Fixing HWRT collisions on console.

- Platform specific wrapper for vfx RT code to access payload and trace result structures better on each platform.

- Moved hash table to use RWBuffers instead of RWStructuredBuffers.
   - Unclear why right now but when using structured buffers it seems like the format was interperted wrong somehow so only the first byte was ever read/written to the buffer. Causeing havok with the rest of the hash table setup.
   - Malformed hash tables also can cause GPU hangs so added safety code for that.

Misc
- Adding transitions after buffer clears.
- Fixed bad define in #if

diff --git a/Engine/Plugins/FX/Niagara/Shaders/Private/NiagaraDataInterfaceCollisionQuery.ush b/Engine/Plugins/FX/Niagara/Shaders/Private/NiagaraDataInterfaceCollisionQuery.ush
--- a/Engine/Plugins/FX/Niagara/Shaders/Private/NiagaraDataInterfaceCollisionQuery.ush
+++ b/Engine/Plugins/FX/Niagara/Shaders/Private/NiagaraDataInterfaceCollisionQuery.ush
@@ -217,7 +217,7 @@ void NDICollisionQuery_CreateAsyncRayTrace(uint MaxRayTraceCount, RWStructuredBu
 	}
 }
 
-void NDICollisionQuery_ReadAsyncRayTrace(uint MaxRayTraceCount, StructuredBuffer<FNiagaraRayTracingPayload> IntersectionResults, int In_PreviousFrameQueryID, out bool Out_CollisionValid, out float Out_CollisionDistance, out float3 Out_CollisionPosWorld, out float3 Out_CollisionNormal)
+void NDICollisionQuery_ReadAsyncRayTrace(uint MaxRayTraceCount, StructuredBuffer<FNiagaraRayTracingResult> IntersectionResults, int In_PreviousFrameQueryID, out bool Out_CollisionValid, out float Out_CollisionDistance, out float3 Out_CollisionPosWorld, out float3 Out_CollisionNormal)
 {
 	Out_CollisionValid = false;
 	Out_CollisionDistance = 0.0f;
@@ -226,7 +226,7 @@ void NDICollisionQuery_ReadAsyncRayTrace(uint MaxRayTraceCount, StructuredBuffer
 
 	if (In_PreviousFrameQueryID >= 0 && uint(In_PreviousFrameQueryID) < MaxRayTraceCount)
 	{
-		FNiagaraRayTracingPayload HitResult = IntersectionResults[In_PreviousFrameQueryID];
+		FNiagaraRayTracingResult HitResult = IntersectionResults[In_PreviousFrameQueryID];
 
 		Out_CollisionValid = HitResult.IsHit();
 		if (Out_CollisionValid)
diff --git a/Engine/Plugins/FX/Niagara/Shaders/Private/NiagaraDataInterfaceCollisionQueryTemplate.ush b/Engine/Plugins/FX/Niagara/Shaders/Private/NiagaraDataInterfaceCollisionQueryTemplate.ush
--- a/Engine/Plugins/FX/Niagara/Shaders/Private/NiagaraDataInterfaceCollisionQueryTemplate.ush
+++ b/Engine/Plugins/FX/Niagara/Shaders/Private/NiagaraDataInterfaceCollisionQueryTemplate.ush
@@ -4,7 +4,7 @@
 bool											RayTracingEnabled_{ParameterName};
 uint                                            MaxRayTraceCount_{ParameterName};
 RWStructuredBuffer<FNiagaraRayData>               RWRayRequests_{ParameterName};
-StructuredBuffer<FNiagaraRayTracingPayload>     IntersectionResults_{ParameterName};
+StructuredBuffer<FNiagaraRayTracingResult>     IntersectionResults_{ParameterName};
 RWBuffer<uint>                                  RWRayTraceCounts_{ParameterName};
 #endif
 
diff --git a/Engine/Plugins/FX/Niagara/Shaders/Private/NiagaraRayTraceCollisionGroupShaders.usf b/Engine/Plugins/FX/Niagara/Shaders/Private/NiagaraRayTraceCollisionGroupShaders.usf
--- a/Engine/Plugins/FX/Niagara/Shaders/Private/NiagaraRayTraceCollisionGroupShaders.usf
+++ b/Engine/Plugins/FX/Niagara/Shaders/Private/NiagaraRayTraceCollisionGroupShaders.usf
@@ -21,8 +21,10 @@ void UpdatePrimIdToCollisionGroupMap(uint Index : SV_DispatchThreadID)
 		//TODO: Handle failures.
 		if (PrimId != -1)
 		{
-			HashTableAdd(PrimId, CollisionGroupIdIndex);
-			RWHashToCollisionGroups[CollisionGroupIdIndex] = CollisionGroupId;
+			if(HashTableAdd(PrimId, CollisionGroupIdIndex))
+			{
+				RWHashToCollisionGroups[CollisionGroupIdIndex] = CollisionGroupId;
+			}
 		}
 	}
-}
+}
\ No newline at end of file
diff --git a/Engine/Plugins/FX/Niagara/Shaders/Private/NiagaraRayTracingCommon.ush b/Engine/Plugins/FX/Niagara/Shaders/Private/NiagaraRayTracingCommon.ush
--- a/Engine/Plugins/FX/Niagara/Shaders/Private/NiagaraRayTracingCommon.ush
+++ b/Engine/Plugins/FX/Niagara/Shaders/Private/NiagaraRayTracingCommon.ush
@@ -14,11 +14,17 @@ struct FNiagaraRayData : FBasicRayData
 	uint CollisionGroup;
 };
 
+struct FNiagaraRayTracingPayload : FVFXTracePayload
+{
+
+};
+
 // payload structure for Niagara ray traces
-struct FNiagaraRayTracingPayload : FDefaultPayload
+struct FNiagaraRayTracingResult : FMinimalPayload
 {
 	float3 WorldPosition;
 	float3 WorldNormal;
 };
 
+
 #endif // NIAGARA_USE_RAYTRACING
\ No newline at end of file
diff --git a/Engine/Plugins/FX/Niagara/Shaders/Private/NiagaraRayTracingShaders.usf b/Engine/Plugins/FX/Niagara/Shaders/Private/NiagaraRayTracingShaders.usf
--- a/Engine/Plugins/FX/Niagara/Shaders/Private/NiagaraRayTracingShaders.usf
+++ b/Engine/Plugins/FX/Niagara/Shaders/Private/NiagaraRayTracingShaders.usf
@@ -7,10 +7,11 @@
 #include "/Engine/Private/RayTracing/RayTracingHitGroupCommon.ush"
 #include "NiagaraRayTracingCommon.ush"
 #include "/Engine/Private/HashTable.ush"
+#include "/Engine/Private/RayTracing/VFXTraceRay.ush"
 
 RaytracingAccelerationStructure               TLAS;
 StructuredBuffer<FNiagaraRayData>               Rays;
-RWStructuredBuffer<FNiagaraRayTracingPayload> CollisionOutput;
+RWStructuredBuffer<FNiagaraRayTracingResult> CollisionOutput;
 
 #if NIAGARA_RAYTRACE_FAKE_INDIRECT
 Buffer<uint>                                  RayTraceCounts;
@@ -70,16 +71,17 @@ RAY_TRACING_ENTRY_RAYGEN(NiagaraCollisionRayTraceRG)
 
 	uint RayFlags = RAY_FLAG_FORCE_OPAQUE | RAY_FLAG_CULL_BACK_FACING_TRIANGLES; // don't run anyhit shader
 
-	FNiagaraRayTracingPayload Payload = (FNiagaraRayTracingPayload)0;
+	FVFXTracePayload Payload = (FVFXTracePayload) 0;
 	uint RayCollisionGroup = InputRay.CollisionGroup;
+
 	if (RayCollisionGroup != -1)
 	{
 		uint NumCasts = 0;
 		LOOP
 		while (NumCasts++ <= MaxRetraces)
 		{
-			Payload = (FNiagaraRayTracingPayload)0;
-			TraceRay(
+			Payload = (FVFXTracePayload)0;
+			VFX_TRACE_RAY(
 				TLAS,   // AccelerationStructure
 				RayFlags,
 				RAY_TRACING_MASK_OPAQUE,          // InstanceInclusionMask
@@ -110,7 +112,7 @@ RAY_TRACING_ENTRY_RAYGEN(NiagaraCollisionRayTraceRG)
 	}
 	else
 	{
-		TraceRay(
+		VFX_TRACE_RAY(
 			TLAS,   // AccelerationStructure
 			RayFlags,
 			RAY_TRACING_MASK_OPAQUE,          // InstanceInclusionMask
@@ -123,14 +125,10 @@ RAY_TRACING_ENTRY_RAYGEN(NiagaraCollisionRayTraceRG)
 	}
 
 	CollisionOutput[RayIndex].HitT = Payload.HitT;
-	if (CollisionOutput[RayIndex].IsHit())
+	if (Payload.IsHit())
 	{
 		CollisionOutput[RayIndex].WorldPosition = Payload.WorldPosition;
 		CollisionOutput[RayIndex].WorldNormal = Payload.WorldNormal;
-
-		CollisionOutput[RayIndex].PrimitiveIndex = Payload.PrimitiveIndex;
-		CollisionOutput[RayIndex].InstanceIndex = Payload.InstanceIndex;
-		CollisionOutput[RayIndex].InstanceID = Payload.InstanceID;
 	}
 }
 
@@ -142,29 +140,13 @@ float3 GetBarycentricPosition(FDefaultAttributes Attributes, FTriangleBaseAttrib
 }
 
 RAY_TRACING_ENTRY_CLOSEST_HIT(NiagaraCollisionRayTraceCH,
-	FNiagaraRayTracingPayload, Payload,
+	FVFXTracePayload, Payload,
 	FDefaultAttributes, Attributes)
 {
-	Payload.HitT = RayTCurrent();
-
-	const int Primitive = GetInstanceUserData();
-	Payload.PrimitiveIndex = Primitive;
-
-	const int PrimitiveInstanceIndex = InstanceIndex() - GetBaseInstanceIndex();
-	Payload.InstanceIndex = PrimitiveInstanceIndex;
-
-	FTriangleBaseAttributes Tri = LoadTriangleBaseAttributes(PrimitiveIndex());
-	float3 LocalPosition = GetBarycentricPosition(Attributes, Tri);
-	float3 LocalEdges[2] = { Tri.LocalPositions[1] - Tri.LocalPositions[0], Tri.LocalPositions[2] - Tri.LocalPositions[0] };
-	float3 LocalNormal = cross(LocalEdges[1], LocalEdges[0]);
-	float3x3 InverseTranspose3x3 = transpose((float3x3)WorldToObject4x3());
-
-	Payload.WorldPosition = WorldRayOrigin() + WorldRayDirection() * RayTCurrent();
-	//Payload.WorldPosition = mul(float4(LocalPosition, 1.0f), ObjectToWorld4x3()).xyz;
-	Payload.WorldNormal = normalize(mul(LocalNormal, InverseTranspose3x3));
+	InitVFXTracePayloadFromHitShader(Attributes, Payload);
 }
 
-RAY_TRACING_ENTRY_MISS(NiagaraCollisionRayTraceMiss, FNiagaraRayTracingPayload, Payload)
+RAY_TRACING_ENTRY_MISS(NiagaraCollisionRayTraceMiss, FVFXTracePayload, Payload)
 {
 	Payload.SetMiss();
 }
diff --git a/Engine/Plugins/FX/Niagara/Source/Niagara/Private/NiagaraDataInterfaceCollisionQuery.cpp b/Engine/Plugins/FX/Niagara/Source/Niagara/Private/NiagaraDataInterfaceCollisionQuery.cpp
--- a/Engine/Plugins/FX/Niagara/Source/Niagara/Private/NiagaraDataInterfaceCollisionQuery.cpp
+++ b/Engine/Plugins/FX/Niagara/Source/Niagara/Private/NiagaraDataInterfaceCollisionQuery.cpp
@@ -180,7 +180,7 @@ struct FNiagaraDataIntefaceProxyCollisionQuery : public FNiagaraDataInterfacePro
 
 			RayTraceIntersections.Initialize(
 				TEXT("NiagaraRayTraceIntersections"),
-				sizeof(FNiagaraRayTracingPayload),
+				sizeof(FNiagaraRayTracingResult),
 				MaxRayTraceCount,
 				BUF_Static,
 				false /*bUseUavCounter*/,
diff --git a/Engine/Plugins/FX/Niagara/Source/Niagara/Private/NiagaraFunctionLibrary.cpp b/Engine/Plugins/FX/Niagara/Source/Niagara/Private/NiagaraFunctionLibrary.cpp
--- a/Engine/Plugins/FX/Niagara/Source/Niagara/Private/NiagaraFunctionLibrary.cpp
+++ b/Engine/Plugins/FX/Niagara/Source/Niagara/Private/NiagaraFunctionLibrary.cpp
@@ -1776,7 +1776,7 @@ void UNiagaraFunctionLibrary::InitVectorVMFastPathOps()
 
 void UNiagaraFunctionLibrary::SetComponentNiagaraGPURayTracedCollisionGroup(UObject* WorldContextObject, UPrimitiveComponent* Primitive, int32 CollisionGroup)
 {
-#if WITH_RAYTRACING
+#if RHI_RAYTRACING
 	if(UWorld* World = GEngine->GetWorldFromContextObject(WorldContextObject, EGetWorldErrorMode::LogAndReturnNull))
 	{
 		if(FNiagaraWorldManager* WorldMan = FNiagaraWorldManager::Get(World))
@@ -1795,7 +1795,7 @@ void UNiagaraFunctionLibrary::SetComponentNiagaraGPURayTracedCollisionGroup(UObj
 
 void UNiagaraFunctionLibrary::SetActorNiagaraGPURayTracedCollisionGroup(UObject* WorldContextObject, AActor* Actor, int32 CollisionGroup)
 {
-#if WITH_RAYTRACING
+#if RHI_RAYTRACING
 	if (Actor == nullptr)
 	{
 		return;
@@ -1822,7 +1822,7 @@ void UNiagaraFunctionLibrary::SetActorNiagaraGPURayTracedCollisionGroup(UObject*
 
 int32 UNiagaraFunctionLibrary::AcquireNiagaraGPURayTracedCollisionGroup(UObject* WorldContextObject)
 {
-#if WITH_RAYTRACING
+#if RHI_RAYTRACING
 	if (UWorld* World = GEngine->GetWorldFromContextObject(WorldContextObject, EGetWorldErrorMode::LogAndReturnNull))
 	{
 		if (FNiagaraWorldManager* WorldMan = FNiagaraWorldManager::Get(World))
@@ -1842,7 +1842,7 @@ int32 UNiagaraFunctionLibrary::AcquireNiagaraGPURayTracedCollisionGroup(UObject*
 
 void UNiagaraFunctionLibrary::ReleaseNiagaraGPURayTracedCollisionGroup(UObject* WorldContextObject, int32 CollisionGroup)
 {
-#if WITH_RAYTRACING
+#if RHI_RAYTRACING
 	if (UWorld* World = GEngine->GetWorldFromContextObject(WorldContextObject, EGetWorldErrorMode::LogAndReturnNull))
 	{
 		if (FNiagaraWorldManager* WorldMan = FNiagaraWorldManager::Get(World))
diff --git a/Engine/Plugins/FX/Niagara/Source/NiagaraShader/Private/NiagaraRayTracingHelper.cpp b/Engine/Plugins/FX/Niagara/Source/NiagaraShader/Private/NiagaraRayTracingHelper.cpp
--- a/Engine/Plugins/FX/Niagara/Source/NiagaraShader/Private/NiagaraRayTracingHelper.cpp
+++ b/Engine/Plugins/FX/Niagara/Source/NiagaraShader/Private/NiagaraRayTracingHelper.cpp
@@ -1,6 +1,7 @@
 // Copyright Epic Games, Inc. All Rights Reserved.
 
 #include "NiagaraRayTracingHelper.h"
+#include "NiagaraShader.h"
 
 #if RHI_RAYTRACING
 
@@ -20,7 +21,7 @@ class FNiagaraUpdateCollisionGroupMapCS : public FGlobalShader
 	SHADER_USE_PARAMETER_STRUCT(FNiagaraUpdateCollisionGroupMapCS, FGlobalShader);
 
 	BEGIN_SHADER_PARAMETER_STRUCT(FParameters, )
-		SHADER_PARAMETER_UAV(RWStructuredBuffer<UINT>, RWHashTable)		
+		SHADER_PARAMETER_UAV(Buffer<UINT>, RWHashTable)		
 		SHADER_PARAMETER(uint32, HashTableSize)
 		SHADER_PARAMETER_UAV(Buffer<UINT>, RWHashToCollisionGroups)
 		SHADER_PARAMETER_SRV(Buffer<UINT>, NewPrimIdCollisionGroupPairs)
@@ -61,12 +62,12 @@ class FNiagaraCollisionRayTraceRG : public FGlobalShader
 
 	BEGIN_SHADER_PARAMETER_STRUCT(FParameters, )
 		SHADER_PARAMETER_STRUCT_INCLUDE(FGPUSceneParameters, GPUSceneParameters)
-		SHADER_PARAMETER_SRV(StructuredBuffer<UINT>, HashTable)
-		SHADER_PARAMETER_UAV(RWStructuredBuffer<UINT>,RWHashTable)
+		SHADER_PARAMETER_SRV(Buffer<UINT>, HashTable)
+		SHADER_PARAMETER_UAV(Buffer<UINT>, RWHashTable)
 		SHADER_PARAMETER(uint32, HashTableSize)
 		SHADER_PARAMETER_SRV(RaytracingAccelerationStructure, TLAS)
-		SHADER_PARAMETER_SRV(StructuredBuffer<FNiagaraRayData>, Rays)
-		SHADER_PARAMETER_UAV(RWStructuredBuffer<FNiagaraRayTracingPayload>, CollisionOutput)
+		SHADER_PARAMETER_SRV(Buffer<FNiagaraRayData>, Rays)
+		SHADER_PARAMETER_UAV(Buffer<FNiagaraRayTracingResult>, CollisionOutput)
 		SHADER_PARAMETER_SRV(Buffer<UINT>, RayTraceCounts)
 		SHADER_PARAMETER_SRV(Buffer<UINT>, HashToCollisionGroups)
 		SHADER_PARAMETER(uint32, MaxRetraces)
@@ -298,14 +299,7 @@ void FNiagaraRayTracingHelper::UpdateCollisionGroupMap(FRHICommandList& RHICmdLi
 			HashTableSize = AllocInstances;
 
 			PrimIdHashTable.Release();
-			PrimIdHashTable.Initialize(
-				TEXT("NiagaraPrimIdHashTable"),
-				sizeof(uint32),
-				AllocInstances,
-				BUF_Static,
-				false /*bUseUavCounter*/,
-				false /*bAppendBuffer*/,
-				ERHIAccess::UAVCompute);
+			PrimIdHashTable.Initialize(TEXT("NiagaraPrimIdHashTable"), sizeof(uint32), AllocInstances, EPixelFormat::PF_R32_UINT, BUF_Static);
 
 			HashToCollisionGroups.Release();
 			HashToCollisionGroups.Initialize(TEXT("NiagaraPrimIdHashToCollisionGroups"), sizeof(uint32), AllocInstances, EPixelFormat::PF_R32_UINT, BUF_Static);
@@ -314,9 +308,11 @@ void FNiagaraRayTracingHelper::UpdateCollisionGroupMap(FRHICommandList& RHICmdLi
 		RHICmdList.Transition(FRHITransitionInfo(PrimIdHashTable.UAV, ERHIAccess::Unknown, ERHIAccess::UAVCompute));
 		RHICmdList.Transition(FRHITransitionInfo(HashToCollisionGroups.UAV, ERHIAccess::Unknown, ERHIAccess::UAVCompute));
 
-		//First we have to clear the buffers. Can probably do this better.
-		NiagaraFillGPUIntBuffer(RHICmdList, FeatureLevel, PrimIdHashTable, 0);
-		NiagaraFillGPUIntBuffer(RHICmdList, FeatureLevel, HashToCollisionGroups, INDEX_NONE);
+ 		//First we have to clear the buffers. Can probably do this better.
+ 		NiagaraFillGPUIntBuffer(RHICmdList, FeatureLevel, PrimIdHashTable, 0);
+ 		RHICmdList.Transition(FRHITransitionInfo(PrimIdHashTable.UAV, ERHIAccess::UAVCompute, ERHIAccess::UAVCompute));
+ 		NiagaraFillGPUIntBuffer(RHICmdList, FeatureLevel, HashToCollisionGroups, INDEX_NONE);
+ 		RHICmdList.Transition(FRHITransitionInfo(HashToCollisionGroups.UAV, ERHIAccess::UAVCompute, ERHIAccess::UAVCompute));
 
 		TShaderMapRef<FNiagaraUpdateCollisionGroupMapCS> ComputeShader(GetGlobalShaderMap(GMaxRHIFeatureLevel));
 		FRHIComputeShader* ShaderRHI = ComputeShader.GetComputeShader();
diff --git a/Engine/Plugins/FX/Niagara/Source/NiagaraShader/Public/NiagaraRayTracingHelper.h b/Engine/Plugins/FX/Niagara/Source/NiagaraShader/Public/NiagaraRayTracingHelper.h
--- a/Engine/Plugins/FX/Niagara/Source/NiagaraShader/Public/NiagaraRayTracingHelper.h
+++ b/Engine/Plugins/FX/Niagara/Source/NiagaraShader/Public/NiagaraRayTracingHelper.h
@@ -26,12 +26,16 @@ struct FNiagaraRayData : public FBasicRayData
 struct FNiagaraRayTracingPayload
 {
 	float HitT;
-
 	uint32 PrimitiveIndex;
 	uint32 InstanceIndex;
 	float Barycentrics[2];
+	float WorldPosition[3];
+	float WorldNormal[3];
+};
 
-	uint32 InstanceID;
+struct FNiagaraRayTracingResult
+{
+	float HitT;
 
 	float WorldPosition[3];
 	float WorldNormal[3];
@@ -75,7 +79,7 @@ class NIAGARASHADER_API FNiagaraRayTracingHelper
 		/** Hash table. 
 		 PrimIdHashTable is the main hash table that maps GPUSceneInstanceIndex to and Index we can use to store Collision Groups inside HashToCollisionGroups.
 */
-		FRWBufferStructured PrimIdHashTable;
+		FRWBuffer PrimIdHashTable;
 		uint32 HashTableSize = 0;
 		FRWBuffer HashToCollisionGroups;
 
diff --git a/Engine/Shaders/Private/HashTable.ush b/Engine/Shaders/Private/HashTable.ush
--- a/Engine/Shaders/Private/HashTable.ush
+++ b/Engine/Shaders/Private/HashTable.ush
@@ -23,9 +23,9 @@ uint MurmurMix( uint Hash )
 }
 
 // Linear probing hash table
-StructuredBuffer< uint >	HashTable;
-RWStructuredBuffer< uint >	RWHashTable;
-uint						HashTableSize;
+Buffer< uint >		HashTable;
+RWBuffer< uint >	RWHashTable;
+uint			HashTableSize;
 
 // Returns true if key is added for the first time.
 // Index output is the hash table bucket this key is stored in.
@@ -34,10 +34,17 @@ bool HashTableAdd( uint Key, out uint Index )
 	// Zero is reserved as invalid
 	Key++;
 
+	uint NumLoops = 0;
 	LOOP
 	[allow_uav_condition]
-	for( Index = MurmurMix( Key );; Index++ )
+	for(Index = MurmurMix(Key); ; Index++ )
 	{
+		//Belt and braces safety code to prevent inf loops if tables are malformed.
+		if(++NumLoops > HashTableSize)
+		{
+			break;
+		}
+
 		//Index &= HashTableSize - 1;
 		Index = Index % HashTableSize;
 
@@ -69,10 +76,17 @@ bool HashTableFind( uint Key, out uint Index )
 	// Zero is reserved as invalid
 	Key++;
 
+	uint NumLoops = 0;
 	LOOP
 	[allow_uav_condition]
 	for( Index = MurmurMix( Key );; Index++ )
 	{
+		//Belt and braces safety code to prevent inf loops if tables are malformed.
+		if (++NumLoops > HashTableSize)
+		{
+			break;
+		}
+
 		//Index &= HashTableSize - 1;
 		Index = Index % HashTableSize;
 
diff --git a/Engine/Shaders/Private/RayTracing/RayTracingCommon.ush b/Engine/Shaders/Private/RayTracing/RayTracingCommon.ush
--- a/Engine/Shaders/Private/RayTracing/RayTracingCommon.ush
+++ b/Engine/Shaders/Private/RayTracing/RayTracingCommon.ush
@@ -119,6 +119,14 @@ struct FDefaultPayload : FIntersectionPayload
 	uint   InstanceID; // Value of FRayTracingGeometryInstance::UserData. Undefined on miss.
 };
 
+struct FVFXTracePayload : FIntersectionPayload
+{
+	void SetWorldNormal(float3 InNormal) { WorldNormal = InNormal; }
+
+	float3 WorldPosition;
+	float3 WorldNormal;
+};
+
 #define RAY_TRACING_DEFERRED_MATERIAL_KEY_INVALID   0xFFFFFFFF
 #define RAY_TRACING_DEFERRED_MATERIAL_KEY_RAY_MISS  0xFFFFFFFE
 
diff --git a/Engine/Shaders/Private/RayTracing/VFXTraceRay.ush b/Engine/Shaders/Private/RayTracing/VFXTraceRay.ush
new file mode 100644
--- /dev/null
+++ b/Engine/Shaders/Private/RayTracing/VFXTraceRay.ush
@@ -0,0 +1,53 @@
+// Copyright Epic Games, Inc. All Rights Reserved.
+
+/*=============================================================================
+	VFXTraceRay.ush: VFX Specific ray tracing utilities.
+=============================================================================*/
+
+#pragma once
+
+#ifndef VFX_TRACE_RAY_USH_INCLUDED
+#define VFX_TRACE_RAY_USH_INCLUDED
+
+#include "RayTracingCommon.ush"
+#include "RayTracingHitGroupCommon.ush"
+
+
+#ifdef OVERRIDE_VFXTRACERAY_USH
+#include "/Platform/Private/VFXTraceRay.ush"
+#else // OVERRIDE_VFXTRACERAY_USH
+
+
+float3 GetWorldNormalCommon(FVFXTracePayload Payload)
+{
+	FTriangleBaseAttributes Tri = LoadTriangleBaseAttributes(Payload.PrimitiveIndex);
+
+	float3 LocalPosition = (1.0f - Payload.Barycentrics.x - Payload.Barycentrics.y) * Tri.LocalPositions[0]
+		+ Payload.Barycentrics.x * Tri.LocalPositions[1]
+		+ Payload.Barycentrics.y * Tri.LocalPositions[2];
+
+	float3 LocalEdges[2] = { Tri.LocalPositions[1] - Tri.LocalPositions[0], Tri.LocalPositions[2] - Tri.LocalPositions[0] };
+	float3 LocalNormal = cross(LocalEdges[1], LocalEdges[0]);
+	float3x3 InverseTranspose3x3 = transpose((float3x3)WorldToObject4x3());
+
+	return normalize(mul(LocalNormal, InverseTranspose3x3));
+}
+
+void InitVFXTracePayloadFromHitShader(FDefaultAttributes Attributes, out FVFXTracePayload OutPayload)
+{
+	OutPayload.HitT = RayTCurrent();
+	OutPayload.PrimitiveIndex = GetInstanceUserData();
+	OutPayload.InstanceIndex = InstanceIndex() - GetBaseInstanceIndex();
+	OutPayload.Barycentrics = Attributes.Barycentrics;
+
+	OutPayload.WorldPosition = WorldRayOrigin() + WorldRayDirection() * RayTCurrent();
+
+	OutPayload.WorldNormal = GetWorldNormalCommon(OutPayload);
+}
+
+#define VFX_TRACE_RAY(TLAS, RayFlags, InstanceInclusionMask, RayContributionToHitGroupIndex, MultiplierForGeometryContributionToShaderIndex, MissShaderIndex, Ray, Payload)\
+		TraceRay(TLAS, RayFlags, InstanceInclusionMask, RayContributionToHitGroupIndex, MultiplierForGeometryContributionToShaderIndex, MissShaderIndex, Ray, Payload);
+
+#endif //#ifndef UE_TRACE_RAY
+
+#endif // UE_TRACE_RAY_USH_INCLUDED // Workarround for UE-66460
