Fixed a bug where the "fail if not previously set" default mode would not correctly pick up previous writes

diff --git a/Engine/Plugins/FX/Niagara/Source/NiagaraCore/Private/NiagaraCustomVersion.cpp b/Engine/Plugins/FX/Niagara/Source/NiagaraCore/Private/NiagaraCustomVersion.cpp
--- a/Engine/Plugins/FX/Niagara/Source/NiagaraCore/Private/NiagaraCustomVersion.cpp
+++ b/Engine/Plugins/FX/Niagara/Source/NiagaraCore/Private/NiagaraCustomVersion.cpp
@@ -8,4 +8,4 @@ const FGuid FNiagaraCustomVersion::GUID(0xFCF57AFA, 0x50764283, 0xB9A9E658, 0xFF
 // Register the custom version with core
 FCustomVersionRegistration GRegisterNiagaraCustomVersion(FNiagaraCustomVersion::GUID, FNiagaraCustomVersion::LatestVersion, TEXT("NiagaraVer"));
 
-const FGuid FNiagaraCustomVersion::LatestScriptCompileVersion(0x02497746, 0x33D75C49, 0x9A60CB65, 0x7D6AC2A2);
+const FGuid FNiagaraCustomVersion::LatestScriptCompileVersion(0x02497746, 0x33D75C49, 0x9A60CB65, 0x7D6AC5A2);
diff --git a/Engine/Plugins/FX/Niagara/Source/NiagaraEditor/Private/NiagaraHlslTranslator.cpp b/Engine/Plugins/FX/Niagara/Source/NiagaraEditor/Private/NiagaraHlslTranslator.cpp
--- a/Engine/Plugins/FX/Niagara/Source/NiagaraEditor/Private/NiagaraHlslTranslator.cpp
+++ b/Engine/Plugins/FX/Niagara/Source/NiagaraEditor/Private/NiagaraHlslTranslator.cpp
@@ -32,11 +32,6 @@
 #include "NiagaraDataInterface.h"
 #include "NiagaraDataInterfaceCurve.h"
 #include "NiagaraDataInterfaceVector2DCurve.h"
-#include "NiagaraDataInterfaceVectorCurve.h"
-#include "NiagaraDataInterfaceVector4Curve.h"
-#include "NiagaraDataInterfaceColorCurve.h"
-#include "NiagaraDataInterfaceStaticMesh.h"
-#include "NiagaraDataInterfaceCurlNoise.h"
 
 #include "NiagaraParameterCollection.h"
 #include "NiagaraEditorTickables.h"
@@ -6915,9 +6910,21 @@ void FHlslNiagaraTranslator::HandleDataInterfaceCall(FNiagaraScriptDataInterface
 	}
 }
 
+bool IsVariableWriteBeforeRead(const TArray<FNiagaraParameterMapHistory::FReadHistory>& ReadHistory)
+{
+	for (const FNiagaraParameterMapHistory::FReadHistory& History : ReadHistory)
+	{
+		if (History.PreviousWritePin.Pin == nullptr)
+		{
+			return false;
+		}
+	}
+	return true;
+}
+
 void FHlslNiagaraTranslator::RegisterFunctionCall(ENiagaraScriptUsage ScriptUsage, const FString& InName, const FString& InFullName, const FGuid& CallNodeId, const FString& InFunctionNameSuffix, UNiagaraScriptSource* Source,
-	FNiagaraFunctionSignature& InSignature, bool bIsCustomHlsl, const FString& InCustomHlsl, TArray<int32>& Inputs, TArrayView<UEdGraphPin* const> CallInputs, TArrayView<UEdGraphPin* const> CallOutputs,
-	FNiagaraFunctionSignature& OutSignature)
+                                                  FNiagaraFunctionSignature& InSignature, bool bIsCustomHlsl, const FString& InCustomHlsl, TArray<int32>& Inputs, TArrayView<UEdGraphPin* const> CallInputs, TArrayView<UEdGraphPin* const> CallOutputs,
+                                                  FNiagaraFunctionSignature& OutSignature)
 {
 	NIAGARA_SCOPE_CYCLE_COUNTER(STAT_NiagaraEditor_Module_NiagaraHLSLTranslator_RegisterFunctionCall);
 
@@ -7031,9 +7038,9 @@ void FHlslNiagaraTranslator::RegisterFunctionCall(ENiagaraScriptUsage ScriptUsag
 
 							for (uint32 VarIdx = History.MapNodeVariableMetaData[FoundIdx].Key; VarIdx < History.MapNodeVariableMetaData[FoundIdx].Value; VarIdx++)
 							{
-								if (History.PerVariableReadHistory[VarIdx].Num() == 0)
+								if (IsVariableWriteBeforeRead(History.PerVariableReadHistory[VarIdx]))
 								{
-									// We don't need to worry about defaults if the variable is only written to.
+									// We don't need to worry about defaults if the variable is written before being read or never read at all.
 									continue;
 								}
 
