Fix for texture coord arrays in Niagara on GL/Vulkan with DXC

diff --git a/Engine/Plugins/FX/Niagara/Shaders/Private/NiagaraMeshVertexFactory.ush b/Engine/Plugins/FX/Niagara/Shaders/Private/NiagaraMeshVertexFactory.ush
--- a/Engine/Plugins/FX/Niagara/Shaders/Private/NiagaraMeshVertexFactory.ush
+++ b/Engine/Plugins/FX/Niagara/Shaders/Private/NiagaraMeshVertexFactory.ush
@@ -103,9 +103,30 @@ struct FVertexFactoryInput
 	half4	VertexColor : ATTRIBUTE3;
 
 #if NUM_MATERIAL_TEXCOORDS_VERTEX
-	float2	TexCoords[NUM_MATERIAL_TEXCOORDS_VERTEX] : ATTRIBUTE4;
+	float2	TexCoords0 : ATTRIBUTE4;
+	#if NUM_MATERIAL_TEXCOORDS_VERTEX > 1
+		float2	TexCoords1 : ATTRIBUTE5;
+	#endif
+	#if NUM_MATERIAL_TEXCOORDS_VERTEX > 2
+		float2	TexCoords2 : ATTRIBUTE6;
+	#endif
+	#if NUM_MATERIAL_TEXCOORDS_VERTEX > 3
+		float2	TexCoords3 : ATTRIBUTE7;
+	#endif
+	#if NUM_MATERIAL_TEXCOORDS_VERTEX > 4
+		float2	TexCoords4 : ATTRIBUTE8;
+	#endif
+	#if NUM_MATERIAL_TEXCOORDS_VERTEX > 5
+		float2	TexCoords5 : ATTRIBUTE9;
+	#endif
+	#if NUM_MATERIAL_TEXCOORDS_VERTEX > 6
+		float2	TexCoords6 : ATTRIBUTE10;
+	#endif
+	#if NUM_MATERIAL_TEXCOORDS_VERTEX > 7
+		float2	TexCoords7 : ATTRIBUTE11;
+	#endif
 #elif USE_PARTICLE_SUBUVS
-	float2	TexCoords[1] : ATTRIBUTE4;
+	float2	TexCoords0 : ATTRIBUTE4;
 #endif
 
 	uint InstanceId	: SV_InstanceID;
@@ -719,10 +740,28 @@ FVertexFactoryIntermediates GetVertexFactoryIntermediates(FVertexFactoryInput In
 	Intermediates.ParticleTranslatedWorldPosition = Intermediates.ParticlePosition.xyz + ResolvedView.PreViewTranslation.xyz;
 
 #if NUM_MATERIAL_TEXCOORDS_VERTEX
-	for(int CoordinateIndex = 0; CoordinateIndex < NUM_MATERIAL_TEXCOORDS_VERTEX; CoordinateIndex++)
-	{
-		Intermediates.TexCoords[CoordinateIndex] = Input.TexCoords[CoordinateIndex].xy;
-	}
+	Intermediates.TexCoords[0] = Input.TexCoords0;
+	#if NUM_MATERIAL_TEXCOORDS_VERTEX > 1
+		Intermediates.TexCoords[1] = Input.TexCoords1;
+	#endif
+	#if NUM_MATERIAL_TEXCOORDS_VERTEX > 2
+		Intermediates.TexCoords[2] = Input.TexCoords2;
+	#endif
+	#if NUM_MATERIAL_TEXCOORDS_VERTEX > 3
+		Intermediates.TexCoords[3] = Input.TexCoords3;
+	#endif
+	#if NUM_MATERIAL_TEXCOORDS_VERTEX > 4
+		Intermediates.TexCoords[4] = Input.TexCoords4;
+	#endif
+	#if NUM_MATERIAL_TEXCOORDS_VERTEX > 5
+		Intermediates.TexCoords[5] = Input.TexCoords5;
+	#endif
+	#if NUM_MATERIAL_TEXCOORDS_VERTEX > 6
+		Intermediates.TexCoords[6] = Input.TexCoords6;
+	#endif
+	#if NUM_MATERIAL_TEXCOORDS_VERTEX > 7
+		Intermediates.TexCoords[7] = Input.TexCoords7;
+	#endif
 #endif
 
 #if USE_PARTICLE_SUBUVS
@@ -735,8 +774,8 @@ FVertexFactoryIntermediates GetVertexFactoryIntermediates(FVertexFactoryInput In
 	float SubImageBH = fmod(SubImageB, NiagaraMeshVF.SubImageSize.x);
 	float SubImageAV = floor(SubImageA * NiagaraMeshVF.SubImageSize.z);
 	float SubImageBV = floor(SubImageB * NiagaraMeshVF.SubImageSize.z);
-	Intermediates.SubUVCoords.xy = (float2(SubImageAH, SubImageAV) + Input.TexCoords[0].xy) * NiagaraMeshVF.SubImageSize.zw;
-	Intermediates.SubUVCoords.zw = (float2(SubImageBH, SubImageBV) + Input.TexCoords[0].xy) * NiagaraMeshVF.SubImageSize.zw;
+	Intermediates.SubUVCoords.xy = (float2(SubImageAH, SubImageAV) + Input.TexCoords0.xy) * NiagaraMeshVF.SubImageSize.zw;
+	Intermediates.SubUVCoords.zw = (float2(SubImageBH, SubImageBV) + Input.TexCoords0.xy) * NiagaraMeshVF.SubImageSize.zw;
 	Intermediates.SubUVLerp = (NiagaraMeshVF.SubImageBlendMode == 1) ? SubImageLerp : 0.0f;
 #endif
 
@@ -1024,14 +1063,34 @@ float4 VertexFactoryGetTranslatedPrimitiveVolumeBounds(FVertexFactoryInterpolant
 		Input.VertexColor = float4(1.0, 1.0, 1.0, 0.5);
 #if NUM_MATERIAL_TEXCOORDS_VERTEX
 UNROLL
-		for (int CoordinateIndex = 0; CoordinateIndex < NUM_MATERIAL_TEXCOORDS_VERTEX; CoordinateIndex++)
-		{
-			Input.TexCoords[CoordinateIndex].x = 0.5f;
-			Input.TexCoords[CoordinateIndex].y = 0.5f;
-		}
+
+#if NUM_MATERIAL_TEXCOORDS_VERTEX
+	Input.TexCoords0.xy = float2(0.5f, 0.5f);
+	#if NUM_MATERIAL_TEXCOORDS_VERTEX > 1
+		Input.TexCoords1.xy = float2(0.5f, 0.5f);
+	#endif
+	#if NUM_MATERIAL_TEXCOORDS_VERTEX > 2
+		Input.TexCoords2.xy = float2(0.5f, 0.5f);
+	#endif
+	#if NUM_MATERIAL_TEXCOORDS_VERTEX > 3
+		Input.TexCoords3.xy = float2(0.5f, 0.5f);
+	#endif
+	#if NUM_MATERIAL_TEXCOORDS_VERTEX > 4
+		Input.TexCoords4.xy = float2(0.5f, 0.5f);
+	#endif
+	#if NUM_MATERIAL_TEXCOORDS_VERTEX > 5
+		Input.TexCoords5.xy = float2(0.5f, 0.5f);
+	#endif
+	#if NUM_MATERIAL_TEXCOORDS_VERTEX > 6
+		Input.TexCoords6.xy = float2(0.5f, 0.5f);
+	#endif
+	#if NUM_MATERIAL_TEXCOORDS_VERTEX > 7
+		Input.TexCoords7.xy = float2(0.5f, 0.5f);
+	#endif
+#endif
+
 #elif USE_PARTICLE_SUBUVS
-		Input.TexCoords[0].x = 0.5f;
-		Input.TexCoords[0].y = 0.5f;
+		Input.TexCoords0.xy = float2(0.5f, 0.5f);
 #endif
 
 #if PASS_NEEDS_VERTEX_ID
