- Add support for setting system or emitter fixed bounds per instance

diff --git a/Engine/Plugins/FX/Niagara/Source/Niagara/Classes/NiagaraEmitterInstance.h b/Engine/Plugins/FX/Niagara/Source/Niagara/Classes/NiagaraEmitterInstance.h
--- a/Engine/Plugins/FX/Niagara/Source/Niagara/Classes/NiagaraEmitterInstance.h
+++ b/Engine/Plugins/FX/Niagara/Source/Niagara/Classes/NiagaraEmitterInstance.h
@@ -147,6 +147,8 @@ public:
 	}
 
 	void SetSystemFixedBoundsOverride(FBox SystemFixedBounds);
+	FORCEINLINE void SetFixedBounds(const FBox& InLocalBounds) { FixedBounds_GT = InLocalBounds; }
+	FBox GetFixedBounds() const;
 
 	UObject* FindBinding(const FNiagaraVariable& InVariable) const;
 	UNiagaraDataInterface* FindDataInterface(const FNiagaraVariable& InVariable) const;
@@ -198,9 +200,13 @@ private:
 	/* Emitter bounds */
 	FBox CachedBounds;
 
+	/** Optional user specified bounds. */
+	FBox FixedBounds_GT;
+	FBox FixedBounds_CNC;
+
 	/** Cached fixed bounds of the parent system which override this Emitter Instances bounds if set. Whenever we initialize the owning SystemInstance we will reconstruct this
 	 ** EmitterInstance and the cached bounds will be unset. */
-	TOptional<FBox> CachedSystemFixedBounds;
+	FBox CachedSystemFixedBounds;
 
 	FNiagaraSystemInstanceID OwnerSystemInstanceID;
 
diff --git a/Engine/Plugins/FX/Niagara/Source/Niagara/Private/NiagaraComponent.cpp b/Engine/Plugins/FX/Niagara/Source/Niagara/Private/NiagaraComponent.cpp
--- a/Engine/Plugins/FX/Niagara/Source/Niagara/Private/NiagaraComponent.cpp
+++ b/Engine/Plugins/FX/Niagara/Source/Niagara/Private/NiagaraComponent.cpp
@@ -494,6 +494,56 @@ void UNiagaraComponent::SetEmitterEnable(FName EmitterName, bool bNewEnableState
 	}
 }
 
+void UNiagaraComponent::SetSystemFixedBounds(FBox LocalBounds)
+{
+	if ( SystemInstanceController.IsValid() )
+	{
+		SystemInstanceController->SetSystemFixedBounds(LocalBounds);
+	}
+}
+
+FBox UNiagaraComponent::GetSystemFixedBounds() const
+{
+	if (SystemInstanceController.IsValid())
+	{
+		return SystemInstanceController->GetSystemFixedBounds();
+	}
+	return FBox(EForceInit::ForceInit);
+}
+
+void UNiagaraComponent::ClearSystemFixedBounds()
+{
+	if (SystemInstanceController.IsValid())
+	{
+		SystemInstanceController->SetSystemFixedBounds(FBox(EForceInit::ForceInit));
+	}
+}
+
+void UNiagaraComponent::SetEmitterFixedBounds(FName EmitterName, FBox LocalBounds)
+{
+	if (SystemInstanceController.IsValid())
+	{
+		SystemInstanceController->SetEmitterFixedBounds(EmitterName, LocalBounds);
+	}
+}
+
+FBox UNiagaraComponent::GetEmitterFixedBounds(FName EmitterName) const
+{
+	if (SystemInstanceController.IsValid())
+	{
+		return SystemInstanceController->GetEmitterFixedBounds(EmitterName);
+	}
+	return FBox(EForceInit::ForceInit);
+}
+
+void UNiagaraComponent::ClearEmitterFixedBounds(FName EmitterName)
+{
+	if (SystemInstanceController.IsValid())
+	{
+		SystemInstanceController->SetEmitterFixedBounds(EmitterName, FBox(EForceInit::ForceInit));
+	}
+}
+
 void UNiagaraComponent::ReleaseToPool()
 {
 	// A component may be marked pending kill before the owner has it's reference set to null.
diff --git a/Engine/Plugins/FX/Niagara/Source/Niagara/Private/NiagaraEmitterInstance.cpp b/Engine/Plugins/FX/Niagara/Source/Niagara/Private/NiagaraEmitterInstance.cpp
--- a/Engine/Plugins/FX/Niagara/Source/Niagara/Private/NiagaraEmitterInstance.cpp
+++ b/Engine/Plugins/FX/Niagara/Source/Niagara/Private/NiagaraEmitterInstance.cpp
@@ -120,9 +120,11 @@ struct FNiagaraEditorOnlyCycleTimer
 //////////////////////////////////////////////////////////////////////////
 
 FNiagaraEmitterInstance::FNiagaraEmitterInstance(FNiagaraSystemInstance* InParentSystemInstance)
-: CachedBounds(ForceInit)
-, CachedSystemFixedBounds()
-, ParentSystemInstance(InParentSystemInstance)
+	: CachedBounds(ForceInit)
+	, FixedBounds_GT(ForceInit)
+	, FixedBounds_CNC(ForceInit)
+	, CachedSystemFixedBounds(ForceInit)
+	, ParentSystemInstance(InParentSystemInstance)
 {
 	ParticleDataSet = new FNiagaraDataSet();
 
@@ -573,6 +575,9 @@ void FNiagaraEmitterInstance::OnPooledReuse()
 	// Ensure we kill any existing particles and mark our buffers for reset
 	bResetPending = true;
 	TotalSpawnedParticles = 0;
+
+	FixedBounds_GT.Init();
+	FixedBounds_CNC.Init();
 }
 
 void FNiagaraEmitterInstance::SetParticleComponentActive(FObjectKey ComponentKey, int32 ParticleID) const
@@ -1019,9 +1024,13 @@ void FNiagaraEmitterInstance::PostTick()
 	}
 
 	CachedBounds.Init();
-	if (CachedSystemFixedBounds.IsSet())
+	if (FixedBounds_CNC.IsValid)
+	{
+		CachedBounds = FixedBounds_CNC;
+	}
+	else if (CachedSystemFixedBounds.IsValid)
 	{
-		CachedBounds = CachedSystemFixedBounds.GetValue();
+		CachedBounds = CachedSystemFixedBounds;
 	}
 	else if (CachedEmitter->bFixedBounds || CachedEmitter->SimTarget == ENiagaraSimTarget::GPUComputeSim)
 	{
@@ -1134,6 +1143,8 @@ void FNiagaraEmitterInstance::PreTick()
 	FScopeCycleCounter SystemStatCounter(CachedEmitter->GetStatID(true, true));
 #endif
 
+	FixedBounds_CNC = FixedBounds_GT;
+
 	checkSlow(ParticleDataSet);
 	FNiagaraDataSet& Data = *ParticleDataSet;
 
@@ -1226,6 +1237,19 @@ void FNiagaraEmitterInstance::SetSystemFixedBoundsOverride(FBox SystemFixedBound
 	CachedSystemFixedBounds = SystemFixedBounds;
 }
 
+FBox FNiagaraEmitterInstance::GetFixedBounds() const
+{
+	if ( FixedBounds_GT.IsValid )
+	{
+		return FixedBounds_GT;
+	}
+	else if ( CachedEmitter && CachedEmitter->bFixedBounds)
+	{
+		return CachedEmitter->FixedBounds;
+	}
+	return FBox(EForceInit::ForceInit);
+}
+
 static int32 GbTriggerCrash = 0;
 static FAutoConsoleVariableRef CVarTriggerCrash(
 	TEXT("fx.TriggerDebugCrash"),
diff --git a/Engine/Plugins/FX/Niagara/Source/Niagara/Private/NiagaraSystemInstance.cpp b/Engine/Plugins/FX/Niagara/Source/Niagara/Private/NiagaraSystemInstance.cpp
--- a/Engine/Plugins/FX/Niagara/Source/Niagara/Private/NiagaraSystemInstance.cpp
+++ b/Engine/Plugins/FX/Niagara/Source/Niagara/Private/NiagaraSystemInstance.cpp
@@ -110,6 +110,10 @@ FNiagaraSystemInstance::FNiagaraSystemInstance(UWorld& InWorld, UNiagaraSystem&
 	, bLODDistanceIsOverridden(false)
 	, bPooled(bInPooled)
 	, CachedDeltaSeconds(0.0f)
+	, TimeSinceLastForceUpdateTransform(0.0f)
+	, FixedBounds_GT(EForceInit::ForceInit)
+	, FixedBounds_CNC(EForceInit::ForceInit)
+	, LocalBounds(EForceInit::ForceInit)
 	, RequestedExecutionState(ENiagaraExecutionState::Complete)
 	, ActualExecutionState(ENiagaraExecutionState::Complete)
 	, FeatureLevel(GMaxRHIFeatureLevel)
@@ -704,6 +708,10 @@ void FNiagaraSystemInstance::Complete(bool bExternalCompletion)
 void FNiagaraSystemInstance::OnPooledReuse(UWorld& NewWorld)
 {
 	World = &NewWorld;
+
+	FixedBounds_GT.Init();
+	FixedBounds_CNC.Init();
+
 	for (auto&& Emitter : Emitters)
 	{
 		Emitter->OnPooledReuse();
@@ -2197,6 +2205,7 @@ void FNiagaraSystemInstance::Tick_GameThread(float DeltaSeconds)
 	}
 
 	CachedDeltaSeconds = DeltaSeconds;
+	FixedBounds_CNC = FixedBounds_GT;
 
 	TickInstanceParameters_GameThread(DeltaSeconds);
 
@@ -2311,7 +2320,11 @@ void FNiagaraSystemInstance::Tick_Concurrent(bool bEnqueueGPUTickIfNeeded)
 	}
 
 	// Update local bounds
-	if ( System->bFixedBounds )
+	if ( FixedBounds_CNC.IsValid )
+	{
+		LocalBounds = FixedBounds_CNC;
+	}
+	else if ( System->bFixedBounds )
 	{
 		LocalBounds = System->GetFixedBounds();
 	}
@@ -2516,6 +2529,53 @@ TConstArrayView<FNiagaraEmitterExecutionIndex> FNiagaraSystemInstance::GetEmitte
 	return MakeArrayView<FNiagaraEmitterExecutionIndex>(nullptr, 0);
 }
 
+FBox FNiagaraSystemInstance::GetSystemFixedBounds() const
+{
+	if (FixedBounds_GT.IsValid)
+	{
+		return FixedBounds_GT;
+	}
+	else
+	{
+		const UNiagaraSystem* NiagaraSystem = GetSystem();
+		if (NiagaraSystem && NiagaraSystem->bFixedBounds)
+		{
+			return NiagaraSystem->GetFixedBounds();
+		}
+	}
+	return FBox(EForceInit::ForceInit);
+}
+
+void FNiagaraSystemInstance::SetEmitterFixedBounds(FName EmitterName, const FBox& InLocalBounds)
+{
+	for (TSharedRef<FNiagaraEmitterInstance, ESPMode::ThreadSafe>& Emitter : Emitters)
+	{
+		if ( Emitter->GetEmitterHandle().GetName() == EmitterName )
+		{
+			Emitter->SetFixedBounds(InLocalBounds);
+			return;
+		}
+	}
+
+	// Failed to find emitter
+	UE_LOG(LogNiagara, Warning, TEXT("SetEmitterFixedBounds: Failed to find Emitter(%s) System(%s) Component(%s)"), *EmitterName.ToString(), *GetNameSafe(Asset.Get()), *GetFullNameSafe(AttachComponent.Get()));
+}
+
+FBox FNiagaraSystemInstance::GetEmitterFixedBounds(FName EmitterName) const
+{
+	for (const TSharedRef<FNiagaraEmitterInstance, ESPMode::ThreadSafe>& Emitter : Emitters)
+	{
+		if (Emitter->GetEmitterHandle().GetName() == EmitterName)
+		{
+			return Emitter->GetFixedBounds();
+		}
+	}
+
+	// Failed to find emitter
+	UE_LOG(LogNiagara, Warning, TEXT("GetEmitterFixedBounds: Failed to find Emitter(%s) System(%s) Component(%s)"), *EmitterName.ToString(), *GetNameSafe(Asset.Get()), *GetFullNameSafe(AttachComponent.Get()));
+	return FBox(EForceInit::ForceInit);
+}
+
 FNiagaraEmitterInstance* FNiagaraSystemInstance::GetEmitterByID(FGuid InID)
 {
 	for (TSharedRef<FNiagaraEmitterInstance, ESPMode::ThreadSafe>& Emitter : Emitters)
diff --git a/Engine/Plugins/FX/Niagara/Source/Niagara/Public/NiagaraComponent.h b/Engine/Plugins/FX/Niagara/Source/Niagara/Public/NiagaraComponent.h
--- a/Engine/Plugins/FX/Niagara/Source/Niagara/Public/NiagaraComponent.h
+++ b/Engine/Plugins/FX/Niagara/Source/Niagara/Public/NiagaraComponent.h
@@ -333,6 +333,46 @@ public:
 	UFUNCTION(BlueprintCallable, Category = Niagara, meta = (DisplayName = "Get Tick Behavior"))
 	ENiagaraTickBehavior GetTickBehavior() const { return TickBehavior; }
 
+	/**
+	Sets the fixed bounds for the system instance, this overrides all other bounds.
+	The box is expected to be in local space not world space.
+	A default uninitialized box will clear the fixed bounds and revert back to system fixed / dynamic bounds.
+	*/
+	UFUNCTION(BlueprintCallable, Category = Niagara)
+	void SetSystemFixedBounds(FBox LocalBounds);
+
+	/**
+	Gets the fixed bounds for the system instance.
+	This will return the user set fixed bounds if set, or the systems fixed bounds if set.
+	Note: The returned box may be invalid if no fixed bounds exist
+	*/
+	UFUNCTION(BlueprintCallable, Category = Niagara)
+	FBox GetSystemFixedBounds() const;
+
+	/** Clear any previously set fixed bounds for the system instance. */
+	UFUNCTION(BlueprintCallable, Category = Niagara)
+	void ClearSystemFixedBounds();
+
+	/**
+	Sets the fixed bounds for an emitter instance, this overrides all other bounds.
+	The box is expected to be in local space not world space.
+	A default uninitialized box will clear the fixed bounds and revert back to emitter fixed / dynamic bounds.
+	*/
+	UFUNCTION(BlueprintCallable, Category = Niagara)
+	void SetEmitterFixedBounds(FName EmitterName, FBox LocalBounds);
+
+	/**
+	Gets the fixed bounds for an emitter instance.
+	This will return the user set fixed bounds if set, or the emitters fixed bounds if set.
+	Note: The returned box may be invalid if no fixed bounds exist
+	*/
+	UFUNCTION(BlueprintCallable, Category = Niagara)
+	FBox GetEmitterFixedBounds(FName EmitterName) const;
+
+	/** Clear any previously set fixed bounds for the emitter instance. */
+	UFUNCTION(BlueprintCallable, Category = Niagara)
+	void ClearEmitterFixedBounds(FName EmitterName);
+
 	UFUNCTION(BlueprintCallable, Category = Niagara, meta = (DisplayName = "Set Random Seed Offset"))
 	void SetRandomSeedOffset(int32 NewRandomSeedOffset);
 
diff --git a/Engine/Plugins/FX/Niagara/Source/Niagara/Public/NiagaraSystemInstance.h b/Engine/Plugins/FX/Niagara/Source/Niagara/Public/NiagaraSystemInstance.h
--- a/Engine/Plugins/FX/Niagara/Source/Niagara/Public/NiagaraSystemInstance.h
+++ b/Engine/Plugins/FX/Niagara/Source/Niagara/Public/NiagaraSystemInstance.h
@@ -233,6 +233,11 @@ public:
 	FORCEINLINE const FBox& GetLocalBounds() const { return LocalBounds;  }
 	TConstArrayView<FNiagaraEmitterExecutionIndex> GetEmitterExecutionOrder() const;
 
+	FORCEINLINE void SetSystemFixedBounds(const FBox& InLocalBounds) { FixedBounds_GT = InLocalBounds; }
+	FBox GetSystemFixedBounds() const;
+	void SetEmitterFixedBounds(FName EmitterName, const FBox& InLocalBounds);
+	FBox GetEmitterFixedBounds(FName EmitterName) const;
+
 	FNiagaraEmitterInstance* GetEmitterByID(FGuid InID);
 
 	FORCEINLINE bool IsSolo() const { return bSolo; }
@@ -536,6 +541,10 @@ private:
 	/** Time since we last forced a bounds update. */
 	float TimeSinceLastForceUpdateTransform;
 
+	/** Optional user specified bounds. */
+	FBox FixedBounds_GT;
+	FBox FixedBounds_CNC;
+
 	/** Current calculated local bounds. */
 	FBox LocalBounds;
 
diff --git a/Engine/Plugins/FX/Niagara/Source/Niagara/Public/NiagaraSystemInstanceController.h b/Engine/Plugins/FX/Niagara/Source/Niagara/Public/NiagaraSystemInstanceController.h
--- a/Engine/Plugins/FX/Niagara/Source/Niagara/Public/NiagaraSystemInstanceController.h
+++ b/Engine/Plugins/FX/Niagara/Source/Niagara/Public/NiagaraSystemInstanceController.h
@@ -128,6 +128,10 @@ public:
 	NIAGARA_SYSTEM_INSTANCE_CONTROLLER_SHIM(Deactivate,)
 	NIAGARA_SYSTEM_INSTANCE_CONTROLLER_SHIM(CalculateTickGroup, const)
 	NIAGARA_SYSTEM_INSTANCE_CONTROLLER_SHIM(GetLocalBounds, const)
+	NIAGARA_SYSTEM_INSTANCE_CONTROLLER_SHIM(SetSystemFixedBounds,)
+	NIAGARA_SYSTEM_INSTANCE_CONTROLLER_SHIM(GetSystemFixedBounds, const)
+	NIAGARA_SYSTEM_INSTANCE_CONTROLLER_SHIM(SetEmitterFixedBounds,)
+	NIAGARA_SYSTEM_INSTANCE_CONTROLLER_SHIM(GetEmitterFixedBounds, const)
 	NIAGARA_SYSTEM_INSTANCE_CONTROLLER_SHIM(WaitForConcurrentTickAndFinalize,)
 	NIAGARA_SYSTEM_INSTANCE_CONTROLLER_SHIM(OnSimulationDestroyed,)
 	NIAGARA_SYSTEM_INSTANCE_CONTROLLER_SHIM(GetAttachComponent,)
