Expose Max GPU Particles per frame on emitter, remove corresponding console var
Useful for fluids when we emit a lot of particles on the first frame, potentially rejection sampling many of them

diff --git a/Engine/Plugins/FX/Niagara/Source/Niagara/Classes/NiagaraEmitter.h b/Engine/Plugins/FX/Niagara/Source/Niagara/Classes/NiagaraEmitter.h
--- a/Engine/Plugins/FX/Niagara/Source/Niagara/Classes/NiagaraEmitter.h
+++ b/Engine/Plugins/FX/Niagara/Source/Niagara/Classes/NiagaraEmitter.h
@@ -386,6 +386,10 @@ public:
 	UPROPERTY(EditAnywhere, AdvancedDisplay, Category = "Emitter", meta = (InlineEditConditionToggle))
 	uint32 bLimitDeltaTime : 1;
 
+	/** An override on the max number of GPU particles we expect to spawn in a single frame. A value of 0 means it'll use fx.MaxNiagaraGPUParticlesSpawnPerFrame.*/
+	UPROPERTY(EditAnywhere, AdvancedDisplay, Category = "Emitter", meta = (EditCondition = "SimTarget == ENiagaraSimTarget::GPUComputeSim", DisplayName = "Max GPU Particles Spawn per Frame"))
+	int32 MaxGPUParticlesSpawnPerFrame;
+
 	void NIAGARA_API GetScripts(TArray<UNiagaraScript*>& OutScripts, bool bCompilableOnly = true, bool bEnabledOnly = false) const;
 
 	NIAGARA_API UNiagaraScript* GetScript(ENiagaraScriptUsage Usage, FGuid UsageId);
diff --git a/Engine/Plugins/FX/Niagara/Source/Niagara/Private/NiagaraEmitter.cpp b/Engine/Plugins/FX/Niagara/Source/Niagara/Private/NiagaraEmitter.cpp
--- a/Engine/Plugins/FX/Niagara/Source/Niagara/Private/NiagaraEmitter.cpp
+++ b/Engine/Plugins/FX/Niagara/Source/Niagara/Private/NiagaraEmitter.cpp
@@ -134,6 +134,7 @@ UNiagaraEmitter::UNiagaraEmitter(const FObjectInitializer& Initializer)
 , bCombineEventSpawn(false)
 , MaxDeltaTimePerTick(0.125)
 , bLimitDeltaTime(true)
+, MaxGPUParticlesSpawnPerFrame(0)
 #if WITH_EDITORONLY_DATA
 , bBakeOutRapidIteration(true)
 , ThumbnailImageOutOfDate(true)
diff --git a/Engine/Plugins/FX/Niagara/Source/Niagara/Private/NiagaraEmitterInstance.cpp b/Engine/Plugins/FX/Niagara/Source/Niagara/Private/NiagaraEmitterInstance.cpp
--- a/Engine/Plugins/FX/Niagara/Source/Niagara/Private/NiagaraEmitterInstance.cpp
+++ b/Engine/Plugins/FX/Niagara/Source/Niagara/Private/NiagaraEmitterInstance.cpp
@@ -1458,9 +1458,12 @@ void FNiagaraEmitterInstance::Tick(float DeltaSeconds)
 						// spawn. Therefore, we'll send the spawn requests to the render thread as if there was no limit, and we'll clamp the values there, when we prepare
 						// the destination dataset for simulation.
 						NumSpawnedOnGPUThisFrame += Info.Count;
-						if (NumSpawnedOnGPUThisFrame > GMaxNiagaraGPUParticlesSpawnPerFrame)
+
+						int32 MaxParticlesSpawnedPerFrame = CachedEmitter->MaxGPUParticlesSpawnPerFrame <= 0 ? GMaxNiagaraGPUParticlesSpawnPerFrame : CachedEmitter->MaxGPUParticlesSpawnPerFrame;
+
+						if (NumSpawnedOnGPUThisFrame > MaxParticlesSpawnedPerFrame)
 						{
-							FString DebugMsg = FString::Printf(TEXT("%s has attempted to execeed max GPU per frame spawn! | Max: %d | Requested: %d | SpawnInfoEntry: %d"), *CachedEmitter->GetFullName(), GMaxNiagaraGPUParticlesSpawnPerFrame, NumSpawnedOnGPUThisFrame, SpawnInfoIdx);
+							FString DebugMsg = FString::Printf(TEXT("%s has attempted to execeed max GPU per frame spawn! | Max: %d | Requested: %d | SpawnInfoEntry: %d"), *CachedEmitter->GetFullName(), MaxParticlesSpawnedPerFrame, NumSpawnedOnGPUThisFrame, SpawnInfoIdx);
 							UE_LOG(LogNiagara, Warning, TEXT("%s"), *DebugMsg);
 							GEngine->AddOnScreenDebugMessage(-1, 5.f, FColor::Yellow, *DebugMsg);
 							break;
