Fix for systems getting stuck in an inactive state requiring a full rebuild
-generating the CompilationId during UNiagaraScript::PostLoad() requires the renderer properties associated with the Emitter to be fully postloaded to ensure that their bindings are up to date

diff --git a/Engine/Plugins/FX/Niagara/Source/Niagara/Private/NiagaraScript.cpp b/Engine/Plugins/FX/Niagara/Source/Niagara/Private/NiagaraScript.cpp
--- a/Engine/Plugins/FX/Niagara/Source/Niagara/Private/NiagaraScript.cpp
+++ b/Engine/Plugins/FX/Niagara/Source/Niagara/Private/NiagaraScript.cpp
@@ -1111,6 +1111,10 @@ void UNiagaraScript::ComputeVMCompilationId_EmitterShared(FNiagaraVMExecutableDa
 	// Gather additional variables from renderers
 	for (UNiagaraRendererProperties* RendererProperty : Emitter->GetRenderers())
 	{
+		// Because we currently generate the CompilationId during PostLoad (as a part of AreScriptAndSourceSynchronized)
+		// we need to ensure that the RendererProperties have done their PostLoad.
+		RendererProperty->ConditionalPostLoad();
+
 		if (RendererProperty->GetCurrentSourceMode() != InSourceMode)
 			continue;
 
