fixed global distance field sampling

diff --git a/Engine/Plugins/FX/Niagara/Source/Niagara/Classes/NiagaraEmitterInstanceBatcher.h b/Engine/Plugins/FX/Niagara/Source/Niagara/Classes/NiagaraEmitterInstanceBatcher.h
--- a/Engine/Plugins/FX/Niagara/Source/Niagara/Classes/NiagaraEmitterInstanceBatcher.h
+++ b/Engine/Plugins/FX/Niagara/Source/Niagara/Classes/NiagaraEmitterInstanceBatcher.h
@@ -114,7 +114,7 @@ public:
 	 */
 	bool AddSortedGPUSimulation(FNiagaraGPUSortInfo& SortInfo);
 
-	const FGlobalDistanceFieldParameterData& GetGlobalDistanceFieldParameters() const { return GlobalDistanceFieldParams; }
+	const FGlobalDistanceFieldParameterData* GetGlobalDistanceFieldParameters() const;
 
 	void SetDataInterfaceParameters(FRHICommandList& RHICmdList, const FNiagaraGPUSystemTick& Tick, const FNiagaraComputeInstanceData& InstanceData, const FNiagaraShaderRef& ComputeShader, const FNiagaraSimStageData& SimStageData) const;
 	void UnsetDataInterfaceParameters(FRHICommandList& RHICmdList, const FNiagaraGPUSystemTick& Tick, const FNiagaraComputeInstanceData& InstanceData, const FNiagaraShaderRef& ComputeShader, const FNiagaraSimStageData& SimStageData) const;
@@ -242,7 +242,7 @@ private:
 
 	uint32 FramesBeforeTickFlush = 0;
 
-	FGlobalDistanceFieldParameterData GlobalDistanceFieldParams;
+	TConstArrayView<FViewInfo> CurrentPassViews;
 
 	/** A buffer of list sizes used by UpdateFreeIDBuffers to allow overlapping several dispatches. */
 	FRWBuffer FreeIDListSizesBuffer;
diff --git a/Engine/Plugins/FX/Niagara/Source/Niagara/Private/NiagaraDataInterfaceCollisionQuery.cpp b/Engine/Plugins/FX/Niagara/Source/Niagara/Private/NiagaraDataInterfaceCollisionQuery.cpp
--- a/Engine/Plugins/FX/Niagara/Source/Niagara/Private/NiagaraDataInterfaceCollisionQuery.cpp
+++ b/Engine/Plugins/FX/Niagara/Source/Niagara/Private/NiagaraDataInterfaceCollisionQuery.cpp
@@ -1028,7 +1028,17 @@ public:
 		if (GlobalDistanceFieldParameters.IsBound())
 		{
 			check(Context.Batcher);
-			GlobalDistanceFieldParameters.Set(RHICmdList, ComputeShaderRHI, Context.Batcher->GetGlobalDistanceFieldParameters());
+			const FGlobalDistanceFieldParameterData* GlobalDistanceFieldParameterData = Context.Batcher->GetGlobalDistanceFieldParameters();
+
+			if (GlobalDistanceFieldParameterData)
+			{
+				GlobalDistanceFieldParameters.Set(RHICmdList, ComputeShaderRHI, *GlobalDistanceFieldParameterData);
+			}
+			else
+			{
+				GlobalDistanceFieldParameters.Set(RHICmdList, ComputeShaderRHI, FGlobalDistanceFieldParameterData());
+			}
+			
 		}
 
 #if RHI_RAYTRACING
diff --git a/Engine/Plugins/FX/Niagara/Source/Niagara/Private/NiagaraEmitterInstanceBatcher.cpp b/Engine/Plugins/FX/Niagara/Source/Niagara/Private/NiagaraEmitterInstanceBatcher.cpp
--- a/Engine/Plugins/FX/Niagara/Source/Niagara/Private/NiagaraEmitterInstanceBatcher.cpp
+++ b/Engine/Plugins/FX/Niagara/Source/Niagara/Private/NiagaraEmitterInstanceBatcher.cpp
@@ -1711,6 +1711,8 @@ void NiagaraEmitterInstanceBatcher::PostRenderOpaque(FRDGBuilder& GraphBuilder,
 			};
 #endif
 
+			CurrentPassViews = Views;
+
 			FRHIUniformBuffer* ViewUniformBuffer = GetReferenceViewUniformBuffer(Views);
 
 			// Setup new readback since if there is no pending request, there is no risk of having invalid data read (offset being allocated after the readback was sent).
@@ -1718,6 +1720,9 @@ void NiagaraEmitterInstanceBatcher::PostRenderOpaque(FRDGBuilder& GraphBuilder,
 
 			FinishDispatches();
 
+			// Clear CurrentPassViews
+			CurrentPassViews = TConstArrayView<FViewInfo>();
+
 			ProcessDebugReadbacks(RHICmdList, false);
 		}
 
@@ -1822,9 +1827,6 @@ void NiagaraEmitterInstanceBatcher::PreRender(FRDGBuilder& GraphBuilder, TConstA
 	}
 
 	LLM_SCOPE(ELLMTag::Niagara);
-
-	const FGlobalDistanceFieldParameterData* ViewGlobalDistanceFieldParams = GetReferenceGlobalDistanceFieldData(Views);
-	GlobalDistanceFieldParams = ViewGlobalDistanceFieldParams ? *ViewGlobalDistanceFieldParams : FGlobalDistanceFieldParameterData();
 }
 
 void NiagaraEmitterInstanceBatcher::OnDestroy()
@@ -1848,6 +1850,12 @@ bool NiagaraEmitterInstanceBatcher::AddSortedGPUSimulation(FNiagaraGPUSortInfo&
 	}
 }
 
+const FGlobalDistanceFieldParameterData* NiagaraEmitterInstanceBatcher::GetGlobalDistanceFieldParameters() const
+{ 
+	check(CurrentPassViews.Num() > 0); 
+	return &CurrentPassViews[0].GlobalDistanceFieldInfo.ParameterData; 
+}
+
 void NiagaraEmitterInstanceBatcher::GenerateSortKeys(FRHICommandListImmediate& RHICmdList, int32 BatchId, int32 NumElementsInBatch, EGPUSortFlags Flags, FRHIUnorderedAccessView* KeysUAV, FRHIUnorderedAccessView* ValuesUAV)
 {
 	const FGPUSortManager::FKeyGenInfo KeyGenInfo((uint32)NumElementsInBatch, EnumHasAnyFlags(Flags, EGPUSortFlags::HighPrecisionKeys));
