Conform EWorldType string conversion methods to LexToString standard

diff --git a/Engine/Plugins/FX/Niagara/Source/NiagaraEditor/Private/Customizations/NiagaraOutlinerCustomization.cpp b/Engine/Plugins/FX/Niagara/Source/NiagaraEditor/Private/Customizations/NiagaraOutlinerCustomization.cpp
--- a/Engine/Plugins/FX/Niagara/Source/NiagaraEditor/Private/Customizations/NiagaraOutlinerCustomization.cpp
+++ b/Engine/Plugins/FX/Niagara/Source/NiagaraEditor/Private/Customizations/NiagaraOutlinerCustomization.cpp
@@ -891,7 +891,7 @@ TSharedRef<SWidget> FNiagaraOutlinerTreeWorldItem::GetHeaderWidget()
 				[
 					SNew(SNiagaraOutlinerTreeItemHeaderDataWidget<FText>, Outliner->ViewSettings)
 					.ToolTipText(LOCTEXT("WorldHeaderTooltip_WorldType", "World Type"))
-					.Data(FText::FromString(ToString((EWorldType::Type)Data->WorldType)))
+					.Data(FText::FromString(LexToString((EWorldType::Type)Data->WorldType)))
 				];
 
 			Box->AddSlot()
diff --git a/Engine/Source/Runtime/Engine/Classes/Engine/EngineTypes.h b/Engine/Source/Runtime/Engine/Classes/Engine/EngineTypes.h
--- a/Engine/Source/Runtime/Engine/Classes/Engine/EngineTypes.h
+++ b/Engine/Source/Runtime/Engine/Classes/Engine/EngineTypes.h
@@ -1046,7 +1046,7 @@ namespace EWorldType
 	};
 }
 
-const TCHAR* LexToString(const EWorldType::Type Value);
+ENGINE_API const TCHAR* LexToString(const EWorldType::Type Value);
 
 /** Describes what parts of level streaming should be forcibly handled immediately */
 enum class EFlushLevelStreamingType : uint8
diff --git a/Engine/Source/Runtime/Engine/Classes/Engine/World.h b/Engine/Source/Runtime/Engine/Classes/Engine/World.h
--- a/Engine/Source/Runtime/Engine/Classes/Engine/World.h
+++ b/Engine/Source/Runtime/Engine/Classes/Engine/World.h
@@ -4109,5 +4109,6 @@ FORCEINLINE_DEBUGGABLE bool UWorld::IsNetMode(ENetMode Mode) const
 #endif
 }
 
+UE_DEPRECATED(5.0, "Please use LexToString(EWorldType::Type Type) instead")
 FString ENGINE_API ToString(EWorldType::Type Type);
 FString ENGINE_API ToString(ENetMode NetMode);
