- Random seed added to system instances, allows for randomness in system scripts

diff --git a/Engine/Plugins/FX/Niagara/Source/Niagara/Classes/NiagaraConstants.h b/Engine/Plugins/FX/Niagara/Source/Niagara/Classes/NiagaraConstants.h
--- a/Engine/Plugins/FX/Niagara/Source/Niagara/Classes/NiagaraConstants.h
+++ b/Engine/Plugins/FX/Niagara/Source/Niagara/Classes/NiagaraConstants.h
@@ -68,6 +68,7 @@
 #define SYS_PARAM_ENGINE_EMITTER_INSTANCE_SEED           INiagaraModule::GetVar_Engine_Emitter_InstanceSeed()
 #define SYS_PARAM_ENGINE_SYSTEM_NUM_EMITTERS_ALIVE       INiagaraModule::GetVar_Engine_System_NumEmittersAlive()
 #define SYS_PARAM_ENGINE_SYSTEM_SIGNIFICANCE_INDEX		 INiagaraModule::GetVar_Engine_System_SignificanceIndex()
+#define SYS_PARAM_ENGINE_SYSTEM_RANDOM_SEED				INiagaraModule::GetVar_Engine_System_RandomSeed()
 #define SYS_PARAM_ENGINE_SYSTEM_NUM_EMITTERS             INiagaraModule::GetVar_Engine_System_NumEmitters()
 #define SYS_PARAM_ENGINE_NUM_SYSTEM_INSTANCES            INiagaraModule::GetVar_Engine_NumSystemInstances()
 
diff --git a/Engine/Plugins/FX/Niagara/Source/Niagara/Classes/NiagaraSystem.h b/Engine/Plugins/FX/Niagara/Source/Niagara/Classes/NiagaraSystem.h
--- a/Engine/Plugins/FX/Niagara/Source/Niagara/Classes/NiagaraSystem.h
+++ b/Engine/Plugins/FX/Niagara/Source/Niagara/Classes/NiagaraSystem.h
@@ -419,6 +419,9 @@ public:
 	FORCEINLINE float GetWarmupTickDelta()const { return WarmupTickDelta; }
 	virtual void GetAssetRegistryTags(TArray<FAssetRegistryTag>& OutTags)  const override;
 
+	FORCEINLINE bool NeedsDeterminism() const { return bDeterminism; }
+	FORCEINLINE int32 GetRandomSeed() const { return RandomSeed; }
+
 #if STATS
 	FNiagaraStatDatabase& GetStatData() { return StatDatabase; }
 #endif
@@ -768,6 +771,17 @@ protected:
 	UPROPERTY(EditAnywhere, Category = Performance, meta = (ToolTip = "Auto-deactivate system if all emitters are determined to not spawn particles again, regardless of lifetime."))
 	bool bAutoDeactivate;
 
+	/**
+	When disabled we will generate a RandomSeed per instance on reset which is not deterministic.
+	When enabled we will always use the RandomSeed from the system plus the components RandomSeedOffset, this allows for determinism but variance between components.
+	*/
+	UPROPERTY(EditAnywhere, Category = "Random")
+	bool bDeterminism = true;
+
+	/** Seed used for system script random number generator. */
+	UPROPERTY(EditAnywhere, Category = "Random", meta = (EditCondition = "bDeterminism"))
+	int32 RandomSeed = 0;
+
 	/** Warm up time in seconds. Used to calculate WarmupTickCount. Rounds down to the nearest multiple of WarmupTickDelta. */
 	UPROPERTY(EditAnywhere, Category = Warmup)
 	float WarmupTime;
diff --git a/Engine/Plugins/FX/Niagara/Source/Niagara/Private/NiagaraConstants.cpp b/Engine/Plugins/FX/Niagara/Source/Niagara/Private/NiagaraConstants.cpp
--- a/Engine/Plugins/FX/Niagara/Source/Niagara/Private/NiagaraConstants.cpp
+++ b/Engine/Plugins/FX/Niagara/Source/Niagara/Private/NiagaraConstants.cpp
@@ -95,6 +95,7 @@ void FNiagaraConstants::Init()
 		SystemParameters.Add(SYS_PARAM_ENGINE_EMITTER_SPAWN_COUNT_SCALE);
 		SystemParameters.Add(SYS_PARAM_ENGINE_SYSTEM_NUM_EMITTERS_ALIVE);
 		SystemParameters.Add(SYS_PARAM_ENGINE_SYSTEM_SIGNIFICANCE_INDEX);
+		SystemParameters.Add(SYS_PARAM_ENGINE_SYSTEM_RANDOM_SEED);
 		SystemParameters.Add(SYS_PARAM_ENGINE_SYSTEM_NUM_EMITTERS);
 		SystemParameters.Add(SYS_PARAM_ENGINE_NUM_SYSTEM_INSTANCES);
 		SystemParameters.Add(SYS_PARAM_ENGINE_GLOBAL_SPAWN_COUNT_SCALE);
@@ -215,6 +216,7 @@ void FNiagaraConstants::Init()
 		SystemStrMap.Add(SYS_PARAM_ENGINE_EMITTER_INSTANCE_SEED, LOCTEXT("EmitterInstanceSeed", "A random seed that changes with every execution of the emitter instance."));
 		SystemStrMap.Add(SYS_PARAM_ENGINE_SYSTEM_NUM_EMITTERS_ALIVE, LOCTEXT("SystemNumEmittersAlive", "The number of emitters still alive attached to this system. Should only be used in System scripts."));
 		SystemStrMap.Add(SYS_PARAM_ENGINE_SYSTEM_SIGNIFICANCE_INDEX, LOCTEXT("SystemSignificanceIndex", "Index denoting how significant this system instance is in relation to others of the same system in this scene. e.g. 0 is the most significanct instance."));
+		SystemStrMap.Add(SYS_PARAM_ENGINE_SYSTEM_RANDOM_SEED, LOCTEXT("SystemRandomSeed", "A random seed controlled used for generating system random numbers."));
 		
 		SystemStrMap.Add(SYS_PARAM_ENGINE_SYSTEM_NUM_EMITTERS, LOCTEXT("SystemNumEmitters", "The number of emitters attached to this system. Should only be used in System scripts."));
 		SystemStrMap.Add(SYS_PARAM_ENGINE_NUM_SYSTEM_INSTANCES, LOCTEXT("SystemNumInstances", "The number of instances of the this system currently ticking. Should only be used in System scripts."));
diff --git a/Engine/Plugins/FX/Niagara/Source/Niagara/Private/NiagaraModule.cpp b/Engine/Plugins/FX/Niagara/Source/Niagara/Private/NiagaraModule.cpp
--- a/Engine/Plugins/FX/Niagara/Source/Niagara/Private/NiagaraModule.cpp
+++ b/Engine/Plugins/FX/Niagara/Source/Niagara/Private/NiagaraModule.cpp
@@ -122,6 +122,7 @@ FNiagaraVariable INiagaraModule::Engine_Emitter_SpawnCountScale;
 FNiagaraVariable INiagaraModule::Engine_System_TickCount;
 FNiagaraVariable INiagaraModule::Engine_System_NumEmittersAlive;
 FNiagaraVariable INiagaraModule::Engine_System_SignificanceIndex;
+FNiagaraVariable INiagaraModule::Engine_System_RandomSeed;
 FNiagaraVariable INiagaraModule::Engine_System_NumEmitters;
 FNiagaraVariable INiagaraModule::Engine_NumSystemInstances;
 
@@ -248,6 +249,7 @@ void INiagaraModule::StartupModule()
 	Engine_System_TickCount = FNiagaraVariable(FNiagaraTypeDefinition::GetIntDef(), TEXT("Engine.System.TickCount"));
 	Engine_System_NumEmittersAlive = FNiagaraVariable(FNiagaraTypeDefinition::GetIntDef(), TEXT("Engine.System.NumEmittersAlive"));
 	Engine_System_SignificanceIndex = FNiagaraVariable(FNiagaraTypeDefinition::GetIntDef(), TEXT("Engine.System.SignificanceIndex"));
+	Engine_System_RandomSeed = FNiagaraVariable(FNiagaraTypeDefinition::GetIntDef(), TEXT("Engine.System.RandomSeed"));
 	Engine_System_NumEmitters = FNiagaraVariable(FNiagaraTypeDefinition::GetIntDef(), TEXT("Engine.System.NumEmitters"));
 	Engine_NumSystemInstances = FNiagaraVariable(FNiagaraTypeDefinition::GetIntDef(), TEXT("Engine.NumSystemInstances"));
 
@@ -1375,7 +1377,6 @@ const TArray<FNiagaraVariable>& FNiagaraSystemParameters::GetVariables()
 {
 	static const FName NAME_NiagaraStructPadding0 = "Engine.System.PaddingInt32_0";
 	static const FName NAME_NiagaraStructPadding1 = "Engine.System.PaddingInt32_1";
-	static const FName NAME_NiagaraStructPadding2 = "Engine.System.PaddingInt32_2";
 
 	static const TArray<FNiagaraVariable> Variables =
 	{
@@ -1388,9 +1389,9 @@ const TArray<FNiagaraVariable>& FNiagaraSystemParameters::GetVariables()
 		SYS_PARAM_ENGINE_SYSTEM_NUM_EMITTERS,
 		SYS_PARAM_ENGINE_SYSTEM_NUM_EMITTERS_ALIVE,
 		SYS_PARAM_ENGINE_SYSTEM_SIGNIFICANCE_INDEX,
+		SYS_PARAM_ENGINE_SYSTEM_RANDOM_SEED,
 		FNiagaraVariable(FNiagaraTypeDefinition::GetIntDef(), NAME_NiagaraStructPadding0),
 		FNiagaraVariable(FNiagaraTypeDefinition::GetIntDef(), NAME_NiagaraStructPadding1),
-		FNiagaraVariable(FNiagaraTypeDefinition::GetIntDef(), NAME_NiagaraStructPadding2),
 	};
 
 	return Variables;
diff --git a/Engine/Plugins/FX/Niagara/Source/Niagara/Private/NiagaraSystemInstance.cpp b/Engine/Plugins/FX/Niagara/Source/Niagara/Private/NiagaraSystemInstance.cpp
--- a/Engine/Plugins/FX/Niagara/Source/Niagara/Private/NiagaraSystemInstance.cpp
+++ b/Engine/Plugins/FX/Niagara/Source/Niagara/Private/NiagaraSystemInstance.cpp
@@ -96,6 +96,7 @@ FNiagaraSystemInstance::FNiagaraSystemInstance(UWorld& InWorld, UNiagaraSystem&
 	, Age(0.0f)
 	, LastRenderTime(0.0f)
 	, TickCount(0)
+	, RandomSeed(0)
 	, RandomSeedOffset(0)
 	, LODDistance(0.0f)
 	, MaxLODDistance(FLT_MAX)
@@ -852,6 +853,10 @@ void FNiagaraSystemInstance::Reset(FNiagaraSystemInstance::EResetMode Mode)
 				SystemGpuComputeProxy->AddToBatcher(GetBatcher());
 			}
 
+			// Create new random seed
+			RandomSeed = GetSystem()->NeedsDeterminism() ? GetSystem()->GetRandomSeed() : FMath::Rand();
+
+			// Add instance to simulation
 			SystemSimulation->AddInstance(this);
 
 			UNiagaraSystem* System = GetSystem();
@@ -1898,6 +1903,7 @@ void FNiagaraSystemInstance::TickInstanceParameters_GameThread(float DeltaSecond
 	CurrentSystemParameters.EngineLodDistance = GetLODDistance();
 	CurrentSystemParameters.EngineLodDistanceFraction = CurrentSystemParameters.EngineLodDistance / MaxLODDistance;
 	CurrentSystemParameters.SignificanceIndex = SignificanceIndex;
+	CurrentSystemParameters.RandomSeed = RandomSeed + RandomSeedOffset;
 
 	if (OverrideParameters)
 	{
@@ -1937,6 +1943,7 @@ void FNiagaraSystemInstance::TickInstanceParameters_Concurrent()
 	CurrentSystemParameters.EngineEmitterCount = GatheredInstanceParameters.EmitterCount;
 	CurrentSystemParameters.EngineAliveEmitterCount = GatheredInstanceParameters.NumAlive;
 	CurrentSystemParameters.SignificanceIndex = SignificanceIndex;
+	CurrentSystemParameters.RandomSeed = RandomSeed + RandomSeedOffset;
 
 	FNiagaraGlobalParameters& CurrentGlobalParameter = GlobalParameters[ParameterIndex];
 	CurrentGlobalParameter.EngineDeltaTime = GatheredInstanceParameters.DeltaSeconds;
diff --git a/Engine/Plugins/FX/Niagara/Source/Niagara/Public/NiagaraModule.h b/Engine/Plugins/FX/Niagara/Source/Niagara/Public/NiagaraModule.h
--- a/Engine/Plugins/FX/Niagara/Source/Niagara/Public/NiagaraModule.h
+++ b/Engine/Plugins/FX/Niagara/Source/Niagara/Public/NiagaraModule.h
@@ -141,6 +141,7 @@ public:
 	FORCEINLINE static const FNiagaraVariable&  GetVar_Engine_System_TickCount() { return Engine_System_TickCount; }
 	FORCEINLINE static const FNiagaraVariable&  GetVar_Engine_System_NumEmittersAlive() { return Engine_System_NumEmittersAlive; }
 	FORCEINLINE static const FNiagaraVariable&  GetVar_Engine_System_SignificanceIndex() { return Engine_System_SignificanceIndex; }
+	FORCEINLINE static const FNiagaraVariable&  GetVar_Engine_System_RandomSeed() { return Engine_System_RandomSeed; }
 	FORCEINLINE static const FNiagaraVariable&  GetVar_Engine_System_NumEmitters() { return Engine_System_NumEmitters; }
 	FORCEINLINE static const FNiagaraVariable&  GetVar_Engine_NumSystemInstances() { return Engine_NumSystemInstances; }
 
@@ -258,6 +259,7 @@ private:
 	static FNiagaraVariable Engine_System_TickCount;
 	static FNiagaraVariable Engine_System_NumEmittersAlive;
 	static FNiagaraVariable Engine_System_SignificanceIndex;
+	static FNiagaraVariable Engine_System_RandomSeed;
 	static FNiagaraVariable Engine_System_NumEmitters;
 	static FNiagaraVariable Engine_NumSystemInstances;
 
diff --git a/Engine/Plugins/FX/Niagara/Source/Niagara/Public/NiagaraSystemInstance.h b/Engine/Plugins/FX/Niagara/Source/Niagara/Public/NiagaraSystemInstance.h
--- a/Engine/Plugins/FX/Niagara/Source/Niagara/Public/NiagaraSystemInstance.h
+++ b/Engine/Plugins/FX/Niagara/Source/Niagara/Public/NiagaraSystemInstance.h
@@ -441,6 +441,9 @@ private:
 	/** The tick count of the System instance. */
 	int32 TickCount;
 
+	/** Random seed used for system simulation random number generation. */
+	int32 RandomSeed;
+
 	/** A system-wide offset to permute the deterministic random seed (allows for variance among multiple instances while still being deterministic) */
 	int32 RandomSeedOffset;
 
diff --git a/Engine/Plugins/FX/Niagara/Source/Niagara/Public/NiagaraTypes.h b/Engine/Plugins/FX/Niagara/Source/Niagara/Public/NiagaraTypes.h
--- a/Engine/Plugins/FX/Niagara/Source/Niagara/Public/NiagaraTypes.h
+++ b/Engine/Plugins/FX/Niagara/Source/Niagara/Public/NiagaraTypes.h
@@ -1684,8 +1684,8 @@ struct alignas(16) FNiagaraSystemParameters
 	int32 EngineEmitterCount = 0;
 	int32 EngineAliveEmitterCount = 0;
 	int32 SignificanceIndex = 0;
+	int32 RandomSeed = 0;
 
-	int32 _Pad0;
 	int32 _Pad1;
 	int32 _Pad2;
 };
diff --git a/Engine/Plugins/FX/Niagara/Source/NiagaraCore/Private/NiagaraCustomVersion.cpp b/Engine/Plugins/FX/Niagara/Source/NiagaraCore/Private/NiagaraCustomVersion.cpp
--- a/Engine/Plugins/FX/Niagara/Source/NiagaraCore/Private/NiagaraCustomVersion.cpp
+++ b/Engine/Plugins/FX/Niagara/Source/NiagaraCore/Private/NiagaraCustomVersion.cpp
@@ -8,4 +8,5 @@ const FGuid FNiagaraCustomVersion::GUID(0xFCF57AFA, 0x50764283, 0xB9A9E658, 0xFF
 // Register the custom version with core
 FCustomVersionRegistration GRegisterNiagaraCustomVersion(FNiagaraCustomVersion::GUID, FNiagaraCustomVersion::LatestVersion, TEXT("NiagaraVer"));
 
-const FGuid FNiagaraCustomVersion::LatestScriptCompileVersion(0x02497746, 0x33D75C49, 0x9A60CB65, 0x7D6AC5A2);
+const FGuid FNiagaraCustomVersion::LatestScriptCompileVersion(0x4C7CC1D0, 0xE6421847, 0xA7879DFE, 0x9C00B2D1);
+
diff --git a/Engine/Plugins/FX/Niagara/Source/NiagaraEditor/Private/NiagaraHlslTranslator.cpp b/Engine/Plugins/FX/Niagara/Source/NiagaraEditor/Private/NiagaraHlslTranslator.cpp
--- a/Engine/Plugins/FX/Niagara/Source/NiagaraEditor/Private/NiagaraHlslTranslator.cpp
+++ b/Engine/Plugins/FX/Niagara/Source/NiagaraEditor/Private/NiagaraHlslTranslator.cpp
@@ -1009,6 +1009,7 @@ static void ConvertFloatToHalf(const FNiagaraCompileOptions& InCompileOptions, T
 			SYS_PARAM_ENGINE_SYSTEM_AGE,
 			SYS_PARAM_ENGINE_SYSTEM_NUM_EMITTERS_ALIVE,
 			SYS_PARAM_ENGINE_SYSTEM_NUM_EMITTERS,
+			SYS_PARAM_ENGINE_SYSTEM_RANDOM_SEED,
 			SYS_PARAM_ENGINE_NUM_SYSTEM_INSTANCES,
 			SYS_PARAM_ENGINE_EMITTER_NUM_PARTICLES,
 			SYS_PARAM_ENGINE_EMITTER_SIMULATION_POSITION,
@@ -1611,6 +1612,7 @@ const FNiagaraTranslateResults &FHlslNiagaraTranslator::Translate(const FNiagara
 				ParameterMapRegisterExternalConstantNamespaceVariable(SYS_PARAM_ENGINE_EMITTER_TOTAL_SPAWNED_PARTICLES, nullptr, 0, OutputIdx, nullptr);
 				ParameterMapRegisterExternalConstantNamespaceVariable(SYS_PARAM_EMITTER_RANDOM_SEED, nullptr, 0, OutputIdx, nullptr);
 				ParameterMapRegisterExternalConstantNamespaceVariable(SYS_PARAM_ENGINE_EMITTER_INSTANCE_SEED, nullptr, 0, OutputIdx, nullptr);
+				ParameterMapRegisterExternalConstantNamespaceVariable(SYS_PARAM_ENGINE_SYSTEM_RANDOM_SEED, nullptr, 0, OutputIdx, nullptr);
 			}
 		}
 
@@ -4143,7 +4145,8 @@ bool FHlslNiagaraTranslator::ShouldInterpolateParameter(const FNiagaraVariable&
 		Parameter == SYS_PARAM_ENGINE_EMITTER_SPAWN_COUNT_SCALE ||
 		Parameter == SYS_PARAM_EMITTER_RANDOM_SEED ||
 		Parameter == SYS_PARAM_ENGINE_EMITTER_INSTANCE_SEED ||
-		Parameter == SYS_PARAM_ENGINE_SYSTEM_TICK_COUNT)
+		Parameter == SYS_PARAM_ENGINE_SYSTEM_TICK_COUNT ||
+		Parameter == SYS_PARAM_ENGINE_SYSTEM_RANDOM_SEED)
 	{
 		return false;
 	}
