Mark parameter definition synchronization metadata as skip for compile hash to avoid arbitrarily dirtying the compile hash.

diff --git a/Engine/Plugins/FX/Niagara/Source/NiagaraEditor/Public/NiagaraScriptVariable.h b/Engine/Plugins/FX/Niagara/Source/NiagaraEditor/Public/NiagaraScriptVariable.h
--- a/Engine/Plugins/FX/Niagara/Source/NiagaraEditor/Public/NiagaraScriptVariable.h
+++ b/Engine/Plugins/FX/Niagara/Source/NiagaraEditor/Public/NiagaraScriptVariable.h
@@ -123,12 +123,12 @@ private:
 	UPROPERTY()
 	bool bIsStaticSwitch;
 
-	UPROPERTY()
+	UPROPERTY(meta = (SkipForCompileHash = "true"))
 	bool bSubscribedToParameterDefinitions;
 
 	UPROPERTY(meta = (SkipForCompileHash = "true"))
 	FGuid ChangeId;
 
-	UPROPERTY()
+	UPROPERTY(meta = (SkipForCompileHash = "true"))
 	bool bOverrideParameterDefinitionsDefaultValue;
 }; 
