- Add missing implementation for texture setting on the controller

diff --git a/Engine/Plugins/FX/Niagara/Source/Niagara/Private/NiagaraSystemInstanceController.cpp b/Engine/Plugins/FX/Niagara/Source/Niagara/Private/NiagaraSystemInstanceController.cpp
--- a/Engine/Plugins/FX/Niagara/Source/Niagara/Private/NiagaraSystemInstanceController.cpp
+++ b/Engine/Plugins/FX/Niagara/Source/Niagara/Private/NiagaraSystemInstanceController.cpp
@@ -444,6 +444,15 @@ void FNiagaraSystemInstanceController::SetVariable(FName InVariableName, TWeakOb
 	}
 }
 
+void FNiagaraSystemInstanceController::SetVariable(FName InVariableName, TWeakObjectPtr<UTexture> Texture)
+{
+	if (OverrideParameters)
+	{
+		const FNiagaraVariable VariableDesc(FNiagaraTypeDefinition::GetUTextureDef(), InVariableName);
+		OverrideParameters->SetUObject(Texture.Get(), VariableDesc);
+	}
+}
+
 void FNiagaraSystemInstanceController::SetVariable(FName InVariableName, TWeakObjectPtr<UTextureRenderTarget> TextureRenderTarget)
 {
 	if (OverrideParameters)
diff --git a/Engine/Plugins/FX/Niagara/Source/Niagara/Public/NiagaraSystemInstanceController.h b/Engine/Plugins/FX/Niagara/Source/Niagara/Public/NiagaraSystemInstanceController.h
--- a/Engine/Plugins/FX/Niagara/Source/Niagara/Public/NiagaraSystemInstanceController.h
+++ b/Engine/Plugins/FX/Niagara/Source/Niagara/Public/NiagaraSystemInstanceController.h
@@ -164,6 +164,7 @@ private:
 	void SetVariable(FName InVariableName, TWeakObjectPtr<UObject> Object);
 	void SetVariable(FName InVariableName, TWeakObjectPtr<UMaterialInterface> Object);
 	void SetVariable(FName InVariableName, TWeakObjectPtr<UStaticMesh> Object);
+	void SetVariable(FName InVariableName, TWeakObjectPtr<UTexture> Texture);
 	void SetVariable(FName InVariableName, TWeakObjectPtr<UTextureRenderTarget> TextureRenderTarget);
 
 	NIAGARA_SYSTEM_INSTANCE_CONTROLLER_SHIM(SetRequestedExecutionState,)
