CVar to disable Niagara GPU ray traced collisions

diff --git a/Engine/Plugins/FX/Niagara/Shaders/Private/NiagaraDataInterfaceCollisionQueryTemplate.ush b/Engine/Plugins/FX/Niagara/Shaders/Private/NiagaraDataInterfaceCollisionQueryTemplate.ush
--- a/Engine/Plugins/FX/Niagara/Shaders/Private/NiagaraDataInterfaceCollisionQueryTemplate.ush
+++ b/Engine/Plugins/FX/Niagara/Shaders/Private/NiagaraDataInterfaceCollisionQueryTemplate.ush
@@ -29,10 +29,10 @@ void IssueAsyncRayTraceGpu_{ParameterName}_UEImpureCall(int In_QueryID, float3 I
 	if (RayTracingEnabled_{ParameterName})
 	{
 		NDICollisionQuery_IssueAsyncRayTrace(MaxRayTraceCount_{ParameterName}, RWRayRequests_{ParameterName}, RWRayTraceCounts_{ParameterName}, In_QueryID, In_TraceStart, In_TraceEnd, In_TraceChannel, Out_IsQueryValid);
+		return;
 	}
-#else
-	Out_IsQueryValid = false;
 #endif
+	Out_IsQueryValid = false;
 }
 
 void ReserveRayTraceIndex_{ParameterName}_UEImpureCall(int In_TraceCount, out int Out_StartQueryID, out bool Out_IndicesValid)
@@ -41,10 +41,10 @@ void ReserveRayTraceIndex_{ParameterName}_UEImpureCall(int In_TraceCount, out in
 	if (RayTracingEnabled_{ParameterName})
 	{
 		NDICollisionQuery_ReserveRayTraceIndex(MaxRayTraceCount_{ParameterName}, RWRayTraceCounts_{ParameterName}, In_TraceCount, Out_StartQueryID, Out_IndicesValid);
+		return;
 	}
-#else
-	Out_IndicesValid = false;
 #endif
+	Out_IndicesValid = false;
 }
 
 void CreateAsyncRayTraceGpu_{ParameterName}_UEImpureCall(float3 In_TraceStart, float3 In_TraceEnd, int In_TraceChannel, out int Out_QueryID, out bool Out_IsQueryValid)
@@ -53,10 +53,10 @@ void CreateAsyncRayTraceGpu_{ParameterName}_UEImpureCall(float3 In_TraceStart, f
 	if (RayTracingEnabled_{ParameterName})
 	{
 		NDICollisionQuery_CreateAsyncRayTrace(MaxRayTraceCount_{ParameterName}, RWRayRequests_{ParameterName}, RWRayTraceCounts_{ParameterName}, In_TraceStart, In_TraceEnd, In_TraceChannel, Out_QueryID, Out_IsQueryValid);
+		return;
 	}
-#else
-	Out_IsQueryValid = false;
 #endif
+	Out_IsQueryValid = false;
 }
 
 void ReadAsyncRayTraceGpu_{ParameterName}(int In_PreviousFrameQueryID, out bool Out_CollisionValid, out float Out_CollisionDistance, out float3 Out_CollisionPosWorld, out float3 Out_CollisionNormal)
@@ -65,8 +65,8 @@ void ReadAsyncRayTraceGpu_{ParameterName}(int In_PreviousFrameQueryID, out bool
 	if (RayTracingEnabled_{ParameterName})
 	{
 		NDICollisionQuery_ReadAsyncRayTrace(MaxRayTraceCount_{ParameterName}, IntersectionResults_{ParameterName}, In_PreviousFrameQueryID, Out_CollisionValid, Out_CollisionDistance, Out_CollisionPosWorld, Out_CollisionNormal);
-	}
-#else
-	Out_CollisionValid = false;
+		return;
+	}	
 #endif
+	Out_CollisionValid = false;
 }
\ No newline at end of file
diff --git a/Engine/Plugins/FX/Niagara/Source/Niagara/Private/NiagaraDataInterfaceCollisionQuery.cpp b/Engine/Plugins/FX/Niagara/Source/Niagara/Private/NiagaraDataInterfaceCollisionQuery.cpp
--- a/Engine/Plugins/FX/Niagara/Source/Niagara/Private/NiagaraDataInterfaceCollisionQuery.cpp
+++ b/Engine/Plugins/FX/Niagara/Source/Niagara/Private/NiagaraDataInterfaceCollisionQuery.cpp
@@ -18,6 +18,14 @@
 
 #define LOCTEXT_NAMESPACE "NiagaraDataInterfaceCollisionQuery"
 
+int32 GEnableGPUHWRTCollisions = 1;
+static FAutoConsoleVariableRef CVarEnableGPUHWRTCollisions(
+	TEXT("fx.Niagara.Collision.EnableGPURayTracedCollisions"),
+	GEnableGPUHWRTCollisions,
+	TEXT("If greater than zero, GPU hardware ray trace collisions are enabled."),
+	ECVF_Default
+);
+
 namespace NDICollisionQueryLocal
 {
 	static const TCHAR* CommonShaderFile = TEXT("/Plugin/FX/Niagara/Private/NiagaraDataInterfaceCollisionQuery.ush");
@@ -164,7 +172,7 @@ struct FNiagaraDataIntefaceProxyCollisionQuery : public FNiagaraDataInterfacePro
 		DEC_MEMORY_STAT_BY(STAT_NiagaraGPUDataInterfaceMemory, RayTraceCounts.NumBytes);
 		RayTraceCounts.Release();
 
-		if (IsRayTracingEnabled() && InMaxRayTraceRequests > 0)
+		if (IsRayTracingEnabled() && GEnableGPUHWRTCollisions && InMaxRayTraceRequests > 0)
 		{
 			MaxRayTraceCount = 16 * FMath::DivideAndRoundUp(InMaxRayTraceRequests, 16);
 
@@ -646,7 +654,7 @@ void UNiagaraDataInterfaceCollisionQuery::ValidateFunction(const FNiagaraFunctio
 
 bool UNiagaraDataInterfaceCollisionQuery::RequiresRayTracingScene() const
 {
-	return IsRayTracingEnabled() && MaxRayTraceCount > 0;
+	return IsRayTracingEnabled() && GEnableGPUHWRTCollisions && MaxRayTraceCount > 0;
 }
 
 #if WITH_EDITORONLY_DATA
@@ -1042,7 +1050,7 @@ public:
 		}
 
 #if RHI_RAYTRACING
-		SetShaderValue(RHICmdList, ComputeShaderRHI, RayTracingEnabledParam, IsRayTracingEnabled() ? 1 : 0);
+		SetShaderValue(RHICmdList, ComputeShaderRHI, RayTracingEnabledParam, IsRayTracingEnabled() && GEnableGPUHWRTCollisions ? 1 : 0);
 		SetShaderValue(RHICmdList, ComputeShaderRHI, MaxRayTraceCountParam, QueryDI->MaxRayTraceCount);
 
 		if (RayRequestsParam.IsUAVBound())
