- Disable low latency path if we cast translucent shadows as the data will not be ready in time

#rnx

diff --git a/Engine/Plugins/FX/Niagara/Source/Niagara/Private/NiagaraRendererMeshes.cpp b/Engine/Plugins/FX/Niagara/Source/Niagara/Private/NiagaraRendererMeshes.cpp
--- a/Engine/Plugins/FX/Niagara/Source/Niagara/Private/NiagaraRendererMeshes.cpp
+++ b/Engine/Plugins/FX/Niagara/Source/Niagara/Private/NiagaraRendererMeshes.cpp
@@ -270,7 +270,7 @@ void FNiagaraRendererMeshes::PrepareParticleMeshRenderData(FParticleMeshRenderDa
 	bAllTranslucentMaterials &= ParticleMeshRenderData.bHasTranslucentMaterials;
 
 	// Anything to render?
-	ParticleMeshRenderData.SourceParticleData = ParticleMeshRenderData.DynamicDataMesh->GetParticleDataToRender(bAllTranslucentMaterials && bGpuLowLatencyTranslucency);
+	ParticleMeshRenderData.SourceParticleData = ParticleMeshRenderData.DynamicDataMesh->GetParticleDataToRender(bAllTranslucentMaterials && bGpuLowLatencyTranslucency && !SceneProxy->CastsVolumetricTranslucentShadow());
 	if ((ParticleMeshRenderData.SourceParticleData == nullptr) ||
 		(SourceMode == ENiagaraRendererSourceDataMode::Particles && ParticleMeshRenderData.SourceParticleData->GetNumInstances() == 0) ||
 		(Meshes.Num() == 0)
diff --git a/Engine/Plugins/FX/Niagara/Source/Niagara/Private/NiagaraRendererSprites.cpp b/Engine/Plugins/FX/Niagara/Source/Niagara/Private/NiagaraRendererSprites.cpp
--- a/Engine/Plugins/FX/Niagara/Source/Niagara/Private/NiagaraRendererSprites.cpp
+++ b/Engine/Plugins/FX/Niagara/Source/Niagara/Private/NiagaraRendererSprites.cpp
@@ -207,7 +207,7 @@ void FNiagaraRendererSprites::PrepareParticleSpriteRenderData(FParticleSpriteRen
 	// Do we have anything to render?
 	EBlendMode BlendMode = MaterialRenderProxy->GetIncompleteMaterialWithFallback(FeatureLevel).GetBlendMode();
 	ParticleSpriteRenderData.bHasTranslucentMaterials = IsTranslucentBlendMode(BlendMode);
-	ParticleSpriteRenderData.SourceParticleData = ParticleSpriteRenderData.DynamicDataSprites->GetParticleDataToRender(ParticleSpriteRenderData.bHasTranslucentMaterials && bGpuLowLatencyTranslucency);
+	ParticleSpriteRenderData.SourceParticleData = ParticleSpriteRenderData.DynamicDataSprites->GetParticleDataToRender(ParticleSpriteRenderData.bHasTranslucentMaterials && bGpuLowLatencyTranslucency && !SceneProxy->CastsVolumetricTranslucentShadow());
 	if ((ParticleSpriteRenderData.SourceParticleData == nullptr) ||
 		(SourceMode == ENiagaraRendererSourceDataMode::Particles && ParticleSpriteRenderData.SourceParticleData->GetNumInstances() == 0) ||
 		(GbEnableNiagaraSpriteRendering == 0)
