Merging CL 17268618, 17268825: Fix nullptr crash for Niagara GetWorldManager()

diff --git a/Engine/Plugins/FX/Niagara/Source/Niagara/Private/NiagaraComponent.cpp b/Engine/Plugins/FX/Niagara/Source/Niagara/Private/NiagaraComponent.cpp
--- a/Engine/Plugins/FX/Niagara/Source/Niagara/Private/NiagaraComponent.cpp
+++ b/Engine/Plugins/FX/Niagara/Source/Niagara/Private/NiagaraComponent.cpp
@@ -1359,12 +1359,18 @@ void UNiagaraComponent::OnSystemComplete(bool bExternalCompletion)
 
 		if (PoolingMethod == ENCPoolMethod::AutoRelease)//Don't release back to the pool if we're completing due to scalability culling.
 		{
-			FNiagaraWorldManager::Get(GetWorld())->GetComponentPool()->ReclaimWorldParticleSystem(this);
+			if (FNiagaraWorldManager* WorldManager = FNiagaraWorldManager::Get(GetWorld()))
+			{
+				WorldManager->GetComponentPool()->ReclaimWorldParticleSystem(this);
+			}
 		}
 		else if (PoolingMethod == ENCPoolMethod::ManualRelease_OnComplete)
 		{
 			PoolingMethod = ENCPoolMethod::ManualRelease;
-			FNiagaraWorldManager::Get(GetWorld())->GetComponentPool()->ReclaimWorldParticleSystem(this);
+			if (FNiagaraWorldManager* WorldManager = FNiagaraWorldManager::Get(GetWorld()))
+			{
+				WorldManager->GetComponentPool()->ReclaimWorldParticleSystem(this);
+			}
 		}
 		else if (bAutoDestroy)
 		{
diff --git a/Engine/Plugins/FX/Niagara/Source/Niagara/Private/NiagaraComponentPool.cpp b/Engine/Plugins/FX/Niagara/Source/Niagara/Private/NiagaraComponentPool.cpp
--- a/Engine/Plugins/FX/Niagara/Source/Niagara/Private/NiagaraComponentPool.cpp
+++ b/Engine/Plugins/FX/Niagara/Source/Niagara/Private/NiagaraComponentPool.cpp
@@ -48,7 +48,10 @@ static FAutoConsoleVariableRef NiagaraKeepPooledComponentsRegistered(
 void DumpPooledWorldNiagaraNiagaraSystemInfo(UWorld* World)
 {
 	check(World);
-	FNiagaraWorldManager::Get(World)->GetComponentPool()->Dump();
+	if (FNiagaraWorldManager* WorldManager = FNiagaraWorldManager::Get(World))
+	{
+		WorldManager->GetComponentPool()->Dump();
+	}
 }
 
 FAutoConsoleCommandWithWorld DumpNCPoolInfoCommand(
diff --git a/Engine/Plugins/FX/Niagara/Source/Niagara/Private/NiagaraFunctionLibrary.cpp b/Engine/Plugins/FX/Niagara/Source/Niagara/Private/NiagaraFunctionLibrary.cpp
--- a/Engine/Plugins/FX/Niagara/Source/Niagara/Private/NiagaraFunctionLibrary.cpp
+++ b/Engine/Plugins/FX/Niagara/Source/Niagara/Private/NiagaraFunctionLibrary.cpp
@@ -62,8 +62,11 @@ UNiagaraComponent* CreateNiagaraSystem(UNiagaraSystem* SystemTemplate, UWorld* W
 		{
 			if (UNiagaraComponentSettings::ShouldForceAutoPooling(SystemTemplate))
 			{
-				UNiagaraComponentPool* ComponentPool = FNiagaraWorldManager::Get(World)->GetComponentPool();
-				NiagaraComponent = ComponentPool->CreateWorldParticleSystem(SystemTemplate, World, ENCPoolMethod::AutoRelease);
+				if (FNiagaraWorldManager* WorldManager = FNiagaraWorldManager::Get(World))
+				{
+					UNiagaraComponentPool* ComponentPool = WorldManager->GetComponentPool();
+					NiagaraComponent = ComponentPool->CreateWorldParticleSystem(SystemTemplate, World, ENCPoolMethod::AutoRelease);
+				}
 			}
 			else
 			{
@@ -74,9 +77,9 @@ UNiagaraComponent* CreateNiagaraSystem(UNiagaraSystem* SystemTemplate, UWorld* W
 				NiagaraComponent->bAllowAnyoneToDestroyMe = true;
 			}
 		}
-		else
+		else if (FNiagaraWorldManager* WorldManager = FNiagaraWorldManager::Get(World))
 		{
-			UNiagaraComponentPool* ComponentPool = FNiagaraWorldManager::Get(World)->GetComponentPool();
+			UNiagaraComponentPool* ComponentPool = WorldManager->GetComponentPool();
 			NiagaraComponent = ComponentPool->CreateWorldParticleSystem(SystemTemplate, World, PoolingMethod);
 		}
 	}
diff --git a/Engine/Plugins/FX/Niagara/Source/Niagara/Private/NiagaraSystemInstance.cpp b/Engine/Plugins/FX/Niagara/Source/Niagara/Private/NiagaraSystemInstance.cpp
--- a/Engine/Plugins/FX/Niagara/Source/Niagara/Private/NiagaraSystemInstance.cpp
+++ b/Engine/Plugins/FX/Niagara/Source/Niagara/Private/NiagaraSystemInstance.cpp
@@ -969,15 +969,21 @@ void FNiagaraSystemInstance::AdvanceSimulation(int32 TickCountToSimulate, float
 
 bool FNiagaraSystemInstance::IsReadyToRun() const
 {
-	bool bAllReadyToRun = true;
+	// check world
+	if (World == nullptr || World->bIsTearingDown)
+	{
+		return false;
+	}
 
+	// check system
 	UNiagaraSystem* System = GetSystem();
-
 	if (!System || !System->IsReadyToRun())
 	{
 		return false;
 	}
 
+	// check emitters
+	bool bAllReadyToRun = true;
 	for (TSharedRef<FNiagaraEmitterInstance, ESPMode::ThreadSafe> Simulation : Emitters)
 	{
 		if (!Simulation->IsReadyToRun())
@@ -2174,7 +2180,7 @@ bool FNiagaraSystemInstance::HandleCompletion()
 	}
 
 	bool bCompletedAlready = IsComplete();
-	if (bCompletedAlready || bEmittersCompleteOrDisabled)
+	if (bCompletedAlready || bEmittersCompleteOrDisabled || World->bIsTearingDown)
 	{
 		//UE_LOG(LogNiagara, Log, TEXT("Completion Achieved"));
 		Complete(false);
diff --git a/Engine/Plugins/FX/Niagara/Source/Niagara/Private/NiagaraSystemSimulation.cpp b/Engine/Plugins/FX/Niagara/Source/Niagara/Private/NiagaraSystemSimulation.cpp
--- a/Engine/Plugins/FX/Niagara/Source/Niagara/Private/NiagaraSystemSimulation.cpp
+++ b/Engine/Plugins/FX/Niagara/Source/Niagara/Private/NiagaraSystemSimulation.cpp
@@ -797,7 +797,7 @@ UNiagaraParameterCollectionInstance* FNiagaraSystemSimulation::GetParameterColle
 	//If no explicit override from the system, just get the current instance set on the world.
 	if (!Ret)
 	{
-		if(FNiagaraWorldManager* WorldMan = FNiagaraWorldManager::Get(World))
+		if (FNiagaraWorldManager* WorldMan = FNiagaraWorldManager::Get(World))
 		{
 			Ret = WorldMan->GetParameterCollection(Collection);
 		}
@@ -920,7 +920,10 @@ void FNiagaraSystemSimulation::AddTickGroupPromotion(FNiagaraSystemInstance* Ins
 
 	check(bIsSolo == false);
 
-	FNiagaraWorldManager::Get(World)->MarkSimulationForPostActorWork(this);
+	if (FNiagaraWorldManager* WorldManager = FNiagaraWorldManager::Get(World))
+	{
+		WorldManager->MarkSimulationForPostActorWork(this);
+	}
 }
 
 void FNiagaraSystemSimulation::RemoveFromInstanceList(FNiagaraSystemInstance* Instance)
@@ -1055,7 +1058,10 @@ void FNiagaraSystemSimulation::SetInstanceState(FNiagaraSystemInstance* Instance
 	// When we add an instance to pending spawn we must also flag for post actor work
 	if ((NewState == ENiagaraSystemInstanceState::PendingSpawn) && !bIsSolo)
 	{
-		FNiagaraWorldManager::Get(World)->MarkSimulationForPostActorWork(this);
+		if (FNiagaraWorldManager* WorldManager = FNiagaraWorldManager::Get(World))
+		{
+			WorldManager->MarkSimulationForPostActorWork(this);
+		}
 	}
 }
 
@@ -1165,6 +1171,8 @@ void FNiagaraSystemSimulation::Tick_GameThread(float DeltaSeconds, const FGraphE
 
 	// Update instances
 	int32 SystemIndex = 0;
+	FNiagaraWorldManager* WorldManager = FNiagaraWorldManager::Get(World);
+	check(WorldManager != nullptr);
 	while (SystemIndex < SystemInstances.Num())
 	{
 		FNiagaraSystemInstance* Instance = SystemInstances[SystemIndex];
@@ -1178,9 +1186,6 @@ void FNiagaraSystemSimulation::Tick_GameThread(float DeltaSeconds, const FGraphE
 				// Tick demotion we need to do this now to ensure we complete in the correct group
 				if (DesiredTickGroup > SystemTickGroup)
 				{
-					FNiagaraWorldManager* WorldManager = FNiagaraWorldManager::Get(World);
-					check(WorldManager != nullptr);
-
 					TSharedPtr<FNiagaraSystemSimulation, ESPMode::ThreadSafe> NewSim = WorldManager->GetSystemSimulation(DesiredTickGroup, System);
 					NewSim->TransferInstance(Instance);
 					continue;
@@ -1228,7 +1233,7 @@ void FNiagaraSystemSimulation::Tick_GameThread(float DeltaSeconds, const FGraphE
 				const ETickingGroup DesiredTickGroup = Instance->CalculateTickGroup();
 				if (DesiredTickGroup != SystemTickGroup)
 				{
-					TSharedPtr<FNiagaraSystemSimulation, ESPMode::ThreadSafe> DestSim = FNiagaraWorldManager::Get(World)->GetSystemSimulation(DesiredTickGroup, System);
+					TSharedPtr<FNiagaraSystemSimulation, ESPMode::ThreadSafe> DestSim = WorldManager->GetSystemSimulation(DesiredTickGroup, System);
 					DestSim->TransferInstance(Instance);
 					continue;
 				}
@@ -1280,7 +1285,7 @@ void FNiagaraSystemSimulation::Tick_GameThread(float DeltaSeconds, const FGraphE
 		}
 		else
 		{
-			FNiagaraWorldManager::Get(World)->MarkSimulationsForEndOfFrameWait(this);
+			WorldManager->MarkSimulationsForEndOfFrameWait(this);
 		}
 	}
 	else
