Fixed a bug where changed to a system mid-compile could lead to a crash

diff --git a/Engine/Plugins/FX/Niagara/Source/Niagara/Private/NiagaraSystem.cpp b/Engine/Plugins/FX/Niagara/Source/Niagara/Private/NiagaraSystem.cpp
--- a/Engine/Plugins/FX/Niagara/Source/Niagara/Private/NiagaraSystem.cpp
+++ b/Engine/Plugins/FX/Niagara/Source/Niagara/Private/NiagaraSystem.cpp
@@ -147,10 +147,10 @@ TSharedPtr<FNiagaraCompileRequestDataBase, ESPMode::ThreadSafe> FNiagaraLazyPrec
 		// Grab the list of user variables that were actually encountered so that we can add to them later.
 		SystemPrecompiledData->GatherPreCompiledVariables(TEXT("User"), EncounteredExposedVars);
 		SystemPrecompiledData->GetReferencedObjects(CompilationRootObjects);
-		TArray<FNiagaraEmitterHandle> EmitterHandles = System->GetEmitterHandles();
+		const TArray<FNiagaraEmitterHandle>& EmitterHandles = System->GetEmitterHandles();
 		for (int32 i = 0; i < EmitterHandles.Num(); i++)
 		{
-			FNiagaraEmitterHandle Handle = EmitterHandles[i];
+			const FNiagaraEmitterHandle& Handle = EmitterHandles[i];
 			if (Handle.GetInstance() && Handle.GetIsEnabled())
 			{
 				TSharedPtr<FNiagaraCompileRequestDataBase, ESPMode::ThreadSafe> EmitterPrecompiledData = SystemPrecompiledData->GetDependentRequest(i);
@@ -179,7 +179,13 @@ TSharedPtr<FNiagaraCompileRequestDataBase, ESPMode::ThreadSafe> FNiagaraLazyPrec
 	{
 		return SystemPrecompiledData;
 	}
-	return EmitterMapping.FindChecked(ForScript);
+	TSharedPtr<FNiagaraCompileRequestDataBase, ESPMode::ThreadSafe>* EmitterPrecompileData = EmitterMapping.Find(ForScript);
+	if (EmitterPrecompileData == nullptr)
+	{
+		UE_LOG(LogNiagara, Log, TEXT("Failed to precompile %s. This is probably because the system was modified between the original compile request and now."), *System->GetPathName());
+		return nullptr;
+	}
+	return *EmitterPrecompileData;
 }
 
 FNiagaraAsyncCompileTask::FNiagaraAsyncCompileTask(FString InAssetPath, const FEmitterCompiledScriptPair& InScriptPair)
@@ -331,6 +337,10 @@ void FNiagaraAsyncCompileTask::PrecompileData()
 	DDCKey = CurrentDDCKey;
 	
 	ComputedPrecompileData = PrecompileReference->GetPrecompileData(ScriptPair.CompiledScript);
+	if (ComputedPrecompileData == nullptr)
+	{
+		AbortTask();
+	}
 }
 
 void FNiagaraAsyncCompileTask::StartCompileJob()
