- Fix not returning the correct instance

diff --git a/Engine/Plugins/FX/Niagara/Source/Niagara/Private/NiagaraDataSetDebugAccessor.cpp b/Engine/Plugins/FX/Niagara/Source/Niagara/Private/NiagaraDataSetDebugAccessor.cpp
--- a/Engine/Plugins/FX/Niagara/Source/Niagara/Private/NiagaraDataSetDebugAccessor.cpp
+++ b/Engine/Plugins/FX/Niagara/Source/Niagara/Private/NiagaraDataSetDebugAccessor.cpp
@@ -51,12 +51,12 @@ float FNiagaraDataSetDebugAccessor::ReadFloat(const FNiagaraDataBuffer* DataBuff
 		if (bIsFloat)
 		{
 			const float* FloatData = reinterpret_cast<const float*>(DataBuffer->GetComponentPtrFloat(ComponentIndex + Component));
-			return *FloatData;
+			return FloatData[Instance];
 		}
 		else if (bIsHalf)
 		{
 			const FFloat16* HalfData = reinterpret_cast<const FFloat16*>(DataBuffer->GetComponentPtrHalf(ComponentIndex + Component));
-			return *HalfData;
+			return HalfData[Instance];
 		}
 	}
 	return 0.0f;
@@ -67,7 +67,7 @@ int32 FNiagaraDataSetDebugAccessor::ReadInt(const FNiagaraDataBuffer* DataBuffer
 	if (DataBuffer != nullptr && ComponentIndex != INDEX_NONE && Instance < DataBuffer->GetNumInstances() && Component < NumComponents)
 	{
 		const int32* IntData = reinterpret_cast<const int32*>(DataBuffer->GetComponentPtrInt32(ComponentIndex + Component));
-		return *IntData;
+		return IntData[Instance];
 	}
 	return 0;
 }
