Fix for missing RayTraceCounts parameter for Niagara HWRT shaders when indirect args aren't supported

diff --git a/Engine/Plugins/FX/Niagara/Source/NiagaraShader/Private/NiagaraRayTracingHelper.cpp b/Engine/Plugins/FX/Niagara/Source/NiagaraShader/Private/NiagaraRayTracingHelper.cpp
--- a/Engine/Plugins/FX/Niagara/Source/NiagaraShader/Private/NiagaraRayTracingHelper.cpp
+++ b/Engine/Plugins/FX/Niagara/Source/NiagaraShader/Private/NiagaraRayTracingHelper.cpp
@@ -408,13 +408,14 @@ void FNiagaraRayTracingHelper::IssueRayTraces(FRHICommandList& RHICmdList, FScen
 	Params.CollisionOutput = RayTraceResults;
 	Params.HashToCollisionGroups = HashToCollisionGroups.SRV;
 	Params.MaxRetraces = MaxRetraces;
-	FRayTracingShaderBindingsWriter GlobalResources;
-	SetShaderParameters(GlobalResources, RGShader, Params);
 
 	if (FNiagaraCollisionRayTraceRG::SupportsIndirectDispatch())
 	{
 		RHICmdList.Transition(FRHITransitionInfo(IndirectArgsBuffer->UAV, ERHIAccess::UAVCompute, ERHIAccess::IndirectArgs | ERHIAccess::SRVCompute));
 
+		FRayTracingShaderBindingsWriter GlobalResources;
+		SetShaderParameters(GlobalResources, RGShader, Params);
+
 		RHICmdList.RayTraceDispatchIndirect(
 			RayTracingPipelineState,
 			RGShader.GetRayTracingShader(),
@@ -428,7 +429,10 @@ void FNiagaraRayTracingHelper::IssueRayTraces(FRHICommandList& RHICmdList, FScen
 	else
 	{
 		RHICmdList.Transition(FRHITransitionInfo(IndirectArgsBuffer->UAV, ERHIAccess::UAVCompute, ERHIAccess::SRVCompute));
-		GlobalResources.SetSRV(2, IndirectArgsBuffer->SRV);
+		Params.RayTraceCounts = IndirectArgsBuffer->SRV;
+
+		FRayTracingShaderBindingsWriter GlobalResources;
+		SetShaderParameters(GlobalResources, RGShader, Params);
 
 		RHICmdList.RayTraceDispatch(
 			RayTracingPipelineState,
