- Fix not clearing variables when using console commands

diff --git a/Engine/Plugins/FX/Niagara/Source/Niagara/Private/NiagaraDebugHud.cpp b/Engine/Plugins/FX/Niagara/Source/Niagara/Private/NiagaraDebugHud.cpp
--- a/Engine/Plugins/FX/Niagara/Source/Niagara/Private/NiagaraDebugHud.cpp
+++ b/Engine/Plugins/FX/Niagara/Source/Niagara/Private/NiagaraDebugHud.cpp
@@ -113,12 +113,12 @@ namespace NiagaraDebugLocal
 		MakeTuple(TEXT("SystemShowActiveOnlyInWorld="), TEXT("When enabled only active systems are shown in world"), [](FString Arg) {Settings.bSystemShowActiveOnlyInWorld = FCString::Atoi(*Arg) != 0; }),
 		MakeTuple(TEXT("SystemDebugVerbosity="), TEXT("Set the in world system debug verbosity"), [](FString Arg) {Settings.SystemDebugVerbosity = FMath::Clamp(ENiagaraDebugHudVerbosity(FCString::Atoi(*Arg)), ENiagaraDebugHudVerbosity::None, ENiagaraDebugHudVerbosity::Verbose); }),
 		MakeTuple(TEXT("SystemEmitterVerbosity="), TEXT("Set the in world system emitter debug verbosity"), [](FString Arg) {Settings.SystemEmitterVerbosity = FMath::Clamp(ENiagaraDebugHudVerbosity(FCString::Atoi(*Arg)), ENiagaraDebugHudVerbosity::None, ENiagaraDebugHudVerbosity::Verbose); }),
-		MakeTuple(TEXT("SystemVariables="), TEXT("Set the system variables to display"), [](FString Arg) {FNiagaraDebugHUDVariable::InitFromString(Arg, Settings.SystemVariables); GCachedSystemVariables.Empty(); }),
+		MakeTuple(TEXT("SystemVariables="), TEXT("Set the system variables to display"), [](FString Arg) {Settings.SystemVariables.Empty(); FNiagaraDebugHUDVariable::InitFromString(Arg, Settings.SystemVariables); GCachedSystemVariables.Empty(); }),
 		MakeTuple(TEXT("ShowSystemVariables="), TEXT("Set system variables visibility"), [](FString Arg) {Settings.bShowSystemVariables = FCString::Atoi(*Arg) != 0; GCachedSystemVariables.Empty(); }),
 
 		// Particle commands
 		MakeTuple(TEXT("EnableGpuParticleReadback="), TEXT("Enables GPU readback support for particle attributes"), [](FString Arg) {Settings.bEnableGpuParticleReadback = FCString::Atoi(*Arg) != 0;}),
-		MakeTuple(TEXT("ParticleVariables="), TEXT("Set the particle variables to display"), [](FString Arg) {FNiagaraDebugHUDVariable::InitFromString(Arg, Settings.ParticlesVariables); GCachedSystemVariables.Empty(); }),
+		MakeTuple(TEXT("ParticleVariables="), TEXT("Set the particle variables to display"), [](FString Arg) {Settings.ParticlesVariables.Empty(); FNiagaraDebugHUDVariable::InitFromString(Arg, Settings.ParticlesVariables); GCachedSystemVariables.Empty(); }),
 		MakeTuple(TEXT("ShowParticleVariables="), TEXT("Set Particle variables visibility"), [](FString Arg) {Settings.bShowParticleVariables = FCString::Atoi(*Arg) != 0; GCachedSystemVariables.Empty(); }),
 		MakeTuple(TEXT("MaxParticlesToDisplay="), TEXT("Maximum number of particles to show variables on"), [](FString Arg) {Settings.MaxParticlesToDisplay = FMath::Max(FCString::Atoi(*Arg), 0); Settings.bUseMaxParticlesToDisplay = Settings.MaxParticlesToDisplay > 0; }),
 		MakeTuple(TEXT("ShowParticlesVariablesWithSystem="), TEXT("When enabled particle variables are shown with the system display"), [](FString Arg) {Settings.bShowParticlesVariablesWithSystem = FCString::Atoi(*Arg) != 0; }),
