- Fixes to allow Niagara compute debugging on cooked targets

#rnx

diff --git a/Engine/Plugins/FX/Niagara/Source/Niagara/Classes/NiagaraDataInterfaceGrid2DCollection.h b/Engine/Plugins/FX/Niagara/Source/Niagara/Classes/NiagaraDataInterfaceGrid2DCollection.h
--- a/Engine/Plugins/FX/Niagara/Source/Niagara/Classes/NiagaraDataInterfaceGrid2DCollection.h
+++ b/Engine/Plugins/FX/Niagara/Source/Niagara/Classes/NiagaraDataInterfaceGrid2DCollection.h
@@ -46,7 +46,7 @@ struct FGrid2DCollectionRWInstanceData_GameThread
 	FVector2D CellSize = FVector2D::ZeroVector;
 	FVector2D WorldBBoxSize = FVector2D::ZeroVector;
 	EPixelFormat PixelFormat = EPixelFormat::PF_R32_FLOAT;
-#if WITH_EDITOR
+#if WITH_EDITORONLY_DATA
 	bool bPreviewGrid = false;
 	FIntVector4 PreviewAttribute = FIntVector4(INDEX_NONE, INDEX_NONE, INDEX_NONE, INDEX_NONE);
 #endif
@@ -83,7 +83,7 @@ struct FGrid2DCollectionRWInstanceData_RenderThread
 	TArray<int32> VarComponents;
 	TArray<uint32> Offsets;
 
-#if WITH_EDITOR
+#if WITH_EDITORONLY_DATA
 	bool bPreviewGrid = false;
 	FIntVector4 PreviewAttribute = FIntVector4(INDEX_NONE, INDEX_NONE, INDEX_NONE, INDEX_NONE);
 #endif
diff --git a/Engine/Plugins/FX/Niagara/Source/Niagara/Private/DataInterface/NiagaraDataInterfaceDebugDraw.cpp b/Engine/Plugins/FX/Niagara/Source/Niagara/Private/DataInterface/NiagaraDataInterfaceDebugDraw.cpp
--- a/Engine/Plugins/FX/Niagara/Source/Niagara/Private/DataInterface/NiagaraDataInterfaceDebugDraw.cpp
+++ b/Engine/Plugins/FX/Niagara/Source/Niagara/Private/DataInterface/NiagaraDataInterfaceDebugDraw.cpp
@@ -3131,8 +3131,7 @@ bool UNiagaraDataInterfaceDebugDraw::GenerateCompilerTagPrefix(const FNiagaraFun
 
 bool UNiagaraDataInterfaceDebugDraw::GPUContextInit(const FNiagaraScriptDataInterfaceCompileInfo& InInfo, void* PerInstanceData, FNiagaraSystemInstance* SystemInstance) const
 {
-
-#if NIAGARA_COMPUTEDEBUG_ENABLED
+#if NIAGARA_COMPUTEDEBUG_ENABLED && WITH_EDITORONLY_DATA
 	FNDIDebugDrawInstanceData_GameThread* InstanceData = reinterpret_cast<FNDIDebugDrawInstanceData_GameThread*>(PerInstanceData);
 	for (const FNiagaraFunctionSignature& Sig : InInfo.RegisteredFunctions)
 	{
diff --git a/Engine/Plugins/FX/Niagara/Source/Niagara/Private/NiagaraDataInterfaceGrid2DCollection.cpp b/Engine/Plugins/FX/Niagara/Source/Niagara/Private/NiagaraDataInterfaceGrid2DCollection.cpp
--- a/Engine/Plugins/FX/Niagara/Source/Niagara/Private/NiagaraDataInterfaceGrid2DCollection.cpp
+++ b/Engine/Plugins/FX/Niagara/Source/Niagara/Private/NiagaraDataInterfaceGrid2DCollection.cpp
@@ -12,9 +12,7 @@
 #include "NiagaraSystemInstance.h"
 #include "NiagaraSettings.h"
 #include "NiagaraBatchedElements.h"
-#if WITH_EDITOR
 #include "NiagaraGpuComputeDebug.h"
-#endif
 #include "Engine/TextureRenderTarget2D.h"
 #include "Engine/TextureRenderTarget2DArray.h"
 #include "NiagaraConstants.h"
@@ -1721,7 +1719,7 @@ bool UNiagaraDataInterfaceGrid2DCollection::CopyToInternal(UNiagaraDataInterface
 	OtherTyped->RenderTargetUserParameter = RenderTargetUserParameter;
 	OtherTyped->OverrideBufferFormat = OverrideBufferFormat;
 	OtherTyped->bOverrideFormat = bOverrideFormat;
-#if WITH_EDITOR
+#if WITH_EDITORONLY_DATA
 	OtherTyped->bPreviewGrid = bPreviewGrid;
 	OtherTyped->PreviewAttribute = PreviewAttribute;
 #endif
@@ -1793,7 +1791,7 @@ bool UNiagaraDataInterfaceGrid2DCollection::InitPerInstanceData(void* PerInstanc
 	InstanceData->RTUserParamBinding.Init(SystemInstance->GetInstanceParameters(), RenderTargetUserParameter.Parameter);
 	InstanceData->UpdateTargetTexture(BufferFormat);
 
-#if WITH_EDITOR
+#if WITH_EDITORONLY_DATA
 	InstanceData->bPreviewGrid = bPreviewGrid;
 	InstanceData->PreviewAttribute = FIntVector4(INDEX_NONE, INDEX_NONE, INDEX_NONE, INDEX_NONE);
 	if (bPreviewGrid && !PreviewAttribute.IsNone())
@@ -1848,7 +1846,7 @@ bool UNiagaraDataInterfaceGrid2DCollection::InitPerInstanceData(void* PerInstanc
 			TargetData->Vars.Emplace(RT_InstanceData.Vars[i].GetName());
 			TargetData->VarComponents.Emplace(RT_InstanceData.Vars[i].GetType().GetSize() / sizeof(float));
 		}
-#if WITH_EDITOR
+#if WITH_EDITORONLY_DATA
 		TargetData->bPreviewGrid = RT_InstanceData.bPreviewGrid;
 		TargetData->PreviewAttribute = RT_InstanceData.PreviewAttribute;
 #endif
@@ -2592,7 +2590,7 @@ void FNiagaraDataInterfaceProxyGrid2DCollectionProxy::PostSimulate(FRHICommandLi
 		}
 	}
 
-#if NIAGARA_COMPUTEDEBUG_ENABLED
+#if NIAGARA_COMPUTEDEBUG_ENABLED && WITH_EDITORONLY_DATA
 	if (ProxyData->bPreviewGrid && ProxyData->CurrentData)
 	{
 		if (FNiagaraGpuComputeDebug* GpuComputeDebug = Context.Batcher->GetGpuComputeDebug())
diff --git a/Engine/Plugins/FX/Niagara/Source/Niagara/Private/NiagaraDataInterfaceRenderTarget2D.cpp b/Engine/Plugins/FX/Niagara/Source/Niagara/Private/NiagaraDataInterfaceRenderTarget2D.cpp
--- a/Engine/Plugins/FX/Niagara/Source/Niagara/Private/NiagaraDataInterfaceRenderTarget2D.cpp
+++ b/Engine/Plugins/FX/Niagara/Source/Niagara/Private/NiagaraDataInterfaceRenderTarget2D.cpp
@@ -13,9 +13,7 @@
 #include "NiagaraSettings.h"
 #include "NiagaraSystemInstance.h"
 #include "NiagaraRenderer.h"
-#if WITH_EDITOR
 #include "NiagaraGpuComputeDebug.h"
-#endif
 #include "NiagaraStats.h"
 
 #define LOCTEXT_NAMESPACE "NiagaraDataInterfaceRenderTarget2D"
@@ -824,7 +822,7 @@ void FNiagaraDataInterfaceProxyRenderTarget2DProxy::PostSimulate(FRHICommandList
 			NiagaraGenerateMips::GenerateMips(RHICmdList, ProxyData->TextureRHI, ProxyData->MipMapGenerationType);
 		}
 
-#if NIAGARA_COMPUTEDEBUG_ENABLED
+#if NIAGARA_COMPUTEDEBUG_ENABLED && WITH_EDITORONLY_DATA
 		if (ProxyData->bPreviewTexture)
 		{
 			if (FNiagaraGpuComputeDebug* GpuComputeDebug = Context.Batcher->GetGpuComputeDebug())
diff --git a/Engine/Plugins/FX/Niagara/Source/Niagara/Private/NiagaraDataInterfaceRenderTarget2DArray.cpp b/Engine/Plugins/FX/Niagara/Source/Niagara/Private/NiagaraDataInterfaceRenderTarget2DArray.cpp
--- a/Engine/Plugins/FX/Niagara/Source/Niagara/Private/NiagaraDataInterfaceRenderTarget2DArray.cpp
+++ b/Engine/Plugins/FX/Niagara/Source/Niagara/Private/NiagaraDataInterfaceRenderTarget2DArray.cpp
@@ -13,9 +13,7 @@
 #include "NiagaraStats.h"
 #include "NiagaraRenderer.h"
 #include "NiagaraSettings.h"
-#if WITH_EDITOR
 #include "NiagaraGpuComputeDebug.h"
-#endif
 
 #define LOCTEXT_NAMESPACE "NiagaraDataInterfaceRenderTarget2DArray"
 
@@ -740,7 +738,7 @@ bool UNiagaraDataInterfaceRenderTarget2DArray::PerInstanceTickPostSimulate(void*
 
 void FNiagaraDataInterfaceProxyRenderTarget2DArrayProxy::PostSimulate(FRHICommandList& RHICmdList, const FNiagaraDataInterfaceArgs& Context)
 {
-#if NIAGARA_COMPUTEDEBUG_ENABLED
+#if NIAGARA_COMPUTEDEBUG_ENABLED && WITH_EDITORONLY_DATA
 	FRenderTarget2DArrayRWInstanceData_RenderThread* ProxyData = SystemInstancesToProxyData_RT.Find(Context.SystemInstanceID);
 
 	if (ProxyData && ProxyData->bPreviewTexture)
diff --git a/Engine/Plugins/FX/Niagara/Source/Niagara/Private/NiagaraDataInterfaceRenderTargetCube.cpp b/Engine/Plugins/FX/Niagara/Source/Niagara/Private/NiagaraDataInterfaceRenderTargetCube.cpp
--- a/Engine/Plugins/FX/Niagara/Source/Niagara/Private/NiagaraDataInterfaceRenderTargetCube.cpp
+++ b/Engine/Plugins/FX/Niagara/Source/Niagara/Private/NiagaraDataInterfaceRenderTargetCube.cpp
@@ -11,9 +11,7 @@
 #include "NiagaraStats.h"
 #include "NiagaraRenderer.h"
 #include "NiagaraSettings.h"
-#if WITH_EDITOR
 #include "NiagaraGpuComputeDebug.h"
-#endif
 
 #define LOCTEXT_NAMESPACE "NiagaraDataInterfaceRenderTargetCube"
 
@@ -716,7 +714,7 @@ bool UNiagaraDataInterfaceRenderTargetCube::PerInstanceTickPostSimulate(void* Pe
 
 void FNiagaraDataInterfaceProxyRenderTargetCubeProxy::PostSimulate(FRHICommandList& RHICmdList, const FNiagaraDataInterfaceArgs& Context)
 {
-#if NIAGARA_COMPUTEDEBUG_ENABLED
+#if NIAGARA_COMPUTEDEBUG_ENABLED && WITH_EDITORONLY_DATA
 	FRenderTargetCubeRWInstanceData_RenderThread* ProxyData = SystemInstancesToProxyData_RT.Find(Context.SystemInstanceID);
 
 	if (ProxyData && ProxyData->bPreviewTexture && ProxyData->TextureRHI.IsValid())
diff --git a/Engine/Plugins/FX/Niagara/Source/Niagara/Private/NiagaraDataInterfaceRenderTargetVolume.cpp b/Engine/Plugins/FX/Niagara/Source/Niagara/Private/NiagaraDataInterfaceRenderTargetVolume.cpp
--- a/Engine/Plugins/FX/Niagara/Source/Niagara/Private/NiagaraDataInterfaceRenderTargetVolume.cpp
+++ b/Engine/Plugins/FX/Niagara/Source/Niagara/Private/NiagaraDataInterfaceRenderTargetVolume.cpp
@@ -12,9 +12,7 @@
 #include "NiagaraStats.h"
 #include "NiagaraRenderer.h"
 #include "NiagaraSettings.h"
-#if WITH_EDITOR
 #include "NiagaraGpuComputeDebug.h"
-#endif
 
 #define LOCTEXT_NAMESPACE "NiagaraDataInterfaceRenderTargetVolume"
 
@@ -738,7 +736,7 @@ bool UNiagaraDataInterfaceRenderTargetVolume::PerInstanceTickPostSimulate(void*
 
 void FNiagaraDataInterfaceProxyRenderTargetVolumeProxy::PostSimulate(FRHICommandList& RHICmdList, const FNiagaraDataInterfaceArgs& Context)
 {
-#if NIAGARA_COMPUTEDEBUG_ENABLED
+#if NIAGARA_COMPUTEDEBUG_ENABLED && WITH_EDITORONLY_DATA
 	FRenderTargetVolumeRWInstanceData_RenderThread* ProxyData = SystemInstancesToProxyData_RT.Find(Context.SystemInstanceID);
 
 	if (ProxyData && ProxyData->bPreviewTexture && ProxyData->TextureRHI.IsValid())
