- Create & update the external UB on the render thread

diff --git a/Engine/Plugins/FX/Niagara/Source/Niagara/Classes/NiagaraComputeExecutionContext.h b/Engine/Plugins/FX/Niagara/Source/Niagara/Classes/NiagaraComputeExecutionContext.h
--- a/Engine/Plugins/FX/Niagara/Source/Niagara/Classes/NiagaraComputeExecutionContext.h
+++ b/Engine/Plugins/FX/Niagara/Source/Niagara/Classes/NiagaraComputeExecutionContext.h
@@ -173,6 +173,7 @@ public:
 	class FNiagaraShaderScript*  GPUScript_RT;
 
 	// persistent layouts used to create the constant buffers for the compute sim shader
+	uint32 ExternalCBufferLayoutSize = 0;
 	TRefCountPtr<FNiagaraRHIUniformBufferLayout> ExternalCBufferLayout;
 
 	//Dynamic state updated either from GT via RT commands or from the RT side sim code itself.
diff --git a/Engine/Plugins/FX/Niagara/Source/Niagara/Private/NiagaraComputeExecutionContext.cpp b/Engine/Plugins/FX/Niagara/Source/Niagara/Private/NiagaraComputeExecutionContext.cpp
--- a/Engine/Plugins/FX/Niagara/Source/Niagara/Private/NiagaraComputeExecutionContext.cpp
+++ b/Engine/Plugins/FX/Niagara/Source/Niagara/Private/NiagaraComputeExecutionContext.cpp
@@ -18,7 +18,6 @@ FNiagaraComputeExecutionContext::FNiagaraComputeExecutionContext()
 	, GPUScript(nullptr)
 	, GPUScript_RT(nullptr)
 {
-	ExternalCBufferLayout = new FNiagaraRHIUniformBufferLayout(TEXT("Niagara GPU External CBuffer"), 0);
 }
 
 FNiagaraComputeExecutionContext::~FNiagaraComputeExecutionContext()
diff --git a/Engine/Plugins/FX/Niagara/Source/Niagara/Private/NiagaraEmitterInstance.cpp b/Engine/Plugins/FX/Niagara/Source/Niagara/Private/NiagaraEmitterInstance.cpp
--- a/Engine/Plugins/FX/Niagara/Source/Niagara/Private/NiagaraEmitterInstance.cpp
+++ b/Engine/Plugins/FX/Niagara/Source/Niagara/Private/NiagaraEmitterInstance.cpp
@@ -1523,9 +1523,15 @@ void FNiagaraEmitterInstance::Tick(float DeltaSeconds)
 
 		// @todo-threadsafety do this once during init. Should not change during runtime...
 		const uint32 ConstantBufferSize = ParmSize / (GPUExecContext->HasInterpolationParameters ? 2 : 1);
-		if (GPUExecContext->ExternalCBufferLayout->ConstantBufferSize != ConstantBufferSize)
+		if (GPUExecContext->ExternalCBufferLayoutSize != ConstantBufferSize)
 		{
-			GPUExecContext->ExternalCBufferLayout = new FNiagaraRHIUniformBufferLayout(TEXT("Niagara GPU External CBuffer"), ConstantBufferSize);
+			GPUExecContext->ExternalCBufferLayoutSize = ConstantBufferSize;
+			ENQUEUE_RENDER_COMMAND(UpdateConstantBuffer)(
+				[GPUExecContext_RT=GPUExecContext, ConstantBufferSize_RT=ConstantBufferSize](FRHICommandListImmediate& RHICmdList)
+				{
+					GPUExecContext_RT->ExternalCBufferLayout = new FNiagaraRHIUniformBufferLayout(TEXT("Niagara GPU External CBuffer"), ConstantBufferSize_RT);
+				}
+			);
 		}
 
 		// Need to call post-tick, which calls the copy to previous for interpolated spawning
diff --git a/Engine/Plugins/FX/Niagara/Source/NiagaraCore/Private/NiagaraCustomVersion.cpp b/Engine/Plugins/FX/Niagara/Source/NiagaraCore/Private/NiagaraCustomVersion.cpp
--- a/Engine/Plugins/FX/Niagara/Source/NiagaraCore/Private/NiagaraCustomVersion.cpp
+++ b/Engine/Plugins/FX/Niagara/Source/NiagaraCore/Private/NiagaraCustomVersion.cpp
@@ -8,5 +8,5 @@ const FGuid FNiagaraCustomVersion::GUID(0xFCF57AFA, 0x50764283, 0xB9A9E658, 0xFF
 // Register the custom version with core
 FCustomVersionRegistration GRegisterNiagaraCustomVersion(FNiagaraCustomVersion::GUID, FNiagaraCustomVersion::LatestVersion, TEXT("NiagaraVer"));
 
-const FGuid FNiagaraCustomVersion::LatestScriptCompileVersion(0x4C7CC1D0, 0xE6421847, 0xA7879DFE, 0x9C00B2D1);
+const FGuid FNiagaraCustomVersion::LatestScriptCompileVersion(0x75FF6720, 0xCD21A74B, 0xB06E5894, 0xC71F0BEC);
 
