- Fix missing path for user fixed system bounds

#rnx

diff --git a/Engine/Plugins/FX/Niagara/Source/Niagara/Private/NiagaraComponent.cpp b/Engine/Plugins/FX/Niagara/Source/Niagara/Private/NiagaraComponent.cpp
--- a/Engine/Plugins/FX/Niagara/Source/Niagara/Private/NiagaraComponent.cpp
+++ b/Engine/Plugins/FX/Niagara/Source/Niagara/Private/NiagaraComponent.cpp
@@ -1794,7 +1794,11 @@ FBoxSphereBounds UNiagaraComponent::CalcBounds(const FTransform& LocalToWorld) c
 	else
 	{
 		FBox SimBounds(ForceInit);
-		if (Asset && Asset->bFixedBounds)
+		if ( SystemInstanceController.IsValid() )
+		{
+			SimBounds = SystemInstanceController->GetSystemFixedBounds();
+		}
+		if (!SimBounds.IsValid && Asset && Asset->bFixedBounds)
 		{
 			SimBounds = Asset->GetFixedBounds();
 		}
