- Change system determinism default to false which matches emitter
- Fix missing CalculateRandomRange not using system seed

diff --git a/Engine/Plugins/FX/Niagara/Source/Niagara/Classes/NiagaraSystem.h b/Engine/Plugins/FX/Niagara/Source/Niagara/Classes/NiagaraSystem.h
--- a/Engine/Plugins/FX/Niagara/Source/Niagara/Classes/NiagaraSystem.h
+++ b/Engine/Plugins/FX/Niagara/Source/Niagara/Classes/NiagaraSystem.h
@@ -776,7 +776,7 @@ protected:
 	When enabled we will always use the RandomSeed from the system plus the components RandomSeedOffset, this allows for determinism but variance between components.
 	*/
 	UPROPERTY(EditAnywhere, Category = "Random")
-	bool bDeterminism = true;
+	bool bDeterminism = false;
 
 	/** Seed used for system script random number generator. */
 	UPROPERTY(EditAnywhere, Category = "Random", meta = (EditCondition = "bDeterminism"))
diff --git a/Engine/Plugins/FX/Niagara/Source/Niagara/Private/NiagaraSystem.cpp b/Engine/Plugins/FX/Niagara/Source/Niagara/Private/NiagaraSystem.cpp
--- a/Engine/Plugins/FX/Niagara/Source/Niagara/Private/NiagaraSystem.cpp
+++ b/Engine/Plugins/FX/Niagara/Source/Niagara/Private/NiagaraSystem.cpp
@@ -902,9 +902,13 @@ void UNiagaraSystem::RecomputeExecutionOrderForDataInterface(class UNiagaraDataI
 
 void UNiagaraSystem::Serialize(FArchive& Ar)
 {
-	Super::Serialize(Ar);
-
 	Ar.UsingCustomVersion(FNiagaraCustomVersion::GUID);
+	if (Ar.CustomVer(FNiagaraCustomVersion::GUID) < FNiagaraCustomVersion::ChangeSystemDeterministicDefault)
+	{
+		bDeterminism = true;
+	}
+
+	Super::Serialize(Ar);
 
 	if (Ar.CustomVer(FNiagaraCustomVersion::GUID) >= FNiagaraCustomVersion::ChangeEmitterCompiledDataToSharedRefs)
 	{
diff --git a/Engine/Plugins/FX/Niagara/Source/NiagaraCore/Public/NiagaraCustomVersion.h b/Engine/Plugins/FX/Niagara/Source/NiagaraCore/Public/NiagaraCustomVersion.h
--- a/Engine/Plugins/FX/Niagara/Source/NiagaraCore/Public/NiagaraCustomVersion.h
+++ b/Engine/Plugins/FX/Niagara/Source/NiagaraCore/Public/NiagaraCustomVersion.h
@@ -166,6 +166,8 @@ struct FNiagaraCustomVersion
 
 		MoveDefaultValueFromFNiagaraVariableMetaDataToUNiagaraScriptVariable, 
 
+		ChangeSystemDeterministicDefault,	// Changed the default mode from deterministic to non-deterministic which matches emitters
+
 		// DO NOT ADD A NEW VERSION UNLESS YOU HAVE TALKED TO THE NIAGARA LEAD. Mismanagement of these versions can lead to data loss if it is adjusted in multiple streams simultaneously.
 		// -----<new versions can be added above this line>  -------------------------------------------------
 		VersionPlusOne,
