Fixed a bug where some engine constants could be marked as "read before being set"

diff --git a/Engine/Plugins/FX/Niagara/Source/NiagaraCore/Private/NiagaraCustomVersion.cpp b/Engine/Plugins/FX/Niagara/Source/NiagaraCore/Private/NiagaraCustomVersion.cpp
--- a/Engine/Plugins/FX/Niagara/Source/NiagaraCore/Private/NiagaraCustomVersion.cpp
+++ b/Engine/Plugins/FX/Niagara/Source/NiagaraCore/Private/NiagaraCustomVersion.cpp
@@ -8,5 +8,4 @@ const FGuid FNiagaraCustomVersion::GUID(0xFCF57AFA, 0x50764283, 0xB9A9E658, 0xFF
 // Register the custom version with core
 FCustomVersionRegistration GRegisterNiagaraCustomVersion(FNiagaraCustomVersion::GUID, FNiagaraCustomVersion::LatestVersion, TEXT("NiagaraVer"));
 
-const FGuid FNiagaraCustomVersion::LatestScriptCompileVersion(0x75FF6720, 0xCD21A74B, 0xB06E5894, 0xC71F0BEC);
-
+const FGuid FNiagaraCustomVersion::LatestScriptCompileVersion(0x02497746, 0x33D75C49, 0x9A60CB65, 0x7D6AC5A4);
diff --git a/Engine/Plugins/FX/Niagara/Source/NiagaraEditor/Private/NiagaraHlslTranslator.cpp b/Engine/Plugins/FX/Niagara/Source/NiagaraEditor/Private/NiagaraHlslTranslator.cpp
--- a/Engine/Plugins/FX/Niagara/Source/NiagaraEditor/Private/NiagaraHlslTranslator.cpp
+++ b/Engine/Plugins/FX/Niagara/Source/NiagaraEditor/Private/NiagaraHlslTranslator.cpp
@@ -8278,10 +8278,9 @@ void FHlslNiagaraTranslator::RegisterCompileDependency(const FNiagaraVariableBas
 	if (InVar.GetType().IsDataInterface() || InVar.GetType().IsUObject() || InVar.IsInNameSpace(FNiagaraConstants::UserNamespace) || InVar.IsInNameSpace(FNiagaraConstants::EngineNamespace) || InVar.IsInNameSpace(FNiagaraConstants::ParameterCollectionNamespace))
 		return;
 
-	if (InVar.GetName() == TEXT("Particles.UniqueID") || InVar.GetName() == TEXT("Emitter.SpawnGroup") || InVar.GetName() == TEXT("Emitter.InterpSpawnStartDt") ||
-		InVar.GetName() == TEXT("Emitter.SpawnInterval") || InVar.GetName() == TEXT("Emitter.RandomSeed"))
+	if (FNiagaraConstants::IsNiagaraConstant(InVar) || InVar.GetName() == TEXT("Emitter.InterpSpawnStartDt") || InVar.GetName() == TEXT("Emitter.SpawnInterval"))
 	{
-		return;																							  
+		return;
 	}
 
 	if (bEmitAsLinker)																					
