LWC support for FQuat and FTransform
Merge Dev-LargeWorldCoordinates @ 17289670

diff --git a/Engine/Plugins/Animation/ControlRig/Source/ControlRig/Private/Units/Hierarchy/RigUnit_SetControlTransform.cpp b/Engine/Plugins/Animation/ControlRig/Source/ControlRig/Private/Units/Hierarchy/RigUnit_SetControlTransform.cpp
--- a/Engine/Plugins/Animation/ControlRig/Source/ControlRig/Private/Units/Hierarchy/RigUnit_SetControlTransform.cpp
+++ b/Engine/Plugins/Animation/ControlRig/Source/ControlRig/Private/Units/Hierarchy/RigUnit_SetControlTransform.cpp
@@ -80,7 +80,7 @@ FRigUnit_SetControlFloat_Execute()
 				const FRigElementKey Key(Control, ERigElementType::Control);
 				if (CachedControlIndex.UpdateCache(Key, Hierarchy))
 				{
-					if(FMath::IsNearlyEqual(Weight, 1.f))
+					if(FMath::IsNearlyEqual((float)Weight, 1.f))
 					{
 						Hierarchy->SetControlValue(CachedControlIndex, FRigControlValue::Make<float>(FloatValue));
 					}
@@ -146,7 +146,7 @@ FRigUnit_SetControlInteger_Execute()
 				const FRigElementKey Key(Control, ERigElementType::Control);
 				if (CachedControlIndex.UpdateCache(Key, Hierarchy))
 				{
-					if(FMath::IsNearlyEqual(Weight, 1.f))
+					if(FMath::IsNearlyEqual((float)Weight, 1.f))
 					{
 						Hierarchy->SetControlValue(CachedControlIndex, FRigControlValue::Make<int32>(IntegerValue));
 					}
@@ -213,7 +213,7 @@ FRigUnit_SetControlVector2D_Execute()
 				if (CachedControlIndex.UpdateCache(Key, Hierarchy))
 				{
 					const FVector3f CurrentValue(Vector.X, Vector.Y, 0.f);
-					if(FMath::IsNearlyEqual(Weight, 1.f))
+					if(FMath::IsNearlyEqual((float)Weight, 1.f))
 					{
 						Hierarchy->SetControlValue(CachedControlIndex, FRigControlValue::Make<FVector3f>(CurrentValue));
 					}
@@ -289,7 +289,7 @@ FRigUnit_SetControlVector_Execute()
 
 					if (ControlType == ERigControlType::Position)
 					{
-						if(FMath::IsNearlyEqual(Weight, 1.f))
+						if(FMath::IsNearlyEqual((float)Weight, 1.f))
 						{
 							Transform.SetLocation(Vector);
 						}
@@ -301,7 +301,7 @@ FRigUnit_SetControlVector_Execute()
 					}
 					else if (ControlType == ERigControlType::Scale)
 					{
-						if(FMath::IsNearlyEqual(Weight, 1.f))
+						if(FMath::IsNearlyEqual((float)Weight, 1.f))
 						{
 							Transform.SetScale3D(Vector);
 						}
@@ -364,7 +364,7 @@ FRigUnit_SetControlRotator_Execute()
 					}
 
 					FQuat Quat = FQuat(Rotator);
-					if (FMath::IsNearlyEqual(Weight, 1.f))
+					if (FMath::IsNearlyEqual((float)Weight, 1.f))
 					{
 						Transform.SetRotation(Quat);
 					}
@@ -452,7 +452,7 @@ FRigUnit_SetControlTransform_Execute()
 					{
 						case EBoneGetterSetterMode::GlobalSpace:
 						{
-							if(FMath::IsNearlyEqual(Weight, 1.f))
+							if(FMath::IsNearlyEqual((float)Weight, 1.f))
 							{
 								Hierarchy->SetGlobalTransform(CachedControlIndex, Transform);
 							}
@@ -465,7 +465,7 @@ FRigUnit_SetControlTransform_Execute()
 						}
 						case EBoneGetterSetterMode::LocalSpace:
 						{
-							if(FMath::IsNearlyEqual(Weight, 1.f))
+							if(FMath::IsNearlyEqual((float)Weight, 1.f))
 							{
 								Hierarchy->SetLocalTransform(CachedControlIndex, Transform);
 							}
diff --git a/Engine/Plugins/Animation/LiveLink/Source/LiveLinkComponents/Public/Controllers/LiveLinkTransformController.h b/Engine/Plugins/Animation/LiveLink/Source/LiveLinkComponents/Public/Controllers/LiveLinkTransformController.h
--- a/Engine/Plugins/Animation/LiveLink/Source/LiveLinkComponents/Public/Controllers/LiveLinkTransformController.h
+++ b/Engine/Plugins/Animation/LiveLink/Source/LiveLinkComponents/Public/Controllers/LiveLinkTransformController.h
@@ -7,7 +7,7 @@
 #include "LiveLinkTransformController.generated.h"
 
 struct FLiveLinkTransformStaticData;
-struct FTransform;
+DECLARE_LWC_TYPE(Transform, 3);
 class USceneComponent;
 
 USTRUCT(BlueprintType)
@@ -34,7 +34,7 @@ struct LIVELINKCOMPONENTS_API FLiveLinkTransformControllerData
 	/**
 	 * Whether we sweep to the destination location, triggering overlaps along the way and stopping short of the target if blocked by something.
 	 * Only the root component is swept and checked for blocking collision, child components move without sweeping. If collision is off, this has no effect.
-	 */
+	*/
 	UPROPERTY(EditAnywhere, BlueprintReadOnly, Category = "LiveLink", AdvancedDisplay)
 	bool bSweep = false;
 
@@ -43,7 +43,7 @@ struct LIVELINKCOMPONENTS_API FLiveLinkTransformControllerData
 	 * If true, physics velocity for this object is unchanged (so ragdoll parts are not affected by change in location).
 	 * If false, physics velocity is updated based on the change in position (affecting ragdoll parts).
 	 * If CCD is on and not teleporting, this will affect objects along the entire sweep volume.
-	 */
+	*/
 	UPROPERTY(EditAnywhere, BlueprintReadOnly, Category = "LiveLink", AdvancedDisplay)
 	bool bTeleport = true;
 	
diff --git a/Engine/Plugins/Experimental/ChaosCloth/Source/Chaos/Private/ChaosCloth/ChaosClothingSimulationSolver.cpp b/Engine/Plugins/Experimental/ChaosCloth/Source/Chaos/Private/ChaosCloth/ChaosClothingSimulationSolver.cpp
--- a/Engine/Plugins/Experimental/ChaosCloth/Source/Chaos/Private/ChaosCloth/ChaosClothingSimulationSolver.cpp
+++ b/Engine/Plugins/Experimental/ChaosCloth/Source/Chaos/Private/ChaosCloth/ChaosClothingSimulationSolver.cpp
@@ -529,7 +529,7 @@ void FClothingSimulationSolver::SetReferenceVelocityScale(
 	OldRootBoneLocalTransform.AddToTranslation(-OldLocalSpaceLocation);
 
 	// Calculate deltas
-	const FTransform DeltaTransform = ReferenceSpaceTransform.GetRelativeTransform(OldReferenceSpaceTransform);
+	const FRigidTransform3 DeltaTransform = ReferenceSpaceTransform.GetRelativeTransform(OldReferenceSpaceTransform);
 
 	// Apply linear velocity scale
 	const FVec3 LinearRatio = FVec3(1.f) - LinearVelocityScale.BoundToBox(FVec3(0.f), FVec3(1.f));
@@ -545,10 +545,10 @@ void FClothingSimulationSolver::SetReferenceVelocityScale(
 	}
 
 	const FReal PartialDeltaAngle = DeltaAngle * FMath::Clamp(1.f - AngularVelocityScale, 0.f, 1.f);
-	DeltaRotation = FQuat(Axis, PartialDeltaAngle);
+	DeltaRotation = UE::Math::TQuat<FReal>(Axis, PartialDeltaAngle);
 
 	// Transform points back into the previous frame of reference before applying the adjusted deltas 
-	PreSimulationTransforms[GroupId] = OldRootBoneLocalTransform.Inverse() * FTransform(DeltaRotation, DeltaPosition) * OldRootBoneLocalTransform;
+	PreSimulationTransforms[GroupId] = OldRootBoneLocalTransform.Inverse() * FRigidTransform3(DeltaPosition, DeltaRotation) * OldRootBoneLocalTransform;
 
 	// Save the reference bone relative angular velocity for calculating the fictitious forces
 	FictitiousAngularDisplacement[GroupId] = ReferenceSpaceTransform.TransformVector(Axis * PartialDeltaAngle * FMath::Min(2.f, FictitiousAngularScale));  // Clamp to 2x the delta angle
diff --git a/Engine/Plugins/Experimental/MeshModelingToolsetExp/Source/GeometryProcessingAdapters/Private/ApproximateActorsImpl.cpp b/Engine/Plugins/Experimental/MeshModelingToolsetExp/Source/GeometryProcessingAdapters/Private/ApproximateActorsImpl.cpp
--- a/Engine/Plugins/Experimental/MeshModelingToolsetExp/Source/GeometryProcessingAdapters/Private/ApproximateActorsImpl.cpp
+++ b/Engine/Plugins/Experimental/MeshModelingToolsetExp/Source/GeometryProcessingAdapters/Private/ApproximateActorsImpl.cpp
@@ -552,6 +552,7 @@ static TSharedPtr<FApproximationMeshData> GenerateApproximationMesh(
 			CurResultMeshSpatial.Build();
 		}
 
+		using FTransform3d = UE::Geometry::FTransform3d;
 		TArray<FTransform3d> NoTransforms;
 		NoTransforms.Add(FTransform3d::Identity());
 		TArray<FDynamicMeshAABBTree3*> Spatials;
diff --git a/Engine/Plugins/Experimental/MeshModelingToolsetExp/Source/MeshModelingToolsEditorOnly/Private/ShapeSprayTool.cpp b/Engine/Plugins/Experimental/MeshModelingToolsetExp/Source/MeshModelingToolsEditorOnly/Private/ShapeSprayTool.cpp
--- a/Engine/Plugins/Experimental/MeshModelingToolsetExp/Source/MeshModelingToolsEditorOnly/Private/ShapeSprayTool.cpp
+++ b/Engine/Plugins/Experimental/MeshModelingToolsetExp/Source/MeshModelingToolsEditorOnly/Private/ShapeSprayTool.cpp
@@ -129,7 +129,7 @@ void UShapeSprayTool::OnUpdateDrag(const FRay& Ray)
 	UDynamicMeshBrushTool::OnUpdateDrag(Ray);
 
 	FFrame3f WorldFrame((FVector3f)LastBrushStamp.WorldPosition, (FVector3f)LastBrushStamp.WorldNormal);
-	FTransform3d Transform(Cast<IPrimitiveComponentBackedTarget>(Target)->GetWorldTransform());
+	UE::Geometry::FTransform3d Transform(Cast<IPrimitiveComponentBackedTarget>(Target)->GetWorldTransform());
 
 
 	FDynamicMesh3* Mesh = AccumMeshComponent->GetMesh();
diff --git a/Engine/Plugins/Experimental/MeshModelingToolsetExp/Source/MeshModelingToolsExp/Private/DynamicMeshBrushTool.cpp b/Engine/Plugins/Experimental/MeshModelingToolsetExp/Source/MeshModelingToolsExp/Private/DynamicMeshBrushTool.cpp
--- a/Engine/Plugins/Experimental/MeshModelingToolsetExp/Source/MeshModelingToolsExp/Private/DynamicMeshBrushTool.cpp
+++ b/Engine/Plugins/Experimental/MeshModelingToolsetExp/Source/MeshModelingToolsExp/Private/DynamicMeshBrushTool.cpp
@@ -25,7 +25,7 @@ void UDynamicMeshBrushTool::Setup()
 	PreviewMesh = NewObject<UPreviewMesh>(this);
 	PreviewMesh->bBuildSpatialDataStructure = true;
 	PreviewMesh->CreateInWorld(UE::ToolTarget::GetTargetActor(Target)->GetWorld(), FTransform::Identity);
-	FTransform3d LocalToWorldTransform = UE::ToolTarget::GetLocalToWorldTransform(Target);
+	UE::Geometry::FTransform3d LocalToWorldTransform = UE::ToolTarget::GetLocalToWorldTransform(Target);
 	PreviewMesh->SetTransform((FTransform)LocalToWorldTransform);
 
 	FComponentMaterialSet MaterialSet = UE::ToolTarget::GetMaterialSet(Target);
diff --git a/Engine/Plugins/Experimental/MeshModelingToolsetExp/Source/MeshModelingToolsExp/Private/MeshSelectionTool.cpp b/Engine/Plugins/Experimental/MeshModelingToolsetExp/Source/MeshModelingToolsExp/Private/MeshSelectionTool.cpp
--- a/Engine/Plugins/Experimental/MeshModelingToolsetExp/Source/MeshModelingToolsExp/Private/MeshSelectionTool.cpp
+++ b/Engine/Plugins/Experimental/MeshModelingToolsetExp/Source/MeshModelingToolsExp/Private/MeshSelectionTool.cpp
@@ -360,7 +360,7 @@ bool UMeshSelectionTool::HitTest(const FRay& Ray, FHitResult& OutHit)
 		SourceMesh->GetTriInfo(OutHit.FaceIndex, Normal, Area, Centroid);
 		FViewCameraState StateOut;
 		GetToolManager()->GetContextQueriesAPI()->GetCurrentViewState(StateOut);
-		FTransform3d LocalToWorld = UE::ToolTarget::GetLocalToWorldTransform(Target);
+		UE::Geometry::FTransform3d LocalToWorld = UE::ToolTarget::GetLocalToWorldTransform(Target);
 		FVector3d LocalEyePosition(LocalToWorld.InverseTransformPosition(StateOut.Position));
 
 		if (Normal.Dot((Centroid - LocalEyePosition)) > 0)
@@ -416,7 +416,7 @@ TUniquePtr<FDynamicMeshOctree3>& UMeshSelectionTool::GetOctree()
 
 void UMeshSelectionTool::CalculateVertexROI(const FBrushStampData& Stamp, TArray<int>& VertexROI)
 {
-	FTransform3d Transform = UE::ToolTarget::GetLocalToWorldTransform(Target);
+	UE::Geometry::FTransform3d Transform = UE::ToolTarget::GetLocalToWorldTransform(Target);
 	FVector3d StampPosLocal = Transform.InverseTransformPosition((FVector3d)Stamp.WorldPosition);
 
 	// TODO: need dynamic vertex hash table!
@@ -438,7 +438,7 @@ void UMeshSelectionTool::CalculateVertexROI(const FBrushStampData& Stamp, TArray
 
 void UMeshSelectionTool::CalculateTriangleROI(const FBrushStampData& Stamp, TArray<int>& TriangleROI)
 {
-	FTransform3d Transform = UE::ToolTarget::GetLocalToWorldTransform(Target);
+	UE::Geometry::FTransform3d Transform = UE::ToolTarget::GetLocalToWorldTransform(Target);
 	FVector3d StampPosLocal = Transform.InverseTransformPosition((FVector3d)Stamp.WorldPosition);
 
 	// always select first triangle
@@ -591,7 +591,7 @@ void UMeshSelectionTool::UpdateFaceSelection(const FBrushStampData& Stamp, const
 	{
 		FViewCameraState StateOut;
 		GetToolManager()->GetContextQueriesAPI()->GetCurrentViewState(StateOut);
-		FTransform3d LocalToWorld = UE::ToolTarget::GetLocalToWorldTransform(Target);
+		UE::Geometry::FTransform3d LocalToWorld = UE::ToolTarget::GetLocalToWorldTransform(Target);
 		FVector3d LocalEyePosition(LocalToWorld.InverseTransformPosition(StateOut.Position));
 
 		for (int tid : TriangleROI)
@@ -702,7 +702,7 @@ FBox UMeshSelectionTool::GetWorldSpaceFocusBox()
 		FAxisAlignedBox3d Bounds = FAxisAlignedBox3d::Empty();
 
 		const FDynamicMesh3* Mesh = PreviewMesh->GetMesh();
-		FTransform3d Transform(PreviewMesh->GetTransform());
+		UE::Geometry::FTransform3d Transform(PreviewMesh->GetTransform());
 		for (int32 tid : SelectedFaces)
 		{
 			FIndex3i Tri = Mesh->GetTriangle(tid);
diff --git a/Engine/Plugins/Experimental/MeshModelingToolsetExp/Source/MeshModelingToolsExp/Private/SkinWeightsPaintTool.cpp b/Engine/Plugins/Experimental/MeshModelingToolsetExp/Source/MeshModelingToolsExp/Private/SkinWeightsPaintTool.cpp
--- a/Engine/Plugins/Experimental/MeshModelingToolsetExp/Source/MeshModelingToolsExp/Private/SkinWeightsPaintTool.cpp
+++ b/Engine/Plugins/Experimental/MeshModelingToolsetExp/Source/MeshModelingToolsExp/Private/SkinWeightsPaintTool.cpp
@@ -280,7 +280,7 @@ void USkinWeightsPaintTool::CalculateVertexROI(const FBrushStampData& Stamp, TAr
 {
 	using namespace UE::Geometry;
 	IPrimitiveComponentBackedTarget* TargetComponent = Cast<IPrimitiveComponentBackedTarget>(Target);
-	FTransform3d Transform(TargetComponent->GetWorldTransform());
+	UE::Geometry::FTransform3d Transform(TargetComponent->GetWorldTransform());
 	FVector3d StampPosLocal = Transform.InverseTransformPosition((FVector3d)Stamp.WorldPosition);
 
 	float RadiusSqr = CurrentBrushRadius * CurrentBrushRadius;
@@ -351,7 +351,7 @@ void USkinWeightsPaintTool::ApplyStamp(const FBrushStampData& Stamp)
 		return;
 
 	IPrimitiveComponentBackedTarget* TargetComponent = Cast<IPrimitiveComponentBackedTarget>(Target);
-	FTransform3d Transform(TargetComponent->GetWorldTransform());
+	UE::Geometry::FTransform3d Transform(TargetComponent->GetWorldTransform());
 	FVector3d StampPosLocal = Transform.InverseTransformPosition((FVector3d)Stamp.WorldPosition);
 
 	TArray<int32> ROIVertices;
diff --git a/Engine/Plugins/Experimental/MeshModelingToolsetExp/Source/MeshModelingToolsExp/Private/SplitMeshesTool.cpp b/Engine/Plugins/Experimental/MeshModelingToolsetExp/Source/MeshModelingToolsExp/Private/SplitMeshesTool.cpp
--- a/Engine/Plugins/Experimental/MeshModelingToolsetExp/Source/MeshModelingToolsExp/Private/SplitMeshesTool.cpp
+++ b/Engine/Plugins/Experimental/MeshModelingToolsetExp/Source/MeshModelingToolsExp/Private/SplitMeshesTool.cpp
@@ -115,7 +115,7 @@ void USplitMeshesTool::Shutdown(EToolShutdownType ShutdownType)
 			check(TargetActor != nullptr);
 			DeleteActors.Add(TargetActor);
 
-			FTransform3d SourceTransform = UE::ToolTarget::GetLocalToWorldTransform(Targets[ti]);
+			UE::Geometry::FTransform3d SourceTransform = UE::ToolTarget::GetLocalToWorldTransform(Targets[ti]);
 			FDynamicMesh3 SourceMesh = UE::ToolTarget::GetDynamicMeshCopy(Targets[ti], true);
 			FString AssetName = TargetActor->GetActorNameOrLabel();
 
@@ -136,7 +136,7 @@ void USplitMeshesTool::Shutdown(EToolShutdownType ShutdownType)
 			{
 				FCreateMeshObjectParams NewMeshObjectParams = BaseMeshObjectParams;
 				NewMeshObjectParams.BaseName = FString::Printf(TEXT("%s_%d"), *AssetName, k);
-				FTransform3d PartTransform = SourceTransform;
+				UE::Geometry::FTransform3d PartTransform = SourceTransform;
 				PartTransform.SetTranslation(SourceTransform.GetTranslation() + SplitInfo.Origins[k]);
 				NewMeshObjectParams.Transform = (FTransform)PartTransform;
 				if (BasicProperties->bTransferMaterials)
diff --git a/Engine/Plugins/Experimental/MeshModelingToolsetExp/Source/MeshModelingToolsExp/Private/ToolActivities/PolyEditCutFacesActivity.cpp b/Engine/Plugins/Experimental/MeshModelingToolsetExp/Source/MeshModelingToolsExp/Private/ToolActivities/PolyEditCutFacesActivity.cpp
--- a/Engine/Plugins/Experimental/MeshModelingToolsetExp/Source/MeshModelingToolsExp/Private/ToolActivities/PolyEditCutFacesActivity.cpp
+++ b/Engine/Plugins/Experimental/MeshModelingToolsetExp/Source/MeshModelingToolsExp/Private/ToolActivities/PolyEditCutFacesActivity.cpp
@@ -100,7 +100,7 @@ void UPolyEditCutFacesActivity::BeginCutFaces()
 	TArray<int32> ActiveTriangleSelection;
 	ActivityContext->CurrentTopology->GetSelectedTriangles(ActiveSelection, ActiveTriangleSelection);
 
-	FTransform3d WorldTransform(ActivityContext->Preview->PreviewMesh->GetTransform());
+	UE::Geometry::FTransform3d WorldTransform(ActivityContext->Preview->PreviewMesh->GetTransform());
 
 	EditPreview = PolyEditActivityUtil::CreatePolyEditPreviewMesh(*ParentTool, *ActivityContext);
 	EditPreview->InitializeStaticType(ActivityContext->CurrentMesh.Get(), ActiveTriangleSelection, &WorldTransform);
@@ -147,7 +147,7 @@ void UPolyEditCutFacesActivity::ApplyCutFaces()
 	}
 	FVector3d PlaneOrigin = 0.5 * (Point0.Origin + Point1.Origin);
 	// map into local space of target mesh
-	FTransform3d WorldTransform(ActivityContext->Preview->PreviewMesh->GetTransform());
+	UE::Geometry::FTransform3d WorldTransform(ActivityContext->Preview->PreviewMesh->GetTransform());
 	PlaneOrigin = WorldTransform.InverseTransformPosition(PlaneOrigin);
 	PlaneNormal = WorldTransform.InverseTransformNormal(PlaneNormal);
 	UE::Geometry::Normalize(PlaneNormal);
diff --git a/Engine/Plugins/Experimental/MeshModelingToolsetExp/Source/MeshModelingToolsExp/Private/ToolActivities/PolyEditExtrudeActivity.cpp b/Engine/Plugins/Experimental/MeshModelingToolsetExp/Source/MeshModelingToolsExp/Private/ToolActivities/PolyEditExtrudeActivity.cpp
--- a/Engine/Plugins/Experimental/MeshModelingToolsetExp/Source/MeshModelingToolsExp/Private/ToolActivities/PolyEditExtrudeActivity.cpp
+++ b/Engine/Plugins/Experimental/MeshModelingToolsetExp/Source/MeshModelingToolsExp/Private/ToolActivities/PolyEditExtrudeActivity.cpp
@@ -125,7 +125,7 @@ void UPolyEditExtrudeActivity::BeginExtrude()
 	TArray<int32> ActiveTriangleSelection;
 	ActivityContext->CurrentTopology->GetSelectedTriangles(ActiveSelection, ActiveTriangleSelection);
 
-	FTransform3d WorldTransform(ActivityContext->Preview->PreviewMesh->GetTransform());
+	UE::Geometry::FTransform3d WorldTransform(ActivityContext->Preview->PreviewMesh->GetTransform());
 
 	// Get the world frame
 	FFrame3d ActiveSelectionFrameLocal = ActivityContext->CurrentTopology->GetSelectionFrame(ActiveSelection);
@@ -179,7 +179,7 @@ void UPolyEditExtrudeActivity::ApplyExtrude()
 	TArray<int32> ActiveTriangleSelection;
 	ActivityContext->CurrentTopology->GetSelectedTriangles(ActiveSelection, ActiveTriangleSelection);
 
-	FTransform3d WorldTransform(ActivityContext->Preview->PreviewMesh->GetTransform());
+	UE::Geometry::FTransform3d WorldTransform(ActivityContext->Preview->PreviewMesh->GetTransform());
 	FVector3d MeshSpaceExtrudeDirection = WorldTransform.InverseTransformVector(ActiveSelectionFrameWorld.Z());
 	double ExtrudeDistance = ExtrudeHeightMechanic->CurrentHeight;
 
@@ -239,6 +239,7 @@ void UPolyEditExtrudeActivity::Clear()
 
 FVector3d UPolyEditExtrudeActivity::GetExtrudeDirection() const
 {
+	using FTransform3d = UE::Geometry::FTransform3d;
 	switch (ExtrudeProperties->Direction)
 	{
 	default:
diff --git a/Engine/Plugins/Experimental/MeshModelingToolsetExp/Source/MeshModelingToolsExp/Private/ToolActivities/PolyEditInsetOutsetActivity.cpp b/Engine/Plugins/Experimental/MeshModelingToolsetExp/Source/MeshModelingToolsExp/Private/ToolActivities/PolyEditInsetOutsetActivity.cpp
--- a/Engine/Plugins/Experimental/MeshModelingToolsetExp/Source/MeshModelingToolsExp/Private/ToolActivities/PolyEditInsetOutsetActivity.cpp
+++ b/Engine/Plugins/Experimental/MeshModelingToolsetExp/Source/MeshModelingToolsExp/Private/ToolActivities/PolyEditInsetOutsetActivity.cpp
@@ -115,7 +115,7 @@ void UPolyEditInsetOutsetActivity::BeginInset()
 	TArray<int32> ActiveTriangleSelection;
 	ActivityContext->CurrentTopology->GetSelectedTriangles(ActiveSelection, ActiveTriangleSelection);
 
-	FTransform3d WorldTransform(ActivityContext->Preview->PreviewMesh->GetTransform());
+	UE::Geometry::FTransform3d WorldTransform(ActivityContext->Preview->PreviewMesh->GetTransform());
 
 	EditPreview = PolyEditActivityUtil::CreatePolyEditPreviewMesh(*ParentTool, *ActivityContext);
 	EditPreview->InitializeInsetType(ActivityContext->CurrentMesh.Get(), ActiveTriangleSelection, &WorldTransform);
@@ -138,7 +138,7 @@ void UPolyEditInsetOutsetActivity::BeginInset()
 	FMeshBoundaryLoops Loops(&InsetHitTargetMesh);
 	TArray<FVector3d> LoopVertices;
 	Loops.Loops[0].GetVertices(LoopVertices);
-	CurveDistMechanic->InitializePolyLoop(LoopVertices, FTransform3d::Identity());
+	CurveDistMechanic->InitializePolyLoop(LoopVertices, UE::Geometry::FTransform3d::Identity());
 
 	SetToolPropertySourceEnabled(Settings, true);
 
diff --git a/Engine/Plugins/Experimental/MeshModelingToolsetExp/Source/MeshModelingToolsExp/Private/UVProjectionTool.cpp b/Engine/Plugins/Experimental/MeshModelingToolsetExp/Source/MeshModelingToolsExp/Private/UVProjectionTool.cpp
--- a/Engine/Plugins/Experimental/MeshModelingToolsetExp/Source/MeshModelingToolsExp/Private/UVProjectionTool.cpp
+++ b/Engine/Plugins/Experimental/MeshModelingToolsetExp/Source/MeshModelingToolsExp/Private/UVProjectionTool.cpp
@@ -113,7 +113,7 @@ void UUVProjectionTool::Setup()
 	ProjectionShapeVisualizer.LineThickness = 2.0;
 
 	// initialize Gizmo
-	FTransform3d GizmoPositionWorld(WorldBounds.Center());
+	UE::Geometry::FTransform3d GizmoPositionWorld(WorldBounds.Center());
 	TransformProxy = NewObject<UTransformProxy>(this);
 	TransformProxy->SetTransform((FTransform)GizmoPositionWorld);
 	TransformProxy->OnTransformChanged.AddUObject(this, &UUVProjectionTool::TransformChanged);
diff --git a/Engine/Plugins/Experimental/MeshModelingToolsetExp/Source/ModelingComponents/Private/ModelingToolTargetUtil.cpp b/Engine/Plugins/Experimental/MeshModelingToolsetExp/Source/ModelingComponents/Private/ModelingToolTargetUtil.cpp
--- a/Engine/Plugins/Experimental/MeshModelingToolsetExp/Source/ModelingComponents/Private/ModelingToolTargetUtil.cpp
+++ b/Engine/Plugins/Experimental/MeshModelingToolsetExp/Source/ModelingComponents/Private/ModelingToolTargetUtil.cpp
@@ -86,15 +86,15 @@ bool UE::ToolTarget::SetSourceObjectVisible(UToolTarget* Target, bool bVisible)
 }
 
 
-FTransform3d UE::ToolTarget::GetLocalToWorldTransform(UToolTarget* Target)
+UE::Geometry::FTransform3d UE::ToolTarget::GetLocalToWorldTransform(UToolTarget* Target)
 {
 	IPrimitiveComponentBackedTarget* TargetComponent = Cast<IPrimitiveComponentBackedTarget>(Target);
 	if (TargetComponent)
 	{
-		return (FTransform3d)TargetComponent->GetWorldTransform();
+		return (UE::Geometry::FTransform3d)TargetComponent->GetWorldTransform();
 	}
 	ensure(false);
-	return FTransform3d();
+	return UE::Geometry::FTransform3d();
 }
 
 FComponentMaterialSet UE::ToolTarget::GetMaterialSet(UToolTarget* Target, bool bPreferAssetMaterials)
diff --git a/Engine/Plugins/Experimental/PlanarCutPlugin/Source/PlanarCut/Private/GeometryMeshConversion.cpp b/Engine/Plugins/Experimental/PlanarCutPlugin/Source/PlanarCut/Private/GeometryMeshConversion.cpp
--- a/Engine/Plugins/Experimental/PlanarCutPlugin/Source/PlanarCut/Private/GeometryMeshConversion.cpp
+++ b/Engine/Plugins/Experimental/PlanarCutPlugin/Source/PlanarCut/Private/GeometryMeshConversion.cpp
@@ -1054,7 +1054,7 @@ FCellMeshes::FCellMeshes(int32 NumUVLayersIn, FDynamicMesh3& SingleCutter, const
 
 	if (Transform.IsSet())
 	{
-		MeshTransforms::ApplyTransform(SingleCutter, FTransform3d(Transform.GetValue()));
+		MeshTransforms::ApplyTransform(SingleCutter, UE::Geometry::FTransform3d(Transform.GetValue()));
 	}
 
 	// Mesh should already be augmented
@@ -1995,6 +1995,7 @@ void FDynamicMeshCollection::Init(const FGeometryCollection* Collection, const T
 			continue;
 		}
 
+		using FTransform3d = UE::Geometry::FTransform3d;
 		FTransform3d CollectionToLocal = FTransform3d(GeometryCollectionAlgo::GlobalMatrix(Collection->Transform, Collection->Parent, TransformIdx) * TransformCollection);
 
 		int32 AddedMeshIdx = Meshes.Add(new FMeshData(NumUVLayers));
diff --git a/Engine/Plugins/Experimental/PlanarCutPlugin/Source/PlanarCut/Private/PlanarCut.cpp b/Engine/Plugins/Experimental/PlanarCutPlugin/Source/PlanarCut/Private/PlanarCut.cpp
--- a/Engine/Plugins/Experimental/PlanarCutPlugin/Source/PlanarCut/Private/PlanarCut.cpp
+++ b/Engine/Plugins/Experimental/PlanarCutPlugin/Source/PlanarCut/Private/PlanarCut.cpp
@@ -53,7 +53,6 @@ using namespace UE::Geometry;
 
 using namespace UE::PlanarCut;
 
-
 // logic from FMeshUtility::GenerateGeometryCollectionFromBlastChunk, sets material IDs based on construction pattern that external materials have even IDs and are matched to internal materials at InternalID = ExternalID+1
 int32 FInternalSurfaceMaterials::GetDefaultMaterialIDForGeometry(const FGeometryCollection& Collection, int32 GeometryIdx) const
 {
@@ -888,6 +887,8 @@ int32 MergeBones(
 	bool bUnionJoinedPieces
 )
 {
+	using UE::Geometry::FTransform3d;
+
 	FTransform CellsToWorld = FTransform::Identity;
 
 	FGeometryCollectionProximityUtility::UpdateProximity(&Collection);
@@ -1269,7 +1270,7 @@ void ConvertToMeshDescription(
 		// unless it's just one mesh that we're going to center later anyway
 		if (!bCenterPivot || NumMeshes > 1)
 		{
-			MeshTransforms::ApplyTransform(Mesh, (FTransform3d)ToCollection);
+			MeshTransforms::ApplyTransform(Mesh, (UE::Geometry::FTransform3d)ToCollection);
 		}
 
 		FMeshNormals::InitializeOverlayToPerVertexNormals(Mesh.Attributes()->PrimaryNormals(), true);
diff --git a/Engine/Plugins/Experimental/UVEditor/Source/UVEditorTools/Private/UVSelectTool.cpp b/Engine/Plugins/Experimental/UVEditor/Source/UVEditorTools/Private/UVSelectTool.cpp
--- a/Engine/Plugins/Experimental/UVEditor/Source/UVEditorTools/Private/UVSelectTool.cpp
+++ b/Engine/Plugins/Experimental/UVEditor/Source/UVEditorTools/Private/UVSelectTool.cpp
@@ -603,8 +603,8 @@ void UUVSelectTool::OnSelectionChanged()
 						}
 					}
 				}
+				}
 			}
-		}
 		else if (Selection.Type == FDynamicMeshSelection::EType::Vertex)
 		{
 			for (int32 Vid : Selection.SelectedIDs)
@@ -648,7 +648,7 @@ void UUVSelectTool::UpdateLivePreviewLines()
 	{
 		FTransform MeshTransform = Targets[SelectionTargetIndex]->AppliedPreview->PreviewMesh->GetTransform();
 		const FDynamicMesh3* LivePreviewMesh = Targets[SelectionTargetIndex]->AppliedCanonical.Get();
-		
+
 		for (int32 Eid : BoundaryEids)
 		{
 			FVector3d Vert1, Vert2;
@@ -756,7 +756,7 @@ void UUVSelectTool::ApplyGizmoTransform()
 {
 	if (bGizmoTransformNeedsApplication)
 	{
-		FTransform3d TransformToApply(UnappliedGizmoTransform);
+		UE::Geometry::FTransform3d TransformToApply(UnappliedGizmoTransform);
 
 		// TODO: The division here is a bit of a hack. Properly-speaking, the scaling handles should act relative to
 		// gizmo size, not the visible space across which we drag, otherwise it becomes dependent on the units we
diff --git a/Engine/Plugins/Experimental/Water/Source/Runtime/Private/WaterBooleanUtils.cpp b/Engine/Plugins/Experimental/Water/Source/Runtime/Private/WaterBooleanUtils.cpp
--- a/Engine/Plugins/Experimental/Water/Source/Runtime/Private/WaterBooleanUtils.cpp
+++ b/Engine/Plugins/Experimental/Water/Source/Runtime/Private/WaterBooleanUtils.cpp
@@ -23,7 +23,6 @@
 #include "Model.h"
 #include "WaterModule.h"
 
-
 double FWaterBooleanUtils::UnitAspectRatio(const FVector3d& A, const FVector3d& B, const FVector3d& C)
 {
 	double AspectRatio = VectorUtil::AspectRatio(A, B, C);
@@ -132,7 +131,7 @@ void FWaterBooleanUtils::ExtractMesh(AVolume* Volume, FDynamicMesh3& Mesh)
 	Mesh.DiscardAttributes();
 
 	UModel* Model = Volume->Brush;
-	FTransform3d XForm(Volume->GetTransform());
+	UE::Geometry::FTransform3d XForm(Volume->GetTransform());
 
 	// Each "BspNode" is a planar polygon, triangulate each polygon and accumulate in a mesh.
 	// Note that this does not make any attempt to weld vertices/edges
@@ -239,10 +238,10 @@ void FWaterBooleanUtils::ApplyBooleanRepairs(FDynamicMesh3& Mesh, double MergeTo
 
 void FWaterBooleanUtils::MakeNormalizationTransform(
 	const FAxisAlignedBox3d& Bounds,
-	FTransform3d& ToNormalizedOut, FTransform3d& FromNormalizedOut)
+	UE::Geometry::FTransform3d& ToNormalizedOut, UE::Geometry::FTransform3d& FromNormalizedOut)
 {
 	double WorldSize = Bounds.MaxDim();
-	FromNormalizedOut = FTransform3d::Identity();
+	FromNormalizedOut = UE::Geometry::FTransform3d::Identity();
 	FromNormalizedOut.SetTranslation(Bounds.Center());
 	FromNormalizedOut.SetScale(WorldSize * FVector3d::One());
 	ToNormalizedOut = FromNormalizedOut.Inverse();
@@ -263,12 +262,12 @@ void FWaterBooleanUtils::SetToFaceNormals(FDynamicMesh3& Mesh)
 	}
 }
 
-FDynamicMesh3 FWaterBooleanUtils::AccumulateExtrusionVolumes(const TArray<AWaterBodyExclusionVolume*>& ExclusionVolumes, const FTransform3d& Transform)
+FDynamicMesh3 FWaterBooleanUtils::AccumulateExtrusionVolumes(const TArray<AWaterBodyExclusionVolume*>& ExclusionVolumes, const UE::Geometry::FTransform3d& Transform)
 {
 	// First step is to boolean-union all the volumes into this mesh
 	FDynamicMesh3 CombinedVolumes;
 
-	FTransform3d IdentityTransform;
+	UE::Geometry::FTransform3d IdentityTransform;
 	for (AWaterBodyExclusionVolume* Volume : ExclusionVolumes)
 	{
 		// convert volume to mesh
@@ -303,6 +302,8 @@ FDynamicMesh3 FWaterBooleanUtils::AccumulateExtrusionVolumes(const TArray<AWater
 
 FDynamicMesh3 FWaterBooleanUtils::MakeBoxCollisionMesh(FAxisAlignedBox3d WorldBoxBounds, const TArray<AWaterBodyExclusionVolume*>& ExclusionVolumes)
 {
+	using FTransform3d = UE::Geometry::FTransform3d;
+
 	// construct a transform that scales the input world-space box to unit bounds, to improve numerical precision
 	FTransform3d ToUnitTransform, FromUnitTransform;
 	MakeNormalizationTransform(WorldBoxBounds, ToUnitTransform, FromUnitTransform);
@@ -1093,7 +1094,7 @@ void FWaterBooleanUtils::BuildOceanCollisionComponents(
 	TArray<TArray<AWaterBodyExclusionVolume*>> ClusteredVolumes;
 	ClusterExclusionVolumes(IntersectingVolumes, ClusteredVolumes, WorldMeshBufferWidth);
 
-	FTransform3d ActorTransformd(ActorTransform);
+	UE::Geometry::FTransform3d ActorTransformd(ActorTransform);
 
 	// for each cluster, clip ocean box to bounds-plus-buffer of cluster and then
 	// compute Boolean subtraction of exclusion volumes from that box
diff --git a/Engine/Plugins/Experimental/Water/Source/Runtime/Private/WaterBooleanUtils.h b/Engine/Plugins/Experimental/Water/Source/Runtime/Private/WaterBooleanUtils.h
--- a/Engine/Plugins/Experimental/Water/Source/Runtime/Private/WaterBooleanUtils.h
+++ b/Engine/Plugins/Experimental/Water/Source/Runtime/Private/WaterBooleanUtils.h
@@ -83,7 +83,7 @@ private:
 		 * Construct transforms to and from a normalized space, for a given bounds.
 		 * ("Normalized space" is such that scaled box fits in origin-centered unit box)
 		 */
-	static void MakeNormalizationTransform(const FAxisAlignedBox3d& Bounds, FTransform3d& ToNormalizedOut, FTransform3d& FromNormalizedOut);
+	static void MakeNormalizationTransform(const FAxisAlignedBox3d& Bounds, UE::Geometry::FTransform3d& ToNormalizedOut, UE::Geometry::FTransform3d& FromNormalizedOut);
 
 		/**
 		 * Utility function to set per-triangle attribute normals on a DynamicMesh
@@ -94,7 +94,7 @@ private:
 		 * Mesh a set of AVolumes and then Boolean-Union them into a FDynamicMesh3. 
 		 * @param Transform this transform is applied to each volume before the Boolean is computed
 		 */
-	static FDynamicMesh3 AccumulateExtrusionVolumes(const TArray<AWaterBodyExclusionVolume*>& ExclusionVolumes, const FTransform3d& Transform);
+	static FDynamicMesh3 AccumulateExtrusionVolumes(const TArray<AWaterBodyExclusionVolume*>& ExclusionVolumes, const UE::Geometry::FTransform3d& Transform);
 
 		/**
 		 * Generate a mesh that is a Box with a set of Volumes Boolean-subtracted.
diff --git a/Engine/Plugins/FX/Niagara/Source/Niagara/Classes/NDISkeletalMeshCommon.h b/Engine/Plugins/FX/Niagara/Source/Niagara/Classes/NDISkeletalMeshCommon.h
--- a/Engine/Plugins/FX/Niagara/Source/Niagara/Classes/NDISkeletalMeshCommon.h
+++ b/Engine/Plugins/FX/Niagara/Source/Niagara/Classes/NDISkeletalMeshCommon.h
@@ -93,10 +93,10 @@ struct FSkinnedPositionAccessorHelper
 	FORCEINLINE FVector3f GetSkinnedVertexPreviousPosition(FSkeletalMeshAccessorHelper& Accessor, int32 VertexIndex) = delete;
 	FORCEINLINE void GetSkinnedTangentBasis(FSkeletalMeshAccessorHelper& Accessor, int32 VertexIndex, FVector3f& OutTangentX, FVector3f& OutTangentY, FVector3f& OutTangentZ) = delete;
 	FORCEINLINE void GetSkinnedPreviousTangentBasis(FSkeletalMeshAccessorHelper& Accessor, int32 VertexIndex, FVector3f& OutTangentX, FVector3f& OutTangentY, FVector3f& OutTangentZ) = delete;
-	FORCEINLINE FVector GetSkinnedBonePosition(FSkeletalMeshAccessorHelper& Accessor, int32 BoneIndex) = delete;
-	FORCEINLINE FVector GetSkinnedBonePreviousPosition(FSkeletalMeshAccessorHelper& Accessor, int32 BoneIndex) = delete;
-	FORCEINLINE_DEBUGGABLE FQuat GetSkinnedBoneRotation(FSkeletalMeshAccessorHelper& Accessor, int32 BoneIndex) = delete;
-	FORCEINLINE_DEBUGGABLE FQuat GetSkinnedBonePreviousRotation(FSkeletalMeshAccessorHelper& Accessor, int32 BoneIndex) = delete;
+	FORCEINLINE FVector3f GetSkinnedBonePosition(FSkeletalMeshAccessorHelper& Accessor, int32 BoneIndex) = delete;
+	FORCEINLINE FVector3f GetSkinnedBonePreviousPosition(FSkeletalMeshAccessorHelper& Accessor, int32 BoneIndex) = delete;
+	FORCEINLINE_DEBUGGABLE FQuat4f GetSkinnedBoneRotation(FSkeletalMeshAccessorHelper& Accessor, int32 BoneIndex) = delete;
+	FORCEINLINE_DEBUGGABLE FQuat4f GetSkinnedBonePreviousRotation(FSkeletalMeshAccessorHelper& Accessor, int32 BoneIndex) = delete;
 };
 
 template<>
@@ -169,19 +169,19 @@ struct FSkinnedPositionAccessorHelper<TNDISkelMesh_SkinningModeInvalid>
 		return GetSkinnedBonePosition(Accessor, BoneIndex);
 	}
 
-	FORCEINLINE_DEBUGGABLE FQuat GetSkinnedBoneRotation(FSkeletalMeshAccessorHelper& Accessor, int32 BoneIndex)
+	FORCEINLINE_DEBUGGABLE FQuat4f GetSkinnedBoneRotation(FSkeletalMeshAccessorHelper& Accessor, int32 BoneIndex)
 	{
 		const int32 NumRealBones = Accessor.Mesh->GetRefSkeleton().GetRawBoneNum();
 		if (BoneIndex < NumRealBones)
 		{
-			return Accessor.Mesh->GetComposedRefPoseMatrix(BoneIndex).GetMatrixWithoutScale().ToQuat();
+			return (FQuat4f)Accessor.Mesh->GetComposedRefPoseMatrix(BoneIndex).GetMatrixWithoutScale().ToQuat();
 		}
 
 		const FTransform& RefTransform = Accessor.Mesh->GetRefSkeleton().GetRefBonePose()[BoneIndex];
-		return RefTransform.GetRotation();
+		return (FQuat4f)RefTransform.GetRotation();
 	}
 
-	FORCEINLINE_DEBUGGABLE FQuat GetSkinnedBonePreviousRotation(FSkeletalMeshAccessorHelper& Accessor, int32 BoneIndex)
+	FORCEINLINE_DEBUGGABLE FQuat4f GetSkinnedBonePreviousRotation(FSkeletalMeshAccessorHelper& Accessor, int32 BoneIndex)
 	{
 		return GetSkinnedBoneRotation(Accessor, BoneIndex);
 	}
@@ -260,19 +260,19 @@ struct FSkinnedPositionAccessorHelper<TNDISkelMesh_SkinningModeNone>
 		return GetSkinnedBonePosition(Accessor, BoneIndex);
 	}
 
-	FORCEINLINE_DEBUGGABLE FQuat GetSkinnedBoneRotation(FSkeletalMeshAccessorHelper& Accessor, int32 BoneIndex)
+	FORCEINLINE_DEBUGGABLE FQuat4f GetSkinnedBoneRotation(FSkeletalMeshAccessorHelper& Accessor, int32 BoneIndex)
 	{
 		const int32 NumRealBones = Accessor.Mesh->GetRefSkeleton().GetRawBoneNum();
 		if (BoneIndex < NumRealBones)
 		{
-			return Accessor.Mesh->GetComposedRefPoseMatrix(BoneIndex).GetMatrixWithoutScale().ToQuat();
+			return (FQuat4f)Accessor.Mesh->GetComposedRefPoseMatrix(BoneIndex).GetMatrixWithoutScale().ToQuat();
 		}
 
 		const FTransform& RefTransform = Accessor.Mesh->GetRefSkeleton().GetRefBonePose()[BoneIndex];
-		return RefTransform.GetRotation();
+		return (FQuat4f)RefTransform.GetRotation();
 	}
 
-	FORCEINLINE_DEBUGGABLE FQuat GetSkinnedBonePreviousRotation(FSkeletalMeshAccessorHelper& Accessor, int32 BoneIndex)
+	FORCEINLINE_DEBUGGABLE FQuat4f GetSkinnedBonePreviousRotation(FSkeletalMeshAccessorHelper& Accessor, int32 BoneIndex)
 	{
 		return GetSkinnedBoneRotation(Accessor, BoneIndex);
 	}
@@ -336,22 +336,22 @@ struct FSkinnedPositionAccessorHelper<TNDISkelMesh_SkinningModeOnTheFly>
 
 	FORCEINLINE_DEBUGGABLE FVector3f GetSkinnedBonePosition(FSkeletalMeshAccessorHelper& Accessor, int32 BoneIndex)
 	{
-		return Accessor.SkinningData->CurrComponentTransforms()[BoneIndex].GetLocation();
+		return (FVector3f)Accessor.SkinningData->CurrComponentTransforms()[BoneIndex].GetLocation();
 	}
 
 	FORCEINLINE_DEBUGGABLE FVector3f GetSkinnedBonePreviousPosition(FSkeletalMeshAccessorHelper& Accessor, int32 BoneIndex)
 	{
-		return Accessor.SkinningData->PrevComponentTransforms()[BoneIndex].GetLocation();
+		return (FVector3f)Accessor.SkinningData->PrevComponentTransforms()[BoneIndex].GetLocation();
 	}
 
-	FORCEINLINE_DEBUGGABLE FQuat GetSkinnedBoneRotation(FSkeletalMeshAccessorHelper& Accessor, int32 BoneIndex)
+	FORCEINLINE_DEBUGGABLE FQuat4f GetSkinnedBoneRotation(FSkeletalMeshAccessorHelper& Accessor, int32 BoneIndex)
 	{
-		return Accessor.SkinningData->CurrComponentTransforms()[BoneIndex].GetRotation();
+		return (FQuat4f)Accessor.SkinningData->CurrComponentTransforms()[BoneIndex].GetRotation();
 	}
 
-	FORCEINLINE_DEBUGGABLE FQuat GetSkinnedBonePreviousRotation(FSkeletalMeshAccessorHelper& Accessor, int32 BoneIndex)
+	FORCEINLINE_DEBUGGABLE FQuat4f GetSkinnedBonePreviousRotation(FSkeletalMeshAccessorHelper& Accessor, int32 BoneIndex)
 	{
-		return Accessor.SkinningData->PrevComponentTransforms()[BoneIndex].GetRotation();
+		return (FQuat4f)Accessor.SkinningData->PrevComponentTransforms()[BoneIndex].GetRotation();
 	}
 };
 
@@ -421,14 +421,14 @@ struct FSkinnedPositionAccessorHelper<TNDISkelMesh_SkinningModePreSkin>
 		return Accessor.SkinningData->PrevComponentTransforms()[BoneIndex].GetLocation();
 	}
 
-	FORCEINLINE_DEBUGGABLE FQuat GetSkinnedBoneRotation(FSkeletalMeshAccessorHelper& Accessor, int32 BoneIndex)
+	FORCEINLINE_DEBUGGABLE FQuat4f GetSkinnedBoneRotation(FSkeletalMeshAccessorHelper& Accessor, int32 BoneIndex)
 	{
-		return Accessor.SkinningData->CurrComponentTransforms()[BoneIndex].GetRotation();
+		return (FQuat4f)Accessor.SkinningData->CurrComponentTransforms()[BoneIndex].GetRotation();
 	}
 
-	FORCEINLINE_DEBUGGABLE FQuat GetSkinnedBonePreviousRotation(FSkeletalMeshAccessorHelper& Accessor, int32 BoneIndex)
+	FORCEINLINE_DEBUGGABLE FQuat4f GetSkinnedBonePreviousRotation(FSkeletalMeshAccessorHelper& Accessor, int32 BoneIndex)
 	{
-		return Accessor.SkinningData->PrevComponentTransforms()[BoneIndex].GetRotation();
+		return (FQuat4f)Accessor.SkinningData->PrevComponentTransforms()[BoneIndex].GetRotation();
 	}
 };
 
diff --git a/Engine/Plugins/FX/Niagara/Source/Niagara/Classes/NiagaraDataInterface.h b/Engine/Plugins/FX/Niagara/Source/Niagara/Classes/NiagaraDataInterface.h
--- a/Engine/Plugins/FX/Niagara/Source/Niagara/Classes/NiagaraDataInterface.h
+++ b/Engine/Plugins/FX/Niagara/Source/Niagara/Classes/NiagaraDataInterface.h
@@ -24,7 +24,8 @@ struct FNDITransformHandlerNoop
 	FORCEINLINE void TransformPosition(FVector3d& V, const FMatrix44d& M) { }
 	FORCEINLINE void TransformVector(FVector3f& V, const FMatrix44f& M) { }
 	FORCEINLINE void TransformVector(FVector3d& V, const FMatrix44d& M) { }
-	FORCEINLINE void TransformRotation(FQuat& Q1, const FQuat& Q2) { }
+	FORCEINLINE void TransformRotation(FQuat4f& Q1, const FQuat4f& Q2) { }
+	FORCEINLINE void TransformRotation(FQuat4d& Q1, const FQuat4d& Q2) { }
 };
 
 struct FNDITransformHandler
@@ -33,7 +34,8 @@ struct FNDITransformHandler
 	FORCEINLINE void TransformPosition(FVector3d& P, const FMatrix44d& M) { P = M.TransformPosition(P); }
 	FORCEINLINE void TransformVector(FVector3f& V, const FMatrix44f& M) { V = M.TransformVector(V).GetUnsafeNormal3(); }
 	FORCEINLINE void TransformVector(FVector3d& V, const FMatrix44d& M) { V = M.TransformVector(V).GetUnsafeNormal3(); }
-	FORCEINLINE void TransformRotation(FQuat& Q1, const FQuat& Q2) { Q1 = Q2 * Q1; }
+	FORCEINLINE void TransformRotation(FQuat4f& Q1, const FQuat4f& Q2) { Q1 = Q2 * Q1; }
+	FORCEINLINE void TransformRotation(FQuat4d& Q1, const FQuat4d& Q2) { Q1 = Q2 * Q1; }
 };
 
 //////////////////////////////////////////////////////////////////////////
@@ -751,14 +753,14 @@ struct FNDIInputParam<FVector4>
 };
 
 template<>
-struct FNDIInputParam<FQuat>
+struct FNDIInputParam<FQuat4f>
 {
 	VectorVM::FExternalFuncInputHandler<float> X;
 	VectorVM::FExternalFuncInputHandler<float> Y;
 	VectorVM::FExternalFuncInputHandler<float> Z;
 	VectorVM::FExternalFuncInputHandler<float> W;
 	FORCEINLINE FNDIInputParam(FVectorVMExternalFunctionContext& Context) : X(Context), Y(Context), Z(Context), W(Context) {}
-	FORCEINLINE FQuat GetAndAdvance() { return FQuat(X.GetAndAdvance(), Y.GetAndAdvance(), Z.GetAndAdvance(), W.GetAndAdvance()); }
+	FORCEINLINE FQuat4f GetAndAdvance() { return FQuat4f(X.GetAndAdvance(), Y.GetAndAdvance(), Z.GetAndAdvance(), W.GetAndAdvance()); }
 };
 
 template<>
@@ -859,7 +861,7 @@ struct FNDIOutputParam<FVector4>
 };
 
 template<>
-struct FNDIOutputParam<FQuat>
+struct FNDIOutputParam<FQuat4f>
 {
 	VectorVM::FExternalFuncRegisterHandler<float> X;
 	VectorVM::FExternalFuncRegisterHandler<float> Y;
@@ -867,7 +869,7 @@ struct FNDIOutputParam<FQuat>
 	VectorVM::FExternalFuncRegisterHandler<float> W;
 	FORCEINLINE FNDIOutputParam(FVectorVMExternalFunctionContext& Context) : X(Context), Y(Context), Z(Context), W(Context) {}
 	FORCEINLINE bool IsValid() const { return X.IsValid() || Y.IsValid() || Z.IsValid() || W.IsValid(); }
-	FORCEINLINE void SetAndAdvance(FQuat Val)
+	FORCEINLINE void SetAndAdvance(FQuat4f Val)
 	{
 		*X.GetDestAndAdvance() = Val.X;
 		*Y.GetDestAndAdvance() = Val.Y;
diff --git a/Engine/Plugins/FX/Niagara/Source/Niagara/Classes/NiagaraDataInterfaceArrayFloat.h b/Engine/Plugins/FX/Niagara/Source/Niagara/Classes/NiagaraDataInterfaceArrayFloat.h
--- a/Engine/Plugins/FX/Niagara/Source/Niagara/Classes/NiagaraDataInterfaceArrayFloat.h
+++ b/Engine/Plugins/FX/Niagara/Source/Niagara/Classes/NiagaraDataInterfaceArrayFloat.h
@@ -33,7 +33,7 @@ class NIAGARA_API UNiagaraDataInterfaceArrayFloat3 : public UNiagaraDataInterfac
 	GENERATED_UCLASS_BODY()
 public:
 	UPROPERTY(EditAnywhere, BlueprintReadWrite, Category = "Array")
-	TArray<FVector> FloatData;		// LWC_TODO: Should be FVector3f, but only FVector is blueprint accessible, so we copy it to FloatDataAsFloat before passing to the interface
+	TArray<FVector> FloatData;		// LWC_TODO: Should be FVector3f, but only FVector is blueprint accessible
 
 	TArray<FVector>& GetArrayReference() { return FloatData; }
 };
diff --git a/Engine/Plugins/FX/Niagara/Source/Niagara/Classes/NiagaraDataInterfaceArrayImpl.h b/Engine/Plugins/FX/Niagara/Source/Niagara/Classes/NiagaraDataInterfaceArrayImpl.h
--- a/Engine/Plugins/FX/Niagara/Source/Niagara/Classes/NiagaraDataInterfaceArrayImpl.h
+++ b/Engine/Plugins/FX/Niagara/Source/Niagara/Classes/NiagaraDataInterfaceArrayImpl.h
@@ -20,6 +20,7 @@ struct FNDIArrayImplHelperBase
 
 	static void GPUGetFetchHLSL(FString& OutHLSL, const TCHAR* BufferName) { OutHLSL.Appendf(TEXT("OutValue = %s[ClampedIndex];"), BufferName); }
 	static int32 GPUGetTypeStride() { return sizeof(TArrayType); }
+	static int32 CPUGetTypeStride() { return sizeof(TArrayType); }
 	static void CopyData(void* Dest, const TArrayType* Src, int32 BufferSize)
 	{
 		FMemory::Memcpy(Dest, Src, BufferSize);
@@ -421,28 +422,27 @@ struct FNiagaraDataInterfaceArrayImpl : public INiagaraDataInterfaceArrayImpl
 				if (RT_Proxy->NumElements > 0)
 				{
 					const int32 BufferStride = T::GPUGetTypeStride();
-					const int32 BufferSize = RT_Array.GetTypeSize() * RT_Array.Num();
+					const int32 BufferSize = T::CPUGetTypeStride() * RT_Array.Num();
 					const int32 BufferNumElements = BufferSize / BufferStride;
-					check((sizeof(TArrayType) % BufferStride) == 0);
-					check(BufferSize == BufferNumElements * BufferStride);
+					check(T::GPUGetTypeStride() <= T::CPUGetTypeStride());
 
 					RT_Proxy->Buffer.Initialize(TEXT("NiagaraArrayFloat"), BufferStride, BufferNumElements, FNDIArrayImplHelper<TArrayType>::PixelFormat, BUF_Static);
 					void* GPUMemory = RHICmdList.LockBuffer(RT_Proxy->Buffer.Buffer, 0, BufferSize, RLM_WriteOnly);
-					T::CopyData(GPUMemory, RT_Array.GetData(), BufferSize);
+					T::CopyData(GPUMemory, RT_Array.GetData(), RT_Array.GetTypeSize() * RT_Array.Num());
 					RHICmdList.UnlockBuffer(RT_Proxy->Buffer.Buffer);
 				}
 				else
 				{
 					const int32 BufferStride = T::GPUGetTypeStride();
-					const int32 BufferSize = RT_Array.GetTypeSize();
-					const int32 BufferNumElements = BufferSize / BufferStride;
-					check((sizeof(TArrayType) % BufferStride) == 0);
+					const int32 BufferSize = T::CPUGetTypeStride();
+					const int32 BufferNumElements = BufferSize / BufferStride;					
+					check(T::GPUGetTypeStride() <= T::CPUGetTypeStride());
 
-					const TArrayType DefaultValue = FNDIArrayImplHelper<TArrayType>::GetDefaultValue();
+					const TArrayType DefaultValue = (TArrayType)FNDIArrayImplHelper<TArrayType>::GetDefaultValue();
 
 					RT_Proxy->Buffer.Initialize(TEXT("NiagaraArrayFloat"), BufferStride, BufferNumElements, FNDIArrayImplHelper<TArrayType>::PixelFormat, BUF_Static);
 					void* GPUMemory = RHICmdList.LockBuffer(RT_Proxy->Buffer.Buffer, 0, BufferSize, RLM_WriteOnly);
-					T::CopyData(GPUMemory, &DefaultValue, BufferSize);
+					T::CopyData(GPUMemory, &DefaultValue, RT_Array.GetTypeSize());
 					RHICmdList.UnlockBuffer(RT_Proxy->Buffer.Buffer);
 				}
 
@@ -555,15 +555,15 @@ struct FNiagaraDataInterfaceArrayImpl : public INiagaraDataInterfaceArrayImpl
 			for (int32 i = 0; i < Context.GetNumInstances(); ++i)
 			{
 				const int32 Index = FMath::Clamp(IndexParam.GetAndAdvance(), 0, Num);
-				OutValue.SetAndAdvance(Data[Index]);
+				OutValue.SetAndAdvance((typename FNDIArrayImplHelper<TArrayType>::TVMArrayType)Data[Index]);
 			}
 		}
 		else
 		{
-			const TArrayType DefaultValue = FNDIArrayImplHelper<TArrayType>::GetDefaultValue();
+			const TArrayType DefaultValue = (TArrayType)FNDIArrayImplHelper<TArrayType>::GetDefaultValue();
 			for (int32 i = 0; i < Context.GetNumInstances(); ++i)
 			{
-				OutValue.SetAndAdvance(DefaultValue);
+				OutValue.SetAndAdvance((typename FNDIArrayImplHelper<TArrayType>::TVMArrayType)DefaultValue);
 			}
 		}
 	}
@@ -591,7 +591,7 @@ struct FNiagaraDataInterfaceArrayImpl : public INiagaraDataInterfaceArrayImpl
 
 		if (NewNum > OldNum)
 		{
-			const TArrayType DefaultValue = FNDIArrayImplHelper<TArrayType>::GetDefaultValue();
+			const TArrayType DefaultValue = (TArrayType)FNDIArrayImplHelper<TArrayType>::GetDefaultValue();
 			for (int32 i = OldNum; i < NewNum; ++i)
 			{
 				Data[i] = DefaultValue;
@@ -609,7 +609,7 @@ struct FNiagaraDataInterfaceArrayImpl : public INiagaraDataInterfaceArrayImpl
 		for (int32 i = 0; i < Context.GetNumInstances(); ++i)
 		{
 			const int32 Index = IndexParam.GetAndAdvance();
-			const TArrayType Value = InValue.GetAndAdvance();
+			const TArrayType Value = (TArrayType)InValue.GetAndAdvance();
 
 			if (Data.IsValidIndex(Index))
 			{
@@ -630,7 +630,7 @@ struct FNiagaraDataInterfaceArrayImpl : public INiagaraDataInterfaceArrayImpl
 		for (int32 i = 0; i < Context.GetNumInstances(); ++i)
 		{
 			const bool bSkipExecute = InSkipExecute.GetAndAdvance();
-			const TArrayType Value = InValue.GetAndAdvance();
+			const TArrayType Value = (TArrayType)InValue.GetAndAdvance();
 			if (!bSkipExecute && (Data.Num() < MaxElements))
 			{
 				Data.Emplace(Value);
@@ -644,7 +644,7 @@ struct FNiagaraDataInterfaceArrayImpl : public INiagaraDataInterfaceArrayImpl
 		FNDIInputParam<FNiagaraBool> InSkipExecute(Context);
 		FNDIOutputParam<typename FNDIArrayImplHelper<TArrayType>::TVMArrayType> OutValue(Context);
 		FNDIOutputParam<FNiagaraBool> OutIsValid(Context);
-		const TArrayType DefaultValue = FNDIArrayImplHelper<TArrayType>::GetDefaultValue();
+		const TArrayType DefaultValue = (TArrayType)FNDIArrayImplHelper<TArrayType>::GetDefaultValue();
 
 		FRWScopeLock WriteLock(Owner->ArrayRWGuard, SLT_Write);
 		for (int32 i=0; i < Context.GetNumInstances(); ++i)
@@ -652,12 +652,12 @@ struct FNiagaraDataInterfaceArrayImpl : public INiagaraDataInterfaceArrayImpl
 			const bool bSkipExecute = InSkipExecute.GetAndAdvance();
 			if (bSkipExecute || (Data.Num() == 0))
 			{
-				OutValue.SetAndAdvance(DefaultValue);
+				OutValue.SetAndAdvance((typename FNDIArrayImplHelper<TArrayType>::TVMArrayType)DefaultValue);
 				OutIsValid.SetAndAdvance(false);
 			}
 			else
 			{
-				OutValue.SetAndAdvance(Data.Pop());
+				OutValue.SetAndAdvance((typename FNDIArrayImplHelper<TArrayType>::TVMArrayType)Data.Pop());
 				OutIsValid.SetAndAdvance(true);
 			}
 		}
diff --git a/Engine/Plugins/FX/Niagara/Source/Niagara/Classes/NiagaraDataInterfaceSkeletalMesh.h b/Engine/Plugins/FX/Niagara/Source/Niagara/Classes/NiagaraDataInterfaceSkeletalMesh.h
--- a/Engine/Plugins/FX/Niagara/Source/Niagara/Classes/NiagaraDataInterfaceSkeletalMesh.h
+++ b/Engine/Plugins/FX/Niagara/Source/Niagara/Classes/NiagaraDataInterfaceSkeletalMesh.h
@@ -171,22 +171,22 @@ public:
 		return BoneRefToLocals[CurrIndex ^ 1];
 	}
 
-	FORCEINLINE TArray<FTransform>& CurrComponentTransforms()
+	FORCEINLINE TArray<FTransform3f>& CurrComponentTransforms()
 	{
 		return ComponentTransforms[CurrIndex];
 	}
 
-	FORCEINLINE const TArray<FTransform>& CurrComponentTransforms() const
+	FORCEINLINE const TArray<FTransform3f>& CurrComponentTransforms() const
 	{
 		return ComponentTransforms[CurrIndex];
 	}
 
-	FORCEINLINE TArray<FTransform>& PrevComponentTransforms()
+	FORCEINLINE TArray<FTransform3f>& PrevComponentTransforms()
 	{
 		return ComponentTransforms[CurrIndex ^ 1];
 	}
 
-	FORCEINLINE const TArray<FTransform>& PrevComponentTransforms() const
+	FORCEINLINE const TArray<FTransform3f>& PrevComponentTransforms() const
 	{
 		return ComponentTransforms[CurrIndex ^ 1];
 	}
@@ -222,7 +222,7 @@ private:
 	TArray<FMatrix44f> BoneRefToLocals[2];
 
 	/** Component space transforms */
-	TArray<FTransform> ComponentTransforms[2];
+	TArray<FTransform3f> ComponentTransforms[2];
 
 	struct FLODData
 	{
@@ -633,7 +633,7 @@ struct FNDISkeletalMesh_InstanceData
 	struct FCachedSocketInfo
 	{
 		FCachedSocketInfo() : BoneIdx(INDEX_NONE){}
-		FTransform Transform;
+		FTransform3f Transform;
 		int32 BoneIdx;
 	};
 	TArray<FCachedSocketInfo> FilteredSocketInfo;
@@ -644,7 +644,7 @@ struct FNDISkeletalMesh_InstanceData
 	/** Index into which socket transforms to use.  */
 	uint32 FilteredSocketTransformsIndex = 0;
 	/** Transforms for sockets. */
-	TStaticArray<TArray<FTransform>, 2> FilteredSocketTransforms;
+	TStaticArray<TArray<FTransform3f>, 2> FilteredSocketTransforms;
 
 	uint32 ChangeId;
 
@@ -697,9 +697,9 @@ struct FNDISkeletalMesh_InstanceData
 	}
 
 	void UpdateFilteredSocketTransforms();
-	TArray<FTransform>& GetFilteredSocketsWriteBuffer() { return FilteredSocketTransforms[FilteredSocketTransformsIndex]; }
-	const TArray<FTransform>& GetFilteredSocketsCurrBuffer() const { return FilteredSocketTransforms[FilteredSocketTransformsIndex]; }
-	const TArray<FTransform>& GetFilteredSocketsPrevBuffer() const { return FilteredSocketTransforms[(FilteredSocketTransformsIndex + 1) % FilteredSocketTransforms.Num()]; }
+	TArray<FTransform3f>& GetFilteredSocketsWriteBuffer() { return FilteredSocketTransforms[FilteredSocketTransformsIndex]; }
+	const TArray<FTransform3f>& GetFilteredSocketsCurrBuffer() const { return FilteredSocketTransforms[FilteredSocketTransformsIndex]; }
+	const TArray<FTransform3f>& GetFilteredSocketsPrevBuffer() const { return FilteredSocketTransforms[(FilteredSocketTransformsIndex + 1) % FilteredSocketTransforms.Num()]; }
 
 	bool HasColorData();
 };
diff --git a/Engine/Plugins/FX/Niagara/Source/Niagara/Classes/NiagaraDataInterfaceStaticMesh.h b/Engine/Plugins/FX/Niagara/Source/Niagara/Classes/NiagaraDataInterfaceStaticMesh.h
--- a/Engine/Plugins/FX/Niagara/Source/Niagara/Classes/NiagaraDataInterfaceStaticMesh.h
+++ b/Engine/Plugins/FX/Niagara/Source/Niagara/Classes/NiagaraDataInterfaceStaticMesh.h
@@ -206,7 +206,7 @@ struct FNDIStaticMesh_InstanceData
 	TRefCountPtr<const FStaticMeshLODResources> CachedLOD;
 
 	/** Cached socket information, if available */
-	TArray<FTransform> CachedSockets;
+	TArray<FTransform3f> CachedSockets;
 
 	/** Number of filtered sockets. */
 	int32 NumFilteredSockets = 0;
diff --git a/Engine/Plugins/FX/Niagara/Source/Niagara/Classes/NiagaraDataSetAccessor.h b/Engine/Plugins/FX/Niagara/Source/Niagara/Classes/NiagaraDataSetAccessor.h
--- a/Engine/Plugins/FX/Niagara/Source/Niagara/Classes/NiagaraDataSetAccessor.h
+++ b/Engine/Plugins/FX/Niagara/Source/Niagara/Classes/NiagaraDataSetAccessor.h
@@ -277,9 +277,9 @@ struct FNiagaraDataSetAccessorTypeInfo<FLinearColor>
 };
 
 template<>
-struct FNiagaraDataSetAccessorTypeInfo<FQuat>
+struct FNiagaraDataSetAccessorTypeInfo<FQuat4f>
 {
-	using TAccessorBaseClass = FNiagaraDataSetAccessorFloat<FQuat>;
+	using TAccessorBaseClass = FNiagaraDataSetAccessorFloat<FQuat4f>;
 
 	static constexpr bool bSupportsHalf = false;
 	static constexpr int32 NumElements = 4;
