LWC support for FQuat and FTransform

diff --git a/Engine/Plugins/FX/Niagara/Source/Niagara/Private/DataInterface/NiagaraDataInterfaceDebugDraw.cpp b/Engine/Plugins/FX/Niagara/Source/Niagara/Private/DataInterface/NiagaraDataInterfaceDebugDraw.cpp
--- a/Engine/Plugins/FX/Niagara/Source/Niagara/Private/DataInterface/NiagaraDataInterfaceDebugDraw.cpp
+++ b/Engine/Plugins/FX/Niagara/Source/Niagara/Private/DataInterface/NiagaraDataInterfaceDebugDraw.cpp
@@ -55,10 +55,10 @@ struct FNDIDebugDrawInstanceData_GameThread
 
 	void AddSphere(const FVector3f& Location, float Radius, int32 Segments, const FLinearColor& Color)
 	{
-		AddSphere(Location, FVector3f(Radius), FQuat::Identity, FQuat::Identity, FVector3f::OneVector, Segments, true, true, true, Color);
+		AddSphere(Location, FVector3f(Radius), FQuat4f::Identity, FQuat4f::Identity, FVector3f::OneVector, Segments, true, true, true, Color);
 	}
 
-	void AddSphere(const FVector3f& Location, const FVector3f& Radius, const FQuat& AxisRotator, const FQuat& WorldRotate, const FVector3f& WorldScale, int32 Segments, bool bHemiX, bool bHemiY, bool bHemiZ, const FLinearColor& Color)
+	void AddSphere(const FVector3f& Location, const FVector3f& Radius, const FQuat4f& AxisRotator, const FQuat4f& WorldRotate, const FVector3f& WorldScale, int32 Segments, bool bHemiX, bool bHemiY, bool bHemiZ, const FLinearColor& Color)
 	{
 		const float uinc = 2.0f * PI / float(Segments);
 
@@ -66,7 +66,7 @@ struct FNDIDebugDrawInstanceData_GameThread
 		float SinX0 = FMath::Sin(ux);
 		float CosX0 = FMath::Cos(ux);
 		FVector3f LerpVector = FVector3f(bHemiX ? 1.0f : 0.0f, bHemiY ? 1.0f : 0.0f, bHemiZ ? 1.0f : 0.0f);
-		FTransform LocalToWorld;
+		FTransform3f LocalToWorld;
 		LocalToWorld.SetComponents(WorldRotate, FVector3f::ZeroVector, WorldScale);
 
 		for (int x = 0; x < Segments; ++x)
@@ -115,14 +115,14 @@ struct FNDIDebugDrawInstanceData_GameThread
 		}
 	}
 
-	void AddBox(const FVector3f& Location, const FQuat& Rotation, const FVector3f& Extents, const FLinearColor& Color)
+	void AddBox(const FVector3f& Location, const FQuat4f& Rotation, const FVector3f& Extents, const FLinearColor& Color)
 	{
-		AddBox(Location, Rotation, Extents, FQuat::Identity, FVector3f::OneVector, Color);
+		AddBox(Location, Rotation, Extents, FQuat4f::Identity, FVector3f::OneVector, Color);
 	}
 
-	void AddBox(const FVector3f& Location, const FQuat& Rotation, const FVector3f& Extents, const FQuat& WorldRotate, const FVector3f& WorldScale, const FLinearColor& Color)
+	void AddBox(const FVector3f& Location, const FQuat4f& Rotation, const FVector3f& Extents, const FQuat4f& WorldRotate, const FVector3f& WorldScale, const FLinearColor& Color)
 	{
-		FTransform LocalToWorld;
+		FTransform3f LocalToWorld;
 		LocalToWorld.SetComponents(WorldRotate, FVector3f::ZeroVector, WorldScale);
 		const FVector3f Points[] =
 		{
@@ -205,12 +205,12 @@ struct FNDIDebugDrawInstanceData_GameThread
 
 	void AddCylinder(const FVector3f& Location, const FVector3f& Axis, float Height, float Radius, int32 NumHeightSegments, int32 NumRadiusSegments, const FLinearColor& Color)
 	{
-		FQuat Rotation = FQuat::FindBetweenNormals(FVector3f::UpVector, Axis);
+		FQuat4f Rotation = FQuat4f::FindBetweenNormals(FVector3f::UpVector, Axis);
 
 		AddCylinder(Location, Rotation, Height, Radius, NumHeightSegments, NumRadiusSegments, Color, false, false, FVector3f::OneVector, true);
 	}
 
-	void AddCylinder(const FVector3f& Location, const FQuat& Rotation, float Height, float Radius, int32 NumHeightSegments, int32 NumRadiusSegments, const FLinearColor& Color, bool bHemiXDraw, bool bHemiYDraw, const FVector3f& NonUniformScale, bool bCenterVertically)
+	void AddCylinder(const FVector3f& Location, const FQuat4f& Rotation, float Height, float Radius, int32 NumHeightSegments, int32 NumRadiusSegments, const FLinearColor& Color, bool bHemiXDraw, bool bHemiYDraw, const FVector3f& NonUniformScale, bool bCenterVertically)
 	{
 		const FVector3f HeightVector = FVector3f::ZAxisVector * Height;
 		const FVector3f HalfHeightOffset = bCenterVertically? HeightVector * 0.5f : FVector3f::ZeroVector;
@@ -253,12 +253,12 @@ struct FNDIDebugDrawInstanceData_GameThread
 
 	void AddCone(const FVector3f& Location, const FVector3f& Axis, float Height, float RadiusTop, float RadiusBottom, int32 NumHeightSegments, int32 NumRadiusSegments, const FLinearColor& Color)
 	{
-		FQuat Rotation = FQuat::FindBetweenNormals(FVector3f::UpVector, Axis);
+		FQuat4f Rotation = FQuat4f::FindBetweenNormals(FVector3f::UpVector, Axis);
 
 		AddCone(Location, Rotation, Height, RadiusTop, RadiusBottom, NumHeightSegments, NumRadiusSegments, Color, FVector3f::OneVector, true);
 	}
 
-	void AddCone(const FVector3f& Location, const FQuat& Rotation, float Height, float RadiusTop, float RadiusBottom, int32 NumHeightSegments, int32 NumRadiusSegments, const FLinearColor& Color, const FVector3f& NonUniformScale, bool bCenterVertically)
+	void AddCone(const FVector3f& Location, const FQuat4f& Rotation, float Height, float RadiusTop, float RadiusBottom, int32 NumHeightSegments, int32 NumRadiusSegments, const FLinearColor& Color, const FVector3f& NonUniformScale, bool bCenterVertically)
 	{
 		const FVector3f HeightVector = FVector3f::ZAxisVector * Height;
 		const FVector3f HalfHeightOffset = bCenterVertically ? HeightVector * 0.5f : FVector3f::ZeroVector;
@@ -309,7 +309,7 @@ struct FNDIDebugDrawInstanceData_GameThread
 	{
 		const FVector3f AxisNorm = Axis.GetSafeNormal();
 
-		const FQuat AxisRotation = FQuat::FindBetweenNormals(FVector3f::UpVector, AxisNorm);
+		const FQuat4f AxisRotation = FQuat4f::FindBetweenNormals(FVector3f::UpVector, AxisNorm);
 
 		const FVector3f TangentX = AxisRotation.RotateVector(FVector3f::XAxisVector);
 		const FVector3f TangentY = AxisRotation.RotateVector(FVector3f::YAxisVector);
@@ -348,7 +348,7 @@ struct FNDIDebugDrawInstanceData_GameThread
 
 	}
 
-	void AddCoordinateSystem(const FVector3f& Location, const FQuat& Rotation, float Scale)
+	void AddCoordinateSystem(const FVector3f& Location, const FQuat4f& Rotation, float Scale)
 	{
 		const FVector3f XAxis = Rotation.RotateVector(FVector3f(Scale, 0.0f, 0.0f));
 		const FVector3f YAxis = Rotation.RotateVector(FVector3f(0.0f, Scale, 0.0f));
@@ -359,7 +359,7 @@ struct FNDIDebugDrawInstanceData_GameThread
 		AddLine(Location, Location + ZAxis, FLinearColor::Blue);
 	}
 
-	void AddGrid2D(const FVector3f& Center, const FQuat& Rotation, const FVector2D& Extents, const FIntPoint& NumCells, const FLinearColor& Color)
+	void AddGrid2D(const FVector3f& Center, const FQuat4f& Rotation, const FVector2D& Extents, const FIntPoint& NumCells, const FLinearColor& Color)
 	{
 		const FVector3f Corner = Center - Rotation.RotateVector(FVector3f(Extents.X, Extents.Y, 0.0f));
 		const FVector3f XLength = Rotation.RotateVector(FVector3f(Extents.X * 2.0f, 0.0f, 0.0f));
@@ -379,7 +379,7 @@ struct FNDIDebugDrawInstanceData_GameThread
 		}
 	}
 
-	void AddGrid3D(const FVector3f& Center, const FQuat& Rotation, const FVector3f& Extents, const FIntVector& NumCells, const FLinearColor& Color)
+	void AddGrid3D(const FVector3f& Center, const FQuat4f& Rotation, const FVector3f& Extents, const FIntVector& NumCells, const FLinearColor& Color)
 	{
 		const FVector3f Corner = Center - Rotation.RotateVector(Extents);
 		const FVector3f XLength = Rotation.RotateVector(FVector3f(Extents.X * 2.0f, 0.0f, 0.0f));
@@ -515,7 +515,7 @@ struct FNDIDebugDrawInstanceData_GameThread
 			}
 		}
 
-		void TransformQuat(bool bRotationWasSet, FQuat& Quat, const TOptional<ENiagaraCoordinateSpace>& SourceSpace, ENiagaraCoordinateSpace DestSpace, const FNiagaraSystemInstance* SystemInstance)
+		void TransformQuat(bool bRotationWasSet, FQuat4f& Quat, const TOptional<ENiagaraCoordinateSpace>& SourceSpace, ENiagaraCoordinateSpace DestSpace, const FNiagaraSystemInstance* SystemInstance)
 		{
 			ENiagaraCoordinateSpace SourceSpaceConcrete = GetConcreteSource(bRotationWasSet, SourceSpace);
 			ENiagaraCoordinateSpace DestSpaceConcrete = GetConcreteSource(true, DestSpace);
@@ -526,11 +526,11 @@ struct FNDIDebugDrawInstanceData_GameThread
 
 			if (SourceSpaceConcrete == ENiagaraCoordinateSpace::Local && DestSpaceConcrete == ENiagaraCoordinateSpace::World)
 			{
-				Quat = SystemInstance->GetWorldTransform().Rotator().Quaternion() * Quat;
+				Quat = (FQuat4f)SystemInstance->GetWorldTransform().Rotator().Quaternion() * Quat;
 			}
 			else if (SourceSpaceConcrete == ENiagaraCoordinateSpace::World && DestSpaceConcrete == ENiagaraCoordinateSpace::Local)
 			{
-				Quat = SystemInstance->GetWorldTransform().Rotator().GetInverse().Quaternion() * Quat;
+				Quat = (FQuat4f)SystemInstance->GetWorldTransform().Rotator().GetInverse().Quaternion() * Quat;
 			}
 			else
 			{
@@ -592,9 +592,9 @@ struct FNDIDebugDrawInstanceData_GameThread
 				return;
 			static FNiagaraTypeDefinition CoordTypeDef = FNiagaraTypeDefinition(FNiagaraTypeDefinition::GetCoordinateSpaceEnum());
 
-			TOptional<FVector> StartLocation = GetTag<FVector>(SystemInstance, FNiagaraTypeDefinition::GetVec3Def(), 0);
+			TOptional<FVector3f> StartLocation = GetTag<FVector3f>(SystemInstance, FNiagaraTypeDefinition::GetVec3Def(), 0);
 			TOptional<ENiagaraCoordinateSpace> StartLocationCoordinateSpace = GetTag<ENiagaraCoordinateSpace>(SystemInstance, CoordTypeDef, 1);
-			TOptional<FVector> EndLocation = GetTag<FVector>(SystemInstance, FNiagaraTypeDefinition::GetVec3Def(), 2);
+			TOptional<FVector3f> EndLocation = GetTag<FVector3f>(SystemInstance, FNiagaraTypeDefinition::GetVec3Def(), 2);
 			TOptional<ENiagaraCoordinateSpace> EndLocationCoordinateSpace = GetTag<ENiagaraCoordinateSpace>(SystemInstance, CoordTypeDef, 3);
 			TOptional<FLinearColor> Color = GetTag<FLinearColor>(SystemInstance, FNiagaraTypeDefinition::GetColorDef(), 4);
 
@@ -617,14 +617,14 @@ struct FNDIDebugDrawInstanceData_GameThread
 				return;
 			static FNiagaraTypeDefinition CoordTypeDef = FNiagaraTypeDefinition(FNiagaraTypeDefinition::GetCoordinateSpaceEnum());
 
-			TOptional<FVector> Center = GetTag<FVector>(SystemInstance, FNiagaraTypeDefinition::GetVec3Def(), 0);
+			TOptional<FVector3f> Center = GetTag<FVector3f>(SystemInstance, FNiagaraTypeDefinition::GetVec3Def(), 0);
 			TOptional<ENiagaraCoordinateSpace> CenterWorldSpace = GetTag<ENiagaraCoordinateSpace>(SystemInstance, CoordTypeDef, 1);
-			TOptional<FVector> Offset = GetTag<FVector>(SystemInstance, FNiagaraTypeDefinition::GetVec3Def(), 2);
+			TOptional<FVector3f> Offset = GetTag<FVector3f>(SystemInstance, FNiagaraTypeDefinition::GetVec3Def(), 2);
 			TOptional<ENiagaraCoordinateSpace> OffsetWorldSpace = GetTag<ENiagaraCoordinateSpace>(SystemInstance, CoordTypeDef, 3);
 			TOptional<FVector2D> Extents = GetTag<FVector2D>(SystemInstance, FNiagaraTypeDefinition::GetVec2Def(), 4);
 			TOptional<FNiagaraBool> HalfExtents = GetTag<FNiagaraBool>(SystemInstance, FNiagaraTypeDefinition::GetBoolDef(), 5);
-			TOptional<FVector> XAxis = GetTag<FVector>(SystemInstance, FNiagaraTypeDefinition::GetVec3Def(), 6);
-			TOptional<FVector> YAxis = GetTag<FVector>(SystemInstance, FNiagaraTypeDefinition::GetVec3Def(), 7);
+			TOptional<FVector3f> XAxis = GetTag<FVector3f>(SystemInstance, FNiagaraTypeDefinition::GetVec3Def(), 6);
+			TOptional<FVector3f> YAxis = GetTag<FVector3f>(SystemInstance, FNiagaraTypeDefinition::GetVec3Def(), 7);
 			TOptional<ENiagaraCoordinateSpace> AxisCoordinateSpace = GetTag<ENiagaraCoordinateSpace>(SystemInstance, CoordTypeDef, 8);
 			TOptional<int32> NumXSegments = GetTag<int32>(SystemInstance, FNiagaraTypeDefinition::GetIntDef(), 9);
 			TOptional<int32> NumYSegments = GetTag<int32>(SystemInstance, FNiagaraTypeDefinition::GetIntDef(), 10);
@@ -664,13 +664,13 @@ struct FNDIDebugDrawInstanceData_GameThread
 				return;
 			static FNiagaraTypeDefinition CoordTypeDef = FNiagaraTypeDefinition(FNiagaraTypeDefinition::GetCoordinateSpaceEnum());
 
-			TOptional<FVector> Center = GetTag<FVector>(SystemInstance, FNiagaraTypeDefinition::GetVec3Def(), 0);
+			TOptional<FVector3f> Center = GetTag<FVector3f>(SystemInstance, FNiagaraTypeDefinition::GetVec3Def(), 0);
 			TOptional<ENiagaraCoordinateSpace> CenterWorldSpace = GetTag<ENiagaraCoordinateSpace>(SystemInstance, CoordTypeDef, 1);
-			TOptional<FVector> Offset = GetTag<FVector>(SystemInstance, FNiagaraTypeDefinition::GetVec3Def(), 2);
+			TOptional<FVector3f> Offset = GetTag<FVector3f>(SystemInstance, FNiagaraTypeDefinition::GetVec3Def(), 2);
 			TOptional<ENiagaraCoordinateSpace> OffsetWorldSpace = GetTag<ENiagaraCoordinateSpace>(SystemInstance, CoordTypeDef, 3);
 			TOptional<float> Radius = GetTag<float>(SystemInstance, FNiagaraTypeDefinition::GetFloatDef(), 4);
-			TOptional<FVector> XAxis = GetTag<FVector>(SystemInstance, FNiagaraTypeDefinition::GetVec3Def(), 5);
-			TOptional<FVector> YAxis = GetTag<FVector>(SystemInstance, FNiagaraTypeDefinition::GetVec3Def(), 6);
+			TOptional<FVector3f> XAxis = GetTag<FVector3f>(SystemInstance, FNiagaraTypeDefinition::GetVec3Def(), 5);
+			TOptional<FVector3f> YAxis = GetTag<FVector3f>(SystemInstance, FNiagaraTypeDefinition::GetVec3Def(), 6);
 			TOptional<ENiagaraCoordinateSpace> AxisCoordinateSpace = GetTag<ENiagaraCoordinateSpace>(SystemInstance, CoordTypeDef, 7);
 			TOptional<int32> NumSegments = GetTag<int32>(SystemInstance, FNiagaraTypeDefinition::GetIntDef(), 8);
 			TOptional<FLinearColor> Color = GetTag<FLinearColor>(SystemInstance, FNiagaraTypeDefinition::GetColorDef(), 9);
@@ -702,14 +702,14 @@ struct FNDIDebugDrawInstanceData_GameThread
 				return;
 			static FNiagaraTypeDefinition CoordTypeDef = FNiagaraTypeDefinition(FNiagaraTypeDefinition::GetCoordinateSpaceEnum());
 
-			TOptional<FVector> Center = GetTag<FVector>(SystemInstance, FNiagaraTypeDefinition::GetVec3Def(), 0);
+			TOptional<FVector3f> Center = GetTag<FVector3f>(SystemInstance, FNiagaraTypeDefinition::GetVec3Def(), 0);
 			TOptional<ENiagaraCoordinateSpace> CenterWorldSpace = GetTag<ENiagaraCoordinateSpace>(SystemInstance, CoordTypeDef, 1);
-			TOptional<FVector> Offset = GetTag<FVector>(SystemInstance, FNiagaraTypeDefinition::GetVec3Def(), 2);
+			TOptional<FVector3f> Offset = GetTag<FVector3f>(SystemInstance, FNiagaraTypeDefinition::GetVec3Def(), 2);
 			TOptional<ENiagaraCoordinateSpace> OffsetWorldSpace = GetTag<ENiagaraCoordinateSpace>(SystemInstance, CoordTypeDef, 3);
-			TOptional<FVector> RotationAxis = GetTag<FVector>(SystemInstance, FNiagaraTypeDefinition::GetVec3Def(), 4);
+			TOptional<FVector3f> RotationAxis = GetTag<FVector3f>(SystemInstance, FNiagaraTypeDefinition::GetVec3Def(), 4);
 			TOptional<float> RotationNormalizedAngle = GetTag<float>(SystemInstance, FNiagaraTypeDefinition::GetFloatDef(), 5);
 			TOptional<ENiagaraCoordinateSpace> RotationWorldSpace = GetTag<ENiagaraCoordinateSpace>(SystemInstance, CoordTypeDef, 6);
-			TOptional<FVector> Extents = GetTag<FVector>(SystemInstance, FNiagaraTypeDefinition::GetVec3Def(), 7);
+			TOptional<FVector3f> Extents = GetTag<FVector3f>(SystemInstance, FNiagaraTypeDefinition::GetVec3Def(), 7);
 			TOptional<FNiagaraBool> HalfExtents = GetTag<FNiagaraBool>(SystemInstance, FNiagaraTypeDefinition::GetBoolDef(), 8);
 			TOptional<FLinearColor> Color = GetTag<FLinearColor>(SystemInstance, FNiagaraTypeDefinition::GetColorDef(), 9);
 
@@ -720,7 +720,7 @@ struct FNDIDebugDrawInstanceData_GameThread
 			FVector3f DrawRotationAxis = RotationAxis.Get(FVector3f::ZAxisVector);
 			float DrawRotationNormalizedAngle = RotationNormalizedAngle.Get(0.0f);
 			FLinearColor DrawColor = Color.Get(FLinearColor::Green);
-			FQuat DrawRotation = FQuat::Identity;
+			FQuat4f DrawRotation = FQuat4f::Identity;
 
 
 			// DrawCenter and DrawOffset both are brought in and are used to locate the center of the box.
@@ -729,9 +729,9 @@ struct FNDIDebugDrawInstanceData_GameThread
 			TransformPosition(Center.IsSet(), DrawCenter, CenterWorldSpace, ENiagaraCoordinateSpace::Simulation, SystemInstance);
 			TransformVector(Offset.IsSet(), DrawOffset, OffsetWorldSpace, ENiagaraCoordinateSpace::Simulation, SystemInstance);
 
-			DrawRotation = FQuat(DrawRotationAxis, FMath::DegreesToRadians(DrawRotationNormalizedAngle * 360.0f));
+			DrawRotation = FQuat4f(DrawRotationAxis, FMath::DegreesToRadians(DrawRotationNormalizedAngle * 360.0f));
 
-			FQuat WorldRot = FQuat::Identity;
+			FQuat4f WorldRot = FQuat4f::Identity;
 			TransformQuat(true, WorldRot, RotationWorldSpace, ENiagaraCoordinateSpace::World, SystemInstance);
 			FVector3f WorldScale = FVector3f::OneVector;
 			ScaleVector(true, WorldScale, RotationWorldSpace, ENiagaraCoordinateSpace::World, SystemInstance);
@@ -756,19 +756,19 @@ struct FNDIDebugDrawInstanceData_GameThread
 				return;
 			static FNiagaraTypeDefinition CoordTypeDef = FNiagaraTypeDefinition(FNiagaraTypeDefinition::GetCoordinateSpaceEnum());
 
-			TOptional<FVector> Center = GetTag<FVector>(SystemInstance, FNiagaraTypeDefinition::GetVec3Def(), 0);
+			TOptional<FVector3f> Center = GetTag<FVector3f>(SystemInstance, FNiagaraTypeDefinition::GetVec3Def(), 0);
 			TOptional<ENiagaraCoordinateSpace> CenterWorldSpace = GetTag<ENiagaraCoordinateSpace>(SystemInstance, CoordTypeDef, 1);
-			TOptional<FVector> Offset = GetTag<FVector>(SystemInstance, FNiagaraTypeDefinition::GetVec3Def(), 2);
+			TOptional<FVector3f> Offset = GetTag<FVector3f>(SystemInstance, FNiagaraTypeDefinition::GetVec3Def(), 2);
 			TOptional<ENiagaraCoordinateSpace> OffsetWorldSpace = GetTag<ENiagaraCoordinateSpace>(SystemInstance, CoordTypeDef, 3);
 			TOptional<float> Radius = GetTag<float>(SystemInstance, FNiagaraTypeDefinition::GetFloatDef(), 4);
 			TOptional<ENiagaraCoordinateSpace> RadiusWorldSpace = GetTag<ENiagaraCoordinateSpace>(SystemInstance, CoordTypeDef, 5);
 			TOptional<FNiagaraBool> HemiX = GetTag<FNiagaraBool>(SystemInstance, FNiagaraTypeDefinition::GetBoolDef(), 6);
 			TOptional<FNiagaraBool> HemiY = GetTag<FNiagaraBool>(SystemInstance, FNiagaraTypeDefinition::GetBoolDef(), 7);
 			TOptional<FNiagaraBool> HemiZ = GetTag<FNiagaraBool>(SystemInstance, FNiagaraTypeDefinition::GetBoolDef(), 8);
-			TOptional<FVector> OrientationAxisForSphere = GetTag<FVector>(SystemInstance, FNiagaraTypeDefinition::GetVec3Def(), 9);
+			TOptional<FVector3f> OrientationAxisForSphere = GetTag<FVector3f>(SystemInstance, FNiagaraTypeDefinition::GetVec3Def(), 9);
 			TOptional<ENiagaraCoordinateSpace> RotationCoordinateSpace = GetTag<ENiagaraCoordinateSpace>(SystemInstance, CoordTypeDef, 10);
-			TOptional<FQuat> AdditionalRotation = GetTag<FQuat>(SystemInstance, FNiagaraTypeDefinition::GetQuatDef(), 11);
-			TOptional<FVector> NonUniformScale = GetTag<FVector>(SystemInstance, FNiagaraTypeDefinition::GetVec3Def(), 12);
+			TOptional<FQuat4f> AdditionalRotation = GetTag<FQuat4f>(SystemInstance, FNiagaraTypeDefinition::GetQuatDef(), 11);
+			TOptional<FVector3f> NonUniformScale = GetTag<FVector3f>(SystemInstance, FNiagaraTypeDefinition::GetVec3Def(), 12);
 			TOptional<int32> NumSegments = GetTag<int32>(SystemInstance, FNiagaraTypeDefinition::GetIntDef(), 13);
 			TOptional<FLinearColor> Color = GetTag<FLinearColor>(SystemInstance, FNiagaraTypeDefinition::GetColorDef(), 14);
 
@@ -791,9 +791,9 @@ struct FNDIDebugDrawInstanceData_GameThread
 			TransformPosition(Center.IsSet(), DrawCenter, CenterWorldSpace, ENiagaraCoordinateSpace::Simulation, SystemInstance);
 			TransformVector(Offset.IsSet(), DrawOffset, OffsetWorldSpace, ENiagaraCoordinateSpace::Simulation, SystemInstance);
 
-			FQuat OrientationAxis = FQuat::FindBetweenVectors(FVector3f(1.0f, 0.0f, 0.0f), DrawOrientationAxis) * AdditionalRotation.Get(FQuat::Identity);
+			FQuat4f OrientationAxis = FQuat4f::FindBetweenVectors(FVector3f(1.0f, 0.0f, 0.0f), DrawOrientationAxis) * AdditionalRotation.Get(FQuat4f::Identity);
 
-			FQuat WorldRot = FQuat::Identity;
+			FQuat4f WorldRot = FQuat4f::Identity;
 			TransformQuat(true, WorldRot, RadiusWorldSpace, ENiagaraCoordinateSpace::World, SystemInstance);
 			FVector3f WorldScale = FVector3f::OneVector;
 			ScaleVector(true, WorldScale, RadiusWorldSpace, ENiagaraCoordinateSpace::World, SystemInstance);
@@ -815,9 +815,9 @@ struct FNDIDebugDrawInstanceData_GameThread
 			static FNiagaraTypeDefinition CoordTypeDef = FNiagaraTypeDefinition(FNiagaraTypeDefinition::GetCoordinateSpaceEnum());
 			static FNiagaraTypeDefinition OrientationAxisTypeDef = FNiagaraTypeDefinition(FNiagaraTypeDefinition::GetOrientationAxisEnum());
 
-			TOptional<FVector> Center = GetTag<FVector>(SystemInstance, FNiagaraTypeDefinition::GetVec3Def(), 0);
+			TOptional<FVector3f> Center = GetTag<FVector3f>(SystemInstance, FNiagaraTypeDefinition::GetVec3Def(), 0);
 			TOptional<ENiagaraCoordinateSpace> CenterWorldSpace = GetTag<ENiagaraCoordinateSpace>(SystemInstance, CoordTypeDef, 1);
-			TOptional<FVector> Offset = GetTag<FVector>(SystemInstance, FNiagaraTypeDefinition::GetVec3Def(), 2);
+			TOptional<FVector3f> Offset = GetTag<FVector3f>(SystemInstance, FNiagaraTypeDefinition::GetVec3Def(), 2);
 			TOptional<ENiagaraCoordinateSpace> OffsetWorldSpace = GetTag<ENiagaraCoordinateSpace>(SystemInstance, CoordTypeDef, 3);
 			TOptional<float> Radius = GetTag<float>(SystemInstance, FNiagaraTypeDefinition::GetFloatDef(), 4);
 			TOptional<float> Height = GetTag<float>(SystemInstance, FNiagaraTypeDefinition::GetFloatDef(), 5);
@@ -826,7 +826,7 @@ struct FNDIDebugDrawInstanceData_GameThread
 			TOptional<FNiagaraBool> HemiY = GetTag<FNiagaraBool>(SystemInstance, FNiagaraTypeDefinition::GetBoolDef(), 8);
 			TOptional<ENiagaraOrientationAxis> OrientationAxis = GetTag<ENiagaraOrientationAxis>(SystemInstance, OrientationAxisTypeDef, 9);
 			TOptional<ENiagaraCoordinateSpace> OrientationAxisCoordinateSpace = GetTag<ENiagaraCoordinateSpace>(SystemInstance, CoordTypeDef, 10);
-			TOptional<FVector> NonUniformScale = GetTag<FVector>(SystemInstance, FNiagaraTypeDefinition::GetVec3Def(), 11);
+			TOptional<FVector3f> NonUniformScale = GetTag<FVector3f>(SystemInstance, FNiagaraTypeDefinition::GetVec3Def(), 11);
 			TOptional<int32> NumRadiusSegments = GetTag<int32>(SystemInstance, FNiagaraTypeDefinition::GetIntDef(), 12);
 			TOptional<int32> NumHeightSegments = GetTag<int32>(SystemInstance, FNiagaraTypeDefinition::GetIntDef(), 13);
 			TOptional<FLinearColor> Color = GetTag<FLinearColor>(SystemInstance, FNiagaraTypeDefinition::GetColorDef(), 14);
@@ -847,10 +847,10 @@ struct FNDIDebugDrawInstanceData_GameThread
 			TransformPosition(Center.IsSet(), DrawCenter, CenterWorldSpace, ENiagaraCoordinateSpace::World, SystemInstance);
 			TransformVector(Offset.IsSet(), DrawOffset, OffsetWorldSpace, ENiagaraCoordinateSpace::World, SystemInstance);
 
-			FQuat WorldRot =
-				DrawOrientationAxisEnum == ENiagaraOrientationAxis::XAxis ? FQuat(FRotator(0, 90, -90)) :
-				DrawOrientationAxisEnum == ENiagaraOrientationAxis::YAxis ? FQuat(FRotator(0, 0, -90)) :
-				FQuat::Identity;
+			FQuat4f WorldRot =
+				DrawOrientationAxisEnum == ENiagaraOrientationAxis::XAxis ? FQuat4f(FRotator(0, 90, -90)) :
+				DrawOrientationAxisEnum == ENiagaraOrientationAxis::YAxis ? FQuat4f(FRotator(0, 0, -90)) :
+				FQuat4f::Identity;
 
 			TransformQuat(true, WorldRot, OrientationAxisCoordinateSpace, ENiagaraCoordinateSpace::World, SystemInstance);
 
@@ -871,15 +871,15 @@ struct FNDIDebugDrawInstanceData_GameThread
 				return;
 			static FNiagaraTypeDefinition CoordTypeDef = FNiagaraTypeDefinition(FNiagaraTypeDefinition::GetCoordinateSpaceEnum());
 
-			TOptional<FVector> Center = GetTag<FVector>(SystemInstance, FNiagaraTypeDefinition::GetVec3Def(), 0);
+			TOptional<FVector3f> Center = GetTag<FVector3f>(SystemInstance, FNiagaraTypeDefinition::GetVec3Def(), 0);
 			TOptional<ENiagaraCoordinateSpace> CenterWorldSpace = GetTag<ENiagaraCoordinateSpace>(SystemInstance, CoordTypeDef, 1);
-			TOptional<FVector> Offset = GetTag<FVector>(SystemInstance, FNiagaraTypeDefinition::GetVec3Def(), 2);
+			TOptional<FVector3f> Offset = GetTag<FVector3f>(SystemInstance, FNiagaraTypeDefinition::GetVec3Def(), 2);
 			TOptional<ENiagaraCoordinateSpace> OffsetWorldSpace = GetTag<ENiagaraCoordinateSpace>(SystemInstance, CoordTypeDef, 3);
 			TOptional<float> Angle = GetTag<float>(SystemInstance, FNiagaraTypeDefinition::GetFloatDef(), 4);
 			TOptional<float> Length = GetTag<float>(SystemInstance, FNiagaraTypeDefinition::GetFloatDef(), 5);
-			TOptional<FVector> OrientationAxis = GetTag<FVector>(SystemInstance, FNiagaraTypeDefinition::GetVec3Def(), 6);
+			TOptional<FVector3f> OrientationAxis = GetTag<FVector3f>(SystemInstance, FNiagaraTypeDefinition::GetVec3Def(), 6);
 			TOptional<ENiagaraCoordinateSpace> OrientationAxisCoordinateSpace = GetTag<ENiagaraCoordinateSpace>(SystemInstance, CoordTypeDef, 7);
-			TOptional<FVector> NonUniformScale = GetTag<FVector>(SystemInstance, FNiagaraTypeDefinition::GetVec3Def(), 8);
+			TOptional<FVector3f> NonUniformScale = GetTag<FVector3f>(SystemInstance, FNiagaraTypeDefinition::GetVec3Def(), 8);
 			TOptional<int32> NumRadiusSegments = GetTag<int32>(SystemInstance, FNiagaraTypeDefinition::GetIntDef(), 9);
 			TOptional<int32> NumHeightSegments = GetTag<int32>(SystemInstance, FNiagaraTypeDefinition::GetIntDef(), 10);
 			TOptional<FLinearColor> Color = GetTag<FLinearColor>(SystemInstance, FNiagaraTypeDefinition::GetColorDef(), 11);
@@ -898,7 +898,7 @@ struct FNDIDebugDrawInstanceData_GameThread
 			TransformVector(Offset.IsSet(), DrawOffset, OffsetWorldSpace, ENiagaraCoordinateSpace::World, SystemInstance);
 
 
-			FQuat WorldRot = FQuat::FindBetweenNormals(FVector3f::ZAxisVector, DrawOrientationAxis);
+			FQuat4f WorldRot = FQuat4f::FindBetweenNormals(FVector3f::ZAxisVector, DrawOrientationAxis);
 			TransformQuat(true, WorldRot, OrientationAxisCoordinateSpace, ENiagaraCoordinateSpace::World, SystemInstance);
 
 			FVector3f HalfOffset = DrawOrientationAxis.GetSafeNormal() * DrawLength * 0.5f;
@@ -920,15 +920,15 @@ struct FNDIDebugDrawInstanceData_GameThread
 				return;
 			static FNiagaraTypeDefinition CoordTypeDef = FNiagaraTypeDefinition(FNiagaraTypeDefinition::GetCoordinateSpaceEnum());
 
-			TOptional<FVector> Center = GetTag<FVector>(SystemInstance, FNiagaraTypeDefinition::GetVec3Def(), 0);
+			TOptional<FVector3f> Center = GetTag<FVector3f>(SystemInstance, FNiagaraTypeDefinition::GetVec3Def(), 0);
 			TOptional<ENiagaraCoordinateSpace> CenterWorldSpace = GetTag<ENiagaraCoordinateSpace>(SystemInstance, CoordTypeDef, 1);
-			TOptional<FVector> Offset = GetTag<FVector>(SystemInstance, FNiagaraTypeDefinition::GetVec3Def(), 2);
+			TOptional<FVector3f> Offset = GetTag<FVector3f>(SystemInstance, FNiagaraTypeDefinition::GetVec3Def(), 2);
 			TOptional<ENiagaraCoordinateSpace> OffsetWorldSpace = GetTag<ENiagaraCoordinateSpace>(SystemInstance, CoordTypeDef, 3);
 			TOptional<float> MajorRadius = GetTag<float>(SystemInstance, FNiagaraTypeDefinition::GetFloatDef(), 4);
 			TOptional<float> MinorRadius = GetTag<float>(SystemInstance, FNiagaraTypeDefinition::GetFloatDef(), 5);
-			TOptional<FVector> OrientationAxis = GetTag<FVector>(SystemInstance, FNiagaraTypeDefinition::GetVec3Def(), 6);
+			TOptional<FVector3f> OrientationAxis = GetTag<FVector3f>(SystemInstance, FNiagaraTypeDefinition::GetVec3Def(), 6);
 			TOptional<ENiagaraCoordinateSpace> OrientationAxisCoordinateSpace = GetTag<ENiagaraCoordinateSpace>(SystemInstance, CoordTypeDef, 7);
-			TOptional<FVector> NonUniformScale = GetTag<FVector>(SystemInstance, FNiagaraTypeDefinition::GetVec3Def(), 8);
+			TOptional<FVector3f> NonUniformScale = GetTag<FVector3f>(SystemInstance, FNiagaraTypeDefinition::GetVec3Def(), 8);
 			TOptional<int32> MajorRadiusSegments = GetTag<int32>(SystemInstance, FNiagaraTypeDefinition::GetFloatDef(), 9);
 			TOptional<int32> MinorRadiusSegments = GetTag<int32>(SystemInstance, FNiagaraTypeDefinition::GetFloatDef(), 10);
 			TOptional<FLinearColor> Color = GetTag<FLinearColor>(SystemInstance, FNiagaraTypeDefinition::GetColorDef(), 11);
@@ -946,7 +946,7 @@ struct FNDIDebugDrawInstanceData_GameThread
 			TransformPosition(Center.IsSet(), DrawCenter, CenterWorldSpace, ENiagaraCoordinateSpace::World, SystemInstance);
 			TransformVector(Offset.IsSet(), DrawOffset, OffsetWorldSpace, ENiagaraCoordinateSpace::World, SystemInstance);
 
-			FQuat WorldRot = FQuat::Identity;
+			FQuat4f WorldRot = FQuat4f::Identity;
 			TransformQuat(true, WorldRot, OrientationAxisCoordinateSpace, ENiagaraCoordinateSpace::World, SystemInstance);
 
 			DrawOrientationAxis = WorldRot.RotateVector(DrawOrientationAxis);
@@ -964,17 +964,17 @@ struct FNDIDebugDrawInstanceData_GameThread
 				return;
 			static FNiagaraTypeDefinition CoordTypeDef = FNiagaraTypeDefinition(FNiagaraTypeDefinition::GetCoordinateSpaceEnum());
 
-			TOptional<FVector> Center = GetTag<FVector>(SystemInstance, FNiagaraTypeDefinition::GetVec3Def(), 0);
+			TOptional<FVector3f> Center = GetTag<FVector3f>(SystemInstance, FNiagaraTypeDefinition::GetVec3Def(), 0);
 			TOptional<ENiagaraCoordinateSpace> CenterWorldSpace = GetTag<ENiagaraCoordinateSpace>(SystemInstance, CoordTypeDef, 1);
-			TOptional<FVector> Offset = GetTag<FVector>(SystemInstance, FNiagaraTypeDefinition::GetVec3Def(), 2);
+			TOptional<FVector3f> Offset = GetTag<FVector3f>(SystemInstance, FNiagaraTypeDefinition::GetVec3Def(), 2);
 			TOptional<ENiagaraCoordinateSpace> OffsetWorldSpace = GetTag<ENiagaraCoordinateSpace>(SystemInstance, CoordTypeDef, 3);
-			TOptional<FQuat> Rotation = GetTag<FQuat>(SystemInstance, FNiagaraTypeDefinition::GetQuatDef(), 4);
+			TOptional<FQuat4f> Rotation = GetTag<FQuat4f>(SystemInstance, FNiagaraTypeDefinition::GetQuatDef(), 4);
 			TOptional<ENiagaraCoordinateSpace> RotationWorldSpace = GetTag<ENiagaraCoordinateSpace>(SystemInstance, CoordTypeDef, 5);
 			TOptional<float> Scale = GetTag<float>(SystemInstance, FNiagaraTypeDefinition::GetFloatDef(), 6);
 
 			FVector3f DrawCenter = Center.Get(FVector3f::ZeroVector);
 			FVector3f DrawOffset = Offset.Get(FVector3f::ZeroVector);
-			FQuat DrawRotation = Rotation.Get(FQuat::Identity);
+			FQuat4f DrawRotation = Rotation.Get(FQuat4f::Identity);
 			float DrawScale = Scale.Get(1.0f);
 
 			TransformPosition(Center.IsSet(), DrawCenter, CenterWorldSpace, ENiagaraCoordinateSpace::Simulation, SystemInstance);
@@ -993,11 +993,11 @@ struct FNDIDebugDrawInstanceData_GameThread
 				return;
 			static FNiagaraTypeDefinition CoordTypeDef = FNiagaraTypeDefinition(FNiagaraTypeDefinition::GetCoordinateSpaceEnum());
 
-			TOptional<FVector> Center = GetTag<FVector>(SystemInstance, FNiagaraTypeDefinition::GetVec3Def(), 0);
+			TOptional<FVector3f> Center = GetTag<FVector3f>(SystemInstance, FNiagaraTypeDefinition::GetVec3Def(), 0);
 			TOptional<ENiagaraCoordinateSpace> CenterWorldSpace = GetTag<ENiagaraCoordinateSpace>(SystemInstance, CoordTypeDef, 1);
-			TOptional<FVector> Offset = GetTag<FVector>(SystemInstance, FNiagaraTypeDefinition::GetVec3Def(), 2);
+			TOptional<FVector3f> Offset = GetTag<FVector3f>(SystemInstance, FNiagaraTypeDefinition::GetVec3Def(), 2);
 			TOptional<ENiagaraCoordinateSpace> OffsetWorldSpace = GetTag<ENiagaraCoordinateSpace>(SystemInstance, CoordTypeDef, 3);
-			TOptional<FQuat> Rotation = GetTag<FQuat>(SystemInstance, FNiagaraTypeDefinition::GetQuatDef(), 4);
+			TOptional<FQuat4f> Rotation = GetTag<FQuat4f>(SystemInstance, FNiagaraTypeDefinition::GetQuatDef(), 4);
 			TOptional<ENiagaraCoordinateSpace> RotationWorldSpace = GetTag<ENiagaraCoordinateSpace>(SystemInstance, CoordTypeDef, 5);
 			TOptional<FVector2D> Extents = GetTag<FVector2D>(SystemInstance, FNiagaraTypeDefinition::GetVec2Def(), 6);
 			TOptional<float> NumCellsX = GetTag<float>(SystemInstance, FNiagaraTypeDefinition::GetFloatDef(), 7);
@@ -1008,7 +1008,7 @@ struct FNDIDebugDrawInstanceData_GameThread
 
 			FVector3f DrawCenter = Center.Get(FVector3f::ZeroVector);
 			FVector3f DrawOffset = Offset.Get(FVector3f::ZeroVector);
-			FQuat DrawRotation = Rotation.Get(FQuat::Identity);
+			FQuat4f DrawRotation = Rotation.Get(FQuat4f::Identity);
 			FVector2D DrawExtents = Extents.Get(FVector2D(10.f));
 			int32 DrawNumCellsX = NumCellsX.Get(10);
 			int32 DrawNumCellsY = NumCellsY.Get(10);
@@ -1030,13 +1030,13 @@ struct FNDIDebugDrawInstanceData_GameThread
 				return;
 			static FNiagaraTypeDefinition CoordTypeDef = FNiagaraTypeDefinition(FNiagaraTypeDefinition::GetCoordinateSpaceEnum());
 
-			TOptional<FVector> Center = GetTag<FVector>(SystemInstance, FNiagaraTypeDefinition::GetVec3Def(), 0);
+			TOptional<FVector3f> Center = GetTag<FVector3f>(SystemInstance, FNiagaraTypeDefinition::GetVec3Def(), 0);
 			TOptional<ENiagaraCoordinateSpace> CenterWorldSpace = GetTag<ENiagaraCoordinateSpace>(SystemInstance, CoordTypeDef, 1);
-			TOptional<FVector> Offset = GetTag<FVector>(SystemInstance, FNiagaraTypeDefinition::GetVec3Def(), 2);
+			TOptional<FVector3f> Offset = GetTag<FVector3f>(SystemInstance, FNiagaraTypeDefinition::GetVec3Def(), 2);
 			TOptional<ENiagaraCoordinateSpace> OffsetWorldSpace = GetTag<ENiagaraCoordinateSpace>(SystemInstance, CoordTypeDef, 3);
-			TOptional<FQuat> Rotation = GetTag<FQuat>(SystemInstance, FNiagaraTypeDefinition::GetQuatDef(), 4);
+			TOptional<FQuat4f> Rotation = GetTag<FQuat4f>(SystemInstance, FNiagaraTypeDefinition::GetQuatDef(), 4);
 			TOptional<ENiagaraCoordinateSpace> RotationWorldSpace = GetTag<ENiagaraCoordinateSpace>(SystemInstance, CoordTypeDef, 5);
-			TOptional<FVector> Extents = GetTag<FVector>(SystemInstance, FNiagaraTypeDefinition::GetVec3Def(), 6);
+			TOptional<FVector3f> Extents = GetTag<FVector3f>(SystemInstance, FNiagaraTypeDefinition::GetVec3Def(), 6);
 			TOptional<float> NumCellsX = GetTag<float>(SystemInstance, FNiagaraTypeDefinition::GetFloatDef(), 7);
 			TOptional<float> NumCellsY = GetTag<float>(SystemInstance, FNiagaraTypeDefinition::GetFloatDef(), 8);
 			TOptional<float> NumCellsZ = GetTag<float>(SystemInstance, FNiagaraTypeDefinition::GetFloatDef(), 9);
@@ -1045,7 +1045,7 @@ struct FNDIDebugDrawInstanceData_GameThread
 
 			FVector3f DrawCenter = Center.Get(FVector3f::ZeroVector);
 			FVector3f DrawOffset = Offset.Get(FVector3f::ZeroVector);
-			FQuat DrawRotation = Rotation.Get(FQuat::Identity);
+			FQuat4f DrawRotation = Rotation.Get(FQuat4f::Identity);
 			FVector3f DrawExtents = Extents.Get(FVector3f(10.f));
 			int32 DrawNumCellsX = NumCellsX.Get(10);
 			int32 DrawNumCellsY = NumCellsY.Get(10);
@@ -1156,8 +1156,8 @@ struct FNDIDebugDrawInstanceData_GameThread
 							ParamNames.Add(GenerateParameterName(ExistingShape, TEXT("EndLocationCoordinateSpace")));
 							ParamNames.Add(GenerateParameterName(ExistingShape, TEXT("Color")));
 
-							TOptional<FVector> StartLocation = GetCompileTag<FVector>(SystemInstance, Script, FNiagaraTypeDefinition::GetVec3Def(), ParamNames[0]);
-							TOptional<FVector> EndLocation = GetCompileTag<FVector>(SystemInstance, Script, FNiagaraTypeDefinition::GetVec3Def(), ParamNames[2]);
+							TOptional<FVector3f> StartLocation = GetCompileTag<FVector3f>(SystemInstance, Script, FNiagaraTypeDefinition::GetVec3Def(), ParamNames[0]);
+							TOptional<FVector3f> EndLocation = GetCompileTag<FVector3f>(SystemInstance, Script, FNiagaraTypeDefinition::GetVec3Def(), ParamNames[2]);
 
 							bIsValid = StartLocation.IsSet() || EndLocation.IsSet();
 
@@ -1179,7 +1179,7 @@ struct FNDIDebugDrawInstanceData_GameThread
 							ParamNames.Add(GenerateParameterName(ExistingShape, TEXT("UnboundedPlane")));
 							ParamNames.Add(GenerateParameterName(ExistingShape, TEXT("Color")));
 
-							TOptional<FVector> Center = GetCompileTag<FVector>(SystemInstance, Script, FNiagaraTypeDefinition::GetVec3Def(), ParamNames[0]);
+							TOptional<FVector3f> Center = GetCompileTag<FVector3f>(SystemInstance, Script, FNiagaraTypeDefinition::GetVec3Def(), ParamNames[0]);
 							TOptional<FVector2D> Extents = GetCompileTag<FVector2D>(SystemInstance, Script, FNiagaraTypeDefinition::GetVec2Def(), ParamNames[4]);
 
 							bIsValid = Center.IsSet() || Extents.IsSet();
@@ -1199,7 +1199,7 @@ struct FNDIDebugDrawInstanceData_GameThread
 							ParamNames.Add(GenerateParameterName(ExistingShape, TEXT("NumSegments")));
 							ParamNames.Add(GenerateParameterName(ExistingShape, TEXT("Color")));
 
-							TOptional<FVector> Center = GetCompileTag<FVector>(SystemInstance, Script, FNiagaraTypeDefinition::GetVec3Def(), ParamNames[0]);
+							TOptional<FVector3f> Center = GetCompileTag<FVector3f>(SystemInstance, Script, FNiagaraTypeDefinition::GetVec3Def(), ParamNames[0]);
 							TOptional<float> Radius = GetCompileTag<float>(SystemInstance, Script, FNiagaraTypeDefinition::GetFloatDef(), ParamNames[4]);
 
 							bIsValid = Center.IsSet() || Radius.IsSet();
@@ -1219,8 +1219,8 @@ struct FNDIDebugDrawInstanceData_GameThread
 							ParamNames.Add(GenerateParameterName(ExistingShape, TEXT("HalfExtents")));
 							ParamNames.Add(GenerateParameterName(ExistingShape, TEXT("Color")));
 
-							TOptional<FVector> Center = GetCompileTag<FVector>(SystemInstance, Script, FNiagaraTypeDefinition::GetVec3Def(), ParamNames[0]);
-							TOptional<FVector> Extents = GetCompileTag<FVector>(SystemInstance, Script, FNiagaraTypeDefinition::GetVec3Def(), ParamNames[7]);
+							TOptional<FVector3f> Center = GetCompileTag<FVector3f>(SystemInstance, Script, FNiagaraTypeDefinition::GetVec3Def(), ParamNames[0]);
+							TOptional<FVector3f> Extents = GetCompileTag<FVector3f>(SystemInstance, Script, FNiagaraTypeDefinition::GetVec3Def(), ParamNames[7]);
 
 							bIsValid = Center.IsSet() || Extents.IsSet();
 
@@ -1244,7 +1244,7 @@ struct FNDIDebugDrawInstanceData_GameThread
 							ParamNames.Add(GenerateParameterName(ExistingShape, TEXT("Num Segments")));
 							ParamNames.Add(GenerateParameterName(ExistingShape, TEXT("Color")));							
 
-							TOptional<FVector> Center = GetCompileTag<FVector>(SystemInstance, Script, FNiagaraTypeDefinition::GetVec3Def(), ParamNames[0]);
+							TOptional<FVector3f> Center = GetCompileTag<FVector3f>(SystemInstance, Script, FNiagaraTypeDefinition::GetVec3Def(), ParamNames[0]);
 							TOptional<float> Radius = GetCompileTag<float>(SystemInstance, Script, FNiagaraTypeDefinition::GetFloatDef(), ParamNames[4]);
 
 							bIsValid = Center.IsSet() || Radius.IsSet();
@@ -1269,7 +1269,7 @@ struct FNDIDebugDrawInstanceData_GameThread
 							ParamNames.Add(GenerateParameterName(ExistingShape, TEXT("Num Height Segments")));
 							ParamNames.Add(GenerateParameterName(ExistingShape, TEXT("Color")));
 
-							TOptional<FVector> Center = GetCompileTag<FVector>(SystemInstance, Script, FNiagaraTypeDefinition::GetVec3Def(), ParamNames[0]);
+							TOptional<FVector3f> Center = GetCompileTag<FVector3f>(SystemInstance, Script, FNiagaraTypeDefinition::GetVec3Def(), ParamNames[0]);
 							TOptional<float> Radius = GetCompileTag<float>(SystemInstance, Script, FNiagaraTypeDefinition::GetFloatDef(), ParamNames[4]);
 
 							bIsValid = Center.IsSet() || Radius.IsSet();
@@ -1291,7 +1291,7 @@ struct FNDIDebugDrawInstanceData_GameThread
 							ParamNames.Add(GenerateParameterName(ExistingShape, TEXT("Num Height Segments")));
 							ParamNames.Add(GenerateParameterName(ExistingShape, TEXT("Color")));
 
-							TOptional<FVector> Center = GetCompileTag<FVector>(SystemInstance, Script, FNiagaraTypeDefinition::GetVec3Def(), ParamNames[0]);
+							TOptional<FVector3f> Center = GetCompileTag<FVector3f>(SystemInstance, Script, FNiagaraTypeDefinition::GetVec3Def(), ParamNames[0]);
 							TOptional<float> Radius = GetCompileTag<float>(SystemInstance, Script, FNiagaraTypeDefinition::GetFloatDef(), ParamNames[4]);
 
 							bIsValid = Center.IsSet() || Radius.IsSet();
@@ -1313,7 +1313,7 @@ struct FNDIDebugDrawInstanceData_GameThread
 							ParamNames.Add(GenerateParameterName(ExistingShape, TEXT("MinorRadiusSegments")));
 							ParamNames.Add(GenerateParameterName(ExistingShape, TEXT("Color")));
 
-							TOptional<FVector> Center = GetCompileTag<FVector>(SystemInstance, Script, FNiagaraTypeDefinition::GetVec3Def(), ParamNames[0]);
+							TOptional<FVector3f> Center = GetCompileTag<FVector3f>(SystemInstance, Script, FNiagaraTypeDefinition::GetVec3Def(), ParamNames[0]);
 							TOptional<float> MajorRadius = GetCompileTag<float>(SystemInstance, Script, FNiagaraTypeDefinition::GetFloatDef(), ParamNames[4]);
 							TOptional<float> MinorRadius = GetCompileTag<float>(SystemInstance, Script, FNiagaraTypeDefinition::GetFloatDef(), ParamNames[5]);
 
@@ -1331,7 +1331,7 @@ struct FNDIDebugDrawInstanceData_GameThread
 							ParamNames.Add(GenerateParameterName(ExistingShape, TEXT("RotationCoordinateSpace")));
 							ParamNames.Add(GenerateParameterName(ExistingShape, TEXT("Scale")));
 
-							TOptional<FVector> Center = GetCompileTag<FVector>(SystemInstance, Script, FNiagaraTypeDefinition::GetVec3Def(), ParamNames[0]);
+							TOptional<FVector3f> Center = GetCompileTag<FVector3f>(SystemInstance, Script, FNiagaraTypeDefinition::GetVec3Def(), ParamNames[0]);
 
 							bIsValid = Center.IsSet();
 
@@ -1350,7 +1350,7 @@ struct FNDIDebugDrawInstanceData_GameThread
 							ParamNames.Add(GenerateParameterName(ExistingShape, TEXT("NumCellsY")));
 							ParamNames.Add(GenerateParameterName(ExistingShape, TEXT("Color")));
 
-							TOptional<FVector> Center = GetCompileTag<FVector>(SystemInstance, Script, FNiagaraTypeDefinition::GetVec3Def(), ParamNames[0]);
+							TOptional<FVector3f> Center = GetCompileTag<FVector3f>(SystemInstance, Script, FNiagaraTypeDefinition::GetVec3Def(), ParamNames[0]);
 
 							bIsValid = Center.IsSet();
 
@@ -1370,7 +1370,7 @@ struct FNDIDebugDrawInstanceData_GameThread
 							ParamNames.Add(GenerateParameterName(ExistingShape, TEXT("NumCellsZ")));
 							ParamNames.Add(GenerateParameterName(ExistingShape, TEXT("Color")));
 
-							TOptional<FVector> Center = GetCompileTag<FVector>(SystemInstance, Script, FNiagaraTypeDefinition::GetVec3Def(), ParamNames[0]);
+							TOptional<FVector3f> Center = GetCompileTag<FVector3f>(SystemInstance, Script, FNiagaraTypeDefinition::GetVec3Def(), ParamNames[0]);
 
 							bIsValid = Center.IsSet();
 
@@ -1673,7 +1673,7 @@ namespace NDIDebugDrawLocal
 			}
 
 			FNDIInputParam<FVector3f> LocationParam;
-			FNDIInputParam<FQuat> RotationParam;
+			FNDIInputParam<FQuat4f> RotationParam;
 			FNDIInputParam<FVector3f> ExtentsParam;
 			FNDIInputParam<FLinearColor> ColorParam;
 		};
@@ -1711,7 +1711,7 @@ namespace NDIDebugDrawLocal
 		static void Draw(FNDIDebugDrawInstanceData_GameThread* InstanceData, VMBindings& Bindings, bool bExecute)
 		{
 			const FVector3f Location = Bindings.LocationParam.GetAndAdvance();
-			const FQuat Rotation = Bindings.RotationParam.GetAndAdvance();
+			const FQuat4f Rotation = Bindings.RotationParam.GetAndAdvance();
 			const FVector3f Extents = Bindings.ExtentsParam.GetAndAdvance();
 			const FLinearColor Color = Bindings.ColorParam.GetAndAdvance();
 
@@ -1773,7 +1773,7 @@ namespace NDIDebugDrawLocal
 			FNDIInputParam<FNiagaraBool> HemiZParam;
 			FNDIInputParam<FVector3f> OrientationAxisParam;
 			FNDIInputParam<ENiagaraCoordinateSpace> RotationWSParam;
-			FNDIInputParam<FQuat> AdditionalRotationParam;
+			FNDIInputParam<FQuat4f> AdditionalRotationParam;
 			FNDIInputParam<FVector3f> NonUniformScaleParam;
 			FNDIInputParam<int32> SegmentsParam;
 			FNDIInputParam<FLinearColor> ColorParam;
@@ -2045,7 +2045,7 @@ namespace NDIDebugDrawLocal
 			}
 
 			FNDIInputParam<FVector3f> LocationParam;
-			FNDIInputParam<FQuat> RotationParam;
+			FNDIInputParam<FQuat4f> RotationParam;
 			FNDIInputParam<float> ScaleParam;
 		};		
 		
@@ -2067,7 +2067,7 @@ namespace NDIDebugDrawLocal
 			FNDIInputParam<ENiagaraCoordinateSpace> LocationWSParam;
 			FNDIInputParam<FVector3f> OffsetParam;
 			FNDIInputParam<ENiagaraCoordinateSpace> OffsetWSParam;
-			FNDIInputParam<FQuat> RotationParam;
+			FNDIInputParam<FQuat4f> RotationParam;
 			FNDIInputParam<ENiagaraCoordinateSpace> RotationCoordinateSpace;
 			FNDIInputParam<float> ScaleParam;
 		};
@@ -2077,7 +2077,7 @@ namespace NDIDebugDrawLocal
 		static void Draw(FNDIDebugDrawInstanceData_GameThread* InstanceData, VMBindings& Bindings, bool bExecute)
 		{
 			const FVector3f Location = Bindings.LocationParam.GetAndAdvance();
-			const FQuat Rotation = Bindings.RotationParam.GetAndAdvance();
+			const FQuat4f Rotation = Bindings.RotationParam.GetAndAdvance();
 			const float Scale = Bindings.ScaleParam.GetAndAdvance();
 
 			if (bExecute)
@@ -2103,7 +2103,7 @@ namespace NDIDebugDrawLocal
 			}
 
 			FNDIInputParam<FVector3f> CenterParam;
-			FNDIInputParam<FQuat> RotationParam;
+			FNDIInputParam<FQuat4f> RotationParam;
 			FNDIInputParam<FVector2D> ExtentsParam;
 			FNDIInputParam<int32> NumCellsXParam;
 			FNDIInputParam<int32> NumCellsYParam;
@@ -2131,7 +2131,7 @@ namespace NDIDebugDrawLocal
 			FNDIInputParam<ENiagaraCoordinateSpace> LocationWSParam;
 			FNDIInputParam<FVector3f> OffsetParam;
 			FNDIInputParam<ENiagaraCoordinateSpace> OffsetWSParam;
-			FNDIInputParam<FQuat> RotationParam;
+			FNDIInputParam<FQuat4f> RotationParam;
 			FNDIInputParam<ENiagaraCoordinateSpace> RotationCoordinateSpace;
 			FNDIInputParam<FVector2D> ExtentsParam;
 			FNDIInputParam<int32> NumCellsXParam;
@@ -2143,7 +2143,7 @@ namespace NDIDebugDrawLocal
 		static void Draw(FNDIDebugDrawInstanceData_GameThread* InstanceData, VMBindings& Bindings, bool bExecute)
 		{
 			const FVector3f Center = Bindings.CenterParam.GetAndAdvance();
-			const FQuat Rotation = Bindings.RotationParam.GetAndAdvance();
+			const FQuat4f Rotation = Bindings.RotationParam.GetAndAdvance();
 			const FVector2D Extents = Bindings.ExtentsParam.GetAndAdvance();
 			const FIntPoint NumCells(Bindings.NumCellsXParam.GetAndAdvance(), Bindings.NumCellsYParam.GetAndAdvance());
 			const FLinearColor Color = Bindings.ColorParam.GetAndAdvance();
@@ -2172,7 +2172,7 @@ namespace NDIDebugDrawLocal
 			}
 
 			FNDIInputParam<FVector3f> CenterParam;
-			FNDIInputParam<FQuat> RotationParam;
+			FNDIInputParam<FQuat4f> RotationParam;
 			FNDIInputParam<FVector3f> ExtentsParam;
 			FNDIInputParam<int32> NumCellsXParam;
 			FNDIInputParam<int32> NumCellsYParam;
@@ -2202,7 +2202,7 @@ namespace NDIDebugDrawLocal
 			FNDIInputParam<ENiagaraCoordinateSpace> LocationWSParam;
 			FNDIInputParam<FVector3f> OffsetParam;
 			FNDIInputParam<ENiagaraCoordinateSpace> OffsetWSParam;
-			FNDIInputParam<FQuat> RotationParam;
+			FNDIInputParam<FQuat4f> RotationParam;
 			FNDIInputParam<ENiagaraCoordinateSpace> RotationCoordinateSpace;
 			FNDIInputParam<FVector3f> ExtentsParam;
 			FNDIInputParam<int32> NumCellsXParam;
@@ -2216,7 +2216,7 @@ namespace NDIDebugDrawLocal
 		static void Draw(FNDIDebugDrawInstanceData_GameThread* InstanceData, VMBindings& Bindings, bool bExecute)
 		{
 			const FVector3f Center = Bindings.CenterParam.GetAndAdvance();
-			const FQuat Rotation = Bindings.RotationParam.GetAndAdvance();
+			const FQuat4f Rotation = Bindings.RotationParam.GetAndAdvance();
 			const FVector3f Extents = Bindings.ExtentsParam.GetAndAdvance();
 			const FIntVector NumCells(Bindings.NumCellsXParam.GetAndAdvance(), Bindings.NumCellsYParam.GetAndAdvance(), Bindings.NumCellsZParam.GetAndAdvance());
 			const FLinearColor Color = Bindings.ColorParam.GetAndAdvance();
@@ -2642,7 +2642,7 @@ void UNiagaraDataInterfaceDebugDraw::GetFunctions(TArray<FNiagaraFunctionSignatu
 		Signature.Inputs.Add_GetRef(FNiagaraVariable(FNiagaraTypeDefinition::GetBoolDef(), TEXT("Hemisphere Z"))).SetValue(true);
 		Signature.Inputs.Add_GetRef(FNiagaraVariable(FNiagaraTypeDefinition::GetVec3Def(), TEXT("Sphere Orientation Axis"))).SetValue(FVector3f::ZAxisVector);
 		Signature.Inputs.Add_GetRef(FNiagaraVariable(CoordTypeDef, TEXT("RotationCoordinateSpace"))).SetValue(ENiagaraCoordinateSpace::Local);
-		Signature.Inputs.Add_GetRef(FNiagaraVariable(FNiagaraTypeDefinition::GetQuatDef(), TEXT("Additional Rotation"))).SetValue(FQuat::Identity);
+		Signature.Inputs.Add_GetRef(FNiagaraVariable(FNiagaraTypeDefinition::GetQuatDef(), TEXT("Additional Rotation"))).SetValue(FQuat4f::Identity);
 		Signature.Inputs.Add_GetRef(FNiagaraVariable(FNiagaraTypeDefinition::GetVec3Def(), TEXT("NonUniform Scale"))).SetValue(FVector3f::OneVector);
 		Signature.Inputs.Add_GetRef(FNiagaraVariable(FNiagaraTypeDefinition::GetIntDef(), TEXT("Num Segments"))).SetValue(12);
 		Signature.Inputs.Add(FNiagaraVariable(FNiagaraTypeDefinition::GetColorDef(), TEXT("Color")));
diff --git a/Engine/Plugins/FX/Niagara/Source/Niagara/Private/NDISkeletalMesh_BoneSampling.cpp b/Engine/Plugins/FX/Niagara/Source/Niagara/Private/NDISkeletalMesh_BoneSampling.cpp
--- a/Engine/Plugins/FX/Niagara/Source/Niagara/Private/NDISkeletalMesh_BoneSampling.cpp
+++ b/Engine/Plugins/FX/Niagara/Source/Niagara/Private/NDISkeletalMesh_BoneSampling.cpp
@@ -712,7 +712,7 @@ struct FBoneSocketSkinnedDataOutputHandler
 	}
 
 	FNDIOutputParam<FVector3f> Position;
-	FNDIOutputParam<FQuat> Rotation;
+	FNDIOutputParam<FQuat4f> Rotation;
 	FNDIOutputParam<FVector3f> Velocity;
 
 	//TODO: Rotation + Scale too? Use quats so we can get proper interpolation between bone and parent.
@@ -743,10 +743,10 @@ void UNiagaraDataInterfaceSkeletalMesh::GetSkinnedBoneData(FVectorVMExternalFunc
 
 	//TODO: Replace this by storing off FTransforms and doing a proper lerp to get a final transform.
 	//Also need to pull in a per particle interpolation factor.
-	const FMatrix& InstanceTransform = InstData->Transform;
-	const FMatrix& PrevInstanceTransform = InstData->PrevTransform;
-	const FQuat InstanceRotation = Output.bNeedsRotation ? InstanceTransform.GetMatrixWithoutScale().ToQuat() : FQuat::Identity;
-	const FQuat PrevInstanceRotation = Output.bNeedsRotation ? PrevInstanceTransform.GetMatrixWithoutScale().ToQuat() : FQuat::Identity;
+	const FMatrix44f InstanceTransform = (FMatrix44f)InstData->Transform;
+	const FMatrix44f PrevInstanceTransform = (FMatrix44f)InstData->PrevTransform;
+	const FQuat4f InstanceRotation = Output.bNeedsRotation ? InstanceTransform.GetMatrixWithoutScale().ToQuat() : FQuat4f::Identity;
+	const FQuat4f PrevInstanceRotation = Output.bNeedsRotation ? PrevInstanceTransform.GetMatrixWithoutScale().ToQuat() : FQuat4f::Identity;
 
 	FSkeletalMeshAccessorHelper Accessor;
 	Accessor.Init<TNDISkelMesh_FilterModeNone, TNDISkelMesh_AreaWeightingOff>(InstData);
@@ -756,8 +756,8 @@ void UNiagaraDataInterfaceSkeletalMesh::GetSkinnedBoneData(FVectorVMExternalFunc
 		const int32 BoneAndSocketCount = BoneCount + InstData->FilteredSocketInfo.Num();
 		float InvDt = 1.0f / InstData->DeltaSeconds;
 
-		const TArray<FTransform>& FilteredSocketCurrTransforms = InstData->GetFilteredSocketsCurrBuffer();
-		const TArray<FTransform>& FilteredSocketPrevTransforms = InstData->GetFilteredSocketsPrevBuffer();
+		const TArray<FTransform3f>& FilteredSocketCurrTransforms = InstData->GetFilteredSocketsCurrBuffer();
+		const TArray<FTransform3f>& FilteredSocketPrevTransforms = InstData->GetFilteredSocketsPrevBuffer();
 
 		for (int32 i = 0; i < Context.GetNumInstances(); ++i)
 		{
@@ -784,13 +784,13 @@ void UNiagaraDataInterfaceSkeletalMesh::GetSkinnedBoneData(FVectorVMExternalFunc
 				}
 				if (Output.bNeedsRotation)
 				{
-					Output.Rotation.SetAndAdvance(FQuat::Identity);
+					Output.Rotation.SetAndAdvance(FQuat4f::Identity);
 				}
 			}
 			else if (bIsSocket)
 			{
-				FTransform CurrSocketTransform = FilteredSocketCurrTransforms[Socket];
-				FTransform PrevSocketTransform = FilteredSocketPrevTransforms[Socket];
+				FTransform3f CurrSocketTransform = FilteredSocketCurrTransforms[Socket];
+				FTransform3f PrevSocketTransform = FilteredSocketPrevTransforms[Socket];
 
 				Pos = CurrSocketTransform.GetLocation();
 				TransformHandler.TransformPosition(Pos, InstanceTransform);
@@ -803,13 +803,13 @@ void UNiagaraDataInterfaceSkeletalMesh::GetSkinnedBoneData(FVectorVMExternalFunc
 
 				if (Output.bNeedsRotation)
 				{
-					FQuat Rotation = CurrSocketTransform.GetRotation();
+					FQuat4f Rotation = CurrSocketTransform.GetRotation();
 					TransformHandler.TransformRotation(Rotation, InstanceRotation);
 					if (bInterpolated::Value)
 					{
-						FQuat PrevRotation = PrevSocketTransform.GetRotation();
+						FQuat4f PrevRotation = PrevSocketTransform.GetRotation();
 						TransformHandler.TransformRotation(PrevRotation, PrevInstanceRotation);
-						Rotation = FQuat::Slerp(PrevRotation, Rotation, Interp);
+						Rotation = FQuat4f::Slerp(PrevRotation, Rotation, Interp);
 					}
 
 					Output.Rotation.SetAndAdvance(Rotation);
@@ -829,13 +829,13 @@ void UNiagaraDataInterfaceSkeletalMesh::GetSkinnedBoneData(FVectorVMExternalFunc
 
 				if (Output.bNeedsRotation)
 				{
-					FQuat Rotation = SkinningHandler.GetSkinnedBoneRotation(Accessor, Bone);
+					FQuat4f Rotation = SkinningHandler.GetSkinnedBoneRotation(Accessor, Bone);
 					TransformHandler.TransformRotation(Rotation, InstanceRotation);
 					if (bInterpolated::Value)
 					{
-						FQuat PrevRotation = SkinningHandler.GetSkinnedBonePreviousRotation(Accessor, Bone);
+						FQuat4f PrevRotation = SkinningHandler.GetSkinnedBonePreviousRotation(Accessor, Bone);
 						TransformHandler.TransformRotation(PrevRotation, PrevInstanceRotation);
-						Rotation = FQuat::Slerp(PrevRotation, Rotation, Interp);
+						Rotation = FQuat4f::Slerp(PrevRotation, Rotation, Interp);
 					}
 
 					Output.Rotation.SetAndAdvance(Rotation);
@@ -879,7 +879,7 @@ void UNiagaraDataInterfaceSkeletalMesh::GetSkinnedBoneData(FVectorVMExternalFunc
 
 			if (Output.bNeedsRotation)
 			{
-				Output.Rotation.SetAndAdvance(FQuat::Identity);
+				Output.Rotation.SetAndAdvance(FQuat4f::Identity);
 			}
 
 			if (Output.bNeedsVelocity || bInterpolated::Value)
@@ -950,22 +950,22 @@ void UNiagaraDataInterfaceSkeletalMesh::GetFilteredSocketTransform(FVectorVMExte
 	FNDIInputParam<FNiagaraBool> ApplyWorldTransform(Context);
 
 	FNDIOutputParam<FVector3f> OutSocketTranslate(Context);
-	FNDIOutputParam<FQuat> OutSocketRotation(Context);
+	FNDIOutputParam<FQuat4f> OutSocketRotation(Context);
 	FNDIOutputParam<FVector3f> OutSocketScale(Context);
 
-	const TArray<FTransform>& CurrentFilteredSockets = InstData->GetFilteredSocketsCurrBuffer();
+	const TArray<FTransform3f>& CurrentFilteredSockets = InstData->GetFilteredSocketsCurrBuffer();
 	const int32 SocketMax = CurrentFilteredSockets.Num() - 1;
 	if (SocketMax >= 0)
 	{
 		const bool bNeedsRotation = OutSocketRotation.IsValid();
-		const FMatrix InstanceTransform = InstData->Transform;
-		const FQuat InstanceRotation = bNeedsRotation ? InstanceTransform.GetMatrixWithoutScale().ToQuat() : FQuat::Identity;
+		const FMatrix44f InstanceTransform = (FMatrix44f)InstData->Transform;
+		const FQuat4f InstanceRotation = bNeedsRotation ? InstanceTransform.GetMatrixWithoutScale().ToQuat() : FQuat4f::Identity;
 
 		for (int32 i = 0; i < Context.GetNumInstances(); ++i)
 		{
 			const int32 SocketIndex = FMath::Clamp(SocketParam.GetAndAdvance(), 0, SocketMax);
 			FVector3f SocketTranslation = CurrentFilteredSockets[SocketIndex].GetTranslation();
-			FQuat SocketRotation = CurrentFilteredSockets[SocketIndex].GetRotation();
+			FQuat4f SocketRotation = CurrentFilteredSockets[SocketIndex].GetRotation();
 			FVector3f SocketScale = CurrentFilteredSockets[SocketIndex].GetScale3D();
 
 			const bool bApplyTransform = ApplyWorldTransform.GetAndAdvance();
@@ -986,7 +986,7 @@ void UNiagaraDataInterfaceSkeletalMesh::GetFilteredSocketTransform(FVectorVMExte
 		for (int32 i=0; i < Context.GetNumInstances(); ++i)
 		{
 			OutSocketTranslate.SetAndAdvance(FVector3f::ZeroVector);
-			OutSocketRotation.SetAndAdvance(FQuat::Identity);
+			OutSocketRotation.SetAndAdvance(FQuat4f::Identity);
 			OutSocketScale.SetAndAdvance(FVector3f::ZeroVector);
 		}
 	}
diff --git a/Engine/Plugins/FX/Niagara/Source/Niagara/Private/NiagaraCommon.cpp b/Engine/Plugins/FX/Niagara/Source/Niagara/Private/NiagaraCommon.cpp
--- a/Engine/Plugins/FX/Niagara/Source/Niagara/Private/NiagaraCommon.cpp
+++ b/Engine/Plugins/FX/Niagara/Source/Niagara/Private/NiagaraCommon.cpp
@@ -145,6 +145,11 @@ UScriptStruct* FNiagaraTypeHelper::FindNiagaraFriendlyTopLevelStruct(UScriptStru
 		return FNiagaraTypeDefinition::GetVec3Struct();
 	}
 
+	if (InStruct->GetFName() == NAME_Quat || InStruct->GetFName() == NAME_Quat4d) // LWC support
+	{
+		return FNiagaraTypeDefinition::GetQuatStruct();
+	}
+
 	return InStruct;
 }
 
diff --git a/Engine/Plugins/FX/Niagara/Source/Niagara/Private/NiagaraComponentRendererProperties.cpp b/Engine/Plugins/FX/Niagara/Source/Niagara/Private/NiagaraComponentRendererProperties.cpp
--- a/Engine/Plugins/FX/Niagara/Source/Niagara/Private/NiagaraComponentRendererProperties.cpp
+++ b/Engine/Plugins/FX/Niagara/Source/Niagara/Private/NiagaraComponentRendererProperties.cpp
@@ -40,6 +40,7 @@ bool UNiagaraComponentRendererProperties::IsConvertible(const FNiagaraTypeDefini
 		(SourceType == FNiagaraTypeDefinition::GetVec3Def() && TargetType.GetStruct() == GetFRotatorDef().GetStruct()) ||
 		(SourceType == FNiagaraTypeDefinition::GetVec3Def() && TargetType.GetStruct() == GetFVectorDef().GetStruct()) ||
 		(SourceType == FNiagaraTypeDefinition::GetVec4Def() && TargetType.GetStruct() == GetFColorDef().GetStruct()) ||
+		(SourceType == FNiagaraTypeDefinition::GetQuatDef() && TargetType.GetStruct() == GetFQuatDef().GetStruct()) ||
 		(SourceType == FNiagaraTypeDefinition::GetQuatDef() && TargetType.GetStruct() == GetFRotatorDef().GetStruct()))
 	{
 		return true;
@@ -112,6 +113,13 @@ FNiagaraTypeDefinition UNiagaraComponentRendererProperties::GetFVectorDef()
 	return FNiagaraTypeDefinition(VectorStruct, FNiagaraTypeDefinition::EAllowUnfriendlyStruct::Allow);
 }
 
+FNiagaraTypeDefinition UNiagaraComponentRendererProperties::GetFQuatDef()
+{
+	static UPackage* CoreUObjectPkg = FindObjectChecked<UPackage>(nullptr, TEXT("/Script/CoreUObject"));
+	static UScriptStruct* Struct = FindObjectChecked<UScriptStruct>(CoreUObjectPkg, TEXT("Quat"));
+	return FNiagaraTypeDefinition(Struct, FNiagaraTypeDefinition::EAllowUnfriendlyStruct::Allow);
+}
+
 TArray<TWeakObjectPtr<UNiagaraComponentRendererProperties>> UNiagaraComponentRendererProperties::ComponentRendererPropertiesToDeferredInit;
 
 UNiagaraComponentRendererProperties::UNiagaraComponentRendererProperties()
diff --git a/Engine/Plugins/FX/Niagara/Source/Niagara/Private/NiagaraConstants.cpp b/Engine/Plugins/FX/Niagara/Source/Niagara/Private/NiagaraConstants.cpp
--- a/Engine/Plugins/FX/Niagara/Source/Niagara/Private/NiagaraConstants.cpp
+++ b/Engine/Plugins/FX/Niagara/Source/Niagara/Private/NiagaraConstants.cpp
@@ -387,7 +387,7 @@ void FNiagaraConstants::Init()
 
 		AttrDefaultsStrMap.Add(SYS_PARAM_PARTICLES_MESH_ORIENTATION, TEXT("0.0,0.0,0.0,1.0"));
 		Var = SYS_PARAM_PARTICLES_MESH_ORIENTATION;
-		Var.SetValue<FQuat>(FQuat::Identity);
+		Var.SetValue<FQuat4f>(FQuat4f::Identity);
 		AttrDefaultsValueMap.Add(SYS_PARAM_PARTICLES_MESH_ORIENTATION, Var);
 
 		AttrDefaultsStrMap.Add(SYS_PARAM_PARTICLES_MESH_INDEX, TEXT("0"));
diff --git a/Engine/Plugins/FX/Niagara/Source/Niagara/Private/NiagaraDataInterfaceArrayFloat.cpp b/Engine/Plugins/FX/Niagara/Source/Niagara/Private/NiagaraDataInterfaceArrayFloat.cpp
--- a/Engine/Plugins/FX/Niagara/Source/Niagara/Private/NiagaraDataInterfaceArrayFloat.cpp
+++ b/Engine/Plugins/FX/Niagara/Source/Niagara/Private/NiagaraDataInterfaceArrayFloat.cpp
@@ -69,13 +69,17 @@ struct FNDIArrayImplHelper<FVector3d> : public FNDIArrayImplHelperBase<FVector3d
 	{
 		return sizeof(float);
 	}
+	static int32 CPUGetTypeStride()
+	{
+		return sizeof(FVector3f);
+	}
 
 	static void CopyData(void* Dest, const FVector3d* Src, int32 BufferSize)
 	{
 		FVector3f* DestFloats = (FVector3f*)Dest;
 		int32 Num = BufferSize / sizeof(FVector3d);
 		for (int32 i = 0; i < Num; i++)
-			DestFloats[i] = Src[i];
+			DestFloats[i] = (FVector3f)Src[i];
 	}
 };
 
@@ -102,14 +106,37 @@ struct FNDIArrayImplHelper<FLinearColor> : public FNDIArrayImplHelperBase<FLinea
 };
 
 template<>
-struct FNDIArrayImplHelper<FQuat> : public FNDIArrayImplHelperBase<FQuat>
+struct FNDIArrayImplHelper<FQuat4f> : public FNDIArrayImplHelperBase<FQuat4f>
 {
-	typedef FQuat TVMArrayType;
+	typedef FQuat4f TVMArrayType;
 	static constexpr TCHAR const* HLSLValueTypeName = TEXT("float4");
 	static constexpr TCHAR const* HLSLBufferTypeName = TEXT("float4");
 	static constexpr EPixelFormat PixelFormat = PF_A32B32G32R32F;
 	static const FNiagaraTypeDefinition& GetTypeDefinition() { return FNiagaraTypeDefinition::GetQuatDef(); }
-	static const FQuat GetDefaultValue() { return FQuat::Identity; }
+	static const FQuat4f GetDefaultValue() { return FQuat4f::Identity; }
+};
+
+// LWC_TODO: This is represented as an FQuat4f internally (array is converted to floats during PushToRenderThread)
+template<>
+struct FNDIArrayImplHelper<FQuat4d> : public FNDIArrayImplHelperBase<FQuat4d>
+{
+	typedef FQuat4f TVMArrayType;
+	static constexpr TCHAR const* HLSLValueTypeName = TEXT("float4");
+	static constexpr TCHAR const* HLSLBufferTypeName = TEXT("float4");
+	static constexpr EPixelFormat PixelFormat = PF_A32B32G32R32F;
+	static const FNiagaraTypeDefinition& GetTypeDefinition() { return FNiagaraTypeDefinition::GetQuatDef(); }
+	static const FQuat4f GetDefaultValue() { return FQuat4f::Identity; }
+
+	static int32 GPUGetTypeStride() { return sizeof(FQuat4f); }
+	static int32 CPUGetTypeStride() { return sizeof(FQuat4f); }
+
+	static void CopyData(void* Dest, const FQuat4d* Src, int32 BufferSize)
+	{
+		FQuat4f* DestFloats = (FQuat4f*)Dest;
+		int32 Num = BufferSize / sizeof(FQuat4d);
+		for (int32 i = 0; i < Num; i++)
+			DestFloats[i] = (FQuat4f)Src[i];
+	}
 };
 
 UNiagaraDataInterfaceArrayFloat::UNiagaraDataInterfaceArrayFloat(FObjectInitializer const& ObjectInitializer)
diff --git a/Engine/Plugins/FX/Niagara/Source/Niagara/Private/NiagaraDataInterfaceParticleRead.cpp b/Engine/Plugins/FX/Niagara/Source/Niagara/Private/NiagaraDataInterfaceParticleRead.cpp
--- a/Engine/Plugins/FX/Niagara/Source/Niagara/Private/NiagaraDataInterfaceParticleRead.cpp
+++ b/Engine/Plugins/FX/Niagara/Source/Niagara/Private/NiagaraDataInterfaceParticleRead.cpp
@@ -1600,7 +1600,7 @@ void UNiagaraDataInterfaceParticleRead::ReadColor(FVectorVMExternalFunctionConte
 
 void UNiagaraDataInterfaceParticleRead::ReadQuat(FVectorVMExternalFunctionContext& Context, FName AttributeToRead)
 {
-	ReadWithCheck<FQuat>(Context, AttributeToRead, FQuat::Identity);
+	ReadWithCheck<FQuat4f>(Context, AttributeToRead, FQuat4f::Identity);
 }
 
 void UNiagaraDataInterfaceParticleRead::ReadID(FVectorVMExternalFunctionContext& Context, FName AttributeToRead)
@@ -1720,7 +1720,7 @@ void UNiagaraDataInterfaceParticleRead::ReadColorByIndex(FVectorVMExternalFuncti
 
 void UNiagaraDataInterfaceParticleRead::ReadQuatByIndex(FVectorVMExternalFunctionContext& Context, FName AttributeToRead)
 {
-	ReadByIndexWithCheck<FQuat>(Context, AttributeToRead, FQuat::Identity);
+	ReadByIndexWithCheck<FQuat4f>(Context, AttributeToRead, FQuat4f::Identity);
 }
 
 void UNiagaraDataInterfaceParticleRead::ReadIDByIndex(FVectorVMExternalFunctionContext& Context, FName AttributeToRead)
diff --git a/Engine/Plugins/FX/Niagara/Source/Niagara/Private/NiagaraDataInterfaceSkeletalMesh.cpp b/Engine/Plugins/FX/Niagara/Source/Niagara/Private/NiagaraDataInterfaceSkeletalMesh.cpp
--- a/Engine/Plugins/FX/Niagara/Source/Niagara/Private/NiagaraDataInterfaceSkeletalMesh.cpp
+++ b/Engine/Plugins/FX/Niagara/Source/Niagara/Private/NiagaraDataInterfaceSkeletalMesh.cpp
@@ -365,7 +365,7 @@ void FSkeletalMeshSkinningData::UpdateBoneTransforms()
 
 	const TArray<FTransform>& BaseCompSpaceTransforms = SkelComp->GetComponentSpaceTransforms();
 	TArray<FMatrix44f>& CurrBones = CurrBoneRefToLocals();
-	TArray<FTransform>& CurrTransforms = CurrComponentTransforms();
+	TArray<FTransform3f>& CurrTransforms = CurrComponentTransforms();
 
 	if (USkinnedMeshComponent* MasterComponent = SkelComp->MasterPoseComponent.Get())
 	{
@@ -395,14 +395,14 @@ void FSkeletalMeshSkinningData::UpdateBoneTransforms()
 					if (MasterIndex != INDEX_NONE && MasterIndex < MasterTransforms.Num())
 					{
 						bFoundMaster = true;
-						CurrTransforms[BoneIndex] = MasterTransforms[MasterIndex];
+						CurrTransforms[BoneIndex] = (FTransform3f)MasterTransforms[MasterIndex];
 					}
 				}
 
 				if ( !bFoundMaster )
 				{
 					const int32 ParentIndex = SkelMesh->GetRefSkeleton().GetParentIndex(BoneIndex);
-					FTransform BoneTransform = SkelMesh->GetRefSkeleton().GetRefBonePose()[BoneIndex];
+					FTransform3f BoneTransform = (FTransform3f)SkelMesh->GetRefSkeleton().GetRefBonePose()[BoneIndex];
 					if ((ParentIndex >= 0) && (ParentIndex < BoneIndex))
 					{
 						BoneTransform = BoneTransform * CurrTransforms[ParentIndex];
@@ -424,7 +424,7 @@ void FSkeletalMeshSkinningData::UpdateBoneTransforms()
 	else
 	{
 		SkelComp->CacheRefToLocalMatrices(CurrBones);
-		CurrTransforms = SkelComp->GetComponentSpaceTransforms();
+		CurrTransforms = LWC::ConvertArrayType<FTransform3f>(SkelComp->GetComponentSpaceTransforms());	// LWC_TODO: Perf pessimization
 	}
 }
 
@@ -1112,9 +1112,9 @@ void FSkeletalMeshGpuDynamicBufferProxy::NewFrame(const FNDISkeletalMesh_Instanc
 			}
 
 			// Append sockets
-			for (const FTransform& SocketTransform : InstanceData->GetFilteredSocketsCurrBuffer())
+			for (const FTransform3f& SocketTransform : InstanceData->GetFilteredSocketsCurrBuffer())
 			{
-				const FQuat Rotation = SocketTransform.GetRotation();
+				const FQuat4f Rotation = SocketTransform.GetRotation();
 				BoneSamplingData.Add(SocketTransform.GetLocation());
 				BoneSamplingData.Add(FVector4(Rotation.X, Rotation.Y, Rotation.Z, Rotation.W));
 				BoneSamplingData.Add(SocketTransform.GetScale3D());
@@ -2138,14 +2138,16 @@ bool FNDISkeletalMesh_InstanceData::Init(UNiagaraDataInterfaceSkeletalMesh* Inte
 			{
 				for (int32 i = 0; i < FilteredSocketInfo.Num(); ++i)
 				{
-					NewSkelComp->GetSocketInfoByName(FilteredSockets[i], FilteredSocketInfo[i].Transform, FilteredSocketInfo[i].BoneIdx);
+					FTransform SocketTransform;
+					NewSkelComp->GetSocketInfoByName(FilteredSockets[i], SocketTransform, FilteredSocketInfo[i].BoneIdx);
+					FilteredSocketInfo[i].Transform = (FTransform3f)SocketTransform;
 				}
 			}
 			else
 			{
 				for (int32 i = 0; i < FilteredSocketInfo.Num(); ++i)
 				{
-					FilteredSocketInfo[i].Transform = FTransform(Mesh->GetComposedRefPoseMatrix(FilteredSockets[i]));
+					FilteredSocketInfo[i].Transform = FTransform3f(Mesh->GetComposedRefPoseMatrix(FilteredSockets[i]));
 					FilteredSocketInfo[i].BoneIdx = INDEX_NONE;
 				}
 			}
@@ -2357,13 +2359,13 @@ bool FNDISkeletalMesh_InstanceData::Tick(UNiagaraDataInterfaceSkeletalMesh* Inte
 void FNDISkeletalMesh_InstanceData::UpdateFilteredSocketTransforms()
 {
 	USkeletalMeshComponent* SkelComp = Cast<USkeletalMeshComponent>(SceneComponent.Get());
-	TArray<FTransform>& WriteBuffer = GetFilteredSocketsWriteBuffer();
+	TArray<FTransform3f>& WriteBuffer = GetFilteredSocketsWriteBuffer();
 
 	for (int32 i = 0; i < FilteredSocketInfo.Num(); ++i)
 	{
 		const FCachedSocketInfo& SocketInfo = FilteredSocketInfo[i];
 		const FTransform& BoneTransform = SocketInfo.BoneIdx != INDEX_NONE ? SkelComp->GetBoneTransform(SocketInfo.BoneIdx, FTransform::Identity) : FTransform::Identity;
-		WriteBuffer[i] = SocketInfo.Transform * BoneTransform;
+		WriteBuffer[i] = SocketInfo.Transform * FTransform3f(BoneTransform);
 	}
 }
 
diff --git a/Engine/Plugins/FX/Niagara/Source/Niagara/Private/NiagaraDataInterfaceStaticMesh.cpp b/Engine/Plugins/FX/Niagara/Source/Niagara/Private/NiagaraDataInterfaceStaticMesh.cpp
--- a/Engine/Plugins/FX/Niagara/Source/Niagara/Private/NiagaraDataInterfaceStaticMesh.cpp
+++ b/Engine/Plugins/FX/Niagara/Source/Niagara/Private/NiagaraDataInterfaceStaticMesh.cpp
@@ -146,9 +146,9 @@ void FStaticMeshGpuSpawnBuffer::Initialise(const FStaticMeshLODResources* Res, c
 	if (NumSockets > 0)
 	{
 		SocketTransformsResourceArray.Reserve(NumFilteredSockets * 3);
-		for (const FTransform& SocketTransform : InstanceData->CachedSockets)
+		for (const FTransform3f& SocketTransform : InstanceData->CachedSockets)
 		{
-			const FQuat SocketRotation = SocketTransform.GetRotation();
+			const FQuat4f SocketRotation = SocketTransform.GetRotation();
 			SocketTransformsResourceArray.Add(FVector4(SocketTransform.GetTranslation(), 0.0f));
 			SocketTransformsResourceArray.Add(FVector4(SocketRotation.X, SocketRotation.Y, SocketRotation.Z, SocketRotation.W));
 			SocketTransformsResourceArray.Add(FVector4(SocketTransform.GetScale3D(), 0.0f));
@@ -397,9 +397,9 @@ bool FNDIStaticMesh_InstanceData::Init(UNiagaraDataInterfaceStaticMesh* Interfac
 			CachedSockets.AddDefaulted(NumMeshSockets);
 			for (int32 i=0; i < NumMeshSockets; ++i)
 			{
-				CachedSockets[i].SetTranslation(Mesh->Sockets[i]->RelativeLocation);
-				CachedSockets[i].SetRotation(FQuat(Mesh->Sockets[i]->RelativeRotation));
-				CachedSockets[i].SetScale3D(Mesh->Sockets[i]->RelativeScale);
+				CachedSockets[i].SetTranslation(FVector3f(Mesh->Sockets[i]->RelativeLocation));
+				CachedSockets[i].SetRotation(FQuat4f(Mesh->Sockets[i]->RelativeRotation));
+				CachedSockets[i].SetScale3D(FVector3f(Mesh->Sockets[i]->RelativeScale));
 			}
 
 			NumFilteredSockets = 0;
@@ -2536,10 +2536,10 @@ void UNiagaraDataInterfaceStaticMesh::GetSocketTransform(FVectorVMExternalFuncti
 	FNDIInputParam<int32> InSocketIndex(Context);
 
 	FNDIOutputParam<FVector3f> OutTranslate(Context);
-	FNDIOutputParam<FQuat> OutRotate(Context);
+	FNDIOutputParam<FQuat4f> OutRotate(Context);
 	FNDIOutputParam<FVector3f> OutScale(Context);
 
-	const FQuat InstRotation = bWorldSpace ? InstData->Rotation : FQuat::Identity;
+	const FQuat4f InstRotation = bWorldSpace ? FQuat4f(InstData->Rotation) : FQuat4f::Identity;
 
 	const int32 SocketMax = InstData->CachedSockets.Num() - 1;
 	if (SocketMax >= 0)
@@ -2548,16 +2548,16 @@ void UNiagaraDataInterfaceStaticMesh::GetSocketTransform(FVectorVMExternalFuncti
 		{
 			const int32 SocketIndex = FMath::Clamp(InSocketIndex.GetAndAdvance(), 0, SocketMax);
 
-			FTransform SocketTransform = InstData->CachedSockets[SocketIndex];
-			OutTranslate.SetAndAdvance(bWorldSpace ? FVector(InstData->Transform.TransformPosition(SocketTransform.GetTranslation())) : SocketTransform.GetTranslation());
+			FTransform3f SocketTransform = (FTransform3f)InstData->CachedSockets[SocketIndex];
+			OutTranslate.SetAndAdvance(bWorldSpace ? FVector3f(InstData->Transform.TransformPosition(SocketTransform.GetTranslation())) : SocketTransform.GetTranslation());
 			OutRotate.SetAndAdvance(bWorldSpace ? InstRotation * SocketTransform.GetRotation() : SocketTransform.GetRotation());
-			OutScale.SetAndAdvance(bWorldSpace ? FVector(InstData->Transform.TransformVector(SocketTransform.GetScale3D())) : SocketTransform.GetScale3D());
+			OutScale.SetAndAdvance(bWorldSpace ? FVector3f(InstData->Transform.TransformVector(SocketTransform.GetScale3D())) : SocketTransform.GetScale3D());
 		}
 	}
 	else
 	{
-		const FVector DefaultTranslate = bWorldSpace ? InstData->Transform.GetOrigin() : FVector::ZeroVector;
-		const FVector DefaultScale = bWorldSpace ? InstData->Transform.ExtractScaling() : FVector::OneVector;
+		const FVector3f DefaultTranslate = bWorldSpace ? FVector3f(InstData->Transform.GetOrigin()) : FVector3f::ZeroVector;
+		const FVector3f DefaultScale = bWorldSpace ? FVector3f(InstData->Transform.ExtractScaling()) : FVector3f::OneVector;
 		for (int32 i = 0; i < Context.GetNumInstances(); ++i)
 		{
 			OutTranslate.SetAndAdvance(DefaultTranslate);
@@ -2574,10 +2574,10 @@ void UNiagaraDataInterfaceStaticMesh::GetFilteredSocketTransform(FVectorVMExtern
 	FNDIInputParam<int32> InSocketIndex(Context);
 
 	FNDIOutputParam<FVector3f> OutTranslate(Context);
-	FNDIOutputParam<FQuat> OutRotate(Context);
+	FNDIOutputParam<FQuat4f> OutRotate(Context);
 	FNDIOutputParam<FVector3f> OutScale(Context);
 
-	const FQuat InstRotation = bWorldSpace ? InstData->Rotation : FQuat::Identity;
+	const FQuat4f InstRotation = bWorldSpace ? FQuat4f(InstData->Rotation) : FQuat4f::Identity;
 
 	const int32 SocketMax = InstData->NumFilteredSockets - 1;
 	if (SocketMax >= 0)
@@ -2586,16 +2586,16 @@ void UNiagaraDataInterfaceStaticMesh::GetFilteredSocketTransform(FVectorVMExtern
 		{
 			const int32 SocketIndex = InstData->FilteredAndUnfilteredSockets[FMath::Clamp(InSocketIndex.GetAndAdvance(), 0, SocketMax)];
 
-			FTransform SocketTransform = InstData->CachedSockets[SocketIndex];
-			OutTranslate.SetAndAdvance(bWorldSpace ? FVector(InstData->Transform.TransformPosition(SocketTransform.GetTranslation())) : SocketTransform.GetTranslation());
+			FTransform3f SocketTransform = (FTransform3f)InstData->CachedSockets[SocketIndex];
+			OutTranslate.SetAndAdvance(bWorldSpace ? FVector3f(InstData->Transform.TransformPosition(SocketTransform.GetTranslation())) : SocketTransform.GetTranslation());
 			OutRotate.SetAndAdvance(bWorldSpace ? InstRotation * SocketTransform.GetRotation() : SocketTransform.GetRotation());
-			OutScale.SetAndAdvance(bWorldSpace ? FVector(InstData->Transform.TransformVector(SocketTransform.GetScale3D())) : SocketTransform.GetScale3D());
+			OutScale.SetAndAdvance(bWorldSpace ? FVector3f(InstData->Transform.TransformVector(SocketTransform.GetScale3D())) : SocketTransform.GetScale3D());
 		}
 	}
 	else
 	{
-		const FVector DefaultTranslate = bWorldSpace ? InstData->Transform.GetOrigin() : FVector::ZeroVector;
-		const FVector DefaultScale = bWorldSpace ? InstData->Transform.ExtractScaling() : FVector::OneVector;
+		const FVector3f DefaultTranslate = bWorldSpace ? FVector3f(InstData->Transform.GetOrigin()) : FVector3f::ZeroVector;
+		const FVector3f DefaultScale = bWorldSpace ? FVector3f(InstData->Transform.ExtractScaling()) : FVector3f::OneVector;
 		for (int32 i = 0; i < Context.GetNumInstances(); ++i)
 		{
 			OutTranslate.SetAndAdvance(DefaultTranslate);
@@ -2612,10 +2612,10 @@ void UNiagaraDataInterfaceStaticMesh::GetUnfilteredSocketTransform(FVectorVMExte
 	FNDIInputParam<int32> InSocketIndex(Context);
 
 	FNDIOutputParam<FVector3f> OutTranslate(Context);
-	FNDIOutputParam<FQuat> OutRotate(Context);
+	FNDIOutputParam<FQuat4f> OutRotate(Context);
 	FNDIOutputParam<FVector3f> OutScale(Context);
 
-	const FQuat InstRotation = bWorldSpace ? InstData->Rotation : FQuat::Identity;
+	const FQuat4f InstRotation = bWorldSpace ? FQuat4f(InstData->Rotation) : FQuat4f::Identity;
 
 	if (InstData->NumFilteredSockets > 0)
 	{
@@ -2627,10 +2627,10 @@ void UNiagaraDataInterfaceStaticMesh::GetUnfilteredSocketTransform(FVectorVMExte
 		{
 			const int32 SocketIndex = InstData->FilteredAndUnfilteredSockets[FMath::Clamp(InSocketIndex.GetAndAdvance(), 0, SocketMax) + SocketOffset];
 
-			FTransform SocketTransform = InstData->CachedSockets[SocketIndex];
-			OutTranslate.SetAndAdvance(bWorldSpace ? FVector(InstData->Transform.TransformPosition(SocketTransform.GetTranslation())) : SocketTransform.GetTranslation());
+			FTransform3f SocketTransform = (FTransform3f)InstData->CachedSockets[SocketIndex];
+			OutTranslate.SetAndAdvance(bWorldSpace ? FVector3f(InstData->Transform.TransformPosition(SocketTransform.GetTranslation())) : SocketTransform.GetTranslation());
 			OutRotate.SetAndAdvance(bWorldSpace ? InstRotation * SocketTransform.GetRotation() : SocketTransform.GetRotation());
-			OutScale.SetAndAdvance(bWorldSpace ? FVector(InstData->Transform.TransformVector(SocketTransform.GetScale3D())) : SocketTransform.GetScale3D());
+			OutScale.SetAndAdvance(bWorldSpace ? FVector3f(InstData->Transform.TransformVector(SocketTransform.GetScale3D())) : SocketTransform.GetScale3D());
 		}
 	}
 	else if (InstData->CachedSockets.Num() > 0)
@@ -2640,16 +2640,16 @@ void UNiagaraDataInterfaceStaticMesh::GetUnfilteredSocketTransform(FVectorVMExte
 		{
 			const int32 SocketIndex = FMath::Clamp(InSocketIndex.GetAndAdvance(), 0, SocketMax);
 
-			FTransform SocketTransform = InstData->CachedSockets[SocketIndex];
-			OutTranslate.SetAndAdvance(bWorldSpace ? FVector(InstData->Transform.TransformPosition(SocketTransform.GetTranslation())) : SocketTransform.GetTranslation());
+			FTransform3f SocketTransform = (FTransform3f)InstData->CachedSockets[SocketIndex];
+			OutTranslate.SetAndAdvance(bWorldSpace ? FVector3f(InstData->Transform.TransformPosition(SocketTransform.GetTranslation())) : SocketTransform.GetTranslation());
 			OutRotate.SetAndAdvance(bWorldSpace ? InstRotation * SocketTransform.GetRotation() : SocketTransform.GetRotation());
-			OutScale.SetAndAdvance(bWorldSpace ? FVector(InstData->Transform.TransformVector(SocketTransform.GetScale3D())) : SocketTransform.GetScale3D());
+			OutScale.SetAndAdvance(bWorldSpace ? FVector3f(InstData->Transform.TransformVector(SocketTransform.GetScale3D())) : SocketTransform.GetScale3D());
 		}
 	}
 	else
 	{
-		const FVector DefaultTranslate = bWorldSpace ? InstData->Transform.GetOrigin() : FVector::ZeroVector;
-		const FVector DefaultScale = bWorldSpace ? InstData->Transform.ExtractScaling() : FVector::OneVector;
+		const FVector3f DefaultTranslate = bWorldSpace ? FVector3f(InstData->Transform.GetOrigin()) : FVector3f::ZeroVector;
+		const FVector3f DefaultScale = bWorldSpace ? FVector3f(InstData->Transform.ExtractScaling()) : FVector3f::OneVector;
 		for (int32 i = 0; i < Context.GetNumInstances(); ++i)
 		{
 			OutTranslate.SetAndAdvance(DefaultTranslate);
diff --git a/Engine/Plugins/FX/Niagara/Source/Niagara/Private/NiagaraModule.cpp b/Engine/Plugins/FX/Niagara/Source/Niagara/Private/NiagaraModule.cpp
--- a/Engine/Plugins/FX/Niagara/Source/Niagara/Private/NiagaraModule.cpp
+++ b/Engine/Plugins/FX/Niagara/Source/Niagara/Private/NiagaraModule.cpp
@@ -670,7 +670,7 @@ void FNiagaraTypeDefinition::Init()
 	FNiagaraTypeDefinition::Vec3Struct = FindObjectChecked<UScriptStruct>(CoreUObjectPkg, TEXT("Vector3f"));
 	FNiagaraTypeDefinition::Vec4Struct = FindObjectChecked<UScriptStruct>(CoreUObjectPkg, TEXT("Vector4"));
 	FNiagaraTypeDefinition::ColorStruct = FindObjectChecked<UScriptStruct>(CoreUObjectPkg, TEXT("LinearColor"));
-	FNiagaraTypeDefinition::QuatStruct = FindObjectChecked<UScriptStruct>(CoreUObjectPkg, TEXT("Quat"));
+	FNiagaraTypeDefinition::QuatStruct = FindObjectChecked<UScriptStruct>(CoreUObjectPkg, TEXT("Quat4f"));
 
 	FNiagaraTypeDefinition::UObjectClass = UObject::StaticClass();
 	FNiagaraTypeDefinition::UMaterialClass = UMaterialInterface::StaticClass();
diff --git a/Engine/Plugins/FX/Niagara/Source/Niagara/Private/NiagaraParameterCollection.cpp b/Engine/Plugins/FX/Niagara/Source/Niagara/Private/NiagaraParameterCollection.cpp
--- a/Engine/Plugins/FX/Niagara/Source/Niagara/Private/NiagaraParameterCollection.cpp
+++ b/Engine/Plugins/FX/Niagara/Source/Niagara/Private/NiagaraParameterCollection.cpp
@@ -280,7 +280,7 @@ FVector2D UNiagaraParameterCollectionInstance::GetVector2DParameter(const FStrin
 
 FVector UNiagaraParameterCollectionInstance::GetVectorParameter(const FString& InVariableName)
 {
-	return ParameterStorage.GetParameterValue<FVector3f>(FNiagaraVariable(FNiagaraTypeDefinition::GetVec3Def(), *Collection->ParameterNameFromFriendlyName(InVariableName)));
+	return (FVector)ParameterStorage.GetParameterValue<FVector3f>(FNiagaraVariable(FNiagaraTypeDefinition::GetVec3Def(), *Collection->ParameterNameFromFriendlyName(InVariableName)));
 }
 
 FVector4 UNiagaraParameterCollectionInstance::GetVector4Parameter(const FString& InVariableName)
@@ -291,7 +291,7 @@ FVector4 UNiagaraParameterCollectionInstance::GetVector4Parameter(const FString&
 
 FQuat UNiagaraParameterCollectionInstance::GetQuatParameter(const FString& InVariableName)
 {
-	return ParameterStorage.GetParameterValue<FQuat>(FNiagaraVariable(FNiagaraTypeDefinition::GetQuatDef(), *Collection->ParameterNameFromFriendlyName(InVariableName)));
+	return (FQuat)ParameterStorage.GetParameterValue<FQuat4f>(FNiagaraVariable(FNiagaraTypeDefinition::GetQuatDef(), *Collection->ParameterNameFromFriendlyName(InVariableName)));
 }
 
 FLinearColor UNiagaraParameterCollectionInstance::GetColorParameter(const FString& InVariableName)
diff --git a/Engine/Plugins/FX/Niagara/Source/Niagara/Private/NiagaraRendererComponents.cpp b/Engine/Plugins/FX/Niagara/Source/Niagara/Private/NiagaraRendererComponents.cpp
--- a/Engine/Plugins/FX/Niagara/Source/Niagara/Private/NiagaraRendererComponents.cpp
+++ b/Engine/Plugins/FX/Niagara/Source/Niagara/Private/NiagaraRendererComponents.cpp
@@ -49,7 +49,7 @@ void SetVariableByType(FNiagaraVariable& DataVariable, FNiagaraDataSet& Data, in
 	else if (VarType == FNiagaraTypeDefinition::GetVec3Def()) { SetValueWithAccessor<FVector3f>(DataVariable, Data, ParticleIndex); }
 	else if (VarType == FNiagaraTypeDefinition::GetVec4Def()) { SetValueWithAccessor<FVector4>(DataVariable, Data, ParticleIndex); }
 	else if (VarType == FNiagaraTypeDefinition::GetColorDef()) { SetValueWithAccessor<FLinearColor>(DataVariable, Data, ParticleIndex); }
-	else if (VarType == FNiagaraTypeDefinition::GetQuatDef()) { SetValueWithAccessor<FQuat>(DataVariable, Data, ParticleIndex); }
+	else if (VarType == FNiagaraTypeDefinition::GetQuatDef()) { SetValueWithAccessor<FQuat4f>(DataVariable, Data, ParticleIndex); }
 }
 
 void ConvertVariableToType(const FNiagaraVariable& SourceVariable, FNiagaraVariable& TargetVariable)
@@ -90,9 +90,14 @@ void ConvertVariableToType(const FNiagaraVariable& SourceVariable, FNiagaraVaria
 		FVector3f Data = SourceVariable.GetValue<FVector3f>();
 		TargetVariable.SetValue<FVector>(FVector(Data));
 	}
+	else if (SourceType == FNiagaraTypeDefinition::GetQuatDef() && TargetType == UNiagaraComponentRendererProperties::GetFQuatDef())
+	{
+		FQuat4f Data = SourceVariable.GetValue<FQuat4f>();
+		TargetVariable.SetValue<FQuat>(FQuat(Data));
+	}
 	else if (SourceType == FNiagaraTypeDefinition::GetQuatDef() && TargetType == UNiagaraComponentRendererProperties::GetFRotatorDef())
 	{
-		FQuat Data = SourceVariable.GetValue<FQuat>();
+		FQuat4f Data = SourceVariable.GetValue<FQuat4f>();
 		TargetVariable.SetValue<FRotator>(Data.Rotator());
 	}
 }
diff --git a/Engine/Plugins/FX/Niagara/Source/Niagara/Private/NiagaraSystemInstance.cpp b/Engine/Plugins/FX/Niagara/Source/Niagara/Private/NiagaraSystemInstance.cpp
--- a/Engine/Plugins/FX/Niagara/Source/Niagara/Private/NiagaraSystemInstance.cpp
+++ b/Engine/Plugins/FX/Niagara/Source/Niagara/Private/NiagaraSystemInstance.cpp
@@ -1930,6 +1930,7 @@ void FNiagaraSystemInstance::TickInstanceParameters_Concurrent()
 
 		const FVector Location = GatheredInstanceParameters.ComponentTrans.GetLocation();
 		const FVector LastLocation = FMath::IsNearlyZero(CurrentSystemParameters.EngineSystemAge) ? Location : FVector(OwnerParameters[GetParameterIndex(true)].EnginePosition);
+		const FQuat LastRotation = GatheredInstanceParameters.ComponentTrans.GetRotation();
 
 		CurrentOwnerParameters.EngineLocalToWorld = LocalToWorld;
 		CurrentOwnerParameters.EngineWorldToLocal = LocalToWorld.Inverse();
@@ -1937,7 +1938,7 @@ void FNiagaraSystemInstance::TickInstanceParameters_Concurrent()
 		CurrentOwnerParameters.EngineWorldToLocalTransposed = CurrentOwnerParameters.EngineWorldToLocal.GetTransposed();
 		CurrentOwnerParameters.EngineLocalToWorldNoScale = LocalToWorldNoScale;
 		CurrentOwnerParameters.EngineWorldToLocalNoScale = LocalToWorldNoScale.Inverse();
-		CurrentOwnerParameters.EngineRotation = GatheredInstanceParameters.ComponentTrans.GetRotation();
+		CurrentOwnerParameters.EngineRotation = FQuat4f((float)LastRotation.X, (float)LastRotation.Y, (float)LastRotation.Z, (float)LastRotation.W);
 		CurrentOwnerParameters.EnginePosition = Location;
 		CurrentOwnerParameters.EngineVelocity = (Location - LastLocation) / GatheredInstanceParameters.DeltaSeconds;
 		CurrentOwnerParameters.EngineXAxis = CurrentOwnerParameters.EngineRotation.GetAxisX();
diff --git a/Engine/Plugins/FX/Niagara/Source/Niagara/Public/NiagaraComponentRendererProperties.h b/Engine/Plugins/FX/Niagara/Source/Niagara/Public/NiagaraComponentRendererProperties.h
--- a/Engine/Plugins/FX/Niagara/Source/Niagara/Public/NiagaraComponentRendererProperties.h
+++ b/Engine/Plugins/FX/Niagara/Source/Niagara/Public/NiagaraComponentRendererProperties.h
@@ -134,6 +134,7 @@ public:
 	static FNiagaraTypeDefinition GetFColorDef();
 	static FNiagaraTypeDefinition GetFRotatorDef();
 	static FNiagaraTypeDefinition GetFVectorDef();
+	static FNiagaraTypeDefinition GetFQuatDef();
 
 	virtual void CacheFromCompiledData(const FNiagaraDataSetCompiledData* CompiledData) override;
 
diff --git a/Engine/Plugins/FX/Niagara/Source/Niagara/Public/NiagaraParameterStore.h b/Engine/Plugins/FX/Niagara/Source/Niagara/Public/NiagaraParameterStore.h
--- a/Engine/Plugins/FX/Niagara/Source/Niagara/Public/NiagaraParameterStore.h
+++ b/Engine/Plugins/FX/Niagara/Source/Niagara/Public/NiagaraParameterStore.h
@@ -663,6 +663,12 @@ FORCEINLINE bool FNiagaraParameterStore::SetParameterValue(const FVector3d& InVa
 	return SetParameterValue((FVector3f)InValue, Param, bAdd);
 }
 
+template<>
+FORCEINLINE bool FNiagaraParameterStore::SetParameterValue(const FQuat4d& InValue, const FNiagaraVariable& Param, bool bAdd)
+{
+	return SetParameterValue((FQuat4f)InValue, Param, bAdd);
+}
+
 FORCEINLINE_DEBUGGABLE void FNiagaraParameterStore::Tick()
 {
 #if NIAGARA_NAN_CHECKING
@@ -1079,9 +1085,9 @@ struct FNiagaraParameterDirectBinding<FVector4>
 };
 
 template<>
-struct FNiagaraParameterDirectBinding<FQuat>
+struct FNiagaraParameterDirectBinding<FQuat4f>
 {
-	mutable FQuat* ValuePtr;
+	mutable FQuat4f* ValuePtr;
 #if NIAGARA_VALIDATE_DIRECT_BINDINGS
 	FNiagaraParameterStore* BoundStore;
 	FNiagaraVariable BoundVariable;
@@ -1095,40 +1101,40 @@ struct FNiagaraParameterDirectBinding<FQuat>
 #endif
 	{}
 
-	FQuat* Init(FNiagaraParameterStore& InStore, const FNiagaraVariable& DestVariable)
+	FQuat4f* Init(FNiagaraParameterStore& InStore, const FNiagaraVariable& DestVariable)
 	{
 #if NIAGARA_VALIDATE_DIRECT_BINDINGS
 		BoundStore = &InStore;
 		BoundVariable = DestVariable;
 		LayoutVersion = BoundStore->GetLayoutVersion();
 #endif
-		check(DestVariable.GetSizeInBytes() == sizeof(FQuat));
-		ValuePtr = (FQuat*)InStore.GetParameterData(DestVariable);
+		check(DestVariable.GetSizeInBytes() == sizeof(FQuat4f));
+		ValuePtr = (FQuat4f*)InStore.GetParameterData(DestVariable);
 		return ValuePtr;
 	}
 
-	FORCEINLINE void SetValue(const FQuat& InValue)
+	FORCEINLINE void SetValue(const FQuat4f& InValue)
 	{
 #if NIAGARA_VALIDATE_DIRECT_BINDINGS
-		checkSlow(BoundVariable.GetSizeInBytes() == sizeof(FQuat));
+		checkSlow(BoundVariable.GetSizeInBytes() == sizeof(FQuat4f));
 		checkfSlow(LayoutVersion == BoundStore->GetLayoutVersion(), TEXT("This binding is invalid, its bound parameter store's layout was changed since it was created"));
 #endif
 		if (ValuePtr)
 		{
-			FMemory::Memcpy(ValuePtr, &InValue, sizeof(FQuat));//Temp annoyance until we fix the alignment issues with parameter stores.
+			FMemory::Memcpy(ValuePtr, &InValue, sizeof(FQuat4f));//Temp annoyance until we fix the alignment issues with parameter stores.
 		}
 	}
 
-	FORCEINLINE FQuat GetValue()const
+	FORCEINLINE FQuat4f GetValue()const
 	{
 #if NIAGARA_VALIDATE_DIRECT_BINDINGS
-		checkSlow(BoundVariable.GetSizeInBytes() == sizeof(FQuat));
+		checkSlow(BoundVariable.GetSizeInBytes() == sizeof(FQuat4f));
 		checkfSlow(LayoutVersion == BoundStore->GetLayoutVersion(), TEXT("This binding is invalid, its bound parameter store's layout was changed since it was created"));
 #endif
-		FQuat Ret;
+		FQuat4f Ret;
 		if (ValuePtr)
 		{
-			FMemory::Memcpy(&Ret, ValuePtr, sizeof(FQuat));//Temp annoyance until we fix the alignment issues with parameter stores.
+			FMemory::Memcpy(&Ret, ValuePtr, sizeof(FQuat4f));//Temp annoyance until we fix the alignment issues with parameter stores.
 		}
 		return Ret;
 	}
diff --git a/Engine/Plugins/FX/Niagara/Source/Niagara/Public/NiagaraTypes.h b/Engine/Plugins/FX/Niagara/Source/Niagara/Public/NiagaraTypes.h
--- a/Engine/Plugins/FX/Niagara/Source/Niagara/Public/NiagaraTypes.h
+++ b/Engine/Plugins/FX/Niagara/Source/Niagara/Public/NiagaraTypes.h
@@ -1043,7 +1043,7 @@ const FNiagaraTypeDefinition& FNiagaraTypeDefinition::Get()
 	if (TIsSame<T, FVector4>::Value) { return FNiagaraTypeDefinition::GetVec4Def(); }
 	if (TIsSame<T, int32>::Value) { return FNiagaraTypeDefinition::GetIntDef(); }
 	if (TIsSame<T, FNiagaraBool>::Value) { return FNiagaraTypeDefinition::GetBoolDef(); }
-	if (TIsSame<T, FQuat>::Value) { return FNiagaraTypeDefinition::GetQuatDef(); }
+	if (TIsSame<T, FQuat4f>::Value) { return FNiagaraTypeDefinition::GetQuatDef(); }
 	if (TIsSame<T, FMatrix44f>::Value) { return FNiagaraTypeDefinition::GetMatrix4Def(); }
 	if (TIsSame<T, FLinearColor>::Value) { return FNiagaraTypeDefinition::GetColorDef(); }
 	if (TIsSame<T, FNiagaraID>::Value) { return FNiagaraTypeDefinition::GetIDDef(); }
@@ -1705,7 +1705,7 @@ struct alignas(16) FNiagaraOwnerParameters
 	FMatrix44f EngineWorldToLocalTransposed = FMatrix44f::Identity;
 	FMatrix44f EngineLocalToWorldNoScale = FMatrix44f::Identity;
 	FMatrix44f EngineWorldToLocalNoScale = FMatrix44f::Identity;
-	FQuat EngineRotation = FQuat::Identity;
+	FQuat4f EngineRotation = FQuat4f::Identity;
 	FVector4 EnginePosition = FVector4(EForceInit::ForceInitToZero);
 	FVector4 EngineVelocity = FVector4(EForceInit::ForceInitToZero);
 	FVector4 EngineXAxis = FVector4(1.0f, 0.0f, 0.0f, 0.0f);
diff --git a/Engine/Plugins/FX/Niagara/Source/NiagaraEditor/Private/TypeEditorUtilities/NiagaraVectorTypeEditorUtilities.cpp b/Engine/Plugins/FX/Niagara/Source/NiagaraEditor/Private/TypeEditorUtilities/NiagaraVectorTypeEditorUtilities.cpp
--- a/Engine/Plugins/FX/Niagara/Source/NiagaraEditor/Private/TypeEditorUtilities/NiagaraVectorTypeEditorUtilities.cpp
+++ b/Engine/Plugins/FX/Niagara/Source/NiagaraEditor/Private/TypeEditorUtilities/NiagaraVectorTypeEditorUtilities.cpp
@@ -366,13 +366,13 @@ public:
 	virtual void UpdateInternalValueFromStruct(TSharedRef<FStructOnScope> Struct) override
 	{
 		checkf(Struct->GetStruct() == FNiagaraTypeDefinition::GetQuatStruct(), TEXT("Struct type not supported."));
-		VectorValue = *((FQuat*)Struct->GetStructMemory());
+		VectorValue = *((FQuat4f*)Struct->GetStructMemory());
 	}
 
 	virtual void UpdateStructFromInternalValue(TSharedRef<FStructOnScope> Struct) override
 	{
 		checkf(Struct->GetStruct() == FNiagaraTypeDefinition::GetQuatStruct(), TEXT("Struct type not supported."));
-		*((FQuat*)Struct->GetStructMemory()) = VectorValue;
+		*((FQuat4f*)Struct->GetStructMemory()) = VectorValue;
 	}
 
 protected:
@@ -416,7 +416,7 @@ protected:
 	}
 
 private:
-	FQuat VectorValue;
+	FQuat4f VectorValue;
 };
 
 TSharedPtr<SNiagaraParameterEditor> FNiagaraEditorQuatTypeUtilities::CreateParameterEditor(const FNiagaraTypeDefinition& ParameterType) const
@@ -433,7 +433,7 @@ FString FNiagaraEditorQuatTypeUtilities::GetPinDefaultStringFromValue(const FNia
 {
 	checkf(AllocatedVariable.IsDataAllocated(), TEXT("Can not generate a default value string for an unallocated variable."));
 	// NOTE: We can not use ToString() here since the vector pin control doesn't use the standard 'X=0,Y=0,Z=0,W=0' syntax.
-	FQuat Value = AllocatedVariable.GetValue<FQuat>();
+	FQuat4f Value = AllocatedVariable.GetValue<FQuat4f>();
 	return FString::Printf(TEXT("%3.3f,%3.3f,%3.3f,%3.3f"), Value.X, Value.Y, Value.Z, Value.W);
 }
 
@@ -443,8 +443,8 @@ bool FNiagaraEditorQuatTypeUtilities::SetValueFromPinDefaultString(const FString
 	FVector4 Value(0, 0, 0, 0);
 	if (FDefaultValueHelper::ParseVector4(StringValue, Value) || !Variable.IsDataAllocated())
 	{
-		FQuat Quat(Value.X, Value.Y, Value.Z, Value.W);
-		Variable.SetValue<FQuat>(Quat);
+		FQuat4f Quat(Value.X, Value.Y, Value.Z, Value.W);
+		Variable.SetValue<FQuat4f>(Quat);
 		return true;
 	}
 	return false;
@@ -457,14 +457,14 @@ FText FNiagaraEditorQuatTypeUtilities::GetSearchTextFromValue(const FNiagaraVari
 
 FText FNiagaraEditorQuatTypeUtilities::GetStackDisplayText(FNiagaraVariable& Variable) const
 {
-	FQuat Value = Variable.GetValue<FQuat>();
+	FQuat4f Value = Variable.GetValue<FQuat4f>();
 	return FText::Format(FText::FromString("({0}, {1}, {2}, {3})"), Value.X, Value.Y, Value.Z, Value.W);
 }
 
 void FNiagaraEditorQuatTypeUtilities::UpdateVariableWithDefaultValue(FNiagaraVariable& Variable) const
 {
 	checkf(Variable.GetType().GetStruct() == FNiagaraTypeDefinition::GetQuatStruct(), TEXT("Struct type not supported."));
-	Variable.SetValue<FQuat>(FQuat(0, 0, 0, 1));
+	Variable.SetValue<FQuat4f>(FQuat4f(0, 0, 0, 1));
 }
 
 bool FNiagaraEditorNiagaraIDTypeUtilities::CanHandlePinDefaults() const
diff --git a/Engine/Plugins/Runtime/AR/AppleAR/AppleARKit/Source/AppleARKit/Public/AppleARKitMeshData.h b/Engine/Plugins/Runtime/AR/AppleAR/AppleARKit/Source/AppleARKit/Public/AppleARKitMeshData.h
--- a/Engine/Plugins/Runtime/AR/AppleAR/AppleARKit/Source/AppleARKit/Public/AppleARKitMeshData.h
+++ b/Engine/Plugins/Runtime/AR/AppleAR/AppleARKit/Source/AppleARKit/Public/AppleARKitMeshData.h
@@ -11,7 +11,6 @@
 	#import <ARKit/ARKit.h>
 #endif
 
-struct FTransform;
 
 /** Represents a cached version of the data coming from an ARKit mesh anchor */
 class FARKitMeshData : public IMRMesh::FBrickDataReceipt
diff --git a/Engine/Plugins/Runtime/Oculus/OculusVR/Source/OculusHMD/Private/OculusHMDPrivate.h b/Engine/Plugins/Runtime/Oculus/OculusVR/Source/OculusHMD/Private/OculusHMDPrivate.h
--- a/Engine/Plugins/Runtime/Oculus/OculusVR/Source/OculusHMD/Private/OculusHMDPrivate.h
+++ b/Engine/Plugins/Runtime/Oculus/OculusVR/Source/OculusHMD/Private/OculusHMDPrivate.h
@@ -70,7 +70,7 @@ namespace OculusHMD
 	/** Converts FQuat to ovrpQuatf */
 	FORCEINLINE ovrpQuatf ToOvrpQuatf(const FQuat& InQuat)
 	{
-		return ovrpQuatf { InQuat.Y, InQuat.Z, -InQuat.X, -InQuat.W };
+		return ovrpQuatf { float(InQuat.Y), float(InQuat.Z), float(-InQuat.X), float(-InQuat.W) };
 	}
 
 	/** Converts vector from Oculus to Unreal */
diff --git a/Engine/Plugins/Runtime/OpenXR/Source/OpenXRHMD/Public/OpenXRCore.h b/Engine/Plugins/Runtime/OpenXR/Source/OpenXRHMD/Public/OpenXRCore.h
--- a/Engine/Plugins/Runtime/OpenXR/Source/OpenXRHMD/Public/OpenXRCore.h
+++ b/Engine/Plugins/Runtime/OpenXR/Source/OpenXRHMD/Public/OpenXRCore.h
@@ -43,7 +43,7 @@ FORCEINLINE FQuat ToFQuat(XrQuaternionf Quat)
 
 FORCEINLINE XrQuaternionf ToXrQuat(FQuat Quat)
 {
-	return XrQuaternionf{ Quat.Y, Quat.Z, -Quat.X, -Quat.W };
+	return XrQuaternionf{ (float)Quat.Y, (float)Quat.Z, -(float)Quat.X, -(float)Quat.W };
 }
 
 FORCEINLINE FVector ToFVector(XrVector3f Vector, float Scale = 1.0f)
diff --git a/Engine/Plugins/Runtime/OpenXRHandTracking/Source/OpenXRHandTracking/Private/OpenXRHandTracking.cpp b/Engine/Plugins/Runtime/OpenXRHandTracking/Source/OpenXRHandTracking/Private/OpenXRHandTracking.cpp
--- a/Engine/Plugins/Runtime/OpenXRHandTracking/Source/OpenXRHandTracking/Private/OpenXRHandTracking.cpp
+++ b/Engine/Plugins/Runtime/OpenXRHandTracking/Source/OpenXRHandTracking/Private/OpenXRHandTracking.cpp
@@ -528,7 +528,7 @@ bool FOpenXRHandTracking::GetKeypointState(EControllerHand Hand, EHandKeypoint K
 	return gotTransform;
 }
 
-bool FOpenXRHandTracking::GetAllKeypointStates(EControllerHand Hand, TArray<FVector>& OutPositions, TArray<struct FQuat>& OutRotations, TArray<float>& OutRadii) const
+bool FOpenXRHandTracking::GetAllKeypointStates(EControllerHand Hand, TArray<FVector>& OutPositions, TArray<FQuat>& OutRotations, TArray<float>& OutRadii) const
 {
 	if (!bHandTrackingAvailable)
 	{
diff --git a/Engine/Plugins/Runtime/OpenXRHandTracking/Source/OpenXRHandTracking/Private/OpenXRHandTracking.h b/Engine/Plugins/Runtime/OpenXRHandTracking/Source/OpenXRHandTracking/Private/OpenXRHandTracking.h
--- a/Engine/Plugins/Runtime/OpenXRHandTracking/Source/OpenXRHandTracking/Private/OpenXRHandTracking.h
+++ b/Engine/Plugins/Runtime/OpenXRHandTracking/Source/OpenXRHandTracking/Private/OpenXRHandTracking.h
@@ -89,7 +89,7 @@ public:
 	virtual FName GetHandTrackerDeviceTypeName() const override;
 	virtual bool IsHandTrackingStateValid() const override;
 	virtual bool GetKeypointState(EControllerHand Hand, EHandKeypoint Keypoint, FTransform& OutTransform, float& OutRadius) const override;
-	virtual bool GetAllKeypointStates(EControllerHand Hand, TArray<FVector>& OutPositions, TArray<struct FQuat>& OutRotations, TArray<float>& OutRadii) const override;
+	virtual bool GetAllKeypointStates(EControllerHand Hand, TArray<FVector>& OutPositions, TArray<FQuat>& OutRotations, TArray<float>& OutRadii) const override;
 
 private:
 	FHandState& GetLeftHandState();
diff --git a/Engine/Plugins/VirtualProduction/CameraCalibration/Source/CameraCalibrationEditor/Private/Calibrators/CameraNodalOffsetAlgoPoints.cpp b/Engine/Plugins/VirtualProduction/CameraCalibration/Source/CameraCalibrationEditor/Private/Calibrators/CameraNodalOffsetAlgoPoints.cpp
--- a/Engine/Plugins/VirtualProduction/CameraCalibration/Source/CameraCalibrationEditor/Private/Calibrators/CameraNodalOffsetAlgoPoints.cpp
+++ b/Engine/Plugins/VirtualProduction/CameraCalibration/Source/CameraCalibrationEditor/Private/Calibrators/CameraNodalOffsetAlgoPoints.cpp
@@ -150,13 +150,13 @@ namespace CameraNodalOffsetAlgoPoints
 
 		// Average the rotation
 
-		float AverageQuatVec[4] = { 0 }; // Simple averaging should work for similar quaterions, which these are.
+		FQuat::FReal AverageQuatVec[4] = { 0 }; // Simple averaging should work for similar quaterions, which these are.
 
 		for (const FSinglePoseResult& SinglePoseResult : SinglePoseResults)
 		{
 			const FQuat Rotation = SinglePoseResult.Transform.GetRotation();
 
-			const float ThisQuat[4] = {
+			const FQuat::FReal ThisQuat[4] = {
 				Rotation.X,
 				Rotation.Y,
 				Rotation.Z,
diff --git a/Engine/Plugins/VirtualProduction/CameraCalibrationCore/Source/CameraCalibrationCore/Public/CameraCalibrationUtils.h b/Engine/Plugins/VirtualProduction/CameraCalibrationCore/Source/CameraCalibrationCore/Public/CameraCalibrationUtils.h
--- a/Engine/Plugins/VirtualProduction/CameraCalibrationCore/Source/CameraCalibrationCore/Public/CameraCalibrationUtils.h
+++ b/Engine/Plugins/VirtualProduction/CameraCalibrationCore/Source/CameraCalibrationCore/Public/CameraCalibrationUtils.h
@@ -4,7 +4,6 @@
 
 #include "CoreMinimal.h"
 
-struct FTransform;
 
 class UTexture2D;
 
diff --git a/Engine/Source/Developer/Datasmith/DatasmithFacade/Private/DatasmithFacadeActor.cpp b/Engine/Source/Developer/Datasmith/DatasmithFacade/Private/DatasmithFacadeActor.cpp
--- a/Engine/Source/Developer/Datasmith/DatasmithFacade/Private/DatasmithFacadeActor.cpp
+++ b/Engine/Source/Developer/Datasmith/DatasmithFacade/Private/DatasmithFacadeActor.cpp
@@ -103,10 +103,10 @@ void FDatasmithFacadeActor::GetRotation(
 {
 	FQuat RotationQuat(GetDatasmithActorElement()->GetRotation());
 
-	OutX = RotationQuat.X;
-	OutY = RotationQuat.Y;
-	OutZ = RotationQuat.Z;
-	OutW = RotationQuat.W;
+	OutX = (float)RotationQuat.X;
+	OutY = (float)RotationQuat.Y;
+	OutZ = (float)RotationQuat.Z;
+	OutW = (float)RotationQuat.W;
 }
 
 void FDatasmithFacadeActor::SetTranslation(
diff --git a/Engine/Source/Developer/Datasmith/DatasmithFacade/Private/DatasmithFacadeActorLight.cpp b/Engine/Source/Developer/Datasmith/DatasmithFacade/Private/DatasmithFacadeActorLight.cpp
--- a/Engine/Source/Developer/Datasmith/DatasmithFacade/Private/DatasmithFacadeActorLight.cpp
+++ b/Engine/Source/Developer/Datasmith/DatasmithFacade/Private/DatasmithFacadeActorLight.cpp
@@ -179,10 +179,10 @@ void FDatasmithFacadeActorLight::GetIesRotation(
 ) const
 {
 	FQuat IesQuat(GetDatasmithLightActorElement()->GetIesRotation());
-	OutX = IesQuat.X;
-	OutY = IesQuat.Y;
-	OutZ = IesQuat.Z;
-	OutW = IesQuat.W;
+	OutX = (float)IesQuat.X;
+	OutY = (float)IesQuat.Y;
+	OutZ = (float)IesQuat.Z;
+	OutW = (float)IesQuat.W;
 }
 
 void FDatasmithFacadeActorLight::GetIesRotation(
