Fixed a bug where the Niagara translator could compile unused parts of the graph, resulting in errors where unset variables are being read

diff --git a/Engine/Plugins/FX/Niagara/Source/NiagaraCore/Private/NiagaraCustomVersion.cpp b/Engine/Plugins/FX/Niagara/Source/NiagaraCore/Private/NiagaraCustomVersion.cpp
--- a/Engine/Plugins/FX/Niagara/Source/NiagaraCore/Private/NiagaraCustomVersion.cpp
+++ b/Engine/Plugins/FX/Niagara/Source/NiagaraCore/Private/NiagaraCustomVersion.cpp
@@ -8,4 +8,5 @@ const FGuid FNiagaraCustomVersion::GUID(0xFCF57AFA, 0x50764283, 0xB9A9E658, 0xFF
 // Register the custom version with core
 FCustomVersionRegistration GRegisterNiagaraCustomVersion(FNiagaraCustomVersion::GUID, FNiagaraCustomVersion::LatestVersion, TEXT("NiagaraVer"));
 
-const FGuid FNiagaraCustomVersion::LatestScriptCompileVersion(0x02497746, 0x33D75C49, 0x9A60CB65, 0x7D6AC5A4);
+const FGuid FNiagaraCustomVersion::LatestScriptCompileVersion(0x64FA8CB0, 0xF7748343, 0x83D275AE, 0xAE764DF9);
+
diff --git a/Engine/Plugins/FX/Niagara/Source/NiagaraEditor/Private/NiagaraHlslTranslator.cpp b/Engine/Plugins/FX/Niagara/Source/NiagaraEditor/Private/NiagaraHlslTranslator.cpp
--- a/Engine/Plugins/FX/Niagara/Source/NiagaraEditor/Private/NiagaraHlslTranslator.cpp
+++ b/Engine/Plugins/FX/Niagara/Source/NiagaraEditor/Private/NiagaraHlslTranslator.cpp
@@ -6363,7 +6363,7 @@ void FHlslNiagaraTranslator::FunctionCall(UNiagaraNodeFunctionCall* FunctionNode
 
 void FHlslNiagaraTranslator::EnterFunctionCallNode(const TSet<FName>& UnusedInputs)
 {
-	FunctionNodeStack.Add(UnusedInputs);
+	FunctionNodeStack.AddDefaulted_GetRef().UnusedInputs = UnusedInputs;
 }
 
 void FHlslNiagaraTranslator::ExitFunctionCallNode()
@@ -6378,7 +6378,39 @@ bool FHlslNiagaraTranslator::IsFunctionVariableCulledFromCompilation(const FName
 	{
 		return false;
 	}
-	return FunctionNodeStack.Last().Contains(InputName);
+	
+	FunctionNodeStackEntry StackEntry = FunctionNodeStack.Last();
+	if (StackEntry.UnusedInputs.Contains(InputName))
+	{
+		return true;
+	}
+
+	FString InputNameString = InputName.ToString();
+	for (const FString& CulledFunction : StackEntry.CulledFunctionNames)
+	{
+		// if the entire function call was culled, we don't want to compile anything related to it
+		if (InputNameString.StartsWith(CulledFunction + "."))
+		{
+			return true;
+		}
+	
+	}
+	return false;
+}
+
+void FHlslNiagaraTranslator::CullMapSetInputPin(UEdGraphPin* InputPin)
+{
+	if (InputPin == nullptr || InputPin->LinkedTo.Num() != 1 || FunctionNodeStack.Num() == 0)
+	{
+		return;
+	}
+
+	// when a map set input is culled that is connected to a function call node (as is the case for dynamic inputs), we also need to cull any upstream pins that set inputs for the culled function call node
+	if (UNiagaraNodeFunctionCall* FunctionNode = Cast<UNiagaraNodeFunctionCall>(InputPin->LinkedTo[0]->GetOwningNode()))
+	{
+		FString FunctionScriptName = FunctionNode->GetFunctionName();
+		FunctionNodeStack.Last().CulledFunctionNames.Add(FunctionScriptName);
+	}
 }
 
 // From a valid list of namespaces, resolve any aliased tokens and promote namespaced variables without a master namespace to the input parameter map instance namespace
diff --git a/Engine/Plugins/FX/Niagara/Source/NiagaraEditor/Private/NiagaraHlslTranslator.h b/Engine/Plugins/FX/Niagara/Source/NiagaraEditor/Private/NiagaraHlslTranslator.h
--- a/Engine/Plugins/FX/Niagara/Source/NiagaraEditor/Private/NiagaraHlslTranslator.h
+++ b/Engine/Plugins/FX/Niagara/Source/NiagaraEditor/Private/NiagaraHlslTranslator.h
@@ -330,6 +330,12 @@ public:
 	bool bCallIDInitialized = false;
 };
 
+struct FunctionNodeStackEntry
+{
+	TSet<FName> UnusedInputs;
+	TSet<FString> CulledFunctionNames;
+};
+
 class NIAGARAEDITOR_API FHlslNiagaraTranslator
 {
 public:
@@ -528,6 +534,7 @@ public:
 	void EnterFunctionCallNode(const TSet<FName>& UnusedInputs);
 	void ExitFunctionCallNode();
 	bool IsFunctionVariableCulledFromCompilation(const FName& InputName) const;
+	void CullMapSetInputPin(UEdGraphPin* InputPin);
 
 	virtual void Convert(class UNiagaraNodeConvert* Convert, TArrayView<const int32> Inputs, TArray<int32>& Outputs);
 	virtual void If(class UNiagaraNodeIf* IfNode, TArray<FNiagaraVariable>& Vars, int32 Condition, TArray<int32>& PathA, TArray<int32>& PathB, TArray<int32>& Outputs);
@@ -684,8 +691,8 @@ private:
 	// Synced to the ParamMapHistories.
 	TArray<TArray<int32>> ParamMapSetVariablesToChunks;
 
-	// Used to keep track of contextual information about the currently compiled function node (currently only the unsused parameter names)
-	TArray<TSet<FName>> FunctionNodeStack;
+	// Used to keep track of contextual information about the currently compiled function node
+	TArray<FunctionNodeStackEntry> FunctionNodeStack;
 
 	// Synced to the System uniforms encountered for parameter maps thus far.
 	struct UniformVariableInfo
diff --git a/Engine/Plugins/FX/Niagara/Source/NiagaraEditor/Private/NiagaraNodeParameterMapSet.cpp b/Engine/Plugins/FX/Niagara/Source/NiagaraEditor/Private/NiagaraNodeParameterMapSet.cpp
--- a/Engine/Plugins/FX/Niagara/Source/NiagaraEditor/Private/NiagaraNodeParameterMapSet.cpp
+++ b/Engine/Plugins/FX/Niagara/Source/NiagaraEditor/Private/NiagaraNodeParameterMapSet.cpp
@@ -214,11 +214,7 @@ void UNiagaraNodeParameterMapSet::Compile(class FHlslNiagaraTranslator* Translat
 		Outputs.Add(INDEX_NONE);
 	}
 
-	const UEdGraphSchema_Niagara* Schema = CastChecked<UEdGraphSchema_Niagara>(GetSchema());
-
-	// First compile fully down the hierarchy for our predecessors..
-	TArray<FCompiledPin, TInlineAllocator<16>> CompileInputs;
-	CompileInputs.Reserve(InputPins.Num());
+	// update the translator with the culled function calls before compiling any further
 	for (UEdGraphPin* InputPin : InputPins)
 	{
 		if (IsAddPin(InputPin))
@@ -227,6 +223,19 @@ void UNiagaraNodeParameterMapSet::Compile(class FHlslNiagaraTranslator* Translat
 		}
 
 		if (Translator->IsFunctionVariableCulledFromCompilation(InputPin->PinName))
+		{
+			Translator->CullMapSetInputPin(InputPin);
+		}
+	}
+
+	const UEdGraphSchema_Niagara* Schema = CastChecked<UEdGraphSchema_Niagara>(GetSchema());
+
+	// First compile fully down the hierarchy for our predecessors..
+	TArray<FCompiledPin, TInlineAllocator<16>> CompileInputs;
+	CompileInputs.Reserve(InputPins.Num());
+	for (UEdGraphPin* InputPin : InputPins)
+	{
+		if (IsAddPin(InputPin) || Translator->IsFunctionVariableCulledFromCompilation(InputPin->PinName))
 		{
 			continue;
 		}
