- Allow bounds to persistent in the editor when we are resetting the system for compilation, etc

#rnx

diff --git a/Engine/Plugins/FX/Niagara/Source/Niagara/Private/NiagaraComponent.cpp b/Engine/Plugins/FX/Niagara/Source/Niagara/Private/NiagaraComponent.cpp
--- a/Engine/Plugins/FX/Niagara/Source/Niagara/Private/NiagaraComponent.cpp
+++ b/Engine/Plugins/FX/Niagara/Source/Niagara/Private/NiagaraComponent.cpp
@@ -421,6 +421,7 @@ UNiagaraComponent::UNiagaraComponent(const FObjectInitializer& ObjectInitializer
 	, ScalabilityManagerHandle(INDEX_NONE)
 	, ForceUpdateTransformTime(0.0f)
 	, CurrLocalBounds(ForceInit)
+	, SystemFixedBounds(ForceInit)
 {
 	OverrideParameters.SetOwner(this);
 
@@ -496,6 +497,7 @@ void UNiagaraComponent::SetEmitterEnable(FName EmitterName, bool bNewEnableState
 
 void UNiagaraComponent::SetSystemFixedBounds(FBox LocalBounds)
 {
+	SystemFixedBounds = LocalBounds;
 	if ( SystemInstanceController.IsValid() )
 	{
 		SystemInstanceController->SetSystemFixedBounds(LocalBounds);
@@ -504,15 +506,12 @@ void UNiagaraComponent::SetSystemFixedBounds(FBox LocalBounds)
 
 FBox UNiagaraComponent::GetSystemFixedBounds() const
 {
-	if (SystemInstanceController.IsValid())
-	{
-		return SystemInstanceController->GetSystemFixedBounds();
-	}
-	return FBox(EForceInit::ForceInit);
+	return SystemFixedBounds;
 }
 
 void UNiagaraComponent::ClearSystemFixedBounds()
 {
+	SystemFixedBounds.Init();
 	if (SystemInstanceController.IsValid())
 	{
 		SystemInstanceController->SetSystemFixedBounds(FBox(EForceInit::ForceInit));
@@ -521,6 +520,7 @@ void UNiagaraComponent::ClearSystemFixedBounds()
 
 void UNiagaraComponent::SetEmitterFixedBounds(FName EmitterName, FBox LocalBounds)
 {
+	EmitterFixedBounds.FindOrAdd(EmitterName, LocalBounds);
 	if (SystemInstanceController.IsValid())
 	{
 		SystemInstanceController->SetEmitterFixedBounds(EmitterName, LocalBounds);
@@ -529,15 +529,16 @@ void UNiagaraComponent::SetEmitterFixedBounds(FName EmitterName, FBox LocalBound
 
 FBox UNiagaraComponent::GetEmitterFixedBounds(FName EmitterName) const
 {
-	if (SystemInstanceController.IsValid())
+	if ( const FBox* EmitterBounds = EmitterFixedBounds.Find(EmitterName) )
 	{
-		return SystemInstanceController->GetEmitterFixedBounds(EmitterName);
+		return *EmitterBounds;
 	}
 	return FBox(EForceInit::ForceInit);
 }
 
 void UNiagaraComponent::ClearEmitterFixedBounds(FName EmitterName)
 {
+	EmitterFixedBounds.Remove(EmitterName);
 	if (SystemInstanceController.IsValid())
 	{
 		SystemInstanceController->SetEmitterFixedBounds(EmitterName, FBox(EForceInit::ForceInit));
@@ -873,6 +874,18 @@ bool UNiagaraComponent::InitializeSystem()
 			SystemInstanceController->SetLODDistance(PreviewLODDistance, PreviewMaxDistance, true);
 		}
 
+		if (SystemFixedBounds.IsValid)
+		{
+			SystemInstanceController->SetSystemFixedBounds(SystemFixedBounds);
+		}
+		if (!EmitterFixedBounds.IsEmpty())
+		{
+			for ( auto it=EmitterFixedBounds.CreateConstIterator(); it; ++it )
+			{
+				SystemInstanceController->SetEmitterFixedBounds(it.Key(), it.Value());
+			}
+		}
+
 #if WITH_EDITORONLY_DATA
 		OnSystemInstanceChangedDelegate.Broadcast();
 #endif
@@ -1463,6 +1476,11 @@ void UNiagaraComponent::OnPooledReuse(UWorld* NewWorld)
 	//Leaving as the default of -1000 causes the visibility code to always assume this should be culled until it's first rendered and initialized by the RT.
 	SetLastRenderTime(GetWorld()->GetTimeSeconds());
 
+	ForceUpdateTransformTime = 0.0f;
+	CurrLocalBounds.Init();
+	SystemFixedBounds.Init();
+	EmitterFixedBounds.Empty();
+
 	if (SystemInstanceController != nullptr)
 	{
 		SystemInstanceController->OnPooledReuse(*NewWorld);
@@ -1793,11 +1811,7 @@ FBoxSphereBounds UNiagaraComponent::CalcBounds(const FTransform& LocalToWorld) c
 	}
 	else
 	{
-		FBox SimBounds(ForceInit);
-		if ( SystemInstanceController.IsValid() )
-		{
-			SimBounds = SystemInstanceController->GetSystemFixedBounds();
-		}
+		FBox SimBounds = SystemFixedBounds;
 		if (!SimBounds.IsValid && Asset && Asset->bFixedBounds)
 		{
 			SimBounds = Asset->GetFixedBounds();
diff --git a/Engine/Plugins/FX/Niagara/Source/Niagara/Public/NiagaraComponent.h b/Engine/Plugins/FX/Niagara/Source/Niagara/Public/NiagaraComponent.h
--- a/Engine/Plugins/FX/Niagara/Source/Niagara/Public/NiagaraComponent.h
+++ b/Engine/Plugins/FX/Niagara/Source/Niagara/Public/NiagaraComponent.h
@@ -136,7 +136,6 @@ private:
 	UPROPERTY()
 	uint32 bRenderingEnabled : 1;
 
-
 	//~ Begin UActorComponent Interface.
 protected:
 	virtual void OnRegister() override;
@@ -727,6 +726,10 @@ private:
 	float ForceUpdateTransformTime;
 	FBox CurrLocalBounds;
 
+	FBox SystemFixedBounds;
+	TMap<FName, FBox> EmitterFixedBounds;
+
+
 	float CustomTimeDilation = 1.0f;
 
 	UPROPERTY()
