- Fix missing transform for tangent basis
- Improve tangent basis reconstruction

#rnx

diff --git a/Engine/Plugins/FX/Niagara/Shaders/Private/NiagaraDataInterfaceSkeletalMesh.ush b/Engine/Plugins/FX/Niagara/Shaders/Private/NiagaraDataInterfaceSkeletalMesh.ush
--- a/Engine/Plugins/FX/Niagara/Shaders/Private/NiagaraDataInterfaceSkeletalMesh.ush
+++ b/Engine/Plugins/FX/Niagara/Shaders/Private/NiagaraDataInterfaceSkeletalMesh.ush
@@ -223,11 +223,13 @@ FDISkelMeshSkinnedVertex DISkelMesh_GetSkinnedVertex(in FDISkelMeshContext DICon
 		TangentZ.xyz = mul(CurrBoneMatrix, float4(TangentZ.xyz, 0.0f)).xyz;
 	}
 
+	float3 TangentY = cross(TangentZ.xyz, TangentX.xyz) * TangentZ.w;
+
 	FDISkelMeshSkinnedVertex SkinnedVertex;
 	SkinnedVertex.Position = CurrPosition;
 	SkinnedVertex.PrevPosition = PrevPosition;
-	SkinnedVertex.TangentX = TangentX;
-	SkinnedVertex.TangentY = cross(TangentZ.xyz, TangentX.xyz);
+	SkinnedVertex.TangentX = cross(TangentY, TangentZ.xyz) * TangentZ.w;
+	SkinnedVertex.TangentY = TangentY;
 	SkinnedVertex.TangentZ = TangentZ.xyz;
 	return SkinnedVertex;
 }
@@ -787,9 +789,9 @@ void DISkelMesh_GetSkinnedVertexWS(in FDISkelMeshContext DIContext, in int Verte
     OutPosition = mul(float4(SkinnedVertex.Position, 1.0f), DIContext.InstanceTransform).xyz;
     float3 PrevPosition = mul(float4(SkinnedVertex.PrevPosition, 1.0f), DIContext.InstancePrevTransform).xyz;
     OutVelocity = (SkinnedVertex.Position - SkinnedVertex.PrevPosition) * DIContext.InstanceInvDeltaTime;
-	OutNormal = SkinnedVertex.TangentZ;
-	OutBinormal = SkinnedVertex.TangentY;
-	OutTangent = SkinnedVertex.TangentX;
+    OutNormal   = normalize(mul(float4(SkinnedVertex.TangentZ, 0.0f), DIContext.InstanceTransform).xyz);
+    OutBinormal = normalize(mul(float4(SkinnedVertex.TangentY, 0.0f), DIContext.InstanceTransform).xyz);
+    OutTangent  = normalize(mul(float4(SkinnedVertex.TangentX, 0.0f), DIContext.InstanceTransform).xyz);
 }
 
 void DISkelMesh_GetVertexColor(in FDISkelMeshContext DIContext, in int VertexIndex, out float4 OutColor)
