improve error checking for grids with resolution too large

diff --git a/Engine/Plugins/FX/Niagara/Source/Niagara/Private/NiagaraDataInterfaceGrid3DCollection.cpp b/Engine/Plugins/FX/Niagara/Source/Niagara/Private/NiagaraDataInterfaceGrid3DCollection.cpp
--- a/Engine/Plugins/FX/Niagara/Source/Niagara/Private/NiagaraDataInterfaceGrid3DCollection.cpp
+++ b/Engine/Plugins/FX/Niagara/Source/Niagara/Private/NiagaraDataInterfaceGrid3DCollection.cpp
@@ -1929,6 +1929,13 @@ bool UNiagaraDataInterfaceGrid3DCollection::InitPerInstanceData(void* PerInstanc
 	int32 MaxTilesY = floor(MaxDim / InstanceData->NumCells.Y);
 	int32 MaxTilesZ = floor(MaxDim / InstanceData->NumCells.Z);
 	int32 MaxAttributes = MaxTilesX * MaxTilesY * MaxTilesZ;
+
+	if (InstanceData->NumCells.X > MaxDim || InstanceData->NumCells.Y > MaxDim || InstanceData->NumCells.Z > MaxDim)
+	{
+		UE_LOG(LogNiagara, Error, TEXT("Resolution is too high on %s... max is 2048, defined as %i, %i %i"), *FNiagaraUtilities::SystemInstanceIDToString(SystemInstance->GetId()), InstanceData->NumCells.X, InstanceData->NumCells.Y, InstanceData->NumCells.Z);
+		return false;
+	}
+
 	if (NumAttribChannelsFound > MaxAttributes && MaxAttributes > 0)
 	{
 		UE_LOG(LogNiagara, Error, TEXT("Invalid number of attributes defined on %s... max is %i, num defined is %i"), *FNiagaraUtilities::SystemInstanceIDToString(SystemInstance->GetId()), MaxAttributes, NumAttribChannelsFound);
@@ -2477,7 +2484,7 @@ bool UNiagaraDataInterfaceGrid3DCollection::PerInstanceTickPostSimulate(void* Pe
 	FGrid3DCollectionRWInstanceData_GameThread* InstanceData = static_cast<FGrid3DCollectionRWInstanceData_GameThread*>(PerInstanceData);
 	bool bNeedsReset = false;
 
-	if (InstanceData->NeedsRealloc && InstanceData->NumCells.X > 0 && InstanceData->NumCells.Y > 0 && InstanceData->NumCells.Z > 0)
+	if (InstanceData && InstanceData->NeedsRealloc && InstanceData->NumCells.X > 0 && InstanceData->NumCells.Y > 0 && InstanceData->NumCells.Z > 0)
 	{
 		InstanceData->NeedsRealloc = false;
 
@@ -2492,7 +2499,14 @@ bool UNiagaraDataInterfaceGrid3DCollection::PerInstanceTickPostSimulate(void* Pe
 		int32 MaxTilesY = floor(MaxDim / InstanceData->NumCells.Y);
 		int32 MaxTilesZ = floor(MaxDim / InstanceData->NumCells.Z);
 		int32 MaxAttributes = MaxTilesX * MaxTilesY * MaxTilesZ;
-		if ((InstanceData->TotalNumAttributes > MaxAttributes && MaxAttributes > 0) || InstanceData->TotalNumAttributes == 0)
+		
+		if (InstanceData->NumCells.X > MaxDim || InstanceData->NumCells.Y > MaxDim || InstanceData->NumCells.Z > MaxDim)
+		{
+			UE_LOG(LogNiagara, Error, TEXT("Resolution is too high on %s... max is 2048, defined as %i, %i %i"), *FNiagaraUtilities::SystemInstanceIDToString(SystemInstance->GetId()), InstanceData->NumCells.X, InstanceData->NumCells.Y, InstanceData->NumCells.Z);
+			return false;
+		}
+		
+		if ((InstanceData->TotalNumAttributes > MaxAttributes && MaxAttributes > 0) || InstanceData->TotalNumAttributes <= 0)
 		{
 			UE_LOG(LogNiagara, Error, TEXT("Invalid number of attributes defined on %s... max is %i, num defined is %i"), *FNiagaraUtilities::SystemInstanceIDToString(SystemInstance->GetId()), MaxAttributes, InstanceData->TotalNumAttributes);
 			return false;
