#UE4 Prevent checkslow failure when MainDataSet.HasVariable() is false in FNiagaraSystemSimulation::Init

diff --git a/Engine/Plugins/FX/Niagara/Source/Niagara/Private/NiagaraSystemSimulation.cpp b/Engine/Plugins/FX/Niagara/Source/Niagara/Private/NiagaraSystemSimulation.cpp
--- a/Engine/Plugins/FX/Niagara/Source/Niagara/Private/NiagaraSystemSimulation.cpp
+++ b/Engine/Plugins/FX/Niagara/Source/Niagara/Private/NiagaraSystemSimulation.cpp
@@ -374,24 +374,31 @@ bool FNiagaraSystemSimulation::Init(UNiagaraSystem* InSystem, UWorld* InWorld, b
 			{
 				FNiagaraEmitterHandle& EmitterHandle = System->GetEmitterHandle(EmitterIdx);
 				UNiagaraEmitter* Emitter = EmitterHandle.GetInstance();
+				bool bAddedAccessor = false;
 				if (Emitter)
 				{
 					FString EmitterName = Emitter->GetUniqueEmitterName();
-					EmitterExecutionStateAccessors.Emplace(MainDataSet, FNiagaraVariable(EnumPtr, *(EmitterName + TEXT(".ExecutionState"))));
-					const TArray<TSharedRef<const FNiagaraEmitterCompiledData>>& EmitterCompiledData = System->GetEmitterCompiledData();
-
-					check(EmitterCompiledData.Num() == System->GetNumEmitters());
-					for (FName AttrName : EmitterCompiledData[EmitterIdx]->SpawnAttributes)
+					FNiagaraVariable ExecutionStateVar(EnumPtr, *(EmitterName + TEXT(".ExecutionState")));
+					if (MainDataSet.HasVariable(ExecutionStateVar))
 					{
-						EmitterSpawnInfoAccessors[EmitterIdx].Emplace(MainDataSet, FNiagaraVariable(FNiagaraTypeDefinition(FNiagaraSpawnInfo::StaticStruct()), AttrName));
-					}
-
-					if (Emitter->bLimitDeltaTime)
-					{
-						MaxDeltaTime = MaxDeltaTime.IsSet() ? FMath::Min(MaxDeltaTime.GetValue(), Emitter->MaxDeltaTimePerTick) : Emitter->MaxDeltaTimePerTick;
+						EmitterExecutionStateAccessors.Emplace(MainDataSet, ExecutionStateVar);
+						bAddedAccessor = true;
+						const TArray<TSharedRef<const FNiagaraEmitterCompiledData>>& EmitterCompiledData = System->GetEmitterCompiledData();
+
+						check(EmitterCompiledData.Num() == System->GetNumEmitters());
+						for (FName AttrName : EmitterCompiledData[EmitterIdx]->SpawnAttributes)
+						{
+							EmitterSpawnInfoAccessors[EmitterIdx].Emplace(MainDataSet, FNiagaraVariable(FNiagaraTypeDefinition(FNiagaraSpawnInfo::StaticStruct()), AttrName));
+						}
+
+						if (Emitter->bLimitDeltaTime)
+						{
+							MaxDeltaTime = MaxDeltaTime.IsSet() ? FMath::Min(MaxDeltaTime.GetValue(), Emitter->MaxDeltaTimePerTick) : Emitter->MaxDeltaTimePerTick;
+						}
 					}
 				}
-				else
+				
+				if (!bAddedAccessor)
 				{
 					EmitterExecutionStateAccessors.Emplace();
 				}
