#UE4 Handling null mesh in UNiagaraDataInterfaceSkeletalMesh::GetSkinnedBoneData

diff --git a/Engine/Plugins/FX/Niagara/Source/Niagara/Private/NDISkeletalMesh_BoneSampling.cpp b/Engine/Plugins/FX/Niagara/Source/Niagara/Private/NDISkeletalMesh_BoneSampling.cpp
--- a/Engine/Plugins/FX/Niagara/Source/Niagara/Private/NDISkeletalMesh_BoneSampling.cpp
+++ b/Engine/Plugins/FX/Niagara/Source/Niagara/Private/NDISkeletalMesh_BoneSampling.cpp
@@ -776,7 +776,6 @@ void UNiagaraDataInterfaceSkeletalMesh::GetSkinnedBoneData(FVectorVMContext& Con
 	}
 
 	checkfSlow(InstData.Get(), TEXT("Skeletal Mesh Interface has invalid instance data. %s"), *GetPathName());
-	checkfSlow(InstData->Mesh, TEXT("Skeletal Mesh Interface has invalid mesh. %s"), *GetPathName());
 
 	FBoneSocketSkinnedDataOutputHandler Output(Context);
 
@@ -787,115 +786,157 @@ void UNiagaraDataInterfaceSkeletalMesh::GetSkinnedBoneData(FVectorVMContext& Con
 	const FQuat InstanceRotation = Output.bNeedsRotation ? InstanceTransform.GetMatrixWithoutScale().ToQuat() : FQuat::Identity;
 	const FQuat PrevInstanceRotation = Output.bNeedsRotation ? PrevInstanceTransform.GetMatrixWithoutScale().ToQuat() : FQuat::Identity;
 
-	FSkeletalMeshAccessorHelper Accessor;
-	Accessor.Init<TNDISkelMesh_FilterModeNone, TNDISkelMesh_AreaWeightingOff>(InstData);
-	check(Accessor.AreBonesAccessible());
-
-	const int32 BoneCount = SkinningHandler.GetBoneCount(Accessor, bInterpolated::Value);
-	const int32 BoneAndSocketCount = BoneCount + InstData->FilteredSocketInfo.Num();
-	float InvDt = 1.0f / InstData->DeltaSeconds;
-
-	const TArray<FTransform>& FilteredSocketCurrTransforms = InstData->GetFilteredSocketsCurrBuffer();
-	const TArray<FTransform>& FilteredSocketPrevTransforms = InstData->GetFilteredSocketsPrevBuffer();
-
-	for (int32 i = 0; i < Context.NumInstances; ++i)
+	if (InstData->Mesh)
 	{
-		const float Interp = bInterpolated::Value ? InterpParam.GetAndAdvance() : 1.0f;
+		FSkeletalMeshAccessorHelper Accessor;
+		Accessor.Init<TNDISkelMesh_FilterModeNone, TNDISkelMesh_AreaWeightingOff>(InstData);
+		check(Accessor.AreBonesAccessible());
 
-		// Determine bone or socket
-		const int32 Bone = BoneParam.GetAndAdvance();
-		const bool bIsSocket = Bone >= BoneCount;
-		const int32 Socket = Bone - BoneCount;
+		const int32 BoneCount = SkinningHandler.GetBoneCount(Accessor, bInterpolated::Value);
+		const int32 BoneAndSocketCount = BoneCount + InstData->FilteredSocketInfo.Num();
+		float InvDt = 1.0f / InstData->DeltaSeconds;
 
-		FVector Pos;
-		FVector Prev;
+		const TArray<FTransform>& FilteredSocketCurrTransforms = InstData->GetFilteredSocketsCurrBuffer();
+		const TArray<FTransform>& FilteredSocketPrevTransforms = InstData->GetFilteredSocketsPrevBuffer();
 
-		// Handle invalid bone indices first
-		if (Bone < 0 || Bone >= BoneAndSocketCount)
+		for (int32 i = 0; i < Context.NumInstances; ++i)
 		{
-			Pos = FVector::ZeroVector;
-			TransformHandler.TransformPosition(Pos, InstanceTransform);
+			const float Interp = bInterpolated::Value ? InterpParam.GetAndAdvance() : 1.0f;
 
-			if (Output.bNeedsVelocity || bInterpolated::Value)
+			// Determine bone or socket
+			const int32 Bone = BoneParam.GetAndAdvance();
+			const bool bIsSocket = Bone >= BoneCount;
+			const int32 Socket = Bone - BoneCount;
+
+			FVector Pos;
+			FVector Prev;
+
+			// Handle invalid bone indices first
+			if (Bone < 0 || Bone >= BoneAndSocketCount)
 			{
-				Prev = FVector::ZeroVector;
-				TransformHandler.TransformPosition(Prev, PrevInstanceTransform);
+				Pos = FVector::ZeroVector;
+				TransformHandler.TransformPosition(Pos, InstanceTransform);
+
+				if (Output.bNeedsVelocity || bInterpolated::Value)
+				{
+					Prev = FVector::ZeroVector;
+					TransformHandler.TransformPosition(Prev, PrevInstanceTransform);
+				}
+				if (Output.bNeedsRotation)
+				{
+					Output.SetRotation(FQuat::Identity);
+				}
 			}
-			if (Output.bNeedsRotation)
+			else if (bIsSocket)
 			{
-				Output.SetRotation(FQuat::Identity);
+				FTransform CurrSocketTransform = FilteredSocketCurrTransforms[Socket];
+				FTransform PrevSocketTransform = FilteredSocketPrevTransforms[Socket];
+
+				Pos = CurrSocketTransform.GetLocation();
+				TransformHandler.TransformPosition(Pos, InstanceTransform);
+
+				if (Output.bNeedsVelocity || bInterpolated::Value)
+				{
+					Prev = PrevSocketTransform.GetLocation();
+					TransformHandler.TransformPosition(Prev, PrevInstanceTransform);
+				}
+
+				if (Output.bNeedsRotation)
+				{
+					FQuat Rotation = CurrSocketTransform.GetRotation();
+					TransformHandler.TransformRotation(Rotation, InstanceRotation);
+					if (bInterpolated::Value)
+					{
+						FQuat PrevRotation = PrevSocketTransform.GetRotation();
+						TransformHandler.TransformRotation(PrevRotation, PrevInstanceRotation);
+						Rotation = FQuat::Slerp(PrevRotation, Rotation, Interp);
+					}
+
+					Output.SetRotation(Rotation);
+				}
 			}
-		}
-		else if (bIsSocket)
-		{
-			FTransform CurrSocketTransform = FilteredSocketCurrTransforms[Socket];
-			FTransform PrevSocketTransform = FilteredSocketPrevTransforms[Socket];
+			// Bone
+			else
+			{
+				Pos = SkinningHandler.GetSkinnedBonePosition(Accessor, Bone);
+				TransformHandler.TransformPosition(Pos, InstanceTransform);
 
-			Pos = CurrSocketTransform.GetLocation();
-			TransformHandler.TransformPosition(Pos, InstanceTransform);
+				if (Output.bNeedsVelocity || bInterpolated::Value)
+				{
+					Prev = SkinningHandler.GetSkinnedBonePreviousPosition(Accessor, Bone);
+					TransformHandler.TransformPosition(Prev, PrevInstanceTransform);
+				}
+
+				if (Output.bNeedsRotation)
+				{
+					FQuat Rotation = SkinningHandler.GetSkinnedBoneRotation(Accessor, Bone);
+					TransformHandler.TransformRotation(Rotation, InstanceRotation);
+					if (bInterpolated::Value)
+					{
+						FQuat PrevRotation = SkinningHandler.GetSkinnedBonePreviousRotation(Accessor, Bone);
+						TransformHandler.TransformRotation(PrevRotation, PrevInstanceRotation);
+						Rotation = FQuat::Slerp(PrevRotation, Rotation, Interp);
+					}
+
+					Output.SetRotation(Rotation);
+				}
+			}
 
 			if (Output.bNeedsVelocity || bInterpolated::Value)
 			{
-				Prev = PrevSocketTransform.GetLocation();
-				TransformHandler.TransformPosition(Prev, PrevInstanceTransform);
+				Pos = FMath::Lerp(Prev, Pos, Interp);
 			}
 
-			if (Output.bNeedsRotation)
+			if (Output.bNeedsPosition)
 			{
-				FQuat Rotation = CurrSocketTransform.GetRotation();
-				TransformHandler.TransformRotation(Rotation, InstanceRotation);
-				if (bInterpolated::Value)
-				{
-					FQuat PrevRotation = PrevSocketTransform.GetRotation();
-					TransformHandler.TransformRotation(PrevRotation, PrevInstanceRotation);
-					Rotation = FQuat::Slerp(PrevRotation, Rotation, Interp);
-				}
+				Output.SetPosition(Pos);
+			}
 
-				Output.SetRotation(Rotation);
+			if(Output.bNeedsVelocity)
+			{
+				//Don't have enough information to get a better interpolated velocity.
+				FVector Velocity = (Pos - Prev) * InvDt;
+				Output.SetVelocity(Velocity);
 			}
 		}
-		// Bone
-		else
+	}
+	else
+	{
+		const float InvDt = 1.0f / InstData->DeltaSeconds;
+
+		for (int32 i = 0; i < Context.NumInstances; ++i)
 		{
-			Pos = SkinningHandler.GetSkinnedBonePosition(Accessor, Bone);
+			const float Interp = bInterpolated::Value ? InterpParam.GetAndAdvance() : 1.0f;
+
+			FVector Prev = FVector::ZeroVector;
+			FVector Pos = FVector::ZeroVector;
 			TransformHandler.TransformPosition(Pos, InstanceTransform);
 
 			if (Output.bNeedsVelocity || bInterpolated::Value)
 			{
-				Prev = SkinningHandler.GetSkinnedBonePreviousPosition(Accessor, Bone);
 				TransformHandler.TransformPosition(Prev, PrevInstanceTransform);
 			}
 
 			if (Output.bNeedsRotation)
 			{
-				FQuat Rotation = SkinningHandler.GetSkinnedBoneRotation(Accessor, Bone);
-				TransformHandler.TransformRotation(Rotation, InstanceRotation);
-				if (bInterpolated::Value)
-				{
-					FQuat PrevRotation = SkinningHandler.GetSkinnedBonePreviousRotation(Accessor, Bone);
-					TransformHandler.TransformRotation(PrevRotation, PrevInstanceRotation);
-					Rotation = FQuat::Slerp(PrevRotation, Rotation, Interp);
-				}
-
-				Output.SetRotation(Rotation);
+				Output.SetRotation(FQuat::Identity);
 			}
-		}
 
-		if (Output.bNeedsVelocity || bInterpolated::Value)
-		{
-			Pos = FMath::Lerp(Prev, Pos, Interp);
-		}
+			if (Output.bNeedsVelocity || bInterpolated::Value)
+			{
+				Pos = FMath::Lerp(Prev, Pos, Interp);
+			}
 
-		if (Output.bNeedsPosition)
-		{
-			Output.SetPosition(Pos);
-		}
+			if (Output.bNeedsPosition)
+			{
+				Output.SetPosition(Pos);
+			}
 
-		if(Output.bNeedsVelocity)
-		{
-			//Don't have enough information to get a better interpolated velocity.
-			FVector Velocity = (Pos - Prev) * InvDt;
-			Output.SetVelocity(Velocity);
+			if (Output.bNeedsVelocity)
+			{
+				FVector Velocity = (Pos - Prev) * InvDt;
+				Output.SetVelocity(Velocity);
+			}
 		}
 	}
 }
