- Fix for NaN/uninitiatlized value in Niagara GPU sort info

diff --git a/Engine/Plugins/FX/Niagara/Source/Niagara/Classes/NiagaraGPUSortInfo.h b/Engine/Plugins/FX/Niagara/Source/Niagara/Classes/NiagaraGPUSortInfo.h
--- a/Engine/Plugins/FX/Niagara/Source/Niagara/Classes/NiagaraGPUSortInfo.h
+++ b/Engine/Plugins/FX/Niagara/Source/Niagara/Classes/NiagaraGPUSortInfo.h
@@ -57,7 +57,7 @@ struct FNiagaraGPUSortInfo
 	int32 CullScaleAttributeOffset = INDEX_NONE;
 	int32 RendererVisTagAttributeOffset = INDEX_NONE;
 	int32 RendererVisibility = 0;
-	FSphere LocalBSphere;
+	FSphere LocalBSphere = FSphere(0);
 	FVector2D DistanceCullRange { 0.0f, FLT_MAX };
 	TArray<FPlane, TFixedAllocator<MaxCullPlanes>> CullPlanes;
 
