Buffer allocation and dispatch rework for Niagara HWRT collisions.

All instances now accumulate their rays per DI with a single dispatch per DI rather than a dispatch per system instance.
MaxRayTraces can now define traces per particle.
The batcher now has shared scratch buffers that are used for all ray requests, results and counts.
Each DI will allocate a section of these scratch buffers to use. Offsets into these buffers are now passed into all shaders.

- New utility classes for GPU scratch buffers.
These allocate out of fixed buffers with allocations being reset each frame.
Float and UInt RWBuffers as well as typed structured buffers available.
These will probably be replaced when we RDGify Niagara fully.

- Misc changes to VFXTracerRay API requested by Yuriy.
- CVar to enable/disable HWRT collisions now togglable at runtime. Requires reinitializing all FX.

diff --git a/Engine/Plugins/FX/Niagara/Shaders/Private/NiagaraDataInterfaceCollisionQuery.ush b/Engine/Plugins/FX/Niagara/Shaders/Private/NiagaraDataInterfaceCollisionQuery.ush
--- a/Engine/Plugins/FX/Niagara/Shaders/Private/NiagaraDataInterfaceCollisionQuery.ush
+++ b/Engine/Plugins/FX/Niagara/Shaders/Private/NiagaraDataInterfaceCollisionQuery.ush
@@ -151,7 +151,7 @@ uint BuildRayTraceMask(int TraceChannel)
 	return 0xFF;
 }
 
-void NDICollisionQuery_IssueAsyncRayTrace(uint MaxRayTraceCount, RWStructuredBuffer<FNiagaraRayData> RWRayRequests, RWBuffer<uint> RWRayTraceCounts, int In_QueryID, float3 In_TraceStart, float3 In_TraceEnd, int In_TraceChannel, out bool Out_IsQueryValid)
+void NDICollisionQuery_IssueAsyncRayTrace(uint MaxRayTraceCount, RWStructuredBuffer<FNiagaraRayData> RWRayRequests, uint RayRequestOffset, RWBuffer<uint> RWRayTraceCounts, uint RayTraceCountsOffset, int In_QueryID, float3 In_TraceStart, float3 In_TraceEnd, int In_TraceChannel, out bool Out_IsQueryValid)
 {
 	Out_IsQueryValid = false;
 
@@ -160,29 +160,31 @@ void NDICollisionQuery_IssueAsyncRayTrace(uint MaxRayTraceCount, RWStructuredBuf
 		float3 Trace = In_TraceEnd - In_TraceStart;
 		float TraceLength = length(Trace);
 
-		RWRayRequests[In_QueryID].Origin = In_TraceStart;
-		RWRayRequests[In_QueryID].Mask = BuildRayTraceMask(In_TraceChannel);
-		RWRayRequests[In_QueryID].Direction = normalize(Trace);
-		RWRayRequests[In_QueryID].TFar = TraceLength;
+		//Offset query into the actual ray buffer.
+		uint RayIndex = RayRequestOffset + In_QueryID;
+		RWRayRequests[RayIndex].Origin = In_TraceStart;
+		RWRayRequests[RayIndex].Mask = BuildRayTraceMask(In_TraceChannel);
+		RWRayRequests[RayIndex].Direction = normalize(Trace);
+		RWRayRequests[RayIndex].TFar = TraceLength;
 		Out_IsQueryValid = true;
 
 		// we want to keep track of the maximum index used
-		InterlockedMax(RWRayTraceCounts[0], In_QueryID + 1);
-		RWRayTraceCounts[1] = 1;
-		RWRayTraceCounts[2] = 1;
+		InterlockedMax(RWRayTraceCounts[RayTraceCountsOffset + 0], In_QueryID + 1);
+		RWRayTraceCounts[RayTraceCountsOffset + 1] = 1;
+		RWRayTraceCounts[RayTraceCountsOffset + 2] = 1;
 	}
 }
 
-void NDICollisionQuery_ReserveRayTraceIndex(uint MaxRayTraceCount, RWBuffer<uint> RWRayTraceCounts, int In_TraceCount, out int Out_StartQueryID, out bool Out_IndicesValid)
+void NDICollisionQuery_ReserveRayTraceIndex(uint MaxRayTraceCount, RWBuffer<uint> RWRayTraceCounts, uint RayTraceCountsOffset, int In_TraceCount, out int Out_StartQueryID, out bool Out_IndicesValid)
 {
 	Out_IndicesValid = false;
 
-	if ((In_TraceCount + RWRayTraceCounts[0]) < MaxRayTraceCount)
+	if ((In_TraceCount + RWRayTraceCounts[RayTraceCountsOffset + 0]) < MaxRayTraceCount)
 	{
 		uint RequestIndex = 0;
-		InterlockedAdd(RWRayTraceCounts[0], In_TraceCount, RequestIndex);
-		RWRayTraceCounts[1] = 1;
-		RWRayTraceCounts[2] = 1;
+		InterlockedAdd(RWRayTraceCounts[RayTraceCountsOffset + 0], In_TraceCount, RequestIndex);
+		RWRayTraceCounts[RayTraceCountsOffset + 1] = 1;
+		RWRayTraceCounts[RayTraceCountsOffset + 2] = 1;
 
 		if ((In_TraceCount + RequestIndex) <= MaxRayTraceCount)
 		{
@@ -191,33 +193,35 @@ void NDICollisionQuery_ReserveRayTraceIndex(uint MaxRayTraceCount, RWBuffer<uint
 		}
 		else
 		{
-			InterlockedMax(RWRayTraceCounts[0], MaxRayTraceCount);
+			InterlockedMax(RWRayTraceCounts[RayTraceCountsOffset + 0], MaxRayTraceCount);
 		}
 	}
 }
 
-void NDICollisionQuery_CreateAsyncRayTrace(uint MaxRayTraceCount, RWStructuredBuffer<FNiagaraRayData> RWRayRequests, RWBuffer<uint> RWRayTraceCounts, float3 In_TraceStart, float3 In_TraceEnd, int In_TraceChannel, out int Out_QueryID, out bool Out_IsQueryValid)
+void NDICollisionQuery_CreateAsyncRayTrace(	uint MaxRayTraceCount, RWStructuredBuffer<FNiagaraRayData> RWRayRequests, uint RayRequestsOffset, RWBuffer<uint> RWRayTraceCounts, uint RayTraceCountsOffset,
+											float3 In_TraceStart, float3 In_TraceEnd, int In_TraceChannel, out int Out_QueryID, out bool Out_IsQueryValid)
 {
 	Out_IsQueryValid = false;
 
-	NDICollisionQuery_ReserveRayTraceIndex(MaxRayTraceCount, RWRayTraceCounts, 1, Out_QueryID, Out_IsQueryValid);
+	NDICollisionQuery_ReserveRayTraceIndex(MaxRayTraceCount, RWRayTraceCounts, RayTraceCountsOffset, 1, Out_QueryID, Out_IsQueryValid);
 
 	if (Out_IsQueryValid)
 	{
 		float3 Trace = In_TraceEnd - In_TraceStart;
 		float TraceLength = length(Trace);
 
-		RWRayRequests[Out_QueryID].Origin = In_TraceStart;
-		RWRayRequests[Out_QueryID].Mask = BuildRayTraceMask(In_TraceChannel);
-		RWRayRequests[Out_QueryID].Direction = normalize(Trace);
-		RWRayRequests[Out_QueryID].TFar = TraceLength;
+		uint RayIndex = Out_QueryID + RayRequestsOffset;
+		RWRayRequests[RayIndex].Origin = In_TraceStart;
+		RWRayRequests[RayIndex].Mask = BuildRayTraceMask(In_TraceChannel);
+		RWRayRequests[RayIndex].Direction = normalize(Trace);
+		RWRayRequests[RayIndex].TFar = TraceLength;
 
 		//TODO: We can pack collision group and trace channels into the same int and test both in the RG Shader.
-		RWRayRequests[Out_QueryID].CollisionGroup = In_TraceChannel;
+		RWRayRequests[RayIndex].CollisionGroup = In_TraceChannel;
 	}
 }
 
-void NDICollisionQuery_ReadAsyncRayTrace(uint MaxRayTraceCount, StructuredBuffer<FNiagaraRayTracingResult> IntersectionResults, int In_PreviousFrameQueryID, out bool Out_CollisionValid, out float Out_CollisionDistance, out float3 Out_CollisionPosWorld, out float3 Out_CollisionNormal)
+void NDICollisionQuery_ReadAsyncRayTrace(uint MaxRayTraceCount, StructuredBuffer<FNiagaraRayTracingResult> IntersectionResults, uint IntersectionRestultsOffest, int In_PreviousFrameQueryID, out bool Out_CollisionValid, out float Out_CollisionDistance, out float3 Out_CollisionPosWorld, out float3 Out_CollisionNormal)
 {
 	Out_CollisionValid = false;
 	Out_CollisionDistance = 0.0f;
@@ -226,7 +230,7 @@ void NDICollisionQuery_ReadAsyncRayTrace(uint MaxRayTraceCount, StructuredBuffer
 
 	if (In_PreviousFrameQueryID >= 0 && uint(In_PreviousFrameQueryID) < MaxRayTraceCount)
 	{
-		FNiagaraRayTracingResult HitResult = IntersectionResults[In_PreviousFrameQueryID];
+		FNiagaraRayTracingResult HitResult = IntersectionResults[In_PreviousFrameQueryID + IntersectionRestultsOffest];
 
 		Out_CollisionValid = HitResult.IsHit();
 		if (Out_CollisionValid)
diff --git a/Engine/Plugins/FX/Niagara/Shaders/Private/NiagaraDataInterfaceCollisionQueryTemplate.ush b/Engine/Plugins/FX/Niagara/Shaders/Private/NiagaraDataInterfaceCollisionQueryTemplate.ush
--- a/Engine/Plugins/FX/Niagara/Shaders/Private/NiagaraDataInterfaceCollisionQueryTemplate.ush
+++ b/Engine/Plugins/FX/Niagara/Shaders/Private/NiagaraDataInterfaceCollisionQueryTemplate.ush
@@ -1,11 +1,15 @@
 // Copyright Epic Games, Inc. All Rights Reserved.
 
 #if NIAGARA_USE_RAYTRACING
-bool											RayTracingEnabled_{ParameterName};
-uint                                            MaxRayTraceCount_{ParameterName};
-RWStructuredBuffer<FNiagaraRayData>               RWRayRequests_{ParameterName};
-StructuredBuffer<FNiagaraRayTracingResult>     IntersectionResults_{ParameterName};
-RWBuffer<uint>                                  RWRayTraceCounts_{ParameterName};
+/** Niagara ray tracing uses shared buffers between dispathces so each buffer has a buffer + and offset into that buffer for the current dispatch. */
+bool										RayTracingEnabled_{ParameterName};
+uint										MaxRayTraceCount_{ParameterName};
+RWStructuredBuffer<FNiagaraRayData>			RWRayRequests_{ParameterName};
+uint										RayRequestsOffset_{ ParameterName };
+StructuredBuffer<FNiagaraRayTracingResult>	IntersectionResults_{ParameterName};
+uint										IntersectionResultsOffset_{ ParameterName };
+RWBuffer<uint>								RWRayTraceCounts_{ParameterName};
+uint										RayTraceCountsOffset_{ ParameterName };
 #endif
 
 void QuerySceneDepthGPU_{ParameterName}(in float3 In_SamplePos, out float Out_SceneDepth, out float3 Out_CameraPosWorld, out bool Out_IsInsideView, out float3 Out_WorldPos, out float3 Out_WorldNormal)
@@ -28,7 +32,7 @@ void IssueAsyncRayTraceGpu_{ParameterName}_UEImpureCall(int In_QueryID, float3 I
 #if NIAGARA_USE_RAYTRACING
 	if (RayTracingEnabled_{ParameterName})
 	{
-		NDICollisionQuery_IssueAsyncRayTrace(MaxRayTraceCount_{ParameterName}, RWRayRequests_{ParameterName}, RWRayTraceCounts_{ParameterName}, In_QueryID, In_TraceStart, In_TraceEnd, In_TraceChannel, Out_IsQueryValid);
+		NDICollisionQuery_IssueAsyncRayTrace(MaxRayTraceCount_{ParameterName}, RWRayRequests_{ParameterName}, RayRequestsOffset_{ ParameterName }, RWRayTraceCounts_{ParameterName}, RayTraceCountsOffset_{ ParameterName }, In_QueryID, In_TraceStart, In_TraceEnd, In_TraceChannel, Out_IsQueryValid);
 		return;
 	}
 #endif
@@ -40,7 +44,7 @@ void ReserveRayTraceIndex_{ParameterName}_UEImpureCall(int In_TraceCount, out in
 #if NIAGARA_USE_RAYTRACING
 	if (RayTracingEnabled_{ParameterName})
 	{
-		NDICollisionQuery_ReserveRayTraceIndex(MaxRayTraceCount_{ParameterName}, RWRayTraceCounts_{ParameterName}, In_TraceCount, Out_StartQueryID, Out_IndicesValid);
+		NDICollisionQuery_ReserveRayTraceIndex(MaxRayTraceCount_{ParameterName}, RWRayTraceCounts_{ParameterName}, RayTraceCountsOffset_{ ParameterName }, In_TraceCount, Out_StartQueryID, Out_IndicesValid);
 		return;
 	}
 #endif
@@ -52,7 +56,7 @@ void CreateAsyncRayTraceGpu_{ParameterName}_UEImpureCall(float3 In_TraceStart, f
 #if NIAGARA_USE_RAYTRACING
 	if (RayTracingEnabled_{ParameterName})
 	{
-		NDICollisionQuery_CreateAsyncRayTrace(MaxRayTraceCount_{ParameterName}, RWRayRequests_{ParameterName}, RWRayTraceCounts_{ParameterName}, In_TraceStart, In_TraceEnd, In_TraceChannel, Out_QueryID, Out_IsQueryValid);
+		NDICollisionQuery_CreateAsyncRayTrace(MaxRayTraceCount_{ParameterName}, RWRayRequests_{ParameterName}, RayRequestsOffset_{ ParameterName }, RWRayTraceCounts_{ParameterName}, RayTraceCountsOffset_{ ParameterName }, In_TraceStart, In_TraceEnd, In_TraceChannel, Out_QueryID, Out_IsQueryValid);
 		return;
 	}
 #endif
@@ -64,9 +68,9 @@ void ReadAsyncRayTraceGpu_{ParameterName}(int In_PreviousFrameQueryID, out bool
 #if NIAGARA_USE_RAYTRACING
 	if (RayTracingEnabled_{ParameterName})
 	{
-		NDICollisionQuery_ReadAsyncRayTrace(MaxRayTraceCount_{ParameterName}, IntersectionResults_{ParameterName}, In_PreviousFrameQueryID, Out_CollisionValid, Out_CollisionDistance, Out_CollisionPosWorld, Out_CollisionNormal);
+		NDICollisionQuery_ReadAsyncRayTrace(MaxRayTraceCount_{ParameterName}, IntersectionResults_{ParameterName}, IntersectionResultsOffset_{ParameterName}, In_PreviousFrameQueryID, Out_CollisionValid, Out_CollisionDistance, Out_CollisionPosWorld, Out_CollisionNormal);
 		return;
-	}	
+	}
 #endif
 	Out_CollisionValid = false;
 }
\ No newline at end of file
diff --git a/Engine/Plugins/FX/Niagara/Shaders/Private/NiagaraRayTracingShaders.usf b/Engine/Plugins/FX/Niagara/Shaders/Private/NiagaraRayTracingShaders.usf
--- a/Engine/Plugins/FX/Niagara/Shaders/Private/NiagaraRayTracingShaders.usf
+++ b/Engine/Plugins/FX/Niagara/Shaders/Private/NiagaraRayTracingShaders.usf
@@ -9,9 +9,11 @@
 #include "/Engine/Private/HashTable.ush"
 #include "/Engine/Private/RayTracing/VFXTraceRay.ush"
 
-RaytracingAccelerationStructure               TLAS;
-StructuredBuffer<FNiagaraRayData>               Rays;
-RWStructuredBuffer<FNiagaraRayTracingResult> CollisionOutput;
+RaytracingAccelerationStructure					TLAS;
+StructuredBuffer<FNiagaraRayData>				Rays;
+uint											RaysOffset;
+RWStructuredBuffer<FNiagaraRayTracingResult>	CollisionOutput;
+uint											CollisionOutputOffset;
 
 #if NIAGARA_RAYTRACE_FAKE_INDIRECT
 Buffer<uint>                                  RayTraceCounts;
@@ -55,13 +57,13 @@ RAY_TRACING_ENTRY_RAYGEN(NiagaraCollisionRayTraceRG)
 	const uint RayIndex = DispatchRaysIndex().x;
 
 #if NIAGARA_RAYTRACE_FAKE_INDIRECT
-	if (RayIndex >= RayTraceCounts[0])
+	if (RayIndex >= RayTraceCounts[0]) 
 	{
 		return;
 	}
 #endif
 
-	FNiagaraRayData InputRay = Rays[RayIndex];
+	FNiagaraRayData InputRay = Rays[RayIndex + RaysOffset];
 
 	RayDesc Ray;
 	Ray.Origin = InputRay.Origin;
@@ -81,7 +83,7 @@ RAY_TRACING_ENTRY_RAYGEN(NiagaraCollisionRayTraceRG)
 		while (NumCasts++ <= MaxRetraces)
 		{
 			Payload = (FVFXTracePayload)0;
-			VFX_TRACE_RAY(
+			VfxTraceRay(
 				TLAS,   // AccelerationStructure
 				RayFlags,
 				RAY_TRACING_MASK_OPAQUE,          // InstanceInclusionMask
@@ -90,10 +92,10 @@ RAY_TRACING_ENTRY_RAYGEN(NiagaraCollisionRayTraceRG)
 				0,      // MissShaderIndex
 				Ray,    // RayDesc
 				Payload // Payload
-			);
+			); 
 
 			uint GPUSceneInstanceIndex = GetGPUSceneInstanceIndex(Payload.PrimitiveIndex, Payload.InstanceIndex);
-			uint PrimCollisionGroup;
+			uint PrimCollisionGroup; 
 			if (FindCollisionGroup(GPUSceneInstanceIndex, PrimCollisionGroup))
 			{
 				//If we collide with something of a different group to the Ray then break and allow the collision. Otherwise loop and attempt to recast just infront of the hit we're discarding.
@@ -112,7 +114,7 @@ RAY_TRACING_ENTRY_RAYGEN(NiagaraCollisionRayTraceRG)
 	}
 	else
 	{
-		VFX_TRACE_RAY(
+		VfxTraceRay(
 			TLAS,   // AccelerationStructure
 			RayFlags,
 			RAY_TRACING_MASK_OPAQUE,          // InstanceInclusionMask
@@ -124,11 +126,12 @@ RAY_TRACING_ENTRY_RAYGEN(NiagaraCollisionRayTraceRG)
 		);
 	}
 
-	CollisionOutput[RayIndex].HitT = Payload.HitT;
+	uint OutputIndex = RayIndex + CollisionOutputOffset;
+	CollisionOutput[OutputIndex].HitT = Payload.HitT;
 	if (Payload.IsHit())
 	{
-		CollisionOutput[RayIndex].WorldPosition = Payload.WorldPosition;
-		CollisionOutput[RayIndex].WorldNormal = Payload.WorldNormal;
+		CollisionOutput[OutputIndex].WorldPosition = Payload.WorldPosition;
+		CollisionOutput[OutputIndex].WorldNormal = Payload.WorldNormal;
 	}
 }
 
diff --git a/Engine/Plugins/FX/Niagara/Source/Niagara/Classes/NiagaraDataInterfaceCollisionQuery.h b/Engine/Plugins/FX/Niagara/Source/Niagara/Classes/NiagaraDataInterfaceCollisionQuery.h
--- a/Engine/Plugins/FX/Niagara/Source/Niagara/Classes/NiagaraDataInterfaceCollisionQuery.h
+++ b/Engine/Plugins/FX/Niagara/Source/Niagara/Classes/NiagaraDataInterfaceCollisionQuery.h
@@ -30,7 +30,7 @@ public:
 	DECLARE_NIAGARA_DI_PARAMETER();
 
 	UPROPERTY(EditAnywhere, Category = "Ray Trace")
-	int32 MaxRayTraceCount = 0;
+	int32 MaxTracesPerParticle = 0;
 
 	/** If a collision is rejected, how many times do we attempt to retrace from that collision point forward to find a new, valid collision.*/
 	UPROPERTY(EditAnywhere, Category = "Ray Trace")
diff --git a/Engine/Plugins/FX/Niagara/Source/Niagara/Classes/NiagaraEmitterInstanceBatcher.h b/Engine/Plugins/FX/Niagara/Source/Niagara/Classes/NiagaraEmitterInstanceBatcher.h
--- a/Engine/Plugins/FX/Niagara/Source/Niagara/Classes/NiagaraEmitterInstanceBatcher.h
+++ b/Engine/Plugins/FX/Niagara/Source/Niagara/Classes/NiagaraEmitterInstanceBatcher.h
@@ -158,7 +158,7 @@ public:
 	FNiagaraGpuReadbackManager* GetGpuReadbackManager() const { return GpuReadbackManagerPtr.Get(); }
 
 #if RHI_RAYTRACING
-	void IssueRayTraces(FRHICommandList& RHICmdList, const FIntPoint& RayTraceCounts, uint32 MaxRetraces, FRHIShaderResourceView* RayTraceRequests, FRWBuffer* IndirectArgsBuffer, uint32 IndirectArgsOffset, FRHIUnorderedAccessView* RayTraceResults) const;
+	FNiagaraRayTracingHelper& GetRayTracingHelper()const;
 	bool HasRayTracingScene() const;
 
 	void SetPrimitiveRayTracingCollisionGroup(UPrimitiveComponent* PrimitiveComponent, uint32 Group);
diff --git a/Engine/Plugins/FX/Niagara/Source/Niagara/Private/NiagaraDataInterfaceCollisionQuery.cpp b/Engine/Plugins/FX/Niagara/Source/Niagara/Private/NiagaraDataInterfaceCollisionQuery.cpp
--- a/Engine/Plugins/FX/Niagara/Source/Niagara/Private/NiagaraDataInterfaceCollisionQuery.cpp
+++ b/Engine/Plugins/FX/Niagara/Source/Niagara/Private/NiagaraDataInterfaceCollisionQuery.cpp
@@ -18,11 +18,19 @@
 
 #define LOCTEXT_NAMESPACE "NiagaraDataInterfaceCollisionQuery"
 
+void OnHWRTCollisionsEnabledChanged(IConsoleVariable* CVar)
+{
+	//Force a reinit of everything just to be safe.
+	FNiagaraSystemUpdateContext Context;
+	Context.AddAll(true);
+}
+
 int32 GEnableGPUHWRTCollisions = 1;
 static FAutoConsoleVariableRef CVarEnableGPUHWRTCollisions(
 	TEXT("fx.Niagara.Collision.EnableGPURayTracedCollisions"),
 	GEnableGPUHWRTCollisions,
 	TEXT("If greater than zero, GPU hardware ray trace collisions are enabled."),
+	FConsoleVariableDelegate::CreateStatic(&OnHWRTCollisionsEnabledChanged),
 	ECVF_Default
 );
 
@@ -44,8 +52,11 @@ namespace NDICollisionQueryLocal
 	static const FString RayTracingEnabledParamName(TEXT("RayTracingEnabled_"));
 	static const FString MaxRayTraceCountParamName(TEXT("MaxRayTraceCount_"));
 	static const FString RayRequestsParamName(TEXT("RayRequests_"));
+	static const FString RayRequestsOffsetParamName(TEXT("RayRequestsOffset_"));
 	static const FString IntersectionResultsParamName(TEXT("IntersectionResults_"));
+	static const FString IntersectionResultsOffsetParamName(TEXT("IntersectionResultsOffset_"));
 	static const FString RayTraceCountsParamName(TEXT("RayTraceCounts_"));
+	static const FString RayTraceCountsOffsetParamName(TEXT("RayTraceCountsOffset_"));
 }
 
 FCriticalSection UNiagaraDataInterfaceCollisionQuery::CriticalSection;
@@ -66,27 +77,12 @@ struct FNiagaraCollisionDIFunctionVersion
 struct FNiagaraDataIntefaceProxyCollisionQuery : public FNiagaraDataInterfaceProxy
 {
 #if RHI_RAYTRACING
-	FRWBufferStructured RayTraceRequests;
-	FRWBufferStructured RayTraceIntersections;
-	FRWBuffer RayTraceCounts;
-
-	int32 MaxRayTraceCount = 0;
-
+	int32 MaxTracesPerParticle = 0;
 	uint32 MaxRetraces = 0;
 #endif
 
 	virtual ~FNiagaraDataIntefaceProxyCollisionQuery()
 	{
-#if RHI_RAYTRACING
-		DEC_MEMORY_STAT_BY(STAT_NiagaraGPUDataInterfaceMemory, RayTraceRequests.NumBytes);
-		RayTraceRequests.Release();
-
-		DEC_MEMORY_STAT_BY(STAT_NiagaraGPUDataInterfaceMemory, RayTraceIntersections.NumBytes);
-		RayTraceIntersections.Release();
-
-		DEC_MEMORY_STAT_BY(STAT_NiagaraGPUDataInterfaceMemory, RayTraceCounts.NumBytes);
-		RayTraceCounts.Release();
-#endif
 	}
 
 	virtual int32 PerInstanceDataPassedToRenderThreadSize() const override
@@ -99,112 +95,31 @@ struct FNiagaraDataIntefaceProxyCollisionQuery : public FNiagaraDataInterfacePro
 		FNiagaraDataInterfaceProxy::PreStage(RHICmdList, Context);
 
 #if RHI_RAYTRACING
-		// First stage needs to clear the indirect args
-		if (MaxRayTraceCount > 0 && Context.SimStageData->bFirstStage)
+
+		if (IsRayTracingEnabled() && GEnableGPUHWRTCollisions && MaxTracesPerParticle > 0)
 		{
-			RHICmdList.ClearUAVUint(RayTraceCounts.UAV, FUintVector4(0, 0, 0, 0));
-			RHICmdList.Transition(FRHITransitionInfo(RayTraceCounts.UAV, ERHIAccess::UAVCompute, ERHIAccess::UAVCompute));
+			//Accumulate the total ray requests for this DI for all dispatches in the stage.
+			int32 RayRequests = MaxTracesPerParticle * Context.SimStageData->DestinationNumInstances;
+			FNiagaraRayTracingHelper& RTHelper = Context.Batcher->GetRayTracingHelper();
+			RTHelper.AddToDispatch(this, RayRequests, MaxRetraces);
 		}
 #endif
 	}
 
 	virtual void PostStage(FRHICommandList& RHICmdList, const FNiagaraDataInterfaceStageArgs& Context) override
 	{
-#if RHI_RAYTRACING
-		// If we are not the last stage we need to transition because another stage may write
-		// If we do not have a ray tracing scene active we also need to transition as PostSimulate will not transition the buffers
-		if (MaxRayTraceCount > 0 && (!Context.SimStageData->bLastStage || !Context.Batcher->HasRayTracingScene()))
-		{
-			FRHITransitionInfo Transitions[] =
-			{
-				FRHITransitionInfo(RayTraceRequests.UAV, ERHIAccess::UAVCompute, ERHIAccess::UAVCompute),
-				FRHITransitionInfo(RayTraceCounts.UAV, ERHIAccess::UAVCompute, ERHIAccess::UAVCompute),
-			};
-			RHICmdList.Transition(Transitions);
-		}
-#endif
 	}
 
 	virtual void PostSimulate(FRHICommandList& RHICmdList, const FNiagaraDataInterfaceArgs& Context) override
 	{
 		FNiagaraDataInterfaceProxy::PostSimulate(RHICmdList, Context);
-
-#if RHI_RAYTRACING
-		if (MaxRayTraceCount > 0 && Context.Batcher->HasRayTracingScene())
-		{
-			{
-				FRHITransitionInfo PreTransitions[] =
-				{
-					FRHITransitionInfo(RayTraceRequests.UAV, ERHIAccess::UAVCompute, ERHIAccess::SRVCompute),
-					FRHITransitionInfo(RayTraceIntersections.UAV, ERHIAccess::SRVCompute, ERHIAccess::UAVCompute),
-					FRHITransitionInfo(RayTraceCounts.UAV, ERHIAccess::UAVCompute, ERHIAccess::IndirectArgs),
-				};
-				RHICmdList.Transition(PreTransitions);
-			}
-
-			Context.Batcher->IssueRayTraces(RHICmdList, FIntPoint(MaxRayTraceCount, 1), MaxRetraces, RayTraceRequests.SRV, &RayTraceCounts, 0, RayTraceIntersections.UAV);
-
-			{
-				FRHITransitionInfo PostTransitions[] =
-				{
-					FRHITransitionInfo(RayTraceRequests.UAV, ERHIAccess::SRVCompute, ERHIAccess::UAVCompute),
-					FRHITransitionInfo(RayTraceIntersections.UAV, ERHIAccess::UAVCompute, ERHIAccess::SRVCompute),
-					FRHITransitionInfo(RayTraceCounts.UAV, ERHIAccess::IndirectArgs, ERHIAccess::UAVCompute),
-				};
-				RHICmdList.Transition(PostTransitions);
-			}
-		}
-#endif
 	}
 
-	void RenderThreadInitialize(int32 InMaxRayTraceRequests, uint32 InMaxRetraces)
+	void RenderThreadInitialize(int32 InMaxTracesPerParticle, uint32 InMaxRetraces)
 	{
 #if RHI_RAYTRACING
-		MaxRayTraceCount = 0;
 		MaxRetraces = InMaxRetraces;
-
-		DEC_MEMORY_STAT_BY(STAT_NiagaraGPUDataInterfaceMemory, RayTraceRequests.NumBytes);
-		RayTraceRequests.Release();
-
-		DEC_MEMORY_STAT_BY(STAT_NiagaraGPUDataInterfaceMemory, RayTraceIntersections.NumBytes);
-		RayTraceIntersections.Release();
-
-		DEC_MEMORY_STAT_BY(STAT_NiagaraGPUDataInterfaceMemory, RayTraceCounts.NumBytes);
-		RayTraceCounts.Release();
-
-		if (IsRayTracingEnabled() && GEnableGPUHWRTCollisions && InMaxRayTraceRequests > 0)
-		{
-			MaxRayTraceCount = 16 * FMath::DivideAndRoundUp(InMaxRayTraceRequests, 16);
-
-			RayTraceRequests.Initialize(
-				TEXT("NiagaraRayTraceRequests"),
-				sizeof(FNiagaraRayData),
-				MaxRayTraceCount,
-				BUF_Static,
-				false /*bUseUavCounter*/,
-				false /*bAppendBuffer*/,
-				ERHIAccess::UAVCompute);
-			INC_MEMORY_STAT_BY(STAT_NiagaraGPUDataInterfaceMemory, RayTraceRequests.NumBytes);
-
-			RayTraceIntersections.Initialize(
-				TEXT("NiagaraRayTraceIntersections"),
-				sizeof(FNiagaraRayTracingResult),
-				MaxRayTraceCount,
-				BUF_Static,
-				false /*bUseUavCounter*/,
-				false /*bAppendBuffer*/,
-				ERHIAccess::SRVCompute);
-			INC_MEMORY_STAT_BY(STAT_NiagaraGPUDataInterfaceMemory, RayTraceIntersections.NumBytes);
-
-			RayTraceCounts.Initialize(
-				TEXT("NiagaraRayTraceIndirectArgs"),
-				sizeof(uint32),
-				4,
-				PF_R32_UINT,
-				ERHIAccess::UAVCompute,
-				BUF_Static | BUF_DrawIndirect);
-			INC_MEMORY_STAT_BY(STAT_NiagaraGPUDataInterfaceMemory, RayTraceCounts.NumBytes);
-		}
+		MaxTracesPerParticle = InMaxTracesPerParticle;
 #endif
 	}
 };
@@ -253,7 +168,7 @@ void UNiagaraDataInterfaceCollisionQuery::PostLoad()
 {
 	Super::PostLoad();
 
-	if (MaxRayTraceCount)
+	if (MaxTracesPerParticle)
 	{
 		MarkRenderDataDirty();
 	}
@@ -654,7 +569,7 @@ void UNiagaraDataInterfaceCollisionQuery::ValidateFunction(const FNiagaraFunctio
 
 bool UNiagaraDataInterfaceCollisionQuery::RequiresRayTracingScene() const
 {
-	return IsRayTracingEnabled() && GEnableGPUHWRTCollisions && MaxRayTraceCount > 0;
+	return IsRayTracingEnabled() && GEnableGPUHWRTCollisions && MaxTracesPerParticle > 0;
 }
 
 #if WITH_EDITORONLY_DATA
@@ -966,7 +881,7 @@ bool UNiagaraDataInterfaceCollisionQuery::Equals(const UNiagaraDataInterface* Ot
 	}
 
 	const UNiagaraDataInterfaceCollisionQuery* OtherTyped = CastChecked<const UNiagaraDataInterfaceCollisionQuery>(Other);
-	return OtherTyped->MaxRayTraceCount == MaxRayTraceCount;
+	return OtherTyped->MaxTracesPerParticle == MaxTracesPerParticle;
 	return OtherTyped->MaxRetraces == MaxRetraces;
 }
 
@@ -978,7 +893,7 @@ bool UNiagaraDataInterfaceCollisionQuery::CopyToInternal(UNiagaraDataInterface*
 	}
 
 	UNiagaraDataInterfaceCollisionQuery* OtherTyped = CastChecked<UNiagaraDataInterfaceCollisionQuery>(Destination);
-	OtherTyped->MaxRayTraceCount = MaxRayTraceCount;
+	OtherTyped->MaxTracesPerParticle = MaxTracesPerParticle;
 	OtherTyped->MaxRetraces = MaxRetraces;
 	OtherTyped->MarkRenderDataDirty();
 	return true;
@@ -990,9 +905,9 @@ void UNiagaraDataInterfaceCollisionQuery::PushToRenderThreadImpl()
 
 	// Push Updates to Proxy, first release any resources
 	ENQUEUE_RENDER_COMMAND(FUpdateDI)(
-		[RT_Proxy, RT_MaxRayTraceRequests = MaxRayTraceCount, RT_MaxRetraces = MaxRetraces](FRHICommandListImmediate& RHICmdList)
+		[RT_Proxy, RT_MaxTracesPerParticle = MaxTracesPerParticle, RT_MaxRetraces = MaxRetraces](FRHICommandListImmediate& RHICmdList)
 		{
-			RT_Proxy->RenderThreadInitialize(RT_MaxRayTraceRequests, RT_MaxRetraces);
+			RT_Proxy->RenderThreadInitialize(RT_MaxTracesPerParticle, RT_MaxRetraces);
 		});
 }
 
@@ -1001,7 +916,7 @@ void UNiagaraDataInterfaceCollisionQuery::PostEditChangeProperty(struct FPropert
 {
 	Super::PostEditChangeProperty(PropertyChangedEvent);
 
-	if (PropertyChangedEvent.Property && PropertyChangedEvent.Property->GetFName() == GET_MEMBER_NAME_CHECKED(UNiagaraDataInterfaceCollisionQuery, MaxRayTraceCount))
+	if (PropertyChangedEvent.Property && PropertyChangedEvent.Property->GetFName() == GET_MEMBER_NAME_CHECKED(UNiagaraDataInterfaceCollisionQuery, MaxTracesPerParticle))
 	{
 		MarkRenderDataDirty();
 	}
@@ -1021,8 +936,11 @@ public:
 		RayTracingEnabledParam.Bind(ParameterMap, *(NDICollisionQueryLocal::RayTracingEnabledParamName + ParameterInfo.DataInterfaceHLSLSymbol));
 		MaxRayTraceCountParam.Bind(ParameterMap, *(NDICollisionQueryLocal::MaxRayTraceCountParamName + ParameterInfo.DataInterfaceHLSLSymbol));
 		RayRequestsParam.Bind(ParameterMap, *(NDICollisionQueryLocal::RayRequestsParamName + ParameterInfo.DataInterfaceHLSLSymbol));
+		RayRequestOffsetParam.Bind(ParameterMap, *(NDICollisionQueryLocal::RayRequestsOffsetParamName + ParameterInfo.DataInterfaceHLSLSymbol));
 		IntersectionResultsParam.Bind(ParameterMap, *(NDICollisionQueryLocal::IntersectionResultsParamName + ParameterInfo.DataInterfaceHLSLSymbol));
+		IntersectionResultOffsetParam.Bind(ParameterMap, *(NDICollisionQueryLocal::IntersectionResultsOffsetParamName + ParameterInfo.DataInterfaceHLSLSymbol));
 		RayTraceCountsParam.Bind(ParameterMap, *(NDICollisionQueryLocal::RayTraceCountsParamName + ParameterInfo.DataInterfaceHLSLSymbol));
+		RayTraceCountsOffsetParam.Bind(ParameterMap, *(NDICollisionQueryLocal::RayTraceCountsOffsetParamName + ParameterInfo.DataInterfaceHLSLSymbol));
 #endif
 	}
 
@@ -1050,22 +968,42 @@ public:
 		}
 
 #if RHI_RAYTRACING
+
+		FNiagaraRayTracingHelper& RTHelper = Context.Batcher->GetRayTracingHelper();
+		const FNiagaraRayTraceDispatchInfo* DispatchInfo = nullptr;
+		if (IsRayTracingEnabled() && GEnableGPUHWRTCollisions && QueryDI->MaxTracesPerParticle > 0)
+		{
+			DispatchInfo = &RTHelper.GetDispatch(QueryDI, RHICmdList);
+		}
+		else
+		{
+			DispatchInfo = &RTHelper.GetDummyDispatch(RHICmdList);
+		}
+
 		SetShaderValue(RHICmdList, ComputeShaderRHI, RayTracingEnabledParam, IsRayTracingEnabled() && GEnableGPUHWRTCollisions ? 1 : 0);
-		SetShaderValue(RHICmdList, ComputeShaderRHI, MaxRayTraceCountParam, QueryDI->MaxRayTraceCount);
+		SetShaderValue(RHICmdList, ComputeShaderRHI, MaxRayTraceCountParam, DispatchInfo->MaxRays);
 
 		if (RayRequestsParam.IsUAVBound())
 		{
-			RHICmdList.SetUAVParameter(ComputeShaderRHI, RayRequestsParam.GetUAVIndex(), QueryDI->RayTraceRequests.UAV);
+			check(DispatchInfo->RayRequests.IsValid());
+			RHICmdList.SetUAVParameter(ComputeShaderRHI, RayRequestsParam.GetUAVIndex(), DispatchInfo->RayRequests.Buffer->UAV);
+			SetShaderValue(RHICmdList, ComputeShaderRHI, RayRequestOffsetParam, DispatchInfo->RayRequests.Offset);
 		}
 
 		if (IntersectionResultsParam.IsBound())
 		{
-			SetSRVParameter(RHICmdList, ComputeShaderRHI, IntersectionResultsParam, QueryDI->RayTraceIntersections.SRV);
+			check(DispatchInfo->LastFrameRayTraceIntersections.IsValid());
+
+			SetSRVParameter(RHICmdList, ComputeShaderRHI, IntersectionResultsParam, DispatchInfo->LastFrameRayTraceIntersections.Buffer->SRV);
+			SetShaderValue(RHICmdList, ComputeShaderRHI, IntersectionResultOffsetParam, DispatchInfo->LastFrameRayTraceIntersections.Offset);
 		}
 
 		if (RayTraceCountsParam.IsUAVBound())
 		{
-			RHICmdList.SetUAVParameter(ComputeShaderRHI, RayTraceCountsParam.GetUAVIndex(), QueryDI->RayTraceCounts.UAV);
+			check(DispatchInfo->RayCounts.IsValid());
+
+			RHICmdList.SetUAVParameter(ComputeShaderRHI, RayTraceCountsParam.GetUAVIndex(), DispatchInfo->RayCounts.Buffer->UAV);
+			SetShaderValue(RHICmdList, ComputeShaderRHI, RayTraceCountsOffsetParam, DispatchInfo->RayCounts.Offset);
 		}
 #endif
 	}
@@ -1095,8 +1033,11 @@ private:
 	LAYOUT_FIELD(FShaderParameter, RayTracingEnabledParam);
 	LAYOUT_FIELD(FShaderParameter, MaxRayTraceCountParam);
 	LAYOUT_FIELD(FRWShaderParameter, RayRequestsParam);
+	LAYOUT_FIELD(FShaderParameter, RayRequestOffsetParam);
 	LAYOUT_FIELD(FShaderResourceParameter, IntersectionResultsParam);
+	LAYOUT_FIELD(FShaderParameter, IntersectionResultOffsetParam);
 	LAYOUT_FIELD(FRWShaderParameter, RayTraceCountsParam);
+	LAYOUT_FIELD(FShaderParameter, RayTraceCountsOffsetParam);
 #endif
 };
 
@@ -1104,4 +1045,4 @@ IMPLEMENT_TYPE_LAYOUT(FNiagaraDataInterfaceParametersCS_CollisionQuery);
 
 IMPLEMENT_NIAGARA_DI_PARAMETER(UNiagaraDataInterfaceCollisionQuery, FNiagaraDataInterfaceParametersCS_CollisionQuery);
 
-#undef LOCTEXT_NAMESPACE
\ No newline at end of file
+#undef LOCTEXT_NAMESPACE
diff --git a/Engine/Plugins/FX/Niagara/Source/Niagara/Private/NiagaraEmitterInstanceBatcher.cpp b/Engine/Plugins/FX/Niagara/Source/Niagara/Private/NiagaraEmitterInstanceBatcher.cpp
--- a/Engine/Plugins/FX/Niagara/Source/Niagara/Private/NiagaraEmitterInstanceBatcher.cpp
+++ b/Engine/Plugins/FX/Niagara/Source/Niagara/Private/NiagaraEmitterInstanceBatcher.cpp
@@ -249,6 +249,10 @@ NiagaraEmitterInstanceBatcher::~NiagaraEmitterInstanceBatcher()
 {
 	FinishDispatches();
 
+#if RHI_RAYTRACING
+	RayTracingHelper->Reset();
+#endif
+
 	GlobalCBufferLayout = nullptr;
 	SystemCBufferLayout = nullptr;
 	OwnerCBufferLayout = nullptr;
@@ -504,6 +508,10 @@ void NiagaraEmitterInstanceBatcher::ProcessPendingTicksFlush(FRHICommandListImme
 			//UE_LOG(LogNiagara, Log, TEXT("NiagaraEmitterInstanceBatcher: Queued ticks are being Destroyed due to not rendering.  This may result in undesirable simulation artifacts."));
 
 			FinishDispatches();
+
+#if RHI_RAYTRACING
+			RayTracingHelper->Reset();
+#endif
 			break;
 		}
 	}
@@ -1618,6 +1626,13 @@ void NiagaraEmitterInstanceBatcher::PreInitViews(FRDGBuilder& GraphBuilder, bool
 	{
 		UpdateInstanceCountManager(GraphBuilder.RHICmdList);
 		PrepareAllTicks(GraphBuilder.RHICmdList);
+		
+#if RHI_RAYTRACING
+		if (HasRayTracingScene())
+		{
+			RayTracingHelper->BeginFrame(GraphBuilder.RHICmdList);
+		}
+#endif
 
 		if ( NiagaraEmitterInstanceBatcherLocal::GDebugLogging)
 		{
@@ -1694,10 +1709,6 @@ void NiagaraEmitterInstanceBatcher::PostRenderOpaque(FRDGBuilder& GraphBuilder,
 
 #if RHI_RAYTRACING
 			BuildRayTracingSceneInfo(RHICmdList, Views);
-			ON_SCOPE_EXIT
-			{
-				ResetRayTracingSceneInfo();
-			};
 #endif
 
 			CurrentPassViews = Views;
@@ -1709,6 +1720,13 @@ void NiagaraEmitterInstanceBatcher::PostRenderOpaque(FRDGBuilder& GraphBuilder,
 
 			FinishDispatches();
 
+#if RHI_RAYTRACING
+			if(HasRayTracingScene())
+			{
+				RayTracingHelper->EndFrame(RHICmdList, GetScene());
+			}
+#endif
+
 			// Clear CurrentPassViews
 			CurrentPassViews = TConstArrayView<FViewInfo>();
 
@@ -1962,11 +1980,11 @@ void NiagaraEmitterInstanceBatcher::ResetRayTracingSceneInfo()
 	RayTracingHelper->Reset();
 }
 
-void NiagaraEmitterInstanceBatcher::IssueRayTraces(FRHICommandList& RHICmdList, const FIntPoint& RayTraceCounts, uint32 MaxRetraces, FRHIShaderResourceView* RayTraceRequests, FRWBuffer* IndirectArgsBuffer, uint32 IndirectArgsOffset, FRHIUnorderedAccessView* RayTraceResults) const
-{
-	check(NumProxiesThatRequireRayTracingScene > 0);
 
-	RayTracingHelper->IssueRayTraces(RHICmdList, GetScene(), RayTraceCounts, MaxRetraces, RayTraceRequests, IndirectArgsBuffer, IndirectArgsOffset, RayTraceResults);
+FNiagaraRayTracingHelper& NiagaraEmitterInstanceBatcher::GetRayTracingHelper()const
+{
+	check(RayTracingHelper.IsValid());
+	return *RayTracingHelper.Get();
 }
 
 bool NiagaraEmitterInstanceBatcher::HasRayTracingScene() const
@@ -2262,4 +2280,4 @@ void NiagaraEmitterInstanceBatcher::BroadcastTemporalEffect(FRHICommandList& RHI
 		}
 	}
 }
-#endif // WITH_MGPU
\ No newline at end of file
+#endif // WITH_MGPU
diff --git a/Engine/Plugins/FX/Niagara/Source/NiagaraShader/Private/NiagaraRayTracingHelper.cpp b/Engine/Plugins/FX/Niagara/Source/NiagaraShader/Private/NiagaraRayTracingHelper.cpp
--- a/Engine/Plugins/FX/Niagara/Source/NiagaraShader/Private/NiagaraRayTracingHelper.cpp
+++ b/Engine/Plugins/FX/Niagara/Source/NiagaraShader/Private/NiagaraRayTracingHelper.cpp
@@ -68,7 +68,9 @@ class FNiagaraCollisionRayTraceRG : public FGlobalShader
 		SHADER_PARAMETER(uint32, HashTableSize)
 		SHADER_PARAMETER_SRV(RaytracingAccelerationStructure, TLAS)
 		SHADER_PARAMETER_SRV(Buffer<FNiagaraRayData>, Rays)
+		SHADER_PARAMETER(uint32, RaysOffset)
 		SHADER_PARAMETER_UAV(Buffer<FNiagaraRayTracingResult>, CollisionOutput)
+		SHADER_PARAMETER(uint32, CollisionOutputOffset)
 		SHADER_PARAMETER_SRV(Buffer<UINT>, RayTraceCounts)
 		SHADER_PARAMETER_SRV(Buffer<UINT>, HashToCollisionGroups)
 		SHADER_PARAMETER(uint32, MaxRetraces)
@@ -155,6 +157,45 @@ IMPLEMENT_GLOBAL_SHADER(FNiagaraCollisionRayTraceRG, "/Plugin/FX/Niagara/Private
 IMPLEMENT_GLOBAL_SHADER(FNiagaraCollisionRayTraceCH, "/Plugin/FX/Niagara/Private/NiagaraRayTracingShaders.usf", "NiagaraCollisionRayTraceCH", SF_RayHitGroup);
 IMPLEMENT_GLOBAL_SHADER(FNiagaraCollisionRayTraceMiss, "/Plugin/FX/Niagara/Private/NiagaraRayTracingShaders.usf", "NiagaraCollisionRayTraceMiss", SF_RayMiss);
 
+
+int32 GNiagaraHWRTScratchBucketSize = 1024;
+static FAutoConsoleVariableRef CVarNiagaraHWRTScratchBucketSize(
+	TEXT("fx.Niagara.RayTracing.ScratchPadBucketSize"),
+	GNiagaraHWRTScratchBucketSize,
+	TEXT("Size (in elements) for Niagara hardware ray tracing scratch buffer buckets. \n"),
+	ECVF_Default
+);
+
+int32 GNiagaraHWRTCountsScratchBucketSize = 3 * 256;
+static FAutoConsoleVariableRef CVarNiagaraHWRTCountsScratchBucketSize(
+	TEXT("fx.Niagara.RayTracing.CountsScratchPadBucketSize"),
+	GNiagaraHWRTCountsScratchBucketSize,
+	TEXT("Scratch bucket size for the HWRT counts buffer. This buffer requires 4. \n"),
+	ECVF_Default
+);
+
+FNiagaraRayTracingHelper::FNiagaraRayTracingHelper(EShaderPlatform InShaderPlatform)
+	: ShaderPlatform(InShaderPlatform)
+	, RayRequests(GNiagaraHWRTScratchBucketSize, TEXT("NiagaraRayRequests"))
+	, RayTraceIntersections(GNiagaraHWRTScratchBucketSize, TEXT("NiagaraRayTraceIntersections"))
+	, RayTraceCounts(GNiagaraHWRTCountsScratchBucketSize, TEXT("NiagaraRayTraceCounts"), BUF_Static | BUF_DrawIndirect)
+{}
+
+FNiagaraRayTracingHelper::~FNiagaraRayTracingHelper()
+{
+	Reset();
+	Dispatches.Reset();
+
+	//TODO: Move RT Helper out of NiagaraShader and into Niagara. Makes general sense plus it would allow use to hook into these memory stats.
+// 	DEC_MEMORY_STAT_BY(STAT_NiagaraGPUDataInterfaceMemory, RayRequests.AllocatedBytes());
+// 	DEC_MEMORY_STAT_BY(STAT_NiagaraGPUDataInterfaceMemory, RayTraceIntersections.AllocatedBytes());
+// 	DEC_MEMORY_STAT_BY(STAT_NiagaraGPUDataInterfaceMemory, RayTraceCounts.AllocatedBytes());
+
+	RayRequests.Release();
+	RayTraceIntersections.Release();
+	RayTraceCounts.Release();
+}
+
 void FNiagaraRayTracingHelper::Reset()
 {
 	RayTracingPipelineState = nullptr;
@@ -163,6 +204,119 @@ void FNiagaraRayTracingHelper::Reset()
 	ViewUniformBuffer = nullptr;
 }
 
+void FNiagaraRayTracingHelper::BeginFrame(FRHICommandList& RHICmdList)
+{
+	//Store off last frames dispatches so we can still access the previous results buffers for reading in this frame's simulations.
+	PreviousFrameDispatches = MoveTemp(Dispatches);
+	Dispatches.Reset();
+
+	//Ensure we're each buffer is definitely in the right access state.
+	RayRequests.Transition(RHICmdList, ERHIAccess::Unknown, ERHIAccess::UAVCompute);//Simulations will write new ray trace requests.
+	RayTraceIntersections.Transition(RHICmdList, ERHIAccess::Unknown, ERHIAccess::SRVCompute);//Simulations will read from the previous frames results.
+	RayTraceCounts.Transition(RHICmdList, ERHIAccess::Unknown, ERHIAccess::UAVCompute);//Simulations will accumulate the number of traces requests here.
+
+	//Reset the allocations.
+	RayRequests.Reset();
+
+	//Note this doesn't change any buffers or data themselves and the ray intersections buffer are still going to be read as SRVs in the up coming simulations dispatches.
+	//We're just clearing existing allocations here.
+	RayTraceIntersections.Reset();
+
+	//We have to also clear the counts/indirect args buffer.
+	RayTraceCounts.Reset(RHICmdList, true);
+
+	// Clear the dummy buffer allocations.
+	DummyDispatch.Reset();
+
+	RayTracingPipelineState = nullptr;
+	RayTracingScene = nullptr;
+	RayTracingSceneView = nullptr;
+	ViewUniformBuffer = nullptr;
+}
+
+void FNiagaraRayTracingHelper::EndFrame(FRHICommandList& RHICmdList, FScene* Scene)
+{
+	RayRequests.Transition(RHICmdList, ERHIAccess::UAVCompute, ERHIAccess::SRVCompute);//Ray trace dispatches will read the ray requests.
+	RayTraceIntersections.Transition(RHICmdList, ERHIAccess::SRVCompute, ERHIAccess::UAVCompute);//Ray trace dispatches will write new results.
+	RayTraceCounts.Transition(RHICmdList, ERHIAccess::UAVCompute, ERHIAccess::IndirectArgs | ERHIAccess::SRVCompute);//Ray trace dispatches read these counts and use them for indirect dispatches.
+	
+	IssueRayTraces(RHICmdList, Scene);
+
+	RayRequests.Transition(RHICmdList, ERHIAccess::SRVCompute, ERHIAccess::UAVCompute);//Next frame, simulations will write new ray trace requests.
+	RayTraceIntersections.Transition(RHICmdList, ERHIAccess::UAVCompute, ERHIAccess::SRVCompute);//Next frame these results will be read by simulations shaders.
+	RayTraceCounts.Transition(RHICmdList, ERHIAccess::IndirectArgs | ERHIAccess::SRVCompute, ERHIAccess::UAVCompute);//Next frame these counts will be written by simulation shaders.
+}
+
+void FNiagaraRayTracingHelper::AddToDispatch(FNiagaraDataInterfaceProxy* DispatchKey, uint32 MaxRays, int32 MaxRetraces)
+{
+	FNiagaraRayTraceDispatchInfo& Dispatch = Dispatches.FindOrAdd(DispatchKey);
+	Dispatch.MaxRays += MaxRays;
+	Dispatch.MaxRetraces = MaxRetraces;
+}
+
+const FNiagaraRayTraceDispatchInfo& FNiagaraRayTracingHelper::GetDispatch(FNiagaraDataInterfaceProxy* DispatchKey, FRHICommandList& RHICmdList)
+{
+	FNiagaraRayTraceDispatchInfo& Dispatch = Dispatches.FindChecked(DispatchKey);
+
+	//Finalize allocations first time we access the dispatch.
+	if (Dispatch.RayRequests.IsValid() == false)
+	{
+		//TODO: Move RT Helper out of NiagaraShader and into Niagara. Makes general sense plus it would allow use to hook into these memory stats.
+// 		DEC_MEMORY_STAT_BY(STAT_NiagaraGPUDataInterfaceMemory, RayRequests.AllocatedBytes());
+// 		DEC_MEMORY_STAT_BY(STAT_NiagaraGPUDataInterfaceMemory, RayTraceIntersections.AllocatedBytes());
+// 		DEC_MEMORY_STAT_BY(STAT_NiagaraGPUDataInterfaceMemory, RayTraceCounts.AllocatedBytes());
+
+		Dispatch.RayRequests = RayRequests.Alloc(Dispatch.MaxRays);
+		Dispatch.RayTraceIntersections = RayTraceIntersections.Alloc(Dispatch.MaxRays);
+		Dispatch.RayCounts = RayTraceCounts.Alloc<uint32>(3, RHICmdList, true);
+
+		//TODO: Could add a delegate to call from the scratch when new buffers are allocated to reduce stat spam here?
+// 		INC_MEMORY_STAT_BY(STAT_NiagaraGPUDataInterfaceMemory, RayRequests.AllocatedBytes());
+// 		INC_MEMORY_STAT_BY(STAT_NiagaraGPUDataInterfaceMemory, RayTraceIntersections.AllocatedBytes());
+// 		INC_MEMORY_STAT_BY(STAT_NiagaraGPUDataInterfaceMemory, RayTraceCounts.AllocatedBytes());
+
+		//Find last frame's results buffer if there was one. Ideally we can do this without the map but it will need a bit of wrangling.
+		FNiagaraRayTraceDispatchInfo* PrevDispatch = PreviousFrameDispatches.Find(DispatchKey);
+		if (PrevDispatch && PrevDispatch->RayTraceIntersections.IsValid())
+		{
+			Dispatch.LastFrameRayTraceIntersections = PrevDispatch->RayTraceIntersections;
+		}
+		else
+		{
+			//Set the current frame results just so this is a valid buffer.
+			//If we didn't have a results buffer last frame then this won't actually be read.
+			Dispatch.LastFrameRayTraceIntersections = Dispatch.RayTraceIntersections;
+		}
+	}
+	
+	return Dispatch;
+}
+
+const FNiagaraRayTraceDispatchInfo& FNiagaraRayTracingHelper::GetDummyDispatch(FRHICommandList& RHICmdList)
+{
+	if (DummyDispatch.IsValid() == false)
+	{
+		//TODO: Move RT Helper out of NiagaraShader and into Niagara. Makes general sense plus it would allow use to hook into these memory stats.
+// 		DEC_MEMORY_STAT_BY(STAT_NiagaraGPUDataInterfaceMemory, RayRequests.AllocatedBytes());
+// 		DEC_MEMORY_STAT_BY(STAT_NiagaraGPUDataInterfaceMemory, RayTraceIntersections.AllocatedBytes());
+// 		DEC_MEMORY_STAT_BY(STAT_NiagaraGPUDataInterfaceMemory, RayTraceCounts.AllocatedBytes());
+
+		DummyDispatch.RayRequests = RayRequests.Alloc(1);
+		DummyDispatch.RayTraceIntersections = RayTraceIntersections.Alloc(1);
+		DummyDispatch.LastFrameRayTraceIntersections = DummyDispatch.RayTraceIntersections;
+		DummyDispatch.RayCounts = RayTraceCounts.Alloc<uint32>(3, RHICmdList, true);
+		
+		//TODO: Move RT Helper out of NiagaraShader and into Niagara. Makes general sense plus it would allow use to hook into these memory stats.
+// 		INC_MEMORY_STAT_BY(STAT_NiagaraGPUDataInterfaceMemory, RayRequests.AllocatedBytes());
+// 		INC_MEMORY_STAT_BY(STAT_NiagaraGPUDataInterfaceMemory, RayTraceIntersections.AllocatedBytes());
+// 		INC_MEMORY_STAT_BY(STAT_NiagaraGPUDataInterfaceMemory, RayTraceCounts.AllocatedBytes());
+
+		DummyDispatch.MaxRays = 0;
+		DummyDispatch.MaxRetraces = 0;
+	}
+	return DummyDispatch;
+}
+
 static FRayTracingPipelineState* CreateNiagaraRayTracingPipelineState(
 	EShaderPlatform Platform,
 	FRHICommandList& RHICmdList,
@@ -382,7 +536,7 @@ void FNiagaraRayTracingHelper::BuildRayTracingSceneInfo(FRHICommandList& RHICmdL
 	}
 }
 
-void FNiagaraRayTracingHelper::IssueRayTraces(FRHICommandList& RHICmdList, FScene* Scene, const FIntPoint& RayTraceCounts, uint32 MaxRetraces, FRHIShaderResourceView* RayTraceRequests, FRWBuffer* IndirectArgsBuffer, uint32 IndirectArgsOffset, FRHIUnorderedAccessView* RayTraceResults) const
+void FNiagaraRayTracingHelper::IssueRayTraces(FRHICommandList& RHICmdList, FScene* Scene)
 {
 	SCOPED_DRAW_EVENT(RHICmdList, NiagaraIssueCollisionRayTraces);
 	check(RayTracingPipelineState);
@@ -390,60 +544,67 @@ void FNiagaraRayTracingHelper::IssueRayTraces(FRHICommandList& RHICmdList, FScen
 	check(RayTracingSceneView);
 	check(Scene);
 
-	FGlobalShaderMap* ShaderMap = GetGlobalShaderMap(ShaderPlatform);
-
-	TShaderRef<FNiagaraCollisionRayTraceRG> RGShader = FNiagaraCollisionRayTraceRG::GetShader(ShaderMap);
-
-	FNiagaraCollisionRayTraceRG::FParameters Params;
-
-	Params.GPUSceneParameters.GPUSceneInstanceSceneData = Scene->GPUScene.InstanceSceneDataBuffer.SRV;
-	Params.GPUSceneParameters.GPUSceneInstancePayloadData = Scene->GPUScene.InstancePayloadDataBuffer.SRV;
-	Params.GPUSceneParameters.GPUScenePrimitiveSceneData = Scene->GPUScene.PrimitiveBuffer.SRV;
-	Params.GPUSceneParameters.GPUSceneFrameNumber = Scene->GPUScene.GetSceneFrameNumber();
-
-	Params.HashTable = PrimIdHashTable.SRV;
-	Params.HashTableSize = HashTableSize;
-	Params.TLAS = RayTracingSceneView;
-	Params.Rays = RayTraceRequests;
-	Params.CollisionOutput = RayTraceResults;
-	Params.HashToCollisionGroups = HashToCollisionGroups.SRV;
-	Params.MaxRetraces = MaxRetraces;
-
-	if (FNiagaraCollisionRayTraceRG::SupportsIndirectDispatch())
+	for (TPair<FNiagaraDataInterfaceProxy*, FNiagaraRayTraceDispatchInfo>& Pair : Dispatches)
 	{
-		RHICmdList.Transition(FRHITransitionInfo(IndirectArgsBuffer->UAV, ERHIAccess::UAVCompute, ERHIAccess::IndirectArgs | ERHIAccess::SRVCompute));
-
-		FRayTracingShaderBindingsWriter GlobalResources;
-		SetShaderParameters(GlobalResources, RGShader, Params);
+		FNiagaraRayTraceDispatchInfo& DispatchInfo = Pair.Value;
+		
+		if (DispatchInfo.MaxRays == 0)
+		{
+			continue;
+		}
 
-		RHICmdList.RayTraceDispatchIndirect(
-			RayTracingPipelineState,
-			RGShader.GetRayTracingShader(),
-			RayTracingScene,
-			GlobalResources,
-			IndirectArgsBuffer->Buffer,
-			IndirectArgsOffset);
+		FGlobalShaderMap* ShaderMap = GetGlobalShaderMap(ShaderPlatform);
 
-		RHICmdList.Transition(FRHITransitionInfo(IndirectArgsBuffer->UAV, ERHIAccess::IndirectArgs | ERHIAccess::SRVCompute, ERHIAccess::UAVCompute));
-	}
-	else
-	{
-		RHICmdList.Transition(FRHITransitionInfo(IndirectArgsBuffer->UAV, ERHIAccess::UAVCompute, ERHIAccess::SRVCompute));
-		Params.RayTraceCounts = IndirectArgsBuffer->SRV;
+		TShaderRef<FNiagaraCollisionRayTraceRG> RGShader = FNiagaraCollisionRayTraceRG::GetShader(ShaderMap);
 
-		FRayTracingShaderBindingsWriter GlobalResources;
-		SetShaderParameters(GlobalResources, RGShader, Params);
+		FNiagaraCollisionRayTraceRG::FParameters Params;
 
-		RHICmdList.RayTraceDispatch(
-			RayTracingPipelineState,
-			RGShader.GetRayTracingShader(),
-			RayTracingScene,
-			GlobalResources,
-			RayTraceCounts.X,
-			RayTraceCounts.Y
-		);
+		Params.GPUSceneParameters.GPUSceneInstanceSceneData = Scene->GPUScene.InstanceSceneDataBuffer.SRV;
+		Params.GPUSceneParameters.GPUSceneInstancePayloadData = Scene->GPUScene.InstancePayloadDataBuffer.SRV;
+		Params.GPUSceneParameters.GPUScenePrimitiveSceneData = Scene->GPUScene.PrimitiveBuffer.SRV;
+		Params.GPUSceneParameters.GPUSceneFrameNumber = Scene->GPUScene.GetSceneFrameNumber();
 
-		RHICmdList.Transition(FRHITransitionInfo(IndirectArgsBuffer->UAV, ERHIAccess::SRVCompute, ERHIAccess::UAVCompute));
+		Params.HashTable = PrimIdHashTable.SRV;
+		Params.HashTableSize = HashTableSize;
+		Params.TLAS = RayTracingSceneView;
+		Params.Rays = DispatchInfo.RayRequests.Buffer->SRV;
+		Params.RaysOffset = DispatchInfo.RayRequests.Offset;
+		Params.CollisionOutput = DispatchInfo.RayTraceIntersections.Buffer->UAV;
+		Params.CollisionOutputOffset = DispatchInfo.RayTraceIntersections.Offset;
+		Params.HashToCollisionGroups = HashToCollisionGroups.SRV;
+		Params.MaxRetraces = DispatchInfo.MaxRetraces;
+
+		if (FNiagaraCollisionRayTraceRG::SupportsIndirectDispatch())
+		{
+			FRayTracingShaderBindingsWriter GlobalResources;
+			SetShaderParameters(GlobalResources, RGShader, Params);
+
+			//Can we wrangle things so we can have one indirect dispatch with each internal dispatch pointing to potentially different Ray and Results buffers?
+			//For now have a each as a unique dispatch.
+			RHICmdList.RayTraceDispatchIndirect(
+				RayTracingPipelineState,
+				RGShader.GetRayTracingShader(),
+				RayTracingScene,
+				GlobalResources,
+				DispatchInfo.RayCounts.Buffer->Buffer,
+				DispatchInfo.RayCounts.Offset * sizeof(uint32));
+		}
+		else
+		{
+			Params.RayTraceCounts = DispatchInfo.RayCounts.Buffer->SRV;
+
+			FRayTracingShaderBindingsWriter GlobalResources;
+			SetShaderParameters(GlobalResources, RGShader, Params);
+
+			RHICmdList.RayTraceDispatch(
+				RayTracingPipelineState,
+				RGShader.GetRayTracingShader(),
+				RayTracingScene,
+				GlobalResources,
+				DispatchInfo.MaxRays,
+				1
+			);
+		}
 	}
 }
 
@@ -455,4 +616,4 @@ bool FNiagaraRayTracingHelper::IsValid() const
 		&& ViewUniformBuffer;
 }
 
-#endif
+#endif
\ No newline at end of file
diff --git a/Engine/Plugins/FX/Niagara/Source/NiagaraShader/Public/NiagaraRayTracingHelper.h b/Engine/Plugins/FX/Niagara/Source/NiagaraShader/Public/NiagaraRayTracingHelper.h
--- a/Engine/Plugins/FX/Niagara/Source/NiagaraShader/Public/NiagaraRayTracingHelper.h
+++ b/Engine/Plugins/FX/Niagara/Source/NiagaraShader/Public/NiagaraRayTracingHelper.h
@@ -4,7 +4,6 @@
 #include "CoreMinimal.h"
 
 #include "RHIDefinitions.h"
-
 #if RHI_RAYTRACING
 
 class FViewInfo;
@@ -15,6 +14,401 @@ class FRHIShaderResourceView;
 class FRHIUniformBuffer;
 class FRHIUnorderedAccessView;
 struct FRWBuffer;
+struct FNiagaraDataInterfaceProxy;
+
+
+//////////////////////////////////////////////////////////////////////////
+//TODO: Move scratch pad buffers to core rendering code along side global dynamic read buffers etc.
+
+/** Allocates space from a pool of fixed size static RWBuffers.
+*	Allocations are valid between calls to Reset.
+*	Useful for transient working data generated in one dispatch and consumed in others.
+*	All buffers in the same scratch pad are transitioned as one.
+*/
+class NIAGARASHADER_API FGPUScratchPad
+{
+public:
+	struct FAllocation
+	{
+		FRWBuffer* Buffer = nullptr;
+		uint32 Offset = 0;
+
+		FORCEINLINE bool IsValid()const { return Buffer != nullptr; }
+	};
+
+	struct FGPUScratchBuffer
+	{
+	public:
+		struct FBuffer
+		{
+			FRWBuffer Buffer;
+			uint32 Used = 0;
+		};
+
+		FGPUScratchBuffer(uint32 InBucketSize, EPixelFormat InPixelFormat, uint32 InBytesPerElement, EBufferUsageFlags InBufferUsageFlags)
+			: MinBucketSize(InBucketSize), PixelFormat(InPixelFormat), BytesPerElement(InBytesPerElement), BufferUsageFlags(InBufferUsageFlags)
+		{}
+
+		/** Allocates NumElements from the buffers. Will try to find an existing buffer to fit the allocation but will create a new one if needed. */
+		bool Allocate(uint32 NumElements, FString& BufferDebugName, FRHICommandList& RHICmdList, FGPUScratchPad::FAllocation& OutAllocation)
+		{
+			//Look for a buffer to fit this allocation.
+			FBuffer* ToUse = nullptr;
+			for (FBuffer& Buffer : Buffers)
+			{
+				uint32 NumBufferElements = Buffer.Buffer.NumBytes / BytesPerElement;
+				if (NumElements + Buffer.Used <= NumBufferElements)
+				{
+					ToUse = &Buffer;
+				}
+			}
+
+			bool bNew = false;
+			if (ToUse == nullptr)
+			{
+				//Allocate a new buffer that will fit this allocation.
+				FBuffer* NewBuffer = new FBuffer();
+				Buffers.Add(NewBuffer);
+
+				//Round up the size to pow 2 to ensure a slowly growing user doesn't just keep allocating new buckets. Probably a better guess size here than power of 2.
+				uint32 NewBufferSize = FPlatformMath::RoundUpToPowerOfTwo(FMath::Max(MinBucketSize, NumElements));
+
+				NewBuffer->Buffer.Initialize(*BufferDebugName, BytesPerElement, NewBufferSize, PixelFormat, CurrentAccess, BufferUsageFlags);
+			
+				ToUse = NewBuffer;
+
+				NumAllocatedBytes += NewBuffer->Buffer.NumBytes;
+
+				bNew = true;
+			}
+
+			OutAllocation.Buffer = &ToUse->Buffer;
+			OutAllocation.Offset = ToUse->Used;
+			ToUse->Used += NumElements;
+			NumUsedBytes += NumElements * BytesPerElement;
+			return bNew;
+		}
+
+		/** Resets the buffer and releases it's internal RWBuffers too. */
+		void Release()
+		{
+			Reset();
+			CurrentAccess = ERHIAccess::UAVCompute;
+			for (FBuffer& Buffer : Buffers)
+			{
+				Buffer.Buffer.Release();
+				Buffer.Used = 0;
+			}
+			Buffers.Reset();
+
+			NumAllocatedBytes = 0;
+		}
+
+		/** Resets the usage for all buffers to zero. */
+		void Reset()
+		{
+			for (FBuffer& Buffer : Buffers)
+			{
+				Buffer.Used = 0;
+			}
+
+			NumUsedBytes = 0;
+		}
+
+		template<typename TFunc>
+		void ForEachBuffer(TFunc Func)
+		{
+			for (FBuffer& Buffer : Buffers)
+			{
+				Func(Buffer.Buffer);
+			}
+		}
+
+		/** Transitions all buffers in the scratch pad. */
+		void Transition(FRHICommandList& RHICmdList, ERHIAccess InPreviousState, ERHIAccess InNewState, EResourceTransitionFlags InFlags = EResourceTransitionFlags::None)
+		{
+			//Store off current access state so we can place any new buffers in the same state.
+			CurrentAccess = InNewState;
+
+ 			TArray<FRHITransitionInfo, TMemStackAllocator<>> Transitions;
+ 			Transitions.Reserve(Buffers.Num());
+ 			for (FBuffer& Buffer : Buffers)
+ 			{
+ 				Transitions.Emplace(Buffer.Buffer.UAV, InPreviousState, InNewState, InFlags);
+ 			}
+ 
+ 			RHICmdList.Transition(Transitions);
+		}
+
+		/** Returns the total bytes allocated for all buffers in the scratch pad. */
+		FORCEINLINE uint32 AllocatedBytes()const
+		{
+			return NumAllocatedBytes;
+		}
+
+		/** Returns the number of used bytes allocated for all buffers in the scratch pad. */
+		FORCEINLINE uint32 UsedBytes()const
+		{
+			return NumUsedBytes;
+		}
+
+		FORCEINLINE ERHIAccess GetExpectedCurrentAccess()const{return CurrentAccess;}
+	private:
+
+		uint32 MinBucketSize = 4096;
+		TIndirectArray<FBuffer> Buffers;
+
+		EPixelFormat PixelFormat;
+		uint32 BytesPerElement = 0;
+		EBufferUsageFlags BufferUsageFlags = EBufferUsageFlags::None;
+		ERHIAccess CurrentAccess = ERHIAccess::UAVCompute;
+
+		uint32 NumAllocatedBytes = 0;
+		uint32 NumUsedBytes = 0;
+	};
+
+	FGPUScratchPad(uint32 InBucketSize, const FString& InDebugName, EBufferUsageFlags BufferUsageFlags = BUF_Static)
+		: BucketSize(InBucketSize)
+		, DebugName(InDebugName)
+		, ScratchBufferFloat(BucketSize, PF_R32_FLOAT, sizeof(float), BufferUsageFlags)
+		, ScratchBufferUInt(BucketSize, PF_R32_UINT, sizeof(uint32), BufferUsageFlags)
+	{
+	}
+
+	/** Allocate a number of elements from the scratch pad. This will attempt to allocate from an existing buffer but will create a new one if it can't. */
+	template<typename T>
+	FORCEINLINE FAllocation Alloc(uint32 Size, FRHICommandList& RHICmdList, bool bClearNew);
+
+	/** Resets the scratch pad and also releases all buffers it holds. */
+	void Release()
+	{
+		ScratchBufferFloat.Release();
+		ScratchBufferUInt.Release();
+	}
+
+	/** Resets the usage counters for buffers in this scratch pad, invalidating all previous allocations. Does not affect actual allocations. Optionally clears all buffers to zero. */
+	void Reset(FRHICommandList& RHICmdList, bool bClear)
+	{
+		ScratchBufferFloat.Reset();
+		ScratchBufferUInt.Reset();
+
+		if (bClear)
+		{
+			Transition(RHICmdList, ERHIAccess::Unknown, ERHIAccess::UAVCompute);
+
+			auto ClearFloatBuffer = [&](FRWBuffer& Buffer)
+			{
+				RHICmdList.ClearUAVFloat(Buffer.UAV, FVector4(0.0f));
+			};
+			ScratchBufferFloat.ForEachBuffer(ClearFloatBuffer);
+
+			auto ClearUIntBuffer = [&](FRWBuffer& Buffer)
+			{
+				RHICmdList.ClearUAVUint(Buffer.UAV, FUintVector4(0, 0, 0, 0));
+			};
+			ScratchBufferUInt.ForEachBuffer(ClearUIntBuffer);
+
+			Transition(RHICmdList, ERHIAccess::UAVCompute, ERHIAccess::UAVCompute);
+		}
+	}
+
+	/** Transitions all buffers in the scratch pad. */
+	void Transition(FRHICommandList& RHICmdList, ERHIAccess InPreviousState, ERHIAccess InNewState, EResourceTransitionFlags InFlags = EResourceTransitionFlags::None)
+	{
+		ScratchBufferFloat.Transition(RHICmdList, InPreviousState, InNewState, InFlags);
+		ScratchBufferUInt.Transition(RHICmdList, InPreviousState, InNewState, InFlags);
+	}
+
+	/** Returns the total bytes allocated for all buffers in the scratch pad. */
+	FORCEINLINE uint32 AllocatedBytes()const
+	{
+		return ScratchBufferFloat.AllocatedBytes() + ScratchBufferUInt.AllocatedBytes();
+	}
+
+	/** Returns the number of used bytes allocated for all buffers in the scratch pad. */
+	FORCEINLINE uint32 UsedBytes()const
+	{
+		return ScratchBufferFloat.UsedBytes() + ScratchBufferUInt.UsedBytes();
+	}
+
+private:
+
+	/** Size of each FRWBuffer in the pool. */
+	uint32 BucketSize = 4096;
+	FString DebugName;
+
+	FGPUScratchBuffer ScratchBufferFloat;
+	FGPUScratchBuffer ScratchBufferUInt;
+};
+
+template<>
+FORCEINLINE FGPUScratchPad::FAllocation FGPUScratchPad::Alloc<float>(uint32 Size, FRHICommandList& RHICmdList, bool bClearNew)
+{
+	FGPUScratchPad::FAllocation Allocation;
+	bool bNewBuffer = ScratchBufferFloat.Allocate(Size, DebugName, RHICmdList, Allocation);
+	if (bNewBuffer && bClearNew)
+	{
+		RHICmdList.Transition(FRHITransitionInfo(Allocation.Buffer->UAV, ScratchBufferFloat.GetExpectedCurrentAccess(), ERHIAccess::UAVCompute));
+		RHICmdList.ClearUAVFloat(Allocation.Buffer->UAV, FVector4(0.0f));
+		RHICmdList.Transition(FRHITransitionInfo(Allocation.Buffer->UAV, ERHIAccess::UAVCompute, ScratchBufferFloat.GetExpectedCurrentAccess()));
+	}
+	return Allocation;
+}
+
+template<>
+FORCEINLINE FGPUScratchPad::FAllocation FGPUScratchPad::Alloc<uint32>(uint32 Size, FRHICommandList& RHICmdList, bool bClearNew)
+{
+	FGPUScratchPad::FAllocation Allocation;
+	bool bNewBuffer = ScratchBufferUInt.Allocate(Size, DebugName, RHICmdList, Allocation);
+	if (bNewBuffer && bClearNew)
+	{
+		RHICmdList.Transition(FRHITransitionInfo(Allocation.Buffer->UAV, ScratchBufferUInt.GetExpectedCurrentAccess(), ERHIAccess::UAVCompute));
+		RHICmdList.ClearUAVUint(Allocation.Buffer->UAV, FUintVector4(0, 0, 0, 0));
+		RHICmdList.Transition(FRHITransitionInfo(Allocation.Buffer->UAV, ERHIAccess::UAVCompute, ScratchBufferUInt.GetExpectedCurrentAccess()));
+	}
+	return Allocation;
+}
+
+//////////////////////////////////////////////////////////////////////////
+
+/** Behaves similar to FGPUScratchPad but for use with RWStructuredBuffers
+*/
+template<typename T>
+class NIAGARASHADER_API FGPUScratchPadStructured
+{
+public:
+	struct NIAGARASHADER_API FAllocation
+	{
+		FRWBufferStructured* Buffer = nullptr;
+		uint32 Offset = 0;
+
+		FORCEINLINE bool IsValid()const { return Buffer != nullptr; }
+	};
+
+	struct NIAGARASHADER_API FBuffer
+	{
+		FRWBufferStructured Buffer;
+		uint32 Used = 0;
+	};
+
+	FGPUScratchPadStructured(uint32 InBucketSize, const FString& InDebugName, EBufferUsageFlags InBufferUsageFlags = EBufferUsageFlags::Static | EBufferUsageFlags::UnorderedAccess)
+		: MinBucketSize(InBucketSize)
+		, DebugName(InDebugName)
+		, BufferUsageFlags(InBufferUsageFlags)
+	{
+	}
+
+	/** Allocates NumElements from the scratch pad. Will try to allocate from an existing buffer but will create a new one if needed. */
+	FAllocation Alloc(uint32 NumElements)
+	{
+		//Look for a buffer to fit this allocation.
+		FBuffer* ToUse = nullptr;
+		for (FBuffer& Buffer : Buffers)
+		{
+			uint32 NumBufferElements = Buffer.Buffer.NumBytes / sizeof(T);
+			if (NumElements + Buffer.Used < NumBufferElements)
+			{
+				ToUse = &Buffer;
+			}
+		}
+
+		if (ToUse == nullptr)
+		{
+			//Allocate a new buffer that will fit this allocation.
+
+			FBuffer* NewBuffer = new FBuffer();
+			Buffers.Add(NewBuffer);
+
+			//Round up the size to pow 2 to ensure a slowly growing user doesn't just keep allocating new buckets. Probably a better guess size here than power of 2.
+			uint32 NewBufferSize = FPlatformMath::RoundUpToPowerOfTwo(FMath::Max(MinBucketSize, NumElements));
+
+			NewBuffer->Buffer.Initialize(
+				*DebugName,
+				sizeof(T),
+				NewBufferSize,
+				BufferUsageFlags,
+				false /*bUseUavCounter*/,
+				false /*bAppendBuffer*/,
+				CurrentAccess);
+
+			ToUse = NewBuffer;
+			NumAllocatedBytes += NewBuffer->Buffer.NumBytes;
+		}
+
+		FAllocation Allocation;
+		Allocation.Buffer = &ToUse->Buffer;
+		Allocation.Offset = ToUse->Used;
+		ToUse->Used += NumElements;
+		NumUsedBytes += NumElements * sizeof(T);
+		return Allocation;
+	}
+
+	/** Resets the scratch pad and releases all buffers it holds. */
+	void Release()
+	{
+		Reset();
+		CurrentAccess = ERHIAccess::UAVCompute;
+		for (FBuffer& Buffer : Buffers)
+		{
+			Buffer.Buffer.Release();
+			Buffer.Used = 0;
+		}
+		Buffers.Reset();
+
+		NumAllocatedBytes = 0;
+	}
+
+	/** Resets the scratch pad, invalidating all previous allocations. Does not modify actual RWBuffer allocations. */
+	void Reset()
+	{
+		for (FBuffer& Buffer : Buffers)
+		{
+			Buffer.Used = 0;
+		}
+
+		NumUsedBytes = 0;
+	}
+
+	/** Transitions all buffers in the scratch pad. */
+	void Transition(FRHICommandList& RHICmdList, ERHIAccess InPreviousState, ERHIAccess InNewState, EResourceTransitionFlags InFlags = EResourceTransitionFlags::None)
+	{
+		//Store off current access state so we can place any new buffers in the same state.
+		CurrentAccess = InNewState;
+
+ 		TArray<FRHITransitionInfo, TMemStackAllocator<>> Transitions;
+ 		Transitions.Reserve(Buffers.Num());
+ 		for (FBuffer& Buffer : Buffers)
+ 		{
+ 			Transitions.Emplace(Buffer.Buffer.UAV, InPreviousState, InNewState, InFlags);
+ 		}
+
+ 		RHICmdList.Transition(Transitions);
+	}
+
+	/** Returns the total bytes allocated for all buffers in the scratch pad. */
+	FORCEINLINE uint32 AllocatedBytes()const
+	{
+		return NumAllocatedBytes;
+	}
+
+	/** Returns the number of used bytes allocated for all buffers in the scratch pad. */
+	FORCEINLINE uint32 UsedBytes()const
+	{
+		return NumUsedBytes;
+	}
+
+private:
+
+	/** Size of each FRWBuffer in the pool. */
+	uint32 MinBucketSize = 4096;
+	TIndirectArray<FBuffer> Buffers;
+	FString DebugName;
+	EBufferUsageFlags BufferUsageFlags = EBufferUsageFlags::None;
+	ERHIAccess CurrentAccess = ERHIAccess::UAVCompute;
+
+	uint32 NumAllocatedBytes = 0;
+	uint32 NumUsedBytes = 0;
+};
 
 //c++ mirror of the ray struct defined in NiagaraRayTracingCommon.ush
 struct FNiagaraRayData : public FBasicRayData
@@ -41,19 +435,58 @@ struct FNiagaraRayTracingResult
 	float WorldNormal[3];
 };
 
-class NIAGARASHADER_API FNiagaraRayTracingHelper
+/** Holds all information on a single RayTracing dispatch. */
+struct FNiagaraRayTraceDispatchInfo
 {
-	public:
-		FNiagaraRayTracingHelper(EShaderPlatform InShaderPlatform)
-		: ShaderPlatform(InShaderPlatform)
-		{}
+	/** Buffer allocation for ray requests to trace. */
+	FGPUScratchPadStructured<FNiagaraRayData>::FAllocation RayRequests;
+	/** Buffer allocation for writing trace results into. */
+	FGPUScratchPadStructured<FNiagaraRayTracingResult>::FAllocation RayTraceIntersections;
+	/** Buffer allocation for last frames trace results to read from in this frame's simulation. */
+	FGPUScratchPadStructured<FNiagaraRayTracingResult>::FAllocation LastFrameRayTraceIntersections;
+	/** Buffer allocation for ray trace counts. Accumulated in simulations shaders and used as dispatch args for ray tracing shaders. */
+	FGPUScratchPad::FAllocation RayCounts;
+	/** Total possible rays to trace. This may be significantly higher than the actual rays request we accumulate into RayCounts. */
+	uint32 MaxRays = 0;
+	/** Max number or times rays can be re-traced in a shader if they hit something that is invalid/filtered. */
+	uint32 MaxRetraces = 0;
 
+	FORCEINLINE bool IsValid(){ return MaxRays > 0; }
+	FORCEINLINE void Reset()
+	{
+		RayRequests = FGPUScratchPadStructured<FNiagaraRayData>::FAllocation();
+		RayTraceIntersections = FGPUScratchPadStructured<FNiagaraRayTracingResult>::FAllocation();
+		RayCounts = FGPUScratchPad::FAllocation();
+		MaxRays = 0;
+		MaxRetraces = 0;
+	}
+};
+
+class NIAGARASHADER_API FNiagaraRayTracingHelper
+{
+public:
+		FNiagaraRayTracingHelper(EShaderPlatform InShaderPlatform);
 		FNiagaraRayTracingHelper() = delete;
+		~FNiagaraRayTracingHelper();
 
 		void Reset();
 		void BuildRayTracingSceneInfo(FRHICommandList& RHICmdList, TConstArrayView<FViewInfo> Views);
-		void IssueRayTraces(FRHICommandList& RHICmdList, FScene* Scene, const FIntPoint& RayTraceCounts, uint32 MaxRetraces, FRHIShaderResourceView* RayTraceRequests, FRWBuffer* IndirectArgsBuffer, uint32 IndirectArgsOffset, FRHIUnorderedAccessView* RayTraceResults) const;
+
+		void IssueRayTraces(FRHICommandList& RHICmdList, FScene* Scene);
+		void IssueRayTraces(FRHICommandList& RHICmdList, FScene* Scene, const FIntPoint& RayTraceCounts, uint32 MaxRetraces, FRHIShaderResourceView* RayTraceRequests, FRWBuffer* IndirectArgsBuffer, uint32 IndirectArgsOffset, FRHIUnorderedAccessView* RayTraceResults);
 		bool IsValid() const;
+		
+		void BeginFrame(FRHICommandList& RHICmdList);
+		void EndFrame(FRHICommandList& RHICmdList, FScene* Scene);
+
+		/** Accumulates ray requests from user DIs into single dispatches per DI. */
+		void AddToDispatch(FNiagaraDataInterfaceProxy* DispatchKey, uint32 MaxRays, int32 MaxRetraces);
+		
+		/** Returns the final buffers for each DI for use in simulations and RT dispatches. This will allocated buffers where needed. */
+		const FNiagaraRayTraceDispatchInfo& GetDispatch(FNiagaraDataInterfaceProxy* DispatchKey, FRHICommandList& RHICmdList);
+
+		/** Returns a dummy dispatch to allow shaders to still bind to valid buffers. */
+		const FNiagaraRayTraceDispatchInfo& GetDummyDispatch(FRHICommandList& RHICmdList);
 
 		/** Adds a primitive to a collision group. This data is sent to the GPU on the start of the next frame. */
 		void SetPrimitiveCollisionGroup(FPrimitiveSceneInfo& Primitive, uint32 CollisionGroup);
@@ -62,6 +495,11 @@ class NIAGARASHADER_API FNiagaraRayTracingHelper
 		void UpdateCollisionGroupMap(FRHICommandList& RHICmdList, FScene* Scene, ERHIFeatureLevel::Type FeatureLevel);
 
 		void RefreshPrimitiveInstanceData();
+
+		FGPUScratchPadStructured<FNiagaraRayData>& GetRayRequests(){ return RayRequests; }
+		FGPUScratchPadStructured<FNiagaraRayTracingResult> GetRayTraceIntersections() { return RayTraceIntersections; }
+		FGPUScratchPad GetRayTraceCounts() { return RayTraceCounts; }
+
 	private:
 		const EShaderPlatform ShaderPlatform;
 
@@ -84,6 +522,22 @@ class NIAGARASHADER_API FNiagaraRayTracingHelper
 		FRWBuffer HashToCollisionGroups;
 
 		bool bCollisionGroupMapDirty = true;
+
+		/** Scratch buffer holding all ray trace requests for all simulations in the scene. */
+		FGPUScratchPadStructured<FNiagaraRayData> RayRequests;
+		/** Scratch buffer holding all ray trace results for all simulations in the scene. */
+		FGPUScratchPadStructured<FNiagaraRayTracingResult> RayTraceIntersections;
+		/** Scratch buffer holding all accumulated ray request counts for all simulations in the scene. */
+		FGPUScratchPad RayTraceCounts;
+
+		/** Map of DI Proxy to Dispatch Info. Each DI will have a single dispatch for all it's instances. */
+		TMap<FNiagaraDataInterfaceProxy*, FNiagaraRayTraceDispatchInfo> Dispatches;
+
+		/** Last frame's DI proxy map. Allows use to retrieve the buffer allocations to read last frames trace results. TODO: Improve. Ideally don't need a map copy here. */
+		TMap<FNiagaraDataInterfaceProxy*, FNiagaraRayTraceDispatchInfo> PreviousFrameDispatches;
+
+		/** In places where ray tracing is disabled etc we use a dummy dispatch allocation set so the simulations can bind to valid buffers. */
+		FNiagaraRayTraceDispatchInfo DummyDispatch;
 };
 
 #endif
diff --git a/Engine/Shaders/Private/RayTracing/RayTracingCommon.ush b/Engine/Shaders/Private/RayTracing/RayTracingCommon.ush
--- a/Engine/Shaders/Private/RayTracing/RayTracingCommon.ush
+++ b/Engine/Shaders/Private/RayTracing/RayTracingCommon.ush
@@ -121,8 +121,6 @@ struct FDefaultPayload : FIntersectionPayload
 
 struct FVFXTracePayload : FIntersectionPayload
 {
-	void SetWorldNormal(float3 InNormal) { WorldNormal = InNormal; }
-
 	float3 WorldPosition;
 	float3 WorldNormal;
 };
diff --git a/Engine/Shaders/Private/RayTracing/VFXTraceRay.ush b/Engine/Shaders/Private/RayTracing/VFXTraceRay.ush
--- a/Engine/Shaders/Private/RayTracing/VFXTraceRay.ush
+++ b/Engine/Shaders/Private/RayTracing/VFXTraceRay.ush
@@ -17,37 +17,43 @@
 #include "/Platform/Private/VFXTraceRay.ush"
 #else // OVERRIDE_VFXTRACERAY_USH
 
-
-float3 GetWorldNormalCommon(FVFXTracePayload Payload)
+void InitVFXTracePayloadFromHitShader(FDefaultAttributes Attributes, out FVFXTracePayload OutPayload)
 {
+	OutPayload.HitT = RayTCurrent();
+	OutPayload.PrimitiveIndex = GetInstanceUserData();
+	OutPayload.InstanceIndex = InstanceIndex() - GetBaseInstanceIndex();
+	OutPayload.Barycentrics = Attributes.Barycentrics;
+
+	//Reconstruct world position.
+	OutPayload.WorldPosition = WorldRayOrigin() + WorldRayDirection() * RayTCurrent();
+
+	//Reconstruct world normal.
 	FTriangleBaseAttributes Tri = LoadTriangleBaseAttributes(PrimitiveIndex());
 
-	float3 LocalPosition = (1.0f - Payload.Barycentrics.x - Payload.Barycentrics.y) * Tri.LocalPositions[0]
-		+ Payload.Barycentrics.x * Tri.LocalPositions[1]
-		+ Payload.Barycentrics.y * Tri.LocalPositions[2];
+	float3 LocalPosition = (1.0f - Attributes.Barycentrics.x - Attributes.Barycentrics.y) * Tri.LocalPositions[0]
+		+ Attributes.Barycentrics.x * Tri.LocalPositions[1]
+		+ Attributes.Barycentrics.y * Tri.LocalPositions[2];
 
 	float3 LocalEdges[2] = { Tri.LocalPositions[1] - Tri.LocalPositions[0], Tri.LocalPositions[2] - Tri.LocalPositions[0] };
 	float3 LocalNormal = cross(LocalEdges[1], LocalEdges[0]);
 	float3x3 InverseTranspose3x3 = transpose((float3x3)WorldToObject4x3());
 
-	return normalize(mul(LocalNormal, InverseTranspose3x3));
+	OutPayload.WorldNormal = normalize(mul(LocalNormal, InverseTranspose3x3));
 }
 
-void InitVFXTracePayloadFromHitShader(FDefaultAttributes Attributes, out FVFXTracePayload OutPayload)
+void VfxTraceRay(
+	RaytracingAccelerationStructure AccelerationStructure,
+	uint RayFlags,
+	uint InstanceInclusionMask,
+	uint RayContributionToHitGroupIndex,
+	uint MultiplierForGeometryContributionToHitGroupIndex,
+	uint MissShaderIndex,
+	RayDesc InRay,
+	inout FVFXTracePayload Payload)
 {
-	OutPayload.HitT = RayTCurrent();
-	OutPayload.PrimitiveIndex = GetInstanceUserData();
-	OutPayload.InstanceIndex = InstanceIndex() - GetBaseInstanceIndex();
-	OutPayload.Barycentrics = Attributes.Barycentrics;
-
-	OutPayload.WorldPosition = WorldRayOrigin() + WorldRayDirection() * RayTCurrent();
-
-	OutPayload.WorldNormal = GetWorldNormalCommon(OutPayload);
+	TraceRay(AccelerationStructure, RayFlags, InstanceInclusionMask, RayContributionToHitGroupIndex, MultiplierForGeometryContributionToHitGroupIndex, MissShaderIndex, InRay, Payload);
 }
 
-#define VFX_TRACE_RAY(TLAS, RayFlags, InstanceInclusionMask, RayContributionToHitGroupIndex, MultiplierForGeometryContributionToShaderIndex, MissShaderIndex, Ray, Payload)\
-		TraceRay(TLAS, RayFlags, InstanceInclusionMask, RayContributionToHitGroupIndex, MultiplierForGeometryContributionToShaderIndex, MissShaderIndex, Ray, Payload);
-
 #endif //#ifndef UE_TRACE_RAY
 
 #endif // UE_TRACE_RAY_USH_INCLUDED // Workarround for UE-66460
