Niagara - Graph state caching and UI optimizations to prevent compile hitches.
+ Cache the results of GetStackFunctionInputs to speed up refreshing module and input stack entries, and to speed up PrepareRapidIterationParameters.  Results are keyed using the inputs to the function and are invalidated based on the owning and called graph change ids.
+ When refreshing stack function inputs, don't refresh the the default input values or data which relies on the metadata if the called graph change id hasn't changed.
+ When refreshing the linked inputs and outputs for a module, skip the update if the called graph hasn't changed.
+ When handling message manager refreshes in stack modules items and stack functions inputs, skip the refresh if the previous messages and new messages are both empty.
+ Remove the top level compile refreshes from the stack view models, and instead refresh the script groups when the script they depend on has been compiled.

diff --git a/Engine/Plugins/FX/Niagara/Source/NiagaraEditor/Private/NiagaraEditorModule.cpp b/Engine/Plugins/FX/Niagara/Source/NiagaraEditor/Private/NiagaraEditorModule.cpp
--- a/Engine/Plugins/FX/Niagara/Source/NiagaraEditor/Private/NiagaraEditorModule.cpp
+++ b/Engine/Plugins/FX/Niagara/Source/NiagaraEditor/Private/NiagaraEditorModule.cpp
@@ -146,6 +146,8 @@
 
 #include "NiagaraPerfBaseline.h"
 
+#include "NiagaraGraphDataCache.h"
+
 IMPLEMENT_MODULE( FNiagaraEditorModule, NiagaraEditor );
 
 #define LOCTEXT_NAMESPACE "NiagaraEditorModule"
@@ -1101,6 +1103,8 @@ void FNiagaraEditorModule::StartupModule()
 #if NIAGARA_PERF_BASELINES
 	UNiagaraEffectType::OnGeneratePerfBaselines().BindRaw(this, &FNiagaraEditorModule::GeneratePerfBaselines);
 #endif
+
+	GraphDataCache = MakeUnique<FNiagaraGraphDataCache>();
 }
 
 void FNiagaraEditorModule::ShutdownModule()
@@ -1215,6 +1219,8 @@ void FNiagaraEditorModule::ShutdownModule()
 #if NIAGARA_PERF_BASELINES
 	UNiagaraEffectType::OnGeneratePerfBaselines().Unbind();
 #endif
+
+	GraphDataCache.Reset();
 }
 
 void FNiagaraEditorModule::OnPostEngineInit()
diff --git a/Engine/Plugins/FX/Niagara/Source/NiagaraEditor/Private/NiagaraEditorUtilities.cpp b/Engine/Plugins/FX/Niagara/Source/NiagaraEditor/Private/NiagaraEditorUtilities.cpp
--- a/Engine/Plugins/FX/Niagara/Source/NiagaraEditor/Private/NiagaraEditorUtilities.cpp
+++ b/Engine/Plugins/FX/Niagara/Source/NiagaraEditor/Private/NiagaraEditorUtilities.cpp
@@ -1139,7 +1139,7 @@ void FNiagaraEditorUtilities::SetStaticSwitchConstants(UNiagaraGraph* Graph, TAr
 		UNiagaraNodeFunctionCall* FunctionNode = Cast<UNiagaraNodeFunctionCall>(Node);
 		if (FunctionNode)
 		{
-			FunctionNode->DebugState = FunctionNode->bInheritDebugStatus? ConstantResolver.GetDebugState() : ENiagaraFunctionDebugState::NoDebug;
+			FunctionNode->DebugState = FunctionNode->bInheritDebugStatus? ConstantResolver.CalculateDebugState() : ENiagaraFunctionDebugState::NoDebug;
 
 			if (FunctionNode->PropagatedStaticSwitchParameters.Num() > 0)
 			{
diff --git a/Engine/Plugins/FX/Niagara/Source/NiagaraEditor/Private/NiagaraGraphDataCache.cpp b/Engine/Plugins/FX/Niagara/Source/NiagaraEditor/Private/NiagaraGraphDataCache.cpp
new file mode 100644
--- /dev/null
+++ b/Engine/Plugins/FX/Niagara/Source/NiagaraEditor/Private/NiagaraGraphDataCache.cpp
@@ -0,0 +1,64 @@
+// Copyright Epic Games, Inc. All Rights Reserved.
+
+#include "NiagaraGraphDataCache.h"
+
+#include "NiagaraNodeFunctionCall.h"
+
+void FNiagaraGraphDataCache::GetStackFunctionInputPins(UNiagaraNodeFunctionCall& FunctionCallNode, TArray<const UEdGraphPin*>& OutInputPins, TSet<const UEdGraphPin*>& OutHiddenPins, FCompileConstantResolver ConstantResolver, FNiagaraStackGraphUtilities::ENiagaraGetStackFunctionInputPinsOptions Options, bool bIgnoreDisabled)
+{
+	FGetStackFunctionInputPinsKey Key(FunctionCallNode, ConstantResolver, Options, bIgnoreDisabled);
+
+	FCachedStackFunctionInputPins& CachedResult = GetStackFunctionInputPinsCache.FindOrAdd(Key);
+	if (CachedResult.IsValid())
+	{
+		OutInputPins = CachedResult.InputPins;
+		OutHiddenPins = CachedResult.HiddenPins;
+		return;
+	}
+
+	FNiagaraStackGraphUtilities::GetStackFunctionInputPinsWithoutCache(FunctionCallNode, OutInputPins, OutHiddenPins, ConstantResolver, Options, bIgnoreDisabled);
+
+	CachedResult.SourceGraphWeak = FunctionCallNode.GetNiagaraGraph();
+	CachedResult.LastSourceGraphChangeId = FunctionCallNode.GetNiagaraGraph()->GetChangeID();
+	CachedResult.FunctionCallNodeWeak = &FunctionCallNode;
+	CachedResult.CalledGraphWeak = FunctionCallNode.GetCalledGraph();
+	CachedResult.LastCalledGraphChangeId = FunctionCallNode.GetCalledGraph() != nullptr ? FunctionCallNode.GetCalledGraph()->GetChangeID() : FGuid();
+	CachedResult.InputPins = OutInputPins;
+	CachedResult.HiddenPins = OutHiddenPins;
+}
+
+FNiagaraGraphDataCache::FGetStackFunctionInputPinsKey::FGetStackFunctionInputPinsKey(UNiagaraNodeFunctionCall& InFunctionCallNode, const FCompileConstantResolver& InConstantResolver, FNiagaraStackGraphUtilities::ENiagaraGetStackFunctionInputPinsOptions InOptions, bool bInIgnoreDisabled)
+{
+	FunctionCallNodeKey = FObjectKey(&InFunctionCallNode);
+
+	ConstantResolverSystemKey = FObjectKey(InConstantResolver.GetSystem());
+	ConstantResolverEmitterKey = FObjectKey(InConstantResolver.GetEmitter());
+	ConstantResolverTranslator = InConstantResolver.GetTranslator();
+	ConstantResolverUsage = InConstantResolver.GetUsage();
+	ConstantResolverDebugState = InConstantResolver.GetDebugState();
+
+	Options = InOptions;
+	bIgnoreDisabled = bInIgnoreDisabled;
+
+	Hash = GetTypeHash(FunctionCallNodeKey);
+	Hash = HashCombine(Hash, GetTypeHash(ConstantResolverSystemKey));
+	Hash = HashCombine(Hash, GetTypeHash(ConstantResolverEmitterKey));
+	Hash = HashCombine(Hash, GetTypeHash(ConstantResolverTranslator));
+	Hash = HashCombine(Hash, GetTypeHash(ConstantResolverUsage));
+	Hash = HashCombine(Hash, GetTypeHash(ConstantResolverDebugState));
+
+	Hash = HashCombine(Hash, GetTypeHash(Options));
+	Hash = HashCombine(Hash, GetTypeHash(bIgnoreDisabled));
+}
+
+bool FNiagaraGraphDataCache::FCachedStackFunctionInputPins::IsValid() const
+{
+	return
+		SourceGraphWeak.IsValid() &&
+		FunctionCallNodeWeak.IsValid() &&
+		CalledGraphWeak.IsValid() &&
+
+		SourceGraphWeak->GetChangeID() == LastSourceGraphChangeId &&
+		FunctionCallNodeWeak->GetCalledGraph() == CalledGraphWeak.Get() &&
+		(CalledGraphWeak.IsValid() == false || CalledGraphWeak->GetChangeID() == LastCalledGraphChangeId);
+}
diff --git a/Engine/Plugins/FX/Niagara/Source/NiagaraEditor/Private/NiagaraParameterMapHistory.cpp b/Engine/Plugins/FX/Niagara/Source/NiagaraEditor/Private/NiagaraParameterMapHistory.cpp
--- a/Engine/Plugins/FX/Niagara/Source/NiagaraEditor/Private/NiagaraParameterMapHistory.cpp
+++ b/Engine/Plugins/FX/Niagara/Source/NiagaraEditor/Private/NiagaraParameterMapHistory.cpp
@@ -1552,7 +1552,7 @@ bool FCompileConstantResolver::ResolveConstant(FNiagaraVariable& OutConstant) co
 	if (OutConstant == FNiagaraVariable(FNiagaraTypeDefinition::GetFunctionDebugStateEnum(), TEXT("Function.DebugState")))
 	{
 		FNiagaraInt32 EnumValue;
-		EnumValue.Value = (uint8)GetDebugState();
+		EnumValue.Value = (uint8)CalculateDebugState();
 		OutConstant.SetValue(EnumValue);
 		return true;
 	}
@@ -1625,7 +1625,7 @@ bool FCompileConstantResolver::ResolveConstant(FNiagaraVariable& OutConstant) co
 	return false;
 }
 
-ENiagaraFunctionDebugState FCompileConstantResolver::GetDebugState() const
+ENiagaraFunctionDebugState FCompileConstantResolver::CalculateDebugState() const
 {
 	const UNiagaraSystem* CurrentSystem =  System? System : (Emitter ? Cast<UNiagaraSystem>(Emitter->GetOuter()) : nullptr);
 	bool bDisableDebug = CurrentSystem ? CurrentSystem->bDisableAllDebugSwitches : false;
diff --git a/Engine/Plugins/FX/Niagara/Source/NiagaraEditor/Private/ViewModels/Stack/NiagaraStackFunctionInput.cpp b/Engine/Plugins/FX/Niagara/Source/NiagaraEditor/Private/ViewModels/Stack/NiagaraStackFunctionInput.cpp
--- a/Engine/Plugins/FX/Niagara/Source/NiagaraEditor/Private/ViewModels/Stack/NiagaraStackFunctionInput.cpp
+++ b/Engine/Plugins/FX/Niagara/Source/NiagaraEditor/Private/ViewModels/Stack/NiagaraStackFunctionInput.cpp
@@ -372,6 +372,12 @@ TArray<UNiagaraStackFunctionInput*> UNiagaraStackFunctionInput::GetChildInputs()
 	return ChildInputs;
 }
 
+void UNiagaraStackFunctionInput::GetCurrentChangeIds(FGuid& OutOwningGraphChangeId, FGuid& OutFunctionGraphChangeId) const
+{
+	OutOwningGraphChangeId = OwningFunctionCallNode->GetNiagaraGraph()->GetChangeID();
+	OutFunctionGraphChangeId = OwningFunctionCallNode->GetCalledGraph() != nullptr ? OwningFunctionCallNode->GetCalledGraph()->GetChangeID() : FGuid();
+}
+
 FText UNiagaraStackFunctionInput::GetCollapsedStateText() const
 {
 	if (IsFinalized())
@@ -756,6 +762,12 @@ void UNiagaraStackFunctionInput::RefreshChildrenInternal(const TArray<UNiagaraSt
 	{
 		NewIssues.Append(MessageManagerIssues);
 	}
+
+	FGuid CurrentOwningGraphChangeId;
+	FGuid CurrentFunctionGraphChangeId;
+	GetCurrentChangeIds(CurrentOwningGraphChangeId, CurrentFunctionGraphChangeId);
+	LastOwningGraphChangeId = CurrentOwningGraphChangeId;
+	LastFunctionGraphChangeId = CurrentFunctionGraphChangeId;
 }
 
 class UNiagaraStackFunctionInputUtilities
@@ -877,10 +889,17 @@ void UNiagaraStackFunctionInput::RefreshValues()
 		return;
 	}
 
-	// First collect the default values which are used to figure out if an input can be reset, and are used to
-	// determine the current displayed value.
-	DefaultInputValues = FInputValues();
-	UpdateValuesFromScriptDefaults(DefaultInputValues);
+	FGuid CurrentOwningGraphChangeId;
+	FGuid CurrentFunctionGraphChangeId;
+	GetCurrentChangeIds(CurrentOwningGraphChangeId, CurrentFunctionGraphChangeId);
+	if (LastOwningGraphChangeId.IsSet() == false || (CurrentOwningGraphChangeId != LastOwningGraphChangeId.GetValue()) ||
+		LastFunctionGraphChangeId.IsSet() == false || (CurrentFunctionGraphChangeId != LastFunctionGraphChangeId.GetValue()))
+	{
+		// First collect the default values which are used to figure out if an input can be reset, and are used to
+		// determine the current displayed value.
+		DefaultInputValues = FInputValues();
+		UpdateValuesFromScriptDefaults(DefaultInputValues);
+	}
 
 	FInputValues OldValues = InputValues;
 	InputValues = FInputValues();
@@ -964,7 +983,18 @@ void UNiagaraStackFunctionInput::RefreshValues()
 
 void UNiagaraStackFunctionInput::RefreshFromMetaData(TArray<FStackIssue>& NewIssues)
 {
+	FGuid CurrentOwningGraphChangeId;
+	FGuid CurrentFunctionGraphChangeId;
+	GetCurrentChangeIds(CurrentOwningGraphChangeId, CurrentFunctionGraphChangeId);
+	if (LastFunctionGraphChangeId.IsSet() && CurrentFunctionGraphChangeId == LastFunctionGraphChangeId)
+	{
+		// If the called function graph hasn't changed, then the metadata will also be the same and we can skip updating these values.
+		NewIssues.Append(InputMetaDataIssues);
+		return;
+	}
+
 	InputMetaData.Reset();
+	InputMetaDataIssues.Empty();
 	if (OwningFunctionCallNode->IsA<UNiagaraNodeAssignment>())
 	{
 		// Set variables nodes have no metadata, but if they're setting a defined constant see if there's metadata for that.
@@ -1007,7 +1037,7 @@ void UNiagaraStackFunctionInput::RefreshFromMetaData(TArray<FStackIssue>& NewIss
 
 		if (EditConditionError.IsEmpty() == false)
 		{
-			NewIssues.Add(FStackIssue(
+			InputMetaDataIssues.Add(FStackIssue(
 				EStackIssueSeverity::Info,
 				LOCTEXT("EditConditionErrorShort", "Edit condition error"),
 				FText::Format(LOCTEXT("EditConditionErrorLongFormat", "Edit condition failed to bind.  Function: {0} Input: {1} Message: {2}"), 
@@ -1023,7 +1053,7 @@ void UNiagaraStackFunctionInput::RefreshFromMetaData(TArray<FStackIssue>& NewIss
 
 		if (VisibleConditionError.IsEmpty() == false)
 		{
-			NewIssues.Add(FStackIssue(
+			InputMetaDataIssues.Add(FStackIssue(
 				EStackIssueSeverity::Info,
 				LOCTEXT("VisibleConditionErrorShort", "Visible condition error"),
 				FText::Format(LOCTEXT("VisibleConditionErrorLongFormat", "Visible condition failed to bind.  Function: {0} Input: {1} Message: {2}"),
@@ -1037,6 +1067,8 @@ void UNiagaraStackFunctionInput::RefreshFromMetaData(TArray<FStackIssue>& NewIss
 		bIsInlineEditConditionToggle = InputType == FNiagaraTypeDefinition::GetBoolDef() && 
 			InputMetaData->bInlineEditConditionToggle;
 	}
+
+	NewIssues.Append(InputMetaDataIssues);
 }
 
 FText UNiagaraStackFunctionInput::GetDisplayName() const
@@ -1760,18 +1792,22 @@ FNiagaraVariable UNiagaraStackFunctionInput::CreateRapidIterationVariable(const
 
 void UNiagaraStackFunctionInput::OnMessageManagerRefresh(const TArray<TSharedRef<const INiagaraMessage>>& NewMessages)
 {
-	MessageManagerIssues.Reset();
-	if (InputValues.DynamicNode.IsValid())
+	if(MessageManagerIssues.Num() != 0 || NewMessages.Num() != 0)
 	{
-		for (TSharedRef<const INiagaraMessage> Message : NewMessages)
+		MessageManagerIssues.Reset();
+		if (InputValues.DynamicNode.IsValid())
 		{
-			FStackIssue Issue = FNiagaraMessageUtilities::MessageToStackIssue(Message, GetStackEditorDataKey());
-			if (MessageManagerIssues.ContainsByPredicate([&Issue](const FStackIssue& NewIssue)
-				{ return NewIssue.GetUniqueIdentifier() == Issue.GetUniqueIdentifier(); }) == false)
+			for (TSharedRef<const INiagaraMessage> Message : NewMessages)
 			{
-				MessageManagerIssues.Add(Issue);	
+				FStackIssue Issue = FNiagaraMessageUtilities::MessageToStackIssue(Message, GetStackEditorDataKey());
+				if (MessageManagerIssues.ContainsByPredicate([&Issue](const FStackIssue& NewIssue)
+					{ return NewIssue.GetUniqueIdentifier() == Issue.GetUniqueIdentifier(); }) == false)
+				{
+					MessageManagerIssues.Add(Issue);	
+				}
 			}
 		}
+		RefreshChildren();
 	}
 }
 
diff --git a/Engine/Plugins/FX/Niagara/Source/NiagaraEditor/Private/ViewModels/Stack/NiagaraStackGraphUtilities.cpp b/Engine/Plugins/FX/Niagara/Source/NiagaraEditor/Private/ViewModels/Stack/NiagaraStackGraphUtilities.cpp
--- a/Engine/Plugins/FX/Niagara/Source/NiagaraEditor/Private/ViewModels/Stack/NiagaraStackGraphUtilities.cpp
+++ b/Engine/Plugins/FX/Niagara/Source/NiagaraEditor/Private/ViewModels/Stack/NiagaraStackGraphUtilities.cpp
@@ -39,6 +39,7 @@
 #include "NiagaraScriptVariable.h"
 #include "INiagaraEditorTypeUtilities.h"
 #include "NiagaraEditorData.h"
+#include "NiagaraGraphDataCache.h"
 
 
 DECLARE_CYCLE_STAT(TEXT("Niagara - StackGraphUtilities - RelayoutGraph"), STAT_NiagaraEditor_StackGraphUtilities_RelayoutGraph, STATGROUP_NiagaraEditor);
@@ -753,6 +754,11 @@ void ExtractInputPinsFromHistory(FNiagaraParameterMapHistory& History, UEdGraph*
 }
 
 void FNiagaraStackGraphUtilities::GetStackFunctionInputPins(UNiagaraNodeFunctionCall& FunctionCallNode, TArray<const UEdGraphPin*>& OutInputPins, TSet<const UEdGraphPin*>& OutHiddenPins, FCompileConstantResolver ConstantResolver, ENiagaraGetStackFunctionInputPinsOptions Options, bool bIgnoreDisabled)
+{
+	FNiagaraEditorModule::Get().GetGraphDataCache().GetStackFunctionInputPins(FunctionCallNode, OutInputPins, OutHiddenPins, ConstantResolver, Options, bIgnoreDisabled);
+}
+
+void FNiagaraStackGraphUtilities::GetStackFunctionInputPinsWithoutCache(UNiagaraNodeFunctionCall& FunctionCallNode, TArray<const UEdGraphPin*>& OutInputPins, TSet<const UEdGraphPin*>& OutHiddenPins, FCompileConstantResolver ConstantResolver, ENiagaraGetStackFunctionInputPinsOptions Options, bool bIgnoreDisabled)
 {
 	FNiagaraParameterMapHistoryBuilder Builder;
 	Builder.SetIgnoreDisabled(bIgnoreDisabled);
diff --git a/Engine/Plugins/FX/Niagara/Source/NiagaraEditor/Private/ViewModels/Stack/NiagaraStackModuleItem.cpp b/Engine/Plugins/FX/Niagara/Source/NiagaraEditor/Private/ViewModels/Stack/NiagaraStackModuleItem.cpp
--- a/Engine/Plugins/FX/Niagara/Source/NiagaraEditor/Private/ViewModels/Stack/NiagaraStackModuleItem.cpp
+++ b/Engine/Plugins/FX/Niagara/Source/NiagaraEditor/Private/ViewModels/Stack/NiagaraStackModuleItem.cpp
@@ -1040,19 +1040,22 @@ void UNiagaraStackModuleItem::RefreshIsEnabled()
 
 void UNiagaraStackModuleItem::OnMessageManagerRefresh(const TArray<TSharedRef<const INiagaraMessage>>& NewMessages)
 {
-	MessageManagerIssues.Reset();
-	for (TSharedRef<const INiagaraMessage> Message : NewMessages)
+	if (MessageManagerIssues.Num() != 0 || NewMessages.Num() != 0)
 	{
-		// Sometimes compile errors with the same info are generated, so guard against duplicates here.
-		FStackIssue Issue = FNiagaraMessageUtilities::MessageToStackIssue(Message, GetStackEditorDataKey());
-		if (MessageManagerIssues.ContainsByPredicate([&Issue](const FStackIssue& NewIssue)
-			{ return NewIssue.GetUniqueIdentifier() == Issue.GetUniqueIdentifier(); }) == false)
+		MessageManagerIssues.Reset();
+		for (TSharedRef<const INiagaraMessage> Message : NewMessages)
 		{
-			MessageManagerIssues.Add(Issue);
+			// Sometimes compile errors with the same info are generated, so guard against duplicates here.
+			FStackIssue Issue = FNiagaraMessageUtilities::MessageToStackIssue(Message, GetStackEditorDataKey());
+			if (MessageManagerIssues.ContainsByPredicate([&Issue](const FStackIssue& NewIssue)
+				{ return NewIssue.GetUniqueIdentifier() == Issue.GetUniqueIdentifier(); }) == false)
+			{
+				MessageManagerIssues.Add(Issue);
+			}
 		}
-	}
 
-	RefreshChildren();
+		RefreshChildren();
+	}
 }
 
 bool UNiagaraStackModuleItem::CanMoveAndDelete() const
diff --git a/Engine/Plugins/FX/Niagara/Source/NiagaraEditor/Private/ViewModels/Stack/NiagaraStackModuleItemLinkedInputCollection.cpp b/Engine/Plugins/FX/Niagara/Source/NiagaraEditor/Private/ViewModels/Stack/NiagaraStackModuleItemLinkedInputCollection.cpp
--- a/Engine/Plugins/FX/Niagara/Source/NiagaraEditor/Private/ViewModels/Stack/NiagaraStackModuleItemLinkedInputCollection.cpp
+++ b/Engine/Plugins/FX/Niagara/Source/NiagaraEditor/Private/ViewModels/Stack/NiagaraStackModuleItemLinkedInputCollection.cpp
@@ -59,6 +59,17 @@ bool UNiagaraStackModuleItemLinkedInputCollection::GetShouldShowInStack() const
 
 void UNiagaraStackModuleItemLinkedInputCollection::RefreshChildrenInternal(const TArray<UNiagaraStackEntry*>& CurrentChildren, TArray<UNiagaraStackEntry*>& NewChildren, TArray<FStackIssue>& NewIssues)
 {
+	FGuid CurrentOwningGraphChangeId;
+	FGuid CurrentFunctionGraphChangeId;
+	GetCurrentChangeIds(CurrentOwningGraphChangeId, CurrentFunctionGraphChangeId);
+	if (LastOwningGraphChangeId.IsSet() && CurrentOwningGraphChangeId == LastOwningGraphChangeId &&
+		LastFunctionGraphChangeId.IsSet() && CurrentFunctionGraphChangeId == LastFunctionGraphChangeId)
+	{
+		// If the owning and called graph haven't changed, then the child inputs haven't changed either.
+		NewChildren.Append(CurrentChildren);
+		return;
+	}
+
 	UEdGraphPin* OutputParameterMapPin = FNiagaraStackGraphUtilities::GetParameterMapOutputPin(*FunctionCallNode);
 	if (ensureMsgf(OutputParameterMapPin != nullptr, TEXT("Invalid Stack Graph - Function call node has no output pin.")))
 	{
@@ -110,4 +121,13 @@ void UNiagaraStackModuleItemLinkedInputCollection::RefreshChildrenInternal(const
 
 		NewChildren.Sort(SortNewChildrenPred);
 	}
+
+	LastOwningGraphChangeId = CurrentOwningGraphChangeId;
+	LastFunctionGraphChangeId = CurrentFunctionGraphChangeId;
+}
+
+void UNiagaraStackModuleItemLinkedInputCollection::GetCurrentChangeIds(FGuid& OutOwningGraphChangeId, FGuid& OutFunctionGraphChangeId) const
+{
+	OutOwningGraphChangeId = FunctionCallNode->GetNiagaraGraph()->GetChangeID();
+	OutFunctionGraphChangeId = FunctionCallNode->GetCalledGraph() != nullptr ? FunctionCallNode->GetCalledGraph()->GetChangeID() : FGuid();
 }
diff --git a/Engine/Plugins/FX/Niagara/Source/NiagaraEditor/Private/ViewModels/Stack/NiagaraStackModuleItemOutputCollection.cpp b/Engine/Plugins/FX/Niagara/Source/NiagaraEditor/Private/ViewModels/Stack/NiagaraStackModuleItemOutputCollection.cpp
--- a/Engine/Plugins/FX/Niagara/Source/NiagaraEditor/Private/ViewModels/Stack/NiagaraStackModuleItemOutputCollection.cpp
+++ b/Engine/Plugins/FX/Niagara/Source/NiagaraEditor/Private/ViewModels/Stack/NiagaraStackModuleItemOutputCollection.cpp
@@ -48,6 +48,17 @@ UNiagaraStackEntry::EStackRowStyle UNiagaraStackModuleItemOutputCollection::GetS
 
 void UNiagaraStackModuleItemOutputCollection::RefreshChildrenInternal(const TArray<UNiagaraStackEntry*>& CurrentChildren, TArray<UNiagaraStackEntry*>& NewChildren, TArray<FStackIssue>& NewIssues)
 {
+	FGuid CurrentOwningGraphChangeId;
+	FGuid CurrentFunctionGraphChangeId;
+	GetCurrentChangeIds(CurrentOwningGraphChangeId, CurrentFunctionGraphChangeId);
+	if (LastOwningGraphChangeId.IsSet() && CurrentOwningGraphChangeId == LastOwningGraphChangeId &&
+		LastFunctionGraphChangeId.IsSet() && CurrentFunctionGraphChangeId == LastFunctionGraphChangeId)
+	{
+		// If the owning and called graph haven't changed, then the child outputs haven't changed either.
+		NewChildren.Append(CurrentChildren);
+		return;
+	}
+
 	UEdGraphPin* OutputParameterMapPin = FNiagaraStackGraphUtilities::GetParameterMapOutputPin(*FunctionCallNode);
 	if (ensureMsgf(OutputParameterMapPin != nullptr, TEXT("Invalid Stack Graph - Function call node has no output pin.")))
 	{
@@ -84,4 +95,13 @@ void UNiagaraStackModuleItemOutputCollection::RefreshChildrenInternal(const TArr
 
 		NewChildren.Sort(SortNewChildrenPred);
 	}
+
+	LastOwningGraphChangeId = CurrentOwningGraphChangeId;
+	LastFunctionGraphChangeId = CurrentFunctionGraphChangeId;
+}
+
+void UNiagaraStackModuleItemOutputCollection::GetCurrentChangeIds(FGuid& OutOwningGraphChangeId, FGuid& OutFunctionGraphChangeId) const
+{
+	OutOwningGraphChangeId = FunctionCallNode->GetNiagaraGraph()->GetChangeID();
+	OutFunctionGraphChangeId = FunctionCallNode->GetCalledGraph() != nullptr ? FunctionCallNode->GetCalledGraph()->GetChangeID() : FGuid();
 }
diff --git a/Engine/Plugins/FX/Niagara/Source/NiagaraEditor/Private/ViewModels/Stack/NiagaraStackScriptItemGroup.cpp b/Engine/Plugins/FX/Niagara/Source/NiagaraEditor/Private/ViewModels/Stack/NiagaraStackScriptItemGroup.cpp
--- a/Engine/Plugins/FX/Niagara/Source/NiagaraEditor/Private/ViewModels/Stack/NiagaraStackScriptItemGroup.cpp
+++ b/Engine/Plugins/FX/Niagara/Source/NiagaraEditor/Private/ViewModels/Stack/NiagaraStackScriptItemGroup.cpp
@@ -381,6 +381,26 @@ void UNiagaraStackScriptItemGroup::Initialize(
 	ScriptGraph = InScriptViewModel->GetGraphViewModel()->GetGraph();
 	ScriptGraph->AddOnGraphChangedHandler(
 		FOnGraphChanged::FDelegate::CreateUObject(this, &UNiagaraStackScriptItemGroup::OnScriptGraphChanged));
+
+	if (ScriptUsage == ENiagaraScriptUsage::SystemSpawnScript || ScriptUsage == ENiagaraScriptUsage::EmitterSpawnScript)
+	{
+		OwningSystemWeak = &GetSystemViewModel()->GetSystem();
+		OwningSystemWeak->GetSystemSpawnScript()->OnVMScriptCompiled().AddUObject(this, &UNiagaraStackScriptItemGroup::OnSystemScriptCompiled);
+	}
+	else if (ScriptUsage == ENiagaraScriptUsage::SystemUpdateScript || ScriptUsage == ENiagaraScriptUsage::EmitterUpdateScript)
+	{
+		OwningSystemWeak = &GetSystemViewModel()->GetSystem();
+		OwningSystemWeak->GetSystemUpdateScript()->OnVMScriptCompiled().AddUObject(this, &UNiagaraStackScriptItemGroup::OnSystemScriptCompiled);
+	}
+	else if (GetEmitterViewModel().IsValid())
+	{
+		OwningParticleScriptWeak = GetEmitterViewModel()->GetEmitter()->GetScript(ScriptUsage, ScriptUsageId);
+		if (OwningParticleScriptWeak.IsValid())
+		{
+			OwningParticleScriptWeak->OnVMScriptCompiled().AddUObject(this, &UNiagaraStackScriptItemGroup::OnParticleScriptCompiled);
+			OwningParticleScriptWeak->OnGPUScriptCompiled().AddUObject(this, &UNiagaraStackScriptItemGroup::OnParticleScriptCompiled);
+		}
+	}
 }
 
 UNiagaraNodeOutput* UNiagaraStackScriptItemGroup::GetScriptOutputNode() const
@@ -411,6 +431,21 @@ void UNiagaraStackScriptItemGroup::FinalizeInternal()
 	{
 		ScriptViewModel.Reset();
 	}
+
+	if ((ScriptUsage == ENiagaraScriptUsage::SystemSpawnScript || ScriptUsage == ENiagaraScriptUsage::EmitterSpawnScript) && OwningSystemWeak.IsValid())
+	{
+		OwningSystemWeak->GetSystemSpawnScript()->OnVMScriptCompiled().RemoveAll(this);
+	}
+	else if ((ScriptUsage == ENiagaraScriptUsage::SystemUpdateScript || ScriptUsage == ENiagaraScriptUsage::EmitterUpdateScript) && OwningSystemWeak.IsValid())
+	{
+		OwningSystemWeak->GetSystemUpdateScript()->OnVMScriptCompiled().RemoveAll(this);
+	}
+	else if (OwningParticleScriptWeak.IsValid())
+	{
+		OwningParticleScriptWeak->OnVMScriptCompiled().RemoveAll(this);
+		OwningParticleScriptWeak->OnGPUScriptCompiled().RemoveAll(this);
+	}
+
 	Super::FinalizeInternal();
 }
 
@@ -1149,6 +1184,22 @@ void UNiagaraStackScriptItemGroup::OnScriptGraphChanged(const struct FEdGraphEdi
 	}
 }
 
+void UNiagaraStackScriptItemGroup::OnSystemScriptCompiled(UNiagaraScript* InScript, const FGuid& ScriptVersion)
+{
+	if (IsFinalized() == false)
+	{
+		RefreshChildren();
+	}
+}
+
+void UNiagaraStackScriptItemGroup::OnParticleScriptCompiled(UNiagaraScript* InScript, const FGuid& ScriptVersion)
+{
+	if (IsFinalized() == false)
+	{
+		RefreshChildren();
+	}
+}
+
 void GatherRenamedModuleOutputs(
 	UNiagaraStackScriptItemGroup* InOwnerEntry,
 	TArray<TPair<const UNiagaraClipboardFunction*, UNiagaraNodeFunctionCall*>> InClipboardFunctionAndNodeFunctionPairs,
diff --git a/Engine/Plugins/FX/Niagara/Source/NiagaraEditor/Private/ViewModels/Stack/NiagaraStackViewModel.cpp b/Engine/Plugins/FX/Niagara/Source/NiagaraEditor/Private/ViewModels/Stack/NiagaraStackViewModel.cpp
--- a/Engine/Plugins/FX/Niagara/Source/NiagaraEditor/Private/ViewModels/Stack/NiagaraStackViewModel.cpp
+++ b/Engine/Plugins/FX/Niagara/Source/NiagaraEditor/Private/ViewModels/Stack/NiagaraStackViewModel.cpp
@@ -89,10 +89,8 @@ void UNiagaraStackViewModel::InitializeWithViewModels(TSharedPtr<FNiagaraSystemV
 	{
 		if (EmitterViewModel.IsValid())
 		{
-			EmitterViewModel->OnScriptCompiled().AddUObject(this, &UNiagaraStackViewModel::OnEmitterCompiled);
 			EmitterViewModel->OnParentRemoved().AddUObject(this, &UNiagaraStackViewModel::EmitterParentRemoved);
 		}
-		SystemViewModelPinned->OnSystemCompiled().AddUObject(this, &UNiagaraStackViewModel::OnSystemCompiled);
 		
 		UNiagaraStackRoot* StackRoot = NewObject<UNiagaraStackRoot>(this);
 		UNiagaraStackEntry::FRequiredEntryData RequiredEntryData(SystemViewModelPinned.ToSharedRef(), EmitterViewModel,
@@ -153,13 +151,11 @@ void UNiagaraStackViewModel::Reset()
 
 	if (EmitterHandleViewModel.IsValid())
 	{
-		EmitterHandleViewModel.Pin()->GetEmitterViewModel()->OnScriptCompiled().RemoveAll(this);
 		EmitterHandleViewModel.Reset();
 	}
 
 	if (SystemViewModel.IsValid())
 	{
-		SystemViewModel.Pin()->OnSystemCompiled().RemoveAll(this);
 		SystemViewModel.Reset();
 	}
 
@@ -347,18 +343,6 @@ void UNiagaraStackViewModel::GetPathForEntry(UNiagaraStackEntry* Entry, TArray<U
 	GeneratePathForEntry(RootEntry, Entry, TArray<UNiagaraStackEntry*>(), EntryPath);
 }
 
-void UNiagaraStackViewModel::OnSystemCompiled()
-{
-	// Queue a refresh for the next tick because forcing a refresh now can cause entries to be finalized while they're still being used.
-	bRefreshPending = true;
-}
-
-void UNiagaraStackViewModel::OnEmitterCompiled(UNiagaraScript*, const FGuid&)
-{
-	// Queue a refresh for the next tick because forcing a refresh now can cause entries to be finalized while they're still being used.
-	bRefreshPending = true;
-}
-
 void UNiagaraStackViewModel::EmitterParentRemoved()
 {
 	RootEntry->RefreshChildren();
diff --git a/Engine/Plugins/FX/Niagara/Source/NiagaraEditor/Public/NiagaraEditorModule.h b/Engine/Plugins/FX/Niagara/Source/NiagaraEditor/Public/NiagaraEditorModule.h
--- a/Engine/Plugins/FX/Niagara/Source/NiagaraEditor/Public/NiagaraEditorModule.h
+++ b/Engine/Plugins/FX/Niagara/Source/NiagaraEditor/Public/NiagaraEditorModule.h
@@ -40,6 +40,7 @@ class FParticlePerfStatsListener_NiagaraBaselineComparisonRender;
 class FNiagaraDebugger;
 class UNiagaraParameterDefinitions;
 class UNiagaraReservedParametersManager;
+class FNiagaraGraphDataCache;
 
 DECLARE_STATS_GROUP(TEXT("Niagara Editor"), STATGROUP_NiagaraEditor, STATCAT_Advanced);
 
@@ -196,6 +197,8 @@ public:
 
 	TArray<UNiagaraParameterDefinitions*>& GetReservedDefinitions() { return ReservedDefinitions; };
 
+	FNiagaraGraphDataCache& GetGraphDataCache() const { return *GraphDataCache.Get(); }
+
 private:
 	class FDeferredDestructionContainerBase
 	{
@@ -335,6 +338,8 @@ private:
 
 	// Set of Parameter Definitions assets to reconcile the unique Id for Definitions that are duplicated from each other.
 	TArray<UNiagaraParameterDefinitions*> ReservedDefinitions;
+
+	TUniquePtr<FNiagaraGraphDataCache> GraphDataCache;
 };
 
 USTRUCT()
diff --git a/Engine/Plugins/FX/Niagara/Source/NiagaraEditor/Public/NiagaraGraphDataCache.h b/Engine/Plugins/FX/Niagara/Source/NiagaraEditor/Public/NiagaraGraphDataCache.h
new file mode 100644
--- /dev/null
+++ b/Engine/Plugins/FX/Niagara/Source/NiagaraEditor/Public/NiagaraGraphDataCache.h
@@ -0,0 +1,72 @@
+// Copyright Epic Games, Inc. All Rights Reserved.
+
+#pragma once
+
+#include "NiagaraCommon.h"
+#include "NiagaraEditorUtilities.h"
+#include "UObject/ObjectKey.h"
+#include "UObject/WeakObjectPtrTemplates.h"
+#include "ViewModels/Stack/NiagaraStackGraphUtilities.h"
+
+class FHlslNiagaraTranslator;
+class UEdGraphPin;
+class UNiagaraGraph;
+class UNiagaraNodeFunctionCall;
+
+class FNiagaraGraphDataCache
+{
+public:
+	void GetStackFunctionInputPins(UNiagaraNodeFunctionCall& FunctionCallNode, TArray<const UEdGraphPin*>& OutInputPins, TSet<const UEdGraphPin*>& OutHiddenPins, FCompileConstantResolver ConstantResolver, FNiagaraStackGraphUtilities::ENiagaraGetStackFunctionInputPinsOptions Options, bool bIgnoreDisabled);
+	
+private:
+	struct FGetStackFunctionInputPinsKey
+	{
+		FGetStackFunctionInputPinsKey(UNiagaraNodeFunctionCall& InFunctionCallNode, const FCompileConstantResolver& InConstantResolver, FNiagaraStackGraphUtilities::ENiagaraGetStackFunctionInputPinsOptions InOptions, bool bInIgnoreDisabled);
+
+		FORCEINLINE bool operator == (const FGetStackFunctionInputPinsKey& Other) const
+		{
+			return
+				FunctionCallNodeKey == Other.FunctionCallNodeKey &&
+				ConstantResolverSystemKey == Other.ConstantResolverSystemKey &&
+				ConstantResolverEmitterKey == Other.ConstantResolverEmitterKey &&
+				ConstantResolverTranslator == Other.ConstantResolverTranslator &&
+				ConstantResolverUsage == Other.ConstantResolverUsage &&
+				ConstantResolverDebugState == Other.ConstantResolverDebugState &&
+				Options == Other.Options &&
+				bIgnoreDisabled == Other.bIgnoreDisabled;
+		}
+
+		friend uint32 GetTypeHash(const FGetStackFunctionInputPinsKey& Key)
+		{
+			return Key.Hash;
+		}
+		
+	private:
+		FObjectKey FunctionCallNodeKey;
+
+		FObjectKey ConstantResolverSystemKey;
+		FObjectKey ConstantResolverEmitterKey;
+		const FHlslNiagaraTranslator* ConstantResolverTranslator;
+		ENiagaraScriptUsage ConstantResolverUsage;
+		ENiagaraFunctionDebugState ConstantResolverDebugState;
+
+		FNiagaraStackGraphUtilities::ENiagaraGetStackFunctionInputPinsOptions Options;
+		bool bIgnoreDisabled;
+
+		int32 Hash;
+	};
+
+	struct FCachedStackFunctionInputPins
+	{
+		TWeakObjectPtr<UNiagaraGraph> SourceGraphWeak;
+		FGuid LastSourceGraphChangeId;
+		TWeakObjectPtr<UNiagaraNodeFunctionCall> FunctionCallNodeWeak;
+		TWeakObjectPtr<UNiagaraGraph> CalledGraphWeak;
+		FGuid LastCalledGraphChangeId;
+		TArray<const UEdGraphPin*> InputPins;
+		TSet<const UEdGraphPin*> HiddenPins;
+		bool IsValid() const;
+	};
+
+	TMap<FGetStackFunctionInputPinsKey, FCachedStackFunctionInputPins> GetStackFunctionInputPinsCache;
+};
\ No newline at end of file
diff --git a/Engine/Plugins/FX/Niagara/Source/NiagaraEditor/Public/NiagaraParameterMapHistory.h b/Engine/Plugins/FX/Niagara/Source/NiagaraEditor/Public/NiagaraParameterMapHistory.h
--- a/Engine/Plugins/FX/Niagara/Source/NiagaraEditor/Public/NiagaraParameterMapHistory.h
+++ b/Engine/Plugins/FX/Niagara/Source/NiagaraEditor/Public/NiagaraParameterMapHistory.h
@@ -21,9 +21,15 @@ public:
 	FCompileConstantResolver(const UNiagaraSystem* System, ENiagaraScriptUsage Usage, ENiagaraFunctionDebugState DebugState = ENiagaraFunctionDebugState::NoDebug) : Emitter(nullptr), System(System), Translator(nullptr), Usage(Usage), DebugState(DebugState) {}
 	FCompileConstantResolver(const FHlslNiagaraTranslator* Translator, ENiagaraFunctionDebugState DebugState = ENiagaraFunctionDebugState::NoDebug) : Emitter(nullptr), System(nullptr), Translator(Translator), Usage(ENiagaraScriptUsage::Function), DebugState(DebugState) {}
 
+	const UNiagaraEmitter* GetEmitter() const { return Emitter; }
+	const UNiagaraSystem* GetSystem() const { return System; }
+	const FHlslNiagaraTranslator* GetTranslator() const { return Translator; }
+	ENiagaraScriptUsage GetUsage() const { return Usage; }
+	ENiagaraFunctionDebugState GetDebugState() const { return DebugState; }
+
 	bool ResolveConstant(FNiagaraVariable& OutConstant) const;
 
-	ENiagaraFunctionDebugState GetDebugState() const;
+	ENiagaraFunctionDebugState CalculateDebugState() const;
 	FCompileConstantResolver WithDebugState(ENiagaraFunctionDebugState InDebugState) const;
 	FCompileConstantResolver WithUsage(ENiagaraScriptUsage ScriptUsage) const;
 private:
diff --git a/Engine/Plugins/FX/Niagara/Source/NiagaraEditor/Public/ViewModels/Stack/NiagaraStackFunctionInput.h b/Engine/Plugins/FX/Niagara/Source/NiagaraEditor/Public/ViewModels/Stack/NiagaraStackFunctionInput.h
--- a/Engine/Plugins/FX/Niagara/Source/NiagaraEditor/Public/ViewModels/Stack/NiagaraStackFunctionInput.h
+++ b/Engine/Plugins/FX/Niagara/Source/NiagaraEditor/Public/ViewModels/Stack/NiagaraStackFunctionInput.h
@@ -348,6 +348,8 @@ private:
 
 	TArray<UNiagaraStackFunctionInput*> GetChildInputs() const;
 
+	void GetCurrentChangeIds(FGuid& OutOwningGraphChangeId, FGuid& OutFunctionGraphChangeId) const;
+
 private:
 	/** The module function call which owns this input entry. NOTE: This input might not be an input to the module function
 		call, it may be an input to a dynamic input function call which is owned by the module. */
@@ -368,6 +370,9 @@ private:
 	/** The meta data for this input, defined in the owning function's script. */
 	TOptional<FNiagaraVariableMetaData> InputMetaData;
 
+	// Stack issues generated when updating from metadata.
+	TArray<FStackIssue> InputMetaDataIssues;
+
 	/** A unique key for this input for looking up editor only UI data. */
 	FString StackEditorDataKey;
 
@@ -471,4 +476,10 @@ private:
 
 	// If true then this stack entry is the semantic child of another stack entry
 	bool bIsSemanticChild = false;
+
+	// The value of the change id for the graph which owns the function call that owns this input as of the last refresh.
+	TOptional<FGuid> LastOwningGraphChangeId;
+
+	// The value of the change id for the called graph of the function call which owns this input as of the last refresh.
+	TOptional<FGuid> LastFunctionGraphChangeId;
 };
diff --git a/Engine/Plugins/FX/Niagara/Source/NiagaraEditor/Public/ViewModels/Stack/NiagaraStackGraphUtilities.h b/Engine/Plugins/FX/Niagara/Source/NiagaraEditor/Public/ViewModels/Stack/NiagaraStackGraphUtilities.h
--- a/Engine/Plugins/FX/Niagara/Source/NiagaraEditor/Public/ViewModels/Stack/NiagaraStackGraphUtilities.h
+++ b/Engine/Plugins/FX/Niagara/Source/NiagaraEditor/Public/ViewModels/Stack/NiagaraStackGraphUtilities.h
@@ -97,6 +97,9 @@ namespace FNiagaraStackGraphUtilities
 	/* Returns the input pins for the given function call node. */
 	void GetStackFunctionInputPins(UNiagaraNodeFunctionCall& FunctionCallNode, TArray<const UEdGraphPin*>& OutInputPins, TSet<const UEdGraphPin*>& OutHiddenPins, FCompileConstantResolver ConstantResolver, ENiagaraGetStackFunctionInputPinsOptions Options = ENiagaraGetStackFunctionInputPinsOptions::AllInputs, bool bIgnoreDisabled = false);
 
+	/* Returns the input pins for the given function call node.  Bypasses the module level caching of this data and generates it each call. */
+	void GetStackFunctionInputPinsWithoutCache(UNiagaraNodeFunctionCall& FunctionCallNode, TArray<const UEdGraphPin*>& OutInputPins, TSet<const UEdGraphPin*>& OutHiddenPins, FCompileConstantResolver ConstantResolver, ENiagaraGetStackFunctionInputPinsOptions Options, bool bIgnoreDisabled);
+
 	/* Module script calls do not have direct inputs, but rely on the parameter map being initialized correctly. This utility function resolves which of the module's parameters are reachable during compilation and returns a list of pins on the parameter map node that do not have to be compiled. */
 	TArray<UEdGraphPin*> GetUnusedFunctionInputPins(UNiagaraNodeFunctionCall& FunctionCallNode, FCompileConstantResolver ConstantResolver);
 
diff --git a/Engine/Plugins/FX/Niagara/Source/NiagaraEditor/Public/ViewModels/Stack/NiagaraStackModuleItemLinkedInputCollection.h b/Engine/Plugins/FX/Niagara/Source/NiagaraEditor/Public/ViewModels/Stack/NiagaraStackModuleItemLinkedInputCollection.h
--- a/Engine/Plugins/FX/Niagara/Source/NiagaraEditor/Public/ViewModels/Stack/NiagaraStackModuleItemLinkedInputCollection.h
+++ b/Engine/Plugins/FX/Niagara/Source/NiagaraEditor/Public/ViewModels/Stack/NiagaraStackModuleItemLinkedInputCollection.h
@@ -27,6 +27,15 @@ public:
 protected:
 	virtual void RefreshChildrenInternal(const TArray<UNiagaraStackEntry*>& CurrentChildren, TArray<UNiagaraStackEntry*>& NewChildren, TArray<FStackIssue>& NewIssues) override;
 
+private:
+	void GetCurrentChangeIds(FGuid& OutOwningGraphChangeId, FGuid& OutFunctionGraphChangeId) const;
+
 private:
 	UNiagaraNodeFunctionCall* FunctionCallNode;
+
+	// The value of the change id for the graph which owns the function call as of the last refresh.
+	TOptional<FGuid> LastOwningGraphChangeId;
+
+	// The value of the change id for the called graph of the function call as of the last refresh.
+	TOptional<FGuid> LastFunctionGraphChangeId;
 };
diff --git a/Engine/Plugins/FX/Niagara/Source/NiagaraEditor/Public/ViewModels/Stack/NiagaraStackModuleItemOutputCollection.h b/Engine/Plugins/FX/Niagara/Source/NiagaraEditor/Public/ViewModels/Stack/NiagaraStackModuleItemOutputCollection.h
--- a/Engine/Plugins/FX/Niagara/Source/NiagaraEditor/Public/ViewModels/Stack/NiagaraStackModuleItemOutputCollection.h
+++ b/Engine/Plugins/FX/Niagara/Source/NiagaraEditor/Public/ViewModels/Stack/NiagaraStackModuleItemOutputCollection.h
@@ -26,6 +26,15 @@ public:
 protected:
 	virtual void RefreshChildrenInternal(const TArray<UNiagaraStackEntry*>& CurrentChildren, TArray<UNiagaraStackEntry*>& NewChildren, TArray<FStackIssue>& NewIssues) override;
 
+private:
+	void GetCurrentChangeIds(FGuid& OutOwningGraphChangeId, FGuid& OutFunctionGraphChangeId) const;
+
 private:
 	UNiagaraNodeFunctionCall* FunctionCallNode;
+
+	// The value of the change id for the graph which owns the function call as of the last refresh.
+	TOptional<FGuid> LastOwningGraphChangeId;
+
+	// The value of the change id for the called graph of the function call as of the last refresh.
+	TOptional<FGuid> LastFunctionGraphChangeId;
 };
diff --git a/Engine/Plugins/FX/Niagara/Source/NiagaraEditor/Public/ViewModels/Stack/NiagaraStackScriptItemGroup.h b/Engine/Plugins/FX/Niagara/Source/NiagaraEditor/Public/ViewModels/Stack/NiagaraStackScriptItemGroup.h
--- a/Engine/Plugins/FX/Niagara/Source/NiagaraEditor/Public/ViewModels/Stack/NiagaraStackScriptItemGroup.h
+++ b/Engine/Plugins/FX/Niagara/Source/NiagaraEditor/Public/ViewModels/Stack/NiagaraStackScriptItemGroup.h
@@ -64,6 +64,10 @@ private:
 
 	void OnScriptGraphChanged(const struct FEdGraphEditAction& InAction);
 
+	void OnSystemScriptCompiled(UNiagaraScript* InScript, const FGuid& ScriptVersion);
+
+	void OnParticleScriptCompiled(UNiagaraScript* InScript, const FGuid& ScriptVersion);
+
 	TOptional<FDropRequestResponse> CanDropOnTarget(const UNiagaraStackEntry& TargetEntry, const FDropRequest& DropRequest);
 
 	TOptional<FDropRequestResponse> CanDropEntriesOnTarget(const UNiagaraStackEntry& TargetEntry, const FDropRequest& DropRequest);
@@ -97,5 +101,9 @@ private:
 
 	TWeakObjectPtr<UEdGraph> ScriptGraph;
 
+	TWeakObjectPtr<UNiagaraSystem> OwningSystemWeak;
+
+	TWeakObjectPtr<UNiagaraScript> OwningParticleScriptWeak;
+
 	FDelegateHandle OnGraphChangedHandle;
 };
diff --git a/Engine/Plugins/FX/Niagara/Source/NiagaraEditor/Public/ViewModels/Stack/NiagaraStackViewModel.h b/Engine/Plugins/FX/Niagara/Source/NiagaraEditor/Public/ViewModels/Stack/NiagaraStackViewModel.h
--- a/Engine/Plugins/FX/Niagara/Source/NiagaraEditor/Public/ViewModels/Stack/NiagaraStackViewModel.h
+++ b/Engine/Plugins/FX/Niagara/Source/NiagaraEditor/Public/ViewModels/Stack/NiagaraStackViewModel.h
@@ -162,8 +162,6 @@ private:
 	void EntryRequestFullRefreshDeferred();
 	void RefreshTopLevelViewModels();
 	void RefreshHasIssues();
-	void OnSystemCompiled();
-	void OnEmitterCompiled(UNiagaraScript* InScript, const FGuid& ScriptVersion);
 	void EmitterParentRemoved();
 	/** Called by the tick function to perform partial search */
 	void SearchTick();
