Manual merge of CL 17047875

Added a "on tab closed" callback for the stack to clear the search text. Previously, the search text would remain as was in the selection stack view model, and re-opening the tab would result in the previous search results being applied without the text in the search box showing

diff --git a/Engine/Plugins/FX/Niagara/Source/NiagaraEditor/Private/Toolkits/NiagaraSystemToolkit.cpp b/Engine/Plugins/FX/Niagara/Source/NiagaraEditor/Private/Toolkits/NiagaraSystemToolkit.cpp
--- a/Engine/Plugins/FX/Niagara/Source/NiagaraEditor/Private/Toolkits/NiagaraSystemToolkit.cpp
+++ b/Engine/Plugins/FX/Niagara/Source/NiagaraEditor/Private/Toolkits/NiagaraSystemToolkit.cpp
@@ -63,6 +63,7 @@
 #include "NiagaraEditor/Private/SNiagaraAssetPickerList.h"
 #include "Widgets/SNiagaraParameterDefinitionsPanel.h"
 #include "ViewModels/NiagaraParameterDefinitionsPanelViewModel.h"
+#include "ViewModels/Stack/NiagaraStackViewModel.h"
 
 
 #define LOCTEXT_NAMESPACE "NiagaraSystemEditor"
@@ -674,6 +675,12 @@ TSharedRef<SDockTab> FNiagaraSystemToolkit::SpawnTab_SelectedEmitterStack(const
 			NiagaraEditorModule.GetWidgetProvider()->CreateStackView(*SystemViewModel->GetSelectionViewModel()->GetSelectionStackViewModel())
 		];
 
+	SDockTab::FOnTabClosedCallback TabClosedCallback = SDockTab::FOnTabClosedCallback::CreateLambda([=](TSharedRef<SDockTab> DockTab)
+	{
+		SystemViewModel->GetSelectionViewModel()->GetSelectionStackViewModel()->ResetSearchText();
+	});
+	
+	SpawnedTab->SetOnTabClosed(TabClosedCallback);
 	return SpawnedTab;
 }
 
diff --git a/Engine/Plugins/FX/Niagara/Source/NiagaraEditor/Private/ViewModels/Stack/NiagaraStackViewModel.cpp b/Engine/Plugins/FX/Niagara/Source/NiagaraEditor/Private/ViewModels/Stack/NiagaraStackViewModel.cpp
--- a/Engine/Plugins/FX/Niagara/Source/NiagaraEditor/Private/ViewModels/Stack/NiagaraStackViewModel.cpp
+++ b/Engine/Plugins/FX/Niagara/Source/NiagaraEditor/Private/ViewModels/Stack/NiagaraStackViewModel.cpp
@@ -199,6 +199,13 @@ void UNiagaraStackViewModel::Tick()
 	}
 }
 
+void UNiagaraStackViewModel::ResetSearchText()
+{
+	CurrentSearchText = FText::GetEmpty();
+	// restarting the search with empty text will reset search results
+	bRestartSearch = true;	
+}
+
 void UNiagaraStackViewModel::OnSearchTextChanged(const FText& SearchText)
 {
 	if (RootEntry && !CurrentSearchText.EqualTo(SearchText))
diff --git a/Engine/Plugins/FX/Niagara/Source/NiagaraEditor/Public/ViewModels/Stack/NiagaraStackViewModel.h b/Engine/Plugins/FX/Niagara/Source/NiagaraEditor/Public/ViewModels/Stack/NiagaraStackViewModel.h
--- a/Engine/Plugins/FX/Niagara/Source/NiagaraEditor/Public/ViewModels/Stack/NiagaraStackViewModel.h
+++ b/Engine/Plugins/FX/Niagara/Source/NiagaraEditor/Public/ViewModels/Stack/NiagaraStackViewModel.h
@@ -111,6 +111,7 @@ public:
 
 	virtual void Tick();
 	//~ stack search stuff
+	void ResetSearchText();
 	void OnSearchTextChanged(const FText& SearchText);
 	FText GetCurrentSearchText() const { return CurrentSearchText; };
 	bool IsSearching();
