- Fix missing swizzle when reading color component on GPU

diff --git a/Engine/Plugins/FX/Niagara/Shaders/Private/NiagaraDataInterfaceSkeletalMesh.ush b/Engine/Plugins/FX/Niagara/Shaders/Private/NiagaraDataInterfaceSkeletalMesh.ush
--- a/Engine/Plugins/FX/Niagara/Shaders/Private/NiagaraDataInterfaceSkeletalMesh.ush
+++ b/Engine/Plugins/FX/Niagara/Shaders/Private/NiagaraDataInterfaceSkeletalMesh.ush
@@ -223,13 +223,11 @@ FDISkelMeshSkinnedVertex DISkelMesh_GetSkinnedVertex(in FDISkelMeshContext DICon
 		TangentZ.xyz = mul(CurrBoneMatrix, float4(TangentZ.xyz, 0.0f)).xyz;
 	}
 
-	float3 TangentY = cross(TangentZ.xyz, TangentX.xyz) * TangentZ.w;
-
 	FDISkelMeshSkinnedVertex SkinnedVertex;
 	SkinnedVertex.Position = CurrPosition;
 	SkinnedVertex.PrevPosition = PrevPosition;
-	SkinnedVertex.TangentX = cross(TangentY, TangentZ.xyz) * TangentZ.w;
-	SkinnedVertex.TangentY = TangentY;
+	SkinnedVertex.TangentX = TangentX;
+	SkinnedVertex.TangentY = cross(TangentZ.xyz, TangentX.xyz);
 	SkinnedVertex.TangentZ = TangentZ.xyz;
 	return SkinnedVertex;
 }
@@ -770,7 +768,7 @@ void DISkelMesh_GetTriColor(in FDISkelMeshContext DIContext, in uint TriangleInd
     uint VertexIndex0 = DIContext.MeshIndexBuffer[IndexBufferOffset];
     uint VertexIndex1 = DIContext.MeshIndexBuffer[IndexBufferOffset + 1];
     uint VertexIndex2 = DIContext.MeshIndexBuffer[IndexBufferOffset + 2];
-    OutColor = (DIContext.MeshColorBuffer[VertexIndex0] * BaryCoord.x) + (DIContext.MeshColorBuffer[VertexIndex1] * BaryCoord.y) + (DIContext.MeshColorBuffer[VertexIndex2] * BaryCoord.z);
+    OutColor = (DIContext.MeshColorBuffer[VertexIndex0] FMANUALFETCH_COLOR_COMPONENT_SWIZZLE * BaryCoord.x) + (DIContext.MeshColorBuffer[VertexIndex1] FMANUALFETCH_COLOR_COMPONENT_SWIZZLE * BaryCoord.y) + (DIContext.MeshColorBuffer[VertexIndex2] FMANUALFETCH_COLOR_COMPONENT_SWIZZLE * BaryCoord.z);
 }
 
 void DISkelMesh_GetSkinnedVertex(in FDISkelMeshContext DIContext, in int VertexIndex, out float3 OutPosition, out float3 OutVelocity, out float3 OutNormal, out float3 OutBinormal, out float3 OutTangent)
@@ -789,14 +787,14 @@ void DISkelMesh_GetSkinnedVertexWS(in FDISkelMeshContext DIContext, in int Verte
     OutPosition = mul(float4(SkinnedVertex.Position, 1.0f), DIContext.InstanceTransform).xyz;
     float3 PrevPosition = mul(float4(SkinnedVertex.PrevPosition, 1.0f), DIContext.InstancePrevTransform).xyz;
     OutVelocity = (SkinnedVertex.Position - SkinnedVertex.PrevPosition) * DIContext.InstanceInvDeltaTime;
-    OutNormal   = normalize(mul(float4(SkinnedVertex.TangentZ, 0.0f), DIContext.InstanceTransform).xyz);
-    OutBinormal = normalize(mul(float4(SkinnedVertex.TangentY, 0.0f), DIContext.InstanceTransform).xyz);
-    OutTangent  = normalize(mul(float4(SkinnedVertex.TangentX, 0.0f), DIContext.InstanceTransform).xyz);
+	OutNormal = SkinnedVertex.TangentZ;
+	OutBinormal = SkinnedVertex.TangentY;
+	OutTangent = SkinnedVertex.TangentX;
 }
 
 void DISkelMesh_GetVertexColor(in FDISkelMeshContext DIContext, in int VertexIndex, out float4 OutColor)
 {
-    OutColor = DIContext.MeshColorBuffer[VertexIndex];
+    OutColor = DIContext.MeshColorBuffer[VertexIndex] FMANUALFETCH_COLOR_COMPONENT_SWIZZLE;
 }
 
 void DISkelMesh_GetVertexUV(in FDISkelMeshContext DIContext, in int VertexIndex, in int UVSet, out float2 OutUV)
