- Fix deprecation warning

#rnx

diff --git a/Engine/Plugins/FX/Niagara/Source/Niagara/Classes/NiagaraEmitter.h b/Engine/Plugins/FX/Niagara/Source/Niagara/Classes/NiagaraEmitter.h
--- a/Engine/Plugins/FX/Niagara/Source/Niagara/Classes/NiagaraEmitter.h
+++ b/Engine/Plugins/FX/Niagara/Source/Niagara/Classes/NiagaraEmitter.h
@@ -239,7 +239,7 @@ public:
 	virtual void PostRename(UObject* OldOuter, const FName OldName) override;
 	virtual void PostDuplicate(EDuplicateMode::Type DuplicateMode) override;
 	virtual void PostEditChangeProperty(struct FPropertyChangedEvent& PropertyChangedEvent) override;
-	virtual void PreSave(const class ITargetPlatform* TargetPlatform) override;
+	virtual void PreSave(FObjectPreSaveContext ObjectSaveContext) override;
 	NIAGARA_API FOnPropertiesChanged& OnPropertiesChanged();
 	NIAGARA_API FOnRenderersChanged& OnRenderersChanged();
 	/** Helper method for when a rename has been detected within the graph. Covers renaming the internal renderer bindings.*/
diff --git a/Engine/Plugins/FX/Niagara/Source/Niagara/Private/NiagaraEmitter.cpp b/Engine/Plugins/FX/Niagara/Source/Niagara/Private/NiagaraEmitter.cpp
--- a/Engine/Plugins/FX/Niagara/Source/Niagara/Private/NiagaraEmitter.cpp
+++ b/Engine/Plugins/FX/Niagara/Source/Niagara/Private/NiagaraEmitter.cpp
@@ -18,6 +18,7 @@
 #include "NiagaraTrace.h"
 #include "Interfaces/ITargetPlatform.h"
 #include "Modules/ModuleManager.h"
+#include "UObject/ObjectSaveContext.h"
 #include "UObject/Package.h"
 
 #if WITH_EDITOR
@@ -926,11 +927,9 @@ void UNiagaraEmitter::PostEditChangeProperty(struct FPropertyChangedEvent& Prope
 #endif
 }
 
-void UNiagaraEmitter::PreSave(const ITargetPlatform* TargetPlatform)
+void UNiagaraEmitter::PreSave(FObjectPreSaveContext ObjectSaveContext)
 {
-	PRAGMA_DISABLE_DEPRECATION_WARNINGS;
-	Super::PreSave(TargetPlatform);
-	PRAGMA_ENABLE_DEPRECATION_WARNINGS;
+	Super::PreSave(ObjectSaveContext);
 	UpdateEmitterAfterLoad();
 }
 
