New DI functions for allowing a finalize stage to be run after all instances of Pre/Post Stage has occured.  Allows for things like batching allocations for multiple instances of a DI.

diff --git a/Engine/Plugins/FX/Niagara/Source/Niagara/Classes/NiagaraDataInterface.h b/Engine/Plugins/FX/Niagara/Source/Niagara/Classes/NiagaraDataInterface.h
--- a/Engine/Plugins/FX/Niagara/Source/Niagara/Classes/NiagaraDataInterface.h
+++ b/Engine/Plugins/FX/Niagara/Source/Niagara/Classes/NiagaraDataInterface.h
@@ -251,6 +251,12 @@ struct FNiagaraDataInterfaceProxy : TSharedFromThis<FNiagaraDataInterfaceProxy,
 	virtual void PostStage(FRHICommandList& RHICmdList, const FNiagaraDataInterfaceStageArgs& Context) {}
 	virtual void PostSimulate(FRHICommandList& RHICmdList, const FNiagaraDataInterfaceArgs& Context) {}
 
+	virtual bool RequiresPreStageFinalize() const { return false; }
+	virtual void FinalizePreStage(FRHICommandList& RHICmdList, const NiagaraEmitterInstanceBatcher* Batcher) {}
+
+	virtual bool RequiresPostStageFinalize() const { return false; }
+	virtual void FinalizePostStage(FRHICommandList& RHICmdList, const NiagaraEmitterInstanceBatcher* Batcher) {}
+
 	virtual FNiagaraDataInterfaceProxyRW* AsIterationProxy() { return nullptr; }
 };
 
diff --git a/Engine/Plugins/FX/Niagara/Source/Niagara/Classes/NiagaraEmitterInstanceBatcher.h b/Engine/Plugins/FX/Niagara/Source/Niagara/Classes/NiagaraEmitterInstanceBatcher.h
--- a/Engine/Plugins/FX/Niagara/Source/Niagara/Classes/NiagaraEmitterInstanceBatcher.h
+++ b/Engine/Plugins/FX/Niagara/Source/Niagara/Classes/NiagaraEmitterInstanceBatcher.h
@@ -131,10 +131,10 @@ public:
 	FNiagaraDataInterfaceProxyRW* FindIterationInterface(FNiagaraComputeInstanceData* Instance, const uint32 SimulationStageIndex) const;
 
 	/** Loop over all the data interfaces and call the pre-stage methods */
-	void PreStageInterface(FRHICommandList& RHICmdList, const FNiagaraGPUSystemTick& Tick, const FNiagaraComputeInstanceData& InstanceData, const FNiagaraSimStageData& SimStageData) const;
+	void PreStageInterface(FRHICommandList& RHICmdList, const FNiagaraGPUSystemTick& Tick, const FNiagaraComputeInstanceData& InstanceData, const FNiagaraSimStageData& SimStageData, TSet<FNiagaraDataInterfaceProxy*>& ProxiesToFinalize) const;
 
 	/** Loop over all the data interfaces and call the post-stage methods */
-	void PostStageInterface(FRHICommandList& RHICmdList, const FNiagaraGPUSystemTick& Tick, const FNiagaraComputeInstanceData& InstanceData, const FNiagaraSimStageData& SimStageData) const;
+	void PostStageInterface(FRHICommandList& RHICmdList, const FNiagaraGPUSystemTick& Tick, const FNiagaraComputeInstanceData& InstanceData, const FNiagaraSimStageData& SimStageData, TSet<FNiagaraDataInterfaceProxy*>& ProxiesToFinalize) const;
 
 	/** Loop over all data interfaces and call the post-simulate methods */
 	void PostSimulateInterface(FRHICommandList& RHICmdList, const FNiagaraGPUSystemTick& Tick, const FNiagaraComputeInstanceData& InstanceData) const;
diff --git a/Engine/Plugins/FX/Niagara/Source/Niagara/Private/NiagaraEmitterInstanceBatcher.cpp b/Engine/Plugins/FX/Niagara/Source/Niagara/Private/NiagaraEmitterInstanceBatcher.cpp
--- a/Engine/Plugins/FX/Niagara/Source/Niagara/Private/NiagaraEmitterInstanceBatcher.cpp
+++ b/Engine/Plugins/FX/Niagara/Source/Niagara/Private/NiagaraEmitterInstanceBatcher.cpp
@@ -565,7 +565,7 @@ FNiagaraDataInterfaceProxyRW* NiagaraEmitterInstanceBatcher::FindIterationInterf
 	return Instance->FindIterationInterface(SimulationStageIndex);
 }
 
-void NiagaraEmitterInstanceBatcher::PreStageInterface(FRHICommandList& RHICmdList, const FNiagaraGPUSystemTick& Tick, const FNiagaraComputeInstanceData& InstanceData, const FNiagaraSimStageData& SimStageData) const
+void NiagaraEmitterInstanceBatcher::PreStageInterface(FRHICommandList& RHICmdList, const FNiagaraGPUSystemTick& Tick, const FNiagaraComputeInstanceData& InstanceData, const FNiagaraSimStageData& SimStageData, TSet<FNiagaraDataInterfaceProxy*>& ProxiesToFinalize) const
 {
 	// Note: All stages will contain the same bindings so if they are valid for one they are valid for all, this could change in the future
 	const FNiagaraShaderRef& ComputeShader = InstanceData.Context->GPUScript_RT->GetShader(0);
@@ -579,12 +579,17 @@ void NiagaraEmitterInstanceBatcher::PreStageInterface(FRHICommandList& RHICmdLis
 		{
 			const FNiagaraDataInterfaceStageArgs TmpContext(Interface, Tick.SystemInstanceID, this, &InstanceData, &SimStageData, InstanceData.IsOutputStage(Interface, SimStageData.StageIndex), InstanceData.IsIterationStage(Interface, SimStageData.StageIndex));
 			Interface->PreStage(RHICmdList, TmpContext);
+
+			if (Interface->RequiresPreStageFinalize())
+			{
+				ProxiesToFinalize.Add(Interface);
+			}
 		}
 		InterfaceIndex++;
 	}
 }
 
-void NiagaraEmitterInstanceBatcher::PostStageInterface(FRHICommandList& RHICmdList, const FNiagaraGPUSystemTick& Tick, const FNiagaraComputeInstanceData& InstanceData, const FNiagaraSimStageData& SimStageData) const
+void NiagaraEmitterInstanceBatcher::PostStageInterface(FRHICommandList& RHICmdList, const FNiagaraGPUSystemTick& Tick, const FNiagaraComputeInstanceData& InstanceData, const FNiagaraSimStageData& SimStageData, TSet<FNiagaraDataInterfaceProxy*>& ProxiesToFinalize) const
 {
 	// Note: All stages will contain the same bindings so if they are valid for one they are valid for all, this could change in the future
 	const FNiagaraShaderRef& ComputeShader = InstanceData.Context->GPUScript_RT->GetShader(0);
@@ -598,6 +603,11 @@ void NiagaraEmitterInstanceBatcher::PostStageInterface(FRHICommandList& RHICmdLi
 		{
 			const FNiagaraDataInterfaceStageArgs TmpContext(Interface, Tick.SystemInstanceID, this, &InstanceData, &SimStageData, InstanceData.IsOutputStage(Interface, SimStageData.StageIndex), InstanceData.IsIterationStage(Interface, SimStageData.StageIndex));
 			Interface->PostStage(RHICmdList, TmpContext);
+
+			if (Interface->RequiresPostStageFinalize())
+			{
+				ProxiesToFinalize.Add(Interface);
+			}
 		}
 		InterfaceIndex++;
 	}
@@ -1285,9 +1295,17 @@ void NiagaraEmitterInstanceBatcher::ExecuteTicks(FRHICommandList& RHICmdList, FR
 		}
 
 		// Execute PreStage
-		for (const FNiagaraGpuDispatchInstance& DispatchInstance : DispatchGroup.DispatchInstances)
 		{
-			PreStageInterface(RHICmdList, DispatchInstance.Tick, DispatchInstance.InstanceData, DispatchInstance.SimStageData);
+			TSet<FNiagaraDataInterfaceProxy*> ProxiesToFinalize;
+			for (const FNiagaraGpuDispatchInstance& DispatchInstance : DispatchGroup.DispatchInstances)
+			{
+				PreStageInterface(RHICmdList, DispatchInstance.Tick, DispatchInstance.InstanceData, DispatchInstance.SimStageData, ProxiesToFinalize);
+			}
+
+			for (FNiagaraDataInterfaceProxy* ProxyToFinalize : ProxiesToFinalize)
+			{
+				ProxyToFinalize->FinalizePreStage(RHICmdList, this);
+			}
 		}
 
 		// Execute Stage
@@ -1316,44 +1334,52 @@ void NiagaraEmitterInstanceBatcher::ExecuteTicks(FRHICommandList& RHICmdList, FR
 		RHICmdList.EndUAVOverlap(GPUInstanceCounterManager.GetInstanceCountBuffer().UAV);
 
 		// Execute PostStage
-		for (const FNiagaraGpuDispatchInstance& DispatchInstance : DispatchGroup.DispatchInstances)
 		{
-			PostStageInterface(RHICmdList, DispatchInstance.Tick, DispatchInstance.InstanceData, DispatchInstance.SimStageData);
-			if ( DispatchInstance.SimStageData.bLastStage )
+			TSet<FNiagaraDataInterfaceProxy*> ProxiesToFinalize;
+			for (const FNiagaraGpuDispatchInstance& DispatchInstance : DispatchGroup.DispatchInstances)
 			{
-				PostSimulateInterface(RHICmdList, DispatchInstance.Tick, DispatchInstance.InstanceData);
-
-				// Update CurrentData with the latest information as things like ParticleReads can use this data
-				FNiagaraComputeExecutionContext* ComputeContext = DispatchInstance.InstanceData.Context;
-				const FNiagaraSimStageData& FinalSimStageData = DispatchInstance.SimStageData;
-				FNiagaraDataBuffer* FinalSimStageDataBuffer = FinalSimStageData.Destination != nullptr ? FinalSimStageData.Destination : FinalSimStageData.Source;
-				check(FinalSimStageDataBuffer != nullptr);
-
-				// If we are setting the data to render we need to ensure we switch back to the original CurrentData then swap the GPU buffers into it
-				if (DispatchInstance.SimStageData.bSetDataToRender)
+				PostStageInterface(RHICmdList, DispatchInstance.Tick, DispatchInstance.InstanceData, DispatchInstance.SimStageData, ProxiesToFinalize);
+				if ( DispatchInstance.SimStageData.bLastStage )
 				{
-					check(ComputeContext->DataSetOriginalBuffer_RT != nullptr);
-					FNiagaraDataBuffer* CurrentData = ComputeContext->DataSetOriginalBuffer_RT;
-					ComputeContext->DataSetOriginalBuffer_RT = nullptr;
+					PostSimulateInterface(RHICmdList, DispatchInstance.Tick, DispatchInstance.InstanceData);
 
-					ComputeContext->MainDataSet->CurrentData = CurrentData;
-					CurrentData->SwapGPU(FinalSimStageDataBuffer);
+					// Update CurrentData with the latest information as things like ParticleReads can use this data
+					FNiagaraComputeExecutionContext* ComputeContext = DispatchInstance.InstanceData.Context;
+					const FNiagaraSimStageData& FinalSimStageData = DispatchInstance.SimStageData;
+					FNiagaraDataBuffer* FinalSimStageDataBuffer = FinalSimStageData.Destination != nullptr ? FinalSimStageData.Destination : FinalSimStageData.Source;
+					check(FinalSimStageDataBuffer != nullptr);
 
-					// Mark data as ready for anyone who picks up the buffer on the next frame
-					CurrentData->SetGPUDataReadyStage(ENiagaraGpuComputeTickStage::First);
+					// If we are setting the data to render we need to ensure we switch back to the original CurrentData then swap the GPU buffers into it
+					if (DispatchInstance.SimStageData.bSetDataToRender)
+					{
+						check(ComputeContext->DataSetOriginalBuffer_RT != nullptr);
+						FNiagaraDataBuffer* CurrentData = ComputeContext->DataSetOriginalBuffer_RT;
+						ComputeContext->DataSetOriginalBuffer_RT = nullptr;
 
-					ComputeContext->SetTranslucentDataToRender(nullptr);
-					ComputeContext->SetDataToRender(CurrentData);
+						ComputeContext->MainDataSet->CurrentData = CurrentData;
+						CurrentData->SwapGPU(FinalSimStageDataBuffer);
+
+						// Mark data as ready for anyone who picks up the buffer on the next frame
+						CurrentData->SetGPUDataReadyStage(ENiagaraGpuComputeTickStage::First);
+                    
+						ComputeContext->SetTranslucentDataToRender(nullptr);
+						ComputeContext->SetDataToRender(CurrentData);
 
 #if WITH_MGPU
-					AddTemporalEffectBuffers(CurrentData);
+						AddTemporalEffectBuffers(CurrentData);
 #endif // WITH_MGPU
+					}
+						// If this is not the final tick of the final stage we need set our temporary buffer for data interfaces, etc, that may snoop from CurrentData
+					else
+					{
+						ComputeContext->MainDataSet->CurrentData = FinalSimStageDataBuffer;
+					}
 				}
-				// If this is not the final tick of the final stage we need set our temporary buffer for data interfaces, etc, that may snoop from CurrentData
-				else
-				{
-					ComputeContext->MainDataSet->CurrentData = FinalSimStageDataBuffer;
-				}
+			}
+
+			for (FNiagaraDataInterfaceProxy* ProxyToFinalize : ProxiesToFinalize)
+			{
+				ProxyToFinalize->FinalizePostStage(RHICmdList, this);
 			}
 		}
 
