Fix for the collision DI which resulted in a small memory leak and systems getting reset.
-pull ray tracing Dispatch construction out of the proxies Set function (which, if a buffer was allocated would lead to a nested dispatch)
-ensure that the ray tracing dispatches (specifically the dummy dispatch) is properly getting reset even if no ray tracing is involved

diff --git a/Engine/Plugins/FX/Niagara/Source/Niagara/Private/NiagaraDataInterfaceCollisionQuery.cpp b/Engine/Plugins/FX/Niagara/Source/Niagara/Private/NiagaraDataInterfaceCollisionQuery.cpp
--- a/Engine/Plugins/FX/Niagara/Source/Niagara/Private/NiagaraDataInterfaceCollisionQuery.cpp
+++ b/Engine/Plugins/FX/Niagara/Source/Niagara/Private/NiagaraDataInterfaceCollisionQuery.cpp
@@ -95,7 +95,6 @@ struct FNiagaraDataIntefaceProxyCollisionQuery : public FNiagaraDataInterfacePro
 		FNiagaraDataInterfaceProxy::PreStage(RHICmdList, Context);
 
 #if RHI_RAYTRACING
-
 		if (IsRayTracingEnabled() && GEnableGPUHWRTCollisions && MaxTracesPerParticle > 0)
 		{
 			//Accumulate the total ray requests for this DI for all dispatches in the stage.
@@ -106,14 +105,25 @@ struct FNiagaraDataIntefaceProxyCollisionQuery : public FNiagaraDataInterfacePro
 #endif
 	}
 
-	virtual void PostStage(FRHICommandList& RHICmdList, const FNiagaraDataInterfaceStageArgs& Context) override
+#if RHI_RAYTRACING
+	virtual bool RequiresPreStageFinalize() const override
 	{
+		return true;
 	}
 
-	virtual void PostSimulate(FRHICommandList& RHICmdList, const FNiagaraDataInterfaceArgs& Context) override
+	virtual void FinalizePreStage(FRHICommandList& RHICmdList, const NiagaraEmitterInstanceBatcher* Batcher) override
 	{
-		FNiagaraDataInterfaceProxy::PostSimulate(RHICmdList, Context);
+		FNiagaraRayTracingHelper& RTHelper = Batcher->GetRayTracingHelper();
+		if (IsRayTracingEnabled() && GEnableGPUHWRTCollisions && MaxTracesPerParticle > 0)
+		{
+			RTHelper.BuildDispatch(RHICmdList, this);
+		}
+		else
+		{
+			RTHelper.BuildDummyDispatch(RHICmdList);
+		}
 	}
+#endif
 
 	void RenderThreadInitialize(int32 InMaxTracesPerParticle, uint32 InMaxRetraces)
 	{
@@ -968,42 +978,53 @@ public:
 		}
 
 #if RHI_RAYTRACING
+		const bool HasRayTracingParametersBound = RayRequestsParam.IsUAVBound()
+			|| IntersectionResultsParam.IsBound()
+			|| RayTraceCountsParam.IsBound();
 
-		FNiagaraRayTracingHelper& RTHelper = Context.Batcher->GetRayTracingHelper();
-		const FNiagaraRayTraceDispatchInfo* DispatchInfo = nullptr;
-		if (IsRayTracingEnabled() && GEnableGPUHWRTCollisions && QueryDI->MaxTracesPerParticle > 0)
+		if ((IsRayTracingEnabled() && GEnableGPUHWRTCollisions) || HasRayTracingParametersBound)
 		{
-			DispatchInfo = &RTHelper.GetDispatch(QueryDI, RHICmdList);
-		}
-		else
-		{
-			DispatchInfo = &RTHelper.GetDummyDispatch(RHICmdList);
-		}
+			FNiagaraRayTracingHelper& RTHelper = Context.Batcher->GetRayTracingHelper();
+			const FNiagaraRayTraceDispatchInfo* DispatchInfo = nullptr;
+			if (IsRayTracingEnabled() && GEnableGPUHWRTCollisions && QueryDI->MaxTracesPerParticle > 0)
+			{
+				DispatchInfo = &RTHelper.GetDispatch(QueryDI);
+			}
+			else
+			{
+				DispatchInfo = &RTHelper.GetDummyDispatch();
+			}
 
-		SetShaderValue(RHICmdList, ComputeShaderRHI, RayTracingEnabledParam, IsRayTracingEnabled() && GEnableGPUHWRTCollisions ? 1 : 0);
-		SetShaderValue(RHICmdList, ComputeShaderRHI, MaxRayTraceCountParam, DispatchInfo->MaxRays);
+			SetShaderValue(RHICmdList, ComputeShaderRHI, RayTracingEnabledParam, IsRayTracingEnabled() && GEnableGPUHWRTCollisions ? 1 : 0);
+			SetShaderValue(RHICmdList, ComputeShaderRHI, MaxRayTraceCountParam, DispatchInfo->MaxRays);
 
-		if (RayRequestsParam.IsUAVBound())
-		{
-			check(DispatchInfo->RayRequests.IsValid());
-			RHICmdList.SetUAVParameter(ComputeShaderRHI, RayRequestsParam.GetUAVIndex(), DispatchInfo->RayRequests.Buffer->UAV);
-			SetShaderValue(RHICmdList, ComputeShaderRHI, RayRequestOffsetParam, DispatchInfo->RayRequests.Offset);
-		}
+			if (RayRequestsParam.IsUAVBound())
+			{
+				check(DispatchInfo->RayRequests.IsValid());
+				RHICmdList.SetUAVParameter(ComputeShaderRHI, RayRequestsParam.GetUAVIndex(), DispatchInfo->RayRequests.Buffer->UAV);
+				SetShaderValue(RHICmdList, ComputeShaderRHI, RayRequestOffsetParam, DispatchInfo->RayRequests.Offset);
+			}
 
-		if (IntersectionResultsParam.IsBound())
-		{
-			check(DispatchInfo->LastFrameRayTraceIntersections.IsValid());
+			if (IntersectionResultsParam.IsBound())
+			{
+				check(DispatchInfo->LastFrameRayTraceIntersections.IsValid());
 
-			SetSRVParameter(RHICmdList, ComputeShaderRHI, IntersectionResultsParam, DispatchInfo->LastFrameRayTraceIntersections.Buffer->SRV);
-			SetShaderValue(RHICmdList, ComputeShaderRHI, IntersectionResultOffsetParam, DispatchInfo->LastFrameRayTraceIntersections.Offset);
-		}
+				SetSRVParameter(RHICmdList, ComputeShaderRHI, IntersectionResultsParam, DispatchInfo->LastFrameRayTraceIntersections.Buffer->SRV);
+				SetShaderValue(RHICmdList, ComputeShaderRHI, IntersectionResultOffsetParam, DispatchInfo->LastFrameRayTraceIntersections.Offset);
+			}
 
-		if (RayTraceCountsParam.IsUAVBound())
-		{
-			check(DispatchInfo->RayCounts.IsValid());
+			if (RayTraceCountsParam.IsUAVBound())
+			{
+				check(DispatchInfo->RayCounts.IsValid());
 
-			RHICmdList.SetUAVParameter(ComputeShaderRHI, RayTraceCountsParam.GetUAVIndex(), DispatchInfo->RayCounts.Buffer->UAV);
-			SetShaderValue(RHICmdList, ComputeShaderRHI, RayTraceCountsOffsetParam, DispatchInfo->RayCounts.Offset);
+				RHICmdList.SetUAVParameter(ComputeShaderRHI, RayTraceCountsParam.GetUAVIndex(), DispatchInfo->RayCounts.Buffer->UAV);
+				SetShaderValue(RHICmdList, ComputeShaderRHI, RayTraceCountsOffsetParam, DispatchInfo->RayCounts.Offset);
+			}
+		}
+		else
+		{
+			SetShaderValue(RHICmdList, ComputeShaderRHI, RayTracingEnabledParam, 0);
+			SetShaderValue(RHICmdList, ComputeShaderRHI, MaxRayTraceCountParam, 0);
 		}
 #endif
 	}
diff --git a/Engine/Plugins/FX/Niagara/Source/Niagara/Private/NiagaraEmitterInstanceBatcher.cpp b/Engine/Plugins/FX/Niagara/Source/Niagara/Private/NiagaraEmitterInstanceBatcher.cpp
--- a/Engine/Plugins/FX/Niagara/Source/Niagara/Private/NiagaraEmitterInstanceBatcher.cpp
+++ b/Engine/Plugins/FX/Niagara/Source/Niagara/Private/NiagaraEmitterInstanceBatcher.cpp
@@ -1654,10 +1654,7 @@ void NiagaraEmitterInstanceBatcher::PreInitViews(FRDGBuilder& GraphBuilder, bool
 		PrepareAllTicks(GraphBuilder.RHICmdList);
 		
 #if RHI_RAYTRACING
-		if (HasRayTracingScene())
-		{
-			RayTracingHelper->BeginFrame(GraphBuilder.RHICmdList);
-		}
+		RayTracingHelper->BeginFrame(GraphBuilder.RHICmdList, HasRayTracingScene());
 #endif
 
 		if ( NiagaraEmitterInstanceBatcherLocal::GDebugLogging)
@@ -1747,10 +1744,7 @@ void NiagaraEmitterInstanceBatcher::PostRenderOpaque(FRDGBuilder& GraphBuilder,
 			FinishDispatches();
 
 #if RHI_RAYTRACING
-			if(HasRayTracingScene())
-			{
-				RayTracingHelper->EndFrame(RHICmdList, GetScene());
-			}
+			RayTracingHelper->EndFrame(RHICmdList, HasRayTracingScene(), GetScene());
 #endif
 
 			// Clear CurrentPassViews
diff --git a/Engine/Plugins/FX/Niagara/Source/NiagaraShader/Private/NiagaraRayTracingHelper.cpp b/Engine/Plugins/FX/Niagara/Source/NiagaraShader/Private/NiagaraRayTracingHelper.cpp
--- a/Engine/Plugins/FX/Niagara/Source/NiagaraShader/Private/NiagaraRayTracingHelper.cpp
+++ b/Engine/Plugins/FX/Niagara/Source/NiagaraShader/Private/NiagaraRayTracingHelper.cpp
@@ -204,7 +204,7 @@ void FNiagaraRayTracingHelper::Reset()
 	ViewUniformBuffer = nullptr;
 }
 
-void FNiagaraRayTracingHelper::BeginFrame(FRHICommandList& RHICmdList)
+void FNiagaraRayTracingHelper::BeginFrame(FRHICommandList& RHICmdList, bool HasRayTracingScene)
 {
 	//Store off last frames dispatches so we can still access the previous results buffers for reading in this frame's simulations.
 	PreviousFrameDispatches = MoveTemp(Dispatches);
@@ -234,17 +234,21 @@ void FNiagaraRayTracingHelper::BeginFrame(FRHICommandList& RHICmdList)
 	ViewUniformBuffer = nullptr;
 }
 
-void FNiagaraRayTracingHelper::EndFrame(FRHICommandList& RHICmdList, FScene* Scene)
+void FNiagaraRayTracingHelper::EndFrame(FRHICommandList& RHICmdList, bool HasRayTracingScene, FScene* Scene)
 {
-	RayRequests.Transition(RHICmdList, ERHIAccess::UAVCompute, ERHIAccess::SRVCompute);//Ray trace dispatches will read the ray requests.
-	RayTraceIntersections.Transition(RHICmdList, ERHIAccess::SRVCompute, ERHIAccess::UAVCompute);//Ray trace dispatches will write new results.
-	RayTraceCounts.Transition(RHICmdList, ERHIAccess::UAVCompute, ERHIAccess::IndirectArgs | ERHIAccess::SRVCompute);//Ray trace dispatches read these counts and use them for indirect dispatches.
+	// skip everything if we don't have a ray tracing scene
+	if (HasRayTracingScene)
+	{
+		RayRequests.Transition(RHICmdList, ERHIAccess::UAVCompute, ERHIAccess::SRVCompute);//Ray trace dispatches will read the ray requests.
+		RayTraceIntersections.Transition(RHICmdList, ERHIAccess::SRVCompute, ERHIAccess::UAVCompute);//Ray trace dispatches will write new results.
+		RayTraceCounts.Transition(RHICmdList, ERHIAccess::UAVCompute, ERHIAccess::IndirectArgs | ERHIAccess::SRVCompute);//Ray trace dispatches read these counts and use them for indirect dispatches.
 	
-	IssueRayTraces(RHICmdList, Scene);
+		IssueRayTraces(RHICmdList, Scene);
 
-	RayRequests.Transition(RHICmdList, ERHIAccess::SRVCompute, ERHIAccess::UAVCompute);//Next frame, simulations will write new ray trace requests.
-	RayTraceIntersections.Transition(RHICmdList, ERHIAccess::UAVCompute, ERHIAccess::SRVCompute);//Next frame these results will be read by simulations shaders.
-	RayTraceCounts.Transition(RHICmdList, ERHIAccess::IndirectArgs | ERHIAccess::SRVCompute, ERHIAccess::UAVCompute);//Next frame these counts will be written by simulation shaders.
+		RayRequests.Transition(RHICmdList, ERHIAccess::SRVCompute, ERHIAccess::UAVCompute);//Next frame, simulations will write new ray trace requests.
+		RayTraceIntersections.Transition(RHICmdList, ERHIAccess::UAVCompute, ERHIAccess::SRVCompute);//Next frame these results will be read by simulations shaders.
+		RayTraceCounts.Transition(RHICmdList, ERHIAccess::IndirectArgs | ERHIAccess::SRVCompute, ERHIAccess::UAVCompute);//Next frame these counts will be written by simulation shaders.
+	}
 }
 
 void FNiagaraRayTracingHelper::AddToDispatch(FNiagaraDataInterfaceProxy* DispatchKey, uint32 MaxRays, int32 MaxRetraces)
@@ -254,7 +258,7 @@ void FNiagaraRayTracingHelper::AddToDispatch(FNiagaraDataInterfaceProxy* Dispatc
 	Dispatch.MaxRetraces = MaxRetraces;
 }
 
-const FNiagaraRayTraceDispatchInfo& FNiagaraRayTracingHelper::GetDispatch(FNiagaraDataInterfaceProxy* DispatchKey, FRHICommandList& RHICmdList)
+void FNiagaraRayTracingHelper::BuildDispatch(FRHICommandList& RHICmdList, FNiagaraDataInterfaceProxy* DispatchKey)
 {
 	FNiagaraRayTraceDispatchInfo& Dispatch = Dispatches.FindChecked(DispatchKey);
 
@@ -288,11 +292,9 @@ const FNiagaraRayTraceDispatchInfo& FNiagaraRayTracingHelper::GetDispatch(FNiaga
 			Dispatch.LastFrameRayTraceIntersections = Dispatch.RayTraceIntersections;
 		}
 	}
-	
-	return Dispatch;
 }
 
-const FNiagaraRayTraceDispatchInfo& FNiagaraRayTracingHelper::GetDummyDispatch(FRHICommandList& RHICmdList)
+void FNiagaraRayTracingHelper::BuildDummyDispatch(FRHICommandList& RHICmdList)
 {
 	if (DummyDispatch.IsValid() == false)
 	{
@@ -314,6 +316,23 @@ const FNiagaraRayTraceDispatchInfo& FNiagaraRayTracingHelper::GetDummyDispatch(F
 		DummyDispatch.MaxRays = 0;
 		DummyDispatch.MaxRetraces = 0;
 	}
+}
+
+const FNiagaraRayTraceDispatchInfo& FNiagaraRayTracingHelper::GetDispatch(FNiagaraDataInterfaceProxy* DispatchKey) const
+{
+	const FNiagaraRayTraceDispatchInfo& Dispatch = Dispatches.FindChecked(DispatchKey);
+
+	// make sure that the buffer is valid (that BuildDispatch has been called)
+	check(Dispatch.RayRequests.IsValid());
+
+	return Dispatch;
+}
+
+const FNiagaraRayTraceDispatchInfo& FNiagaraRayTracingHelper::GetDummyDispatch() const
+{
+	// make sure that the buffers are valid (that BuildDummyDispatch has been called)
+	check(DummyDispatch.RayRequests.IsValid());
+
 	return DummyDispatch;
 }
 
diff --git a/Engine/Plugins/FX/Niagara/Source/NiagaraShader/Public/NiagaraRayTracingHelper.h b/Engine/Plugins/FX/Niagara/Source/NiagaraShader/Public/NiagaraRayTracingHelper.h
--- a/Engine/Plugins/FX/Niagara/Source/NiagaraShader/Public/NiagaraRayTracingHelper.h
+++ b/Engine/Plugins/FX/Niagara/Source/NiagaraShader/Public/NiagaraRayTracingHelper.h
@@ -451,7 +451,7 @@ struct FNiagaraRayTraceDispatchInfo
 	/** Max number or times rays can be re-traced in a shader if they hit something that is invalid/filtered. */
 	uint32 MaxRetraces = 0;
 
-	FORCEINLINE bool IsValid(){ return MaxRays > 0; }
+	FORCEINLINE bool IsValid() const { return MaxRays > 0; }
 	FORCEINLINE void Reset()
 	{
 		RayRequests = FGPUScratchPadStructured<FNiagaraRayData>::FAllocation();
@@ -476,17 +476,23 @@ public:
 		void IssueRayTraces(FRHICommandList& RHICmdList, FScene* Scene, const FIntPoint& RayTraceCounts, uint32 MaxRetraces, FRHIShaderResourceView* RayTraceRequests, FRWBuffer* IndirectArgsBuffer, uint32 IndirectArgsOffset, FRHIUnorderedAccessView* RayTraceResults);
 		bool IsValid() const;
 		
-		void BeginFrame(FRHICommandList& RHICmdList);
-		void EndFrame(FRHICommandList& RHICmdList, FScene* Scene);
+		void BeginFrame(FRHICommandList& RHICmdList, bool HasRayTracingScene);
+		void EndFrame(FRHICommandList& RHICmdList, bool HasRayTracingScene, FScene* Scene);
 
 		/** Accumulates ray requests from user DIs into single dispatches per DI. */
 		void AddToDispatch(FNiagaraDataInterfaceProxy* DispatchKey, uint32 MaxRays, int32 MaxRetraces);
-		
-		/** Returns the final buffers for each DI for use in simulations and RT dispatches. This will allocated buffers where needed. */
-		const FNiagaraRayTraceDispatchInfo& GetDispatch(FNiagaraDataInterfaceProxy* DispatchKey, FRHICommandList& RHICmdList);
+
+		/** Ensures that the final buffers for the provided proxy have been allocated. */
+		void BuildDispatch(FRHICommandList& RHICmdList, FNiagaraDataInterfaceProxy* DispatchKey);
+
+		/** Ensures that buffers for a dummy buffer have been allocated. */
+		void BuildDummyDispatch(FRHICommandList& RHICmdList);
+
+		/** Returns the final buffers for each DI for use in simulations and RT dispatches. */
+		const FNiagaraRayTraceDispatchInfo& GetDispatch(FNiagaraDataInterfaceProxy* DispatchKey) const;
 
 		/** Returns a dummy dispatch to allow shaders to still bind to valid buffers. */
-		const FNiagaraRayTraceDispatchInfo& GetDummyDispatch(FRHICommandList& RHICmdList);
+		const FNiagaraRayTraceDispatchInfo& GetDummyDispatch() const;
 
 		/** Adds a primitive to a collision group. This data is sent to the GPU on the start of the next frame. */
 		void SetPrimitiveCollisionGroup(FPrimitiveSceneInfo& Primitive, uint32 CollisionGroup);
