The previous IsValid check checked on the existence of a SystemInstance instead of the system itself.
This fixes an issue with data interfaces not properly displaying the UI for user parameter bindings since an if check failed due to this.

diff --git a/Engine/Plugins/FX/Niagara/Source/NiagaraEditor/Private/ViewModels/NiagaraSystemViewModel.cpp b/Engine/Plugins/FX/Niagara/Source/NiagaraEditor/Private/ViewModels/NiagaraSystemViewModel.cpp
--- a/Engine/Plugins/FX/Niagara/Source/NiagaraEditor/Private/ViewModels/NiagaraSystemViewModel.cpp
+++ b/Engine/Plugins/FX/Niagara/Source/NiagaraEditor/Private/ViewModels/NiagaraSystemViewModel.cpp
@@ -143,7 +143,7 @@ void FNiagaraSystemViewModel::Initialize(UNiagaraSystem& InSystem, FNiagaraSyste
 
 bool FNiagaraSystemViewModel::IsValid() const
 {
-	return SystemInstance != nullptr;
+	return System != nullptr;
 }
 
 void FNiagaraSystemViewModel::DumpToText(FString& ExportText)
