Updated editor style to include the same style for odd and even rows. This fixes an issue where every second row in the Niagara asset windows was white when hovered.

diff --git a/Engine/Plugins/FX/Niagara/Source/NiagaraEditor/Private/NiagaraEditorStyle.cpp b/Engine/Plugins/FX/Niagara/Source/NiagaraEditor/Private/NiagaraEditorStyle.cpp
--- a/Engine/Plugins/FX/Niagara/Source/NiagaraEditor/Private/NiagaraEditorStyle.cpp
+++ b/Engine/Plugins/FX/Niagara/Source/NiagaraEditor/Private/NiagaraEditorStyle.cpp
@@ -169,13 +169,13 @@ void InitActionMenu(TSharedRef< FSlateStyleSet > Style)
 	const FTableRowStyle ActionMenuRowStyle = FTableRowStyle()
             .SetEvenRowBackgroundBrush(FSlateNoResource())
             .SetOddRowBackgroundBrush(FSlateNoResource())
-            .SetOddRowBackgroundHoveredBrush(IMAGE_CORE_BRUSH("Common/Selection", Icon8x8, FLinearColor(1.0f, 1.0f, 1.0f, 0.1f)))
+			.SetEvenRowBackgroundHoveredBrush(IMAGE_CORE_BRUSH("Common/Selection", Icon8x8, FLinearColor(1.0f, 1.0f, 1.0f, 0.1f)))
+			.SetOddRowBackgroundHoveredBrush(IMAGE_CORE_BRUSH("Common/Selection", Icon8x8, FLinearColor(1.0f, 1.0f, 1.0f, 0.1f)))
             .SetSelectorFocusedBrush(BORDER_CORE_BRUSH("Common/Selector", FMargin(4.f / 16.f), FStarshipCoreStyle::GetCoreStyle().GetSlateColor("SelectorColor")))
             .SetActiveBrush(FSlateColorBrush(FStyleColors::Select))
             .SetActiveHoveredBrush(FSlateColorBrush(FStyleColors::Select))
-			// note: inactive brush is used for selections that don't have keyboard focus. SelectionColor_Inactive would be gray. We use the normal SelectionColor here
             .SetInactiveBrush(FSlateColorBrush(FStyleColors::SelectInactive))
-            .SetInactiveHoveredBrush(FSlateColorBrush(FStyleColors::SelectInactive))
+            .SetInactiveHoveredBrush(FSlateColorBrush(FStyleColors::SelectHover))
             .SetActiveHighlightedBrush(FSlateColorBrush(FStyleColors::PrimaryHover))
             .SetInactiveHighlightedBrush(FSlateColorBrush(FStyleColors::SelectParent))
 			.SetTextColor(FStyleColors::Foreground)
