- Grid3DCollection reader : Inherits from the Grid3DCollection itself rather than being a copy of the DI (like the 2D reader)
- Small changes required for the Grid3DCollection to allow for this inheritance hierarchy
- Small test case

diff --git a/Engine/Plugins/FX/Niagara/Source/Niagara/Classes/NiagaraDataInterfaceGrid3DCollection.h b/Engine/Plugins/FX/Niagara/Source/Niagara/Classes/NiagaraDataInterfaceGrid3DCollection.h
--- a/Engine/Plugins/FX/Niagara/Source/Niagara/Classes/NiagaraDataInterfaceGrid3DCollection.h
+++ b/Engine/Plugins/FX/Niagara/Source/Niagara/Classes/NiagaraDataInterfaceGrid3DCollection.h
@@ -54,6 +54,12 @@ struct FGrid3DCollectionRWInstanceData_GameThread
 	TArray<FNiagaraVariableBase> Vars;
 	TArray<uint32> Offsets;
 
+	// We need to essentially make this a linked list to avoid more refactoring for now
+	// eventually we can clean this logic up, but this allows us to have a subclass that
+	// overrides the render thread data, which in this case is for a grid reader
+	UNiagaraDataInterface* OtherDI = nullptr;
+	FGrid3DCollectionRWInstanceData_GameThread* OtherInstanceData = nullptr;
+
 	int32 FindAttributeIndexByName(const FName& InName, int32 NumChannels);
 	bool UpdateTargetTexture(ENiagaraGpuBufferFormat BufferFormat);
 };
@@ -85,6 +91,13 @@ struct FGrid3DCollectionRWInstanceData_RenderThread
 
 	FTextureRHIRef RenderTargetToCopyTo;
 
+	TArray<FName> AttributeNames;
+
+	// We need to essentially make this a linked list to avoid more refactoring for now
+	// eventually we can clean this logic up, but this allows us to have a subclass that
+	// overrides the render thread data, which in this case is for a grid reader
+	FNiagaraDataInterfaceProxy* OtherProxy = nullptr;
+
 	void BeginSimulate(FRHICommandList& RHICmdList);
 	void EndSimulate(FRHICommandList& RHICmdList);
 };
@@ -105,6 +118,36 @@ struct FNiagaraDataInterfaceProxyGrid3DCollectionProxy : public FNiagaraDataInte
 	TMap<FNiagaraSystemInstanceID, FGrid3DCollectionRWInstanceData_RenderThread> SystemInstancesToProxyData_RT;
 };
 
+struct FNiagaraDataInterfaceParametersCS_Grid3DCollection : public FNiagaraDataInterfaceParametersCS
+{
+	DECLARE_TYPE_LAYOUT(FNiagaraDataInterfaceParametersCS_Grid3DCollection, NonVirtual);
+public:
+	void Bind(const FNiagaraDataInterfaceGPUParamInfo& ParameterInfo, const class FShaderParameterMap& ParameterMap);
+	void Set(FRHICommandList& RHICmdList, const FNiagaraDataInterfaceSetArgs& Context) const;
+	void Unset(FRHICommandList& RHICmdList, const FNiagaraDataInterfaceSetArgs& Context) const;
+
+private:
+	LAYOUT_FIELD(FShaderParameter, NumAttributesParam);
+	LAYOUT_FIELD(FShaderParameter, NumNamedAttributesParam);
+	LAYOUT_FIELD(FShaderParameter, UnitToUVParam);
+	LAYOUT_FIELD(FShaderParameter, NumCellsParam);
+	LAYOUT_FIELD(FShaderParameter, NumTilesParam);
+	LAYOUT_FIELD(FShaderParameter, OneOverNumTilesParam);
+	LAYOUT_FIELD(FShaderParameter, UnitClampMinParam);
+	LAYOUT_FIELD(FShaderParameter, UnitClampMaxParam);
+	LAYOUT_FIELD(FShaderParameter, CellSizeParam);
+	LAYOUT_FIELD(FShaderParameter, WorldBBoxSizeParam);
+
+	LAYOUT_FIELD(FShaderResourceParameter, GridParam);
+	LAYOUT_FIELD(FRWShaderParameter, OutputGridParam);
+	LAYOUT_FIELD(FShaderParameter, AttributeIndicesParam);
+	LAYOUT_FIELD(FShaderResourceParameter, PerAttributeDataParam);
+
+	LAYOUT_FIELD(FShaderResourceParameter, SamplerParam);
+	LAYOUT_FIELD(TMemoryImageArray<FName>, AttributeNames);
+	LAYOUT_FIELD(TMemoryImageArray<uint32>, AttributeChannelCount);
+};
+
 UCLASS(EditInlineNew, Category = "Grid", meta = (DisplayName = "Grid3D Collection", Experimental), Blueprintable, BlueprintType)
 class NIAGARA_API UNiagaraDataInterfaceGrid3DCollection : public UNiagaraDataInterfaceGrid3D
 {
@@ -270,6 +313,8 @@ public:
 	static int32 GetComponentCountFromFuncName(const FName& FuncName);
 	static FNiagaraTypeDefinition GetValueTypeFromFuncName(const FName& FuncName);
 	static bool CanCreateVarFromFuncName(const FName& FuncName);
+
+	TMap<FNiagaraSystemInstanceID, FGrid3DCollectionRWInstanceData_GameThread*>& GetSystemInstancesToProxyData_GT() { return SystemInstancesToProxyData_GT; }
 protected:
 #if WITH_EDITORONLY_DATA
 	void WriteSetHLSL(const FNiagaraDataInterfaceGPUParamInfo& ParamInfo, const FNiagaraDataInterfaceGeneratedFunction& FunctionInfo, int FunctionInstanceIndex, int32 InNumChannels, FString& OutHLSL);
@@ -291,3 +336,4 @@ protected:
 
 	static FNiagaraVariableBase ExposedRTVar;
 };
+
diff --git a/Engine/Plugins/FX/Niagara/Source/Niagara/Classes/NiagaraDataInterfaceGrid3DCollectionReader.h b/Engine/Plugins/FX/Niagara/Source/Niagara/Classes/NiagaraDataInterfaceGrid3DCollectionReader.h
new file mode 100644
--- /dev/null
+++ b/Engine/Plugins/FX/Niagara/Source/Niagara/Classes/NiagaraDataInterfaceGrid3DCollectionReader.h
@@ -0,0 +1,43 @@
+// Copyright Epic Games, Inc. All Rights Reserved.
+#pragma once
+
+#include "NiagaraDataInterface.h"
+#include "NiagaraDataInterfaceRW.h"
+#include "ClearQuad.h"
+#include "NiagaraComponent.h"
+#include "Niagara/Private/NiagaraStats.h"
+#include "NiagaraDataInterfaceGrid3DCollection.h"
+
+#include "NiagaraDataInterfaceGrid3DCollectionReader.generated.h"
+
+class FNiagaraSystemInstance;
+class UTextureRenderTarget;
+class UTextureRenderTargetVolume;
+
+UCLASS(EditInlineNew, Category = "Grid", hidecategories=(Grid3DCollection,Grid, Deprecated), meta = (DisplayName = "Grid3D Collection Reader", Experimental), Blueprintable, BlueprintType)
+class NIAGARA_API UNiagaraDataInterfaceGrid3DCollectionReader : public UNiagaraDataInterfaceGrid3DCollection
+{
+	GENERATED_UCLASS_BODY()
+
+public:
+	DECLARE_NIAGARA_DI_PARAMETER();
+
+	// Name of the emitter to read from
+	UPROPERTY(EditAnywhere, Category = "Reader")
+	FString EmitterName;
+
+	// Name of the Grid2DCollection Data Interface on the emitter
+	UPROPERTY(EditAnywhere, Category = "Reader")
+	FString DIName;
+
+	virtual bool Equals(const UNiagaraDataInterface* Other) const override;
+	virtual bool InitPerInstanceData(void* PerInstanceData, FNiagaraSystemInstance* SystemInstance) override;
+	virtual void GetEmitterDependencies(UNiagaraSystem* Asset, TArray<UNiagaraEmitter*>& Dependencies) const override;
+	virtual void GetFunctions(TArray<FNiagaraFunctionSignature>& OutFunctions) override;
+	virtual bool PerInstanceTickPostSimulate(void* PerInstanceData, FNiagaraSystemInstance* SystemInstance, float DeltaSeconds) override;
+
+protected:
+	//~ UNiagaraDataInterface interface
+	virtual bool CopyToInternal(UNiagaraDataInterface* Destination) const override;
+	//~ UNiagaraDataInterface interface END
+};
\ No newline at end of file
diff --git a/Engine/Plugins/FX/Niagara/Source/Niagara/Private/NiagaraDataInterfaceGrid3DCollection.cpp b/Engine/Plugins/FX/Niagara/Source/Niagara/Private/NiagaraDataInterfaceGrid3DCollection.cpp
--- a/Engine/Plugins/FX/Niagara/Source/Niagara/Private/NiagaraDataInterfaceGrid3DCollection.cpp
+++ b/Engine/Plugins/FX/Niagara/Source/Niagara/Private/NiagaraDataInterfaceGrid3DCollection.cpp
@@ -143,11 +143,8 @@ static FAutoConsoleVariableRef CVarNiagaraGrid3DOverrideFormat(
 );
 
 /*--------------------------------------------------------------------------------------------------------------------------*/
-struct FNiagaraDataInterfaceParametersCS_Grid3DCollection : public FNiagaraDataInterfaceParametersCS
-{
-	DECLARE_TYPE_LAYOUT(FNiagaraDataInterfaceParametersCS_Grid3DCollection, NonVirtual);
-public:
-	void Bind(const FNiagaraDataInterfaceGPUParamInfo& ParameterInfo, const class FShaderParameterMap& ParameterMap)
+
+void FNiagaraDataInterfaceParametersCS_Grid3DCollection::Bind(const FNiagaraDataInterfaceGPUParamInfo& ParameterInfo, const class FShaderParameterMap& ParameterMap)
 	{
 		NumAttributesParam.Bind(ParameterMap, *(UNiagaraDataInterfaceRWBase::NumAttributesName + ParameterInfo.DataInterfaceHLSLSymbol));
 		NumNamedAttributesParam.Bind(ParameterMap, *(UNiagaraDataInterfaceRWBase::NumNamedAttributesName + ParameterInfo.DataInterfaceHLSLSymbol));
@@ -194,8 +191,7 @@ public:
 			}
 		}
 	}
-
-	void Set(FRHICommandList& RHICmdList, const FNiagaraDataInterfaceSetArgs& Context) const
+void FNiagaraDataInterfaceParametersCS_Grid3DCollection::Set(FRHICommandList& RHICmdList, const FNiagaraDataInterfaceSetArgs& Context) const
 	{
 		check(IsInRenderingThread());
 
@@ -204,7 +200,15 @@ public:
 		FNiagaraDataInterfaceProxyGrid3DCollectionProxy* VFDI = static_cast<FNiagaraDataInterfaceProxyGrid3DCollectionProxy*>(Context.DataInterface);
 
 		FGrid3DCollectionRWInstanceData_RenderThread* ProxyData = VFDI->SystemInstancesToProxyData_RT.Find(Context.SystemInstanceID);
-		check(ProxyData);
+		FGrid3DCollectionRWInstanceData_RenderThread* OriginalProxyData = ProxyData;
+
+		check(ProxyData);		
+		if (ProxyData->OtherProxy != nullptr)
+		{			
+			FNiagaraDataInterfaceProxyGrid3DCollectionProxy* OtherGrid3DProxy = static_cast<FNiagaraDataInterfaceProxyGrid3DCollectionProxy*>(ProxyData->OtherProxy);
+			ProxyData = OtherGrid3DProxy->SystemInstancesToProxyData_RT.Find(Context.SystemInstanceID);
+			check(ProxyData);
+		}
 
 		SetShaderValue(RHICmdList, ComputeShaderRHI, NumCellsParam, ProxyData->NumCells);	
 		SetShaderValue(RHICmdList, ComputeShaderRHI, NumTilesParam, ProxyData->NumTiles);
@@ -214,10 +218,10 @@ public:
 		SetShaderValue(RHICmdList, ComputeShaderRHI, UnitClampMinParam, HalfPixelOffset);
 		SetShaderValue(RHICmdList, ComputeShaderRHI, UnitClampMaxParam, FVector3f::OneVector - HalfPixelOffset);
 
-		if (ProxyData->AttributeIndices.Num() == 0 && AttributeNames.Num() > 0)
+		if (OriginalProxyData->AttributeIndices.Num() == 0 && AttributeNames.Num() > 0)
 		{
 			int NumAttrIndices = Align(AttributeNames.Num(), 4);
-			ProxyData->AttributeIndices.SetNumZeroed(NumAttrIndices);
+			OriginalProxyData->AttributeIndices.SetNumZeroed(NumAttrIndices);
 
 			// TODO handle mismatched types!
 			for (int32 i = 0; i < AttributeNames.Num(); i++)
@@ -228,11 +232,11 @@ public:
 				check(ProxyData->Offsets.Num() == ProxyData->Vars.Num());
 				if (ProxyData->Offsets.IsValidIndex(FoundIdx) && AttributeChannelCount[i] == ProxyData->VarComponents[FoundIdx])
 				{
-					ProxyData->AttributeIndices[i] = ProxyData->Offsets[FoundIdx];
+					OriginalProxyData->AttributeIndices[i] = ProxyData->Offsets[FoundIdx];
 				}
 				else
 				{
-					ProxyData->AttributeIndices[i] = -1; // We may need to protect against this in the hlsl as this might underflow an array lookup if used incorrectly.
+					OriginalProxyData->AttributeIndices[i] = -1; // We may need to protect against this in the hlsl as this might underflow an array lookup if used incorrectly.
 				}
 			}
 		}
@@ -246,7 +250,7 @@ public:
 
 		SetShaderValue(RHICmdList, ComputeShaderRHI, WorldBBoxSizeParam, (FVector3f)ProxyData->WorldBBoxSize);
 
-		SetShaderValueArray(RHICmdList, ComputeShaderRHI, AttributeIndicesParam, ProxyData->AttributeIndices.GetData(), ProxyData->AttributeIndices.Num());
+		SetShaderValueArray(RHICmdList, ComputeShaderRHI, AttributeIndicesParam, OriginalProxyData->AttributeIndices.GetData(), OriginalProxyData->AttributeIndices.Num());
 		
 		SetSRVParameter(RHICmdList, ComputeShaderRHI, PerAttributeDataParam, ProxyData->PerAttributeData.SRV);
 
@@ -281,8 +285,7 @@ public:
 			RHICmdList.SetUAVParameter(ComputeShaderRHI, OutputGridParam.GetUAVIndex(), OutputGridUAV);
 		}
 	}
-
-	void Unset(FRHICommandList& RHICmdList, const FNiagaraDataInterfaceSetArgs& Context) const
+void FNiagaraDataInterfaceParametersCS_Grid3DCollection::Unset(FRHICommandList& RHICmdList, const FNiagaraDataInterfaceSetArgs& Context) const
 	{
 		if (OutputGridParam.IsBound())
 		{
@@ -291,28 +294,6 @@ public:
 		}
 	}
 
-private:
-	LAYOUT_FIELD(FShaderParameter, NumAttributesParam);
-	LAYOUT_FIELD(FShaderParameter, NumNamedAttributesParam);
-	LAYOUT_FIELD(FShaderParameter, UnitToUVParam);
-	LAYOUT_FIELD(FShaderParameter, NumCellsParam);
-	LAYOUT_FIELD(FShaderParameter, NumTilesParam);
-	LAYOUT_FIELD(FShaderParameter, OneOverNumTilesParam);
-	LAYOUT_FIELD(FShaderParameter, UnitClampMinParam);
-	LAYOUT_FIELD(FShaderParameter, UnitClampMaxParam);
-	LAYOUT_FIELD(FShaderParameter, CellSizeParam);
-	LAYOUT_FIELD(FShaderParameter, WorldBBoxSizeParam);
-
-	LAYOUT_FIELD(FShaderResourceParameter, GridParam);
-	LAYOUT_FIELD(FRWShaderParameter, OutputGridParam);
-	LAYOUT_FIELD(FShaderParameter, AttributeIndicesParam);
-	LAYOUT_FIELD(FShaderResourceParameter, PerAttributeDataParam);
-
-	LAYOUT_FIELD(FShaderResourceParameter, SamplerParam);
-	LAYOUT_FIELD(TMemoryImageArray<FName>, AttributeNames);
-	LAYOUT_FIELD(TMemoryImageArray<uint32>, AttributeChannelCount);
-};
-
 IMPLEMENT_TYPE_LAYOUT(FNiagaraDataInterfaceParametersCS_Grid3DCollection);
 
 IMPLEMENT_NIAGARA_DI_PARAMETER(UNiagaraDataInterfaceGrid3DCollection, FNiagaraDataInterfaceParametersCS_Grid3DCollection);
@@ -1059,11 +1040,21 @@ void UNiagaraDataInterfaceGrid3DCollection::GetNumCells(FVectorVMExternalFunctio
 	FNDIOutputParam<int32> NumCellsY(Context);
 	FNDIOutputParam<int32> NumCellsZ(Context);
 
+	int32 TmpNumCellsX = InstData->NumCells.X;
+	int32 TmpNumCellsY = InstData->NumCells.Y;
+	int32 TmpNumCellsZ = InstData->NumCells.Z;
+	if (InstData->OtherInstanceData != nullptr)
+	{
+		TmpNumCellsX = InstData->OtherInstanceData->NumCells.X;
+		TmpNumCellsY = InstData->OtherInstanceData->NumCells.Y;
+		TmpNumCellsZ = InstData->OtherInstanceData->NumCells.Z;
+	}
+
 	for (int32 InstanceIdx = 0; InstanceIdx < Context.GetNumInstances(); ++InstanceIdx)
 	{
-		NumCellsX.SetAndAdvance(InstData->NumCells.X);
-		NumCellsY.SetAndAdvance(InstData->NumCells.Y);
-		NumCellsZ.SetAndAdvance(InstData->NumCells.Z);
+		NumCellsX.SetAndAdvance(TmpNumCellsX);
+		NumCellsY.SetAndAdvance(TmpNumCellsY);
+		NumCellsZ.SetAndAdvance(TmpNumCellsZ);
 	}
 }
 
@@ -2820,7 +2811,17 @@ FIntVector FNiagaraDataInterfaceProxyGrid3DCollectionProxy::GetElementCount(FNia
 {
 	if ( const FGrid3DCollectionRWInstanceData_RenderThread* ProxyData = SystemInstancesToProxyData_RT.Find(SystemInstanceID) )
 	{
-		return ProxyData->NumCells;
+		// support a grid reader acting as an iteration source
+		if (ProxyData->OtherProxy != nullptr)
+		{
+			FNiagaraDataInterfaceProxyGrid3DCollectionProxy *OtherGrid3DProxy = static_cast<FNiagaraDataInterfaceProxyGrid3DCollectionProxy*>(ProxyData->OtherProxy);
+			const FGrid3DCollectionRWInstanceData_RenderThread* OtherProxyData = OtherGrid3DProxy->SystemInstancesToProxyData_RT.Find(SystemInstanceID);
+			return OtherProxyData->NumCells;
+		}
+		else
+		{
+			return ProxyData->NumCells;
+		}
 	}
 	return FIntVector::ZeroValue;
 }
diff --git a/Engine/Plugins/FX/Niagara/Source/Niagara/Private/NiagaraDataInterfaceGrid3DCollectionReader.cpp b/Engine/Plugins/FX/Niagara/Source/Niagara/Private/NiagaraDataInterfaceGrid3DCollectionReader.cpp
new file mode 100644
--- /dev/null
+++ b/Engine/Plugins/FX/Niagara/Source/Niagara/Private/NiagaraDataInterfaceGrid3DCollectionReader.cpp
@@ -0,0 +1,210 @@
+// Copyright Epic Games, Inc. All Rights Reserved.
+#include "NiagaraDataInterfaceGrid3DCollectionReader.h"
+#include "NiagaraShader.h"
+#include "ShaderParameterUtils.h"
+#include "ClearQuad.h"
+#include "TextureResource.h"
+#include "NiagaraEmitterInstanceBatcher.h"
+#include "NiagaraSystemInstance.h"
+#include "NiagaraRenderer.h"
+#include "Engine/VolumeTexture.h"
+#include "Engine/TextureRenderTargetVolume.h"
+#include "NiagaraSettings.h"
+#include "NiagaraConstants.h"
+#if WITH_EDITOR
+#include "NiagaraGpuComputeDebug.h"
+#endif
+
+#define LOCTEXT_NAMESPACE "NiagaraDataInterfaceGrid3DCollectionReader"
+
+IMPLEMENT_NIAGARA_DI_PARAMETER(UNiagaraDataInterfaceGrid3DCollectionReader, FNiagaraDataInterfaceParametersCS_Grid3DCollection);
+
+UNiagaraDataInterfaceGrid3DCollectionReader::UNiagaraDataInterfaceGrid3DCollectionReader(FObjectInitializer const& ObjectInitializer)
+	: Super(ObjectInitializer)
+{
+	Proxy.Reset(new FNiagaraDataInterfaceProxyGrid3DCollectionProxy());
+
+	FNiagaraTypeDefinition Def(UTextureRenderTarget::StaticClass());
+	RenderTargetUserParameter.Parameter.SetType(Def);
+}
+
+bool UNiagaraDataInterfaceGrid3DCollectionReader::Equals(const UNiagaraDataInterface* Other) const
+{
+	if (!Super::Equals(Other))
+	{
+		return false;
+	}
+	const UNiagaraDataInterfaceGrid3DCollectionReader* OtherTyped = CastChecked<const UNiagaraDataInterfaceGrid3DCollectionReader>(Other);
+
+	return OtherTyped != nullptr && OtherTyped->EmitterName == EmitterName && OtherTyped->DIName == DIName;
+}
+
+bool UNiagaraDataInterfaceGrid3DCollectionReader::CopyToInternal(UNiagaraDataInterface* Destination) const
+{
+	if (!Super::CopyToInternal(Destination))
+	{
+		return false;
+	}
+
+	UNiagaraDataInterfaceGrid3DCollectionReader* OtherTyped = CastChecked<UNiagaraDataInterfaceGrid3DCollectionReader>(Destination);
+	OtherTyped->EmitterName = EmitterName;
+	OtherTyped->DIName = DIName;
+	return true;
+}
+
+bool UNiagaraDataInterfaceGrid3DCollectionReader::InitPerInstanceData(void* PerInstanceData, FNiagaraSystemInstance* SystemInstance)
+{
+	check(Proxy);
+	
+	FGrid3DCollectionRWInstanceData_GameThread* InstanceData = new (PerInstanceData) FGrid3DCollectionRWInstanceData_GameThread();
+	SystemInstancesToProxyData_GT.Emplace(SystemInstance->GetId(), InstanceData);
+
+	TSharedPtr<FNiagaraEmitterInstance, ESPMode::ThreadSafe> RT_EmitterInstance;
+
+	
+	FNiagaraEmitterInstance* EmitterInstanceToUse = nullptr;
+	for (TSharedPtr<FNiagaraEmitterInstance, ESPMode::ThreadSafe> EmitterInstance : SystemInstance->GetEmitters())
+	{
+		UNiagaraEmitter* Emitter = EmitterInstance->GetCachedEmitter();
+		if (Emitter == nullptr)
+		{
+			continue;
+		}
+
+		if (EmitterName == Emitter->GetUniqueEmitterName())
+		{
+			EmitterInstanceToUse = EmitterInstance.Get();
+			break;
+		}
+	}
+
+	// Look for proxy we are going to use
+	FNiagaraDataInterfaceProxy* ProxyToUse = nullptr;
+	UNiagaraDataInterfaceGrid3DCollection* GridInterface = nullptr;
+
+	if (EmitterInstanceToUse != nullptr)
+	{
+		FNiagaraComputeExecutionContext* ExecContext = EmitterInstanceToUse->GetGPUContext();
+		if ((ExecContext != nullptr) && (ExecContext->GPUScript != nullptr))
+		{
+			const TArray<FNiagaraScriptDataInterfaceCompileInfo>& DataInterfaceInfo = ExecContext->GPUScript->GetVMExecutableData().DataInterfaceInfo;
+			const TArray<UNiagaraDataInterface*>& DataInterfaces = ExecContext->CombinedParamStore.GetDataInterfaces();
+
+			FString FullName = FString("Emitter.") + DIName;
+			int Index = 0;
+
+			// #todo(dmp): we are looking at the UObjects that define the DIs here 
+			
+			for (UNiagaraDataInterface* Interface : DataInterfaces)
+			{
+				if (DataInterfaceInfo[Index].Name.GetPlainNameString() == FullName)
+				{
+					GridInterface = static_cast<UNiagaraDataInterfaceGrid3DCollection*>(Interface);
+					ProxyToUse = Interface->GetProxy();
+					break;
+				}
+				++Index;
+			}
+		}
+
+		if (GridInterface && ProxyToUse)
+		{
+
+			FGrid3DCollectionRWInstanceData_GameThread* Grid3DInstanceData = GridInterface->GetSystemInstancesToProxyData_GT().FindRef(SystemInstance->GetId());
+			InstanceData->OtherDI = GridInterface;
+			InstanceData->OtherInstanceData = Grid3DInstanceData;
+
+			// Push Updates to Proxy.
+			FNiagaraDataInterfaceProxyGrid3DCollectionProxy* RT_Proxy = GetProxyAs<FNiagaraDataInterfaceProxyGrid3DCollectionProxy>();
+			ENQUEUE_RENDER_COMMAND(FUpdateData)(
+				[RT_Proxy, InstanceID = SystemInstance->GetId(), RT_ProxyToUse = ProxyToUse](FRHICommandListImmediate& RHICmdList)
+			{
+				check(!RT_Proxy->SystemInstancesToProxyData_RT.Contains(InstanceID));
+				//check(RT_ProxyToUse->SystemInstancesToProxyData_RT.Contains(InstanceID));
+				//FGrid3DCollectionRWInstanceData_RenderThread* DataToRead = RT_ProxyToUse->SystemInstancesToProxyData_RT.Find(InstanceID);
+
+				FGrid3DCollectionRWInstanceData_RenderThread* TargetData = &RT_Proxy->SystemInstancesToProxyData_RT.Add(InstanceID);
+				TargetData->OtherProxy = RT_ProxyToUse;
+			});
+			return true;
+		}
+	}
+
+	return false;
+}
+
+void UNiagaraDataInterfaceGrid3DCollectionReader::GetEmitterDependencies(UNiagaraSystem* Asset, TArray<UNiagaraEmitter*>& Dependencies) const
+{
+	if (!Asset)
+	{
+		return;
+	}
+
+	UNiagaraEmitter* FoundSourceEmitter = nullptr;
+	for (const FNiagaraEmitterHandle& EmitterHandle : Asset->GetEmitterHandles())
+	{
+		UNiagaraEmitter* EmitterInstance = EmitterHandle.GetInstance();
+		if (EmitterInstance && EmitterInstance->GetUniqueEmitterName() == EmitterName)
+		{
+			Dependencies.Add(EmitterInstance);
+			return;
+		}
+	}
+}
+
+
+
+void UNiagaraDataInterfaceGrid3DCollectionReader::GetFunctions(TArray<FNiagaraFunctionSignature>& OutFunctions)
+{
+	int32 NumFunctionsBefore = OutFunctions.Num();
+
+	Super::GetFunctions(OutFunctions);
+	TSet<FName> FunctionsToRemove = 
+	{
+		ClearCellFunctionName,
+		CopyPreviousToCurrentForCellFunctionName,
+		CopyMaskedPreviousToCurrentForCellFunctionName,
+		SetValueFunctionName,
+		SetFullGridValueFunctionName,
+		SetVector4ValueFunctionName,
+		SetVector3ValueFunctionName,
+		SetVector2ValueFunctionName,
+		SetFloatValueFunctionName,
+		GetPreviousVector4ValueFunctionName,
+		SetVectorValueFunctionName,
+		SetVector2DValueFunctionName,
+		SetNumCellsFunctionName,
+		SetNumCellsFunctionName
+	};
+	
+	for (int32 ItemIndex = NumFunctionsBefore; ItemIndex < OutFunctions.Num();)
+	{
+		if (FunctionsToRemove.Find(OutFunctions[ItemIndex].Name) != nullptr)
+		{
+			OutFunctions.RemoveAtSwap(ItemIndex, 1);
+		}
+		else
+		{
+			++ItemIndex;
+		}
+	}	
+}
+
+bool UNiagaraDataInterfaceGrid3DCollectionReader::PerInstanceTickPostSimulate(void* PerInstanceData, FNiagaraSystemInstance* SystemInstance, float DeltaSeconds)
+{
+	FGrid3DCollectionRWInstanceData_GameThread* InstanceData = SystemInstancesToProxyData_GT.FindRef(SystemInstance->GetId());
+	
+	// we can run into the case where depending on the ordering of DI initialization, we might have not been able to grab the other grid's InstanceData
+	// in InitPerInstanceData.  If this is the case, we ensure it is correct here.
+	if (InstanceData->OtherInstanceData == nullptr && InstanceData->OtherDI)
+	{
+		UNiagaraDataInterfaceGrid3DCollection*  OtherGridDI = static_cast<UNiagaraDataInterfaceGrid3DCollection*>(InstanceData->OtherDI);
+		FGrid3DCollectionRWInstanceData_GameThread* Grid3DInstanceData = OtherGridDI->GetSystemInstancesToProxyData_GT().FindRef(SystemInstance->GetId());
+
+		InstanceData->OtherInstanceData = Grid3DInstanceData;
+	}
+
+	return false;
+}
+
+#undef LOCTEXT_NAMESPACE
