Fix a crash in batcher on nullptr for auto test.

diff --git a/Engine/Plugins/FX/Niagara/Source/Niagara/Private/NiagaraEmitterInstanceBatcher.cpp b/Engine/Plugins/FX/Niagara/Source/Niagara/Private/NiagaraEmitterInstanceBatcher.cpp
--- a/Engine/Plugins/FX/Niagara/Source/Niagara/Private/NiagaraEmitterInstanceBatcher.cpp
+++ b/Engine/Plugins/FX/Niagara/Source/Niagara/Private/NiagaraEmitterInstanceBatcher.cpp
@@ -391,7 +391,7 @@ void NiagaraEmitterInstanceBatcher::BuildConstantBuffers(FNiagaraGPUSystemTick&
 				if (BoundParameterCounts[InstanceTypeIt][InterpIt])
 				{
 					FRHIUniformBufferLayout* Layout = InstanceTypeIt == FNiagaraGPUSystemTick::UBT_Emitter ? EmitterCBufferLayout : Instance.Context->ExternalCBufferLayout;
-					if (Layout->Resources.Num() > 0 || Layout->ConstantBufferSize > 0)
+					if (Layout && (Layout->Resources.Num() > 0 || Layout->ConstantBufferSize > 0))
 					{
 						BufferRef = RHICreateUniformBuffer(
 							Tick.GetUniformBufferSource((FNiagaraGPUSystemTick::EUniformBufferType) InstanceTypeIt, &Instance, !InterpIt),
