Manual merge
Made stack expansion state persistent

diff --git a/Engine/Plugins/FX/Niagara/Source/NiagaraEditor/Private/NiagaraStackEditorData.cpp b/Engine/Plugins/FX/Niagara/Source/NiagaraEditor/Private/NiagaraStackEditorData.cpp
--- a/Engine/Plugins/FX/Niagara/Source/NiagaraEditor/Private/NiagaraStackEditorData.cpp
+++ b/Engine/Plugins/FX/Niagara/Source/NiagaraEditor/Private/NiagaraStackEditorData.cpp
@@ -21,10 +21,18 @@ bool UNiagaraStackEditorData::GetStackEntryIsExpanded(const FString& StackEntryK
 
 void UNiagaraStackEditorData::SetStackEntryIsExpanded(const FString& StackEntryKey, bool bIsExpanded)
 {
+	bool bBroadcast = false;
 	if (ensureMsgf(StackEntryKey.IsEmpty() == false, TEXT("Can not set the expanded state with an empty key")))
 	{
+		// we assume elements are expanded by default, so collapsing (bIsExpanded = false) will cause a broadcast
+		bBroadcast = GetStackEntryIsExpanded(StackEntryKey, true) != bIsExpanded;
 		StackEntryKeyToExpandedMap.FindOrAdd(StackEntryKey) = bIsExpanded;
 	}
+
+	if (bBroadcast)
+	{
+		OnPersistentDataChanged().Broadcast();
+	}
 }
 
 bool UNiagaraStackEditorData::GetStackEntryIsExpandedInOverview(const FString& StackEntryKey, bool bIsExpandedDefault) const
diff --git a/Engine/Plugins/FX/Niagara/Source/NiagaraEditor/Public/NiagaraStackEditorData.h b/Engine/Plugins/FX/Niagara/Source/NiagaraEditor/Public/NiagaraStackEditorData.h
--- a/Engine/Plugins/FX/Niagara/Source/NiagaraEditor/Public/NiagaraStackEditorData.h
+++ b/Engine/Plugins/FX/Niagara/Source/NiagaraEditor/Public/NiagaraStackEditorData.h
@@ -145,6 +145,7 @@ public:
 private:
 	TMap<FString, bool> StackEntryKeyToRenamePendingMap;
 
+	UPROPERTY()
 	TMap<FString, bool> StackEntryKeyToExpandedMap;
 
 	TMap<FString, bool> StackEntryKeyToPreSearchExpandedMap;
