- Allow data interfaces to opt out of running beyond PostActorTick
- This is important for the collision data interface as it relies on the async traces being issues this frame

diff --git a/Engine/Plugins/FX/Niagara/Source/Niagara/Classes/NiagaraDataInterface.h b/Engine/Plugins/FX/Niagara/Source/Niagara/Classes/NiagaraDataInterface.h
--- a/Engine/Plugins/FX/Niagara/Source/Niagara/Classes/NiagaraDataInterface.h
+++ b/Engine/Plugins/FX/Niagara/Source/Niagara/Classes/NiagaraDataInterface.h
@@ -368,6 +368,12 @@ public:
 	virtual bool HasPreSimulateTick() const { return false; }
 	virtual bool HasPostSimulateTick() const { return false; }
 
+	/**
+	When set to true the simulation may not complete in the same frame it started, allowing maximum overlap with the GameThread.
+	You must override and set to false if you require the data interface to complete before PostActorTick completes.
+	*/
+	virtual bool PostSimulateCanOverlapFrames() const { return true; }
+
 	virtual bool RequiresDistanceFieldData() const { return false; }
 	virtual bool RequiresDepthBuffer() const { return false; }
 	virtual bool RequiresEarlyViewData() const { return false; }
diff --git a/Engine/Plugins/FX/Niagara/Source/Niagara/Classes/NiagaraDataInterfaceAudioPlayer.h b/Engine/Plugins/FX/Niagara/Source/Niagara/Classes/NiagaraDataInterfaceAudioPlayer.h
--- a/Engine/Plugins/FX/Niagara/Source/Niagara/Classes/NiagaraDataInterfaceAudioPlayer.h
+++ b/Engine/Plugins/FX/Niagara/Source/Niagara/Classes/NiagaraDataInterfaceAudioPlayer.h
@@ -110,6 +110,7 @@ public:
 
 	virtual bool HasPreSimulateTick() const override { return true; }
 	virtual bool HasPostSimulateTick() const override { return true; }
+	virtual bool PostSimulateCanOverlapFrames() const { return false; }
 	//UNiagaraDataInterface Interface
 
 	virtual void PlayOneShotAudio(FVectorVMExternalFunctionContext& Context);
diff --git a/Engine/Plugins/FX/Niagara/Source/Niagara/Classes/NiagaraDataInterfaceCollisionQuery.h b/Engine/Plugins/FX/Niagara/Source/Niagara/Classes/NiagaraDataInterfaceCollisionQuery.h
--- a/Engine/Plugins/FX/Niagara/Source/Niagara/Classes/NiagaraDataInterfaceCollisionQuery.h
+++ b/Engine/Plugins/FX/Niagara/Source/Niagara/Classes/NiagaraDataInterfaceCollisionQuery.h
@@ -87,6 +87,9 @@ public:
 	
 	virtual bool HasPreSimulateTick() const override{ return true; }
 	virtual bool HasPostSimulateTick() const override { return true; }
+	virtual bool PostSimulateCanOverlapFrames() const { return false; }
+
+private:
 
 	virtual bool Equals(const UNiagaraDataInterface* Other) const override;
 
diff --git a/Engine/Plugins/FX/Niagara/Source/Niagara/Classes/NiagaraDataInterfaceExport.h b/Engine/Plugins/FX/Niagara/Source/Niagara/Classes/NiagaraDataInterfaceExport.h
--- a/Engine/Plugins/FX/Niagara/Source/Niagara/Classes/NiagaraDataInterfaceExport.h
+++ b/Engine/Plugins/FX/Niagara/Source/Niagara/Classes/NiagaraDataInterfaceExport.h
@@ -94,6 +94,7 @@ public:
 
 	virtual bool HasPreSimulateTick() const override { return true; }
 	virtual bool HasPostSimulateTick() const override { return true; }
+	virtual bool PostSimulateCanOverlapFrames() const { return false; }
 	//UNiagaraDataInterface Interface
 
 	virtual bool HasInternalAttributeReads(const UNiagaraEmitter* OwnerEmitter, const UNiagaraEmitter* Provider) const override { return OwnerEmitter == Provider; };
diff --git a/Engine/Plugins/FX/Niagara/Source/Niagara/Classes/NiagaraSystem.h b/Engine/Plugins/FX/Niagara/Source/Niagara/Classes/NiagaraSystem.h
--- a/Engine/Plugins/FX/Niagara/Source/Niagara/Classes/NiagaraSystem.h
+++ b/Engine/Plugins/FX/Niagara/Source/Niagara/Classes/NiagaraSystem.h
@@ -611,7 +611,8 @@ public:
 	bool bRequireCurrentFrameData = true;
 
 	bool HasSystemScriptDIsWithPerInstanceData() const;
-	FORCEINLINE bool HasDIsWithPostSimulateTick()const{ return bHasDIsWithPostSimulateTick; }
+	FORCEINLINE bool HasDIsWithPostSimulateTick() const { return bHasDIsWithPostSimulateTick; }
+	FORCEINLINE bool AllDIsPostSimulateCanOverlapFrames() const { return bAllDIsPostSimulateCanOverlapFrames; }
 	FORCEINLINE bool HasAnyGPUEmitters()const{ return bHasAnyGPUEmitters; }
 	FORCEINLINE bool NeedsGPUContextInitForDataInterfaces() const { return bNeedsGPUContextInitForDataInterfaces; }
 
@@ -852,6 +853,7 @@ protected:
 	mutable FString CrashReporterTag;
 
 	uint32 bHasDIsWithPostSimulateTick : 1;
+	uint32 bAllDIsPostSimulateCanOverlapFrames : 1;
 	uint32 bHasAnyGPUEmitters : 1;
 	uint32 bNeedsSortedSignificanceCull : 1;
 
diff --git a/Engine/Plugins/FX/Niagara/Source/Niagara/Private/NiagaraSystem.cpp b/Engine/Plugins/FX/Niagara/Source/Niagara/Private/NiagaraSystem.cpp
--- a/Engine/Plugins/FX/Niagara/Source/Niagara/Private/NiagaraSystem.cpp
+++ b/Engine/Plugins/FX/Niagara/Source/Niagara/Private/NiagaraSystem.cpp
@@ -115,6 +115,8 @@ UNiagaraSystem::UNiagaraSystem(const FObjectInitializer& ObjectInitializer)
 , WarmupTickDelta(1.0f / 15.0f)
 , bHasSystemScriptDIsWithPerInstanceData(false)
 , bNeedsGPUContextInitForDataInterfaces(false)
+, bHasDIsWithPostSimulateTick(false)
+, bAllDIsPostSimulateCanOverlapFrames(true)
 , bHasAnyGPUEmitters(false)
 , bNeedsSortedSignificanceCull(false)
 , ActiveInstances(0)
@@ -1958,6 +1960,7 @@ void UNiagaraSystem::UpdatePostCompileDIInfo()
 void UNiagaraSystem::UpdateDITickFlags()
 {
 	bHasDIsWithPostSimulateTick = false;
+	bAllDIsPostSimulateCanOverlapFrames = true;
 	auto CheckPostSimTick = [&](UNiagaraScript* Script)
 	{
 		if (Script)
@@ -1968,6 +1971,7 @@ void UNiagaraSystem::UpdateDITickFlags()
 				if (DefaultDataInterface && DefaultDataInterface->HasPostSimulateTick())
 				{
 					bHasDIsWithPostSimulateTick |= true;
+					bAllDIsPostSimulateCanOverlapFrames &= DefaultDataInterface->PostSimulateCanOverlapFrames();
 				}
 			}
 		}
diff --git a/Engine/Plugins/FX/Niagara/Source/Niagara/Private/NiagaraSystemSimulation.cpp b/Engine/Plugins/FX/Niagara/Source/Niagara/Private/NiagaraSystemSimulation.cpp
--- a/Engine/Plugins/FX/Niagara/Source/Niagara/Private/NiagaraSystemSimulation.cpp
+++ b/Engine/Plugins/FX/Niagara/Source/Niagara/Private/NiagaraSystemSimulation.cpp
@@ -1285,7 +1285,14 @@ void FNiagaraSystemSimulation::Tick_GameThread(float DeltaSeconds, const FGraphE
 		}
 		else
 		{
-			WorldManager->MarkSimulationsForEndOfFrameWait(this);
+			if ( System->AllDIsPostSimulateCanOverlapFrames() )
+			{
+				WorldManager->MarkSimulationsForEndOfFrameWait(this);
+			}
+			else
+			{
+				WorldManager->MarkSimulationForPostActorWork(this);
+			}
 		}
 	}
 	else
@@ -1433,9 +1440,8 @@ void FNiagaraSystemSimulation::Spawn_GameThread(float DeltaSeconds, bool bPostAc
 			InitParameterDataSetBindings(SpawningInstances[0]);
 		}
 
-		// Can we spawn async?
-		const bool bCanRunAsync = !bIsSolo && (!bPostActorTick || !GNiagaraSystemSimulationTickTaskShouldWait);
-		FNiagaraSystemSimulationTickContext Context(this, SpawningInstances, SpawningDataSet, DeltaSeconds, SpawningInstances.Num(), bCanRunAsync);
+		// Execute spawning we only force solo onto the GameThread
+		FNiagaraSystemSimulationTickContext Context(this, SpawningInstances, SpawningDataSet, DeltaSeconds, SpawningInstances.Num(), !bIsSolo);
 		if ( Context.IsRunningAsync() )
 		{
 			auto ConcurrentTickTask = TGraphTask<FNiagaraSystemSimulationSpawnConcurrentTask>::CreateTask(nullptr, ENamedThreads::GameThread).ConstructAndHold(Context, AllWorkCompleteGraphEvent);
@@ -1446,8 +1452,15 @@ void FNiagaraSystemSimulation::Spawn_GameThread(float DeltaSeconds, bool bPostAc
 			}
 			ConcurrentTickTask->Unlock();
 
-			// Note we always mark for EOF wait as the execution could be ongoing beyond PostActorTick
-			WorldManager->MarkSimulationsForEndOfFrameWait(this);
+			// Some simulations require us to complete inside PostActorTick and can not overlap to EOF updates / next frame, ensure we wait for them in the right location
+			if ( System->AllDIsPostSimulateCanOverlapFrames() && !GNiagaraSystemSimulationTickTaskShouldWait)
+			{
+				WorldManager->MarkSimulationForPostActorWork(this);
+			}
+			else
+			{
+				WorldManager->MarkSimulationsForEndOfFrameWait(this);
+			}
 		}
 		else
 		{
diff --git a/Engine/Plugins/FX/Niagara/Source/Niagara/Private/NiagaraWorldManager.cpp b/Engine/Plugins/FX/Niagara/Source/Niagara/Private/NiagaraWorldManager.cpp
--- a/Engine/Plugins/FX/Niagara/Source/Niagara/Private/NiagaraWorldManager.cpp
+++ b/Engine/Plugins/FX/Niagara/Source/Niagara/Private/NiagaraWorldManager.cpp
@@ -719,39 +719,44 @@ void FNiagaraWorldManager::PostActorTick(float DeltaSeconds)
 	// - Instances that were spawned and we need to ensure the async tick is complete
 	if (SimulationsWithPostActorWork.Num() > 0)
 	{
-		for (int32 i=0; i < SimulationsWithPostActorWork.Num(); ++i )
+		// Ensure completion of all systems
+		for ( auto& Simulation : SimulationsWithPostActorWork )
 		{
-			if (!SimulationsWithPostActorWork[i]->IsValid())
+			if ( Simulation->IsValid() )
 			{
-				SimulationsWithPostActorWork.RemoveAtSwap(i, 1, false);
-				--i;
-			}
-			else
-			{
-				SimulationsWithPostActorWork[i]->WaitForInstancesTickComplete();
+				Simulation->WaitForInstancesTickComplete();
 			}
 		}
 
-		for (int32 i=0; i < SimulationsWithPostActorWork.Num(); ++i)
+		// Update tick groups
+		for (auto& Simulation : SimulationsWithPostActorWork)
 		{
-			if (!SimulationsWithPostActorWork[i]->IsValid())
+			if (Simulation->IsValid())
 			{
-				SimulationsWithPostActorWork.RemoveAtSwap(i, 1, false);
-				--i;
+				Simulation->UpdateTickGroups_GameThread();
 			}
-			else
+		}
+
+		// Execute spawning
+		TArray<TSharedRef<FNiagaraSystemSimulation, ESPMode::ThreadSafe>> LocalSimulationsWithPostActorWork;
+		Swap(SimulationsWithPostActorWork, LocalSimulationsWithPostActorWork);
+		for (auto& Simulation : LocalSimulationsWithPostActorWork)
+		{
+			if (Simulation->IsValid())
 			{
-				SimulationsWithPostActorWork[i]->UpdateTickGroups_GameThread();
+				Simulation->Spawn_GameThread(DeltaSeconds, true);
 			}
 		}
 
-		for (const auto& Simulation : SimulationsWithPostActorWork)
+		// Wait for any spawning that may be required to complete in PostActorTick
+		for (auto& Simulation : SimulationsWithPostActorWork)
 		{
 			if (Simulation->IsValid())
 			{
-				Simulation->Spawn_GameThread(DeltaSeconds, true);
+				Simulation->WaitForInstancesTickComplete();
 			}
 		}
+		SimulationsWithPostActorWork.Reset();
 	}
 
 	// Clear cached player view location list, it should never be used outside of the world tick
