- Don't clip ticks when not running the simulation
- Cache view rect so tick flushing has a resonable resolution

#rnx

diff --git a/Engine/Plugins/FX/Niagara/Source/Niagara/Classes/NiagaraEmitterInstanceBatcher.h b/Engine/Plugins/FX/Niagara/Source/Niagara/Classes/NiagaraEmitterInstanceBatcher.h
--- a/Engine/Plugins/FX/Niagara/Source/Niagara/Classes/NiagaraEmitterInstanceBatcher.h
+++ b/Engine/Plugins/FX/Niagara/Source/Niagara/Classes/NiagaraEmitterInstanceBatcher.h
@@ -310,4 +310,7 @@ private:
 	/** Pool of free GPU ray traced collision groups. */
 	TArray<int32> FreeGPURayTracedCollisionGroups;
 	int32 NumGPURayTracedCollisionGroups = 0;
+
+	// Cached information to build a dummy view info if neccessary
+	FIntRect CachedViewRect;
 };
diff --git a/Engine/Plugins/FX/Niagara/Source/Niagara/Private/NiagaraEmitterInstanceBatcher.cpp b/Engine/Plugins/FX/Niagara/Source/Niagara/Private/NiagaraEmitterInstanceBatcher.cpp
--- a/Engine/Plugins/FX/Niagara/Source/Niagara/Private/NiagaraEmitterInstanceBatcher.cpp
+++ b/Engine/Plugins/FX/Niagara/Source/Niagara/Private/NiagaraEmitterInstanceBatcher.cpp
@@ -58,14 +58,6 @@ uint32 FNiagaraComputeExecutionContext::TickCounter = 0;
 const FName NiagaraEmitterInstanceBatcher::TemporalEffectName("NiagaraEmitterInstanceBatcher");
 #endif // WITH_MGPU
 
-int32 GNiagaraGpuMaxQueuedRenderFrames = 1;
-static FAutoConsoleVariableRef CVarNiagaraGpuMaxQueuedRenderFrames(
-	TEXT("fx.Niagara.Batcher.MaxQueuedFramesWithoutRender"),
-	GNiagaraGpuMaxQueuedRenderFrames,
-	TEXT("Number of frames we allow to be queued before we force ticks to be released or executed.\n"),
-	ECVF_Default
-);
-
 int32 GNiagaraGpuSubmitCommandHint = WITH_EDITOR ? 10 : 0;
 static FAutoConsoleVariableRef CVarNiagaraGpuSubmitCommandHint(
 	TEXT("fx.NiagaraGpuSubmitCommandHint"),
@@ -181,6 +173,7 @@ NiagaraEmitterInstanceBatcher::NiagaraEmitterInstanceBatcher(ERHIFeatureLevel::T
 	: FeatureLevel(InFeatureLevel)
 	, ShaderPlatform(InShaderPlatform)
 	, GPUSortManager(InGPUSortManager)
+	, CachedViewRect(0, 0, 64, 64)
 {
 	// Register the batcher callback in the GPUSortManager.
 	// The callback is used to generate the initial keys and values for the GPU sort tasks,
@@ -469,10 +462,9 @@ void NiagaraEmitterInstanceBatcher::ProcessPendingTicksFlush(FRHICommandListImme
 				.SetWorldTimes(0, 0, 0)
 				.SetGammaCorrection(1.0f));
 
-			const FIntRect DummViewRect(0, 0, 128, 128);
 			FSceneViewInitOptions ViewInitOptions;
 			ViewInitOptions.ViewFamily = &ViewFamily;
-			ViewInitOptions.SetViewRectangle(DummViewRect);
+			ViewInitOptions.SetViewRectangle(CachedViewRect);
 			ViewInitOptions.ViewOrigin = FVector::ZeroVector;
 			ViewInitOptions.ViewRotationMatrix = FMatrix::Identity;
 			ViewInitOptions.ProjectionMatrix = FMatrix::Identity;
@@ -488,15 +480,15 @@ void NiagaraEmitterInstanceBatcher::ProcessPendingTicksFlush(FRHICommandListImme
 			DummyView->ViewUniformBuffer = TUniformBufferRef<FViewUniformShaderParameters>::CreateUniformBufferImmediate(*DummyView->CachedViewUniformShaderParameters, UniformBuffer_SingleFrame);
 
 			TConstArrayView<FViewInfo> DummyViews = MakeArrayView(DummyView, 1);
-			const bool AllowGPUParticleUpdate = true;
+			const bool bAllowGPUParticleUpdate = true;
 			
 			// Execute all ticks that we can support without invalid simulations
 			FRDGBuilder GraphBuilder(RHICmdList);
 			CreateSystemTextures(GraphBuilder);
-			PreInitViews(GraphBuilder, AllowGPUParticleUpdate);
+			PreInitViews(GraphBuilder, bAllowGPUParticleUpdate);
 			GPUInstanceCounterManager.UpdateDrawIndirectBuffers(*this, RHICmdList, FeatureLevel, ENiagaraGPUCountUpdatePhase::PreOpaque);
-			PostInitViews(GraphBuilder, DummyViews, AllowGPUParticleUpdate);
-			PostRenderOpaque(GraphBuilder, DummyViews, AllowGPUParticleUpdate);
+			PostInitViews(GraphBuilder, DummyViews, bAllowGPUParticleUpdate);
+			PostRenderOpaque(GraphBuilder, DummyViews, bAllowGPUParticleUpdate);
 			GPUInstanceCounterManager.UpdateDrawIndirectBuffers(*this, RHICmdList, FeatureLevel, ENiagaraGPUCountUpdatePhase::PostOpaque);
 			GraphBuilder.Execute();
 			break;
@@ -1592,7 +1584,6 @@ void NiagaraEmitterInstanceBatcher::PreInitViews(FRDGBuilder& GraphBuilder, bool
 {
 	bRequiresReadback = false;
 	GNiagaraViewDataManager.ClearSceneTextureParameters();
-	FramesBeforeTickFlush = 0;
 
 	GpuReadbackManagerPtr->Tick();
 #if NIAGARA_COMPUTEDEBUG_ENABLED
@@ -1647,35 +1638,39 @@ void NiagaraEmitterInstanceBatcher::PreInitViews(FRDGBuilder& GraphBuilder, bool
 	// (note that currently there are no callback appropriate to do it)
 	SimulationsToSort.Reset();
 
-
-	if (bAllowGPUParticleUpdate && FNiagaraUtilities::AllowGPUParticles(GetShaderPlatform()))
+	if (FNiagaraUtilities::AllowGPUParticles(GetShaderPlatform()))
 	{
-		UpdateInstanceCountManager(GraphBuilder.RHICmdList);
-		PrepareAllTicks(GraphBuilder.RHICmdList);
+		if ( bAllowGPUParticleUpdate )
+		{
+			FramesBeforeTickFlush = 0;
+
+			UpdateInstanceCountManager(GraphBuilder.RHICmdList);
+			PrepareAllTicks(GraphBuilder.RHICmdList);
 		
 #if RHI_RAYTRACING
-		RayTracingHelper->BeginFrame(GraphBuilder.RHICmdList, HasRayTracingScene());
+			RayTracingHelper->BeginFrame(GraphBuilder.RHICmdList, HasRayTracingScene());
 #endif
 
-		if ( NiagaraEmitterInstanceBatcherLocal::GDebugLogging)
-		{
-			DumpDebugFrame();
-		}
+			if ( NiagaraEmitterInstanceBatcherLocal::GDebugLogging)
+			{
+				DumpDebugFrame();
+			}
 
-		FNiagaraSceneTextureParameters* PassParameters = GraphBuilder.AllocParameters<FNiagaraSceneTextureParameters>();
-		GNiagaraViewDataManager.GetSceneTextureParameters(GraphBuilder, *PassParameters);
+			FNiagaraSceneTextureParameters* PassParameters = GraphBuilder.AllocParameters<FNiagaraSceneTextureParameters>();
+			GNiagaraViewDataManager.GetSceneTextureParameters(GraphBuilder, *PassParameters);
 
-		GraphBuilder.AddPass(
-			RDG_EVENT_NAME("Niagara::PreInitViews"),
-			PassParameters,
-			ERDGPassFlags::Compute | ERDGPassFlags::NeverCull,
-			[this, PassParameters](FRHICommandListImmediate& RHICmdList)
-		{
-			NiagaraSceneTextures = PassParameters;
-			ON_SCOPE_EXIT { NiagaraSceneTextures = nullptr; };
+			GraphBuilder.AddPass(
+				RDG_EVENT_NAME("Niagara::PreInitViews"),
+				PassParameters,
+				ERDGPassFlags::Compute | ERDGPassFlags::NeverCull,
+				[this, PassParameters](FRHICommandListImmediate& RHICmdList)
+			{
+				NiagaraSceneTextures = PassParameters;
+				ON_SCOPE_EXIT { NiagaraSceneTextures = nullptr; };
 
-			ExecuteTicks(RHICmdList, nullptr, ENiagaraGpuComputeTickStage::PreInitViews);
-		});
+				ExecuteTicks(RHICmdList, nullptr, ENiagaraGpuComputeTickStage::PreInitViews);
+			});
+		}
 	}
 	else
 	{
@@ -1716,6 +1711,11 @@ void NiagaraEmitterInstanceBatcher::PostRenderOpaque(FRDGBuilder& GraphBuilder,
 
 	bAllowGPUParticleUpdate = bAllowGPUParticleUpdate && GetReferenceAllowGPUUpdate(Views);
 
+	if ( bAllowGPUParticleUpdate && Views.IsValidIndex(0) )
+	{
+		CachedViewRect = Views[0].ViewRect;
+	}
+
 	FNiagaraSceneTextureParameters* PassParameters = GraphBuilder.AllocParameters<FNiagaraSceneTextureParameters>();
 	GNiagaraViewDataManager.GetSceneTextureParameters(GraphBuilder, *PassParameters);
 
