- Ensure we use weak pointers for the customization
- There's a small window when closing the editor where the objects can be cleaned and Slate still attempts to access the data

#rnx

diff --git a/Engine/Plugins/FX/Niagara/Source/NiagaraEditor/Private/Customizations/NiagaraTypeCustomizations.cpp b/Engine/Plugins/FX/Niagara/Source/NiagaraEditor/Private/Customizations/NiagaraTypeCustomizations.cpp
--- a/Engine/Plugins/FX/Niagara/Source/NiagaraEditor/Private/Customizations/NiagaraTypeCustomizations.cpp
+++ b/Engine/Plugins/FX/Niagara/Source/NiagaraEditor/Private/Customizations/NiagaraTypeCustomizations.cpp
@@ -505,7 +505,7 @@ void FNiagaraVariableAttributeBindingCustomization::CustomizeHeader(TSharedRef<I
 
 FName FNiagaraUserParameterBindingCustomization::GetVariableName() const
 {
-	if (BaseSystem && TargetUserParameterBinding)
+	if (IsValid() && TargetUserParameterBinding)
 	{
 		return (TargetUserParameterBinding->Parameter.GetName());
 	}
@@ -514,7 +514,7 @@ FName FNiagaraUserParameterBindingCustomization::GetVariableName() const
 
 FText FNiagaraUserParameterBindingCustomization::GetCurrentText() const
 {
-	if (BaseSystem && TargetUserParameterBinding)
+	if (IsValid() && TargetUserParameterBinding)
 	{
 		return FText::FromName(TargetUserParameterBinding->Parameter.GetName());
 	}
@@ -523,7 +523,7 @@ FText FNiagaraUserParameterBindingCustomization::GetCurrentText() const
 
 FText FNiagaraUserParameterBindingCustomization::GetTooltipText() const
 {
-	if (BaseSystem && TargetUserParameterBinding && TargetUserParameterBinding->Parameter.IsValid())
+	if (IsValid() && TargetUserParameterBinding && TargetUserParameterBinding->Parameter.IsValid())
 	{
 		FText TooltipDesc = FText::Format(LOCTEXT("UserParameterBindingTooltip", "Bound to the user parameter \"{0}\""), FText::FromName(TargetUserParameterBinding->Parameter.GetName()));
 		return TooltipDesc;
@@ -554,8 +554,9 @@ TArray<FName> FNiagaraUserParameterBindingCustomization::GetNames() const
 {
 	TArray<FName> Names;
 
-	if (BaseSystem && TargetUserParameterBinding)
+	if (IsValid() && TargetUserParameterBinding)
 	{
+		UNiagaraSystem* BaseSystem = BaseSystemWeakPtr.Get();
 		for (const FNiagaraVariable Var : BaseSystem->GetExposedParameters().ReadParameterVariables())
 		{
 			if (FNiagaraParameterMapHistory::IsUserParameter(Var) && Var.GetType() == TargetUserParameterBinding->Parameter.GetType())
@@ -641,10 +642,13 @@ void FNiagaraUserParameterBindingCustomization::CustomizeHeader(TSharedRef<IProp
 	TArray<UObject*> Objects;
 	PropertyHandle->GetOuterObjects(Objects);
 	bool bAddDefault = true;
+
+	ObjectCustomizingWeakPtr.Reset();
+	BaseSystemWeakPtr.Reset();
+
 	if (Objects.Num() == 1)
 	{
-		BaseSystem = Objects[0]->GetTypedOuter<UNiagaraSystem>();
-		
+		UNiagaraSystem* BaseSystem = Objects[0]->GetTypedOuter<UNiagaraSystem>();
 		if (BaseSystem == nullptr)
 		{
 			if (UNiagaraDataInterface* ObjectAsDataInterface = Cast<UNiagaraDataInterface>(Objects[0]))
@@ -657,6 +661,9 @@ void FNiagaraUserParameterBindingCustomization::CustomizeHeader(TSharedRef<IProp
 		
 		if (BaseSystem)
 		{
+			ObjectCustomizingWeakPtr = Objects[0];
+			BaseSystemWeakPtr = BaseSystem;
+
 			TargetUserParameterBinding = (FNiagaraUserParameterBinding*)PropertyHandle->GetValueBaseAddress((uint8*)Objects[0]);
 
 			HeaderRow
diff --git a/Engine/Plugins/FX/Niagara/Source/NiagaraEditor/Private/Customizations/NiagaraTypeCustomizations.h b/Engine/Plugins/FX/Niagara/Source/NiagaraEditor/Private/Customizations/NiagaraTypeCustomizations.h
--- a/Engine/Plugins/FX/Niagara/Source/NiagaraEditor/Private/Customizations/NiagaraTypeCustomizations.h
+++ b/Engine/Plugins/FX/Niagara/Source/NiagaraEditor/Private/Customizations/NiagaraTypeCustomizations.h
@@ -215,8 +215,11 @@ private:
 	TSharedRef<SWidget> OnCreateWidgetForAction(struct FCreateWidgetForActionData* const InCreateData);
 	void OnActionSelected(const TArray< TSharedPtr<FEdGraphSchemaAction> >& SelectedActions, ESelectInfo::Type InSelectionType);
 
+	bool IsValid() const { return BaseSystemWeakPtr.IsValid() && ObjectCustomizingWeakPtr.IsValid(); }
+
 	TSharedPtr<IPropertyHandle> PropertyHandle;
-	class UNiagaraSystem* BaseSystem;
+	TWeakObjectPtr<UObject> ObjectCustomizingWeakPtr;
+	TWeakObjectPtr<class UNiagaraSystem> BaseSystemWeakPtr;
 	struct FNiagaraUserParameterBinding* TargetUserParameterBinding;
 
 };
