integrate whitewater into river

diff --git a/Engine/Plugins/FX/Niagara/Source/NiagaraEditorWidgets/Private/DetailCustomizations/NiagaraDataInterfaceMeshRendererInfoDetails.cpp b/Engine/Plugins/FX/Niagara/Source/NiagaraEditorWidgets/Private/DetailCustomizations/NiagaraDataInterfaceMeshRendererInfoDetails.cpp
--- a/Engine/Plugins/FX/Niagara/Source/NiagaraEditorWidgets/Private/DetailCustomizations/NiagaraDataInterfaceMeshRendererInfoDetails.cpp
+++ b/Engine/Plugins/FX/Niagara/Source/NiagaraEditorWidgets/Private/DetailCustomizations/NiagaraDataInterfaceMeshRendererInfoDetails.cpp
@@ -191,6 +191,10 @@ FText FNiagaraDataInterfaceMeshRendererInfoDetails::CreateRendererTextLabel(cons
 	const FText EmitterText = FText::FromString(Emitter->GetUniqueEmitterName());
 
 	UNiagaraEmitterEditorData* EmitterEditorData = static_cast<UNiagaraEmitterEditorData*>(Emitter->GetEditorData());
+	if (EmitterEditorData == nullptr)
+	{
+		return LOCTEXT("NoneOption", "None");
+	}
 	UNiagaraStackEditorData& EmitterStackEditorData = EmitterEditorData->GetStackEditorData();
 	FString RendererStackEditorDataKey = FString::Printf(TEXT("Renderer-%s"), *Properties->GetName());
 	const FText* RendererDisplayName = EmitterStackEditorData.GetStackEntryDisplayName(RendererStackEditorDataKey);
