Automated search-and-replace pass to change uses of UE_TRANSITIONAL_OBJECT_PTR(Type) to TObjectPtr<Type>.
There were two instance of the type being incorrect and they were manually fixed.

diff --git a/Engine/Plugins/Animation/ControlRig/Source/ControlRigDeveloper/Private/ControlRigBlueprint.cpp b/Engine/Plugins/Animation/ControlRig/Source/ControlRigDeveloper/Private/ControlRigBlueprint.cpp
--- a/Engine/Plugins/Animation/ControlRig/Source/ControlRigDeveloper/Private/ControlRigBlueprint.cpp
+++ b/Engine/Plugins/Animation/ControlRig/Source/ControlRigDeveloper/Private/ControlRigBlueprint.cpp
@@ -1104,7 +1104,7 @@ URigVMController* UControlRigBlueprint::GetController(URigVMGraph* InGraph) cons
 		InGraph = Model;
 	}
 
-	UE_TRANSITIONAL_OBJECT_PTR(URigVMController) const* ControllerPtr = Controllers.Find(InGraph);
+	TObjectPtr<URigVMController> const* ControllerPtr = Controllers.Find(InGraph);
 	if (ControllerPtr)
 	{
 		return *ControllerPtr;
diff --git a/Engine/Plugins/Animation/ControlRig/Source/ControlRigDeveloper/Private/Graph/ControlRigGraphNode.cpp b/Engine/Plugins/Animation/ControlRig/Source/ControlRigDeveloper/Private/Graph/ControlRigGraphNode.cpp
--- a/Engine/Plugins/Animation/ControlRig/Source/ControlRigDeveloper/Private/Graph/ControlRigGraphNode.cpp
+++ b/Engine/Plugins/Animation/ControlRig/Source/ControlRigDeveloper/Private/Graph/ControlRigGraphNode.cpp
@@ -978,7 +978,7 @@ FName UControlRigGraphNode::GetModelNodeName() const
 
 URigVMPin* UControlRigGraphNode::GetModelPinFromPinPath(const FString& InPinPath) const
 {
-	if (UE_TRANSITIONAL_OBJECT_PTR(URigVMPin) const* CachedModelPinPtr = CachedModelPins.Find(InPinPath))
+	if (TObjectPtr<URigVMPin> const* CachedModelPinPtr = CachedModelPins.Find(InPinPath))
 	{
 		URigVMPin* CachedModelPin = *CachedModelPinPtr;
 		if (!CachedModelPin->HasAnyFlags(RF_Transient) && CachedModelPin->GetNode())
diff --git a/Engine/Plugins/Animation/ControlRig/Source/ControlRigEditor/Private/EditMode/ControlRigControlsProxy.cpp b/Engine/Plugins/Animation/ControlRig/Source/ControlRigEditor/Private/EditMode/ControlRigControlsProxy.cpp
--- a/Engine/Plugins/Animation/ControlRig/Source/ControlRigEditor/Private/EditMode/ControlRigControlsProxy.cpp
+++ b/Engine/Plugins/Animation/ControlRig/Source/ControlRigEditor/Private/EditMode/ControlRigControlsProxy.cpp
@@ -656,7 +656,7 @@ void UControlRigBoolControlProxy::SetKey(const IPropertyHandle& KeyedPropertyHan
 
 UControlRigControlsProxy* UControlRigDetailPanelControlProxies::FindProxy(const FName& Name) const
 {
-	UE_TRANSITIONAL_OBJECT_PTR(UControlRigControlsProxy) const* Proxy = AllProxies.Find(Name);
+	TObjectPtr<UControlRigControlsProxy> const* Proxy = AllProxies.Find(Name);
 	if (Proxy &&  Proxy[0])
 	{
 		return Proxy[0];
diff --git a/Engine/Plugins/Animation/LiveLink/Source/LiveLinkComponents/Private/LiveLinkComponentController.cpp b/Engine/Plugins/Animation/LiveLink/Source/LiveLinkComponents/Private/LiveLinkComponentController.cpp
--- a/Engine/Plugins/Animation/LiveLink/Source/LiveLinkComponents/Private/LiveLinkComponentController.cpp
+++ b/Engine/Plugins/Animation/LiveLink/Source/LiveLinkComponents/Private/LiveLinkComponentController.cpp
@@ -98,7 +98,7 @@ void ULiveLinkComponentController::SetControllerClassForRole(TSubclassOf<ULiveLi
 {
 	if (ControllerMap.Contains(RoleClass))
 	{
-		UE_TRANSITIONAL_OBJECT_PTR(ULiveLinkControllerBase)& CurrentController = ControllerMap.FindOrAdd(RoleClass);
+		TObjectPtr<ULiveLinkControllerBase>& CurrentController = ControllerMap.FindOrAdd(RoleClass);
 		if (CurrentController == nullptr || CurrentController->GetClass() != DesiredControllerClass)
 		{
 			//Controller is about to change, cleanup current one before 
@@ -291,7 +291,7 @@ bool ULiveLinkComponentController::IsControllerMapOutdated() const
 	//Check if all map matches class hierarchy
 	for (const TSubclassOf<ULiveLinkRole>& RoleClass : SelectedRoleHierarchy)
 	{
-		UE_TRANSITIONAL_OBJECT_PTR(ULiveLinkControllerBase) const* FoundController = ControllerMap.Find(RoleClass);
+		TObjectPtr<ULiveLinkControllerBase> const* FoundController = ControllerMap.Find(RoleClass);
 
 		//If ControllerMap doesn't have an entry for one of the role class hierarchy, we need to update
 		if (FoundController == nullptr)
diff --git a/Engine/Plugins/Animation/LiveLink/Source/LiveLinkEditor/Private/LiveLinkComponentDetailCustomization.cpp b/Engine/Plugins/Animation/LiveLink/Source/LiveLinkEditor/Private/LiveLinkComponentDetailCustomization.cpp
--- a/Engine/Plugins/Animation/LiveLink/Source/LiveLinkEditor/Private/LiveLinkComponentDetailCustomization.cpp
+++ b/Engine/Plugins/Animation/LiveLink/Source/LiveLinkEditor/Private/LiveLinkComponentDetailCustomization.cpp
@@ -245,7 +245,7 @@ bool FLiveLinkComponentDetailCustomization::IsControllerItemSelected(FName Item,
 	if (ULiveLinkComponentController* EditedObjectPtr = EditedObject.Get())
 	{
 
-		UE_TRANSITIONAL_OBJECT_PTR(ULiveLinkControllerBase)* CurrentClass = EditedObjectPtr->ControllerMap.Find(RoleClass);
+		TObjectPtr<ULiveLinkControllerBase>* CurrentClass = EditedObjectPtr->ControllerMap.Find(RoleClass);
 		if (CurrentClass == nullptr || *CurrentClass == nullptr)
 		{
 			return Item.IsNone();
@@ -266,7 +266,7 @@ FSlateColor FLiveLinkComponentDetailCustomization::HandleControllerStatusColorAn
 
 	if (ULiveLinkComponentController* EditorObjectPtr = EditedObject.Get())
 	{
-		UE_TRANSITIONAL_OBJECT_PTR(ULiveLinkControllerBase)* AssociatedControllerPtr = EditorObjectPtr->ControllerMap.Find(RoleClassEntry);
+		TObjectPtr<ULiveLinkControllerBase>* AssociatedControllerPtr = EditorObjectPtr->ControllerMap.Find(RoleClassEntry);
 		if (AssociatedControllerPtr && *AssociatedControllerPtr)
 		{
 			if (UActorComponent* SelectedComponent = EditorObjectPtr->ComponentToControl.GetComponent(EditorObjectPtr->GetOwner()))
@@ -301,7 +301,7 @@ FText FLiveLinkComponentDetailCustomization::HandleControllerStatusText(TSubclas
 
 	if (ULiveLinkComponentController* EditorObjectPtr = EditedObject.Get())
 	{
-		UE_TRANSITIONAL_OBJECT_PTR(ULiveLinkControllerBase)* AssociatedControllerPtr = EditorObjectPtr->ControllerMap.Find(RoleClassEntry);
+		TObjectPtr<ULiveLinkControllerBase>* AssociatedControllerPtr = EditorObjectPtr->ControllerMap.Find(RoleClassEntry);
 		if (AssociatedControllerPtr && *AssociatedControllerPtr)
 		{
 			if (UActorComponent* SelectedComponent = EditorObjectPtr->ComponentToControl.GetComponent(EditorObjectPtr->GetOwner()))
@@ -336,7 +336,7 @@ FText FLiveLinkComponentDetailCustomization::HandleControllerStatusToolTipText(T
 
 	if (ULiveLinkComponentController* EditorObjectPtr = EditedObject.Get())
 	{
-		UE_TRANSITIONAL_OBJECT_PTR(ULiveLinkControllerBase)* AssociatedControllerPtr = EditorObjectPtr->ControllerMap.Find(RoleClassEntry);
+		TObjectPtr<ULiveLinkControllerBase>* AssociatedControllerPtr = EditorObjectPtr->ControllerMap.Find(RoleClassEntry);
 		if (AssociatedControllerPtr && *AssociatedControllerPtr)
 		{
 			if (UActorComponent* SelectedComponent = EditorObjectPtr->ComponentToControl.GetComponent(EditorObjectPtr->GetOwner()))
diff --git a/Engine/Plugins/Enterprise/MDLImporter/Source/MDLImporter/Private/material/BakedMaterialFactory.cpp b/Engine/Plugins/Enterprise/MDLImporter/Source/MDLImporter/Private/material/BakedMaterialFactory.cpp
--- a/Engine/Plugins/Enterprise/MDLImporter/Source/MDLImporter/Private/material/BakedMaterialFactory.cpp
+++ b/Engine/Plugins/Enterprise/MDLImporter/Source/MDLImporter/Private/material/BakedMaterialFactory.cpp
@@ -58,7 +58,7 @@ namespace Mat
 		// get under clear coat output
 		UMaterialExpressionClearCoatNormalCustomOutput* UnderClearCoat = nullptr;
 		{
-			UE_TRANSITIONAL_OBJECT_PTR(UMaterialExpression)* Found = Material.Expressions.FindByPredicate(
+			TObjectPtr<UMaterialExpression>* Found = Material.Expressions.FindByPredicate(
 			    [](const UMaterialExpression* Expr) { return Expr->IsA<UMaterialExpressionClearCoatNormalCustomOutput>(); });
 			if (Found)
 				UnderClearCoat = Cast<UMaterialExpressionClearCoatNormalCustomOutput>(*Found);
@@ -149,7 +149,7 @@ namespace Mat
 	{
 		FString               Name;
 
-		UE_TRANSITIONAL_OBJECT_PTR(UMaterialExpression)* Found = Material.Expressions.FindByPredicate(
+		TObjectPtr<UMaterialExpression>* Found = Material.Expressions.FindByPredicate(
 			[&Name](UMaterialExpression* Expression)  //
 		{
 			if (UMaterialExpressionScalarParameter* ScalarParameter = Cast<UMaterialExpressionScalarParameter>(Expression))
@@ -166,7 +166,7 @@ namespace Mat
 			return Name.Find(TEXT("Tiling Factor")) != INDEX_NONE;
 		});
 
-		UE_TRANSITIONAL_OBJECT_PTR(UMaterialExpression)* FoundU = Material.Expressions.FindByPredicate(
+		TObjectPtr<UMaterialExpression>* FoundU = Material.Expressions.FindByPredicate(
 		    [&Name](UMaterialExpression* Expression)  //
 		    {
 			    if (UMaterialExpressionScalarParameter* ScalarParameter = Cast<UMaterialExpressionScalarParameter>(Expression))
@@ -183,7 +183,7 @@ namespace Mat
 			    return Name.Find(TEXT("U Tiling")) != INDEX_NONE;
 		    });
 
-		UE_TRANSITIONAL_OBJECT_PTR(UMaterialExpression)* FoundV = Material.Expressions.FindByPredicate(
+		TObjectPtr<UMaterialExpression>* FoundV = Material.Expressions.FindByPredicate(
 			[&Name](UMaterialExpression* Expression)  //
 		{
 			if (UMaterialExpressionScalarParameter* ScalarParameter = Cast<UMaterialExpressionScalarParameter>(Expression))
diff --git a/Engine/Plugins/Enterprise/VariantManagerContent/Source/VariantManagerContent/Private/LevelVariantSets.cpp b/Engine/Plugins/Enterprise/VariantManagerContent/Source/VariantManagerContent/Private/LevelVariantSets.cpp
--- a/Engine/Plugins/Enterprise/VariantManagerContent/Source/VariantManagerContent/Private/LevelVariantSets.cpp
+++ b/Engine/Plugins/Enterprise/VariantManagerContent/Source/VariantManagerContent/Private/LevelVariantSets.cpp
@@ -305,7 +305,7 @@ UObject* ULevelVariantSets::GetDirectorInstance(UObject* WorldContext)
 	}
 
 	ULevelVariantSetsFunctionDirector* TargetDirector = nullptr;
-	if ( UE_TRANSITIONAL_OBJECT_PTR(ULevelVariantSetsFunctionDirector*)* ExistingDirector = DirectorOuter->DirectorInstances.Find( DirectorClass ) )
+	if ( TObjectPtr<ULevelVariantSetsFunctionDirector>* ExistingDirector = DirectorOuter->DirectorInstances.Find( DirectorClass ) )
 	{
 		TargetDirector = *ExistingDirector;
 	}
@@ -344,7 +344,7 @@ UVariantSet* ULevelVariantSets::GetVariantSet(int32 VariantSetIndex)
 
 UVariantSet* ULevelVariantSets::GetVariantSetByName(FString VariantSetName)
 {
-	UE_TRANSITIONAL_OBJECT_PTR(UVariantSet)* VarSetPtr = VariantSets.FindByPredicate([VariantSetName](const UVariantSet* VarSet)
+	TObjectPtr<UVariantSet>* VarSetPtr = VariantSets.FindByPredicate([VariantSetName](const UVariantSet* VarSet)
 	{
 		return VarSet->GetDisplayText().ToString() == VariantSetName;
 	});
diff --git a/Engine/Plugins/Enterprise/VariantManagerContent/Source/VariantManagerContent/Private/Variant.cpp b/Engine/Plugins/Enterprise/VariantManagerContent/Source/VariantManagerContent/Private/Variant.cpp
--- a/Engine/Plugins/Enterprise/VariantManagerContent/Source/VariantManagerContent/Private/Variant.cpp
+++ b/Engine/Plugins/Enterprise/VariantManagerContent/Source/VariantManagerContent/Private/Variant.cpp
@@ -234,7 +234,7 @@ AActor* UVariant::GetActor(int32 ActorIndex)
 
 UVariantObjectBinding* UVariant::GetBindingByName(const FString& ActorName)
 {
-	UE_TRANSITIONAL_OBJECT_PTR(UVariantObjectBinding)* FoundBindingPtr = ObjectBindings.FindByPredicate([&ActorName](const UVariantObjectBinding* Binding)
+	TObjectPtr<UVariantObjectBinding>* FoundBindingPtr = ObjectBindings.FindByPredicate([&ActorName](const UVariantObjectBinding* Binding)
 	{
 		UObject* ThisActor = Binding->GetObject();
 		return ThisActor && ThisActor->GetName() == ActorName;
diff --git a/Engine/Plugins/Enterprise/VariantManagerContent/Source/VariantManagerContent/Private/VariantSet.cpp b/Engine/Plugins/Enterprise/VariantManagerContent/Source/VariantManagerContent/Private/VariantSet.cpp
--- a/Engine/Plugins/Enterprise/VariantManagerContent/Source/VariantManagerContent/Private/VariantSet.cpp
+++ b/Engine/Plugins/Enterprise/VariantManagerContent/Source/VariantManagerContent/Private/VariantSet.cpp
@@ -282,7 +282,7 @@ UVariant* UVariantSet::GetVariant(int32 VariantIndex)
 
 UVariant* UVariantSet::GetVariantByName(FString VariantName)
 {
-	UE_TRANSITIONAL_OBJECT_PTR(UVariant)* VarPtr = Variants.FindByPredicate([VariantName](const UVariant* Var)
+	TObjectPtr<UVariant>* VarPtr = Variants.FindByPredicate([VariantName](const UVariant* Var)
 	{
 		return Var->GetDisplayText().ToString() == VariantName;
 	});
diff --git a/Engine/Plugins/Experimental/CommonConversation/Source/CommonConversationEditor/Private/SConversationDiff.cpp b/Engine/Plugins/Experimental/CommonConversation/Source/CommonConversationEditor/Private/SConversationDiff.cpp
--- a/Engine/Plugins/Experimental/CommonConversation/Source/CommonConversationEditor/Private/SConversationDiff.cpp
+++ b/Engine/Plugins/Experimental/CommonConversation/Source/CommonConversationEditor/Private/SConversationDiff.cpp
@@ -376,7 +376,7 @@ void SConversationDiff::OnSelectionChanged(FSharedDiffOnGraph Item, ESelectInfo:
 				{
 					// This is a sub-node, we need to find our parent node in the graph
 					// todo: work out why BTNode->ParentNode is always null
-					UE_TRANSITIONAL_OBJECT_PTR(UEdGraphNode)* ParentNodePtr = NodeGraph->Nodes.FindByPredicate([BTNode](UEdGraphNode* PotentialParentNode) -> bool
+					TObjectPtr<UEdGraphNode>* ParentNodePtr = NodeGraph->Nodes.FindByPredicate([BTNode](UEdGraphNode* PotentialParentNode) -> bool
 					{
 						UConversationGraphNode* BTPotentialParentNode = Cast<UConversationGraphNode>(PotentialParentNode);
 						return BTPotentialParentNode && (BTPotentialParentNode->SubNodes.Contains(BTNode));
diff --git a/Engine/Plugins/Experimental/MeshModelingToolsetExp/Source/MeshModelingToolsExp/Private/Sculpting/MeshSculptToolBase.cpp b/Engine/Plugins/Experimental/MeshModelingToolsetExp/Source/MeshModelingToolsExp/Private/Sculpting/MeshSculptToolBase.cpp
--- a/Engine/Plugins/Experimental/MeshModelingToolsetExp/Source/MeshModelingToolsExp/Private/Sculpting/MeshSculptToolBase.cpp
+++ b/Engine/Plugins/Experimental/MeshModelingToolsetExp/Source/MeshModelingToolsExp/Private/Sculpting/MeshSculptToolBase.cpp
@@ -328,7 +328,7 @@ void UMeshSculptToolBase::SetActivePrimaryBrushType(int32 Identifier)
 	PrimaryBrushOp = (*Factory)->Build();
 	PrimaryBrushOp->Falloff = PrimaryFalloff;
 
-	UE_TRANSITIONAL_OBJECT_PTR(UMeshSculptBrushOpProps)* FoundProps = BrushOpPropSets.Find(Identifier);
+	TObjectPtr<UMeshSculptBrushOpProps>* FoundProps = BrushOpPropSets.Find(Identifier);
 	if (FoundProps != nullptr)
 	{
 		SetToolPropertySourceEnabled(*FoundProps, bBrushOpPropsVisible);
@@ -360,7 +360,7 @@ void UMeshSculptToolBase::SetActiveSecondaryBrushType(int32 Identifier)
 	SecondaryFalloff->FalloffFunc = UE::SculptFalloffs::MakeStandardSmoothFalloff();
 	SecondaryBrushOp->Falloff = SecondaryFalloff;
 
-	UE_TRANSITIONAL_OBJECT_PTR(UMeshSculptBrushOpProps)* FoundProps = SecondaryBrushOpPropSets.Find(Identifier);
+	TObjectPtr<UMeshSculptBrushOpProps>* FoundProps = SecondaryBrushOpPropSets.Find(Identifier);
 	if (FoundProps != nullptr)
 	{
 		SetToolPropertySourceEnabled(*FoundProps, bBrushOpPropsVisible);
diff --git a/Engine/Plugins/Experimental/MeshModelingToolsetExp/Source/ModelingComponents/Private/Drawing/PreviewGeometryActor.cpp b/Engine/Plugins/Experimental/MeshModelingToolsetExp/Source/ModelingComponents/Private/Drawing/PreviewGeometryActor.cpp
--- a/Engine/Plugins/Experimental/MeshModelingToolsetExp/Source/ModelingComponents/Private/Drawing/PreviewGeometryActor.cpp
+++ b/Engine/Plugins/Experimental/MeshModelingToolsetExp/Source/ModelingComponents/Private/Drawing/PreviewGeometryActor.cpp
@@ -102,7 +102,7 @@ ULineSetComponent* UPreviewGeometry::AddLineSet(const FString& SetIdentifier)
 
 ULineSetComponent* UPreviewGeometry::FindLineSet(const FString& LineSetIdentifier)
 {
-	UE_TRANSITIONAL_OBJECT_PTR(ULineSetComponent)* Found = LineSets.Find(LineSetIdentifier);
+	TObjectPtr<ULineSetComponent>* Found = LineSets.Find(LineSetIdentifier);
 	if (Found != nullptr)
 	{
 		return *Found;
@@ -112,7 +112,7 @@ ULineSetComponent* UPreviewGeometry::FindLineSet(const FString& LineSetIdentifie
 
 bool UPreviewGeometry::RemoveLineSet(const FString& LineSetIdentifier, bool bDestroy)
 {
-	UE_TRANSITIONAL_OBJECT_PTR(ULineSetComponent)* Found = LineSets.Find(LineSetIdentifier);
+	TObjectPtr<ULineSetComponent>* Found = LineSets.Find(LineSetIdentifier);
 	if (Found != nullptr)
 	{
 		ULineSetComponent* LineSet = *Found;
@@ -146,7 +146,7 @@ void UPreviewGeometry::RemoveAllLineSets(bool bDestroy)
 
 bool UPreviewGeometry::SetLineSetVisibility(const FString& LineSetIdentifier, bool bVisible)
 {
-	UE_TRANSITIONAL_OBJECT_PTR(ULineSetComponent)* Found = LineSets.Find(LineSetIdentifier);
+	TObjectPtr<ULineSetComponent>* Found = LineSets.Find(LineSetIdentifier);
 	if (Found != nullptr)
 	{
 		(*Found)->SetVisibility(bVisible);
@@ -158,7 +158,7 @@ bool UPreviewGeometry::SetLineSetVisibility(const FString& LineSetIdentifier, bo
 
 bool UPreviewGeometry::SetLineSetMaterial(const FString& LineSetIdentifier, UMaterialInterface* NewMaterial)
 {
-	UE_TRANSITIONAL_OBJECT_PTR(ULineSetComponent)* Found = LineSets.Find(LineSetIdentifier);
+	TObjectPtr<ULineSetComponent>* Found = LineSets.Find(LineSetIdentifier);
 	if (Found != nullptr)
 	{
 		(*Found)->SetLineMaterial(NewMaterial);
diff --git a/Engine/Plugins/FX/Niagara/Source/Niagara/Private/MovieScene/MovieSceneNiagaraSystemTrack.cpp b/Engine/Plugins/FX/Niagara/Source/Niagara/Private/MovieScene/MovieSceneNiagaraSystemTrack.cpp
--- a/Engine/Plugins/FX/Niagara/Source/Niagara/Private/MovieScene/MovieSceneNiagaraSystemTrack.cpp
+++ b/Engine/Plugins/FX/Niagara/Source/Niagara/Private/MovieScene/MovieSceneNiagaraSystemTrack.cpp
@@ -29,7 +29,7 @@ FMovieSceneEvalTemplatePtr UMovieSceneNiagaraSystemTrack::CreateTemplateForSecti
 
 void UMovieSceneNiagaraSystemTrack::PostCompile(FMovieSceneEvaluationTrack& OutTrack, const FMovieSceneTrackCompilerArgs& Args) const
 {
-	UE_TRANSITIONAL_OBJECT_PTR(UMovieSceneSection) const* SpawnSectionPtr = Sections.FindByPredicate([](const UMovieSceneSection* Section) { return Section->GetClass() == UMovieSceneNiagaraSystemSpawnSection::StaticClass(); });
+	TObjectPtr<UMovieSceneSection> const* SpawnSectionPtr = Sections.FindByPredicate([](const UMovieSceneSection* Section) { return Section->GetClass() == UMovieSceneNiagaraSystemSpawnSection::StaticClass(); });
 	if (SpawnSectionPtr != nullptr)
 	{
 		UMovieSceneNiagaraSystemSpawnSection* SpawnSection = CastChecked<UMovieSceneNiagaraSystemSpawnSection>(*SpawnSectionPtr);
diff --git a/Engine/Plugins/FX/Niagara/Source/Niagara/Private/NiagaraEmitter.cpp b/Engine/Plugins/FX/Niagara/Source/Niagara/Private/NiagaraEmitter.cpp
--- a/Engine/Plugins/FX/Niagara/Source/Niagara/Private/NiagaraEmitter.cpp
+++ b/Engine/Plugins/FX/Niagara/Source/Niagara/Private/NiagaraEmitter.cpp
@@ -1768,7 +1768,7 @@ void UNiagaraEmitter::UpdateFromMergedCopy(const INiagaraMergeManager& MergeMana
 	}
 
 	// Copy shader stages
-	for (UE_TRANSITIONAL_OBJECT_PTR(UNiagaraSimulationStageBase)& SimulationStage : SimulationStages)
+	for (TObjectPtr<UNiagaraSimulationStageBase>& SimulationStage : SimulationStages)
 	{
 		SimulationStage->OnChanged().RemoveAll(this);
 		SimulationStage->Script->RapidIterationParameters.RemoveAllOnChangedHandlers(this);
@@ -1961,7 +1961,7 @@ void UNiagaraEmitter::RemoveEventHandlerByUsageId(FGuid EventHandlerUsageId)
 
 UNiagaraSimulationStageBase* UNiagaraEmitter::GetSimulationStageById(FGuid ScriptUsageId) const
 {
-	UE_TRANSITIONAL_OBJECT_PTR(UNiagaraSimulationStageBase) const* FoundSimulationStagePtr = SimulationStages.FindByPredicate([&ScriptUsageId](UNiagaraSimulationStageBase* SimulationStage) { return SimulationStage->Script->GetUsageId() == ScriptUsageId; });
+	TObjectPtr<UNiagaraSimulationStageBase> const* FoundSimulationStagePtr = SimulationStages.FindByPredicate([&ScriptUsageId](UNiagaraSimulationStageBase* SimulationStage) { return SimulationStage->Script->GetUsageId() == ScriptUsageId; });
 	return FoundSimulationStagePtr != nullptr ? *FoundSimulationStagePtr : nullptr;
 }
 
diff --git a/Engine/Plugins/FX/Niagara/Source/Niagara/Private/NiagaraPlatformSet.cpp b/Engine/Plugins/FX/Niagara/Source/Niagara/Private/NiagaraPlatformSet.cpp
--- a/Engine/Plugins/FX/Niagara/Source/Niagara/Private/NiagaraPlatformSet.cpp
+++ b/Engine/Plugins/FX/Niagara/Source/Niagara/Private/NiagaraPlatformSet.cpp
@@ -259,7 +259,7 @@ void FNiagaraPlatformSet::GetOverridenDeviceProfiles(int32 QualityLevel, TArray<
 	int32 QLMask = CreateQualityLevelMask(QualityLevel);
 	for (const FNiagaraDeviceProfileStateEntry& Entry : DeviceProfileStates)
 	{
-		if (UE_TRANSITIONAL_OBJECT_PTR(UObject)* DeviceProfile = UDeviceProfileManager::Get().Profiles.FindByPredicate([&](UObject* CheckProfile) {return CheckProfile->GetFName() == Entry.ProfileName;}))
+		if (TObjectPtr<UObject>* DeviceProfile = UDeviceProfileManager::Get().Profiles.FindByPredicate([&](UObject* CheckProfile) {return CheckProfile->GetFName() == Entry.ProfileName;}))
 		{
 			UDeviceProfile* Profile = CastChecked<UDeviceProfile>(*DeviceProfile);
 			//If this platform cannot change at runtime then we store all EQs in the state so that the device is still overridden if someone changes it's EQ CVar.
diff --git a/Engine/Plugins/FX/Niagara/Source/NiagaraEditor/Private/NiagaraEditorData.cpp b/Engine/Plugins/FX/Niagara/Source/NiagaraEditor/Private/NiagaraEditorData.cpp
--- a/Engine/Plugins/FX/Niagara/Source/NiagaraEditor/Private/NiagaraEditorData.cpp
+++ b/Engine/Plugins/FX/Niagara/Source/NiagaraEditor/Private/NiagaraEditorData.cpp
@@ -126,7 +126,7 @@ TOptional<TTuple<FName /*SyncedOldName*/, FName /*SyncedNewName*/>> UNiagaraEdit
 	if (DestScriptVar == nullptr)
 	{
 		const FGuid& SourceScriptVarId = SourceScriptVar->Metadata.GetVariableGuid();
-		UE_TRANSITIONAL_OBJECT_PTR(UNiagaraScriptVariable)* ScriptVarPtr = EditorOnlyScriptVars.FindByPredicate([&SourceScriptVarId](const UNiagaraScriptVariable* ScriptVar) { return ScriptVar->Metadata.GetVariableGuid() == SourceScriptVarId; });
+		TObjectPtr<UNiagaraScriptVariable>* ScriptVarPtr = EditorOnlyScriptVars.FindByPredicate([&SourceScriptVarId](const UNiagaraScriptVariable* ScriptVar) { return ScriptVar->Metadata.GetVariableGuid() == SourceScriptVarId; });
 		if (ScriptVarPtr == nullptr)
 		{
 			// Failed to find a DestScriptVar with an Id matching that of SourceScriptVar.
diff --git a/Engine/Plugins/FX/Niagara/Source/NiagaraEditor/Private/NiagaraGraph.cpp b/Engine/Plugins/FX/Niagara/Source/NiagaraEditor/Private/NiagaraGraph.cpp
--- a/Engine/Plugins/FX/Niagara/Source/NiagaraEditor/Private/NiagaraGraph.cpp
+++ b/Engine/Plugins/FX/Niagara/Source/NiagaraEditor/Private/NiagaraGraph.cpp
@@ -147,7 +147,7 @@ void UNiagaraGraph::PostLoad()
 	for (auto It = VariableToScriptVariable.CreateIterator(); It; ++It)
 	{
 		FNiagaraVariable Var = It.Key();
-		UE_TRANSITIONAL_OBJECT_PTR(UNiagaraScriptVariable)& ScriptVar = It.Value();
+		TObjectPtr<UNiagaraScriptVariable>& ScriptVar = It.Value();
 
 		if (ScriptVar == nullptr)
 		{
@@ -318,7 +318,7 @@ void UNiagaraGraph::PostLoad()
 
 		for (auto& VariableToScriptVariableItem : VariableToScriptVariable)
 		{
-			UE_TRANSITIONAL_OBJECT_PTR(UNiagaraScriptVariable)& MetaData = VariableToScriptVariableItem.Value;
+			TObjectPtr<UNiagaraScriptVariable>& MetaData = VariableToScriptVariableItem.Value;
 
 			// Migrate advanced display.
 			if (MetaData->Metadata.PropertyMetaData.Contains("AdvancedDisplay"))
@@ -1325,7 +1325,7 @@ UNiagaraScriptVariable* UNiagaraGraph::GetScriptVariable(FNiagaraVariable Parame
 	{
 		RefreshParameterReferences();
 	}
-	if (UE_TRANSITIONAL_OBJECT_PTR(UNiagaraScriptVariable)* FoundScriptVariable = VariableToScriptVariable.Find(Parameter))
+	if (TObjectPtr<UNiagaraScriptVariable>* FoundScriptVariable = VariableToScriptVariable.Find(Parameter))
 	{
 		return *FoundScriptVariable;
 	}
@@ -1368,7 +1368,7 @@ UNiagaraScriptVariable* UNiagaraGraph::AddParameter(const FNiagaraVariable& Para
 		ParameterToReferencesMap.Add(Parameter, NewReferenceCollection);
 	}
 
-	UE_TRANSITIONAL_OBJECT_PTR(UNiagaraScriptVariable)* FoundScriptVariable = VariableToScriptVariable.Find(Parameter);
+	TObjectPtr<UNiagaraScriptVariable>* FoundScriptVariable = VariableToScriptVariable.Find(Parameter);
 	if (!FoundScriptVariable)
 	{
 		Modify();
@@ -1389,7 +1389,7 @@ UNiagaraScriptVariable* UNiagaraGraph::AddParameter(const FNiagaraVariable& Para
 
 UNiagaraScriptVariable* UNiagaraGraph::AddParameter(const UNiagaraScriptVariable* InScriptVar)
 {
-	UE_TRANSITIONAL_OBJECT_PTR(UNiagaraScriptVariable)* FoundScriptVariable = VariableToScriptVariable.Find(InScriptVar->Variable);
+	TObjectPtr<UNiagaraScriptVariable>* FoundScriptVariable = VariableToScriptVariable.Find(InScriptVar->Variable);
 	if (!FoundScriptVariable)
 	{
 		Modify();
@@ -1558,8 +1558,8 @@ bool UNiagaraGraph::RenameParameterFromPin(const FNiagaraVariable& Parameter, FN
 	FNiagaraVariable NewParameter = Parameter;
 	NewParameter.SetName(NewName);
 
-	UE_TRANSITIONAL_OBJECT_PTR(UNiagaraScriptVariable)* FoundOldScriptVariablePtr = VariableToScriptVariable.Find(Parameter);
-	UE_TRANSITIONAL_OBJECT_PTR(UNiagaraScriptVariable)* FoundNewScriptVariablePtr = VariableToScriptVariable.Find(NewParameter);
+	TObjectPtr<UNiagaraScriptVariable>* FoundOldScriptVariablePtr = VariableToScriptVariable.Find(Parameter);
+	TObjectPtr<UNiagaraScriptVariable>* FoundNewScriptVariablePtr = VariableToScriptVariable.Find(NewParameter);
 
 	bool bMerged = false;
 	const bool bOldScriptVariableIsStaticSwitch = FoundOldScriptVariablePtr ? (*FoundOldScriptVariablePtr)->GetIsStaticSwitch() : false;
@@ -1634,7 +1634,7 @@ bool UNiagaraGraph::RenameParameter(const FNiagaraVariable& Parameter, FName New
 	FNiagaraVariable NewParameter = Parameter;
 	NewParameter.SetName(NewName);
 
-	UE_TRANSITIONAL_OBJECT_PTR(UNiagaraScriptVariable)* OldScriptVariablePtr = VariableToScriptVariable.Find(Parameter);
+	TObjectPtr<UNiagaraScriptVariable>* OldScriptVariablePtr = VariableToScriptVariable.Find(Parameter);
 	UNiagaraScriptVariable* OldScriptVariable = OldScriptVariablePtr ? *OldScriptVariablePtr : nullptr;
 	FNiagaraVariableMetaData OldMetaData;
 	OldMetaData.CreateNewGuid();
@@ -1649,7 +1649,7 @@ bool UNiagaraGraph::RenameParameter(const FNiagaraVariable& Parameter, FName New
 		OldMetaData = OldScriptVariable->Metadata;
 	}
 		
-	UE_TRANSITIONAL_OBJECT_PTR(UNiagaraScriptVariable)* NewScriptVariablePtr = VariableToScriptVariable.Find(NewParameter);
+	TObjectPtr<UNiagaraScriptVariable>* NewScriptVariablePtr = VariableToScriptVariable.Find(NewParameter);
 
 	// Swap metadata to the new parameter; put the new parameter into VariableToScriptVariable
 	if (OldScriptVariable != nullptr)
@@ -1783,7 +1783,7 @@ bool UNiagaraGraph::SynchronizeScriptVariable(const UNiagaraScriptVariable* Sour
 		const FGuid& SourceScriptVarId = SourceScriptVar->Metadata.GetVariableGuid();
 		TArray<TObjectPtr<UNiagaraScriptVariable>> ScriptVariables;
 		VariableToScriptVariable.GenerateValueArray(ScriptVariables);
-		UE_TRANSITIONAL_OBJECT_PTR(UNiagaraScriptVariable)* ScriptVarPtr = ScriptVariables.FindByPredicate([&SourceScriptVarId](const UNiagaraScriptVariable* ScriptVar) { return ScriptVar->Metadata.GetVariableGuid() == SourceScriptVarId; });
+		TObjectPtr<UNiagaraScriptVariable>* ScriptVarPtr = ScriptVariables.FindByPredicate([&SourceScriptVarId](const UNiagaraScriptVariable* ScriptVar) { return ScriptVar->Metadata.GetVariableGuid() == SourceScriptVarId; });
 		if(ScriptVarPtr == nullptr)
 		{
 			// Failed to find a DestScriptVar with an Id matching that of SourceScriptVar.
@@ -2679,7 +2679,7 @@ void UNiagaraGraph::MarkGraphRequiresSynchronization(FString Reason)
 
 TOptional<FNiagaraVariableMetaData> UNiagaraGraph::GetMetaData(const FNiagaraVariable& InVar) const
 {
-	if (UE_TRANSITIONAL_OBJECT_PTR(UNiagaraScriptVariable)* MetaData = VariableToScriptVariable.Find(InVar))
+	if (TObjectPtr<UNiagaraScriptVariable>* MetaData = VariableToScriptVariable.Find(InVar))
 	{
 		if (*MetaData)
 		{
@@ -2691,7 +2691,7 @@ TOptional<FNiagaraVariableMetaData> UNiagaraGraph::GetMetaData(const FNiagaraVar
 
 void UNiagaraGraph::SetMetaData(const FNiagaraVariable& InVar, const FNiagaraVariableMetaData& InMetaData)
 {
-	if (UE_TRANSITIONAL_OBJECT_PTR(UNiagaraScriptVariable)* FoundMetaData = VariableToScriptVariable.Find(InVar))
+	if (TObjectPtr<UNiagaraScriptVariable>* FoundMetaData = VariableToScriptVariable.Find(InVar))
 	{
 		if (*FoundMetaData)
 		{
@@ -2708,7 +2708,7 @@ void UNiagaraGraph::SetMetaData(const FNiagaraVariable& InVar, const FNiagaraVar
 	else 
 	{
 		Modify();
-		UE_TRANSITIONAL_OBJECT_PTR(UNiagaraScriptVariable)& NewScriptVariable = VariableToScriptVariable.Add(InVar, NewObject<UNiagaraScriptVariable>(this, FName(), RF_Transactional));
+		TObjectPtr<UNiagaraScriptVariable>& NewScriptVariable = VariableToScriptVariable.Add(InVar, NewObject<UNiagaraScriptVariable>(this, FName(), RF_Transactional));
 		NewScriptVariable->Init(InVar, InMetaData);
 		NewScriptVariable->SetIsStaticSwitch(FindStaticSwitchInputs().Contains(InVar));
 	}
diff --git a/Engine/Plugins/FX/Niagara/Source/NiagaraEditor/Private/NiagaraMessages.cpp b/Engine/Plugins/FX/Niagara/Source/NiagaraEditor/Private/NiagaraMessages.cpp
--- a/Engine/Plugins/FX/Niagara/Source/NiagaraEditor/Private/NiagaraMessages.cpp
+++ b/Engine/Plugins/FX/Niagara/Source/NiagaraEditor/Private/NiagaraMessages.cpp
@@ -300,7 +300,7 @@ bool FNiagaraMessageJobCompileEvent::RecursiveGetScriptNamesAndAssetPathsFromCon
 
 	if (NodeGuid.IsValid())
 	{
-		UE_TRANSITIONAL_OBJECT_PTR(UEdGraphNode) const* EventNodePtr = InGraphToSearch->Nodes.FindByPredicate([NodeGuid](UEdGraphNode* Node) { return Node->NodeGuid == NodeGuid; });
+		TObjectPtr<UEdGraphNode> const* EventNodePtr = InGraphToSearch->Nodes.FindByPredicate([NodeGuid](UEdGraphNode* Node) { return Node->NodeGuid == NodeGuid; });
 		if (EventNodePtr != nullptr)
 		{
 			OutContextNodeObjectKeys.AddUnique(FObjectKey(*EventNodePtr));
diff --git a/Engine/Plugins/FX/Niagara/Source/NiagaraEditor/Private/NiagaraParameterDefinitions.cpp b/Engine/Plugins/FX/Niagara/Source/NiagaraEditor/Private/NiagaraParameterDefinitions.cpp
--- a/Engine/Plugins/FX/Niagara/Source/NiagaraEditor/Private/NiagaraParameterDefinitions.cpp
+++ b/Engine/Plugins/FX/Niagara/Source/NiagaraEditor/Private/NiagaraParameterDefinitions.cpp
@@ -73,7 +73,7 @@ void UNiagaraParameterDefinitions::PostEditChangeProperty(struct FPropertyChange
 void UNiagaraParameterDefinitions::AddParameter(const FNiagaraVariable& NewVariable)
 {
 	Modify();
-	UE_TRANSITIONAL_OBJECT_PTR(UNiagaraScriptVariable)& NewScriptVariable = ScriptVariables.Add_GetRef(NewObject<UNiagaraScriptVariable>(this, FName(), RF_Transactional));
+	TObjectPtr<UNiagaraScriptVariable>& NewScriptVariable = ScriptVariables.Add_GetRef(NewObject<UNiagaraScriptVariable>(this, FName(), RF_Transactional));
 	NewScriptVariable->Init(NewVariable, FNiagaraVariableMetaData());
 	NewScriptVariable->SetIsStaticSwitch(false);
 	NewScriptVariable->SetIsSubscribedToParameterDefinitions(true);
@@ -98,7 +98,7 @@ void UNiagaraParameterDefinitions::RemoveParameter(const FNiagaraVariable& Varia
 
 void UNiagaraParameterDefinitions::RenameParameter(const FNiagaraVariable& VariableToRename, const FName NewName)
 {
-	if (UE_TRANSITIONAL_OBJECT_PTR(UNiagaraScriptVariable)* ScriptVariablePtr = ScriptVariables.FindByPredicate([&VariableToRename](const UNiagaraScriptVariable* ScriptVariable) { return ScriptVariable->Variable == VariableToRename; }))
+	if (TObjectPtr<UNiagaraScriptVariable>* ScriptVariablePtr = ScriptVariables.FindByPredicate([&VariableToRename](const UNiagaraScriptVariable* ScriptVariable) { return ScriptVariable->Variable == VariableToRename; }))
 	{
 		const FNiagaraVariable OldVariable = VariableToRename;
 		Modify();
@@ -141,7 +141,7 @@ TSet<FGuid> UNiagaraParameterDefinitions::GetParameterIds() const
 
 UNiagaraScriptVariable* UNiagaraParameterDefinitions::GetScriptVariable(const FNiagaraVariable& Var)
 {
-	if (UE_TRANSITIONAL_OBJECT_PTR(UNiagaraScriptVariable) const* ScriptVarPtr = ScriptVariables.FindByPredicate([Var](const UNiagaraScriptVariable* ScriptVar) { return ScriptVar->Variable == Var; }))
+	if (TObjectPtr<UNiagaraScriptVariable> const* ScriptVarPtr = ScriptVariables.FindByPredicate([Var](const UNiagaraScriptVariable* ScriptVar) { return ScriptVar->Variable == Var; }))
 	{
 		return *ScriptVarPtr;
 	}
@@ -150,7 +150,7 @@ UNiagaraScriptVariable* UNiagaraParameterDefinitions::GetScriptVariable(const FN
 
 UNiagaraScriptVariable* UNiagaraParameterDefinitions::GetScriptVariable(const FGuid& ScriptVarId)
 {
-	if (UE_TRANSITIONAL_OBJECT_PTR(UNiagaraScriptVariable) const* ScriptVarPtr = ScriptVariables.FindByPredicate([ScriptVarId](const UNiagaraScriptVariable* ScriptVar) { return ScriptVar->Metadata.GetVariableGuid() == ScriptVarId; }))
+	if (TObjectPtr<UNiagaraScriptVariable> const* ScriptVarPtr = ScriptVariables.FindByPredicate([ScriptVarId](const UNiagaraScriptVariable* ScriptVar) { return ScriptVar->Metadata.GetVariableGuid() == ScriptVarId; }))
 	{
 		return *ScriptVarPtr;
 	}
@@ -213,7 +213,7 @@ void UNiagaraParameterDefinitions::SynchronizeWithSubscribedParameterDefinitions
 				const FName& ExternalScriptVarName = ExternalScriptVar->Variable.GetName();
 				for (const FGuid& InternalScriptVarId : BindingNameSubscription.InternalScriptVarIds)
 				{
-					if (UE_TRANSITIONAL_OBJECT_PTR(UNiagaraScriptVariable) const* ScriptVarPtr = ScriptVariables.FindByPredicate([&InternalScriptVarId](const UNiagaraScriptVariable* ScriptVar) { return ScriptVar->Metadata.GetVariableGuid() == InternalScriptVarId; }))
+					if (TObjectPtr<UNiagaraScriptVariable> const* ScriptVarPtr = ScriptVariables.FindByPredicate([&InternalScriptVarId](const UNiagaraScriptVariable* ScriptVar) { return ScriptVar->Metadata.GetVariableGuid() == InternalScriptVarId; }))
 					{
 						UNiagaraScriptVariable* ScriptVar = *ScriptVarPtr;
 						if (ScriptVar->GetChangeId() != ExternalScriptVarChangeId)
diff --git a/Engine/Plugins/FX/Niagara/Source/NiagaraEditor/Private/Sequencer/NiagaraSequence/MovieSceneNiagaraEmitterTrack.cpp b/Engine/Plugins/FX/Niagara/Source/NiagaraEditor/Private/Sequencer/NiagaraSequence/MovieSceneNiagaraEmitterTrack.cpp
--- a/Engine/Plugins/FX/Niagara/Source/NiagaraEditor/Private/Sequencer/NiagaraSequence/MovieSceneNiagaraEmitterTrack.cpp
+++ b/Engine/Plugins/FX/Niagara/Source/NiagaraEditor/Private/Sequencer/NiagaraSequence/MovieSceneNiagaraEmitterTrack.cpp
@@ -253,7 +253,7 @@ void UMovieSceneNiagaraEmitterTrack::CreateSections(const FFrameRate& InFrameRes
 					UMovieSceneNiagaraEmitterSectionBase* EmitterSection = nullptr;
 					bool bSectionCreated = false;
 
-					UE_TRANSITIONAL_OBJECT_PTR(UMovieSceneSection)* SharedEmitterSection = Sections.FindByPredicate([=](UMovieSceneSection* Section) {	return Section->IsA(SectionClass); });
+					TObjectPtr<UMovieSceneSection>* SharedEmitterSection = Sections.FindByPredicate([=](UMovieSceneSection* Section) {	return Section->IsA(SectionClass); });
 					if (SharedEmitterSection != nullptr)
 					{
 						EmitterSection = CastChecked<UMovieSceneNiagaraEmitterSectionBase>(*SharedEmitterSection);
diff --git a/Engine/Plugins/FX/Niagara/Source/NiagaraEditor/Private/ViewModels/NiagaraParameterPanelViewModel.cpp b/Engine/Plugins/FX/Niagara/Source/NiagaraEditor/Private/ViewModels/NiagaraParameterPanelViewModel.cpp
--- a/Engine/Plugins/FX/Niagara/Source/NiagaraEditor/Private/ViewModels/NiagaraParameterPanelViewModel.cpp
+++ b/Engine/Plugins/FX/Niagara/Source/NiagaraEditor/Private/ViewModels/NiagaraParameterPanelViewModel.cpp
@@ -1522,7 +1522,7 @@ TArray<FNiagaraParameterPanelItem> FNiagaraSystemToolkitParameterPanelViewModel:
 		for (const FNiagaraVariable& Var : Vars)
 		{
 			UNiagaraScriptVariable* ScriptVar;
-			if (UE_TRANSITIONAL_OBJECT_PTR(UNiagaraScriptVariable) const* ScriptVarPtr = ParameterToScriptVariableMap.Find(Var))
+			if (TObjectPtr<UNiagaraScriptVariable> const* ScriptVarPtr = ParameterToScriptVariableMap.Find(Var))
 			{
 				ScriptVar = *ScriptVarPtr;
 			}
@@ -1674,7 +1674,7 @@ TArray<FNiagaraParameterPanelItem> FNiagaraSystemToolkitParameterPanelViewModel:
 					}
 
 					// -Lookup the script variable.
-					UE_TRANSITIONAL_OBJECT_PTR(UNiagaraScriptVariable) const* ScriptVarPtr = ParameterToScriptVariableMap.Find(Var);
+					TObjectPtr<UNiagaraScriptVariable> const* ScriptVarPtr = ParameterToScriptVariableMap.Find(Var);
 					UNiagaraScriptVariable* ScriptVar;
 					if (ScriptVarPtr != nullptr)
 					{
@@ -2616,7 +2616,7 @@ TArray<FNiagaraParameterPanelItem> FNiagaraScriptToolkitParameterPanelViewModel:
 			{
 				// This variable has not been registered, prepare the FNiagaraParameterPanelItem.
 				// -First lookup the script variable.
-				UE_TRANSITIONAL_OBJECT_PTR(UNiagaraScriptVariable) const* ScriptVarPtr = ParameterToScriptVariableMap.Find(Var);
+				TObjectPtr<UNiagaraScriptVariable> const* ScriptVarPtr = ParameterToScriptVariableMap.Find(Var);
 				UNiagaraScriptVariable* ScriptVar = ScriptVarPtr != nullptr ? *ScriptVarPtr : nullptr;
 				if (!ScriptVar)
 				{
diff --git a/Engine/Plugins/FX/Niagara/Source/NiagaraEditor/Private/ViewModels/Stack/NiagaraStackEntry.cpp b/Engine/Plugins/FX/Niagara/Source/NiagaraEditor/Private/ViewModels/Stack/NiagaraStackEntry.cpp
--- a/Engine/Plugins/FX/Niagara/Source/NiagaraEditor/Private/ViewModels/Stack/NiagaraStackEntry.cpp
+++ b/Engine/Plugins/FX/Niagara/Source/NiagaraEditor/Private/ViewModels/Stack/NiagaraStackEntry.cpp
@@ -732,7 +732,7 @@ void UNiagaraStackEntry::RefreshStackErrorChildren()
 	{
 		const FStackIssue& Issue = StackIssues[i];
 		UNiagaraStackErrorItem* ErrorEntry = nullptr;
-		UE_TRANSITIONAL_OBJECT_PTR(UNiagaraStackErrorItem)* Found = ErrorChildren.FindByPredicate(
+		TObjectPtr<UNiagaraStackErrorItem>* Found = ErrorChildren.FindByPredicate(
 			[&](UNiagaraStackErrorItem* CurrentChild) { return CurrentChild->GetStackIssue().GetUniqueIdentifier() == Issue.GetUniqueIdentifier(); });
 		if (Found == nullptr)
 		{
diff --git a/Engine/Plugins/FX/Niagara/Source/NiagaraEditor/Private/Widgets/SNiagaraScriptGraph.cpp b/Engine/Plugins/FX/Niagara/Source/NiagaraEditor/Private/Widgets/SNiagaraScriptGraph.cpp
--- a/Engine/Plugins/FX/Niagara/Source/NiagaraEditor/Private/Widgets/SNiagaraScriptGraph.cpp
+++ b/Engine/Plugins/FX/Niagara/Source/NiagaraEditor/Private/Widgets/SNiagaraScriptGraph.cpp
@@ -506,7 +506,7 @@ void SNiagaraScriptGraph::FocusGraphElement(const INiagaraScriptGraphFocusInfo*
 	{
 		const FNiagaraScriptGraphNodeToFocusInfo* NodeFocusInfo = static_cast<const FNiagaraScriptGraphNodeToFocusInfo*>(FocusInfo);
 		const FGuid& NodeGuidToMatch = NodeFocusInfo->GetNodeGuidToFocus();
-		UE_TRANSITIONAL_OBJECT_PTR(UEdGraphNode) const* NodeToFocus = ViewModel->GetGraph()->Nodes.FindByPredicate([&NodeGuidToMatch](const UEdGraphNode* Node) {return Node->NodeGuid == NodeGuidToMatch; });
+		TObjectPtr<UEdGraphNode> const* NodeToFocus = ViewModel->GetGraph()->Nodes.FindByPredicate([&NodeGuidToMatch](const UEdGraphNode* Node) {return Node->NodeGuid == NodeGuidToMatch; });
 		if (NodeToFocus != nullptr && *NodeToFocus != nullptr)
 		{
 			GetGraphEditor()->JumpToNode(*NodeToFocus);
diff --git a/Engine/Plugins/Media/MediaPlayerEditor/Source/MediaPlayerEditor/Private/Customizations/PlatformMediaSourceCustomization.cpp b/Engine/Plugins/Media/MediaPlayerEditor/Source/MediaPlayerEditor/Private/Customizations/PlatformMediaSourceCustomization.cpp
--- a/Engine/Plugins/Media/MediaPlayerEditor/Source/MediaPlayerEditor/Private/Customizations/PlatformMediaSourceCustomization.cpp
+++ b/Engine/Plugins/Media/MediaPlayerEditor/Source/MediaPlayerEditor/Private/Customizations/PlatformMediaSourceCustomization.cpp
@@ -142,7 +142,7 @@ void FPlatformMediaSourceCustomization::SetPlatformMediaSourcesValue(FString Pla
 
 	for (auto Object : OuterObjects)
 	{
-		UE_TRANSITIONAL_OBJECT_PTR(UMediaSource)& OldMediaSource = Cast<UPlatformMediaSource>(Object)->PlatformMediaSources.FindOrAdd(PlatformName);;
+		TObjectPtr<UMediaSource>& OldMediaSource = Cast<UPlatformMediaSource>(Object)->PlatformMediaSources.FindOrAdd(PlatformName);;
 
 		if (OldMediaSource != MediaSource)
 		{
@@ -165,7 +165,7 @@ void FPlatformMediaSourceCustomization::HandleMediaSourceEntryBoxChanged(const F
 
 	for (auto Object : OuterObjects)
 	{
-		UE_TRANSITIONAL_OBJECT_PTR(UMediaSource)& OldMediaSource = Cast<UPlatformMediaSource>(Object)->PlatformMediaSources.FindOrAdd(PlatformName);;
+		TObjectPtr<UMediaSource>& OldMediaSource = Cast<UPlatformMediaSource>(Object)->PlatformMediaSources.FindOrAdd(PlatformName);;
 
 		if (OldMediaSource != AssetData.GetAsset())
 		{
diff --git a/Engine/Plugins/MeshPainting/Source/MeshPaintingToolset/Private/MeshTexturePaintingTool.cpp b/Engine/Plugins/MeshPainting/Source/MeshPaintingToolset/Private/MeshTexturePaintingTool.cpp
--- a/Engine/Plugins/MeshPainting/Source/MeshPaintingToolset/Private/MeshTexturePaintingTool.cpp
+++ b/Engine/Plugins/MeshPainting/Source/MeshPaintingToolset/Private/MeshTexturePaintingTool.cpp
@@ -1189,7 +1189,7 @@ void UMeshTexturePaintingTool::CycleTextures(int32 Direction)
 	{
 		return;
 	}
-	UE_TRANSITIONAL_OBJECT_PTR(UTexture2D)& SelectedTexture = TextureProperties->PaintTexture;
+	TObjectPtr<UTexture2D>& SelectedTexture = TextureProperties->PaintTexture;
 	const int32 TextureIndex = (SelectedTexture != nullptr) ? PaintableTextures.IndexOfByKey(SelectedTexture) : 0;
 	if (TextureIndex != INDEX_NONE)
 	{
diff --git a/Engine/Plugins/Runtime/HairStrands/Source/HairStrandsCore/Private/GroomComponent.cpp b/Engine/Plugins/Runtime/HairStrands/Source/HairStrandsCore/Private/GroomComponent.cpp
--- a/Engine/Plugins/Runtime/HairStrands/Source/HairStrandsCore/Private/GroomComponent.cpp
+++ b/Engine/Plugins/Runtime/HairStrands/Source/HairStrandsCore/Private/GroomComponent.cpp
@@ -1334,7 +1334,7 @@ void UGroomComponent::ReleaseHairSimulation()
 {
 	for (int32 CompIndex = 0; CompIndex < NiagaraComponents.Num(); ++CompIndex)
 	{
-		UE_TRANSITIONAL_OBJECT_PTR(UNiagaraComponent)& NiagaraComponent = NiagaraComponents[CompIndex];
+		TObjectPtr<UNiagaraComponent>& NiagaraComponent = NiagaraComponents[CompIndex];
 		if (NiagaraComponent && !NiagaraComponent->IsBeingDestroyed())
 		{
 			if (GetWorld() && NiagaraComponent->IsRegistered())
@@ -1353,7 +1353,7 @@ void UGroomComponent::CreateHairSimulation(const int32 GroupIndex, const int32 L
 {
 	if (GroupIndex < NiagaraComponents.Num())
 	{
-		UE_TRANSITIONAL_OBJECT_PTR(UNiagaraComponent)& NiagaraComponent = NiagaraComponents[GroupIndex];
+		TObjectPtr<UNiagaraComponent>& NiagaraComponent = NiagaraComponents[GroupIndex];
 		if (GroomAsset && (GroupIndex < GroomAsset->HairGroupsPhysics.Num()) && (LODIndex < GroomAsset->HairGroupsLOD[GroupIndex].LODs.Num()) && IsSimulationEnable(GroupIndex, LODIndex))
 		{
 			if (!NiagaraComponent)
diff --git a/Engine/Plugins/Runtime/Steam/SteamSockets/Source/SteamSockets/Private/SteamSocketsNetDriver.cpp b/Engine/Plugins/Runtime/Steam/SteamSockets/Source/SteamSockets/Private/SteamSocketsNetDriver.cpp
--- a/Engine/Plugins/Runtime/Steam/SteamSockets/Source/SteamSockets/Private/SteamSocketsNetDriver.cpp
+++ b/Engine/Plugins/Runtime/Steam/SteamSockets/Source/SteamSockets/Private/SteamSocketsNetDriver.cpp
@@ -435,7 +435,7 @@ void USteamSocketsNetDriver::ResetSocketInfo(const FSteamSocket* RemovedSocket)
 
 UNetConnection* USteamSocketsNetDriver::FindClientConnectionForHandle(SteamSocketHandles SocketHandle)
 {
-	for (UE_TRANSITIONAL_OBJECT_PTR(UNetConnection)& ClientConnection : ClientConnections)
+	for (TObjectPtr<UNetConnection>& ClientConnection : ClientConnections)
 	{
 		USteamSocketsNetConnection* SteamConnection = static_cast<USteamSocketsNetConnection*>(ClientConnection);
 		if (SteamConnection && SteamConnection->GetRawSocket() != nullptr)
diff --git a/Engine/Source/Developer/FunctionalTesting/Private/FunctionalTest.cpp b/Engine/Source/Developer/FunctionalTesting/Private/FunctionalTest.cpp
--- a/Engine/Source/Developer/FunctionalTesting/Private/FunctionalTest.cpp
+++ b/Engine/Source/Developer/FunctionalTesting/Private/FunctionalTest.cpp
@@ -388,7 +388,7 @@ void AFunctionalTest::FinishTest(EFunctionalTestResult TestResult, const FString
 
 	OnTestFinished.Broadcast();
 
-	UE_TRANSITIONAL_OBJECT_PTR(AActor)* ActorToDestroy = AutoDestroyActors.GetData();
+	TObjectPtr<AActor>* ActorToDestroy = AutoDestroyActors.GetData();
 
 	for (int32 ActorIndex = 0; ActorIndex < AutoDestroyActors.Num(); ++ActorIndex, ++ActorToDestroy)
 	{
diff --git a/Engine/Source/Developer/ToolMenus/Private/ToolMenus.cpp b/Engine/Source/Developer/ToolMenus/Private/ToolMenus.cpp
--- a/Engine/Source/Developer/ToolMenus/Private/ToolMenus.cpp
+++ b/Engine/Source/Developer/ToolMenus/Private/ToolMenus.cpp
@@ -397,13 +397,13 @@ bool UToolMenus::GetDisplayUIExtensionPoints() const
 
 UToolMenu* UToolMenus::FindMenu(const FName Name)
 {
-	UE_TRANSITIONAL_OBJECT_PTR(UToolMenu)* Found = Menus.Find(Name);
+	TObjectPtr<UToolMenu>* Found = Menus.Find(Name);
 	return Found ? *Found : nullptr;
 }
 
 bool UToolMenus::IsMenuRegistered(const FName Name) const
 {
-	UE_TRANSITIONAL_OBJECT_PTR(UToolMenu) const * Found = Menus.Find(Name);
+	TObjectPtr<UToolMenu> const * Found = Menus.Find(Name);
 	return Found && *Found && (*Found)->IsRegistered();
 }
 
diff --git a/Engine/Source/Editor/BehaviorTreeEditor/Private/SBehaviorTreeDiff.cpp b/Engine/Source/Editor/BehaviorTreeEditor/Private/SBehaviorTreeDiff.cpp
--- a/Engine/Source/Editor/BehaviorTreeEditor/Private/SBehaviorTreeDiff.cpp
+++ b/Engine/Source/Editor/BehaviorTreeEditor/Private/SBehaviorTreeDiff.cpp
@@ -369,7 +369,7 @@ void SBehaviorTreeDiff::OnSelectionChanged(FSharedDiffOnGraph Item, ESelectInfo:
 				{
 					// This is a sub-node, we need to find our parent node in the graph
 					// todo: work out why BTNode->ParentNode is always null
-					UE_TRANSITIONAL_OBJECT_PTR(UEdGraphNode)* ParentNodePtr = NodeGraph->Nodes.FindByPredicate([BTNode](UEdGraphNode* PotentialParentNode) -> bool
+					TObjectPtr<UEdGraphNode>* ParentNodePtr = NodeGraph->Nodes.FindByPredicate([BTNode](UEdGraphNode* PotentialParentNode) -> bool
 					{
 						UBehaviorTreeGraphNode* BTPotentialParentNode = Cast<UBehaviorTreeGraphNode>(PotentialParentNode);
 						return BTPotentialParentNode && (BTPotentialParentNode->Decorators.Contains(BTNode) || BTPotentialParentNode->Services.Contains(BTNode));
diff --git a/Engine/Source/Editor/Experimental/EditorInteractiveToolsFramework/Private/EditorInteractiveGizmoRegistry.cpp b/Engine/Source/Editor/Experimental/EditorInteractiveToolsFramework/Private/EditorInteractiveGizmoRegistry.cpp
--- a/Engine/Source/Editor/Experimental/EditorInteractiveToolsFramework/Private/EditorInteractiveGizmoRegistry.cpp
+++ b/Engine/Source/Editor/Experimental/EditorInteractiveToolsFramework/Private/EditorInteractiveGizmoRegistry.cpp
@@ -156,7 +156,7 @@ void UEditorInteractiveGizmoRegistry::RegisterEditorGizmoType(EEditorGizmoCatego
 {
 	if (ensure(InGizmoBuilder))
 	{
-		if (UE_TRANSITIONAL_OBJECT_PTR(UEditorGizmoRegistryCategoryEntry)* RegEntry = GizmoCategoryMap.Find(InGizmoCategory))
+		if (TObjectPtr<UEditorGizmoRegistryCategoryEntry>* RegEntry = GizmoCategoryMap.Find(InGizmoCategory))
 		{
 			if (*RegEntry)
 			{
@@ -168,7 +168,7 @@ void UEditorInteractiveGizmoRegistry::RegisterEditorGizmoType(EEditorGizmoCatego
 
 void UEditorInteractiveGizmoRegistry::GetQualifiedEditorGizmoBuilders(EEditorGizmoCategory InGizmoCategory, const FToolBuilderState& InToolBuilderState, TArray<UInteractiveGizmoBuilder*>& InFoundBuilders)
 {
-	if (UE_TRANSITIONAL_OBJECT_PTR(UEditorGizmoRegistryCategoryEntry)* RegEntry = GizmoCategoryMap.Find(InGizmoCategory))
+	if (TObjectPtr<UEditorGizmoRegistryCategoryEntry>* RegEntry = GizmoCategoryMap.Find(InGizmoCategory))
 	{
 		if (*RegEntry)
 		{
@@ -181,7 +181,7 @@ void UEditorInteractiveGizmoRegistry::DeregisterEditorGizmoType(EEditorGizmoCate
 {
 	if (ensure(InGizmoBuilder))
 	{
-		if (UE_TRANSITIONAL_OBJECT_PTR(UEditorGizmoRegistryCategoryEntry)* RegEntry = GizmoCategoryMap.Find(InGizmoCategory))
+		if (TObjectPtr<UEditorGizmoRegistryCategoryEntry>* RegEntry = GizmoCategoryMap.Find(InGizmoCategory))
 		{
 			if (*RegEntry)
 			{
diff --git a/Engine/Source/Editor/LocalizationDashboard/Private/LocalizationTargetDetailCustomization.cpp b/Engine/Source/Editor/LocalizationDashboard/Private/LocalizationTargetDetailCustomization.cpp
--- a/Engine/Source/Editor/LocalizationDashboard/Private/LocalizationTargetDetailCustomization.cpp
+++ b/Engine/Source/Editor/LocalizationDashboard/Private/LocalizationTargetDetailCustomization.cpp
@@ -761,7 +761,7 @@ void FLocalizationTargetDetailCustomization::RebuildTargetsList()
 
 		for (const FGuid& OtherTargetDependencyGuid : OtherTarget->Settings.TargetDependencies)
 		{
-			UE_TRANSITIONAL_OBJECT_PTR(ULocalizationTarget)* const OtherTargetDependency = TargetSet->TargetObjects.FindByPredicate([OtherTargetDependencyGuid](ULocalizationTarget* SomeLocalizationTarget)->bool{return SomeLocalizationTarget->Settings.Guid == OtherTargetDependencyGuid;});
+			TObjectPtr<ULocalizationTarget>* const OtherTargetDependency = TargetSet->TargetObjects.FindByPredicate([OtherTargetDependencyGuid](ULocalizationTarget* SomeLocalizationTarget)->bool{return SomeLocalizationTarget->Settings.Guid == OtherTargetDependencyGuid;});
 			if (OtherTargetDependency && DoesTargetDependOnUs(*OtherTargetDependency))
 			{
 				return true;
diff --git a/Engine/Source/Editor/MovieSceneTools/Private/MovieSceneToolsModule.cpp b/Engine/Source/Editor/MovieSceneTools/Private/MovieSceneToolsModule.cpp
--- a/Engine/Source/Editor/MovieSceneTools/Private/MovieSceneToolsModule.cpp
+++ b/Engine/Source/Editor/MovieSceneTools/Private/MovieSceneToolsModule.cpp
@@ -335,9 +335,9 @@ bool FMovieSceneToolsModule::UpgradeLegacyEventEndpointForSection(UMovieSceneEve
 			{
 				if (EntryPoint.NodeGuid_DEPRECATED.IsValid())
 				{
-					if (UE_TRANSITIONAL_OBJECT_PTR(UEdGraph) const* GraphPtr = Algo::FindBy(SequenceDirectorBP->UbergraphPages, EntryPoint.GraphGuid_DEPRECATED, &UEdGraph::GraphGuid))
+					if (TObjectPtr<UEdGraph> const* GraphPtr = Algo::FindBy(SequenceDirectorBP->UbergraphPages, EntryPoint.GraphGuid_DEPRECATED, &UEdGraph::GraphGuid))
 					{
-						UE_TRANSITIONAL_OBJECT_PTR(UEdGraphNode) const* NodePtr  = Algo::FindBy((*GraphPtr)->Nodes, EntryPoint.NodeGuid_DEPRECATED, &UEdGraphNode::NodeGuid);
+						TObjectPtr<UEdGraphNode> const* NodePtr  = Algo::FindBy((*GraphPtr)->Nodes, EntryPoint.NodeGuid_DEPRECATED, &UEdGraphNode::NodeGuid);
 						if (NodePtr)
 						{
 							UK2Node_CustomEvent* CustomEvent = Cast<UK2Node_CustomEvent>(*NodePtr);
@@ -350,9 +350,9 @@ bool FMovieSceneToolsModule::UpgradeLegacyEventEndpointForSection(UMovieSceneEve
 					}
 				}
 				// If the node guid is invalid, this must be a function graph on the BP
-				else if (UE_TRANSITIONAL_OBJECT_PTR(UEdGraph) const* GraphPtr = Algo::FindBy(SequenceDirectorBP->FunctionGraphs, EntryPoint.GraphGuid_DEPRECATED, &UEdGraph::GraphGuid))
+				else if (TObjectPtr<UEdGraph> const* GraphPtr = Algo::FindBy(SequenceDirectorBP->FunctionGraphs, EntryPoint.GraphGuid_DEPRECATED, &UEdGraph::GraphGuid))
 				{
-					UE_TRANSITIONAL_OBJECT_PTR(UEdGraphNode) const* NodePtr = Algo::FindByPredicate((*GraphPtr)->Nodes, [](UEdGraphNode* InNode){ return InNode && InNode->IsA<UK2Node_FunctionEntry>(); });
+					TObjectPtr<UEdGraphNode> const* NodePtr = Algo::FindByPredicate((*GraphPtr)->Nodes, [](UEdGraphNode* InNode){ return InNode && InNode->IsA<UK2Node_FunctionEntry>(); });
 					if (NodePtr)
 					{
 						UK2Node_FunctionEntry* FunctionEntry = CastChecked<UK2Node_FunctionEntry>(*NodePtr);
diff --git a/Engine/Source/Editor/VREditor/UI/VREditorUISystem.cpp b/Engine/Source/Editor/VREditor/UI/VREditorUISystem.cpp
--- a/Engine/Source/Editor/VREditor/UI/VREditorUISystem.cpp
+++ b/Engine/Source/Editor/VREditor/UI/VREditorUISystem.cpp
@@ -1102,7 +1102,7 @@ void UVREditorUISystem::ShowEditorUIPanel(AVREditorFloatingUI* Panel, UVREditorI
 			else if (PanelID.ToString().Contains(ActorPreviewPrefix))
 			{
 				// if we are closing a preview panel, make sure to unpin it also
-				UE_TRANSITIONAL_OBJECT_PTR(AActor)* PreviewedActorPtr = PreviewWindowInfo.Find(PanelID);
+				TObjectPtr<AActor>* PreviewedActorPtr = PreviewWindowInfo.Find(PanelID);
 				if (PreviewedActorPtr != nullptr)
 				{
 					AActor* PreviewedActor = *PreviewedActorPtr;
@@ -1832,7 +1832,7 @@ AVREditorFloatingUI* UVREditorUISystem::GetPanel(const VREditorPanelID& InPanelI
 	AVREditorFloatingUI* Result = nullptr;
 	if (FloatingUIs.Num() > 0)
 	{
-		UE_TRANSITIONAL_OBJECT_PTR(AVREditorFloatingUI) const * FoundPanel = FloatingUIs.Find(InPanelID);
+		TObjectPtr<AVREditorFloatingUI> const * FoundPanel = FloatingUIs.Find(InPanelID);
 		if (FoundPanel != nullptr)
 		{
 			Result = *FoundPanel;
diff --git a/Engine/Source/Runtime/AIModule/Private/Perception/AIPerceptionComponent.cpp b/Engine/Source/Runtime/AIModule/Private/Perception/AIPerceptionComponent.cpp
--- a/Engine/Source/Runtime/AIModule/Private/Perception/AIPerceptionComponent.cpp
+++ b/Engine/Source/Runtime/AIModule/Private/Perception/AIPerceptionComponent.cpp
@@ -118,7 +118,7 @@ void UAIPerceptionComponent::ConfigureSense(UAISenseConfig& Config)
 {
 	// first check if we're reconfiguring a sense
 	bool bIsNewConfig = true;
-	for (UE_TRANSITIONAL_OBJECT_PTR(UAISenseConfig)& SenseConfig : SensesConfig)
+	for (TObjectPtr<UAISenseConfig>& SenseConfig : SensesConfig)
 	{
 		if (SenseConfig != nullptr && SenseConfig->GetClass() == Config.GetClass())
 		{
diff --git a/Engine/Source/Runtime/ClothingSystemRuntimeCommon/Private/ClothingAsset.cpp b/Engine/Source/Runtime/ClothingSystemRuntimeCommon/Private/ClothingAsset.cpp
--- a/Engine/Source/Runtime/ClothingSystemRuntimeCommon/Private/ClothingAsset.cpp
+++ b/Engine/Source/Runtime/ClothingSystemRuntimeCommon/Private/ClothingAsset.cpp
@@ -1091,7 +1091,7 @@ void UClothingAssetCommon::PropagateSharedConfigs(bool bMigrateSharedConfigToCon
 		for (const auto& ClothSharedConfigItem : ClothSharedConfigs)
 		{
 			// Set share config
-			if (UE_TRANSITIONAL_OBJECT_PTR(UClothConfigBase)* const ClothConfigBase = ClothConfigs.Find(ClothSharedConfigItem.Key))
+			if (TObjectPtr<UClothConfigBase>* const ClothConfigBase = ClothConfigs.Find(ClothSharedConfigItem.Key))
 			{
 				// Reset this shared config
 				*ClothConfigBase = ClothSharedConfigItem.Value;
diff --git a/Engine/Source/Runtime/ClothingSystemRuntimeInterface/Private/ClothingSimulationInteractor.cpp b/Engine/Source/Runtime/ClothingSystemRuntimeInterface/Private/ClothingSimulationInteractor.cpp
--- a/Engine/Source/Runtime/ClothingSystemRuntimeInterface/Private/ClothingSimulationInteractor.cpp
+++ b/Engine/Source/Runtime/ClothingSystemRuntimeInterface/Private/ClothingSimulationInteractor.cpp
@@ -24,7 +24,7 @@ void UClothingSimulationInteractor::DestroyClothingInteractors()
 
 UClothingInteractor* UClothingSimulationInteractor::GetClothingInteractor(const FString& ClothingAssetName) const
 {
-	if (UE_TRANSITIONAL_OBJECT_PTR(UClothingInteractor* const)* const ClothingInteractor = ClothingInteractors.Find(FName(ClothingAssetName)))
+	if (const TObjectPtr<UClothingInteractor>* ClothingInteractor = ClothingInteractors.Find(FName(ClothingAssetName)))
 	{
 		return *ClothingInteractor;
 	}
diff --git a/Engine/Source/Runtime/Engine/Private/Animation/AssetMappingTable.cpp b/Engine/Source/Runtime/Engine/Private/Animation/AssetMappingTable.cpp
--- a/Engine/Source/Runtime/Engine/Private/Animation/AssetMappingTable.cpp
+++ b/Engine/Source/Runtime/Engine/Private/Animation/AssetMappingTable.cpp
@@ -191,7 +191,7 @@ void UAssetMappingTable::ReplaceReferredAnimations(const TMap<UAnimationAsset*,
 {
 	for (FAssetMapping& AssetMapping : MappedAssets)
 	{
-		UE_TRANSITIONAL_OBJECT_PTR(UAnimationAsset)& AnimAsset = AssetMapping.SourceAsset;
+		TObjectPtr<UAnimationAsset>& AnimAsset = AssetMapping.SourceAsset;
 		if (AnimAsset)
 		{
 			// now fix everythign else
diff --git a/Engine/Source/Runtime/Engine/Private/Animation/Skeleton.cpp b/Engine/Source/Runtime/Engine/Private/Animation/Skeleton.cpp
--- a/Engine/Source/Runtime/Engine/Private/Animation/Skeleton.cpp
+++ b/Engine/Source/Runtime/Engine/Private/Animation/Skeleton.cpp
@@ -2491,7 +2491,7 @@ UBlendProfile* USkeleton::CreateNewBlendProfile(const FName& InProfileName)
 
 UBlendProfile* USkeleton::GetBlendProfile(const FName& InProfileName)
 {
-	UE_TRANSITIONAL_OBJECT_PTR(UBlendProfile)* FoundProfile = BlendProfiles.FindByPredicate([InProfileName](const UBlendProfile* Profile)
+	TObjectPtr<UBlendProfile>* FoundProfile = BlendProfiles.FindByPredicate([InProfileName](const UBlendProfile* Profile)
 	{
 		return Profile->GetName() == InProfileName.ToString();
 	});
diff --git a/Engine/Source/Runtime/Engine/Private/CompositeCurveTable.cpp b/Engine/Source/Runtime/Engine/Private/CompositeCurveTable.cpp
--- a/Engine/Source/Runtime/Engine/Private/CompositeCurveTable.cpp
+++ b/Engine/Source/Runtime/Engine/Private/CompositeCurveTable.cpp
@@ -23,7 +23,7 @@ UCompositeCurveTable::UCompositeCurveTable(const FObjectInitializer& ObjectIniti
 void UCompositeCurveTable::GetPreloadDependencies(TArray<UObject*>& OutDeps)
 {
 	Super::GetPreloadDependencies(OutDeps);
-	for (const UE_TRANSITIONAL_OBJECT_PTR(UCurveTable)& ParentTable : ParentTables)
+	for (const TObjectPtr<UCurveTable>& ParentTable : ParentTables)
 	{
 		if (ParentTable)
 		{
@@ -266,7 +266,7 @@ const UCompositeCurveTable* UCompositeCurveTable::FindLoops(TArray<const UCompos
 {
 	AlreadySeenTables.Add(this);
 
-	for (const UE_TRANSITIONAL_OBJECT_PTR(UCurveTable)& CurveTable : ParentTables)
+	for (const TObjectPtr<UCurveTable>& CurveTable : ParentTables)
 	{
 		// we only care about composite tables since regular tables terminate the chain and can't be in loops
 		if (UCompositeCurveTable* CompositeCurveTable = Cast<UCompositeCurveTable>(CurveTable))
diff --git a/Engine/Source/Runtime/Engine/Private/Materials/MaterialInstance.cpp b/Engine/Source/Runtime/Engine/Private/Materials/MaterialInstance.cpp
--- a/Engine/Source/Runtime/Engine/Private/Materials/MaterialInstance.cpp
+++ b/Engine/Source/Runtime/Engine/Private/Materials/MaterialInstance.cpp
@@ -600,7 +600,7 @@ UMaterialInstance::UMaterialInstance(const FObjectInitializer& ObjectInitializer
 	ShadingModels = MSM_Unlit;
 
 	PhysMaterial = nullptr;
-	for (UE_TRANSITIONAL_OBJECT_PTR(UPhysicalMaterial)& PhysMat : PhysicalMaterialMap)
+	for (TObjectPtr<UPhysicalMaterial>& PhysMat : PhysicalMaterialMap)
 	{
 		PhysMat = nullptr;
 	}
diff --git a/Engine/Source/Runtime/Engine/Private/WorldSettings.cpp b/Engine/Source/Runtime/Engine/Private/WorldSettings.cpp
--- a/Engine/Source/Runtime/Engine/Private/WorldSettings.cpp
+++ b/Engine/Source/Runtime/Engine/Private/WorldSettings.cpp
@@ -780,7 +780,7 @@ class UBookmarkBase* AWorldSettings::GetOrAddBookmark(const uint32 BookmarkIndex
 {
 	if (BookmarkArray.IsValidIndex(BookmarkIndex))
 	{
-		UE_TRANSITIONAL_OBJECT_PTR(UBookmarkBase)& Bookmark = BookmarkArray[BookmarkIndex];
+		TObjectPtr<UBookmarkBase>& Bookmark = BookmarkArray[BookmarkIndex];
 
 		if (Bookmark == nullptr || (bRecreateOnClassMismatch && Bookmark->GetClass() != GetDefaultBookmarkClass()))
 		{
@@ -830,7 +830,7 @@ void AWorldSettings::ClearBookmark(const uint32 BookmarkIndex)
 {
 	if (BookmarkArray.IsValidIndex(BookmarkIndex))
 	{
-		if (UE_TRANSITIONAL_OBJECT_PTR(UBookmarkBase)& Bookmark = BookmarkArray[BookmarkIndex])
+		if (TObjectPtr<UBookmarkBase>& Bookmark = BookmarkArray[BookmarkIndex])
 		{
 			Modify();
 			Bookmark->OnCleared();
@@ -842,7 +842,7 @@ void AWorldSettings::ClearBookmark(const uint32 BookmarkIndex)
 void AWorldSettings::ClearAllBookmarks()
 {
 	Modify();
-	for (UE_TRANSITIONAL_OBJECT_PTR(UBookmarkBase)& Bookmark : BookmarkArray)
+	for (TObjectPtr<UBookmarkBase>& Bookmark : BookmarkArray)
 	{
 		if (Bookmark)
 		{
@@ -891,7 +891,7 @@ void AWorldSettings::SanitizeBookmarkClasses()
 		bool bFoundInvalidBookmarks = false;
 		for (int32 i = 0; i < BookmarkArray.Num(); ++i)
 		{
-			if (UE_TRANSITIONAL_OBJECT_PTR(UBookmarkBase)& Bookmark = BookmarkArray[i])
+			if (TObjectPtr<UBookmarkBase>& Bookmark = BookmarkArray[i])
 			{
 				if (Bookmark->GetClass() != ExpectedClass)
 				{
diff --git a/Engine/Source/Runtime/Landscape/Private/Landscape.cpp b/Engine/Source/Runtime/Landscape/Private/Landscape.cpp
--- a/Engine/Source/Runtime/Landscape/Private/Landscape.cpp
+++ b/Engine/Source/Runtime/Landscape/Private/Landscape.cpp
@@ -372,11 +372,11 @@ void ULandscapeComponent::Serialize(FArchive& Ar)
 		TArray<UObject**> TexturesAndMaterials;
 		TexturesAndMaterials.Add((UObject**)&HeightmapTexture);
 		TexturesAndMaterials.Add((UObject**)&XYOffsetmapTexture);
-		for (UE_TRANSITIONAL_OBJECT_PTR(UTexture2D)& WeightmapTexture : WeightmapTextures)
+		for (TObjectPtr<UTexture2D>& WeightmapTexture : WeightmapTextures)
 		{
 			TexturesAndMaterials.Add((UObject**)&static_cast<UTexture2D*&>(WeightmapTexture));
 		}
-		for (UE_TRANSITIONAL_OBJECT_PTR(UTexture2D)& MobileWeightmapTexture : MobileWeightmapTextures)
+		for (TObjectPtr<UTexture2D>& MobileWeightmapTexture : MobileWeightmapTextures)
 		{
 			TexturesAndMaterials.Add((UObject**)&static_cast<UTexture2D*&>(MobileWeightmapTexture));
 		}
@@ -389,15 +389,15 @@ void ULandscapeComponent::Serialize(FArchive& Ar)
 				TexturesAndMaterials.Add((UObject**)&WeightmapTexture);
 			}
 		}
-		for (UE_TRANSITIONAL_OBJECT_PTR(UMaterialInstanceConstant)& MaterialInstance : MaterialInstances)
+		for (TObjectPtr<UMaterialInstanceConstant>& MaterialInstance : MaterialInstances)
 		{
 			TexturesAndMaterials.Add((UObject**)&static_cast<UMaterialInstanceConstant*&>(MaterialInstance));
 		}
-		for (UE_TRANSITIONAL_OBJECT_PTR(UMaterialInterface)& MobileMaterialInterface : MobileMaterialInterfaces)
+		for (TObjectPtr<UMaterialInterface>& MobileMaterialInterface : MobileMaterialInterfaces)
 		{
 			TexturesAndMaterials.Add((UObject**)(&static_cast<UMaterialInterface*&>(MobileMaterialInterface)));
 		}
-		for (UE_TRANSITIONAL_OBJECT_PTR(UMaterialInstanceConstant)& MobileCombinationMaterialInstance : MobileCombinationMaterialInstances)
+		for (TObjectPtr<UMaterialInstanceConstant>& MobileCombinationMaterialInstance : MobileCombinationMaterialInstances)
 		{
 			TexturesAndMaterials.Add((UObject**)&static_cast<UMaterialInstanceConstant*&>(MobileCombinationMaterialInstance));
 		}
@@ -1338,7 +1338,7 @@ void ULandscapeComponent::BeginDestroy()
 			if (WeightmapTextures.IsValidIndex(WeightmapIndex))
 			{
 				UTexture2D* WeightmapTexture = WeightmapTextures[WeightmapIndex];
-				UE_TRANSITIONAL_OBJECT_PTR(ULandscapeWeightmapUsage)* Usage = Proxy->WeightmapUsageMap.Find(WeightmapTexture);
+				TObjectPtr<ULandscapeWeightmapUsage>* Usage = Proxy->WeightmapUsageMap.Find(WeightmapTexture);
 				if (Usage != nullptr && (*Usage) != nullptr)
 				{
 					(*Usage)->ChannelUsage[WeightmapLayerAllocations[LayerIdx].WeightmapTextureChannel] = nullptr;
diff --git a/Engine/Source/Runtime/Landscape/Private/LandscapeEdit.cpp b/Engine/Source/Runtime/Landscape/Private/LandscapeEdit.cpp
--- a/Engine/Source/Runtime/Landscape/Private/LandscapeEdit.cpp
+++ b/Engine/Source/Runtime/Landscape/Private/LandscapeEdit.cpp
@@ -733,7 +733,7 @@ void ULandscapeComponent::FixupWeightmaps()
 
 				UTexture2D* WeightmapTexture = WeightmapTextures[Allocation.WeightmapTextureIndex];
 
-				UE_TRANSITIONAL_OBJECT_PTR(ULandscapeWeightmapUsage)* TempUsage = Proxy->WeightmapUsageMap.Find(WeightmapTexture);
+				TObjectPtr<ULandscapeWeightmapUsage>* TempUsage = Proxy->WeightmapUsageMap.Find(WeightmapTexture);
 
 				if (TempUsage == nullptr)
 				{
diff --git a/Engine/Source/Runtime/Landscape/Private/LandscapeEditLayers.cpp b/Engine/Source/Runtime/Landscape/Private/LandscapeEditLayers.cpp
--- a/Engine/Source/Runtime/Landscape/Private/LandscapeEditLayers.cpp
+++ b/Engine/Source/Runtime/Landscape/Private/LandscapeEditLayers.cpp
@@ -1910,7 +1910,7 @@ void ALandscapeProxy::InitializeProxyLayersWeightmapUsage()
 					{
 						FWeightmapLayerAllocationInfo& Allocation = LayerData->WeightmapData.LayerAllocations[LayerIdx];
 						UTexture2D* WeightmapTexture = LayerData->WeightmapData.Textures[Allocation.WeightmapTextureIndex];
-						UE_TRANSITIONAL_OBJECT_PTR(ULandscapeWeightmapUsage)* TempUsage = WeightmapUsageMap.Find(WeightmapTexture);
+						TObjectPtr<ULandscapeWeightmapUsage>* TempUsage = WeightmapUsageMap.Find(WeightmapTexture);
 
 						if (TempUsage == nullptr)
 						{
@@ -5583,7 +5583,7 @@ void ALandscapeProxy::DeleteLayer(const FGuid& InLayerGuid)
 			for (const FWeightmapLayerAllocationInfo& Allocation : LayerComponentData->WeightmapData.LayerAllocations)
 			{
 				UTexture2D* WeightmapTexture = LayerComponentData->WeightmapData.Textures[Allocation.WeightmapTextureIndex];
-				UE_TRANSITIONAL_OBJECT_PTR(ULandscapeWeightmapUsage)* Usage = WeightmapUsageMap.Find(WeightmapTexture);
+				TObjectPtr<ULandscapeWeightmapUsage>* Usage = WeightmapUsageMap.Find(WeightmapTexture);
 
 				if (Usage != nullptr && (*Usage) != nullptr)
 				{
diff --git a/Engine/Source/Runtime/MaterialShaderQualitySettings/Private/MaterialShaderQualitySettings.cpp b/Engine/Source/Runtime/MaterialShaderQualitySettings/Private/MaterialShaderQualitySettings.cpp
--- a/Engine/Source/Runtime/MaterialShaderQualitySettings/Private/MaterialShaderQualitySettings.cpp
+++ b/Engine/Source/Runtime/MaterialShaderQualitySettings/Private/MaterialShaderQualitySettings.cpp
@@ -40,7 +40,7 @@ const FName& UMaterialShaderQualitySettings::GetPreviewPlatform()
 
 void UMaterialShaderQualitySettings::SetPreviewPlatform(FName PlatformName)
 {
-	 UE_TRANSITIONAL_OBJECT_PTR(UShaderPlatformQualitySettings)* FoundPlatform = ForwardSettingMap.Find(PlatformName);
+	 TObjectPtr<UShaderPlatformQualitySettings>* FoundPlatform = ForwardSettingMap.Find(PlatformName);
 	 PreviewPlatformSettings = FoundPlatform == nullptr ? nullptr : *FoundPlatform;
 	 PreviewPlatformName = PlatformName;
 }
diff --git a/Engine/Source/Runtime/MediaAssets/Private/Assets/PlatformMediaSource.cpp b/Engine/Source/Runtime/MediaAssets/Private/Assets/PlatformMediaSource.cpp
--- a/Engine/Source/Runtime/MediaAssets/Private/Assets/PlatformMediaSource.cpp
+++ b/Engine/Source/Runtime/MediaAssets/Private/Assets/PlatformMediaSource.cpp
@@ -28,7 +28,7 @@ void UPlatformMediaSource::PreSave(FObjectPreSaveContext ObjectSaveContext)
 	const ITargetPlatform* TargetPlatform = ObjectSaveContext.GetTargetPlatform();
 	if (TargetPlatform)
 	{
-		UE_TRANSITIONAL_OBJECT_PTR(UMediaSource)* PlatformMediaSource = PlatformMediaSources.Find(TargetPlatform->IniPlatformName());
+		TObjectPtr<UMediaSource>* PlatformMediaSource = PlatformMediaSources.Find(TargetPlatform->IniPlatformName());
 		MediaSource = (PlatformMediaSource != nullptr) ? *PlatformMediaSource : nullptr;
 	}
 #endif
@@ -200,7 +200,7 @@ UMediaSource* UPlatformMediaSource::GetMediaSource() const
 {
 #if WITH_EDITORONLY_DATA
 	const FString RunningPlatformName(FPlatformProperties::IniPlatformName());
-	UE_TRANSITIONAL_OBJECT_PTR(UMediaSource) const* PlatformMediaSource = PlatformMediaSources.Find(RunningPlatformName);
+	TObjectPtr<UMediaSource> const* PlatformMediaSource = PlatformMediaSources.Find(RunningPlatformName);
 
 	if (PlatformMediaSource == nullptr)
 	{
diff --git a/Engine/Source/Runtime/UMG/Private/UserWidget.cpp b/Engine/Source/Runtime/UMG/Private/UserWidget.cpp
--- a/Engine/Source/Runtime/UMG/Private/UserWidget.cpp
+++ b/Engine/Source/Runtime/UMG/Private/UserWidget.cpp
@@ -373,7 +373,7 @@ UWorld* UUserWidget::GetWorld() const
 
 UUMGSequencePlayer* UUserWidget::GetSequencePlayer(const UWidgetAnimation* InAnimation) const
 {
-	UE_TRANSITIONAL_OBJECT_PTR(UUMGSequencePlayer) const* FoundPlayer = ActiveSequencePlayers.FindByPredicate(
+	TObjectPtr<UUMGSequencePlayer> const* FoundPlayer = ActiveSequencePlayers.FindByPredicate(
 		[&](const UUMGSequencePlayer* Player)
 	{
 		return Player->GetAnimation() == InAnimation;
@@ -669,7 +669,7 @@ bool UUserWidget::IsAnimationPlayingForward(const UWidgetAnimation* InAnimation)
 {
 	if (InAnimation)
 	{
-		UE_TRANSITIONAL_OBJECT_PTR(UUMGSequencePlayer)* FoundPlayer = ActiveSequencePlayers.FindByPredicate([&](const UUMGSequencePlayer* Player) { return Player->GetAnimation() == InAnimation; });
+		TObjectPtr<UUMGSequencePlayer>* FoundPlayer = ActiveSequencePlayers.FindByPredicate([&](const UUMGSequencePlayer* Player) { return Player->GetAnimation() == InAnimation; });
 
 		if (FoundPlayer)
 		{
