Merging //UE5/Release-Engine-Test to Main (//UE5/Main)

Refactor compile dependency checking to occur after compile results are recorded; now separately evaluate dependencies with all cached compile results.

diff --git a/Engine/Plugins/FX/Niagara/Source/Niagara/Classes/NiagaraSystem.h b/Engine/Plugins/FX/Niagara/Source/Niagara/Classes/NiagaraSystem.h
--- a/Engine/Plugins/FX/Niagara/Source/Niagara/Classes/NiagaraSystem.h
+++ b/Engine/Plugins/FX/Niagara/Source/Niagara/Classes/NiagaraSystem.h
@@ -157,7 +157,6 @@ struct FEmitterCompiledScriptPair
 	uint32 PendingJobID = INDEX_NONE; // this is the ID for any active shader compiler worker job
 	FNiagaraVMExecutableDataId CompileId;
 	TSharedPtr<FNiagaraVMExecutableData> CompileResults;
-	int32 ParentIndex = INDEX_NONE;
 };
 
 UENUM()
@@ -730,6 +729,8 @@ private:
 
 	bool CompilationResultsValid(FNiagaraSystemCompileRequest& CompileRequest) const;
 
+	void EvaluateCompileResultDependencies() const;
+
 	void InitEmitterCompiledData();
 
 	void InitSystemCompiledData();
diff --git a/Engine/Plugins/FX/Niagara/Source/Niagara/Private/NiagaraSystem.cpp b/Engine/Plugins/FX/Niagara/Source/Niagara/Private/NiagaraSystem.cpp
--- a/Engine/Plugins/FX/Niagara/Source/Niagara/Private/NiagaraSystem.cpp
+++ b/Engine/Plugins/FX/Niagara/Source/Niagara/Private/NiagaraSystem.cpp
@@ -2363,56 +2363,165 @@ bool UNiagaraSystem::CompilationResultsValid(FNiagaraSystemCompileRequest& Compi
 			return false;
 		}
 	}
+	return true;
+}
 
-	// Now iterate over all dependencies and verify that they are met. If not, emit an error.
-	for (auto& AsyncTask : CompileRequest.DDCTasks)
+void UNiagaraSystem::EvaluateCompileResultDependencies() const
+{
+	struct FScriptCompileResultValidationInfo
 	{
-		const FEmitterCompiledScriptPair& CompilePair = AsyncTask->GetScriptPair();
-		bool bValid = CompilePair.CompileResults.IsValid()
-			&& CompilePair.CompileResults->LastCompileStatus != ENiagaraScriptCompileStatus::NCS_Error;
+		UNiagaraEmitter* Emitter = nullptr;
+		FNiagaraVMExecutableData* CompileResults = nullptr;
+		int32 ParentIndex = INDEX_NONE;
+		bool bCompileResultStatusDirty = false;
 
-		if (CompilePair.CompileResults.IsValid() && CompilePair.CompileResults->ExternalDependencies.Num() != 0)
+		void ReEvaluateCompileStatus()
 		{
-			for (const FNiagaraCompileDependency& Dependency : CompilePair.CompileResults->ExternalDependencies)
+			// Early out for invalid existing compile statuses.
+			if (CompileResults->LastCompileStatus == ENiagaraScriptCompileStatus::NCS_BeingCreated)
+			{
+				ensureMsgf(false, TEXT("Encountered \"Being Created\" Compile Status when evaluating compile dependencies after compile! Compilation should be complete!"));
+				return;
+			}
+			else if (CompileResults->LastCompileStatus == ENiagaraScriptCompileStatus::NCS_Dirty)
+			{
+				ensureMsgf(false, TEXT("Encountered \"Dirty\" Compile Status when evaluating compile dependencies after compile! Compilation should be complete!"));
+				return;
+			}
+
+			// Tally warnings and errors.
+			int32 NumErrors = 0;
+			int32 NumWarnings = 0;
+
+			for (const FNiagaraCompileEvent& Event : CompileResults->LastCompileEvents)
 			{
-				FNiagaraVariable TestVar = Dependency.DependentVariable;
-				ensure(TestVar.GetName() != NAME_None);
-				if (CompilePair.Emitter)
+				if (Event.Severity == FNiagaraCompileEventSeverity::Error)
 				{
-					FName NewName = GetEmitterVariableAliasName(TestVar, CompilePair.Emitter);
-					TestVar.SetName(NewName);
+					NumErrors++;
 				}
+				else if (Event.Severity == FNiagaraCompileEventSeverity::Warning)
+				{
+					NumWarnings++;
+				}
+			}
+
+			// Set last compile status based on error and warning count.
+			if (NumErrors > 0)
+			{
+				CompileResults->LastCompileStatus = ENiagaraScriptCompileStatus::NCS_Error;
+			}
+			else if (NumWarnings > 0)
+			{
+				CompileResults->LastCompileStatus = ENiagaraScriptCompileStatus::NCS_UpToDateWithWarnings;
+			}
+			else
+			{
+				CompileResults->LastCompileStatus = ENiagaraScriptCompileStatus::NCS_UpToDate;
+			}
+		};
+	};
+
+	TArray<FScriptCompileResultValidationInfo> ScriptCompilesToValidate;
+	int32 ParentIdx = INDEX_NONE; 
+
+	// Add a compiled script pair to be evaluated for external dependencies to ScriptCompilesToValidate and clear all of its compile events that are from script dependencies.
+	auto AddCompiledScriptForValidation = [&ScriptCompilesToValidate](UNiagaraScript* InScript, int32 ParentIdx, UNiagaraEmitter* InEmitter = nullptr)->int32
+	{
+		FScriptCompileResultValidationInfo ValidationInfo = FScriptCompileResultValidationInfo();
+		ValidationInfo.CompileResults = &InScript->GetVMExecutableData();
+		TArray<FNiagaraCompileEvent>& CompileEvents = ValidationInfo.CompileResults->LastCompileEvents;
+		for (int32 Idx = CompileEvents.Num() - 1; Idx > -1; --Idx)
+		{
+			const FNiagaraCompileEvent& Event = CompileEvents[Idx];
+			if (Event.Source == FNiagaraCompileEventSource::ScriptDependency)
+			{
+				CompileEvents.RemoveAt(Idx);
+				ValidationInfo.bCompileResultStatusDirty = true;
+			}
+		}
+
+		ValidationInfo.Emitter = InEmitter;
+		ValidationInfo.ParentIndex = ParentIdx;
+		return ScriptCompilesToValidate.Add(ValidationInfo);
+	};
+
+	// Gather per-emitter scripts to evaluate dependencies.
+	for (int32 i = 0; i < EmitterHandles.Num(); i++)
+	{
+		FNiagaraEmitterHandle Handle = EmitterHandles[i];
+		if (Handle.GetInstance() && Handle.GetIsEnabled())
+		{
+			TArray<UNiagaraScript*> EmitterScripts;
+			bool bOnlyGetCompilableScripts = true;
+			bool bOnlyGetEnabledScripts = true;
+			Handle.GetInstance()->GetScripts(EmitterScripts, bOnlyGetCompilableScripts, bOnlyGetEnabledScripts);
+			check(EmitterScripts.Num() > 0);
+			for (UNiagaraScript* EmitterScript : EmitterScripts)
+			{
+				ParentIdx = AddCompiledScriptForValidation(EmitterScript, ParentIdx, Handle.GetInstance());
+			}
+		}
+	}
 
-				bool bDependencyMet = false;
-				int32 TestIdx = CompilePair.ParentIndex;
-				while (TestIdx != INDEX_NONE && bDependencyMet == false)
+	// Gather system scripts to evaluate dependencies.
+	ParentIdx = AddCompiledScriptForValidation(SystemSpawnScript, INDEX_NONE);
+	AddCompiledScriptForValidation(SystemUpdateScript, ParentIdx);
+
+	// Fixup the ParentIndex for emitter scripts that can have dependencies resolved by system scripts.
+	for (FScriptCompileResultValidationInfo& ValidationInfo : ScriptCompilesToValidate)
+	{
+		if (ValidationInfo.Emitter != nullptr && ValidationInfo.ParentIndex == INDEX_NONE)
+		{
+			ValidationInfo.ParentIndex = ParentIdx;
+		}
+	}
+
+	// Iterate over all ScriptCompilesToValidate and add compile events for dependencies that are not met.
+	for (FScriptCompileResultValidationInfo& ValidationInfo : ScriptCompilesToValidate)
+	{
+		for (const FNiagaraCompileDependency& Dependency : ValidationInfo.CompileResults->ExternalDependencies)
+		{
+			FNiagaraVariableBase TestVar = Dependency.DependentVariable;
+			ensure(TestVar.GetName() != NAME_None);
+			if (ValidationInfo.Emitter)
+			{
+				FName NewName = GetEmitterVariableAliasName(TestVar, ValidationInfo.Emitter);
+				TestVar.SetName(NewName);
+			}
+
+			bool bDependencyMet = false;
+			int32 TestIdx = ValidationInfo.ParentIndex;
+			while (TestIdx != INDEX_NONE && bDependencyMet == false)
+			{
+				if (ScriptCompilesToValidate.IsValidIndex(TestIdx))
 				{
-					if (CompileRequest.DDCTasks.IsValidIndex(TestIdx))
+					const FScriptCompileResultValidationInfo& TestInfo = ScriptCompilesToValidate[TestIdx];
+					if (TestInfo.CompileResults->AttributesWritten.Num() > 0)
 					{
-						const FEmitterCompiledScriptPair& TestPair = CompileRequest.DDCTasks[TestIdx]->GetScriptPair();
-						if (TestPair.CompileResults.IsValid() && TestPair.CompileResults->AttributesWritten.Num() > 0)
+						if (TestInfo.CompileResults->AttributesWritten.Contains(TestVar))
 						{
-							if (TestPair.CompileResults->AttributesWritten.Contains(TestVar))
-							{
-								bDependencyMet = true;
-								break;
-							}
+							bDependencyMet = true;
+							break;
 						}
-						TestIdx = TestPair.ParentIndex;
 					}
-				}
-				if (!bDependencyMet)
-				{
-					FNiagaraCompileEvent LinkerErrorEvent(
-						FNiagaraCompileEventSeverity::Error, Dependency.LinkerErrorMessage, FString(), false, Dependency.NodeGuid, Dependency.PinGuid, Dependency.StackGuids);
-					CompilePair.CompileResults->LastCompileEvents.Add(LinkerErrorEvent);
-					CompilePair.CompileResults->LastCompileStatus = ENiagaraScriptCompileStatus::NCS_Error;
+					TestIdx = TestInfo.ParentIndex;
 				}
 			}
+			if (!bDependencyMet)
+			{
+				FNiagaraCompileEvent LinkerErrorEvent(
+					FNiagaraCompileEventSeverity::Error, Dependency.LinkerErrorMessage, FString(), false, Dependency.NodeGuid, Dependency.PinGuid, Dependency.StackGuids, FNiagaraCompileEventSource::ScriptDependency);
+				ValidationInfo.CompileResults->LastCompileEvents.Add(LinkerErrorEvent);
+				ValidationInfo.bCompileResultStatusDirty = true;
+			}
+		}
+
+		// If the compile events of the compile results have changed, re-evaluate the compile status as it could have changed.
+		if (ValidationInfo.bCompileResultStatusDirty)
+		{
+			ValidationInfo.ReEvaluateCompileStatus();
 		}
 	}
-	
-	return true;
 }
 
 void UNiagaraSystem::PreProcessWaitingDDCTasks(bool bProcessForWait)
@@ -2539,6 +2648,9 @@ bool UNiagaraSystem::QueryCompileComplete(bool bWait, bool bDoPost, bool bDoNotA
 			}
 		}
 
+		// Once compile results have been set, check dependencies found during compile and evaluate whether dependency compile events should be added.
+		EvaluateCompileResultDependencies();
+
 		if (bDoPost)
 		{
 			for (FNiagaraEmitterHandle Handle : EmitterHandles)
@@ -2812,13 +2924,11 @@ bool UNiagaraSystem::RequestCompile(bool bForce, FNiagaraSystemUpdateContext* Op
 					TArray<UNiagaraScript*> EmitterScripts;
 					Handle.GetInstance()->GetScripts(EmitterScripts, false, true);
 					check(EmitterScripts.Num() > 0);
-					int32 Parent = INDEX_NONE;
 					for (UNiagaraScript* EmitterScript : EmitterScripts)
 					{
 						FEmitterCompiledScriptPair Pair;
 						Pair.Emitter = Handle.GetInstance();
 						Pair.CompiledScript = EmitterScript;
-						Pair.ParentIndex = Parent;
 
 						FNiagaraVMExecutableDataId NewID;
 						// we need to compute the vmID here to check later in the ddc task before doing the precompile if anything has changed in the meantime
@@ -2836,7 +2946,7 @@ bool UNiagaraSystem::RequestCompile(bool bForce, FNiagaraSystemUpdateContext* Op
 							AsyncTask->CurrentState = ENiagaraCompilationState::Finished;
 						}
 						AsyncTask->UniqueEmitterName = Handle.GetInstance()->GetUniqueEmitterName();
-						Parent = ActiveCompilation.DDCTasks.Add(AsyncTask);
+						ActiveCompilation.DDCTasks.Add(AsyncTask);
 					}
 				}
 			}
@@ -2844,7 +2954,6 @@ bool UNiagaraSystem::RequestCompile(bool bForce, FNiagaraSystemUpdateContext* Op
 			bAnyCompiled = bAnyUnsynchronized || bForce;
 
 			// Now add the system scripts for compilation...
-			int32 Parent = INDEX_NONE;
 			{
 				FEmitterCompiledScriptPair Pair;
 				Pair.CompiledScript = SystemSpawnScript;
@@ -2864,13 +2973,12 @@ bool UNiagaraSystem::RequestCompile(bool bForce, FNiagaraSystemUpdateContext* Op
 				{
 					AsyncTask->CurrentState = ENiagaraCompilationState::Finished;
 				}
-				Parent = ActiveCompilation.DDCTasks.Add(AsyncTask);
+				ActiveCompilation.DDCTasks.Add(AsyncTask);
 			}
 
 			{
 				FEmitterCompiledScriptPair Pair;
 				Pair.CompiledScript = SystemUpdateScript;
-				Pair.ParentIndex = Parent;
 				
 				FNiagaraVMExecutableDataId NewID;
 				// we need to compute the vmID here to check later in the ddc task before doing the precompile if anything has changed in the meantime
@@ -2887,17 +2995,7 @@ bool UNiagaraSystem::RequestCompile(bool bForce, FNiagaraSystemUpdateContext* Op
 				{
 					AsyncTask->CurrentState = ENiagaraCompilationState::Finished;
 				}
-				Parent = ActiveCompilation.DDCTasks.Add(AsyncTask);
-			}
-
-			// Need to set the EmitterParent on the emitter spawn scripts
-			for (auto & AsyncTask: ActiveCompilation.DDCTasks)
-			{
-				FEmitterCompiledScriptPair& Pair = AsyncTask->ScriptPair;
-				if (Pair.Emitter != nullptr && Pair.ParentIndex == INDEX_NONE)
-				{
-					Pair.ParentIndex = Parent;
-				}
+				ActiveCompilation.DDCTasks.Add(AsyncTask);
 			}
 		}
 
diff --git a/Engine/Plugins/FX/Niagara/Source/NiagaraCore/Private/NiagaraCustomVersion.cpp b/Engine/Plugins/FX/Niagara/Source/NiagaraCore/Private/NiagaraCustomVersion.cpp
--- a/Engine/Plugins/FX/Niagara/Source/NiagaraCore/Private/NiagaraCustomVersion.cpp
+++ b/Engine/Plugins/FX/Niagara/Source/NiagaraCore/Private/NiagaraCustomVersion.cpp
@@ -8,5 +8,5 @@ const FGuid FNiagaraCustomVersion::GUID(0xFCF57AFA, 0x50764283, 0xB9A9E658, 0xFF
 // Register the custom version with core
 FCustomVersionRegistration GRegisterNiagaraCustomVersion(FNiagaraCustomVersion::GUID, FNiagaraCustomVersion::LatestVersion, TEXT("NiagaraVer"));
 
-const FGuid FNiagaraCustomVersion::LatestScriptCompileVersion(0x64FA8CB0, 0xF7748343, 0x83D275AE, 0xAE764DF9);
+const FGuid FNiagaraCustomVersion::LatestScriptCompileVersion(0x37BF6E63, 0x7F1144F2, 0xAD4EC385, 0x0A16472D);
 
diff --git a/Engine/Plugins/FX/Niagara/Source/NiagaraShader/Public/NiagaraShared.h b/Engine/Plugins/FX/Niagara/Source/NiagaraShader/Public/NiagaraShared.h
--- a/Engine/Plugins/FX/Niagara/Source/NiagaraShader/Public/NiagaraShared.h
+++ b/Engine/Plugins/FX/Niagara/Source/NiagaraShader/Public/NiagaraShared.h
@@ -43,6 +43,13 @@ enum class FNiagaraCompileEventSeverity : uint8
 	Error = 2
 };
 
+UENUM()
+enum class FNiagaraCompileEventSource : uint8
+{
+	Unset = 0, // No specific source of compile event.
+	ScriptDependency = 1 // Compile event is from a dependency to other scripts.
+};
+
 /** Records necessary information to give UI cues for errors/logs/warnings during compile.*/
 USTRUCT()
 struct FNiagaraCompileEvent
@@ -60,8 +67,25 @@ public:
 		StackGuids.Empty();
 	}
  
-	FNiagaraCompileEvent(FNiagaraCompileEventSeverity InSeverity, const FString& InMessage, FString InShortDescription = FString(), bool bInDismissable = true,FGuid InNodeGuid = FGuid(), FGuid InPinGuid = FGuid(), const TArray<FGuid>& InCallstackGuids = TArray<FGuid>())
-		: Severity(InSeverity), Message(InMessage), ShortDescription(InShortDescription), bDismissable(bInDismissable), NodeGuid(InNodeGuid), PinGuid(InPinGuid), StackGuids(InCallstackGuids) {}
+	FNiagaraCompileEvent(
+		FNiagaraCompileEventSeverity InSeverity,
+		const FString& InMessage,
+		FString InShortDescription = FString(),
+		bool bInDismissable = true,
+		FGuid InNodeGuid = FGuid(),
+		FGuid InPinGuid = FGuid(),
+		const TArray<FGuid>& InCallstackGuids = TArray<FGuid>(),
+		FNiagaraCompileEventSource InSource = FNiagaraCompileEventSource::Unset
+	)
+		: Severity(InSeverity)
+		, Message(InMessage)
+		, ShortDescription(InShortDescription)
+		, bDismissable(bInDismissable)
+		, NodeGuid(InNodeGuid)
+		, PinGuid(InPinGuid)
+		, StackGuids(InCallstackGuids)
+		, Source(InSource) 
+	{}
  
 	/** Whether or not this is an error, warning, or info*/
 	UPROPERTY()
@@ -84,6 +108,10 @@ public:
 	/** The compile stack frame of node id's*/
 	UPROPERTY()
 	TArray<FGuid> StackGuids;
+	/** The source of the compile event for partial invalidation purposes. */
+	UPROPERTY()
+	FNiagaraCompileEventSource Source;
+
 };
 
 //
