Fix CIS error FNiagaraCompileEvent ctor

diff --git a/Engine/Plugins/FX/Niagara/Source/NiagaraShader/Public/NiagaraShared.h b/Engine/Plugins/FX/Niagara/Source/NiagaraShader/Public/NiagaraShared.h
--- a/Engine/Plugins/FX/Niagara/Source/NiagaraShader/Public/NiagaraShared.h
+++ b/Engine/Plugins/FX/Niagara/Source/NiagaraShader/Public/NiagaraShared.h
@@ -65,6 +65,7 @@ public:
 		NodeGuid = FGuid();
 		PinGuid = FGuid();
 		StackGuids.Empty();
+		Source = FNiagaraCompileEventSource::Unset;
 	}
  
 	FNiagaraCompileEvent(
