- Support renderer enabled dynamic binding on component renderer

#rnx

diff --git a/Engine/Plugins/FX/Niagara/Source/Niagara/Private/NiagaraComponentRendererProperties.cpp b/Engine/Plugins/FX/Niagara/Source/Niagara/Private/NiagaraComponentRendererProperties.cpp
--- a/Engine/Plugins/FX/Niagara/Source/Niagara/Private/NiagaraComponentRendererProperties.cpp
+++ b/Engine/Plugins/FX/Niagara/Source/Niagara/Private/NiagaraComponentRendererProperties.cpp
@@ -245,6 +245,23 @@ void UNiagaraComponentRendererProperties::CacheFromCompiledData(const FNiagaraDa
 	UpdateSourceModeDerivates(ENiagaraRendererSourceDataMode::Particles);
 }
 
+#if WITH_EDITORONLY_DATA
+bool UNiagaraComponentRendererProperties::IsSupportedVariableForBinding(const FNiagaraVariableBase& InSourceForBinding, const FName& InTargetBindingName) const
+{
+	if ( InTargetBindingName == GET_MEMBER_NAME_CHECKED(UNiagaraComponentRendererProperties, RendererEnabledBinding) )
+	{
+		if (
+			InSourceForBinding.IsInNameSpace(FNiagaraConstants::UserNamespace) ||
+			InSourceForBinding.IsInNameSpace(FNiagaraConstants::SystemNamespace) ||
+			InSourceForBinding.IsInNameSpace(FNiagaraConstants::EmitterNamespace))
+		{
+			return true;
+		}
+	}
+	return false;
+}
+#endif
+
 void UNiagaraComponentRendererProperties::UpdateSetterFunctions()
 {
 	SetterFunctionMapping.Empty();
diff --git a/Engine/Plugins/FX/Niagara/Source/Niagara/Private/NiagaraRendererComponents.cpp b/Engine/Plugins/FX/Niagara/Source/Niagara/Private/NiagaraRendererComponents.cpp
--- a/Engine/Plugins/FX/Niagara/Source/Niagara/Private/NiagaraRendererComponents.cpp
+++ b/Engine/Plugins/FX/Niagara/Source/Niagara/Private/NiagaraRendererComponents.cpp
@@ -299,9 +299,11 @@ void FNiagaraRendererComponents::PostSystemTick_GameThread(const UNiagaraRendere
 	FNiagaraDataSetReaderInt32<int32> VisTagAccessor = FNiagaraDataSetAccessor<int32>::CreateReader(Data, Properties->RendererVisibilityTagBinding.GetDataSetBindableVariable().GetName());
 	FNiagaraDataSetReaderInt32<int32> UniqueIDAccessor = FNiagaraDataSetAccessor<int32>::CreateReader(Data, FName("UniqueID"));
 
-	auto IsParticleEnabled = [&EnabledAccessor, &VisTagAccessor, Properties](int32 ParticleIndex)
+	const bool bIsRendererEnabled = IsRendererEnabled(Properties, Emitter);
+
+	auto IsParticleEnabled = [&bIsRendererEnabled, &EnabledAccessor, &VisTagAccessor, Properties](int32 ParticleIndex)
 	{
-		if (EnabledAccessor.GetSafe(ParticleIndex, true))
+		if (bIsRendererEnabled && EnabledAccessor.GetSafe(ParticleIndex, true))
 		{
 			if (VisTagAccessor.IsValid())
 			{
diff --git a/Engine/Plugins/FX/Niagara/Source/Niagara/Public/NiagaraComponentRendererProperties.h b/Engine/Plugins/FX/Niagara/Source/Niagara/Public/NiagaraComponentRendererProperties.h
--- a/Engine/Plugins/FX/Niagara/Source/Niagara/Public/NiagaraComponentRendererProperties.h
+++ b/Engine/Plugins/FX/Niagara/Source/Niagara/Public/NiagaraComponentRendererProperties.h
@@ -138,6 +138,10 @@ public:
 
 	virtual void CacheFromCompiledData(const FNiagaraDataSetCompiledData* CompiledData) override;
 
+#if WITH_EDITORONLY_DATA
+	virtual bool IsSupportedVariableForBinding(const FNiagaraVariableBase& InSourceForBinding, const FName& InTargetBindingName) const override;
+#endif
+
 	virtual bool NeedsSystemPostTick() const override { return true; }
 	virtual bool NeedsSystemCompletion() const override { return true; }
 
