- Check type size with NDI input / output params
- Fix up areas where LWC has introduced crashes due to incorrect type usage

#rnx

diff --git a/Engine/Plugins/Experimental/ChaosNiagara/Source/ChaosNiagara/Private/NiagaraDataInterfacePhysicsField.cpp b/Engine/Plugins/Experimental/ChaosNiagara/Source/ChaosNiagara/Private/NiagaraDataInterfacePhysicsField.cpp
--- a/Engine/Plugins/Experimental/ChaosNiagara/Source/ChaosNiagara/Private/NiagaraDataInterfacePhysicsField.cpp
+++ b/Engine/Plugins/Experimental/ChaosNiagara/Source/ChaosNiagara/Private/NiagaraDataInterfacePhysicsField.cpp
@@ -527,11 +527,11 @@ void UNiagaraDataInterfacePhysicsField::SamplePhysicsVectorField(FVectorVMExtern
 	VectorVM::FUserPtrHandler<FNDIPhysicsFieldData> InstData(Context);
 
 	// Inputs 
-	FNDIInputParam<FVector> SamplePositionParam(Context);
+	FNDIInputParam<FVector3f> SamplePositionParam(Context);
 	FNDIInputParam<EFieldVectorType> VectorTargetParam(Context);
 
 	// Outputs...
-	FNDIOutputParam<FVector> OutVectorFieldParam(Context);
+	FNDIOutputParam<FVector3f> OutVectorFieldParam(Context);
 
 	if (InstData)
 	{
diff --git a/Engine/Plugins/FX/Niagara/Source/Niagara/Classes/NiagaraDataInterface.h b/Engine/Plugins/FX/Niagara/Source/Niagara/Classes/NiagaraDataInterface.h
--- a/Engine/Plugins/FX/Niagara/Source/Niagara/Classes/NiagaraDataInterface.h
+++ b/Engine/Plugins/FX/Niagara/Source/Niagara/Classes/NiagaraDataInterface.h
@@ -20,22 +20,22 @@ struct FNiagaraDataInterfaceProxy;
 
 struct FNDITransformHandlerNoop
 {
-	FORCEINLINE void TransformPosition(FVector3f& V, const FMatrix44f& M) { }
-	FORCEINLINE void TransformPosition(FVector3d& V, const FMatrix44d& M) { }
-	FORCEINLINE void TransformVector(FVector3f& V, const FMatrix44f& M) { }
-	FORCEINLINE void TransformVector(FVector3d& V, const FMatrix44d& M) { }
-	FORCEINLINE void TransformRotation(FQuat4f& Q1, const FQuat4f& Q2) { }
-	FORCEINLINE void TransformRotation(FQuat4d& Q1, const FQuat4d& Q2) { }
+	FORCEINLINE void TransformPosition(FVector3f& V, const FMatrix44f& M) const { }
+	FORCEINLINE void TransformPosition(FVector3d& V, const FMatrix44d& M) const { }
+	FORCEINLINE void TransformVector(FVector3f& V, const FMatrix44f& M) const { }
+	FORCEINLINE void TransformVector(FVector3d& V, const FMatrix44d& M) const { }
+	FORCEINLINE void TransformRotation(FQuat4f& Q1, const FQuat4f& Q2) const { }
+	FORCEINLINE void TransformRotation(FQuat4d& Q1, const FQuat4d& Q2) const { }
 };
 
 struct FNDITransformHandler
 {
-	FORCEINLINE void TransformPosition(FVector3f& P, const FMatrix44f& M) { P = M.TransformPosition(P); }
-	FORCEINLINE void TransformPosition(FVector3d& P, const FMatrix44d& M) { P = M.TransformPosition(P); }
-	FORCEINLINE void TransformVector(FVector3f& V, const FMatrix44f& M) { V = M.TransformVector(V).GetUnsafeNormal3(); }
-	FORCEINLINE void TransformVector(FVector3d& V, const FMatrix44d& M) { V = M.TransformVector(V).GetUnsafeNormal3(); }
-	FORCEINLINE void TransformRotation(FQuat4f& Q1, const FQuat4f& Q2) { Q1 = Q2 * Q1; }
-	FORCEINLINE void TransformRotation(FQuat4d& Q1, const FQuat4d& Q2) { Q1 = Q2 * Q1; }
+	FORCEINLINE void TransformPosition(FVector3f& P, const FMatrix44f& M) const { P = M.TransformPosition(P); }
+	FORCEINLINE void TransformPosition(FVector3d& P, const FMatrix44d& M) const { P = M.TransformPosition(P); }
+	FORCEINLINE void TransformVector(FVector3f& V, const FMatrix44f& M) const { V = M.TransformVector(V).GetUnsafeNormal3(); }
+	FORCEINLINE void TransformVector(FVector3d& V, const FMatrix44d& M) const { V = M.TransformVector(V).GetUnsafeNormal3(); }
+	FORCEINLINE void TransformRotation(FQuat4f& Q1, const FQuat4f& Q2) const { Q1 = Q2 * Q1; }
+	FORCEINLINE void TransformRotation(FQuat4d& Q1, const FQuat4d& Q2) const { Q1 = Q2 * Q1; }
 };
 
 //////////////////////////////////////////////////////////////////////////
@@ -572,14 +572,14 @@ struct FNDIRandomHelper
 		RandParam.GetAndAdvance(RandInfo);
 	}
 
-	FORCEINLINE bool IsDeterministic()
+	FORCEINLINE bool IsDeterministic() const
 	{
 		return RandInfo.Seed3 != INDEX_NONE;
 	}
 
 	//////////////////////////////////////////////////////////////////////////
 	
-	FORCEINLINE_DEBUGGABLE FVector4 Rand4(int32 InstanceIndex)
+	FORCEINLINE_DEBUGGABLE FVector4 Rand4(int32 InstanceIndex) const
 	{
 		if (IsDeterministic())
 		{
@@ -610,7 +610,7 @@ struct FNDIRandomHelper
 		}
 	}
 
-	FORCEINLINE_DEBUGGABLE FVector Rand3(int32 InstanceIndex)
+	FORCEINLINE_DEBUGGABLE FVector Rand3(int32 InstanceIndex) const
 	{
 		if (IsDeterministic())
 		{
@@ -633,7 +633,7 @@ struct FNDIRandomHelper
 		}
 	}
 
-	FORCEINLINE_DEBUGGABLE FVector2D Rand2(int32 InstanceIndex)
+	FORCEINLINE_DEBUGGABLE FVector2D Rand2(int32 InstanceIndex) const
 	{
 		if (IsDeterministic())
 		{
@@ -646,7 +646,7 @@ struct FNDIRandomHelper
 		}
 	}
 
-	FORCEINLINE_DEBUGGABLE float Rand(int32 InstanceIndex)
+	FORCEINLINE_DEBUGGABLE float Rand(int32 InstanceIndex) const
 	{
 		if (IsDeterministic())
 		{
@@ -658,31 +658,31 @@ struct FNDIRandomHelper
 		}
 	}
 
-	FORCEINLINE_DEBUGGABLE FVector4 RandRange(int32 InstanceIndex, FVector4 Min, FVector4 Max)
+	FORCEINLINE_DEBUGGABLE FVector4 RandRange(int32 InstanceIndex, FVector4 Min, FVector4 Max) const
 	{
 		FVector4 Range = Max - Min;
 		return Min + (Rand(InstanceIndex) * Range);
 	}
 
-	FORCEINLINE_DEBUGGABLE FVector RandRange(int32 InstanceIndex, FVector Min, FVector Max)
+	FORCEINLINE_DEBUGGABLE FVector RandRange(int32 InstanceIndex, FVector Min, FVector Max) const
 	{
 		FVector Range = Max - Min;
 		return Min + (Rand(InstanceIndex) * Range);
 	}
 
-	FORCEINLINE_DEBUGGABLE FVector2D RandRange(int32 InstanceIndex, FVector2D Min, FVector2D Max)
+	FORCEINLINE_DEBUGGABLE FVector2D RandRange(int32 InstanceIndex, FVector2D Min, FVector2D Max) const
 	{
 		FVector2D Range = Max - Min;
 		return Min + (Rand(InstanceIndex) * Range);
 	}
 
-	FORCEINLINE_DEBUGGABLE float RandRange(int32 InstanceIndex, float Min, float Max)
+	FORCEINLINE_DEBUGGABLE float RandRange(int32 InstanceIndex, float Min, float Max) const
 	{
 		float Range = Max - Min;
 		return Min + (Rand(InstanceIndex) * Range);
 	}
 
-	FORCEINLINE_DEBUGGABLE int32 RandRange(int32 InstanceIndex, int32 Min, int32 Max)
+	FORCEINLINE_DEBUGGABLE int32 RandRange(int32 InstanceIndex, int32 Min, int32 Max) const
 	{
 		// NOTE: Scaling a uniform float range provides better distribution of 
 		//       numbers than using %.
@@ -691,7 +691,7 @@ struct FNDIRandomHelper
 		return Min + (int(Rand(InstanceIndex) * (Range + 1)));
 	}
 
-	FORCEINLINE_DEBUGGABLE FVector RandomBarycentricCoord(int32 InstanceIndex)
+	FORCEINLINE_DEBUGGABLE FVector RandomBarycentricCoord(int32 InstanceIndex) const
 	{
 		//TODO: This is gonna be slooooow. Move to an LUT possibly or find faster method.
 		//Can probably handle lower quality randoms / uniformity for a decent speed win.
@@ -711,6 +711,7 @@ struct FNDIRandomHelper
 template<typename T>
 struct FNDIInputParam
 {
+	static_assert(sizeof(T) == sizeof(float), "Generic template assumes 4 bytes per element");
 	VectorVM::FExternalFuncInputHandler<T> Data;
 	FORCEINLINE FNDIInputParam(FVectorVMExternalFunctionContext& Context) : Data(Context) {}
 	FORCEINLINE T GetAndAdvance() { return Data.GetAndAdvance(); }
@@ -799,6 +800,7 @@ struct FNDIInputParam<FNiagaraID>
 template<typename T>
 struct FNDIOutputParam
 {
+	static_assert(sizeof(T) == sizeof(float), "Generic template assumes 4 bytes per element");
 	VectorVM::FExternalFuncRegisterHandler<T> Data;
 	FORCEINLINE FNDIOutputParam(FVectorVMExternalFunctionContext& Context) : Data(Context) {}
 	FORCEINLINE bool IsValid() const { return Data.IsValid();  }
diff --git a/Engine/Plugins/FX/Niagara/Source/Niagara/Private/DataInterface/NiagaraDataInterfaceActorComponent.cpp b/Engine/Plugins/FX/Niagara/Source/Niagara/Private/DataInterface/NiagaraDataInterfaceActorComponent.cpp
--- a/Engine/Plugins/FX/Niagara/Source/Niagara/Private/DataInterface/NiagaraDataInterfaceActorComponent.cpp
+++ b/Engine/Plugins/FX/Niagara/Source/Niagara/Private/DataInterface/NiagaraDataInterfaceActorComponent.cpp
@@ -364,14 +364,14 @@ void UNiagaraDataInterfaceActorComponent::VMGetTransform(FVectorVMExternalFuncti
 	VectorVM::FUserPtrHandler<FInstanceData_GameThread> InstanceData(Context);
 	FNDIOutputParam<bool>		OutValid(Context);
 	FNDIOutputParam<FVector3f>	OutPosition(Context);
-	FNDIOutputParam<FQuat>		OutRotation(Context);
+	FNDIOutputParam<FQuat4f>	OutRotation(Context);
 	FNDIOutputParam<FVector3f>	OutScale(Context);
 
 	for (int32 i=0; i < Context.GetNumInstances(); ++i)
 	{
 		OutValid.SetAndAdvance(InstanceData->bCachedValid);
 		OutPosition.SetAndAdvance(InstanceData->CachedTransform.GetLocation());
-		OutRotation.SetAndAdvance(InstanceData->CachedTransform.GetRotation());
+		OutRotation.SetAndAdvance(FQuat4f(InstanceData->CachedTransform.GetRotation()));
 		OutScale.SetAndAdvance(InstanceData->CachedTransform.GetScale3D());
 	}
 }
