- Use actor label when available as it gives more meaningful information

#rnx

diff --git a/Engine/Plugins/FX/Niagara/Source/Niagara/Private/NiagaraDebugHud.cpp b/Engine/Plugins/FX/Niagara/Source/Niagara/Private/NiagaraDebugHud.cpp
--- a/Engine/Plugins/FX/Niagara/Source/Niagara/Private/NiagaraDebugHud.cpp
+++ b/Engine/Plugins/FX/Niagara/Source/Niagara/Private/NiagaraDebugHud.cpp
@@ -515,6 +515,12 @@ namespace NiagaraDebugLocal
 			}
 		}
 	}
+
+	FString GetHudActorName(AActor* Actor)
+	{
+		FString Name = Actor->GetActorNameOrLabel();
+		return Name.Len() == 0 ? GetNameSafe(Actor) : Name;
+	}
 }
 
 FNiagaraDebugHud::FNiagaraDebugHud(UWorld* World)
@@ -637,7 +643,7 @@ void FNiagaraDebugHud::GatherSystemInfo()
 			if ( Settings.bActorFilterEnabled )
 			{
 				AActor* Actor = NiagaraComponent->GetOwner();
-				bIsMatch &= (Actor != nullptr) && Actor->GetName().MatchesWildcard(Settings.ActorFilter);
+				bIsMatch &= (Actor != nullptr) && GetHudActorName(Actor).MatchesWildcard(Settings.ActorFilter);
 			}
 
 			// Filter by component
@@ -1398,7 +1404,7 @@ void FNiagaraDebugHud::DrawComponents(FNiagaraWorldManager* WorldManager, UCanva
 				StringBuilder.Append(TEXT("Component - "));
 				if ( OwnerActor )
 				{
-					StringBuilder.Append(*GetNameSafe(OwnerActor));
+					StringBuilder.Append(*GetHudActorName(OwnerActor));
 					StringBuilder.Append(TEXT("/"));
 				}
 				StringBuilder.Append(*GetNameSafe(NiagaraComponent));
