Niagara node localization fix for pin names in staticswitch & select node for enums
Also removed a function that wasn't being used

diff --git a/Engine/Plugins/FX/Niagara/Source/NiagaraCore/Private/NiagaraCustomVersion.cpp b/Engine/Plugins/FX/Niagara/Source/NiagaraCore/Private/NiagaraCustomVersion.cpp
--- a/Engine/Plugins/FX/Niagara/Source/NiagaraCore/Private/NiagaraCustomVersion.cpp
+++ b/Engine/Plugins/FX/Niagara/Source/NiagaraCore/Private/NiagaraCustomVersion.cpp
@@ -8,5 +8,5 @@ const FGuid FNiagaraCustomVersion::GUID(0xFCF57AFA, 0x50764283, 0xB9A9E658, 0xFF
 // Register the custom version with core
 FCustomVersionRegistration GRegisterNiagaraCustomVersion(FNiagaraCustomVersion::GUID, FNiagaraCustomVersion::LatestVersion, TEXT("NiagaraVer"));
 
-const FGuid FNiagaraCustomVersion::LatestScriptCompileVersion(0x37BF6E63, 0x7F1144F2, 0xAD4EC385, 0x0A16472D);
+const FGuid FNiagaraCustomVersion::LatestScriptCompileVersion(0x375F6E63, 0x7F2144F2, 0xAD4EC387, 0x0AB6472D);
 
diff --git a/Engine/Plugins/FX/Niagara/Source/NiagaraEditor/Private/NiagaraNodeSelect.cpp b/Engine/Plugins/FX/Niagara/Source/NiagaraEditor/Private/NiagaraNodeSelect.cpp
--- a/Engine/Plugins/FX/Niagara/Source/NiagaraEditor/Private/NiagaraNodeSelect.cpp
+++ b/Engine/Plugins/FX/Niagara/Source/NiagaraEditor/Private/NiagaraNodeSelect.cpp
@@ -511,7 +511,11 @@ FString UNiagaraNodeSelect::GetInputCaseName(int32 Case) const
 	else if(SelectorPinType.IsEnum() && SelectorPinType.GetEnum())
 	{
 		UEnum* Enum = SelectorPinType.GetEnum();
-		return Enum->GetDisplayNameTextByValue(Case).ToString();
+		// the display name is subject to localization and some automatic prettification. To avoid the localization aspect, we retrieve the source string of the text
+		// which is essentially the still prettified non-localized base text. We have to keep it this way for backwards compatibility until we can do a full upgrade pass.
+		// Same in StaticSwitch @todo
+		FText EnumDisplayText = Enum->GetDisplayNameTextByValue(Case);
+		return *FTextInspector::GetSourceString(EnumDisplayText);
 	}
 
 	return TEXT("");
diff --git a/Engine/Plugins/FX/Niagara/Source/NiagaraEditor/Private/NiagaraNodeStaticSwitch.cpp b/Engine/Plugins/FX/Niagara/Source/NiagaraEditor/Private/NiagaraNodeStaticSwitch.cpp
--- a/Engine/Plugins/FX/Niagara/Source/NiagaraEditor/Private/NiagaraNodeStaticSwitch.cpp
+++ b/Engine/Plugins/FX/Niagara/Source/NiagaraEditor/Private/NiagaraNodeStaticSwitch.cpp
@@ -70,7 +70,8 @@ FString UNiagaraNodeStaticSwitch::GetInputCaseName(int32 Case) const
 	}
 	else if (SwitchTypeData.SwitchType == ENiagaraStaticSwitchType::Enum && SwitchTypeData.Enum)
 	{
-		return SwitchTypeData.Enum->GetDisplayNameTextByValue(Case).ToString();
+		FText EnumDisplayText = SwitchTypeData.Enum->GetDisplayNameTextByValue(Case);
+		return *FTextInspector::GetSourceString(EnumDisplayText);
 	}
 
 	return TEXT("");
@@ -142,7 +143,10 @@ FName UNiagaraNodeStaticSwitch::GetOptionPinName(const FNiagaraVariable& Variabl
 	}
 	else if (SwitchTypeData.SwitchType == ENiagaraStaticSwitchType::Enum && SwitchTypeData.Enum)
 	{
-		Suffix = SwitchTypeData.Enum->GetDisplayNameTextByValue(Value).ToString();
+		// the display name is subject to localization and some automatic prettification. To avoid the localization aspect, we retrieve the source string of the text
+		// which is essentially the still prettified non-localized base text. We have to keep it this way for backwards compatibility, until we do a full upgrade pass
+		FText EnumDisplayText = SwitchTypeData.Enum->GetDisplayNameTextByValue(Value);
+		Suffix = *FTextInspector::GetSourceString(EnumDisplayText);
 	}
 
 	return FName(Variable.GetName().ToString() + FString::Printf(TEXT(" if %s"), *Suffix));
@@ -267,26 +271,6 @@ void UNiagaraNodeStaticSwitch::AllocateDefaultPins()
 	GetNiagaraGraph()->GetParameterReferenceMap();
 }
 
-FString UNiagaraNodeStaticSwitch::GetOptionPinSuffix(int32 Index) const
-{
-	FString PathSuffix = TEXT(" if ");
-	if (SwitchTypeData.SwitchType == ENiagaraStaticSwitchType::Bool)
-	{
-		PathSuffix += Index == 0 ? TEXT("true") : TEXT("false");
-	}
-	else if (SwitchTypeData.SwitchType == ENiagaraStaticSwitchType::Integer)
-	{
-		PathSuffix += FString::FromInt(Index);
-	}
-	else if (SwitchTypeData.SwitchType == ENiagaraStaticSwitchType::Enum && SwitchTypeData.Enum)
-	{
-		FText EnumName = SwitchTypeData.Enum->GetDisplayNameTextByIndex(Index);
-		PathSuffix += EnumName.ToString();
-	}
-
-	return PathSuffix;
-}
-
 bool UNiagaraNodeStaticSwitch::AllowNiagaraTypeForAddPin(const FNiagaraTypeDefinition& InType) const
 {
 	// explicitly allow parameter maps and numeric types
diff --git a/Engine/Plugins/FX/Niagara/Source/NiagaraEditor/Private/NiagaraNodeStaticSwitch.h b/Engine/Plugins/FX/Niagara/Source/NiagaraEditor/Private/NiagaraNodeStaticSwitch.h
--- a/Engine/Plugins/FX/Niagara/Source/NiagaraEditor/Private/NiagaraNodeStaticSwitch.h
+++ b/Engine/Plugins/FX/Niagara/Source/NiagaraEditor/Private/NiagaraNodeStaticSwitch.h
@@ -111,8 +111,7 @@ private:
 	bool GetVarIndex(class FHlslNiagaraTranslator* Translator, int32 InputPinCount, int32 Value, int32& VarIndexOut) const;
 
 	void RemoveUnusedGraphParameter(const FNiagaraVariable& OldParameter);
-
-	FString GetOptionPinSuffix(int32 Index) const;
+	
 	void AddIntegerInputPin();
 	void RemoveIntegerInputPin();
 	FText GetIntegerAddButtonTooltipText() const;
