Fixed static analysis warning (C6239: (<non-zero constant> && <expression>) always evaluates to the result of <expression>)

diff --git a/Engine/Plugins/FX/Niagara/Source/Niagara/Private/NiagaraRendererComponents.cpp b/Engine/Plugins/FX/Niagara/Source/Niagara/Private/NiagaraRendererComponents.cpp
--- a/Engine/Plugins/FX/Niagara/Source/Niagara/Private/NiagaraRendererComponents.cpp
+++ b/Engine/Plugins/FX/Niagara/Source/Niagara/Private/NiagaraRendererComponents.cpp
@@ -606,7 +606,12 @@ void FNiagaraRendererComponents::TickPropertyBindings(
 			}
 		}
 
-		bool bForceStructConversion = !UE_LARGE_WORLD_COORDINATES_DISABLED && PropertyType.GetScriptStruct() && !FNiagaraTypeHelper::IsNiagaraFriendlyTopLevelStruct(PropertyType.GetScriptStruct());
+#if UE_LARGE_WORLD_COORDINATES_DISABLED
+		bool bForceStructConversion = false;
+#else
+		bool bForceStructConversion = PropertyType.GetScriptStruct() && !FNiagaraTypeHelper::IsNiagaraFriendlyTopLevelStruct(PropertyType.GetScriptStruct());
+#endif
+		
 		if (PropertyType.IsValid() && DataVariable.GetType() != PropertyType && (!PropertySetter->bIgnoreConversion || bForceStructConversion))
 		{
 			FNiagaraVariable TargetVariable(PropertyType, DataVariable.GetName());
