Speculative fix for a "ParameterOffsets.Num() should be 0 is 1" warning during cook

diff --git a/Engine/Plugins/FX/Niagara/Source/Niagara/Private/NiagaraScript.cpp b/Engine/Plugins/FX/Niagara/Source/Niagara/Private/NiagaraScript.cpp
--- a/Engine/Plugins/FX/Niagara/Source/Niagara/Private/NiagaraScript.cpp
+++ b/Engine/Plugins/FX/Niagara/Source/Niagara/Private/NiagaraScript.cpp
@@ -1355,6 +1355,7 @@ void UNiagaraScript::PreSave(FObjectPreSaveContext ObjectSaveContext)
 	UNiagaraSystem* SystemOwner = FindRootSystem();
 	if (SystemOwner)
 	{
+		SystemOwner->EnsureFullyLoaded();
 		SystemOwner->WaitForCompilationComplete();
 	}
 
