- Expose RandomInfo to Niagara in a way that it can be used consistently across GPU

diff --git a/Engine/Plugins/FX/Niagara/Shaders/Private/NiagaraEmitterInstanceShader.usf b/Engine/Plugins/FX/Niagara/Shaders/Private/NiagaraEmitterInstanceShader.usf
--- a/Engine/Plugins/FX/Niagara/Shaders/Private/NiagaraEmitterInstanceShader.usf
+++ b/Engine/Plugins/FX/Niagara/Shaders/Private/NiagaraEmitterInstanceShader.usf
@@ -775,5 +775,36 @@ int rand_int(int x, int Seed1, int Seed2, int Seed3)
 			return AcquireIndex(DataSetID, bIsValid);
 	}
 
+////////////////////////////////////////////////////////////////////////////////////
+// Random number functions
+struct NiagaraRandInfo
+{
+	int Seed1;
+	int Seed2;
+	int Seed3;
+};
+
+#if GPU_SIMULATION
+float NiagaraRandomFloat(NiagaraRandInfo RandInfo)
+{
+	return uint(RandInfo.Seed3) == 0xffffffff ? NiagaraInternalNoise(uint(RandInfo.Seed1), uint(RandInfo.Seed2), uint(RandInfo.Seed3)) : rand(1.0f, uint(RandInfo.Seed1), uint(RandInfo.Seed2), uint(RandInfo.Seed3));
+}
+
+int NiagaraRandomInt(NiagaraRandInfo RandInfo, int Range)
+{
+	float T = NiagaraRandomFloat(RandInfo);
+	return int(floor(float(Range) * T));
+}
+
+float3 NiagaraRandomBaryCoord(NiagaraRandInfo RandInfo)
+{
+	float2 r = float2(NiagaraRandomFloat(RandInfo), NiagaraRandomFloat(RandInfo));
+	float sqrt0 = sqrt(r.x);
+	float sqrt1 = sqrt(r.y);
+	return float3(1.0f - sqrt0, sqrt0 * (1.0 - r.y), r.y * sqrt0);
+}
+#endif
+
+////////////////////////////////////////////////////////////////////////////////////
 //Include the simulation shader code generated by the node graph.
 #include "/Engine/Generated/NiagaraEmitterInstance.ush"
diff --git a/Engine/Plugins/FX/Niagara/Source/Niagara/Private/NiagaraModule.cpp b/Engine/Plugins/FX/Niagara/Source/Niagara/Private/NiagaraModule.cpp
--- a/Engine/Plugins/FX/Niagara/Source/Niagara/Private/NiagaraModule.cpp
+++ b/Engine/Plugins/FX/Niagara/Source/Niagara/Private/NiagaraModule.cpp
@@ -571,6 +571,7 @@ void INiagaraModule::OnAssetLoaded(UObject* Asset)
 
 UScriptStruct* FNiagaraTypeDefinition::ParameterMapStruct;
 UScriptStruct* FNiagaraTypeDefinition::IDStruct;
+UScriptStruct* FNiagaraTypeDefinition::RandInfoStruct;
 UScriptStruct* FNiagaraTypeDefinition::NumericStruct;
 UScriptStruct* FNiagaraTypeDefinition::FloatStruct;
 UScriptStruct* FNiagaraTypeDefinition::BoolStruct;
@@ -608,6 +609,7 @@ UEnum* FNiagaraTypeDefinition::FunctionDebugStateEnum;
 
 FNiagaraTypeDefinition FNiagaraTypeDefinition::ParameterMapDef;
 FNiagaraTypeDefinition FNiagaraTypeDefinition::IDDef;
+FNiagaraTypeDefinition FNiagaraTypeDefinition::RandInfoDef;
 FNiagaraTypeDefinition FNiagaraTypeDefinition::NumericDef;
 FNiagaraTypeDefinition FNiagaraTypeDefinition::FloatDef;
 FNiagaraTypeDefinition FNiagaraTypeDefinition::BoolDef;
@@ -654,6 +656,7 @@ void FNiagaraTypeDefinition::Init()
 	static auto* NiagaraPkg = FindObjectChecked<UPackage>(nullptr, TEXT("/Script/Niagara"));
 	FNiagaraTypeDefinition::ParameterMapStruct = FindObjectChecked<UScriptStruct>(NiagaraPkg, TEXT("NiagaraParameterMap"));
 	FNiagaraTypeDefinition::IDStruct = FindObjectChecked<UScriptStruct>(NiagaraPkg, TEXT("NiagaraID"));
+	FNiagaraTypeDefinition::RandInfoStruct = FindObjectChecked<UScriptStruct>(NiagaraPkg, TEXT("NiagaraRandInfo"));
 	FNiagaraTypeDefinition::NumericStruct = FindObjectChecked<UScriptStruct>(NiagaraPkg, TEXT("NiagaraNumeric"));
 	FNiagaraTypeDefinition::FloatStruct = FindObjectChecked<UScriptStruct>(NiagaraPkg, TEXT("NiagaraFloat"));
 	FNiagaraTypeDefinition::BoolStruct = FindObjectChecked<UScriptStruct>(NiagaraPkg, TEXT("NiagaraBool"));
@@ -679,6 +682,7 @@ void FNiagaraTypeDefinition::Init()
 	
 	ParameterMapDef = FNiagaraTypeDefinition(ParameterMapStruct);
 	IDDef = FNiagaraTypeDefinition(IDStruct);
+	RandInfoDef = FNiagaraTypeDefinition(RandInfoStruct);
 	NumericDef = FNiagaraTypeDefinition(NumericStruct);
 	FloatDef = FNiagaraTypeDefinition(FloatStruct);
 	BoolDef = FNiagaraTypeDefinition(BoolStruct);
@@ -871,6 +875,7 @@ void FNiagaraTypeDefinition::RecreateUserDefinedTypeRegistry()
 
 	FNiagaraTypeRegistry::Register(ParameterMapDef, ParamFlags);
 	FNiagaraTypeRegistry::Register(IDDef, ParamFlags | PayloadFlags);
+	FNiagaraTypeRegistry::Register(RandInfoDef, ParamFlags | PayloadFlags);
 	FNiagaraTypeRegistry::Register(NumericDef, ParamFlags);
 	FNiagaraTypeRegistry::Register(FloatDef, ParamFlags | PayloadFlags);
 	FNiagaraTypeRegistry::Register(HalfDef, ParamFlags | PayloadFlags);
@@ -898,7 +903,6 @@ void FNiagaraTypeDefinition::RecreateUserDefinedTypeRegistry()
 	FNiagaraTypeRegistry::Register(UTextureDef, VarFlags);
 	FNiagaraTypeRegistry::Register(UTextureRenderTargetDef, VarFlags);
 	FNiagaraTypeRegistry::Register(UStaticMeshDef, VarFlags);
-	FNiagaraTypeRegistry::Register(FNiagaraRandInfo::StaticStruct(), ParamFlags | PayloadFlags);
 	FNiagaraTypeRegistry::Register(StaticEnum<ENiagaraLegacyTrailWidthMode>(), ParamFlags | PayloadFlags);
 
 	
diff --git a/Engine/Plugins/FX/Niagara/Source/Niagara/Public/NiagaraCommon.h b/Engine/Plugins/FX/Niagara/Source/Niagara/Public/NiagaraCommon.h
--- a/Engine/Plugins/FX/Niagara/Source/Niagara/Public/NiagaraCommon.h
+++ b/Engine/Plugins/FX/Niagara/Source/Niagara/Public/NiagaraCommon.h
@@ -1266,22 +1266,6 @@ struct TStructOpsTypeTraits<FNiagaraUserParameterBinding> : public TStructOpsTyp
 	};
 };
 
-USTRUCT()
-struct FNiagaraRandInfo
-{
-	GENERATED_USTRUCT_BODY()
-
-	UPROPERTY(EditAnywhere, Category = "Random")
-	int32 Seed1 = 0;
-	
-	UPROPERTY(EditAnywhere, Category = "Random")
-	int32 Seed2 = 0;
-
-	UPROPERTY(EditAnywhere, Category = "Random")
-	int32 Seed3 = 0;
-};
-
-
 //////////////////////////////////////////////////////////////////////////
 // Legacy Anim Trail Support
 
diff --git a/Engine/Plugins/FX/Niagara/Source/Niagara/Public/NiagaraTypes.h b/Engine/Plugins/FX/Niagara/Source/Niagara/Public/NiagaraTypes.h
--- a/Engine/Plugins/FX/Niagara/Source/Niagara/Public/NiagaraTypes.h
+++ b/Engine/Plugins/FX/Niagara/Source/Niagara/Public/NiagaraTypes.h
@@ -241,6 +241,21 @@ struct FNiagaraID
 	bool operator<(const FNiagaraID& Other)const { return Index < Other.Index || (Index == Other.Index && AcquireTag < Other.AcquireTag); }
 };
 
+USTRUCT()
+struct FNiagaraRandInfo
+{
+	GENERATED_USTRUCT_BODY()
+
+	UPROPERTY(EditAnywhere, Category = "Random")
+	int32 Seed1 = 0;
+	
+	UPROPERTY(EditAnywhere, Category = "Random")
+	int32 Seed2 = 0;
+
+	UPROPERTY(EditAnywhere, Category = "Random")
+	int32 Seed3 = 0;
+};
+
 #define NIAGARA_INVALID_ID (FNiagaraID({(INDEX_NONE), (INDEX_NONE)}))
 
 FORCEINLINE uint32 GetTypeHash(const FNiagaraID& ID)
@@ -878,6 +893,7 @@ public:
 	static const FNiagaraTypeDefinition& GetGenericNumericDef() { return NumericDef; }
 	static const FNiagaraTypeDefinition& GetParameterMapDef() { return ParameterMapDef; }
 	static const FNiagaraTypeDefinition& GetIDDef() { return IDDef; }
+	static const FNiagaraTypeDefinition& GetRandInfoDef() { return RandInfoDef; }
 	static const FNiagaraTypeDefinition& GetUObjectDef() { return UObjectDef; }
 	static const FNiagaraTypeDefinition& GetUMaterialDef() { return UMaterialDef; }
 	static const FNiagaraTypeDefinition& GetUStaticMeshDef() { return UStaticMeshDef; }
@@ -905,6 +921,7 @@ public:
 	static UScriptStruct* GetWildcardStruct() { return WildcardStruct; }
 	static UScriptStruct* GetParameterMapStruct() { return ParameterMapStruct; }
 	static UScriptStruct* GetIDStruct() { return IDStruct; }
+	static UScriptStruct* GetRandInfoStruct() { return RandInfoStruct; }
 
 	static UScriptStruct* GetHalfStruct() { return HalfStruct; }
 	static UScriptStruct* GetHalfVec2Struct() { return HalfVec2Struct; }
@@ -960,6 +977,7 @@ private:
 	static FNiagaraTypeDefinition NumericDef;
 	static FNiagaraTypeDefinition ParameterMapDef;
 	static FNiagaraTypeDefinition IDDef;
+	static FNiagaraTypeDefinition RandInfoDef;
 	static FNiagaraTypeDefinition UObjectDef;
 	static FNiagaraTypeDefinition UMaterialDef;
 	static FNiagaraTypeDefinition UTextureDef;
@@ -1009,6 +1027,7 @@ private:
 
 	static UScriptStruct* ParameterMapStruct;
 	static UScriptStruct* IDStruct;
+	static UScriptStruct* RandInfoStruct;
 
 	static TSet<UScriptStruct*> NumericStructs;
 	static TArray<FNiagaraTypeDefinition> OrderedNumericTypes;
@@ -1047,6 +1066,7 @@ const FNiagaraTypeDefinition& FNiagaraTypeDefinition::Get()
 	if (TIsSame<T, FMatrix44f>::Value) { return FNiagaraTypeDefinition::GetMatrix4Def(); }
 	if (TIsSame<T, FLinearColor>::Value) { return FNiagaraTypeDefinition::GetColorDef(); }
 	if (TIsSame<T, FNiagaraID>::Value) { return FNiagaraTypeDefinition::GetIDDef(); }
+	if (TIsSame<T, FNiagaraRandInfo>::Value) { return FNiagaraTypeDefinition::GetRandInfoDef(); }
 }
 
 FORCEINLINE uint32 GetTypeHash(const FNiagaraTypeDefinition& Type)
diff --git a/Engine/Plugins/FX/Niagara/Source/NiagaraEditor/Private/NiagaraHlslTranslator.cpp b/Engine/Plugins/FX/Niagara/Source/NiagaraEditor/Private/NiagaraHlslTranslator.cpp
--- a/Engine/Plugins/FX/Niagara/Source/NiagaraEditor/Private/NiagaraHlslTranslator.cpp
+++ b/Engine/Plugins/FX/Niagara/Source/NiagaraEditor/Private/NiagaraHlslTranslator.cpp
@@ -1616,15 +1616,17 @@ const FNiagaraTranslateResults &FHlslNiagaraTranslator::Translate(const FNiagara
 			}
 		}
 
-		// NiagaraID is a fundamental type which can be used in custom HLSL without declaring any pins of that type in any nodes.
-		// Make sure it's always defined in the generated HLSL.
-		StructsToDefine.AddUnique(FNiagaraTypeDefinition::GetIDDef());
-
 		// Generate the Parameter Map HLSL definitions. We don't add to the final HLSL output here. We just build up the strings and tables
 		// that are needed later.
 		TArray<FNiagaraVariable> PrimaryDataSetOutputEntries;
 		FString ParameterMapDefinitionStr = BuildParameterMapHlslDefinitions(PrimaryDataSetOutputEntries);
 
+		// Ensure some structures are always added as we use them in custom HLSL / data interface templates
+		// Remove some structures which we define inside NiagaraEmitterInstanceShader.usf as we want a common set of functions
+		// Ensure we always add structures that are fundamental to custom HLSL or data interface templates shader files
+		StructsToDefine.AddUnique(FNiagaraTypeDefinition::GetIDDef());
+		StructsToDefine.Remove(FNiagaraTypeDefinition::GetRandInfoDef());
+
 		for (const FNiagaraTypeDefinition& Type : StructsToDefine)
 		{
 			FText ErrorMessage;
