- Merging latest from UE5 Private branch

#rnx

diff --git a/Engine/Plugins/FX/Niagara/Shaders/Private/NiagaraDataInterfaceSkeletalMesh.ush b/Engine/Plugins/FX/Niagara/Shaders/Private/NiagaraDataInterfaceSkeletalMesh.ush
--- a/Engine/Plugins/FX/Niagara/Shaders/Private/NiagaraDataInterfaceSkeletalMesh.ush
+++ b/Engine/Plugins/FX/Niagara/Shaders/Private/NiagaraDataInterfaceSkeletalMesh.ush
@@ -223,11 +223,13 @@ FDISkelMeshSkinnedVertex DISkelMesh_GetSkinnedVertex(in FDISkelMeshContext DICon
 		TangentZ.xyz = mul(CurrBoneMatrix, float4(TangentZ.xyz, 0.0f)).xyz;
 	}
 
+	float3 TangentY = cross(TangentZ.xyz, TangentX.xyz) * TangentZ.w;
+
 	FDISkelMeshSkinnedVertex SkinnedVertex;
 	SkinnedVertex.Position = CurrPosition;
 	SkinnedVertex.PrevPosition = PrevPosition;
-	SkinnedVertex.TangentX = TangentX;
-	SkinnedVertex.TangentY = cross(TangentZ.xyz, TangentX.xyz);
+	SkinnedVertex.TangentX = cross(TangentY, TangentZ.xyz) * TangentZ.w;
+	SkinnedVertex.TangentY = TangentY;
 	SkinnedVertex.TangentZ = TangentZ.xyz;
 	return SkinnedVertex;
 }
@@ -787,9 +789,9 @@ void DISkelMesh_GetSkinnedVertexWS(in FDISkelMeshContext DIContext, in int Verte
     OutPosition = mul(float4(SkinnedVertex.Position, 1.0f), DIContext.InstanceTransform).xyz;
     float3 PrevPosition = mul(float4(SkinnedVertex.PrevPosition, 1.0f), DIContext.InstancePrevTransform).xyz;
     OutVelocity = (SkinnedVertex.Position - SkinnedVertex.PrevPosition) * DIContext.InstanceInvDeltaTime;
-	OutNormal = SkinnedVertex.TangentZ;
-	OutBinormal = SkinnedVertex.TangentY;
-	OutTangent = SkinnedVertex.TangentX;
+    OutNormal   = normalize(mul(float4(SkinnedVertex.TangentZ, 0.0f), DIContext.InstanceTransform).xyz);
+    OutBinormal = normalize(mul(float4(SkinnedVertex.TangentY, 0.0f), DIContext.InstanceTransform).xyz);
+    OutTangent  = normalize(mul(float4(SkinnedVertex.TangentX, 0.0f), DIContext.InstanceTransform).xyz);
 }
 
 void DISkelMesh_GetVertexColor(in FDISkelMeshContext DIContext, in int VertexIndex, out float4 OutColor)
diff --git a/Engine/Plugins/FX/Niagara/Source/Niagara/Private/NDISkeletalMesh_TriangleSampling.cpp b/Engine/Plugins/FX/Niagara/Source/Niagara/Private/NDISkeletalMesh_TriangleSampling.cpp
--- a/Engine/Plugins/FX/Niagara/Source/Niagara/Private/NDISkeletalMesh_TriangleSampling.cpp
+++ b/Engine/Plugins/FX/Niagara/Source/Niagara/Private/NDISkeletalMesh_TriangleSampling.cpp
@@ -929,12 +929,11 @@ void UNiagaraDataInterfaceSkeletalMesh::GetTriCoordColor(FVectorVMExternalFuncti
 	if ( const FSkeletalMeshLODRenderData* LODData = InstData->CachedLODData )
 	{
 		const FColorVertexBuffer& Colors = LODData->StaticVertexBuffers.ColorVertexBuffer;
-		checkfSlow(Colors.GetNumVertices() != 0, TEXT("Trying to access vertex colors from mesh without any."));
 
 		const FMultiSizeIndexContainer& Indices = LODData->MultiSizeIndexContainer;
 		const FRawStaticIndexBuffer16or32Interface* IndexBuffer = Indices.GetIndexBuffer();
 		const int32 TriMax = (IndexBuffer->Num() / 3) - 1;
-		if (TriMax >= 0)
+		if (TriMax >= 0 && (Colors.GetNumVertices() > 0) && (Colors.GetVertexData() != nullptr) && Colors.GetAllowCPUAccess() )
 		{
 			for (int32 i = 0; i < Context.GetNumInstances(); ++i)
 			{
diff --git a/Engine/Plugins/FX/Niagara/Source/Niagara/Private/NiagaraComponent.cpp b/Engine/Plugins/FX/Niagara/Source/Niagara/Private/NiagaraComponent.cpp
--- a/Engine/Plugins/FX/Niagara/Source/Niagara/Private/NiagaraComponent.cpp
+++ b/Engine/Plugins/FX/Niagara/Source/Niagara/Private/NiagaraComponent.cpp
@@ -1657,8 +1657,8 @@ void UNiagaraComponent::BeginDestroy()
 		PoolingMethod = ENCPoolMethod::None;
 	}
 
-	//By now we will have already unregistered with the scalability manger. Either directly in OnComponentDestroyed, or via the post GC callbacks in the manager it's self in the case of someone calling MarkPendingKill() directly on a component.
-	ScalabilityManagerHandle = INDEX_NONE;
+	//In most cases we have already unregistered by now but there are few edge cases where we need to unregister here.
+	UnregisterWithScalabilityManager();
 
 	DestroyInstance();
 
diff --git a/Engine/Plugins/FX/Niagara/Source/Niagara/Private/NiagaraEmitter.cpp b/Engine/Plugins/FX/Niagara/Source/Niagara/Private/NiagaraEmitter.cpp
--- a/Engine/Plugins/FX/Niagara/Source/Niagara/Private/NiagaraEmitter.cpp
+++ b/Engine/Plugins/FX/Niagara/Source/Niagara/Private/NiagaraEmitter.cpp
@@ -1235,6 +1235,19 @@ void UNiagaraEmitter::CacheFromCompiledData(const FNiagaraDataSetCompiledData* C
 	{
 		MaxInstanceCount = 0;
 	}
+
+	// Cache shaders for all GPU scripts
+#if WITH_EDITORONLY_DATA
+	if ( AreAllScriptAndSourcesSynchronized() )
+	{
+		ForEachScript(
+			[](UNiagaraScript* Script)
+			{
+				Script->CacheResourceShadersForRendering(false, false);
+			}
+		);
+	}
+#endif
 }
 
 void UNiagaraEmitter::CacheFromShaderCompiled()
@@ -1351,7 +1364,6 @@ void UNiagaraEmitter::UpdateEmitterAfterLoad()
 	}
 	DebugSimName.Append(GetName());
 #endif
-
 }
 
 bool UNiagaraEmitter::IsAllowedByScalability()const
diff --git a/Engine/Plugins/FX/Niagara/Source/Niagara/Private/NiagaraModule.cpp b/Engine/Plugins/FX/Niagara/Source/Niagara/Private/NiagaraModule.cpp
--- a/Engine/Plugins/FX/Niagara/Source/Niagara/Private/NiagaraModule.cpp
+++ b/Engine/Plugins/FX/Niagara/Source/Niagara/Private/NiagaraModule.cpp
@@ -1180,18 +1180,18 @@ UNiagaraDataInterfaceBase* FNiagaraTypeRegistry::GetDefaultDataInterfaceByName(c
 		FNiagaraTypeRegistry& Registry = Get();
 		FReadScopeLock Lock(Registry.RegisteredTypesLock);
 		for (const FNiagaraTypeDefinition& Def : Registry.RegisteredTypes)
-	{
-		if (Def.IsDataInterface())
 		{
-			UClass* FoundDIClass = Def.GetClass();
-			if (FoundDIClass && (FoundDIClass->GetName() == DIClassName || FoundDIClass->GetFullName() == DIClassName))
+			if (Def.IsDataInterface())
 			{
-				DIClass = FoundDIClass;
-				break;
+				UClass* FoundDIClass = Def.GetClass();
+				if (FoundDIClass && (FoundDIClass->GetName() == DIClassName || FoundDIClass->GetFullName() == DIClassName))
+				{
+					DIClass = FoundDIClass;
+					break;
+				}
 			}
 		}
 	}
-	}
 
 	// Consider the possibility of a redirector pointing to a new location..
 	if (DIClass == nullptr)
diff --git a/Engine/Plugins/FX/Niagara/Source/Niagara/Private/NiagaraRendererMeshes.cpp b/Engine/Plugins/FX/Niagara/Source/Niagara/Private/NiagaraRendererMeshes.cpp
--- a/Engine/Plugins/FX/Niagara/Source/Niagara/Private/NiagaraRendererMeshes.cpp
+++ b/Engine/Plugins/FX/Niagara/Source/Niagara/Private/NiagaraRendererMeshes.cpp
@@ -52,6 +52,14 @@ struct FNiagaraDynamicDataMesh : public FNiagaraDynamicDataBase
 	{
 	}
 
+	virtual void ApplyMaterialOverride(int32 MaterialIndex, UMaterialInterface* MaterialOverride) override
+	{
+		if (Materials.IsValidIndex(MaterialIndex) && MaterialOverride)
+		{
+			Materials[MaterialIndex] = MaterialOverride->GetRenderProxy();
+		}
+	}
+
 	TArray<FMaterialRenderProxy*, TInlineAllocator<8>> Materials;
 	TArray<UNiagaraDataInterface*> DataInterfacesBound;
 	TArray<UObject*> ObjectsBound;
diff --git a/Engine/Plugins/FX/Niagara/Source/Niagara/Private/NiagaraRendererRibbons.cpp b/Engine/Plugins/FX/Niagara/Source/Niagara/Private/NiagaraRendererRibbons.cpp
--- a/Engine/Plugins/FX/Niagara/Source/Niagara/Private/NiagaraRendererRibbons.cpp
+++ b/Engine/Plugins/FX/Niagara/Source/Niagara/Private/NiagaraRendererRibbons.cpp
@@ -96,6 +96,14 @@ struct FNiagaraDynamicDataRibbon : public FNiagaraDynamicDataBase
 	{
 	}
 
+	virtual void ApplyMaterialOverride(int32 MaterialIndex, UMaterialInterface* MaterialOverride) override
+	{
+		if (MaterialIndex == 0 && MaterialOverride)
+		{
+			Material = MaterialOverride->GetRenderProxy();
+		}
+	}
+
 	/** Material to use passed to the Renderer. */
 	FMaterialRenderProxy* Material;
 
diff --git a/Engine/Plugins/FX/Niagara/Source/Niagara/Private/NiagaraRendererSprites.cpp b/Engine/Plugins/FX/Niagara/Source/Niagara/Private/NiagaraRendererSprites.cpp
--- a/Engine/Plugins/FX/Niagara/Source/Niagara/Private/NiagaraRendererSprites.cpp
+++ b/Engine/Plugins/FX/Niagara/Source/Niagara/Private/NiagaraRendererSprites.cpp
@@ -47,6 +47,14 @@ struct FNiagaraDynamicDataSprites : public FNiagaraDynamicDataBase
 	{
 	}
 
+	virtual void ApplyMaterialOverride(int32 MaterialIndex, UMaterialInterface* MaterialOverride) override
+	{
+		if (MaterialIndex == 0 && MaterialOverride)
+		{
+			Material = MaterialOverride->GetRenderProxy();
+		}
+	}
+
 	FMaterialRenderProxy* Material = nullptr;
 	TArray<UNiagaraDataInterface*> DataInterfacesBound;
 	TArray<UObject*> ObjectsBound;
diff --git a/Engine/Plugins/FX/Niagara/Source/Niagara/Private/NiagaraScript.cpp b/Engine/Plugins/FX/Niagara/Source/Niagara/Private/NiagaraScript.cpp
--- a/Engine/Plugins/FX/Niagara/Source/Niagara/Private/NiagaraScript.cpp
+++ b/Engine/Plugins/FX/Niagara/Source/Niagara/Private/NiagaraScript.cpp
@@ -1793,13 +1793,6 @@ void UNiagaraScript::PostLoad()
 
 	ProcessSerializedShaderMaps();
 
-#if WITH_EDITORONLY_DATA
-	if (CachedScriptVMId.BaseScriptCompileHash.IsValid() && AreScriptAndSourceSynchronized())
-	{
-		CacheResourceShadersForRendering(false, false /*bNeedsRecompile*/);
-	}
-#endif
-
 	GenerateStatIDs();
 
 	// Optimize the VM script for runtime usage
diff --git a/Engine/Plugins/FX/Niagara/Source/Niagara/Private/NiagaraSystemInstanceController.cpp b/Engine/Plugins/FX/Niagara/Source/Niagara/Private/NiagaraSystemInstanceController.cpp
--- a/Engine/Plugins/FX/Niagara/Source/Niagara/Private/NiagaraSystemInstanceController.cpp
+++ b/Engine/Plugins/FX/Niagara/Source/Niagara/Private/NiagaraSystemInstanceController.cpp
@@ -61,7 +61,7 @@ FNiagaraSystemRenderData* FNiagaraSystemInstanceController::CreateSystemRenderDa
 
 void FNiagaraSystemInstanceController::GenerateSetDynamicDataCommands(FNiagaraSystemRenderData::FSetDynamicDataCommandList& Commands, FNiagaraSystemRenderData& RenderData, const FNiagaraSceneProxy& SceneProxy)
 {
-	RenderData.GenerateSetDynamicDataCommands(Commands, SceneProxy, SystemInstance.Get());
+	RenderData.GenerateSetDynamicDataCommands(Commands, SceneProxy, SystemInstance.Get(), EmitterMaterials);
 }
 
 void FNiagaraSystemInstanceController::PostTickRenderers(FNiagaraSystemRenderData& RenderData)
diff --git a/Engine/Plugins/FX/Niagara/Source/Niagara/Private/NiagaraSystemRenderData.cpp b/Engine/Plugins/FX/Niagara/Source/Niagara/Private/NiagaraSystemRenderData.cpp
--- a/Engine/Plugins/FX/Niagara/Source/Niagara/Private/NiagaraSystemRenderData.cpp
+++ b/Engine/Plugins/FX/Niagara/Source/Niagara/Private/NiagaraSystemRenderData.cpp
@@ -118,7 +118,7 @@ uint32 FNiagaraSystemRenderData::GetDynamicDataSize() const
 	return Size;
 }
 
-void FNiagaraSystemRenderData::GenerateSetDynamicDataCommands(FSetDynamicDataCommandList& Commands, const FNiagaraSceneProxy& SceneProxy, const FNiagaraSystemInstance* SystemInstance)
+void FNiagaraSystemRenderData::GenerateSetDynamicDataCommands(FSetDynamicDataCommandList& Commands, const FNiagaraSceneProxy& SceneProxy, const FNiagaraSystemInstance* SystemInstance, TConstArrayView<FMaterialOverride> MaterialOverrides)
 {
 	if (!SystemInstance)
 	{
@@ -183,6 +183,17 @@ void FNiagaraSystemRenderData::GenerateSetDynamicDataCommands(FSetDynamicDataCom
 					if (bRendererEditorEnabled && !EmitterInst->IsComplete() && !SystemInstance->IsComplete())
 					{
 						NewData = Renderer->GenerateDynamicData(&SceneProxy, Properties, EmitterInst);
+
+						if (NewData)
+						{
+							for (const FMaterialOverride& MaterialOverride : MaterialOverrides)
+							{
+								if (MaterialOverride.EmitterRendererProperty == Properties)
+								{
+									NewData->ApplyMaterialOverride(MaterialOverride.MaterialSubIndex, MaterialOverride.Material);
+								}
+							}
+						}
 					}
 
 					Commands.Add(FSetDynamicDataCommand(Renderer, NewData));
diff --git a/Engine/Plugins/FX/Niagara/Source/Niagara/Public/NiagaraRenderer.h b/Engine/Plugins/FX/Niagara/Source/Niagara/Public/NiagaraRenderer.h
--- a/Engine/Plugins/FX/Niagara/Source/Niagara/Public/NiagaraRenderer.h
+++ b/Engine/Plugins/FX/Niagara/Source/Niagara/Public/NiagaraRenderer.h
@@ -42,6 +42,7 @@ struct FNiagaraDynamicDataBase
 	FORCEINLINE FNiagaraSystemInstanceID GetSystemInstanceID() const { return SystemInstanceID; }
 
 	void SetVertexFactoryData(class FNiagaraVertexFactoryBase& VertexFactory);
+	virtual void ApplyMaterialOverride(int32 MaterialIndex, UMaterialInterface* MaterialOverride) {};
 
 protected:
 	FMaterialRelevance MaterialRelevance;
diff --git a/Engine/Plugins/FX/Niagara/Source/Niagara/Public/NiagaraSystemInstanceController.h b/Engine/Plugins/FX/Niagara/Source/Niagara/Public/NiagaraSystemInstanceController.h
--- a/Engine/Plugins/FX/Niagara/Source/Niagara/Public/NiagaraSystemInstanceController.h
+++ b/Engine/Plugins/FX/Niagara/Source/Niagara/Public/NiagaraSystemInstanceController.h
@@ -141,12 +141,7 @@ public:
 	void OnPooledReuse(UWorld& NewWorld) { SystemInstance->OnPooledReuse(NewWorld); }
 
 private:
-	struct FMaterialOverride
-	{
-		UMaterialInterface* Material = nullptr;
-		const UNiagaraRendererProperties* EmitterRendererProperty = nullptr; // Pointer only used for lookups, don't actually dereference
-		uint32 MaterialSubIndex = 0;
-	};
+	using FMaterialOverride = FNiagaraSystemRenderData::FMaterialOverride;
 
 	//~ FGCObject interface
 	virtual void AddReferencedObjects(FReferenceCollector& Collector) override;
diff --git a/Engine/Plugins/FX/Niagara/Source/Niagara/Public/NiagaraSystemRenderData.h b/Engine/Plugins/FX/Niagara/Source/Niagara/Public/NiagaraSystemRenderData.h
--- a/Engine/Plugins/FX/Niagara/Source/Niagara/Public/NiagaraSystemRenderData.h
+++ b/Engine/Plugins/FX/Niagara/Source/Niagara/Public/NiagaraSystemRenderData.h
@@ -20,6 +20,13 @@ class FSceneView;
 class FNiagaraSystemRenderData
 {
 public:
+	struct FMaterialOverride
+	{
+		UMaterialInterface* Material = nullptr;
+		const UNiagaraRendererProperties* EmitterRendererProperty = nullptr; // Pointer only used for lookups, don't actually dereference
+		uint32 MaterialSubIndex = 0;
+	};
+
 	struct FSetDynamicDataCommand
 	{
 		FNiagaraRenderer* Renderer = nullptr;
@@ -53,7 +60,7 @@ public:
 	/** Gets the total size of dynamic data allocated */
 	uint32 GetDynamicDataSize() const;
 	/** Called at the end of the frame, before rendering when necessary to provide renderers with render thread data */
-	void GenerateSetDynamicDataCommands(FSetDynamicDataCommandList& Commands, const FNiagaraSceneProxy& SceneProxy, const FNiagaraSystemInstance* SystemInstance);
+	void GenerateSetDynamicDataCommands(FSetDynamicDataCommandList& Commands, const FNiagaraSceneProxy& SceneProxy, const FNiagaraSystemInstance* SystemInstance, TConstArrayView<FMaterialOverride> MaterialOverrides);
 	/** Gets the dynamic mesh elements from all renderers */
 	void GetDynamicMeshElements(const TArray<const FSceneView*>& Views, const FSceneViewFamily& ViewFamily, uint32 VisibilityMap, FMeshElementCollector& Collector, const FNiagaraSceneProxy& SceneProxy);
 #if RHI_RAYTRACING
diff --git a/Engine/Plugins/FX/Niagara/Source/NiagaraCore/Private/NiagaraCustomVersion.cpp b/Engine/Plugins/FX/Niagara/Source/NiagaraCore/Private/NiagaraCustomVersion.cpp
--- a/Engine/Plugins/FX/Niagara/Source/NiagaraCore/Private/NiagaraCustomVersion.cpp
+++ b/Engine/Plugins/FX/Niagara/Source/NiagaraCore/Private/NiagaraCustomVersion.cpp
@@ -8,5 +8,5 @@ const FGuid FNiagaraCustomVersion::GUID(0xFCF57AFA, 0x50764283, 0xB9A9E658, 0xFF
 // Register the custom version with core
 FCustomVersionRegistration GRegisterNiagaraCustomVersion(FNiagaraCustomVersion::GUID, FNiagaraCustomVersion::LatestVersion, TEXT("NiagaraVer"));
 
-const FGuid FNiagaraCustomVersion::LatestScriptCompileVersion(0x375F6E63, 0x7F2144F2, 0xAD4EC387, 0x0AB6472D);
+const FGuid FNiagaraCustomVersion::LatestScriptCompileVersion(0xDDF215A2, 0x07D62F47, 0x86F0F7B5, 0xC9DEAD89);
 
diff --git a/Engine/Source/Runtime/Engine/Public/Rendering/ColorVertexBuffer.h b/Engine/Source/Runtime/Engine/Public/Rendering/ColorVertexBuffer.h
--- a/Engine/Source/Runtime/Engine/Public/Rendering/ColorVertexBuffer.h
+++ b/Engine/Source/Runtime/Engine/Public/Rendering/ColorVertexBuffer.h
@@ -100,7 +100,7 @@ public:
 	{
 		return NumVertices;
 	}
-	bool GetAllowCPUAccess() const;
+	ENGINE_API bool GetAllowCPUAccess() const;
 
 	/** Useful for memory profiling. */
 	ENGINE_API uint32 GetAllocatedSize() const;
