- Fix not setting mesh changed delegates when we change meshes

#rnx

diff --git a/Engine/Plugins/FX/Niagara/Source/Niagara/Private/NiagaraMeshRendererProperties.cpp b/Engine/Plugins/FX/Niagara/Source/Niagara/Private/NiagaraMeshRendererProperties.cpp
--- a/Engine/Plugins/FX/Niagara/Source/Niagara/Private/NiagaraMeshRendererProperties.cpp
+++ b/Engine/Plugins/FX/Niagara/Source/Niagara/Private/NiagaraMeshRendererProperties.cpp
@@ -891,13 +891,14 @@ bool UNiagaraMeshRendererProperties::ChangeRequiresMeshListRebuild(const FProper
 	}
 
 	// If any of these are changed, we have to rebuild the mesh list
-	static const TArray<FName, TInlineAllocator<4>> RebuildMeshPropertyNames
+	static const TArray<FName, TInlineAllocator<6>> RebuildMeshPropertyNames
 	{
 		GET_MEMBER_NAME_CHECKED(UNiagaraMeshRendererProperties, bEnableMeshFlipbook),
 		GET_MEMBER_NAME_CHECKED(UNiagaraMeshRendererProperties, FirstFlipbookFrame),
 		GET_MEMBER_NAME_CHECKED(UNiagaraMeshRendererProperties, FlipbookSuffixFormat),
 		GET_MEMBER_NAME_CHECKED(UNiagaraMeshRendererProperties, FlipbookSuffixNumDigits),
 		GET_MEMBER_NAME_CHECKED(UNiagaraMeshRendererProperties, NumFlipbookFrames),
+		GET_MEMBER_NAME_CHECKED(FNiagaraMeshRendererMeshProperties, Mesh),
 	};
 	return RebuildMeshPropertyNames.Contains(Property->GetFName());
 }
