rigid mesh collisions (only static meshes are considered (all in the scene))

diff --git a/Engine/Plugins/FX/Niagara/Shaders/Private/NiagaraDataInterfaceRigidMeshCollisionQuery.ush b/Engine/Plugins/FX/Niagara/Shaders/Private/NiagaraDataInterfaceRigidMeshCollisionQuery.ush
new file mode 100644
--- /dev/null
+++ b/Engine/Plugins/FX/Niagara/Shaders/Private/NiagaraDataInterfaceRigidMeshCollisionQuery.ush
@@ -0,0 +1,392 @@
+// Copyright Epic Games, Inc. All Rights Reserved.
+
+/*=============================================================================
+NiagaraDataInterfaceRigidMeshCollisionQuery.ush 
+=============================================================================*/
+
+#pragma once
+
+/* -----------------------------------------------------------------
+ * Physics Asset constants and context
+ * -----------------------------------------------------------------
+ */
+
+ #define RIGID_MESH_COLLISION_QUERY_MAX_TRANSFORMS 200
+ #define SMALL_NUMBER 1e-8
+
+ #define CURRENT_OFFSET 0
+ #define PREVIOUS_OFFSET RIGID_MESH_COLLISION_QUERY_MAX_TRANSFORMS
+
+#define MAX_DISTANCE 1e+38
+#define BOX_INDEX 0
+#define SPHERE_INDEX 1
+#define CAPSULE_INDEX 2
+
+struct FDIRigidMeshCollisionContext
+{
+	uint4 ElementOffsets;
+	Buffer<float4>		WorldTransformBuffer;
+	Buffer<float4>		InverseTransformBuffer;
+	Buffer<float4>		ElementExtentBuffer;
+	Buffer<uint>		PhysicsTypeBuffer;
+}; 
+
+#define DIRIGIDMESHCOLLISIONQUERY_DECLARE_CONSTANTS(NAME)\
+uint4 ElementOffsets_##NAME;\
+Buffer<float4>		WorldTransformBuffer_##NAME;\
+Buffer<float4>		InverseTransformBuffer_##NAME;\
+Buffer<float4>		ElementExtentBuffer_##NAME;\
+Buffer<uint>		PhysicsTypeBuffer_##NAME;\
+
+
+#define DIRIGIDMESHCOLLISIONQUERY_MAKE_CONTEXT(NAME)\
+FDIRigidMeshCollisionContext DIContext; \
+DIContext.ElementOffsets = ElementOffsets_##NAME;\
+DIContext.WorldTransformBuffer = WorldTransformBuffer_##NAME;\
+DIContext.InverseTransformBuffer = InverseTransformBuffer_##NAME;\
+DIContext.ElementExtentBuffer = ElementExtentBuffer_##NAME;\
+DIContext.PhysicsTypeBuffer = PhysicsTypeBuffer_##NAME;\
+
+
+/* -----------------------------------------------------------------
+ * Element accessors
+ * -----------------------------------------------------------------
+ */
+
+int DIRigidMeshCollision_GetNumBoxes(in FDIRigidMeshCollisionContext DIContext)
+{
+	return DIContext.ElementOffsets[BOX_INDEX+1] - DIContext.ElementOffsets[BOX_INDEX];
+}
+
+int DIRigidMeshCollision_GetNumSpheres(in FDIRigidMeshCollisionContext DIContext)
+{
+	return DIContext.ElementOffsets[SPHERE_INDEX+1] - DIContext.ElementOffsets[SPHERE_INDEX];
+}
+
+int DIRigidMeshCollision_GetNumCapsules(in FDIRigidMeshCollisionContext DIContext)
+{
+	return DIContext.ElementOffsets[CAPSULE_INDEX+1] - DIContext.ElementOffsets[CAPSULE_INDEX];
+}
+
+
+/* -----------------------------------------------------------------
+ * Position / Transform utils
+ * -----------------------------------------------------------------
+ */
+
+float3 DIRigidMeshCollision_GetCurrentPosition(in FDIRigidMeshCollisionContext DIContext, in int ElementIndex)
+{
+	const int BufferOffset = 3 * ElementIndex + CURRENT_OFFSET;
+    return float3(DIContext.WorldTransformBuffer[BufferOffset].w, DIContext.WorldTransformBuffer[BufferOffset+1].w, DIContext.WorldTransformBuffer[BufferOffset+2].w);
+}
+
+float3 DIRigidMeshCollision_GetPreviousPosition(in FDIRigidMeshCollisionContext DIContext, in int ElementIndex)
+{
+	const int BufferOffset = 3 * ElementIndex + PREVIOUS_OFFSET;
+    return float3(DIContext.WorldTransformBuffer[BufferOffset].w, DIContext.WorldTransformBuffer[BufferOffset+1].w, DIContext.WorldTransformBuffer[BufferOffset+2].w);
+}
+
+float3x4 DIRigidMeshCollision_GetCurrentTransform(in FDIRigidMeshCollisionContext DIContext, in int ElementIndex)
+{
+	const int BufferOffset = 3 * ElementIndex + CURRENT_OFFSET;
+    return float3x4(DIContext.WorldTransformBuffer[BufferOffset], DIContext.WorldTransformBuffer[BufferOffset+1], DIContext.WorldTransformBuffer[BufferOffset+2]);
+}
+
+float3x4 DIRigidMeshCollision_GetInverseTransform(in FDIRigidMeshCollisionContext DIContext, in int ElementIndex)
+{
+	const int BufferOffset = 3 * ElementIndex + CURRENT_OFFSET;
+    return float3x4(DIContext.InverseTransformBuffer[BufferOffset], DIContext.InverseTransformBuffer[BufferOffset+1], DIContext.InverseTransformBuffer[BufferOffset+2]);
+}
+
+float3x4 DIRigidMeshCollision_GetPreviousTransform(in FDIRigidMeshCollisionContext DIContext, in int ElementIndex)
+{
+	const int BufferOffset = 3 * ElementIndex + PREVIOUS_OFFSET;
+    return float3x4(DIContext.WorldTransformBuffer[BufferOffset], DIContext.WorldTransformBuffer[BufferOffset+1], DIContext.WorldTransformBuffer[BufferOffset+2]);
+}
+
+float3x4 DIRigidMeshCollision_GetPreviousInverse(in FDIRigidMeshCollisionContext DIContext, in int ElementIndex)
+{
+	const int BufferOffset = 3 * ElementIndex + PREVIOUS_OFFSET;
+    return float3x4(DIContext.InverseTransformBuffer[BufferOffset], DIContext.InverseTransformBuffer[BufferOffset+1], DIContext.InverseTransformBuffer[BufferOffset+2]);
+}
+
+
+/* -----------------------------------------------------------------
+ * Collision detection utils
+ * -----------------------------------------------------------------
+ */
+
+// Given a world space position (NodePosition) compute the sphere closest point (position,normal,velocity)
+float DIRigidMeshCollision_GetSphereProjection(in float3 LocalPosition, in float3 SphereCenter, in float SphereRadius, in int SphereIndex,
+			inout float3 OutClosestPosition, inout float3 OutClosestNormal, inout int OutElementIndex, inout float OutMinDistance)
+{
+	const float3 DeltaPosition = LocalPosition - SphereCenter;
+	const float DeltaLength = length(DeltaPosition);
+	const float SphereDistance = DeltaLength - SphereRadius;
+
+	if (SphereDistance < OutMinDistance)
+	{
+		OutMinDistance = SphereDistance;
+		OutElementIndex = SphereIndex;
+
+		OutClosestNormal = (DeltaLength > SMALL_NUMBER) ? DeltaPosition / DeltaLength : float3(0,0,0);
+		OutClosestPosition = LocalPosition - OutClosestNormal*SphereDistance;
+	}
+	return SphereDistance;
+}
+
+// Given a world space position (NodePosition) compute the sphere closest point (position,normal,velocity)
+float DIRigidMeshCollision_GetCapsuleProjection(in float3 LocalPosition, in float2 CapsuleSize, in int CapsuleIndex,
+			inout float3 OutClosestPosition, inout float3 OutClosestNormal, inout int OutElementIndex, inout float OutMinDistance)
+{
+	const float HalfLength = 0.5 * CapsuleSize.y;
+	const float3 SphereCenter = float3(0,0,clamp(LocalPosition.z, -HalfLength, HalfLength));
+
+	return DIRigidMeshCollision_GetSphereProjection(LocalPosition,SphereCenter,CapsuleSize.x,CapsuleIndex,
+			OutClosestPosition,OutClosestNormal,OutElementIndex,OutMinDistance);
+}
+
+// Given a world space position (NodePosition) compute the sphere closest point (position,normal,velocity)
+float DIRigidMeshCollision_GetBoxProjection(in float3 LocalPosition, in float3 BoxExtent, in int BoxIndex, 
+			inout float3 OutClosestPosition, inout float3 OutClosestNormal, inout int OutElementIndex, inout float OutMinDistance)
+{
+	const float3 HalfExtent = 0.5 * BoxExtent;
+	const float3 DeltaPosition = abs(LocalPosition) - HalfExtent;
+
+	const int ClosestAxis = ((DeltaPosition.x > DeltaPosition.y) && (DeltaPosition.x > DeltaPosition.z)) ? 0 : ( DeltaPosition.y > DeltaPosition.z) ? 1 : 2;
+	const float OutsideDistance = length(max(DeltaPosition,0.0));
+	const float BoxDistance = OutsideDistance + min(DeltaPosition[ClosestAxis],0.0);
+
+	if (BoxDistance < OutMinDistance)
+	{
+		OutMinDistance = BoxDistance;
+		OutElementIndex = BoxIndex;
+
+		if (BoxDistance <= 0)
+		{
+			const bool NegativeSide = LocalPosition[ClosestAxis] < 0.0;
+
+			OutClosestPosition = LocalPosition;
+			OutClosestNormal = float3(0,0,0);
+			if( ClosestAxis == 0) 
+			{
+				OutClosestPosition.x = NegativeSide ? -HalfExtent.x : HalfExtent.x;
+				OutClosestNormal.x = NegativeSide ? -1.0 : 1.0;
+			}
+			else if( ClosestAxis == 1) 
+			{
+				OutClosestPosition.y = NegativeSide ? -HalfExtent.y : HalfExtent.y;
+				OutClosestNormal.y = NegativeSide ? -1.0 : 1.0;
+			}
+			else if( ClosestAxis == 2) 
+			{
+				OutClosestPosition.z = NegativeSide ? -HalfExtent.z : HalfExtent.z;
+				OutClosestNormal.z = NegativeSide ? -1.0 : 1.0;
+			}
+		}
+		else
+		{
+			OutClosestPosition = clamp(LocalPosition,-HalfExtent,HalfExtent);
+			OutClosestNormal = (LocalPosition - OutClosestPosition) / OutsideDistance;
+		}
+	}
+	return BoxDistance;
+}
+
+
+/* -----------------------------------------------------------------
+ * Get the closest element to the node position
+ * -----------------------------------------------------------------
+ */
+
+float3 DIRigidMeshCollision_GetLocalPosition(in FDIRigidMeshCollisionContext DIContext, in float3 NodePosition, in int ElementIndex, in float TimeFraction )
+{
+	const float3 CurrentLocal = mul(DIRigidMeshCollision_GetInverseTransform(DIContext,ElementIndex), float4(NodePosition,1.0)).xyz;
+	const float3 PreviousLocal = mul(DIRigidMeshCollision_GetPreviousInverse(DIContext,ElementIndex), float4(NodePosition,1.0)).xyz;
+
+	return PreviousLocal + TimeFraction * (CurrentLocal-PreviousLocal);
+}
+
+// Given a world space position (NodePosition) compute the sphere closest point (position,normal,velocity)
+int DIRigidMeshCollision_GetClosestElement(in FDIRigidMeshCollisionContext DIContext, in float3 NodePosition, out float3 OutClosestPosition, 
+										  out float3 OutClosestNormal, out float OutMinDistance, in float TimeFraction)
+{	
+	float MinDistance = MAX_DISTANCE;
+	int ElementIndex = -1;
+	float3 CollisionPosition = float3(0,0,0);
+	float3 CollisionNormal = float3(0,0,0);
+
+	const int SpheresBegin = DIContext.ElementOffsets[SPHERE_INDEX];
+	const int SpheresEnd = DIContext.ElementOffsets[SPHERE_INDEX+1];
+	for (int SphereIndex = SpheresBegin; SphereIndex < SpheresEnd; ++SphereIndex)
+	{
+		const float3 LocalPosition = DIRigidMeshCollision_GetLocalPosition(DIContext,NodePosition,SphereIndex,TimeFraction);
+
+		DIRigidMeshCollision_GetSphereProjection(LocalPosition, float3(0,0,0), DIContext.ElementExtentBuffer[SphereIndex].x, SphereIndex,
+				CollisionPosition, CollisionNormal, ElementIndex, MinDistance);
+	}
+
+	const int CapsulesBegin = DIContext.ElementOffsets[CAPSULE_INDEX];
+	const int CapsulesEnd = DIContext.ElementOffsets[CAPSULE_INDEX+1];
+	for (int CapsuleIndex = CapsulesBegin; CapsuleIndex < CapsulesEnd; ++CapsuleIndex)
+	{
+		const float3 LocalPosition = DIRigidMeshCollision_GetLocalPosition(DIContext,NodePosition,CapsuleIndex,TimeFraction);
+
+		DIRigidMeshCollision_GetCapsuleProjection(LocalPosition, DIContext.ElementExtentBuffer[CapsuleIndex].xy, CapsuleIndex,
+				CollisionPosition, CollisionNormal, ElementIndex, MinDistance);
+	}
+
+	const int BoxesBegin = DIContext.ElementOffsets[BOX_INDEX];
+	const int BoxesEnd = DIContext.ElementOffsets[BOX_INDEX+1];
+	for (int BoxIndex = BoxesBegin; BoxIndex < BoxesEnd; ++BoxIndex)
+	{
+		const float3 LocalPosition = DIRigidMeshCollision_GetLocalPosition(DIContext,NodePosition,BoxIndex,TimeFraction);
+
+		DIRigidMeshCollision_GetBoxProjection(LocalPosition, DIContext.ElementExtentBuffer[BoxIndex].xyz, BoxIndex,
+				CollisionPosition, CollisionNormal, ElementIndex, MinDistance);
+	}
+	OutClosestPosition = CollisionPosition;
+	OutClosestNormal = CollisionNormal;
+	OutMinDistance = MinDistance;
+	return ElementIndex;
+}
+
+/* -----------------------------------------------------------------
+ * Get the closest point to the physics asset 
+ * -----------------------------------------------------------------
+ */
+
+// Given a world space position (NodePosition) and an element index compute the physics asset closest point (position,normal,velocity)
+void DIRigidMeshCollision_GetElementPoint(in FDIRigidMeshCollisionContext DIContext, in float3 NodePosition, in float DeltaTime, in float TimeFraction, in int ElementIndex, out float3 OutClosestPosition, 
+										  out float3 OutClosestNormal, out float3 OutClosestVelocity)
+{	
+	const float3 LocalPosition = DIRigidMeshCollision_GetLocalPosition(DIContext,NodePosition,ElementIndex,TimeFraction);
+
+	float MinDistance = MAX_DISTANCE;
+	int OutputIndex = -1;
+	float3 CollisionPosition = float3(0,0,0);
+	float3 CollisionNormal = float3(0,0,0);
+
+	if( ElementIndex >= DIContext.ElementOffsets[SPHERE_INDEX] && ElementIndex < DIContext.ElementOffsets[SPHERE_INDEX+1] )
+	{
+		DIRigidMeshCollision_GetSphereProjection(LocalPosition, float3(0,0,0), DIContext.ElementExtentBuffer[ElementIndex].x, ElementIndex,
+				CollisionPosition, CollisionNormal, OutputIndex, MinDistance);
+	}
+	else if ( ElementIndex >= DIContext.ElementOffsets[CAPSULE_INDEX] && ElementIndex < DIContext.ElementOffsets[CAPSULE_INDEX+1] )
+	{
+		DIRigidMeshCollision_GetCapsuleProjection(LocalPosition, DIContext.ElementExtentBuffer[ElementIndex].xy, ElementIndex,
+				CollisionPosition, CollisionNormal, OutputIndex, MinDistance);
+	}
+	else if ( ElementIndex >= DIContext.ElementOffsets[BOX_INDEX] && ElementIndex < DIContext.ElementOffsets[BOX_INDEX+1] )
+	{
+		DIRigidMeshCollision_GetBoxProjection(LocalPosition, DIContext.ElementExtentBuffer[ElementIndex].xyz, ElementIndex,
+				CollisionPosition, CollisionNormal, OutputIndex, MinDistance);
+	}
+
+	OutClosestNormal = float3(0,0,0);
+	OutClosestPosition = float3(0,0,0);
+	OutClosestVelocity = float3(0,0,0);
+
+	if (ElementIndex != -1)
+	{
+		const float3 PreviousPosition = mul(DIRigidMeshCollision_GetPreviousTransform(DIContext,ElementIndex), float4(CollisionPosition,1.0)).xyz;
+		const float3 CurrentPosition = mul(DIRigidMeshCollision_GetCurrentTransform(DIContext,ElementIndex), float4(CollisionPosition,1.0)).xyz;
+
+		const float3 PreviousNormal = mul(DIRigidMeshCollision_GetPreviousTransform(DIContext,ElementIndex), float4(CollisionNormal,0.0)).xyz;
+		const float3 CurrentNormal = mul(DIRigidMeshCollision_GetCurrentTransform(DIContext,ElementIndex), float4(CollisionNormal,0.0)).xyz;
+
+		OutClosestPosition = PreviousPosition + TimeFraction * (CurrentPosition-PreviousPosition);
+		OutClosestNormal = normalize(PreviousNormal + TimeFraction * (CurrentNormal-PreviousNormal));
+
+		OutClosestVelocity = ( CurrentPosition - PreviousPosition ) / DeltaTime;
+	}
+}
+
+// Given a world space position (NodePosition) and an element index compute the physics asset closest distance
+void DIRigidMeshCollision_GetElementDistance(in FDIRigidMeshCollisionContext DIContext, in float3 NodePosition, in float TimeFraction, in int ElementIndex, out float OutClosestDistance)
+{	
+	const float3 LocalPosition = DIRigidMeshCollision_GetLocalPosition(DIContext,NodePosition,ElementIndex,TimeFraction);
+
+	float MinDistance = MAX_DISTANCE;
+	int OutputIndex = -1;
+	float3 CollisionPosition = float3(0,0,0);
+	float3 CollisionNormal = float3(0,0,0);
+
+	if( ElementIndex >= DIContext.ElementOffsets[SPHERE_INDEX] && ElementIndex < DIContext.ElementOffsets[SPHERE_INDEX+1] )
+	{
+		DIRigidMeshCollision_GetSphereProjection(LocalPosition, float3(0,0,0), DIContext.ElementExtentBuffer[ElementIndex].x, ElementIndex,
+				CollisionPosition, CollisionNormal, OutputIndex, MinDistance);
+	}
+	else if ( ElementIndex >= DIContext.ElementOffsets[CAPSULE_INDEX] && ElementIndex < DIContext.ElementOffsets[CAPSULE_INDEX+1] )
+	{
+		DIRigidMeshCollision_GetCapsuleProjection(LocalPosition, DIContext.ElementExtentBuffer[ElementIndex].xy, ElementIndex,
+				CollisionPosition, CollisionNormal, OutputIndex, MinDistance);
+	}
+	else if ( ElementIndex >= DIContext.ElementOffsets[BOX_INDEX] && ElementIndex < DIContext.ElementOffsets[BOX_INDEX+1] )
+	{
+		DIRigidMeshCollision_GetBoxProjection(LocalPosition, DIContext.ElementExtentBuffer[ElementIndex].xyz, ElementIndex,
+				CollisionPosition, CollisionNormal, OutputIndex, MinDistance);
+	}
+	
+	OutClosestDistance = 0.0;
+
+	if (ElementIndex != -1)
+	{
+		OutClosestDistance = (DIContext.PhysicsTypeBuffer[ElementIndex] == 1) ? MinDistance : 0;
+	}
+}
+
+ // Given a world space position (NodePosition) compute the physics asset closest element
+void DIRigidMeshCollision_GetClosestElement(in FDIRigidMeshCollisionContext DIContext, in float3 NodePosition, in float TimeFraction, out int OutElementIndex)
+{	
+	float3 CollisionPosition = float3(0,0,0);
+	float3 CollisionNormal = float3(0,0,0);
+	float MinDistance = 0.0;
+
+	OutElementIndex = DIRigidMeshCollision_GetClosestElement(DIContext,NodePosition,CollisionPosition,CollisionNormal,MinDistance,TimeFraction);
+}
+
+// Given a world space position (NodePosition) compute the physics asset closest point (position,normal,velocity)
+void DIRigidMeshCollision_GetClosestPoint(in FDIRigidMeshCollisionContext DIContext, in float3 NodePosition, in float DeltaTime, in float TimeFraction, out float OutClosestDistance, out float3 OutClosestPosition, out float3 OutClosestNormal, out float3 OutClosestVelocity)
+{	
+	float3 CollisionPosition = float3(0,0,0);
+	float3 CollisionNormal = float3(0,0,0);
+	OutClosestDistance = 0.0;
+
+	const int ElementIndex = DIRigidMeshCollision_GetClosestElement(DIContext,NodePosition,CollisionPosition,CollisionNormal,OutClosestDistance,TimeFraction);
+
+	OutClosestNormal = float3(0,0,0);
+	OutClosestPosition = float3(0,0,0);
+	OutClosestVelocity = float3(0,0,0);
+
+	if (ElementIndex != -1)
+	{
+		const float3 PreviousPosition = mul(DIRigidMeshCollision_GetPreviousTransform(DIContext,ElementIndex), float4(CollisionPosition,1.0)).xyz;
+		const float3 CurrentPosition = mul(DIRigidMeshCollision_GetCurrentTransform(DIContext,ElementIndex), float4(CollisionPosition,1.0)).xyz;
+
+		const float3 PreviousNormal = mul(DIRigidMeshCollision_GetPreviousTransform(DIContext,ElementIndex), float4(CollisionNormal,0.0)).xyz;
+		const float3 CurrentNormal = mul(DIRigidMeshCollision_GetCurrentTransform(DIContext,ElementIndex), float4(CollisionNormal,0.0)).xyz;
+
+		OutClosestPosition = PreviousPosition + TimeFraction * (CurrentPosition-PreviousPosition);
+		OutClosestNormal = normalize(PreviousNormal + TimeFraction * (CurrentNormal-PreviousNormal));
+
+		OutClosestVelocity = ( CurrentPosition - PreviousPosition ) / DeltaTime;
+	}
+}
+
+// Given a world space position (NodePosition) compute the physics asset closest distance
+void DIRigidMeshCollision_GetClosestDistance(in FDIRigidMeshCollisionContext DIContext, in float3 NodePosition, in float TimeFraction, out float OutClosestDistance)
+{	
+	float3 CollisionPosition = float3(0,0,0);
+	float3 CollisionNormal = float3(0,0,0);
+	float MinDistance = 0.0;
+
+	const int ElementIndex = DIRigidMeshCollision_GetClosestElement(DIContext,NodePosition,CollisionPosition,CollisionNormal,MinDistance,TimeFraction);
+
+	OutClosestDistance = 0.0;
+
+	if (ElementIndex != -1)
+	{
+		OutClosestDistance = (DIContext.PhysicsTypeBuffer[ElementIndex] == 1) ? MinDistance : 0;
+	}
+}
\ No newline at end of file
diff --git a/Engine/Plugins/FX/Niagara/Source/Niagara/Private/NiagaraDataInterfaceRigidMeshCollisionQuery.cpp b/Engine/Plugins/FX/Niagara/Source/Niagara/Private/NiagaraDataInterfaceRigidMeshCollisionQuery.cpp
new file mode 100644
--- /dev/null
+++ b/Engine/Plugins/FX/Niagara/Source/Niagara/Private/NiagaraDataInterfaceRigidMeshCollisionQuery.cpp
@@ -0,0 +1,922 @@
+// Copyright Epic Games, Inc. All Rights Reserved.
+
+#include "NiagaraDataInterfaceRigidMeshCollisionQuery.h"
+#include "Animation/SkeletalMeshActor.h"
+#include "Components/SkeletalMeshComponent.h"
+#include "SkeletalRenderPublic.h"
+#include "SkeletalMeshTypes.h"
+#include "AnimationRuntime.h"
+#include "NiagaraShader.h"
+#include "NiagaraComponent.h"
+#include "NiagaraRenderer.h"
+#include "NiagaraSystemInstance.h"
+#include "NiagaraEmitterInstanceBatcher.h"
+#include "ShaderParameterUtils.h"
+#include "EngineUtils.h"
+#include "Engine/StaticMeshActor.h"
+
+#define LOCTEXT_NAMESPACE "NiagaraDataInterfaceRigidMeshCollisionQuery"
+DEFINE_LOG_CATEGORY_STATIC(LogRigidMeshCollision, Log, All);
+
+//------------------------------------------------------------------------------------------------------------
+
+static const FName GetNumBoxesName(TEXT("GetNumBoxes"));
+static const FName GetNumSpheresName(TEXT("GetNumSpheres"));
+static const FName GetNumCapsulesName(TEXT("GetNumCapsules"));
+
+//------------------------------------------------------------------------------------------------------------
+
+static const FName GetClosestElementName(TEXT("GetClosestElement"));
+static const FName GetElementPointName(TEXT("GetElementPoint"));
+static const FName GetElementDistanceName(TEXT("GetElementDistance"));
+static const FName GetClosestPointName(TEXT("GetClosestPoint"));
+static const FName GetClosestDistanceName(TEXT("GetClosestDistance"));
+
+//------------------------------------------------------------------------------------------------------------
+
+const FString UNiagaraDataInterfaceRigidMeshCollisionQuery::ElementOffsetsName(TEXT("ElementOffsets_"));
+
+const FString UNiagaraDataInterfaceRigidMeshCollisionQuery::WorldTransformBufferName(TEXT("WorldTransformBuffer_"));
+const FString UNiagaraDataInterfaceRigidMeshCollisionQuery::InverseTransformBufferName(TEXT("InverseTransformBuffer_"));
+const FString UNiagaraDataInterfaceRigidMeshCollisionQuery::ElementExtentBufferName(TEXT("ElementExtentBuffer_"));
+const FString UNiagaraDataInterfaceRigidMeshCollisionQuery::PhysicsTypeBufferName(TEXT("PhysicsTypeBuffer_"));
+
+//------------------------------------------------------------------------------------------------------------
+
+struct FNDIRigidMeshCollisionParametersName
+{
+	FNDIRigidMeshCollisionParametersName(const FString& Suffix)
+	{
+		ElementOffsetsName = UNiagaraDataInterfaceRigidMeshCollisionQuery::ElementOffsetsName + Suffix;
+
+		WorldTransformBufferName = UNiagaraDataInterfaceRigidMeshCollisionQuery::WorldTransformBufferName + Suffix;
+		InverseTransformBufferName = UNiagaraDataInterfaceRigidMeshCollisionQuery::InverseTransformBufferName + Suffix;
+		ElementExtentBufferName = UNiagaraDataInterfaceRigidMeshCollisionQuery::ElementExtentBufferName + Suffix;
+		PhysicsTypeBufferName = UNiagaraDataInterfaceRigidMeshCollisionQuery::PhysicsTypeBufferName + Suffix;		
+	}
+
+	FString ElementOffsetsName;
+
+	FString WorldTransformBufferName;
+	FString InverseTransformBufferName;
+	FString ElementExtentBufferName;
+	FString PhysicsTypeBufferName;
+};
+
+//------------------------------------------------------------------------------------------------------------
+
+template<typename BufferType, EPixelFormat PixelFormat, uint32 ElementCount, uint32 BufferCount = 1>
+void CreateInternalBuffer(FRWBuffer& OutputBuffer)
+{
+	if (ElementCount > 0)
+	{
+		OutputBuffer.Initialize(TEXT("FNDIRigidMeshCollisionBuffer"), sizeof(BufferType), ElementCount * BufferCount, PixelFormat, BUF_Static);
+	}
+}
+
+template<typename BufferType, EPixelFormat PixelFormat, uint32 ElementCount, uint32 BufferCount = 1>
+void UpdateInternalBuffer(const TStaticArray<BufferType,ElementCount*BufferCount>& InputData, FRWBuffer& OutputBuffer)
+{
+	if (ElementCount > 0 && OutputBuffer.Buffer.IsValid())
+	{
+		const uint32 BufferBytes = sizeof(BufferType) * ElementCount * BufferCount;
+
+		void* OutputData = RHILockBuffer(OutputBuffer.Buffer, 0, BufferBytes, RLM_WriteOnly);
+
+		FMemory::Memcpy(OutputData, InputData.GetData(), BufferBytes);
+		RHIUnlockBuffer(OutputBuffer.Buffer);
+	}
+}
+
+void FillCurrentTransforms(const FTransform& ElementTransform, uint32& ElementCount,
+	TStaticArray<FVector4,RIGID_MESH_COLLISION_QUERY_MAX_TRANSFORMS>& OutCurrentTransform, TStaticArray<FVector4, RIGID_MESH_COLLISION_QUERY_MAX_TRANSFORMS>& OutCurrentInverse)
+{
+	const uint32 ElementOffset = 3 * ElementCount;
+	const FMatrix44f ElementMatrix = ElementTransform.ToMatrixWithScale();
+	const FMatrix44f ElementInverse = ElementMatrix.Inverse();
+
+	ElementMatrix.To3x4MatrixTranspose(&OutCurrentTransform[ElementOffset].X);
+	ElementInverse.To3x4MatrixTranspose(&OutCurrentInverse[ElementOffset].X);
+	++ElementCount;
+}
+
+void GetNumPrimitives(TArray<AStaticMeshActor*> StaticMeshActors, uint32& NumBoxes, uint32& NumSpheres, uint32& NumCapsules)
+{
+	NumBoxes = 0;
+	NumSpheres = 0;
+	NumCapsules = 0;
+
+	for (int32 ActorIndex = 0; ActorIndex < StaticMeshActors.Num(); ++ActorIndex)
+	{
+		AStaticMeshActor* StaticMeshActor = StaticMeshActors[ActorIndex];
+		UStaticMeshComponent* StaticMeshComponent = StaticMeshActor != nullptr ? StaticMeshActor->GetStaticMeshComponent() : nullptr;
+		UBodySetup* BodySetup = StaticMeshComponent != nullptr ? StaticMeshComponent->GetBodySetup() : nullptr;
+		if (BodySetup != nullptr)
+		{
+			// #todo(dmp): save static mesh component reference
+
+			for (const FKConvexElem& ConvexElem : BodySetup->AggGeom.ConvexElems)
+			{
+				if (CollisionEnabledHasPhysics(ConvexElem.GetCollisionEnabled()))
+				{
+					UE_LOG(LogRigidMeshCollision, Warning, TEXT("Convex collision objects encountered and will be skipped"));
+				}
+			}
+			for (const FKBoxElem& BoxElem : BodySetup->AggGeom.BoxElems)
+			{
+				if (CollisionEnabledHasPhysics(BoxElem.GetCollisionEnabled()))
+				{
+					NumBoxes += 1;
+				}
+			}
+			for (const FKSphereElem& SphereElem : BodySetup->AggGeom.SphereElems)
+			{
+				if (CollisionEnabledHasPhysics(SphereElem.GetCollisionEnabled()))
+				{
+					NumSpheres += 1;
+				}
+			}
+			for (const FKSphylElem& CapsuleElem : BodySetup->AggGeom.SphylElems)
+			{
+				if (CollisionEnabledHasPhysics(CapsuleElem.GetCollisionEnabled()))
+				{
+					NumCapsules += 1;
+				}
+			}				
+		}
+	}
+}
+
+void CompactInternalArrays(FNDIRigidMeshCollisionArrays* OutAssetArrays)
+{
+	for (uint32 TransformIndex = 0; TransformIndex < RIGID_MESH_COLLISION_QUERY_MAX_TRANSFORMS; ++TransformIndex)
+	{
+		uint32 OffsetIndex = TransformIndex;
+		OutAssetArrays->WorldTransform[OffsetIndex] = OutAssetArrays->CurrentTransform[TransformIndex];
+		OutAssetArrays->InverseTransform[OffsetIndex] = OutAssetArrays->CurrentInverse[TransformIndex];
+
+		OffsetIndex += RIGID_MESH_COLLISION_QUERY_MAX_TRANSFORMS;
+		OutAssetArrays->WorldTransform[OffsetIndex] = OutAssetArrays->PreviousTransform[TransformIndex];
+		OutAssetArrays->InverseTransform[OffsetIndex] = OutAssetArrays->PreviousInverse[TransformIndex];		
+	}
+}
+
+void CreateInternalArrays(TArray<AStaticMeshActor*> StaticMeshActors,
+	FNDIRigidMeshCollisionArrays* OutAssetArrays)
+{
+	if (OutAssetArrays != nullptr)
+	{
+		OutAssetArrays->ElementOffsets.BoxOffset = 0;
+		OutAssetArrays->ElementOffsets.SphereOffset = 0;
+		OutAssetArrays->ElementOffsets.CapsuleOffset = 0;
+		OutAssetArrays->ElementOffsets.NumElements = 0;
+
+		uint32 NumBoxes = 0;
+		uint32 NumSpheres = 0;
+		uint32 NumCapsules = 0;
+
+		GetNumPrimitives(StaticMeshActors, NumBoxes, NumSpheres, NumCapsules);
+		
+		if ((NumBoxes + NumSpheres + NumCapsules) < RIGID_MESH_COLLISION_QUERY_MAX_PRIMITIVES)
+		{
+			OutAssetArrays->ElementOffsets.BoxOffset = 0;
+			OutAssetArrays->ElementOffsets.SphereOffset = OutAssetArrays->ElementOffsets.BoxOffset + NumBoxes;
+			OutAssetArrays->ElementOffsets.CapsuleOffset = OutAssetArrays->ElementOffsets.SphereOffset + NumSpheres;
+			OutAssetArrays->ElementOffsets.NumElements = OutAssetArrays->ElementOffsets.CapsuleOffset + NumCapsules;
+
+			const uint32 NumTransforms = OutAssetArrays->ElementOffsets.NumElements * 2;
+			const uint32 NumExtents = OutAssetArrays->ElementOffsets.NumElements;
+
+			uint32 BoxCount = OutAssetArrays->ElementOffsets.BoxOffset;
+			uint32 SphereCount = OutAssetArrays->ElementOffsets.SphereOffset;
+			uint32 CapsuleCount = OutAssetArrays->ElementOffsets.CapsuleOffset;
+
+			for (int32 ActorIndex = 0; ActorIndex < StaticMeshActors.Num(); ++ActorIndex)
+			{				
+				AStaticMeshActor* StaticMeshActor = StaticMeshActors[ActorIndex];
+				UStaticMeshComponent* StaticMeshComponent = StaticMeshActor != nullptr ? StaticMeshActor->GetStaticMeshComponent() : nullptr;
+				UBodySetup* BodySetup = StaticMeshComponent != nullptr ? StaticMeshComponent->GetBodySetup() : nullptr;
+
+				if (BodySetup != nullptr)
+				{										
+					const FTransform MeshTransform = StaticMeshActor->GetTransform();				
+					for (const FKBoxElem& BoxElem : BodySetup->AggGeom.BoxElems)
+					{
+						if (CollisionEnabledHasPhysics(BoxElem.GetCollisionEnabled()))
+						{							
+							OutAssetArrays->PhysicsType[BoxCount] = (BoxElem.GetCollisionEnabled() == ECollisionEnabled::QueryAndPhysics);
+
+							const FTransform ElementTransform = FTransform(BoxElem.Rotation, BoxElem.Center) * MeshTransform;
+							OutAssetArrays->ElementExtent[BoxCount] = FVector4(BoxElem.X, BoxElem.Y, BoxElem.Z, 0);
+							FillCurrentTransforms(ElementTransform, BoxCount, OutAssetArrays->CurrentTransform, OutAssetArrays->CurrentInverse);
+							//--BoxCount;
+						}
+					}
+
+					for (const FKSphereElem& SphereElem : BodySetup->AggGeom.SphereElems)
+					{
+						if (CollisionEnabledHasPhysics(SphereElem.GetCollisionEnabled()))
+						{														
+							OutAssetArrays->PhysicsType[SphereCount] = (SphereElem.GetCollisionEnabled() == ECollisionEnabled::QueryAndPhysics);
+
+							const FTransform ElementTransform = FTransform(SphereElem.Center) * MeshTransform;
+							OutAssetArrays->ElementExtent[SphereCount] = FVector4(SphereElem.Radius, 0, 0, 0);
+							FillCurrentTransforms(ElementTransform, SphereCount, OutAssetArrays->CurrentTransform, OutAssetArrays->CurrentInverse);
+							//--SphereCount;
+						}
+					}
+
+					for (const FKSphylElem& CapsuleElem : BodySetup->AggGeom.SphylElems)
+					{
+						if (CollisionEnabledHasPhysics(CapsuleElem.GetCollisionEnabled()))
+						{														
+							OutAssetArrays->PhysicsType[CapsuleCount] = (CapsuleElem.GetCollisionEnabled() == ECollisionEnabled::QueryAndPhysics);
+
+							const FTransform ElementTransform = FTransform(CapsuleElem.Rotation, CapsuleElem.Center) * MeshTransform;
+							OutAssetArrays->ElementExtent[CapsuleCount] = FVector4(CapsuleElem.Radius, CapsuleElem.Length, 0, 0);
+							FillCurrentTransforms(ElementTransform, CapsuleCount, OutAssetArrays->CurrentTransform, OutAssetArrays->CurrentInverse);
+							//--CapsuleCount;
+						}
+					}					
+				}
+			}
+			OutAssetArrays->PreviousTransform = OutAssetArrays->CurrentTransform;
+			OutAssetArrays->PreviousInverse = OutAssetArrays->CurrentInverse;
+
+			CompactInternalArrays(OutAssetArrays);
+		}
+		else
+		{
+			UE_LOG(LogRigidMeshCollision, Warning, TEXT("Number of Collision DI primitives is higher than the niagara %d limit"), RIGID_MESH_COLLISION_QUERY_MAX_PRIMITIVES);
+		}
+	}
+}
+
+void UpdateInternalArrays(TArray<AStaticMeshActor*> StaticMeshActors,
+	FNDIRigidMeshCollisionArrays* OutAssetArrays)
+{
+	if (OutAssetArrays != nullptr && OutAssetArrays->ElementOffsets.NumElements < RIGID_MESH_COLLISION_QUERY_MAX_PRIMITIVES)
+	{
+		uint32 NumBoxes = 0;
+		uint32 NumSpheres = 0;
+		uint32 NumCapsules = 0;
+
+		GetNumPrimitives(StaticMeshActors, NumBoxes, NumSpheres, NumCapsules);
+
+		if (((OutAssetArrays->ElementOffsets.SphereOffset - OutAssetArrays->ElementOffsets.BoxOffset) != NumBoxes) || 
+			((OutAssetArrays->ElementOffsets.CapsuleOffset - OutAssetArrays->ElementOffsets.SphereOffset) != NumSpheres) ||
+			((OutAssetArrays->ElementOffsets.NumElements - OutAssetArrays->ElementOffsets.CapsuleOffset) != NumCapsules))
+		{
+			CreateInternalArrays(StaticMeshActors, OutAssetArrays);
+		}
+
+		OutAssetArrays->PreviousTransform = OutAssetArrays->CurrentTransform;
+		OutAssetArrays->PreviousInverse = OutAssetArrays->CurrentInverse;
+
+		uint32 BoxCount = OutAssetArrays->ElementOffsets.BoxOffset;
+		uint32 SphereCount = OutAssetArrays->ElementOffsets.SphereOffset;
+		uint32 CapsuleCount = OutAssetArrays->ElementOffsets.CapsuleOffset;
+
+		for (int32 ActorIndex = 0; ActorIndex < StaticMeshActors.Num(); ++ActorIndex)
+		{
+			AStaticMeshActor* StaticMeshActor = StaticMeshActors[ActorIndex];
+			UStaticMeshComponent* StaticMeshComponent = StaticMeshActor != nullptr ? StaticMeshActor->GetStaticMeshComponent() : nullptr;
+			UBodySetup* BodySetup = StaticMeshComponent != nullptr ? StaticMeshComponent->GetBodySetup() : nullptr;
+			if (BodySetup != nullptr)
+			{				
+				const FTransform MeshTransform = StaticMeshActor->GetTransform();
+
+				for (const FKBoxElem& BoxElem : BodySetup->AggGeom.BoxElems)
+				{
+					if (CollisionEnabledHasPhysics(BoxElem.GetCollisionEnabled()))
+					{
+						const FTransform ElementTransform = FTransform(BoxElem.Rotation, BoxElem.Center) * MeshTransform;
+						FillCurrentTransforms(ElementTransform, BoxCount, OutAssetArrays->CurrentTransform, OutAssetArrays->CurrentInverse);
+					}
+				}
+
+				for (const FKSphereElem& SphereElem : BodySetup->AggGeom.SphereElems)
+				{
+					if (CollisionEnabledHasPhysics(SphereElem.GetCollisionEnabled()))
+					{
+						const FTransform ElementTransform = FTransform(SphereElem.Center) * MeshTransform;
+						FillCurrentTransforms(ElementTransform, SphereCount, OutAssetArrays->CurrentTransform, OutAssetArrays->CurrentInverse);
+					}
+				}
+
+				for (const FKSphylElem& CapsuleElem : BodySetup->AggGeom.SphylElems)
+				{
+					if (CollisionEnabledHasPhysics(CapsuleElem.GetCollisionEnabled()))
+					{
+						const FTransform ElementTransform = FTransform(CapsuleElem.Rotation, CapsuleElem.Center) * MeshTransform;
+						FillCurrentTransforms(ElementTransform, CapsuleCount, OutAssetArrays->CurrentTransform, OutAssetArrays->CurrentInverse);
+					}
+				}
+			}
+		}
+		CompactInternalArrays(OutAssetArrays);
+	}
+}
+
+//------------------------------------------------------------------------------------------------------------
+
+void FNDIRigidMeshCollisionBuffer::InitRHI()
+{
+	CreateInternalBuffer<FVector4, EPixelFormat::PF_A32B32G32R32F, RIGID_MESH_COLLISION_QUERY_MAX_TRANSFORMS, 2>(WorldTransformBuffer);
+	CreateInternalBuffer<FVector4, EPixelFormat::PF_A32B32G32R32F, RIGID_MESH_COLLISION_QUERY_MAX_TRANSFORMS, 2>(InverseTransformBuffer);
+
+	CreateInternalBuffer<FVector4, EPixelFormat::PF_A32B32G32R32F, RIGID_MESH_COLLISION_QUERY_MAX_PRIMITIVES>(ElementExtentBuffer);
+	CreateInternalBuffer<uint32, EPixelFormat::PF_R32_UINT, RIGID_MESH_COLLISION_QUERY_MAX_PRIMITIVES>(PhysicsTypeBuffer);
+}
+
+void FNDIRigidMeshCollisionBuffer::ReleaseRHI()
+{
+	WorldTransformBuffer.Release();
+	InverseTransformBuffer.Release();
+	ElementExtentBuffer.Release();
+	PhysicsTypeBuffer.Release();
+}
+
+//------------------------------------------------------------------------------------------------------------
+
+ETickingGroup ComputeTickingGroup(const TArray<TWeakObjectPtr<class USkeletalMeshComponent>> SkeletalMeshes)
+{
+	ETickingGroup TickingGroup = NiagaraFirstTickGroup;
+	for (int32 ComponentIndex = 0; ComponentIndex < SkeletalMeshes.Num(); ++ComponentIndex)
+	{
+		if (SkeletalMeshes[ComponentIndex].Get() != nullptr)
+		{
+			const USkeletalMeshComponent* Component = Cast<USkeletalMeshComponent>(SkeletalMeshes[ComponentIndex].Get());
+
+			const ETickingGroup ComponentTickGroup = FMath::Max(Component->PrimaryComponentTick.TickGroup, Component->PrimaryComponentTick.EndTickGroup);
+			const ETickingGroup PhysicsTickGroup = Component->bBlendPhysics ? FMath::Max(ComponentTickGroup, TG_EndPhysics) : ComponentTickGroup;
+			const ETickingGroup ClampedTickGroup = FMath::Clamp(ETickingGroup(PhysicsTickGroup + 1), NiagaraFirstTickGroup, NiagaraLastTickGroup);
+
+			TickingGroup = FMath::Max(TickingGroup, ClampedTickGroup);
+		}
+	}
+	return TickingGroup;
+}
+
+void FNDIRigidMeshCollisionData::Release()
+{
+	if (AssetBuffer)
+	{
+		BeginReleaseResource(AssetBuffer);
+		ENQUEUE_RENDER_COMMAND(DeleteResource)(
+			[ParamPointerToRelease = AssetBuffer](FRHICommandListImmediate& RHICmdList)
+			{
+				delete ParamPointerToRelease;
+			});
+		AssetBuffer = nullptr;
+	}
+}
+
+void FNDIRigidMeshCollisionData::Init(UNiagaraDataInterfaceRigidMeshCollisionQuery* Interface, FNiagaraSystemInstance* SystemInstance)
+{
+	AssetBuffer = nullptr;
+
+	UWorld* World = SystemInstance->GetWorld();		
+
+	StaticMeshActors.Empty();
+	for (TActorIterator<AStaticMeshActor> It(World); It; ++It)
+	{
+		AStaticMeshActor* StaticMeshActor = *It;		
+		StaticMeshActors.Add(StaticMeshActor);
+	
+		// Only deal with static meshes that are movable
+		// #todo(dmp): also nonmoveable?
+		
+		// Get local bounds
+		// FVector MinLocalBounds;
+		// FVector MaxLocalBounds;
+		// StaticMeshActor->GetStaticMeshComponent()->GetLocalBounds(MinLocalBounds, MaxLocalBounds);
+
+		// Initial Transform 
+		//FTransform InitialTransform = StaticMeshActor->GetTransform();
+		
+		// Use Body setup if it exists
+		//UBodySetup *BodySetup = StaticMeshComponent->GetBodySetup();
+	}
+
+	if (Interface != nullptr && SystemInstance != nullptr)
+	{	
+		if (0 < StaticMeshActors.Num() && StaticMeshActors[0] != nullptr)
+		{
+			CreateInternalArrays(StaticMeshActors, &AssetArrays);
+		}
+
+		AssetBuffer = new FNDIRigidMeshCollisionBuffer();
+		BeginInitResource(AssetBuffer);
+	}
+}
+
+void FNDIRigidMeshCollisionData::Update(UNiagaraDataInterfaceRigidMeshCollisionQuery* Interface, FNiagaraSystemInstance* SystemInstance)
+{
+	if (Interface != nullptr && SystemInstance != nullptr)
+	{		
+		//TickingGroup = ComputeTickingGroup(Interface->SourceComponents);
+
+		if (0 < StaticMeshActors.Num() && StaticMeshActors[0] != nullptr)
+		{
+			UpdateInternalArrays(StaticMeshActors, &AssetArrays);
+		}
+	}
+}
+
+//------------------------------------------------------------------------------------------------------------
+
+struct FNDIRigidMeshCollisionParametersCS : public FNiagaraDataInterfaceParametersCS
+{
+	DECLARE_TYPE_LAYOUT(FNDIRigidMeshCollisionParametersCS, NonVirtual);
+public:
+	void Bind(const FNiagaraDataInterfaceGPUParamInfo& ParameterInfo, const class FShaderParameterMap& ParameterMap)
+	{
+		FNDIRigidMeshCollisionParametersName ParamNames(*ParameterInfo.DataInterfaceHLSLSymbol);
+
+		ElementOffsets.Bind(ParameterMap, *ParamNames.ElementOffsetsName);
+
+		WorldTransformBuffer.Bind(ParameterMap, *ParamNames.WorldTransformBufferName);
+		InverseTransformBuffer.Bind(ParameterMap, *ParamNames.InverseTransformBufferName);
+		ElementExtentBuffer.Bind(ParameterMap, *ParamNames.ElementExtentBufferName);
+		PhysicsTypeBuffer.Bind(ParameterMap, *ParamNames.PhysicsTypeBufferName);		
+	}
+
+	void Set(FRHICommandList& RHICmdList, const FNiagaraDataInterfaceSetArgs& Context) const
+	{
+		check(IsInRenderingThread());
+
+		FRHIComputeShader* ComputeShaderRHI = RHICmdList.GetBoundComputeShader();
+
+		FNDIRigidMeshCollisionProxy* InterfaceProxy =
+			static_cast<FNDIRigidMeshCollisionProxy*>(Context.DataInterface);
+		FNDIRigidMeshCollisionData* ProxyData =
+			InterfaceProxy->SystemInstancesToProxyData.Find(Context.SystemInstanceID);
+
+		if (ProxyData != nullptr && ProxyData->AssetBuffer && ProxyData->AssetBuffer->IsInitialized())
+		{
+			FNDIRigidMeshCollisionBuffer* AssetBuffer = ProxyData->AssetBuffer;
+
+			FRHITransitionInfo Transitions[] = {
+				FRHITransitionInfo(AssetBuffer->WorldTransformBuffer.UAV, ERHIAccess::Unknown, ERHIAccess::SRVCompute),
+				FRHITransitionInfo(AssetBuffer->InverseTransformBuffer.UAV, ERHIAccess::Unknown, ERHIAccess::SRVCompute),
+				FRHITransitionInfo(AssetBuffer->ElementExtentBuffer.UAV, ERHIAccess::Unknown, ERHIAccess::SRVCompute),
+				FRHITransitionInfo(AssetBuffer->PhysicsTypeBuffer.UAV, ERHIAccess::Unknown, ERHIAccess::SRVCompute)
+			};
+			RHICmdList.Transition(MakeArrayView(Transitions, UE_ARRAY_COUNT(Transitions)));
+
+			SetSRVParameter(RHICmdList, ComputeShaderRHI, WorldTransformBuffer, AssetBuffer->WorldTransformBuffer.SRV);
+			SetSRVParameter(RHICmdList, ComputeShaderRHI, InverseTransformBuffer, AssetBuffer->InverseTransformBuffer.SRV);
+			SetSRVParameter(RHICmdList, ComputeShaderRHI, ElementExtentBuffer, AssetBuffer->ElementExtentBuffer.SRV);
+			SetSRVParameter(RHICmdList, ComputeShaderRHI, PhysicsTypeBuffer, AssetBuffer->PhysicsTypeBuffer.SRV);
+
+			SetShaderValue(RHICmdList, ComputeShaderRHI, ElementOffsets, ProxyData->AssetArrays.ElementOffsets);			
+		}
+		else
+		{
+			SetSRVParameter(RHICmdList, ComputeShaderRHI, WorldTransformBuffer, FNiagaraRenderer::GetDummyFloatBuffer());
+			SetSRVParameter(RHICmdList, ComputeShaderRHI, InverseTransformBuffer, FNiagaraRenderer::GetDummyFloatBuffer());
+			SetSRVParameter(RHICmdList, ComputeShaderRHI, ElementExtentBuffer, FNiagaraRenderer::GetDummyFloatBuffer());
+			SetSRVParameter(RHICmdList, ComputeShaderRHI, PhysicsTypeBuffer, FNiagaraRenderer::GetDummyIntBuffer());
+
+			static const FElementOffset DummyOffsets(0, 0, 0, 0);
+			SetShaderValue(RHICmdList, ComputeShaderRHI, ElementOffsets, DummyOffsets);	
+		}
+	}
+
+	void Unset(FRHICommandList& RHICmdList, const FNiagaraDataInterfaceSetArgs& Context) const
+	{
+	}
+
+private:
+
+	LAYOUT_FIELD(FShaderParameter, ElementOffsets);
+
+	LAYOUT_FIELD(FShaderResourceParameter, WorldTransformBuffer);
+	LAYOUT_FIELD(FShaderResourceParameter, InverseTransformBuffer);
+	LAYOUT_FIELD(FShaderResourceParameter, ElementExtentBuffer);
+	LAYOUT_FIELD(FShaderResourceParameter, PhysicsTypeBuffer);
+};
+
+IMPLEMENT_TYPE_LAYOUT(FNDIRigidMeshCollisionParametersCS);
+
+IMPLEMENT_NIAGARA_DI_PARAMETER(UNiagaraDataInterfaceRigidMeshCollisionQuery, FNDIRigidMeshCollisionParametersCS);
+
+
+//------------------------------------------------------------------------------------------------------------
+
+void FNDIRigidMeshCollisionProxy::ConsumePerInstanceDataFromGameThread(void* PerInstanceData, const FNiagaraSystemInstanceID& Instance)
+{
+	FNDIRigidMeshCollisionData* SourceData = static_cast<FNDIRigidMeshCollisionData*>(PerInstanceData);
+	FNDIRigidMeshCollisionData* TargetData = &(SystemInstancesToProxyData.FindOrAdd(Instance));
+
+	ensure(TargetData);
+	if (TargetData)
+	{
+		TargetData->AssetBuffer = SourceData->AssetBuffer;		
+		TargetData->AssetArrays = SourceData->AssetArrays;
+		TargetData->TickingGroup = SourceData->TickingGroup;
+	}
+	else
+	{
+		UE_LOG(LogRigidMeshCollision, Log, TEXT("ConsumePerInstanceDataFromGameThread() ... could not find %d"), Instance);
+	}
+}
+
+void FNDIRigidMeshCollisionProxy::InitializePerInstanceData(const FNiagaraSystemInstanceID& SystemInstance)
+{
+	check(IsInRenderingThread());
+
+	FNDIRigidMeshCollisionData* TargetData = SystemInstancesToProxyData.Find(SystemInstance);
+	TargetData = &SystemInstancesToProxyData.Add(SystemInstance);
+}
+
+void FNDIRigidMeshCollisionProxy::DestroyPerInstanceData(NiagaraEmitterInstanceBatcher* Batcher, const FNiagaraSystemInstanceID& SystemInstance)
+{
+	check(IsInRenderingThread());
+	SystemInstancesToProxyData.Remove(SystemInstance);
+}
+
+void FNDIRigidMeshCollisionProxy::PreStage(FRHICommandList& RHICmdList, const FNiagaraDataInterfaceStageArgs& Context)
+{
+	FNDIRigidMeshCollisionData* ProxyData =
+		SystemInstancesToProxyData.Find(Context.SystemInstanceID);
+
+	if (ProxyData != nullptr && ProxyData->AssetBuffer)
+	{
+		if (Context.SimStageData->bFirstStage)
+		{
+			FRHITransitionInfo Transitions[] = {
+				FRHITransitionInfo(ProxyData->AssetBuffer->WorldTransformBuffer.UAV, ERHIAccess::Unknown, ERHIAccess::UAVCompute),
+				FRHITransitionInfo(ProxyData->AssetBuffer->InverseTransformBuffer.UAV, ERHIAccess::Unknown, ERHIAccess::UAVCompute),
+				FRHITransitionInfo(ProxyData->AssetBuffer->PhysicsTypeBuffer.UAV, ERHIAccess::Unknown, ERHIAccess::UAVCompute),
+				FRHITransitionInfo(ProxyData->AssetBuffer->ElementExtentBuffer.UAV, ERHIAccess::Unknown, ERHIAccess::UAVCompute)
+			};
+			RHICmdList.Transition(MakeArrayView(Transitions, UE_ARRAY_COUNT(Transitions)));
+
+			UpdateInternalBuffer<FVector4, EPixelFormat::PF_A32B32G32R32F, RIGID_MESH_COLLISION_QUERY_MAX_TRANSFORMS, 2>(ProxyData->AssetArrays.WorldTransform, ProxyData->AssetBuffer->WorldTransformBuffer);
+			UpdateInternalBuffer<FVector4, EPixelFormat::PF_A32B32G32R32F, RIGID_MESH_COLLISION_QUERY_MAX_TRANSFORMS, 2>(ProxyData->AssetArrays.InverseTransform, ProxyData->AssetBuffer->InverseTransformBuffer);
+			UpdateInternalBuffer<FVector4, EPixelFormat::PF_A32B32G32R32F, RIGID_MESH_COLLISION_QUERY_MAX_PRIMITIVES>(ProxyData->AssetArrays.ElementExtent, ProxyData->AssetBuffer->ElementExtentBuffer);
+			UpdateInternalBuffer<uint32, EPixelFormat::PF_R32_UINT, RIGID_MESH_COLLISION_QUERY_MAX_PRIMITIVES>(ProxyData->AssetArrays.PhysicsType, ProxyData->AssetBuffer->PhysicsTypeBuffer);
+		}
+	}
+}
+
+void FNDIRigidMeshCollisionProxy::ResetData(FRHICommandList& RHICmdList, const FNiagaraDataInterfaceArgs& Context)
+{}
+
+//------------------------------------------------------------------------------------------------------------
+
+UNiagaraDataInterfaceRigidMeshCollisionQuery::UNiagaraDataInterfaceRigidMeshCollisionQuery(FObjectInitializer const& ObjectInitializer)
+	: Super(ObjectInitializer)	
+{
+	Proxy.Reset(new FNDIRigidMeshCollisionProxy());
+}
+
+
+bool UNiagaraDataInterfaceRigidMeshCollisionQuery::InitPerInstanceData(void* PerInstanceData, FNiagaraSystemInstance* SystemInstance)
+{
+	FNDIRigidMeshCollisionData* InstanceData = new (PerInstanceData) FNDIRigidMeshCollisionData();
+
+	check(InstanceData);
+	InstanceData->Init(this, SystemInstance);
+
+	return true;
+}
+
+ETickingGroup UNiagaraDataInterfaceRigidMeshCollisionQuery::CalculateTickGroup(const void* PerInstanceData) const
+{
+	const FNDIRigidMeshCollisionData* InstanceData = static_cast<const FNDIRigidMeshCollisionData*>(PerInstanceData);
+
+	if (InstanceData)
+	{
+		return InstanceData->TickingGroup;
+	}
+	return NiagaraFirstTickGroup;
+}
+
+void UNiagaraDataInterfaceRigidMeshCollisionQuery::DestroyPerInstanceData(void* PerInstanceData, FNiagaraSystemInstance* SystemInstance)
+{
+	FNDIRigidMeshCollisionData* InstanceData = static_cast<FNDIRigidMeshCollisionData*>(PerInstanceData);
+
+	InstanceData->Release();
+	InstanceData->~FNDIRigidMeshCollisionData();
+
+	FNDIRigidMeshCollisionProxy* ThisProxy = GetProxyAs<FNDIRigidMeshCollisionProxy>();
+	ENQUEUE_RENDER_COMMAND(FNiagaraDIDestroyInstanceData) (
+		[ThisProxy, InstanceID = SystemInstance->GetId(), Batcher = SystemInstance->GetBatcher()](FRHICommandListImmediate& CmdList)
+	{
+		ThisProxy->SystemInstancesToProxyData.Remove(InstanceID);
+	}
+	);
+}
+
+bool UNiagaraDataInterfaceRigidMeshCollisionQuery::PerInstanceTick(void* PerInstanceData, FNiagaraSystemInstance* SystemInstance, float InDeltaSeconds)
+{
+	FNDIRigidMeshCollisionData* InstanceData = static_cast<FNDIRigidMeshCollisionData*>(PerInstanceData);
+	if (InstanceData && InstanceData->AssetBuffer && SystemInstance)
+	{
+		InstanceData->Update(this, SystemInstance);
+	}
+	return false;
+}
+
+bool UNiagaraDataInterfaceRigidMeshCollisionQuery::CopyToInternal(UNiagaraDataInterface* Destination) const
+{
+	if (!Super::CopyToInternal(Destination))
+	{
+		return false;
+	}
+
+	UNiagaraDataInterfaceRigidMeshCollisionQuery* OtherTyped = CastChecked<UNiagaraDataInterfaceRigidMeshCollisionQuery>(Destination);		
+
+	return true;
+}
+
+bool UNiagaraDataInterfaceRigidMeshCollisionQuery::Equals(const UNiagaraDataInterface* Other) const
+{
+	if (!Super::Equals(Other))
+	{
+		return false;
+	}
+	const UNiagaraDataInterfaceRigidMeshCollisionQuery* OtherTyped = CastChecked<const UNiagaraDataInterfaceRigidMeshCollisionQuery>(Other);
+
+	return true;
+}
+
+void UNiagaraDataInterfaceRigidMeshCollisionQuery::PostInitProperties()
+{
+	Super::PostInitProperties();
+
+	if (HasAnyFlags(RF_ClassDefaultObject))
+	{
+		ENiagaraTypeRegistryFlags Flags = ENiagaraTypeRegistryFlags::AllowAnyVariable | ENiagaraTypeRegistryFlags::AllowParameter;
+		FNiagaraTypeRegistry::Register(FNiagaraTypeDefinition(GetClass()), Flags);
+	}
+}
+
+void UNiagaraDataInterfaceRigidMeshCollisionQuery::GetFunctions(TArray<FNiagaraFunctionSignature>& OutFunctions)
+{
+	{
+		FNiagaraFunctionSignature Sig;
+		Sig.Name = GetNumBoxesName;
+		Sig.bSupportsGPU = true;
+		Sig.bSupportsCPU = false;
+		Sig.bMemberFunction = true;
+		Sig.bRequiresContext = false;
+		Sig.Inputs.Add(FNiagaraVariable(FNiagaraTypeDefinition(GetClass()), TEXT("Collision DI")));
+		Sig.Outputs.Add(FNiagaraVariable(FNiagaraTypeDefinition::GetIntDef(), TEXT("Num Boxes")));
+
+		OutFunctions.Add(Sig);
+	}
+	{
+		FNiagaraFunctionSignature Sig;
+		Sig.Name = GetNumSpheresName;
+		Sig.bSupportsGPU = true;
+		Sig.bSupportsCPU = false;
+		Sig.bMemberFunction = true;
+		Sig.bRequiresContext = false;
+		Sig.Inputs.Add(FNiagaraVariable(FNiagaraTypeDefinition(GetClass()), TEXT("Collision DI")));
+		Sig.Outputs.Add(FNiagaraVariable(FNiagaraTypeDefinition::GetIntDef(), TEXT("Num Spheres")));
+
+		OutFunctions.Add(Sig);
+	}
+	{
+		FNiagaraFunctionSignature Sig;
+		Sig.Name = GetNumCapsulesName;
+		Sig.bSupportsGPU = true;
+		Sig.bSupportsCPU = false;
+		Sig.bMemberFunction = true;
+		Sig.bRequiresContext = false;
+		Sig.Inputs.Add(FNiagaraVariable(FNiagaraTypeDefinition(GetClass()), TEXT("Collision DI")));
+		Sig.Outputs.Add(FNiagaraVariable(FNiagaraTypeDefinition::GetIntDef(), TEXT("Num Capsules")));
+
+		OutFunctions.Add(Sig);
+	}
+	{
+		FNiagaraFunctionSignature Sig;
+		Sig.Name = GetClosestPointName;
+		Sig.bSupportsGPU = true;
+		Sig.bSupportsCPU = false;
+		Sig.bMemberFunction = true;
+		Sig.bRequiresContext = false;
+		Sig.Inputs.Add(FNiagaraVariable(FNiagaraTypeDefinition(GetClass()), TEXT("Collision DI")));
+		Sig.Inputs.Add(FNiagaraVariable(FNiagaraTypeDefinition::GetVec3Def(), TEXT("Node Position")));
+		Sig.Inputs.Add(FNiagaraVariable(FNiagaraTypeDefinition::GetFloatDef(), TEXT("Delta Time")));
+		Sig.Inputs.Add(FNiagaraVariable(FNiagaraTypeDefinition::GetFloatDef(), TEXT("Time Fraction")));
+		Sig.Outputs.Add(FNiagaraVariable(FNiagaraTypeDefinition::GetFloatDef(), TEXT("Closest Distance")));
+		Sig.Outputs.Add(FNiagaraVariable(FNiagaraTypeDefinition::GetVec3Def(), TEXT("Closest Position")));
+		Sig.Outputs.Add(FNiagaraVariable(FNiagaraTypeDefinition::GetVec3Def(), TEXT("Closest Normal")));
+		Sig.Outputs.Add(FNiagaraVariable(FNiagaraTypeDefinition::GetVec3Def(), TEXT("Closest Velocity")));
+
+		OutFunctions.Add(Sig);
+	}
+	{
+		FNiagaraFunctionSignature Sig;
+		Sig.Name = GetClosestElementName;
+		Sig.bSupportsGPU = true;
+		Sig.bSupportsCPU = false;
+		Sig.bMemberFunction = true;
+		Sig.bRequiresContext = false;
+		Sig.Inputs.Add(FNiagaraVariable(FNiagaraTypeDefinition(GetClass()), TEXT("Collision DI")));
+		Sig.Inputs.Add(FNiagaraVariable(FNiagaraTypeDefinition::GetVec3Def(), TEXT("Node Position")));
+		Sig.Inputs.Add(FNiagaraVariable(FNiagaraTypeDefinition::GetFloatDef(), TEXT("Time Fraction")));
+		Sig.Outputs.Add(FNiagaraVariable(FNiagaraTypeDefinition::GetIntDef(), TEXT("Closest Element")));
+
+		OutFunctions.Add(Sig);
+	}
+	{
+		FNiagaraFunctionSignature Sig;
+		Sig.Name = GetElementPointName;
+		Sig.bSupportsGPU = true;
+		Sig.bSupportsCPU = false;
+		Sig.bMemberFunction = true;
+		Sig.bRequiresContext = false;
+		Sig.Inputs.Add(FNiagaraVariable(FNiagaraTypeDefinition(GetClass()), TEXT("Collision DI")));
+		Sig.Inputs.Add(FNiagaraVariable(FNiagaraTypeDefinition::GetVec3Def(), TEXT("Node Position")));
+		Sig.Inputs.Add(FNiagaraVariable(FNiagaraTypeDefinition::GetFloatDef(), TEXT("Delta Time")));
+		Sig.Inputs.Add(FNiagaraVariable(FNiagaraTypeDefinition::GetFloatDef(), TEXT("Time Fraction")));
+		Sig.Inputs.Add(FNiagaraVariable(FNiagaraTypeDefinition::GetIntDef(), TEXT("Element Index")));
+		Sig.Outputs.Add(FNiagaraVariable(FNiagaraTypeDefinition::GetVec3Def(), TEXT("Closest Position")));
+		Sig.Outputs.Add(FNiagaraVariable(FNiagaraTypeDefinition::GetVec3Def(), TEXT("Closest Normal")));
+		Sig.Outputs.Add(FNiagaraVariable(FNiagaraTypeDefinition::GetVec3Def(), TEXT("Closest Velocity")));
+
+		OutFunctions.Add(Sig);
+	}
+	{
+		FNiagaraFunctionSignature Sig;
+		Sig.Name = GetElementDistanceName;
+		Sig.bSupportsGPU = true;
+		Sig.bSupportsCPU = false;
+		Sig.bMemberFunction = true;
+		Sig.bRequiresContext = false;
+		Sig.Inputs.Add(FNiagaraVariable(FNiagaraTypeDefinition(GetClass()), TEXT("Collision DI")));
+		Sig.Inputs.Add(FNiagaraVariable(FNiagaraTypeDefinition::GetVec3Def(), TEXT("Node Position")));
+		Sig.Inputs.Add(FNiagaraVariable(FNiagaraTypeDefinition::GetFloatDef(), TEXT("Time Fraction")));
+		Sig.Inputs.Add(FNiagaraVariable(FNiagaraTypeDefinition::GetIntDef(), TEXT("Element Index")));
+		Sig.Outputs.Add(FNiagaraVariable(FNiagaraTypeDefinition::GetFloatDef(), TEXT("Closest Distance")));
+
+		OutFunctions.Add(Sig);
+	}
+	{
+		FNiagaraFunctionSignature Sig;
+		Sig.Name = GetClosestDistanceName;
+		Sig.bSupportsGPU = true;
+		Sig.bSupportsCPU = false;
+		Sig.bMemberFunction = true;
+		Sig.bRequiresContext = false;
+		Sig.Inputs.Add(FNiagaraVariable(FNiagaraTypeDefinition(GetClass()), TEXT("Collision DI")));
+		Sig.Inputs.Add(FNiagaraVariable(FNiagaraTypeDefinition::GetVec3Def(), TEXT("Node Position")));
+		Sig.Inputs.Add(FNiagaraVariable(FNiagaraTypeDefinition::GetFloatDef(), TEXT("Time Fraction")));
+		Sig.Outputs.Add(FNiagaraVariable(FNiagaraTypeDefinition::GetFloatDef(), TEXT("Closest Distance")));
+
+		OutFunctions.Add(Sig);
+	}
+}
+
+void UNiagaraDataInterfaceRigidMeshCollisionQuery::GetVMExternalFunction(const FVMExternalFunctionBindingInfo& BindingInfo, void* InstanceData, FVMExternalFunction& OutFunc)
+{
+}
+
+#if WITH_EDITORONLY_DATA
+bool UNiagaraDataInterfaceRigidMeshCollisionQuery::GetFunctionHLSL(const FNiagaraDataInterfaceGPUParamInfo& ParamInfo, const FNiagaraDataInterfaceGeneratedFunction& FunctionInfo, int FunctionInstanceIndex, FString& OutHLSL)
+{
+	FNDIRigidMeshCollisionParametersName ParamNames(ParamInfo.DataInterfaceHLSLSymbol);
+
+	TMap<FString, FStringFormatArg> ArgsSample = {
+		{TEXT("InstanceFunctionName"), FunctionInfo.InstanceName},
+		{TEXT("ElementOffsetsName"), ParamNames.ElementOffsetsName},
+		{TEXT("WorldTransformBufferName"), ParamNames.WorldTransformBufferName},
+		{TEXT("InverseTransformBufferName"), ParamNames.InverseTransformBufferName},
+		{TEXT("ElementExtentBufferName"), ParamNames.ElementExtentBufferName},
+		{TEXT("RigidMeshCollisionContextName"), TEXT("DIRIGIDMESHCOLLISIONQUERY_MAKE_CONTEXT(") + ParamInfo.DataInterfaceHLSLSymbol + TEXT(")")},
+	};
+
+	if (FunctionInfo.DefinitionName == GetNumBoxesName)
+	{
+		static const TCHAR* FormatSample = TEXT(R"(
+		void {InstanceFunctionName}(out int OutNumBoxes)
+		{
+			{RigidMeshCollisionContextName}
+			OutNumBoxes = DIRigidMeshCollision_GetNumBoxes(DIContext);
+		}
+		)");
+		OutHLSL += FString::Format(FormatSample, ArgsSample);
+		return true;
+	}
+	else if (FunctionInfo.DefinitionName == GetNumCapsulesName)
+	{
+		static const TCHAR* FormatSample = TEXT(R"(
+		void {InstanceFunctionName}(out int OutNumCapsules)
+		{
+			{RigidMeshCollisionContextName}
+			OutNumCapsules = DIRigidMeshCollision_GetNumCapsules(DIContext);
+		}
+		)");
+		OutHLSL += FString::Format(FormatSample, ArgsSample);
+		return true;
+	}
+	else if (FunctionInfo.DefinitionName == GetNumSpheresName)
+	{
+		static const TCHAR* FormatSample = TEXT(R"(
+		void {InstanceFunctionName}(out int OutNumSpheres)
+		{
+			{RigidMeshCollisionContextName}
+			OutNumSpheres = DIRigidMeshCollision_GetNumSpheres(DIContext);
+		}
+		)");
+		OutHLSL += FString::Format(FormatSample, ArgsSample);
+		return true;
+	}
+	else if (FunctionInfo.DefinitionName == GetClosestPointName)
+	{
+		static const TCHAR* FormatSample = TEXT(R"(
+		void {InstanceFunctionName}(in float3 NodePosition, in float DeltaTime, in float TimeFraction, out float ClosestDistance, out float3 OutClosestPosition, 
+							out float3 OutClosestNormal, out float3 OutClosestVelocity)
+		{
+			{RigidMeshCollisionContextName} DIRigidMeshCollision_GetClosestPoint(DIContext,NodePosition,DeltaTime,TimeFraction, ClosestDistance,
+				OutClosestPosition,OutClosestNormal,OutClosestVelocity);
+		}
+		)");
+		OutHLSL += FString::Format(FormatSample, ArgsSample);
+		return true;
+	}
+	else if (FunctionInfo.DefinitionName == GetClosestElementName)
+	{
+		static const TCHAR* FormatSample = TEXT(R"(
+		void {InstanceFunctionName}(in float3 NodePosition, in float TimeFraction, out int OutClosestElement)
+		{
+			{RigidMeshCollisionContextName} DIRigidMeshCollision_GetClosestElement(DIContext,NodePosition,TimeFraction,
+				OutClosestElement);
+		}
+		)");
+		OutHLSL += FString::Format(FormatSample, ArgsSample);
+		return true;
+	}
+	else if (FunctionInfo.DefinitionName == GetElementPointName)
+	{
+		static const TCHAR* FormatSample = TEXT(R"(
+		void {InstanceFunctionName}(in float3 NodePosition, in float DeltaTime, in float TimeFraction, in int ElementIndex, out float3 OutClosestPosition, 
+							out float3 OutClosestNormal, out float3 OutClosestVelocity)
+		{
+			{RigidMeshCollisionContextName} DIRigidMeshCollision_GetElementPoint(DIContext,NodePosition,DeltaTime,TimeFraction,ElementIndex,
+				OutClosestPosition,OutClosestNormal,OutClosestVelocity);
+		}
+		)");
+		OutHLSL += FString::Format(FormatSample, ArgsSample);
+		return true;
+	}
+	else if (FunctionInfo.DefinitionName == GetElementDistanceName)
+	{
+		static const TCHAR* FormatSample = TEXT(R"(
+		void {InstanceFunctionName}(in float3 NodePosition, in float TimeFraction, in int ElementIndex, out float OutClosestDistance)
+		{
+			{RigidMeshCollisionContextName} DIRigidMeshCollision_GetElementDistance(DIContext,NodePosition,TimeFraction,ElementIndex,
+				OutClosestDistance);
+		}
+		)");
+		OutHLSL += FString::Format(FormatSample, ArgsSample);
+		return true;
+	}
+	else if (FunctionInfo.DefinitionName == GetClosestDistanceName)
+	{
+		static const TCHAR* FormatSample = TEXT(R"(
+		void {InstanceFunctionName}(in float3 NodePosition, in float TimeFraction, out float OutClosestDistance)
+		{
+			{RigidMeshCollisionContextName} DIRigidMeshCollision_GetClosestDistance(DIContext,NodePosition,TimeFraction,OutClosestDistance);
+		}
+		)");
+		OutHLSL += FString::Format(FormatSample, ArgsSample);
+		return true;
+	}
+
+	OutHLSL += TEXT("\n");
+	return false;
+}
+
+void UNiagaraDataInterfaceRigidMeshCollisionQuery::GetCommonHLSL(FString& OutHLSL)
+{
+	OutHLSL += TEXT("#include \"/Plugin/FX/Niagara/Private/NiagaraDataInterfaceRigidMeshCollisionQuery.ush\"\n");		
+}
+
+void UNiagaraDataInterfaceRigidMeshCollisionQuery::GetParameterDefinitionHLSL(const FNiagaraDataInterfaceGPUParamInfo& ParamInfo, FString& OutHLSL)
+{
+	OutHLSL += TEXT("DIRIGIDMESHCOLLISIONQUERY_DECLARE_CONSTANTS(") + ParamInfo.DataInterfaceHLSLSymbol + TEXT(")\n");
+}
+#endif
+
+void UNiagaraDataInterfaceRigidMeshCollisionQuery::ProvidePerInstanceDataForRenderThread(void* DataForRenderThread, void* PerInstanceData, const FNiagaraSystemInstanceID& SystemInstance)
+{
+	FNDIRigidMeshCollisionData* GameThreadData = static_cast<FNDIRigidMeshCollisionData*>(PerInstanceData);
+	FNDIRigidMeshCollisionData* RenderThreadData = static_cast<FNDIRigidMeshCollisionData*>(DataForRenderThread);
+
+	if (GameThreadData != nullptr && RenderThreadData != nullptr)
+	{
+		RenderThreadData->AssetBuffer = GameThreadData->AssetBuffer;		
+		RenderThreadData->AssetArrays = GameThreadData->AssetArrays;
+		RenderThreadData->TickingGroup = GameThreadData->TickingGroup;
+	}
+	check(Proxy);
+}
+
+#undef LOCTEXT_NAMESPACE
\ No newline at end of file
diff --git a/Engine/Plugins/FX/Niagara/Source/Niagara/Public/NiagaraDataInterfaceRigidMeshCollisionQuery.h b/Engine/Plugins/FX/Niagara/Source/Niagara/Public/NiagaraDataInterfaceRigidMeshCollisionQuery.h
new file mode 100644
--- /dev/null
+++ b/Engine/Plugins/FX/Niagara/Source/Niagara/Public/NiagaraDataInterfaceRigidMeshCollisionQuery.h
@@ -0,0 +1,171 @@
+// Copyright Epic Games, Inc. All Rights Reserved.
+
+#pragma once
+
+#include "NiagaraDataInterface.h"
+#include "NiagaraCommon.h"
+#include "VectorVM.h"
+#include "NiagaraDataInterfaceRigidMeshCollisionQuery.generated.h"
+
+#define RIGID_MESH_COLLISION_QUERY_MAX_PRIMITIVES 100
+#define RIGID_MESH_COLLISION_QUERY_MAX_TRANSFORMS RIGID_MESH_COLLISION_QUERY_MAX_PRIMITIVES * 2
+
+/** Element offsets in the array list */
+struct FElementOffset
+{
+	FElementOffset(const uint32 InBoxOffset, const uint32 InSphereOffset, const uint32 InCapsuleOffset, const uint32 InNumElements) :
+		BoxOffset(InBoxOffset), SphereOffset(InSphereOffset), CapsuleOffset(InCapsuleOffset), NumElements(InNumElements)
+	{}
+
+	FElementOffset() :
+		BoxOffset(0), SphereOffset(0), CapsuleOffset(0), NumElements(0)
+	{}
+	uint32 BoxOffset;
+	uint32 SphereOffset;
+	uint32 CapsuleOffset;
+	uint32 NumElements;
+};
+
+/** Arrays in which the cpu datas will be str */
+struct FNDIRigidMeshCollisionArrays
+{
+	FElementOffset ElementOffsets;
+	TStaticArray<FVector4, 2 * RIGID_MESH_COLLISION_QUERY_MAX_TRANSFORMS> WorldTransform;
+	TStaticArray<FVector4, 2 * RIGID_MESH_COLLISION_QUERY_MAX_TRANSFORMS> InverseTransform;
+	TStaticArray<FVector4, RIGID_MESH_COLLISION_QUERY_MAX_TRANSFORMS> CurrentTransform;
+	TStaticArray<FVector4, RIGID_MESH_COLLISION_QUERY_MAX_TRANSFORMS> CurrentInverse;
+	TStaticArray<FVector4, RIGID_MESH_COLLISION_QUERY_MAX_TRANSFORMS> PreviousTransform;
+	TStaticArray<FVector4, RIGID_MESH_COLLISION_QUERY_MAX_TRANSFORMS> PreviousInverse;
+	TStaticArray<FVector4, RIGID_MESH_COLLISION_QUERY_MAX_PRIMITIVES> ElementExtent;
+	TStaticArray<uint32, RIGID_MESH_COLLISION_QUERY_MAX_PRIMITIVES> PhysicsType;
+};
+
+/** Render buffers that will be used in hlsl functions */
+struct FNDIRigidMeshCollisionBuffer : public FRenderResource
+{
+	/** Init the buffer */
+	virtual void InitRHI() override;
+
+	/** Release the buffer */
+	virtual void ReleaseRHI() override;
+
+	/** Get the resource name */
+	virtual FString GetFriendlyName() const override { return TEXT("FNDIRigidMeshCollisionBuffer"); }
+
+	/** World transform buffer */
+	FRWBuffer WorldTransformBuffer;
+
+	/** Inverse transform buffer*/
+	FRWBuffer InverseTransformBuffer;
+
+	/** Element extent buffer */
+	FRWBuffer ElementExtentBuffer;
+
+	/** Physics type buffer */
+	FRWBuffer PhysicsTypeBuffer;
+};
+
+/** Data stored per physics asset instance*/
+struct FNDIRigidMeshCollisionData
+{
+	/** Initialize the cpu datas */
+	void Init(class UNiagaraDataInterfaceRigidMeshCollisionQuery* Interface, FNiagaraSystemInstance* SystemInstance);
+
+	/** Update the gpu datas */
+	void Update(class UNiagaraDataInterfaceRigidMeshCollisionQuery* Interface, FNiagaraSystemInstance* SystemInstance);
+
+	/** Release the buffers */
+	void Release();
+
+	/** Physics asset Gpu buffer */
+	FNDIRigidMeshCollisionBuffer* AssetBuffer;
+
+	/** Physics asset Cpu arrays */
+	FNDIRigidMeshCollisionArrays AssetArrays;  
+
+	/** The instance ticking group */
+	ETickingGroup TickingGroup;
+
+	/** Static Mesh Components **/
+	TArray<AStaticMeshActor*> StaticMeshActors;
+};
+
+/** Data Interface for the Collisions */
+UCLASS(EditInlineNew, Category = "Collision", meta = (DisplayName = "Rigid Mesh Collision Query"))
+class UNiagaraDataInterfaceRigidMeshCollisionQuery : public UNiagaraDataInterface
+{
+	GENERATED_UCLASS_BODY()
+
+public:
+
+	DECLARE_NIAGARA_DI_PARAMETER();
+
+	/** UObject Interface */
+	virtual void PostInitProperties() override;
+
+	/** UNiagaraDataInterface Interface */
+	virtual void GetFunctions(TArray<FNiagaraFunctionSignature>& OutFunctions) override;
+	virtual void GetVMExternalFunction(const FVMExternalFunctionBindingInfo& BindingInfo, void* InstanceData, FVMExternalFunction& OutFunc) override;
+	virtual bool CanExecuteOnTarget(ENiagaraSimTarget Target) const override { return Target == ENiagaraSimTarget::GPUComputeSim; }
+	virtual bool InitPerInstanceData(void* PerInstanceData, FNiagaraSystemInstance* SystemInstance) override;
+	virtual void DestroyPerInstanceData(void* PerInstanceData, FNiagaraSystemInstance* SystemInstance)override;
+	virtual bool PerInstanceTick(void* PerInstanceData, FNiagaraSystemInstance* SystemInstance, float DeltaSeconds) override;
+	virtual int32 PerInstanceDataSize()const override { return sizeof(FNDIRigidMeshCollisionData); }
+	virtual bool Equals(const UNiagaraDataInterface* Other) const override;
+	virtual bool HasPreSimulateTick() const override { return true; }
+	virtual bool HasTickGroupPrereqs() const override { return true; }
+	virtual ETickingGroup CalculateTickGroup(const void* PerInstanceData) const override;
+
+	/** GPU simulation  functionality */
+#if WITH_EDITORONLY_DATA
+	virtual void GetCommonHLSL(FString& OutHLSL) override;
+	virtual void GetParameterDefinitionHLSL(const FNiagaraDataInterfaceGPUParamInfo& ParamInfo, FString& OutHLSL) override;
+	virtual bool GetFunctionHLSL(const FNiagaraDataInterfaceGPUParamInfo& ParamInfo, const FNiagaraDataInterfaceGeneratedFunction& FunctionInfo, int FunctionInstanceIndex, FString& OutHLSL) override;
+#endif
+	virtual void ProvidePerInstanceDataForRenderThread(void* DataForRenderThread, void* PerInstanceData, const FNiagaraSystemInstanceID& SystemInstance) override;
+
+	/** Name of element offsets */
+	static const FString ElementOffsetsName;
+
+	/** Name of the world transform buffer */
+	static const FString WorldTransformBufferName;
+
+	/** Name of the inverse transform buffer */
+	static const FString InverseTransformBufferName;
+
+	/** Name of the element extent buffer */
+	static const FString ElementExtentBufferName;
+
+	/** Name of the physics type buffer */
+	static const FString PhysicsTypeBufferName;
+
+protected:
+	/** Copy one niagara DI to this */
+	virtual bool CopyToInternal(UNiagaraDataInterface* Destination) const override;
+};
+
+/** Proxy to send data to gpu */
+struct FNDIRigidMeshCollisionProxy : public FNiagaraDataInterfaceProxy
+{
+	/** Get the size of the data that will be passed to render*/
+	virtual int32 PerInstanceDataPassedToRenderThreadSize() const override { return sizeof(FNDIRigidMeshCollisionData); }
+
+	/** Get the data that will be passed to render*/
+	virtual void ConsumePerInstanceDataFromGameThread(void* PerInstanceData, const FNiagaraSystemInstanceID& Instance) override;
+
+	/** Initialize the Proxy data buffer */
+	void InitializePerInstanceData(const FNiagaraSystemInstanceID& SystemInstance);
+
+	/** Destroy the proxy data if necessary */
+	void DestroyPerInstanceData(NiagaraEmitterInstanceBatcher* Batcher, const FNiagaraSystemInstanceID& SystemInstance);
+
+	/** Launch all pre stage functions */
+	virtual void PreStage(FRHICommandList& RHICmdList, const FNiagaraDataInterfaceStageArgs& Context) override;
+
+	/** Reset the buffers  */
+	virtual void ResetData(FRHICommandList& RHICmdList, const FNiagaraDataInterfaceArgs& Context) override;
+
+	/** List of proxy data for each system instances*/
+	TMap<FNiagaraSystemInstanceID, FNDIRigidMeshCollisionData> SystemInstancesToProxyData;
+};
+
