CIS fix

diff --git a/Engine/Plugins/FX/Niagara/Source/Niagara/Public/NiagaraDataInterfaceRigidMeshCollisionQuery.h b/Engine/Plugins/FX/Niagara/Source/Niagara/Public/NiagaraDataInterfaceRigidMeshCollisionQuery.h
--- a/Engine/Plugins/FX/Niagara/Source/Niagara/Public/NiagaraDataInterfaceRigidMeshCollisionQuery.h
+++ b/Engine/Plugins/FX/Niagara/Source/Niagara/Public/NiagaraDataInterfaceRigidMeshCollisionQuery.h
@@ -10,6 +10,9 @@
 #define RIGID_MESH_COLLISION_QUERY_MAX_PRIMITIVES 100
 #define RIGID_MESH_COLLISION_QUERY_MAX_TRANSFORMS RIGID_MESH_COLLISION_QUERY_MAX_PRIMITIVES * 2
 
+// Forward declaration
+class AStaticMeshActor;
+
 /** Element offsets in the array list */
 struct FElementOffset
 {
@@ -100,6 +103,15 @@ public:
 
 	DECLARE_NIAGARA_DI_PARAMETER();
 
+	UPROPERTY(EditAnywhere, Category = "Static Mesh")
+	bool StaticMesh = true;
+
+	UPROPERTY(EditAnywhere, Category = "Static Mesh", meta = (EditCondition = "UseStaticMeshes"))
+	FString Tag = TEXT("");
+
+	UPROPERTY(EditAnywhere, Category = "Static Mesh", meta = (EditCondition = "UseStaticMeshes"))
+	bool OnlyUseMoveable = true;
+
 	/** UObject Interface */
 	virtual void PostInitProperties() override;
 
