flags:
turn on and off static mesh collisions
use tags
only consider moveable objects

diff --git a/Engine/Plugins/FX/Niagara/Source/Niagara/Private/NiagaraDataInterfaceRigidMeshCollisionQuery.cpp b/Engine/Plugins/FX/Niagara/Source/Niagara/Private/NiagaraDataInterfaceRigidMeshCollisionQuery.cpp
--- a/Engine/Plugins/FX/Niagara/Source/Niagara/Private/NiagaraDataInterfaceRigidMeshCollisionQuery.cpp
+++ b/Engine/Plugins/FX/Niagara/Source/Niagara/Private/NiagaraDataInterfaceRigidMeshCollisionQuery.cpp
@@ -376,27 +376,35 @@ void FNDIRigidMeshCollisionData::Init(UNiagaraDataInterfaceRigidMeshCollisionQue
 {
 	AssetBuffer = nullptr;
 
-	UWorld* World = SystemInstance->GetWorld();		
+	UWorld* World = SystemInstance->GetWorld();
 
 	StaticMeshActors.Empty();
-	for (TActorIterator<AStaticMeshActor> It(World); It; ++It)
+
+	if (Interface->StaticMesh)
 	{
-		AStaticMeshActor* StaticMeshActor = *It;		
-		StaticMeshActors.Add(StaticMeshActor);
-	
-		// Only deal with static meshes that are movable
-		// #todo(dmp): also nonmoveable?
-		
-		// Get local bounds
-		// FVector MinLocalBounds;
-		// FVector MaxLocalBounds;
-		// StaticMeshActor->GetStaticMeshComponent()->GetLocalBounds(MinLocalBounds, MaxLocalBounds);
+		for (TActorIterator<AStaticMeshActor> It(World); It; ++It)
+		{
+			AStaticMeshActor* StaticMeshActor = *It;
 
-		// Initial Transform 
-		//FTransform InitialTransform = StaticMeshActor->GetTransform();
-		
-		// Use Body setup if it exists
-		//UBodySetup *BodySetup = StaticMeshComponent->GetBodySetup();
+			if (!Interface->OnlyUseMoveable || (Interface->OnlyUseMoveable && StaticMeshActor->IsRootComponentMovable()) && 
+				(Interface->Tag == FString("") || (Interface->Tag != FString("") && StaticMeshActor->Tags.Contains(FName(Interface->Tag)))))
+			{
+				StaticMeshActors.Add(StaticMeshActor);
+			}
+			// Only deal with static meshes that are movable
+			// #todo(dmp): also nonmoveable?
+
+			// Get local bounds
+			// FVector MinLocalBounds;
+			// FVector MaxLocalBounds;
+			// StaticMeshActor->GetStaticMeshComponent()->GetLocalBounds(MinLocalBounds, MaxLocalBounds);
+
+			// Initial Transform 
+			//FTransform InitialTransform = StaticMeshActor->GetTransform();
+
+			// Use Body setup if it exists
+			//UBodySetup *BodySetup = StaticMeshComponent->GetBodySetup();
+		}
 	}
 
 	if (Interface != nullptr && SystemInstance != nullptr)
@@ -630,6 +638,10 @@ bool UNiagaraDataInterfaceRigidMeshCollisionQuery::CopyToInternal(UNiagaraDataIn
 
 	UNiagaraDataInterfaceRigidMeshCollisionQuery* OtherTyped = CastChecked<UNiagaraDataInterfaceRigidMeshCollisionQuery>(Destination);		
 
+	OtherTyped->StaticMesh = StaticMesh;
+	OtherTyped->Tag = Tag;
+	OtherTyped->OnlyUseMoveable = OnlyUseMoveable;
+
 	return true;
 }
 
@@ -641,7 +653,7 @@ bool UNiagaraDataInterfaceRigidMeshCollisionQuery::Equals(const UNiagaraDataInte
 	}
 	const UNiagaraDataInterfaceRigidMeshCollisionQuery* OtherTyped = CastChecked<const UNiagaraDataInterfaceRigidMeshCollisionQuery>(Other);
 
-	return true;
+	return  (OtherTyped->StaticMesh == StaticMesh) && (OtherTyped->Tag == Tag) && (OtherTyped->OnlyUseMoveable == OnlyUseMoveable);
 }
 
 void UNiagaraDataInterfaceRigidMeshCollisionQuery::PostInitProperties()
