cis fixes

diff --git a/Engine/Plugins/FX/Niagara/Source/Niagara/Private/NiagaraDataInterfaceRigidMeshCollisionQuery.cpp b/Engine/Plugins/FX/Niagara/Source/Niagara/Private/NiagaraDataInterfaceRigidMeshCollisionQuery.cpp
--- a/Engine/Plugins/FX/Niagara/Source/Niagara/Private/NiagaraDataInterfaceRigidMeshCollisionQuery.cpp
+++ b/Engine/Plugins/FX/Niagara/Source/Niagara/Private/NiagaraDataInterfaceRigidMeshCollisionQuery.cpp
@@ -339,24 +339,7 @@ void FNDIRigidMeshCollisionBuffer::ReleaseRHI()
 
 //------------------------------------------------------------------------------------------------------------
 
-ETickingGroup ComputeTickingGroup(const TArray<TWeakObjectPtr<class USkeletalMeshComponent>> SkeletalMeshes)
-{
-	ETickingGroup TickingGroup = NiagaraFirstTickGroup;
-	for (int32 ComponentIndex = 0; ComponentIndex < SkeletalMeshes.Num(); ++ComponentIndex)
-	{
-		if (SkeletalMeshes[ComponentIndex].Get() != nullptr)
-		{
-			const USkeletalMeshComponent* Component = Cast<USkeletalMeshComponent>(SkeletalMeshes[ComponentIndex].Get());
 
-			const ETickingGroup ComponentTickGroup = FMath::Max(Component->PrimaryComponentTick.TickGroup, Component->PrimaryComponentTick.EndTickGroup);
-			const ETickingGroup PhysicsTickGroup = Component->bBlendPhysics ? FMath::Max(ComponentTickGroup, TG_EndPhysics) : ComponentTickGroup;
-			const ETickingGroup ClampedTickGroup = FMath::Clamp(ETickingGroup(PhysicsTickGroup + 1), NiagaraFirstTickGroup, NiagaraLastTickGroup);
-
-			TickingGroup = FMath::Max(TickingGroup, ClampedTickGroup);
-		}
-	}
-	return TickingGroup;
-}
 
 void FNDIRigidMeshCollisionData::Release()
 {
@@ -386,7 +369,7 @@ void FNDIRigidMeshCollisionData::Init(UNiagaraDataInterfaceRigidMeshCollisionQue
 		{
 			AStaticMeshActor* StaticMeshActor = *It;
 
-			if (!Interface->OnlyUseMoveable || (Interface->OnlyUseMoveable && StaticMeshActor->IsRootComponentMovable()) && 
+			if ((!Interface->OnlyUseMoveable || (Interface->OnlyUseMoveable && StaticMeshActor->IsRootComponentMovable())) && 
 				(Interface->Tag == FString("") || (Interface->Tag != FString("") && StaticMeshActor->Tags.Contains(FName(Interface->Tag)))))
 			{
 				StaticMeshActors.Add(StaticMeshActor);
@@ -423,7 +406,7 @@ void FNDIRigidMeshCollisionData::Update(UNiagaraDataInterfaceRigidMeshCollisionQ
 {
 	if (Interface != nullptr && SystemInstance != nullptr)
 	{		
-		//TickingGroup = ComputeTickingGroup(Interface->SourceComponents);
+		TickingGroup = ComputeTickingGroup();
 
 		if (0 < StaticMeshActors.Num() && StaticMeshActors[0] != nullptr)
 		{
@@ -432,6 +415,29 @@ void FNDIRigidMeshCollisionData::Update(UNiagaraDataInterfaceRigidMeshCollisionQ
 	}
 }
 
+ETickingGroup FNDIRigidMeshCollisionData::ComputeTickingGroup()
+{
+	TickingGroup = NiagaraFirstTickGroup;
+	for (int32 ComponentIndex = 0; ComponentIndex < StaticMeshActors.Num(); ++ComponentIndex)
+	{
+		if (StaticMeshActors[ComponentIndex] != nullptr)
+		{			
+			UStaticMeshComponent* Component = StaticMeshActors[ComponentIndex]->GetStaticMeshComponent();
+
+			if (Component == nullptr)
+			{
+				continue;
+			}
+				const ETickingGroup ComponentTickGroup = FMath::Max(Component->PrimaryComponentTick.TickGroup, Component->PrimaryComponentTick.EndTickGroup);
+				const ETickingGroup PhysicsTickGroup = ComponentTickGroup;
+				const ETickingGroup ClampedTickGroup = FMath::Clamp(ETickingGroup(PhysicsTickGroup + 1), NiagaraFirstTickGroup, NiagaraLastTickGroup);		
+
+			TickingGroup = FMath::Max(TickingGroup, ClampedTickGroup);
+		}
+	}
+	return TickingGroup;
+}
+
 //------------------------------------------------------------------------------------------------------------
 
 struct FNDIRigidMeshCollisionParametersCS : public FNiagaraDataInterfaceParametersCS
@@ -581,7 +587,6 @@ UNiagaraDataInterfaceRigidMeshCollisionQuery::UNiagaraDataInterfaceRigidMeshColl
 	Proxy.Reset(new FNDIRigidMeshCollisionProxy());
 }
 
-
 bool UNiagaraDataInterfaceRigidMeshCollisionQuery::InitPerInstanceData(void* PerInstanceData, FNiagaraSystemInstance* SystemInstance)
 {
 	FNDIRigidMeshCollisionData* InstanceData = new (PerInstanceData) FNDIRigidMeshCollisionData();
diff --git a/Engine/Plugins/FX/Niagara/Source/Niagara/Public/NiagaraDataInterfaceRigidMeshCollisionQuery.h b/Engine/Plugins/FX/Niagara/Source/Niagara/Public/NiagaraDataInterfaceRigidMeshCollisionQuery.h
--- a/Engine/Plugins/FX/Niagara/Source/Niagara/Public/NiagaraDataInterfaceRigidMeshCollisionQuery.h
+++ b/Engine/Plugins/FX/Niagara/Source/Niagara/Public/NiagaraDataInterfaceRigidMeshCollisionQuery.h
@@ -80,15 +80,17 @@ struct FNDIRigidMeshCollisionData
 	/** Release the buffers */
 	void Release();
 
+	ETickingGroup ComputeTickingGroup();
+
+	/** The instance ticking group */
+	ETickingGroup TickingGroup;
+
 	/** Physics asset Gpu buffer */
 	FNDIRigidMeshCollisionBuffer* AssetBuffer;
 
 	/** Physics asset Cpu arrays */
 	FNDIRigidMeshCollisionArrays AssetArrays;  
 
-	/** The instance ticking group */
-	ETickingGroup TickingGroup;
-
 	/** Static Mesh Components **/
 	TArray<AStaticMeshActor*> StaticMeshActors;
 };
