remove todos, etc

diff --git a/Engine/Plugins/FX/Niagara/Source/Niagara/Private/NiagaraDataInterfaceRigidMeshCollisionQuery.cpp b/Engine/Plugins/FX/Niagara/Source/Niagara/Private/NiagaraDataInterfaceRigidMeshCollisionQuery.cpp
--- a/Engine/Plugins/FX/Niagara/Source/Niagara/Private/NiagaraDataInterfaceRigidMeshCollisionQuery.cpp
+++ b/Engine/Plugins/FX/Niagara/Source/Niagara/Private/NiagaraDataInterfaceRigidMeshCollisionQuery.cpp
@@ -373,20 +373,7 @@ void FNDIRigidMeshCollisionData::Init(UNiagaraDataInterfaceRigidMeshCollisionQue
 				(Interface->Tag == FString("") || (Interface->Tag != FString("") && StaticMeshActor->Tags.Contains(FName(Interface->Tag)))))
 			{
 				StaticMeshActors.Add(StaticMeshActor);
-			}
-			// Only deal with static meshes that are movable
-			// #todo(dmp): also nonmoveable?
-
-			// Get local bounds
-			// FVector MinLocalBounds;
-			// FVector MaxLocalBounds;
-			// StaticMeshActor->GetStaticMeshComponent()->GetLocalBounds(MinLocalBounds, MaxLocalBounds);
-
-			// Initial Transform 
-			//FTransform InitialTransform = StaticMeshActor->GetTransform();
-
-			// Use Body setup if it exists
-			//UBodySetup *BodySetup = StaticMeshComponent->GetBodySetup();
+			}			
 		}
 	}
 
