Niagara - Refactor precompile into 2 steps, a single precompile for the system compile, and then separate duplicate steps for spawn/update script pairs.  Ideally each script could have a separate request, but the current implementation of the precomied histories require that spawn and update are duplicate together before traversal.

diff --git a/Engine/Plugins/FX/Niagara/Source/Niagara/Classes/NiagaraScript.h b/Engine/Plugins/FX/Niagara/Source/Niagara/Classes/NiagaraScript.h
--- a/Engine/Plugins/FX/Niagara/Source/Niagara/Classes/NiagaraScript.h
+++ b/Engine/Plugins/FX/Niagara/Source/Niagara/Classes/NiagaraScript.h
@@ -22,6 +22,7 @@
 
 class UNiagaraDataInterface;
 class FNiagaraCompileRequestDataBase;
+class FNiagaraCompileRequestDuplicateDataBase;
 class UNiagaraConvertInPlaceUtilityBase;
 
 #define NIAGARA_INVALID_MEMORY (0xBA)
@@ -1001,7 +1002,7 @@ public:
 
 	/** Request an asynchronous compile for the script, possibly forcing it to compile. The output values are the compilation id of the data as well as the async handle to 
 		gather up the results with. The function returns whether or not any compiles were actually issued. */
-	NIAGARA_API bool RequestExternallyManagedAsyncCompile(const TSharedPtr<FNiagaraCompileRequestDataBase, ESPMode::ThreadSafe>& RequestData, FNiagaraVMExecutableDataId& OutCompileId, uint32& OutAsyncHandle);
+	NIAGARA_API bool RequestExternallyManagedAsyncCompile(const TSharedPtr<FNiagaraCompileRequestDataBase, ESPMode::ThreadSafe>& RequestData, const TSharedPtr<FNiagaraCompileRequestDuplicateDataBase, ESPMode::ThreadSafe>& RequestDuplicateData, FNiagaraVMExecutableDataId& OutCompileId, uint32& OutAsyncHandle);
 
 	/** Builds the DDC string for the derived data cache using the supplied CompiledId */
 	static FString BuildNiagaraDDCKeyString(const FNiagaraVMExecutableDataId& CompileId);
diff --git a/Engine/Plugins/FX/Niagara/Source/Niagara/Classes/NiagaraScriptSourceBase.h b/Engine/Plugins/FX/Niagara/Source/Niagara/Classes/NiagaraScriptSourceBase.h
--- a/Engine/Plugins/FX/Niagara/Source/Niagara/Classes/NiagaraScriptSourceBase.h
+++ b/Engine/Plugins/FX/Niagara/Source/Niagara/Classes/NiagaraScriptSourceBase.h
@@ -30,13 +30,24 @@ class FNiagaraCompileRequestDataBase
 {
 public:
 	virtual ~FNiagaraCompileRequestDataBase() {}
-	virtual bool GatherPreCompiledVariables(const FString& InNamespaceFilter, TArray<FNiagaraVariable>& OutVars) = 0;
-	virtual void GetReferencedObjects(TArray<UObject*>& Objects) = 0;
-	virtual const TMap<FName, UNiagaraDataInterface*>& GetObjectNameMap() = 0;
+	virtual void GatherPreCompiledVariables(const FString& InNamespaceFilter, TArray<FNiagaraVariable>& OutVars) const = 0;
 	virtual int32 GetDependentRequestCount() const = 0;
 	virtual TSharedPtr<FNiagaraCompileRequestDataBase, ESPMode::ThreadSafe> GetDependentRequest(int32 Index) = 0;
+	virtual const FNiagaraCompileRequestDataBase* GetDependentRequest(int32 Index) const = 0;
 	virtual FName ResolveEmitterAlias(FName VariableName) const = 0;
 	virtual bool GetUseRapidIterationParams() const = 0;
+};
+
+/** External reference to the compile request data generated.*/
+class FNiagaraCompileRequestDuplicateDataBase
+{
+public:
+	virtual ~FNiagaraCompileRequestDuplicateDataBase() {}
+	virtual bool IsDuplicateDataFor(UNiagaraSystem* InSystem, UNiagaraEmitter* InEmitter, UNiagaraScript* InScript) const = 0;
+	virtual void GetDuplicatedObjects(TArray<UObject*>& Objects) = 0;
+	virtual const TMap<FName, UNiagaraDataInterface*>& GetObjectNameMap() = 0;
+	virtual int32 GetDependentRequestCount() const = 0;
+	virtual TSharedPtr<FNiagaraCompileRequestDuplicateDataBase, ESPMode::ThreadSafe> GetDependentRequest(int32 Index) = 0;
 	virtual void ReleaseCompilationCopies() = 0;
 };
 
diff --git a/Engine/Plugins/FX/Niagara/Source/Niagara/Classes/NiagaraSystem.h b/Engine/Plugins/FX/Niagara/Source/Niagara/Classes/NiagaraSystem.h
--- a/Engine/Plugins/FX/Niagara/Source/Niagara/Classes/NiagaraSystem.h
+++ b/Engine/Plugins/FX/Niagara/Source/Niagara/Classes/NiagaraSystem.h
@@ -175,16 +175,17 @@ enum class ENiagaraCompilationState : uint8
 struct FNiagaraLazyPrecompileReference
 {
 	TSharedPtr<FNiagaraCompileRequestDataBase, ESPMode::ThreadSafe> GetPrecompileData(UNiagaraScript* ForScript);
+	TSharedPtr<FNiagaraCompileRequestDuplicateDataBase, ESPMode::ThreadSafe> GetPrecompileDuplicateData(UNiagaraEmitter* OwningEmitter, UNiagaraScript* TargetScript);
 	
 	UNiagaraSystem* System = nullptr;
 	TArray<UNiagaraScript*> Scripts;
 	TMap<UNiagaraScript*, int32> EmitterScriptIndex;
 	TArray<TObjectPtr<UObject>> CompilationRootObjects;
-	TArray<FNiagaraVariable> EncounteredExposedVars;
 
 private:
 	TSharedPtr<FNiagaraCompileRequestDataBase, ESPMode::ThreadSafe> SystemPrecompiledData;
 	TMap<UNiagaraScript*, TSharedPtr<FNiagaraCompileRequestDataBase, ESPMode::ThreadSafe>> EmitterMapping;
+	TArray<TSharedPtr<FNiagaraCompileRequestDuplicateDataBase, ESPMode::ThreadSafe>> PrecompileDuplicateDatas;
 };
 
 struct FNiagaraAsyncTaskSharedData
@@ -210,7 +211,9 @@ public:
 	bool bWaitForCompileJob = false;
 	bool bUsedShaderCompilerWorker = false;
 	bool bFetchedGCObjects = false;
+	UNiagaraSystem* OwningSystem;
 	FEmitterCompiledScriptPair ScriptPair;
+	TArray<FNiagaraVariable> EncounteredExposedVars;
 
 	ENiagaraCompilationState CurrentState;
 
@@ -220,8 +223,9 @@ public:
 	// this data is shared between the ddc thread and the game thread that starts the compilation
 	TSharedPtr<FNiagaraLazyPrecompileReference, ESPMode::ThreadSafe> PrecompileReference;
 	TSharedPtr<FNiagaraCompileRequestDataBase, ESPMode::ThreadSafe> ComputedPrecompileData;
+	TSharedPtr<FNiagaraCompileRequestDuplicateDataBase, ESPMode::ThreadSafe> ComputedPrecompileDuplicateData;
 
-	FNiagaraAsyncCompileTask(FString InAssetPath, const FEmitterCompiledScriptPair& InScriptPair);
+	FNiagaraAsyncCompileTask(UNiagaraSystem* InOwningSystem, FString InAssetPath, const FEmitterCompiledScriptPair& InScriptPair);
 
 	const FEmitterCompiledScriptPair& GetScriptPair() const { return ScriptPair; };
 	void WaitAndResolveResult();
diff --git a/Engine/Plugins/FX/Niagara/Source/Niagara/Private/NiagaraModule.cpp b/Engine/Plugins/FX/Niagara/Source/Niagara/Private/NiagaraModule.cpp
--- a/Engine/Plugins/FX/Niagara/Source/Niagara/Private/NiagaraModule.cpp
+++ b/Engine/Plugins/FX/Niagara/Source/Niagara/Private/NiagaraModule.cpp
@@ -492,10 +492,10 @@ void INiagaraModule::UnregisterEditorOnlyDataUtilities(TSharedRef<INiagaraEditor
 	EditorOnlyDataUtilities.Reset();
 }
 
-int32 INiagaraModule::StartScriptCompileJob(const FNiagaraCompileRequestDataBase* InCompileData, const FNiagaraCompileOptions& InCompileOptions)
+int32 INiagaraModule::StartScriptCompileJob(const FNiagaraCompileRequestDataBase* InCompileData, const FNiagaraCompileRequestDuplicateDataBase* InCompileDuplicateData, const FNiagaraCompileOptions& InCompileOptions)
 {
 	checkf(ScriptCompilerDelegate.IsBound(), TEXT("Create default script compiler delegate not bound."));
-	return ScriptCompilerDelegate.Execute(InCompileData, InCompileOptions);
+	return ScriptCompilerDelegate.Execute(InCompileData, InCompileDuplicateData, InCompileOptions);
 }
 
 TSharedPtr<FNiagaraVMExecutableData> INiagaraModule::GetCompileJobResult(int32 JobID, bool bWait)
@@ -535,22 +535,47 @@ void INiagaraModule::UnregisterCompileResultDelegate(FDelegateHandle DelegateHan
 
 TSharedPtr<FNiagaraCompileRequestDataBase, ESPMode::ThreadSafe> INiagaraModule::Precompile(UObject* Obj, FGuid Version)
 {
-	checkf(ObjectPrecompilerDelegate.IsBound(), TEXT("ObjectPrecompiler delegate not bound."));
-	return ObjectPrecompilerDelegate.Execute(Obj, Version);
+	checkf(PrecompileDelegate.IsBound(), TEXT("Precompile delegate not bound."));
+	return PrecompileDelegate.Execute(Obj, Version);
+}
+
+TSharedPtr<FNiagaraCompileRequestDuplicateDataBase, ESPMode::ThreadSafe> INiagaraModule::PrecompileDuplicate(
+	const FNiagaraCompileRequestDataBase* OwningSystemRequestData,
+	UNiagaraSystem* OwningSystem,
+	UNiagaraEmitter* OwningEmitter,
+	UNiagaraScript* TargetScript,
+	FGuid TargetVersion)
+{
+	checkf(PrecompileDuplicateDelegate.IsBound(), TEXT("PrecompileDuplicate delegate not bound."));
+	return PrecompileDuplicateDelegate.Execute(OwningSystemRequestData, OwningSystem, OwningEmitter, TargetScript, TargetVersion);
 }
 
 FDelegateHandle INiagaraModule::RegisterPrecompiler(FOnPrecompile PreCompiler)
 {
-	checkf(ObjectPrecompilerDelegate.IsBound() == false, TEXT("Only one handler is allowed for the ObjectPrecompiler delegate"));
-	ObjectPrecompilerDelegate = PreCompiler;
-	return ObjectPrecompilerDelegate.GetHandle();
+	checkf(PrecompileDelegate.IsBound() == false, TEXT("Only one handler is allowed for the Precompiler delegate"));
+	PrecompileDelegate = PreCompiler;
+	return PrecompileDelegate.GetHandle();
 }
 
 void INiagaraModule::UnregisterPrecompiler(FDelegateHandle DelegateHandle)
 {
-	checkf(ObjectPrecompilerDelegate.IsBound(), TEXT("ObjectPrecompiler is not registered"));
-	checkf(ObjectPrecompilerDelegate.GetHandle() == DelegateHandle, TEXT("Can only unregister the ObjectPrecompiler delegate with the handle it was registered with."));
-	ObjectPrecompilerDelegate.Unbind();
+	checkf(PrecompileDelegate.IsBound(), TEXT("ObjectPrecompiler is not registered"));
+	checkf(PrecompileDelegate.GetHandle() == DelegateHandle, TEXT("Can only unregister the Precompiler delegate with the handle it was registered with."));
+	PrecompileDelegate.Unbind();
+}
+
+FDelegateHandle INiagaraModule::RegisterPrecompileDuplicator(FOnPrecompileDuplicate PreCompileDuplicator)
+{
+	checkf(PrecompileDuplicateDelegate.IsBound() == false, TEXT("Only one handler is allowed for the PrecompileDuplicator delegate"));
+	PrecompileDuplicateDelegate = PreCompileDuplicator;
+	return PrecompileDuplicateDelegate.GetHandle();
+}
+
+void INiagaraModule::UnregisterPrecompileDuplicator(FDelegateHandle DelegateHandle)
+{
+	checkf(PrecompileDuplicateDelegate.IsBound(), TEXT("PrecompileDuplicator is not registered"));
+	checkf(PrecompileDuplicateDelegate.GetHandle() == DelegateHandle, TEXT("Can only unregister the PrecompileDuplicator delegate with the handle it was registered with."));
+	PrecompileDuplicateDelegate.Unbind();
 }
 
 void INiagaraModule::OnAssetLoaded(UObject* Asset)
@@ -1468,4 +1493,4 @@ void INiagaraModule::GeneratePerfBaselines(TArray<UNiagaraEffectType*>& Baseline
 
 #endif
 
-#undef LOCTEXT_NAMESPACE
\ No newline at end of file
+#undef LOCTEXT_NAMESPACE
diff --git a/Engine/Plugins/FX/Niagara/Source/Niagara/Private/NiagaraScript.cpp b/Engine/Plugins/FX/Niagara/Source/Niagara/Private/NiagaraScript.cpp
--- a/Engine/Plugins/FX/Niagara/Source/Niagara/Private/NiagaraScript.cpp
+++ b/Engine/Plugins/FX/Niagara/Source/Niagara/Private/NiagaraScript.cpp
@@ -2391,8 +2391,8 @@ void UNiagaraScript::RequestCompile(const FGuid& ScriptVersion, bool bForceCompi
 		TArray<TSharedPtr<FNiagaraCompileRequestDataBase, ESPMode::ThreadSafe>> DependentRequests;
 		TArray<uint8> OutData;
 		INiagaraModule& NiagaraModule = FModuleManager::Get().LoadModuleChecked<INiagaraModule>(TEXT("Niagara"));
-		TSharedPtr<FNiagaraCompileRequestDataBase, ESPMode::ThreadSafe> RequestData = NiagaraModule.Precompile(this, ScriptVersion);
 
+		TSharedPtr<FNiagaraCompileRequestDataBase, ESPMode::ThreadSafe> RequestData = NiagaraModule.Precompile(this, ScriptVersion);
 		if (RequestData.IsValid() == false)
 		{
 			COOK_STAT(Timer.TrackCyclesOnly());
@@ -2400,6 +2400,16 @@ void UNiagaraScript::RequestCompile(const FGuid& ScriptVersion, bool bForceCompi
 			return;
 		}
 
+		TArray<FNiagaraVariable> PrecompiledVariables;
+		RequestData->GatherPreCompiledVariables(FString(), PrecompiledVariables);
+		TSharedPtr<FNiagaraCompileRequestDuplicateDataBase, ESPMode::ThreadSafe> RequestDuplicateData = NiagaraModule.PrecompileDuplicate(RequestData.Get(), nullptr, nullptr, this, ScriptVersion);
+		if (RequestDuplicateData.IsValid() == false)
+		{
+			COOK_STAT(Timer.TrackCyclesOnly());
+			UE_LOG(LogNiagara, Error, TEXT("Failed to precompile duplicate %s.  This is due to unexpected invalid or broken data.  Additional details should be in the log."), *GetPathName());
+			return;
+		}
+
 		// check the ddc first
 		if (GetDerivedDataCacheRef().GetSynchronous(*GetNiagaraDDCKeyString(ScriptVersion), OutData, GetPathName()))
 		{
@@ -2407,20 +2417,20 @@ void UNiagaraScript::RequestCompile(const FGuid& ScriptVersion, bool bForceCompi
 			if (BinaryToExecData(this, OutData, ExeData))
 			{
 				COOK_STAT(Timer.AddHit(OutData.Num()));
-				SetVMCompilationResults(LastGeneratedVMId, ExeData, FString(), RequestData->GetObjectNameMap());
+				SetVMCompilationResults(LastGeneratedVMId, ExeData, FString(), RequestDuplicateData->GetObjectNameMap());
 				return;
 			}
 		}
 
 		ActiveCompileRoots.Empty();
-		RequestData->GetReferencedObjects(ActiveCompileRoots);
+		RequestDuplicateData->GetDuplicatedObjects(ActiveCompileRoots);
 
 		FNiagaraCompileOptions Options(GetUsage(), GetUsageId(), ScriptData->ModuleUsageBitmask, GetPathName(), GetFullName(), GetName());
-		int32 JobHandle = NiagaraModule.StartScriptCompileJob(RequestData.Get(), Options);
+		int32 JobHandle = NiagaraModule.StartScriptCompileJob(RequestData.Get(), RequestDuplicateData.Get(), Options);
 		TSharedPtr<FNiagaraVMExecutableData> ExeData = NiagaraModule.GetCompileJobResult(JobHandle, true);
 		if (ExeData)
 		{
-			SetVMCompilationResults(LastGeneratedVMId, *ExeData, FString(), RequestData->GetObjectNameMap());
+			SetVMCompilationResults(LastGeneratedVMId, *ExeData, FString(), RequestDuplicateData->GetObjectNameMap());
 			// save result to the ddc
 			if (ExecToBinaryData(this, OutData, *ExeData))
 			{
@@ -2436,7 +2446,7 @@ void UNiagaraScript::RequestCompile(const FGuid& ScriptVersion, bool bForceCompi
 	}
 }
 
-bool UNiagaraScript::RequestExternallyManagedAsyncCompile(const TSharedPtr<FNiagaraCompileRequestDataBase, ESPMode::ThreadSafe>& RequestData, FNiagaraVMExecutableDataId& OutCompileId, uint32& OutAsyncHandle)
+bool UNiagaraScript::RequestExternallyManagedAsyncCompile(const TSharedPtr<FNiagaraCompileRequestDataBase, ESPMode::ThreadSafe>& RequestData, const TSharedPtr<FNiagaraCompileRequestDuplicateDataBase, ESPMode::ThreadSafe>& RequestDuplicateData, FNiagaraVMExecutableDataId& OutCompileId, uint32& OutAsyncHandle)
 {
 	COOK_STAT(auto Timer = NiagaraScriptCookStats::UsageStats.TimeSyncWork());
 	COOK_STAT(Timer.TrackCyclesOnly());
@@ -2462,7 +2472,7 @@ bool UNiagaraScript::RequestExternallyManagedAsyncCompile(const TSharedPtr<FNiag
 		FNiagaraCompileOptions Options(GetUsage(), GetUsageId(), ScriptData->ModuleUsageBitmask, GetPathName(), GetFullName(), GetName());
 		Options.AdditionalDefines = LastGeneratedVMId.AdditionalDefines;
 		Options.AdditionalVariables = LastGeneratedVMId.AdditionalVariables;
-		OutAsyncHandle = NiagaraModule.StartScriptCompileJob(RequestData.Get(), Options);
+		OutAsyncHandle = NiagaraModule.StartScriptCompileJob(RequestData.Get(), RequestDuplicateData.Get(), Options);
 		UE_LOG(LogNiagara, Verbose, TEXT("Script '%s' is requesting compile.."), *GetFullName());
 		return true;
 	}
diff --git a/Engine/Plugins/FX/Niagara/Source/Niagara/Private/NiagaraSystem.cpp b/Engine/Plugins/FX/Niagara/Source/Niagara/Private/NiagaraSystem.cpp
--- a/Engine/Plugins/FX/Niagara/Source/Niagara/Private/NiagaraSystem.cpp
+++ b/Engine/Plugins/FX/Niagara/Source/Niagara/Private/NiagaraSystem.cpp
@@ -147,9 +147,6 @@ TSharedPtr<FNiagaraCompileRequestDataBase, ESPMode::ThreadSafe> FNiagaraLazyPrec
 		INiagaraModule& NiagaraModule = FModuleManager::Get().LoadModuleChecked<INiagaraModule>(TEXT("Niagara"));
 		SystemPrecompiledData = NiagaraModule.Precompile(Container, FGuid());
 
-		// Grab the list of user variables that were actually encountered so that we can add to them later.
-		SystemPrecompiledData->GatherPreCompiledVariables(TEXT("User"), EncounteredExposedVars);
-		SystemPrecompiledData->GetReferencedObjects(CompilationRootObjects);
 		const TArray<FNiagaraEmitterHandle>& EmitterHandles = System->GetEmitterHandles();
 		for (int32 i = 0; i < EmitterHandles.Num(); i++)
 		{
@@ -157,8 +154,6 @@ TSharedPtr<FNiagaraCompileRequestDataBase, ESPMode::ThreadSafe> FNiagaraLazyPrec
 			if (Handle.GetInstance() && Handle.GetIsEnabled())
 			{
 				TSharedPtr<FNiagaraCompileRequestDataBase, ESPMode::ThreadSafe> EmitterPrecompiledData = SystemPrecompiledData->GetDependentRequest(i);
-				EmitterPrecompiledData->GetReferencedObjects(CompilationRootObjects);
-				
 				TArray<UNiagaraScript*> EmitterScripts;
 				Handle.GetInstance()->GetScripts(EmitterScripts, false, true);
 				check(EmitterScripts.Num() > 0);
@@ -166,9 +161,6 @@ TSharedPtr<FNiagaraCompileRequestDataBase, ESPMode::ThreadSafe> FNiagaraLazyPrec
 				{
 					EmitterMapping.Add(EmitterScript, EmitterPrecompiledData);
 				}
-
-				// Add the emitter's User variables to the encountered list to expose for later.
-				EmitterPrecompiledData->GatherPreCompiledVariables(TEXT("User"), EncounteredExposedVars);
 			}
 		}
 	}
@@ -191,8 +183,37 @@ TSharedPtr<FNiagaraCompileRequestDataBase, ESPMode::ThreadSafe> FNiagaraLazyPrec
 	return *EmitterPrecompileData;
 }
 
-FNiagaraAsyncCompileTask::FNiagaraAsyncCompileTask(FString InAssetPath, const FEmitterCompiledScriptPair& InScriptPair)
+TSharedPtr<FNiagaraCompileRequestDuplicateDataBase, ESPMode::ThreadSafe> FNiagaraLazyPrecompileReference::GetPrecompileDuplicateData(UNiagaraEmitter* OwningEmitter, UNiagaraScript* TargetScript)
+{
+	TSharedPtr<FNiagaraCompileRequestDuplicateDataBase, ESPMode::ThreadSafe> PrecompileDuplicateData;
+	for (TSharedPtr<FNiagaraCompileRequestDuplicateDataBase, ESPMode::ThreadSafe> ExistingPrecompileDuplicateData : PrecompileDuplicateDatas)
+	{
+		if (ExistingPrecompileDuplicateData->IsDuplicateDataFor(System, OwningEmitter, TargetScript))
+		{
+			PrecompileDuplicateData = ExistingPrecompileDuplicateData;
+		}
+	}
+
+	if (PrecompileDuplicateData.IsValid() == false)
+	{
+		INiagaraModule& NiagaraModule = FModuleManager::Get().LoadModuleChecked<INiagaraModule>(TEXT("Niagara"));
+
+		PrecompileDuplicateData = NiagaraModule.PrecompileDuplicate(SystemPrecompiledData.Get(), System, OwningEmitter, TargetScript, FGuid());
+		PrecompileDuplicateDatas.Add(PrecompileDuplicateData);
+
+		PrecompileDuplicateData->GetDuplicatedObjects(CompilationRootObjects);
+		for (int32 i = 0; i < PrecompileDuplicateData->GetDependentRequestCount(); i++)
+		{
+			PrecompileDuplicateData->GetDependentRequest(i)->GetDuplicatedObjects(CompilationRootObjects);
+		}
+	}
+
+	return PrecompileDuplicateData;
+}
+
+FNiagaraAsyncCompileTask::FNiagaraAsyncCompileTask(UNiagaraSystem* InOwningSystem, FString InAssetPath, const FEmitterCompiledScriptPair& InScriptPair)
 {
+	OwningSystem = InOwningSystem;
 	AssetPath = InAssetPath;
 	ScriptPair = InScriptPair;
 
@@ -340,9 +361,34 @@ void FNiagaraAsyncCompileTask::PrecompileData()
 	DDCKey = CurrentDDCKey;
 	
 	ComputedPrecompileData = PrecompileReference->GetPrecompileData(ScriptPair.CompiledScript);
-	if (ComputedPrecompileData == nullptr)
+	TSharedPtr<FNiagaraCompileRequestDuplicateDataBase, ESPMode::ThreadSafe> SystemPrecompileDuplicateData = PrecompileReference->GetPrecompileDuplicateData(ScriptPair.Emitter, ScriptPair.CompiledScript);
+
+	if (ComputedPrecompileData == nullptr || SystemPrecompileDuplicateData == nullptr)
 	{
 		AbortTask();
+		return;
+	}
+
+	// Grab the list of user variables that were actually encountered so that we can add to them later.
+	ComputedPrecompileData->GatherPreCompiledVariables(TEXT("User"), EncounteredExposedVars);
+	for(int32 i = 0; i < ComputedPrecompileData->GetDependentRequestCount(); i++)
+	{
+		ComputedPrecompileData->GetDependentRequest(i)->GatherPreCompiledVariables(TEXT("User"), EncounteredExposedVars);
+	}
+
+	// Assign the correct duplicate data for this request.
+	if(ScriptPair.CompiledScript->GetUsage() == ENiagaraScriptUsage::SystemSpawnScript || ScriptPair.CompiledScript->GetUsage() == ENiagaraScriptUsage::SystemUpdateScript)
+	{
+		ComputedPrecompileDuplicateData = SystemPrecompileDuplicateData;
+	}
+	else
+	{
+		UNiagaraEmitter* OwningEmitter = ScriptPair.Emitter;
+		int32 EmitterIndex = OwningSystem->GetEmitterHandles().IndexOfByPredicate([OwningEmitter](const FNiagaraEmitterHandle& EmitterHandle) { return EmitterHandle.GetInstance() == OwningEmitter; });
+		if(EmitterIndex != INDEX_NONE)
+		{
+			ComputedPrecompileDuplicateData = SystemPrecompileDuplicateData->GetDependentRequest(EmitterIndex);
+		}
 	}
 }
 
@@ -351,7 +397,7 @@ void FNiagaraAsyncCompileTask::StartCompileJob()
 	UE_LOG(LogNiagara, Verbose, TEXT("Starting compilation for %s!"), *AssetPath);
 	
 	// Fire off the compile job
-	if (ScriptPair.CompiledScript->RequestExternallyManagedAsyncCompile(ComputedPrecompileData, ScriptPair.CompileId, ScriptPair.PendingJobID))
+	if (ScriptPair.CompiledScript->RequestExternallyManagedAsyncCompile(ComputedPrecompileData, ComputedPrecompileDuplicateData, ScriptPair.CompileId, ScriptPair.PendingJobID))
 	{
 		UE_LOG(LogNiagara, Verbose, TEXT("Designated compilation ID %i for compilation of %s!"), ScriptPair.PendingJobID, *AssetPath);
 		bUsedShaderCompilerWorker = true;
@@ -2554,11 +2600,17 @@ bool UNiagaraSystem::QueryCompileComplete(bool bWait, bool bDoPost, bool bDoNotA
 		
 		FNiagaraSystemCompileRequest& CompileRequest = ActiveCompilations[0];
 		bool bAreWeWaitingForAnyResults = false;
+		double StartTime = FPlatformTime::Seconds();
 		
 		// Check to see if ALL of the sub-requests have resolved.
 		for (auto& AsyncTask : CompileRequest.DDCTasks)
 		{
-			AsyncTask->ProcessCurrentState();
+			// if a task is very expensive we bail and continue on the next frame  
+			if (FPlatformTime::Seconds() - StartTime < 0.125)
+			{
+				AsyncTask->ProcessCurrentState();
+			}
+
 			if (AsyncTask->IsDone())
 			{
 				continue;
@@ -2618,7 +2670,7 @@ bool UNiagaraSystem::QueryCompileComplete(bool bWait, bool bDoPost, bool bDoNotA
 				TMap<FName, UNiagaraDataInterface*> ObjectNameMap;
 				if (AsyncTask->ComputedPrecompileData.IsValid())
 				{
-					ObjectNameMap = AsyncTask->ComputedPrecompileData->GetObjectNameMap(); 
+					ObjectNameMap = AsyncTask->ComputedPrecompileDuplicateData->GetObjectNameMap(); 
 				}
 				else
 				{
@@ -2633,7 +2685,7 @@ bool UNiagaraSystem::QueryCompileComplete(bool bWait, bool bDoPost, bool bDoNotA
 				// Synchronize the variables that we actually encountered during precompile so that we can expose them to the end user.
 				TArray<FNiagaraVariable> OriginalExposedParams;
 				ExposedParameters.GetParameters(OriginalExposedParams);
-				TArray<FNiagaraVariable>& EncounteredExposedVars = AsyncTask->PrecompileReference->EncounteredExposedVars;
+				TArray<FNiagaraVariable>& EncounteredExposedVars = AsyncTask->EncounteredExposedVars;
 				for (int32 i = 0; i < EncounteredExposedVars.Num(); i++)
 				{
 					if (OriginalExposedParams.Contains(EncounteredExposedVars[i]) == false)
@@ -2643,7 +2695,8 @@ bool UNiagaraSystem::QueryCompileComplete(bool bWait, bool bDoPost, bool bDoNotA
 					}
 				}
 				
-				AsyncTask->ComputedPrecompileData->ReleaseCompilationCopies();
+				AsyncTask->ComputedPrecompileDuplicateData->ReleaseCompilationCopies();
+				AsyncTask->ComputedPrecompileDuplicateData.Reset();
 				AsyncTask->ComputedPrecompileData.Reset();
 			}
 		}
@@ -2935,7 +2988,7 @@ bool UNiagaraSystem::RequestCompile(bool bForce, FNiagaraSystemUpdateContext* Op
 						Pair.CompiledScript->ComputeVMCompilationId(NewID, FGuid());
 						Pair.CompileId = NewID;
 
-						TSharedPtr<FNiagaraAsyncCompileTask, ESPMode::ThreadSafe> AsyncTask = MakeShared<FNiagaraAsyncCompileTask, ESPMode::ThreadSafe>(EmitterScript->GetPathName(), Pair);
+						TSharedPtr<FNiagaraAsyncCompileTask, ESPMode::ThreadSafe> AsyncTask = MakeShared<FNiagaraAsyncCompileTask, ESPMode::ThreadSafe>(this, EmitterScript->GetPathName(), Pair);
 						if (EmitterScript->IsCompilable() && !EmitterScript->AreScriptAndSourceSynchronized())
 						{
 							ScriptsNeedingCompile.Add(EmitterScript);
@@ -2963,7 +3016,7 @@ bool UNiagaraSystem::RequestCompile(bool bForce, FNiagaraSystemUpdateContext* Op
 				Pair.CompiledScript->ComputeVMCompilationId(NewID, FGuid());
 				Pair.CompileId = NewID;
 
-				TSharedPtr<FNiagaraAsyncCompileTask, ESPMode::ThreadSafe> AsyncTask = MakeShared<FNiagaraAsyncCompileTask, ESPMode::ThreadSafe>(SystemSpawnScript->GetPathName(), Pair);
+				TSharedPtr<FNiagaraAsyncCompileTask, ESPMode::ThreadSafe> AsyncTask = MakeShared<FNiagaraAsyncCompileTask, ESPMode::ThreadSafe>(this, SystemSpawnScript->GetPathName(), Pair);
 				if (!SystemSpawnScript->AreScriptAndSourceSynchronized())
 				{
 					ScriptsNeedingCompile.Add(SystemSpawnScript);
@@ -2985,7 +3038,7 @@ bool UNiagaraSystem::RequestCompile(bool bForce, FNiagaraSystemUpdateContext* Op
 				Pair.CompiledScript->ComputeVMCompilationId(NewID, FGuid());
 				Pair.CompileId = NewID;
 
-				TSharedPtr<FNiagaraAsyncCompileTask, ESPMode::ThreadSafe> AsyncTask = MakeShared<FNiagaraAsyncCompileTask, ESPMode::ThreadSafe>(SystemUpdateScript->GetPathName(), Pair);
+				TSharedPtr<FNiagaraAsyncCompileTask, ESPMode::ThreadSafe> AsyncTask = MakeShared<FNiagaraAsyncCompileTask, ESPMode::ThreadSafe>(this, SystemUpdateScript->GetPathName(), Pair);
 				if (!SystemUpdateScript->AreScriptAndSourceSynchronized())
 				{
 					ScriptsNeedingCompile.Add(SystemUpdateScript);
diff --git a/Engine/Plugins/FX/Niagara/Source/Niagara/Public/NiagaraModule.h b/Engine/Plugins/FX/Niagara/Source/Niagara/Public/NiagaraModule.h
--- a/Engine/Plugins/FX/Niagara/Source/Niagara/Public/NiagaraModule.h
+++ b/Engine/Plugins/FX/Niagara/Source/Niagara/Public/NiagaraModule.h
@@ -18,6 +18,7 @@ struct FNiagaraVMExecutableData;
 class UNiagaraScript;
 class FNiagaraCompileOptions;
 class FNiagaraCompileRequestDataBase;
+class FNiagaraCompileRequestDuplicateDataBase;
 class INiagaraMergeManager;
 class INiagaraEditorOnlyDataUtilities;
 struct FNiagaraParameterStore;
@@ -34,9 +35,11 @@ class NIAGARA_API INiagaraModule : public IModuleInterface
 public:
 #if WITH_EDITOR
 	typedef TSharedPtr<FNiagaraCompileRequestDataBase, ESPMode::ThreadSafe> CompileRequestPtr;
-	DECLARE_DELEGATE_RetVal_TwoParams(int32, FScriptCompiler, const FNiagaraCompileRequestDataBase*, const FNiagaraCompileOptions&);
+	typedef TSharedPtr<FNiagaraCompileRequestDuplicateDataBase, ESPMode::ThreadSafe> CompileRequestDuplicatePtr;
+	DECLARE_DELEGATE_RetVal_ThreeParams(int32, FScriptCompiler, const FNiagaraCompileRequestDataBase*, const FNiagaraCompileRequestDuplicateDataBase*, const FNiagaraCompileOptions&);
 	DECLARE_DELEGATE_RetVal_TwoParams(TSharedPtr<FNiagaraVMExecutableData>, FCheckCompilationResult, int32, bool);
 	DECLARE_DELEGATE_RetVal_TwoParams(CompileRequestPtr, FOnPrecompile, UObject*, FGuid);
+	DECLARE_DELEGATE_RetVal_FiveParams(CompileRequestDuplicatePtr, FOnPrecompileDuplicate, const FNiagaraCompileRequestDataBase* /*OwningSystemRequestData*/, UNiagaraSystem* /*OwningSystem*/, UNiagaraEmitter* /*OwningEmitter*/, UNiagaraScript* /*TargetScript*/, FGuid /*Version*/);
 #endif
 	DECLARE_DELEGATE_RetVal(void, FOnProcessQueue);
 
@@ -72,7 +75,7 @@ public:
 
 	void UnregisterEditorOnlyDataUtilities(TSharedRef<INiagaraEditorOnlyDataUtilities> InEditorOnlyDataUtilities);
 
-	int32 StartScriptCompileJob(const FNiagaraCompileRequestDataBase* InCompileData, const FNiagaraCompileOptions& InCompileOptions);
+	int32 StartScriptCompileJob(const FNiagaraCompileRequestDataBase* InCompileData, const FNiagaraCompileRequestDuplicateDataBase* InCompileDuplicateData, const FNiagaraCompileOptions& InCompileOptions);
 	TSharedPtr<FNiagaraVMExecutableData> GetCompileJobResult(int32 JobID, bool bWait);
 
 	FDelegateHandle RegisterScriptCompiler(FScriptCompiler ScriptCompiler);
@@ -82,8 +85,16 @@ public:
 	void UnregisterCompileResultDelegate(FDelegateHandle DelegateHandle);
 
 	TSharedPtr<FNiagaraCompileRequestDataBase, ESPMode::ThreadSafe> Precompile(UObject* InObj, FGuid Version);
+	TSharedPtr<FNiagaraCompileRequestDuplicateDataBase, ESPMode::ThreadSafe> PrecompileDuplicate(
+		const FNiagaraCompileRequestDataBase* OwningSystemRequestData,
+		UNiagaraSystem* OwningSystem,
+		UNiagaraEmitter* OwningEmitter,
+		UNiagaraScript* TargetScript,
+		FGuid TargetVersion);
 	FDelegateHandle RegisterPrecompiler(FOnPrecompile PreCompiler);
 	void UnregisterPrecompiler(FDelegateHandle DelegateHandle);
+	FDelegateHandle RegisterPrecompileDuplicator(FOnPrecompileDuplicate PreCompileDuplicator);
+	void UnregisterPrecompileDuplicator(FDelegateHandle DelegateHandle);
 
 	void OnAssetLoaded(UObject* Asset);
 
@@ -215,7 +226,8 @@ public:
 
 	FScriptCompiler ScriptCompilerDelegate;
 	FCheckCompilationResult CompilationResultDelegate;
-	FOnPrecompile ObjectPrecompilerDelegate;
+	FOnPrecompile PrecompileDelegate;
+	FOnPrecompileDuplicate PrecompileDuplicateDelegate;
 #endif
 
 	static int32 EngineEffectsQuality;
diff --git a/Engine/Plugins/FX/Niagara/Source/NiagaraCore/Private/NiagaraCustomVersion.cpp b/Engine/Plugins/FX/Niagara/Source/NiagaraCore/Private/NiagaraCustomVersion.cpp
--- a/Engine/Plugins/FX/Niagara/Source/NiagaraCore/Private/NiagaraCustomVersion.cpp
+++ b/Engine/Plugins/FX/Niagara/Source/NiagaraCore/Private/NiagaraCustomVersion.cpp
@@ -8,5 +8,5 @@ const FGuid FNiagaraCustomVersion::GUID(0xFCF57AFA, 0x50764283, 0xB9A9E658, 0xFF
 // Register the custom version with core
 FCustomVersionRegistration GRegisterNiagaraCustomVersion(FNiagaraCustomVersion::GUID, FNiagaraCustomVersion::LatestVersion, TEXT("NiagaraVer"));
 
-const FGuid FNiagaraCustomVersion::LatestScriptCompileVersion(0xDDF215A2, 0x07D62F47, 0x86F0F7B5, 0xC9DEAD89);
+const FGuid FNiagaraCustomVersion::LatestScriptCompileVersion(0x1448E584, 0xBC404634, 0xA29E0B9B, 0xFDF57669);
 
diff --git a/Engine/Plugins/FX/Niagara/Source/NiagaraEditor/Private/NiagaraAttributeTrimmer.cpp b/Engine/Plugins/FX/Niagara/Source/NiagaraEditor/Private/NiagaraAttributeTrimmer.cpp
--- a/Engine/Plugins/FX/Niagara/Source/NiagaraEditor/Private/NiagaraAttributeTrimmer.cpp
+++ b/Engine/Plugins/FX/Niagara/Source/NiagaraEditor/Private/NiagaraAttributeTrimmer.cpp
@@ -470,7 +470,7 @@ static void TrimAttributes_Safe(TConstArrayView<const FNiagaraParameterMapHistor
 	}));
 }
 
-static void TrimAttributes_Aggressive(const FNiagaraCompileRequestData* CompileData, TConstArrayView<const FNiagaraParameterMapHistory*> LocalParamHistories, TSet<FName>& AttributesToPreserve, TArray<FNiagaraVariable>& Attributes)
+static void TrimAttributes_Aggressive(const FNiagaraCompileRequestDuplicateData* CompileDuplicateData, TConstArrayView<const FNiagaraParameterMapHistory*> LocalParamHistories, TSet<FName>& AttributesToPreserve, TArray<FNiagaraVariable>& Attributes)
 {
 	// variable references hidden in custom hlsl nodes may not be present in a specific stages parameter map
 	// so we consolidate all the variables into one list to use when going through custom hlsl nodes
@@ -492,7 +492,7 @@ static void TrimAttributes_Aggressive(const FNiagaraCompileRequestData* CompileD
 
 	for (const FNiagaraParameterMapHistory* ParamMap : LocalParamHistories)
 	{
-		FFunctionInputResolver InputResolver(CompileData->NodeGraphDeepCopy.Get(), ParamMap->OriginatingScriptUsage);
+		FFunctionInputResolver InputResolver(CompileDuplicateData->NodeGraphDeepCopy.Get(), ParamMap->OriginatingScriptUsage);
 
 		const int32 VariableCount = ParamMap->Variables.Num();
 
@@ -689,7 +689,7 @@ void FHlslNiagaraTranslator::TrimAttributes(const FNiagaraCompileOptions& InComp
 		}
 		else if (AggressiveTrimAttributesEnabled)
 		{
-			TrimAttributes_Aggressive(CompileData, LocalParamHistories, AttributesToPreserve, Attributes);
+			TrimAttributes_Aggressive(CompileDuplicateData, LocalParamHistories, AttributesToPreserve, Attributes);
 		}
 
 		for (const FNiagaraVariable& Attribute : PreTrimmedAttributes)
diff --git a/Engine/Plugins/FX/Niagara/Source/NiagaraEditor/Private/NiagaraCompiler.cpp b/Engine/Plugins/FX/Niagara/Source/NiagaraEditor/Private/NiagaraCompiler.cpp
--- a/Engine/Plugins/FX/Niagara/Source/NiagaraEditor/Private/NiagaraCompiler.cpp
+++ b/Engine/Plugins/FX/Niagara/Source/NiagaraEditor/Private/NiagaraCompiler.cpp
@@ -40,6 +40,7 @@ DECLARE_CYCLE_STAT(TEXT("Niagara - HlslCompiler - CompileScript"), STAT_NiagaraE
 DECLARE_CYCLE_STAT(TEXT("Niagara - HlslCompiler - CompileShader_VectorVM"), STAT_NiagaraEditor_HlslCompiler_CompileShader_VectorVM, STATGROUP_NiagaraEditor);
 DECLARE_CYCLE_STAT(TEXT("Niagara - Module - CompileShader_VectorVMSucceeded"), STAT_NiagaraEditor_HlslCompiler_CompileShader_VectorVMSucceeded, STATGROUP_NiagaraEditor);
 DECLARE_CYCLE_STAT(TEXT("Niagara - ScriptSource - PreCompile"), STAT_NiagaraEditor_ScriptSource_PreCompile, STATGROUP_NiagaraEditor);
+DECLARE_CYCLE_STAT(TEXT("Niagara - ScriptSource - PreCompileDuplicate"), STAT_NiagaraEditor_ScriptSource_PreCompileDuplicate, STATGROUP_NiagaraEditor);
 DECLARE_CYCLE_STAT(TEXT("Niagara - HlslCompiler - TestCompileShader_VectorVM"), STAT_NiagaraEditor_HlslCompiler_TestCompileShader_VectorVM, STATGROUP_NiagaraEditor);
 
 static int32 GbForceNiagaraTranslatorSingleThreaded = 1;
@@ -98,23 +99,21 @@ T* PrecompileDuplicateObject(T const* SourceObject, UObject* Outer, const FName
 
 }
 
-void FNiagaraCompileRequestData::VisitReferencedGraphs(UNiagaraGraph* InSrcGraph, UNiagaraGraph* InDupeGraph, ENiagaraScriptUsage InUsage, FCompileConstantResolver ConstantResolver, bool bNeedsCompilation, TMap<UNiagaraNodeFunctionCall*, ENiagaraScriptUsage> FunctionsWithUsage)
+void FNiagaraCompileRequestDuplicateData::DuplicateReferencedGraphs(UNiagaraGraph* InSrcGraph, UNiagaraGraph* InDupeGraph, ENiagaraScriptUsage InUsage, FCompileConstantResolver ConstantResolver, TMap<UNiagaraNodeFunctionCall*, ENiagaraScriptUsage> FunctionsWithUsage)
 {
 	if (!InDupeGraph && !InSrcGraph)
 	{
 		return;
 	}
-	FunctionData Data;
-	Data.ClonedScript = nullptr;
-	Data.ClonedGraph = InDupeGraph;
-	Data.Usage = InUsage;
-	Data.bHasNumericInputs = false;
-	TArray<FunctionData> MadeArray;
-	MadeArray.Add(Data);
-	PreprocessedFunctions.Add(InSrcGraph, MadeArray);
 
-	bool bStandaloneScript = false;
+	TArray<FDuplicatedGraphData>& DuplicatedGraphDataArray = SharedSourceGraphToDuplicatedGraphsMap->FindOrAdd(InSrcGraph);
+	FDuplicatedGraphData& DuplicatedGraphData = DuplicatedGraphDataArray.AddDefaulted_GetRef();
+	DuplicatedGraphData.ClonedScript = nullptr;
+	DuplicatedGraphData.ClonedGraph = InDupeGraph;
+	DuplicatedGraphData.Usage = InUsage;
+	DuplicatedGraphData.bHasNumericParameters = false;
 
+	bool bStandaloneScript = false;
 	TArray<UNiagaraNodeOutput*> OutputNodes;
 	InDupeGraph->FindOutputNodes(OutputNodes);
 	if (OutputNodes.Num() == 1 && UNiagaraScript::IsStandaloneScript(OutputNodes[0]->GetUsage()))
@@ -122,16 +121,11 @@ void FNiagaraCompileRequestData::VisitReferencedGraphs(UNiagaraGraph* InSrcGraph
 		bStandaloneScript = true;
 	}
 
-	if (bNeedsCompilation)
-	{
-		FNiagaraEditorUtilities::ResolveNumerics(InDupeGraph, bStandaloneScript, ChangedFromNumericVars);
-	}
-	ClonedGraphs.AddUnique(InDupeGraph);
-
-	VisitReferencedGraphsRecursive(InDupeGraph, ConstantResolver, bNeedsCompilation, FunctionsWithUsage);
+	FNiagaraEditorUtilities::ResolveNumerics(InDupeGraph, bStandaloneScript, ChangedFromNumericVars);
+	DuplicateReferencedGraphsRecursive(InDupeGraph, ConstantResolver, FunctionsWithUsage);
 }
 
-void FNiagaraCompileRequestData::VisitReferencedGraphsRecursive(UNiagaraGraph* InGraph, const FCompileConstantResolver& ConstantResolver, bool bNeedsCompilation, TMap<UNiagaraNodeFunctionCall*, ENiagaraScriptUsage> FunctionsWithUsage)
+void FNiagaraCompileRequestDuplicateData::DuplicateReferencedGraphsRecursive(UNiagaraGraph* InGraph, const FCompileConstantResolver& ConstantResolver, TMap<UNiagaraNodeFunctionCall*, ENiagaraScriptUsage> FunctionsWithUsage)
 {
 	if (!InGraph)
 	{
@@ -141,8 +135,6 @@ void FNiagaraCompileRequestData::VisitReferencedGraphsRecursive(UNiagaraGraph* I
 	TArray<UNiagaraNode*> Nodes;
 	InGraph->GetNodesOfClass(Nodes);
 	const UEdGraphSchema_Niagara* Schema = GetDefault<UEdGraphSchema_Niagara>();
-	FPinCollectorArray CallOutputs;
-	FPinCollectorArray CallInputs;
 
 	for (UEdGraphNode* Node : Nodes)
 	{
@@ -156,144 +148,68 @@ void FNiagaraCompileRequestData::VisitReferencedGraphsRecursive(UNiagaraGraph* I
 					UNiagaraDataInterface* DataInterface = InputNode->GetDataInterface();
 					bool bIsParameterMapDataInterface = false;
 					FName DIName = FHlslNiagaraTranslator::GetDataInterfaceName(InputNode->Input.GetName(), EmitterUniqueName, bIsParameterMapDataInterface);
-					UNiagaraDataInterface* Dupe = bNeedsCompilation ? PrecompileDuplicateObject<UNiagaraDataInterface>(DataInterface, GetTransientPackage()) : DataInterface;
-					CopiedDataInterfacesByName.Add(DIName, Dupe);
+					UNiagaraDataInterface* Dupe = PrecompileDuplicateObject<UNiagaraDataInterface>(DataInterface, GetTransientPackage());
+					SharedNameToDuplicatedDataInterfaceMap->Add(DIName, Dupe);
 				}
 				continue;
 			}
 
-			ENiagaraScriptUsage ScriptUsage = ENiagaraScriptUsage::Function;
-
 			UNiagaraNodeFunctionCall* FunctionCallNode = Cast<UNiagaraNodeFunctionCall>(InNode);
 			if (FunctionCallNode)
 			{
 				UNiagaraScript* FunctionScript = FunctionCallNode->FunctionScript;
-				ScriptUsage = FunctionCallNode->GetCalledUsage();
-
-				FCompileConstantResolver FunctionConstantResolver = ConstantResolver;
-				if (FunctionsWithUsage.Contains(FunctionCallNode))
+				bool bFunctionCallOwnsScript = FunctionScript != nullptr && FunctionScript->GetOuter() == FunctionCallNode;
+				if(FunctionScript != nullptr && bFunctionCallOwnsScript == false)
 				{
-					FunctionConstantResolver = ConstantResolver.WithUsage(FunctionsWithUsage[FunctionCallNode]);
-				}
+					// If the function call doesn't already own the script it's pointing at then script needs to be duplicated since it's a referenced
+					// script and will need to be preprocessed.
+					ENiagaraScriptUsage ScriptUsage = FunctionCallNode->GetCalledUsage();
 
-				if (FunctionScript != nullptr)
-				{
-					UNiagaraGraph* FunctionGraph = FunctionCallNode->GetCalledGraph();
+					FCompileConstantResolver FunctionConstantResolver = ConstantResolver;
+					if (FunctionsWithUsage.Contains(FunctionCallNode))
 					{
-						bool bHasNumericParams = FunctionGraph->HasNumericParameters();
-						bool bHasNumericInputs = false;
-
-						// Any function which is not directly owned by it's outer function call node must be cloned since its graph
-						// will be modified in some way with it's internals function calls replaced with cloned references, or with
-						// numeric fixup.  Currently the only scripts which don't need cloning are scripts used by UNiagaraNodeAssignment
-						// module nodes and UNiagaraNodeCustomHlsl expression dynamic input nodes.
-						bool bRequiresClonedScript = FunctionScript->GetOuter()->IsA<UNiagaraNodeFunctionCall>() == false;
+						FunctionConstantResolver = ConstantResolver.WithUsage(FunctionsWithUsage[FunctionCallNode]);
+					}
 
-						CallOutputs.Reset();
-						CallInputs.Reset();
+					UNiagaraGraph* FunctionGraph = FunctionCallNode->GetCalledGraph();
+					bool bHasNumericParams = FunctionGraph->HasNumericParameters();
+					if (bHasNumericParams || SharedSourceGraphToDuplicatedGraphsMap->Contains(FunctionGraph) == false)
+					{
+						// Duplicate the script, the source, and graph
+						UNiagaraScript* DupeScript = DupeScript = FunctionScript->CreateCompilationCopy();
+						TArray<ENiagaraScriptUsage> CompileUsages = { DupeScript->GetUsage() };
+						UNiagaraScriptSource* DupeScriptSource = CastChecked<UNiagaraScriptSource>(DupeScript->GetSource(FunctionCallNode->SelectedScriptVersion))->CreateCompilationCopy(CompileUsages);
+						UNiagaraGraph* DupeGraph = DupeScriptSource->NodeGraph;
+						DupeScript->SetSource(DupeScriptSource, FunctionCallNode->SelectedScriptVersion);
+						
+						// Do any preprocessing necessary
+						FEdGraphUtilities::MergeChildrenGraphsIn(DupeGraph, DupeGraph, /*bRequireSchemaMatch=*/ true);	
+						FPinCollectorArray CallOutputs;
+						FPinCollectorArray CallInputs;
 						InNode->GetOutputPins(CallOutputs);
 						InNode->GetInputPins(CallInputs);
-
-						TArray<UNiagaraNodeInput*> InputNodes;
-						UNiagaraGraph::FFindInputNodeOptions Options;
-						Options.bFilterDuplicates = true;
-						Options.bIncludeParameters = true;
-						Options.bIncludeAttributes = false;
-						Options.bIncludeSystemConstants = false;
-						Options.bIncludeTranslatorConstants = false;
-						FunctionGraph->FindInputNodes(InputNodes, Options);
-
-						for (UNiagaraNodeInput* Input : InputNodes)
-						{
-							if (Input->Input.GetType() == FNiagaraTypeDefinition::GetGenericNumericDef())
-							{
-								bHasNumericInputs = true;
-							}
-						}
-
-						/*UE_LOG(LogNiagaraEditor, Display, TEXT("FunctionGraph = %p SrcScriptName = %s DiffPackage? %s Numerics? %s NumericInputs? %s"), FunctionGraph, *FunctionGraph->GetPathName(),
-							bFromDifferentPackage ? TEXT("yes") : TEXT("no"), bHasNumericParams ? TEXT("yes") : TEXT("no"), bHasNumericInputs ? TEXT("yes") : TEXT("no"));*/
-
-						UNiagaraGraph* ProcessedGraph = nullptr;
-						// We only need to clone a non-numeric graph once.
-
-						if (!PreprocessedFunctions.Contains(FunctionGraph))
-						{
-							UNiagaraScript* DupeScript = nullptr;
-							if (bRequiresClonedScript == false)
-							{
-								DupeScript = FunctionScript;
-								ProcessedGraph = FunctionGraph;
-							}
-							else
-							{
-								DupeScript = FunctionScript;
-								if (bNeedsCompilation)
-								{
-									DupeScript = FunctionScript->CreateCompilationCopy();
-									UNiagaraScriptSource* DupeScriptSource = CastChecked<UNiagaraScriptSource>(DupeScript->GetSource(FunctionCallNode->SelectedScriptVersion))->CreateCompilationCopy();
-									DupeScript->SetSource(DupeScriptSource, FunctionCallNode->SelectedScriptVersion);
-								}
-								ProcessedGraph = Cast<UNiagaraScriptSource>(DupeScript->GetSource(FunctionCallNode->SelectedScriptVersion))->NodeGraph;
-								if (bNeedsCompilation)
-								{
-									FEdGraphUtilities::MergeChildrenGraphsIn(ProcessedGraph, ProcessedGraph, /*bRequireSchemaMatch=*/ true);
-									
-									FNiagaraEditorUtilities::PreprocessFunctionGraph(Schema, ProcessedGraph, CallInputs, CallOutputs, ScriptUsage, FunctionConstantResolver);
-								}
-								FunctionCallNode->FunctionScript = DupeScript;
-							}
-
-							FunctionData Data;
-							Data.ClonedScript = DupeScript;
-							Data.ClonedGraph = ProcessedGraph;
-							Data.CallInputs = CallInputs;
-							Data.CallOutputs = CallOutputs;
-							Data.Usage = ScriptUsage;
-							Data.bHasNumericInputs = bHasNumericInputs;
-							TArray<FunctionData> MadeArray;
-							MadeArray.Add(Data);
-							PreprocessedFunctions.Add(FunctionGraph, MadeArray);
-							VisitReferencedGraphsRecursive(ProcessedGraph, FunctionConstantResolver, bNeedsCompilation, FunctionsWithUsage);
-							ClonedGraphs.AddUnique(ProcessedGraph);
-						}
-						else if (bHasNumericParams)
-						{
-							UNiagaraScript* DupeScript = FunctionScript;
-							if (bNeedsCompilation)
-							{
-								DupeScript = FunctionScript->CreateCompilationCopy();
-								UNiagaraScriptSource* DupeScriptSource = CastChecked<UNiagaraScriptSource>(DupeScript->GetSource(FunctionCallNode->SelectedScriptVersion))->CreateCompilationCopy();
-								DupeScript->SetSource(DupeScriptSource, FunctionCallNode->SelectedScriptVersion);
-							}
-							ProcessedGraph = Cast<UNiagaraScriptSource>(DupeScript->GetSource(FunctionCallNode->SelectedScriptVersion))->NodeGraph;
-							if (bNeedsCompilation)
-							{
-								FEdGraphUtilities::MergeChildrenGraphsIn(ProcessedGraph, ProcessedGraph, /*bRequireSchemaMatch=*/ true);
-								FNiagaraEditorUtilities::PreprocessFunctionGraph(Schema, ProcessedGraph, CallInputs, CallOutputs, ScriptUsage, FunctionConstantResolver);
-							}
-							FunctionCallNode->FunctionScript = DupeScript;
-
-							TArray<FunctionData>* FoundArray = PreprocessedFunctions.Find(FunctionGraph);
-							ClonedGraphs.AddUnique(ProcessedGraph);
-
-							FunctionData Data;
-							Data.ClonedScript = DupeScript;
-							Data.ClonedGraph = ProcessedGraph;
-							Data.CallInputs = CallInputs;
-							Data.CallOutputs = CallOutputs;
-							Data.Usage = ScriptUsage;
-							Data.bHasNumericInputs = bHasNumericInputs;
-
-							FoundArray->Add(Data);
-							VisitReferencedGraphsRecursive(ProcessedGraph, FunctionConstantResolver, bNeedsCompilation, FunctionsWithUsage);
-						}
-						else if(bRequiresClonedScript)
-						{
-							TArray<FunctionData>* FoundArray = PreprocessedFunctions.Find(FunctionGraph);
-							check(FoundArray != nullptr && FoundArray->Num() != 0);
-							FunctionCallNode->FunctionScript = (*FoundArray)[0].ClonedScript;
-						}
+						FNiagaraEditorUtilities::PreprocessFunctionGraph(Schema, DupeGraph, CallInputs, CallOutputs, ScriptUsage, FunctionConstantResolver);
+							
+						// Record the data for this duplicate.
+						TArray<FDuplicatedGraphData>& DuplicatedGraphDataArray = SharedSourceGraphToDuplicatedGraphsMap->FindOrAdd(FunctionGraph);
+						FDuplicatedGraphData& DuplicatedGraphData = DuplicatedGraphDataArray.AddDefaulted_GetRef();
+						DuplicatedGraphData.ClonedScript = DupeScript;
+						DuplicatedGraphData.ClonedGraph = DupeGraph;
+						DuplicatedGraphData.CallInputs = CallInputs;
+						DuplicatedGraphData.CallOutputs = CallOutputs;
+						DuplicatedGraphData.Usage = ScriptUsage;
+						DuplicatedGraphData.bHasNumericParameters = bHasNumericParams;
+
+						// Assign the copied script and process any child scripts.
+						FunctionCallNode->FunctionScript = DupeScript;
+						DuplicateReferencedGraphsRecursive(DupeGraph, FunctionConstantResolver, FunctionsWithUsage);
+					}
+					else
+					{
+						// This graph was already processed and doesn't need per-call duplication so use the previous copy.
+						TArray<FDuplicatedGraphData>* DuplicatedGraphDataArray = SharedSourceGraphToDuplicatedGraphsMap->Find(FunctionGraph);
+						check(DuplicatedGraphDataArray != nullptr && DuplicatedGraphDataArray->Num() != 0);
+						FunctionCallNode->FunctionScript = (*DuplicatedGraphDataArray)[0].ClonedScript;
 					}
 				}
 			}
@@ -301,13 +217,13 @@ void FNiagaraCompileRequestData::VisitReferencedGraphsRecursive(UNiagaraGraph* I
 			UNiagaraNodeEmitter* EmitterNode = Cast<UNiagaraNodeEmitter>(InNode);
 			if (EmitterNode)
 			{
-				for (TSharedPtr<FNiagaraCompileRequestData, ESPMode::ThreadSafe>& Ptr : EmitterData)
+				for (TSharedPtr<FNiagaraCompileRequestDuplicateData, ESPMode::ThreadSafe>& Ptr : EmitterData)
 				{
-					if (Ptr->EmitterUniqueName == EmitterNode->GetEmitterUniqueName() && bNeedsCompilation)
+					if (Ptr->EmitterUniqueName == EmitterNode->GetEmitterUniqueName())
 					{
 						EmitterNode->SyncEnabledState(); // Just to be safe, sync here while we likely still have the handle source.
 						EmitterNode->SetOwnerSystem(nullptr);
-						EmitterNode->SetCachedVariablesForCompilation(*Ptr->EmitterUniqueName, Ptr->NodeGraphDeepCopy.Get(), Ptr->Source.Get());
+						EmitterNode->SetCachedVariablesForCompilation(*Ptr->EmitterUniqueName, Ptr->NodeGraphDeepCopy.Get(), Ptr->SourceDeepCopy.Get());
 					}
 				}
 			}
@@ -315,39 +231,52 @@ void FNiagaraCompileRequestData::VisitReferencedGraphsRecursive(UNiagaraGraph* I
 	}
 }
 
-const TMap<FName, UNiagaraDataInterface*>& FNiagaraCompileRequestData::GetObjectNameMap()
+const TMap<FName, UNiagaraDataInterface*>& FNiagaraCompileRequestDuplicateData::GetObjectNameMap()
 {
-	return CopiedDataInterfacesByName;
+	return *SharedNameToDuplicatedDataInterfaceMap.Get();
 }
 
-void FNiagaraCompileRequestData::MergeInEmitterPrecompiledData(FNiagaraCompileRequestDataBase* InEmitterDataBase)
+UNiagaraDataInterface* FNiagaraCompileRequestDuplicateData::GetDuplicatedDataInterfaceCDOForClass(UClass* Class) const
 {
-	FNiagaraCompileRequestData* InEmitterData = (FNiagaraCompileRequestData*)InEmitterDataBase;
-	if (InEmitterData)
+	if (SharedDataInterfaceClassToDuplicatedCDOMap.IsValid())
 	{
-		CopiedDataInterfacesByName.Append(InEmitterData->CopiedDataInterfacesByName);
+		UNiagaraDataInterface*const* DuplicatedCDOPtr = SharedDataInterfaceClassToDuplicatedCDOMap->Find(Class);
+		if (DuplicatedCDOPtr != nullptr)
+		{
+			return *DuplicatedCDOPtr;
+		}
 	}
+	return nullptr;
 }
 
-
-
 FName FNiagaraCompileRequestData::ResolveEmitterAlias(FName VariableName) const
 {
 	return FNiagaraParameterMapHistory::ResolveEmitterAlias(VariableName, EmitterUniqueName);
 }
 
-void FNiagaraCompileRequestData::GetReferencedObjects(TArray<UObject*>& Objects)
+bool FNiagaraCompileRequestDuplicateData::IsDuplicateDataFor(UNiagaraSystem* InSystem, UNiagaraEmitter* InEmitter, UNiagaraScript* InScript) const
+{
+	return OwningSystem.Get() == InSystem && OwningEmitter.Get() == InEmitter && ValidUsages.Contains(InScript->GetUsage());
+}
+
+void FNiagaraCompileRequestDuplicateData::GetDuplicatedObjects(TArray<UObject*>& Objects)
 {
+	Objects.Add(SourceDeepCopy.Get());
 	Objects.Add(NodeGraphDeepCopy.Get());
-	TArray<UNiagaraDataInterface*> DIs;
-	CopiedDataInterfacesByName.GenerateValueArray(DIs);
-	for (UNiagaraDataInterface* DI : DIs)
+
+	if (SharedNameToDuplicatedDataInterfaceMap.IsValid())
 	{
-		Objects.Add(DI);
+		TArray<UNiagaraDataInterface*> DIs;
+		SharedNameToDuplicatedDataInterfaceMap->GenerateValueArray(DIs);
+		for (UNiagaraDataInterface* DI : DIs)
+		{
+			Objects.Add(DI);
+		}
 	}
 
+	if (SharedDataInterfaceClassToDuplicatedCDOMap.IsValid())
 	{
-		auto Iter = CDOs.CreateIterator();
+		auto Iter = SharedDataInterfaceClassToDuplicatedCDOMap->CreateIterator();
 		while (Iter)
 		{
 			Objects.Add(Iter.Value());
@@ -355,8 +284,9 @@ void FNiagaraCompileRequestData::GetReferencedObjects(TArray<UObject*>& Objects)
 		}
 	}
 
+	if (SharedSourceGraphToDuplicatedGraphsMap.IsValid())
 	{
-		auto Iter = PreprocessedFunctions.CreateIterator();
+		auto Iter = SharedSourceGraphToDuplicatedGraphsMap->CreateIterator();
 		while (Iter)
 		{
 			for (int32 i = 0; i < Iter.Value().Num(); i++)
@@ -369,21 +299,20 @@ void FNiagaraCompileRequestData::GetReferencedObjects(TArray<UObject*>& Objects)
 	}
 }
 
-bool FNiagaraCompileRequestData::GatherPreCompiledVariables(const FString& InNamespaceFilter, TArray<FNiagaraVariable>& OutVars)
+void FNiagaraCompileRequestData::GatherPreCompiledVariables(const FString& InNamespaceFilter, TArray<FNiagaraVariable>& OutVars) const
 {
-	if (PrecompiledHistories.Num() == 0)
+	if (InNamespaceFilter.Len() == 0)
 	{
-		return false;
+		OutVars.Append(EncounteredVariables);
 	}
-
-	for (const FNiagaraParameterMapHistory& History : PrecompiledHistories)
+	else
 	{
-		for (const FNiagaraVariable& Var : History.Variables)
+		for (const FNiagaraVariable& EncounteredVariable : EncounteredVariables)
 		{
-			if (FNiagaraParameterMapHistory::IsInNamespace(Var, InNamespaceFilter))
+			if (FNiagaraParameterMapHistory::IsInNamespace(EncounteredVariable, InNamespaceFilter))
 			{
-				FNiagaraVariable NewVar = Var;
-				if (NewVar.IsDataAllocated() == false && !Var.IsDataInterface())
+				FNiagaraVariable NewVar = EncounteredVariable;
+				if (NewVar.IsDataAllocated() == false && !NewVar.IsDataInterface())
 				{
 					FNiagaraEditorUtilities::ResetVariableToDefaultValue(NewVar);
 				}
@@ -391,48 +320,44 @@ bool FNiagaraCompileRequestData::GatherPreCompiledVariables(const FString& InNam
 			}
 		}
 	}
-	return true;
 }
 
-void FNiagaraCompileRequestData::DeepCopyGraphs(UNiagaraScriptSource* ScriptSource, ENiagaraScriptUsage InUsage, FCompileConstantResolver ConstantResolver, bool bNeedsCompilation)
+void FNiagaraCompileRequestDuplicateData::DeepCopyGraphs(UNiagaraScriptSource* ScriptSource, ENiagaraScriptUsage InUsage, FCompileConstantResolver ConstantResolver)
 {
 	// Clone the source graph so we can modify it as needed; merging in the child graphs
-	//NodeGraphDeepCopy = CastChecked<UNiagaraGraph>(FEdGraphUtilities::CloneGraph(ScriptSource->NodeGraph, GetTransientPackage(), 0));
-	Source = bNeedsCompilation ? ScriptSource->CreateCompilationCopy() : ScriptSource;
-	NodeGraphDeepCopy = Source->NodeGraph;
-	
-	//double StartTime = FPlatformTime::Seconds();
-	if (bNeedsCompilation)
-	{
-		FEdGraphUtilities::MergeChildrenGraphsIn(NodeGraphDeepCopy.Get(), NodeGraphDeepCopy.Get(), /*bRequireSchemaMatch=*/ true);
-	}
-	VisitReferencedGraphs(ScriptSource->NodeGraph, NodeGraphDeepCopy.Get(), InUsage, ConstantResolver, bNeedsCompilation);
-	//float DeltaTime = (float)(FPlatformTime::Seconds() - StartTime);
-	//if (DeltaTime > 0.01f)
-	//{
-	//	UE_LOG(LogNiagaraEditor, Log, TEXT("\tPrecompile VisitReferencedGraphs %s took %f sec"), *ScriptSource->GetPathName(), DeltaTime);
-	//}
+	SourceDeepCopy = ScriptSource->CreateCompilationCopy(ValidUsages);
+	NodeGraphDeepCopy = SourceDeepCopy->NodeGraph;
+	FEdGraphUtilities::MergeChildrenGraphsIn(NodeGraphDeepCopy.Get(), NodeGraphDeepCopy.Get(), /*bRequireSchemaMatch=*/ true);
+	DuplicateReferencedGraphs(ScriptSource->NodeGraph, NodeGraphDeepCopy.Get(), InUsage, ConstantResolver);
 }
 
-void FNiagaraCompileRequestData::DeepCopyGraphs(UNiagaraScriptSource* ScriptSource, UNiagaraEmitter* Emitter, bool bNeedsCompilation)
+void FNiagaraCompileRequestDuplicateData::DeepCopyGraphs(UNiagaraEmitter* Emitter)
 {
-	Source = bNeedsCompilation ? ScriptSource->CreateCompilationCopy() : ScriptSource;
-	NodeGraphDeepCopy = Source->NodeGraph;
-	if (bNeedsCompilation)
-	{
-		FEdGraphUtilities::MergeChildrenGraphsIn(NodeGraphDeepCopy.Get(), NodeGraphDeepCopy.Get(), /*bRequireSchemaMatch=*/ true);
-	}
+	UNiagaraScriptSource* ScriptSource = CastChecked<UNiagaraScriptSource>(Emitter->GraphSource);
+
+	SourceDeepCopy = ScriptSource->CreateCompilationCopy(ValidUsages);
+	NodeGraphDeepCopy = SourceDeepCopy->NodeGraph;
+	FEdGraphUtilities::MergeChildrenGraphsIn(NodeGraphDeepCopy.Get(), NodeGraphDeepCopy.Get(), /*bRequireSchemaMatch=*/ true);
+
 	TMap<UNiagaraNodeFunctionCall*, ENiagaraScriptUsage> FunctionsWithUsage;
-	for (UEdGraphNode* Node : ScriptSource->NodeGraph->Nodes)
+	TArray<UNiagaraNodeOutput*> OutputNodes;
+	NodeGraphDeepCopy->GetNodesOfClass(OutputNodes);
+	for (UNiagaraNodeOutput* OutputNode : OutputNodes)
 	{
-		if (UNiagaraNodeFunctionCall* FunctionCall = Cast<UNiagaraNodeFunctionCall>(Node))
+		TArray<UNiagaraNode*> TraversedNodes;
+		NodeGraphDeepCopy->BuildTraversal(TraversedNodes, OutputNode);
+		for (UNiagaraNode* TraversedNode : TraversedNodes)
 		{
-			UNiagaraNodeOutput* OutputNode = FNiagaraStackGraphUtilities::GetEmitterOutputNodeForStackNode(*FunctionCall);
-			FunctionsWithUsage.Add(FunctionCall, OutputNode ? OutputNode->GetUsage() : ENiagaraScriptUsage::EmitterSpawnScript);
+			UNiagaraNodeFunctionCall* FunctionCallNode = Cast<UNiagaraNodeFunctionCall>(TraversedNode);
+			if (FunctionCallNode != nullptr)
+			{
+				FunctionsWithUsage.Add(FunctionCallNode, OutputNode->GetUsage());
+			}
 		}
 	}
+
 	FCompileConstantResolver ConstantResolver(Emitter, ENiagaraScriptUsage::EmitterSpawnScript);
-	VisitReferencedGraphs(ScriptSource->NodeGraph, NodeGraphDeepCopy.Get(), ENiagaraScriptUsage::EmitterSpawnScript, ConstantResolver, bNeedsCompilation, FunctionsWithUsage);
+	DuplicateReferencedGraphs(ScriptSource->NodeGraph, NodeGraphDeepCopy.Get(), ENiagaraScriptUsage::EmitterSpawnScript, ConstantResolver, FunctionsWithUsage);
 }
 
 
@@ -494,80 +419,66 @@ void FNiagaraCompileRequestData::AddRapidIterationParameters(const FNiagaraParam
 	}
 }
 
-void FNiagaraCompileRequestData::ReleaseCompilationCopies()
+void FNiagaraCompileRequestDuplicateData::ReleaseCompilationCopies()
 {
 	// clean up graph copies
-	for (auto& ClonedGraph : ClonedGraphs)
-	{
-		if (ClonedGraph.IsValid())
-		{
-			ClonedGraph->ReleaseCompilationCopy();
-		}
-	}
-	ClonedGraphs.Empty();
-	
-	for (auto& EmitterCompileData : EmitterData)
+	if (SharedSourceGraphToDuplicatedGraphsMap.IsValid())
 	{
-		if (EmitterCompileData.IsValid())
+		for (const TPair<const UNiagaraGraph*, TArray<FDuplicatedGraphData>>& SourceGraphToDuplicatedGraphs : *(SharedSourceGraphToDuplicatedGraphsMap.Get()))
 		{
-			EmitterCompileData->ReleaseCompilationCopies();
+			for (const FDuplicatedGraphData& DuplicatedGraphData : SourceGraphToDuplicatedGraphs.Value)
+			{
+				DuplicatedGraphData.ClonedGraph->ReleaseCompilationCopy();
+			}
 		}
+		SharedSourceGraphToDuplicatedGraphsMap->Empty();
 	}
-	EmitterData.Empty();
 
 	// clean up script source
-	if (Source.IsValid())
+	if (SourceDeepCopy.IsValid())
 	{
-		Source->ReleaseCompilationCopy();
+		SourceDeepCopy->ReleaseCompilationCopy();
 	}
-	Source = nullptr;
+	SourceDeepCopy = nullptr;
 }
 
-void FNiagaraCompileRequestData::FinishPrecompile(const TArray<FNiagaraVariable>& EncounterableVariables, ENiagaraScriptUsage InUsage, FCompileConstantResolver ConstantResolver, const TArray<UNiagaraSimulationStageBase*>* SimStages)
+void FNiagaraCompileRequestData::FinishPrecompile(const TArray<FNiagaraVariable>& EncounterableVariables, FCompileConstantResolver ConstantResolver, const TArray<ENiagaraScriptUsage>& UsagesToProcess, const TArray<class UNiagaraSimulationStageBase*>* SimStages)
 {
-	//double StartTime = FPlatformTime::Seconds();
 	{
 		ENiagaraScriptCompileStatusEnum = StaticEnum<ENiagaraScriptCompileStatus>();
 		ENiagaraScriptUsageEnum = StaticEnum<ENiagaraScriptUsage>();
 
-		PrecompiledHistories.Empty();
-
 		TArray<UNiagaraNodeOutput*> OutputNodes;
-		NodeGraphDeepCopy->FindOutputNodes(OutputNodes);
-		PrecompiledHistories.Empty();
+		if (Source.IsValid() && Source->NodeGraph != nullptr)
+		{
+			Source->NodeGraph->FindOutputNodes(OutputNodes);
+		}
 
 		int32 NumSimStageNodes = 0;
 		for (UNiagaraNodeOutput* FoundOutputNode : OutputNodes)
 		{
+			if (UsagesToProcess.Contains(FoundOutputNode->GetUsage()) == false)
+			{
+				continue;
+			}
+
 			FName SimStageName;
 			bool bStageEnabled = true;
-			if (FoundOutputNode->GetUsage() == ENiagaraScriptUsage::ParticleSimulationStageScript && SimStages)
+			if (FoundOutputNode->GetUsage() == ENiagaraScriptUsage::ParticleSimulationStageScript && bSimulationStagesEnabled && SimStages)
 			{
-				if ( bSimulationStagesEnabled )
-				{
-					const FGuid& UsageId = FoundOutputNode->GetUsageId();
+				// Find the simulation stage for this output node.
+				const FGuid& UsageId = FoundOutputNode->GetUsageId();
+				UNiagaraSimulationStageBase*const* MatchingStagePtr = SimStages->FindByPredicate([UsageId](UNiagaraSimulationStageBase* SimStage)
+				{ 
+					return SimStage != nullptr && SimStage->Script != nullptr && SimStage->Script->GetUsageId() == UsageId;
+				});
 
-					// Find the matching simstage to the output node
-					for (UNiagaraSimulationStageBase* SimStage : *SimStages)
-					{
-						if (SimStage && SimStage->Script)
-						{
-							if (SimStage->Script->GetUsageId() == UsageId)
-							{
-								bStageEnabled = SimStage->bEnabled;
-								UNiagaraSimulationStageGeneric* GenericStage = Cast<UNiagaraSimulationStageGeneric>(SimStage);
-								if (GenericStage && SimStage->bEnabled)
-								{
-									SimStageName = (GenericStage->IterationSource == ENiagaraIterationSource::DataInterface) ? GenericStage->DataInterface.BoundVariable.GetName() : FName();
-									break;
-								}
-							}
-						}
-					}
-				}
-				else
+				// Set whether or not the stage is enabled, and get the iteration source name if available.
+				bStageEnabled = MatchingStagePtr != nullptr && (*MatchingStagePtr)->bEnabled;
+				if(bStageEnabled && (*MatchingStagePtr)->IsA<UNiagaraSimulationStageGeneric>())
 				{
-					bStageEnabled = false;
+					UNiagaraSimulationStageGeneric* GenericStage = CastChecked<UNiagaraSimulationStageGeneric>(*MatchingStagePtr);
+					SimStageName = GenericStage->IterationSource == ENiagaraIterationSource::DataInterface ? GenericStage->DataInterface.BoundVariable.GetName() : FName();
 				}
 			}
 
@@ -576,7 +487,7 @@ void FNiagaraCompileRequestData::FinishPrecompile(const TArray<FNiagaraVariable>
 				// Map all for this output node
 				FNiagaraParameterMapHistoryWithMetaDataBuilder Builder;
 				Builder.ConstantResolver = ConstantResolver;
-				Builder.AddGraphToCallingGraphContextStack(NodeGraphDeepCopy.Get());
+				Builder.AddGraphToCallingGraphContextStack(Source->NodeGraph);
 				Builder.RegisterEncounterableVariables(EncounterableVariables);
 
 				FString TranslationName = TEXT("Emitter");
@@ -593,6 +504,7 @@ void FNiagaraCompileRequestData::FinishPrecompile(const TArray<FNiagaraVariable>
 					History.OriginatingScriptUsage = FoundOutputNode->GetUsage();
 					for (FNiagaraVariable& Var : History.Variables)
 					{
+						EncounteredVariables.AddUnique(Var);
 						if (Var.GetType() == FNiagaraTypeDefinition::GetGenericNumericDef())
 						{
 							UE_LOG(LogNiagaraEditor, Log, TEXT("Invalid numeric parameter found! %s"), *Var.GetName().ToString())
@@ -605,14 +517,8 @@ void FNiagaraCompileRequestData::FinishPrecompile(const TArray<FNiagaraVariable>
 					NumSimStageNodes++;
 				}
 
-				PrecompiledHistories.Append(Builder.Histories);
 				Builder.EndTranslation(TranslationName);
 			}
-			else
-			{
-				// Add in a blank spot
-				PrecompiledHistories.Emplace();
-			}
 		}
 
 		if (SimStages && NumSimStageNodes)
@@ -644,68 +550,117 @@ void FNiagaraCompileRequestData::FinishPrecompile(const TArray<FNiagaraVariable>
 				}
 			}
 		}
+	}
+}
 
+void FNiagaraCompileRequestDuplicateData::FinishPrecompileDuplicate(const TArray<FNiagaraVariable>& EncounterableVariables, FCompileConstantResolver ConstantResolver, const TArray<class UNiagaraSimulationStageBase*>* SimStages)
+{
+	PrecompiledHistories.Empty();
+
+	TArray<UNiagaraNodeOutput*> OutputNodes;
+	if (NodeGraphDeepCopy.IsValid())
+	{
+		NodeGraphDeepCopy->FindOutputNodes(OutputNodes);
+	}
 
-		// Generate CDO's for any referenced data interfaces...
-		for (int32 i = 0; i < PrecompiledHistories.Num(); i++)
+	for (UNiagaraNodeOutput* FoundOutputNode : OutputNodes)
+	{
+		FName SimStageName;
+		bool bStageEnabled = true;
+		if (FoundOutputNode->GetUsage() == ENiagaraScriptUsage::ParticleSimulationStageScript && bSimulationStagesEnabled && SimStages)
 		{
-			for (const FNiagaraVariable& Var : PrecompiledHistories[i].Variables)
-			{
-				if (Var.IsDataInterface())
+			// Find the simulation stage for this output node.
+			const FGuid& UsageId = FoundOutputNode->GetUsageId();
+			UNiagaraSimulationStageBase* const* MatchingStagePtr = SimStages->FindByPredicate([UsageId](UNiagaraSimulationStageBase* SimStage)
 				{
-					UClass* Class = const_cast<UClass*>(Var.GetType().GetClass());
-					UObject* Obj = PrecompileDuplicateObject(Class->GetDefaultObject(true), GetTransientPackage());
-					CDOs.Add(Class, Obj);
-				}
+					return SimStage != nullptr && SimStage->Script != nullptr && SimStage->Script->GetUsageId() == UsageId;
+				});
+
+			// Set whether or not the stage is enabled, and get the iteration source name if available.
+			bStageEnabled = MatchingStagePtr != nullptr && (*MatchingStagePtr)->bEnabled;
+			if (bStageEnabled && (*MatchingStagePtr)->IsA<UNiagaraSimulationStageGeneric>())
+			{
+				UNiagaraSimulationStageGeneric* GenericStage = CastChecked<UNiagaraSimulationStageGeneric>(*MatchingStagePtr);
+				SimStageName = GenericStage->IterationSource == ENiagaraIterationSource::DataInterface ? GenericStage->DataInterface.BoundVariable.GetName() : FName();
 			}
 		}
 
-		// Generate CDO's for data interfaces that are passed in to function or dynamic input scripts compiled standalone as we do not have a history
-		if (InUsage == ENiagaraScriptUsage::Function || InUsage == ENiagaraScriptUsage::DynamicInput)
+		if (bStageEnabled)
 		{
-			for (const auto ReferencedGraph : ClonedGraphs)
+			// Map all for this output node
+			FNiagaraParameterMapHistoryWithMetaDataBuilder Builder;
+			Builder.ConstantResolver = ConstantResolver;
+			Builder.AddGraphToCallingGraphContextStack(NodeGraphDeepCopy.Get());
+			Builder.RegisterEncounterableVariables(EncounterableVariables);
+
+			FString TranslationName = TEXT("Emitter");
+			Builder.BeginTranslation(TranslationName);
+			Builder.BeginUsage(FoundOutputNode->GetUsage(), SimStageName);
+			Builder.EnableScriptWhitelist(true, FoundOutputNode->GetUsage());
+			Builder.BuildParameterMaps(FoundOutputNode, true);
+			Builder.EndUsage();
+
+			ensure(Builder.Histories.Num() <= 1);
+
+			for (FNiagaraParameterMapHistory& History : Builder.Histories)
 			{
-				TArray<UNiagaraNodeInput*> InputNodes;
-				TArray<FNiagaraVariable*> InputVariables;
-				ReferencedGraph->FindInputNodes(InputNodes);
-				for (const auto InputNode : InputNodes)
+				History.OriginatingScriptUsage = FoundOutputNode->GetUsage();
+				for (FNiagaraVariable& Var : History.Variables)
 				{
-					InputVariables.Add(&InputNode->Input);
-				}
-
-				for (const auto InputVariable : InputVariables)
-				{
-					if (InputVariable->IsDataInterface())
+					if (Var.GetType() == FNiagaraTypeDefinition::GetGenericNumericDef())
 					{
-						UClass* Class = const_cast<UClass*>(InputVariable->GetType().GetClass());
-						UObject* Obj = PrecompileDuplicateObject(Class->GetDefaultObject(true), GetTransientPackage());
-						CDOs.Add(Class, Obj);
+						UE_LOG(LogNiagaraEditor, Log, TEXT("Invalid numeric parameter found! %s"), *Var.GetName().ToString())
 					}
 				}
 			}
+
+			PrecompiledHistories.Append(Builder.Histories);
+			Builder.EndTranslation(TranslationName);
+		}
+		else
+		{
+			// Add in a blank spot
+			PrecompiledHistories.Emplace();
 		}
 	}
 
-	//float DeltaTime = (float)(FPlatformTime::Seconds() - StartTime);
-	//if (DeltaTime > 0.01f)
-	//{
-	//	UE_LOG(LogNiagaraEditor, Log, TEXT("\tPrecompile FinishPrecompile %s took %f sec"), *ScriptSource->GetPathName(), DeltaTime);
-	//}
-}
+	// Collect classes for any data interfaces found in the duplicated graphs
+	TArray<UClass*> DataInterfaceClasses;
+	for (const TPair<const UNiagaraGraph*, TArray<FDuplicatedGraphData>>& SourceGraphToDuplicatedGraphs : *(SharedSourceGraphToDuplicatedGraphsMap.Get()))
+	{
+		for (const FDuplicatedGraphData& DuplicatedGraphData : SourceGraphToDuplicatedGraphs.Value)
+		{
+			TArray<UNiagaraNodeInput*> InputNodes;
+			DuplicatedGraphData.ClonedGraph->FindInputNodes(InputNodes);
+			for (const UNiagaraNodeInput* InputNode : InputNodes)
+			{
+				if (InputNode->Input.IsDataInterface())
+				{
+					DataInterfaceClasses.AddUnique(InputNode->Input.GetType().GetClass());
+				}
+			}
+		}
+	}
 
-bool ScriptSourceNeedsCompiling(UNiagaraScriptSource* InSource, TArray<UNiagaraScript*>& InCompilingScripts)
-{
-	if (!InSource || InCompilingScripts.Num() == 0)
-		return false;
+	// Collect classes for external encounterable variables
+	for (const FNiagaraVariable& EncounterableVariable : EncounterableVariables)
+	{
+		if (EncounterableVariable.IsDataInterface())
+		{
+			DataInterfaceClasses.AddUnique(EncounterableVariable.GetType().GetClass());
+		}
+	}
 
-	// For now, we are just going to duplicate everything if there are any scripts needing compilation. 
-	// We do this to mitigate any complications around cross-graph communication.
-	return true;
+	// Generate copies of the CDOs for any encountered data interfaces.
+	for (UClass* DataInterfaceClass : DataInterfaceClasses)
+	{
+		UNiagaraDataInterface* DuplicateCDO = CastChecked<UNiagaraDataInterface>(PrecompileDuplicateObject(DataInterfaceClass->GetDefaultObject(true), GetTransientPackage()));
+		SharedDataInterfaceClassToDuplicatedCDOMap->Add(DataInterfaceClass, DuplicateCDO);
+	}
 }
 
 TSharedPtr<FNiagaraCompileRequestDataBase, ESPMode::ThreadSafe> FNiagaraEditorModule::Precompile(UObject* InObj, FGuid Version)
 {
-	TArray<UNiagaraScript*> InCompilingScripts;
 	UNiagaraScript* Script = Cast<UNiagaraScript>(InObj);
 	UNiagaraPrecompileContainer* Container = Cast<UNiagaraPrecompileContainer>(InObj);
 	UNiagaraSystem* System = Cast<UNiagaraSystem>(InObj);
@@ -714,7 +669,6 @@ TSharedPtr<FNiagaraCompileRequestDataBase, ESPMode::ThreadSafe> FNiagaraEditorMo
 	if (Container)
 	{
 		System = Container->System;
-		InCompilingScripts = Container->Scripts;
 		if (System)
 		{
 			LogPackage = System->GetOutermost();
@@ -722,7 +676,6 @@ TSharedPtr<FNiagaraCompileRequestDataBase, ESPMode::ThreadSafe> FNiagaraEditorMo
 	}
 	else if (Script)
 	{
-		InCompilingScripts.Add(Script);
 		LogPackage = Script->GetOutermost();
 	}
 
@@ -746,15 +699,15 @@ TSharedPtr<FNiagaraCompileRequestDataBase, ESPMode::ThreadSafe> FNiagaraEditorMo
 
 	if (Script)
 	{
-		UNiagaraScriptSource* Source = Cast<UNiagaraScriptSource>(Script->GetSource(Version));
-		bool bNeedsCompilation = ScriptSourceNeedsCompiling(Source, InCompilingScripts);
-		BasePtr->DeepCopyGraphs(Source, Script->GetUsage(), EmptyResolver, bNeedsCompilation);
+		BasePtr->Source = Cast<UNiagaraScriptSource>(Script->GetSource(Version));
 		const TArray<FNiagaraVariable> EncounterableVariables;
-		BasePtr->FinishPrecompile(EncounterableVariables, Script->GetUsage(), EmptyResolver, nullptr);
+		const TArray<ENiagaraScriptUsage> ValidUsages = { ENiagaraScriptUsage::Function, ENiagaraScriptUsage::Module, ENiagaraScriptUsage::DynamicInput };
+		BasePtr->FinishPrecompile(EncounterableVariables, EmptyResolver, ValidUsages, nullptr);
 	}
 	else if (System)
 	{
 		check(System->GetSystemSpawnScript()->GetLatestSource() == System->GetSystemUpdateScript()->GetLatestSource());
+		BasePtr->Source = Cast<UNiagaraScriptSource>(System->GetSystemSpawnScript()->GetLatestSource());
 		BasePtr->bUseRapidIterationParams = !System->bBakeOutRapidIteration;
 
 		// Store off the current variables in the exposed parameters list.
@@ -771,13 +724,8 @@ TSharedPtr<FNiagaraCompileRequestDataBase, ESPMode::ThreadSafe> FNiagaraEditorMo
 			const FNiagaraEmitterHandle& Handle = System->GetEmitterHandle(i);
 			TSharedPtr<FNiagaraCompileRequestData, ESPMode::ThreadSafe> EmitterPtr = MakeShared<FNiagaraCompileRequestData, ESPMode::ThreadSafe>();
 			EmitterPtr->EmitterUniqueName = Handle.GetInstance()->GetUniqueEmitterName();
-			if (Handle.GetIsEnabled()) // Don't need to copy the graph if we aren't going to use it.
-			{
-				UNiagaraScriptSource* Source = Cast<UNiagaraScriptSource>(Handle.GetInstance()->GraphSource);
-				bool bNeedsCompilation = ScriptSourceNeedsCompiling(Source, InCompilingScripts);
-				EmitterPtr->DeepCopyGraphs(Source, Handle.GetInstance(), bNeedsCompilation);
-			}
 			EmitterPtr->SourceName = BasePtr->SourceName;
+			EmitterPtr->Source = Cast<UNiagaraScriptSource>(Handle.GetInstance()->GraphSource);
 			EmitterPtr->bUseRapidIterationParams = BasePtr->bUseRapidIterationParams || (!Handle.GetInstance()->bBakeOutRapidIteration);
 			EmitterPtr->bSimulationStagesEnabled = Handle.GetInstance()->bSimulationStagesEnabled;
 			BasePtr->EmitterData.Add(EmitterPtr);
@@ -787,9 +735,8 @@ TSharedPtr<FNiagaraCompileRequestDataBase, ESPMode::ThreadSafe> FNiagaraEditorMo
 		{
 			FCompileConstantResolver ConstantResolver(System, ENiagaraScriptUsage::SystemSpawnScript);
 			UNiagaraScriptSource* Source = Cast<UNiagaraScriptSource>(System->GetSystemSpawnScript()->GetLatestSource());
-			bool bNeedsCompilation = ScriptSourceNeedsCompiling(Source, InCompilingScripts);
-			BasePtr->DeepCopyGraphs(Source, ENiagaraScriptUsage::SystemSpawnScript, ConstantResolver, bNeedsCompilation);
-			BasePtr->FinishPrecompile(EncounterableVars, ENiagaraScriptUsage::SystemSpawnScript, ConstantResolver, nullptr);
+			static TArray<ENiagaraScriptUsage> SystemUsages = { ENiagaraScriptUsage::SystemSpawnScript, ENiagaraScriptUsage::SystemUpdateScript };
+			BasePtr->FinishPrecompile(EncounterableVars, ConstantResolver, SystemUsages, nullptr);
 		}
 
 		// Add the User and System variables that we did encounter to the list that emitters might also encounter.
@@ -819,16 +766,21 @@ TSharedPtr<FNiagaraCompileRequestDataBase, ESPMode::ThreadSafe> FNiagaraEditorMo
 						ScriptToEmitterMap.Add(EmitterScripts[ScriptIdx], Handle.GetInstance());
 					}
 				}
-				BasePtr->EmitterData[i]->FinishPrecompile(EncounterableVars, ENiagaraScriptUsage::EmitterSpawnScript, ConstantResolver, &Handle.GetInstance()->GetSimulationStages());
-				BasePtr->MergeInEmitterPrecompiledData(BasePtr->EmitterData[i].Get());
-
-				for (UNiagaraRendererProperties* RendererProperty : Handle.GetInstance()->GetRenderers())
-				{
-					for (const FNiagaraVariable& BoundAttribute : RendererProperty->GetBoundAttributes())
-					{
-						BasePtr->EmitterData[i]->RequiredRendererVariables.AddUnique(BoundAttribute);
-					}
-				}
+				// First finish the precompile for the emitters so that the encountered emitter variables can be found.
+				static TArray<ENiagaraScriptUsage> EmitterUsages = { ENiagaraScriptUsage::EmitterSpawnScript, ENiagaraScriptUsage::EmitterUpdateScript };
+				BasePtr->EmitterData[i]->FinishPrecompile(EncounterableVars, ConstantResolver, EmitterUsages, nullptr);
+
+				// Then finish the precompile for the particle scripts once we've gathered the emitter vars which might be referenced.
+				TArray<FNiagaraVariable> ParticleEncounterableVars = EncounterableVars;
+				BasePtr->EmitterData[i]->GatherPreCompiledVariables(TEXT("Emitter"), ParticleEncounterableVars);
+				static TArray<ENiagaraScriptUsage> ParticleUsages = {
+					ENiagaraScriptUsage::ParticleSpawnScript,
+					ENiagaraScriptUsage::ParticleSpawnScriptInterpolated,
+					ENiagaraScriptUsage::ParticleUpdateScript,
+					ENiagaraScriptUsage::ParticleEventScript,
+					ENiagaraScriptUsage::ParticleGPUComputeScript,
+					ENiagaraScriptUsage::ParticleSimulationStageScript };
+				BasePtr->EmitterData[i]->FinishPrecompile(ParticleEncounterableVars, ConstantResolver, ParticleUsages, &Handle.GetInstance()->GetSimulationStages());
 			}
 		}
 
@@ -921,12 +873,179 @@ TSharedPtr<FNiagaraCompileRequestDataBase, ESPMode::ThreadSafe> FNiagaraEditorMo
 	return BasePtr;
 }
 
-int32 FNiagaraEditorModule::CompileScript(const FNiagaraCompileRequestDataBase* InCompileRequest, const FNiagaraCompileOptions& InCompileOptions)
+void GetUsagesToDuplicate(ENiagaraScriptUsage TargetUsage, TArray<ENiagaraScriptUsage>& DuplicateUsages)
+{
+	/*switch (TargetUsage)
+	{
+	case ENiagaraScriptUsage::SystemSpawnScript:
+		DuplicateUsages.Add(ENiagaraScriptUsage::SystemSpawnScript);
+		DuplicateUsages.Add(ENiagaraScriptUsage::EmitterSpawnScript);
+		break;
+	case ENiagaraScriptUsage::SystemUpdateScript:
+		DuplicateUsages.Add(ENiagaraScriptUsage::SystemUpdateScript);
+		DuplicateUsages.Add(ENiagaraScriptUsage::EmitterUpdateScript);
+		break;
+	case ENiagaraScriptUsage::ParticleSpawnScript:
+	case ENiagaraScriptUsage::ParticleUpdateScript:
+	case ENiagaraScriptUsage::ParticleEventScript:
+		DuplicateUsages.Add(TargetUsage);
+		break;
+	case ENiagaraScriptUsage::ParticleSpawnScriptInterpolated:
+		DuplicateUsages.Add(ENiagaraScriptUsage::ParticleSpawnScript);
+		DuplicateUsages.Add(ENiagaraScriptUsage::ParticleSpawnScriptInterpolated);
+		DuplicateUsages.Add(ENiagaraScriptUsage::ParticleUpdateScript);
+		break;
+	case ENiagaraScriptUsage::ParticleGPUComputeScript:
+		DuplicateUsages.Add(ENiagaraScriptUsage::ParticleSpawnScript);
+		DuplicateUsages.Add(ENiagaraScriptUsage::ParticleSpawnScriptInterpolated);
+		DuplicateUsages.Add(ENiagaraScriptUsage::ParticleUpdateScript);
+		DuplicateUsages.Add(ENiagaraScriptUsage::ParticleSimulationStageScript);
+		DuplicateUsages.Add(ENiagaraScriptUsage::ParticleGPUComputeScript);
+		break;
+	}*/
+
+	// For now we need to include both spawn and update for each target usage, otherwise attribute lists in the precompiled histories aren't generated correctly.
+	switch (TargetUsage)
+	{
+	case ENiagaraScriptUsage::SystemSpawnScript:
+	case ENiagaraScriptUsage::SystemUpdateScript:
+		DuplicateUsages.Add(ENiagaraScriptUsage::SystemSpawnScript);
+		DuplicateUsages.Add(ENiagaraScriptUsage::SystemUpdateScript);
+		DuplicateUsages.Add(ENiagaraScriptUsage::EmitterSpawnScript);
+		DuplicateUsages.Add(ENiagaraScriptUsage::EmitterUpdateScript);
+		break;
+	case ENiagaraScriptUsage::ParticleSpawnScript:
+	case ENiagaraScriptUsage::ParticleSpawnScriptInterpolated:
+	case ENiagaraScriptUsage::ParticleUpdateScript:
+	case ENiagaraScriptUsage::ParticleEventScript:
+	case ENiagaraScriptUsage::ParticleSimulationStageScript:
+	case ENiagaraScriptUsage::ParticleGPUComputeScript:
+		DuplicateUsages.Add(ENiagaraScriptUsage::ParticleSpawnScript);
+		DuplicateUsages.Add(ENiagaraScriptUsage::ParticleSpawnScriptInterpolated);
+		DuplicateUsages.Add(ENiagaraScriptUsage::ParticleUpdateScript);
+		DuplicateUsages.Add(ENiagaraScriptUsage::ParticleEventScript);
+		DuplicateUsages.Add(ENiagaraScriptUsage::ParticleSimulationStageScript);
+		DuplicateUsages.Add(ENiagaraScriptUsage::ParticleGPUComputeScript);
+		break;
+	}
+}
+
+TSharedPtr<FNiagaraCompileRequestDuplicateDataBase, ESPMode::ThreadSafe> FNiagaraEditorModule::PrecompileDuplicate(
+	const FNiagaraCompileRequestDataBase* OwningSystemRequestData,
+	UNiagaraSystem* OwningSystem,
+	UNiagaraEmitter* OwningEmitter,
+	UNiagaraScript* TargetScript,
+	FGuid TargetVersion)
+{
+	UPackage* LogPackage = nullptr;
+
+	if (OwningSystem != nullptr)
+	{
+		LogPackage = OwningSystem->GetOutermost();
+	}
+	else if (TargetScript != nullptr)
+	{
+		LogPackage = TargetScript->GetOutermost();
+	}
+	else
+	{
+		TSharedPtr<FNiagaraCompileRequestDuplicateDataBase, ESPMode::ThreadSafe> InvalidPtr;
+		return InvalidPtr;
+	}
+
+	TRACE_CPUPROFILER_EVENT_SCOPE(NiagaraPrecompileDuplicate);
+	TRACE_CPUPROFILER_EVENT_SCOPE_TEXT_ON_CHANNEL(LogPackage ? *LogPackage->GetName() : *TargetScript->GetName(), NiagaraChannel);
+
+	SCOPE_CYCLE_COUNTER(STAT_NiagaraEditor_ScriptSource_PreCompileDuplicate);
+	double StartTime = FPlatformTime::Seconds();
+
+	TSharedPtr<FNiagaraCompileRequestDuplicateData, ESPMode::ThreadSafe> BasePtr = MakeShared<FNiagaraCompileRequestDuplicateData, ESPMode::ThreadSafe>();
+	TArray<TSharedPtr<FNiagaraCompileRequestDuplicateData, ESPMode::ThreadSafe>> DependentRequests;
+	FCompileConstantResolver EmptyResolver;
+
+	BasePtr->SharedSourceGraphToDuplicatedGraphsMap = MakeShared<TMap<const UNiagaraGraph*, TArray<FNiagaraCompileRequestDuplicateData::FDuplicatedGraphData>>>();
+	BasePtr->SharedNameToDuplicatedDataInterfaceMap = MakeShared<TMap<FName, UNiagaraDataInterface*>>();
+	BasePtr->SharedDataInterfaceClassToDuplicatedCDOMap = MakeShared<TMap<UClass*, UNiagaraDataInterface*>>();
+
+	if (OwningSystem == nullptr)
+	{
+		UNiagaraScriptSource* Source = CastChecked<UNiagaraScriptSource>(TargetScript->GetSource(TargetVersion));
+		BasePtr->ValidUsages.Add(TargetScript->GetUsage());
+		BasePtr->DeepCopyGraphs(Source, TargetScript->GetUsage(), EmptyResolver);
+		TArray<FNiagaraVariable> EncounterableScriptVariables;
+		OwningSystemRequestData->GatherPreCompiledVariables(FString(), EncounterableScriptVariables);
+		BasePtr->FinishPrecompileDuplicate(EncounterableScriptVariables, EmptyResolver, nullptr);
+	}
+	else
+	{
+		GetUsagesToDuplicate(TargetScript->GetUsage(), BasePtr->ValidUsages);
+
+		check(OwningSystem->GetSystemSpawnScript()->GetLatestSource() == OwningSystem->GetSystemUpdateScript()->GetLatestSource());
+
+		// Store off the current variables in the exposed parameters list.
+		TArray<FNiagaraVariable> OriginalExposedParams;
+		OwningSystem->GetExposedParameters().GetParameters(OriginalExposedParams);
+
+		// Create an array of variables that we might encounter when traversing the graphs (include the originally exposed vars above)
+		TArray<FNiagaraVariable> EncounterableVars(OriginalExposedParams);
+
+		// First deep copy all the emitter graphs referenced by the system so that we can later hook up emitter handles in the system traversal.
+		BasePtr->EmitterData.Empty();
+		for (int32 i = 0; i < OwningSystem->GetEmitterHandles().Num(); i++)
+		{
+			const FNiagaraEmitterHandle& Handle = OwningSystem->GetEmitterHandle(i);
+			TSharedPtr<FNiagaraCompileRequestDuplicateData, ESPMode::ThreadSafe> EmitterPtr = MakeShared<FNiagaraCompileRequestDuplicateData, ESPMode::ThreadSafe>();
+			EmitterPtr->EmitterUniqueName = Handle.GetInstance()->GetUniqueEmitterName();
+			EmitterPtr->ValidUsages = BasePtr->ValidUsages;
+			EmitterPtr->SharedSourceGraphToDuplicatedGraphsMap = BasePtr->SharedSourceGraphToDuplicatedGraphsMap;
+			EmitterPtr->SharedNameToDuplicatedDataInterfaceMap = BasePtr->SharedNameToDuplicatedDataInterfaceMap;
+			EmitterPtr->SharedDataInterfaceClassToDuplicatedCDOMap = BasePtr->SharedDataInterfaceClassToDuplicatedCDOMap;
+			EmitterPtr->bSimulationStagesEnabled = Handle.GetInstance()->bSimulationStagesEnabled;
+			if (Handle.GetIsEnabled() && (OwningEmitter == nullptr || OwningEmitter == Handle.GetInstance())) // Don't need to copy the graph if we aren't going to use it.
+			{
+				EmitterPtr->DeepCopyGraphs(Handle.GetInstance());
+			}
+			EmitterPtr->ValidUsages = BasePtr->ValidUsages;
+			BasePtr->EmitterData.Add(EmitterPtr);
+		}
+
+		// Now deep copy the system graphs, skipping traversal into any emitter references.
+		{
+			FCompileConstantResolver ConstantResolver(OwningSystem, ENiagaraScriptUsage::SystemSpawnScript);
+			UNiagaraScriptSource* Source = Cast<UNiagaraScriptSource>(OwningSystem->GetSystemSpawnScript()->GetLatestSource());
+			BasePtr->DeepCopyGraphs(Source, ENiagaraScriptUsage::SystemSpawnScript, ConstantResolver);
+			TArray<FNiagaraVariable> EncounterableSystemVariables;
+			OwningSystemRequestData->GatherPreCompiledVariables(FString(), EncounterableSystemVariables);
+			BasePtr->FinishPrecompileDuplicate(EncounterableSystemVariables, ConstantResolver, nullptr);
+		}
+
+		// Now we can finish off the emitters.
+		for (int32 i = 0; i < OwningSystem->GetEmitterHandles().Num(); i++)
+		{
+			const FNiagaraEmitterHandle& Handle = OwningSystem->GetEmitterHandle(i);
+			FCompileConstantResolver ConstantResolver(Handle.GetInstance(), ENiagaraScriptUsage::EmitterSpawnScript);
+			if (Handle.GetIsEnabled())
+			{
+				TArray<FNiagaraVariable> EncounterableEmitterVariables;
+				OwningSystemRequestData->GetDependentRequest(i)->GatherPreCompiledVariables(FString(), EncounterableEmitterVariables);
+				BasePtr->EmitterData[i]->FinishPrecompileDuplicate(EncounterableEmitterVariables, ConstantResolver, &Handle.GetInstance()->GetSimulationStages());
+			}
+		}
+	}
+
+	UE_LOG(LogNiagaraEditor, Verbose, TEXT("'%s' PrecompileDuplicate took %f sec."), *LogPackage->GetName(),
+		(float)(FPlatformTime::Seconds() - StartTime));
+
+	return BasePtr;
+}
+
+int32 FNiagaraEditorModule::CompileScript(const FNiagaraCompileRequestDataBase* InCompileRequest, const FNiagaraCompileRequestDuplicateDataBase* InCompileRequestDuplicate, const FNiagaraCompileOptions& InCompileOptions)
 {
 	SCOPE_CYCLE_COUNTER(STAT_NiagaraEditor_Module_CompileScript);
 
 	check(InCompileRequest != NULL);
 	const FNiagaraCompileRequestData* CompileRequest = (const FNiagaraCompileRequestData*)InCompileRequest;
+	const FNiagaraCompileRequestDuplicateData* CompileRequestDuplicate = (const FNiagaraCompileRequestDuplicateData*)InCompileRequestDuplicate;
 	TArray<FNiagaraVariable> CookedRapidIterationParams = CompileRequest->GetUseRapidIterationParams() ? TArray<FNiagaraVariable>() : CompileRequest->RapidIterationParams;
 
 	UE_LOG(LogNiagaraEditor, Verbose, TEXT("Compiling System %s ..................................................................."), *InCompileOptions.FullName);
@@ -954,11 +1073,11 @@ int32 FNiagaraEditorModule::CompileScript(const FNiagaraCompileRequestDataBase*
 	if (GbForceNiagaraTranslatorSingleThreaded > 0)
 	{
 		FScopeLock Lock(&TranslationCritSec);
-		TranslateResults = Translator.Translate(CompileRequest, InCompileOptions, TranslateOptions);
+		TranslateResults = Translator.Translate(CompileRequest, CompileRequestDuplicate, InCompileOptions, TranslateOptions);
 	}
 	else
 	{
-		TranslateResults = Translator.Translate(CompileRequest, InCompileOptions, TranslateOptions);
+		TranslateResults = Translator.Translate(CompileRequest, CompileRequestDuplicate, InCompileOptions, TranslateOptions);
 	}
 	UE_LOG(LogNiagaraEditor, Verbose, TEXT("Translating System %s took %f sec."), *InCompileOptions.FullName, (float)(FPlatformTime::Seconds() - TranslationStartTime));
 
diff --git a/Engine/Plugins/FX/Niagara/Source/NiagaraEditor/Private/NiagaraEditorModule.cpp b/Engine/Plugins/FX/Niagara/Source/NiagaraEditor/Private/NiagaraEditorModule.cpp
--- a/Engine/Plugins/FX/Niagara/Source/NiagaraEditor/Private/NiagaraEditorModule.cpp
+++ b/Engine/Plugins/FX/Niagara/Source/NiagaraEditor/Private/NiagaraEditorModule.cpp
@@ -1051,9 +1051,9 @@ void FNiagaraEditorModule::StartupModule()
 	NiagaraModule.RegisterEditorOnlyDataUtilities(EditorOnlyDataUtilities.ToSharedRef());
 
 	// Register the script compiler
-	ScriptCompilerHandle = NiagaraModule.RegisterScriptCompiler(INiagaraModule::FScriptCompiler::CreateLambda([this](const FNiagaraCompileRequestDataBase* CompileRequest, const FNiagaraCompileOptions& Options)
+	ScriptCompilerHandle = NiagaraModule.RegisterScriptCompiler(INiagaraModule::FScriptCompiler::CreateLambda([this](const FNiagaraCompileRequestDataBase* CompileRequest, const FNiagaraCompileRequestDuplicateDataBase* CompileRequestDuplicate, const FNiagaraCompileOptions& Options)
 	{
-		return CompileScript(CompileRequest, Options);
+		return CompileScript(CompileRequest, CompileRequestDuplicate, Options);
 	}));
 
 	CompileResultHandle = NiagaraModule.RegisterCompileResultDelegate(INiagaraModule::FCheckCompilationResult::CreateLambda([this](int32 JobID, bool bWait)
@@ -1066,6 +1066,11 @@ void FNiagaraEditorModule::StartupModule()
 		return Precompile(InObj, Version);
 	}));
 
+	PrecompileDuplicatorHandle = NiagaraModule.RegisterPrecompileDuplicator(INiagaraModule::FOnPrecompileDuplicate::CreateLambda([this](const FNiagaraCompileRequestDataBase* OwningSystemRequestData, UNiagaraSystem* OwningSystem, UNiagaraEmitter* OwningEmitter, UNiagaraScript* TargetScript,  FGuid Version)
+	{
+		return PrecompileDuplicate(OwningSystemRequestData, OwningSystem, OwningEmitter, TargetScript, Version);
+	}));
+
 	TestCompileScriptCommand = IConsoleManager::Get().RegisterConsoleCommand(
 		TEXT("fx.TestCompileNiagaraScript"),
 		TEXT("Compiles the specified script on disk for the niagara vector vm"),
diff --git a/Engine/Plugins/FX/Niagara/Source/NiagaraEditor/Private/NiagaraEditorUtilities.cpp b/Engine/Plugins/FX/Niagara/Source/NiagaraEditor/Private/NiagaraEditorUtilities.cpp
--- a/Engine/Plugins/FX/Niagara/Source/NiagaraEditor/Private/NiagaraEditorUtilities.cpp
+++ b/Engine/Plugins/FX/Niagara/Source/NiagaraEditor/Private/NiagaraEditorUtilities.cpp
@@ -1097,6 +1097,10 @@ void FixUpNumericPinsVisitor(const UEdGraphSchema_Niagara* Schema, UNiagaraNode*
 
 void FNiagaraEditorUtilities::FixUpNumericPins(const UEdGraphSchema_Niagara* Schema, UNiagaraNode* Node)
 {
+	if (ensureMsgf(Node->GetOutermost() == GetTransientPackage(), TEXT("Can not fix up numerics on non-transient node {0}"), *Node->GetPathName()) == false)
+	{
+		return;
+	}
 	auto FixUpVisitor = [&](const UEdGraphSchema_Niagara* LSchema, UNiagaraNode* LNode) { FixUpNumericPinsVisitor(LSchema, LNode); };
 	TraverseGraphFromOutputDepthFirst(Schema, Node, FixUpVisitor);
 }
@@ -1329,6 +1333,11 @@ void FNiagaraEditorUtilities::ResolveNumerics(UNiagaraGraph* SourceGraph, bool b
 void FNiagaraEditorUtilities::PreprocessFunctionGraph(const UEdGraphSchema_Niagara* Schema, UNiagaraGraph* Graph, TArrayView<UEdGraphPin* const> CallInputs, TArrayView<UEdGraphPin* const> CallOutputs,
 	ENiagaraScriptUsage ScriptUsage, const FCompileConstantResolver& ConstantResolver)
 {
+	if (ensureMsgf(Graph->GetOutermost() == GetTransientPackage(), TEXT("Can not preprocess non-transient function graph {0}"), *Graph->GetPathName()) == false)
+	{
+		return;
+	}
+
 	// Change any numeric inputs or outputs to match the types from the call node.
 	TArray<UNiagaraNodeInput*> InputNodes;
 
diff --git a/Engine/Plugins/FX/Niagara/Source/NiagaraEditor/Private/NiagaraGraph.cpp b/Engine/Plugins/FX/Niagara/Source/NiagaraEditor/Private/NiagaraGraph.cpp
--- a/Engine/Plugins/FX/Niagara/Source/NiagaraEditor/Private/NiagaraGraph.cpp
+++ b/Engine/Plugins/FX/Niagara/Source/NiagaraEditor/Private/NiagaraGraph.cpp
@@ -951,7 +951,28 @@ void UNiagaraGraph::FindEquivalentOutputNodes(ENiagaraScriptUsage TargetUsageTyp
 	OutputNodes = NodesFound;
 }
 
-UNiagaraGraph* UNiagaraGraph::CreateCompilationCopy()
+void BuildCompleteTraversal(UEdGraphNode* CurrentNode, TArray<UEdGraphNode*>& AllNodes)
+{
+	if (CurrentNode == nullptr || AllNodes.Contains(CurrentNode))
+	{
+		return;
+	}
+
+	AllNodes.Add(CurrentNode);
+
+	for (UEdGraphPin* Pin : CurrentNode->GetAllPins())
+	{
+		if(Pin->Direction == EGPD_Input)
+		{
+			for (UEdGraphPin* LinkedPin : Pin->LinkedTo)
+			{
+				BuildCompleteTraversal(LinkedPin->GetOwningNode(), AllNodes);
+			}
+		}
+	}
+}
+
+UNiagaraGraph* UNiagaraGraph::CreateCompilationCopy(const TArray<ENiagaraScriptUsage>& CompileUsages)
 {
 	check(!bIsForCompilationOnly);
 	UNiagaraGraph* Result = NewObject<UNiagaraGraph>();
@@ -961,6 +982,11 @@ UNiagaraGraph* UNiagaraGraph::CreateCompilationCopy()
 	for (TFieldIterator<FProperty> PropertyIt(GetClass(), EFieldIteratorFlags::IncludeSuper); PropertyIt; ++PropertyIt)
 	{
 		FProperty* Property = *PropertyIt;
+		if (*Property->GetNameCPP() == GET_MEMBER_NAME_CHECKED(UEdGraph, Nodes))
+		{
+			// The nodes will be handled separately below.
+			continue;
+		}
 		const uint8* SourceAddr = Property->ContainerPtrToValuePtr<uint8>(this);
         uint8* DestinationAddr = Property->ContainerPtrToValuePtr<uint8>(Result);
 
@@ -974,8 +1000,20 @@ UNiagaraGraph* UNiagaraGraph::CreateCompilationCopy()
 		It.Value = CastChecked<UNiagaraScriptVariable>(NiagaraEditorModule.GetPooledDuplicateObject(It.Value));
 	}
 
-	// duplicate the nodes
+	// Only add references to nodes which match a complie usage.
 	TMap<UEdGraphNode*, UEdGraphNode*> DuplicationMapping;
+	TArray<UNiagaraNodeOutput*> OutputNodes;
+	GetNodesOfClass(OutputNodes);
+	for (UNiagaraNodeOutput* OutputNode : OutputNodes)
+	{
+		if(CompileUsages.ContainsByPredicate([OutputNode](ENiagaraScriptUsage CompileUsage) { return UNiagaraScript::IsEquivalentUsage(CompileUsage, OutputNode->GetUsage()); }))
+		{
+			TArray<UEdGraphNode*> TraversedNodes;
+			BuildCompleteTraversal(OutputNode, TraversedNodes);
+			Result->Nodes.Append(TraversedNodes);
+		}
+	}
+
 	TArray<UEdGraphNode*> NewNodes;
 	for (UEdGraphNode* Node : Result->Nodes)
 	{
@@ -992,11 +1030,22 @@ UNiagaraGraph* UNiagaraGraph::CreateCompilationCopy()
 			for (int i = 0; i < Pin->LinkedTo.Num(); i++)
 			{
 				UEdGraphPin* LinkedPin = Pin->LinkedTo[i];
-				UEdGraphNode* NewNode = DuplicationMapping[LinkedPin->GetOwningNode()];
-				UEdGraphPin** NodePin = NewNode->Pins.FindByPredicate([LinkedPin](UEdGraphPin* Pin) { return Pin->PinId == LinkedPin->PinId; });
-				check(NodePin);
-				Pin->LinkedTo[i] = *NodePin;
+				UEdGraphNode** NewNodePtr = DuplicationMapping.Find(LinkedPin->GetOwningNode());
+				if (NewNodePtr == nullptr)
+				{
+					// If output pins were connected to another node which wasn't encountered in traversal then it's possible it won't be found here.
+					// In that case set the linked pin to nullptr to be removed later.
+					Pin->LinkedTo[i] = nullptr;
+				}
+				else
+				{
+					UEdGraphPin** NodePin = (*NewNodePtr)->Pins.FindByPredicate([LinkedPin](UEdGraphPin* Pin) { return Pin->PinId == LinkedPin->PinId; });
+					check(NodePin);
+					Pin->LinkedTo[i] = *NodePin;
+				}
 			}
+			// Remove any null linked pins.
+			Pin->LinkedTo.RemoveAll([](const UEdGraphPin* Pin) { return Pin == nullptr; });
 		}
 	}
 	Result->Nodes = NewNodes;
diff --git a/Engine/Plugins/FX/Niagara/Source/NiagaraEditor/Private/NiagaraHlslTranslator.cpp b/Engine/Plugins/FX/Niagara/Source/NiagaraEditor/Private/NiagaraHlslTranslator.cpp
--- a/Engine/Plugins/FX/Niagara/Source/NiagaraEditor/Private/NiagaraHlslTranslator.cpp
+++ b/Engine/Plugins/FX/Niagara/Source/NiagaraEditor/Private/NiagaraHlslTranslator.cpp
@@ -1043,22 +1043,24 @@ static void ConvertFloatToHalf(const FNiagaraCompileOptions& InCompileOptions, T
 	}
 }
 
-const FNiagaraTranslateResults &FHlslNiagaraTranslator::Translate(const FNiagaraCompileRequestData* InCompileData, const FNiagaraCompileOptions& InCompileOptions, FHlslNiagaraTranslatorOptions InTranslateOptions)
+const FNiagaraTranslateResults &FHlslNiagaraTranslator::Translate(const FNiagaraCompileRequestData* InCompileData, const FNiagaraCompileRequestDuplicateData* InCompileDuplicateData, const FNiagaraCompileOptions& InCompileOptions, FHlslNiagaraTranslatorOptions InTranslateOptions)
 {
 	TRACE_CPUPROFILER_EVENT_SCOPE(NiagaraHlslTranslate);
 	TRACE_CPUPROFILER_EVENT_SCOPE_TEXT_ON_CHANNEL(*InCompileOptions.GetPathName(), NiagaraChannel);
 
 	NIAGARA_SCOPE_CYCLE_COUNTER(STAT_NiagaraEditor_HlslTranslator_Translate);
 	check(InCompileData);
+	check(InCompileDuplicateData);
 
 	CompileOptions = InCompileOptions;
 	CompileData = InCompileData;
+	CompileDuplicateData = InCompileDuplicateData;
 	TranslationOptions = InTranslateOptions;
 	CompilationTarget = TranslationOptions.SimTarget;
 	TranslateResults.bHLSLGenSucceeded = false;
 	TranslateResults.OutputHLSL = "";
 
-	TWeakObjectPtr<UNiagaraGraph> SourceGraph = CompileData->NodeGraphDeepCopy;
+	TWeakObjectPtr<UNiagaraGraph> SourceGraph = CompileDuplicateData->NodeGraphDeepCopy;
 
 	if (!SourceGraph.IsValid())
 	{
@@ -1089,7 +1091,7 @@ const FNiagaraTranslateResults &FHlslNiagaraTranslator::Translate(const FNiagara
 	ParamMapHistories.Empty();
 	ParamMapSetVariablesToChunks.Empty();
 
-	OtherOutputParamMapHistories = CompileData->GetPrecomputedHistories();
+	OtherOutputParamMapHistories = CompileDuplicateData->GetPrecomputedHistories();
 
 	// Make the sanitized variable version of this list.
 	OtherOutputParamMapHistoriesSanitizedVariables.AddDefaulted(OtherOutputParamMapHistories.Num());
@@ -1233,7 +1235,7 @@ const FNiagaraTranslateResults &FHlslNiagaraTranslator::Translate(const FNiagara
 					// If we allow partial writes we need to ensure that we are not reading from our own buffer, we ask our data interfaces if this is true or not
 					if (TranslationStages[Index].bPartialParticleUpdate)
 					{
-						for (auto it = InCompileData->CopiedDataInterfacesByName.CreateConstIterator(); it; ++it)
+						for (auto it = InCompileDuplicateData->SharedNameToDuplicatedDataInterfaceMap->CreateConstIterator(); it; ++it)
 						{
 							const UNiagaraDataInterface* DataInterface = it->Value;
 							if (DataInterface->ReadsEmitterParticleData(InCompileData->EmitterUniqueName))
@@ -1957,8 +1959,8 @@ void FHlslNiagaraTranslator::HandleSimStageSetupAndTeardown(int32 InWhichStage,
 		return;
 	}
 
-	UObject* const* FoundCDO = CompileData->CDOs.Find(IterationSourceVar.GetType().GetClass());
-	if (!FoundCDO || !IterationSourceVar.GetType().GetClass())
+	UNiagaraDataInterface* CDO = CompileDuplicateData->GetDuplicatedDataInterfaceCDOForClass(IterationSourceVar.GetType().GetClass());
+	if (!CDO || !IterationSourceVar.GetType().GetClass())
 	{
 		Error(FText::Format(LOCTEXT("CannotFindIterationSourceCDOInParamMap", "Variable {0}'s cached CDO for class was missing during compile!"), FText::FromName(TranslationStage.IterationSource)), nullptr, nullptr);
 		return;
@@ -2009,7 +2011,6 @@ void FHlslNiagaraTranslator::HandleSimStageSetupAndTeardown(int32 InWhichStage,
 		bWroteToIterationSource = true;
 
 	// Now decide if we need to put in the pre/post
-	UNiagaraDataInterface* CDO = Cast<UNiagaraDataInterface>(*FoundCDO);
 	if (CDO && CDO->CanExecuteOnTarget(ENiagaraSimTarget::GPUComputeSim))
 	{
 		bool bNeedsDIOwner = false;
@@ -2631,9 +2632,8 @@ void FHlslNiagaraTranslator::DefineDataInterfaceHLSL(FString& InHlslOutput)
 	{
 		FNiagaraScriptDataInterfaceCompileInfo& Info = CompilationOutput.ScriptData.DataInterfaceInfo[i];
 
-		UObject* const* FoundCDO = CompileData->CDOs.Find(Info.Type.GetClass());
-		check(FoundCDO != nullptr);
-		UNiagaraDataInterface* CDO = Cast<UNiagaraDataInterface>(*FoundCDO);
+		UNiagaraDataInterface* CDO = CompileDuplicateData->GetDuplicatedDataInterfaceCDOForClass(Info.Type.GetClass());
+		check(CDO != nullptr);
 		if (CDO && CDO->CanExecuteOnTarget(ENiagaraSimTarget::GPUComputeSim))
 		{
 			if ( !InterfaceClasses.Contains(Info.Type.GetFName()) )
@@ -5124,9 +5124,8 @@ bool FHlslNiagaraTranslator::ParameterMapRegisterExternalConstantNamespaceVariab
 					}
 					else
 					{
-						UObject*const* Obj = CompileData->CDOs.Find(const_cast<UClass*>(InVariable.GetType().GetClass()));
-						check(Obj != nullptr);
-						DataInterface = CastChecked<UNiagaraDataInterface>(*Obj);
+						DataInterface = CompileDuplicateData->GetDuplicatedDataInterfaceCDOForClass(const_cast<UClass*>(InVariable.GetType().GetClass()));
+						check(DataInterface != nullptr);
 					}
 					if (ensure(DataInterface))
 					{
@@ -5874,9 +5873,8 @@ void FHlslNiagaraTranslator::HandleParameterRead(int32 ParamMapHistoryIdx, const
 		// for variables encountered without ever being set. We do this by creating an instance of the CDO.
 		if (UNiagaraScript::IsStandaloneScript(CompileOptions.TargetUsage) && LastSetChunkIdx == INDEX_NONE)
 		{
-			UObject*const* Obj = CompileData->CDOs.Find(const_cast<UClass*>(Var.GetType().GetClass()));
-			check(Obj != nullptr);
-			UNiagaraDataInterface* DataInterface = CastChecked<UNiagaraDataInterface>(*Obj);
+			UNiagaraDataInterface* DataInterface = CompileDuplicateData->GetDuplicatedDataInterfaceCDOForClass(const_cast<UClass*>(Var.GetType().GetClass()));
+			check(DataInterface != nullptr);
 			if (DataInterface)
 			{
 				LastSetChunkIdx = RegisterDataInterface(Var, DataInterface, true, false);
@@ -6291,9 +6289,7 @@ void FHlslNiagaraTranslator::FunctionCall(UNiagaraNodeFunctionCall* FunctionNode
 				FNiagaraTypeDefinition TypeDef = Schema->PinToTypeDefinition(CallInputs[i]);
 				if (TypeDef.IsDataInterface())
 				{
-					UObject* const* FoundCDO = CompileData->CDOs.Find(TypeDef.GetClass());
-					UNiagaraDataInterface* CDO = (FoundCDO ? Cast<UNiagaraDataInterface>(*FoundCDO) : nullptr);
-
+					UNiagaraDataInterface* CDO = CompileDuplicateData->GetDuplicatedDataInterfaceCDOForClass(TypeDef.GetClass());
 					if (CDO == nullptr)
 					{
 						// If the cdo wasn't found, the data interface was not passed through a parameter map and so it won't be bound correctly, so add a compile error
@@ -6660,9 +6656,7 @@ void FHlslNiagaraTranslator::ProcessCustomHlsl(const FString& InCustomHlsl, ENia
 		}
 		else if (Input.GetType().IsDataInterface())
 		{
-			UObject* const* FoundCDO = CompileData->CDOs.Find(Input.GetType().GetClass());
-			UNiagaraDataInterface* CDO = (FoundCDO ? Cast<UNiagaraDataInterface>(*FoundCDO) : nullptr);
-
+			UNiagaraDataInterface* CDO = CompileDuplicateData->GetDuplicatedDataInterfaceCDOForClass(Input.GetType().GetClass());
 			if (CDO == nullptr)
 			{
 				// If the cdo wasn't found, the data interface was not passed through a parameter map and so it won't be bound correctly, so add a compile error
@@ -7285,8 +7279,8 @@ void FHlslNiagaraTranslator::RegisterFunctionCall(ENiagaraScriptUsage ScriptUsag
 			// Double-check to make sure that the signature matches those specified by the data 
 			// interface. It could be that the existing node has been removed and the graph
 			// needs to be refactored. If that's the case, emit an error.
-			UObject* const* FoundCDO = CompileData->CDOs.Find(Info.Type.GetClass());
-			if (FoundCDO == nullptr)
+			UNiagaraDataInterface* CDO = CompileDuplicateData->GetDuplicatedDataInterfaceCDOForClass(Info.Type.GetClass());
+			if (CDO == nullptr)
 			{
 				// If the cdo wasn't found, the data interface was not passed through a parameter map and so it won't be bound correctly, so add a compile error
 				// and invalidate the signature.
@@ -7295,8 +7289,7 @@ void FHlslNiagaraTranslator::RegisterFunctionCall(ENiagaraScriptUsage ScriptUsag
 				return;
 			}
 
-			UNiagaraDataInterface* CDO = Cast<UNiagaraDataInterface>(*FoundCDO);
-			if (CDO != nullptr && OutSignature.bMemberFunction)
+			if (OutSignature.bMemberFunction)
 			{
 				TArray<FNiagaraFunctionSignature> DataInterfaceFunctions;
 				CDO->GetFunctions(DataInterfaceFunctions);
diff --git a/Engine/Plugins/FX/Niagara/Source/NiagaraEditor/Private/NiagaraHlslTranslator.h b/Engine/Plugins/FX/Niagara/Source/NiagaraEditor/Private/NiagaraHlslTranslator.h
--- a/Engine/Plugins/FX/Niagara/Source/NiagaraEditor/Private/NiagaraHlslTranslator.h
+++ b/Engine/Plugins/FX/Niagara/Source/NiagaraEditor/Private/NiagaraHlslTranslator.h
@@ -66,29 +66,23 @@ public:
 
 	}
 
-	virtual bool GatherPreCompiledVariables(const FString& InNamespaceFilter, TArray<FNiagaraVariable>& OutVars) override;
-	virtual void GetReferencedObjects(TArray<UObject*>& Objects) override;
-	virtual const TMap<FName, UNiagaraDataInterface*>& GetObjectNameMap() override;
-	void MergeInEmitterPrecompiledData(FNiagaraCompileRequestDataBase* InEmitterDataBase);
+	virtual void GatherPreCompiledVariables(const FString& InNamespaceFilter, TArray<FNiagaraVariable>& OutVars) const override;
 	virtual FName ResolveEmitterAlias(FName VariableName) const override;
 
-	TArray<FNiagaraParameterMapHistory>& GetPrecomputedHistories() { return PrecompiledHistories; }
-	const TArray<FNiagaraParameterMapHistory>& GetPrecomputedHistories() const { return PrecompiledHistories; }
-	virtual const class UNiagaraGraph* GetPrecomputedNodeGraph() const { return NodeGraphDeepCopy.Get(); }
 	const FString& GetUniqueEmitterName() const { return EmitterUniqueName; }
-	void VisitReferencedGraphs(UNiagaraGraph* InSrcGraph, UNiagaraGraph* InDupeGraph, ENiagaraScriptUsage InUsage, FCompileConstantResolver ConstantResolver, bool bNeedsCompilation, TMap<UNiagaraNodeFunctionCall*, ENiagaraScriptUsage> FunctionsWithUsage = TMap<UNiagaraNodeFunctionCall*, ENiagaraScriptUsage>());
-	void DeepCopyGraphs(UNiagaraScriptSource* ScriptSource, ENiagaraScriptUsage InUsage, FCompileConstantResolver ConstantResolver, bool bNeedsCompilation);
-	void DeepCopyGraphs(UNiagaraScriptSource* ScriptSource, UNiagaraEmitter* Emitter, bool bNeedsCompilation);
-	void FinishPrecompile(const TArray<FNiagaraVariable>& EncounterableVariables, ENiagaraScriptUsage InUsage, FCompileConstantResolver ConstantResolver, const TArray<class UNiagaraSimulationStageBase*>* SimStages);
+	void FinishPrecompile(const TArray<FNiagaraVariable>& EncounterableVariables, FCompileConstantResolver ConstantResolver, const TArray<ENiagaraScriptUsage>& UsagesToProcess, const TArray<class UNiagaraSimulationStageBase*>* SimStages);
 	virtual int32 GetDependentRequestCount() const override {
 		return EmitterData.Num();
 	};
 	virtual TSharedPtr<FNiagaraCompileRequestDataBase, ESPMode::ThreadSafe> GetDependentRequest(int32 Index) override {
 		return EmitterData[Index];
 	}
+	virtual const FNiagaraCompileRequestDataBase* GetDependentRequest(int32 Index) const override
+	{
+		return EmitterData[Index].Get();
+	}
 	void AddRapidIterationParameters(const FNiagaraParameterStore& InParamStore, FCompileConstantResolver InResolver);
 	virtual bool GetUseRapidIterationParams() const override { return bUseRapidIterationParams; }
-	virtual void ReleaseCompilationCopies() override;
 
 	// Simulation Stage Variables. Sim stage of 0 is always Spawn/Update
 	struct FCompileSimStageData
@@ -104,13 +98,7 @@ public:
 	};
 	TArray<FCompileSimStageData> CompileSimStageData;
 
-	//TArray<FName> IterationSourcePerStage;
-
-	TWeakObjectPtr<UNiagaraGraph> NodeGraphDeepCopy;
-	TArray<FNiagaraParameterMapHistory> PrecompiledHistories;
-	TArray<FNiagaraVariable> ChangedFromNumericVars;
-	TMap<FName, UNiagaraDataInterface*> CopiedDataInterfacesByName;
-	TMap<UClass*, UObject*> CDOs;
+	TArray<FNiagaraVariable> EncounteredVariables;
 	FString EmitterUniqueName;
 	TArray<TSharedPtr<FNiagaraCompileRequestData, ESPMode::ThreadSafe>> EmitterData;
 	TWeakObjectPtr<UNiagaraScriptSource> Source;
@@ -122,28 +110,67 @@ public:
 	UEnum* ENiagaraScriptUsageEnum;
 
 	TArray<FNiagaraVariable> RapidIterationParams;
+};
 
-	struct FunctionData
+class FNiagaraCompileRequestDuplicateData : public FNiagaraCompileRequestDuplicateDataBase
+{
+public:
+	FNiagaraCompileRequestDuplicateData()
+	{
+	}
+
+	virtual bool IsDuplicateDataFor(UNiagaraSystem* InSystem, UNiagaraEmitter* InEmitter, UNiagaraScript* InScript) const override;
+	virtual void GetDuplicatedObjects(TArray<UObject*>& Objects) override;
+	virtual const TMap<FName, UNiagaraDataInterface*>& GetObjectNameMap() override;
+	UNiagaraDataInterface* GetDuplicatedDataInterfaceCDOForClass(UClass* Class) const;
+
+	TArray<FNiagaraParameterMapHistory>& GetPrecomputedHistories() { return PrecompiledHistories; }
+	const TArray<FNiagaraParameterMapHistory>& GetPrecomputedHistories() const { return PrecompiledHistories; }
+
+	void DuplicateReferencedGraphs(UNiagaraGraph* InSrcGraph, UNiagaraGraph* InDupeGraph, ENiagaraScriptUsage InUsage, FCompileConstantResolver ConstantResolver, TMap<UNiagaraNodeFunctionCall*, ENiagaraScriptUsage> FunctionsWithUsage = TMap<UNiagaraNodeFunctionCall*, ENiagaraScriptUsage>());
+	void DeepCopyGraphs(UNiagaraScriptSource* ScriptSource, ENiagaraScriptUsage InUsage, FCompileConstantResolver ConstantResolver);
+	void DeepCopyGraphs(UNiagaraEmitter* Emitter);
+	
+	void FinishPrecompileDuplicate(const TArray<FNiagaraVariable>& EncounterableVariables, FCompileConstantResolver ConstantResolver, const TArray<class UNiagaraSimulationStageBase*>* SimStages);
+
+	virtual int32 GetDependentRequestCount() const override	{ return EmitterData.Num();	}
+	virtual TSharedPtr<FNiagaraCompileRequestDuplicateDataBase, ESPMode::ThreadSafe> GetDependentRequest(int32 Index) override { return EmitterData[Index];	}
+
+	virtual void ReleaseCompilationCopies() override;
+
+	TWeakObjectPtr<UNiagaraSystem> OwningSystem;
+	TWeakObjectPtr<UNiagaraEmitter> OwningEmitter;
+	TArray<ENiagaraScriptUsage> ValidUsages;
+
+	TWeakObjectPtr<UNiagaraScriptSource> SourceDeepCopy;
+	TWeakObjectPtr<UNiagaraGraph> NodeGraphDeepCopy;
+
+	TArray<FNiagaraParameterMapHistory> PrecompiledHistories;
+
+	TArray<FNiagaraVariable> ChangedFromNumericVars;
+	FString EmitterUniqueName;
+	TArray<TSharedPtr<FNiagaraCompileRequestDuplicateData, ESPMode::ThreadSafe>> EmitterData;
+	bool bSimulationStagesEnabled = false;
+
+	struct FDuplicatedGraphData
 	{
 		UNiagaraScript* ClonedScript;
 		UNiagaraGraph* ClonedGraph;
 		TArray<UEdGraphPin*> CallInputs;
 		TArray<UEdGraphPin*> CallOutputs;
 		ENiagaraScriptUsage Usage;
-		bool bHasNumericInputs;
+		bool bHasNumericParameters;
 	};
-	TMap<const UNiagaraGraph*, TArray<FunctionData>> PreprocessedFunctions;
-	TArray<TWeakObjectPtr<UNiagaraGraph>> ClonedGraphs;
 
-	// Copy of the variables that are required by Renderers if this is a compile request for an Emitter. 
-	TArray<FNiagaraVariable> RequiredRendererVariables;
+	TSharedPtr<TMap<const UNiagaraGraph*, TArray<FDuplicatedGraphData>>> SharedSourceGraphToDuplicatedGraphsMap;
+	TSharedPtr<TMap<FName, UNiagaraDataInterface*>> SharedNameToDuplicatedDataInterfaceMap;
+	TSharedPtr<TMap<UClass*, UNiagaraDataInterface*>> SharedDataInterfaceClassToDuplicatedCDOMap;
 
 protected:
-	void VisitReferencedGraphsRecursive(UNiagaraGraph* InGraph, const FCompileConstantResolver& ConstantResolver, bool bNeedsCompile, TMap<UNiagaraNodeFunctionCall*, ENiagaraScriptUsage> FunctionsWithUsage);
+	void DuplicateReferencedGraphsRecursive(UNiagaraGraph* InGraph, const FCompileConstantResolver& ConstantResolver, TMap<UNiagaraNodeFunctionCall*, ENiagaraScriptUsage> FunctionsWithUsage);
 };
 
 
-
 /** Data which is generated from the hlsl by the VectorVMBackend and fed back into the */
 struct FNiagaraTranslatorOutput
 {
@@ -356,6 +383,7 @@ public:
 
 protected:
 	const FNiagaraCompileRequestData* CompileData;
+	const FNiagaraCompileRequestDuplicateData* CompileDuplicateData;
 	FNiagaraCompileOptions CompileOptions;
 
 	FHlslNiagaraTranslatorOptions TranslationOptions;
@@ -409,9 +437,6 @@ protected:
 	TMap<FNiagaraFunctionSignature, TArray<FName> > FunctionStageWriteTargets;
 	TArray<TArray<FName>> ActiveStageWriteTargets;
 
-	/** Map of function graphs we've seen before and already pre-processed. */
-	TMap<const UNiagaraGraph*, UNiagaraGraph*> PreprocessedFunctions;
-
 	void RegisterFunctionCall(ENiagaraScriptUsage ScriptUsage, const FString& InName, const FString& InFullName, const FGuid& CallNodeId, const FString& InFunctionNameSuffix, UNiagaraScriptSource* Source, FNiagaraFunctionSignature& InSignature, bool bIsCustomHlsl, const FString& InCustomHlsl, TArray<int32>& Inputs, TArrayView<UEdGraphPin* const> CallInputs, TArrayView<UEdGraphPin* const> CallOutputs,
 		FNiagaraFunctionSignature& OutSignature);
 	void GenerateFunctionCall(ENiagaraScriptUsage ScriptUsage, FNiagaraFunctionSignature& FunctionSignature, TArrayView<const int32> Inputs, TArray<int32>& Outputs);
@@ -482,7 +507,7 @@ public:
 
 	static void Init();
 	
-	virtual const FNiagaraTranslateResults &Translate(const FNiagaraCompileRequestData* InCompileData, const FNiagaraCompileOptions& InCompileOptions, FHlslNiagaraTranslatorOptions Options);
+	virtual const FNiagaraTranslateResults &Translate(const FNiagaraCompileRequestData* InCompileData, const FNiagaraCompileRequestDuplicateData* InCompileDuplicateData, const FNiagaraCompileOptions& InCompileOptions, FHlslNiagaraTranslatorOptions Options);
 	FNiagaraTranslatorOutput &GetTranslateOutput() { return CompilationOutput; }
 
 	virtual int32 CompilePin(const UEdGraphPin* Pin);
diff --git a/Engine/Plugins/FX/Niagara/Source/NiagaraEditor/Private/NiagaraScriptSource.cpp b/Engine/Plugins/FX/Niagara/Source/NiagaraEditor/Private/NiagaraScriptSource.cpp
--- a/Engine/Plugins/FX/Niagara/Source/NiagaraEditor/Private/NiagaraScriptSource.cpp
+++ b/Engine/Plugins/FX/Niagara/Source/NiagaraEditor/Private/NiagaraScriptSource.cpp
@@ -104,10 +104,10 @@ void UNiagaraScriptSource::PostLoad()
 	}
 }
 
-UNiagaraScriptSource* UNiagaraScriptSource::CreateCompilationCopy()
+UNiagaraScriptSource* UNiagaraScriptSource::CreateCompilationCopy(const TArray<ENiagaraScriptUsage>& CompileUsages)
 {
 	UNiagaraScriptSource* Result = NewObject<UNiagaraScriptSource>();
-	Result->NodeGraph = NodeGraph->CreateCompilationCopy();
+	Result->NodeGraph = NodeGraph->CreateCompilationCopy(CompileUsages);
 	Result->bIsCompilationCopy = true;
 	Result->AddToRoot();
 	return Result;
diff --git a/Engine/Plugins/FX/Niagara/Source/NiagaraEditor/Public/NiagaraEditorModule.h b/Engine/Plugins/FX/Niagara/Source/NiagaraEditor/Public/NiagaraEditorModule.h
--- a/Engine/Plugins/FX/Niagara/Source/NiagaraEditor/Public/NiagaraEditorModule.h
+++ b/Engine/Plugins/FX/Niagara/Source/NiagaraEditor/Public/NiagaraEditorModule.h
@@ -24,6 +24,7 @@ class FNiagaraSystemViewModel;
 class FNiagaraScriptMergeManager;
 class FNiagaraCompileOptions;
 class FNiagaraCompileRequestDataBase;
+class FNiagaraCompileRequestDuplicateDataBase;
 class UMovieSceneNiagaraParameterTrack;
 struct IConsoleCommand;
 class INiagaraEditorOnlyDataUtilities;
@@ -109,10 +110,16 @@ public:
 	NIAGARAEDITOR_API static FNiagaraEditorModule& Get();
 
 	/** Start the compilation of the specified script. */
-	virtual int32 CompileScript(const FNiagaraCompileRequestDataBase* InCompileRequest, const FNiagaraCompileOptions& InCompileOptions);
+	virtual int32 CompileScript(const FNiagaraCompileRequestDataBase* InCompileRequest, const FNiagaraCompileRequestDuplicateDataBase* InCompileRequestDuplicate, const FNiagaraCompileOptions& InCompileOptions);
 	virtual TSharedPtr<FNiagaraVMExecutableData> GetCompilationResult(int32 JobID, bool bWait);
 
-	TSharedPtr<FNiagaraCompileRequestDataBase, ESPMode::ThreadSafe> Precompile(UObject* Obj, FGuid Version);
+	TSharedPtr<FNiagaraCompileRequestDataBase, ESPMode::ThreadSafe> Precompile(UObject* InObj, FGuid Version);
+	TSharedPtr<FNiagaraCompileRequestDuplicateDataBase, ESPMode::ThreadSafe> PrecompileDuplicate(
+		const FNiagaraCompileRequestDataBase* OwningSystemRequestData,
+		UNiagaraSystem* OwningSystem,
+		UNiagaraEmitter* OwningEmitter,
+		UNiagaraScript* TargetScript,
+		FGuid TargetVersion);
 
 	/** Gets the extensibility managers for outside entities to extend static mesh editor's menus and toolbars */
 	virtual TSharedPtr<FExtensibilityManager> GetMenuExtensibilityManager() override {return MenuExtensibilityManager;}
@@ -280,6 +287,7 @@ private:
 	FDelegateHandle ScriptCompilerHandle;
 	FDelegateHandle CompileResultHandle;
 	FDelegateHandle PrecompilerHandle;
+	FDelegateHandle PrecompileDuplicatorHandle;
 
 	FDelegateHandle DeviceProfileManagerUpdatedHandle;
 
diff --git a/Engine/Plugins/FX/Niagara/Source/NiagaraEditor/Public/NiagaraEditorUtilities.h b/Engine/Plugins/FX/Niagara/Source/NiagaraEditor/Public/NiagaraEditorUtilities.h
--- a/Engine/Plugins/FX/Niagara/Source/NiagaraEditor/Public/NiagaraEditorUtilities.h
+++ b/Engine/Plugins/FX/Niagara/Source/NiagaraEditor/Public/NiagaraEditorUtilities.h
@@ -330,6 +330,9 @@ namespace FNiagaraEditorUtilities
 	void RunPythonUpgradeScripts(UNiagaraNodeFunctionCall* SourceNode, const TArray<FVersionedNiagaraScriptData*>& UpgradeVersionData, const FNiagaraScriptVersionUpgradeContext& UpgradeContext, FString& OutWarnings);
 
 	void RefreshAllScriptsFromExternalChanges(FRefreshAllScriptsFromExternalChangesArgs Args);
+
+	DECLARE_DELEGATE_OneParam(FNodeVisitor, UEdGraphNode* /*VisitedNode*/);
+	void VisitAllNodesConnectedToInputs(UEdGraphNode* StartNode, FNodeVisitor Visitor);
 };
 
 namespace FNiagaraParameterUtilities
diff --git a/Engine/Plugins/FX/Niagara/Source/NiagaraEditor/Public/NiagaraGraph.h b/Engine/Plugins/FX/Niagara/Source/NiagaraEditor/Public/NiagaraGraph.h
--- a/Engine/Plugins/FX/Niagara/Source/NiagaraEditor/Public/NiagaraGraph.h
+++ b/Engine/Plugins/FX/Niagara/Source/NiagaraEditor/Public/NiagaraGraph.h
@@ -156,7 +156,7 @@ class UNiagaraGraph : public UEdGraph
 	bool IsEmpty() const { return Nodes.Num() == 0; }
 
 	/** Creates a transient copy of this graph for compilation purposes. */
-	UNiagaraGraph* CreateCompilationCopy();
+	UNiagaraGraph* CreateCompilationCopy(const TArray<ENiagaraScriptUsage>& CompileUsages);
 	void ReleaseCompilationCopy();
 			
 	/** Find the first output node bound to the target usage type.*/
diff --git a/Engine/Plugins/FX/Niagara/Source/NiagaraEditor/Public/NiagaraScriptSource.h b/Engine/Plugins/FX/Niagara/Source/NiagaraEditor/Public/NiagaraScriptSource.h
--- a/Engine/Plugins/FX/Niagara/Source/NiagaraEditor/Public/NiagaraScriptSource.h
+++ b/Engine/Plugins/FX/Niagara/Source/NiagaraEditor/Public/NiagaraScriptSource.h
@@ -26,7 +26,7 @@ class UNiagaraScriptSource : public UNiagaraScriptSourceBase
 	// UObject interface
 	virtual void PostLoad() override;
 
-	UNiagaraScriptSource* CreateCompilationCopy();
+	UNiagaraScriptSource* CreateCompilationCopy(const TArray<ENiagaraScriptUsage>& CompileUsages);
 
 	void ReleaseCompilationCopy();
 
