Manual merge

Fixed emitter property changes not causing a refresh of the relevant emitter stack view model.
Updating the enable sim stages bool didn't refresh the stack because of it

diff --git a/Engine/Plugins/FX/Niagara/Source/NiagaraEditor/Private/ViewModels/NiagaraSystemViewModel.cpp b/Engine/Plugins/FX/Niagara/Source/NiagaraEditor/Private/ViewModels/NiagaraSystemViewModel.cpp
--- a/Engine/Plugins/FX/Niagara/Source/NiagaraEditor/Private/ViewModels/NiagaraSystemViewModel.cpp
+++ b/Engine/Plugins/FX/Niagara/Source/NiagaraEditor/Private/ViewModels/NiagaraSystemViewModel.cpp
@@ -1137,9 +1137,9 @@ void FNiagaraSystemViewModel::RefreshAll()
 
 void FNiagaraSystemViewModel::NotifyDataObjectChanged(TArray<UObject*> ChangedObjects, ENiagaraDataObjectChange ChangeType)
 {
-	if (ChangedObjects.Num() == 1 && (ChangedObjects[0]->IsA<UNiagaraEmitter>() || ChangedObjects[0]->IsA<UNiagaraSystem>()))
+	if (ChangedObjects.Num() == 1 && ChangedObjects[0]->IsA<UNiagaraSystem>())
 	{
-		// we do nothing on emitter or system changes here, because they will trigger a compile and reset on their own, depending on the changed property
+		// we do nothing on system changes here, because they will trigger a compile and reset on their own, depending on the changed property
 		return;
 	}
 
@@ -1159,6 +1159,16 @@ void FNiagaraSystemViewModel::NotifyDataObjectChanged(TArray<UObject*> ChangedOb
 			{
 				bRefreshCurveSelectionViewModel = true;
 			}
+
+			if(UNiagaraEmitter* Emitter = Cast<UNiagaraEmitter>(ChangedObject))
+			{
+				UNiagaraStackViewModel* StackViewModel = GetEmitterHandleViewModelForEmitter(Emitter)->GetEmitterStackViewModel();
+
+				if(ensure(StackViewModel))
+				{
+					StackViewModel->Refresh();
+				}
+			}
 		}
 	}
 	else
diff --git a/Engine/Plugins/FX/Niagara/Source/NiagaraEditor/Public/ViewModels/Stack/NiagaraStackViewModel.h b/Engine/Plugins/FX/Niagara/Source/NiagaraEditor/Public/ViewModels/Stack/NiagaraStackViewModel.h
--- a/Engine/Plugins/FX/Niagara/Source/NiagaraEditor/Public/ViewModels/Stack/NiagaraStackViewModel.h
+++ b/Engine/Plugins/FX/Niagara/Source/NiagaraEditor/Public/ViewModels/Stack/NiagaraStackViewModel.h
@@ -140,6 +140,7 @@ public:
 
 	bool HasIssues() const;
 
+	void Refresh() { bRefreshPending = true; }
 private:
 	/** Recursively Expands all groups and collapses all items in the stack. */
 	void CollapseToHeadersRecursive(TArray<UNiagaraStackEntry*> Entries);
