- Move NiagaraEmitterInstanceBatcher from public to private
- Replace with FNiagaraGpuBatcherInterface which will be the external interface to the batcher for data interfaces

diff --git a/Engine/Plugins/Experimental/ChaosNiagara/Source/ChaosNiagara/Classes/NiagaraDataInterfaceChaosDestruction.h b/Engine/Plugins/Experimental/ChaosNiagara/Source/ChaosNiagara/Classes/NiagaraDataInterfaceChaosDestruction.h
--- a/Engine/Plugins/Experimental/ChaosNiagara/Source/ChaosNiagara/Classes/NiagaraDataInterfaceChaosDestruction.h
+++ b/Engine/Plugins/Experimental/ChaosNiagara/Source/ChaosNiagara/Classes/NiagaraDataInterfaceChaosDestruction.h
@@ -770,7 +770,7 @@ struct FNiagaraDataInterfaceProxyChaosDestruction : public FNiagaraDataInterface
 
 	void CreatePerInstanceData(const FNiagaraSystemInstanceID& SystemInstance);
 
-	void DestroyInstanceData(NiagaraEmitterInstanceBatcher* Batcher, const FNiagaraSystemInstanceID& SystemInstance);
+	void DestroyInstanceData(const FNiagaraSystemInstanceID& SystemInstance);
 
 	float SolverTime;
 	int32 LastSpawnedPointID;
diff --git a/Engine/Plugins/Experimental/ChaosNiagara/Source/ChaosNiagara/Classes/NiagaraDataInterfacePhysicsField.h b/Engine/Plugins/Experimental/ChaosNiagara/Source/ChaosNiagara/Classes/NiagaraDataInterfacePhysicsField.h
--- a/Engine/Plugins/Experimental/ChaosNiagara/Source/ChaosNiagara/Classes/NiagaraDataInterfacePhysicsField.h
+++ b/Engine/Plugins/Experimental/ChaosNiagara/Source/ChaosNiagara/Classes/NiagaraDataInterfacePhysicsField.h
@@ -108,7 +108,7 @@ struct FNDIPhysicsFieldProxy : public FNiagaraDataInterfaceProxy
 	void InitializePerInstanceData(const FNiagaraSystemInstanceID& SystemInstance);
 
 	/** Destroy the proxy data if necessary */
-	void DestroyPerInstanceData(NiagaraEmitterInstanceBatcher* Batcher, const FNiagaraSystemInstanceID& SystemInstance);
+	void DestroyPerInstanceData(const FNiagaraSystemInstanceID& SystemInstance);
 
 	/** List of proxy data for each system instances*/
 	TMap<FNiagaraSystemInstanceID, FNDIPhysicsFieldData> SystemInstancesToProxyData;
diff --git a/Engine/Plugins/Experimental/ChaosNiagara/Source/ChaosNiagara/Private/NiagaraDataInterfaceChaosDestruction.cpp b/Engine/Plugins/Experimental/ChaosNiagara/Source/ChaosNiagara/Private/NiagaraDataInterfaceChaosDestruction.cpp
--- a/Engine/Plugins/Experimental/ChaosNiagara/Source/ChaosNiagara/Private/NiagaraDataInterfaceChaosDestruction.cpp
+++ b/Engine/Plugins/Experimental/ChaosNiagara/Source/ChaosNiagara/Private/NiagaraDataInterfaceChaosDestruction.cpp
@@ -12,7 +12,6 @@
 #include "Chaos/CollisionResolutionTypes.h"
 #include "PhysicalMaterials/PhysicalMaterial.h"
 #include "NiagaraComponent.h"
-#include "NiagaraEmitterInstanceBatcher.h"
 
 #include <memory>
 #include "GeometryCollection/GeometryCollectionComponent.h"
@@ -534,10 +533,10 @@ void UNiagaraDataInterfaceChaosDestruction::DestroyPerInstanceData(void* PerInst
 	{
 		FNiagaraDataInterfaceProxyChaosDestruction* ThisProxy = GetProxyAs<FNiagaraDataInterfaceProxyChaosDestruction>();
 		ENQUEUE_RENDER_COMMAND(FNiagaraDIChaosDestructionDestroyInstanceData) (
-			[ThisProxy, InstanceID = SystemInstance->GetId(), Batcher = SystemInstance->GetBatcher()](FRHICommandListImmediate& CmdList)
-		{
-			ThisProxy->DestroyInstanceData(Batcher, InstanceID);
-		}
+			[ThisProxy, InstanceID = SystemInstance->GetId()](FRHICommandListImmediate& CmdList)
+			{
+				ThisProxy->DestroyInstanceData(InstanceID);
+			}
 		);
 	}
 }
@@ -4121,7 +4120,7 @@ void FNiagaraDataInterfaceProxyChaosDestruction::CreatePerInstanceData(const FNi
 	SystemsToGPUInstanceData.Add(SystemInstance, FNiagaraDIChaosDestruction_GPUData());
 }
 
-void FNiagaraDataInterfaceProxyChaosDestruction::DestroyInstanceData(NiagaraEmitterInstanceBatcher* Batcher, const FNiagaraSystemInstanceID& SystemInstance)
+void FNiagaraDataInterfaceProxyChaosDestruction::DestroyInstanceData(const FNiagaraSystemInstanceID& SystemInstance)
 {
 	check(IsInRenderingThread());
 	check(SystemsToGPUInstanceData.Contains(SystemInstance));
diff --git a/Engine/Plugins/Experimental/ChaosNiagara/Source/ChaosNiagara/Private/NiagaraDataInterfacePhysicsField.cpp b/Engine/Plugins/Experimental/ChaosNiagara/Source/ChaosNiagara/Private/NiagaraDataInterfacePhysicsField.cpp
--- a/Engine/Plugins/Experimental/ChaosNiagara/Source/ChaosNiagara/Private/NiagaraDataInterfacePhysicsField.cpp
+++ b/Engine/Plugins/Experimental/ChaosNiagara/Source/ChaosNiagara/Private/NiagaraDataInterfacePhysicsField.cpp
@@ -5,7 +5,6 @@
 #include "NiagaraComponent.h"
 #include "NiagaraRenderer.h"
 #include "NiagaraSystemInstance.h"
-#include "NiagaraEmitterInstanceBatcher.h"
 #include "ShaderParameterUtils.h"
 #include "Field/FieldSystemNodes.h"
 #include "PhysicsField/PhysicsFieldComponent.h"
@@ -214,7 +213,7 @@ void FNDIPhysicsFieldProxy::InitializePerInstanceData(const FNiagaraSystemInstan
 	TargetData = &SystemInstancesToProxyData.Add(SystemInstance);
 }
 
-void FNDIPhysicsFieldProxy::DestroyPerInstanceData(NiagaraEmitterInstanceBatcher* Batcher, const FNiagaraSystemInstanceID& SystemInstance)
+void FNDIPhysicsFieldProxy::DestroyPerInstanceData(const FNiagaraSystemInstanceID& SystemInstance)
 {
 	check(IsInRenderingThread());
 	SystemInstancesToProxyData.Remove(SystemInstance);
@@ -249,7 +248,7 @@ void UNiagaraDataInterfacePhysicsField::DestroyPerInstanceData(void* PerInstance
 
 	FNDIPhysicsFieldProxy* ThisProxy = GetProxyAs<FNDIPhysicsFieldProxy>();
 	ENQUEUE_RENDER_COMMAND(FNiagaraDIDestroyInstanceData) (
-		[ThisProxy, InstanceID = SystemInstance->GetId(), Batcher = SystemInstance->GetBatcher()](FRHICommandListImmediate& CmdList)
+		[ThisProxy, InstanceID = SystemInstance->GetId()](FRHICommandListImmediate& CmdList)
 	{
 		ThisProxy->SystemInstancesToProxyData.Remove(InstanceID);
 	}
diff --git a/Engine/Plugins/FX/Niagara/Source/Niagara/Classes/NiagaraComputeExecutionContext.h b/Engine/Plugins/FX/Niagara/Source/Niagara/Classes/NiagaraComputeExecutionContext.h
--- a/Engine/Plugins/FX/Niagara/Source/Niagara/Classes/NiagaraComputeExecutionContext.h
+++ b/Engine/Plugins/FX/Niagara/Source/Niagara/Classes/NiagaraComputeExecutionContext.h
@@ -2,11 +2,12 @@
 
 #pragma once
 
+#include "NiagaraSimStageData.h"
 #include "NiagaraScriptExecutionContext.h"
 #include "RHIGPUReadback.h"
 
 class FNiagaraGPUInstanceCountManager;
-class NiagaraEmitterInstanceBatcher;
+class FNiagaraGpuComputeDispatchInterface;
 class FNiagaraGPUSystemTick;
 struct FNiagaraComputeInstanceData;
 struct FNiagaraComputeExecutionContext;
@@ -20,34 +21,6 @@ public:
 	}
 };
 
-struct FNiagaraSimStageData
-{
-	FNiagaraSimStageData()
-	{
-		bFirstStage = false;
-		bLastStage = false;
-		bSetDataToRender = false;
-	}
-
-	uint32 bFirstStage : 1;
-	uint32 bLastStage : 1;
-	uint32 bSetDataToRender : 1;
-
-	uint32 StageIndex = INDEX_NONE;
-	uint32 IterationIndex = 0;
-
-	FNiagaraDataBuffer* Source = nullptr;
-	uint32 SourceCountOffset = INDEX_NONE;
-	uint32 SourceNumInstances = 0;
-
-	FNiagaraDataBuffer* Destination = nullptr;
-	uint32 DestinationCountOffset = INDEX_NONE;
-	uint32 DestinationNumInstances = 0;
-
-	FNiagaraDataInterfaceProxyRW* AlternateIterationSource = nullptr;
-	const FSimulationStageMetaData* StageMetaData = nullptr;
-};
-
 struct FNiagaraGpuDispatchInstance
 {
 	FNiagaraGpuDispatchInstance(const FNiagaraGPUSystemTick& InTick, const FNiagaraComputeInstanceData& InInstanceData)
@@ -123,7 +96,7 @@ struct FNiagaraComputeExecutionContext
 	FNiagaraComputeExecutionContext();
 	~FNiagaraComputeExecutionContext();
 
-	void Reset(NiagaraEmitterInstanceBatcher* Batcher);
+	void Reset(FNiagaraGpuComputeDispatchInterface* ComputeDispatchInterface);
 
 	void InitParams(UNiagaraScript* InGPUComputeScript, ENiagaraSimTarget InSimTarget);
 	void DirtyDataInterfaces();
@@ -154,7 +127,7 @@ struct FNiagaraComputeExecutionContext
 #endif
 
 //-TOOD:private:
-	void ResetInternal(NiagaraEmitterInstanceBatcher* Batcher);
+	void ResetInternal(FNiagaraGpuComputeDispatchInterface* ComputeDispatchInterface);
 
 public:
 	static uint32 TickCounter;
diff --git a/Engine/Plugins/FX/Niagara/Source/Niagara/Classes/NiagaraDataInterface.h b/Engine/Plugins/FX/Niagara/Source/Niagara/Classes/NiagaraDataInterface.h
--- a/Engine/Plugins/FX/Niagara/Source/Niagara/Classes/NiagaraDataInterface.h
+++ b/Engine/Plugins/FX/Niagara/Source/Niagara/Classes/NiagaraDataInterface.h
@@ -252,10 +252,10 @@ struct FNiagaraDataInterfaceProxy : TSharedFromThis<FNiagaraDataInterfaceProxy,
 	virtual void PostSimulate(FRHICommandList& RHICmdList, const FNiagaraDataInterfaceArgs& Context) {}
 
 	virtual bool RequiresPreStageFinalize() const { return false; }
-	virtual void FinalizePreStage(FRHICommandList& RHICmdList, const NiagaraEmitterInstanceBatcher* Batcher) {}
+	virtual void FinalizePreStage(FRHICommandList& RHICmdList, const FNiagaraGpuComputeDispatchInterface* ComputeDispatchInterface) {}
 
 	virtual bool RequiresPostStageFinalize() const { return false; }
-	virtual void FinalizePostStage(FRHICommandList& RHICmdList, const NiagaraEmitterInstanceBatcher* Batcher) {}
+	virtual void FinalizePostStage(FRHICommandList& RHICmdList, const FNiagaraGpuComputeDispatchInterface* ComputeDispatchInterface) {}
 
 	virtual FNiagaraDataInterfaceProxyRW* AsIterationProxy() { return nullptr; }
 };
diff --git a/Engine/Plugins/FX/Niagara/Source/Niagara/Classes/NiagaraDataInterfaceStaticMesh.h b/Engine/Plugins/FX/Niagara/Source/Niagara/Classes/NiagaraDataInterfaceStaticMesh.h
--- a/Engine/Plugins/FX/Niagara/Source/Niagara/Classes/NiagaraDataInterfaceStaticMesh.h
+++ b/Engine/Plugins/FX/Niagara/Source/Niagara/Classes/NiagaraDataInterfaceStaticMesh.h
@@ -491,7 +491,7 @@ struct FNiagaraDataInterfaceProxyStaticMesh : public FNiagaraDataInterfaceProxy
 	virtual void ConsumePerInstanceDataFromGameThread(void* PerInstanceData, const FNiagaraSystemInstanceID& Instance) override;
 
 	void InitializePerInstanceData(const FNiagaraSystemInstanceID& SystemInstance, FStaticMeshGpuSpawnBuffer* MeshGPUSpawnBuffer);
-	void DestroyPerInstanceData(NiagaraEmitterInstanceBatcher* Batcher, const FNiagaraSystemInstanceID& SystemInstance);
+	void DestroyPerInstanceData(const FNiagaraSystemInstanceID& SystemInstance);
 
 	TMap<FNiagaraSystemInstanceID, FNiagaraStaticMeshData> SystemInstancesToMeshData;
 };
diff --git a/Engine/Plugins/FX/Niagara/Source/Niagara/Classes/NiagaraDataSet.h b/Engine/Plugins/FX/Niagara/Source/Niagara/Classes/NiagaraDataSet.h
--- a/Engine/Plugins/FX/Niagara/Source/Niagara/Classes/NiagaraDataSet.h
+++ b/Engine/Plugins/FX/Niagara/Source/Niagara/Classes/NiagaraDataSet.h
@@ -9,8 +9,6 @@
 #include "RenderingThread.h"
 #include "NiagaraDataSet.generated.h"
 
-class NiagaraEmitterInstanceBatcher;
-
 /** Helper class defining the layout and location of an FNiagaraVariable in an FNiagaraDataBuffer-> */
 USTRUCT()
 struct FNiagaraVariableLayoutInfo
@@ -41,7 +39,7 @@ struct FNiagaraVariableLayoutInfo
 class FNiagaraDataSet;
 class FNiagaraShader;
 class FNiagaraGPUInstanceCountManager;
-class NiagaraEmitterInstanceBatcher;
+class FNiagaraGpuComputeDispatchInterface;
 struct FNiagaraComputeExecutionContext;
 
 //Base class for objects in Niagara that are owned by one object but are then passed for reading to other objects, potentially on other threads.
@@ -105,7 +103,7 @@ protected:
 class NIAGARA_API FNiagaraDataBuffer : public FNiagaraSharedObject
 {
 	friend class FScopedNiagaraDataSetGPUReadback;
-	friend class NiagaraEmitterInstanceBatcher;
+	//friend class FNiagaraGpuComputeDispatchInterface;
 	
 protected:
 	virtual ~FNiagaraDataBuffer();
@@ -171,6 +169,8 @@ public:
 	FORCEINLINE uint32 GetGPUInstanceCountBufferOffset() const { return GPUInstanceCountBufferOffset; }
 	FORCEINLINE FRWBuffer& GetGPUIDToIndexTable() { return GPUIDToIndexTable; }
 
+	FORCEINLINE void SetGPUInstanceCountBufferOffset(uint32 Offset) { GPUInstanceCountBufferOffset = Offset; }
+
 	FORCEINLINE int32 GetSafeComponentBufferSize() const { return GetSafeComponentBufferSize(GetNumInstancesAllocated()); }
 	FORCEINLINE uint32 GetFloatStride() const { return FloatStride; }
 	FORCEINLINE uint32 GetInt32Stride() const { return Int32Stride; }
@@ -192,10 +192,7 @@ public:
 	static void SetOutputShaderParams(FRHICommandList& RHICmdList, class FNiagaraShader* Shader, FNiagaraDataBuffer* Buffer);
 	static void UnsetShaderParams(FRHICommandList& RHICmdList, class FNiagaraShader* Shader);
 
-	void ClearGPUInstanceCount()
-	{
-		GPUInstanceCountBufferOffset = INDEX_NONE;
-	}
+	FORCEINLINE void ClearGPUInstanceCount() { GPUInstanceCountBufferOffset = INDEX_NONE; }
 
 	void BuildRegisterTable();
 
@@ -317,7 +314,7 @@ General storage class for all per instance simulation data in Niagara.
 class NIAGARA_API FNiagaraDataSet
 {
 	friend FNiagaraDataBuffer;
-	friend class NiagaraEmitterInstanceBatcher;
+	friend class FNiagaraGpuComputeDispatch;
 
 public:
 
@@ -551,13 +548,13 @@ public:
 		}
 	}
 
-	void ReadbackData(NiagaraEmitterInstanceBatcher* Batcher, FNiagaraDataSet* InDataSet);
+	void ReadbackData(FNiagaraGpuComputeDispatchInterface* ComputeDispatchInterface, FNiagaraDataSet* InDataSet);
 	uint32 GetNumInstances() const { check(DataSet != nullptr); return NumInstances; }
 
 private:
 	FNiagaraDataSet*	DataSet = nullptr;
 	FNiagaraDataBuffer* DataBuffer = nullptr;
-	NiagaraEmitterInstanceBatcher* Batcher = nullptr;
+	FNiagaraGpuComputeDispatchInterface* ComputeDispatchInterface = nullptr;
 	uint32				NumInstances = 0;
 };
 #endif
diff --git a/Engine/Plugins/FX/Niagara/Source/Niagara/Classes/NiagaraEmitterInstance.h b/Engine/Plugins/FX/Niagara/Source/Niagara/Classes/NiagaraEmitterInstance.h
--- a/Engine/Plugins/FX/Niagara/Source/Niagara/Classes/NiagaraEmitterInstance.h
+++ b/Engine/Plugins/FX/Niagara/Source/Niagara/Classes/NiagaraEmitterInstance.h
@@ -16,7 +16,7 @@ class FNiagaraSystemInstance;
 struct FNiagaraEmitterHandle;
 class UNiagaraParameterCollection;
 class UNiagaraParameterCollectionInstance;
-class NiagaraEmitterInstanceBatcher;
+class FNiagaraGpuComputeDispatchInterface;
 struct FNiagaraEmitterCompiledData;
 
 /**
@@ -210,7 +210,7 @@ private:
 
 	FNiagaraSystemInstanceID OwnerSystemInstanceID;
 
-	NiagaraEmitterInstanceBatcher* Batcher = nullptr;
+	FNiagaraGpuComputeDispatchInterface* ComputeDispatchInterface = nullptr;
 
 	/** particle simulation data. Must be a shared ref as various things on the RT can have direct ref to it. */
 	FNiagaraDataSet* ParticleDataSet = nullptr;
diff --git a/Engine/Plugins/FX/Niagara/Source/Niagara/Classes/NiagaraEmitterInstanceBatcher.h b/Engine/Plugins/FX/Niagara/Source/Niagara/Classes/NiagaraEmitterInstanceBatcher.h
deleted file mode 100644
diff --git a/Engine/Plugins/FX/Niagara/Source/Niagara/Classes/NiagaraGPUInstanceCountManager.h b/Engine/Plugins/FX/Niagara/Source/Niagara/Classes/NiagaraGPUInstanceCountManager.h
--- a/Engine/Plugins/FX/Niagara/Source/Niagara/Classes/NiagaraGPUInstanceCountManager.h
+++ b/Engine/Plugins/FX/Niagara/Source/Niagara/Classes/NiagaraGPUInstanceCountManager.h
@@ -12,7 +12,7 @@ NiagaraGPUInstanceCountManager.h: GPU particle count handling
 #include "RHIGPUReadback.h"
 
 class FRHIGPUMemoryReadback;
-class NiagaraEmitterInstanceBatcher;
+class FNiagaraGpuComputeDispatchInterface;
 
 // The number of GPU renderers registered in the instance count manager.
 // Shared between the manager and the renderers.
@@ -99,7 +99,7 @@ public:
 	void FlushIndirectArgsPool();
 
 	// Generate the draw indirect buffers, and reset all release counts.
-	void UpdateDrawIndirectBuffers(NiagaraEmitterInstanceBatcher& Batcher, FRHICommandList& RHICmdList, ERHIFeatureLevel::Type FeatureLevel, ENiagaraGPUCountUpdatePhase::Type CountPhase);
+	void UpdateDrawIndirectBuffers(FNiagaraGpuComputeDispatchInterface* ComputeDispatchInterface, FRHICommandList& RHICmdList, ERHIFeatureLevel::Type FeatureLevel, ENiagaraGPUCountUpdatePhase::Type CountPhase);
 
 	static const ERHIAccess kCountBufferDefaultState;
 
diff --git a/Engine/Plugins/FX/Niagara/Source/Niagara/Private/DataInterface/NiagaraDataInterfaceActorComponent.cpp b/Engine/Plugins/FX/Niagara/Source/Niagara/Private/DataInterface/NiagaraDataInterfaceActorComponent.cpp
--- a/Engine/Plugins/FX/Niagara/Source/Niagara/Private/DataInterface/NiagaraDataInterfaceActorComponent.cpp
+++ b/Engine/Plugins/FX/Niagara/Source/Niagara/Private/DataInterface/NiagaraDataInterfaceActorComponent.cpp
@@ -3,12 +3,12 @@
 #include "NiagaraDataInterfaceActorComponent.h"
 #include "NiagaraTypes.h"
 #include "NiagaraCustomVersion.h"
-#include "NiagaraEmitterInstanceBatcher.h"
 #include "NiagaraSystemInstance.h"
 
 #include "Components/SceneComponent.h"
 #include "GameFramework/Actor.h"
 #include "ShaderCompilerCore.h"
+#include "ShaderParameterUtils.h"
 
 #define LOCTEXT_NAMESPACE "NiagaraDataInterfaceActorComponent"
 
diff --git a/Engine/Plugins/FX/Niagara/Source/Niagara/Private/DataInterface/NiagaraDataInterfaceDebugDraw.cpp b/Engine/Plugins/FX/Niagara/Source/Niagara/Private/DataInterface/NiagaraDataInterfaceDebugDraw.cpp
--- a/Engine/Plugins/FX/Niagara/Source/Niagara/Private/DataInterface/NiagaraDataInterfaceDebugDraw.cpp
+++ b/Engine/Plugins/FX/Niagara/Source/Niagara/Private/DataInterface/NiagaraDataInterfaceDebugDraw.cpp
@@ -2,13 +2,14 @@
 
 #include "NiagaraDataInterfaceDebugDraw.h"
 #include "NiagaraTypes.h"
-#include "NiagaraEmitterInstanceBatcher.h"
+#include "NiagaraGpuComputeDispatchInterface.h"
 #include "NiagaraGpuComputeDebug.h"
 #include "NiagaraWorldManager.h"
 #include "NiagaraSystemInstance.h"
 
 #include "Async/Async.h"
 #include "DrawDebugHelpers.h"
+#include "ShaderParameterUtils.h"
 
 //////////////////////////////////////////////////////////////////////////
 
@@ -3241,11 +3242,11 @@ bool UNiagaraDataInterfaceDebugDraw::InitPerInstanceData(void* PerInstanceData,
 
 #if NIAGARA_COMPUTEDEBUG_ENABLED
 	ENQUEUE_RENDER_COMMAND(NDIDebugDrawInit)(
-		[RT_Proxy=GetProxyAs<FNDIDebugDrawProxy>(), RT_InstanceID=SystemInstance->GetId(), RT_Batcher=SystemInstance->GetBatcher()](FRHICommandListImmediate& RHICmdList)
+		[RT_Proxy=GetProxyAs<FNDIDebugDrawProxy>(), RT_InstanceID=SystemInstance->GetId(), RT_ComputeDispatchInterface=SystemInstance->GetComputeDispatchInterface()](FRHICommandListImmediate& RHICmdList)
 		{
 			check(!RT_Proxy->SystemInstancesToProxyData_RT.Contains(RT_InstanceID));
 			FNDIDebugDrawInstanceData_RenderThread* RT_InstanceData = &RT_Proxy->SystemInstancesToProxyData_RT.Add(RT_InstanceID);
-			RT_InstanceData->GpuComputeDebug = RT_Batcher->GetGpuComputeDebug();
+			RT_InstanceData->GpuComputeDebug = RT_ComputeDispatchInterface->GetGpuComputeDebug();
 		}
 	);
 #endif
diff --git a/Engine/Plugins/FX/Niagara/Source/Niagara/Private/DataInterface/NiagaraDataInterfaceGBuffer.cpp b/Engine/Plugins/FX/Niagara/Source/Niagara/Private/DataInterface/NiagaraDataInterfaceGBuffer.cpp
--- a/Engine/Plugins/FX/Niagara/Source/Niagara/Private/DataInterface/NiagaraDataInterfaceGBuffer.cpp
+++ b/Engine/Plugins/FX/Niagara/Source/Niagara/Private/DataInterface/NiagaraDataInterfaceGBuffer.cpp
@@ -1,7 +1,8 @@
 // Copyright Epic Games, Inc. All Rights Reserved.
 
 #include "NiagaraDataInterfaceGBuffer.h"
-#include "NiagaraEmitterInstanceBatcher.h"
+#include "NiagaraGpuComputeDispatchInterface.h"
+#include "NiagaraGpuComputeDispatch.h"
 #include "NiagaraTypes.h"
 #include "NiagaraRenderViewDataManager.h"
 #include "NiagaraWorldManager.h"
@@ -124,7 +125,7 @@ public:
 
 		FRHISamplerState* VelocitySamplerState = TStaticSamplerState<SF_Bilinear, AM_Clamp, AM_Clamp, AM_Clamp>::GetRHI();
 		FRHITexture* VelocityRHITexture = GBlackTexture->TextureRHI;
-		if ( FNiagaraSceneTextureParameters* NiagaraSceneTextures = Context.Batcher->GetNiagaraSceneTextures() )
+		if ( FNiagaraSceneTextureParameters* NiagaraSceneTextures = static_cast<const FNiagaraGpuComputeDispatch*>(Context.ComputeDispatchInterface)->GetNiagaraSceneTextures() )	//-BATCHERTODO:
 		{
 			if ( FRDGTexture* VelocityRDGTexture = NiagaraSceneTextures->Velocity.GetTexture() )
 			{
diff --git a/Engine/Plugins/FX/Niagara/Source/Niagara/Private/NiagaraComponent.cpp b/Engine/Plugins/FX/Niagara/Source/Niagara/Private/NiagaraComponent.cpp
--- a/Engine/Plugins/FX/Niagara/Source/Niagara/Private/NiagaraComponent.cpp
+++ b/Engine/Plugins/FX/Niagara/Source/Niagara/Private/NiagaraComponent.cpp
@@ -13,8 +13,8 @@
 #include "NiagaraDataInterface.h"
 #include "NiagaraDataSetAccessor.h"
 #include "NiagaraEmitterInstance.h"
-#include "NiagaraEmitterInstanceBatcher.h"
 #include "NiagaraFunctionLibrary.h"
+#include "NiagaraGpuComputeDispatchInterface.h"
 #include "NiagaraRenderer.h"
 #include "NiagaraStats.h"
 #include "NiagaraSystem.h"
@@ -158,7 +158,7 @@ FNiagaraSceneProxy::FNiagaraSceneProxy(UNiagaraComponent* InComponent)
 		LLM_SCOPE(ELLMTag::Niagara);
 
 		RenderData = SystemInstanceController->CreateSystemRenderData(GetScene().GetFeatureLevel());
-		Batcher = SystemInstanceController->GetBatcher();
+		ComputeDispatchInterface = SystemInstanceController->GetComputeDispatchInterface();
 
 		bAlwaysHasVelocity = RenderData->HasAnyMotionBlurEnabled();
 
@@ -212,7 +212,7 @@ void FNiagaraSceneProxy::CreateRenderThreadResources()
 	if (RenderData)
 	{
 		LLM_SCOPE(ELLMTag::Niagara);
-		RenderData->CreateRenderThreadResources(*Batcher);
+		RenderData->CreateRenderThreadResources();
 	}
 }
 
@@ -814,33 +814,23 @@ UNiagaraDataInterface* UNiagaraComponent::GetDataInterface(const FString& Name)
 
 bool UNiagaraComponent::IsWorldReadyToRun() const
 {
-	// The niagara system instance assumes that a batcher exists when it is created. We need to wait until this has happened before successfully activating this system.
-	bool FXSystemExists = false;
-	bool WorldManagerExists = false;
+	// The Niagara system instance assumes that a ComputeDispatchInterface exists when it is created. We need to wait until this has happened before successfully activating this system.
 	UWorld* World = GetWorld();
-	if (World)
+	if ( World == nullptr)
 	{
-		if (World->Scene)
-		{
-			FFXSystemInterface*  FXSystemInterface = World->Scene->GetFXSystem();
-			if (FXSystemInterface)
-			{
-				NiagaraEmitterInstanceBatcher* FoundBatcher = static_cast<NiagaraEmitterInstanceBatcher*>(FXSystemInterface->GetInterface(NiagaraEmitterInstanceBatcher::Name));
-				if (FoundBatcher != nullptr)
-				{
-					FXSystemExists = true;
-				}
-			}
-		}
+		return false;
+	}
 
-		FNiagaraWorldManager* WorldManager = FNiagaraWorldManager::Get(World);
-		if (WorldManager)
-		{
-			WorldManagerExists = true;
-		}
+	if ( FNiagaraGpuComputeDispatchInterface::Get(World) == nullptr )
+	{
+		return false;
 	}
 
-	return WorldManagerExists && FXSystemExists;
+	if (FNiagaraWorldManager::Get(World) == nullptr)
+	{
+		return false;
+	}
+	return true;
 }
 
 bool UNiagaraComponent::InitializeSystem()
diff --git a/Engine/Plugins/FX/Niagara/Source/Niagara/Private/NiagaraComponentPool.cpp b/Engine/Plugins/FX/Niagara/Source/Niagara/Private/NiagaraComponentPool.cpp
--- a/Engine/Plugins/FX/Niagara/Source/Niagara/Private/NiagaraComponentPool.cpp
+++ b/Engine/Plugins/FX/Niagara/Source/Niagara/Private/NiagaraComponentPool.cpp
@@ -8,7 +8,7 @@
 #include "NiagaraComponent.h"
 #include "NiagaraWorldManager.h"
 #include "NiagaraCrashReporterHandler.h"
-#include "NiagaraEmitterInstanceBatcher.h"
+#include "NiagaraGpuComputeDispatchInterface.h"
 
 static float GNiagaraSystemPoolKillUnusedTime = 180.0f;
 static FAutoConsoleVariableRef NiagaraSystemPoolKillUnusedTime(
@@ -337,9 +337,9 @@ void UNiagaraComponentPool::PrimePool(UNiagaraSystem* Template, UWorld* World)
 	FFXSystemInterface* FXSystemInterface = World->Scene->GetFXSystem();
 	if (FXSystemInterface)
 	{
-		if (FXSystemInterface->GetInterface(NiagaraEmitterInstanceBatcher::Name) == nullptr)
+		if (FNiagaraGpuComputeDispatchInterface::Get(FXSystemInterface) == nullptr)
 		{
-			UE_LOG(LogNiagara, Verbose, TEXT("Failed to prime particle pool as the world does not have a NiagaraEmitterInstanceBatcher."));
+			UE_LOG(LogNiagara, Verbose, TEXT("Failed to prime particle pool as the world does not have a FNiagaraGpuComputeDispatchInterface."));
 			return;
 		}
 	}
diff --git a/Engine/Plugins/FX/Niagara/Source/Niagara/Private/NiagaraComputeExecutionContext.cpp b/Engine/Plugins/FX/Niagara/Source/Niagara/Private/NiagaraComputeExecutionContext.cpp
--- a/Engine/Plugins/FX/Niagara/Source/Niagara/Private/NiagaraComputeExecutionContext.cpp
+++ b/Engine/Plugins/FX/Niagara/Source/Niagara/Private/NiagaraComputeExecutionContext.cpp
@@ -2,7 +2,8 @@
 
 #include "NiagaraComputeExecutionContext.h"
 #include "NiagaraStats.h"
-#include "NiagaraEmitterInstanceBatcher.h"
+#include "NiagaraGpuComputeDispatchInterface.h"
+#include "NiagaraGpuComputeDispatch.h"
 #include "NiagaraDataInterface.h"
 #include "NiagaraSystemInstance.h"
 #include "NiagaraSystemGpuComputeProxy.h"
@@ -23,20 +24,20 @@ FNiagaraComputeExecutionContext::FNiagaraComputeExecutionContext()
 FNiagaraComputeExecutionContext::~FNiagaraComputeExecutionContext()
 {
 	// EmitterInstanceReadback.GPUCountOffset should be INDEX_NONE at this point to ensure the index is reused.
-	// When the batcher is being destroyed though, we don't free the index, but this would not be leaking.
+	// When the ComputeDispatchInterface is being destroyed though, we don't free the index, but this would not be leaking.
 	// check(EmitterInstanceReadback.GPUCountOffset == INDEX_NONE);
 	SetDataToRender(nullptr);
 
 	ExternalCBufferLayout = nullptr;
 }
 
-void FNiagaraComputeExecutionContext::Reset(NiagaraEmitterInstanceBatcher* Batcher)
+void FNiagaraComputeExecutionContext::Reset(FNiagaraGpuComputeDispatchInterface* ComputeDispatchInterface)
 {
-	NiagaraEmitterInstanceBatcher* RT_Batcher = Batcher && !Batcher->IsPendingKill() ? Batcher : nullptr;
+	FNiagaraGpuComputeDispatchInterface* RT_ComputeDispatchInterface = ComputeDispatchInterface && !ComputeDispatchInterface->IsPendingKill() ? ComputeDispatchInterface : nullptr;
 	ENQUEUE_RENDER_COMMAND(ResetRT)(
-		[RT_Batcher, RT_Context=this](FRHICommandListImmediate& RHICmdList)
+		[RT_ComputeDispatchInterface, RT_Context=this](FRHICommandListImmediate& RHICmdList)
 		{
-			RT_Context->ResetInternal(RT_Batcher);
+			RT_Context->ResetInternal(RT_ComputeDispatchInterface);
 		}
 	);
 }
@@ -189,13 +190,13 @@ void FNiagaraComputeExecutionContext::PostTick()
 	}
 }
 
-void FNiagaraComputeExecutionContext::ResetInternal(NiagaraEmitterInstanceBatcher* Batcher)
+void FNiagaraComputeExecutionContext::ResetInternal(FNiagaraGpuComputeDispatchInterface* ComputeDispatchInterface)
 {
 	checkf(IsInRenderingThread(), TEXT("Can only reset the gpu context from the render thread"));
 
-	if (Batcher)
+	if (ComputeDispatchInterface)
 	{
-		Batcher->GetGPUInstanceCounterManager().FreeEntry(CountOffset_RT);
+		static_cast<FNiagaraGpuComputeDispatch*>(ComputeDispatchInterface)->GetGPUInstanceCounterManager().FreeEntry(CountOffset_RT);
 	}
 
 	CurrentNumInstances_RT = 0;
diff --git a/Engine/Plugins/FX/Niagara/Source/Niagara/Private/NiagaraDataInterfaceAudioOscilloscope.cpp b/Engine/Plugins/FX/Niagara/Source/Niagara/Private/NiagaraDataInterfaceAudioOscilloscope.cpp
--- a/Engine/Plugins/FX/Niagara/Source/Niagara/Private/NiagaraDataInterfaceAudioOscilloscope.cpp
+++ b/Engine/Plugins/FX/Niagara/Source/Niagara/Private/NiagaraDataInterfaceAudioOscilloscope.cpp
@@ -9,7 +9,6 @@
 #include "ClearQuad.h"
 #include "TextureResource.h"
 #include "Engine/Texture2D.h"
-#include "NiagaraEmitterInstanceBatcher.h"
 #include "NiagaraSystemInstance.h"
 #include "Engine/TextureRenderTarget2D.h"
 #include "Engine/Engine.h"
diff --git a/Engine/Plugins/FX/Niagara/Source/Niagara/Private/NiagaraDataInterfaceCollisionQuery.cpp b/Engine/Plugins/FX/Niagara/Source/Niagara/Private/NiagaraDataInterfaceCollisionQuery.cpp
--- a/Engine/Plugins/FX/Niagara/Source/Niagara/Private/NiagaraDataInterfaceCollisionQuery.cpp
+++ b/Engine/Plugins/FX/Niagara/Source/Niagara/Private/NiagaraDataInterfaceCollisionQuery.cpp
@@ -4,8 +4,10 @@
 
 #include "GlobalDistanceFieldParameters.h"
 #include "NiagaraComponent.h"
-#include "NiagaraEmitterInstanceBatcher.h"
+#include "NiagaraGpuComputeDispatchInterface.h"
+#include "NiagaraGpuComputeDispatch.h"
 #include "NiagaraRayTracingHelper.h"
+#include "NiagaraSimStageData.h"
 #include "NiagaraStats.h"
 #include "NiagaraTypes.h"
 #include "NiagaraWorldManager.h"
@@ -99,7 +101,7 @@ struct FNiagaraDataIntefaceProxyCollisionQuery : public FNiagaraDataInterfacePro
 		{
 			//Accumulate the total ray requests for this DI for all dispatches in the stage.
 			int32 RayRequests = MaxTracesPerParticle * Context.SimStageData->DestinationNumInstances;
-			FNiagaraRayTracingHelper& RTHelper = Context.Batcher->GetRayTracingHelper();
+			FNiagaraRayTracingHelper& RTHelper = Context.ComputeDispatchInterface->GetRayTracingHelper();
 			RTHelper.AddToDispatch(this, RayRequests, MaxRetraces);
 		}
 #endif
@@ -111,9 +113,9 @@ struct FNiagaraDataIntefaceProxyCollisionQuery : public FNiagaraDataInterfacePro
 		return true;
 	}
 
-	virtual void FinalizePreStage(FRHICommandList& RHICmdList, const NiagaraEmitterInstanceBatcher* Batcher) override
+	virtual void FinalizePreStage(FRHICommandList& RHICmdList, const FNiagaraGpuComputeDispatchInterface* ComputeDispatchInterface) override
 	{
-		FNiagaraRayTracingHelper& RTHelper = Batcher->GetRayTracingHelper();
+		FNiagaraRayTracingHelper& RTHelper = ComputeDispatchInterface->GetRayTracingHelper();
 		if (IsRayTracingEnabled() && GEnableGPUHWRTCollisions && MaxTracesPerParticle > 0)
 		{
 			RTHelper.BuildDispatch(RHICmdList, this);
@@ -963,8 +965,8 @@ public:
 		// Bind distance field parameters
 		if (GlobalDistanceFieldParameters.IsBound())
 		{
-			check(Context.Batcher);
-			const FGlobalDistanceFieldParameterData* GlobalDistanceFieldParameterData = Context.Batcher->GetGlobalDistanceFieldParameters();
+			check(Context.ComputeDispatchInterface);
+			const FGlobalDistanceFieldParameterData* GlobalDistanceFieldParameterData = static_cast<const FNiagaraGpuComputeDispatch*>(Context.ComputeDispatchInterface)->GetGlobalDistanceFieldParameters();//-BATCHERTODO:
 
 			if (GlobalDistanceFieldParameterData)
 			{
@@ -984,7 +986,7 @@ public:
 
 		if ((IsRayTracingEnabled() && GEnableGPUHWRTCollisions) || HasRayTracingParametersBound)
 		{
-			FNiagaraRayTracingHelper& RTHelper = Context.Batcher->GetRayTracingHelper();
+			FNiagaraRayTracingHelper& RTHelper = Context.ComputeDispatchInterface->GetRayTracingHelper();
 			const FNiagaraRayTraceDispatchInfo* DispatchInfo = nullptr;
 			if (IsRayTracingEnabled() && GEnableGPUHWRTCollisions && QueryDI->MaxTracesPerParticle > 0)
 			{
diff --git a/Engine/Plugins/FX/Niagara/Source/Niagara/Private/NiagaraDataInterfaceExport.cpp b/Engine/Plugins/FX/Niagara/Source/Niagara/Private/NiagaraDataInterfaceExport.cpp
--- a/Engine/Plugins/FX/Niagara/Source/Niagara/Private/NiagaraDataInterfaceExport.cpp
+++ b/Engine/Plugins/FX/Niagara/Source/Niagara/Private/NiagaraDataInterfaceExport.cpp
@@ -3,8 +3,10 @@
 #include "NiagaraDataInterfaceExport.h"
 #include "NiagaraTypes.h"
 #include "NiagaraCustomVersion.h"
-#include "NiagaraEmitterInstanceBatcher.h"
+#include "NiagaraComputeExecutionContext.h"
+#include "NiagaraGpuComputeDispatchInterface.h"
 #include "NiagaraGpuReadbackManager.h"
+#include "NiagaraGPUSystemTick.h"
 #include "NiagaraSystemInstance.h"
 #include "Internationalization/Internationalization.h"
 #include "ShaderParameterUtils.h"
@@ -164,7 +166,7 @@ struct FNDIExportProxy : public FNiagaraDataInterfaceProxy
 
 			RHICmdList.Transition(FRHITransitionInfo(InstanceData->WriteBuffer.UAV, ERHIAccess::UAVCompute, ERHIAccess::CopySrc));
 
-			FNiagaraGpuReadbackManager* ReadbackManager = Context.Batcher->GetGpuReadbackManager();
+			FNiagaraGpuReadbackManager* ReadbackManager = Context.ComputeDispatchInterface->GetGpuReadbackManager();
 			ReadbackManager->EnqueueReadback(
 				RHICmdList,
 				InstanceData->WriteBuffer.Buffer,
@@ -247,7 +249,7 @@ public:
 			if (InstanceData->WeakCallbackHandler.IsExplicitlyNull())
 			{
 				WriteBufferSize = 0;
-				RHICmdList.SetUAVParameter(Context.Shader.GetComputeShader(), WriteBufferParam.GetUAVIndex(), Context.Batcher->GetEmptyUAVFromPool(RHICmdList, PF_R32_UINT, ENiagaraEmptyUAVType::Buffer));
+				RHICmdList.SetUAVParameter(Context.Shader.GetComputeShader(), WriteBufferParam.GetUAVIndex(), Context.ComputeDispatchInterface->GetEmptyUAVFromPool(RHICmdList, PF_R32_UINT, ENiagaraEmptyUAVType::Buffer));
 			}
 			else
 			{
@@ -326,7 +328,7 @@ void UNiagaraDataInterfaceExport::DestroyPerInstanceData(void* PerInstanceData,
 
 	ENQUEUE_RENDER_COMMAND(FNDIExport_RemoveProxy)
 	(
-		[RT_Proxy=GetProxyAs<FNDIExportProxy>(), InstanceID=SystemInstance->GetId(), Batcher=SystemInstance->GetBatcher()](FRHICommandListImmediate& CmdList)
+		[RT_Proxy=GetProxyAs<FNDIExportProxy>(), InstanceID=SystemInstance->GetId()](FRHICommandListImmediate& CmdList)
 		{
 			RT_Proxy->SystemInstancesToProxyData_RT.Remove(InstanceID);
 		}
diff --git a/Engine/Plugins/FX/Niagara/Source/Niagara/Private/NiagaraDataInterfaceGrid2DCollection.cpp b/Engine/Plugins/FX/Niagara/Source/Niagara/Private/NiagaraDataInterfaceGrid2DCollection.cpp
--- a/Engine/Plugins/FX/Niagara/Source/Niagara/Private/NiagaraDataInterfaceGrid2DCollection.cpp
+++ b/Engine/Plugins/FX/Niagara/Source/Niagara/Private/NiagaraDataInterfaceGrid2DCollection.cpp
@@ -7,7 +7,7 @@
 #include "CanvasItem.h"
 #include "TextureResource.h"
 #include "Engine/Texture2DArray.h"
-#include "NiagaraEmitterInstanceBatcher.h"
+#include "NiagaraGpuComputeDispatchInterface.h"
 #include "NiagaraRenderer.h"
 #include "NiagaraSystemInstance.h"
 #include "NiagaraSettings.h"
@@ -314,7 +314,7 @@ public:
 			}
 			else
 			{
-				OutputGridUAV = Context.Batcher->GetEmptyUAVFromPool(RHICmdList, PF_R32_FLOAT, ENiagaraEmptyUAVType::Texture2DArray);
+				OutputGridUAV = Context.ComputeDispatchInterface->GetEmptyUAVFromPool(RHICmdList, PF_R32_FLOAT, ENiagaraEmptyUAVType::Texture2DArray);
 			}
 			RHICmdList.SetUAVParameter(ComputeShaderRHI, OutputGridParam.GetUAVIndex(), OutputGridUAV);
 		}
@@ -1875,7 +1875,7 @@ void UNiagaraDataInterfaceGrid2DCollection::DestroyPerInstanceData(void* PerInst
 
 	FNiagaraDataInterfaceProxyGrid2DCollectionProxy* RT_Proxy = GetProxyAs<FNiagaraDataInterfaceProxyGrid2DCollectionProxy>();
 	ENQUEUE_RENDER_COMMAND(FNiagaraDIDestroyInstanceData) (
-		[RT_Proxy, InstanceID=SystemInstance->GetId(), Batcher=SystemInstance->GetBatcher()](FRHICommandListImmediate& CmdList)
+		[RT_Proxy, InstanceID=SystemInstance->GetId()](FRHICommandListImmediate& CmdList)
 		{
 			//check(ThisProxy->SystemInstancesToProxyData.Contains(InstanceID));
 			RT_Proxy->SystemInstancesToProxyData_RT.Remove(InstanceID);
@@ -2593,7 +2593,7 @@ void FNiagaraDataInterfaceProxyGrid2DCollectionProxy::PostSimulate(FRHICommandLi
 #if NIAGARA_COMPUTEDEBUG_ENABLED && WITH_EDITORONLY_DATA
 	if (ProxyData->bPreviewGrid && ProxyData->CurrentData)
 	{
-		if (FNiagaraGpuComputeDebug* GpuComputeDebug = Context.Batcher->GetGpuComputeDebug())
+		if (FNiagaraGpuComputeDebug* GpuComputeDebug = Context.ComputeDispatchInterface->GetGpuComputeDebug())
 		{
 			if (ProxyData->PreviewAttribute[0] != INDEX_NONE)
 			{
diff --git a/Engine/Plugins/FX/Niagara/Source/Niagara/Private/NiagaraDataInterfaceGrid2DCollectionReader.cpp b/Engine/Plugins/FX/Niagara/Source/Niagara/Private/NiagaraDataInterfaceGrid2DCollectionReader.cpp
--- a/Engine/Plugins/FX/Niagara/Source/Niagara/Private/NiagaraDataInterfaceGrid2DCollectionReader.cpp
+++ b/Engine/Plugins/FX/Niagara/Source/Niagara/Private/NiagaraDataInterfaceGrid2DCollectionReader.cpp
@@ -1,11 +1,11 @@
 // Copyright Epic Games, Inc. All Rights Reserved.
 #include "NiagaraDataInterfaceGrid2DCollectionReader.h"
+#include "NiagaraComputeExecutionContext.h"
 #include "NiagaraShader.h"
 #include "ShaderParameterUtils.h"
 #include "ClearQuad.h"
 #include "TextureResource.h"
 #include "Engine/Texture2DArray.h"
-#include "NiagaraEmitterInstanceBatcher.h"
 #include "NiagaraSystemInstance.h"
 #include "NiagaraRenderer.h"
 #include "Engine/TextureRenderTarget2D.h"
@@ -356,7 +356,7 @@ void UNiagaraDataInterfaceGrid2DCollectionReader::DestroyPerInstanceData(void* P
 
 	FNiagaraDataInterfaceProxyGrid2DCollectionReaderProxy* RT_Proxy = GetProxyAs<FNiagaraDataInterfaceProxyGrid2DCollectionReaderProxy>();
 	ENQUEUE_RENDER_COMMAND(FNiagaraDIDestroyInstanceData) (
-		[RT_Proxy, InstanceID=SystemInstance->GetId(), Batcher=SystemInstance->GetBatcher()](FRHICommandListImmediate& CmdList)
+		[RT_Proxy, InstanceID=SystemInstance->GetId()](FRHICommandListImmediate& CmdList)
 		{
 			//check(ThisProxy->SystemInstancesToProxyData.Contains(InstanceID));
 			RT_Proxy->SystemInstancesToProxyData_RT.Remove(InstanceID);
diff --git a/Engine/Plugins/FX/Niagara/Source/Niagara/Private/NiagaraDataInterfaceGrid3DCollection.cpp b/Engine/Plugins/FX/Niagara/Source/Niagara/Private/NiagaraDataInterfaceGrid3DCollection.cpp
--- a/Engine/Plugins/FX/Niagara/Source/Niagara/Private/NiagaraDataInterfaceGrid3DCollection.cpp
+++ b/Engine/Plugins/FX/Niagara/Source/Niagara/Private/NiagaraDataInterfaceGrid3DCollection.cpp
@@ -4,7 +4,8 @@
 #include "ShaderParameterUtils.h"
 #include "ClearQuad.h"
 #include "TextureResource.h"
-#include "NiagaraEmitterInstanceBatcher.h"
+#include "NiagaraGpuComputeDispatchInterface.h"
+#include "NiagaraSimStageData.h"
 #include "NiagaraSystemInstance.h"
 #include "NiagaraRenderer.h"
 #include "Engine/VolumeTexture.h"
@@ -280,7 +281,7 @@ void FNiagaraDataInterfaceParametersCS_Grid3DCollection::Set(FRHICommandList& RH
 			}
 			else
 			{
-				OutputGridUAV = Context.Batcher->GetEmptyUAVFromPool(RHICmdList, PF_R32_FLOAT, ENiagaraEmptyUAVType::Texture3D);
+				OutputGridUAV = Context.ComputeDispatchInterface->GetEmptyUAVFromPool(RHICmdList, PF_R32_FLOAT, ENiagaraEmptyUAVType::Texture3D);
 			}
 			RHICmdList.SetUAVParameter(ComputeShaderRHI, OutputGridParam.GetUAVIndex(), OutputGridUAV);
 		}
@@ -2043,7 +2044,7 @@ void UNiagaraDataInterfaceGrid3DCollection::DestroyPerInstanceData(void* PerInst
 
 	FNiagaraDataInterfaceProxyGrid3DCollectionProxy* RT_Proxy = GetProxyAs<FNiagaraDataInterfaceProxyGrid3DCollectionProxy>();
 	ENQUEUE_RENDER_COMMAND(FNiagaraDIDestroyInstanceData) (
-		[RT_Proxy, InstanceID=SystemInstance->GetId(), Batcher=SystemInstance->GetBatcher()](FRHICommandListImmediate& CmdList)
+		[RT_Proxy, InstanceID=SystemInstance->GetId()](FRHICommandListImmediate& CmdList)
 		{
 			//check(ThisProxy->SystemInstancesToProxyData.Contains(InstanceID));
 			RT_Proxy->SystemInstancesToProxyData_RT.Remove(InstanceID);
@@ -2757,7 +2758,7 @@ void FNiagaraDataInterfaceProxyGrid3DCollectionProxy::PostSimulate(FRHICommandLi
 #if WITH_EDITOR
 	if (ProxyData->bPreviewGrid && ProxyData->CurrentData)
 	{
-		if (FNiagaraGpuComputeDebug* GpuComputeDebug = Context.Batcher->GetGpuComputeDebug())
+		if (FNiagaraGpuComputeDebug* GpuComputeDebug = Context.ComputeDispatchInterface->GetGpuComputeDebug())
 		{
 				if (ProxyData->PreviewAttribute[0] != INDEX_NONE)
 			{
diff --git a/Engine/Plugins/FX/Niagara/Source/Niagara/Private/NiagaraDataInterfaceGrid3DCollectionReader.cpp b/Engine/Plugins/FX/Niagara/Source/Niagara/Private/NiagaraDataInterfaceGrid3DCollectionReader.cpp
--- a/Engine/Plugins/FX/Niagara/Source/Niagara/Private/NiagaraDataInterfaceGrid3DCollectionReader.cpp
+++ b/Engine/Plugins/FX/Niagara/Source/Niagara/Private/NiagaraDataInterfaceGrid3DCollectionReader.cpp
@@ -4,13 +4,13 @@
 #include "ShaderParameterUtils.h"
 #include "ClearQuad.h"
 #include "TextureResource.h"
-#include "NiagaraEmitterInstanceBatcher.h"
 #include "NiagaraSystemInstance.h"
 #include "NiagaraRenderer.h"
 #include "Engine/VolumeTexture.h"
 #include "Engine/TextureRenderTargetVolume.h"
 #include "NiagaraSettings.h"
 #include "NiagaraConstants.h"
+#include "NiagaraComputeExecutionContext.h"
 #if WITH_EDITOR
 #include "NiagaraGpuComputeDebug.h"
 #endif
diff --git a/Engine/Plugins/FX/Niagara/Source/Niagara/Private/NiagaraDataInterfaceIntRenderTarget2D.cpp b/Engine/Plugins/FX/Niagara/Source/Niagara/Private/NiagaraDataInterfaceIntRenderTarget2D.cpp
--- a/Engine/Plugins/FX/Niagara/Source/Niagara/Private/NiagaraDataInterfaceIntRenderTarget2D.cpp
+++ b/Engine/Plugins/FX/Niagara/Source/Niagara/Private/NiagaraDataInterfaceIntRenderTarget2D.cpp
@@ -1,13 +1,13 @@
 // Copyright Epic Games, Inc. All Rights Reserved.
 #include "NiagaraDataInterfaceIntRenderTarget2D.h"
 #include "NiagaraShader.h"
-#include "NiagaraEmitterInstanceBatcher.h"
 #include "NiagaraSettings.h"
 #include "NiagaraSystemInstance.h"
 #include "NiagaraRenderer.h"
-#if WITH_EDITOR
+#if NIAGARA_COMPUTEDEBUG_ENABLED
 #include "NiagaraGpuComputeDebug.h"
 #endif
+#include "NiagaraGpuComputeDispatchInterface.h"
 #include "NiagaraStats.h"
 
 #include "ShaderParameterUtils.h"
@@ -110,7 +110,7 @@ struct FNDIIntRenderTarget2DProxy : public FNiagaraDataInterfaceProxyRW
 #if NIAGARA_COMPUTEDEBUG_ENABLED && WITH_EDITORONLY_DATA
 			if (InstanceData->bPreviewRenderTarget)
 			{
-				if (FNiagaraGpuComputeDebug* GpuComputeDebug = Context.Batcher->GetGpuComputeDebug())
+				if (FNiagaraGpuComputeDebug* GpuComputeDebug = Context.ComputeDispatchInterface->GetGpuComputeDebug())
 				{
 					if (FRHITexture* RHITexture = InstanceData->TextureRHI)
 					{
@@ -168,7 +168,7 @@ public:
 			}
 			else
 			{
-				OutputUAV = Context.Batcher->GetEmptyUAVFromPool(RHICmdList, EPixelFormat::PF_A16B16G16R16, ENiagaraEmptyUAVType::Texture2D);
+				OutputUAV = Context.ComputeDispatchInterface->GetEmptyUAVFromPool(RHICmdList, EPixelFormat::PF_A16B16G16R16, ENiagaraEmptyUAVType::Texture2D);
 			}
 
 			RHICmdList.SetUAVParameter(ComputeShaderRHI, TextureUAVParam.GetUAVIndex(), OutputUAV);
@@ -578,7 +578,7 @@ void UNiagaraDataInterfaceIntRenderTarget2D::DestroyPerInstanceData(void* PerIns
 
 	FNDIIntRenderTarget2DProxy* RT_Proxy = GetProxyAs<FNDIIntRenderTarget2DProxy>();
 	ENQUEUE_RENDER_COMMAND(FNiagaraDIDestroyInstanceData) (
-		[RT_Proxy, InstanceID=SystemInstance->GetId(), Batcher=SystemInstance->GetBatcher()](FRHICommandListImmediate& CmdList)
+		[RT_Proxy, InstanceID=SystemInstance->GetId()](FRHICommandListImmediate& CmdList)
 		{
 			RT_Proxy->SystemInstancesToProxyData_RT.Remove(InstanceID);
 		}
diff --git a/Engine/Plugins/FX/Niagara/Source/Niagara/Private/NiagaraDataInterfaceNeighborGrid3D.cpp b/Engine/Plugins/FX/Niagara/Source/Niagara/Private/NiagaraDataInterfaceNeighborGrid3D.cpp
--- a/Engine/Plugins/FX/Niagara/Source/Niagara/Private/NiagaraDataInterfaceNeighborGrid3D.cpp
+++ b/Engine/Plugins/FX/Niagara/Source/Niagara/Private/NiagaraDataInterfaceNeighborGrid3D.cpp
@@ -2,7 +2,7 @@
 #include "NiagaraDataInterfaceNeighborGrid3D.h"
 #include "NiagaraShader.h"
 #include "ShaderParameterUtils.h"
-#include "NiagaraEmitterInstanceBatcher.h"
+#include "NiagaraGpuComputeDispatchInterface.h"
 #include "NiagaraSystemInstance.h"
 #include "NiagaraRenderer.h"
 #include "NiagaraShaderParticleID.h"
@@ -82,12 +82,12 @@ public:
 			SetSRVParameter(RHICmdList, ComputeShaderRHI, ParticleNeighborCountGridParam, FNiagaraRenderer::GetDummyIntBuffer());
 			if (OutputParticleNeighborsGridParam.IsUAVBound())
 			{
-				RHICmdList.SetUAVParameter(ComputeShaderRHI, OutputParticleNeighborsGridParam.GetUAVIndex(), Context.Batcher->GetEmptyUAVFromPool(RHICmdList, PF_R32_SINT, ENiagaraEmptyUAVType::Buffer));
+				RHICmdList.SetUAVParameter(ComputeShaderRHI, OutputParticleNeighborsGridParam.GetUAVIndex(), Context.ComputeDispatchInterface->GetEmptyUAVFromPool(RHICmdList, PF_R32_SINT, ENiagaraEmptyUAVType::Buffer));
 			}
 
 			if (OutputParticleNeighborCountGridParam.IsUAVBound())
 			{
-				RHICmdList.SetUAVParameter(ComputeShaderRHI, OutputParticleNeighborCountGridParam.GetUAVIndex(), Context.Batcher->GetEmptyUAVFromPool(RHICmdList, PF_R32_SINT, ENiagaraEmptyUAVType::Buffer));
+				RHICmdList.SetUAVParameter(ComputeShaderRHI, OutputParticleNeighborCountGridParam.GetUAVIndex(), Context.ComputeDispatchInterface->GetEmptyUAVFromPool(RHICmdList, PF_R32_SINT, ENiagaraEmptyUAVType::Buffer));
 			}
 
 			return;
@@ -120,12 +120,12 @@ public:
 
 			if (OutputParticleNeighborsGridParam.IsUAVBound())
 			{
-				RHICmdList.SetUAVParameter(ComputeShaderRHI, OutputParticleNeighborsGridParam.GetUAVIndex(), Context.Batcher->GetEmptyUAVFromPool(RHICmdList, PF_R32_SINT, ENiagaraEmptyUAVType::Buffer));
+				RHICmdList.SetUAVParameter(ComputeShaderRHI, OutputParticleNeighborsGridParam.GetUAVIndex(), Context.ComputeDispatchInterface->GetEmptyUAVFromPool(RHICmdList, PF_R32_SINT, ENiagaraEmptyUAVType::Buffer));
 			}
 
 			if (OutputParticleNeighborCountGridParam.IsUAVBound())
 			{
-				RHICmdList.SetUAVParameter(ComputeShaderRHI, OutputParticleNeighborCountGridParam.GetUAVIndex(), Context.Batcher->GetEmptyUAVFromPool(RHICmdList, PF_R32_SINT, ENiagaraEmptyUAVType::Buffer));
+				RHICmdList.SetUAVParameter(ComputeShaderRHI, OutputParticleNeighborCountGridParam.GetUAVIndex(), Context.ComputeDispatchInterface->GetEmptyUAVFromPool(RHICmdList, PF_R32_SINT, ENiagaraEmptyUAVType::Buffer));
 			}
 		}
 		else
@@ -719,7 +719,7 @@ void UNiagaraDataInterfaceNeighborGrid3D::DestroyPerInstanceData(void* PerInstan
 		return;
 
 	ENQUEUE_RENDER_COMMAND(FNiagaraDIDestroyInstanceData) (
-		[ThisProxy, InstanceID = SystemInstance->GetId(), Batcher = SystemInstance->GetBatcher()](FRHICommandListImmediate& CmdList)
+		[ThisProxy, InstanceID = SystemInstance->GetId()](FRHICommandListImmediate& CmdList)
 	{
 		//check(ThisProxy->SystemInstancesToProxyData.Contains(InstanceID));
 		ThisProxy->SystemInstancesToProxyData.Remove(InstanceID);
@@ -734,7 +734,7 @@ void FNiagaraDataInterfaceProxyNeighborGrid3D::PreStage(FRHICommandList& RHICmdL
 		NeighborGrid3DRWInstanceData* ProxyData = SystemInstancesToProxyData.Find(Context.SystemInstanceID);
 
 		SCOPED_DRAW_EVENT(RHICmdList, NiagaraNeighborGrid3DClearNeighborInfo);
-		ERHIFeatureLevel::Type FeatureLevel = Context.Batcher->GetFeatureLevel();
+		ERHIFeatureLevel::Type FeatureLevel = Context.ComputeDispatchInterface->GetFeatureLevel();
 
 		FRHITransitionInfo TransitionInfos[] =
 		{
diff --git a/Engine/Plugins/FX/Niagara/Source/Niagara/Private/NiagaraDataInterfaceParticleRead.cpp b/Engine/Plugins/FX/Niagara/Source/Niagara/Private/NiagaraDataInterfaceParticleRead.cpp
--- a/Engine/Plugins/FX/Niagara/Source/Niagara/Private/NiagaraDataInterfaceParticleRead.cpp
+++ b/Engine/Plugins/FX/Niagara/Source/Niagara/Private/NiagaraDataInterfaceParticleRead.cpp
@@ -1,12 +1,13 @@
 // Copyright Epic Games, Inc. All Rights Reserved.
 
 #include "NiagaraDataInterfaceParticleRead.h"
+#include "NiagaraComputeExecutionContext.h"
+#include "NiagaraGPUSystemTick.h"
 #include "NiagaraConstants.h"
 #include "NiagaraSystemInstance.h"
 #include "ShaderParameterUtils.h"
 #include "NiagaraRenderer.h"
 #include "NiagaraDataSetAccessor.h"
-#include "NiagaraEmitterInstanceBatcher.h"
 
 #define LOCTEXT_NAMESPACE "NiagaraDataInterfaceParticleRead"
 
@@ -568,7 +569,7 @@ struct FNiagaraDataInterfaceParametersCS_ParticleRead : public FNiagaraDataInter
 		const uint32 ParticleStrideInt = SourceData->GetInt32Stride() / sizeof(int32);
 		const uint32 ParticleStrideHalf = SourceData->GetHalfStride() / sizeof(FFloat16);
 
-		// There's no need to transition the input buffers, because the grouping logic inside NiagaraEmitterInstanceBatcher ensures that our source emitter has ran before us,
+		// There's no need to transition the input buffers, because the grouping logic inside FNiagaraGpuComputeDispatch ensures that our source emitter has ran before us,
 		// and its buffers have been transitioned to readable.
 
 		SetShaderValue(RHICmdList, ComputeShader, InstanceCountOffsetParam, SourceData->GetGPUInstanceCountBufferOffset());
diff --git a/Engine/Plugins/FX/Niagara/Source/Niagara/Private/NiagaraDataInterfaceRasterizationGrid3D.cpp b/Engine/Plugins/FX/Niagara/Source/Niagara/Private/NiagaraDataInterfaceRasterizationGrid3D.cpp
--- a/Engine/Plugins/FX/Niagara/Source/Niagara/Private/NiagaraDataInterfaceRasterizationGrid3D.cpp
+++ b/Engine/Plugins/FX/Niagara/Source/Niagara/Private/NiagaraDataInterfaceRasterizationGrid3D.cpp
@@ -2,7 +2,7 @@
 #include "NiagaraDataInterfaceRasterizationGrid3D.h"
 #include "NiagaraShader.h"
 #include "ShaderParameterUtils.h"
-#include "NiagaraEmitterInstanceBatcher.h"
+#include "NiagaraGpuComputeDispatchInterface.h"
 #include "NiagaraSystemInstance.h"
 #include "NiagaraRenderer.h"
 #include "NiagaraShaderParticleID.h"
@@ -79,7 +79,7 @@ public:
 
 			if (OutputIntGridParam.IsUAVBound())
 			{
-				RHICmdList.SetUAVParameter(ComputeShaderRHI, OutputIntGridParam.GetUAVIndex(), Context.Batcher->GetEmptyUAVFromPool(RHICmdList, PF_R32_SINT, ENiagaraEmptyUAVType::Buffer));
+				RHICmdList.SetUAVParameter(ComputeShaderRHI, OutputIntGridParam.GetUAVIndex(), Context.ComputeDispatchInterface->GetEmptyUAVFromPool(RHICmdList, PF_R32_SINT, ENiagaraEmptyUAVType::Buffer));
 			}
 
 			return;
@@ -99,7 +99,7 @@ public:
 
 			if (OutputIntGridParam.IsUAVBound())
 			{
-				RHICmdList.SetUAVParameter(ComputeShaderRHI, OutputIntGridParam.GetUAVIndex(), Context.Batcher->GetEmptyUAVFromPool(RHICmdList, PF_R32_SINT, ENiagaraEmptyUAVType::Buffer));
+				RHICmdList.SetUAVParameter(ComputeShaderRHI, OutputIntGridParam.GetUAVIndex(), Context.ComputeDispatchInterface->GetEmptyUAVFromPool(RHICmdList, PF_R32_SINT, ENiagaraEmptyUAVType::Buffer));
 			}
 		}
 		else
@@ -730,7 +730,7 @@ void UNiagaraDataInterfaceRasterizationGrid3D::DestroyPerInstanceData(void* PerI
 		return;
 
 	ENQUEUE_RENDER_COMMAND(FNiagaraDIDestroyInstanceData) (
-		[ThisProxy, InstanceID = SystemInstance->GetId(), Batcher = SystemInstance->GetBatcher()](FRHICommandListImmediate& CmdList)
+		[ThisProxy, InstanceID = SystemInstance->GetId()](FRHICommandListImmediate& CmdList)
 	{
 		//check(ThisProxy->SystemInstancesToProxyData.Contains(InstanceID));
 		ThisProxy->SystemInstancesToProxyData.Remove(InstanceID);
@@ -745,7 +745,7 @@ void FNiagaraDataInterfaceProxyRasterizationGrid3D::PreStage(FRHICommandList& RH
 		RasterizationGrid3DRWInstanceData* ProxyData = SystemInstancesToProxyData.Find(Context.SystemInstanceID);
 
 		SCOPED_DRAW_EVENT(RHICmdList, NiagaraRasterizationGrid3DClearNeighborInfo);
-		ERHIFeatureLevel::Type FeatureLevel = Context.Batcher->GetFeatureLevel();
+		ERHIFeatureLevel::Type FeatureLevel = Context.ComputeDispatchInterface->GetFeatureLevel();
 
 		FRHITransitionInfo TransitionInfos[] =
 		{			
diff --git a/Engine/Plugins/FX/Niagara/Source/Niagara/Private/NiagaraDataInterfaceRenderTarget2D.cpp b/Engine/Plugins/FX/Niagara/Source/Niagara/Private/NiagaraDataInterfaceRenderTarget2D.cpp
--- a/Engine/Plugins/FX/Niagara/Source/Niagara/Private/NiagaraDataInterfaceRenderTarget2D.cpp
+++ b/Engine/Plugins/FX/Niagara/Source/Niagara/Private/NiagaraDataInterfaceRenderTarget2D.cpp
@@ -9,11 +9,11 @@
 #include "Engine/Texture2D.h"
 #include "Engine/TextureRenderTarget2D.h"
 
-#include "NiagaraEmitterInstanceBatcher.h"
 #include "NiagaraSettings.h"
 #include "NiagaraSystemInstance.h"
 #include "NiagaraRenderer.h"
 #include "NiagaraGpuComputeDebug.h"
+#include "NiagaraGpuComputeDispatchInterface.h"
 #include "NiagaraStats.h"
 
 #define LOCTEXT_NAMESPACE "NiagaraDataInterfaceRenderTarget2D"
@@ -110,7 +110,7 @@ public:
 			}
 			else
 			{
-				OutputUAV = Context.Batcher->GetEmptyUAVFromPool(RHICmdList, EPixelFormat::PF_A16B16G16R16, ENiagaraEmptyUAVType::Buffer);
+				OutputUAV = Context.ComputeDispatchInterface->GetEmptyUAVFromPool(RHICmdList, EPixelFormat::PF_A16B16G16R16, ENiagaraEmptyUAVType::Buffer);
 			}
 
 			RHICmdList.SetUAVParameter(ComputeShaderRHI, OutputParam.GetUAVIndex(), OutputUAV);
@@ -561,7 +561,7 @@ void UNiagaraDataInterfaceRenderTarget2D::DestroyPerInstanceData(void* PerInstan
 
 	FNiagaraDataInterfaceProxyRenderTarget2DProxy* RT_Proxy = GetProxyAs<FNiagaraDataInterfaceProxyRenderTarget2DProxy>();
 	ENQUEUE_RENDER_COMMAND(FNiagaraDIDestroyInstanceData) (
-		[RT_Proxy, InstanceID=SystemInstance->GetId(), Batcher=SystemInstance->GetBatcher()](FRHICommandListImmediate& CmdList)
+		[RT_Proxy, InstanceID=SystemInstance->GetId()](FRHICommandListImmediate& CmdList)
 		{
 #if STATS
 			if (FRenderTarget2DRWInstanceData_RenderThread* TargetData = RT_Proxy->SystemInstancesToProxyData_RT.Find(InstanceID))
@@ -825,7 +825,7 @@ void FNiagaraDataInterfaceProxyRenderTarget2DProxy::PostSimulate(FRHICommandList
 #if NIAGARA_COMPUTEDEBUG_ENABLED && WITH_EDITORONLY_DATA
 		if (ProxyData->bPreviewTexture)
 		{
-			if (FNiagaraGpuComputeDebug* GpuComputeDebug = Context.Batcher->GetGpuComputeDebug())
+			if (FNiagaraGpuComputeDebug* GpuComputeDebug = Context.ComputeDispatchInterface->GetGpuComputeDebug())
 			{
 				if (FRHITexture* RHITexture = ProxyData->TextureRHI)
 				{
diff --git a/Engine/Plugins/FX/Niagara/Source/Niagara/Private/NiagaraDataInterfaceRenderTarget2DArray.cpp b/Engine/Plugins/FX/Niagara/Source/Niagara/Private/NiagaraDataInterfaceRenderTarget2DArray.cpp
--- a/Engine/Plugins/FX/Niagara/Source/Niagara/Private/NiagaraDataInterfaceRenderTarget2DArray.cpp
+++ b/Engine/Plugins/FX/Niagara/Source/Niagara/Private/NiagaraDataInterfaceRenderTarget2DArray.cpp
@@ -8,12 +8,12 @@
 #include "Engine/Texture2DArray.h"
 #include "Engine/TextureRenderTarget2DArray.h"
 
-#include "NiagaraEmitterInstanceBatcher.h"
 #include "NiagaraSystemInstance.h"
 #include "NiagaraStats.h"
 #include "NiagaraRenderer.h"
 #include "NiagaraSettings.h"
 #include "NiagaraGpuComputeDebug.h"
+#include "NiagaraGpuComputeDispatchInterface.h"
 
 #define LOCTEXT_NAMESPACE "NiagaraDataInterfaceRenderTarget2DArray"
 
@@ -82,7 +82,7 @@ public:
 			}
 			else
 			{
-				OutputUAV = Context.Batcher->GetEmptyUAVFromPool(RHICmdList, EPixelFormat::PF_A16B16G16R16, ENiagaraEmptyUAVType::Texture2DArray);
+				OutputUAV = Context.ComputeDispatchInterface->GetEmptyUAVFromPool(RHICmdList, EPixelFormat::PF_A16B16G16R16, ENiagaraEmptyUAVType::Texture2DArray);
 			}
 
 			RHICmdList.SetUAVParameter(ComputeShaderRHI, OutputParam.GetUAVIndex(), OutputUAV);
@@ -531,7 +531,7 @@ void UNiagaraDataInterfaceRenderTarget2DArray::DestroyPerInstanceData(void* PerI
 
 	FNiagaraDataInterfaceProxyRenderTarget2DArrayProxy* RT_Proxy = GetProxyAs<FNiagaraDataInterfaceProxyRenderTarget2DArrayProxy>();
 	ENQUEUE_RENDER_COMMAND(FNiagaraDIDestroyInstanceData) (
-		[RT_Proxy, InstanceID=SystemInstance->GetId(), Batcher=SystemInstance->GetBatcher()](FRHICommandListImmediate& CmdList)
+		[RT_Proxy, InstanceID=SystemInstance->GetId()](FRHICommandListImmediate& CmdList)
 		{
 #if STATS
 			if (FRenderTarget2DArrayRWInstanceData_RenderThread* TargetData = RT_Proxy->SystemInstancesToProxyData_RT.Find(InstanceID))
@@ -743,7 +743,7 @@ void FNiagaraDataInterfaceProxyRenderTarget2DArrayProxy::PostSimulate(FRHIComman
 
 	if (ProxyData && ProxyData->bPreviewTexture)
 	{
-		if (FNiagaraGpuComputeDebug* GpuComputeDebug = Context.Batcher->GetGpuComputeDebug())
+		if (FNiagaraGpuComputeDebug* GpuComputeDebug = Context.ComputeDispatchInterface->GetGpuComputeDebug())
 		{
 			if (FRHITexture* RHITexture = ProxyData->TextureRHI)
 			{
diff --git a/Engine/Plugins/FX/Niagara/Source/Niagara/Private/NiagaraDataInterfaceRenderTargetCube.cpp b/Engine/Plugins/FX/Niagara/Source/Niagara/Private/NiagaraDataInterfaceRenderTargetCube.cpp
--- a/Engine/Plugins/FX/Niagara/Source/Niagara/Private/NiagaraDataInterfaceRenderTargetCube.cpp
+++ b/Engine/Plugins/FX/Niagara/Source/Niagara/Private/NiagaraDataInterfaceRenderTargetCube.cpp
@@ -6,12 +6,12 @@
 #include "TextureResource.h"
 #include "Engine/TextureRenderTargetCube.h"
 
-#include "NiagaraEmitterInstanceBatcher.h"
 #include "NiagaraSystemInstance.h"
 #include "NiagaraStats.h"
 #include "NiagaraRenderer.h"
 #include "NiagaraSettings.h"
 #include "NiagaraGpuComputeDebug.h"
+#include "NiagaraGpuComputeDispatchInterface.h"
 
 #define LOCTEXT_NAMESPACE "NiagaraDataInterfaceRenderTargetCube"
 
@@ -80,7 +80,7 @@ public:
 			}
 			else
 			{
-				OutputUAV = Context.Batcher->GetEmptyUAVFromPool(RHICmdList, EPixelFormat::PF_A16B16G16R16, ENiagaraEmptyUAVType::Texture2DArray);
+				OutputUAV = Context.ComputeDispatchInterface->GetEmptyUAVFromPool(RHICmdList, EPixelFormat::PF_A16B16G16R16, ENiagaraEmptyUAVType::Texture2DArray);
 			}
 			RHICmdList.SetUAVParameter(ComputeShaderRHI, OutputParam.GetUAVIndex(), OutputUAV);
 		}
@@ -518,7 +518,7 @@ void UNiagaraDataInterfaceRenderTargetCube::DestroyPerInstanceData(void* PerInst
 
 	FNiagaraDataInterfaceProxyRenderTargetCubeProxy* RT_Proxy = GetProxyAs<FNiagaraDataInterfaceProxyRenderTargetCubeProxy>();
 	ENQUEUE_RENDER_COMMAND(FNiagaraDIDestroyInstanceData) (
-		[RT_Proxy, InstanceID=SystemInstance->GetId(), Batcher=SystemInstance->GetBatcher()](FRHICommandListImmediate& CmdList)
+		[RT_Proxy, InstanceID=SystemInstance->GetId()](FRHICommandListImmediate& CmdList)
 		{
 #if STATS
 			if (FRenderTargetCubeRWInstanceData_RenderThread* TargetData = RT_Proxy->SystemInstancesToProxyData_RT.Find(InstanceID))
@@ -719,7 +719,7 @@ void FNiagaraDataInterfaceProxyRenderTargetCubeProxy::PostSimulate(FRHICommandLi
 
 	if (ProxyData && ProxyData->bPreviewTexture && ProxyData->TextureRHI.IsValid())
 	{
-		if (FNiagaraGpuComputeDebug* GpuComputeDebug = Context.Batcher->GetGpuComputeDebug())
+		if (FNiagaraGpuComputeDebug* GpuComputeDebug = Context.ComputeDispatchInterface->GetGpuComputeDebug())
 		{
 			if ( FRHITexture* RHITexture = ProxyData->TextureRHI )
 			{
diff --git a/Engine/Plugins/FX/Niagara/Source/Niagara/Private/NiagaraDataInterfaceRenderTargetVolume.cpp b/Engine/Plugins/FX/Niagara/Source/Niagara/Private/NiagaraDataInterfaceRenderTargetVolume.cpp
--- a/Engine/Plugins/FX/Niagara/Source/Niagara/Private/NiagaraDataInterfaceRenderTargetVolume.cpp
+++ b/Engine/Plugins/FX/Niagara/Source/Niagara/Private/NiagaraDataInterfaceRenderTargetVolume.cpp
@@ -7,7 +7,7 @@
 #include "TextureRenderTargetVolumeResource.h"
 #include "Engine/TextureRenderTargetVolume.h"
 
-#include "NiagaraEmitterInstanceBatcher.h"
+#include "NiagaraGpuComputeDispatchInterface.h"
 #include "NiagaraSystemInstance.h"
 #include "NiagaraStats.h"
 #include "NiagaraRenderer.h"
@@ -81,7 +81,7 @@ public:
 			}
 			else
 			{
-				OutputUAV = Context.Batcher->GetEmptyUAVFromPool(RHICmdList, EPixelFormat::PF_A16B16G16R16, ENiagaraEmptyUAVType::Texture3D);
+				OutputUAV = Context.ComputeDispatchInterface->GetEmptyUAVFromPool(RHICmdList, EPixelFormat::PF_A16B16G16R16, ENiagaraEmptyUAVType::Texture3D);
 			}
 			RHICmdList.SetUAVParameter(ComputeShaderRHI, OutputParam.GetUAVIndex(), OutputUAV);
 		}
@@ -527,7 +527,7 @@ void UNiagaraDataInterfaceRenderTargetVolume::DestroyPerInstanceData(void* PerIn
 
 	FNiagaraDataInterfaceProxyRenderTargetVolumeProxy* RT_Proxy = GetProxyAs<FNiagaraDataInterfaceProxyRenderTargetVolumeProxy>();
 	ENQUEUE_RENDER_COMMAND(FNiagaraDIDestroyInstanceData) (
-		[RT_Proxy, InstanceID=SystemInstance->GetId(), Batcher=SystemInstance->GetBatcher()](FRHICommandListImmediate& CmdList)
+		[RT_Proxy, InstanceID=SystemInstance->GetId()](FRHICommandListImmediate& CmdList)
 		{
 #if STATS
 			if (FRenderTargetVolumeRWInstanceData_RenderThread* TargetData = RT_Proxy->SystemInstancesToProxyData_RT.Find(InstanceID))
@@ -741,7 +741,7 @@ void FNiagaraDataInterfaceProxyRenderTargetVolumeProxy::PostSimulate(FRHICommand
 
 	if (ProxyData && ProxyData->bPreviewTexture && ProxyData->TextureRHI.IsValid())
 	{
-		if (FNiagaraGpuComputeDebug* GpuComputeDebug = Context.Batcher->GetGpuComputeDebug())
+		if (FNiagaraGpuComputeDebug* GpuComputeDebug = Context.ComputeDispatchInterface->GetGpuComputeDebug())
 		{
 			if ( FRHITexture* RHITexture = ProxyData->TextureRHI )
 			{
diff --git a/Engine/Plugins/FX/Niagara/Source/Niagara/Private/NiagaraDataInterfaceRigidMeshCollisionQuery.cpp b/Engine/Plugins/FX/Niagara/Source/Niagara/Private/NiagaraDataInterfaceRigidMeshCollisionQuery.cpp
--- a/Engine/Plugins/FX/Niagara/Source/Niagara/Private/NiagaraDataInterfaceRigidMeshCollisionQuery.cpp
+++ b/Engine/Plugins/FX/Niagara/Source/Niagara/Private/NiagaraDataInterfaceRigidMeshCollisionQuery.cpp
@@ -9,8 +9,8 @@
 #include "NiagaraShader.h"
 #include "NiagaraComponent.h"
 #include "NiagaraRenderer.h"
+#include "NiagaraSimStageData.h"
 #include "NiagaraSystemInstance.h"
-#include "NiagaraEmitterInstanceBatcher.h"
 #include "ShaderParameterUtils.h"
 #include "EngineUtils.h"
 #include "Engine/StaticMeshActor.h"
@@ -532,7 +532,7 @@ void FNDIRigidMeshCollisionProxy::InitializePerInstanceData(const FNiagaraSystem
 	TargetData = &SystemInstancesToProxyData.Add(SystemInstance);
 }
 
-void FNDIRigidMeshCollisionProxy::DestroyPerInstanceData(NiagaraEmitterInstanceBatcher* Batcher, const FNiagaraSystemInstanceID& SystemInstance)
+void FNDIRigidMeshCollisionProxy::DestroyPerInstanceData(const FNiagaraSystemInstanceID& SystemInstance)
 {
 	check(IsInRenderingThread());
 	SystemInstancesToProxyData.Remove(SystemInstance);
@@ -604,7 +604,7 @@ void UNiagaraDataInterfaceRigidMeshCollisionQuery::DestroyPerInstanceData(void*
 
 	FNDIRigidMeshCollisionProxy* ThisProxy = GetProxyAs<FNDIRigidMeshCollisionProxy>();
 	ENQUEUE_RENDER_COMMAND(FNiagaraDIDestroyInstanceData) (
-		[ThisProxy, InstanceID = SystemInstance->GetId(), Batcher = SystemInstance->GetBatcher()](FRHICommandListImmediate& CmdList)
+		[ThisProxy, InstanceID = SystemInstance->GetId()](FRHICommandListImmediate& CmdList)
 	{
 		ThisProxy->SystemInstancesToProxyData.Remove(InstanceID);
 	}
diff --git a/Engine/Plugins/FX/Niagara/Source/Niagara/Private/NiagaraDataInterfaceSpline.cpp b/Engine/Plugins/FX/Niagara/Source/Niagara/Private/NiagaraDataInterfaceSpline.cpp
--- a/Engine/Plugins/FX/Niagara/Source/Niagara/Private/NiagaraDataInterfaceSpline.cpp
+++ b/Engine/Plugins/FX/Niagara/Source/Niagara/Private/NiagaraDataInterfaceSpline.cpp
@@ -504,7 +504,7 @@ void UNiagaraDataInterfaceSpline::DestroyPerInstanceData(void* PerInstanceData,
 
 	FNiagaraDataInterfaceProxySpline* RT_Proxy = GetProxyAs<FNiagaraDataInterfaceProxySpline>();
 	ENQUEUE_RENDER_COMMAND(FNiagaraDIDestroyInstanceData) (
-		[RT_Proxy, InstanceID=SystemInstance->GetId(), Batcher=SystemInstance->GetBatcher()](FRHICommandListImmediate& CmdList)
+		[RT_Proxy, InstanceID=SystemInstance->GetId()](FRHICommandListImmediate& CmdList)
 		{
 #if STATS
 			if (FNDISpline_InstanceData_RenderThread* TargetData = RT_Proxy->SystemInstancesToProxyData_RT.Find(InstanceID))
diff --git a/Engine/Plugins/FX/Niagara/Source/Niagara/Private/NiagaraDataInterfaceStaticMesh.cpp b/Engine/Plugins/FX/Niagara/Source/Niagara/Private/NiagaraDataInterfaceStaticMesh.cpp
--- a/Engine/Plugins/FX/Niagara/Source/Niagara/Private/NiagaraDataInterfaceStaticMesh.cpp
+++ b/Engine/Plugins/FX/Niagara/Source/Niagara/Private/NiagaraDataInterfaceStaticMesh.cpp
@@ -10,7 +10,6 @@
 #include "Internationalization/Internationalization.h"
 #include "NiagaraScript.h"
 #include "ShaderParameterUtils.h"
-#include "NiagaraEmitterInstanceBatcher.h"
 #include "NiagaraStats.h"
 
 #define LOCTEXT_NAMESPACE "NiagaraDataInterfaceStaticMesh"
@@ -834,7 +833,7 @@ void FNiagaraDataInterfaceProxyStaticMesh::InitializePerInstanceData(const FNiag
 	Data.MeshGpuSpawnBuffer = MeshGPUSpawnBuffer;
 }
 
-void FNiagaraDataInterfaceProxyStaticMesh::DestroyPerInstanceData(NiagaraEmitterInstanceBatcher* Batcher, const FNiagaraSystemInstanceID& SystemInstance)
+void FNiagaraDataInterfaceProxyStaticMesh::DestroyPerInstanceData(const FNiagaraSystemInstanceID& SystemInstance)
 {
 	check(IsInRenderingThread());
 	//check(SystemInstancesToMeshData.Contains(SystemInstance));
@@ -1591,9 +1590,9 @@ void UNiagaraDataInterfaceStaticMesh::DestroyPerInstanceData(void* PerInstanceDa
 	{
 		FNiagaraDataInterfaceProxyStaticMesh* ThisProxy = GetProxyAs<FNiagaraDataInterfaceProxyStaticMesh>();
 		ENQUEUE_RENDER_COMMAND(FNiagaraDIDestroyInstanceData) (
-			[ThisProxy, InstanceID = SystemInstance->GetId(), Batcher = SystemInstance->GetBatcher()](FRHICommandListImmediate& CmdList)
+			[ThisProxy, InstanceID = SystemInstance->GetId()](FRHICommandListImmediate& CmdList)
 			{
-				ThisProxy->DestroyPerInstanceData(Batcher, InstanceID);
+				ThisProxy->DestroyPerInstanceData(InstanceID);
 			}
 		);
 	}
diff --git a/Engine/Plugins/FX/Niagara/Source/Niagara/Private/NiagaraDataSet.cpp b/Engine/Plugins/FX/Niagara/Source/Niagara/Private/NiagaraDataSet.cpp
--- a/Engine/Plugins/FX/Niagara/Source/Niagara/Private/NiagaraDataSet.cpp
+++ b/Engine/Plugins/FX/Niagara/Source/Niagara/Private/NiagaraDataSet.cpp
@@ -9,7 +9,8 @@
 #include "NiagaraStats.h"
 #include "NiagaraRenderer.h"
 #include "NiagaraGPUInstanceCountManager.h"
-#include "NiagaraEmitterInstanceBatcher.h"
+#include "NiagaraGpuComputeDispatchInterface.h"
+#include "NiagaraGpuComputeDispatch.h"
 #include "NiagaraShaderParticleID.h"
 
 DECLARE_CYCLE_STAT(TEXT("InitRenderData"), STAT_InitRenderData, STATGROUP_Niagara);
@@ -1298,12 +1299,12 @@ void FNiagaraDataBuffer::UnsetShaderParams(FRHICommandList& RHICmdList, FNiagara
 }
 
 #if WITH_EDITOR
-void FScopedNiagaraDataSetGPUReadback::ReadbackData(NiagaraEmitterInstanceBatcher* InBatcher, FNiagaraDataSet* InDataSet)
+void FScopedNiagaraDataSetGPUReadback::ReadbackData(FNiagaraGpuComputeDispatchInterface* InComputeDispatchInterface, FNiagaraDataSet* InDataSet)
 {
 	check(DataSet == nullptr);
 	check(InDataSet != nullptr);
 
-	Batcher = InBatcher && !InBatcher->IsPendingKill() ? Batcher : nullptr;
+	ComputeDispatchInterface = InComputeDispatchInterface && !InComputeDispatchInterface->IsPendingKill() ? ComputeDispatchInterface : nullptr;
 	DataSet = InDataSet;
 	DataBuffer = DataSet->GetCurrentData();
 
@@ -1317,9 +1318,9 @@ void FScopedNiagaraDataSetGPUReadback::ReadbackData(NiagaraEmitterInstanceBatche
 		{
 			// Read DrawIndirect Params
 			const uint32 BufferOffset = DataBuffer->GetGPUInstanceCountBufferOffset();
-			if (Batcher && BufferOffset != INDEX_NONE)
+			if (ComputeDispatchInterface && BufferOffset != INDEX_NONE)
 			{
-				FRHIBuffer* InstanceCountBuffer = Batcher->GetGPUInstanceCounterManager().GetInstanceCountBuffer().Buffer;
+				FRHIBuffer* InstanceCountBuffer = static_cast<FNiagaraGpuComputeDispatch*>(ComputeDispatchInterface)->GetGPUInstanceCounterManager().GetInstanceCountBuffer().Buffer;
 
 				void* Data = RHICmdList.LockBuffer(InstanceCountBuffer, 0, (BufferOffset + 1) * sizeof(int32), RLM_ReadOnly);
 				NumInstances = reinterpret_cast<int32*>(Data)[BufferOffset];
diff --git a/Engine/Plugins/FX/Niagara/Source/Niagara/Private/NiagaraDebugHud.cpp b/Engine/Plugins/FX/Niagara/Source/Niagara/Private/NiagaraDebugHud.cpp
--- a/Engine/Plugins/FX/Niagara/Source/Niagara/Private/NiagaraDebugHud.cpp
+++ b/Engine/Plugins/FX/Niagara/Source/Niagara/Private/NiagaraDebugHud.cpp
@@ -2,8 +2,9 @@
 
 #include "NiagaraDebugHud.h"
 #include "NiagaraComponent.h"
-#include "NiagaraEmitterInstanceBatcher.h"
 #include "NiagaraFunctionLibrary.h"
+#include "NiagaraGpuComputeDispatchInterface.h"
+#include "NiagaraComputeExecutionContext.h"
 #include "NiagaraScript.h"
 #include "NiagaraSystem.h"
 #include "NiagaraWorldManager.h"
@@ -755,9 +756,9 @@ FNiagaraDataSet* FNiagaraDebugHud::GetParticleDataSet(FNiagaraSystemInstance* Sy
 			GpuCachedData->PendingEmitterData[iEmitter]->Frame.Init(&AllEmittersCompiledData[iEmitter]->DataSetCompiledData);
 
 			ENQUEUE_RENDER_COMMAND(NiagaraReadbackGpuSim)(
-				[RT_Batcher=SystemInstance->GetBatcher(), RT_InstanceID=SystemInstance->GetId(), RT_DebugInfo=GpuCachedData->PendingEmitterData[iEmitter], RT_Context=GPUExecContext](FRHICommandListImmediate& RHICmdList)
+				[RT_ComputeDispatchInterface=SystemInstance->GetComputeDispatchInterface(), RT_InstanceID=SystemInstance->GetId(), RT_DebugInfo=GpuCachedData->PendingEmitterData[iEmitter], RT_Context=GPUExecContext](FRHICommandListImmediate& RHICmdList)
 				{
-					RT_Batcher->AddDebugReadback(RT_InstanceID, RT_DebugInfo, RT_Context);
+					RT_ComputeDispatchInterface->AddDebugReadback(RT_InstanceID, RT_DebugInfo, RT_Context);
 				}
 			);
 		}
diff --git a/Engine/Plugins/FX/Niagara/Source/Niagara/Private/NiagaraEmitterInstance.cpp b/Engine/Plugins/FX/Niagara/Source/Niagara/Private/NiagaraEmitterInstance.cpp
--- a/Engine/Plugins/FX/Niagara/Source/Niagara/Private/NiagaraEmitterInstance.cpp
+++ b/Engine/Plugins/FX/Niagara/Source/Niagara/Private/NiagaraEmitterInstance.cpp
@@ -6,8 +6,9 @@
 #include "NiagaraRenderer.h"
 #include "NiagaraSystemInstance.h"
 #include "NiagaraSystemGpuComputeProxy.h"
+#include "NiagaraComputeExecutionContext.h"
 #include "NiagaraDataInterface.h"
-#include "NiagaraEmitterInstanceBatcher.h"
+#include "NiagaraGpuComputeDispatchInterface.h"
 #include "NiagaraScriptExecutionContext.h"
 #include "NiagaraParameterCollection.h"
 #include "NiagaraWorldManager.h"
@@ -128,8 +129,8 @@ FNiagaraEmitterInstance::FNiagaraEmitterInstance(FNiagaraSystemInstance* InParen
 {
 	ParticleDataSet = new FNiagaraDataSet();
 
-	Batcher = ParentSystemInstance ? ParentSystemInstance->GetBatcher() : nullptr;
-	check(Batcher != nullptr);
+	ComputeDispatchInterface = ParentSystemInstance ? ParentSystemInstance->GetComputeDispatchInterface() : nullptr;
+	check(ComputeDispatchInterface != nullptr);
 }
 
 FNiagaraEmitterInstance::~FNiagaraEmitterInstance()
@@ -147,9 +148,8 @@ FNiagaraEmitterInstance::~FNiagaraEmitterInstance()
 		GPUExecContext->CombinedParamStore.UnbindAll();
 
 		/** We defer the deletion of the particle dataset and the compute context to the RT to be sure all in-flight RT commands have finished using it.*/
-		NiagaraEmitterInstanceBatcher* Batcher_RT = Batcher && !Batcher->IsPendingKill() ? Batcher : nullptr;
 		ENQUEUE_RENDER_COMMAND(FDeleteContextCommand)(
-			[Batcher_RT, ExecContext=GPUExecContext, DataSet= ParticleDataSet](FRHICommandListImmediate& RHICmdList)
+			[ExecContext=GPUExecContext, DataSet=ParticleDataSet](FRHICommandListImmediate& RHICmdList)
 			{
 				if ( ExecContext != nullptr )
 				{
@@ -234,7 +234,7 @@ bool FNiagaraEmitterInstance::IsAllowedToExecute() const
 	if (CachedEmitter->SimTarget == ENiagaraSimTarget::GPUComputeSim)
 	{
 		//-TODO: Could replace with CPU side sim in some cases
-		if ( !Batcher || !FNiagaraUtilities::AllowGPUParticles(Batcher->GetShaderPlatform()) )
+		if ( !ComputeDispatchInterface || !FNiagaraUtilities::AllowGPUParticles(ComputeDispatchInterface->GetShaderPlatform()) )
 		{
 			return false;
 		}
@@ -973,7 +973,7 @@ void FNiagaraEmitterInstance::CalculateFixedBounds(const FTransform& ToWorldSpac
 			return;
 		}
 
-		ScopedGPUReadback.ReadbackData(Batcher, GPUExecContext->MainDataSet);
+		ScopedGPUReadback.ReadbackData(ComputeDispatchInterface, GPUExecContext->MainDataSet);
 		NumInstances = ScopedGPUReadback.GetNumInstances();
 	}
 	else
@@ -1081,7 +1081,7 @@ bool FNiagaraEmitterInstance::HandleCompletion(bool bForce)
 	{
 		if( GPUExecContext )
 		{
-			GPUExecContext->Reset(Batcher);
+			GPUExecContext->Reset(ComputeDispatchInterface);
 		}
 		ParticleDataSet->ResetBuffers();
 		if (EventInstanceData.IsValid())
@@ -1177,7 +1177,7 @@ void FNiagaraEmitterInstance::PreTick()
 		UpdateExecContext.PostTick();
 		if (CachedEmitter->SimTarget == ENiagaraSimTarget::GPUComputeSim && GPUExecContext != nullptr)
 		{
-			//We PostTick the GPUExecContext here to prime crucial PREV parameters (such as PREV_Engine.Owner.Position). This PostTick call is necessary as the GPUExecContext has not been sent to the batcher yet.
+			//We PostTick the GPUExecContext here to prime crucial PREV parameters (such as PREV_Engine.Owner.Position). This PostTick call is necessary as the GPUExecContext has not been sent to the dispatch interface yet.
 			GPUExecContext->PostTick();
 		}
 
@@ -1225,7 +1225,7 @@ bool FNiagaraEmitterInstance::WaitForDebugInfo()
 	FNiagaraComputeExecutionContext* DebugContext = GPUExecContext;
 	if (CachedEmitter->SimTarget == ENiagaraSimTarget::GPUComputeSim && DebugContext)
 	{
-		ENQUEUE_RENDER_COMMAND(CaptureCommand)([=](FRHICommandListImmediate& RHICmdList) { Batcher->ProcessDebugReadbacks(RHICmdList, true); });
+		ENQUEUE_RENDER_COMMAND(CaptureCommand)([=](FRHICommandListImmediate& RHICmdList) { ComputeDispatchInterface->ProcessDebugReadbacks(RHICmdList, true); });
 		FlushRenderingCommands(); 
 		return true;
 	}
@@ -1297,7 +1297,7 @@ void FNiagaraEmitterInstance::Tick(float DeltaSeconds)
 	{
 		if (GPUExecContext)
 		{
-			GPUExecContext->Reset(Batcher);
+			GPUExecContext->Reset(ComputeDispatchInterface);
 		}
 		Data.ResetBuffers();
 		ExecutionState = ENiagaraExecutionState::Inactive;
@@ -1404,7 +1404,7 @@ void FNiagaraEmitterInstance::Tick(float DeltaSeconds)
 		EmitterParameters.EmitterInstanceSeed = InstanceSeed;
 	}
 
-	/* GPU simulation -  we just create an FNiagaraComputeExecutionContext, queue it, and let the batcher take care of the rest
+	/* GPU simulation -  we just create an FNiagaraComputeExecutionContext, queue it, and let the dispatch interface take care of the rest
 	 */
 	if (CachedEmitter->SimTarget == ENiagaraSimTarget::GPUComputeSim && GPUExecContext != nullptr)
 	{
@@ -1415,7 +1415,7 @@ void FNiagaraEmitterInstance::Tick(float DeltaSeconds)
 		if (ParentSystemInstance->ShouldCaptureThisFrame())
 		{
 			TSharedPtr<struct FNiagaraScriptDebuggerInfo, ESPMode::ThreadSafe> DebugInfo = ParentSystemInstance->GetActiveCaptureWrite(CachedIDName, ENiagaraScriptUsage::ParticleGPUComputeScript, FGuid());
-			if (DebugInfo && Batcher != nullptr)
+			if (DebugInfo && ComputeDispatchInterface != nullptr)
 			{
 				//Data.Dump(DebugInfo->Frame, true, 0, OrigNumParticles);
 				//DebugInfo->Frame.Dump(true, 0, OrigNumParticles);
@@ -1426,9 +1426,9 @@ void FNiagaraEmitterInstance::Tick(float DeltaSeconds)
 
 				// Execute a readback
 				ENQUEUE_RENDER_COMMAND(NiagaraReadbackGpuSim)(
-					[RT_Batcher=Batcher, RT_InstanceID=OwnerSystemInstanceID, RT_DebugInfo=DebugInfo, RT_Context=GPUExecContext](FRHICommandListImmediate& RHICmdList)
+					[RT_ComputeDispatchInterface = ComputeDispatchInterface, RT_InstanceID=OwnerSystemInstanceID, RT_DebugInfo=DebugInfo, RT_Context=GPUExecContext](FRHICommandListImmediate& RHICmdList)
 					{
-						RT_Batcher->AddDebugReadback(RT_InstanceID, RT_DebugInfo, RT_Context);
+						RT_ComputeDispatchInterface->AddDebugReadback(RT_InstanceID, RT_DebugInfo, RT_Context);
 					}
 				);
 			}
@@ -1538,7 +1538,7 @@ void FNiagaraEmitterInstance::Tick(float DeltaSeconds)
 		SpawnExecContext.PostTick();
 		UpdateExecContext.PostTick();
 
-		// At this stage GPU execution is being handled by the batcher so we do not need to call PostTick() for it
+		// At this stage GPU execution is being handled by the dispatch interface so we do not need to call PostTick() for it
 		for (FNiagaraScriptExecutionContext& EventContext : GetEventExecutionContexts())
 		{
 			EventContext.PostTick();
@@ -2072,7 +2072,7 @@ void FNiagaraEmitterInstance::Tick(float DeltaSeconds)
 
 	SpawnExecContext.PostTick();
 	UpdateExecContext.PostTick();
-	// At this stage GPU execution is being handled by the batcher so we do not need to call PostTick() for it
+	// At this stage GPU execution is being handled by the dispatch interface so we do not need to call PostTick() for it
 
 	for (FNiagaraScriptExecutionContext& EventContext : GetEventExecutionContexts())
 	{
diff --git a/Engine/Plugins/FX/Niagara/Source/Niagara/Private/NiagaraEmitterInstanceBatcher.cpp b/Engine/Plugins/FX/Niagara/Source/Niagara/Private/NiagaraEmitterInstanceBatcher.cpp
deleted file mode 100644
diff --git a/Engine/Plugins/FX/Niagara/Source/Niagara/Private/NiagaraEmptyUAVPool.cpp b/Engine/Plugins/FX/Niagara/Source/Niagara/Private/NiagaraEmptyUAVPool.cpp
new file mode 100644
--- /dev/null
+++ b/Engine/Plugins/FX/Niagara/Source/Niagara/Private/NiagaraEmptyUAVPool.cpp
@@ -0,0 +1,113 @@
+// Copyright Epic Games, Inc. All Rights Reserved.
+
+#include "NiagaraEmptyUAVPool.h"
+#include "NiagaraStats.h"
+
+DECLARE_DWORD_ACCUMULATOR_STAT(TEXT("# EmptyUAVs"), STAT_NiagaraEmptyUAVPool, STATGROUP_Niagara);
+
+FNiagaraEmptyUAVPoolScopedAccess::FNiagaraEmptyUAVPoolScopedAccess(class FNiagaraEmptyUAVPool* InEmptyUAVPool)
+	: EmptyUAVPool(InEmptyUAVPool)
+{
+	check(EmptyUAVPool);
+	++EmptyUAVPool->UAVAccessCounter;
+}
+
+FNiagaraEmptyUAVPoolScopedAccess::~FNiagaraEmptyUAVPoolScopedAccess()
+{
+	--EmptyUAVPool->UAVAccessCounter;
+	if (EmptyUAVPool->UAVAccessCounter == 0)
+	{
+		EmptyUAVPool->ResetEmptyUAVPools();
+	}
+}
+
+FRHIUnorderedAccessView* FNiagaraEmptyUAVPool::GetEmptyUAVFromPool(FRHICommandList& RHICmdList, EPixelFormat Format, ENiagaraEmptyUAVType Type)
+{
+	check(IsInRenderingThread());
+	checkf(UAVAccessCounter != 0, TEXT("Accessing Niagara's UAV Pool while not within a scope, this could result in a memory leak!"));
+
+	TMap<EPixelFormat, FEmptyUAVPool>& UAVMap = UAVPools[(int)Type];
+	FEmptyUAVPool& Pool = UAVMap.FindOrAdd(Format);
+	checkSlow(Pool.NextFreeIndex <= Pool.UAVs.Num());
+	if (Pool.NextFreeIndex == Pool.UAVs.Num())
+	{
+		FEmptyUAV& NewUAV = Pool.UAVs.AddDefaulted_GetRef();
+
+		// Initialize the UAV
+		FRHIResourceCreateInfo CreateInfo(TEXT("FNiagaraGpuComputeDispatch::EmptyUAV"));
+		switch ( Type )
+		{
+			case ENiagaraEmptyUAVType::Buffer:
+			{
+				const uint32 BytesPerElement = GPixelFormats[Format].BlockBytes;
+				NewUAV.Buffer = RHICreateVertexBuffer(BytesPerElement, BUF_UnorderedAccess | BUF_ShaderResource, CreateInfo);
+				NewUAV.UAV = RHICreateUnorderedAccessView(NewUAV.Buffer, Format);
+				break;
+			}
+			
+			case ENiagaraEmptyUAVType::Texture2D:
+			{
+				NewUAV.Texture = RHICreateTexture2D(1, 1, Format, 1, 1, TexCreate_ShaderResource | TexCreate_UAV, CreateInfo);
+				NewUAV.UAV = RHICreateUnorderedAccessView(NewUAV.Texture, 0);
+				break;
+			}
+			
+			case ENiagaraEmptyUAVType::Texture2DArray:
+			{
+				NewUAV.Texture = RHICreateTexture2DArray(1, 1, 1, Format, 1, 1, TexCreate_ShaderResource | TexCreate_UAV, CreateInfo);
+				NewUAV.UAV = RHICreateUnorderedAccessView(NewUAV.Texture, 0);
+				break;
+			}
+			
+			case ENiagaraEmptyUAVType::Texture3D:
+			{
+				NewUAV.Texture = RHICreateTexture3D(1, 1, 1, Format, 1, TexCreate_ShaderResource | TexCreate_UAV, CreateInfo);
+				NewUAV.UAV = RHICreateUnorderedAccessView(NewUAV.Texture, 0);
+				break;
+			}
+			
+			default:
+			{
+				checkNoEntry();
+				return nullptr;
+			}
+		}
+
+		RHICmdList.Transition(FRHITransitionInfo(NewUAV.UAV, ERHIAccess::Unknown, ERHIAccess::UAVCompute));
+
+		// Dispatches which use Empty UAVs are allowed to overlap, since we don't care about the contents of these buffers.
+		// We never need to calll EndUAVOverlap() on these.
+		RHICmdList.BeginUAVOverlap(NewUAV.UAV);
+
+		INC_DWORD_STAT(STAT_NiagaraEmptyUAVPool);
+	}
+
+	FRHIUnorderedAccessView* UAV = Pool.UAVs[Pool.NextFreeIndex].UAV;
+	++Pool.NextFreeIndex;
+	return UAV;
+}
+
+void FNiagaraEmptyUAVPool::ResetEmptyUAVPools()
+{
+	for (int Type = 0; Type < UE_ARRAY_COUNT(UAVPools); ++Type)
+	{
+		for (TPair<EPixelFormat, FEmptyUAVPool>& Entry : UAVPools[Type])
+		{
+			Entry.Value.NextFreeIndex = 0;
+		}
+	}
+}
+
+FNiagaraEmptyUAVPool::FEmptyUAV::~FEmptyUAV()
+{
+	Buffer.SafeRelease();
+	Texture.SafeRelease();
+	UAV.SafeRelease();
+}
+
+FNiagaraEmptyUAVPool::FEmptyUAVPool::~FEmptyUAVPool()
+{
+	UE_CLOG(NextFreeIndex != 0, LogNiagara, Warning, TEXT("EmptyUAVPool is potentially in use during destruction."));
+	DEC_DWORD_STAT_BY(STAT_NiagaraEmptyUAVPool, UAVs.Num());
+}
+
diff --git a/Engine/Plugins/FX/Niagara/Source/Niagara/Private/NiagaraFunctionLibrary.cpp b/Engine/Plugins/FX/Niagara/Source/Niagara/Private/NiagaraFunctionLibrary.cpp
--- a/Engine/Plugins/FX/Niagara/Source/Niagara/Private/NiagaraFunctionLibrary.cpp
+++ b/Engine/Plugins/FX/Niagara/Source/Niagara/Private/NiagaraFunctionLibrary.cpp
@@ -20,7 +20,7 @@
 #include "Engine/VolumeTexture.h"
 #include "Engine/Texture2DArray.h"
 #include "Engine/TextureCube.h"
-#include "NiagaraEmitterInstanceBatcher.h"
+#include "NiagaraGpuComputeDispatchInterface.h"
 
 #define LOCTEXT_NAMESPACE "NiagaraFunctionLibrary"
 
@@ -1786,17 +1786,11 @@ void UNiagaraFunctionLibrary::InitVectorVMFastPathOps()
 void UNiagaraFunctionLibrary::SetComponentNiagaraGPURayTracedCollisionGroup(UObject* WorldContextObject, UPrimitiveComponent* Primitive, int32 CollisionGroup)
 {
 #if RHI_RAYTRACING
-	if(UWorld* World = GEngine->GetWorldFromContextObject(WorldContextObject, EGetWorldErrorMode::LogAndReturnNull))
+	if ( UWorld* World = GEngine->GetWorldFromContextObject(WorldContextObject, EGetWorldErrorMode::LogAndReturnNull) )
 	{
-		if(FNiagaraWorldManager* WorldMan = FNiagaraWorldManager::Get(World))
+		if ( FNiagaraGpuComputeDispatchInterface* ComputeDispatchInterface = FNiagaraGpuComputeDispatchInterface::Get(World) )
 		{
-			if(FFXSystemInterface* FXSystemInterface = World->Scene->GetFXSystem())			
-			{
-				if (NiagaraEmitterInstanceBatcher* Batcher = static_cast<NiagaraEmitterInstanceBatcher*>(FXSystemInterface->GetInterface(NiagaraEmitterInstanceBatcher::Name)))
-				{
-					Batcher->SetPrimitiveRayTracingCollisionGroup(Primitive, CollisionGroup);
-				}
-			}
+			ComputeDispatchInterface->SetPrimitiveRayTracingCollisionGroup(Primitive, CollisionGroup);
 		}
 	}
 #endif
@@ -1810,21 +1804,17 @@ void UNiagaraFunctionLibrary::SetActorNiagaraGPURayTracedCollisionGroup(UObject*
 		return;
 	}
 
-	if (UWorld* World = GEngine->GetWorldFromContextObject(WorldContextObject, EGetWorldErrorMode::LogAndReturnNull))
+	if ( UWorld* World = GEngine->GetWorldFromContextObject(WorldContextObject, EGetWorldErrorMode::LogAndReturnNull) )
 	{
-		if (FNiagaraWorldManager* WorldMan = FNiagaraWorldManager::Get(World))
+		if (FNiagaraGpuComputeDispatchInterface* ComputeDispatchInterface = FNiagaraGpuComputeDispatchInterface::Get(World))
 		{
-			if (FFXSystemInterface* FXSystemInterface = World->Scene->GetFXSystem())
-			{
-				if (NiagaraEmitterInstanceBatcher* Batcher = static_cast<NiagaraEmitterInstanceBatcher*>(FXSystemInterface->GetInterface(NiagaraEmitterInstanceBatcher::Name)))
+			Actor->ForEachComponent<UPrimitiveComponent>(
+				/*bIncludeFromChildActors*/true,
+				[&ComputeDispatchInterface, &CollisionGroup](UPrimitiveComponent* PrimitiveComponent)
 				{
-					Actor->ForEachComponent<UPrimitiveComponent>(/*bIncludeFromChildActors*/true, [&Batcher, &CollisionGroup](UPrimitiveComponent* PrimitiveComponent)
-					{
-							Batcher->SetPrimitiveRayTracingCollisionGroup(PrimitiveComponent, CollisionGroup);
-					});
-				}
+					ComputeDispatchInterface->SetPrimitiveRayTracingCollisionGroup(PrimitiveComponent, CollisionGroup);
+				});
 			}
-		}
 	}
 #endif
 }
@@ -1834,15 +1824,9 @@ int32 UNiagaraFunctionLibrary::AcquireNiagaraGPURayTracedCollisionGroup(UObject*
 #if RHI_RAYTRACING
 	if (UWorld* World = GEngine->GetWorldFromContextObject(WorldContextObject, EGetWorldErrorMode::LogAndReturnNull))
 	{
-		if (FNiagaraWorldManager* WorldMan = FNiagaraWorldManager::Get(World))
+		if (FNiagaraGpuComputeDispatchInterface* ComputeDispatchInterface = FNiagaraGpuComputeDispatchInterface::Get(World))
 		{
-			if (FFXSystemInterface* FXSystemInterface = World->Scene->GetFXSystem())
-			{
-				if (NiagaraEmitterInstanceBatcher* Batcher = static_cast<NiagaraEmitterInstanceBatcher*>(FXSystemInterface->GetInterface(NiagaraEmitterInstanceBatcher::Name)))
-				{
-					return Batcher->AcquireGPURayTracedCollisionGroup();
-				}
-			}
+			return ComputeDispatchInterface->AcquireGPURayTracedCollisionGroup();
 		}
 	}
 #endif
@@ -1854,15 +1838,9 @@ void UNiagaraFunctionLibrary::ReleaseNiagaraGPURayTracedCollisionGroup(UObject*
 #if RHI_RAYTRACING
 	if (UWorld* World = GEngine->GetWorldFromContextObject(WorldContextObject, EGetWorldErrorMode::LogAndReturnNull))
 	{
-		if (FNiagaraWorldManager* WorldMan = FNiagaraWorldManager::Get(World))
+		if (FNiagaraGpuComputeDispatchInterface* ComputeDispatchInterface = FNiagaraGpuComputeDispatchInterface::Get(World))
 		{
-			if (FFXSystemInterface* FXSystemInterface = World->Scene->GetFXSystem())
-			{
-				if (NiagaraEmitterInstanceBatcher* Batcher = static_cast<NiagaraEmitterInstanceBatcher*>(FXSystemInterface->GetInterface(NiagaraEmitterInstanceBatcher::Name)))
-				{
-					Batcher->ReleaseGPURayTracedCollisionGroup(CollisionGroup);
-				}
-			}
+			ComputeDispatchInterface->ReleaseGPURayTracedCollisionGroup(CollisionGroup);
 		}
 	}
 #endif
diff --git a/Engine/Plugins/FX/Niagara/Source/Niagara/Private/NiagaraGPUInstanceCountManager.cpp b/Engine/Plugins/FX/Niagara/Source/Niagara/Private/NiagaraGPUInstanceCountManager.cpp
--- a/Engine/Plugins/FX/Niagara/Source/Niagara/Private/NiagaraGPUInstanceCountManager.cpp
+++ b/Engine/Plugins/FX/Niagara/Source/Niagara/Private/NiagaraGPUInstanceCountManager.cpp
@@ -1,12 +1,14 @@
 // Copyright Epic Games, Inc. All Rights Reserved.
 
 #include "NiagaraGPUInstanceCountManager.h"
+#include "NiagaraEmptyUAVPool.h"
+#include "NiagaraGpuComputeDispatch.h"
+#include "NiagaraRenderer.h"
 #include "NiagaraStats.h"
+
 #include "Containers/DynamicRHIResourceArray.h"
 #include "GPUSortManager.h" // CopyUIntBufferToTargets
 #include "ProfilingDebugging/RealtimeGPUProfiler.h"
-#include "NiagaraRenderer.h"
-#include "NiagaraEmitterInstanceBatcher.h"
 #include "ClearQuad.h"
 #include "PipelineStateCache.h"
 
@@ -250,7 +252,7 @@ void FNiagaraGPUInstanceCountManager::ResizeBuffers(FRHICommandListImmediate& RH
 			int32 UsedIndexCounts[] = { AllocatedInstanceCounts };
 			CopyUIntBufferToTargets(RHICmdList, FeatureLevel, CountBuffer.SRV, UAVs, UsedIndexCounts, 0, UE_ARRAY_COUNT(UAVs));
 
-			// NiagaraEmitterInstanceBatcher expects the count buffer to be readable and copyable before running the sim.
+			// FNiagaraGpuComputeDispatch expects the count buffer to be readable and copyable before running the sim.
 			RHICmdList.Transition(FRHITransitionInfo(NextCountBuffer.UAV, ERHIAccess::UAVCompute, kCountBufferDefaultState));
 
 			// Swap the buffers
@@ -350,7 +352,7 @@ FNiagaraGPUInstanceCountManager::FIndirectArgSlot FNiagaraGPUInstanceCountManage
 
 ////////////////////////////////////////////////////////////////////////////////////////////////////////////////////////////////////////////////////////////////////////////////////////////////////////
 
-void FNiagaraGPUInstanceCountManager::UpdateDrawIndirectBuffers(NiagaraEmitterInstanceBatcher& Batcher, FRHICommandList& RHICmdList, ERHIFeatureLevel::Type FeatureLevel, ENiagaraGPUCountUpdatePhase::Type CountPhase)
+void FNiagaraGPUInstanceCountManager::UpdateDrawIndirectBuffers(FNiagaraGpuComputeDispatchInterface* ComputeDispatchInterface, FRHICommandList& RHICmdList, ERHIFeatureLevel::Type FeatureLevel, ENiagaraGPUCountUpdatePhase::Type CountPhase)
 {
 	// Anything to process?
 	TArray<FNiagaraDrawIndirectArgGenTaskInfo>& ArgTasks = DrawIndirectArgGenTasks[CountPhase];
@@ -375,7 +377,7 @@ void FNiagaraGPUInstanceCountManager::UpdateDrawIndirectBuffers(NiagaraEmitterIn
 			RHIUnlockBuffer(TaskInfosBuffer.Buffer);
 		}
 
-		FNiagaraUAVPoolAccessScope UAVPoolAccessScope(Batcher);
+		FNiagaraEmptyUAVPoolScopedAccess UAVPoolAccessScope(ComputeDispatchInterface->GetEmptyUAVPool());
 		TArray<FRHITransitionInfo, TInlineAllocator<10>> Transitions;
 		Transitions.Reserve(DrawIndirectPool.Num() + 2);
 
@@ -390,7 +392,7 @@ void FNiagaraGPUInstanceCountManager::UpdateDrawIndirectBuffers(NiagaraEmitterIn
 		else
 		{
 			// This can happen if there are no InstanceCountClearTasks and all DrawIndirectArgGenTasks_PreOpaque are using culled counts
-			CountsUAV = Batcher.GetEmptyUAVFromPool(RHICmdList, PF_R32_UINT, ENiagaraEmptyUAVType::Buffer);
+			CountsUAV = ComputeDispatchInterface->GetEmptyUAVFromPool(RHICmdList, PF_R32_UINT, ENiagaraEmptyUAVType::Buffer);
 		}
 
 		// Get culled counts buffer
@@ -441,7 +443,7 @@ void FNiagaraGPUInstanceCountManager::UpdateDrawIndirectBuffers(NiagaraEmitterIn
 			}
 			else
 			{
-				ArgsUAV = Batcher.GetEmptyUAVFromPool(RHICmdList, PF_R32_UINT, ENiagaraEmptyUAVType::Buffer);
+				ArgsUAV = ComputeDispatchInterface->GetEmptyUAVFromPool(RHICmdList, PF_R32_UINT, ENiagaraEmptyUAVType::Buffer);
 			}
 
 			const bool bIsLastDispatch = DispatchIdx == (NumDispatches - 1);
@@ -591,8 +593,8 @@ void FNiagaraGPUInstanceCountManager::EnqueueGPUReadback(FRHICommandListImmediat
 			CountReadback = new FRHIGPUBufferReadback(TEXT("Niagara GPU Instance Count Readback"));
 		}
 		CountReadbackSize = UsedInstanceCounts;
-		// No need for a transition, NiagaraEmitterInstanceBatcher ensures that the buffer is left in the
-		// correct state after the sim.
+
+		// No need for a transition, FNiagaraGpuComputeDispatch ensures that the buffer is left in the correct state after the simulation.
 		CountReadback->EnqueueCopy(RHICmdList, CountBuffer.Buffer);
 	}
 }
diff --git a/Engine/Plugins/FX/Niagara/Source/Niagara/Private/NiagaraGpuComputeDispatch.cpp b/Engine/Plugins/FX/Niagara/Source/Niagara/Private/NiagaraGpuComputeDispatch.cpp
new file mode 100644
--- /dev/null
+++ b/Engine/Plugins/FX/Niagara/Source/Niagara/Private/NiagaraGpuComputeDispatch.cpp
@@ -0,0 +1,2204 @@
+// Copyright Epic Games, Inc. All Rights Reserved.
+
+#include "NiagaraGpuComputeDispatch.h"
+
+#include "Async/Async.h"
+#include "CanvasTypes.h"
+#include "ClearQuad.h"
+#include "Engine/Canvas.h"
+#include "Engine/GameViewportClient.h"
+#include "Engine/World.h"
+#include "GPUSort.h"
+#include "Misc/ScopeExit.h"
+
+#include "NiagaraDataInterfaceRW.h"
+#if NIAGARA_COMPUTEDEBUG_ENABLED
+#include "NiagaraGpuComputeDebug.h"
+#endif
+#include "NiagaraGpuReadbackManager.h"
+#include "NiagaraRayTracingHelper.h"
+#include "NiagaraRenderer.h"
+#include "NiagaraScriptExecutionContext.h"
+#include "NiagaraSystemGpuComputeProxy.h"
+#include "NiagaraShader.h"
+#include "NiagaraShaderParticleID.h"
+#include "NiagaraSortingGPU.h"
+#include "NiagaraStats.h"
+#include "NiagaraRenderViewDataManager.h"
+#include "NiagaraWorldManager.h"
+#include "RHI.h"
+#include "RHIGPUReadback.h"
+#include "Runtime/Renderer/Private/SceneRendering.h"
+#include "Runtime/Renderer/Private/PostProcess/SceneRenderTargets.h"
+#include "PipelineStateCache.h"
+#include "SceneUtils.h"
+#include "ShaderParameterUtils.h"
+
+DECLARE_CYCLE_STAT(TEXT("Niagara Dispatch Setup"), STAT_NiagaraGPUDispatchSetup_RT, STATGROUP_Niagara);
+DECLARE_CYCLE_STAT(TEXT("GPU Emitter Dispatch [RT]"), STAT_NiagaraGPUSimTick_RT, STATGROUP_Niagara);
+DECLARE_CYCLE_STAT(TEXT("GPU Data Readback [RT]"), STAT_NiagaraGPUReadback_RT, STATGROUP_Niagara);
+DECLARE_FLOAT_COUNTER_STAT(TEXT("Niagara GPU Sim"), STAT_GPU_NiagaraSim, STATGROUP_GPU);
+DECLARE_DWORD_COUNTER_STAT(TEXT("# GPU Particles"), STAT_NiagaraGPUParticles, STATGROUP_Niagara);
+DECLARE_DWORD_COUNTER_STAT(TEXT("# GPU Sorted Particles"), STAT_NiagaraGPUSortedParticles, STATGROUP_Niagara);
+DECLARE_DWORD_COUNTER_STAT(TEXT("# GPU Sorted Buffers"), STAT_NiagaraGPUSortedBuffers, STATGROUP_Niagara);
+DECLARE_DWORD_COUNTER_STAT(TEXT("Readback latency (frames)"), STAT_NiagaraReadbackLatency, STATGROUP_Niagara);
+DECLARE_DWORD_COUNTER_STAT(TEXT("# GPU Dispatches"), STAT_NiagaraGPUDispatches, STATGROUP_Niagara);
+
+DECLARE_GPU_STAT_NAMED(NiagaraGPU, TEXT("Niagara"));
+DECLARE_GPU_STAT_NAMED(NiagaraGPUSimulation, TEXT("Niagara GPU Simulation"));
+DECLARE_GPU_STAT_NAMED(NiagaraGPUClearIDTables, TEXT("NiagaraGPU Clear ID Tables"));
+DECLARE_GPU_STAT_NAMED(NiagaraGPUComputeFreeIDs, TEXT("Niagara GPU Compute All Free IDs"));
+DECLARE_GPU_STAT_NAMED(NiagaraGPUComputeFreeIDsEmitter, TEXT("Niagara GPU Compute Emitter Free IDs"));
+DECLARE_GPU_STAT_NAMED(NiagaraGPUSorting, TEXT("Niagara GPU sorting"));
+
+uint32 FNiagaraComputeExecutionContext::TickCounter = 0;
+
+#if WITH_MGPU
+const FName FNiagaraGpuComputeDispatch::TemporalEffectName("FNiagaraGpuComputeDispatch");
+#endif // WITH_MGPU
+
+int32 GNiagaraGpuSubmitCommandHint = WITH_EDITOR ? 10 : 0;
+static FAutoConsoleVariableRef CVarNiagaraGpuSubmitCommandHint(
+	TEXT("fx.NiagaraGpuSubmitCommandHint"),
+	GNiagaraGpuSubmitCommandHint,
+	TEXT("If greater than zero, we use this value to submit commands after the number of dispatches have been issued."),
+	ECVF_Default
+);
+
+int32 GNiagaraGpuLowLatencyTranslucencyEnabled = 0;
+static FAutoConsoleVariableRef CVarNiagaraGpuLowLatencyTranslucencyEnabled(
+	TEXT("fx.NiagaraGpuLowLatencyTranslucencyEnabled"),
+	GNiagaraGpuLowLatencyTranslucencyEnabled,
+	TEXT("When enabled translucent materials can use the current frames simulation data no matter which tick pass Niagara uses.\n")
+	TEXT("This can result in an additional data buffer being required but will reduce any latency when using view uniform buffer / depth buffer / distance fields / etc"),
+	ECVF_Default
+);
+
+int32 GNiagaraBatcherFreeBufferEarly = 1;
+static FAutoConsoleVariableRef CVarNiagaraBatcherFreeBufferEarly(
+	TEXT("fx.NiagaraBatcher.FreeBufferEarly"),
+	GNiagaraBatcherFreeBufferEarly,
+	TEXT("Will take the path to release GPU buffers when possible.\n")
+	TEXT("This will reduce memory pressure but can result in more allocations if you buffers ping pong from zero particles to many."),
+	ECVF_Default
+);
+
+const FName FNiagaraGpuComputeDispatch::Name(TEXT("FNiagaraGpuComputeDispatch"));
+
+namespace FNiagaraGpuComputeDispatchLocal
+{
+	int32 GTickFlushMaxQueuedFrames = 10;
+	static FAutoConsoleVariableRef CVarNiagaraTickFlushMaxQueuedFrames(
+		TEXT("fx.Niagara.Batcher.TickFlush.MaxQueuedFrames"),
+		GTickFlushMaxQueuedFrames,
+		TEXT("The number of unprocessed frames with queued ticks before we process them.\n")
+		TEXT("The larger the number the more data we process in a single frame, this is generally only a concern when the application does not have focus."),
+		ECVF_Default
+	);
+
+	int32 GTickFlushMode = 1;
+	static FAutoConsoleVariableRef CVarNiagaraTickFlushMode(
+		TEXT("fx.Niagara.Batcher.TickFlush.Mode"),
+		GTickFlushMode,
+		TEXT("What to do when we go over our max queued frames.\n")
+		TEXT("0 = Keep ticks queued, can result in a long pause when gaining focus again.\n")
+		TEXT("1 = (Default) Process all queued ticks with dummy view / buffer data, may result in incorrect simulation due to missing depth collisions, etc.\n")
+		TEXT("2 = Kill all pending ticks, may result in incorrect simulation due to missing frames of data, i.e. a particle reset.\n"),
+		ECVF_Default
+	);
+
+	static int32 GAddDispatchGroupDrawEvent = 0;
+	static FAutoConsoleVariableRef CVarAddDispatchGroupDrawEvent(
+		TEXT("fx.Niagara.Batcher.AddDispatchGroupDrawEvent"),
+		GAddDispatchGroupDrawEvent,
+		TEXT("Add a draw event marker around each dispatch group."),
+		ECVF_Default
+	);
+
+	#if !WITH_EDITOR
+		constexpr int32 GDebugLogging = 0;
+	#else
+		static int32 GDebugLogging = 0;
+		static FAutoConsoleVariableRef CVarNiagaraDebugLogging(
+			TEXT("fx.Niagara.Batcher.DebugLogging"),
+			GDebugLogging,
+			TEXT("Enables a lot of spew to the log to debug the batcher."),
+			ECVF_Default
+		);
+	#endif
+
+	template<typename TTransitionArrayType>
+	static void AddDataBufferTransitions(TTransitionArrayType& BeforeTransitionArray, TTransitionArrayType& AfterTransitionArray, FNiagaraDataBuffer* DestinationData, ERHIAccess BeforeState = ERHIAccess::SRVMask, ERHIAccess AfterState = ERHIAccess::UAVCompute)
+	{
+		if (FRHIUnorderedAccessView* FloatUAV = DestinationData->GetGPUBufferFloat().UAV )
+		{
+			BeforeTransitionArray.Emplace(FloatUAV, BeforeState, AfterState);
+			AfterTransitionArray.Emplace(FloatUAV, AfterState, BeforeState);
+		}
+		if (FRHIUnorderedAccessView* HalfUAV = DestinationData->GetGPUBufferHalf().UAV)
+		{
+			BeforeTransitionArray.Emplace(HalfUAV, BeforeState, AfterState);
+			AfterTransitionArray.Emplace(HalfUAV, AfterState, BeforeState);
+		}
+		if (FRHIUnorderedAccessView* IntUAV = DestinationData->GetGPUBufferInt().UAV)
+		{
+			BeforeTransitionArray.Emplace(IntUAV, BeforeState, AfterState);
+			AfterTransitionArray.Emplace(IntUAV, AfterState, BeforeState);
+		}
+	}
+}
+
+FFXSystemInterface* FNiagaraGpuComputeDispatch::GetInterface(const FName& InName)
+{
+	return InName == Name ? this : nullptr;
+}
+
+FNiagaraGpuComputeDispatch::FNiagaraGpuComputeDispatch(ERHIFeatureLevel::Type InFeatureLevel, EShaderPlatform InShaderPlatform, FGPUSortManager* InGPUSortManager)
+	: FNiagaraGpuComputeDispatchInterface(InShaderPlatform, InFeatureLevel)
+	, GPUSortManager(InGPUSortManager)
+	, CachedViewRect(0, 0, 64, 64)
+{
+	// Register the batcher callback in the GPUSortManager.
+	// The callback is used to generate the initial keys and values for the GPU sort tasks,
+	// the values being the sorted particle indices used by the Niagara renderers.
+	// The registration also involves defining the list of flags possibly used in GPUSortManager::AddTask()
+	if (GPUSortManager)
+	{
+		GPUSortManager->Register(
+			FGPUSortKeyGenDelegate::CreateLambda(
+				[this](FRHICommandListImmediate& RHICmdList, int32 BatchId, int32 NumElementsInBatch, EGPUSortFlags Flags, FRHIUnorderedAccessView* KeysUAV, FRHIUnorderedAccessView* ValuesUAV)
+				{
+					GenerateSortKeys(RHICmdList, BatchId, NumElementsInBatch, Flags, KeysUAV, ValuesUAV);
+				}
+			),
+			EGPUSortFlags::AnyKeyPrecision | EGPUSortFlags::AnyKeyGenLocation | EGPUSortFlags::AnySortLocation | EGPUSortFlags::ValuesAsInt32,
+			Name
+		);
+
+		if (FNiagaraUtilities::AllowComputeShaders(GetShaderPlatform()))
+		{
+			// Because of culled indirect draw args, we have to update the draw indirect buffer after the sort key generation
+			GPUSortManager->PostPreRenderEvent.AddLambda(
+				[this](FRHICommandListImmediate& RHICmdList)
+				{
+					GPUInstanceCounterManager.UpdateDrawIndirectBuffers(this, RHICmdList, FeatureLevel, ENiagaraGPUCountUpdatePhase::PreOpaque);
+#if WITH_MGPU
+					// For PreInitViews we actually need to broadcast here and not in ExecuteAll since
+					// this is the last place the instance count buffer is written to.
+					if (StageToBroadcastTemporalEffect == ENiagaraGpuComputeTickStage::PreInitViews)
+					{
+						BroadcastTemporalEffect(RHICmdList);
+					}
+#endif // WITH_MGPU
+				}
+			);
+
+			GPUSortManager->PostPostRenderEvent.AddLambda
+			(
+				[this](FRHICommandListImmediate& RHICmdList)
+				{
+					GPUInstanceCounterManager.UpdateDrawIndirectBuffers(this, RHICmdList, FeatureLevel, ENiagaraGPUCountUpdatePhase::PostOpaque);
+				}
+			);
+		}
+	}
+
+	GlobalCBufferLayout  = new FNiagaraRHIUniformBufferLayout(TEXT("Niagara GPU Global CBuffer"),  sizeof(FNiagaraGlobalParameters));
+	SystemCBufferLayout  = new FNiagaraRHIUniformBufferLayout(TEXT("Niagara GPU System CBuffer"),  sizeof(FNiagaraSystemParameters));
+	OwnerCBufferLayout   = new FNiagaraRHIUniformBufferLayout(TEXT("Niagara GPU Owner CBuffer"),   sizeof(FNiagaraOwnerParameters));
+	EmitterCBufferLayout = new FNiagaraRHIUniformBufferLayout(TEXT("Niagara GPU Emitter CBuffer"), sizeof(FNiagaraEmitterParameters));
+
+#if RHI_RAYTRACING
+	RayTracingHelper.Reset(new FNiagaraRayTracingHelper(InShaderPlatform));
+#endif
+
+#if NIAGARA_COMPUTEDEBUG_ENABLED
+	GpuComputeDebugPtr.Reset(new FNiagaraGpuComputeDebug(FeatureLevel));
+#endif
+	GpuReadbackManagerPtr.Reset(new FNiagaraGpuReadbackManager());
+	EmptyUAVPoolPtr.Reset(new FNiagaraEmptyUAVPool());
+
+	FPrimitiveSceneInfo::OnGPUSceneInstancesAllocated.AddRaw(this, &FNiagaraGpuComputeDispatch::OnPrimitiveGPUSceneInstancesAllocated);
+	FPrimitiveSceneInfo::OnGPUSceneInstancesAllocated.AddRaw(this, &FNiagaraGpuComputeDispatch::OnPrimitiveGPUSceneInstancesFreed);
+}
+
+FNiagaraGpuComputeDispatch::~FNiagaraGpuComputeDispatch()
+{
+	FinishDispatches();
+
+#if RHI_RAYTRACING
+	RayTracingHelper->Reset();
+#endif
+
+	GlobalCBufferLayout = nullptr;
+	SystemCBufferLayout = nullptr;
+	OwnerCBufferLayout = nullptr;
+	EmitterCBufferLayout = nullptr;
+
+	FPrimitiveSceneInfo::OnGPUSceneInstancesAllocated.RemoveAll(this);
+	FPrimitiveSceneInfo::OnGPUSceneInstancesFreed.RemoveAll(this);
+}
+
+void FNiagaraGpuComputeDispatch::AddGpuComputeProxy(FNiagaraSystemGpuComputeProxy* ComputeProxy)
+{
+	check(ComputeProxy->ComputeDispatchIndex == INDEX_NONE);
+
+	const ENiagaraGpuComputeTickStage::Type TickStage = ComputeProxy->GetComputeTickStage();
+	ComputeProxy->ComputeDispatchIndex = ProxiesPerStage[TickStage].Num();
+	ProxiesPerStage[TickStage].Add(ComputeProxy);
+
+	NumProxiesThatRequireDistanceFieldData	+= ComputeProxy->RequiresDistanceFieldData() ? 1 : 0;
+	NumProxiesThatRequireDepthBuffer		+= ComputeProxy->RequiresDepthBuffer() ? 1 : 0;
+	NumProxiesThatRequireEarlyViewData		+= ComputeProxy->RequiresEarlyViewData() ? 1 : 0;
+	NumProxiesThatRequireRayTracingScene	+= ComputeProxy->RequiresRayTracingScene() ? 1 : 0;
+}
+
+void FNiagaraGpuComputeDispatch::RemoveGpuComputeProxy(FNiagaraSystemGpuComputeProxy* ComputeProxy)
+{
+	check(ComputeProxy->ComputeDispatchIndex != INDEX_NONE);
+
+	const int32 TickStage = int32(ComputeProxy->GetComputeTickStage());
+	const int32 ProxyIndex = ComputeProxy->ComputeDispatchIndex;
+	check(ProxiesPerStage[TickStage][ProxyIndex] == ComputeProxy);
+
+	ProxiesPerStage[TickStage].RemoveAtSwap(ProxyIndex);
+	if (ProxiesPerStage[TickStage].IsValidIndex(ProxyIndex))
+	{
+		ProxiesPerStage[TickStage][ProxyIndex]->ComputeDispatchIndex = ProxyIndex;
+	}
+	ComputeProxy->ComputeDispatchIndex = INDEX_NONE;
+
+	NumProxiesThatRequireDistanceFieldData	-= ComputeProxy->RequiresDistanceFieldData() ? 1 : 0;
+	NumProxiesThatRequireDepthBuffer		-= ComputeProxy->RequiresDepthBuffer() ? 1 : 0;
+	NumProxiesThatRequireEarlyViewData		-= ComputeProxy->RequiresEarlyViewData() ? 1 : 0;
+	NumProxiesThatRequireRayTracingScene	-= ComputeProxy->RequiresRayTracingScene() ? 1 : 0;
+
+#if NIAGARA_COMPUTEDEBUG_ENABLED
+	if (FNiagaraGpuComputeDebug* GpuComputeDebug = GpuComputeDebugPtr.Get())
+	{
+		GpuComputeDebug->OnSystemDeallocated(ComputeProxy->GetSystemInstanceID());
+	}
+#endif
+#if !UE_BUILD_SHIPPING
+	GpuDebugReadbackInfos.RemoveAll(
+		[&](const FDebugReadbackInfo& Info)
+		{
+			// In the unlikely event we have one in the queue make sure it's marked as complete with no data in it
+			if ( Info.InstanceID == ComputeProxy->GetSystemInstanceID())
+			{
+				Info.DebugInfo->Frame.CopyFromGPUReadback(nullptr, nullptr, nullptr, 0, 0, 0, 0, 0);
+				Info.DebugInfo->bWritten = true;
+			}
+			return Info.InstanceID == ComputeProxy->GetSystemInstanceID();
+		}
+	);
+#endif
+}
+
+void FNiagaraGpuComputeDispatch::BuildConstantBuffers(FNiagaraGPUSystemTick& Tick) const
+{
+	if (!Tick.InstanceCount)
+	{
+		return;
+	}
+
+	TArrayView<FNiagaraComputeInstanceData> InstanceData = Tick.GetInstances();
+
+	// first go through and figure out if we need to support interpolated spawning
+	bool HasInterpolationParameters = false;
+	bool HasMultipleStages = false;
+	for (const FNiagaraComputeInstanceData& Instance : InstanceData)
+	{
+		HasInterpolationParameters = HasInterpolationParameters || Instance.Context->HasInterpolationParameters;
+		HasMultipleStages = HasMultipleStages || Instance.bUsesSimStages;
+	}
+
+	int32 BoundParameterCounts[FNiagaraGPUSystemTick::UBT_NumTypes][2];
+	FMemory::Memzero(BoundParameterCounts);
+
+	for (const FNiagaraComputeInstanceData& Instance : InstanceData)
+	{
+		for (int32 InterpIt = 0; InterpIt < (HasInterpolationParameters ? 2 : 1); ++InterpIt)
+		{
+			BoundParameterCounts[FNiagaraGPUSystemTick::UBT_Global][InterpIt] += Instance.Context->GPUScript_RT->IsGlobalConstantBufferUsed_RenderThread(InterpIt) ? 1 : 0;
+			BoundParameterCounts[FNiagaraGPUSystemTick::UBT_System][InterpIt] += Instance.Context->GPUScript_RT->IsSystemConstantBufferUsed_RenderThread(InterpIt) ? 1 : 0;
+			BoundParameterCounts[FNiagaraGPUSystemTick::UBT_Owner][InterpIt] += Instance.Context->GPUScript_RT->IsOwnerConstantBufferUsed_RenderThread(InterpIt) ? 1 : 0;
+			BoundParameterCounts[FNiagaraGPUSystemTick::UBT_Emitter][InterpIt] += Instance.Context->GPUScript_RT->IsEmitterConstantBufferUsed_RenderThread(InterpIt) ? 1 : 0;
+			BoundParameterCounts[FNiagaraGPUSystemTick::UBT_External][InterpIt] += Instance.Context->GPUScript_RT->IsExternalConstantBufferUsed_RenderThread(InterpIt) ? 1 : 0;
+		}
+	}
+
+	const int32 InterpScale = HasInterpolationParameters ? 2 : 1;
+	const int32 BufferCount = InterpScale * (FNiagaraGPUSystemTick::UBT_NumSystemTypes + FNiagaraGPUSystemTick::UBT_NumInstanceTypes * Tick.InstanceCount);
+
+	Tick.UniformBuffers.Empty(BufferCount);
+
+	const FRHIUniformBufferLayout* SystemLayouts[FNiagaraGPUSystemTick::UBT_NumSystemTypes] =
+	{
+		GlobalCBufferLayout,
+		SystemCBufferLayout,
+		OwnerCBufferLayout
+	};
+
+	for (int32 InterpIt = 0; InterpIt < InterpScale; ++InterpIt)
+	{
+		for (int32 SystemTypeIt = FNiagaraGPUSystemTick::UBT_FirstSystemType; SystemTypeIt < FNiagaraGPUSystemTick::UBT_NumSystemTypes; ++SystemTypeIt)
+		{
+			FUniformBufferRHIRef BufferRef;
+
+			if (BoundParameterCounts[SystemTypeIt][InterpIt])
+			{
+				BufferRef = RHICreateUniformBuffer(
+					Tick.GetUniformBufferSource((FNiagaraGPUSystemTick::EUniformBufferType) SystemTypeIt, nullptr, !InterpIt),
+					SystemLayouts[SystemTypeIt],
+					((BoundParameterCounts[SystemTypeIt][InterpIt] > 1) || HasMultipleStages)
+						? EUniformBufferUsage::UniformBuffer_SingleFrame
+						: EUniformBufferUsage::UniformBuffer_SingleDraw);
+			}
+
+			Tick.UniformBuffers.Add(BufferRef);
+		}
+
+		for (int32 InstanceTypeIt = FNiagaraGPUSystemTick::UBT_FirstInstanceType; InstanceTypeIt < FNiagaraGPUSystemTick::UBT_NumTypes; ++InstanceTypeIt)
+		{
+			for (const FNiagaraComputeInstanceData& Instance : InstanceData)
+			{
+				FUniformBufferRHIRef BufferRef;
+
+				if (BoundParameterCounts[InstanceTypeIt][InterpIt])
+				{
+					FRHIUniformBufferLayout* Layout = InstanceTypeIt == FNiagaraGPUSystemTick::UBT_Emitter ? EmitterCBufferLayout : Instance.Context->ExternalCBufferLayout;
+					if (Layout && (Layout->Resources.Num() > 0 || Layout->ConstantBufferSize > 0))
+					{
+						BufferRef = RHICreateUniformBuffer(
+							Tick.GetUniformBufferSource((FNiagaraGPUSystemTick::EUniformBufferType) InstanceTypeIt, &Instance, !InterpIt),
+							Layout,
+							((BoundParameterCounts[InstanceTypeIt][InterpIt] > 1) || HasMultipleStages)
+							? EUniformBufferUsage::UniformBuffer_SingleFrame
+							: EUniformBufferUsage::UniformBuffer_SingleDraw);
+					}
+				}
+
+				Tick.UniformBuffers.Add(BufferRef);
+			}
+		}
+	}
+}
+
+void FNiagaraGpuComputeDispatch::Tick(UWorld* World, float DeltaTime)
+{
+	check(IsInGameThread());
+	ENQUEUE_RENDER_COMMAND(NiagaraPumpBatcher)(
+		[RT_NiagaraBatcher=this](FRHICommandListImmediate& RHICmdList)
+		{
+			RT_NiagaraBatcher->ProcessPendingTicksFlush(RHICmdList, false);
+			RT_NiagaraBatcher->GetGPUInstanceCounterManager().FlushIndirectArgsPool();
+		}
+	);
+}
+
+void FNiagaraGpuComputeDispatch::FlushPendingTicks_GameThread()
+{
+	check(IsInGameThread());
+	ENQUEUE_RENDER_COMMAND(NiagaraFlushPendingTicks)(
+		[RT_NiagaraBatcher=this](FRHICommandListImmediate& RHICmdList)
+		{
+			RT_NiagaraBatcher->ProcessPendingTicksFlush(RHICmdList, true);
+			RT_NiagaraBatcher->GetGPUInstanceCounterManager().FlushIndirectArgsPool();
+		}
+	);
+}
+
+void FNiagaraGpuComputeDispatch::ProcessPendingTicksFlush(FRHICommandListImmediate& RHICmdList, bool bForceFlush)
+{
+	// No ticks are pending
+	bool bHasTicks = false;
+	for ( int iTickStage=0; iTickStage < ENiagaraGpuComputeTickStage::Max; ++iTickStage)
+	{
+		if ( ProxiesPerStage[iTickStage].Num() > 0 )
+		{
+			bHasTicks = true;
+			break;
+		}
+	}
+
+	if ( !bHasTicks )
+	{
+		return;
+	}
+
+	// We have pending ticks increment our counter, once we cross the threshold we will perform the appropriate operation
+	++FramesBeforeTickFlush;
+	if (!bForceFlush && (FramesBeforeTickFlush < uint32(FNiagaraGpuComputeDispatchLocal::GTickFlushMaxQueuedFrames)) )
+	{
+		return;
+	}
+	FramesBeforeTickFlush = 0;
+
+	switch (FNiagaraGpuComputeDispatchLocal::GTickFlushMode)
+	{
+		// Do nothing
+		default:
+		case 0:
+		{
+			//UE_LOG(LogNiagara, Log, TEXT("FNiagaraGpuComputeDispatch: Queued ticks (%d) are building up, this may cause a stall when released."), Ticks_RT.Num());
+			break;
+		}
+
+		// Process all the pending ticks that have built up
+		case 1:
+		{
+			//UE_LOG(LogNiagara, Log, TEXT("FNiagaraGpuComputeDispatch: Queued ticks are being Processed due to not rendering.  This may result in undesirable simulation artifacts."));
+
+			// Make a temporary ViewInfo
+			//-TODO: We could gather some more information here perhaps?
+			FMemMark Mark(FMemStack::Get());
+
+			FSceneViewFamily ViewFamily(FSceneViewFamily::ConstructionValues(nullptr, nullptr, FEngineShowFlags(ESFIM_Game))
+				.SetWorldTimes(0, 0, 0)
+				.SetGammaCorrection(1.0f));
+
+			FSceneViewInitOptions ViewInitOptions;
+			ViewInitOptions.ViewFamily = &ViewFamily;
+			ViewInitOptions.SetViewRectangle(CachedViewRect);
+			ViewInitOptions.ViewOrigin = FVector::ZeroVector;
+			ViewInitOptions.ViewRotationMatrix = FMatrix::Identity;
+			ViewInitOptions.ProjectionMatrix = FMatrix::Identity;
+
+			FViewInfo* DummyView = new(FMemStack::Get().Alloc(sizeof(FViewInfo), alignof(FViewInfo))) FViewInfo(ViewInitOptions);
+
+			DummyView->ViewRect = DummyView->UnscaledViewRect;
+			DummyView->CachedViewUniformShaderParameters = MakeUnique<FViewUniformShaderParameters>();
+
+			FBox UnusedVolumeBounds[TVC_MAX];
+			DummyView->SetupUniformBufferParameters(UnusedVolumeBounds, TVC_MAX, *DummyView->CachedViewUniformShaderParameters);
+
+			DummyView->ViewUniformBuffer = TUniformBufferRef<FViewUniformShaderParameters>::CreateUniformBufferImmediate(*DummyView->CachedViewUniformShaderParameters, UniformBuffer_SingleFrame);
+
+			TConstArrayView<FViewInfo> DummyViews = MakeArrayView(DummyView, 1);
+			const bool bAllowGPUParticleUpdate = true;
+			
+			// Execute all ticks that we can support without invalid simulations
+			FRDGBuilder GraphBuilder(RHICmdList);
+			CreateSystemTextures(GraphBuilder);
+			PreInitViews(GraphBuilder, bAllowGPUParticleUpdate);
+			GPUInstanceCounterManager.UpdateDrawIndirectBuffers(this, RHICmdList, FeatureLevel, ENiagaraGPUCountUpdatePhase::PreOpaque);
+			PostInitViews(GraphBuilder, DummyViews, bAllowGPUParticleUpdate);
+			PostRenderOpaque(GraphBuilder, DummyViews, bAllowGPUParticleUpdate);
+			GPUInstanceCounterManager.UpdateDrawIndirectBuffers(this, RHICmdList, FeatureLevel, ENiagaraGPUCountUpdatePhase::PostOpaque);
+			GraphBuilder.Execute();
+			break;
+		}
+
+		// Kill all the pending ticks that have built up
+		case 2:
+		{
+			//UE_LOG(LogNiagara, Log, TEXT("FNiagaraGpuComputeDispatch: Queued ticks are being Destroyed due to not rendering.  This may result in undesirable simulation artifacts."));
+
+			FinishDispatches();
+
+#if RHI_RAYTRACING
+			RayTracingHelper->Reset();
+#endif
+			break;
+		}
+	}
+}
+
+void FNiagaraGpuComputeDispatch::FinishDispatches()
+{
+	check(IsInRenderingThread());
+
+	for (int iTickStage = 0; iTickStage < ENiagaraGpuComputeTickStage::Max; ++iTickStage)
+	{
+		for (FNiagaraSystemGpuComputeProxy* ComputeProxy : ProxiesPerStage[iTickStage])
+		{
+			ComputeProxy->ReleaseTicks(GPUInstanceCounterManager);
+		}
+	}
+
+	for (FNiagaraGpuDispatchList& DispatchList : DispatchListPerStage)
+	{
+		DispatchList.DispatchGroups.Empty();
+		if (DispatchList.CountsToRelease.Num() > 0)
+		{
+			GPUInstanceCounterManager.FreeEntryArray(DispatchList.CountsToRelease);
+			DispatchList.CountsToRelease.Empty();
+		}
+	}
+}
+
+void FNiagaraGpuComputeDispatch::ResetDataInterfaces(FRHICommandList& RHICmdList, const FNiagaraGPUSystemTick& Tick, const FNiagaraComputeInstanceData& InstanceData) const
+{
+	// Note: All stages will contain the same bindings so if they are valid for one they are valid for all, this could change in the future
+	const FNiagaraShaderRef& ComputeShader = InstanceData.Context->GPUScript_RT->GetShader(0);
+	const TMemoryImageArray<FNiagaraDataInterfaceParamRef>& DIParameters = ComputeShader->GetDIParameters();
+
+	uint32 InterfaceIndex = 0;
+	for (FNiagaraDataInterfaceProxy* Interface : InstanceData.DataInterfaceProxies)
+	{
+		const FNiagaraDataInterfaceParamRef& DIParam = DIParameters[InterfaceIndex];
+		if (DIParam.Parameters.IsValid())
+		{
+			const FNiagaraDataInterfaceArgs TmpContext(Interface, Tick.SystemInstanceID, this);
+			Interface->ResetData(RHICmdList, TmpContext);
+		}
+		InterfaceIndex++;
+	}
+}
+
+FNiagaraDataInterfaceProxyRW* FNiagaraGpuComputeDispatch::FindIterationInterface(FNiagaraComputeInstanceData* Instance, const uint32 SimulationStageIndex) const
+{
+	// Determine if the iteration is outputting to a custom data size
+	return Instance->FindIterationInterface(SimulationStageIndex);
+}
+
+void FNiagaraGpuComputeDispatch::PreStageInterface(FRHICommandList& RHICmdList, const FNiagaraGPUSystemTick& Tick, const FNiagaraComputeInstanceData& InstanceData, const FNiagaraSimStageData& SimStageData, TSet<FNiagaraDataInterfaceProxy*>& ProxiesToFinalize) const
+{
+	// Note: All stages will contain the same bindings so if they are valid for one they are valid for all, this could change in the future
+	const FNiagaraShaderRef& ComputeShader = InstanceData.Context->GPUScript_RT->GetShader(0);
+	const TMemoryImageArray<FNiagaraDataInterfaceParamRef>& DIParameters = ComputeShader->GetDIParameters();
+
+	uint32 InterfaceIndex = 0;
+	for (FNiagaraDataInterfaceProxy* Interface : InstanceData.DataInterfaceProxies)
+	{
+		const FNiagaraDataInterfaceParamRef& DIParam = DIParameters[InterfaceIndex];
+		if (DIParam.Parameters.IsValid())
+		{
+			const FNiagaraDataInterfaceStageArgs TmpContext(Interface, Tick.SystemInstanceID, this, &InstanceData, &SimStageData, InstanceData.IsOutputStage(Interface, SimStageData.StageIndex), InstanceData.IsIterationStage(Interface, SimStageData.StageIndex));
+			Interface->PreStage(RHICmdList, TmpContext);
+
+			if (Interface->RequiresPreStageFinalize())
+			{
+				ProxiesToFinalize.Add(Interface);
+			}
+		}
+		InterfaceIndex++;
+	}
+}
+
+void FNiagaraGpuComputeDispatch::PostStageInterface(FRHICommandList& RHICmdList, const FNiagaraGPUSystemTick& Tick, const FNiagaraComputeInstanceData& InstanceData, const FNiagaraSimStageData& SimStageData, TSet<FNiagaraDataInterfaceProxy*>& ProxiesToFinalize) const
+{
+	// Note: All stages will contain the same bindings so if they are valid for one they are valid for all, this could change in the future
+	const FNiagaraShaderRef& ComputeShader = InstanceData.Context->GPUScript_RT->GetShader(0);
+	const TMemoryImageArray<FNiagaraDataInterfaceParamRef>& DIParameters = ComputeShader->GetDIParameters();
+
+	uint32 InterfaceIndex = 0;
+	for (FNiagaraDataInterfaceProxy* Interface : InstanceData.DataInterfaceProxies)
+	{
+		const FNiagaraDataInterfaceParamRef& DIParam = DIParameters[InterfaceIndex];
+		if (DIParam.Parameters.IsValid())
+		{
+			const FNiagaraDataInterfaceStageArgs TmpContext(Interface, Tick.SystemInstanceID, this, &InstanceData, &SimStageData, InstanceData.IsOutputStage(Interface, SimStageData.StageIndex), InstanceData.IsIterationStage(Interface, SimStageData.StageIndex));
+			Interface->PostStage(RHICmdList, TmpContext);
+
+			if (Interface->RequiresPostStageFinalize())
+			{
+				ProxiesToFinalize.Add(Interface);
+			}
+		}
+		InterfaceIndex++;
+	}
+}
+
+void FNiagaraGpuComputeDispatch::PostSimulateInterface(FRHICommandList& RHICmdList, const FNiagaraGPUSystemTick& Tick, const FNiagaraComputeInstanceData& InstanceData) const
+{
+	// Note: All stages will contain the same bindings so if they are valid for one they are valid for all, this could change in the future
+	const FNiagaraShaderRef& ComputeShader = InstanceData.Context->GPUScript_RT->GetShader(0);
+	const TMemoryImageArray<FNiagaraDataInterfaceParamRef>& DIParameters = ComputeShader->GetDIParameters();
+
+	uint32 InterfaceIndex = 0;
+	for (FNiagaraDataInterfaceProxy* Interface : InstanceData.DataInterfaceProxies)
+	{
+		const FNiagaraDataInterfaceParamRef& DIParam = DIParameters[InterfaceIndex];
+		if (DIParam.Parameters.IsValid())
+		{
+			const FNiagaraDataInterfaceArgs TmpContext(Interface, Tick.SystemInstanceID, this);
+			Interface->PostSimulate(RHICmdList, TmpContext);
+		}
+		InterfaceIndex++;
+	}
+}
+
+
+void FNiagaraGpuComputeDispatch::UpdateFreeIDsListSizesBuffer(FRHICommandList& RHICmdList, uint32 NumInstances)
+{
+	if (NumInstances > NumAllocatedFreeIDListSizes)
+	{
+		constexpr uint32 ALLOC_CHUNK_SIZE = 128;
+		NumAllocatedFreeIDListSizes = Align(NumInstances, ALLOC_CHUNK_SIZE);
+		if (FreeIDListSizesBuffer.Buffer)
+		{
+			FreeIDListSizesBuffer.Release();
+		}
+		FreeIDListSizesBuffer.Initialize(TEXT("NiagaraFreeIDListSizes"), sizeof(uint32), NumAllocatedFreeIDListSizes, EPixelFormat::PF_R32_SINT, ERHIAccess::UAVCompute, BUF_Static);
+	}
+
+	{
+		SCOPED_DRAW_EVENT(RHICmdList, NiagaraGPUComputeClearFreeIDListSizes);
+		RHICmdList.Transition(FRHITransitionInfo(FreeIDListSizesBuffer.UAV, ERHIAccess::UAVCompute, ERHIAccess::UAVCompute));
+		NiagaraFillGPUIntBuffer(RHICmdList, FeatureLevel, FreeIDListSizesBuffer, 0);
+	}
+}
+
+void FNiagaraGpuComputeDispatch::UpdateFreeIDBuffers(FRHICommandList& RHICmdList, TConstArrayView<FNiagaraComputeExecutionContext*> Instances)
+{
+	if (Instances.Num() == 0)
+	{
+		return;
+	}
+
+	SCOPED_DRAW_EVENT(RHICmdList, NiagaraGPUComputeFreeIDs);
+	SCOPED_GPU_STAT(RHICmdList, NiagaraGPUComputeFreeIDs);
+
+	TArray<FRHITransitionInfo, TMemStackAllocator<>> TransitionsBefore;
+	TransitionsBefore.Emplace(FreeIDListSizesBuffer.UAV, ERHIAccess::UAVCompute, ERHIAccess::UAVCompute);
+	RHICmdList.Transition(TransitionsBefore);
+
+	check((uint32)Instances.Num() <= NumAllocatedFreeIDListSizes);
+
+	RHICmdList.BeginUAVOverlap(FreeIDListSizesBuffer.UAV);
+	for (uint32 iInstance = 0; iInstance < (uint32)Instances.Num(); ++iInstance)
+	{
+		FNiagaraComputeExecutionContext* ComputeContext = Instances[iInstance];
+		FNiagaraDataSet* MainDataSet = ComputeContext->MainDataSet;
+		FNiagaraDataBuffer* CurrentData = ComputeContext->MainDataSet->GetCurrentData();
+
+		SCOPED_DRAW_EVENTF(RHICmdList, NiagaraGPUComputeFreeIDsEmitter, TEXT("Update Free ID Buffer - %s"), ComputeContext->GetDebugSimName());
+		NiagaraComputeGPUFreeIDs(RHICmdList, FeatureLevel, MainDataSet->GetGPUNumAllocatedIDs(), CurrentData->GetGPUIDToIndexTable().SRV, MainDataSet->GetGPUFreeIDs(), FreeIDListSizesBuffer, iInstance);
+	}
+	RHICmdList.EndUAVOverlap(FreeIDListSizesBuffer.UAV);
+}
+
+void FNiagaraGpuComputeDispatch::DumpDebugFrame()
+{
+	// Anything doing?
+	bool bHasAnyWork = false;
+	for (int iTickStage = 0; iTickStage < ENiagaraGpuComputeTickStage::Max; ++iTickStage)
+	{
+		bHasAnyWork |= DispatchListPerStage[iTickStage].HasWork();
+	}
+	if ( !bHasAnyWork )
+	{
+		return;
+	}
+
+	// Dump Frame
+	UE_LOG(LogNiagara, Warning, TEXT("====== BatcherFrame(%d)"), GFrameNumberRenderThread);
+
+	for (int iTickStage = 0; iTickStage < ENiagaraGpuComputeTickStage::Max; ++iTickStage)
+	{
+		if (!DispatchListPerStage[iTickStage].HasWork())
+		{
+			continue;
+		}
+
+		FNiagaraGpuDispatchList& DispatchList = DispatchListPerStage[iTickStage];
+		UE_LOG(LogNiagara, Warning, TEXT("==== TickStage(%d) TotalGroups(%d)"), iTickStage, DispatchList.DispatchGroups.Num());
+
+		for ( int iDispatchGroup=0; iDispatchGroup < DispatchList.DispatchGroups.Num(); ++iDispatchGroup )
+		{
+			const FNiagaraGpuDispatchGroup& DispatchGroup = DispatchListPerStage[iTickStage].DispatchGroups[iDispatchGroup];
+
+			if (DispatchGroup.TicksWithPerInstanceData.Num() > 0)
+			{
+				UE_LOG(LogNiagara, Warning, TEXT("====== TicksWithPerInstanceData(%s)"), DispatchGroup.TicksWithPerInstanceData.Num());
+				for (FNiagaraGPUSystemTick* Tick : DispatchGroup.TicksWithPerInstanceData)
+				{
+					for (const auto& Pair : Tick->DIInstanceData->InterfaceProxiesToOffsets)
+					{
+						FNiagaraDataInterfaceProxy* Proxy = Pair.Key;
+						UE_LOG(LogNiagara, Warning, TEXT("Proxy(%s)"), *Proxy->SourceDIName.ToString());
+					}
+				}
+			}
+
+			UE_LOG(LogNiagara, Warning, TEXT("====== DispatchGroup(%d)"), iDispatchGroup);
+			for ( const FNiagaraGpuDispatchInstance& DispatchInstance : DispatchGroup.DispatchInstances )
+			{
+				const FNiagaraSimStageData& SimStageData = DispatchInstance.SimStageData;
+				const FNiagaraComputeInstanceData& InstanceData = DispatchInstance.InstanceData;
+
+				TStringBuilder<512> Builder;
+				Builder.Appendf(TEXT("Proxy(%p) "), DispatchInstance.Tick.SystemGpuComputeProxy);
+				Builder.Appendf(TEXT("ComputeContext(%p) "), InstanceData.Context);
+				Builder.Appendf(TEXT("Emitter(%s) "), InstanceData.Context->GetDebugSimName());
+				Builder.Appendf(TEXT("Stage(%d | %s) "), SimStageData.StageIndex, *SimStageData.StageMetaData->SimulationStageName.ToString());
+
+				if (InstanceData.bResetData)
+				{
+					Builder.Append(TEXT("ResetData "));
+				}
+
+				if (InstanceData.Context->MainDataSet->RequiresPersistentIDs())
+				{
+					Builder.Append(TEXT("HasPersistentIDs "));
+				}
+
+				if ( DispatchInstance.SimStageData.bFirstStage )
+				{
+					Builder.Append(TEXT("FirstStage "));
+				}
+
+				if ( DispatchInstance.SimStageData.bLastStage )
+				{
+					Builder.Append(TEXT("LastStage "));
+				}
+
+				if (DispatchInstance.SimStageData.bSetDataToRender)
+				{
+					Builder.Append(TEXT("SetDataToRender "));
+				}
+					
+
+				if (InstanceData.Context->EmitterInstanceReadback.GPUCountOffset != INDEX_NONE)
+				{
+					if (InstanceData.Context->EmitterInstanceReadback.GPUCountOffset == SimStageData.SourceCountOffset)
+					{
+						Builder.Appendf(TEXT("ReadbackSource(%d) "), InstanceData.Context->EmitterInstanceReadback.CPUCount);
+					}
+				}
+				Builder.Appendf(TEXT("Source(%p 0x%08x %d) "), SimStageData.Source, SimStageData.SourceCountOffset, SimStageData.SourceNumInstances);
+				Builder.Appendf(TEXT("Destination(%p 0x%08x %d) "), SimStageData.Destination, SimStageData.DestinationCountOffset, SimStageData.DestinationNumInstances);
+				Builder.Appendf(TEXT("Iteration(%d | %s) "), SimStageData.IterationIndex, SimStageData.AlternateIterationSource ? *SimStageData.AlternateIterationSource->SourceDIName.ToString() : TEXT("Particles"));
+				UE_LOG(LogNiagara, Warning, TEXT("%s"), Builder.ToString());
+			}
+
+			if (DispatchGroup.FreeIDUpdates.Num() > 0)
+			{
+				UE_LOG(LogNiagara, Warning, TEXT("====== FreeIDUpdates"));
+				for (FNiagaraComputeExecutionContext* ComputeContext : DispatchGroup.FreeIDUpdates)
+				{
+					UE_LOG(LogNiagara, Warning, TEXT("ComputeContext(%p) Emitter(%s)"), ComputeContext, ComputeContext->GetDebugSimName());
+				}
+			}
+		}
+		if (DispatchList.CountsToRelease.Num() > 0)
+		{
+			UE_LOG(LogNiagara, Warning, TEXT("====== CountsToRelease"));
+
+			const int NumPerLine = 16;
+
+			TStringBuilder<512> StringBuilder;
+			for ( int32 i=0; i < DispatchList.CountsToRelease.Num(); ++i )
+			{
+				const bool bFirst = (i % NumPerLine) == 0;
+				const bool bLast = ((i % NumPerLine) == NumPerLine - 1) || (i == DispatchList.CountsToRelease.Num() -1);
+
+				if ( !bFirst )
+				{
+					StringBuilder.Append(TEXT(", "));
+				}
+				StringBuilder.Appendf(TEXT("0x%08x"), DispatchList.CountsToRelease[i]);
+
+				if ( bLast )
+				{
+					UE_LOG(LogNiagara, Warning, TEXT("%s"), StringBuilder.ToString());
+					StringBuilder.Reset();
+				}
+			}
+		}
+	}
+}
+
+void FNiagaraGpuComputeDispatch::UpdateInstanceCountManager(FRHICommandListImmediate& RHICmdList)
+{
+	TRACE_CPUPROFILER_EVENT_SCOPE(FNiagaraGpuComputeDispatch_UpdateInstanceCountManager);
+
+	// Resize dispatch buffer count
+	//-OPT: No need to iterate over all the ticks, we can store this as ticks are queued
+	{
+		int32 TotalDispatchCount = 0;
+		for (int iTickStage = 0; iTickStage < ENiagaraGpuComputeTickStage::Max; ++iTickStage)
+		{
+			for (FNiagaraSystemGpuComputeProxy* ComputeProxy : ProxiesPerStage[iTickStage])
+			{
+				for (FNiagaraGPUSystemTick& Tick : ComputeProxy->PendingTicks)
+				{
+					TotalDispatchCount += (int32)Tick.TotalDispatches;
+
+					for ( FNiagaraComputeInstanceData& InstanceData : Tick.GetInstances() )
+					{
+						if (InstanceData.bResetData)
+						{
+							InstanceData.Context->EmitterInstanceReadback.GPUCountOffset = INDEX_NONE;
+						}
+					}
+				}
+			}
+		}
+		GPUInstanceCounterManager.ResizeBuffers(RHICmdList, FeatureLevel, TotalDispatchCount);
+	}
+
+	// Consume any pending readbacks that are ready
+	{
+		SCOPE_CYCLE_COUNTER(STAT_NiagaraGPUReadback_RT);
+		if ( const uint32* Counts = GPUInstanceCounterManager.GetGPUReadback() )
+		{
+			if (FNiagaraGpuComputeDispatchLocal::GDebugLogging)
+			{
+				UE_LOG(LogNiagara, Warning, TEXT("====== BatcherFrame(%d) Readback Complete"), GFrameNumberRenderThread);
+			}
+
+			for (int iTickStage=0; iTickStage < ENiagaraGpuComputeTickStage::Max; ++iTickStage)
+			{
+				for (FNiagaraSystemGpuComputeProxy* ComputeProxy : ProxiesPerStage[iTickStage])
+				{
+					for ( FNiagaraComputeExecutionContext* ComputeContext : ComputeProxy->ComputeContexts )
+					{
+						if (ComputeContext->EmitterInstanceReadback.GPUCountOffset == INDEX_NONE )
+						{
+							continue;
+						}
+
+						const uint32 DeadInstanceCount = ComputeContext->EmitterInstanceReadback.CPUCount - Counts[ComputeContext->EmitterInstanceReadback.GPUCountOffset];
+						if ( DeadInstanceCount <= ComputeContext->CurrentNumInstances_RT )
+						{
+							ComputeContext->CurrentNumInstances_RT -= DeadInstanceCount;
+						}
+						if (FNiagaraGpuComputeDispatchLocal::GDebugLogging)
+						{
+							UE_LOG(LogNiagara, Warning, TEXT("ComputeContext(%p) Emitter(%s) DeadInstances(%d) CountReleased(0x%08x)"), ComputeContext, ComputeContext->GetDebugSimName(), DeadInstanceCount, ComputeContext->EmitterInstanceReadback.GPUCountOffset);
+						}
+
+						// Readback complete
+						ComputeContext->EmitterInstanceReadback.GPUCountOffset = INDEX_NONE;
+					}
+				}
+			}
+
+			// Release the readback buffer
+			GPUInstanceCounterManager.ReleaseGPUReadback();
+		}
+	}
+}
+
+void FNiagaraGpuComputeDispatch::PrepareTicksForProxy(FRHICommandListImmediate& RHICmdList, FNiagaraSystemGpuComputeProxy* ComputeProxy, FNiagaraGpuDispatchList& GpuDispatchList)
+{
+	for (FNiagaraComputeExecutionContext* ComputeContext : ComputeProxy->ComputeContexts)
+	{
+		ComputeContext->CurrentMaxInstances_RT = 0;
+		ComputeContext->CurrentMaxAllocateInstances_RT = 0;
+		ComputeContext->BufferSwapsThisFrame_RT = 0;
+	}
+
+	if (ComputeProxy->PendingTicks.Num() == 0)
+	{
+		return;
+	}
+
+	const bool bEnqueueCountReadback = !GPUInstanceCounterManager.HasPendingGPUReadback();
+
+	// Set final tick flag
+	ComputeProxy->PendingTicks.Last().bIsFinalTick = true;
+
+	// Process ticks
+	int32 iTickStartDispatchGroup = 0;
+
+	for (FNiagaraGPUSystemTick& Tick : ComputeProxy->PendingTicks)
+	{
+		int32 iInstanceStartDispatchGroup = iTickStartDispatchGroup;
+		int32 iInstanceCurrDispatchGroup = iTickStartDispatchGroup;
+		bool bHasFreeIDUpdates = false;
+
+		// Track that we need to consume per instance data before executing the ticks
+		//if (Tick.DIInstanceData)
+		//{
+		//	GpuDispatchList.PreAllocateGroups(iTickStartDispatchGroup + 1);
+		//	GpuDispatchList.DispatchGroups[iTickStartDispatchGroup].TicksWithPerInstanceData.Add(&Tick);
+		//}
+
+		// Iterate over all instances preparing our number of instances
+		for (FNiagaraComputeInstanceData& InstanceData : Tick.GetInstances())
+		{
+			FNiagaraComputeExecutionContext* ComputeContext = InstanceData.Context;
+
+			// Instance requires a reset?
+			if (InstanceData.bResetData)
+			{
+				ComputeContext->CurrentNumInstances_RT = 0;
+				if (ComputeContext->CountOffset_RT != INDEX_NONE)
+				{
+					GpuDispatchList.CountsToRelease.Add(ComputeContext->CountOffset_RT);
+					ComputeContext->CountOffset_RT = INDEX_NONE;
+				}
+			}
+
+			// If shader is not ready don't do anything
+			if (!ComputeContext->GPUScript_RT->IsShaderMapComplete_RenderThread())
+			{
+				continue;
+			}
+
+#if WITH_EDITOR
+			//-TODO: Temporary change to avoid crashing on incorrect FeatureLevel
+			if (ComputeContext->GPUScript_RT->GetFeatureLevel() != FeatureLevel)
+			{
+				continue;
+			}
+#endif
+
+			// Nothing to dispatch?
+			if ( InstanceData.TotalDispatches == 0 )
+			{
+				continue;
+			}
+
+			// Determine this instances start dispatch group, in the case of emitter dependencies (i.e. particle reads) we need to continue rather than starting again
+			iInstanceStartDispatchGroup = InstanceData.bStartNewOverlapGroup ? iInstanceCurrDispatchGroup : iInstanceStartDispatchGroup;
+			iInstanceCurrDispatchGroup = iInstanceStartDispatchGroup;
+
+			// Pre-allocator groups
+			GpuDispatchList.PreAllocateGroups(iInstanceCurrDispatchGroup + InstanceData.TotalDispatches);
+
+			// Calculate instance counts
+			const uint32 MaxBufferInstances = ComputeContext->MainDataSet->GetMaxInstanceCount();
+			const uint32 PrevNumInstances = ComputeContext->CurrentNumInstances_RT;
+
+			ComputeContext->CurrentNumInstances_RT = PrevNumInstances + InstanceData.SpawnInfo.SpawnRateInstances + InstanceData.SpawnInfo.EventSpawnTotal;
+			ComputeContext->CurrentNumInstances_RT = FMath::Min(ComputeContext->CurrentNumInstances_RT, MaxBufferInstances);
+
+			// Calculate new maximum count
+			ComputeContext->CurrentMaxInstances_RT = FMath::Max(ComputeContext->CurrentMaxInstances_RT, ComputeContext->CurrentNumInstances_RT);
+
+			if (!GNiagaraBatcherFreeBufferEarly || (ComputeContext->CurrentMaxInstances_RT > 0))
+			{
+				ComputeContext->CurrentMaxAllocateInstances_RT = FMath::Max3(ComputeContext->CurrentMaxAllocateInstances_RT, ComputeContext->CurrentMaxInstances_RT, InstanceData.SpawnInfo.MaxParticleCount);
+			}
+			else
+			{
+				ComputeContext->CurrentMaxAllocateInstances_RT = FMath::Max(ComputeContext->CurrentMaxAllocateInstances_RT, ComputeContext->CurrentMaxInstances_RT);
+			}
+
+			bHasFreeIDUpdates |= ComputeContext->MainDataSet->RequiresPersistentIDs();
+
+			//-OPT: Do we need this test?  Can remove in favor of MaxUpdateIterations
+			bool bFirstStage = true;
+			for (int32 SimStageIndex=0; SimStageIndex < ComputeContext->SimStageInfo.Num(); ++SimStageIndex)
+			{
+				const FSimulationStageMetaData& SimStageMetaData = ComputeContext->SimStageInfo[SimStageIndex];
+				if (InstanceData.NumIterationsPerStage[SimStageIndex] == 0)
+				{
+					continue;
+				}
+
+				FNiagaraDataInterfaceProxyRW* IterationInterface = InstanceData.FindIterationInterface(SimStageIndex);
+				for ( int32 IterationIndex=0; IterationIndex < InstanceData.NumIterationsPerStage[SimStageIndex]; ++IterationIndex )
+				{
+					// Build SimStage data
+					FNiagaraGpuDispatchGroup& DispatchGroup = GpuDispatchList.DispatchGroups[iInstanceCurrDispatchGroup++];
+					FNiagaraGpuDispatchInstance& DispatchInstance = DispatchGroup.DispatchInstances.Emplace_GetRef(Tick, InstanceData);
+					FNiagaraSimStageData& SimStageData = DispatchInstance.SimStageData;
+					SimStageData.bFirstStage = bFirstStage;
+					SimStageData.StageIndex = SimStageIndex;
+					SimStageData.IterationIndex = IterationIndex;
+					SimStageData.StageMetaData = &SimStageMetaData;
+					SimStageData.AlternateIterationSource = IterationInterface;
+
+					bFirstStage = false;
+
+					FNiagaraDataBuffer* SourceData = ComputeContext->bHasTickedThisFrame_RT ? ComputeContext->GetPrevDataBuffer() : ComputeContext->MainDataSet->GetCurrentData();
+
+					// This stage does not modify particle data, i.e. read only or not related to particles at all
+					if (!SimStageData.StageMetaData->bWritesParticles)
+					{
+						SimStageData.Source = SourceData;
+						SimStageData.SourceCountOffset = ComputeContext->CountOffset_RT;
+						SimStageData.SourceNumInstances = ComputeContext->CurrentNumInstances_RT;
+						SimStageData.Destination = nullptr;
+						SimStageData.DestinationCountOffset = ComputeContext->CountOffset_RT;
+						SimStageData.DestinationNumInstances = ComputeContext->CurrentNumInstances_RT;
+					}
+					// This stage writes particles but will not kill any, we can use the buffer as both source and destination
+					else if (SimStageData.StageMetaData->bPartialParticleUpdate)
+					{
+						SimStageData.Source = nullptr;
+						SimStageData.SourceCountOffset = ComputeContext->CountOffset_RT;
+						SimStageData.SourceNumInstances = ComputeContext->CurrentNumInstances_RT;
+						SimStageData.Destination = SourceData;
+						SimStageData.DestinationCountOffset = ComputeContext->CountOffset_RT;
+						SimStageData.DestinationNumInstances = ComputeContext->CurrentNumInstances_RT;
+					}
+					// This stage may kill particles, we need to allocate a new destination buffer
+					else
+					{
+						SimStageData.Source = SourceData;
+						SimStageData.SourceCountOffset = ComputeContext->CountOffset_RT;
+						//-TODO: This is a little odd, perhaps we need to change the preallocate
+						SimStageData.SourceNumInstances = SimStageIndex == 0 && IterationIndex == 0 ? PrevNumInstances : ComputeContext->CurrentNumInstances_RT;
+						SimStageData.Destination = ComputeContext->GetNextDataBuffer();
+						SimStageData.DestinationCountOffset = GPUInstanceCounterManager.AcquireEntry();
+						SimStageData.DestinationNumInstances = ComputeContext->CurrentNumInstances_RT;
+
+						ComputeContext->AdvanceDataBuffer();
+						ComputeContext->CountOffset_RT = SimStageData.DestinationCountOffset;
+						ComputeContext->bHasTickedThisFrame_RT = true;
+
+						// If we are the last tick then we may want to enqueue for a readback
+						// Note: Do not pull count from SimStageData as a reset tick will be INDEX_NONE
+						check(SimStageData.SourceCountOffset != INDEX_NONE || SimStageData.SourceNumInstances == 0);
+						if (SimStageData.SourceCountOffset != INDEX_NONE)
+						{
+							if ( bEnqueueCountReadback && Tick.bIsFinalTick && (ComputeContext->EmitterInstanceReadback.GPUCountOffset == INDEX_NONE))
+							{
+								bRequiresReadback = true;
+								ComputeContext->EmitterInstanceReadback.CPUCount = SimStageData.SourceNumInstances;
+								ComputeContext->EmitterInstanceReadback.GPUCountOffset = SimStageData.SourceCountOffset;
+							}
+							GpuDispatchList.CountsToRelease.Add(SimStageData.SourceCountOffset);
+						}
+					}
+				}
+			}
+
+			// The final instance in the last group we modified is the final stage for this tick
+			FNiagaraGpuDispatchGroup& FinalDispatchGroup = GpuDispatchList.DispatchGroups[iInstanceCurrDispatchGroup - 1];
+			FinalDispatchGroup.DispatchInstances.Last().SimStageData.bLastStage = true;
+			FinalDispatchGroup.DispatchInstances.Last().SimStageData.bSetDataToRender = Tick.bIsFinalTick;
+
+			// Keep track of where the next set of dispatch should occur
+			iTickStartDispatchGroup = FMath::Max(iTickStartDispatchGroup, iInstanceCurrDispatchGroup);
+		}
+
+		// Accumulate Free ID updates
+		// Note: These must be done at the end of the tick due to the way spawned instances read from the free list
+		if (bHasFreeIDUpdates)
+		{
+			FNiagaraGpuDispatchGroup& FinalDispatchGroup = GpuDispatchList.DispatchGroups[iInstanceCurrDispatchGroup - 1];
+			for (FNiagaraComputeInstanceData& InstanceData : Tick.GetInstances())
+			{
+				FNiagaraComputeExecutionContext* ComputeContext = InstanceData.Context;
+				if (!ComputeContext->GPUScript_RT->IsShaderMapComplete_RenderThread())
+				{
+					continue;
+				}
+
+				if ( ComputeContext->MainDataSet->RequiresPersistentIDs() )
+				{
+					FinalDispatchGroup.FreeIDUpdates.Add(ComputeContext);
+				}
+			}
+		}
+
+		// Build constant buffers for tick
+		BuildConstantBuffers(Tick);
+	}
+
+	// Now that all ticks have been processed we can adjust our output buffers to the correct size
+	// We will also set the translucent data to render, i.e. this frames data.
+	for (FNiagaraComputeExecutionContext* ComputeContext : ComputeProxy->ComputeContexts)
+	{
+		if (!ComputeContext->bHasTickedThisFrame_RT)
+		{
+			continue;
+		}
+
+		// We need to store the current data from the main data set as we will be temporarily stomping it during multi-ticking
+		ComputeContext->DataSetOriginalBuffer_RT = ComputeContext->MainDataSet->GetCurrentData();
+
+		//-OPT: We should allocate all GPU free IDs together since they require a transition
+		if (ComputeContext->MainDataSet->RequiresPersistentIDs())
+		{
+			ComputeContext->MainDataSet->AllocateGPUFreeIDs(ComputeContext->CurrentMaxAllocateInstances_RT + 1, RHICmdList, FeatureLevel, ComputeContext->GetDebugSimName());
+		}
+
+		// Allocate space for the buffers we need to perform ticking.  In cases of multiple ticks or multiple write stages we need 3 buffers (current rendered and two simulation buffers).
+		//-OPT: We can batch the allocation of persistent IDs together so the compute shaders overlap
+		const uint32 NumBuffersToResize = FMath::Min<uint32>(ComputeContext->BufferSwapsThisFrame_RT, UE_ARRAY_COUNT(ComputeContext->DataBuffers_RT));
+		for (uint32 i=0; i < NumBuffersToResize; ++i)
+		{
+			ComputeContext->DataBuffers_RT[i]->AllocateGPU(RHICmdList, ComputeContext->CurrentMaxAllocateInstances_RT + 1, FeatureLevel, ComputeContext->GetDebugSimName());
+		}
+
+		// Ensure we don't keep multi-tick buffers around longer than they are required by releasing them
+		for (uint32 i=NumBuffersToResize; i < UE_ARRAY_COUNT(ComputeContext->DataBuffers_RT); ++i)
+		{
+			ComputeContext->DataBuffers_RT[i]->ReleaseGPU();
+		}
+
+		// RDG will defer the Niagara dispatches until the graph is executed.
+		// Therefore we need to setup the DataToRender for MeshProcessors & sorting to use the correct data,
+		// that is anything that happens before PostRenderOpaque
+		if ((ComputeProxy->GetComputeTickStage() == ENiagaraGpuComputeTickStage::PreInitViews) || (ComputeProxy->GetComputeTickStage() == ENiagaraGpuComputeTickStage::PostInitViews))
+		{
+			FNiagaraDataBuffer* FinalBuffer = ComputeContext->GetPrevDataBuffer();
+			FinalBuffer->SetGPUInstanceCountBufferOffset(ComputeContext->CountOffset_RT);
+			FinalBuffer->SetNumInstances(ComputeContext->CurrentNumInstances_RT);
+			FinalBuffer->SetGPUDataReadyStage(ComputeProxy->GetComputeTickStage());
+			ComputeContext->SetDataToRender(FinalBuffer);
+#if WITH_MGPU
+			AddTemporalEffectBuffers(FinalBuffer);
+#endif // WITH_MGPU
+		}
+		// When low latency translucency is enabled we can setup the final buffer / final count here.
+		// This will allow our mesh processor commands to pickup the data for the same frame.
+		// This allows simulations that use the depth buffer, for example, to execute with no latency
+		else if ( GNiagaraGpuLowLatencyTranslucencyEnabled )
+		{
+			FNiagaraDataBuffer* FinalBuffer = ComputeContext->GetPrevDataBuffer();
+			FinalBuffer->SetGPUInstanceCountBufferOffset(ComputeContext->CountOffset_RT);
+			FinalBuffer->SetNumInstances(ComputeContext->CurrentNumInstances_RT);
+			FinalBuffer->SetGPUDataReadyStage(ComputeProxy->GetComputeTickStage());
+			ComputeContext->SetTranslucentDataToRender(FinalBuffer);
+		}
+	}
+
+#if WITH_MGPU
+	if (GNumAlternateFrameRenderingGroups > 1)
+	{
+		StageToBroadcastTemporalEffect = ENiagaraGpuComputeTickStage::Last;
+		while (StageToBroadcastTemporalEffect > ENiagaraGpuComputeTickStage::First && !DispatchListPerStage[static_cast<int32>(StageToBroadcastTemporalEffect)].HasWork())
+		{
+			StageToBroadcastTemporalEffect = static_cast<ENiagaraGpuComputeTickStage::Type>(static_cast<int32>(StageToBroadcastTemporalEffect) - 1);
+		}
+
+		StageToWaitForTemporalEffect = ENiagaraGpuComputeTickStage::First;
+		// If we're going to write to the instance count buffer after PreInitViews then
+		// that needs to be the wait stage, regardless of whether or not we're ticking
+		// anything in that stage.
+		if (!GPUInstanceCounterManager.HasEntriesPendingFree())
+		{
+			while (StageToWaitForTemporalEffect < StageToBroadcastTemporalEffect && !DispatchListPerStage[static_cast<int32>(StageToWaitForTemporalEffect)].HasWork())
+			{
+				StageToWaitForTemporalEffect = static_cast<ENiagaraGpuComputeTickStage::Type>(static_cast<int32>(StageToWaitForTemporalEffect) + 1);
+			}
+		}
+	}
+#endif // WITH_MGPU
+}
+
+void FNiagaraGpuComputeDispatch::PrepareAllTicks(FRHICommandListImmediate& RHICmdList)
+{
+	for (int iTickStage=0; iTickStage < ENiagaraGpuComputeTickStage::Max; ++iTickStage)
+	{
+		for (FNiagaraSystemGpuComputeProxy* ComputeProxy : ProxiesPerStage[iTickStage])
+		{
+			PrepareTicksForProxy(RHICmdList, ComputeProxy, DispatchListPerStage[iTickStage]);
+		}
+	}
+}
+
+void FNiagaraGpuComputeDispatch::ExecuteTicks(FRHICommandList& RHICmdList, FRHIUniformBuffer* ViewUniformBuffer, ENiagaraGpuComputeTickStage::Type TickStage)
+{
+#if WITH_MGPU
+	if (GNumAlternateFrameRenderingGroups > 1 && TickStage == StageToWaitForTemporalEffect)
+	{
+		RHICmdList.WaitForTemporalEffect(TemporalEffectName);
+	}
+#endif // WITH_MGPU
+
+	// Anything to execute for this stage
+	FNiagaraGpuDispatchList& DispatchList = DispatchListPerStage[TickStage];
+	if ( !DispatchList.HasWork() )
+	{
+		return;
+	}
+
+	TRACE_CPUPROFILER_EVENT_SCOPE(FNiagaraGpuComputeDispatch_ExecuteTicks);
+	SCOPED_DRAW_EVENTF(RHICmdList, FNiagaraGpuComputeDispatch_ExecuteTicks, TEXT("FNiagaraGpuComputeDispatch_ExecuteTicks - TickStage(%d)"), TickStage);
+	SCOPE_CYCLE_COUNTER(STAT_NiagaraGPUSimTick_RT);
+	SCOPED_GPU_STAT(RHICmdList, NiagaraGPUSimulation);
+
+	FMemMark Mark(FMemStack::Get());
+	TArray<FRHITransitionInfo, TMemStackAllocator<>> TransitionsBefore;
+	TArray<FRHITransitionInfo, TMemStackAllocator<>> TransitionsAfter;
+	TArray<FNiagaraDataBuffer*, TMemStackAllocator<>> IDToIndexInit;
+
+	for ( const FNiagaraGpuDispatchGroup& DispatchGroup : DispatchList.DispatchGroups )
+	{
+		SCOPED_CONDITIONAL_DRAW_EVENT(RHICmdList, NiagaraDispatchGroup, FNiagaraGpuComputeDispatchLocal::GAddDispatchGroupDrawEvent);
+
+		const bool bIsFirstGroup = &DispatchGroup == &DispatchList.DispatchGroups[0];
+		const bool bIsLastGroup = &DispatchGroup == &DispatchList.DispatchGroups.Last();
+
+		// Generate transitions and discover free / ID table updates
+		TransitionsBefore.Reserve(DispatchGroup.DispatchInstances.Num() * 3);
+		TransitionsAfter.Reserve(DispatchGroup.DispatchInstances.Num() * 3);
+		for (const FNiagaraGpuDispatchInstance& DispatchInstance : DispatchGroup.DispatchInstances)
+		{
+			if (FNiagaraDataBuffer* DestinationBuffer = DispatchInstance.SimStageData.Destination)
+			{
+				FNiagaraGpuComputeDispatchLocal::AddDataBufferTransitions(TransitionsBefore, TransitionsAfter, DestinationBuffer);
+			}
+
+			FNiagaraComputeExecutionContext* ComputeContext = DispatchInstance.InstanceData.Context;
+			const bool bRequiresPersistentIDs = ComputeContext->MainDataSet->RequiresPersistentIDs();
+			if (bRequiresPersistentIDs)
+			{
+				if ( FNiagaraDataBuffer* IDtoIndexBuffer = DispatchInstance.SimStageData.Destination )
+				{
+					IDToIndexInit.Add(IDtoIndexBuffer);
+					TransitionsBefore.Emplace(IDtoIndexBuffer->GetGPUIDToIndexTable().UAV, ERHIAccess::SRVCompute, ERHIAccess::UAVCompute);
+					TransitionsAfter.Emplace(IDtoIndexBuffer->GetGPUIDToIndexTable().UAV, ERHIAccess::UAVCompute, ERHIAccess::SRVCompute);
+				}
+			}
+		}
+
+		TransitionsBefore.Emplace(GPUInstanceCounterManager.GetInstanceCountBuffer().UAV, bIsFirstGroup ? FNiagaraGPUInstanceCountManager::kCountBufferDefaultState : ERHIAccess::UAVCompute, ERHIAccess::UAVCompute);
+		if (bIsLastGroup)
+		{
+			TransitionsAfter.Emplace(GPUInstanceCounterManager.GetInstanceCountBuffer().UAV, ERHIAccess::UAVCompute, FNiagaraGPUInstanceCountManager::kCountBufferDefaultState);
+		}
+
+		if (DispatchGroup.FreeIDUpdates.Num() > 0)
+		{
+			for (FNiagaraComputeExecutionContext* ComputeContext : DispatchGroup.FreeIDUpdates )
+			{
+				TransitionsAfter.Emplace(ComputeContext->MainDataSet->GetGPUFreeIDs().UAV, ERHIAccess::SRVCompute, ERHIAccess::UAVCompute);
+			}
+		}
+
+		// Consume per tick data from the game thread
+		//for (FNiagaraGPUSystemTick* Tick : DispatchGroup.TicksWithPerInstanceData)
+		//{
+		//	uint8* BasePointer = reinterpret_cast<uint8*>(Tick->DIInstanceData->PerInstanceDataForRT);
+
+		//	for (const auto& Pair : Tick->DIInstanceData->InterfaceProxiesToOffsets)
+		//	{
+		//		FNiagaraDataInterfaceProxy* Proxy = Pair.Key;
+		//		uint8* InstanceDataPtr = BasePointer + Pair.Value;
+		//		Proxy->ConsumePerInstanceDataFromGameThread(InstanceDataPtr, Tick->SystemInstanceID);
+		//	}
+		//}
+
+		// Execute Before Transitions
+		RHICmdList.Transition(TransitionsBefore);
+		TransitionsBefore.Reset();
+
+		// Initialize the IDtoIndex tables
+		if (IDToIndexInit.Num() > 0)
+		{
+			SCOPED_DRAW_EVENT(RHICmdList, NiagaraGPUComputeClearIDToIndexBuffer);
+
+			TArray<FRHITransitionInfo, TMemStackAllocator<>> IDToIndexTransitions;
+			IDToIndexTransitions.Reserve(IDToIndexInit.Num());
+
+			for (FNiagaraDataBuffer* IDtoIndexBuffer : IDToIndexInit)
+			{
+				NiagaraFillGPUIntBuffer(RHICmdList, FeatureLevel, IDtoIndexBuffer->GetGPUIDToIndexTable(), INDEX_NONE);
+				IDToIndexTransitions.Emplace(IDtoIndexBuffer->GetGPUIDToIndexTable().UAV, ERHIAccess::UAVCompute, ERHIAccess::UAVCompute);
+			}
+			IDToIndexInit.Reset();
+			RHICmdList.Transition(IDToIndexTransitions);
+		}
+
+		// Execute PreStage
+		{
+			TSet<FNiagaraDataInterfaceProxy*> ProxiesToFinalize;
+			for (const FNiagaraGpuDispatchInstance& DispatchInstance : DispatchGroup.DispatchInstances)
+			{
+				PreStageInterface(RHICmdList, DispatchInstance.Tick, DispatchInstance.InstanceData, DispatchInstance.SimStageData, ProxiesToFinalize);
+			}
+
+			for (FNiagaraDataInterfaceProxy* ProxyToFinalize : ProxiesToFinalize)
+			{
+				ProxyToFinalize->FinalizePreStage(RHICmdList, this);
+			}
+		}
+
+		// Execute Stage
+		RHICmdList.BeginUAVOverlap(GPUInstanceCounterManager.GetInstanceCountBuffer().UAV);
+		for (const FNiagaraGpuDispatchInstance& DispatchInstance : DispatchGroup.DispatchInstances)
+		{
+			++FNiagaraComputeExecutionContext::TickCounter;
+			if (DispatchInstance.InstanceData.bResetData && DispatchInstance.SimStageData.bFirstStage )
+			{
+				ResetDataInterfaces(RHICmdList, DispatchInstance.Tick, DispatchInstance.InstanceData);
+			}
+#if STATS
+			if (GPUProfiler.IsProfilingEnabled() && DispatchInstance.Tick.bIsFinalTick)
+			{			
+				const TStatId StatId = FDynamicStats::CreateStatId<FStatGroup_STATGROUP_NiagaraDetailed>("GPU_Stage_" + DispatchInstance.SimStageData.StageMetaData->SimulationStageName.ToString());
+				const int32 TimerHandle = GPUProfiler.StartTimer((uint64)DispatchInstance.InstanceData.Context, StatId, RHICmdList);
+				DispatchStage(RHICmdList, ViewUniformBuffer, DispatchInstance.Tick, DispatchInstance.InstanceData, DispatchInstance.SimStageData);
+				GPUProfiler.EndTimer(TimerHandle, RHICmdList);
+			}
+			else
+#endif
+			{
+				DispatchStage(RHICmdList, ViewUniformBuffer, DispatchInstance.Tick, DispatchInstance.InstanceData, DispatchInstance.SimStageData);
+			}
+		}
+		RHICmdList.EndUAVOverlap(GPUInstanceCounterManager.GetInstanceCountBuffer().UAV);
+
+		// Execute PostStage
+		{
+			TSet<FNiagaraDataInterfaceProxy*> ProxiesToFinalize;
+			for (const FNiagaraGpuDispatchInstance& DispatchInstance : DispatchGroup.DispatchInstances)
+			{
+				PostStageInterface(RHICmdList, DispatchInstance.Tick, DispatchInstance.InstanceData, DispatchInstance.SimStageData, ProxiesToFinalize);
+				if ( DispatchInstance.SimStageData.bLastStage )
+				{
+					PostSimulateInterface(RHICmdList, DispatchInstance.Tick, DispatchInstance.InstanceData);
+
+					// Update CurrentData with the latest information as things like ParticleReads can use this data
+					FNiagaraComputeExecutionContext* ComputeContext = DispatchInstance.InstanceData.Context;
+					const FNiagaraSimStageData& FinalSimStageData = DispatchInstance.SimStageData;
+					FNiagaraDataBuffer* FinalSimStageDataBuffer = FinalSimStageData.Destination != nullptr ? FinalSimStageData.Destination : FinalSimStageData.Source;
+					check(FinalSimStageDataBuffer != nullptr);
+
+					// If we are setting the data to render we need to ensure we switch back to the original CurrentData then swap the GPU buffers into it
+					if (DispatchInstance.SimStageData.bSetDataToRender)
+					{
+						check(ComputeContext->DataSetOriginalBuffer_RT != nullptr);
+						FNiagaraDataBuffer* CurrentData = ComputeContext->DataSetOriginalBuffer_RT;
+						ComputeContext->DataSetOriginalBuffer_RT = nullptr;
+
+						ComputeContext->MainDataSet->CurrentData = CurrentData;
+						CurrentData->SwapGPU(FinalSimStageDataBuffer);
+
+						// Mark data as ready for anyone who picks up the buffer on the next frame
+						CurrentData->SetGPUDataReadyStage(ENiagaraGpuComputeTickStage::First);
+                    
+						ComputeContext->SetTranslucentDataToRender(nullptr);
+						ComputeContext->SetDataToRender(CurrentData);
+
+#if WITH_MGPU
+						AddTemporalEffectBuffers(CurrentData);
+#endif // WITH_MGPU
+					}
+						// If this is not the final tick of the final stage we need set our temporary buffer for data interfaces, etc, that may snoop from CurrentData
+					else
+					{
+						ComputeContext->MainDataSet->CurrentData = FinalSimStageDataBuffer;
+					}
+				}
+			}
+
+			for (FNiagaraDataInterfaceProxy* ProxyToFinalize : ProxiesToFinalize)
+			{
+				ProxyToFinalize->FinalizePostStage(RHICmdList, this);
+			}
+		}
+
+		// Execute After Transitions
+		RHICmdList.Transition(TransitionsAfter);
+		TransitionsAfter.Reset();
+
+#if WITH_MGPU
+		// For PreInitViews we actually need to broadcast after UpdateDrawIndirectBuffer
+		// since this is the last place the instance count buffer is written to.
+		if (TickStage == StageToBroadcastTemporalEffect && StageToBroadcastTemporalEffect != ENiagaraGpuComputeTickStage::PreInitViews)
+		{
+			BroadcastTemporalEffect(RHICmdList);
+		}
+#endif // WITH_MGPU
+
+		// Update free IDs
+		if (DispatchGroup.FreeIDUpdates.Num() > 0)
+		{
+			UpdateFreeIDsListSizesBuffer(RHICmdList, DispatchGroup.FreeIDUpdates.Num());
+			UpdateFreeIDBuffers(RHICmdList, DispatchGroup.FreeIDUpdates);
+
+			for (FNiagaraComputeExecutionContext* ComputeContext : DispatchGroup.FreeIDUpdates)
+			{
+				TransitionsAfter.Emplace(ComputeContext->MainDataSet->GetGPUFreeIDs().UAV, ERHIAccess::UAVCompute, ERHIAccess::SRVCompute);
+			}
+			RHICmdList.Transition(TransitionsAfter);
+			TransitionsAfter.Reset();
+		}
+	}
+
+	// Clear dispatch groups
+	// We do not release the counts as we won't do that until we finish the dispatches
+	DispatchList.DispatchGroups.Empty();
+}
+
+void FNiagaraGpuComputeDispatch::DispatchStage(FRHICommandList& RHICmdList, FRHIUniformBuffer* ViewUniformBuffer, const FNiagaraGPUSystemTick& Tick, const FNiagaraComputeInstanceData& InstanceData, const FNiagaraSimStageData& SimStageData)
+{
+	// Setup source buffer
+	if ( SimStageData.Source != nullptr )
+	{
+		SimStageData.Source->SetNumInstances(SimStageData.SourceNumInstances);
+		SimStageData.Source->SetGPUInstanceCountBufferOffset(SimStageData.SourceCountOffset);
+	}
+
+	// Setup destination buffer
+	int32 InstancesToSpawn = 0;
+	if ( SimStageData.Destination != nullptr )
+	{
+		SimStageData.Destination->SetNumInstances(SimStageData.DestinationNumInstances);
+		SimStageData.Destination->SetGPUInstanceCountBufferOffset(SimStageData.DestinationCountOffset);
+		SimStageData.Destination->SetIDAcquireTag(FNiagaraComputeExecutionContext::TickCounter);
+
+		if ( SimStageData.bFirstStage )
+		{
+			check(SimStageData.DestinationNumInstances >= SimStageData.SourceNumInstances);
+			InstancesToSpawn = SimStageData.DestinationNumInstances - SimStageData.SourceNumInstances;
+		}
+		SimStageData.Destination->SetNumSpawnedInstances(InstancesToSpawn);
+	}
+
+	// Get dispatch count
+	uint32 DispatchCount = 0;
+	if ( SimStageData.AlternateIterationSource )
+	{
+		const FIntVector ElementCount = SimStageData.AlternateIterationSource->GetElementCount(Tick.SystemInstanceID);
+		// Verify the number of elements isn't higher that what we can handle
+		checkf(ElementCount.X * ElementCount.Y * ElementCount.Z < uint64(TNumericLimits<uint32>::Max()), TEXT("ElementCount(%d, %d, %d) for IterationInterface(%s) overflows an int32 this is not allowed"), ElementCount.X, ElementCount.Y, ElementCount.Z, *SimStageData.AlternateIterationSource->SourceDIName.ToString());
+
+		DispatchCount = ElementCount.X * ElementCount.Y * ElementCount.Z;
+	}
+	else
+	{
+		DispatchCount = SimStageData.DestinationNumInstances;
+	}
+
+	if ( DispatchCount == 0 )
+	{
+		return;
+	}
+
+	// Get Shader
+	const FNiagaraShaderRef ComputeShader = InstanceData.Context->GPUScript_RT->GetShader(SimStageData.StageIndex);
+	FRHIComputeShader* RHIComputeShader = ComputeShader.GetComputeShader();
+	SetComputePipelineState(RHICmdList, RHIComputeShader);
+
+	SCOPED_DRAW_EVENTF(RHICmdList, NiagaraGPUSimulationCS,
+		TEXT("NiagaraGpuSim(%s) DispatchCount(%u) Stage(%s %u) Iteration(%d) NumInstructions(%u)"),
+		InstanceData.Context->GetDebugSimName(),
+		DispatchCount,
+		*SimStageData.StageMetaData->SimulationStageName.ToString(),
+		SimStageData.StageIndex,
+		SimStageData.IterationIndex,
+		ComputeShader->GetNumInstructions()
+	);
+	FNiagaraEmptyUAVPoolScopedAccess UAVPoolAccessScope(GetEmptyUAVPool());
+
+	// Set Parameters
+	const bool bRequiresPersistentIDs = InstanceData.Context->MainDataSet->RequiresPersistentIDs();
+
+	SetShaderValue(RHICmdList, RHIComputeShader, ComputeShader->SimStartParam, InstanceData.bResetData ? 1U : 0U);
+	SetShaderValue(RHICmdList, RHIComputeShader, ComputeShader->EmitterTickCounterParam, FNiagaraComputeExecutionContext::TickCounter);
+	SetShaderValue(RHICmdList, RHIComputeShader, ComputeShader->UpdateStartInstanceParam, 0);
+	SetShaderValue(RHICmdList, RHIComputeShader, ComputeShader->NumSpawnedInstancesParam, InstancesToSpawn);
+	SetSRVParameter(RHICmdList, RHIComputeShader, ComputeShader->FreeIDBufferParam, bRequiresPersistentIDs ? InstanceData.Context->MainDataSet->GetGPUFreeIDs().SRV.GetReference() : FNiagaraRenderer::GetDummyIntBuffer());
+
+	// Set spawn Information
+	// This parameter is an array of structs with 2 floats and 2 ints on CPU, but a float4 array on GPU. The shader uses asint() to cast the integer values. To set the parameter,
+	// we pass the structure array as a float* to SetShaderValueArray() and specify the number of floats (not float vectors).
+	static_assert((sizeof(InstanceData.SpawnInfo.SpawnInfoStartOffsets) % SHADER_PARAMETER_ARRAY_ELEMENT_ALIGNMENT) == 0, "sizeof SpawnInfoStartOffsets should be a multiple of SHADER_PARAMETER_ARRAY_ELEMENT_ALIGNMENT");
+	static_assert((sizeof(InstanceData.SpawnInfo.SpawnInfoParams) % SHADER_PARAMETER_ARRAY_ELEMENT_ALIGNMENT) == 0, "sizeof SpawnInfoParams should be a multiple of SHADER_PARAMETER_ARRAY_ELEMENT_ALIGNMENT");
+	SetShaderValueArray(RHICmdList, RHIComputeShader, ComputeShader->EmitterSpawnInfoOffsetsParam, InstanceData.SpawnInfo.SpawnInfoStartOffsets, NIAGARA_MAX_GPU_SPAWN_INFOS);
+	SetShaderValueArray(RHICmdList, RHIComputeShader, ComputeShader->EmitterSpawnInfoParamsParam, &InstanceData.SpawnInfo.SpawnInfoParams[0].IntervalDt, 4 * NIAGARA_MAX_GPU_SPAWN_INFOS);
+
+	if (ComputeShader->ViewUniformBufferParam.IsBound())
+	{
+		if (ensureMsgf(ViewUniformBuffer != nullptr, TEXT("ViewUniformBuffer is required for '%s' but we do not have one to bind"), InstanceData.Context->GetDebugSimName()))
+		{
+			RHICmdList.SetShaderUniformBuffer(RHIComputeShader, ComputeShader->ViewUniformBufferParam.GetBaseIndex(), ViewUniformBuffer);
+		}
+	}
+
+	// Setup instance counts
+	if (ComputeShader->InstanceCountsParam.IsBound())
+	{
+		ComputeShader->InstanceCountsParam.SetBuffer(RHICmdList, RHIComputeShader, GPUInstanceCounterManager.GetInstanceCountBuffer());
+
+		if (SimStageData.AlternateIterationSource)
+		{
+			SetShaderValue(RHICmdList, RHIComputeShader, ComputeShader->ReadInstanceCountOffsetParam, -1);
+			SetShaderValue(RHICmdList, RHIComputeShader, ComputeShader->WriteInstanceCountOffsetParam, -1);
+		}
+		else
+		{
+			SetShaderValue(RHICmdList, RHIComputeShader, ComputeShader->ReadInstanceCountOffsetParam, SimStageData.SourceCountOffset);
+			SetShaderValue(RHICmdList, RHIComputeShader, ComputeShader->WriteInstanceCountOffsetParam, SimStageData.DestinationCountOffset);
+		}
+	}
+
+	// Simulation Stage Information
+	// X = Count Buffer Instance Count Offset (INDEX_NONE == Use Instance Count)
+	// Y = Instance Count
+	// Z = Iteration Index
+	// W = Num Iterations
+	{
+		FIntVector4 SimulationStageIterationInfo = { INDEX_NONE, -1, 0, 0 };
+		float SimulationStageNormalizedIterationIndex = 0.0f;
+
+		if (SimStageData.AlternateIterationSource != nullptr)
+		{
+			const uint32 IterationInstanceCountOffset = SimStageData.AlternateIterationSource->GetGPUInstanceCountOffset(Tick.SystemInstanceID);
+			SimulationStageIterationInfo.X = IterationInstanceCountOffset;
+			SimulationStageIterationInfo.Y = IterationInstanceCountOffset == INDEX_NONE ? DispatchCount : 0;
+		}
+
+		const int32 NumIterations = InstanceData.NumIterationsPerStage[SimStageData.StageIndex];
+		const int32 IterationIndex = SimStageData.IterationIndex;
+		SimulationStageIterationInfo.Z = IterationIndex;
+		SimulationStageIterationInfo.W = NumIterations;
+		SimulationStageNormalizedIterationIndex = NumIterations > 1 ? float(IterationIndex) / float(NumIterations - 1) : 1.0f;
+
+		SetShaderValue(RHICmdList, RHIComputeShader, ComputeShader->SimulationStageIterationInfoParam, SimulationStageIterationInfo);
+		SetShaderValue(RHICmdList, RHIComputeShader, ComputeShader->SimulationStageNormalizedIterationIndexParam, SimulationStageNormalizedIterationIndex);
+	}
+
+	// Set Input & Output buffers
+	FNiagaraDataBuffer::SetInputShaderParams(RHICmdList, ComputeShader.GetShader(), SimStageData.Source);
+	FNiagaraDataBuffer::SetOutputShaderParams(RHICmdList, ComputeShader.GetShader(), SimStageData.Destination);
+
+	// Set data interface parameters
+	SetDataInterfaceParameters(RHICmdList, Tick, InstanceData, ComputeShader, SimStageData);
+
+	// Execute the dispatch
+	{
+		const uint32 ThreadGroupSize = FNiagaraShader::GetGroupSize(ShaderPlatform);
+		checkf(ThreadGroupSize <= NiagaraComputeMaxThreadGroupSize, TEXT("ShaderPlatform(%d) is set to ThreadGroup size (%d) which is > the NiagaraComputeMaxThreadGroupSize(%d)"), ShaderPlatform, ThreadGroupSize, NiagaraComputeMaxThreadGroupSize);
+		const uint32 TotalThreadGroups = FMath::DivideAndRoundUp(DispatchCount, ThreadGroupSize);
+
+		const uint32 ThreadGroupCountZ = 1;
+		const uint32 ThreadGroupCountY = FMath::DivideAndRoundUp(TotalThreadGroups, NiagaraMaxThreadGroupCountPerDimension);
+		const uint32 ThreadGroupCountX = FMath::DivideAndRoundUp(TotalThreadGroups, ThreadGroupCountY);
+
+		SetShaderValue(RHICmdList, RHIComputeShader, ComputeShader->DispatchThreadIdToLinearParam, ThreadGroupCountY > 1 ? ThreadGroupCountX * ThreadGroupSize : 0);
+
+		SetConstantBuffers(RHICmdList, ComputeShader, Tick, &InstanceData);
+
+		DispatchComputeShader(RHICmdList, ComputeShader, ThreadGroupCountX, ThreadGroupCountY, ThreadGroupCountZ);
+
+		INC_DWORD_STAT(STAT_NiagaraGPUDispatches);
+	}
+
+	// Unset UAV parameters
+	UnsetDataInterfaceParameters(RHICmdList, Tick, InstanceData, ComputeShader, SimStageData);
+	FNiagaraDataBuffer::UnsetShaderParams(RHICmdList, ComputeShader.GetShader());
+	ComputeShader->InstanceCountsParam.UnsetUAV(RHICmdList, RHIComputeShader);
+
+	// Optionally submit commands to the GPU
+	// This can be used to avoid accidental TDR detection in the editor especially when issuing multiple ticks in the same frame
+	if (GNiagaraGpuSubmitCommandHint > 0)
+	{
+		++TotalDispatchesThisFrame;
+		if ((TotalDispatchesThisFrame % GNiagaraGpuSubmitCommandHint) == 0)
+		{
+			RHICmdList.SubmitCommandsHint();
+		}
+	}
+}
+
+void FNiagaraGpuComputeDispatch::PreInitViews(FRDGBuilder& GraphBuilder, bool bAllowGPUParticleUpdate)
+{
+	bRequiresReadback = false;
+	GNiagaraViewDataManager.ClearSceneTextureParameters();
+
+	GpuReadbackManagerPtr->Tick();
+#if NIAGARA_COMPUTEDEBUG_ENABLED
+	if ( FNiagaraGpuComputeDebug* GpuComputeDebug = GetGpuComputeDebug() )
+	{
+		GpuComputeDebug->Tick(GraphBuilder.RHICmdList);
+	}
+#endif
+
+#if STATS
+	// check if we can process profiling results from previous frames
+	if (GPUProfiler.IsProfilingEnabled())
+	{
+		TArray<FNiagaraGPUTimingResult, TInlineAllocator<16>> ProfilingResults;
+		GPUProfiler.QueryTimingResults(GraphBuilder.RHICmdList, ProfilingResults);
+		TMap<TWeakObjectPtr<UNiagaraEmitter>, TMap<TStatIdData const*, float>> CapturedStats;
+		for (const FNiagaraGPUTimingResult& Result : ProfilingResults)
+		{
+			// map the individual timings to the emitters and source scripts
+			if (Result.ReporterHandle)
+			{
+				FNiagaraComputeExecutionContext* Context = (FNiagaraComputeExecutionContext*)Result.ReporterHandle;
+				if (Context->EmitterPtr.IsValid())
+				{
+					CapturedStats.FindOrAdd(Context->EmitterPtr).FindOrAdd(Result.StatId.GetRawPointer()) += Result.ElapsedMicroseconds;
+				}
+			}
+		}
+		if (CapturedStats.Num() > 0)
+		{
+			// report the captured gpu stats on the game thread
+			uint64 ReporterHandle = (uint64)this;
+			AsyncTask(ENamedThreads::GameThread, [LocalCapturedStats = MoveTemp(CapturedStats), ReporterHandle]
+            {
+                for (auto& Entry : LocalCapturedStats)
+                {
+                    if (Entry.Key.IsValid())
+                    {
+                        Entry.Key->GetStatData().AddStatCapture(TTuple<uint64, ENiagaraScriptUsage>(ReporterHandle, ENiagaraScriptUsage::ParticleGPUComputeScript), Entry.Value);
+                    }
+                }
+            });
+		}
+	}
+#endif
+
+	LLM_SCOPE(ELLMTag::Niagara);
+	TotalDispatchesThisFrame = 0;
+
+	// Reset the list of GPUSort tasks and release any resources they hold on to.
+	// It might be worth considering doing so at the end of the render to free the resources immediately.
+	// (note that currently there are no callback appropriate to do it)
+	SimulationsToSort.Reset();
+
+	if (FNiagaraUtilities::AllowGPUParticles(GetShaderPlatform()))
+	{
+		if ( bAllowGPUParticleUpdate )
+		{
+			FramesBeforeTickFlush = 0;
+
+			UpdateInstanceCountManager(GraphBuilder.RHICmdList);
+			PrepareAllTicks(GraphBuilder.RHICmdList);
+		
+#if RHI_RAYTRACING
+			RayTracingHelper->BeginFrame(GraphBuilder.RHICmdList, HasRayTracingScene());
+#endif
+
+			if ( FNiagaraGpuComputeDispatchLocal::GDebugLogging)
+			{
+				DumpDebugFrame();
+			}
+
+			FNiagaraSceneTextureParameters* PassParameters = GraphBuilder.AllocParameters<FNiagaraSceneTextureParameters>();
+			GNiagaraViewDataManager.GetSceneTextureParameters(GraphBuilder, *PassParameters);
+
+			GraphBuilder.AddPass(
+				RDG_EVENT_NAME("Niagara::PreInitViews"),
+				PassParameters,
+				ERDGPassFlags::Compute | ERDGPassFlags::NeverCull,
+				[this, PassParameters](FRHICommandListImmediate& RHICmdList)
+			{
+				NiagaraSceneTextures = PassParameters;
+				ON_SCOPE_EXIT { NiagaraSceneTextures = nullptr; };
+
+				ExecuteTicks(RHICmdList, nullptr, ENiagaraGpuComputeTickStage::PreInitViews);
+			});
+		}
+	}
+	else
+	{
+		GPUInstanceCounterManager.ResizeBuffers(GraphBuilder.RHICmdList, FeatureLevel,  0);
+		FinishDispatches();
+	}
+}
+
+void FNiagaraGpuComputeDispatch::PostInitViews(FRDGBuilder& GraphBuilder, TArrayView<const class FViewInfo> Views, bool bAllowGPUParticleUpdate)
+{
+	LLM_SCOPE(ELLMTag::Niagara);
+
+	bAllowGPUParticleUpdate = bAllowGPUParticleUpdate && GetReferenceAllowGPUUpdate(Views);
+
+	if (bAllowGPUParticleUpdate && FNiagaraUtilities::AllowGPUParticles(GetShaderPlatform()))
+	{
+		FNiagaraSceneTextureParameters* PassParameters = GraphBuilder.AllocParameters<FNiagaraSceneTextureParameters>();
+		GNiagaraViewDataManager.GetSceneTextureParameters(GraphBuilder, *PassParameters);
+
+		GraphBuilder.AddPass(
+			RDG_EVENT_NAME("Niagara::PostInitViews"),
+			PassParameters,
+			ERDGPassFlags::Compute | ERDGPassFlags::NeverCull,
+			[this, PassParameters, Views](FRHICommandListImmediate& RHICmdList)
+		{
+			NiagaraSceneTextures = PassParameters;
+			ON_SCOPE_EXIT{ NiagaraSceneTextures = nullptr; };
+
+			FRHIUniformBuffer* ViewUniformBuffer = GetReferenceViewUniformBuffer(Views);
+			ExecuteTicks(RHICmdList, ViewUniformBuffer, ENiagaraGpuComputeTickStage::PostInitViews);
+		});
+	}
+}
+
+void FNiagaraGpuComputeDispatch::PostRenderOpaque(FRDGBuilder& GraphBuilder, TConstArrayView<FViewInfo> Views, bool bAllowGPUParticleUpdate)
+{
+	LLM_SCOPE(ELLMTag::Niagara);
+
+	bAllowGPUParticleUpdate = bAllowGPUParticleUpdate && GetReferenceAllowGPUUpdate(Views);
+
+	if ( bAllowGPUParticleUpdate && Views.IsValidIndex(0) )
+	{
+		CachedViewRect = Views[0].ViewRect;
+	}
+
+	FNiagaraSceneTextureParameters* PassParameters = GraphBuilder.AllocParameters<FNiagaraSceneTextureParameters>();
+	GNiagaraViewDataManager.GetSceneTextureParameters(GraphBuilder, *PassParameters);
+
+	GraphBuilder.AddPass(
+		RDG_EVENT_NAME("Niagara::PostRenderOpaque"),
+		PassParameters,
+		ERDGPassFlags::Compute | ERDGPassFlags::NeverCull,
+		[this, PassParameters, Views, bAllowGPUParticleUpdate](FRHICommandListImmediate& RHICmdList)
+	{
+		if (bAllowGPUParticleUpdate && FNiagaraUtilities::AllowGPUParticles(GetShaderPlatform()))
+		{
+			NiagaraSceneTextures = PassParameters;
+			ON_SCOPE_EXIT{ NiagaraSceneTextures = nullptr; };
+
+#if RHI_RAYTRACING
+			BuildRayTracingSceneInfo(RHICmdList, Views);
+#endif
+
+			CurrentPassViews = Views;
+
+			FRHIUniformBuffer* ViewUniformBuffer = GetReferenceViewUniformBuffer(Views);
+
+			// Setup new readback since if there is no pending request, there is no risk of having invalid data read (offset being allocated after the readback was sent).
+			ExecuteTicks(RHICmdList, ViewUniformBuffer, ENiagaraGpuComputeTickStage::PostOpaqueRender);
+
+			FinishDispatches();
+
+#if RHI_RAYTRACING
+			RayTracingHelper->EndFrame(RHICmdList, HasRayTracingScene(), GetScene());
+#endif
+
+			// Clear CurrentPassViews
+			CurrentPassViews = TConstArrayView<FViewInfo>();
+
+			ProcessDebugReadbacks(RHICmdList, false);
+		}
+
+		if ( bRequiresReadback )
+		{
+			check(!GPUInstanceCounterManager.HasPendingGPUReadback());
+			GPUInstanceCounterManager.EnqueueGPUReadback(RHICmdList);
+			bRequiresReadback = false;
+		}
+	});
+
+	GNiagaraViewDataManager.ClearSceneTextureParameters();
+}
+
+void FNiagaraGpuComputeDispatch::ProcessDebugReadbacks(FRHICommandListImmediate& RHICmdList, bool bWaitCompletion)
+{
+#if !UE_BUILD_SHIPPING
+	// Execute any pending readbacks as the ticks have now all been processed
+	for (const FDebugReadbackInfo& DebugReadback : GpuDebugReadbackInfos)
+	{
+		FNiagaraDataBuffer* CurrentDataBuffer = DebugReadback.Context->MainDataSet->GetCurrentData();
+		if ( CurrentDataBuffer == nullptr )
+		{
+			// Data is invalid
+			DebugReadback.DebugInfo->Frame.CopyFromGPUReadback(nullptr, nullptr, nullptr, 0, 0, 0, 0, 0);
+			DebugReadback.DebugInfo->bWritten = true;
+			continue;
+		}
+
+		const uint32 CountOffset = CurrentDataBuffer->GetGPUInstanceCountBufferOffset();
+		if (CountOffset == INDEX_NONE)
+		{
+			// Data is invalid
+			DebugReadback.DebugInfo->Frame.CopyFromGPUReadback(nullptr, nullptr, nullptr, 0, 0, 0, 0, 0);
+			DebugReadback.DebugInfo->bWritten = true;
+			continue;
+		}
+
+		// Execute readback
+		constexpr int32 MaxReadbackBuffers = 4;
+		TArray<FRHIBuffer*, TInlineAllocator<MaxReadbackBuffers>> ReadbackBuffers;
+
+		const int32 CountBufferIndex = ReadbackBuffers.Add(GPUInstanceCounterManager.GetInstanceCountBuffer().Buffer);
+		const int32 FloatBufferIndex = (CurrentDataBuffer->GetGPUBufferFloat().NumBytes == 0) ? INDEX_NONE : ReadbackBuffers.Add(CurrentDataBuffer->GetGPUBufferFloat().Buffer);
+		const int32 HalfBufferIndex = (CurrentDataBuffer->GetGPUBufferHalf().NumBytes == 0) ? INDEX_NONE : ReadbackBuffers.Add(CurrentDataBuffer->GetGPUBufferHalf().Buffer);
+		const int32 IntBufferIndex = (CurrentDataBuffer->GetGPUBufferInt().NumBytes == 0) ? INDEX_NONE : ReadbackBuffers.Add(CurrentDataBuffer->GetGPUBufferInt().Buffer);
+
+		const int32 FloatBufferStride = CurrentDataBuffer->GetFloatStride();
+		const int32 HalfBufferStride = CurrentDataBuffer->GetHalfStride();
+		const int32 IntBufferStride = CurrentDataBuffer->GetInt32Stride();
+
+		GpuReadbackManagerPtr->EnqueueReadbacks(
+			RHICmdList,
+			MakeArrayView(ReadbackBuffers),
+			[=, DebugInfo=DebugReadback.DebugInfo](TConstArrayView<TPair<void*, uint32>> BufferData)
+			{
+				checkf(4 + (CountOffset * 4) <= BufferData[CountBufferIndex].Value, TEXT("CountOffset %d is out of bounds"), CountOffset, BufferData[CountBufferIndex].Value);
+				const int32 InstanceCount = reinterpret_cast<int32*>(BufferData[CountBufferIndex].Key)[CountOffset];
+				const float* FloatDataBuffer = FloatBufferIndex == INDEX_NONE ? nullptr : reinterpret_cast<float*>(BufferData[FloatBufferIndex].Key);
+				const FFloat16* HalfDataBuffer = HalfBufferIndex == INDEX_NONE ? nullptr : reinterpret_cast<FFloat16*>(BufferData[HalfBufferIndex].Key);
+				const int32* IntDataBuffer = IntBufferIndex == INDEX_NONE ? nullptr : reinterpret_cast<int32*>(BufferData[IntBufferIndex].Key);
+
+				DebugInfo->Frame.CopyFromGPUReadback(FloatDataBuffer, IntDataBuffer, HalfDataBuffer, 0, InstanceCount, FloatBufferStride, IntBufferStride, HalfBufferStride);
+				DebugInfo->bWritten = true;
+			}
+		);
+	}
+	GpuDebugReadbackInfos.Empty();
+
+	if (bWaitCompletion)
+	{
+		GpuReadbackManagerPtr->WaitCompletion(RHICmdList);
+	}
+#endif
+}
+
+bool FNiagaraGpuComputeDispatch::UsesGlobalDistanceField() const
+{
+	return NumProxiesThatRequireDistanceFieldData > 0;
+}
+
+bool FNiagaraGpuComputeDispatch::UsesDepthBuffer() const
+{
+	return NumProxiesThatRequireDepthBuffer > 0;
+}
+
+bool FNiagaraGpuComputeDispatch::RequiresEarlyViewUniformBuffer() const
+{
+	return NumProxiesThatRequireEarlyViewData > 0;
+}
+
+bool FNiagaraGpuComputeDispatch::RequiresRayTracingScene() const
+{
+	return NumProxiesThatRequireRayTracingScene > 0;
+}
+
+void FNiagaraGpuComputeDispatch::PreRender(FRDGBuilder& GraphBuilder, TConstArrayView<FViewInfo> Views, bool bAllowGPUParticleUpdate)
+{
+	if (!FNiagaraUtilities::AllowGPUParticles(GetShaderPlatform()))
+	{
+		return;
+	}
+
+	LLM_SCOPE(ELLMTag::Niagara);
+}
+
+void FNiagaraGpuComputeDispatch::OnDestroy()
+{
+	FNiagaraWorldManager::OnComputeDispatchInterfaceDestroyed(this);
+	FFXSystemInterface::OnDestroy();
+}
+
+bool FNiagaraGpuComputeDispatch::AddSortedGPUSimulation(FNiagaraGPUSortInfo& SortInfo)
+{
+	if (GPUSortManager && GPUSortManager->AddTask(SortInfo.AllocationInfo, SortInfo.ParticleCount, SortInfo.SortFlags))
+	{
+		// It's not worth currently to have a map between SortInfo.AllocationInfo.SortBatchId and the relevant indices in SimulationsToSort
+		// because the number of batches is expect to be very small (1 or 2). If this change, it might be worth reconsidering.
+		SimulationsToSort.Add(SortInfo);
+		return true;
+	}
+	else
+	{
+		return false;
+	}
+}
+
+const FGlobalDistanceFieldParameterData* FNiagaraGpuComputeDispatch::GetGlobalDistanceFieldParameters() const
+{ 
+	check(CurrentPassViews.Num() > 0); 
+	return &CurrentPassViews[0].GlobalDistanceFieldInfo.ParameterData; 
+}
+
+void FNiagaraGpuComputeDispatch::GenerateSortKeys(FRHICommandListImmediate& RHICmdList, int32 BatchId, int32 NumElementsInBatch, EGPUSortFlags Flags, FRHIUnorderedAccessView* KeysUAV, FRHIUnorderedAccessView* ValuesUAV)
+{
+	const FGPUSortManager::FKeyGenInfo KeyGenInfo((uint32)NumElementsInBatch, EnumHasAnyFlags(Flags, EGPUSortFlags::HighPrecisionKeys));
+
+	FNiagaraSortKeyGenCS::FPermutationDomain SortPermutationVector;
+	SortPermutationVector.Set<FNiagaraSortKeyGenCS::FSortUsingMaxPrecision>(GNiagaraGPUSortingUseMaxPrecision != 0);
+	SortPermutationVector.Set<FNiagaraSortKeyGenCS::FEnableCulling>(false);
+
+	FNiagaraSortKeyGenCS::FPermutationDomain SortAndCullPermutationVector;
+	SortAndCullPermutationVector.Set<FNiagaraSortKeyGenCS::FSortUsingMaxPrecision>(GNiagaraGPUSortingUseMaxPrecision != 0);
+	SortAndCullPermutationVector.Set<FNiagaraSortKeyGenCS::FEnableCulling>(true);
+
+	TShaderMapRef<FNiagaraSortKeyGenCS> SortKeyGenCS(GetGlobalShaderMap(FeatureLevel), SortPermutationVector);
+	TShaderMapRef<FNiagaraSortKeyGenCS> SortAndCullKeyGenCS(GetGlobalShaderMap(FeatureLevel), SortAndCullPermutationVector);
+
+	FRWBuffer* CulledCountsBuffer = GPUInstanceCounterManager.AcquireCulledCountsBuffer(RHICmdList, FeatureLevel);
+
+	FNiagaraSortKeyGenCS::FParameters Params;
+	Params.SortKeyMask = KeyGenInfo.SortKeyParams.X;
+	Params.SortKeyShift = KeyGenInfo.SortKeyParams.Y;
+	Params.SortKeySignBit = KeyGenInfo.SortKeyParams.Z;
+	Params.OutKeys = KeysUAV;
+	Params.OutParticleIndices = ValuesUAV;
+
+	FRHIUnorderedAccessView* OverlapUAVs[3];
+	uint32 NumOverlapUAVs = 0;
+
+	OverlapUAVs[NumOverlapUAVs] = KeysUAV;
+	++NumOverlapUAVs;
+	OverlapUAVs[NumOverlapUAVs] = ValuesUAV;
+	++NumOverlapUAVs;
+
+	if (CulledCountsBuffer)
+	{
+		Params.OutCulledParticleCounts = CulledCountsBuffer->UAV;
+		OverlapUAVs[NumOverlapUAVs] = CulledCountsBuffer->UAV;
+		++NumOverlapUAVs;
+	}
+	else
+	{
+		// Note: We don't care that the buffer will be allowed to be reused
+		FNiagaraEmptyUAVPoolScopedAccess UAVPoolAccessScope(GetEmptyUAVPool());
+		Params.OutCulledParticleCounts = GetEmptyUAVFromPool(RHICmdList, PF_R32_UINT, ENiagaraEmptyUAVType::Buffer);
+	}
+
+	RHICmdList.BeginUAVOverlap(MakeArrayView(OverlapUAVs, NumOverlapUAVs));
+
+	for (const FNiagaraGPUSortInfo& SortInfo : SimulationsToSort)
+	{
+		if (SortInfo.AllocationInfo.SortBatchId == BatchId)
+		{
+			Params.NiagaraParticleDataFloat = SortInfo.ParticleDataFloatSRV;
+			Params.NiagaraParticleDataHalf = SortInfo.ParticleDataHalfSRV;
+			Params.NiagaraParticleDataInt = SortInfo.ParticleDataIntSRV;
+			Params.GPUParticleCountBuffer = SortInfo.GPUParticleCountSRV;
+			Params.FloatDataStride = SortInfo.FloatDataStride;
+			Params.HalfDataStride = SortInfo.HalfDataStride;
+			Params.IntDataStride = SortInfo.IntDataStride;
+			Params.ParticleCount = SortInfo.ParticleCount;
+			Params.GPUParticleCountOffset = SortInfo.GPUParticleCountOffset;
+			Params.CulledGPUParticleCountOffset = SortInfo.CulledGPUParticleCountOffset;
+			Params.EmitterKey = (uint32)SortInfo.AllocationInfo.ElementIndex << KeyGenInfo.ElementKeyShift;
+			Params.OutputOffset = SortInfo.AllocationInfo.BufferOffset;
+			Params.CameraPosition = SortInfo.ViewOrigin;
+			Params.CameraDirection = SortInfo.ViewDirection;
+			Params.SortMode = (uint32)SortInfo.SortMode;
+			Params.SortAttributeOffset = SortInfo.SortAttributeOffset;
+			Params.CullPositionAttributeOffset = SortInfo.CullPositionAttributeOffset;
+			Params.CullOrientationAttributeOffset = SortInfo.CullOrientationAttributeOffset;
+			Params.CullScaleAttributeOffset = SortInfo.CullScaleAttributeOffset;
+			Params.RendererVisibility = SortInfo.RendererVisibility;
+			Params.RendererVisTagAttributeOffset = SortInfo.RendererVisTagAttributeOffset;
+			Params.MeshIndex = SortInfo.MeshIndex;
+			Params.MeshIndexAttributeOffset = SortInfo.MeshIndexAttributeOffset;
+			Params.CullDistanceRangeSquared = SortInfo.DistanceCullRange * SortInfo.DistanceCullRange;
+			Params.LocalBoundingSphere = FVector4(SortInfo.LocalBSphere.Center, SortInfo.LocalBSphere.W);
+			Params.CullingWorldSpaceOffset = SortInfo.CullingWorldSpaceOffset;
+
+			Params.NumCullPlanes = 0;
+			for (const FPlane& Plane : SortInfo.CullPlanes)
+			{
+				Params.CullPlanes[Params.NumCullPlanes++] = FVector4(Plane.X, Plane.Y, Plane.Z, Plane.W);
+			}
+
+			// Choose the shader to bind
+			TShaderMapRef<FNiagaraSortKeyGenCS> KeyGenCS = SortInfo.bEnableCulling ? SortAndCullKeyGenCS : SortKeyGenCS;
+			SetComputePipelineState(RHICmdList, KeyGenCS.GetComputeShader());
+
+			SetShaderParameters(RHICmdList, KeyGenCS, KeyGenCS.GetComputeShader(), Params);
+			DispatchComputeShader(RHICmdList, KeyGenCS, FMath::DivideAndRoundUp(SortInfo.ParticleCount, NIAGARA_KEY_GEN_THREAD_COUNT), 1, 1);
+			UnsetShaderUAVs(RHICmdList, KeyGenCS, KeyGenCS.GetComputeShader());
+		}
+	}
+
+	RHICmdList.EndUAVOverlap(MakeArrayView(OverlapUAVs, NumOverlapUAVs));
+}
+
+#if RHI_RAYTRACING
+void FNiagaraGpuComputeDispatch::BuildRayTracingSceneInfo(FRHICommandList& RHICmdList, TConstArrayView<FViewInfo> Views)
+{
+	if (NumProxiesThatRequireRayTracingScene > 0)
+	{
+		RayTracingHelper->BuildRayTracingSceneInfo(RHICmdList, Views);
+
+		RayTracingHelper->UpdateCollisionGroupMap(RHICmdList, GetScene(), FeatureLevel);
+	}
+	else
+	{
+		RayTracingHelper->Reset();
+	}
+}
+
+void FNiagaraGpuComputeDispatch::ResetRayTracingSceneInfo()
+{
+	RayTracingHelper->Reset();
+}
+
+
+FNiagaraRayTracingHelper& FNiagaraGpuComputeDispatch::GetRayTracingHelper()const
+{
+	check(RayTracingHelper.IsValid());
+	return *RayTracingHelper.Get();
+}
+
+bool FNiagaraGpuComputeDispatch::HasRayTracingScene() const
+{
+	return NumProxiesThatRequireRayTracingScene && RayTracingHelper->IsValid();
+}
+
+void FNiagaraGpuComputeDispatch::SetPrimitiveRayTracingCollisionGroup(UPrimitiveComponent* Primitive, uint32 Group)
+{
+	check(IsInGameThread());
+	if (Primitive == nullptr || Primitive->SceneProxy == nullptr)
+	{
+		return;
+	}
+
+	ENQUEUE_RENDER_COMMAND(NiagaraSetPrimHWRTCollisionGroup)(
+		[RT_NiagaraBatcher = this, Proxy = Primitive->SceneProxy, CollisionGroup=Group](FRHICommandListImmediate& RHICmdList)
+		{
+			check(Proxy);
+			RT_NiagaraBatcher->RayTracingHelper->SetPrimitiveCollisionGroup(*Proxy->GetPrimitiveSceneInfo(), CollisionGroup);
+		}
+	);
+}
+
+int32 FNiagaraGpuComputeDispatch::AcquireGPURayTracedCollisionGroup()
+{
+	check(IsInGameThread());//one of the few batcher functions that should be called from the GT
+	if (FreeGPURayTracedCollisionGroups.Num() > 0)
+	{
+		return FreeGPURayTracedCollisionGroups.Pop();
+	}
+
+	return NumGPURayTracedCollisionGroups++;
+}
+
+void FNiagaraGpuComputeDispatch::ReleaseGPURayTracedCollisionGroup(int32 CollisionGroup)
+{
+	check(IsInGameThread());//one of the few batcher functions that should be called from the GT
+	FreeGPURayTracedCollisionGroups.Add(CollisionGroup);
+}
+
+#endif
+
+void FNiagaraGpuComputeDispatch::OnPrimitiveGPUSceneInstancesAllocated()
+{
+#if RHI_RAYTRACING
+	RayTracingHelper->RefreshPrimitiveInstanceData();
+#endif
+}
+
+void FNiagaraGpuComputeDispatch::OnPrimitiveGPUSceneInstancesFreed()
+{
+#if RHI_RAYTRACING
+	RayTracingHelper->RefreshPrimitiveInstanceData();
+#endif
+}
+
+/* Set shader parameters for data interfaces
+ */
+void FNiagaraGpuComputeDispatch::SetDataInterfaceParameters(FRHICommandList& RHICmdList, const FNiagaraGPUSystemTick& Tick, const FNiagaraComputeInstanceData& InstanceData, const FNiagaraShaderRef& ComputeShader, const FNiagaraSimStageData& SimStageData) const
+{
+	// @todo-threadsafety This is a bit gross. Need to rethink this api.
+	const FNiagaraShaderMapPointerTable& PointerTable = ComputeShader.GetPointerTable();
+
+	uint32 InterfaceIndex = 0;
+	for (FNiagaraDataInterfaceProxy* Interface : InstanceData.DataInterfaceProxies)
+	{
+		const FNiagaraDataInterfaceParamRef& DIParam = ComputeShader->GetDIParameters()[InterfaceIndex];
+		if (DIParam.Parameters.IsValid())
+		{
+			FNiagaraDataInterfaceSetArgs Context(Interface, Tick.SystemInstanceID, this, ComputeShader, &InstanceData, &SimStageData, InstanceData.IsOutputStage(Interface, SimStageData.StageIndex), InstanceData.IsIterationStage(Interface, SimStageData.StageIndex));
+			DIParam.DIType.Get(PointerTable.DITypes)->SetParameters(DIParam.Parameters.Get(), RHICmdList, Context);
+		}
+
+		InterfaceIndex++;
+	}
+}
+
+void FNiagaraGpuComputeDispatch::UnsetDataInterfaceParameters(FRHICommandList& RHICmdList, const FNiagaraGPUSystemTick& Tick, const FNiagaraComputeInstanceData& InstanceData, const FNiagaraShaderRef& ComputeShader, const FNiagaraSimStageData& SimStageData) const
+{
+	// @todo-threadsafety This is a bit gross. Need to rethink this api.
+	const FNiagaraShaderMapPointerTable& PointerTable = ComputeShader.GetPointerTable();
+
+	uint32 InterfaceIndex = 0;
+	for (FNiagaraDataInterfaceProxy* Interface : InstanceData.DataInterfaceProxies)
+	{
+		const FNiagaraDataInterfaceParamRef& DIParam = ComputeShader->GetDIParameters()[InterfaceIndex];
+		if (DIParam.Parameters.IsValid())
+		{
+			FNiagaraDataInterfaceSetArgs Context(Interface, Tick.SystemInstanceID, this, ComputeShader, &InstanceData, &SimStageData, InstanceData.IsOutputStage(Interface, SimStageData.StageIndex), InstanceData.IsIterationStage(Interface, SimStageData.StageIndex));
+			DIParam.DIType.Get(PointerTable.DITypes)->UnsetParameters(DIParam.Parameters.Get(), RHICmdList, Context);
+		}
+
+		InterfaceIndex++;
+	}
+}
+
+static void SetConstantBuffer(FRHICommandList &RHICmdList, FRHIComputeShader* ComputeShader, const FShaderUniformBufferParameter& BufferParam, const FUniformBufferRHIRef& UniformBuffer)
+{
+	if (BufferParam.IsBound() && UniformBuffer.IsValid())
+	{
+		RHICmdList.SetShaderUniformBuffer(ComputeShader, BufferParam.GetBaseIndex(), UniformBuffer);
+	}
+}
+
+void FNiagaraGpuComputeDispatch::SetConstantBuffers(FRHICommandList &RHICmdList, const FNiagaraShaderRef& Shader, const FNiagaraGPUSystemTick& Tick, const FNiagaraComputeInstanceData* Instance)
+{
+	FNiagaraComputeExecutionContext* Context = Instance->Context;
+	FRHIComputeShader* ComputeShader = Shader.GetComputeShader();
+
+	SetConstantBuffer(RHICmdList, ComputeShader, Shader->GlobalConstantBufferParam[0], Tick.GetUniformBuffer(FNiagaraGPUSystemTick::UBT_Global, nullptr, true));
+	SetConstantBuffer(RHICmdList, ComputeShader, Shader->SystemConstantBufferParam[0], Tick.GetUniformBuffer(FNiagaraGPUSystemTick::UBT_System, nullptr, true));
+	SetConstantBuffer(RHICmdList, ComputeShader, Shader->OwnerConstantBufferParam[0], Tick.GetUniformBuffer(FNiagaraGPUSystemTick::UBT_Owner, nullptr, true));
+	SetConstantBuffer(RHICmdList, ComputeShader, Shader->EmitterConstantBufferParam[0], Tick.GetUniformBuffer(FNiagaraGPUSystemTick::UBT_Emitter, Instance, true));
+	SetConstantBuffer(RHICmdList, ComputeShader, Shader->ExternalConstantBufferParam[0], Tick.GetUniformBuffer(FNiagaraGPUSystemTick::UBT_External, Instance, true));
+
+	if (Context->HasInterpolationParameters)
+	{
+		SetConstantBuffer(RHICmdList, ComputeShader, Shader->GlobalConstantBufferParam[1], Tick.GetUniformBuffer(FNiagaraGPUSystemTick::UBT_Global, nullptr, false));
+		SetConstantBuffer(RHICmdList, ComputeShader, Shader->SystemConstantBufferParam[1], Tick.GetUniformBuffer(FNiagaraGPUSystemTick::UBT_System, nullptr, false));
+		SetConstantBuffer(RHICmdList, ComputeShader, Shader->OwnerConstantBufferParam[1], Tick.GetUniformBuffer(FNiagaraGPUSystemTick::UBT_Owner, nullptr, false));
+		SetConstantBuffer(RHICmdList, ComputeShader, Shader->EmitterConstantBufferParam[1], Tick.GetUniformBuffer(FNiagaraGPUSystemTick::UBT_Emitter, Instance, false));
+		SetConstantBuffer(RHICmdList, ComputeShader, Shader->ExternalConstantBufferParam[1], Tick.GetUniformBuffer(FNiagaraGPUSystemTick::UBT_External, Instance, false));
+	}
+}
+
+FGPUSortManager* FNiagaraGpuComputeDispatch::GetGPUSortManager() const
+{
+	return GPUSortManager;
+}
+
+void FNiagaraGpuComputeDispatch::AddDebugReadback(FNiagaraSystemInstanceID InstanceID, TSharedPtr<struct FNiagaraScriptDebuggerInfo, ESPMode::ThreadSafe> DebugInfo, FNiagaraComputeExecutionContext* Context)
+{
+	FDebugReadbackInfo& ReadbackInfo = GpuDebugReadbackInfos.AddDefaulted_GetRef();
+	ReadbackInfo.InstanceID = InstanceID;
+	ReadbackInfo.DebugInfo = DebugInfo;
+	ReadbackInfo.Context = Context;
+}
+
+bool FNiagaraGpuComputeDispatch::ShouldDebugDraw_RenderThread() const
+{
+#if NIAGARA_COMPUTEDEBUG_ENABLED
+	if (FNiagaraGpuComputeDebug* GpuComputeDebug = GpuComputeDebugPtr.Get())
+	{
+		return GpuComputeDebug->ShouldDrawDebug();
+	}
+#endif
+	return false;
+}
+
+void FNiagaraGpuComputeDispatch::DrawDebug_RenderThread(class FRDGBuilder& GraphBuilder, const class FViewInfo& View, const struct FScreenPassRenderTarget& Output)
+{
+#if NIAGARA_COMPUTEDEBUG_ENABLED
+	if (FNiagaraGpuComputeDebug* GpuComputeDebug = GpuComputeDebugPtr.Get())
+	{
+		GpuComputeDebug->DrawDebug(GraphBuilder, View, Output);
+	}
+#endif
+}
+
+void FNiagaraGpuComputeDispatch::DrawSceneDebug_RenderThread(class FRDGBuilder& GraphBuilder, const class FViewInfo& View, FRDGTextureRef SceneColor, FRDGTextureRef SceneDepth)
+{
+#if NIAGARA_COMPUTEDEBUG_ENABLED
+	if (FNiagaraGpuComputeDebug* GpuComputeDebug = GpuComputeDebugPtr.Get())
+	{
+		GpuComputeDebug->DrawSceneDebug(GraphBuilder, View, SceneColor, SceneDepth);
+	}
+#endif
+}
+
+#if WITH_MGPU
+void FNiagaraGpuComputeDispatch::AddTemporalEffectBuffers(FNiagaraDataBuffer* FinalData)
+{
+	if (GNumAlternateFrameRenderingGroups > 1)
+	{
+		for (FRHIBuffer* Buffer : { FinalData->GetGPUBufferFloat().Buffer, FinalData->GetGPUBufferInt().Buffer, FinalData->GetGPUBufferHalf().Buffer })
+		{
+			if (Buffer)
+			{
+				TemporalEffectBuffers.Add(Buffer);
+			}
+		}
+	}
+}
+
+void FNiagaraGpuComputeDispatch::BroadcastTemporalEffect(FRHICommandList& RHICmdList)
+{
+	if (GNumAlternateFrameRenderingGroups > 1)
+	{
+		if (GPUInstanceCounterManager.GetInstanceCountBuffer().Buffer)
+		{
+			TemporalEffectBuffers.Add(GPUInstanceCounterManager.GetInstanceCountBuffer().Buffer);
+		}
+		if (TemporalEffectBuffers.Num())
+		{
+			RHICmdList.BroadcastTemporalEffect(TemporalEffectName, TemporalEffectBuffers);
+			TemporalEffectBuffers.Reset();
+		}
+	}
+}
+#endif // WITH_MGPU
diff --git a/Engine/Plugins/FX/Niagara/Source/Niagara/Private/NiagaraGpuComputeDispatch.h b/Engine/Plugins/FX/Niagara/Source/Niagara/Private/NiagaraGpuComputeDispatch.h
new file mode 100644
--- /dev/null
+++ b/Engine/Plugins/FX/Niagara/Source/Niagara/Private/NiagaraGpuComputeDispatch.h
@@ -0,0 +1,243 @@
+// Copyright Epic Games, Inc. All Rights Reserved.
+
+/*==============================================================================
+FNiagaraGpuComputeDispatch.h: Queueing and batching for Niagara simulation;
+use to reduce per-simulation overhead by batching together simulations using
+the same VectorVM byte code / compute shader code
+==============================================================================*/
+#pragma once
+
+#include "CoreMinimal.h"
+#include "FXSystem.h"
+#include "Modules/ModuleManager.h"
+#include "ParticleResources.h"
+#include "NiagaraCommon.h"
+#include "NiagaraComputeExecutionContext.h"
+#include "NiagaraEmitter.h"
+#include "NiagaraGpuComputeDispatchInterface.h"
+#include "NiagaraGPUProfiler.h"
+#include "NiagaraGPUSortInfo.h"
+#include "NiagaraGPUSystemTick.h"
+#include "NiagaraParameters.h"
+#include "NiagaraRendererProperties.h"
+#include "NiagaraSystemGpuComputeProxy.h"
+#include "RendererInterface.h"
+#include "RHIResources.h"
+#include "Runtime/Engine/Private/Particles/ParticleSortingGPU.h"
+#include "Tickable.h"
+
+class FGPUSortManager;
+class FNiagaraRayTracingHelper;
+
+class FNiagaraGpuComputeDispatch : public FNiagaraGpuComputeDispatchInterface
+{
+public:
+	static const FName Name;
+
+	virtual FFXSystemInterface* GetInterface(const FName& InName) override;
+
+	FNiagaraGpuComputeDispatch(ERHIFeatureLevel::Type InFeatureLevel, EShaderPlatform InShaderPlatform, FGPUSortManager* InGPUSortManager);
+
+	~FNiagaraGpuComputeDispatch();
+
+	/** Add system instance proxy to the batcher for tracking. */
+	virtual void AddGpuComputeProxy(FNiagaraSystemGpuComputeProxy* ComputeProxy) override;
+
+	/** Remove system instance proxy from the batcher. */
+	virtual void RemoveGpuComputeProxy(FNiagaraSystemGpuComputeProxy* ComputeProxy) override;
+
+#if WITH_EDITOR
+	virtual void Suspend() override {}
+	virtual void Resume() override {}
+#endif // #if WITH_EDITOR
+
+	virtual void DrawDebug(FCanvas* Canvas) override {}
+	virtual bool ShouldDebugDraw_RenderThread() const override;
+	virtual void DrawDebug_RenderThread(class FRDGBuilder& GraphBuilder, const class FViewInfo& View, const struct FScreenPassRenderTarget& Output) override;
+	virtual void DrawSceneDebug_RenderThread(class FRDGBuilder& GraphBuilder, const class FViewInfo& View, FRDGTextureRef SceneColor, FRDGTextureRef SceneDepth) override;
+	virtual void AddVectorField(UVectorFieldComponent* VectorFieldComponent) override {}
+	virtual void RemoveVectorField(UVectorFieldComponent* VectorFieldComponent) override {}
+	virtual void UpdateVectorField(UVectorFieldComponent* VectorFieldComponent) override {}
+	virtual void PreInitViews(FRDGBuilder& GraphBuilder, bool bAllowGPUParticleUpdate) override;
+	virtual void PostInitViews(FRDGBuilder& GraphBuilder, TConstArrayView<FViewInfo> Views, bool bAllowGPUParticleUpdate) override;
+	virtual bool UsesGlobalDistanceField() const override;
+	virtual bool UsesDepthBuffer() const override;
+	virtual bool RequiresEarlyViewUniformBuffer() const override;
+	virtual bool RequiresRayTracingScene() const override;
+	virtual void PreRender(FRDGBuilder& GraphBuilder, TConstArrayView<FViewInfo> Views, bool bAllowGPUParticleUpdate) override;
+	virtual void OnDestroy() override; // Called on the gamethread to delete the batcher on the renderthread.
+
+	virtual void Tick(UWorld* World, float DeltaTime) override;
+
+	virtual void PostRenderOpaque(FRDGBuilder& GraphBuilder, TConstArrayView<FViewInfo> Views, bool bAllowGPUParticleUpdate) override;
+
+	// FNiagaraGpuComputeDispatchInterface Impl
+	virtual void FlushPendingTicks_GameThread() override;
+	// FNiagaraGpuComputeDispatchInterface Impl
+
+	/**
+	 * Process and respond to a build up of excessive ticks inside the batcher.
+	 * In the case of the application not having focus the game thread may continue
+	 * to process and send ticks to the render thread but the rendering thread may
+	 * never process them.  The World Manager will ensure this is called once per
+	 * game frame so we have an opportunity to flush the ticks avoiding a stall
+	 * when we gain focus again.
+	 */
+	void ProcessPendingTicksFlush(FRHICommandListImmediate& RHICmdList, bool bForceFlush);
+
+	/** Processes all pending readbacks */
+	virtual void ProcessDebugReadbacks(FRHICommandListImmediate& RHICmdList, bool bWaitCompletion) override;
+
+	virtual bool AddSortedGPUSimulation(FNiagaraGPUSortInfo& SortInfo) override;
+
+	const FGlobalDistanceFieldParameterData* GetGlobalDistanceFieldParameters() const;
+
+	void SetDataInterfaceParameters(FRHICommandList& RHICmdList, const FNiagaraGPUSystemTick& Tick, const FNiagaraComputeInstanceData& InstanceData, const FNiagaraShaderRef& ComputeShader, const FNiagaraSimStageData& SimStageData) const;
+	void UnsetDataInterfaceParameters(FRHICommandList& RHICmdList, const FNiagaraGPUSystemTick& Tick, const FNiagaraComputeInstanceData& InstanceData, const FNiagaraShaderRef& ComputeShader, const FNiagaraSimStageData& SimStageData) const;
+
+	/** Reset the data interfaces and check if the spawn stages are valid */
+	void ResetDataInterfaces(FRHICommandList& RHICmdList, const FNiagaraGPUSystemTick& Tick, const FNiagaraComputeInstanceData& InstanceData) const;
+
+	/** Given a shader stage index, find the corresponding data interface */
+	FNiagaraDataInterfaceProxyRW* FindIterationInterface(FNiagaraComputeInstanceData* Instance, const uint32 SimulationStageIndex) const;
+
+	/** Loop over all the data interfaces and call the pre-stage methods */
+	void PreStageInterface(FRHICommandList& RHICmdList, const FNiagaraGPUSystemTick& Tick, const FNiagaraComputeInstanceData& InstanceData, const FNiagaraSimStageData& SimStageData, TSet<FNiagaraDataInterfaceProxy*>& ProxiesToFinalize) const;
+
+	/** Loop over all the data interfaces and call the post-stage methods */
+	void PostStageInterface(FRHICommandList& RHICmdList, const FNiagaraGPUSystemTick& Tick, const FNiagaraComputeInstanceData& InstanceData, const FNiagaraSimStageData& SimStageData, TSet<FNiagaraDataInterfaceProxy*>& ProxiesToFinalize) const;
+
+	/** Loop over all data interfaces and call the post-simulate methods */
+	void PostSimulateInterface(FRHICommandList& RHICmdList, const FNiagaraGPUSystemTick& Tick, const FNiagaraComputeInstanceData& InstanceData) const;
+
+	/** Get the shared SortManager, used in the rendering loop to call FGPUSortManager::OnPreRender() and FGPUSortManager::OnPostRenderOpaque() */
+	virtual FGPUSortManager* GetGPUSortManager() const override;
+
+	/** When inside a pass returns the active scene textures for the pass, can be nullptr. */
+	class FNiagaraSceneTextureParameters* GetNiagaraSceneTextures() const { return NiagaraSceneTextures; }
+
+	/** Debug only function to readback data. */
+	virtual void AddDebugReadback(FNiagaraSystemInstanceID InstanceID, TSharedPtr<struct FNiagaraScriptDebuggerInfo, ESPMode::ThreadSafe> DebugInfo, FNiagaraComputeExecutionContext* Context) override;
+
+#if RHI_RAYTRACING
+	virtual FNiagaraRayTracingHelper& GetRayTracingHelper() const override;
+	virtual bool HasRayTracingScene() const override;
+
+	virtual void SetPrimitiveRayTracingCollisionGroup(UPrimitiveComponent* PrimitiveComponent, uint32 Group) override;
+
+	virtual int32 AcquireGPURayTracedCollisionGroup() override;
+	virtual void ReleaseGPURayTracedCollisionGroup(int32 CollisionGroup) override;
+#endif
+
+	void OnPrimitiveGPUSceneInstancesAllocated();
+	void OnPrimitiveGPUSceneInstancesFreed();
+
+private:
+	void DumpDebugFrame();
+	void UpdateInstanceCountManager(FRHICommandListImmediate& RHICmdList);
+	void PrepareTicksForProxy(FRHICommandListImmediate& RHICmdList, FNiagaraSystemGpuComputeProxy* ComputeProxy, FNiagaraGpuDispatchList& GpuDispatchList);
+	void PrepareAllTicks(FRHICommandListImmediate& RHICmdList);
+	void ExecuteTicks(FRHICommandList& RHICmdList, FRHIUniformBuffer* ViewUniformBuffer, ENiagaraGpuComputeTickStage::Type TickStage);
+	void DispatchStage(FRHICommandList& RHICmdList, FRHIUniformBuffer* ViewUniformBuffer, const FNiagaraGPUSystemTick& Tick, const FNiagaraComputeInstanceData& InstanceData, const FNiagaraSimStageData& SimStageData);
+
+	/**
+	 * Generate all the initial keys and values for a GPUSortManager sort batch.
+	 * Sort batches are created when GPU sort tasks are registered in AddSortedGPUSimulation().
+	 * Each sort task defines constraints about when the initial sort data can generated and
+	 * and when the sorted results are needed (see EGPUSortFlags for details).
+	 * Currently, for Niagara, all the sort tasks have the EGPUSortFlags::KeyGenAfterPreRender flag
+	 * and so the callback registered in GPUSortManager->Register() only has the EGPUSortFlags::KeyGenAfterPreRender usage.
+	 * This garanties that GenerateSortKeys() only gets called after PreRender(), which is a constraint required because
+	 * Niagara renders the current state of the GPU emitters, before the are ticked
+	 * (Niagara GPU emitters are ticked at InitView and in PostRenderOpaque).
+	 * Note that this callback must only initialize the content for the elements that relates to the tasks it has registered in this batch.
+	 *
+	 * @param RHICmdList - The command list used to initiate the keys and values on GPU.
+	 * @param BatchId - The GPUSortManager batch id (regrouping several similar sort tasks).
+	 * @param NumElementsInBatch - The number of elements grouped in the batch (each element maps to a sort task)
+	 * @param Flags - Details about the key precision (see EGPUSortFlags::AnyKeyPrecision) and the keygen location (see EGPUSortFlags::AnyKeyGenLocation).
+	 * @param KeysUAV - The UAV that holds all the initial keys used to sort the values (being the particle indices here).
+	 * @param ValuesUAV - The UAV that holds the initial values (particle indices) to be sorted accordingly to the keys.
+	 */
+	void GenerateSortKeys(FRHICommandListImmediate& RHICmdList, int32 BatchId, int32 NumElementsInBatch, EGPUSortFlags Flags, FRHIUnorderedAccessView* KeysUAV, FRHIUnorderedAccessView* ValuesUAV);
+
+#if RHI_RAYTRACING
+	void BuildRayTracingSceneInfo(FRHICommandList& RHICmdList, TConstArrayView<FViewInfo> Views);
+	void ResetRayTracingSceneInfo();
+#endif
+
+	void FinishDispatches();
+
+	void UpdateFreeIDsListSizesBuffer(FRHICommandList& RHICmdList, uint32 NumInstances);
+	void UpdateFreeIDBuffers(FRHICommandList& RHICmdList, TConstArrayView<FNiagaraComputeExecutionContext*> Instances);
+
+	void SetConstantBuffers(FRHICommandList &RHICmdList, const FNiagaraShaderRef& Shader, const FNiagaraGPUSystemTick& Tick, const FNiagaraComputeInstanceData* Instance);
+	void BuildConstantBuffers(FNiagaraGPUSystemTick& Tick) const;
+
+	/** The shared GPUSortManager, used to register GPU sort tasks in order to generate sorted particle indices per emitter. */
+	TRefCountPtr<FGPUSortManager> GPUSortManager;
+	/** All sort tasks registered in AddSortedGPUSimulation(). Holds all the data required in GenerateSortKeys(). */
+	TArray<FNiagaraGPUSortInfo> SimulationsToSort;
+
+#if RHI_RAYTRACING
+	// helper object which is valid during PostRenderOpaque if DataInterfaces are present that require a RayTrace scene
+	TUniquePtr<FNiagaraRayTracingHelper> RayTracingHelper;
+#endif
+
+#if STATS
+	FNiagaraGPUProfiler GPUProfiler;
+#endif
+
+	// persistent layouts used to create the constant buffers for the compute sim shader
+	TRefCountPtr<FNiagaraRHIUniformBufferLayout> GlobalCBufferLayout;
+	TRefCountPtr<FNiagaraRHIUniformBufferLayout> SystemCBufferLayout;
+	TRefCountPtr<FNiagaraRHIUniformBufferLayout> OwnerCBufferLayout;
+	TRefCountPtr<FNiagaraRHIUniformBufferLayout> EmitterCBufferLayout;
+
+	uint32 FramesBeforeTickFlush = 0;
+
+	TConstArrayView<FViewInfo> CurrentPassViews;
+
+	/** A buffer of list sizes used by UpdateFreeIDBuffers to allow overlapping several dispatches. */
+	FRWBuffer FreeIDListSizesBuffer;
+	uint32 NumAllocatedFreeIDListSizes = 0;
+
+	uint32 NumProxiesThatRequireDistanceFieldData = 0;
+	uint32 NumProxiesThatRequireDepthBuffer = 0;
+	uint32 NumProxiesThatRequireEarlyViewData = 0;
+	uint32 NumProxiesThatRequireRayTracingScene = 0;
+
+	int32 TotalDispatchesThisFrame = 0;
+
+	bool bRequiresReadback = false;
+	TArray<FNiagaraSystemGpuComputeProxy*> ProxiesPerStage[ENiagaraGpuComputeTickStage::Max];
+
+	FNiagaraGpuDispatchList DispatchListPerStage[ENiagaraGpuComputeTickStage::Max];
+
+	class FNiagaraSceneTextureParameters* NiagaraSceneTextures = nullptr;
+
+	struct FDebugReadbackInfo
+	{
+		FNiagaraSystemInstanceID InstanceID;
+		TSharedPtr<struct FNiagaraScriptDebuggerInfo, ESPMode::ThreadSafe> DebugInfo;
+		FNiagaraComputeExecutionContext* Context;
+	};
+	TArray<FDebugReadbackInfo> GpuDebugReadbackInfos;
+
+#if WITH_MGPU
+	static const FName TemporalEffectName;
+	TArray<FRHIBuffer*> TemporalEffectBuffers;
+	ENiagaraGpuComputeTickStage::Type StageToWaitForTemporalEffect = ENiagaraGpuComputeTickStage::First;
+	ENiagaraGpuComputeTickStage::Type StageToBroadcastTemporalEffect = ENiagaraGpuComputeTickStage::First;
+
+	void AddTemporalEffectBuffers(FNiagaraDataBuffer* FinalData);
+	void BroadcastTemporalEffect(FRHICommandList& RHICmdList);
+#endif // WITH_MGPU
+
+	/** Pool of free GPU ray traced collision groups. */
+	TArray<int32> FreeGPURayTracedCollisionGroups;
+	int32 NumGPURayTracedCollisionGroups = 0;
+
+	// Cached information to build a dummy view info if neccessary
+	FIntRect CachedViewRect;
+};
diff --git a/Engine/Plugins/FX/Niagara/Source/Niagara/Private/NiagaraGpuComputeDispatchInterface.cpp b/Engine/Plugins/FX/Niagara/Source/Niagara/Private/NiagaraGpuComputeDispatchInterface.cpp
new file mode 100644
--- /dev/null
+++ b/Engine/Plugins/FX/Niagara/Source/Niagara/Private/NiagaraGpuComputeDispatchInterface.cpp
@@ -0,0 +1,35 @@
+// Copyright Epic Games, Inc. All Rights Reserved.
+
+#include "NiagaraGpuComputeDispatchInterface.h"
+#include "NiagaraGpuComputeDispatch.h"
+#include "NiagaraGpuComputeDebug.h"
+#include "NiagaraGpuReadbackManager.h"
+
+#include "Engine/World.h"
+#include "SceneInterface.h"
+#include "FXSystem.h"
+
+FNiagaraGpuComputeDispatchInterface::FNiagaraGpuComputeDispatchInterface(EShaderPlatform InShaderPlatform, ERHIFeatureLevel::Type InFeatureLevel)
+	: ShaderPlatform(InShaderPlatform)
+	, FeatureLevel(InFeatureLevel)
+{
+}
+
+FNiagaraGpuComputeDispatchInterface::~FNiagaraGpuComputeDispatchInterface()
+{
+}
+
+FNiagaraGpuComputeDispatchInterface* FNiagaraGpuComputeDispatchInterface::Get(UWorld* World)
+{
+	return World ? Get(World->Scene) : nullptr;
+}
+
+FNiagaraGpuComputeDispatchInterface* FNiagaraGpuComputeDispatchInterface::Get(FSceneInterface* Scene)
+{
+	return Scene ? Get(Scene->GetFXSystem()) : nullptr;
+}
+
+FNiagaraGpuComputeDispatchInterface* FNiagaraGpuComputeDispatchInterface::Get(FFXSystemInterface* FXSceneInterface)
+{
+	return FXSceneInterface ? static_cast<FNiagaraGpuComputeDispatchInterface*>(FXSceneInterface->GetInterface(FNiagaraGpuComputeDispatch::Name)) : nullptr;
+}
diff --git a/Engine/Plugins/FX/Niagara/Source/Niagara/Private/NiagaraModule.cpp b/Engine/Plugins/FX/Niagara/Source/Niagara/Private/NiagaraModule.cpp
--- a/Engine/Plugins/FX/Niagara/Source/Niagara/Private/NiagaraModule.cpp
+++ b/Engine/Plugins/FX/Niagara/Source/Niagara/Private/NiagaraModule.cpp
@@ -22,7 +22,7 @@
 #include "NiagaraRenderer.h"
 #include "NiagaraShaderModule.h"
 #include "UObject/CoreRedirects.h"
-#include "NiagaraEmitterInstanceBatcher.h"
+#include "NiagaraGpuComputeDispatch.h"
 #include "Misc/CoreDelegates.h"
 #include "NiagaraDebuggerClient.h"
 #include "Particles/FXBudget.h"
@@ -332,10 +332,10 @@ void INiagaraModule::StartupModule()
 	}));
 
 	FFXSystemInterface::RegisterCustomFXSystem(
-		NiagaraEmitterInstanceBatcher::Name, 
+		FNiagaraGpuComputeDispatch::Name,
 		FCreateCustomFXSystemDelegate::CreateLambda([](ERHIFeatureLevel::Type InFeatureLevel, EShaderPlatform InShaderPlatform, FGPUSortManager* InGPUSortManager) -> FFXSystemInterface*
 	{
-		return new NiagaraEmitterInstanceBatcher(InFeatureLevel, InShaderPlatform, InGPUSortManager);
+		return new FNiagaraGpuComputeDispatch(InFeatureLevel, InShaderPlatform, InGPUSortManager);
 	}));
 
 	// Needed for NiagaraDataInterfaceAudioSpectrum
@@ -420,7 +420,7 @@ void INiagaraModule::OnWorldBeginTearDown(UWorld* World)
 
 void INiagaraModule::ShutdownRenderingResources()
 {
-	FFXSystemInterface::UnregisterCustomFXSystem(NiagaraEmitterInstanceBatcher::Name);
+	FFXSystemInterface::UnregisterCustomFXSystem(FNiagaraGpuComputeDispatch::Name);
 
 	FNiagaraRenderViewDataManager::Shutdown();
 }
diff --git a/Engine/Plugins/FX/Niagara/Source/Niagara/Private/NiagaraRenderer.cpp b/Engine/Plugins/FX/Niagara/Source/Niagara/Private/NiagaraRenderer.cpp
--- a/Engine/Plugins/FX/Niagara/Source/Niagara/Private/NiagaraRenderer.cpp
+++ b/Engine/Plugins/FX/Niagara/Source/Niagara/Private/NiagaraRenderer.cpp
@@ -9,7 +9,7 @@
 #include "NiagaraComponent.h"
 #include "Engine/Engine.h"
 #include "DynamicBufferAllocator.h"
-#include "NiagaraEmitterInstanceBatcher.h"
+#include "NiagaraComputeExecutionContext.h"
 #include "NiagaraGPUSortInfo.h"
 #include "NiagaraEmitterInstance.h"
 #include "NiagaraSystemInstanceController.h"
diff --git a/Engine/Plugins/FX/Niagara/Source/Niagara/Private/NiagaraRendererMeshes.cpp b/Engine/Plugins/FX/Niagara/Source/Niagara/Private/NiagaraRendererMeshes.cpp
--- a/Engine/Plugins/FX/Niagara/Source/Niagara/Private/NiagaraRendererMeshes.cpp
+++ b/Engine/Plugins/FX/Niagara/Source/Niagara/Private/NiagaraRendererMeshes.cpp
@@ -8,9 +8,9 @@
 #include "NiagaraComponent.h"
 #include "Async/ParallelFor.h"
 #include "Engine/StaticMesh.h"
-#include "NiagaraEmitterInstanceBatcher.h"
 #include "NiagaraSortingGPU.h"
 #include "NiagaraGPURayTracingTransformsShader.h"
+#include "NiagaraGpuComputeDispatchInterface.h"
 #include "RayTracingDefinitions.h"
 #include "RayTracingDynamicGeometryCollection.h"
 #include "RayTracingInstance.h"
@@ -258,7 +258,7 @@ void FNiagaraRendererMeshes::PrepareParticleMeshRenderData(FParticleMeshRenderDa
 {
 	// Anything to render?
 	ParticleMeshRenderData.DynamicDataMesh = static_cast<FNiagaraDynamicDataMesh*>(InDynamicData);
-	if (!ParticleMeshRenderData.DynamicDataMesh || !SceneProxy->GetBatcher())
+	if (!ParticleMeshRenderData.DynamicDataMesh || !SceneProxy->GetComputeDispatchInterface())
 	{
 		return;
 	}
@@ -294,7 +294,7 @@ void FNiagaraRendererMeshes::PrepareParticleMeshRenderData(FParticleMeshRenderDa
 	// Particle source mode
 	if (SourceMode == ENiagaraRendererSourceDataMode::Particles)
 	{
-		const EShaderPlatform ShaderPlatform = SceneProxy->GetBatcher()->GetShaderPlatform();
+		const EShaderPlatform ShaderPlatform = SceneProxy->GetComputeDispatchInterface()->GetShaderPlatform();
 
 		// Determine if we need sorting
 		ParticleMeshRenderData.bNeedsSort = SortMode != ENiagaraSortMode::None && (ParticleMeshRenderData.bHasTranslucentMaterials || !bSortOnlyWhenTranslucent);
@@ -424,7 +424,7 @@ void FNiagaraRendererMeshes::InitializeSortInfo(FParticleMeshRenderData& Particl
 
 	if (ParticleMeshRenderData.bSortCullOnGpu)
 	{
-		NiagaraEmitterInstanceBatcher* Batcher = SceneProxy.GetBatcher();
+		FNiagaraGpuComputeDispatchInterface* ComputeDispatchInterface = SceneProxy.GetComputeDispatchInterface();
 
 		OutSortInfo.ParticleDataFloatSRV = ParticleMeshRenderData.ParticleFloatSRV;
 		OutSortInfo.ParticleDataHalfSRV = ParticleMeshRenderData.ParticleHalfSRV;
@@ -432,7 +432,7 @@ void FNiagaraRendererMeshes::InitializeSortInfo(FParticleMeshRenderData& Particl
 		OutSortInfo.FloatDataStride = ParticleMeshRenderData.ParticleFloatDataStride;
 		OutSortInfo.HalfDataStride = ParticleMeshRenderData.ParticleHalfDataStride;
 		OutSortInfo.IntDataStride = ParticleMeshRenderData.ParticleIntDataStride;
-		OutSortInfo.GPUParticleCountSRV = GetSrvOrDefaultUInt(Batcher->GetGPUInstanceCounterManager().GetInstanceCountBuffer());
+		OutSortInfo.GPUParticleCountSRV = GetSrvOrDefaultUInt(ComputeDispatchInterface->GetGPUInstanceCounterManager().GetInstanceCountBuffer());
 		OutSortInfo.GPUParticleCountOffset = ParticleMeshRenderData.SourceParticleData->GetGPUInstanceCountBufferOffset();
 	}
 
@@ -585,7 +585,7 @@ void FNiagaraRendererMeshes::PreparePerMeshData(FParticleMeshRenderData& Particl
 	}
 }
 
-uint32 FNiagaraRendererMeshes::PerformSortAndCull(FParticleMeshRenderData& ParticleMeshRenderData, FGlobalDynamicReadBuffer& ReadBuffer, FNiagaraGPUSortInfo& SortInfo, NiagaraEmitterInstanceBatcher* Batcher, int32 MeshIndex) const
+uint32 FNiagaraRendererMeshes::PerformSortAndCull(FParticleMeshRenderData& ParticleMeshRenderData, FGlobalDynamicReadBuffer& ReadBuffer, FNiagaraGPUSortInfo& SortInfo, FNiagaraGpuComputeDispatchInterface* ComputeDispatchInterface, int32 MeshIndex) const
 {
 	// Emitter mode culls earlier on
 	if (SourceMode == ENiagaraRendererSourceDataMode::Emitter)
@@ -603,8 +603,8 @@ uint32 FNiagaraRendererMeshes::PerformSortAndCull(FParticleMeshRenderData& Parti
 		SortInfo.MeshIndex = MeshIndex;
 		if (ParticleMeshRenderData.bSortCullOnGpu)
 		{
-			SortInfo.CulledGPUParticleCountOffset = ParticleMeshRenderData.bNeedsCull ? Batcher->GetGPUInstanceCounterManager().AcquireCulledEntry() : INDEX_NONE;
-			if (Batcher->AddSortedGPUSimulation(SortInfo))
+			SortInfo.CulledGPUParticleCountOffset = ParticleMeshRenderData.bNeedsCull ? ComputeDispatchInterface->GetGPUInstanceCounterManager().AcquireCulledEntry() : INDEX_NONE;
+			if (ComputeDispatchInterface->AddSortedGPUSimulation(SortInfo))
 			{
 				ParticleMeshRenderData.ParticleSortedIndicesSRV = SortInfo.AllocationInfo.BufferSRV;
 				ParticleMeshRenderData.ParticleSortedIndicesOffset = SortInfo.AllocationInfo.BufferOffset;
@@ -899,10 +899,10 @@ void FNiagaraRendererMeshes::CreateMeshBatchForSection(
 	if ((SourceMode == ENiagaraRendererSourceDataMode::Particles) && (GPUCountBufferOffset != INDEX_NONE))
 	{
 		// We need to use indirect draw args, because the number of actual instances is coming from the GPU
-		auto Batcher = SceneProxy.GetBatcher();
-		check(Batcher);
+		auto ComputeDispatchInterface = SceneProxy.GetComputeDispatchInterface();
+		check(ComputeDispatchInterface);
 
-		auto& CountManager = Batcher->GetGPUInstanceCounterManager();
+		auto& CountManager = ComputeDispatchInterface->GetGPUInstanceCounterManager();
 		auto IndirectDraw = CountManager.AddDrawIndirect(
 			GPUCountBufferOffset,
 			Section.NumTriangles * 3,
@@ -982,7 +982,7 @@ void FNiagaraRendererMeshes::GetDynamicMeshElements(const TArray<const FSceneVie
 				InitializeSortInfo(ParticleMeshRenderData, *SceneProxy, *View, ViewIndex, bIsInstancedStereo, SortInfo);
 			}
 
-			NiagaraEmitterInstanceBatcher* Batcher = SceneProxy->GetBatcher();
+			FNiagaraGpuComputeDispatchInterface* ComputeDispatchInterface = SceneProxy->GetComputeDispatchInterface();
 			for (int32 MeshIndex = 0; MeshIndex < Meshes.Num(); ++MeshIndex)
 			{
 				// No binding for mesh index or we don't allow culling we will only render the first mesh for all particles
@@ -1025,7 +1025,7 @@ void FNiagaraRendererMeshes::GetDynamicMeshElements(const TArray<const FSceneVie
 				PreparePerMeshData(ParticleMeshRenderData, *SceneProxy, MeshData);
 
 				// Sort/Cull particles if needed.
-				const uint32 NumInstances = PerformSortAndCull(ParticleMeshRenderData, Collector.GetDynamicReadBuffer(), SortInfo, Batcher, MeshData.SourceMeshIndex);
+				const uint32 NumInstances = PerformSortAndCull(ParticleMeshRenderData, Collector.GetDynamicReadBuffer(), SortInfo, ComputeDispatchInterface, MeshData.SourceMeshIndex);
 				if ( NumInstances > 0 )
 				{
 					// Increment stats
@@ -1154,8 +1154,8 @@ void FNiagaraRendererMeshes::GetDynamicRayTracingInstances(FRayTracingMaterialGa
 		PreparePerMeshData(ParticleMeshRenderData, *SceneProxy, MeshData);
 
 		// Sort/Cull particles if needed.
-		NiagaraEmitterInstanceBatcher* Batcher = SceneProxy->GetBatcher();
-		const uint32 NumInstances = PerformSortAndCull(ParticleMeshRenderData, Context.RayTracingMeshResourceCollector.GetDynamicReadBuffer(), SortInfo, Batcher, MeshData.SourceMeshIndex);
+		FNiagaraGpuComputeDispatchInterface* ComputeDispatchInterface = SceneProxy->GetComputeDispatchInterface();
+		const uint32 NumInstances = PerformSortAndCull(ParticleMeshRenderData, Context.RayTracingMeshResourceCollector.GetDynamicReadBuffer(), SortInfo, ComputeDispatchInterface, MeshData.SourceMeshIndex);
 		if ( NumInstances == 0 )
 		{
 			continue;
@@ -1392,7 +1392,7 @@ void FNiagaraRendererMeshes::GetDynamicRayTracingInstances(FRayTracingMaterialGa
 					PassParameters.ParticleDataFloatBuffer		= ParticleMeshRenderData.ParticleFloatSRV;
 					//PassParameters.ParticleDataHalfBuffer		= ParticleMeshRenderData.ParticleHalfSRV;
 					PassParameters.ParticleDataIntBuffer		= ParticleMeshRenderData.ParticleIntSRV;
-					PassParameters.GPUInstanceCountBuffer		= Batcher->GetGPUInstanceCounterManager().GetInstanceCountBuffer().SRV;
+					PassParameters.GPUInstanceCountBuffer		= ComputeDispatchInterface->GetGPUInstanceCounterManager().GetInstanceCountBuffer().SRV;
 					PassParameters.TLASTransforms				= InstanceGPUTransformsBuffer.UAV;
 				}
 
diff --git a/Engine/Plugins/FX/Niagara/Source/Niagara/Private/NiagaraRendererRibbons.cpp b/Engine/Plugins/FX/Niagara/Source/Niagara/Private/NiagaraRendererRibbons.cpp
--- a/Engine/Plugins/FX/Niagara/Source/Niagara/Private/NiagaraRendererRibbons.cpp
+++ b/Engine/Plugins/FX/Niagara/Source/Niagara/Private/NiagaraRendererRibbons.cpp
@@ -235,9 +235,9 @@ void FNiagaraRendererRibbons::ReleaseRenderThreadResources()
 }
 
 // FPrimitiveSceneProxy interface.
-void FNiagaraRendererRibbons::CreateRenderThreadResources(NiagaraEmitterInstanceBatcher* Batcher)
+void FNiagaraRendererRibbons::CreateRenderThreadResources()
 {
-	FNiagaraRenderer::CreateRenderThreadResources(Batcher);
+	FNiagaraRenderer::CreateRenderThreadResources();
 #if RHI_RAYTRACING
 	if (IsRayTracingEnabled())
 	{
@@ -1436,8 +1436,8 @@ void FNiagaraRendererRibbons::GetDynamicRayTracingInstances(FRayTracingMaterialG
 	check(SceneProxy);
 
 	FNiagaraDynamicDataRibbon *DynamicDataRibbon = static_cast<FNiagaraDynamicDataRibbon*>(DynamicDataRender);
-	NiagaraEmitterInstanceBatcher* Batcher = SceneProxy->GetBatcher();
-	if (!DynamicDataRibbon || !Batcher)
+	FNiagaraGpuComputeDispatchInterface* ComputeDispatchInterface = SceneProxy->GetComputeDispatchInterface();
+	if (!DynamicDataRibbon || !ComputeDispatchInterface)
 	{
 		return;
 	}
diff --git a/Engine/Plugins/FX/Niagara/Source/Niagara/Private/NiagaraRendererSprites.cpp b/Engine/Plugins/FX/Niagara/Source/Niagara/Private/NiagaraRendererSprites.cpp
--- a/Engine/Plugins/FX/Niagara/Source/Niagara/Private/NiagaraRendererSprites.cpp
+++ b/Engine/Plugins/FX/Niagara/Source/Niagara/Private/NiagaraRendererSprites.cpp
@@ -4,10 +4,10 @@
 #include "ParticleResources.h"
 #include "NiagaraDataSet.h"
 #include "NiagaraStats.h"
-#include "NiagaraEmitterInstanceBatcher.h"
 #include "NiagaraSortingGPU.h"
 #include "NiagaraComponent.h"
 #include "NiagaraCutoutVertexBuffer.h"
+#include "NiagaraGpuComputeDispatchInterface.h"
 #include "RayTracingDefinitions.h"
 #include "RayTracingDynamicGeometryCollection.h"
 #include "RayTracingInstance.h"
@@ -178,9 +178,9 @@ void FNiagaraRendererSprites::ReleaseRenderThreadResources()
 #endif
 }
 
-void FNiagaraRendererSprites::CreateRenderThreadResources(NiagaraEmitterInstanceBatcher* Batcher)
+void FNiagaraRendererSprites::CreateRenderThreadResources()
 {
-	FNiagaraRenderer::CreateRenderThreadResources(Batcher);
+	FNiagaraRenderer::CreateRenderThreadResources();
 	CutoutVertexBuffer.InitResource();
 
 #if RHI_RAYTRACING
@@ -203,7 +203,7 @@ void FNiagaraRendererSprites::CreateRenderThreadResources(NiagaraEmitterInstance
 void FNiagaraRendererSprites::PrepareParticleSpriteRenderData(FParticleSpriteRenderData& ParticleSpriteRenderData, FNiagaraDynamicDataBase* InDynamicData, const FNiagaraSceneProxy* SceneProxy) const
 {
 	ParticleSpriteRenderData.DynamicDataSprites = static_cast<FNiagaraDynamicDataSprites*>(InDynamicData);
-	if (!ParticleSpriteRenderData.DynamicDataSprites || !SceneProxy->GetBatcher())
+	if (!ParticleSpriteRenderData.DynamicDataSprites || !SceneProxy->GetComputeDispatchInterface())
 	{
 		ParticleSpriteRenderData.SourceParticleData = nullptr;
 		return;
@@ -244,7 +244,7 @@ void FNiagaraRendererSprites::PrepareParticleSpriteRenderData(FParticleSpriteRen
 	// Particle source mode
 	if (SourceMode == ENiagaraRendererSourceDataMode::Particles)
 	{
-		const EShaderPlatform ShaderPlatform = SceneProxy->GetBatcher()->GetShaderPlatform();
+		const EShaderPlatform ShaderPlatform = SceneProxy->GetComputeDispatchInterface()->GetShaderPlatform();
 
 		// Determine if we need sorting
 		ParticleSpriteRenderData.bNeedsSort = SortMode != ENiagaraSortMode::None && (BlendMode == BLEND_AlphaComposite || BlendMode == BLEND_AlphaHoldout || BlendMode == BLEND_Translucent || !bSortOnlyWhenTranslucent);
@@ -395,7 +395,7 @@ void FNiagaraRendererSprites::InitializeSortInfo(FParticleSpriteRenderData& Part
 
 	if (ParticleSpriteRenderData.bSortCullOnGpu)
 	{
-		NiagaraEmitterInstanceBatcher* Batcher = SceneProxy.GetBatcher();
+		FNiagaraGpuComputeDispatchInterface* ComputeDispatchInterface = SceneProxy.GetComputeDispatchInterface();
 
 		OutSortInfo.ParticleDataFloatSRV = ParticleSpriteRenderData.ParticleFloatSRV;
 		OutSortInfo.ParticleDataHalfSRV = ParticleSpriteRenderData.ParticleHalfSRV;
@@ -403,7 +403,7 @@ void FNiagaraRendererSprites::InitializeSortInfo(FParticleSpriteRenderData& Part
 		OutSortInfo.FloatDataStride = ParticleSpriteRenderData.ParticleFloatDataStride;
 		OutSortInfo.HalfDataStride = ParticleSpriteRenderData.ParticleHalfDataStride;
 		OutSortInfo.IntDataStride = ParticleSpriteRenderData.ParticleIntDataStride;
-		OutSortInfo.GPUParticleCountSRV = GetSrvOrDefaultUInt(Batcher->GetGPUInstanceCounterManager().GetInstanceCountBuffer());
+		OutSortInfo.GPUParticleCountSRV = GetSrvOrDefaultUInt(ComputeDispatchInterface->GetGPUInstanceCounterManager().GetInstanceCountBuffer());
 		OutSortInfo.GPUParticleCountOffset = ParticleSpriteRenderData.SourceParticleData->GetGPUInstanceCountBufferOffset();
 	}
 
@@ -761,10 +761,10 @@ void FNiagaraRendererSprites::CreateMeshBatchForView(
 	FNiagaraGPUInstanceCountManager::FIndirectArgSlot IndirectDraw;
 	if ((SourceMode == ENiagaraRendererSourceDataMode::Particles) && (GPUCountBufferOffset != INDEX_NONE))
 	{
-		NiagaraEmitterInstanceBatcher* Batcher = SceneProxy.GetBatcher();
-		check(Batcher);
+		FNiagaraGpuComputeDispatchInterface* ComputeDispatchInterface = SceneProxy.GetComputeDispatchInterface();
+		check(ComputeDispatchInterface);
 
-		IndirectDraw = Batcher->GetGPUInstanceCounterManager().AddDrawIndirect(
+		IndirectDraw = ComputeDispatchInterface->GetGPUInstanceCounterManager().AddDrawIndirect(
 			GPUCountBufferOffset,
 			NumIndicesPerInstance,
 			0,
@@ -913,13 +913,13 @@ void FNiagaraRendererSprites::GetDynamicMeshElements(const TArray<const FSceneVi
 			uint32 NumInstances = SourceMode == ENiagaraRendererSourceDataMode::Particles ? ParticleSpriteRenderData.SourceParticleData->GetNumInstances() : 1;
 
 			VertexFactory.SetSortedIndices(nullptr, 0xFFFFFFFF);
-			NiagaraEmitterInstanceBatcher* Batcher = SceneProxy->GetBatcher();
+			FNiagaraGpuComputeDispatchInterface* ComputeDispatchInterface = SceneProxy->GetComputeDispatchInterface();
 			if (ParticleSpriteRenderData.bNeedsCull || ParticleSpriteRenderData.bNeedsSort)
 			{
 				if (ParticleSpriteRenderData.bSortCullOnGpu)
 				{
-					SortInfo.CulledGPUParticleCountOffset = ParticleSpriteRenderData.bNeedsCull ? Batcher->GetGPUInstanceCounterManager().AcquireCulledEntry() : INDEX_NONE;
-					if (Batcher->AddSortedGPUSimulation(SortInfo))
+					SortInfo.CulledGPUParticleCountOffset = ParticleSpriteRenderData.bNeedsCull ? ComputeDispatchInterface->GetGPUInstanceCounterManager().AcquireCulledEntry() : INDEX_NONE;
+					if (ComputeDispatchInterface->AddSortedGPUSimulation(SortInfo))
 					{
 						VertexFactory.SetSortedIndices(SortInfo.AllocationInfo.BufferSRV, SortInfo.AllocationInfo.BufferOffset);
 					}
@@ -1004,13 +1004,13 @@ void FNiagaraRendererSprites::GetDynamicRayTracingInstances(FRayTracingMaterialG
 	uint32 NumInstances = SourceMode == ENiagaraRendererSourceDataMode::Particles ? ParticleSpriteRenderData.SourceParticleData->GetNumInstances() : 1;
 
 	VertexFactory.SetSortedIndices(nullptr, 0xFFFFFFFF);
-	NiagaraEmitterInstanceBatcher* Batcher = SceneProxy->GetBatcher();
+	FNiagaraGpuComputeDispatchInterface* ComputeDispatchInterface = SceneProxy->GetComputeDispatchInterface();
 	if (ParticleSpriteRenderData.bNeedsCull || ParticleSpriteRenderData.bNeedsSort)
 	{
 		if (ParticleSpriteRenderData.bSortCullOnGpu)
 		{
-			SortInfo.CulledGPUParticleCountOffset = ParticleSpriteRenderData.bNeedsCull ? Batcher->GetGPUInstanceCounterManager().AcquireCulledEntry() : INDEX_NONE;
-			if (Batcher->AddSortedGPUSimulation(SortInfo))
+			SortInfo.CulledGPUParticleCountOffset = ParticleSpriteRenderData.bNeedsCull ? ComputeDispatchInterface->GetGPUInstanceCounterManager().AcquireCulledEntry() : INDEX_NONE;
+			if (ComputeDispatchInterface->AddSortedGPUSimulation(SortInfo))
 			{
 				VertexFactory.SetSortedIndices(SortInfo.AllocationInfo.BufferSRV, SortInfo.AllocationInfo.BufferOffset);
 			}
diff --git a/Engine/Plugins/FX/Niagara/Source/Niagara/Private/NiagaraScriptExecutionContext.cpp b/Engine/Plugins/FX/Niagara/Source/Niagara/Private/NiagaraScriptExecutionContext.cpp
--- a/Engine/Plugins/FX/Niagara/Source/Niagara/Private/NiagaraScriptExecutionContext.cpp
+++ b/Engine/Plugins/FX/Niagara/Source/Niagara/Private/NiagaraScriptExecutionContext.cpp
@@ -2,7 +2,6 @@
 
 #include "NiagaraScriptExecutionContext.h"
 #include "NiagaraStats.h"
-#include "NiagaraEmitterInstanceBatcher.h"
 #include "NiagaraDataInterface.h"
 #include "NiagaraSystemInstance.h"
 #include "NiagaraSystemGpuComputeProxy.h"
diff --git a/Engine/Plugins/FX/Niagara/Source/Niagara/Private/NiagaraScriptExecutionParameterStore.cpp b/Engine/Plugins/FX/Niagara/Source/Niagara/Private/NiagaraScriptExecutionParameterStore.cpp
--- a/Engine/Plugins/FX/Niagara/Source/Niagara/Private/NiagaraScriptExecutionParameterStore.cpp
+++ b/Engine/Plugins/FX/Niagara/Source/Niagara/Private/NiagaraScriptExecutionParameterStore.cpp
@@ -2,7 +2,6 @@
 
 #include "NiagaraScriptExecutionParameterStore.h"
 #include "NiagaraStats.h"
-#include "NiagaraEmitterInstanceBatcher.h"
 #include "NiagaraDataInterface.h"
 #include "NiagaraSystemInstance.h"
 
diff --git a/Engine/Plugins/FX/Niagara/Source/Niagara/Private/NiagaraSystem.cpp b/Engine/Plugins/FX/Niagara/Source/Niagara/Private/NiagaraSystem.cpp
--- a/Engine/Plugins/FX/Niagara/Source/Niagara/Private/NiagaraSystem.cpp
+++ b/Engine/Plugins/FX/Niagara/Source/Niagara/Private/NiagaraSystem.cpp
@@ -9,7 +9,6 @@
 #include "NiagaraEditorDataBase.h"
 #include "NiagaraEmitter.h"
 #include "NiagaraEmitterHandle.h"
-#include "NiagaraEmitterInstanceBatcher.h"
 #include "NiagaraModule.h"
 #include "NiagaraPrecompileContainer.h"
 #include "NiagaraRendererProperties.h"
diff --git a/Engine/Plugins/FX/Niagara/Source/Niagara/Private/NiagaraSystemGpuComputeProxy.cpp b/Engine/Plugins/FX/Niagara/Source/Niagara/Private/NiagaraSystemGpuComputeProxy.cpp
--- a/Engine/Plugins/FX/Niagara/Source/Niagara/Private/NiagaraSystemGpuComputeProxy.cpp
+++ b/Engine/Plugins/FX/Niagara/Source/Niagara/Private/NiagaraSystemGpuComputeProxy.cpp
@@ -3,13 +3,13 @@
 #include "NiagaraSystemGpuComputeProxy.h"
 #include "NiagaraSystemInstance.h"
 #include "NiagaraEmitterInstance.h"
-#include "NiagaraEmitterInstanceBatcher.h"
 #include "NiagaraComputeExecutionContext.h"
+#include "NiagaraGpuComputeDispatchInterface.h"
 #include "NiagaraGPUSystemTick.h"
 
 FNiagaraSystemGpuComputeProxy::FNiagaraSystemGpuComputeProxy(FNiagaraSystemInstance* OwnerInstance)
 	: DebugOwnerInstance(OwnerInstance)
-	, DebugOwnerBatcher(nullptr)
+	, DebugOwnerComputeDispatchInterface(nullptr)
 	, SystemInstanceID(OwnerInstance->GetId())
 {
 	bRequiresDistanceFieldData = OwnerInstance->RequiresDistanceFieldData();
@@ -49,19 +49,19 @@ FNiagaraSystemGpuComputeProxy::FNiagaraSystemGpuComputeProxy(FNiagaraSystemInsta
 FNiagaraSystemGpuComputeProxy::~FNiagaraSystemGpuComputeProxy()
 {
 	check(IsInRenderingThread());
-	check(DebugOwnerBatcher == nullptr);
+	check(DebugOwnerComputeDispatchInterface == nullptr);
 }
 
-void FNiagaraSystemGpuComputeProxy::AddToBatcher(NiagaraEmitterInstanceBatcher* Batcher)
+void FNiagaraSystemGpuComputeProxy::AddToRenderThread(FNiagaraGpuComputeDispatchInterface* ComputeDispatchInterface)
 {
 	check(IsInGameThread());
-	check(DebugOwnerBatcher == nullptr);
-	DebugOwnerBatcher = Batcher;
+	check(DebugOwnerComputeDispatchInterface == nullptr);
+	DebugOwnerComputeDispatchInterface = ComputeDispatchInterface;
 
-	ENQUEUE_RENDER_COMMAND(AddProxyToBatcher)(
+	ENQUEUE_RENDER_COMMAND(AddProxyToComputeDispatchInterface)(
 		[=](FRHICommandListImmediate& RHICmdList)
 		{
-			Batcher->AddGpuComputeProxy(this);
+			ComputeDispatchInterface->AddGpuComputeProxy(this);
 
 			for (FNiagaraComputeExecutionContext* ComputeContext : ComputeContexts)
 			{
@@ -82,21 +82,21 @@ void FNiagaraSystemGpuComputeProxy::AddToBatcher(NiagaraEmitterInstanceBatcher*
 	);
 }
 
-void FNiagaraSystemGpuComputeProxy::RemoveFromBatcher(NiagaraEmitterInstanceBatcher* Batcher, bool bDeleteProxy)
+void FNiagaraSystemGpuComputeProxy::RemoveFromRenderThread(FNiagaraGpuComputeDispatchInterface* ComputeDispatchInterface, bool bDeleteProxy)
 {
 	check(IsInGameThread());
-	check(DebugOwnerBatcher == Batcher);
-	DebugOwnerBatcher = nullptr;
+	check(DebugOwnerComputeDispatchInterface == ComputeDispatchInterface);
+	DebugOwnerComputeDispatchInterface = nullptr;
 
-	ENQUEUE_RENDER_COMMAND(RemoveFromBatcher)(
+	ENQUEUE_RENDER_COMMAND(RemoveFromRenderThread)(
 		[=](FRHICommandListImmediate& RHICmdList)
 		{
-			Batcher->RemoveGpuComputeProxy(this);
-			ReleaseTicks(Batcher->GetGPUInstanceCounterManager());
+			ComputeDispatchInterface->RemoveGpuComputeProxy(this);
+			ReleaseTicks(ComputeDispatchInterface->GetGPUInstanceCounterManager());
 
 			for (FNiagaraComputeExecutionContext* ComputeContext : ComputeContexts)
 			{
-				ComputeContext->ResetInternal(Batcher);
+				ComputeContext->ResetInternal(ComputeDispatchInterface);
 
 				//-TODO: Can we move this inside the context???
 				for (int i = 0; i < UE_ARRAY_COUNT(ComputeContext->DataBuffers_RT); ++i)
diff --git a/Engine/Plugins/FX/Niagara/Source/Niagara/Private/NiagaraSystemInstance.cpp b/Engine/Plugins/FX/Niagara/Source/Niagara/Private/NiagaraSystemInstance.cpp
--- a/Engine/Plugins/FX/Niagara/Source/Niagara/Private/NiagaraSystemInstance.cpp
+++ b/Engine/Plugins/FX/Niagara/Source/Niagara/Private/NiagaraSystemInstance.cpp
@@ -4,14 +4,16 @@
 #include "NiagaraSystemGpuComputeProxy.h"
 #include "NiagaraConstants.h"
 #include "NiagaraCommon.h"
+#include "NiagaraComputeExecutionContext.h"
 #include "NiagaraDataInterface.h"
+#include "NiagaraGPUSystemTick.h"
 #include "NiagaraStats.h"
 #include "NiagaraParameterCollection.h"
 #include "NiagaraWorldManager.h"
 #include "NiagaraComponent.h"
 #include "NiagaraRenderer.h"
 #include "NiagaraGpuComputeDebug.h"
-#include "NiagaraEmitterInstanceBatcher.h"
+#include "NiagaraGpuComputeDispatchInterface.h"
 #include "NiagaraCrashReporterHandler.h"
 
 #include "Async/Async.h"
@@ -128,15 +130,8 @@ FNiagaraSystemInstance::FNiagaraSystemInstance(UWorld& InWorld, UNiagaraSystem&
 		InstanceParameters.SetOwner(InAttachComponent);
 	}
 
-	if (World->Scene)
-	{
-		FFXSystemInterface*  FXSystemInterface = World->Scene->GetFXSystem();
-		if (FXSystemInterface)
-		{
-			Batcher = static_cast<NiagaraEmitterInstanceBatcher*>(FXSystemInterface->GetInterface(NiagaraEmitterInstanceBatcher::Name));
-		}
-		FeatureLevel = World->FeatureLevel;
-	}
+	ComputeDispatchInterface = FNiagaraGpuComputeDispatchInterface::Get(World);
+	FeatureLevel = World->FeatureLevel;
 
 	// In some cases the system may have already stated that you should ignore dependencies and tick as early as possible.
 	if (!InAsset.bRequireCurrentFrameData)
@@ -499,7 +494,7 @@ void FNiagaraSystemInstance::SetGpuComputeDebug(bool bEnableDebug)
 {
 #if NIAGARA_COMPUTEDEBUG_ENABLED
 	UNiagaraSystem* System  = GetSystem();
-	if (Batcher == nullptr || System == nullptr)
+	if (ComputeDispatchInterface == nullptr || System == nullptr)
 	{
 		return;
 	}
@@ -522,9 +517,9 @@ void FNiagaraSystemInstance::SetGpuComputeDebug(bool bEnableDebug)
 
 		ENQUEUE_RENDER_COMMAND(NiagaraAddGPUSystemDebug)
 		(
-			[RT_Batcher=Batcher, RT_InstanceID=GetId(), RT_SystemName=SystemName](FRHICommandListImmediate& RHICmdList)
+			[RT_ComputeDispatchInterface=ComputeDispatchInterface, RT_InstanceID=GetId(), RT_SystemName=SystemName](FRHICommandListImmediate& RHICmdList)
 			{
-				if (FNiagaraGpuComputeDebug* GpuComputeDebug = RT_Batcher->GetGpuComputeDebug())
+				if (FNiagaraGpuComputeDebug* GpuComputeDebug = RT_ComputeDispatchInterface->GetGpuComputeDebug())
 				{
 					GpuComputeDebug->AddSystemInstance(RT_InstanceID, RT_SystemName);
 				}
@@ -535,9 +530,9 @@ void FNiagaraSystemInstance::SetGpuComputeDebug(bool bEnableDebug)
 	{
 		ENQUEUE_RENDER_COMMAND(NiagaraRemoveGPUSystemDebug)
 		(
-			[RT_Batcher=Batcher, RT_InstanceID=GetId()](FRHICommandListImmediate& RHICmdList)
+			[RT_ComputeDispatchInterface=ComputeDispatchInterface, RT_InstanceID=GetId()](FRHICommandListImmediate& RHICmdList)
 			{
-				if (FNiagaraGpuComputeDebug* GpuComputeDebug = RT_Batcher->GetGpuComputeDebug())
+				if (FNiagaraGpuComputeDebug* GpuComputeDebug = RT_ComputeDispatchInterface->GetGpuComputeDebug())
 				{
 					GpuComputeDebug->RemoveSystemInstance(RT_InstanceID);
 				}
@@ -634,7 +629,7 @@ bool FNiagaraSystemInstance::DeallocateSystemInstance(FNiagaraSystemInstancePtr&
 		if (SystemInstanceAllocation->SystemGpuComputeProxy)
 		{
 			FNiagaraSystemGpuComputeProxy* Proxy = SystemInstanceAllocation->SystemGpuComputeProxy.Release();
-			Proxy->RemoveFromBatcher(SystemInstanceAllocation->GetBatcher(), true);
+			Proxy->RemoveFromRenderThread(SystemInstanceAllocation->GetComputeDispatchInterface(), true);
 		}
 
 		// Queue deferred deletion from the WorldManager
@@ -687,7 +682,7 @@ void FNiagaraSystemInstance::Complete(bool bExternalCompletion)
 	if (SystemGpuComputeProxy)
 	{
 		FNiagaraSystemGpuComputeProxy* Proxy = SystemGpuComputeProxy.Release();
-		Proxy->RemoveFromBatcher(GetBatcher(), true);
+		Proxy->RemoveFromRenderThread(GetComputeDispatchInterface(), true);
 	}
 
 	if (!bPooled)
@@ -783,14 +778,14 @@ void FNiagaraSystemInstance::Reset(FNiagaraSystemInstance::EResetMode Mode)
 		Mode = EResetMode::ReInit;
 	}
 
-	// Remove any existing proxy from the batcher
+	// Remove any existing proxy from the diaptcher
 	// This MUST be done before the emitters array is re-initialized
 	if (SystemGpuComputeProxy.IsValid())
 	{
 		if (IsComplete() || (Mode != EResetMode::ResetSystem))
 		{
 			FNiagaraSystemGpuComputeProxy* Proxy = SystemGpuComputeProxy.Release();
-			Proxy->RemoveFromBatcher(GetBatcher(), true);
+			Proxy->RemoveFromRenderThread(GetComputeDispatchInterface(), true);
 		}
 	}
 
@@ -846,11 +841,11 @@ void FNiagaraSystemInstance::Reset(FNiagaraSystemInstance::EResetMode Mode)
 		//Interface init can disable the system.
 		if (!IsComplete())
 		{
-			// Create the shared context for the batcher if we have a single active GPU emitter in the system
+			// Create the shared context for the dispatcher if we have a single active GPU emitter in the system
 			if (bHasGPUEmitters && !SystemGpuComputeProxy.IsValid())
 			{
 				SystemGpuComputeProxy.Reset(new FNiagaraSystemGpuComputeProxy(this));
-				SystemGpuComputeProxy->AddToBatcher(GetBatcher());
+				SystemGpuComputeProxy->AddToRenderThread(GetComputeDispatchInterface());
 			}
 
 			// Create new random seed
@@ -1158,7 +1153,7 @@ void FNiagaraSystemInstance::Cleanup()
 	if (SystemGpuComputeProxy)
 	{
 		FNiagaraSystemGpuComputeProxy* Proxy = SystemGpuComputeProxy.Release();
-		Proxy->RemoveFromBatcher(GetBatcher(), true);
+		Proxy->RemoveFromRenderThread(GetComputeDispatchInterface(), true);
 	}
 
 	UnbindParameters();
diff --git a/Engine/Plugins/FX/Niagara/Source/Niagara/Private/NiagaraSystemRenderData.cpp b/Engine/Plugins/FX/Niagara/Source/Niagara/Private/NiagaraSystemRenderData.cpp
--- a/Engine/Plugins/FX/Niagara/Source/Niagara/Private/NiagaraSystemRenderData.cpp
+++ b/Engine/Plugins/FX/Niagara/Source/Niagara/Private/NiagaraSystemRenderData.cpp
@@ -31,13 +31,13 @@ FNiagaraSystemRenderData::~FNiagaraSystemRenderData()
 	}
 }
 
-void FNiagaraSystemRenderData::CreateRenderThreadResources(NiagaraEmitterInstanceBatcher& Batcher)
+void FNiagaraSystemRenderData::CreateRenderThreadResources()
 {
 	for (auto Renderer : EmitterRenderers_RT)
 	{
 		if (Renderer)
 		{
-			Renderer->CreateRenderThreadResources(&Batcher);
+			Renderer->CreateRenderThreadResources();
 		}
 	}
 }
diff --git a/Engine/Plugins/FX/Niagara/Source/Niagara/Private/NiagaraSystemSimulation.cpp b/Engine/Plugins/FX/Niagara/Source/Niagara/Private/NiagaraSystemSimulation.cpp
--- a/Engine/Plugins/FX/Niagara/Source/Niagara/Private/NiagaraSystemSimulation.cpp
+++ b/Engine/Plugins/FX/Niagara/Source/Niagara/Private/NiagaraSystemSimulation.cpp
@@ -14,8 +14,11 @@
 #include "Async/ParallelFor.h"
 #include "NiagaraComponent.h"
 #include "NiagaraWorldManager.h"
-#include "NiagaraEmitterInstanceBatcher.h"
+#include "NiagaraGpuComputeDispatchInterface.h"
+#include "NiagaraGPUSystemTick.h"
 #include "NiagaraCrashReporterHandler.h"
+#include "NiagaraSystemGpuComputeProxy.h"
+
 
 // Niagara Async Ticking Sequence
 // - NiagaraSystemSimulation::Tick_GameThread
@@ -630,15 +633,7 @@ bool FNiagaraSystemSimulation::Init(UNiagaraSystem* InSystem, UWorld* InWorld, b
 	FNiagaraWorldManager* WorldMan = FNiagaraWorldManager::Get(InWorld);
 	check(WorldMan);
 
-	Batcher = nullptr;
-	if (World && World->Scene)
-	{
-		FFXSystemInterface* FXSystemInterface = World->Scene->GetFXSystem();
-		if (FXSystemInterface)
-		{
-			Batcher = static_cast<NiagaraEmitterInstanceBatcher*>(FXSystemInterface->GetInterface(NiagaraEmitterInstanceBatcher::Name));
-		}
-	}
+	DispatchInterface = FNiagaraGpuComputeDispatchInterface::Get(World);
 
 	bCanExecute = System->GetSystemSpawnScript()->GetVMExecutableData().IsValid() && System->GetSystemUpdateScript()->GetVMExecutableData().IsValid();
 
@@ -2256,7 +2251,7 @@ void FNiagaraSystemSimulation::BuildConstantBufferTable(
 ENiagaraGPUTickHandlingMode FNiagaraSystemSimulation::GetGPUTickHandlingMode()const
 {
 	const UNiagaraSystem* System = GetSystem();
-	if (Batcher && FNiagaraUtilities::AllowGPUParticles(Batcher->GetShaderPlatform()) && System && System->HasAnyGPUEmitters())
+	if (DispatchInterface && FNiagaraUtilities::AllowGPUParticles(DispatchInterface->GetShaderPlatform()) && System && System->HasAnyGPUEmitters())
 	{
 		//TODO: Maybe some DI post ticks can even be done concurrent too which would also remove this restriction.
 		bool bGT = System->HasDIsWithPostSimulateTick() || GNiagaraSystemSimulationConcurrentGPUTickInit == 0;
diff --git a/Engine/Plugins/FX/Niagara/Source/Niagara/Private/NiagaraUserRedirectionParameterStore.cpp b/Engine/Plugins/FX/Niagara/Source/Niagara/Private/NiagaraUserRedirectionParameterStore.cpp
--- a/Engine/Plugins/FX/Niagara/Source/Niagara/Private/NiagaraUserRedirectionParameterStore.cpp
+++ b/Engine/Plugins/FX/Niagara/Source/Niagara/Private/NiagaraUserRedirectionParameterStore.cpp
@@ -2,7 +2,6 @@
 
 #include "NiagaraUserRedirectionParameterStore.h"
 #include "NiagaraStats.h"
-#include "NiagaraEmitterInstanceBatcher.h"
 #include "NiagaraDataInterface.h"
 #include "NiagaraSystemInstance.h"
 
diff --git a/Engine/Plugins/FX/Niagara/Source/Niagara/Private/NiagaraWorldManager.cpp b/Engine/Plugins/FX/Niagara/Source/Niagara/Private/NiagaraWorldManager.cpp
--- a/Engine/Plugins/FX/Niagara/Source/Niagara/Private/NiagaraWorldManager.cpp
+++ b/Engine/Plugins/FX/Niagara/Source/Niagara/Private/NiagaraWorldManager.cpp
@@ -15,11 +15,11 @@
 #include "GameFramework/PlayerController.h"
 #include "EngineModule.h"
 #include "NiagaraStats.h"
-#include "NiagaraEmitterInstanceBatcher.h"
 #include "NiagaraComponentPool.h"
 #include "NiagaraComponent.h"
 #include "NiagaraEffectType.h"
 #include "NiagaraDebugHud.h"
+#include "NiagaraGpuComputeDispatchInterface.h"
 #include "HAL/PlatformApplicationMisc.h"
 #include "Particles/FXBudget.h"
 #include "NiagaraFunctionLibrary.h"
@@ -509,11 +509,11 @@ void FNiagaraWorldManager::OnSystemPostChange(UNiagaraSystem* System)
 }
 #endif//WITH_EDITOR
 
-void FNiagaraWorldManager::OnBatcherDestroyed_Internal(NiagaraEmitterInstanceBatcher* InBatcher)
+void FNiagaraWorldManager::OnComputeDispatchInterfaceDestroyed_Internal(FNiagaraGpuComputeDispatchInterface* InComputeDispatchInterface)
 {
-	// Process the deferred deletion queue before deleting the batcher of this world.
-	// This is required because the batcher is accessed in FNiagaraEmitterInstance::~FNiagaraEmitterInstance
-	if (World && World->FXSystem && World->FXSystem->GetInterface(NiagaraEmitterInstanceBatcher::Name) == InBatcher)
+	// Process the deferred deletion queue before deleting the ComputeDispatchInterface of this world.
+	// This is required because the ComputeDispatchInterface is accessed in FNiagaraEmitterInstance::~FNiagaraEmitterInstance
+	if (FNiagaraGpuComputeDispatchInterface::Get(World) == InComputeDispatchInterface)
 	{
 		DeferredDeletionQueue.Empty();
 	}
@@ -650,11 +650,11 @@ void FNiagaraWorldManager::OnWorldBeginTearDown(UWorld* World)
 // 	}
 }
 
-void FNiagaraWorldManager::OnBatcherDestroyed(NiagaraEmitterInstanceBatcher* InBatcher)
+void FNiagaraWorldManager::OnComputeDispatchInterfaceDestroyed(FNiagaraGpuComputeDispatchInterface* InComputeDispatchInterface)
 {
 	for (TPair<UWorld*, FNiagaraWorldManager*>& Pair : WorldManagers)
 	{
-		Pair.Value->OnBatcherDestroyed_Internal(InBatcher);
+		Pair.Value->OnComputeDispatchInterfaceDestroyed_Internal(InComputeDispatchInterface);
 	}
 }
 
diff --git a/Engine/Plugins/FX/Niagara/Source/Niagara/Public/NiagaraComponent.h b/Engine/Plugins/FX/Niagara/Source/Niagara/Public/NiagaraComponent.h
--- a/Engine/Plugins/FX/Niagara/Source/Niagara/Public/NiagaraComponent.h
+++ b/Engine/Plugins/FX/Niagara/Source/Niagara/Public/NiagaraComponent.h
@@ -22,7 +22,7 @@ class UNiagaraSystem;
 class UNiagaraParameterCollection;
 class UNiagaraParameterCollectionInstance;
 class FNiagaraSystemSimulation;
-class NiagaraEmitterInstanceBatcher;
+class FNiagaraGpuComputeDispatchInterface;
 
 // Called when the particle system is done
 DECLARE_DYNAMIC_MULTICAST_DELEGATE_OneParam(FOnNiagaraSystemFinished, class UNiagaraComponent*, PSystem);
@@ -768,7 +768,7 @@ public:
 	/** Sets whether or not this scene proxy should be rendered. */
 	void SetRenderingEnabled(bool bInRenderingEnabled);
 
-	NiagaraEmitterInstanceBatcher* GetBatcher() const { return Batcher; }
+	FNiagaraGpuComputeDispatchInterface* GetComputeDispatchInterface() const { return ComputeDispatchInterface; }
 
 #if RHI_RAYTRACING
 	virtual void GetDynamicRayTracingInstances(FRayTracingMaterialGatheringContext& Context, TArray<FRayTracingInstance>& OutRayTracingInstances) override;
@@ -826,7 +826,7 @@ private:
 	/** The data required to render a single instance of a NiagaraSystem */
 	FNiagaraSystemRenderData* RenderData = nullptr;
 
-	NiagaraEmitterInstanceBatcher* Batcher = nullptr;
+	FNiagaraGpuComputeDispatchInterface* ComputeDispatchInterface = nullptr;
 
 	FMatrix LocalToWorldInverse;
 
diff --git a/Engine/Plugins/FX/Niagara/Source/Niagara/Public/NiagaraDataInterfaceRigidMeshCollisionQuery.h b/Engine/Plugins/FX/Niagara/Source/Niagara/Public/NiagaraDataInterfaceRigidMeshCollisionQuery.h
--- a/Engine/Plugins/FX/Niagara/Source/Niagara/Public/NiagaraDataInterfaceRigidMeshCollisionQuery.h
+++ b/Engine/Plugins/FX/Niagara/Source/Niagara/Public/NiagaraDataInterfaceRigidMeshCollisionQuery.h
@@ -171,7 +171,7 @@ struct FNDIRigidMeshCollisionProxy : public FNiagaraDataInterfaceProxy
 	void InitializePerInstanceData(const FNiagaraSystemInstanceID& SystemInstance);
 
 	/** Destroy the proxy data if necessary */
-	void DestroyPerInstanceData(NiagaraEmitterInstanceBatcher* Batcher, const FNiagaraSystemInstanceID& SystemInstance);
+	void DestroyPerInstanceData(const FNiagaraSystemInstanceID& SystemInstance);
 
 	/** Launch all pre stage functions */
 	virtual void PreStage(FRHICommandList& RHICmdList, const FNiagaraDataInterfaceStageArgs& Context) override;
diff --git a/Engine/Plugins/FX/Niagara/Source/Niagara/Public/NiagaraEmptyUAVPool.h b/Engine/Plugins/FX/Niagara/Source/Niagara/Public/NiagaraEmptyUAVPool.h
new file mode 100644
--- /dev/null
+++ b/Engine/Plugins/FX/Niagara/Source/Niagara/Public/NiagaraEmptyUAVPool.h
@@ -0,0 +1,71 @@
+// Copyright Epic Games, Inc. All Rights Reserved.
+
+#pragma once
+
+#include "NiagaraCommon.h"
+
+// Scoped access for UAVs this is not required when running inside the main dispatch loop
+// but is required for external usage that is outside, i.e. if you are doing some custom dispatch setup
+struct FNiagaraEmptyUAVPoolScopedAccess
+{
+public:
+	explicit FNiagaraEmptyUAVPoolScopedAccess(class FNiagaraEmptyUAVPool* EmptyUAVPool);
+	FNiagaraEmptyUAVPoolScopedAccess(const FNiagaraEmptyUAVPoolScopedAccess&) = delete;
+	FNiagaraEmptyUAVPoolScopedAccess(FNiagaraEmptyUAVPoolScopedAccess&&) = delete;
+	~FNiagaraEmptyUAVPoolScopedAccess();
+
+	FNiagaraEmptyUAVPoolScopedAccess& operator=(const FNiagaraEmptyUAVPoolScopedAccess&) = delete;
+	FNiagaraEmptyUAVPoolScopedAccess& operator=(FNiagaraEmptyUAVPoolScopedAccess&&) = delete;
+
+private:
+	class FNiagaraEmptyUAVPool* EmptyUAVPool;
+};
+
+// Type of empty UAV to get
+enum class ENiagaraEmptyUAVType
+{
+	Buffer,
+	Texture2D,
+	Texture2DArray,
+	Texture3D,
+	Num
+};
+
+// Empty UAV pool used for ensuring we bind a buffer when one does not exist
+class NIAGARA_API FNiagaraEmptyUAVPool
+{
+	friend struct FNiagaraEmptyUAVPoolScopedAccess;
+
+public:
+	/**
+	Grab a temporary empty RW buffer from the pool.
+	Note: When doing this outside of Niagara you must be within a FNiagaraUAVPoolAccessScope.
+	*/
+	FRHIUnorderedAccessView* GetEmptyUAVFromPool(FRHICommandList& RHICmdList, EPixelFormat Format, ENiagaraEmptyUAVType Type);
+
+	/**
+	Returns all used UAVs back to the pool
+	*/
+	void ResetEmptyUAVPools();
+
+protected:
+	struct FEmptyUAV
+	{
+		~FEmptyUAV();
+
+		FBufferRHIRef Buffer;
+		FTextureRHIRef Texture;
+		FUnorderedAccessViewRHIRef UAV;
+	};
+
+	struct FEmptyUAVPool
+	{
+		~FEmptyUAVPool();
+
+		int32 NextFreeIndex = 0;
+		TArray<FEmptyUAV> UAVs;
+	};
+
+	uint32 UAVAccessCounter = 0;
+	TMap<EPixelFormat, FEmptyUAVPool> UAVPools[(int)ENiagaraEmptyUAVType::Num];
+};
diff --git a/Engine/Plugins/FX/Niagara/Source/Niagara/Public/NiagaraGpuComputeDispatchInterface.h b/Engine/Plugins/FX/Niagara/Source/Niagara/Public/NiagaraGpuComputeDispatchInterface.h
new file mode 100644
--- /dev/null
+++ b/Engine/Plugins/FX/Niagara/Source/Niagara/Public/NiagaraGpuComputeDispatchInterface.h
@@ -0,0 +1,100 @@
+// Copyright Epic Games, Inc. All Rights Reserved.
+
+#pragma once
+
+#include "NiagaraCommon.h"
+#include "NiagaraEmptyUAVPool.h"
+#include "NiagaraGPUInstanceCountManager.h"
+#include "FXSystem.h"
+
+struct FNiagaraComputeExecutionContext;
+class FNiagaraGpuComputeDebug;
+class FNiagaraGPUInstanceCountManager;
+class FNiagaraGpuReadbackManager;
+class FNiagaraRayTracingHelper;
+struct FNiagaraScriptDebuggerInfo;
+class FNiagaraSystemGpuComputeProxy;
+
+// Public API for Niagara's Compute Dispatcher
+// This is generally used with DataInterfaces or Custom Renderers
+class NIAGARA_API FNiagaraGpuComputeDispatchInterface : public FFXSystemInterface
+{
+public:
+	static FNiagaraGpuComputeDispatchInterface* Get(class UWorld* World);
+	static FNiagaraGpuComputeDispatchInterface* Get(class FSceneInterface* Scene);
+	static FNiagaraGpuComputeDispatchInterface* Get(class FFXSystemInterface* FXSceneInterface);
+
+	explicit FNiagaraGpuComputeDispatchInterface(EShaderPlatform InShaderPlatform, ERHIFeatureLevel::Type InFeatureLevel);
+	virtual ~FNiagaraGpuComputeDispatchInterface();
+
+	/** Get ShaderPlatform the batcher is bound to */
+	EShaderPlatform GetShaderPlatform() const { return ShaderPlatform; }
+	/** Get FeatureLevel the batcher is bound to */
+	ERHIFeatureLevel::Type GetFeatureLevel() const { return FeatureLevel; }
+
+	/** Add system instance proxy to the batcher for tracking. */
+	virtual void AddGpuComputeProxy(FNiagaraSystemGpuComputeProxy* ComputeProxy) = 0;
+	/** Remove system instance proxy from the batcher. */
+	virtual void RemoveGpuComputeProxy(FNiagaraSystemGpuComputeProxy* ComputeProxy) = 0;
+
+	/**
+	 * Register work for GPU sorting (using the GPUSortManager).
+	 * The constraints of the sort request are defined in SortInfo.SortFlags.
+	 * The sort task bindings are set in SortInfo.AllocationInfo.
+	 * The initial keys and values are generated in the GenerateSortKeys() callback.
+	 *
+	 * Return true if the work was registered, or false it GPU sorting is not available or impossible.
+	 */
+	virtual bool AddSortedGPUSimulation(struct FNiagaraGPUSortInfo& SortInfo) = 0;
+
+	/** Get access to the instance count manager. */
+	FORCEINLINE FNiagaraGPUInstanceCountManager& GetGPUInstanceCounterManager() { check(IsInRenderingThread()); return GPUInstanceCounterManager; }
+
+#if NIAGARA_COMPUTEDEBUG_ENABLED
+	/** Get access to Niagara's GpuComputeDebug */
+	FNiagaraGpuComputeDebug* GetGpuComputeDebug() const { return GpuComputeDebugPtr.Get(); }
+#endif
+
+	/** Get access to Niagara's GpuReadbackManager. */
+	FNiagaraGpuReadbackManager* GetGpuReadbackManager() const { return GpuReadbackManagerPtr.Get(); }
+
+	/** Get access to Niagara's GpuReadbackManager. */
+	FNiagaraEmptyUAVPool* GetEmptyUAVPool() const { return EmptyUAVPoolPtr.Get(); }
+
+	/** Convenience wrapped to get a UAV from the poo. */
+	FRHIUnorderedAccessView* GetEmptyUAVFromPool(FRHICommandList& RHICmdList, EPixelFormat Format, ENiagaraEmptyUAVType Type) const { return EmptyUAVPoolPtr->GetEmptyUAVFromPool(RHICmdList, Format, Type); }
+
+	/**
+	Call this to force all pending ticks to be flushed from the batcher.
+	Doing so will execute them outside of a view context which may result in undesirable results.
+	*/
+	virtual void FlushPendingTicks_GameThread() = 0;
+
+	/** Debug only function to readback data. */
+	virtual void AddDebugReadback(FNiagaraSystemInstanceID InstanceID, TSharedPtr<FNiagaraScriptDebuggerInfo, ESPMode::ThreadSafe> DebugInfo, FNiagaraComputeExecutionContext* Context) = 0;
+
+	/** Processes all pending debug readbacks */
+	virtual void ProcessDebugReadbacks(FRHICommandListImmediate& RHICmdList, bool bWaitCompletion) = 0;
+
+#if RHI_RAYTRACING
+	virtual FNiagaraRayTracingHelper& GetRayTracingHelper() const = 0;
+	virtual bool HasRayTracingScene() const = 0;
+
+	virtual void SetPrimitiveRayTracingCollisionGroup(UPrimitiveComponent* PrimitiveComponent, uint32 Group) = 0;
+
+	virtual int32 AcquireGPURayTracedCollisionGroup() = 0;
+	virtual void ReleaseGPURayTracedCollisionGroup(int32 CollisionGroup) = 0;
+#endif
+
+protected:
+	EShaderPlatform							ShaderPlatform;
+	ERHIFeatureLevel::Type					FeatureLevel;
+#if NIAGARA_COMPUTEDEBUG_ENABLED
+	TUniquePtr<FNiagaraGpuComputeDebug>		GpuComputeDebugPtr;
+#endif
+	TUniquePtr<FNiagaraGpuReadbackManager>	GpuReadbackManagerPtr;
+	TUniquePtr<FNiagaraEmptyUAVPool>		EmptyUAVPoolPtr;
+
+	// GPU emitter instance count buffer. Contains the actual particle / instance count generate in the GPU tick.
+	FNiagaraGPUInstanceCountManager			GPUInstanceCounterManager;
+};
diff --git a/Engine/Plugins/FX/Niagara/Source/Niagara/Public/NiagaraRenderer.h b/Engine/Plugins/FX/Niagara/Source/Niagara/Public/NiagaraRenderer.h
--- a/Engine/Plugins/FX/Niagara/Source/Niagara/Public/NiagaraRenderer.h
+++ b/Engine/Plugins/FX/Niagara/Source/Niagara/Public/NiagaraRenderer.h
@@ -82,7 +82,7 @@ public:
 	FNiagaraRenderer& operator=(const FNiagaraRenderer& Other) = delete;
 
 	virtual void Initialize(const UNiagaraRendererProperties *InProps, const FNiagaraEmitterInstance* Emitter, const FNiagaraSystemInstanceController& InComponent);
-	virtual void CreateRenderThreadResources(NiagaraEmitterInstanceBatcher* Batcher) {}
+	virtual void CreateRenderThreadResources() {}
 	virtual void ReleaseRenderThreadResources() {}
 	virtual void DestroyRenderState_Concurrent() {}
 
diff --git a/Engine/Plugins/FX/Niagara/Source/Niagara/Public/NiagaraRendererMeshes.h b/Engine/Plugins/FX/Niagara/Source/Niagara/Public/NiagaraRendererMeshes.h
--- a/Engine/Plugins/FX/Niagara/Source/Niagara/Public/NiagaraRendererMeshes.h
+++ b/Engine/Plugins/FX/Niagara/Source/Niagara/Public/NiagaraRendererMeshes.h
@@ -108,7 +108,7 @@ protected:
 	void PrepareParticleRenderBuffers(FParticleMeshRenderData& ParticleMeshRenderData, FGlobalDynamicReadBuffer& DynamicReadBuffer) const;
 	void InitializeSortInfo(FParticleMeshRenderData& ParticleMeshRenderData, const FNiagaraSceneProxy& SceneProxy, const FSceneView& View, int32 ViewIndex, bool bIsInstancedStereo, FNiagaraGPUSortInfo& OutSortInfo) const;
 	void PreparePerMeshData(FParticleMeshRenderData& ParticleMeshRenderData, const FNiagaraSceneProxy& SceneProxy, const FMeshData& MeshData) const;
-	uint32 PerformSortAndCull(FParticleMeshRenderData& ParticleMeshRenderData, FGlobalDynamicReadBuffer& ReadBuffer, FNiagaraGPUSortInfo& SortInfo, class NiagaraEmitterInstanceBatcher* Batcher, int32 MeshIndex) const;
+	uint32 PerformSortAndCull(FParticleMeshRenderData& ParticleMeshRenderData, FGlobalDynamicReadBuffer& ReadBuffer, FNiagaraGPUSortInfo& SortInfo, class FNiagaraGpuComputeDispatchInterface* ComputeDispatchInterface, int32 MeshIndex) const;
 	FNiagaraMeshUniformBufferRef CreatePerViewUniformBuffer(FParticleMeshRenderData& ParticleMeshRenderData, const FSceneView& View, const FMeshData& MeshData, const FNiagaraSceneProxy& SceneProxy) const;
 
 	void CreateMeshBatchForSection(
diff --git a/Engine/Plugins/FX/Niagara/Source/Niagara/Public/NiagaraRendererRibbons.h b/Engine/Plugins/FX/Niagara/Source/Niagara/Public/NiagaraRendererRibbons.h
--- a/Engine/Plugins/FX/Niagara/Source/Niagara/Public/NiagaraRendererRibbons.h
+++ b/Engine/Plugins/FX/Niagara/Source/Niagara/Public/NiagaraRendererRibbons.h
@@ -22,7 +22,7 @@ public:
 	~FNiagaraRendererRibbons();
 
 	// FNiagaraRenderer Interface 
-	virtual void CreateRenderThreadResources(NiagaraEmitterInstanceBatcher* Batcher) override;
+	virtual void CreateRenderThreadResources() override;
 	virtual void ReleaseRenderThreadResources() override;
 
 	virtual void GetDynamicMeshElements(const TArray<const FSceneView*>& Views, const FSceneViewFamily& ViewFamily, uint32 VisibilityMap, FMeshElementCollector& Collector, const FNiagaraSceneProxy *SceneProxy) const override;
diff --git a/Engine/Plugins/FX/Niagara/Source/Niagara/Public/NiagaraRendererSprites.h b/Engine/Plugins/FX/Niagara/Source/Niagara/Public/NiagaraRendererSprites.h
--- a/Engine/Plugins/FX/Niagara/Source/Niagara/Public/NiagaraRendererSprites.h
+++ b/Engine/Plugins/FX/Niagara/Source/Niagara/Public/NiagaraRendererSprites.h
@@ -22,7 +22,7 @@ public:
 	~FNiagaraRendererSprites();
 
 	//FNiagaraRenderer interface
-	virtual void CreateRenderThreadResources(NiagaraEmitterInstanceBatcher* Batcher) override;
+	virtual void CreateRenderThreadResources() override;
 	virtual void ReleaseRenderThreadResources() override;
 
 	virtual void GetDynamicMeshElements(const TArray<const FSceneView*>& Views, const FSceneViewFamily& ViewFamily, uint32 VisibilityMap, FMeshElementCollector& Collector, const FNiagaraSceneProxy *SceneProxy) const override;
diff --git a/Engine/Plugins/FX/Niagara/Source/Niagara/Public/NiagaraSimStageData.h b/Engine/Plugins/FX/Niagara/Source/Niagara/Public/NiagaraSimStageData.h
new file mode 100644
--- /dev/null
+++ b/Engine/Plugins/FX/Niagara/Source/Niagara/Public/NiagaraSimStageData.h
@@ -0,0 +1,31 @@
+// Copyright Epic Games, Inc. All Rights Reserved.
+
+#pragma once
+
+struct FNiagaraSimStageData
+{
+	FNiagaraSimStageData()
+	{
+		bFirstStage = false;
+		bLastStage = false;
+		bSetDataToRender = false;
+	}
+
+	uint32 bFirstStage : 1;
+	uint32 bLastStage : 1;
+	uint32 bSetDataToRender : 1;
+
+	uint32 StageIndex = INDEX_NONE;
+	uint32 IterationIndex = 0;
+
+	class FNiagaraDataBuffer* Source = nullptr;
+	uint32 SourceCountOffset = INDEX_NONE;
+	uint32 SourceNumInstances = 0;
+
+	class FNiagaraDataBuffer* Destination = nullptr;
+	uint32 DestinationCountOffset = INDEX_NONE;
+	uint32 DestinationNumInstances = 0;
+
+	struct FNiagaraDataInterfaceProxyRW* AlternateIterationSource = nullptr;
+	const struct FSimulationStageMetaData* StageMetaData = nullptr;
+};
diff --git a/Engine/Plugins/FX/Niagara/Source/Niagara/Public/NiagaraSystemGpuComputeProxy.h b/Engine/Plugins/FX/Niagara/Source/Niagara/Public/NiagaraSystemGpuComputeProxy.h
--- a/Engine/Plugins/FX/Niagara/Source/Niagara/Public/NiagaraSystemGpuComputeProxy.h
+++ b/Engine/Plugins/FX/Niagara/Source/Niagara/Public/NiagaraSystemGpuComputeProxy.h
@@ -5,21 +5,21 @@
 #include "NiagaraCommon.h"
 
 class FNiagaraSystemInstance;
-class NiagaraEmitterInstanceBatcher;
+class FNiagaraGpuComputeDispatchInterface;
 struct FNiagaraComputeExecutionContext;
 class FNiagaraGPUInstanceCountManager;
 class FNiagaraGPUSystemTick;
 
 class FNiagaraSystemGpuComputeProxy
 {
-	friend class NiagaraEmitterInstanceBatcher;
+	friend class FNiagaraGpuComputeDispatch;
 
 public:
 	FNiagaraSystemGpuComputeProxy(FNiagaraSystemInstance* OwnerInstance);
 	~FNiagaraSystemGpuComputeProxy();
 
-	void AddToBatcher(NiagaraEmitterInstanceBatcher* Batcher);
-	void RemoveFromBatcher(NiagaraEmitterInstanceBatcher* Batcher, bool bDeleteProxy);
+	void AddToRenderThread(FNiagaraGpuComputeDispatchInterface* ComputeDispatchInterface);
+	void RemoveFromRenderThread(FNiagaraGpuComputeDispatchInterface* ComputeDispatchInterface, bool bDeleteProxy);
 
 	FNiagaraSystemInstanceID GetSystemInstanceID() const { return SystemInstanceID; }
 	ENiagaraGpuComputeTickStage::Type GetComputeTickStage() const { return ComputeTickStage; }
@@ -34,8 +34,8 @@ public:
 
 private:
 	FNiagaraSystemInstance*						DebugOwnerInstance = nullptr;
-	NiagaraEmitterInstanceBatcher*				DebugOwnerBatcher = nullptr;
-	int32										BatcherIndex = INDEX_NONE;
+	FNiagaraGpuComputeDispatchInterface*		DebugOwnerComputeDispatchInterface = nullptr;
+	int32										ComputeDispatchIndex = INDEX_NONE;
 
 	FNiagaraSystemInstanceID					SystemInstanceID = FNiagaraSystemInstanceID();
 	ENiagaraGpuComputeTickStage::Type			ComputeTickStage = ENiagaraGpuComputeTickStage::PostOpaqueRender;
diff --git a/Engine/Plugins/FX/Niagara/Source/Niagara/Public/NiagaraSystemInstance.h b/Engine/Plugins/FX/Niagara/Source/Niagara/Public/NiagaraSystemInstance.h
--- a/Engine/Plugins/FX/Niagara/Source/Niagara/Public/NiagaraSystemInstance.h
+++ b/Engine/Plugins/FX/Niagara/Source/Niagara/Public/NiagaraSystemInstance.h
@@ -14,7 +14,7 @@
 class FNiagaraWorldManager;
 class FNiagaraSystemInstance;
 class FNiagaraSystemSimulation;
-class NiagaraEmitterInstanceBatcher;
+class FNiagaraGpuComputeDispatchInterface;
 class FNiagaraGPUSystemTick;
 class FNiagaraSystemGpuComputeProxy;
 
@@ -338,7 +338,7 @@ public:
 	/** Dumps information about the instances tick to the log */
 	void DumpTickInfo(FOutputDevice& Ar);
 
-	NiagaraEmitterInstanceBatcher* GetBatcher() const { return Batcher; }
+	FNiagaraGpuComputeDispatchInterface* GetComputeDispatchInterface() const { return ComputeDispatchInterface; }
 
 	static bool AllocateSystemInstance(FNiagaraSystemInstancePtr& OutSystemInstanceAllocation, UWorld& InWorld, UNiagaraSystem& InAsset,
 		FNiagaraUserRedirectionParameterStore* InOverrideParameters = nullptr, USceneComponent* InAttachComponent = nullptr,
@@ -557,7 +557,7 @@ private:
 	/** Copy of simulations internal state so that it can be passed to emitters etc. */
 	ENiagaraExecutionState ActualExecutionState;
 
-	NiagaraEmitterInstanceBatcher* Batcher = nullptr;
+	FNiagaraGpuComputeDispatchInterface* ComputeDispatchInterface = nullptr;
 	TUniquePtr<FNiagaraSystemGpuComputeProxy> SystemGpuComputeProxy;
 
 	/** Tag we feed into crash reporter for this instance. */
diff --git a/Engine/Plugins/FX/Niagara/Source/Niagara/Public/NiagaraSystemInstanceController.h b/Engine/Plugins/FX/Niagara/Source/Niagara/Public/NiagaraSystemInstanceController.h
--- a/Engine/Plugins/FX/Niagara/Source/Niagara/Public/NiagaraSystemInstanceController.h
+++ b/Engine/Plugins/FX/Niagara/Source/Niagara/Public/NiagaraSystemInstanceController.h
@@ -112,7 +112,7 @@ public:
 	NIAGARA_SYSTEM_INSTANCE_CONTROLLER_SHIM(GetLODDistance, const)
 	NIAGARA_SYSTEM_INSTANCE_CONTROLLER_SHIM(GetAreDataInterfacesInitialized, const)
 	NIAGARA_SYSTEM_INSTANCE_CONTROLLER_SHIM(GetAge, const)
-	NIAGARA_SYSTEM_INSTANCE_CONTROLLER_SHIM(GetBatcher, const)
+	NIAGARA_SYSTEM_INSTANCE_CONTROLLER_SHIM(GetComputeDispatchInterface, const)
 	NIAGARA_SYSTEM_INSTANCE_CONTROLLER_SHIM(IsSolo, const)
 	NIAGARA_SYSTEM_INSTANCE_CONTROLLER_SHIM(IsComplete, const)
 	NIAGARA_SYSTEM_INSTANCE_CONTROLLER_SHIM(IsPendingSpawn, const)
diff --git a/Engine/Plugins/FX/Niagara/Source/Niagara/Public/NiagaraSystemRenderData.h b/Engine/Plugins/FX/Niagara/Source/Niagara/Public/NiagaraSystemRenderData.h
--- a/Engine/Plugins/FX/Niagara/Source/Niagara/Public/NiagaraSystemRenderData.h
+++ b/Engine/Plugins/FX/Niagara/Source/Niagara/Public/NiagaraSystemRenderData.h
@@ -47,7 +47,7 @@ public:
 	~FNiagaraSystemRenderData();
 
 	/** Called from the render thread to give renderers an opportunity to create resources needed for rendering. */
-	void CreateRenderThreadResources(NiagaraEmitterInstanceBatcher& Batcher);
+	void CreateRenderThreadResources();
 	/** Called from the render thread to give renderers an opportunity to release their resources. */
 	void ReleaseRenderThreadResources();
 	/** Gives the system's renderers an opportunity to free resources */
diff --git a/Engine/Plugins/FX/Niagara/Source/Niagara/Public/NiagaraSystemSimulation.h b/Engine/Plugins/FX/Niagara/Source/Niagara/Public/NiagaraSystemSimulation.h
--- a/Engine/Plugins/FX/Niagara/Source/Niagara/Public/NiagaraSystemSimulation.h
+++ b/Engine/Plugins/FX/Niagara/Source/Niagara/Public/NiagaraSystemSimulation.h
@@ -14,6 +14,7 @@ class UWorld;
 class UNiagaraParameterCollection;
 class UNiagaraParameterCollectionInstance;
 class FNiagaraSystemSimulation;
+class FNiagaraGpuComputeDispatchInterface;
 
 using FNiagaraSystemSimulationPtr = TSharedPtr<FNiagaraSystemSimulation, ESPMode::ThreadSafe>;
 
@@ -308,9 +309,9 @@ public:
 
 	ETickingGroup GetTickGroup() const { return SystemTickGroup; }
 
-	FORCEINLINE NiagaraEmitterInstanceBatcher* GetBatcher()const { return Batcher; }
+	FORCEINLINE FNiagaraGpuComputeDispatchInterface* GetDispatchInterface() const { return DispatchInterface; }
 
-	ENiagaraGPUTickHandlingMode GetGPUTickHandlingMode()const;
+	ENiagaraGPUTickHandlingMode GetGPUTickHandlingMode() const;
 
 	/** If true we use legacy simulation contexts that could not handle per instance DI calls in the system scripts and would force the whole simulation solo. */
 	static bool UseLegacySystemSimulationContexts();
@@ -426,7 +427,7 @@ protected:
 
 	mutable FString CrashReporterTag;
 
-	NiagaraEmitterInstanceBatcher* Batcher = nullptr;
+	FNiagaraGpuComputeDispatchInterface* DispatchInterface = nullptr;
 
 	static bool bUseLegacyExecContexts;
 };
diff --git a/Engine/Plugins/FX/Niagara/Source/Niagara/Public/NiagaraWorldManager.h b/Engine/Plugins/FX/Niagara/Source/Niagara/Public/NiagaraWorldManager.h
--- a/Engine/Plugins/FX/Niagara/Source/Niagara/Public/NiagaraWorldManager.h
+++ b/Engine/Plugins/FX/Niagara/Source/Niagara/Public/NiagaraWorldManager.h
@@ -72,8 +72,8 @@ public:
 	static void OnStartup();
 	static void OnShutdown();
 
-	// Gamethread callback to cleanup references to the given batcher before it gets deleted on the renderthread.
-	static void OnBatcherDestroyed(class NiagaraEmitterInstanceBatcher* InBatcher);
+	// Gamethread callback to cleanup references to the given ComputeDispatchInterface before it gets deleted on the renderthread.
+	static void OnComputeDispatchInterfaceDestroyed(class FNiagaraGpuComputeDispatchInterface* InComputeDispatchInterface);
 
 	static void DestroyAllSystemSimulations(class UNiagaraSystem* System);
 
@@ -214,8 +214,8 @@ private:
 	// Callback to handle any pre GC processing needed.
 	static void OnPreGarbageCollectBeginDestroy();
 		
-	// Gamethread callback to cleanup references to the given batcher before it gets deleted on the renderthread.
-	void OnBatcherDestroyed_Internal(NiagaraEmitterInstanceBatcher* InBatcher);
+	// Gamethread callback to cleanup references to the given ComputeDispatchInterface before it gets deleted on the renderthread.
+	void OnComputeDispatchInterfaceDestroyed_Internal(class FNiagaraGpuComputeDispatchInterface* InComputeDispatchInterface);
 
 	bool CanPreCull(UNiagaraEffectType* EffectType);
 
diff --git a/Engine/Plugins/FX/Niagara/Source/NiagaraCore/Public/NiagaraDataInterfaceBase.h b/Engine/Plugins/FX/Niagara/Source/NiagaraCore/Public/NiagaraDataInterfaceBase.h
--- a/Engine/Plugins/FX/Niagara/Source/NiagaraCore/Public/NiagaraDataInterfaceBase.h
+++ b/Engine/Plugins/FX/Niagara/Source/NiagaraCore/Public/NiagaraDataInterfaceBase.h
@@ -16,11 +16,11 @@ class FNiagaraSystemInstance;
 class FNiagaraShader;
 class FNiagaraShaderMapPointerTable;
 class UNiagaraDataInterfaceBase;
+class FNiagaraGpuComputeDispatchInterface;
 struct FNiagaraDataInterfaceGPUParamInfo;
 class FRHICommandList;
 
 struct FNiagaraDataInterfaceProxy;
-class NiagaraEmitterInstanceBatcher;
 struct FNiagaraComputeInstanceData;
 struct FNiagaraSimStageData;
 
@@ -28,22 +28,22 @@ DECLARE_EXPORTED_TEMPLATE_INTRINSIC_TYPE_LAYOUT(template<>, TIndexedPtr<UNiagara
 
 struct FNiagaraDataInterfaceArgs
 {
-	FNiagaraDataInterfaceArgs(FNiagaraDataInterfaceProxy* InDataInterface, FNiagaraSystemInstanceID InSystemInstanceID, const NiagaraEmitterInstanceBatcher* InBatcher)
+	FNiagaraDataInterfaceArgs(FNiagaraDataInterfaceProxy* InDataInterface, FNiagaraSystemInstanceID InSystemInstanceID, const FNiagaraGpuComputeDispatchInterface* InComputeDispatchInterface)
 		: DataInterface(InDataInterface)
 		, SystemInstanceID(InSystemInstanceID)
-		, Batcher(InBatcher)
+		, ComputeDispatchInterface(InComputeDispatchInterface)
 	{
 	}
 
 	FNiagaraDataInterfaceProxy*	DataInterface;
 	FNiagaraSystemInstanceID SystemInstanceID;
-	const NiagaraEmitterInstanceBatcher* Batcher;
+	const FNiagaraGpuComputeDispatchInterface* ComputeDispatchInterface;
 };
 
 struct FNiagaraDataInterfaceStageArgs : public FNiagaraDataInterfaceArgs
 {
-	FNiagaraDataInterfaceStageArgs(FNiagaraDataInterfaceProxy* InDataInterface, FNiagaraSystemInstanceID InSystemInstanceID, const NiagaraEmitterInstanceBatcher* InBatcher, const FNiagaraComputeInstanceData* InComputeInstanceData, const FNiagaraSimStageData* InSimStageData, bool InIsOutputStage, bool InIsIterationStage)
-		: FNiagaraDataInterfaceArgs(InDataInterface, InSystemInstanceID, InBatcher)
+	FNiagaraDataInterfaceStageArgs(FNiagaraDataInterfaceProxy* InDataInterface, FNiagaraSystemInstanceID InSystemInstanceID, const FNiagaraGpuComputeDispatchInterface* InComputeDispatchInterface, const FNiagaraComputeInstanceData* InComputeInstanceData, const FNiagaraSimStageData* InSimStageData, bool InIsOutputStage, bool InIsIterationStage)
+		: FNiagaraDataInterfaceArgs(InDataInterface, InSystemInstanceID, InComputeDispatchInterface)
 		, ComputeInstanceData(InComputeInstanceData)
 		, SimStageData(InSimStageData)
 		, IsOutputStage(InIsOutputStage)
@@ -62,8 +62,8 @@ struct FNiagaraDataInterfaceSetArgs : public FNiagaraDataInterfaceArgs
 	typedef TShaderRefBase<FNiagaraShader, FNiagaraShaderMapPointerTable> FShaderReference;
 
 	FNiagaraDataInterfaceSetArgs(
-		FNiagaraDataInterfaceProxy* InDataInterface, FNiagaraSystemInstanceID InSystemInstanceID, const NiagaraEmitterInstanceBatcher* InBatcher, const FShaderReference& InShader, const FNiagaraComputeInstanceData* InComputeInstanceData, const FNiagaraSimStageData* InSimStageData, bool InIsOutputStage, bool InIsIterationStage)
-		: FNiagaraDataInterfaceArgs(InDataInterface, InSystemInstanceID, InBatcher)
+		FNiagaraDataInterfaceProxy* InDataInterface, FNiagaraSystemInstanceID InSystemInstanceID, const FNiagaraGpuComputeDispatchInterface* InComputeDispatchInterface, const FShaderReference& InShader, const FNiagaraComputeInstanceData* InComputeInstanceData, const FNiagaraSimStageData* InSimStageData, bool InIsOutputStage, bool InIsIterationStage)
+		: FNiagaraDataInterfaceArgs(InDataInterface, InSystemInstanceID, InComputeDispatchInterface)
 		, Shader(InShader)
 		, ComputeInstanceData(InComputeInstanceData)
 		, SimStageData(InSimStageData)
diff --git a/Engine/Plugins/FX/Niagara/Source/NiagaraEditor/Private/ViewModels/NiagaraBakerViewModel.cpp b/Engine/Plugins/FX/Niagara/Source/NiagaraEditor/Private/ViewModels/NiagaraBakerViewModel.cpp
--- a/Engine/Plugins/FX/Niagara/Source/NiagaraEditor/Private/ViewModels/NiagaraBakerViewModel.cpp
+++ b/Engine/Plugins/FX/Niagara/Source/NiagaraEditor/Private/ViewModels/NiagaraBakerViewModel.cpp
@@ -9,7 +9,7 @@
 #include "ViewModels/NiagaraSystemViewModel.h"
 #include "Widgets/SNiagaraBakerWidget.h"
 #include "NiagaraBakerRenderer.h"
-#include "NiagaraEmitterInstanceBatcher.h"
+#include "NiagaraGpuComputeDispatchInterface.h"
 #include "NiagaraComponent.h"
 #include "NiagaraSystem.h"
 
@@ -154,12 +154,8 @@ void FNiagaraBakerViewModel::RenderBaker()
 	UWorld* World = PreviewComponent->GetWorld();
 	FSceneInterface* BakerScene = World->Scene;
 
-	NiagaraEmitterInstanceBatcher* NiagaraBatcher = nullptr;
-	if (World->Scene && World->Scene->GetFXSystem())
-	{
-		NiagaraBatcher = static_cast<NiagaraEmitterInstanceBatcher*>(World->Scene->GetFXSystem()->GetInterface(NiagaraEmitterInstanceBatcher::Name));
-	}
-	check(NiagaraBatcher);
+	FNiagaraGpuComputeDispatchInterface* ComputeDispatchInterface = FNiagaraGpuComputeDispatchInterface::Get(World);
+	ensureMsgf(ComputeDispatchInterface, TEXT("The batcher was not valid on the world this may result in incorrect baking"));
 
 	// Create output render targets & output Baker data
 	TArray<UTextureRenderTarget2D*, TInlineAllocator<4>> RenderTargets;
@@ -197,14 +193,9 @@ void FNiagaraBakerViewModel::RenderBaker()
 		PreviewComponent->TickComponent(BakerSettings->GetSeekDelta(), ELevelTick::LEVELTICK_All, nullptr);
 
 		// We need any GPU sims to flush pending ticks
-		if (NiagaraBatcher)
+		if (ComputeDispatchInterface)
 		{
-			ENQUEUE_RENDER_COMMAND(NiagaraFlushBatcher)(
-				[RT_NiagaraBatcher=NiagaraBatcher](FRHICommandListImmediate& RHICmdList)
-				{
-					RT_NiagaraBatcher->ProcessPendingTicksFlush(RHICmdList, true);
-				}
-			);
+			ComputeDispatchInterface->FlushPendingTicks_GameThread();
 		}
 
 		const float WorldTime = FApp::GetCurrentTime() - BakerSettings->StartSeconds - BakerSettings->DurationSeconds + FrameTime;
diff --git a/Engine/Plugins/Runtime/HairStrands/Source/HairStrandsCore/Private/Niagara/NiagaraDataInterfaceHairStrands.cpp b/Engine/Plugins/Runtime/HairStrands/Source/HairStrandsCore/Private/Niagara/NiagaraDataInterfaceHairStrands.cpp
--- a/Engine/Plugins/Runtime/HairStrands/Source/HairStrandsCore/Private/Niagara/NiagaraDataInterfaceHairStrands.cpp
+++ b/Engine/Plugins/Runtime/HairStrands/Source/HairStrandsCore/Private/Niagara/NiagaraDataInterfaceHairStrands.cpp
@@ -3,9 +3,10 @@
 #include "Niagara/NiagaraDataInterfaceHairStrands.h"
 #include "NiagaraShader.h"
 #include "NiagaraComponent.h"
+#include "NiagaraGpuComputeDispatchInterface.h"
 #include "NiagaraRenderer.h"
+#include "NiagaraSimStageData.h"
 #include "NiagaraSystemInstance.h"
-#include "NiagaraEmitterInstanceBatcher.h"
 
 #include "Components/SkeletalMeshComponent.h"
 #include "ShaderParameterUtils.h"
@@ -781,8 +782,8 @@ struct FNDIHairStrandsParametersCS : public FNiagaraDataInterfaceParametersCS
 			SetShaderValue(RHICmdList, ComputeShaderRHI, SampleCount, 0);
 
 			// Set Shader UAV
-			SetUAVParameter(RHICmdList, ComputeShaderRHI, DeformedPositionBuffer, Context.Batcher->GetEmptyUAVFromPool(RHICmdList, PF_R32_FLOAT, ENiagaraEmptyUAVType::Buffer));
-			SetUAVParameter(RHICmdList, ComputeShaderRHI, BoundingBoxBuffer, Context.Batcher->GetEmptyUAVFromPool(RHICmdList, PF_R32_UINT, ENiagaraEmptyUAVType::Buffer));
+			SetUAVParameter(RHICmdList, ComputeShaderRHI, DeformedPositionBuffer, Context.ComputeDispatchInterface->GetEmptyUAVFromPool(RHICmdList, PF_R32_FLOAT, ENiagaraEmptyUAVType::Buffer));
+			SetUAVParameter(RHICmdList, ComputeShaderRHI, BoundingBoxBuffer, Context.ComputeDispatchInterface->GetEmptyUAVFromPool(RHICmdList, PF_R32_UINT, ENiagaraEmptyUAVType::Buffer));
 
 			// Set Shader SRV
 			SetSRVParameter(RHICmdList, ComputeShaderRHI, CurvesOffsetsBuffer, FNiagaraRenderer::GetDummyUIntBuffer());
@@ -877,7 +878,7 @@ void FNDIHairStrandsProxy::InitializePerInstanceData(const FNiagaraSystemInstanc
 	TargetData = &SystemInstancesToProxyData.Add(SystemInstance);
 }
 
-void FNDIHairStrandsProxy::DestroyPerInstanceData(NiagaraEmitterInstanceBatcher* Batcher, const FNiagaraSystemInstanceID& SystemInstance)
+void FNDIHairStrandsProxy::DestroyPerInstanceData(const FNiagaraSystemInstanceID& SystemInstance)
 {
 	check(IsInRenderingThread());
 	//check(SystemInstancesToProxyData.Contains(SystemInstance));
@@ -1068,10 +1069,10 @@ void UNiagaraDataInterfaceHairStrands::DestroyPerInstanceData(void* PerInstanceD
 
 	FNDIHairStrandsProxy* ThisProxy = GetProxyAs<FNDIHairStrandsProxy>();
 	ENQUEUE_RENDER_COMMAND(FNiagaraDIDestroyInstanceData) (
-		[ThisProxy, InstanceID = SystemInstance->GetId(), Batcher = SystemInstance->GetBatcher()](FRHICommandListImmediate& CmdList)
-	{
-		ThisProxy->SystemInstancesToProxyData.Remove(InstanceID);
-	}
+		[ThisProxy, InstanceID = SystemInstance->GetId()](FRHICommandListImmediate& CmdList)
+		{
+			ThisProxy->SystemInstancesToProxyData.Remove(InstanceID);
+		}
 	);
 }
 
diff --git a/Engine/Plugins/Runtime/HairStrands/Source/HairStrandsCore/Private/Niagara/NiagaraDataInterfacePhysicsAsset.cpp b/Engine/Plugins/Runtime/HairStrands/Source/HairStrandsCore/Private/Niagara/NiagaraDataInterfacePhysicsAsset.cpp
--- a/Engine/Plugins/Runtime/HairStrands/Source/HairStrandsCore/Private/Niagara/NiagaraDataInterfacePhysicsAsset.cpp
+++ b/Engine/Plugins/Runtime/HairStrands/Source/HairStrandsCore/Private/Niagara/NiagaraDataInterfacePhysicsAsset.cpp
@@ -10,8 +10,8 @@
 #include "NiagaraComponent.h"
 #include "NiagaraRenderer.h"
 #include "GroomComponent.h"
+#include "NiagaraSimStageData.h"
 #include "NiagaraSystemInstance.h"
-#include "NiagaraEmitterInstanceBatcher.h"
 #include "ShaderParameterUtils.h"
 
 #define LOCTEXT_NAMESPACE "NiagaraDataInterfacePhysicsAsset"
@@ -628,7 +628,7 @@ void FNDIPhysicsAssetProxy::InitializePerInstanceData(const FNiagaraSystemInstan
 	TargetData = &SystemInstancesToProxyData.Add(SystemInstance);
 }
 
-void FNDIPhysicsAssetProxy::DestroyPerInstanceData(NiagaraEmitterInstanceBatcher* Batcher, const FNiagaraSystemInstanceID& SystemInstance)
+void FNDIPhysicsAssetProxy::DestroyPerInstanceData(const FNiagaraSystemInstanceID& SystemInstance)
 {
 	check(IsInRenderingThread());
 	SystemInstancesToProxyData.Remove(SystemInstance);
@@ -805,10 +805,10 @@ void UNiagaraDataInterfacePhysicsAsset::DestroyPerInstanceData(void* PerInstance
 
 	FNDIPhysicsAssetProxy* ThisProxy = GetProxyAs<FNDIPhysicsAssetProxy>();
 	ENQUEUE_RENDER_COMMAND(FNiagaraDIDestroyInstanceData) (
-		[ThisProxy, InstanceID = SystemInstance->GetId(), Batcher = SystemInstance->GetBatcher()](FRHICommandListImmediate& CmdList)
-	{
-		ThisProxy->SystemInstancesToProxyData.Remove(InstanceID);
-	}
+		[ThisProxy, InstanceID = SystemInstance->GetId()](FRHICommandListImmediate& CmdList)
+		{
+			ThisProxy->SystemInstancesToProxyData.Remove(InstanceID);
+		}
 	);
 }
 
diff --git a/Engine/Plugins/Runtime/HairStrands/Source/HairStrandsCore/Private/Niagara/NiagaraDataInterfacePressureGrid.cpp b/Engine/Plugins/Runtime/HairStrands/Source/HairStrandsCore/Private/Niagara/NiagaraDataInterfacePressureGrid.cpp
--- a/Engine/Plugins/Runtime/HairStrands/Source/HairStrandsCore/Private/Niagara/NiagaraDataInterfacePressureGrid.cpp
+++ b/Engine/Plugins/Runtime/HairStrands/Source/HairStrandsCore/Private/Niagara/NiagaraDataInterfacePressureGrid.cpp
@@ -3,9 +3,10 @@
 #include "Niagara/NiagaraDataInterfacePressureGrid.h"
 #include "NiagaraShader.h"
 #include "NiagaraComponent.h"
+#include "NiagaraGpuComputeDispatchInterface.h"
 #include "NiagaraRenderer.h"
+#include "NiagaraSimStageData.h"
 #include "NiagaraSystemInstance.h"
-#include "NiagaraEmitterInstanceBatcher.h"
 
 #include "ShaderParameterUtils.h"
 #include "ClearQuad.h"
@@ -538,7 +539,7 @@ void FNDIPressureGridProxy::PreStage(FRHICommandList& RHICmdList, const FNiagara
 	{
 		if (Context.SimStageData->bFirstStage)
 		{
-			ClearBuffer(RHICmdList, Context.Batcher->GetFeatureLevel(), ProxyData->CurrentGridBuffer, ProxyData->DestinationGridBuffer, ProxyData->GridSize, true);
+			ClearBuffer(RHICmdList, Context.ComputeDispatchInterface->GetFeatureLevel(), ProxyData->CurrentGridBuffer, ProxyData->DestinationGridBuffer, ProxyData->GridSize, true);
 		}
 	}
 }
diff --git a/Engine/Plugins/Runtime/HairStrands/Source/HairStrandsCore/Private/Niagara/NiagaraDataInterfaceVelocityGrid.cpp b/Engine/Plugins/Runtime/HairStrands/Source/HairStrandsCore/Private/Niagara/NiagaraDataInterfaceVelocityGrid.cpp
--- a/Engine/Plugins/Runtime/HairStrands/Source/HairStrandsCore/Private/Niagara/NiagaraDataInterfaceVelocityGrid.cpp
+++ b/Engine/Plugins/Runtime/HairStrands/Source/HairStrandsCore/Private/Niagara/NiagaraDataInterfaceVelocityGrid.cpp
@@ -3,9 +3,10 @@
 #include "Niagara/NiagaraDataInterfaceVelocityGrid.h"
 #include "NiagaraShader.h"
 #include "NiagaraComponent.h"
+#include "NiagaraGpuComputeDispatchInterface.h"
 #include "NiagaraRenderer.h"
+#include "NiagaraSimStageData.h"
 #include "NiagaraSystemInstance.h"
-#include "NiagaraEmitterInstanceBatcher.h"
 
 #include "ShaderParameterUtils.h"
 #include "ClearQuad.h"
@@ -189,7 +190,7 @@ void FNDIVelocityGridParametersCS::Set(FRHICommandList& RHICmdList, const FNiaga
 	}
 	else
 	{
-		SetUAVParameter(RHICmdList, ComputeShaderRHI, GridDestinationBuffer, Context.Batcher->GetEmptyUAVFromPool(RHICmdList, PF_R32_UINT, ENiagaraEmptyUAVType::Buffer));
+		SetUAVParameter(RHICmdList, ComputeShaderRHI, GridDestinationBuffer, Context.ComputeDispatchInterface->GetEmptyUAVFromPool(RHICmdList, PF_R32_UINT, ENiagaraEmptyUAVType::Buffer));
 		SetSRVParameter(RHICmdList, ComputeShaderRHI, GridCurrentBuffer, FNiagaraRenderer::GetDummyUIntBuffer());
 
 		SetShaderValue(RHICmdList, ComputeShaderRHI, GridSize, FIntVector());
@@ -243,7 +244,7 @@ void UNiagaraDataInterfaceVelocityGrid::DestroyPerInstanceData(void* PerInstance
 
 	FNDIVelocityGridProxy* ThisProxy = GetProxyAs<FNDIVelocityGridProxy>();
 	ENQUEUE_RENDER_COMMAND(FNiagaraDIDestroyInstanceData) (
-		[ThisProxy, InstanceID = SystemInstance->GetId(), Batcher = SystemInstance->GetBatcher()](FRHICommandListImmediate& CmdList)
+		[ThisProxy, InstanceID = SystemInstance->GetId()](FRHICommandListImmediate& CmdList)
 	{
 		ThisProxy->SystemInstancesToProxyData.Remove(InstanceID);
 	}
@@ -830,7 +831,7 @@ void FNDIVelocityGridProxy::PostStage(FRHICommandList& RHICmdList, const FNiagar
 	if (ProxyData != nullptr)
 	{
 		//ProxyData->Swap();
-		CopyTexture(RHICmdList, Context.Batcher->GetFeatureLevel(), ProxyData->DestinationGridBuffer, ProxyData->CurrentGridBuffer,  ProxyData->GridSize);
+		CopyTexture(RHICmdList, Context.ComputeDispatchInterface->GetFeatureLevel(), ProxyData->DestinationGridBuffer, ProxyData->CurrentGridBuffer,  ProxyData->GridSize);
 		//FRHICopyTextureInfo CopyInfo;
 		//RHICmdList.CopyTexture(ProxyData->DestinationGridBuffer->GridDataBuffer.Buffer,
 		//	ProxyData->CurrentGridBuffer->GridDataBuffer.Buffer, CopyInfo);
diff --git a/Engine/Plugins/Runtime/HairStrands/Source/HairStrandsCore/Public/Niagara/NiagaraDataInterfaceHairStrands.h b/Engine/Plugins/Runtime/HairStrands/Source/HairStrandsCore/Public/Niagara/NiagaraDataInterfaceHairStrands.h
--- a/Engine/Plugins/Runtime/HairStrands/Source/HairStrandsCore/Public/Niagara/NiagaraDataInterfaceHairStrands.h
+++ b/Engine/Plugins/Runtime/HairStrands/Source/HairStrandsCore/Public/Niagara/NiagaraDataInterfaceHairStrands.h
@@ -711,7 +711,7 @@ struct FNDIHairStrandsProxy : public FNiagaraDataInterfaceProxy
 	void InitializePerInstanceData(const FNiagaraSystemInstanceID& SystemInstance);
 
 	/** Destroy the proxy data if necessary */
-	void DestroyPerInstanceData(NiagaraEmitterInstanceBatcher* Batcher, const FNiagaraSystemInstanceID& SystemInstance);
+	void DestroyPerInstanceData(const FNiagaraSystemInstanceID& SystemInstance);
 
 	/** Launch all pre stage functions */
 	virtual void PreStage(FRHICommandList& RHICmdList, const FNiagaraDataInterfaceStageArgs& Context) override;
diff --git a/Engine/Plugins/Runtime/HairStrands/Source/HairStrandsCore/Public/Niagara/NiagaraDataInterfacePhysicsAsset.h b/Engine/Plugins/Runtime/HairStrands/Source/HairStrandsCore/Public/Niagara/NiagaraDataInterfacePhysicsAsset.h
--- a/Engine/Plugins/Runtime/HairStrands/Source/HairStrandsCore/Public/Niagara/NiagaraDataInterfacePhysicsAsset.h
+++ b/Engine/Plugins/Runtime/HairStrands/Source/HairStrandsCore/Public/Niagara/NiagaraDataInterfacePhysicsAsset.h
@@ -216,7 +216,7 @@ struct FNDIPhysicsAssetProxy : public FNiagaraDataInterfaceProxy
 	void InitializePerInstanceData(const FNiagaraSystemInstanceID& SystemInstance);
 
 	/** Destroy the proxy data if necessary */
-	void DestroyPerInstanceData(NiagaraEmitterInstanceBatcher* Batcher, const FNiagaraSystemInstanceID& SystemInstance);
+	void DestroyPerInstanceData(const FNiagaraSystemInstanceID& SystemInstance);
 
 	/** Launch all pre stage functions */
 	virtual void PreStage(FRHICommandList& RHICmdList, const FNiagaraDataInterfaceStageArgs& Context) override;
