- Detect if a data interface is used with a CPU script  or not
- Allow the static mesh data interface to access GPU resources if it's not used by a CPU script and SRVs will be available on the target platform

diff --git a/Engine/Plugins/FX/Niagara/Source/Niagara/Classes/NiagaraDataInterface.h b/Engine/Plugins/FX/Niagara/Source/Niagara/Classes/NiagaraDataInterface.h
--- a/Engine/Plugins/FX/Niagara/Source/Niagara/Classes/NiagaraDataInterface.h
+++ b/Engine/Plugins/FX/Niagara/Source/Niagara/Classes/NiagaraDataInterface.h
@@ -382,8 +382,11 @@ public:
 	virtual bool HasTickGroupPrereqs() const { return false; }
 	virtual ETickingGroup CalculateTickGroup(const void* PerInstanceData) const { return NiagaraFirstTickGroup; }
 
+	/** Used to determine if we need to create CPU resources for the emitter. */
+	bool IsUsedWithCPUEmitter() const;
+
 	/** Used to determine if we need to create GPU resources for the emitter. */
-	bool IsUsedWithGPUEmitter(class FNiagaraSystemInstance* SystemInstance) const;
+	bool IsUsedWithGPUEmitter() const;
 
 	/** Determines if this type definition matches to a known data interface type.*/
 	static bool IsDataInterfaceType(const FNiagaraTypeDefinition& TypeDef);
@@ -481,6 +484,12 @@ public:
 		PushToRenderThread();
 	}
 
+	void SetUsedByCPUEmitter(bool bUsed = true)
+	{
+		check(IsInGameThread());
+		bUsedByCPUEmitter = bUsed;
+	}
+
 	void SetUsedByGPUEmitter(bool bUsed = true)
 	{
 		check(IsInGameThread());
@@ -510,6 +519,7 @@ protected:
 	TUniquePtr<FNiagaraDataInterfaceProxy> Proxy;
 
 	uint32 bRenderDataDirty : 1;
+	uint32 bUsedByCPUEmitter : 1;
 	uint32 bUsedByGPUEmitter : 1;
 
 private:
diff --git a/Engine/Plugins/FX/Niagara/Source/Niagara/Private/NiagaraCommon.cpp b/Engine/Plugins/FX/Niagara/Source/Niagara/Private/NiagaraCommon.cpp
--- a/Engine/Plugins/FX/Niagara/Source/Niagara/Private/NiagaraCommon.cpp
+++ b/Engine/Plugins/FX/Niagara/Source/Niagara/Private/NiagaraCommon.cpp
@@ -14,6 +14,7 @@
 #include "NiagaraWorldManager.h"
 #include "String/ParseTokens.h"
 #include "UObject/Class.h"
+#include "GPUSkinCache.h"
 
 DECLARE_CYCLE_STAT(TEXT("Niagara - Utilities - PrepareRapidIterationParameters"), STAT_Niagara_Utilities_PrepareRapidIterationParameters, STATGROUP_Niagara);
 
@@ -936,6 +937,11 @@ bool FNiagaraUtilities::AllowGPUCulling(EShaderPlatform ShaderPlatform)
 	return GNiagaraGPUCulling && AllowGPUSorting(ShaderPlatform) && AllowComputeShaders(ShaderPlatform);
 }
 
+bool FNiagaraUtilities::AreBufferSRVsAlwaysCreated(EShaderPlatform ShaderPlatform)
+{
+	return RHISupportsManualVertexFetch(ShaderPlatform) || IsGPUSkinCacheAvailable(ShaderPlatform);
+}
+
 ENiagaraCompileUsageStaticSwitch FNiagaraUtilities::ConvertScriptUsageToStaticSwitchUsage(ENiagaraScriptUsage ScriptUsage)
 {
 	if (ScriptUsage == ENiagaraScriptUsage::ParticleEventScript)
diff --git a/Engine/Plugins/FX/Niagara/Source/Niagara/Private/NiagaraComponent.cpp b/Engine/Plugins/FX/Niagara/Source/Niagara/Private/NiagaraComponent.cpp
--- a/Engine/Plugins/FX/Niagara/Source/Niagara/Private/NiagaraComponent.cpp
+++ b/Engine/Plugins/FX/Niagara/Source/Niagara/Private/NiagaraComponent.cpp
@@ -2956,7 +2956,8 @@ void UNiagaraComponent::SetOverrideParameterStoreValue(const FNiagaraVariableBas
 		UNiagaraDataInterface* OriginalDI = InValue.GetDataInterface();
 		UNiagaraDataInterface* DuplicatedDI = DuplicateObject(OriginalDI, this);
 		OverrideParameters.SetDataInterface(DuplicatedDI, InKey);
-		DuplicatedDI->SetUsedByGPUEmitter(OriginalDI->IsUsedWithGPUEmitter(nullptr));
+		DuplicatedDI->SetUsedByCPUEmitter(OriginalDI->IsUsedWithCPUEmitter());
+		DuplicatedDI->SetUsedByGPUEmitter(OriginalDI->IsUsedWithGPUEmitter());
 	}
 	else if (InKey.IsUObject())
 	{
diff --git a/Engine/Plugins/FX/Niagara/Source/Niagara/Private/NiagaraDataInterface.cpp b/Engine/Plugins/FX/Niagara/Source/Niagara/Private/NiagaraDataInterface.cpp
--- a/Engine/Plugins/FX/Niagara/Source/Niagara/Private/NiagaraDataInterface.cpp
+++ b/Engine/Plugins/FX/Niagara/Source/Niagara/Private/NiagaraDataInterface.cpp
@@ -14,6 +14,7 @@
 UNiagaraDataInterface::UNiagaraDataInterface(FObjectInitializer const& ObjectInitializer)
 {
 	bRenderDataDirty = false;
+	bUsedByCPUEmitter = false;
 	bUsedByGPUEmitter = false;
 }
 
@@ -99,7 +100,12 @@ bool UNiagaraDataInterface::Equals(const UNiagaraDataInterface* Other) const
 	return true;
 }
 
-bool UNiagaraDataInterface::IsUsedWithGPUEmitter(FNiagaraSystemInstance* SystemInstance) const
+bool UNiagaraDataInterface::IsUsedWithCPUEmitter() const
+{
+	return bUsedByCPUEmitter;
+}
+
+bool UNiagaraDataInterface::IsUsedWithGPUEmitter() const
 {
 	return bUsedByGPUEmitter;
 }
diff --git a/Engine/Plugins/FX/Niagara/Source/Niagara/Private/NiagaraDataInterfaceSkeletalMesh.cpp b/Engine/Plugins/FX/Niagara/Source/Niagara/Private/NiagaraDataInterfaceSkeletalMesh.cpp
--- a/Engine/Plugins/FX/Niagara/Source/Niagara/Private/NiagaraDataInterfaceSkeletalMesh.cpp
+++ b/Engine/Plugins/FX/Niagara/Source/Niagara/Private/NiagaraDataInterfaceSkeletalMesh.cpp
@@ -1963,7 +1963,7 @@ bool FNDISkeletalMesh_InstanceData::Init(UNiagaraDataInterfaceSkeletalMesh* Inte
 
 		const bool MeshValid = SkeletalMesh.IsValid();
 		const bool SupportUvMappingCpu = UsedByCpuUvMapping && MeshValid;
-		const bool SupportUvMappingGpu = UsedByGpuUvMapping && MeshValid && Interface->IsUsedWithGPUEmitter(SystemInstance);
+		const bool SupportUvMappingGpu = UsedByGpuUvMapping && MeshValid && Interface->IsUsedWithGPUEmitter();
 
 		FSkeletalMeshUvMappingUsage UvMappingUsage(SupportUvMappingCpu, SupportUvMappingGpu);
 
@@ -1990,7 +1990,7 @@ bool FNDISkeletalMesh_InstanceData::Init(UNiagaraDataInterfaceSkeletalMesh* Inte
 
 		const bool MeshValid = SkeletalMesh.IsValid();
 		const bool SupportConnectivityCpu = UsedByCpuConnectivity && MeshValid;
-		const bool SupportConnectivityGpu = UsedByGpuConnectivity && MeshValid && Interface->IsUsedWithGPUEmitter(SystemInstance);
+		const bool SupportConnectivityGpu = UsedByGpuConnectivity && MeshValid && Interface->IsUsedWithGPUEmitter();
 
 		FSkeletalMeshConnectivityUsage ConnectivityUsage(SupportConnectivityCpu, SupportConnectivityGpu);
 
@@ -2187,7 +2187,7 @@ bool FNDISkeletalMesh_InstanceData::Init(UNiagaraDataInterfaceSkeletalMesh* Inte
 			}
 		}
 
-		if (Interface->IsUsedWithGPUEmitter(SystemInstance))
+		if (Interface->IsUsedWithGPUEmitter())
 		{
 			GPUSkinBoneInfluenceType BoneInfluenceType = SkinWeightBuffer->GetBoneInfluenceType();
 			bUnlimitedBoneInfluences = (BoneInfluenceType == GPUSkinBoneInfluenceType::UnlimitedBoneInfluence);
diff --git a/Engine/Plugins/FX/Niagara/Source/Niagara/Private/NiagaraDataInterfaceSpline.cpp b/Engine/Plugins/FX/Niagara/Source/Niagara/Private/NiagaraDataInterfaceSpline.cpp
--- a/Engine/Plugins/FX/Niagara/Source/Niagara/Private/NiagaraDataInterfaceSpline.cpp
+++ b/Engine/Plugins/FX/Niagara/Source/Niagara/Private/NiagaraDataInterfaceSpline.cpp
@@ -592,7 +592,7 @@ bool UNiagaraDataInterfaceSpline::PerInstanceTick(void* PerInstanceData, FNiagar
 		InstData->ComponentTransform = SplineComponent->GetComponentTransform();
 		InstData->DefaultUpVector = SplineComponent->DefaultUpVector;
 
-		bool bShouldBuildLUT = (bUseLUT || IsUsedWithGPUEmitter(SystemInstance)) && InstData->SplineLUT.MaxIndex < 0;
+		bool bShouldBuildLUT = (bUseLUT || IsUsedWithGPUEmitter()) && InstData->SplineLUT.MaxIndex < 0;
 		
 		if (InstData->SplineCurvesVersion != SplineComponent->SplineCurves.Version)
 		{
@@ -601,17 +601,17 @@ bool UNiagaraDataInterfaceSpline::PerInstanceTick(void* PerInstanceData, FNiagar
 			InstData->bSyncedGPUCopy = false;
 			InstData->SplineLUT.Reset();
 
-			bShouldBuildLUT = bUseLUT || IsUsedWithGPUEmitter(SystemInstance);
-	}
+			bShouldBuildLUT = bUseLUT || IsUsedWithGPUEmitter();
+		}
 
-		bool bShouldSyncToGPU = IsUsedWithGPUEmitter(SystemInstance) && !InstData->bSyncedGPUCopy && InstData->SplineLUT.MaxIndex != INDEX_NONE;
+		bool bShouldSyncToGPU = IsUsedWithGPUEmitter() && !InstData->bSyncedGPUCopy && InstData->SplineLUT.MaxIndex != INDEX_NONE;
 	
 		// We must build the LUT if this is for GPU regardless of settings
 		if (bShouldBuildLUT)
 		{
 			InstData->SplineLUT.BuildLUT(InstData->SplineCurves, bUseLUT? NumLUTSteps : 256/*Default the LUT to a reasonable value if it's not specifically enabled*/);
 
-			bShouldSyncToGPU = IsUsedWithGPUEmitter(SystemInstance);				
+			bShouldSyncToGPU = IsUsedWithGPUEmitter();				
 		}
 
 
diff --git a/Engine/Plugins/FX/Niagara/Source/Niagara/Private/NiagaraDataInterfaceStaticMesh.cpp b/Engine/Plugins/FX/Niagara/Source/Niagara/Private/NiagaraDataInterfaceStaticMesh.cpp
--- a/Engine/Plugins/FX/Niagara/Source/Niagara/Private/NiagaraDataInterfaceStaticMesh.cpp
+++ b/Engine/Plugins/FX/Niagara/Source/Niagara/Private/NiagaraDataInterfaceStaticMesh.cpp
@@ -312,9 +312,20 @@ bool FNDIStaticMesh_InstanceData::Init(UNiagaraDataInterfaceStaticMesh* Interfac
 	}
 	else if (!Mesh->bAllowCPUAccess)
 	{
-		UE_LOG(LogNiagara, Log, TEXT("StaticMesh data interface using a mesh that does not allow CPU access. Interface: %s, Mesh: %s"),
-			*Interface->GetFullName(), *Mesh->GetFullName());
-		Mesh = nullptr; // Disallow usage of this mesh to prevent issues on cooked builds
+		//-TODO: This will be modified with the refactor to allow GPU to be used when CPU is also used
+		if (Interface->IsUsedWithGPUEmitter() && !Interface->IsUsedWithCPUEmitter())
+		{
+			if ( !FNiagaraUtilities::AreBufferSRVsAlwaysCreated(GMaxRHIShaderPlatform) )
+			{
+				UE_LOG(LogNiagara, Log, TEXT("StaticMesh data interface used by GPU emitter but does not have SRV access on this platform.  Enable CPU access to fix this issue. Interface: %s, Mesh: %s"), *Interface->GetFullName(), *StaticMesh->GetFullName());
+				Mesh = nullptr;
+			}
+		}
+		else
+		{
+			UE_LOG(LogNiagara, Log, TEXT("StaticMesh data interface used by CPU emitter and does not allow CPU access. Interface: %s, Mesh: %s"), *Interface->GetFullName(), *StaticMesh->GetFullName());
+			Mesh = nullptr;
+		}
 	}
 
 	if (Mesh != nullptr)
@@ -1538,11 +1549,10 @@ bool UNiagaraDataInterfaceStaticMesh::InitPerInstanceData(void* PerInstanceData,
 		FStaticMeshGpuSpawnBuffer* MeshGpuSpawnBuffer = nullptr;
 		const FStaticMeshLODResources* GpuMeshLODResource = Inst->CachedLOD;
 
-		if (GpuMeshLODResource && IsUsedWithGPUEmitter(SystemInstance))
+		if (GpuMeshLODResource && IsUsedWithGPUEmitter())
 		{
-			// Always allocate when bAllowCPUAccess (index buffer can only have SRV created in this case as of today)
-			// We do not know if this interface is allocated for CPU or GPU so we allocate for both case... TODO: have some cached data created in case a GPU version is needed?
-			ensure(Inst->StaticMesh->bAllowCPUAccess); // this should have been verified in Init()
+			// This should already have been validated, but to be safe
+			ensureMsgf(Inst->StaticMesh->bAllowCPUAccess || FNiagaraUtilities::AreBufferSRVsAlwaysCreated(GMaxRHIShaderPlatform), TEXT("Creating GPU resources in an invalid case CPUAccess(%d) SRVsAlwaysCreated(%d)"), Inst->StaticMesh->bAllowCPUAccess, FNiagaraUtilities::AreBufferSRVsAlwaysCreated(GMaxRHIShaderPlatform));
 
 			MeshGpuSpawnBuffer = new FStaticMeshGpuSpawnBuffer;
 			MeshGpuSpawnBuffer->Initialise(GpuMeshLODResource, *this, Inst);
diff --git a/Engine/Plugins/FX/Niagara/Source/Niagara/Private/NiagaraSystemInstance.cpp b/Engine/Plugins/FX/Niagara/Source/Niagara/Private/NiagaraSystemInstance.cpp
--- a/Engine/Plugins/FX/Niagara/Source/Niagara/Private/NiagaraSystemInstance.cpp
+++ b/Engine/Plugins/FX/Niagara/Source/Niagara/Private/NiagaraSystemInstance.cpp
@@ -1394,7 +1394,7 @@ void FNiagaraSystemInstance::InitDataInterfaces()
 	//Now the interfaces in the simulations are all correct, we can build the per instance data table.
 	int32 InstanceDataSize = 0;
 	DataInterfaceInstanceDataOffsets.Empty();
-	auto CalcInstDataSize = [&](const FNiagaraParameterStore& ParamStore, bool bIsGPUSimulation, bool bSearchInstanceParams)
+	auto CalcInstDataSize = [&](const FNiagaraParameterStore& ParamStore, bool bIsCPUSimulation, bool bIsGPUSimulation, bool bSearchInstanceParams)
 	{
 		const TArrayView<const FNiagaraVariableWithOffset> Params = ParamStore.ReadParameterVariables();
 		const TArray<UNiagaraDataInterface*>& Interfaces = ParamStore.GetDataInterfaces();
@@ -1446,25 +1446,29 @@ void FNiagaraSystemInstance::InitDataInterfaces()
 							Proxy->SourceDIName = Var.GetName();
 						}
 					}
+					else if (bIsCPUSimulation)
+					{
+						Interface->SetUsedByCPUEmitter(true);
+					}
 				}
 			}
 		}
 	};
 
-	CalcInstDataSize(InstanceParameters, false, false);//This probably should be a proper exec context.
+	CalcInstDataSize(InstanceParameters, false, false, false);//This probably should be a proper exec context.
 
 	if (SystemSimulation->GetIsSolo() && FNiagaraSystemSimulation::UseLegacySystemSimulationContexts())
 	{
-		CalcInstDataSize(SystemSimulation->GetSpawnExecutionContext()->Parameters, false, false);
+		CalcInstDataSize(SystemSimulation->GetSpawnExecutionContext()->Parameters, true, false, false);
 		SystemSimulation->GetSpawnExecutionContext()->DirtyDataInterfaces();
 
-		CalcInstDataSize(SystemSimulation->GetUpdateExecutionContext()->Parameters, false, false);
+		CalcInstDataSize(SystemSimulation->GetUpdateExecutionContext()->Parameters, true, false, false);
 		SystemSimulation->GetUpdateExecutionContext()->DirtyDataInterfaces();
 	}
 	else
 	{
-		CalcInstDataSize(SystemSimulation->GetSpawnExecutionContext()->Parameters, false, true);
-		CalcInstDataSize(SystemSimulation->GetUpdateExecutionContext()->Parameters, false, true);
+		CalcInstDataSize(SystemSimulation->GetSpawnExecutionContext()->Parameters, true, false, true);
+		CalcInstDataSize(SystemSimulation->GetUpdateExecutionContext()->Parameters, true, false, true);
 	}
 
 	//Iterate over interfaces to get size for table and clear their interface bindings.
@@ -1478,16 +1482,16 @@ void FNiagaraSystemInstance::InitDataInterfaces()
 
 		const bool bGPUSimulation = Sim.GetCachedEmitter() && (Sim.GetCachedEmitter()->SimTarget == ENiagaraSimTarget::GPUComputeSim);
 
-		CalcInstDataSize(Sim.GetSpawnExecutionContext().Parameters, bGPUSimulation, false);
-		CalcInstDataSize(Sim.GetUpdateExecutionContext().Parameters, bGPUSimulation, false);
+		CalcInstDataSize(Sim.GetSpawnExecutionContext().Parameters, !bGPUSimulation, bGPUSimulation, false);
+		CalcInstDataSize(Sim.GetUpdateExecutionContext().Parameters, !bGPUSimulation, bGPUSimulation, false);
 		for (int32 i = 0; i < Sim.GetEventExecutionContexts().Num(); i++)
 		{
-			CalcInstDataSize(Sim.GetEventExecutionContexts()[i].Parameters, bGPUSimulation, false);
+			CalcInstDataSize(Sim.GetEventExecutionContexts()[i].Parameters, !bGPUSimulation, bGPUSimulation, false);
 		}
 
 		if (Sim.GetCachedEmitter() && Sim.GetCachedEmitter()->SimTarget == ENiagaraSimTarget::GPUComputeSim && Sim.GetGPUContext())
 		{
-			CalcInstDataSize(Sim.GetGPUContext()->CombinedParamStore, bGPUSimulation, false);
+			CalcInstDataSize(Sim.GetGPUContext()->CombinedParamStore, !bGPUSimulation, bGPUSimulation, false);
 		}
 
 		//Also force a rebind while we're here.
diff --git a/Engine/Plugins/FX/Niagara/Source/Niagara/Public/NiagaraCommon.h b/Engine/Plugins/FX/Niagara/Source/Niagara/Public/NiagaraCommon.h
--- a/Engine/Plugins/FX/Niagara/Source/Niagara/Public/NiagaraCommon.h
+++ b/Engine/Plugins/FX/Niagara/Source/Niagara/Public/NiagaraCommon.h
@@ -1208,6 +1208,9 @@ namespace FNiagaraUtilities
 	// Are we able to use the GPU for sorting?
 	bool AllowGPUSorting(EShaderPlatform ShaderPlatform);
 
+	// Helper function to detect if SRVs are always created for buffers or not
+	bool AreBufferSRVsAlwaysCreated(EShaderPlatform ShaderPlatform);
+
 	ENiagaraCompileUsageStaticSwitch NIAGARA_API ConvertScriptUsageToStaticSwitchUsage(ENiagaraScriptUsage ScriptUsage);
 	ENiagaraScriptContextStaticSwitch NIAGARA_API ConvertScriptUsageToStaticSwitchContext(ENiagaraScriptUsage ScriptUsage);
 	
