- Disable culling & sorting if the cull count buffer has already been acquired

#rnx

diff --git a/Engine/Plugins/FX/Niagara/Source/Niagara/Classes/NiagaraGPUInstanceCountManager.h b/Engine/Plugins/FX/Niagara/Source/Niagara/Classes/NiagaraGPUInstanceCountManager.h
--- a/Engine/Plugins/FX/Niagara/Source/Niagara/Classes/NiagaraGPUInstanceCountManager.h
+++ b/Engine/Plugins/FX/Niagara/Source/Niagara/Classes/NiagaraGPUInstanceCountManager.h
@@ -77,6 +77,11 @@ public:
 
 	uint32 AcquireEntry();
 
+	bool CanAcquireCulledEntry() const
+	{
+		return !bAcquiredCulledCounts;
+	}
+
 	uint32 AcquireCulledEntry()
 	{
 		check(!bAcquiredCulledCounts);
diff --git a/Engine/Plugins/FX/Niagara/Source/Niagara/Private/NiagaraRendererMeshes.cpp b/Engine/Plugins/FX/Niagara/Source/Niagara/Private/NiagaraRendererMeshes.cpp
--- a/Engine/Plugins/FX/Niagara/Source/Niagara/Private/NiagaraRendererMeshes.cpp
+++ b/Engine/Plugins/FX/Niagara/Source/Niagara/Private/NiagaraRendererMeshes.cpp
@@ -320,6 +320,14 @@ void FNiagaraRendererMeshes::PrepareParticleMeshRenderData(FParticleMeshRenderDa
 				ParticleMeshRenderData.bNeedsCull = false;
 			}
 			ParticleMeshRenderData.bNeedsSort &= ParticleMeshRenderData.bSortCullOnGpu;
+
+			//-TODO: Culling and sorting from InitViewsAfterPrePass can not be respected if the culled entries have already been acquired
+			if ((ParticleMeshRenderData.bNeedsSort || ParticleMeshRenderData.bNeedsCull) && !SceneProxy->GetComputeDispatchInterface()->GetGPUInstanceCounterManager().CanAcquireCulledEntry())
+			{
+				//ensureMsgf(false, TEXT("Culling & sorting is not supported once the culled counts have been acquired, sorting & culling will be disabled for these draws."));
+				ParticleMeshRenderData.bNeedsSort = false;
+				ParticleMeshRenderData.bNeedsCull = false;
+			}
 		}
 		else
 		{
diff --git a/Engine/Plugins/FX/Niagara/Source/Niagara/Private/NiagaraRendererSprites.cpp b/Engine/Plugins/FX/Niagara/Source/Niagara/Private/NiagaraRendererSprites.cpp
--- a/Engine/Plugins/FX/Niagara/Source/Niagara/Private/NiagaraRendererSprites.cpp
+++ b/Engine/Plugins/FX/Niagara/Source/Niagara/Private/NiagaraRendererSprites.cpp
@@ -270,6 +270,14 @@ void FNiagaraRendererSprites::PrepareParticleSpriteRenderData(FParticleSpriteRen
 				ParticleSpriteRenderData.bNeedsCull = false;
 			}
 			ParticleSpriteRenderData.bNeedsSort &= ParticleSpriteRenderData.bSortCullOnGpu;
+
+			//-TODO: Culling and sorting from InitViewsAfterPrePass can not be respected if the culled entries have already been acquired
+			if ((ParticleSpriteRenderData.bNeedsSort || ParticleSpriteRenderData.bNeedsCull) && !SceneProxy->GetComputeDispatchInterface()->GetGPUInstanceCounterManager().CanAcquireCulledEntry())
+			{
+				//ensureMsgf(false, TEXT("Culling & sorting is not supported once the culled counts have been acquired, sorting & culling will be disabled for these draws."));
+				ParticleSpriteRenderData.bNeedsSort = false;
+				ParticleSpriteRenderData.bNeedsCull = false;
+			}
 		}
 		else
 		{
