- Don't add to GPUDataInterfaces if the data interface is not used by a GPU emitter
- This reduces the performance & memory of generating GPU ticks in some situations

diff --git a/Engine/Plugins/FX/Niagara/Source/Niagara/Private/NiagaraSystemInstance.cpp b/Engine/Plugins/FX/Niagara/Source/Niagara/Private/NiagaraSystemInstance.cpp
--- a/Engine/Plugins/FX/Niagara/Source/Niagara/Private/NiagaraSystemInstance.cpp
+++ b/Engine/Plugins/FX/Niagara/Source/Niagara/Private/NiagaraSystemInstance.cpp
@@ -1524,7 +1524,7 @@ void FNiagaraSystemInstance::InitDataInterfaces()
 				PostTickDataInterfaces.Add(i);
 			}
 
-			if (bHasGPUEmitters)
+			if (Interface->IsUsedWithGPUEmitter())
 			{
 				const int32 GPUDataSize = Interface->PerInstanceDataPassedToRenderThreadSize();
 				if (GPUDataSize > 0)
