- Remove case that blocked SRV for uniform sampling
- Reduce the size of GPU triangle uniform sampling buffer from 16 bytes per triangle to 8 bytes

diff --git a/Engine/Plugins/FX/Niagara/Source/Niagara/Private/NiagaraDataInterfaceStaticMesh.cpp b/Engine/Plugins/FX/Niagara/Source/Niagara/Private/NiagaraDataInterfaceStaticMesh.cpp
--- a/Engine/Plugins/FX/Niagara/Source/Niagara/Private/NiagaraDataInterfaceStaticMesh.cpp
+++ b/Engine/Plugins/FX/Niagara/Source/Niagara/Private/NiagaraDataInterfaceStaticMesh.cpp
@@ -136,7 +136,7 @@ void FStaticMeshGpuSpawnBuffer::Initialise(const FStaticMeshLODResources* Res, c
 
 	if (InstanceData->bIsGpuUniformlyDistributedSampling)
 	{
-		BufferUniformTriangleSamplingSRV = Res->AreaWeightedSectionSamplersBuffer.GetBufferSRV(); // Cache that SRV for later
+		BufferUniformTriangleSamplingSRV = Res->AreaWeightedSectionSamplersBuffer.GetBufferSRV();
 	}
 
 	// Init socket data
@@ -2797,7 +2797,7 @@ bool UNiagaraDataInterfaceStaticMesh::GetFunctionHLSL(const FNiagaraDataInterfac
 				{
 					// Uniform area weighted position selection (using alias method from Alias method from FWeightedRandomSampler)
 					uint TriangleIndex = min(uint(RandT0*float(SectionTriangleCount)), SectionTriangleCount-1);
-					uint4 TriangleData = {MeshTriangleBufferName}[SectionFirstTriangle + TriangleIndex];
+					uint2 TriangleData = {MeshTriangleBufferName}[SectionFirstTriangle + TriangleIndex];
 
 					// Alias check
 					float RandT1 = NiagaraInternalNoise(1, 2, 3);
@@ -2850,7 +2850,7 @@ bool UNiagaraDataInterfaceStaticMesh::GetFunctionHLSL(const FNiagaraDataInterfac
 				{
 					// Uniform area weighted position selection (using alias method from Alias method from FWeightedRandomSampler)
 					uint TriangleIndex = min(uint(RandT0*float(SectionTriangleCount)), SectionTriangleCount-1);
-					uint4 TriangleData = {MeshTriangleBufferName}[SectionFirstTriangle + TriangleIndex];
+					uint2 TriangleData = {MeshTriangleBufferName}[SectionFirstTriangle + TriangleIndex];
 
 					// Alias check
 					float RandT1 = NiagaraInternalNoise(1, 2, 3);
@@ -3429,7 +3429,7 @@ void UNiagaraDataInterfaceStaticMesh::GetParameterDefinitionHLSL(const FNiagaraD
 	OutHLSL += TEXT("Buffer<float2> ") + ParamNames.MeshTexCoordBufferName + TEXT(";\n");
 	OutHLSL += TEXT("Buffer<float4> ") + ParamNames.MeshColorBufferName + TEXT(";\n");
 	OutHLSL += TEXT("Buffer<uint4> ") + ParamNames.MeshSectionBufferName + TEXT(";\n");
-	OutHLSL += TEXT("Buffer<uint4> ") + ParamNames.MeshTriangleBufferName + TEXT(";\n");
+	OutHLSL += TEXT("Buffer<uint2> ") + ParamNames.MeshTriangleBufferName + TEXT(";\n");
 	OutHLSL += TEXT("uint ") + ParamNames.UseColorBufferName + TEXT(";\n");
 	OutHLSL += TEXT("uint ") + ParamNames.SectionCountName + TEXT(";\n");
 	OutHLSL += TEXT("float4x4 ") + ParamNames.InstanceTransformName + TEXT(";\n");
diff --git a/Engine/Plugins/FX/Niagara/Source/NiagaraCore/Private/NiagaraCustomVersion.cpp b/Engine/Plugins/FX/Niagara/Source/NiagaraCore/Private/NiagaraCustomVersion.cpp
--- a/Engine/Plugins/FX/Niagara/Source/NiagaraCore/Private/NiagaraCustomVersion.cpp
+++ b/Engine/Plugins/FX/Niagara/Source/NiagaraCore/Private/NiagaraCustomVersion.cpp
@@ -8,5 +8,5 @@ const FGuid FNiagaraCustomVersion::GUID(0xFCF57AFA, 0x50764283, 0xB9A9E658, 0xFF
 // Register the custom version with core
 FCustomVersionRegistration GRegisterNiagaraCustomVersion(FNiagaraCustomVersion::GUID, FNiagaraCustomVersion::LatestVersion, TEXT("NiagaraVer"));
 
-const FGuid FNiagaraCustomVersion::LatestScriptCompileVersion(0x1448E584, 0xBC404634, 0xA29E0B9B, 0xFDF57669);
+const FGuid FNiagaraCustomVersion::LatestScriptCompileVersion(0xC6985CB9, 0xA75A5C45, 0xA99D541B, 0x4E8A9FB0);
 
diff --git a/Engine/Source/Runtime/Engine/Private/StaticMesh.cpp b/Engine/Source/Runtime/Engine/Private/StaticMesh.cpp
--- a/Engine/Source/Runtime/Engine/Private/StaticMesh.cpp
+++ b/Engine/Source/Runtime/Engine/Private/StaticMesh.cpp
@@ -229,7 +229,7 @@ void FStaticMeshSectionAreaWeightedTriangleSamplerBuffer::InitRHI()
 		{
 			TriangleCount += (*Samplers)[i].GetNumEntries();
 		}
-		uint32 SizeByte = TriangleCount * sizeof(SectionTriangleInfo);
+		const uint32 SizeByte = TriangleCount * sizeof(SectionTriangleInfo);
 
 		FRHIResourceCreateInfo CreateInfo(TEXT("StaticMeshSectionAreaWeightedTriangleSamplerBuffer"));
 		BufferSectionTriangleRHI = RHICreateBuffer(SizeByte, BUF_Static | BUF_VertexBuffer | BUF_ShaderResource, 0, ERHIAccess::VertexOrIndexBuffer | ERHIAccess::SRVMask, CreateInfo);
@@ -245,14 +245,14 @@ void FStaticMeshSectionAreaWeightedTriangleSamplerBuffer::InitRHI()
 
 			for (uint32 t = 0; t < NumTriangle; ++t)
 			{
-				SectionTriangleInfo NewTriangleInfo = { ProbTris[t], (uint32)AliasTris[t], 0, 0 };
+				SectionTriangleInfo NewTriangleInfo = { ProbTris[t], (uint32)AliasTris[t] };
 				*SectionTriangleInfoBuffer = NewTriangleInfo;
 				SectionTriangleInfoBuffer++;
 			}
 		}
 		RHIUnlockBuffer(BufferSectionTriangleRHI);
 
-		BufferSectionTriangleSRV = RHICreateShaderResourceView(BufferSectionTriangleRHI, sizeof(SectionTriangleInfo), PF_R32G32B32A32_UINT);
+		BufferSectionTriangleSRV = RHICreateShaderResourceView(BufferSectionTriangleRHI, sizeof(SectionTriangleInfo), PF_R32G32_UINT);
 	}
 }
 
@@ -1307,7 +1307,7 @@ void FStaticMeshLODResources::InitResources(UStaticMesh* Parent)
 		BeginInitResource(&AdditionalIndexBuffers->ReversedDepthOnlyIndexBuffer);
 	}
 
-	if (Parent->bSupportGpuUniformlyDistributedSampling && Parent->bSupportUniformlyDistributedSampling && Parent->bAllowCPUAccess)
+	if (Parent->bSupportGpuUniformlyDistributedSampling && Parent->bSupportUniformlyDistributedSampling && (AreaWeightedSampler.GetNumEntries() > 0))
 	{
 		AreaWeightedSectionSamplersBuffer.Init(&AreaWeightedSectionSamplers);
 		BeginInitResource(&AreaWeightedSectionSamplersBuffer);
diff --git a/Engine/Source/Runtime/Engine/Public/StaticMeshResources.h b/Engine/Source/Runtime/Engine/Public/StaticMeshResources.h
--- a/Engine/Source/Runtime/Engine/Public/StaticMeshResources.h
+++ b/Engine/Source/Runtime/Engine/Public/StaticMeshResources.h
@@ -280,8 +280,6 @@ private:
 	{
 		float  Prob;
 		uint32 Alias;
-		uint32 pad0;
-		uint32 pad1;
 	};
 
 	FBufferRHIRef BufferSectionTriangleRHI = nullptr;
