- Fix buffer overflow from LWC changes, when FVector became doubles the data we copied to the GPU was the incorrect size

#rnx

diff --git a/Engine/Plugins/FX/Niagara/Source/Niagara/Private/NiagaraGpuComputeDebug.cpp b/Engine/Plugins/FX/Niagara/Source/Niagara/Private/NiagaraGpuComputeDebug.cpp
--- a/Engine/Plugins/FX/Niagara/Source/Niagara/Private/NiagaraGpuComputeDebug.cpp
+++ b/Engine/Plugins/FX/Niagara/Source/Niagara/Private/NiagaraGpuComputeDebug.cpp
@@ -74,7 +74,10 @@ void FNiagaraGpuComputeDebug::Tick(FRHICommandListImmediate& RHICmdList)
 		DebugDrawData->StaticLineCount = DebugDrawData->StaticLines.Num();
 		if (DebugDrawData->StaticLineCount > 0 )
 		{
-			const uint32 NumElements = FMath::DivideAndRoundUp(DebugDrawData->StaticLineCount, 64u) * 64u * 7u;
+			constexpr uint32 NumFloatsPerLine = 7;
+			static_assert(sizeof(FNiagaraSimulationDebugDrawData::FGpuLine) == (NumFloatsPerLine * sizeof(float)), "Line size does not match expected GPU size");
+
+			const uint32 NumElements = FMath::DivideAndRoundUp(DebugDrawData->StaticLineCount, 64u) * 64u * NumFloatsPerLine;
 			const uint32 RequiredBytes = NumElements * sizeof(float);
 			if ( DebugDrawData->StaticLineBuffer.NumBytes < RequiredBytes )
 			{
diff --git a/Engine/Plugins/FX/Niagara/Source/Niagara/Private/NiagaraGpuComputeDebug.h b/Engine/Plugins/FX/Niagara/Source/Niagara/Private/NiagaraGpuComputeDebug.h
--- a/Engine/Plugins/FX/Niagara/Source/Niagara/Private/NiagaraGpuComputeDebug.h
+++ b/Engine/Plugins/FX/Niagara/Source/Niagara/Private/NiagaraGpuComputeDebug.h
@@ -16,9 +16,9 @@ struct FNiagaraSimulationDebugDrawData
 {
 	struct FGpuLine
 	{
-		FVector	Start;
-		FVector	End;
-		uint32	Color;
+		FVector3f	Start;
+		FVector3f	End;
+		uint32		Color;
 	};
 
 	bool				bRequiresUpdate = true;
