Fixed a bug that prevented playing audio from Niagara in cooked builds

diff --git a/Engine/Plugins/FX/Niagara/Source/Niagara/Classes/NiagaraDataInterfaceAudioPlayer.h b/Engine/Plugins/FX/Niagara/Source/Niagara/Classes/NiagaraDataInterfaceAudioPlayer.h
--- a/Engine/Plugins/FX/Niagara/Source/Niagara/Classes/NiagaraDataInterfaceAudioPlayer.h
+++ b/Engine/Plugins/FX/Niagara/Source/Niagara/Classes/NiagaraDataInterfaceAudioPlayer.h
@@ -91,6 +91,8 @@ public:
 	/** If true then this data interface only processes sounds during active gameplay. This is useful when you are working in the preview window and the sounds annoy you. */
 	UPROPERTY(EditAnywhere, AdvancedDisplay, Category = "Audio")
 	bool bOnlyActiveDuringGameplay = false;
+
+	virtual bool UpgradeFunctionCall(FNiagaraFunctionSignature& FunctionSignature) override;
 #endif
 	
 	//UObject Interface
diff --git a/Engine/Plugins/FX/Niagara/Source/Niagara/Private/NiagaraDataInterfaceAudioPlayer.cpp b/Engine/Plugins/FX/Niagara/Source/Niagara/Private/NiagaraDataInterfaceAudioPlayer.cpp
--- a/Engine/Plugins/FX/Niagara/Source/Niagara/Private/NiagaraDataInterfaceAudioPlayer.cpp
+++ b/Engine/Plugins/FX/Niagara/Source/Niagara/Private/NiagaraDataInterfaceAudioPlayer.cpp
@@ -87,6 +87,22 @@ UNiagaraDataInterfaceAudioPlayer::UNiagaraDataInterfaceAudioPlayer(FObjectInitia
 	MaxPlaysPerTick = 10;
 }
 
+#if WITH_EDITORONLY_DATA
+bool UNiagaraDataInterfaceAudioPlayer::UpgradeFunctionCall(FNiagaraFunctionSignature& FunctionSignature)
+{
+	bool bWasChanged = false;
+
+	// The play audio function was marked to have side effects, so older version need the exec pin 
+	if (FunctionSignature.Name == PlayAudioName && FunctionSignature.bRequiresExecPin == false)
+	{
+		FunctionSignature.bRequiresExecPin = true;
+		bWasChanged = true;
+	}
+
+	return bWasChanged;
+}
+#endif
+
 void UNiagaraDataInterfaceAudioPlayer::PostInitProperties()
 {
 	Super::PostInitProperties();
@@ -237,12 +253,11 @@ void UNiagaraDataInterfaceAudioPlayer::GetFunctions(TArray<FNiagaraFunctionSigna
 	Sig.Name = PlayAudioName;
 #if WITH_EDITORONLY_DATA
 	Sig.Description = NSLOCTEXT("Niagara", "PlayAudioDIFunctionDescription", "This function plays a sound at the given location after the simulation has ticked.");
-	Sig.ExperimentalMessage = NSLOCTEXT("Niagara", "PlayAudioDIFunctionExperimental", "The return value of the audio function call currently needs to be wired to a particle parameter, because otherwise it will be removed by the compiler.");
 #endif
 	Sig.bMemberFunction = true;
 	Sig.bRequiresContext = false;
 	Sig.bSupportsGPU = false;
-	Sig.bExperimental = true;
+	Sig.bRequiresExecPin = true;
 	Sig.Inputs.Add(FNiagaraVariable(FNiagaraTypeDefinition(GetClass()), TEXT("Audio interface")));
 	Sig.Inputs.Add(FNiagaraVariable(FNiagaraTypeDefinition::GetBoolDef(), TEXT("Play Audio")));
 	Sig.Inputs.Add(FNiagaraVariable(FNiagaraTypeDefinition::GetVec3Def(), TEXT("PositionWS")));
