support for sampling mesh distance fields
rename to use "world" instead of "node"
examples that use mesh distance fields/velocities
updated fluids collision module to use mesh distance fields

diff --git a/Engine/Plugins/FX/Niagara/Shaders/Private/NiagaraDataInterfaceRigidMeshCollisionQuery.ush b/Engine/Plugins/FX/Niagara/Shaders/Private/NiagaraDataInterfaceRigidMeshCollisionQuery.ush
--- a/Engine/Plugins/FX/Niagara/Shaders/Private/NiagaraDataInterfaceRigidMeshCollisionQuery.ush
+++ b/Engine/Plugins/FX/Niagara/Shaders/Private/NiagaraDataInterfaceRigidMeshCollisionQuery.ush
@@ -7,7 +7,7 @@ NiagaraDataInterfaceRigidMeshCollisionQuery.ush
 #pragma once
 
 /* -----------------------------------------------------------------
- * Physics Asset constants and context
+ * static mesh constants and context
  * -----------------------------------------------------------------
  */
 
@@ -29,6 +29,7 @@ struct FDIRigidMeshCollisionContext
 	Buffer<float4>		InverseTransformBuffer;
 	Buffer<float4>		ElementExtentBuffer;
 	Buffer<uint>		PhysicsTypeBuffer;
+	Buffer<uint>		DFIndexBuffer;
 }; 
 
 #define DIRIGIDMESHCOLLISIONQUERY_DECLARE_CONSTANTS(NAME)\
@@ -37,6 +38,7 @@ Buffer<float4>		WorldTransformBuffer_##NAME;\
 Buffer<float4>		InverseTransformBuffer_##NAME;\
 Buffer<float4>		ElementExtentBuffer_##NAME;\
 Buffer<uint>		PhysicsTypeBuffer_##NAME;\
+Buffer<uint>		DFIndexBuffer_##NAME;\
 
 
 #define DIRIGIDMESHCOLLISIONQUERY_MAKE_CONTEXT(NAME)\
@@ -46,6 +48,7 @@ DIContext.WorldTransformBuffer = WorldTransformBuffer_##NAME;\
 DIContext.InverseTransformBuffer = InverseTransformBuffer_##NAME;\
 DIContext.ElementExtentBuffer = ElementExtentBuffer_##NAME;\
 DIContext.PhysicsTypeBuffer = PhysicsTypeBuffer_##NAME;\
+DIContext.DFIndexBuffer = DFIndexBuffer_##NAME;\
 
 
 /* -----------------------------------------------------------------
@@ -116,7 +119,7 @@ float3x4 DIRigidMeshCollision_GetPreviousInverse(in FDIRigidMeshCollisionContext
  * -----------------------------------------------------------------
  */
 
-// Given a world space position (NodePosition) compute the sphere closest point (position,normal,velocity)
+// Given a world space position (WorldPosition) compute the sphere closest point (position,normal,velocity)
 float DIRigidMeshCollision_GetSphereProjection(in float3 LocalPosition, in float3 SphereCenter, in float SphereRadius, in int SphereIndex,
 			inout float3 OutClosestPosition, inout float3 OutClosestNormal, inout int OutElementIndex, inout float OutMinDistance)
 {
@@ -135,7 +138,7 @@ float DIRigidMeshCollision_GetSphereProjection(in float3 LocalPosition, in float
 	return SphereDistance;
 }
 
-// Given a world space position (NodePosition) compute the sphere closest point (position,normal,velocity)
+// Given a world space position (WorldPosition) compute the sphere closest point (position,normal,velocity)
 float DIRigidMeshCollision_GetCapsuleProjection(in float3 LocalPosition, in float2 CapsuleSize, in int CapsuleIndex,
 			inout float3 OutClosestPosition, inout float3 OutClosestNormal, inout int OutElementIndex, inout float OutMinDistance)
 {
@@ -146,7 +149,7 @@ float DIRigidMeshCollision_GetCapsuleProjection(in float3 LocalPosition, in floa
 			OutClosestPosition,OutClosestNormal,OutElementIndex,OutMinDistance);
 }
 
-// Given a world space position (NodePosition) compute the sphere closest point (position,normal,velocity)
+// Given a world space position (WorldPosition) compute the sphere closest point (position,normal,velocity)
 float DIRigidMeshCollision_GetBoxProjection(in float3 LocalPosition, in float3 BoxExtent, in int BoxIndex, 
 			inout float3 OutClosestPosition, inout float3 OutClosestNormal, inout int OutElementIndex, inout float OutMinDistance)
 {
@@ -195,20 +198,20 @@ float DIRigidMeshCollision_GetBoxProjection(in float3 LocalPosition, in float3 B
 
 
 /* -----------------------------------------------------------------
- * Get the closest element to the node position
+ * Get the closest element to the world position
  * -----------------------------------------------------------------
  */
 
-float3 DIRigidMeshCollision_GetLocalPosition(in FDIRigidMeshCollisionContext DIContext, in float3 NodePosition, in int ElementIndex, in float TimeFraction )
+float3 DIRigidMeshCollision_GetLocalPosition(in FDIRigidMeshCollisionContext DIContext, in float3 WorldPosition, in int ElementIndex, in float TimeFraction )
 {
-	const float3 CurrentLocal = mul(DIRigidMeshCollision_GetInverseTransform(DIContext,ElementIndex), float4(NodePosition,1.0)).xyz;
-	const float3 PreviousLocal = mul(DIRigidMeshCollision_GetPreviousInverse(DIContext,ElementIndex), float4(NodePosition,1.0)).xyz;
+	const float3 CurrentLocal = mul(DIRigidMeshCollision_GetInverseTransform(DIContext,ElementIndex), float4(WorldPosition,1.0)).xyz;
+	const float3 PreviousLocal = mul(DIRigidMeshCollision_GetPreviousInverse(DIContext,ElementIndex), float4(WorldPosition,1.0)).xyz;
 
 	return PreviousLocal + TimeFraction * (CurrentLocal-PreviousLocal);
 }
 
-// Given a world space position (NodePosition) compute the sphere closest point (position,normal,velocity)
-int DIRigidMeshCollision_GetClosestElement(in FDIRigidMeshCollisionContext DIContext, in float3 NodePosition, out float3 OutClosestPosition, 
+// Given a world space position (WorldPosition) compute the sphere closest point (position,normal,velocity)
+int DIRigidMeshCollision_GetClosestElement(in FDIRigidMeshCollisionContext DIContext, in float3 WorldPosition, out float3 OutClosestPosition, 
 										  out float3 OutClosestNormal, out float OutMinDistance, in float TimeFraction)
 {	
 	float MinDistance = MAX_DISTANCE;
@@ -220,7 +223,7 @@ int DIRigidMeshCollision_GetClosestElement(in FDIRigidMeshCollisionContext DICon
 	const int SpheresEnd = DIContext.ElementOffsets[SPHERE_INDEX+1];
 	for (int SphereIndex = SpheresBegin; SphereIndex < SpheresEnd; ++SphereIndex)
 	{
-		const float3 LocalPosition = DIRigidMeshCollision_GetLocalPosition(DIContext,NodePosition,SphereIndex,TimeFraction);
+		const float3 LocalPosition = DIRigidMeshCollision_GetLocalPosition(DIContext,WorldPosition,SphereIndex,TimeFraction);
 
 		DIRigidMeshCollision_GetSphereProjection(LocalPosition, float3(0,0,0), DIContext.ElementExtentBuffer[SphereIndex].x, SphereIndex,
 				CollisionPosition, CollisionNormal, ElementIndex, MinDistance);
@@ -230,7 +233,7 @@ int DIRigidMeshCollision_GetClosestElement(in FDIRigidMeshCollisionContext DICon
 	const int CapsulesEnd = DIContext.ElementOffsets[CAPSULE_INDEX+1];
 	for (int CapsuleIndex = CapsulesBegin; CapsuleIndex < CapsulesEnd; ++CapsuleIndex)
 	{
-		const float3 LocalPosition = DIRigidMeshCollision_GetLocalPosition(DIContext,NodePosition,CapsuleIndex,TimeFraction);
+		const float3 LocalPosition = DIRigidMeshCollision_GetLocalPosition(DIContext,WorldPosition,CapsuleIndex,TimeFraction);
 
 		DIRigidMeshCollision_GetCapsuleProjection(LocalPosition, DIContext.ElementExtentBuffer[CapsuleIndex].xy, CapsuleIndex,
 				CollisionPosition, CollisionNormal, ElementIndex, MinDistance);
@@ -240,7 +243,7 @@ int DIRigidMeshCollision_GetClosestElement(in FDIRigidMeshCollisionContext DICon
 	const int BoxesEnd = DIContext.ElementOffsets[BOX_INDEX+1];
 	for (int BoxIndex = BoxesBegin; BoxIndex < BoxesEnd; ++BoxIndex)
 	{
-		const float3 LocalPosition = DIRigidMeshCollision_GetLocalPosition(DIContext,NodePosition,BoxIndex,TimeFraction);
+		const float3 LocalPosition = DIRigidMeshCollision_GetLocalPosition(DIContext,WorldPosition,BoxIndex,TimeFraction);
 
 		DIRigidMeshCollision_GetBoxProjection(LocalPosition, DIContext.ElementExtentBuffer[BoxIndex].xyz, BoxIndex,
 				CollisionPosition, CollisionNormal, ElementIndex, MinDistance);
@@ -252,15 +255,15 @@ int DIRigidMeshCollision_GetClosestElement(in FDIRigidMeshCollisionContext DICon
 }
 
 /* -----------------------------------------------------------------
- * Get the closest point to the physics asset 
+ * Get the closest point to the static mesh 
  * -----------------------------------------------------------------
  */
 
-// Given a world space position (NodePosition) and an element index compute the physics asset closest point (position,normal,velocity)
-void DIRigidMeshCollision_GetElementPoint(in FDIRigidMeshCollisionContext DIContext, in float3 NodePosition, in float DeltaTime, in float TimeFraction, in int ElementIndex, out float3 OutClosestPosition, 
+// Given a world space position (WorldPosition) and an element index compute the static mesh closest point (position,normal,velocity)
+void DIRigidMeshCollision_GetElementPoint(in FDIRigidMeshCollisionContext DIContext, in float3 WorldPosition, in float DeltaTime, in float TimeFraction, in int ElementIndex, out float3 OutClosestPosition, 
 										  out float3 OutClosestNormal, out float3 OutClosestVelocity)
 {	
-	const float3 LocalPosition = DIRigidMeshCollision_GetLocalPosition(DIContext,NodePosition,ElementIndex,TimeFraction);
+	const float3 LocalPosition = DIRigidMeshCollision_GetLocalPosition(DIContext,WorldPosition,ElementIndex,TimeFraction);
 
 	float MinDistance = MAX_DISTANCE;
 	int OutputIndex = -1;
@@ -302,10 +305,10 @@ void DIRigidMeshCollision_GetElementPoint(in FDIRigidMeshCollisionContext DICont
 	}
 }
 
-// Given a world space position (NodePosition) and an element index compute the physics asset closest distance
-void DIRigidMeshCollision_GetElementDistance(in FDIRigidMeshCollisionContext DIContext, in float3 NodePosition, in float TimeFraction, in int ElementIndex, out float OutClosestDistance)
+// Given a world space position (WorldPosition) and an element index compute the static mesh closest distance
+void DIRigidMeshCollision_GetElementDistance(in FDIRigidMeshCollisionContext DIContext, in float3 WorldPosition, in float TimeFraction, in int ElementIndex, out float OutClosestDistance)
 {	
-	const float3 LocalPosition = DIRigidMeshCollision_GetLocalPosition(DIContext,NodePosition,ElementIndex,TimeFraction);
+	const float3 LocalPosition = DIRigidMeshCollision_GetLocalPosition(DIContext,WorldPosition,ElementIndex,TimeFraction);
 
 	float MinDistance = MAX_DISTANCE;
 	int OutputIndex = -1;
@@ -336,24 +339,24 @@ void DIRigidMeshCollision_GetElementDistance(in FDIRigidMeshCollisionContext DIC
 	}
 }
 
- // Given a world space position (NodePosition) compute the physics asset closest element
-void DIRigidMeshCollision_GetClosestElement(in FDIRigidMeshCollisionContext DIContext, in float3 NodePosition, in float TimeFraction, out int OutElementIndex)
+ // Given a world space position (WorldPosition) compute the static mesh closest element
+void DIRigidMeshCollision_GetClosestElement(in FDIRigidMeshCollisionContext DIContext, in float3 WorldPosition, in float TimeFraction, out int OutElementIndex)
 {	
 	float3 CollisionPosition = float3(0,0,0);
 	float3 CollisionNormal = float3(0,0,0);
 	float MinDistance = 0.0;
 
-	OutElementIndex = DIRigidMeshCollision_GetClosestElement(DIContext,NodePosition,CollisionPosition,CollisionNormal,MinDistance,TimeFraction);
+	OutElementIndex = DIRigidMeshCollision_GetClosestElement(DIContext,WorldPosition,CollisionPosition,CollisionNormal,MinDistance,TimeFraction);
 }
 
-// Given a world space position (NodePosition) compute the physics asset closest point (position,normal,velocity)
-void DIRigidMeshCollision_GetClosestPoint(in FDIRigidMeshCollisionContext DIContext, in float3 NodePosition, in float DeltaTime, in float TimeFraction, out float OutClosestDistance, out float3 OutClosestPosition, out float3 OutClosestNormal, out float3 OutClosestVelocity)
+// Given a world space position (WorldPosition) compute the static mesh closest point (position,normal,velocity)
+void DIRigidMeshCollision_GetClosestPoint(in FDIRigidMeshCollisionContext DIContext, in float3 WorldPosition, in float DeltaTime, in float TimeFraction, out float OutClosestDistance, out float3 OutClosestPosition, out float3 OutClosestNormal, out float3 OutClosestVelocity)
 {	
 	float3 CollisionPosition = float3(0,0,0);
 	float3 CollisionNormal = float3(0,0,0);
 	OutClosestDistance = 0.0;
 
-	const int ElementIndex = DIRigidMeshCollision_GetClosestElement(DIContext,NodePosition,CollisionPosition,CollisionNormal,OutClosestDistance,TimeFraction);
+	const int ElementIndex = DIRigidMeshCollision_GetClosestElement(DIContext,WorldPosition,CollisionPosition,CollisionNormal,OutClosestDistance,TimeFraction);
 
 	OutClosestNormal = float3(0,0,0);
 	OutClosestPosition = float3(0,0,0);
@@ -374,14 +377,14 @@ void DIRigidMeshCollision_GetClosestPoint(in FDIRigidMeshCollisionContext DICont
 	}
 }
 
-// Given a world space position (NodePosition) compute the physics asset closest distance
-void DIRigidMeshCollision_GetClosestDistance(in FDIRigidMeshCollisionContext DIContext, in float3 NodePosition, in float TimeFraction, out float OutClosestDistance)
+// Given a world space position (WorldPosition) compute the static mesh closest distance
+void DIRigidMeshCollision_GetClosestDistance(in FDIRigidMeshCollisionContext DIContext, in float3 WorldPosition, in float TimeFraction, out float OutClosestDistance)
 {	
 	float3 CollisionPosition = float3(0,0,0);
 	float3 CollisionNormal = float3(0,0,0);
 	float MinDistance = 0.0;
 
-	const int ElementIndex = DIRigidMeshCollision_GetClosestElement(DIContext,NodePosition,CollisionPosition,CollisionNormal,MinDistance,TimeFraction);
+	const int ElementIndex = DIRigidMeshCollision_GetClosestElement(DIContext,WorldPosition,CollisionPosition,CollisionNormal,MinDistance,TimeFraction);
 
 	OutClosestDistance = 0.0;
 
@@ -389,4 +392,51 @@ void DIRigidMeshCollision_GetClosestDistance(in FDIRigidMeshCollisionContext DIC
 	{
 		OutClosestDistance = (DIContext.PhysicsTypeBuffer[ElementIndex] == 1) ? MinDistance : 0;
 	}
+}
+
+// Given a world space position (WorldPosition) compute the static mesh closest point (position,normal,velocity)
+void DIRigidMeshCollision_GetClosestPointMeshDistanceField(in FDIRigidMeshCollisionContext DIContext, in float3 WorldPosition, in float DeltaTime, in float TimeFraction, out float OutClosestDistance, out float3 OutClosestPosition, out float3 OutClosestNormal, out float3 OutClosestVelocity)
+{	
+	float3 CollisionPosition = float3(0,0,0);
+	float3 CollisionNormal = float3(0,0,0);
+	OutClosestDistance = MAX_DISTANCE;	
+
+	OutClosestNormal = float3(0,0,0);
+	OutClosestPosition = float3(0,0,0);
+	OutClosestVelocity = float3(0,0,0);
+
+	const int ElementIndex = DIRigidMeshCollision_GetClosestElement(DIContext,WorldPosition,CollisionPosition,CollisionNormal,OutClosestDistance,TimeFraction);
+
+	if (ElementIndex != -1)
+	{
+		int DFIndex = DIContext.DFIndexBuffer[ElementIndex];
+
+		if (DFIndex >= 0 && DFIndex < NumSceneObjects)
+		{
+			// #todo(dmp): this does a dynamic branch based on intersecting the bbox.  Maybe we can factor that out due to the broadphase here?
+			OutClosestDistance = DistanceToNearestSurfaceForObject(DFIndex, WorldPosition, 1e-3);
+		}
+
+		// #todo(dmp): figure out calculating SDF gradient so we can get more accurate collision velocities and collision normal
+		/*		
+		if (OutClosestDistance < 1e-5)
+		{
+			FDFObjectData DFObjectData = LoadDFObjectData(ElementIndex);
+			uint NumMips = LoadDFAssetData(DFObjectData.AssetIndex, 0).NumMips;
+			FDFAssetData DFAssetData = LoadDFAssetData(DFObjectData.AssetIndex, NumMips - 1);
+			float3 VolumeGradient = CalculateMeshSDFGradient(WorldPosition, DFAssetData);
+			float VolumeGradientLength = length(VolumeGradient);
+			float3 VolumeNormal = VolumeGradientLength > 0.0f ? VolumeGradient / VolumeGradientLength : 0;
+			// Transform by transposed inverse to handle non-uniform scaling
+			float3 WorldGradient = mul(VolumeNormal, transpose((float3x3)DFObjectData.WorldToVolume));
+			float WorldGradientLength = length(WorldGradient);
+			OutClosestNormal = WorldGradientLength > 0.0f ? WorldGradient / WorldGradientLength : 0;
+		}
+		*/
+
+		const float3 PreviousPosition = mul(DIRigidMeshCollision_GetPreviousTransform(DIContext,ElementIndex), float4(CollisionPosition,1.0)).xyz;
+		const float3 CurrentPosition = mul(DIRigidMeshCollision_GetCurrentTransform(DIContext,ElementIndex), float4(CollisionPosition,1.0)).xyz;
+		OutClosestVelocity = ( CurrentPosition - PreviousPosition ) / DeltaTime;
+		OutClosestPosition = PreviousPosition + TimeFraction * (CurrentPosition-PreviousPosition);
+	}
 }
\ No newline at end of file
diff --git a/Engine/Plugins/FX/Niagara/Source/Niagara/Private/NiagaraDataInterfaceRigidMeshCollisionQuery.cpp b/Engine/Plugins/FX/Niagara/Source/Niagara/Private/NiagaraDataInterfaceRigidMeshCollisionQuery.cpp
--- a/Engine/Plugins/FX/Niagara/Source/Niagara/Private/NiagaraDataInterfaceRigidMeshCollisionQuery.cpp
+++ b/Engine/Plugins/FX/Niagara/Source/Niagara/Private/NiagaraDataInterfaceRigidMeshCollisionQuery.cpp
@@ -14,6 +14,11 @@
 #include "ShaderParameterUtils.h"
 #include "EngineUtils.h"
 #include "Engine/StaticMeshActor.h"
+#include "Renderer/Private/ScenePrivate.h"
+#include "DistanceFieldAtlas.h"
+#include "Renderer/Private/DistanceFieldLightingShared.h"
+#include "NiagaraGpuComputeDispatchInterface.h"
+#include "NiagaraGpuComputeDispatch.h"
 
 #define LOCTEXT_NAMESPACE "NiagaraDataInterfaceRigidMeshCollisionQuery"
 DEFINE_LOG_CATEGORY_STATIC(LogRigidMeshCollision, Log, All);
@@ -31,6 +36,8 @@ static const FName GetElementPointName(TEXT("GetElementPoint"));
 static const FName GetElementDistanceName(TEXT("GetElementDistance"));
 static const FName GetClosestPointName(TEXT("GetClosestPoint"));
 static const FName GetClosestDistanceName(TEXT("GetClosestDistance"));
+static const FName GetClosestPointMeshDistanceFieldName(TEXT("GetClosestPointMeshDistanceField"));
+
 
 //------------------------------------------------------------------------------------------------------------
 
@@ -40,6 +47,13 @@ const FString UNiagaraDataInterfaceRigidMeshCollisionQuery::WorldTransformBuffer
 const FString UNiagaraDataInterfaceRigidMeshCollisionQuery::InverseTransformBufferName(TEXT("InverseTransformBuffer_"));
 const FString UNiagaraDataInterfaceRigidMeshCollisionQuery::ElementExtentBufferName(TEXT("ElementExtentBuffer_"));
 const FString UNiagaraDataInterfaceRigidMeshCollisionQuery::PhysicsTypeBufferName(TEXT("PhysicsTypeBuffer_"));
+const FString UNiagaraDataInterfaceRigidMeshCollisionQuery::DFIndexBufferName(TEXT("DFIndexBuffer_"));
+
+bool IsMeshDistanceFieldEnabled()
+{
+	static const auto* CVarGenerateMeshDistanceFields = IConsoleManager::Get().FindTConsoleVariableDataInt(TEXT("r.GenerateMeshDistanceFields"));
+	return CVarGenerateMeshDistanceFields != nullptr && CVarGenerateMeshDistanceFields->GetValueOnAnyThread() > 0;
+}
 
 //------------------------------------------------------------------------------------------------------------
 
@@ -53,6 +67,7 @@ struct FNDIRigidMeshCollisionParametersName
 		InverseTransformBufferName = UNiagaraDataInterfaceRigidMeshCollisionQuery::InverseTransformBufferName + Suffix;
 		ElementExtentBufferName = UNiagaraDataInterfaceRigidMeshCollisionQuery::ElementExtentBufferName + Suffix;
 		PhysicsTypeBufferName = UNiagaraDataInterfaceRigidMeshCollisionQuery::PhysicsTypeBufferName + Suffix;		
+		DFIndexBufferName = UNiagaraDataInterfaceRigidMeshCollisionQuery::DFIndexBufferName + Suffix;
 	}
 
 	FString ElementOffsetsName;
@@ -61,6 +76,7 @@ struct FNDIRigidMeshCollisionParametersName
 	FString InverseTransformBufferName;
 	FString ElementExtentBufferName;
 	FString PhysicsTypeBufferName;
+	FString DFIndexBufferName;
 };
 
 //------------------------------------------------------------------------------------------------------------
@@ -89,11 +105,12 @@ void UpdateInternalBuffer(const TStaticArray<BufferType,ElementCount*BufferCount
 }
 
 void FillCurrentTransforms(const FTransform& ElementTransform, uint32& ElementCount,
-	TStaticArray<FVector4,RIGID_MESH_COLLISION_QUERY_MAX_TRANSFORMS>& OutCurrentTransform, TStaticArray<FVector4, RIGID_MESH_COLLISION_QUERY_MAX_TRANSFORMS>& OutCurrentInverse)
+	TStaticArray<FVector4f,RIGID_MESH_COLLISION_QUERY_MAX_TRANSFORMS>& OutCurrentTransform, TStaticArray<FVector4f, RIGID_MESH_COLLISION_QUERY_MAX_TRANSFORMS>& OutCurrentInverse)
 {
+	// LWC_TODO: precision loss
 	const uint32 ElementOffset = 3 * ElementCount;
-	const FMatrix ElementMatrix = ElementTransform.ToMatrixWithScale();
-	const FMatrix ElementInverse = ElementMatrix.Inverse();
+	const FMatrix44f ElementMatrix = ElementTransform.ToMatrixWithScale();
+	const FMatrix44f ElementInverse = ElementMatrix.Inverse();
 
 	ElementMatrix.To3x4MatrixTranspose(&OutCurrentTransform[ElementOffset].X);
 	ElementInverse.To3x4MatrixTranspose(&OutCurrentInverse[ElementOffset].X);
@@ -114,14 +131,6 @@ void GetNumPrimitives(TArray<AStaticMeshActor*> StaticMeshActors, uint32& NumBox
 		if (BodySetup != nullptr)
 		{
 			// #todo(dmp): save static mesh component reference
-
-			for (const FKConvexElem& ConvexElem : BodySetup->AggGeom.ConvexElems)
-			{
-				if (CollisionEnabledHasPhysics(ConvexElem.GetCollisionEnabled()))
-				{
-					UE_LOG(LogRigidMeshCollision, Warning, TEXT("Convex collision objects encountered and will be skipped"));
-				}
-			}
 			for (const FKBoxElem& BoxElem : BodySetup->AggGeom.BoxElems)
 			{
 				if (CollisionEnabledHasPhysics(BoxElem.GetCollisionEnabled()))
@@ -195,21 +204,29 @@ void CreateInternalArrays(TArray<AStaticMeshActor*> StaticMeshActors,
 			{				
 				AStaticMeshActor* StaticMeshActor = StaticMeshActors[ActorIndex];
 				UStaticMeshComponent* StaticMeshComponent = StaticMeshActor != nullptr ? StaticMeshActor->GetStaticMeshComponent() : nullptr;
-				UBodySetup* BodySetup = StaticMeshComponent != nullptr ? StaticMeshComponent->GetBodySetup() : nullptr;
+				UBodySetup* BodySetup = StaticMeshComponent != nullptr ? StaticMeshComponent->GetBodySetup() : nullptr;				
 
 				if (BodySetup != nullptr)
 				{										
 					const FTransform MeshTransform = StaticMeshActor->GetTransform();				
+
+					for (const FKConvexElem& ConvexElem : BodySetup->AggGeom.ConvexElems)
+					{
+						if (CollisionEnabledHasPhysics(ConvexElem.GetCollisionEnabled()))
+						{
+							UE_LOG(LogRigidMeshCollision, Warning, TEXT("Convex collision objects encountered and will be skipped"));
+						}
+					}
 					for (const FKBoxElem& BoxElem : BodySetup->AggGeom.BoxElems)
 					{
 						if (CollisionEnabledHasPhysics(BoxElem.GetCollisionEnabled()))
 						{							
-							OutAssetArrays->PhysicsType[BoxCount] = (BoxElem.GetCollisionEnabled() == ECollisionEnabled::QueryAndPhysics);
+							OutAssetArrays->PhysicsType[BoxCount] = (BoxElem.GetCollisionEnabled() == ECollisionEnabled::QueryAndPhysics);							
+							OutAssetArrays->SourceSceneProxy[BoxCount] = StaticMeshComponent->SceneProxy;
 
 							const FTransform ElementTransform = FTransform(BoxElem.Rotation, BoxElem.Center) * MeshTransform;
 							OutAssetArrays->ElementExtent[BoxCount] = FVector4(BoxElem.X, BoxElem.Y, BoxElem.Z, 0);
-							FillCurrentTransforms(ElementTransform, BoxCount, OutAssetArrays->CurrentTransform, OutAssetArrays->CurrentInverse);
-							//--BoxCount;
+							FillCurrentTransforms(ElementTransform, BoxCount, OutAssetArrays->CurrentTransform, OutAssetArrays->CurrentInverse);							
 						}
 					}
 
@@ -217,12 +234,12 @@ void CreateInternalArrays(TArray<AStaticMeshActor*> StaticMeshActors,
 					{
 						if (CollisionEnabledHasPhysics(SphereElem.GetCollisionEnabled()))
 						{														
-							OutAssetArrays->PhysicsType[SphereCount] = (SphereElem.GetCollisionEnabled() == ECollisionEnabled::QueryAndPhysics);
+							OutAssetArrays->PhysicsType[SphereCount] = (SphereElem.GetCollisionEnabled() == ECollisionEnabled::QueryAndPhysics);							
+							OutAssetArrays->SourceSceneProxy[BoxCount] = StaticMeshComponent->SceneProxy;
 
 							const FTransform ElementTransform = FTransform(SphereElem.Center) * MeshTransform;
 							OutAssetArrays->ElementExtent[SphereCount] = FVector4(SphereElem.Radius, 0, 0, 0);
-							FillCurrentTransforms(ElementTransform, SphereCount, OutAssetArrays->CurrentTransform, OutAssetArrays->CurrentInverse);
-							//--SphereCount;
+							FillCurrentTransforms(ElementTransform, SphereCount, OutAssetArrays->CurrentTransform, OutAssetArrays->CurrentInverse);							
 						}
 					}
 
@@ -230,12 +247,12 @@ void CreateInternalArrays(TArray<AStaticMeshActor*> StaticMeshActors,
 					{
 						if (CollisionEnabledHasPhysics(CapsuleElem.GetCollisionEnabled()))
 						{														
-							OutAssetArrays->PhysicsType[CapsuleCount] = (CapsuleElem.GetCollisionEnabled() == ECollisionEnabled::QueryAndPhysics);
+							OutAssetArrays->PhysicsType[CapsuleCount] = (CapsuleElem.GetCollisionEnabled() == ECollisionEnabled::QueryAndPhysics);							
+							OutAssetArrays->SourceSceneProxy[BoxCount] = StaticMeshComponent->SceneProxy;
 
 							const FTransform ElementTransform = FTransform(CapsuleElem.Rotation, CapsuleElem.Center) * MeshTransform;
 							OutAssetArrays->ElementExtent[CapsuleCount] = FVector4(CapsuleElem.Radius, CapsuleElem.Length, 0, 0);
-							FillCurrentTransforms(ElementTransform, CapsuleCount, OutAssetArrays->CurrentTransform, OutAssetArrays->CurrentInverse);
-							//--CapsuleCount;
+							FillCurrentTransforms(ElementTransform, CapsuleCount, OutAssetArrays->CurrentTransform, OutAssetArrays->CurrentInverse);							
 						}
 					}					
 				}
@@ -252,8 +269,7 @@ void CreateInternalArrays(TArray<AStaticMeshActor*> StaticMeshActors,
 	}
 }
 
-void UpdateInternalArrays(TArray<AStaticMeshActor*> StaticMeshActors,
-	FNDIRigidMeshCollisionArrays* OutAssetArrays)
+void UpdateInternalArrays(const TArray<AStaticMeshActor*> &StaticMeshActors, FNDIRigidMeshCollisionArrays* OutAssetArrays)
 {
 	if (OutAssetArrays != nullptr && OutAssetArrays->ElementOffsets.NumElements < RIGID_MESH_COLLISION_QUERY_MAX_PRIMITIVES)
 	{
@@ -283,7 +299,7 @@ void UpdateInternalArrays(TArray<AStaticMeshActor*> StaticMeshActors,
 			UStaticMeshComponent* StaticMeshComponent = StaticMeshActor != nullptr ? StaticMeshActor->GetStaticMeshComponent() : nullptr;
 			UBodySetup* BodySetup = StaticMeshComponent != nullptr ? StaticMeshComponent->GetBodySetup() : nullptr;
 			if (BodySetup != nullptr)
-			{				
+			{
 				const FTransform MeshTransform = StaticMeshActor->GetTransform();
 
 				for (const FKBoxElem& BoxElem : BodySetup->AggGeom.BoxElems)
@@ -314,6 +330,7 @@ void UpdateInternalArrays(TArray<AStaticMeshActor*> StaticMeshActors,
 				}
 			}
 		}
+
 		CompactInternalArrays(OutAssetArrays);
 	}
 }
@@ -322,11 +339,12 @@ void UpdateInternalArrays(TArray<AStaticMeshActor*> StaticMeshActors,
 
 void FNDIRigidMeshCollisionBuffer::InitRHI()
 {
-	CreateInternalBuffer<FVector4, EPixelFormat::PF_A32B32G32R32F, RIGID_MESH_COLLISION_QUERY_MAX_TRANSFORMS, 2>(WorldTransformBuffer);
-	CreateInternalBuffer<FVector4, EPixelFormat::PF_A32B32G32R32F, RIGID_MESH_COLLISION_QUERY_MAX_TRANSFORMS, 2>(InverseTransformBuffer);
+	CreateInternalBuffer<FVector4f, EPixelFormat::PF_A32B32G32R32F, RIGID_MESH_COLLISION_QUERY_MAX_TRANSFORMS, 2>(WorldTransformBuffer);
+	CreateInternalBuffer<FVector4f, EPixelFormat::PF_A32B32G32R32F, RIGID_MESH_COLLISION_QUERY_MAX_TRANSFORMS, 2>(InverseTransformBuffer);
 
-	CreateInternalBuffer<FVector4, EPixelFormat::PF_A32B32G32R32F, RIGID_MESH_COLLISION_QUERY_MAX_PRIMITIVES>(ElementExtentBuffer);
+	CreateInternalBuffer<FVector4f, EPixelFormat::PF_A32B32G32R32F, RIGID_MESH_COLLISION_QUERY_MAX_PRIMITIVES>(ElementExtentBuffer);
 	CreateInternalBuffer<uint32, EPixelFormat::PF_R32_UINT, RIGID_MESH_COLLISION_QUERY_MAX_PRIMITIVES>(PhysicsTypeBuffer);
+	CreateInternalBuffer<uint32, EPixelFormat::PF_R32_UINT, RIGID_MESH_COLLISION_QUERY_MAX_PRIMITIVES>(DFIndexBuffer);
 }
 
 void FNDIRigidMeshCollisionBuffer::ReleaseRHI()
@@ -335,6 +353,7 @@ void FNDIRigidMeshCollisionBuffer::ReleaseRHI()
 	InverseTransformBuffer.Release();
 	ElementExtentBuffer.Release();
 	PhysicsTypeBuffer.Release();
+	DFIndexBuffer.Release();
 }
 
 //------------------------------------------------------------------------------------------------------------
@@ -378,10 +397,11 @@ void FNDIRigidMeshCollisionData::Init(UNiagaraDataInterfaceRigidMeshCollisionQue
 	}
 
 	if (Interface != nullptr && SystemInstance != nullptr)
-	{	
+	{
 		if (0 < StaticMeshActors.Num() && StaticMeshActors[0] != nullptr)
 		{
-			CreateInternalArrays(StaticMeshActors, &AssetArrays);
+			// @note: we are not creating internal arrays here and letting it happen on the call to update
+			 //CreateInternalArrays(StaticMeshActors, &AssetArrays);
 		}
 
 		AssetBuffer = new FNDIRigidMeshCollisionBuffer();
@@ -427,6 +447,94 @@ ETickingGroup FNDIRigidMeshCollisionData::ComputeTickingGroup()
 
 //------------------------------------------------------------------------------------------------------------
 
+class FDistanceFieldParameters
+{
+	DECLARE_INLINE_TYPE_LAYOUT(FDistanceFieldParameters, NonVirtual);
+public:
+	void Bind(const FShaderParameterMap& ParameterMap)
+	{
+		SceneObjectBounds.Bind(ParameterMap, TEXT("SceneObjectBounds"));
+		SceneObjectData.Bind(ParameterMap, TEXT("SceneObjectData"));
+		NumSceneObjects.Bind(ParameterMap, TEXT("NumSceneObjects"));
+		SceneDistanceFieldAssetData.Bind(ParameterMap, TEXT("SceneDistanceFieldAssetData"));
+		DistanceFieldIndirectionTable.Bind(ParameterMap, TEXT("DistanceFieldIndirectionTable"));
+		DistanceFieldBrickTexture.Bind(ParameterMap, TEXT("DistanceFieldBrickTexture"));
+		DistanceFieldSampler.Bind(ParameterMap, TEXT("DistanceFieldSampler"));
+		DistanceFieldBrickSize.Bind(ParameterMap, TEXT("DistanceFieldBrickSize"));
+		DistanceFieldUniqueDataBrickSize.Bind(ParameterMap, TEXT("DistanceFieldUniqueDataBrickSize"));
+		DistanceFieldBrickAtlasSizeInBricks.Bind(ParameterMap, TEXT("DistanceFieldBrickAtlasSizeInBricks"));
+		DistanceFieldBrickAtlasMask.Bind(ParameterMap, TEXT("DistanceFieldBrickAtlasMask"));
+		DistanceFieldBrickAtlasSizeLog2.Bind(ParameterMap, TEXT("DistanceFieldBrickAtlasSizeLog2"));
+		DistanceFieldBrickAtlasTexelSize.Bind(ParameterMap, TEXT("DistanceFieldBrickAtlasTexelSize"));
+	}
+
+	bool IsBound() const
+	{
+		return SceneDistanceFieldAssetData.IsBound();
+	}
+
+	friend FArchive& operator<<(FArchive& Ar, FDistanceFieldParameters& Parameters)
+	{
+		Ar << Parameters.SceneObjectBounds;
+		Ar << Parameters.SceneObjectData;
+		Ar << Parameters.NumSceneObjects;
+		Ar << Parameters.SceneDistanceFieldAssetData;
+		Ar << Parameters.DistanceFieldIndirectionTable;
+		Ar << Parameters.DistanceFieldBrickTexture;
+		Ar << Parameters.DistanceFieldSampler;
+		Ar << Parameters.DistanceFieldBrickSize;
+		Ar << Parameters.DistanceFieldUniqueDataBrickSize;
+		Ar << Parameters.DistanceFieldBrickAtlasSizeInBricks;
+		Ar << Parameters.DistanceFieldBrickAtlasMask;
+		Ar << Parameters.DistanceFieldBrickAtlasSizeLog2;
+		Ar << Parameters.DistanceFieldBrickAtlasTexelSize;
+
+		return Ar;
+	}
+
+	template<typename ShaderRHIParamRef>
+	FORCEINLINE_DEBUGGABLE void Set(FRHICommandList& RHICmdList, const ShaderRHIParamRef ShaderRHI, const FDistanceFieldSceneData* ParameterData) const
+	{
+		if (IsBound())
+		{
+			SetSRVParameter(RHICmdList, ShaderRHI, SceneObjectBounds, ParameterData->GetCurrentObjectBuffers()->Bounds.SRV);
+			SetSRVParameter(RHICmdList, ShaderRHI, SceneObjectData, ParameterData->GetCurrentObjectBuffers()->Data.SRV);
+			SetShaderValue(RHICmdList, ShaderRHI, NumSceneObjects, ParameterData->NumObjectsInBuffer);
+			SetSRVParameter(RHICmdList, ShaderRHI, SceneDistanceFieldAssetData, ParameterData->AssetDataBuffer.SRV);
+			SetSRVParameter(RHICmdList, ShaderRHI, DistanceFieldIndirectionTable, ParameterData->IndirectionTable.SRV);
+			SetTextureParameter(RHICmdList, ShaderRHI, DistanceFieldBrickTexture, ParameterData->DistanceFieldBrickVolumeTexture->GetRenderTargetItem().ShaderResourceTexture);
+			SetSamplerParameter(RHICmdList, ShaderRHI, DistanceFieldSampler, TStaticSamplerState<SF_Bilinear, AM_Clamp, AM_Clamp, AM_Clamp>::GetRHI());
+			SetShaderValue(RHICmdList, ShaderRHI, DistanceFieldBrickSize, FVector3f(DistanceField::BrickSize));
+			SetShaderValue(RHICmdList, ShaderRHI, DistanceFieldUniqueDataBrickSize, FVector3f(DistanceField::UniqueDataBrickSize));
+			SetShaderValue(RHICmdList, ShaderRHI, DistanceFieldBrickAtlasSizeInBricks, ParameterData->BrickTextureDimensionsInBricks);
+			SetShaderValue(RHICmdList, ShaderRHI, DistanceFieldBrickAtlasMask, ParameterData->BrickTextureDimensionsInBricks - FIntVector(1));
+			SetShaderValue(RHICmdList, ShaderRHI, DistanceFieldBrickAtlasSizeLog2, FIntVector(
+				FMath::FloorLog2(ParameterData->BrickTextureDimensionsInBricks.X),
+				FMath::FloorLog2(ParameterData->BrickTextureDimensionsInBricks.Y),
+				FMath::FloorLog2(ParameterData->BrickTextureDimensionsInBricks.Z)));
+			SetShaderValue(RHICmdList, ShaderRHI, DistanceFieldBrickAtlasTexelSize, FVector3f(1.0f) / FVector3f(ParameterData->BrickTextureDimensionsInBricks * DistanceField::BrickSize));
+		}
+	}
+
+private:
+
+	LAYOUT_FIELD(FShaderResourceParameter, SceneObjectBounds)
+		LAYOUT_FIELD(FShaderResourceParameter, SceneObjectData)
+		LAYOUT_FIELD(FShaderParameter, NumSceneObjects)
+		LAYOUT_FIELD(FShaderResourceParameter, SceneDistanceFieldAssetData)
+		LAYOUT_FIELD(FShaderResourceParameter, DistanceFieldIndirectionTable)
+		LAYOUT_FIELD(FShaderResourceParameter, DistanceFieldBrickTexture)
+		LAYOUT_FIELD(FShaderResourceParameter, DistanceFieldSampler)
+		LAYOUT_FIELD(FShaderParameter, DistanceFieldBrickSize)
+		LAYOUT_FIELD(FShaderParameter, DistanceFieldUniqueDataBrickSize)
+		LAYOUT_FIELD(FShaderParameter, DistanceFieldBrickAtlasSizeInBricks)
+		LAYOUT_FIELD(FShaderParameter, DistanceFieldBrickAtlasMask)
+		LAYOUT_FIELD(FShaderParameter, DistanceFieldBrickAtlasSizeLog2)
+		LAYOUT_FIELD(FShaderParameter, DistanceFieldBrickAtlasTexelSize)
+};
+
+//------------------------------------------------------------------------------------------------------------
+
 struct FNDIRigidMeshCollisionParametersCS : public FNiagaraDataInterfaceParametersCS
 {
 	DECLARE_TYPE_LAYOUT(FNDIRigidMeshCollisionParametersCS, NonVirtual);
@@ -440,7 +548,10 @@ public:
 		WorldTransformBuffer.Bind(ParameterMap, *ParamNames.WorldTransformBufferName);
 		InverseTransformBuffer.Bind(ParameterMap, *ParamNames.InverseTransformBufferName);
 		ElementExtentBuffer.Bind(ParameterMap, *ParamNames.ElementExtentBufferName);
-		PhysicsTypeBuffer.Bind(ParameterMap, *ParamNames.PhysicsTypeBufferName);		
+		PhysicsTypeBuffer.Bind(ParameterMap, *ParamNames.PhysicsTypeBufferName);	
+		DFIndexBuffer.Bind(ParameterMap, *ParamNames.DFIndexBufferName);
+
+		DistanceFieldParameters.Bind(ParameterMap);
 	}
 
 	void Set(FRHICommandList& RHICmdList, const FNiagaraDataInterfaceSetArgs& Context) const
@@ -462,7 +573,8 @@ public:
 				FRHITransitionInfo(AssetBuffer->WorldTransformBuffer.UAV, ERHIAccess::Unknown, ERHIAccess::SRVCompute),
 				FRHITransitionInfo(AssetBuffer->InverseTransformBuffer.UAV, ERHIAccess::Unknown, ERHIAccess::SRVCompute),
 				FRHITransitionInfo(AssetBuffer->ElementExtentBuffer.UAV, ERHIAccess::Unknown, ERHIAccess::SRVCompute),
-				FRHITransitionInfo(AssetBuffer->PhysicsTypeBuffer.UAV, ERHIAccess::Unknown, ERHIAccess::SRVCompute)
+				FRHITransitionInfo(AssetBuffer->PhysicsTypeBuffer.UAV, ERHIAccess::Unknown, ERHIAccess::SRVCompute),
+				FRHITransitionInfo(AssetBuffer->DFIndexBuffer.UAV, ERHIAccess::Unknown, ERHIAccess::SRVCompute)
 			};
 			RHICmdList.Transition(MakeArrayView(Transitions, UE_ARRAY_COUNT(Transitions)));
 
@@ -470,8 +582,25 @@ public:
 			SetSRVParameter(RHICmdList, ComputeShaderRHI, InverseTransformBuffer, AssetBuffer->InverseTransformBuffer.SRV);
 			SetSRVParameter(RHICmdList, ComputeShaderRHI, ElementExtentBuffer, AssetBuffer->ElementExtentBuffer.SRV);
 			SetSRVParameter(RHICmdList, ComputeShaderRHI, PhysicsTypeBuffer, AssetBuffer->PhysicsTypeBuffer.SRV);
+			SetSRVParameter(RHICmdList, ComputeShaderRHI, DFIndexBuffer, AssetBuffer->DFIndexBuffer.SRV);
+
+			SetShaderValue(RHICmdList, ComputeShaderRHI, ElementOffsets, ProxyData->AssetArrays.ElementOffsets);
 
-			SetShaderValue(RHICmdList, ComputeShaderRHI, ElementOffsets, ProxyData->AssetArrays.ElementOffsets);			
+			if (DistanceFieldParameters.IsBound())
+			{
+				const FDistanceFieldSceneData *DistanceFieldSceneData = static_cast<const FNiagaraGpuComputeDispatch*>(Context.ComputeDispatchInterface)->GetMeshDistanceFieldParameters();	//-BATCHERTODO:
+
+				if (DistanceFieldSceneData == nullptr)
+				{
+					UE_LOG(LogRigidMeshCollision, Error, TEXT("Distance fields are not available for use"));
+					// #todo(dmp): should we set something here in the case where distance field data is bound but we don't have it?
+					// there is no trivial constructor
+				}
+				else 
+				{
+					DistanceFieldParameters.Set(RHICmdList, ComputeShaderRHI, DistanceFieldSceneData);
+				}				
+			}
 		}
 		else
 		{
@@ -479,6 +608,7 @@ public:
 			SetSRVParameter(RHICmdList, ComputeShaderRHI, InverseTransformBuffer, FNiagaraRenderer::GetDummyFloatBuffer());
 			SetSRVParameter(RHICmdList, ComputeShaderRHI, ElementExtentBuffer, FNiagaraRenderer::GetDummyFloatBuffer());
 			SetSRVParameter(RHICmdList, ComputeShaderRHI, PhysicsTypeBuffer, FNiagaraRenderer::GetDummyIntBuffer());
+			SetSRVParameter(RHICmdList, ComputeShaderRHI, DFIndexBuffer, FNiagaraRenderer::GetDummyIntBuffer());
 
 			static const FElementOffset DummyOffsets(0, 0, 0, 0);
 			SetShaderValue(RHICmdList, ComputeShaderRHI, ElementOffsets, DummyOffsets);	
@@ -497,6 +627,9 @@ private:
 	LAYOUT_FIELD(FShaderResourceParameter, InverseTransformBuffer);
 	LAYOUT_FIELD(FShaderResourceParameter, ElementExtentBuffer);
 	LAYOUT_FIELD(FShaderResourceParameter, PhysicsTypeBuffer);
+	LAYOUT_FIELD(FShaderResourceParameter, DFIndexBuffer);
+
+	LAYOUT_FIELD(FDistanceFieldParameters, DistanceFieldParameters);	
 };
 
 IMPLEMENT_TYPE_LAYOUT(FNDIRigidMeshCollisionParametersCS);
@@ -508,7 +641,7 @@ IMPLEMENT_NIAGARA_DI_PARAMETER(UNiagaraDataInterfaceRigidMeshCollisionQuery, FND
 
 void FNDIRigidMeshCollisionProxy::ConsumePerInstanceDataFromGameThread(void* PerInstanceData, const FNiagaraSystemInstanceID& Instance)
 {
-	FNDIRigidMeshCollisionData* SourceData = static_cast<FNDIRigidMeshCollisionData*>(PerInstanceData);
+	FNDIRigidMeshCollisionData* SourceData = static_cast<FNDIRigidMeshCollisionData*>(PerInstanceData);	
 	FNDIRigidMeshCollisionData* TargetData = &(SystemInstancesToProxyData.FindOrAdd(Instance));
 
 	ensure(TargetData);
@@ -517,6 +650,17 @@ void FNDIRigidMeshCollisionProxy::ConsumePerInstanceDataFromGameThread(void* Per
 		TargetData->AssetBuffer = SourceData->AssetBuffer;		
 		TargetData->AssetArrays = SourceData->AssetArrays;
 		TargetData->TickingGroup = SourceData->TickingGroup;
+
+		// loop over proxies and compute df indices on the RT only on new data
+		// @todo(dmp): for now we do this every frame because it seems like sometimes DFindices are not set
+		for (uint32 i = 0; i < SourceData->AssetArrays.ElementOffsets.NumElements; ++i)
+		{
+			FPrimitiveSceneProxy* Proxy = SourceData->AssetArrays.SourceSceneProxy[i];
+								
+			const TArray<int32, TInlineAllocator<1>>& DFIndices = Proxy != nullptr && Proxy->GetPrimitiveSceneInfo() != nullptr ?
+				Proxy->GetPrimitiveSceneInfo()->DistanceFieldInstanceIndices : TArray<int32, TInlineAllocator<1>>();			
+			TargetData->AssetArrays.DFIndex[i] = DFIndices.Num() > 0 ? DFIndices[0] : -1;
+		}		
 	}
 	else
 	{
@@ -528,8 +672,8 @@ void FNDIRigidMeshCollisionProxy::InitializePerInstanceData(const FNiagaraSystem
 {
 	check(IsInRenderingThread());
 
-	FNDIRigidMeshCollisionData* TargetData = SystemInstancesToProxyData.Find(SystemInstance);
-	TargetData = &SystemInstancesToProxyData.Add(SystemInstance);
+	check(!SystemInstancesToProxyData.Contains(SystemInstance));
+	FNDIRigidMeshCollisionData* TargetData = &SystemInstancesToProxyData.Add(SystemInstance);
 }
 
 void FNDIRigidMeshCollisionProxy::DestroyPerInstanceData(const FNiagaraSystemInstanceID& SystemInstance)
@@ -540,6 +684,8 @@ void FNDIRigidMeshCollisionProxy::DestroyPerInstanceData(const FNiagaraSystemIns
 
 void FNDIRigidMeshCollisionProxy::PreStage(FRHICommandList& RHICmdList, const FNiagaraDataInterfaceStageArgs& Context)
 {
+	check(SystemInstancesToProxyData.Contains(Context.SystemInstanceID));
+
 	FNDIRigidMeshCollisionData* ProxyData =
 		SystemInstancesToProxyData.Find(Context.SystemInstanceID);
 
@@ -551,14 +697,16 @@ void FNDIRigidMeshCollisionProxy::PreStage(FRHICommandList& RHICmdList, const FN
 				FRHITransitionInfo(ProxyData->AssetBuffer->WorldTransformBuffer.UAV, ERHIAccess::Unknown, ERHIAccess::UAVCompute),
 				FRHITransitionInfo(ProxyData->AssetBuffer->InverseTransformBuffer.UAV, ERHIAccess::Unknown, ERHIAccess::UAVCompute),
 				FRHITransitionInfo(ProxyData->AssetBuffer->PhysicsTypeBuffer.UAV, ERHIAccess::Unknown, ERHIAccess::UAVCompute),
-				FRHITransitionInfo(ProxyData->AssetBuffer->ElementExtentBuffer.UAV, ERHIAccess::Unknown, ERHIAccess::UAVCompute)
+				FRHITransitionInfo(ProxyData->AssetBuffer->ElementExtentBuffer.UAV, ERHIAccess::Unknown, ERHIAccess::UAVCompute),
+				FRHITransitionInfo(ProxyData->AssetBuffer->DFIndexBuffer.UAV, ERHIAccess::Unknown, ERHIAccess::UAVCompute)
 			};
 			RHICmdList.Transition(MakeArrayView(Transitions, UE_ARRAY_COUNT(Transitions)));
 
-			UpdateInternalBuffer<FVector4, EPixelFormat::PF_A32B32G32R32F, RIGID_MESH_COLLISION_QUERY_MAX_TRANSFORMS, 2>(ProxyData->AssetArrays.WorldTransform, ProxyData->AssetBuffer->WorldTransformBuffer);
-			UpdateInternalBuffer<FVector4, EPixelFormat::PF_A32B32G32R32F, RIGID_MESH_COLLISION_QUERY_MAX_TRANSFORMS, 2>(ProxyData->AssetArrays.InverseTransform, ProxyData->AssetBuffer->InverseTransformBuffer);
-			UpdateInternalBuffer<FVector4, EPixelFormat::PF_A32B32G32R32F, RIGID_MESH_COLLISION_QUERY_MAX_PRIMITIVES>(ProxyData->AssetArrays.ElementExtent, ProxyData->AssetBuffer->ElementExtentBuffer);
+			UpdateInternalBuffer<FVector4f, EPixelFormat::PF_A32B32G32R32F, RIGID_MESH_COLLISION_QUERY_MAX_TRANSFORMS, 2>(ProxyData->AssetArrays.WorldTransform, ProxyData->AssetBuffer->WorldTransformBuffer);
+			UpdateInternalBuffer<FVector4f, EPixelFormat::PF_A32B32G32R32F, RIGID_MESH_COLLISION_QUERY_MAX_TRANSFORMS, 2>(ProxyData->AssetArrays.InverseTransform, ProxyData->AssetBuffer->InverseTransformBuffer);
+			UpdateInternalBuffer<FVector4f, EPixelFormat::PF_A32B32G32R32F, RIGID_MESH_COLLISION_QUERY_MAX_PRIMITIVES>(ProxyData->AssetArrays.ElementExtent, ProxyData->AssetBuffer->ElementExtentBuffer);
 			UpdateInternalBuffer<uint32, EPixelFormat::PF_R32_UINT, RIGID_MESH_COLLISION_QUERY_MAX_PRIMITIVES>(ProxyData->AssetArrays.PhysicsType, ProxyData->AssetBuffer->PhysicsTypeBuffer);
+			UpdateInternalBuffer<uint32, EPixelFormat::PF_R32_UINT, RIGID_MESH_COLLISION_QUERY_MAX_PRIMITIVES>(ProxyData->AssetArrays.DFIndex, ProxyData->AssetBuffer->DFIndexBuffer);
 		}
 	}
 }
@@ -705,7 +853,7 @@ void UNiagaraDataInterfaceRigidMeshCollisionQuery::GetFunctions(TArray<FNiagaraF
 		Sig.bMemberFunction = true;
 		Sig.bRequiresContext = false;
 		Sig.Inputs.Add(FNiagaraVariable(FNiagaraTypeDefinition(GetClass()), TEXT("Collision DI")));
-		Sig.Inputs.Add(FNiagaraVariable(FNiagaraTypeDefinition::GetVec3Def(), TEXT("Node Position")));
+		Sig.Inputs.Add(FNiagaraVariable(FNiagaraTypeDefinition::GetVec3Def(), TEXT("World Position")));
 		Sig.Inputs.Add(FNiagaraVariable(FNiagaraTypeDefinition::GetFloatDef(), TEXT("Delta Time")));
 		Sig.Inputs.Add(FNiagaraVariable(FNiagaraTypeDefinition::GetFloatDef(), TEXT("Time Fraction")));
 		Sig.Outputs.Add(FNiagaraVariable(FNiagaraTypeDefinition::GetFloatDef(), TEXT("Closest Distance")));
@@ -723,7 +871,7 @@ void UNiagaraDataInterfaceRigidMeshCollisionQuery::GetFunctions(TArray<FNiagaraF
 		Sig.bMemberFunction = true;
 		Sig.bRequiresContext = false;
 		Sig.Inputs.Add(FNiagaraVariable(FNiagaraTypeDefinition(GetClass()), TEXT("Collision DI")));
-		Sig.Inputs.Add(FNiagaraVariable(FNiagaraTypeDefinition::GetVec3Def(), TEXT("Node Position")));
+		Sig.Inputs.Add(FNiagaraVariable(FNiagaraTypeDefinition::GetVec3Def(), TEXT("World Position")));
 		Sig.Inputs.Add(FNiagaraVariable(FNiagaraTypeDefinition::GetFloatDef(), TEXT("Time Fraction")));
 		Sig.Outputs.Add(FNiagaraVariable(FNiagaraTypeDefinition::GetIntDef(), TEXT("Closest Element")));
 
@@ -737,7 +885,7 @@ void UNiagaraDataInterfaceRigidMeshCollisionQuery::GetFunctions(TArray<FNiagaraF
 		Sig.bMemberFunction = true;
 		Sig.bRequiresContext = false;
 		Sig.Inputs.Add(FNiagaraVariable(FNiagaraTypeDefinition(GetClass()), TEXT("Collision DI")));
-		Sig.Inputs.Add(FNiagaraVariable(FNiagaraTypeDefinition::GetVec3Def(), TEXT("Node Position")));
+		Sig.Inputs.Add(FNiagaraVariable(FNiagaraTypeDefinition::GetVec3Def(), TEXT("World Position")));
 		Sig.Inputs.Add(FNiagaraVariable(FNiagaraTypeDefinition::GetFloatDef(), TEXT("Delta Time")));
 		Sig.Inputs.Add(FNiagaraVariable(FNiagaraTypeDefinition::GetFloatDef(), TEXT("Time Fraction")));
 		Sig.Inputs.Add(FNiagaraVariable(FNiagaraTypeDefinition::GetIntDef(), TEXT("Element Index")));
@@ -755,7 +903,7 @@ void UNiagaraDataInterfaceRigidMeshCollisionQuery::GetFunctions(TArray<FNiagaraF
 		Sig.bMemberFunction = true;
 		Sig.bRequiresContext = false;
 		Sig.Inputs.Add(FNiagaraVariable(FNiagaraTypeDefinition(GetClass()), TEXT("Collision DI")));
-		Sig.Inputs.Add(FNiagaraVariable(FNiagaraTypeDefinition::GetVec3Def(), TEXT("Node Position")));
+		Sig.Inputs.Add(FNiagaraVariable(FNiagaraTypeDefinition::GetVec3Def(), TEXT("World Position")));
 		Sig.Inputs.Add(FNiagaraVariable(FNiagaraTypeDefinition::GetFloatDef(), TEXT("Time Fraction")));
 		Sig.Inputs.Add(FNiagaraVariable(FNiagaraTypeDefinition::GetIntDef(), TEXT("Element Index")));
 		Sig.Outputs.Add(FNiagaraVariable(FNiagaraTypeDefinition::GetFloatDef(), TEXT("Closest Distance")));
@@ -770,9 +918,28 @@ void UNiagaraDataInterfaceRigidMeshCollisionQuery::GetFunctions(TArray<FNiagaraF
 		Sig.bMemberFunction = true;
 		Sig.bRequiresContext = false;
 		Sig.Inputs.Add(FNiagaraVariable(FNiagaraTypeDefinition(GetClass()), TEXT("Collision DI")));
-		Sig.Inputs.Add(FNiagaraVariable(FNiagaraTypeDefinition::GetVec3Def(), TEXT("Node Position")));
+		Sig.Inputs.Add(FNiagaraVariable(FNiagaraTypeDefinition::GetVec3Def(), TEXT("World Position")));
+		Sig.Inputs.Add(FNiagaraVariable(FNiagaraTypeDefinition::GetFloatDef(), TEXT("Time Fraction")));
+		Sig.Outputs.Add(FNiagaraVariable(FNiagaraTypeDefinition::GetFloatDef(), TEXT("Closest Distance")));
+
+		OutFunctions.Add(Sig);
+	}
+
+	{
+		FNiagaraFunctionSignature Sig;
+		Sig.Name = GetClosestPointMeshDistanceFieldName;
+		Sig.bSupportsGPU = true;
+		Sig.bSupportsCPU = false;
+		Sig.bMemberFunction = true;
+		Sig.bRequiresContext = false;
+		Sig.Inputs.Add(FNiagaraVariable(FNiagaraTypeDefinition(GetClass()), TEXT("Collision DI")));
+		Sig.Inputs.Add(FNiagaraVariable(FNiagaraTypeDefinition::GetVec3Def(), TEXT("World Position")));
+		Sig.Inputs.Add(FNiagaraVariable(FNiagaraTypeDefinition::GetFloatDef(), TEXT("Delta Time")));
 		Sig.Inputs.Add(FNiagaraVariable(FNiagaraTypeDefinition::GetFloatDef(), TEXT("Time Fraction")));
 		Sig.Outputs.Add(FNiagaraVariable(FNiagaraTypeDefinition::GetFloatDef(), TEXT("Closest Distance")));
+		Sig.Outputs.Add(FNiagaraVariable(FNiagaraTypeDefinition::GetVec3Def(), TEXT("Closest Position")));
+		Sig.Outputs.Add(FNiagaraVariable(FNiagaraTypeDefinition::GetVec3Def(), TEXT("Closest Normal")));
+		Sig.Outputs.Add(FNiagaraVariable(FNiagaraTypeDefinition::GetVec3Def(), TEXT("Closest Velocity")));
 
 		OutFunctions.Add(Sig);
 	}
@@ -835,10 +1002,10 @@ bool UNiagaraDataInterfaceRigidMeshCollisionQuery::GetFunctionHLSL(const FNiagar
 	else if (FunctionInfo.DefinitionName == GetClosestPointName)
 	{
 		static const TCHAR* FormatSample = TEXT(R"(
-		void {InstanceFunctionName}(in float3 NodePosition, in float DeltaTime, in float TimeFraction, out float ClosestDistance, out float3 OutClosestPosition, 
+		void {InstanceFunctionName}(in float3 WorldPosition, in float DeltaTime, in float TimeFraction, out float ClosestDistance, out float3 OutClosestPosition, 
 							out float3 OutClosestNormal, out float3 OutClosestVelocity)
 		{
-			{RigidMeshCollisionContextName} DIRigidMeshCollision_GetClosestPoint(DIContext,NodePosition,DeltaTime,TimeFraction, ClosestDistance,
+			{RigidMeshCollisionContextName} DIRigidMeshCollision_GetClosestPoint(DIContext,WorldPosition,DeltaTime,TimeFraction, ClosestDistance,
 				OutClosestPosition,OutClosestNormal,OutClosestVelocity);
 		}
 		)");
@@ -848,9 +1015,9 @@ bool UNiagaraDataInterfaceRigidMeshCollisionQuery::GetFunctionHLSL(const FNiagar
 	else if (FunctionInfo.DefinitionName == GetClosestElementName)
 	{
 		static const TCHAR* FormatSample = TEXT(R"(
-		void {InstanceFunctionName}(in float3 NodePosition, in float TimeFraction, out int OutClosestElement)
+		void {InstanceFunctionName}(in float3 WorldPosition, in float TimeFraction, out int OutClosestElement)
 		{
-			{RigidMeshCollisionContextName} DIRigidMeshCollision_GetClosestElement(DIContext,NodePosition,TimeFraction,
+			{RigidMeshCollisionContextName} DIRigidMeshCollision_GetClosestElement(DIContext,WorldPosition,TimeFraction,
 				OutClosestElement);
 		}
 		)");
@@ -860,10 +1027,10 @@ bool UNiagaraDataInterfaceRigidMeshCollisionQuery::GetFunctionHLSL(const FNiagar
 	else if (FunctionInfo.DefinitionName == GetElementPointName)
 	{
 		static const TCHAR* FormatSample = TEXT(R"(
-		void {InstanceFunctionName}(in float3 NodePosition, in float DeltaTime, in float TimeFraction, in int ElementIndex, out float3 OutClosestPosition, 
+		void {InstanceFunctionName}(in float3 WorldPosition, in float DeltaTime, in float TimeFraction, in int ElementIndex, out float3 OutClosestPosition, 
 							out float3 OutClosestNormal, out float3 OutClosestVelocity)
 		{
-			{RigidMeshCollisionContextName} DIRigidMeshCollision_GetElementPoint(DIContext,NodePosition,DeltaTime,TimeFraction,ElementIndex,
+			{RigidMeshCollisionContextName} DIRigidMeshCollision_GetElementPoint(DIContext,WorldPosition,DeltaTime,TimeFraction,ElementIndex,
 				OutClosestPosition,OutClosestNormal,OutClosestVelocity);
 		}
 		)");
@@ -873,9 +1040,9 @@ bool UNiagaraDataInterfaceRigidMeshCollisionQuery::GetFunctionHLSL(const FNiagar
 	else if (FunctionInfo.DefinitionName == GetElementDistanceName)
 	{
 		static const TCHAR* FormatSample = TEXT(R"(
-		void {InstanceFunctionName}(in float3 NodePosition, in float TimeFraction, in int ElementIndex, out float OutClosestDistance)
+		void {InstanceFunctionName}(in float3 WorldPosition, in float TimeFraction, in int ElementIndex, out float OutClosestDistance)
 		{
-			{RigidMeshCollisionContextName} DIRigidMeshCollision_GetElementDistance(DIContext,NodePosition,TimeFraction,ElementIndex,
+			{RigidMeshCollisionContextName} DIRigidMeshCollision_GetElementDistance(DIContext,WorldPosition,TimeFraction,ElementIndex,
 				OutClosestDistance);
 		}
 		)");
@@ -885,9 +1052,22 @@ bool UNiagaraDataInterfaceRigidMeshCollisionQuery::GetFunctionHLSL(const FNiagar
 	else if (FunctionInfo.DefinitionName == GetClosestDistanceName)
 	{
 		static const TCHAR* FormatSample = TEXT(R"(
-		void {InstanceFunctionName}(in float3 NodePosition, in float TimeFraction, out float OutClosestDistance)
+		void {InstanceFunctionName}(in float3 WorldPosition, in float TimeFraction, out float OutClosestDistance)
+		{
+			{RigidMeshCollisionContextName} DIRigidMeshCollision_GetClosestDistance(DIContext,WorldPosition,TimeFraction,OutClosestDistance);
+		}
+		)");
+		OutHLSL += FString::Format(FormatSample, ArgsSample);
+		return true;
+	}
+	else if (FunctionInfo.DefinitionName == GetClosestPointMeshDistanceFieldName)
+	{
+		static const TCHAR* FormatSample = TEXT(R"(
+		void {InstanceFunctionName}(in float3 WorldPosition, in float DeltaTime, in float TimeFraction, out float ClosestDistance, out float3 OutClosestPosition, 
+							out float3 OutClosestNormal, out float3 OutClosestVelocity)
 		{
-			{RigidMeshCollisionContextName} DIRigidMeshCollision_GetClosestDistance(DIContext,NodePosition,TimeFraction,OutClosestDistance);
+			{RigidMeshCollisionContextName} DIRigidMeshCollision_GetClosestPointMeshDistanceField(DIContext,WorldPosition,DeltaTime,TimeFraction, ClosestDistance,
+				OutClosestPosition,OutClosestNormal,OutClosestVelocity);
 		}
 		)");
 		OutHLSL += FString::Format(FormatSample, ArgsSample);
@@ -900,7 +1080,9 @@ bool UNiagaraDataInterfaceRigidMeshCollisionQuery::GetFunctionHLSL(const FNiagar
 
 void UNiagaraDataInterfaceRigidMeshCollisionQuery::GetCommonHLSL(FString& OutHLSL)
 {
-	OutHLSL += TEXT("#include \"/Plugin/FX/Niagara/Private/NiagaraDataInterfaceRigidMeshCollisionQuery.ush\"\n");		
+	OutHLSL += TEXT("#include \"/Engine/Private/DistanceFieldLightingShared.ush\"\n");
+	OutHLSL += TEXT("#include \"/Engine/Private/MeshDistanceFieldCommon.ush\"\n");
+	OutHLSL += TEXT("#include \"/Plugin/FX/Niagara/Private/NiagaraDataInterfaceRigidMeshCollisionQuery.ush\"\n");			
 }
 
 void UNiagaraDataInterfaceRigidMeshCollisionQuery::GetParameterDefinitionHLSL(const FNiagaraDataInterfaceGPUParamInfo& ParamInfo, FString& OutHLSL)
@@ -909,13 +1091,26 @@ void UNiagaraDataInterfaceRigidMeshCollisionQuery::GetParameterDefinitionHLSL(co
 }
 #endif
 
+#if WITH_EDITOR
+void UNiagaraDataInterfaceRigidMeshCollisionQuery::ValidateFunction(const FNiagaraFunctionSignature& Function, TArray<FText>& OutValidationErrors)
+{
+	if (Function.Name == GetClosestPointMeshDistanceFieldName)
+	{
+		if (!IsMeshDistanceFieldEnabled())
+		{
+			OutValidationErrors.Add(NSLOCTEXT("UNiagaraDataInterfaceRigidMeshCollisionQuery", "NiagaraDistanceFieldNotEnabledMsg", "The mesh distance field generation is currently not enabled, please check the project settings.\nNiagara cannot query the mesh distance fields otherwise."));
+		}
+	}
+}
+#endif
+
 void UNiagaraDataInterfaceRigidMeshCollisionQuery::ProvidePerInstanceDataForRenderThread(void* DataForRenderThread, void* PerInstanceData, const FNiagaraSystemInstanceID& SystemInstance)
 {
 	FNDIRigidMeshCollisionData* GameThreadData = static_cast<FNDIRigidMeshCollisionData*>(PerInstanceData);
 	FNDIRigidMeshCollisionData* RenderThreadData = static_cast<FNDIRigidMeshCollisionData*>(DataForRenderThread);
 
 	if (GameThreadData != nullptr && RenderThreadData != nullptr)
-	{
+	{		
 		RenderThreadData->AssetBuffer = GameThreadData->AssetBuffer;		
 		RenderThreadData->AssetArrays = GameThreadData->AssetArrays;
 		RenderThreadData->TickingGroup = GameThreadData->TickingGroup;
diff --git a/Engine/Plugins/FX/Niagara/Source/Niagara/Private/NiagaraGpuComputeDispatch.cpp b/Engine/Plugins/FX/Niagara/Source/Niagara/Private/NiagaraGpuComputeDispatch.cpp
--- a/Engine/Plugins/FX/Niagara/Source/Niagara/Private/NiagaraGpuComputeDispatch.cpp
+++ b/Engine/Plugins/FX/Niagara/Source/Niagara/Private/NiagaraGpuComputeDispatch.cpp
@@ -33,6 +33,7 @@
 #include "PipelineStateCache.h"
 #include "SceneUtils.h"
 #include "ShaderParameterUtils.h"
+#include "Renderer/Private/ScenePrivate.h"
 
 DECLARE_CYCLE_STAT(TEXT("Niagara Dispatch Setup"), STAT_NiagaraGPUDispatchSetup_RT, STATGROUP_Niagara);
 DECLARE_CYCLE_STAT(TEXT("GPU Emitter Dispatch [RT]"), STAT_NiagaraGPUSimTick_RT, STATGROUP_Niagara);
@@ -1879,6 +1880,17 @@ const FGlobalDistanceFieldParameterData* FNiagaraGpuComputeDispatch::GetGlobalDi
 	return &CurrentPassViews[0].GlobalDistanceFieldInfo.ParameterData; 
 }
 
+const FDistanceFieldSceneData* FNiagaraGpuComputeDispatch::GetMeshDistanceFieldParameters() const
+{
+	check(CurrentPassViews.Num() > 0);
+	if (CurrentPassViews[0].Family == nullptr || CurrentPassViews[0].Family->Scene == nullptr || CurrentPassViews[0].Family->Scene->GetRenderScene() == nullptr)
+	{
+		return nullptr;
+	}
+
+	return &CurrentPassViews[0].Family->Scene->GetRenderScene()->DistanceFieldSceneData;
+}
+
 void FNiagaraGpuComputeDispatch::GenerateSortKeys(FRHICommandListImmediate& RHICmdList, int32 BatchId, int32 NumElementsInBatch, EGPUSortFlags Flags, FRHIUnorderedAccessView* KeysUAV, FRHIUnorderedAccessView* ValuesUAV)
 {
 	const FGPUSortManager::FKeyGenInfo KeyGenInfo((uint32)NumElementsInBatch, EnumHasAnyFlags(Flags, EGPUSortFlags::HighPrecisionKeys));
diff --git a/Engine/Plugins/FX/Niagara/Source/Niagara/Private/NiagaraGpuComputeDispatch.h b/Engine/Plugins/FX/Niagara/Source/Niagara/Private/NiagaraGpuComputeDispatch.h
--- a/Engine/Plugins/FX/Niagara/Source/Niagara/Private/NiagaraGpuComputeDispatch.h
+++ b/Engine/Plugins/FX/Niagara/Source/Niagara/Private/NiagaraGpuComputeDispatch.h
@@ -28,6 +28,7 @@ the same VectorVM byte code / compute shader code
 
 class FGPUSortManager;
 class FNiagaraRayTracingHelper;
+class FDistanceFieldSceneData;
 
 class FNiagaraGpuComputeDispatch : public FNiagaraGpuComputeDispatchInterface
 {
@@ -91,6 +92,7 @@ public:
 	virtual bool AddSortedGPUSimulation(FNiagaraGPUSortInfo& SortInfo) override;
 
 	const FGlobalDistanceFieldParameterData* GetGlobalDistanceFieldParameters() const;
+	const FDistanceFieldSceneData* GetMeshDistanceFieldParameters() const;
 
 	void SetDataInterfaceParameters(FRHICommandList& RHICmdList, const FNiagaraGPUSystemTick& Tick, const FNiagaraComputeInstanceData& InstanceData, const FNiagaraShaderRef& ComputeShader, const FNiagaraSimStageData& SimStageData) const;
 	void UnsetDataInterfaceParameters(FRHICommandList& RHICmdList, const FNiagaraGPUSystemTick& Tick, const FNiagaraComputeInstanceData& InstanceData, const FNiagaraShaderRef& ComputeShader, const FNiagaraSimStageData& SimStageData) const;
diff --git a/Engine/Plugins/FX/Niagara/Source/Niagara/Public/NiagaraDataInterfaceRigidMeshCollisionQuery.h b/Engine/Plugins/FX/Niagara/Source/Niagara/Public/NiagaraDataInterfaceRigidMeshCollisionQuery.h
--- a/Engine/Plugins/FX/Niagara/Source/Niagara/Public/NiagaraDataInterfaceRigidMeshCollisionQuery.h
+++ b/Engine/Plugins/FX/Niagara/Source/Niagara/Public/NiagaraDataInterfaceRigidMeshCollisionQuery.h
@@ -33,14 +33,16 @@ struct FElementOffset
 struct FNDIRigidMeshCollisionArrays
 {
 	FElementOffset ElementOffsets;
-	TStaticArray<FVector4, 2 * RIGID_MESH_COLLISION_QUERY_MAX_TRANSFORMS> WorldTransform;
-	TStaticArray<FVector4, 2 * RIGID_MESH_COLLISION_QUERY_MAX_TRANSFORMS> InverseTransform;
-	TStaticArray<FVector4, RIGID_MESH_COLLISION_QUERY_MAX_TRANSFORMS> CurrentTransform;
-	TStaticArray<FVector4, RIGID_MESH_COLLISION_QUERY_MAX_TRANSFORMS> CurrentInverse;
-	TStaticArray<FVector4, RIGID_MESH_COLLISION_QUERY_MAX_TRANSFORMS> PreviousTransform;
-	TStaticArray<FVector4, RIGID_MESH_COLLISION_QUERY_MAX_TRANSFORMS> PreviousInverse;
-	TStaticArray<FVector4, RIGID_MESH_COLLISION_QUERY_MAX_PRIMITIVES> ElementExtent;
+	TStaticArray<FVector4f, 2 * RIGID_MESH_COLLISION_QUERY_MAX_TRANSFORMS> WorldTransform;
+	TStaticArray<FVector4f, 2 * RIGID_MESH_COLLISION_QUERY_MAX_TRANSFORMS> InverseTransform;
+	TStaticArray<FVector4f, RIGID_MESH_COLLISION_QUERY_MAX_TRANSFORMS> CurrentTransform;
+	TStaticArray<FVector4f, RIGID_MESH_COLLISION_QUERY_MAX_TRANSFORMS> CurrentInverse;
+	TStaticArray<FVector4f, RIGID_MESH_COLLISION_QUERY_MAX_TRANSFORMS> PreviousTransform;
+	TStaticArray<FVector4f, RIGID_MESH_COLLISION_QUERY_MAX_TRANSFORMS> PreviousInverse;
+	TStaticArray<FVector4f, RIGID_MESH_COLLISION_QUERY_MAX_PRIMITIVES> ElementExtent;
 	TStaticArray<uint32, RIGID_MESH_COLLISION_QUERY_MAX_PRIMITIVES> PhysicsType;
+	TStaticArray<uint32, RIGID_MESH_COLLISION_QUERY_MAX_PRIMITIVES> DFIndex;
+	TStaticArray<FPrimitiveSceneProxy*, RIGID_MESH_COLLISION_QUERY_MAX_PRIMITIVES> SourceSceneProxy;
 };
 
 /** Render buffers that will be used in hlsl functions */
@@ -66,6 +68,9 @@ struct FNDIRigidMeshCollisionBuffer : public FRenderResource
 
 	/** Physics type buffer */
 	FRWBuffer PhysicsTypeBuffer;
+
+	/** Distance field index buffer */
+	FRWBuffer DFIndexBuffer;
 };
 
 /** Data stored per physics asset instance*/
@@ -130,11 +135,15 @@ public:
 	virtual bool HasTickGroupPrereqs() const override { return true; }
 	virtual ETickingGroup CalculateTickGroup(const void* PerInstanceData) const override;
 
+	virtual bool RequiresDistanceFieldData() const override { return true; }
+
 	/** GPU simulation  functionality */
 #if WITH_EDITORONLY_DATA
 	virtual void GetCommonHLSL(FString& OutHLSL) override;
 	virtual void GetParameterDefinitionHLSL(const FNiagaraDataInterfaceGPUParamInfo& ParamInfo, FString& OutHLSL) override;
 	virtual bool GetFunctionHLSL(const FNiagaraDataInterfaceGPUParamInfo& ParamInfo, const FNiagaraDataInterfaceGeneratedFunction& FunctionInfo, int FunctionInstanceIndex, FString& OutHLSL) override;
+
+	virtual void ValidateFunction(const FNiagaraFunctionSignature& Function, TArray<FText>& OutValidationErrors) override;
 #endif
 	virtual void ProvidePerInstanceDataForRenderThread(void* DataForRenderThread, void* PerInstanceData, const FNiagaraSystemInstanceID& SystemInstance) override;
 
@@ -153,6 +162,9 @@ public:
 	/** Name of the physics type buffer */
 	static const FString PhysicsTypeBufferName;
 
+	/** Name of the DF Index type buffer */
+	static const FString DFIndexBufferName;
+
 protected:
 	/** Copy one niagara DI to this */
 	virtual bool CopyToInternal(UNiagaraDataInterface* Destination) const override;
@@ -181,5 +193,4 @@ struct FNDIRigidMeshCollisionProxy : public FNiagaraDataInterfaceProxy
 
 	/** List of proxy data for each system instances*/
 	TMap<FNiagaraSystemInstanceID, FNDIRigidMeshCollisionData> SystemInstancesToProxyData;
-};
-
+};
\ No newline at end of file
diff --git a/Engine/Plugins/FX/Niagara/Source/NiagaraCore/Private/NiagaraCustomVersion.cpp b/Engine/Plugins/FX/Niagara/Source/NiagaraCore/Private/NiagaraCustomVersion.cpp
--- a/Engine/Plugins/FX/Niagara/Source/NiagaraCore/Private/NiagaraCustomVersion.cpp
+++ b/Engine/Plugins/FX/Niagara/Source/NiagaraCore/Private/NiagaraCustomVersion.cpp
@@ -8,5 +8,4 @@ const FGuid FNiagaraCustomVersion::GUID(0xFCF57AFA, 0x50764283, 0xB9A9E658, 0xFF
 // Register the custom version with core
 FCustomVersionRegistration GRegisterNiagaraCustomVersion(FNiagaraCustomVersion::GUID, FNiagaraCustomVersion::LatestVersion, TEXT("NiagaraVer"));
 
-const FGuid FNiagaraCustomVersion::LatestScriptCompileVersion(0xC6985CB9, 0xA75A5C45, 0xA99D541B, 0x4E8A9FB0);
-
+const FGuid FNiagaraCustomVersion::LatestScriptCompileVersion(0x6E575DF3, 0x9F8B432E, 0x928A9DDA, 0x1497B4A5);
\ No newline at end of file
