cis fix
fix for check against nullptr

diff --git a/Engine/Plugins/FX/Niagara/Source/Niagara/Private/NiagaraDataInterfaceRigidMeshCollisionQuery.cpp b/Engine/Plugins/FX/Niagara/Source/Niagara/Private/NiagaraDataInterfaceRigidMeshCollisionQuery.cpp
--- a/Engine/Plugins/FX/Niagara/Source/Niagara/Private/NiagaraDataInterfaceRigidMeshCollisionQuery.cpp
+++ b/Engine/Plugins/FX/Niagara/Source/Niagara/Private/NiagaraDataInterfaceRigidMeshCollisionQuery.cpp
@@ -378,26 +378,26 @@ void FNDIRigidMeshCollisionData::Init(UNiagaraDataInterfaceRigidMeshCollisionQue
 {
 	AssetBuffer = nullptr;
 
-	UWorld* World = SystemInstance->GetWorld();
+	if (Interface != nullptr && SystemInstance != nullptr)
+	{
+		UWorld* World = SystemInstance->GetWorld();
 
-	StaticMeshActors.Empty();
+		StaticMeshActors.Empty();
 
-	if (Interface->StaticMesh)
-	{
-		for (TActorIterator<AStaticMeshActor> It(World); It; ++It)
+		if (Interface->StaticMesh)
 		{
-			AStaticMeshActor* StaticMeshActor = *It;
-
-			if ((!Interface->OnlyUseMoveable || (Interface->OnlyUseMoveable && StaticMeshActor->IsRootComponentMovable())) && 
-				(Interface->Tag == FString("") || (Interface->Tag != FString("") && StaticMeshActor->Tags.Contains(FName(Interface->Tag)))))
+			for (TActorIterator<AStaticMeshActor> It(World); It; ++It)
 			{
-				StaticMeshActors.Add(StaticMeshActor);
-			}			
+				AStaticMeshActor* StaticMeshActor = *It;
+
+				if ((!Interface->OnlyUseMoveable || (Interface->OnlyUseMoveable && StaticMeshActor->IsRootComponentMovable())) &&
+					(Interface->Tag == FString("") || (Interface->Tag != FString("") && StaticMeshActor->Tags.Contains(FName(Interface->Tag)))))
+				{
+					StaticMeshActors.Add(StaticMeshActor);
+				}
+			}
 		}
-	}
 
-	if (Interface != nullptr && SystemInstance != nullptr)
-	{
 		if (0 < StaticMeshActors.Num() && StaticMeshActors[0] != nullptr)
 		{
 			// @note: we are not creating internal arrays here and letting it happen on the call to update
@@ -657,9 +657,15 @@ void FNDIRigidMeshCollisionProxy::ConsumePerInstanceDataFromGameThread(void* Per
 		{
 			FPrimitiveSceneProxy* Proxy = SourceData->AssetArrays.SourceSceneProxy[i];
 								
-			const TArray<int32, TInlineAllocator<1>>& DFIndices = Proxy != nullptr && Proxy->GetPrimitiveSceneInfo() != nullptr ?
-				Proxy->GetPrimitiveSceneInfo()->DistanceFieldInstanceIndices : TArray<int32, TInlineAllocator<1>>();			
-			TargetData->AssetArrays.DFIndex[i] = DFIndices.Num() > 0 ? DFIndices[0] : -1;
+			if (Proxy != nullptr && Proxy->GetPrimitiveSceneInfo() != nullptr)
+			{
+				const TArray<int32, TInlineAllocator<1>>& DFIndices = Proxy->GetPrimitiveSceneInfo()->DistanceFieldInstanceIndices;
+				TargetData->AssetArrays.DFIndex[i] = DFIndices.Num() > 0 ? DFIndices[0] : -1;
+			}
+			else
+			{
+				TargetData->AssetArrays.DFIndex[i] = -1;
+			}
 		}		
 	}
 	else
