- Fix missing FVector4d specialization

#rnx

diff --git a/Engine/Plugins/FX/Niagara/Source/Niagara/Public/NiagaraParameterStore.h b/Engine/Plugins/FX/Niagara/Source/Niagara/Public/NiagaraParameterStore.h
--- a/Engine/Plugins/FX/Niagara/Source/Niagara/Public/NiagaraParameterStore.h
+++ b/Engine/Plugins/FX/Niagara/Source/Niagara/Public/NiagaraParameterStore.h
@@ -663,6 +663,12 @@ FORCEINLINE bool FNiagaraParameterStore::SetParameterValue(const FVector3d& InVa
 	return SetParameterValue((FVector3f)InValue, Param, bAdd);
 }
 
+template<>
+FORCEINLINE bool FNiagaraParameterStore::SetParameterValue(const FVector4d& InValue, const FNiagaraVariable& Param, bool bAdd)
+{
+	return SetParameterValue((FVector4f)InValue, Param, bAdd);
+}
+
 template<>
 FORCEINLINE bool FNiagaraParameterStore::SetParameterValue(const FQuat4d& InValue, const FNiagaraVariable& Param, bool bAdd)
 {
