fix mesh distance field sampling
more test content

diff --git a/Engine/Plugins/FX/Niagara/Source/Niagara/Private/NiagaraDataInterfaceRigidMeshCollisionQuery.cpp b/Engine/Plugins/FX/Niagara/Source/Niagara/Private/NiagaraDataInterfaceRigidMeshCollisionQuery.cpp
--- a/Engine/Plugins/FX/Niagara/Source/Niagara/Private/NiagaraDataInterfaceRigidMeshCollisionQuery.cpp
+++ b/Engine/Plugins/FX/Niagara/Source/Niagara/Private/NiagaraDataInterfaceRigidMeshCollisionQuery.cpp
@@ -466,6 +466,9 @@ public:
 		DistanceFieldBrickAtlasMask.Bind(ParameterMap, TEXT("DistanceFieldBrickAtlasMask"));
 		DistanceFieldBrickAtlasSizeLog2.Bind(ParameterMap, TEXT("DistanceFieldBrickAtlasSizeLog2"));
 		DistanceFieldBrickAtlasTexelSize.Bind(ParameterMap, TEXT("DistanceFieldBrickAtlasTexelSize"));
+		DistanceFieldBrickAtlasHalfTexelSize.Bind(ParameterMap, TEXT("DistanceFieldBrickAtlasHalfTexelSize"));
+		DistanceFieldBrickOffsetToAtlasUVScale.Bind(ParameterMap, TEXT("DistanceFieldBrickOffsetToAtlasUVScale"));
+		DistanceFieldUniqueDataBrickSizeInAtlasTexels.Bind(ParameterMap, TEXT("DistanceFieldUniqueDataBrickSizeInAtlasTexels"));
 	}
 
 	bool IsBound() const
@@ -488,6 +491,9 @@ public:
 		Ar << Parameters.DistanceFieldBrickAtlasMask;
 		Ar << Parameters.DistanceFieldBrickAtlasSizeLog2;
 		Ar << Parameters.DistanceFieldBrickAtlasTexelSize;
+		Ar << Parameters.DistanceFieldBrickAtlasHalfTexelSize;
+		Ar << Parameters.DistanceFieldBrickOffsetToAtlasUVScale;
+		Ar << Parameters.DistanceFieldUniqueDataBrickSizeInAtlasTexels;
 
 		return Ar;
 	}
@@ -512,7 +518,11 @@ public:
 				FMath::FloorLog2(ParameterData->BrickTextureDimensionsInBricks.X),
 				FMath::FloorLog2(ParameterData->BrickTextureDimensionsInBricks.Y),
 				FMath::FloorLog2(ParameterData->BrickTextureDimensionsInBricks.Z)));
-			SetShaderValue(RHICmdList, ShaderRHI, DistanceFieldBrickAtlasTexelSize, FVector3f(1.0f) / FVector3f(ParameterData->BrickTextureDimensionsInBricks * DistanceField::BrickSize));
+			FVector3f DistanceFieldBrickAtlasTexelSizeTmp = FVector3f(1.0f) / FVector3f(ParameterData->BrickTextureDimensionsInBricks * DistanceField::BrickSize);
+			SetShaderValue(RHICmdList, ShaderRHI, DistanceFieldBrickAtlasTexelSize, DistanceFieldBrickAtlasTexelSizeTmp);
+			SetShaderValue(RHICmdList, ShaderRHI, DistanceFieldBrickAtlasHalfTexelSize, 0.5f * DistanceFieldBrickAtlasTexelSizeTmp);
+			SetShaderValue(RHICmdList, ShaderRHI, DistanceFieldBrickOffsetToAtlasUVScale, FVector3f(DistanceField::BrickSize) * DistanceFieldBrickAtlasTexelSizeTmp);
+			SetShaderValue(RHICmdList, ShaderRHI, DistanceFieldUniqueDataBrickSizeInAtlasTexels, FVector3f(DistanceField::UniqueDataBrickSize) * DistanceFieldBrickAtlasTexelSizeTmp);
 		}
 	}
 
@@ -531,6 +541,9 @@ private:
 		LAYOUT_FIELD(FShaderParameter, DistanceFieldBrickAtlasMask)
 		LAYOUT_FIELD(FShaderParameter, DistanceFieldBrickAtlasSizeLog2)
 		LAYOUT_FIELD(FShaderParameter, DistanceFieldBrickAtlasTexelSize)
+		LAYOUT_FIELD(FShaderParameter, DistanceFieldBrickAtlasHalfTexelSize)
+		LAYOUT_FIELD(FShaderParameter, DistanceFieldBrickOffsetToAtlasUVScale)
+		LAYOUT_FIELD(FShaderParameter, DistanceFieldUniqueDataBrickSizeInAtlasTexels)
 };
 
 //------------------------------------------------------------------------------------------------------------
@@ -595,6 +608,8 @@ public:
 					UE_LOG(LogRigidMeshCollision, Error, TEXT("Distance fields are not available for use"));
 					// #todo(dmp): should we set something here in the case where distance field data is bound but we don't have it?
 					// there is no trivial constructor
+					//FDistanceFieldSceneData DummyDistanceFieldSceneData(Context.Shader->GetShaderPlatform());
+					//DistanceFieldParameters.SetEmpty(RHICmdList, ComputeShaderRHI, DummyDistanceFieldSceneData);
 				}
 				else 
 				{
@@ -641,6 +656,8 @@ IMPLEMENT_NIAGARA_DI_PARAMETER(UNiagaraDataInterfaceRigidMeshCollisionQuery, FND
 
 void FNDIRigidMeshCollisionProxy::ConsumePerInstanceDataFromGameThread(void* PerInstanceData, const FNiagaraSystemInstanceID& Instance)
 {
+	check(IsInRenderingThread());
+
 	FNDIRigidMeshCollisionData* SourceData = static_cast<FNDIRigidMeshCollisionData*>(PerInstanceData);	
 	FNDIRigidMeshCollisionData* TargetData = &(SystemInstancesToProxyData.FindOrAdd(Instance));
 
@@ -653,20 +670,21 @@ void FNDIRigidMeshCollisionProxy::ConsumePerInstanceDataFromGameThread(void* Per
 
 		// loop over proxies and compute df indices on the RT only on new data
 		// @todo(dmp): for now we do this every frame because it seems like sometimes DFindices are not set
+
 		for (uint32 i = 0; i < SourceData->AssetArrays.ElementOffsets.NumElements; ++i)
 		{
 			FPrimitiveSceneProxy* Proxy = SourceData->AssetArrays.SourceSceneProxy[i];
 								
 			if (Proxy != nullptr && Proxy->GetPrimitiveSceneInfo() != nullptr)
 			{
-				const TArray<int32, TInlineAllocator<1>>& DFIndices = Proxy->GetPrimitiveSceneInfo()->DistanceFieldInstanceIndices;
+				const TArray<int32, TInlineAllocator<1>>& DFIndices = Proxy->GetPrimitiveSceneInfo()->DistanceFieldInstanceIndices;				
 				TargetData->AssetArrays.DFIndex[i] = DFIndices.Num() > 0 ? DFIndices[0] : -1;
 			}
 			else
 			{
 				TargetData->AssetArrays.DFIndex[i] = -1;
 			}
-		}		
+		}	
 	}
 	else
 	{
