fix crash with quering static mesh scene proxy
bump niagara custom version for fix w/ mesh distance field sampling

diff --git a/Engine/Plugins/FX/Niagara/Source/Niagara/Private/NiagaraDataInterfaceRigidMeshCollisionQuery.cpp b/Engine/Plugins/FX/Niagara/Source/Niagara/Private/NiagaraDataInterfaceRigidMeshCollisionQuery.cpp
--- a/Engine/Plugins/FX/Niagara/Source/Niagara/Private/NiagaraDataInterfaceRigidMeshCollisionQuery.cpp
+++ b/Engine/Plugins/FX/Niagara/Source/Niagara/Private/NiagaraDataInterfaceRigidMeshCollisionQuery.cpp
@@ -306,6 +306,7 @@ void UpdateInternalArrays(const TArray<AStaticMeshActor*> &StaticMeshActors, FND
 				{
 					if (CollisionEnabledHasPhysics(BoxElem.GetCollisionEnabled()))
 					{
+						OutAssetArrays->SourceSceneProxy[BoxCount] = StaticMeshComponent->SceneProxy;
 						const FTransform ElementTransform = FTransform(BoxElem.Rotation, BoxElem.Center) * MeshTransform;
 						FillCurrentTransforms(ElementTransform, BoxCount, OutAssetArrays->CurrentTransform, OutAssetArrays->CurrentInverse);
 					}
@@ -315,6 +316,7 @@ void UpdateInternalArrays(const TArray<AStaticMeshActor*> &StaticMeshActors, FND
 				{
 					if (CollisionEnabledHasPhysics(SphereElem.GetCollisionEnabled()))
 					{
+						OutAssetArrays->SourceSceneProxy[SphereCount] = StaticMeshComponent->SceneProxy;
 						const FTransform ElementTransform = FTransform(SphereElem.Center) * MeshTransform;
 						FillCurrentTransforms(ElementTransform, SphereCount, OutAssetArrays->CurrentTransform, OutAssetArrays->CurrentInverse);
 					}
@@ -324,6 +326,7 @@ void UpdateInternalArrays(const TArray<AStaticMeshActor*> &StaticMeshActors, FND
 				{
 					if (CollisionEnabledHasPhysics(CapsuleElem.GetCollisionEnabled()))
 					{
+						OutAssetArrays->SourceSceneProxy[CapsuleCount] = StaticMeshComponent->SceneProxy;
 						const FTransform ElementTransform = FTransform(CapsuleElem.Rotation, CapsuleElem.Center) * MeshTransform;
 						FillCurrentTransforms(ElementTransform, CapsuleCount, OutAssetArrays->CurrentTransform, OutAssetArrays->CurrentInverse);
 					}
diff --git a/Engine/Plugins/FX/Niagara/Source/NiagaraCore/Private/NiagaraCustomVersion.cpp b/Engine/Plugins/FX/Niagara/Source/NiagaraCore/Private/NiagaraCustomVersion.cpp
--- a/Engine/Plugins/FX/Niagara/Source/NiagaraCore/Private/NiagaraCustomVersion.cpp
+++ b/Engine/Plugins/FX/Niagara/Source/NiagaraCore/Private/NiagaraCustomVersion.cpp
@@ -8,4 +8,4 @@ const FGuid FNiagaraCustomVersion::GUID(0xFCF57AFA, 0x50764283, 0xB9A9E658, 0xFF
 // Register the custom version with core
 FCustomVersionRegistration GRegisterNiagaraCustomVersion(FNiagaraCustomVersion::GUID, FNiagaraCustomVersion::LatestVersion, TEXT("NiagaraVer"));
 
-const FGuid FNiagaraCustomVersion::LatestScriptCompileVersion(0x6E575DF3, 0x9F8B432E, 0x928A9DDA, 0x1497B4A5);
\ No newline at end of file
+const FGuid FNiagaraCustomVersion::LatestScriptCompileVersion(0xCBDDE73F, 0xA0584B4F, 0x864D4335, 0x59C858DA);
\ No newline at end of file
