- Use GetActorNameOrLabel for actor name in actor filter suggestions
- Clean up code in NiagaraHud

#rnx

diff --git a/Engine/Plugins/FX/Niagara/Source/Niagara/Private/NiagaraDebugHud.cpp b/Engine/Plugins/FX/Niagara/Source/Niagara/Private/NiagaraDebugHud.cpp
--- a/Engine/Plugins/FX/Niagara/Source/Niagara/Private/NiagaraDebugHud.cpp
+++ b/Engine/Plugins/FX/Niagara/Source/Niagara/Private/NiagaraDebugHud.cpp
@@ -516,12 +516,6 @@ namespace NiagaraDebugLocal
 			}
 		}
 	}
-
-	FString GetHudActorName(AActor* Actor)
-	{
-		FString Name = Actor->GetActorNameOrLabel();
-		return Name.Len() == 0 ? GetNameSafe(Actor) : Name;
-	}
 }
 
 FNiagaraDebugHud::FNiagaraDebugHud(UWorld* World)
@@ -644,7 +638,7 @@ void FNiagaraDebugHud::GatherSystemInfo()
 			if ( Settings.bActorFilterEnabled )
 			{
 				AActor* Actor = NiagaraComponent->GetOwner();
-				bIsMatch &= (Actor != nullptr) && GetHudActorName(Actor).MatchesWildcard(Settings.ActorFilter);
+				bIsMatch &= (Actor != nullptr) && Actor->GetActorNameOrLabel().MatchesWildcard(Settings.ActorFilter);
 			}
 
 			// Filter by component
@@ -1405,7 +1399,7 @@ void FNiagaraDebugHud::DrawComponents(FNiagaraWorldManager* WorldManager, UCanva
 				StringBuilder.Append(TEXT("Component - "));
 				if ( OwnerActor )
 				{
-					StringBuilder.Append(*GetHudActorName(OwnerActor));
+					StringBuilder.Append(*OwnerActor->GetActorNameOrLabel());
 					StringBuilder.Append(TEXT("/"));
 				}
 				StringBuilder.Append(*GetNameSafe(NiagaraComponent));
diff --git a/Engine/Plugins/FX/Niagara/Source/Niagara/Private/NiagaraDebuggerClient.cpp b/Engine/Plugins/FX/Niagara/Source/Niagara/Private/NiagaraDebuggerClient.cpp
--- a/Engine/Plugins/FX/Niagara/Source/Niagara/Private/NiagaraDebuggerClient.cpp
+++ b/Engine/Plugins/FX/Niagara/Source/Niagara/Private/NiagaraDebuggerClient.cpp
@@ -95,7 +95,7 @@ void FNiagaraDebuggerClient::UpdateClientInfo()
 		{
 			if (Actor)
 			{
-				NewInfo->Actors.Add(Actor->GetName());
+				NewInfo->Actors.Add(Actor->GetActorNameOrLabel());
 			}
 		}
 
