- Static Mesh Data Interface API consistency changes

diff --git a/Engine/Plugins/FX/Niagara/Shaders/Private/NiagaraDataInterfaceStaticMeshTemplate.ush b/Engine/Plugins/FX/Niagara/Shaders/Private/NiagaraDataInterfaceStaticMeshTemplate.ush
new file mode 100644
--- /dev/null
+++ b/Engine/Plugins/FX/Niagara/Shaders/Private/NiagaraDataInterfaceStaticMeshTemplate.ush
@@ -0,0 +1,559 @@
+// Copyright Epic Games, Inc. All Rights Reserved.
+
+int3			NumTriangles_{ParameterName};		// x = Num Sections, y = Num Filtered, z = Num Unfiltered
+int				NumVertices_{ParameterName};
+int				NumUVs_{ParameterName};
+uint			HasColors_{ParameterName};
+Buffer<uint>	IndexBuffer_{ParameterName};
+Buffer<float>	PositionBuffer_{ParameterName};
+Buffer<float4>	TangentBuffer_{ParameterName};
+Buffer<float2>	UVBuffer_{ParameterName};
+Buffer<float4>	ColorBuffer_{ParameterName};
+
+uint			HasUniformSampling_{ParameterName};
+Buffer<uint2>	UniformSamplingTriangles_{ParameterName};
+
+int3			SectionCounts_{ParameterName};		// x = Num Sections, y = Num Filtered, z = Num Unfiltered
+Buffer<uint4>	SectionInfos_{ParameterName};		// FirstTriangle, NumTriangles, Prob, Alias
+Buffer<uint>	FilteredAndUnfilteredSections_{ParameterName};
+
+int3			SocketCounts_{ParameterName};		// x = Num Sockets, y = Num Filtered, z = Num Unfiltered
+Buffer<float4>	SocketTransforms_{ParameterName};
+Buffer<uint>	FilteredAndUnfilteredSockets_{ParameterName};
+
+float			InvDeltaSeconds_{ParameterName};
+float4x4		InstanceTransform_{ParameterName};
+float4x4		InstanceTransformInverseTransposed_{ParameterName};
+float4			InstanceRotation_{ParameterName};
+float4x4		InstancePreviousTransform_{ParameterName};
+float4x4		InstancePreviousTransformInverseTransposed_{ParameterName};
+float4			InstancePreviousRotation_{ParameterName};
+float3			InstanceWorldVelocity_{ParameterName};
+
+//////////////////////////////////////////////////////////////////////////
+// Fwd declarations
+void RandomSection_{ParameterName}(NiagaraRandInfo RandInfo, out int Section);
+void RandomFilteredSection_{ParameterName}(NiagaraRandInfo RandInfo, out int Section);
+void RandomUnfilteredSection_{ParameterName}(NiagaraRandInfo RandInfo, out int Section);
+void RandomSectionTriangle_{ParameterName}(NiagaraRandInfo RandInfo, int Section, out int Triangle, out float3 BaryCoord);
+
+//////////////////////////////////////////////////////////////////////////
+int UniformTriangle_{ParameterName}(NiagaraRandInfo RandInfo, int Triangle, int SectionOffset)
+{
+	uint2 ProbAlias = UniformSamplingTriangles_{ParameterName}[Triangle];
+	float Prob = asfloat(ProbAlias.x);
+	int Alias = ProbAlias.y;
+	return NiagaraRandomFloat(RandInfo) < Prob ? Triangle : Alias + SectionOffset;
+}
+
+//////////////////////////////////////////////////////////////////////////
+// Vertex Sampling
+void IsValidVertex_{ParameterName}(out bool IsValid)
+{
+	IsValid = NumVertices_{ParameterName} > 0;
+}
+
+void RandomVertex_{ParameterName}(NiagaraRandInfo RandInfo, out int Vertex)
+{
+	Vertex = NiagaraRandomInt(RandInfo, NumVertices_{ParameterName});
+}
+
+void GetVertexCount_{ParameterName}(out int Count)
+{
+	Count = NumVertices_{ParameterName};
+}
+
+void GetVertex_{ParameterName}(int Vertex, out float3 Position, out float3 Velocity, out float3 Normal, out float3 Bitangent, out float3 Tangent)
+{
+	[branch]
+	if ( NumVertices_{ParameterName} > 0 )
+	{
+		Vertex = clamp(Vertex, 0, NumVertices_{ParameterName} - 1);
+
+		Position.x = PositionBuffer_{ParameterName}[Vertex * 3 + 0];
+		Position.y = PositionBuffer_{ParameterName}[Vertex * 3 + 1];
+		Position.z = PositionBuffer_{ParameterName}[Vertex * 3 + 2];
+
+		float4 TangentX = TangentBias(TangentBuffer_{ParameterName}[Vertex * 2 + 0]);
+		float4 TangentZ = TangentBias(TangentBuffer_{ParameterName}[Vertex * 2 + 1]);
+		float3 TangentY = cross(TangentZ.xyz, TangentX.xyz) * TangentZ.w;
+
+		Normal		= TangentZ.xyz;
+		Bitangent	= TangentY;
+		Tangent		= TangentX.xyz;
+	}
+	else
+	{
+		Position	= float3(0, 0, 0);
+		Velocity	= float3(0, 0, 0);
+		Normal		= float3(0, 0, 1);
+		Bitangent	= float3(0, 1, 0);
+		Tangent		= float3(1, 0, 0);
+	}
+}
+
+void GetVertexWS_{ParameterName}(int Vertex, out float3 Position, out float3 Velocity, out float3 Normal, out float3 Bitangent, out float3 Tangent)
+{
+	float3 LocalPosition;
+	GetVertex_{ParameterName}(Vertex, LocalPosition, Velocity, Normal, Bitangent, Tangent);
+	
+	Position				= mul(float4(LocalPosition, 1), InstanceTransform_{ParameterName}).xyz;
+	float3 PreviousPosition	= mul(float4(LocalPosition, 1), InstancePreviousTransform_{ParameterName}).xyz;
+	Velocity				= (Position - PreviousPosition) * InvDeltaSeconds_{ParameterName};
+	Normal					= normalize(mul(float4(Normal, 0), InstanceTransformInverseTransposed_{ParameterName}).xyz);
+	Bitangent				= normalize(mul(float4(Bitangent, 0), InstanceTransformInverseTransposed_{ParameterName}).xyz);
+	Tangent					= normalize(mul(float4(Tangent, 0), InstanceTransformInverseTransposed_{ParameterName}).xyz);
+}
+
+void GetVertexColor_{ParameterName}(int Vertex, out float4 Color)
+{
+	Vertex = clamp(Vertex, 0, NumVertices_{ParameterName} - 1);
+	Color = HasColors_{ParameterName} ? ColorBuffer_{ParameterName}[Vertex] : float4(1, 1, 1, 1);
+}
+
+void GetVertexUV_{ParameterName}(int Vertex, int UVSet, out float2 UV)
+{
+	Vertex = clamp(Vertex, 0, NumVertices_{ParameterName} - 1);
+	UVSet = clamp(UVSet, 0, NumUVs_{ParameterName} - 1);
+	UV = UVSet >= 0 && UVSet < NumUVs_{ParameterName} ? UVBuffer_{ParameterName}[Vertex * NumUVs_{ParameterName} + UVSet] : float2(0, 0);
+}
+
+//////////////////////////////////////////////////////////////////////////
+// VM Triangle Sampling
+void IsValidTriangle_{ParameterName}(int Triangle, out bool IsValid)
+{
+	IsValid = Triangle >= 0 && Triangle < NumTriangles_{ParameterName}.x;
+}
+
+void RandomTriangle_{ParameterName}(NiagaraRandInfo RandInfo, out int Triangle, out float3 BaryCoord)
+{
+	int Section;
+	RandomSection_{ParameterName}(RandInfo, Section);
+	RandomSectionTriangle_{ParameterName}(RandInfo, Section, Triangle, BaryCoord);
+}
+
+void GetTriangleCount_{ParameterName}(out int Count)
+{
+	Count = NumTriangles_{ParameterName}.x;
+}
+
+void IsValidFilteredTriangle_{ParameterName}(int Triangle, out bool IsValid)
+{
+	IsValid = Triangle >= 0 && Triangle < NumTriangles_{ParameterName}.y;
+}
+
+void RandomFilteredTriangle_{ParameterName}(NiagaraRandInfo RandInfo, out int Triangle, out float3 BaryCoord)
+{
+	if ( NumTriangles_{ParameterName}.y > 0 )
+	{
+		int Section;
+		RandomFilteredSection_{ParameterName}(RandInfo, Section);
+		RandomSectionTriangle_{ParameterName}(RandInfo, Section, Triangle, BaryCoord);
+	}
+	else
+	{
+		Triangle = 0;
+		BaryCoord = NiagaraRandomBaryCoord(RandInfo);
+	}
+}
+
+void GetFilteredTriangleCount_{ParameterName}(out int Count)
+{
+	Count = NumTriangles_{ParameterName}.y;
+}
+
+void GetFilteredTriangleAt_{ParameterName}(int TriangleIndex, out int Triangle)
+{
+	Triangle = 0;
+
+	int SectionInfoOffset = SectionCounts_{ParameterName}.x;
+	int SectionFirstTriangle = 0;
+	for ( int i=0; i < SectionCounts_{ParameterName}.y; ++i )
+	{
+		int SectionLastTriangle = SectionFirstTriangle + SectionInfos_{ParameterName}[SectionInfoOffset + i].y;
+		if ( TriangleIndex >= SectionFirstTriangle && TriangleIndex < SectionLastTriangle )
+		{
+			Triangle = (TriangleIndex - SectionFirstTriangle) + SectionInfos_{ParameterName}[SectionInfoOffset + i].x;
+			return;
+		}
+		SectionFirstTriangle = SectionLastTriangle;
+	}
+}
+
+void IsValidUnfilteredTriangle_{ParameterName}(int Triangle, out bool IsValid)
+{
+	IsValid = Triangle >= 0 && Triangle < NumTriangles_{ParameterName}.z;
+}
+
+void RandomUnfilteredTriangle_{ParameterName}(NiagaraRandInfo RandInfo, out int Triangle, out float3 BaryCoord)
+{
+	if ( NumTriangles_{ParameterName}.z > 0 )
+	{
+		int Section;
+		RandomUnfilteredSection_{ParameterName}(RandInfo, Section);
+		RandomSectionTriangle_{ParameterName}(RandInfo, Section, Triangle, BaryCoord);
+	}
+	else
+	{
+		Triangle = 0;
+		BaryCoord = NiagaraRandomBaryCoord(RandInfo);
+	}
+}
+
+void GetUnfilteredTriangleCount_{ParameterName}(out int Count)
+{
+	Count = NumTriangles_{ParameterName}.z;
+}
+
+void GetUnfilteredTriangleAt_{ParameterName}(int TriangleIndex, out int Triangle)
+{
+	Triangle = 0;
+
+	int SectionInfoOffset = SectionCounts_{ParameterName}.x + SectionCounts_{ParameterName}.y;
+	int SectionFirstTriangle = 0;
+	for ( int i=0; i < SectionCounts_{ParameterName}.z; ++i )
+	{
+		int SectionLastTriangle = SectionFirstTriangle + SectionInfos_{ParameterName}[SectionInfoOffset + i].y;
+		if ( TriangleIndex >= SectionFirstTriangle && TriangleIndex < SectionLastTriangle )
+		{
+			Triangle = (TriangleIndex - SectionFirstTriangle) + SectionInfos_{ParameterName}[SectionInfoOffset + i].x;
+			return;
+		}
+		SectionFirstTriangle = SectionLastTriangle;
+	}
+}
+
+void GetTriangleIndices_{ParameterName}(int Triangle, out int Index0, out int Index1, out int Index2)
+{
+	[branch]
+	if ( NumTriangles_{ParameterName}.x > 0 )
+	{
+		Triangle = clamp(Triangle, 0, NumTriangles_{ParameterName}.x - 1);
+		Index0 = IndexBuffer_{ParameterName}[Triangle * 3 + 0];
+		Index1 = IndexBuffer_{ParameterName}[Triangle * 3 + 1];
+		Index2 = IndexBuffer_{ParameterName}[Triangle * 3 + 2];
+	}
+	else
+	{
+		Index0 = 0;
+		Index1 = 0;
+		Index2 = 0;
+	}
+}
+
+void GetTriangle_{ParameterName}(int Triangle, float3 BaryCoord, out float3 Position, out float3 Velocity, out float3 Normal, out float3 Bitangent, out float3 Tangent)
+{
+	int Index[3];
+	GetTriangleIndices_{ParameterName}(Triangle, Index[0], Index[1], Index[2]);
+
+	float3 Positions[3], Velocities[3], Normals[3], Bitangents[3], Tangents[3];
+	GetVertex_{ParameterName}(Index[0], Positions[0], Velocities[0], Normals[0], Bitangents[0], Tangents[0]);
+	GetVertex_{ParameterName}(Index[1], Positions[1], Velocities[1], Normals[1], Bitangents[1], Tangents[1]);
+	GetVertex_{ParameterName}(Index[2], Positions[2], Velocities[2], Normals[2], Bitangents[2], Tangents[2]);
+
+	Position	= Positions[0] * BaryCoord.x + Positions[1] * BaryCoord.y + Positions[2] * BaryCoord.z;
+	Velocity	= float3(0, 0, 0);
+	Normal		= Normals[0] * BaryCoord.x + Normals[1] * BaryCoord.y + Normals[2] * BaryCoord.x;
+	Bitangent	= Bitangents[0] * BaryCoord.x + Bitangents[1] * BaryCoord.y + Bitangents[2] * BaryCoord.x;
+	Tangent		= Tangents[0] * BaryCoord.x + Tangents[1] * BaryCoord.y + Tangents[2] * BaryCoord.x;
+}
+
+void GetTriangleWS_{ParameterName}(int Triangle, float3 BaryCoord, out float3 Position, out float3 Velocity, out float3 Normal, out float3 Bitangent, out float3 Tangent)
+{
+	float3 LocalPosition;
+	GetTriangle_{ParameterName}(Triangle, BaryCoord, LocalPosition, Velocity, Normal, Bitangent, Tangent);
+
+	Position				= mul(float4(LocalPosition, 1), InstanceTransform_{ParameterName}).xyz;
+	float3 PreviousPosition	= mul(float4(LocalPosition, 1), InstancePreviousTransform_{ParameterName}).xyz;
+	Velocity				= (Position - PreviousPosition) * InvDeltaSeconds_{ParameterName};
+	Normal					= normalize(mul(float4(Normal, 0), InstanceTransformInverseTransposed_{ParameterName}).xyz);
+	Bitangent				= normalize(mul(float4(Bitangent, 0), InstanceTransformInverseTransposed_{ParameterName}).xyz);
+	Tangent					= normalize(mul(float4(Tangent, 0), InstanceTransformInverseTransposed_{ParameterName}).xyz);
+}
+
+void GetTriangleColor_{ParameterName}(int Triangle, float3 BaryCoord, out float4 Color)
+{
+	int Index[3];
+	GetTriangleIndices_{ParameterName}(Triangle, Index[0], Index[1], Index[2]);
+
+	float4 Colors[3];
+	GetVertexColor_{ParameterName}(Index[0], Colors[0]);
+	GetVertexColor_{ParameterName}(Index[1], Colors[1]);
+	GetVertexColor_{ParameterName}(Index[2], Colors[2]);
+
+	Color = Colors[0] * BaryCoord.x + Colors[1] * BaryCoord.y + Colors[2] * BaryCoord.z;
+}
+
+void GetTriangleUV_{ParameterName}(int Triangle, float3 BaryCoord, int UVSet, out float2 UV)
+{
+	int Index[3];
+	GetTriangleIndices_{ParameterName}(Triangle, Index[0], Index[1], Index[2]);
+
+	float2 UVs[3];
+	GetVertexUV_{ParameterName}(Index[0], UVSet, UVs[0]);
+	GetVertexUV_{ParameterName}(Index[1], UVSet, UVs[1]);
+	GetVertexUV_{ParameterName}(Index[2], UVSet, UVs[2]);
+
+	UV = UVs[0] * BaryCoord.x + UVs[1] * BaryCoord.y + UVs[2] * BaryCoord.z;
+}
+
+//////////////////////////////////////////////////////////////////////////
+// Socket Functions
+void GetSocketCount_{ParameterName}(out int Count)
+{
+	Count = SocketCounts_{ParameterName}.x;
+}
+
+void GetFilteredSocketCount_{ParameterName}(out int Count)
+{
+	Count = SocketCounts_{ParameterName}.y;
+}
+void GetUnfilteredSocketCount_{ParameterName}(out int Count)
+{
+	Count = SocketCounts_{ParameterName}.z;
+}
+
+void GetSocketTransform_{ParameterName}(int Socket, out float3 Position, out float4 Rotation, out float3 Scale)
+{
+	[branch]
+	if ( SocketCounts_{ParameterName}.x > 0 )
+	{
+		Socket = clamp(Socket, 0, SocketCounts_{ParameterName}.x - 1);
+		Position = SocketTransforms_{ParameterName}[Socket * 3 + 0].xyz;
+		Rotation = SocketTransforms_{ParameterName}[Socket * 3 + 1];
+		Scale    = SocketTransforms_{ParameterName}[Socket * 3 + 2].xyz;
+	}
+	else
+	{
+		Position = float3(0, 0, 0);
+		Rotation = float4(0, 0, 0, 1);
+		Scale = float3(1, 1, 1);
+	}
+}
+
+void GetSocketTransformWS_{ParameterName}(int Socket, out float3 Position, out float4 Rotation, out float3 Scale)
+{
+	GetSocketTransform_{ParameterName}(Socket, Position, Rotation, Scale);
+	Position = mul(float4(Position, 1.0), InstanceTransform_{ParameterName}).xyz;
+	Rotation = NiagaraGPU_QuatMul(InstanceRotation_{ParameterName}, Rotation);
+	Scale = mul(float4(Scale, 0.0), InstanceTransform_{ParameterName}).xyz;
+}
+
+void GetFilteredSocketTransform_{ParameterName}(int SocketIndex, out float3 Position, out float4 Rotation, out float3 Scale)
+{
+	[branch]
+	if ( SocketCounts_{ParameterName}.y > 0 )
+	{
+		SocketIndex = clamp(SocketIndex, 0, SocketCounts_{ParameterName}.y - 1);
+		int Socket = FilteredAndUnfilteredSockets_{ParameterName}[SocketIndex];
+		GetSocketTransform_{ParameterName}(Socket, Position, Rotation, Scale);
+	}
+	else
+	{
+		Position = float3(0, 0, 0);
+		Rotation = float4(0, 0, 0, 1);
+		Scale = float3(1, 1, 1);
+	}
+}
+
+void GetFilteredSocketTransformWS_{ParameterName}(int SocketIndex, out float3 Position, out float4 Rotation, out float3 Scale)
+{
+	GetFilteredSocketTransform_{ParameterName}(SocketIndex, Position, Rotation, Scale);
+	Position = mul(float4(Position, 1.0), InstanceTransform_{ParameterName}).xyz;
+	Rotation = NiagaraGPU_QuatMul(InstanceRotation_{ParameterName}, Rotation);
+	Scale = mul(float4(Scale, 0.0), InstanceTransform_{ParameterName}).xyz;
+}
+
+void GetUnfilteredSocketTransform_{ParameterName}(int SocketIndex, out float3 Position, out float4 Rotation, out float3 Scale)
+{
+	[branch]
+	if ( SocketCounts_{ParameterName}.z > 0 )
+	{
+		SocketIndex = SocketCounts_{ParameterName}.y + clamp(SocketIndex, 0, SocketCounts_{ParameterName}.z - 1);
+		int Socket = FilteredAndUnfilteredSockets_{ParameterName}[SocketIndex];
+		GetSocketTransform_{ParameterName}(Socket, Position, Rotation, Scale);
+	}
+	else
+	{
+		Position = float3(0, 0, 0);
+		Rotation = float4(0, 0, 0, 1);
+		Scale = float3(1, 1, 1);
+	}
+}
+
+void GetUnfilteredSocketTransformWS_{ParameterName}(int SocketIndex, out float3 Position, out float4 Rotation, out float3 Scale)
+{
+	GetUnfilteredSocketTransform_{ParameterName}(SocketIndex, Position, Rotation, Scale);
+	Position = mul(float4(Position, 1.0), InstanceTransform_{ParameterName}).xyz;
+	Rotation = NiagaraGPU_QuatMul(InstanceRotation_{ParameterName}, Rotation);
+	Scale = mul(float4(Scale, 0.0), InstanceTransform_{ParameterName}).xyz;
+}
+
+//////////////////////////////////////////////////////////////////////////
+// Section functions
+void IsValidSection_{ParameterName}(int Section, out bool IsValid)
+{
+	IsValid = Section >= 0 && Section < SectionCounts_{ParameterName}.x;
+}
+
+void GetSectionTriangleCount_{ParameterName}(int Section, out int Count)
+{
+	Count = Section >= 0 && Section < SectionCounts_{ParameterName}.x ? SectionInfos_{ParameterName}[Section].y : 0;
+}
+
+void RandomSectionTriangle_{ParameterName}(NiagaraRandInfo RandInfo, int Section, out int Triangle, out float3 BaryCoord)
+{
+	int SectionTriangleCount;
+	GetSectionTriangleCount_{ParameterName}(Section, SectionTriangleCount);
+
+	if (SectionTriangleCount > 0)
+	{
+		int SectionTriangleOffset = SectionInfos_{ParameterName}[Section].x;
+		Triangle = NiagaraRandomInt(RandInfo, SectionTriangleCount) + SectionTriangleOffset;
+
+		[branch]
+		if ( HasUniformSampling_{ParameterName} )
+		{
+			Triangle = UniformTriangle_{ParameterName}(RandInfo, Triangle, SectionTriangleOffset);
+		}
+	}
+	BaryCoord = NiagaraRandomBaryCoord(RandInfo);
+}
+
+void GetSectionTriangleAt_{ParameterName}(int Section, int SectionTriangleIndex, out int Triangle)
+{
+	int SectionTriangleCount;
+	GetSectionTriangleCount_{ParameterName}(Section, SectionTriangleCount);
+
+	if ( SectionTriangleCount > 0 )
+	{
+		int SectionTriangleMax = SectionInfos_{ParameterName}[Section].y - 1;
+		int SectionTriangleOffset = SectionInfos_{ParameterName}[Section].x;
+		Triangle = clamp(SectionTriangleIndex, 0, SectionTriangleMax) + SectionTriangleOffset;
+	}
+	else
+	{
+		Triangle = 0;
+	}
+}
+
+void GetFilteredSectionAt_{ParameterName}(int SectionIndex, out int Section)
+{
+	int FilteredSectionCount = SectionCounts_{ParameterName}.y;
+	if ( FilteredSectionCount > 0 )
+	{
+		SectionIndex = clamp(SectionIndex, 0, FilteredSectionCount - 1);
+		Section = FilteredAndUnfilteredSections_{ParameterName}[SectionIndex];
+	}
+	else
+	{
+		Section = 0;
+	}
+}
+
+void GetUnfilteredSectionAt_{ParameterName}(int SectionIndex, out int Section)
+{
+	int FilteredSectionCount = SectionCounts_{ParameterName}.y;
+	int UnfilteredSectionCount = SectionCounts_{ParameterName}.z;
+	if ( UnfilteredSectionCount > 0 )
+	{
+		SectionIndex = clamp(SectionIndex, 0, UnfilteredSectionCount - 1);
+		Section = FilteredAndUnfilteredSections_{ParameterName}[SectionIndex + FilteredSectionCount];
+	}
+	else
+	{
+		Section = 0;
+	}
+}
+
+void GetSectionCount_{ParameterName}(out int Count)
+{
+	Count = SectionCounts_{ParameterName}.x;
+}
+
+void GetFilteredSectionCount_{ParameterName}(out int Count)
+{
+	Count = SectionCounts_{ParameterName}.y;
+}
+
+void GetUnfilteredSectionCount_{ParameterName}(out int Count)
+{
+	Count = SectionCounts_{ParameterName}.z;
+}
+
+void RandomSection_{ParameterName}(NiagaraRandInfo RandInfo, out int Section)
+{
+	int SectionCount = SectionCounts_{ParameterName}.x;
+	if (SectionCount > 0)
+	{
+		Section = NiagaraRandomInt(RandInfo, SectionCount);
+		float Prob = asfloat(SectionInfos_{ParameterName}[Section].z);
+		int Alias = SectionInfos_{ParameterName}[Section].w;
+		Section = NiagaraRandomFloat(RandInfo) < Prob ? Section : Alias;
+	}
+	else
+	{
+		Section = 0;
+	}
+}
+
+void RandomFilteredSection_{ParameterName}(NiagaraRandInfo RandInfo, out int Section)
+{
+	int SectionCount = SectionCounts_{ParameterName}.x;
+	int FilteredSectionCount = SectionCounts_{ParameterName}.y;
+	if (FilteredSectionCount > 0)
+	{
+		int SectionIndex = NiagaraRandomInt(RandInfo, FilteredSectionCount);
+		int SectionInfoIndex = SectionIndex + SectionCount;
+		float Prob = asfloat(SectionInfos_{ParameterName}[SectionInfoIndex].z);
+		int Alias = SectionInfos_{ParameterName}[SectionInfoIndex].w;
+		SectionIndex = NiagaraRandomFloat(RandInfo) < Prob ? SectionIndex : Alias;
+		Section = FilteredAndUnfilteredSections_{ParameterName}[SectionIndex];
+	}
+	else
+	{
+		Section = 0;
+	}
+}
+
+void RandomUnfilteredSection_{ParameterName}(NiagaraRandInfo RandInfo, out int Section)
+{
+	int SectionCount = SectionCounts_{ParameterName}.x;
+	int FilteredSectionCount = SectionCounts_{ParameterName}.y;
+	int UnfilteredSectionCount = SectionCounts_{ParameterName}.z;
+	if (UnfilteredSectionCount > 0)
+	{
+		int SectionIndex = NiagaraRandomInt(RandInfo, UnfilteredSectionCount);
+		int SectionInfoIndex = SectionIndex + SectionCount + FilteredSectionCount;
+		float Prob = asfloat(SectionInfos_{ParameterName}[SectionInfoIndex].z);
+		int Alias = SectionInfos_{ParameterName}[SectionInfoIndex].w;
+		SectionIndex = NiagaraRandomFloat(RandInfo) < Prob ? SectionIndex : Alias;
+		Section = FilteredAndUnfilteredSections_{ParameterName}[SectionIndex + FilteredSectionCount];
+	}
+	else
+	{
+		Section = 0;
+	}
+}
+
+//////////////////////////////////////////////////////////////////////////
+// VM Misc Functions
+void IsValid_{ParameterName}(out bool IsValid)
+{
+	IsValid = NumVertices_{ParameterName} > 0;
+}
+
+void GetLocalToWorld_{ParameterName}(out float4x4 Matrix)
+{
+	Matrix = InstanceTransform_{ParameterName};
+}
+
+void GetLocalToWorldInverseTransposed_{ParameterName}(out float4x4 Matrix)
+{
+	Matrix = InstanceTransformInverseTransposed_{ParameterName};
+}
+
+void GetWorldVelocity_{ParameterName}(out float3 Velocity)
+{
+	Velocity = InstanceWorldVelocity_{ParameterName};
+}
diff --git a/Engine/Plugins/FX/Niagara/Shaders/Private/NiagaraEmitterInstanceShader.usf b/Engine/Plugins/FX/Niagara/Shaders/Private/NiagaraEmitterInstanceShader.usf
--- a/Engine/Plugins/FX/Niagara/Shaders/Private/NiagaraEmitterInstanceShader.usf
+++ b/Engine/Plugins/FX/Niagara/Shaders/Private/NiagaraEmitterInstanceShader.usf
@@ -785,6 +785,17 @@ struct NiagaraRandInfo
 };
 
 #if GPU_SIMULATION
+NiagaraRandInfo MakeRandInfo()
+{
+	GRandomSeedOffset += 1664525u;
+
+	NiagaraRandInfo RandInfo;
+	RandInfo.Seed1 = GLinearThreadId;
+	RandInfo.Seed2 = GRandomSeedOffset;
+	RandInfo.Seed3 = GEmitterTickCounter;
+	return RandInfo;
+}
+
 float NiagaraRandomFloat(NiagaraRandInfo RandInfo)
 {
 	return uint(RandInfo.Seed3) == 0xffffffff ? NiagaraInternalNoise(uint(RandInfo.Seed1), uint(RandInfo.Seed2), uint(RandInfo.Seed3)) : rand(1.0f, uint(RandInfo.Seed1), uint(RandInfo.Seed2), uint(RandInfo.Seed3));
diff --git a/Engine/Plugins/FX/Niagara/Source/Niagara/Classes/NiagaraDataInterface.h b/Engine/Plugins/FX/Niagara/Source/Niagara/Classes/NiagaraDataInterface.h
--- a/Engine/Plugins/FX/Niagara/Source/Niagara/Classes/NiagaraDataInterface.h
+++ b/Engine/Plugins/FX/Niagara/Source/Niagara/Classes/NiagaraDataInterface.h
@@ -568,27 +568,33 @@ struct FNDIParameter<FNiagaraRandInfo>
 	}
 };
 
-struct FNDIRandomHelper
+// Completely random policy which will pull from the contexts random stream
+struct FNDIRandomStreamPolicy
 {
-	FNDIRandomHelper(FVectorVMExternalFunctionContext& InContext)
-		: Context(InContext)
-		, RandParam(Context)
-	{
+	FNDIRandomStreamPolicy(FVectorVMExternalFunctionContext& InContext) : Context(InContext) {}
 
-	}
+	FORCEINLINE void GetAndAdvance() {}
+	FORCEINLINE bool IsDeterministic() const { return false; }
+	FORCEINLINE_DEBUGGABLE FVector4 Rand4(int32 InstanceIndex) const { return FVector4(Context.GetRandStream().GetFraction(), Context.GetRandStream().GetFraction(), Context.GetRandStream().GetFraction(), Context.GetRandStream().GetFraction()); }
+	FORCEINLINE_DEBUGGABLE FVector Rand3(int32 InstanceIndex) const { return FVector(Context.GetRandStream().GetFraction(), Context.GetRandStream().GetFraction(), Context.GetRandStream().GetFraction()); }
+	FORCEINLINE_DEBUGGABLE FVector2D Rand2(int32 InstanceIndex) const { return FVector2D(Context.GetRandStream().GetFraction(), Context.GetRandStream().GetFraction()); }
+	FORCEINLINE_DEBUGGABLE float Rand(int32 InstanceIndex) const { return Context.GetRandStream().GetFraction(); }
 
-	FORCEINLINE void GetAndAdvance()
-	{
-		RandParam.GetAndAdvance(RandInfo);
-	}
+	FVectorVMExternalFunctionContext& Context;
+};
 
-	FORCEINLINE bool IsDeterministic() const
-	{
-		return RandInfo.Seed3 != INDEX_NONE;
-	}
+// Random policy which can be optionally deterministic depending on the info
+struct FNDIRandomInfoPolicy
+{
+	FNDIRandomInfoPolicy(FVectorVMExternalFunctionContext& InContext)
+		: Context(InContext)
+		, RandParam(InContext)
+	{}
+
+	FORCEINLINE void GetAndAdvance() { RandParam.GetAndAdvance(RandInfo); }
+	FORCEINLINE bool IsDeterministic() const { return RandInfo.Seed3 != INDEX_NONE; }
 
 	//////////////////////////////////////////////////////////////////////////
-	
 	FORCEINLINE_DEBUGGABLE FVector4 Rand4(int32 InstanceIndex) const
 	{
 		if (IsDeterministic())
@@ -597,14 +603,14 @@ struct FNDIRandomHelper
 
 			FIntVector4 v = FIntVector4(RandomCounter, RandInfo.Seed1, RandInfo.Seed2, RandInfo.Seed3) * 1664525 + FIntVector4(1013904223);
 
-			v.X += v.Y*v.W;
-			v.Y += v.Z*v.X;
-			v.Z += v.X*v.Y;
-			v.W += v.Y*v.Z;
-			v.X += v.Y*v.W;
-			v.Y += v.Z*v.X;
-			v.Z += v.X*v.Y;
-			v.W += v.Y*v.Z;
+			v.X += v.Y * v.W;
+			v.Y += v.Z * v.X;
+			v.Z += v.X * v.Y;
+			v.W += v.Y * v.Z;
+			v.X += v.Y * v.W;
+			v.Y += v.Z * v.X;
+			v.Z += v.X * v.Y;
+			v.W += v.Y * v.Z;
 
 			// NOTE(mv): We can use 24 bits of randomness, as all integers in [0, 2^24] 
 			//           are exactly representable in single precision floats.
@@ -628,12 +634,12 @@ struct FNDIRandomHelper
 
 			FIntVector v = FIntVector(RandInfo.Seed1, RandInfo.Seed2, RandomCounter | (RandInfo.Seed3 << 16)) * 1664525 + FIntVector(1013904223);
 
-			v.X += v.Y*v.Z;
-			v.Y += v.Z*v.X;
-			v.Z += v.X*v.Y;
-			v.X += v.Y*v.Z;
-			v.Y += v.Z*v.X;
-			v.Z += v.X*v.Y;
+			v.X += v.Y * v.Z;
+			v.Y += v.Z * v.X;
+			v.Z += v.X * v.Y;
+			v.X += v.Y * v.Z;
+			v.Y += v.Z * v.X;
+			v.Z += v.X * v.Y;
 
 			return FVector((v >> 8) & 0x00ffffff) / 16777216.0; // 0x01000000 == 16777216
 		}
@@ -668,28 +674,41 @@ struct FNDIRandomHelper
 		}
 	}
 
+	FVectorVMExternalFunctionContext& Context;
+	FNDIParameter<FNiagaraRandInfo> RandParam;
+	FNiagaraRandInfo RandInfo;
+};
+
+template<typename TRandomPolicy>
+struct TNDIRandomHelper : public TRandomPolicy
+{
+	TNDIRandomHelper(FVectorVMExternalFunctionContext& InContext)
+		: TRandomPolicy(InContext)
+	{
+	}
+
 	FORCEINLINE_DEBUGGABLE FVector4 RandRange(int32 InstanceIndex, FVector4 Min, FVector4 Max) const
 	{
 		FVector4 Range = Max - Min;
-		return Min + (Rand(InstanceIndex) * Range);
+		return Min + (TRandomPolicy::Rand(InstanceIndex) * Range);
 	}
 
 	FORCEINLINE_DEBUGGABLE FVector RandRange(int32 InstanceIndex, FVector Min, FVector Max) const
 	{
 		FVector Range = Max - Min;
-		return Min + (Rand(InstanceIndex) * Range);
+		return Min + (TRandomPolicy::Rand(InstanceIndex) * Range);
 	}
 
 	FORCEINLINE_DEBUGGABLE FVector2D RandRange(int32 InstanceIndex, FVector2D Min, FVector2D Max) const
 	{
 		FVector2D Range = Max - Min;
-		return Min + (Rand(InstanceIndex) * Range);
+		return Min + (TRandomPolicy::Rand(InstanceIndex) * Range);
 	}
 
 	FORCEINLINE_DEBUGGABLE float RandRange(int32 InstanceIndex, float Min, float Max) const
 	{
 		float Range = Max - Min;
-		return Min + (Rand(InstanceIndex) * Range);
+		return Min + (TRandomPolicy::Rand(InstanceIndex) * Range);
 	}
 
 	FORCEINLINE_DEBUGGABLE int32 RandRange(int32 InstanceIndex, int32 Min, int32 Max) const
@@ -698,25 +717,23 @@ struct FNDIRandomHelper
 		//       numbers than using %.
 		// NOTE: Inclusive! So [0, x] instead of [0, x)
 		int32 Range = Max - Min;
-		return Min + (int(Rand(InstanceIndex) * (Range + 1)));
+		return Min + (int(TRandomPolicy::Rand(InstanceIndex) * (Range + 1)));
 	}
 
 	FORCEINLINE_DEBUGGABLE FVector RandomBarycentricCoord(int32 InstanceIndex) const
 	{
 		//TODO: This is gonna be slooooow. Move to an LUT possibly or find faster method.
 		//Can probably handle lower quality randoms / uniformity for a decent speed win.
-		FVector2D r = Rand2(InstanceIndex);
+		FVector2D r = TRandomPolicy::Rand2(InstanceIndex);
 		float sqrt0 = FMath::Sqrt(r.X);
 		float sqrt1 = FMath::Sqrt(r.Y);
 		return FVector(1.0f - sqrt0, sqrt0 * (1.0 - r.Y), r.Y * sqrt0);
 	}
-
-	FVectorVMExternalFunctionContext& Context;
-	FNDIParameter<FNiagaraRandInfo> RandParam;
-
-	FNiagaraRandInfo RandInfo;
 };
 
+using FNDIRandomHelper = TNDIRandomHelper<FNDIRandomInfoPolicy>;
+using FNDIRandomHelperFromStream = TNDIRandomHelper<FNDIRandomStreamPolicy>;
+
 //Helper to deal with types with potentially several input registers.
 template<typename T>
 struct FNDIInputParam
diff --git a/Engine/Plugins/FX/Niagara/Source/Niagara/Classes/NiagaraDataInterfaceStaticMesh.h b/Engine/Plugins/FX/Niagara/Source/Niagara/Classes/NiagaraDataInterfaceStaticMesh.h
deleted file mode 100644
diff --git a/Engine/Plugins/FX/Niagara/Source/Niagara/Private/DataInterface/NiagaraDataInterfaceStaticMesh.cpp b/Engine/Plugins/FX/Niagara/Source/Niagara/Private/DataInterface/NiagaraDataInterfaceStaticMesh.cpp
new file mode 100644
--- /dev/null
+++ b/Engine/Plugins/FX/Niagara/Source/Niagara/Private/DataInterface/NiagaraDataInterfaceStaticMesh.cpp
@@ -0,0 +1,3770 @@
+// Copyright Epic Games, Inc. All Rights Reserved.
+
+#include "NiagaraDataInterfaceStaticMesh.h"
+#include "NiagaraEmitterInstance.h"
+#include "NiagaraComponent.h"
+#include "NiagaraSystemInstance.h"
+#include "NiagaraRenderer.h"
+#include "NiagaraScript.h"
+#include "ShaderParameterUtils.h"
+#include "NiagaraStats.h"
+
+#include "Engine/StaticMeshActor.h"
+#include "Engine/StaticMeshSocket.h"
+#include "Internationalization/Internationalization.h"
+#include "GPUSkinCache.h"
+#include "ShaderCompilerCore.h"
+
+#define LOCTEXT_NAMESPACE "NiagaraDataInterfaceStaticMesh"
+
+//-OPT: We should be able to share data between data interfaces where possible
+
+namespace NDIStaticMeshLocal
+{
+	static const TCHAR* TemplateShaderFile = TEXT("/Plugin/FX/Niagara/Private/NiagaraDataInterfaceStaticMeshTemplate.ush");
+
+	static int32 GNiagaraFailStaticMeshDataInterface = 0;
+	static FAutoConsoleVariableRef CVarNiagaraFailStaticMeshDataInterface(
+		TEXT("fx.Niagara.FailStaticMeshDataInterface"),
+		GNiagaraFailStaticMeshDataInterface,
+		TEXT("When enabled we will fail out using static mesh data interfaces."),
+		ECVF_Default
+	);
+
+	static int32 GNDIStaticMesh_UseInlineLODsOnly = 1;
+	static FAutoConsoleVariableRef CVarNDIStaticMesh_UseInlineLODsOnly(
+		TEXT("fx.Niagara.NDIStaticMesh.UseInlineLODsOnly"),
+		GNDIStaticMesh_UseInlineLODsOnly,
+		TEXT("When enabled Niagara will never use streaming LOD levels, only inline LODs."),
+		ECVF_Default
+	);
+
+	//////////////////////////////////////////////////////////////////////////
+	struct EDIFunctionVersion
+	{
+		enum Type
+		{
+			InitialVersion = 0,
+			RefactoredV1 = 1,
+
+			VersionPlusOne,
+			LatestVersion = VersionPlusOne - 1
+		};
+	};
+
+	//////////////////////////////////////////////////////////////////////////
+	// Vertex Sampling
+	static const FName	IsValidVertexName("IsValidVertex");
+	static const FName	RandomVertexName("RandomVertex");
+	static const FName	GetVertexCountName("GetVertexCount");
+
+	static const FName	GetVertexName("GetVertex");
+	static const FName	GetVertexWSName("GetVertexWS");
+	static const FName	GetVertexColorName("GetVertexColor");
+	static const FName	GetVertexUVName("GetVertexUV");
+
+	//////////////////////////////////////////////////////////////////////////
+	// VM Triangle Sampling
+	static const FName	IsValidTriangleName("IsValidTriangle");
+	static const FName	RandomTriangleName("RandomTriangle");
+	static const FName	GetTriangleCountName("GetTriangleCount");
+
+	static const FName	IsValidFilteredTriangleName("IsValidFilteredTriangle");
+	static const FName	RandomFilteredTriangleName("RandomFilteredTriangle");
+	static const FName	GetFilteredTriangleCountName("GetFilteredTriangleCount");
+	static const FName	GetFilteredTriangleAtName("GetFilteredTriangleAt");
+
+	static const FName	IsValidUnfilteredTriangleName("IsValidUnfilteredTriangle");
+	static const FName	RandomUnfilteredTriangleName("RandomUnfilteredTriangle");
+	static const FName	GetUnfilteredTriangleCountName("GetUnfilteredTriangleCount");
+	static const FName	GetUnfilteredTriangleAtName("GetUnfilteredTriangleAt");
+
+	static const FName	GetTriangleName("GetTriangle");
+	static const FName	GetTriangleWSName("GetTriangleWS");
+	static const FName	GetTriangleColorName("GetTriangleColor");
+	static const FName	GetTriangleUVName("GetTriangleUV");
+
+	static const FName	GetTriangleIndicesName("GetTriangleIndices");
+
+	//////////////////////////////////////////////////////////////////////////
+	// Socket Functions
+	static const FName	GetSocketCountName("GetSocketCount");
+	static const FName	GetFilteredSocketCountName("GetFilteredSocketCount");
+	static const FName	GetUnfilteredSocketCountName("GetUnfilteredSocketCount");
+	static const FName	GetSocketTransformName("GetSocketTransform");
+	static const FName	GetSocketTransformWSName("GetSocketTransformWS");
+	static const FName	GetFilteredSocketTransformName("GetFilteredSocketTransform");
+	static const FName	GetFilteredSocketTransformWSName("GetFilteredSocketTransformWS");
+	static const FName	GetUnfilteredSocketTransformName("GetUnfilteredSocketTransform");
+	static const FName	GetUnfilteredSocketTransformWSName("GetUnfilteredSocketTransformWS");
+
+	//////////////////////////////////////////////////////////////////////////
+	// Section functions
+	static const FName	IsValidSectionName("IsValidSection");
+	static const FName	GetSectionTriangleCountName("GetSectionTriangleCount");
+	static const FName	RandomSectionTriangleName("RandomSectionTriangle");
+	static const FName	GetSectionTriangleAtName("GetSectionTriangleAt");
+
+	static const FName	GetFilteredSectionAtName("GetFilteredSectionAt");
+	static const FName	GetUnfilteredSectionAtName("GetUnfilteredSectionAt");
+
+	static const FName	GetSectionCountName("GetSectionCount");
+	static const FName	GetFilteredSectionCountName("GetFilteredSectionCount");
+	static const FName	GetUnfilteredSectionCountName("GetUnfilteredSectionCount");
+
+	static const FName	RandomSectionName("RandomSection");
+	static const FName	RandomFilteredSectionName("RandomFilteredSection");
+	static const FName	RandomUnfilteredSectionName("RandomUnfilteredSection");
+
+	//////////////////////////////////////////////////////////////////////////
+	// VM Misc Functions
+	static const FName	IsValidName("IsValid");
+
+	static const FName	GetLocalToWorldName("GetLocalToWorld");
+	static const FName	GetLocalToWorldInverseTransposedName("GetLocalToWorldInverseTransposed");
+	static const FName	GetWorldVelocityName("GetWorldVelocity");
+
+	//////////////////////////////////////////////////////////////////////////
+	// Deprecated Functions
+	static const FName Deprecated_RandomSectionName("DeprecatedRandomSection");
+	static const FName Deprecated_RandomTriCoordName("DeprecatedRandomTriCoord");
+	static const FName Deprecated_RandomTriCoordOnSectionName("DeprecatedRandomTriCoordOnSection");
+	static const FName Deprecated_RandomTriCoordVCFilteredName("DeprecatedRandomTriCoordUsingVertexColorFilter");
+	static const FName Deprecated_GetVertexPositionName("DeprecatedGetVertexPosition");
+	static const FName Deprecated_GetVertexPositionWSName("DeprecatedGetVertexPositionWS");
+	static const FName Deprecated_GetTriPositionName("DeprecatedGetTriPosition");
+	static const FName Deprecated_GetTriNormalName("DeprecatedGetTriNormal");
+	static const FName Deprecated_GetTriTangentsName("DeprecatedGetTriTangents");
+	static const FName Deprecated_GetTriPositionWSName("DeprecatedGetTriPositionWS");
+	static const FName Deprecated_GetTriNormalWSName("DeprecatedGetTriNormalWS");
+	static const FName Deprecated_GetTriTangentsWSName("DeprecatedGetTriTangentsWS");
+	static const FName Deprecated_GetTriColorName("DeprecatedGetTriColor");
+	static const FName Deprecated_GetTriUVName("DeprecatedGetTriUV");
+	static const FName Deprecated_GetTriPositionAndVelocityWSName("DeprecatedGetTriPositionAndVelocityWS");
+	//static const FName Deprecated_GetMeshLocalToWorldInverseTransposedName("DeprecatedGetMeshLocalToWorldInverseTransposed");
+
+	//////////////////////////////////////////////////////////////////////////
+	// Shader Parameters
+	static FString	NumTriangles_String("NumTriangles_");
+	static FString	NumVertices_String("NumVertices_");
+	static FString	NumUVs_String("NumUVs_");
+	static FString	HasColors_String("HasColors_");
+	static FString	IndexBuffer_String("IndexBuffer_");
+	static FString	PositionBuffer_String("PositionBuffer_");
+	static FString	TangentBuffer_String("TangentBuffer_");
+	static FString	UVBuffer_String("UVBuffer_");
+	static FString	ColorBuffer_String("ColorBuffer_");
+
+	static FString	HasUniformSampling_String("HasUniformSampling_");
+	static FString	UniformSamplingTriangles_String("UniformSamplingTriangles_");
+
+	static FString	SectionCounts_String("SectionCounts_");
+	static FString	SectionInfos_String("SectionInfos_");
+	static FString	FilteredAndUnfilteredSections_String("FilteredAndUnfilteredSections_");
+
+	static FString	SocketCounts_String("SocketCounts_");
+	static FString	SocketTransforms_String("SocketTransforms_");
+	static FString	FilteredAndUnfilteredSockets_String("FilteredAndUnfilteredSockets_");
+
+	static FString	InvDeltaSeconds_String("InvDeltaSeconds_");
+	static FString	InstanceTransform_String("InstanceTransform_");
+	static FString	InstanceTransformInverseTransposed_String("InstanceTransformInverseTransposed_");
+	static FString	InstanceRotation_String("InstanceRotation_");
+	static FString	InstancePreviousTransform_String("InstancePreviousTransform_");
+	static FString	InstancePreviousTransformInverseTransposed_String("InstancePreviousTransformInverseTransposed_");
+	static FString	InstancePreviousRotation_String("InstancePreviousRotation_");
+	static FString	InstanceWorldVelocity_String("InstanceWorldVelocity_");
+
+	//////////////////////////////////////////////////////////////////////////
+	struct FNDISectionInfo
+	{
+		uint32 FirstTriangle = 0;
+		uint32 NumTriangles = 0;
+		float  Prob = 0.0f;
+		uint32 Alias = 0;
+	};
+
+	class FNDISectionAreaWeightedSampler : public FWeightedRandomSampler
+	{
+	public:
+		static void Init(const FStaticMeshLODResources* LODResource, TConstArrayView<int32> Sections, const FStaticMeshSectionAreaWeightedTriangleSamplerArray* MeshSectionSamplers, TArray<FNDISectionInfo>& OutSectionInfos)
+		{
+			if (MeshSectionSamplers)
+			{
+				FNDISectionAreaWeightedSampler Sampler;
+				Sampler.Init(Sections, *MeshSectionSamplers);
+				for (int32 i=0; i < Sections.Num(); ++i)
+				{
+					const int32 SectionIndex	= Sections[i];
+					FNDISectionInfo& SectionInfo= OutSectionInfos.AddDefaulted_GetRef();
+					SectionInfo.FirstTriangle	= LODResource->Sections[SectionIndex].FirstIndex / 3;
+					SectionInfo.NumTriangles	= LODResource->Sections[SectionIndex].NumTriangles;
+					SectionInfo.Prob			= Sampler.GetProb()[i];
+					SectionInfo.Alias			= Sampler.GetAlias()[i];
+				}
+			}
+			else
+			{
+				for (int32 i=0; i < Sections.Num(); ++i)
+				{
+					const int32 SectionIndex	= Sections[i];
+					FNDISectionInfo& SectionInfo= OutSectionInfos.AddDefaulted_GetRef();
+					SectionInfo.FirstTriangle	= LODResource->Sections[SectionIndex].FirstIndex / 3;
+					SectionInfo.NumTriangles	= LODResource->Sections[SectionIndex].NumTriangles;
+					SectionInfo.Prob			= 1.0f;
+					SectionInfo.Alias			= i;
+				}
+			}
+		}
+
+		void Init(TConstArrayView<int32> Sections, const FStaticMeshSectionAreaWeightedTriangleSamplerArray& MeshSectionSamplers)
+		{
+			TempTotalWeight = 0.0f;
+			TempWeights.Empty();
+
+			for (const int32 SectionIndex : Sections)
+			{
+				const FStaticMeshSectionAreaWeightedTriangleSampler& SectionSampler = MeshSectionSamplers[SectionIndex];
+				const float SectionWeight = SectionSampler.GetTotalWeight();
+				TempTotalWeight += SectionWeight;
+				TempWeights.Add(SectionWeight);
+			}
+
+			FWeightedRandomSampler::Initialize();
+		}
+
+	protected:
+		virtual float GetWeights(TArray<float>& OutWeights) override
+		{
+			OutWeights = MoveTemp(TempWeights);
+			return TempTotalWeight;
+		}
+
+		float			TempTotalWeight = 0.0f;
+		TArray<float>	TempWeights;
+	};
+
+	//////////////////////////////////////////////////////////////////////////
+	struct FGpuInitializeData
+	{
+		struct FRenderProxy*				RenderProxy = nullptr;
+		FNiagaraSystemInstanceID			SystemInstanceID;
+		const FStaticMeshLODResources*		LODResource = nullptr;
+
+		int32								NumFilteredTriangles = 0;
+		int32								NumUnfilteredTriangles = 0;
+
+		bool								bGpuUniformDistribution = false;
+
+		int32								NumSections = 0;
+		int32								NumFilteredSections = 0;
+		int32								NumUnfilteredSections = 0;
+		TResourceArray<FIntVector4>			SectionInfos;
+		TResourceArray<uint16>				FilteredAndUnfilteredSections;
+
+		int32								NumSockets = 0;
+		int32								NumFilteredSockets = 0;
+		int32								NumUnfilteredSockets = 0;
+		TResourceArray<FVector4f>			SocketTransforms;
+		TResourceArray<uint16>				FilteredAndUnfilteredSockets;
+	};
+
+	//////////////////////////////////////////////////////////////////////////
+	struct FInstanceData_RenderThread
+	{
+		bool		bIsValid = false;
+		bool		bGpuUniformDistribution = false;
+		FMatrix44f	Transform = FMatrix44f::Identity;
+		FMatrix44f	PrevTransform = FMatrix44f::Identity;
+		FQuat4f		Rotation = FQuat4f::Identity;
+		FQuat4f		PrevRotation = FQuat4f::Identity;
+		float		DeltaSeconds = 0.0f;
+
+		FIntVector NumTriangles = FIntVector::ZeroValue;
+		int32 NumVertices = 0;
+		int32 NumUVs = 0;
+		FShaderResourceViewRHIRef MeshIndexBufferSRV;
+		FShaderResourceViewRHIRef MeshPositionBufferSRV;
+		FShaderResourceViewRHIRef MeshTangentBufferSRV;
+		FShaderResourceViewRHIRef MeshUVBufferSRV;
+		FShaderResourceViewRHIRef MeshColorBufferSRV;
+
+		FShaderResourceViewRHIRef MeshUniformSamplingTriangleSRV;
+
+		FIntVector	SectionCounts = FIntVector::ZeroValue;	// X=NumSections, y = Num Filtered, Z=NumUnfiltered
+		FReadBuffer SectionInfos;
+		FReadBuffer FilteredAndUnfilteredSections;
+
+		FIntVector	SocketCounts = FIntVector::ZeroValue;	// X=NumSockets, Y=NumFiltered, Z=NumUnfiltered
+		FReadBuffer	SocketTransforms;
+		FReadBuffer	FilteredAndUnfilteredSockets;
+
+	#if STATS
+		int64 GPUMemoryUsage = 0;
+	#endif
+
+		~FInstanceData_RenderThread()
+		{
+			Release();
+		}
+
+		void Init(FGpuInitializeData& GpuInitializeData)
+		{
+			// Gather mesh buffers
+			if (GpuInitializeData.LODResource)
+			{
+				MeshIndexBufferSRV = RHICreateShaderResourceView(GpuInitializeData.LODResource->IndexBuffer.IndexBufferRHI);
+				MeshPositionBufferSRV = GpuInitializeData.LODResource->VertexBuffers.PositionVertexBuffer.GetSRV();
+				MeshTangentBufferSRV = GpuInitializeData.LODResource->VertexBuffers.StaticMeshVertexBuffer.GetTangentsSRV();
+				MeshUVBufferSRV = GpuInitializeData.LODResource->VertexBuffers.StaticMeshVertexBuffer.GetTexCoordsSRV();
+				MeshColorBufferSRV = GpuInitializeData.LODResource->VertexBuffers.ColorVertexBuffer.GetColorComponentsSRV();
+
+				NumTriangles.X = GpuInitializeData.LODResource->IndexBuffer.GetNumIndices() / 3;
+				NumTriangles.Y = GpuInitializeData.NumFilteredTriangles;
+				NumTriangles.Z = GpuInitializeData.NumUnfilteredTriangles;
+				NumVertices = GpuInitializeData.LODResource->VertexBuffers.PositionVertexBuffer.GetNumVertices();
+				NumUVs = MeshUVBufferSRV.IsValid() ? GpuInitializeData.LODResource->VertexBuffers.StaticMeshVertexBuffer.GetNumTexCoords() : 0;
+
+				bIsValid = MeshIndexBufferSRV.IsValid() && MeshPositionBufferSRV.IsValid() && MeshTangentBufferSRV.IsValid();
+
+				// Build uniform sampling buffer
+				MeshUniformSamplingTriangleSRV = GpuInitializeData.bGpuUniformDistribution ? GpuInitializeData.LODResource->AreaWeightedSectionSamplersBuffer.GetBufferSRV() : nullptr;
+				bGpuUniformDistribution = GpuInitializeData.bGpuUniformDistribution && MeshUniformSamplingTriangleSRV.IsValid();
+
+				// Build section buffers
+				SectionCounts.X = GpuInitializeData.NumSections;
+				SectionCounts.Y = GpuInitializeData.NumFilteredSections;
+				SectionCounts.Z = GpuInitializeData.NumUnfilteredSections;
+				if (GpuInitializeData.SectionInfos.Num() > 0)
+				{
+					SectionInfos.Initialize(TEXT("NDISkelMesh_SectionInfos"), sizeof(FIntVector4), GpuInitializeData.SectionInfos.Num(), EPixelFormat::PF_R32G32B32A32_UINT, BUF_Static, &GpuInitializeData.SectionInfos);
+				}
+				if (GpuInitializeData.FilteredAndUnfilteredSections.Num() > 0)
+				{
+					FilteredAndUnfilteredSections.Initialize(TEXT("NDISkelMesh_FilteredAndUnfilteredSections"), sizeof(uint16), GpuInitializeData.FilteredAndUnfilteredSections.Num(), EPixelFormat::PF_R16_UINT, BUF_Static, &GpuInitializeData.FilteredAndUnfilteredSections);
+				}
+			}
+			else
+			{
+				NumTriangles = FIntVector::ZeroValue;
+				NumVertices = 0;
+				NumUVs = 0;
+
+				bIsValid = false;
+
+				bGpuUniformDistribution = false;
+
+				SectionCounts = FIntVector::ZeroValue;
+			}
+
+			// Build socket buffers
+			SocketCounts.X = GpuInitializeData.NumSockets;
+			SocketCounts.Y = GpuInitializeData.NumFilteredSockets;
+			SocketCounts.Z = GpuInitializeData.NumUnfilteredSockets;
+			if ( GpuInitializeData.SocketTransforms.Num() > 0 )
+			{
+				SocketTransforms.Initialize(TEXT("NDISkelMesh_SocketTransforms"), sizeof(FVector4f), GpuInitializeData.SocketTransforms.Num(), EPixelFormat::PF_A32B32G32R32F, BUF_Static, &GpuInitializeData.SocketTransforms);
+			}
+			if ( GpuInitializeData.FilteredAndUnfilteredSockets.Num() > 0 )
+			{
+				FilteredAndUnfilteredSockets.Initialize(TEXT("NDISkelMesh_FilteredAndUnfilteredSockets"), sizeof(uint16), GpuInitializeData.FilteredAndUnfilteredSockets.Num(), EPixelFormat::PF_R16_UINT, BUF_Static, &GpuInitializeData.FilteredAndUnfilteredSockets);
+			}
+		#if STATS
+			GPUMemoryUsage = SectionInfos.NumBytes + FilteredAndUnfilteredSections.NumBytes + SocketTransforms.NumBytes + FilteredAndUnfilteredSockets.NumBytes;
+			INC_MEMORY_STAT_BY(STAT_NiagaraGPUDataInterfaceMemory, GPUMemoryUsage);
+		#endif
+		}
+
+		void Release()
+		{
+			bIsValid = false;
+
+		#if STATS
+			INC_MEMORY_STAT_BY(STAT_NiagaraGPUDataInterfaceMemory, GPUMemoryUsage);
+			GPUMemoryUsage = 0;
+		#endif
+
+			MeshIndexBufferSRV.SafeRelease();
+			MeshPositionBufferSRV.SafeRelease();
+			MeshTangentBufferSRV.SafeRelease();
+			MeshUVBufferSRV.SafeRelease();
+			MeshColorBufferSRV.SafeRelease();
+
+			SectionInfos.Release();
+			FilteredAndUnfilteredSections.Release();
+
+			SocketTransforms.Release();
+			FilteredAndUnfilteredSockets.Release();
+		}
+	};
+
+	struct FInstanceData_FromGameThread
+	{
+		FMatrix44f	Transform = FMatrix44f::Identity;
+		FMatrix44f	PrevTransform = FMatrix44f::Identity;
+		FQuat4f		Rotation = FQuat4f::Identity;
+		FQuat4f		PrevRotation = FQuat4f::Identity;
+		float		DeltaSeconds = 0.0f;
+	};
+
+	struct FInstanceData_GameThread
+	{
+		/** Cached ptr to StaticMeshComponent we sample from, when found. Otherwise, the SceneComponent to use to transform the Default or Preview mesh. */
+		TWeakObjectPtr<USceneComponent> SceneComponentWeakPtr;
+
+		/** Cached ptr to the mesh so that we can make sure that we haven't been deleted. */
+		TWeakObjectPtr<UStaticMesh> StaticMeshWeakPtr;
+
+		/** Cached ComponentToWorld. (Falls back to WorldTransform of the system instance) */
+		FMatrix Transform;
+		/** InverseTranspose of above for transforming normals/tangents. */
+		FMatrix TransformInverseTransposed;
+
+		/** Cached ComponentToWorld from previous tick. */
+		FMatrix PrevTransform;
+		/** Cached Previous InverseTranspose of above for transforming normals/tangents. */
+		FMatrix PrevTransformInverseTransposed;
+
+		/** Cached Rotation. */
+		FQuat4f Rotation;
+		/** Cached Previous Rotation. */
+		FQuat4f PrevRotation;
+
+		/** Time separating Transform and PrevTransform. */
+		float DeltaSeconds;
+
+		/** Velocity set by the physics body of the mesh component */
+		FVector PhysicsVelocity;
+		/** True if velocity should not be calculated via the transforms, but rather read the physics data from the mesh component */
+		uint32 bUsePhysicsVelocity : 1;
+
+		/** True if SceneComponent was valid on initialization (used to track invalidation of the component on tick) */
+		uint32 bComponentValid : 1;
+
+		/** True if StaticMesh was valid on initialization (used to track invalidation of the mesh on tick) */
+		uint32 bMeshValid : 1;
+
+		/** True if the mesh allows CPU access. Use to reset the instance in the editor*/
+		uint32 bMeshAllowsCpuAccess : 1;
+		/** True if the mesh we're using allows area weighted sampling on CPU. */
+		uint32 bIsCpuUniformlyDistributedSampling : 1;
+		/** True if the mesh we're using allows area weighted sampling on GPU. */
+		uint32 bIsGpuUniformlyDistributedSampling : 1;
+
+		/** Filtered sections followed by unfiltered sections. */
+		int32 NumFilteredSections = 0;
+		int32 NumUnfilteredSections = 0;
+		int32 NumFilteredTriangles = 0;
+		int32 NumUnfilteredTriangles = 0;
+		TArray<int32> FilteredAndUnfilteredSections;
+		TArray<FNDISectionInfo> FilteredAndUnfilteredSectionInfos;
+
+		///** Allows sampling of the mesh's tris based on a dynamic color range. */
+		//TSharedPtr<struct FDynamicVertexColorFilterData> DynamicVertexColorSampler;
+
+		/** Cached change id off of the data interface.*/
+		uint32 ChangeId = 0;
+
+		/** The MinLOD, see UStaticMesh::MinLOD which is platform specific.*/
+		int32 MinLOD = 0;
+		/** The cached LODIdx used to initialize the FNDIStaticMesh_InstanceData.*/
+		int32 CachedLODIdx = 0;
+
+		/** Cached socket information, if available */
+		TArray<FTransform3f> CachedSockets;
+
+		/** Number of filtered sockets. */
+		int32 NumFilteredSockets = 0;
+
+		/** Filter sockets followed by unfiltered sockets */
+		TArray<uint16> FilteredAndUnfilteredSockets;
+
+#if WITH_EDITOR
+		FDelegateHandle OnMeshChanged;
+		FDelegateHandle OnPostMeshBuild;
+		FDelegateHandle OnMeshReimported;
+
+		void OnEditorMeshChanged(UObject* OptionalMesh)
+		{
+			UStaticMesh* StaticMesh = StaticMeshWeakPtr.Get();
+			if (!OptionalMesh || (OptionalMesh == StaticMesh))
+			{
+				StaticMeshWeakPtr.Reset();
+				if (StaticMesh)
+				{
+					StaticMesh->GetOnMeshChanged().Remove(OnMeshChanged);
+					StaticMesh->OnPostMeshBuild().Remove(OnPostMeshBuild);
+				}
+				if (UImportSubsystem* ImportSubsystem = GEditor->GetEditorSubsystem<UImportSubsystem>())
+				{
+					ImportSubsystem->OnAssetReimport.Remove(OnMeshReimported);
+				}
+			}
+			OnMeshChanged.Reset();
+			OnPostMeshBuild.Reset();
+			OnMeshReimported.Reset();
+		}
+#endif
+
+		bool Init(UNiagaraDataInterfaceStaticMesh* Interface, FNiagaraSystemInstance* SystemInstance)
+		{
+			check(SystemInstance);
+
+			// Initialize members
+			SceneComponentWeakPtr = nullptr;
+			StaticMeshWeakPtr = nullptr;
+			Transform = FMatrix::Identity;
+			TransformInverseTransposed = FMatrix::Identity;
+			PrevTransform = FMatrix::Identity;
+			PrevTransformInverseTransposed = FMatrix::Identity;
+			Rotation = FQuat4f::Identity;
+			PrevRotation = FQuat4f::Identity;
+			DeltaSeconds = 0.0f;
+			PhysicsVelocity = FVector::ZeroVector;
+			bUsePhysicsVelocity = Interface->bUsePhysicsBodyVelocity;
+			bComponentValid = false;
+			bMeshValid = false;
+			bMeshAllowsCpuAccess = false;
+			bIsCpuUniformlyDistributedSampling = false;
+			bIsGpuUniformlyDistributedSampling = false;
+			NumFilteredSections = 0;
+			NumUnfilteredSections = 0;
+			FilteredAndUnfilteredSections.Empty();
+			FilteredAndUnfilteredSectionInfos.Empty();
+			NumFilteredTriangles = 0;
+			NumUnfilteredTriangles = 0;
+			ChangeId = Interface->ChangeId;
+			MinLOD = 0;
+			CachedLODIdx = 0;
+			CachedSockets.Empty();
+			NumFilteredSockets = 0;
+			FilteredAndUnfilteredSockets.Empty();
+
+			// Get component / mesh we are using
+			USceneComponent* SceneComponent = nullptr;
+			UStaticMesh* StaticMesh = Interface->GetStaticMesh(SceneComponent, SystemInstance);
+			SceneComponentWeakPtr = SceneComponent;
+
+			// Gather attached information
+			bComponentValid = SceneComponent != nullptr;
+			const FTransform& ComponentTransform = bComponentValid ? SceneComponent->GetComponentToWorld() : SystemInstance->GetWorldTransform();
+
+			Transform = ComponentTransform.ToMatrixWithScale();
+			TransformInverseTransposed = ComponentTransform.Inverse().ToMatrixWithScale().GetTransposed();
+			PrevTransform = Transform;
+			PrevTransformInverseTransposed = ComponentTransform.Inverse().ToMatrixWithScale().GetTransposed();
+
+			Rotation = FQuat4f(ComponentTransform.GetRotation());
+			PrevRotation = Rotation;
+
+			if (bUsePhysicsVelocity)
+			{
+				if (UStaticMeshComponent* MeshComponent = Cast<UStaticMeshComponent>(SceneComponent))
+				{
+					PhysicsVelocity = MeshComponent->GetPhysicsLinearVelocity();
+				}
+			}
+
+			if (GNiagaraFailStaticMeshDataInterface != 0)
+			{
+				UE_LOG(LogNiagara, Log, TEXT("StaticMesh data interface aborting because \"fx.Niagara.FailStaticMeshDataInterface 1\". Failed InitPerInstanceData - %s"), *Interface->GetFullName());
+				return false;
+			}
+
+			// Report missing or inaccessible meshes to the log
+			if (StaticMesh == nullptr)
+			{
+				UE_LOG(LogNiagara, Log, TEXT("StaticMesh data interface has no valid mesh - %s"), *Interface->GetFullName());
+			}
+			else if (!StaticMesh->bAllowCPUAccess)
+			{
+				//-TODO: This will be modified with the refactor to allow GPU to be used when CPU is also used
+				if (Interface->IsUsedWithGPUEmitter() && !Interface->IsUsedWithCPUEmitter())
+				{
+					if (!FNiagaraUtilities::AreBufferSRVsAlwaysCreated(GMaxRHIShaderPlatform))
+					{
+						UE_LOG(LogNiagara, Log, TEXT("NiagaraStaticMeshDataInterface used by GPU emitter but does not have SRV access on this platform.  Enable CPU access to fix this issue. Interface: %s, Mesh: %s"), *GetFullNameSafe(Interface), *GetFullNameSafe(StaticMesh));
+						StaticMesh = nullptr;
+					}
+				}
+				else
+				{
+					UE_LOG(LogNiagara, Log, TEXT("NiagaraStaticMeshDataInterface used by CPU emitter and does not allow CPU access. Interface: %s, Mesh: %s"), *GetFullNameSafe(Interface), *GetFullNameSafe(StaticMesh));
+					StaticMesh = nullptr;
+				}
+			}
+
+			TRefCountPtr<const FStaticMeshLODResources> LODData;
+			if (StaticMesh != nullptr)
+			{
+				// Check if any valid LODs are found. If not, we won't use this mesh
+				MinLOD = StaticMesh->GetMinLOD().GetValue();
+				if (GNDIStaticMesh_UseInlineLODsOnly)
+				{
+					MinLOD = StaticMesh->GetNumLODs() - StaticMesh->GetRenderData()->NumInlinedLODs;
+				}
+
+				FStaticMeshRenderData* RenderData = StaticMesh->GetRenderData();
+				if (RenderData)
+				{
+					CachedLODIdx = RenderData->GetCurrentFirstLODIdx(MinLOD);
+					if (RenderData->LODResources.IsValidIndex(CachedLODIdx))
+					{
+						LODData = &RenderData->LODResources[CachedLODIdx];
+					}
+				}
+
+				if (!LODData.IsValid())
+				{
+					StaticMesh = nullptr;
+				}
+			}
+
+			StaticMeshWeakPtr = StaticMesh;
+			bMeshValid = StaticMesh != nullptr;
+
+			if (StaticMesh != nullptr)
+			{
+#if WITH_EDITOR
+				if ( GIsEditor )
+				{
+					OnMeshChanged = StaticMesh->GetOnMeshChanged().AddLambda([this]() { OnEditorMeshChanged(nullptr); });
+					OnPostMeshBuild = StaticMesh->OnPostMeshBuild().AddLambda([this](UStaticMesh* StaticMesh) { OnEditorMeshChanged(StaticMesh); });
+					if (UImportSubsystem* ImportSubsystem = GEditor->GetEditorSubsystem<UImportSubsystem>())
+					{
+						OnMeshReimported = ImportSubsystem->OnAssetReimport.AddLambda([this](UObject* Object) { OnEditorMeshChanged(Object); });
+					}
+				}
+#endif
+
+				bMeshAllowsCpuAccess = StaticMesh->bAllowCPUAccess;
+				bIsCpuUniformlyDistributedSampling = StaticMesh->bSupportUniformlyDistributedSampling;
+				if (bIsCpuUniformlyDistributedSampling)
+				{
+					bIsCpuUniformlyDistributedSampling &= LODData->AreaWeightedSampler.GetNumEntries() == LODData->Sections.Num();
+					bIsCpuUniformlyDistributedSampling &= LODData->AreaWeightedSectionSamplers.Num() == LODData->Sections.Num();
+				}
+				bIsGpuUniformlyDistributedSampling = bIsCpuUniformlyDistributedSampling && StaticMesh->bSupportGpuUniformlyDistributedSampling;
+
+				// Init section filter
+				//-OPT: I don't foresee a lot of entries here, could potentially optimize this later to avoid insert
+				FilteredAndUnfilteredSections.Reserve(LODData->Sections.Num());
+				FilteredAndUnfilteredSectionInfos.Reserve(LODData->Sections.Num());
+				for (int32 i = 0; i < LODData->Sections.Num(); ++i)
+				{
+					// Skip garbage data
+					if ( LODData->Sections[i].NumTriangles == 0 )
+					{
+						continue;
+					}
+
+					if (Interface->SectionFilter.AllowedMaterialSlots.Num() == 0 || Interface->SectionFilter.AllowedMaterialSlots.Contains(LODData->Sections[i].MaterialIndex))
+					{
+						FilteredAndUnfilteredSections.Insert(i, NumFilteredSections);
+						++NumFilteredSections;
+						NumFilteredTriangles += LODData->Sections[i].NumTriangles;
+					}
+					else
+					{
+						FilteredAndUnfilteredSections.Add(i);
+						++NumUnfilteredSections;
+						NumUnfilteredTriangles += LODData->Sections[i].NumTriangles;
+					}
+				}
+
+				// Initialize section information
+				{
+					const FStaticMeshSectionAreaWeightedTriangleSamplerArray* SectionSamplerArray = bIsCpuUniformlyDistributedSampling ? &LODData->AreaWeightedSectionSamplers : nullptr;
+					FNDISectionAreaWeightedSampler::Init(LODData, GetFilteredSections(), SectionSamplerArray, FilteredAndUnfilteredSectionInfos);
+					FNDISectionAreaWeightedSampler::Init(LODData, GetUnfilteredSections(), SectionSamplerArray, FilteredAndUnfilteredSectionInfos);
+				}
+
+				// Init socket information
+				const int32 NumMeshSockets = StaticMesh->Sockets.Num();
+				if (NumMeshSockets > 0)
+				{
+					if (NumMeshSockets > TNumericLimits<uint16>::Max())
+					{
+						UE_LOG(LogNiagara, Warning, TEXT("Static Mesh Data Interface '%s' requires more sockets '%d' than we currently support '%d' Mesh '%s' Component '%s'"), *GetFullNameSafe(Interface), NumMeshSockets, TNumericLimits<uint16>::Max(), *GetFullNameSafe(StaticMesh), *GetFullNameSafe(SceneComponent));
+						return false;
+					}
+
+					CachedSockets.AddDefaulted(NumMeshSockets);
+					for (int32 i = 0; i < NumMeshSockets; ++i)
+					{
+						CachedSockets[i].SetTranslation(StaticMesh->Sockets[i]->RelativeLocation);
+						CachedSockets[i].SetRotation(FQuat4f(StaticMesh->Sockets[i]->RelativeRotation));
+						CachedSockets[i].SetScale3D(StaticMesh->Sockets[i]->RelativeScale);
+					}
+
+					NumFilteredSockets = 0;
+					if (Interface->FilteredSockets.Num() > 0)
+					{
+						FilteredAndUnfilteredSockets.Reserve(NumMeshSockets);
+						for (int32 i = 0; i < Interface->FilteredSockets.Num(); ++i)
+						{
+							const FName FilteredSocketName = Interface->FilteredSockets[i];
+							const int32 SocketIndex = StaticMesh->Sockets.IndexOfByPredicate([&](const UStaticMeshSocket* Socket) { return Socket->SocketName == FilteredSocketName; });
+							if (SocketIndex == INDEX_NONE)
+							{
+								UE_LOG(LogNiagara, Warning, TEXT("Static Mesh Data Interface '%s' could not find socket '%s' Mesh '%s' Component '%s'"), *GetFullNameSafe(Interface), *FilteredSocketName.ToString(), *GetFullNameSafe(StaticMesh), *GetFullNameSafe(SceneComponent));
+								continue;
+							}
+							++NumFilteredSockets;
+							FilteredAndUnfilteredSockets.Add(uint16(SocketIndex));
+						}
+
+						for (int32 i = 0; i < NumMeshSockets; ++i)
+						{
+							FilteredAndUnfilteredSockets.AddUnique(uint16(i));
+						}
+					}
+				}
+			}
+
+			return true;
+		}
+
+		bool Tick(UNiagaraDataInterfaceStaticMesh* Interface, FNiagaraSystemInstance* SystemInstance, float InDeltaSeconds)
+		{
+			if (ResetRequired(Interface))
+			{
+				return true;
+			}
+
+			DeltaSeconds = InDeltaSeconds;
+
+			USceneComponent* SceneComponent = SceneComponentWeakPtr.Get();
+			const FTransform& ComponentTransform = SceneComponent != nullptr ? SceneComponent->GetComponentToWorld() : SystemInstance->GetWorldTransform();
+
+			PrevTransform = Transform;
+			PrevTransformInverseTransposed = TransformInverseTransposed;
+			Transform = ComponentTransform.ToMatrixWithScale();
+			TransformInverseTransposed = ComponentTransform.Inverse().ToMatrixWithScale().GetTransposed();
+
+			PrevRotation = Rotation;
+			Rotation = FQuat4f(ComponentTransform.GetRotation());
+
+			if (bUsePhysicsVelocity)
+			{
+				if (UStaticMeshComponent* MeshComponent = Cast<UStaticMeshComponent>(SceneComponent))
+				{
+					PhysicsVelocity = MeshComponent->GetPhysicsLinearVelocity();
+				}
+				else
+				{
+					PhysicsVelocity = FVector::ZeroVector;
+				}
+			}
+
+			return false;
+		}
+
+		void Release()
+		{
+#if WITH_EDITOR
+			if (GIsEditor)
+			{
+				// We reuse the function to clear the delegates out
+				OnEditorMeshChanged(nullptr);
+			}
+#endif
+		}
+
+		bool ResetRequired(UNiagaraDataInterfaceStaticMesh* Interface) const
+		{
+			USceneComponent* Component = SceneComponentWeakPtr.Get();
+			if (bComponentValid && !Component)
+			{
+				// The component we were bound to is no longer valid so we have to trigger a reset.
+				return true;
+			}
+
+			UStaticMesh* Mesh = StaticMeshWeakPtr.Get();
+			if (bMeshValid)
+			{
+				if (!Mesh)
+				{
+					// The static mesh we were bound to is no longer valid so we have to trigger a reset.
+					return true;
+				}
+				else if (UStaticMeshComponent* StaticMeshComp = Cast<UStaticMeshComponent>(Component))
+				{
+					if (Mesh != StaticMeshComp->GetStaticMesh())
+					{
+						// The mesh changed on the component we're attached to so we have to reset
+						return true;
+					}
+				}
+			}
+
+			if (Interface != nullptr && ChangeId != Interface->ChangeId)
+			{
+				return true;
+			}
+
+			if (Mesh != nullptr)
+			{
+				// Currently we only reset if the cached LOD was streamed out, to avoid performance hits. To revisit.
+				// We could probably just recache the data derived from the LOD instead of resetting everything.
+				if (Mesh->GetRenderData()->GetCurrentFirstLODIdx(MinLOD) > CachedLODIdx)
+				{
+					return true;
+				}
+			}
+
+			return false;
+		}
+
+		const FStaticMeshLODResources* GetCurrentFirstLOD()
+		{
+			//-OPT: Perhaps we could cache this during the tick function?
+			UStaticMesh* StaticMesh = StaticMeshWeakPtr.Get();
+			if ( bMeshValid && StaticMesh )
+			{
+				if ( const FStaticMeshLODResources* LODResource = StaticMesh->GetRenderData()->GetCurrentFirstLOD(MinLOD) )
+				{
+					return LODResource;
+				}
+			}
+			return nullptr;
+		}
+
+		TConstArrayView<int32> GetFilteredSections() const { return MakeArrayView(FilteredAndUnfilteredSections.GetData(), NumFilteredSections); }
+		TConstArrayView<int32> GetUnfilteredSections() const { return MakeArrayView(FilteredAndUnfilteredSections.GetData() + NumFilteredSections, NumUnfilteredSections); }
+
+		TConstArrayView<FNDISectionInfo> GetFilteredSectionInfos() const { return MakeArrayView(FilteredAndUnfilteredSectionInfos.GetData(), NumFilteredSections); }
+		TConstArrayView<FNDISectionInfo> GetUnfilteredSectionInfos() const { return MakeArrayView(FilteredAndUnfilteredSectionInfos.GetData() + NumFilteredSections, NumUnfilteredSections); }
+	};
+
+	struct FRenderProxy : public FNiagaraDataInterfaceProxy
+	{
+		FRenderProxy() {}
+
+		virtual void ConsumePerInstanceDataFromGameThread(void* FromGameThreadData, const FNiagaraSystemInstanceID& InstanceID) override
+		{
+			FInstanceData_FromGameThread* FromGameThread = reinterpret_cast<FInstanceData_FromGameThread*>(FromGameThreadData);
+			FInstanceData_RenderThread* InstanceData = &PerInstanceData_RT.FindChecked(InstanceID);
+
+			InstanceData->Transform		= FromGameThread->Transform;
+			InstanceData->PrevTransform	= FromGameThread->PrevTransform;
+			InstanceData->Rotation		= FromGameThread->Rotation;
+			InstanceData->PrevRotation	= FromGameThread->PrevRotation;
+			InstanceData->DeltaSeconds	= FromGameThread->DeltaSeconds;
+		}
+
+		virtual int32 PerInstanceDataPassedToRenderThreadSize() const override
+		{
+			return sizeof(FInstanceData_FromGameThread);
+		}
+
+		TMap<FNiagaraSystemInstanceID, FInstanceData_RenderThread> PerInstanceData_RT;
+	};
+
+	//////////////////////////////////////////////////////////////////////////
+
+	template<typename TTransformHandler = FNDITransformHandlerNoop>
+	struct FStaticMeshCpuHelper
+	{
+		FORCEINLINE FStaticMeshCpuHelper(FVectorVMExternalFunctionContext& Context)
+			: InstanceData(Context)
+		{
+			LODResource = InstanceData->GetCurrentFirstLOD();
+		}
+
+
+		FORCEINLINE float GetInvDeltaSeconds() const
+		{
+			return InstanceData->DeltaSeconds > 0.0f ? 1.0f / InstanceData->DeltaSeconds : 0.0f;
+		}
+
+		FORCEINLINE FVector3f TransformPosition(FVector3f Position) const
+		{
+			TransformHandler.TransformPosition(Position, InstanceData->Transform);
+			return Position;
+		}
+
+		FORCEINLINE FVector3f TransformVector(FVector3f Vector) const
+		{
+			TransformHandler.TransformVector(Vector, InstanceData->TransformInverseTransposed);
+			return Vector;
+		}
+
+		FORCEINLINE FQuat4f TransformRotation(FQuat4f Rotation) const
+		{
+			TransformHandler.TransformRotation(Rotation, InstanceData->Rotation);
+			return Rotation;
+		}
+
+		FORCEINLINE FVector3f PreviousTransformPosition(FVector3f Position) const
+		{
+			TransformHandler.TransformPosition(Position, InstanceData->PrevTransform);
+			return Position;
+		}
+
+		FORCEINLINE FVector3f PreviousTransformVector(FVector3f Vector) const
+		{
+			TransformHandler.TransformVector(Vector, InstanceData->PrevTransformInverseTransposed);
+			return Vector;
+		}
+
+		FORCEINLINE FQuat4f PreviousTransformRotation(FQuat4f Rotation) const
+		{
+			TransformHandler.TransformRotation(Rotation, InstanceData->PrevRotation);
+			return Rotation;
+		}
+
+		FORCEINLINE FIndexArrayView GetIndexArrayView() const
+		{
+			return LODResource ? LODResource->IndexBuffer.GetArrayView() : FIndexArrayView();
+		}
+
+		FORCEINLINE int32 GetNumPositionVertices() const
+		{
+			return LODResource && LODResource->VertexBuffers.PositionVertexBuffer.GetVertexData() && LODResource->VertexBuffers.PositionVertexBuffer.GetAllowCPUAccess() ? LODResource->VertexBuffers.PositionVertexBuffer.GetNumVertices() : 0;
+		}
+
+		FORCEINLINE int32 GetNumTangentVertices() const
+		{
+			return LODResource && LODResource->VertexBuffers.StaticMeshVertexBuffer.GetTangentData() && LODResource->VertexBuffers.StaticMeshVertexBuffer.GetAllowCPUAccess() ? LODResource->VertexBuffers.StaticMeshVertexBuffer.GetNumVertices() : 0;
+		}
+
+		uint32 GetNumUVs() const
+		{
+			return LODResource && LODResource->VertexBuffers.StaticMeshVertexBuffer.GetTexCoordData() && LODResource->VertexBuffers.StaticMeshVertexBuffer.GetAllowCPUAccess() ? LODResource->VertexBuffers.StaticMeshVertexBuffer.GetNumTexCoords() : 0;
+		}
+
+		uint32 GetNumUVVertices() const
+		{
+			return LODResource && LODResource->VertexBuffers.StaticMeshVertexBuffer.GetTexCoordData() && LODResource->VertexBuffers.StaticMeshVertexBuffer.GetAllowCPUAccess() ? LODResource->VertexBuffers.StaticMeshVertexBuffer.GetNumVertices() : 0;
+		}
+
+		FORCEINLINE int32 GetNumColorVertices() const
+		{
+			return LODResource && LODResource->VertexBuffers.ColorVertexBuffer.GetVertexData() && LODResource->VertexBuffers.ColorVertexBuffer.GetAllowCPUAccess() ? LODResource->VertexBuffers.ColorVertexBuffer.GetNumVertices() : 0;
+		}
+
+		FORCEINLINE FVector3f GetLocalTrianglePosition(const FVector3f& BaryCoord, int32 Index0, int32 Index1, int32 Index2) const
+		{
+			FVector3f Position;
+			Position = LODResource->VertexBuffers.PositionVertexBuffer.VertexPosition(Index0) * BaryCoord.X;
+			Position += LODResource->VertexBuffers.PositionVertexBuffer.VertexPosition(Index1) * BaryCoord.Y;
+			Position += LODResource->VertexBuffers.PositionVertexBuffer.VertexPosition(Index2) * BaryCoord.Z;
+			return Position;
+		}
+
+		FORCEINLINE FVector3f GetTrianglePosition(const FVector3f& BaryCoord, int32 Index0, int32 Index1, int32 Index2) const
+		{
+			FVector3f Position = GetLocalTrianglePosition(BaryCoord, Index0, Index1, Index2);
+			TransformHandler.TransformPosition(Position, InstanceData->Transform);
+			return Position;
+		}
+
+		FORCEINLINE FVector3f GetTriangleTangentX(const FVector3f& BaryCoord, int32 Index0, int32 Index1, int32 Index2) const
+		{
+			FVector3f Tangent;
+			Tangent  = LODResource->VertexBuffers.StaticMeshVertexBuffer.VertexTangentX(Index0) * BaryCoord.X;
+			Tangent += LODResource->VertexBuffers.StaticMeshVertexBuffer.VertexTangentX(Index1) * BaryCoord.Y;
+			Tangent += LODResource->VertexBuffers.StaticMeshVertexBuffer.VertexTangentX(Index2) * BaryCoord.Z;
+			TransformHandler.TransformVector(Tangent, InstanceData->Transform);
+			return Tangent;
+		}
+
+		FORCEINLINE FVector3f GetTriangleTangentY(const FVector3f& BaryCoord, int32 Index0, int32 Index1, int32 Index2) const
+		{
+			FVector3f Tangent;
+			Tangent  = LODResource->VertexBuffers.StaticMeshVertexBuffer.VertexTangentY(Index0) * BaryCoord.X;
+			Tangent += LODResource->VertexBuffers.StaticMeshVertexBuffer.VertexTangentY(Index1) * BaryCoord.Y;
+			Tangent += LODResource->VertexBuffers.StaticMeshVertexBuffer.VertexTangentY(Index2) * BaryCoord.Z;
+			TransformHandler.TransformVector(Tangent, InstanceData->Transform);
+			return Tangent;
+		}
+
+		FORCEINLINE FVector3f GetTriangleTangentZ(const FVector3f& BaryCoord, int32 Index0, int32 Index1, int32 Index2) const
+		{
+			FVector3f Tangent;
+			Tangent  = LODResource->VertexBuffers.StaticMeshVertexBuffer.VertexTangentZ(Index0) * BaryCoord.X;
+			Tangent += LODResource->VertexBuffers.StaticMeshVertexBuffer.VertexTangentZ(Index1) * BaryCoord.Y;
+			Tangent += LODResource->VertexBuffers.StaticMeshVertexBuffer.VertexTangentZ(Index2) * BaryCoord.Z;
+			TransformHandler.TransformVector(Tangent, InstanceData->Transform);
+			return Tangent;
+		}
+
+		FORCEINLINE FLinearColor GetTriangleColor(const FVector3f& BaryCoord, int32 Index0, int32 Index1, int32 Index2) const
+		{
+			FLinearColor Color;
+			Color  = FLinearColor(LODResource->VertexBuffers.ColorVertexBuffer.VertexColor(Index0)) * BaryCoord.X;
+			Color += FLinearColor(LODResource->VertexBuffers.ColorVertexBuffer.VertexColor(Index1)) * BaryCoord.Y;
+			Color += FLinearColor(LODResource->VertexBuffers.ColorVertexBuffer.VertexColor(Index2)) * BaryCoord.Z;
+			return Color;
+		}
+
+		FORCEINLINE FVector2D GetTriangleUV(const FVector3f& BaryCoord, int32 Index0, int32 Index1, int32 Index2, int32 UVSet) const
+		{
+			FVector2D UV;
+			UV = LODResource->VertexBuffers.StaticMeshVertexBuffer.GetVertexUV(Index0, UVSet) * BaryCoord.X;
+			UV += LODResource->VertexBuffers.StaticMeshVertexBuffer.GetVertexUV(Index1, UVSet) * BaryCoord.Y;
+			UV += LODResource->VertexBuffers.StaticMeshVertexBuffer.GetVertexUV(Index2, UVSet) * BaryCoord.Z;
+			return UV;
+		}
+
+		FORCEINLINE FVector3f GetPosition(int32 Vertex) const
+		{
+			FVector3f Position = LODResource->VertexBuffers.PositionVertexBuffer.VertexPosition(Vertex);
+			TransformHandler.TransformPosition(Position, InstanceData->Transform);
+			return Position;
+		}
+
+		FORCEINLINE FVector3f GetLocalPosition(int32 Vertex) const
+		{
+			return LODResource->VertexBuffers.PositionVertexBuffer.VertexPosition(Vertex);
+		}
+
+		FORCEINLINE FVector3f GetTangentX(int32 Vertex) const
+		{
+			FVector3f TangentX = LODResource->VertexBuffers.StaticMeshVertexBuffer.VertexTangentX(Vertex);
+			TransformHandler.TransformVector(TangentX, InstanceData->TransformInverseTransposed);
+			return TangentX;
+		}
+
+		FORCEINLINE FVector3f GetTangentY(int32 Vertex) const
+		{
+			FVector3f TangentY = LODResource->VertexBuffers.StaticMeshVertexBuffer.VertexTangentY(Vertex);
+			TransformHandler.TransformVector(TangentY, InstanceData->TransformInverseTransposed);
+			return TangentY;
+		}
+
+		FORCEINLINE FVector3f GetTangentZ(int32 Vertex) const
+		{
+			FVector3f TangentZ = LODResource->VertexBuffers.StaticMeshVertexBuffer.VertexTangentZ(Vertex);
+			TransformHandler.TransformVector(TangentZ, InstanceData->TransformInverseTransposed);
+			return TangentZ;
+		}
+
+		FORCEINLINE FLinearColor GetColor(int32 Vertex) const
+		{
+			return LODResource->VertexBuffers.ColorVertexBuffer.VertexColor(Vertex);
+		}
+
+		FORCEINLINE FVector2D GetUV(int32 Vertex, int32 UVSet) const
+		{
+			return LODResource->VertexBuffers.StaticMeshVertexBuffer.GetVertexUV(Vertex, UVSet);
+		}
+
+		FORCEINLINE bool IsCpuUniformlyDistributedSampling() const
+		{
+			return LODResource && InstanceData->bIsCpuUniformlyDistributedSampling;
+		}
+
+		template<typename TRandomHelper>
+		FORCEINLINE int32 RandomUniformTriangle(const TRandomHelper& RandHelper, int32 InstanceIndex) const
+		{
+			int32 Triangle = 0;
+			const int32 Section = LODResource->AreaWeightedSampler.GetNumEntries() ? LODResource->AreaWeightedSampler.GetEntryIndex(RandHelper.Rand(InstanceIndex), RandHelper.Rand(InstanceIndex)) : 0;
+			if ( LODResource->AreaWeightedSectionSamplers.IsValidIndex(Section) && LODResource->AreaWeightedSectionSamplers[Section].GetNumEntries() )
+			{
+				const int32 SectionTriangle = LODResource->AreaWeightedSectionSamplers[Section].GetEntryIndex(RandHelper.Rand(InstanceIndex), RandHelper.Rand(InstanceIndex));
+				Triangle = (LODResource->Sections[Section].FirstIndex / 3) + SectionTriangle;
+			}
+			return Triangle;
+		}
+
+		template<typename TRandomHelper>
+		FORCEINLINE int32 RandomUniformTriangle(const TRandomHelper& RandHelper, int32 InstanceIndex, const TConstArrayView<FNDISectionInfo>& SectionInfos, TConstArrayView<int32>& SectionRemap) const
+		{
+			int32 Triangle = 0;
+			int32 SectionIndex = RandHelper.Rand(InstanceIndex) * SectionInfos.Num();
+			SectionIndex = RandHelper.Rand(InstanceIndex) < SectionInfos[SectionIndex].Prob ? SectionIndex : SectionInfos[SectionIndex].Alias;
+			const int32 Section = SectionRemap[SectionIndex];
+			if ( LODResource->AreaWeightedSectionSamplers.IsValidIndex(Section) && LODResource->AreaWeightedSectionSamplers[Section].GetNumEntries() )
+			{
+				const int32 SectionTriangle = LODResource->AreaWeightedSectionSamplers[Section].GetEntryIndex(RandHelper.Rand(InstanceIndex), RandHelper.Rand(InstanceIndex));
+				Triangle = SectionInfos[SectionIndex].FirstTriangle + SectionTriangle;
+			}
+			return Triangle;
+		}
+
+		template<typename TRandomHelper>
+		FORCEINLINE int32 RandomTriangle(const TRandomHelper& RandHelper, int32 InstanceIndex, const TConstArrayView<FNDISectionInfo>& SectionInfos) const
+		{
+			const int32 Section = RandHelper.Rand(InstanceIndex) * SectionInfos.Num();
+			const int32 Triangle = SectionInfos[Section].FirstTriangle + (RandHelper.Rand(InstanceIndex) * SectionInfos[Section].NumTriangles);
+			return Triangle;
+		}
+
+		template<typename TRandomHelper>
+		FORCEINLINE int32 RandomUniformSection(const TRandomHelper& RandHelper, int32 InstanceIndex) const
+		{
+			return LODResource->AreaWeightedSampler.GetEntryIndex(RandHelper.Rand(InstanceIndex), RandHelper.Rand(InstanceIndex));
+		}
+
+		template<typename TRandomHelper>
+		FORCEINLINE int32 RandomUniformSection(const TRandomHelper& RandHelper, int32 InstanceIndex, const TConstArrayView<FNDISectionInfo>& SectionInfos, TConstArrayView<int32>& SectionRemap) const
+		{
+			int32 SectionIndex = RandHelper.Rand(InstanceIndex) * SectionInfos.Num();
+			SectionIndex = RandHelper.Rand(InstanceIndex) < SectionInfos[SectionIndex].Prob ? SectionIndex : SectionInfos[SectionIndex].Alias;
+			return SectionRemap[SectionIndex];
+		}
+
+		template<typename TRandomHelper>
+		FORCEINLINE int32 RandomSectionUniformTriangle(const TRandomHelper& RandHelper, int32 InstanceIndex, int32 Section) const
+		{
+			int32 Triangle = 0;
+			if (LODResource->AreaWeightedSectionSamplers.IsValidIndex(Section) && LODResource->AreaWeightedSectionSamplers[Section].GetNumEntries())
+			{
+				const int32 SectionTriangle = LODResource->AreaWeightedSectionSamplers[Section].GetEntryIndex(RandHelper.Rand(InstanceIndex), RandHelper.Rand(InstanceIndex));
+				Triangle = (LODResource->Sections[Section].FirstIndex / 3) + SectionTriangle;
+			}
+			return Triangle;
+		}
+
+		template<typename TRandomHelper>
+		FORCEINLINE int32 RandomSectionTriangle(const TRandomHelper& RandHelper, int32 InstanceIndex, int32 Section) const
+		{
+			const int32 Triangle = (LODResource->Sections[Section].FirstIndex / 3) + (RandHelper.Rand(InstanceIndex) * LODResource->Sections[Section].NumTriangles);
+			return Triangle;
+		}
+
+		FORCEINLINE int32 GetNumTriangles() const
+		{
+			return LODResource ? LODResource->IndexBuffer.GetArrayView().Num() / 3 : 0;
+		}
+
+		FORCEINLINE int32 GetNumFilteredTriangles() const
+		{
+			return LODResource ? InstanceData->NumFilteredTriangles : 0;
+		}
+
+		FORCEINLINE int32 GetNumUnfilteredTriangles() const
+		{
+			return LODResource ? InstanceData->NumUnfilteredTriangles : 0;
+		}
+
+		FORCEINLINE int32 SectionTriangleToTriangle(const TConstArrayView<FNDISectionInfo>& SectionInfos, int32 Index)
+		{
+			int32 SectionFirstTriangle = 0;
+			for ( const FNDISectionInfo& SectionInfo : SectionInfos )
+			{
+				const int32 SectionLastTriangle = SectionFirstTriangle + SectionInfo.NumTriangles;
+				if ( (Index >= SectionFirstTriangle) && (Index < SectionLastTriangle) )
+				{
+					return (Index - SectionFirstTriangle) + SectionInfo.FirstTriangle;
+				}
+				SectionFirstTriangle = SectionLastTriangle;
+			}
+			return 0;
+		}
+
+		VectorVM::FUserPtrHandler<FInstanceData_GameThread> InstanceData;
+		TRefCountPtr<const FStaticMeshLODResources> LODResource;
+		TTransformHandler TransformHandler;
+	};
+}
+
+//////////////////////////////////////////////////////////////////////////
+
+struct FNiagaraDataInterfaceParametersCS_StaticMesh : public FNiagaraDataInterfaceParametersCS
+{
+	DECLARE_TYPE_LAYOUT(FNiagaraDataInterfaceParametersCS_StaticMesh, NonVirtual);
+
+public:
+	void Bind(const FNiagaraDataInterfaceGPUParamInfo& ParameterInfo, const class FShaderParameterMap& ParameterMap)
+	{
+		using namespace NDIStaticMeshLocal;
+
+		NumTrianglesParam.Bind(ParameterMap, *(NumTriangles_String + ParameterInfo.DataInterfaceHLSLSymbol));
+		NumVerticesParam.Bind(ParameterMap, *(NumVertices_String + ParameterInfo.DataInterfaceHLSLSymbol));
+		NumUVsParam.Bind(ParameterMap, *(NumUVs_String + ParameterInfo.DataInterfaceHLSLSymbol));
+		HasColorsParam.Bind(ParameterMap, *(HasColors_String + ParameterInfo.DataInterfaceHLSLSymbol));
+		IndexBufferParam.Bind(ParameterMap, *(IndexBuffer_String + ParameterInfo.DataInterfaceHLSLSymbol));
+		PositionBufferParam.Bind(ParameterMap, *(PositionBuffer_String + ParameterInfo.DataInterfaceHLSLSymbol));
+		TangentBufferParam.Bind(ParameterMap, *(TangentBuffer_String + ParameterInfo.DataInterfaceHLSLSymbol));
+		UVBufferParam.Bind(ParameterMap, *(UVBuffer_String + ParameterInfo.DataInterfaceHLSLSymbol));
+		ColorBufferParam.Bind(ParameterMap, *(ColorBuffer_String + ParameterInfo.DataInterfaceHLSLSymbol));
+
+		HasUniformSamplingParam.Bind(ParameterMap, *(HasUniformSampling_String + ParameterInfo.DataInterfaceHLSLSymbol));
+		UniformSamplingTrianglesParam.Bind(ParameterMap, *(UniformSamplingTriangles_String + ParameterInfo.DataInterfaceHLSLSymbol));
+
+		SectionCountsParam.Bind(ParameterMap, *(SectionCounts_String + ParameterInfo.DataInterfaceHLSLSymbol));
+		SectionInfosParam.Bind(ParameterMap, *(SectionInfos_String + ParameterInfo.DataInterfaceHLSLSymbol));
+		FilteredAndUnfilteredSectionsParam.Bind(ParameterMap, *(FilteredAndUnfilteredSections_String + ParameterInfo.DataInterfaceHLSLSymbol));
+
+		SocketCountsParam.Bind(ParameterMap, *(SocketCounts_String + ParameterInfo.DataInterfaceHLSLSymbol));
+		SocketTransformsParam.Bind(ParameterMap, *(SocketTransforms_String + ParameterInfo.DataInterfaceHLSLSymbol));
+		FilteredAndUnfilteredSocketsParam.Bind(ParameterMap, *(FilteredAndUnfilteredSockets_String + ParameterInfo.DataInterfaceHLSLSymbol));
+
+		InvDeltaSecondsParam.Bind(ParameterMap, *(InvDeltaSeconds_String + ParameterInfo.DataInterfaceHLSLSymbol));
+		InstanceTransformParam.Bind(ParameterMap, *(InstanceTransform_String + ParameterInfo.DataInterfaceHLSLSymbol));
+		InstanceTransformInverseTransposedParam.Bind(ParameterMap, *(InstanceTransformInverseTransposed_String + ParameterInfo.DataInterfaceHLSLSymbol));
+		InstanceRotationParam.Bind(ParameterMap, *(InstanceRotation_String + ParameterInfo.DataInterfaceHLSLSymbol));
+		InstancePreviousTransformParam.Bind(ParameterMap, *(InstancePreviousTransform_String + ParameterInfo.DataInterfaceHLSLSymbol));
+		InstancePreviousTransformInverseTransposedParam.Bind(ParameterMap, *(InstancePreviousTransformInverseTransposed_String + ParameterInfo.DataInterfaceHLSLSymbol));
+		InstancePreviousRotationParam.Bind(ParameterMap, *(InstancePreviousRotation_String + ParameterInfo.DataInterfaceHLSLSymbol));
+		InstanceWorldVelocityParam.Bind(ParameterMap, *(InstanceWorldVelocity_String + ParameterInfo.DataInterfaceHLSLSymbol));
+	}
+
+	void Set(FRHICommandList& RHICmdList, const FNiagaraDataInterfaceSetArgs& Context) const
+	{
+		check(IsInRenderingThread());
+
+		FRHIComputeShader* ComputeShaderRHI = Context.Shader.GetComputeShader();
+
+		NDIStaticMeshLocal::FRenderProxy* InterfaceProxy = static_cast<NDIStaticMeshLocal::FRenderProxy*>(Context.DataInterface);
+		const NDIStaticMeshLocal::FInstanceData_RenderThread& InstanceData = InterfaceProxy->PerInstanceData_RT.FindChecked(Context.SystemInstanceID);
+
+		// Set mesh sampling data
+		if ( InstanceData.bIsValid )
+		{
+			const bool bHasColors = InstanceData.MeshColorBufferSRV.IsValid();
+			SetShaderValue(RHICmdList, ComputeShaderRHI, NumTrianglesParam, InstanceData.NumTriangles);
+			SetShaderValue(RHICmdList, ComputeShaderRHI, NumVerticesParam, InstanceData.NumVertices);
+			SetShaderValue(RHICmdList, ComputeShaderRHI, NumUVsParam, InstanceData.NumUVs);
+			SetShaderValue(RHICmdList, ComputeShaderRHI, HasColorsParam, bHasColors ? 1 : 0);
+			SetSRVParameter(RHICmdList, ComputeShaderRHI, IndexBufferParam, InstanceData.MeshIndexBufferSRV);
+			SetSRVParameter(RHICmdList, ComputeShaderRHI, PositionBufferParam, InstanceData.MeshPositionBufferSRV);
+			SetSRVParameter(RHICmdList, ComputeShaderRHI, TangentBufferParam, InstanceData.MeshTangentBufferSRV);
+			SetSRVParameter(RHICmdList, ComputeShaderRHI, UVBufferParam, InstanceData.MeshUVBufferSRV.IsValid() ? InstanceData.MeshUVBufferSRV.GetReference() : FNiagaraRenderer::GetDummyFloat2Buffer());
+			SetSRVParameter(RHICmdList, ComputeShaderRHI, ColorBufferParam, bHasColors ? InstanceData.MeshColorBufferSRV.GetReference() : FNiagaraRenderer::GetDummyWhiteColorBuffer());
+
+			SetShaderValue(RHICmdList, ComputeShaderRHI, HasUniformSamplingParam, InstanceData.bGpuUniformDistribution ? 1 : 0);
+			SetSRVParameter(RHICmdList, ComputeShaderRHI, UniformSamplingTrianglesParam, InstanceData.bGpuUniformDistribution ? InstanceData.MeshUniformSamplingTriangleSRV.GetReference() : FNiagaraRenderer::GetDummyUInt2Buffer());
+		}
+		else
+		{
+			SetShaderValue(RHICmdList, ComputeShaderRHI, NumTrianglesParam, 0);
+			SetShaderValue(RHICmdList, ComputeShaderRHI, NumVerticesParam, 0);
+			SetShaderValue(RHICmdList, ComputeShaderRHI, NumUVsParam, 0);
+			SetShaderValue(RHICmdList, ComputeShaderRHI, HasColorsParam, 0);
+			SetSRVParameter(RHICmdList, ComputeShaderRHI, IndexBufferParam, FNiagaraRenderer::GetDummyUIntBuffer());
+			SetSRVParameter(RHICmdList, ComputeShaderRHI, PositionBufferParam, FNiagaraRenderer::GetDummyFloatBuffer());
+			SetSRVParameter(RHICmdList, ComputeShaderRHI, TangentBufferParam, FNiagaraRenderer::GetDummyFloat4Buffer());
+			SetSRVParameter(RHICmdList, ComputeShaderRHI, UVBufferParam, FNiagaraRenderer::GetDummyFloat2Buffer());
+			SetSRVParameter(RHICmdList, ComputeShaderRHI, ColorBufferParam, FNiagaraRenderer::GetDummyWhiteColorBuffer());
+
+			SetShaderValue(RHICmdList, ComputeShaderRHI, HasUniformSamplingParam, 0);
+			SetSRVParameter(RHICmdList, ComputeShaderRHI, UniformSamplingTrianglesParam, FNiagaraRenderer::GetDummyUInt2Buffer());
+		}
+
+		// Section information
+		SetShaderValue(RHICmdList, ComputeShaderRHI, SectionCountsParam, InstanceData.SectionCounts);
+		SetSRVParameter(RHICmdList, ComputeShaderRHI, SectionInfosParam, InstanceData.SectionInfos.SRV.IsValid() ? InstanceData.SectionInfos.SRV.GetReference() : FNiagaraRenderer::GetDummyUInt4Buffer());
+		SetSRVParameter(RHICmdList, ComputeShaderRHI, FilteredAndUnfilteredSectionsParam, InstanceData.FilteredAndUnfilteredSections.SRV.IsValid() ? InstanceData.FilteredAndUnfilteredSections.SRV.GetReference() : FNiagaraRenderer::GetDummyUIntBuffer());
+
+		// Set socket sampling information
+		SetShaderValue(RHICmdList, ComputeShaderRHI, SocketCountsParam, InstanceData.SocketCounts);
+		SetSRVParameter(RHICmdList, ComputeShaderRHI, SocketTransformsParam, InstanceData.SocketTransforms.SRV.IsValid() ? InstanceData.SocketTransforms.SRV.GetReference() : FNiagaraRenderer::GetDummyFloat4Buffer());
+		SetSRVParameter(RHICmdList, ComputeShaderRHI, FilteredAndUnfilteredSocketsParam, InstanceData.FilteredAndUnfilteredSockets.SRV.IsValid() ? InstanceData.FilteredAndUnfilteredSockets.SRV.GetReference() : FNiagaraRenderer::GetDummyUIntBuffer());
+
+		// Set misc data
+		const float InvDeltaTime = InstanceData.DeltaSeconds > 0.0f ? 1.0f / InstanceData.DeltaSeconds : 0.0f;
+		const FVector3f DeltaPosition = InstanceData.Transform.GetOrigin() - InstanceData.PrevTransform.GetOrigin();
+
+		SetShaderValue(RHICmdList, ComputeShaderRHI, InvDeltaSecondsParam, InvDeltaTime);
+		SetShaderValue(RHICmdList, ComputeShaderRHI, InstanceTransformParam, InstanceData.Transform);
+		SetShaderValue(RHICmdList, ComputeShaderRHI, InstanceTransformInverseTransposedParam, InstanceData.Transform.Inverse().GetTransposed());
+		SetShaderValue(RHICmdList, ComputeShaderRHI, InstanceRotationParam, InstanceData.Rotation);
+		SetShaderValue(RHICmdList, ComputeShaderRHI, InstancePreviousTransformParam, InstanceData.PrevTransform);
+		SetShaderValue(RHICmdList, ComputeShaderRHI, InstancePreviousTransformInverseTransposedParam, InstanceData.PrevTransform.Inverse().GetTransposed());
+		SetShaderValue(RHICmdList, ComputeShaderRHI, InstancePreviousRotationParam, InstanceData.PrevRotation);
+		SetShaderValue(RHICmdList, ComputeShaderRHI, InstanceWorldVelocityParam, DeltaPosition);
+	}
+
+private:
+	LAYOUT_FIELD(FShaderParameter, NumTrianglesParam);
+	LAYOUT_FIELD(FShaderParameter, NumVerticesParam);
+	LAYOUT_FIELD(FShaderParameter, NumUVsParam);
+	LAYOUT_FIELD(FShaderParameter, HasColorsParam);
+	LAYOUT_FIELD(FShaderResourceParameter, IndexBufferParam);
+	LAYOUT_FIELD(FShaderResourceParameter, PositionBufferParam);
+	LAYOUT_FIELD(FShaderResourceParameter, TangentBufferParam);
+	LAYOUT_FIELD(FShaderResourceParameter, UVBufferParam);
+	LAYOUT_FIELD(FShaderResourceParameter, ColorBufferParam);
+
+	LAYOUT_FIELD(FShaderParameter, HasUniformSamplingParam);
+	LAYOUT_FIELD(FShaderResourceParameter, UniformSamplingTrianglesParam);
+
+	LAYOUT_FIELD(FShaderParameter,  SectionCountsParam);
+	LAYOUT_FIELD(FShaderResourceParameter, SectionInfosParam);
+	LAYOUT_FIELD(FShaderResourceParameter, FilteredAndUnfilteredSectionsParam);
+
+	LAYOUT_FIELD(FShaderParameter, SocketCountsParam);
+	LAYOUT_FIELD(FShaderResourceParameter, SocketTransformsParam);
+	LAYOUT_FIELD(FShaderResourceParameter, FilteredAndUnfilteredSocketsParam);
+
+	LAYOUT_FIELD(FShaderParameter, InvDeltaSecondsParam);
+	LAYOUT_FIELD(FShaderParameter, InstanceTransformParam);
+	LAYOUT_FIELD(FShaderParameter, InstanceTransformInverseTransposedParam);
+	LAYOUT_FIELD(FShaderParameter, InstanceRotationParam);
+	LAYOUT_FIELD(FShaderParameter, InstancePreviousTransformParam);
+	LAYOUT_FIELD(FShaderParameter, InstancePreviousTransformInverseTransposedParam);
+	LAYOUT_FIELD(FShaderParameter, InstancePreviousRotationParam);
+	LAYOUT_FIELD(FShaderParameter, InstanceWorldVelocityParam);
+};
+
+IMPLEMENT_TYPE_LAYOUT(FNiagaraDataInterfaceParametersCS_StaticMesh);
+IMPLEMENT_NIAGARA_DI_PARAMETER(UNiagaraDataInterfaceStaticMesh, FNiagaraDataInterfaceParametersCS_StaticMesh);
+
+//////////////////////////////////////////////////////////////////////////
+
+UNiagaraDataInterfaceStaticMesh::UNiagaraDataInterfaceStaticMesh(FObjectInitializer const& ObjectInitializer)
+	: Super(ObjectInitializer)
+{
+	Proxy.Reset(new NDIStaticMeshLocal::FRenderProxy());
+}
+
+void UNiagaraDataInterfaceStaticMesh::PostInitProperties()
+{
+	Super::PostInitProperties();
+
+	// Can we register data interfaces as regular types and fold them into the FNiagaraVariable framework for UI and function calls etc?
+	if (HasAnyFlags(RF_ClassDefaultObject))
+	{
+		ENiagaraTypeRegistryFlags DIFlags =
+			ENiagaraTypeRegistryFlags::AllowAnyVariable |
+			ENiagaraTypeRegistryFlags::AllowParameter;
+		FNiagaraTypeRegistry::Register(FNiagaraTypeDefinition(GetClass()), DIFlags);
+
+		// Still some issues with using custom structs. Convert node for example throws a wobbler. TODO after GDC.
+		ENiagaraTypeRegistryFlags CoordFlags =
+			ENiagaraTypeRegistryFlags::AllowAnyVariable |
+			ENiagaraTypeRegistryFlags::AllowParameter |
+			ENiagaraTypeRegistryFlags::AllowPayload;
+		FNiagaraTypeRegistry::Register(FMeshTriCoordinate::StaticStruct(), CoordFlags);
+	}
+}
+
+#if WITH_EDITOR
+void UNiagaraDataInterfaceStaticMesh::PostEditChangeProperty(struct FPropertyChangedEvent& PropertyChangedEvent)
+{
+	Super::PostEditChangeProperty(PropertyChangedEvent);
+	++ChangeId;
+
+	if (PropertyChangedEvent.Property &&
+		PropertyChangedEvent.Property->GetFName() == GET_MEMBER_NAME_CHECKED(UNiagaraDataInterfaceStaticMesh, SourceMode))
+	{
+		if (SourceMode != ENDIStaticMesh_SourceMode::Default && SourceMode != ENDIStaticMesh_SourceMode::Source)
+		{
+			// Ensure we don't have any reference to a source actor that we'll never use
+			Source = nullptr;
+			SourceComponent = nullptr;
+		}
+	
+		if (SourceMode != ENDIStaticMesh_SourceMode::Default && SourceMode != ENDIStaticMesh_SourceMode::DefaultMeshOnly)
+		{
+			// Ensure we don't cook in a default mesh we'll never use
+			DefaultMesh = nullptr;
+		}
+	}
+}
+
+bool UNiagaraDataInterfaceStaticMesh::CanEditChange(const FProperty* InProperty) const
+{
+	if (!Super::CanEditChange(InProperty))
+	{
+		return false;
+	}
+
+	if (InProperty->GetFName() == GET_MEMBER_NAME_CHECKED(UNiagaraDataInterfaceStaticMesh, Source) &&
+		SourceMode != ENDIStaticMesh_SourceMode::Default &&
+		SourceMode != ENDIStaticMesh_SourceMode::Source)
+	{
+		// Disable Source if we'll never use it
+		return false;
+	}
+
+	if (InProperty->GetFName() == GET_MEMBER_NAME_CHECKED(UNiagaraDataInterfaceStaticMesh, DefaultMesh) &&
+		SourceMode != ENDIStaticMesh_SourceMode::Default &&
+		SourceMode != ENDIStaticMesh_SourceMode::DefaultMeshOnly)
+	{
+		// Disable Default Mesh if we'll never use it
+		return false;
+	}
+
+	return true;
+}
+#endif
+
+bool UNiagaraDataInterfaceStaticMesh::InitPerInstanceData(void* PerInstanceData, FNiagaraSystemInstance* SystemInstance)
+{
+	NDIStaticMeshLocal::FInstanceData_GameThread* InstanceData = new(PerInstanceData) NDIStaticMeshLocal::FInstanceData_GameThread();
+	const bool bIsValid = InstanceData->Init(this, SystemInstance);
+
+	// Create render thread data?
+	if (bIsValid && IsUsedWithGPUEmitter() )
+	{
+		TUniquePtr<NDIStaticMeshLocal::FGpuInitializeData> GpuInitializeData = MakeUnique<NDIStaticMeshLocal::FGpuInitializeData>();
+		GpuInitializeData->RenderProxy = GetProxyAs<NDIStaticMeshLocal::FRenderProxy>();
+		GpuInitializeData->SystemInstanceID = SystemInstance->GetId();
+		GpuInitializeData->LODResource = InstanceData->GetCurrentFirstLOD();
+		if ( GpuInitializeData->LODResource )
+		{
+			GpuInitializeData->NumFilteredTriangles = InstanceData->NumFilteredTriangles;
+			GpuInitializeData->NumUnfilteredTriangles = InstanceData->NumUnfilteredTriangles;
+
+			GpuInitializeData->bGpuUniformDistribution = InstanceData->bIsGpuUniformlyDistributedSampling;
+
+			// Build section data to pass
+			GpuInitializeData->NumSections = GpuInitializeData->LODResource->Sections.Num();
+			if (GpuInitializeData->NumSections > 0)
+			{
+				const FStaticMeshAreaWeightedSectionSampler* SectionSamplers = GpuInitializeData->bGpuUniformDistribution ? &GpuInitializeData->LODResource->AreaWeightedSampler : nullptr;
+				for (int i = 0; i < GpuInitializeData->NumSections; ++i)
+				{
+					const FStaticMeshSection& MeshSection = GpuInitializeData->LODResource->Sections[i];
+					const float Prob = SectionSamplers ? SectionSamplers->GetProb()[i] : 1.0f;
+					const int32 Alias = SectionSamplers ? SectionSamplers->GetAlias()[i] : i;
+					GpuInitializeData->SectionInfos.Emplace(MeshSection.FirstIndex / 3, MeshSection.NumTriangles, reinterpret_cast<const int32&>(Prob), Alias);
+				}
+
+				GpuInitializeData->NumFilteredSections = InstanceData->NumFilteredSections;
+				for (const NDIStaticMeshLocal::FNDISectionInfo& SectionInfo : InstanceData->GetFilteredSectionInfos())
+				{
+					GpuInitializeData->SectionInfos.Emplace(SectionInfo.FirstTriangle, SectionInfo.NumTriangles, reinterpret_cast<const int32&>(SectionInfo.Prob), SectionInfo.Alias);
+				}
+				for (int Section : InstanceData->GetFilteredSections())
+				{
+					GpuInitializeData->FilteredAndUnfilteredSections.Add(Section);
+				}
+
+				GpuInitializeData->NumUnfilteredSections = InstanceData->NumUnfilteredSections;
+				for (const NDIStaticMeshLocal::FNDISectionInfo& SectionInfo : InstanceData->GetUnfilteredSectionInfos())
+				{
+					GpuInitializeData->SectionInfos.Emplace(SectionInfo.FirstTriangle, SectionInfo.NumTriangles, reinterpret_cast<const int32&>(SectionInfo.Prob), SectionInfo.Alias);
+				}
+				for (int Section : InstanceData->GetUnfilteredSections())
+				{
+					GpuInitializeData->FilteredAndUnfilteredSections.Add(Section);
+				}
+			}
+		}
+
+		// Build socket data to pass
+		GpuInitializeData->NumSockets = InstanceData->CachedSockets.Num();
+		if (GpuInitializeData->NumSockets > 0)
+		{
+			GpuInitializeData->NumFilteredSockets = InstanceData->NumFilteredSockets;
+			GpuInitializeData->NumUnfilteredSockets = GpuInitializeData->NumSockets - InstanceData->NumFilteredSockets;
+
+			GpuInitializeData->SocketTransforms.Reserve(GpuInitializeData->NumFilteredSockets * 3);
+			for (const FTransform3f& SocketTransform : InstanceData->CachedSockets)
+			{
+				const FQuat4f SocketRotation = SocketTransform.GetRotation();
+				GpuInitializeData->SocketTransforms.Emplace(SocketTransform.GetTranslation(), 0.0f);
+				GpuInitializeData->SocketTransforms.Emplace(SocketRotation.X, SocketRotation.Y, SocketRotation.Z, SocketRotation.W);
+				GpuInitializeData->SocketTransforms.Emplace(SocketTransform.GetScale3D(), 0.0f);
+			}
+
+			if (InstanceData->NumFilteredSockets > 0)
+			{
+				GpuInitializeData->FilteredAndUnfilteredSockets.Reserve(InstanceData->FilteredAndUnfilteredSockets.Num());
+				for (uint16 v : InstanceData->FilteredAndUnfilteredSockets)
+				{
+					GpuInitializeData->FilteredAndUnfilteredSockets.Add(v);
+				}
+			}
+		}
+
+		// Initialize the render thread side data
+		ENQUEUE_RENDER_COMMAND(InitInstanceData)(
+			[GpuInitializeData_RT=MoveTemp(GpuInitializeData)](FRHICommandListImmediate& CmdList)
+			{
+				NDIStaticMeshLocal::FInstanceData_RenderThread& InstanceData_RT = GpuInitializeData_RT->RenderProxy->PerInstanceData_RT.Add(GpuInitializeData_RT->SystemInstanceID);
+				InstanceData_RT.Init(*(GpuInitializeData_RT.Get()));
+			}
+		);
+	}
+
+	return bIsValid;
+}
+
+void UNiagaraDataInterfaceStaticMesh::DestroyPerInstanceData(void* PerInstanceData, FNiagaraSystemInstance* SystemInstance)
+{
+	NDIStaticMeshLocal::FInstanceData_GameThread* InstanceData = reinterpret_cast<NDIStaticMeshLocal::FInstanceData_GameThread*>(PerInstanceData);
+	InstanceData->Release();
+	InstanceData->~FInstanceData_GameThread();
+
+	if ( IsUsedWithGPUEmitter() )
+	{
+		NDIStaticMeshLocal::FRenderProxy* Proxy_RT = GetProxyAs<NDIStaticMeshLocal::FRenderProxy>();
+		FNiagaraSystemInstanceID InstanceID_RT = SystemInstance->GetId();
+		ENQUEUE_RENDER_COMMAND(RemoveInstanceData)(
+			[Proxy_RT, InstanceID_RT](FRHICommandListImmediate& CmdList)
+			{
+				Proxy_RT->PerInstanceData_RT.Remove(InstanceID_RT);
+			}
+		);
+	}
+}
+
+bool UNiagaraDataInterfaceStaticMesh::PerInstanceTick(void* PerInstanceData, FNiagaraSystemInstance* SystemInstance, float DeltaSeconds)
+{
+	using namespace NDIStaticMeshLocal;
+
+	FInstanceData_GameThread* InstanceData = reinterpret_cast<FInstanceData_GameThread*>(PerInstanceData);
+	return InstanceData->Tick(this, SystemInstance, DeltaSeconds);
+}
+
+int32 UNiagaraDataInterfaceStaticMesh::PerInstanceDataSize() const
+{
+	return sizeof(NDIStaticMeshLocal::FInstanceData_GameThread);
+}
+
+void UNiagaraDataInterfaceStaticMesh::ProvidePerInstanceDataForRenderThread(void* InDataFromGT, void* InInstanceData, const FNiagaraSystemInstanceID& SystemInstance)
+{
+	NDIStaticMeshLocal::FInstanceData_GameThread* InstanceData = static_cast<NDIStaticMeshLocal::FInstanceData_GameThread*>(InInstanceData);
+	NDIStaticMeshLocal::FInstanceData_FromGameThread* DataFromGT = static_cast<NDIStaticMeshLocal::FInstanceData_FromGameThread*>(InDataFromGT);
+
+	DataFromGT->Transform		= InstanceData->Transform;
+	DataFromGT->PrevTransform	= InstanceData->PrevTransform;
+	DataFromGT->Rotation		= InstanceData->Rotation;
+	DataFromGT->PrevRotation	= InstanceData->PrevRotation;
+	DataFromGT->DeltaSeconds	= InstanceData->DeltaSeconds;
+}
+
+void UNiagaraDataInterfaceStaticMesh::GetFunctions(TArray<FNiagaraFunctionSignature>& OutFunctions)
+{
+	using namespace NDIStaticMeshLocal;
+
+	//OutFunctions.Reserve(OutFunctions.Num() + NumFunctions);
+
+	// Setup base signature
+	FNiagaraFunctionSignature BaseSignature;
+	BaseSignature.Inputs.Emplace(FNiagaraTypeDefinition(GetClass()), TEXT("StaticMesh"));
+	BaseSignature.bMemberFunction = true;
+	BaseSignature.bRequiresContext = false;
+#if WITH_EDITORONLY_DATA
+	BaseSignature.FunctionVersion = EDIFunctionVersion::LatestVersion;
+#endif
+
+	//////////////////////////////////////////////////////////////////////////
+	// VM Vertex Sampling
+	{
+		FNiagaraFunctionSignature& Sig = OutFunctions.Add_GetRef(BaseSignature);
+		Sig.Name = IsValidVertexName;
+		Sig.Inputs.Emplace(FNiagaraTypeDefinition::GetIntDef(), TEXT("Vertex"));
+		Sig.Outputs.Emplace(FNiagaraTypeDefinition::GetBoolDef(), TEXT("Valid"));
+	}
+	{
+		FNiagaraFunctionSignature& Sig = OutFunctions.Add_GetRef(BaseSignature);
+		Sig.Name = RandomVertexName;
+		Sig.Inputs.Emplace(FNiagaraTypeDefinition::GetRandInfoDef(), TEXT("RandomInfo"));
+		Sig.Outputs.Emplace(FNiagaraTypeDefinition::GetIntDef(), TEXT("Vertex"));
+	}
+	{
+		FNiagaraFunctionSignature& Sig = OutFunctions.Add_GetRef(BaseSignature);
+		Sig.Name = GetVertexCountName;
+		Sig.Outputs.Emplace(FNiagaraTypeDefinition::GetIntDef(), TEXT("Count"));
+	}
+
+	//{
+	//	FNiagaraFunctionSignature& Sig = OutFunctions.Add_GetRef(BaseSignature);
+	//	Sig.Name = IsValidFilteredVertexName;
+	//	Sig.Inputs.Emplace(FNiagaraTypeDefinition::GetIntDef(), TEXT("Vertex"));
+	//	Sig.Outputs.Emplace(FNiagaraTypeDefinition::GetBoolDef(), TEXT("Valid"));
+	//}
+	//{
+	//	FNiagaraFunctionSignature& Sig = OutFunctions.Add_GetRef(BaseSignature);
+	//	Sig.Name = RandomFilteredVertexName;
+	//	Sig.Outputs.Emplace(FNiagaraTypeDefinition::GetIntDef(), TEXT("Vertex"));
+	//}
+	//{
+	//	FNiagaraFunctionSignature& Sig = OutFunctions.Add_GetRef(BaseSignature);
+	//	Sig.Name = GetFilteredVertexCountName;
+	//	Sig.Outputs.Emplace(FNiagaraTypeDefinition::GetIntDef(), TEXT("Count"));
+	//}
+	//{
+	//	FNiagaraFunctionSignature& Sig = OutFunctions.Add_GetRef(BaseSignature);
+	//	Sig.Name = GetFilteredVertexAtName;
+	//	Sig.Inputs.Emplace(FNiagaraTypeDefinition::GetIntDef(), TEXT("Index"));
+	//	Sig.Outputs.Emplace(FNiagaraTypeDefinition::GetIntDef(), TEXT("Vertex"));
+	//}
+
+	{
+		FNiagaraFunctionSignature Sig = BaseSignature;
+		Sig.Inputs.Emplace(FNiagaraTypeDefinition::GetIntDef(), TEXT("Vertex"));
+		Sig.Outputs.Emplace(FNiagaraTypeDefinition::GetVec3Def(), TEXT("Position"));
+		Sig.Outputs.Emplace(FNiagaraTypeDefinition::GetVec3Def(), TEXT("Velocity"));
+		Sig.Outputs.Emplace(FNiagaraTypeDefinition::GetVec3Def(), TEXT("Normal"));
+		Sig.Outputs.Emplace(FNiagaraTypeDefinition::GetVec3Def(), TEXT("Bitangent"));
+		Sig.Outputs.Emplace(FNiagaraTypeDefinition::GetVec3Def(), TEXT("Tangent"));
+		OutFunctions.Add_GetRef(Sig).Name = GetVertexName;
+		OutFunctions.Add_GetRef(Sig).Name = GetVertexWSName;
+	}
+	{
+		FNiagaraFunctionSignature& Sig = OutFunctions.Add_GetRef(BaseSignature);
+		Sig.Name = GetVertexColorName;
+		Sig.Inputs.Emplace(FNiagaraTypeDefinition::GetIntDef(), TEXT("Vertex"));
+		Sig.Outputs.Emplace(FNiagaraTypeDefinition::GetColorDef(), TEXT("Color"));
+	}
+	{
+		FNiagaraFunctionSignature& Sig = OutFunctions.Add_GetRef(BaseSignature);
+		Sig.Name = GetVertexUVName;
+		Sig.Inputs.Emplace(FNiagaraTypeDefinition::GetIntDef(), TEXT("Vertex"));
+		Sig.Inputs.Emplace(FNiagaraTypeDefinition::GetIntDef(), TEXT("UV Set"));
+		Sig.Outputs.Emplace(FNiagaraTypeDefinition::GetVec2Def(), TEXT("UV"));
+	}
+
+	//////////////////////////////////////////////////////////////////////////
+	// VM Triangle Sampling
+	{
+		FNiagaraFunctionSignature& Sig = OutFunctions.Add_GetRef(BaseSignature);
+		Sig.Name = IsValidTriangleName;
+		Sig.Inputs.Emplace(FNiagaraTypeDefinition::GetIntDef(), TEXT("Triangle"));
+		Sig.Outputs.Emplace(FNiagaraTypeDefinition::GetBoolDef(), TEXT("Valid"));
+	}
+	{
+		FNiagaraFunctionSignature& Sig = OutFunctions.Add_GetRef(BaseSignature);
+		Sig.Name = RandomTriangleName;
+		Sig.Inputs.Emplace(FNiagaraTypeDefinition::GetRandInfoDef(), TEXT("RandomInfo"));
+		Sig.Outputs.Emplace(FNiagaraTypeDefinition::GetIntDef(), TEXT("Triangle"));
+		Sig.Outputs.Emplace(FNiagaraTypeDefinition::GetVec3Def(), TEXT("BaryCoord"));
+	}
+	{
+		FNiagaraFunctionSignature& Sig = OutFunctions.Add_GetRef(BaseSignature);
+		Sig.Name = GetTriangleCountName;
+		Sig.Outputs.Emplace(FNiagaraTypeDefinition::GetIntDef(), TEXT("Count"));
+	}
+
+	{
+		FNiagaraFunctionSignature& Sig = OutFunctions.Add_GetRef(BaseSignature);
+		Sig.Name = IsValidFilteredTriangleName;
+		Sig.Inputs.Emplace(FNiagaraTypeDefinition::GetIntDef(), TEXT("Index"));
+		Sig.Outputs.Emplace(FNiagaraTypeDefinition::GetBoolDef(), TEXT("Valid"));
+	}
+	{
+		FNiagaraFunctionSignature& Sig = OutFunctions.Add_GetRef(BaseSignature);
+		Sig.Name = RandomFilteredTriangleName;
+		Sig.Inputs.Emplace(FNiagaraTypeDefinition::GetRandInfoDef(), TEXT("RandomInfo"));
+		Sig.Outputs.Emplace(FNiagaraTypeDefinition::GetIntDef(), TEXT("Triangle"));
+		Sig.Outputs.Emplace(FNiagaraTypeDefinition::GetVec3Def(), TEXT("BaryCoord"));
+	}
+	{
+		FNiagaraFunctionSignature& Sig = OutFunctions.Add_GetRef(BaseSignature);
+		Sig.Name = GetFilteredTriangleCountName;
+		Sig.Outputs.Emplace(FNiagaraTypeDefinition::GetIntDef(), TEXT("Count"));
+	}
+	{
+		FNiagaraFunctionSignature& Sig = OutFunctions.Add_GetRef(BaseSignature);
+		Sig.Name = GetFilteredTriangleAtName;
+		Sig.Inputs.Emplace(FNiagaraTypeDefinition::GetIntDef(), TEXT("Index"));
+		Sig.Outputs.Emplace(FNiagaraTypeDefinition::GetIntDef(), TEXT("Triangle"));
+	}
+
+	{
+		FNiagaraFunctionSignature& Sig = OutFunctions.Add_GetRef(BaseSignature);
+		Sig.Name = IsValidUnfilteredTriangleName;
+		Sig.Inputs.Emplace(FNiagaraTypeDefinition::GetIntDef(), TEXT("Index"));
+		Sig.Outputs.Emplace(FNiagaraTypeDefinition::GetBoolDef(), TEXT("Valid"));
+	}
+	{
+		FNiagaraFunctionSignature& Sig = OutFunctions.Add_GetRef(BaseSignature);
+		Sig.Name = RandomUnfilteredTriangleName;
+		Sig.Inputs.Emplace(FNiagaraTypeDefinition::GetRandInfoDef(), TEXT("RandomInfo"));
+		Sig.Outputs.Emplace(FNiagaraTypeDefinition::GetIntDef(), TEXT("Triangle"));
+		Sig.Outputs.Emplace(FNiagaraTypeDefinition::GetVec3Def(), TEXT("BaryCoord"));
+	}
+	{
+		FNiagaraFunctionSignature& Sig = OutFunctions.Add_GetRef(BaseSignature);
+		Sig.Name = GetUnfilteredTriangleCountName;
+		Sig.Outputs.Emplace(FNiagaraTypeDefinition::GetIntDef(), TEXT("Count"));
+	}
+	{
+		FNiagaraFunctionSignature& Sig = OutFunctions.Add_GetRef(BaseSignature);
+		Sig.Name = GetUnfilteredTriangleAtName;
+		Sig.Inputs.Emplace(FNiagaraTypeDefinition::GetIntDef(), TEXT("Index"));
+		Sig.Outputs.Emplace(FNiagaraTypeDefinition::GetIntDef(), TEXT("Triangle"));
+	}
+
+	{
+		FNiagaraFunctionSignature Sig = BaseSignature;
+		Sig.Inputs.Emplace(FNiagaraTypeDefinition::GetIntDef(), TEXT("Triangle"));
+		Sig.Inputs.Emplace_GetRef(FNiagaraVariable(FNiagaraTypeDefinition::GetVec3Def(), TEXT("BaryCoord"))).SetValue(FVector3f(1.0f / 3.0f));
+		Sig.Outputs.Emplace(FNiagaraTypeDefinition::GetVec3Def(), TEXT("Position"));
+		Sig.Outputs.Emplace(FNiagaraTypeDefinition::GetVec3Def(), TEXT("Velocity"));
+		Sig.Outputs.Emplace(FNiagaraTypeDefinition::GetVec3Def(), TEXT("Normal"));
+		Sig.Outputs.Emplace(FNiagaraTypeDefinition::GetVec3Def(), TEXT("Bitangent"));
+		Sig.Outputs.Emplace(FNiagaraTypeDefinition::GetVec3Def(), TEXT("Tangent"));
+		OutFunctions.Add_GetRef(Sig).Name = GetTriangleName;
+		OutFunctions.Add_GetRef(Sig).Name = GetTriangleWSName;
+	}
+	{
+		FNiagaraFunctionSignature& Sig = OutFunctions.Add_GetRef(BaseSignature);
+		Sig.Name = GetTriangleColorName;
+		Sig.Inputs.Emplace(FNiagaraTypeDefinition::GetIntDef(), TEXT("Triangle"));
+		Sig.Inputs.Emplace_GetRef(FNiagaraVariable(FNiagaraTypeDefinition::GetVec3Def(), TEXT("BaryCoord"))).SetValue(FVector3f(1.0f / 3.0f));
+		Sig.Outputs.Emplace(FNiagaraTypeDefinition::GetColorDef(), TEXT("Color"));
+	}
+	{
+		FNiagaraFunctionSignature& Sig = OutFunctions.Add_GetRef(BaseSignature);
+		Sig.Name = GetTriangleUVName;
+		Sig.Inputs.Emplace(FNiagaraTypeDefinition::GetIntDef(), TEXT("Triangle"));
+		Sig.Inputs.Emplace_GetRef(FNiagaraVariable(FNiagaraTypeDefinition::GetVec3Def(), TEXT("BaryCoord"))).SetValue(FVector3f(1.0f / 3.0f));
+		Sig.Inputs.Emplace(FNiagaraTypeDefinition::GetIntDef(), TEXT("UV Set"));
+		Sig.Outputs.Emplace(FNiagaraTypeDefinition::GetVec2Def(), TEXT("UV"));
+	}
+
+	{
+		FNiagaraFunctionSignature& Sig = OutFunctions.Add_GetRef(BaseSignature);
+		Sig.Name = GetTriangleIndicesName;
+		Sig.Inputs.Emplace(FNiagaraTypeDefinition::GetIntDef(), TEXT("Triangle"));
+		Sig.Outputs.Emplace(FNiagaraTypeDefinition::GetIntDef(), TEXT("Index0"));
+		Sig.Outputs.Emplace(FNiagaraTypeDefinition::GetIntDef(), TEXT("Index1"));
+		Sig.Outputs.Emplace(FNiagaraTypeDefinition::GetIntDef(), TEXT("Index2"));
+	}
+
+	//////////////////////////////////////////////////////////////////////////
+	// Socket Functions
+	{
+		FNiagaraFunctionSignature Sig = BaseSignature;
+		Sig.Outputs.Add(FNiagaraVariable(FNiagaraTypeDefinition::GetIntDef(), TEXT("Num")));
+
+		OutFunctions.Add_GetRef(Sig).Name = GetSocketCountName;
+		OutFunctions.Add_GetRef(Sig).Name = GetFilteredSocketCountName;
+		OutFunctions.Add_GetRef(Sig).Name = GetUnfilteredSocketCountName;
+	}
+	{
+		FNiagaraFunctionSignature Sig = BaseSignature;
+		Sig.Inputs.Add(FNiagaraVariable(FNiagaraTypeDefinition::GetIntDef(), TEXT("Index")));
+		Sig.Outputs.Add(FNiagaraVariable(FNiagaraTypeDefinition::GetVec3Def(), TEXT("Position")));
+		Sig.Outputs.Add(FNiagaraVariable(FNiagaraTypeDefinition::GetQuatDef(), TEXT("Rotation")));
+		Sig.Outputs.Add(FNiagaraVariable(FNiagaraTypeDefinition::GetVec3Def(), TEXT("Scale")));
+
+		OutFunctions.Add_GetRef(Sig).Name = GetSocketTransformName;
+		OutFunctions.Add_GetRef(Sig).Name = GetSocketTransformWSName;
+		OutFunctions.Add_GetRef(Sig).Name = GetFilteredSocketTransformName;
+		OutFunctions.Add_GetRef(Sig).Name = GetFilteredSocketTransformWSName;
+		OutFunctions.Add_GetRef(Sig).Name = GetUnfilteredSocketTransformName;
+		OutFunctions.Add_GetRef(Sig).Name = GetUnfilteredSocketTransformWSName;
+	}
+
+	////////////////////////////////////////////////////////////////////////////
+	// Section functions
+	{
+		FNiagaraFunctionSignature& Sig = OutFunctions.Add_GetRef(BaseSignature);
+		Sig.Name = IsValidSectionName;
+		Sig.Inputs.Emplace(FNiagaraTypeDefinition::GetIntDef(), TEXT("Index"));
+		Sig.Outputs.Emplace(FNiagaraTypeDefinition::GetBoolDef(), TEXT("Valid"));
+	}
+	{
+		FNiagaraFunctionSignature& Sig = OutFunctions.Add_GetRef(BaseSignature);
+		Sig.Name = GetSectionTriangleCountName;
+		Sig.Inputs.Emplace(FNiagaraTypeDefinition::GetIntDef(), TEXT("Section"));
+		Sig.Outputs.Emplace(FNiagaraTypeDefinition::GetIntDef(), TEXT("Count"));
+	}
+	{
+		FNiagaraFunctionSignature& Sig = OutFunctions.Add_GetRef(BaseSignature);
+		Sig.Name = RandomSectionTriangleName;
+		Sig.Inputs.Emplace(FNiagaraTypeDefinition::GetRandInfoDef(), TEXT("RandomInfo"));
+		Sig.Inputs.Emplace(FNiagaraTypeDefinition::GetIntDef(), TEXT("Section"));
+		Sig.Outputs.Emplace(FNiagaraTypeDefinition::GetIntDef(), TEXT("Triangle"));
+		Sig.Outputs.Emplace(FNiagaraTypeDefinition::GetVec3Def(), TEXT("BaryCoord"));
+	}
+	{
+		FNiagaraFunctionSignature& Sig = OutFunctions.Add_GetRef(BaseSignature);
+		Sig.Name = GetSectionTriangleAtName;
+		Sig.Inputs.Emplace(FNiagaraTypeDefinition::GetIntDef(), TEXT("Section"));
+		Sig.Inputs.Emplace(FNiagaraTypeDefinition::GetIntDef(), TEXT("Index"));
+		Sig.Outputs.Emplace(FNiagaraTypeDefinition::GetIntDef(), TEXT("Triangle"));
+	}
+
+	{
+		FNiagaraFunctionSignature Sig = BaseSignature;
+		Sig.Inputs.Emplace(FNiagaraTypeDefinition::GetIntDef(), TEXT("Section"));
+		Sig.Outputs.Emplace(FNiagaraTypeDefinition::GetIntDef(), TEXT("Count"));
+		OutFunctions.Add_GetRef(Sig).Name = GetFilteredSectionAtName;
+		OutFunctions.Add_GetRef(Sig).Name = GetUnfilteredSectionAtName;
+	}
+
+	{
+		FNiagaraFunctionSignature Sig = BaseSignature;
+		Sig.Outputs.Emplace(FNiagaraTypeDefinition::GetIntDef(), TEXT("Count"));
+		OutFunctions.Add_GetRef(Sig).Name = GetSectionCountName;
+		OutFunctions.Add_GetRef(Sig).Name = GetFilteredSectionCountName;
+		OutFunctions.Add_GetRef(Sig).Name = GetUnfilteredSectionCountName;
+	}
+
+	{
+		FNiagaraFunctionSignature Sig = BaseSignature;
+		Sig.Inputs.Emplace(FNiagaraTypeDefinition::GetRandInfoDef(), TEXT("RandomInfo"));
+		Sig.Outputs.Emplace(FNiagaraTypeDefinition::GetIntDef(), TEXT("Section"));
+		OutFunctions.Add_GetRef(Sig).Name = RandomSectionName;
+		OutFunctions.Add_GetRef(Sig).Name = RandomFilteredSectionName;
+		OutFunctions.Add_GetRef(Sig).Name = RandomUnfilteredSectionName;
+	}
+
+	//////////////////////////////////////////////////////////////////////////
+	// VM Misc Functions
+	{
+		FNiagaraFunctionSignature& Sig = OutFunctions.Add_GetRef(BaseSignature);
+		Sig.Name = IsValidName;
+		Sig.Outputs.Emplace(FNiagaraTypeDefinition::GetBoolDef(), TEXT("Valid"));
+	}
+	{
+		FNiagaraFunctionSignature& Sig = OutFunctions.Add_GetRef(BaseSignature);
+		Sig.Name = GetLocalToWorldName;
+		Sig.Outputs.Emplace(FNiagaraTypeDefinition::GetMatrix4Def(), TEXT("Matrix"));
+	}
+	{
+		FNiagaraFunctionSignature& Sig = OutFunctions.Add_GetRef(BaseSignature);
+		Sig.Name = GetLocalToWorldInverseTransposedName;
+		Sig.Outputs.Emplace(FNiagaraTypeDefinition::GetMatrix4Def(), TEXT("Matrix"));
+	}
+	{
+		FNiagaraFunctionSignature& Sig = OutFunctions.Add_GetRef(BaseSignature);
+		Sig.Name = GetWorldVelocityName;
+		Sig.Outputs.Emplace(FNiagaraTypeDefinition::GetVec3Def(), TEXT("Velocity"));
+	}
+
+	//////////////////////////////////////////////////////////////////////////
+	// Deprecated Functions
+	{
+		FNiagaraFunctionSignature& Sig = OutFunctions.Add_GetRef(BaseSignature);
+		Sig.bSoftDeprecatedFunction = true;
+		Sig.Name = Deprecated_GetTriColorName;
+		Sig.Inputs.Emplace(FNiagaraTypeDefinition(FMeshTriCoordinate::StaticStruct()), TEXT("Coord"));
+		Sig.Outputs.Emplace(FNiagaraTypeDefinition::GetColorDef(), TEXT("Color"));
+	}
+	{
+		FNiagaraFunctionSignature& Sig = OutFunctions.Add_GetRef(BaseSignature);
+		Sig.bSoftDeprecatedFunction = true;
+		Sig.Name = Deprecated_GetTriUVName;
+		Sig.Inputs.Emplace(FNiagaraTypeDefinition(FMeshTriCoordinate::StaticStruct()), TEXT("Coord"));
+		Sig.Inputs.Emplace(FNiagaraTypeDefinition::GetIntDef(), TEXT("UV Set"));
+		Sig.Outputs.Emplace(FNiagaraTypeDefinition::GetVec2Def(), TEXT("UV"));
+	}
+	{
+		FNiagaraFunctionSignature & Sig = OutFunctions.Add_GetRef(BaseSignature);
+		Sig.bSoftDeprecatedFunction = true;
+		Sig.Name = Deprecated_RandomSectionName;
+		Sig.Outputs.Emplace(FNiagaraTypeDefinition::GetIntDef(), TEXT("Section"));
+	}
+	{
+		FNiagaraFunctionSignature & Sig = OutFunctions.Add_GetRef(BaseSignature);
+		Sig.bSoftDeprecatedFunction = true;
+		Sig.Name = Deprecated_RandomTriCoordName;
+		Sig.Outputs.Emplace(FNiagaraTypeDefinition(FMeshTriCoordinate::StaticStruct()), TEXT("Coord"));
+	}
+	{
+		FNiagaraFunctionSignature & Sig = OutFunctions.Add_GetRef(BaseSignature);
+		Sig.bSoftDeprecatedFunction = true;
+		Sig.Name = Deprecated_RandomTriCoordOnSectionName;
+		Sig.Inputs.Emplace(FNiagaraTypeDefinition::GetIntDef(), TEXT("Section"));
+		Sig.Outputs.Emplace(FNiagaraTypeDefinition(FMeshTriCoordinate::StaticStruct()), TEXT("Coord"));
+	}
+	// Need custom code
+	{
+		FNiagaraFunctionSignature Sig = BaseSignature;
+		Sig.bSoftDeprecatedFunction = true;
+		Sig.Inputs.Emplace(FNiagaraTypeDefinition::GetIntDef(), TEXT("Vertex"));
+		Sig.Outputs.Emplace(FNiagaraTypeDefinition::GetVec3Def(), TEXT("Position"));
+		OutFunctions.Add_GetRef(Sig).Name = Deprecated_GetVertexPositionName;
+		OutFunctions.Add_GetRef(Sig).Name = Deprecated_GetVertexPositionWSName;
+	}
+
+	{
+		FNiagaraFunctionSignature Sig = BaseSignature;
+		Sig.bSoftDeprecatedFunction = true;
+		Sig.Inputs.Emplace(FNiagaraTypeDefinition(FMeshTriCoordinate::StaticStruct()), TEXT("Coord"));
+		Sig.Outputs.Emplace(FNiagaraTypeDefinition::GetVec3Def(), TEXT("Position"));
+		OutFunctions.Add_GetRef(Sig).Name = Deprecated_GetTriPositionName;
+		OutFunctions.Add_GetRef(Sig).Name = Deprecated_GetTriPositionWSName;
+	}
+	{
+		FNiagaraFunctionSignature& Sig = OutFunctions.Add_GetRef(BaseSignature);
+		Sig.bSoftDeprecatedFunction = true;
+		Sig.Name = Deprecated_GetTriPositionAndVelocityWSName;
+		Sig.Inputs.Emplace(FNiagaraTypeDefinition(FMeshTriCoordinate::StaticStruct()), TEXT("Coord"));
+		Sig.Outputs.Emplace(FNiagaraTypeDefinition::GetVec3Def(), TEXT("Position"));
+		Sig.Outputs.Emplace(FNiagaraTypeDefinition::GetVec3Def(), TEXT("Velocity"));
+	}
+	{
+		FNiagaraFunctionSignature Sig = BaseSignature;
+		Sig.bSoftDeprecatedFunction = true;
+		Sig.Inputs.Emplace(FNiagaraTypeDefinition(FMeshTriCoordinate::StaticStruct()), TEXT("Coord"));
+		Sig.Outputs.Emplace(FNiagaraTypeDefinition::GetVec3Def(), TEXT("Tangent"));
+		Sig.Outputs.Emplace(FNiagaraTypeDefinition::GetVec3Def(), TEXT("Binormal"));
+		Sig.Outputs.Emplace(FNiagaraTypeDefinition::GetVec3Def(), TEXT("Normal"));
+		OutFunctions.Add_GetRef(Sig).Name = Deprecated_GetTriTangentsName;
+		OutFunctions.Add_GetRef(Sig).Name = Deprecated_GetTriTangentsWSName;
+	}
+	{
+		FNiagaraFunctionSignature Sig = BaseSignature;
+		Sig.bSoftDeprecatedFunction = true;
+		Sig.Inputs.Emplace(FNiagaraTypeDefinition(FMeshTriCoordinate::StaticStruct()), TEXT("Coord"));
+		Sig.Outputs.Emplace(FNiagaraTypeDefinition::GetVec3Def(), TEXT("Normal"));
+		OutFunctions.Add_GetRef(Sig).Name = Deprecated_GetTriNormalName;
+		OutFunctions.Add_GetRef(Sig).Name = Deprecated_GetTriNormalWSName;
+	}
+}
+
+#if WITH_EDITORONLY_DATA
+bool UNiagaraDataInterfaceStaticMesh::UpgradeFunctionCall(FNiagaraFunctionSignature& FunctionSignature)
+{
+	using namespace NDIStaticMeshLocal;
+
+	// Early out for version matching
+	if (FunctionSignature.FunctionVersion == EDIFunctionVersion::LatestVersion)
+	{
+		return false;
+	}
+
+	// Major refactor for consistency across the data interface
+	if (FunctionSignature.FunctionVersion < EDIFunctionVersion::RefactoredV1)
+	{
+		static const TPair<FName, FName> NameTranslations[] =
+		{
+			MakeTuple("RandomSection", Deprecated_RandomSectionName),
+			MakeTuple("RandomTriCoord", Deprecated_RandomTriCoordName),
+			MakeTuple("RandomTriCoordOnSection", Deprecated_RandomTriCoordOnSectionName),
+			MakeTuple("RandomTriCoordUsingVertexColorFilter", Deprecated_RandomTriCoordVCFilteredName),
+			MakeTuple("GetVertexPosition", Deprecated_GetVertexPositionName),
+			MakeTuple("GetVertexPositionWS", Deprecated_GetVertexPositionWSName),
+			MakeTuple("GetTriPosition", Deprecated_GetTriPositionName),
+			MakeTuple("GetTriNormal", Deprecated_GetTriNormalName),
+			MakeTuple("GetTriTangents", Deprecated_GetTriTangentsName),
+			MakeTuple("GetTriPositionWS", Deprecated_GetTriPositionWSName),
+			MakeTuple("GetTriNormalWS", Deprecated_GetTriNormalWSName),
+			MakeTuple("GetTriTangentsWS", Deprecated_GetTriTangentsWSName),
+			MakeTuple("GetTriColor", Deprecated_GetTriColorName),
+			MakeTuple("GetTriUV", Deprecated_GetTriUVName),
+			MakeTuple("GetTriPositionAndVelocityWS", Deprecated_GetTriPositionAndVelocityWSName),
+			MakeTuple("GetMeshLocalToWorldInverseTransposed", GetLocalToWorldInverseTransposedName),
+		};
+
+		for ( const auto& NameTrans : NameTranslations )
+		{
+			if ( FunctionSignature.Name == NameTrans.Key )
+			{
+				FunctionSignature.Name = NameTrans.Value;
+				break;
+			}
+		}
+
+		FunctionSignature.FunctionVersion = EDIFunctionVersion::RefactoredV1;
+	}
+
+	return true;
+}
+#endif //WITH_EDITORONLY_DATA
+
+void UNiagaraDataInterfaceStaticMesh::GetVMExternalFunction(const FVMExternalFunctionBindingInfo& BindingInfo, void* InstanceData, FVMExternalFunction& OutFunc)
+{
+	using namespace NDIStaticMeshLocal;
+
+	//////////////////////////////////////////////////////////////////////////
+	// VM Vertex Sampling
+	if (BindingInfo.Name == IsValidVertexName)
+	{
+		OutFunc = FVMExternalFunction::CreateLambda([this](FVectorVMExternalFunctionContext& Context) { VMIsValidVertex(Context); });
+	}
+	else if (BindingInfo.Name == RandomVertexName)
+	{
+		OutFunc = FVMExternalFunction::CreateLambda([this](FVectorVMExternalFunctionContext& Context) { VMRandomVertex(Context); });
+	}
+	else if (BindingInfo.Name == GetVertexCountName)
+	{
+		OutFunc = FVMExternalFunction::CreateLambda([this](FVectorVMExternalFunctionContext& Context) { VMGetVertexCount(Context); });
+	}
+	//else if (BindingInfo.Name == IsValidFilteredVertexName)
+	//{
+	//	OutFunc = FVMExternalFunction::CreateLambda([this](FVectorVMExternalFunctionContext& Context) { VMIsValidFilteredVertex(Context); });
+	//}
+	//else if (BindingInfo.Name == RandomFilteredVertexName)
+	//{
+	//	OutFunc = FVMExternalFunction::CreateLambda([this](FVectorVMExternalFunctionContext& Context) { VMRandomFilteredVertexAt(Context); });
+	//}
+	//else if (BindingInfo.Name == GetFilteredVertexCountName)
+	//{
+	//	OutFunc = FVMExternalFunction::CreateLambda([this](FVectorVMExternalFunctionContext& Context) { VMGetFilteredVertexCount(Context); });
+	//}
+	//else if (BindingInfo.Name == GetFilteredVertexAtName)
+	//{
+	//	OutFunc = FVMExternalFunction::CreateLambda([this](FVectorVMExternalFunctionContext& Context) { VMGetFilteredVertexAt(Context); });
+	//}
+	else if (BindingInfo.Name == GetVertexName)
+	{
+		OutFunc = FVMExternalFunction::CreateLambda([this](FVectorVMExternalFunctionContext& Context) { VMGetVertex<FNDITransformHandlerNoop>(Context); });
+	}
+	else if (BindingInfo.Name == GetVertexWSName)
+	{
+		OutFunc = FVMExternalFunction::CreateLambda([this](FVectorVMExternalFunctionContext& Context) { VMGetVertex<FNDITransformHandler>(Context); });
+	}
+	else if (BindingInfo.Name == GetVertexColorName)
+	{
+		OutFunc = FVMExternalFunction::CreateLambda([this](FVectorVMExternalFunctionContext& Context) { VMGetVertexColor(Context); });
+	}
+	else if (BindingInfo.Name == GetVertexUVName)
+	{
+		OutFunc = FVMExternalFunction::CreateLambda([this](FVectorVMExternalFunctionContext& Context) { VMGetVertexUV(Context); });
+	}
+
+	//////////////////////////////////////////////////////////////////////////
+	// VM Triangle Sampling
+	else if (BindingInfo.Name == IsValidTriangleName)
+	{
+		OutFunc = FVMExternalFunction::CreateLambda([this](FVectorVMExternalFunctionContext& Context) { VMIsValidTriangle(Context); });
+	}
+	else if (BindingInfo.Name == RandomTriangleName)
+	{
+		OutFunc = FVMExternalFunction::CreateLambda([this](FVectorVMExternalFunctionContext& Context) { VMRandomTriangle<FNDIRandomHelper>(Context); });
+	}
+	else if (BindingInfo.Name == GetTriangleCountName)
+	{
+		OutFunc = FVMExternalFunction::CreateLambda([this](FVectorVMExternalFunctionContext& Context) { VMGetTriangleCount(Context); });
+	}
+	else if (BindingInfo.Name == IsValidFilteredTriangleName)
+	{
+		OutFunc = FVMExternalFunction::CreateLambda([this](FVectorVMExternalFunctionContext& Context) { VMIsValidFilteredTriangle(Context); });
+	}
+	else if (BindingInfo.Name == RandomFilteredTriangleName)
+	{
+		OutFunc = FVMExternalFunction::CreateLambda([this](FVectorVMExternalFunctionContext& Context) { VMRandomFilteredTriangle(Context); });
+	}
+	else if (BindingInfo.Name == GetFilteredTriangleCountName)
+	{
+		OutFunc = FVMExternalFunction::CreateLambda([this](FVectorVMExternalFunctionContext& Context) { VMGetFilteredTriangleCount(Context); });
+	}
+	else if (BindingInfo.Name == GetFilteredTriangleAtName)
+	{
+		OutFunc = FVMExternalFunction::CreateLambda([this](FVectorVMExternalFunctionContext& Context) { VMGetFilteredTriangleAt(Context); });
+	}
+	else if (BindingInfo.Name == IsValidUnfilteredTriangleName)
+	{
+		OutFunc = FVMExternalFunction::CreateLambda([this](FVectorVMExternalFunctionContext& Context) { VMIsValidUnfilteredTriangle(Context); });
+	}
+	else if (BindingInfo.Name == RandomUnfilteredTriangleName)
+	{
+		OutFunc = FVMExternalFunction::CreateLambda([this](FVectorVMExternalFunctionContext& Context) { VMRandomUnfilteredTriangle(Context); });
+	}
+	else if (BindingInfo.Name == GetUnfilteredTriangleCountName)
+	{
+		OutFunc = FVMExternalFunction::CreateLambda([this](FVectorVMExternalFunctionContext& Context) { VMGetUnfilteredTriangleCount(Context); });
+	}
+	else if (BindingInfo.Name == GetUnfilteredTriangleAtName)
+	{
+		OutFunc = FVMExternalFunction::CreateLambda([this](FVectorVMExternalFunctionContext& Context) { VMGetUnfilteredTriangleAt(Context); });
+	}
+	else if (BindingInfo.Name == GetTriangleName)
+	{
+		OutFunc = FVMExternalFunction::CreateLambda([this](FVectorVMExternalFunctionContext& Context) { VMGetTriangle<FNDITransformHandlerNoop>(Context); });
+	}
+	else if (BindingInfo.Name == GetTriangleWSName)
+	{
+		OutFunc = FVMExternalFunction::CreateLambda([this](FVectorVMExternalFunctionContext& Context) { VMGetTriangle<FNDITransformHandler>(Context); });
+	}
+	else if (BindingInfo.Name == GetTriangleColorName)
+	{
+		OutFunc = FVMExternalFunction::CreateLambda([this](FVectorVMExternalFunctionContext& Context) { VMGetTriangleColor(Context); });
+	}
+	else if (BindingInfo.Name == GetTriangleUVName)
+	{
+		OutFunc = FVMExternalFunction::CreateLambda([this](FVectorVMExternalFunctionContext& Context) { VMGetTriangleUV(Context); });
+	}
+	else if (BindingInfo.Name == GetTriangleIndicesName)
+	{
+		OutFunc = FVMExternalFunction::CreateLambda([this](FVectorVMExternalFunctionContext& Context) { VMGetTriangleIndices(Context); });
+	}
+
+	//////////////////////////////////////////////////////////////////////////
+	// Socket Functions
+	else if (BindingInfo.Name == GetSocketCountName)
+	{
+		OutFunc = FVMExternalFunction::CreateLambda([this](FVectorVMExternalFunctionContext& Context) { VMGetSocketCount(Context); });
+	}
+	else if (BindingInfo.Name == GetFilteredSocketCountName)
+	{
+		OutFunc = FVMExternalFunction::CreateLambda([this](FVectorVMExternalFunctionContext& Context) { VMGetFilteredSocketCount(Context); });
+	}
+	else if (BindingInfo.Name == GetUnfilteredSocketCountName)
+	{
+		OutFunc = FVMExternalFunction::CreateLambda([this](FVectorVMExternalFunctionContext& Context) { VMGetUnfilteredSocketCount(Context); });
+	}
+	else if (BindingInfo.Name == GetSocketTransformName)
+	{
+		OutFunc = FVMExternalFunction::CreateLambda([this](FVectorVMExternalFunctionContext& Context) { VMGetSocketTransform<FNDITransformHandlerNoop>(Context); });
+	}
+	else if (BindingInfo.Name == GetSocketTransformWSName)
+	{
+		OutFunc = FVMExternalFunction::CreateLambda([this](FVectorVMExternalFunctionContext& Context) { VMGetSocketTransform<FNDITransformHandler>(Context); });
+	}
+	else if (BindingInfo.Name == GetFilteredSocketTransformName)
+	{
+		OutFunc = FVMExternalFunction::CreateLambda([this](FVectorVMExternalFunctionContext& Context) { VMGetFilteredSocketTransform<FNDITransformHandlerNoop>(Context); });
+	}
+	else if (BindingInfo.Name == GetFilteredSocketTransformWSName)
+	{
+		OutFunc = FVMExternalFunction::CreateLambda([this](FVectorVMExternalFunctionContext& Context) { VMGetFilteredSocketTransform<FNDITransformHandler>(Context); });
+	}
+	else if (BindingInfo.Name == GetUnfilteredSocketTransformName)
+	{
+		OutFunc = FVMExternalFunction::CreateLambda([this](FVectorVMExternalFunctionContext& Context) { VMGetUnfilteredSocketTransform<FNDITransformHandlerNoop>(Context); });
+	}
+	else if (BindingInfo.Name == GetUnfilteredSocketTransformWSName)
+	{
+		OutFunc = FVMExternalFunction::CreateLambda([this](FVectorVMExternalFunctionContext& Context) { VMGetUnfilteredSocketTransform<FNDITransformHandler>(Context); });
+	}
+
+	//////////////////////////////////////////////////////////////////////////
+	// Section functions
+	else if (BindingInfo.Name == IsValidSectionName)
+	{
+		OutFunc = FVMExternalFunction::CreateLambda([this](FVectorVMExternalFunctionContext& Context) { VMIsValidSection(Context); });
+	}
+	else if (BindingInfo.Name == GetSectionTriangleCountName)
+	{
+		OutFunc = FVMExternalFunction::CreateLambda([this](FVectorVMExternalFunctionContext& Context) { VMGetSectionTriangleCount(Context); });
+	}
+	else if (BindingInfo.Name == RandomSectionTriangleName)
+	{
+		OutFunc = FVMExternalFunction::CreateLambda([this](FVectorVMExternalFunctionContext& Context) { VMRandomSectionTriangle<FNDIRandomHelper>(Context); });
+	}
+	else if (BindingInfo.Name == GetSectionTriangleAtName)
+	{
+		OutFunc = FVMExternalFunction::CreateLambda([this](FVectorVMExternalFunctionContext& Context) { VMGetSectionTriangleAt(Context); });
+	}
+	else if (BindingInfo.Name == GetFilteredSectionAtName)
+	{
+		OutFunc = FVMExternalFunction::CreateLambda([this](FVectorVMExternalFunctionContext& Context) { VMGetFilteredSectionAt(Context); });
+	}
+	else if (BindingInfo.Name == GetUnfilteredSectionAtName)
+	{
+		OutFunc = FVMExternalFunction::CreateLambda([this](FVectorVMExternalFunctionContext& Context) { VMGetUnfilteredSectionAt(Context); });
+	}
+	else if (BindingInfo.Name == GetSectionCountName)
+	{
+		OutFunc = FVMExternalFunction::CreateLambda([this](FVectorVMExternalFunctionContext& Context) { VMGetSectionCount(Context); });
+	}
+	else if (BindingInfo.Name == GetFilteredSectionCountName)
+	{
+		OutFunc = FVMExternalFunction::CreateLambda([this](FVectorVMExternalFunctionContext& Context) { VMGetFilteredSectionCount(Context); });
+	}
+	else if (BindingInfo.Name == GetUnfilteredSectionCountName)
+	{
+		OutFunc = FVMExternalFunction::CreateLambda([this](FVectorVMExternalFunctionContext& Context) { VMGetUnfilteredSectionCount(Context); });
+	}
+	else if (BindingInfo.Name == RandomSectionName)
+	{
+		OutFunc = FVMExternalFunction::CreateLambda([this](FVectorVMExternalFunctionContext& Context) { VMRandomSection<FNDIRandomHelper>(Context); });
+	}
+	else if (BindingInfo.Name == RandomFilteredSectionName)
+	{
+		OutFunc = FVMExternalFunction::CreateLambda([this](FVectorVMExternalFunctionContext& Context) { VMRandomFilteredSection<FNDIRandomHelper>(Context); });
+	}
+	else if (BindingInfo.Name == RandomUnfilteredSectionName)
+	{
+		OutFunc = FVMExternalFunction::CreateLambda([this](FVectorVMExternalFunctionContext& Context) { VMRandomUnfilteredSection(Context); });
+	}
+
+	//////////////////////////////////////////////////////////////////////////
+	// VM Misc Functions
+	else if (BindingInfo.Name == IsValidName)
+	{
+		OutFunc = FVMExternalFunction::CreateLambda([this](FVectorVMExternalFunctionContext& Context) { VMIsValid(Context); });
+	}
+	else if (BindingInfo.Name == GetLocalToWorldName)
+	{
+		OutFunc = FVMExternalFunction::CreateLambda([this](FVectorVMExternalFunctionContext& Context) { VMGetLocalToWorld(Context); });
+	}
+	else if (BindingInfo.Name == GetLocalToWorldInverseTransposedName)
+	{
+		OutFunc = FVMExternalFunction::CreateLambda([this](FVectorVMExternalFunctionContext& Context) { VMGetLocalToWorldInverseTransposed(Context); });
+	}
+	else if (BindingInfo.Name == GetWorldVelocityName)
+	{
+		OutFunc = FVMExternalFunction::CreateLambda([this](FVectorVMExternalFunctionContext& Context) { VMGetWorldVelocity(Context); });
+	}
+
+	//////////////////////////////////////////////////////////////////////////
+	// Deprecated Functions
+	else if (BindingInfo.Name == Deprecated_GetTriColorName)
+	{
+		OutFunc = FVMExternalFunction::CreateLambda([this](FVectorVMExternalFunctionContext& Context) { VMGetTriangleColor(Context); });
+	}
+	else if (BindingInfo.Name == Deprecated_GetTriUVName)
+	{
+		OutFunc = FVMExternalFunction::CreateLambda([this](FVectorVMExternalFunctionContext& Context) { VMGetTriangleUV(Context); });
+	}
+	// Custom VM bindings
+	else if (BindingInfo.Name == Deprecated_RandomSectionName)
+	{
+		OutFunc = FVMExternalFunction::CreateLambda([this](FVectorVMExternalFunctionContext& Context) { VMRandomFilteredSection<FNDIRandomHelperFromStream>(Context); });
+	}
+	else if (BindingInfo.Name == Deprecated_RandomTriCoordName)
+	{
+		OutFunc = FVMExternalFunction::CreateLambda([this](FVectorVMExternalFunctionContext& Context) { VMRandomTriangle<FNDIRandomHelperFromStream>(Context); });
+	}
+	else if (BindingInfo.Name == Deprecated_RandomTriCoordOnSectionName)
+	{
+		OutFunc = FVMExternalFunction::CreateLambda([this](FVectorVMExternalFunctionContext& Context) { VMRandomSectionTriangle<FNDIRandomHelperFromStream>(Context); });
+	}
+	else if (BindingInfo.Name == Deprecated_GetVertexPositionName)
+	{
+		OutFunc = FVMExternalFunction::CreateLambda([this](FVectorVMExternalFunctionContext& Context) { VMGetVertexPosition_Deprecated<FNDITransformHandlerNoop>(Context); });
+	}
+	else if (BindingInfo.Name == Deprecated_GetVertexPositionWSName)
+	{
+		OutFunc = FVMExternalFunction::CreateLambda([this](FVectorVMExternalFunctionContext& Context) { VMGetVertexPosition_Deprecated<FNDITransformHandler>(Context); });
+	}
+	else if (BindingInfo.Name == Deprecated_GetTriPositionName)
+	{
+		OutFunc = FVMExternalFunction::CreateLambda([this](FVectorVMExternalFunctionContext& Context) { VMGetTriPosition_Deprecated<FNDITransformHandlerNoop>(Context); });
+	}
+	else if (BindingInfo.Name == Deprecated_GetTriPositionWSName)
+	{
+		OutFunc = FVMExternalFunction::CreateLambda([this](FVectorVMExternalFunctionContext& Context) { VMGetTriPosition_Deprecated<FNDITransformHandler>(Context); });
+	}
+	else if (BindingInfo.Name == Deprecated_GetTriPositionWSName)
+	{
+		OutFunc = FVMExternalFunction::CreateLambda([this](FVectorVMExternalFunctionContext& Context) { VMGetTriPosition_Deprecated<FNDITransformHandler>(Context); });
+	}
+	else if (BindingInfo.Name == Deprecated_GetTriTangentsName)
+	{	
+		OutFunc = FVMExternalFunction::CreateLambda([this](FVectorVMExternalFunctionContext& Context) { VMGetTriangleTangentBasis_Deprecated<FNDITransformHandlerNoop>(Context); });
+	}
+	else if (BindingInfo.Name == Deprecated_GetTriTangentsWSName)
+	{
+		OutFunc = FVMExternalFunction::CreateLambda([this](FVectorVMExternalFunctionContext& Context) { VMGetTriangleTangentBasis_Deprecated<FNDITransformHandler>(Context); });
+	}
+	else if (BindingInfo.Name == Deprecated_GetTriNormalName)
+	{
+		OutFunc = FVMExternalFunction::CreateLambda([this](FVectorVMExternalFunctionContext& Context) { VMGetTriangleNormal_Deprecated<FNDITransformHandlerNoop>(Context); });
+	}
+	else if (BindingInfo.Name == Deprecated_GetTriNormalWSName)
+	{
+		OutFunc = FVMExternalFunction::CreateLambda([this](FVectorVMExternalFunctionContext& Context) { VMGetTriangleNormal_Deprecated<FNDITransformHandler>(Context); });
+	}
+	else if (BindingInfo.Name == Deprecated_GetTriPositionAndVelocityWSName)
+	{
+	OutFunc = FVMExternalFunction::CreateLambda([this](FVectorVMExternalFunctionContext& Context) { VMGetTriPositionAndVelocity_Deprecated<FNDITransformHandler>(Context); });
+	}
+}
+
+#if WITH_EDITORONLY_DATA
+bool UNiagaraDataInterfaceStaticMesh::AppendCompileHash(FNiagaraCompileHashVisitor* InVisitor) const
+{
+	bool bSuccess = Super::AppendCompileHash(InVisitor);
+	FSHAHash Hash = GetShaderFileHash(NDIStaticMeshLocal::TemplateShaderFile, EShaderPlatform::SP_PCD3D_SM5);
+	InVisitor->UpdateString(TEXT("NiagaraDataInterfaceStaticMeshTemplateHLSLSource"), Hash.ToString());
+	return bSuccess;
+}
+
+void UNiagaraDataInterfaceStaticMesh::GetParameterDefinitionHLSL(const FNiagaraDataInterfaceGPUParamInfo& ParamInfo, FString& OutHLSL)
+{
+	TMap<FString, FStringFormatArg> TemplateArgs =
+	{
+		{TEXT("ParameterName"),	ParamInfo.DataInterfaceHLSLSymbol},
+	};
+
+	FString TemplateFile;
+	LoadShaderSourceFile(NDIStaticMeshLocal::TemplateShaderFile, EShaderPlatform::SP_PCD3D_SM5, &TemplateFile, nullptr);
+	OutHLSL += FString::Format(*TemplateFile, TemplateArgs);
+}
+
+bool UNiagaraDataInterfaceStaticMesh::GetFunctionHLSL(const FNiagaraDataInterfaceGPUParamInfo& ParamInfo, const FNiagaraDataInterfaceGeneratedFunction& FunctionInfo, int FunctionInstanceIndex, FString& OutHLSL)
+{
+	using namespace NDIStaticMeshLocal;
+
+	// Currently we support all functions on the GPU
+	// Add an exclusion if the function is not supported by GPU
+
+	TMap<FString, FStringFormatArg> ArgsSample =
+	{
+		{TEXT("ParameterName"),	ParamInfo.DataInterfaceHLSLSymbol},
+		{TEXT("InstanceFunctionName"), FunctionInfo.InstanceName},
+	};
+
+	// We still need to go through some hoops for the old deprecated functions
+	if (FunctionInfo.DefinitionName == Deprecated_RandomSectionName)
+	{
+		static const TCHAR* FormatSample = TEXT("void {InstanceFunctionName}(out int OutSection) { RandomFilteredSection_{ParameterName}(MakeRandInfo(), OutSection); }\r\n");
+		OutHLSL += FString::Format(FormatSample, ArgsSample);
+	}
+	else if (FunctionInfo.DefinitionName == Deprecated_RandomTriCoordName )
+	{
+		static const TCHAR* FormatSample = TEXT("void {InstanceFunctionName}(out MeshTriCoordinate OutTriCoord) { RandomTriangle_{ParameterName}(MakeRandInfo(), OutTriCoord.Tri, OutTriCoord.BaryCoord); }\r\n");
+		OutHLSL += FString::Format(FormatSample, ArgsSample);
+	}
+	else if (FunctionInfo.DefinitionName == Deprecated_RandomTriCoordOnSectionName )
+	{
+		static const TCHAR* FormatSample = TEXT("void {InstanceFunctionName}(int Section, out MeshTriCoordinate OutTriCoord) { RandomSectionTriangle_{ParameterName}(MakeRandInfo(), Section, OutTriCoord.Tri, OutTriCoord.BaryCoord); }\r\n");
+		OutHLSL += FString::Format(FormatSample, ArgsSample);
+	}
+	else if (FunctionInfo.DefinitionName == Deprecated_RandomTriCoordVCFilteredName )
+	{
+		// Not supported anymore
+		return false;
+	}
+	else if (FunctionInfo.DefinitionName == Deprecated_GetVertexPositionName)
+	{
+		static const TCHAR* FormatSample = TEXT("void {InstanceFunctionName}(int Vertex, out float3 Position) { float3 Dummy; GetVertex_{ParameterName}(Vertex, Position, Dummy, Dummy, Dummy, Dummy); }");
+		OutHLSL += FString::Format(FormatSample, ArgsSample);
+	}
+	else if (FunctionInfo.DefinitionName == Deprecated_GetVertexPositionWSName)
+	{
+		static const TCHAR* FormatSample = TEXT("void {InstanceFunctionName}(int Vertex, out float3 Position) { float3 Dummy; GetVertexWS_{ParameterName}(Vertex, Position, Dummy, Dummy, Dummy, Dummy); }");
+		OutHLSL += FString::Format(FormatSample, ArgsSample);
+	}
+	else if (FunctionInfo.DefinitionName == Deprecated_GetTriPositionName )
+	{
+		static const TCHAR* FormatSample = TEXT("void {InstanceFunctionName}(MeshTriCoordinate TriCoord, out float3 Position) { float3 Dummy; GetTriangle_{ParameterName}(TriCoord.Tri, TriCoord.BaryCoord, Position, Dummy, Dummy, Dummy, Dummy); }\r\n");
+		OutHLSL += FString::Format(FormatSample, ArgsSample);
+	}
+	else if ( FunctionInfo.DefinitionName == Deprecated_GetTriNormalName )
+	{
+		static const TCHAR* FormatSample = TEXT("void {InstanceFunctionName}(MeshTriCoordinate TriCoord, out float3 Normal) { float3 Dummy; GetTriangle_{ParameterName}(TriCoord.Tri, TriCoord.BaryCoord, Dummy, Dummy, Normal, Dummy, Dummy); }\r\n");
+		OutHLSL += FString::Format(FormatSample, ArgsSample);
+	}
+	else if ( FunctionInfo.DefinitionName == Deprecated_GetTriTangentsName )
+	{
+		static const TCHAR* FormatSample = TEXT("void {InstanceFunctionName}(MeshTriCoordinate TriCoord, out float3 Tangent, out float3 Bitangent, out float3 Normal) { float3 Dummy; GetTriangle_{ParameterName}(TriCoord.Tri, TriCoord.BaryCoord, Dummy, Dummy, Normal, Bitangent, Tangent); }\r\n");
+		OutHLSL += FString::Format(FormatSample, ArgsSample);
+	}
+	else if ( FunctionInfo.DefinitionName == Deprecated_GetTriPositionWSName )
+	{
+		static const TCHAR* FormatSample = TEXT("void {InstanceFunctionName}(MeshTriCoordinate TriCoord, out float3 Position) { float3 Dummy; GetTriangleWS_{ParameterName}(TriCoord.Tri, TriCoord.BaryCoord, Position, Dummy, Dummy, Dummy, Dummy); }\r\n");
+		OutHLSL += FString::Format(FormatSample, ArgsSample);
+	}
+	else if ( FunctionInfo.DefinitionName == Deprecated_GetTriNormalWSName )
+	{
+		static const TCHAR* FormatSample = TEXT("void {InstanceFunctionName}(MeshTriCoordinate TriCoord, out float3 Normal) { float3 Dummy; GetTriangleWS_{ParameterName}(TriCoord.Tri, TriCoord.BaryCoord, Dummy, Dummy, Normal, Dummy, Dummy); }\r\n");
+		OutHLSL += FString::Format(FormatSample, ArgsSample);
+	}
+	else if ( FunctionInfo.DefinitionName == Deprecated_GetTriTangentsWSName )
+	{
+		static const TCHAR* FormatSample = TEXT("void {InstanceFunctionName}(MeshTriCoordinate TriCoord, out float3 Tangent, out float3 Bitangent, out float3 Normal) { float3 Dummy; GetTriangleWS_{ParameterName}(TriCoord.Tri, TriCoord.BaryCoord, Dummy, Dummy, Normal, Bitangent, Tangent); }\r\n");
+		OutHLSL += FString::Format(FormatSample, ArgsSample);
+	}
+	else if (FunctionInfo.DefinitionName == Deprecated_GetTriColorName)
+	{
+		static const TCHAR* FormatSample = TEXT("void {InstanceFunctionName}(MeshTriCoordinate TriCoord, out float4 Color) { GetTriangleColor_{ParameterName}(TriCoord.Tri, TriCoord.BaryCoord, Color); }\r\n");
+		OutHLSL += FString::Format(FormatSample, ArgsSample);
+	}
+	else if (FunctionInfo.DefinitionName == Deprecated_GetTriUVName)
+	{
+		static const TCHAR* FormatSample = TEXT("void {InstanceFunctionName}(MeshTriCoordinate TriCoord, int UVSet, out float2 UV) { GetTriangleUV_{ParameterName}(TriCoord.Tri, TriCoord.BaryCoord, UVSet, UV); }\r\n");
+		OutHLSL += FString::Format(FormatSample, ArgsSample);
+	}
+	else if ( FunctionInfo.DefinitionName == Deprecated_GetTriPositionAndVelocityWSName )
+	{
+		static const TCHAR* FormatSample = TEXT("void {InstanceFunctionName}(MeshTriCoordinate TriCoord, out float3 Position, out float3 Velocity) { float3 Dummy; GetTriangleWS_{ParameterName}(TriCoord.Tri, TriCoord.BaryCoord, Position, Velocity, Dummy, Dummy, Dummy); }\r\n");
+		OutHLSL += FString::Format(FormatSample, ArgsSample);
+	}
+
+	return true;
+}
+#endif
+
+bool UNiagaraDataInterfaceStaticMesh::Equals(const UNiagaraDataInterface* Other) const
+{
+	if (!Super::Equals(Other))
+	{
+		return false;
+	}
+	const UNiagaraDataInterfaceStaticMesh* OtherTyped = CastChecked<const UNiagaraDataInterfaceStaticMesh>(Other);
+	return OtherTyped->SourceMode == SourceMode &&
+		OtherTyped->Source == Source &&
+		OtherTyped->SourceComponent == SourceComponent &&
+		OtherTyped->DefaultMesh == DefaultMesh &&
+		OtherTyped->SectionFilter.AllowedMaterialSlots == SectionFilter.AllowedMaterialSlots &&
+		OtherTyped->bUsePhysicsBodyVelocity == bUsePhysicsBodyVelocity &&
+		OtherTyped->FilteredSockets == FilteredSockets;
+}
+
+bool UNiagaraDataInterfaceStaticMesh::CopyToInternal(UNiagaraDataInterface* Destination) const
+{
+	if (!Super::CopyToInternal(Destination))
+	{
+		return false;
+	}
+
+	UNiagaraDataInterfaceStaticMesh* OtherTyped = CastChecked<UNiagaraDataInterfaceStaticMesh>(Destination);
+	OtherTyped->SourceMode = SourceMode;
+	OtherTyped->Source = Source;
+	OtherTyped->SourceComponent = SourceComponent;
+	OtherTyped->DefaultMesh = DefaultMesh;
+#if WITH_EDITORONLY_DATA
+	OtherTyped->PreviewMesh = PreviewMesh;
+#endif
+	OtherTyped->SectionFilter = SectionFilter;
+	OtherTyped->bUsePhysicsBodyVelocity = bUsePhysicsBodyVelocity;
+	OtherTyped->FilteredSockets = FilteredSockets;
+	return true;
+}
+
+#if WITH_EDITOR
+void UNiagaraDataInterfaceStaticMesh::GetFeedback(UNiagaraSystem* Asset, UNiagaraComponent* Component, TArray<FNiagaraDataInterfaceError>& OutErrors, TArray<FNiagaraDataInterfaceFeedback>& OutWarnings, TArray<FNiagaraDataInterfaceFeedback>& OutInfo)
+{
+	if (Source == nullptr && DefaultMesh != nullptr && !DefaultMesh->bAllowCPUAccess)
+	{
+		FNiagaraDataInterfaceError CPUAccessNotAllowedError(FText::Format(LOCTEXT("CPUAccessNotAllowedError", "This mesh needs CPU access in order to be used properly.({0})"), FText::FromString(DefaultMesh->GetName())),
+			LOCTEXT("CPUAccessNotAllowedErrorSummary", "CPU access error"),
+			FNiagaraDataInterfaceFix::CreateLambda([=]()
+		{
+			DefaultMesh->Modify();
+			DefaultMesh->bAllowCPUAccess = true;
+			return true;
+		}));
+
+		OutErrors.Add(CPUAccessNotAllowedError);
+	}
+
+	bool bHasNoMeshAssignedWarning = (Source == nullptr && DefaultMesh == nullptr);
+#if WITH_EDITORONLY_DATA
+	if (bHasNoMeshAssignedWarning)
+	{
+		if (UStaticMesh* LocalPreviewMesh = PreviewMesh.LoadSynchronous())
+		{
+			bHasNoMeshAssignedWarning = false;
+
+			if (!LocalPreviewMesh->bAllowCPUAccess)
+			{
+				FNiagaraDataInterfaceError CPUAccessNotAllowedError(FText::Format(LOCTEXT("CPUAccessNotAllowedError", "This mesh needs CPU access in order to be used properly.({0})"), FText::FromString(LocalPreviewMesh->GetName())),
+					LOCTEXT("CPUAccessNotAllowedErrorSummary", "CPU access error"),
+					FNiagaraDataInterfaceFix::CreateLambda([=]()
+					{
+						LocalPreviewMesh->Modify();
+						LocalPreviewMesh->bAllowCPUAccess = true;
+						return true;
+					}));
+
+				OutErrors.Add(CPUAccessNotAllowedError);
+			}
+		}
+	}
+#endif
+
+	if (bHasNoMeshAssignedWarning)
+	{
+		FNiagaraDataInterfaceFeedback NoMeshAssignedError(LOCTEXT("NoMeshAssignedError", "This Data Interface should be assigned a static mesh to operate correctly."),
+			LOCTEXT("NoMeshAssignedErrorSummary", "No mesh assigned warning"),
+			FNiagaraDataInterfaceFix());
+
+		OutWarnings.Add(NoMeshAssignedError);
+	}
+}
+#endif //WITH_EDITOR
+
+UStaticMesh* UNiagaraDataInterfaceStaticMesh::GetStaticMesh(USceneComponent*& OutComponent, class FNiagaraSystemInstance* SystemInstance)
+{
+	// Helper to scour an actor (or its parents) for a valid Static mesh component
+	auto FindActorMeshComponent = [](AActor* Actor, bool bRecurseParents = false) -> UStaticMeshComponent*
+	{
+		if (AStaticMeshActor* MeshActor = Cast<AStaticMeshActor>(Actor))
+		{
+			UStaticMeshComponent* Comp = MeshActor->GetStaticMeshComponent();
+			if (::IsValid(Comp))
+			{
+				return Comp;
+			}
+		}
+
+		// Fall back on any valid component on the actor
+		while (Actor)
+		{
+			for (UActorComponent* ActorComp : Actor->GetComponents())
+			{
+				UStaticMeshComponent* Comp = Cast<UStaticMeshComponent>(ActorComp);
+				if (::IsValid(Comp) && Comp->GetStaticMesh() != nullptr)
+				{
+					return Comp;
+				}
+			}
+
+			if (bRecurseParents)
+			{
+				Actor = Actor->GetParentActor();
+			}
+			else
+			{
+				break;
+			}
+		}
+
+		return nullptr;
+	};
+
+	UStaticMeshComponent* FoundMeshComponent = nullptr;	
+
+	const bool bTrySource = SourceMode == ENDIStaticMesh_SourceMode::Default || SourceMode == ENDIStaticMesh_SourceMode::Source;
+	const bool bTryAttachParent = SourceMode == ENDIStaticMesh_SourceMode::Default || SourceMode == ENDIStaticMesh_SourceMode::AttachParent;
+	const bool bTryDefaultMesh = SourceMode == ENDIStaticMesh_SourceMode::Default || SourceMode == ENDIStaticMesh_SourceMode::DefaultMeshOnly;
+
+	if (bTrySource && ::IsValid(SourceComponent))
+	{
+		FoundMeshComponent = SourceComponent;
+	}
+	else if (bTrySource && Source)
+	{
+		FoundMeshComponent = FindActorMeshComponent(Source);
+	}
+	else if (bTryAttachParent && SystemInstance)
+	{
+		if (USceneComponent* AttachComponent = SystemInstance->GetAttachComponent())
+		{
+			// First, try to find the mesh component up the attachment hierarchy
+			for (USceneComponent* Curr = AttachComponent; Curr; Curr = Curr->GetAttachParent())
+			{
+				UStaticMeshComponent* ParentComp = Cast<UStaticMeshComponent>(Curr);
+				if (::IsValid(ParentComp))
+				{
+					FoundMeshComponent = ParentComp;
+					break;
+				}
+			}
+			
+			if (!FoundMeshComponent)
+			{
+				// Next, try to find one in our outer chain
+				UStaticMeshComponent* OuterComp = AttachComponent->GetTypedOuter<UStaticMeshComponent>();
+				if (::IsValid(OuterComp))
+				{
+					FoundMeshComponent = OuterComp;
+				}
+				else if (AActor* Actor = AttachComponent->GetAttachmentRootActor())
+				{
+					// Final fall-back, look for any mesh component on our root actor or any of its parents
+					FoundMeshComponent = FindActorMeshComponent(Actor, true);
+				}
+			}
+		}
+	}
+
+	UStaticMesh* Mesh = nullptr;
+	OutComponent = nullptr;
+	if (FoundMeshComponent)
+	{
+		Mesh = FoundMeshComponent->GetStaticMesh();
+		OutComponent = FoundMeshComponent;
+	}
+	else if (bTryDefaultMesh)
+	{
+		Mesh = DefaultMesh;
+	}
+
+#if WITH_EDITORONLY_DATA
+	if (!Mesh && !FoundMeshComponent && (!SystemInstance || !SystemInstance->GetWorld()->IsGameWorld()))
+	{
+		// NOTE: We don't fall back on the preview mesh if we have a valid static mesh component referenced
+		Mesh = PreviewMesh.LoadSynchronous();		
+	}
+#endif
+
+	return Mesh;
+}
+
+void UNiagaraDataInterfaceStaticMesh::SetSourceComponentFromBlueprints(UStaticMeshComponent* ComponentToUse)
+{
+	// NOTE: When ChangeId changes the next tick will be skipped and a reset of the per-instance data will be initiated. 
+	++ChangeId;
+	SourceComponent = ComponentToUse;
+	Source = ComponentToUse->GetOwner();
+}
+
+void UNiagaraDataInterfaceStaticMesh::SetDefaultMeshFromBlueprints(UStaticMesh* MeshToUse)
+{
+	// NOTE: When ChangeId changes the next tick will be skipped and a reset of the per-instance data will be initiated. 
+	++ChangeId;
+	SourceComponent = nullptr;
+	Source = nullptr;
+	DefaultMesh = MeshToUse;
+}
+
+//////////////////////////////////////////////////////////////////////////
+// VM Vertex Sampling
+void UNiagaraDataInterfaceStaticMesh::VMIsValidVertex(FVectorVMExternalFunctionContext& Context)
+{
+	NDIStaticMeshLocal::FStaticMeshCpuHelper<FNDITransformHandlerNoop> StaticMeshHelper(Context);
+	FNDIInputParam<int32> VertexParam(Context);
+	FNDIOutputParam<bool> OutValid(Context);
+
+	const int32 NumVertices = StaticMeshHelper.GetNumPositionVertices();
+	for (int32 i = 0; i < Context.GetNumInstances(); ++i)
+	{
+		const int32 Vertex = VertexParam.GetAndAdvance();
+		OutValid.SetAndAdvance(Vertex >= 0 && Vertex < NumVertices);
+	}
+}
+
+void UNiagaraDataInterfaceStaticMesh::VMRandomVertex(FVectorVMExternalFunctionContext& Context)
+{
+	NDIStaticMeshLocal::FStaticMeshCpuHelper<FNDITransformHandlerNoop> StaticMeshHelper(Context);
+	FNDIRandomHelper RandHelper(Context);
+	FNDIOutputParam<int32> OutVertex(Context);
+
+	const int32 VertexMax = FMath::Max(StaticMeshHelper.GetNumPositionVertices() - 1, 0);
+	for (int32 i=0; i < Context.GetNumInstances(); ++i)
+	{
+		RandHelper.GetAndAdvance();
+		const int32 Vertex = RandHelper.RandRange(i, 0, VertexMax);
+		OutVertex.SetAndAdvance(Vertex);
+	}
+}
+
+void UNiagaraDataInterfaceStaticMesh::VMGetVertexCount(FVectorVMExternalFunctionContext& Context)
+{
+	NDIStaticMeshLocal::FStaticMeshCpuHelper<FNDITransformHandlerNoop> StaticMeshHelper(Context);
+	FNDIOutputParam<int32> OutNumVertices(Context);
+
+	const int32 NumVertices = StaticMeshHelper.GetNumPositionVertices();
+	for (int32 i = 0; i < Context.GetNumInstances(); ++i)
+	{
+		OutNumVertices.SetAndAdvance(NumVertices);
+	}
+}
+
+//void UNiagaraDataInterfaceStaticMesh::VMIsValidFilteredVertex(FVectorVMExternalFunctionContext& Context)
+//{
+//}
+
+//void UNiagaraDataInterfaceStaticMesh::VMRandomFilteredVertex(FVectorVMExternalFunctionContext& Context)
+//{
+//}
+
+//void UNiagaraDataInterfaceStaticMesh::VMGetFilteredVertexCount(FVectorVMExternalFunctionContext& Context)
+//{
+//}
+//
+//void UNiagaraDataInterfaceStaticMesh::VMGetFilteredVertexAt(FVectorVMExternalFunctionContext& Context)
+//{
+//}
+
+template<typename TTransformHandler>
+void UNiagaraDataInterfaceStaticMesh::VMGetVertex(FVectorVMExternalFunctionContext& Context)
+{
+	NDIStaticMeshLocal::FStaticMeshCpuHelper<TTransformHandler> StaticMeshHelper(Context);
+	FNDIInputParam<int32> VertexParam(Context);
+	FNDIOutputParam<FVector3f> OutPosition(Context);
+	FNDIOutputParam<FVector3f> OutVelocity(Context);
+	FNDIOutputParam<FVector3f> OutNormal(Context);
+	FNDIOutputParam<FVector3f> OutBitangent(Context);
+	FNDIOutputParam<FVector3f> OutTangent(Context);
+
+	const int32 PositionsMax = StaticMeshHelper.GetNumPositionVertices() - 1;
+	const int32 TangentsMax = StaticMeshHelper.GetNumTangentVertices() - 1;
+	if ( (PositionsMax >= 0) && (PositionsMax == TangentsMax) )
+	{
+		const float InvDt = StaticMeshHelper.GetInvDeltaSeconds();
+		for (int32 i = 0; i < Context.GetNumInstances(); ++i)
+		{
+			const int32 Vertex = FMath::Clamp(VertexParam.GetAndAdvance(), 0, PositionsMax);
+			const FVector3f LocalPosition = StaticMeshHelper.GetLocalPosition(Vertex);
+			const FVector3f Position = StaticMeshHelper.TransformPosition(LocalPosition);
+			const FVector3f PreviousPosition = StaticMeshHelper.PreviousTransformPosition(LocalPosition);
+			const FVector3f Velocity = (Position - PreviousPosition) * InvDt;
+			OutPosition.SetAndAdvance(Position);
+			OutVelocity.SetAndAdvance(Velocity);
+			OutNormal.SetAndAdvance(StaticMeshHelper.GetTangentZ(Vertex));
+			OutBitangent.SetAndAdvance(StaticMeshHelper.GetTangentY(Vertex));
+			OutTangent.SetAndAdvance(StaticMeshHelper.GetTangentX(Vertex));
+		}
+	}
+	else
+	{
+		const FVector3f Position = StaticMeshHelper.TransformPosition(FVector3f::ZeroVector);
+		const FVector3f Velocity = FVector3f::ZeroVector;
+		const FVector3f TangentX = StaticMeshHelper.TransformVector(FVector3f(1.0f, 0.0f, 0.0f));
+		const FVector3f TangentY = StaticMeshHelper.TransformVector(FVector3f(0.0f, 1.0f, 0.0f));
+		const FVector3f TangentZ = StaticMeshHelper.TransformVector(FVector3f(0.0f, 0.0f, 1.0f));
+		for ( int32 i=0; i < Context.GetNumInstances(); ++i )
+		{
+			OutPosition.SetAndAdvance(Position);
+			OutVelocity.SetAndAdvance(Velocity);
+			OutNormal.SetAndAdvance(TangentZ);
+			OutBitangent.SetAndAdvance(TangentY);
+			OutTangent.SetAndAdvance(TangentX);
+		}
+	}
+}
+
+void UNiagaraDataInterfaceStaticMesh::VMGetVertexColor(FVectorVMExternalFunctionContext& Context)
+{
+	NDIStaticMeshLocal::FStaticMeshCpuHelper StaticMeshHelper(Context);
+	FNDIInputParam<int32> VertexParam(Context);
+	FNDIOutputParam<FLinearColor> OutColor(Context);
+
+	const int32 ColorsMax = StaticMeshHelper.GetNumColorVertices() - 1;
+	if ( ColorsMax >= 0 )
+	{
+		for (int32 i = 0; i < Context.GetNumInstances(); ++i)
+		{
+			const int32 Vertex = VertexParam.GetAndAdvance();
+			OutColor.SetAndAdvance(StaticMeshHelper.GetColor(FMath::Clamp(Vertex, 0, ColorsMax)));
+		}
+	}
+	else
+	{
+		for (int32 i = 0; i < Context.GetNumInstances(); ++i)
+		{
+			OutColor.SetAndAdvance(FLinearColor::White);
+		}
+	}
+}
+
+void UNiagaraDataInterfaceStaticMesh::VMGetVertexUV(FVectorVMExternalFunctionContext& Context)
+{
+	NDIStaticMeshLocal::FStaticMeshCpuHelper StaticMeshHelper(Context);
+	FNDIInputParam<int32> VertexParam(Context);
+	FNDIInputParam<int32> UVSetParam(Context);
+	FNDIOutputParam<FVector2D> OutUV(Context);
+
+	const int32 VertexMax = StaticMeshHelper.GetNumUVVertices() - 1;
+	const int32 UVSetMax = StaticMeshHelper.GetNumUVs() - 1;
+	if (VertexMax >= 0 && UVSetMax >= 0)
+	{
+		for (int32 i = 0; i < Context.GetNumInstances(); ++i)
+		{
+			const int32 Vertex = VertexParam.GetAndAdvance();
+			const int32 UVSet = UVSetParam.GetAndAdvance();
+			const FVector2D UV = StaticMeshHelper.GetUV(FMath::Clamp(Vertex, 0, VertexMax), FMath::Clamp(UVSet, 0, UVSetMax));
+			OutUV.SetAndAdvance(UV);
+		}
+	}
+	else
+	{
+		for (int32 i = 0; i < Context.GetNumInstances(); ++i)
+		{
+			OutUV.SetAndAdvance(FVector2D::ZeroVector);
+		}
+	}
+}
+
+//////////////////////////////////////////////////////////////////////////
+// VM Triangle Sampling
+void UNiagaraDataInterfaceStaticMesh::VMIsValidTriangle(FVectorVMExternalFunctionContext& Context)
+{
+	NDIStaticMeshLocal::FStaticMeshCpuHelper<FNDITransformHandlerNoop> StaticMeshHelper(Context);
+	FNDIInputParam<int32> TriangleParam(Context);
+	FNDIOutputParam<bool> OutValid(Context);
+
+	const int32 NumTriangles = StaticMeshHelper.GetNumTriangles();
+	for (int32 i = 0; i < Context.GetNumInstances(); ++i)
+	{
+		const int32 Triangle = TriangleParam.GetAndAdvance();
+		OutValid.SetAndAdvance(Triangle >= 0 && Triangle < NumTriangles);
+	}
+}
+
+template<typename TRandomHelper>
+void UNiagaraDataInterfaceStaticMesh::VMRandomTriangle(FVectorVMExternalFunctionContext& Context)
+{
+	NDIStaticMeshLocal::FStaticMeshCpuHelper<FNDITransformHandlerNoop> StaticMeshHelper(Context);
+	TRandomHelper RandHelper(Context);
+	FNDIOutputParam<int32> OutTriangleParam(Context);
+	FNDIOutputParam<FVector3f> OutBaryCoordParam(Context);
+
+	if ( StaticMeshHelper.IsCpuUniformlyDistributedSampling() )
+	{
+		for (int32 i = 0; i < Context.GetNumInstances(); ++i)
+		{
+			RandHelper.GetAndAdvance();
+			OutTriangleParam.SetAndAdvance(StaticMeshHelper.RandomUniformTriangle(RandHelper, i));
+			OutBaryCoordParam.SetAndAdvance(RandHelper.RandomBarycentricCoord(i));
+		}
+	}
+	else
+	{
+		const int32 NumTriangles = StaticMeshHelper.GetNumTriangles();
+		const int32 TriangleMax = FMath::Max(NumTriangles - 1, 0);
+		for (int32 i = 0; i < Context.GetNumInstances(); ++i)
+		{
+			RandHelper.GetAndAdvance();
+			OutTriangleParam.SetAndAdvance(RandHelper.RandRange(i, 0, TriangleMax));
+			OutBaryCoordParam.SetAndAdvance(RandHelper.RandomBarycentricCoord(i));
+		}
+	}
+}
+
+void UNiagaraDataInterfaceStaticMesh::VMGetTriangleCount(FVectorVMExternalFunctionContext& Context)
+{
+	NDIStaticMeshLocal::FStaticMeshCpuHelper<FNDITransformHandlerNoop> StaticMeshHelper(Context);
+	FNDIOutputParam<int32> OutCountParam(Context);
+
+	const int32 NumTriangles = StaticMeshHelper.GetNumTriangles();
+	for (int32 i = 0; i < Context.GetNumInstances(); ++i)
+	{
+		OutCountParam.SetAndAdvance(NumTriangles);
+	}
+}
+
+void UNiagaraDataInterfaceStaticMesh::VMIsValidFilteredTriangle(FVectorVMExternalFunctionContext& Context)
+{
+	NDIStaticMeshLocal::FStaticMeshCpuHelper<FNDITransformHandlerNoop> StaticMeshHelper(Context);
+	FNDIInputParam<int32> TriangleParam(Context);
+	FNDIOutputParam<bool> OutValid(Context);
+
+	const int32 NumTriangles = StaticMeshHelper.GetNumFilteredTriangles();
+	for (int32 i = 0; i < Context.GetNumInstances(); ++i)
+	{
+		const int32 Triangle = TriangleParam.GetAndAdvance();
+		OutValid.SetAndAdvance(Triangle >= 0 && Triangle < NumTriangles);
+	}
+}
+
+void UNiagaraDataInterfaceStaticMesh::VMRandomFilteredTriangle(FVectorVMExternalFunctionContext& Context)
+{
+	NDIStaticMeshLocal::FStaticMeshCpuHelper<FNDITransformHandlerNoop> StaticMeshHelper(Context);
+	FNDIRandomHelper RandHelper(Context);
+	FNDIOutputParam<int32> OutTriangleParam(Context);
+	FNDIOutputParam<FVector3f> OutBaryCoordParam(Context);
+
+	TConstArrayView<NDIStaticMeshLocal::FNDISectionInfo> SectionInfos = StaticMeshHelper.InstanceData->GetFilteredSectionInfos();
+	if (SectionInfos.Num() )
+	{
+		if ( StaticMeshHelper.IsCpuUniformlyDistributedSampling() )
+		{
+			TConstArrayView<int32> SectionRemap = StaticMeshHelper.InstanceData->GetFilteredSections();
+			for (int32 i=0; i < Context.GetNumInstances(); ++i)
+			{
+				RandHelper.GetAndAdvance();
+				OutTriangleParam.SetAndAdvance(StaticMeshHelper.RandomUniformTriangle(RandHelper, i, SectionInfos, SectionRemap));
+				OutBaryCoordParam.SetAndAdvance(RandHelper.RandomBarycentricCoord(i));
+			}
+		}
+		else
+		{
+			for (int32 i = 0; i < Context.GetNumInstances(); ++i)
+			{
+				RandHelper.GetAndAdvance();
+				OutTriangleParam.SetAndAdvance(StaticMeshHelper.RandomTriangle(RandHelper, i, SectionInfos));
+				OutBaryCoordParam.SetAndAdvance(RandHelper.RandomBarycentricCoord(i));
+			}
+		}
+	}
+	else
+	{
+		for (int32 i = 0; i < Context.GetNumInstances(); ++i)
+		{
+			RandHelper.GetAndAdvance();
+			OutTriangleParam.SetAndAdvance(0);
+			OutBaryCoordParam.SetAndAdvance(RandHelper.RandomBarycentricCoord(i));
+		}
+	}
+}
+
+void UNiagaraDataInterfaceStaticMesh::VMGetFilteredTriangleCount(FVectorVMExternalFunctionContext& Context)
+{
+	NDIStaticMeshLocal::FStaticMeshCpuHelper<FNDITransformHandlerNoop> StaticMeshHelper(Context);
+	FNDIOutputParam<int32> OutCount(Context);
+
+	const int32 NumTriangles = StaticMeshHelper.GetNumFilteredTriangles();
+	for (int32 i = 0; i < Context.GetNumInstances(); ++i)
+	{
+		OutCount.SetAndAdvance(NumTriangles);
+	}
+}
+
+void UNiagaraDataInterfaceStaticMesh::VMGetFilteredTriangleAt(FVectorVMExternalFunctionContext& Context)
+{
+	NDIStaticMeshLocal::FStaticMeshCpuHelper<FNDITransformHandlerNoop> StaticMeshHelper(Context);
+	FNDIInputParam<int32> IndexParam(Context);
+	FNDIOutputParam<int32> OutTriangle(Context);
+
+	const int32 TriangleMax = StaticMeshHelper.GetNumFilteredTriangles() - 1;
+	if ( TriangleMax > 0 )
+	{
+		TConstArrayView<NDIStaticMeshLocal::FNDISectionInfo> FilteredSections = StaticMeshHelper.InstanceData->GetFilteredSectionInfos();
+		for (int32 i = 0; i < Context.GetNumInstances(); ++i)
+		{
+			const int32 Index = FMath::Clamp(IndexParam.GetAndAdvance(), 0, TriangleMax);
+			const int32 Triangle = StaticMeshHelper.SectionTriangleToTriangle(FilteredSections, Index);
+			OutTriangle.SetAndAdvance(Triangle);
+		}
+	}
+	else
+	{
+		for (int32 i = 0; i < Context.GetNumInstances(); ++i)
+		{
+			OutTriangle.SetAndAdvance(0);
+		}
+	}
+}
+
+void UNiagaraDataInterfaceStaticMesh::VMIsValidUnfilteredTriangle(FVectorVMExternalFunctionContext& Context)
+{
+	NDIStaticMeshLocal::FStaticMeshCpuHelper<FNDITransformHandlerNoop> StaticMeshHelper(Context);
+	FNDIInputParam<int32> TriangleParam(Context);
+	FNDIOutputParam<bool> OutValid(Context);
+
+	const int32 NumTriangles = StaticMeshHelper.GetNumUnfilteredTriangles();
+	for (int32 i = 0; i < Context.GetNumInstances(); ++i)
+	{
+		const int32 Triangle = TriangleParam.GetAndAdvance();
+		OutValid.SetAndAdvance(Triangle >= 0 && Triangle < NumTriangles);
+	}
+}
+
+void UNiagaraDataInterfaceStaticMesh::VMRandomUnfilteredTriangle(FVectorVMExternalFunctionContext& Context)
+{
+	NDIStaticMeshLocal::FStaticMeshCpuHelper<FNDITransformHandlerNoop> StaticMeshHelper(Context);
+	FNDIRandomHelper RandHelper(Context);
+	FNDIOutputParam<int32> OutTriangleParam(Context);
+	FNDIOutputParam<FVector3f> OutBaryCoordParam(Context);
+
+	TConstArrayView<NDIStaticMeshLocal::FNDISectionInfo> SectionInfos = StaticMeshHelper.InstanceData->GetUnfilteredSectionInfos();
+	if (SectionInfos.Num() )
+	{
+		if ( StaticMeshHelper.IsCpuUniformlyDistributedSampling() )
+		{
+			TConstArrayView<int32> SectionRemap = StaticMeshHelper.InstanceData->GetUnfilteredSections();
+			for (int32 i=0; i < Context.GetNumInstances(); ++i)
+			{
+				RandHelper.GetAndAdvance();
+				OutTriangleParam.SetAndAdvance(StaticMeshHelper.RandomUniformTriangle(RandHelper, i, SectionInfos, SectionRemap));
+				OutBaryCoordParam.SetAndAdvance(RandHelper.RandomBarycentricCoord(i));
+			}
+		}
+		else
+		{
+			for (int32 i = 0; i < Context.GetNumInstances(); ++i)
+			{
+				RandHelper.GetAndAdvance();
+				OutTriangleParam.SetAndAdvance(StaticMeshHelper.RandomTriangle(RandHelper, i, SectionInfos));
+				OutBaryCoordParam.SetAndAdvance(RandHelper.RandomBarycentricCoord(i));
+			}
+		}
+	}
+	else
+	{
+		for (int32 i = 0; i < Context.GetNumInstances(); ++i)
+		{
+			RandHelper.GetAndAdvance();
+			OutTriangleParam.SetAndAdvance(0);
+			OutBaryCoordParam.SetAndAdvance(RandHelper.RandomBarycentricCoord(i));
+		}
+	}
+}
+
+void UNiagaraDataInterfaceStaticMesh::VMGetUnfilteredTriangleCount(FVectorVMExternalFunctionContext& Context)
+{
+	NDIStaticMeshLocal::FStaticMeshCpuHelper<FNDITransformHandlerNoop> StaticMeshHelper(Context);
+	FNDIOutputParam<int32> OutCount(Context);
+
+	const int32 NumTriangles = StaticMeshHelper.GetNumUnfilteredTriangles();
+	for (int32 i = 0; i < Context.GetNumInstances(); ++i)
+	{
+		OutCount.SetAndAdvance(NumTriangles);
+	}
+}
+
+void UNiagaraDataInterfaceStaticMesh::VMGetUnfilteredTriangleAt(FVectorVMExternalFunctionContext& Context)
+{
+	NDIStaticMeshLocal::FStaticMeshCpuHelper<FNDITransformHandlerNoop> StaticMeshHelper(Context);
+	FNDIInputParam<int32> IndexParam(Context);
+	FNDIOutputParam<int32> OutTriangle(Context);
+
+	const int32 TriangleMax = StaticMeshHelper.GetNumUnfilteredTriangles() - 1;
+	if ( TriangleMax > 0 )
+	{
+		TConstArrayView<NDIStaticMeshLocal::FNDISectionInfo> UnfilteredSections = StaticMeshHelper.InstanceData->GetUnfilteredSectionInfos();
+		for (int32 i = 0; i < Context.GetNumInstances(); ++i)
+		{
+			const int32 Index = FMath::Clamp(IndexParam.GetAndAdvance(), 0, TriangleMax);
+			const int32 Triangle = StaticMeshHelper.SectionTriangleToTriangle(UnfilteredSections, Index);
+			OutTriangle.SetAndAdvance(Triangle);
+		}
+	}
+	else
+	{
+		for (int32 i = 0; i < Context.GetNumInstances(); ++i)
+		{
+			OutTriangle.SetAndAdvance(0);
+		}
+	}
+}
+
+template<typename TTransformHandler>
+void UNiagaraDataInterfaceStaticMesh::VMGetTriangle(FVectorVMExternalFunctionContext& Context)
+{
+	NDIStaticMeshLocal::FStaticMeshCpuHelper<TTransformHandler> StaticMeshHelper(Context);
+	FNDIInputParam<int32> TriangleParam(Context);
+	FNDIInputParam<FVector3f> BaryCoordParam(Context);
+	FNDIOutputParam<FVector3f> OutPositionParam(Context);
+	FNDIOutputParam<FVector3f> OutVelocityParam(Context);
+	FNDIOutputParam<FVector3f> OutNormalParam(Context);
+	FNDIOutputParam<FVector3f> OutBitangentParam(Context);
+	FNDIOutputParam<FVector3f> OutTangentParam(Context);
+
+	const FIndexArrayView IndexArray = StaticMeshHelper.GetIndexArrayView();
+	const int32 TriangleMax = StaticMeshHelper.GetNumTriangles() - 1;
+	if (TriangleMax >= 0 && StaticMeshHelper.GetNumPositionVertices() > 0 && StaticMeshHelper.GetNumTangentVertices() > 0)
+	{
+		for (int32 i = 0; i < Context.GetNumInstances(); ++i)
+		{
+			const int32 Triangle = FMath::Clamp(TriangleParam.GetAndAdvance(), 0, TriangleMax);
+			const FVector3f BaryCoord = BaryCoordParam.GetAndAdvance();
+			const FVector3f LocalPosition = StaticMeshHelper.GetLocalTrianglePosition(BaryCoord, IndexArray[Triangle * 3 + 0], IndexArray[Triangle * 3 + 1], IndexArray[Triangle * 3 + 2]);
+			const FVector3f Position = StaticMeshHelper.TransformPosition(LocalPosition);
+			const FVector3f PreviousPosition = StaticMeshHelper.PreviousTransformPosition(LocalPosition);
+			const FVector3f Velocity = (Position - PreviousPosition) * StaticMeshHelper.GetInvDeltaSeconds();
+			const FVector3f TangentX = StaticMeshHelper.GetTriangleTangentX(BaryCoord, IndexArray[Triangle * 3 + 0], IndexArray[Triangle * 3 + 1], IndexArray[Triangle * 3 + 2]);
+			const FVector3f TangentY = StaticMeshHelper.GetTriangleTangentY(BaryCoord, IndexArray[Triangle * 3 + 0], IndexArray[Triangle * 3 + 1], IndexArray[Triangle * 3 + 2]);
+			const FVector3f TangentZ = StaticMeshHelper.GetTriangleTangentZ(BaryCoord, IndexArray[Triangle * 3 + 0], IndexArray[Triangle * 3 + 1], IndexArray[Triangle * 3 + 2]);
+			OutPositionParam.SetAndAdvance(Position);
+			OutVelocityParam.SetAndAdvance(Velocity);
+			OutNormalParam.SetAndAdvance(TangentZ);
+			OutBitangentParam.SetAndAdvance(TangentY);
+			OutTangentParam.SetAndAdvance(TangentX);
+		}
+	}
+	else
+	{
+		const FVector3f Position = StaticMeshHelper.TransformPosition(FVector3f::ZeroVector);
+		const FVector3f PreviousPosition = StaticMeshHelper.PreviousTransformPosition(FVector3f::ZeroVector);
+		const FVector3f Velocity = FVector3f::ZeroVector;
+		const FVector3f TangentX = StaticMeshHelper.TransformVector(FVector3f(1.0f, 0.0f, 0.0f));
+		const FVector3f TangentY = StaticMeshHelper.TransformVector(FVector3f(0.0f, 1.0f, 0.0f));
+		const FVector3f TangentZ = StaticMeshHelper.TransformVector(FVector3f(0.0f, 0.0f, 1.0f));
+		for (int32 i = 0; i < Context.GetNumInstances(); ++i)
+		{
+			OutPositionParam.SetAndAdvance(Position);
+			OutVelocityParam.SetAndAdvance(Velocity);
+			OutNormalParam.SetAndAdvance(TangentZ);
+			OutBitangentParam.SetAndAdvance(TangentY);
+			OutTangentParam.SetAndAdvance(TangentX);
+		}
+	}
+}
+
+void UNiagaraDataInterfaceStaticMesh::VMGetTriangleColor(FVectorVMExternalFunctionContext& Context)
+{
+	NDIStaticMeshLocal::FStaticMeshCpuHelper<FNDITransformHandlerNoop> StaticMeshHelper(Context);
+	FNDIInputParam<int32> TriangleParam(Context);
+	FNDIInputParam<FVector3f> BaryCoordParam(Context);
+	FNDIOutputParam<FLinearColor> OutColorParam(Context);
+
+	const FIndexArrayView IndexArray = StaticMeshHelper.GetIndexArrayView();
+	const int32 TriangleMax = StaticMeshHelper.GetNumTriangles() - 1;
+	if (TriangleMax >= 0 && StaticMeshHelper.GetNumColorVertices() > 0)
+	{
+		for (int32 i = 0; i < Context.GetNumInstances(); ++i)
+		{
+			const int32 Triangle = FMath::Clamp(TriangleParam.GetAndAdvance(), 0, TriangleMax);
+			const FVector3f BaryCoord = BaryCoordParam.GetAndAdvance();
+			const FLinearColor Color = StaticMeshHelper.GetTriangleColor(BaryCoord, IndexArray[Triangle * 3 + 0], IndexArray[Triangle * 3 + 1], IndexArray[Triangle * 3 + 2]);
+			OutColorParam.SetAndAdvance(Color);
+		}
+	}
+	else
+	{
+		for (int32 i = 0; i < Context.GetNumInstances(); ++i)
+		{
+			OutColorParam.SetAndAdvance(FLinearColor::White);
+		}
+	}
+}
+
+void UNiagaraDataInterfaceStaticMesh::VMGetTriangleUV(FVectorVMExternalFunctionContext& Context)
+{
+	NDIStaticMeshLocal::FStaticMeshCpuHelper<FNDITransformHandlerNoop> StaticMeshHelper(Context);
+	FNDIInputParam<int32> TriangleParam(Context);
+	FNDIInputParam<FVector3f> BaryCoordParam(Context);
+	FNDIInputParam<int32> UVSetParam(Context);
+	FNDIOutputParam<FVector2D> OutUVParam(Context);
+
+	const FIndexArrayView IndexArray = StaticMeshHelper.GetIndexArrayView();
+	const int32 TriangleMax = StaticMeshHelper.GetNumTriangles() - 1;
+	const int32 UVSetMax = StaticMeshHelper.GetNumUVs() - 1;
+	if (TriangleMax >= 0 && UVSetMax >= 0)
+	{
+		for (int32 i = 0; i < Context.GetNumInstances(); ++i)
+		{
+			const int32 Triangle = FMath::Clamp(TriangleParam.GetAndAdvance(), 0, TriangleMax);
+			const FVector3f BaryCoord = BaryCoordParam.GetAndAdvance();
+			const int32 UVSet = FMath::Clamp(UVSetParam.GetAndAdvance(), 0, UVSetMax);
+			const FVector2D UV = StaticMeshHelper.GetTriangleUV(BaryCoord, IndexArray[Triangle * 3 + 0], IndexArray[Triangle * 3 + 1], IndexArray[Triangle * 3 + 2], UVSet);
+			OutUVParam.SetAndAdvance(UV);
+		}
+	}
+	else
+	{
+		for (int32 i = 0; i < Context.GetNumInstances(); ++i)
+		{
+			OutUVParam.SetAndAdvance(FVector2D::ZeroVector);
+		}
+	}
+}
+
+void UNiagaraDataInterfaceStaticMesh::VMGetTriangleIndices(FVectorVMExternalFunctionContext& Context)
+{
+	NDIStaticMeshLocal::FStaticMeshCpuHelper<FNDITransformHandlerNoop> StaticMeshHelper(Context);
+	FNDIInputParam<int32> TriangleParam(Context);
+	FNDIOutputParam<int32> OutIndex0Param(Context);
+	FNDIOutputParam<int32> OutIndex1Param(Context);
+	FNDIOutputParam<int32> OutIndex2Param(Context);
+
+	const FIndexArrayView IndexArray = StaticMeshHelper.GetIndexArrayView();
+	const int32 TriangleMax = StaticMeshHelper.GetNumTriangles() - 1;
+	if (TriangleMax >= 0)
+	{
+		for (int32 i = 0; i < Context.GetNumInstances(); ++i)
+		{
+			const int32 Triangle = FMath::Clamp(TriangleParam.GetAndAdvance(), 0, TriangleMax);
+			OutIndex0Param.SetAndAdvance(IndexArray[Triangle * 3 + 0]);
+			OutIndex1Param.SetAndAdvance(IndexArray[Triangle * 3 + 1]);
+			OutIndex2Param.SetAndAdvance(IndexArray[Triangle * 3 + 2]);
+		}
+	}
+	else
+	{
+		for (int32 i = 0; i < Context.GetNumInstances(); ++i)
+		{
+			OutIndex0Param.SetAndAdvance(0);
+			OutIndex1Param.SetAndAdvance(0);
+			OutIndex2Param.SetAndAdvance(0);
+		}
+	}
+}
+
+//////////////////////////////////////////////////////////////////////////
+// Socket Functions
+void UNiagaraDataInterfaceStaticMesh::VMGetSocketCount(FVectorVMExternalFunctionContext& Context)
+{
+	VectorVM::FUserPtrHandler<NDIStaticMeshLocal::FInstanceData_GameThread> InstanceData(Context);
+	FNDIOutputParam<int32> OutCountParam(Context);
+
+	for (int32 i = 0; i < Context.GetNumInstances(); ++i)
+	{
+		OutCountParam.SetAndAdvance(InstanceData->CachedSockets.Num());
+	}
+}
+
+void UNiagaraDataInterfaceStaticMesh::VMGetFilteredSocketCount(FVectorVMExternalFunctionContext& Context)
+{
+	VectorVM::FUserPtrHandler<NDIStaticMeshLocal::FInstanceData_GameThread> InstanceData(Context);
+	FNDIOutputParam<int32> OutCountParam(Context);
+
+	for (int32 i = 0; i < Context.GetNumInstances(); ++i)
+	{
+		OutCountParam.SetAndAdvance(InstanceData->NumFilteredSockets);
+	}
+}
+
+void UNiagaraDataInterfaceStaticMesh::VMGetUnfilteredSocketCount(FVectorVMExternalFunctionContext& Context)
+{
+	VectorVM::FUserPtrHandler<NDIStaticMeshLocal::FInstanceData_GameThread> InstanceData(Context);
+	FNDIOutputParam<int32> OutCountParam(Context);
+
+	const int32 NumUnfilteredSockets = InstanceData->CachedSockets.Num() - InstanceData->NumFilteredSockets;
+	for (int32 i = 0; i < Context.GetNumInstances(); ++i)
+	{
+		OutCountParam.SetAndAdvance(NumUnfilteredSockets);
+	}
+}
+
+template<typename TTransformHandler>
+void UNiagaraDataInterfaceStaticMesh::VMGetSocketTransform(FVectorVMExternalFunctionContext& Context)
+{
+	NDIStaticMeshLocal::FStaticMeshCpuHelper<TTransformHandler> StaticMeshHelper(Context);
+	FNDIInputParam<int32> SocketIndexParam(Context);
+	FNDIOutputParam<FVector3f> OutTranslateParam(Context);
+	FNDIOutputParam<FQuat4f> OutRotateParam(Context);
+	FNDIOutputParam<FVector3f> OutScaleParam(Context);
+
+	const int32 SocketMax = StaticMeshHelper.InstanceData->CachedSockets.Num() - 1;
+	if (SocketMax >= 0)
+	{
+		for (int32 i = 0; i < Context.GetNumInstances(); ++i)
+		{
+			const int32 SocketIndex = FMath::Clamp(SocketIndexParam.GetAndAdvance(), 0, SocketMax);
+			const FTransform3f& SocketTransform = StaticMeshHelper.InstanceData->CachedSockets[SocketIndex];
+			OutTranslateParam.SetAndAdvance(StaticMeshHelper.TransformPosition(SocketTransform.GetTranslation()));
+			OutRotateParam.SetAndAdvance(StaticMeshHelper.TransformRotation(SocketTransform.GetRotation()));
+			OutScaleParam.SetAndAdvance(StaticMeshHelper.TransformVector(SocketTransform.GetScale3D()));
+		}
+	}
+	else
+	{
+		const FVector3f DefaultTranslate = StaticMeshHelper.TransformPosition(FVector3f::ZeroVector);
+		const FQuat4f DefaultRotation = StaticMeshHelper.TransformRotation(FQuat4f::Identity);
+		const FVector3f DefaultScale = StaticMeshHelper.TransformVector(FVector::OneVector);
+		for (int32 i = 0; i < Context.GetNumInstances(); ++i)
+		{
+			OutTranslateParam.SetAndAdvance(DefaultTranslate);
+			OutRotateParam.SetAndAdvance(DefaultRotation);
+			OutScaleParam.SetAndAdvance(DefaultScale);
+		}
+	}
+}
+
+template<typename TTransformHandler>
+void UNiagaraDataInterfaceStaticMesh::VMGetFilteredSocketTransform(FVectorVMExternalFunctionContext& Context)
+{
+	NDIStaticMeshLocal::FStaticMeshCpuHelper<TTransformHandler> StaticMeshHelper(Context);
+	FNDIInputParam<int32> SocketIndexParam(Context);
+	FNDIOutputParam<FVector3f> OutTranslateParam(Context);
+	FNDIOutputParam<FQuat4f> OutRotateParam(Context);
+	FNDIOutputParam<FVector3f> OutScaleParam(Context);
+
+	const int32 SocketMax = StaticMeshHelper.InstanceData->NumFilteredSockets - 1;
+	if (SocketMax >= 0)
+	{
+		for (int32 i = 0; i < Context.GetNumInstances(); ++i)
+		{
+			const int32 FilteredSocketIndex = FMath::Clamp(SocketIndexParam.GetAndAdvance(), 0, SocketMax);
+			const int32 SocketIndex = StaticMeshHelper.InstanceData->FilteredAndUnfilteredSockets[FilteredSocketIndex];
+			const FTransform3f& SocketTransform = StaticMeshHelper.InstanceData->CachedSockets[SocketIndex];
+			OutTranslateParam.SetAndAdvance(StaticMeshHelper.TransformPosition(SocketTransform.GetTranslation()));
+			OutRotateParam.SetAndAdvance(StaticMeshHelper.TransformRotation(SocketTransform.GetRotation()));
+			OutScaleParam.SetAndAdvance(StaticMeshHelper.TransformVector(SocketTransform.GetScale3D()));
+		}
+	}
+	else
+	{
+		const FVector3f DefaultTranslate = StaticMeshHelper.TransformPosition(FVector3f::ZeroVector);
+		const FQuat4f DefaultRotation = StaticMeshHelper.TransformRotation(FQuat4f::Identity);
+		const FVector3f DefaultScale = StaticMeshHelper.TransformVector(FVector::OneVector);
+		for (int32 i = 0; i < Context.GetNumInstances(); ++i)
+		{
+			OutTranslateParam.SetAndAdvance(DefaultTranslate);
+			OutRotateParam.SetAndAdvance(DefaultRotation);
+			OutScaleParam.SetAndAdvance(DefaultScale);
+		}
+	}
+}
+
+template<typename TTransformHandler>
+void UNiagaraDataInterfaceStaticMesh::VMGetUnfilteredSocketTransform(FVectorVMExternalFunctionContext& Context)
+{
+	NDIStaticMeshLocal::FStaticMeshCpuHelper<TTransformHandler> StaticMeshHelper(Context);
+	FNDIInputParam<int32> SocketIndexParam(Context);
+	FNDIOutputParam<FVector3f> OutTranslateParam(Context);
+	FNDIOutputParam<FQuat4f> OutRotateParam(Context);
+	FNDIOutputParam<FVector3f> OutScaleParam(Context);
+
+	const int32 UnfilteredOffset = StaticMeshHelper.InstanceData->NumFilteredSockets;
+	const int32 SocketMax = StaticMeshHelper.InstanceData->FilteredAndUnfilteredSockets.Num() - UnfilteredOffset - 1;
+	if (SocketMax >= 0)
+	{
+		for (int32 i = 0; i < Context.GetNumInstances(); ++i)
+		{
+			const int32 UnfilteredSocketIndex = FMath::Clamp(SocketIndexParam.GetAndAdvance(), 0, SocketMax) + UnfilteredOffset;
+			const int32 SocketIndex = StaticMeshHelper.InstanceData->FilteredAndUnfilteredSockets[UnfilteredSocketIndex];
+			const FTransform3f& SocketTransform = StaticMeshHelper.InstanceData->CachedSockets[SocketIndex];
+			OutTranslateParam.SetAndAdvance(StaticMeshHelper.TransformPosition(SocketTransform.GetTranslation()));
+			OutRotateParam.SetAndAdvance(StaticMeshHelper.TransformRotation(SocketTransform.GetRotation()));
+			OutScaleParam.SetAndAdvance(StaticMeshHelper.TransformVector(SocketTransform.GetScale3D()));
+		}
+	}
+	else
+	{
+		const FVector3f DefaultTranslate = StaticMeshHelper.TransformPosition(FVector3f::ZeroVector);
+		const FQuat4f DefaultRotation = StaticMeshHelper.TransformRotation(FQuat4f::Identity);
+		const FVector3f DefaultScale = StaticMeshHelper.TransformVector(FVector::OneVector);
+		for (int32 i = 0; i < Context.GetNumInstances(); ++i)
+		{
+			OutTranslateParam.SetAndAdvance(DefaultTranslate);
+			OutRotateParam.SetAndAdvance(DefaultRotation);
+			OutScaleParam.SetAndAdvance(DefaultScale);
+		}
+	}
+}
+
+//////////////////////////////////////////////////////////////////////////
+// Section functions
+void UNiagaraDataInterfaceStaticMesh::VMIsValidSection(FVectorVMExternalFunctionContext& Context)
+{
+	NDIStaticMeshLocal::FStaticMeshCpuHelper<FNDITransformHandlerNoop> StaticMeshHelper(Context);
+	FNDIInputParam<int32> SectionIndexParam(Context);
+	FNDIOutputParam<bool> OutIsValid(Context);
+
+	const int32 NumSections = StaticMeshHelper.LODResource ? StaticMeshHelper.LODResource->Sections.Num() : 0;
+	for ( int32 i=0; i < Context.GetNumInstances(); ++i )
+	{
+		const int32 Section = SectionIndexParam.GetAndAdvance();
+		OutIsValid.SetAndAdvance(Section >= 0 && Section < NumSections);
+	}
+}
+
+void UNiagaraDataInterfaceStaticMesh::VMGetSectionTriangleCount(FVectorVMExternalFunctionContext& Context)
+{
+	NDIStaticMeshLocal::FStaticMeshCpuHelper<FNDITransformHandlerNoop> StaticMeshHelper(Context);
+	FNDIInputParam<int32> SectionIndexParam(Context);
+	FNDIOutputParam<int32> OutTriangleCount(Context);
+
+	const int32 SectionMax = StaticMeshHelper.LODResource ? StaticMeshHelper.LODResource->Sections.Num() : 0;
+	for (int32 i = 0; i < Context.GetNumInstances(); ++i)
+	{
+		const int32 Section = SectionIndexParam.GetAndAdvance();
+		const bool ValidSection = Section >= 0 && Section < SectionMax;
+		OutTriangleCount.SetAndAdvance(ValidSection ? StaticMeshHelper.LODResource->Sections[Section].NumTriangles : 0);
+	}
+}
+
+template<typename TRandomHandler>
+void UNiagaraDataInterfaceStaticMesh::VMRandomSectionTriangle(FVectorVMExternalFunctionContext& Context)
+{
+	NDIStaticMeshLocal::FStaticMeshCpuHelper<FNDITransformHandlerNoop> StaticMeshHelper(Context);
+	TRandomHandler RandHelper(Context);
+	FNDIInputParam<int32> SectionIndexParam(Context);
+	FNDIOutputParam<int32> OutTriangle(Context);
+	FNDIOutputParam<FVector3f> OutBaryCoord(Context);
+
+	const int32 SectionMax = StaticMeshHelper.LODResource ? StaticMeshHelper.LODResource->Sections.Num() - 1 : 0;
+	if (SectionMax >= 0)
+	{
+		if ( StaticMeshHelper.IsCpuUniformlyDistributedSampling() )
+		{
+			for (int32 i = 0; i < Context.GetNumInstances(); ++i)
+			{
+				RandHelper.GetAndAdvance();
+				const int32 Section = FMath::Clamp(SectionIndexParam.GetAndAdvance(), 0, SectionMax);
+				OutTriangle.SetAndAdvance(StaticMeshHelper.RandomSectionUniformTriangle(RandHelper, i, Section));
+				OutBaryCoord.SetAndAdvance(RandHelper.RandomBarycentricCoord(i));
+			}
+		}
+		else
+		{
+			for (int32 i = 0; i < Context.GetNumInstances(); ++i)
+			{
+				RandHelper.GetAndAdvance();
+				const int32 Section = FMath::Clamp(SectionIndexParam.GetAndAdvance(), 0, SectionMax);
+				OutTriangle.SetAndAdvance(StaticMeshHelper.RandomSectionTriangle(RandHelper, i, Section));
+				OutBaryCoord.SetAndAdvance(RandHelper.RandomBarycentricCoord(i));
+			}
+		}
+	}
+	else
+	{
+		for (int32 i = 0; i < Context.GetNumInstances(); ++i)
+		{
+			RandHelper.GetAndAdvance();
+			OutTriangle.SetAndAdvance(0);
+			OutBaryCoord.SetAndAdvance(RandHelper.RandomBarycentricCoord(i));
+		}
+	}
+}
+
+void UNiagaraDataInterfaceStaticMesh::VMGetSectionTriangleAt(FVectorVMExternalFunctionContext& Context)
+{
+	NDIStaticMeshLocal::FStaticMeshCpuHelper<FNDITransformHandlerNoop> StaticMeshHelper(Context);
+	FNDIInputParam<int32> SectionIndexParam(Context);
+	FNDIInputParam<int32> TriangleIndexParam(Context);
+	FNDIOutputParam<int32> OutTriangle(Context);
+
+	const int32 NumSections = StaticMeshHelper.LODResource ? StaticMeshHelper.LODResource->Sections.Num() : 0;
+	if (NumSections >= 0)
+	{
+		for (int32 i = 0; i < Context.GetNumInstances(); ++i)
+		{
+			const int32 Section = SectionIndexParam.GetAndAdvance();
+			const int32 SectionTriangleIndex = TriangleIndexParam.GetAndAdvance();
+			const bool SectionValid = Section >= 0 && Section < NumSections;
+			const int32 SectionTriangleMax = SectionValid ? StaticMeshHelper.LODResource->Sections[Section].NumTriangles - 1 : 0;
+			const int32 SectionTriangleOffset = SectionValid ? StaticMeshHelper.LODResource->Sections[Section].FirstIndex / 3 : 0;
+			OutTriangle.SetAndAdvance(FMath::Clamp(SectionTriangleIndex, 0, SectionTriangleMax) + SectionTriangleOffset);
+		}
+	}
+	else
+	{
+		for (int32 i = 0; i < Context.GetNumInstances(); ++i)
+		{
+			OutTriangle.SetAndAdvance(0);
+		}
+	}
+}
+
+void UNiagaraDataInterfaceStaticMesh::VMGetFilteredSectionAt(FVectorVMExternalFunctionContext& Context)
+{
+	NDIStaticMeshLocal::FStaticMeshCpuHelper<FNDITransformHandlerNoop> StaticMeshHelper(Context);
+	FNDIInputParam<int32> SectionIndexParam(Context);
+	FNDIOutputParam<int32> OutSection(Context);
+
+	TConstArrayView<int32> FilteredSections = StaticMeshHelper.InstanceData->GetFilteredSections();
+	const int32 SectionMax = FilteredSections.Num() - 1;
+	if ( SectionMax >= 0 )
+	{
+		for (int32 i = 0; i < Context.GetNumInstances(); ++i)
+		{
+			const int32 SectionIndex = FMath::Clamp(SectionIndexParam.GetAndAdvance(), 0, SectionMax);
+			OutSection.SetAndAdvance(FilteredSections[SectionIndex]);
+		}
+	}
+	else
+	{
+		for (int32 i = 0; i < Context.GetNumInstances(); ++i)
+		{
+			OutSection.SetAndAdvance(0);
+		}
+	}
+}
+
+void UNiagaraDataInterfaceStaticMesh::VMGetUnfilteredSectionAt(FVectorVMExternalFunctionContext& Context)
+{
+	NDIStaticMeshLocal::FStaticMeshCpuHelper<FNDITransformHandlerNoop> StaticMeshHelper(Context);
+	FNDIInputParam<int32> SectionIndexParam(Context);
+	FNDIOutputParam<int32> OutSection(Context);
+
+	TConstArrayView<int32> UnfilteredSections = StaticMeshHelper.InstanceData->GetUnfilteredSections();
+	const int32 SectionMax = UnfilteredSections.Num() - 1;
+	if (SectionMax >= 0)
+	{
+		for (int32 i = 0; i < Context.GetNumInstances(); ++i)
+		{
+			const int32 SectionIndex = FMath::Clamp(SectionIndexParam.GetAndAdvance(), 0, SectionMax);
+			OutSection.SetAndAdvance(UnfilteredSections[SectionIndex]);
+		}
+	}
+	else
+	{
+		for (int32 i = 0; i < Context.GetNumInstances(); ++i)
+		{
+			OutSection.SetAndAdvance(0);
+		}
+	}
+}
+
+void UNiagaraDataInterfaceStaticMesh::VMGetSectionCount(FVectorVMExternalFunctionContext& Context)
+{
+	NDIStaticMeshLocal::FStaticMeshCpuHelper<FNDITransformHandlerNoop> StaticMeshHelper(Context);
+	FNDIOutputParam<int32> OutSectionCount(Context);
+
+	const int32 NumSections = StaticMeshHelper.LODResource ? StaticMeshHelper.LODResource->Sections.Num() : 0;
+	for (int32 i = 0; i < Context.GetNumInstances(); ++i)
+	{
+		OutSectionCount.SetAndAdvance(NumSections);
+	}
+}
+
+void UNiagaraDataInterfaceStaticMesh::VMGetFilteredSectionCount(FVectorVMExternalFunctionContext& Context)
+{
+	NDIStaticMeshLocal::FStaticMeshCpuHelper<FNDITransformHandlerNoop> StaticMeshHelper(Context);
+	FNDIOutputParam<int32> OutSectionCount(Context);
+
+	const int32 NumSections = StaticMeshHelper.InstanceData->NumFilteredSections;
+	for (int32 i = 0; i < Context.GetNumInstances(); ++i)
+	{
+		OutSectionCount.SetAndAdvance(NumSections);
+	}
+}
+
+void UNiagaraDataInterfaceStaticMesh::VMGetUnfilteredSectionCount(FVectorVMExternalFunctionContext& Context)
+{
+	NDIStaticMeshLocal::FStaticMeshCpuHelper<FNDITransformHandlerNoop> StaticMeshHelper(Context);
+	FNDIOutputParam<int32> OutSectionCount(Context);
+
+	const int32 NumSections = StaticMeshHelper.InstanceData->NumUnfilteredSections;
+	for (int32 i = 0; i < Context.GetNumInstances(); ++i)
+	{
+		OutSectionCount.SetAndAdvance(NumSections);
+	}
+}
+
+template<typename TRandomHandler>
+void UNiagaraDataInterfaceStaticMesh::VMRandomSection(FVectorVMExternalFunctionContext& Context)
+{
+	NDIStaticMeshLocal::FStaticMeshCpuHelper<FNDITransformHandlerNoop> StaticMeshHelper(Context);
+	TRandomHandler RandHelper(Context);
+	FNDIOutputParam<int32> OutSection(Context);
+
+	if ( StaticMeshHelper.IsCpuUniformlyDistributedSampling() )
+	{
+		for (int32 i = 0; i < Context.GetNumInstances(); ++i)
+		{
+			RandHelper.GetAndAdvance();
+			OutSection.SetAndAdvance(StaticMeshHelper.RandomUniformSection(RandHelper, i));
+		}
+	}
+	else
+	{
+		const int32 SectionMax = StaticMeshHelper.LODResource ? StaticMeshHelper.LODResource->Sections.Num() - 1 : 0;
+		for (int32 i = 0; i < Context.GetNumInstances(); ++i)
+		{
+			RandHelper.GetAndAdvance();
+			OutSection.SetAndAdvance(RandHelper.RandRange(i, 0, SectionMax));
+		}
+	}
+}
+
+template<typename TRandomHandler>
+void UNiagaraDataInterfaceStaticMesh::VMRandomFilteredSection(FVectorVMExternalFunctionContext& Context)
+{
+	NDIStaticMeshLocal::FStaticMeshCpuHelper<FNDITransformHandlerNoop> StaticMeshHelper(Context);
+	TRandomHandler RandHelper(Context);
+	FNDIOutputParam<int32> OutSection(Context);
+
+	TConstArrayView<int32> FilteredSections = StaticMeshHelper.InstanceData->GetFilteredSections();
+	const int32 SectionMax = FilteredSections.Num() - 1;
+	if (SectionMax >= 0)
+	{
+		if ( StaticMeshHelper.IsCpuUniformlyDistributedSampling() )
+		{
+			TConstArrayView<NDIStaticMeshLocal::FNDISectionInfo> FilteredSectionInfos = StaticMeshHelper.InstanceData->GetFilteredSectionInfos();
+			for (int32 i = 0; i < Context.GetNumInstances(); ++i)
+			{
+				RandHelper.GetAndAdvance();
+				const int32 Section = StaticMeshHelper.RandomUniformSection(RandHelper, i, FilteredSectionInfos, FilteredSections);
+				OutSection.SetAndAdvance(FilteredSections[Section]);
+			}
+		}
+		else
+		{
+			for (int32 i = 0; i < Context.GetNumInstances(); ++i)
+			{
+				RandHelper.GetAndAdvance();
+				OutSection.SetAndAdvance(FilteredSections[RandHelper.RandRange(i, 0, SectionMax)]);
+			}
+		}
+	}
+	else
+	{
+		for (int32 i = 0; i < Context.GetNumInstances(); ++i)
+		{
+			RandHelper.GetAndAdvance();
+			OutSection.SetAndAdvance(0);
+		}
+	}
+}
+
+void UNiagaraDataInterfaceStaticMesh::VMRandomUnfilteredSection(FVectorVMExternalFunctionContext& Context)
+{
+	NDIStaticMeshLocal::FStaticMeshCpuHelper<FNDITransformHandlerNoop> StaticMeshHelper(Context);
+	FNDIRandomHelper RandHelper(Context);
+	FNDIOutputParam<int32> OutSection(Context);
+
+	TConstArrayView<int32> UnfilteredSections = StaticMeshHelper.InstanceData->GetUnfilteredSections();
+	const int32 SectionMax = UnfilteredSections.Num() - 1;
+	if (SectionMax >= 0)
+	{
+		if (StaticMeshHelper.IsCpuUniformlyDistributedSampling())
+		{
+			TConstArrayView<NDIStaticMeshLocal::FNDISectionInfo> UnfilteredSectionInfos = StaticMeshHelper.InstanceData->GetUnfilteredSectionInfos();
+			for (int32 i = 0; i < Context.GetNumInstances(); ++i)
+			{
+				RandHelper.GetAndAdvance();
+				const int32 Section = StaticMeshHelper.RandomUniformSection(RandHelper, i, UnfilteredSectionInfos, UnfilteredSections);
+				OutSection.SetAndAdvance(UnfilteredSections[Section]);
+			}
+		}
+		else
+		{
+			for (int32 i = 0; i < Context.GetNumInstances(); ++i)
+			{
+				RandHelper.GetAndAdvance();
+				OutSection.SetAndAdvance(UnfilteredSections[RandHelper.RandRange(i, 0, SectionMax)]);
+			}
+		}
+	}
+	else
+	{
+		for (int32 i = 0; i < Context.GetNumInstances(); ++i)
+		{
+			RandHelper.GetAndAdvance();
+			OutSection.SetAndAdvance(0);
+		}
+	}
+}
+
+//////////////////////////////////////////////////////////////////////////
+// VM Misc Functions
+void UNiagaraDataInterfaceStaticMesh::VMIsValid(FVectorVMExternalFunctionContext& Context)
+{
+	VectorVM::FUserPtrHandler<NDIStaticMeshLocal::FInstanceData_GameThread> InstanceData(Context);
+	FNDIOutputParam<bool> OutValid(Context);
+
+	for (int32 i = 0; i < Context.GetNumInstances(); ++i)
+	{
+		OutValid.SetAndAdvance(InstanceData->bMeshValid);
+	}
+}
+
+void UNiagaraDataInterfaceStaticMesh::VMGetLocalToWorld(FVectorVMExternalFunctionContext& Context)
+{
+	VectorVM::FUserPtrHandler<NDIStaticMeshLocal::FInstanceData_GameThread> InstanceData(Context);
+	FNDIOutputParam<FMatrix44f> OutMatrix(Context);
+
+	for (int32 i = 0; i < Context.GetNumInstances(); ++i)
+	{
+		OutMatrix.SetAndAdvance(InstanceData->Transform);
+	}
+}
+
+void UNiagaraDataInterfaceStaticMesh::VMGetLocalToWorldInverseTransposed(FVectorVMExternalFunctionContext& Context)
+{
+	VectorVM::FUserPtrHandler<NDIStaticMeshLocal::FInstanceData_GameThread> InstanceData(Context);
+	FNDIOutputParam<FMatrix44f> OutMatrix(Context);
+
+	for (int32 i = 0; i < Context.GetNumInstances(); ++i)
+	{
+		OutMatrix.SetAndAdvance(InstanceData->TransformInverseTransposed);
+	}
+}
+
+void UNiagaraDataInterfaceStaticMesh::VMGetWorldVelocity(FVectorVMExternalFunctionContext& Context)
+{
+	VectorVM::FUserPtrHandler<NDIStaticMeshLocal::FInstanceData_GameThread> InstanceData(Context);
+	FNDIOutputParam<FVector3f> OutVelocity(Context);
+
+	for (int32 i = 0; i < Context.GetNumInstances(); ++i)
+	{
+		OutVelocity.SetAndAdvance(InstanceData->PhysicsVelocity);
+	}
+}
+
+//////////////////////////////////////////////////////////////////////////
+// Deprecated VM Functions
+template<typename TTransformHandler>
+void UNiagaraDataInterfaceStaticMesh::VMGetVertexPosition_Deprecated(FVectorVMExternalFunctionContext& Context)
+{
+	NDIStaticMeshLocal::FStaticMeshCpuHelper<TTransformHandler> StaticMeshHelper(Context);
+	FNDIInputParam<int32> VertexParam(Context);
+	FNDIOutputParam<FVector3f> OutPosition(Context);
+
+	const int32 PositionsMax = StaticMeshHelper.GetNumPositionVertices() - 1;
+	if ( PositionsMax >= 0 )
+	{
+		for (int32 i = 0; i < Context.GetNumInstances(); ++i)
+		{
+			const int32 Vertex = FMath::Clamp(VertexParam.GetAndAdvance(), 0, PositionsMax);
+			const FVector3f Position = StaticMeshHelper.GetPosition(Vertex);
+			OutPosition.SetAndAdvance(Position);
+		}
+	}
+	else
+	{
+		const FVector3f Position = StaticMeshHelper.TransformPosition(FVector3f::ZeroVector);
+		for (int32 i = 0; i < Context.GetNumInstances(); ++i)
+		{
+			OutPosition.SetAndAdvance(Position);
+		}
+	}
+}
+
+template<typename TTransformHandler>
+void UNiagaraDataInterfaceStaticMesh::VMGetTriPosition_Deprecated(FVectorVMExternalFunctionContext& Context)
+{
+	NDIStaticMeshLocal::FStaticMeshCpuHelper<TTransformHandler> StaticMeshHelper(Context);
+	FNDIInputParam<int32> TriangleParam(Context);
+	FNDIInputParam<FVector3f> BaryCoordParam(Context);
+	FNDIOutputParam<FVector3f> OutPositionParam(Context);
+
+	const FIndexArrayView IndexArray = StaticMeshHelper.GetIndexArrayView();
+	const int32 TriangleMax = StaticMeshHelper.GetNumTriangles() - 1;
+	if (TriangleMax >= 0 && StaticMeshHelper.GetNumPositionVertices() > 0)
+	{
+		for (int32 i = 0; i < Context.GetNumInstances(); ++i)
+		{
+			const int32 Triangle = FMath::Clamp(TriangleParam.GetAndAdvance(), 0, TriangleMax);
+			const FVector3f BaryCoord = BaryCoordParam.GetAndAdvance();
+			const FVector3f Position = StaticMeshHelper.GetTrianglePosition(BaryCoord, IndexArray[Triangle * 3 + 0], IndexArray[Triangle * 3 + 1], IndexArray[Triangle * 3 + 2]);
+			OutPositionParam.SetAndAdvance(Position);
+		}
+	}
+	else
+	{
+		const FVector3f Position = StaticMeshHelper.TransformPosition(FVector3f::ZeroVector);
+		for (int32 i = 0; i < Context.GetNumInstances(); ++i)
+		{
+			OutPositionParam.SetAndAdvance(Position);
+		}
+	}
+}
+
+template<typename TTransformHandler>
+void UNiagaraDataInterfaceStaticMesh::VMGetTriPositionAndVelocity_Deprecated(FVectorVMExternalFunctionContext& Context)
+{
+	NDIStaticMeshLocal::FStaticMeshCpuHelper<TTransformHandler> StaticMeshHelper(Context);
+	FNDIInputParam<int32> TriangleParam(Context);
+	FNDIInputParam<FVector3f> BaryCoordParam(Context);
+	FNDIOutputParam<FVector3f> OutPositionParam(Context);
+	FNDIOutputParam<FVector3f> OutVelocityParam(Context);
+
+	const FIndexArrayView IndexArray = StaticMeshHelper.GetIndexArrayView();
+	const int32 TriangleMax = StaticMeshHelper.GetNumTriangles() - 1;
+	if (TriangleMax >= 0 && StaticMeshHelper.GetNumPositionVertices() > 0)
+	{
+		for (int32 i = 0; i < Context.GetNumInstances(); ++i)
+		{
+			const int32 Triangle = FMath::Clamp(TriangleParam.GetAndAdvance(), 0, TriangleMax);
+			const FVector3f BaryCoord = BaryCoordParam.GetAndAdvance();
+			const FVector3f LocalPosition = StaticMeshHelper.GetLocalTrianglePosition(BaryCoord, IndexArray[Triangle * 3 + 0], IndexArray[Triangle * 3 + 1], IndexArray[Triangle * 3 + 2]);
+			const FVector3f Position = StaticMeshHelper.TransformPosition(LocalPosition);
+			const FVector3f PreviousPosition = StaticMeshHelper.PreviousTransformPosition(LocalPosition);
+			const FVector3f Velocity = (Position - PreviousPosition) * StaticMeshHelper.GetInvDeltaSeconds();
+			OutPositionParam.SetAndAdvance(Position);
+			OutVelocityParam.SetAndAdvance(Velocity);
+		}
+	}
+	else
+	{
+		const FVector3f Position = StaticMeshHelper.TransformPosition(FVector3f::ZeroVector);
+		const FVector3f PreviousPosition = StaticMeshHelper.PreviousTransformPosition(FVector3f::ZeroVector);
+		const FVector3f Velocity = FVector3f::ZeroVector;
+		for (int32 i = 0; i < Context.GetNumInstances(); ++i)
+		{
+			OutPositionParam.SetAndAdvance(Position);
+			OutVelocityParam.SetAndAdvance(Velocity);
+		}
+	}
+}
+
+template<typename TTransformHandler>
+void UNiagaraDataInterfaceStaticMesh::VMGetTriangleTangentBasis_Deprecated(FVectorVMExternalFunctionContext& Context)
+{
+	NDIStaticMeshLocal::FStaticMeshCpuHelper<TTransformHandler> StaticMeshHelper(Context);
+	FNDIInputParam<int32> TriangleParam(Context);
+	FNDIInputParam<FVector3f> BaryCoordParam(Context);
+	FNDIOutputParam<FVector3f> OutTangent(Context);
+	FNDIOutputParam<FVector3f> OutBitangent(Context);
+	FNDIOutputParam<FVector3f> OutNormal(Context);
+
+	const FIndexArrayView IndexArray = StaticMeshHelper.GetIndexArrayView();
+	const int32 TriangleMax = StaticMeshHelper.GetNumTriangles() - 1;
+	if (TriangleMax >= 0 && StaticMeshHelper.GetNumTangentVertices() > 0)
+	{
+		for (int32 i = 0; i < Context.GetNumInstances(); ++i)
+		{
+			const int32 Triangle = FMath::Clamp(TriangleParam.GetAndAdvance(), 0, TriangleMax);
+			const FVector3f BaryCoord = BaryCoordParam.GetAndAdvance();
+			const FVector3f TangentX = StaticMeshHelper.GetTriangleTangentX(BaryCoord, IndexArray[Triangle * 3 + 0], IndexArray[Triangle * 3 + 1], IndexArray[Triangle * 3 + 2]);
+			const FVector3f TangentY = StaticMeshHelper.GetTriangleTangentY(BaryCoord, IndexArray[Triangle * 3 + 0], IndexArray[Triangle * 3 + 1], IndexArray[Triangle * 3 + 2]);
+			const FVector3f TangentZ = StaticMeshHelper.GetTriangleTangentZ(BaryCoord, IndexArray[Triangle * 3 + 0], IndexArray[Triangle * 3 + 1], IndexArray[Triangle * 3 + 2]);
+			OutNormal.SetAndAdvance(TangentZ);
+			OutBitangent.SetAndAdvance(TangentY);
+			OutTangent.SetAndAdvance(TangentX);
+		}
+	}
+	else
+	{
+		const FVector3f TangentX = StaticMeshHelper.TransformVector(FVector3f(1.0f, 0.0f, 0.0f));
+		const FVector3f TangentY = StaticMeshHelper.TransformVector(FVector3f(0.0f, 1.0f, 0.0f));
+		const FVector3f TangentZ = StaticMeshHelper.TransformVector(FVector3f(0.0f, 0.0f, 1.0f));
+		for (int32 i = 0; i < Context.GetNumInstances(); ++i)
+		{
+			OutNormal.SetAndAdvance(TangentZ);
+			OutBitangent.SetAndAdvance(TangentY);
+			OutTangent.SetAndAdvance(TangentX);
+		}
+	}
+}
+
+template<typename TTransformHandler>
+void UNiagaraDataInterfaceStaticMesh::VMGetTriangleNormal_Deprecated(FVectorVMExternalFunctionContext& Context)
+{
+	NDIStaticMeshLocal::FStaticMeshCpuHelper<TTransformHandler> StaticMeshHelper(Context);
+	FNDIInputParam<int32> TriangleParam(Context);
+	FNDIInputParam<FVector3f> BaryCoordParam(Context);
+	FNDIOutputParam<FVector3f> OutNormal(Context);
+
+	const FIndexArrayView IndexArray = StaticMeshHelper.GetIndexArrayView();
+	const int32 TriangleMax = StaticMeshHelper.GetNumTriangles() - 1;
+	if (TriangleMax >= 0 && StaticMeshHelper.GetNumTangentVertices() > 0)
+	{
+		for (int32 i = 0; i < Context.GetNumInstances(); ++i)
+		{
+			const int32 Triangle = FMath::Clamp(TriangleParam.GetAndAdvance(), 0, TriangleMax);
+			const FVector3f BaryCoord = BaryCoordParam.GetAndAdvance();
+			const FVector3f TangentZ = StaticMeshHelper.GetTriangleTangentZ(BaryCoord, IndexArray[Triangle * 3 + 0], IndexArray[Triangle * 3 + 1], IndexArray[Triangle * 3 + 2]);
+			OutNormal.SetAndAdvance(TangentZ);
+		}
+	}
+	else
+	{
+		const FVector3f TangentZ = StaticMeshHelper.TransformVector(FVector3f(0.0f, 0.0f, 1.0f));
+		for (int32 i = 0; i < Context.GetNumInstances(); ++i)
+		{
+			OutNormal.SetAndAdvance(TangentZ);
+		}
+	}
+}
+
+#undef LOCTEXT_NAMESPACE
diff --git a/Engine/Plugins/FX/Niagara/Source/Niagara/Private/DataInterface/NiagaraDataInterfaceStaticMesh.h b/Engine/Plugins/FX/Niagara/Source/Niagara/Private/DataInterface/NiagaraDataInterfaceStaticMesh.h
new file mode 100644
--- /dev/null
+++ b/Engine/Plugins/FX/Niagara/Source/Niagara/Private/DataInterface/NiagaraDataInterfaceStaticMesh.h
@@ -0,0 +1,225 @@
+// Copyright Epic Games, Inc. All Rights Reserved.
+#pragma once
+
+#include "NiagaraDataInterfaceMeshCommon.h"
+#include "StaticMeshResources.h"
+#include "Engine/StaticMesh.h"
+#include "NiagaraDataInterfaceStaticMesh.generated.h"
+
+UENUM()
+enum class ENDIStaticMesh_SourceMode : uint8
+{
+	/**
+	Default behavior.
+	- Use "Source" when specified (either set explicitly or via blueprint with Set Niagara Static Mesh Component).
+	- When no source is specified, attempt to find a Static Mesh Component on an attached actor or component.
+	- If no source actor/component specified and no attached component found, fall back to the "Default Mesh" specified.
+	*/
+	Default,
+
+	/**
+	Only use "Source" (either set explicitly or via blueprint with Set Niagara Static Mesh Component).
+	*/
+	Source,
+
+	/**
+	Only use the parent actor or component the system is attached to.
+	*/
+	AttachParent,
+
+	/**
+	Only use the "Default Mesh" specified.
+	*/
+	DefaultMeshOnly,
+};
+
+USTRUCT()
+struct FNDIStaticMeshSectionFilter
+{
+	GENERATED_USTRUCT_BODY();
+
+	/** Only allow sections these material slots. */
+	UPROPERTY(EditAnywhere, Category="Section Filter")
+	TArray<int32> AllowedMaterialSlots;
+
+	//Others?
+	//Banned material slots
+	
+	void Init(class UNiagaraDataInterfaceStaticMesh* Owner, bool bAreaWeighted);
+	FORCEINLINE bool CanEverReject()const { return AllowedMaterialSlots.Num() > 0; }
+};
+
+/** Data Interface allowing sampling of static meshes. */
+UCLASS(EditInlineNew, Category = "Meshes", meta = (DisplayName = "Static Mesh"))
+class NIAGARA_API UNiagaraDataInterfaceStaticMesh : public UNiagaraDataInterface
+{
+	GENERATED_UCLASS_BODY()
+
+public:
+	DECLARE_NIAGARA_DI_PARAMETER();
+
+	/** Controls how to retrieve the Static Mesh Component to attach to. */
+	UPROPERTY(EditAnywhere, Category = "Mesh")
+	ENDIStaticMesh_SourceMode SourceMode = ENDIStaticMesh_SourceMode::Default;
+	
+#if WITH_EDITORONLY_DATA
+	/** Mesh used to sample from when not overridden by a source actor from the scene. Only available in editor for previewing. This is removed in cooked builds. */
+	UPROPERTY(EditAnywhere, Category = "Mesh")
+	TSoftObjectPtr<UStaticMesh> PreviewMesh;
+#endif
+
+	/** Mesh used to sample from when not overridden by a source actor from the scene. This mesh is NOT removed from cooked builds. */
+	UPROPERTY(EditAnywhere, Category = "Mesh")
+	TObjectPtr<UStaticMesh> DefaultMesh;
+
+	/** The source actor from which to sample. Takes precedence over the direct mesh. Note that this can only be set when used as a user variable on a component in the world. */
+	UPROPERTY(EditAnywhere, Category = "Mesh")
+	TObjectPtr<AActor> Source;
+
+	/** The source component from which to sample. Takes precedence over the direct mesh. Not exposed to the user, only indirectly accessible from blueprints. */
+	UPROPERTY(Transient)
+	TObjectPtr<UStaticMeshComponent> SourceComponent;
+
+	/** Array of filters the can be used to limit sampling to certain sections of the mesh. */
+	UPROPERTY(EditAnywhere, Category = "Mesh")
+	FNDIStaticMeshSectionFilter SectionFilter;
+
+	/** If true then the mesh velocity is taken from the mesh component's physics data. Otherwise it will be calculated by diffing the component transforms between ticks, which is more reliable but won't work on the first frame. */
+	UPROPERTY(EditAnywhere, AdvancedDisplay, Category = "Mesh")
+    bool bUsePhysicsBodyVelocity = false;
+
+	/** List of filtered sockets to use. */
+	UPROPERTY(EditAnywhere, Category = "Mesh")
+	TArray<FName> FilteredSockets;
+
+    /** Changed within the editor on PostEditChangeProperty. Should be changed whenever a refresh is desired.*/
+	uint32 ChangeId = 0;
+
+	//~ UObject interface
+	virtual void PostInitProperties() override;
+#if WITH_EDITOR
+	virtual void PostEditChangeProperty(struct FPropertyChangedEvent& PropertyChangedEvent) override;
+	virtual bool CanEditChange(const FProperty* InProperty) const override;
+#endif
+	//~ UObject interface
+
+	//~ UNiagaraDataInterface interface
+	virtual bool InitPerInstanceData(void* PerInstanceData, FNiagaraSystemInstance* SystemInstance) override;
+	virtual void DestroyPerInstanceData(void* PerInstanceData, FNiagaraSystemInstance* SystemInstance) override;
+	virtual bool PerInstanceTick(void* PerInstanceData, FNiagaraSystemInstance* SystemInstance, float DeltaSeconds) override;
+	virtual int32 PerInstanceDataSize() const override;
+
+	virtual void ProvidePerInstanceDataForRenderThread(void* DataForRenderThread, void* PerInstanceData, const FNiagaraSystemInstanceID& SystemInstance) override;
+
+	virtual void GetFunctions(TArray<FNiagaraFunctionSignature>& OutFunctions) override;
+#if WITH_EDITORONLY_DATA
+	virtual bool UpgradeFunctionCall(FNiagaraFunctionSignature& FunctionSignature) override;
+#endif
+	virtual void GetVMExternalFunction(const FVMExternalFunctionBindingInfo& BindingInfo, void* InstanceData, FVMExternalFunction &OutFunc) override;
+	virtual bool CanExecuteOnTarget(ENiagaraSimTarget Target) const override { return true; }
+	virtual bool HasPreSimulateTick() const override { return true; }
+
+#if WITH_EDITORONLY_DATA
+	virtual bool AppendCompileHash(FNiagaraCompileHashVisitor* InVisitor) const override;
+	virtual void GetParameterDefinitionHLSL(const FNiagaraDataInterfaceGPUParamInfo& ParamInfo, FString& OutHLSL) override;
+	virtual bool GetFunctionHLSL(const FNiagaraDataInterfaceGPUParamInfo& ParamInfo, const FNiagaraDataInterfaceGeneratedFunction& FunctionInfo, int FunctionInstanceIndex, FString& OutHLSL) override;
+#endif
+
+	virtual bool Equals(const UNiagaraDataInterface* Other) const override;
+protected:
+	virtual bool CopyToInternal(UNiagaraDataInterface* Destination) const override;
+public:
+#if WITH_EDITOR
+	virtual void GetFeedback(UNiagaraSystem* Asset, UNiagaraComponent* Component, TArray<FNiagaraDataInterfaceError>& OutErrors,
+		TArray<FNiagaraDataInterfaceFeedback>& OutWarnings, TArray<FNiagaraDataInterfaceFeedback>& OutInfo) override;
+#endif
+	//~ UNiagaraDataInterface interface
+
+	UStaticMesh* GetStaticMesh(USceneComponent*& OutComponent, class FNiagaraSystemInstance* SystemInstance = nullptr);
+	void SetSourceComponentFromBlueprints(UStaticMeshComponent* ComponentToUse);
+	void SetDefaultMeshFromBlueprints(UStaticMesh* MeshToUse);
+
+protected:
+	// VM Vertex Sampling
+	void VMIsValidVertex(FVectorVMExternalFunctionContext& Context);
+	void VMRandomVertex(FVectorVMExternalFunctionContext& Context);
+	void VMGetVertexCount(FVectorVMExternalFunctionContext& Context);
+
+	template<typename TTransformHandler>
+	void VMGetVertex(FVectorVMExternalFunctionContext& Context);
+	void VMGetVertexColor(FVectorVMExternalFunctionContext& Context);
+	void VMGetVertexUV(FVectorVMExternalFunctionContext& Context);
+
+	// VM Triangle Sampling
+	void VMIsValidTriangle(FVectorVMExternalFunctionContext& Context);
+	template<typename TRandomHelper>
+	void VMRandomTriangle(FVectorVMExternalFunctionContext& Context);
+	void VMGetTriangleCount(FVectorVMExternalFunctionContext& Context);
+
+	void VMIsValidFilteredTriangle(FVectorVMExternalFunctionContext& Context);
+	void VMRandomFilteredTriangle(FVectorVMExternalFunctionContext& Context);
+	void VMGetFilteredTriangleCount(FVectorVMExternalFunctionContext& Context);
+	void VMGetFilteredTriangleAt(FVectorVMExternalFunctionContext& Context);
+
+	void VMIsValidUnfilteredTriangle(FVectorVMExternalFunctionContext& Context);
+	void VMRandomUnfilteredTriangle(FVectorVMExternalFunctionContext& Context);
+	void VMGetUnfilteredTriangleCount(FVectorVMExternalFunctionContext& Context);
+	void VMGetUnfilteredTriangleAt(FVectorVMExternalFunctionContext& Context);
+
+	template<typename TTransformHandler>
+	void VMGetTriangle(FVectorVMExternalFunctionContext& Context);
+	void VMGetTriangleColor(FVectorVMExternalFunctionContext& Context);
+	void VMGetTriangleUV(FVectorVMExternalFunctionContext& Context);
+
+	void VMGetTriangleIndices(FVectorVMExternalFunctionContext& Context);
+
+	// Socket Functions
+	void VMGetSocketCount(FVectorVMExternalFunctionContext& Context);
+	void VMGetFilteredSocketCount(FVectorVMExternalFunctionContext& Context);
+	void VMGetUnfilteredSocketCount(FVectorVMExternalFunctionContext& Context);
+	template<typename TTransformHandler>
+	void VMGetSocketTransform(FVectorVMExternalFunctionContext& Context);
+	template<typename TTransformHandler>
+	void VMGetFilteredSocketTransform(FVectorVMExternalFunctionContext& Context);
+	template<typename TTransformHandler>
+	void VMGetUnfilteredSocketTransform(FVectorVMExternalFunctionContext& Context);
+
+	// Section functions
+	void VMIsValidSection(FVectorVMExternalFunctionContext& Context);
+	void VMGetSectionTriangleCount(FVectorVMExternalFunctionContext& Context);
+	template<typename TRandomHandler>
+	void VMRandomSectionTriangle(FVectorVMExternalFunctionContext& Context);
+	void VMGetSectionTriangleAt(FVectorVMExternalFunctionContext& Context);
+
+	void VMGetFilteredSectionAt(FVectorVMExternalFunctionContext& Context);
+	void VMGetUnfilteredSectionAt(FVectorVMExternalFunctionContext& Context);
+
+	void VMGetSectionCount(FVectorVMExternalFunctionContext& Context);
+	void VMGetFilteredSectionCount(FVectorVMExternalFunctionContext& Context);
+	void VMGetUnfilteredSectionCount(FVectorVMExternalFunctionContext& Context);
+
+	template<typename TRandomHandler>
+	void VMRandomSection(FVectorVMExternalFunctionContext& Context);
+	template<typename TRandomHandler>
+	void VMRandomFilteredSection(FVectorVMExternalFunctionContext& Context);
+	void VMRandomUnfilteredSection(FVectorVMExternalFunctionContext& Context);
+
+	// VM Misc Functions
+	void VMIsValid(FVectorVMExternalFunctionContext& Context);
+
+	void VMGetLocalToWorld(FVectorVMExternalFunctionContext& Context);
+	void VMGetLocalToWorldInverseTransposed(FVectorVMExternalFunctionContext& Context);
+	void VMGetWorldVelocity(FVectorVMExternalFunctionContext& Context);
+
+	// Deprecated VM Functions
+	template<typename TTransformHandler>
+	void VMGetVertexPosition_Deprecated(FVectorVMExternalFunctionContext& Context);
+	template<typename TTransformHandler>
+	void VMGetTriPosition_Deprecated(FVectorVMExternalFunctionContext& Context);
+	template<typename TTransformHandler>
+	void VMGetTriPositionAndVelocity_Deprecated(FVectorVMExternalFunctionContext& Context);
+	template<typename TTransformHandler>
+	void VMGetTriangleTangentBasis_Deprecated(FVectorVMExternalFunctionContext& Context);
+	template<typename TTransformHandler>
+	void VMGetTriangleNormal_Deprecated(FVectorVMExternalFunctionContext& Context);
+};
diff --git a/Engine/Plugins/FX/Niagara/Source/Niagara/Private/NiagaraDataInterfaceStaticMesh.cpp b/Engine/Plugins/FX/Niagara/Source/Niagara/Private/NiagaraDataInterfaceStaticMesh.cpp
deleted file mode 100644
diff --git a/Engine/Plugins/FX/Niagara/Source/Niagara/Private/NiagaraFunctionLibrary.cpp b/Engine/Plugins/FX/Niagara/Source/Niagara/Private/NiagaraFunctionLibrary.cpp
--- a/Engine/Plugins/FX/Niagara/Source/Niagara/Private/NiagaraFunctionLibrary.cpp
+++ b/Engine/Plugins/FX/Niagara/Source/Niagara/Private/NiagaraFunctionLibrary.cpp
@@ -11,7 +11,7 @@
 #include "Internationalization/Internationalization.h"
 
 #include "NiagaraWorldManager.h"
-#include "NiagaraDataInterfaceStaticMesh.h"
+#include "DataInterface/NiagaraDataInterfaceStaticMesh.h"
 #include "NiagaraStats.h"
 #include "NiagaraDataInterfaceTexture.h"
 #include "NiagaraDataInterface2DArrayTexture.h"
diff --git a/Engine/Plugins/FX/Niagara/Source/Niagara/Private/NiagaraRenderer.cpp b/Engine/Plugins/FX/Niagara/Source/Niagara/Private/NiagaraRenderer.cpp
--- a/Engine/Plugins/FX/Niagara/Source/Niagara/Private/NiagaraRenderer.cpp
+++ b/Engine/Plugins/FX/Niagara/Source/Niagara/Private/NiagaraRenderer.cpp
@@ -187,6 +187,13 @@ FRHIShaderResourceView* FNiagaraRenderer::GetDummyUIntBuffer()
 	return DummyUIntBuffer.SRV;
 }
 
+FRHIShaderResourceView* FNiagaraRenderer::GetDummyUInt2Buffer()
+{
+	check(IsInRenderingThread());
+	static TGlobalResource<FNiagaraEmptyBufferSRV> DummyUInt2Buffer(PF_R32G32_UINT, TEXT("NiagaraRenderer::DummyUInt2"));
+	return DummyUInt2Buffer.SRV;
+}
+
 FRHIShaderResourceView* FNiagaraRenderer::GetDummyUInt4Buffer()
 {
 	check(IsInRenderingThread());
diff --git a/Engine/Plugins/FX/Niagara/Source/Niagara/Private/NiagaraScript.cpp b/Engine/Plugins/FX/Niagara/Source/Niagara/Private/NiagaraScript.cpp
--- a/Engine/Plugins/FX/Niagara/Source/Niagara/Private/NiagaraScript.cpp
+++ b/Engine/Plugins/FX/Niagara/Source/Niagara/Private/NiagaraScript.cpp
@@ -21,7 +21,7 @@
 #include "UObject/ObjectSaveContext.h"
 #include "UObject/ReleaseObjectVersion.h"
 #include "NiagaraDataInterfaceSkeletalMesh.h"
-#include "NiagaraDataInterfaceStaticMesh.h"
+#include "DataInterface/NiagaraDataInterfaceStaticMesh.h"
 #if WITH_EDITOR
 	#include "DerivedDataCacheInterface.h"
 	#include "Interfaces/ITargetPlatform.h"
diff --git a/Engine/Plugins/FX/Niagara/Source/Niagara/Public/NiagaraRenderer.h b/Engine/Plugins/FX/Niagara/Source/Niagara/Public/NiagaraRenderer.h
--- a/Engine/Plugins/FX/Niagara/Source/Niagara/Public/NiagaraRenderer.h
+++ b/Engine/Plugins/FX/Niagara/Source/Niagara/Public/NiagaraRenderer.h
@@ -123,6 +123,7 @@ public:
 	static FRHIShaderResourceView* GetDummyWhiteColorBuffer();
 	static FRHIShaderResourceView* GetDummyIntBuffer();
 	static FRHIShaderResourceView* GetDummyUIntBuffer();
+	static FRHIShaderResourceView* GetDummyUInt2Buffer();
 	static FRHIShaderResourceView* GetDummyUInt4Buffer();
 	static FRHIShaderResourceView* GetDummyTextureReadBuffer2D();
 	static FRHIShaderResourceView* GetDummyTextureReadBuffer2DArray();
diff --git a/Engine/Plugins/FX/Niagara/Source/NiagaraCore/Private/NiagaraCustomVersion.cpp b/Engine/Plugins/FX/Niagara/Source/NiagaraCore/Private/NiagaraCustomVersion.cpp
--- a/Engine/Plugins/FX/Niagara/Source/NiagaraCore/Private/NiagaraCustomVersion.cpp
+++ b/Engine/Plugins/FX/Niagara/Source/NiagaraCore/Private/NiagaraCustomVersion.cpp
@@ -8,4 +8,4 @@ const FGuid FNiagaraCustomVersion::GUID(0xFCF57AFA, 0x50764283, 0xB9A9E658, 0xFF
 // Register the custom version with core
 FCustomVersionRegistration GRegisterNiagaraCustomVersion(FNiagaraCustomVersion::GUID, FNiagaraCustomVersion::LatestVersion, TEXT("NiagaraVer"));
 
-const FGuid FNiagaraCustomVersion::LatestScriptCompileVersion(0xCBDDE73F, 0xA0584B4F, 0x864D4335, 0x59C858DA);
\ No newline at end of file
+const FGuid FNiagaraCustomVersion::LatestScriptCompileVersion(0x47F836A4, 0x880FE04D, 0x8E8AFB04, 0x6D8CA5CD);
diff --git a/Engine/Plugins/FX/Niagara/Source/NiagaraEditorWidgets/Private/DetailCustomizations/NiagaraDataInterfaceStaticMeshDetails.cpp b/Engine/Plugins/FX/Niagara/Source/NiagaraEditorWidgets/Private/DetailCustomizations/NiagaraDataInterfaceStaticMeshDetails.cpp
--- a/Engine/Plugins/FX/Niagara/Source/NiagaraEditorWidgets/Private/DetailCustomizations/NiagaraDataInterfaceStaticMeshDetails.cpp
+++ b/Engine/Plugins/FX/Niagara/Source/NiagaraEditorWidgets/Private/DetailCustomizations/NiagaraDataInterfaceStaticMeshDetails.cpp
@@ -3,7 +3,7 @@
 #include "NiagaraDataInterfaceStaticMeshDetails.h"
 #include "NiagaraDetailSourcedArrayBuilder.h"
 #include "NiagaraDataInterfaceDetails.h"
-#include "NiagaraDataInterfaceStaticMesh.h" 
+#include "Niagara/Private/DataInterface/NiagaraDataInterfaceStaticMesh.h" 
 #include "DetailLayoutBuilder.h"
 #include "DetailCategoryBuilder.h"
 #include "IDetailGroup.h"
@@ -41,7 +41,7 @@ void FNiagaraDataInterfaceStaticMeshDetails::CustomizeDetails(IDetailLayoutBuild
 	 MeshInterface->OnChanged().RemoveAll(this);
 	 MeshInterface->OnChanged().AddSP(this, &FNiagaraDataInterfaceStaticMeshDetails::OnInterfaceChanged);
 
-	 TWeakObjectPtr<USceneComponent> SceneComponent;
+	 USceneComponent* SceneComponent = nullptr;
 	 if (UNiagaraDataInterfaceStaticMesh* Interface = MeshInterface.Get())
 	 {
 		 MeshObject = Interface->GetStaticMesh(SceneComponent);
@@ -91,7 +91,7 @@ void FNiagaraDataInterfaceStaticMeshDetails::CustomizeDetails(IDetailLayoutBuild
 		 MeshObject = nullptr;
 	 }
 
-	 TWeakObjectPtr<USceneComponent> SceneComponent;
+	 USceneComponent* SceneComponent = nullptr;
 	 if (UNiagaraDataInterfaceStaticMesh* Interface = MeshInterface.Get())
 	 {
 		 MeshObject = Interface->GetStaticMesh(SceneComponent);
diff --git a/Engine/Source/Runtime/Engine/Public/Rendering/PositionVertexBuffer.h b/Engine/Source/Runtime/Engine/Public/Rendering/PositionVertexBuffer.h
--- a/Engine/Source/Runtime/Engine/Public/Rendering/PositionVertexBuffer.h
+++ b/Engine/Source/Runtime/Engine/Public/Rendering/PositionVertexBuffer.h
@@ -93,7 +93,7 @@ public:
 	{
 		return NumVertices;
 	}
-	bool GetAllowCPUAccess() const;
+	ENGINE_API bool GetAllowCPUAccess() const;
 
 	/** Create an RHI vertex buffer with CPU data. CPU data may be discarded after creation (see TResourceArray::Discard) */
 	FBufferRHIRef CreateRHIBuffer_RenderThread();
diff --git a/Engine/Source/Runtime/VectorVM/Public/VectorVM.h b/Engine/Source/Runtime/VectorVM/Public/VectorVM.h
--- a/Engine/Source/Runtime/VectorVM/Public/VectorVM.h
+++ b/Engine/Source/Runtime/VectorVM/Public/VectorVM.h
@@ -576,6 +576,10 @@ namespace VectorVM
 		FORCEINLINE T* Get() { return Ptr; }
 		FORCEINLINE T* operator->() { return Ptr; }
 		FORCEINLINE operator T*() { return Ptr; }
+
+		FORCEINLINE const T* Get() const { return Ptr; }
+		FORCEINLINE const T* operator->() const { return Ptr; }
+		FORCEINLINE operator const T* () const { return Ptr; }
 	};
 
 	// A flexible handler that can deal with either constant or register inputs.
