* Prevent Niagara compilations to stack, have at least one pending compile request
* Fixed a bunch of formatting problems

diff --git a/Engine/Plugins/FX/Niagara/Source/Niagara/Private/NiagaraSystem.cpp b/Engine/Plugins/FX/Niagara/Source/Niagara/Private/NiagaraSystem.cpp
--- a/Engine/Plugins/FX/Niagara/Source/Niagara/Private/NiagaraSystem.cpp
+++ b/Engine/Plugins/FX/Niagara/Source/Niagara/Private/NiagaraSystem.cpp
@@ -222,10 +222,10 @@ FNiagaraAsyncCompileTask::FNiagaraAsyncCompileTask(UNiagaraSystem* InOwningSyste
 void FNiagaraAsyncCompileTask::ProcessCurrentState()
 {
 	SCOPE_CYCLE_COUNTER(STAT_Niagara_System_CompileScriptTaskGT);
-		if (IsDone())
-		{
-			return;
-		}
+	if (IsDone())
+	{
+		return;
+	}
 	
 	if (CurrentState == ENiagaraCompilationState::CheckDDC)
 	{
@@ -233,42 +233,42 @@ void FNiagaraAsyncCompileTask::ProcessCurrentState()
 		CheckDDCResult();
 	}
 	else if (CurrentState == ENiagaraCompilationState::Precompile)
-		{
-			// gather precompile data
+	{
+		// gather precompile data
 		StartCompileTime = FPlatformTime::Seconds();
-				PrecompileData();
-				MoveToState(ENiagaraCompilationState::StartCompileJob);
-			}
+		PrecompileData();
+		MoveToState(ENiagaraCompilationState::StartCompileJob);
+	}
 	else if (CurrentState == ENiagaraCompilationState::StartCompileJob)
-			{
-			// start the async compile job
-				StartCompileJob();
-				MoveToState(ENiagaraCompilationState::AwaitResult);
-			}
+	{
+		// start the async compile job
+		StartCompileJob();
+		MoveToState(ENiagaraCompilationState::AwaitResult);
+	}
 	else if (CurrentState == ENiagaraCompilationState::AwaitResult)
-			{
-				if (AwaitResult())
-				{
-					MoveToState(ENiagaraCompilationState::ProcessResult);
-				}
-			}
-	else if (CurrentState == ENiagaraCompilationState::ProcessResult)
+	{
+		if (AwaitResult())
 		{
-			// save the result from the compile job
-				ProcessResult();
+			MoveToState(ENiagaraCompilationState::ProcessResult);
+		}
+	}
+	else if (CurrentState == ENiagaraCompilationState::ProcessResult)
+	{
+		// save the result from the compile job
+		ProcessResult();
 		MoveToState(ENiagaraCompilationState::PutToDDC);
-			}
+	}
 	else if (CurrentState == ENiagaraCompilationState::PutToDDC)
-			{
+	{
 		// put the result from the compile job into the ddc
 		PutToDDC();
 		MoveToState(ENiagaraCompilationState::Finished);
-		}
-		else
-		{
-			check(false);
-		}
 	}
+	else
+	{
+		check(false);
+	}
+}
 
 void FNiagaraAsyncCompileTask::MoveToState(ENiagaraCompilationState NewState)
 {
@@ -307,7 +307,7 @@ void FNiagaraAsyncCompileTask::MoveToState(ENiagaraCompilationState NewState)
 
 	UE_LOG(LogNiagara, Verbose, TEXT("Changing state %i -> %i for for %s!"), CurrentState, NewState, *AssetPath);
 	CurrentState = NewState;
-	}
+}
 
 bool FNiagaraAsyncCompileTask::IsDone() const
 {
@@ -433,7 +433,6 @@ void FNiagaraAsyncCompileTask::WaitAndResolveResult()
 {
 	check(IsInGameThread());
 	
-
 	bWaitForCompileJob = true;
 	while (!IsDone())
 	{
@@ -465,16 +464,16 @@ void FNiagaraAsyncCompileTask::CheckDDCResult()
 				ScriptPair.bResultsReady = true;
 				MoveToState(ENiagaraCompilationState::Finished);
 				UE_LOG(LogNiagara, Verbose, TEXT("Compilation data for %s could be pulled from the ddc."), *AssetPath);
-}
+			}
 			else
-{
+			{
 				UE_LOG(LogNiagara, Warning, TEXT("Unable to create exec data from ddc data for script %s, going to recompile it from scratch. DDC might be corrupted or there is a problem with script serialization."), *AssetPath);
 				ExeData.Reset();
 				MoveToState(ENiagaraCompilationState::Precompile);
 			}
 		}
 		else
-	{
+		{
 			MoveToState(ENiagaraCompilationState::Precompile);
 			UE_LOG(LogNiagara, Verbose, TEXT("No compilation data for %s found in the ddc."), *AssetPath);
 		}
@@ -2634,18 +2633,18 @@ bool UNiagaraSystem::QueryCompileComplete(bool bWait, bool bDoPost, bool bDoNotA
 		// Make sure that we aren't waiting for any results to come back.
 		if (bAreWeWaitingForAnyResults)
 		{
-				return false;
-			}
-			// if we've gotten all the results, run a quick check to see if the data is valid, if it's not then that indicates that
-			// we've run into a compatibility issue and so we should see if we should issue a full rebuild
-			const bool ResultsValid = CompilationResultsValid(CompileRequest);
-			if (!ResultsValid && !CompileRequest.bForced)
-			{
-				CompileRequest.RootObjects.Empty();
-				ActiveCompilations.RemoveAt(0);
-				RequestCompile(true, nullptr);
-				return false;
-			}
+			return false;
+		}
+		// if we've gotten all the results, run a quick check to see if the data is valid, if it's not then that indicates that
+		// we've run into a compatibility issue and so we should see if we should issue a full rebuild
+		const bool ResultsValid = CompilationResultsValid(CompileRequest);
+		if (!ResultsValid && !CompileRequest.bForced)
+		{
+			CompileRequest.RootObjects.Empty();
+			ActiveCompilations.RemoveAt(0);
+			RequestCompile(true, nullptr);
+			return false;
+		}
 
 		// In the world of do not apply, we're exiting the system completely so let's just kill any active compilations altogether.
 		if (bDoNotApply || CompileRequest.bIsValid == false)
@@ -2939,11 +2938,16 @@ bool UNiagaraSystem::RequestCompile(bool bForce, FNiagaraSystemUpdateContext* Op
 		return false;
 	}
 
-	if (ActiveCompilations.Num() > 0)
+	// we don't want to stack compilations, so we remove requests that have not started processing yet
+	for (int i = ActiveCompilations.Num() - 1; i >= 1; i--)
 	{
-		PollForCompilationComplete();
-	}
-	
+		FNiagaraSystemCompileRequest& Request = ActiveCompilations[i];
+		for (auto& AsyncTask : Request.DDCTasks)
+		{
+			AsyncTask->AbortTask();
+		}
+		ActiveCompilations.RemoveAt(i);
+	}	
 
 	// Record that we entered this function already.
 	SCOPE_CYCLE_COUNTER(STAT_Niagara_System_CompileScript);
@@ -3098,7 +3102,7 @@ bool UNiagaraSystem::RequestCompile(bool bForce, FNiagaraSystemUpdateContext* Op
 
 
 	// We might be able to just complete compilation right now if nothing needed compilation.
-	if (ScriptsNeedingCompile.Num() == 0)
+	if (ScriptsNeedingCompile.Num() == 0 && ActiveCompilations.Num() == 1)
 	{
 		ActiveCompilation.bAllScriptsSynchronized = true;
 		PollForCompilationComplete();
