- Add support for pulling bind pose vertex / triangle

diff --git a/Engine/Plugins/FX/Niagara/Shaders/Private/NiagaraDataInterfaceSkeletalMesh.ush b/Engine/Plugins/FX/Niagara/Shaders/Private/NiagaraDataInterfaceSkeletalMesh.ush
--- a/Engine/Plugins/FX/Niagara/Shaders/Private/NiagaraDataInterfaceSkeletalMesh.ush
+++ b/Engine/Plugins/FX/Niagara/Shaders/Private/NiagaraDataInterfaceSkeletalMesh.ush
@@ -144,6 +144,37 @@ float4 DISkelMesh_UnpackWeights4(uint Packed)
 	return float4(Packed & 0xff, Packed >> 8 & 0xff, Packed >> 16 & 0xff, Packed >> 24 & 0xff) / 255.0f;
 }
 
+void DISkelMesh_GetVertex(in FDISkelMeshContext DIContext, uint VertexIndex, out float3 OutPosition, out float3 OutTangentX, out float3 OutTangentY, out float3 OutTangentZ)
+{
+	//-TODO: Use a R32G32B32 format
+	float3 Position = float3(DIContext.MeshVertexBuffer[VertexIndex * 3], DIContext.MeshVertexBuffer[VertexIndex * 3 + 1], DIContext.MeshVertexBuffer[VertexIndex * 3 + 2]);
+	float3 TangentX = TangentBias(DIContext.MeshTangentBuffer[VertexIndex * 2    ].xyz);
+	float4 TangentZ = TangentBias(DIContext.MeshTangentBuffer[VertexIndex * 2 + 1].xyzw);	
+	float3 TangentY = cross(TangentZ.xyz, TangentX.xyz) * TangentZ.w;
+
+	OutPosition = Position;
+	OutTangentX = cross(TangentY, TangentZ.xyz) * TangentZ.w;
+	OutTangentY = TangentY;
+	OutTangentZ = TangentZ.xyz;
+}
+
+void DISkelMesh_GetPointOnTriangle(in FDISkelMeshContext DIContext, uint TriangleIndex, float3 BaryCoord, out float3 OutPosition, out float3 OutTangentX, out float3 OutTangentY, out float3 OutTangentZ)
+{
+	uint IndexBufferOffset = TriangleIndex * 3;
+	uint VertexIndex0 = DIContext.MeshIndexBuffer[IndexBufferOffset];
+	uint VertexIndex1 = DIContext.MeshIndexBuffer[IndexBufferOffset + 1];
+	uint VertexIndex2 = DIContext.MeshIndexBuffer[IndexBufferOffset + 2];
+
+	float3 Positions[3], TangentsX[3], TangentsY[3], TangentsZ[3];
+	DISkelMesh_GetVertex(DIContext, VertexIndex0, Positions[0], TangentsX[0], TangentsY[0], TangentsZ[0]);
+	DISkelMesh_GetVertex(DIContext, VertexIndex1, Positions[1], TangentsX[1], TangentsY[1], TangentsZ[1]);
+	DISkelMesh_GetVertex(DIContext, VertexIndex2, Positions[2], TangentsX[2], TangentsY[2], TangentsZ[2]);
+	OutPosition = Positions[0] * BaryCoord.x + Positions[1] * BaryCoord.y  + Positions[2] * BaryCoord.z;
+	OutTangentX = TangentsX[0] * BaryCoord.x + TangentsX[1] * BaryCoord.y  + TangentsX[2] * BaryCoord.z;
+	OutTangentY = TangentsY[0] * BaryCoord.x + TangentsY[1] * BaryCoord.y  + TangentsY[2] * BaryCoord.z;
+	OutTangentZ = TangentsZ[0] * BaryCoord.x + TangentsZ[1] * BaryCoord.y  + TangentsZ[2] * BaryCoord.z;
+}
+
 FDISkelMeshSkinnedVertex DISkelMesh_GetSkinnedVertex(in FDISkelMeshContext DIContext, uint VertexIndex)
 {
 	//-TODO: Use a R32G32B32 format
diff --git a/Engine/Plugins/FX/Niagara/Source/Niagara/Classes/NiagaraDataInterfaceMeshCommon.h b/Engine/Plugins/FX/Niagara/Source/Niagara/Classes/NiagaraDataInterfaceMeshCommon.h
--- a/Engine/Plugins/FX/Niagara/Source/Niagara/Classes/NiagaraDataInterfaceMeshCommon.h
+++ b/Engine/Plugins/FX/Niagara/Source/Niagara/Classes/NiagaraDataInterfaceMeshCommon.h
@@ -54,13 +54,13 @@ FORCEINLINE FVector4 BarycentricInterpolate(float BaryX, float BaryY, float Bary
 }
 
 template<typename T>
-FORCEINLINE T BarycentricInterpolate(FVector BaryCoord, T V0, T V1, T V2)
+FORCEINLINE T BarycentricInterpolate(FVector3f BaryCoord, T V0, T V1, T V2)
 {
 	return V0 * BaryCoord.X + V1 * BaryCoord.Y + V2 * BaryCoord.Z;
 }
 
 // Overload for FVector4 to work around C2719: (formal parameter with requested alignment of 16 won't be aligned)
-FORCEINLINE FVector4 BarycentricInterpolate(FVector BaryCoord, const FVector4& V0, const FVector4& V1, const FVector4& V2)
+FORCEINLINE FVector4 BarycentricInterpolate(FVector3f BaryCoord, const FVector4& V0, const FVector4& V1, const FVector4& V2)
 {
 	return V0 * BaryCoord.X + V1 * BaryCoord.Y + V2 * BaryCoord.Z;
 }
diff --git a/Engine/Plugins/FX/Niagara/Source/Niagara/Classes/NiagaraDataInterfaceSkeletalMesh.h b/Engine/Plugins/FX/Niagara/Source/Niagara/Classes/NiagaraDataInterfaceSkeletalMesh.h
--- a/Engine/Plugins/FX/Niagara/Source/Niagara/Classes/NiagaraDataInterfaceSkeletalMesh.h
+++ b/Engine/Plugins/FX/Niagara/Source/Niagara/Classes/NiagaraDataInterfaceSkeletalMesh.h
@@ -893,6 +893,8 @@ public:
 	template<typename FilterMode, typename AreaWeightingMode>
 	void IsValidTriCoord(FVectorVMExternalFunctionContext& Context);
 
+	void GetTriangleData(FVectorVMExternalFunctionContext& Context);
+
 	template<typename SkinningHandlerType, typename TransformHandlerType, typename VertexAccessorType, typename bInterpolated>
 	void GetTriCoordSkinnedData(FVectorVMExternalFunctionContext& Context);
 
@@ -951,6 +953,9 @@ public:
 	template<typename FilterMode>
 	void GetFilteredVertexAt(FVectorVMExternalFunctionContext& Context);
 
+	template<typename VertexAccessorType>
+	void GetVertexData(FVectorVMExternalFunctionContext& Context);
+
 	template<typename SkinningHandlerType, typename TransformHandlerType, typename VertexAccessorType>
 	void GetVertexSkinnedData(FVectorVMExternalFunctionContext& Context);
 
@@ -1019,6 +1024,7 @@ public:
 	// Triangle Sampling
 	static const FName RandomTriCoordName;
 	static const FName IsValidTriCoordName;
+	static const FName GetTriangleDataName;
 	static const FName GetSkinnedTriangleDataName;
 	static const FName GetSkinnedTriangleDataWSName;
 	static const FName GetSkinnedTriangleDataInterpName;
@@ -1059,6 +1065,7 @@ public:
 	static const FName GetFilteredSocketOrBoneAtName;
 
 	// Vertex Sampling
+	static const FName GetVertexDataName;
 	static const FName GetSkinnedVertexDataName;
 	static const FName GetSkinnedVertexDataWSName;
 	static const FName GetVertexColorName;
diff --git a/Engine/Plugins/FX/Niagara/Source/Niagara/Private/NDISkeletalMesh_TriangleSampling.cpp b/Engine/Plugins/FX/Niagara/Source/Niagara/Private/NDISkeletalMesh_TriangleSampling.cpp
--- a/Engine/Plugins/FX/Niagara/Source/Niagara/Private/NDISkeletalMesh_TriangleSampling.cpp
+++ b/Engine/Plugins/FX/Niagara/Source/Niagara/Private/NDISkeletalMesh_TriangleSampling.cpp
@@ -31,6 +31,7 @@ DEFINE_NDI_FUNC_BINDER(UNiagaraDataInterfaceSkeletalMesh, GetTriangleCoordInAabb
 
 const FName FSkeletalMeshInterfaceHelper::RandomTriCoordName("RandomTriCoord");
 const FName FSkeletalMeshInterfaceHelper::IsValidTriCoordName("IsValidTriCoord");
+const FName FSkeletalMeshInterfaceHelper::GetTriangleDataName("GetTriangleData");
 const FName FSkeletalMeshInterfaceHelper::GetSkinnedTriangleDataName("GetSkinnedTriangleData");
 const FName FSkeletalMeshInterfaceHelper::GetSkinnedTriangleDataWSName("GetSkinnedTriangleDataWS");
 const FName FSkeletalMeshInterfaceHelper::GetSkinnedTriangleDataInterpName("GetSkinnedTriangleDataInterpolated");
@@ -73,6 +74,22 @@ void UNiagaraDataInterfaceSkeletalMesh::GetTriangleSamplingFunctions(TArray<FNia
 		Sig.Description = LOCTEXT("IsValidDesc", "Determine if this tri coordinate's triangle index is valid for this mesh. Note that this only checks the mesh index buffer size and does not include any filtering settings.");
 #endif
 	}
+	
+	{
+		FNiagaraFunctionSignature & Sig = OutFunctions.AddDefaulted_GetRef();
+		Sig.Name = FSkeletalMeshInterfaceHelper::GetTriangleDataName;
+		Sig.Inputs.Add(FNiagaraVariable(FNiagaraTypeDefinition(GetClass()), TEXT("SkeletalMesh")));
+		Sig.Inputs.Add(FNiagaraVariable(FNiagaraTypeDefinition(FMeshTriCoordinate::StaticStruct()), TEXT("Coord")));
+		Sig.Outputs.Add(FNiagaraVariable(FNiagaraTypeDefinition::GetVec3Def(), TEXT("Position")));
+		Sig.Outputs.Add(FNiagaraVariable(FNiagaraTypeDefinition::GetVec3Def(), TEXT("Normal")));
+		Sig.Outputs.Add(FNiagaraVariable(FNiagaraTypeDefinition::GetVec3Def(), TEXT("Binormal")));
+		Sig.Outputs.Add(FNiagaraVariable(FNiagaraTypeDefinition::GetVec3Def(), TEXT("Tangent")));
+		Sig.bMemberFunction = true;
+		Sig.bRequiresContext = false;
+#if WITH_EDITORONLY_DATA
+		Sig.Description = LOCTEXT("GetTriangleDataDesc", "Returns bind pose triangle data.");
+#endif
+	}
 
 	{
 		FNiagaraFunctionSignature& Sig = OutFunctions.AddDefaulted_GetRef();
@@ -312,6 +329,11 @@ void UNiagaraDataInterfaceSkeletalMesh::BindTriangleSamplingFunction(const FVMEx
 		check(BindingInfo.GetNumInputs() == 5 && BindingInfo.GetNumOutputs() == 1);
 		TFilterModeBinder<TAreaWeightingModeBinder<NDI_FUNC_BINDER(UNiagaraDataInterfaceSkeletalMesh, IsValidTriCoord)>>::Bind(this, BindingInfo, InstanceData, OutFunc);
 	}
+	else if (BindingInfo.Name == FSkeletalMeshInterfaceHelper::GetTriangleDataName)
+	{
+		check(BindingInfo.GetNumInputs() == 5 && BindingInfo.GetNumOutputs() == 12);
+		OutFunc = FVMExternalFunction::CreateLambda([this](FVectorVMExternalFunctionContext& Context) {this->GetTriangleData(Context);});
+	}
 	else if (BindingInfo.Name == FSkeletalMeshInterfaceHelper::GetSkinnedTriangleDataName)
 	{
 		check(BindingInfo.GetNumInputs() == 5 && BindingInfo.GetNumOutputs() == 15);
@@ -1207,6 +1229,64 @@ struct FGetTriCoordSkinnedDataOutputHandler
 	const bool bNeedsTangent;
 };
 
+void UNiagaraDataInterfaceSkeletalMesh::GetTriangleData(FVectorVMExternalFunctionContext& Context)
+{
+	SCOPE_CYCLE_COUNTER(STAT_NiagaraSkel_Sample);
+	VectorVM::FUserPtrHandler<FNDISkeletalMesh_InstanceData> InstanceData(Context);
+	FNDIInputParam<int32> TriParam(Context);
+	FNDIInputParam<FVector3f> BaryParam(Context);
+
+	FNDIOutputParam<FVector3f> OutPositionParam(Context);
+	FNDIOutputParam<FVector3f> OutTangentZParam(Context);
+	FNDIOutputParam<FVector3f> OutTangentYParam(Context);
+	FNDIOutputParam<FVector3f> OutTangentXParam(Context);
+
+	FSkeletalMeshAccessorHelper Accessor;
+	Accessor.Init<TNDISkelMesh_FilterModeNone, TNDISkelMesh_AreaWeightingOff>(InstanceData);
+
+	// Is the data valid?
+	if (!InstanceData->bAllowCPUMeshDataAccess || !Accessor.IsLODAccessible())
+	{
+		for (int32 i = 0; i < Context.GetNumInstances(); ++i)
+		{
+			OutPositionParam.SetAndAdvance(FVector3f::ZeroVector);
+			OutTangentZParam.SetAndAdvance(FVector3f::ZAxisVector);
+			OutTangentYParam.SetAndAdvance(FVector3f::YAxisVector);
+			OutTangentXParam.SetAndAdvance(FVector3f::XAxisVector);
+		}
+		return;
+	}
+
+	// Data should be considered valid here
+	const FSkeletalMeshLODRenderData* LODData = Accessor.LODData;
+	const int32 TriMax = (Accessor.IndexBuffer->Num() / 3) - 1;
+	FSkinnedPositionAccessorHelper<TNDISkelMesh_SkinningModeNone> SkinningHandler;
+
+	for (int32 i = 0; i < Context.GetNumInstances(); ++i)
+	{
+		const int Triangle = FMath::Clamp(TriParam.GetAndAdvance(), 0, TriMax);
+		const FVector3f BaryCoord = BaryParam.GetAndAdvance();
+
+		int32 Indices[3];
+		SkinningHandler.GetTriangleIndices(Accessor, Triangle, Indices[0], Indices[1], Indices[2]);
+
+		FVector3f Positions[3];
+		SkinningHandler.GetSkinnedTrianglePositions(Accessor, Indices[0], Indices[1], Indices[2], Positions[0], Positions[1], Positions[2]);
+
+		FVector3f TangentsX[3];
+		FVector3f TangentsY[3];
+		FVector3f TangentsZ[3];
+		SkinningHandler.GetSkinnedTangentBasis(Accessor, Indices[0], TangentsX[0], TangentsY[0], TangentsZ[0]);
+		SkinningHandler.GetSkinnedTangentBasis(Accessor, Indices[1], TangentsX[1], TangentsY[1], TangentsZ[1]);
+		SkinningHandler.GetSkinnedTangentBasis(Accessor, Indices[2], TangentsX[2], TangentsY[2], TangentsZ[2]);
+
+		OutPositionParam.SetAndAdvance(BarycentricInterpolate(BaryCoord, Positions[0], Positions[1], Positions[2]));
+		OutTangentZParam.SetAndAdvance(BarycentricInterpolate(BaryCoord, TangentsX[0], TangentsX[1], TangentsX[2]));
+		OutTangentYParam.SetAndAdvance(BarycentricInterpolate(BaryCoord, TangentsY[0], TangentsY[1], TangentsY[2]));
+		OutTangentXParam.SetAndAdvance(BarycentricInterpolate(BaryCoord, TangentsZ[0], TangentsZ[1], TangentsZ[2]));
+	}
+}
+
 template<typename SkinningHandlerType, typename TransformHandlerType, typename VertexAccessorType, typename bInterpolated>
 void UNiagaraDataInterfaceSkeletalMesh::GetTriCoordSkinnedData(FVectorVMExternalFunctionContext& Context)
 {
diff --git a/Engine/Plugins/FX/Niagara/Source/Niagara/Private/NDISkeletalMesh_VertexSampling.cpp b/Engine/Plugins/FX/Niagara/Source/Niagara/Private/NDISkeletalMesh_VertexSampling.cpp
--- a/Engine/Plugins/FX/Niagara/Source/Niagara/Private/NDISkeletalMesh_VertexSampling.cpp
+++ b/Engine/Plugins/FX/Niagara/Source/Niagara/Private/NDISkeletalMesh_VertexSampling.cpp
@@ -14,6 +14,7 @@
 DECLARE_CYCLE_STAT(TEXT("Skel Mesh Vertex Sampling"), STAT_NiagaraSkel_Vertex_Sample, STATGROUP_Niagara);
 
 //Final binders for all static mesh interface functions.
+DEFINE_NDI_FUNC_BINDER(UNiagaraDataInterfaceSkeletalMesh, GetVertexData);
 DEFINE_NDI_FUNC_BINDER(UNiagaraDataInterfaceSkeletalMesh, GetVertexSkinnedData);
 DEFINE_NDI_DIRECT_FUNC_BINDER(UNiagaraDataInterfaceSkeletalMesh, GetVertexColor);
 DEFINE_NDI_DIRECT_FUNC_BINDER(UNiagaraDataInterfaceSkeletalMesh, GetVertexColorFallback);
@@ -24,6 +25,7 @@ DEFINE_NDI_FUNC_BINDER(UNiagaraDataInterfaceSkeletalMesh, RandomFilteredVertex);
 DEFINE_NDI_FUNC_BINDER(UNiagaraDataInterfaceSkeletalMesh, GetFilteredVertexCount);
 DEFINE_NDI_FUNC_BINDER(UNiagaraDataInterfaceSkeletalMesh, GetFilteredVertexAt);
 
+const FName FSkeletalMeshInterfaceHelper::GetVertexDataName("GetVertexData");
 const FName FSkeletalMeshInterfaceHelper::GetSkinnedVertexDataName("GetSkinnedVertexData");
 const FName FSkeletalMeshInterfaceHelper::GetSkinnedVertexDataWSName("GetSkinnedVertexDataWS");
 const FName FSkeletalMeshInterfaceHelper::GetVertexColorName("GetVertexColor");
@@ -40,6 +42,22 @@ const FName FSkeletalMeshInterfaceHelper::GetFilteredVertexAtName("GetFilteredVe
 
 void UNiagaraDataInterfaceSkeletalMesh::GetVertexSamplingFunctions(TArray<FNiagaraFunctionSignature>& OutFunctions)
 {
+	{
+		FNiagaraFunctionSignature Sig;
+		Sig.Name = FSkeletalMeshInterfaceHelper::GetVertexDataName;
+		Sig.Inputs.Add(FNiagaraVariable(FNiagaraTypeDefinition(GetClass()), TEXT("SkeletalMesh")));
+		Sig.Inputs.Add(FNiagaraVariable(FNiagaraTypeDefinition::GetIntDef(), TEXT("Vertex")));
+		Sig.Outputs.Add(FNiagaraVariable(FNiagaraTypeDefinition::GetVec3Def(), TEXT("Position")));
+		Sig.Outputs.Add(FNiagaraVariable(FNiagaraTypeDefinition::GetVec3Def(), TEXT("Normal")));
+		Sig.Outputs.Add(FNiagaraVariable(FNiagaraTypeDefinition::GetVec3Def(), TEXT("Binormal")));
+		Sig.Outputs.Add(FNiagaraVariable(FNiagaraTypeDefinition::GetVec3Def(), TEXT("Tangent")));
+		Sig.bMemberFunction = true;
+		Sig.bRequiresContext = false;
+#if WITH_EDITORONLY_DATA
+		Sig.Description = LOCTEXT("GetVertexDataDesc", "Returns bind pose for the vertex");
+#endif
+		OutFunctions.Add(Sig);
+	}
 	{
 		FNiagaraFunctionSignature Sig;
 		Sig.Name = FSkeletalMeshInterfaceHelper::GetSkinnedVertexDataName;
@@ -174,7 +192,12 @@ void UNiagaraDataInterfaceSkeletalMesh::GetVertexSamplingFunctions(TArray<FNiaga
 void UNiagaraDataInterfaceSkeletalMesh::BindVertexSamplingFunction(const FVMExternalFunctionBindingInfo& BindingInfo, FNDISkeletalMesh_InstanceData* InstanceData, FVMExternalFunction &OutFunc)
 {
 	//////////////////////////////////////////////////////////////////////////
-	if (BindingInfo.Name == FSkeletalMeshInterfaceHelper::GetSkinnedVertexDataName)
+	if (BindingInfo.Name == FSkeletalMeshInterfaceHelper::GetVertexDataName)
+	{
+		check(BindingInfo.GetNumInputs() == 2 && BindingInfo.GetNumOutputs() == 12);
+		TVertexAccessorBinder<NDI_FUNC_BINDER(UNiagaraDataInterfaceSkeletalMesh, GetVertexData)>::Bind(this, BindingInfo, InstanceData, OutFunc);
+	}
+	else if (BindingInfo.Name == FSkeletalMeshInterfaceHelper::GetSkinnedVertexDataName)
 	{
 		check(BindingInfo.GetNumInputs() == 2 && BindingInfo.GetNumOutputs() == 15);
 		TSkinningModeBinder<TNDIExplicitBinder<FNDITransformHandlerNoop, TVertexAccessorBinder<NDI_FUNC_BINDER(UNiagaraDataInterfaceSkeletalMesh, GetVertexSkinnedData)>>>::BindCheckCPUAccess(this, BindingInfo, InstanceData, OutFunc);
@@ -677,6 +700,58 @@ struct FGetVertexSkinnedDataOutputHandler
 	const bool bNeedsTangentZ;
 };
 
+template<typename VertexAccessorType>
+void UNiagaraDataInterfaceSkeletalMesh::GetVertexData(FVectorVMExternalFunctionContext& Context)
+{
+	SCOPE_CYCLE_COUNTER(STAT_NiagaraSkel_Vertex_Sample);
+	VectorVM::FUserPtrHandler<FNDISkeletalMesh_InstanceData> InstData(Context);
+	FNDIInputParam<int32> VertParam(Context);
+
+	FNDIOutputParam<FVector3f> OutPositionParam(Context);
+	FNDIOutputParam<FVector3f> OutTangentZParam(Context);
+	FNDIOutputParam<FVector3f> OutTangentYParam(Context);
+	FNDIOutputParam<FVector3f> OutTangentXParam(Context);
+
+	FSkeletalMeshAccessorHelper Accessor;
+	Accessor.Init<TNDISkelMesh_FilterModeNone, TNDISkelMesh_AreaWeightingOff>(InstData);
+	if (Accessor.IsLODAccessible())
+	{
+		FSkinnedPositionAccessorHelper<TNDISkelMesh_SkinningModeNone> SkinningHandler;
+
+		const int32 VertMax = Accessor.LODData->GetNumVertices() - 1;
+		if (VertMax >= 0)
+		{
+			for (int32 i = 0; i < Context.GetNumInstances(); ++i)
+			{
+				const int32 Vertex = FMath::Clamp(VertParam.GetAndAdvance(), 0, VertMax);
+
+				FVector3f Position = SkinningHandler.GetSkinnedVertexPosition(Accessor, Vertex);
+				FVector3f TangentX = FVector3f::ZeroVector;
+				FVector3f TangentY = FVector3f::ZeroVector;
+				FVector3f TangentZ = FVector3f::ZeroVector;
+				SkinningHandler.GetSkinnedTangentBasis(Accessor, Vertex, TangentX, TangentY, TangentZ);
+
+				OutPositionParam.SetAndAdvance(Position);
+				OutTangentXParam.SetAndAdvance(TangentX);
+				OutTangentYParam.SetAndAdvance(TangentY);
+				OutTangentZParam.SetAndAdvance(TangentZ);
+			}
+
+			// We are done
+			return;
+		}
+	}
+
+	// Fall though for bad data
+	for (int32 i = 0; i < Context.GetNumInstances(); ++i)
+	{
+		OutPositionParam.SetAndAdvance(FVector3f::ZeroVector);
+		OutTangentXParam.SetAndAdvance(FVector3f::ZeroVector);
+		OutTangentYParam.SetAndAdvance(FVector3f::ZeroVector);
+		OutTangentZParam.SetAndAdvance(FVector3f::ZeroVector);
+	}
+}
+
 template<typename SkinningHandlerType, typename TransformHandlerType, typename VertexAccessorType>
 void UNiagaraDataInterfaceSkeletalMesh::GetVertexSkinnedData(FVectorVMExternalFunctionContext& Context)
 {
diff --git a/Engine/Plugins/FX/Niagara/Source/Niagara/Private/NiagaraDataInterfaceSkeletalMesh.cpp b/Engine/Plugins/FX/Niagara/Source/Niagara/Private/NiagaraDataInterfaceSkeletalMesh.cpp
--- a/Engine/Plugins/FX/Niagara/Source/Niagara/Private/NiagaraDataInterfaceSkeletalMesh.cpp
+++ b/Engine/Plugins/FX/Niagara/Source/Niagara/Private/NiagaraDataInterfaceSkeletalMesh.cpp
@@ -2977,6 +2977,11 @@ bool UNiagaraDataInterfaceSkeletalMesh::GetFunctionHLSL(const FNiagaraDataInterf
 		static const TCHAR* FormatSample = TEXT("void {InstanceFunctionName} (in {MeshTriCoordinateStructName} InCoord, out bool IsValid) { {GetDISkelMeshContextName} IsValid = InCoord.Tri < DIContext.MeshTriangleCount; }");
 		OutHLSL += FString::Format(FormatSample, ArgsSample);
 	}
+	else if (FunctionInfo.DefinitionName == FSkeletalMeshInterfaceHelper::GetTriangleDataName)
+	{
+		static const TCHAR* FormatSample = TEXT("void {InstanceFunctionName} (in {MeshTriCoordinateStructName} InCoord, out float3 OutPosition, out float3 OutNormal, out float3 OutBinormal, out float3 OutTangent) { {GetDISkelMeshContextName} DISkelMesh_GetPointOnTriangle(DIContext, InCoord.Tri, InCoord.BaryCoord, OutPosition, OutTangent, OutBinormal, OutNormal); }");
+		OutHLSL += FString::Format(FormatSample, ArgsSample);
+	}
 	else if (FunctionInfo.DefinitionName == FSkeletalMeshInterfaceHelper::GetSkinnedTriangleDataWSName)
 	{
 		static const TCHAR* FormatSample = TEXT("void {InstanceFunctionName} (in {MeshTriCoordinateStructName} InCoord, out float3 OutPosition, out float3 OutVelocity, out float3 OutNormal, out float3 OutBinormal, out float3 OutTangent) { {GetDISkelMeshContextName} DISKelMesh_GetSkinnedTriangleDataWS(DIContext, InCoord.Tri, InCoord.BaryCoord, OutPosition, OutVelocity, OutNormal, OutBinormal, OutTangent); }");
@@ -3061,6 +3066,11 @@ bool UNiagaraDataInterfaceSkeletalMesh::GetFunctionHLSL(const FNiagaraDataInterf
 	}
 	//////////////////////////////////////////////////////////////////////////////////////////////////////////////////////////////
 	// Vertex Sampling
+	else if (FunctionInfo.DefinitionName == FSkeletalMeshInterfaceHelper::GetVertexDataName)
+	{
+		static const TCHAR* FormatSample = TEXT("void {InstanceFunctionName} (in int Vertex, out float3 OutPosition, out float3 OutNormal, out float3 OutBinormal, out float3 OutTangent) { {GetDISkelMeshContextName} DISkelMesh_GetVertex(DIContext, Vertex, OutPosition, OutTangent, OutBinormal, OutNormal); }");
+		OutHLSL += FString::Format(FormatSample, ArgsSample);
+	}
 	else if (FunctionInfo.DefinitionName == FSkeletalMeshInterfaceHelper::GetSkinnedVertexDataName)
 	{
 		static const TCHAR* FormatSample = TEXT("void {InstanceFunctionName} (in int Vertex, out float3 OutPosition, out float3 OutVelocity, out float3 OutNormal, out float3 OutBinormal, out float3 OutTangent) { {GetDISkelMeshContextName} DISkelMesh_GetSkinnedVertex(DIContext, Vertex, OutPosition, OutVelocity, OutNormal, OutBinormal, OutTangent); }");
