- Adjust display names to reduce give clarity around component & cascade

#rnx

diff --git a/Engine/Plugins/FX/Niagara/Source/Niagara/Public/NiagaraComponent.h b/Engine/Plugins/FX/Niagara/Source/Niagara/Public/NiagaraComponent.h
--- a/Engine/Plugins/FX/Niagara/Source/Niagara/Public/NiagaraComponent.h
+++ b/Engine/Plugins/FX/Niagara/Source/Niagara/Public/NiagaraComponent.h
@@ -32,7 +32,7 @@ DECLARE_DYNAMIC_MULTICAST_DELEGATE_OneParam(FOnNiagaraSystemFinished, class UNia
 * @see ANiagaraActor
 * @see UNiagaraSystem
 */
-UCLASS(ClassGroup = (Rendering, Common), hidecategories = Object, hidecategories = Physics, hidecategories = Collision, showcategories = Trigger, editinlinenew, meta = (BlueprintSpawnableComponent, DisplayName = "Niagara Particle System"))
+UCLASS(ClassGroup = (Rendering, Common), hidecategories = Object, hidecategories = Physics, hidecategories = Collision, showcategories = Trigger, editinlinenew, meta = (BlueprintSpawnableComponent, DisplayName = "Niagara Particle System Component"))
 class NIAGARA_API UNiagaraComponent : public UFXSystemComponent
 {
 	friend struct FNiagaraScalabilityManager;
diff --git a/Engine/Source/Runtime/Engine/Classes/Particles/ParticleSystemComponent.h b/Engine/Source/Runtime/Engine/Classes/Particles/ParticleSystemComponent.h
--- a/Engine/Source/Runtime/Engine/Classes/Particles/ParticleSystemComponent.h
+++ b/Engine/Source/Runtime/Engine/Classes/Particles/ParticleSystemComponent.h
@@ -430,7 +430,7 @@ public:
 /** 
  * A particle emitter.
  */
-UCLASS(ClassGroup=(Rendering), hidecategories=Object, hidecategories=Physics, hidecategories=Collision, showcategories=Trigger, editinlinenew, meta=(BlueprintSpawnableComponent))
+UCLASS(ClassGroup=(Rendering), hidecategories=Object, hidecategories=Physics, hidecategories=Collision, showcategories=Trigger, editinlinenew, meta=(BlueprintSpawnableComponent, DisplayName = "Cascade Particle System Component"))
 class ENGINE_API UParticleSystemComponent : public UFXSystemComponent
 {
 	friend class FParticleSystemWorldManager;
