- Add null check for GetScriptData

#rnx

diff --git a/Engine/Plugins/FX/Niagara/Source/NiagaraEditor/Private/ViewModels/Stack/NiagaraStackScriptItemGroup.cpp b/Engine/Plugins/FX/Niagara/Source/NiagaraEditor/Private/ViewModels/Stack/NiagaraStackScriptItemGroup.cpp
--- a/Engine/Plugins/FX/Niagara/Source/NiagaraEditor/Private/ViewModels/Stack/NiagaraStackScriptItemGroup.cpp
+++ b/Engine/Plugins/FX/Niagara/Source/NiagaraEditor/Private/ViewModels/Stack/NiagaraStackScriptItemGroup.cpp
@@ -843,10 +843,14 @@ TOptional<UNiagaraStackEntry::FDropRequestResponse> UNiagaraStackScriptItemGroup
 		return FDropRequestResponse(TOptional<EItemDropZone>(), LOCTEXT("CantMoveModuleError", "This inherited module can't be moved."));
 	}
 
-	TArray<ENiagaraScriptUsage> SourceUsages = SourceModuleItem->GetModuleNode().GetScriptData()->GetSupportedUsageContexts();
-	if (SourceUsages.ContainsByPredicate([this](ENiagaraScriptUsage SourceUsage) { return UNiagaraScript::IsEquivalentUsage(ScriptUsage, SourceUsage); }) == false)
+	TArray<ENiagaraScriptUsage> SourceUsages;
+	if ( FVersionedNiagaraScriptData* ScriptData = SourceModuleItem->GetModuleNode().GetScriptData() )
 	{
-		return FDropRequestResponse(TOptional<EItemDropZone>(), LOCTEXT("CantMoveModuleByUsage", "This module can't be moved to this section of the\nstack because it's not valid for this usage context."));
+		ScriptData->GetSupportedUsageContexts();
+		if (SourceUsages.ContainsByPredicate([this](ENiagaraScriptUsage SourceUsage) { return UNiagaraScript::IsEquivalentUsage(ScriptUsage, SourceUsage); }) == false)
+		{
+			return FDropRequestResponse(TOptional<EItemDropZone>(), LOCTEXT("CantMoveModuleByUsage", "This module can't be moved to this section of the\nstack because it's not valid for this usage context."));
+		}
 	}
 
 	TOptional<EItemDropZone> TargetDropZone = GetTargetDropZoneForTargetEntry(this, TargetEntry, DropRequest.DropZone);
