- Pulled across a few changes / fixes

#rnx

diff --git a/Engine/Plugins/FX/Niagara/Source/Niagara/Classes/NiagaraComputeExecutionContext.h b/Engine/Plugins/FX/Niagara/Source/Niagara/Classes/NiagaraComputeExecutionContext.h
--- a/Engine/Plugins/FX/Niagara/Source/Niagara/Classes/NiagaraComputeExecutionContext.h
+++ b/Engine/Plugins/FX/Niagara/Source/Niagara/Classes/NiagaraComputeExecutionContext.h
@@ -184,6 +184,8 @@ public:
 	FNiagaraDataBuffer* DataBuffers_RT[2] = { nullptr, nullptr };
 	uint32 BufferSwapsThisFrame_RT = 0;
 	uint32 CountOffset_RT = INDEX_NONE;
+	uint32 FinalDispatchGroup_RT = INDEX_NONE;
+	uint32 FinalDispatchGroupInstance_RT = INDEX_NONE;
 
 	// Used only when we multi-tick and need to keep track of pointing back to the correct FNiagaraDataBuffer
 	FNiagaraDataBuffer* DataSetOriginalBuffer_RT = nullptr;
diff --git a/Engine/Plugins/FX/Niagara/Source/Niagara/Classes/NiagaraEmitter.h b/Engine/Plugins/FX/Niagara/Source/Niagara/Classes/NiagaraEmitter.h
--- a/Engine/Plugins/FX/Niagara/Source/Niagara/Classes/NiagaraEmitter.h
+++ b/Engine/Plugins/FX/Niagara/Source/Niagara/Classes/NiagaraEmitter.h
@@ -108,7 +108,9 @@ enum class EParticleAllocationMode : uint8
 	/** This mode tries to estimate the max particle count at runtime by using previous simulations as reference.*/
 	AutomaticEstimate = 0,
 	/** This mode is useful if the particle count can vary wildly at runtime (e.g. due to user parameters) and a lot of reallocations happen.*/
-	ManualEstimate
+	ManualEstimate,
+	/** This mode defines an upper limit on the number of particles that will be simulated.  Useful for rejection sampling where we expect many spawned particles to get killed. */
+	FixedCount
 };
 
 USTRUCT()
@@ -303,7 +305,7 @@ public:
 	The emitter will allocate at least this many particles on it's first tick.
 	This can aid performance by avoiding many allocations as an emitter ramps up to it's max size.
 	*/
-	UPROPERTY(EditAnywhere, AdvancedDisplay, Category = "Emitter", meta = (EditCondition = "AllocationMode == EParticleAllocationMode::ManualEstimate"))
+	UPROPERTY(EditAnywhere, AdvancedDisplay, Category = "Emitter", meta = (EditCondition = "AllocationMode == EParticleAllocationMode::ManualEstimate || AllocationMode == EParticleAllocationMode::FixedCount"))
 	int32 PreAllocationCount;
 
 	UPROPERTY()
@@ -563,6 +565,7 @@ public:
 
 	/* Returns the number of max expected particles for memory allocations. */
 	NIAGARA_API int32 GetMaxParticleCountEstimate();
+	NIAGARA_API uint32 GetMaxInstanceCount() const;
 
 #if WITH_EDITORONLY_DATA
 	NIAGARA_API UNiagaraEmitter* GetParent() const;
@@ -589,8 +592,6 @@ public:
 
 	bool RequiresViewUniformBuffer() const { return bRequiresViewUniformBuffer; }
 
-	uint32 GetMaxInstanceCount() const { return MaxInstanceCount; }
-
 	TConstArrayView<TUniquePtr<FNiagaraBoundsCalculator>> GetBoundsCalculators() const { return MakeArrayView(BoundsCalculators); }
 
 protected:
diff --git a/Engine/Plugins/FX/Niagara/Source/Niagara/Private/NiagaraEmitter.cpp b/Engine/Plugins/FX/Niagara/Source/Niagara/Private/NiagaraEmitter.cpp
--- a/Engine/Plugins/FX/Niagara/Source/Niagara/Private/NiagaraEmitter.cpp
+++ b/Engine/Plugins/FX/Niagara/Source/Niagara/Private/NiagaraEmitter.cpp
@@ -1222,14 +1222,6 @@ void UNiagaraEmitter::CacheFromCompiledData(const FNiagaraDataSetCompiledData* C
 		uint64 MaxBufferElements = (GDebugForcedMaxGPUBufferElements > 0) ? (uint64)GDebugForcedMaxGPUBufferElements : GetMaxBufferDimension();
 		// Don't just cast the result of the division to 32-bit, since that will produce garbage if MaxNumInstances is larger than UINT_MAX. Saturate instead.
 		MaxInstanceCount = (uint32)FMath::Min(MaxBufferElements / MaxGPUBufferComponents, (uint64)UINT_MAX);
-
-		if (SimTarget == ENiagaraSimTarget::GPUComputeSim)
-		{
-			// On GPU, the size of the allocated buffers must be a multiple of NiagaraComputeMaxThreadGroupSize, so round down.
-			MaxInstanceCount = FMath::DivideAndRoundDown(MaxInstanceCount, NiagaraComputeMaxThreadGroupSize) * NiagaraComputeMaxThreadGroupSize;
-			// We will need an extra scratch instance, so the maximum number of usable instances is one less than the value we computed.
-			MaxInstanceCount -= 1;
-		}
 	}
 	else
 	{
@@ -2163,7 +2155,8 @@ int32 UNiagaraEmitter::AddRuntimeAllocation(uint64 ReporterHandle, int32 Allocat
 
 int32 UNiagaraEmitter::GetMaxParticleCountEstimate()
 {
-	if (AllocationMode == EParticleAllocationMode::ManualEstimate)
+	if ((AllocationMode == EParticleAllocationMode::ManualEstimate)
+		|| (AllocationMode == EParticleAllocationMode::FixedCount))
 	{
 		return PreAllocationCount;
 	}
@@ -2193,6 +2186,26 @@ int32 UNiagaraEmitter::GetMaxParticleCountEstimate()
 	return RuntimeEstimation.AllocationEstimate;
 }
 
+uint32 UNiagaraEmitter::GetMaxInstanceCount() const
+{
+	uint32 Result = MaxInstanceCount;
+
+	if (AllocationMode == EParticleAllocationMode::FixedCount)
+	{
+		Result = FMath::Min<uint32>(Result, PreAllocationCount);
+	}
+
+	if (SimTarget == ENiagaraSimTarget::GPUComputeSim)
+	{
+		// On GPU, the size of the allocated buffers must be a multiple of NiagaraComputeMaxThreadGroupSize, so round down.
+		Result = FMath::DivideAndRoundDown(Result, NiagaraComputeMaxThreadGroupSize) * NiagaraComputeMaxThreadGroupSize;
+		// We will need an extra scratch instance, so the maximum number of usable instances is one less than the value we computed.
+		Result -= 1;
+	}
+
+	return Result;
+}
+
 void UNiagaraEmitter::GenerateStatID()const
 {
 #if STATS
diff --git a/Engine/Plugins/FX/Niagara/Source/Niagara/Private/NiagaraGpuComputeDispatch.cpp b/Engine/Plugins/FX/Niagara/Source/Niagara/Private/NiagaraGpuComputeDispatch.cpp
--- a/Engine/Plugins/FX/Niagara/Source/Niagara/Private/NiagaraGpuComputeDispatch.cpp
+++ b/Engine/Plugins/FX/Niagara/Source/Niagara/Private/NiagaraGpuComputeDispatch.cpp
@@ -881,6 +881,8 @@ void FNiagaraGpuComputeDispatch::PrepareTicksForProxy(FRHICommandListImmediate&
 		ComputeContext->CurrentMaxInstances_RT = 0;
 		ComputeContext->CurrentMaxAllocateInstances_RT = 0;
 		ComputeContext->BufferSwapsThisFrame_RT = 0;
+		ComputeContext->FinalDispatchGroup_RT = INDEX_NONE;
+		ComputeContext->FinalDispatchGroupInstance_RT = INDEX_NONE;
 	}
 
 	if (ComputeProxy->PendingTicks.Num() == 0)
@@ -1052,10 +1054,12 @@ void FNiagaraGpuComputeDispatch::PrepareTicksForProxy(FRHICommandListImmediate&
 				}
 			}
 
-			// The final instance in the last group we modified is the final stage for this tick
+			// Set this as the last stage and store the final dispatch group / instance
 			FNiagaraGpuDispatchGroup& FinalDispatchGroup = GpuDispatchList.DispatchGroups[iInstanceCurrDispatchGroup - 1];
 			FinalDispatchGroup.DispatchInstances.Last().SimStageData.bLastStage = true;
-			FinalDispatchGroup.DispatchInstances.Last().SimStageData.bSetDataToRender = Tick.bIsFinalTick;
+
+			ComputeContext->FinalDispatchGroup_RT = iInstanceCurrDispatchGroup - 1;
+			ComputeContext->FinalDispatchGroupInstance_RT = FinalDispatchGroup.DispatchInstances.Num() - 1;
 
 			// Keep track of where the next set of dispatch should occur
 			iTickStartDispatchGroup = FMath::Max(iTickStartDispatchGroup, iInstanceCurrDispatchGroup);
@@ -1094,6 +1098,11 @@ void FNiagaraGpuComputeDispatch::PrepareTicksForProxy(FRHICommandListImmediate&
 			continue;
 		}
 
+		// Ensure we set the data to render as the context may have been dropped during a multi-tick
+		check(ComputeContext->FinalDispatchGroup_RT != INDEX_NONE);
+		FNiagaraGpuDispatchGroup& FinalDispatchGroup = GpuDispatchList.DispatchGroups[ComputeContext->FinalDispatchGroup_RT];
+		FinalDispatchGroup.DispatchInstances[ComputeContext->FinalDispatchGroupInstance_RT].SimStageData.bSetDataToRender = true;
+
 		// We need to store the current data from the main data set as we will be temporarily stomping it during multi-ticking
 		ComputeContext->DataSetOriginalBuffer_RT = ComputeContext->MainDataSet->GetCurrentData();
 
