- Ensure we are not using double variants in templated code
- Fixes for using the wrong type which will result in a check

#rnx

diff --git a/Engine/Plugins/FX/Niagara/Source/Niagara/Private/NiagaraRendererComponents.cpp b/Engine/Plugins/FX/Niagara/Source/Niagara/Private/NiagaraRendererComponents.cpp
--- a/Engine/Plugins/FX/Niagara/Source/Niagara/Private/NiagaraRendererComponents.cpp
+++ b/Engine/Plugins/FX/Niagara/Source/Niagara/Private/NiagaraRendererComponents.cpp
@@ -88,17 +88,17 @@ void ConvertVariableToType(const FNiagaraVariable& SourceVariable, FNiagaraVaria
 	else if (SourceType == FNiagaraTypeDefinition::GetVec3Def() && TargetType == UNiagaraComponentRendererProperties::GetFVectorDef())
 	{
 		FVector3f Data = SourceVariable.GetValue<FVector3f>();
-		TargetVariable.SetValue<FVector>(FVector(Data));
+		TargetVariable.SetDoubleValue<FVector>(FVector(Data));
 	}
 	else if (SourceType == FNiagaraTypeDefinition::GetVec4Def() && TargetType == UNiagaraComponentRendererProperties::GetFVector4Def())
 	{
 		FVector4f Data = SourceVariable.GetValue<FVector4f>();
-		TargetVariable.SetValue<FVector>(FVector4(Data));
+		TargetVariable.SetDoubleValue<FVector4>(FVector4(Data));
 	}
 	else if (SourceType == FNiagaraTypeDefinition::GetQuatDef() && TargetType == UNiagaraComponentRendererProperties::GetFQuatDef())
 	{
 		FQuat4f Data = SourceVariable.GetValue<FQuat4f>();
-		TargetVariable.SetValue<FQuat>(FQuat(Data));
+		TargetVariable.SetDoubleValue<FQuat>(FQuat(Data));
 	}
 	else if (SourceType == FNiagaraTypeDefinition::GetQuatDef() && TargetType == UNiagaraComponentRendererProperties::GetFRotatorDef())
 	{
diff --git a/Engine/Plugins/FX/Niagara/Source/Niagara/Public/NiagaraTypes.h b/Engine/Plugins/FX/Niagara/Source/Niagara/Public/NiagaraTypes.h
--- a/Engine/Plugins/FX/Niagara/Source/Niagara/Public/NiagaraTypes.h
+++ b/Engine/Plugins/FX/Niagara/Source/Niagara/Public/NiagaraTypes.h
@@ -1548,6 +1548,7 @@ struct FNiagaraVariable : public FNiagaraVariableBase
 	template<typename T>
 	void SetValue(const T& Data)
 	{
+		static_assert(!TIsUECoreVariant<T, double>::Value, "Double core variant, please use SetDoubleValue.");
 		check(sizeof(T) == TypeDefHandle->GetSize());
 		AllocateData();
 		FMemory::Memcpy(VarData.GetData(), &Data, VarData.Num());
@@ -1556,6 +1557,27 @@ struct FNiagaraVariable : public FNiagaraVariableBase
 	template<typename T>
 	T GetValue() const
 	{
+		static_assert(!TIsUECoreVariant<T, double>::Value, "Double core variant, please use GetDoubleValue.");
+		check(sizeof(T) == TypeDefHandle->GetSize());
+		check(IsDataAllocated());
+		T Value;
+		FMemory::Memcpy(&Value, GetData(), TypeDefHandle->GetSize());
+		return Value;
+	}
+
+	template<typename T>
+	void SetDoubleValue(const T& Data)
+	{
+		static_assert(TIsUECoreVariant<T, double>::Value, "Float core variant, please use SetValue.");
+		check(sizeof(T) == TypeDefHandle->GetSize());
+		AllocateData();
+		FMemory::Memcpy(VarData.GetData(), &Data, VarData.Num());
+	}
+
+	template<typename T>
+	T GetDoubleValue() const
+	{
+		static_assert(TIsUECoreVariant<T, double>::Value, "Float core variant, please use GetValue.");
 		check(sizeof(T) == TypeDefHandle->GetSize());
 		check(IsDataAllocated());
 		T Value;
diff --git a/Engine/Plugins/FX/Niagara/Source/NiagaraEditor/Private/NiagaraEmitterFactoryNew.cpp b/Engine/Plugins/FX/Niagara/Source/NiagaraEditor/Private/NiagaraEmitterFactoryNew.cpp
--- a/Engine/Plugins/FX/Niagara/Source/NiagaraEditor/Private/NiagaraEmitterFactoryNew.cpp
+++ b/Engine/Plugins/FX/Niagara/Source/NiagaraEditor/Private/NiagaraEmitterFactoryNew.cpp
@@ -102,6 +102,7 @@ template<typename ValueType>
 void SetRapidIterationParameter(FString UniqueEmitterName, UNiagaraScript& TargetScript, UNiagaraNodeFunctionCall& TargetFunctionCallNode,
 	FName InputName, FNiagaraTypeDefinition InputType, ValueType Value)
 {
+	static_assert(!TIsUECoreVariant<ValueType, double>::Value, "Double core variant. Must be float type!");
 	FNiagaraParameterHandle InputHandle = FNiagaraParameterHandle::CreateModuleParameterHandle(InputName);
 	FNiagaraParameterHandle AliasedInputHandle = FNiagaraParameterHandle::CreateAliasedModuleParameterHandle(InputHandle, &TargetFunctionCallNode);
 	FNiagaraVariable RapidIterationParameter = FNiagaraStackGraphUtilities::CreateRapidIterationParameter(UniqueEmitterName, TargetScript.GetUsage(),
@@ -181,7 +182,7 @@ UObject* UNiagaraEmitterFactoryNew::FactoryCreateNew(UClass* Class, UObject* InP
 			if (AddVelocityNode != nullptr)
 			{
 				SetRapidIterationParameter(NewEmitter->GetUniqueEmitterName(), *NewEmitter->SpawnScriptProps.Script, *AddVelocityNode,
-					"Velocity", FNiagaraTypeDefinition::GetVec3Def(), FVector(0.0f, 0.0f, 100.0f));
+					"Velocity", FNiagaraTypeDefinition::GetVec3Def(), FVector3f(0.0f, 0.0f, 100.0f));
 			}
 
 			TArray<FNiagaraVariable> Vars =
