- Fix for LWC disabled

#rnx

diff --git a/Engine/Plugins/FX/Niagara/Source/Niagara/Public/NiagaraTypes.h b/Engine/Plugins/FX/Niagara/Source/Niagara/Public/NiagaraTypes.h
--- a/Engine/Plugins/FX/Niagara/Source/Niagara/Public/NiagaraTypes.h
+++ b/Engine/Plugins/FX/Niagara/Source/Niagara/Public/NiagaraTypes.h
@@ -1568,21 +1568,29 @@ struct FNiagaraVariable : public FNiagaraVariableBase
 	template<typename T>
 	void SetDoubleValue(const T& Data)
 	{
+	#if UE_LARGE_WORLD_COORDINATES_DISABLED
+		SetValue(Data);
+	#else
 		static_assert(TIsUECoreVariant<T, double>::Value, "Float core variant, please use SetValue.");
 		check(sizeof(T) == TypeDefHandle->GetSize());
 		AllocateData();
 		FMemory::Memcpy(VarData.GetData(), &Data, VarData.Num());
+	#endif
 	}
 
 	template<typename T>
 	T GetDoubleValue() const
 	{
+	#if UE_LARGE_WORLD_COORDINATES_DISABLED
+		return GetValue(Data);
+	#else
 		static_assert(TIsUECoreVariant<T, double>::Value, "Float core variant, please use GetValue.");
 		check(sizeof(T) == TypeDefHandle->GetSize());
 		check(IsDataAllocated());
 		T Value;
 		FMemory::Memcpy(&Value, GetData(), TypeDefHandle->GetSize());
 		return Value;
+	#endif
 	}
 
 	void SetData(const uint8* Data)
