- One more fix for LWC disabled

#rnx

diff --git a/Engine/Plugins/FX/Niagara/Source/Niagara/Public/NiagaraTypes.h b/Engine/Plugins/FX/Niagara/Source/Niagara/Public/NiagaraTypes.h
--- a/Engine/Plugins/FX/Niagara/Source/Niagara/Public/NiagaraTypes.h
+++ b/Engine/Plugins/FX/Niagara/Source/Niagara/Public/NiagaraTypes.h
@@ -1582,7 +1582,7 @@ struct FNiagaraVariable : public FNiagaraVariableBase
 	T GetDoubleValue() const
 	{
 	#if UE_LARGE_WORLD_COORDINATES_DISABLED
-		return GetValue(Data);
+		return GetValue();
 	#else
 		static_assert(TIsUECoreVariant<T, double>::Value, "Float core variant, please use GetValue.");
 		check(sizeof(T) == TypeDefHandle->GetSize());
