Fix for a combo button being constructed twice, resulting in the combo box contents not being properly displayed

diff --git a/Engine/Plugins/FX/Niagara/Source/NiagaraEditorWidgets/Private/DetailCustomizations/SNiagaraNamePropertySelector.cpp b/Engine/Plugins/FX/Niagara/Source/NiagaraEditorWidgets/Private/DetailCustomizations/SNiagaraNamePropertySelector.cpp
--- a/Engine/Plugins/FX/Niagara/Source/NiagaraEditorWidgets/Private/DetailCustomizations/SNiagaraNamePropertySelector.cpp
+++ b/Engine/Plugins/FX/Niagara/Source/NiagaraEditorWidgets/Private/DetailCustomizations/SNiagaraNamePropertySelector.cpp
@@ -76,8 +76,51 @@ void SNiagaraNamePropertySelector::Construct(const FArguments& InArgs, TSharedRe
 		.ContentPadding(FMargin(5, 0))
 		.OnComboBoxOpened(this, &SNiagaraNamePropertySelector::OnComboOpening);
 
-	SAssignNew(ElementButton, SComboButton);
-	ElementButton->Construct(Args);
+	SAssignNew(ElementButton, SComboButton)
+		.ButtonContent()
+		[
+			SNew(SHorizontalBox)
+			+ SHorizontalBox::Slot()
+			.VAlign(VAlign_Center)
+			.AutoWidth()
+			.Padding(1.f, 1.f)
+			[
+				SNew(STextBlock)
+				.ColorAndOpacity(FLinearColor(1, 1, 1, 1))
+				.Text(this, &SNiagaraNamePropertySelector::GetComboText)
+			]
+		]
+		.MenuContent()
+		[
+			SNew(SListViewSelectorDropdownMenu<TSharedPtr<FName>>, SearchBox, ElementsListView)
+			[
+				SNew(SBorder)
+				.BorderImage(FEditorStyle::GetBrush("Menu.Background"))
+				.Padding(2)
+				[
+					SNew(SBox)
+					.WidthOverride(175)
+					[
+						SNew(SVerticalBox)
+						+ SVerticalBox::Slot()
+							.Padding(1.f)
+							.AutoHeight()
+							[
+								SearchBox.ToSharedRef()
+							]
+						+ SVerticalBox::Slot()
+							.MaxHeight(400)
+							.Padding(8.f)
+							[
+								ElementsListView.ToSharedRef()
+							]
+					]
+				]
+			]
+		]
+		.IsFocusable(true)
+		.ContentPadding(FMargin(5, 0))
+		.OnComboBoxOpened(this, &SNiagaraNamePropertySelector::OnComboOpening);
 
 	ElementsListView->EnableToolTipForceField(true);
 	// SComboButton can automatically handle setting focus to a specified control when the combo button is opened
