Added refresh of staticswitch constants before accessing staticswitch pins. This is a bandaid fix. Traversing the map history should not change the graph by calling SetStaticSwitchConstants. Sometimes, the constant resolver isn't valid and will clear the constants, but not re-set them. The graph data cache exposed the problem. Multiple UI refreshes that come in after a value change in the stack cause a staticswitch constant to be updated multiple times, and if the last time before staticswitch pins are accessed clears the constants we get false results.
This happened in UNiagaraStackFunctionInputCollection::RefreshChildrenInternal and could happen elsewhere too.

diff --git a/Engine/Plugins/FX/Niagara/Source/NiagaraEditor/Private/NiagaraGraphDataCache.cpp b/Engine/Plugins/FX/Niagara/Source/NiagaraEditor/Private/NiagaraGraphDataCache.cpp
--- a/Engine/Plugins/FX/Niagara/Source/NiagaraEditor/Private/NiagaraGraphDataCache.cpp
+++ b/Engine/Plugins/FX/Niagara/Source/NiagaraEditor/Private/NiagaraGraphDataCache.cpp
@@ -16,6 +16,8 @@ void FNiagaraGraphDataCache::GetStackFunctionInputPins(UNiagaraNodeFunctionCall&
 		return;
 	}
 
+	OutInputPins.Empty();
+	OutHiddenPins.Empty();
 	FNiagaraStackGraphUtilities::GetStackFunctionInputPinsWithoutCache(FunctionCallNode, OutInputPins, OutHiddenPins, ConstantResolver, Options, bIgnoreDisabled);
 
 	CachedResult.SourceGraphWeak = FunctionCallNode.GetNiagaraGraph();
diff --git a/Engine/Plugins/FX/Niagara/Source/NiagaraEditor/Private/NiagaraScriptMergeManager.cpp b/Engine/Plugins/FX/Niagara/Source/NiagaraEditor/Private/NiagaraScriptMergeManager.cpp
--- a/Engine/Plugins/FX/Niagara/Source/NiagaraEditor/Private/NiagaraScriptMergeManager.cpp
+++ b/Engine/Plugins/FX/Niagara/Source/NiagaraEditor/Private/NiagaraScriptMergeManager.cpp
@@ -212,6 +212,8 @@ FNiagaraStackFunctionMergeAdapter::FNiagaraStackFunctionMergeAdapter(const UNiag
 
 	FString UniqueEmitterName = InOwningEmitter.GetUniqueEmitterName();
 
+	FCompileConstantResolver ConstantResolver(&InOwningEmitter, FNiagaraStackGraphUtilities::GetEmitterOutputNodeForStackNode(*FunctionCallNode)->GetUsage());
+
 	// Collect explicit overrides set via parameter map set nodes.
 	TSet<FName> AliasedInputsAdded;
 	UNiagaraNodeParameterMapSet* OverrideNode = FNiagaraStackGraphUtilities::GetStackFunctionOverrideNode(*FunctionCallNode);
@@ -235,13 +237,12 @@ FNiagaraStackFunctionMergeAdapter::FNiagaraStackFunctionMergeAdapter(const UNiag
 		FCompileConstantResolver Resolver(&InOwningEmitter, ENiagaraScriptUsage::Function);
 		TArray<const UEdGraphPin*> FunctionInputPins;
 		FNiagaraStackGraphUtilities::GetStackFunctionInputPins(*FunctionCallNode, FunctionInputPins, HiddenPins, Resolver, FNiagaraStackGraphUtilities::ENiagaraGetStackFunctionInputPinsOptions::ModuleInputsOnly, false);
-
+		
 		for (const UEdGraphPin* FunctionInputPin : FunctionInputPins)
 		{
 			FNiagaraVariable FunctionInputVariable = NiagaraSchema->PinToNiagaraVariable(FunctionInputPin);
 			if (FunctionInputVariable.IsValid() && FNiagaraStackGraphUtilities::IsRapidIterationType(FunctionInputVariable.GetType()))
 			{
-				FCompileConstantResolver ConstantResolver(&InOwningEmitter, FNiagaraStackGraphUtilities::GetEmitterOutputNodeForStackNode(*FunctionCallNode)->GetUsage());
 				UEdGraphPin* FunctionInputDefaultPin = FunctionCallNode->FindParameterMapDefaultValuePin(FunctionInputPin->PinName, OwningScript->GetUsage(), ConstantResolver);
 				if (FunctionInputDefaultPin != nullptr)
 				{
@@ -318,7 +319,8 @@ FNiagaraStackFunctionMergeAdapter::FNiagaraStackFunctionMergeAdapter(const UNiag
 
 	TArray<UEdGraphPin*> StaticSwitchPins;
 	TSet<UEdGraphPin*> StaticSwitchPinsHidden;
-	FNiagaraStackGraphUtilities::GetStackFunctionStaticSwitchPins(*FunctionCallNode.Get(), StaticSwitchPins, StaticSwitchPinsHidden);
+	FNiagaraStackGraphUtilities::GetStackFunctionStaticSwitchPins(*FunctionCallNode.Get(), StaticSwitchPins, StaticSwitchPinsHidden, ConstantResolver);
+
 	for (UEdGraphPin* StaticSwitchPin : StaticSwitchPins)
 	{
 		// TODO: Only add static switch overrides when the current value is different from the default.  This requires a 
@@ -2745,7 +2747,8 @@ FNiagaraScriptMergeManager::FApplyDiffResults FNiagaraScriptMergeManager::AddInp
 		{
 			TArray<UEdGraphPin*> StaticSwitchPins;
 			TSet<UEdGraphPin*> StaticSwitchPinsHidden;
-			FNiagaraStackGraphUtilities::GetStackFunctionStaticSwitchPins(TargetFunctionCall, StaticSwitchPins, StaticSwitchPinsHidden);
+			FCompileConstantResolver ConstantResolver(BaseEmitterAdapter->GetEditableEmitter(), OwningScript.GetUsage());
+			FNiagaraStackGraphUtilities::GetStackFunctionStaticSwitchPins(TargetFunctionCall, StaticSwitchPins, StaticSwitchPinsHidden, ConstantResolver);
 			UEdGraphPin** MatchingStaticSwitchPinPtr = StaticSwitchPins.FindByPredicate([&OverrideToAdd](UEdGraphPin* StaticSwitchPin) { return StaticSwitchPin->PinName == *OverrideToAdd->GetInputName(); });
 			if (MatchingStaticSwitchPinPtr != nullptr)
 			{
diff --git a/Engine/Plugins/FX/Niagara/Source/NiagaraEditor/Private/ViewModels/Stack/NiagaraStackFunctionInputCollection.cpp b/Engine/Plugins/FX/Niagara/Source/NiagaraEditor/Private/ViewModels/Stack/NiagaraStackFunctionInputCollection.cpp
--- a/Engine/Plugins/FX/Niagara/Source/NiagaraEditor/Private/ViewModels/Stack/NiagaraStackFunctionInputCollection.cpp
+++ b/Engine/Plugins/FX/Niagara/Source/NiagaraEditor/Private/ViewModels/Stack/NiagaraStackFunctionInputCollection.cpp
@@ -215,7 +215,7 @@ void UNiagaraStackFunctionInputCollection::RefreshChildrenInternal(const TArray<
 	// Gather static switch parameters
 	TSet<UEdGraphPin*> HiddenSwitchPins;
 	TArray<UEdGraphPin*> SwitchPins;
-	FNiagaraStackGraphUtilities::GetStackFunctionStaticSwitchPins(*InputFunctionCallNode, SwitchPins, HiddenSwitchPins);
+	FNiagaraStackGraphUtilities::GetStackFunctionStaticSwitchPins(*InputFunctionCallNode, SwitchPins, HiddenSwitchPins, ConstantResolver);
 	for (UEdGraphPin* InputPin : SwitchPins)
 	{
 		// The static switch pin names to not contain the module namespace, as they are not part of the parameter maps.
diff --git a/Engine/Plugins/FX/Niagara/Source/NiagaraEditor/Private/ViewModels/Stack/NiagaraStackGraphUtilities.cpp b/Engine/Plugins/FX/Niagara/Source/NiagaraEditor/Private/ViewModels/Stack/NiagaraStackGraphUtilities.cpp
--- a/Engine/Plugins/FX/Niagara/Source/NiagaraEditor/Private/ViewModels/Stack/NiagaraStackGraphUtilities.cpp
+++ b/Engine/Plugins/FX/Niagara/Source/NiagaraEditor/Private/ViewModels/Stack/NiagaraStackGraphUtilities.cpp
@@ -889,7 +889,8 @@ void FNiagaraStackGraphUtilities::GetStackFunctionInputPins(UNiagaraNodeFunction
 	GetStackFunctionInputPins(FunctionCallNode, OutInputPins, HiddenPins, EmptyResolver, Options, bIgnoreDisabled);
 }
 
-void FNiagaraStackGraphUtilities::GetStackFunctionStaticSwitchPins(UNiagaraNodeFunctionCall& FunctionCallNode, TArray<UEdGraphPin*>& OutInputPins, TSet<UEdGraphPin*>& OutHiddenPins)
+void FNiagaraStackGraphUtilities::GetStackFunctionStaticSwitchPins(UNiagaraNodeFunctionCall& FunctionCallNode, TArray<UEdGraphPin*>& OutInputPins, TSet<UEdGraphPin*>& OutHiddenPins,
+	FCompileConstantResolver& ConstantResolver)
 {
 	const UEdGraphSchema_Niagara* Schema = CastChecked<UEdGraphSchema_Niagara>(FunctionCallNode.GetSchema());
 	UNiagaraGraph* FunctionCallGraph = FunctionCallNode.GetCalledGraph();
@@ -898,8 +899,13 @@ void FNiagaraStackGraphUtilities::GetStackFunctionStaticSwitchPins(UNiagaraNodeF
 		return;
 	}
 
+	FPinCollectorArray InputPins;
+	FunctionCallNode.GetInputPins(InputPins);
+	FNiagaraEditorUtilities::SetStaticSwitchConstants(FunctionCallGraph, InputPins, ConstantResolver);
+
+	TArray<FNiagaraVariable> SwitchInputs = FunctionCallGraph->FindStaticSwitchInputs(false);
 	TArray<FNiagaraVariable> ReachableInputs = FunctionCallGraph->FindStaticSwitchInputs(true);
-	for (FNiagaraVariable SwitchInput : FunctionCallGraph->FindStaticSwitchInputs(false))
+	for (FNiagaraVariable SwitchInput : SwitchInputs)
 	{
 		FEdGraphPinType PinType = Schema->TypeDefinitionToPinType(SwitchInput.GetType());
 		for (UEdGraphPin* Pin : FunctionCallNode.Pins)
diff --git a/Engine/Plugins/FX/Niagara/Source/NiagaraEditor/Public/ViewModels/Stack/NiagaraStackGraphUtilities.h b/Engine/Plugins/FX/Niagara/Source/NiagaraEditor/Public/ViewModels/Stack/NiagaraStackGraphUtilities.h
--- a/Engine/Plugins/FX/Niagara/Source/NiagaraEditor/Public/ViewModels/Stack/NiagaraStackGraphUtilities.h
+++ b/Engine/Plugins/FX/Niagara/Source/NiagaraEditor/Public/ViewModels/Stack/NiagaraStackGraphUtilities.h
@@ -103,7 +103,7 @@ namespace FNiagaraStackGraphUtilities
 	/* Module script calls do not have direct inputs, but rely on the parameter map being initialized correctly. This utility function resolves which of the module's parameters are reachable during compilation and returns a list of pins on the parameter map node that do not have to be compiled. */
 	TArray<UEdGraphPin*> GetUnusedFunctionInputPins(UNiagaraNodeFunctionCall& FunctionCallNode, FCompileConstantResolver ConstantResolver);
 
-	void GetStackFunctionStaticSwitchPins(UNiagaraNodeFunctionCall& FunctionCallNode, TArray<UEdGraphPin*>& OutInputPins, TSet<UEdGraphPin*>& OutHiddenPins);
+	void GetStackFunctionStaticSwitchPins(UNiagaraNodeFunctionCall& FunctionCallNode, TArray<UEdGraphPin*>& OutInputPins, TSet<UEdGraphPin*>& OutHiddenPins, FCompileConstantResolver& ConstantResolver);
 
 	void GetStackFunctionOutputVariables(UNiagaraNodeFunctionCall& FunctionCallNode, FCompileConstantResolver ConstantResolver, TArray<FNiagaraVariable>& OutOutputVariables, TArray<FNiagaraVariable>& OutOutputVariablesWithOriginalAliasesIntact);
 
