- Handle invalid output texture size in the baker, i.e. one that will result in a frame size < 0

#rnx

diff --git a/Engine/Plugins/FX/Niagara/Source/Niagara/Classes/NiagaraBakerSettings.h b/Engine/Plugins/FX/Niagara/Source/Niagara/Classes/NiagaraBakerSettings.h
--- a/Engine/Plugins/FX/Niagara/Source/Niagara/Classes/NiagaraBakerSettings.h
+++ b/Engine/Plugins/FX/Niagara/Source/Niagara/Classes/NiagaraBakerSettings.h
@@ -64,10 +64,12 @@ struct FNiagaraBakerTextureSettings
 	UPROPERTY(EditAnywhere, Category = "Texture")
 	TObjectPtr<UTexture2D> GeneratedTexture = nullptr;
 
+	bool IsValidForBake() const { return FrameSize.X > 0 && FrameSize.Y > 0; }
+
 	bool Equals(const FNiagaraBakerTextureSettings& Other) const;
 
-	FNiagaraBakerTextureSettings() :
-		bUseFrameSize(false)
+	FNiagaraBakerTextureSettings()
+		: bUseFrameSize(false)
 	{}
 };
 
diff --git a/Engine/Plugins/FX/Niagara/Source/NiagaraEditor/Private/ViewModels/NiagaraBakerViewModel.cpp b/Engine/Plugins/FX/Niagara/Source/NiagaraEditor/Private/ViewModels/NiagaraBakerViewModel.cpp
--- a/Engine/Plugins/FX/Niagara/Source/NiagaraEditor/Private/ViewModels/NiagaraBakerViewModel.cpp
+++ b/Engine/Plugins/FX/Niagara/Source/NiagaraEditor/Private/ViewModels/NiagaraBakerViewModel.cpp
@@ -159,13 +159,20 @@ void FNiagaraBakerViewModel::RenderBaker()
 
 	for ( const auto& OutputTexture : BakerSettings->OutputTextures )
 	{
-		UTextureRenderTarget2D* RenderTarget = NewObject<UTextureRenderTarget2D>();
-		RenderTarget->AddToRoot();
-		RenderTarget->ClearColor = FLinearColor::Transparent;
-		RenderTarget->TargetGamma = 1.0f;
-		RenderTarget->InitCustomFormat(OutputTexture.FrameSize.X, OutputTexture.FrameSize.Y, PF_FloatRGBA, false);
+		if (OutputTexture.IsValidForBake())
+		{
+			UTextureRenderTarget2D* RenderTarget = NewObject<UTextureRenderTarget2D>();
+			RenderTarget->AddToRoot();
+			RenderTarget->ClearColor = FLinearColor::Transparent;
+			RenderTarget->TargetGamma = 1.0f;
+			RenderTarget->InitCustomFormat(OutputTexture.FrameSize.X, OutputTexture.FrameSize.Y, PF_FloatRGBA, false);
 
-		RenderTargets.Add(RenderTarget);
+			RenderTargets.Add(RenderTarget);
+		}
+		else
+		{
+			RenderTargets.Add(nullptr);
+		}
 		OutputBakers.AddDefaulted_GetRef().SetNumZeroed(OutputTexture.TextureSize.X * OutputTexture.TextureSize.Y);
 	}
 
@@ -197,6 +204,10 @@ void FNiagaraBakerViewModel::RenderBaker()
 		{
 			const FNiagaraBakerTextureSettings& OutputTexture = BakerSettings->OutputTextures[iOutputTexture];
 			UTextureRenderTarget2D* RenderTarget = RenderTargets[iOutputTexture];
+			if (RenderTarget == nullptr)
+			{
+				continue;
+			}
 			FTextureRenderTargetResource* RenderTargetResource = RenderTarget->GameThread_GetRenderTargetResource();
 
 			BakerRenderer.RenderView(PreviewComponent, BakerSettings, WorldTime, RenderTarget, iOutputTexture);
@@ -265,8 +276,11 @@ void FNiagaraBakerViewModel::RenderBaker()
 	// Clean up render targets
 	for ( UTextureRenderTarget2D* RenderTarget : RenderTargets )
 	{
-		RenderTarget->RemoveFromRoot();
-		RenderTarget = nullptr;
+		if (RenderTarget)
+		{
+			RenderTarget->RemoveFromRoot();
+			RenderTarget = nullptr;
+		}
 	}
 }
 
diff --git a/Engine/Plugins/FX/Niagara/Source/NiagaraEditor/Private/Widgets/SNiagaraBakerViewport.cpp b/Engine/Plugins/FX/Niagara/Source/NiagaraEditor/Private/Widgets/SNiagaraBakerViewport.cpp
--- a/Engine/Plugins/FX/Niagara/Source/NiagaraEditor/Private/Widgets/SNiagaraBakerViewport.cpp
+++ b/Engine/Plugins/FX/Niagara/Source/NiagaraEditor/Private/Widgets/SNiagaraBakerViewport.cpp
@@ -291,44 +291,54 @@ public:
 			FVector2D TextPosition(PreviewViewRect.Min.X + TextStartOffset.X, PreviewViewRect.Min.Y + TextStartOffset.Y);
 			if (BakerSettings->OutputTextures.IsValidIndex(PreviewTextureIndex))
 			{
-				// Ensure render target is the correct size
 				const FNiagaraBakerTextureSettings& OutputTexture = BakerSettings->OutputTextures[PreviewTextureIndex];
-				if ( !RealtimeRenderTarget || RealtimeRenderTarget->SizeX != OutputTexture.FrameSize.X || RealtimeRenderTarget->SizeY != OutputTexture.FrameSize.Y )
+				if ( OutputTexture.IsValidForBake() )
 				{
-					if (RealtimeRenderTarget == nullptr)
+					// Ensure render target is the correct size
+					if ( !RealtimeRenderTarget || RealtimeRenderTarget->SizeX != OutputTexture.FrameSize.X || RealtimeRenderTarget->SizeY != OutputTexture.FrameSize.Y )
 					{
-						RealtimeRenderTarget = NewObject<UTextureRenderTarget2D>();
+						if (RealtimeRenderTarget == nullptr)
+						{
+							RealtimeRenderTarget = NewObject<UTextureRenderTarget2D>();
+						}
+						RealtimeRenderTarget->ClearColor = FLinearColor::Transparent;
+						RealtimeRenderTarget->TargetGamma = 1.0f;
+						RealtimeRenderTarget->InitCustomFormat(OutputTexture.FrameSize.X, OutputTexture.FrameSize.Y, PF_FloatRGBA, false);
 					}
-					RealtimeRenderTarget->ClearColor = FLinearColor::Transparent;
-					RealtimeRenderTarget->TargetGamma = 1.0f;
-					RealtimeRenderTarget->InitCustomFormat(OutputTexture.FrameSize.X, OutputTexture.FrameSize.Y, PF_FloatRGBA, false);
-				}
-
-				// Seek to correct time and render
-				UNiagaraComponent* PreviewComponent = ViewModel->GetPreviewComponent();
-				PreviewComponent->SetSeekDelta(BakerSettings->GetSeekDelta());
-				PreviewComponent->SeekToDesiredAge(WorldTime);
-				PreviewComponent->TickComponent(BakerSettings->GetSeekDelta(), ELevelTick::LEVELTICK_All, nullptr);
 
-				BakerRenderer.RenderView(PreviewComponent, BakerSettings, WorldTime, RealtimeRenderTarget, PreviewTextureIndex);
-
-				const FVector2D HalfPixel(0.5f / float(OutputTexture.FrameSize.X), 0.5f / float(OutputTexture.FrameSize.Y));
-				Canvas->DrawTile(
-					PreviewViewRect.Min.X, PreviewViewRect.Min.Y, PreviewViewRect.Width(), PreviewViewRect.Height(),
-					HalfPixel.X, HalfPixel.Y, 1.0f - HalfPixel.X, 1.0f - HalfPixel.Y,
-					FLinearColor::White,
-					RealtimeRenderTarget->GetResource(),
-					bAlphaBlend ? SE_BLEND_AlphaComposite : SE_BLEND_Opaque
-				);
+					// Seek to correct time and render
+					UNiagaraComponent* PreviewComponent = ViewModel->GetPreviewComponent();
+					PreviewComponent->SetSeekDelta(BakerSettings->GetSeekDelta());
+					PreviewComponent->SeekToDesiredAge(WorldTime);
+					PreviewComponent->TickComponent(BakerSettings->GetSeekDelta(), ELevelTick::LEVELTICK_All, nullptr);
+
+					BakerRenderer.RenderView(PreviewComponent, BakerSettings, WorldTime, RealtimeRenderTarget, PreviewTextureIndex);
+
+					const FVector2D HalfPixel(0.5f / float(OutputTexture.FrameSize.X), 0.5f / float(OutputTexture.FrameSize.Y));
+					Canvas->DrawTile(
+						PreviewViewRect.Min.X, PreviewViewRect.Min.Y, PreviewViewRect.Width(), PreviewViewRect.Height(),
+						HalfPixel.X, HalfPixel.Y, 1.0f - HalfPixel.X, 1.0f - HalfPixel.Y,
+						FLinearColor::White,
+						RealtimeRenderTarget->GetResource(),
+						bAlphaBlend ? SE_BLEND_AlphaComposite : SE_BLEND_Opaque
+					);
+
+					// Display info text
+					if (bShowInfoText)
+					{
+						Canvas->DrawShadowedString(TextPosition.X, TextPosition.Y, TEXT("Live Preview"), DisplayFont, FLinearColor::White);
+						TextPosition.Y += FontHeight;
 
-				// Display info text
-				if (bShowInfoText)
+						Canvas->DrawShadowedString(TextPosition.X + 5.0f, TextPosition.Y, *FString::Printf(TEXT("Texture(%d) FrameSize(%d x %d)"), PreviewTextureIndex, OutputTexture.FrameSize.X, OutputTexture.FrameSize.Y), DisplayFont, FLinearColor::White);
+						TextPosition.Y += FontHeight;
+					}
+				}
+				else
 				{
 					Canvas->DrawShadowedString(TextPosition.X, TextPosition.Y, TEXT("Live Preview"), DisplayFont, FLinearColor::White);
 					TextPosition.Y += FontHeight;
 
-					Canvas->DrawShadowedString(TextPosition.X + 5.0f, TextPosition.Y, *FString::Printf(TEXT("Texture(%d) FrameSize(%d x %d)"), PreviewTextureIndex, OutputTexture.FrameSize.X, OutputTexture.FrameSize.Y), DisplayFont, FLinearColor::White);
-					TextPosition.Y += FontHeight;
+					Canvas->DrawShadowedString(TextPosition.X + 5.0f, TextPosition.Y, *FString::Printf(TEXT("Texture(%d) Is Invalid For Baking FrameSize(%d x %d)"), PreviewTextureIndex, OutputTexture.FrameSize.X, OutputTexture.FrameSize.Y), DisplayFont, FLinearColor::White);
 				}
 			}
 			else
