- Don't include invalid (pending kill) objects or ones that are the CDO when generating filter list

#rnx

diff --git a/Engine/Plugins/FX/Niagara/Source/Niagara/Private/NiagaraDebuggerClient.cpp b/Engine/Plugins/FX/Niagara/Source/Niagara/Private/NiagaraDebuggerClient.cpp
--- a/Engine/Plugins/FX/Niagara/Source/Niagara/Private/NiagaraDebuggerClient.cpp
+++ b/Engine/Plugins/FX/Niagara/Source/Niagara/Private/NiagaraDebuggerClient.cpp
@@ -84,7 +84,7 @@ void FNiagaraDebuggerClient::UpdateClientInfo()
 		for (TObjectIterator<UNiagaraComponent> It; It; ++It)
 		{
 			UNiagaraComponent* Comp = *It;
-			if (Comp)
+			if (IsValid(Comp) && !Comp->HasAnyFlags(EObjectFlags::RF_ClassDefaultObject))
 			{
 				NewInfo->Components.AddUnique(Comp->GetName());
 				Actors.Add(Comp->GetOwner());
@@ -93,9 +93,9 @@ void FNiagaraDebuggerClient::UpdateClientInfo()
 
 		for (AActor* Actor : Actors)
 		{
-			if (Actor)
+			if (IsValid(Actor) && !Actor->HasAnyFlags(EObjectFlags::RF_ClassDefaultObject))
 			{
-				NewInfo->Actors.Add(Actor->GetActorNameOrLabel());
+				NewInfo->Actors.AddUnique(Actor->GetActorNameOrLabel());
 			}
 		}
 
