- Fix a crash if we have a user parameter driving a static mesh but no mesh selected
- Fix meshes not showing in the above situation

diff --git a/Engine/Plugins/FX/Niagara/Source/Niagara/Private/NiagaraRendererMeshes.cpp b/Engine/Plugins/FX/Niagara/Source/Niagara/Private/NiagaraRendererMeshes.cpp
--- a/Engine/Plugins/FX/Niagara/Source/Niagara/Private/NiagaraRendererMeshes.cpp
+++ b/Engine/Plugins/FX/Niagara/Source/Niagara/Private/NiagaraRendererMeshes.cpp
@@ -186,7 +186,7 @@ void FNiagaraRendererMeshes::Initialize(const UNiagaraRendererProperties* InProp
 			MeshData.PivotOffset = MeshProperties.PivotOffset;
 			MeshData.PivotOffsetSpace = MeshProperties.PivotOffsetSpace;
 			MeshData.Scale = MeshProperties.Scale;
-			MeshData.MinimumLOD = MeshProperties.Mesh->GetMinLODIdx();
+			MeshData.MinimumLOD = Mesh->GetMinLODIdx();
 
 			// Create an index remap from mesh material index to it's index in the master material list
 			TArray<UMaterialInterface*> MeshMaterials;
@@ -1456,7 +1456,7 @@ FNiagaraDynamicDataBase* FNiagaraRendererMeshes::GenerateDynamicData(const FNiag
 	// Bail if we have cached mesh render data for any meshes that are no longer valid
 	for (const auto& MeshData : Meshes)
 	{
-		if (!Properties->Meshes.IsValidIndex(MeshData.SourceMeshIndex) || !Properties->Meshes[MeshData.SourceMeshIndex].Mesh)
+		if ( !Properties->Meshes.IsValidIndex(MeshData.SourceMeshIndex) )
 		{
 			return nullptr;
 		}
