Fixed merge error causing a compilation issue in Niagara vertex factory shaders

diff --git a/Engine/Plugins/FX/Niagara/Shaders/Private/NiagaraMeshVertexFactory.ush b/Engine/Plugins/FX/Niagara/Shaders/Private/NiagaraMeshVertexFactory.ush
--- a/Engine/Plugins/FX/Niagara/Shaders/Private/NiagaraMeshVertexFactory.ush
+++ b/Engine/Plugins/FX/Niagara/Shaders/Private/NiagaraMeshVertexFactory.ush
@@ -8,6 +8,11 @@
 #include "NiagaraCommon.ush"
 #include "NiagaraVFCommon.usf"
 
+#if COMPUTESHADER || RAYHITGROUPSHADER
+#include "/Engine/Private/RayTracing/RayTracingCommon.ush"
+#include "/Engine/Private/RayTracing/RayTracingHitGroupCommon.ush"
+#endif
+
 #define USE_PARTICLE_POSITION (NEEDS_PARTICLE_POSITION)
 #define USE_PARTICLE_VELOCITY (NEEDS_PARTICLE_VELOCITY)
 #define USE_PARTICLE_TIME (NEEDS_PARTICLE_TIME)
@@ -990,7 +995,7 @@ float4 VertexFactoryGetTranslatedPrimitiveVolumeBounds(FVertexFactoryInterpolant
 #endif // #if NEEDS_VERTEX_FACTORY_INTERPOLATION
 
 // RHI_RAYTRACING
-#if COMPUTESHADER || RAYHITGROUPSHADER
+#if RAYHITGROUPSHADER
 	FVertexFactoryInput LoadVertexFactoryInputForHGS(uint TriangleIndex, int VertexIndex)
 	{
 		FVertexFactoryInput Input;
@@ -1006,8 +1011,12 @@ float4 VertexFactoryGetTranslatedPrimitiveVolumeBounds(FVertexFactoryInterpolant
 		Input.TangentZ = 0.0;
 		Input.VertexColor = float4(1.0, 1.0, 1.0, 0.5);
 #if NUM_MATERIAL_TEXCOORDS_VERTEX
-		Input.TexCoords[0].x = 0.5f;
-		Input.TexCoords[0].y = 0.5f;
+UNROLL
+		for (int CoordinateIndex = 0; CoordinateIndex < NUM_MATERIAL_TEXCOORDS_VERTEX; CoordinateIndex++)
+		{
+			Input.TexCoords[CoordinateIndex].x = 0.5f;
+			Input.TexCoords[CoordinateIndex].y = 0.5f;
+		}
 #elif USE_PARTICLE_SUBUVS
 		Input.TexCoords[0].x = 0.5f;
 		Input.TexCoords[0].y = 0.5f;
diff --git a/Engine/Plugins/FX/Niagara/Source/NiagaraVertexFactories/Private/NiagaraMeshVertexFactory.cpp b/Engine/Plugins/FX/Niagara/Source/NiagaraVertexFactories/Private/NiagaraMeshVertexFactory.cpp
--- a/Engine/Plugins/FX/Niagara/Source/NiagaraVertexFactories/Private/NiagaraMeshVertexFactory.cpp
+++ b/Engine/Plugins/FX/Niagara/Source/NiagaraVertexFactories/Private/NiagaraMeshVertexFactory.cpp
@@ -156,7 +156,6 @@ IMPLEMENT_VERTEX_FACTORY_TYPE(FNiagaraMeshVertexFactory, "/Plugin/FX/Niagara/Pri
 	  EVertexFactoryFlags::UsedWithMaterials 
 	| EVertexFactoryFlags::SupportsDynamicLighting
 	| EVertexFactoryFlags::SupportsRayTracing
-	| EVertexFactoryFlags::SupportsRayTracingDynamicGeometry
 );
 
 /////////////////////////////////////////////////////////////////////////////////////////////////////
@@ -172,6 +171,5 @@ IMPLEMENT_VERTEX_FACTORY_TYPE(FNiagaraMeshVertexFactoryEx, "/Plugin/FX/Niagara/P
 	  EVertexFactoryFlags::UsedWithMaterials
 	| EVertexFactoryFlags::SupportsDynamicLighting
 	| EVertexFactoryFlags::SupportsRayTracing
-	| EVertexFactoryFlags::SupportsRayTracingDynamicGeometry
 	| EVertexFactoryFlags::SupportsPrecisePrevWorldPos
 );
