Merge UE5/RET [at] 16305968 to UE5/Main
This represents UE4/Main [at] 16261013 and Dev-PerfTest [at] 16259937

diff --git a/Engine/Plugins/FX/Niagara/Source/Niagara/Private/NiagaraEmitter.cpp b/Engine/Plugins/FX/Niagara/Source/Niagara/Private/NiagaraEmitter.cpp
--- a/Engine/Plugins/FX/Niagara/Source/Niagara/Private/NiagaraEmitter.cpp
+++ b/Engine/Plugins/FX/Niagara/Source/Niagara/Private/NiagaraEmitter.cpp
@@ -442,6 +442,14 @@ void UNiagaraEmitter::PostLoad()
 	}
 #endif
 
+#if WITH_EDITORONLY_DATA
+	if (EditorParameters == nullptr)
+	{
+		INiagaraModule& NiagaraModule = FModuleManager::GetModuleChecked<INiagaraModule>("Niagara");
+		EditorParameters = NiagaraModule.GetEditorOnlyDataUtilities().CreateDefaultEditorParameters(this);
+	}
+#endif
+
 #if WITH_EDITORONLY_DATA
 	if (!GPUComputeScript)
 	{
@@ -639,6 +647,7 @@ void UNiagaraEmitter::PostLoad()
 			EditorData->OnPersistentDataChanged().AddUObject(this, &UNiagaraEmitter::PersistentEditorDataChanged);
 		}
 	}
+	
 
 	// this can only ever be true for old assets that haven't been loaded yet, so this won't overwrite subsequent changes to the template specification
 	if(bIsTemplateAsset_DEPRECATED)
diff --git a/Engine/Plugins/FX/Niagara/Source/Niagara/Private/NiagaraSystem.cpp b/Engine/Plugins/FX/Niagara/Source/Niagara/Private/NiagaraSystem.cpp
--- a/Engine/Plugins/FX/Niagara/Source/Niagara/Private/NiagaraSystem.cpp
+++ b/Engine/Plugins/FX/Niagara/Source/Niagara/Private/NiagaraSystem.cpp
@@ -85,6 +85,15 @@ static FAutoConsoleVariableRef CVarNiagaraCompileWaitLoggingThreshold(
 	ECVF_Default
 );
 
+static int GNiagaraCompileWaitLoggingTerminationCap = 3;
+static FAutoConsoleVariableRef CVarNiagaraCompileWaitLoggingTerminationCap(
+	TEXT("fx.Niagara.CompileWaitLoggingCap"),
+	GNiagaraCompileWaitLoggingTerminationCap,
+	TEXT("During automation, how many times do we log before failing compilation?"),
+	ECVF_Default
+);
+
+
 static int GNiagaraCompileWaitLoggingTerminationCap = 3;
 static FAutoConsoleVariableRef CVarNiagaraCompileWaitLoggingTerminationCap(
 	TEXT("fx.Niagara.CompileWaitLoggingCap"),
diff --git a/Engine/Source/Runtime/Engine/Private/SoundConcurrency.cpp b/Engine/Source/Runtime/Engine/Private/SoundConcurrency.cpp
--- a/Engine/Source/Runtime/Engine/Private/SoundConcurrency.cpp
+++ b/Engine/Source/Runtime/Engine/Private/SoundConcurrency.cpp
@@ -402,7 +402,7 @@ bool FConcurrencyGroup::CanPlaySoundNow(float InCurrentTime) const
 		const float DeltaTime = InCurrentTime - LastTimePlayed;
 		if (DeltaTime < Settings.RetriggerTime)
 		{
-			UE_LOG(LogAudio, VeryVerbose, TEXT("Rejected Sound for Group ID (%d) with DeltaTime: %.2f"), GroupID, DeltaTime);
+			UE_LOG(LogAudioConcurrency, VeryVerbose, TEXT("Rejected Sound for Group ID (%d) with DeltaTime: %.2f"), GroupID, DeltaTime);
 			return false;
 		}
 	}
