fixed discrepancies between UE5/Main and Release-Engine-Test

diff --git a/Engine/Plugins/FX/Niagara/Source/Niagara/Classes/NiagaraEmitter.h b/Engine/Plugins/FX/Niagara/Source/Niagara/Classes/NiagaraEmitter.h
--- a/Engine/Plugins/FX/Niagara/Source/Niagara/Classes/NiagaraEmitter.h
+++ b/Engine/Plugins/FX/Niagara/Source/Niagara/Classes/NiagaraEmitter.h
@@ -393,6 +393,11 @@ public:
 	void CacheFromCompiledData(const FNiagaraDataSetCompiledData* CompiledData);
 	void CacheFromShaderCompiled();
 
+	// Reference to the async update task that is responsible to do any outstanding work like emitter merging on the game thread
+	FGraphEventRef UpdateTaskRef;
+
+	void UpdateEmitterAfterLoad();
+
 #if WITH_EDITORONLY_DATA
 	/** 'Source' data/graphs for the scripts used by this emitter. */
 	UPROPERTY()
@@ -620,6 +625,8 @@ private:
 	void RaiseOnEmitterGPUCompiled(UNiagaraScript* InScript, const FGuid& ScriptVersion);
 #endif
 
+	bool bFullyLoaded = false;
+
 #if !UE_BUILD_SHIPPING
 	FString DebugSimName;
 #endif
diff --git a/Engine/Plugins/FX/Niagara/Source/Niagara/Classes/NiagaraSystem.h b/Engine/Plugins/FX/Niagara/Source/Niagara/Classes/NiagaraSystem.h
--- a/Engine/Plugins/FX/Niagara/Source/Niagara/Classes/NiagaraSystem.h
+++ b/Engine/Plugins/FX/Niagara/Source/Niagara/Classes/NiagaraSystem.h
@@ -262,6 +262,7 @@ public:
 
 private:
 	bool IsValidInternal() const;
+	void EnsureFullyLoaded() const;
 public:
 	/** Returns true if this system is valid and can be instanced. False otherwise. */
 	bool IsValid() const { return FPlatformProperties::RequiresCookedData() ? bIsValidCached : IsValidInternal(); }
@@ -271,6 +272,9 @@ public:
 		original asset. */
 	FNiagaraEmitterHandle AddEmitterHandle(UNiagaraEmitter& SourceEmitter, FName EmitterName);
 
+	/** Adds a new emitter handle to this system without copying the original asset. This should only be used for temporary systems and never for live assets. */
+	void AddEmitterHandleDirect(FNiagaraEmitterHandle& EmitterHandleToAdd);
+
 	/** Duplicates an existing emitter handle and adds it to the System.  The new handle will reference the same source asset,
 		but will have a copy of the duplicated Instance value. */
 	FNiagaraEmitterHandle DuplicateEmitterHandle(const FNiagaraEmitterHandle& EmitterHandleToDuplicate, FName EmitterName);
@@ -410,6 +414,8 @@ public:
 	FNiagaraSystemUpdateContext UpdateContext;
 #endif
 
+	void UpdateSystemAfterLoad();
+
 	bool ShouldAutoDeactivate() const { return bAutoDeactivate; }
 	bool IsLooping() const;
 
@@ -501,6 +507,7 @@ public:
 	UPROPERTY(EditAnywhere, Category = "Debug")
 	bool bDumpDebugEmitterInfo;
 
+	bool bFullyLoaded = false;
 
 	/** When enabled, we follow the settings on the UNiagaraComponent for tick order. When this option is disabled, we ignore any dependencies from data interfaces or other variables and instead fire off the simulation as early in the frame as possible. This greatly
 	reduces overhead and allows the game thread to run faster, but comes at a tradeoff if the dependencies might leave gaps or other visual artifacts.*/
@@ -688,6 +695,9 @@ protected:
 	UNiagaraBakerSettings* BakerGeneratedSettings;
 #endif
 
+	// Reference to the async system update task that is responsible to do any outstanding work like emitter merging or script compilation 
+	FGraphEventRef UpdateTaskRef;
+	
 	UPROPERTY()
 	bool bHasSystemScriptDIsWithPerInstanceData;
 
diff --git a/Engine/Plugins/FX/Niagara/Source/Niagara/Private/NiagaraEmitter.cpp b/Engine/Plugins/FX/Niagara/Source/Niagara/Private/NiagaraEmitter.cpp
--- a/Engine/Plugins/FX/Niagara/Source/Niagara/Private/NiagaraEmitter.cpp
+++ b/Engine/Plugins/FX/Niagara/Source/Niagara/Private/NiagaraEmitter.cpp
@@ -66,6 +66,33 @@
 	ECVF_Default
 );
 
+/**
+Game thread task to do any outstanding work on the loaded emitter
+*/
+struct FNiagaraEmitterUpdateTask
+{
+	FNiagaraEmitterUpdateTask(TWeakObjectPtr<UNiagaraEmitter> InEmitter) : WeakEmitter(InEmitter)
+	{
+	}
+
+	FORCEINLINE TStatId GetStatId() const { RETURN_QUICK_DECLARE_CYCLE_STAT(FNiagaraEmitterUpdateTask, STATGROUP_TaskGraphTasks); }
+	static FORCEINLINE ENamedThreads::Type GetDesiredThread() { return ENamedThreads::GameThread; }
+	static FORCEINLINE ESubsequentsMode::Type GetSubsequentsMode() { return ESubsequentsMode::TrackSubsequents; }
+
+	void DoTask(ENamedThreads::Type CurrentThread, const FGraphEventRef& MyCompletionGraphEvent)
+	{
+		UNiagaraEmitter* Emitter = WeakEmitter.Get();
+		if (Emitter == nullptr)
+		{
+			// probably got garbage collected in the meantime
+			return;
+		}
+		Emitter->UpdateEmitterAfterLoad();
+	}
+
+	TWeakObjectPtr<UNiagaraEmitter> WeakEmitter;
+};
+
 FNiagaraDetailsLevelScaleOverrides::FNiagaraDetailsLevelScaleOverrides()
 {
 	Low = 0.125f;
@@ -357,34 +384,6 @@
 		SetFlags(RF_Transactional);
 	}
 
-	for (int32 RendererIndex = RendererProperties.Num() - 1; RendererIndex >= 0; --RendererIndex)
-	{
-		if (RendererProperties[RendererIndex] == nullptr)
-		{
-#if WITH_EDITORONLY_DATA
-			//In cooked builds these can be cooked out and null on purpose.
-			ensureMsgf(IsCooked, TEXT("Null renderer found in %s at index %i, removing it to prevent crashes."), *GetPathName(), RendererIndex);
-#endif
-			RendererProperties.RemoveAt(RendererIndex);
-		}
-		else
-		{
-			RendererProperties[RendererIndex]->ConditionalPostLoad();
-		}
-	}
-
-	for (int32 SimulationStageIndex = SimulationStages.Num() - 1; SimulationStageIndex >= 0; --SimulationStageIndex)
-	{
-		if (ensureMsgf(SimulationStages[SimulationStageIndex] != nullptr && SimulationStages[SimulationStageIndex]->Script != nullptr, TEXT("Null simulation stage, or simulation stage with a null script found in %s at index %i, removing it to prevent crashes."), *GetPathName(), SimulationStageIndex) == false)
-		{
-			SimulationStages.RemoveAt(SimulationStageIndex);
-		}
-		else
-		{
-			SimulationStages[SimulationStageIndex]->ConditionalPostLoad();
-		}
-	}
-
 	const int32 NiagaraVer = GetLinkerCustomVersion(FNiagaraCustomVersion::GUID);
 	if (NiagaraVer < FNiagaraCustomVersion::PlatformScalingRefactor)
 	{
@@ -442,14 +441,39 @@
 	}
 #endif
 
+	for (int32 RendererIndex = RendererProperties.Num() - 1; RendererIndex >= 0; --RendererIndex)
+	{
+		if (RendererProperties[RendererIndex] == nullptr)
+		{
 #if WITH_EDITORONLY_DATA
-	if (EditorParameters == nullptr)
+			//In cooked builds these can be cooked out and null on purpose.
+			ensureMsgf(IsCooked, TEXT("Null renderer found in %s at index %i, removing it to prevent crashes."), *GetPathName(), RendererIndex);
+#endif
+			RendererProperties.RemoveAt(RendererIndex);
+		}
+		else
+		{
+			RendererProperties[RendererIndex]->ConditionalPostLoad();
+		}
+	}
+
+	for (int32 SimulationStageIndex = SimulationStages.Num() - 1; SimulationStageIndex >= 0; --SimulationStageIndex)
 	{
-		INiagaraModule& NiagaraModule = FModuleManager::GetModuleChecked<INiagaraModule>("Niagara");
-		EditorParameters = NiagaraModule.GetEditorOnlyDataUtilities().CreateDefaultEditorParameters(this);
+		if (ensureMsgf(SimulationStages[SimulationStageIndex] != nullptr && SimulationStages[SimulationStageIndex]->Script != nullptr, TEXT("Null simulation stage, or simulation stage with a null script found in %s at index %i, removing it to prevent crashes."), *GetPathName(), SimulationStageIndex) == false)
+		{
+			SimulationStages.RemoveAt(SimulationStageIndex);
+		}
+		else
+		{
+			SimulationStages[SimulationStageIndex]->ConditionalPostLoad();
+		}
 	}
-#endif
 
+	if (SpawnScriptProps.Script)
+	{
+		SpawnScriptProps.Script->ConditionalPostLoad();
+	}
+	
 #if WITH_EDITORONLY_DATA
 	if (!GPUComputeScript)
 	{
@@ -457,13 +481,8 @@
 		GPUComputeScript->SetUsage(ENiagaraScriptUsage::ParticleGPUComputeScript);
 		GPUComputeScript->SetLatestSource(SpawnScriptProps.Script ? SpawnScriptProps.Script->GetLatestSource() : nullptr);
 	}
-
 	GPUComputeScript->OnGPUScriptCompiled().AddUObject(this, &UNiagaraEmitter::RaiseOnEmitterGPUCompiled);
-#else
-	check(GPUComputeScript == nullptr || SimTarget == ENiagaraSimTarget::GPUComputeSim);
-#endif
 
-#if WITH_EDITORONLY_DATA
 	if (EmitterSpawnScriptProps.Script == nullptr || EmitterUpdateScriptProps.Script == nullptr)
 	{
 		EmitterSpawnScriptProps.Script = NewObject<UNiagaraScript>(this, "EmitterSpawnScript", EObjectFlags::RF_Transactional);
@@ -478,48 +497,13 @@
 			EmitterUpdateScriptProps.Script->SetLatestSource(SpawnScriptProps.Script->GetLatestSource());
 		}
 	}
-#endif
-
-	//Temporarily disabling interpolated spawn if the script type and flag don't match.
-	if (SpawnScriptProps.Script)
-	{
-		SpawnScriptProps.Script->ConditionalPostLoad();
-		bool bActualInterpolatedSpawning = SpawnScriptProps.Script->IsInterpolatedParticleSpawnScript();
-		if (bInterpolatedSpawning != bActualInterpolatedSpawning)
-		{
-			bInterpolatedSpawning = false;
-			if (bActualInterpolatedSpawning)
-			{
-#if WITH_EDITORONLY_DATA
-				SpawnScriptProps.Script->InvalidateCompileResults(TEXT("Interpolated spawn changed."));//clear out the script as it was compiled with interpolated spawn.
-#endif
-				SpawnScriptProps.Script->SetUsage(ENiagaraScriptUsage::ParticleSpawnScript);
-			}
-			UE_LOG(LogNiagara, Warning, TEXT("Disabling interpolated spawn because emitter flag and script type don't match. Did you adjust this value in the UI? Emitter may need recompile.. %s"), *GetFullName());
-		}
-	}
-
-#if WITH_EDITORONLY_DATA
-	if (GetOuter()->IsA<UNiagaraEmitter>())
-	{
-		// If this emitter is owned by another emitter, remove it's inheritance information so that it doesn't try to merge changes.
-		Parent = nullptr;
-		ParentAtLastMerge = nullptr;
-	}
 
 	if (!GetOutermost()->bIsCookedForEditor)
 	{
 		GraphSource->ConditionalPostLoad();
 		GraphSource->PostLoadFromEmitter(*this);
-	}
-#endif
-
-	EnsureScriptsPostLoaded();
-
-#if WITH_EDITORONLY_DATA
-	if (!GetOutermost()->bIsCookedForEditor)
-	{
-		// Handle emitter inheritance.
+		
+		// Prepare for emitter inheritance.
 		if (Parent != nullptr)
 		{
 			Parent->ConditionalPostLoad();
@@ -567,86 +551,31 @@
 		{
 			// Modify here so that the asset will be marked dirty when using the resave commandlet.  This will be ignored during regular post load.
 			Modify();
-			MergeChangesFromParent();
-		}
-
-		// Reset scripts if recompile is forced.
-		bool bGenerateNewChangeId = false;
-		FString GenerateNewChangeIdReason;
-		if (GetForceCompileOnLoad())
-		{
-			// If we are a standalone emitter, then we invalidate id's, which should cause systems dependent on us to regenerate.
-			UObject* OuterObj = GetOuter();
-			if (OuterObj == GetOutermost())
-			{
-				GraphSource->ForceGraphToRecompileOnNextCheck();
-				bGenerateNewChangeId = true;
-				GenerateNewChangeIdReason = TEXT("PostLoad - Force compile on load");
-				if (GEnableVerboseNiagaraChangeIdLogging)
-				{
-					UE_LOG(LogNiagara, Log, TEXT("InvalidateCachedCompileIds for %s because GbForceNiagaraCompileOnLoad = %d"), *GetPathName(), GbForceNiagaraCompileOnLoad);
-				}
-			}
-		}
-	
-		if (ChangeId.IsValid() == false)
-		{
-			// If the change id is already invalid we need to generate a new one, and can skip checking the owned scripts.
-			bGenerateNewChangeId = true;
-			GenerateNewChangeIdReason = TEXT("PostLoad - Change id was invalid.");
-			if (GEnableVerboseNiagaraChangeIdLogging)
-			{
-				UE_LOG(LogNiagara, Log, TEXT("Change ID updated for emitter %s because the ID was invalid."), *GetPathName());
-			}
-		}
-
-		if (bGenerateNewChangeId)
-		{
-			UpdateChangeId(GenerateNewChangeIdReason);
-		}
-
-		GraphSource->OnChanged().AddUObject(this, &UNiagaraEmitter::GraphSourceChanged);
-
-		EmitterSpawnScriptProps.Script->RapidIterationParameters.AddOnChangedHandler(
-			FNiagaraParameterStore::FOnChanged::FDelegate::CreateUObject(this, &UNiagaraEmitter::ScriptRapidIterationParameterChanged));
-		EmitterUpdateScriptProps.Script->RapidIterationParameters.AddOnChangedHandler(
-			FNiagaraParameterStore::FOnChanged::FDelegate::CreateUObject(this, &UNiagaraEmitter::ScriptRapidIterationParameterChanged));
-
-		if (SpawnScriptProps.Script)
-		{
-			SpawnScriptProps.Script->RapidIterationParameters.AddOnChangedHandler(
-				FNiagaraParameterStore::FOnChanged::FDelegate::CreateUObject(this, &UNiagaraEmitter::ScriptRapidIterationParameterChanged));
-		}
-	
-		if (UpdateScriptProps.Script)
-		{
-			UpdateScriptProps.Script->RapidIterationParameters.AddOnChangedHandler(
-				FNiagaraParameterStore::FOnChanged::FDelegate::CreateUObject(this, &UNiagaraEmitter::ScriptRapidIterationParameterChanged));
-		}
-
-		for (FNiagaraEventScriptProperties& EventScriptProperties : EventHandlerScriptProps)
-		{
-			EventScriptProperties.Script->RapidIterationParameters.AddOnChangedHandler(
-				FNiagaraParameterStore::FOnChanged::FDelegate::CreateUObject(this, &UNiagaraEmitter::ScriptRapidIterationParameterChanged));
-		}
-
-		for (UNiagaraSimulationStageBase* SimulationStage : SimulationStages)
-		{
-			SimulationStage->OnChanged().AddUObject(this, &UNiagaraEmitter::SimulationStageChanged);
-			SimulationStage->Script->RapidIterationParameters.AddOnChangedHandler(
-				FNiagaraParameterStore::FOnChanged::FDelegate::CreateUObject(this, &UNiagaraEmitter::ScriptRapidIterationParameterChanged));
-		}
-
-		for (UNiagaraRendererProperties* Renderer : RendererProperties)
-		{
-			Renderer->OnChanged().AddUObject(this, &UNiagaraEmitter::RendererChanged);
-		}
-
-		if (EditorData != nullptr)
-		{
-			EditorData->OnPersistentDataChanged().AddUObject(this, &UNiagaraEmitter::PersistentEditorDataChanged);
 		}
 	}
+#else
+	check(GPUComputeScript == nullptr || SimTarget == ENiagaraSimTarget::GPUComputeSim);
+#endif
+
+	//Temporarily disabling interpolated spawn if the script type and flag don't match.
+	if (SpawnScriptProps.Script)
+	{
+		bool bActualInterpolatedSpawning = SpawnScriptProps.Script->IsInterpolatedParticleSpawnScript();
+		if (bInterpolatedSpawning != bActualInterpolatedSpawning)
+		{
+			bInterpolatedSpawning = false;
+			if (bActualInterpolatedSpawning)
+			{
+#if WITH_EDITORONLY_DATA
+				//clear out the script as it was compiled with interpolated spawn.
+				SpawnScriptProps.Script->InvalidateCompileResults(TEXT("Interpolated spawn changed."));
+#endif
+				SpawnScriptProps.Script->SetUsage(ENiagaraScriptUsage::ParticleSpawnScript);
+			}
+			UE_LOG(LogNiagara, Warning, TEXT("Disabling interpolated spawn because emitter flag and script type don't match. Did you adjust this value in the UI? Emitter may need recompile.. %s"), *GetFullName());
+		}
+	}
+	EnsureScriptsPostLoaded();
 	
 #if WITH_EDITORONLY_DATA
 	// this can only ever be true for old assets that haven't been loaded yet, so this won't overwrite subsequent changes to the template specification
@@ -656,19 +585,8 @@
 	}
 #endif
 	
-#endif
-
-	ResolveScalabilitySettings();
-
-#if !UE_BUILD_SHIPPING
-	DebugSimName.Empty();
-	if (const UNiagaraSystem* SystemOwner = Cast<const UNiagaraSystem>(GetOuter()))
-	{
-		DebugSimName = SystemOwner->GetName();
-		DebugSimName.AppendChar(':');
-	}
-	DebugSimName.Append(GetName());
-#endif
+	// we are not yet finished, but we do the rest of the work after postload in a separate task
+	UpdateTaskRef = TGraphTask<FNiagaraEmitterUpdateTask>::CreateTask().ConstructAndDispatchWhenReady(this);
 }
 
 bool UNiagaraEmitter::IsEditorOnly() const
@@ -1326,6 +1244,125 @@
 	}
 }
 
+void UNiagaraEmitter::UpdateEmitterAfterLoad()
+{
+	if (bFullyLoaded)
+	{
+		return;
+	}
+	bFullyLoaded = true;
+	UpdateTaskRef = nullptr;
+	
+#if WITH_EDITORONLY_DATA
+	check(IsInGameThread());
+	if (GetOuter()->IsA<UNiagaraEmitter>())
+	{
+		// If this emitter is owned by another emitter, remove it's inheritance information so that it doesn't try to merge changes.
+		Parent = nullptr;
+		ParentAtLastMerge = nullptr;
+	}
+	
+	if (!GetOutermost()->bIsCookedForEditor)
+	{
+		if (IsSynchronizedWithParent() == false)
+		{
+			// Modify here so that the asset will be marked dirty when using the resave commandlet.  This will be ignored during regular post load.
+			MergeChangesFromParent();
+		}
+
+		// Reset scripts if recompile is forced.
+		bool bGenerateNewChangeId = false;
+		FString GenerateNewChangeIdReason;
+		if (GetForceCompileOnLoad())
+		{
+			// If we are a standalone emitter, then we invalidate id's, which should cause systems dependent on us to regenerate.
+			UObject* OuterObj = GetOuter();
+			if (OuterObj == GetOutermost())
+			{
+				GraphSource->ForceGraphToRecompileOnNextCheck();
+				bGenerateNewChangeId = true;
+				GenerateNewChangeIdReason = TEXT("PostLoad - Force compile on load");
+				if (GEnableVerboseNiagaraChangeIdLogging)
+				{
+					UE_LOG(LogNiagara, Log, TEXT("InvalidateCachedCompileIds for %s because GbForceNiagaraCompileOnLoad = %d"), *GetPathName(), GbForceNiagaraCompileOnLoad);
+				}
+			}
+		}
+	
+		if (ChangeId.IsValid() == false)
+		{
+			// If the change id is already invalid we need to generate a new one, and can skip checking the owned scripts.
+			bGenerateNewChangeId = true;
+			GenerateNewChangeIdReason = TEXT("PostLoad - Change id was invalid.");
+			if (GEnableVerboseNiagaraChangeIdLogging)
+			{
+				UE_LOG(LogNiagara, Log, TEXT("Change ID updated for emitter %s because the ID was invalid."), *GetPathName());
+			}
+		}
+
+		if (bGenerateNewChangeId)
+		{
+			UpdateChangeId(GenerateNewChangeIdReason);
+		}
+
+		GraphSource->OnChanged().AddUObject(this, &UNiagaraEmitter::GraphSourceChanged);
+
+		EmitterSpawnScriptProps.Script->RapidIterationParameters.AddOnChangedHandler(
+			FNiagaraParameterStore::FOnChanged::FDelegate::CreateUObject(this, &UNiagaraEmitter::ScriptRapidIterationParameterChanged));
+		EmitterUpdateScriptProps.Script->RapidIterationParameters.AddOnChangedHandler(
+			FNiagaraParameterStore::FOnChanged::FDelegate::CreateUObject(this, &UNiagaraEmitter::ScriptRapidIterationParameterChanged));
+
+		if (SpawnScriptProps.Script)
+		{
+			SpawnScriptProps.Script->RapidIterationParameters.AddOnChangedHandler(
+				FNiagaraParameterStore::FOnChanged::FDelegate::CreateUObject(this, &UNiagaraEmitter::ScriptRapidIterationParameterChanged));
+		}
+	
+		if (UpdateScriptProps.Script)
+		{
+			UpdateScriptProps.Script->RapidIterationParameters.AddOnChangedHandler(
+				FNiagaraParameterStore::FOnChanged::FDelegate::CreateUObject(this, &UNiagaraEmitter::ScriptRapidIterationParameterChanged));
+		}
+
+		for (FNiagaraEventScriptProperties& EventScriptProperties : EventHandlerScriptProps)
+		{
+			EventScriptProperties.Script->RapidIterationParameters.AddOnChangedHandler(
+				FNiagaraParameterStore::FOnChanged::FDelegate::CreateUObject(this, &UNiagaraEmitter::ScriptRapidIterationParameterChanged));
+		}
+
+		for (UNiagaraSimulationStageBase* SimulationStage : SimulationStages)
+		{
+			SimulationStage->OnChanged().AddUObject(this, &UNiagaraEmitter::SimulationStageChanged);
+			SimulationStage->Script->RapidIterationParameters.AddOnChangedHandler(
+				FNiagaraParameterStore::FOnChanged::FDelegate::CreateUObject(this, &UNiagaraEmitter::ScriptRapidIterationParameterChanged));
+		}
+
+		for (UNiagaraRendererProperties* Renderer : RendererProperties)
+		{
+			Renderer->OnChanged().AddUObject(this, &UNiagaraEmitter::RendererChanged);
+		}
+
+		if (EditorData != nullptr)
+		{
+			EditorData->OnPersistentDataChanged().AddUObject(this, &UNiagaraEmitter::PersistentEditorDataChanged);
+		}
+	}
+#endif
+
+	ResolveScalabilitySettings();
+
+#if !UE_BUILD_SHIPPING
+	DebugSimName.Empty();
+	if (const UNiagaraSystem* SystemOwner = Cast<const UNiagaraSystem>(GetOuter()))
+	{
+		DebugSimName = SystemOwner->GetName();
+		DebugSimName.AppendChar(':');
+	}
+	DebugSimName.Append(GetName());
+#endif
+
+}
+
 bool UNiagaraEmitter::IsAllowedByScalability()const
 {
 	return Platforms.IsActive();
diff --git a/Engine/Plugins/FX/Niagara/Source/Niagara/Private/NiagaraSystem.cpp b/Engine/Plugins/FX/Niagara/Source/Niagara/Private/NiagaraSystem.cpp
--- a/Engine/Plugins/FX/Niagara/Source/Niagara/Private/NiagaraSystem.cpp
+++ b/Engine/Plugins/FX/Niagara/Source/Niagara/Private/NiagaraSystem.cpp
@@ -12,7 +12,6 @@
 #include "NiagaraEmitterInstanceBatcher.h"
 #include "NiagaraModule.h"
 #include "NiagaraPrecompileContainer.h"
-#include "NiagaraRenderer.h"
 #include "NiagaraRendererProperties.h"
 #include "NiagaraScriptSourceBase.h"
 #include "NiagaraSettings.h"
@@ -93,16 +92,45 @@
 	ECVF_Default
 );
 
-
-static int GNiagaraCompileWaitLoggingTerminationCap = 3;
-static FAutoConsoleVariableRef CVarNiagaraCompileWaitLoggingTerminationCap(
-	TEXT("fx.Niagara.CompileWaitLoggingCap"),
-	GNiagaraCompileWaitLoggingTerminationCap,
-	TEXT("During automation, how many times do we log before failing compilation?"),
-	ECVF_Default
-);
-
 //////////////////////////////////////////////////////////////////////////
 
+/**
+Game thread task to do any outstanding work on the loaded system
+*/
+struct FNiagaraSystemUpdateTask
+{
+	FNiagaraSystemUpdateTask(TWeakObjectPtr<UNiagaraSystem> InSystem) : WeakSystem(InSystem)
+	{
+	}
+
+	FORCEINLINE TStatId GetStatId() const { RETURN_QUICK_DECLARE_CYCLE_STAT(FNiagaraSystemUpdateTask, STATGROUP_TaskGraphTasks); }
+	static FORCEINLINE ENamedThreads::Type GetDesiredThread() { return ENamedThreads::GameThread; }
+	static FORCEINLINE ESubsequentsMode::Type GetSubsequentsMode() { return ESubsequentsMode::TrackSubsequents; }
+
+	void DoTask(ENamedThreads::Type CurrentThread, const FGraphEventRef& MyCompletionGraphEvent)
+	{
+		check(IsInGameThread());
+		UNiagaraSystem* System = WeakSystem.Get();
+		if (System == nullptr)
+		{
+			// probably got garbage collected in the meantime
+			return;
+		}
+#if WITH_EDITORONLY_DATA
+		for (const FNiagaraEmitterHandle& Handle : System->GetEmitterHandles())
+		{
+			if (Handle.GetInstance() && Handle.GetInstance()->UpdateTaskRef.IsValid())
+			{
+				// the emitter tasks should all be done by now, if not then something is very wrong
+				ensureMsgf(Handle.GetInstance()->UpdateTaskRef->IsComplete(), TEXT("Unfinished emitter update task for system %s"), *GetPathNameSafe(System));
+			}
+		}
+#endif
+		System->UpdateSystemAfterLoad();
+	}
+
+	TWeakObjectPtr<UNiagaraSystem> WeakSystem;
+};
+
 UNiagaraSystem::UNiagaraSystem(const FObjectInitializer& ObjectInitializer)
 : Super(ObjectInitializer)
@@ -176,6 +204,7 @@
 #if WITH_EDITORONLY_DATA
 	WaitForCompilationComplete();
 #endif
+	EnsureFullyLoaded();
 }
 
 #if WITH_EDITOR
@@ -318,9 +347,192 @@
 	return false;
 }
 
+void UNiagaraSystem::UpdateSystemAfterLoad()
+{
+	// guard against deadlocks by having wait called on it during the update
+	if (bFullyLoaded)
+	{
+		return;
+	}
+	bFullyLoaded = true;
+	UpdateTaskRef = nullptr;
+	
 #if WITH_EDITORONLY_DATA
-bool UNiagaraSystem::UsesScript(const UNiagaraScript* Script)const
+	// We remove emitters and scripts on dedicated servers, so skip further work.
+	const bool bIsDedicatedServer = !GIsClient && GIsServer;
+
+	if (!GetOutermost()->bIsCookedForEditor && !bIsDedicatedServer)
+	{		
+		TArray<UNiagaraScript*> AllSystemScripts;
+		UNiagaraScriptSourceBase* SystemScriptSource;
+		if (SystemSpawnScript == nullptr)
+		{
+			SystemSpawnScript = NewObject<UNiagaraScript>(this, "SystemSpawnScript", RF_Transactional);
+			SystemSpawnScript->SetUsage(ENiagaraScriptUsage::SystemSpawnScript);
+			INiagaraModule& NiagaraModule = FModuleManager::GetModuleChecked<INiagaraModule>("Niagara");
+			SystemScriptSource = NiagaraModule.GetEditorOnlyDataUtilities().CreateDefaultScriptSource(this);
+			SystemSpawnScript->SetLatestSource(SystemScriptSource);
+		}
+		else
+		{
+			SystemSpawnScript->ConditionalPostLoad();
+			SystemScriptSource = SystemSpawnScript->GetLatestSource();
+		}
+		AllSystemScripts.Add(SystemSpawnScript);
+
+		if (SystemUpdateScript == nullptr)
+		{
+			SystemUpdateScript = NewObject<UNiagaraScript>(this, "SystemUpdateScript", RF_Transactional);
+			SystemUpdateScript->SetUsage(ENiagaraScriptUsage::SystemUpdateScript);
+			SystemUpdateScript->SetLatestSource(SystemScriptSource);
+		}
+		else
+		{
+			SystemUpdateScript->ConditionalPostLoad();
+		}
+		AllSystemScripts.Add(SystemUpdateScript);
+
+		bool bSystemScriptsAreSynchronized = true;
+		for (UNiagaraScript* SystemScript : AllSystemScripts)
+		{
+			bSystemScriptsAreSynchronized &= SystemScript->AreScriptAndSourceSynchronized();
+		}
+
+		bool bEmitterScriptsAreSynchronized = true;
+
+#if 0
+		UE_LOG(LogNiagara, Log, TEXT("PreMerger"));
+		for (FNiagaraEmitterHandle& EmitterHandle : EmitterHandles)
+		{
+			UE_LOG(LogNiagara, Log, TEXT("Emitter Handle: %s"), *EmitterHandle.GetUniqueInstanceName());
+			UNiagaraScript* UpdateScript = EmitterHandle.GetInstance()->GetScript(ENiagaraScriptUsage::ParticleUpdateScript, FGuid());
+			UNiagaraScript* SpawnScript = EmitterHandle.GetInstance()->GetScript(ENiagaraScriptUsage::ParticleSpawnScript, FGuid());
+			UE_LOG(LogNiagara, Log, TEXT("Spawn Parameters"));
+			SpawnScript->GetVMExecutableData().Parameters.DumpParameters();
+			UE_LOG(LogNiagara, Log, TEXT("Spawn RI Parameters"));
+			SpawnScript->RapidIterationParameters.DumpParameters();
+			UE_LOG(LogNiagara, Log, TEXT("Update Parameters"));
+			UpdateScript->GetVMExecutableData().Parameters.DumpParameters();
+			UE_LOG(LogNiagara, Log, TEXT("Update RI Parameters"));
+			UpdateScript->RapidIterationParameters.DumpParameters();
+		}
+#endif
+
+		for (FNiagaraEmitterHandle& EmitterHandle : EmitterHandles)
+		{
+			if (EmitterHandle.GetIsEnabled() && EmitterHandle.GetInstance() && !EmitterHandle.GetInstance()->AreAllScriptAndSourcesSynchronized())
+			{
+				bEmitterScriptsAreSynchronized = false;
+			}
+		}
+
+		if (UNiagaraEmitter::GetForceCompileOnLoad())
+		{
+			ForceGraphToRecompileOnNextCheck();
+			UE_LOG(LogNiagara, Log, TEXT("System %s being rebuilt because UNiagaraEmitter::GetForceCompileOnLoad() == true."), *GetPathName());
+		}
+
+		if (bSystemScriptsAreSynchronized == false && GEnableVerboseNiagaraChangeIdLogging)
+		{
+			UE_LOG(LogNiagara, Log, TEXT("System %s being compiled because there were changes to a system script Change ID."), *GetPathName());
+		}
+
+		if (bEmitterScriptsAreSynchronized == false && GEnableVerboseNiagaraChangeIdLogging)
+		{
+			UE_LOG(LogNiagara, Log, TEXT("System %s being compiled because there were changes to an emitter script Change ID."), *GetPathName());
+		}
+
+		if (EmitterCompiledData.Num() == 0 || EmitterCompiledData[0]->DataSetCompiledData.Variables.Num() == 0)
+		{
+			InitEmitterCompiledData();
+		}
+
+		if (SystemCompiledData.InstanceParamStore.ReadParameterVariables().Num() == 0 ||SystemCompiledData.DataSetCompiledData.Variables.Num() == 0)
+		{
+			InitSystemCompiledData();
+		}
+
+#if 0
+		UE_LOG(LogNiagara, Log, TEXT("Before"));
+		for (FNiagaraEmitterHandle& EmitterHandle : EmitterHandles)
+		{
+			UE_LOG(LogNiagara, Log, TEXT("Emitter Handle: %s"), *EmitterHandle.GetUniqueInstanceName());
+			UNiagaraScript* UpdateScript = EmitterHandle.GetInstance()->GetScript(ENiagaraScriptUsage::ParticleUpdateScript, FGuid());
+			UNiagaraScript* SpawnScript = EmitterHandle.GetInstance()->GetScript(ENiagaraScriptUsage::ParticleSpawnScript, FGuid());
+			UE_LOG(LogNiagara, Log, TEXT("Spawn Parameters"));
+			SpawnScript->GetVMExecutableData().Parameters.DumpParameters();
+			UE_LOG(LogNiagara, Log, TEXT("Spawn RI Parameters"));
+			SpawnScript->RapidIterationParameters.DumpParameters();
+			UE_LOG(LogNiagara, Log, TEXT("Update Parameters"));
+			UpdateScript->GetVMExecutableData().Parameters.DumpParameters();
+			UE_LOG(LogNiagara, Log, TEXT("Update RI Parameters"));
+			UpdateScript->RapidIterationParameters.DumpParameters();
+		}
+#endif
+
+		if (bSystemScriptsAreSynchronized == false || bEmitterScriptsAreSynchronized == false)
+		{
+			// Call modify here so that the system will resave the compile ids and script vm when running the resave
+			// commandlet.  In normal post load, it will be ignored.
+			Modify();
+			RequestCompile(false);
+		}
+
+#if 0
+		UE_LOG(LogNiagara, Log, TEXT("After"));
+		for (FNiagaraEmitterHandle& EmitterHandle : EmitterHandles)
+		{
+			UE_LOG(LogNiagara, Log, TEXT("Emitter Handle: %s"), *EmitterHandle.GetUniqueInstanceName());
+			UNiagaraScript* UpdateScript = EmitterHandle.GetInstance()->GetScript(ENiagaraScriptUsage::ParticleUpdateScript, FGuid());
+			UNiagaraScript* SpawnScript = EmitterHandle.GetInstance()->GetScript(ENiagaraScriptUsage::ParticleSpawnScript, FGuid());
+			UE_LOG(LogNiagara, Log, TEXT("Spawn Parameters"));
+			SpawnScript->GetVMExecutableData().Parameters.DumpParameters();
+			UE_LOG(LogNiagara, Log, TEXT("Spawn RI Parameters"));
+			SpawnScript->RapidIterationParameters.DumpParameters();
+			UE_LOG(LogNiagara, Log, TEXT("Update Parameters"));
+			UpdateScript->GetVMExecutableData().Parameters.DumpParameters();
+			UE_LOG(LogNiagara, Log, TEXT("Update RI Parameters"));
+			UpdateScript->RapidIterationParameters.DumpParameters();
+		}
+#endif
+	}
+	if (GNiagaraForceSystemsToCookOutRapidIterationOnLoad == 1 && !bBakeOutRapidIteration)
+	{
+		WaitForCompilationComplete();
+		bBakeOutRapidIteration = true;
+		RequestCompile(false);
+	}
+#endif
+
+	if ( FPlatformProperties::RequiresCookedData() )
+	{
+		bIsValidCached = IsValidInternal();
+		bIsReadyToRunCached = IsReadyToRunInternal();
+	}
+
+	ResolveScalabilitySettings();
+
+	ComputeEmittersExecutionOrder();
+
+	ComputeRenderersDrawOrder();
+
+	CacheFromCompiledData();
+
+	//TODO: Move to serialized properties?
+	UpdateDITickFlags();
+	UpdateHasGPUEmitters();
+
+	// Run task to prime pools this must happen on the GameThread
+	if (PoolPrimeSize > 0 && MaxPoolSize > 0)
+	{
+		FNiagaraWorldManager::PrimePoolForAllWorlds(this);
+	}
+}
+
+#if WITH_EDITORONLY_DATA
+bool UNiagaraSystem::UsesScript(const UNiagaraScript* Script) const
 {
+	EnsureFullyLoaded();
 	if (SystemSpawnScript == Script ||
 		SystemUpdateScript == Script)
 	{
@@ -567,242 +779,38 @@
 	ExposedParameters.RecreateRedirections();
 #endif
 
+	FGraphEventArray EmitterPrerequisiteTasks;
+	for (FNiagaraEmitterHandle& EmitterHandle : EmitterHandles)
+	{
 #if WITH_EDITORONLY_DATA
-	// We remove emitters and scripts on dedicated servers, so skip further work.
-	const bool bIsDedicatedServer = !GIsClient && GIsServer;
-
-	if (!GetOutermost()->bIsCookedForEditor && !bIsDedicatedServer)
-	{
-		if (EditorParameters == nullptr)
+		EmitterHandle.ConditionalPostLoad(NiagaraVer);
+#else
+		if (UNiagaraEmitter* NiagaraEmitter = EmitterHandle.GetInstance())
 		{
-			INiagaraModule& NiagaraModule = FModuleManager::GetModuleChecked<INiagaraModule>("Niagara");
-			EditorParameters = NiagaraModule.GetEditorOnlyDataUtilities().CreateDefaultEditorParameters(this);
-		}
-
-		TArray<UNiagaraScript*> AllSystemScripts;
-		UNiagaraScriptSourceBase* SystemScriptSource = nullptr;
-		if (SystemSpawnScript == nullptr)
-		{
-			SystemSpawnScript = NewObject<UNiagaraScript>(this, "SystemSpawnScript", RF_Transactional);
-			SystemSpawnScript->SetUsage(ENiagaraScriptUsage::SystemSpawnScript);
-			INiagaraModule& NiagaraModule = FModuleManager::GetModuleChecked<INiagaraModule>("Niagara");
-			SystemScriptSource = NiagaraModule.GetEditorOnlyDataUtilities().CreateDefaultScriptSource(this);
-			SystemSpawnScript->SetLatestSource(SystemScriptSource);
-		}
-		else
-		{
-			SystemSpawnScript->ConditionalPostLoad();
-			SystemScriptSource = SystemSpawnScript->GetLatestSource();
-		}
-		AllSystemScripts.Add(SystemSpawnScript);
-
-		if (SystemUpdateScript == nullptr)
-		{
-			SystemUpdateScript = NewObject<UNiagaraScript>(this, "SystemUpdateScript", RF_Transactional);
-			SystemUpdateScript->SetUsage(ENiagaraScriptUsage::SystemUpdateScript);
-			SystemUpdateScript->SetLatestSource(SystemScriptSource);
-		}
-		else
-		{
-			SystemUpdateScript->ConditionalPostLoad();
-		}
-		AllSystemScripts.Add(SystemUpdateScript);
-
-		//TODO: This causes a crash becuase the script source ptr is null? Fix
-		//For existing emitters before the lifecylce rework, ensure they have the system lifecycle module.
-		if (NiagaraVer < FNiagaraCustomVersion::LifeCycleRework)
-		{
-			/*UNiagaraScriptSourceBase* SystemScriptSource = SystemUpdateScript->GetSource();
-			if (SystemScriptSource)
-			{
-				bool bFoundModule;
-				if (SystemScriptSource->AddModuleIfMissing(TEXT("/Niagara/Modules/System/SystemLifeCycle.SystemLifeCycle"), ENiagaraScriptUsage::SystemUpdateScript, bFoundModule))
-				{
-					bNeedsRecompile = true;
-				}
-			}*/
-		}
-
-		bool bSystemScriptsAreSynchronized = true;
-		for (UNiagaraScript* SystemScript : AllSystemScripts)
-		{
-			bSystemScriptsAreSynchronized &= SystemScript->AreScriptAndSourceSynchronized();
-		}
-
-		bool bEmitterScriptsAreSynchronized = true;
-
-#if 0
-		UE_LOG(LogNiagara, Log, TEXT("PreMerger"));
-		for (FNiagaraEmitterHandle& EmitterHandle : EmitterHandles)
-		{
-			UE_LOG(LogNiagara, Log, TEXT("Emitter Handle: %s"), *EmitterHandle.GetUniqueInstanceName());
-			UNiagaraScript* UpdateScript = EmitterHandle.GetInstance()->GetScript(ENiagaraScriptUsage::ParticleUpdateScript, FGuid());
-			UNiagaraScript* SpawnScript = EmitterHandle.GetInstance()->GetScript(ENiagaraScriptUsage::ParticleSpawnScript, FGuid());
-			UE_LOG(LogNiagara, Log, TEXT("Spawn Parameters"));
-			SpawnScript->GetVMExecutableData().Parameters.DumpParameters();
-			UE_LOG(LogNiagara, Log, TEXT("Spawn RI Parameters"));
-			SpawnScript->RapidIterationParameters.DumpParameters();
-			UE_LOG(LogNiagara, Log, TEXT("Update Parameters"));
-			UpdateScript->GetVMExecutableData().Parameters.DumpParameters();
-			UE_LOG(LogNiagara, Log, TEXT("Update RI Parameters"));
-			UpdateScript->RapidIterationParameters.DumpParameters();
+			NiagaraEmitter->ConditionalPostLoad();
 		}
 #endif
-
-		for (FNiagaraEmitterHandle& EmitterHandle : EmitterHandles)
+		if (UNiagaraEmitter* NiagaraEmitter = EmitterHandle.GetInstance())
 		{
-			EmitterHandle.ConditionalPostLoad(NiagaraVer);
-			if (EmitterHandle.GetIsEnabled() && EmitterHandle.GetInstance() && !EmitterHandle.GetInstance()->AreAllScriptAndSourcesSynchronized())
-			{
-				bEmitterScriptsAreSynchronized = false;
-			}
+			EmitterPrerequisiteTasks.Add(NiagaraEmitter->UpdateTaskRef);
 		}
-
-		if (EditorData == nullptr)
-		{
-			INiagaraModule& NiagaraModule = FModuleManager::GetModuleChecked<INiagaraModule>("Niagara");
-			EditorData = NiagaraModule.GetEditorOnlyDataUtilities().CreateDefaultEditorData(this);
-		}
-		else
-		{
-			EditorData->PostLoadFromOwner(this);
-		}
-
-		if (UNiagaraEmitter::GetForceCompileOnLoad())
-		{
-			ForceGraphToRecompileOnNextCheck();
-			UE_LOG(LogNiagara, Log, TEXT("System %s being rebuilt because UNiagaraEmitter::GetForceCompileOnLoad() == true."), *GetPathName());
-		}
-
-		if (bSystemScriptsAreSynchronized == false && GEnableVerboseNiagaraChangeIdLogging)
-		{
-			UE_LOG(LogNiagara, Log, TEXT("System %s being compiled because there were changes to a system script Change ID."), *GetPathName());
-		}
-
-		if (bEmitterScriptsAreSynchronized == false && GEnableVerboseNiagaraChangeIdLogging)
-		{
-			UE_LOG(LogNiagara, Log, TEXT("System %s being compiled because there were changes to an emitter script Change ID."), *GetPathName());
-		}
-
-		if (EmitterCompiledData.Num() == 0 || EmitterCompiledData[0]->DataSetCompiledData.Variables.Num() == 0)
-		{
-			InitEmitterCompiledData();
-		}
-
-		if (SystemCompiledData.InstanceParamStore.ReadParameterVariables().Num() == 0 ||SystemCompiledData.DataSetCompiledData.Variables.Num() == 0)
-		{
-			InitSystemCompiledData();
-		}
-
-#if 0
-		UE_LOG(LogNiagara, Log, TEXT("Before"));
-		for (FNiagaraEmitterHandle& EmitterHandle : EmitterHandles)
-		{
-			UE_LOG(LogNiagara, Log, TEXT("Emitter Handle: %s"), *EmitterHandle.GetUniqueInstanceName());
-			UNiagaraScript* UpdateScript = EmitterHandle.GetInstance()->GetScript(ENiagaraScriptUsage::ParticleUpdateScript, FGuid());
-			UNiagaraScript* SpawnScript = EmitterHandle.GetInstance()->GetScript(ENiagaraScriptUsage::ParticleSpawnScript, FGuid());
-			UE_LOG(LogNiagara, Log, TEXT("Spawn Parameters"));
-			SpawnScript->GetVMExecutableData().Parameters.DumpParameters();
-			UE_LOG(LogNiagara, Log, TEXT("Spawn RI Parameters"));
-			SpawnScript->RapidIterationParameters.DumpParameters();
-			UE_LOG(LogNiagara, Log, TEXT("Update Parameters"));
-			UpdateScript->GetVMExecutableData().Parameters.DumpParameters();
-			UE_LOG(LogNiagara, Log, TEXT("Update RI Parameters"));
-			UpdateScript->RapidIterationParameters.DumpParameters();
-		}
-#endif
-
-		// Synchronize with definitions before compiling.
-		const UNiagaraSettings* Settings = GetDefault<UNiagaraSettings>();
-		check(Settings);
-		TArray<FGuid> DefaultDefinitionsUniqueIds;
-		for (const FSoftObjectPath& DefaultLinkedParameterDefinitionObjPath : Settings->DefaultLinkedParameterDefinitions)
-		{
-			UNiagaraParameterDefinitionsBase* DefaultLinkedParameterDefinitions = CastChecked<UNiagaraParameterDefinitionsBase>(DefaultLinkedParameterDefinitionObjPath.TryLoad());
-			DefaultDefinitionsUniqueIds.Add(DefaultLinkedParameterDefinitions->GetDefinitionsUniqueId());
-			const bool bDoNotAssertIfAlreadySubscribed = true;
-			SubscribeToParameterDefinitions(DefaultLinkedParameterDefinitions, bDoNotAssertIfAlreadySubscribed);
-		}
-		FSynchronizeWithParameterDefinitionsArgs Args;
-		Args.SpecificDefinitionsUniqueIds = DefaultDefinitionsUniqueIds;
-		Args.bForceSynchronizeDefinitions = true;
-		Args.bSubscribeAllNameMatchParameters = true;
-		SynchronizeWithParameterDefinitions(Args);
-		InitParameterDefinitionsSubscriptions();
-
-		if (bSystemScriptsAreSynchronized == false || bEmitterScriptsAreSynchronized == false)
-		{
-			// Call modify here so that the system will resave the compile ids and script vm when running the resave
-			// commandlet.  In normal post load, it will be ignored.
-			Modify();
-			RequestCompile(false);
-		}
-
-#if 0
-		UE_LOG(LogNiagara, Log, TEXT("After"));
-		for (FNiagaraEmitterHandle& EmitterHandle : EmitterHandles)
-		{
-			UE_LOG(LogNiagara, Log, TEXT("Emitter Handle: %s"), *EmitterHandle.GetUniqueInstanceName());
-			UNiagaraScript* UpdateScript = EmitterHandle.GetInstance()->GetScript(ENiagaraScriptUsage::ParticleUpdateScript, FGuid());
-			UNiagaraScript* SpawnScript = EmitterHandle.GetInstance()->GetScript(ENiagaraScriptUsage::ParticleSpawnScript, FGuid());
-			UE_LOG(LogNiagara, Log, TEXT("Spawn Parameters"));
-			SpawnScript->GetVMExecutableData().Parameters.DumpParameters();
-			UE_LOG(LogNiagara, Log, TEXT("Spawn RI Parameters"));
-			SpawnScript->RapidIterationParameters.DumpParameters();
-			UE_LOG(LogNiagara, Log, TEXT("Update Parameters"));
-			UpdateScript->GetVMExecutableData().Parameters.DumpParameters();
-			UE_LOG(LogNiagara, Log, TEXT("Update RI Parameters"));
-			UpdateScript->RapidIterationParameters.DumpParameters();
-		}
-#endif
 	}
-	if (GNiagaraForceSystemsToCookOutRapidIterationOnLoad == 1 && !bBakeOutRapidIteration)
+	
+#if WITH_EDITORONLY_DATA
+	if (EditorData == nullptr)
 	{
-		WaitForCompilationComplete();
-		bBakeOutRapidIteration = true;
-		RequestCompile(false);
+		INiagaraModule& NiagaraModule = FModuleManager::GetModuleChecked<INiagaraModule>("Niagara");
+		EditorData = NiagaraModule.GetEditorOnlyDataUtilities().CreateDefaultEditorData(this);
 	}
-
-	// see the equivalent in NiagaraEmitter for details
-	if(bIsTemplateAsset_DEPRECATED)
+	else
 	{
-		TemplateSpecification = bIsTemplateAsset_DEPRECATED ? ENiagaraScriptTemplateSpecification::Template : ENiagaraScriptTemplateSpecification::None;
+		EditorData->PostLoadFromOwner(this);
 	}
+	
 #endif // WITH_EDITORONLY_DATA
 
-	if ( FPlatformProperties::RequiresCookedData() )
-	{
-		bIsValidCached = IsValidInternal();
-		bIsReadyToRunCached = IsReadyToRunInternal();
-	}
-
-	ResolveScalabilitySettings();
-
-	ComputeEmittersExecutionOrder();
-
-	ComputeRenderersDrawOrder();
-
-	CacheFromCompiledData();
-
-	//TODO: Move to serialized properties?
-	UpdateDITickFlags();
-	UpdateHasGPUEmitters();
-
-	// Run task to prime pools this must happen on the GameThread
-	if (PoolPrimeSize > 0 && MaxPoolSize > 0)
-	{
-		AsyncTask(
-			ENamedThreads::GameThread,
-			[WeakSystem=TWeakObjectPtr<UNiagaraSystem>(this)]()
-			{
-				if ( UNiagaraSystem* NiagaraSystem = WeakSystem.Get() )
-				{
-					check(IsInGameThread());
-					FNiagaraWorldManager::PrimePoolForAllWorlds(NiagaraSystem);
-				}
-			}
-		);
-	}
+	// we are not yet finished, but we do the rest of the work after postload in a separate task
+	UpdateTaskRef = TGraphTask<FNiagaraSystemUpdateTask>::CreateTask(&EmitterPrerequisiteTasks).ConstructAndDispatchWhenReady(this);
 }
 
 #if WITH_EDITORONLY_DATA
@@ -889,6 +897,7 @@
 		return false;
 	}
 
+	EnsureFullyLoaded();
 	if (!SystemSpawnScript || !SystemUpdateScript)
 	{
 		if (FPlatformProperties::RequiresCookedData())
@@ -958,6 +967,7 @@
 
 void UNiagaraSystem::GetAssetRegistryTags(TArray<FAssetRegistryTag>& OutTags) const
 {
+	EnsureFullyLoaded();
 #if WITH_EDITOR
 	OutTags.Add(FAssetRegistryTag("HasGPUEmitter", HasAnyGPUEmitters() ? TEXT("True") : TEXT("False"), FAssetRegistryTag::TT_Alphabetical));
 
@@ -1706,6 +1716,23 @@
 	return true;
 }
 
+void UNiagaraSystem::EnsureFullyLoaded() const
+{
+	if (UpdateTaskRef.IsValid() && UpdateTaskRef->IsComplete() == false)
+	{
+		ensureMsgf(FUObjectThreadContext::Get().IsRoutingPostLoad == false, TEXT("Niagara system data was forced to load during PostLoad instead of the taskgraph! Asset %s"), *GetPathName());
+		for (const FNiagaraEmitterHandle& EmitterHandle : EmitterHandles)
+		{
+			if (UNiagaraEmitter* NiagaraEmitter = EmitterHandle.GetInstance())
+			{
+				NiagaraEmitter->UpdateEmitterAfterLoad();
+			}
+		}
+		UNiagaraSystem* System = const_cast<UNiagaraSystem*>(this); // Sadly necessary because we are called from other const functions
+		System->UpdateSystemAfterLoad();
+	}	
+}
 
+
 bool UNiagaraSystem::CanObtainEmitterAttribute(const FNiagaraVariableBase& InVarWithUniqueNameNamespace) const
 {
@@ -1741,6 +1768,12 @@
 	return EmitterHandle;
 }
 
+void UNiagaraSystem::AddEmitterHandleDirect(FNiagaraEmitterHandle& EmitterHandle)
+{
+	EmitterHandles.Add(EmitterHandle);
+	RefreshSystemParametersFromEmitter(EmitterHandle);
+}
+
 FNiagaraEmitterHandle UNiagaraSystem::DuplicateEmitterHandle(const FNiagaraEmitterHandle& EmitterHandleToDuplicate, FName EmitterName)
 {
 	UNiagaraEmitter* DuplicateEmitter = UNiagaraEmitter::CreateAsDuplicate(*EmitterHandleToDuplicate.GetInstance(), EmitterName, *this);
@@ -2406,11 +2439,9 @@
 
 
 		//Compile all emitters
-		bool bTrulyAsync = true;
 		bool bAnyUnsynchronized = false;	
 
 		// Pass one... determine if any need to be compiled.
-		bool bForceSystems = false;
 		{
 
 			for (int32 i = 0; i < EmitterHandles.Num(); i++)
@@ -2418,8 +2449,6 @@
 				FNiagaraEmitterHandle Handle = EmitterHandles[i];
 				if (Handle.GetInstance() && Handle.GetIsEnabled())
 				{
-					UNiagaraScriptSourceBase* GraphSource = Handle.GetInstance()->GraphSource;
-
 					TArray<UNiagaraScript*> EmitterScripts;
 					Handle.GetInstance()->GetScripts(EmitterScripts, false, true);
 					check(EmitterScripts.Num() > 0);
@@ -2443,7 +2472,6 @@
 				}
 			}
 
-			bForceSystems = bForce || bAnyUnsynchronized;
 			bAnyCompiled = bAnyUnsynchronized || bForce;
 
 			// Now add the system scripts for compilation...
@@ -2486,66 +2514,59 @@
 		}
 
 
-	
+		// We found things needing compilation, now we have to go through an static duplicate everything that will be translated...
 		{
+			UNiagaraPrecompileContainer* Container = NewObject<UNiagaraPrecompileContainer>(GetTransientPackage());
+			Container->System = this;
+			Container->Scripts = ScriptsNeedingCompile;
+			TSharedPtr<FNiagaraCompileRequestDataBase, ESPMode::ThreadSafe> SystemPrecompiledData = NiagaraModule.Precompile(Container, FGuid());
 
-			// We found things needing compilation, now we have to go through an static duplicate everything that will be translated...
+			if (SystemPrecompiledData.IsValid() == false)
 			{
-				UNiagaraPrecompileContainer* Container = NewObject<UNiagaraPrecompileContainer>(GetTransientPackage());
-				Container->System = this;
-				Container->Scripts = ScriptsNeedingCompile;
-				TSharedPtr<FNiagaraCompileRequestDataBase, ESPMode::ThreadSafe> SystemPrecompiledData = NiagaraModule.Precompile(Container, FGuid());
+				UE_LOG(LogNiagara, Error, TEXT("Failed to precompile %s.  This is due to unexpected invalid or broken data.  Additional details should be in the log."), *GetPathName());
+				return false;
+			}
 
-				if (SystemPrecompiledData.IsValid() == false)
-				{
-					UE_LOG(LogNiagara, Error, TEXT("Failed to precompile %s.  This is due to unexpected invalid or broken data.  Additional details should be in the log."), *GetPathName());
-					return false;
-				}
+			SystemPrecompiledData->GetReferencedObjects(ActiveCompilation.RootObjects);
+			ActiveCompilation.MappedData.Add(SystemSpawnScript, SystemPrecompiledData);
+			ActiveCompilation.MappedData.Add(SystemUpdateScript, SystemPrecompiledData);
 
-				SystemPrecompiledData->GetReferencedObjects(ActiveCompilation.RootObjects);
-				ActiveCompilation.MappedData.Add(SystemSpawnScript, SystemPrecompiledData);
-				ActiveCompilation.MappedData.Add(SystemUpdateScript, SystemPrecompiledData);
-
-				check(EmitterHandles.Num() == SystemPrecompiledData->GetDependentRequestCount());
+			check(EmitterHandles.Num() == SystemPrecompiledData->GetDependentRequestCount());
 
 
-				// Grab the list of user variables that were actually encountered so that we can add to them later.
-				TArray<FNiagaraVariable> EncounteredExposedVars;
-				SystemPrecompiledData->GatherPreCompiledVariables(TEXT("User"), EncounteredExposedVars);
+			// Grab the list of user variables that were actually encountered so that we can add to them later.
+			TArray<FNiagaraVariable> EncounteredExposedVars;
+			SystemPrecompiledData->GatherPreCompiledVariables(TEXT("User"), EncounteredExposedVars);
 
-				for (int32 i = 0; i < EmitterHandles.Num(); i++)
+			for (int32 i = 0; i < EmitterHandles.Num(); i++)
+			{
+				FNiagaraEmitterHandle Handle = EmitterHandles[i];
+				if (Handle.GetInstance() && Handle.GetIsEnabled())
 				{
-					FNiagaraEmitterHandle Handle = EmitterHandles[i];
-					if (Handle.GetInstance() && Handle.GetIsEnabled())
+					TSharedPtr<FNiagaraCompileRequestDataBase, ESPMode::ThreadSafe> EmitterPrecompiledData = SystemPrecompiledData->GetDependentRequest(i);
+					EmitterPrecompiledData->GetReferencedObjects(ActiveCompilation.RootObjects);
+
+					TArray<UNiagaraScript*> EmitterScripts;
+					Handle.GetInstance()->GetScripts(EmitterScripts, false, true);
+					check(EmitterScripts.Num() > 0);
+					for (UNiagaraScript* EmitterScript : EmitterScripts)
 					{
-						UNiagaraScriptSourceBase* GraphSource = Handle.GetInstance()->GraphSource;
-						TSharedPtr<FNiagaraCompileRequestDataBase, ESPMode::ThreadSafe> EmitterPrecompiledData = SystemPrecompiledData->GetDependentRequest(i);
-						EmitterPrecompiledData->GetReferencedObjects(ActiveCompilation.RootObjects);
+						ActiveCompilation.MappedData.Add(EmitterScript, EmitterPrecompiledData);
+					}
 
-						TArray<UNiagaraScript*> EmitterScripts;
-						Handle.GetInstance()->GetScripts(EmitterScripts, false, true);
-						check(EmitterScripts.Num() > 0);
-						for (UNiagaraScript* EmitterScript : EmitterScripts)
-						{
-							ActiveCompilation.MappedData.Add(EmitterScript, EmitterPrecompiledData);
-						}
+					// Add the emitter's User variables to the encountered list to expose for later.
+					EmitterPrecompiledData->GatherPreCompiledVariables(TEXT("User"), EncounteredExposedVars);
+				}
+			}
 
 
-						// Add the emitter's User variables to the encountered list to expose for later.
-						EmitterPrecompiledData->GatherPreCompiledVariables(TEXT("User"), EncounteredExposedVars);
-
-					}
-				}
-
-
-				// Now let's synchronize the variables that we actually encountered during precompile so that we can expose them to the end user.
-				for (int32 i = 0; i < EncounteredExposedVars.Num(); i++)
+			// Now let's synchronize the variables that we actually encountered during precompile so that we can expose them to the end user.
+			for (int32 i = 0; i < EncounteredExposedVars.Num(); i++)
+			{
+				if (OriginalExposedParams.Contains(EncounteredExposedVars[i]) == false)
 				{
-					if (OriginalExposedParams.Contains(EncounteredExposedVars[i]) == false)
-					{
-						// Just in case it wasn't added previously..
-						ExposedParameters.AddParameter(EncounteredExposedVars[i]);
-					}
+					// Just in case it wasn't added previously..
+					ExposedParameters.AddParameter(EncounteredExposedVars[i]);
 				}
 			}
 		}
diff --git a/Engine/Plugins/FX/Niagara/Source/NiagaraEditor/NiagaraEditor.Build.cs b/Engine/Plugins/FX/Niagara/Source/NiagaraEditor/NiagaraEditor.Build.cs
--- a/Engine/Plugins/FX/Niagara/Source/NiagaraEditor/NiagaraEditor.Build.cs
+++ b/Engine/Plugins/FX/Niagara/Source/NiagaraEditor/NiagaraEditor.Build.cs
@@ -27,6 +27,7 @@ public class NiagaraEditor : ModuleRules
 				"RenderCore",
 				"Slate", 
 				"SlateCore",
+				"SlateNullRenderer",
 				"Kismet",
                 "EditorStyle",
 				"UnrealEd", 
diff --git a/Engine/Plugins/FX/Niagara/Source/NiagaraEditor/Private/NiagaraEditorUtilities.cpp b/Engine/Plugins/FX/Niagara/Source/NiagaraEditor/Private/NiagaraEditorUtilities.cpp
--- a/Engine/Plugins/FX/Niagara/Source/NiagaraEditor/Private/NiagaraEditorUtilities.cpp
+++ b/Engine/Plugins/FX/Niagara/Source/NiagaraEditor/Private/NiagaraEditorUtilities.cpp
@@ -2031,7 +2031,7 @@ TArray<UNiagaraComponent*> FNiagaraEditorUtilities::GetComponentsThatReferenceSy
 	return ReferencingComponents;
 }
 
-const FGuid FNiagaraEditorUtilities::AddEmitterToSystem(UNiagaraSystem& InSystem, UNiagaraEmitter& InEmitterToAdd)
+const FGuid FNiagaraEditorUtilities::AddEmitterToSystem(UNiagaraSystem& InSystem, UNiagaraEmitter& InEmitterToAdd, bool bCreateCopy)
 {
 	// Kill all system instances before modifying the emitter handle list to prevent accessing deleted data.
 	KillSystemInstances(InSystem);
@@ -2049,7 +2049,7 @@ const FGuid FNiagaraEditorUtilities::AddEmitterToSystem(UNiagaraSystem& InSystem
 		InSystem.Modify();
 		EmitterHandle = InSystem.AddEmitterHandle(InEmitterToAdd, FNiagaraUtilities::GetUniqueName(InEmitterToAdd.GetFName(), EmitterHandleNames));
 	}
-	else
+	else if (bCreateCopy)
 	{
 		// When editing an emitter asset we add the emitter as a duplicate so that the parent emitter is duplicated, but it's parent emitter
 		// information is maintained.
@@ -2057,6 +2057,11 @@ const FGuid FNiagaraEditorUtilities::AddEmitterToSystem(UNiagaraSystem& InSystem
 		FNiagaraEmitterHandle TemporaryEmitterHandle(InEmitterToAdd);
 		EmitterHandle = InSystem.DuplicateEmitterHandle(TemporaryEmitterHandle, *InEmitterToAdd.GetUniqueEmitterName());
 	}
+	else
+	{
+		EmitterHandle = FNiagaraEmitterHandle(InEmitterToAdd);
+		InSystem.AddEmitterHandleDirect(EmitterHandle);		
+	}
 	
 	FNiagaraStackGraphUtilities::RebuildEmitterNodes(InSystem);
 	SystemEditorData->SynchronizeOverviewGraphWithSystem(InSystem);
diff --git a/Engine/Plugins/FX/Niagara/Source/NiagaraEditor/Private/ViewModels/NiagaraSystemViewModel.cpp b/Engine/Plugins/FX/Niagara/Source/NiagaraEditor/Private/ViewModels/NiagaraSystemViewModel.cpp
--- a/Engine/Plugins/FX/Niagara/Source/NiagaraEditor/Private/ViewModels/NiagaraSystemViewModel.cpp
+++ b/Engine/Plugins/FX/Niagara/Source/NiagaraEditor/Private/ViewModels/NiagaraSystemViewModel.cpp
@@ -361,14 +361,14 @@ TSharedPtr<FNiagaraEmitterHandleViewModel> FNiagaraSystemViewModel::AddEmitter(U
 		GEditor->BeginTransaction(LOCTEXT("AddEmitter", "Add emitter"));
 	}
 
-	const FGuid NewEmitterHandleId = FNiagaraEditorUtilities::AddEmitterToSystem(GetSystem(), Emitter);
+	const FGuid NewEmitterHandleId = FNiagaraEditorUtilities::AddEmitterToSystem(GetSystem(), Emitter, EditMode != ENiagaraSystemViewModelEditMode::EmitterDuringMerge);
 
 	if (false == bSystemIsPlaceholder)
 	{
 		GEditor->EndTransaction();
 	}
 
-	if (GetSystem().GetNumEmitters() == 1 && EditorSettings->GetAutoPlay())
+	if (GetSystem().GetNumEmitters() == 1 && EditorSettings->GetAutoPlay() && Sequencer.IsValid())
 	{
 		// When adding a new emitter to an empty system start playing.
 		Sequencer->SetPlaybackStatus(EMovieScenePlayerStatus::Playing);
@@ -386,10 +386,13 @@ TSharedPtr<FNiagaraEmitterHandleViewModel> FNiagaraSystemViewModel::AddEmitter(U
 
 	TGuardValue<bool> UpdateGuard(bUpdatingSequencerFromEmitterDataChange, true);
 
-	FFrameTime NewStartFrame = NewSystemPlaybackRange.GetLowerBoundValue() * NiagaraSequence->GetMovieScene()->GetTickResolution();
-	int32 NewDuration = (NewSystemPlaybackRange.Size<float>() * NiagaraSequence->GetMovieScene()->GetTickResolution()).FrameNumber.Value;
+	if (NiagaraSequence)
+	{
+		FFrameTime NewStartFrame = NewSystemPlaybackRange.GetLowerBoundValue() * NiagaraSequence->GetMovieScene()->GetTickResolution();
+		int32 NewDuration = (NewSystemPlaybackRange.Size<float>() * NiagaraSequence->GetMovieScene()->GetTickResolution()).FrameNumber.Value;
 
-	NiagaraSequence->GetMovieScene()->SetPlaybackRange(NewStartFrame.RoundToFrame(), NewDuration);
+		NiagaraSequence->GetMovieScene()->SetPlaybackRange(NewStartFrame.RoundToFrame(), NewDuration);
+	}
 
 	TSharedPtr<FNiagaraEmitterHandleViewModel> NewEmitterHandleViewModel = GetEmitterHandleViewModelById(NewEmitterHandleId);
 	TArray<UNiagaraStackEntry*> SelectedStackEntries;
@@ -1222,12 +1225,15 @@ void FNiagaraSystemViewModel::ResetEmitterHandleViewModelsAndTracks()
 		EmitterHandleViewModel->Reset();
 	}
 
-	TArray<UMovieSceneTrack*> MasterTracks = NiagaraSequence->GetMovieScene()->GetMasterTracks();
-	for (UMovieSceneTrack* MasterTrack : MasterTracks)
+	if (NiagaraSequence)
 	{
-		if (MasterTrack != nullptr)
+		TArray<UMovieSceneTrack*> MasterTracks = NiagaraSequence->GetMovieScene()->GetMasterTracks();
+		for (UMovieSceneTrack* MasterTrack : MasterTracks)
 		{
-			NiagaraSequence->GetMovieScene()->RemoveMasterTrack(*MasterTrack);
+			if (MasterTrack != nullptr)
+			{
+				NiagaraSequence->GetMovieScene()->RemoveMasterTrack(*MasterTrack);
+			}
 		}
 	}
 }
@@ -1334,64 +1340,72 @@ void FNiagaraSystemViewModel::RefreshSequencerTracks()
 {
 	TGuardValue<bool> UpdateGuard(bUpdatingSequencerFromEmitterDataChange, true);
 
-	TArray<UMovieSceneTrack*> MasterTracks = NiagaraSequence->GetMovieScene()->GetMasterTracks();
-	for (UMovieSceneTrack* MasterTrack : MasterTracks)
+	if (Sequencer.IsValid())
 	{
-		if (MasterTrack != nullptr)
+		TArray<UMovieSceneTrack*> MasterTracks = NiagaraSequence->GetMovieScene()->GetMasterTracks();
+		for (UMovieSceneTrack* MasterTrack : MasterTracks)
 		{
-			NiagaraSequence->GetMovieScene()->RemoveMasterTrack(*MasterTrack);
+			if (MasterTrack != nullptr)
+			{
+				NiagaraSequence->GetMovieScene()->RemoveMasterTrack(*MasterTrack);
+			}
 		}
-	}
 
-	TMap<FGuid, UMovieSceneNiagaraEmitterTrack*> EmitterHandleIdToTrackMap;
-	for (TSharedRef<FNiagaraEmitterHandleViewModel> EmitterHandleViewModel : EmitterHandleViewModels)
-	{
-		UMovieSceneNiagaraEmitterTrack* EmitterTrack = Cast<UMovieSceneNiagaraEmitterTrack>(NiagaraSequence->GetMovieScene()->AddMasterTrack(UMovieSceneNiagaraEmitterTrack::StaticClass()));
-		EmitterTrack->Initialize(*this, EmitterHandleViewModel, NiagaraSequence->GetMovieScene()->GetTickResolution());
-		EmitterHandleIdToTrackMap.Add(EmitterHandleViewModel->GetId(), EmitterTrack);
-	}
+		TMap<FGuid, UMovieSceneNiagaraEmitterTrack*> EmitterHandleIdToTrackMap;
+		for (TSharedRef<FNiagaraEmitterHandleViewModel> EmitterHandleViewModel : EmitterHandleViewModels)
+		{
+			UMovieSceneNiagaraEmitterTrack* EmitterTrack = Cast<UMovieSceneNiagaraEmitterTrack>(NiagaraSequence->GetMovieScene()->AddMasterTrack(UMovieSceneNiagaraEmitterTrack::StaticClass()));
+			EmitterTrack->Initialize(*this, EmitterHandleViewModel, NiagaraSequence->GetMovieScene()->GetTickResolution());
+			EmitterHandleIdToTrackMap.Add(EmitterHandleViewModel->GetId(), EmitterTrack);
+		}
 
-	TArray<UMovieSceneFolder*>& MovieSceneRootFolders = NiagaraSequence->GetMovieScene()->GetRootFolders();
-	MovieSceneRootFolders.Empty();
+		TArray<UMovieSceneFolder*>& MovieSceneRootFolders = NiagaraSequence->GetMovieScene()->GetRootFolders();
+		MovieSceneRootFolders.Empty();
 
-	const UNiagaraSystemEditorData& SystemEditorData = GetEditorData();
-	UNiagaraSystemEditorFolder& RootFolder = SystemEditorData.GetRootFolder();
-	for (const UNiagaraSystemEditorFolder* RootChildFolder : RootFolder.GetChildFolders())
-	{
-		UMovieSceneFolder* MovieSceneRootFolder = NewObject<UMovieSceneFolder>(NiagaraSequence->GetMovieScene(), RootChildFolder->GetFolderName(), RF_Transactional);
-		MovieSceneRootFolder->SetFolderName(RootChildFolder->GetFolderName());
-		MovieSceneRootFolders.Add(MovieSceneRootFolder);
-		PopulateChildMovieSceneFoldersFromNiagaraFolders(RootChildFolder, MovieSceneRootFolder, EmitterHandleIdToTrackMap);
+		const UNiagaraSystemEditorData& SystemEditorData = GetEditorData();
+		UNiagaraSystemEditorFolder& RootFolder = SystemEditorData.GetRootFolder();
+		for (const UNiagaraSystemEditorFolder* RootChildFolder : RootFolder.GetChildFolders())
+		{
+			UMovieSceneFolder* MovieSceneRootFolder = NewObject<UMovieSceneFolder>(NiagaraSequence->GetMovieScene(), RootChildFolder->GetFolderName(), RF_Transactional);
+			MovieSceneRootFolder->SetFolderName(RootChildFolder->GetFolderName());
+			MovieSceneRootFolders.Add(MovieSceneRootFolder);
+			PopulateChildMovieSceneFoldersFromNiagaraFolders(RootChildFolder, MovieSceneRootFolder, EmitterHandleIdToTrackMap);
+		}
+		Sequencer->NotifyMovieSceneDataChanged(EMovieSceneDataChangeType::MovieSceneStructureItemsChanged);
 	}
 
-	Sequencer->NotifyMovieSceneDataChanged(EMovieSceneDataChangeType::MovieSceneStructureItemsChanged);
-
 	// Since we just rebuilt all of the sequencer tracks, these updates don't need to be done.
 	EmitterIdsRequiringSequencerTrackUpdate.Empty();
 }
 
 void FNiagaraSystemViewModel::UpdateSequencerTracksForEmitters(const TArray<FGuid>& EmitterIdsRequiringUpdate)
 {
-	TGuardValue<bool> UpdateGuard(bUpdatingSequencerFromEmitterDataChange, true);
-	for (UMovieSceneTrack* Track : NiagaraSequence->GetMovieScene()->GetMasterTracks())
+	if (Sequencer.IsValid())
 	{
-		UMovieSceneNiagaraEmitterTrack* EmitterTrack = CastChecked<UMovieSceneNiagaraEmitterTrack>(Track);
-		if (EmitterIdsRequiringUpdate.Contains(EmitterTrack->GetEmitterHandleViewModel()->GetId()))
+		TGuardValue<bool> UpdateGuard(bUpdatingSequencerFromEmitterDataChange, true);
+		for (UMovieSceneTrack* Track : NiagaraSequence->GetMovieScene()->GetMasterTracks())
 		{
-			EmitterTrack->UpdateTrackFromEmitterGraphChange(NiagaraSequence->GetMovieScene()->GetTickResolution());
+			UMovieSceneNiagaraEmitterTrack* EmitterTrack = CastChecked<UMovieSceneNiagaraEmitterTrack>(Track);
+			if (EmitterIdsRequiringUpdate.Contains(EmitterTrack->GetEmitterHandleViewModel()->GetId()))
+			{
+				EmitterTrack->UpdateTrackFromEmitterGraphChange(NiagaraSequence->GetMovieScene()->GetTickResolution());
+			}
 		}
+		Sequencer->NotifyMovieSceneDataChanged(EMovieSceneDataChangeType::MovieSceneStructureItemsChanged);
 	}
-	Sequencer->NotifyMovieSceneDataChanged(EMovieSceneDataChangeType::MovieSceneStructureItemsChanged);
 }
 
 UMovieSceneNiagaraEmitterTrack* FNiagaraSystemViewModel::GetTrackForHandleViewModel(TSharedRef<FNiagaraEmitterHandleViewModel> EmitterHandleViewModel)
 {
-	for (UMovieSceneTrack* Track : NiagaraSequence->GetMovieScene()->GetMasterTracks())
+	if (NiagaraSequence)
 	{
-		UMovieSceneNiagaraEmitterTrack* EmitterTrack = CastChecked<UMovieSceneNiagaraEmitterTrack>(Track);
-		if (EmitterTrack->GetEmitterHandleViewModel() == EmitterHandleViewModel)
+		for (UMovieSceneTrack* Track : NiagaraSequence->GetMovieScene()->GetMasterTracks())
 		{
-			return EmitterTrack;
+			UMovieSceneNiagaraEmitterTrack* EmitterTrack = CastChecked<UMovieSceneNiagaraEmitterTrack>(Track);
+			if (EmitterTrack->GetEmitterHandleViewModel() == EmitterHandleViewModel)
+			{
+				return EmitterTrack;
+			}
 		}
 	}
 	return nullptr;
@@ -1399,6 +1413,11 @@ UMovieSceneNiagaraEmitterTrack* FNiagaraSystemViewModel::GetTrackForHandleViewMo
 
 void FNiagaraSystemViewModel::SetupSequencer()
 {
+	if (EditMode == ENiagaraSystemViewModelEditMode::EmitterDuringMerge)
+	{
+		// we don't need a sequencer when merging emitters
+		return;
+	}
 	NiagaraSequence = NewObject<UNiagaraSequence>(GetTransientPackage());
 	UMovieScene* MovieScene = NewObject<UMovieScene>(NiagaraSequence, FName("Niagara System MovieScene"), RF_Transactional);
 	MovieScene->SetDisplayRate(FFrameRate(240, 1));
@@ -1448,7 +1467,7 @@ void FNiagaraSystemViewModel::SetupSequencer()
 void FNiagaraSystemViewModel::SnapToNextSpeed(const FString& PropertyName, const UNiagaraEditorSettings* Settings)
 {
 	// we update the speed in any case
-	if (Sequencer != nullptr)
+	if (Sequencer.IsValid())
 	{
 		Sequencer->SnapToClosestPlaybackSpeed();
 	}
@@ -1473,19 +1492,23 @@ void FNiagaraSystemViewModel::ResetSystem()
 
 void FNiagaraSystemViewModel::ResetSystem(ETimeResetMode TimeResetMode, EMultiResetMode MultiResetMode, EReinitMode ReinitMode)
 {
-	bool bResetAge = TimeResetMode == ETimeResetMode::AllowResetTime && (Sequencer->GetPlaybackStatus() == EMovieScenePlayerStatus::Playing || Sequencer->GetPlaybackStatus() == EMovieScenePlayerStatus::Stopped || EditorSettings->GetResimulateOnChangeWhilePaused() == false);
-	if (bResetAge)
+	bool bResetAge = false;
+	if (Sequencer.IsValid())
 	{
-		TGuardValue<bool> Guard(bSettingSequencerTimeDirectly, true);
-		if (Sequencer->GetPlaybackStatus() == EMovieScenePlayerStatus::Playing)
+		bResetAge = TimeResetMode == ETimeResetMode::AllowResetTime && (Sequencer->GetPlaybackStatus() == EMovieScenePlayerStatus::Playing || Sequencer->GetPlaybackStatus() == EMovieScenePlayerStatus::Stopped || EditorSettings->GetResimulateOnChangeWhilePaused() == false);
+		if (bResetAge)
 		{
-			Sequencer->SetPlaybackStatus(EMovieScenePlayerStatus::Paused);
-			Sequencer->SetGlobalTime(0);
-			Sequencer->SetPlaybackStatus(EMovieScenePlayerStatus::Playing);
-		}
-		else
-		{
-			Sequencer->SetGlobalTime(0);
+			TGuardValue<bool> Guard(bSettingSequencerTimeDirectly, true);
+			if (Sequencer->GetPlaybackStatus() == EMovieScenePlayerStatus::Playing)
+			{
+				Sequencer->SetPlaybackStatus(EMovieScenePlayerStatus::Paused);
+				Sequencer->SetGlobalTime(0);
+				Sequencer->SetPlaybackStatus(EMovieScenePlayerStatus::Playing);
+			}
+			else
+			{
+				Sequencer->SetGlobalTime(0);
+			}
 		}
 	}
 
@@ -1634,7 +1657,7 @@ void FNiagaraSystemViewModel::EmitterHandlePropertyChanged(FGuid OwningEmitterHa
 {
 	// When the emitter handle changes, refresh the System scripts emitter nodes and the sequencer tracks just in case the
 	// property that changed was the handles emitter.
-	if (bUpdatingEmittersFromSequencerDataChange == false)
+	if (bUpdatingEmittersFromSequencerDataChange == false && NiagaraSequence)
 	{
 		TGuardValue<bool> UpdateGuard(bUpdatingSequencerFromEmitterDataChange, true);
 		for (UMovieSceneTrack* Track : NiagaraSequence->GetMovieScene()->GetMasterTracks())
@@ -1715,7 +1738,7 @@ void FNiagaraSystemViewModel::EmitterParameterStoreChanged(const FNiagaraParamet
 		{
 			UpdateSequencerTracksForEmitters(EmitterIdsRequiringSequencerTrackUpdate);
 		}
-		else
+		else if (Sequencer.IsValid())
 		{
 			TGuardValue<bool> UpdateGuard(bUpdatingSequencerFromEmitterDataChange, true);
 			for (UMovieSceneTrack* Track : NiagaraSequence->GetMovieScene()->GetMasterTracks())
@@ -1723,6 +1746,7 @@ void FNiagaraSystemViewModel::EmitterParameterStoreChanged(const FNiagaraParamet
 				UMovieSceneNiagaraEmitterTrack* EmitterTrack = CastChecked<UMovieSceneNiagaraEmitterTrack>(Track);
 				EmitterTrack->UpdateTrackFromEmitterParameterChange(NiagaraSequence->GetMovieScene()->GetTickResolution());
 			}
+			
 			Sequencer->NotifyMovieSceneDataChanged(EMovieSceneDataChangeType::TrackValueChanged);
 		}
 	}
@@ -1739,7 +1763,7 @@ void FNiagaraSystemViewModel::UpdateSimulationFromParameterChange()
 	}
 	else
 	{
-		if (Sequencer->GetPlaybackStatus() == EMovieScenePlayerStatus::Stopped)
+		if (Sequencer && Sequencer->GetPlaybackStatus() == EMovieScenePlayerStatus::Stopped)
 		{
 			// TODO: Update the view when paused and reset on change is turned off.
 		}
@@ -1809,7 +1833,7 @@ void PopulateNiagaraFoldersFromMovieSceneFolders(const TArray<UMovieSceneFolder*
 
 void FNiagaraSystemViewModel::SequencerDataChanged(EMovieSceneDataChangeType DataChangeType)
 {
-	if (bUpdatingSequencerFromEmitterDataChange == false && GIsTransacting == false)
+	if (bUpdatingSequencerFromEmitterDataChange == false && GIsTransacting == false && NiagaraSequence)
 	{
 		TGuardValue<bool> UpdateGuard(bUpdatingEmittersFromSequencerDataChange, true);
 
@@ -1919,7 +1943,7 @@ void FNiagaraSystemViewModel::SequencerDataChanged(EMovieSceneDataChangeType Dat
 
 void FNiagaraSystemViewModel::SequencerTimeChanged()
 {
-	if (!PreviewComponent || !SystemInstance || !SystemInstance->GetAreDataInterfacesInitialized())
+	if (!PreviewComponent || !SystemInstance || !SystemInstance->GetAreDataInterfacesInitialized() || !Sequencer.IsValid())
 	{
 		return;
 	}
@@ -1997,6 +2021,11 @@ void FNiagaraSystemViewModel::SequencerSectionSelectionChanged(TArray<UMovieScen
 
 void FNiagaraSystemViewModel::UpdateEmitterHandleSelectionFromSequencer()
 {
+	if (!Sequencer.IsValid())
+	{
+		return;
+	}
+	
 	TArray<FGuid> NewSelectedEmitterHandleIds;
 
 	TArray<UMovieSceneTrack*> SelectedTracks;
@@ -2054,6 +2083,10 @@ void FNiagaraSystemViewModel::UpdateEmitterHandleSelectionFromSequencer()
 
 void FNiagaraSystemViewModel::UpdateSequencerFromEmitterHandleSelection()
 {
+	if (!Sequencer.IsValid())
+	{
+		return;
+	}
 	TGuardValue<bool> UpdateGuard(bUpdatingSequencerSelectionFromSystem, true);
 	Sequencer->EmptySelection();
 	for (FGuid SelectedEmitterHandleId : SelectionViewModel->GetSelectedEmitterHandleIds())
diff --git a/Engine/Plugins/FX/Niagara/Source/NiagaraEditor/Public/NiagaraEditorUtilities.h b/Engine/Plugins/FX/Niagara/Source/NiagaraEditor/Public/NiagaraEditorUtilities.h
--- a/Engine/Plugins/FX/Niagara/Source/NiagaraEditor/Public/NiagaraEditorUtilities.h
+++ b/Engine/Plugins/FX/Niagara/Source/NiagaraEditor/Public/NiagaraEditorUtilities.h
@@ -252,7 +252,7 @@ namespace FNiagaraEditorUtilities
 
 	TArray<UNiagaraComponent*> GetComponentsThatReferenceSystemViewModel(const FNiagaraSystemViewModel& ReferencedSystemViewModel);
 
-	NIAGARAEDITOR_API const FGuid AddEmitterToSystem(UNiagaraSystem& InSystem, UNiagaraEmitter& InEmitterToAdd);
+	NIAGARAEDITOR_API const FGuid AddEmitterToSystem(UNiagaraSystem& InSystem, UNiagaraEmitter& InEmitterToAdd, bool bCreateCopy = true);
 
 	void RemoveEmittersFromSystemByEmitterHandleId(UNiagaraSystem& InSystem, TSet<FGuid> EmitterHandleIdsToDelete);
 
diff --git a/Engine/Plugins/FX/Niagara/Source/NiagaraEditor/Public/ViewModels/NiagaraSystemViewModel.h b/Engine/Plugins/FX/Niagara/Source/NiagaraEditor/Public/ViewModels/NiagaraSystemViewModel.h
--- a/Engine/Plugins/FX/Niagara/Source/NiagaraEditor/Public/ViewModels/NiagaraSystemViewModel.h
+++ b/Engine/Plugins/FX/Niagara/Source/NiagaraEditor/Public/ViewModels/NiagaraSystemViewModel.h
@@ -55,6 +55,8 @@ enum class NIAGARAEDITOR_API ENiagaraSystemViewModelEditMode
 	SystemAsset,
 	/** An emitter asset is being edited.  In this mode the system scripts will not be editable and all emitter editing options are available. */
 	EmitterAsset,
+	/** Special case where an emitter asset is being edited during a merge. In this mode changes made to the emitter being edited don't have to be applied, but directly affect the source. */
+	EmitterDuringMerge,
 };
 
 /** Defines options for the niagara System view model */
diff --git a/Engine/Source/Runtime/Engine/Private/AudioDevice.cpp b/Engine/Source/Runtime/Engine/Private/AudioDevice.cpp
--- a/Engine/Source/Runtime/Engine/Private/AudioDevice.cpp
+++ b/Engine/Source/Runtime/Engine/Private/AudioDevice.cpp
@@ -4705,7 +4705,7 @@ void FAudioDevice::AddNewActiveSoundInternal(const FActiveSound& NewActiveSound,
 			FAudioVirtualLoop VirtualLoop;
 			if (FAudioVirtualLoop::Virtualize(NewActiveSound, *this, bDoRangeCheck, VirtualLoop))
 			{
-				UE_LOG(LogAudio, Verbose, TEXT("New ActiveSound %s Virtualizing: Failed to pass concurrency"), *Sound->GetName());
+				UE_LOG(LogAudioConcurrency, Verbose, TEXT("New ActiveSound %s Virtualizing: Failed to pass concurrency"), *Sound->GetName());
 				AddVirtualLoop(VirtualLoop);
 			}
 			else
diff --git a/Engine/Source/Runtime/Engine/Private/SoundConcurrency.cpp b/Engine/Source/Runtime/Engine/Private/SoundConcurrency.cpp
--- a/Engine/Source/Runtime/Engine/Private/SoundConcurrency.cpp
+++ b/Engine/Source/Runtime/Engine/Private/SoundConcurrency.cpp
@@ -11,6 +11,7 @@
 #include "HAL/IConsoleManager.h"
 #include "Sound/SoundBase.h"
 
+DEFINE_LOG_CATEGORY(LogAudioConcurrency);
 
 static float ConcurrencyMinVolumeScaleCVar = 1.e-3f;
 FAutoConsoleVariableRef CVarConcurrencyMinVolumeScale(
@@ -34,7 +35,7 @@ namespace SoundConcurrency
 		{
 			if (const USoundBase* Sound = ActiveSound.GetSound())
 			{
-				UE_LOG(LogAudio, Verbose,
+				UE_LOG(LogAudioConcurrency, Verbose,
 					TEXT("Sound '%s' concurrency generation '%i' target volume update: %.3f to %.3f."),
 					*Sound->GetName(),
 					SoundData.Generation,
@@ -180,7 +181,7 @@ void FConcurrencyGroup::AddActiveSound(FActiveSound& ActiveSound)
 
 	if (ActiveSound.ConcurrencyGroupData.Contains(GroupID))
 	{
-		UE_LOG(LogAudio, Fatal, TEXT("Attempting to add active sound '%s' to concurrency group multiple times."), *ActiveSound.GetOwnerName());
+		UE_LOG(LogAudioConcurrency, Fatal, TEXT("Attempting to add active sound '%s' to concurrency group multiple times."), *ActiveSound.GetOwnerName());
 		return;
 	}
 
@@ -544,7 +545,7 @@ FConcurrencyGroup* FSoundConcurrencyManager::CanPlaySound(const FActiveSound& Ne
 	FConcurrencyGroup* ConcurrencyGroup = ConcurrencyGroups.FindRef(GroupID);
 	if (!ConcurrencyGroup)
 	{
-		UE_LOG(LogAudio, Warning, TEXT("Attempting to add active sound '%s' (owner '%s') to invalid concurrency group."),
+		UE_LOG(LogAudioConcurrency, Warning, TEXT("Attempting to add active sound '%s' (owner '%s') to invalid concurrency group."),
 			NewActiveSound.GetSound() ? *NewActiveSound.GetSound()->GetFullName() : TEXT("Unset"),
 			*NewActiveSound.GetOwnerName());
 		return nullptr;
@@ -859,6 +860,11 @@ FActiveSound* FSoundConcurrencyManager::CreateAndEvictActiveSounds(const FActive
 			continue;
 		}
 
+		if (SoundToEvict->GetSound())
+		{
+			UE_LOG(LogAudioConcurrency, VeryVerbose, TEXT("Evicting Sound %s due to concurrency"), *(SoundToEvict->GetSound()->GetName()));
+		}
+
 		StopDueToVoiceStealing(*SoundToEvict);
 	}
 
@@ -876,7 +882,7 @@ void FSoundConcurrencyManager::RemoveActiveSound(FActiveSound& ActiveSound)
 		FConcurrencyGroup* ConcurrencyGroup = ConcurrencyGroups.FindRef(ConcurrencyGroupID);
 		if (!ConcurrencyGroup)
 		{
-			UE_LOG(LogAudio, Error, TEXT("Attempting to remove stopped sound '%s' from inactive concurrency group."),
+			UE_LOG(LogAudioConcurrency, Error, TEXT("Attempting to remove stopped sound '%s' from inactive concurrency group."),
 				ActiveSound.GetSound() ? *ActiveSound.GetSound()->GetName(): TEXT("Unset"));
 			continue;
 		}
@@ -962,7 +968,7 @@ void FSoundConcurrencyManager::StopDueToVoiceStealing(FActiveSound& ActiveSound)
 		ActiveSound.ClearAudioComponent();
 		if (USoundBase* Sound = ActiveSound.GetSound())
 		{
-			UE_LOG(LogAudio, Verbose, TEXT("Playing ActiveSound %s Virtualizing: Sound's voice stollen due to concurrency group maximum met."), *Sound->GetName());
+			UE_LOG(LogAudioConcurrency, Verbose, TEXT("Playing ActiveSound %s Virtualizing: Sound's voice stollen due to concurrency group maximum met."), *Sound->GetName());
 		}
 		ActiveSound.AudioDevice->AddVirtualLoop(VirtualLoop);
 	}
