Niagara - Prevent creating property row generators when a system view model is goign to be used for data processing only and won't be displayed in the UI.

diff --git a/Engine/Plugins/FX/Niagara/Source/NiagaraEditor/Private/NiagaraScriptMergeManager.cpp b/Engine/Plugins/FX/Niagara/Source/NiagaraEditor/Private/NiagaraScriptMergeManager.cpp
--- a/Engine/Plugins/FX/Niagara/Source/NiagaraEditor/Private/NiagaraScriptMergeManager.cpp
+++ b/Engine/Plugins/FX/Niagara/Source/NiagaraEditor/Private/NiagaraScriptMergeManager.cpp
@@ -1341,6 +1341,7 @@ void FNiagaraScriptMergeManager::UpdateModuleVersions(UNiagaraEmitter& Instance,
 		SystemOptions.bCanModifyEmittersFromTimeline = false;
 		SystemOptions.bCanSimulate = false;
 		SystemOptions.bCanAutoCompile = false;
+		SystemOptions.bIsForDataProcessingOnly = true;
 		SystemOptions.EditMode = ENiagaraSystemViewModelEditMode::EmitterDuringMerge;
 		SystemOptions.MessageLogGuid = System->GetAssetGuid();
 		SystemViewModel->Initialize(*System, SystemOptions);
diff --git a/Engine/Plugins/FX/Niagara/Source/NiagaraEditor/Private/ViewModels/NiagaraSystemViewModel.cpp b/Engine/Plugins/FX/Niagara/Source/NiagaraEditor/Private/ViewModels/NiagaraSystemViewModel.cpp
--- a/Engine/Plugins/FX/Niagara/Source/NiagaraEditor/Private/ViewModels/NiagaraSystemViewModel.cpp
+++ b/Engine/Plugins/FX/Niagara/Source/NiagaraEditor/Private/ViewModels/NiagaraSystemViewModel.cpp
@@ -67,6 +67,7 @@ template<> TMap<UNiagaraSystem*, TArray<FNiagaraSystemViewModel*>> TNiagaraViewM
 FNiagaraSystemViewModelOptions::FNiagaraSystemViewModelOptions()
 	: bCanAutoCompile(true)
 	, bCanSimulate(true)
+	, bIsForDataProcessingOnly(false)
 {
 }
 
@@ -103,6 +104,7 @@ void FNiagaraSystemViewModel::Initialize(UNiagaraSystem& InSystem, FNiagaraSyste
 	EditMode = InOptions.EditMode;
 	OnGetSequencerAddMenuContent = InOptions.OnGetSequencerAddMenuContent;
 	SystemMessageLogGuidKey = InOptions.MessageLogGuid;
+	bIsForDataProcessingOnly = InOptions.bIsForDataProcessingOnly;
 
 	SystemChangedDelegateHandle = System->OnSystemPostEditChange().AddSP(this, &FNiagaraSystemViewModel::SystemChanged);
 
@@ -874,6 +876,11 @@ TSharedRef<FNiagaraPlaceholderDataInterfaceManager> FNiagaraSystemViewModel::Get
 	return PlaceholderDataInterfaceManager.ToSharedRef();
 }
 
+bool FNiagaraSystemViewModel::GetIsForDataProcessingOnly() const
+{
+	return bIsForDataProcessingOnly;
+}
+
 void FNiagaraSystemViewModel::SendLastCompileMessageJobs() const
 {
 	if (SystemMessageLogGuidKey.IsSet() == false)
diff --git a/Engine/Plugins/FX/Niagara/Source/NiagaraEditor/Private/ViewModels/Stack/NiagaraStackObject.cpp b/Engine/Plugins/FX/Niagara/Source/NiagaraEditor/Private/ViewModels/Stack/NiagaraStackObject.cpp
--- a/Engine/Plugins/FX/Niagara/Source/NiagaraEditor/Private/ViewModels/Stack/NiagaraStackObject.cpp
+++ b/Engine/Plugins/FX/Niagara/Source/NiagaraEditor/Private/ViewModels/Stack/NiagaraStackObject.cpp
@@ -2,6 +2,7 @@
 
 #include "ViewModels/Stack/NiagaraStackObject.h"
 #include "ViewModels/Stack/NiagaraStackPropertyRow.h"
+#include "ViewModels/NiagaraSystemViewModel.h"
 #include "NiagaraNode.h"
 #include "NiagaraEditorModule.h"
 
@@ -25,6 +26,7 @@ void UNiagaraStackObject::Initialize(FRequiredEntryData InRequiredEntryData, UOb
 	Super::Initialize(InRequiredEntryData, InOwnerStackItemEditorDataKey, ObjectStackEditorDataKey);
 	Object = InObject;
 	OwningNiagaraNode = InOwningNiagaraNode;
+	bIsRefresingDataInterfaceErrors = false;
 }
 
 void UNiagaraStackObject::SetOnSelectRootNodes(FOnSelectRootNodes OnSelectRootNodes)
@@ -143,7 +145,7 @@ void UNiagaraStackObject::RefreshChildrenInternal(const TArray<UNiagaraStackEntr
 
 	GatherIssueFromProperties((uint8*)Object, Object->GetClass(), true);
 
-	if (PropertyRowGenerator.IsValid() == false)
+	if (GetSystemViewModel()->GetIsForDataProcessingOnly() == false && PropertyRowGenerator.IsValid() == false)
 	{
 		FPropertyEditorModule& PropertyEditorModule = FModuleManager::LoadModuleChecked<FPropertyEditorModule>("PropertyEditor");
 		FPropertyRowGeneratorArgs Args;
@@ -191,11 +193,11 @@ void UNiagaraStackObject::RefreshChildrenInternal(const TArray<UNiagaraStackEntr
 	if (DataInterfaceObject != nullptr)
 	{
 		// First we need to refresh the errors on the data interface so that the rows in the property row generator 
-		// are correct.  We need to remove the delegate handler first so that we don't cause a reentrant refresh.
-		// This can be updated to use a guard value post 4.25.
-		PropertyRowGenerator->OnRowsRefreshed().RemoveAll(this);
-		DataInterfaceObject->RefreshErrors();
-		PropertyRowGenerator->OnRowsRefreshed().AddUObject(this, &UNiagaraStackObject::PropertyRowsRefreshed);
+		// are correct.
+		{
+			TGuardValue<bool> RefreGuard(bIsRefresingDataInterfaceErrors, true);
+			DataInterfaceObject->RefreshErrors();
+		}
 
 		// Generate the summary stack issue for any errors which are generated.
 		TArray<FNiagaraDataInterfaceError> Errors;
@@ -221,35 +223,37 @@ void UNiagaraStackObject::RefreshChildrenInternal(const TArray<UNiagaraStackEntr
 		}
 	}
 
-
-	TArray<TSharedRef<IDetailTreeNode>> DefaultRootTreeNodes = PropertyRowGenerator->GetRootTreeNodes();
-	TArray<TSharedRef<IDetailTreeNode>> RootTreeNodes;
-	if (OnSelectRootNodesDelegate.IsBound())
-	{
-		OnSelectRootNodesDelegate.Execute(DefaultRootTreeNodes, &RootTreeNodes);
-	}
-	else
-	{
-		RootTreeNodes = DefaultRootTreeNodes;
-	}
-
-	for (TSharedRef<IDetailTreeNode> RootTreeNode : RootTreeNodes)
+	if(PropertyRowGenerator.IsValid())
 	{
-		if (RootTreeNode->GetNodeType() == EDetailNodeType::Advanced)
+		TArray<TSharedRef<IDetailTreeNode>> DefaultRootTreeNodes = PropertyRowGenerator->GetRootTreeNodes();
+		TArray<TSharedRef<IDetailTreeNode>> RootTreeNodes;
+		if (OnSelectRootNodesDelegate.IsBound())
 		{
-			continue;
+			OnSelectRootNodesDelegate.Execute(DefaultRootTreeNodes, &RootTreeNodes);
 		}
-
-		UNiagaraStackPropertyRow* ChildRow = FindCurrentChildOfTypeByPredicate<UNiagaraStackPropertyRow>(CurrentChildren,
-			[=](UNiagaraStackPropertyRow* CurrentChild) { return CurrentChild->GetDetailTreeNode() == RootTreeNode; });
-
-		if (ChildRow == nullptr)
+		else
 		{
-			ChildRow = NewObject<UNiagaraStackPropertyRow>(this);
-			ChildRow->Initialize(CreateDefaultChildRequiredData(), RootTreeNode, GetOwnerStackItemEditorDataKey(), GetOwnerStackItemEditorDataKey(), OwningNiagaraNode);
+			RootTreeNodes = DefaultRootTreeNodes;
 		}
 
-		NewChildren.Add(ChildRow);
+		for (TSharedRef<IDetailTreeNode> RootTreeNode : RootTreeNodes)
+		{
+			if (RootTreeNode->GetNodeType() == EDetailNodeType::Advanced)
+			{
+				continue;
+			}
+
+			UNiagaraStackPropertyRow* ChildRow = FindCurrentChildOfTypeByPredicate<UNiagaraStackPropertyRow>(CurrentChildren,
+				[=](UNiagaraStackPropertyRow* CurrentChild) { return CurrentChild->GetDetailTreeNode() == RootTreeNode; });
+
+			if (ChildRow == nullptr)
+			{
+				ChildRow = NewObject<UNiagaraStackPropertyRow>(this);
+				ChildRow->Initialize(CreateDefaultChildRequiredData(), RootTreeNode, GetOwnerStackItemEditorDataKey(), GetOwnerStackItemEditorDataKey(), OwningNiagaraNode);
+			}
+
+			NewChildren.Add(ChildRow);
+		}
 	}
 }
 
@@ -260,5 +264,8 @@ void UNiagaraStackObject::PostRefreshChildrenInternal()
 
 void UNiagaraStackObject::PropertyRowsRefreshed()
 {
-	RefreshChildren();
+	if(bIsRefresingDataInterfaceErrors == false)
+	{
+		RefreshChildren();
+	}
 }
diff --git a/Engine/Plugins/FX/Niagara/Source/NiagaraEditor/Public/ViewModels/NiagaraSystemViewModel.h b/Engine/Plugins/FX/Niagara/Source/NiagaraEditor/Public/ViewModels/NiagaraSystemViewModel.h
--- a/Engine/Plugins/FX/Niagara/Source/NiagaraEditor/Public/ViewModels/NiagaraSystemViewModel.h
+++ b/Engine/Plugins/FX/Niagara/Source/NiagaraEditor/Public/ViewModels/NiagaraSystemViewModel.h
@@ -81,6 +81,9 @@ struct NIAGARAEDITOR_API FNiagaraSystemViewModelOptions
 
 	/** Gets the current editing mode for this system. */
 	ENiagaraSystemViewModelEditMode EditMode;
+
+	/** Specifies that the view model is being constructed for data processing only and will not be displayed in the UI. */
+	bool bIsForDataProcessingOnly;
 };
 
 struct FNiagaraStackModuleData
@@ -378,6 +381,9 @@ public:
 
 	TSharedRef<FNiagaraPlaceholderDataInterfaceManager> GetPlaceholderDataInterfaceManager();
 
+	/** Gets whether or not this view model is for data processing only and will not be displayed in the UI. */
+	bool GetIsForDataProcessingOnly() const;
+
 private:
 	/** Sends message jobs to FNiagaraMessageManager for all compile events from the last compile. */
 	void SendLastCompileMessageJobs() const;
@@ -651,4 +657,7 @@ private:
 
 	/** Flag for when messages have been added/removed through the viewmodel to signal that the message manager needs to be refreshed. */
 	mutable bool bPendingAssetMessagesChanged;
+
+	/** Specifies that this view model is for data processing only and will not be displayed in the UI. */
+	bool bIsForDataProcessingOnly;
 };
\ No newline at end of file
diff --git a/Engine/Plugins/FX/Niagara/Source/NiagaraEditor/Public/ViewModels/Stack/NiagaraStackObject.h b/Engine/Plugins/FX/Niagara/Source/NiagaraEditor/Public/ViewModels/Stack/NiagaraStackObject.h
--- a/Engine/Plugins/FX/Niagara/Source/NiagaraEditor/Public/ViewModels/Stack/NiagaraStackObject.h
+++ b/Engine/Plugins/FX/Niagara/Source/NiagaraEditor/Public/ViewModels/Stack/NiagaraStackObject.h
@@ -68,4 +68,5 @@ private:
 	TArray<FRegisteredClassCustomization> RegisteredClassCustomizations;
 	TArray<FRegisteredPropertyCustomization> RegisteredPropertyCustomizations;
 	TSharedPtr<IPropertyRowGenerator> PropertyRowGenerator;
+	bool bIsRefresingDataInterfaceErrors;
 };
