Merge from Release-Engine-Test [at] 16487383 to UE5/Main
This represents UE4/Main [at] 16445039 and Dev-PerfTest [at] 16444526

diff --git a/Engine/Plugins/FX/Niagara/Source/Niagara/Private/NiagaraSystem.cpp b/Engine/Plugins/FX/Niagara/Source/Niagara/Private/NiagaraSystem.cpp
--- a/Engine/Plugins/FX/Niagara/Source/Niagara/Private/NiagaraSystem.cpp
+++ b/Engine/Plugins/FX/Niagara/Source/Niagara/Private/NiagaraSystem.cpp
@@ -331,7 +331,7 @@ void UNiagaraSystem::UpdateSystemAfterLoad()
 	const bool bIsDedicatedServer = !GIsClient && GIsServer;
 
 	if (!GetOutermost()->bIsCookedForEditor && !bIsDedicatedServer)
-	{		
+	{
 		TArray<UNiagaraScript*> AllSystemScripts;
 		UNiagaraScriptSourceBase* SystemScriptSource;
 		if (SystemSpawnScript == nullptr)
@@ -782,6 +782,11 @@ void UNiagaraSystem::PostLoad()
 	// there will be no merging or compiling which makes it safe to do so.
 	UpdateSystemAfterLoad();
 #endif
+	// see the equivalent in NiagaraEmitter for details
+	if(bIsTemplateAsset_DEPRECATED)
+	{
+		TemplateSpecification = bIsTemplateAsset_DEPRECATED ? ENiagaraScriptTemplateSpecification::Template : ENiagaraScriptTemplateSpecification::None;
+	}
 }
 
 #if WITH_EDITORONLY_DATA
diff --git a/Engine/Plugins/FX/Niagara/Source/NiagaraEditor/Private/Customizations/NiagaraScriptVariableCustomization.cpp b/Engine/Plugins/FX/Niagara/Source/NiagaraEditor/Private/Customizations/NiagaraScriptVariableCustomization.cpp
--- a/Engine/Plugins/FX/Niagara/Source/NiagaraEditor/Private/Customizations/NiagaraScriptVariableCustomization.cpp
+++ b/Engine/Plugins/FX/Niagara/Source/NiagaraEditor/Private/Customizations/NiagaraScriptVariableCustomization.cpp
@@ -89,6 +89,15 @@ void FNiagaraScriptVariableDetails::PostUndo(bool bSuccess)
 			ParameterEditorLibraryValue->UpdateInternalValueFromStruct(ParameterValue.ToSharedRef());
 		}
 	}
+	else if (Variable->GetOuter()->IsA<UNiagaraParameterDefinitions>())
+	{
+		if (TypeUtilityLibraryValue && ParameterEditorLibraryValue)
+		{
+			TSharedPtr<FStructOnScope> ParameterValue = MakeShareable(new FStructOnScope(Variable->Variable.GetType().GetStruct()));
+			Variable->CopyDefaultValueDataTo(ParameterValue->GetStructMemory());
+			ParameterEditorLibraryValue->UpdateInternalValueFromStruct(ParameterValue.ToSharedRef());
+		}
+	}
 	else
 	{
 		if (UEdGraphPin* Pin = GetAnyDefaultPin())
