- Allow data interfaces to opt out of running beyond PostActorTick
- This is important for the collision data interface as it relies on the async traces being issues this frame

diff --git a/Engine/Plugins/FX/Niagara/Source/Niagara/Private/NiagaraWorldManager.cpp b/Engine/Plugins/FX/Niagara/Source/Niagara/Private/NiagaraWorldManager.cpp
--- a/Engine/Plugins/FX/Niagara/Source/Niagara/Private/NiagaraWorldManager.cpp
+++ b/Engine/Plugins/FX/Niagara/Source/Niagara/Private/NiagaraWorldManager.cpp
@@ -720,7 +720,7 @@ void FNiagaraWorldManager::PostActorTick(float DeltaSeconds)
 	if (SimulationsWithPostActorWork.Num() > 0)
 	{
 		// Ensure completion of all systems
-		for (int32 i=0; i < SimulationsWithPostActorWork.Num(); ++i )
+		for ( auto& Simulation : SimulationsWithPostActorWork )
 		{
 			FNiagaraSystemSimulationRef& Simulation = SimulationsWithPostActorWork[i];
 			if ( Simulation->IsValid() )
@@ -730,7 +730,7 @@ void FNiagaraWorldManager::PostActorTick(float DeltaSeconds)
 		}
 
 		// Update tick groups
-		for (int32 i=0; i < SimulationsWithPostActorWork.Num(); ++i)
+		for (auto& Simulation : SimulationsWithPostActorWork)
 		{
 			FNiagaraSystemSimulationRef& Simulation = SimulationsWithPostActorWork[i];
 			if (Simulation->IsValid())
@@ -743,7 +743,7 @@ void FNiagaraWorldManager::PostActorTick(float DeltaSeconds)
 		TArray<FNiagaraSystemSimulationRef> LocalSimulationsWithPostActorWork;
 		Swap(SimulationsWithPostActorWork, LocalSimulationsWithPostActorWork);
 		for (FNiagaraSystemSimulationRef& Simulation : LocalSimulationsWithPostActorWork)
-			{
+		{
 			if (Simulation->IsValid())
 			{
 				Simulation->Spawn_GameThread(DeltaSeconds, true);
