Merge from Release-Engine-Test [at] 17462327 to UE5/Main
This represents UE4/Main [REDACTED_HANDLE] and Dev-PerfTest [REDACTED_HANDLE]

diff --git a/Engine/Plugins/FX/Niagara/Source/Niagara/Private/NiagaraWorldManager.cpp b/Engine/Plugins/FX/Niagara/Source/Niagara/Private/NiagaraWorldManager.cpp
--- a/Engine/Plugins/FX/Niagara/Source/Niagara/Private/NiagaraWorldManager.cpp
+++ b/Engine/Plugins/FX/Niagara/Source/Niagara/Private/NiagaraWorldManager.cpp
@@ -720,7 +720,7 @@ void FNiagaraWorldManager::PostActorTick(float DeltaSeconds)
 	if (SimulationsWithPostActorWork.Num() > 0)
 	{
 		// Ensure completion of all systems
-		for ( auto& Simulation : SimulationsWithPostActorWork )
+		for (int32 i = 0; i < SimulationsWithPostActorWork.Num(); ++i)
 		{
 			FNiagaraSystemSimulationRef& Simulation = SimulationsWithPostActorWork[i];
 			if ( Simulation->IsValid() )
@@ -730,7 +730,7 @@ void FNiagaraWorldManager::PostActorTick(float DeltaSeconds)
 		}
 
 		// Update tick groups
-		for (auto& Simulation : SimulationsWithPostActorWork)
+		for (int32 i = 0; i < SimulationsWithPostActorWork.Num(); ++i)
 		{
 			FNiagaraSystemSimulationRef& Simulation = SimulationsWithPostActorWork[i];
 			if (Simulation->IsValid())
