Fixing bad resolve
#rnx

diff --git a/Engine/Plugins/FX/Niagara/Source/NiagaraEditor/Private/ViewModels/NiagaraBakerViewModel.cpp b/Engine/Plugins/FX/Niagara/Source/NiagaraEditor/Private/ViewModels/NiagaraBakerViewModel.cpp
--- a/Engine/Plugins/FX/Niagara/Source/NiagaraEditor/Private/ViewModels/NiagaraBakerViewModel.cpp
+++ b/Engine/Plugins/FX/Niagara/Source/NiagaraEditor/Private/ViewModels/NiagaraBakerViewModel.cpp
@@ -9,7 +9,6 @@
 #include "ViewModels/NiagaraSystemViewModel.h"
 #include "Widgets/SNiagaraBakerWidget.h"
 #include "NiagaraBakerRenderer.h"
-#include "NiagaraGpuComputeDispatchInterface.h"
 #include "NiagaraComponent.h"
 #include "NiagaraSystem.h"
 
@@ -56,7 +55,8 @@ void FNiagaraBakerViewModel::Initialize(TWeakPtr<FNiagaraSystemViewModel> InWeak
 		PreviewComponent->SetAsset(System);
 		PreviewComponent->SetForceSolo(true);
 		PreviewComponent->SetAgeUpdateMode(ENiagaraAgeUpdateMode::DesiredAge);
-		PreviewComponent->SetCanRenderWhileSeeking(false);
+		PreviewComponent->SetCanRenderWhileSeeking(true);
+		PreviewComponent->SetMaxSimTime(0.0f);
 		PreviewComponent->Activate(true);
 
 		AdvancedPreviewScene = MakeShareable(new FAdvancedPreviewScene(FPreviewScene::ConstructionValues()));
@@ -152,10 +152,6 @@ void FNiagaraBakerViewModel::RenderBaker()
 	}
 
 	UWorld* World = PreviewComponent->GetWorld();
-	FSceneInterface* BakerScene = World->Scene;
-
-	FNiagaraGpuComputeDispatchInterface* ComputeDispatchInterface = FNiagaraGpuComputeDispatchInterface::Get(World);
-	ensureMsgf(ComputeDispatchInterface, TEXT("The batcher was not valid on the world this may result in incorrect baking"));
 
 	// Create output render targets & output Baker data
 	TArray<UTextureRenderTarget2D*, TInlineAllocator<4>> RenderTargets;
@@ -182,6 +178,8 @@ void FNiagaraBakerViewModel::RenderBaker()
 	FScopedSlowTask SlowTask(TotalFrames, LOCTEXT("RenderingBaker", "Rendering Baker..."));
 	SlowTask.MakeDialog();
 
+	FNiagaraBakerRenderer BakerRenderer;
+
 	for ( int32 iFrame=0; iFrame < TotalFrames; ++iFrame)
 	{
 		SlowTask.EnterProgressFrame(1);
@@ -192,23 +190,16 @@ void FNiagaraBakerViewModel::RenderBaker()
 		PreviewComponent->SeekToDesiredAge(FrameTime);
 		PreviewComponent->TickComponent(BakerSettings->GetSeekDelta(), ELevelTick::LEVELTICK_All, nullptr);
 
-		// We need any GPU sims to flush pending ticks
-		if (ComputeDispatchInterface)
-		{
-			ComputeDispatchInterface->FlushPendingTicks_GameThread();
-		}
-
 		const float WorldTime = FApp::GetCurrentTime() - BakerSettings->StartSeconds - BakerSettings->DurationSeconds + FrameTime;
 
 		// Render frame
-		FNiagaraBakerRenderer BakerRenderer(PreviewComponent, BakerSettings, WorldTime);
 		for (int32 iOutputTexture=0; iOutputTexture < BakerSettings->OutputTextures.Num(); ++iOutputTexture)
 		{
 			const FNiagaraBakerTextureSettings& OutputTexture = BakerSettings->OutputTextures[iOutputTexture];
 			UTextureRenderTarget2D* RenderTarget = RenderTargets[iOutputTexture];
 			FTextureRenderTargetResource* RenderTargetResource = RenderTarget->GameThread_GetRenderTargetResource();
 
-			BakerRenderer.RenderView(RenderTarget, iOutputTexture);
+			BakerRenderer.RenderView(PreviewComponent, BakerSettings, WorldTime, RenderTarget, iOutputTexture);
 
 			TArray<FFloat16Color> OutSamples;
 			RenderTargetResource->ReadFloat16Pixels(OutSamples);
