- Fix LWC issue

#rnx

diff --git a/Engine/Plugins/FX/Niagara/Source/NiagaraEditor/Private/NiagaraHlslTranslator.cpp b/Engine/Plugins/FX/Niagara/Source/NiagaraEditor/Private/NiagaraHlslTranslator.cpp
--- a/Engine/Plugins/FX/Niagara/Source/NiagaraEditor/Private/NiagaraHlslTranslator.cpp
+++ b/Engine/Plugins/FX/Niagara/Source/NiagaraEditor/Private/NiagaraHlslTranslator.cpp
@@ -5051,7 +5051,7 @@ bool FHlslNiagaraTranslator::GetLiteralConstantVariable(FNiagaraVariable& OutVar
 		FNiagaraVariable ResolvedLocalSpaceCompileOptionVar = ActiveHistoryForFunctionCalls.ResolveAliases(SYS_PARAM_EMITTER_LOCALSPACE);
 		if (CompileOptions.AdditionalDefines.Contains(ResolvedLocalSpaceCompileOptionVar.GetName().ToString()))
 		{
-			OutVar.SetValue(FVector(EForceInit::ForceInitToZero));
+			OutVar.SetValue(FVector3f(EForceInit::ForceInitToZero));
 			return true;
 		}
 	}
