FFloat16 is not currently included in TShaderParameterTypeInfo, and increasing it's array count will not increase its NumRows as seems to be intended.  Simply hardcode correct values.

diff --git a/Engine/Plugins/FX/Niagara/Source/Niagara/Private/NiagaraScriptExecutionParameterStore.cpp b/Engine/Plugins/FX/Niagara/Source/Niagara/Private/NiagaraScriptExecutionParameterStore.cpp
--- a/Engine/Plugins/FX/Niagara/Source/Niagara/Private/NiagaraScriptExecutionParameterStore.cpp
+++ b/Engine/Plugins/FX/Niagara/Source/Niagara/Private/NiagaraScriptExecutionParameterStore.cpp
@@ -107,28 +107,28 @@ uint32 FNiagaraScriptExecutionParameterStore::GenerateLayoutInfoInternal(TArray<
 	}
 	else if (InSrcStruct == FNiagaraTypeDefinition::GetHalfStruct())
 	{
-		uint32 HalfSize = (TShaderParameterTypeInfo<FFloat16>::NumRows * TShaderParameterTypeInfo<FFloat16>::NumColumns) * sizeof(FFloat16);
+		uint32 HalfSize = (/*NumRows*/1 * /*NumColumns*/ 1) * sizeof(FFloat16);
 		Members.Emplace(InSrcOffset, Align(NextMemberOffset, TShaderParameterTypeInfo<FFloat16>::Alignment), HalfSize, HalfSize);
 		InSrcOffset += sizeof(FFloat16);
 		NextMemberOffset = Members[Members.Num() - 1].DestOffset + Members[Members.Num() - 1].DestSize;
 	}
 	else if (InSrcStruct == FNiagaraTypeDefinition::GetHalfVec2Struct())
 	{
-		uint32 StructFinalSize = (TShaderParameterTypeInfo<FFloat16[2]>::NumRows * TShaderParameterTypeInfo<FFloat16>::NumColumns) * sizeof(FFloat16);
+		uint32 StructFinalSize = (/*NumRows*/2 * /*NumColumns*/1) * sizeof(FFloat16);
 		Members.Emplace(InSrcOffset, OffsetAlign(NextMemberOffset, VectorPaddedSize), StructFinalSize, VectorPaddedSize);
 		InSrcOffset += sizeof(FFloat16[2]);
 		NextMemberOffset = Members[Members.Num() - 1].DestOffset + Members[Members.Num() - 1].DestSize;
 	}
 	else if (InSrcStruct == FNiagaraTypeDefinition::GetHalfVec3Struct())
 	{
-		uint32 StructFinalSize = (TShaderParameterTypeInfo<FFloat16[3]>::NumRows * TShaderParameterTypeInfo<FFloat16>::NumColumns) * sizeof(FFloat16);
+		uint32 StructFinalSize = (/*NumRows*/3 * /*NumColumns*/1) * sizeof(FFloat16);
 		Members.Emplace(InSrcOffset, OffsetAlign(NextMemberOffset, VectorPaddedSize), StructFinalSize, VectorPaddedSize);
 		InSrcOffset += sizeof(FFloat16[3]);
 		NextMemberOffset = Members[Members.Num() - 1].DestOffset + Members[Members.Num() - 1].DestSize;
 	}
 	else if (InSrcStruct == FNiagaraTypeDefinition::GetHalfVec4Struct())
 	{
-		uint32 StructFinalSize = (TShaderParameterTypeInfo<FFloat16[4]>::NumRows * TShaderParameterTypeInfo<FFloat16>::NumColumns) * sizeof(FFloat16);
+		uint32 StructFinalSize = (/*NumRows*/4 * /*NumColumns*/1) * sizeof(FFloat16);
 		Members.Emplace(InSrcOffset, OffsetAlign(NextMemberOffset, VectorPaddedSize), StructFinalSize, VectorPaddedSize);
 		InSrcOffset += sizeof(FFloat16[4]);
 		NextMemberOffset = Members[Members.Num() - 1].DestOffset + Members[Members.Num() - 1].DestSize;
