Fixed a bug where pooled components would not respect the visibility of sublevels

diff --git a/Engine/Plugins/FX/Niagara/Source/Niagara/Private/NiagaraComponent.cpp b/Engine/Plugins/FX/Niagara/Source/Niagara/Private/NiagaraComponent.cpp
--- a/Engine/Plugins/FX/Niagara/Source/Niagara/Private/NiagaraComponent.cpp
+++ b/Engine/Plugins/FX/Niagara/Source/Niagara/Private/NiagaraComponent.cpp
@@ -1921,6 +1921,17 @@ void UNiagaraComponent::OnChildDetached(USceneComponent* ChildComponent)
 
 }
 
+bool UNiagaraComponent::IsVisible() const
+{
+	if (PoolingMethod == ENCPoolMethod::None || GetAttachParent() == nullptr)
+	{
+		return Super::IsVisible();
+	}
+
+	// pooled components are owned by the persistent level, so we need to check their current attached parent visibility, otherwise they ignore sublevel visibility
+	return Super::IsVisible() && GetAttachParent()->IsVisible();
+}
+
 void UNiagaraComponent::SetTickBehavior(ENiagaraTickBehavior NewTickBehavior)
 {
 	TickBehavior = NewTickBehavior;
diff --git a/Engine/Plugins/FX/Niagara/Source/Niagara/Public/NiagaraComponent.h b/Engine/Plugins/FX/Niagara/Source/Niagara/Public/NiagaraComponent.h
--- a/Engine/Plugins/FX/Niagara/Source/Niagara/Public/NiagaraComponent.h
+++ b/Engine/Plugins/FX/Niagara/Source/Niagara/Public/NiagaraComponent.h
@@ -221,6 +221,7 @@ public:
 	//~ Begin USceneComponent Interface
 	virtual void OnChildAttached(USceneComponent* ChildComponent) override;
 	virtual void OnChildDetached(USceneComponent* ChildComponent) override;
+	virtual bool IsVisible() const override;
 	//~ Begin USceneComponent Interface
 
 	UE_DEPRECATED(5.0, "This interface is no longer safe to access directly. Use the interface provided by GetSystemInstanceController instead.")
