- Modify some variable names to fix reserved keyword issues
- Resave content with adjusted version

#rnx

diff --git a/Engine/Plugins/FX/Niagara/Shaders/Private/NiagaraDataInterfaceStaticMeshTemplate.ush b/Engine/Plugins/FX/Niagara/Shaders/Private/NiagaraDataInterfaceStaticMeshTemplate.ush
--- a/Engine/Plugins/FX/Niagara/Shaders/Private/NiagaraDataInterfaceStaticMeshTemplate.ush
+++ b/Engine/Plugins/FX/Niagara/Shaders/Private/NiagaraDataInterfaceStaticMeshTemplate.ush
@@ -35,15 +35,15 @@ float3			InstanceWorldVelocity_{ParameterName};
 void RandomSection_{ParameterName}(NiagaraRandInfo RandInfo, out int Section);
 void RandomFilteredSection_{ParameterName}(NiagaraRandInfo RandInfo, out int Section);
 void RandomUnfilteredSection_{ParameterName}(NiagaraRandInfo RandInfo, out int Section);
-void RandomSectionTriangle_{ParameterName}(NiagaraRandInfo RandInfo, int Section, out int Triangle, out float3 BaryCoord);
+void RandomSectionTriangle_{ParameterName}(NiagaraRandInfo RandInfo, int Section, out int Tri, out float3 BaryCoord);
 
 //////////////////////////////////////////////////////////////////////////
-int UniformTriangle_{ParameterName}(NiagaraRandInfo RandInfo, int Triangle, int SectionOffset)
+int UniformTriangle_{ParameterName}(NiagaraRandInfo RandInfo, int Tri, int SectionOffset)
 {
-	uint2 ProbAlias = UniformSamplingTriangles_{ParameterName}[Triangle];
+	uint2 ProbAlias = UniformSamplingTriangles_{ParameterName}[Tri];
 	float Prob = asfloat(ProbAlias.x);
 	int Alias = ProbAlias.y;
-	return NiagaraRandomFloat(RandInfo) < Prob ? Triangle : Alias + SectionOffset;
+	return NiagaraRandomFloat(RandInfo) < Prob ? Tri : Alias + SectionOffset;
 }
 
 //////////////////////////////////////////////////////////////////////////
@@ -120,16 +120,16 @@ void GetVertexUV_{ParameterName}(int Vertex, int UVSet, out float2 UV)
 
 //////////////////////////////////////////////////////////////////////////
 // VM Triangle Sampling
-void IsValidTriangle_{ParameterName}(int Triangle, out bool IsValid)
+void IsValidTriangle_{ParameterName}(int Tri, out bool IsValid)
 {
-	IsValid = Triangle >= 0 && Triangle < NumTriangles_{ParameterName}.x;
+	IsValid = Tri >= 0 && Tri < NumTriangles_{ParameterName}.x;
 }
 
-void RandomTriangle_{ParameterName}(NiagaraRandInfo RandInfo, out int Triangle, out float3 BaryCoord)
+void RandomTriangle_{ParameterName}(NiagaraRandInfo RandInfo, out int Tri, out float3 BaryCoord)
 {
 	int Section;
 	RandomSection_{ParameterName}(RandInfo, Section);
-	RandomSectionTriangle_{ParameterName}(RandInfo, Section, Triangle, BaryCoord);
+	RandomSectionTriangle_{ParameterName}(RandInfo, Section, Tri, BaryCoord);
 }
 
 void GetTriangleCount_{ParameterName}(out int Count)
@@ -137,22 +137,22 @@ void GetTriangleCount_{ParameterName}(out int Count)
 	Count = NumTriangles_{ParameterName}.x;
 }
 
-void IsValidFilteredTriangle_{ParameterName}(int Triangle, out bool IsValid)
+void IsValidFilteredTriangle_{ParameterName}(int Tri, out bool IsValid)
 {
-	IsValid = Triangle >= 0 && Triangle < NumTriangles_{ParameterName}.y;
+	IsValid = Tri >= 0 && Tri < NumTriangles_{ParameterName}.y;
 }
 
-void RandomFilteredTriangle_{ParameterName}(NiagaraRandInfo RandInfo, out int Triangle, out float3 BaryCoord)
+void RandomFilteredTriangle_{ParameterName}(NiagaraRandInfo RandInfo, out int Tri, out float3 BaryCoord)
 {
 	if ( NumTriangles_{ParameterName}.y > 0 )
 	{
 		int Section;
 		RandomFilteredSection_{ParameterName}(RandInfo, Section);
-		RandomSectionTriangle_{ParameterName}(RandInfo, Section, Triangle, BaryCoord);
+		RandomSectionTriangle_{ParameterName}(RandInfo, Section, Tri, BaryCoord);
 	}
 	else
 	{
-		Triangle = 0;
+		Tri = 0;
 		BaryCoord = NiagaraRandomBaryCoord(RandInfo);
 	}
 }
@@ -162,9 +162,9 @@ void GetFilteredTriangleCount_{ParameterName}(out int Count)
 	Count = NumTriangles_{ParameterName}.y;
 }
 
-void GetFilteredTriangleAt_{ParameterName}(int TriangleIndex, out int Triangle)
+void GetFilteredTriangleAt_{ParameterName}(int TriangleIndex, out int Tri)
 {
-	Triangle = 0;
+	Tri = 0;
 
 	int SectionInfoOffset = SectionCounts_{ParameterName}.x;
 	int SectionFirstTriangle = 0;
@@ -173,29 +173,29 @@ void GetFilteredTriangleAt_{ParameterName}(int TriangleIndex, out int Triangle)
 		int SectionLastTriangle = SectionFirstTriangle + SectionInfos_{ParameterName}[SectionInfoOffset + i].y;
 		if ( TriangleIndex >= SectionFirstTriangle && TriangleIndex < SectionLastTriangle )
 		{
-			Triangle = (TriangleIndex - SectionFirstTriangle) + SectionInfos_{ParameterName}[SectionInfoOffset + i].x;
+			Tri = (TriangleIndex - SectionFirstTriangle) + SectionInfos_{ParameterName}[SectionInfoOffset + i].x;
 			return;
 		}
 		SectionFirstTriangle = SectionLastTriangle;
 	}
 }
 
-void IsValidUnfilteredTriangle_{ParameterName}(int Triangle, out bool IsValid)
+void IsValidUnfilteredTriangle_{ParameterName}(int Tri, out bool IsValid)
 {
-	IsValid = Triangle >= 0 && Triangle < NumTriangles_{ParameterName}.z;
+	IsValid = Tri >= 0 && Tri < NumTriangles_{ParameterName}.z;
 }
 
-void RandomUnfilteredTriangle_{ParameterName}(NiagaraRandInfo RandInfo, out int Triangle, out float3 BaryCoord)
+void RandomUnfilteredTriangle_{ParameterName}(NiagaraRandInfo RandInfo, out int Tri, out float3 BaryCoord)
 {
 	if ( NumTriangles_{ParameterName}.z > 0 )
 	{
 		int Section;
 		RandomUnfilteredSection_{ParameterName}(RandInfo, Section);
-		RandomSectionTriangle_{ParameterName}(RandInfo, Section, Triangle, BaryCoord);
+		RandomSectionTriangle_{ParameterName}(RandInfo, Section, Tri, BaryCoord);
 	}
 	else
 	{
-		Triangle = 0;
+		Tri = 0;
 		BaryCoord = NiagaraRandomBaryCoord(RandInfo);
 	}
 }
@@ -205,9 +205,9 @@ void GetUnfilteredTriangleCount_{ParameterName}(out int Count)
 	Count = NumTriangles_{ParameterName}.z;
 }
 
-void GetUnfilteredTriangleAt_{ParameterName}(int TriangleIndex, out int Triangle)
+void GetUnfilteredTriangleAt_{ParameterName}(int TriangleIndex, out int Tri)
 {
-	Triangle = 0;
+	Tri = 0;
 
 	int SectionInfoOffset = SectionCounts_{ParameterName}.x + SectionCounts_{ParameterName}.y;
 	int SectionFirstTriangle = 0;
@@ -216,22 +216,22 @@ void GetUnfilteredTriangleAt_{ParameterName}(int TriangleIndex, out int Triangle
 		int SectionLastTriangle = SectionFirstTriangle + SectionInfos_{ParameterName}[SectionInfoOffset + i].y;
 		if ( TriangleIndex >= SectionFirstTriangle && TriangleIndex < SectionLastTriangle )
 		{
-			Triangle = (TriangleIndex - SectionFirstTriangle) + SectionInfos_{ParameterName}[SectionInfoOffset + i].x;
+			Tri = (TriangleIndex - SectionFirstTriangle) + SectionInfos_{ParameterName}[SectionInfoOffset + i].x;
 			return;
 		}
 		SectionFirstTriangle = SectionLastTriangle;
 	}
 }
 
-void GetTriangleIndices_{ParameterName}(int Triangle, out int Index0, out int Index1, out int Index2)
+void GetTriangleIndices_{ParameterName}(int Tri, out int Index0, out int Index1, out int Index2)
 {
 	[branch]
 	if ( NumTriangles_{ParameterName}.x > 0 )
 	{
-		Triangle = clamp(Triangle, 0, NumTriangles_{ParameterName}.x - 1);
-		Index0 = IndexBuffer_{ParameterName}[Triangle * 3 + 0];
-		Index1 = IndexBuffer_{ParameterName}[Triangle * 3 + 1];
-		Index2 = IndexBuffer_{ParameterName}[Triangle * 3 + 2];
+		Tri = clamp(Tri, 0, NumTriangles_{ParameterName}.x - 1);
+		Index0 = IndexBuffer_{ParameterName}[Tri * 3 + 0];
+		Index1 = IndexBuffer_{ParameterName}[Tri * 3 + 1];
+		Index2 = IndexBuffer_{ParameterName}[Tri * 3 + 2];
 	}
 	else
 	{
@@ -241,10 +241,10 @@ void GetTriangleIndices_{ParameterName}(int Triangle, out int Index0, out int In
 	}
 }
 
-void GetTriangle_{ParameterName}(int Triangle, float3 BaryCoord, out float3 Position, out float3 Velocity, out float3 Normal, out float3 Bitangent, out float3 Tangent)
+void GetTriangle_{ParameterName}(int Tri, float3 BaryCoord, out float3 Position, out float3 Velocity, out float3 Normal, out float3 Bitangent, out float3 Tangent)
 {
 	int Index[3];
-	GetTriangleIndices_{ParameterName}(Triangle, Index[0], Index[1], Index[2]);
+	GetTriangleIndices_{ParameterName}(Tri, Index[0], Index[1], Index[2]);
 
 	float3 Positions[3], Velocities[3], Normals[3], Bitangents[3], Tangents[3];
 	GetVertex_{ParameterName}(Index[0], Positions[0], Velocities[0], Normals[0], Bitangents[0], Tangents[0]);
@@ -258,10 +258,10 @@ void GetTriangle_{ParameterName}(int Triangle, float3 BaryCoord, out float3 Posi
 	Tangent		= Tangents[0] * BaryCoord.x + Tangents[1] * BaryCoord.y + Tangents[2] * BaryCoord.x;
 }
 
-void GetTriangleWS_{ParameterName}(int Triangle, float3 BaryCoord, out float3 Position, out float3 Velocity, out float3 Normal, out float3 Bitangent, out float3 Tangent)
+void GetTriangleWS_{ParameterName}(int Tri, float3 BaryCoord, out float3 Position, out float3 Velocity, out float3 Normal, out float3 Bitangent, out float3 Tangent)
 {
 	float3 LocalPosition;
-	GetTriangle_{ParameterName}(Triangle, BaryCoord, LocalPosition, Velocity, Normal, Bitangent, Tangent);
+	GetTriangle_{ParameterName}(Tri, BaryCoord, LocalPosition, Velocity, Normal, Bitangent, Tangent);
 
 	Position				= mul(float4(LocalPosition, 1), InstanceTransform_{ParameterName}).xyz;
 	float3 PreviousPosition	= mul(float4(LocalPosition, 1), InstancePreviousTransform_{ParameterName}).xyz;
@@ -271,10 +271,10 @@ void GetTriangleWS_{ParameterName}(int Triangle, float3 BaryCoord, out float3 Po
 	Tangent					= normalize(mul(float4(Tangent, 0), InstanceTransformInverseTransposed_{ParameterName}).xyz);
 }
 
-void GetTriangleColor_{ParameterName}(int Triangle, float3 BaryCoord, out float4 Color)
+void GetTriangleColor_{ParameterName}(int Tri, float3 BaryCoord, out float4 Color)
 {
 	int Index[3];
-	GetTriangleIndices_{ParameterName}(Triangle, Index[0], Index[1], Index[2]);
+	GetTriangleIndices_{ParameterName}(Tri, Index[0], Index[1], Index[2]);
 
 	float4 Colors[3];
 	GetVertexColor_{ParameterName}(Index[0], Colors[0]);
@@ -284,10 +284,10 @@ void GetTriangleColor_{ParameterName}(int Triangle, float3 BaryCoord, out float4
 	Color = Colors[0] * BaryCoord.x + Colors[1] * BaryCoord.y + Colors[2] * BaryCoord.z;
 }
 
-void GetTriangleUV_{ParameterName}(int Triangle, float3 BaryCoord, int UVSet, out float2 UV)
+void GetTriangleUV_{ParameterName}(int Tri, float3 BaryCoord, int UVSet, out float2 UV)
 {
 	int Index[3];
-	GetTriangleIndices_{ParameterName}(Triangle, Index[0], Index[1], Index[2]);
+	GetTriangleIndices_{ParameterName}(Tri, Index[0], Index[1], Index[2]);
 
 	float2 UVs[3];
 	GetVertexUV_{ParameterName}(Index[0], UVSet, UVs[0]);
@@ -401,7 +401,7 @@ void GetSectionTriangleCount_{ParameterName}(int Section, out int Count)
 	Count = Section >= 0 && Section < SectionCounts_{ParameterName}.x ? SectionInfos_{ParameterName}[Section].y : 0;
 }
 
-void RandomSectionTriangle_{ParameterName}(NiagaraRandInfo RandInfo, int Section, out int Triangle, out float3 BaryCoord)
+void RandomSectionTriangle_{ParameterName}(NiagaraRandInfo RandInfo, int Section, out int Tri, out float3 BaryCoord)
 {
 	int SectionTriangleCount;
 	GetSectionTriangleCount_{ParameterName}(Section, SectionTriangleCount);
@@ -409,18 +409,18 @@ void RandomSectionTriangle_{ParameterName}(NiagaraRandInfo RandInfo, int Section
 	if (SectionTriangleCount > 0)
 	{
 		int SectionTriangleOffset = SectionInfos_{ParameterName}[Section].x;
-		Triangle = NiagaraRandomInt(RandInfo, SectionTriangleCount) + SectionTriangleOffset;
+		Tri = NiagaraRandomInt(RandInfo, SectionTriangleCount) + SectionTriangleOffset;
 
 		[branch]
 		if ( HasUniformSampling_{ParameterName} )
 		{
-			Triangle = UniformTriangle_{ParameterName}(RandInfo, Triangle, SectionTriangleOffset);
+			Tri = UniformTriangle_{ParameterName}(RandInfo, Tri, SectionTriangleOffset);
 		}
 	}
 	BaryCoord = NiagaraRandomBaryCoord(RandInfo);
 }
 
-void GetSectionTriangleAt_{ParameterName}(int Section, int SectionTriangleIndex, out int Triangle)
+void GetSectionTriangleAt_{ParameterName}(int Section, int SectionTriangleIndex, out int Tri)
 {
 	int SectionTriangleCount;
 	GetSectionTriangleCount_{ParameterName}(Section, SectionTriangleCount);
@@ -429,11 +429,11 @@ void GetSectionTriangleAt_{ParameterName}(int Section, int SectionTriangleIndex,
 	{
 		int SectionTriangleMax = SectionInfos_{ParameterName}[Section].y - 1;
 		int SectionTriangleOffset = SectionInfos_{ParameterName}[Section].x;
-		Triangle = clamp(SectionTriangleIndex, 0, SectionTriangleMax) + SectionTriangleOffset;
+		Tri = clamp(SectionTriangleIndex, 0, SectionTriangleMax) + SectionTriangleOffset;
 	}
 	else
 	{
-		Triangle = 0;
+		Tri = 0;
 	}
 }
 
@@ -543,14 +543,14 @@ void IsValid_{ParameterName}(out bool IsValid)
 	IsValid = NumVertices_{ParameterName} > 0;
 }
 
-void GetLocalToWorld_{ParameterName}(out float4x4 Matrix)
+void GetLocalToWorld_{ParameterName}(out float4x4 OutMatrix)
 {
-	Matrix = InstanceTransform_{ParameterName};
+	OutMatrix = InstanceTransform_{ParameterName};
 }
 
-void GetLocalToWorldInverseTransposed_{ParameterName}(out float4x4 Matrix)
+void GetLocalToWorldInverseTransposed_{ParameterName}(out float4x4 OutMatrix)
 {
-	Matrix = InstanceTransformInverseTransposed_{ParameterName};
+	OutMatrix = InstanceTransformInverseTransposed_{ParameterName};
 }
 
 void GetWorldVelocity_{ParameterName}(out float3 Velocity)
