Variable  max transforms for rigid mesh collisions

diff --git a/Engine/Plugins/FX/Niagara/Shaders/Private/NiagaraDataInterfaceRigidMeshCollisionQuery.ush b/Engine/Plugins/FX/Niagara/Shaders/Private/NiagaraDataInterfaceRigidMeshCollisionQuery.ush
--- a/Engine/Plugins/FX/Niagara/Shaders/Private/NiagaraDataInterfaceRigidMeshCollisionQuery.ush
+++ b/Engine/Plugins/FX/Niagara/Shaders/Private/NiagaraDataInterfaceRigidMeshCollisionQuery.ush
@@ -11,12 +11,7 @@ NiagaraDataInterfaceRigidMeshCollisionQuery.ush
  * -----------------------------------------------------------------
  */
 
- #define RIGID_MESH_COLLISION_QUERY_MAX_TRANSFORMS 200
- #define SMALL_NUMBER 1e-8
-
- #define CURRENT_OFFSET 0
- #define PREVIOUS_OFFSET RIGID_MESH_COLLISION_QUERY_MAX_TRANSFORMS
-
+#define SMALL_NUMBER 1e-8
 #define MAX_DISTANCE 1e+38
 #define BOX_INDEX 0
 #define SPHERE_INDEX 1
@@ -24,6 +19,9 @@ NiagaraDataInterfaceRigidMeshCollisionQuery.ush
 
 struct FDIRigidMeshCollisionContext
 {
+	uint MaxTransforms;
+	uint CurrentOffset;
+	uint PreviousOffset;
 	uint4 ElementOffsets;
 	Buffer<float4>		WorldTransformBuffer;
 	Buffer<float4>		InverseTransformBuffer;
@@ -33,6 +31,9 @@ struct FDIRigidMeshCollisionContext
 }; 
 
 #define DIRIGIDMESHCOLLISIONQUERY_DECLARE_CONSTANTS(NAME)\
+uint MaxTransforms_##NAME;\
+uint CurrentOffset_##NAME;\
+uint PreviousOffset_##NAME;\
 uint4 ElementOffsets_##NAME;\
 Buffer<float4>		WorldTransformBuffer_##NAME;\
 Buffer<float4>		InverseTransformBuffer_##NAME;\
@@ -43,6 +44,9 @@ Buffer<uint>		DFIndexBuffer_##NAME;\
 
 #define DIRIGIDMESHCOLLISIONQUERY_MAKE_CONTEXT(NAME)\
 FDIRigidMeshCollisionContext DIContext; \
+DIContext.MaxTransforms = MaxTransforms_##NAME;\
+DIContext.CurrentOffset =  CurrentOffset_##NAME;\
+DIContext.PreviousOffset = PreviousOffset_##NAME;\
 DIContext.ElementOffsets = ElementOffsets_##NAME;\
 DIContext.WorldTransformBuffer = WorldTransformBuffer_##NAME;\
 DIContext.InverseTransformBuffer = InverseTransformBuffer_##NAME;\
@@ -79,37 +83,37 @@ int DIRigidMeshCollision_GetNumCapsules(in FDIRigidMeshCollisionContext DIContex
 
 float3 DIRigidMeshCollision_GetCurrentPosition(in FDIRigidMeshCollisionContext DIContext, in int ElementIndex)
 {
-	const int BufferOffset = 3 * ElementIndex + CURRENT_OFFSET;
+	const int BufferOffset = 3 * ElementIndex + DIContext.CurrentOffset;
     return float3(DIContext.WorldTransformBuffer[BufferOffset].w, DIContext.WorldTransformBuffer[BufferOffset+1].w, DIContext.WorldTransformBuffer[BufferOffset+2].w);
 }
 
 float3 DIRigidMeshCollision_GetPreviousPosition(in FDIRigidMeshCollisionContext DIContext, in int ElementIndex)
 {
-	const int BufferOffset = 3 * ElementIndex + PREVIOUS_OFFSET;
+	const int BufferOffset = 3 * ElementIndex + DIContext.PreviousOffset;
     return float3(DIContext.WorldTransformBuffer[BufferOffset].w, DIContext.WorldTransformBuffer[BufferOffset+1].w, DIContext.WorldTransformBuffer[BufferOffset+2].w);
 }
 
 float3x4 DIRigidMeshCollision_GetCurrentTransform(in FDIRigidMeshCollisionContext DIContext, in int ElementIndex)
 {
-	const int BufferOffset = 3 * ElementIndex + CURRENT_OFFSET;
+	const int BufferOffset = 3 * ElementIndex + DIContext.CurrentOffset;
     return float3x4(DIContext.WorldTransformBuffer[BufferOffset], DIContext.WorldTransformBuffer[BufferOffset+1], DIContext.WorldTransformBuffer[BufferOffset+2]);
 }
 
 float3x4 DIRigidMeshCollision_GetInverseTransform(in FDIRigidMeshCollisionContext DIContext, in int ElementIndex)
 {
-	const int BufferOffset = 3 * ElementIndex + CURRENT_OFFSET;
+	const int BufferOffset = 3 * ElementIndex + DIContext.CurrentOffset;
     return float3x4(DIContext.InverseTransformBuffer[BufferOffset], DIContext.InverseTransformBuffer[BufferOffset+1], DIContext.InverseTransformBuffer[BufferOffset+2]);
 }
 
 float3x4 DIRigidMeshCollision_GetPreviousTransform(in FDIRigidMeshCollisionContext DIContext, in int ElementIndex)
 {
-	const int BufferOffset = 3 * ElementIndex + PREVIOUS_OFFSET;
+	const int BufferOffset = 3 * ElementIndex + DIContext.PreviousOffset;
     return float3x4(DIContext.WorldTransformBuffer[BufferOffset], DIContext.WorldTransformBuffer[BufferOffset+1], DIContext.WorldTransformBuffer[BufferOffset+2]);
 }
 
 float3x4 DIRigidMeshCollision_GetPreviousInverse(in FDIRigidMeshCollisionContext DIContext, in int ElementIndex)
 {
-	const int BufferOffset = 3 * ElementIndex + PREVIOUS_OFFSET;
+	const int BufferOffset = 3 * ElementIndex + DIContext.PreviousOffset;
     return float3x4(DIContext.InverseTransformBuffer[BufferOffset], DIContext.InverseTransformBuffer[BufferOffset+1], DIContext.InverseTransformBuffer[BufferOffset+2]);
 }
 
diff --git a/Engine/Plugins/FX/Niagara/Source/Niagara/Private/NiagaraDataInterfaceRigidMeshCollisionQuery.cpp b/Engine/Plugins/FX/Niagara/Source/Niagara/Private/NiagaraDataInterfaceRigidMeshCollisionQuery.cpp
--- a/Engine/Plugins/FX/Niagara/Source/Niagara/Private/NiagaraDataInterfaceRigidMeshCollisionQuery.cpp
+++ b/Engine/Plugins/FX/Niagara/Source/Niagara/Private/NiagaraDataInterfaceRigidMeshCollisionQuery.cpp
@@ -41,6 +41,10 @@ static const FName GetClosestPointMeshDistanceFieldName(TEXT("GetClosestPointMes
 
 //------------------------------------------------------------------------------------------------------------
 
+const FString UNiagaraDataInterfaceRigidMeshCollisionQuery::MaxTransformsName(TEXT("MaxTransforms_"));
+const FString UNiagaraDataInterfaceRigidMeshCollisionQuery::CurrentOffsetName(TEXT("CurrentOffset_"));
+const FString UNiagaraDataInterfaceRigidMeshCollisionQuery::PreviousOffsetName(TEXT("PreviousOffset_"));
+
 const FString UNiagaraDataInterfaceRigidMeshCollisionQuery::ElementOffsetsName(TEXT("ElementOffsets_"));
 
 const FString UNiagaraDataInterfaceRigidMeshCollisionQuery::WorldTransformBufferName(TEXT("WorldTransformBuffer_"));
@@ -61,6 +65,10 @@ struct FNDIRigidMeshCollisionParametersName
 {
 	FNDIRigidMeshCollisionParametersName(const FString& Suffix)
 	{
+		MaxTransformsName = UNiagaraDataInterfaceRigidMeshCollisionQuery::MaxTransformsName + Suffix;
+		CurrentOffsetName = UNiagaraDataInterfaceRigidMeshCollisionQuery::CurrentOffsetName + Suffix;
+		PreviousOffsetName = UNiagaraDataInterfaceRigidMeshCollisionQuery::PreviousOffsetName + Suffix;
+
 		ElementOffsetsName = UNiagaraDataInterfaceRigidMeshCollisionQuery::ElementOffsetsName + Suffix;
 
 		WorldTransformBufferName = UNiagaraDataInterfaceRigidMeshCollisionQuery::WorldTransformBufferName + Suffix;
@@ -70,6 +78,10 @@ struct FNDIRigidMeshCollisionParametersName
 		DFIndexBufferName = UNiagaraDataInterfaceRigidMeshCollisionQuery::DFIndexBufferName + Suffix;
 	}
 
+	FString MaxTransformsName;
+	FString CurrentOffsetName;
+	FString PreviousOffsetName;
+
 	FString ElementOffsetsName;
 
 	FString WorldTransformBufferName;
@@ -81,21 +93,22 @@ struct FNDIRigidMeshCollisionParametersName
 
 //------------------------------------------------------------------------------------------------------------
 
-template<typename BufferType, EPixelFormat PixelFormat, uint32 ElementCount, uint32 BufferCount = 1>
-void CreateInternalBuffer(FRWBuffer& OutputBuffer)
+template<typename BufferType, EPixelFormat PixelFormat>
+void CreateInternalBuffer(FRWBuffer& OutputBuffer, uint32 ElementCount)
 {
 	if (ElementCount > 0)
 	{
-		OutputBuffer.Initialize(TEXT("FNDIRigidMeshCollisionBuffer"), sizeof(BufferType), ElementCount * BufferCount, PixelFormat, BUF_Static);
+		OutputBuffer.Initialize(TEXT("FNDIRigidMeshCollisionBuffer"), sizeof(BufferType), ElementCount, PixelFormat, BUF_Static);
 	}
 }
 
-template<typename BufferType, EPixelFormat PixelFormat, uint32 ElementCount, uint32 BufferCount = 1>
-void UpdateInternalBuffer(const TStaticArray<BufferType,ElementCount*BufferCount>& InputData, FRWBuffer& OutputBuffer)
+template<typename BufferType, EPixelFormat PixelFormat>
+void UpdateInternalBuffer(const TArray<BufferType>& InputData, FRWBuffer& OutputBuffer)
 {
+	uint32 ElementCount = InputData.Num();
 	if (ElementCount > 0 && OutputBuffer.Buffer.IsValid())
 	{
-		const uint32 BufferBytes = sizeof(BufferType) * ElementCount * BufferCount;
+		const uint32 BufferBytes = sizeof(BufferType) * ElementCount;
 
 		void* OutputData = RHILockBuffer(OutputBuffer.Buffer, 0, BufferBytes, RLM_WriteOnly);
 
@@ -105,7 +118,7 @@ void UpdateInternalBuffer(const TStaticArray<BufferType,ElementCount*BufferCount
 }
 
 void FillCurrentTransforms(const FTransform& ElementTransform, uint32& ElementCount,
-	TStaticArray<FVector4f,RIGID_MESH_COLLISION_QUERY_MAX_TRANSFORMS>& OutCurrentTransform, TStaticArray<FVector4f, RIGID_MESH_COLLISION_QUERY_MAX_TRANSFORMS>& OutCurrentInverse)
+	TArray<FVector4f>& OutCurrentTransform, TArray<FVector4f>& OutCurrentInverse)
 {
 	// LWC_TODO: precision loss
 	const uint32 ElementOffset = 3 * ElementCount;
@@ -158,13 +171,13 @@ void GetNumPrimitives(TArray<AStaticMeshActor*> StaticMeshActors, uint32& NumBox
 
 void CompactInternalArrays(FNDIRigidMeshCollisionArrays* OutAssetArrays)
 {
-	for (uint32 TransformIndex = 0; TransformIndex < RIGID_MESH_COLLISION_QUERY_MAX_TRANSFORMS; ++TransformIndex)
+	for (uint32 TransformIndex = 0; TransformIndex < OutAssetArrays->MaxTransforms; ++TransformIndex)
 	{
 		uint32 OffsetIndex = TransformIndex;
 		OutAssetArrays->WorldTransform[OffsetIndex] = OutAssetArrays->CurrentTransform[TransformIndex];
 		OutAssetArrays->InverseTransform[OffsetIndex] = OutAssetArrays->CurrentInverse[TransformIndex];
 
-		OffsetIndex += RIGID_MESH_COLLISION_QUERY_MAX_TRANSFORMS;
+		OffsetIndex += OutAssetArrays->MaxTransforms;
 		OutAssetArrays->WorldTransform[OffsetIndex] = OutAssetArrays->PreviousTransform[TransformIndex];
 		OutAssetArrays->InverseTransform[OffsetIndex] = OutAssetArrays->PreviousInverse[TransformIndex];		
 	}
@@ -186,7 +199,7 @@ void CreateInternalArrays(TArray<AStaticMeshActor*> StaticMeshActors,
 
 		GetNumPrimitives(StaticMeshActors, NumBoxes, NumSpheres, NumCapsules);
 		
-		if ((NumBoxes + NumSpheres + NumCapsules) < RIGID_MESH_COLLISION_QUERY_MAX_PRIMITIVES)
+		if ((NumBoxes + NumSpheres + NumCapsules) < OutAssetArrays->MaxPrimitives)
 		{
 			OutAssetArrays->ElementOffsets.BoxOffset = 0;
 			OutAssetArrays->ElementOffsets.SphereOffset = OutAssetArrays->ElementOffsets.BoxOffset + NumBoxes;
@@ -264,18 +277,18 @@ void CreateInternalArrays(TArray<AStaticMeshActor*> StaticMeshActors,
 		}
 		else
 		{
-			UE_LOG(LogRigidMeshCollision, Warning, TEXT("Number of Collision DI primitives is higher than the niagara %d limit"), RIGID_MESH_COLLISION_QUERY_MAX_PRIMITIVES);
+			UE_LOG(LogRigidMeshCollision, Warning, TEXT("Number of Collision DI primitives is higher than the %d limit.  Please increase it."), OutAssetArrays->MaxPrimitives);
 		}
 	}
 }
 
 void UpdateInternalArrays(const TArray<AStaticMeshActor*> &StaticMeshActors, FNDIRigidMeshCollisionArrays* OutAssetArrays)
 {
-	if (OutAssetArrays != nullptr && OutAssetArrays->ElementOffsets.NumElements < RIGID_MESH_COLLISION_QUERY_MAX_PRIMITIVES)
+	if (OutAssetArrays != nullptr && OutAssetArrays->ElementOffsets.NumElements < OutAssetArrays->MaxPrimitives)
 	{
 		uint32 NumBoxes = 0;
 		uint32 NumSpheres = 0;
-		uint32 NumCapsules = 0;
+		uint32 NumCapsules = 0;		
 
 		GetNumPrimitives(StaticMeshActors, NumBoxes, NumSpheres, NumCapsules);
 
@@ -342,12 +355,12 @@ void UpdateInternalArrays(const TArray<AStaticMeshActor*> &StaticMeshActors, FND
 
 void FNDIRigidMeshCollisionBuffer::InitRHI()
 {
-	CreateInternalBuffer<FVector4f, EPixelFormat::PF_A32B32G32R32F, RIGID_MESH_COLLISION_QUERY_MAX_TRANSFORMS, 2>(WorldTransformBuffer);
-	CreateInternalBuffer<FVector4f, EPixelFormat::PF_A32B32G32R32F, RIGID_MESH_COLLISION_QUERY_MAX_TRANSFORMS, 2>(InverseTransformBuffer);
+	CreateInternalBuffer<FVector4f, EPixelFormat::PF_A32B32G32R32F>(WorldTransformBuffer, 3 * MaxNumTransforms);
+	CreateInternalBuffer<FVector4f, EPixelFormat::PF_A32B32G32R32F>(InverseTransformBuffer, 3 * MaxNumTransforms);
 
-	CreateInternalBuffer<FVector4f, EPixelFormat::PF_A32B32G32R32F, RIGID_MESH_COLLISION_QUERY_MAX_PRIMITIVES>(ElementExtentBuffer);
-	CreateInternalBuffer<uint32, EPixelFormat::PF_R32_UINT, RIGID_MESH_COLLISION_QUERY_MAX_PRIMITIVES>(PhysicsTypeBuffer);
-	CreateInternalBuffer<uint32, EPixelFormat::PF_R32_UINT, RIGID_MESH_COLLISION_QUERY_MAX_PRIMITIVES>(DFIndexBuffer);
+	CreateInternalBuffer<FVector4f, EPixelFormat::PF_A32B32G32R32F>(ElementExtentBuffer, MaxNumPrimitives);
+	CreateInternalBuffer<uint32, EPixelFormat::PF_R32_UINT>(PhysicsTypeBuffer, MaxNumPrimitives);
+	CreateInternalBuffer<uint32, EPixelFormat::PF_R32_UINT>(DFIndexBuffer, MaxNumPrimitives);
 }
 
 void FNDIRigidMeshCollisionBuffer::ReleaseRHI()
@@ -406,8 +419,11 @@ void FNDIRigidMeshCollisionData::Init(UNiagaraDataInterfaceRigidMeshCollisionQue
 			// @note: we are not creating internal arrays here and letting it happen on the call to update
 			 //CreateInternalArrays(StaticMeshActors, &AssetArrays);
 		}
+		AssetArrays = new FNDIRigidMeshCollisionArrays();
+		AssetArrays->Resize(Interface->MaxNumPrimitives);
 
 		AssetBuffer = new FNDIRigidMeshCollisionBuffer();
+		AssetBuffer->SetMaxNumPrimitives(Interface->MaxNumPrimitives);		
 		BeginInitResource(AssetBuffer);
 	}
 }
@@ -420,7 +436,7 @@ void FNDIRigidMeshCollisionData::Update(UNiagaraDataInterfaceRigidMeshCollisionQ
 
 		if (0 < StaticMeshActors.Num() && StaticMeshActors[0] != nullptr)
 		{
-			UpdateInternalArrays(StaticMeshActors, &AssetArrays);
+			UpdateInternalArrays(StaticMeshActors, AssetArrays);
 		}
 	}
 }
@@ -559,6 +575,10 @@ public:
 	{
 		FNDIRigidMeshCollisionParametersName ParamNames(*ParameterInfo.DataInterfaceHLSLSymbol);
 
+		MaxTransforms.Bind(ParameterMap, *ParamNames.MaxTransformsName);
+		CurrentOffset.Bind(ParameterMap, *ParamNames.CurrentOffsetName);
+		PreviousOffset.Bind(ParameterMap, *ParamNames.PreviousOffsetName);
+
 		ElementOffsets.Bind(ParameterMap, *ParamNames.ElementOffsetsName);
 
 		WorldTransformBuffer.Bind(ParameterMap, *ParamNames.WorldTransformBufferName);
@@ -600,7 +620,12 @@ public:
 			SetSRVParameter(RHICmdList, ComputeShaderRHI, PhysicsTypeBuffer, AssetBuffer->PhysicsTypeBuffer.SRV);
 			SetSRVParameter(RHICmdList, ComputeShaderRHI, DFIndexBuffer, AssetBuffer->DFIndexBuffer.SRV);
 
-			SetShaderValue(RHICmdList, ComputeShaderRHI, ElementOffsets, ProxyData->AssetArrays.ElementOffsets);
+
+			SetShaderValue(RHICmdList, ComputeShaderRHI, MaxTransforms, ProxyData->AssetArrays->MaxTransforms);
+			SetShaderValue(RHICmdList, ComputeShaderRHI, CurrentOffset, 0);
+			SetShaderValue(RHICmdList, ComputeShaderRHI, PreviousOffset, ProxyData->AssetArrays->MaxTransforms);
+
+			SetShaderValue(RHICmdList, ComputeShaderRHI, ElementOffsets, ProxyData->AssetArrays->ElementOffsets);
 
 			if (DistanceFieldParameters.IsBound())
 			{
@@ -629,6 +654,9 @@ public:
 			SetSRVParameter(RHICmdList, ComputeShaderRHI, DFIndexBuffer, FNiagaraRenderer::GetDummyIntBuffer());
 
 			static const FElementOffset DummyOffsets(0, 0, 0, 0);
+			SetShaderValue(RHICmdList, ComputeShaderRHI, MaxTransforms, 0);
+			SetShaderValue(RHICmdList, ComputeShaderRHI, CurrentOffset, 0);
+			SetShaderValue(RHICmdList, ComputeShaderRHI, PreviousOffset, 0);
 			SetShaderValue(RHICmdList, ComputeShaderRHI, ElementOffsets, DummyOffsets);	
 		}
 	}
@@ -639,6 +667,10 @@ public:
 
 private:
 
+	LAYOUT_FIELD(FShaderParameter, MaxTransforms);
+	LAYOUT_FIELD(FShaderParameter, CurrentOffset);
+	LAYOUT_FIELD(FShaderParameter, PreviousOffset);
+
 	LAYOUT_FIELD(FShaderParameter, ElementOffsets);
 
 	LAYOUT_FIELD(FShaderResourceParameter, WorldTransformBuffer);
@@ -674,18 +706,18 @@ void FNDIRigidMeshCollisionProxy::ConsumePerInstanceDataFromGameThread(void* Per
 		// loop over proxies and compute df indices on the RT only on new data
 		// @todo(dmp): for now we do this every frame because it seems like sometimes DFindices are not set
 
-		for (uint32 i = 0; i < SourceData->AssetArrays.ElementOffsets.NumElements; ++i)
+		for (uint32 i = 0; i < SourceData->AssetArrays->ElementOffsets.NumElements; ++i)
 		{
-			FPrimitiveSceneProxy* Proxy = SourceData->AssetArrays.SourceSceneProxy[i];
+			FPrimitiveSceneProxy* Proxy = SourceData->AssetArrays->SourceSceneProxy[i];
 								
 			if (Proxy != nullptr && Proxy->GetPrimitiveSceneInfo() != nullptr)
 			{
 				const TArray<int32, TInlineAllocator<1>>& DFIndices = Proxy->GetPrimitiveSceneInfo()->DistanceFieldInstanceIndices;				
-				TargetData->AssetArrays.DFIndex[i] = DFIndices.Num() > 0 ? DFIndices[0] : -1;
+				TargetData->AssetArrays->DFIndex[i] = DFIndices.Num() > 0 ? DFIndices[0] : -1;
 			}
 			else
 			{
-				TargetData->AssetArrays.DFIndex[i] = -1;
+				TargetData->AssetArrays->DFIndex[i] = -1;
 			}
 		}	
 	}
@@ -705,7 +737,8 @@ void FNDIRigidMeshCollisionProxy::InitializePerInstanceData(const FNiagaraSystem
 
 void FNDIRigidMeshCollisionProxy::DestroyPerInstanceData(const FNiagaraSystemInstanceID& SystemInstance)
 {
-	check(IsInRenderingThread());
+	check(IsInRenderingThread());	
+
 	SystemInstancesToProxyData.Remove(SystemInstance);
 }
 
@@ -729,11 +762,11 @@ void FNDIRigidMeshCollisionProxy::PreStage(FRHICommandList& RHICmdList, const FN
 			};
 			RHICmdList.Transition(MakeArrayView(Transitions, UE_ARRAY_COUNT(Transitions)));
 
-			UpdateInternalBuffer<FVector4f, EPixelFormat::PF_A32B32G32R32F, RIGID_MESH_COLLISION_QUERY_MAX_TRANSFORMS, 2>(ProxyData->AssetArrays.WorldTransform, ProxyData->AssetBuffer->WorldTransformBuffer);
-			UpdateInternalBuffer<FVector4f, EPixelFormat::PF_A32B32G32R32F, RIGID_MESH_COLLISION_QUERY_MAX_TRANSFORMS, 2>(ProxyData->AssetArrays.InverseTransform, ProxyData->AssetBuffer->InverseTransformBuffer);
-			UpdateInternalBuffer<FVector4f, EPixelFormat::PF_A32B32G32R32F, RIGID_MESH_COLLISION_QUERY_MAX_PRIMITIVES>(ProxyData->AssetArrays.ElementExtent, ProxyData->AssetBuffer->ElementExtentBuffer);
-			UpdateInternalBuffer<uint32, EPixelFormat::PF_R32_UINT, RIGID_MESH_COLLISION_QUERY_MAX_PRIMITIVES>(ProxyData->AssetArrays.PhysicsType, ProxyData->AssetBuffer->PhysicsTypeBuffer);
-			UpdateInternalBuffer<uint32, EPixelFormat::PF_R32_UINT, RIGID_MESH_COLLISION_QUERY_MAX_PRIMITIVES>(ProxyData->AssetArrays.DFIndex, ProxyData->AssetBuffer->DFIndexBuffer);
+			UpdateInternalBuffer<FVector4f, EPixelFormat::PF_A32B32G32R32F>(ProxyData->AssetArrays->WorldTransform, ProxyData->AssetBuffer->WorldTransformBuffer);
+			UpdateInternalBuffer<FVector4f, EPixelFormat::PF_A32B32G32R32F>(ProxyData->AssetArrays->InverseTransform, ProxyData->AssetBuffer->InverseTransformBuffer);
+			UpdateInternalBuffer<FVector4f, EPixelFormat::PF_A32B32G32R32F>(ProxyData->AssetArrays->ElementExtent, ProxyData->AssetBuffer->ElementExtentBuffer);
+			UpdateInternalBuffer<uint32, EPixelFormat::PF_R32_UINT>(ProxyData->AssetArrays->PhysicsType, ProxyData->AssetBuffer->PhysicsTypeBuffer);
+			UpdateInternalBuffer<uint32, EPixelFormat::PF_R32_UINT>(ProxyData->AssetArrays->DFIndex, ProxyData->AssetBuffer->DFIndexBuffer);
 		}
 	}
 }
@@ -774,14 +807,21 @@ void UNiagaraDataInterfaceRigidMeshCollisionQuery::DestroyPerInstanceData(void*
 {
 	FNDIRigidMeshCollisionData* InstanceData = static_cast<FNDIRigidMeshCollisionData*>(PerInstanceData);
 
-	InstanceData->Release();
+	InstanceData->Release();	
 	InstanceData->~FNDIRigidMeshCollisionData();
 
 	FNDIRigidMeshCollisionProxy* ThisProxy = GetProxyAs<FNDIRigidMeshCollisionProxy>();
 	ENQUEUE_RENDER_COMMAND(FNiagaraDIDestroyInstanceData) (
 		[ThisProxy, InstanceID = SystemInstance->GetId()](FRHICommandListImmediate& CmdList)
 	{
-		ThisProxy->SystemInstancesToProxyData.Remove(InstanceID);
+		FNDIRigidMeshCollisionData* ProxyData =
+			ThisProxy->SystemInstancesToProxyData.Find(InstanceID);
+
+		if (ProxyData != nullptr && ProxyData->AssetArrays)
+		{			
+			ThisProxy->SystemInstancesToProxyData.Remove(InstanceID);
+			delete ProxyData->AssetArrays;
+		}		
 	}
 	);
 }
@@ -983,6 +1023,9 @@ bool UNiagaraDataInterfaceRigidMeshCollisionQuery::GetFunctionHLSL(const FNiagar
 
 	TMap<FString, FStringFormatArg> ArgsSample = {
 		{TEXT("InstanceFunctionName"), FunctionInfo.InstanceName},
+		{TEXT("MaxTransformsName"), ParamNames.MaxTransformsName},
+		{TEXT("CurrentOffsetName"), ParamNames.CurrentOffsetName},
+		{TEXT("PreviousOffsetName"), ParamNames.PreviousOffsetName},
 		{TEXT("ElementOffsetsName"), ParamNames.ElementOffsetsName},
 		{TEXT("WorldTransformBufferName"), ParamNames.WorldTransformBufferName},
 		{TEXT("InverseTransformBufferName"), ParamNames.InverseTransformBufferName},
@@ -1138,8 +1181,10 @@ void UNiagaraDataInterfaceRigidMeshCollisionQuery::ProvidePerInstanceDataForRend
 
 	if (GameThreadData != nullptr && RenderThreadData != nullptr)
 	{		
-		RenderThreadData->AssetBuffer = GameThreadData->AssetBuffer;		
-		RenderThreadData->AssetArrays = GameThreadData->AssetArrays;
+		RenderThreadData->AssetBuffer = GameThreadData->AssetBuffer;				
+
+		RenderThreadData->AssetArrays = new FNDIRigidMeshCollisionArrays();
+		RenderThreadData->AssetArrays->CopyFrom(GameThreadData->AssetArrays);		
 		RenderThreadData->TickingGroup = GameThreadData->TickingGroup;
 	}
 	check(Proxy);
diff --git a/Engine/Plugins/FX/Niagara/Source/Niagara/Public/NiagaraDataInterfaceRigidMeshCollisionQuery.h b/Engine/Plugins/FX/Niagara/Source/Niagara/Public/NiagaraDataInterfaceRigidMeshCollisionQuery.h
--- a/Engine/Plugins/FX/Niagara/Source/Niagara/Public/NiagaraDataInterfaceRigidMeshCollisionQuery.h
+++ b/Engine/Plugins/FX/Niagara/Source/Niagara/Public/NiagaraDataInterfaceRigidMeshCollisionQuery.h
@@ -7,9 +7,6 @@
 #include "VectorVM.h"
 #include "NiagaraDataInterfaceRigidMeshCollisionQuery.generated.h"
 
-#define RIGID_MESH_COLLISION_QUERY_MAX_PRIMITIVES 100
-#define RIGID_MESH_COLLISION_QUERY_MAX_TRANSFORMS RIGID_MESH_COLLISION_QUERY_MAX_PRIMITIVES * 2
-
 // Forward declaration
 class AStaticMeshActor;
 
@@ -33,16 +30,63 @@ struct FElementOffset
 struct FNDIRigidMeshCollisionArrays
 {
 	FElementOffset ElementOffsets;
-	TStaticArray<FVector4f, 2 * RIGID_MESH_COLLISION_QUERY_MAX_TRANSFORMS> WorldTransform;
-	TStaticArray<FVector4f, 2 * RIGID_MESH_COLLISION_QUERY_MAX_TRANSFORMS> InverseTransform;
-	TStaticArray<FVector4f, RIGID_MESH_COLLISION_QUERY_MAX_TRANSFORMS> CurrentTransform;
-	TStaticArray<FVector4f, RIGID_MESH_COLLISION_QUERY_MAX_TRANSFORMS> CurrentInverse;
-	TStaticArray<FVector4f, RIGID_MESH_COLLISION_QUERY_MAX_TRANSFORMS> PreviousTransform;
-	TStaticArray<FVector4f, RIGID_MESH_COLLISION_QUERY_MAX_TRANSFORMS> PreviousInverse;
-	TStaticArray<FVector4f, RIGID_MESH_COLLISION_QUERY_MAX_PRIMITIVES> ElementExtent;
-	TStaticArray<uint32, RIGID_MESH_COLLISION_QUERY_MAX_PRIMITIVES> PhysicsType;
-	TStaticArray<uint32, RIGID_MESH_COLLISION_QUERY_MAX_PRIMITIVES> DFIndex;
-	TStaticArray<FPrimitiveSceneProxy*, RIGID_MESH_COLLISION_QUERY_MAX_PRIMITIVES> SourceSceneProxy;
+	TArray<FVector4f> WorldTransform;
+	TArray<FVector4f> InverseTransform;
+	TArray<FVector4f> CurrentTransform;
+	TArray<FVector4f> CurrentInverse;
+	TArray<FVector4f> PreviousTransform;
+	TArray<FVector4f> PreviousInverse;
+	TArray<FVector4f> ElementExtent;
+	TArray<uint32> PhysicsType;
+	TArray<uint32> DFIndex;
+	TArray<FPrimitiveSceneProxy*> SourceSceneProxy;
+
+	FNDIRigidMeshCollisionArrays()
+	{
+		Resize(100);
+	}
+
+	FNDIRigidMeshCollisionArrays(uint32 Num)
+	{
+		Resize(Num);
+	}
+
+	void CopyFrom(const FNDIRigidMeshCollisionArrays* Other)
+	{
+		Resize(Other->MaxPrimitives);
+
+		ElementOffsets = Other->ElementOffsets;
+		WorldTransform = Other->WorldTransform;
+		InverseTransform = Other->InverseTransform;
+		CurrentTransform = Other->CurrentTransform;
+		CurrentInverse = Other->CurrentInverse;
+		PreviousTransform = Other->PreviousTransform;
+		PreviousInverse = Other->PreviousInverse;
+		ElementExtent = Other->ElementExtent;
+		PhysicsType = Other->PhysicsType;
+		DFIndex = Other->DFIndex;
+		SourceSceneProxy = Other->SourceSceneProxy;
+	}
+
+	void Resize(uint32 Num)
+	{
+		MaxPrimitives = Num;
+		MaxTransforms = MaxPrimitives * 2;				
+
+		WorldTransform.Init(FVector4f(0, 0, 0, 0), 3 * MaxTransforms);
+		InverseTransform.Init(FVector4f(0, 0, 0, 0), 3 * MaxTransforms);
+		CurrentTransform.Init(FVector4f(0, 0, 0, 0), 3 * MaxPrimitives);
+		CurrentInverse.Init(FVector4f(0, 0, 0, 0), 3 * MaxPrimitives);
+		PreviousTransform.Init(FVector4f(0, 0, 0, 0), 3 * MaxPrimitives);
+		PreviousInverse.Init(FVector4f(0, 0, 0, 0), 3 * MaxPrimitives);
+		ElementExtent.Init(FVector4f(0, 0, 0, 0), MaxPrimitives);
+		PhysicsType.Init(0, MaxPrimitives);
+		DFIndex.Init(0, MaxPrimitives);
+		SourceSceneProxy.Init(nullptr, MaxPrimitives);
+	}
+
+	uint32 MaxPrimitives = 100;
+	uint32 MaxTransforms = MaxPrimitives * 2;
 };
 
 /** Render buffers that will be used in hlsl functions */
@@ -71,6 +115,18 @@ struct FNDIRigidMeshCollisionBuffer : public FRenderResource
 
 	/** Distance field index buffer */
 	FRWBuffer DFIndexBuffer;
+
+	/** Max number of primitives */
+	uint32 MaxNumPrimitives;
+
+	/** Max number of transforms (prev and next needed) */
+	uint32 MaxNumTransforms;
+
+	void SetMaxNumPrimitives(uint32 Num)
+	{
+		MaxNumPrimitives = Num;
+		MaxNumTransforms = 2 * Num;
+	}
 };
 
 /** Data stored per physics asset instance*/
@@ -91,10 +147,10 @@ struct FNDIRigidMeshCollisionData
 	ETickingGroup TickingGroup;
 
 	/** Physics asset Gpu buffer */
-	FNDIRigidMeshCollisionBuffer* AssetBuffer;
+	FNDIRigidMeshCollisionBuffer* AssetBuffer = nullptr;
 
 	/** Physics asset Cpu arrays */
-	FNDIRigidMeshCollisionArrays AssetArrays;  
+	FNDIRigidMeshCollisionArrays *AssetArrays = nullptr;
 
 	/** Static Mesh Components **/
 	TArray<AStaticMeshActor*> StaticMeshActors;
@@ -119,6 +175,9 @@ public:
 	UPROPERTY(EditAnywhere, Category = "Static Mesh", meta = (EditCondition = "UseStaticMeshes"))
 	bool OnlyUseMoveable = true;
 
+	UPROPERTY(EditAnywhere, Category = "General")
+	int MaxNumPrimitives = 100;
+
 	/** UObject Interface */
 	virtual void PostInitProperties() override;
 
@@ -147,6 +206,15 @@ public:
 #endif
 	virtual void ProvidePerInstanceDataForRenderThread(void* DataForRenderThread, void* PerInstanceData, const FNiagaraSystemInstanceID& SystemInstance) override;
 
+	/** Name of element offsets */
+	static const FString MaxTransformsName;
+
+	/** Name of element offsets */
+	static const FString CurrentOffsetName;
+
+	/** Name of element offsets */
+	static const FString PreviousOffsetName;
+
 	/** Name of element offsets */
 	static const FString ElementOffsetsName;
 
diff --git a/Engine/Plugins/FX/Niagara/Source/NiagaraCore/Private/NiagaraCustomVersion.cpp b/Engine/Plugins/FX/Niagara/Source/NiagaraCore/Private/NiagaraCustomVersion.cpp
--- a/Engine/Plugins/FX/Niagara/Source/NiagaraCore/Private/NiagaraCustomVersion.cpp
+++ b/Engine/Plugins/FX/Niagara/Source/NiagaraCore/Private/NiagaraCustomVersion.cpp
@@ -8,4 +8,4 @@ const FGuid FNiagaraCustomVersion::GUID(0xFCF57AFA, 0x50764283, 0xB9A9E658, 0xFF
 // Register the custom version with core
 FCustomVersionRegistration GRegisterNiagaraCustomVersion(FNiagaraCustomVersion::GUID, FNiagaraCustomVersion::LatestVersion, TEXT("NiagaraVer"));
 
-const FGuid FNiagaraCustomVersion::LatestScriptCompileVersion(0x47F836A4, 0x880FE04D, 0x8E8AFB04, 0x6D8CA5CD);
+const FGuid FNiagaraCustomVersion::LatestScriptCompileVersion(0x67F13F7B, 0xAF06430C, 0x8660D50E, 0x4BB0DB90);
\ No newline at end of file
