better errors for cases when no collision shapes found

diff --git a/Engine/Plugins/FX/Niagara/Source/Niagara/Private/NiagaraDataInterfaceRigidMeshCollisionQuery.cpp b/Engine/Plugins/FX/Niagara/Source/Niagara/Private/NiagaraDataInterfaceRigidMeshCollisionQuery.cpp
--- a/Engine/Plugins/FX/Niagara/Source/Niagara/Private/NiagaraDataInterfaceRigidMeshCollisionQuery.cpp
+++ b/Engine/Plugins/FX/Niagara/Source/Niagara/Private/NiagaraDataInterfaceRigidMeshCollisionQuery.cpp
@@ -218,7 +218,7 @@ void CreateInternalArrays(TArray<AStaticMeshActor*> StaticMeshActors,
 				AStaticMeshActor* StaticMeshActor = StaticMeshActors[ActorIndex];
 				UStaticMeshComponent* StaticMeshComponent = StaticMeshActor != nullptr ? StaticMeshActor->GetStaticMeshComponent() : nullptr;
 				UBodySetup* BodySetup = StaticMeshComponent != nullptr ? StaticMeshComponent->GetBodySetup() : nullptr;				
-
+				bool FoundCollisionShapes = false;
 				if (BodySetup != nullptr)
 				{										
 					const FTransform MeshTransform = StaticMeshActor->GetTransform();				
@@ -227,7 +227,7 @@ void CreateInternalArrays(TArray<AStaticMeshActor*> StaticMeshActors,
 					{
 						if (CollisionEnabledHasPhysics(ConvexElem.GetCollisionEnabled()))
 						{
-							UE_LOG(LogRigidMeshCollision, Warning, TEXT("Convex collision objects encountered and will be skipped"));
+							UE_LOG(LogRigidMeshCollision, Warning, TEXT("Convex collision objects encountered and will be skipped on %s"), *StaticMeshActor->GetName());
 						}
 					}
 					for (const FKBoxElem& BoxElem : BodySetup->AggGeom.BoxElems)
@@ -240,6 +240,8 @@ void CreateInternalArrays(TArray<AStaticMeshActor*> StaticMeshActors,
 							const FTransform ElementTransform = FTransform(BoxElem.Rotation, BoxElem.Center) * MeshTransform;
 							OutAssetArrays->ElementExtent[BoxCount] = FVector4(BoxElem.X, BoxElem.Y, BoxElem.Z, 0);
 							FillCurrentTransforms(ElementTransform, BoxCount, OutAssetArrays->CurrentTransform, OutAssetArrays->CurrentInverse);							
+
+							FoundCollisionShapes = true;
 						}
 					}
 
@@ -253,6 +255,8 @@ void CreateInternalArrays(TArray<AStaticMeshActor*> StaticMeshActors,
 							const FTransform ElementTransform = FTransform(SphereElem.Center) * MeshTransform;
 							OutAssetArrays->ElementExtent[SphereCount] = FVector4(SphereElem.Radius, 0, 0, 0);
 							FillCurrentTransforms(ElementTransform, SphereCount, OutAssetArrays->CurrentTransform, OutAssetArrays->CurrentInverse);							
+
+							FoundCollisionShapes = true;
 						}
 					}
 
@@ -266,9 +270,17 @@ void CreateInternalArrays(TArray<AStaticMeshActor*> StaticMeshActors,
 							const FTransform ElementTransform = FTransform(CapsuleElem.Rotation, CapsuleElem.Center) * MeshTransform;
 							OutAssetArrays->ElementExtent[CapsuleCount] = FVector4(CapsuleElem.Radius, CapsuleElem.Length, 0, 0);
 							FillCurrentTransforms(ElementTransform, CapsuleCount, OutAssetArrays->CurrentTransform, OutAssetArrays->CurrentInverse);							
+
+							FoundCollisionShapes = true;
 						}
 					}					
 				}
+				
+				if (!FoundCollisionShapes)
+				{
+					UE_LOG(LogRigidMeshCollision, Warning, TEXT("No useable collision body setup found on mesh %s"), *StaticMeshActor->GetName());
+
+				}
 			}
 			OutAssetArrays->PreviousTransform = OutAssetArrays->CurrentTransform;
 			OutAssetArrays->PreviousInverse = OutAssetArrays->CurrentInverse;
