cis fix

diff --git a/Engine/Plugins/FX/Niagara/Source/Niagara/Private/NiagaraDataInterfaceRigidMeshCollisionQuery.cpp b/Engine/Plugins/FX/Niagara/Source/Niagara/Private/NiagaraDataInterfaceRigidMeshCollisionQuery.cpp
--- a/Engine/Plugins/FX/Niagara/Source/Niagara/Private/NiagaraDataInterfaceRigidMeshCollisionQuery.cpp
+++ b/Engine/Plugins/FX/Niagara/Source/Niagara/Private/NiagaraDataInterfaceRigidMeshCollisionQuery.cpp
@@ -216,70 +216,74 @@ void CreateInternalArrays(TArray<AStaticMeshActor*> StaticMeshActors,
 			for (int32 ActorIndex = 0; ActorIndex < StaticMeshActors.Num(); ++ActorIndex)
 			{				
 				AStaticMeshActor* StaticMeshActor = StaticMeshActors[ActorIndex];
-				UStaticMeshComponent* StaticMeshComponent = StaticMeshActor != nullptr ? StaticMeshActor->GetStaticMeshComponent() : nullptr;
-				UBodySetup* BodySetup = StaticMeshComponent != nullptr ? StaticMeshComponent->GetBodySetup() : nullptr;				
-				bool FoundCollisionShapes = false;
-				if (BodySetup != nullptr)
-				{										
-					const FTransform MeshTransform = StaticMeshActor->GetTransform();				
-
-					for (const FKConvexElem& ConvexElem : BodySetup->AggGeom.ConvexElems)
+
+				if (StaticMeshActor != nullptr && StaticMeshActor->GetStaticMeshComponent() != nullptr)
+				{
+					UStaticMeshComponent* StaticMeshComponent = StaticMeshActor->GetStaticMeshComponent();
+					UBodySetup* BodySetup = StaticMeshComponent->GetBodySetup();
+					bool FoundCollisionShapes = false;
+					if (BodySetup != nullptr)
 					{
-						if (CollisionEnabledHasPhysics(ConvexElem.GetCollisionEnabled()))
+						const FTransform MeshTransform = StaticMeshActor->GetTransform();
+
+						for (const FKConvexElem& ConvexElem : BodySetup->AggGeom.ConvexElems)
 						{
-							UE_LOG(LogRigidMeshCollision, Warning, TEXT("Convex collision objects encountered and will be skipped on %s"), *StaticMeshActor->GetName());
+							if (CollisionEnabledHasPhysics(ConvexElem.GetCollisionEnabled()))
+							{
+								UE_LOG(LogRigidMeshCollision, Warning, TEXT("Convex collision objects encountered and will be skipped on %s"), *StaticMeshActor->GetName());
+							}
 						}
-					}
-					for (const FKBoxElem& BoxElem : BodySetup->AggGeom.BoxElems)
-					{
-						if (CollisionEnabledHasPhysics(BoxElem.GetCollisionEnabled()))
-						{							
-							OutAssetArrays->PhysicsType[BoxCount] = (BoxElem.GetCollisionEnabled() == ECollisionEnabled::QueryAndPhysics);							
-							OutAssetArrays->SourceSceneProxy[BoxCount] = StaticMeshComponent->SceneProxy;
+						for (const FKBoxElem& BoxElem : BodySetup->AggGeom.BoxElems)
+						{
+							if (CollisionEnabledHasPhysics(BoxElem.GetCollisionEnabled()))
+							{
+								OutAssetArrays->PhysicsType[BoxCount] = (BoxElem.GetCollisionEnabled() == ECollisionEnabled::QueryAndPhysics);
+								OutAssetArrays->SourceSceneProxy[BoxCount] = StaticMeshComponent->SceneProxy;
 
-							const FTransform ElementTransform = FTransform(BoxElem.Rotation, BoxElem.Center) * MeshTransform;
-							OutAssetArrays->ElementExtent[BoxCount] = FVector4(BoxElem.X, BoxElem.Y, BoxElem.Z, 0);
-							FillCurrentTransforms(ElementTransform, BoxCount, OutAssetArrays->CurrentTransform, OutAssetArrays->CurrentInverse);							
+								const FTransform ElementTransform = FTransform(BoxElem.Rotation, BoxElem.Center) * MeshTransform;
+								OutAssetArrays->ElementExtent[BoxCount] = FVector4(BoxElem.X, BoxElem.Y, BoxElem.Z, 0);
+								FillCurrentTransforms(ElementTransform, BoxCount, OutAssetArrays->CurrentTransform, OutAssetArrays->CurrentInverse);
 
-							FoundCollisionShapes = true;
+								FoundCollisionShapes = true;
+							}
 						}
-					}
 
-					for (const FKSphereElem& SphereElem : BodySetup->AggGeom.SphereElems)
-					{
-						if (CollisionEnabledHasPhysics(SphereElem.GetCollisionEnabled()))
-						{														
-							OutAssetArrays->PhysicsType[SphereCount] = (SphereElem.GetCollisionEnabled() == ECollisionEnabled::QueryAndPhysics);							
-							OutAssetArrays->SourceSceneProxy[BoxCount] = StaticMeshComponent->SceneProxy;
+						for (const FKSphereElem& SphereElem : BodySetup->AggGeom.SphereElems)
+						{
+							if (CollisionEnabledHasPhysics(SphereElem.GetCollisionEnabled()))
+							{
+								OutAssetArrays->PhysicsType[SphereCount] = (SphereElem.GetCollisionEnabled() == ECollisionEnabled::QueryAndPhysics);
+								OutAssetArrays->SourceSceneProxy[BoxCount] = StaticMeshComponent->SceneProxy;
 
-							const FTransform ElementTransform = FTransform(SphereElem.Center) * MeshTransform;
-							OutAssetArrays->ElementExtent[SphereCount] = FVector4(SphereElem.Radius, 0, 0, 0);
-							FillCurrentTransforms(ElementTransform, SphereCount, OutAssetArrays->CurrentTransform, OutAssetArrays->CurrentInverse);							
+								const FTransform ElementTransform = FTransform(SphereElem.Center) * MeshTransform;
+								OutAssetArrays->ElementExtent[SphereCount] = FVector4(SphereElem.Radius, 0, 0, 0);
+								FillCurrentTransforms(ElementTransform, SphereCount, OutAssetArrays->CurrentTransform, OutAssetArrays->CurrentInverse);
 
-							FoundCollisionShapes = true;
+								FoundCollisionShapes = true;
+							}
 						}
-					}
 
-					for (const FKSphylElem& CapsuleElem : BodySetup->AggGeom.SphylElems)
-					{
-						if (CollisionEnabledHasPhysics(CapsuleElem.GetCollisionEnabled()))
-						{														
-							OutAssetArrays->PhysicsType[CapsuleCount] = (CapsuleElem.GetCollisionEnabled() == ECollisionEnabled::QueryAndPhysics);							
-							OutAssetArrays->SourceSceneProxy[BoxCount] = StaticMeshComponent->SceneProxy;
+						for (const FKSphylElem& CapsuleElem : BodySetup->AggGeom.SphylElems)
+						{
+							if (CollisionEnabledHasPhysics(CapsuleElem.GetCollisionEnabled()))
+							{
+								OutAssetArrays->PhysicsType[CapsuleCount] = (CapsuleElem.GetCollisionEnabled() == ECollisionEnabled::QueryAndPhysics);
+								OutAssetArrays->SourceSceneProxy[BoxCount] = StaticMeshComponent->SceneProxy;
 
-							const FTransform ElementTransform = FTransform(CapsuleElem.Rotation, CapsuleElem.Center) * MeshTransform;
-							OutAssetArrays->ElementExtent[CapsuleCount] = FVector4(CapsuleElem.Radius, CapsuleElem.Length, 0, 0);
-							FillCurrentTransforms(ElementTransform, CapsuleCount, OutAssetArrays->CurrentTransform, OutAssetArrays->CurrentInverse);							
+								const FTransform ElementTransform = FTransform(CapsuleElem.Rotation, CapsuleElem.Center) * MeshTransform;
+								OutAssetArrays->ElementExtent[CapsuleCount] = FVector4(CapsuleElem.Radius, CapsuleElem.Length, 0, 0);
+								FillCurrentTransforms(ElementTransform, CapsuleCount, OutAssetArrays->CurrentTransform, OutAssetArrays->CurrentInverse);
 
-							FoundCollisionShapes = true;
+								FoundCollisionShapes = true;
+							}
 						}
-					}					
-				}
-				
-				if (!FoundCollisionShapes)
-				{
-					UE_LOG(LogRigidMeshCollision, Warning, TEXT("No useable collision body setup found on mesh %s"), *StaticMeshActor->GetName());
+					}
+
+					if (!FoundCollisionShapes)
+					{
+						UE_LOG(LogRigidMeshCollision, Warning, TEXT("No useable collision body setup found on mesh %s"), *StaticMeshActor->GetName());
 
+					}
 				}
 			}
 			OutAssetArrays->PreviousTransform = OutAssetArrays->CurrentTransform;
diff --git a/Engine/Plugins/FX/Niagara/Source/Niagara/Public/NiagaraDataInterfaceRigidMeshCollisionQuery.h b/Engine/Plugins/FX/Niagara/Source/Niagara/Public/NiagaraDataInterfaceRigidMeshCollisionQuery.h
--- a/Engine/Plugins/FX/Niagara/Source/Niagara/Public/NiagaraDataInterfaceRigidMeshCollisionQuery.h
+++ b/Engine/Plugins/FX/Niagara/Source/Niagara/Public/NiagaraDataInterfaceRigidMeshCollisionQuery.h
@@ -169,10 +169,10 @@ public:
 	UPROPERTY(EditAnywhere, Category = "Static Mesh")
 	bool StaticMesh = true;
 
-	UPROPERTY(EditAnywhere, Category = "Static Mesh", meta = (EditCondition = "UseStaticMeshes"))
+	UPROPERTY(EditAnywhere, Category = "Static Mesh", meta = (EditCondition = "StaticMesh"))
 	FString Tag = TEXT("");
 
-	UPROPERTY(EditAnywhere, Category = "Static Mesh", meta = (EditCondition = "UseStaticMeshes"))
+	UPROPERTY(EditAnywhere, Category = "Static Mesh", meta = (EditCondition = "StaticMesh"))
 	bool OnlyUseMoveable = true;
 
 	UPROPERTY(EditAnywhere, Category = "General")
