expose UnitToFloatIndex as vm function

diff --git a/Engine/Plugins/FX/Niagara/Source/Niagara/Classes/NiagaraDataInterfaceGrid3DCollection.h b/Engine/Plugins/FX/Niagara/Source/Niagara/Classes/NiagaraDataInterfaceGrid3DCollection.h
--- a/Engine/Plugins/FX/Niagara/Source/Niagara/Classes/NiagaraDataInterfaceGrid3DCollection.h
+++ b/Engine/Plugins/FX/Niagara/Source/Niagara/Classes/NiagaraDataInterfaceGrid3DCollection.h
@@ -246,6 +246,7 @@ public:
 
 	void GetNumCells(FVectorVMExternalFunctionContext& Context);
 	void SetNumCells(FVectorVMExternalFunctionContext& Context);
+	void UnitToFloatIndex(FVectorVMExternalFunctionContext& Context);
 
 	void GetAttributeIndex(FVectorVMExternalFunctionContext& Context, const FName& InName, int32 NumChannels);
 
diff --git a/Engine/Plugins/FX/Niagara/Source/Niagara/Private/NiagaraDataInterfaceGrid3DCollection.cpp b/Engine/Plugins/FX/Niagara/Source/Niagara/Private/NiagaraDataInterfaceGrid3DCollection.cpp
--- a/Engine/Plugins/FX/Niagara/Source/Niagara/Private/NiagaraDataInterfaceGrid3DCollection.cpp
+++ b/Engine/Plugins/FX/Niagara/Source/Niagara/Private/NiagaraDataInterfaceGrid3DCollection.cpp
@@ -978,6 +978,7 @@ DEFINE_NDI_DIRECT_FUNC_BINDER(UNiagaraDataInterfaceGrid3DCollection, GetWorldBBo
 DEFINE_NDI_DIRECT_FUNC_BINDER(UNiagaraDataInterfaceGrid3DCollection, GetCellSize);
 DEFINE_NDI_DIRECT_FUNC_BINDER(UNiagaraDataInterfaceGrid3DCollection, SetNumCells);
 DEFINE_NDI_DIRECT_FUNC_BINDER(UNiagaraDataInterfaceGrid3DCollection, GetNumCells);
+DEFINE_NDI_DIRECT_FUNC_BINDER(UNiagaraDataInterfaceGrid3DCollection, UnitToFloatIndex);
 DEFINE_NDI_DIRECT_FUNC_BINDER_WITH_PAYLOAD(UNiagaraDataInterfaceGrid3DCollection, GetAttributeIndex);
 void UNiagaraDataInterfaceGrid3DCollection::GetVMExternalFunction(const FVMExternalFunctionBindingInfo& BindingInfo, void* InstanceData, FVMExternalFunction &OutFunc)
 {
@@ -1007,6 +1008,11 @@ void UNiagaraDataInterfaceGrid3DCollection::GetVMExternalFunction(const FVMExter
 		check(BindingInfo.GetNumInputs() == 1 && BindingInfo.GetNumOutputs() == 3);
 		NDI_FUNC_BINDER(UNiagaraDataInterfaceGrid3DCollection, GetNumCells)::Bind(this, OutFunc);
 	}
+	else if (BindingInfo.Name == UnitToFloatIndexFunctionName)
+	{		
+		check(BindingInfo.GetNumInputs() == 4 && BindingInfo.GetNumOutputs() == 3);
+		NDI_FUNC_BINDER(UNiagaraDataInterfaceGrid3DCollection, UnitToFloatIndex)::Bind(this, OutFunc);
+	}
 	else if (BindingInfo.Name == GetVector4AttributeIndexFunctionName)
 	{
 		FName AttributeName = BindingInfo.FindSpecifier(NAME_Attribute)->Value;
@@ -2471,6 +2477,32 @@ void UNiagaraDataInterfaceGrid3DCollection::SetNumCells(FVectorVMExternalFunctio
 	}
 }
 
+void UNiagaraDataInterfaceGrid3DCollection::UnitToFloatIndex(FVectorVMExternalFunctionContext& Context)
+{
+	// This should only be called from a system or emitter script due to a need for only setting up initially.
+	VectorVM::FUserPtrHandler<FGrid3DCollectionRWInstanceData_GameThread> InstData(Context);
+	VectorVM::FExternalFuncInputHandler<float> InUnitX(Context);
+	VectorVM::FExternalFuncInputHandler<float> InUnitY(Context);
+	VectorVM::FExternalFuncInputHandler<float> InUnitZ(Context);
+
+	VectorVM::FExternalFuncRegisterHandler<float> OutFloatX(Context);
+	VectorVM::FExternalFuncRegisterHandler<float> OutFloatY(Context);
+	VectorVM::FExternalFuncRegisterHandler<float> OutFloatZ(Context);	
+
+	for (int32 InstanceIdx = 0; InstanceIdx < Context.GetNumInstances(); ++InstanceIdx)
+	{
+		FVector3f InUnit;
+		InUnit.X = InUnitX.GetAndAdvance();
+		InUnit.Y = InUnitY.GetAndAdvance();
+		InUnit.Z = InUnitZ.GetAndAdvance();
+
+		FVector3f OutIndex = InUnit * FVector3f(InstData->NumCells) - .5;
+		*OutFloatX.GetDestAndAdvance() = OutIndex.X;
+		*OutFloatY.GetDestAndAdvance() = OutIndex.Y;
+		*OutFloatZ.GetDestAndAdvance() = OutIndex.Z;
+	}
+}
+
 bool UNiagaraDataInterfaceGrid3DCollection::PerInstanceTickPostSimulate(void* PerInstanceData, FNiagaraSystemInstance* SystemInstance, float DeltaSeconds)
 {
 	FGrid3DCollectionRWInstanceData_GameThread* InstanceData = static_cast<FGrid3DCollectionRWInstanceData_GameThread*>(PerInstanceData);
