fix sampling normals from mesh distance fields

diff --git a/Engine/Plugins/FX/Niagara/Shaders/Private/NiagaraDataInterfaceRigidMeshCollisionQuery.ush b/Engine/Plugins/FX/Niagara/Shaders/Private/NiagaraDataInterfaceRigidMeshCollisionQuery.ush
--- a/Engine/Plugins/FX/Niagara/Shaders/Private/NiagaraDataInterfaceRigidMeshCollisionQuery.ush
+++ b/Engine/Plugins/FX/Niagara/Shaders/Private/NiagaraDataInterfaceRigidMeshCollisionQuery.ush
@@ -409,7 +409,8 @@ void DIRigidMeshCollision_GetClosestPointMeshDistanceField(in FDIRigidMeshCollis
 	OutClosestPosition = float3(0,0,0);
 	OutClosestVelocity = float3(0,0,0);
 
-	const int ElementIndex = DIRigidMeshCollision_GetClosestElement(DIContext,WorldPosition,CollisionPosition,CollisionNormal,OutClosestDistance,TimeFraction);
+	float OutClosestDistanceTmp;
+	const int ElementIndex = DIRigidMeshCollision_GetClosestElement(DIContext,WorldPosition,CollisionPosition,CollisionNormal,OutClosestDistanceTmp,TimeFraction);
 
 	if (ElementIndex != -1)
 	{
@@ -419,28 +420,53 @@ void DIRigidMeshCollision_GetClosestPointMeshDistanceField(in FDIRigidMeshCollis
 		{
 			// #todo(dmp): this does a dynamic branch based on intersecting the bbox.  Maybe we can factor that out due to the broadphase here?
 			OutClosestDistance = DistanceToNearestSurfaceForObject(DFIndex, WorldPosition, 1e-3);
-		}
 
-		// #todo(dmp): figure out calculating SDF gradient so we can get more accurate collision velocities and collision normal
-		/*		
-		if (OutClosestDistance < 1e-5)
-		{
-			FDFObjectData DFObjectData = LoadDFObjectData(ElementIndex);
+			FDFObjectData DFObjectData = LoadDFObjectData(DFIndex);
 			uint NumMips = LoadDFAssetData(DFObjectData.AssetIndex, 0).NumMips;
 			FDFAssetData DFAssetData = LoadDFAssetData(DFObjectData.AssetIndex, NumMips - 1);
-			float3 VolumeGradient = CalculateMeshSDFGradient(WorldPosition, DFAssetData);
+			float3 VolumePosition = mul(float4(WorldPosition, 1), DFObjectData.WorldToVolume).xyz;		
+			float3 VolumeGradient = CalculateMeshSDFGradient(VolumePosition, DFAssetData);
 			float VolumeGradientLength = length(VolumeGradient);
 			float3 VolumeNormal = VolumeGradientLength > 0.0f ? VolumeGradient / VolumeGradientLength : 0;
 			// Transform by transposed inverse to handle non-uniform scaling
 			float3 WorldGradient = mul(VolumeNormal, transpose((float3x3)DFObjectData.WorldToVolume));
 			float WorldGradientLength = length(WorldGradient);
 			OutClosestNormal = WorldGradientLength > 0.0f ? WorldGradient / WorldGradientLength : 0;
+
+			const float3 PreviousPosition = mul(DIRigidMeshCollision_GetPreviousTransform(DIContext,ElementIndex), float4(CollisionPosition,1.0)).xyz;
+			const float3 CurrentPosition = mul(DIRigidMeshCollision_GetCurrentTransform(DIContext,ElementIndex), float4(CollisionPosition,1.0)).xyz;
+			OutClosestVelocity = ( CurrentPosition - PreviousPosition ) / DeltaTime;
+			OutClosestPosition = PreviousPosition + TimeFraction * (CurrentPosition-PreviousPosition);
 		}
-		*/
+	}
+}
 
-		const float3 PreviousPosition = mul(DIRigidMeshCollision_GetPreviousTransform(DIContext,ElementIndex), float4(CollisionPosition,1.0)).xyz;
-		const float3 CurrentPosition = mul(DIRigidMeshCollision_GetCurrentTransform(DIContext,ElementIndex), float4(CollisionPosition,1.0)).xyz;
-		OutClosestVelocity = ( CurrentPosition - PreviousPosition ) / DeltaTime;
-		OutClosestPosition = PreviousPosition + TimeFraction * (CurrentPosition-PreviousPosition);
+// Given a world space position (WorldPosition) compute the static mesh closest point (position,normal,velocity)
+void DIRigidMeshCollision_GetClosestPointMeshDistanceFieldNoNormal(in FDIRigidMeshCollisionContext DIContext, in float3 WorldPosition, in float DeltaTime, in float TimeFraction, out float OutClosestDistance, out float3 OutClosestPosition, out float3 OutClosestVelocity)
+{	
+	float3 CollisionPosition = float3(0,0,0);
+	float3 CollisionNormal = float3(0,0,0);
+	OutClosestDistance = MAX_DISTANCE;	
+	
+	OutClosestPosition = float3(0,0,0);
+	OutClosestVelocity = float3(0,0,0);
+
+	float OutClosestDistanceTmp;
+	const int ElementIndex = DIRigidMeshCollision_GetClosestElement(DIContext,WorldPosition,CollisionPosition,CollisionNormal,OutClosestDistanceTmp,TimeFraction);
+
+	if (ElementIndex != -1)
+	{
+		int DFIndex = DIContext.DFIndexBuffer[ElementIndex];
+
+		if (DFIndex >= 0 && DFIndex < NumSceneObjects)
+		{
+			// #todo(dmp): this does a dynamic branch based on intersecting the bbox.  Maybe we can factor that out due to the broadphase here?
+			OutClosestDistance = DistanceToNearestSurfaceForObject(DFIndex, WorldPosition, 1e-3);		
+
+			const float3 PreviousPosition = mul(DIRigidMeshCollision_GetPreviousTransform(DIContext,ElementIndex), float4(CollisionPosition,1.0)).xyz;
+			const float3 CurrentPosition = mul(DIRigidMeshCollision_GetCurrentTransform(DIContext,ElementIndex), float4(CollisionPosition,1.0)).xyz;
+			OutClosestVelocity = ( CurrentPosition - PreviousPosition ) / DeltaTime;
+			OutClosestPosition = PreviousPosition + TimeFraction * (CurrentPosition-PreviousPosition);
+		}
 	}
 }
\ No newline at end of file
diff --git a/Engine/Plugins/FX/Niagara/Source/Niagara/Private/NiagaraDataInterfaceRigidMeshCollisionQuery.cpp b/Engine/Plugins/FX/Niagara/Source/Niagara/Private/NiagaraDataInterfaceRigidMeshCollisionQuery.cpp
--- a/Engine/Plugins/FX/Niagara/Source/Niagara/Private/NiagaraDataInterfaceRigidMeshCollisionQuery.cpp
+++ b/Engine/Plugins/FX/Niagara/Source/Niagara/Private/NiagaraDataInterfaceRigidMeshCollisionQuery.cpp
@@ -37,6 +37,7 @@ static const FName GetElementDistanceName(TEXT("GetElementDistance"));
 static const FName GetClosestPointName(TEXT("GetClosestPoint"));
 static const FName GetClosestDistanceName(TEXT("GetClosestDistance"));
 static const FName GetClosestPointMeshDistanceFieldName(TEXT("GetClosestPointMeshDistanceField"));
+static const FName GetClosestPointMeshDistanceFieldNoNormalName(TEXT("GetClosestPointMeshDistanceFieldNoNormal"));
 
 
 //------------------------------------------------------------------------------------------------------------
@@ -1026,6 +1027,24 @@ void UNiagaraDataInterfaceRigidMeshCollisionQuery::GetFunctions(TArray<FNiagaraF
 
 		OutFunctions.Add(Sig);
 	}
+
+	{
+		FNiagaraFunctionSignature Sig;
+		Sig.Name = GetClosestPointMeshDistanceFieldNoNormalName;
+		Sig.bSupportsGPU = true;
+		Sig.bSupportsCPU = false;
+		Sig.bMemberFunction = true;
+		Sig.bRequiresContext = false;
+		Sig.Inputs.Add(FNiagaraVariable(FNiagaraTypeDefinition(GetClass()), TEXT("Collision DI")));
+		Sig.Inputs.Add(FNiagaraVariable(FNiagaraTypeDefinition::GetVec3Def(), TEXT("World Position")));
+		Sig.Inputs.Add(FNiagaraVariable(FNiagaraTypeDefinition::GetFloatDef(), TEXT("Delta Time")));
+		Sig.Inputs.Add(FNiagaraVariable(FNiagaraTypeDefinition::GetFloatDef(), TEXT("Time Fraction")));
+		Sig.Outputs.Add(FNiagaraVariable(FNiagaraTypeDefinition::GetFloatDef(), TEXT("Closest Distance")));
+		Sig.Outputs.Add(FNiagaraVariable(FNiagaraTypeDefinition::GetVec3Def(), TEXT("Closest Position")));		
+		Sig.Outputs.Add(FNiagaraVariable(FNiagaraTypeDefinition::GetVec3Def(), TEXT("Closest Velocity")));
+
+		OutFunctions.Add(Sig);
+	}
 }
 
 void UNiagaraDataInterfaceRigidMeshCollisionQuery::GetVMExternalFunction(const FVMExternalFunctionBindingInfo& BindingInfo, void* InstanceData, FVMExternalFunction& OutFunc)
@@ -1159,7 +1178,19 @@ bool UNiagaraDataInterfaceRigidMeshCollisionQuery::GetFunctionHLSL(const FNiagar
 		OutHLSL += FString::Format(FormatSample, ArgsSample);
 		return true;
 	}
-
+	else if (FunctionInfo.DefinitionName == GetClosestPointMeshDistanceFieldNoNormalName)
+	{
+	static const TCHAR* FormatSample = TEXT(R"(
+		void {InstanceFunctionName}(in float3 WorldPosition, in float DeltaTime, in float TimeFraction, out float ClosestDistance, out float3 OutClosestPosition, 
+							out float3 OutClosestVelocity)
+		{
+			{RigidMeshCollisionContextName} DIRigidMeshCollision_GetClosestPointMeshDistanceFieldNoNormal(DIContext,WorldPosition,DeltaTime,TimeFraction, ClosestDistance,
+				OutClosestPosition,OutClosestVelocity);
+		}
+		)");
+	OutHLSL += FString::Format(FormatSample, ArgsSample);
+	return true;
+	}
 	OutHLSL += TEXT("\n");
 	return false;
 }
@@ -1180,7 +1211,7 @@ void UNiagaraDataInterfaceRigidMeshCollisionQuery::GetParameterDefinitionHLSL(co
 #if WITH_EDITOR
 void UNiagaraDataInterfaceRigidMeshCollisionQuery::ValidateFunction(const FNiagaraFunctionSignature& Function, TArray<FText>& OutValidationErrors)
 {
-	if (Function.Name == GetClosestPointMeshDistanceFieldName)
+	if (Function.Name == GetClosestPointMeshDistanceFieldName || Function.Name == GetClosestPointMeshDistanceFieldNoNormalName)
 	{
 		if (!IsMeshDistanceFieldEnabled())
 		{
diff --git a/Engine/Plugins/FX/Niagara/Source/NiagaraCore/Private/NiagaraCustomVersion.cpp b/Engine/Plugins/FX/Niagara/Source/NiagaraCore/Private/NiagaraCustomVersion.cpp
--- a/Engine/Plugins/FX/Niagara/Source/NiagaraCore/Private/NiagaraCustomVersion.cpp
+++ b/Engine/Plugins/FX/Niagara/Source/NiagaraCore/Private/NiagaraCustomVersion.cpp
@@ -8,4 +8,4 @@ const FGuid FNiagaraCustomVersion::GUID(0xFCF57AFA, 0x50764283, 0xB9A9E658, 0xFF
 // Register the custom version with core
 FCustomVersionRegistration GRegisterNiagaraCustomVersion(FNiagaraCustomVersion::GUID, FNiagaraCustomVersion::LatestVersion, TEXT("NiagaraVer"));
 
-const FGuid FNiagaraCustomVersion::LatestScriptCompileVersion(0x67F13F7B, 0xAF06430C, 0x8660D50E, 0x4BB0DB90);
\ No newline at end of file
+const FGuid FNiagaraCustomVersion::LatestScriptCompileVersion(0x03D31049, 0xA9144EC4, 0x8D77B697, 0x6841EF4E);
\ No newline at end of file
