- Make LWC fix up handle referencing individual members, i.e. .xyz

#rnx

diff --git a/Engine/Plugins/FX/Niagara/Source/NiagaraCore/Private/NiagaraCustomVersion.cpp b/Engine/Plugins/FX/Niagara/Source/NiagaraCore/Private/NiagaraCustomVersion.cpp
--- a/Engine/Plugins/FX/Niagara/Source/NiagaraCore/Private/NiagaraCustomVersion.cpp
+++ b/Engine/Plugins/FX/Niagara/Source/NiagaraCore/Private/NiagaraCustomVersion.cpp
@@ -8,4 +8,4 @@ const FGuid FNiagaraCustomVersion::GUID(0xFCF57AFA, 0x50764283, 0xB9A9E658, 0xFF
 // Register the custom version with core
 FCustomVersionRegistration GRegisterNiagaraCustomVersion(FNiagaraCustomVersion::GUID, FNiagaraCustomVersion::LatestVersion, TEXT("NiagaraVer"));
 
-const FGuid FNiagaraCustomVersion::LatestScriptCompileVersion(0x56A96476, 0x62074AFC, 0xAF0593C6, 0x88E90637);
\ No newline at end of file
+const FGuid FNiagaraCustomVersion::LatestScriptCompileVersion(0x9FE9C854, 0x226FD449, 0x855AC15D, 0x98841030);
\ No newline at end of file
diff --git a/Engine/Plugins/FX/Niagara/Source/NiagaraEditor/Private/NiagaraHlslTranslator.cpp b/Engine/Plugins/FX/Niagara/Source/NiagaraEditor/Private/NiagaraHlslTranslator.cpp
--- a/Engine/Plugins/FX/Niagara/Source/NiagaraEditor/Private/NiagaraHlslTranslator.cpp
+++ b/Engine/Plugins/FX/Niagara/Source/NiagaraEditor/Private/NiagaraHlslTranslator.cpp
@@ -6616,16 +6616,31 @@ void FHlslNiagaraTranslator::ProcessCustomHlsl(const FString& InCustomHlsl, ENia
 			TEXT("PreViewTranslation"),
 			TEXT("PrevPreViewTranslation"),
 		};
+		static const FString ViewNamespace(TEXT("View."));
 
 		for (FString& Token : Tokens)
 		{
-			if (Token.StartsWith(TEXT("View."), ESearchCase::CaseSensitive))
+			if (Token.StartsWith(ViewNamespace, ESearchCase::CaseSensitive))
 			{
+				FStringView TokenMemberName = FStringView(Token).Mid(ViewNamespace.Len());
+				FStringView TokenPostfix;
+
+				int32 MemberEnd = INDEX_NONE;
+				if (TokenMemberName.FindChar(TCHAR('.'), MemberEnd))
+				{
+					TokenPostfix = TokenMemberName.Mid(MemberEnd);
+					TokenMemberName = TokenMemberName.Mid(0, MemberEnd);
+				}
+
 				for (const TCHAR* LWCMember : LWCViewMembers)
 				{
-					if (Token.Equals(FString(TEXT("View.")) + LWCMember))
+					if (TokenMemberName.Equals(LWCMember, ESearchCase::CaseSensitive))
 					{
 						Token = FString::Printf(TEXT("LWCToFloat(PrimaryView.%s)"), LWCMember);
+						if (TokenPostfix.Len() > 0)
+						{
+							Token.Append(TokenPostfix);
+						}
 						break;
 					}
 				}
