internally store a list of AActors rather than explictly static mesh actors.  Fixes using blueprint actors that have static mesh components.

diff --git a/Engine/Plugins/FX/Niagara/Source/Niagara/Private/NiagaraDataInterfaceRigidMeshCollisionQuery.cpp b/Engine/Plugins/FX/Niagara/Source/Niagara/Private/NiagaraDataInterfaceRigidMeshCollisionQuery.cpp
--- a/Engine/Plugins/FX/Niagara/Source/Niagara/Private/NiagaraDataInterfaceRigidMeshCollisionQuery.cpp
+++ b/Engine/Plugins/FX/Niagara/Source/Niagara/Private/NiagaraDataInterfaceRigidMeshCollisionQuery.cpp
@@ -13,7 +13,6 @@
 #include "NiagaraSystemInstance.h"
 #include "ShaderParameterUtils.h"
 #include "EngineUtils.h"
-#include "Engine/StaticMeshActor.h"
 #include "Renderer/Private/ScenePrivate.h"
 #include "DistanceFieldAtlas.h"
 #include "Renderer/Private/DistanceFieldLightingShared.h"
@@ -131,16 +130,20 @@ void FillCurrentTransforms(const FTransform& ElementTransform, uint32& ElementCo
 	++ElementCount;
 }
 
-void GetNumPrimitives(TArray<AStaticMeshActor*> StaticMeshActors, uint32& NumBoxes, uint32& NumSpheres, uint32& NumCapsules)
+void GetNumPrimitives(TArray<AActor*> Actors, uint32& NumBoxes, uint32& NumSpheres, uint32& NumCapsules)
 {
 	NumBoxes = 0;
 	NumSpheres = 0;
 	NumCapsules = 0;
 
-	for (int32 ActorIndex = 0; ActorIndex < StaticMeshActors.Num(); ++ActorIndex)
+	for (int32 ActorIndex = 0; ActorIndex < Actors.Num(); ++ActorIndex)
 	{
-		AStaticMeshActor* StaticMeshActor = StaticMeshActors[ActorIndex];
-		UStaticMeshComponent* StaticMeshComponent = StaticMeshActor != nullptr ? StaticMeshActor->GetStaticMeshComponent() : nullptr;
+		AActor* Actor = Actors[ActorIndex];
+		UStaticMeshComponent* StaticMeshComponent = nullptr;
+		if (Actor != nullptr)
+		{
+			StaticMeshComponent = Cast<UStaticMeshComponent>(Actor->GetComponentByClass(UStaticMeshComponent::StaticClass()));
+		}
 		UBodySetup* BodySetup = StaticMeshComponent != nullptr ? StaticMeshComponent->GetBodySetup() : nullptr;
 		if (BodySetup != nullptr)
 		{
@@ -184,7 +187,7 @@ void CompactInternalArrays(FNDIRigidMeshCollisionArrays* OutAssetArrays)
 	}
 }
 
-void CreateInternalArrays(TArray<AStaticMeshActor*> StaticMeshActors,
+void CreateInternalArrays(TArray<AActor*> Actors,
 	FNDIRigidMeshCollisionArrays* OutAssetArrays)
 {
 	if (OutAssetArrays != nullptr)
@@ -198,7 +201,7 @@ void CreateInternalArrays(TArray<AStaticMeshActor*> StaticMeshActors,
 		uint32 NumSpheres = 0;
 		uint32 NumCapsules = 0;
 
-		GetNumPrimitives(StaticMeshActors, NumBoxes, NumSpheres, NumCapsules);
+		GetNumPrimitives(Actors, NumBoxes, NumSpheres, NumCapsules);
 		
 		if ((NumBoxes + NumSpheres + NumCapsules) < OutAssetArrays->MaxPrimitives)
 		{
@@ -214,24 +217,24 @@ void CreateInternalArrays(TArray<AStaticMeshActor*> StaticMeshActors,
 			uint32 SphereCount = OutAssetArrays->ElementOffsets.SphereOffset;
 			uint32 CapsuleCount = OutAssetArrays->ElementOffsets.CapsuleOffset;
 
-			for (int32 ActorIndex = 0; ActorIndex < StaticMeshActors.Num(); ++ActorIndex)
+			for (int32 ActorIndex = 0; ActorIndex < Actors.Num(); ++ActorIndex)
 			{				
-				AStaticMeshActor* StaticMeshActor = StaticMeshActors[ActorIndex];
-
-				if (StaticMeshActor != nullptr && StaticMeshActor->GetStaticMeshComponent() != nullptr)
+				AActor* Actor = Actors[ActorIndex];
+				
+				if (Actor != nullptr && Actor->GetComponentByClass(UStaticMeshComponent::StaticClass()) != nullptr)
 				{
-					UStaticMeshComponent* StaticMeshComponent = StaticMeshActor->GetStaticMeshComponent();
+					UStaticMeshComponent* StaticMeshComponent = Cast< UStaticMeshComponent>(Actor->GetComponentByClass(UStaticMeshComponent::StaticClass()));
 					UBodySetup* BodySetup = StaticMeshComponent->GetBodySetup();
 					bool FoundCollisionShapes = false;
 					if (BodySetup != nullptr)
 					{
-						const FTransform MeshTransform = StaticMeshActor->GetTransform();
+						const FTransform MeshTransform = Actor->GetTransform();
 
 						for (const FKConvexElem& ConvexElem : BodySetup->AggGeom.ConvexElems)
 						{
 							if (CollisionEnabledHasPhysics(ConvexElem.GetCollisionEnabled()))
 							{
-								UE_LOG(LogRigidMeshCollision, Warning, TEXT("Convex collision objects encountered and will be skipped on %s"), *StaticMeshActor->GetName());
+								UE_LOG(LogRigidMeshCollision, Warning, TEXT("Convex collision objects encountered and will be skipped on %s"), *Actor->GetName());
 							}
 						}
 						for (const FKBoxElem& BoxElem : BodySetup->AggGeom.BoxElems)
@@ -282,7 +285,7 @@ void CreateInternalArrays(TArray<AStaticMeshActor*> StaticMeshActors,
 
 					if (!FoundCollisionShapes)
 					{
-						UE_LOG(LogRigidMeshCollision, Warning, TEXT("No useable collision body setup found on mesh %s"), *StaticMeshActor->GetName());
+						UE_LOG(LogRigidMeshCollision, Warning, TEXT("No useable collision body setup found on mesh %s"), *Actor->GetName());
 
 					}
 				}
@@ -299,7 +302,7 @@ void CreateInternalArrays(TArray<AStaticMeshActor*> StaticMeshActors,
 	}
 }
 
-void UpdateInternalArrays(const TArray<AStaticMeshActor*> &StaticMeshActors, FNDIRigidMeshCollisionArrays* OutAssetArrays)
+void UpdateInternalArrays(const TArray<AActor*> &Actors, FNDIRigidMeshCollisionArrays* OutAssetArrays)
 {
 	if (OutAssetArrays != nullptr && OutAssetArrays->ElementOffsets.NumElements < OutAssetArrays->MaxPrimitives)
 	{
@@ -307,13 +310,13 @@ void UpdateInternalArrays(const TArray<AStaticMeshActor*> &StaticMeshActors, FND
 		uint32 NumSpheres = 0;
 		uint32 NumCapsules = 0;		
 
-		GetNumPrimitives(StaticMeshActors, NumBoxes, NumSpheres, NumCapsules);
+		GetNumPrimitives(Actors, NumBoxes, NumSpheres, NumCapsules);
 
 		if (((OutAssetArrays->ElementOffsets.SphereOffset - OutAssetArrays->ElementOffsets.BoxOffset) != NumBoxes) || 
 			((OutAssetArrays->ElementOffsets.CapsuleOffset - OutAssetArrays->ElementOffsets.SphereOffset) != NumSpheres) ||
 			((OutAssetArrays->ElementOffsets.NumElements - OutAssetArrays->ElementOffsets.CapsuleOffset) != NumCapsules))
 		{
-			CreateInternalArrays(StaticMeshActors, OutAssetArrays);
+			CreateInternalArrays(Actors, OutAssetArrays);
 		}
 
 		OutAssetArrays->PreviousTransform = OutAssetArrays->CurrentTransform;
@@ -323,14 +326,18 @@ void UpdateInternalArrays(const TArray<AStaticMeshActor*> &StaticMeshActors, FND
 		uint32 SphereCount = OutAssetArrays->ElementOffsets.SphereOffset;
 		uint32 CapsuleCount = OutAssetArrays->ElementOffsets.CapsuleOffset;
 
-		for (int32 ActorIndex = 0; ActorIndex < StaticMeshActors.Num(); ++ActorIndex)
+		for (int32 ActorIndex = 0; ActorIndex < Actors.Num(); ++ActorIndex)
 		{
-			AStaticMeshActor* StaticMeshActor = StaticMeshActors[ActorIndex];
-			UStaticMeshComponent* StaticMeshComponent = StaticMeshActor != nullptr ? StaticMeshActor->GetStaticMeshComponent() : nullptr;
+			AActor* Actor = Actors[ActorIndex];
+			UStaticMeshComponent* StaticMeshComponent = nullptr;
+			if (Actor != nullptr)
+			{
+				StaticMeshComponent = Cast<UStaticMeshComponent>(Actor->GetComponentByClass(UStaticMeshComponent::StaticClass()));
+			}
 			UBodySetup* BodySetup = StaticMeshComponent != nullptr ? StaticMeshComponent->GetBodySetup() : nullptr;
 			if (BodySetup != nullptr)
 			{
-				const FTransform MeshTransform = StaticMeshActor->GetTransform();
+				const FTransform MeshTransform = Actor->GetTransform();
 
 				for (const FKBoxElem& BoxElem : BodySetup->AggGeom.BoxElems)
 				{
@@ -415,26 +422,26 @@ void FNDIRigidMeshCollisionData::Init(UNiagaraDataInterfaceRigidMeshCollisionQue
 	{
 		UWorld* World = SystemInstance->GetWorld();
 
-		StaticMeshActors.Empty();
+		Actors.Empty();
 
 		if (Interface->StaticMesh)
 		{
-			for (TActorIterator<AStaticMeshActor> It(World); It; ++It)
+			for (TActorIterator<AActor> It(World); It; ++It)
 			{
-				AStaticMeshActor* StaticMeshActor = *It;
+				AActor* Actor = *It;
 
-				if ((!Interface->OnlyUseMoveable || (Interface->OnlyUseMoveable && StaticMeshActor->IsRootComponentMovable())) &&
-					(Interface->Tag == FString("") || (Interface->Tag != FString("") && StaticMeshActor->Tags.Contains(FName(Interface->Tag)))))
+				if ((!Interface->OnlyUseMoveable || (Interface->OnlyUseMoveable && Actor->IsRootComponentMovable())) &&
+					(Interface->Tag == FString("") || (Interface->Tag != FString("") && Actor->Tags.Contains(FName(Interface->Tag)))))
 				{
-					StaticMeshActors.Add(StaticMeshActor);
+					Actors.Add(Actor);
 				}
 			}
 		}
 
-		if (0 < StaticMeshActors.Num() && StaticMeshActors[0] != nullptr)
+		if (0 < Actors.Num() && Actors[0] != nullptr)
 		{
 			// @note: we are not creating internal arrays here and letting it happen on the call to update
-			 //CreateInternalArrays(StaticMeshActors, &AssetArrays);
+			 //CreateInternalArrays(Actors, &AssetArrays);
 		}
 		AssetArrays = new FNDIRigidMeshCollisionArrays();
 		AssetArrays->Resize(Interface->MaxNumPrimitives);
@@ -451,9 +458,9 @@ void FNDIRigidMeshCollisionData::Update(UNiagaraDataInterfaceRigidMeshCollisionQ
 	{		
 		TickingGroup = ComputeTickingGroup();
 
-		if (0 < StaticMeshActors.Num() && StaticMeshActors[0] != nullptr)
+		if (0 < Actors.Num() && Actors[0] != nullptr)
 		{
-			UpdateInternalArrays(StaticMeshActors, AssetArrays);
+			UpdateInternalArrays(Actors, AssetArrays);
 		}
 	}
 }
@@ -461,11 +468,12 @@ void FNDIRigidMeshCollisionData::Update(UNiagaraDataInterfaceRigidMeshCollisionQ
 ETickingGroup FNDIRigidMeshCollisionData::ComputeTickingGroup()
 {
 	TickingGroup = NiagaraFirstTickGroup;
-	for (int32 ComponentIndex = 0; ComponentIndex < StaticMeshActors.Num(); ++ComponentIndex)
+	for (int32 ComponentIndex = 0; ComponentIndex < Actors.Num(); ++ComponentIndex)
 	{
-		if (StaticMeshActors[ComponentIndex] != nullptr)
+		if (Actors[ComponentIndex] != nullptr)
 		{			
-			UStaticMeshComponent* Component = StaticMeshActors[ComponentIndex]->GetStaticMeshComponent();
+		
+			UStaticMeshComponent* Component = Cast<UStaticMeshComponent>(Actors[ComponentIndex]->GetComponentByClass(UStaticMeshComponent::StaticClass()));
 
 			if (Component == nullptr)
 			{
diff --git a/Engine/Plugins/FX/Niagara/Source/Niagara/Public/NiagaraDataInterfaceRigidMeshCollisionQuery.h b/Engine/Plugins/FX/Niagara/Source/Niagara/Public/NiagaraDataInterfaceRigidMeshCollisionQuery.h
--- a/Engine/Plugins/FX/Niagara/Source/Niagara/Public/NiagaraDataInterfaceRigidMeshCollisionQuery.h
+++ b/Engine/Plugins/FX/Niagara/Source/Niagara/Public/NiagaraDataInterfaceRigidMeshCollisionQuery.h
@@ -8,7 +8,7 @@
 #include "NiagaraDataInterfaceRigidMeshCollisionQuery.generated.h"
 
 // Forward declaration
-class AStaticMeshActor;
+class AActor;
 
 /** Element offsets in the array list */
 struct FElementOffset
@@ -153,7 +153,7 @@ struct FNDIRigidMeshCollisionData
 	FNDIRigidMeshCollisionArrays *AssetArrays = nullptr;
 
 	/** Static Mesh Components **/
-	TArray<AStaticMeshActor*> StaticMeshActors;
+	TArray<AActor*> Actors;
 };
 
 /** Data Interface for the Collisions */
