- Fix merge error

#rnx

diff --git a/Engine/Plugins/FX/Niagara/Source/Niagara/Private/NiagaraRendererMeshes.cpp b/Engine/Plugins/FX/Niagara/Source/Niagara/Private/NiagaraRendererMeshes.cpp
--- a/Engine/Plugins/FX/Niagara/Source/Niagara/Private/NiagaraRendererMeshes.cpp
+++ b/Engine/Plugins/FX/Niagara/Source/Niagara/Private/NiagaraRendererMeshes.cpp
@@ -1113,7 +1113,7 @@ void FNiagaraRendererMeshes::GetDynamicRayTracingInstances(FRayTracingMaterialGa
 	// NOTE: have to run the GPU sort when culling is enabled if supported on this platform
 	// TODO: implement culling and renderer visibility on the CPU for other platforms
 	const bool bGPUSortEnabled = FNiagaraUtilities::AllowComputeShaders(Batcher->GetShaderPlatform());
-	const bool bDoGPUCulling = SourceMode == ENiagaraRendererSourceDataMode::Particles && bEnableCulling && GNiagaraGPUCulling && FNiagaraUtilities::AllowComputeShaders(Batcher->GetShaderPlatform());
+	const bool bDoGPUCulling = SourceMode == ENiagaraRendererSourceDataMode::Particles && bEnableCulling && FNiagaraUtilities::AllowGPUCulling(Batcher->GetShaderPlatform());
 	const bool bShouldSort = SourceMode == ENiagaraRendererSourceDataMode::Particles && SortMode != ENiagaraSortMode::None && (bHasTranslucentMaterials || !bSortOnlyWhenTranslucent);
 	const bool bCustomSorting = SortMode == ENiagaraSortMode::CustomAscending || SortMode == ENiagaraSortMode::CustomDecending;
 
