Fixed emitter inheritance not working when creating a new emitter from a parent emitter

diff --git a/Engine/Plugins/FX/Niagara/Source/NiagaraEditor/Private/SNewEmitterDialog.cpp b/Engine/Plugins/FX/Niagara/Source/NiagaraEditor/Private/SNewEmitterDialog.cpp
--- a/Engine/Plugins/FX/Niagara/Source/NiagaraEditor/Private/SNewEmitterDialog.cpp
+++ b/Engine/Plugins/FX/Niagara/Source/NiagaraEditor/Private/SNewEmitterDialog.cpp
@@ -27,7 +27,7 @@ void SNewEmitterDialog::Construct(const FArguments& InArgs)
 	TabOptions.ChangeTabState(ENiagaraScriptTemplateSpecification::None, true);
 	TabOptions.ChangeTabState(ENiagaraScriptTemplateSpecification::Behavior, true);
 
-	SAssignNew(TemplateAssetPicker, SNiagaraAssetPickerList, UNiagaraEmitter::StaticClass())
+	SAssignNew(NewAssetPicker, SNiagaraAssetPickerList, UNiagaraEmitter::StaticClass())
 		.ViewOptions(DisplayAllViewOptions)
 		.TabOptions(TabOptions);
 
@@ -41,22 +41,25 @@ void SNewEmitterDialog::Construct(const FArguments& InArgs)
 				LOCTEXT("CreateFromEmitterLabel", "New emitter"),
 				LOCTEXT("CreateFromEmitterDescription", "Create a new emitter from a template or behavior emitter (no inheritance) or from a parent (inheritance)"),
 				LOCTEXT("EmitterPickerHeader", "Select an Emitter"),
-				SNiagaraNewAssetDialog::FOnGetSelectedAssetsFromPicker::CreateSP(this, &SNewEmitterDialog::GetSelectedEmitterTemplateAssets),
-				SNiagaraNewAssetDialog::FOnSelectionConfirmed(),
-				TemplateAssetPicker.ToSharedRef(), TemplateAssetPicker->GetSearchBox() 
+				SNiagaraNewAssetDialog::FOnGetSelectedAssetsFromPicker::CreateSP(this, &SNewEmitterDialog::GetSelectedEmitterAssets_NewAssets),
+				SNiagaraNewAssetDialog::FOnSelectionConfirmed::CreateSP(this, &SNewEmitterDialog::CheckUseInheritance),
+				NewAssetPicker.ToSharedRef(), NewAssetPicker->GetSearchBox() 
 				),
 			SNiagaraNewAssetDialog::FNiagaraNewAssetDialogOption(
 				LOCTEXT("CreateFromOtherEmitterLabel", "Copy existing emitter"),
 				LOCTEXT("CreateFromOtherEmitterDescription", "Copies an existing emitter from your project content"),
 				LOCTEXT("ProjectEmitterPickerHeader", "Select a Project Emitter"),
-				SNiagaraNewAssetDialog::FOnGetSelectedAssetsFromPicker::CreateSP(this, &SNewEmitterDialog::GetSelectedProjectEmiterAssets),
-				SNiagaraNewAssetDialog::FOnSelectionConfirmed(),
+				SNiagaraNewAssetDialog::FOnGetSelectedAssetsFromPicker::CreateSP(this, &SNewEmitterDialog::GetSelectedEmitterAssets_CopyAssets),
+				SNiagaraNewAssetDialog::FOnSelectionConfirmed::CreateLambda([&]()
+				{
+					bUseInheritance = false;
+				}),
 				CopyAssetPicker.ToSharedRef(), CopyAssetPicker->GetSearchBox()
 				)
 		});
 }
 
-TOptional<FAssetData> SNewEmitterDialog::GetSelectedEmitterAsset()
+TOptional<FAssetData> SNewEmitterDialog::GetSelectedEmitterAsset() const
 {
 	const TArray<FAssetData>& SelectedEmitterAssets = GetSelectedAssets();
 	if (SelectedEmitterAssets.Num() > 0)
@@ -71,25 +74,30 @@ bool SNewEmitterDialog::GetUseInheritance() const
 	return bUseInheritance;
 }
 
-void SNewEmitterDialog::GetSelectedEmitterTemplateAssets(TArray<FAssetData>& OutSelectedAssets)
+void SNewEmitterDialog::GetSelectedEmitterAssets_NewAssets(TArray<FAssetData>& OutSelectedAssets) const
 {
-	OutSelectedAssets.Append(TemplateAssetPicker->GetSelectedAssets());
+	OutSelectedAssets.Append(NewAssetPicker->GetSelectedAssets());
 }
 
-void SNewEmitterDialog::GetSelectedParentEmitterAssets(TArray<FAssetData>& OutSelectedAssets)
-{
-	OutSelectedAssets.Append(InheritAssetPicker->GetSelectedAssets());
-}
-
-void SNewEmitterDialog::GetSelectedProjectEmiterAssets(TArray<FAssetData>& OutSelectedAssets)
+void SNewEmitterDialog::GetSelectedEmitterAssets_CopyAssets(TArray<FAssetData>& OutSelectedAssets) const
 {
 	OutSelectedAssets.Append(CopyAssetPicker->GetSelectedAssets());
 }
 
-
 void SNewEmitterDialog::InheritanceOptionConfirmed()
 {
 	bUseInheritance = true;
 }
 
+void SNewEmitterDialog::CheckUseInheritance()
+{
+	TOptional<FAssetData> EmitterAssetData = GetSelectedEmitterAsset();
+
+	if(EmitterAssetData.IsSet())
+	{
+		UNiagaraEmitter* Emitter = Cast<UNiagaraEmitter>(EmitterAssetData.GetValue().GetAsset());
+		bUseInheritance = Emitter->TemplateSpecification == ENiagaraScriptTemplateSpecification::None;
+	}
+}
+
 #undef LOCTEXT_NAMESPACE
diff --git a/Engine/Plugins/FX/Niagara/Source/NiagaraEditor/Private/SNewEmitterDialog.h b/Engine/Plugins/FX/Niagara/Source/NiagaraEditor/Private/SNewEmitterDialog.h
--- a/Engine/Plugins/FX/Niagara/Source/NiagaraEditor/Private/SNewEmitterDialog.h
+++ b/Engine/Plugins/FX/Niagara/Source/NiagaraEditor/Private/SNewEmitterDialog.h
@@ -21,23 +21,23 @@ public:
 
 	void Construct(const FArguments& InArgs);
 
-	TOptional<FAssetData> GetSelectedEmitterAsset();
+	TOptional<FAssetData> GetSelectedEmitterAsset() const;
 
 	bool GetUseInheritance() const;
 
 private:
-	void GetSelectedEmitterTemplateAssets(TArray<FAssetData>& OutSelectedAssets);
+	void GetSelectedEmitterAssets_NewAssets(TArray<FAssetData>& OutSelectedAssets) const;
 
-	void GetSelectedParentEmitterAssets(TArray<FAssetData>& OutSelectedAssets);
-
-	void GetSelectedProjectEmiterAssets(TArray<FAssetData>& OutSelectedAssets);
+	void GetSelectedEmitterAssets_CopyAssets(TArray<FAssetData>& OutSelectedAssets) const;
 
 	void InheritanceOptionConfirmed();
 
+	void CheckUseInheritance();
+
 private:
-	TSharedPtr<SNiagaraAssetPickerList> TemplateAssetPicker;
+	TSharedPtr<SNiagaraAssetPickerList> NewAssetPicker;
 	TSharedPtr<SNiagaraAssetPickerList> InheritAssetPicker;
 	TSharedPtr<SNiagaraAssetPickerList> CopyAssetPicker;
 
-	bool bUseInheritance;
+	bool bUseInheritance = false;
 };
\ No newline at end of file
