Temporary fix for deadlock that can happen when opening maps

diff --git a/Engine/Plugins/FX/Niagara/Source/Niagara/Private/NiagaraSystem.cpp b/Engine/Plugins/FX/Niagara/Source/Niagara/Private/NiagaraSystem.cpp
--- a/Engine/Plugins/FX/Niagara/Source/Niagara/Private/NiagaraSystem.cpp
+++ b/Engine/Plugins/FX/Niagara/Source/Niagara/Private/NiagaraSystem.cpp
@@ -974,7 +974,7 @@ bool UNiagaraSystem::IsReadyToRunInternal() const
 
 void UNiagaraSystem::GetAssetRegistryTags(TArray<FAssetRegistryTag>& OutTags) const
 {
-	EnsureFullyLoaded();
+	//EnsureFullyLoaded();
 #if WITH_EDITOR
 	OutTags.Add(FAssetRegistryTag("HasGPUEmitter", HasAnyGPUEmitters() ? TEXT("True") : TEXT("False"), FAssetRegistryTag::TT_Alphabetical));
 
