fix for deadlock that can happen when opening maps

diff --git a/Engine/Plugins/FX/Niagara/Source/Niagara/Classes/NiagaraEmitter.h b/Engine/Plugins/FX/Niagara/Source/Niagara/Classes/NiagaraEmitter.h
--- a/Engine/Plugins/FX/Niagara/Source/Niagara/Classes/NiagaraEmitter.h
+++ b/Engine/Plugins/FX/Niagara/Source/Niagara/Classes/NiagaraEmitter.h
@@ -622,6 +622,8 @@ private:
 	/** A multicast delegate which is called whenever all the scripts for this emitter have been compiled (successfully or not). */
 	FOnEmitterCompiled OnGPUScriptCompiledDelegate;
 
+	bool bFullyLoaded = false;
+
 	void RaiseOnEmitterGPUCompiled(UNiagaraScript* InScript, const FGuid& ScriptVersion);
 #endif
 
diff --git a/Engine/Plugins/FX/Niagara/Source/Niagara/Classes/NiagaraSystem.h b/Engine/Plugins/FX/Niagara/Source/Niagara/Classes/NiagaraSystem.h
--- a/Engine/Plugins/FX/Niagara/Source/Niagara/Classes/NiagaraSystem.h
+++ b/Engine/Plugins/FX/Niagara/Source/Niagara/Classes/NiagaraSystem.h
@@ -507,6 +507,7 @@ public:
 	UPROPERTY(EditAnywhere, Category = "Debug")
 	bool bDumpDebugEmitterInfo;
 
+	bool bFullyLoaded = false;
 
 	/** When enabled, we follow the settings on the UNiagaraComponent for tick order. When this option is disabled, we ignore any dependencies from data interfaces or other variables and instead fire off the simulation as early in the frame as possible. This greatly
 	reduces overhead and allows the game thread to run faster, but comes at a tradeoff if the dependencies might leave gaps or other visual artifacts.*/
diff --git a/Engine/Plugins/FX/Niagara/Source/Niagara/Private/NiagaraEmitter.cpp b/Engine/Plugins/FX/Niagara/Source/Niagara/Private/NiagaraEmitter.cpp
--- a/Engine/Plugins/FX/Niagara/Source/Niagara/Private/NiagaraEmitter.cpp
+++ b/Engine/Plugins/FX/Niagara/Source/Niagara/Private/NiagaraEmitter.cpp
@@ -1227,6 +1227,13 @@ void UNiagaraEmitter::CacheFromShaderCompiled()
 
 void UNiagaraEmitter::UpdateEmitterAfterLoad()
 {
+	if (bFullyLoaded)
+	{
+		return;
+	}
+	bFullyLoaded = true;
+	UpdateTaskRef = nullptr;
+	
 #if WITH_EDITORONLY_DATA
 	check(IsInGameThread());
 	if (GetOuter()->IsA<UNiagaraEmitter>())
diff --git a/Engine/Plugins/FX/Niagara/Source/Niagara/Private/NiagaraSystem.cpp b/Engine/Plugins/FX/Niagara/Source/Niagara/Private/NiagaraSystem.cpp
--- a/Engine/Plugins/FX/Niagara/Source/Niagara/Private/NiagaraSystem.cpp
+++ b/Engine/Plugins/FX/Niagara/Source/Niagara/Private/NiagaraSystem.cpp
@@ -351,6 +351,11 @@ bool UNiagaraSystem::UsesCollection(const UNiagaraParameterCollection* Collectio
 void UNiagaraSystem::UpdateSystemAfterLoad()
 {
 	// guard against deadlocks by having wait called on it during the update
+	if (bFullyLoaded)
+	{
+		return;
+	}
+	bFullyLoaded = true;
 	UpdateTaskRef = nullptr;
 	
 #if WITH_EDITORONLY_DATA
@@ -532,7 +537,7 @@ void UNiagaraSystem::UpdateSystemAfterLoad()
 }
 
 #if WITH_EDITORONLY_DATA
-bool UNiagaraSystem::UsesScript(const UNiagaraScript* Script)const
+bool UNiagaraSystem::UsesScript(const UNiagaraScript* Script) const
 {
 	EnsureFullyLoaded();
 	if (SystemSpawnScript == Script ||
@@ -974,7 +979,7 @@ bool UNiagaraSystem::IsReadyToRunInternal() const
 
 void UNiagaraSystem::GetAssetRegistryTags(TArray<FAssetRegistryTag>& OutTags) const
 {
-	//EnsureFullyLoaded();
+	EnsureFullyLoaded();
 #if WITH_EDITOR
 	OutTags.Add(FAssetRegistryTag("HasGPUEmitter", HasAnyGPUEmitters() ? TEXT("True") : TEXT("False"), FAssetRegistryTag::TT_Alphabetical));
 
@@ -1727,9 +1732,17 @@ void UNiagaraSystem::EnsureFullyLoaded() const
 {
 	if (UpdateTaskRef.IsValid() && UpdateTaskRef->IsComplete() == false)
 	{
-		UpdateTaskRef->Wait();
-	}
-	ensureMsgf(UpdateTaskRef.IsValid() == false || UpdateTaskRef->IsComplete(), TEXT("Niagara system data was accessed before being fully loaded! Asset %s"), *GetPathName());
+		ensureMsgf(FUObjectThreadContext::Get().IsRoutingPostLoad == false, TEXT("Niagara system data was forced to load during PostLoad instead of the taskgraph! Asset %s"), *GetPathName());
+		for (const FNiagaraEmitterHandle& EmitterHandle : EmitterHandles)
+		{
+			if (UNiagaraEmitter* NiagaraEmitter = EmitterHandle.GetInstance())
+			{
+				NiagaraEmitter->UpdateEmitterAfterLoad();
+			}
+		}
+		UNiagaraSystem* System = const_cast<UNiagaraSystem*>(this); // Sadly necessary because we are called from other const functions
+		System->UpdateSystemAfterLoad();
+	}	
 }
 
 
