Removed old debug messages

diff --git a/Engine/Plugins/FX/Niagara/Source/NiagaraEditor/Private/ViewModels/NiagaraSystemViewModel.cpp b/Engine/Plugins/FX/Niagara/Source/NiagaraEditor/Private/ViewModels/NiagaraSystemViewModel.cpp
--- a/Engine/Plugins/FX/Niagara/Source/NiagaraEditor/Private/ViewModels/NiagaraSystemViewModel.cpp
+++ b/Engine/Plugins/FX/Niagara/Source/NiagaraEditor/Private/ViewModels/NiagaraSystemViewModel.cpp
@@ -144,8 +144,6 @@ void FNiagaraSystemViewModel::DumpToText(FString& ExportText)
 
 void FNiagaraSystemViewModel::Cleanup()
 {
-	UE_LOG(LogNiagaraEditor, Warning, TEXT("Cleanup System view model %p"), this);
-
 	if (SystemInstance)
 	{
 		SystemInstance->OnInitialized().RemoveAll(this);
@@ -234,8 +232,6 @@ void FNiagaraSystemViewModel::Cleanup()
 FNiagaraSystemViewModel::~FNiagaraSystemViewModel()
 {
 	Cleanup();
-
-	UE_LOG(LogNiagaraEditor, Warning, TEXT("Deleting System view model %p"), this);
 }
 
 INiagaraParameterDefinitionsSubscriber* FNiagaraSystemViewModel::GetParameterDefinitionsSubscriber()
