Fix mismerge in NiagaraPinTypeSelector

diff --git a/Engine/Plugins/FX/Niagara/Source/NiagaraEditor/Private/Widgets/SNiagaraPinTypeSelector.h b/Engine/Plugins/FX/Niagara/Source/NiagaraEditor/Private/Widgets/SNiagaraPinTypeSelector.h
--- a/Engine/Plugins/FX/Niagara/Source/NiagaraEditor/Private/Widgets/SNiagaraPinTypeSelector.h
+++ b/Engine/Plugins/FX/Niagara/Source/NiagaraEditor/Private/Widgets/SNiagaraPinTypeSelector.h
@@ -3,16 +3,17 @@
 #pragma once
 
 #include "CoreMinimal.h"
-#include "SlateFwd.h"
-#include "Misc/Attribute.h"
 #include "EdGraph/EdGraphPin.h"
-#include "Widgets/DeclarativeSyntaxSupport.h"
-#include "Styling/SlateColor.h"
 #include "Input/Reply.h"
-#include "Widgets/SWidget.h"
+#include "Misc/Attribute.h"
 #include "NiagaraTypes.h"
+#include "SGraphPalette.h"
+#include "SlateFwd.h"
+#include "Styling/SlateColor.h"
+#include "Widgets/DeclarativeSyntaxSupport.h"
 #include "Widgets/Input/SComboButton.h"
 #include "Widgets/SCompoundWidget.h"
+#include "Widgets/SWidget.h"
 
 class NIAGARAEDITOR_API SNiagaraPinTypeSelector : public SCompoundWidget
 {
@@ -31,10 +32,6 @@ private:
 	UEdGraphPin* Pin = nullptr;
 	TSharedPtr<SComboButton> SelectorButton;
 };
-#include "Widgets/DeclarativeSyntaxSupport.h"
-#include "Widgets/SWidget.h"
-#include "Widgets/SCompoundWidget.h"
-#include "SGraphPalette.h"
 
 class SNiagaraIconWidget : public SGraphPaletteItem
 {
