- Editgate parts of CL 16281957 to fix material compilation error

#rnx

diff --git a/Engine/Plugins/FX/Niagara/Shaders/Private/NiagaraMeshVertexFactory.ush b/Engine/Plugins/FX/Niagara/Shaders/Private/NiagaraMeshVertexFactory.ush
--- a/Engine/Plugins/FX/Niagara/Shaders/Private/NiagaraMeshVertexFactory.ush
+++ b/Engine/Plugins/FX/Niagara/Shaders/Private/NiagaraMeshVertexFactory.ush
@@ -1023,7 +1023,7 @@ UNROLL
 #endif
 
 #if PASS_NEEDS_VERTEX_ID
-		Input.VertexId : Tri.Indices[VertexIndex];
+		Input.VertexId = VertexId;
 #endif // #if PASS_NEEDS_VERTEX_ID
 
 		return Input;
