Fix IsRotationNormalized() assert in Niagara meshes with ray tracing enabled

diff --git a/Engine/Plugins/FX/Niagara/Source/Niagara/Private/NiagaraRendererMeshes.cpp b/Engine/Plugins/FX/Niagara/Source/Niagara/Private/NiagaraRendererMeshes.cpp
--- a/Engine/Plugins/FX/Niagara/Source/Niagara/Private/NiagaraRendererMeshes.cpp
+++ b/Engine/Plugins/FX/Niagara/Source/Niagara/Private/NiagaraRendererMeshes.cpp
@@ -1377,7 +1377,7 @@ void FNiagaraRendererMeshes::GetDynamicRayTracingInstances(FRayTracingMaterialGa
 				if (bHasRotation)
 				{
 					FQuat InstanceQuat = GetInstanceQuat(InstanceIndex);
-					FTransform RotationTransform(InstanceQuat/*.GetNormalized()*/);
+					FTransform RotationTransform(InstanceQuat.GetNormalized());
 					FMatrix RotationMatrix = RotationTransform.ToMatrixWithScale();
 					InstanceTransform = RotationMatrix * InstanceTransform;
 				}
