- Handle edge case where we back to back tick and did not send EOF updates

#rnx

diff --git a/Engine/Plugins/FX/Niagara/Source/Niagara/Private/NiagaraSystemSimulation.cpp b/Engine/Plugins/FX/Niagara/Source/Niagara/Private/NiagaraSystemSimulation.cpp
--- a/Engine/Plugins/FX/Niagara/Source/Niagara/Private/NiagaraSystemSimulation.cpp
+++ b/Engine/Plugins/FX/Niagara/Source/Niagara/Private/NiagaraSystemSimulation.cpp
@@ -33,6 +33,8 @@
 //  - Enqueue finalize tasks (FNiagaraSystemInstanceFinalizeTask), tracks finalize inside instances
 //  - Appends all finalize tasks to a completion task (FNiagaraSystemSimulationAllWorkCompleteTask), when complete everything is done
 
+#define NIAGARA_SYSTEMSIMULATION_DEBUGGING 0
+
 //High level stats for system sim tick.
 DECLARE_CYCLE_STAT(TEXT("System Simulaton Tick [GT]"), STAT_NiagaraSystemSim_TickGT, STATGROUP_Niagara);
 DECLARE_CYCLE_STAT(TEXT("System Simulaton Tick [CNC]"), STAT_NiagaraSystemSim_TickCNC, STATGROUP_Niagara);
@@ -110,26 +112,111 @@ static FAutoConsoleVariableRef CVarNiagaraSkipTickDeltaSeconds(
 	ECVF_Default
 );
 
-#if 0
-static int GNiagaraSystemSimulationDebugKillOnSpawn = 0;
-static FAutoConsoleVariableRef CVarNiagaraSystemSimulationDebugKillOnSpawn(
-	TEXT("fx.Niagara.SystemSimulation.DebugKillOnSpawn"),
-	GNiagaraSystemSimulationDebugKillOnSpawn,
-	TEXT("Force kill instance on spawn."),
-	ECVF_Default
-);
+namespace NiagaraSystemSimulationLocal
+{
+#if NIAGARA_SYSTEMSIMULATION_DEBUGGING
+	static int GDebugKillAllOnSpawn = 0;
+	static FAutoConsoleVariableRef CVarDebugKillAllOnSpawn(
+		TEXT("fx.Niagara.SystemSimulation.DebugKillAllOnSpawn"),
+		GDebugKillAllOnSpawn,
+		TEXT("Will randomly kill all spawning instances based on probability of 0-100."),
+		ECVF_Default
+	);
+
+	static int GDebugKillAllOnUpdate = 0;
+	static FAutoConsoleVariableRef CVarDebugKillAllOnUpdate(
+		TEXT("fx.Niagara.SystemSimulation.DebugKillAllOnUpdate"),
+		GDebugKillAllOnUpdate,
+		TEXT("Will randomly kill all updating instances based on probability of 0-100."),
+		ECVF_Default
+	);
+
+	static int GDebugKillInstanceOnTick = 0;
+	static FAutoConsoleVariableRef CVarDebugKillInstanceOnTick(
+		TEXT("fx.Niagara.SystemSimulation.DebugKillInstanceOnTick"),
+		GDebugKillInstanceOnTick,
+		TEXT("Will randomly kill an instance on tick based on probability of 0-100."),
+		ECVF_Default
+	);
+
+	static float GDebugForceDelayConcurrentTask = 0.0f;
+	static FAutoConsoleVariableRef CVarDebugForceDelayConcurrentTask(
+		TEXT("fx.Niagara.SystemSimulation.DebugForceDelayConcurrentTask"),
+		GDebugForceDelayConcurrentTask,
+		TEXT("Forces a time delay into the concurrent task."),
+		ECVF_Default
+	);
+
+	static float GDebugForceDelayInstancesTask = 0.0f;
+	static FAutoConsoleVariableRef CVarDebugForceDelayInstancesTask(
+		TEXT("fx.Niagara.SystemSimulation.DebugForceDelayInstancesTask"),
+		GDebugForceDelayInstancesTask,
+		TEXT("Forces a time delay into the instances task."),
+		ECVF_Default
+	);
+
+	bool DebugKillAllOnSpawn(FNiagaraSystemSimulationTickContext& Context)
+	{
+		if (GDebugKillAllOnSpawn > 0)
+		{
+			if ((FMath::Rand() % 100) < GDebugKillAllOnSpawn)
+			{
+				for (FNiagaraSystemInstance* SystemInst : Context.Instances)
+				{
+					SystemInst->SetActualExecutionState(ENiagaraExecutionState::Disabled);
+				}
+				Context.DataSet.EndSimulate();
+				return true;
+			}
+		}
+		return false;
+	}
 
-static int GNiagaraSystemSimulationDebugKillOnUpdate = 0;
-static FAutoConsoleVariableRef CVarNiagaraSystemSimulationDebugKillOnUpdate(
-	TEXT("fx.Niagara.SystemSimulation.DebugKillOnUpdate"),
-	GNiagaraSystemSimulationDebugKillOnUpdate,
-	TEXT("Force kill instance on update."),
-	ECVF_Default
-);
-#else
-static constexpr int GNiagaraSystemSimulationDebugKillOnSpawn = 0;
-static constexpr int GNiagaraSystemSimulationDebugKillOnUpdate = 0;
+	bool DebugKillAllOnUpdate(FNiagaraSystemSimulationTickContext& Context)
+	{
+		if (GDebugKillAllOnUpdate > 0)
+		{
+			if ((FMath::Rand() % 100) < GDebugKillAllOnUpdate)
+			{
+				for (FNiagaraSystemInstance* SystemInst : Context.Instances)
+				{
+					SystemInst->SetActualExecutionState(ENiagaraExecutionState::Disabled);
+				}
+				Context.DataSet.EndSimulate();
+				return true;
+			}
+		}
+		return false;
+	}
+
+	void DebugKillInstanceOnTick(FNiagaraSystemInstance* Instance)
+	{
+		if (GDebugKillInstanceOnTick > 0)
+		{
+			if ((FMath::Rand() % 100) < GDebugKillInstanceOnTick)
+			{
+				Instance->Complete(true);
+			}
+		}
+	}
+
+	void DebugDelayConcurrentTask()
+	{
+		if (GDebugForceDelayConcurrentTask > 0.0f)
+		{
+			FPlatformProcess::Sleep(GDebugForceDelayConcurrentTask);
+		}
+	}
+
+	void DebugDelayInstancesTask()
+	{
+		if (GDebugForceDelayInstancesTask > 0.0f)
+		{
+			FPlatformProcess::Sleep(GDebugForceDelayInstancesTask);
+		}
+	}
 #endif
+}
 
 //////////////////////////////////////////////////////////////////////////
 // Task priorities for simulation tasks
@@ -261,6 +348,10 @@ struct FNiagaraSystemSimulationTickConcurrentTask
 	void DoTask(ENamedThreads::Type CurrentThread, const FGraphEventRef& MyCompletionGraphEvent)
 	{
 		PARTICLE_PERF_STAT_CYCLES_GT(FParticlePerfStatsContext(Context.World, Context.System), TickConcurrent);
+#if NIAGARA_SYSTEMSIMULATION_DEBUGGING
+		NiagaraSystemSimulationLocal::DebugDelayConcurrentTask();
+#endif
+
 		Context.Owner->Tick_Concurrent(Context);
 		CompletionTask->Unlock();
 	}
@@ -287,6 +378,9 @@ struct FNiagaraSystemSimulationSpawnConcurrentTask
 	void DoTask(ENamedThreads::Type CurrentThread, const FGraphEventRef& MyCompletionGraphEvent)
 	{
 		PARTICLE_PERF_STAT_CYCLES_GT(FParticlePerfStatsContext(Context.World, Context.System), TickConcurrent);
+#if NIAGARA_SYSTEMSIMULATION_DEBUGGING
+		NiagaraSystemSimulationLocal::DebugDelayConcurrentTask();
+#endif
 
 		Context.BeforeInstancesTickGraphEvents.Add(MyCompletionGraphEvent);
 		Context.Owner->Spawn_Concurrent(Context);
@@ -313,6 +407,10 @@ struct FNiagaraSystemInstanceTickConcurrentTask
 
 	void DoTask(ENamedThreads::Type CurrentThread, const FGraphEventRef& MyCompletionGraphEvent)
 	{
+#if NIAGARA_SYSTEMSIMULATION_DEBUGGING
+		NiagaraSystemSimulationLocal::DebugDelayInstancesTask();
+#endif
+
 		if (SystemSimulation->GetGPUTickHandlingMode() == ENiagaraGPUTickHandlingMode::ConcurrentBatched)
 		{
 			TArray<TPair<FNiagaraSystemGpuComputeProxy*, FNiagaraGPUSystemTick>, TInlineAllocator<NiagaraSystemTickBatchSize>> GPUTicks;
@@ -985,7 +1083,8 @@ void FNiagaraSystemSimulation::Tick_GameThread(float DeltaSeconds, const FGraphE
 
 	FNiagaraCrashReporterScope CRScope(this);
 
-	WaitForInstancesTickComplete(true);
+	// Work may not be complete if we back to back tick the GameThread without sending EOF updates
+	WaitForInstancesTickComplete();
 
 	SCOPE_CYCLE_COUNTER(STAT_NiagaraOverview_GT);
 	SCOPE_CYCLE_COUNTER(STAT_NiagaraSystemSim_TickGT);
@@ -1057,6 +1156,9 @@ void FNiagaraSystemSimulation::Tick_GameThread(float DeltaSeconds, const FGraphE
 
 		// Perform instance tick
 		Instance->Tick_GameThread(DeltaSeconds);
+#if NIAGARA_SYSTEMSIMULATION_DEBUGGING
+		NiagaraSystemSimulationLocal::DebugKillInstanceOnTick(Instance);
+#endif
 
 		// TickDataInterfaces could remove the system so we only increment if the system has changed
 		// Also possible for this system to have been transfered to another system simulation.
@@ -1099,6 +1201,9 @@ void FNiagaraSystemSimulation::Tick_GameThread(float DeltaSeconds, const FGraphE
 			// Execute instance tick
 			PARTICLE_PERF_STAT_CYCLES_WITH_COUNT_GT(FParticlePerfStatsContext(GetInstancePerfStats(Instance)), TickGameThread, 1);
 			Instance->Tick_GameThread(DeltaSeconds);
+#if NIAGARA_SYSTEMSIMULATION_DEBUGGING
+			NiagaraSystemSimulationLocal::DebugKillInstanceOnTick(Instance);
+#endif
 
 			if (Instance->SystemInstanceIndex != INDEX_NONE)
 			{
@@ -1263,6 +1368,9 @@ void FNiagaraSystemSimulation::Spawn_GameThread(float DeltaSeconds, bool bPostAc
 
 		// We are about to spawn execute game thread tick part
 		Instance->Tick_GameThread(DeltaSeconds);
+#if NIAGARA_SYSTEMSIMULATION_DEBUGGING
+		NiagaraSystemSimulationLocal::DebugKillInstanceOnTick(Instance);
+#endif
 
 		// If we are still alive spawn
 		if (Instance->SystemInstanceIndex != INDEX_NONE)
@@ -1355,10 +1463,12 @@ void FNiagaraSystemSimulation::WaitForInstancesTickComplete(bool bEnsureComplete
 	check(IsInGameThread());
 
 	// If our AllWorkCompleteGraphEvent is not complete we need to wait on all instances to complete manually
-	if (AllWorkCompleteGraphEvent.IsValid() && !AllWorkCompleteGraphEvent->IsComplete())
+	const bool bConcurrentInFlight = ConcurrentTickGraphEvent.IsValid() && !ConcurrentTickGraphEvent->IsComplete();
+	const bool bInstancesInFlight = AllWorkCompleteGraphEvent.IsValid() && !AllWorkCompleteGraphEvent->IsComplete();
+	if (bConcurrentInFlight || bInstancesInFlight)
 	{
 		SCOPE_CYCLE_COUNTER(STAT_NiagaraSystemSim_ForceWaitForAsync);
-		ensureAlwaysMsgf(!bEnsureComplete, TEXT("NiagaraSystemSimulation(%s) AllWorkCompleteGraphEvent is not completed."), *GetSystem()->GetPathName());
+		ensureAlwaysMsgf(!bEnsureComplete, TEXT("NiagaraSystemSimulation(%s) AllWorkCompleteGraphEvent is not completed (ConcurrentInFlight=%d, InstancesInFlight=%d)."), *GetSystem()->GetPathName(), bConcurrentInFlight, bInstancesInFlight);
 		WaitForConcurrentTickComplete(false);
 
 		TArray<FNiagaraSystemInstance*>& SystemInstances = GetSystemInstances(ENiagaraSystemInstanceState::Running);
@@ -1548,7 +1658,7 @@ void FNiagaraSystemSimulation::SpawnSystemInstances(FNiagaraSystemSimulationTick
 	Context.DataSet.GetDestinationDataChecked().SetNumInstances(NumInstances);
 
 	// Run Spawn
-	if (SpawnExecContext->Tick(SoloSystemInstance, ENiagaraSimTarget::CPUSim) == false || GNiagaraSystemSimulationDebugKillOnSpawn)
+	if (SpawnExecContext->Tick(SoloSystemInstance, ENiagaraSimTarget::CPUSim) == false)
 	{
 		for (FNiagaraSystemInstance* SystemInst : Context.Instances)
 		{
@@ -1558,6 +1668,13 @@ void FNiagaraSystemSimulation::SpawnSystemInstances(FNiagaraSystemSimulationTick
 		return;
 	}
 
+#if NIAGARA_SYSTEMSIMULATION_DEBUGGING
+	if ( NiagaraSystemSimulationLocal::DebugKillAllOnSpawn(Context) )
+	{
+		return;
+	}
+#endif
+
 	SpawnExecContext->BindSystemInstances(Context.Instances);
 	SpawnExecContext->BindData(0, Context.DataSet, OrigNum, false);
 	SpawnExecContext->BindData(1, SpawnInstanceParameterDataSet, OrigNum, false);
@@ -1609,7 +1726,7 @@ void FNiagaraSystemSimulation::UpdateSystemInstances(FNiagaraSystemSimulationTic
 		DestinationData.SetNumInstances(NumInstances);
 
 		// Tick UpdateExecContext, this can fail to bind VM functions if this happens we become invalid so mark all instances as disabled
-		if (UpdateExecContext->Tick(Context.Instances[0], ENiagaraSimTarget::CPUSim) == false || GNiagaraSystemSimulationDebugKillOnUpdate)
+		if (UpdateExecContext->Tick(Context.Instances[0], ENiagaraSimTarget::CPUSim) == false)
 		{
 			for (FNiagaraSystemInstance* SystemInst : Context.Instances)
 			{
@@ -1619,6 +1736,13 @@ void FNiagaraSystemSimulation::UpdateSystemInstances(FNiagaraSystemSimulationTic
 			return;
 		}
 
+#if NIAGARA_SYSTEMSIMULATION_DEBUGGING
+		if ( NiagaraSystemSimulationLocal::DebugKillAllOnUpdate(Context) )
+		{
+			return;
+		}
+#endif
+
 		// Run update.
 		if (OrigNum > 0)
 		{
