Merge UE5/RET [at] 16305968 to UE5/Main
This represents UE4/Main [at] 16261013 and Dev-PerfTest [at] 16259937

diff --git a/Engine/Plugins/FX/Niagara/Source/Niagara/Private/NiagaraSystem.cpp b/Engine/Plugins/FX/Niagara/Source/Niagara/Private/NiagaraSystem.cpp
--- a/Engine/Plugins/FX/Niagara/Source/Niagara/Private/NiagaraSystem.cpp
+++ b/Engine/Plugins/FX/Niagara/Source/Niagara/Private/NiagaraSystem.cpp
@@ -84,6 +84,15 @@ static FAutoConsoleVariableRef CVarNiagaraCompileWaitLoggingThreshold(
 	ECVF_Default
 );
 
+static int GNiagaraCompileWaitLoggingTerminationCap = 3;
+static FAutoConsoleVariableRef CVarNiagaraCompileWaitLoggingTerminationCap(
+	TEXT("fx.Niagara.CompileWaitLoggingCap"),
+	GNiagaraCompileWaitLoggingTerminationCap,
+	TEXT("During automation, how many times do we log before failing compilation?"),
+	ECVF_Default
+);
+
+
 static int GNiagaraCompileWaitLoggingTerminationCap = 3;
 static FAutoConsoleVariableRef CVarNiagaraCompileWaitLoggingTerminationCap(
 	TEXT("fx.Niagara.CompileWaitLoggingCap"),
diff --git a/Engine/Source/Runtime/Engine/Private/SoundConcurrency.cpp b/Engine/Source/Runtime/Engine/Private/SoundConcurrency.cpp
--- a/Engine/Source/Runtime/Engine/Private/SoundConcurrency.cpp
+++ b/Engine/Source/Runtime/Engine/Private/SoundConcurrency.cpp
@@ -403,7 +403,7 @@ bool FConcurrencyGroup::CanPlaySoundNow(float InCurrentTime) const
 		const float DeltaTime = InCurrentTime - LastTimePlayed;
 		if (DeltaTime < Settings.RetriggerTime)
 		{
-			UE_LOG(LogAudio, VeryVerbose, TEXT("Rejected Sound for Group ID (%d) with DeltaTime: %.2f"), GroupID, DeltaTime);
+			UE_LOG(LogAudioConcurrency, VeryVerbose, TEXT("Rejected Sound for Group ID (%d) with DeltaTime: %.2f"), GroupID, DeltaTime);
 			return false;
 		}
 	}
