Remove duplicated code from merge

diff --git a/Engine/Plugins/FX/Niagara/Source/Niagara/Private/NiagaraSystem.cpp b/Engine/Plugins/FX/Niagara/Source/Niagara/Private/NiagaraSystem.cpp
--- a/Engine/Plugins/FX/Niagara/Source/Niagara/Private/NiagaraSystem.cpp
+++ b/Engine/Plugins/FX/Niagara/Source/Niagara/Private/NiagaraSystem.cpp
@@ -84,15 +84,6 @@ static FAutoConsoleVariableRef CVarNiagaraCompileWaitLoggingThreshold(
 	ECVF_Default
 );
 
-static int GNiagaraCompileWaitLoggingTerminationCap = 3;
-static FAutoConsoleVariableRef CVarNiagaraCompileWaitLoggingTerminationCap(
-	TEXT("fx.Niagara.CompileWaitLoggingCap"),
-	GNiagaraCompileWaitLoggingTerminationCap,
-	TEXT("During automation, how many times do we log before failing compilation?"),
-	ECVF_Default
-);
-
-
 static int GNiagaraCompileWaitLoggingTerminationCap = 3;
 static FAutoConsoleVariableRef CVarNiagaraCompileWaitLoggingTerminationCap(
 	TEXT("fx.Niagara.CompileWaitLoggingCap"),
