Editgrate 15997724 - Add the ability to bind meshes in the mesh renderer to a user parameter

diff --git a/Engine/Plugins/FX/Niagara/Source/Niagara/Public/NiagaraComponent.h b/Engine/Plugins/FX/Niagara/Source/Niagara/Public/NiagaraComponent.h
--- a/Engine/Plugins/FX/Niagara/Source/Niagara/Public/NiagaraComponent.h
+++ b/Engine/Plugins/FX/Niagara/Source/Niagara/Public/NiagaraComponent.h
@@ -428,6 +428,9 @@ public:
 	UFUNCTION(BlueprintCallable, Category = Niagara, meta = (DisplayName = "Set Niagara Variable (Texture)"))
 	void SetVariableTexture(FName InVariableName, class UTexture* Texture);
 
+	UFUNCTION(BlueprintCallable, Category = Niagara, meta = (DisplayName = "Set Niagara Variable (Static Mesh)"))
+	void SetVariableStaticMesh(FName InVariableName, UStaticMesh* InValue);
+
 	UFUNCTION(BlueprintCallable, Category = Niagara, meta = (DisplayName = "Set Niagara Variable (TextureRenderTarget)"))
 	void SetVariableTextureRenderTarget(FName InVariableName, class UTextureRenderTarget* TextureRenderTarget);
 
