- Cache TextureSize and ensure we handle it changing behind our backs correctly by re-pushing data to the RT
- Engine changes mean that it is no longer allowed to get the SizeXYZ / etc on worker threads

diff --git a/Engine/Plugins/FX/Niagara/Source/Niagara/Classes/NiagaraDataInterface2DArrayTexture.h b/Engine/Plugins/FX/Niagara/Source/Niagara/Classes/NiagaraDataInterface2DArrayTexture.h
--- a/Engine/Plugins/FX/Niagara/Source/Niagara/Classes/NiagaraDataInterface2DArrayTexture.h
+++ b/Engine/Plugins/FX/Niagara/Source/Niagara/Classes/NiagaraDataInterface2DArrayTexture.h
@@ -33,6 +33,9 @@ public:
 	virtual void GetFunctions(TArray<FNiagaraFunctionSignature>& OutFunctions)override;
 	virtual void GetVMExternalFunction(const FVMExternalFunctionBindingInfo& BindingInfo, void* InstanceData, FVMExternalFunction &OutFunc) override;
 	virtual bool CanExecuteOnTarget(ENiagaraSimTarget Target)const override { return Target == ENiagaraSimTarget::GPUComputeSim; }
+
+	virtual bool HasPreSimulateTick() const override { return true; }
+	virtual bool PerInstanceTick(void* PerInstanceData, FNiagaraSystemInstance* SystemInstance, float DeltaSeconds) override;
 	//UNiagaraDataInterface Interface End
 
 	void GetTextureDimensions(FVectorVMContext& Context);
@@ -57,6 +60,9 @@ protected:
 
 	virtual void PushToRenderThreadImpl() override;
 
+protected:
+	FIntVector TextureSize = FIntVector::ZeroValue;
+
 	static const FName SampleTextureName;
 	static const FName TextureDimsName;
 };
diff --git a/Engine/Plugins/FX/Niagara/Source/Niagara/Classes/NiagaraDataInterfaceCubeTexture.h b/Engine/Plugins/FX/Niagara/Source/Niagara/Classes/NiagaraDataInterfaceCubeTexture.h
--- a/Engine/Plugins/FX/Niagara/Source/Niagara/Classes/NiagaraDataInterfaceCubeTexture.h
+++ b/Engine/Plugins/FX/Niagara/Source/Niagara/Classes/NiagaraDataInterfaceCubeTexture.h
@@ -33,6 +33,9 @@ public:
 	virtual void GetFunctions(TArray<FNiagaraFunctionSignature>& OutFunctions)override;
 	virtual void GetVMExternalFunction(const FVMExternalFunctionBindingInfo& BindingInfo, void* InstanceData, FVMExternalFunction &OutFunc) override;
 	virtual bool CanExecuteOnTarget(ENiagaraSimTarget Target)const override { return Target == ENiagaraSimTarget::GPUComputeSim; }
+
+	virtual bool HasPreSimulateTick() const override { return true; }
+	virtual bool PerInstanceTick(void* PerInstanceData, FNiagaraSystemInstance* SystemInstance, float DeltaSeconds) override;
 	//UNiagaraDataInterface Interface End
 
 	void SampleCubeTexture(FVectorVMContext& Context);
@@ -58,6 +61,9 @@ protected:
 
 	virtual void PushToRenderThreadImpl() override;
 
+protected:
+	FIntPoint TextureSize = FIntPoint::ZeroValue;
+
 	static const FName SampleCubeTextureName;
 	static const FName TextureDimsName;
 };
diff --git a/Engine/Plugins/FX/Niagara/Source/Niagara/Classes/NiagaraDataInterfaceTexture.h b/Engine/Plugins/FX/Niagara/Source/Niagara/Classes/NiagaraDataInterfaceTexture.h
--- a/Engine/Plugins/FX/Niagara/Source/Niagara/Classes/NiagaraDataInterfaceTexture.h
+++ b/Engine/Plugins/FX/Niagara/Source/Niagara/Classes/NiagaraDataInterfaceTexture.h
@@ -32,6 +32,9 @@ public:
 	virtual void GetFunctions(TArray<FNiagaraFunctionSignature>& OutFunctions)override;
 	virtual void GetVMExternalFunction(const FVMExternalFunctionBindingInfo& BindingInfo, void* InstanceData, FVMExternalFunction &OutFunc) override;
 	virtual bool CanExecuteOnTarget(ENiagaraSimTarget Target)const override { return Target==ENiagaraSimTarget::GPUComputeSim; }
+
+	virtual bool HasPreSimulateTick() const override { return true; }
+	virtual bool PerInstanceTick(void* PerInstanceData, FNiagaraSystemInstance* SystemInstance, float DeltaSeconds) override;
 	//UNiagaraDataInterface Interface End
 
 	void SampleTexture(FVectorVMContext& Context);
@@ -58,6 +61,9 @@ protected:
 
 	virtual void PushToRenderThreadImpl() override;
 
+protected:
+	FIntPoint TextureSize = FIntPoint::ZeroValue;
+
 	static const FName SampleTexture2DName;
 	static const FName SampleVolumeTextureName;
 	static const FName SamplePseudoVolumeTextureName;
diff --git a/Engine/Plugins/FX/Niagara/Source/Niagara/Classes/NiagaraDataInterfaceVolumeTexture.h b/Engine/Plugins/FX/Niagara/Source/Niagara/Classes/NiagaraDataInterfaceVolumeTexture.h
--- a/Engine/Plugins/FX/Niagara/Source/Niagara/Classes/NiagaraDataInterfaceVolumeTexture.h
+++ b/Engine/Plugins/FX/Niagara/Source/Niagara/Classes/NiagaraDataInterfaceVolumeTexture.h
@@ -33,6 +33,9 @@ public:
 	virtual void GetFunctions(TArray<FNiagaraFunctionSignature>& OutFunctions)override;
 	virtual void GetVMExternalFunction(const FVMExternalFunctionBindingInfo& BindingInfo, void* InstanceData, FVMExternalFunction &OutFunc) override;
 	virtual bool CanExecuteOnTarget(ENiagaraSimTarget Target)const override { return Target == ENiagaraSimTarget::GPUComputeSim; }
+
+	virtual bool HasPreSimulateTick() const override { return true; }
+	virtual bool PerInstanceTick(void* PerInstanceData, FNiagaraSystemInstance* SystemInstance, float DeltaSeconds) override;
 	//UNiagaraDataInterface Interface End
 
 	void SampleVolumeTexture(FVectorVMContext& Context);
@@ -58,6 +61,9 @@ protected:
 
 	virtual void PushToRenderThreadImpl() override;
 
+protected:
+	FIntVector TextureSize = FIntVector::ZeroValue;
+
 	static const FName SampleVolumeTextureName;
 	static const FName TextureDimsName;
 };
diff --git a/Engine/Plugins/FX/Niagara/Source/Niagara/Private/NiagaraDataInterface2DArrayTexture.cpp b/Engine/Plugins/FX/Niagara/Source/Niagara/Private/NiagaraDataInterface2DArrayTexture.cpp
--- a/Engine/Plugins/FX/Niagara/Source/Niagara/Private/NiagaraDataInterface2DArrayTexture.cpp
+++ b/Engine/Plugins/FX/Niagara/Source/Niagara/Private/NiagaraDataInterface2DArrayTexture.cpp
@@ -116,21 +116,25 @@ void UNiagaraDataInterface2DArrayTexture::GetVMExternalFunction(const FVMExterna
 	}
 }
 
-void UNiagaraDataInterface2DArrayTexture::GetTextureDimensions(FVectorVMContext& Context)
+bool UNiagaraDataInterface2DArrayTexture::PerInstanceTick(void* PerInstanceData, FNiagaraSystemInstance* SystemInstance, float DeltaSeconds)
 {
-	FNDIOutputParam<FVector> DimensionsOut(Context);
-
-	FVector TextureDimensions = FVector::ZeroVector;
-	if (Texture)
+	const FIntVector CurrentTextureSize = Texture != nullptr ? FIntVector(Texture->GetSizeX(), Texture->GetSizeY(), Texture->GetArraySize()) : FIntVector::ZeroValue;
+	if ( CurrentTextureSize != TextureSize )
 	{
-		TextureDimensions.X = Texture->GetSizeX();
-		TextureDimensions.Y = Texture->GetSizeY();
-		TextureDimensions.Z = Texture->GetArraySize();
+		TextureSize = CurrentTextureSize;
+		MarkRenderDataDirty();
 	}
+	return false;
+}
+
+void UNiagaraDataInterface2DArrayTexture::GetTextureDimensions(FVectorVMContext& Context)
+{
+	FNDIOutputParam<FVector> DimensionsOut(Context);
 
+	const FVector FloatTextureSize(TextureSize.X, TextureSize.Y, TextureSize.Z);
 	for (int32 i = 0; i < Context.NumInstances; ++i)
 	{
-		DimensionsOut.SetAndAdvance(TextureDimensions);
+		DimensionsOut.SetAndAdvance(FloatTextureSize);
 	}
 }
 
@@ -235,21 +239,21 @@ void UNiagaraDataInterface2DArrayTexture::PushToRenderThreadImpl()
 {
 	FNiagaraDataInterfaceProxy2DArrayTexture* RT_Proxy = GetProxyAs<FNiagaraDataInterfaceProxy2DArrayTexture>();
 
-	FVector RT_TexDims(EForceInit::ForceInitToZero);
+	TextureSize = FIntVector::ZeroValue;
 	if (Texture)
 	{
-		RT_TexDims.X = Texture->GetSizeX();
-		RT_TexDims.Y = Texture->GetSizeY();
-		RT_TexDims.Z = Texture->GetArraySize();
+		TextureSize.X = Texture->GetSizeX();
+		TextureSize.Y = Texture->GetSizeY();
+		TextureSize.Z = Texture->GetArraySize();
 	}
 
 	ENQUEUE_RENDER_COMMAND(FPushDITextureToRT)
 	(
-		[RT_Proxy, RT_Resource=Texture ? Texture->Resource : nullptr, RT_TexDims](FRHICommandListImmediate& RHICmdList)
+		[RT_Proxy, RT_Resource=Texture ? Texture->Resource : nullptr, RT_TexDims=TextureSize](FRHICommandListImmediate& RHICmdList)
 		{
 			RT_Proxy->TextureRHI = RT_Resource ? RT_Resource->TextureRHI : nullptr;
 			RT_Proxy->SamplerStateRHI = RT_Resource ? RT_Resource->SamplerStateRHI : nullptr;
-			RT_Proxy->TexDims = RT_TexDims;
+			RT_Proxy->TexDims = FVector(RT_TexDims.X, RT_TexDims.Y, RT_TexDims.Z);
 		}
 	);
 }
diff --git a/Engine/Plugins/FX/Niagara/Source/Niagara/Private/NiagaraDataInterfaceCubeTexture.cpp b/Engine/Plugins/FX/Niagara/Source/Niagara/Private/NiagaraDataInterfaceCubeTexture.cpp
--- a/Engine/Plugins/FX/Niagara/Source/Niagara/Private/NiagaraDataInterfaceCubeTexture.cpp
+++ b/Engine/Plugins/FX/Niagara/Source/Niagara/Private/NiagaraDataInterfaceCubeTexture.cpp
@@ -131,28 +131,26 @@ void UNiagaraDataInterfaceCubeTexture::GetVMExternalFunction(const FVMExternalFu
 	}
 }
 
+bool UNiagaraDataInterfaceCubeTexture::PerInstanceTick(void* PerInstanceData, FNiagaraSystemInstance* SystemInstance, float DeltaSeconds)
+{
+	const FIntPoint CurrentTextureSize = Texture != nullptr ? FIntPoint(Texture->GetSizeX(), Texture->GetSizeY()) : FIntPoint::ZeroValue;
+	if ( CurrentTextureSize != TextureSize )
+	{
+		TextureSize = CurrentTextureSize;
+		MarkRenderDataDirty();
+	}
+	return false;
+}
+
 void UNiagaraDataInterfaceCubeTexture::GetTextureDimensions(FVectorVMContext& Context)
 {
 	FNDIOutputParam<int32> OutWidth(Context);
 	FNDIOutputParam<int32> OutHeight(Context);
 
-	if (Texture == nullptr)
-	{
-		for (int32 i = 0; i < Context.NumInstances; ++i)
-		{
-			OutWidth.SetAndAdvance(0);
-			OutHeight.SetAndAdvance(0);
-		}
-	}
-	else
+	for (int32 i = 0; i < Context.NumInstances; ++i)
 	{
-		const int Width = Texture->GetSizeX();
-		const int Height = Texture->GetSizeY();
-		for (int32 i = 0; i < Context.NumInstances; ++i)
-		{
-			OutWidth.SetAndAdvance(Width);
-			OutHeight.SetAndAdvance(Height);
-		}
+		OutWidth.SetAndAdvance(TextureSize.X);
+		OutHeight.SetAndAdvance(TextureSize.Y);
 	}
 }
 
@@ -274,16 +272,16 @@ void UNiagaraDataInterfaceCubeTexture::PushToRenderThreadImpl()
 {
 	FNiagaraDataInterfaceProxyCubeTexture* RT_Proxy = GetProxyAs<FNiagaraDataInterfaceProxyCubeTexture>();
 
-	FIntPoint RT_TexDims(0, 0);
+	TextureSize = FIntPoint::ZeroValue;
 	if (Texture)
 	{
-		RT_TexDims.X = Texture->GetSizeX();
-		RT_TexDims.Y = Texture->GetSizeY();
+		TextureSize.X = Texture->GetSizeX();
+		TextureSize.Y = Texture->GetSizeY();
 	}
 
 	ENQUEUE_RENDER_COMMAND(FPushDITextureToRT)
 	(
-		[RT_Proxy, RT_Resource=Texture ? Texture->Resource : nullptr, RT_TexDims](FRHICommandListImmediate& RHICmdList)
+		[RT_Proxy, RT_Resource=Texture ? Texture->Resource : nullptr, RT_TexDims=TextureSize](FRHICommandListImmediate& RHICmdList)
 		{
 			RT_Proxy->TextureRHI = RT_Resource ? RT_Resource->TextureRHI : nullptr;
 			RT_Proxy->SamplerStateRHI = RT_Resource ? RT_Resource->SamplerStateRHI : nullptr;
diff --git a/Engine/Plugins/FX/Niagara/Source/Niagara/Private/NiagaraDataInterfaceTexture.cpp b/Engine/Plugins/FX/Niagara/Source/Niagara/Private/NiagaraDataInterfaceTexture.cpp
--- a/Engine/Plugins/FX/Niagara/Source/Niagara/Private/NiagaraDataInterfaceTexture.cpp
+++ b/Engine/Plugins/FX/Niagara/Source/Niagara/Private/NiagaraDataInterfaceTexture.cpp
@@ -217,28 +217,26 @@ void UNiagaraDataInterfaceTexture::GetVMExternalFunction(const FVMExternalFuncti
 	}
 }
 
-void UNiagaraDataInterfaceTexture::GetTextureDimensions(FVectorVMContext& Context)
+bool UNiagaraDataInterfaceTexture::PerInstanceTick(void* PerInstanceData, FNiagaraSystemInstance* SystemInstance, float DeltaSeconds)
 {
-	VectorVM::FExternalFuncRegisterHandler<float> OutWidth(Context);
-	VectorVM::FExternalFuncRegisterHandler<float> OutHeight(Context);
-
-	if (Texture == nullptr)
+	const FIntPoint CurrentTextureSize = Texture != nullptr ? FIntPoint(Texture->GetSurfaceWidth(), Texture->GetSurfaceHeight()) : FIntPoint::ZeroValue;
+	if ( CurrentTextureSize !=  TextureSize )
 	{
-		for (int32 i = 0; i < Context.NumInstances; ++i)
-		{
-			*OutWidth.GetDestAndAdvance() = 0.0f;
-			*OutHeight.GetDestAndAdvance() = 0.0f;
-		}
+		TextureSize = CurrentTextureSize;
+		MarkRenderDataDirty();
 	}
-	else
+	return false;
+}
+
+void UNiagaraDataInterfaceTexture::GetTextureDimensions(FVectorVMContext& Context)
+{
+	FNDIOutputParam<float> OutWidth(Context);
+	FNDIOutputParam<float> OutHeight(Context);
+
+	for (int32 i = 0; i < Context.NumInstances; ++i)
 	{
-		float Width = Texture->GetSurfaceWidth();
-		float Height = Texture->GetSurfaceHeight();
-		for (int32 i = 0; i < Context.NumInstances; ++i)
-		{
-			*OutWidth.GetDestAndAdvance() = Width;
-			*OutHeight.GetDestAndAdvance() = Height;
-		}
+		OutWidth.SetAndAdvance(TextureSize.X);
+		OutHeight.SetAndAdvance(TextureSize.Y);
 	}
 }
 
@@ -380,10 +378,7 @@ public:
 		TextureParam.Bind(ParameterMap, *TexName);
 		SamplerParam.Bind(ParameterMap, *SampleName);
 		
-		
-
 		Dimensions.Bind(ParameterMap, *(UNiagaraDataInterfaceTexture::DimensionsBaseName + ParameterInfo.DataInterfaceHLSLSymbol));
-
 	}
 
 	void Set(FRHICommandList& RHICmdList, const FNiagaraDataInterfaceSetArgs& Context) const
@@ -412,7 +407,7 @@ public:
 				SamplerStateRHI,
 				TextureDI->TextureReferenceRHI->GetReferencedTexture()
 			);
-			SetShaderValue(RHICmdList, ComputeShaderRHI, Dimensions, TextureDI->TexDims);
+			SetShaderValue(RHICmdList, ComputeShaderRHI, Dimensions, FVector2D(TextureDI->TexDims.X, TextureDI->TexDims.Y));
 		}
 		else
 		{
@@ -425,7 +420,7 @@ public:
 				GBlackTexture->TextureRHI
 			);
 			FVector2D TexDims(EForceInit::ForceInitToZero);
-			SetShaderValue(RHICmdList, ComputeShaderRHI, Dimensions, TexDims);
+			SetShaderValue(RHICmdList, ComputeShaderRHI, Dimensions, FVector2D::ZeroVector);
 		}
 	}
 
@@ -443,17 +438,16 @@ void UNiagaraDataInterfaceTexture::PushToRenderThreadImpl()
 {
 	FNiagaraDataInterfaceProxyTexture* RT_Proxy = GetProxyAs<FNiagaraDataInterfaceProxyTexture>();
 
-	FVector2D RT_TexDims(EForceInit::ForceInitToZero);
-
+	TextureSize = FIntPoint::ZeroValue;
 	if (Texture)
 	{
-		RT_TexDims.X = Texture->GetSurfaceWidth();
-		RT_TexDims.Y = Texture->GetSurfaceHeight();
+		TextureSize.X = Texture->GetSurfaceWidth();
+		TextureSize.Y = Texture->GetSurfaceHeight();
 	}
 
 	ENQUEUE_RENDER_COMMAND(FPushDITextureToRT)
 	(
-		[RT_Proxy, RT_Texture=Texture, RT_TexDims](FRHICommandListImmediate& RHICmdList)
+		[RT_Proxy, RT_Texture=Texture, RT_TexDims=TextureSize](FRHICommandListImmediate& RHICmdList)
 		{
 			if (RT_Texture)
 			{
@@ -465,7 +459,7 @@ void UNiagaraDataInterfaceTexture::PushToRenderThreadImpl()
 				RT_Proxy->TextureReferenceRHI = nullptr;
 				RT_Proxy->SamplerStateRHI = nullptr;
 			}
-			RT_Proxy->TexDims = RT_TexDims;
+			RT_Proxy->TexDims = FVector2D(RT_TexDims.X, RT_TexDims.Y);
 		}
 	);
 }
diff --git a/Engine/Plugins/FX/Niagara/Source/Niagara/Private/NiagaraDataInterfaceVolumeTexture.cpp b/Engine/Plugins/FX/Niagara/Source/Niagara/Private/NiagaraDataInterfaceVolumeTexture.cpp
--- a/Engine/Plugins/FX/Niagara/Source/Niagara/Private/NiagaraDataInterfaceVolumeTexture.cpp
+++ b/Engine/Plugins/FX/Niagara/Source/Niagara/Private/NiagaraDataInterfaceVolumeTexture.cpp
@@ -126,32 +126,29 @@ void UNiagaraDataInterfaceVolumeTexture::GetVMExternalFunction(const FVMExternal
 	}
 }
 
-void UNiagaraDataInterfaceVolumeTexture::GetTextureDimensions(FVectorVMContext& Context)
+bool UNiagaraDataInterfaceVolumeTexture::PerInstanceTick(void* PerInstanceData, FNiagaraSystemInstance* SystemInstance, float DeltaSeconds)
 {
-	VectorVM::FExternalFuncRegisterHandler<float> OutWidth(Context);
-	VectorVM::FExternalFuncRegisterHandler<float> OutHeight(Context);
-	VectorVM::FExternalFuncRegisterHandler<float> OutDepth(Context);
-
-	if (Texture == nullptr)
+	const FIntVector CurrentTextureSize = Texture != nullptr ? FIntVector(Texture->GetSizeX(), Texture->GetSizeY(), Texture->GetSizeZ()) : FIntVector::ZeroValue;
+	if ( CurrentTextureSize !=  TextureSize )
 	{
-		for (int32 i = 0; i < Context.NumInstances; ++i)
-		{
-			*OutWidth.GetDestAndAdvance() = 0.0f;
-			*OutHeight.GetDestAndAdvance() = 0.0f;
-			*OutDepth.GetDestAndAdvance() = 0.0f;
-		}
+		TextureSize = CurrentTextureSize;
+		MarkRenderDataDirty();
 	}
-	else
+	return false;
+}
+
+void UNiagaraDataInterfaceVolumeTexture::GetTextureDimensions(FVectorVMContext& Context)
+{
+	FNDIOutputParam<float> OutWidth(Context);
+	FNDIOutputParam<float> OutHeight(Context);
+	FNDIOutputParam<float> OutDepth(Context);
+
+	FVector FloatTextureSize(TextureSize.X, TextureSize.Y, TextureSize.Z);
+	for (int32 i = 0; i < Context.NumInstances; ++i)
 	{
-		float Width = Texture->GetSizeX();
-		float Height = Texture->GetSizeY();
-		float Depth = Texture->GetSizeZ();
-		for (int32 i = 0; i < Context.NumInstances; ++i)
-		{
-			*OutWidth.GetDestAndAdvance() = Width;
-			*OutHeight.GetDestAndAdvance() = Height;
-			*OutDepth.GetDestAndAdvance() = Depth;
-		}
+		OutWidth.SetAndAdvance(FloatTextureSize.X);
+		OutHeight.SetAndAdvance(FloatTextureSize.Y);
+		OutDepth.SetAndAdvance(FloatTextureSize.Z);
 	}
 }
 
@@ -281,21 +278,21 @@ void UNiagaraDataInterfaceVolumeTexture::PushToRenderThreadImpl()
 {
 	FNiagaraDataInterfaceProxyVolumeTexture* RT_Proxy = GetProxyAs<FNiagaraDataInterfaceProxyVolumeTexture>();
 
-	FVector RT_TexDims(EForceInit::ForceInitToZero);
+	TextureSize = FIntVector::ZeroValue;
 	if (Texture)
 	{
-		RT_TexDims.X = Texture->GetSizeX();
-		RT_TexDims.Y = Texture->GetSizeY();
-		RT_TexDims.Z = Texture->GetSizeZ();
+		TextureSize.X = Texture->GetSizeX();
+		TextureSize.Y = Texture->GetSizeY();
+		TextureSize.Z = Texture->GetSizeZ();
 	}
 
 	ENQUEUE_RENDER_COMMAND(FPushDITextureToRT)
 	(
-		[RT_Proxy, RT_Resource=Texture ? Texture->Resource : nullptr, RT_TexDims](FRHICommandListImmediate& RHICmdList)
+		[RT_Proxy, RT_Resource=Texture ? Texture->Resource : nullptr, RT_TexDims=TextureSize](FRHICommandListImmediate& RHICmdList)
 		{
 			RT_Proxy->TextureRHI = RT_Resource ? RT_Resource->TextureRHI : nullptr;
 			RT_Proxy->SamplerStateRHI = RT_Resource ? RT_Resource->SamplerStateRHI : nullptr;
-			RT_Proxy->TexDims = RT_TexDims;
+			RT_Proxy->TexDims = FVector(RT_TexDims.X, RT_TexDims.Y, RT_TexDims.Z);
 		}
 	);
 }
