Fixed issue where new DIs were'nt having thier state copied from the asset correctly due to a recent change making copying that asset data optional for SetAsset calls.
Now all new data is initialized and only existing parameters are subject to the optional skipping of this copy.

diff --git a/Engine/Plugins/FX/Niagara/Source/Niagara/Private/NiagaraComponent.cpp b/Engine/Plugins/FX/Niagara/Source/Niagara/Private/NiagaraComponent.cpp
--- a/Engine/Plugins/FX/Niagara/Source/Niagara/Private/NiagaraComponent.cpp
+++ b/Engine/Plugins/FX/Niagara/Source/Niagara/Private/NiagaraComponent.cpp
@@ -2529,26 +2529,25 @@ void UNiagaraComponent::FixDataInterfaceOuters()
 
 void UNiagaraComponent::CopyParametersFromAsset(bool bResetExistingOverrideParameters)
 {
+	TArray<FNiagaraVariable> ExistingVars;
+	OverrideParameters.GetParameters(ExistingVars);
+
 	TArray<FNiagaraVariable> SourceVars;
 	Asset->GetExposedParameters().GetParameters(SourceVars);
 	for (FNiagaraVariable& Param : SourceVars)
 	{
-		OverrideParameters.AddParameter(Param, true);
-	}
+		bool bNewParam = OverrideParameters.AddParameter(Param, true);
 
-	TArray<FNiagaraVariable> ExistingVars;
-	OverrideParameters.GetParameters(ExistingVars);
+		bool bCopyAssetValue = bResetExistingOverrideParameters || bNewParam;
+		if (bCopyAssetValue)
+		{
+			Asset->GetExposedParameters().CopyParameterData(OverrideParameters, Param);
+		}
+	}
 
 	for (FNiagaraVariable ExistingVar : ExistingVars)
 	{
-		if (SourceVars.Contains(ExistingVar))
-		{
-			if (bResetExistingOverrideParameters)
-			{
-				Asset->GetExposedParameters().CopyParameterData(OverrideParameters, ExistingVar);
-			}
-		}
-		else
+		if (!SourceVars.Contains(ExistingVar))
 		{
 			OverrideParameters.RemoveParameter(ExistingVar);
 		}
