Fix Engine.Emitter.SimulationPosition being backwards in respect to Emitter localspace flag. Now map directly to engine_owner_position instead of map.engine.owner.position.

diff --git a/Engine/Plugins/FX/Niagara/Source/NiagaraEditor/Private/NiagaraHlslTranslator.cpp b/Engine/Plugins/FX/Niagara/Source/NiagaraEditor/Private/NiagaraHlslTranslator.cpp
--- a/Engine/Plugins/FX/Niagara/Source/NiagaraEditor/Private/NiagaraHlslTranslator.cpp
+++ b/Engine/Plugins/FX/Niagara/Source/NiagaraEditor/Private/NiagaraHlslTranslator.cpp
@@ -5132,11 +5132,11 @@ bool FHlslNiagaraTranslator::HandleBoundConstantVariableToDataSetRead(FNiagaraVa
 		// Simulation position is 0 for localspace emitters.
 		// If we are not in localspace then this will not be a literal constant and is instead a default linked variable as handled in GenerateConstantString().
 		const bool bEmitterLocalSpace = CompileOptions.AdditionalDefines.Contains(SYS_PARAM_EMITTER_LOCALSPACE.GetName().ToString());
-		if (bEmitterLocalSpace)
+		if (bEmitterLocalSpace == false)
 		{
-			const FString ConstantStr = FString::Printf(TEXT("%s"), *SYS_PARAM_ENGINE_POSITION.GetName().ToString());
-			const FString ParameterMapInstanceNameStr = GetParameterMapInstanceName(InParamMapHistoryIdx);
-			Output = AddBodyChunk(GetUniqueSymbolName(TEXT("Constant")), FString::Printf(TEXT("%s.%s"), *ParameterMapInstanceNameStr, *ConstantStr), InVariable.GetType());
+			const FString SysParamEnginePositionSymbolNameString = GetSanitizedSymbolName(SYS_PARAM_ENGINE_POSITION.GetName().ToString(), true);
+			const FString ConstantStr = FString::Printf(TEXT("%s%s"), *SysParamEnginePositionSymbolNameString, *FString(TEXT(".xyz")));
+			Output = AddBodyChunk(GetUniqueSymbolName(TEXT("Constant")), ConstantStr, InVariable.GetType());
 			return true;
 		}
 		else
