Added some more reporting on mismatched event names, and also prevented Emitters from being ready if a valid event generator hasn't been found.

diff --git a/Engine/Plugins/FX/Niagara/Source/Niagara/Private/NiagaraEmitter.cpp b/Engine/Plugins/FX/Niagara/Source/Niagara/Private/NiagaraEmitter.cpp
--- a/Engine/Plugins/FX/Niagara/Source/Niagara/Private/NiagaraEmitter.cpp
+++ b/Engine/Plugins/FX/Niagara/Source/Niagara/Private/NiagaraEmitter.cpp
@@ -1009,13 +1009,28 @@ bool UNiagaraEmitter::IsValid()const
 		}
 		if (EventHandlerScriptProps.Num() != 0)
 		{
-			for (int32 i = 0; i < EventHandlerScriptProps.Num(); i++)
+			for (const FNiagaraEventScriptProperties& EventHandler : EventHandlerScriptProps)
 			{
-				if (!EventHandlerScriptProps[i].Script->IsScriptCompilationPending(false) &&
-					!EventHandlerScriptProps[i].Script->DidScriptCompilationSucceed(false))
+				if (!EventHandler.Script->IsScriptCompilationPending(false) &&
+					!EventHandler.Script->DidScriptCompilationSucceed(false))
 				{
 					return false;
 				}
+
+				const FNiagaraVMExecutableData& ScriptData = EventHandler.Script->GetVMExecutableData();
+
+				// validate that the read data sets for any events match the name that we've found
+				for (const FNiagaraDataSetID& ReadDataSet : ScriptData.ReadDataSets)
+				{
+					if (ReadDataSet.Type == ENiagaraDataSetType::Event
+					&& EventHandler.SourceEventName != NAME_None
+					&& ReadDataSet.Name != EventHandler.SourceEventName)
+					{
+						UE_LOG(LogNiagara, Warning, TEXT("%s being marked invalid because it has mismatched events (%s is expected, but %s has been found)."),
+							*GetPathName(), *ReadDataSet.Name.ToString(), *EventHandler.SourceEventName.ToString());
+						return false;
+					}
+				}
 			}
 		}
 	}
diff --git a/Engine/Plugins/FX/Niagara/Source/NiagaraEditor/Private/ViewModels/Stack/NiagaraStackEventScriptItemGroup.cpp b/Engine/Plugins/FX/Niagara/Source/NiagaraEditor/Private/ViewModels/Stack/NiagaraStackEventScriptItemGroup.cpp
--- a/Engine/Plugins/FX/Niagara/Source/NiagaraEditor/Private/ViewModels/Stack/NiagaraStackEventScriptItemGroup.cpp
+++ b/Engine/Plugins/FX/Niagara/Source/NiagaraEditor/Private/ViewModels/Stack/NiagaraStackEventScriptItemGroup.cpp
@@ -221,6 +221,31 @@ void UNiagaraStackEventScriptItemGroup::RefreshChildrenInternal(const TArray<UNi
 	if (EventScriptProperties != nullptr)
 	{
 		SetDisplayName(FText::Format(LOCTEXT("FormatEventScriptDisplayName", "Event Handler - Source: {0}"), FText::FromName(EventScriptProperties->SourceEventName)));
+
+		const bool HasValidScript = EventScriptProperties->Script
+			&& !EventScriptProperties->Script->IsScriptCompilationPending(false)
+			&& EventScriptProperties->Script->DidScriptCompilationSucceed(false);
+
+		if (HasValidScript && EventScriptProperties->SourceEventName != NAME_None)
+		{
+			const FNiagaraVMExecutableData& ScriptData = EventScriptProperties->Script->GetVMExecutableData();
+
+			// validate that the read data sets for any events match the name that we've found
+			for (const FNiagaraDataSetID& ReadDataSet : ScriptData.ReadDataSets)
+			{
+				if (ReadDataSet.Type == ENiagaraDataSetType::Event && ReadDataSet.Name != EventScriptProperties->SourceEventName)
+				{
+					FStackIssue EventMismatchError(
+						EStackIssueSeverity::Error,
+						FText::Format(LOCTEXT("EventMismatchSummaryText", "Event {0} can't be found"), FText::FromName(ReadDataSet.Name)),
+						FText::Format(LOCTEXT("EventMismatchText", "Looking for an event named {0} but the Emitter is currently using {1}."), FText::FromName(ReadDataSet.Name), FText::FromName(EventScriptProperties->SourceEventName)),
+						GetStackEditorDataKey(),
+						false);
+
+					NewIssues.Add(EventMismatchError);
+				}
+			}
+		}
 	}
 	else
 	{
