Rename the emitter's option for rapid iteration paramter baking to 'Supports Baked Rapid Iteration' to better represent what it's doing.

diff --git a/Engine/Plugins/FX/Niagara/Source/Niagara/Classes/NiagaraEmitter.h b/Engine/Plugins/FX/Niagara/Source/Niagara/Classes/NiagaraEmitter.h
--- a/Engine/Plugins/FX/Niagara/Source/Niagara/Classes/NiagaraEmitter.h
+++ b/Engine/Plugins/FX/Niagara/Source/Niagara/Classes/NiagaraEmitter.h
@@ -404,7 +404,7 @@ public:
 	TObjectPtr<class UNiagaraScriptSourceBase>	GraphSource;
 
 	/** Should we enable rapid iteration removal if the system is also set to remove rapid iteration parameters on compile? This value defaults to true.*/
-	UPROPERTY(EditAnywhere, AdvancedDisplay, Category = "Emitter")
+	UPROPERTY(EditAnywhere, AdvancedDisplay, Category = "Emitter", meta = (DisplayName = "Supports Baked Rapid Iteration"))
 	uint32 bBakeOutRapidIteration : 1;
 
 	bool NIAGARA_API AreAllScriptAndSourcesSynchronized() const;
