Speculative fix for build error

diff --git a/Engine/Plugins/FX/Niagara/Source/Niagara/Classes/NiagaraEmitter.h b/Engine/Plugins/FX/Niagara/Source/Niagara/Classes/NiagaraEmitter.h
--- a/Engine/Plugins/FX/Niagara/Source/Niagara/Classes/NiagaraEmitter.h
+++ b/Engine/Plugins/FX/Niagara/Source/Niagara/Classes/NiagaraEmitter.h
@@ -239,6 +239,7 @@ public:
 	virtual void PostRename(UObject* OldOuter, const FName OldName) override;
 	virtual void PostDuplicate(EDuplicateMode::Type DuplicateMode) override;
 	virtual void PostEditChangeProperty(struct FPropertyChangedEvent& PropertyChangedEvent) override;
+	virtual void PreSave(const class ITargetPlatform* TargetPlatform) override;
 	NIAGARA_API FOnPropertiesChanged& OnPropertiesChanged();
 	NIAGARA_API FOnRenderersChanged& OnRenderersChanged();
 	/** Helper method for when a rename has been detected within the graph. Covers renaming the internal renderer bindings.*/
@@ -396,7 +397,7 @@ public:
 	// Reference to the async update task that is responsible to do any outstanding work like emitter merging on the game thread
 	FGraphEventRef UpdateTaskRef;
 
-	void UpdateEmitterAfterLoad();
+	NIAGARA_API void UpdateEmitterAfterLoad();
 
 #if WITH_EDITORONLY_DATA
 	/** 'Source' data/graphs for the scripts used by this emitter. */
diff --git a/Engine/Plugins/FX/Niagara/Source/Niagara/Private/NiagaraEmitter.cpp b/Engine/Plugins/FX/Niagara/Source/Niagara/Private/NiagaraEmitter.cpp
--- a/Engine/Plugins/FX/Niagara/Source/Niagara/Private/NiagaraEmitter.cpp
+++ b/Engine/Plugins/FX/Niagara/Source/Niagara/Private/NiagaraEmitter.cpp
@@ -919,6 +919,15 @@ void UNiagaraEmitter::PostEditChangeProperty(struct FPropertyChangedEvent& Prope
 #endif
 }
 
+void UNiagaraEmitter::PreSave(const ITargetPlatform* TargetPlatform)
+{
+	Super::PreSave(TargetPlatform);
+	if (UpdateTaskRef.IsValid() && UpdateTaskRef->IsComplete() == false)
+	{
+		UpdateEmitterAfterLoad();
+	}
+}
+
 
 UNiagaraEmitter::FOnPropertiesChanged& UNiagaraEmitter::OnPropertiesChanged()
 {
@@ -1388,6 +1397,10 @@ UNiagaraEditorParametersAdapterBase* UNiagaraEmitter::GetEditorParameters()
 
 void UNiagaraEmitter::SetEditorData(UNiagaraEditorDataBase* InEditorData)
 {
+	if (EditorData == InEditorData)
+	{
+		return;
+	}
 	if (EditorData != nullptr)
 	{
 		EditorData->OnPersistentDataChanged().RemoveAll(this);
