Changed restoration of orbit mode to happen after the view transition is done rather than after the start of it, as it led to 'rubberbanding' where the camera would start focusing on the particles just to revert to its previous position right after. As there is no delegate "OnEnd" we have to make due with tick.

diff --git a/Engine/Plugins/FX/Niagara/Source/NiagaraEditor/Private/Widgets/SNiagaraSystemViewport.cpp b/Engine/Plugins/FX/Niagara/Source/NiagaraEditor/Private/Widgets/SNiagaraSystemViewport.cpp
--- a/Engine/Plugins/FX/Niagara/Source/NiagaraEditor/Private/Widgets/SNiagaraSystemViewport.cpp
+++ b/Engine/Plugins/FX/Niagara/Source/NiagaraEditor/Private/Widgets/SNiagaraSystemViewport.cpp
@@ -478,6 +478,18 @@ void SNiagaraSystemViewport::RefreshViewport()
 void SNiagaraSystemViewport::Tick( const FGeometry& AllottedGeometry, const double InCurrentTime, const float InDeltaTime )
 {
 	SEditorViewport::Tick( AllottedGeometry, InCurrentTime, InDeltaTime );
+
+	// this marks the end of the transition, so we restore orbit mode if needed
+	if(bIsViewTransitioning && !SystemViewportClient->GetViewTransform().IsPlaying())
+	{
+		if(bShouldActivateOrbitAfterTransitioning)
+		{
+			SystemViewportClient->ToggleOrbitCamera(true);
+		}
+
+		bShouldActivateOrbitAfterTransitioning = false;
+		bIsViewTransitioning = false;
+	}
 }
 
 void SNiagaraSystemViewport::SetPreviewComponent(UNiagaraComponent* NiagaraComponent)
@@ -591,7 +603,9 @@ void SNiagaraSystemViewport::OnFocusViewportToSelection()
 
 		SystemViewportClient->FocusViewportOnBox(PreviewComponent->Bounds.GetBox());
 
-		SystemViewportClient->ToggleOrbitCamera(bIsOrbit);
+		// this will reactivate orbit mode after the transition is done, if needed
+		bIsViewTransitioning = true;
+		bShouldActivateOrbitAfterTransitioning = bIsOrbit;
 	}
 }
 
diff --git a/Engine/Plugins/FX/Niagara/Source/NiagaraEditor/Private/Widgets/SNiagaraSystemViewport.h b/Engine/Plugins/FX/Niagara/Source/NiagaraEditor/Private/Widgets/SNiagaraSystemViewport.h
--- a/Engine/Plugins/FX/Niagara/Source/NiagaraEditor/Private/Widgets/SNiagaraSystemViewport.h
+++ b/Engine/Plugins/FX/Niagara/Source/NiagaraEditor/Private/Widgets/SNiagaraSystemViewport.h
@@ -131,6 +131,12 @@ private:
 	uint32 DrawFlags;
 
 	FOnThumbnailCaptured OnThumbnailCaptured;
+
+	/** Used on tick to determine if a view transition was active so we can restore view settings at the end of it */
+	bool bIsViewTransitioning = false;
+
+	/** True if orbit mode was active before we started a view transition. Used to restore orbit mode at the the end of the transition */
+	bool bShouldActivateOrbitAfterTransitioning = false;
 };
 
 #if NIAGARA_PERF_BASELINES
