Niagara - Add an emitters parent update task as a prerequisite to it's own update task so that it doesn't try to merge from a file that's not fully loaded and merged.  Also force the parent and parent at last merge to update before updating a child emitter.

diff --git a/Engine/Plugins/FX/Niagara/Source/Niagara/Private/NiagaraEmitter.cpp b/Engine/Plugins/FX/Niagara/Source/Niagara/Private/NiagaraEmitter.cpp
--- a/Engine/Plugins/FX/Niagara/Source/Niagara/Private/NiagaraEmitter.cpp
+++ b/Engine/Plugins/FX/Niagara/Source/Niagara/Private/NiagaraEmitter.cpp
@@ -565,9 +565,17 @@ void UNiagaraEmitter::PostLoad()
 	}
 
 	EnsureScriptsPostLoaded();
+
+	FGraphEventArray ParentPrerequisiteTasks;
+#if WITH_EDITORONLY_DATA
+	if (Parent != nullptr && Parent->UpdateTaskRef.IsValid())
+	{
+		ParentPrerequisiteTasks.Add(Parent->UpdateTaskRef);
+	}
+#endif
 	
 	// we are not yet finished, but we do the rest of the work after postload in a separate task
-	UpdateTaskRef = TGraphTask<FNiagaraEmitterUpdateTask>::CreateTask().ConstructAndDispatchWhenReady(this);
+	UpdateTaskRef = TGraphTask<FNiagaraEmitterUpdateTask>::CreateTask(&ParentPrerequisiteTasks).ConstructAndDispatchWhenReady(this);
 }
 
 bool UNiagaraEmitter::IsEditorOnly() const
@@ -1266,6 +1274,17 @@ void UNiagaraEmitter::UpdateEmitterAfterLoad()
 		Parent = nullptr;
 		ParentAtLastMerge = nullptr;
 	}
+
+	// The task prerequisites should have updated the parent emitter before updating this emitter, but if this emitter
+	// has been forced to update we need to make sure the parent has been updated too.
+	if (Parent != nullptr)
+	{
+		Parent->UpdateEmitterAfterLoad();
+	}
+	if (ParentAtLastMerge != nullptr)
+	{
+		ParentAtLastMerge->UpdateEmitterAfterLoad();
+	}
 	
 	if (!GetOutermost()->bIsCookedForEditor)
 	{
diff --git a/Engine/Plugins/FX/Niagara/Source/Niagara/Private/NiagaraSystem.cpp b/Engine/Plugins/FX/Niagara/Source/Niagara/Private/NiagaraSystem.cpp
--- a/Engine/Plugins/FX/Niagara/Source/Niagara/Private/NiagaraSystem.cpp
+++ b/Engine/Plugins/FX/Niagara/Source/Niagara/Private/NiagaraSystem.cpp
@@ -201,10 +201,10 @@ void UNiagaraSystem::PreSave(const class ITargetPlatform* TargetPlatform)
 void UNiagaraSystem::PreSave(FObjectPreSaveContext ObjectSaveContext)
 {
 	Super::PreSave(ObjectSaveContext);
+	EnsureFullyLoaded();
 #if WITH_EDITORONLY_DATA
 	WaitForCompilationComplete();
 #endif
-	EnsureFullyLoaded();
 }
 
 #if WITH_EDITOR
@@ -213,8 +213,8 @@ void UNiagaraSystem::BeginCacheForCookedPlatformData(const ITargetPlatform *Targ
 	//UE_LOG(LogNiagara, Display, TEXT("UNiagaraSystem::BeginCacheForCookedPlatformData %s %s"), *GetFullName(), GIsSavingPackage ? TEXT("Saving...") : TEXT("Not Saving..."));
 	Super::BeginCacheForCookedPlatformData(TargetPlatform);
 	
+	EnsureFullyLoaded();
 #if WITH_EDITORONLY_DATA
-
 	WaitForCompilationComplete();
 #endif
 }
