Removed old emitter merge call since it was moved to the update task after postload

diff --git a/Engine/Plugins/FX/Niagara/Source/Niagara/Private/NiagaraEmitterHandle.cpp b/Engine/Plugins/FX/Niagara/Source/Niagara/Private/NiagaraEmitterHandle.cpp
--- a/Engine/Plugins/FX/Niagara/Source/Niagara/Private/NiagaraEmitterHandle.cpp
+++ b/Engine/Plugins/FX/Niagara/Source/Niagara/Private/NiagaraEmitterHandle.cpp
@@ -177,13 +177,6 @@ void FNiagaraEmitterHandle::ConditionalPostLoad(int32 NiagaraCustomVersion)
 				Instance->ParentAtLastMerge->Rename(nullptr, Instance, REN_ForceNoResetLoaders);
 				LastMergedSource_DEPRECATED = nullptr;
 			}
-
-			// Since we've previously post loaded the emitter it wouldn't have merged on load since it didn't have
-			// the parent information so we check this again here, now that the parent information has been set.
-			if (Instance->IsSynchronizedWithParent() == false)
-			{
-				Instance->MergeChangesFromParent();
-			}
 		}
 
 		FText Reason;
