Fixed a crash caused by checking compile results when the rendering thread is suspended

diff --git a/Engine/Plugins/FX/Niagara/Source/NiagaraShader/Private/NiagaraShader.cpp b/Engine/Plugins/FX/Niagara/Source/NiagaraShader/Private/NiagaraShader.cpp
--- a/Engine/Plugins/FX/Niagara/Source/NiagaraShader/Private/NiagaraShader.cpp
+++ b/Engine/Plugins/FX/Niagara/Source/NiagaraShader/Private/NiagaraShader.cpp
@@ -932,7 +932,7 @@ bool FNiagaraShaderMap::IsNiagaraShaderComplete(const FNiagaraShaderScript* Scri
 
 bool FNiagaraShaderMap::IsComplete(const FNiagaraShaderScript* Script, bool bSilent)
 {
-	check(!GIsThreadedRendering || !IsInRenderingThread());
+	check(!GIsThreadedRendering || !IsInActualRenderingThread());
 	// Make sure we are operating on a referenced shader map or the below Find will cause this shader map to be deleted,
 	// Since it creates a temporary ref counted pointer.
 	check(NumRefs > 0);
