- Add Matrix4 support for RI bake out

diff --git a/Engine/Plugins/FX/Niagara/Source/NiagaraCore/Private/NiagaraCustomVersion.cpp b/Engine/Plugins/FX/Niagara/Source/NiagaraCore/Private/NiagaraCustomVersion.cpp
--- a/Engine/Plugins/FX/Niagara/Source/NiagaraCore/Private/NiagaraCustomVersion.cpp
+++ b/Engine/Plugins/FX/Niagara/Source/NiagaraCore/Private/NiagaraCustomVersion.cpp
@@ -8,4 +8,4 @@ const FGuid FNiagaraCustomVersion::GUID(0xFCF57AFA, 0x50764283, 0xB9A9E658, 0xFF
 // Register the custom version with core
 FCustomVersionRegistration GRegisterNiagaraCustomVersion(FNiagaraCustomVersion::GUID, FNiagaraCustomVersion::LatestVersion, TEXT("NiagaraVer"));
 
-const FGuid FNiagaraCustomVersion::LatestScriptCompileVersion(0xD19D9D89, 0xD23D422F, 0xA16E1C27, 0xC6B3CC2C);
+const FGuid FNiagaraCustomVersion::LatestScriptCompileVersion(0x38F94A17, 0xC44A414C, 0x9F8AFB2C, 0x605115EB);
diff --git a/Engine/Plugins/FX/Niagara/Source/NiagaraEditor/Private/NiagaraHlslTranslator.cpp b/Engine/Plugins/FX/Niagara/Source/NiagaraEditor/Private/NiagaraHlslTranslator.cpp
--- a/Engine/Plugins/FX/Niagara/Source/NiagaraEditor/Private/NiagaraHlslTranslator.cpp
+++ b/Engine/Plugins/FX/Niagara/Source/NiagaraEditor/Private/NiagaraHlslTranslator.cpp
@@ -4564,6 +4564,17 @@ FString FHlslNiagaraTranslator::GenerateConstantString(const FNiagaraVariable& C
 			int32* ValuePtr = (int32*)Constant.GetData();
 			ConstantStr = FString::Printf(TEXT("%d"), *ValuePtr);
 		}
+		else if (Type == FNiagaraTypeDefinition::GetMatrix4Def() )
+		{
+			float* ValuePtr = (float*)Constant.GetData();
+			ConstantStr = FString::Printf(
+				TEXT("float4x4(%g,%g,%g,%g, %g,%g,%g,%g, %g,%g,%g,%g, %g,%g,%g,%g)"),
+				ValuePtr[ 0], ValuePtr[ 1], ValuePtr[ 2], ValuePtr[ 3],
+				ValuePtr[ 4], ValuePtr[ 5], ValuePtr[ 6], ValuePtr[ 7],
+				ValuePtr[ 8], ValuePtr[ 9], ValuePtr[10], ValuePtr[11],
+				ValuePtr[12], ValuePtr[13], ValuePtr[14], ValuePtr[15]
+			);
+		}
 		else if (Type == FNiagaraTypeDefinition::GetBoolDef())
 		{
 			if (Constant.GetValue<FNiagaraBool>().IsValid() == false)
