fix for deprecation warning treated as error in RES

diff --git a/Engine/Plugins/FX/Niagara/Source/Niagara/Private/NiagaraEmitter.cpp b/Engine/Plugins/FX/Niagara/Source/Niagara/Private/NiagaraEmitter.cpp
--- a/Engine/Plugins/FX/Niagara/Source/Niagara/Private/NiagaraEmitter.cpp
+++ b/Engine/Plugins/FX/Niagara/Source/Niagara/Private/NiagaraEmitter.cpp
@@ -929,7 +929,9 @@ void UNiagaraEmitter::PostEditChangeProperty(struct FPropertyChangedEvent& Prope
 
 void UNiagaraEmitter::PreSave(const ITargetPlatform* TargetPlatform)
 {
+	PRAGMA_DISABLE_DEPRECATION_WARNINGS;
 	Super::PreSave(TargetPlatform);
+	PRAGMA_ENABLE_DEPRECATION_WARNINGS;
 	if (UpdateTaskRef.IsValid() && UpdateTaskRef->IsComplete() == false)
 	{
 		UpdateEmitterAfterLoad();
