Enabled double click on all emitter and system wizard asset pickers

diff --git a/Engine/Plugins/FX/Niagara/Source/NiagaraEditor/Private/SNewEmitterDialog.cpp b/Engine/Plugins/FX/Niagara/Source/NiagaraEditor/Private/SNewEmitterDialog.cpp
--- a/Engine/Plugins/FX/Niagara/Source/NiagaraEditor/Private/SNewEmitterDialog.cpp
+++ b/Engine/Plugins/FX/Niagara/Source/NiagaraEditor/Private/SNewEmitterDialog.cpp
@@ -28,10 +28,12 @@ void SNewEmitterDialog::Construct(const FArguments& InArgs)
 	TabOptions.ChangeTabState(ENiagaraScriptTemplateSpecification::Behavior, true);
 
 	SAssignNew(NewAssetPicker, SNiagaraAssetPickerList, UNiagaraEmitter::StaticClass())
+		.OnTemplateAssetActivated(this, &SNewEmitterDialog::ConfirmSelection)
 		.ViewOptions(DisplayAllViewOptions)
 		.TabOptions(TabOptions);
 
 	SAssignNew(CopyAssetPicker, SNiagaraAssetPickerList, UNiagaraEmitter::StaticClass())
+		.OnTemplateAssetActivated(this, &SNewEmitterDialog::ConfirmSelection)
 		.ViewOptions(DisplayAllViewOptions)
 		.TabOptions(TabOptions);
 	
diff --git a/Engine/Plugins/FX/Niagara/Source/NiagaraEditor/Private/SNewSystemDialog.cpp b/Engine/Plugins/FX/Niagara/Source/NiagaraEditor/Private/SNewSystemDialog.cpp
--- a/Engine/Plugins/FX/Niagara/Source/NiagaraEditor/Private/SNewSystemDialog.cpp
+++ b/Engine/Plugins/FX/Niagara/Source/NiagaraEditor/Private/SNewSystemDialog.cpp
@@ -46,10 +46,12 @@ void SNewSystemDialog::Construct(const FArguments& InArgs)
 		.bAllowMultiSelect(true);
 
 	SAssignNew(TemplateBehaviorAssetPicker, SNiagaraAssetPickerList, UNiagaraSystem::StaticClass())
+		.OnTemplateAssetActivated(this, &SNewSystemDialog::ConfirmSelection)
 		.ViewOptions(DisplayAllViewOptions)
 		.TabOptions(TemplateAndBehaviorsOnlyTabOptions);
 
 	SAssignNew(SystemAssetPicker, SNiagaraAssetPickerList, UNiagaraSystem::StaticClass())
+		.OnTemplateAssetActivated(this, &SNewSystemDialog::ConfirmSelection)
 		.ViewOptions(DisplayAllViewOptions);
 	
 	SNiagaraNewAssetDialog::Construct(SNiagaraNewAssetDialog::FArguments(), UNiagaraSystem::StaticClass()->GetFName(), LOCTEXT("AssetTypeName", "system"),
diff --git a/Engine/Plugins/FX/Niagara/Source/NiagaraEditor/Private/SNiagaraNewAssetDialog.cpp b/Engine/Plugins/FX/Niagara/Source/NiagaraEditor/Private/SNiagaraNewAssetDialog.cpp
--- a/Engine/Plugins/FX/Niagara/Source/NiagaraEditor/Private/SNiagaraNewAssetDialog.cpp
+++ b/Engine/Plugins/FX/Niagara/Source/NiagaraEditor/Private/SNiagaraNewAssetDialog.cpp
@@ -174,6 +174,16 @@ void SNiagaraNewAssetDialog::ConfirmSelection()
 	RequestDestroyWindow();
 }
 
+void SNiagaraNewAssetDialog::ConfirmSelection(const FAssetData& AssetData)
+{
+	SelectedAssets.Add(AssetData);
+	
+	const FNiagaraNewAssetDialogOption& SelectedOption = Options[SelectedOptionIndex];
+	SelectedOption.OnSelectionConfirmed.ExecuteIfBound();
+	bUserConfirmedSelection = true;
+	RequestDestroyWindow();
+}
+
 void SNiagaraNewAssetDialog::OnWindowClosed(const TSharedRef<SWindow>& Window)
 {
 	SaveConfig();
diff --git a/Engine/Plugins/FX/Niagara/Source/NiagaraEditor/Private/SNiagaraNewAssetDialog.h b/Engine/Plugins/FX/Niagara/Source/NiagaraEditor/Private/SNiagaraNewAssetDialog.h
--- a/Engine/Plugins/FX/Niagara/Source/NiagaraEditor/Private/SNiagaraNewAssetDialog.h
+++ b/Engine/Plugins/FX/Niagara/Source/NiagaraEditor/Private/SNiagaraNewAssetDialog.h
@@ -58,6 +58,8 @@ protected:
 	const TArray<FAssetData>& GetSelectedAssets() const;
 
 	void ConfirmSelection();
+	void ConfirmSelection(const FAssetData& AssetData);
+
 	int32 GetSelectedObjectIndex() const { return SelectedOptionIndex; };
 
 protected:
