Non-unity build fix for notes feature
Missing transaction header

diff --git a/Engine/Plugins/FX/Niagara/Source/NiagaraEditor/Private/NiagaraNodeFunctionCall.cpp b/Engine/Plugins/FX/Niagara/Source/NiagaraEditor/Private/NiagaraNodeFunctionCall.cpp
--- a/Engine/Plugins/FX/Niagara/Source/NiagaraEditor/Private/NiagaraNodeFunctionCall.cpp
+++ b/Engine/Plugins/FX/Niagara/Source/NiagaraEditor/Private/NiagaraNodeFunctionCall.cpp
@@ -26,6 +26,7 @@
 #include "ViewModels/Stack/NiagaraParameterHandle.h"
 #include "ViewModels/Stack/NiagaraStackGraphUtilities.h"
 #include "NiagaraNodeStaticSwitch.h"
+#include "ScopedTransaction.h"
 
 
 #define LOCTEXT_NAMESPACE "NiagaraNodeFunctionCall"
