Added default initialized MessageSeverity to FNiagaraStackMessage struct to fix a CIS issue

diff --git a/Engine/Plugins/FX/Niagara/Source/NiagaraEditor/Public/NiagaraMessages.h b/Engine/Plugins/FX/Niagara/Source/NiagaraEditor/Public/NiagaraMessages.h
--- a/Engine/Plugins/FX/Niagara/Source/NiagaraEditor/Public/NiagaraMessages.h
+++ b/Engine/Plugins/FX/Niagara/Source/NiagaraEditor/Public/NiagaraMessages.h
@@ -249,7 +249,7 @@ struct NIAGARAEDITOR_API FNiagaraStackMessage
 	FText ShortDescription;
 
 	UPROPERTY()
-	ENiagaraMessageSeverity MessageSeverity;
+	ENiagaraMessageSeverity MessageSeverity = ENiagaraMessageSeverity::Info;
 
 	UPROPERTY()
 	bool bAllowDismissal;
