Added default initialized bAllowDismissal to FNiagaraStackMessage struct to fix a CIS issue

diff --git a/Engine/Plugins/FX/Niagara/Source/NiagaraEditor/Public/NiagaraMessages.h b/Engine/Plugins/FX/Niagara/Source/NiagaraEditor/Public/NiagaraMessages.h
--- a/Engine/Plugins/FX/Niagara/Source/NiagaraEditor/Public/NiagaraMessages.h
+++ b/Engine/Plugins/FX/Niagara/Source/NiagaraEditor/Public/NiagaraMessages.h
@@ -252,7 +252,7 @@ struct NIAGARAEDITOR_API FNiagaraStackMessage
 	ENiagaraMessageSeverity MessageSeverity = ENiagaraMessageSeverity::Info;
 
 	UPROPERTY()
-	bool bAllowDismissal;
+	bool bAllowDismissal = true;
 
 	UPROPERTY()
 	FGuid Guid;
