Don't create full copies of scripts and graphs when compiling, reuse objects using an object pool

diff --git a/Engine/Plugins/FX/Niagara/Source/Niagara/Classes/NiagaraScript.h b/Engine/Plugins/FX/Niagara/Source/Niagara/Classes/NiagaraScript.h
--- a/Engine/Plugins/FX/Niagara/Source/Niagara/Classes/NiagaraScript.h
+++ b/Engine/Plugins/FX/Niagara/Source/Niagara/Classes/NiagaraScript.h
@@ -588,6 +588,9 @@ public:
 
 	/** Makes sure that the default version data is available and fixes old script assets. */
 	NIAGARA_API void CheckVersionDataAvailable();
+
+	/** Creates a shallow transient copy of this script for compilation purposes. */
+	NIAGARA_API UNiagaraScript* CreateCompilationCopy();
 #endif
 
 	// how this script is to be used. cannot be private due to use of GET_MEMBER_NAME_CHECKED
diff --git a/Engine/Plugins/FX/Niagara/Source/Niagara/Private/NiagaraScript.cpp b/Engine/Plugins/FX/Niagara/Source/Niagara/Private/NiagaraScript.cpp
--- a/Engine/Plugins/FX/Niagara/Source/Niagara/Private/NiagaraScript.cpp
+++ b/Engine/Plugins/FX/Niagara/Source/Niagara/Private/NiagaraScript.cpp
@@ -377,6 +377,23 @@ void UNiagaraScript::CheckVersionDataAvailable()
 	ExposedVersion = Data.Version.VersionGuid;
 }
 
+UNiagaraScript* UNiagaraScript::CreateCompilationCopy()
+{
+	UNiagaraScript* Result = NewObject<UNiagaraScript>();
+
+	// create a shallow copy
+	for (TFieldIterator<FProperty> PropertyIt(GetClass(), EFieldIteratorFlags::IncludeSuper); PropertyIt; ++PropertyIt)
+	{
+		FProperty* Property = *PropertyIt;
+		const uint8* SourceAddr = Property->ContainerPtrToValuePtr<uint8>(this);
+		uint8* DestinationAddr = Property->ContainerPtrToValuePtr<uint8>(Result);
+
+		Property->CopyCompleteValue(DestinationAddr, SourceAddr);
+	}
+	
+	return Result;
+}
+
 #endif
 
 bool FNiagaraVMExecutableDataId::IsValid() const
diff --git a/Engine/Plugins/FX/Niagara/Source/NiagaraEditor/Private/NiagaraCompiler.cpp b/Engine/Plugins/FX/Niagara/Source/NiagaraEditor/Private/NiagaraCompiler.cpp
--- a/Engine/Plugins/FX/Niagara/Source/NiagaraEditor/Private/NiagaraCompiler.cpp
+++ b/Engine/Plugins/FX/Niagara/Source/NiagaraEditor/Private/NiagaraCompiler.cpp
@@ -221,14 +221,12 @@ void FNiagaraCompileRequestData::VisitReferencedGraphsRecursive(UNiagaraGraph* I
 							}
 							else
 							{
-								DupeScript = bNeedsCompilation ? PrecompileDuplicateObject<UNiagaraScript>(FunctionScript, InNode, FunctionScript->GetFName()) : FunctionScript;
+								DupeScript = FunctionScript;
 								if (bNeedsCompilation)
 								{
-									if (DupeScript->GetSource(FunctionCallNode->SelectedScriptVersion)->GetOuter() != DupeScript)
-									{
-										UNiagaraScriptSource* DupeScriptSource = CastChecked<UNiagaraScriptSource>(DupeScript->GetSource(FunctionCallNode->SelectedScriptVersion));
-										DupeScript->SetSource(DuplicateObject<UNiagaraScriptSource>(DupeScriptSource, DupeScript, DupeScript->GetLatestSource()->GetFName()), FunctionCallNode->SelectedScriptVersion);
-									}
+									DupeScript = FunctionScript->CreateCompilationCopy();
+									UNiagaraScriptSource* DupeScriptSource = CastChecked<UNiagaraScriptSource>(DupeScript->GetSource(FunctionCallNode->SelectedScriptVersion))->CreateCompilationCopy();
+									DupeScript->SetSource(DupeScriptSource, FunctionCallNode->SelectedScriptVersion);
 								}
 								ProcessedGraph = Cast<UNiagaraScriptSource>(DupeScript->GetSource(FunctionCallNode->SelectedScriptVersion))->NodeGraph;
 								if (bNeedsCompilation)
@@ -252,11 +250,16 @@ void FNiagaraCompileRequestData::VisitReferencedGraphsRecursive(UNiagaraGraph* I
 							PreprocessedFunctions.Add(FunctionGraph, MadeArray);
 							VisitReferencedGraphsRecursive(ProcessedGraph, ConstantResolver, bNeedsCompilation);
 							ClonedGraphs.AddUnique(ProcessedGraph);
-
 						}
 						else if (bHasNumericParams)
 						{
-							UNiagaraScript* DupeScript = bNeedsCompilation ? PrecompileDuplicateObject<UNiagaraScript>(FunctionScript, InNode, FunctionScript->GetFName()) : FunctionScript;
+							UNiagaraScript* DupeScript = FunctionScript;
+							if (bNeedsCompilation)
+							{
+								DupeScript = FunctionScript->CreateCompilationCopy();
+								UNiagaraScriptSource* DupeScriptSource = CastChecked<UNiagaraScriptSource>(DupeScript->GetSource(FunctionCallNode->SelectedScriptVersion))->CreateCompilationCopy();
+								DupeScript->SetSource(DupeScriptSource, FunctionCallNode->SelectedScriptVersion);
+							}
 							ProcessedGraph = Cast<UNiagaraScriptSource>(DupeScript->GetSource(FunctionCallNode->SelectedScriptVersion))->NodeGraph;
 							if (bNeedsCompilation)
 							{
@@ -298,7 +301,7 @@ void FNiagaraCompileRequestData::VisitReferencedGraphsRecursive(UNiagaraGraph* I
 					{
 						EmitterNode->SyncEnabledState(); // Just to be safe, sync here while we likely still have the handle source.
 						EmitterNode->SetOwnerSystem(nullptr);
-						EmitterNode->SetCachedVariablesForCompilation(*Ptr->EmitterUniqueName, Ptr->NodeGraphDeepCopy, Ptr->Source);
+						EmitterNode->SetCachedVariablesForCompilation(*Ptr->EmitterUniqueName, Ptr->NodeGraphDeepCopy.Get(), Ptr->Source);
 					}
 				}
 			}
@@ -329,7 +332,7 @@ FName FNiagaraCompileRequestData::ResolveEmitterAlias(FName VariableName) const
 
 void FNiagaraCompileRequestData::GetReferencedObjects(TArray<UObject*>& Objects)
 {
-	Objects.Add(NodeGraphDeepCopy);
+	Objects.Add(NodeGraphDeepCopy.Get());
 	TArray<UNiagaraDataInterface*> DIs;
 	CopiedDataInterfacesByName.GenerateValueArray(DIs);
 	for (UNiagaraDataInterface* DI : DIs)
@@ -360,6 +363,21 @@ void FNiagaraCompileRequestData::GetReferencedObjects(TArray<UObject*>& Objects)
 	}
 }
 
+FNiagaraCompileRequestData::~FNiagaraCompileRequestData()
+{
+	if (NodeGraphDeepCopy.IsValid())
+	{
+		NodeGraphDeepCopy->ReleaseCompilationCopy();
+	}
+	for (TWeakObjectPtr<UNiagaraGraph> Graph : ClonedGraphs)
+	{
+		if (Graph.IsValid())
+		{
+			Graph->ReleaseCompilationCopy();
+		}
+	}
+}
+
 bool FNiagaraCompileRequestData::GatherPreCompiledVariables(const FString& InNamespaceFilter, TArray<FNiagaraVariable>& OutVars)
 {
 	if (PrecompiledHistories.Num() == 0)
@@ -389,15 +407,15 @@ void FNiagaraCompileRequestData::DeepCopyGraphs(UNiagaraScriptSource* ScriptSour
 {
 	// Clone the source graph so we can modify it as needed; merging in the child graphs
 	//NodeGraphDeepCopy = CastChecked<UNiagaraGraph>(FEdGraphUtilities::CloneGraph(ScriptSource->NodeGraph, GetTransientPackage(), 0));
-	Source = bNeedsCompilation ? PrecompileDuplicateObject<UNiagaraScriptSource>(ScriptSource, GetTransientPackage()) : ScriptSource;
+	Source = bNeedsCompilation ? ScriptSource->CreateCompilationCopy() : ScriptSource;
 	NodeGraphDeepCopy = Source->NodeGraph;
 	
 	//double StartTime = FPlatformTime::Seconds();
 	if (bNeedsCompilation)
 	{
-		FEdGraphUtilities::MergeChildrenGraphsIn(NodeGraphDeepCopy, NodeGraphDeepCopy, /*bRequireSchemaMatch=*/ true);
+		FEdGraphUtilities::MergeChildrenGraphsIn(NodeGraphDeepCopy.Get(), NodeGraphDeepCopy.Get(), /*bRequireSchemaMatch=*/ true);
 	}
-	VisitReferencedGraphs(ScriptSource->NodeGraph, NodeGraphDeepCopy, InUsage, ConstantResolver, bNeedsCompilation);
+	VisitReferencedGraphs(ScriptSource->NodeGraph, NodeGraphDeepCopy.Get(), InUsage, ConstantResolver, bNeedsCompilation);
 	//float DeltaTime = (float)(FPlatformTime::Seconds() - StartTime);
 	//if (DeltaTime > 0.01f)
 	//{
@@ -517,7 +535,7 @@ void FNiagaraCompileRequestData::FinishPrecompile(UNiagaraScriptSource* ScriptSo
 				// Map all for this output node
 				FNiagaraParameterMapHistoryWithMetaDataBuilder Builder;
 				Builder.ConstantResolver = ConstantResolver;
-				Builder.AddGraphToCallingGraphContextStack(NodeGraphDeepCopy);
+				Builder.AddGraphToCallingGraphContextStack(NodeGraphDeepCopy.Get());
 				Builder.RegisterEncounterableVariables(EncounterableVariables);
 
 				FString TranslationName = TEXT("Emitter");
@@ -627,6 +645,12 @@ void FNiagaraCompileRequestData::FinishPrecompile(UNiagaraScriptSource* ScriptSo
 				}
 			}
 		}
+
+		// clean up graph copies
+		for (auto& ClonedGraph : ClonedGraphs)
+		{
+			ClonedGraph->ReleaseCompilationCopy();
+		}
 	}
 
 	//float DeltaTime = (float)(FPlatformTime::Seconds() - StartTime);
diff --git a/Engine/Plugins/FX/Niagara/Source/NiagaraEditor/Private/NiagaraEditorModule.cpp b/Engine/Plugins/FX/Niagara/Source/NiagaraEditor/Private/NiagaraEditorModule.cpp
--- a/Engine/Plugins/FX/Niagara/Source/NiagaraEditor/Private/NiagaraEditorModule.cpp
+++ b/Engine/Plugins/FX/Niagara/Source/NiagaraEditor/Private/NiagaraEditorModule.cpp
@@ -755,6 +755,7 @@ void FNiagaraEditorModule::OnPreExit()
 			Sys->WaitForCompilationComplete();
 		}
 	}
+	ClearObjectPool();
 }
 
 void FNiagaraEditorModule::StartupModule()
@@ -1322,6 +1323,54 @@ TSharedRef<FNiagaraScriptMergeManager> FNiagaraEditorModule::GetScriptMergeManag
 	return ScriptMergeManager.ToSharedRef();
 }
 
+UObject* FNiagaraEditorModule::GetPooledDuplicateObject(UObject* Source, EFieldIteratorFlags::SuperClassFlags CopySuperProperties)
+{
+	check(IsInGameThread());
+	TArray<UObject*>& Pool = ObjectPool.FindOrAdd(Source->GetClass());
+	UObject* OutPooledObj;
+	if (Pool.Num() > 0)
+	{
+		OutPooledObj = Pool.Pop();
+		for (TFieldIterator<FProperty> PropertyIt(OutPooledObj->GetClass(), CopySuperProperties); PropertyIt; ++PropertyIt)
+		{
+			FProperty* Property = *PropertyIt;
+			const uint8* SourceAddr = Property->ContainerPtrToValuePtr<uint8>(Source);
+			uint8* DestinationAddr = Property->ContainerPtrToValuePtr<uint8>(OutPooledObj);
+			Property->CopyCompleteValue(DestinationAddr, SourceAddr);
+		}
+	}
+	else
+	{
+		OutPooledObj = DuplicateObject(Source, GetTransientPackage());
+		OutPooledObj->AddToRoot();
+	}
+	return OutPooledObj;
+}
+
+void FNiagaraEditorModule::ReleaseObjectToPool(UObject* Obj)
+{
+	check(IsInGameThread());
+	if (Obj->GetOuter() == GetTransientPackage())
+	{
+		Obj->AddToRoot();
+		ObjectPool.FindOrAdd(Obj->GetClass()).Push(Obj);
+	}
+}
+
+void FNiagaraEditorModule::ClearObjectPool()
+{
+	check(IsInGameThread());
+	for (auto& Pool : ObjectPool)
+	{
+		for (UObject* Var : Pool.Value)
+		{
+			Var->RemoveFromRoot();
+			Var->MarkPendingKill();
+		}
+	}
+	ObjectPool.Empty();
+}
+
 void FNiagaraEditorModule::RegisterParameterTrackCreatorForType(const UScriptStruct& StructType, FOnCreateMovieSceneTrackForParameter CreateTrack)
 {
 	checkf(TypeToParameterTrackCreatorMap.Contains(&StructType) == false, TEXT("Type already registered"));
diff --git a/Engine/Plugins/FX/Niagara/Source/NiagaraEditor/Private/NiagaraGraph.cpp b/Engine/Plugins/FX/Niagara/Source/NiagaraEditor/Private/NiagaraGraph.cpp
--- a/Engine/Plugins/FX/Niagara/Source/NiagaraEditor/Private/NiagaraGraph.cpp
+++ b/Engine/Plugins/FX/Niagara/Source/NiagaraEditor/Private/NiagaraGraph.cpp
@@ -29,6 +29,7 @@
 #include "Misc/SecureHash.h"
 #include "HAL/FileManager.h"
 #include "Misc/Paths.h"
+#include "Serialization/DuplicatedDataWriter.h"
 #include "String/ParseTokens.h"
 #include "ViewModels/NiagaraParameterDefinitionsSubscriberViewModel.h"
 
@@ -776,6 +777,12 @@ void UNiagaraGraph::PostEditChangeProperty(FPropertyChangedEvent& PropertyChange
 	RefreshParameterReferences();
 }
 
+void UNiagaraGraph::BeginDestroy()
+{
+	Super::BeginDestroy();
+	ReleaseCompilationCopy();
+}
+
 class UNiagaraScriptSource* UNiagaraGraph::GetSource() const
 {
 	return CastChecked<UNiagaraScriptSource>(GetOuter());
@@ -939,6 +946,70 @@ void UNiagaraGraph::FindEquivalentOutputNodes(ENiagaraScriptUsage TargetUsageTyp
 	OutputNodes = NodesFound;
 }
 
+UNiagaraGraph* UNiagaraGraph::CreateCompilationCopy()
+{
+	UNiagaraGraph* Result = NewObject<UNiagaraGraph>();
+	FNiagaraEditorModule& NiagaraEditorModule = FModuleManager::GetModuleChecked<FNiagaraEditorModule>("NiagaraEditor");
+
+	// create a shallow copy
+	for (TFieldIterator<FProperty> PropertyIt(GetClass(), EFieldIteratorFlags::IncludeSuper); PropertyIt; ++PropertyIt)
+	{
+		FProperty* Property = *PropertyIt;
+		const uint8* SourceAddr = Property->ContainerPtrToValuePtr<uint8>(this);
+        uint8* DestinationAddr = Property->ContainerPtrToValuePtr<uint8>(Result);
+
+        Property->CopyCompleteValue(DestinationAddr, SourceAddr);
+	}
+	Result->bIsForCompilationOnly = true;
+
+	// get new script variables from the pool
+	for (auto& It : VariableToScriptVariable)
+	{
+		It.Value = CastChecked<UNiagaraScriptVariable>(NiagaraEditorModule.GetPooledDuplicateObject(It.Value));
+	}
+
+	// duplicate the nodes
+	TMap<UEdGraphNode*, UEdGraphNode*> DuplicationMapping;
+	TArray<UEdGraphNode*> NewNodes;
+	for (UEdGraphNode* Node : Result->Nodes)
+	{
+		UEdGraphNode* DupNode = NewNodes.Add_GetRef(DuplicateObject(Node, Result));
+		DuplicationMapping.Add(Node, DupNode);
+	}
+	for (UEdGraphNode* Node : NewNodes)
+	{
+		// fix up linked pins
+		for (UEdGraphPin* Pin : Node->Pins)
+		{
+			for (int i = 0; i < Pin->LinkedTo.Num(); i++)
+			{
+				UEdGraphPin* LinkedPin = Pin->LinkedTo[i];
+				UEdGraphNode* NewNode = DuplicationMapping[LinkedPin->GetOwningNode()];
+				UEdGraphPin** NodePin = NewNode->Pins.FindByPredicate([LinkedPin](UEdGraphPin* Pin) { return Pin->PinId == LinkedPin->PinId; });
+				check(NodePin);
+				Pin->LinkedTo[i] = *NodePin;
+			}
+		}
+	}
+	Result->Nodes = NewNodes;
+	
+	return Result;
+}
+
+void UNiagaraGraph::ReleaseCompilationCopy()
+{
+	if (!bIsForCompilationOnly)
+	{
+		return;
+	}
+	FNiagaraEditorModule& NiagaraEditorModule = FModuleManager::GetModuleChecked<FNiagaraEditorModule>("NiagaraEditor");
+	for (auto It : VariableToScriptVariable)
+	{
+		NiagaraEditorModule.ReleaseObjectToPool(It.Value);
+	}
+	VariableToScriptVariable.Empty();
+}
+
 UNiagaraNodeOutput* UNiagaraGraph::FindOutputNode(ENiagaraScriptUsage TargetUsageType, FGuid TargetUsageId) const
 {
 	NIAGARA_SCOPE_CYCLE_COUNTER(STAT_NiagaraEditor_Graph_FindOutputNode);
diff --git a/Engine/Plugins/FX/Niagara/Source/NiagaraEditor/Private/NiagaraHlslTranslator.cpp b/Engine/Plugins/FX/Niagara/Source/NiagaraEditor/Private/NiagaraHlslTranslator.cpp
--- a/Engine/Plugins/FX/Niagara/Source/NiagaraEditor/Private/NiagaraHlslTranslator.cpp
+++ b/Engine/Plugins/FX/Niagara/Source/NiagaraEditor/Private/NiagaraHlslTranslator.cpp
@@ -1169,9 +1169,9 @@ const FNiagaraTranslateResults &FHlslNiagaraTranslator::Translate(const FNiagara
 	TranslateResults.bHLSLGenSucceeded = false;
 	TranslateResults.OutputHLSL = "";
 
-	UNiagaraGraph* SourceGraph = CompileData->NodeGraphDeepCopy;
+	TWeakObjectPtr<UNiagaraGraph> SourceGraph = CompileData->NodeGraphDeepCopy;
 
-	if (!SourceGraph)
+	if (!SourceGraph.IsValid())
 	{
 		Error(LOCTEXT("GetGraphFail", "Cannot find graph node!"), nullptr, nullptr);
 		return TranslateResults;
diff --git a/Engine/Plugins/FX/Niagara/Source/NiagaraEditor/Private/NiagaraHlslTranslator.h b/Engine/Plugins/FX/Niagara/Source/NiagaraEditor/Private/NiagaraHlslTranslator.h
--- a/Engine/Plugins/FX/Niagara/Source/NiagaraEditor/Private/NiagaraHlslTranslator.h
+++ b/Engine/Plugins/FX/Niagara/Source/NiagaraEditor/Private/NiagaraHlslTranslator.h
@@ -66,7 +66,7 @@ public:
 
 	}
 
-	virtual ~FNiagaraCompileRequestData() {}
+	virtual ~FNiagaraCompileRequestData();
 	virtual bool GatherPreCompiledVariables(const FString& InNamespaceFilter, TArray<FNiagaraVariable>& OutVars) override;
 	virtual void GetReferencedObjects(TArray<UObject*>& Objects) override;
 	virtual const TMap<FName, UNiagaraDataInterface*>& GetObjectNameMap() override;
@@ -75,7 +75,7 @@ public:
 
 	TArray<FNiagaraParameterMapHistory>& GetPrecomputedHistories() { return PrecompiledHistories; }
 	const TArray<FNiagaraParameterMapHistory>& GetPrecomputedHistories() const { return PrecompiledHistories; }
-	virtual const class UNiagaraGraph* GetPrecomputedNodeGraph() const { return NodeGraphDeepCopy; }
+	virtual const class UNiagaraGraph* GetPrecomputedNodeGraph() const { return NodeGraphDeepCopy.Get(); }
 	const FString& GetUniqueEmitterName() const { return EmitterUniqueName; }
 	void VisitReferencedGraphs(UNiagaraGraph* InSrcGraph, UNiagaraGraph* InDupeGraph, ENiagaraScriptUsage InUsage, FCompileConstantResolver ConstantResolver, bool bNeedsCompilation);
 	void DeepCopyGraphs(UNiagaraScriptSource* ScriptSource, ENiagaraScriptUsage InUsage, FCompileConstantResolver ConstantResolver, bool bNeedsCompilation);
@@ -97,9 +97,7 @@ public:
 	TArray<FGuid> StageGuids;
 	TArray<FName> StageNames;
 
-	// If this is being held onto for any length of time, make sure to hold onto it in a gc-aware object. Right now in this information-passing struct,
-	// we could have a leaked garbage collected pointer if not held onto by someone capable of registering a reference.
-	UNiagaraGraph* NodeGraphDeepCopy;
+	TWeakObjectPtr<UNiagaraGraph> NodeGraphDeepCopy;
 	TArray<FNiagaraParameterMapHistory> PrecompiledHistories;
 	TArray<FNiagaraVariable> ChangedFromNumericVars;
 	TMap<FName, UNiagaraDataInterface*> CopiedDataInterfacesByName;
@@ -126,7 +124,7 @@ public:
 		bool bHasNumericInputs;
 	};
 	TMap<const UNiagaraGraph*, TArray<FunctionData>> PreprocessedFunctions;
-	TArray<UNiagaraGraph*> ClonedGraphs;
+	TArray<TWeakObjectPtr<UNiagaraGraph>> ClonedGraphs;
 
 	// Copy of the variables that are required by Renderers if this is a compile request for an Emitter. 
 	TArray<FNiagaraVariable> RequiredRendererVariables;
diff --git a/Engine/Plugins/FX/Niagara/Source/NiagaraEditor/Private/NiagaraScriptSource.cpp b/Engine/Plugins/FX/Niagara/Source/NiagaraEditor/Private/NiagaraScriptSource.cpp
--- a/Engine/Plugins/FX/Niagara/Source/NiagaraEditor/Private/NiagaraScriptSource.cpp
+++ b/Engine/Plugins/FX/Niagara/Source/NiagaraEditor/Private/NiagaraScriptSource.cpp
@@ -104,6 +104,13 @@ void UNiagaraScriptSource::PostLoad()
 	}
 }
 
+UNiagaraScriptSource* UNiagaraScriptSource::CreateCompilationCopy()
+{
+	UNiagaraScriptSource* Result = NewObject<UNiagaraScriptSource>();
+	Result->NodeGraph = NodeGraph->CreateCompilationCopy();
+	return Result;
+}
+
 bool UNiagaraScriptSource::IsSynchronized(const FGuid& InChangeId)
 {
 	if (NodeGraph)
diff --git a/Engine/Plugins/FX/Niagara/Source/NiagaraEditor/Private/NiagaraScriptVariable.cpp b/Engine/Plugins/FX/Niagara/Source/NiagaraEditor/Private/NiagaraScriptVariable.cpp
--- a/Engine/Plugins/FX/Niagara/Source/NiagaraEditor/Private/NiagaraScriptVariable.cpp
+++ b/Engine/Plugins/FX/Niagara/Source/NiagaraEditor/Private/NiagaraScriptVariable.cpp
@@ -51,6 +51,18 @@ void UNiagaraScriptVariable::Init(const FNiagaraVariable& InVar, const FNiagaraV
 	}
 }
 
+void UNiagaraScriptVariable::InitFrom(UNiagaraScriptVariable* Value)
+{
+	for (TFieldIterator<FProperty> PropertyIt(GetClass(), EFieldIteratorFlags::ExcludeSuper); PropertyIt; ++PropertyIt)
+	{
+		FProperty* Property = *PropertyIt;
+		const uint8* SourceAddr = Property->ContainerPtrToValuePtr<uint8>(Value);
+		uint8* DestinationAddr = Property->ContainerPtrToValuePtr<uint8>(this);
+
+		Property->CopyCompleteValue(DestinationAddr, SourceAddr);
+	}
+}
+
 void UNiagaraScriptVariable::PostLoad()
 {
 	Super::PostLoad();
diff --git a/Engine/Plugins/FX/Niagara/Source/NiagaraEditor/Public/NiagaraEditorModule.h b/Engine/Plugins/FX/Niagara/Source/NiagaraEditor/Public/NiagaraEditorModule.h
--- a/Engine/Plugins/FX/Niagara/Source/NiagaraEditor/Public/NiagaraEditorModule.h
+++ b/Engine/Plugins/FX/Niagara/Source/NiagaraEditor/Public/NiagaraEditorModule.h
@@ -101,6 +101,11 @@ public:
 
 	TSharedRef<FNiagaraScriptMergeManager> GetScriptMergeManager() const;
 
+	// Object pooling methods used to prevent unnecessary object allocation during compiles
+	NIAGARAEDITOR_API UObject* GetPooledDuplicateObject(UObject* Source, EFieldIteratorFlags::SuperClassFlags CopySuperProperties = EFieldIteratorFlags::ExcludeSuper);
+	NIAGARAEDITOR_API void ReleaseObjectToPool(UObject* Obj);
+	NIAGARAEDITOR_API void ClearObjectPool();
+
 	void RegisterParameterTrackCreatorForType(const UScriptStruct& StructType, FOnCreateMovieSceneTrackForParameter CreateTrack);
 	void UnregisterParameterTrackCreatorForType(const UScriptStruct& StructType);
 	bool CanCreateParameterTrackForType(const UScriptStruct& StructType);
@@ -273,6 +278,8 @@ private:
 
 	TMap<FName, INiagaraStackObjectIssueGenerator*> StackIssueGenerators;
 
+	TMap<UClass*, TArray<UObject*>> ObjectPool;
+
 	// Parameter library assets reserve the name/type pairs of their parameters engine-wide.
 	mutable TSet<FName> ReservedLibraryParameterNames;
 
diff --git a/Engine/Plugins/FX/Niagara/Source/NiagaraEditor/Public/NiagaraGraph.h b/Engine/Plugins/FX/Niagara/Source/NiagaraEditor/Public/NiagaraGraph.h
--- a/Engine/Plugins/FX/Niagara/Source/NiagaraEditor/Public/NiagaraGraph.h
+++ b/Engine/Plugins/FX/Niagara/Source/NiagaraEditor/Public/NiagaraGraph.h
@@ -146,6 +146,7 @@ class UNiagaraGraph : public UEdGraph
 	//~ Begin UObject Interface
 	virtual void PostLoad() override;
 	virtual void PostEditChangeProperty(struct FPropertyChangedEvent& PropertyChangedEvent) override;
+	virtual void BeginDestroy() override;
 	//~ End UObjet Interface
 	
 	/** Get the source that owns this graph */
@@ -153,6 +154,10 @@ class UNiagaraGraph : public UEdGraph
 
 	/** Determine if there are any nodes in this graph.*/
 	bool IsEmpty() const { return Nodes.Num() == 0; }
+
+	/** Creates a transient copy of this graph for compilation purposes. */
+	UNiagaraGraph* CreateCompilationCopy();
+	void ReleaseCompilationCopy();
 			
 	/** Find the first output node bound to the target usage type.*/
 	class UNiagaraNodeOutput* FindOutputNode(ENiagaraScriptUsage TargetUsageType, FGuid TargetUsageId = FGuid()) const;
@@ -424,4 +429,6 @@ private:
 	bool bIsRenamingParameter;
 
 	mutable bool bParameterReferenceRefreshPending;
+
+	bool bIsForCompilationOnly = false;
 };
diff --git a/Engine/Plugins/FX/Niagara/Source/NiagaraEditor/Public/NiagaraScriptSource.h b/Engine/Plugins/FX/Niagara/Source/NiagaraEditor/Public/NiagaraScriptSource.h
--- a/Engine/Plugins/FX/Niagara/Source/NiagaraEditor/Public/NiagaraScriptSource.h
+++ b/Engine/Plugins/FX/Niagara/Source/NiagaraEditor/Public/NiagaraScriptSource.h
@@ -23,6 +23,8 @@ class UNiagaraScriptSource : public UNiagaraScriptSourceBase
 	// UObject interface
 	virtual void PostLoad() override;
 
+	UNiagaraScriptSource* CreateCompilationCopy();
+
 	// UNiagaraScriptSourceBase interface.
 	//virtual ENiagaraScriptCompileStatus Compile(UNiagaraScript* ScriptOwner, FString& OutGraphLevelErrorMessages) override;
 	virtual bool IsSynchronized(const FGuid& InChangeId) override;
diff --git a/Engine/Plugins/FX/Niagara/Source/NiagaraEditor/Public/NiagaraScriptVariable.h b/Engine/Plugins/FX/Niagara/Source/NiagaraEditor/Public/NiagaraScriptVariable.h
--- a/Engine/Plugins/FX/Niagara/Source/NiagaraEditor/Public/NiagaraScriptVariable.h
+++ b/Engine/Plugins/FX/Niagara/Source/NiagaraEditor/Public/NiagaraScriptVariable.h
@@ -23,6 +23,7 @@ public:
 	virtual void PostEditChangeProperty(struct FPropertyChangedEvent& PropertyChangedEvent) override;
 
 	void Init(const FNiagaraVariable& InVar, const FNiagaraVariableMetaData& InVarMetaData);
+	void InitFrom(UNiagaraScriptVariable* Value);
 
 	/** The default mode. Can be Value, Binding or Custom. */
 	UPROPERTY(EditAnywhere, Category = "Default Value")
