- Allow BP get & set a single element in a Niagara array

diff --git a/Engine/Plugins/FX/Niagara/Source/Niagara/Classes/NiagaraDataInterfaceArrayFloat.h b/Engine/Plugins/FX/Niagara/Source/Niagara/Classes/NiagaraDataInterfaceArrayFloat.h
--- a/Engine/Plugins/FX/Niagara/Source/Niagara/Classes/NiagaraDataInterfaceArrayFloat.h
+++ b/Engine/Plugins/FX/Niagara/Source/Niagara/Classes/NiagaraDataInterfaceArrayFloat.h
@@ -12,6 +12,8 @@ class NIAGARA_API UNiagaraDataInterfaceArrayFloat : public UNiagaraDataInterface
 public:
 	UPROPERTY(EditAnywhere, BlueprintReadWrite, Category = "Array")
 	TArray<float> FloatData;
+
+	TArray<float>& GetArrayReference() { return FloatData; }
 };
 
 UCLASS(EditInlineNew, Category = "Array", meta = (DisplayName = "Vector 2D Array", Experimental), Blueprintable, BlueprintType)
@@ -21,6 +23,8 @@ class NIAGARA_API UNiagaraDataInterfaceArrayFloat2 : public UNiagaraDataInterfac
 public:
 	UPROPERTY(EditAnywhere, BlueprintReadWrite, Category = "Array")
 	TArray<FVector2D> FloatData;
+
+	TArray<FVector2D>& GetArrayReference() { return FloatData; }
 };
 
 UCLASS(EditInlineNew, Category = "Array", meta = (DisplayName = "Vector Array", Experimental), Blueprintable, BlueprintType)
@@ -30,6 +34,8 @@ class NIAGARA_API UNiagaraDataInterfaceArrayFloat3 : public UNiagaraDataInterfac
 public:
 	UPROPERTY(EditAnywhere, BlueprintReadWrite, Category = "Array")
 	TArray<FVector> FloatData;
+
+	TArray<FVector>& GetArrayReference() { return FloatData; }
 };
 
 UCLASS(EditInlineNew, Category = "Array", meta = (DisplayName = "Vector 4 Array", Experimental), Blueprintable, BlueprintType)
@@ -39,6 +45,8 @@ class NIAGARA_API UNiagaraDataInterfaceArrayFloat4 : public UNiagaraDataInterfac
 public:
 	UPROPERTY(EditAnywhere, BlueprintReadWrite, Category = "Array")
 	TArray<FVector4> FloatData;
+
+	TArray<FVector4>& GetArrayReference() { return FloatData; }
 };
 
 UCLASS(EditInlineNew, Category = "Array", meta = (DisplayName = "Color Array", Experimental), Blueprintable, BlueprintType)
@@ -48,6 +56,8 @@ class NIAGARA_API UNiagaraDataInterfaceArrayColor : public UNiagaraDataInterface
 public:
 	UPROPERTY(EditAnywhere, BlueprintReadWrite, Category = "Array")
 	TArray<FLinearColor> ColorData;
+
+	TArray<FLinearColor>& GetArrayReference() { return ColorData; }
 };
 
 UCLASS(EditInlineNew, Category = "Array", meta = (DisplayName = "Quaternion Array", Experimental), Blueprintable, BlueprintType)
@@ -57,4 +67,6 @@ class NIAGARA_API UNiagaraDataInterfaceArrayQuat : public UNiagaraDataInterfaceA
 public:
 	UPROPERTY(EditAnywhere, BlueprintReadWrite, Category = "Array")
 	TArray<FQuat> QuatData;
+
+	TArray<FQuat>& GetArrayReference() { return QuatData; }
 };
diff --git a/Engine/Plugins/FX/Niagara/Source/Niagara/Classes/NiagaraDataInterfaceArrayFunctionLibrary.h b/Engine/Plugins/FX/Niagara/Source/Niagara/Classes/NiagaraDataInterfaceArrayFunctionLibrary.h
--- a/Engine/Plugins/FX/Niagara/Source/Niagara/Classes/NiagaraDataInterfaceArrayFunctionLibrary.h
+++ b/Engine/Plugins/FX/Niagara/Source/Niagara/Classes/NiagaraDataInterfaceArrayFunctionLibrary.h
@@ -62,4 +62,54 @@ class NIAGARA_API UNiagaraDataInterfaceArrayFunctionLibrary : public UBlueprintF
 	/** Gets a copy of Niagara Bool Data. */
 	UFUNCTION(BlueprintCallable, Category = Niagara, meta = (DisplayName = "Niagara Get Bool Array"))
 	static TArray<bool> GetNiagaraArrayBool(UNiagaraComponent* NiagaraSystem, FName OverrideName);
+
+	/** Sets a single value within a Niagara Array Float. */
+	UFUNCTION(BlueprintCallable, Category = Niagara, meta = (DisplayName = "Niagara Set Float Array Value"))
+	static void SetNiagaraArrayFloatValue(UNiagaraComponent* NiagaraSystem, FName OverrideName, int Index, float Value, bool bSizeToFit);
+	/** Sets a single value within a Niagara Array FVector2D. */
+	UFUNCTION(BlueprintCallable, Category = Niagara, meta = (DisplayName = "Niagara Set Vector 2D Array Value"))
+	static void SetNiagaraArrayVector2DValue(UNiagaraComponent* NiagaraSystem, FName OverrideName, int Index, const FVector2D& Value, bool bSizeToFit);
+	/** Sets a single value within a Niagara Array FVector. */
+	UFUNCTION(BlueprintCallable, Category = Niagara, meta = (DisplayName = "Niagara Set Vector Array Value"))
+	static void SetNiagaraArrayVectorValue(UNiagaraComponent* NiagaraSystem, FName OverrideName, int Index, const FVector& Value, bool bSizeToFit);
+	/** Sets a single value within a Niagara Array FVector4. */
+	UFUNCTION(BlueprintCallable, Category = Niagara, meta = (DisplayName = "Niagara Set Vector 4 Array Value"))
+	static void SetNiagaraArrayVector4Value(UNiagaraComponent* NiagaraSystem, FName OverrideName, int Index, const FVector4& Value, bool bSizeToFit);
+	/** Sets a single value within a Niagara Array FLinearColor. */
+	UFUNCTION(BlueprintCallable, Category = Niagara, meta = (DisplayName = "Niagara Set Color Array Value"))
+	static void SetNiagaraArrayColorValue(UNiagaraComponent* NiagaraSystem, FName OverrideName, int Index, const FLinearColor& Value, bool bSizeToFit);
+	/** Sets a single value within a Niagara Array FQuat. */
+	UFUNCTION(BlueprintCallable, Category = Niagara, meta = (DisplayName = "Niagara Set Quaternion Array Value"))
+	static void SetNiagaraArrayQuatValue(UNiagaraComponent* NiagaraSystem, FName OverrideName, int Index, const FQuat& Value, bool bSizeToFit);
+	/** Sets a single value within a Niagara Array Int32. */
+	UFUNCTION(BlueprintCallable, Category = Niagara, meta = (DisplayName = "Niagara Set Int32 Array Value"))
+	static void SetNiagaraArrayInt32Value(UNiagaraComponent* NiagaraSystem, FName OverrideName, int Index, int32 Value, bool bSizeToFit);
+	/** Sets a single value within a Niagara Array Bool. */
+	UFUNCTION(BlueprintCallable, Category = Niagara, meta = (DisplayName = "Niagara Set Bool Array Value"))
+	static void SetNiagaraArrayBoolValue(UNiagaraComponent* NiagaraSystem, FName OverrideName, int Index, const bool& Value, bool bSizeToFit);
+
+	/** Gets a single value within a Niagara Array Float. */
+	UFUNCTION(BlueprintCallable, Category = Niagara, meta = (DisplayName = "Niagara Get Float Array Value"))
+	static float GetNiagaraArrayFloatValue(UNiagaraComponent* NiagaraSystem, FName OverrideName, int Index);
+	/** Gets a single value within a Niagara Array FVector2D. */
+	UFUNCTION(BlueprintCallable, Category = Niagara, meta = (DisplayName = "Niagara Get Vector 2D Array Value"))
+	static FVector2D GetNiagaraArrayVector2DValue(UNiagaraComponent* NiagaraSystem, FName OverrideName, int Index);
+	/** Gets a single value within a Niagara Array FVector. */
+	UFUNCTION(BlueprintCallable, Category = Niagara, meta = (DisplayName = "Niagara Get Vector Array Value"))
+	static FVector GetNiagaraArrayVectorValue(UNiagaraComponent* NiagaraSystem, FName OverrideName, int Index);
+	/** Gets a single value within a Niagara Array FVector4. */
+	UFUNCTION(BlueprintCallable, Category = Niagara, meta = (DisplayName = "Niagara Get Vector 4 Array Value"))
+	static FVector4 GetNiagaraArrayVector4Value(UNiagaraComponent* NiagaraSystem, FName OverrideName, int Index);
+	/** Gets a single value within a Niagara Array FLinearColor. */
+	UFUNCTION(BlueprintCallable, Category = Niagara, meta = (DisplayName = "Niagara Get Color Array Value"))
+	static FLinearColor GetNiagaraArrayColorValue(UNiagaraComponent* NiagaraSystem, FName OverrideName, int Index);
+	/** Gets a single value within a Niagara Array FQuat. */
+	UFUNCTION(BlueprintCallable, Category = Niagara, meta = (DisplayName = "Niagara Get Quaternion Array Value"))
+	static FQuat GetNiagaraArrayQuatValue(UNiagaraComponent* NiagaraSystem, FName OverrideName, int Index);
+	/** Gets a single value within a Niagara Array Int32. */
+	UFUNCTION(BlueprintCallable, Category = Niagara, meta = (DisplayName = "Niagara Get Int32 Array Value"))
+	static int32 GetNiagaraArrayInt32Value(UNiagaraComponent* NiagaraSystem, FName OverrideName, int Index);
+	/** Gets a single value within a Niagara Array Bool. */
+	UFUNCTION(BlueprintCallable, Category = Niagara, meta = (DisplayName = "Niagara Get Bool Array Value"))
+	static bool GetNiagaraArrayBoolValue(UNiagaraComponent* NiagaraSystem, FName OverrideName, int Index);
 };
diff --git a/Engine/Plugins/FX/Niagara/Source/Niagara/Classes/NiagaraDataInterfaceArrayInt.h b/Engine/Plugins/FX/Niagara/Source/Niagara/Classes/NiagaraDataInterfaceArrayInt.h
--- a/Engine/Plugins/FX/Niagara/Source/Niagara/Classes/NiagaraDataInterfaceArrayInt.h
+++ b/Engine/Plugins/FX/Niagara/Source/Niagara/Classes/NiagaraDataInterfaceArrayInt.h
@@ -12,6 +12,8 @@ class NIAGARA_API UNiagaraDataInterfaceArrayInt32 : public UNiagaraDataInterface
 public:
 	UPROPERTY(EditAnywhere, BlueprintReadWrite, Category = "Array")
 	TArray<int32> IntData;
+
+	TArray<int32>& GetArrayReference() { return IntData; }
 };
 
 UCLASS(EditInlineNew, Category = "Array", meta = (DisplayName = "Bool Array", Experimental), Blueprintable, BlueprintType)
@@ -21,4 +23,6 @@ class NIAGARA_API UNiagaraDataInterfaceArrayBool : public UNiagaraDataInterfaceA
 public:
 	UPROPERTY(EditAnywhere, BlueprintReadWrite, Category = "Array")
 	TArray<bool> BoolData;
+
+	TArray<bool>& GetArrayReference() { return BoolData; }
 };
diff --git a/Engine/Plugins/FX/Niagara/Source/Niagara/Private/NiagaraDataInterfaceArrayFunctionLibrary.cpp b/Engine/Plugins/FX/Niagara/Source/Niagara/Private/NiagaraDataInterfaceArrayFunctionLibrary.cpp
--- a/Engine/Plugins/FX/Niagara/Source/Niagara/Private/NiagaraDataInterfaceArrayFunctionLibrary.cpp
+++ b/Engine/Plugins/FX/Niagara/Source/Niagara/Private/NiagaraDataInterfaceArrayFunctionLibrary.cpp
@@ -3,165 +3,229 @@
 #include "NiagaraDataInterfaceArrayFunctionLibrary.h"
 #include "NiagaraDataInterfaceArrayFloat.h"
 #include "NiagaraDataInterfaceArrayInt.h"
+#include "NiagaraDataInterfaceArrayImpl.h"
 #include "NiagaraFunctionLibrary.h"
 
-void UNiagaraDataInterfaceArrayFunctionLibrary::SetNiagaraArrayFloat(UNiagaraComponent* NiagaraSystem, FName OverrideName, const TArray<float>& ArrayData)
+template<typename TArrayType, typename TDataInterace>
+void SetNiagaraArray(UNiagaraComponent* NiagaraSystem, FName OverrideName, const TArray<TArrayType>& InArray)
 {
-	if (UNiagaraDataInterfaceArrayFloat* ArrayDI = UNiagaraFunctionLibrary::GetDataInterface<UNiagaraDataInterfaceArrayFloat>(NiagaraSystem, OverrideName))
+	if (TDataInterace* ArrayDI = UNiagaraFunctionLibrary::GetDataInterface<TDataInterace>(NiagaraSystem, OverrideName))
 	{
 		FRWScopeLock WriteLock(ArrayDI->ArrayRWGuard, SLT_Write);
-		ArrayDI->FloatData = ArrayData;
+		ArrayDI->GetArrayReference() = InArray;
 		ArrayDI->MarkRenderDataDirty();
 	}
 }
 
-void UNiagaraDataInterfaceArrayFunctionLibrary::SetNiagaraArrayVector2D(UNiagaraComponent* NiagaraSystem, FName OverrideName, const TArray<FVector2D>& ArrayData)
+template<typename TArrayType, typename TDataInterace>
+TArray<TArrayType> GetNiagaraArray(UNiagaraComponent* NiagaraSystem, FName OverrideName)
 {
-	if (UNiagaraDataInterfaceArrayFloat2* ArrayDI = UNiagaraFunctionLibrary::GetDataInterface<UNiagaraDataInterfaceArrayFloat2>(NiagaraSystem, OverrideName))
+	if (TDataInterace* ArrayDI = UNiagaraFunctionLibrary::GetDataInterface<TDataInterace>(NiagaraSystem, OverrideName))
 	{
-		FRWScopeLock WriteLock(ArrayDI->ArrayRWGuard, SLT_Write);
-		ArrayDI->FloatData = ArrayData;
-		ArrayDI->MarkRenderDataDirty();
+		FRWScopeLock WriteLock(ArrayDI->ArrayRWGuard, SLT_ReadOnly);
+		return ArrayDI->GetArrayReference();
 	}
+	return TArray<TArrayType>();
 }
 
-void UNiagaraDataInterfaceArrayFunctionLibrary::SetNiagaraArrayVector(UNiagaraComponent* NiagaraSystem, FName OverrideName, const TArray<FVector>& ArrayData)
+template<typename TArrayType, typename TDataInterace>
+void SetNiagaraArrayValue(UNiagaraComponent* NiagaraSystem, FName OverrideName, int Index, const TArrayType& Value, bool bSizeToFit)
 {
-	if (UNiagaraDataInterfaceArrayFloat3* ArrayDI = UNiagaraFunctionLibrary::GetDataInterface<UNiagaraDataInterfaceArrayFloat3>(NiagaraSystem, OverrideName))
+	if (TDataInterace* ArrayDI = UNiagaraFunctionLibrary::GetDataInterface<TDataInterace>(NiagaraSystem, OverrideName))
 	{
 		FRWScopeLock WriteLock(ArrayDI->ArrayRWGuard, SLT_Write);
-		ArrayDI->FloatData = ArrayData;
+		auto& ArrayData = ArrayDI->GetArrayReference();
+
+		if (!ArrayData.IsValidIndex(Index))
+		{
+			if (!bSizeToFit)
+			{
+				return;
+			}
+			ArrayData.AddDefaulted(Index + 1 - ArrayData.Num());
+		}
+
+		ArrayData[Index] = Value;
 		ArrayDI->MarkRenderDataDirty();
 	}
 }
 
-void UNiagaraDataInterfaceArrayFunctionLibrary::SetNiagaraArrayVector4(UNiagaraComponent* NiagaraSystem, FName OverrideName, const TArray<FVector4>& ArrayData)
+template<typename TArrayType, typename TDataInterace>
+TArrayType GetNiagaraArrayValue(UNiagaraComponent* NiagaraSystem, FName OverrideName, int Index)
 {
-	if (UNiagaraDataInterfaceArrayFloat4* ArrayDI = UNiagaraFunctionLibrary::GetDataInterface<UNiagaraDataInterfaceArrayFloat4>(NiagaraSystem, OverrideName))
+	if (TDataInterace* ArrayDI = UNiagaraFunctionLibrary::GetDataInterface<TDataInterace>(NiagaraSystem, OverrideName))
 	{
-		FRWScopeLock WriteLock(ArrayDI->ArrayRWGuard, SLT_Write);
-		ArrayDI->FloatData = ArrayData;
-		ArrayDI->MarkRenderDataDirty();
+		FRWScopeLock WriteLock(ArrayDI->ArrayRWGuard, SLT_ReadOnly);
+		auto& ArrayData = ArrayDI->GetArrayReference();
+		return ArrayData.IsValidIndex(Index) ? ArrayData[Index] : TArrayType();
 	}
+	return TArrayType();
+}
+
+//////////////////////////////////////////////////////////////////////////
+
+void UNiagaraDataInterfaceArrayFunctionLibrary::SetNiagaraArrayFloat(UNiagaraComponent* NiagaraSystem, FName OverrideName, const TArray<float>& ArrayData)
+{
+	SetNiagaraArray<float, UNiagaraDataInterfaceArrayFloat>(NiagaraSystem, OverrideName, ArrayData);
+}
+
+void UNiagaraDataInterfaceArrayFunctionLibrary::SetNiagaraArrayVector2D(UNiagaraComponent* NiagaraSystem, FName OverrideName, const TArray<FVector2D>& ArrayData)
+{
+	SetNiagaraArray<FVector2D, UNiagaraDataInterfaceArrayFloat2>(NiagaraSystem, OverrideName, ArrayData);
+}
+
+void UNiagaraDataInterfaceArrayFunctionLibrary::SetNiagaraArrayVector(UNiagaraComponent* NiagaraSystem, FName OverrideName, const TArray<FVector>& ArrayData)
+{
+	SetNiagaraArray<FVector, UNiagaraDataInterfaceArrayFloat3>(NiagaraSystem, OverrideName, ArrayData);
+}
+
+void UNiagaraDataInterfaceArrayFunctionLibrary::SetNiagaraArrayVector4(UNiagaraComponent* NiagaraSystem, FName OverrideName, const TArray<FVector4>& ArrayData)
+{
+	SetNiagaraArray<FVector4, UNiagaraDataInterfaceArrayFloat4>(NiagaraSystem, OverrideName, ArrayData);
 }
 
 void UNiagaraDataInterfaceArrayFunctionLibrary::SetNiagaraArrayColor(UNiagaraComponent* NiagaraSystem, FName OverrideName, const TArray<FLinearColor>& ArrayData)
 {
-	if (UNiagaraDataInterfaceArrayColor* ArrayDI = UNiagaraFunctionLibrary::GetDataInterface<UNiagaraDataInterfaceArrayColor>(NiagaraSystem, OverrideName))
-	{
-		FRWScopeLock WriteLock(ArrayDI->ArrayRWGuard, SLT_Write);
-		ArrayDI->ColorData = ArrayData;
-		ArrayDI->MarkRenderDataDirty();
-	}
+	SetNiagaraArray<FLinearColor, UNiagaraDataInterfaceArrayColor>(NiagaraSystem, OverrideName, ArrayData);
 }
 
 void UNiagaraDataInterfaceArrayFunctionLibrary::SetNiagaraArrayQuat(UNiagaraComponent* NiagaraSystem, FName OverrideName, const TArray<FQuat>& ArrayData)
 {
-	if (UNiagaraDataInterfaceArrayQuat* ArrayDI = UNiagaraFunctionLibrary::GetDataInterface<UNiagaraDataInterfaceArrayQuat>(NiagaraSystem, OverrideName))
-	{
-		FRWScopeLock WriteLock(ArrayDI->ArrayRWGuard, SLT_Write);
-		ArrayDI->QuatData = ArrayData;
-		ArrayDI->MarkRenderDataDirty();
-	}
+	SetNiagaraArray<FQuat, UNiagaraDataInterfaceArrayQuat>(NiagaraSystem, OverrideName, ArrayData);
 }
 
 void UNiagaraDataInterfaceArrayFunctionLibrary::SetNiagaraArrayInt32(UNiagaraComponent* NiagaraSystem, FName OverrideName, const TArray<int32>& ArrayData)
 {
-	if (UNiagaraDataInterfaceArrayInt32* ArrayDI = UNiagaraFunctionLibrary::GetDataInterface<UNiagaraDataInterfaceArrayInt32>(NiagaraSystem, OverrideName))
-	{
-		FRWScopeLock WriteLock(ArrayDI->ArrayRWGuard, SLT_Write);
-		ArrayDI->IntData = ArrayData;
-		ArrayDI->MarkRenderDataDirty();
-	}
+	SetNiagaraArray<int32, UNiagaraDataInterfaceArrayInt32>(NiagaraSystem, OverrideName, ArrayData);
 }
 
 void UNiagaraDataInterfaceArrayFunctionLibrary::SetNiagaraArrayBool(UNiagaraComponent* NiagaraSystem, FName OverrideName, const TArray<bool>& ArrayData)
 {
-	if (UNiagaraDataInterfaceArrayBool* ArrayDI = UNiagaraFunctionLibrary::GetDataInterface<UNiagaraDataInterfaceArrayBool>(NiagaraSystem, OverrideName))
-	{
-		FRWScopeLock WriteLock(ArrayDI->ArrayRWGuard, SLT_Write);
-		ArrayDI->BoolData = ArrayData;
-		ArrayDI->MarkRenderDataDirty();
-	}
+	SetNiagaraArray<bool, UNiagaraDataInterfaceArrayBool>(NiagaraSystem, OverrideName, ArrayData);
 }
 
+//////////////////////////////////////////////////////////////////////////
+
 TArray<float> UNiagaraDataInterfaceArrayFunctionLibrary::GetNiagaraArrayFloat(UNiagaraComponent* NiagaraSystem, FName OverrideName)
 {
-	if (UNiagaraDataInterfaceArrayFloat* ArrayDI = UNiagaraFunctionLibrary::GetDataInterface<UNiagaraDataInterfaceArrayFloat>(NiagaraSystem, OverrideName))
-	{
-		FRWScopeLock ReadLock(ArrayDI->ArrayRWGuard, SLT_ReadOnly);
-		return ArrayDI->FloatData;
-	}
-	return TArray<float>();
+	return GetNiagaraArray<float, UNiagaraDataInterfaceArrayFloat>(NiagaraSystem, OverrideName);
 }
 
 TArray<FVector2D> UNiagaraDataInterfaceArrayFunctionLibrary::GetNiagaraArrayVector2D(UNiagaraComponent* NiagaraSystem, FName OverrideName)
 {
-	if (UNiagaraDataInterfaceArrayFloat2* ArrayDI = UNiagaraFunctionLibrary::GetDataInterface<UNiagaraDataInterfaceArrayFloat2>(NiagaraSystem, OverrideName))
-	{
-		FRWScopeLock ReadLock(ArrayDI->ArrayRWGuard, SLT_ReadOnly);
-		return ArrayDI->FloatData;
-	}
-	return TArray<FVector2D>();
+	return GetNiagaraArray<FVector2D, UNiagaraDataInterfaceArrayFloat2>(NiagaraSystem, OverrideName);
 }
 
 TArray<FVector> UNiagaraDataInterfaceArrayFunctionLibrary::GetNiagaraArrayVector(UNiagaraComponent* NiagaraSystem, FName OverrideName)
 {
-	if (UNiagaraDataInterfaceArrayFloat3* ArrayDI = UNiagaraFunctionLibrary::GetDataInterface<UNiagaraDataInterfaceArrayFloat3>(NiagaraSystem, OverrideName))
-	{
-		FRWScopeLock ReadLock(ArrayDI->ArrayRWGuard, SLT_ReadOnly);
-		return ArrayDI->FloatData;
-	}
-	return TArray<FVector>();
+	return GetNiagaraArray<FVector, UNiagaraDataInterfaceArrayFloat3>(NiagaraSystem, OverrideName);
 }
 
 TArray<FVector4> UNiagaraDataInterfaceArrayFunctionLibrary::GetNiagaraArrayVector4(UNiagaraComponent* NiagaraSystem, FName OverrideName)
 {
-	if (UNiagaraDataInterfaceArrayFloat4* ArrayDI = UNiagaraFunctionLibrary::GetDataInterface<UNiagaraDataInterfaceArrayFloat4>(NiagaraSystem, OverrideName))
-	{
-		FRWScopeLock ReadLock(ArrayDI->ArrayRWGuard, SLT_ReadOnly);
-		return ArrayDI->FloatData;
-	}
-	return TArray<FVector4>();
+	return GetNiagaraArray<FVector4, UNiagaraDataInterfaceArrayFloat4>(NiagaraSystem, OverrideName);
 }
 
 TArray<FLinearColor> UNiagaraDataInterfaceArrayFunctionLibrary::GetNiagaraArrayColor(UNiagaraComponent* NiagaraSystem, FName OverrideName)
 {
-	if (UNiagaraDataInterfaceArrayColor* ArrayDI = UNiagaraFunctionLibrary::GetDataInterface<UNiagaraDataInterfaceArrayColor>(NiagaraSystem, OverrideName))
-	{
-		FRWScopeLock ReadLock(ArrayDI->ArrayRWGuard, SLT_ReadOnly);
-		return ArrayDI->ColorData;
-	}
-	return TArray<FLinearColor>();
+	return GetNiagaraArray<FLinearColor, UNiagaraDataInterfaceArrayColor>(NiagaraSystem, OverrideName);
 }
 
 TArray<FQuat> UNiagaraDataInterfaceArrayFunctionLibrary::GetNiagaraArrayQuat(UNiagaraComponent* NiagaraSystem, FName OverrideName)
 {
-	if (UNiagaraDataInterfaceArrayQuat* ArrayDI = UNiagaraFunctionLibrary::GetDataInterface<UNiagaraDataInterfaceArrayQuat>(NiagaraSystem, OverrideName))
-	{
-		FRWScopeLock ReadLock(ArrayDI->ArrayRWGuard, SLT_ReadOnly);
-		return ArrayDI->QuatData;
-	}
-	return TArray<FQuat>();
+	return GetNiagaraArray<FQuat, UNiagaraDataInterfaceArrayQuat>(NiagaraSystem, OverrideName);
 }
 
 TArray<int32> UNiagaraDataInterfaceArrayFunctionLibrary::GetNiagaraArrayInt32(UNiagaraComponent* NiagaraSystem, FName OverrideName)
 {
-	if (UNiagaraDataInterfaceArrayInt32* ArrayDI = UNiagaraFunctionLibrary::GetDataInterface<UNiagaraDataInterfaceArrayInt32>(NiagaraSystem, OverrideName))
-	{
-		FRWScopeLock ReadLock(ArrayDI->ArrayRWGuard, SLT_ReadOnly);
-		return ArrayDI->IntData;
-	}
-	return TArray<int32>();
+	return GetNiagaraArray<int32, UNiagaraDataInterfaceArrayInt32>(NiagaraSystem, OverrideName);
 }
 
 TArray<bool> UNiagaraDataInterfaceArrayFunctionLibrary::GetNiagaraArrayBool(UNiagaraComponent* NiagaraSystem, FName OverrideName)
 {
-	if (UNiagaraDataInterfaceArrayBool* ArrayDI = UNiagaraFunctionLibrary::GetDataInterface<UNiagaraDataInterfaceArrayBool>(NiagaraSystem, OverrideName))
-	{
-		FRWScopeLock ReadLock(ArrayDI->ArrayRWGuard, SLT_ReadOnly);
-		return ArrayDI->BoolData;
-	}
-	return TArray<bool>();
+	return GetNiagaraArray<bool, UNiagaraDataInterfaceArrayBool>(NiagaraSystem, OverrideName);
+}
+
+//////////////////////////////////////////////////////////////////////////
+
+void UNiagaraDataInterfaceArrayFunctionLibrary::SetNiagaraArrayFloatValue(UNiagaraComponent* NiagaraSystem, FName OverrideName, int Index, float Value, bool bSizeToFit)
+{
+	SetNiagaraArrayValue<float, UNiagaraDataInterfaceArrayFloat>(NiagaraSystem, OverrideName, Index, Value, bSizeToFit);
+}
+
+void UNiagaraDataInterfaceArrayFunctionLibrary::SetNiagaraArrayVector2DValue(UNiagaraComponent* NiagaraSystem, FName OverrideName, int Index, const FVector2D& Value, bool bSizeToFit)
+{
+	SetNiagaraArrayValue<FVector2D, UNiagaraDataInterfaceArrayFloat2>(NiagaraSystem, OverrideName, Index, Value, bSizeToFit);
+}
+
+void UNiagaraDataInterfaceArrayFunctionLibrary::SetNiagaraArrayVectorValue(UNiagaraComponent* NiagaraSystem, FName OverrideName, int Index, const FVector& Value, bool bSizeToFit)
+{
+	SetNiagaraArrayValue<FVector, UNiagaraDataInterfaceArrayFloat3>(NiagaraSystem, OverrideName, Index, Value, bSizeToFit);
 }
 
+void UNiagaraDataInterfaceArrayFunctionLibrary::SetNiagaraArrayVector4Value(UNiagaraComponent* NiagaraSystem, FName OverrideName, int Index, const FVector4& Value, bool bSizeToFit)
+{
+	SetNiagaraArrayValue<FVector4, UNiagaraDataInterfaceArrayFloat4>(NiagaraSystem, OverrideName, Index, Value, bSizeToFit);
+}
+
+void UNiagaraDataInterfaceArrayFunctionLibrary::SetNiagaraArrayColorValue(UNiagaraComponent* NiagaraSystem, FName OverrideName, int Index, const FLinearColor& Value, bool bSizeToFit)
+{
+	SetNiagaraArrayValue<FLinearColor, UNiagaraDataInterfaceArrayColor>(NiagaraSystem, OverrideName, Index, Value, bSizeToFit);
+}
+
+void UNiagaraDataInterfaceArrayFunctionLibrary::SetNiagaraArrayQuatValue(UNiagaraComponent* NiagaraSystem, FName OverrideName, int Index, const FQuat& Value, bool bSizeToFit)
+{
+	SetNiagaraArrayValue<FQuat, UNiagaraDataInterfaceArrayQuat>(NiagaraSystem, OverrideName, Index, Value, bSizeToFit);
+}
+
+void UNiagaraDataInterfaceArrayFunctionLibrary::SetNiagaraArrayInt32Value(UNiagaraComponent* NiagaraSystem, FName OverrideName, int Index, int32 Value, bool bSizeToFit)
+{
+	SetNiagaraArrayValue<int32, UNiagaraDataInterfaceArrayInt32>(NiagaraSystem, OverrideName, Index, Value, bSizeToFit);
+}
+
+void UNiagaraDataInterfaceArrayFunctionLibrary::SetNiagaraArrayBoolValue(UNiagaraComponent* NiagaraSystem, FName OverrideName, int Index, const bool& Value, bool bSizeToFit)
+{
+	SetNiagaraArrayValue<bool, UNiagaraDataInterfaceArrayBool>(NiagaraSystem, OverrideName, Index, Value, bSizeToFit);
+}
+
+//////////////////////////////////////////////////////////////////////////
+
+float UNiagaraDataInterfaceArrayFunctionLibrary::GetNiagaraArrayFloatValue(UNiagaraComponent* NiagaraSystem, FName OverrideName, int Index)
+{
+	return GetNiagaraArrayValue<float, UNiagaraDataInterfaceArrayFloat>(NiagaraSystem, OverrideName, Index);
+}
+
+FVector2D UNiagaraDataInterfaceArrayFunctionLibrary::GetNiagaraArrayVector2DValue(UNiagaraComponent* NiagaraSystem, FName OverrideName, int Index)
+{
+	return GetNiagaraArrayValue<FVector2D, UNiagaraDataInterfaceArrayFloat2>(NiagaraSystem, OverrideName, Index);
+}
+
+FVector UNiagaraDataInterfaceArrayFunctionLibrary::GetNiagaraArrayVectorValue(UNiagaraComponent* NiagaraSystem, FName OverrideName, int Index)
+{
+	return GetNiagaraArrayValue<FVector, UNiagaraDataInterfaceArrayFloat3>(NiagaraSystem, OverrideName, Index);
+}
+
+FVector4 UNiagaraDataInterfaceArrayFunctionLibrary::GetNiagaraArrayVector4Value(UNiagaraComponent* NiagaraSystem, FName OverrideName, int Index)
+{
+	return GetNiagaraArrayValue<FVector4, UNiagaraDataInterfaceArrayFloat4>(NiagaraSystem, OverrideName, Index);
+}
+
+FLinearColor UNiagaraDataInterfaceArrayFunctionLibrary::GetNiagaraArrayColorValue(UNiagaraComponent* NiagaraSystem, FName OverrideName, int Index)
+{
+	return GetNiagaraArrayValue<FLinearColor, UNiagaraDataInterfaceArrayColor>(NiagaraSystem, OverrideName, Index);
+}
+
+FQuat UNiagaraDataInterfaceArrayFunctionLibrary::GetNiagaraArrayQuatValue(UNiagaraComponent* NiagaraSystem, FName OverrideName, int Index)
+{
+	return GetNiagaraArrayValue<FQuat, UNiagaraDataInterfaceArrayQuat>(NiagaraSystem, OverrideName, Index);
+}
+
+int32 UNiagaraDataInterfaceArrayFunctionLibrary::GetNiagaraArrayInt32Value(UNiagaraComponent* NiagaraSystem, FName OverrideName, int Index)
+{
+	return GetNiagaraArrayValue<int32, UNiagaraDataInterfaceArrayInt32>(NiagaraSystem, OverrideName, Index);
+}
+
+bool UNiagaraDataInterfaceArrayFunctionLibrary::GetNiagaraArrayBoolValue(UNiagaraComponent* NiagaraSystem, FName OverrideName, int Index)
+{
+	return GetNiagaraArrayValue<bool, UNiagaraDataInterfaceArrayBool>(NiagaraSystem, OverrideName, Index);
+}
