Fixed non-unity build error

diff --git a/Engine/Plugins/FX/Niagara/Source/NiagaraEditor/Private/NiagaraGraph.cpp b/Engine/Plugins/FX/Niagara/Source/NiagaraEditor/Private/NiagaraGraph.cpp
--- a/Engine/Plugins/FX/Niagara/Source/NiagaraEditor/Private/NiagaraGraph.cpp
+++ b/Engine/Plugins/FX/Niagara/Source/NiagaraEditor/Private/NiagaraGraph.cpp
@@ -1,37 +1,38 @@
 // Copyright Epic Games, Inc. All Rights Reserved.
 
 #include "NiagaraGraph.h"
+
+#include "EdGraphSchema_Niagara.h"
+#include "GraphEditAction.h"
+#include "INiagaraEditorTypeUtilities.h"
 #include "NiagaraCommon.h"
-#include "NiagaraEditorModule.h"
-#include "NiagaraScript.h"
 #include "NiagaraComponent.h"
 #include "NiagaraConstants.h"
-#include "NiagaraSystem.h"
-#include "NiagaraNodeOutput.h"
-#include "NiagaraNodeInput.h"
-#include "NiagaraScriptSource.h"
-#include "NiagaraParameterDefinitions.h"
-#include "GraphEditAction.h"
-#include "EdGraphSchema_Niagara.h"
+#include "NiagaraCustomVersion.h"
+#include "NiagaraEditorModule.h"
+#include "NiagaraEditorModule.h"
+#include "NiagaraEditorUtilities.h"
+#include "NiagaraHlslTranslator.h"
+#include "NiagaraModule.h"
+#include "NiagaraNode.h"
 #include "NiagaraNodeAssignment.h"
+#include "NiagaraNodeFunctionCall.h"
+#include "NiagaraNodeInput.h"
+#include "NiagaraNodeOutput.h"
 #include "NiagaraNodeParameterMapBase.h"
 #include "NiagaraNodeParameterMapGet.h"
 #include "NiagaraNodeParameterMapSet.h"
 #include "NiagaraNodeReroute.h"
-#include "INiagaraEditorTypeUtilities.h"
-#include "NiagaraEditorUtilities.h"
-#include "NiagaraNode.h"
-#include "NiagaraCustomVersion.h"
 #include "NiagaraNodeStaticSwitch.h"
+#include "NiagaraParameterDefinitions.h"
+#include "NiagaraScript.h"
+#include "NiagaraScriptSource.h"
 #include "NiagaraScriptVariable.h"
-#include "NiagaraHlslTranslator.h"
-#include "NiagaraNodeFunctionCall.h"
-#include "Misc/SecureHash.h"
 #include "HAL/FileManager.h"
 #include "Misc/Paths.h"
-#include "Serialization/DuplicatedDataWriter.h"
+#include "Misc/SecureHash.h"
+#include "Modules/ModuleManager.h"
 #include "String/ParseTokens.h"
-#include "ViewModels/NiagaraParameterDefinitionsSubscriberViewModel.h"
 
 
 DECLARE_CYCLE_STAT(TEXT("NiagaraEditor - Graph - PostLoad"), STAT_NiagaraEditor_Graph_PostLoad, STATGROUP_NiagaraEditor);
