- Optionally allow Niagara system simulations to not delay the tick task from completion
- This will also enable concurrent spawning in PostActorTick
- Pre EOF updates will wait for all concurrent work, or the next time the simulation ticks it will wait

#rnx

diff --git a/Engine/Plugins/FX/Niagara/Source/Niagara/Private/NiagaraSystemSimulation.cpp b/Engine/Plugins/FX/Niagara/Source/Niagara/Private/NiagaraSystemSimulation.cpp
--- a/Engine/Plugins/FX/Niagara/Source/Niagara/Private/NiagaraSystemSimulation.cpp
+++ b/Engine/Plugins/FX/Niagara/Source/Niagara/Private/NiagaraSystemSimulation.cpp
@@ -112,6 +112,14 @@ static FAutoConsoleVariableRef CVarNiagaraSkipTickDeltaSeconds(
 	ECVF_Default
 );
 
+static int32 GNiagaraSystemSimulationTickTaskShouldWait = 1;
+static FAutoConsoleVariableRef CVarNiagaraSystemSimulationTickTaskShouldWait(
+	TEXT("fx.Niagara.SystemSimulation.TickTaskShouldWait"),
+	GNiagaraSystemSimulationTickTaskShouldWait,
+	TEXT("When enabled the tick task will wait for concurrent work to complete, when disabled the task is complete one the GT tick is complete."),
+	ECVF_Default
+);
+
 namespace NiagaraSystemSimulationLocal
 {
 #if NIAGARA_SYSTEMSIMULATION_DEBUGGING
@@ -1240,7 +1248,14 @@ void FNiagaraSystemSimulation::Tick_GameThread(float DeltaSeconds, const FGraphE
 			Instance->ConcurrentTickGraphEvent = ConcurrentTickGraphEvent;
 		}
 		ConcurrentTickTask->Unlock();
-		MyCompletionGraphEvent->DontCompleteUntil(AllWorkCompleteGraphEvent);
+		if (GNiagaraSystemSimulationTickTaskShouldWait)
+		{
+			MyCompletionGraphEvent->DontCompleteUntil(AllWorkCompleteGraphEvent);
+		}
+		else
+		{
+			FNiagaraWorldManager::Get(World)->MarkSimulationsForEndOfFrameWait(this);
+		}
 	}
 	else
 	{
@@ -1388,7 +1403,8 @@ void FNiagaraSystemSimulation::Spawn_GameThread(float DeltaSeconds, bool bPostAc
 		}
 
 		// Can we spawn async?
-		FNiagaraSystemSimulationTickContext Context(this, SpawningInstances, SpawningDataSet, DeltaSeconds, SpawningInstances.Num(), !bPostActorTick);
+		const bool bCanRunAsync = !bPostActorTick || !GNiagaraSystemSimulationTickTaskShouldWait;
+		FNiagaraSystemSimulationTickContext Context(this, SpawningInstances, SpawningDataSet, DeltaSeconds, SpawningInstances.Num(), bCanRunAsync);
 		if ( Context.IsRunningAsync() )
 		{
 			check(bIsSolo == false);
@@ -1400,7 +1416,9 @@ void FNiagaraSystemSimulation::Spawn_GameThread(float DeltaSeconds, bool bPostAc
 				Instance->ConcurrentTickGraphEvent = ConcurrentTickGraphEvent;
 			}
 			ConcurrentTickTask->Unlock();
-			WorldManager->MarkSimulationForPostActorWork(this);
+
+			// Note we always mark for EOF wait as the execution could be ongoing beyond PostActorTick
+			WorldManager->MarkSimulationsForEndOfFrameWait(this);
 		}
 		else
 		{
diff --git a/Engine/Plugins/FX/Niagara/Source/Niagara/Private/NiagaraWorldManager.cpp b/Engine/Plugins/FX/Niagara/Source/Niagara/Private/NiagaraWorldManager.cpp
--- a/Engine/Plugins/FX/Niagara/Source/Niagara/Private/NiagaraWorldManager.cpp
+++ b/Engine/Plugins/FX/Niagara/Source/Niagara/Private/NiagaraWorldManager.cpp
@@ -768,14 +768,27 @@ void FNiagaraWorldManager::PostActorTick(float DeltaSeconds)
 
 void FNiagaraWorldManager::PreSendAllEndOfFrameUpdates()
 {
-	for (const auto& Simulation : SimulationsWithPostActorWork)
+	if ( SimulationsWithPostActorWork.Num() > 0 )
+	{
+		//-TODO: This should not happen
+		for (const auto& Simulation : SimulationsWithPostActorWork)
+		{
+			if (Simulation->IsValid())
+			{
+				Simulation->WaitForInstancesTickComplete();
+			}
+		}
+		SimulationsWithPostActorWork.Reset();
+	}
+
+	for (const auto& Simulation : SimulationsWithEndOfFrameWait)
 	{
 		if (Simulation->IsValid())
 		{
 			Simulation->WaitForInstancesTickComplete();
 		}
 	}
-	SimulationsWithPostActorWork.Reset();
+	SimulationsWithEndOfFrameWait.Reset();
 }
 
 void FNiagaraWorldManager::MarkSimulationForPostActorWork(FNiagaraSystemSimulation* SystemSimulation)
@@ -787,6 +800,15 @@ void FNiagaraWorldManager::MarkSimulationForPostActorWork(FNiagaraSystemSimulati
 	}
 }
 
+void FNiagaraWorldManager::MarkSimulationsForEndOfFrameWait(FNiagaraSystemSimulation* SystemSimulation)
+{
+	check(SystemSimulation);
+	if (!SimulationsWithEndOfFrameWait.ContainsByPredicate([&](const TSharedRef<FNiagaraSystemSimulation, ESPMode::ThreadSafe>& Existing) { return &Existing.Get() == SystemSimulation; }))
+	{
+		SimulationsWithEndOfFrameWait.Add(SystemSimulation->AsShared());
+	}
+}
+
 void FNiagaraWorldManager::Tick(ETickingGroup TickGroup, float DeltaSeconds, ELevelTick TickType, ENamedThreads::Type CurrentThread, const FGraphEventRef& MyCompletionGraphEvent)
 {
 	check(TickGroup >= NiagaraFirstTickGroup && TickGroup <= NiagaraLastTickGroup);
diff --git a/Engine/Plugins/FX/Niagara/Source/Niagara/Public/NiagaraWorldManager.h b/Engine/Plugins/FX/Niagara/Source/Niagara/Public/NiagaraWorldManager.h
--- a/Engine/Plugins/FX/Niagara/Source/Niagara/Public/NiagaraWorldManager.h
+++ b/Engine/Plugins/FX/Niagara/Source/Niagara/Public/NiagaraWorldManager.h
@@ -145,6 +145,7 @@ public:
 	void DestroySystemInstance(FNiagaraSystemInstancePtr& InPtr);	
 
 	void MarkSimulationForPostActorWork(FNiagaraSystemSimulation* SystemSimulation);
+	void MarkSimulationsForEndOfFrameWait(FNiagaraSystemSimulation* SystemSimulation);
 
 	void Tick(ETickingGroup TickGroup, float DeltaSeconds, ELevelTick TickType, ENamedThreads::Type CurrentThread, const FGraphEventRef& MyCompletionGraphEvent);
 
@@ -302,6 +303,8 @@ private:
 
 	TArray<TSharedRef<FNiagaraSystemSimulation, ESPMode::ThreadSafe>> SimulationsWithPostActorWork;
 
+	TArray<TSharedRef<FNiagaraSystemSimulation, ESPMode::ThreadSafe>> SimulationsWithEndOfFrameWait;
+
 	int32 CachedEffectsQuality;
 
 	bool bCachedPlayerViewLocationsValid = false;
