Recover lost change for preventing non-determinism with the textures generated by [REDACTED_PERSON]

diff --git a/Engine/Plugins/FX/Niagara/Source/Niagara/Private/NiagaraDataInterfaceCurveBase.cpp b/Engine/Plugins/FX/Niagara/Source/Niagara/Private/NiagaraDataInterfaceCurveBase.cpp
--- a/Engine/Plugins/FX/Niagara/Source/Niagara/Private/NiagaraDataInterfaceCurveBase.cpp
+++ b/Engine/Plugins/FX/Niagara/Source/Niagara/Private/NiagaraDataInterfaceCurveBase.cpp
@@ -291,6 +291,9 @@ void UNiagaraDataInterfaceCurveBase::UpdateExposedTexture()
 	ExposedTexture->Source.UnlockMip(0);
 
 	ExposedTexture->PostEditChange();
+
+	// PostEditChange() will assign a random GUID to the texture, which leads to non-deterministic builds.
+	ExposedTexture->SetDeterministicLightingGuid();
 }
 #endif
 
