Clean up error message for missing data interface

diff --git a/Engine/Plugins/FX/Niagara/Source/NiagaraEditor/Private/NiagaraHlslTranslator.cpp b/Engine/Plugins/FX/Niagara/Source/NiagaraEditor/Private/NiagaraHlslTranslator.cpp
--- a/Engine/Plugins/FX/Niagara/Source/NiagaraEditor/Private/NiagaraHlslTranslator.cpp
+++ b/Engine/Plugins/FX/Niagara/Source/NiagaraEditor/Private/NiagaraHlslTranslator.cpp
@@ -6725,13 +6725,16 @@ void FHlslNiagaraTranslator::ProcessCustomHlsl(const FString& InCustomHlsl, ENia
 			{
 				// If the cdo wasn't found, the data interface was not passed through a parameter map and so it won't be bound correctly, so add a compile error
 				// and invalidate the signature.
-				Error(LOCTEXT("DataInterfaceNotFoundCustomHLSL", "Data interface used by custom hlsl, but not found in precompiled data. Please notify Niagara team of bug."), nullptr, nullptr);
+				Error(FText::Format(LOCTEXT("DataInterfaceNotFoundCustomHLSL", "Data interface ({0}) used by custom hlsl, but not found in precompiled data. Please notify Niagara team of bug."),
+					FText::FromName(Input.GetName())),
+					InNodeForErrorReporting,
+					nullptr);
 				return;
 			}
 			int32 OwnerIdx = Inputs[i];
 			if (OwnerIdx < 0 || OwnerIdx >= CompilationOutput.ScriptData.DataInterfaceInfo.Num())
 			{
-				Error(LOCTEXT("FunctionCallDataInterfaceMissingRegistration", "Function call signature does not match to a registered DataInterface. Valid DataInterfaces should be wired into a DataInterface function call."), nullptr, nullptr);
+				Error(LOCTEXT("FunctionCallDataInterfaceMissingRegistration", "Function call signature does not match to a registered DataInterface. Valid DataInterfaces should be wired into a DataInterface function call."), InNodeForErrorReporting, nullptr);
 				return;
 			}
 
