fixed discrepancies between UE5/Main and Release-Engine-Test

diff --git a/Engine/Plugins/FX/Niagara/Source/Niagara/Private/NiagaraEmitter.cpp b/Engine/Plugins/FX/Niagara/Source/Niagara/Private/NiagaraEmitter.cpp
--- a/Engine/Plugins/FX/Niagara/Source/Niagara/Private/NiagaraEmitter.cpp
+++ b/Engine/Plugins/FX/Niagara/Source/Niagara/Private/NiagaraEmitter.cpp
@@ -592,7 +592,6 @@ void UNiagaraEmitter::PostLoad()
 			UE_LOG(LogNiagara, Warning, TEXT("Disabling interpolated spawn because emitter flag and script type don't match. Did you adjust this value in the UI? Emitter may need recompile.. %s"), *GetFullName());
 		}
 	}
-
 	EnsureScriptsPostLoaded();
 
 	FGraphEventArray ParentPrerequisiteTasks;
@@ -601,6 +600,11 @@ void UNiagaraEmitter::PostLoad()
 	{
 		ParentPrerequisiteTasks.Add(Parent->UpdateTaskRef);
 	}
+	// this can only ever be true for old assets that haven't been loaded yet, so this won't overwrite subsequent changes to the template specification
+	if(bIsTemplateAsset_DEPRECATED)
+	{
+		TemplateSpecification = ENiagaraScriptTemplateSpecification::Template;
+	}
 #endif
 	
 	// we are not yet finished, but we do the rest of the work after postload in a separate task
diff --git a/Engine/Plugins/FX/Niagara/Source/Niagara/Private/NiagaraSystem.cpp b/Engine/Plugins/FX/Niagara/Source/Niagara/Private/NiagaraSystem.cpp
--- a/Engine/Plugins/FX/Niagara/Source/Niagara/Private/NiagaraSystem.cpp
+++ b/Engine/Plugins/FX/Niagara/Source/Niagara/Private/NiagaraSystem.cpp
@@ -361,7 +361,7 @@ void UNiagaraSystem::UpdateSystemAfterLoad()
 	const bool bIsDedicatedServer = !GIsClient && GIsServer;
 
 	if (!GetOutermost()->bIsCookedForEditor && !bIsDedicatedServer)
-	{
+	{		
 		TArray<UNiagaraScript*> AllSystemScripts;
 		UNiagaraScriptSourceBase* SystemScriptSource;
 		if (SystemSpawnScript == nullptr)
@@ -841,11 +841,6 @@ void UNiagaraSystem::PostLoad()
 	// After forcing syncing all DefaultLinkedParameterDefinitions, call SynchronizeWithParameterDefinitions again to sync with all definitions the system was already subscribed to, and do not force the sync.
 	SynchronizeWithParameterDefinitions();
 	
-	// see the equivalent in NiagaraEmitter for details
-	if(bIsTemplateAsset_DEPRECATED)
-	{
-		TemplateSpecification = bIsTemplateAsset_DEPRECATED ? ENiagaraScriptTemplateSpecification::Template : ENiagaraScriptTemplateSpecification::None;
-	}
 #endif // WITH_EDITORONLY_DATA
 
 	// we are not yet finished, but we do the rest of the work after postload in a separate task
