Move parameter definitions synchronization to AfterPostLoad methods of System and Emitter. Use Init method of UNiagaraScriptVariable to properly initialize default value variant member when created in some places.

diff --git a/Engine/Plugins/FX/Niagara/Source/Niagara/Private/NiagaraEmitter.cpp b/Engine/Plugins/FX/Niagara/Source/Niagara/Private/NiagaraEmitter.cpp
--- a/Engine/Plugins/FX/Niagara/Source/Niagara/Private/NiagaraEmitter.cpp
+++ b/Engine/Plugins/FX/Niagara/Source/Niagara/Private/NiagaraEmitter.cpp
@@ -539,31 +539,6 @@ void UNiagaraEmitter::PostLoad()
 			}
 		}
 
-		// Synchronize with definitions before merging.
-		// First force sync with all definitions in the DefaultLinkedParameterDefinitions array.
-		const UNiagaraSettings* Settings = GetDefault<UNiagaraSettings>();
-		check(Settings);
-		TArray<FGuid> DefaultDefinitionsUniqueIds;
-		for (const FSoftObjectPath& DefaultLinkedParameterDefinitionObjPath : Settings->DefaultLinkedParameterDefinitions)
-		{
-			UNiagaraParameterDefinitionsBase* DefaultLinkedParameterDefinitions = Cast<UNiagaraParameterDefinitionsBase>(DefaultLinkedParameterDefinitionObjPath.TryLoad());
-			if (DefaultLinkedParameterDefinitions == nullptr)
-			{
-				continue;
-			}
-			DefaultDefinitionsUniqueIds.Add(DefaultLinkedParameterDefinitions->GetDefinitionsUniqueId());
-			const bool bDoNotAssertIfAlreadySubscribed = true;
-			SubscribeToParameterDefinitions(DefaultLinkedParameterDefinitions, bDoNotAssertIfAlreadySubscribed);
-		}
-		FSynchronizeWithParameterDefinitionsArgs Args;
-		Args.SpecificDefinitionsUniqueIds = DefaultDefinitionsUniqueIds;
-		Args.bForceSynchronizeDefinitions = true;
-		Args.bSubscribeAllNameMatchParameters = true;
-		SynchronizeWithParameterDefinitions(Args);
-		
-		// After forcing syncing all DefaultLinkedParameterDefinitions, call SynchronizeWithParameterDefinitions again to sync with all definitions the emitter was already subscribed to, and do not force the sync.
-		SynchronizeWithParameterDefinitions();
-
 		if (IsSynchronizedWithParent() == false)
 		{
 			// Modify here so that the asset will be marked dirty when using the resave commandlet.  This will be ignored during regular post load.
@@ -1288,6 +1263,33 @@ void UNiagaraEmitter::UpdateEmitterAfterLoad()
 	
 #if WITH_EDITORONLY_DATA
 	check(IsInGameThread());
+
+	// Synchronize with definitions before merging.
+	// First force sync with all definitions in the DefaultLinkedParameterDefinitions array.
+	const UNiagaraSettings* Settings = GetDefault<UNiagaraSettings>();
+	check(Settings);
+	TArray<FGuid> DefaultDefinitionsUniqueIds;
+	for (const FSoftObjectPath& DefaultLinkedParameterDefinitionObjPath : Settings->DefaultLinkedParameterDefinitions)
+	{
+		UNiagaraParameterDefinitionsBase* DefaultLinkedParameterDefinitions = Cast<UNiagaraParameterDefinitionsBase>(DefaultLinkedParameterDefinitionObjPath.TryLoad());
+		if (DefaultLinkedParameterDefinitions == nullptr)
+		{
+			continue;
+		}
+		DefaultDefinitionsUniqueIds.Add(DefaultLinkedParameterDefinitions->GetDefinitionsUniqueId());
+		const bool bDoNotAssertIfAlreadySubscribed = true;
+		SubscribeToParameterDefinitions(DefaultLinkedParameterDefinitions, bDoNotAssertIfAlreadySubscribed);
+	}
+	FSynchronizeWithParameterDefinitionsArgs Args;
+	Args.SpecificDefinitionsUniqueIds = DefaultDefinitionsUniqueIds;
+	Args.bForceSynchronizeDefinitions = true;
+	Args.bSubscribeAllNameMatchParameters = true;
+	//SynchronizeWithParameterDefinitions(Args);
+
+	// After forcing syncing all DefaultLinkedParameterDefinitions, call SynchronizeWithParameterDefinitions again to sync with all definitions the emitter was already subscribed to, and do not force the sync.
+	//SynchronizeWithParameterDefinitions();
+
+	// Merge with parent if necessary.
 	if (GetOuter()->IsA<UNiagaraEmitter>())
 	{
 		// If this emitter is owned by another emitter, remove it's inheritance information so that it doesn't try to merge changes.
diff --git a/Engine/Plugins/FX/Niagara/Source/Niagara/Private/NiagaraSystem.cpp b/Engine/Plugins/FX/Niagara/Source/Niagara/Private/NiagaraSystem.cpp
--- a/Engine/Plugins/FX/Niagara/Source/Niagara/Private/NiagaraSystem.cpp
+++ b/Engine/Plugins/FX/Niagara/Source/Niagara/Private/NiagaraSystem.cpp
@@ -397,6 +397,31 @@ void UNiagaraSystem::UpdateSystemAfterLoad()
 			bSystemScriptsAreSynchronized &= SystemScript->AreScriptAndSourceSynchronized();
 		}
 
+		// Synchronize with parameter definitions
+		// First force sync with all definitions in the DefaultLinkedParameterDefinitions array.
+		const UNiagaraSettings* Settings = GetDefault<UNiagaraSettings>();
+		check(Settings);
+		TArray<FGuid> DefaultDefinitionsUniqueIds;
+		for (const FSoftObjectPath& DefaultLinkedParameterDefinitionObjPath : Settings->DefaultLinkedParameterDefinitions)
+		{
+			UNiagaraParameterDefinitionsBase* DefaultLinkedParameterDefinitions = Cast<UNiagaraParameterDefinitionsBase>(DefaultLinkedParameterDefinitionObjPath.TryLoad());
+			if (DefaultLinkedParameterDefinitions == nullptr)
+			{
+				continue;
+			}
+			DefaultDefinitionsUniqueIds.Add(DefaultLinkedParameterDefinitions->GetDefinitionsUniqueId());
+			const bool bDoNotAssertIfAlreadySubscribed = true;
+			SubscribeToParameterDefinitions(DefaultLinkedParameterDefinitions, bDoNotAssertIfAlreadySubscribed);
+		}
+		FSynchronizeWithParameterDefinitionsArgs Args;
+		Args.SpecificDefinitionsUniqueIds = DefaultDefinitionsUniqueIds;
+		Args.bForceSynchronizeDefinitions = true;
+		Args.bSubscribeAllNameMatchParameters = true;
+		//SynchronizeWithParameterDefinitions(Args);
+
+		// After forcing syncing all DefaultLinkedParameterDefinitions, call SynchronizeWithParameterDefinitions again to sync with all definitions the system was already subscribed to, and do not force the sync.
+		//SynchronizeWithParameterDefinitions();
+
 		bool bEmitterScriptsAreSynchronized = true;
 
 #if 0
@@ -815,31 +840,6 @@ void UNiagaraSystem::PostLoad()
 		INiagaraModule& NiagaraModule = FModuleManager::GetModuleChecked<INiagaraModule>("Niagara");
 		EditorParameters = NiagaraModule.GetEditorOnlyDataUtilities().CreateDefaultEditorParameters(this);
 	}
-
-	// Synchronize with parameter definitions
-	// First force sync with all definitions in the DefaultLinkedParameterDefinitions array.
-	const UNiagaraSettings* Settings = GetDefault<UNiagaraSettings>();
-	check(Settings);
-	TArray<FGuid> DefaultDefinitionsUniqueIds;
-	for (const FSoftObjectPath& DefaultLinkedParameterDefinitionObjPath : Settings->DefaultLinkedParameterDefinitions)
-	{
-		UNiagaraParameterDefinitionsBase* DefaultLinkedParameterDefinitions = Cast<UNiagaraParameterDefinitionsBase>(DefaultLinkedParameterDefinitionObjPath.TryLoad());
-		if (DefaultLinkedParameterDefinitions == nullptr)
-		{
-			continue;
-		}
-		DefaultDefinitionsUniqueIds.Add(DefaultLinkedParameterDefinitions->GetDefinitionsUniqueId());
-		const bool bDoNotAssertIfAlreadySubscribed = true;
-		SubscribeToParameterDefinitions(DefaultLinkedParameterDefinitions, bDoNotAssertIfAlreadySubscribed);
-	}
-	FSynchronizeWithParameterDefinitionsArgs Args;
-	Args.SpecificDefinitionsUniqueIds = DefaultDefinitionsUniqueIds;
-	Args.bForceSynchronizeDefinitions = true;
-	Args.bSubscribeAllNameMatchParameters = true;
-	SynchronizeWithParameterDefinitions(Args);
-
-	// After forcing syncing all DefaultLinkedParameterDefinitions, call SynchronizeWithParameterDefinitions again to sync with all definitions the system was already subscribed to, and do not force the sync.
-	SynchronizeWithParameterDefinitions();
 	
 #endif // WITH_EDITORONLY_DATA
 
diff --git a/Engine/Plugins/FX/Niagara/Source/NiagaraEditor/Private/NiagaraGraph.cpp b/Engine/Plugins/FX/Niagara/Source/NiagaraEditor/Private/NiagaraGraph.cpp
--- a/Engine/Plugins/FX/Niagara/Source/NiagaraEditor/Private/NiagaraGraph.cpp
+++ b/Engine/Plugins/FX/Niagara/Source/NiagaraEditor/Private/NiagaraGraph.cpp
@@ -150,6 +150,7 @@ void UNiagaraGraph::PostLoad()
 		if (ScriptVar == nullptr)
 		{
 			ScriptVar = NewObject<UNiagaraScriptVariable>(const_cast<UNiagaraGraph*>(this));
+			ScriptVar->Init(Var, FNiagaraVariableMetaData());
 			UE_LOG(LogNiagaraEditor, Display, TEXT("Fixed null UNiagaraScriptVariable | variable %s | asset path %s"), *Var.GetName().ToString(), *GetPathName());
 		}
 		else
diff --git a/Engine/Plugins/FX/Niagara/Source/NiagaraEditor/Private/ViewModels/NiagaraParameterPanelViewModel.cpp b/Engine/Plugins/FX/Niagara/Source/NiagaraEditor/Private/ViewModels/NiagaraParameterPanelViewModel.cpp
--- a/Engine/Plugins/FX/Niagara/Source/NiagaraEditor/Private/ViewModels/NiagaraParameterPanelViewModel.cpp
+++ b/Engine/Plugins/FX/Niagara/Source/NiagaraEditor/Private/ViewModels/NiagaraParameterPanelViewModel.cpp
@@ -1306,7 +1306,7 @@ TArray<FNiagaraParameterPanelItem> FNiagaraSystemToolkitParameterPanelViewModel:
 				// Create a new UNiagaraScriptVariable to represent this parameter for the lifetime of the ParameterPanelViewModel.
 				ScriptVar = NewObject<UNiagaraScriptVariable>(GetTransientPackage());
 				ScriptVar->AddToRoot();
-				ScriptVar->Variable = Var;
+				ScriptVar->Init(Var, FNiagaraVariableMetaData());
 				TransientParameterToScriptVarMap.Add(Var, ScriptVar);
 			}
 
@@ -1460,7 +1460,7 @@ TArray<FNiagaraParameterPanelItem> FNiagaraSystemToolkitParameterPanelViewModel:
 						// Create a new UNiagaraScriptVariable to represent this parameter for the lifetime of the ParameterPanelViewModel.
 						ScriptVar = NewObject<UNiagaraScriptVariable>(&SystemViewModel->GetSystem());
 						ScriptVar->AddToRoot();
-						ScriptVar->Variable = Var;
+						ScriptVar->Init(Var, FNiagaraVariableMetaData());
 						TransientParameterToScriptVarMap.Add(Var, ScriptVar);
 					}
 
@@ -2318,7 +2318,7 @@ TArray<FNiagaraParameterPanelItem> FNiagaraScriptToolkitParameterPanelViewModel:
 					// Create a new UNiagaraScriptVariable to represent this parameter for the lifetime of the ParameterPanelViewModel.
 					ScriptVar = NewObject<UNiagaraScriptVariable>(ScriptViewModel->GetStandaloneScript().Script);
 					ScriptVar->AddToRoot();
-					ScriptVar->Variable = Var;
+					ScriptVar->Init(Var, FNiagaraVariableMetaData());
 					TransientParameterToScriptVarMap.Add(Var, ScriptVar);
 				}
 
