Fixed a bug where the parameters panel would show all attributes and not filter by static switch selection

diff --git a/Engine/Plugins/FX/Niagara/Source/NiagaraEditor/Private/ViewModels/NiagaraParameterPanelViewModel.cpp b/Engine/Plugins/FX/Niagara/Source/NiagaraEditor/Private/ViewModels/NiagaraParameterPanelViewModel.cpp
--- a/Engine/Plugins/FX/Niagara/Source/NiagaraEditor/Private/ViewModels/NiagaraParameterPanelViewModel.cpp
+++ b/Engine/Plugins/FX/Niagara/Source/NiagaraEditor/Private/ViewModels/NiagaraParameterPanelViewModel.cpp
@@ -1377,7 +1377,7 @@ TArray<FNiagaraParameterPanelItem> FNiagaraSystemToolkitParameterPanelViewModel:
 			FNiagaraParameterMapHistoryBuilder Builder;
 			UNiagaraEmitter* GraphOwningEmitter = Graph->GetTypedOuter<UNiagaraEmitter>();
 			FCompileConstantResolver ConstantResolver = GraphOwningEmitter != nullptr
-				? FCompileConstantResolver(GraphOwningEmitter, ENiagaraScriptUsage::Function)
+				? FCompileConstantResolver(GraphOwningEmitter, OutputNode->GetUsage())
 				: FCompileConstantResolver();
 				 
 			Builder.SetIgnoreDisabled(bIgnoreDisabled);
@@ -1393,7 +1393,7 @@ TArray<FNiagaraParameterPanelItem> FNiagaraSystemToolkitParameterPanelViewModel:
 				}
 			}
 			Builder.BeginUsage(StageUsage, StageName);
-			NodeToTraverse->BuildParameterMapHistory(Builder, true, false);
+			NodeToTraverse->BuildParameterMapHistory(Builder, true, true);
 			Builder.EndUsage();
 
 			if (Builder.Histories.Num() != 1)
