Moved "Add Note" action from stack module actions to item actions

diff --git a/Engine/Plugins/FX/Niagara/Source/NiagaraEditorWidgets/Private/NiagaraEditorWidgetsUtilities.cpp b/Engine/Plugins/FX/Niagara/Source/NiagaraEditorWidgets/Private/NiagaraEditorWidgetsUtilities.cpp
--- a/Engine/Plugins/FX/Niagara/Source/NiagaraEditorWidgets/Private/NiagaraEditorWidgetsUtilities.cpp
+++ b/Engine/Plugins/FX/Niagara/Source/NiagaraEditorWidgets/Private/NiagaraEditorWidgetsUtilities.cpp
@@ -274,6 +274,12 @@ void ToggleShouldDebugDraw(TWeakObjectPtr<UNiagaraStackItem> StackItemWeak)
 	}
 }
 
+void EnableNoteMode(TWeakObjectPtr<UNiagaraStackModuleItem> ModuleItem)
+{
+	ModuleItem->GetSystemViewModel()->GetSelectionViewModel()->UpdateSelectedEntries({ModuleItem.Get()}, {}, true);
+	ModuleItem->SetNoteMode(true);
+}
+
 bool FNiagaraStackEditorWidgetsUtilities::AddStackItemContextMenuActions(FMenuBuilder& MenuBuilder, UNiagaraStackItem& StackItem)
 {
 	if (StackItem.SupportsChangeEnabled())
@@ -295,6 +301,7 @@ bool FNiagaraStackEditorWidgetsUtilities::AddStackItemContextMenuActions(FMenuBu
 			}
 
 			UNiagaraStackModuleItem* ModuleItem = Cast<UNiagaraStackModuleItem>(&StackItem);
+			
 			if (ModuleItem && ModuleItem->GetModuleNode().ContainsDebugSwitch())
 			{
 				FUIAction Action(FExecuteAction::CreateStatic(&ToggleShouldDebugDraw, TWeakObjectPtr<UNiagaraStackItem>(&StackItem)),
@@ -309,6 +316,15 @@ bool FNiagaraStackEditorWidgetsUtilities::AddStackItemContextMenuActions(FMenuBu
 					EUserInterfaceActionType::Check);
 			}
 
+			if(ModuleItem)
+			{
+				MenuBuilder.AddMenuEntry(
+					LOCTEXT("AddNote", "Add Note"),
+					LOCTEXT("AddNoteToolTip", "Add a note to this module item."),
+					FSlateIcon(),
+					FUIAction(FExecuteAction::CreateStatic(&EnableNoteMode, TWeakObjectPtr<UNiagaraStackModuleItem>(ModuleItem))));
+			}
+
 		}
 		MenuBuilder.EndSection();
 		return true;
@@ -331,12 +347,6 @@ void ShowInsertModuleMenu(TWeakObjectPtr<UNiagaraStackModuleItem> StackModuleIte
 	}
 }
 
-void EnableNoteMode(TWeakObjectPtr<UNiagaraStackModuleItem> ModuleItem)
-{
-	ModuleItem->GetSystemViewModel()->GetSelectionViewModel()->UpdateSelectedEntries({ModuleItem.Get()}, {}, true);
-	ModuleItem->SetNoteMode(true);
-}
-
 bool FNiagaraStackEditorWidgetsUtilities::AddStackModuleItemContextMenuActions(FMenuBuilder& MenuBuilder, UNiagaraStackModuleItem& StackModuleItem, TSharedRef<SWidget> TargetWidget)
 {
 	MenuBuilder.BeginSection("ModuleActions", LOCTEXT("ModuleActions", "Module Actions"));
@@ -352,12 +362,6 @@ bool FNiagaraStackEditorWidgetsUtilities::AddStackModuleItemContextMenuActions(F
 			LOCTEXT("InsertModuleBelowToolTip", "Insert a new module below this module in the stack."),
 			FSlateIcon(),
 			FUIAction(FExecuteAction::CreateStatic(&ShowInsertModuleMenu, TWeakObjectPtr<UNiagaraStackModuleItem>(&StackModuleItem), 1, TWeakPtr<SWidget>(TargetWidget))));
-
-		MenuBuilder.AddMenuEntry(
-			LOCTEXT("AddNote", "Add Note"),
-			LOCTEXT("AddNoteToolTip", "Add a note to this module item."),
-			FSlateIcon(),
-			FUIAction(FExecuteAction::CreateStatic(&EnableNoteMode, TWeakObjectPtr<UNiagaraStackModuleItem>(&StackModuleItem))));
 	}
 	MenuBuilder.EndSection();
 	return true;
