Toggle Enabled and Toggle Debug Draw commands will now work on all selected valid entries.
All values are updated to the one the rightclicked item should assume rather than toggling each item individually

diff --git a/Engine/Plugins/FX/Niagara/Source/NiagaraEditorWidgets/Private/NiagaraEditorWidgetsUtilities.cpp b/Engine/Plugins/FX/Niagara/Source/NiagaraEditorWidgets/Private/NiagaraEditorWidgetsUtilities.cpp
--- a/Engine/Plugins/FX/Niagara/Source/NiagaraEditorWidgets/Private/NiagaraEditorWidgetsUtilities.cpp
+++ b/Engine/Plugins/FX/Niagara/Source/NiagaraEditorWidgets/Private/NiagaraEditorWidgetsUtilities.cpp
@@ -22,6 +22,7 @@
 #include "EditorFontGlyphs.h"
 #include "NiagaraMessages.h"
 #include "ViewModels/NiagaraSystemSelectionViewModel.h"
+#include "ViewModels/NiagaraSystemViewModel.h"
 
 #define LOCTEXT_NAMESPACE "NiagaraStackEditorWidgetsUtilities"
 
@@ -258,19 +259,64 @@ void DeleteItem(TWeakObjectPtr<UNiagaraStackItem> StackItemWeak)
 
 void ToggleEnabledState(TWeakObjectPtr<UNiagaraStackItem> StackItemWeak)
 {
+	TSet<UNiagaraStackItem*> ItemsToToggle;
+
 	UNiagaraStackItem* StackItem = StackItemWeak.Get();
 	if (StackItem != nullptr)
 	{
-		StackItem->SetIsEnabled(!StackItem->GetIsEnabled());
+		ItemsToToggle.Add(StackItem);
+		bool bShouldBeEnabled = !StackItem->GetIsEnabled();
+		
+		TArray<UNiagaraStackEntry*> StackEntries;
+		StackItem->GetSystemViewModel()->GetSelectionViewModel()->GetSelectedEntries(StackEntries);
+
+		for(UNiagaraStackEntry* StackEntry : StackEntries)
+		{
+			UNiagaraStackItem* AdditionalStackItem = Cast<UNiagaraStackItem>(StackEntry);
+
+			if (AdditionalStackItem != nullptr)
+			{
+				ItemsToToggle.Add(AdditionalStackItem);
+			}
+		}
+
+		// we assume the same state for all of the selected entries rather than toggling item-by-item
+		for(UNiagaraStackItem* CurrentStackItem : ItemsToToggle)
+		{
+			CurrentStackItem->SetIsEnabled(bShouldBeEnabled);		
+		}
 	}
+
 }
 
 void ToggleShouldDebugDraw(TWeakObjectPtr<UNiagaraStackItem> StackItemWeak)
 {
+	TSet<UNiagaraStackModuleItem*> ItemsToToggle;
+
 	UNiagaraStackModuleItem* ModuleItem = Cast<UNiagaraStackModuleItem>(StackItemWeak.Get());
 	if (ModuleItem != nullptr)
 	{
-		ModuleItem->SetDebugDrawEnabled(!ModuleItem->IsDebugDrawEnabled());
+		ItemsToToggle.Add(ModuleItem);
+		bool bShouldBeEnabled = !ModuleItem->IsDebugDrawEnabled();
+
+		TArray<UNiagaraStackEntry*> StackEntries;
+		ModuleItem->GetSystemViewModel()->GetSelectionViewModel()->GetSelectedEntries(StackEntries);
+
+		for(UNiagaraStackEntry* StackEntry : StackEntries)
+		{
+			UNiagaraStackModuleItem* AdditionalStackModuleItem = Cast<UNiagaraStackModuleItem>(StackEntry);
+
+			if (AdditionalStackModuleItem != nullptr)
+			{
+				ItemsToToggle.Add(AdditionalStackModuleItem);
+			}
+		}
+
+		// we assume the same state for all of the selected entries rather than toggling item-by-item
+		for(UNiagaraStackModuleItem* CurrentStackItem : ItemsToToggle)
+		{
+			CurrentStackItem->SetDebugDrawEnabled(bShouldBeEnabled);
+		}
 	}
 }
 
