Various smaller fixes/changes to Niagara actions
- Removed "Add X Parameter" actions from module scripts and dynamic input scripts, as they needlessly clutter the menu and parameters can be accessed in the parameters panel
- Added parent category "Add Parameter" to all add parameter actions so they no longer needlessly clutter the menus with different categories
- Removed previous "Add Parameter" category for the "Add Niagara Parameter Map Parameter" action as it would be a duplicate nested category
- Removed "Suggested" markdown from "Add new scratchpad module" action
- Removed the "Add rapid iteration param" action

diff --git a/Engine/Plugins/FX/Niagara/Source/NiagaraEditor/Private/EdGraphSchema_Niagara.cpp b/Engine/Plugins/FX/Niagara/Source/NiagaraEditor/Private/EdGraphSchema_Niagara.cpp
--- a/Engine/Plugins/FX/Niagara/Source/NiagaraEditor/Private/EdGraphSchema_Niagara.cpp
+++ b/Engine/Plugins/FX/Niagara/Source/NiagaraEditor/Private/EdGraphSchema_Niagara.cpp
@@ -798,9 +798,11 @@ TArray<TSharedPtr<FNiagaraAction_NewNode>> UEdGraphSchema_Niagara::GetGraphActio
 				PinType = PinToTypeDefinition(FromPin);
 			}
 
-			if (PinType.GetStruct())
+			// we don't want the add parameter list in module or dynamic input graphs
+			if (PinType.GetStruct() && !bModuleGraph && !bDynamicInputGraph)
 			{
 				const FText MenuDescFmt = LOCTEXT("Add ParameterFmt", "Add {0} Parameter");
+				const FText AddParameterCategory = LOCTEXT("AddParameterCat", "Add Parameter");
 				const TArray<FNiagaraTypeDefinition>& RegisteredTypes = FNiagaraTypeRegistry::GetRegisteredParameterTypes();
 				for (FNiagaraTypeDefinition Type : RegisteredTypes)
 				{
@@ -809,15 +811,15 @@ TArray<TSharedPtr<FNiagaraAction_NewNode>> UEdGraphSchema_Niagara::GetGraphActio
 						continue;
 					}
 
-					FText MenuCat;
+					TArray<FString> Categories;
+					Categories.Add(AddParameterCategory.ToString());
+					
 					if (const UClass* Class = Type.GetClass())
 					{						
-						MenuCat = FObjectEditorUtils::GetCategoryText(Class);
+						Categories.Add(FObjectEditorUtils::GetCategoryText(Class).ToString());
 					}
 					else
 					{
-						MenuCat = LOCTEXT("AddParameterCat", "Add Parameter");
-
 						// If you are in dynamic inputs or modules, we only allow free-range variables for 
 						// data interfaces and parameter maps.
 						if (bDynamicInputGraph || bModuleGraph)
@@ -832,34 +834,11 @@ TArray<TSharedPtr<FNiagaraAction_NewNode>> UEdGraphSchema_Niagara::GetGraphActio
 					const FText DisplayName = FText::Format(MenuDescFmt, Type.GetNameText());
 					
 					UNiagaraNodeInput* InputNode = NewObject<UNiagaraNodeInput>(OwnerOfTemporaries);
-					AddNewNodeMenuAction(NewActions, InputNode, DisplayName, ENiagaraMenuSections::General, {MenuCat.ToString()}, FText::GetEmpty(), FText::GetEmpty());
+					AddNewNodeMenuAction(NewActions, InputNode, DisplayName, ENiagaraMenuSections::General, Categories, FText::GetEmpty(), FText::GetEmpty());
 					FNiagaraEditorUtilities::InitializeParameterInputNode(*InputNode, Type, NiagaraGraph);
 				}
 
-				// TODO sckime please remove this..
-				if (bSystemGraph || IsParticleGraph(NiagaraGraph))
-				{
-					for (FNiagaraTypeDefinition Type : RegisteredTypes)
-					{
-						FText MenuCat;
-						if (const UClass* Class = Type.GetClass())
-						{
-							continue;
-						}
-						else
-						{
-							MenuCat = LOCTEXT("AddRIParameterCat", "Add Rapid Iteration Param");
-						}
-
-						const FText DisplayName = FText::Format(MenuDescFmt, Type.GetNameText());
-						
-						UNiagaraNodeInput* InputNode = NewObject<UNiagaraNodeInput>(OwnerOfTemporaries);
-						AddNewNodeMenuAction(NewActions, InputNode, DisplayName, ENiagaraMenuSections::General, {MenuCat.ToString()}, FText::GetEmpty(), FText::GetEmpty());
-						FNiagaraEditorUtilities::InitializeParameterInputNode(*InputNode, Type, NiagaraGraph);
-						InputNode->Usage = ENiagaraInputNodeUsage::RapidIterationParameter;
-					}
-				}
-				
+				// this allows adding a parameter of the type of the dragged-from input pin
 				if (PinType != FNiagaraTypeDefinition::GetGenericNumericDef())
 				{
 					//For correctly typed pins, offer the correct type at the top level.				
diff --git a/Engine/Plugins/FX/Niagara/Source/NiagaraEditor/Private/ViewModels/Stack/NiagaraStackScriptItemGroup.cpp b/Engine/Plugins/FX/Niagara/Source/NiagaraEditor/Private/ViewModels/Stack/NiagaraStackScriptItemGroup.cpp
--- a/Engine/Plugins/FX/Niagara/Source/NiagaraEditor/Private/ViewModels/Stack/NiagaraStackScriptItemGroup.cpp
+++ b/Engine/Plugins/FX/Niagara/Source/NiagaraEditor/Private/ViewModels/Stack/NiagaraStackScriptItemGroup.cpp
@@ -111,7 +111,7 @@ public:
 	{
 		FText DisplayName = LOCTEXT("NewScratchModuleName", "New Scratch Pad Module");
 		FText Description = LOCTEXT("NewScratchModuleDescription", "Description: Create a new scratch pad module.");
-		return MakeShareable(new FScriptGroupAddAction( DisplayName, {}, Description, FText(), true, true, FNiagaraVariable(), false, FAssetData(), nullptr, true, false));
+		return MakeShareable(new FScriptGroupAddAction( DisplayName, {}, Description, FText(), false, true, FNiagaraVariable(), false, FAssetData(), nullptr, true, false));
 	}
 
 	static TSharedRef<FScriptGroupAddAction> CreateNewSetSpecificModuleAction()
