Ensure that systems and emitter are fully loaded before editing them

diff --git a/Engine/Plugins/FX/Niagara/Source/NiagaraEditor/Private/Toolkits/NiagaraSystemToolkit.cpp b/Engine/Plugins/FX/Niagara/Source/NiagaraEditor/Private/Toolkits/NiagaraSystemToolkit.cpp
--- a/Engine/Plugins/FX/Niagara/Source/NiagaraEditor/Private/Toolkits/NiagaraSystemToolkit.cpp
+++ b/Engine/Plugins/FX/Niagara/Source/NiagaraEditor/Private/Toolkits/NiagaraSystemToolkit.cpp
@@ -235,6 +235,7 @@ void FNiagaraSystemToolkit::InitializeWithSystem(const EToolkitMode::Type Mode,
 {
 	System = &InSystem;
 	Emitter = nullptr;
+	System->EnsureFullyLoaded();
 
 	FNiagaraSystemViewModelOptions SystemOptions;
 	SystemOptions.bCanModifyEmittersFromTimeline = true;
@@ -276,8 +277,10 @@ void FNiagaraSystemToolkit::InitializeWithEmitter(const EToolkitMode::Type Mode,
 {
 	System = NewObject<UNiagaraSystem>(GetTransientPackage(), NAME_None, RF_Transient | RF_Transactional);
 	UNiagaraSystemFactoryNew::InitializeSystem(System, true);
+	System->EnsureFullyLoaded();
 
 	Emitter = &InEmitter;
+	Emitter->UpdateEmitterAfterLoad();
 
 	// Before copying the emitter prepare the rapid iteration parameters so that the post compile prepare doesn't
 	// cause the change ids to become out of sync.
