Fixed a bug where propagating a static switch would delete it from the parameter panel

diff --git a/Engine/Plugins/FX/Niagara/Source/NiagaraEditor/Private/Customizations/NiagaraFunctionCallNodeDetails.cpp b/Engine/Plugins/FX/Niagara/Source/NiagaraEditor/Private/Customizations/NiagaraFunctionCallNodeDetails.cpp
--- a/Engine/Plugins/FX/Niagara/Source/NiagaraEditor/Private/Customizations/NiagaraFunctionCallNodeDetails.cpp
+++ b/Engine/Plugins/FX/Niagara/Source/NiagaraEditor/Private/Customizations/NiagaraFunctionCallNodeDetails.cpp
@@ -132,6 +132,7 @@ void FNiagaraFunctionCallNodeDetails::CustomizeDetails(IDetailLayoutBuilder& Det
 						CopyMetadataFromCalledGraph(SwitchInput);
 					}
 					Node->RefreshFromExternalChanges();
+					Node->GetGraph()->NotifyGraphChanged();
 				})
 			]
 			+ SVerticalBox::Slot()
diff --git a/Engine/Plugins/FX/Niagara/Source/NiagaraEditor/Private/NiagaraEditorUtilities.cpp b/Engine/Plugins/FX/Niagara/Source/NiagaraEditor/Private/NiagaraEditorUtilities.cpp
--- a/Engine/Plugins/FX/Niagara/Source/NiagaraEditor/Private/NiagaraEditorUtilities.cpp
+++ b/Engine/Plugins/FX/Niagara/Source/NiagaraEditor/Private/NiagaraEditorUtilities.cpp
@@ -342,6 +342,11 @@ void FNiagaraEditorUtilities::FixUpPastedNodes(UEdGraph* Graph, TSet<UEdGraphNod
 				ExistingNames.Add(ActualPastedFunctionCallName);
 				OldFunctionToNewFunctionNameMap.Add(FunctionCallName, ActualPastedFunctionCallName);
 			}
+			UNiagaraGraph* NiagaraGraph = CastChecked<UNiagaraGraph>(Graph);
+			for (const FNiagaraPropagatedVariable& PropagatedVariable : PastedFunctionCallNode->PropagatedStaticSwitchParameters)
+			{
+				NiagaraGraph->AddParameter(PropagatedVariable.ToVariable(), true);
+			}
 		}
 	}
 
diff --git a/Engine/Plugins/FX/Niagara/Source/NiagaraEditor/Private/NiagaraGraph.cpp b/Engine/Plugins/FX/Niagara/Source/NiagaraEditor/Private/NiagaraGraph.cpp
--- a/Engine/Plugins/FX/Niagara/Source/NiagaraEditor/Private/NiagaraGraph.cpp
+++ b/Engine/Plugins/FX/Niagara/Source/NiagaraEditor/Private/NiagaraGraph.cpp
@@ -158,6 +158,7 @@ void UNiagaraGraph::PostLoad()
 			// Conditional postload all ScriptVars to ensure static switch default values are allocated as these are required when postloading all graph nodes later.
 			ScriptVar->ConditionalPostLoad();
 		}
+		ScriptVar->SetIsStaticSwitch(FindStaticSwitchInputs().Contains(Var));
 	}
 
 	for (UEdGraphNode* Node : Nodes)
@@ -2653,6 +2654,7 @@ void UNiagaraGraph::SetMetaData(const FNiagaraVariable& InVar, const FNiagaraVar
 		Modify();
 		UE_TRANSITIONAL_OBJECT_PTR(UNiagaraScriptVariable)& NewScriptVariable = VariableToScriptVariable.Add(InVar, NewObject<UNiagaraScriptVariable>(this, FName(), RF_Transactional));
 		NewScriptVariable->Init(InVar, InMetaData);
+		NewScriptVariable->SetIsStaticSwitch(FindStaticSwitchInputs().Contains(InVar));
 	}
 }
 
@@ -2750,7 +2752,7 @@ void UNiagaraGraph::RefreshParameterReferences() const
 			}
 		}
 	}
-
+	
 	// Check all pins on all nodes in the graph to find parameter pins which may have been missed in the parameter map traversal.  This
 	// can happen for nodes which are not fully connected and therefore don't show up in the traversal.
 	const UEdGraphSchema_Niagara* NiagaraSchema = GetNiagaraSchema();
