Fix args mismatch in NotifyPostChange() call.

diff --git a/Engine/Plugins/FX/Niagara/Source/NiagaraEditor/Private/Customizations/NiagaraScriptVariableCustomization.cpp b/Engine/Plugins/FX/Niagara/Source/NiagaraEditor/Private/Customizations/NiagaraScriptVariableCustomization.cpp
--- a/Engine/Plugins/FX/Niagara/Source/NiagaraEditor/Private/Customizations/NiagaraScriptVariableCustomization.cpp
+++ b/Engine/Plugins/FX/Niagara/Source/NiagaraEditor/Private/Customizations/NiagaraScriptVariableCustomization.cpp
@@ -534,7 +534,7 @@ void FNiagaraScriptVariableDetails::OnEndLibraryValueChanged()
 	if (TypeUtilityLibraryValue && ParameterEditorLibraryValue && CachedDetailBuilder.IsValid())
 	{
 		const TSharedPtr<IPropertyHandle> DefaultValueHandle = CachedDetailBuilder.Pin()->GetProperty("DefaultValueVariant", UNiagaraScriptVariable::StaticClass());
-		DefaultValueHandle->NotifyPostChange();
+		DefaultValueHandle->NotifyPostChange(EPropertyChangeType::ValueSet);
 		DefaultValueHandle->NotifyFinishedChangingProperties();
 	}
 
@@ -586,7 +586,7 @@ void FNiagaraScriptVariableDetails::OnLibraryValueChanged()
 			{
 				Variable->UpdateChangeId();
 
-				DefaultValueHandle->NotifyPostChange();
+				DefaultValueHandle->NotifyPostChange(EPropertyChangeType::ValueSet);
 				DefaultValueHandle->NotifyFinishedChangingProperties();
 
 				if (UNiagaraParameterDefinitions* OuterParameterDefinitions = Cast<UNiagaraParameterDefinitions>(Variable->GetOuter()))
