Fix object pooling during compilation; release duplicated graphs after translation instead of directly after precompile is finished.

diff --git a/Engine/Plugins/FX/Niagara/Source/Niagara/Classes/NiagaraScriptSourceBase.h b/Engine/Plugins/FX/Niagara/Source/Niagara/Classes/NiagaraScriptSourceBase.h
--- a/Engine/Plugins/FX/Niagara/Source/Niagara/Classes/NiagaraScriptSourceBase.h
+++ b/Engine/Plugins/FX/Niagara/Source/Niagara/Classes/NiagaraScriptSourceBase.h
@@ -36,6 +36,7 @@ public:
 	virtual TSharedPtr<FNiagaraCompileRequestDataBase, ESPMode::ThreadSafe> GetDependentRequest(int32 Index) = 0;
 	virtual FName ResolveEmitterAlias(FName VariableName) const = 0;
 	virtual bool GetUseRapidIterationParams() const = 0;
+	virtual void ReleaseCompilationCopies() = 0;
 };
 
 class FNiagaraCompileOptions
diff --git a/Engine/Plugins/FX/Niagara/Source/Niagara/Private/NiagaraSystem.cpp b/Engine/Plugins/FX/Niagara/Source/Niagara/Private/NiagaraSystem.cpp
--- a/Engine/Plugins/FX/Niagara/Source/Niagara/Private/NiagaraSystem.cpp
+++ b/Engine/Plugins/FX/Niagara/Source/Niagara/Private/NiagaraSystem.cpp
@@ -2633,6 +2633,9 @@ bool UNiagaraSystem::RequestCompile(bool bForce, FNiagaraSystemUpdateContext* Op
 			}
 		}
 
+		// clean up copies from precompile data
+		TSharedPtr<FNiagaraCompileRequestDataBase, ESPMode::ThreadSafe> SystemPrecompiledData = ActiveCompilation.MappedData.FindChecked(SystemSpawnScript);
+		SystemPrecompiledData->ReleaseCompilationCopies();
 	}
 
 
diff --git a/Engine/Plugins/FX/Niagara/Source/NiagaraCore/Private/NiagaraCustomVersion.cpp b/Engine/Plugins/FX/Niagara/Source/NiagaraCore/Private/NiagaraCustomVersion.cpp
--- a/Engine/Plugins/FX/Niagara/Source/NiagaraCore/Private/NiagaraCustomVersion.cpp
+++ b/Engine/Plugins/FX/Niagara/Source/NiagaraCore/Private/NiagaraCustomVersion.cpp
@@ -8,4 +8,4 @@ const FGuid FNiagaraCustomVersion::GUID(0xFCF57AFA, 0x50764283, 0xB9A9E658, 0xFF
 // Register the custom version with core
 FCustomVersionRegistration GRegisterNiagaraCustomVersion(FNiagaraCustomVersion::GUID, FNiagaraCustomVersion::LatestVersion, TEXT("NiagaraVer"));
 
-const FGuid FNiagaraCustomVersion::LatestScriptCompileVersion(0x0EEBCD8A, 0xE45341FF, 0x84B53A98, 0x46EC4EC1);
+const FGuid FNiagaraCustomVersion::LatestScriptCompileVersion(0x17E1CB82, 0x150B4598, 0xA267B22C, 0x31B13B19);
diff --git a/Engine/Plugins/FX/Niagara/Source/NiagaraEditor/Private/NiagaraCompiler.cpp b/Engine/Plugins/FX/Niagara/Source/NiagaraEditor/Private/NiagaraCompiler.cpp
--- a/Engine/Plugins/FX/Niagara/Source/NiagaraEditor/Private/NiagaraCompiler.cpp
+++ b/Engine/Plugins/FX/Niagara/Source/NiagaraEditor/Private/NiagaraCompiler.cpp
@@ -363,21 +363,6 @@ void FNiagaraCompileRequestData::GetReferencedObjects(TArray<UObject*>& Objects)
 	}
 }
 
-FNiagaraCompileRequestData::~FNiagaraCompileRequestData()
-{
-	if (NodeGraphDeepCopy.IsValid())
-	{
-		NodeGraphDeepCopy->ReleaseCompilationCopy();
-	}
-	for (TWeakObjectPtr<UNiagaraGraph> Graph : ClonedGraphs)
-	{
-		if (Graph.IsValid())
-		{
-			Graph->ReleaseCompilationCopy();
-		}
-	}
-}
-
 bool FNiagaraCompileRequestData::GatherPreCompiledVariables(const FString& InNamespaceFilter, TArray<FNiagaraVariable>& OutVars)
 {
 	if (PrecompiledHistories.Num() == 0)
@@ -482,6 +467,19 @@ void FNiagaraCompileRequestData::AddRapidIterationParameters(const FNiagaraParam
 	}
 }
 
+void FNiagaraCompileRequestData::ReleaseCompilationCopies()
+{
+	// clean up graph copies
+	for (auto& ClonedGraph : ClonedGraphs)
+	{
+		ClonedGraph->ReleaseCompilationCopy();
+	}
+	for (auto& EmitterCompileData : EmitterData)
+	{
+		EmitterCompileData->ReleaseCompilationCopies();
+	}
+}
+
 void FNiagaraCompileRequestData::FinishPrecompile(UNiagaraScriptSource* ScriptSource, const TArray<FNiagaraVariable>& EncounterableVariables, ENiagaraScriptUsage InUsage, FCompileConstantResolver ConstantResolver, const TArray<UNiagaraSimulationStageBase*>* SimStages)
 {
 	//double StartTime = FPlatformTime::Seconds();
@@ -645,12 +643,6 @@ void FNiagaraCompileRequestData::FinishPrecompile(UNiagaraScriptSource* ScriptSo
 				}
 			}
 		}
-
-		// clean up graph copies
-		for (auto& ClonedGraph : ClonedGraphs)
-		{
-			ClonedGraph->ReleaseCompilationCopy();
-		}
 	}
 
 	//float DeltaTime = (float)(FPlatformTime::Seconds() - StartTime);
diff --git a/Engine/Plugins/FX/Niagara/Source/NiagaraEditor/Private/NiagaraGraph.cpp b/Engine/Plugins/FX/Niagara/Source/NiagaraEditor/Private/NiagaraGraph.cpp
--- a/Engine/Plugins/FX/Niagara/Source/NiagaraEditor/Private/NiagaraGraph.cpp
+++ b/Engine/Plugins/FX/Niagara/Source/NiagaraEditor/Private/NiagaraGraph.cpp
@@ -951,6 +951,7 @@ void UNiagaraGraph::FindEquivalentOutputNodes(ENiagaraScriptUsage TargetUsageTyp
 
 UNiagaraGraph* UNiagaraGraph::CreateCompilationCopy()
 {
+	check(!bIsForCompilationOnly);
 	UNiagaraGraph* Result = NewObject<UNiagaraGraph>();
 	FNiagaraEditorModule& NiagaraEditorModule = FModuleManager::GetModuleChecked<FNiagaraEditorModule>("NiagaraEditor");
 
@@ -966,7 +967,7 @@ UNiagaraGraph* UNiagaraGraph::CreateCompilationCopy()
 	Result->bIsForCompilationOnly = true;
 
 	// get new script variables from the pool
-	for (auto& It : VariableToScriptVariable)
+	for (auto& It : Result->VariableToScriptVariable)
 	{
 		It.Value = CastChecked<UNiagaraScriptVariable>(NiagaraEditorModule.GetPooledDuplicateObject(It.Value));
 	}
@@ -995,6 +996,10 @@ UNiagaraGraph* UNiagaraGraph::CreateCompilationCopy()
 		}
 	}
 	Result->Nodes = NewNodes;
+
+	// probably not necessary for compilation, but remove references to the original graph
+	Result->ParameterToReferencesMap.Empty();
+	Result->RefreshParameterReferences();
 	
 	return Result;
 }
@@ -1010,7 +1015,10 @@ void UNiagaraGraph::ReleaseCompilationCopy()
 	{
 		NiagaraEditorModule.ReleaseObjectToPool(It.Value);
 	}
+
+	// clear script variables and kill this object to surface any invalid access after it's released
 	VariableToScriptVariable.Empty();
+	MarkPendingKill();
 }
 
 UNiagaraNodeOutput* UNiagaraGraph::FindOutputNode(ENiagaraScriptUsage TargetUsageType, FGuid TargetUsageId) const
diff --git a/Engine/Plugins/FX/Niagara/Source/NiagaraEditor/Private/NiagaraHlslTranslator.h b/Engine/Plugins/FX/Niagara/Source/NiagaraEditor/Private/NiagaraHlslTranslator.h
--- a/Engine/Plugins/FX/Niagara/Source/NiagaraEditor/Private/NiagaraHlslTranslator.h
+++ b/Engine/Plugins/FX/Niagara/Source/NiagaraEditor/Private/NiagaraHlslTranslator.h
@@ -66,7 +66,6 @@ public:
 
 	}
 
-	virtual ~FNiagaraCompileRequestData();
 	virtual bool GatherPreCompiledVariables(const FString& InNamespaceFilter, TArray<FNiagaraVariable>& OutVars) override;
 	virtual void GetReferencedObjects(TArray<UObject*>& Objects) override;
 	virtual const TMap<FName, UNiagaraDataInterface*>& GetObjectNameMap() override;
@@ -88,6 +87,7 @@ public:
 	}
 	void AddRapidIterationParameters(const FNiagaraParameterStore& InParamStore, FCompileConstantResolver InResolver);
 	virtual bool GetUseRapidIterationParams() const override { return bUseRapidIterationParams; }
+	virtual void ReleaseCompilationCopies() override;
 
 	// Simulation Stage Variables. Sim stage of 0 is always Spawn/Update
 	TArray<uint32> NumIterationsPerStage;
diff --git a/Engine/Plugins/FX/Niagara/Source/NiagaraEditor/Private/NiagaraScriptSource.cpp b/Engine/Plugins/FX/Niagara/Source/NiagaraEditor/Private/NiagaraScriptSource.cpp
--- a/Engine/Plugins/FX/Niagara/Source/NiagaraEditor/Private/NiagaraScriptSource.cpp
+++ b/Engine/Plugins/FX/Niagara/Source/NiagaraEditor/Private/NiagaraScriptSource.cpp
@@ -107,7 +107,7 @@ void UNiagaraScriptSource::PostLoad()
 UNiagaraScriptSource* UNiagaraScriptSource::CreateCompilationCopy()
 {
 	UNiagaraScriptSource* Result = NewObject<UNiagaraScriptSource>();
-	Result->NodeGraph = DuplicateObject(NodeGraph, GetTransientPackage());// NodeGraph->CreateCompilationCopy();
+	Result->NodeGraph = NodeGraph->CreateCompilationCopy();
 	return Result;
 }
 
