- Fix typo in console variable description

#rnx

diff --git a/Engine/Plugins/FX/Niagara/Source/Niagara/Private/NiagaraSystemSimulation.cpp b/Engine/Plugins/FX/Niagara/Source/Niagara/Private/NiagaraSystemSimulation.cpp
--- a/Engine/Plugins/FX/Niagara/Source/Niagara/Private/NiagaraSystemSimulation.cpp
+++ b/Engine/Plugins/FX/Niagara/Source/Niagara/Private/NiagaraSystemSimulation.cpp
@@ -116,7 +116,7 @@ static int32 GNiagaraSystemSimulationTickTaskShouldWait = 1;
 static FAutoConsoleVariableRef CVarNiagaraSystemSimulationTickTaskShouldWait(
 	TEXT("fx.Niagara.SystemSimulation.TickTaskShouldWait"),
 	GNiagaraSystemSimulationTickTaskShouldWait,
-	TEXT("When enabled the tick task will wait for concurrent work to complete, when disabled the task is complete one the GT tick is complete."),
+	TEXT("When enabled the tick task will wait for concurrent work to complete, when disabled the task is complete once the GT tick is complete."),
 	ECVF_Default
 );
 
diff --git a/Engine/Plugins/FX/Niagara/Source/Niagara/Private/NiagaraWorldManager.cpp b/Engine/Plugins/FX/Niagara/Source/Niagara/Private/NiagaraWorldManager.cpp
--- a/Engine/Plugins/FX/Niagara/Source/Niagara/Private/NiagaraWorldManager.cpp
+++ b/Engine/Plugins/FX/Niagara/Source/Niagara/Private/NiagaraWorldManager.cpp
@@ -689,7 +689,6 @@ void FNiagaraWorldManager::PreSendAllEndOfFrameUpdates()
 {
 	if ( SimulationsWithPostActorWork.Num() > 0 )
 	{
-		//-TODO: This should not happen
 		for (const auto& Simulation : SimulationsWithPostActorWork)
 		{
 			if (Simulation->IsValid())
