Fixed a bug where the Script.Usage static switch would generate an error if a new default value was defined behind one of its branches. This was because during precompile of the emitter scripts, the parameter map histories all used the same constant resolver.

diff --git a/Engine/Plugins/FX/Niagara/Source/NiagaraCore/Private/NiagaraCustomVersion.cpp b/Engine/Plugins/FX/Niagara/Source/NiagaraCore/Private/NiagaraCustomVersion.cpp
--- a/Engine/Plugins/FX/Niagara/Source/NiagaraCore/Private/NiagaraCustomVersion.cpp
+++ b/Engine/Plugins/FX/Niagara/Source/NiagaraCore/Private/NiagaraCustomVersion.cpp
@@ -8,4 +8,4 @@ const FGuid FNiagaraCustomVersion::GUID(0xFCF57AFA, 0x50764283, 0xB9A9E658, 0xFF
 // Register the custom version with core
 FCustomVersionRegistration GRegisterNiagaraCustomVersion(FNiagaraCustomVersion::GUID, FNiagaraCustomVersion::LatestVersion, TEXT("NiagaraVer"));
 
-const FGuid FNiagaraCustomVersion::LatestScriptCompileVersion(0x17E1CB82, 0x150B4598, 0xA267B22C, 0x31B13B19);
+const FGuid FNiagaraCustomVersion::LatestScriptCompileVersion(0x7F8E9CA7, 0x61806B41, 0xA6C7E8A3, 0xCB2120F5);
diff --git a/Engine/Plugins/FX/Niagara/Source/NiagaraEditor/Private/NiagaraCompiler.cpp b/Engine/Plugins/FX/Niagara/Source/NiagaraEditor/Private/NiagaraCompiler.cpp
--- a/Engine/Plugins/FX/Niagara/Source/NiagaraEditor/Private/NiagaraCompiler.cpp
+++ b/Engine/Plugins/FX/Niagara/Source/NiagaraEditor/Private/NiagaraCompiler.cpp
@@ -98,7 +98,7 @@ T* PrecompileDuplicateObject(T const* SourceObject, UObject* Outer, const FName
 
 }
 
-void FNiagaraCompileRequestData::VisitReferencedGraphs(UNiagaraGraph* InSrcGraph, UNiagaraGraph* InDupeGraph, ENiagaraScriptUsage InUsage, FCompileConstantResolver ConstantResolver, bool bNeedsCompilation)
+void FNiagaraCompileRequestData::VisitReferencedGraphs(UNiagaraGraph* InSrcGraph, UNiagaraGraph* InDupeGraph, ENiagaraScriptUsage InUsage, FCompileConstantResolver ConstantResolver, bool bNeedsCompilation, TMap<UNiagaraNodeFunctionCall*, ENiagaraScriptUsage> FunctionsWithUsage)
 {
 	if (!InDupeGraph && !InSrcGraph)
 	{
@@ -128,10 +128,10 @@ void FNiagaraCompileRequestData::VisitReferencedGraphs(UNiagaraGraph* InSrcGraph
 	}
 	ClonedGraphs.AddUnique(InDupeGraph);
 
-	VisitReferencedGraphsRecursive(InDupeGraph, ConstantResolver, bNeedsCompilation);
+	VisitReferencedGraphsRecursive(InDupeGraph, ConstantResolver, bNeedsCompilation, FunctionsWithUsage);
 }
 
-void FNiagaraCompileRequestData::VisitReferencedGraphsRecursive(UNiagaraGraph* InGraph, const FCompileConstantResolver& ConstantResolver, bool bNeedsCompilation)
+void FNiagaraCompileRequestData::VisitReferencedGraphsRecursive(UNiagaraGraph* InGraph, const FCompileConstantResolver& ConstantResolver, bool bNeedsCompilation, TMap<UNiagaraNodeFunctionCall*, ENiagaraScriptUsage> FunctionsWithUsage)
 {
 	if (!InGraph)
 	{
@@ -170,6 +170,12 @@ void FNiagaraCompileRequestData::VisitReferencedGraphsRecursive(UNiagaraGraph* I
 				UNiagaraScript* FunctionScript = FunctionCallNode->FunctionScript;
 				ScriptUsage = FunctionCallNode->GetCalledUsage();
 
+				FCompileConstantResolver FunctionConstantResolver = ConstantResolver;
+				if (FunctionsWithUsage.Contains(FunctionCallNode))
+				{
+					FunctionConstantResolver = ConstantResolver.WithUsage(FunctionsWithUsage[FunctionCallNode]);
+				}
+
 				if (FunctionScript != nullptr)
 				{
 					UNiagaraGraph* FunctionGraph = FunctionCallNode->GetCalledGraph();
@@ -232,8 +238,8 @@ void FNiagaraCompileRequestData::VisitReferencedGraphsRecursive(UNiagaraGraph* I
 								if (bNeedsCompilation)
 								{
 									FEdGraphUtilities::MergeChildrenGraphsIn(ProcessedGraph, ProcessedGraph, /*bRequireSchemaMatch=*/ true);
-
-									FNiagaraEditorUtilities::PreprocessFunctionGraph(Schema, ProcessedGraph, CallInputs, CallOutputs, ScriptUsage, ConstantResolver);
+									
+									FNiagaraEditorUtilities::PreprocessFunctionGraph(Schema, ProcessedGraph, CallInputs, CallOutputs, ScriptUsage, FunctionConstantResolver);
 								}
 								FunctionCallNode->FunctionScript = DupeScript;
 							}
@@ -248,7 +254,7 @@ void FNiagaraCompileRequestData::VisitReferencedGraphsRecursive(UNiagaraGraph* I
 							TArray<FunctionData> MadeArray;
 							MadeArray.Add(Data);
 							PreprocessedFunctions.Add(FunctionGraph, MadeArray);
-							VisitReferencedGraphsRecursive(ProcessedGraph, ConstantResolver, bNeedsCompilation);
+							VisitReferencedGraphsRecursive(ProcessedGraph, FunctionConstantResolver, bNeedsCompilation, FunctionsWithUsage);
 							ClonedGraphs.AddUnique(ProcessedGraph);
 						}
 						else if (bHasNumericParams)
@@ -264,7 +270,7 @@ void FNiagaraCompileRequestData::VisitReferencedGraphsRecursive(UNiagaraGraph* I
 							if (bNeedsCompilation)
 							{
 								FEdGraphUtilities::MergeChildrenGraphsIn(ProcessedGraph, ProcessedGraph, /*bRequireSchemaMatch=*/ true);
-								FNiagaraEditorUtilities::PreprocessFunctionGraph(Schema, ProcessedGraph, CallInputs, CallOutputs, ScriptUsage, ConstantResolver);
+								FNiagaraEditorUtilities::PreprocessFunctionGraph(Schema, ProcessedGraph, CallInputs, CallOutputs, ScriptUsage, FunctionConstantResolver);
 							}
 							FunctionCallNode->FunctionScript = DupeScript;
 
@@ -280,7 +286,7 @@ void FNiagaraCompileRequestData::VisitReferencedGraphsRecursive(UNiagaraGraph* I
 							Data.bHasNumericInputs = bHasNumericInputs;
 
 							FoundArray->Add(Data);
-							VisitReferencedGraphsRecursive(ProcessedGraph, ConstantResolver, bNeedsCompilation);
+							VisitReferencedGraphsRecursive(ProcessedGraph, FunctionConstantResolver, bNeedsCompilation, FunctionsWithUsage);
 						}
 						else if(bRequiresClonedScript)
 						{
@@ -408,6 +414,27 @@ void FNiagaraCompileRequestData::DeepCopyGraphs(UNiagaraScriptSource* ScriptSour
 	//}
 }
 
+void FNiagaraCompileRequestData::DeepCopyGraphs(UNiagaraScriptSource* ScriptSource, UNiagaraEmitter* Emitter, bool bNeedsCompilation)
+{
+	Source = bNeedsCompilation ? PrecompileDuplicateObject<UNiagaraScriptSource>(ScriptSource, GetTransientPackage()) : ScriptSource;
+	NodeGraphDeepCopy = Source->NodeGraph;
+	if (bNeedsCompilation)
+	{
+		FEdGraphUtilities::MergeChildrenGraphsIn(NodeGraphDeepCopy.Get(), NodeGraphDeepCopy.Get(), /*bRequireSchemaMatch=*/ true);
+	}
+	TMap<UNiagaraNodeFunctionCall*, ENiagaraScriptUsage> FunctionsWithUsage;
+	for (UEdGraphNode* Node : ScriptSource->NodeGraph->Nodes)
+	{
+		if (UNiagaraNodeFunctionCall* FunctionCall = Cast<UNiagaraNodeFunctionCall>(Node))
+		{
+			UNiagaraNodeOutput* OutputNode = FNiagaraStackGraphUtilities::GetEmitterOutputNodeForStackNode(*FunctionCall);
+			FunctionsWithUsage.Add(FunctionCall, OutputNode ? OutputNode->GetUsage() : ENiagaraScriptUsage::EmitterSpawnScript);
+		}
+	}
+	FCompileConstantResolver ConstantResolver(Emitter, ENiagaraScriptUsage::EmitterSpawnScript);
+	VisitReferencedGraphs(ScriptSource->NodeGraph, NodeGraphDeepCopy.Get(), ENiagaraScriptUsage::EmitterSpawnScript, ConstantResolver, bNeedsCompilation, FunctionsWithUsage);
+}
+
 
 void FNiagaraCompileRequestData::AddRapidIterationParameters(const FNiagaraParameterStore& InParamStore, FCompileConstantResolver InResolver)
 {
@@ -480,7 +507,7 @@ void FNiagaraCompileRequestData::ReleaseCompilationCopies()
 	}
 }
 
-void FNiagaraCompileRequestData::FinishPrecompile(UNiagaraScriptSource* ScriptSource, const TArray<FNiagaraVariable>& EncounterableVariables, ENiagaraScriptUsage InUsage, FCompileConstantResolver ConstantResolver, const TArray<UNiagaraSimulationStageBase*>* SimStages)
+void FNiagaraCompileRequestData::FinishPrecompile(const TArray<FNiagaraVariable>& EncounterableVariables, ENiagaraScriptUsage InUsage, FCompileConstantResolver ConstantResolver, const TArray<UNiagaraSimulationStageBase*>* SimStages)
 {
 	//double StartTime = FPlatformTime::Seconds();
 	{
@@ -709,10 +736,11 @@ TSharedPtr<FNiagaraCompileRequestDataBase, ESPMode::ThreadSafe> FNiagaraEditorMo
 		bool bNeedsCompilation = ScriptSourceNeedsCompiling(Source, InCompilingScripts);
 		BasePtr->DeepCopyGraphs(Source, Script->GetUsage(), EmptyResolver, bNeedsCompilation);
 		const TArray<FNiagaraVariable> EncounterableVariables;
-		BasePtr->FinishPrecompile(Source, EncounterableVariables, Script->GetUsage(), EmptyResolver, nullptr);
+		BasePtr->FinishPrecompile(EncounterableVariables, Script->GetUsage(), EmptyResolver, nullptr);
 	}
 	else if (System)
 	{
+		check(System->GetSystemSpawnScript()->GetLatestSource() == System->GetSystemUpdateScript()->GetLatestSource());
 		BasePtr->bUseRapidIterationParams = !System->bBakeOutRapidIteration;
 
 		// Store off the current variables in the exposed parameters list.
@@ -722,22 +750,18 @@ TSharedPtr<FNiagaraCompileRequestDataBase, ESPMode::ThreadSafe> FNiagaraEditorMo
 		// Create an array of variables that we might encounter when traversing the graphs (include the originally exposed vars above)
 		TArray<FNiagaraVariable> EncounterableVars(OriginalExposedParams);
 
-		// Create an array of variables that we *did* encounter when traversing the graphs.
-		check(System->GetSystemSpawnScript()->GetLatestSource() == System->GetSystemUpdateScript()->GetLatestSource());
-
 		// First deep copy all the emitter graphs referenced by the system so that we can later hook up emitter handles in the system traversal.
 		BasePtr->EmitterData.Empty();
 		for (int32 i = 0; i < System->GetEmitterHandles().Num(); i++)
 		{
 			const FNiagaraEmitterHandle& Handle = System->GetEmitterHandle(i);
-			FCompileConstantResolver ConstantResolver(Handle.GetInstance(), ENiagaraScriptUsage::EmitterSpawnScript);
 			TSharedPtr<FNiagaraCompileRequestData, ESPMode::ThreadSafe> EmitterPtr = MakeShared<FNiagaraCompileRequestData, ESPMode::ThreadSafe>();
 			EmitterPtr->EmitterUniqueName = Handle.GetInstance()->GetUniqueEmitterName();
 			if (Handle.GetIsEnabled()) // Don't need to copy the graph if we aren't going to use it.
 			{
 				UNiagaraScriptSource* Source = Cast<UNiagaraScriptSource>(Handle.GetInstance()->GraphSource);
 				bool bNeedsCompilation = ScriptSourceNeedsCompiling(Source, InCompilingScripts);
-				EmitterPtr->DeepCopyGraphs(Source, ENiagaraScriptUsage::EmitterSpawnScript, ConstantResolver, bNeedsCompilation);
+				EmitterPtr->DeepCopyGraphs(Source, Handle.GetInstance(), bNeedsCompilation);
 			}
 			EmitterPtr->SourceName = BasePtr->SourceName;
 			EmitterPtr->bUseRapidIterationParams = BasePtr->bUseRapidIterationParams || (!Handle.GetInstance()->bBakeOutRapidIteration);
@@ -751,7 +775,7 @@ TSharedPtr<FNiagaraCompileRequestDataBase, ESPMode::ThreadSafe> FNiagaraEditorMo
 			UNiagaraScriptSource* Source = Cast<UNiagaraScriptSource>(System->GetSystemSpawnScript()->GetLatestSource());
 			bool bNeedsCompilation = ScriptSourceNeedsCompiling(Source, InCompilingScripts);
 			BasePtr->DeepCopyGraphs(Source, ENiagaraScriptUsage::SystemSpawnScript, ConstantResolver, bNeedsCompilation);
-			BasePtr->FinishPrecompile(Source, EncounterableVars, ENiagaraScriptUsage::SystemSpawnScript, ConstantResolver, nullptr);
+			BasePtr->FinishPrecompile(EncounterableVars, ENiagaraScriptUsage::SystemSpawnScript, ConstantResolver, nullptr);
 		}
 
 		// Add the User and System variables that we did encounter to the list that emitters might also encounter.
@@ -781,7 +805,7 @@ TSharedPtr<FNiagaraCompileRequestDataBase, ESPMode::ThreadSafe> FNiagaraEditorMo
 						ScriptToEmitterMap.Add(EmitterScripts[ScriptIdx], Handle.GetInstance());
 					}
 				}
-				BasePtr->EmitterData[i]->FinishPrecompile(Cast<UNiagaraScriptSource>(Handle.GetInstance()->GraphSource), EncounterableVars, ENiagaraScriptUsage::EmitterSpawnScript, ConstantResolver, &Handle.GetInstance()->GetSimulationStages());
+				BasePtr->EmitterData[i]->FinishPrecompile(EncounterableVars, ENiagaraScriptUsage::EmitterSpawnScript, ConstantResolver, &Handle.GetInstance()->GetSimulationStages());
 				BasePtr->MergeInEmitterPrecompiledData(BasePtr->EmitterData[i].Get());
 
 				for (UNiagaraRendererProperties* RendererProperty : Handle.GetInstance()->GetRenderers())
diff --git a/Engine/Plugins/FX/Niagara/Source/NiagaraEditor/Private/NiagaraHlslTranslator.h b/Engine/Plugins/FX/Niagara/Source/NiagaraEditor/Private/NiagaraHlslTranslator.h
--- a/Engine/Plugins/FX/Niagara/Source/NiagaraEditor/Private/NiagaraHlslTranslator.h
+++ b/Engine/Plugins/FX/Niagara/Source/NiagaraEditor/Private/NiagaraHlslTranslator.h
@@ -76,9 +76,10 @@ public:
 	const TArray<FNiagaraParameterMapHistory>& GetPrecomputedHistories() const { return PrecompiledHistories; }
 	virtual const class UNiagaraGraph* GetPrecomputedNodeGraph() const { return NodeGraphDeepCopy.Get(); }
 	const FString& GetUniqueEmitterName() const { return EmitterUniqueName; }
-	void VisitReferencedGraphs(UNiagaraGraph* InSrcGraph, UNiagaraGraph* InDupeGraph, ENiagaraScriptUsage InUsage, FCompileConstantResolver ConstantResolver, bool bNeedsCompilation);
+	void VisitReferencedGraphs(UNiagaraGraph* InSrcGraph, UNiagaraGraph* InDupeGraph, ENiagaraScriptUsage InUsage, FCompileConstantResolver ConstantResolver, bool bNeedsCompilation, TMap<UNiagaraNodeFunctionCall*, ENiagaraScriptUsage> FunctionsWithUsage = TMap<UNiagaraNodeFunctionCall*, ENiagaraScriptUsage>());
 	void DeepCopyGraphs(UNiagaraScriptSource* ScriptSource, ENiagaraScriptUsage InUsage, FCompileConstantResolver ConstantResolver, bool bNeedsCompilation);
-	void FinishPrecompile(UNiagaraScriptSource* ScriptSource, const TArray<FNiagaraVariable>& EncounterableVariables, ENiagaraScriptUsage InUsage, FCompileConstantResolver ConstantResolver, const TArray<class UNiagaraSimulationStageBase*>* SimStages);
+	void DeepCopyGraphs(UNiagaraScriptSource* ScriptSource, UNiagaraEmitter* Emitter, bool bNeedsCompilation);
+	void FinishPrecompile(const TArray<FNiagaraVariable>& EncounterableVariables, ENiagaraScriptUsage InUsage, FCompileConstantResolver ConstantResolver, const TArray<class UNiagaraSimulationStageBase*>* SimStages);
 	virtual int32 GetDependentRequestCount() const override {
 		return EmitterData.Num();
 	};
@@ -130,7 +131,7 @@ public:
 	TArray<FNiagaraVariable> RequiredRendererVariables;
 
 protected:
-	void VisitReferencedGraphsRecursive(UNiagaraGraph* InGraph, const FCompileConstantResolver& ConstantResolver, bool bNeedsCompile);
+	void VisitReferencedGraphsRecursive(UNiagaraGraph* InGraph, const FCompileConstantResolver& ConstantResolver, bool bNeedsCompile, TMap<UNiagaraNodeFunctionCall*, ENiagaraScriptUsage> FunctionsWithUsage);
 };
 
 
diff --git a/Engine/Plugins/FX/Niagara/Source/NiagaraEditor/Private/NiagaraNodeFunctionCall.cpp b/Engine/Plugins/FX/Niagara/Source/NiagaraEditor/Private/NiagaraNodeFunctionCall.cpp
--- a/Engine/Plugins/FX/Niagara/Source/NiagaraEditor/Private/NiagaraNodeFunctionCall.cpp
+++ b/Engine/Plugins/FX/Niagara/Source/NiagaraEditor/Private/NiagaraNodeFunctionCall.cpp
@@ -1259,9 +1259,13 @@ void UNiagaraNodeFunctionCall::BuildParameterMapHistory(FNiagaraParameterMapHist
 		FPinCollectorArray InputPins;
 		GetInputPins(InputPins);
 
-		ENiagaraFunctionDebugState OldDebugState = OutHistory.ConstantResolver.SetDebugState(DebugState);
-		FNiagaraEditorUtilities::SetStaticSwitchConstants(FunctionGraph, InputPins, OutHistory.ConstantResolver);
-		OutHistory.ConstantResolver.SetDebugState(OldDebugState);
+		FCompileConstantResolver FunctionResolver = OutHistory.ConstantResolver.WithDebugState(DebugState);
+		if (OutHistory.HasCurrentUsageContext())
+		{
+			// if we traverse a full emitter graph the usage might change during the traversal, so we need to update the constant resolver
+			FunctionResolver = FunctionResolver.WithUsage(OutHistory.GetCurrentUsageContext());
+		}
+		FNiagaraEditorUtilities::SetStaticSwitchConstants(FunctionGraph, InputPins, FunctionResolver);
 
 		int32 ParamMapIdx = INDEX_NONE;
 		uint32 NodeIdx = INDEX_NONE;
@@ -1401,9 +1405,7 @@ UEdGraphPin* UNiagaraNodeFunctionCall::FindParameterMapDefaultValuePin(const FNa
 			TArray<UEdGraphPin*> InputPins;
 			GetInputPins(InputPins);
 
-			ENiagaraFunctionDebugState PreviousDebugState = ConstantResolver.SetDebugState(DebugState);
-			FNiagaraEditorUtilities::SetStaticSwitchConstants(ScriptSource->NodeGraph, InputPins, ConstantResolver);
-			ConstantResolver.SetDebugState(PreviousDebugState);
+			FNiagaraEditorUtilities::SetStaticSwitchConstants(ScriptSource->NodeGraph, InputPins, ConstantResolver.WithDebugState(DebugState));
 			
 			return ScriptSource->NodeGraph->FindParameterMapDefaultValuePin(VariableName, FunctionScript->GetUsage(), InParentUsage);
 		}
diff --git a/Engine/Plugins/FX/Niagara/Source/NiagaraEditor/Private/NiagaraParameterMapHistory.cpp b/Engine/Plugins/FX/Niagara/Source/NiagaraEditor/Private/NiagaraParameterMapHistory.cpp
--- a/Engine/Plugins/FX/Niagara/Source/NiagaraEditor/Private/NiagaraParameterMapHistory.cpp
+++ b/Engine/Plugins/FX/Niagara/Source/NiagaraEditor/Private/NiagaraParameterMapHistory.cpp
@@ -749,6 +749,11 @@ void FNiagaraParameterMapHistoryBuilder::EnableScriptWhitelist(bool bInEnable, E
 	FilterScriptType = InScriptType;
 }
 
+bool FNiagaraParameterMapHistoryBuilder::HasCurrentUsageContext() const
+{
+	return RelevantScriptUsageContext.Num() > 0;
+}
+
 
 bool FNiagaraParameterMapHistoryBuilder::ContextContains(ENiagaraScriptUsage InUsage) const
 {
@@ -1617,11 +1622,18 @@ ENiagaraFunctionDebugState FCompileConstantResolver::GetDebugState() const
 }
 
 
-ENiagaraFunctionDebugState FCompileConstantResolver::SetDebugState(ENiagaraFunctionDebugState InDebugState)
+FCompileConstantResolver FCompileConstantResolver::WithDebugState(ENiagaraFunctionDebugState InDebugState) const
+{
+	FCompileConstantResolver Copy = *this;
+	Copy.DebugState = InDebugState;
+	return Copy;
+}
+
+FCompileConstantResolver FCompileConstantResolver::WithUsage(ENiagaraScriptUsage ScriptUsage) const
 {
-	ENiagaraFunctionDebugState OldState = DebugState;
-	DebugState = InDebugState; 
-	return OldState;
+	FCompileConstantResolver Copy = *this;
+	Copy.Usage = ScriptUsage;
+	return Copy;
 }
 
 int32 FNiagaraParameterMapHistoryWithMetaDataBuilder::AddVariableToHistory(FNiagaraParameterMapHistory& History, const FNiagaraVariable& InVar, const FNiagaraVariable& InAliasedVar, const UEdGraphPin* InPin)
diff --git a/Engine/Plugins/FX/Niagara/Source/NiagaraEditor/Public/NiagaraParameterMapHistory.h b/Engine/Plugins/FX/Niagara/Source/NiagaraEditor/Public/NiagaraParameterMapHistory.h
--- a/Engine/Plugins/FX/Niagara/Source/NiagaraEditor/Public/NiagaraParameterMapHistory.h
+++ b/Engine/Plugins/FX/Niagara/Source/NiagaraEditor/Public/NiagaraParameterMapHistory.h
@@ -24,7 +24,8 @@ public:
 	bool ResolveConstant(FNiagaraVariable& OutConstant) const;
 
 	ENiagaraFunctionDebugState GetDebugState() const;
-	ENiagaraFunctionDebugState SetDebugState(ENiagaraFunctionDebugState InDebugState);
+	FCompileConstantResolver WithDebugState(ENiagaraFunctionDebugState InDebugState) const;
+	FCompileConstantResolver WithUsage(ENiagaraScriptUsage ScriptUsage) const;
 private:
 	const UNiagaraEmitter* Emitter;
 	const UNiagaraSystem* System;
@@ -371,6 +372,7 @@ public:
 	*/
 	void EnableScriptWhitelist(bool bInEnable, ENiagaraScriptUsage InScriptType);
 
+	bool HasCurrentUsageContext() const;
 	ENiagaraScriptUsage GetCurrentUsageContext()const;
 	ENiagaraScriptUsage GetBaseUsageContext()const;
 	bool ContextContains(ENiagaraScriptUsage InUsage) const;
