Added unique names to the parameter map get input pins as that is required by some internal APIs.
Fixes the issue of all pins with the same name on a node being dragged at once if one of the connections being dragged.

diff --git a/Engine/Plugins/FX/Niagara/Source/NiagaraEditor/Private/NiagaraNodeParameterMapGet.cpp b/Engine/Plugins/FX/Niagara/Source/NiagaraEditor/Private/NiagaraNodeParameterMapGet.cpp
--- a/Engine/Plugins/FX/Niagara/Source/NiagaraEditor/Private/NiagaraNodeParameterMapGet.cpp
+++ b/Engine/Plugins/FX/Niagara/Source/NiagaraEditor/Private/NiagaraNodeParameterMapGet.cpp
@@ -92,8 +92,6 @@ UEdGraphPin* UNiagaraNodeParameterMapGet::CreateDefaultPin(UEdGraphPin* OutputPi
 		return nullptr;
 	}
 
-	UEdGraphPin* DefaultPin = CreatePin(EEdGraphPinDirection::EGPD_Input, OutputPin->PinType, TEXT(""));
-
 	// make sure the new pin name is legal
 	FName OutName;
 	if (FNiagaraEditorUtilities::DecomposeVariableNamespace(OutputPin->GetFName(), OutName).Num() == 0)
@@ -101,6 +99,9 @@ UEdGraphPin* UNiagaraNodeParameterMapGet::CreateDefaultPin(UEdGraphPin* OutputPi
 		OutputPin->PinName = FName(FNiagaraConstants::LocalNamespace.ToString() + "." + OutputPin->GetName());
 	}
 
+	FName PinName(OutputPin->GetName() + TEXT("_Input"));
+	UEdGraphPin* DefaultPin = CreatePin(EEdGraphPinDirection::EGPD_Input, OutputPin->PinType, PinName);
+	
 	// we make the pin read only because the default value is set in the parameter panel unless the default mode is set to "custom" by the user
 	DefaultPin->bNotConnectable = true;
 	DefaultPin->bDefaultValueIsReadOnly = true;
@@ -336,6 +337,8 @@ void UNiagaraNodeParameterMapGet::SynchronizeDefaultInputPin(UEdGraphPin* Defaul
 		DefaultPin->PinToolTip = FText::Format(LOCTEXT("DefaultValueTooltip_UnlessOverridden", "Default value for {0} if no other module has set it previously in the stack."), FText::FromName(OutputPin->PinName)).ToString();
 	}
 
+	DefaultPin->PinName = FName(OutputPin->GetName() + TEXT("_Input"));
+	
 	// Sync pin visibility with the configured default mode
 	if (ScriptVar) {
 		if (ScriptVar->DefaultMode == ENiagaraDefaultMode::Value)
