Fixed build error

diff --git a/Engine/Plugins/FX/Niagara/Source/NiagaraEditor/Private/NiagaraCompiler.cpp b/Engine/Plugins/FX/Niagara/Source/NiagaraEditor/Private/NiagaraCompiler.cpp
--- a/Engine/Plugins/FX/Niagara/Source/NiagaraEditor/Private/NiagaraCompiler.cpp
+++ b/Engine/Plugins/FX/Niagara/Source/NiagaraEditor/Private/NiagaraCompiler.cpp
@@ -420,7 +420,7 @@ void FNiagaraCompileRequestData::DeepCopyGraphs(UNiagaraScriptSource* ScriptSour
 	NodeGraphDeepCopy = Source->NodeGraph;
 	if (bNeedsCompilation)
 	{
-		FEdGraphUtilities::MergeChildrenGraphsIn(NodeGraphDeepCopy.Get(), NodeGraphDeepCopy.Get(), /*bRequireSchemaMatch=*/ true);
+		FEdGraphUtilities::MergeChildrenGraphsIn(NodeGraphDeepCopy, NodeGraphDeepCopy, /*bRequireSchemaMatch=*/ true);
 	}
 	TMap<UNiagaraNodeFunctionCall*, ENiagaraScriptUsage> FunctionsWithUsage;
 	for (UEdGraphNode* Node : ScriptSource->NodeGraph->Nodes)
@@ -432,7 +432,7 @@ void FNiagaraCompileRequestData::DeepCopyGraphs(UNiagaraScriptSource* ScriptSour
 		}
 	}
 	FCompileConstantResolver ConstantResolver(Emitter, ENiagaraScriptUsage::EmitterSpawnScript);
-	VisitReferencedGraphs(ScriptSource->NodeGraph, NodeGraphDeepCopy.Get(), ENiagaraScriptUsage::EmitterSpawnScript, ConstantResolver, bNeedsCompilation, FunctionsWithUsage);
+	VisitReferencedGraphs(ScriptSource->NodeGraph, NodeGraphDeepCopy, ENiagaraScriptUsage::EmitterSpawnScript, ConstantResolver, bNeedsCompilation, FunctionsWithUsage);
 }
 
 
