When duplicating a parameter in the module editor, also copy its metadata

diff --git a/Engine/Plugins/FX/Niagara/Source/NiagaraEditor/Private/ViewModels/NiagaraParameterPanelViewModel.cpp b/Engine/Plugins/FX/Niagara/Source/NiagaraEditor/Private/ViewModels/NiagaraParameterPanelViewModel.cpp
--- a/Engine/Plugins/FX/Niagara/Source/NiagaraEditor/Private/ViewModels/NiagaraParameterPanelViewModel.cpp
+++ b/Engine/Plugins/FX/Niagara/Source/NiagaraEditor/Private/ViewModels/NiagaraParameterPanelViewModel.cpp
@@ -1929,6 +1929,38 @@ void FNiagaraScriptToolkitParameterPanelViewModel::RenameParameter(const FNiagar
 	}
 }
 
+void FNiagaraScriptToolkitParameterPanelViewModel::DuplicateParameter(const FNiagaraParameterPanelItem ItemToDuplicate) const
+{
+	FScopedTransaction Transaction(LOCTEXT("DuplicateParameterTransaction", "Duplicate parameter"));
+	TGuardValue<bool> AddParameterRefreshGuard(bIsAddingParameter, true);
+	bool bSuccess = false;
+
+	TSet<FName> Names;
+	for (const UNiagaraGraph* Graph : GetEditableGraphsConst())
+	{
+		for (auto It = Graph->GetParameterReferenceMap().CreateConstIterator(); It; ++It)
+		{
+			Names.Add(It.Key().GetName());
+		}
+	}
+	const FName NewUniqueName = FNiagaraUtilities::GetUniqueName(ItemToDuplicate.GetVariable().GetName(), Names);
+	FNiagaraVariable NewVariable(ItemToDuplicate.GetVariable().GetType(), NewUniqueName);
+	FNiagaraVariableMetaData ParameterMetadata = ItemToDuplicate.ScriptVariable ? ItemToDuplicate.ScriptVariable->Metadata : FNiagaraVariableMetaData();
+
+	for (UNiagaraGraph* Graph : GetEditableGraphs())
+	{
+		Graph->Modify();
+		Graph->AddParameter(NewVariable, ParameterMetadata, false, false);
+		bSuccess = true;
+	}
+
+	if (bSuccess)
+	{
+		Refresh();
+		SelectParameterItemByName(NewUniqueName, true);
+	}
+}
+
 void FNiagaraScriptToolkitParameterPanelViewModel::SetParameterIsSubscribedToLibrary(const FNiagaraParameterPanelItem ItemToModify, const bool bSubscribed) const
 {
 	if (ensureMsgf(ItemToModify.bExternallyReferenced == false, TEXT("Cannot modify an externally referenced parameter.")) == false)
diff --git a/Engine/Plugins/FX/Niagara/Source/NiagaraEditor/Public/ViewModels/NiagaraParameterPanelViewModel.h b/Engine/Plugins/FX/Niagara/Source/NiagaraEditor/Public/ViewModels/NiagaraParameterPanelViewModel.h
--- a/Engine/Plugins/FX/Niagara/Source/NiagaraEditor/Public/ViewModels/NiagaraParameterPanelViewModel.h
+++ b/Engine/Plugins/FX/Niagara/Source/NiagaraEditor/Public/ViewModels/NiagaraParameterPanelViewModel.h
@@ -358,6 +358,8 @@ public:
 
 	virtual void RenameParameter(const FNiagaraParameterPanelItem& ItemToRename, const FName NewName) const override;
 
+	virtual void DuplicateParameter(const FNiagaraParameterPanelItem ItemToDuplicate) const override;
+
 	virtual void SetParameterIsSubscribedToLibrary(const FNiagaraParameterPanelItem ItemToModify, const bool bSubscribed) const override;
 
 	virtual void OnParameterItemSelected(const FNiagaraParameterPanelItem& SelectedItem, ESelectInfo::Type SelectInfo) const override;
