Manually postload editor data and parent emitters

diff --git a/Engine/Plugins/FX/Niagara/Source/Niagara/Private/NiagaraEmitter.cpp b/Engine/Plugins/FX/Niagara/Source/Niagara/Private/NiagaraEmitter.cpp
--- a/Engine/Plugins/FX/Niagara/Source/Niagara/Private/NiagaraEmitter.cpp
+++ b/Engine/Plugins/FX/Niagara/Source/Niagara/Private/NiagaraEmitter.cpp
@@ -485,6 +485,11 @@ void UNiagaraEmitter::PostLoad()
 	}
 	
 #if WITH_EDITORONLY_DATA
+	if (EditorData != nullptr)
+	{
+		EditorData->ConditionalPostLoad();
+	}
+	
 	if (!GPUComputeScript)
 	{
 		GPUComputeScript = NewObject<UNiagaraScript>(this, "GPUComputeScript", EObjectFlags::RF_Transactional);
@@ -571,9 +576,19 @@ void UNiagaraEmitter::PostLoad()
 
 	FGraphEventArray ParentPrerequisiteTasks;
 #if WITH_EDITORONLY_DATA
-	if (Parent != nullptr && Parent->UpdateTaskRef.IsValid())
+	if (Parent != nullptr)
+	{
+		if (Parent->UpdateTaskRef.IsValid())
+		{
+			ParentPrerequisiteTasks.Add(Parent->UpdateTaskRef);
+		}
+	}
+	if (ParentAtLastMerge != nullptr)
 	{
-		ParentPrerequisiteTasks.Add(Parent->UpdateTaskRef);
+		if (ParentAtLastMerge->UpdateTaskRef.IsValid())
+		{
+			ParentPrerequisiteTasks.Add(ParentAtLastMerge->UpdateTaskRef);
+		}
 	}
 	// this can only ever be true for old assets that haven't been loaded yet, so this won't overwrite subsequent changes to the template specification
 	if(bIsTemplateAsset_DEPRECATED)
