Speculative editor data build fix

diff --git a/Engine/Plugins/FX/Niagara/Source/NiagaraEditor/Private/NiagaraEmitterEditorData.cpp b/Engine/Plugins/FX/Niagara/Source/NiagaraEditor/Private/NiagaraEmitterEditorData.cpp
--- a/Engine/Plugins/FX/Niagara/Source/NiagaraEditor/Private/NiagaraEmitterEditorData.cpp
+++ b/Engine/Plugins/FX/Niagara/Source/NiagaraEditor/Private/NiagaraEmitterEditorData.cpp
@@ -24,6 +24,7 @@ void UNiagaraEmitterEditorData::PostLoad()
 		StackEditorData = NewObject<UNiagaraStackEditorData>(this, TEXT("StackEditorData"), RF_Transactional);
 		StackEditorData->OnPersistentDataChanged().AddUObject(this, &UNiagaraEmitterEditorData::StackEditorDataChanged);
 	}
+	StackEditorData->ConditionalPostLoad();
 }
 
 UNiagaraStackEditorData& UNiagaraEmitterEditorData::GetStackEditorData() const
