Build fix for input node not being fully postloaded

diff --git a/Engine/Plugins/FX/Niagara/Source/NiagaraEditor/Private/NiagaraNodeInput.cpp b/Engine/Plugins/FX/Niagara/Source/NiagaraEditor/Private/NiagaraNodeInput.cpp
--- a/Engine/Plugins/FX/Niagara/Source/NiagaraEditor/Private/NiagaraNodeInput.cpp
+++ b/Engine/Plugins/FX/Niagara/Source/NiagaraEditor/Private/NiagaraNodeInput.cpp
@@ -69,6 +69,7 @@ void UNiagaraNodeInput::PostLoad()
 	Super::PostLoad();
 	if (DataInterface != nullptr)
 	{
+		DataInterface->ConditionalPostLoad();
 		DataInterface->OnChanged().AddUObject(this, &UNiagaraNodeInput::DataInterfaceChanged);
 	}
 
