Integrated 14786457, 14787544 : AFR now works with Niagara

diff --git a/Engine/Plugins/FX/Niagara/Source/Niagara/Classes/NiagaraEmitterInstanceBatcher.h b/Engine/Plugins/FX/Niagara/Source/Niagara/Classes/NiagaraEmitterInstanceBatcher.h
--- a/Engine/Plugins/FX/Niagara/Source/Niagara/Classes/NiagaraEmitterInstanceBatcher.h
+++ b/Engine/Plugins/FX/Niagara/Source/Niagara/Classes/NiagaraEmitterInstanceBatcher.h
@@ -288,4 +288,15 @@ private:
 	};
 	TArray<FDebugReadbackInfo> GpuDebugReadbackInfos;
 	TUniquePtr<FNiagaraGpuReadbackManager> GpuReadbackManagerPtr;
+
+#if WITH_MGPU
+	static const FName TemporalEffectName;
+	TArray<FRHIVertexBuffer*> TemporalEffectBuffers;
+	ENiagaraGpuComputeTickStage::Type StageToWaitForTemporalEffect = ENiagaraGpuComputeTickStage::First;
+	ENiagaraGpuComputeTickStage::Type StageToBroadcastTemporalEffect = ENiagaraGpuComputeTickStage::First;
+
+	void AddTemporalEffectBuffers(FNiagaraDataBuffer* FinalData);
+	void BroadcastTemporalEffect(FRHICommandList& RHICmdList);
+#endif // WITH_MGPU
+
 };
diff --git a/Engine/Plugins/FX/Niagara/Source/Niagara/Classes/NiagaraGPUInstanceCountManager.h b/Engine/Plugins/FX/Niagara/Source/Niagara/Classes/NiagaraGPUInstanceCountManager.h
--- a/Engine/Plugins/FX/Niagara/Source/Niagara/Classes/NiagaraGPUInstanceCountManager.h
+++ b/Engine/Plugins/FX/Niagara/Source/Niagara/Classes/NiagaraGPUInstanceCountManager.h
@@ -94,6 +94,8 @@ public:
 
 	static const ERHIAccess kCountBufferDefaultState;
 
+	bool HasEntriesPendingFree() const { return InstanceCountClearTasks.Num() > 0; }
+
 protected:
 	struct FIndirectArgsPoolEntry
 	{
diff --git a/Engine/Plugins/FX/Niagara/Source/Niagara/Private/NiagaraEmitterInstanceBatcher.cpp b/Engine/Plugins/FX/Niagara/Source/Niagara/Private/NiagaraEmitterInstanceBatcher.cpp
--- a/Engine/Plugins/FX/Niagara/Source/Niagara/Private/NiagaraEmitterInstanceBatcher.cpp
+++ b/Engine/Plugins/FX/Niagara/Source/Niagara/Private/NiagaraEmitterInstanceBatcher.cpp
@@ -51,6 +51,10 @@ DECLARE_GPU_STAT_NAMED(NiagaraGPUSorting, TEXT("Niagara GPU sorting"));
 
 uint32 FNiagaraComputeExecutionContext::TickCounter = 0;
 
+#if WITH_MGPU
+const FName NiagaraEmitterInstanceBatcher::TemporalEffectName("NiagaraEmitterInstanceBatcher");
+#endif // WITH_MGPU
+
 int32 GNiagaraGpuMaxQueuedRenderFrames = 10;
 static FAutoConsoleVariableRef CVarNiagaraGpuMaxQueuedRenderFrames(
 	TEXT("fx.Niagara.Batcher.MaxQueuedFramesWithoutRender"),
@@ -206,6 +210,14 @@ NiagaraEmitterInstanceBatcher::NiagaraEmitterInstanceBatcher(ERHIFeatureLevel::T
 				[this](FRHICommandListImmediate& RHICmdList)
 				{
 					GPUInstanceCounterManager.UpdateDrawIndirectBuffers(*this, RHICmdList, FeatureLevel);
+#if WITH_MGPU
+					// For PreInitViews we actually need to broadcast here and not in ExecuteAll since
+					// this is the last place the instance count buffer is written to.
+					if (StageToBroadcastTemporalEffect == ENiagaraGpuComputeTickStage::PreInitViews)
+					{
+						BroadcastTemporalEffect(RHICmdList);
+					}
+#endif // WITH_MGPU
 				}
 			);
 		}
@@ -1120,6 +1132,9 @@ void NiagaraEmitterInstanceBatcher::PrepareTicksForProxy(FRHICommandListImmediat
 			FinalBuffer->GPUInstanceCountBufferOffset = ComputeContext->CountOffset_RT;
 			FinalBuffer->SetNumInstances(ComputeContext->CurrentNumInstances_RT);
 			ComputeContext->SetDataToRender(FinalBuffer);
+#if WITH_MGPU
+			AddTemporalEffectBuffers(FinalBuffer);
+#endif // WITH_MGPU
 		}
 		// When low latency translucency is enabled we can setup the final buffer / final count here.
 		// This will allow our mesh processor commands to pickup the data for the same frame.
@@ -1132,6 +1147,29 @@ void NiagaraEmitterInstanceBatcher::PrepareTicksForProxy(FRHICommandListImmediat
 			ComputeContext->SetTranslucentDataToRender(FinalBuffer);
 		}
 	}
+
+#if WITH_MGPU
+	if (GNumAlternateFrameRenderingGroups > 1)
+	{
+		StageToBroadcastTemporalEffect = ENiagaraGpuComputeTickStage::Last;
+		while (StageToBroadcastTemporalEffect > ENiagaraGpuComputeTickStage::First && !DispatchListPerStage[static_cast<int32>(StageToBroadcastTemporalEffect)].HasWork())
+		{
+			StageToBroadcastTemporalEffect = static_cast<ENiagaraGpuComputeTickStage::Type>(static_cast<int32>(StageToBroadcastTemporalEffect) - 1);
+		}
+
+		StageToWaitForTemporalEffect = ENiagaraGpuComputeTickStage::First;
+		// If we're going to write to the instance count buffer after PreInitViews then
+		// that needs to be the wait stage, regardless of whether or not we're ticking
+		// anything in that stage.
+		if (!GPUInstanceCounterManager.HasEntriesPendingFree())
+		{
+			while (StageToWaitForTemporalEffect < StageToBroadcastTemporalEffect && !DispatchListPerStage[static_cast<int32>(StageToWaitForTemporalEffect)].HasWork())
+			{
+				StageToWaitForTemporalEffect = static_cast<ENiagaraGpuComputeTickStage::Type>(static_cast<int32>(StageToWaitForTemporalEffect) + 1);
+			}
+		}
+	}
+#endif // WITH_MGPU
 }
 
 void NiagaraEmitterInstanceBatcher::PrepareAllTicks(FRHICommandListImmediate& RHICmdList)
@@ -1147,6 +1185,13 @@ void NiagaraEmitterInstanceBatcher::PrepareAllTicks(FRHICommandListImmediate& RH
 
 void NiagaraEmitterInstanceBatcher::ExecuteTicks(FRHICommandList& RHICmdList, FRHIUniformBuffer* ViewUniformBuffer, ENiagaraGpuComputeTickStage::Type TickStage)
 {
+#if WITH_MGPU
+	if (GNumAlternateFrameRenderingGroups > 1 && TickStage == StageToWaitForTemporalEffect)
+	{
+		RHICmdList.WaitForTemporalEffect(TemporalEffectName);
+	}
+#endif // WITH_MGPU
+
 	// Anything to execute for this stage
 	FNiagaraGpuDispatchList& DispatchList = DispatchListPerStage[TickStage];
 	if ( !DispatchList.HasWork() )
@@ -1301,6 +1346,9 @@ void NiagaraEmitterInstanceBatcher::ExecuteTicks(FRHICommandList& RHICmdList, FR
 
 				ComputeContext->SetTranslucentDataToRender(CurrentData);
 				ComputeContext->SetDataToRender(CurrentData);
+#if WITH_MGPU
+				AddTemporalEffectBuffers(CurrentData);
+#endif // WITH_MGPU
 			}
 		}
 
@@ -1308,6 +1356,15 @@ void NiagaraEmitterInstanceBatcher::ExecuteTicks(FRHICommandList& RHICmdList, FR
 		RHICmdList.Transition(TransitionsAfter);
 		TransitionsAfter.Reset();
 
+#if WITH_MGPU
+		// For PreInitViews we actually need to broadcast after UpdateDrawIndirectBuffer
+		// since this is the last place the instance count buffer is written to.
+		if (TickStage == StageToBroadcastTemporalEffect && StageToBroadcastTemporalEffect != ENiagaraGpuComputeTickStage::PreInitViews)
+		{
+			BroadcastTemporalEffect(RHICmdList);
+		}
+#endif // WITH_MGPU
+
 		// Update free IDs
 		if (DispatchGroup.FreeIDUpdates.Num() > 0)
 		{
@@ -2124,3 +2181,35 @@ void NiagaraEmitterInstanceBatcher::DrawSceneDebug_RenderThread(class FRDGBuilde
 	}
 #endif
 }
+
+#if WITH_MGPU
+void NiagaraEmitterInstanceBatcher::AddTemporalEffectBuffers(FNiagaraDataBuffer* FinalData)
+{
+	if (GNumAlternateFrameRenderingGroups > 1)
+	{
+		for (FRHIVertexBuffer* Buffer : { FinalData->GetGPUBufferFloat().Buffer, FinalData->GetGPUBufferInt().Buffer, FinalData->GetGPUBufferHalf().Buffer })
+		{
+			if (Buffer)
+			{
+				TemporalEffectBuffers.Add(Buffer);
+			}
+		}
+	}
+}
+
+void NiagaraEmitterInstanceBatcher::BroadcastTemporalEffect(FRHICommandList& RHICmdList)
+{
+	if (GNumAlternateFrameRenderingGroups > 1)
+	{
+		if (GPUInstanceCounterManager.GetInstanceCountBuffer().Buffer)
+		{
+			TemporalEffectBuffers.Add(GPUInstanceCounterManager.GetInstanceCountBuffer().Buffer);
+		}
+		if (TemporalEffectBuffers.Num())
+		{
+			RHICmdList.BroadcastTemporalEffect(TemporalEffectName, TemporalEffectBuffers);
+			TemporalEffectBuffers.Reset();
+		}
+	}
+}
+#endif // WITH_MGPU
\ No newline at end of file
diff --git a/Engine/Plugins/FX/Niagara/Source/Niagara/Public/NiagaraSystemGpuComputeProxy.h b/Engine/Plugins/FX/Niagara/Source/Niagara/Public/NiagaraSystemGpuComputeProxy.h
--- a/Engine/Plugins/FX/Niagara/Source/Niagara/Public/NiagaraSystemGpuComputeProxy.h
+++ b/Engine/Plugins/FX/Niagara/Source/Niagara/Public/NiagaraSystemGpuComputeProxy.h
@@ -17,7 +17,9 @@ namespace ENiagaraGpuComputeTickStage
 		PreInitViews,
 		PostInitViews,
 		PostOpaqueRender,
-		Max
+		Max,
+		First = PreInitViews,
+		Last = PostOpaqueRender,
 	};
 };
 
