- Fix issue around BulkData being passed and used on RHI thread on OpenGL

#rnx

diff --git a/Engine/Plugins/FX/Niagara/Source/Niagara/Private/NiagaraDataInterfaceLandscape.cpp b/Engine/Plugins/FX/Niagara/Source/Niagara/Private/NiagaraDataInterfaceLandscape.cpp
--- a/Engine/Plugins/FX/Niagara/Source/Niagara/Private/NiagaraDataInterfaceLandscape.cpp
+++ b/Engine/Plugins/FX/Niagara/Source/Niagara/Private/NiagaraDataInterfaceLandscape.cpp
@@ -92,7 +92,8 @@ public:
 
 	struct FTextureBulkData : public FResourceBulkDataInterface
 	{
-		FTextureBulkData(const void* InData, uint32 InDataSize) : Data(InData), DataSize(InDataSize) {}
+		void Init(const void* InData, uint32 InDataSize) { Data=InData; DataSize=InDataSize; }
+		void Clear() { Data = nullptr; DataSize = 0; }
 		virtual const void* GetResourceBulkData() const override { return Data; }
 		virtual uint32 GetResourceBulkDataSize() const override { return DataSize; }
 		virtual void Discard() override { }
@@ -114,8 +115,8 @@ public:
 
 	void Initialize();
 
-	FTextureBulkData GetHeightBulkData() { return FTextureBulkData(HeightValues.GetData(), HeightValues.Num() * HeightValues.GetTypeSize()); }
-	FTextureBulkData GetPhysMatBulkData() { return FTextureBulkData(PhysMatValues.GetData(), PhysMatValues.Num() * PhysMatValues.GetTypeSize()); }
+	FTextureBulkData* GetHeightBulkData() { return HeightBulkData.DataSize > 0 ? &HeightBulkData : nullptr; }
+	FTextureBulkData* GetPhysMatBulkData() { return PhysMatBulkData.DataSize > 0 ? &PhysMatBulkData : nullptr; }
 
 	FLandscapeTextureResource LandscapeTextures;
 	FMatrix ActorToWorldTransform = FMatrix::Identity;
@@ -136,6 +137,9 @@ private:
 	TArray<float> HeightValues;
 	TArray<uint8> PhysMatValues;
 
+	FTextureBulkData HeightBulkData;
+	FTextureBulkData PhysMatBulkData;
+
 	FThreadSafeBool GpuDataReleasedByRt = false;
 	FThreadSafeBool CpuDataReleasedByRt = false;
 	bool GpuDataQueuedForRelease = false;
@@ -225,16 +229,14 @@ FLandscapeTextureResource::FLandscapeTextureResource(FNDI_Landscape_SharedResour
 
 void FLandscapeTextureResource::InitRHI()
 {
-	auto HeightBulkData = Owner.GetHeightBulkData();
-	if (HeightBulkData.DataSize)
+	if (FNDI_Landscape_SharedResource::FTextureBulkData* HeightBulkData = Owner.GetHeightBulkData())
 	{
-		HeightTexture.Initialize(TEXT("FLandscapeTextureResource_HeightTexture"), sizeof(float), Owner.CellCount.X, Owner.CellCount.Y, EPixelFormat::PF_R32_FLOAT, FTextureReadBuffer2D::DefaultTextureInitFlag, &HeightBulkData);
+		HeightTexture.Initialize(TEXT("FLandscapeTextureResource_HeightTexture"), sizeof(float), Owner.CellCount.X, Owner.CellCount.Y, EPixelFormat::PF_R32_FLOAT, FTextureReadBuffer2D::DefaultTextureInitFlag, HeightBulkData);
 	}
 
-	auto PhysMatBulkData = Owner.GetPhysMatBulkData();
-	if (PhysMatBulkData.DataSize)
+	if (FNDI_Landscape_SharedResource::FTextureBulkData* PhysMatBulkData = Owner.GetPhysMatBulkData())
 	{
-		PhysMatTexture.Initialize(TEXT("FLandscapeTextureResource_PhysMatTexture"), sizeof(uint8), Owner.CellCount.X, Owner.CellCount.Y, EPixelFormat::PF_R8_UINT, FTextureReadBuffer2D::DefaultTextureInitFlag, &PhysMatBulkData);
+		PhysMatTexture.Initialize(TEXT("FLandscapeTextureResource_PhysMatTexture"), sizeof(uint8), Owner.CellCount.X, Owner.CellCount.Y, EPixelFormat::PF_R8_UINT, FTextureReadBuffer2D::DefaultTextureInitFlag, PhysMatBulkData);
 
 		PhysMatDimensions = FIntPoint(Owner.CellCount.X, Owner.CellCount.Y);
 	}
@@ -383,6 +385,11 @@ void FNDI_Landscape_SharedResource::Initialize()
 				const float DefaultHeight = 0.0f;
 				HeightValues.Reset(SampleCount);
 				HeightValues.Init(DefaultHeight, SampleCount);
+				HeightBulkData.Init(HeightValues.GetData(), HeightValues.Num() * HeightValues.GetTypeSize());
+			}
+			else
+			{
+				HeightBulkData.Clear();
 			}
 
 			if (ResourceKey.IncludesCachedPhysMat)
@@ -390,6 +397,11 @@ void FNDI_Landscape_SharedResource::Initialize()
 				const uint8 DefaultPhysMat = INDEX_NONE;
 				PhysMatValues.Reset(SampleCount);
 				PhysMatValues.Init(DefaultPhysMat, SampleCount);
+				PhysMatBulkData.Init(PhysMatValues.GetData(), PhysMatValues.Num() * PhysMatValues.GetTypeSize());
+			}
+			else
+			{
+				PhysMatBulkData.Clear();
 			}
 
 			const FIntPoint RegionVertexBase = ResourceKey.MinCaptureRegion * ComponentQuadCount;
@@ -496,7 +508,14 @@ void FNDI_Landscape_SharedResource::ReleaseGpu()
 	ENQUEUE_RENDER_COMMAND(BeginDestroyCommand)(
 		[Released](FRHICommandListImmediate& RHICmdList)
 	{
-		*Released = true;
+		// On some RHIs textures will push data on the RHI thread
+		// Therefore we are not 'released' until the RHI thread has processed all commands
+		RHICmdList.EnqueueLambda(
+			[Released](FRHICommandListImmediate& RHICmdList)
+			{
+				*Released = true;
+			}
+		);
 	});
 }
 
