fixed deprecation warning for FRHIVertexBuffer

diff --git a/Engine/Plugins/FX/Niagara/Source/Niagara/Private/NiagaraEmitterInstanceBatcher.cpp b/Engine/Plugins/FX/Niagara/Source/Niagara/Private/NiagaraEmitterInstanceBatcher.cpp
--- a/Engine/Plugins/FX/Niagara/Source/Niagara/Private/NiagaraEmitterInstanceBatcher.cpp
+++ b/Engine/Plugins/FX/Niagara/Source/Niagara/Private/NiagaraEmitterInstanceBatcher.cpp
@@ -2187,7 +2187,7 @@ void NiagaraEmitterInstanceBatcher::AddTemporalEffectBuffers(FNiagaraDataBuffer*
 {
 	if (GNumAlternateFrameRenderingGroups > 1)
 	{
-		for (FRHIVertexBuffer* Buffer : { FinalData->GetGPUBufferFloat().Buffer, FinalData->GetGPUBufferInt().Buffer, FinalData->GetGPUBufferHalf().Buffer })
+		for (FRHIBuffer* Buffer : { FinalData->GetGPUBufferFloat().Buffer, FinalData->GetGPUBufferInt().Buffer, FinalData->GetGPUBufferHalf().Buffer })
 		{
 			if (Buffer)
 			{
