Fixed the "Link Input" parent category not being applied in the dynamic inputs menu

diff --git a/Engine/Plugins/FX/Niagara/Source/NiagaraEditorWidgets/Private/Stack/SNiagaraStackFunctionInputValue.cpp b/Engine/Plugins/FX/Niagara/Source/NiagaraEditorWidgets/Private/Stack/SNiagaraStackFunctionInputValue.cpp
--- a/Engine/Plugins/FX/Niagara/Source/NiagaraEditorWidgets/Private/Stack/SNiagaraStackFunctionInputValue.cpp
+++ b/Engine/Plugins/FX/Niagara/Source/NiagaraEditorWidgets/Private/Stack/SNiagaraStackFunctionInputValue.cpp
@@ -967,7 +967,7 @@ TArray<TSharedPtr<FNiagaraMenuAction_Generic>> SNiagaraStackFunctionInputValue::
 			
 			TSharedPtr<FNiagaraMenuAction_Generic> LinkAction(new FNiagaraMenuAction_Generic(
 				FNiagaraMenuAction_Generic::FOnExecuteAction::CreateSP(this, &SNiagaraStackFunctionInputValue::ParameterHandleSelected, AvailableHandle),
-				DisplayName, ENiagaraMenuSections::General, {Category.ToString()}, Tooltip, FText()));
+				DisplayName, ENiagaraMenuSections::General, {RootCategoryName, Category.ToString()}, Tooltip, FText()));
 			
 			LinkAction->SetParameterVariable(FNiagaraVariable(FunctionInput->GetInputType(), AvailableHandle.GetParameterHandleString()));
 			LinkAction->SourceData = NiagaraSourceData;
