Fix for "Apply" not being enabled when changing a string name in a node on Attribute Reader Function. Also fixed potential crash with too long of text entry in the same field.

diff --git a/Engine/Plugins/FX/Niagara/Source/NiagaraEditor/Private/Widgets/SNiagaraGraphNodeFunctionCallWithSpecifiers.cpp b/Engine/Plugins/FX/Niagara/Source/NiagaraEditor/Private/Widgets/SNiagaraGraphNodeFunctionCallWithSpecifiers.cpp
--- a/Engine/Plugins/FX/Niagara/Source/NiagaraEditor/Private/Widgets/SNiagaraGraphNodeFunctionCallWithSpecifiers.cpp
+++ b/Engine/Plugins/FX/Niagara/Source/NiagaraEditor/Private/Widgets/SNiagaraGraphNodeFunctionCallWithSpecifiers.cpp
@@ -6,6 +6,7 @@
 #include "NiagaraNodeFunctionCall.h"
 #include "GraphEditorSettings.h"
 #include "Widgets/Input/SEditableTextBox.h"
+#include "NiagaraConstants.h"
 
 #define LOCTEXT_NAMESPACE "SNiagaraGraphNodeFunctionCallWithSpecifiers"
 
@@ -26,6 +27,7 @@ TSharedRef<SWidget> SNiagaraGraphNodeFunctionCallWithSpecifiers::CreateNodeConte
 			.VAlign(VAlign_Center)
 			[
 				SNew(SNiagaraFunctionSpecifier, Entry.Key, Entry.Value, *FunctionSpecifiers)
+					.OnValueNameChanged(this, &SNiagaraGraphNodeFunctionCallWithSpecifiers::OnValueNameChanged)
 			];
 	}
 
@@ -47,11 +49,19 @@ TSharedRef<SWidget> SNiagaraGraphNodeFunctionCallWithSpecifiers::CreateNodeConte
 	return RetWidget;
 }
 
+void SNiagaraGraphNodeFunctionCallWithSpecifiers::OnValueNameChanged()
+{
+	Cast<UNiagaraNodeFunctionCall>(GraphNode)->MarkNodeRequiresSynchronization(__FUNCTION__, true);
+}
+
 void SNiagaraFunctionSpecifier::Construct(const FArguments& InArgs, FName InAttributeName, FName InValueName, TMap<FName, FName>& InSpecifiers)
 {
+	OnValueNameChanged = InArgs._OnValueNameChanged;
+
 	AttributeName = InAttributeName;
 	ValueName = InValueName;
 	Specifiers = &InSpecifiers;
+
 	ChildSlot
 		.VAlign(VAlign_Fill)
 		.HAlign(HAlign_Fill)
@@ -72,17 +82,29 @@ void SNiagaraFunctionSpecifier::Construct(const FArguments& InArgs, FName InAttr
 			[
 				SNew(SEditableTextBox)
 				.Text(FText::FromName(ValueName))
+				.OnVerifyTextChanged(this, &SNiagaraFunctionSpecifier::VerifyNameTextChanged)
 				.OnTextCommitted(this, &SNiagaraFunctionSpecifier::OnValueNameCommitted)
 			]
 		];
 }
 
+bool SNiagaraFunctionSpecifier::VerifyNameTextChanged(const FText& NewText, FText& OutErrorMessage) const
+{
+	if (NewText.ToString().Len() >= FNiagaraConstants::MaxParameterLength)
+	{
+		OutErrorMessage = FText::FormatOrdered(LOCTEXT("NameToLongError", "Name cannot exceed {0} characters."), FNiagaraConstants::MaxParameterLength);
+		return false;
+	}
+	return true;
+}
+
 void SNiagaraFunctionSpecifier::OnValueNameCommitted(const FText& InText, ETextCommit::Type InCommitType)
 {
 	ValueName = FName(*InText.ToString());
 	if (Specifiers)
 	{
 		Specifiers->Add(AttributeName, ValueName);
+		OnValueNameChanged.ExecuteIfBound();
 	}
 }
 
diff --git a/Engine/Plugins/FX/Niagara/Source/NiagaraEditor/Private/Widgets/SNiagaraGraphNodeFunctionCallWithSpecifiers.h b/Engine/Plugins/FX/Niagara/Source/NiagaraEditor/Private/Widgets/SNiagaraGraphNodeFunctionCallWithSpecifiers.h
--- a/Engine/Plugins/FX/Niagara/Source/NiagaraEditor/Private/Widgets/SNiagaraGraphNodeFunctionCallWithSpecifiers.h
+++ b/Engine/Plugins/FX/Niagara/Source/NiagaraEditor/Private/Widgets/SNiagaraGraphNodeFunctionCallWithSpecifiers.h
@@ -13,6 +13,8 @@ public:
 
 	void Construct(const FArguments& InArgs, UEdGraphNode* InGraphNode);
 	virtual TSharedRef<SWidget> CreateNodeContentArea() override;
+
+	void OnValueNameChanged();
 	
 	TMap<FName, FName>* FunctionSpecifiers;
 };
@@ -21,11 +23,17 @@ class SNiagaraFunctionSpecifier : public SCompoundWidget
 {
 public:
 	SLATE_BEGIN_ARGS(SNiagaraFunctionSpecifier) {}
+	SLATE_EVENT(FSimpleDelegate, OnValueNameChanged)
 	SLATE_END_ARGS();
+
+
 	void Construct(const FArguments& InArgs, FName InAttributeName, FName InValueName, TMap<FName, FName>& InSpecifiers);
 
+	bool VerifyNameTextChanged(const FText& NewText, FText& OutErrorMessage) const;
 	void OnValueNameCommitted(const FText& InText, ETextCommit::Type InCommitType);
 
+	FSimpleDelegate OnValueNameChanged;
+
 	FName AttributeName;
 	FName ValueName;
 	TMap<FName, FName>* Specifiers;
