Fix landscape actor bounds evaluation for case with streaming actors (in which case we need to use the LandscapeInfo instead of the actor directly).

diff --git a/Engine/Plugins/FX/Niagara/Source/Niagara/Private/NiagaraDataInterfaceLandscape.cpp b/Engine/Plugins/FX/Niagara/Source/Niagara/Private/NiagaraDataInterfaceLandscape.cpp
--- a/Engine/Plugins/FX/Niagara/Source/Niagara/Private/NiagaraDataInterfaceLandscape.cpp
+++ b/Engine/Plugins/FX/Niagara/Source/Niagara/Private/NiagaraDataInterfaceLandscape.cpp
@@ -1032,11 +1032,15 @@ ALandscape* UNiagaraDataInterfaceLandscape::GetLandscape(const FNiagaraSystemIns
 	{
 		if (InLandscape->GetWorld() == SystemInstance.GetWorld())
 		{
-			const FBox& LandscapeBounds = InLandscape->GetComponentsBoundingBox();
-
-			if (LandscapeBounds.IntersectXY(WorldBounds))
+			if (const ULandscapeInfo* LandscapeInfo = InLandscape->GetLandscapeInfo())
 			{
-				return true;
+				for (const auto& ComponentIt : LandscapeInfo->XYtoCollisionComponentMap)
+				{
+					if (WorldBounds.IntersectXY(ComponentIt.Value->Bounds.GetBox()))
+					{
+						return true;
+					}
+				}
 			}
 		}
 
