Fixing situation where the VortexField debug visualization was not working with bindings to Engine.Emitter.SimulationPosition.

diff --git a/Engine/Plugins/FX/Niagara/Source/NiagaraCore/Private/NiagaraCustomVersion.cpp b/Engine/Plugins/FX/Niagara/Source/NiagaraCore/Private/NiagaraCustomVersion.cpp
--- a/Engine/Plugins/FX/Niagara/Source/NiagaraCore/Private/NiagaraCustomVersion.cpp
+++ b/Engine/Plugins/FX/Niagara/Source/NiagaraCore/Private/NiagaraCustomVersion.cpp
@@ -8,4 +8,4 @@ const FGuid FNiagaraCustomVersion::GUID(0xFCF57AFA, 0x50764283, 0xB9A9E658, 0xFF
 // Register the custom version with core
 FCustomVersionRegistration GRegisterNiagaraCustomVersion(FNiagaraCustomVersion::GUID, FNiagaraCustomVersion::LatestVersion, TEXT("NiagaraVer"));
 
-const FGuid FNiagaraCustomVersion::LatestScriptCompileVersion(0x7F8E9CA7, 0x61806B41, 0xA6C7E8A3, 0xCB2120F5);
+const FGuid FNiagaraCustomVersion::LatestScriptCompileVersion(0x37B399DD, 0x825B4DE2, 0xB302EED3, 0xD7413330);
diff --git a/Engine/Plugins/FX/Niagara/Source/NiagaraEditor/Private/NiagaraHlslTranslator.cpp b/Engine/Plugins/FX/Niagara/Source/NiagaraEditor/Private/NiagaraHlslTranslator.cpp
--- a/Engine/Plugins/FX/Niagara/Source/NiagaraEditor/Private/NiagaraHlslTranslator.cpp
+++ b/Engine/Plugins/FX/Niagara/Source/NiagaraEditor/Private/NiagaraHlslTranslator.cpp
@@ -5148,6 +5148,10 @@ bool FHlslNiagaraTranslator::HandleBoundConstantVariableToDataSetRead(FNiagaraVa
 			const FString SysParamEnginePositionSymbolNameString = GetSanitizedSymbolName(SYS_PARAM_ENGINE_POSITION.GetName().ToString(), true);
 			const FString ConstantStr = FString::Printf(TEXT("%s%s"), *SysParamEnginePositionSymbolNameString, *FString(TEXT(".xyz")));
 			Output = AddBodyChunk(GetUniqueSymbolName(TEXT("Constant")), ConstantStr, InVariable.GetType());
+			if (CodeChunks.IsValidIndex(Output))
+			{
+				CodeChunks[Output].Original = InVariable;
+			}
 			return true;
 		}
 		else
@@ -5155,7 +5159,11 @@ bool FHlslNiagaraTranslator::HandleBoundConstantVariableToDataSetRead(FNiagaraVa
 			InVariable.SetValue(FVector(EForceInit::ForceInitToZero));
 			float* ValuePtr = (float*)InVariable.GetData();
 			const FString ConstantStr = FString::Printf(TEXT("float3(%g,%g,%g)"), *ValuePtr, *(ValuePtr + 1), *(ValuePtr + 2));
-			Output = AddBodyChunk(GetUniqueSymbolName(TEXT("Constant")), ConstantStr, InVariable.GetType());
+			Output = AddBodyChunk(GetUniqueSymbolName(TEXT("Constant")), ConstantStr, InVariable.GetType()); 
+			if (CodeChunks.IsValidIndex(Output))
+			{
+				CodeChunks[Output].Original = InVariable;
+			}
 			return true;
 		}
 	}
@@ -8037,13 +8045,18 @@ void FHlslNiagaraTranslator::WriteCompilerTag(int32 InputCompileResult, const UE
 						{
 							InputCompileResult = CodeChunks[InputCompileResult].SourceChunks[0]; // Follow the linkage
 						}
-						else if (CodeChunks[InputCompileResult].Original.IsDataAllocated())
+						else if (CodeChunks[InputCompileResult].Original.IsDataAllocated()) // Handle constants
 						{
 							Variable.AllocateData();
 							CodeChunks[InputCompileResult].Original.CopyTo(Variable.GetData());
 							bSearch = false;
 						}
-						else
+						else if (CodeChunks[InputCompileResult].Original.IsValid()) // Handle default assignments
+						{
+							Value = CodeChunks[InputCompileResult].Original.GetName().ToString();
+							bSearch = false;
+						}
+						else // Handle setting to defaults as we didn't find a match.
 						{
 							TSharedPtr<INiagaraEditorTypeUtilities, ESPMode::ThreadSafe> TypeEditorUtilities = NiagaraEditorModule.GetTypeUtilities(TypeDef);
 							if (TypeEditorUtilities.IsValid() && TypeEditorUtilities->CanHandlePinDefaults() && CodeChunks[InputCompileResult].Definition.IsEmpty() == false)
