Niagara - Move update after load for emitters and system from task graph tasks to the asset loaded delegate.

diff --git a/Engine/Plugins/FX/Niagara/Source/Niagara/Classes/NiagaraEmitter.h b/Engine/Plugins/FX/Niagara/Source/Niagara/Classes/NiagaraEmitter.h
--- a/Engine/Plugins/FX/Niagara/Source/Niagara/Classes/NiagaraEmitter.h
+++ b/Engine/Plugins/FX/Niagara/Source/Niagara/Classes/NiagaraEmitter.h
@@ -394,9 +394,6 @@ public:
 	void CacheFromCompiledData(const FNiagaraDataSetCompiledData* CompiledData);
 	void CacheFromShaderCompiled();
 
-	// Reference to the async update task that is responsible to do any outstanding work like emitter merging on the game thread
-	FGraphEventRef UpdateTaskRef;
-
 	NIAGARA_API void UpdateEmitterAfterLoad();
 
 #if WITH_EDITORONLY_DATA
diff --git a/Engine/Plugins/FX/Niagara/Source/Niagara/Classes/NiagaraSystem.h b/Engine/Plugins/FX/Niagara/Source/Niagara/Classes/NiagaraSystem.h
--- a/Engine/Plugins/FX/Niagara/Source/Niagara/Classes/NiagaraSystem.h
+++ b/Engine/Plugins/FX/Niagara/Source/Niagara/Classes/NiagaraSystem.h
@@ -696,9 +696,6 @@ protected:
 	UNiagaraBakerSettings* BakerGeneratedSettings;
 #endif
 
-	// Reference to the async system update task that is responsible to do any outstanding work like emitter merging or script compilation 
-	FGraphEventRef UpdateTaskRef;
-	
 	UPROPERTY()
 	bool bHasSystemScriptDIsWithPerInstanceData;
 
diff --git a/Engine/Plugins/FX/Niagara/Source/Niagara/Private/NiagaraEmitter.cpp b/Engine/Plugins/FX/Niagara/Source/Niagara/Private/NiagaraEmitter.cpp
--- a/Engine/Plugins/FX/Niagara/Source/Niagara/Private/NiagaraEmitter.cpp
+++ b/Engine/Plugins/FX/Niagara/Source/Niagara/Private/NiagaraEmitter.cpp
@@ -66,33 +66,6 @@ static FAutoConsoleVariableRef CVarNiagaraDebugForcedMaxGPUBufferElements(
 	ECVF_Default
 );
 
-/**
-Game thread task to do any outstanding work on the loaded emitter
-*/
-struct FNiagaraEmitterUpdateTask
-{
-	FNiagaraEmitterUpdateTask(TWeakObjectPtr<UNiagaraEmitter> InEmitter) : WeakEmitter(InEmitter)
-	{
-	}
-
-	FORCEINLINE TStatId GetStatId() const { RETURN_QUICK_DECLARE_CYCLE_STAT(FNiagaraEmitterUpdateTask, STATGROUP_TaskGraphTasks); }
-	static FORCEINLINE ENamedThreads::Type GetDesiredThread() { return ENamedThreads::GameThread; }
-	static FORCEINLINE ESubsequentsMode::Type GetSubsequentsMode() { return ESubsequentsMode::TrackSubsequents; }
-
-	void DoTask(ENamedThreads::Type CurrentThread, const FGraphEventRef& MyCompletionGraphEvent)
-	{
-		UNiagaraEmitter* Emitter = WeakEmitter.Get();
-		if (Emitter == nullptr)
-		{
-			// probably got garbage collected in the meantime
-			return;
-		}
-		Emitter->UpdateEmitterAfterLoad();
-	}
-
-	TWeakObjectPtr<UNiagaraEmitter> WeakEmitter;
-};
-
 FNiagaraDetailsLevelScaleOverrides::FNiagaraDetailsLevelScaleOverrides()
 {
 	Low = 0.125f;
@@ -572,33 +545,14 @@ void UNiagaraEmitter::PostLoad()
 			UE_LOG(LogNiagara, Warning, TEXT("Disabling interpolated spawn because emitter flag and script type don't match. Did you adjust this value in the UI? Emitter may need recompile.. %s"), *GetFullName());
 		}
 	}
+
 	EnsureScriptsPostLoaded();
 
-	FGraphEventArray ParentPrerequisiteTasks;
-#if WITH_EDITORONLY_DATA
-	if (Parent != nullptr)
-	{
-		if (Parent->UpdateTaskRef.IsValid())
-		{
-			ParentPrerequisiteTasks.Add(Parent->UpdateTaskRef);
-		}
-	}
-	if (ParentAtLastMerge != nullptr)
-	{
-		if (ParentAtLastMerge->UpdateTaskRef.IsValid())
-		{
-			ParentPrerequisiteTasks.Add(ParentAtLastMerge->UpdateTaskRef);
-		}
-	}
-	// this can only ever be true for old assets that haven't been loaded yet, so this won't overwrite subsequent changes to the template specification
-	if(bIsTemplateAsset_DEPRECATED)
-	{
-		TemplateSpecification = ENiagaraScriptTemplateSpecification::Template;
-	}
+#if !WITH_EDITOR
+	// When running without the editor in a cooked build we run the update immediately in post load since
+	// there will be no merging or compiling which makes it safe to do so.
+	UpdateEmitterAfterLoad();
 #endif
-	
-	// we are not yet finished, but we do the rest of the work after postload in a separate task
-	UpdateTaskRef = TGraphTask<FNiagaraEmitterUpdateTask>::CreateTask(&ParentPrerequisiteTasks).ConstructAndDispatchWhenReady(this);
 }
 
 bool UNiagaraEmitter::IsEditorOnly() const
@@ -955,10 +909,7 @@ void UNiagaraEmitter::PreSave(const ITargetPlatform* TargetPlatform)
 	PRAGMA_DISABLE_DEPRECATION_WARNINGS;
 	Super::PreSave(TargetPlatform);
 	PRAGMA_ENABLE_DEPRECATION_WARNINGS;
-	if (UpdateTaskRef.IsValid() && UpdateTaskRef->IsComplete() == false)
-	{
-		UpdateEmitterAfterLoad();
-	}
+	UpdateEmitterAfterLoad();
 }
 
 
@@ -1274,7 +1225,6 @@ void UNiagaraEmitter::UpdateEmitterAfterLoad()
 		return;
 	}
 	bFullyLoaded = true;
-	UpdateTaskRef = nullptr;
 	
 #if WITH_EDITORONLY_DATA
 	check(IsInGameThread());
@@ -1312,8 +1262,6 @@ void UNiagaraEmitter::UpdateEmitterAfterLoad()
 		ParentAtLastMerge = nullptr;
 	}
 
-	// The task prerequisites should have updated the parent emitter before updating this emitter, but if this emitter
-	// has been forced to update we need to make sure the parent has been updated too.
 	if (Parent != nullptr)
 	{
 		Parent->UpdateEmitterAfterLoad();
diff --git a/Engine/Plugins/FX/Niagara/Source/Niagara/Private/NiagaraModule.cpp b/Engine/Plugins/FX/Niagara/Source/Niagara/Private/NiagaraModule.cpp
--- a/Engine/Plugins/FX/Niagara/Source/Niagara/Private/NiagaraModule.cpp
+++ b/Engine/Plugins/FX/Niagara/Source/Niagara/Private/NiagaraModule.cpp
@@ -27,6 +27,7 @@
 #include "NiagaraDebuggerClient.h"
 #include "Particles/FXBudget.h"
 #include "Engine/StaticMesh.h"
+#include "UObject/UObjectGlobals.h"
 
 IMPLEMENT_MODULE(INiagaraModule, Niagara);
 
@@ -195,6 +196,8 @@ void INiagaraModule::StartupModule()
 	// This includes the ability to compile scripts and load WITH_EDITOR_ONLY data.
 	// Note that when loading with the Editor, the NiagaraEditor module is loaded based on the plugin description.
 	FModuleManager::Get().LoadModule(TEXT("NiagaraEditor"));
+
+	FCoreUObjectDelegates::OnAssetLoaded.AddRaw(this, &INiagaraModule::OnAssetLoaded);
 #endif
 
 	//Init commonly used FNiagaraVariables
@@ -406,6 +409,10 @@ void INiagaraModule::ShutdownModule()
 
 	FNiagaraTypeRegistry::TearDown();
 
+#if WITH_EDITOR
+	FCoreUObjectDelegates::OnAssetLoaded.RemoveAll(this);
+#endif
+
 #if WITH_NIAGARA_DEBUGGER
 	DebuggerClient.Reset();
 #endif
@@ -515,6 +522,18 @@ void INiagaraModule::UnregisterPrecompiler(FDelegateHandle DelegateHandle)
 	ObjectPrecompilerDelegate.Unbind();
 }
 
+void INiagaraModule::OnAssetLoaded(UObject* Asset)
+{
+	if (Asset->IsA<UNiagaraSystem>())
+	{
+		CastChecked<UNiagaraSystem>(Asset)->UpdateSystemAfterLoad();
+	}
+	else if (Asset->IsA<UNiagaraEmitter>())
+	{
+		CastChecked<UNiagaraEmitter>(Asset)->UpdateEmitterAfterLoad();
+	}
+}
+
 #endif
 
 //////////////////////////////////////////////////////////////////////////
diff --git a/Engine/Plugins/FX/Niagara/Source/Niagara/Private/NiagaraSystem.cpp b/Engine/Plugins/FX/Niagara/Source/Niagara/Private/NiagaraSystem.cpp
--- a/Engine/Plugins/FX/Niagara/Source/Niagara/Private/NiagaraSystem.cpp
+++ b/Engine/Plugins/FX/Niagara/Source/Niagara/Private/NiagaraSystem.cpp
@@ -95,44 +95,6 @@ static FAutoConsoleVariableRef CVarNiagaraCompileWaitLoggingTerminationCap(
 
 //////////////////////////////////////////////////////////////////////////
 
-/**
-Game thread task to do any outstanding work on the loaded system
-*/
-struct FNiagaraSystemUpdateTask
-{
-	FNiagaraSystemUpdateTask(TWeakObjectPtr<UNiagaraSystem> InSystem) : WeakSystem(InSystem)
-	{
-	}
-
-	FORCEINLINE TStatId GetStatId() const { RETURN_QUICK_DECLARE_CYCLE_STAT(FNiagaraSystemUpdateTask, STATGROUP_TaskGraphTasks); }
-	static FORCEINLINE ENamedThreads::Type GetDesiredThread() { return ENamedThreads::GameThread; }
-	static FORCEINLINE ESubsequentsMode::Type GetSubsequentsMode() { return ESubsequentsMode::TrackSubsequents; }
-
-	void DoTask(ENamedThreads::Type CurrentThread, const FGraphEventRef& MyCompletionGraphEvent)
-	{
-		check(IsInGameThread());
-		UNiagaraSystem* System = WeakSystem.Get();
-		if (System == nullptr)
-		{
-			// probably got garbage collected in the meantime
-			return;
-		}
-#if WITH_EDITORONLY_DATA
-		for (const FNiagaraEmitterHandle& Handle : System->GetEmitterHandles())
-		{
-			if (Handle.GetInstance() && Handle.GetInstance()->UpdateTaskRef.IsValid())
-			{
-				// the emitter tasks should all be done by now, if not then something is very wrong
-				ensureMsgf(Handle.GetInstance()->UpdateTaskRef->IsComplete(), TEXT("Unfinished emitter update task for system %s"), *GetPathNameSafe(System));
-			}
-		}
-#endif
-		System->UpdateSystemAfterLoad();
-	}
-
-	TWeakObjectPtr<UNiagaraSystem> WeakSystem;
-};
-
 UNiagaraSystem::UNiagaraSystem(const FObjectInitializer& ObjectInitializer)
 : Super(ObjectInitializer)
 #if WITH_EDITORONLY_DATA
@@ -355,8 +317,15 @@ void UNiagaraSystem::UpdateSystemAfterLoad()
 		return;
 	}
 	bFullyLoaded = true;
-	UpdateTaskRef = nullptr;
-	
+
+	for (FNiagaraEmitterHandle& EmitterHandle : EmitterHandles)
+	{
+		if (EmitterHandle.GetInstance() != nullptr)
+		{
+			EmitterHandle.GetInstance()->UpdateEmitterAfterLoad();
+		}
+	}
+
 #if WITH_EDITORONLY_DATA
 	// We remove emitters and scripts on dedicated servers, so skip further work.
 	const bool bIsDedicatedServer = !GIsClient && GIsServer;
@@ -808,44 +777,11 @@ void UNiagaraSystem::PostLoad()
 	ExposedParameters.RecreateRedirections();
 #endif
 
-	FGraphEventArray EmitterPrerequisiteTasks;
-	for (FNiagaraEmitterHandle& EmitterHandle : EmitterHandles)
-	{
-#if WITH_EDITORONLY_DATA
-		EmitterHandle.ConditionalPostLoad(NiagaraVer);
-#else
-		if (UNiagaraEmitter* NiagaraEmitter = EmitterHandle.GetInstance())
-		{
-			NiagaraEmitter->ConditionalPostLoad();
-		}
+#if !WITH_EDITOR
+	// When running without the editor in a cooked build we run the update immediately in post load since
+	// there will be no merging or compiling which makes it safe to do so.
+	UpdateSystemAfterLoad();
 #endif
-		if (UNiagaraEmitter* NiagaraEmitter = EmitterHandle.GetInstance())
-		{
-			EmitterPrerequisiteTasks.Add(NiagaraEmitter->UpdateTaskRef);
-		}
-	}
-	
-#if WITH_EDITORONLY_DATA
-	if (EditorData == nullptr)
-	{
-		INiagaraModule& NiagaraModule = FModuleManager::GetModuleChecked<INiagaraModule>("Niagara");
-		EditorData = NiagaraModule.GetEditorOnlyDataUtilities().CreateDefaultEditorData(this);
-	}
-	else
-	{
-		EditorData->PostLoadFromOwner(this);
-	}
-
-	if (EditorParameters == nullptr)
-	{
-		INiagaraModule& NiagaraModule = FModuleManager::GetModuleChecked<INiagaraModule>("Niagara");
-		EditorParameters = NiagaraModule.GetEditorOnlyDataUtilities().CreateDefaultEditorParameters(this);
-	}
-	
-#endif // WITH_EDITORONLY_DATA
-
-	// we are not yet finished, but we do the rest of the work after postload in a separate task
-	UpdateTaskRef = TGraphTask<FNiagaraSystemUpdateTask>::CreateTask(&EmitterPrerequisiteTasks).ConstructAndDispatchWhenReady(this);
 }
 
 #if WITH_EDITORONLY_DATA
@@ -1002,7 +938,11 @@ bool UNiagaraSystem::IsReadyToRunInternal() const
 
 void UNiagaraSystem::GetAssetRegistryTags(TArray<FAssetRegistryTag>& OutTags) const
 {
-	EnsureFullyLoaded();
+	if (HasAnyFlags(RF_ClassDefaultObject) == false)
+	{
+		EnsureFullyLoaded();
+	}
+
 #if WITH_EDITOR
 	OutTags.Add(FAssetRegistryTag("HasGPUEmitter", HasAnyGPUEmitters() ? TEXT("True") : TEXT("False"), FAssetRegistryTag::TT_Alphabetical));
 
@@ -1753,19 +1693,8 @@ bool UNiagaraSystem::IsValidInternal() const
 
 void UNiagaraSystem::EnsureFullyLoaded() const
 {
-	if (UpdateTaskRef.IsValid() && UpdateTaskRef->IsComplete() == false)
-	{
-		ensureMsgf(FUObjectThreadContext::Get().IsRoutingPostLoad == false, TEXT("Niagara system data was forced to load during PostLoad instead of the taskgraph! Asset %s"), *GetPathName());
-		for (const FNiagaraEmitterHandle& EmitterHandle : EmitterHandles)
-		{
-			if (UNiagaraEmitter* NiagaraEmitter = EmitterHandle.GetInstance())
-			{
-				NiagaraEmitter->UpdateEmitterAfterLoad();
-			}
-		}
-		UNiagaraSystem* System = const_cast<UNiagaraSystem*>(this); // Sadly necessary because we are called from other const functions
-		System->UpdateSystemAfterLoad();
-	}	
+	UNiagaraSystem* System = const_cast<UNiagaraSystem*>(this);
+	System->UpdateSystemAfterLoad();
 }
 
 
diff --git a/Engine/Plugins/FX/Niagara/Source/Niagara/Public/NiagaraModule.h b/Engine/Plugins/FX/Niagara/Source/Niagara/Public/NiagaraModule.h
--- a/Engine/Plugins/FX/Niagara/Source/Niagara/Public/NiagaraModule.h
+++ b/Engine/Plugins/FX/Niagara/Source/Niagara/Public/NiagaraModule.h
@@ -85,6 +85,8 @@ public:
 	FDelegateHandle RegisterPrecompiler(FOnPrecompile PreCompiler);
 	void UnregisterPrecompiler(FDelegateHandle DelegateHandle);
 
+	void OnAssetLoaded(UObject* Asset);
+
 #endif
 
 #if NIAGARA_PERF_BASELINES
