Fix for pooled systems not reactivating on system edit.

diff --git a/Engine/Plugins/FX/Niagara/Source/Niagara/Private/NiagaraCommon.cpp b/Engine/Plugins/FX/Niagara/Source/Niagara/Private/NiagaraCommon.cpp
--- a/Engine/Plugins/FX/Niagara/Source/Niagara/Private/NiagaraCommon.cpp
+++ b/Engine/Plugins/FX/Niagara/Source/Niagara/Private/NiagaraCommon.cpp
@@ -150,6 +150,7 @@ void FNiagaraSystemUpdateContext::CommitUpdate()
 		{
 			Comp->ReinitializeSystem();
 			Comp->EndUpdateContextReset();
+			PostWork.ExecuteIfBound(Comp);
 		}
 	}
 	ComponentsToReInit.Empty();
@@ -160,6 +161,7 @@ void FNiagaraSystemUpdateContext::CommitUpdate()
 		{
 			Comp->ResetSystem();
 			Comp->EndUpdateContextReset();
+			PostWork.ExecuteIfBound(Comp);
 		}
 	}
 	ComponentsToReset.Empty();
@@ -173,6 +175,7 @@ void FNiagaraSystemUpdateContext::CommitUpdate()
 				SystemInstanceController->OnSimulationDestroyed();
 			}
 			Comp->EndUpdateContextReset();
+			PostWork.ExecuteIfBound(Comp);
 		}
 	}
 	ComponentsToReInit.Empty();
@@ -189,6 +192,12 @@ void FNiagaraSystemUpdateContext::AddAll(bool bReInit)
 	}
 }
 
+void FNiagaraSystemUpdateContext::Add(UNiagaraComponent* Component, bool bReInit)
+{
+	check(Component);
+	AddInternal(Component, bReInit);
+}
+
 void FNiagaraSystemUpdateContext::Add(const UNiagaraSystem* System, bool bReInit)
 {
 	for (TObjectIterator<UNiagaraComponent> It; It; ++It)
@@ -248,6 +257,8 @@ void FNiagaraSystemUpdateContext::Add(const UNiagaraParameterCollection* Collect
 
 void FNiagaraSystemUpdateContext::AddInternal(UNiagaraComponent* Comp, bool bReInit)
 {
+	PreWork.ExecuteIfBound(Comp);
+
 	Comp->BeginUpdateContextReset();
 
 	if (bReInit)
diff --git a/Engine/Plugins/FX/Niagara/Source/Niagara/Public/NiagaraCommon.h b/Engine/Plugins/FX/Niagara/Source/Niagara/Public/NiagaraCommon.h
--- a/Engine/Plugins/FX/Niagara/Source/Niagara/Public/NiagaraCommon.h
+++ b/Engine/Plugins/FX/Niagara/Source/Niagara/Public/NiagaraCommon.h
@@ -13,6 +13,7 @@
 #include "UObject/WeakFieldPtr.h"
 #include "NiagaraCommon.generated.h"
 
+class UNiagaraComponent;
 class UNiagaraSystem;
 class UNiagaraScript;
 class UNiagaraDataInterface;
@@ -701,6 +702,7 @@ struct NIAGARA_API FNiagaraSystemUpdateContext
 {
 	GENERATED_BODY()
 
+	FNiagaraSystemUpdateContext(UNiagaraComponent* Component, bool bReInit, bool bInDestroyOnAdd = false, bool bInOnlyActive = false) :bDestroyOnAdd(bInDestroyOnAdd), bOnlyActive(bInOnlyActive) { Add(Component, bReInit); }
 	FNiagaraSystemUpdateContext(const UNiagaraSystem* System, bool bReInit, bool bInDestroyOnAdd = false, bool bInOnlyActive = false) :bDestroyOnAdd(bInDestroyOnAdd), bOnlyActive(bInOnlyActive) { Add(System, bReInit); }
 #if WITH_EDITORONLY_DATA
 	FNiagaraSystemUpdateContext(const UNiagaraEmitter* Emitter, bool bReInit, bool bInDestroyOnAdd = false, bool bInOnlyActive = false) :bDestroyOnAdd(bInDestroyOnAdd), bOnlyActive(bInOnlyActive)  { Add(Emitter, bReInit); }
@@ -715,6 +717,7 @@ struct NIAGARA_API FNiagaraSystemUpdateContext
 	void SetDestroyOnAdd(bool bInDestroyOnAdd) { bDestroyOnAdd = bInDestroyOnAdd; }
 	void SetOnlyActive(bool bInOnlyActive) { bOnlyActive = bInOnlyActive; }
 
+	void Add(UNiagaraComponent* Component, bool bReInit);
 	void Add(const UNiagaraSystem* System, bool bReInit);
 #if WITH_EDITORONLY_DATA
 	void Add(const UNiagaraEmitter* Emitter, bool bReInit);
@@ -724,10 +727,14 @@ struct NIAGARA_API FNiagaraSystemUpdateContext
 
 	/** Adds all currently active systems.*/
 	void AddAll(bool bReInit);
-
+	
 	/** Handles any pending reinits or resets of system instances in this update context. */
 	void CommitUpdate();
 
+	DECLARE_DELEGATE_OneParam(FCustomWorkDelegate, UNiagaraComponent*);
+	FCustomWorkDelegate& GetPreWork(){ return PreWork; }
+	FCustomWorkDelegate& GetPostWork() { return PostWork; }
+
 private:
 	void AddInternal(class UNiagaraComponent* Comp, bool bReInit);
 
@@ -743,6 +750,12 @@ private:
 	bool bDestroyOnAdd;
 	bool bOnlyActive;
 	//TODO: When we allow component less systems we'll also want to find and reset those.
+
+	/** Delegate called before a component is added to the context for update. */
+	FCustomWorkDelegate PreWork;
+	
+	/** Delegate called before after a component has been updated. */
+	FCustomWorkDelegate PostWork;
 };
 
 /** Defines different usages for a niagara script. */
diff --git a/Engine/Plugins/FX/Niagara/Source/NiagaraEditor/Private/ViewModels/NiagaraSystemViewModel.cpp b/Engine/Plugins/FX/Niagara/Source/NiagaraEditor/Private/ViewModels/NiagaraSystemViewModel.cpp
--- a/Engine/Plugins/FX/Niagara/Source/NiagaraEditor/Private/ViewModels/NiagaraSystemViewModel.cpp
+++ b/Engine/Plugins/FX/Niagara/Source/NiagaraEditor/Private/ViewModels/NiagaraSystemViewModel.cpp
@@ -1515,34 +1515,29 @@ void FNiagaraSystemViewModel::ResetSystem(ETimeResetMode TimeResetMode, EMultiRe
 		System->CacheFromCompiledData();
 	}
 
-	TArray<UNiagaraComponent*> ReferencingComponents;
+	FNiagaraSystemUpdateContext UpdateContext;
+	UpdateContext.GetPostWork().BindLambda(
+		[ReinitMode, bResetAge](UNiagaraComponent* Component)
+		{
+			if (ReinitMode == EReinitMode::ResetSystem && bResetAge  && Component->GetAgeUpdateMode() == ENiagaraAgeUpdateMode::DesiredAge)
+			{
+				Component->SetDesiredAge(0);
+			}
+		}
+	);
 	if (MultiResetMode == EMultiResetMode::ResetThisInstance)
 	{
 		if (PreviewComponent != nullptr)
 		{
-			ReferencingComponents.Add(PreviewComponent);
+			UpdateContext.Add(PreviewComponent, ReinitMode == EReinitMode::ReinitializeSystem);
 		}
 	}
 	else
 	{
-		ReferencingComponents = FNiagaraEditorUtilities::GetComponentsThatReferenceSystem(GetSystem());
+		UpdateContext.Add(&GetSystem(), ReinitMode == EReinitMode::ReinitializeSystem);
 	}
 
-	for (auto Component : ReferencingComponents)
-	{
-		if (ReinitMode == EReinitMode::ResetSystem)
-		{
-			Component->ResetSystem();
-			if (bResetAge && Component->GetAgeUpdateMode() == ENiagaraAgeUpdateMode::DesiredAge)
-			{
-				Component->SetDesiredAge(0);
-			}
-		}
-		else if (ReinitMode == EReinitMode::ReinitializeSystem)
-		{
-			Component->ReinitializeSystem();
-		}
-	}
+	UpdateContext.CommitUpdate();
 
 	if (EditMode == ENiagaraSystemViewModelEditMode::EmitterAsset && MultiResetMode == EMultiResetMode::AllowResetAllInstances && EditorSettings->GetResetDependentSystemsWhenEditingEmitters())
 	{
