Fixing static analyss warning

diff --git a/Engine/Plugins/FX/Niagara/Source/Niagara/Private/NiagaraComponent.cpp b/Engine/Plugins/FX/Niagara/Source/Niagara/Private/NiagaraComponent.cpp
--- a/Engine/Plugins/FX/Niagara/Source/Niagara/Private/NiagaraComponent.cpp
+++ b/Engine/Plugins/FX/Niagara/Source/Niagara/Private/NiagaraComponent.cpp
@@ -798,7 +798,7 @@ bool UNiagaraComponent::InitializeSystem()
 
 		MarkRenderStateDirty();
 
-		if (GIsAutomationTesting && GetAsset())
+		if (GIsAutomationTesting && GetAsset() != nullptr)
 		{
 			UE_LOG(LogNiagara, Log, TEXT("UNiagaraComponent> System %s initialized."), *GetAsset()->GetPathName());
 		}
@@ -1283,7 +1283,7 @@ void UNiagaraComponent::OnSystemComplete(bool bExternalCompletion)
 		}
 	}
 
-	if (GIsAutomationTesting && GetAsset())
+	if (GIsAutomationTesting && GetAsset() != nullptr)
 	{
 		UE_LOG(LogNiagara, Log, TEXT("UNiagaraComponent> System %s completed."), *GetAsset()->GetPathName());
 	}
