Making Warning a Display instead

diff --git a/Engine/Plugins/FX/Niagara/Source/Niagara/Private/NiagaraEmitterInstance.cpp b/Engine/Plugins/FX/Niagara/Source/Niagara/Private/NiagaraEmitterInstance.cpp
--- a/Engine/Plugins/FX/Niagara/Source/Niagara/Private/NiagaraEmitterInstance.cpp
+++ b/Engine/Plugins/FX/Niagara/Source/Niagara/Private/NiagaraEmitterInstance.cpp
@@ -1417,7 +1417,9 @@ void FNiagaraEmitterInstance::Tick(float DeltaSeconds)
 						NumSpawnedOnGPUThisFrame += Info.Count;
 						if (NumSpawnedOnGPUThisFrame > GMaxNiagaraGPUParticlesSpawnPerFrame)
 						{
-							UE_LOG(LogNiagara, Warning, TEXT("%s has attempted to execeed max GPU per frame spawn! | Max: %d | Requested: %d | SpawnInfoEntry: %d"), *CachedEmitter->GetUniqueEmitterName(), GMaxNiagaraGPUParticlesSpawnPerFrame, NumSpawnedOnGPUThisFrame, SpawnInfoIdx);
+							FString DebugMsg = FString::Printf(TEXT("%s has attempted to execeed max GPU per frame spawn! | Max: %d | Requested: %d | SpawnInfoEntry: %d"), *CachedEmitter->GetFullName(), GMaxNiagaraGPUParticlesSpawnPerFrame, NumSpawnedOnGPUThisFrame, SpawnInfoIdx);
+							UE_LOG(LogNiagara, Warning, TEXT("%s"), *DebugMsg);
+							GEngine->AddOnScreenDebugMessage(-1, 5.f, FColor::Yellow, *DebugMsg);
 							break;
 						}
 
@@ -1433,7 +1435,9 @@ void FNiagaraEmitterInstance::Tick(float DeltaSeconds)
 					}
 					else if (Info.Count > 0)
 					{
-						UE_LOG(LogNiagara, Warning, TEXT("%s Exceeded Gpu spawn info count, see NIAGARA_MAX_GPU_SPAWN_INFOS for more information!"), *CachedEmitter->GetUniqueEmitterName());
+						FString DebugMsg = FString::Printf(TEXT("%s Exceeded Gpu spawn info count, see NIAGARA_MAX_GPU_SPAWN_INFOS for more information!"), *CachedEmitter->GetUniqueEmitterName());
+						UE_LOG(LogNiagara, Warning, TEXT("%s"), *DebugMsg);
+						GEngine->AddOnScreenDebugMessage(-1, 5.f, FColor::Yellow, *DebugMsg);
 						break;
 					}
 
@@ -1499,7 +1503,9 @@ void FNiagaraEmitterInstance::Tick(float DeltaSeconds)
 	//TODO: These current limits can be improved relatively easily. Though perf in at these counts will obviously be an issue anyway.
 	if (CachedEmitter->SimTarget == ENiagaraSimTarget::CPUSim && AllocationSize > GMaxNiagaraCPUParticlesPerEmitter)
 	{
-		UE_LOG(LogNiagara, Warning, TEXT("%s has attempted to exceed the max CPU particle count! | Max: %d | Requested: %u"), *CachedEmitter->GetFullName(), GMaxNiagaraCPUParticlesPerEmitter, AllocationSize);
+		FString DebugMsg = FString::Printf(TEXT("%s has attempted to exceed the max CPU particle count! | Max: %d | Requested: %u"), *CachedEmitter->GetFullName(), GMaxNiagaraCPUParticlesPerEmitter, AllocationSize);
+		UE_LOG(LogNiagara, Warning, TEXT("%s"), *DebugMsg);
+		GEngine->AddOnScreenDebugMessage(-1, 5.f, FColor::Yellow, *DebugMsg);
 
 		//We clear the emitters estimate otherwise we get stuck in this state forever.
 		CachedEmitter->ClearRuntimeAllocationEstimate();
