- Allow arrays to be an iteration source
- Update GPU copy when we dirty the data on the CPU
- Allow us to skip the SetElement

diff --git a/Engine/Plugins/FX/Niagara/Source/Niagara/Classes/NiagaraDataInterfaceArray.h b/Engine/Plugins/FX/Niagara/Source/Niagara/Classes/NiagaraDataInterfaceArray.h
--- a/Engine/Plugins/FX/Niagara/Source/Niagara/Classes/NiagaraDataInterfaceArray.h
+++ b/Engine/Plugins/FX/Niagara/Source/Niagara/Classes/NiagaraDataInterfaceArray.h
@@ -2,7 +2,7 @@
 
 #pragma once
 
-#include "NiagaraDataInterface.h"
+#include "NiagaraDataInterfaceRW.h"
 #include "NiagaraCommon.h"
 #include "NiagaraShared.h"
 #include "NiagaraShader.h"
@@ -32,7 +32,7 @@ struct INiagaraDataInterfaceArrayImpl
 };
 
 UCLASS(abstract, EditInlineNew)
-class NIAGARA_API UNiagaraDataInterfaceArray : public UNiagaraDataInterface
+class NIAGARA_API UNiagaraDataInterfaceArray : public UNiagaraDataInterfaceRWBase
 {
 	GENERATED_UCLASS_BODY()
 
diff --git a/Engine/Plugins/FX/Niagara/Source/Niagara/Classes/NiagaraDataInterfaceArrayImpl.h b/Engine/Plugins/FX/Niagara/Source/Niagara/Classes/NiagaraDataInterfaceArrayImpl.h
--- a/Engine/Plugins/FX/Niagara/Source/Niagara/Classes/NiagaraDataInterfaceArrayImpl.h
+++ b/Engine/Plugins/FX/Niagara/Source/Niagara/Classes/NiagaraDataInterfaceArrayImpl.h
@@ -33,7 +33,7 @@ struct FNDIArrayImplHelper : public FNDIArrayImplHelperBase<TArrayType>
 {
 };
 
-struct FNiagaraDataInterfaceProxyArrayImpl : public FNiagaraDataInterfaceProxy
+struct FNiagaraDataInterfaceProxyArrayImpl : public FNiagaraDataInterfaceProxyRW
 {
 	FReadBuffer		Buffer;
 	int32			NumElements = 0;
@@ -46,10 +46,25 @@ struct FNiagaraDataInterfaceProxyArrayImpl : public FNiagaraDataInterfaceProxy
 
 	virtual void ConsumePerInstanceDataFromGameThread(void* PerInstanceData, const FNiagaraSystemInstanceID& Instance) override { check(false); }
 	virtual int32 PerInstanceDataPassedToRenderThreadSize() const override { return 0; }
+
+	virtual FIntVector GetElementCount(FNiagaraSystemInstanceID SystemInstanceID) const override { return FIntVector(NumElements, 1, 1); }
+	//virtual uint32 GetGPUInstanceCountOffset(FNiagaraSystemInstanceID SystemInstanceID) const override { return INDEX_NONE; }
 };
 
 struct FNiagaraDataInterfaceArrayImplHelper
 {
+	struct FFunctionVersion
+	{
+		enum Type
+		{
+			InitialVersion = 0,
+			AddOptionalExecuteToSet = 1,
+
+			VersionPlusOne,
+			LatestVersion = VersionPlusOne - 1
+		};
+	};
+
 	static const FName Function_LengthName;
 	static const FName Function_IsValidIndexName;
 	static const FName Function_LastIndexName;
@@ -108,155 +123,109 @@ struct FNiagaraDataInterfaceArrayImpl : public INiagaraDataInterfaceArrayImpl
 
 	virtual void GetFunctions(TArray<FNiagaraFunctionSignature>& OutFunctions) const override
 	{
-		OutFunctions.Reserve(OutFunctions.Num() + 3);
+		OutFunctions.Reserve(OutFunctions.Num() + 9);
 
 		// Immutable functions
+		FNiagaraFunctionSignature DefaultImmutableSig;
+		DefaultImmutableSig.bMemberFunction = true;
+		DefaultImmutableSig.bRequiresContext = false;
+		DefaultImmutableSig.bSupportsCPU = FNDIArrayImplHelper<TArrayType>::bSupportsCPU;
+		DefaultImmutableSig.bSupportsGPU = FNDIArrayImplHelper<TArrayType>::bSupportsGPU;
+		DefaultImmutableSig.Inputs.Emplace(FNiagaraTypeDefinition(TObjectType::StaticClass()), TEXT("Array interface"));
+#if WITH_EDITORONLY_DATA
+		DefaultImmutableSig.FunctionVersion = FNiagaraDataInterfaceArrayImplHelper::FFunctionVersion::LatestVersion;
+#endif
 		{
-			FNiagaraFunctionSignature& Sig = OutFunctions.AddDefaulted_GetRef();
+			FNiagaraFunctionSignature& Sig = OutFunctions.Add_GetRef(DefaultImmutableSig);
 			Sig.Name = FNiagaraDataInterfaceArrayImplHelper::Function_LengthName;
 			#if WITH_EDITORONLY_DATA
 				Sig.Description = NSLOCTEXT("Niagara", "Array_LengthDesc", "Gets the number of elements in the array.");
 			#endif
-			Sig.bMemberFunction = true;
-			Sig.bRequiresContext = false;
-			Sig.bExperimental = true;
-			Sig.bSupportsCPU = FNDIArrayImplHelper<TArrayType>::bSupportsCPU;
-			Sig.bSupportsGPU = FNDIArrayImplHelper<TArrayType>::bSupportsGPU;
-			Sig.Inputs.Add(FNiagaraVariable(FNiagaraTypeDefinition(TObjectType::StaticClass()), TEXT("Array interface")));
-			Sig.Outputs.Add(FNiagaraVariable(FNiagaraTypeDefinition::GetIntDef(), TEXT("Num")));
+			Sig.Outputs.Emplace(FNiagaraTypeDefinition::GetIntDef(), TEXT("Num"));
 		}
 
 		{
-			FNiagaraFunctionSignature& Sig = OutFunctions.AddDefaulted_GetRef();
+			FNiagaraFunctionSignature& Sig = OutFunctions.Add_GetRef(DefaultImmutableSig);
 			Sig.Name = FNiagaraDataInterfaceArrayImplHelper::Function_IsValidIndexName;
 			#if WITH_EDITORONLY_DATA
 				Sig.Description = NSLOCTEXT("Niagara", "Array_IsValidIndexDesc", "Tests to see if the index is valid and exists in the array.");
 			#endif
-			Sig.bMemberFunction = true;
-			Sig.bRequiresContext = false;
-			Sig.bExperimental = true;
-			Sig.bSupportsCPU = FNDIArrayImplHelper<TArrayType>::bSupportsCPU;
-			Sig.bSupportsGPU = FNDIArrayImplHelper<TArrayType>::bSupportsGPU;
-			Sig.Inputs.Add(FNiagaraVariable(FNiagaraTypeDefinition(TObjectType::StaticClass()), TEXT("Array interface")));
-			Sig.Inputs.Add(FNiagaraVariable(FNiagaraTypeDefinition::GetIntDef(), TEXT("Index")));
-			Sig.Outputs.Add(FNiagaraVariable(FNiagaraTypeDefinition::GetBoolDef(), TEXT("Valid")));
+			Sig.Inputs.Emplace(FNiagaraTypeDefinition::GetIntDef(), TEXT("Index"));
+			Sig.Outputs.Emplace(FNiagaraTypeDefinition::GetBoolDef(), TEXT("Valid"));
 		}
 
 		{
-			FNiagaraFunctionSignature& Sig = OutFunctions.AddDefaulted_GetRef();
+			FNiagaraFunctionSignature& Sig = OutFunctions.Add_GetRef(DefaultImmutableSig);
 			Sig.Name = FNiagaraDataInterfaceArrayImplHelper::Function_LastIndexName;
 #if WITH_EDITORONLY_DATA
 			Sig.Description = NSLOCTEXT("Niagara", "Array_LastIndexDesc", "Returns the last valid index in the array, will be -1 if no elements.");
 #endif
-			Sig.bMemberFunction = true;
-			Sig.bRequiresContext = false;
-			Sig.bExperimental = true;
-			Sig.bSupportsCPU = FNDIArrayImplHelper<TArrayType>::bSupportsCPU;
-			Sig.bSupportsGPU = FNDIArrayImplHelper<TArrayType>::bSupportsGPU;
-			Sig.Inputs.Add(FNiagaraVariable(FNiagaraTypeDefinition(TObjectType::StaticClass()), TEXT("Array interface")));
-			Sig.Outputs.Add(FNiagaraVariable(FNiagaraTypeDefinition::GetIntDef(), TEXT("Index")));
+			Sig.Outputs.Emplace(FNiagaraTypeDefinition::GetIntDef(), TEXT("Index"));
 		}
 
 		{
-			FNiagaraFunctionSignature& Sig = OutFunctions.AddDefaulted_GetRef();
+			FNiagaraFunctionSignature& Sig = OutFunctions.Add_GetRef(DefaultImmutableSig);
 			Sig.Name = FNiagaraDataInterfaceArrayImplHelper::Function_GetName;
 			#if WITH_EDITORONLY_DATA
 				Sig.Description = NSLOCTEXT("Niagara", "Array_GetDesc", "Gets the value from the array at the given zero based index.");
 			#endif
-			Sig.bMemberFunction = true;
-			Sig.bRequiresContext = false;
-			Sig.bExperimental = true;
-			Sig.bSupportsCPU = FNDIArrayImplHelper<TArrayType>::bSupportsCPU;
-			Sig.bSupportsGPU = FNDIArrayImplHelper<TArrayType>::bSupportsGPU;
-			Sig.Inputs.Add(FNiagaraVariable(FNiagaraTypeDefinition(TObjectType::StaticClass()), TEXT("Array interface")));
-			Sig.Inputs.Add(FNiagaraVariable(FNiagaraTypeDefinition::GetIntDef(), TEXT("Index")));
-			Sig.Outputs.Add(FNiagaraVariable(FNDIArrayImplHelper<TArrayType>::GetTypeDefinition(), TEXT("Value")));
+			Sig.Inputs.Emplace(FNiagaraTypeDefinition::GetIntDef(), TEXT("Index"));
+			Sig.Outputs.Emplace(FNDIArrayImplHelper<TArrayType>::GetTypeDefinition(), TEXT("Value"));
 		}
 
 		// Mutable functions
+		FNiagaraFunctionSignature DefaultMutableSig = DefaultImmutableSig;
+		DefaultMutableSig.bSupportsGPU = false;
+		DefaultMutableSig.bRequiresExecPin = true;
 		{
-			FNiagaraFunctionSignature& Sig = OutFunctions.AddDefaulted_GetRef();
+			FNiagaraFunctionSignature& Sig = OutFunctions.Add_GetRef(DefaultMutableSig);
 			Sig.Name = FNiagaraDataInterfaceArrayImplHelper::Function_ClearName;
 			#if WITH_EDITORONLY_DATA
 				Sig.Description = NSLOCTEXT("Niagara", "Array_ClearDesc", "Clears the array, removing all elements");
 			#endif
-			Sig.bMemberFunction = true;
-			Sig.bRequiresContext = false;
-			Sig.bExperimental = true;
-			Sig.bSupportsCPU = FNDIArrayImplHelper<TArrayType>::bSupportsCPU;
-			Sig.bSupportsGPU = false;
-			Sig.bRequiresExecPin = true;
 			Sig.ModuleUsageBitmask = ENiagaraScriptUsageMask::System | ENiagaraScriptUsageMask::Emitter;
-			Sig.Inputs.Add(FNiagaraVariable(FNiagaraTypeDefinition(TObjectType::StaticClass()), TEXT("Array interface")));
 		}
 
 		{
-			FNiagaraFunctionSignature& Sig = OutFunctions.AddDefaulted_GetRef();
+			FNiagaraFunctionSignature& Sig = OutFunctions.Add_GetRef(DefaultMutableSig);
 			Sig.Name = FNiagaraDataInterfaceArrayImplHelper::Function_ResizeName;
 			#if WITH_EDITORONLY_DATA
 				Sig.Description = NSLOCTEXT("Niagara", "Array_ResizeDesc", "Resizes the array to the specified size, initializing new elements with the default value.");
 			#endif
-			Sig.bMemberFunction = true;
-			Sig.bRequiresContext = false;
-			Sig.bExperimental = true;
-			Sig.bSupportsCPU = FNDIArrayImplHelper<TArrayType>::bSupportsCPU;
-			Sig.bSupportsGPU = false;
-			Sig.bRequiresExecPin = true;
 			Sig.ModuleUsageBitmask = ENiagaraScriptUsageMask::System | ENiagaraScriptUsageMask::Emitter;
-			Sig.Inputs.Add(FNiagaraVariable(FNiagaraTypeDefinition(TObjectType::StaticClass()), TEXT("Array interface")));
-			Sig.Inputs.Add(FNiagaraVariable(FNiagaraTypeDefinition::GetIntDef(), TEXT("Num")));
+			Sig.Inputs.Emplace(FNiagaraTypeDefinition::GetIntDef(), TEXT("Num"));
 		}
 
 		{
-			FNiagaraFunctionSignature& Sig = OutFunctions.AddDefaulted_GetRef();
+			FNiagaraFunctionSignature& Sig = OutFunctions.Add_GetRef(DefaultMutableSig);
 			Sig.Name = FNiagaraDataInterfaceArrayImplHelper::Function_SetArrayElemName;
 			#if WITH_EDITORONLY_DATA
 				Sig.Description = NSLOCTEXT("Niagara", "Array_SetArrayElemDesc", "Sets the value at the given zero based index (i.e the first element is 0).");
 			#endif
-			Sig.bMemberFunction = true;
-			Sig.bRequiresContext = false;
-			Sig.bExperimental = true;
-			Sig.bSupportsCPU = FNDIArrayImplHelper<TArrayType>::bSupportsCPU;
-			Sig.bSupportsGPU = false;
-			Sig.bRequiresExecPin = true;
-			Sig.Inputs.Add(FNiagaraVariable(FNiagaraTypeDefinition(TObjectType::StaticClass()), TEXT("Array interface")));
-			Sig.Inputs.Add(FNiagaraVariable(FNiagaraTypeDefinition::GetIntDef(), TEXT("Index")));
-			Sig.Inputs.Add(FNiagaraVariable(FNDIArrayImplHelper<TArrayType>::GetTypeDefinition(), TEXT("Value")));
+			Sig.Inputs.Emplace(FNiagaraTypeDefinition::GetBoolDef(), TEXT("SkipSet"));
+			Sig.Inputs.Emplace(FNiagaraTypeDefinition::GetIntDef(), TEXT("Index"));
+			Sig.Inputs.Emplace(FNDIArrayImplHelper<TArrayType>::GetTypeDefinition(), TEXT("Value"));
 		}
 
 		{
-			FNiagaraFunctionSignature& Sig = OutFunctions.AddDefaulted_GetRef();
+			FNiagaraFunctionSignature& Sig = OutFunctions.Add_GetRef(DefaultMutableSig);
 			Sig.Name = FNiagaraDataInterfaceArrayImplHelper::Function_AddName;
 			#if WITH_EDITORONLY_DATA
 				Sig.Description = NSLOCTEXT("Niagara", "Array_AddDesc", "Optionally add a value onto the end of the array.");
 			#endif
-			Sig.bMemberFunction = true;
-			Sig.bRequiresContext = false;
-			Sig.bExperimental = true;
-			Sig.bSupportsCPU = FNDIArrayImplHelper<TArrayType>::bSupportsCPU;
-			Sig.bSupportsGPU = false;
-			Sig.bRequiresExecPin = true;
-			Sig.Inputs.Add(FNiagaraVariable(FNiagaraTypeDefinition(TObjectType::StaticClass()), TEXT("Array interface")));
-			Sig.Inputs.Add(FNiagaraVariable(FNiagaraTypeDefinition::GetBoolDef(), TEXT("SkipAdd")));
-			Sig.Inputs.Add(FNiagaraVariable(FNDIArrayImplHelper<TArrayType>::GetTypeDefinition(), TEXT("Value")));
+			Sig.Inputs.Emplace(FNiagaraTypeDefinition::GetBoolDef(), TEXT("SkipAdd"));
+			Sig.Inputs.Emplace(FNDIArrayImplHelper<TArrayType>::GetTypeDefinition(), TEXT("Value"));
 		}
 
 		{
-			FNiagaraFunctionSignature& Sig = OutFunctions.AddDefaulted_GetRef();
+			FNiagaraFunctionSignature& Sig = OutFunctions.Add_GetRef(DefaultMutableSig);
 			Sig.Name = FNiagaraDataInterfaceArrayImplHelper::Function_RemoveLastElemName;
 #if WITH_EDITORONLY_DATA
 			Sig.Description = NSLOCTEXT("Niagara", "Array_RemoveLastElemDesc", "Optionally remove the last element from the array.  Returns the default value if no elements are in the array or you skip the remove.");
 #endif
-			Sig.bMemberFunction = true;
-			Sig.bRequiresContext = false;
-			Sig.bExperimental = true;
-			Sig.bSupportsCPU = FNDIArrayImplHelper<TArrayType>::bSupportsCPU;
-			Sig.bSupportsGPU = false;
-			Sig.bRequiresExecPin = true;
-			Sig.Inputs.Add(FNiagaraVariable(FNiagaraTypeDefinition(TObjectType::StaticClass()), TEXT("Array interface")));
-			Sig.Inputs.Add(FNiagaraVariable(FNiagaraTypeDefinition::GetBoolDef(), TEXT("SkipRemove")));
-			Sig.Outputs.Add(FNiagaraVariable(FNDIArrayImplHelper<TArrayType>::GetTypeDefinition(), TEXT("Value")));
-			Sig.Outputs.Add(FNiagaraVariable(FNiagaraTypeDefinition::GetBoolDef(), TEXT("IsValid")));
+			Sig.Inputs.Emplace(FNiagaraTypeDefinition::GetBoolDef(), TEXT("SkipRemove"));
+			Sig.Outputs.Emplace(FNDIArrayImplHelper<TArrayType>::GetTypeDefinition(), TEXT("Value"));
+			Sig.Outputs.Emplace(FNiagaraTypeDefinition::GetBoolDef(), TEXT("IsValid"));
 		}
 	}
 
@@ -570,16 +539,16 @@ struct FNiagaraDataInterfaceArrayImpl : public INiagaraDataInterfaceArrayImpl
 
 	void Clear(FVectorVMExternalFunctionContext& Context)
 	{
-		//-TODO: This dirties the GPU data
 		ensureMsgf(Context.GetNumInstances() == 1, TEXT("Setting the number of values in an array with more than one instance, which doesn't make sense"));
 
 		FRWScopeLock WriteLock(Owner->ArrayRWGuard, SLT_Write);
 		Data.Reset();
+
+		Owner->MarkRenderDataDirty();
 	}
 
 	void Resize(FVectorVMExternalFunctionContext& Context)
 	{
-		//-TODO: This dirties the GPU data
 		FNDIInputParam<int32> NewNumParam(Context);
 
 		ensureMsgf(Context.GetNumInstances() == 1, TEXT("Setting the number of values in an array with more than one instance, which doesn't make sense"));
@@ -597,11 +566,13 @@ struct FNiagaraDataInterfaceArrayImpl : public INiagaraDataInterfaceArrayImpl
 				Data[i] = DefaultValue;
 			}
 		}
+
+		Owner->MarkRenderDataDirty();
 	}
 
 	void SetValue(FVectorVMExternalFunctionContext& Context)
 	{
-		//-TODO: This dirties the GPU data
+		FNDIInputParam<FNiagaraBool> InSkipSet(Context);
 		FNDIInputParam<int32> IndexParam(Context);
 		FNDIInputParam<typename FNDIArrayImplHelper<TArrayType>::TVMArrayType> InValue(Context);
 
@@ -610,17 +581,19 @@ struct FNiagaraDataInterfaceArrayImpl : public INiagaraDataInterfaceArrayImpl
 		{
 			const int32 Index = IndexParam.GetAndAdvance();
 			const TArrayType Value = (TArrayType)InValue.GetAndAdvance();
+			const bool bSkipSet = InSkipSet.GetAndAdvance();
 
-			if (Data.IsValidIndex(Index))
+			if (!bSkipSet && Data.IsValidIndex(Index))
 			{
 				Data[Index] = Value;
 			}
 		}
+
+		Owner->MarkRenderDataDirty();
 	}
 
 	void PushValue(FVectorVMExternalFunctionContext& Context)
 	{
-		//-TODO: This dirties the GPU data
 		FNDIInputParam<FNiagaraBool> InSkipExecute(Context);
 		FNDIInputParam<typename FNDIArrayImplHelper<TArrayType>::TVMArrayType> InValue(Context);
 
@@ -636,11 +609,12 @@ struct FNiagaraDataInterfaceArrayImpl : public INiagaraDataInterfaceArrayImpl
 				Data.Emplace(Value);
 			}
 		}
+
+		Owner->MarkRenderDataDirty();
 	}
 
 	void PopValue(FVectorVMExternalFunctionContext& Context)
 	{
-		//-TODO: This dirties the GPU data
 		FNDIInputParam<FNiagaraBool> InSkipExecute(Context);
 		FNDIOutputParam<typename FNDIArrayImplHelper<TArrayType>::TVMArrayType> OutValue(Context);
 		FNDIOutputParam<FNiagaraBool> OutIsValid(Context);
@@ -661,5 +635,7 @@ struct FNiagaraDataInterfaceArrayImpl : public INiagaraDataInterfaceArrayImpl
 				OutIsValid.SetAndAdvance(true);
 			}
 		}
+
+		Owner->MarkRenderDataDirty();
 	}
 };
\ No newline at end of file
diff --git a/Engine/Plugins/FX/Niagara/Source/Niagara/Private/NiagaraDataInterfaceArrayImpl.cpp b/Engine/Plugins/FX/Niagara/Source/Niagara/Private/NiagaraDataInterfaceArrayImpl.cpp
--- a/Engine/Plugins/FX/Niagara/Source/Niagara/Private/NiagaraDataInterfaceArrayImpl.cpp
+++ b/Engine/Plugins/FX/Niagara/Source/Niagara/Private/NiagaraDataInterfaceArrayImpl.cpp
@@ -26,30 +26,44 @@ FString FNiagaraDataInterfaceArrayImplHelper::GetBufferSizeName(const FString& I
 #if WITH_EDITORONLY_DATA
 bool FNiagaraDataInterfaceArrayImplHelper::UpgradeFunctionCall(FNiagaraFunctionSignature& FunctionSignature)
 {
-	bool bModified = false;
-
-	static const TPair<FName, FName> NodeRenames[] =
+	// Early out, nothing to do here
+	if ( FunctionSignature.FunctionVersion == FFunctionVersion::LatestVersion )
 	{
-		MakeTuple(FName("GetNum"),			FNiagaraDataInterfaceArrayImplHelper::Function_LengthName),
-		MakeTuple(FName("GetValue"),		FNiagaraDataInterfaceArrayImplHelper::Function_GetName),
-		MakeTuple(FName("Reset"),			FNiagaraDataInterfaceArrayImplHelper::Function_ClearName),
-		MakeTuple(FName("SetNum"),			FNiagaraDataInterfaceArrayImplHelper::Function_ResizeName),
-		MakeTuple(FName("SetValue"),		FNiagaraDataInterfaceArrayImplHelper::Function_SetArrayElemName),
-		MakeTuple(FName("PushValue"),		FNiagaraDataInterfaceArrayImplHelper::Function_AddName),
-		MakeTuple(FName("PopValue"),		FNiagaraDataInterfaceArrayImplHelper::Function_RemoveLastElemName),
-	};
-
-	for (const auto& Pair : NodeRenames)
+		return false;
+	}
+
+	if ( FunctionSignature.FunctionVersion < FFunctionVersion::AddOptionalExecuteToSet )
 	{
-		if (Pair.Key == FunctionSignature.Name)
+		static const TPair<FName, FName> NodeRenames[] =
+		{
+			MakeTuple(FName("GetNum"),			FNiagaraDataInterfaceArrayImplHelper::Function_LengthName),
+			MakeTuple(FName("GetValue"),		FNiagaraDataInterfaceArrayImplHelper::Function_GetName),
+			MakeTuple(FName("Reset"),			FNiagaraDataInterfaceArrayImplHelper::Function_ClearName),
+			MakeTuple(FName("SetNum"),			FNiagaraDataInterfaceArrayImplHelper::Function_ResizeName),
+			MakeTuple(FName("SetValue"),		FNiagaraDataInterfaceArrayImplHelper::Function_SetArrayElemName),
+			MakeTuple(FName("PushValue"),		FNiagaraDataInterfaceArrayImplHelper::Function_AddName),
+			MakeTuple(FName("PopValue"),		FNiagaraDataInterfaceArrayImplHelper::Function_RemoveLastElemName),
+		};
+
+		for (const auto& Pair : NodeRenames)
+		{
+			if (Pair.Key == FunctionSignature.Name)
+			{
+				FunctionSignature.Name = Pair.Value;
+				break;
+			}
+		}
+
+		FunctionSignature.bExperimental = false;
+
+		if ( FunctionSignature.Name == FNiagaraDataInterfaceArrayImplHelper::Function_SetArrayElemName )
 		{
-			FunctionSignature.Name = Pair.Value;
-			bModified = true;
-			break;
+			FunctionSignature.Inputs.EmplaceAt(1, FNiagaraTypeDefinition::GetBoolDef(), TEXT("SkipSet"));
 		}
 	}
+	FunctionSignature.FunctionVersion = FFunctionVersion::LatestVersion;
 
-	return bModified;
+	return true;
 }
 #endif
 
