- Allow readback manager to select a sub resource region to read from

diff --git a/Engine/Plugins/FX/Niagara/Source/Niagara/Private/NiagaraGpuReadbackManager.cpp b/Engine/Plugins/FX/Niagara/Source/Niagara/Private/NiagaraGpuReadbackManager.cpp
--- a/Engine/Plugins/FX/Niagara/Source/Niagara/Private/NiagaraGpuReadbackManager.cpp
+++ b/Engine/Plugins/FX/Niagara/Source/Niagara/Private/NiagaraGpuReadbackManager.cpp
@@ -68,17 +68,24 @@ void FNiagaraGpuReadbackManager::WaitCompletion(FRHICommandListImmediate& RHICmd
 }
 
 void FNiagaraGpuReadbackManager::EnqueueReadback(FRHICommandList& RHICmdList, FRHIBuffer* Buffer, FCompletionCallback Callback)
+{
+	check(IsInRenderingThread());
+	
+	EnqueueReadback(RHICmdList, Buffer, 0, Buffer->GetSize(), MoveTemp(Callback));
+}
+
+void FNiagaraGpuReadbackManager::EnqueueReadback(FRHICommandList& RHICmdList, FRHIBuffer* Buffer, uint32 Offset, uint32 NumBytes, FCompletionCallback Callback)
 {
 	static const FName FenceName(TEXT("NiagaraGpuReadback"));
 
 	check(IsInRenderingThread());
 
 	FPendingReadback Readback;
-	Readback.StagingBuffers.Emplace(RHICreateStagingBuffer(), Buffer->GetSize());
+	Readback.StagingBuffers.Emplace(RHICreateStagingBuffer(), NumBytes);
 	Readback.Fence = RHICreateGPUFence(FenceName);
 	Readback.Callback = Callback;
 
-	RHICmdList.CopyToStagingBuffer(Buffer, Readback.StagingBuffers[0].Key, 0, Readback.StagingBuffers[0].Value);
+	RHICmdList.CopyToStagingBuffer(Buffer, Readback.StagingBuffers[0].Key, Offset, Readback.StagingBuffers[0].Value);
 
 	Readback.Fence->Clear();
 	RHICmdList.WriteGPUFence(Readback.Fence);
diff --git a/Engine/Plugins/FX/Niagara/Source/Niagara/Private/NiagaraGpuReadbackManager.h b/Engine/Plugins/FX/Niagara/Source/Niagara/Private/NiagaraGpuReadbackManager.h
--- a/Engine/Plugins/FX/Niagara/Source/Niagara/Private/NiagaraGpuReadbackManager.h
+++ b/Engine/Plugins/FX/Niagara/Source/Niagara/Private/NiagaraGpuReadbackManager.h
@@ -36,6 +36,7 @@ public:
 
 	// Enqueue a readback of a single buffer
 	void EnqueueReadback(FRHICommandList& RHICmdList, FRHIBuffer* Buffer, FCompletionCallback Callback);
+	void EnqueueReadback(FRHICommandList& RHICmdList, FRHIBuffer* Buffer, uint32 Offset, uint32 NumBytes, FCompletionCallback Callback);
 
 	// Enqueue a readback of multiple buffers
 	void EnqueueReadbacks(FRHICommandList& RHICmdList, TConstArrayView<FRHIBuffer*> Buffers, FCompletionCallback Callback);
