- Allow allocations from the count manager outside of the general compute dispatch

diff --git a/Engine/Plugins/FX/Niagara/Source/Niagara/Classes/NiagaraGPUInstanceCountManager.h b/Engine/Plugins/FX/Niagara/Source/Niagara/Classes/NiagaraGPUInstanceCountManager.h
--- a/Engine/Plugins/FX/Niagara/Source/Niagara/Classes/NiagaraGPUInstanceCountManager.h
+++ b/Engine/Plugins/FX/Niagara/Source/Niagara/Classes/NiagaraGPUInstanceCountManager.h
@@ -55,7 +55,7 @@ public:
 		FORCEINLINE bool IsValid() const { return Offset != INDEX_NONE; }
 	};
 
-	FNiagaraGPUInstanceCountManager();
+	FNiagaraGPUInstanceCountManager(ERHIFeatureLevel::Type FeatureLevel);
 	~FNiagaraGPUInstanceCountManager();
 
 	// Init resource for the first time.
@@ -63,20 +63,22 @@ public:
 	// Free resources.
 	void ReleaseRHI();
 
-	FRWBuffer& GetInstanceCountBuffer()
+	const FRWBuffer& GetInstanceCountBuffer() const
 	{
 		check(UsedInstanceCounts <= AllocatedInstanceCounts); // Can't resize after after the buffer gets bound.
 		return CountBuffer;
 	}
 
+	/** Acquire an entry from the free list, assumes this comes from being presized. */
+	uint32 AcquireEntry();
+	/** Acquire an entry, this will either come from the free list or reallocate the buffer. */
+	uint32 AcquireOrAllocateEntry(FRHICommandListImmediate& RHICmdList);
+
 	/** Free the entry and reset it to INDEX_NONE if valid. */
 	void FreeEntry(uint32& BufferOffset);
-
 	/** Free and array of entries, you are expected to reset or change to INDEX_NONE. */
 	void FreeEntryArray(TConstArrayView<uint32> EntryArray);
 
-	uint32 AcquireEntry();
-
 	bool CanAcquireCulledEntry() const
 	{
 		return !bAcquiredCulledCounts;
@@ -88,7 +90,7 @@ public:
 		return RequiredCulledCounts++;
 	}
 
-	FRWBuffer* AcquireCulledCountsBuffer(FRHICommandListImmediate& RHICmdList, ERHIFeatureLevel::Type FeatureLevel);
+	FRWBuffer* AcquireCulledCountsBuffer(FRHICommandListImmediate& RHICmdList);
 
 	const uint32* GetGPUReadback();
 	void ReleaseGPUReadback();
@@ -99,12 +101,12 @@ public:
 	FIndirectArgSlot AddDrawIndirect(uint32 InstanceCountBufferOffset, uint32 NumIndicesPerInstance, uint32 StartIndexLocation, bool bIsInstancedStereoEnabled, bool bCulled, ENiagaraGpuComputeTickStage::Type ReadyTickStage);
 
 	// Resize instance count and draw indirect buffers to ensure it is big enough to hold all draw indirect args.
-	void ResizeBuffers(FRHICommandListImmediate& RHICmdList, ERHIFeatureLevel::Type FeatureLevel, int32 ReservedInstanceCounts);
+	void ResizeBuffers(FRHICommandListImmediate& RHICmdList, int32 ReservedInstanceCounts);
 
 	void FlushIndirectArgsPool();
 
 	// Generate the draw indirect buffers, and reset all release counts.
-	void UpdateDrawIndirectBuffers(FNiagaraGpuComputeDispatchInterface* ComputeDispatchInterface, FRHICommandList& RHICmdList, ERHIFeatureLevel::Type FeatureLevel, ENiagaraGPUCountUpdatePhase::Type CountPhase);
+	void UpdateDrawIndirectBuffers(FNiagaraGpuComputeDispatchInterface* ComputeDispatchInterface, FRHICommandList& RHICmdList, ENiagaraGPUCountUpdatePhase::Type CountPhase);
 
 	static const ERHIAccess kCountBufferDefaultState;
 
@@ -124,6 +126,8 @@ protected:
 
 	void ReleaseCounts();
 
+	ERHIFeatureLevel::Type FeatureLevel;
+
 	/** The current used instance counts allocated from FNiagaraDataBuffer::AllocateGPU() */
 	int32 UsedInstanceCounts = 0;
 	/** The allocated instance counts in CountBuffer */
diff --git a/Engine/Plugins/FX/Niagara/Source/Niagara/Private/NiagaraGPUInstanceCountManager.cpp b/Engine/Plugins/FX/Niagara/Source/Niagara/Private/NiagaraGPUInstanceCountManager.cpp
--- a/Engine/Plugins/FX/Niagara/Source/Niagara/Private/NiagaraGPUInstanceCountManager.cpp
+++ b/Engine/Plugins/FX/Niagara/Source/Niagara/Private/NiagaraGPUInstanceCountManager.cpp
@@ -20,6 +20,14 @@ static FAutoConsoleVariableRef CVarNiagaraMinGPUInstanceCount(
 	ECVF_Default
 );
 
+int32 GNiagaraGPUCountManagerAllocateIncrement = 64;
+static FAutoConsoleVariableRef CVarNiagaraGPUCountManagerAllocateIncrement(
+	TEXT("Niagara.GPUCountManager.AllocateIncrement"),
+	GNiagaraGPUCountManagerAllocateIncrement,
+	TEXT("If we run out of space for allocations this is how many allocate rather than a single entry. (default=64)"),
+	ECVF_Default
+);
+
 int32 GNiagaraMinCulledGPUInstanceCount = 2048;
 static FAutoConsoleVariableRef CVarNiagaraMinCulledGPUInstanceCount(
 	TEXT("Niagara.MinCulledGPUInstanceCount"),
@@ -94,7 +102,8 @@ DECLARE_CYCLE_STAT(TEXT("GPU Readback Lock"), STAT_NiagaraGPUReadbackLock, STATG
 
 const ERHIAccess FNiagaraGPUInstanceCountManager::kCountBufferDefaultState = ERHIAccess::SRVMask | ERHIAccess::CopySrc;
 
-FNiagaraGPUInstanceCountManager::FNiagaraGPUInstanceCountManager()
+FNiagaraGPUInstanceCountManager::FNiagaraGPUInstanceCountManager(ERHIFeatureLevel::Type InFeatureLevel)
+	: FeatureLevel(InFeatureLevel)
 {
 }
 
@@ -156,6 +165,27 @@ uint32 FNiagaraGPUInstanceCountManager::AcquireEntry()
 	}
 }
 
+uint32 FNiagaraGPUInstanceCountManager::AcquireOrAllocateEntry(FRHICommandListImmediate& RHICmdList)
+{
+	checkSlow(IsInRenderingThread());
+
+	// Free entries?
+	if (FreeEntries.Num())
+	{
+		return FreeEntries.Pop();
+	}
+	else if (UsedInstanceCounts < AllocatedInstanceCounts)
+	{
+		return UsedInstanceCounts++;
+	}
+
+	// We need to resize
+	ResizeBuffers(RHICmdList, AllocatedInstanceCounts + GNiagaraGPUCountManagerAllocateIncrement);
+
+	check(UsedInstanceCounts < AllocatedInstanceCounts);
+	return UsedInstanceCounts++;
+}
+
 void FNiagaraGPUInstanceCountManager::FreeEntry(uint32& BufferOffset)
 {
 	checkSlow(IsInRenderingThread());
@@ -188,7 +218,7 @@ void FNiagaraGPUInstanceCountManager::FreeEntryArray(TConstArrayView<uint32> Ent
 	}
 }
 
-FRWBuffer* FNiagaraGPUInstanceCountManager::AcquireCulledCountsBuffer(FRHICommandListImmediate& RHICmdList, ERHIFeatureLevel::Type FeatureLevel)
+FRWBuffer* FNiagaraGPUInstanceCountManager::AcquireCulledCountsBuffer(FRHICommandListImmediate& RHICmdList)
 {
 	if (RequiredCulledCounts > 0)
 	{
@@ -220,7 +250,7 @@ FRWBuffer* FNiagaraGPUInstanceCountManager::AcquireCulledCountsBuffer(FRHIComman
 	return nullptr;
 }
 
-void FNiagaraGPUInstanceCountManager::ResizeBuffers(FRHICommandListImmediate& RHICmdList, ERHIFeatureLevel::Type FeatureLevel, int32 ReservedInstanceCounts)
+void FNiagaraGPUInstanceCountManager::ResizeBuffers(FRHICommandListImmediate& RHICmdList, int32 ReservedInstanceCounts)
 {
 	const int32 RequiredInstanceCounts = UsedInstanceCounts + FMath::Max<int32>(ReservedInstanceCounts - FreeEntries.Num(), 0);
 	if (RequiredInstanceCounts > 0)
@@ -352,7 +382,7 @@ FNiagaraGPUInstanceCountManager::FIndirectArgSlot FNiagaraGPUInstanceCountManage
 
 ////////////////////////////////////////////////////////////////////////////////////////////////////////////////////////////////////////////////////////////////////////////////////////////////////////
 
-void FNiagaraGPUInstanceCountManager::UpdateDrawIndirectBuffers(FNiagaraGpuComputeDispatchInterface* ComputeDispatchInterface, FRHICommandList& RHICmdList, ERHIFeatureLevel::Type FeatureLevel, ENiagaraGPUCountUpdatePhase::Type CountPhase)
+void FNiagaraGPUInstanceCountManager::UpdateDrawIndirectBuffers(FNiagaraGpuComputeDispatchInterface* ComputeDispatchInterface, FRHICommandList& RHICmdList, ENiagaraGPUCountUpdatePhase::Type CountPhase)
 {
 	// Anything to process?
 	TArray<FNiagaraDrawIndirectArgGenTaskInfo>& ArgTasks = DrawIndirectArgGenTasks[CountPhase];
diff --git a/Engine/Plugins/FX/Niagara/Source/Niagara/Private/NiagaraGpuComputeDispatch.cpp b/Engine/Plugins/FX/Niagara/Source/Niagara/Private/NiagaraGpuComputeDispatch.cpp
--- a/Engine/Plugins/FX/Niagara/Source/Niagara/Private/NiagaraGpuComputeDispatch.cpp
+++ b/Engine/Plugins/FX/Niagara/Source/Niagara/Private/NiagaraGpuComputeDispatch.cpp
@@ -182,7 +182,7 @@ FNiagaraGpuComputeDispatch::FNiagaraGpuComputeDispatch(ERHIFeatureLevel::Type In
 			GPUSortManager->PostPreRenderEvent.AddLambda(
 				[this](FRHICommandListImmediate& RHICmdList)
 				{
-					GPUInstanceCounterManager.UpdateDrawIndirectBuffers(this, RHICmdList, FeatureLevel, ENiagaraGPUCountUpdatePhase::PreOpaque);
+					GPUInstanceCounterManager.UpdateDrawIndirectBuffers(this, RHICmdList, ENiagaraGPUCountUpdatePhase::PreOpaque);
 #if WITH_MGPU
 					// For PreInitViews we actually need to broadcast here and not in ExecuteAll since
 					// this is the last place the instance count buffer is written to.
@@ -198,7 +198,7 @@ FNiagaraGpuComputeDispatch::FNiagaraGpuComputeDispatch(ERHIFeatureLevel::Type In
 			(
 				[this](FRHICommandListImmediate& RHICmdList)
 				{
-					GPUInstanceCounterManager.UpdateDrawIndirectBuffers(this, RHICmdList, FeatureLevel, ENiagaraGPUCountUpdatePhase::PostOpaque);
+					GPUInstanceCounterManager.UpdateDrawIndirectBuffers(this, RHICmdList, ENiagaraGPUCountUpdatePhase::PostOpaque);
 				}
 			);
 		}
@@ -483,10 +483,10 @@ void FNiagaraGpuComputeDispatch::ProcessPendingTicksFlush(FRHICommandListImmedia
 			FRDGBuilder GraphBuilder(RHICmdList);
 			CreateSystemTextures(GraphBuilder);
 			PreInitViews(GraphBuilder, bAllowGPUParticleUpdate);
-			GPUInstanceCounterManager.UpdateDrawIndirectBuffers(this, RHICmdList, FeatureLevel, ENiagaraGPUCountUpdatePhase::PreOpaque);
+			GPUInstanceCounterManager.UpdateDrawIndirectBuffers(this, RHICmdList, ENiagaraGPUCountUpdatePhase::PreOpaque);
 			PostInitViews(GraphBuilder, DummyViews, bAllowGPUParticleUpdate);
 			PostRenderOpaque(GraphBuilder, DummyViews, bAllowGPUParticleUpdate);
-			GPUInstanceCounterManager.UpdateDrawIndirectBuffers(this, RHICmdList, FeatureLevel, ENiagaraGPUCountUpdatePhase::PostOpaque);
+			GPUInstanceCounterManager.UpdateDrawIndirectBuffers(this, RHICmdList, ENiagaraGPUCountUpdatePhase::PostOpaque);
 			GraphBuilder.Execute();
 			break;
 		}
@@ -828,7 +828,7 @@ void FNiagaraGpuComputeDispatch::UpdateInstanceCountManager(FRHICommandListImmed
 				}
 			}
 		}
-		GPUInstanceCounterManager.ResizeBuffers(RHICmdList, FeatureLevel, TotalDispatchCount);
+		GPUInstanceCounterManager.ResizeBuffers(RHICmdList, TotalDispatchCount);
 	}
 
 	// Consume any pending readbacks that are ready
@@ -1680,7 +1680,7 @@ void FNiagaraGpuComputeDispatch::PreInitViews(FRDGBuilder& GraphBuilder, bool bA
 	}
 	else
 	{
-		GPUInstanceCounterManager.ResizeBuffers(GraphBuilder.RHICmdList, FeatureLevel,  0);
+		GPUInstanceCounterManager.ResizeBuffers(GraphBuilder.RHICmdList, 0);
 		FinishDispatches();
 	}
 }
@@ -1915,7 +1915,7 @@ void FNiagaraGpuComputeDispatch::GenerateSortKeys(FRHICommandListImmediate& RHIC
 	TShaderMapRef<FNiagaraSortKeyGenCS> SortKeyGenCS(GetGlobalShaderMap(FeatureLevel), SortPermutationVector);
 	TShaderMapRef<FNiagaraSortKeyGenCS> SortAndCullKeyGenCS(GetGlobalShaderMap(FeatureLevel), SortAndCullPermutationVector);
 
-	FRWBuffer* CulledCountsBuffer = GPUInstanceCounterManager.AcquireCulledCountsBuffer(RHICmdList, FeatureLevel);
+	FRWBuffer* CulledCountsBuffer = GPUInstanceCounterManager.AcquireCulledCountsBuffer(RHICmdList);
 
 	FNiagaraSortKeyGenCS::FParameters Params;
 	Params.SortKeyMask = KeyGenInfo.SortKeyParams.X;
diff --git a/Engine/Plugins/FX/Niagara/Source/Niagara/Private/NiagaraGpuComputeDispatchInterface.cpp b/Engine/Plugins/FX/Niagara/Source/Niagara/Private/NiagaraGpuComputeDispatchInterface.cpp
--- a/Engine/Plugins/FX/Niagara/Source/Niagara/Private/NiagaraGpuComputeDispatchInterface.cpp
+++ b/Engine/Plugins/FX/Niagara/Source/Niagara/Private/NiagaraGpuComputeDispatchInterface.cpp
@@ -12,6 +12,7 @@
 FNiagaraGpuComputeDispatchInterface::FNiagaraGpuComputeDispatchInterface(EShaderPlatform InShaderPlatform, ERHIFeatureLevel::Type InFeatureLevel)
 	: ShaderPlatform(InShaderPlatform)
 	, FeatureLevel(InFeatureLevel)
+	, GPUInstanceCounterManager(InFeatureLevel)
 {
 }
 
diff --git a/Engine/Plugins/FX/Niagara/Source/Niagara/Public/NiagaraGpuComputeDispatchInterface.h b/Engine/Plugins/FX/Niagara/Source/Niagara/Public/NiagaraGpuComputeDispatchInterface.h
--- a/Engine/Plugins/FX/Niagara/Source/Niagara/Public/NiagaraGpuComputeDispatchInterface.h
+++ b/Engine/Plugins/FX/Niagara/Source/Niagara/Public/NiagaraGpuComputeDispatchInterface.h
@@ -49,6 +49,7 @@ public:
 
 	/** Get access to the instance count manager. */
 	FORCEINLINE FNiagaraGPUInstanceCountManager& GetGPUInstanceCounterManager() { check(IsInRenderingThread()); return GPUInstanceCounterManager; }
+	FORCEINLINE const FNiagaraGPUInstanceCountManager& GetGPUInstanceCounterManager() const { check(IsInRenderingThread()); return GPUInstanceCounterManager; }
 
 #if NIAGARA_COMPUTEDEBUG_ENABLED
 	/** Get access to Niagara's GpuComputeDebug */
diff --git a/Engine/Plugins/FX/Niagara/Source/Niagara/Public/NiagaraRenderer.h b/Engine/Plugins/FX/Niagara/Source/Niagara/Public/NiagaraRenderer.h
--- a/Engine/Plugins/FX/Niagara/Source/Niagara/Public/NiagaraRenderer.h
+++ b/Engine/Plugins/FX/Niagara/Source/Niagara/Public/NiagaraRenderer.h
@@ -107,10 +107,10 @@ public:
 	virtual void GetDynamicRayTracingInstances(FRayTracingMaterialGatheringContext& Context, TArray<FRayTracingInstance>& OutRayTracingInstances, const FNiagaraSceneProxy* Proxy) {}
 #endif
 
-	FORCEINLINE static FRHIShaderResourceView* GetSrvOrDefaultFloat(FRWBuffer& RWBuffer) { return RWBuffer.SRV.IsValid() ? (FRHIShaderResourceView*)RWBuffer.SRV : GetDummyFloatBuffer(); }
-	FORCEINLINE static FRHIShaderResourceView* GetSrvOrDefaultHalf(FRWBuffer& RWBuffer) { return RWBuffer.SRV.IsValid() ? (FRHIShaderResourceView*)RWBuffer.SRV : GetDummyHalfBuffer(); }
-	FORCEINLINE static FRHIShaderResourceView* GetSrvOrDefaultUInt(FRWBuffer& RWBuffer) { return RWBuffer.SRV.IsValid() ? (FRHIShaderResourceView*)RWBuffer.SRV : GetDummyUIntBuffer(); }
-	FORCEINLINE static FRHIShaderResourceView* GetSrvOrDefaultInt(FRWBuffer& RWBuffer) { return RWBuffer.SRV.IsValid() ? (FRHIShaderResourceView*)RWBuffer.SRV : GetDummyIntBuffer(); }
+	FORCEINLINE static FRHIShaderResourceView* GetSrvOrDefaultFloat(const FRWBuffer& RWBuffer) { return RWBuffer.SRV.IsValid() ? (FRHIShaderResourceView*)RWBuffer.SRV : GetDummyFloatBuffer(); }
+	FORCEINLINE static FRHIShaderResourceView* GetSrvOrDefaultHalf(const FRWBuffer& RWBuffer) { return RWBuffer.SRV.IsValid() ? (FRHIShaderResourceView*)RWBuffer.SRV : GetDummyHalfBuffer(); }
+	FORCEINLINE static FRHIShaderResourceView* GetSrvOrDefaultUInt(const FRWBuffer& RWBuffer) { return RWBuffer.SRV.IsValid() ? (FRHIShaderResourceView*)RWBuffer.SRV : GetDummyUIntBuffer(); }
+	FORCEINLINE static FRHIShaderResourceView* GetSrvOrDefaultInt(const FRWBuffer& RWBuffer) { return RWBuffer.SRV.IsValid() ? (FRHIShaderResourceView*)RWBuffer.SRV : GetDummyIntBuffer(); }
 
 	FORCEINLINE static FRHIShaderResourceView* GetSrvOrDefaultFloat(FGlobalDynamicReadBuffer::FAllocation& Allocation) { return Allocation.IsValid() ? (FRHIShaderResourceView*)Allocation.SRV : GetDummyFloatBuffer(); }
 	FORCEINLINE static FRHIShaderResourceView* GetSrvOrDefaultHalf(FGlobalDynamicReadBuffer::FAllocation& Allocation) { return Allocation.IsValid() ? (FRHIShaderResourceView*)Allocation.SRV : GetDummyHalfBuffer(); }
