- Static analysis fixes

#rnx

diff --git a/Engine/Plugins/FX/Niagara/Source/Niagara/Private/DataInterface/NiagaraDataInterfaceStaticMesh.cpp b/Engine/Plugins/FX/Niagara/Source/Niagara/Private/DataInterface/NiagaraDataInterfaceStaticMesh.cpp
--- a/Engine/Plugins/FX/Niagara/Source/Niagara/Private/DataInterface/NiagaraDataInterfaceStaticMesh.cpp
+++ b/Engine/Plugins/FX/Niagara/Source/Niagara/Private/DataInterface/NiagaraDataInterfaceStaticMesh.cpp
@@ -2223,10 +2223,6 @@ void UNiagaraDataInterfaceStaticMesh::GetVMExternalFunction(const FVMExternalFun
 	{
 		OutFunc = FVMExternalFunction::CreateLambda([this](FVectorVMExternalFunctionContext& Context) { VMGetTriPosition_Deprecated<FNDITransformHandler>(Context); });
 	}
-	else if (BindingInfo.Name == Deprecated_GetTriPositionWSName)
-	{
-		OutFunc = FVMExternalFunction::CreateLambda([this](FVectorVMExternalFunctionContext& Context) { VMGetTriPosition_Deprecated<FNDITransformHandler>(Context); });
-	}
 	else if (BindingInfo.Name == Deprecated_GetTriTangentsName)
 	{	
 		OutFunc = FVMExternalFunction::CreateLambda([this](FVectorVMExternalFunctionContext& Context) { VMGetTriangleTangentBasis_Deprecated<FNDITransformHandlerNoop>(Context); });
