- Move Simple Counter in private and extend functionality to support more operations and GPU

diff --git a/Engine/Plugins/FX/Niagara/Shaders/Private/NiagaraClearCounts.usf b/Engine/Plugins/FX/Niagara/Shaders/Private/NiagaraClearCounts.usf
new file mode 100644
--- /dev/null
+++ b/Engine/Plugins/FX/Niagara/Shaders/Private/NiagaraClearCounts.usf
@@ -0,0 +1,23 @@
+// Copyright Epic Games, Inc. All Rights Reserved.
+
+#include "/Engine/Private/Common.ush"
+
+#if NIAGARA_CLEAR_COUNTS_INT_CS
+
+RWBuffer<int>	CountBuffer;
+Buffer<int>		CountsAndValuesBuffer;
+uint			NumCounts;
+
+[numthreads(64, 1, 1)]
+void ClearCountsIntCS(uint3 DispatchThreadId : SV_DispatchThreadID)
+{
+	uint Index = DispatchThreadId.x;
+	if(Index < NumCounts)
+	{
+		int ClearIndex = CountsAndValuesBuffer[Index * 2];
+		int ClearValue = CountsAndValuesBuffer[(Index * 2) + 1];
+		CountBuffer[ClearIndex] = ClearValue;
+	}
+}
+
+#endif
diff --git a/Engine/Plugins/FX/Niagara/Shaders/Private/NiagaraDataInterfaceSimpleCounterTemplate.ush b/Engine/Plugins/FX/Niagara/Shaders/Private/NiagaraDataInterfaceSimpleCounterTemplate.ush
new file mode 100644
--- /dev/null
+++ b/Engine/Plugins/FX/Niagara/Shaders/Private/NiagaraDataInterfaceSimpleCounterTemplate.ush
@@ -0,0 +1,45 @@
+// Copyright Epic Games, Inc. All Rights Reserved.
+// RWInstanceCounts come from NiagaraEmitterInstanceSimulation.ush as we use counts from that buffer
+
+int CountOffset_{ParameterName};
+
+void Get_{ParameterName}(out int CurrValue)
+{
+	CurrValue = RWInstanceCounts[CountOffset_{ParameterName}];
+}
+
+void Exchange_{ParameterName}_UEImpureCall(in bool bExecute, in int NewValue, out int PrevValue)
+{
+	if ( bExecute )
+	{
+		InterlockedExchange(RWInstanceCounts[CountOffset_{ParameterName}], NewValue, PrevValue);
+	}
+	else
+	{
+		PrevValue = RWInstanceCounts[CountOffset_{ParameterName}];
+	}
+}
+
+void Add_{ParameterName}_UEImpureCall(in bool bExecute, in int Amount, out int PrevValue, out int CurrValue)
+{
+	if ( bExecute )
+	{
+		InterlockedAdd(RWInstanceCounts[CountOffset_{ParameterName}], Amount, PrevValue);
+		CurrValue = PrevValue + Amount;
+	}
+	else
+	{
+		PrevValue = RWInstanceCounts[CountOffset_{ParameterName}];
+		CurrValue = PrevValue;
+	}
+}
+
+void Increment_{ParameterName}_UEImpureCall(in bool bExecute, out int PrevValue, out int CurrValue)
+{
+	Add_{ParameterName}_UEImpureCall(bExecute, 1, PrevValue, CurrValue);
+}
+
+void Decrement_{ParameterName}_UEImpureCall(in bool bExecute, out int PrevValue, out int CurrValue)
+{
+	Add_{ParameterName}_UEImpureCall(bExecute, -1, PrevValue, CurrValue);
+}
diff --git a/Engine/Plugins/FX/Niagara/Source/Niagara/Classes/NiagaraDataInterface.h b/Engine/Plugins/FX/Niagara/Source/Niagara/Classes/NiagaraDataInterface.h
--- a/Engine/Plugins/FX/Niagara/Source/Niagara/Classes/NiagaraDataInterface.h
+++ b/Engine/Plugins/FX/Niagara/Source/Niagara/Classes/NiagaraDataInterface.h
@@ -742,6 +742,7 @@ struct FNDIInputParam
 	VectorVM::FExternalFuncInputHandler<T> Data;
 	FORCEINLINE FNDIInputParam(FVectorVMExternalFunctionContext& Context) : Data(Context) {}
 	FORCEINLINE T GetAndAdvance() { return Data.GetAndAdvance(); }
+	FORCEINLINE bool IsConstant() const { return Data.IsConstant(); }
 };
 
 template<>
@@ -750,6 +751,7 @@ struct FNDIInputParam<FNiagaraBool>
 	VectorVM::FExternalFuncInputHandler<FNiagaraBool> Data;
 	FORCEINLINE FNDIInputParam(FVectorVMExternalFunctionContext& Context) : Data(Context) {}
 	FORCEINLINE bool GetAndAdvance() { return Data.GetAndAdvance().GetValue(); }
+	FORCEINLINE bool IsConstant() const { return Data.IsConstant(); }
 };
 
 template<>
@@ -758,6 +760,7 @@ struct FNDIInputParam<bool>
 	VectorVM::FExternalFuncInputHandler<FNiagaraBool> Data;
 	FORCEINLINE FNDIInputParam(FVectorVMExternalFunctionContext& Context) : Data(Context) {}
 	FORCEINLINE bool GetAndAdvance() { return Data.GetAndAdvance().GetValue(); }
+	FORCEINLINE bool IsConstant() const { return Data.IsConstant(); }
 };
 
 template<>
@@ -767,6 +770,7 @@ struct FNDIInputParam<FVector2D>
 	VectorVM::FExternalFuncInputHandler<float> Y;
 	FORCEINLINE FNDIInputParam(FVectorVMExternalFunctionContext& Context) : X(Context), Y(Context) {}
 	FORCEINLINE FVector2D GetAndAdvance() { return FVector2D(X.GetAndAdvance(), Y.GetAndAdvance()); }
+	FORCEINLINE bool IsConstant() const { return X.IsConstant() && Y.IsConstant(); }
 };
 
 template<>
@@ -777,10 +781,9 @@ struct FNDIInputParam<FVector3f>
 	VectorVM::FExternalFuncInputHandler<float> Z;
 	FNDIInputParam(FVectorVMExternalFunctionContext& Context) : X(Context), Y(Context), Z(Context) {}
 	FORCEINLINE FVector3f GetAndAdvance() { return FVector3f(X.GetAndAdvance(), Y.GetAndAdvance(), Z.GetAndAdvance()); }
+	FORCEINLINE bool IsConstant() const { return X.IsConstant() && Y.IsConstant() && Z.IsConstant(); }
 };
 
-
-
 template<>
 struct FNDIInputParam<FVector4f>
 {
@@ -790,6 +793,7 @@ struct FNDIInputParam<FVector4f>
 	VectorVM::FExternalFuncInputHandler<float> W;
 	FORCEINLINE FNDIInputParam(FVectorVMExternalFunctionContext& Context) : X(Context), Y(Context), Z(Context), W(Context) {}
 	FORCEINLINE FVector4f GetAndAdvance() { return FVector4f(X.GetAndAdvance(), Y.GetAndAdvance(), Z.GetAndAdvance(), W.GetAndAdvance()); }
+	FORCEINLINE bool IsConstant() const { return X.IsConstant() && Y.IsConstant() && Z.IsConstant() && W.IsConstant(); }
 };
 
 template<>
@@ -801,6 +805,7 @@ struct FNDIInputParam<FQuat4f>
 	VectorVM::FExternalFuncInputHandler<float> W;
 	FORCEINLINE FNDIInputParam(FVectorVMExternalFunctionContext& Context) : X(Context), Y(Context), Z(Context), W(Context) {}
 	FORCEINLINE FQuat4f GetAndAdvance() { return FQuat4f(X.GetAndAdvance(), Y.GetAndAdvance(), Z.GetAndAdvance(), W.GetAndAdvance()); }
+	FORCEINLINE bool IsConstant() const { return X.IsConstant() && Y.IsConstant() && Z.IsConstant() && W.IsConstant(); }
 };
 
 template<>
@@ -812,6 +817,7 @@ struct FNDIInputParam<FLinearColor>
 	VectorVM::FExternalFuncInputHandler<float> A;
 	FORCEINLINE FNDIInputParam(FVectorVMExternalFunctionContext& Context) : R(Context), G(Context), B(Context), A(Context) {}
 	FORCEINLINE FLinearColor GetAndAdvance() { return FLinearColor(R.GetAndAdvance(), G.GetAndAdvance(), B.GetAndAdvance(), A.GetAndAdvance()); }
+	FORCEINLINE bool IsConstant() const { return R.IsConstant() && G.IsConstant() && B.IsConstant() && A.IsConstant(); }
 };
 
 template<>
@@ -821,6 +827,7 @@ struct FNDIInputParam<FNiagaraID>
 	VectorVM::FExternalFuncInputHandler<int32> AcquireTag;
 	FORCEINLINE FNDIInputParam(FVectorVMExternalFunctionContext& Context) : Index(Context), AcquireTag(Context) {}
 	FORCEINLINE FNiagaraID GetAndAdvance() { return FNiagaraID(Index.GetAndAdvance(), AcquireTag.GetAndAdvance()); }
+	FORCEINLINE bool IsConstant() const { return Index.IsConstant() && AcquireTag.IsConstant(); }
 };
 
 //Helper to deal with types with potentially several output registers.
diff --git a/Engine/Plugins/FX/Niagara/Source/Niagara/Classes/NiagaraDataInterfaceSimpleCounter.h b/Engine/Plugins/FX/Niagara/Source/Niagara/Classes/NiagaraDataInterfaceSimpleCounter.h
deleted file mode 100644
diff --git a/Engine/Plugins/FX/Niagara/Source/Niagara/Private/DataInterface/NiagaraDataInterfaceSimpleCounter.cpp b/Engine/Plugins/FX/Niagara/Source/Niagara/Private/DataInterface/NiagaraDataInterfaceSimpleCounter.cpp
new file mode 100644
--- /dev/null
+++ b/Engine/Plugins/FX/Niagara/Source/Niagara/Private/DataInterface/NiagaraDataInterfaceSimpleCounter.cpp
@@ -0,0 +1,626 @@
+// Copyright Epic Games, Inc. All Rights Reserved.
+
+#include "NiagaraDataInterfaceSimpleCounter.h"
+#include "NiagaraClearCounts.h"
+#include "NiagaraGpuComputeDispatchInterface.h"
+#include "NiagaraGpuReadbackManager.h"
+#include "NiagaraSystemInstance.h"
+
+#include "Internationalization/Internationalization.h"
+#include "Async/Async.h"
+#include "ShaderParameterUtils.h"
+#include "ShaderCompilerCore.h"
+
+#define LOCTEXT_NAMESPACE "NiagaraDataInterfaceSimpleCounter"
+
+//////////////////////////////////////////////////////////////////////////
+
+namespace NDISimpleCounterLocal
+{
+	static const TCHAR* TemplateShaderFile = TEXT("/Plugin/FX/Niagara/Private/NiagaraDataInterfaceSimpleCounterTemplate.ush");
+	static const FName NAME_GetNextValue_Deprecated(TEXT("GetNextValue"));
+	static const FName NAME_Get(TEXT("Get"));
+	static const FName NAME_Exchange(TEXT("Exchange"));
+	static const FName NAME_Add(TEXT("Add"));
+	static const FName NAME_Increment(TEXT("Increment"));
+	static const FName NAME_Decrement(TEXT("Decrement"));
+}
+
+struct FNDISimpleCounterInstanceData_RenderThread
+{
+	TOptional<int32>	CountValue;
+	uint32				CountOffset = INDEX_NONE;
+};
+
+struct FNDISimpleCounterInstanceData_GameThread
+{
+	bool				bModified = true;
+	std::atomic<int32>	Counter;
+};
+
+struct FNDISimpleCounterProxy : public FNiagaraDataInterfaceProxy
+{
+	FNDISimpleCounterProxy(UNiagaraDataInterfaceSimpleCounter* Owner)
+		: WeakOwner(Owner)
+		, GpuSyncMode(Owner->GpuSyncMode)
+	{
+	}
+
+	virtual void ConsumePerInstanceDataFromGameThread(void* PerInstanceData, const FNiagaraSystemInstanceID& Instance) override {}
+	virtual int32 PerInstanceDataPassedToRenderThreadSize() const override { return 0; }
+
+	virtual void PreStage(FRHICommandList& RHICmdList, const FNiagaraDataInterfaceStageArgs& Context) override
+	{
+		if (FNDISimpleCounterInstanceData_RenderThread* InstanceData = PerInstanceData_RenderThread.Find(Context.SystemInstanceID) )
+		{
+			if ( InstanceData->CountValue.IsSet() )
+			{
+				//-OPT: We could push this into the count manager and batch set as part of the clear process
+				const FNiagaraGPUInstanceCountManager& CounterManager = Context.ComputeDispatchInterface->GetGPUInstanceCounterManager();
+				const FRWBuffer& CountBuffer = CounterManager.GetInstanceCountBuffer();
+
+				const TPair<uint32, int32> DataToClear(InstanceData->CountOffset, InstanceData->CountValue.GetValue());
+				RHICmdList.Transition(FRHITransitionInfo(CountBuffer.UAV, ERHIAccess::UAVCompute, ERHIAccess::UAVCompute));
+				NiagaraClearCounts::ClearCountsInt(RHICmdList, CountBuffer.UAV, MakeArrayView(&DataToClear, 1) );
+				RHICmdList.Transition(FRHITransitionInfo(CountBuffer.UAV, ERHIAccess::UAVCompute, ERHIAccess::UAVCompute));
+				InstanceData->CountValue.Reset();
+			}
+		}
+	}
+
+	virtual void PostSimulate(FRHICommandList& RHICmdList, const FNiagaraDataInterfaceArgs& Context) override
+	{
+		if (FNiagaraUtilities::ShouldSyncGpuToCpu(GpuSyncMode))
+		{
+			if (FNDISimpleCounterInstanceData_RenderThread* InstanceData = PerInstanceData_RenderThread.Find(Context.SystemInstanceID))
+			{
+				const FNiagaraGPUInstanceCountManager& CountManager = Context.ComputeDispatchInterface->GetGPUInstanceCounterManager();
+				FNiagaraGpuReadbackManager* ReadbackManager = Context.ComputeDispatchInterface->GetGpuReadbackManager();
+				ReadbackManager->EnqueueReadback(
+					RHICmdList,
+					CountManager.GetInstanceCountBuffer().Buffer,
+					InstanceData->CountOffset, sizeof(int32),
+					[SystemInstanceID=Context.SystemInstanceID, WeakOwner=WeakOwner, Proxy=this](TConstArrayView<TPair<void*, uint32>> ReadbackData)
+					{
+						const int32 CounterValue = *reinterpret_cast<const int32*>(ReadbackData[0].Key);
+						AsyncTask(
+							ENamedThreads::GameThread,
+							[SystemInstanceID, CounterValue, WeakOwner, Proxy]()
+							{
+								// FNiagaraDataInterfaceProxy do not outlive UNiagaraDataInterface so if our Object is valid so is the proxy
+								// Equally because we do not share instance IDs (monotonically increasing number) we won't ever stomp something that has 'gone away'
+								if ( WeakOwner.Get() )
+								{
+									if ( FNDISimpleCounterInstanceData_GameThread* InstanceData_GT = Proxy->PerInstanceData_GameThread.FindRef(SystemInstanceID) )
+									{
+										InstanceData_GT->Counter = CounterValue;
+									}
+								}
+							}
+						);
+					}
+				);
+			}
+		}
+	}
+
+	TWeakObjectPtr<UNiagaraDataInterfaceSimpleCounter> WeakOwner;
+	ENiagaraGpuSyncMode GpuSyncMode = ENiagaraGpuSyncMode::None;
+
+	TMap<FNiagaraSystemInstanceID, FNDISimpleCounterInstanceData_RenderThread>	PerInstanceData_RenderThread;
+	TMap<FNiagaraSystemInstanceID, FNDISimpleCounterInstanceData_GameThread*>	PerInstanceData_GameThread;
+};
+
+struct FNDISimpleCounterCS : public FNiagaraDataInterfaceParametersCS
+{
+	DECLARE_TYPE_LAYOUT(FNDISimpleCounterCS, NonVirtual);
+
+public:
+	void Bind(const FNiagaraDataInterfaceGPUParamInfo& ParameterInfo, const class FShaderParameterMap& ParameterMap)
+	{
+		CountOffsetParam.Bind(ParameterMap, *(TEXT("CountOffset_") + ParameterInfo.DataInterfaceHLSLSymbol));
+	}
+
+	void Set(FRHICommandList& RHICmdList, const FNiagaraDataInterfaceSetArgs& Context) const
+	{
+		using namespace NDISimpleCounterLocal;
+
+		check(IsInRenderingThread());
+
+		if (CountOffsetParam.IsBound())
+		{
+			FNDISimpleCounterProxy* DIProxy = static_cast<FNDISimpleCounterProxy*>(Context.DataInterface);
+			FNDISimpleCounterInstanceData_RenderThread* InstanceData = &DIProxy->PerInstanceData_RenderThread.FindOrAdd(Context.SystemInstanceID);
+			check(InstanceData->CountOffset != INDEX_NONE);
+
+			FRHIComputeShader* ComputeShaderRHI = Context.Shader.GetComputeShader();
+			SetShaderValue(RHICmdList, ComputeShaderRHI, CountOffsetParam, InstanceData->CountOffset);
+		}
+	}
+
+private:
+	LAYOUT_FIELD(FShaderParameter, CountOffsetParam);
+};
+
+IMPLEMENT_TYPE_LAYOUT(FNDISimpleCounterCS);
+IMPLEMENT_NIAGARA_DI_PARAMETER(UNiagaraDataInterfaceSimpleCounter, FNDISimpleCounterCS);
+
+//////////////////////////////////////////////////////////////////////////
+
+UNiagaraDataInterfaceSimpleCounter::UNiagaraDataInterfaceSimpleCounter(FObjectInitializer const& ObjectInitializer)
+	: Super(ObjectInitializer)
+{
+	Proxy.Reset(new FNDISimpleCounterProxy(this));
+}
+
+void UNiagaraDataInterfaceSimpleCounter::PostInitProperties()
+{
+	Super::PostInitProperties();
+
+	//Can we register data interfaces as regular types and fold them into the FNiagaraVariable framework for UI and function calls etc?
+	if (HasAnyFlags(RF_ClassDefaultObject))
+	{
+		ENiagaraTypeRegistryFlags Flags = ENiagaraTypeRegistryFlags::AllowAnyVariable | ENiagaraTypeRegistryFlags::AllowParameter;
+		FNiagaraTypeRegistry::Register(FNiagaraTypeDefinition(GetClass()), Flags);
+	}
+}
+
+#if WITH_EDITOR
+void UNiagaraDataInterfaceSimpleCounter::PostEditChangeProperty(struct FPropertyChangedEvent& PropertyChangedEvent)
+{
+	Super::PostEditChangeProperty(PropertyChangedEvent);
+
+	// Ensure proxy properties are up to date
+	FNDISimpleCounterProxy* Proxy_GT = GetProxyAs<FNDISimpleCounterProxy>();
+	Proxy_GT->GpuSyncMode = GpuSyncMode;
+}
+#endif
+
+bool UNiagaraDataInterfaceSimpleCounter::Equals(const UNiagaraDataInterface* Other) const
+{
+	if (!Super::Equals(Other))
+	{
+		return false;
+	}
+
+	auto* OtherTyped = CastChecked<const UNiagaraDataInterfaceSimpleCounter>(Other);
+	return
+		OtherTyped->GpuSyncMode == GpuSyncMode &&
+		OtherTyped->InitialValue == InitialValue;
+}
+
+bool UNiagaraDataInterfaceSimpleCounter::CopyToInternal(UNiagaraDataInterface* Destination) const
+{
+	if (!Super::CopyToInternal(Destination))
+	{
+		return false;
+	}
+
+	auto* DestinationTyped = CastChecked<UNiagaraDataInterfaceSimpleCounter>(Destination);
+	DestinationTyped->GpuSyncMode = GpuSyncMode;
+	DestinationTyped->InitialValue = InitialValue;
+
+	// Ensure proxy properties are up to date
+	FNDISimpleCounterProxy* DestinationProxy = DestinationTyped->GetProxyAs<FNDISimpleCounterProxy>();
+	DestinationProxy->GpuSyncMode = GpuSyncMode;
+
+	return true;
+}
+
+bool UNiagaraDataInterfaceSimpleCounter::InitPerInstanceData(void* PerInstanceData, FNiagaraSystemInstance* SystemInstance)
+{
+	FNDISimpleCounterInstanceData_GameThread* InstanceData_GT = new(PerInstanceData) FNDISimpleCounterInstanceData_GameThread();
+	InstanceData_GT->Counter = InitialValue;
+
+	if ( IsUsedWithGPUEmitter() )
+	{
+		FNDISimpleCounterProxy* Proxy_GT = GetProxyAs<FNDISimpleCounterProxy>();
+		Proxy_GT->PerInstanceData_GameThread.Add(SystemInstance->GetId(), InstanceData_GT);
+
+		ENQUEUE_RENDER_COMMAND(FNDISimpleCounter_RemoveProxy)
+		(
+			[Proxy_RT=Proxy_GT, InstanceID_RT=SystemInstance->GetId(), ComputeInterface_RT=SystemInstance->GetComputeDispatchInterface(), InitialValue_RT=InitialValue](FRHICommandListImmediate& RHICmdList)
+			{
+				FNDISimpleCounterInstanceData_RenderThread* InstanceData = &Proxy_RT->PerInstanceData_RenderThread.Add(InstanceID_RT);
+				InstanceData->CountValue = InitialValue_RT;
+				InstanceData->CountOffset = ComputeInterface_RT->GetGPUInstanceCounterManager().AcquireOrAllocateEntry(RHICmdList);
+			}
+		);
+	}
+	return true;
+}
+
+void UNiagaraDataInterfaceSimpleCounter::DestroyPerInstanceData(void* PerInstanceData, FNiagaraSystemInstance* SystemInstance)
+{
+	auto* InstanceData = reinterpret_cast<FNDISimpleCounterInstanceData_GameThread*>(PerInstanceData);
+	InstanceData->~FNDISimpleCounterInstanceData_GameThread();
+
+	if ( IsUsedWithGPUEmitter() )
+	{
+		FNDISimpleCounterProxy* Proxy_GT = GetProxyAs<FNDISimpleCounterProxy>();
+		Proxy_GT->PerInstanceData_GameThread.Remove(SystemInstance->GetId());
+
+		ENQUEUE_RENDER_COMMAND(FNDISimpleCounter_RemoveProxy)
+		(
+			[Proxy_RT=Proxy_GT, InstanceID_RT=SystemInstance->GetId(), ComputeInterface_RT=SystemInstance->GetComputeDispatchInterface()](FRHICommandListImmediate& RHICmdList)
+			{
+				if ( FNDISimpleCounterInstanceData_RenderThread* InstanceData = Proxy_RT->PerInstanceData_RenderThread.Find(InstanceID_RT) )
+				{
+					if ( ensure(InstanceData->CountOffset != INDEX_NONE) )
+					{
+						ComputeInterface_RT->GetGPUInstanceCounterManager().FreeEntry(InstanceData->CountOffset);
+						InstanceData->CountOffset = INDEX_NONE;
+					}
+					Proxy_RT->PerInstanceData_RenderThread.Remove(InstanceID_RT);
+				}
+			}
+		);
+	}
+}
+
+int32 UNiagaraDataInterfaceSimpleCounter::PerInstanceDataSize() const
+{
+	return sizeof(FNDISimpleCounterInstanceData_GameThread);
+}
+
+void UNiagaraDataInterfaceSimpleCounter::GetFunctions(TArray<FNiagaraFunctionSignature>& OutFunctions)
+{
+	using namespace NDISimpleCounterLocal;
+
+	// Deprecated function
+	{
+		FNiagaraFunctionSignature& Sig = OutFunctions.AddDefaulted_GetRef();
+		Sig.Name = NAME_GetNextValue_Deprecated;
+		Sig.bMemberFunction = true;
+		Sig.bRequiresContext = false;
+		Sig.bSoftDeprecatedFunction = true;
+		Sig.Inputs.Add(FNiagaraVariable(FNiagaraTypeDefinition(GetClass()), TEXT("Counter")));
+		Sig.Outputs.Add(FNiagaraVariable(FNiagaraTypeDefinition::GetIntDef(), TEXT("Value")));
+		Sig.SetDescription(LOCTEXT("UNiagaraDataInterfaceSimpleCounter_GetNextValue", "Increment the internal counter. Note that it is possible for this counter to roll over eventually, so make sure that your particles do not live extremely long lifetimes."));
+	}
+
+	{
+		FNiagaraFunctionSignature & Sig = OutFunctions.AddDefaulted_GetRef();
+		Sig.Name = NAME_Get;
+		Sig.bMemberFunction = true;
+		Sig.bRequiresContext = false;
+		Sig.Inputs.Add(FNiagaraVariable(FNiagaraTypeDefinition(GetClass()), TEXT("Counter")));
+		Sig.Outputs.Add(FNiagaraVariable(FNiagaraTypeDefinition::GetIntDef(), TEXT("Current Value")));
+		Sig.SetDescription(LOCTEXT("GetDesc", "Gets the current value of the counter."));
+	}
+	{
+		FNiagaraFunctionSignature& Sig = OutFunctions.AddDefaulted_GetRef();
+		Sig.Name = NAME_Exchange;
+		Sig.bMemberFunction = true;
+		Sig.bRequiresContext = false;
+		Sig.bRequiresExecPin = true;
+		Sig.Inputs.Add(FNiagaraVariable(FNiagaraTypeDefinition(GetClass()), TEXT("Counter")));
+		Sig.Inputs.Add_GetRef(FNiagaraVariable(FNiagaraTypeDefinition::GetBoolDef(), TEXT("Execute"))).SetValue(true);
+		Sig.Inputs.Add(FNiagaraVariable(FNiagaraTypeDefinition::GetIntDef(), TEXT("New Value")));
+		Sig.Outputs.Add(FNiagaraVariable(FNiagaraTypeDefinition::GetIntDef(), TEXT("Previous Value")));
+		Sig.SetDescription(LOCTEXT("ExchangeDesc", "Exchanges the current value with the new one."));
+	}
+	{
+		FNiagaraFunctionSignature& Sig = OutFunctions.AddDefaulted_GetRef();
+		Sig.Name = NAME_Add;
+		Sig.bMemberFunction = true;
+		Sig.bRequiresContext = false;
+		Sig.bRequiresExecPin = true;
+		Sig.Inputs.Add(FNiagaraVariable(FNiagaraTypeDefinition(GetClass()), TEXT("Counter")));
+		Sig.Inputs.Add_GetRef(FNiagaraVariable(FNiagaraTypeDefinition::GetBoolDef(), TEXT("Execute"))).SetValue(true);
+		Sig.Inputs.Add(FNiagaraVariable(FNiagaraTypeDefinition::GetIntDef(), TEXT("Amount")));
+		Sig.Outputs.Add(FNiagaraVariable(FNiagaraTypeDefinition::GetIntDef(), TEXT("Previous Value")));
+		Sig.Outputs.Add(FNiagaraVariable(FNiagaraTypeDefinition::GetIntDef(), TEXT("Current Value")));
+		Sig.SetDescription(LOCTEXT("AddDesc", "Adds the Amount to the counter."));
+	}
+	{
+		FNiagaraFunctionSignature& Sig = OutFunctions.AddDefaulted_GetRef();
+		Sig.Name = NAME_Increment;
+		Sig.bMemberFunction = true;
+		Sig.bRequiresContext = false;
+		Sig.bRequiresExecPin = true;
+		Sig.Inputs.Add(FNiagaraVariable(FNiagaraTypeDefinition(GetClass()), TEXT("Counter")));
+		Sig.Inputs.Add_GetRef(FNiagaraVariable(FNiagaraTypeDefinition::GetBoolDef(), TEXT("Execute"))).SetValue(true);
+		Sig.Outputs.Add(FNiagaraVariable(FNiagaraTypeDefinition::GetIntDef(), TEXT("Previous Value")));
+		Sig.Outputs.Add(FNiagaraVariable(FNiagaraTypeDefinition::GetIntDef(), TEXT("Current Value")));
+		Sig.SetDescription(LOCTEXT("IncrementDesc", "Increments the counter by 1."));
+	}
+	{
+		FNiagaraFunctionSignature& Sig = OutFunctions.AddDefaulted_GetRef();
+		Sig.Name = NAME_Decrement;
+		Sig.bMemberFunction = true;
+		Sig.bRequiresContext = false;
+		Sig.bRequiresExecPin = true;
+		Sig.Inputs.Add(FNiagaraVariable(FNiagaraTypeDefinition(GetClass()), TEXT("Counter")));
+		Sig.Inputs.Add_GetRef(FNiagaraVariable(FNiagaraTypeDefinition::GetBoolDef(), TEXT("Execute"))).SetValue(true);
+		Sig.Outputs.Add(FNiagaraVariable(FNiagaraTypeDefinition::GetIntDef(), TEXT("Previous Value")));
+		Sig.Outputs.Add(FNiagaraVariable(FNiagaraTypeDefinition::GetIntDef(), TEXT("Current Value")));
+		Sig.SetDescription(LOCTEXT("DecrementDesc", "Decrements the counter by 1."));
+	}
+}
+
+void UNiagaraDataInterfaceSimpleCounter::GetVMExternalFunction(const FVMExternalFunctionBindingInfo& BindingInfo, void* InstanceData, FVMExternalFunction& OutFunc)
+{
+	using namespace NDISimpleCounterLocal;
+
+	if ( BindingInfo.Name == NAME_GetNextValue_Deprecated )
+	{
+		OutFunc = FVMExternalFunction::CreateUObject(this, &UNiagaraDataInterfaceSimpleCounter::GetNextValue_Deprecated);
+	}
+	else if (BindingInfo.Name == NAME_Get)
+	{
+		OutFunc = FVMExternalFunction::CreateUObject(this, &UNiagaraDataInterfaceSimpleCounter::VMGet);
+	}
+	else if (BindingInfo.Name == NAME_Exchange)
+	{
+		OutFunc = FVMExternalFunction::CreateUObject(this, &UNiagaraDataInterfaceSimpleCounter::VMExchange);
+	}
+	else if (BindingInfo.Name == NAME_Add)
+	{
+		OutFunc = FVMExternalFunction::CreateUObject(this, &UNiagaraDataInterfaceSimpleCounter::VMAdd);
+	}
+	else if (BindingInfo.Name == NAME_Increment)
+	{
+		OutFunc = FVMExternalFunction::CreateUObject(this, &UNiagaraDataInterfaceSimpleCounter::VMIncrement);
+	}
+	else if (BindingInfo.Name == NAME_Decrement)
+	{
+		OutFunc = FVMExternalFunction::CreateUObject(this, &UNiagaraDataInterfaceSimpleCounter::VMDecrement);
+	}
+}
+
+#if WITH_EDITORONLY_DATA
+void UNiagaraDataInterfaceSimpleCounter::GetParameterDefinitionHLSL(const FNiagaraDataInterfaceGPUParamInfo& ParamInfo, FString& OutHLSL)
+{
+	TMap<FString, FStringFormatArg> TemplateArgs =
+	{
+		{TEXT("ParameterName"),	ParamInfo.DataInterfaceHLSLSymbol},
+	};
+
+	FString TemplateFile;
+	LoadShaderSourceFile(NDISimpleCounterLocal::TemplateShaderFile, EShaderPlatform::SP_PCD3D_SM5, &TemplateFile, nullptr);
+	OutHLSL += FString::Format(*TemplateFile, TemplateArgs);
+}
+
+bool UNiagaraDataInterfaceSimpleCounter::GetFunctionHLSL(const FNiagaraDataInterfaceGPUParamInfo& ParamInfo, const FNiagaraDataInterfaceGeneratedFunction& FunctionInfo, int FunctionInstanceIndex, FString& OutHLSL)
+{
+	using namespace NDISimpleCounterLocal;
+
+	if ((FunctionInfo.DefinitionName == NAME_Get) ||
+		(FunctionInfo.DefinitionName == NAME_Exchange) ||
+		(FunctionInfo.DefinitionName == NAME_Add) ||
+		(FunctionInfo.DefinitionName == NAME_Increment) ||
+		(FunctionInfo.DefinitionName == NAME_Decrement) )
+	{
+		return true;
+	}
+
+	// Invalid function
+	return true;
+}
+
+bool UNiagaraDataInterfaceSimpleCounter::AppendCompileHash(FNiagaraCompileHashVisitor* InVisitor) const
+{
+	bool bSuccess = Super::AppendCompileHash(InVisitor);
+	FSHAHash Hash = GetShaderFileHash(NDISimpleCounterLocal::TemplateShaderFile, EShaderPlatform::SP_PCD3D_SM5);
+	InVisitor->UpdateString(TEXT("NiagaraDataInterfaceSimpleCounterTemplateHLSLSource"), Hash.ToString());
+	return bSuccess;
+}
+#endif
+
+void UNiagaraDataInterfaceSimpleCounter::PushToRenderThreadImpl()
+{
+	TArray<TPair<FNiagaraSystemInstanceID, int32>, TInlineAllocator<4>> DataToPush;
+
+	FNDISimpleCounterProxy* Proxy_GT = GetProxyAs<FNDISimpleCounterProxy>();
+	for ( auto it=Proxy_GT->PerInstanceData_GameThread.CreateIterator(); it; ++it )
+	{
+		if ( it.Value()->bModified )
+		{
+			it.Value()->bModified = false;
+			DataToPush.Emplace(it.Key(), it.Value()->Counter);
+		}
+	}
+
+	if (DataToPush.Num() > 0)
+	{
+		ENQUEUE_RENDER_COMMAND(FNDISimpleCounter_PushToRender)
+		(
+			[Proxy_RT=GetProxyAs<FNDISimpleCounterProxy>(), DataToPush_RT=MoveTemp(DataToPush)](FRHICommandListImmediate& RHICmdList)
+			{
+				for (const auto& Pair : DataToPush_RT)
+				{
+					if ( FNDISimpleCounterInstanceData_RenderThread* InstanceData_RT=Proxy_RT->PerInstanceData_RenderThread.Find(Pair.Key) )
+					{
+						InstanceData_RT->CountValue = Pair.Value;
+					}
+				}
+			}
+		);
+	}
+}
+
+void UNiagaraDataInterfaceSimpleCounter::VMGet(FVectorVMExternalFunctionContext& Context)
+{
+	VectorVM::FUserPtrHandler<FNDISimpleCounterInstanceData_GameThread> InstanceData(Context);
+	FNDIOutputParam<int32> OutValue(Context);
+
+	const int32 CurrValue = InstanceData->Counter;
+	for (int32 i = 0; i < Context.GetNumInstances(); ++i)
+	{
+		OutValue.SetAndAdvance(CurrValue);
+	}
+}
+
+void UNiagaraDataInterfaceSimpleCounter::VMExchange(FVectorVMExternalFunctionContext& Context)
+{	
+	VectorVM::FUserPtrHandler<FNDISimpleCounterInstanceData_GameThread> InstanceData(Context);
+	FNDIInputParam<bool> InExecute(Context);
+	FNDIInputParam<int32> InValue(Context);
+	FNDIOutputParam<int32> OutPreviousValue(Context);
+
+	for (int32 i = 0; i < Context.GetNumInstances(); ++i)
+	{
+		const bool bExecute = InExecute.GetAndAdvance();
+		const int32 NewValue = InValue.GetAndAdvance();
+		if (bExecute)
+		{
+			const int32 PrevValue = InstanceData->Counter.exchange(NewValue);
+			OutPreviousValue.SetAndAdvance(PrevValue);
+		}
+		else
+		{
+			OutPreviousValue.SetAndAdvance(InstanceData->Counter);
+		}
+	}
+
+	if (FNiagaraUtilities::ShouldSyncCpuToGpu(GpuSyncMode))
+	{
+		InstanceData->bModified = true;
+		MarkRenderDataDirty();
+	}
+}
+
+void UNiagaraDataInterfaceSimpleCounter::VMAdd(FVectorVMExternalFunctionContext& Context)
+{
+	VectorVM::FUserPtrHandler<FNDISimpleCounterInstanceData_GameThread> InstanceData(Context);
+	FNDIInputParam<bool> InExecute(Context);
+	FNDIInputParam<int32> InValue(Context);
+	FNDIOutputParam<int32> OutPreviousValue(Context);
+	FNDIOutputParam<int32> OutCurrentValue(Context);
+
+	for (int32 i=0; i < Context.GetNumInstances(); ++i)
+	{
+		const bool bExecute = InExecute.GetAndAdvance();
+		const int32 Value = InValue.GetAndAdvance();
+		if (bExecute)
+		{
+			const int32 PrevValue = InstanceData->Counter.fetch_add(Value);
+			OutPreviousValue.SetAndAdvance(PrevValue);
+			OutCurrentValue.SetAndAdvance(PrevValue + Value);
+		}
+		else
+		{
+			const int32 CurrValue = InstanceData->Counter;
+			OutPreviousValue.SetAndAdvance(CurrValue);
+			OutCurrentValue.SetAndAdvance(CurrValue);
+		}
+	}
+
+	if (FNiagaraUtilities::ShouldSyncCpuToGpu(GpuSyncMode))
+	{
+		InstanceData->bModified = true;
+		MarkRenderDataDirty();
+	}
+}
+
+void UNiagaraDataInterfaceSimpleCounter::VMIncrement(FVectorVMExternalFunctionContext& Context)
+{
+	VectorVM::FUserPtrHandler<FNDISimpleCounterInstanceData_GameThread> InstanceData(Context);
+	FNDIInputParam<bool> InExecute(Context);
+	FNDIOutputParam<int32> OutPreviousValue(Context);
+	FNDIOutputParam<int32> OutCurrentValue(Context);
+
+	if (InExecute.IsConstant())
+	{
+		const bool bIncrement = InExecute.GetAndAdvance();
+		int32 PrevValue = bIncrement ? InstanceData->Counter.fetch_add(Context.GetNumInstances()) : InstanceData->Counter.load();
+		for (int32 i = 0; i < Context.GetNumInstances(); ++i)
+		{
+			const int32 CurrValue = PrevValue + (bIncrement ? 1 : 0);
+			OutPreviousValue.SetAndAdvance(PrevValue);
+			OutCurrentValue.SetAndAdvance(CurrValue);
+			PrevValue = CurrValue;
+		}
+	}
+	else
+	{
+		for (int32 i=0; i < Context.GetNumInstances(); ++i)
+		{
+			const bool bExecute = InExecute.GetAndAdvance();
+			if (bExecute)
+			{
+				const int32 PrevValue = InstanceData->Counter.fetch_add(1);
+				OutPreviousValue.SetAndAdvance(PrevValue);
+				OutCurrentValue.SetAndAdvance(PrevValue + 1);
+			}
+			else
+			{
+				const int32 CurrValue = InstanceData->Counter;
+				OutPreviousValue.SetAndAdvance(CurrValue);
+				OutCurrentValue.SetAndAdvance(CurrValue);
+			}
+		}
+	}
+
+	if (FNiagaraUtilities::ShouldSyncCpuToGpu(GpuSyncMode))
+	{
+		InstanceData->bModified = true;
+		MarkRenderDataDirty();
+	}
+}
+
+void UNiagaraDataInterfaceSimpleCounter::VMDecrement(FVectorVMExternalFunctionContext& Context)
+{
+	VectorVM::FUserPtrHandler<FNDISimpleCounterInstanceData_GameThread> InstanceData(Context);
+	FNDIInputParam<bool> InExecute(Context);
+	FNDIOutputParam<int32> OutPreviousValue(Context);
+	FNDIOutputParam<int32> OutCurrentValue(Context);
+
+	if (InExecute.IsConstant())
+	{
+		const bool bDecrement = InExecute.GetAndAdvance();
+		int32 PrevValue = bDecrement ? InstanceData->Counter.fetch_sub(Context.GetNumInstances()) : InstanceData->Counter.load();
+		for (int32 i = 0; i < Context.GetNumInstances(); ++i)
+		{
+			const int32 CurrValue = PrevValue - (bDecrement ? 1 : 0);
+			OutPreviousValue.SetAndAdvance(PrevValue);
+			OutCurrentValue.SetAndAdvance(CurrValue);
+			PrevValue = CurrValue;
+		}
+	}
+	else
+	{
+		for (int32 i = 0; i < Context.GetNumInstances(); ++i)
+		{
+			const bool bExecute = InExecute.GetAndAdvance();
+			if (bExecute)
+			{
+				const int32 PrevValue = InstanceData->Counter.fetch_sub(1);
+				OutPreviousValue.SetAndAdvance(PrevValue);
+				OutCurrentValue.SetAndAdvance(PrevValue - 1);
+			}
+			else
+			{
+				const int32 CurrValue = InstanceData->Counter;
+				OutPreviousValue.SetAndAdvance(CurrValue);
+				OutCurrentValue.SetAndAdvance(CurrValue);
+			}
+		}
+	}
+
+	if (FNiagaraUtilities::ShouldSyncCpuToGpu(GpuSyncMode))
+	{
+		InstanceData->bModified = true;
+		MarkRenderDataDirty();
+	}
+}
+
+void UNiagaraDataInterfaceSimpleCounter::GetNextValue_Deprecated(FVectorVMExternalFunctionContext& Context)
+{
+	VectorVM::FUserPtrHandler<FNDISimpleCounterInstanceData_GameThread> InstanceData(Context);
+	FNDIOutputParam<int32> OutValue(Context);
+
+	for (int32 i = 0; i < Context.GetNumInstances(); ++i)
+	{
+		OutValue.SetAndAdvance(InstanceData->Counter.fetch_add(1));
+	}
+
+	if (FNiagaraUtilities::ShouldSyncCpuToGpu(GpuSyncMode))
+	{
+		InstanceData->bModified = true;
+		MarkRenderDataDirty();
+	}
+}
+
+#undef LOCTEXT_NAMESPACE
diff --git a/Engine/Plugins/FX/Niagara/Source/Niagara/Private/DataInterface/NiagaraDataInterfaceSimpleCounter.h b/Engine/Plugins/FX/Niagara/Source/Niagara/Private/DataInterface/NiagaraDataInterfaceSimpleCounter.h
new file mode 100644
--- /dev/null
+++ b/Engine/Plugins/FX/Niagara/Source/Niagara/Private/DataInterface/NiagaraDataInterfaceSimpleCounter.h
@@ -0,0 +1,65 @@
+// Copyright Epic Games, Inc. All Rights Reserved.
+#pragma once
+
+#include "NiagaraCommon.h"
+#include "NiagaraShared.h"
+#include "VectorVM.h"
+#include "NiagaraDataInterface.h"
+#include "HAL/PlatformAtomics.h"
+#include "NiagaraDataInterfaceSimpleCounter.generated.h"
+
+/**
+Thread safe counter starts at the initial value on start / reset.
+When operating between CPU & GPU ensure you set the appropriate sync mode.
+*/
+UCLASS(EditInlineNew, Category = "Counting", meta = (DisplayName = "Simple Counter"))
+class NIAGARA_API UNiagaraDataInterfaceSimpleCounter : public UNiagaraDataInterface
+{
+	GENERATED_UCLASS_BODY()
+public:
+	DECLARE_NIAGARA_DI_PARAMETER();
+
+	// UObject Interface
+	virtual void PostInitProperties() override;
+#if WITH_EDITOR
+	virtual void PostEditChangeProperty(struct FPropertyChangedEvent& PropertyChangedEvent) override;
+#endif
+	// UObject Interface End
+
+	// UNiagaraDataInterface Interface Begin
+	virtual bool CanExecuteOnTarget(ENiagaraSimTarget Target) const override { return true; }
+
+	virtual bool Equals(const UNiagaraDataInterface* Other) const override;
+	virtual bool CopyToInternal(UNiagaraDataInterface* Destination) const override;
+
+	virtual bool InitPerInstanceData(void* PerInstanceData, FNiagaraSystemInstance* SystemInstance) override;
+	virtual void DestroyPerInstanceData(void* PerInstanceData, FNiagaraSystemInstance* SystemInstance) override;
+	virtual int32 PerInstanceDataSize() const override;
+
+	virtual void GetFunctions(TArray<FNiagaraFunctionSignature>& OutFunctions) override;
+	virtual void GetVMExternalFunction(const FVMExternalFunctionBindingInfo& BindingInfo, void* InstanceData, FVMExternalFunction &OutFunc) override;
+#if WITH_EDITORONLY_DATA
+	virtual void GetParameterDefinitionHLSL(const FNiagaraDataInterfaceGPUParamInfo& ParamInfo, FString& OutHLSL) override;
+	virtual bool GetFunctionHLSL(const FNiagaraDataInterfaceGPUParamInfo& ParamInfo, const FNiagaraDataInterfaceGeneratedFunction& FunctionInfo, int FunctionInstanceIndex, FString& OutHLSL) override;
+	virtual bool AppendCompileHash(FNiagaraCompileHashVisitor* InVisitor) const override;
+#endif
+	virtual void PushToRenderThreadImpl() override;
+	// UNiagaraDataInterface Interface End
+
+	// VM functions
+	void VMGet(FVectorVMExternalFunctionContext& Context);
+	void VMExchange(FVectorVMExternalFunctionContext& Context);
+	void VMAdd(FVectorVMExternalFunctionContext& Context);
+	void VMIncrement(FVectorVMExternalFunctionContext& Context);
+	void VMDecrement(FVectorVMExternalFunctionContext& Context);
+
+	void GetNextValue_Deprecated(FVectorVMExternalFunctionContext& Context);
+
+	/** Select how we should synchronize the counter between Cpu & Gpu */
+	UPROPERTY(EditAnywhere, Category = "Simple Counter")
+	ENiagaraGpuSyncMode GpuSyncMode = ENiagaraGpuSyncMode::None;
+
+	/** This is the value the counter will have when the instance is reset / created. */
+	UPROPERTY(EditAnywhere, Category = "Simple Counter")
+	int32 InitialValue = 0;
+};
diff --git a/Engine/Plugins/FX/Niagara/Source/Niagara/Private/NiagaraDataInterfaceSimpleCounter.cpp b/Engine/Plugins/FX/Niagara/Source/Niagara/Private/NiagaraDataInterfaceSimpleCounter.cpp
deleted file mode 100644
diff --git a/Engine/Plugins/FX/Niagara/Source/Niagara/Public/NiagaraCommon.h b/Engine/Plugins/FX/Niagara/Source/Niagara/Public/NiagaraCommon.h
--- a/Engine/Plugins/FX/Niagara/Source/Niagara/Public/NiagaraCommon.h
+++ b/Engine/Plugins/FX/Niagara/Source/Niagara/Public/NiagaraCommon.h
@@ -85,6 +85,19 @@ enum class ENiagaraGpuBufferFormat : uint8
 	Max UMETA(Hidden),
 };
 
+UENUM()
+enum class ENiagaraGpuSyncMode
+{
+	/** Data will not be automatically pushed and could diverge between Cpu & Gpu. */
+	None,
+	/** Cpu modifications will be pushed to the Gpu. */
+	SyncCpuToGpu,
+	/** Gpu will continuously push back to the Cpu, this will incur a performance penalty. */
+	SyncGpuToCpu,
+	/** Gpu will continuously push back to the Cpu and Cpu modifications will be pushed to the Gpu. */
+	SyncBoth,
+};
+
 UENUM()
 enum class ENiagaraMipMapGeneration : uint8
 {
@@ -1211,6 +1224,18 @@ namespace FNiagaraUtilities
 	// Helper function to detect if SRVs are always created for buffers or not
 	bool AreBufferSRVsAlwaysCreated(EShaderPlatform ShaderPlatform);
 
+	// Helper function to determine if we should sync data from CPU to GPU
+	FORCEINLINE bool ShouldSyncCpuToGpu(ENiagaraGpuSyncMode SyncMode)
+	{
+		return SyncMode == ENiagaraGpuSyncMode::SyncBoth || SyncMode == ENiagaraGpuSyncMode::SyncCpuToGpu;
+	}
+
+	// Helper function to determine if we should sync data from GPU to CPU
+	FORCEINLINE bool ShouldSyncGpuToCpu(ENiagaraGpuSyncMode SyncMode)
+	{
+		return SyncMode == ENiagaraGpuSyncMode::SyncBoth || SyncMode == ENiagaraGpuSyncMode::SyncGpuToCpu;
+	}
+
 	ENiagaraCompileUsageStaticSwitch NIAGARA_API ConvertScriptUsageToStaticSwitchUsage(ENiagaraScriptUsage ScriptUsage);
 	ENiagaraScriptContextStaticSwitch NIAGARA_API ConvertScriptUsageToStaticSwitchContext(ENiagaraScriptUsage ScriptUsage);
 	
diff --git a/Engine/Plugins/FX/Niagara/Source/NiagaraShader/Private/NiagaraClearCounts.cpp b/Engine/Plugins/FX/Niagara/Source/NiagaraShader/Private/NiagaraClearCounts.cpp
new file mode 100644
--- /dev/null
+++ b/Engine/Plugins/FX/Niagara/Source/NiagaraShader/Private/NiagaraClearCounts.cpp
@@ -0,0 +1,73 @@
+// Copyright Epic Games, Inc. All Rights Reserved.
+
+#include "NiagaraClearCounts.h"
+
+#include "RHI.h"
+#include "ScreenRendering.h"
+#include "CommonRenderResources.h"
+#include "Modules/ModuleManager.h"
+
+#include "RenderGraphBuilder.h"
+#include "Runtime/Renderer/Private/ScreenPass.h"
+
+//////////////////////////////////////////////////////////////////////////
+
+class NIAGARASHADER_API FNiagaraClearCountsIntCS : public FGlobalShader
+{
+public:
+	DECLARE_SHADER_TYPE(FNiagaraClearCountsIntCS, Global);
+	SHADER_USE_PARAMETER_STRUCT(FNiagaraClearCountsIntCS, FGlobalShader);
+
+	static constexpr uint32 ThreadGroupSize = 64;
+
+	static bool ShouldCompilePermutation(const FGlobalShaderPermutationParameters& Parameters)
+	{
+		return !IsMobilePlatform(Parameters.Platform);
+	}
+
+	static void ModifyCompilationEnvironment(const FGlobalShaderPermutationParameters& Parameters, FShaderCompilerEnvironment& OutEnvironment)
+	{
+		FGlobalShader::ModifyCompilationEnvironment(Parameters, OutEnvironment);
+		OutEnvironment.SetDefine(TEXT("NIAGARA_CLEAR_COUNTS_INT_CS"), 1);
+		OutEnvironment.SetDefine(TEXT("THREADGROUP_SIZE"), ThreadGroupSize);
+	}
+
+	BEGIN_SHADER_PARAMETER_STRUCT(FParameters, )
+		SHADER_PARAMETER_UAV(RWBuffer<int>,	CountBuffer)
+		SHADER_PARAMETER_SRV(Buffer<int>,	CountsAndValuesBuffer)
+		SHADER_PARAMETER(uint32,			NumCounts)
+	END_SHADER_PARAMETER_STRUCT()
+};
+
+IMPLEMENT_GLOBAL_SHADER(FNiagaraClearCountsIntCS, "/Plugin/FX/Niagara/Private/NiagaraClearCounts.usf", "ClearCountsIntCS", SF_Compute);
+
+void NiagaraClearCounts::ClearCountsInt(FRHICommandList& RHICmdList, FRHIUnorderedAccessView* UAV, TConstArrayView<TPair<uint32, int32>> IndexAndValueArray)
+{
+	check(UAV != nullptr);
+	check(IndexAndValueArray.Num() > 0);
+
+	FReadBuffer IndexAndValueBuffer;
+	IndexAndValueBuffer.Initialize(TEXT("NiagaraClearCounts"), sizeof(int32), IndexAndValueArray.Num() * 2, PF_R32_SINT);
+	{
+		const uint32 BufferSize = IndexAndValueArray.Num() * sizeof(int32) * 2;
+		void* GPUMemory = RHILockBuffer(IndexAndValueBuffer.Buffer, 0, BufferSize, RLM_WriteOnly);
+		FMemory::Memcpy(GPUMemory, IndexAndValueArray.GetData(), BufferSize);
+		RHIUnlockBuffer(IndexAndValueBuffer.Buffer);
+	}
+
+	FNiagaraClearCountsIntCS::FParameters ShaderParameters;
+	ShaderParameters.CountBuffer			= UAV;
+	ShaderParameters.CountsAndValuesBuffer	= IndexAndValueBuffer.SRV;
+	ShaderParameters.NumCounts				= IndexAndValueArray.Num();
+
+	TShaderMapRef<FNiagaraClearCountsIntCS> ComputeShader(GetGlobalShaderMap(GMaxRHIFeatureLevel));
+	FRHIComputeShader* ShaderRHI = ComputeShader.GetComputeShader();
+	const uint32 NumThreadGroups = FMath::DivideAndRoundUp<uint32>(IndexAndValueArray.Num(), FNiagaraClearCountsIntCS::ThreadGroupSize);
+
+	RHICmdList.Transition(FRHITransitionInfo(UAV, ERHIAccess::Unknown, ERHIAccess::UAVCompute));
+	RHICmdList.SetComputeShader(ShaderRHI);
+	SetShaderParameters(RHICmdList, ComputeShader, ShaderRHI, ShaderParameters);
+	RHICmdList.DispatchComputeShader(NumThreadGroups, 1, 1);
+	UnsetShaderUAVs(RHICmdList, ComputeShader, ShaderRHI);
+	RHICmdList.Transition(FRHITransitionInfo(UAV, ERHIAccess::UAVCompute, ERHIAccess::UAVCompute));
+}
diff --git a/Engine/Plugins/FX/Niagara/Source/NiagaraShader/Public/NiagaraClearCounts.h b/Engine/Plugins/FX/Niagara/Source/NiagaraShader/Public/NiagaraClearCounts.h
new file mode 100644
--- /dev/null
+++ b/Engine/Plugins/FX/Niagara/Source/NiagaraShader/Public/NiagaraClearCounts.h
@@ -0,0 +1,11 @@
+// Copyright Epic Games, Inc. All Rights Reserved.
+
+#pragma once
+
+#include "RHI.h"
+#include "CommonRenderResources.h"
+
+namespace NiagaraClearCounts
+{
+	NIAGARASHADER_API void ClearCountsInt(FRHICommandList& RHICmdList, FRHIUnorderedAccessView* UAV, TConstArrayView<TPair<uint32, int32>> IndexAndValueArray);
+}
