convex shapes use a bbox for collision and still give a warning

diff --git a/Engine/Plugins/FX/Niagara/Source/Niagara/Private/NiagaraDataInterfaceRigidMeshCollisionQuery.cpp b/Engine/Plugins/FX/Niagara/Source/Niagara/Private/NiagaraDataInterfaceRigidMeshCollisionQuery.cpp
--- a/Engine/Plugins/FX/Niagara/Source/Niagara/Private/NiagaraDataInterfaceRigidMeshCollisionQuery.cpp
+++ b/Engine/Plugins/FX/Niagara/Source/Niagara/Private/NiagaraDataInterfaceRigidMeshCollisionQuery.cpp
@@ -148,6 +148,14 @@ void GetNumPrimitives(TArray<AActor*> Actors, uint32& NumBoxes, uint32& NumSpher
 		if (BodySetup != nullptr)
 		{
 			// #todo(dmp): save static mesh component reference
+
+			for (const FKConvexElem& ConvexElem : BodySetup->AggGeom.ConvexElems)
+			{
+				if (CollisionEnabledHasPhysics(ConvexElem.GetCollisionEnabled()))
+				{
+					NumBoxes += 1;
+				}
+			}
 			for (const FKBoxElem& BoxElem : BodySetup->AggGeom.BoxElems)
 			{
 				if (CollisionEnabledHasPhysics(BoxElem.GetCollisionEnabled()))
@@ -231,10 +239,23 @@ void CreateInternalArrays(TArray<AActor*> Actors,
 						const FTransform MeshTransform = Actor->GetTransform();
 
 						for (const FKConvexElem& ConvexElem : BodySetup->AggGeom.ConvexElems)
-						{
+						{							
 							if (CollisionEnabledHasPhysics(ConvexElem.GetCollisionEnabled()))
 							{
-								UE_LOG(LogRigidMeshCollision, Warning, TEXT("Convex collision objects encountered and will be skipped on %s"), *Actor->GetName());
+								UE_LOG(LogRigidMeshCollision, Warning, TEXT("Convex collision objects encountered and will be interpreted as a bounding box on %s"), *Actor->GetName());
+
+
+								FBox BBox = ConvexElem.ElemBox;
+								FVector3f Extent = BBox.Max - BBox.Min;
+
+								OutAssetArrays->PhysicsType[BoxCount] = (ConvexElem.GetCollisionEnabled() == ECollisionEnabled::QueryAndPhysics);
+								OutAssetArrays->SourceSceneProxy[BoxCount] = StaticMeshComponent->SceneProxy;
+								
+								const FTransform ElementTransform = ConvexElem.GetTransform() * MeshTransform;
+								OutAssetArrays->ElementExtent[BoxCount] = FVector4(Extent.X, Extent.Y, Extent.Z, 0);
+								FillCurrentTransforms(ElementTransform, BoxCount, OutAssetArrays->CurrentTransform, OutAssetArrays->CurrentInverse);
+
+								FoundCollisionShapes = true;
 							}
 						}
 						for (const FKBoxElem& BoxElem : BodySetup->AggGeom.BoxElems)
@@ -339,6 +360,21 @@ void UpdateInternalArrays(const TArray<AActor*> &Actors, FNDIRigidMeshCollisionA
 			{
 				const FTransform MeshTransform = Actor->GetTransform();
 
+				for (const FKConvexElem& ConvexElem : BodySetup->AggGeom.ConvexElems)
+				{
+					if (CollisionEnabledHasPhysics(ConvexElem.GetCollisionEnabled()))
+					{
+						UE_LOG(LogRigidMeshCollision, Warning, TEXT("Convex collision objects encountered and will be interpreted as a bounding box on %s"), *Actor->GetName());
+
+
+						FBox BBox = ConvexElem.ElemBox;						
+
+						OutAssetArrays->SourceSceneProxy[BoxCount] = StaticMeshComponent->SceneProxy;
+						const FTransform ElementTransform = ConvexElem.GetTransform() * MeshTransform;
+						FillCurrentTransforms(ElementTransform, BoxCount, OutAssetArrays->CurrentTransform, OutAssetArrays->CurrentInverse);
+					}
+				}
+
 				for (const FKBoxElem& BoxElem : BodySetup->AggGeom.BoxElems)
 				{
 					if (CollisionEnabledHasPhysics(BoxElem.GetCollisionEnabled()))
