- Hide GPU profiler implementation and expose only what is neccessary outside of the plugin
- Extend GPU profiler to capture stage execution cost & pass the data to the particle perf stats & debug HUD

diff --git a/Engine/Plugins/FX/Niagara/Source/Niagara/Classes/NiagaraComputeExecutionContext.h b/Engine/Plugins/FX/Niagara/Source/Niagara/Classes/NiagaraComputeExecutionContext.h
--- a/Engine/Plugins/FX/Niagara/Source/Niagara/Classes/NiagaraComputeExecutionContext.h
+++ b/Engine/Plugins/FX/Niagara/Source/Niagara/Classes/NiagaraComputeExecutionContext.h
@@ -119,9 +119,15 @@ struct FNiagaraComputeExecutionContext
 	}  EmitterInstanceReadback;
 	
 #if !UE_BUILD_SHIPPING
+	FName GetDebugSimFName() const { return DebugSimFName; }
 	const TCHAR* GetDebugSimName() const { return *DebugSimName; }
-	void SetDebugSimName(const TCHAR* InDebugSimName) { DebugSimName = InDebugSimName; }
+	void SetDebugSimName(const TCHAR* InDebugSimName)
+	{
+		DebugSimName = InDebugSimName;
+		DebugSimFName = FName(DebugSimName);
+	}
 #else
+	FName GetDebugSimFName() const { return NAME_None; }
 	const TCHAR* GetDebugSimName() const { return TEXT(""); }
 	void SetDebugSimName(const TCHAR*) { }
 #endif
@@ -133,11 +139,11 @@ public:
 	static uint32 TickCounter;
 
 #if !UE_BUILD_SHIPPING
+	FName DebugSimFName;
 	FString DebugSimName;
 #endif
-#if STATS
-	TWeakObjectPtr<UNiagaraEmitter> EmitterPtr; // emitter pointer used to report captured gpu stats
-#endif
+	TWeakObjectPtr<class USceneComponent>	ProfilingComponentPtr;
+	TWeakObjectPtr<UNiagaraEmitter>			ProfilingEmitterPtr;
 
 	const TArray<UNiagaraDataInterface*>& GetDataInterfaces()const { return CombinedParamStore.GetDataInterfaces(); }
 
diff --git a/Engine/Plugins/FX/Niagara/Source/Niagara/Classes/NiagaraGPUProfiler.h b/Engine/Plugins/FX/Niagara/Source/Niagara/Classes/NiagaraGPUProfiler.h
deleted file mode 100644
diff --git a/Engine/Plugins/FX/Niagara/Source/Niagara/Classes/NiagaraGPUProfilerInterface.h b/Engine/Plugins/FX/Niagara/Source/Niagara/Classes/NiagaraGPUProfilerInterface.h
new file mode 100644
--- /dev/null
+++ b/Engine/Plugins/FX/Niagara/Source/Niagara/Classes/NiagaraGPUProfilerInterface.h
@@ -0,0 +1,80 @@
+// Copyright Epic Games, Inc. All Rights Reserved.
+
+#pragma once
+
+#include "NiagaraCommon.h"
+#include "NiagaraSystemGpuComputeProxy.h"
+
+//////////////////////////////////////////////////////////////////////////
+// Public API for tracking GPU time when the profiler is enabled
+struct NIAGARA_API FNiagaraGpuProfileScope
+{
+#if WITH_NIAGARA_GPU_PROFILER
+	FNiagaraGpuProfileScope(FRHICommandList& RHICmdList, const class FNiagaraGpuComputeDispatchInterface* ComputeDispatchInterface, FName StageName);
+	FNiagaraGpuProfileScope(FRHICommandList& RHICmdList, const class FNiagaraGpuComputeDispatchInterface* ComputeDispatchInterface, const struct FNiagaraGpuDispatchInstance& DispatchInstance);
+	FNiagaraGpuProfileScope(FRHICommandList& RHICmdList, const struct FNiagaraDataInterfaceArgs& Context, FName StageName);
+	FNiagaraGpuProfileScope(FRHICommandList& RHICmdList, const struct FNiagaraDataInterfaceSetArgs& Context, FName StageName);
+	FNiagaraGpuProfileScope(FRHICommandList& RHICmdList, const struct FNiagaraDataInterfaceStageArgs& Context, FName StageName);
+	~FNiagaraGpuProfileScope();
+
+private:
+	FRHICommandList& RHICmdList;
+	class FNiagaraGPUProfiler* GPUProfiler = nullptr;
+#else
+	FNiagaraGpuProfileScope(FRHICommandList& RHICmdList, const class FNiagaraGpuComputeDispatchInterface* ComputeDispatchInterface, FName StageName) {}
+	FNiagaraGpuProfileScope(FRHICommandList& RHICmdList, const class FNiagaraGpuComputeDispatchInterface* ComputeDispatchInterface, const struct FNiagaraGpuDispatchInstance& DispatchInstance) {}
+	FNiagaraGpuProfileScope(FRHICommandList& RHICmdList, const struct FNiagaraDataInterfaceArgs& Context, FName StageName) {}
+	FNiagaraGpuProfileScope(FRHICommandList& RHICmdList, const struct FNiagaraDataInterfaceSetArgs& Context, FName StageName) {}
+	FNiagaraGpuProfileScope(FRHICommandList& RHICmdList, const struct FNiagaraDataInterfaceStageArgs& Context, FName StageName) {}
+	~FNiagaraGpuProfileScope() {}
+#endif
+};
+
+#if WITH_NIAGARA_GPU_PROFILER
+//////////////////////////////////////////////////////////////////////////
+/** Results generated when the frame is ready and sent to that game thread */
+struct FNiagaraGpuFrameResults : public TSharedFromThis<FNiagaraGpuFrameResults, ESPMode::ThreadSafe>
+{
+	struct FStageResults
+	{
+		int32	NumDispatchGroups = 0;
+		int32	NumDispatches = 0;
+		uint64	DurationMicroseconds = 0;
+	};
+
+	struct FDispatchResults
+	{
+		uint32									bUniqueInstance : 1;		// Set only once for all dispatches from an instance across all ticks
+		TWeakObjectPtr<class USceneComponent>	OwnerComponent;				// Optional pointer back to owning Component
+		TWeakObjectPtr<class UNiagaraEmitter>	OwnerEmitter;				// Optional pointer back to owning Emitter
+		FName									StageName;					// Generally the simulation stage but may be a DataInterface name
+		uint64									DurationMicroseconds;		// Duration in microseconds of the dispatch
+	};
+
+	uintptr_t					OwnerContext = 0;
+	FStageResults				StageResults[ENiagaraGpuComputeTickStage::Max];
+	TArray<FDispatchResults>	DispatchResults;
+};
+
+using FNiagaraGpuFrameResultsPtr = TSharedPtr<FNiagaraGpuFrameResults, ESPMode::ThreadSafe>;
+
+//////////////////////////////////////////////////////////////////////////
+/** Public API to Niagara GPU Profiling. */
+class FNiagaraGPUProfilerInterface
+{
+	DECLARE_MULTICAST_DELEGATE_OneParam(FOnFrameResults, const FNiagaraGpuFrameResultsPtr&);
+
+public:
+	static FOnFrameResults& GetOnFrameResults_GameThread() { check(IsInGameThread()); return OnFrameResults_GameThread; }
+	static FOnFrameResults& GetOnFrameResults_RenderThread() { check(IsInRenderingThread()); return OnFrameResults_RenderThread; }
+
+protected:
+	void PostResults(const FNiagaraGpuFrameResultsPtr& FrameResults);
+
+protected:
+	static FOnFrameResults	OnFrameResults_GameThread;
+	static FOnFrameResults	OnFrameResults_RenderThread;
+};
+
+#endif //WITH_NIAGARA_GPU_PROFILER
+
diff --git a/Engine/Plugins/FX/Niagara/Source/Niagara/Private/NiagaraDataInterfaceRenderTarget2D.cpp b/Engine/Plugins/FX/Niagara/Source/Niagara/Private/NiagaraDataInterfaceRenderTarget2D.cpp
--- a/Engine/Plugins/FX/Niagara/Source/Niagara/Private/NiagaraDataInterfaceRenderTarget2D.cpp
+++ b/Engine/Plugins/FX/Niagara/Source/Niagara/Private/NiagaraDataInterfaceRenderTarget2D.cpp
@@ -14,6 +14,7 @@
 #include "NiagaraRenderer.h"
 #include "NiagaraGpuComputeDebug.h"
 #include "NiagaraGpuComputeDispatchInterface.h"
+#include "NiagaraGPUProfilerInterface.h"
 #include "NiagaraStats.h"
 
 #define LOCTEXT_NAMESPACE "NiagaraDataInterfaceRenderTarget2D"
@@ -31,6 +32,8 @@ const FName UNiagaraDataInterfaceRenderTarget2D::SetSizeFunctionName("SetRenderT
 const FName UNiagaraDataInterfaceRenderTarget2D::GetSizeFunctionName("GetRenderTargetSize");
 const FName UNiagaraDataInterfaceRenderTarget2D::LinearToIndexName("LinearToIndex");
 
+static const FName GNiagaraRenderTarget2DGenerateMipsName("RenderTarget2D::GenerateMips");
+
 FNiagaraVariableBase UNiagaraDataInterfaceRenderTarget2D::ExposedRTVar;
 
 int32 GNiagaraReleaseResourceOnRemove = false;
@@ -807,6 +810,7 @@ void FNiagaraDataInterfaceProxyRenderTarget2DProxy::PostStage(FRHICommandList& R
 		if (ProxyData->bWasWrittenTo && (ProxyData->MipMapGeneration == ENiagaraMipMapGeneration::PostStage))
 		{
 			ProxyData->bWasWrittenTo = false;
+			FNiagaraGpuProfileScope GpuProfileScope(RHICmdList, Context, GNiagaraRenderTarget2DGenerateMipsName);
 			NiagaraGenerateMips::GenerateMips(RHICmdList, ProxyData->TextureRHI, ProxyData->MipMapGenerationType);
 		}
 	}
@@ -819,6 +823,7 @@ void FNiagaraDataInterfaceProxyRenderTarget2DProxy::PostSimulate(FRHICommandList
 		if (ProxyData->bWasWrittenTo && (ProxyData->MipMapGeneration == ENiagaraMipMapGeneration::PostSimulate))
 		{
 			ProxyData->bWasWrittenTo = false;
+			FNiagaraGpuProfileScope GpuProfileScope(RHICmdList, Context, GNiagaraRenderTarget2DGenerateMipsName);
 			NiagaraGenerateMips::GenerateMips(RHICmdList, ProxyData->TextureRHI, ProxyData->MipMapGenerationType);
 		}
 
diff --git a/Engine/Plugins/FX/Niagara/Source/Niagara/Private/NiagaraDebugHud.cpp b/Engine/Plugins/FX/Niagara/Source/Niagara/Private/NiagaraDebugHud.cpp
--- a/Engine/Plugins/FX/Niagara/Source/Niagara/Private/NiagaraDebugHud.cpp
+++ b/Engine/Plugins/FX/Niagara/Source/Niagara/Private/NiagaraDebugHud.cpp
@@ -535,6 +535,90 @@ FNiagaraDebugHud::FNiagaraDebugHud(UWorld* World)
 		GDebugDrawHandle = UDebugDrawService::Register(TEXT("Particles"), FDebugDrawDelegate::CreateStatic(&FNiagaraDebugHud::DebugDrawCallback));
 	}
 	++GDebugDrawHandleUsers;
+#if WITH_NIAGARA_GPU_PROFILER
+	GpuProfilerDelegateHandle = FNiagaraGPUProfilerInterface::GetOnFrameResults_GameThread().AddLambda(
+		[&](const FNiagaraGpuFrameResultsPtr& InGpuResults)
+		{
+			FNiagaraGpuComputeDispatchInterface* DispatchInterface = FNiagaraGpuComputeDispatchInterface::Get(WeakWorld.Get());
+			if ( InGpuResults->OwnerContext == uintptr_t(DispatchInterface) )
+			{
+				GpuResults = InGpuResults;
+				GpuResultsGameFrameCounter = GFrameCounter;
+
+				for ( const auto& DispatchResult : GpuResults->DispatchResults )
+				{
+					if ( DispatchResult.OwnerEmitter.IsExplicitlyNull() )
+					{
+						FGpuUsagePerEvent& EventUsage = GpuUsagePerEvent.FindOrAdd(DispatchResult.StageName);
+						EventUsage.InstanceCount.Accumulate(GpuResultsGameFrameCounter, 1);
+						EventUsage.Microseconds.Accumulate(GpuResultsGameFrameCounter, DispatchResult.DurationMicroseconds);
+					}
+					else
+					{
+						UNiagaraEmitter* OwnerEmitter = DispatchResult.OwnerEmitter.Get();
+						UNiagaraSystem* OwnerSystem = OwnerEmitter ? OwnerEmitter->GetTypedOuter<UNiagaraSystem>() : nullptr;
+						if (OwnerSystem == nullptr)
+						{
+							continue;
+						}
+
+						FGpuUsagePerSystem& SystemUsage = GpuUsagePerSystem.FindOrAdd(OwnerSystem);
+						SystemUsage.Microseconds.Accumulate(GpuResultsGameFrameCounter, DispatchResult.DurationMicroseconds);
+
+						FGpuUsagePerEmitter& EmitterUsage = SystemUsage.Emitters.FindOrAdd(OwnerEmitter);
+						EmitterUsage.Microseconds.Accumulate(GpuResultsGameFrameCounter, DispatchResult.DurationMicroseconds);
+						EmitterUsage.InstanceCount.Accumulate(GpuResultsGameFrameCounter, 1);
+
+						FGpuUsagePerStage& StageUsage = EmitterUsage.Stages.FindOrAdd(DispatchResult.StageName);
+						StageUsage.Microseconds.Accumulate(GpuResultsGameFrameCounter, DispatchResult.DurationMicroseconds);
+						StageUsage.InstanceCount.Accumulate(GpuResultsGameFrameCounter, 1);
+
+						if ( DispatchResult.bUniqueInstance )
+						{
+							SystemUsage.InstanceCount.Accumulate(GpuResultsGameFrameCounter, 1);
+						}
+					}
+				}
+
+				// Prune data that hasn't been seen in a while
+				constexpr uint64 FramesBeforePrune = SmoothedNumFrames;
+				for ( auto SystemIt=GpuUsagePerSystem.CreateIterator(); SystemIt; ++SystemIt)
+				{
+					if ( SystemIt.Value().Microseconds.ShouldPrune(GpuResultsGameFrameCounter) )
+					{
+						SystemIt.RemoveCurrent();
+						continue;
+					}
+
+					for (auto EmitterIt=SystemIt.Value().Emitters.CreateIterator(); EmitterIt; ++EmitterIt)
+					{
+						if (EmitterIt.Value().Microseconds.ShouldPrune(GpuResultsGameFrameCounter) )
+						{
+							EmitterIt.RemoveCurrent();
+							continue;
+						}
+
+						for (auto StageIt=EmitterIt.Value().Stages.CreateIterator(); StageIt; ++StageIt)
+						{
+							if (StageIt.Value().Microseconds.ShouldPrune(GpuResultsGameFrameCounter))
+							{
+								StageIt.RemoveCurrent();
+							}
+						}
+					}
+				}
+
+				for ( auto EventIt=GpuUsagePerEvent.CreateIterator(); EventIt; ++EventIt)
+				{
+					if (EventIt.Value().Microseconds.ShouldPrune(GpuResultsGameFrameCounter))
+					{
+						EventIt.RemoveCurrent();
+					}
+				}
+			}
+		}
+	);
+#endif
 }
 
 FNiagaraDebugHud::~FNiagaraDebugHud()
@@ -547,6 +631,9 @@ FNiagaraDebugHud::~FNiagaraDebugHud()
 		UDebugDrawService::Unregister(GDebugDrawHandle);
 		GDebugDrawHandle.Reset();
 	}
+#if WITH_NIAGARA_GPU_PROFILER
+	FNiagaraGPUProfilerInterface::GetOnFrameResults_GameThread().Remove(GpuProfilerDelegateHandle);
+#endif
 }
 
 void FNiagaraDebugHud::UpdateSettings(const FNiagaraDebugHUDSettingsData& NewSettings)
@@ -957,6 +1044,14 @@ void FNiagaraDebugHud::Draw(FNiagaraWorldManager* WorldManager, UCanvas* Canvas,
 		}
 	}
 
+#if WITH_NIAGARA_GPU_PROFILER
+	// Kill GPU results that are no longer relevant
+	if ( GpuResults.IsValid() && (GFrameCounter - GpuResultsGameFrameCounter) > SmoothedNumFrames)
+	{
+		GpuResults = nullptr;
+	}
+#endif
+
 	LastDrawTime = CurrTime;
 }
 
@@ -1270,7 +1365,7 @@ void FNiagaraDebugHud::DrawOverview(class FNiagaraWorldManager* WorldManager, FC
 
 			if (Settings.PerfGraphMode == ENiagaraDebugHUDPerfGraphMode::None)
 			{
-				OverviewColumns.Emplace(TEXT("Game Thread Avg:"), LexToSanitizedString(GlobalPerfStats.Avg.Time_GT), TEXT("GT Avg"), ColumnOffset, Font, TEXT("000.0000"),
+				OverviewColumns.Emplace(TEXT("Game Thread Avg:"), LexToSanitizedString(GlobalPerfStats.Avg.Time_GT), TEXT("GT Avg (us)"), ColumnOffset, Font, TEXT("000.0000"),
 					[&DetailColor, &DetailHighlightColor, &fAdvanceHeight](FCanvas* Canvas, UFont* Font, float X, float Y, FOverviewColumn& Col, const FSystemDebugInfo& SystemInfo)
 					{
 						FLinearColor RowBGColor = SystemInfo.UniqueColor;
@@ -1280,7 +1375,7 @@ void FNiagaraDebugHud::DrawOverview(class FNiagaraWorldManager* WorldManager, FC
 						Canvas->DrawShadowedString(X, Y, *LexToSanitizedString(SystemInfo.PerfStats ? SystemInfo.PerfStats->Avg.Time_GT : 0.0), Font, RowColor);
 					});
 
-				OverviewColumns.Emplace(TEXT("Game Thread Max:"), LexToSanitizedString(GlobalPerfStats.Max.Time_GT), TEXT("GT Max"), ColumnOffset, Font, TEXT("000.0000"),
+				OverviewColumns.Emplace(TEXT("Game Thread Max:"), LexToSanitizedString(GlobalPerfStats.Max.Time_GT), TEXT("GT Max (us)"), ColumnOffset, Font, TEXT("000.0000"),
 					[&DetailColor, &DetailHighlightColor, &fAdvanceHeight](FCanvas* Canvas, UFont* Font, float X, float Y, FOverviewColumn& Col, const FSystemDebugInfo& SystemInfo)
 					{
 						FLinearColor RowBGColor = SystemInfo.UniqueColor;
@@ -1290,7 +1385,7 @@ void FNiagaraDebugHud::DrawOverview(class FNiagaraWorldManager* WorldManager, FC
 						Canvas->DrawShadowedString(X, Y, *LexToSanitizedString(SystemInfo.PerfStats ? SystemInfo.PerfStats->Max.Time_GT : 0.0), Font, RowColor);
 					});
 
-				OverviewColumns.Emplace_GetRef(TEXT("Render Thread Avg:"), LexToSanitizedString(GlobalPerfStats.Avg.Time_RT), TEXT("RT Avg"), ColumnOffset, Font, TEXT("000.0000"),
+				OverviewColumns.Emplace_GetRef(TEXT("Render Thread Avg:"), LexToSanitizedString(GlobalPerfStats.Avg.Time_RT), TEXT("RT Avg (us)"), ColumnOffset, Font, TEXT("000.0000"),
 					[&DetailColor, &DetailHighlightColor, &fAdvanceHeight](FCanvas* Canvas, UFont* Font, float X, float Y, FOverviewColumn& Col, const FSystemDebugInfo& SystemInfo)
 					{
 						FLinearColor RowBGColor = SystemInfo.UniqueColor;
@@ -1300,7 +1395,7 @@ void FNiagaraDebugHud::DrawOverview(class FNiagaraWorldManager* WorldManager, FC
 						Canvas->DrawShadowedString(X, Y, *LexToSanitizedString(SystemInfo.PerfStats ? SystemInfo.PerfStats->Avg.Time_RT : 0.0), Font, RowColor);
 					});
 
-				OverviewColumns.Emplace(TEXT("Render Thread Max:"), LexToSanitizedString(GlobalPerfStats.Max.Time_RT), TEXT("RT Max"), ColumnOffset, Font, TEXT("000.0000"),
+				OverviewColumns.Emplace(TEXT("Render Thread Max:"), LexToSanitizedString(GlobalPerfStats.Max.Time_RT), TEXT("RT Max (us)"), ColumnOffset, Font, TEXT("000.0000"),
 					[&DetailColor, &DetailHighlightColor, &fAdvanceHeight](FCanvas* Canvas, UFont* Font, float X, float Y, FOverviewColumn& Col, const FSystemDebugInfo& SystemInfo)
 					{
 						FLinearColor RowBGColor = SystemInfo.UniqueColor;
@@ -1309,6 +1404,26 @@ void FNiagaraDebugHud::DrawOverview(class FNiagaraWorldManager* WorldManager, FC
 						const FLinearColor RowColor = SystemInfo.bShowInWorld ? DetailHighlightColor : DetailColor;
 						Canvas->DrawShadowedString(X, Y, *LexToSanitizedString(SystemInfo.PerfStats ? SystemInfo.PerfStats->Max.Time_RT : 0.0), Font, RowColor);
 					});
+
+				OverviewColumns.Emplace_GetRef(TEXT("Gpu Avg:"), LexToSanitizedString(GlobalPerfStats.Avg.Time_GPU), TEXT("Gpu Avg (us)"), ColumnOffset, Font, TEXT("000.0000"),
+					[&DetailColor, &DetailHighlightColor, &fAdvanceHeight](FCanvas* Canvas, UFont* Font, float X, float Y, FOverviewColumn& Col, const FSystemDebugInfo& SystemInfo)
+					{
+						FLinearColor RowBGColor = SystemInfo.UniqueColor;
+						RowBGColor.A = 0.1f;
+						Canvas->DrawTile(X, Y, Col.MaxWidth, fAdvanceHeight, 0, 0, 0, 0, RowBGColor);
+						const FLinearColor RowColor = SystemInfo.bShowInWorld ? DetailHighlightColor : DetailColor;
+						Canvas->DrawShadowedString(X, Y, *LexToSanitizedString(SystemInfo.PerfStats ? SystemInfo.PerfStats->Avg.Time_GPU : 0.0), Font, RowColor);
+					});
+
+				OverviewColumns.Emplace(TEXT("Gpu Max:"), LexToSanitizedString(GlobalPerfStats.Max.Time_GPU), TEXT("Gpu Max (us)"), ColumnOffset, Font, TEXT("000.0000"),
+					[&DetailColor, &DetailHighlightColor, &fAdvanceHeight](FCanvas* Canvas, UFont* Font, float X, float Y, FOverviewColumn& Col, const FSystemDebugInfo& SystemInfo)
+					{
+						FLinearColor RowBGColor = SystemInfo.UniqueColor;
+						RowBGColor.A = 0.1f;
+						Canvas->DrawTile(X, Y, Col.MaxWidth, fAdvanceHeight, 0, 0, 0, 0, RowBGColor);
+						const FLinearColor RowColor = SystemInfo.bShowInWorld ? DetailHighlightColor : DetailColor;
+						Canvas->DrawShadowedString(X, Y, *LexToSanitizedString(SystemInfo.PerfStats ? SystemInfo.PerfStats->Max.Time_GPU : 0.0), Font, RowColor);
+					});
 			}
 			else
 			{
@@ -1323,6 +1438,7 @@ void FNiagaraDebugHud::DrawOverview(class FNiagaraWorldManager* WorldManager, FC
 			}
 		}
 #endif//WITH_PARTICLE_PERF_STATS
+
 		TotalOverviewWidth = ColumnOffset;
 	}
 
@@ -1430,8 +1546,16 @@ void FNiagaraDebugHud::DrawOverview(class FNiagaraWorldManager* WorldManager, FC
 		}
 	}
 
+	//-TODO: Improve the column view to allow custom views like this
+#if WITH_NIAGARA_GPU_PROFILER
+	if (Settings.bOverviewEnabled && Settings.OverviewMode == ENiagaraDebugHUDOverviewMode::GpuComputePerformance)
+	{
+		DrawGpuComputeOverriew(WorldManager, DrawCanvas, TextLocation);
+	}
+#endif //WITH_NIAGARA_GPU_PROFILER
+
 	// Display active systems information
-	if (Settings.bOverviewEnabled)
+	if (Settings.bOverviewEnabled && Settings.OverviewMode != ENiagaraDebugHUDOverviewMode::GpuComputePerformance)
 	{
 		TextLocation.Y += fAdvanceHeight;		
 		uint32 NumLines = 1;
@@ -1506,6 +1630,10 @@ void FNiagaraDebugHud::DrawOverview(class FNiagaraWorldManager* WorldManager, FC
 						{
 							SysInfo.PerfStats->History.GetHistoryFrames_RT(Frames);
 						}
+						else if (Settings.PerfGraphMode == ENiagaraDebugHUDPerfGraphMode::GPU)
+						{
+							SysInfo.PerfStats->History.GetHistoryFrames_GPU(Frames);
+						}
 
 						if (Settings.bEnableSmoothing)
 						{
@@ -1540,89 +1668,344 @@ void FNiagaraDebugHud::DrawOverview(class FNiagaraWorldManager* WorldManager, FC
 
 	TextLocation.Y += fAdvanceHeight;
 
-	//Display global budget usage information
-	if(Settings.bShowGlobalBudgetInfo)
+	DrawGlobalBudgetInfo(WorldManager, DrawCanvas, TextLocation);
+	DrawValidation(WorldManager, DrawCanvas, TextLocation);
+	DrawMessages(WorldManager, DrawCanvas, TextLocation);
+}
+
+void FNiagaraDebugHud::DrawGpuComputeOverriew(class FNiagaraWorldManager* WorldManager, class FCanvas* DrawCanvas, FVector2D& TextLocation)
+{
+#if WITH_NIAGARA_GPU_PROFILER
+	using namespace NiagaraDebugLocal;
+
+	UFont* Font = GetFont(Settings.OverviewFont);
+	const float fAdvanceHeight = Font->GetMaxCharHeight() + 1.0f;
+
+	const FLinearColor HeadingColor = FLinearColor::Green;
+	const FLinearColor DetailColor = FLinearColor::White;
+	const FLinearColor DetailHighlightColor = FLinearColor::Yellow;
+
+	if (GpuResults == nullptr)
+	{
+		static const FString EnableCVarWarning(TEXT("No GPU data ensure 'fx.NiagaraGpuProfilingEnabled' is enabled"));
+
+		const FVector2D StringSize = GetStringSize(Font, *EnableCVarWarning);
+		DrawCanvas->DrawTile(TextLocation.X - 1.0f, TextLocation.Y - 1.0f, StringSize.X + 1.0f, 2.0f + StringSize.Y, 0.0f, 0.0f, 0.0f, 0.0f, BackgroundColor);
+		DrawCanvas->DrawShadowedString(TextLocation.X, TextLocation.Y, *EnableCVarWarning, Font, HeadingColor);
+		TextLocation.Y += StringSize.Y;
+		return;
+	}
+
+	struct FSimpleTableDraw
 	{
-		static const float GuessWidth = 400.0f;
-		if (FFXBudget::Enabled())
+		struct FColumn
 		{
-			FFXTimeData Time = FFXBudget::GetTime();
-			FFXTimeData Budget = FFXBudget::GetBudget();
-			FFXTimeData Usage = FFXBudget::GetUsage();
-			FFXTimeData AdjustedUsage = FFXBudget::GetAdjustedUsage();
+			float				DefaultWidth = 150.0f;
+			FVector2D			MeasuredSize = FVector2D::ZeroVector;
+			float				DrawOffset = 0.0f;
+			TStringBuilder<128>	AllRowsText;
+		};
 
-			const int32 NumLines = 5;//Header, time, budget, usage and adjusted usage.
-			DrawCanvas->DrawTile(TextLocation.X - 1.0f, TextLocation.Y - 1.0f, GuessWidth + 1.0f, 2.0f + (float(NumLines) * fAdvanceHeight), 0.0f, 0.0f, 0.0f, 0.0f, BackgroundColor);
+		void AddColumns(int32 NumToAdd, float DefaultWidth = 150.0f)
+		{
+			for ( int32 i=0; i < NumToAdd; ++i )
+			{
+				FColumn& Column = Columns.AddDefaulted_GetRef();
+				Column.DefaultWidth = DefaultWidth;
+			}
+		}
 
-			FString Temp;
-			DrawCanvas->DrawShadowedString(TextLocation.X, TextLocation.Y, TEXT("Global Budget Info"), Font, HeadingColor);
+		TStringBuilder<128>& GetColumnText(int32 Column)
+		{
+			return Columns[Column].AllRowsText;
+		}
 
-			TextLocation.Y += fAdvanceHeight;
-			FString TimeLabels[] =
+		void RowComplete()
+		{
+			for ( int32 i=0; i < Columns.Num(); ++i )
 			{
-				TEXT("Time: "),
-				TEXT("Budget: "),
-				TEXT("Usage: "),
-				TEXT("Adjusted: ")
-			};
+				Columns[i].AllRowsText.Append(TEXT('\n'));
+			}
+		}
 
-			int32 LabelSizes[] =
+		void Draw(UFont* Font, FCanvas* DrawCanvas, FVector2D& Location, FLinearColor TextColor)
+		{
+			FVector2D TotalSize = FVector2D::ZeroVector;
+			for ( FColumn& Column : Columns)
 			{
-				Font->GetStringSize(*TimeLabels[0]),
-				Font->GetStringSize(*TimeLabels[1]),
-				Font->GetStringSize(*TimeLabels[2]),
-				Font->GetStringSize(*TimeLabels[3])
-			};
-			int32 MaxLabelSize = FMath::Max(FMath::Max(LabelSizes[0], LabelSizes[1]), FMath::Max(LabelSizes[2], LabelSizes[3]));
+				Column.MeasuredSize	= GetStringSize(Font, Column.AllRowsText.ToString());
+				Column.DrawOffset	= TotalSize.X;
+
+				const float SpaceBetweenColumns = (&Column == &Columns.Last()) ? 0.0f : DefaultSpaceBetweenColumns;
+				TotalSize.X = Column.DrawOffset + FMath::Max(Column.MeasuredSize.X + SpaceBetweenColumns, Column.DefaultWidth);
+				TotalSize.Y = FMath::Max(TotalSize.Y, Column.MeasuredSize.Y);
+			}
 
-			auto DrawTimeData = [&](FString Heading, FFXTimeData Time, float HighlightThreshold = FLT_MAX)
+			DrawCanvas->DrawTile(Location.X - 1.0f, Location.Y - 1.0f, TotalSize.X + 1.0f, 2.0f + TotalSize.Y, 0.0f, 0.0f, 0.0f, 0.0f, BackgroundColor);
+			for (FColumn& Column : Columns)
 			{
-				//Draw Time
-				int32 XOff = 0;
-				int32 AlignSize = 50;
-				FString Text[] = {
-					Heading,
-					FString::Printf(TEXT("GT = %2.3f"), Time.GT),
-					FString::Printf(TEXT("CNC = %2.3f"), Time.GTConcurrent),
-					FString::Printf(TEXT("RT = %2.3f"), Time.RT)
-				};
-				int32 Sizes[] =
-				{
-					MaxLabelSize,
-					Font->GetStringSize(*Text[1]),
-					Font->GetStringSize(*Text[2]),
-					Font->GetStringSize(*Text[3])
-				};
-				DrawCanvas->DrawShadowedString(TextLocation.X + XOff, TextLocation.Y, *Text[0], Font, HeadingColor);
-				XOff = AlignArbitrary(XOff + Sizes[0], AlignSize);
-				DrawCanvas->DrawShadowedString(TextLocation.X + XOff, TextLocation.Y, *Text[1], Font, Time.GT > HighlightThreshold ? DetailHighlightColor : DetailColor);
-				XOff = AlignArbitrary(XOff + Sizes[1], AlignSize);
-				DrawCanvas->DrawShadowedString(TextLocation.X + XOff, TextLocation.Y, *Text[2], Font, Time.GTConcurrent > HighlightThreshold ? DetailHighlightColor : DetailColor);
-				XOff = AlignArbitrary(XOff + Sizes[2], AlignSize);
-				DrawCanvas->DrawShadowedString(TextLocation.X + XOff, TextLocation.Y, *Text[3], Font, Time.RT > HighlightThreshold ? DetailHighlightColor : DetailColor);
-				XOff = AlignArbitrary(XOff + Sizes[3], AlignSize);
+				DrawCanvas->DrawShadowedString(Location.X + Column.DrawOffset, Location.Y, Column.AllRowsText.ToString(), Font, TextColor);
+			}
 
-				TextLocation.Y += fAdvanceHeight;
-			};
-			DrawTimeData(TimeLabels[0], Time);
-			DrawTimeData(TimeLabels[1], Budget);
-			DrawTimeData(TimeLabels[2], Usage, 1.0f);
-			DrawTimeData(TimeLabels[3], AdjustedUsage, 1.0f);
+			Location.Y += TotalSize.Y;
 		}
-		else
+
+	private:
+		float DefaultSpaceBetweenColumns = 10.0f;
+		TArray<FColumn, TInlineAllocator<16>> Columns;
+	};
+
+	// Display overview of results
+	{
+		int32 TotalDispatches = 0;
+		int32 TotalDispatchGroups = 0;
+		uint64 TotalMicroseconds = 0;
+		for (const auto& StageResult : GpuResults->StageResults)
 		{
-			int32 NumLines = 2;
-			DrawCanvas->DrawTile(TextLocation.X - 1.0f, TextLocation.Y - 1.0f, GuessWidth + 1.0f, 2.0f + (float(NumLines) * fAdvanceHeight), 0.0f, 0.0f, 0.0f, 0.0f, BackgroundColor);
-			DrawCanvas->DrawShadowedString(TextLocation.X, TextLocation.Y, TEXT("Global Budget Info"), Font, HeadingColor);
-			TextLocation.Y += fAdvanceHeight;
-			DrawCanvas->DrawShadowedString(TextLocation.X, TextLocation.Y, TEXT("Global budget tracking is disabled."), Font, DetailHighlightColor);
-			TextLocation.Y += fAdvanceHeight;
+			TotalDispatches += StageResult.NumDispatches;
+			TotalDispatchGroups += StageResult.NumDispatchGroups;
+			TotalMicroseconds += StageResult.DurationMicroseconds;
 		}
+
+		FSimpleTableDraw SimpleTable;
+		SimpleTable.AddColumns(1, 250.0f);
+		SimpleTable.AddColumns(2, 150.0f);
+		SimpleTable.GetColumnText(0).Append(TEXT("Gpu Compute Overview"));
+		SimpleTable.RowComplete();
+
+		SimpleTable.GetColumnText(0).Appendf(TEXT("TotalMicroseconds : %llu"), TotalMicroseconds);
+		SimpleTable.GetColumnText(1).Appendf(TEXT("TotalDispatches : %d"), TotalDispatches);
+		SimpleTable.GetColumnText(2).Appendf(TEXT("TotalDispatchGroups : %d"), TotalDispatchGroups);
+
+		SimpleTable.Draw(Font, DrawCanvas, TextLocation, DetailColor);
+		TextLocation.Y += fAdvanceHeight;
 	}
 
-	DrawValidation(WorldManager, DrawCanvas, TextLocation);
+	// Do a pass to determine what will be detailed view
+	bool bHasDetailedView = false;
+	bool bHasSimpleView = false;
+	for (auto SystemIt = GpuUsagePerSystem.CreateIterator(); SystemIt; ++SystemIt)
+	{
+		// Prune data systems
+		UNiagaraSystem* OwnerSystem = SystemIt.Key().Get();
+		if (OwnerSystem == nullptr)
+		{
+			SystemIt.RemoveCurrent();
+			continue;
+		}
 
-	DrawMessages(WorldManager, DrawCanvas, TextLocation);
+		// Prune dead emitters
+		for (auto EmitterIt = SystemIt.Value().Emitters.CreateIterator(); EmitterIt; ++EmitterIt)
+		{
+			UNiagaraEmitter* OwnerEmitter = EmitterIt.Key().Get();
+			if (OwnerEmitter == nullptr)
+			{
+				SystemIt.RemoveCurrent();
+			}
+		}
+
+		const bool bShowDetailed = Settings.bSystemFilterEnabled && OwnerSystem->GetName().MatchesWildcard(Settings.SystemFilter);
+		SystemIt.Value().bShowDetailed = bShowDetailed;
+		bHasDetailedView |= bShowDetailed;
+		bHasSimpleView |= !bShowDetailed;
+	}
+
+	// Show detailed information first
+	if (bHasDetailedView)
+	{
+		FSimpleTableDraw SimpleTable;
+		SimpleTable.AddColumns(3, 200.0f);
+		SimpleTable.AddColumns(3, 100.0f);
+		SimpleTable.GetColumnText(0).Append(TEXT("SystemName"));
+		SimpleTable.GetColumnText(1).Append(TEXT("EmitterName"));
+		SimpleTable.GetColumnText(2).Append(TEXT("StageName"));
+		SimpleTable.GetColumnText(3).Append(TEXT("Avg Instances"));
+		SimpleTable.GetColumnText(4).Append(TEXT("Avg us"));
+		SimpleTable.GetColumnText(5).Append(TEXT("Max us"));
+		SimpleTable.RowComplete();
+
+		for (auto SystemIt = GpuUsagePerSystem.CreateIterator(); SystemIt; ++SystemIt)
+		{
+			if ( SystemIt.Value().bShowDetailed == false )
+			{
+				continue;
+			}
+
+			UNiagaraSystem* OwnerSystem = SystemIt.Key().Get();
+			check(OwnerSystem);
+			for (auto EmitterIt = SystemIt.Value().Emitters.CreateIterator(); EmitterIt; ++EmitterIt)
+			{
+				UNiagaraEmitter* OwnerEmitter = EmitterIt.Key().Get();
+				check(OwnerEmitter);
+				for (auto StageIt=EmitterIt.Value().Stages.CreateIterator(); StageIt; ++StageIt)
+				{
+					const FGpuUsagePerStage& StageUsage = StageIt.Value();
+					OwnerSystem->GetFName().AppendString(SimpleTable.GetColumnText(0));
+					OwnerEmitter->GetFName().AppendString(SimpleTable.GetColumnText(1));
+					StageIt.Key().AppendString(SimpleTable.GetColumnText(2));
+					SimpleTable.GetColumnText(3).Appendf(TEXT("%4.1f"), StageUsage.InstanceCount.GetAverage<float>());
+					SimpleTable.GetColumnText(4).Appendf(TEXT("%llu"), StageUsage.Microseconds.GetAverage());
+					SimpleTable.GetColumnText(5).Appendf(TEXT("%llu"), StageUsage.Microseconds.GetMax());
+					SimpleTable.RowComplete();
+				}
+			}
+		}
+		SimpleTable.Draw(Font, DrawCanvas, TextLocation, DetailColor);
+		TextLocation.Y += fAdvanceHeight;
+	}
+
+	// Show simple information
+	if (bHasSimpleView)
+	{
+		FSimpleTableDraw SimpleTable;
+		SimpleTable.AddColumns(1, 200.0f);
+		SimpleTable.AddColumns(3, 100.0f);
+		SimpleTable.GetColumnText(0).Append(TEXT("SystemName"));
+		SimpleTable.GetColumnText(1).Append(TEXT("Avg Instances"));
+		SimpleTable.GetColumnText(2).Append(TEXT("Avg us"));
+		SimpleTable.GetColumnText(3).Append(TEXT("Max us"));
+		SimpleTable.RowComplete();
+
+		for (auto SystemIt=GpuUsagePerSystem.CreateIterator(); SystemIt; ++SystemIt)
+		{
+			if (SystemIt.Value().bShowDetailed == true)
+			{
+				continue;
+			}
+
+			UNiagaraSystem* OwnerSystem = SystemIt.Key().Get();
+			check(OwnerSystem);
+
+			const FGpuUsagePerSystem& SystemUsage = SystemIt.Value();
+
+			OwnerSystem->GetFName().AppendString(SimpleTable.GetColumnText(0));
+			SimpleTable.GetColumnText(1).Appendf(TEXT("%4.1f"), SystemUsage.InstanceCount.GetAverage<float>());
+			SimpleTable.GetColumnText(2).Appendf(TEXT("%llu"), SystemUsage.Microseconds.GetAverage());
+			SimpleTable.GetColumnText(3).Appendf(TEXT("%llu"), SystemUsage.Microseconds.GetMax());
+			SimpleTable.RowComplete();
+		}
+		SimpleTable.Draw(Font, DrawCanvas, TextLocation, DetailColor);
+		TextLocation.Y += fAdvanceHeight;
+	}
+
+	// Show events
+	if (GpuUsagePerEvent.Num() > 0)
+	{
+		FSimpleTableDraw SimpleTable;
+		SimpleTable.AddColumns(1, 200.0f);
+		SimpleTable.AddColumns(3, 100.0f);
+		SimpleTable.GetColumnText(0).Append(TEXT("EventName"));
+		SimpleTable.GetColumnText(1).Append(TEXT("Avg Instances"));
+		SimpleTable.GetColumnText(2).Append(TEXT("Avg us"));
+		SimpleTable.GetColumnText(3).Append(TEXT("Max us"));
+		SimpleTable.RowComplete();
+
+		for (auto EventIt=GpuUsagePerEvent.CreateIterator(); EventIt; ++EventIt)
+		{
+			const FGpuUsagePerEvent& EventUsage = EventIt.Value();
+			EventIt.Key().AppendString(SimpleTable.GetColumnText(0));
+			SimpleTable.GetColumnText(1).Appendf(TEXT("%4.1f"), EventUsage.InstanceCount.GetAverage<float>());
+			SimpleTable.GetColumnText(2).Appendf(TEXT("%llu"), EventUsage.Microseconds.GetAverage());
+			SimpleTable.GetColumnText(3).Appendf(TEXT("%llu"), EventUsage.Microseconds.GetMax());
+			SimpleTable.RowComplete();
+		}
+		SimpleTable.Draw(Font, DrawCanvas, TextLocation, DetailColor);
+		TextLocation.Y += fAdvanceHeight;
+	}
+#endif //WITH_NIAGARA_GPU_PROFILER
+}
+
+void FNiagaraDebugHud::DrawGlobalBudgetInfo(class FNiagaraWorldManager* WorldManager, class FCanvas* DrawCanvas, FVector2D& TextLocation)
+{
+	using namespace NiagaraDebugLocal;
+
+	if ( !Settings.bShowGlobalBudgetInfo )
+	{
+		return;
+	}
+
+	static const float GuessWidth = 400.0f;
+
+	UFont* Font = GetFont(Settings.OverviewFont);
+	const float fAdvanceHeight = Font->GetMaxCharHeight() + 1.0f;
+
+	const FLinearColor HeadingColor = FLinearColor::Green;
+	const FLinearColor DetailColor = FLinearColor::White;
+	const FLinearColor DetailHighlightColor = FLinearColor::Yellow;
+
+	if (FFXBudget::Enabled())
+	{
+		const FFXTimeData Time = FFXBudget::GetTime();
+		const FFXTimeData Budget = FFXBudget::GetBudget();
+		const FFXTimeData Usage = FFXBudget::GetUsage();
+		const FFXTimeData AdjustedUsage = FFXBudget::GetAdjustedUsage();
+
+		const int32 NumLines = 5;//Header, time, budget, usage and adjusted usage.
+		DrawCanvas->DrawTile(TextLocation.X - 1.0f, TextLocation.Y - 1.0f, GuessWidth + 1.0f, 2.0f + (float(NumLines) * fAdvanceHeight), 0.0f, 0.0f, 0.0f, 0.0f, BackgroundColor);
+
+		DrawCanvas->DrawShadowedString(TextLocation.X, TextLocation.Y, TEXT("Global Budget Info"), Font, HeadingColor);
+
+		TextLocation.Y += fAdvanceHeight;
+		FString TimeLabels[] =
+		{
+			TEXT("Time: "),
+			TEXT("Budget: "),
+			TEXT("Usage: "),
+			TEXT("Adjusted: ")
+		};
+
+		int32 LabelSizes[] =
+		{
+			Font->GetStringSize(*TimeLabels[0]),
+			Font->GetStringSize(*TimeLabels[1]),
+			Font->GetStringSize(*TimeLabels[2]),
+			Font->GetStringSize(*TimeLabels[3])
+		};
+		int32 MaxLabelSize = FMath::Max(FMath::Max(LabelSizes[0], LabelSizes[1]), FMath::Max(LabelSizes[2], LabelSizes[3]));
+
+		auto DrawTimeData = [&](FString Heading, FFXTimeData Time, float HighlightThreshold = FLT_MAX)
+		{
+			//Draw Time
+			int32 XOff = 0;
+			int32 AlignSize = 50;
+			FString Text[] = {
+				Heading,
+				FString::Printf(TEXT("GT = %2.3f"), Time.GT),
+				FString::Printf(TEXT("CNC = %2.3f"), Time.GTConcurrent),
+				FString::Printf(TEXT("RT = %2.3f"), Time.RT)
+			};
+			int32 Sizes[] =
+			{
+				MaxLabelSize,
+				Font->GetStringSize(*Text[1]),
+				Font->GetStringSize(*Text[2]),
+				Font->GetStringSize(*Text[3])
+			};
+			DrawCanvas->DrawShadowedString(TextLocation.X + XOff, TextLocation.Y, *Text[0], Font, HeadingColor);
+			XOff = AlignArbitrary(XOff + Sizes[0], AlignSize);
+			DrawCanvas->DrawShadowedString(TextLocation.X + XOff, TextLocation.Y, *Text[1], Font, Time.GT > HighlightThreshold ? DetailHighlightColor : DetailColor);
+			XOff = AlignArbitrary(XOff + Sizes[1], AlignSize);
+			DrawCanvas->DrawShadowedString(TextLocation.X + XOff, TextLocation.Y, *Text[2], Font, Time.GTConcurrent > HighlightThreshold ? DetailHighlightColor : DetailColor);
+			XOff = AlignArbitrary(XOff + Sizes[2], AlignSize);
+			DrawCanvas->DrawShadowedString(TextLocation.X + XOff, TextLocation.Y, *Text[3], Font, Time.RT > HighlightThreshold ? DetailHighlightColor : DetailColor);
+			XOff = AlignArbitrary(XOff + Sizes[3], AlignSize);
+
+			TextLocation.Y += fAdvanceHeight;
+		};
+		DrawTimeData(TimeLabels[0], Time);
+		DrawTimeData(TimeLabels[1], Budget);
+		DrawTimeData(TimeLabels[2], Usage, 1.0f);
+		DrawTimeData(TimeLabels[3], AdjustedUsage, 1.0f);
+	}
+	else
+	{
+		int32 NumLines = 2;
+		DrawCanvas->DrawTile(TextLocation.X - 1.0f, TextLocation.Y - 1.0f, GuessWidth + 1.0f, 2.0f + (float(NumLines) * fAdvanceHeight), 0.0f, 0.0f, 0.0f, 0.0f, BackgroundColor);
+		DrawCanvas->DrawShadowedString(TextLocation.X, TextLocation.Y, TEXT("Global Budget Info"), Font, HeadingColor);
+		TextLocation.Y += fAdvanceHeight;
+		DrawCanvas->DrawShadowedString(TextLocation.X, TextLocation.Y, TEXT("Global budget tracking is disabled."), Font, DetailHighlightColor);
+		TextLocation.Y += fAdvanceHeight;
+	}
 }
 
 void FNiagaraDebugHud::DrawValidation(class FNiagaraWorldManager* WorldManager, class FCanvas* DrawCanvas, FVector2D& TextLocation)
@@ -2234,6 +2617,14 @@ void FNiagaraDebugHudStatHistory::AddFrame_RT(double Time)
 	RTFrames[CurrFrameRT] = Time;
 }
 
+void FNiagaraDebugHudStatHistory::AddFrame_GPU(double Time)
+{
+	FScopeLock Lock(&NiagaraDebugLocal::RTFramesGuard);
+	GPUFrames.SetNumZeroed(NiagaraDebugLocal::Settings.PerfHistoryFrames);
+	CurrFrameGPU = FMath::Wrap(CurrFrameGPU + 1, 0, GPUFrames.Num() - 1);
+	GPUFrames[CurrFrameGPU] = Time;
+}
+
 void FNiagaraDebugHudStatHistory::GetHistoryFrames_GT(TArray<double>& OutHistoryGT)
 {
 	if (GTFrames.Num() == 0)
@@ -2267,6 +2658,24 @@ void FNiagaraDebugHudStatHistory::GetHistoryFrames_RT(TArray<double>& OutHistory
 	} while (WriteFrame != CurrFrameRT);
 }
 
+void FNiagaraDebugHudStatHistory::GetHistoryFrames_GPU(TArray<double>& OutHistoryGPU)
+{
+	FScopeLock Lock(&NiagaraDebugLocal::RTFramesGuard);
+
+	if (GPUFrames.Num() == 0)
+	{
+		return;
+	}
+
+	OutHistoryGPU.Reserve(GPUFrames.Num());
+	int32 WriteFrame = CurrFrameGPU;
+	do
+	{
+		OutHistoryGPU.Add(GPUFrames[WriteFrame]);
+		WriteFrame = FMath::Wrap(WriteFrame + 1, 0, GPUFrames.Num() - 1);
+	} while (WriteFrame != CurrFrameGPU);
+}
+
 TSharedPtr<FNiagaraDebugHUDPerfStats> FNiagaraDebugHUDStatsListener::GetSystemStats(UNiagaraSystem* System)
 {
 	FScopeLock Lock(&SystemStatsGuard);//Locking on every get isn't great but it's debug hud code so meh.
@@ -2406,8 +2815,10 @@ void FNiagaraDebugHUDStatsListener::TickRT()
 	FScopeLock Lock(&AccumulatedStatsGuard);
 	FScopeLock SysInfoLock(&SystemStatsGuard);
 
-	double GlobalAvg = 0.0;
-	double GlobalMax = 0.0;
+	double RTGlobalAvg = 0.0;
+	double RTGlobalMax = 0.0;
+	double GPUGlobalAvg = 0.0;
+	double GPUGlobalMax = 0.0;
 	for (auto& WorldStatsPair : AccumulatedWorldStats)
 	{
 		if (const UWorld* World = Cast<UWorld>(WorldStatsPair.Key.Get()))
@@ -2415,18 +2826,23 @@ void FNiagaraDebugHUDStatsListener::TickRT()
 			FAccumulatedParticlePerfStats* Stats = WorldStatsPair.Value.Get();
 			check(Stats);
 
-			GlobalAvg += Stats->GetRenderThreadStats().GetPerFrameAvg();
-			GlobalMax += Stats->GetRenderThreadStats().GetPerFrameMax();
+			RTGlobalAvg += Stats->GetRenderThreadStats().GetPerFrameAvg();
+			RTGlobalMax += Stats->GetRenderThreadStats().GetPerFrameMax();
+			GPUGlobalAvg += Stats->GetGPUStats().GetPerFrameAvgMicroseconds();
+			GPUGlobalMax += Stats->GetGPUStats().GetPerFrameMaxMicroseconds();
 
 			Stats->ResetRT();
 		}
 	}
-	GlobalStats.History.AddFrame_RT(GlobalAvg);
+	GlobalStats.History.AddFrame_RT(RTGlobalAvg);
+	GlobalStats.History.AddFrame_GPU(GPUGlobalAvg);
 
 	if (bPushStats)
 	{
-		GlobalStats.Avg.Time_RT = GlobalAvg;
-		GlobalStats.Max.Time_RT = GlobalMax;
+		GlobalStats.Avg.Time_RT = RTGlobalAvg;
+		GlobalStats.Max.Time_RT = RTGlobalMax;
+		GlobalStats.Avg.Time_GPU = GPUGlobalAvg;
+		GlobalStats.Max.Time_GPU = GPUGlobalMax;
 	}
 
 #if WITH_PER_SYSTEM_PARTICLE_PERF_STATS
@@ -2444,25 +2860,34 @@ void FNiagaraDebugHUDStatsListener::TickRT()
 
 		if (FAccumulatedParticlePerfStats* Stats = GetStats(System))
 		{
-			double SysAvg;
-			double SysMax;
+			double RTSysAvg;
+			double RTSysMax;
+			double GPUSysAvg;
+			double GPUSysMax;
 			if (Settings.PerfSampleMode == ENiagaraDebugHUDPerfSampleMode::FrameTotal)
 			{
-				SysAvg = Stats->GetRenderThreadStats().GetPerFrameAvg();
-				SysMax = Stats->GetRenderThreadStats().GetPerFrameAvg();
+				RTSysAvg = Stats->GetRenderThreadStats().GetPerFrameAvg();
+				RTSysMax = Stats->GetRenderThreadStats().GetPerFrameAvg();
+				GPUSysAvg = Stats->GetGPUStats().GetPerFrameAvgMicroseconds();
+				GPUSysMax = Stats->GetGPUStats().GetPerFrameAvgMicroseconds();
 			}
 			else
 			{
-				SysAvg = Stats->GetRenderThreadStats().GetPerInstanceAvg();
-				SysMax = Stats->GetRenderThreadStats().GetPerInstanceMax();
+				RTSysAvg = Stats->GetRenderThreadStats().GetPerInstanceAvg();
+				RTSysMax = Stats->GetRenderThreadStats().GetPerInstanceMax();
+				GPUSysAvg = Stats->GetGPUStats().GetPerInstanceAvgMicroseconds();
+				GPUSysMax = Stats->GetGPUStats().GetPerInstanceMaxMicroseconds();
 			}
 
-			HUDStats->History.AddFrame_RT(SysAvg);
+			HUDStats->History.AddFrame_RT(RTSysAvg);
+			HUDStats->History.AddFrame_GPU(GPUSysAvg);
 
 			if (bPushStats)
 			{
-				HUDStats->Avg.Time_RT = SysAvg;
-				HUDStats->Max.Time_RT = SysMax;
+				HUDStats->Avg.Time_RT = RTSysAvg;
+				HUDStats->Max.Time_RT = RTSysMax;
+				HUDStats->Avg.Time_GPU = GPUSysAvg;
+				HUDStats->Max.Time_GPU = GPUSysMax;
 			}
 
 			Stats->ResetRT();
diff --git a/Engine/Plugins/FX/Niagara/Source/Niagara/Private/NiagaraDebugHud.h b/Engine/Plugins/FX/Niagara/Source/Niagara/Private/NiagaraDebugHud.h
--- a/Engine/Plugins/FX/Niagara/Source/Niagara/Private/NiagaraDebugHud.h
+++ b/Engine/Plugins/FX/Niagara/Source/Niagara/Private/NiagaraDebugHud.h
@@ -10,6 +10,7 @@ Class used help debugging Niagara simulations
 #include "NiagaraTypes.h"
 #include "RHICommandList.h"
 #include "NiagaraDebuggerCommon.h"
+#include "NiagaraGPUProfilerInterface.h"
 #include "Particles/ParticlePerfStatsManager.h"
 
 #if WITH_PARTICLE_PERF_STATS
@@ -18,8 +19,9 @@ class FNiagaraDebugHud;
 
 struct FNiagaraDebugHudFrameStat
 {
-	double Time_GT;
-	double Time_RT;
+	double Time_GT = 0.0;
+	double Time_RT = 0.0;
+	double Time_GPU = 0.0;
 };
 
 /** Ring buffer history of stats. */
@@ -27,15 +29,19 @@ struct FNiagaraDebugHudStatHistory
 {
 	TArray<double> GTFrames;
 	TArray<double> RTFrames;
+	TArray<double> GPUFrames;
 
 	int32 CurrFrame = 0;
 	int32 CurrFrameRT = 0;
+	int32 CurrFrameGPU = 0;
 
 	void AddFrame_GT(double Time);
 	void AddFrame_RT(double Time);
+	void AddFrame_GPU(double Time);
 
 	void GetHistoryFrames_GT(TArray<double>& OutHistoryGT);
 	void GetHistoryFrames_RT(TArray<double>& OutHistoryRT);
+	void GetHistoryFrames_GPU(TArray<double>& OutHistoryGPU);
 };
 
 struct FNiagaraDebugHUDPerfStats
@@ -142,6 +148,65 @@ class FNiagaraDebugHud
 		TMap<FName, TArray<FName>>	ParticleVariablesWithErrors;
 	};
 
+	static constexpr int32 SmoothedNumFrames = 32;
+	template<typename TCounterType>
+	struct FSmoothedCounter
+	{
+		uint64			LastFrameSeen = 0;
+		TCounterType	FrameCount[SmoothedNumFrames];
+		uint32			MaxIndex = 0;
+		uint32			CurrentIndex = SmoothedNumFrames - 1;
+
+		void Reset()
+		{
+			LastFrameSeen = 0;
+			MaxIndex = 0;
+			CurrentIndex = SmoothedNumFrames - 1;
+		}
+
+		bool ShouldPrune(uint64 FrameCounter)
+		{
+			return (FrameCounter - LastFrameSeen) >= SmoothedNumFrames;
+		}
+
+		void Accumulate(uint64 InLastFrameSeen, TCounterType Value)
+		{
+			if ( LastFrameSeen != InLastFrameSeen )
+			{
+				LastFrameSeen = InLastFrameSeen;
+				CurrentIndex = (CurrentIndex + 1) % SmoothedNumFrames;
+				FrameCount[CurrentIndex] = TCounterType(0);
+				MaxIndex = FMath::Max(MaxIndex, CurrentIndex + 1);
+			}
+			FrameCount[CurrentIndex] += Value;
+		}
+
+		template<typename TReturnType = TCounterType>
+		TReturnType GetAverage() const
+		{
+			TReturnType Result = 0;
+			if (MaxIndex > 0)
+			{
+				for (uint32 i = 0; i < MaxIndex; ++i)
+				{
+					Result += TReturnType(FrameCount[i]);
+				}
+				Result /= TReturnType(MaxIndex);
+			}
+			return Result;
+		}
+
+		TCounterType GetMax() const
+		{
+			TCounterType Result = 0;
+			for (uint32 i = 0; i < MaxIndex; ++i)
+			{
+				Result = FMath::Max(Result, FrameCount[i]);
+			}
+			return Result;
+		}
+	};
+
 public:
 	FNiagaraDebugHud(class UWorld* World);
 	~FNiagaraDebugHud();
@@ -161,6 +226,8 @@ private:
 
 	void Draw(class FNiagaraWorldManager* WorldManager, class UCanvas* Canvas, class APlayerController* PC);
 	void DrawOverview(class FNiagaraWorldManager* WorldManager, class FCanvas* DrawCanvas);
+	void DrawGpuComputeOverriew(class FNiagaraWorldManager* WorldManager, class FCanvas* DrawCanvas, FVector2D& TextLocation);
+	void DrawGlobalBudgetInfo(class FNiagaraWorldManager* WorldManager, class FCanvas* DrawCanvas, FVector2D& TextLocation);
 	void DrawValidation(class FNiagaraWorldManager* WorldManager, class FCanvas* DrawCanvas, FVector2D& TextLocation);
 	void DrawComponents(class FNiagaraWorldManager* WorldManager, class UCanvas* Canvas);
 	void DrawMessages(class FNiagaraWorldManager* WorldManager, class FCanvas* DrawCanvas, FVector2D& TextLocation);
@@ -194,6 +261,42 @@ private:
 	TSharedPtr<FNiagaraDebugHUDStatsListener, ESPMode::ThreadSafe> StatsListener;
 #endif
 
+#if WITH_NIAGARA_GPU_PROFILER
+	FDelegateHandle				GpuProfilerDelegateHandle;
+	FNiagaraGpuFrameResultsPtr	GpuResults;
+	uint64						GpuResultsGameFrameCounter = 0;
+
+	struct FGpuUsagePerStage
+	{
+		FSmoothedCounter<uint32> InstanceCount;
+		FSmoothedCounter<uint64> Microseconds;
+	};
+
+	struct FGpuUsagePerEmitter
+	{
+		FSmoothedCounter<uint32> InstanceCount;
+		FSmoothedCounter<uint64> Microseconds;
+		TMap<FName, FGpuUsagePerStage> Stages;
+	};
+
+	struct FGpuUsagePerSystem
+	{
+		bool bShowDetailed = false;
+		FSmoothedCounter<uint32> InstanceCount;
+		FSmoothedCounter<uint64> Microseconds;
+		TMap<TWeakObjectPtr<UNiagaraEmitter>, FGpuUsagePerEmitter> Emitters;
+	};
+
+	struct FGpuUsagePerEvent
+	{
+		FSmoothedCounter<uint32> InstanceCount;
+		FSmoothedCounter<uint64> Microseconds;
+	};
+
+	TMap<TWeakObjectPtr<UNiagaraSystem>, FGpuUsagePerSystem> GpuUsagePerSystem;
+	TMap<FName, FGpuUsagePerEvent> GpuUsagePerEvent;
+#endif
+
 	float LastDrawTime = 0.0f;
 	float DeltaSeconds = 0.0f;
 
diff --git a/Engine/Plugins/FX/Niagara/Source/Niagara/Private/NiagaraEmitterInstance.cpp b/Engine/Plugins/FX/Niagara/Source/Niagara/Private/NiagaraEmitterInstance.cpp
--- a/Engine/Plugins/FX/Niagara/Source/Niagara/Private/NiagaraEmitterInstance.cpp
+++ b/Engine/Plugins/FX/Niagara/Source/Niagara/Private/NiagaraEmitterInstance.cpp
@@ -405,9 +405,8 @@ void FNiagaraEmitterInstance::Init(int32 InEmitterIdx, FNiagaraSystemInstanceID
 			GPUExecContext = new FNiagaraComputeExecutionContext();
 			GPUExecContext->InitParams(CachedEmitter->GetGPUComputeScript(), CachedEmitter->SimTarget);
 			GPUExecContext->SetDebugSimName(CachedEmitter->GetDebugSimName());
-#if STATS
-			GPUExecContext->EmitterPtr = GetEmitterHandle().GetInstance();
-#endif
+			GPUExecContext->ProfilingComponentPtr = ParentSystemInstance ? ParentSystemInstance->GetAttachComponent() : nullptr;
+			GPUExecContext->ProfilingEmitterPtr = GetEmitterHandle().GetInstance();
 			GPUExecContext->MainDataSet = ParticleDataSet;
 			GPUExecContext->GPUScript_RT = CachedEmitter->GetGPUComputeScript()->GetRenderThreadScript();
 		}
diff --git a/Engine/Plugins/FX/Niagara/Source/Niagara/Private/NiagaraGPUProfiler.cpp b/Engine/Plugins/FX/Niagara/Source/Niagara/Private/NiagaraGPUProfiler.cpp
--- a/Engine/Plugins/FX/Niagara/Source/Niagara/Private/NiagaraGPUProfiler.cpp
+++ b/Engine/Plugins/FX/Niagara/Source/Niagara/Private/NiagaraGPUProfiler.cpp
@@ -1,92 +1,260 @@
 ﻿// Copyright Epic Games, Inc. All Rights Reserved.
 
 #include "NiagaraGPUProfiler.h"
+#include "NiagaraGPUSystemTick.h"
+#include "NiagaraSimStageData.h"
 #include "HAL/IConsoleManager.h"
+#include "Misc/StringBuilder.h"
 
 int32 GNiagaraGpuProfilingEnabled = 0;
 static FAutoConsoleVariableRef CVarNiagaraGpuProfilingEnabled(
-    TEXT("fx.NiagaraGpuProfilingEnabled"),
-    GNiagaraGpuProfilingEnabled,
-    TEXT("Used by the profiling tool in the system overview to enable or disable gathering of gpu stats.\n"),
-    ECVF_Default
+	TEXT("fx.NiagaraGpuProfilingEnabled"),
+	GNiagaraGpuProfilingEnabled,
+	TEXT("Used by the profiling tool in the system overview to enable or disable gathering of gpu stats.\n"),
+	ECVF_Default
 );
 
-#if STATS
-FNiagaraGPUProfiler::FNiagaraGPUProfiler()
+#if WITH_NIAGARA_GPU_PROFILER
+
+FNiagaraGPUProfiler::FNiagaraGPUProfiler(uintptr_t InOwnerContext)
+	: OwnerContext(InOwnerContext)
 {
 	QueryPool = RHICreateRenderQueryPool(RQT_AbsoluteTime);
 }
 
-int32 FNiagaraGPUProfiler::StartTimer(uint64 ReporterHandle, TStatId StatId, FRHICommandList& RHICmdList)
+FNiagaraGPUProfiler::~FNiagaraGPUProfiler()
 {
-	check(IsInRenderingThread());
+	for (FGpuFrameData& Frame : GpuFrames)
+	{
+		Frame.EndQuery.ReleaseQuery();
+		for (FGpuStageTimer& StageTimer : Frame.StageTimers)
+		{
+			StageTimer.StartQuery.ReleaseQuery();
+			StageTimer.EndQuery.ReleaseQuery();
+		}
+
+		for (FGpuDispatchTimer& DispatchTimer : Frame.DispatchTimers)
+		{
+			DispatchTimer.StartQuery.ReleaseQuery();
+			DispatchTimer.EndQuery.ReleaseQuery();
+		}
+	}
+}
 
-	if (GNiagaraGpuProfilingEnabled == 0)
+void FNiagaraGPUProfiler::BeginFrame(FRHICommandListImmediate& RHICmdList)
+{
+	bProfilingFrame = false;
+	bProfilingDispatches = false;
+
+	// Process any frames that are potentially complete
+	while ( CurrentReadFrame != CurrentWriteFrame )
 	{
-		return -1;
+		FGpuFrameData& ReadFrame = GpuFrames[CurrentReadFrame];
+		if ( !ProcessFrame(RHICmdList, ReadFrame) )
+		{
+			break;
+		}
+
+		// Frame was processed
+		CurrentReadFrame = (CurrentReadFrame + 1) % NumBufferFrames;
 	}
 
-	int32 TimerIndex = ActiveTimers.AddZeroed();
-	FTimerData& Timer = ActiveTimers[TimerIndex];
-	Timer.ReporterHandle = ReporterHandle;
-	Timer.StatId = StatId;
-	Timer.StartQuery = QueryPool->AllocateQuery();
-	Timer.EndQuery = QueryPool->AllocateQuery();
-	Timer.FrameNumber = GFrameNumberRenderThread;
+	// If we are not enabled nothing to profile
+	if (!GNiagaraGpuProfilingEnabled)
+	{
+		return;
+	}
+
+	// Have a processed all previous data?
+	if (GpuFrames[CurrentWriteFrame].EndQuery.GetQuery() != nullptr)
+	{
+		return;
+	}
 
-	RHICmdList.EndRenderQuery(Timer.StartQuery.GetQuery());
-	return TimerIndex;
+	// Latch values for frame
+	bProfilingFrame = true;
+	bProfilingDispatches = true;
 }
 
-void FNiagaraGPUProfiler::EndTimer(int32 TimerHandle, FRHICommandList& RHICmdList)
+void FNiagaraGPUProfiler::EndFrame(FRHICommandList& RHICmdList)
 {
-	check(IsInRenderingThread());
-	if (GNiagaraGpuProfilingEnabled == 0 || TimerHandle < 0)
+	if ( !bProfilingFrame )
 	{
 		return;
 	}
-	check(ActiveTimers.IsValidIndex(TimerHandle));
-	
-	RHICmdList.EndRenderQuery(ActiveTimers[TimerHandle].EndQuery.GetQuery());
+
+	// Inject end marker so we know if all dispatches are complete
+	FGpuFrameData& WriteFrame = GetWriteFrame();
+
+	WriteFrame.EndQuery = QueryPool->AllocateQuery();
+	RHICmdList.EndRenderQuery(WriteFrame.EndQuery.GetQuery());
+
+	CurrentWriteFrame = (CurrentWriteFrame + 1) % NumBufferFrames;
 }
 
-void FNiagaraGPUProfiler::QueryTimingResults(FRHICommandList& RHICmdList, TArray<FNiagaraGPUTimingResult, TInlineAllocator<16>>& OutResults)
+void FNiagaraGPUProfiler::BeginStage(FRHICommandList& RHICmdList, ENiagaraGpuComputeTickStage::Type TickStage, int32 NumDispatchGroups)
 {
-	check(IsInRenderingThread());
-	
-	if (ActiveTimers.Num() == 0 || !RHICmdList.IsImmediate())
+	if (!bProfilingFrame)
 	{
 		return;
 	}
 
-	FRHICommandListImmediate* ImmediateCmdList = (FRHICommandListImmediate*)&RHICmdList;	
-	for (int i = ActiveTimers.Num() - 1; i >= 0; i--)
+	FGpuFrameData& WriteFrame = GetWriteFrame();
+	FGpuStageTimer& StageTimer = WriteFrame.StageTimers[TickStage];
+	StageTimer.NumDispatchGroups = NumDispatchGroups;
+	StageTimer.StartQuery = QueryPool->AllocateQuery();
+	RHICmdList.EndRenderQuery(StageTimer.StartQuery.GetQuery());
+}
+
+void FNiagaraGPUProfiler::EndStage(FRHICommandList& RHICmdList, ENiagaraGpuComputeTickStage::Type TickStage, int32 NumDispatches)
+{
+	if (!bProfilingFrame)
 	{
-		FTimerData& Timer = ActiveTimers[i];
-		if (Timer.FrameNumber != GFrameNumberRenderThread)
-		{
-			if (Timer.StartTime == 0 && !ImmediateCmdList->GetRenderQueryResult(Timer.StartQuery.GetQuery(), Timer.StartTime, false))
-			{
-				continue;
-			}
-			if (Timer.EndTime == 0 && !ImmediateCmdList->GetRenderQueryResult(Timer.EndQuery.GetQuery(), Timer.EndTime, false))
-			{
-				continue;
-			}
-
-			FNiagaraGPUTimingResult TimingResult;
-			TimingResult.ReporterHandle = Timer.ReporterHandle;
-			TimingResult.StatId = Timer.StatId;
-			TimingResult.ElapsedMicroseconds = Timer.EndTime - Timer.StartTime;
-			OutResults.Add(TimingResult);
-			ActiveTimers.RemoveAtSwap(i);
-		}
+		return;
 	}
+
+	FGpuFrameData& WriteFrame = GetWriteFrame();
+	FGpuStageTimer& StageTimer = WriteFrame.StageTimers[TickStage];
+	StageTimer.NumDispatches = NumDispatches;
+	StageTimer.EndQuery = QueryPool->AllocateQuery();
+	RHICmdList.EndRenderQuery(StageTimer.EndQuery.GetQuery());
 }
 
-bool FNiagaraGPUProfiler::IsProfilingEnabled() const
+void FNiagaraGPUProfiler::BeginDispatch(FRHICommandList& RHICmdList, const FNiagaraGpuDispatchInstance& DispatchInstance)
 {
-	return GNiagaraGpuProfilingEnabled > 0;
+	if (!bProfilingDispatches)
+	{
+		return;
+	}
+	check(bDispatchRecursionGuard == false);
+	bDispatchRecursionGuard = true;
+
+	FGpuFrameData& WriteFrame = GetWriteFrame();
+	FGpuDispatchTimer& DispatchTimer = WriteFrame.DispatchTimers.AddDefaulted_GetRef();
+
+	DispatchTimer.bUniqueInstance = DispatchInstance.SimStageData.bSetDataToRender && (&DispatchInstance.InstanceData == &DispatchInstance.Tick.GetInstances()[0]);
+	DispatchTimer.OwnerComponent = DispatchInstance.InstanceData.Context->ProfilingComponentPtr;
+	DispatchTimer.OwnerEmitter = DispatchInstance.InstanceData.Context->ProfilingEmitterPtr;
+	DispatchTimer.StageName = DispatchInstance.SimStageData.StageMetaData->SimulationStageName;
+	DispatchTimer.StartQuery = QueryPool->AllocateQuery();
+	RHICmdList.EndRenderQuery(DispatchTimer.StartQuery.GetQuery());
+}
+
+void FNiagaraGPUProfiler::BeginDispatch(FRHICommandList& RHICmdList, const struct FNiagaraComputeInstanceData& InstanceData, FName StageName)
+{
+	if (!bProfilingDispatches)
+	{
+		return;
+	}
+	check(bDispatchRecursionGuard == false);
+	bDispatchRecursionGuard = true;
+
+	FGpuFrameData& WriteFrame = GetWriteFrame();
+	FGpuDispatchTimer& DispatchTimer = WriteFrame.DispatchTimers.AddDefaulted_GetRef();
+	DispatchTimer.bUniqueInstance = false;
+	DispatchTimer.OwnerComponent = InstanceData.Context->ProfilingComponentPtr;
+	DispatchTimer.OwnerEmitter = InstanceData.Context->ProfilingEmitterPtr;
+	DispatchTimer.StageName = StageName;
+	DispatchTimer.StartQuery = QueryPool->AllocateQuery();
+	RHICmdList.EndRenderQuery(DispatchTimer.StartQuery.GetQuery());
+}
+
+void FNiagaraGPUProfiler::BeginDispatch(FRHICommandList& RHICmdList, FName StageName)
+{
+	if (!bProfilingDispatches)
+	{
+		return;
+	}
+	check(bDispatchRecursionGuard == false);
+	bDispatchRecursionGuard = true;
+
+	FGpuFrameData& WriteFrame = GetWriteFrame();
+	FGpuDispatchTimer& DispatchTimer = WriteFrame.DispatchTimers.AddDefaulted_GetRef();
+	DispatchTimer.bUniqueInstance = false;
+	DispatchTimer.OwnerComponent = nullptr;
+	DispatchTimer.OwnerEmitter = nullptr;
+	DispatchTimer.StageName = StageName;
+	DispatchTimer.StartQuery = QueryPool->AllocateQuery();
+	RHICmdList.EndRenderQuery(DispatchTimer.StartQuery.GetQuery());
+}
+
+void FNiagaraGPUProfiler::EndDispatch(FRHICommandList& RHICmdList)
+{
+	if (!bProfilingDispatches)
+	{
+		return;
+	}
+	check(bDispatchRecursionGuard == true);
+	bDispatchRecursionGuard = false;
+
+	FGpuFrameData& WriteFrame = GetWriteFrame();
+	FGpuDispatchTimer& DispatchTimer = WriteFrame.DispatchTimers.Last();
+	DispatchTimer.EndQuery = QueryPool->AllocateQuery();
+	RHICmdList.EndRenderQuery(DispatchTimer.EndQuery.GetQuery());
+}
+
+bool FNiagaraGPUProfiler::ProcessFrame(FRHICommandListImmediate& RHICmdList, FGpuFrameData& ReadFrame)
+{
+	// Frame ready to process?
+	//-OPT: We can just look at the last write stage end timer here, but that relies on the batcher always executing
+	uint64 DummyEndTime;
+	if (!RHICmdList.GetRenderQueryResult(ReadFrame.EndQuery.GetQuery(), DummyEndTime, false))
+	{
+		return false;
+	}
+	ReadFrame.EndQuery.ReleaseQuery();
+
+	//-OPT: Potentially pool these
+	FNiagaraGpuFrameResultsPtr FrameResults = MakeShared<FNiagaraGpuFrameResults, ESPMode::ThreadSafe>();
+	FrameResults->OwnerContext = OwnerContext;
+	FrameResults->DispatchResults.Reserve(ReadFrame.DispatchTimers.Num());
+
+	// Process results
+	for (int32 i = 0; i < ENiagaraGpuComputeTickStage::Max; ++i)
+	{
+		FGpuStageTimer& StageTimer = ReadFrame.StageTimers[i];
+		auto& StageResults = FrameResults->StageResults[i];
+		StageResults.NumDispatches = StageTimer.NumDispatches;
+		StageResults.NumDispatchGroups = StageTimer.NumDispatchGroups;
+		StageResults.DurationMicroseconds = 0;
+		if ( StageTimer.StartQuery.GetQuery() != nullptr )
+		{
+			uint64 StartMicroseconds;
+			uint64 EndMicroseconds;
+			ensure(RHICmdList.GetRenderQueryResult(StageTimer.StartQuery.GetQuery(), StartMicroseconds, false));
+			ensure(RHICmdList.GetRenderQueryResult(StageTimer.EndQuery.GetQuery(), EndMicroseconds, false));
+			StageResults.DurationMicroseconds = EndMicroseconds - StartMicroseconds;
+			StageTimer.StartQuery.ReleaseQuery();
+			StageTimer.EndQuery.ReleaseQuery();
+		}
+		StageTimer.NumDispatches = 0;
+		StageTimer.NumDispatchGroups = 0;
+	}
+
+	for (FGpuDispatchTimer& DispatchTimer : ReadFrame.DispatchTimers)
+	{
+		auto& DispatchResults = FrameResults->DispatchResults.AddDefaulted_GetRef();
+			
+		uint64 StartMicroseconds;
+		uint64 EndMicroseconds;
+		ensure(RHICmdList.GetRenderQueryResult(DispatchTimer.StartQuery.GetQuery(), StartMicroseconds, false));
+		ensure(RHICmdList.GetRenderQueryResult(DispatchTimer.EndQuery.GetQuery(), EndMicroseconds, false));
+		DispatchTimer.StartQuery.ReleaseQuery();
+		DispatchTimer.EndQuery.ReleaseQuery();
+
+		DispatchResults.bUniqueInstance = DispatchTimer.bUniqueInstance;
+		DispatchResults.OwnerComponent = DispatchTimer.OwnerComponent;
+		DispatchResults.OwnerEmitter = DispatchTimer.OwnerEmitter;
+		DispatchResults.StageName = DispatchTimer.StageName;
+		DispatchResults.DurationMicroseconds = EndMicroseconds - StartMicroseconds;
+	}
+
+	ReadFrame.DispatchTimers.Empty();
+
+	// Post Results
+	PostResults(FrameResults);
+
+	return true;
 }
 
 #endif
diff --git a/Engine/Plugins/FX/Niagara/Source/Niagara/Private/NiagaraGPUProfiler.h b/Engine/Plugins/FX/Niagara/Source/Niagara/Private/NiagaraGPUProfiler.h
new file mode 100644
--- /dev/null
+++ b/Engine/Plugins/FX/Niagara/Source/Niagara/Private/NiagaraGPUProfiler.h
@@ -0,0 +1,83 @@
+// Copyright Epic Games, Inc. All Rights Reserved.
+
+#pragma once
+
+#include "NiagaraCommon.h"
+#include "RHI.h"
+#include "RHICommandList.h"
+#include "RHIResources.h"
+#include "HAL/Platform.h"
+#include "Misc/Build.h"
+#include "Particles//ParticlePerfStats.h"
+
+#include "NiagaraGPUProfilerInterface.h"
+
+#if WITH_NIAGARA_GPU_PROFILER
+
+/** Helper class to time gpu runtime cost of dispatches */
+class FNiagaraGPUProfiler : public FNiagaraGPUProfilerInterface
+{
+	static constexpr int32 NumBufferFrames = 5;
+
+	struct FGpuStageTimer
+	{
+		int32					NumDispatchGroups = 0;
+		int32					NumDispatches = 0;
+		FRHIPooledRenderQuery	StartQuery;
+		FRHIPooledRenderQuery	EndQuery;
+	};
+
+	struct FGpuDispatchTimer
+	{
+		uint32									bUniqueInstance : 1;
+		TWeakObjectPtr<class USceneComponent>	OwnerComponent;
+		TWeakObjectPtr<class UNiagaraEmitter>	OwnerEmitter;
+		FName									StageName;
+		FRHIPooledRenderQuery					StartQuery;
+		FRHIPooledRenderQuery					EndQuery;
+	};
+
+	struct FGpuFrameData
+	{
+		FRHIPooledRenderQuery		EndQuery;
+		FGpuStageTimer				StageTimers[ENiagaraGpuComputeTickStage::Max];
+		TArray<FGpuDispatchTimer>	DispatchTimers;
+	};
+
+	DECLARE_MULTICAST_DELEGATE_OneParam(FOnFrameResults, const FNiagaraGpuFrameResultsPtr&);
+
+public:
+	FNiagaraGPUProfiler(uintptr_t InOwnerContext);
+	~FNiagaraGPUProfiler();
+
+	void BeginFrame(FRHICommandListImmediate& RHICmdList);
+	void EndFrame(FRHICommandList& RHICmdList);
+
+	void BeginStage(FRHICommandList& RHICmdList, ENiagaraGpuComputeTickStage::Type TickStage, int32 NumDispatchGroups);
+	void EndStage(FRHICommandList& RHICmdList, ENiagaraGpuComputeTickStage::Type TickStage, int32 NumDispatches);
+
+	void BeginDispatch(FRHICommandList& RHICmdList, const struct FNiagaraGpuDispatchInstance& DispatchInstance);
+	void BeginDispatch(FRHICommandList& RHICmdList, const struct FNiagaraComputeInstanceData& InstanceData, FName StageName);
+	void BeginDispatch(FRHICommandList& RHICmdList, FName StageName);
+	void EndDispatch(FRHICommandList& RHICmdList);
+
+private:
+	FGpuFrameData& GetWriteFrame() { return GpuFrames[CurrentWriteFrame]; }
+	bool ProcessFrame(FRHICommandListImmediate& RHICmdList, FGpuFrameData& ReadFrame);
+
+private:
+	uintptr_t				OwnerContext = 0;
+
+	bool					bProfilingFrame = false;
+	bool					bProfilingDispatches = false;
+
+	bool					bDispatchRecursionGuard = false;
+
+	int32					CurrentReadFrame = 0;
+	int32					CurrentWriteFrame = 0;
+	FGpuFrameData			GpuFrames[NumBufferFrames];
+
+	FRenderQueryPoolRHIRef	QueryPool;
+};
+
+#endif //WITH_NIAGARA_GPU_PROFILER
diff --git a/Engine/Plugins/FX/Niagara/Source/Niagara/Private/NiagaraGPUProfilerInterface.cpp b/Engine/Plugins/FX/Niagara/Source/Niagara/Private/NiagaraGPUProfilerInterface.cpp
new file mode 100644
--- /dev/null
+++ b/Engine/Plugins/FX/Niagara/Source/Niagara/Private/NiagaraGPUProfilerInterface.cpp
@@ -0,0 +1,161 @@
+// Copyright Epic Games, Inc. All Rights Reserved.
+
+#include "NiagaraGPUProfilerInterface.h"
+#include "NiagaraGpuComputeDispatchInterface.h"
+#include "NiagaraGpuProfiler.h"
+#include "NiagaraGPUSystemTick.h"
+#include "NiagaraDataInterfaceBase.h"
+#include "NiagaraComponent.h"
+
+#include "Async/Async.h"
+
+#if WITH_NIAGARA_GPU_PROFILER
+
+//////////////////////////////////////////////////////////////////////////
+FNiagaraGpuProfileScope::FNiagaraGpuProfileScope(FRHICommandList& InRHICmdList, const FNiagaraGpuComputeDispatchInterface* ComputeDispatchInterface, FName StageName)
+	: RHICmdList(InRHICmdList)
+	, GPUProfiler(static_cast<FNiagaraGPUProfiler*>(ComputeDispatchInterface->GetGPUProfiler()))
+{
+	if (GPUProfiler)
+	{
+		GPUProfiler->BeginDispatch(RHICmdList, StageName);
+	}
+}
+
+FNiagaraGpuProfileScope::FNiagaraGpuProfileScope(FRHICommandList& InRHICmdList, const FNiagaraGpuComputeDispatchInterface* ComputeDispatchInterface, const struct FNiagaraGpuDispatchInstance& DispatchInstance)
+	: RHICmdList(InRHICmdList)
+	, GPUProfiler(static_cast<FNiagaraGPUProfiler*>(ComputeDispatchInterface->GetGPUProfiler()))
+{
+	if (GPUProfiler)
+	{
+		GPUProfiler->BeginDispatch(RHICmdList, DispatchInstance);
+	}
+}
+
+FNiagaraGpuProfileScope::FNiagaraGpuProfileScope(FRHICommandList& InRHICmdList, const struct FNiagaraDataInterfaceArgs& Context, FName StageName)
+	: RHICmdList(InRHICmdList)
+	, GPUProfiler(static_cast<FNiagaraGPUProfiler*>(Context.ComputeDispatchInterface->GetGPUProfiler()))
+{
+	if (GPUProfiler)
+	{
+		GPUProfiler->BeginDispatch(RHICmdList, StageName);
+	}
+}
+
+FNiagaraGpuProfileScope::FNiagaraGpuProfileScope(FRHICommandList& InRHICmdList, const struct FNiagaraDataInterfaceSetArgs& Context, FName StageName)
+	: RHICmdList(InRHICmdList)
+	, GPUProfiler(static_cast<FNiagaraGPUProfiler*>(Context.ComputeDispatchInterface->GetGPUProfiler()))
+{
+	if (GPUProfiler)
+	{
+		GPUProfiler->BeginDispatch(RHICmdList, *Context.ComputeInstanceData, StageName);
+	}
+}
+
+FNiagaraGpuProfileScope::FNiagaraGpuProfileScope(FRHICommandList& InRHICmdList, const struct FNiagaraDataInterfaceStageArgs& Context, FName StageName)
+	: RHICmdList(InRHICmdList)
+	, GPUProfiler(static_cast<FNiagaraGPUProfiler*>(Context.ComputeDispatchInterface->GetGPUProfiler()))
+{
+	if (GPUProfiler)
+	{
+		GPUProfiler->BeginDispatch(RHICmdList, *Context.ComputeInstanceData, StageName);
+	}
+}
+
+FNiagaraGpuProfileScope::~FNiagaraGpuProfileScope()
+{
+	if (GPUProfiler)
+	{
+		GPUProfiler->EndDispatch(RHICmdList);
+	}
+}
+
+//////////////////////////////////////////////////////////////////////////
+
+void SetResultsForParticlePerfStats(const FNiagaraGpuFrameResultsPtr& FrameResults)
+{
+#if WITH_PARTICLE_PERF_STATS
+	auto PostStats =
+		[](FParticlePerfStats* Stats, bool bUniqueInstance, uint64 DurationMicroseconds)
+		{
+			if ( Stats )
+			{
+				Stats->GetGPUStats().NumInstances += bUniqueInstance ? 1: 0;
+				Stats->GetGPUStats().TotalMicroseconds += DurationMicroseconds;
+			}
+		};
+
+	for (const auto& DispatchResult : FrameResults->DispatchResults)
+	{
+		UNiagaraComponent* OwnerComponent = Cast<UNiagaraComponent>(DispatchResult.OwnerComponent.Get());
+		if ( OwnerComponent == nullptr )
+		{
+			return;
+		}
+
+		FParticlePerfStatsContext StatsContext = OwnerComponent->GetPerfStatsContext();
+		PostStats(StatsContext.GetWorldStats(), DispatchResult.bUniqueInstance, DispatchResult.DurationMicroseconds);
+		PostStats(StatsContext.GetSystemStats(), DispatchResult.bUniqueInstance, DispatchResult.DurationMicroseconds);
+		PostStats(StatsContext.GetComponentStats(), DispatchResult.bUniqueInstance, DispatchResult.DurationMicroseconds);
+	}
+#endif
+}
+
+//////////////////////////////////////////////////////////////////////////
+//-TODO: Move out of this file
+void SetResultsForEditorStats(const FNiagaraGpuFrameResultsPtr& FrameResults)
+{
+#if WITH_NIAGARA_GPU_PROFILER_EDITOR
+	// Send data for editor stat display
+	//-TODO: editor stats should merge with particle perf stats
+	TMap<UNiagaraEmitter*, TMap<TStatIdData const*, float>> CapturedStats;
+	CapturedStats.Reserve(FrameResults->DispatchResults.Num());
+
+	for (const auto& DispatchResult : FrameResults->DispatchResults)
+	{
+		UNiagaraEmitter* Emitter = DispatchResult.OwnerEmitter.Get();
+		if ( Emitter == nullptr )
+		{
+			continue;
+		}
+
+		// Build stat name
+		TStringBuilder<128> StatNameBuilder;
+		StatNameBuilder.Append(TEXT("GPU_Stage_"));
+		DispatchResult.StageName.AppendString(StatNameBuilder);
+		const FName StatFName(StatNameBuilder.ToString());
+
+		const TStatId StatId = FDynamicStats::CreateStatId<FStatGroup_STATGROUP_NiagaraDetailed>(StatFName);
+		CapturedStats.FindOrAdd(Emitter).FindOrAdd(StatId.GetRawPointer()) += DispatchResult.DurationMicroseconds;
+	}
+
+	for (auto& CapturedStat : CapturedStats)
+	{
+		CapturedStat.Key->GetStatData().AddStatCapture(TTuple<uint64, ENiagaraScriptUsage>(uint64(FrameResults->OwnerContext), ENiagaraScriptUsage::ParticleGPUComputeScript), CapturedStat.Value);
+	}
+#endif
+}
+
+//////////////////////////////////////////////////////////////////////////
+
+FNiagaraGPUProfilerInterface::FOnFrameResults FNiagaraGPUProfilerInterface::OnFrameResults_GameThread;
+FNiagaraGPUProfilerInterface::FOnFrameResults FNiagaraGPUProfilerInterface::OnFrameResults_RenderThread;
+
+void FNiagaraGPUProfilerInterface::PostResults(const FNiagaraGpuFrameResultsPtr& FrameResults)
+{
+	// Set data for particle perf stats
+	SetResultsForParticlePerfStats(FrameResults);
+	OnFrameResults_RenderThread.Broadcast(FrameResults);
+
+	// Send data to GameThread
+	AsyncTask(
+		ENamedThreads::GameThread,
+		[FrameResults_GT=FrameResults]
+		{
+			SetResultsForEditorStats(FrameResults_GT);
+			OnFrameResults_GameThread.Broadcast(FrameResults_GT);
+		}
+	);
+}
+
+#endif //WITH_NIAGARA_GPU_PROFILER
diff --git a/Engine/Plugins/FX/Niagara/Source/Niagara/Private/NiagaraGpuComputeDispatch.cpp b/Engine/Plugins/FX/Niagara/Source/Niagara/Private/NiagaraGpuComputeDispatch.cpp
--- a/Engine/Plugins/FX/Niagara/Source/Niagara/Private/NiagaraGpuComputeDispatch.cpp
+++ b/Engine/Plugins/FX/Niagara/Source/Niagara/Private/NiagaraGpuComputeDispatch.cpp
@@ -15,6 +15,7 @@
 #if NIAGARA_COMPUTEDEBUG_ENABLED
 #include "NiagaraGpuComputeDebug.h"
 #endif
+#include "NiagaraGpuProfilerInterface.h"
 #include "NiagaraGpuReadbackManager.h"
 #include "NiagaraRayTracingHelper.h"
 #include "NiagaraRenderer.h"
@@ -215,6 +216,9 @@ FNiagaraGpuComputeDispatch::FNiagaraGpuComputeDispatch(ERHIFeatureLevel::Type In
 
 #if NIAGARA_COMPUTEDEBUG_ENABLED
 	GpuComputeDebugPtr.Reset(new FNiagaraGpuComputeDebug(FeatureLevel));
+#endif
+#if WITH_NIAGARA_GPU_PROFILER
+	GPUProfilerPtr = MakeUnique<FNiagaraGPUProfiler>(uintptr_t(static_cast<FNiagaraGpuComputeDispatchInterface*>(this)));
 #endif
 	GpuReadbackManagerPtr.Reset(new FNiagaraGpuReadbackManager());
 	EmptyUAVPoolPtr.Reset(new FNiagaraEmptyUAVPool());
@@ -1214,6 +1218,11 @@ void FNiagaraGpuComputeDispatch::ExecuteTicks(FRHICommandList& RHICmdList, FRHIU
 	TArray<FRHITransitionInfo, TMemStackAllocator<>> TransitionsAfter;
 	TArray<FNiagaraDataBuffer*, TMemStackAllocator<>> IDToIndexInit;
 
+#if WITH_NIAGARA_GPU_PROFILER
+	GPUProfilerPtr->BeginStage(RHICmdList, TickStage, DispatchList.DispatchGroups.Num());
+	const int32 StageStartTotalDispatches = TotalDispatchesThisFrame;
+#endif
+
 	for ( const FNiagaraGpuDispatchGroup& DispatchGroup : DispatchList.DispatchGroups )
 	{
 		SCOPED_CONDITIONAL_DRAW_EVENT(RHICmdList, NiagaraDispatchGroup, FNiagaraGpuComputeDispatchLocal::GAddDispatchGroupDrawEvent);
@@ -1315,19 +1324,9 @@ void FNiagaraGpuComputeDispatch::ExecuteTicks(FRHICommandList& RHICmdList, FRHIU
 			{
 				ResetDataInterfaces(RHICmdList, DispatchInstance.Tick, DispatchInstance.InstanceData);
 			}
-#if STATS
-			if (GPUProfiler.IsProfilingEnabled() && DispatchInstance.Tick.bIsFinalTick)
-			{			
-				const TStatId StatId = FDynamicStats::CreateStatId<FStatGroup_STATGROUP_NiagaraDetailed>("GPU_Stage_" + DispatchInstance.SimStageData.StageMetaData->SimulationStageName.ToString());
-				const int32 TimerHandle = GPUProfiler.StartTimer((uint64)DispatchInstance.InstanceData.Context, StatId, RHICmdList);
-				DispatchStage(RHICmdList, ViewUniformBuffer, DispatchInstance.Tick, DispatchInstance.InstanceData, DispatchInstance.SimStageData);
-				GPUProfiler.EndTimer(TimerHandle, RHICmdList);
-			}
-			else
-#endif
-			{
-				DispatchStage(RHICmdList, ViewUniformBuffer, DispatchInstance.Tick, DispatchInstance.InstanceData, DispatchInstance.SimStageData);
-			}
+
+			FNiagaraGpuProfileScope GpuProfileDispatchScope(RHICmdList, this, DispatchInstance);
+			DispatchStage(RHICmdList, ViewUniformBuffer, DispatchInstance.Tick, DispatchInstance.InstanceData, DispatchInstance.SimStageData);
 		}
 		RHICmdList.EndUAVOverlap(GPUInstanceCounterManager.GetInstanceCountBuffer().UAV);
 
@@ -1409,6 +1408,11 @@ void FNiagaraGpuComputeDispatch::ExecuteTicks(FRHICommandList& RHICmdList, FRHIU
 		}
 	}
 
+#if WITH_NIAGARA_GPU_PROFILER
+	const int32 StageTotalDispatches = TotalDispatchesThisFrame - StageStartTotalDispatches;
+	GPUProfilerPtr->EndStage(RHICmdList, TickStage, StageTotalDispatches);
+#endif
+
 	// Clear dispatch groups
 	// We do not release the counts as we won't do that until we finish the dispatches
 	DispatchList.DispatchGroups.Empty();
@@ -1576,9 +1580,9 @@ void FNiagaraGpuComputeDispatch::DispatchStage(FRHICommandList& RHICmdList, FRHI
 
 	// Optionally submit commands to the GPU
 	// This can be used to avoid accidental TDR detection in the editor especially when issuing multiple ticks in the same frame
+	++TotalDispatchesThisFrame;
 	if (GNiagaraGpuSubmitCommandHint > 0)
 	{
-		++TotalDispatchesThisFrame;
 		if ((TotalDispatchesThisFrame % GNiagaraGpuSubmitCommandHint) == 0)
 		{
 			RHICmdList.SubmitCommandsHint();
@@ -1599,46 +1603,24 @@ void FNiagaraGpuComputeDispatch::PreInitViews(FRDGBuilder& GraphBuilder, bool bA
 	}
 #endif
 
-#if STATS
-	// check if we can process profiling results from previous frames
-	if (GPUProfiler.IsProfilingEnabled())
+	LLM_SCOPE(ELLMTag::Niagara);
+	TotalDispatchesThisFrame = 0;
+
+	// Add pass to begin the gpu profiler frame
+#if WITH_NIAGARA_GPU_PROFILER
 	{
-		TArray<FNiagaraGPUTimingResult, TInlineAllocator<16>> ProfilingResults;
-		GPUProfiler.QueryTimingResults(GraphBuilder.RHICmdList, ProfilingResults);
-		TMap<TWeakObjectPtr<UNiagaraEmitter>, TMap<TStatIdData const*, float>> CapturedStats;
-		for (const FNiagaraGPUTimingResult& Result : ProfilingResults)
-		{
-			// map the individual timings to the emitters and source scripts
-			if (Result.ReporterHandle)
+		int32* NullParameter = GraphBuilder.AllocParameters<int32>();
+		GraphBuilder.AddPass(
+			RDG_EVENT_NAME("Niagara::GPUProfiler"),
+			ERDGPassFlags::None,
+			[this](FRHICommandListImmediate& RHICmdList)
 			{
-				FNiagaraComputeExecutionContext* Context = (FNiagaraComputeExecutionContext*)Result.ReporterHandle;
-				if (Context->EmitterPtr.IsValid())
-				{
-					CapturedStats.FindOrAdd(Context->EmitterPtr).FindOrAdd(Result.StatId.GetRawPointer()) += Result.ElapsedMicroseconds;
-				}
+				GPUProfilerPtr->BeginFrame(RHICmdList);
 			}
-		}
-		if (CapturedStats.Num() > 0)
-		{
-			// report the captured gpu stats on the game thread
-			uint64 ReporterHandle = (uint64)this;
-			AsyncTask(ENamedThreads::GameThread, [LocalCapturedStats = MoveTemp(CapturedStats), ReporterHandle]
-            {
-                for (auto& Entry : LocalCapturedStats)
-                {
-                    if (Entry.Key.IsValid())
-                    {
-                        Entry.Key->GetStatData().AddStatCapture(TTuple<uint64, ENiagaraScriptUsage>(ReporterHandle, ENiagaraScriptUsage::ParticleGPUComputeScript), Entry.Value);
-                    }
-                }
-            });
-		}
+		);
 	}
 #endif
 
-	LLM_SCOPE(ELLMTag::Niagara);
-	TotalDispatchesThisFrame = 0;
-
 	// Reset the list of GPUSort tasks and release any resources they hold on to.
 	// It might be worth considering doing so at the end of the render to free the resources immediately.
 	// (note that currently there are no callback appropriate to do it)
@@ -1765,6 +1747,9 @@ void FNiagaraGpuComputeDispatch::PostRenderOpaque(FRDGBuilder& GraphBuilder, TCo
 			GPUInstanceCounterManager.EnqueueGPUReadback(RHICmdList);
 			bRequiresReadback = false;
 		}
+#if WITH_NIAGARA_GPU_PROFILER
+		GPUProfilerPtr->EndFrame(RHICmdList);
+#endif
 	});
 
 	GNiagaraViewDataManager.ClearSceneTextureParameters();
diff --git a/Engine/Plugins/FX/Niagara/Source/Niagara/Private/NiagaraGpuComputeDispatch.h b/Engine/Plugins/FX/Niagara/Source/Niagara/Private/NiagaraGpuComputeDispatch.h
--- a/Engine/Plugins/FX/Niagara/Source/Niagara/Private/NiagaraGpuComputeDispatch.h
+++ b/Engine/Plugins/FX/Niagara/Source/Niagara/Private/NiagaraGpuComputeDispatch.h
@@ -115,6 +115,10 @@ public:
 	/** Get the shared SortManager, used in the rendering loop to call FGPUSortManager::OnPreRender() and FGPUSortManager::OnPostRenderOpaque() */
 	virtual FGPUSortManager* GetGPUSortManager() const override;
 
+#if WITH_NIAGARA_GPU_PROFILER
+	virtual FNiagaraGPUProfilerInterface* GetGPUProfiler() const final { return GPUProfilerPtr.Get(); }
+#endif
+
 	/** When inside a pass returns the active scene textures for the pass, can be nullptr. */
 	class FNiagaraSceneTextureParameters* GetNiagaraSceneTextures() const { return NiagaraSceneTextures; }
 
@@ -186,8 +190,8 @@ private:
 	TUniquePtr<FNiagaraRayTracingHelper> RayTracingHelper;
 #endif
 
-#if STATS
-	FNiagaraGPUProfiler GPUProfiler;
+#if WITH_NIAGARA_GPU_PROFILER
+	TUniquePtr<FNiagaraGPUProfiler> GPUProfilerPtr;
 #endif
 
 	// persistent layouts used to create the constant buffers for the compute sim shader
diff --git a/Engine/Plugins/FX/Niagara/Source/Niagara/Private/NiagaraSimulationStageBase.cpp b/Engine/Plugins/FX/Niagara/Source/Niagara/Private/NiagaraSimulationStageBase.cpp
--- a/Engine/Plugins/FX/Niagara/Source/Niagara/Private/NiagaraSimulationStageBase.cpp
+++ b/Engine/Plugins/FX/Niagara/Source/Niagara/Private/NiagaraSimulationStageBase.cpp
@@ -5,6 +5,8 @@
 #include "NiagaraSystem.h"
 #include "NiagaraScriptSourceBase.h"
 
+const FName UNiagaraSimulationStageBase::ParticleSpawnUpdateName("ParticleSpawnUpdate");
+
 bool UNiagaraSimulationStageBase::AppendCompileHash(FNiagaraCompileHashVisitor* InVisitor) const
 {
 #if WITH_EDITORONLY_DATA
diff --git a/Engine/Plugins/FX/Niagara/Source/Niagara/Public/NiagaraCommon.h b/Engine/Plugins/FX/Niagara/Source/Niagara/Public/NiagaraCommon.h
--- a/Engine/Plugins/FX/Niagara/Source/Niagara/Public/NiagaraCommon.h
+++ b/Engine/Plugins/FX/Niagara/Source/Niagara/Public/NiagaraCommon.h
@@ -5,12 +5,14 @@
 #include "CoreMinimal.h"
 #include "Engine/EngineBaseTypes.h"
 #include "Engine/TextureRenderTarget2D.h"
-#include "UObject/ObjectMacros.h"
-#include "NiagaraTypes.h"
 #include "Engine/BlueprintGeneratedClass.h"
-#include "NiagaraCore.h"
+#include "Particles/ParticlePerfStats.h"
 #include "UObject/ObjectKey.h"
+#include "UObject/ObjectMacros.h"
 #include "UObject/WeakFieldPtr.h"
+
+#include "NiagaraTypes.h"
+#include "NiagaraCore.h"
 #include "NiagaraCommon.generated.h"
 
 class UNiagaraComponent;
@@ -32,6 +34,9 @@ struct FNiagaraParameterStore;
 	#define NIAGARA_COMPUTEDEBUG_ENABLED WITH_EDITOR
 #endif
 
+#define WITH_NIAGARA_GPU_PROFILER_EDITOR (WITH_EDITOR && STATS)
+#define WITH_NIAGARA_GPU_PROFILER ((WITH_PARTICLE_PERF_STATS && !UE_BUILD_SHIPPING) || WITH_NIAGARA_GPU_PROFILER_EDITOR)
+
 #define INTERPOLATED_PARAMETER_PREFIX TEXT("PREV_")
 
 /** Defines The maximum ThreadGroup size we allow in Niagara.  This is important for how memory is allocated as we always need to round this and the final instance is used to avoid overflowing the buffer. */
diff --git a/Engine/Plugins/FX/Niagara/Source/Niagara/Public/NiagaraDebuggerCommon.h b/Engine/Plugins/FX/Niagara/Source/Niagara/Public/NiagaraDebuggerCommon.h
--- a/Engine/Plugins/FX/Niagara/Source/Niagara/Public/NiagaraDebuggerCommon.h
+++ b/Engine/Plugins/FX/Niagara/Source/Niagara/Public/NiagaraDebuggerCommon.h
@@ -385,6 +385,7 @@ enum class ENiagaraDebugHUDOverviewMode
 	Overview,
 	Scalability,	
 	Performance,
+	GpuComputePerformance,
 }; 
 
 UENUM()
@@ -393,7 +394,7 @@ enum class ENiagaraDebugHUDPerfGraphMode
 	None,
 	GameThread,
 	RenderThread,
-	//GPU,
+	GPU,
 }; 
 
 UENUM()
diff --git a/Engine/Plugins/FX/Niagara/Source/Niagara/Public/NiagaraGpuComputeDispatchInterface.h b/Engine/Plugins/FX/Niagara/Source/Niagara/Public/NiagaraGpuComputeDispatchInterface.h
--- a/Engine/Plugins/FX/Niagara/Source/Niagara/Public/NiagaraGpuComputeDispatchInterface.h
+++ b/Engine/Plugins/FX/Niagara/Source/Niagara/Public/NiagaraGpuComputeDispatchInterface.h
@@ -56,6 +56,11 @@ public:
 	FNiagaraGpuComputeDebug* GetGpuComputeDebug() const { return GpuComputeDebugPtr.Get(); }
 #endif
 
+#if WITH_NIAGARA_GPU_PROFILER
+	/** Access to Niagara's GPU Profiler */
+	virtual class FNiagaraGPUProfilerInterface* GetGPUProfiler() const = 0;
+#endif
+
 	/** Get access to Niagara's GpuReadbackManager. */
 	FNiagaraGpuReadbackManager* GetGpuReadbackManager() const { return GpuReadbackManagerPtr.Get(); }
 
diff --git a/Engine/Plugins/FX/Niagara/Source/Niagara/Public/NiagaraSimulationStageBase.h b/Engine/Plugins/FX/Niagara/Source/Niagara/Public/NiagaraSimulationStageBase.h
--- a/Engine/Plugins/FX/Niagara/Source/Niagara/Public/NiagaraSimulationStageBase.h
+++ b/Engine/Plugins/FX/Niagara/Source/Niagara/Public/NiagaraSimulationStageBase.h
@@ -20,6 +20,8 @@ class NIAGARA_API UNiagaraSimulationStageBase : public UNiagaraMergeable
 	GENERATED_BODY()
 
 public:
+	static const FName ParticleSpawnUpdateName;
+
 	UNiagaraSimulationStageBase()
 	{
 		bEnabled = true;
diff --git a/Engine/Plugins/FX/Niagara/Source/NiagaraEditor/Private/NiagaraHlslTranslator.cpp b/Engine/Plugins/FX/Niagara/Source/NiagaraEditor/Private/NiagaraHlslTranslator.cpp
--- a/Engine/Plugins/FX/Niagara/Source/NiagaraEditor/Private/NiagaraHlslTranslator.cpp
+++ b/Engine/Plugins/FX/Niagara/Source/NiagaraEditor/Private/NiagaraHlslTranslator.cpp
@@ -1165,7 +1165,7 @@ const FNiagaraTranslateResults &FHlslNiagaraTranslator::Translate(const FNiagara
 		// Add the spawn / update stage
 		{
 			FSimulationStageMetaData& SimulationStageMetaData = CompilationOutput.ScriptData.SimulationStageMetaData.AddDefaulted_GetRef();
-			SimulationStageMetaData.SimulationStageName = FName("ParticleSpawnUpdate");
+			SimulationStageMetaData.SimulationStageName = UNiagaraSimulationStageBase::ParticleSpawnUpdateName;
 			SimulationStageMetaData.NumIterations = 1;
 			SimulationStageMetaData.bWritesParticles = true;
 			SimulationStageMetaData.bPartialParticleUpdate = false;
diff --git a/Engine/Plugins/FX/Niagara/Source/NiagaraEditorWidgets/Private/Stack/SNiagaraStackRowPerfWidget.cpp b/Engine/Plugins/FX/Niagara/Source/NiagaraEditorWidgets/Private/Stack/SNiagaraStackRowPerfWidget.cpp
--- a/Engine/Plugins/FX/Niagara/Source/NiagaraEditorWidgets/Private/Stack/SNiagaraStackRowPerfWidget.cpp
+++ b/Engine/Plugins/FX/Niagara/Source/NiagaraEditorWidgets/Private/Stack/SNiagaraStackRowPerfWidget.cpp
@@ -123,7 +123,8 @@ void SNiagaraStackRowPerfWidget::Tick(const FGeometry& AllottedGeometry, const d
 	{
 		if (GetUsage() == ENiagaraScriptUsage::ParticleUpdateScript)
 		{
-			GroupOverallTime = CalculateGroupOverallTime("GPU_Stage_SpawnUpdate");
+			static const FString GpuStatName = TEXT("GPU_Stage_") + UNiagaraSimulationStageBase::ParticleSpawnUpdateName.ToString();
+			GroupOverallTime = CalculateGroupOverallTime(GpuStatName);
 		}
 		if (GetUsage() == ENiagaraScriptUsage::ParticleSimulationStageScript)
 		{
diff --git a/Engine/Source/Runtime/Engine/Classes/Particles/ParticlePerfStats.h b/Engine/Source/Runtime/Engine/Classes/Particles/ParticlePerfStats.h
--- a/Engine/Source/Runtime/Engine/Classes/Particles/ParticlePerfStats.h
+++ b/Engine/Source/Runtime/Engine/Classes/Particles/ParticlePerfStats.h
@@ -112,8 +112,25 @@ struct ENGINE_API FParticlePerfStats_RT
 		RenderUpdateCycles = 0;
 		GetDynamicMeshElementsCycles = 0;
 	}
-	FORCEINLINE uint64 GetTotalCycles()const { return RenderUpdateCycles + GetDynamicMeshElementsCycles; }
-	FORCEINLINE uint64 GetPerInstanceAvgCycles()const { return NumInstances > 0 ? (RenderUpdateCycles + GetDynamicMeshElementsCycles) / NumInstances : 0; }
+	FORCEINLINE uint64 GetTotalCycles() const { return RenderUpdateCycles + GetDynamicMeshElementsCycles; }
+	FORCEINLINE uint64 GetPerInstanceAvgCycles() const { return NumInstances > 0 ? (RenderUpdateCycles + GetDynamicMeshElementsCycles) / NumInstances : 0; }
+};
+
+/** Stats gathered from the GPU */
+struct ENGINE_API FParticlePerfStats_GPU
+{
+	uint64 NumInstances = 0;
+	uint64 TotalMicroseconds = 0;
+
+	FORCEINLINE uint64 GetTotalMicroseconds() const { return TotalMicroseconds; }
+	FORCEINLINE uint64 GetPerInstanceAvgMicroseconds() const { return NumInstances > 0 ? GetTotalMicroseconds() / NumInstances : 0; }
+
+	FParticlePerfStats_GPU() { Reset(); }
+	FORCEINLINE void Reset()
+	{
+		NumInstances = 0;
+		TotalMicroseconds = 0;
+	}
 };
 
 struct ENGINE_API FParticlePerfStats
@@ -198,6 +215,9 @@ struct ENGINE_API FParticlePerfStats
 	/** Stats on RT work. */
 	FParticlePerfStats_RT RenderThreadStats;
 
+	/** Stats from GPU work. */
+	FParticlePerfStats_GPU GPUStats;
+
 	/** Returns the current frame Game Thread stats. */
 	FORCEINLINE FParticlePerfStats_GT& GetGameThreadStats()
 	{
@@ -210,8 +230,13 @@ struct ENGINE_API FParticlePerfStats
 		return RenderThreadStats;
 	}
 
-private:
+	/** Returns the current frame GPU stats. */
+	FORCEINLINE FParticlePerfStats_GPU& GetGPUStats()
+	{
+		return GPUStats;
+	}
 
+private:
 	static FParticlePerfStats* GetWorldPerfStats(const UWorld* World);
 	static FParticlePerfStats* GetSystemPerfStats(const UFXSystemAsset* FXAsset);
 	static FParticlePerfStats* GetComponentPerfStats(const UFXSystemComponent* FXComponent);
diff --git a/Engine/Source/Runtime/Engine/Classes/Particles/ParticlePerfStatsManager.h b/Engine/Source/Runtime/Engine/Classes/Particles/ParticlePerfStatsManager.h
--- a/Engine/Source/Runtime/Engine/Classes/Particles/ParticlePerfStatsManager.h
+++ b/Engine/Source/Runtime/Engine/Classes/Particles/ParticlePerfStatsManager.h
@@ -60,7 +60,7 @@ struct ENGINE_API FAccumulatedParticlePerfStats_RT
 	TArray<uint64, TInlineAllocator<ACCUMULATED_PARTICLE_PERF_STAT_MAX_SAMPLES>>	MaxPerInstanceCycles;
 
 	FAccumulatedParticlePerfStats_RT();
-	FORCEINLINE void Reset();
+	void Reset();
 	void Tick(FParticlePerfStats& Stats);
 
 	/** Returns the total cycles used by all RenderThread stats. */
@@ -87,6 +87,28 @@ struct ENGINE_API FAccumulatedParticlePerfStats_RT
 	FORCEINLINE double GetPerInstanceMax(int32 Index = 0)const { return FPlatformTime::ToMilliseconds64(GetPerInstanceMaxCycles(Index)) * 1000.0; }
 };
 
+struct ENGINE_API FAccumulatedParticlePerfStats_GPU
+{
+	uint32 NumFrames;
+	FParticlePerfStats_GPU AccumulatedStats;
+
+	TArray<uint64, TInlineAllocator<ACCUMULATED_PARTICLE_PERF_STAT_MAX_SAMPLES>>	MaxPerFrameTotalMicroseconds;
+	TArray<uint64, TInlineAllocator<ACCUMULATED_PARTICLE_PERF_STAT_MAX_SAMPLES>>	MaxPerInstanceMicroseconds;
+
+	FAccumulatedParticlePerfStats_GPU();
+	void Reset();
+	void Tick(FParticlePerfStats& Stats);
+
+	/** Returns the total microseconds used by GPU. */
+	FORCEINLINE uint64 GetTotalMicroseconds() const { return AccumulatedStats.GetTotalMicroseconds(); }
+
+	FORCEINLINE uint64 GetPerFrameAvgMicroseconds() const { return NumFrames > 0 ? AccumulatedStats.GetTotalMicroseconds() / NumFrames : 0; }
+	FORCEINLINE uint64 GetPerFrameMaxMicroseconds(int32 Index = 0) const { return MaxPerFrameTotalMicroseconds[Index]; }
+
+	FORCEINLINE uint64 GetPerInstanceAvgMicroseconds()const { return AccumulatedStats.GetPerInstanceAvgMicroseconds(); }
+	FORCEINLINE uint64 GetPerInstanceMaxMicroseconds(int32 Index = 0) const { return MaxPerInstanceMicroseconds[Index]; }
+};
+
 /** Utility class for accumulating many frames worth of stats data. */
 struct ENGINE_API FAccumulatedParticlePerfStats
 {
@@ -100,6 +122,7 @@ struct ENGINE_API FAccumulatedParticlePerfStats
 
 	FAccumulatedParticlePerfStats_GT GameThreadStats;
 	FAccumulatedParticlePerfStats_RT RenderThreadStats;
+	FAccumulatedParticlePerfStats_GPU GPUStats;
 
 	static void AddMax(TArray<uint64, TInlineAllocator<ACCUMULATED_PARTICLE_PERF_STAT_MAX_SAMPLES>>& MaxArray, int64 NewValue);
 	static void ResetMaxArray(TArray<uint64, TInlineAllocator<ACCUMULATED_PARTICLE_PERF_STAT_MAX_SAMPLES>>& MaxArray);
@@ -110,26 +133,41 @@ struct ENGINE_API FAccumulatedParticlePerfStats
 		return GameThreadStats;
 	}
 
-	/** Returns the RT stats. Must be called on the Render Thread. */
-	FORCEINLINE FAccumulatedParticlePerfStats_RT& GetRenderThreadStats()
+	/** Returns the RenderThread stats with an optional flush when on the GameThread. */
+	FORCEINLINE FAccumulatedParticlePerfStats_RT& GetRenderThreadStats(bool bFlushForGameThread = false)
 	{
+		if ( IsInGameThread() )
+		{
+			if ( bFlushForGameThread )
+			{
+				FlushRenderingCommands();
+			}
+		}
+		else
+		{
+			ensure(IsInRenderingThread());
+		}
 		return RenderThreadStats;
 	}
 
-	/** Returns the RenderThread stats for use on the GameThread. Optional sync with the RenderThread. */
-	FORCEINLINE const FAccumulatedParticlePerfStats_RT& GetRenderThreadStats_GameThread(bool bSyncRT=false)const
+	/** Returns the GPU stats with an optional flush when on the GameThread. */
+	FORCEINLINE FAccumulatedParticlePerfStats_GPU& GetGPUStats(bool bFlushForGameThread = false)
 	{
-		if (bSyncRT)
+		if (IsInGameThread())
 		{
-			FlushRenderingCommands();
+			if (bFlushForGameThread)
+			{
+				FlushRenderingCommands();
+			}
 		}
-		return RenderThreadStats;
+		else
+		{
+			ensure(IsInRenderingThread());
+		}
+		return GPUStats;
 	}
 };
 
-// ENGINE_API FString ToStringGT(FAccumulatedParticlePerfStats& Stats);
-// ENGINE_API FString ToStringRT(FAccumulatedParticlePerfStats& Stats);
-
 class ENGINE_API FParticlePerfStatsListener : public TSharedFromThis<FParticlePerfStatsListener, ESPMode::ThreadSafe>
 {
 public:
@@ -266,10 +304,10 @@ public:
 
 	virtual ~FParticlePerfStatsListener_GatherAll() {}
 
-	virtual void Begin()override;
-	virtual void End()override;
-	virtual bool Tick()override;
-	virtual void TickRT()override;
+	virtual void Begin() override;
+	virtual void End() override;
+	virtual bool Tick() override;
+	virtual void TickRT() override;
 
 	virtual void OnAddWorld(const TWeakObjectPtr<const UWorld>& NewWorld)override;
 	virtual void OnRemoveWorld(const TWeakObjectPtr<const UWorld>& World)override;
diff --git a/Engine/Source/Runtime/Engine/Classes/Particles/ParticleSystem.h b/Engine/Source/Runtime/Engine/Classes/Particles/ParticleSystem.h
--- a/Engine/Source/Runtime/Engine/Classes/Particles/ParticleSystem.h
+++ b/Engine/Source/Runtime/Engine/Classes/Particles/ParticleSystem.h
@@ -141,6 +141,8 @@ public:
 	FName CSVStat_InstAvgGT = NAME_None;
 	FName CSVStat_RT = NAME_None;
 	FName CSVStat_InstAvgRT = NAME_None;
+	FName CSVStat_GPU = NAME_None;
+	FName CSVStat_InstAvgGPU = NAME_None;
 	FName CSVStat_Activation = NAME_None;
 	FName CSVStat_Waits = NAME_None;
 	FName CSVStat_Culled = NAME_None;
diff --git a/Engine/Source/Runtime/Engine/Private/Particles/ParticleComponents.cpp b/Engine/Source/Runtime/Engine/Private/Particles/ParticleComponents.cpp
--- a/Engine/Source/Runtime/Engine/Private/Particles/ParticleComponents.cpp
+++ b/Engine/Source/Runtime/Engine/Private/Particles/ParticleComponents.cpp
@@ -143,6 +143,8 @@ void UFXSystemAsset::PostInitProperties()
 	CSVStat_InstAvgGT = *FString::Printf(TEXT("InstAvgGT/%s"), *GetFName().ToString());
 	CSVStat_RT = *FString::Printf(TEXT("RT/%s"), *GetFName().ToString());
 	CSVStat_InstAvgRT = *FString::Printf(TEXT("InstAvgRT/%s"), *GetFName().ToString());
+	CSVStat_GPU = *FString::Printf(TEXT("GPU/%s"), *GetFName().ToString());
+	CSVStat_InstAvgGPU = *FString::Printf(TEXT("InstAvgGPU/%s"), *GetFName().ToString());
 	CSVStat_Count = *FString::Printf(TEXT("Count/%s"), *GetFName().ToString());
 	CSVStat_Activation = *FString::Printf(TEXT("Activation/%s"), *GetFName().ToString());
 	CSVStat_Waits = *FString::Printf(TEXT("Waits/%s"), *GetFName().ToString());
diff --git a/Engine/Source/Runtime/Engine/Private/Particles/ParticlePerfStats.cpp b/Engine/Source/Runtime/Engine/Private/Particles/ParticlePerfStats.cpp
--- a/Engine/Source/Runtime/Engine/Private/Particles/ParticlePerfStats.cpp
+++ b/Engine/Source/Runtime/Engine/Private/Particles/ParticlePerfStats.cpp
@@ -541,6 +541,7 @@ void FParticlePerfStats::ResetRT()
 {
 	check(IsInActualRenderingThread());
 	GetRenderThreadStats().Reset();
+	GetGPUStats().Reset();
 }
 
 FParticlePerfStats::FParticlePerfStats()
@@ -580,6 +581,7 @@ void FParticlePerfStats::TickRT()
 {	
 	check(IsInRenderingThread());
 	GetRenderThreadStats().Reset();
+	GetGPUStats().Reset();
 }
 
 //////////////////////////////////////////////////////////////////////////
@@ -648,6 +650,36 @@ void FAccumulatedParticlePerfStats_RT::Tick(FParticlePerfStats& Stats)
 
 //////////////////////////////////////////////////////////////////////////
 
+FAccumulatedParticlePerfStats_GPU::FAccumulatedParticlePerfStats_GPU()
+{
+	Reset();
+}
+
+void FAccumulatedParticlePerfStats_GPU::Reset()
+{
+	FAccumulatedParticlePerfStats::ResetMaxArray(MaxPerFrameTotalMicroseconds);
+	FAccumulatedParticlePerfStats::ResetMaxArray(MaxPerInstanceMicroseconds);
+
+	NumFrames = 0;
+	AccumulatedStats.Reset();
+}
+
+void FAccumulatedParticlePerfStats_GPU::Tick(FParticlePerfStats& Stats)
+{
+	FParticlePerfStats_GPU& GPUStats = Stats.GetGPUStats();
+	if (GPUStats.NumInstances > 0)
+	{
+		++NumFrames;
+		AccumulatedStats.NumInstances += GPUStats.NumInstances;
+		AccumulatedStats.TotalMicroseconds += GPUStats.TotalMicroseconds;
+
+		FAccumulatedParticlePerfStats::AddMax(MaxPerFrameTotalMicroseconds, GPUStats.GetTotalMicroseconds());
+		FAccumulatedParticlePerfStats::AddMax(MaxPerInstanceMicroseconds, GPUStats.GetPerInstanceAvgMicroseconds());
+	}
+}
+
+//////////////////////////////////////////////////////////////////////////
+
 FAccumulatedParticlePerfStats::FAccumulatedParticlePerfStats()
 {
 	ResetGT();
@@ -662,6 +694,7 @@ void FAccumulatedParticlePerfStats::ResetGT()
 void FAccumulatedParticlePerfStats::ResetRT()
 {
 	RenderThreadStats.Reset();
+	GPUStats.Reset();
 }
 
 void FAccumulatedParticlePerfStats::Reset(bool bSyncWithRT)
@@ -696,6 +729,7 @@ void FAccumulatedParticlePerfStats::TickRT(FParticlePerfStats& Stats)
 {
 	check(IsInActualRenderingThread());
 	RenderThreadStats.Tick(Stats);
+	GPUStats.Tick(Stats);
 }
 
 void FAccumulatedParticlePerfStats::AddMax(TArray<uint64, TInlineAllocator<ACCUMULATED_PARTICLE_PERF_STAT_MAX_SAMPLES>>& MaxArray, int64 NewValue)
@@ -899,48 +933,48 @@ void FParticlePerfStatsListener_GatherAll::DumpStatsToDevice(FOutputDevice& Ar)
 			return;
 		}
 
-		const FAccumulatedParticlePerfStats_GT& GTSTats = PerfStats->GetGameThreadStats();
-		const FAccumulatedParticlePerfStats_RT& RTSTats = PerfStats->GetRenderThreadStats_GameThread();
+		const FAccumulatedParticlePerfStats_GT& GTStats = PerfStats->GetGameThreadStats();
+		const FAccumulatedParticlePerfStats_RT& RTStats = PerfStats->GetRenderThreadStats();
 
-		if ((GTSTats.NumFrames == 0 && RTSTats.NumFrames == 0) || (GTSTats.AccumulatedStats.NumInstances == 0 && RTSTats.AccumulatedStats.NumInstances == 0))
+		if ((GTStats.NumFrames == 0 && RTStats.NumFrames == 0) || (GTStats.AccumulatedStats.NumInstances == 0 && RTStats.AccumulatedStats.NumInstances == 0))
 		{
 			return;
 		}
 
-		const uint64 TotalGameThread = GTSTats.GetTotalCycles();
-		const uint64 TotalRenderThread = RTSTats.GetTotalCycles();
+		const uint64 TotalGameThread = GTStats.GetTotalCycles();
+		const uint64 TotalRenderThread = RTStats.GetTotalCycles();
 
-		const uint64 MaxPerFrameTotalGameThreadFirst = GTSTats.MaxPerFrameTotalCycles.Num() > 0 ? GTSTats.MaxPerFrameTotalCycles[0] : 0;
-		const uint64 MaxPerFrameTotalGameThreadLast = GTSTats.MaxPerFrameTotalCycles.Num() > 0 ? GTSTats.MaxPerFrameTotalCycles.Last() : 0;
+		const uint64 MaxPerFrameTotalGameThreadFirst = GTStats.MaxPerFrameTotalCycles.Num() > 0 ? GTStats.MaxPerFrameTotalCycles[0] : 0;
+		const uint64 MaxPerFrameTotalGameThreadLast = GTStats.MaxPerFrameTotalCycles.Num() > 0 ? GTStats.MaxPerFrameTotalCycles.Last() : 0;
 
-		const uint64 MaxPerFrameTotalRenderThreadFirst = RTSTats.MaxPerFrameTotalCycles.Num() > 0 ? RTSTats.MaxPerFrameTotalCycles[0] : 0;
-		const uint64 MaxPerFrameTotalRenderThreadLast = RTSTats.MaxPerFrameTotalCycles.Num() > 0 ? RTSTats.MaxPerFrameTotalCycles.Last() : 0;
+		const uint64 MaxPerFrameTotalRenderThreadFirst = RTStats.MaxPerFrameTotalCycles.Num() > 0 ? RTStats.MaxPerFrameTotalCycles[0] : 0;
+		const uint64 MaxPerFrameTotalRenderThreadLast = RTStats.MaxPerFrameTotalCycles.Num() > 0 ? RTStats.MaxPerFrameTotalCycles.Last() : 0;
 		//TODO: Add per instance max?
 
 		tempString.Reset();
 		tempString.Appendf(TEXT(",%s"), *Name);
-		tempString.Appendf(TEXT(",%u"), uint32(FPlatformTime::ToMilliseconds64(GTSTats.GetPerFrameAvgCycles()) * 1000.0));
-		tempString.Appendf(TEXT(",%u"), uint32(FPlatformTime::ToMilliseconds64(GTSTats.GetPerInstanceAvgCycles()) * 1000.0));
-		tempString.Appendf(TEXT(",%u"), uint32(FPlatformTime::ToMilliseconds64(RTSTats.GetPerFrameAvgCycles()) * 1000.0));
-		tempString.Appendf(TEXT(",%u"), uint32(FPlatformTime::ToMilliseconds64(RTSTats.GetPerInstanceAvgCycles()) * 1000.0));
-		tempString.Appendf(TEXT(",%u"), uint32(GTSTats.NumFrames));
-		tempString.Appendf(TEXT(",%u"), uint32(GTSTats.AccumulatedStats.NumInstances));
-		tempString.Appendf(TEXT(",%u"), uint32(FPlatformTime::ToMilliseconds64(GTSTats.AccumulatedStats.TickGameThreadCycles) * 1000.0));
-		tempString.Appendf(TEXT(",%u"), uint32(FPlatformTime::ToMilliseconds64(GTSTats.AccumulatedStats.TickConcurrentCycles) * 1000.0));
-		tempString.Appendf(TEXT(",%u"), uint32(FPlatformTime::ToMilliseconds64(GTSTats.AccumulatedStats.FinalizeCycles) * 1000.0));
-		tempString.Appendf(TEXT(",%u"), uint32(FPlatformTime::ToMilliseconds64(GTSTats.AccumulatedStats.EndOfFrameCycles) * 1000.0));
-		tempString.Appendf(TEXT(",%u"), uint32(FPlatformTime::ToMilliseconds64(RTSTats.AccumulatedStats.RenderUpdateCycles) * 1000.0));
-		tempString.Appendf(TEXT(",%u"), uint32(FPlatformTime::ToMilliseconds64(RTSTats.AccumulatedStats.GetDynamicMeshElementsCycles) * 1000.0));
-
-		tempString.Appendf(TEXT(",%u,[ "), uint32(FPlatformTime::ToMilliseconds64(GTSTats.MaxPerFrameTotalCycles.Num() > 0 ? GTSTats.MaxPerFrameTotalCycles[0] : 0) * 1000.0));
-		for (uint64 v : GTSTats.MaxPerFrameTotalCycles)
+		tempString.Appendf(TEXT(",%u"), uint32(FPlatformTime::ToMilliseconds64(GTStats.GetPerFrameAvgCycles()) * 1000.0));
+		tempString.Appendf(TEXT(",%u"), uint32(FPlatformTime::ToMilliseconds64(GTStats.GetPerInstanceAvgCycles()) * 1000.0));
+		tempString.Appendf(TEXT(",%u"), uint32(FPlatformTime::ToMilliseconds64(RTStats.GetPerFrameAvgCycles()) * 1000.0));
+		tempString.Appendf(TEXT(",%u"), uint32(FPlatformTime::ToMilliseconds64(RTStats.GetPerInstanceAvgCycles()) * 1000.0));
+		tempString.Appendf(TEXT(",%u"), uint32(GTStats.NumFrames));
+		tempString.Appendf(TEXT(",%u"), uint32(GTStats.AccumulatedStats.NumInstances));
+		tempString.Appendf(TEXT(",%u"), uint32(FPlatformTime::ToMilliseconds64(GTStats.AccumulatedStats.TickGameThreadCycles) * 1000.0));
+		tempString.Appendf(TEXT(",%u"), uint32(FPlatformTime::ToMilliseconds64(GTStats.AccumulatedStats.TickConcurrentCycles) * 1000.0));
+		tempString.Appendf(TEXT(",%u"), uint32(FPlatformTime::ToMilliseconds64(GTStats.AccumulatedStats.FinalizeCycles) * 1000.0));
+		tempString.Appendf(TEXT(",%u"), uint32(FPlatformTime::ToMilliseconds64(GTStats.AccumulatedStats.EndOfFrameCycles) * 1000.0));
+		tempString.Appendf(TEXT(",%u"), uint32(FPlatformTime::ToMilliseconds64(RTStats.AccumulatedStats.RenderUpdateCycles) * 1000.0));
+		tempString.Appendf(TEXT(",%u"), uint32(FPlatformTime::ToMilliseconds64(RTStats.AccumulatedStats.GetDynamicMeshElementsCycles) * 1000.0));
+
+		tempString.Appendf(TEXT(",%u,[ "), uint32(FPlatformTime::ToMilliseconds64(GTStats.MaxPerFrameTotalCycles.Num() > 0 ? GTStats.MaxPerFrameTotalCycles[0] : 0) * 1000.0));
+		for (uint64 v : GTStats.MaxPerFrameTotalCycles)
 		{
 			tempString.Appendf(TEXT("%u "), uint32(FPlatformTime::ToMilliseconds64(v) * 1000.0));
 		}
 		tempString.Append(TEXT("]"));
 
-		tempString.Appendf(TEXT(",%u,[ "), uint32(FPlatformTime::ToMilliseconds64(RTSTats.MaxPerInstanceCycles.Num() > 0 ? RTSTats.MaxPerInstanceCycles[0] : 0) * 1000.0));
-		for (uint64 v : RTSTats.MaxPerFrameTotalCycles)
+		tempString.Appendf(TEXT(",%u,[ "), uint32(FPlatformTime::ToMilliseconds64(RTStats.MaxPerInstanceCycles.Num() > 0 ? RTStats.MaxPerInstanceCycles[0] : 0) * 1000.0));
+		for (uint64 v : RTStats.MaxPerFrameTotalCycles)
 		{
 			tempString.Appendf(TEXT("%u "), uint32(FPlatformTime::ToMilliseconds64(v) * 1000.0));
 		}
@@ -1163,7 +1197,6 @@ bool FParticlePerfStatsListener_CSVProfiler::Tick()
 	return false;
 }
 
-
 void FParticlePerfStatsListener_CSVProfiler::TickRT()
 {
 	FParticlePerfStatsListener_GatherAll::TickRT();
@@ -1175,10 +1208,15 @@ void FParticlePerfStatsListener_CSVProfiler::TickRT()
 			{
 				if (const UFXSystemAsset* System = WeakSystem.Get())
 				{
-					float RTTime = FPlatformTime::ToMilliseconds64(Stats->GetRenderThreadStats().GetTotalCycles()) * 1000.0f;
-					float AvgTime = FPlatformTime::ToMilliseconds64(Stats->GetRenderThreadStats().GetPerInstanceAvgCycles()) * 1000.0f;
+					const float RTTime = FPlatformTime::ToMilliseconds64(Stats->GetRenderThreadStats().GetTotalCycles()) * 1000.0f;
+					const float RTAvgTime = FPlatformTime::ToMilliseconds64(Stats->GetRenderThreadStats().GetPerInstanceAvgCycles()) * 1000.0f;
 					CSVProfiler->RecordCustomStat(System->CSVStat_RT, CSV_CATEGORY_INDEX(Particles), RTTime, ECsvCustomStatOp::Set);
-					CSVProfiler->RecordCustomStat(System->CSVStat_InstAvgRT, CSV_CATEGORY_INDEX(Particles), AvgTime, ECsvCustomStatOp::Set);
+					CSVProfiler->RecordCustomStat(System->CSVStat_InstAvgRT, CSV_CATEGORY_INDEX(Particles), RTAvgTime, ECsvCustomStatOp::Set);
+
+					const float GpuTime = float(Stats->GetGPUStats().GetTotalMicroseconds());
+					const float GpuAvgTime = float(Stats->GetGPUStats().GetPerInstanceAvgMicroseconds());
+					CSVProfiler->RecordCustomStat(System->CSVStat_GPU, CSV_CATEGORY_INDEX(Particles), GpuTime, ECsvCustomStatOp::Set);
+					CSVProfiler->RecordCustomStat(System->CSVStat_InstAvgGPU, CSV_CATEGORY_INDEX(Particles), GpuAvgTime, ECsvCustomStatOp::Set);
 				}
 			}
 		);
@@ -1186,7 +1224,6 @@ void FParticlePerfStatsListener_CSVProfiler::TickRT()
 	}
 }
 
-
 void FParticlePerfStatsListener_CSVProfiler::End()
 {
 	if (GLog != nullptr)
@@ -1250,7 +1287,7 @@ int32 FParticlePerfStatsListener_DebugRender::RenderStats(class UWorld* World, c
 		}
 
 		const FAccumulatedParticlePerfStats_GT& GTStats = PerfStats->GetGameThreadStats();
-		const FAccumulatedParticlePerfStats_RT& RTStats = PerfStats->GetRenderThreadStats_GameThread();
+		const FAccumulatedParticlePerfStats_RT& RTStats = PerfStats->GetRenderThreadStats();
 		const UFXSystemAsset* System = it.Key().Get();
 		if (GTStats.NumFrames == 0 || RTStats.NumFrames == 0 || System == nullptr)
 		{
