- Header name portability fix

diff --git a/Engine/Plugins/FX/Niagara/Source/Niagara/Private/NiagaraGpuComputeDispatch.cpp b/Engine/Plugins/FX/Niagara/Source/Niagara/Private/NiagaraGpuComputeDispatch.cpp
--- a/Engine/Plugins/FX/Niagara/Source/Niagara/Private/NiagaraGpuComputeDispatch.cpp
+++ b/Engine/Plugins/FX/Niagara/Source/Niagara/Private/NiagaraGpuComputeDispatch.cpp
@@ -15,7 +15,7 @@
 #if NIAGARA_COMPUTEDEBUG_ENABLED
 #include "NiagaraGpuComputeDebug.h"
 #endif
-#include "NiagaraGpuProfilerInterface.h"
+#include "NiagaraGPUProfilerInterface.h"
 #include "NiagaraGpuReadbackManager.h"
 #include "NiagaraRayTracingHelper.h"
 #include "NiagaraRenderer.h"
