- Header portability fix

diff --git a/Engine/Plugins/FX/Niagara/Source/Niagara/Private/NiagaraGPUProfilerInterface.cpp b/Engine/Plugins/FX/Niagara/Source/Niagara/Private/NiagaraGPUProfilerInterface.cpp
--- a/Engine/Plugins/FX/Niagara/Source/Niagara/Private/NiagaraGPUProfilerInterface.cpp
+++ b/Engine/Plugins/FX/Niagara/Source/Niagara/Private/NiagaraGPUProfilerInterface.cpp
@@ -2,7 +2,7 @@
 
 #include "NiagaraGPUProfilerInterface.h"
 #include "NiagaraGpuComputeDispatchInterface.h"
-#include "NiagaraGpuProfiler.h"
+#include "NiagaraGPUProfiler.h"
 #include "NiagaraGPUSystemTick.h"
 #include "NiagaraDataInterfaceBase.h"
 #include "NiagaraComponent.h"
