- Remove // from include

#rnx

diff --git a/Engine/Plugins/FX/Niagara/Source/Niagara/Private/NiagaraGPUProfiler.h b/Engine/Plugins/FX/Niagara/Source/Niagara/Private/NiagaraGPUProfiler.h
--- a/Engine/Plugins/FX/Niagara/Source/Niagara/Private/NiagaraGPUProfiler.h
+++ b/Engine/Plugins/FX/Niagara/Source/Niagara/Private/NiagaraGPUProfiler.h
@@ -8,7 +8,7 @@
 #include "RHIResources.h"
 #include "HAL/Platform.h"
 #include "Misc/Build.h"
-#include "Particles//ParticlePerfStats.h"
+#include "Particles/ParticlePerfStats.h"
 
 #include "NiagaraGPUProfilerInterface.h"
 
