- Fix lights not disabling when Niagara / Particles show flags are diabled

#rnx

diff --git a/Engine/Plugins/FX/Niagara/Source/Niagara/Private/NiagaraRendererLights.cpp b/Engine/Plugins/FX/Niagara/Source/Niagara/Private/NiagaraRendererLights.cpp
--- a/Engine/Plugins/FX/Niagara/Source/Niagara/Private/NiagaraRendererLights.cpp
+++ b/Engine/Plugins/FX/Niagara/Source/Niagara/Private/NiagaraRendererLights.cpp
@@ -39,7 +39,7 @@ FNiagaraRendererLights::FNiagaraRendererLights(ERHIFeatureLevel::Type FeatureLev
 FPrimitiveViewRelevance FNiagaraRendererLights::GetViewRelevance(const FSceneView* View, const FNiagaraSceneProxy *SceneProxy)const
 {
 	FPrimitiveViewRelevance Result;
-	Result.bDrawRelevance = bHasLights;
+	Result.bDrawRelevance = bHasLights && View->Family->EngineShowFlags.Particles && View->Family->EngineShowFlags.Niagara;
 	Result.bShadowRelevance = false;
 	Result.bDynamicRelevance = false;
 	Result.bOpaque = false;
