- Work around incorrect static analysis warnings

#rnx

diff --git a/Engine/Plugins/FX/Niagara/Source/Niagara/Private/NiagaraGPUProfiler.cpp b/Engine/Plugins/FX/Niagara/Source/Niagara/Private/NiagaraGPUProfiler.cpp
--- a/Engine/Plugins/FX/Niagara/Source/Niagara/Private/NiagaraGPUProfiler.cpp
+++ b/Engine/Plugins/FX/Niagara/Source/Niagara/Private/NiagaraGPUProfiler.cpp
@@ -49,7 +49,7 @@ void FNiagaraGPUProfiler::BeginFrame(FRHICommandListImmediate& RHICmdList)
 	// Process any frames that are potentially complete
 	while ( CurrentReadFrame != CurrentWriteFrame )
 	{
-		FGpuFrameData& ReadFrame = GpuFrames[CurrentReadFrame];
+		FGpuFrameData& ReadFrame = GetReadFrame();
 		if ( !ProcessFrame(RHICmdList, ReadFrame) )
 		{
 			break;
@@ -66,7 +66,7 @@ void FNiagaraGPUProfiler::BeginFrame(FRHICommandListImmediate& RHICmdList)
 	}
 
 	// Have a processed all previous data?
-	if (GpuFrames[CurrentWriteFrame].EndQuery.GetQuery() != nullptr)
+	if (GetWriteFrame().EndQuery.GetQuery() != nullptr)
 	{
 		return;
 	}
diff --git a/Engine/Plugins/FX/Niagara/Source/Niagara/Private/NiagaraGPUProfiler.h b/Engine/Plugins/FX/Niagara/Source/Niagara/Private/NiagaraGPUProfiler.h
--- a/Engine/Plugins/FX/Niagara/Source/Niagara/Private/NiagaraGPUProfiler.h
+++ b/Engine/Plugins/FX/Niagara/Source/Niagara/Private/NiagaraGPUProfiler.h
@@ -62,7 +62,8 @@ public:
 	void EndDispatch(FRHICommandList& RHICmdList);
 
 private:
-	FGpuFrameData& GetWriteFrame() { return GpuFrames[CurrentWriteFrame]; }
+	FGpuFrameData& GetReadFrame() { check(CurrentReadFrame >= 0 && CurrentReadFrame < UE_ARRAY_COUNT(GpuFrames)); return GpuFrames[CurrentReadFrame]; }
+	FGpuFrameData& GetWriteFrame() { check(CurrentWriteFrame >= 0 && CurrentWriteFrame < UE_ARRAY_COUNT(GpuFrames)); return GpuFrames[CurrentWriteFrame]; }
 	bool ProcessFrame(FRHICommandListImmediate& RHICmdList, FGpuFrameData& ReadFrame);
 
 private:
