- Improve tooltips for custom sorting modes

#rnx

diff --git a/Engine/Plugins/FX/Niagara/Source/Niagara/Classes/NiagaraGPUSortInfo.h b/Engine/Plugins/FX/Niagara/Source/Niagara/Classes/NiagaraGPUSortInfo.h
--- a/Engine/Plugins/FX/Niagara/Source/Niagara/Classes/NiagaraGPUSortInfo.h
+++ b/Engine/Plugins/FX/Niagara/Source/Niagara/Classes/NiagaraGPUSortInfo.h
@@ -21,9 +21,15 @@ enum class ENiagaraSortMode : uint8
 	ViewDepth,
 	/** Sort by distance to the camera's origin.*/
 	ViewDistance,
-	/** Custom sorting according to a per particle attribute. Which attribute is defined by the renderer's CustomSortingBinding which defaults to Particles.NormalizedAge. Lower values are rendered before higher values. */
+	/**
+	Sort according to particles CustomSortingBinding (defaults to Particles.NormalizedAge).
+	Lower values will be sorted before higher values, i.e. 1 would draw on top of 0.
+	*/
 	CustomAscending,
-	/** Custom sorting according to a per particle attribute. Which attribute is defined by the renderer's CustomSortingBinding which defaults to Particles.NormalizedAge. Higher values are rendered before lower values. */
+	/**
+	Sort according to particles CustomSortingBinding (defaults to Particles.NormalizedAge).
+	Higher values will be sorted before lower values, i.e. 0 would draw on top of 1.
+	*/
 	CustomDecending,
 };
 
