deprecate/hide all functions and grid parameters that specify resolution and grid size
this should be all done via emitter scripts

diff --git a/Engine/Plugins/FX/Niagara/Source/Niagara/Classes/NiagaraDataInterfaceGrid2DCollection.h b/Engine/Plugins/FX/Niagara/Source/Niagara/Classes/NiagaraDataInterfaceGrid2DCollection.h
--- a/Engine/Plugins/FX/Niagara/Source/Niagara/Classes/NiagaraDataInterfaceGrid2DCollection.h
+++ b/Engine/Plugins/FX/Niagara/Source/Niagara/Classes/NiagaraDataInterfaceGrid2DCollection.h
@@ -119,21 +119,21 @@ public:
 	DECLARE_NIAGARA_DI_PARAMETER();
 
 	/** Reference to a user parameter if we're reading one. */
-	UPROPERTY(EditAnywhere, Category = "Grid2DCollection")
+	UPROPERTY(EditAnywhere, Category = "Grid")
 	FNiagaraUserParameterBinding RenderTargetUserParameter;
 
 	/** When enabled overrides the format used to store data inside the grid, otherwise uses the project default setting.  Lower bit depth formats will save memory and performance at the cost of precision. */
-	UPROPERTY(EditAnywhere, Category = "Grid2DCollection", meta = (EditCondition = "bOverrideFormat"))
+	UPROPERTY(EditAnywhere, Category = "Grid", meta = (EditCondition = "bOverrideFormat"))
 	ENiagaraGpuBufferFormat OverrideBufferFormat;
 
-	UPROPERTY(EditAnywhere, Category = "Grid2DCollection", meta = (PinHiddenByDefault, InlineEditConditionToggle))
+	UPROPERTY(EditAnywhere, Category = "Grid", meta = (PinHiddenByDefault, InlineEditConditionToggle))
 	uint8 bOverrideFormat : 1;
 
 #if WITH_EDITORONLY_DATA
-	UPROPERTY(EditAnywhere, Category = "Grid2DCollection", meta = (PinHiddenByDefault, InlineEditConditionToggle))
+	UPROPERTY(EditAnywhere, Category = "Grid", meta = (PinHiddenByDefault, InlineEditConditionToggle))
 	uint8 bPreviewGrid : 1;
 
-	UPROPERTY(EditAnywhere, Category = "Grid2DCollection", meta = (EditCondition = "bPreviewGrid", ToolTip = "When enabled allows you to preview the grid in a debug display") )
+	UPROPERTY(EditAnywhere, Category = "Grid", meta = (EditCondition = "bPreviewGrid", ToolTip = "When enabled allows you to preview the grid in a debug display") )
 	FName PreviewAttribute = NAME_None;
 #endif
 
diff --git a/Engine/Plugins/FX/Niagara/Source/Niagara/Classes/NiagaraDataInterfaceGrid3DCollection.h b/Engine/Plugins/FX/Niagara/Source/Niagara/Classes/NiagaraDataInterfaceGrid3DCollection.h
--- a/Engine/Plugins/FX/Niagara/Source/Niagara/Classes/NiagaraDataInterfaceGrid3DCollection.h
+++ b/Engine/Plugins/FX/Niagara/Source/Niagara/Classes/NiagaraDataInterfaceGrid3DCollection.h
@@ -161,21 +161,21 @@ public:
 	int32 NumAttributes;
 
 	/** Reference to a user parameter if we're reading one. */
-	UPROPERTY(EditAnywhere, Category = "Grid3DCollection")
+	UPROPERTY(EditAnywhere, Category = "Grid")
 	FNiagaraUserParameterBinding RenderTargetUserParameter;
 
 	/** When enabled overrides the format used to store data inside the grid, otherwise uses the project default setting.  Lower bit depth formats will save memory and performance at the cost of precision. */
-	UPROPERTY(EditAnywhere, Category = "Grid3DCollection", meta = (EditCondition = "bOverrideFormat"))
+	UPROPERTY(EditAnywhere, Category = "Grid", meta = (EditCondition = "bOverrideFormat"))
 	ENiagaraGpuBufferFormat OverrideBufferFormat;
 
-	UPROPERTY(EditAnywhere, Category = "Grid3DCollection", meta = (PinHiddenByDefault, InlineEditConditionToggle))
+	UPROPERTY(EditAnywhere, Category = "Grid", meta = (PinHiddenByDefault, InlineEditConditionToggle))
 	uint8 bOverrideFormat : 1;
 
 #if WITH_EDITORONLY_DATA
-	UPROPERTY(EditAnywhere, Category = "Grid3DCollection", meta = (PinHiddenByDefault, InlineEditConditionToggle))
+	UPROPERTY(EditAnywhere, Category = "Grid", meta = (PinHiddenByDefault, InlineEditConditionToggle))
 	uint8 bPreviewGrid : 1;
 
-	UPROPERTY(EditAnywhere, Category = "Grid3DCollection", meta = (EditCondition = "bPreviewGrid", ToolTip = "When enabled allows you to preview the grid in a debug display") )
+	UPROPERTY(EditAnywhere, Category = "Grid", meta = (EditCondition = "bPreviewGrid", ToolTip = "When enabled allows you to preview the grid in a debug display") )
 	FName PreviewAttribute = NAME_None;
 #endif
 
diff --git a/Engine/Plugins/FX/Niagara/Source/Niagara/Classes/NiagaraDataInterfaceRW.h b/Engine/Plugins/FX/Niagara/Source/Niagara/Classes/NiagaraDataInterfaceRW.h
--- a/Engine/Plugins/FX/Niagara/Source/Niagara/Classes/NiagaraDataInterfaceRW.h
+++ b/Engine/Plugins/FX/Niagara/Source/Niagara/Classes/NiagaraDataInterfaceRW.h
@@ -68,10 +68,10 @@ public:
 	static const FName ExecutionIndexToUnitFunctionName;
 
 public:
-	UPROPERTY(EditAnywhere, Category = "Deprecated", AdvancedDisplay)
+	UPROPERTY(EditAnywhere, Category = "Grid", meta = (EditCondition = "false", EditConditionHides))
 	TSet<int> OutputShaderStages;
 
-	UPROPERTY(EditAnywhere, Category = "Deprecated", AdvancedDisplay)
+	UPROPERTY(EditAnywhere, Category = "Grid", meta = (EditCondition = "false", EditConditionHides))
 	TSet<int> IterationShaderStages;
 
 public:
@@ -131,23 +131,23 @@ class NIAGARA_API UNiagaraDataInterfaceGrid3D : public UNiagaraDataInterfaceRWBa
 
 public:
 	// Number of cells
-	UPROPERTY(EditAnywhere, Category = "Grid")
+	UPROPERTY(EditAnywhere, Category = "Grid", meta=(EditCondition = "false", EditConditionHides))
 	FIntVector NumCells;
 
 	// World space size of a cell
-	UPROPERTY(EditAnywhere, Category = "Deprecated", AdvancedDisplay)
+	UPROPERTY(EditAnywhere, Category = "Grid", meta = (EditCondition = "false", EditConditionHides))
 	float CellSize;
 
 	// Number of cells on the longest axis
-	UPROPERTY(EditAnywhere, Category = "Deprecated", AdvancedDisplay)
+	UPROPERTY(EditAnywhere, Category = "Grid", meta = (EditCondition = "false", EditConditionHides))
 	int32 NumCellsMaxAxis;
 
 	// Method for setting the grid resolution
-	UPROPERTY(EditAnywhere, Category = "Deprecated", AdvancedDisplay)
+	UPROPERTY(EditAnywhere, Category = "Grid", meta = (EditCondition = "false", EditConditionHides))
 	ESetResolutionMethod SetResolutionMethod;
 	
 	// World size of the grid
-	UPROPERTY(EditAnywhere, Category = "Deprecated", AdvancedDisplay)
+	UPROPERTY(EditAnywhere, Category = "Grid", meta = (EditCondition = "false", EditConditionHides))
 	FVector WorldBBoxSize;
 
 public:
@@ -181,15 +181,15 @@ class NIAGARA_API UNiagaraDataInterfaceGrid2D : public UNiagaraDataInterfaceRWBa
 
 public:
 	// Number of cells in X
-	UPROPERTY(EditAnywhere, Category = "Grid")
+	UPROPERTY(EditAnywhere, Category = "Grid", meta = (EditCondition = "false", EditConditionHides))
 	int32 NumCellsX;
 
 	// Number of cells in Y
-	UPROPERTY(EditAnywhere, Category = "Grid")
+	UPROPERTY(EditAnywhere, Category = "Grid", meta = (EditCondition = "false", EditConditionHides))
 	int32 NumCellsY;
 	
 	// Number of cells on the longest axis
-	UPROPERTY(EditAnywhere, Category = "Deprecated", AdvancedDisplay)
+	UPROPERTY(EditAnywhere, Category = "Grid", meta = (EditCondition = "false", EditConditionHides))
 	int32 NumCellsMaxAxis;
 
 	// Number of Attributes
@@ -197,11 +197,11 @@ public:
 	int32 NumAttributes;
 
 	// Set grid resolution according to longest axis
-	UPROPERTY(EditAnywhere, Category = "Deprecated", AdvancedDisplay)
+	UPROPERTY(EditAnywhere, Category = "Grid", meta = (EditCondition = "false", EditConditionHides))
 	bool SetGridFromMaxAxis;	
 
 	// World size of the grid
-	UPROPERTY(EditAnywhere, Category = "Deprecated", AdvancedDisplay)
+	UPROPERTY(EditAnywhere, Category = "Grid", meta = (EditCondition = "false", EditConditionHides))
 	FVector2D WorldBBoxSize;
 
 
diff --git a/Engine/Plugins/FX/Niagara/Source/Niagara/Private/NiagaraDataInterfaceRW.cpp b/Engine/Plugins/FX/Niagara/Source/Niagara/Private/NiagaraDataInterfaceRW.cpp
--- a/Engine/Plugins/FX/Niagara/Source/Niagara/Private/NiagaraDataInterfaceRW.cpp
+++ b/Engine/Plugins/FX/Niagara/Source/Niagara/Private/NiagaraDataInterfaceRW.cpp
@@ -88,7 +88,8 @@ void UNiagaraDataInterfaceGrid3D::GetFunctions(TArray<FNiagaraFunctionSignature>
 		Sig.Name = UNiagaraDataInterfaceRWBase::WorldBBoxSizeFunctionName;
 		Sig.Inputs.Add(FNiagaraVariable(FNiagaraTypeDefinition(GetClass()), TEXT("Grid")));
 		Sig.Outputs.Add(FNiagaraVariable(FNiagaraTypeDefinition::GetVec3Def(), TEXT("WorldBBoxSize")));
-
+		Sig.bHidden = true;
+		Sig.bSoftDeprecatedFunction = true;
 		Sig.bMemberFunction = true;
 		Sig.bRequiresContext = false;
 		OutFunctions.Add(Sig);
@@ -243,7 +244,8 @@ void UNiagaraDataInterfaceGrid3D::GetFunctions(TArray<FNiagaraFunctionSignature>
 		Sig.Name = UNiagaraDataInterfaceRWBase::CellSizeFunctionName;
 		Sig.Inputs.Add(FNiagaraVariable(FNiagaraTypeDefinition(GetClass()), TEXT("Grid")));
 		Sig.Outputs.Add(FNiagaraVariable(FNiagaraTypeDefinition::GetVec3Def(), TEXT("CellSize")));		
-
+		Sig.bHidden = true;
+		Sig.bSoftDeprecatedFunction = true;
 		Sig.bMemberFunction = true;
 		Sig.bRequiresContext = false;
 		OutFunctions.Add(Sig);
@@ -533,7 +535,8 @@ void UNiagaraDataInterfaceGrid2D::GetFunctions(TArray<FNiagaraFunctionSignature>
 		Sig.Name = UNiagaraDataInterfaceRWBase::WorldBBoxSizeFunctionName;
 		Sig.Inputs.Add(FNiagaraVariable(FNiagaraTypeDefinition(GetClass()), TEXT("Grid")));
 		Sig.Outputs.Add(FNiagaraVariable(FNiagaraTypeDefinition::GetVec2Def(), TEXT("WorldBBoxSize")));		
-
+		Sig.bHidden = true;
+		Sig.bSoftDeprecatedFunction = true;
 		Sig.bMemberFunction = true;
 		Sig.bRequiresContext = false;
 		OutFunctions.Add(Sig);
@@ -696,7 +699,8 @@ void UNiagaraDataInterfaceGrid2D::GetFunctions(TArray<FNiagaraFunctionSignature>
 		Sig.Name = UNiagaraDataInterfaceRWBase::CellSizeFunctionName;
 		Sig.Inputs.Add(FNiagaraVariable(FNiagaraTypeDefinition(GetClass()), TEXT("Grid")));
 		Sig.Outputs.Add(FNiagaraVariable(FNiagaraTypeDefinition::GetVec2Def(), TEXT("CellSize")));
-
+		Sig.bHidden = true;
+		Sig.bSoftDeprecatedFunction = true;
 		Sig.bMemberFunction = true;
 		Sig.bRequiresContext = false;
 		OutFunctions.Add(Sig);
