- Move Niagara Debug HUD colors into the settings to allow user to adjust them

diff --git a/Engine/Plugins/FX/Niagara/Source/Niagara/Private/NiagaraDebugHud.cpp b/Engine/Plugins/FX/Niagara/Source/Niagara/Private/NiagaraDebugHud.cpp
--- a/Engine/Plugins/FX/Niagara/Source/Niagara/Private/NiagaraDebugHud.cpp
+++ b/Engine/Plugins/FX/Niagara/Source/Niagara/Private/NiagaraDebugHud.cpp
@@ -19,12 +19,6 @@
 
 namespace NiagaraDebugLocal
 {
-	const FLinearColor InfoTextColor = FLinearColor::White;
-	const FLinearColor WarningTextColor = FLinearColor(0.9f, 0.7f, 0.0, 1.0f);
-	const FLinearColor ErrorTextColor = FLinearColor(1.0f, 0.4, 0.3, 1.0f);
-	const FLinearColor WorldTextColor = FLinearColor::White;
-	const FLinearColor BackgroundColor = FLinearColor(0.0f, 0.0f, 0.0f, 0.5f);
-
 	FCriticalSection RTFramesGuard;
 
 	enum class EEngineVariables : uint8
@@ -1154,9 +1148,10 @@ void FNiagaraDebugHud::DrawOverview(class FNiagaraWorldManager* WorldManager, FC
 	UFont* Font = GetFont(Settings.OverviewFont);
 	const float fAdvanceHeight = Font->GetMaxCharHeight() + 1.0f;
 
-	const FLinearColor HeadingColor = FLinearColor::Green;
-	const FLinearColor DetailColor = FLinearColor::White;
-	const FLinearColor DetailHighlightColor = FLinearColor::Yellow;
+	const FLinearColor HeadingColor = Settings.OverviewHeadingColor;
+	const FLinearColor DetailColor = Settings.OverviewDetailColor;
+	const FLinearColor DetailHighlightColor = Settings.OverviewDetailHighlightColor;
+	const FLinearColor BackgroundColor = Settings.DefaultBackgroundColor;
 
 	FVector2D TextLocation = Settings.OverviewLocation;
 
@@ -1681,9 +1676,10 @@ void FNiagaraDebugHud::DrawGpuComputeOverriew(class FNiagaraWorldManager* WorldM
 	UFont* Font = GetFont(Settings.OverviewFont);
 	const float fAdvanceHeight = Font->GetMaxCharHeight() + 1.0f;
 
-	const FLinearColor HeadingColor = FLinearColor::Green;
-	const FLinearColor DetailColor = FLinearColor::White;
-	const FLinearColor DetailHighlightColor = FLinearColor::Yellow;
+	const FLinearColor HeadingColor = Settings.OverviewHeadingColor;
+	const FLinearColor DetailColor = Settings.OverviewDetailColor;
+	const FLinearColor DetailHighlightColor = Settings.OverviewDetailHighlightColor;
+	const FLinearColor BackgroundColor = Settings.DefaultBackgroundColor;
 
 	if (GpuResults == nullptr)
 	{
@@ -1728,7 +1724,7 @@ void FNiagaraDebugHud::DrawGpuComputeOverriew(class FNiagaraWorldManager* WorldM
 			}
 		}
 
-		void Draw(UFont* Font, FCanvas* DrawCanvas, FVector2D& Location, FLinearColor TextColor)
+		void Draw(UFont* Font, FCanvas* DrawCanvas, FVector2D& Location, FLinearColor TextColor, FLinearColor BackgroundColor)
 		{
 			FVector2D TotalSize = FVector2D::ZeroVector;
 			for ( FColumn& Column : Columns)
@@ -1777,7 +1773,7 @@ void FNiagaraDebugHud::DrawGpuComputeOverriew(class FNiagaraWorldManager* WorldM
 		SimpleTable.GetColumnText(1).Appendf(TEXT("TotalDispatches : %d"), TotalDispatches);
 		SimpleTable.GetColumnText(2).Appendf(TEXT("TotalDispatchGroups : %d"), TotalDispatchGroups);
 
-		SimpleTable.Draw(Font, DrawCanvas, TextLocation, DetailColor);
+		SimpleTable.Draw(Font, DrawCanvas, TextLocation, DetailColor, BackgroundColor);
 		TextLocation.Y += fAdvanceHeight;
 	}
 
@@ -1850,7 +1846,7 @@ void FNiagaraDebugHud::DrawGpuComputeOverriew(class FNiagaraWorldManager* WorldM
 				}
 			}
 		}
-		SimpleTable.Draw(Font, DrawCanvas, TextLocation, DetailColor);
+		SimpleTable.Draw(Font, DrawCanvas, TextLocation, DetailColor, BackgroundColor);
 		TextLocation.Y += fAdvanceHeight;
 	}
 
@@ -1884,7 +1880,7 @@ void FNiagaraDebugHud::DrawGpuComputeOverriew(class FNiagaraWorldManager* WorldM
 			SimpleTable.GetColumnText(3).Appendf(TEXT("%llu"), SystemUsage.Microseconds.GetMax());
 			SimpleTable.RowComplete();
 		}
-		SimpleTable.Draw(Font, DrawCanvas, TextLocation, DetailColor);
+		SimpleTable.Draw(Font, DrawCanvas, TextLocation, DetailColor, BackgroundColor);
 		TextLocation.Y += fAdvanceHeight;
 	}
 
@@ -1909,7 +1905,7 @@ void FNiagaraDebugHud::DrawGpuComputeOverriew(class FNiagaraWorldManager* WorldM
 			SimpleTable.GetColumnText(3).Appendf(TEXT("%llu"), EventUsage.Microseconds.GetMax());
 			SimpleTable.RowComplete();
 		}
-		SimpleTable.Draw(Font, DrawCanvas, TextLocation, DetailColor);
+		SimpleTable.Draw(Font, DrawCanvas, TextLocation, DetailColor, BackgroundColor);
 		TextLocation.Y += fAdvanceHeight;
 	}
 #endif //WITH_NIAGARA_GPU_PROFILER
@@ -1929,9 +1925,10 @@ void FNiagaraDebugHud::DrawGlobalBudgetInfo(class FNiagaraWorldManager* WorldMan
 	UFont* Font = GetFont(Settings.OverviewFont);
 	const float fAdvanceHeight = Font->GetMaxCharHeight() + 1.0f;
 
-	const FLinearColor HeadingColor = FLinearColor::Green;
-	const FLinearColor DetailColor = FLinearColor::White;
-	const FLinearColor DetailHighlightColor = FLinearColor::Yellow;
+	const FLinearColor HeadingColor = Settings.OverviewHeadingColor;
+	const FLinearColor DetailColor = Settings.OverviewDetailColor;
+	const FLinearColor DetailHighlightColor = Settings.OverviewDetailHighlightColor;
+	const FLinearColor BackgroundColor = Settings.DefaultBackgroundColor;
 
 	if (FFXBudget::Enabled())
 	{
@@ -2171,12 +2168,12 @@ void FNiagaraDebugHud::DrawValidation(class FNiagaraWorldManager* WorldManager,
 			const FVector2D ErrorStringSize = GetStringSize(Font, ErrorString.ToString());
 
 			TextLocation.Y += fAdvanceHeight;
-			DrawCanvas->DrawTile(TextLocation.X - 1.0f, TextLocation.Y - 1.0f, ErrorStringSize.X + 2.0f, ErrorStringSize.Y + fAdvanceHeight + 2.0f, 0.0f, 0.0f, 0.0f, 0.0f, BackgroundColor);
+			DrawCanvas->DrawTile(TextLocation.X - 1.0f, TextLocation.Y - 1.0f, ErrorStringSize.X + 2.0f, ErrorStringSize.Y + fAdvanceHeight + 2.0f, 0.0f, 0.0f, 0.0f, 0.0f, Settings.DefaultBackgroundColor);
 
-			DrawCanvas->DrawShadowedString(TextLocation.X, TextLocation.Y, TEXT("Found Errors:"), Font, ErrorTextColor);
+			DrawCanvas->DrawShadowedString(TextLocation.X, TextLocation.Y, TEXT("Found Errors:"), Font, Settings.MessageErrorTextColor);
 			TextLocation.Y += fAdvanceHeight;
 
-			DrawCanvas->DrawShadowedString(TextLocation.X, TextLocation.Y, ErrorString.ToString(), Font, ErrorTextColor);
+			DrawCanvas->DrawShadowedString(TextLocation.X, TextLocation.Y, ErrorString.ToString(), Font, Settings.MessageErrorTextColor);
 		}
 	}
 }
@@ -2231,8 +2228,7 @@ void FNiagaraDebugHud::DrawComponents(FNiagaraWorldManager* WorldManager, UCanva
 			SystemSimulation->WaitForInstancesTickComplete();
 		}
 
-		const FLinearColor TextColor = ValidationErrors.Contains(NiagaraComponent) ? ErrorTextColor : WorldTextColor;
-
+		const FLinearColor TextColor = ValidationErrors.Contains(NiagaraComponent) ? Settings.InWorldErrorTextColor : Settings.InWorldTextColor;
 
 		// Show particle data in world
 		if (!Settings.bShowParticlesVariablesWithSystem && bSystemSimulationValid)
@@ -2289,7 +2285,7 @@ void FNiagaraDebugHud::DrawComponents(FNiagaraWorldManager* WorldManager, UCanva
 
 						const TCHAR* FinalString = StringBuilder.ToString();
 						const TPair<FVector2D, FVector2D> SizeAndLocation = GetTextLocation(ParticleFont, FinalString, Settings.ParticleTextOptions, FVector2D(ParticleScreenLocation));
-						DrawCanvas->DrawTile(SizeAndLocation.Value.X - 1.0f, SizeAndLocation.Value.Y - 1.0f, SizeAndLocation.Key.X + 2.0f, SizeAndLocation.Key.Y + 2.0f, 0.0f, 0.0f, 0.0f, 0.0f, BackgroundColor);
+						DrawCanvas->DrawTile(SizeAndLocation.Value.X - 1.0f, SizeAndLocation.Value.Y - 1.0f, SizeAndLocation.Key.X + 2.0f, SizeAndLocation.Key.Y + 2.0f, 0.0f, 0.0f, 0.0f, 0.0f, Settings.DefaultBackgroundColor);
 						DrawCanvas->DrawShadowedString(SizeAndLocation.Value.X, SizeAndLocation.Value.Y, FinalString, ParticleFont, TextColor);
 					}
 				}
@@ -2535,7 +2531,7 @@ void FNiagaraDebugHud::DrawComponents(FNiagaraWorldManager* WorldManager, UCanva
 
 				const TCHAR* FinalString = StringBuilder.ToString();
 				const TPair<FVector2D, FVector2D> SizeAndLocation = GetTextLocation(SystemFont, FinalString, Settings.SystemTextOptions, FVector2D(ScreenLocation));
-				DrawCanvas->DrawTile(SizeAndLocation.Value.X - 1.0f, SizeAndLocation.Value.Y - 1.0f, SizeAndLocation.Key.X + 2.0f, SizeAndLocation.Key.Y + 2.0f, 0.0f, 0.0f, 0.0f, 0.0f, BackgroundColor);
+				DrawCanvas->DrawTile(SizeAndLocation.Value.X - 1.0f, SizeAndLocation.Value.Y - 1.0f, SizeAndLocation.Key.X + 2.0f, SizeAndLocation.Key.Y + 2.0f, 0.0f, 0.0f, 0.0f, 0.0f, Settings.DefaultBackgroundColor);
 				DrawCanvas->DrawShadowedString(SizeAndLocation.Value.X, SizeAndLocation.Value.Y, FinalString, SystemFont, TextColor);
 			}
 		}
@@ -2580,7 +2576,7 @@ void FNiagaraDebugHud::DrawMessages(class FNiagaraWorldManager* WorldManager, cl
 	if (Messages.Num())
 	{
 		TextLocation.Y += fAdvanceHeight;
-		DrawCanvas->DrawTile(TextLocation.X - 1.0f, TextLocation.Y - 1.0f, BackgroundSize.X + 2.0f, BackgroundSize.Y + 2.0f, 0.0f, 0.0f, 0.0f, 0.0f, BackgroundColor);
+		DrawCanvas->DrawTile(TextLocation.X - 1.0f, TextLocation.Y - 1.0f, BackgroundSize.X + 2.0f, BackgroundSize.Y + 2.0f, 0.0f, 0.0f, 0.0f, 0.0f, Settings.DefaultBackgroundColor);
 
 		for (TPair<FName, FNiagaraDebugMessage>& Pair : Messages)
 		{
@@ -2588,9 +2584,9 @@ void FNiagaraDebugHud::DrawMessages(class FNiagaraWorldManager* WorldManager, cl
 			FNiagaraDebugMessage& Message = Pair.Value;
 
 			FLinearColor MessageColor;
-			if (Message.Type == ENiagaraDebugMessageType::Info) MessageColor = InfoTextColor;
-			else if (Message.Type == ENiagaraDebugMessageType::Warning) MessageColor = WarningTextColor;
-			else if (Message.Type == ENiagaraDebugMessageType::Error) MessageColor = ErrorTextColor;
+			if (Message.Type == ENiagaraDebugMessageType::Info) MessageColor = Settings.MessageInfoTextColor;
+			else if (Message.Type == ENiagaraDebugMessageType::Warning) MessageColor = Settings.MessageWarningTextColor;
+			else if (Message.Type == ENiagaraDebugMessageType::Error) MessageColor = Settings.MessageErrorTextColor;
 
 			DrawCanvas->DrawShadowedString(TextLocation.X, TextLocation.Y, *Message.Message, Font, MessageColor);//TODO: Sort by type / lifetime?
 			TextLocation.Y += fAdvanceHeight;
diff --git a/Engine/Plugins/FX/Niagara/Source/Niagara/Public/NiagaraDebuggerCommon.h b/Engine/Plugins/FX/Niagara/Source/Niagara/Public/NiagaraDebuggerCommon.h
--- a/Engine/Plugins/FX/Niagara/Source/Niagara/Public/NiagaraDebuggerCommon.h
+++ b/Engine/Plugins/FX/Niagara/Source/Niagara/Public/NiagaraDebuggerCommon.h
@@ -592,6 +592,41 @@ struct NIAGARA_API FNiagaraDebugHUDSettingsData
 	UPROPERTY(Config, EditAnywhere, Category = "Perf Overview", meta = (EditCondition = "bEnableSmoothing"))
 	int32 SmoothingWidth = 4;
 
+	// Default background color used generally for panels
+	UPROPERTY(EditAnywhere, Category = "Colors")
+	FLinearColor DefaultBackgroundColor = FLinearColor(0.0f, 0.0f, 0.0f, 0.5f);
+
+	// Overview heading text color
+	UPROPERTY(EditAnywhere, Category = "Colors")
+	FLinearColor OverviewHeadingColor = FLinearColor::Green;
+
+	// Overview detail text color
+	UPROPERTY(EditAnywhere, Category = "Colors")
+	FLinearColor OverviewDetailColor = FLinearColor::White;
+
+	// Overview detail highlight text color
+	UPROPERTY(EditAnywhere, Category = "Colors")
+	FLinearColor OverviewDetailHighlightColor = FLinearColor::Yellow;
+
+	// In world text if an error is detected
+	UPROPERTY(EditAnywhere, Category = "Colors")
+	FLinearColor InWorldErrorTextColor = FLinearColor(1.0f, 0.4, 0.3, 1.0f);
+	// In world text color
+	UPROPERTY(EditAnywhere, Category = "Colors")
+	FLinearColor InWorldTextColor = FLinearColor::White;
+
+	// Message display text color
+	UPROPERTY(EditAnywhere, Category = "Colors")
+	FLinearColor MessageInfoTextColor = FLinearColor::White;
+
+	// Message display warning text color
+	UPROPERTY(EditAnywhere, Category = "Colors")
+	FLinearColor MessageWarningTextColor = FLinearColor(0.9f, 0.7f, 0.0, 1.0f);
+
+	// Message display error text color
+	UPROPERTY(EditAnywhere, Category = "Colors")
+	FLinearColor MessageErrorTextColor = FLinearColor(1.0f, 0.4, 0.3, 1.0f);
+
 	UPROPERTY()
 	ENiagaraDebugPlaybackMode PlaybackMode = ENiagaraDebugPlaybackMode::Play;
 
