- Fix for reading data from the RT when in GPU mode

#rnx

diff --git a/Engine/Plugins/FX/Niagara/Source/Niagara/Private/NiagaraRendererMeshes.cpp b/Engine/Plugins/FX/Niagara/Source/Niagara/Private/NiagaraRendererMeshes.cpp
--- a/Engine/Plugins/FX/Niagara/Source/Niagara/Private/NiagaraRendererMeshes.cpp
+++ b/Engine/Plugins/FX/Niagara/Source/Niagara/Private/NiagaraRendererMeshes.cpp
@@ -1445,10 +1445,11 @@ FNiagaraDynamicDataBase* FNiagaraRendererMeshes::GenerateDynamicData(const FNiag
 		}
 	}
 
+	
 	FNiagaraDataBuffer* DataToRender = Emitter->GetData().GetCurrentData();
 	if (!DataToRender || 
 		Meshes.Num() == 0 ||
-		(SourceMode == ENiagaraRendererSourceDataMode::Particles && DataToRender->GetNumInstances() == 0))
+		(SourceMode == ENiagaraRendererSourceDataMode::Particles && SimTarget != ENiagaraSimTarget::GPUComputeSim && DataToRender->GetNumInstances() == 0))
 	{
 		return nullptr;
 	}
