fix bug with actors being deleted during simulation by checking to see if the number of potential collision actors has changed

diff --git a/Engine/Plugins/FX/Niagara/Source/Niagara/Private/NiagaraDataInterfaceRigidMeshCollisionQuery.cpp b/Engine/Plugins/FX/Niagara/Source/Niagara/Private/NiagaraDataInterfaceRigidMeshCollisionQuery.cpp
--- a/Engine/Plugins/FX/Niagara/Source/Niagara/Private/NiagaraDataInterfaceRigidMeshCollisionQuery.cpp
+++ b/Engine/Plugins/FX/Niagara/Source/Niagara/Private/NiagaraDataInterfaceRigidMeshCollisionQuery.cpp
@@ -492,6 +492,31 @@ void FNDIRigidMeshCollisionData::Update(UNiagaraDataInterfaceRigidMeshCollisionQ
 {
 	if (Interface != nullptr && SystemInstance != nullptr)
 	{		
+		TArray<AActor*> ActorsTmp;		
+
+		if (Interface->StaticMesh)
+		{
+			// #todo(dmp): loop over actors and count number of potential colliders.  This might change frame to frame and we might need a
+			// better way of doing this
+			UWorld* World = SystemInstance->GetWorld();
+			for (TActorIterator<AActor> It(World); It; ++It)
+			{
+				AActor* Actor = *It;
+
+				if ((!Interface->OnlyUseMoveable || (Interface->OnlyUseMoveable && Actor->IsRootComponentMovable())) &&
+					(Interface->Tag == FString("") || (Interface->Tag != FString("") && Actor->Tags.Contains(FName(Interface->Tag)))))
+				{
+					ActorsTmp.Add(Actor);
+				}
+			}
+		}
+
+		// check if the number of active actors is the same.  If not, reinitialize
+		if (ActorsTmp.Num() != Actors.Num())
+		{
+			Init(Interface, SystemInstance);
+		}
+
 		TickingGroup = ComputeTickingGroup();
 
 		if (0 < Actors.Num() && Actors[0] != nullptr)
