fix proxy boxes for convex hull collisions

diff --git a/Engine/Plugins/FX/Niagara/Source/Niagara/Private/NiagaraDataInterfaceRigidMeshCollisionQuery.cpp b/Engine/Plugins/FX/Niagara/Source/Niagara/Private/NiagaraDataInterfaceRigidMeshCollisionQuery.cpp
--- a/Engine/Plugins/FX/Niagara/Source/Niagara/Private/NiagaraDataInterfaceRigidMeshCollisionQuery.cpp
+++ b/Engine/Plugins/FX/Niagara/Source/Niagara/Private/NiagaraDataInterfaceRigidMeshCollisionQuery.cpp
@@ -247,11 +247,11 @@ void CreateInternalArrays(TArray<AActor*> Actors,
 
 								FBox BBox = ConvexElem.ElemBox;
 								FVector3f Extent = BBox.Max - BBox.Min;
-
+								FVector3f Center = (BBox.Max + BBox.Min) * .5;
 								OutAssetArrays->PhysicsType[BoxCount] = (ConvexElem.GetCollisionEnabled() == ECollisionEnabled::QueryAndPhysics);
 								OutAssetArrays->SourceSceneProxy[BoxCount] = StaticMeshComponent->SceneProxy;
 								
-								const FTransform ElementTransform = ConvexElem.GetTransform() * MeshTransform;
+								const FTransform ElementTransform = FTransform(Center) * MeshTransform;
 								OutAssetArrays->ElementExtent[BoxCount] = FVector4(Extent.X, Extent.Y, Extent.Z, 0);
 								FillCurrentTransforms(ElementTransform, BoxCount, OutAssetArrays->CurrentTransform, OutAssetArrays->CurrentInverse);
 
@@ -367,10 +367,11 @@ void UpdateInternalArrays(const TArray<AActor*> &Actors, FNDIRigidMeshCollisionA
 						UE_LOG(LogRigidMeshCollision, Warning, TEXT("Convex collision objects encountered and will be interpreted as a bounding box on %s"), *Actor->GetName());
 
 
-						FBox BBox = ConvexElem.ElemBox;						
+						FBox BBox = ConvexElem.ElemBox;												
+						FVector3f Center = (BBox.Max + BBox.Min) * .5;
 
 						OutAssetArrays->SourceSceneProxy[BoxCount] = StaticMeshComponent->SceneProxy;
-						const FTransform ElementTransform = ConvexElem.GetTransform() * MeshTransform;
+						const FTransform ElementTransform = FTransform(Center) * MeshTransform;
 						FillCurrentTransforms(ElementTransform, BoxCount, OutAssetArrays->CurrentTransform, OutAssetArrays->CurrentInverse);
 					}
 				}
