suppress repetitive warnings about convex collisions

diff --git a/Engine/Plugins/FX/Niagara/Source/Niagara/Private/NiagaraDataInterfaceRigidMeshCollisionQuery.cpp b/Engine/Plugins/FX/Niagara/Source/Niagara/Private/NiagaraDataInterfaceRigidMeshCollisionQuery.cpp
--- a/Engine/Plugins/FX/Niagara/Source/Niagara/Private/NiagaraDataInterfaceRigidMeshCollisionQuery.cpp
+++ b/Engine/Plugins/FX/Niagara/Source/Niagara/Private/NiagaraDataInterfaceRigidMeshCollisionQuery.cpp
@@ -364,9 +364,6 @@ void UpdateInternalArrays(const TArray<AActor*> &Actors, FNDIRigidMeshCollisionA
 				{
 					if (CollisionEnabledHasPhysics(ConvexElem.GetCollisionEnabled()))
 					{
-						UE_LOG(LogRigidMeshCollision, Warning, TEXT("Convex collision objects encountered and will be interpreted as a bounding box on %s"), *Actor->GetName());
-
-
 						FBox BBox = ConvexElem.ElemBox;												
 						FVector3f Center = (BBox.Max + BBox.Min) * .5;
 
