- Fix DebugName not always being correct / cached
- Fix crash when execution order is enabled and we are re-compiling

#rnx

diff --git a/Engine/Plugins/FX/Niagara/Source/Niagara/Private/NiagaraEmitter.cpp b/Engine/Plugins/FX/Niagara/Source/Niagara/Private/NiagaraEmitter.cpp
--- a/Engine/Plugins/FX/Niagara/Source/Niagara/Private/NiagaraEmitter.cpp
+++ b/Engine/Plugins/FX/Niagara/Source/Niagara/Private/NiagaraEmitter.cpp
@@ -1240,6 +1240,17 @@ void UNiagaraEmitter::CacheFromCompiledData(const FNiagaraDataSetCompiledData* C
 		);
 	}
 #endif
+
+#if !UE_BUILD_SHIPPING
+	// Ensure our debug simulation name is up to date
+	DebugSimName.Empty();
+	if (const UNiagaraSystem* SystemOwner = Cast<const UNiagaraSystem>(GetOuter()))
+	{
+		DebugSimName = SystemOwner->GetName();
+		DebugSimName.AppendChar(':');
+	}
+	DebugSimName.Append(GetName());
+#endif
 }
 
 void UNiagaraEmitter::CacheFromShaderCompiled()
diff --git a/Engine/Plugins/FX/Niagara/Source/NiagaraEditor/Private/Widgets/SNiagaraSystemViewport.cpp b/Engine/Plugins/FX/Niagara/Source/NiagaraEditor/Private/Widgets/SNiagaraSystemViewport.cpp
--- a/Engine/Plugins/FX/Niagara/Source/NiagaraEditor/Private/Widgets/SNiagaraSystemViewport.cpp
+++ b/Engine/Plugins/FX/Niagara/Source/NiagaraEditor/Private/Widgets/SNiagaraSystemViewport.cpp
@@ -300,15 +300,18 @@ void FNiagaraSystemViewportClient::DrawEmitterExecutionOrder(UNiagaraComponent*
 	Canvas->DrawShadowedString(CurrentX, CurrentY, TEXT("Emitter Execution Order"), Font, FLinearColor::White);
 	CurrentY += FontHeight;
 
-	TConstArrayView<FNiagaraEmitterExecutionIndex> ExecutionOrder = NiagaraSystem->GetEmitterExecutionOrder();
-	int32 DisplayIndex = 0;
-	for (const FNiagaraEmitterExecutionIndex& EmitterExecIndex : ExecutionOrder)
+	if ( NiagaraSystem->IsReadyToRun() )
 	{
-		const FNiagaraEmitterHandle& EmitterHandle = NiagaraSystem->GetEmitterHandle(EmitterExecIndex.EmitterIndex);
-		if (UNiagaraEmitter* NiagaraEmitter = EmitterHandle.GetInstance())
+		TConstArrayView<FNiagaraEmitterExecutionIndex> ExecutionOrder = NiagaraSystem->GetEmitterExecutionOrder();
+		int32 DisplayIndex = 0;
+		for (const FNiagaraEmitterExecutionIndex& EmitterExecIndex : ExecutionOrder)
 		{
-			Canvas->DrawShadowedString(CurrentX, CurrentY, *FString::Printf(TEXT("%d - %s"), ++DisplayIndex, NiagaraEmitter->GetDebugSimName()), Font, FLinearColor::White);
-			CurrentY += FontHeight;
+			const FNiagaraEmitterHandle& EmitterHandle = NiagaraSystem->GetEmitterHandle(EmitterExecIndex.EmitterIndex);
+			if (UNiagaraEmitter* NiagaraEmitter = EmitterHandle.GetInstance())
+			{
+				Canvas->DrawShadowedString(CurrentX, CurrentY, *FString::Printf(TEXT("%d - %s"), ++DisplayIndex, NiagaraEmitter->GetDebugSimName()), Font, FLinearColor::White);
+				CurrentY += FontHeight;
+			}
 		}
 	}
 }
