- Fix various static analysis warnings

#rnx

diff --git a/Engine/Plugins/FX/Niagara/Source/Niagara/Private/NiagaraGpuComputeDispatch.cpp b/Engine/Plugins/FX/Niagara/Source/Niagara/Private/NiagaraGpuComputeDispatch.cpp
--- a/Engine/Plugins/FX/Niagara/Source/Niagara/Private/NiagaraGpuComputeDispatch.cpp
+++ b/Engine/Plugins/FX/Niagara/Source/Niagara/Private/NiagaraGpuComputeDispatch.cpp
@@ -2074,7 +2074,8 @@ void FNiagaraGpuComputeDispatch::SetDataInterfaceParameters(FRHICommandList& RHI
 	uint32 InterfaceIndex = 0;
 	for (FNiagaraDataInterfaceProxy* Interface : InstanceData.DataInterfaceProxies)
 	{
-		const FNiagaraDataInterfaceParamRef& DIParam = ComputeShader->GetDIParameters()[InterfaceIndex];
+		const TMemoryImageArray<FNiagaraDataInterfaceParamRef>& DIParameters = ComputeShader->GetDIParameters();
+		const FNiagaraDataInterfaceParamRef& DIParam = DIParameters[InterfaceIndex];
 		if (DIParam.Parameters.IsValid())
 		{
 			FNiagaraDataInterfaceSetArgs Context(Interface, Tick.SystemInstanceID, this, ComputeShader, &InstanceData, &SimStageData, InstanceData.IsOutputStage(Interface, SimStageData.StageIndex), InstanceData.IsIterationStage(Interface, SimStageData.StageIndex));
@@ -2093,6 +2094,7 @@ void FNiagaraGpuComputeDispatch::UnsetDataInterfaceParameters(FRHICommandList& R
 	uint32 InterfaceIndex = 0;
 	for (FNiagaraDataInterfaceProxy* Interface : InstanceData.DataInterfaceProxies)
 	{
+		const TMemoryImageArray<FNiagaraDataInterfaceParamRef>& DIParameters = ComputeShader->GetDIParameters();
 		const FNiagaraDataInterfaceParamRef& DIParam = ComputeShader->GetDIParameters()[InterfaceIndex];
 		if (DIParam.Parameters.IsValid())
 		{
diff --git a/Engine/Plugins/FX/Niagara/Source/NiagaraEditor/Private/NiagaraCompiler.cpp b/Engine/Plugins/FX/Niagara/Source/NiagaraEditor/Private/NiagaraCompiler.cpp
--- a/Engine/Plugins/FX/Niagara/Source/NiagaraEditor/Private/NiagaraCompiler.cpp
+++ b/Engine/Plugins/FX/Niagara/Source/NiagaraEditor/Private/NiagaraCompiler.cpp
@@ -176,7 +176,7 @@ void FNiagaraCompileRequestDuplicateData::DuplicateReferencedGraphsRecursive(UNi
 					if (bHasNumericParams || SharedSourceGraphToDuplicatedGraphsMap->Contains(FunctionGraph) == false)
 					{
 						// Duplicate the script, the source, and graph
-						UNiagaraScript* DupeScript = DupeScript = FunctionScript->CreateCompilationCopy();
+						UNiagaraScript* DupeScript = FunctionScript->CreateCompilationCopy();
 						TArray<ENiagaraScriptUsage> CompileUsages = { DupeScript->GetUsage() };
 						UNiagaraScriptSource* DupeScriptSource = CastChecked<UNiagaraScriptSource>(DupeScript->GetSource(FunctionCallNode->SelectedScriptVersion))->CreateCompilationCopy(CompileUsages);
 						UNiagaraGraph* DupeGraph = DupeScriptSource->NodeGraph;
diff --git a/Engine/Plugins/FX/Niagara/Source/NiagaraEditorWidgets/Private/Stack/SNiagaraStackFunctionInputValue.cpp b/Engine/Plugins/FX/Niagara/Source/NiagaraEditorWidgets/Private/Stack/SNiagaraStackFunctionInputValue.cpp
--- a/Engine/Plugins/FX/Niagara/Source/NiagaraEditorWidgets/Private/Stack/SNiagaraStackFunctionInputValue.cpp
+++ b/Engine/Plugins/FX/Niagara/Source/NiagaraEditorWidgets/Private/Stack/SNiagaraStackFunctionInputValue.cpp
@@ -1028,7 +1028,6 @@ TArray<TSharedPtr<FNiagaraMenuAction_Generic>> SNiagaraStackFunctionInputValue::
 
 		TSharedPtr<FNiagaraMenuAction_Generic> ExpressionAction(new FNiagaraMenuAction_Generic(
                 FNiagaraMenuAction_Generic::FOnExecuteAction::CreateSP(this, &SNiagaraStackFunctionInputValue::CustomExpressionSelected),
-                FNiagaraMenuAction_Generic::FCanExecuteAction::CreateLambda([bIsDataInterfaceOrObject]() { return bIsDataInterfaceOrObject == false; }),
                 DisplayName, ENiagaraMenuSections::General, {}, Tooltip, FText()));
 
 		ExpressionAction->SourceData = NiagaraSourceData;
@@ -1046,7 +1045,6 @@ TArray<TSharedPtr<FNiagaraMenuAction_Generic>> SNiagaraStackFunctionInputValue::
 
 		TSharedPtr<FNiagaraMenuAction_Generic> CreateScratchAction(new FNiagaraMenuAction_Generic(
            FNiagaraMenuAction_Generic::FOnExecuteAction::CreateSP(this, &SNiagaraStackFunctionInputValue::CreateScratchSelected),
-           FNiagaraMenuAction_Generic::FCanExecuteAction::CreateLambda([bIsDataInterfaceOrObject]() { return bIsDataInterfaceOrObject == false; }),
            DisplayName, ENiagaraMenuSections::General, {}, Tooltip, FText()));
 
 		CreateScratchAction->SourceData = NiagaraSourceData;
diff --git a/Engine/Plugins/FX/Niagara/Source/NiagaraEditorWidgets/Private/Stack/SNiagaraStackTableRow.cpp b/Engine/Plugins/FX/Niagara/Source/NiagaraEditorWidgets/Private/Stack/SNiagaraStackTableRow.cpp
--- a/Engine/Plugins/FX/Niagara/Source/NiagaraEditorWidgets/Private/Stack/SNiagaraStackTableRow.cpp
+++ b/Engine/Plugins/FX/Niagara/Source/NiagaraEditorWidgets/Private/Stack/SNiagaraStackTableRow.cpp
@@ -44,7 +44,7 @@ void SNiagaraStackTableRow::Construct(const FArguments& InArgs, UNiagaraStackVie
 	ForegroundColor = InArgs._ItemForegroundColor;
 	IndicatorColor = InArgs._IndicatorColor;
 
-	ExecutionCategoryToolTipText = InStackEntry->GetExecutionSubcategoryName() != NAME_None
+	ExecutionCategoryToolTipText = (InStackEntry->GetExecutionSubcategoryName() != NAME_None)
 		? FText::Format(LOCTEXT("ExecutionCategoryToolTipFormat", "{0} - {1}"), FText::FromName(InStackEntry->GetExecutionCategoryName()), FText::FromName(InStackEntry->GetExecutionSubcategoryName()))
 		: FText::FromName(InStackEntry->GetExecutionCategoryName());
 
