- Fix for element -> bytes in offset

#rnx

diff --git a/Engine/Plugins/FX/Niagara/Source/Niagara/Private/DataInterface/NiagaraDataInterfaceSimpleCounter.cpp b/Engine/Plugins/FX/Niagara/Source/Niagara/Private/DataInterface/NiagaraDataInterfaceSimpleCounter.cpp
--- a/Engine/Plugins/FX/Niagara/Source/Niagara/Private/DataInterface/NiagaraDataInterfaceSimpleCounter.cpp
+++ b/Engine/Plugins/FX/Niagara/Source/Niagara/Private/DataInterface/NiagaraDataInterfaceSimpleCounter.cpp
@@ -79,7 +79,7 @@ struct FNDISimpleCounterProxy : public FNiagaraDataInterfaceProxy
 				ReadbackManager->EnqueueReadback(
 					RHICmdList,
 					CountManager.GetInstanceCountBuffer().Buffer,
-					InstanceData->CountOffset, sizeof(int32),
+					InstanceData->CountOffset * sizeof(uint32), sizeof(uint32),
 					[SystemInstanceID=Context.SystemInstanceID, WeakOwner=WeakOwner, Proxy=this](TConstArrayView<TPair<void*, uint32>> ReadbackData)
 					{
 						const int32 CounterValue = *reinterpret_cast<const int32*>(ReadbackData[0].Key);
