Small improvements to debug hud.

- Bug fix for broken state when re-entering perf graph mode.
- Fix for case where we can leak listeners.
- System colours are now customizable
- Fix for coloured bg tile width.

diff --git a/Engine/Plugins/FX/Niagara/Source/Niagara/Private/NiagaraDebugHud.cpp b/Engine/Plugins/FX/Niagara/Source/Niagara/Private/NiagaraDebugHud.cpp
--- a/Engine/Plugins/FX/Niagara/Source/Niagara/Private/NiagaraDebugHud.cpp
+++ b/Engine/Plugins/FX/Niagara/Source/Niagara/Private/NiagaraDebugHud.cpp
@@ -884,10 +884,11 @@ void FNiagaraDebugHud::GatherSystemInfo()
 		}
 
 		//Generate a unique-ish random color for use in graphs and the world to help visually ID this system.
-		FRandomStream Rands(GetTypeHash(NiagaraComponent->GetAsset()->GetName()));
-		uint8 RandomHue = (uint8)Rands.RandRange(0, 255);
-		uint8 RandomValue = (uint8)Rands.RandRange(128, 255);
-		SystemDebugInfo.UniqueColor = FLinearColor::MakeFromHSV8(RandomHue, 255, RandomValue);
+		FRandomStream Rands(GetTypeHash(NiagaraComponent->GetAsset()->GetName()) + Settings.SystemColorSeed);
+		uint8 RandomHue = (uint8)Rands.RandRange(Settings.SystemColorHSVMin.X, Settings.SystemColorHSVMax.X);
+		uint8 RandomSat = (uint8)Rands.RandRange(Settings.SystemColorHSVMin.Y, Settings.SystemColorHSVMax.Y);
+		uint8 RandomValue = (uint8)Rands.RandRange(Settings.SystemColorHSVMin.Z, Settings.SystemColorHSVMax.Z);
+		SystemDebugInfo.UniqueColor = FLinearColor::MakeFromHSV8(RandomHue, RandomSat, RandomValue);
 
 
 		bool bWillBeVisible = false;
@@ -1181,9 +1182,9 @@ void FNiagaraDebugHud::DrawOverview(class FNiagaraWorldManager* WorldManager, FC
 		{
 			HeaderWidth = GetStringSize(Font, *GlobalHeader).X;
 			SystemWidth = FMath::Max(GetStringSize(Font, *SystemHeader).X, GetStringSize(Font, *ExampleSystemString).X);
-			MaxWidth = FMath::Max(HeaderWidth + GlobalDataSeparator + GetStringSize(Font, *GlobalData).X, SystemWidth);
+			MaxWidth = FMath::Max(HeaderWidth + GlobalDataSeparator + GetStringSize(Font, *GlobalData).X, SystemWidth) + ColumnSeparator;
 
-			InOffset += MaxWidth + ColumnSeparator;
+			InOffset += MaxWidth;
 		}
 
 		void DrawGlobalHeader(FCanvas* DrawCanvas, UFont* Font, float X, float Y, FLinearColor HeadingColor, FLinearColor DetailColor)
@@ -1213,7 +1214,7 @@ void FNiagaraDebugHud::DrawOverview(class FNiagaraWorldManager* WorldManager, FC
 			[&DetailColor, &DetailHighlightColor, &fAdvanceHeight](FCanvas* Canvas, UFont* Font, float X, float Y, FOverviewColumn& Col, const FSystemDebugInfo& SystemInfo)
 			{
 				FLinearColor RowBGColor = SystemInfo.UniqueColor;
-				RowBGColor.A = 0.1f;
+				RowBGColor.A = Settings.SystemColorTableOpacity;
 				Canvas->DrawTile(X, Y, Col.MaxWidth, fAdvanceHeight, 0,0,0,0, RowBGColor);
 				const FLinearColor RowColor = SystemInfo.bShowInWorld ? DetailHighlightColor : DetailColor;
 				Canvas->DrawShadowedString(X, Y, *SystemInfo.SystemName, Font, RowColor);
@@ -1225,7 +1226,7 @@ void FNiagaraDebugHud::DrawOverview(class FNiagaraWorldManager* WorldManager, FC
 				[&DetailColor, &DetailHighlightColor, &fAdvanceHeight](FCanvas* Canvas, UFont* Font, float X, float Y, FOverviewColumn& Col, const FSystemDebugInfo& SystemInfo)
 				{
 					FLinearColor RowBGColor = SystemInfo.UniqueColor;
-					RowBGColor.A = 0.1f;
+					RowBGColor.A = Settings.SystemColorTableOpacity;
 					Canvas->DrawTile(X, Y, Col.MaxWidth, fAdvanceHeight, 0, 0, 0, 0, RowBGColor);
 					const FLinearColor RowColor = SystemInfo.bShowInWorld ? DetailHighlightColor : DetailColor;
 					Canvas->DrawShadowedString(X, Y, *FString::FromInt(SystemInfo.TotalSystems), Font, RowColor);
@@ -1235,7 +1236,7 @@ void FNiagaraDebugHud::DrawOverview(class FNiagaraWorldManager* WorldManager, FC
 				[&DetailColor, &DetailHighlightColor, &fAdvanceHeight](FCanvas* Canvas, UFont* Font, float X, float Y, FOverviewColumn& Col, const FSystemDebugInfo& SystemInfo)
 				{
 					FLinearColor RowBGColor = SystemInfo.UniqueColor;
-					RowBGColor.A = 0.1f;
+					RowBGColor.A = Settings.SystemColorTableOpacity;
 					Canvas->DrawTile(X, Y, Col.MaxWidth, fAdvanceHeight, 0, 0, 0, 0, RowBGColor);
 					const FLinearColor RowColor = SystemInfo.bShowInWorld ? DetailHighlightColor : DetailColor;
 					Canvas->DrawShadowedString(X, Y, *FString::FromInt(SystemInfo.TotalScalability), Font, RowColor);
@@ -1245,7 +1246,7 @@ void FNiagaraDebugHud::DrawOverview(class FNiagaraWorldManager* WorldManager, FC
 				[&DetailColor, &DetailHighlightColor, &fAdvanceHeight](FCanvas* Canvas, UFont* Font, float X, float Y, FOverviewColumn& Col, const FSystemDebugInfo& SystemInfo)
 				{
 					FLinearColor RowBGColor = SystemInfo.UniqueColor;
-					RowBGColor.A = 0.1f;
+					RowBGColor.A = Settings.SystemColorTableOpacity;
 					Canvas->DrawTile(X, Y, Col.MaxWidth, fAdvanceHeight, 0, 0, 0, 0, RowBGColor);
 					const FLinearColor RowColor = SystemInfo.bShowInWorld ? DetailHighlightColor : DetailColor;
 					Canvas->DrawShadowedString(X, Y, *FString::FromInt(SystemInfo.TotalEmitters), Font, RowColor);
@@ -1255,7 +1256,7 @@ void FNiagaraDebugHud::DrawOverview(class FNiagaraWorldManager* WorldManager, FC
 				[&DetailColor, &DetailHighlightColor, &fAdvanceHeight](FCanvas* Canvas, UFont* Font, float X, float Y, FOverviewColumn& Col, const FSystemDebugInfo& SystemInfo)
 				{
 					FLinearColor RowBGColor = SystemInfo.UniqueColor;
-					RowBGColor.A = 0.1f;
+					RowBGColor.A = Settings.SystemColorTableOpacity;
 					Canvas->DrawTile(X, Y, Col.MaxWidth, fAdvanceHeight, 0, 0, 0, 0, RowBGColor);
 					const FLinearColor RowColor = SystemInfo.bShowInWorld ? DetailHighlightColor : DetailColor;
 					Canvas->DrawShadowedString(X, Y, *FString::FromInt(SystemInfo.TotalParticles), Font, RowColor);
@@ -1265,7 +1266,7 @@ void FNiagaraDebugHud::DrawOverview(class FNiagaraWorldManager* WorldManager, FC
 				[&DetailColor, &DetailHighlightColor, &fAdvanceHeight](FCanvas* Canvas, UFont* Font, float X, float Y, FOverviewColumn& Col, const FSystemDebugInfo& SystemInfo)
 				{
 					FLinearColor RowBGColor = SystemInfo.UniqueColor;
-					RowBGColor.A = 0.1f;
+					RowBGColor.A = Settings.SystemColorTableOpacity;
 					Canvas->DrawTile(X, Y, Col.MaxWidth, fAdvanceHeight, 0, 0, 0, 0, RowBGColor);
 					const FLinearColor RowColor = SystemInfo.bShowInWorld ? DetailHighlightColor : DetailColor;
 					Canvas->DrawShadowedString(X, Y, *FString::Printf(TEXT("%6.2f"), double(SystemInfo.TotalBytes) / (1024.0 * 1024.0)), Font, RowColor);
@@ -1277,7 +1278,7 @@ void FNiagaraDebugHud::DrawOverview(class FNiagaraWorldManager* WorldManager, FC
 				[&DetailColor, &DetailHighlightColor, &fAdvanceHeight](FCanvas* Canvas, UFont* Font, float X, float Y, FOverviewColumn& Col, const FSystemDebugInfo& SystemInfo)
 				{
 					FLinearColor RowBGColor = SystemInfo.UniqueColor;
-					RowBGColor.A = 0.1f;
+					RowBGColor.A = Settings.SystemColorTableOpacity;
 					Canvas->DrawTile(X, Y, Col.MaxWidth, fAdvanceHeight, 0, 0, 0, 0, RowBGColor);
 					const FLinearColor RowColor = SystemInfo.bShowInWorld ? DetailHighlightColor : DetailColor;
 					Canvas->DrawShadowedString(X, Y, *FString::FromInt(SystemInfo.TotalSystems), Font, RowColor);
@@ -1287,7 +1288,7 @@ void FNiagaraDebugHud::DrawOverview(class FNiagaraWorldManager* WorldManager, FC
 				[&DetailColor, &DetailHighlightColor, &fAdvanceHeight](FCanvas* Canvas, UFont* Font, float X, float Y, FOverviewColumn& Col, const FSystemDebugInfo& SystemInfo)
 				{
 					FLinearColor RowBGColor = SystemInfo.UniqueColor;
-					RowBGColor.A = 0.1f;
+					RowBGColor.A = Settings.SystemColorTableOpacity;
 					Canvas->DrawTile(X, Y, Col.MaxWidth, fAdvanceHeight, 0, 0, 0, 0, RowBGColor);
 					const FLinearColor RowColor = SystemInfo.bShowInWorld ? DetailHighlightColor : DetailColor;
 					Canvas->DrawShadowedString(X, Y, *FString::FromInt(SystemInfo.TotalScalability), Font, RowColor);
@@ -1297,7 +1298,7 @@ void FNiagaraDebugHud::DrawOverview(class FNiagaraWorldManager* WorldManager, FC
 				[&DetailColor, &DetailHighlightColor, &fAdvanceHeight](FCanvas* Canvas, UFont* Font, float X, float Y, FOverviewColumn& Col, const FSystemDebugInfo& SystemInfo)
 				{
 					FLinearColor RowBGColor = SystemInfo.UniqueColor;
-					RowBGColor.A = 0.1f;
+					RowBGColor.A = Settings.SystemColorTableOpacity;
 					Canvas->DrawTile(X, Y, Col.MaxWidth, fAdvanceHeight, 0, 0, 0, 0, RowBGColor);
 					const FLinearColor RowColor = SystemInfo.bShowInWorld ? DetailHighlightColor : DetailColor;
 					Canvas->DrawShadowedString(X, Y, *FString::FromInt(SystemInfo.TotalCulled), Font, RowColor);
@@ -1307,7 +1308,7 @@ void FNiagaraDebugHud::DrawOverview(class FNiagaraWorldManager* WorldManager, FC
 				[&DetailColor, &DetailHighlightColor, &fAdvanceHeight](FCanvas* Canvas, UFont* Font, float X, float Y, FOverviewColumn& Col, const FSystemDebugInfo& SystemInfo)
 				{
 					FLinearColor RowBGColor = SystemInfo.UniqueColor;
-					RowBGColor.A = 0.1f;
+					RowBGColor.A = Settings.SystemColorTableOpacity;
 					Canvas->DrawTile(X, Y, Col.MaxWidth, fAdvanceHeight, 0, 0, 0, 0, RowBGColor);
 					const FLinearColor RowColor = SystemInfo.bShowInWorld ? DetailHighlightColor : DetailColor;
 					Canvas->DrawShadowedString(X, Y, *FString::FromInt(SystemInfo.TotalCulledByDistance), Font, RowColor);
@@ -1317,7 +1318,7 @@ void FNiagaraDebugHud::DrawOverview(class FNiagaraWorldManager* WorldManager, FC
 				[&DetailColor, &DetailHighlightColor, &fAdvanceHeight](FCanvas* Canvas, UFont* Font, float X, float Y, FOverviewColumn& Col, const FSystemDebugInfo& SystemInfo)
 				{
 					FLinearColor RowBGColor = SystemInfo.UniqueColor;
-					RowBGColor.A = 0.1f;
+					RowBGColor.A = Settings.SystemColorTableOpacity;
 					Canvas->DrawTile(X, Y, Col.MaxWidth, fAdvanceHeight, 0, 0, 0, 0, RowBGColor);
 					const FLinearColor RowColor = SystemInfo.bShowInWorld ? DetailHighlightColor : DetailColor;
 					Canvas->DrawShadowedString(X, Y, *FString::FromInt(SystemInfo.TotalCulledByVisibility), Font, RowColor);
@@ -1327,7 +1328,7 @@ void FNiagaraDebugHud::DrawOverview(class FNiagaraWorldManager* WorldManager, FC
 				[&DetailColor, &DetailHighlightColor, &fAdvanceHeight](FCanvas* Canvas, UFont* Font, float X, float Y, FOverviewColumn& Col, const FSystemDebugInfo& SystemInfo)
 				{
 					FLinearColor RowBGColor = SystemInfo.UniqueColor;
-					RowBGColor.A = 0.1f;
+					RowBGColor.A = Settings.SystemColorTableOpacity;
 					Canvas->DrawTile(X, Y, Col.MaxWidth, fAdvanceHeight, 0, 0, 0, 0, RowBGColor);
 					const FLinearColor RowColor = SystemInfo.bShowInWorld ? DetailHighlightColor : DetailColor;
 					Canvas->DrawShadowedString(X, Y, *FString::FromInt(SystemInfo.TotalCulledByInstanceCount), Font, RowColor);
@@ -1337,7 +1338,7 @@ void FNiagaraDebugHud::DrawOverview(class FNiagaraWorldManager* WorldManager, FC
 				[&DetailColor, &DetailHighlightColor, &fAdvanceHeight](FCanvas* Canvas, UFont* Font, float X, float Y, FOverviewColumn& Col, const FSystemDebugInfo& SystemInfo)
 				{
 					FLinearColor RowBGColor = SystemInfo.UniqueColor;
-					RowBGColor.A = 0.1f;
+					RowBGColor.A = Settings.SystemColorTableOpacity;
 					Canvas->DrawTile(X, Y, Col.MaxWidth, fAdvanceHeight, 0, 0, 0, 0, RowBGColor);
 					const FLinearColor RowColor = SystemInfo.bShowInWorld ? DetailHighlightColor : DetailColor;
 					Canvas->DrawShadowedString(X, Y, *FString::FromInt(SystemInfo.TotalCulledByBudget), Font, RowColor);
@@ -1347,7 +1348,7 @@ void FNiagaraDebugHud::DrawOverview(class FNiagaraWorldManager* WorldManager, FC
 				[&DetailColor, &DetailHighlightColor, &fAdvanceHeight](FCanvas* Canvas, UFont* Font, float X, float Y, FOverviewColumn& Col, const FSystemDebugInfo& SystemInfo)
 				{
 					FLinearColor RowBGColor = SystemInfo.UniqueColor;
-					RowBGColor.A = 0.1f;
+					RowBGColor.A = Settings.SystemColorTableOpacity;
 					Canvas->DrawTile(X, Y, Col.MaxWidth, fAdvanceHeight, 0, 0, 0, 0, RowBGColor);
 					const FLinearColor RowColor = SystemInfo.bShowInWorld ? DetailHighlightColor : DetailColor;
 					Canvas->DrawShadowedString(X, Y, *FString::FromInt(SystemInfo.TotalPlayerSystems), Font, RowColor);
@@ -1364,7 +1365,7 @@ void FNiagaraDebugHud::DrawOverview(class FNiagaraWorldManager* WorldManager, FC
 					[&DetailColor, &DetailHighlightColor, &fAdvanceHeight](FCanvas* Canvas, UFont* Font, float X, float Y, FOverviewColumn& Col, const FSystemDebugInfo& SystemInfo)
 					{
 						FLinearColor RowBGColor = SystemInfo.UniqueColor;
-						RowBGColor.A = 0.1f;
+						RowBGColor.A = Settings.SystemColorTableOpacity;
 						Canvas->DrawTile(X, Y, Col.MaxWidth, fAdvanceHeight, 0, 0, 0, 0, RowBGColor);
 						const FLinearColor RowColor = SystemInfo.bShowInWorld ? DetailHighlightColor : DetailColor;
 						Canvas->DrawShadowedString(X, Y, *LexToSanitizedString(SystemInfo.PerfStats ? SystemInfo.PerfStats->Avg.Time_GT : 0.0), Font, RowColor);
@@ -1374,7 +1375,7 @@ void FNiagaraDebugHud::DrawOverview(class FNiagaraWorldManager* WorldManager, FC
 					[&DetailColor, &DetailHighlightColor, &fAdvanceHeight](FCanvas* Canvas, UFont* Font, float X, float Y, FOverviewColumn& Col, const FSystemDebugInfo& SystemInfo)
 					{
 						FLinearColor RowBGColor = SystemInfo.UniqueColor;
-						RowBGColor.A = 0.1f;
+						RowBGColor.A = Settings.SystemColorTableOpacity;
 						Canvas->DrawTile(X, Y, Col.MaxWidth, fAdvanceHeight, 0, 0, 0, 0, RowBGColor);
 						const FLinearColor RowColor = SystemInfo.bShowInWorld ? DetailHighlightColor : DetailColor;
 						Canvas->DrawShadowedString(X, Y, *LexToSanitizedString(SystemInfo.PerfStats ? SystemInfo.PerfStats->Max.Time_GT : 0.0), Font, RowColor);
@@ -1384,7 +1385,7 @@ void FNiagaraDebugHud::DrawOverview(class FNiagaraWorldManager* WorldManager, FC
 					[&DetailColor, &DetailHighlightColor, &fAdvanceHeight](FCanvas* Canvas, UFont* Font, float X, float Y, FOverviewColumn& Col, const FSystemDebugInfo& SystemInfo)
 					{
 						FLinearColor RowBGColor = SystemInfo.UniqueColor;
-						RowBGColor.A = 0.1f;
+						RowBGColor.A = Settings.SystemColorTableOpacity;
 						Canvas->DrawTile(X, Y, Col.MaxWidth, fAdvanceHeight, 0, 0, 0, 0, RowBGColor);
 						const FLinearColor RowColor = SystemInfo.bShowInWorld ? DetailHighlightColor : DetailColor;
 						Canvas->DrawShadowedString(X, Y, *LexToSanitizedString(SystemInfo.PerfStats ? SystemInfo.PerfStats->Avg.Time_RT : 0.0), Font, RowColor);
@@ -1394,7 +1395,7 @@ void FNiagaraDebugHud::DrawOverview(class FNiagaraWorldManager* WorldManager, FC
 					[&DetailColor, &DetailHighlightColor, &fAdvanceHeight](FCanvas* Canvas, UFont* Font, float X, float Y, FOverviewColumn& Col, const FSystemDebugInfo& SystemInfo)
 					{
 						FLinearColor RowBGColor = SystemInfo.UniqueColor;
-						RowBGColor.A = 0.1f;
+						RowBGColor.A = Settings.SystemColorTableOpacity;
 						Canvas->DrawTile(X, Y, Col.MaxWidth, fAdvanceHeight, 0, 0, 0, 0, RowBGColor);
 						const FLinearColor RowColor = SystemInfo.bShowInWorld ? DetailHighlightColor : DetailColor;
 						Canvas->DrawShadowedString(X, Y, *LexToSanitizedString(SystemInfo.PerfStats ? SystemInfo.PerfStats->Max.Time_RT : 0.0), Font, RowColor);
@@ -1404,7 +1405,7 @@ void FNiagaraDebugHud::DrawOverview(class FNiagaraWorldManager* WorldManager, FC
 					[&DetailColor, &DetailHighlightColor, &fAdvanceHeight](FCanvas* Canvas, UFont* Font, float X, float Y, FOverviewColumn& Col, const FSystemDebugInfo& SystemInfo)
 					{
 						FLinearColor RowBGColor = SystemInfo.UniqueColor;
-						RowBGColor.A = 0.1f;
+						RowBGColor.A = Settings.SystemColorTableOpacity;
 						Canvas->DrawTile(X, Y, Col.MaxWidth, fAdvanceHeight, 0, 0, 0, 0, RowBGColor);
 						const FLinearColor RowColor = SystemInfo.bShowInWorld ? DetailHighlightColor : DetailColor;
 						Canvas->DrawShadowedString(X, Y, *LexToSanitizedString(SystemInfo.PerfStats ? SystemInfo.PerfStats->Avg.Time_GPU : 0.0), Font, RowColor);
@@ -1414,7 +1415,7 @@ void FNiagaraDebugHud::DrawOverview(class FNiagaraWorldManager* WorldManager, FC
 					[&DetailColor, &DetailHighlightColor, &fAdvanceHeight](FCanvas* Canvas, UFont* Font, float X, float Y, FOverviewColumn& Col, const FSystemDebugInfo& SystemInfo)
 					{
 						FLinearColor RowBGColor = SystemInfo.UniqueColor;
-						RowBGColor.A = 0.1f;
+						RowBGColor.A = Settings.SystemColorTableOpacity;
 						Canvas->DrawTile(X, Y, Col.MaxWidth, fAdvanceHeight, 0, 0, 0, 0, RowBGColor);
 						const FLinearColor RowColor = SystemInfo.bShowInWorld ? DetailHighlightColor : DetailColor;
 						Canvas->DrawShadowedString(X, Y, *LexToSanitizedString(SystemInfo.PerfStats ? SystemInfo.PerfStats->Max.Time_GPU : 0.0), Font, RowColor);
diff --git a/Engine/Plugins/FX/Niagara/Source/Niagara/Public/NiagaraDebuggerCommon.h b/Engine/Plugins/FX/Niagara/Source/Niagara/Public/NiagaraDebuggerCommon.h
--- a/Engine/Plugins/FX/Niagara/Source/Niagara/Public/NiagaraDebuggerCommon.h
+++ b/Engine/Plugins/FX/Niagara/Source/Niagara/Public/NiagaraDebuggerCommon.h
@@ -626,6 +626,22 @@ struct NIAGARA_API FNiagaraDebugHUDSettingsData
 	// Message display error text color
 	UPROPERTY(EditAnywhere, Category = "Colors")
 	FLinearColor MessageErrorTextColor = FLinearColor(1.0f, 0.4, 0.3, 1.0f);
+	
+	/** Opacity of the system color background tile in overview table rows. */
+	UPROPERTY(Config, EditAnywhere, Category = "Colors")
+	float SystemColorTableOpacity = 0.2f;
+
+	/** Additional seed value for random system colors. Useful if current colors of systems are too similar. */
+	UPROPERTY(Config, EditAnywhere, Category = "Colors")
+	uint32 SystemColorSeed = 0;
+
+	/** Minimum HSV values for the random colors generated for each System. */
+	UPROPERTY(Config, EditAnywhere, Category = "Colors")
+	FVector SystemColorHSVMin = FVector(0, 200, 200);
+
+	/** Maximum HSV values for the random colors generated for each System. */
+	UPROPERTY(Config, EditAnywhere, Category = "Colors")
+	FVector SystemColorHSVMax = FVector(255, 255, 255);
 
 	UPROPERTY()
 	ENiagaraDebugPlaybackMode PlaybackMode = ENiagaraDebugPlaybackMode::Play;
diff --git a/Engine/Source/Runtime/Engine/Private/Particles/ParticlePerfStats.cpp b/Engine/Source/Runtime/Engine/Private/Particles/ParticlePerfStats.cpp
--- a/Engine/Source/Runtime/Engine/Private/Particles/ParticlePerfStats.cpp
+++ b/Engine/Source/Runtime/Engine/Private/Particles/ParticlePerfStats.cpp
@@ -214,7 +214,7 @@ void FParticlePerfStatsManager::Reset()
 		for (TObjectIterator<UFXSystemComponent> CompIt; CompIt; ++CompIt)
 		{
 			CompIt->ParticlePerfStats = nullptr;
-			CompIt->RecreateRenderState_Concurrent();
+			CompIt->MarkRenderStateDirty();
 		}
 		ComponentToPerfStats.Empty();
 	}
@@ -222,7 +222,7 @@ void FParticlePerfStatsManager::Reset()
 	{
 		for (TObjectIterator<UFXSystemComponent> CompIt; CompIt; ++CompIt)
 		{
-			CompIt->RecreateRenderState_Concurrent();
+			CompIt->MarkRenderStateDirty();
 		}
 	}
 #endif
@@ -240,7 +240,7 @@ void FParticlePerfStatsManager::Tick()
 		TArray<FParticlePerfStatsListenerPtr, TInlineAllocator<8>> ToRemove;
 		for (FParticlePerfStatsListenerPtr& Listener : Listeners)
 		{
-			if (Listener->Tick() == false)
+			if (Listener.IsValid() == false || Listener.IsUnique() || Listener->Tick() == false)
 			{
 				ToRemove.Add(Listener);
 			}
