remove Static Mesh flag

diff --git a/Engine/Plugins/FX/Niagara/Source/Niagara/Private/NiagaraDataInterfaceRigidMeshCollisionQuery.cpp b/Engine/Plugins/FX/Niagara/Source/Niagara/Private/NiagaraDataInterfaceRigidMeshCollisionQuery.cpp
--- a/Engine/Plugins/FX/Niagara/Source/Niagara/Private/NiagaraDataInterfaceRigidMeshCollisionQuery.cpp
+++ b/Engine/Plugins/FX/Niagara/Source/Niagara/Private/NiagaraDataInterfaceRigidMeshCollisionQuery.cpp
@@ -458,19 +458,16 @@ void FNDIRigidMeshCollisionData::Init(UNiagaraDataInterfaceRigidMeshCollisionQue
 
 		Actors.Empty();
 
-		if (Interface->StaticMesh)
+		for (TActorIterator<AActor> It(World); It; ++It)
 		{
-			for (TActorIterator<AActor> It(World); It; ++It)
-			{
-				AActor* Actor = *It;
+			AActor* Actor = *It;
 
-				if ((!Interface->OnlyUseMoveable || (Interface->OnlyUseMoveable && Actor->IsRootComponentMovable())) &&
-					(Interface->Tag == FString("") || (Interface->Tag != FString("") && Actor->Tags.Contains(FName(Interface->Tag)))))
-				{
-					Actors.Add(Actor);
-				}
+			if ((!Interface->OnlyUseMoveable || (Interface->OnlyUseMoveable && Actor->IsRootComponentMovable())) &&
+				(Interface->Tag == FString("") || (Interface->Tag != FString("") && Actor->Tags.Contains(FName(Interface->Tag)))))
+			{
+				Actors.Add(Actor);
 			}
-		}
+		}	
 
 		if (0 < Actors.Num() && Actors[0] != nullptr)
 		{
@@ -492,22 +489,21 @@ void FNDIRigidMeshCollisionData::Update(UNiagaraDataInterfaceRigidMeshCollisionQ
 	{		
 		TArray<AActor*> ActorsTmp;		
 
-		if (Interface->StaticMesh)
+
+		// #todo(dmp): loop over actors and count number of potential colliders.  This might change frame to frame and we might need a
+		// better way of doing this
+		UWorld* World = SystemInstance->GetWorld();
+		for (TActorIterator<AActor> It(World); It; ++It)
 		{
-			// #todo(dmp): loop over actors and count number of potential colliders.  This might change frame to frame and we might need a
-			// better way of doing this
-			UWorld* World = SystemInstance->GetWorld();
-			for (TActorIterator<AActor> It(World); It; ++It)
-			{
-				AActor* Actor = *It;
+			AActor* Actor = *It;
 
-				if ((!Interface->OnlyUseMoveable || (Interface->OnlyUseMoveable && Actor->IsRootComponentMovable())) &&
-					(Interface->Tag == FString("") || (Interface->Tag != FString("") && Actor->Tags.Contains(FName(Interface->Tag)))))
-				{
-					ActorsTmp.Add(Actor);
-				}
+			if ((!Interface->OnlyUseMoveable || (Interface->OnlyUseMoveable && Actor->IsRootComponentMovable())) &&
+				(Interface->Tag == FString("") || (Interface->Tag != FString("") && Actor->Tags.Contains(FName(Interface->Tag)))))
+			{
+				ActorsTmp.Add(Actor);
 			}
 		}
+		
 
 		// check if the number of active actors is the same.  If not, reinitialize
 		if (ActorsTmp.Num() != Actors.Num())
@@ -928,8 +924,7 @@ bool UNiagaraDataInterfaceRigidMeshCollisionQuery::CopyToInternal(UNiagaraDataIn
 	}
 
 	UNiagaraDataInterfaceRigidMeshCollisionQuery* OtherTyped = CastChecked<UNiagaraDataInterfaceRigidMeshCollisionQuery>(Destination);		
-
-	OtherTyped->StaticMesh = StaticMesh;
+	
 	OtherTyped->Tag = Tag;
 	OtherTyped->OnlyUseMoveable = OnlyUseMoveable;
 
@@ -944,7 +939,7 @@ bool UNiagaraDataInterfaceRigidMeshCollisionQuery::Equals(const UNiagaraDataInte
 	}
 	const UNiagaraDataInterfaceRigidMeshCollisionQuery* OtherTyped = CastChecked<const UNiagaraDataInterfaceRigidMeshCollisionQuery>(Other);
 
-	return  (OtherTyped->StaticMesh == StaticMesh) && (OtherTyped->Tag == Tag) && (OtherTyped->OnlyUseMoveable == OnlyUseMoveable);
+	return  (OtherTyped->Tag == Tag) && (OtherTyped->OnlyUseMoveable == OnlyUseMoveable);
 }
 
 void UNiagaraDataInterfaceRigidMeshCollisionQuery::PostInitProperties()
diff --git a/Engine/Plugins/FX/Niagara/Source/Niagara/Public/NiagaraDataInterfaceRigidMeshCollisionQuery.h b/Engine/Plugins/FX/Niagara/Source/Niagara/Public/NiagaraDataInterfaceRigidMeshCollisionQuery.h
--- a/Engine/Plugins/FX/Niagara/Source/Niagara/Public/NiagaraDataInterfaceRigidMeshCollisionQuery.h
+++ b/Engine/Plugins/FX/Niagara/Source/Niagara/Public/NiagaraDataInterfaceRigidMeshCollisionQuery.h
@@ -167,12 +167,9 @@ public:
 	DECLARE_NIAGARA_DI_PARAMETER();
 
 	UPROPERTY(EditAnywhere, Category = "Static Mesh")
-	bool StaticMesh = true;
-
-	UPROPERTY(EditAnywhere, Category = "Static Mesh", meta = (EditCondition = "StaticMesh"))
 	FString Tag = TEXT("");
 
-	UPROPERTY(EditAnywhere, Category = "Static Mesh", meta = (EditCondition = "StaticMesh"))
+	UPROPERTY(EditAnywhere, Category = "Static Mesh")
 	bool OnlyUseMoveable = true;
 
 	UPROPERTY(EditAnywhere, Category = "General")
