Fix for crash when priming Niagara Pools.

Flow through pool priming can cause us to try to compile a system and reset all current instances inside a call to InitializeSystem() on a pooled component.
Leading to a SystemInstanceController that we're initialiing being nulled.
Fix avoids this by checking the system is ready and triggering any recompile/resets before trying to create any systems for the pool.

diff --git a/Engine/Plugins/FX/Niagara/Source/Niagara/Private/NiagaraComponentPool.cpp b/Engine/Plugins/FX/Niagara/Source/Niagara/Private/NiagaraComponentPool.cpp
--- a/Engine/Plugins/FX/Niagara/Source/Niagara/Private/NiagaraComponentPool.cpp
+++ b/Engine/Plugins/FX/Niagara/Source/Niagara/Private/NiagaraComponentPool.cpp
@@ -334,7 +334,13 @@ void UNiagaraComponentPool::PrimePool(UNiagaraSystem* Template, UWorld* World)
 		UE_LOG(LogNiagara, Verbose, TEXT("Failed to prime particle pool as the world does not have a scene."));
 		return;
 	}	
-	
+
+	if (Template->IsReadyToRun() == false)
+	{
+		UE_LOG(LogNiagara, Verbose, TEXT("Failed to prime particle pool as the Niagara System is not ready to run."));
+		return;
+	}
+
 	FFXSystemInterface* FXSystemInterface = World->Scene->GetFXSystem();
 	if (FXSystemInterface)
 	{
