- Show instance tick group on HUD in verbose view
- Add options to clip particle attribute display for performance improvements
- Add option to display particle attrributes vertical or horizontal

diff --git a/Engine/Plugins/FX/Niagara/Source/Niagara/Private/NiagaraDebugHud.cpp b/Engine/Plugins/FX/Niagara/Source/Niagara/Private/NiagaraDebugHud.cpp
--- a/Engine/Plugins/FX/Niagara/Source/Niagara/Private/NiagaraDebugHud.cpp
+++ b/Engine/Plugins/FX/Niagara/Source/Niagara/Private/NiagaraDebugHud.cpp
@@ -2232,7 +2232,7 @@ void FNiagaraDebugHud::DrawComponents(FNiagaraWorldManager* WorldManager, UCanva
 		const FLinearColor TextColor = ValidationErrors.Contains(NiagaraComponent) ? Settings.InWorldErrorTextColor : Settings.InWorldTextColor;
 
 		// Show particle data in world
-		if (!Settings.bShowParticlesVariablesWithSystem && bSystemSimulationValid)
+		if (!Settings.bShowParticlesVariablesWithSystem && bSystemSimulationValid && Canvas->SceneView)
 		{
 			const FCachedVariables& CachedVariables = GetCachedVariables(NiagaraSystem);
 			for (int32 iEmitter = 0; iEmitter < CachedVariables.ParticleVariables.Num(); ++iEmitter)
@@ -2262,32 +2262,55 @@ void FNiagaraDebugHud::DrawComponents(FNiagaraWorldManager* WorldManager, UCanva
 					continue;
 				}
 
+				FSceneView* SceneView = Canvas->SceneView;
+
 				const FTransform& SystemTransform = SystemInstance->GetWorldTransform();
 				const bool bParticlesLocalSpace = EmitterInstance->GetCachedEmitter()->bLocalSpace;
-
-				const uint32 NumParticles = Settings.bUseMaxParticlesToDisplay ? FMath::Min((uint32)Settings.MaxParticlesToDisplay, DataBuffer->GetNumInstances()) : DataBuffer->GetNumInstances();
-				for (uint32 iInstance = 0; iInstance < NumParticles; ++iInstance)
+				//const float ClipRadius = Settings.bUseParticleDisplayRadius ? 1.0f : 0.0f;
+				const float ParticleDisplayCenterRadiusSq = Settings.bUseParticleDisplayCenterRadius ? (Settings.ParticleDisplayCenterRadius * Settings.ParticleDisplayCenterRadius) : 0.0f;
+				const float ParticleDisplayClipNearPlane = Settings.bUseParticleDisplayClip ? FMath::Max(Settings.ParticleDisplayClip.X, 0.0f) : 0.0f;
+				const float ParticleDisplayClipFarPlane = Settings.bUseParticleDisplayClip ? FMath::Max(Settings.ParticleDisplayClip.Y, ParticleDisplayClipNearPlane) : FLT_MAX;
+				const uint32 MaxDisplayParticles = Settings.bUseMaxParticlesToDisplay ? Settings.MaxParticlesToDisplay : 0xFFFFFFFF;
+				uint32 NumDisplayedParticles = 0;
+
+				for (uint32 iInstance = 0; iInstance < DataBuffer->GetNumInstances(); ++iInstance)
 				{
 					const FVector ParticalLocalPosition = PositionReader.Get(iInstance);
 					const FVector ParticleWorldPosition = bParticlesLocalSpace ? SystemTransform.TransformPosition(ParticalLocalPosition) : ParticalLocalPosition;
 
+					const FPlane ClipPosition = SceneView->Project(ParticleWorldPosition);
+					if ((ParticleDisplayCenterRadiusSq > 0) && (((ClipPosition.X * ClipPosition.X) + (ClipPosition.Y * ClipPosition.Y)) >= ParticleDisplayCenterRadiusSq))
+					{
+						continue;
+					}
+
+					if ((ClipPosition.W <= ParticleDisplayClipNearPlane) || (ClipPosition.W > ParticleDisplayClipFarPlane))
+					{
+						continue;
+					}
+
 					const FVector ParticleScreenLocation = Canvas->Project(ParticleWorldPosition);
-					if (!FMath::IsNearlyZero(ParticleScreenLocation.Z))
+
+					TStringBuilder<1024> StringBuilder;
+					StringBuilder.Appendf(TEXT("Particle(%u)"), iInstance);
+					for (const auto& ParticleVariable : CachedVariables.ParticleVariables[iEmitter])
 					{
-						TStringBuilder<1024> StringBuilder;
-						StringBuilder.Appendf(TEXT("Particle(%u) "), iInstance);
-						for (const auto& ParticleVariable : CachedVariables.ParticleVariables[iEmitter])
-						{
-							StringBuilder.Append(ParticleVariable.GetName().ToString());
-							StringBuilder.Append(TEXT("("));
-							ParticleVariable.StringAppend(StringBuilder, DataBuffer, iInstance);
-							StringBuilder.Append(TEXT(") "));
-						}
+						StringBuilder.Append(Settings.bShowParticleVariablesVertical ? '\n' : ' ');
+						StringBuilder.Append(ParticleVariable.GetName().ToString());
+						StringBuilder.Append('(');
+						ParticleVariable.StringAppend(StringBuilder, DataBuffer, iInstance);
+						StringBuilder.Append(')');
+					}
+
+					const TCHAR* FinalString = StringBuilder.ToString();
+					const TPair<FVector2D, FVector2D> SizeAndLocation = GetTextLocation(ParticleFont, FinalString, Settings.ParticleTextOptions, FVector2D(ParticleScreenLocation));
+					DrawCanvas->DrawTile(SizeAndLocation.Value.X - 1.0f, SizeAndLocation.Value.Y - 1.0f, SizeAndLocation.Key.X + 2.0f, SizeAndLocation.Key.Y + 2.0f, 0.0f, 0.0f, 0.0f, 0.0f, Settings.DefaultBackgroundColor);
+					DrawCanvas->DrawShadowedString(SizeAndLocation.Value.X, SizeAndLocation.Value.Y, FinalString, ParticleFont, TextColor);
 
-						const TCHAR* FinalString = StringBuilder.ToString();
-						const TPair<FVector2D, FVector2D> SizeAndLocation = GetTextLocation(ParticleFont, FinalString, Settings.ParticleTextOptions, FVector2D(ParticleScreenLocation));
-						DrawCanvas->DrawTile(SizeAndLocation.Value.X - 1.0f, SizeAndLocation.Value.Y - 1.0f, SizeAndLocation.Key.X + 2.0f, SizeAndLocation.Key.Y + 2.0f, 0.0f, 0.0f, 0.0f, 0.0f, Settings.DefaultBackgroundColor);
-						DrawCanvas->DrawShadowedString(SizeAndLocation.Value.X, SizeAndLocation.Value.Y, FinalString, ParticleFont, TextColor);
+					++NumDisplayedParticles;
+					if (++NumDisplayedParticles >= MaxDisplayParticles)
+					{
+						break;
 					}
 				}
 			}
@@ -2325,12 +2348,20 @@ void FNiagaraDebugHud::DrawComponents(FNiagaraWorldManager* WorldManager, UCanva
 					{
 						StringBuilder.Appendf(TEXT("Pooled - %s\n"), *PoolingMethodEnum->GetNameStringByIndex((int32)NiagaraComponent->PoolingMethod));
 					}
-					if (bIsActive && NiagaraComponent->IsRegisteredWithScalabilityManager())
+					if (bIsActive)
 					{
-						StringBuilder.Appendf(TEXT("Scalability - %s\n"), *GetNameSafe(NiagaraSystem->GetEffectType()));
-						if (SystemInstance->SignificanceIndex != INDEX_NONE )
+						if (NiagaraComponent->IsRegisteredWithScalabilityManager())
+						{
+							StringBuilder.Appendf(TEXT("Scalability - %s\n"), *GetNameSafe(NiagaraSystem->GetEffectType()));
+							if (SystemInstance->SignificanceIndex != INDEX_NONE)
+							{
+								StringBuilder.Appendf(TEXT("SignificanceIndex - %d\n"), SystemInstance->SignificanceIndex);
+							}
+						}
+						static UEnum* TickingGroupEnum = StaticEnum<ETickingGroup>();
+						if (TickingGroupEnum)
 						{
-							StringBuilder.Appendf(TEXT("SignificanceIndex - %d\n"), SystemInstance->SignificanceIndex);
+							StringBuilder.Appendf(TEXT("TickGroup - %s\n"), *TickingGroupEnum->GetNameStringByValue(SystemInstance->CalculateTickGroup()));
 						}
 					}
 
diff --git a/Engine/Plugins/FX/Niagara/Source/Niagara/Public/NiagaraDebuggerCommon.h b/Engine/Plugins/FX/Niagara/Source/Niagara/Public/NiagaraDebuggerCommon.h
--- a/Engine/Plugins/FX/Niagara/Source/Niagara/Public/NiagaraDebuggerCommon.h
+++ b/Engine/Plugins/FX/Niagara/Source/Niagara/Public/NiagaraDebuggerCommon.h
@@ -548,9 +548,28 @@ struct NIAGARA_API FNiagaraDebugHUDSettingsData
 	UPROPERTY(Config, EditAnywhere, Category = "Debug Particles", meta = (EditCondition = "bShowParticleVariables", DisplayName="Show Particles Attributes With System"))
 	bool bShowParticlesVariablesWithSystem = false;
 
+	UPROPERTY(Config, EditAnywhere, Category = "Debug Particles", meta = (EditCondition = "bShowParticleVariables", DisplayName = "Show Particle Attributes Vertical"))
+	bool bShowParticleVariablesVertical = false;
+
 	UPROPERTY(Config, EditAnywhere, Category = "Debug Particles", meta = (EditCondition = "bShowParticleVariables"))
 	bool bUseMaxParticlesToDisplay = true;
 
+	/** When enabled we use the clip planes to narrow down which particles to display */
+	UPROPERTY(Config, EditAnywhere, Category = "Debug Particles", meta = (EditCondition = "bShowParticleVariables"))
+	bool bUseParticleDisplayClip = false;
+
+	/** Clipping planes used to display particle attributes. */
+	UPROPERTY(Config, EditAnywhere, Category = "Debug Particles", meta = (EditCondition = "bShowParticleVariables", UIMin = "0", ClampMin = "0"))
+	FVector2D ParticleDisplayClip = FVector2D(0.0f, 10000.0f);
+
+	/** When enabled we use a radius from the display center to avoid showing too many particle attributes. */
+	UPROPERTY(Config, EditAnywhere, Category = "Debug Particles", meta = (EditCondition = "bShowParticleVariables"))
+	bool bUseParticleDisplayCenterRadius = false;
+
+	/** Radius from screen center where 0 is center to 1.0 is edge to avoid display too many particle attributes. */
+	UPROPERTY(Config, EditAnywhere, Category = "Debug Particles", meta = (EditCondition = "bShowParticleVariables", UIMin = "0", ClampMin = "0"))
+	float ParticleDisplayCenterRadius = 1.0f;
+
 	/**
 	When enabled, the maximum number of particles to show information about.
 	When disabled all particles will show attributes, this can result in poor performance & potential OOM on some platforms.
