Fix Niagara Depth Sample for Mobile(Correctly sample SceneDepthAux or SceneDepthTexture)

diff --git a/Engine/Plugins/FX/Niagara/Shaders/Private/NiagaraDataInterfaceCollisionQuery.ush b/Engine/Plugins/FX/Niagara/Shaders/Private/NiagaraDataInterfaceCollisionQuery.ush
--- a/Engine/Plugins/FX/Niagara/Shaders/Private/NiagaraDataInterfaceCollisionQuery.ush
+++ b/Engine/Plugins/FX/Niagara/Shaders/Private/NiagaraDataInterfaceCollisionQuery.ush
@@ -6,19 +6,6 @@
 
 #define NDICOLLISIONQUERY_USE_GBUFFER_NORMAL    (FEATURE_LEVEL >= FEATURE_LEVEL_SM5)
 
-float NDICollisionQuery_GetSceneDepth(float2 ScreenUV)
-{
-#if (FEATURE_LEVEL <= FEATURE_LEVEL_ES3_1)
-	#if MOBILE_DEFERRED_SHADING
-		return ConvertFromDeviceZ(Texture2DSample(MobileSceneTextures.SceneDepthTexture, MobileSceneTextures.SceneDepthTextureSampler, ScreenUV).r);
-	#else
-		return ConvertFromDeviceZ(Texture2DSample(MobileSceneTextures.SceneColorTexture, MobileSceneTextures.SceneColorTextureSampler, ScreenUV).a);
-	#endif
-#else
-		return CalcSceneDepth(ScreenUV);
-#endif
-}
-
 float NDICollisionQuery_GetCustomSceneDepth(float2 ScreenUV)
 {
 #if (FEATURE_LEVEL <= FEATURE_LEVEL_ES3_1)
@@ -45,15 +32,15 @@ void NDICollisionQuery_QuerySceneDepthGPU(in float3 In_SamplePos, out float Out_
 	{
 		// Sample the depth buffer to get a world position near the sample position.
 		float2 ScreenUV = ScreenPosition * View.ScreenPositionScaleBias.xy + View.ScreenPositionScaleBias.wz;
-		float SceneDepth = NDICollisionQuery_GetSceneDepth(ScreenUV);
+		float SceneDepth = CalcSceneDepth(ScreenUV);
 		float3 WorldPosition = WorldPositionFromSceneDepth(ScreenPosition.xy, SceneDepth);
 
 #if NDICOLLISIONQUERY_USE_GBUFFER_NORMAL
 		float3 WorldNormal = Texture2DSampleLevel(SceneTexturesStruct.GBufferATexture, SceneTexturesStruct_GBufferATextureSampler, ScreenUV, 0).xyz * 2.0 - 1.0;
 #else
 		float CollisionDepthBounds = 500.0f;
-		float SceneDepth0 = NDICollisionQuery_GetSceneDepth(ScreenUV + float2(View.BufferSizeAndInvSize.z, 0.0));
-		float SceneDepth1 = NDICollisionQuery_GetSceneDepth(ScreenUV + float2(0.0, View.BufferSizeAndInvSize.w));
+		float SceneDepth0 = CalcSceneDepth(ScreenUV + float2(View.BufferSizeAndInvSize.z, 0.0));
+		float SceneDepth1 = CalcSceneDepth(ScreenUV + float2(0.0, View.BufferSizeAndInvSize.w));
 		// When using the forward shading, the normal of the pixel is approximated by the derivative of the world position
 		// of the pixel. But in on the visible edge this derivative can become very high, making CollisionPlane almost
 		// perpendicular to the view plane. In these case the particle may collide the visible edges of the diferent meshes
diff --git a/Engine/Plugins/FX/Niagara/Shaders/Private/NiagaraDataInterfaceGBuffer.ush b/Engine/Plugins/FX/Niagara/Shaders/Private/NiagaraDataInterfaceGBuffer.ush
--- a/Engine/Plugins/FX/Niagara/Shaders/Private/NiagaraDataInterfaceGBuffer.ush
+++ b/Engine/Plugins/FX/Niagara/Shaders/Private/NiagaraDataInterfaceGBuffer.ush
@@ -174,11 +174,7 @@ void DIGBuffer_DecodeDepth(in float2 ScreenUV, in bool bApplyViewportOffset, out
 	OutValue = GBufferData.Depth;
 #elif NDIGBUFFER_IS_MOBILE
 	IsValid = true;
-	#if MOBILE_DEFERRED_SHADING
-		OutValue = ConvertFromDeviceZ(Texture2DSample(MobileSceneTextures.SceneDepthTexture, MobileSceneTextures.SceneDepthTextureSampler, ScreenUV).r);
-	#else
-		OutValue = ConvertFromDeviceZ(Texture2DSample(MobileSceneTextures.SceneColorTexture, MobileSceneTextures.SceneColorTextureSampler, ScreenUV).a);
-	#endif
+	OutValue = CalcSceneDepth(ScreenUV);
 #else
 	IsValid = false;
 	OutValue = 0;
diff --git a/Engine/Plugins/FX/Niagara/Shaders/Private/NiagaraDataInterfaceOcclusion.ush b/Engine/Plugins/FX/Niagara/Shaders/Private/NiagaraDataInterfaceOcclusion.ush
--- a/Engine/Plugins/FX/Niagara/Shaders/Private/NiagaraDataInterfaceOcclusion.ush
+++ b/Engine/Plugins/FX/Niagara/Shaders/Private/NiagaraDataInterfaceOcclusion.ush
@@ -4,19 +4,6 @@
 
 #include "/Engine/Private/DeferredShadingCommon.ush"
 
-float DIOcclusion_CalcSceneDepth(float2 ScreenUV)
-{
-#if (FEATURE_LEVEL <= FEATURE_LEVEL_ES3_1)
-	#if MOBILE_DEFERRED_SHADING
-		return ConvertFromDeviceZ(Texture2DSample(MobileSceneTextures.SceneDepthTexture, MobileSceneTextures.SceneDepthTextureSampler, ScreenUV).r);
-	#else
-		return ConvertFromDeviceZ(Texture2DSample(MobileSceneTextures.SceneColorTexture, MobileSceneTextures.SceneColorTextureSampler, ScreenUV).a);
-	#endif
-#else
-		return CalcSceneDepth(ScreenUV);
-#endif
-}
-
 void DIOcclusion_Rectangle(in float3 In_SampleCenterWorldPos, in float In_SampleWindowWidthWorld, in float In_SampleWindowHeightWorld, in float In_SampleSteps, out float Out_VisibilityFraction, out float Out_SampleFraction)
 {
 	// we need to project the camera vector to a plane to compare it with scene depth
@@ -56,7 +43,7 @@ void DIOcclusion_Rectangle(in float3 In_SampleCenterWorldPos, in float In_Sample
 					continue;
 				}
 				
-				float Depth = DIOcclusion_CalcSceneDepth(float2(SampleX, SampleY));
+				float Depth = CalcSceneDepth(float2(SampleX, SampleY));
 				if (Depth < CameraDistance) 
 				{
 					OccludedSamples++;
@@ -86,7 +73,7 @@ void DIOcclusion_Circle(in float3 In_SampleCenterWorldPos, in float In_SampleWin
 
 	if (ScreenUV.x <= 1 && ScreenUV.x >= 0 && ScreenUV.y <= 1 && ScreenUV.y >= 0)
 	{
-		float Depth = DIOcclusion_CalcSceneDepth(ScreenUV);
+		float Depth = CalcSceneDepth(ScreenUV);
 		if (Depth < CameraDistance) 
 		{
 			OccludedSamples++;
diff --git a/Engine/Shaders/Private/SceneTexturesCommon.ush b/Engine/Shaders/Private/SceneTexturesCommon.ush
--- a/Engine/Shaders/Private/SceneTexturesCommon.ush
+++ b/Engine/Shaders/Private/SceneTexturesCommon.ush
@@ -142,13 +142,13 @@ float LookupDeviceZ( float2 ScreenUV )
 #elif FORCE_DEPTH_TEXTURE_READS || COMPILER_SWITCH
 	// native Depth buffer lookup
 	return Texture2DSampleLevel(MobileSceneTextures.SceneDepthTexture, MobileSceneTextures.SceneDepthTextureSampler, ScreenUV, 0).r;
-#elif (POST_PROCESS_MATERIAL || POST_PROCESS_MATERIAL_MOBILE) && !POST_PROCESS_AR_PASSTHROUGH
+#elif (POST_PROCESS_MATERIAL || POST_PROCESS_MATERIAL_MOBILE || COMPUTESHADER) && !POST_PROCESS_AR_PASSTHROUGH
 	#if MOBILE_DEFERRED_SHADING
-		return Texture2DSample(MobileSceneTextures.SceneDepthTexture, MobileSceneTextures.SceneDepthTextureSampler, ScreenUV).r;
+		return Texture2DSampleLevel(MobileSceneTextures.SceneDepthTexture, MobileSceneTextures.SceneDepthTextureSampler, ScreenUV, 0).r;
 	#else
 		// SceneDepth texture is not accessible during post-processing as we discard it at the end of mobile BasePass
 		// instead fetch DeviceZ from SceneDepthAuxTexture
-		return Texture2DSample(MobileSceneTextures.SceneDepthAuxTexture, MobileSceneTextures.SceneDepthAuxTextureSampler, ScreenUV);
+		return Texture2DSampleLevel(MobileSceneTextures.SceneDepthAuxTexture, MobileSceneTextures.SceneDepthAuxTextureSampler, ScreenUV, 0);
 	#endif
 #elif COMPILER_GLSL_ES3_1 && PIXELSHADER
 	return DepthbufferFetchES2();
