- Fix for data being set on user arrays before the system is created

#okforigithub public
#rnx

diff --git a/Engine/Plugins/FX/Niagara/Source/Niagara/Classes/NiagaraDataInterfaceArrayImpl.h b/Engine/Plugins/FX/Niagara/Source/Niagara/Classes/NiagaraDataInterfaceArrayImpl.h
--- a/Engine/Plugins/FX/Niagara/Source/Niagara/Classes/NiagaraDataInterfaceArrayImpl.h
+++ b/Engine/Plugins/FX/Niagara/Source/Niagara/Classes/NiagaraDataInterfaceArrayImpl.h
@@ -242,36 +242,70 @@ struct FNDIArrayProxyImpl : public INDIArrayProxyBase
 	struct FReadArrayRef
 	{
 		FReadArrayRef(TOwnerType* Owner, FNDIArrayInstanceData_GameThread<TArrayType>* InstanceData)
-			: ScopeLock(InstanceData->ArrayRWGuard)
 		{
-			ArrayData = InstanceData->bIsModified  ? &InstanceData->ArrayData : &Owner->GetArrayReference();
+			if ( InstanceData )
+			{
+				LockObject = &InstanceData->ArrayRWGuard;
+				LockObject->ReadLock();
+				ArrayData = InstanceData->bIsModified ? &InstanceData->ArrayData : &Owner->GetArrayReference();
+			}
+			else
+			{
+				ArrayData = &Owner->GetArrayReference();
+			}
+		}
+		~FReadArrayRef()
+		{
+			if ( LockObject )
+			{
+				LockObject->ReadUnlock();
+			}
 		}
 
 		const TArray<TArrayType>& GetArray() { return *ArrayData; }
 
 	private:
-		FReadScopeLock				ScopeLock;
+		FRWLock*					LockObject = nullptr;
 		const TArray<TArrayType>*	ArrayData = nullptr;
+		UE_NONCOPYABLE(FReadArrayRef);
 	};
 
 	struct FWriteArrayRef
 	{
 		FWriteArrayRef(TOwnerType* Owner, FNDIArrayInstanceData_GameThread<TArrayType>* InstanceData)
-			: ScopeLock(InstanceData->ArrayRWGuard)
 		{
-			if ( InstanceData->bIsModified == false )
+			if (InstanceData)
+			{
+				LockObject = &InstanceData->ArrayRWGuard;
+				LockObject->WriteLock();
+
+				if (InstanceData->bIsModified == false)
+				{
+					InstanceData->bIsModified = true;
+					InstanceData->ArrayData = Owner->GetArrayReference();
+				}
+				ArrayData = &InstanceData->ArrayData;
+			}
+			else
+			{
+				ArrayData = &Owner->GetArrayReference();
+			}
+		}
+
+		~FWriteArrayRef()
+		{
+			if (LockObject)
 			{
-				InstanceData->bIsModified = true;
-				InstanceData->ArrayData = Owner->GetArrayReference();
+				LockObject->WriteUnlock();
 			}
-			ArrayData = &InstanceData->ArrayData;
 		}
 
 		TArray<TArrayType>& GetArray() { return *ArrayData; }
 
 	private:
-		FWriteScopeLock		ScopeLock;
-		TArray<TArrayType>* ArrayData = nullptr;
+		FRWLock*			LockObject = nullptr;
+		TArray<TArrayType>*	ArrayData = nullptr;
+		UE_NONCOPYABLE(FWriteArrayRef);
 	};
 
 	struct FGameToRenderInstanceData
@@ -412,6 +446,69 @@ struct FNDIArrayProxyImpl : public INDIArrayProxyBase
 	}
 
 	//////////////////////////////////////////////////////////////////////////
+	// BP user parameter accessors, should remove if we every start to share the object between instances
+	void SetArrayData(const TArray<TArrayType>& InArrayData)
+	{
+		if (PerInstanceData_GameThread.IsEmpty())
+		{
+			Owner->GetArrayReference() = InArrayData;
+		}
+		else
+		{
+			ensure(PerInstanceData_GameThread.Num() == 1);
+			FNDIArrayInstanceData_GameThread<TArrayType>* InstanceData = PerInstanceData_GameThread.CreateConstIterator().Value();
+			FWriteScopeLock	ScopeLock(InstanceData->ArrayRWGuard);
+			InstanceData->bIsModified = false;
+			InstanceData->bIsRenderDirty = bShouldSyncToGpu;
+			InstanceData->ArrayData.Empty();
+			Owner->GetArrayReference() = InArrayData;
+		}
+	}
+
+	TArray<TArrayType> GetArrayData()
+	{
+		ensure(PerInstanceData_GameThread.Num() <= 1);
+		FNDIArrayInstanceData_GameThread<TArrayType>* InstanceData = PerInstanceData_GameThread.IsEmpty() ? nullptr : PerInstanceData_GameThread.CreateConstIterator().Value();
+		FReadArrayRef ArrayRef(Owner, InstanceData);
+		return ArrayRef.GetArray();
+	}
+
+	void SetArrayValue(int Index, const TArrayType& Value, bool bSizeToFit)
+	{
+		ensure(PerInstanceData_GameThread.Num() <= 1);
+		FNDIArrayInstanceData_GameThread<TArrayType>* InstanceData = PerInstanceData_GameThread.IsEmpty() ? nullptr : PerInstanceData_GameThread.CreateConstIterator().Value();
+		FWriteArrayRef ArrayRef(Owner, InstanceData);
+
+		if (!ArrayRef.GetArray().IsValidIndex(Index))
+		{
+			if (!bSizeToFit)
+			{
+				return;
+			}
+			ArrayRef.GetArray().AddDefaulted(Index + 1 - ArrayRef.GetArray().Num());
+		}
+
+		ArrayRef.GetArray()[Index] = Value;
+	}
+
+	TArrayType GetArrayValue(int Index)
+	{
+		TArrayType ValueOut = TArrayType(FNDIArrayImplHelper<TArrayType>::GetDefaultValue());
+
+		ensure(PerInstanceData_GameThread.Num() <= 1);
+		FNDIArrayInstanceData_GameThread<TArrayType>* InstanceData = PerInstanceData_GameThread.IsEmpty() ? nullptr : PerInstanceData_GameThread.CreateConstIterator().Value();
+		FReadArrayRef ArrayRef(Owner, InstanceData);
+
+		if (!ArrayRef.GetArray().IsValidIndex(Index))
+		{
+			ValueOut = ArrayRef.GetArray()[Index];
+		}
+
+		return ValueOut;
+	}
+
+	//////////////////////////////////////////////////////////////////////////
+	// VM accessors to ensure we maintain per correctness for shared data interfaces
 	void SetArrayData(FNiagaraSystemInstanceID InstanceID, const TArray<TArrayType>& InArrayData)
 	{
 		if ( FNDIArrayInstanceData_GameThread<TArrayType>* InstanceData = PerInstanceData_GameThread.FindRef(InstanceID) )
diff --git a/Engine/Plugins/FX/Niagara/Source/Niagara/Private/NiagaraDataInterfaceArrayFunctionLibrary.cpp b/Engine/Plugins/FX/Niagara/Source/Niagara/Private/NiagaraDataInterfaceArrayFunctionLibrary.cpp
--- a/Engine/Plugins/FX/Niagara/Source/Niagara/Private/NiagaraDataInterfaceArrayFunctionLibrary.cpp
+++ b/Engine/Plugins/FX/Niagara/Source/Niagara/Private/NiagaraDataInterfaceArrayFunctionLibrary.cpp
@@ -8,29 +8,24 @@
 #include "NiagaraSystemInstance.h"
 #include "NiagaraFunctionLibrary.h"
 
+// If / when we share user parameter UObjects we will need to make this per instance which introduces some tricky things about allocating before the instance is active
 template<typename TArrayType, typename TDataInterace>
 void SetNiagaraArray(UNiagaraComponent* NiagaraSystem, FName OverrideName, const TArray<TArrayType>& InArray)
 {
-	if (auto SystemInstanceController = NiagaraSystem->GetSystemInstanceController())
+	if (TDataInterace* ArrayDI = UNiagaraFunctionLibrary::GetDataInterface<TDataInterace>(NiagaraSystem, OverrideName))
 	{
-		if (TDataInterace* ArrayDI = UNiagaraFunctionLibrary::GetDataInterface<TDataInterace>(NiagaraSystem, OverrideName))
-		{
-			auto* ArrayProxy = static_cast<FNDIArrayProxyImpl<TArrayType, TDataInterace>*>(ArrayDI->GetProxy());
-			ArrayProxy->SetArrayData(SystemInstanceController->GetSystemInstanceID(), InArray);
-		}
+		auto* ArrayProxy = static_cast<FNDIArrayProxyImpl<TArrayType, TDataInterace>*>(ArrayDI->GetProxy());
+		ArrayProxy->SetArrayData(InArray);
 	}
 }
 
 template<typename TArrayType, typename TDataInterace>
 TArray<TArrayType> GetNiagaraArray(UNiagaraComponent* NiagaraSystem, FName OverrideName)
 {
-	if (auto SystemInstanceController = NiagaraSystem->GetSystemInstanceController())
+	if (TDataInterace* ArrayDI = UNiagaraFunctionLibrary::GetDataInterface<TDataInterace>(NiagaraSystem, OverrideName))
 	{
-		if (TDataInterace* ArrayDI = UNiagaraFunctionLibrary::GetDataInterface<TDataInterace>(NiagaraSystem, OverrideName))
-		{
-			auto* ArrayProxy = static_cast<FNDIArrayProxyImpl<TArrayType, TDataInterace>*>(ArrayDI->GetProxy());
-			return ArrayProxy->GetArrayData(SystemInstanceController->GetSystemInstanceID());
-		}
+		auto* ArrayProxy = static_cast<FNDIArrayProxyImpl<TArrayType, TDataInterace>*>(ArrayDI->GetProxy());
+		return ArrayProxy->GetArrayData();
 	}
 	return TArray<TArrayType>();
 }
@@ -38,26 +33,20 @@ TArray<TArrayType> GetNiagaraArray(UNiagaraComponent* NiagaraSystem, FName Overr
 template<typename TArrayType, typename TDataInterace>
 void SetNiagaraArrayValue(UNiagaraComponent* NiagaraSystem, FName OverrideName, int Index, const TArrayType& Value, bool bSizeToFit)
 {
-	if (auto SystemInstanceController = NiagaraSystem->GetSystemInstanceController())
+	if (TDataInterace* ArrayDI = UNiagaraFunctionLibrary::GetDataInterface<TDataInterace>(NiagaraSystem, OverrideName))
 	{
-		if (TDataInterace* ArrayDI = UNiagaraFunctionLibrary::GetDataInterface<TDataInterace>(NiagaraSystem, OverrideName))
-		{
-			auto* ArrayProxy = static_cast<FNDIArrayProxyImpl<TArrayType, TDataInterace>*>(ArrayDI->GetProxy());
-			ArrayProxy->SetArrayValue(SystemInstanceController->GetSystemInstanceID(), Index, Value, bSizeToFit);
-		}
+		auto* ArrayProxy = static_cast<FNDIArrayProxyImpl<TArrayType, TDataInterace>*>(ArrayDI->GetProxy());
+		ArrayProxy->SetArrayValue(Index, Value, bSizeToFit);
 	}
 }
 
 template<typename TArrayType, typename TDataInterace>
 TArrayType GetNiagaraArrayValue(UNiagaraComponent* NiagaraSystem, FName OverrideName, int Index)
 {
-	if (auto SystemInstanceController = NiagaraSystem->GetSystemInstanceController())
+	if (TDataInterace* ArrayDI = UNiagaraFunctionLibrary::GetDataInterface<TDataInterace>(NiagaraSystem, OverrideName))
 	{
-		if (TDataInterace* ArrayDI = UNiagaraFunctionLibrary::GetDataInterface<TDataInterace>(NiagaraSystem, OverrideName))
-		{
-			auto* ArrayProxy = static_cast<FNDIArrayProxyImpl<TArrayType, TDataInterace>*>(ArrayDI->GetProxy());
-			return ArrayProxy->GetArrayValue(SystemInstanceController->GetSystemInstanceID(), Index);
-		}
+		auto* ArrayProxy = static_cast<FNDIArrayProxyImpl<TArrayType, TDataInterace>*>(ArrayDI->GetProxy());
+		return ArrayProxy->GetArrayValue(Index);
 	}
 	//-TODO: Should be DefaultValue
 	return TArrayType();
