- Ensure cached max instance count is up to date when scalability changes
- Move emitter debug name enable into NiagaraCommon.h so we can easily enable / disable

#rnx

diff --git a/Engine/Plugins/FX/Niagara/Source/Niagara/Classes/NiagaraComputeExecutionContext.h b/Engine/Plugins/FX/Niagara/Source/Niagara/Classes/NiagaraComputeExecutionContext.h
--- a/Engine/Plugins/FX/Niagara/Source/Niagara/Classes/NiagaraComputeExecutionContext.h
+++ b/Engine/Plugins/FX/Niagara/Source/Niagara/Classes/NiagaraComputeExecutionContext.h
@@ -2,6 +2,7 @@
 
 #pragma once
 
+#include "NiagaraCommon.h"
 #include "NiagaraScriptExecutionContext.h"
 #include "NiagaraSimStageData.h"
 #include "RHIGPUReadback.h"
@@ -118,7 +119,7 @@ struct FNiagaraComputeExecutionContext
 		uint32 CPUCount = 0;
 	}  EmitterInstanceReadback;
 	
-#if 1//!UE_BUILD_SHIPPING
+#if WITH_NIAGARA_DEBUG_EMITTER_NAME
 	FName GetDebugSimFName() const { return DebugSimFName; }
 	const TCHAR* GetDebugSimName() const { return *DebugSimName; }
 	void SetDebugSimName(const TCHAR* InDebugSimName)
@@ -138,7 +139,7 @@ struct FNiagaraComputeExecutionContext
 public:
 	static uint32 TickCounter;
 
-#if 1//!UE_BUILD_SHIPPING
+#if WITH_NIAGARA_DEBUG_EMITTER_NAME
 	FName DebugSimFName;
 	FString DebugSimName;
 #endif
diff --git a/Engine/Plugins/FX/Niagara/Source/Niagara/Classes/NiagaraEmitter.h b/Engine/Plugins/FX/Niagara/Source/Niagara/Classes/NiagaraEmitter.h
--- a/Engine/Plugins/FX/Niagara/Source/Niagara/Classes/NiagaraEmitter.h
+++ b/Engine/Plugins/FX/Niagara/Source/Niagara/Classes/NiagaraEmitter.h
@@ -503,10 +503,10 @@ public:
 	bool CanObtainSystemAttribute(const FNiagaraVariableBase& InVar) const;
 	bool CanObtainUserVariable(const FNiagaraVariableBase& InVar) const;
 
-#if UE_BUILD_SHIPPING
-	const TCHAR* GetDebugSimName() const { return TEXT(""); }
-#else
+#if WITH_NIAGARA_DEBUG_EMITTER_NAME
 	const TCHAR* GetDebugSimName() const { return *DebugSimName; }
+#else
+	const TCHAR* GetDebugSimName() const { return TEXT(""); }
 #endif
 
 	NIAGARA_API const FString& GetUniqueEmitterName() const;
@@ -637,7 +637,7 @@ private:
 
 	bool bFullyLoaded = false;
 
-#if !UE_BUILD_SHIPPING
+#if WITH_NIAGARA_DEBUG_EMITTER_NAME
 	FString DebugSimName;
 #endif
 
diff --git a/Engine/Plugins/FX/Niagara/Source/Niagara/Private/NiagaraDataSet.cpp b/Engine/Plugins/FX/Niagara/Source/Niagara/Private/NiagaraDataSet.cpp
--- a/Engine/Plugins/FX/Niagara/Source/Niagara/Private/NiagaraDataSet.cpp
+++ b/Engine/Plugins/FX/Niagara/Source/Niagara/Private/NiagaraDataSet.cpp
@@ -90,7 +90,7 @@ static FAutoConsoleVariableRef CVarNiagaraGPUDataWarningSize(
 	TEXT("fx.NiagaraGPUDataWarningSize"),
 	GNiagaraGPUDataWarningSize,
 	TEXT("Allocation size where we should log a warning."),
-	ECVF_Default
+	ECVF_Scalability
 );
 
 void NiagaraWarnGpuBufferSize(const int32 ByteSize, const TCHAR* DebugSimName)
diff --git a/Engine/Plugins/FX/Niagara/Source/Niagara/Private/NiagaraEmitter.cpp b/Engine/Plugins/FX/Niagara/Source/Niagara/Private/NiagaraEmitter.cpp
--- a/Engine/Plugins/FX/Niagara/Source/Niagara/Private/NiagaraEmitter.cpp
+++ b/Engine/Plugins/FX/Niagara/Source/Niagara/Private/NiagaraEmitter.cpp
@@ -66,7 +66,7 @@ static FAutoConsoleVariableRef CVarNiagaraEmitterMaxGPUBufferElements(
 	TEXT("Maximum elements per GPU buffer, for example 4k elements would restrict a float buffer to be 16k maximum per buffer.\n")
 	TEXT("Note: If you request something smaller than what will satisfy a single unit of work it will be increased to that size.\n")
 	TEXT("Default 0 which will allow the buffer to be the maximum allowed by the RHI.\n"),
-	ECVF_Default
+	ECVF_Scalability
 );
 
 FNiagaraDetailsLevelScaleOverrides::FNiagaraDetailsLevelScaleOverrides()
@@ -1259,15 +1259,19 @@ void UNiagaraEmitter::CacheFromCompiledData(const FNiagaraDataSetCompiledData* C
 	}
 #endif
 
-#if !UE_BUILD_SHIPPING
+#if WITH_NIAGARA_DEBUG_EMITTER_NAME
 	// Ensure our debug simulation name is up to date
-	DebugSimName.Empty();
-	if (const UNiagaraSystem* SystemOwner = Cast<const UNiagaraSystem>(GetOuter()))
+	// Only required on uncooked as it can change due to compilation
+	if ( !FPlatformProperties::RequiresCookedData() )
 	{
-		DebugSimName = SystemOwner->GetName();
-		DebugSimName.AppendChar(':');
+		DebugSimName.Empty();
+		if (const UNiagaraSystem* SystemOwner = Cast<const UNiagaraSystem>(GetOuter()))
+		{
+			DebugSimName = SystemOwner->GetName();
+			DebugSimName.AppendChar(':');
+		}
+		DebugSimName.Append(GetName());
 	}
-	DebugSimName.Append(GetName());
 #endif
 }
 
@@ -1376,7 +1380,7 @@ void UNiagaraEmitter::UpdateEmitterAfterLoad()
 
 	ResolveScalabilitySettings();
 
-#if !UE_BUILD_SHIPPING
+#if WITH_NIAGARA_DEBUG_EMITTER_NAME
 	DebugSimName.Empty();
 	if (const UNiagaraSystem* SystemOwner = Cast<const UNiagaraSystem>(GetOuter()))
 	{
diff --git a/Engine/Plugins/FX/Niagara/Source/Niagara/Private/NiagaraSystem.cpp b/Engine/Plugins/FX/Niagara/Source/Niagara/Private/NiagaraSystem.cpp
--- a/Engine/Plugins/FX/Niagara/Source/Niagara/Private/NiagaraSystem.cpp
+++ b/Engine/Plugins/FX/Niagara/Source/Niagara/Private/NiagaraSystem.cpp
@@ -3514,6 +3514,8 @@ void UNiagaraSystem::ResolveScalabilitySettings()
 
 void UNiagaraSystem::OnScalabilityCVarChanged()
 {
+	CacheFromCompiledData();
+
 	ResolveScalabilitySettings();
 
 	for (FNiagaraEmitterHandle& Handle : EmitterHandles)
diff --git a/Engine/Plugins/FX/Niagara/Source/Niagara/Public/NiagaraCommon.h b/Engine/Plugins/FX/Niagara/Source/Niagara/Public/NiagaraCommon.h
--- a/Engine/Plugins/FX/Niagara/Source/Niagara/Public/NiagaraCommon.h
+++ b/Engine/Plugins/FX/Niagara/Source/Niagara/Public/NiagaraCommon.h
@@ -37,6 +37,9 @@ struct FNiagaraParameterStore;
 #define WITH_NIAGARA_GPU_PROFILER_EDITOR (WITH_EDITOR && STATS)
 #define WITH_NIAGARA_GPU_PROFILER ((WITH_PARTICLE_PERF_STATS && !UE_BUILD_SHIPPING) || WITH_NIAGARA_GPU_PROFILER_EDITOR)
 
+/** When enabled allows builds to warn about system / emitter information such as hitting buffer caps.  Small performance overhead for copying the name around so RT has access. */
+#define WITH_NIAGARA_DEBUG_EMITTER_NAME	1//!UE_BUILD_SHIPPING
+
 #define INTERPOLATED_PARAMETER_PREFIX TEXT("PREV_")
 
 /** Defines The maximum ThreadGroup size we allow in Niagara.  This is important for how memory is allocated as we always need to round this and the final instance is used to avoid overflowing the buffer. */
