- Fix for grid not showing in Niagara editor

diff --git a/Engine/Plugins/FX/Niagara/Source/NiagaraEditor/Private/Widgets/SNiagaraSystemViewport.cpp b/Engine/Plugins/FX/Niagara/Source/NiagaraEditor/Private/Widgets/SNiagaraSystemViewport.cpp
--- a/Engine/Plugins/FX/Niagara/Source/NiagaraEditor/Private/Widgets/SNiagaraSystemViewport.cpp
+++ b/Engine/Plugins/FX/Niagara/Source/NiagaraEditor/Private/Widgets/SNiagaraSystemViewport.cpp
@@ -77,7 +77,7 @@ FNiagaraSystemViewportClient::FNiagaraSystemViewportClient(FAdvancedPreviewScene
 	DrawHelper.bDrawPivot = false;
 	DrawHelper.bDrawWorldBox = false;
 	DrawHelper.bDrawKillZ = false;
-	DrawHelper.bDrawGrid = false;
+	DrawHelper.bDrawGrid = true;
 	DrawHelper.GridColorAxis = FColor(80,80,80);
 	DrawHelper.GridColorMajor = FColor(72,72,72);
 	DrawHelper.GridColorMinor = FColor(64,64,64);
