Made sure to generate a new variable guid when duplicating a variable in niagara scripts

diff --git a/Engine/Plugins/FX/Niagara/Source/NiagaraEditor/Private/ViewModels/NiagaraParameterPanelViewModel.cpp b/Engine/Plugins/FX/Niagara/Source/NiagaraEditor/Private/ViewModels/NiagaraParameterPanelViewModel.cpp
--- a/Engine/Plugins/FX/Niagara/Source/NiagaraEditor/Private/ViewModels/NiagaraParameterPanelViewModel.cpp
+++ b/Engine/Plugins/FX/Niagara/Source/NiagaraEditor/Private/ViewModels/NiagaraParameterPanelViewModel.cpp
@@ -2229,7 +2229,8 @@ void FNiagaraScriptToolkitParameterPanelViewModel::DuplicateParameter(const FNia
 	const FName NewUniqueName = FNiagaraUtilities::GetUniqueName(ItemToDuplicate.GetVariable().GetName(), Names);
 	FNiagaraVariable NewVariable(ItemToDuplicate.GetVariable().GetType(), NewUniqueName);
 	FNiagaraVariableMetaData ParameterMetadata = ItemToDuplicate.ScriptVariable ? ItemToDuplicate.ScriptVariable->Metadata : FNiagaraVariableMetaData();
-
+	ParameterMetadata.CreateNewGuid();
+	
 	for (UNiagaraGraph* Graph : GetEditableGraphs())
 	{
 		Graph->Modify();
