- Remove TSharedFromThis as DataInteraceProxy's are Unique and the UNiagaraDataInterface owns the lifetime of these

#okforgithuib public
#rnx

diff --git a/Engine/Plugins/FX/Niagara/Source/Niagara/Classes/NiagaraDataInterface.h b/Engine/Plugins/FX/Niagara/Source/Niagara/Classes/NiagaraDataInterface.h
--- a/Engine/Plugins/FX/Niagara/Source/Niagara/Classes/NiagaraDataInterface.h
+++ b/Engine/Plugins/FX/Niagara/Source/Niagara/Classes/NiagaraDataInterface.h
@@ -234,7 +234,7 @@ private:
 //////////////////////////////////////////////////////////////////////////
 
 struct FNiagaraDataInterfaceProxyRW;
-struct FNiagaraDataInterfaceProxy : TSharedFromThis<FNiagaraDataInterfaceProxy, ESPMode::ThreadSafe>
+struct FNiagaraDataInterfaceProxy
 {
 	FNiagaraDataInterfaceProxy() {}
 	virtual ~FNiagaraDataInterfaceProxy() {/*check(IsInRenderingThread());*/}
diff --git a/Engine/Plugins/FX/Niagara/Source/Niagara/Private/NiagaraDataInterface.cpp b/Engine/Plugins/FX/Niagara/Source/Niagara/Private/NiagaraDataInterface.cpp
--- a/Engine/Plugins/FX/Niagara/Source/Niagara/Private/NiagaraDataInterface.cpp
+++ b/Engine/Plugins/FX/Niagara/Source/Niagara/Private/NiagaraDataInterface.cpp
@@ -20,14 +20,12 @@ UNiagaraDataInterface::UNiagaraDataInterface(FObjectInitializer const& ObjectIni
 
 UNiagaraDataInterface::~UNiagaraDataInterface()
 {
-	// @todo-threadsafety Can there be a UNiagaraDataInterface class itself created? Perhaps by the system?
-	if ( Proxy.IsValid() )
+	if ( FNiagaraDataInterfaceProxy* ReleasedProxy = Proxy.Release() )
 	{
 		ENQUEUE_RENDER_COMMAND(FDeleteProxyRT) (
-			[RT_Proxy=MoveTemp(Proxy)](FRHICommandListImmediate& CmdList) mutable
+			[RT_Proxy= ReleasedProxy](FRHICommandListImmediate& CmdList) mutable
 			{
-				// This will release RT_Proxy on the RT
-				RT_Proxy = nullptr;
+				delete RT_Proxy;
 			}
 		);
 		check(Proxy.IsValid() == false);
