- Disable GPU sort / cull via CPU path if we have already acquired the cull buffer.  This will occur when using RDG immediate mode.

#rnx

diff --git a/Engine/Plugins/FX/Niagara/Source/Niagara/Private/NiagaraRendererMeshes.cpp b/Engine/Plugins/FX/Niagara/Source/Niagara/Private/NiagaraRendererMeshes.cpp
--- a/Engine/Plugins/FX/Niagara/Source/Niagara/Private/NiagaraRendererMeshes.cpp
+++ b/Engine/Plugins/FX/Niagara/Source/Niagara/Private/NiagaraRendererMeshes.cpp
@@ -331,6 +331,13 @@ void FNiagaraRendererMeshes::PrepareParticleMeshRenderData(FParticleMeshRenderDa
 		}
 		else
 		{
+			//-TODO: Culling and sorting from InitViewsAfterPrePass can not be respected if the culled entries have already been acquired
+			if (ParticleMeshRenderData.bSortCullOnGpu)
+			{
+				//ensureMsgf(false, TEXT("Culling & sorting is not supported once the culled counts have been acquired, sorting & culling will be disabled for these draws."));
+				ParticleMeshRenderData.bSortCullOnGpu &= SceneProxy->GetComputeDispatchInterface()->GetGPUInstanceCounterManager().CanAcquireCulledEntry();
+			}
+
 			// For CPU sims, decide if we should sort / cull on the GPU or not
 			if ( ParticleMeshRenderData.bSortCullOnGpu )
 			{
diff --git a/Engine/Plugins/FX/Niagara/Source/Niagara/Private/NiagaraRendererSprites.cpp b/Engine/Plugins/FX/Niagara/Source/Niagara/Private/NiagaraRendererSprites.cpp
--- a/Engine/Plugins/FX/Niagara/Source/Niagara/Private/NiagaraRendererSprites.cpp
+++ b/Engine/Plugins/FX/Niagara/Source/Niagara/Private/NiagaraRendererSprites.cpp
@@ -283,6 +283,13 @@ void FNiagaraRendererSprites::PrepareParticleSpriteRenderData(FParticleSpriteRen
 		}
 		else
 		{
+			//-TODO: Culling and sorting from InitViewsAfterPrePass can not be respected if the culled entries have already been acquired
+			if (ParticleSpriteRenderData.bSortCullOnGpu)
+			{
+				//ensureMsgf(false, TEXT("Culling & sorting is not supported once the culled counts have been acquired, sorting & culling will be disabled for these draws."));
+				ParticleSpriteRenderData.bSortCullOnGpu &= SceneProxy->GetComputeDispatchInterface()->GetGPUInstanceCounterManager().CanAcquireCulledEntry();
+			}
+
 			// Should we GPU sort for CPU systems?
 			if ( ParticleSpriteRenderData.bSortCullOnGpu )
 			{
