- Add optional task timeout detect to attempt to track down a rare task lost issue

#rnx

diff --git a/Engine/Plugins/FX/Niagara/Source/Niagara/Private/NiagaraSystemInstance.cpp b/Engine/Plugins/FX/Niagara/Source/Niagara/Private/NiagaraSystemInstance.cpp
--- a/Engine/Plugins/FX/Niagara/Source/Niagara/Private/NiagaraSystemInstance.cpp
+++ b/Engine/Plugins/FX/Niagara/Source/Niagara/Private/NiagaraSystemInstance.cpp
@@ -2133,6 +2133,18 @@ void FNiagaraSystemInstance::ManualTick(float DeltaSeconds, const FGraphEventRef
 	SystemSim->Tick_GameThread(DeltaSeconds, MyCompletionGraphEvent);
 }
 
+void FNiagaraSystemInstance::DumpStalledInfo()
+{
+	TStringBuilder<128> Builder;
+	Builder.Appendf(TEXT("System (%s)\n"), *GetNameSafe(GetSystem()));
+	Builder.Appendf(TEXT("ConcurrentTickGraphEvent Complete (%d)\n"), ConcurrentTickGraphEvent ? ConcurrentTickGraphEvent->IsComplete() : true);
+	Builder.Appendf(TEXT("FinalizePending (%d)\n"), FinalizeRef.IsPending());
+	Builder.Appendf(TEXT("SystemInstanceIndex (%d)\n"), SystemInstanceIndex);
+	Builder.Appendf(TEXT("SystemInstanceState (%d)\n"), SystemInstanceState);
+
+	UE_LOG(LogNiagara, Fatal, TEXT("FNiagaraSystemInstance is stalled.\n%s"), Builder.ToString());
+}
+
 void FNiagaraSystemInstance::WaitForConcurrentTickDoNotFinalize(bool bEnsureComplete)
 {
 	check(IsInGameThread());
@@ -2148,10 +2160,32 @@ void FNiagaraSystemInstance::WaitForConcurrentTickDoNotFinalize(bool bEnsureComp
 		const double WarnSeconds = 5.0;
 		const uint64 WarnCycles = StartCycles + uint64(WarnSeconds / FPlatformTime::GetSecondsPerCycle64());
 
-		do
+		extern int32 GNiagaraSystemSimulationTaskStallTimeout;
+		if (GNiagaraSystemSimulationTaskStallTimeout > 0)
+		{
+			do
+			{
+				const double EndTimeoutSeconds = FPlatformTime::Seconds() + (double(GNiagaraSystemSimulationTaskStallTimeout) / 1000.0);
+				LowLevelTasks::BusyWaitUntil(
+					[this, EndTimeoutSeconds]()
+					{
+						if (FPlatformTime::Seconds() > EndTimeoutSeconds)
+						{
+							DumpStalledInfo();
+							return true;
+						}
+						return ConcurrentTickGraphEvent->IsComplete();
+					}
+				);
+			} while (ConcurrentTickGraphEvent && !ConcurrentTickGraphEvent->IsComplete());
+		}
+		else
 		{
-			FTaskGraphInterface::Get().WaitUntilTaskCompletes(ConcurrentTickGraphEvent, ENamedThreads::GameThread);
-		} while (ConcurrentTickGraphEvent && !ConcurrentTickGraphEvent->IsComplete());
+			do
+			{
+				FTaskGraphInterface::Get().WaitUntilTaskCompletes(ConcurrentTickGraphEvent, ENamedThreads::GameThread);
+			} while (ConcurrentTickGraphEvent && !ConcurrentTickGraphEvent->IsComplete());
+		}
 
 		const double StallTimeMS = FPlatformTime::ToMilliseconds64(FPlatformTime::Cycles64() - StartCycles);
 		if ((GWaitForAsyncStallWarnThresholdMS > 0.0f) && (StallTimeMS > GWaitForAsyncStallWarnThresholdMS))
diff --git a/Engine/Plugins/FX/Niagara/Source/Niagara/Private/NiagaraSystemSimulation.cpp b/Engine/Plugins/FX/Niagara/Source/Niagara/Private/NiagaraSystemSimulation.cpp
--- a/Engine/Plugins/FX/Niagara/Source/Niagara/Private/NiagaraSystemSimulation.cpp
+++ b/Engine/Plugins/FX/Niagara/Source/Niagara/Private/NiagaraSystemSimulation.cpp
@@ -56,7 +56,6 @@ DECLARE_CYCLE_STAT(TEXT("System Sim Init (DirectBindings) [GT]"), STAT_NiagaraSy
 DECLARE_CYCLE_STAT(TEXT("ForcedWaitForAsync"), STAT_NiagaraSystemSim_ForceWaitForAsync, STATGROUP_Niagara);
 DECLARE_CYCLE_STAT(TEXT("ForcedWait Fake Stall"), STAT_NiagaraSystemSim_ForceWaitFakeStall, STATGROUP_Niagara);
 
-
 static int32 GbDumpSystemData = 0;
 static FAutoConsoleVariableRef CVarNiagaraDumpSystemData(
 	TEXT("fx.DumpSystemData"),
@@ -81,6 +80,15 @@ static FAutoConsoleVariableRef CVarNiagaraSystemSimulationAllowASync(
 	ECVF_Default
 );
 
+int32 GNiagaraSystemSimulationTaskStallTimeout = 0;
+static FAutoConsoleVariableRef CVarNiagaraSystemSimulationTaskStallTimeout(
+	TEXT("fx.Niagara.SystemSimulation.TaskStallTimeout"),
+	GNiagaraSystemSimulationTaskStallTimeout,
+	TEXT("Timeout in microseconds for Niagara simulation tasks to be considered stalled.\n")
+	TEXT("When this is > 0 we busy wait as opposed to joining the TG so avoid using execpt for debugging."),
+	ECVF_Default
+);
+
 static int32 GNiagaraSystemSimulationTickBatchSize = NiagaraSystemTickBatchSize;
 static FAutoConsoleVariableRef CVarParallelSystemInstanceTickBatchSize(
 	TEXT("fx.Niagara.SystemSimulation.TickBatchSize"),
@@ -1554,6 +1562,23 @@ void FNiagaraSystemSimulation::Spawn_Concurrent(FNiagaraSystemSimulationTickCont
 	bInSpawnPhase = false;
 }
 
+void FNiagaraSystemSimulation::DumpStalledInfo()
+{
+	TStringBuilder<128> Builder;
+	Builder.Appendf(TEXT("ConcurrentTickGraphEvent Complete (%d)\n"), ConcurrentTickGraphEvent ? ConcurrentTickGraphEvent->IsComplete() : true);
+	Builder.Appendf(TEXT("AllWorkCompleteGraphEvent Complete (%d)\n"), AllWorkCompleteGraphEvent ? ConcurrentTickGraphEvent->IsComplete() : true);
+	Builder.Appendf(TEXT("SystemTickGroup (%d)\n"), SystemTickGroup);
+	Builder.Appendf(TEXT("bInSpawnPhase (%d)\n"), bInSpawnPhase);
+	Builder.Appendf(TEXT("bIsSolo (%d)\n"), bIsSolo);
+	Builder.Appendf(TEXT("InstancesPendingSpawn (%d)\n"), GetSystemInstances(ENiagaraSystemInstanceState::PendingSpawn).Num());
+	Builder.Appendf(TEXT("InstancesPendingSpawnPaused (%d)\n"), GetSystemInstances(ENiagaraSystemInstanceState::PendingSpawnPaused).Num());
+	Builder.Appendf(TEXT("InstancesSpawning (%d)\n"), GetSystemInstances(ENiagaraSystemInstanceState::Spawning).Num());
+	Builder.Appendf(TEXT("InstancesRunning (%d)\n"), GetSystemInstances(ENiagaraSystemInstanceState::Running).Num());
+	Builder.Appendf(TEXT("InstancesPaused (%d)\n"), GetSystemInstances(ENiagaraSystemInstanceState::Paused).Num());
+
+	UE_LOG(LogNiagara, Fatal, TEXT("NiagaraSystemSimulation(%s) is stalled.\n%s"), *GetNameSafe(GetSystem()), Builder.ToString());
+}
+
 void FNiagaraSystemSimulation::WaitForConcurrentTickComplete(bool bEnsureComplete)
 {
 	check(IsInGameThread());
@@ -1562,7 +1587,26 @@ void FNiagaraSystemSimulation::WaitForConcurrentTickComplete(bool bEnsureComplet
 	{
 		SCOPE_CYCLE_COUNTER(STAT_NiagaraSystemSim_ForceWaitForAsync);
 		ensureAlwaysMsgf(!bEnsureComplete, TEXT("NiagaraSystemSimulation(%s) ConcurrentTickGraphEvent is not completed."), *GetSystem()->GetPathName());
-		FTaskGraphInterface::Get().WaitUntilTaskCompletes(ConcurrentTickGraphEvent, ENamedThreads::GameThread);
+
+		if (GNiagaraSystemSimulationTaskStallTimeout > 0)
+		{
+			const double EndTimeoutSeconds = FPlatformTime::Seconds() + (double(GNiagaraSystemSimulationTaskStallTimeout) / 1000.0);
+			LowLevelTasks::BusyWaitUntil(
+				[this, EndTimeoutSeconds]()
+				{
+					if (FPlatformTime::Seconds() > EndTimeoutSeconds)
+					{
+						DumpStalledInfo();
+						return true;
+					}
+					return ConcurrentTickGraphEvent->IsComplete();
+				}
+			);
+		}
+		else
+		{
+			FTaskGraphInterface::Get().WaitUntilTaskCompletes(ConcurrentTickGraphEvent, ENamedThreads::GameThread);
+		}
 	}
 	ConcurrentTickGraphEvent = nullptr;
 }
diff --git a/Engine/Plugins/FX/Niagara/Source/Niagara/Public/NiagaraSystemInstance.h b/Engine/Plugins/FX/Niagara/Source/Niagara/Public/NiagaraSystemInstance.h
--- a/Engine/Plugins/FX/Niagara/Source/Niagara/Public/NiagaraSystemInstance.h
+++ b/Engine/Plugins/FX/Niagara/Source/Niagara/Public/NiagaraSystemInstance.h
@@ -399,6 +399,8 @@ public:
 	int32 GetRandomSeedOffset() const { return RandomSeedOffset; }
 
 private:
+	void DumpStalledInfo();
+
 	void DestroyDataInterfaceInstanceData();
 
 	/** Builds the emitter simulations. */
diff --git a/Engine/Plugins/FX/Niagara/Source/Niagara/Public/NiagaraSystemSimulation.h b/Engine/Plugins/FX/Niagara/Source/Niagara/Public/NiagaraSystemSimulation.h
--- a/Engine/Plugins/FX/Niagara/Source/Niagara/Public/NiagaraSystemSimulation.h
+++ b/Engine/Plugins/FX/Niagara/Source/Niagara/Public/NiagaraSystemSimulation.h
@@ -320,6 +320,8 @@ public:
 	FORCEINLINE UWorld* GetWorld()const{return World;}
 
 protected:
+	void DumpStalledInfo();
+
 	/** Sets constant parameter values */
 	void SetupParameters_GameThread(float DeltaSeconds);
 
