- Fix for missing icons

#rnx

diff --git a/Engine/Plugins/FX/Niagara/Source/NiagaraEditor/Private/NiagaraEditorStyle.cpp b/Engine/Plugins/FX/Niagara/Source/NiagaraEditor/Private/NiagaraEditorStyle.cpp
--- a/Engine/Plugins/FX/Niagara/Source/NiagaraEditor/Private/NiagaraEditorStyle.cpp
+++ b/Engine/Plugins/FX/Niagara/Source/NiagaraEditor/Private/NiagaraEditorStyle.cpp
@@ -718,6 +718,7 @@ TSharedRef< FSlateStyleSet > FNiagaraEditorStyle::Create()
 	InitToolbar(Style);
 	InitTabIcons(Style);
 	InitIcons(Style);
+	InitOverview(Style);
 	InitEmitterDetails(Style);
 	InitAssetColors(Style);
 	InitThumbnails(Style);
