- Fix potential memory leak / uninitialized struct

#rnx

diff --git a/Engine/Plugins/FX/Niagara/Source/Niagara/Private/NiagaraDataInterfaceRigidMeshCollisionQuery.cpp b/Engine/Plugins/FX/Niagara/Source/Niagara/Private/NiagaraDataInterfaceRigidMeshCollisionQuery.cpp
--- a/Engine/Plugins/FX/Niagara/Source/Niagara/Private/NiagaraDataInterfaceRigidMeshCollisionQuery.cpp
+++ b/Engine/Plugins/FX/Niagara/Source/Niagara/Private/NiagaraDataInterfaceRigidMeshCollisionQuery.cpp
@@ -805,6 +805,7 @@ void FNDIRigidMeshCollisionProxy::ConsumePerInstanceDataFromGameThread(void* Per
 	{
 		UE_LOG(LogRigidMeshCollision, Log, TEXT("ConsumePerInstanceDataFromGameThread() ... could not find %d"), Instance);
 	}
+	SourceData->~FNDIRigidMeshCollisionData();
 }
 
 void FNDIRigidMeshCollisionProxy::InitializePerInstanceData(const FNiagaraSystemInstanceID& SystemInstance)
@@ -1286,7 +1287,7 @@ void UNiagaraDataInterfaceRigidMeshCollisionQuery::ValidateFunction(const FNiaga
 void UNiagaraDataInterfaceRigidMeshCollisionQuery::ProvidePerInstanceDataForRenderThread(void* DataForRenderThread, void* PerInstanceData, const FNiagaraSystemInstanceID& SystemInstance)
 {
 	FNDIRigidMeshCollisionData* GameThreadData = static_cast<FNDIRigidMeshCollisionData*>(PerInstanceData);
-	FNDIRigidMeshCollisionData* RenderThreadData = static_cast<FNDIRigidMeshCollisionData*>(DataForRenderThread);
+	FNDIRigidMeshCollisionData* RenderThreadData = new(DataForRenderThread) FNDIRigidMeshCollisionData();
 
 	if (GameThreadData != nullptr && RenderThreadData != nullptr)
 	{		
