Build time optimization

diff --git a/Engine/Plugins/FX/Niagara/Source/Niagara/Private/DataInterface/NiagaraDataInterfaceDebugDraw.cpp b/Engine/Plugins/FX/Niagara/Source/Niagara/Private/DataInterface/NiagaraDataInterfaceDebugDraw.cpp
--- a/Engine/Plugins/FX/Niagara/Source/Niagara/Private/DataInterface/NiagaraDataInterfaceDebugDraw.cpp
+++ b/Engine/Plugins/FX/Niagara/Source/Niagara/Private/DataInterface/NiagaraDataInterfaceDebugDraw.cpp
@@ -2385,6 +2385,9 @@ void UNiagaraDataInterfaceDebugDraw::PostInitProperties()
 	}
 }
 
+// Codegen optimization degenerates for very long functions like GetFunctions when combined with the invocation of lots of FORCEINLINE methods.
+BEGIN_FUNCTION_BUILD_OPTIMIZATION
+
 void UNiagaraDataInterfaceDebugDraw::GetFunctions(TArray<FNiagaraFunctionSignature>& OutFunctions)
 {
 	OutFunctions.Reserve(OutFunctions.Num() + 22);
@@ -2780,6 +2783,8 @@ void UNiagaraDataInterfaceDebugDraw::GetFunctions(TArray<FNiagaraFunctionSignatu
 #endif
 }
 
+END_FUNCTION_BUILD_OPTIMIZATION
+
 #if WITH_EDITORONLY_DATA
 bool UNiagaraDataInterfaceDebugDraw::UpgradeFunctionCall(FNiagaraFunctionSignature& FunctionSignature)
 {
