Fix for precise motion blur on first frame when spawned from events
-ensures that Particles.Previous values are properly set at the end of the event spawn stage

diff --git a/Engine/Plugins/FX/Niagara/Source/Niagara/Classes/NiagaraScriptSourceBase.h b/Engine/Plugins/FX/Niagara/Source/Niagara/Classes/NiagaraScriptSourceBase.h
--- a/Engine/Plugins/FX/Niagara/Source/Niagara/Classes/NiagaraScriptSourceBase.h
+++ b/Engine/Plugins/FX/Niagara/Source/Niagara/Classes/NiagaraScriptSourceBase.h
@@ -56,6 +56,7 @@ class FNiagaraCompileOptions
 public:
 	NIAGARA_API static const FString CpuScriptDefine;
 	NIAGARA_API static const FString GpuScriptDefine;
+	NIAGARA_API static const FString EventSpawnDefine;
 
 	FNiagaraCompileOptions() : TargetUsage(ENiagaraScriptUsage::Function), TargetUsageBitmask(0)
 	{
diff --git a/Engine/Plugins/FX/Niagara/Source/Niagara/Private/NiagaraCommon.cpp b/Engine/Plugins/FX/Niagara/Source/Niagara/Private/NiagaraCommon.cpp
--- a/Engine/Plugins/FX/Niagara/Source/Niagara/Private/NiagaraCommon.cpp
+++ b/Engine/Plugins/FX/Niagara/Source/Niagara/Private/NiagaraCommon.cpp
@@ -1544,6 +1544,7 @@ bool FVMExternalFunctionBindingInfo::Serialize(FArchive& Ar)
 
 const FString FNiagaraCompileOptions::CpuScriptDefine = TEXT("CPUSim");
 const FString FNiagaraCompileOptions::GpuScriptDefine = TEXT("GPUComputeSim");
+const FString FNiagaraCompileOptions::EventSpawnDefine = TEXT("EventSpawn");
 
 FSynchronizeWithParameterDefinitionsArgs::FSynchronizeWithParameterDefinitionsArgs()
 	: SpecificDefinitionsUniqueIds(TArray<FGuid>())
diff --git a/Engine/Plugins/FX/Niagara/Source/Niagara/Private/NiagaraScript.cpp b/Engine/Plugins/FX/Niagara/Source/Niagara/Private/NiagaraScript.cpp
--- a/Engine/Plugins/FX/Niagara/Source/Niagara/Private/NiagaraScript.cpp
+++ b/Engine/Plugins/FX/Niagara/Source/Niagara/Private/NiagaraScript.cpp
@@ -945,6 +945,17 @@ void UNiagaraScript::ComputeVMCompilationId(FNiagaraVMExecutableDataId& Id, FGui
 			Id.ReferencedCompileHashes.AddUnique(Hash);
 			Id.DebugReferencedObjects.Add(TEXT("SimulationStageHeaders"));
 		}
+
+		if (IsParticleEventScript(Usage))
+		{
+			if (const FNiagaraEventScriptProperties* EventScriptProps = Emitter->GetEventHandlerByIdUnsafe(UsageId))
+			{
+				if (EventScriptProps->ExecutionMode == EScriptExecutionMode::SpawnedParticles)
+				{
+					Id.AdditionalDefines.AddUnique(FNiagaraCompileOptions::EventSpawnDefine);
+				}
+			}
+		}
 	}
 
 	UObject* Obj = GetOuter();
diff --git a/Engine/Plugins/FX/Niagara/Source/NiagaraEditor/Private/NiagaraHlslTranslator.cpp b/Engine/Plugins/FX/Niagara/Source/NiagaraEditor/Private/NiagaraHlslTranslator.cpp
--- a/Engine/Plugins/FX/Niagara/Source/NiagaraEditor/Private/NiagaraHlslTranslator.cpp
+++ b/Engine/Plugins/FX/Niagara/Source/NiagaraEditor/Private/NiagaraHlslTranslator.cpp
@@ -1509,6 +1509,15 @@ const FNiagaraTranslateResults &FHlslNiagaraTranslator::Translate(const FNiagara
 		{
 			BuildMissingDefaults();
 		}
+
+		if (UNiagaraScript::IsParticleEventScript(TranslationStages[0].ScriptUsage))
+		{
+			if (CompileOptions.AdditionalDefines.Contains(FNiagaraCompileOptions::EventSpawnDefine))
+			{
+				AddBodyComment(TEXT("//Handle resetting previous values at the end of spawn so that they match outputs! (Needed for motion blur/etc)"));
+				AddBodyChunk(TEXT("HandlePreviousValuesForSpawn(Context);"));
+			}
+		}
 	}
 	else
 	{
@@ -2284,7 +2293,13 @@ void FHlslNiagaraTranslator::DefinePreviousParametersFunction(FString& HlslOutpu
 	HlslOutputString +=
 		TEXT("#if (SimulationStageIndex == 0) // MapSpawn\n")
 		TEXT("void HandlePreviousValuesForSpawn(inout FSimulationContext Context)\n{\n");
-	if (UNiagaraScript::IsParticleSpawnScript(CompileOptions.TargetUsage) || UNiagaraScript::IsGPUScript(CompileOptions.TargetUsage))
+
+	const bool WriteFunctionInternals = UNiagaraScript::IsParticleSpawnScript(CompileOptions.TargetUsage)
+		|| UNiagaraScript::IsGPUScript(CompileOptions.TargetUsage)
+		|| (UNiagaraScript::IsParticleEventScript(CompileOptions.TargetUsage)
+			&& CompileOptions.AdditionalDefines.Contains(FNiagaraCompileOptions::EventSpawnDefine));
+
+	if (WriteFunctionInternals)
 	{
 		TArray<FNiagaraDataSetID> ReadDataSetIDs;
 		TArray<FNiagaraDataSetID> WriteDataSetIDs;
