Record whether parameter defaults are implicit or explicit when building ParameterMapHistory to resolve fail if not set errors.
Do not consider default failure for certain scenarios when setting a variable in the stack.

diff --git a/Engine/Plugins/FX/Niagara/Source/Niagara/Private/NiagaraSystem.cpp b/Engine/Plugins/FX/Niagara/Source/Niagara/Private/NiagaraSystem.cpp
--- a/Engine/Plugins/FX/Niagara/Source/Niagara/Private/NiagaraSystem.cpp
+++ b/Engine/Plugins/FX/Niagara/Source/Niagara/Private/NiagaraSystem.cpp
@@ -14,6 +14,7 @@
 #include "NiagaraRendererProperties.h"
 #include "NiagaraScriptSourceBase.h"
 #include "NiagaraSettings.h"
+#include "NiagaraSimulationStageBase.h"
 #include "NiagaraStats.h"
 #include "NiagaraTrace.h"
 #include "NiagaraTypes.h"
@@ -2433,14 +2434,11 @@ bool UNiagaraSystem::CompilationResultsValid(FNiagaraSystemCompileRequest& Compi
 
 void UNiagaraSystem::EvaluateCompileResultDependencies() const
 {
-	//@todo(ng) disable dependency checking pending fixes for false positive errors
-	return;
-
 	struct FScriptCompileResultValidationInfo
 	{
 		UNiagaraEmitter* Emitter = nullptr;
 		FNiagaraVMExecutableData* CompileResults = nullptr;
-		int32 ParentIndex = INDEX_NONE;
+		TArray<int32> ParentIndices;
 		bool bCompileResultStatusDirty = false;
 
 		void ReEvaluateCompileStatus()
@@ -2492,7 +2490,7 @@ void UNiagaraSystem::EvaluateCompileResultDependencies() const
 	TArray<FScriptCompileResultValidationInfo> ScriptCompilesToValidate;
 
 	// Add a compiled script pair to be evaluated for external dependencies to ScriptCompilesToValidate and clear all of its compile events that are from script dependencies.
-	auto AddCompiledScriptForValidation = [&ScriptCompilesToValidate](UNiagaraScript* InScript, int32 ParentIdx, UNiagaraEmitter* InEmitter = nullptr)->int32
+	auto AddCompiledScriptForValidation = [&ScriptCompilesToValidate](UNiagaraScript* InScript, const TArray<int32>& InParentIndices, UNiagaraEmitter* InEmitter = nullptr)->int32
 	{
 		FScriptCompileResultValidationInfo ValidationInfo = FScriptCompileResultValidationInfo();
 		ValidationInfo.CompileResults = &InScript->GetVMExecutableData();
@@ -2508,38 +2506,80 @@ void UNiagaraSystem::EvaluateCompileResultDependencies() const
 		}
 
 		ValidationInfo.Emitter = InEmitter;
-		ValidationInfo.ParentIndex = ParentIdx;
+		ValidationInfo.ParentIndices = InParentIndices;
 		return ScriptCompilesToValidate.Add(ValidationInfo);
 	};
 
-	// Gather system scripts to evaluate dependencies.
-	int32 SystemSpawnScriptIdx = AddCompiledScriptForValidation(SystemSpawnScript, INDEX_NONE);
-	int32 SystemUpdateScriptIdx = AddCompiledScriptForValidation(SystemUpdateScript, SystemSpawnScriptIdx);
-
-	// Gather per-emitter scripts to evaluate dependencies.
+	int32 SystemSpawnScriptIdx = AddCompiledScriptForValidation(SystemSpawnScript, TArray<int32>());
+	TArray<int32> SystemSpawnVisitOrder;
+	SystemSpawnVisitOrder.Add(SystemSpawnScriptIdx);
+
+	int32 SystemUpdateScriptIdx = AddCompiledScriptForValidation(SystemUpdateScript, SystemSpawnVisitOrder);
+	TArray<int32> SystemUpdateVisitOrder;
+	SystemUpdateVisitOrder.Add(SystemSpawnScriptIdx);
+	SystemUpdateVisitOrder.Add(SystemUpdateScriptIdx);
+
+	// Gather per-emitter scripts to evaluate dependencies. The dependencies need to form a linear chain walking up from the current execution path all the way to SystemSpawn, with 
+	// all scripts that execute in between being visited. If there are holes, then something could be marked as not being written even though it was written. The flow is thus:
+	// System Spawn
+	// Emitter Spawn
+	// System Update
+	// Emitter Update
+	// Particle Spawn/Interpolated
+	// Particle Spawn Event
+	// Particle Update
+	// Particle Non-Spawn Event
+	// Particle Sim Stages
 	for (const FNiagaraEmitterHandle& Handle : EmitterHandles)
 	{
 		if (Handle.GetInstance() && Handle.GetIsEnabled())
 		{
 			UNiagaraEmitter* Emitter = Handle.GetInstance();
 
-			const int32 EmitterSpawnScriptIdx = AddCompiledScriptForValidation(Emitter->EmitterSpawnScriptProps.Script, SystemSpawnScriptIdx, Emitter);
-			const int32 EmitterUpdateScriptIdx = AddCompiledScriptForValidation(Emitter->EmitterUpdateScriptProps.Script, SystemUpdateScriptIdx, Emitter);
-			const int32 ParentIdxForParticleSpawnScript = Emitter->bInterpolatedSpawning ? EmitterUpdateScriptIdx : EmitterSpawnScriptIdx;
-			const int32 ParticleSpawnScriptIdx = AddCompiledScriptForValidation(Emitter->SpawnScriptProps.Script, ParentIdxForParticleSpawnScript, Emitter);
-			const int32 ParticleUpdateScriptIdx = AddCompiledScriptForValidation(Emitter->UpdateScriptProps.Script, EmitterUpdateScriptIdx, Emitter);
+			const int32 EmitterSpawnScriptIdx = AddCompiledScriptForValidation(Emitter->EmitterSpawnScriptProps.Script, SystemSpawnVisitOrder, Emitter);
+			TArray<int32> EmitterSpawnVisitOrder = SystemSpawnVisitOrder;
+			TArray<int32> EmitterUpdateVisitOrder = SystemUpdateVisitOrder;
+			EmitterSpawnVisitOrder.Add(EmitterSpawnScriptIdx);
+			EmitterUpdateVisitOrder.Add(EmitterSpawnScriptIdx);
+
+			const int32 EmitterUpdateScriptIdx = AddCompiledScriptForValidation(Emitter->EmitterUpdateScriptProps.Script, EmitterUpdateVisitOrder, Emitter);
+			EmitterUpdateVisitOrder.Add(EmitterUpdateScriptIdx);
 
-			int32 TempIdx = ParticleUpdateScriptIdx;
+			TArray<int32> ParticleSpawnVisitOrder = EmitterUpdateVisitOrder;
+			TArray<int32> ParticleUpdateVisitOrder = ParticleSpawnVisitOrder;
+			const int32 ParticleSpawnScriptIdx = AddCompiledScriptForValidation(Emitter->SpawnScriptProps.Script, ParticleSpawnVisitOrder, Emitter);
+			ParticleSpawnVisitOrder.Add(ParticleSpawnScriptIdx);
+			ParticleUpdateVisitOrder.Add(ParticleSpawnScriptIdx);
+
+			const int32 ParticleUpdateScriptIdx = AddCompiledScriptForValidation(Emitter->UpdateScriptProps.Script, ParticleUpdateVisitOrder, Emitter);
+			ParticleUpdateVisitOrder.Add(ParticleUpdateScriptIdx);
+
+			// If the spawn script is interpolated, allow the particle update script to resolve dependencies for spawn events and spawn sim stages.
+			if (Emitter->bInterpolatedSpawning)
+			{
+				ParticleSpawnVisitOrder.Add(ParticleUpdateScriptIdx);
+			}
+			
 			const TArray<FNiagaraEventScriptProperties>& EventHandlerScriptProps = Emitter->GetEventHandlers();
 			for (int32 i = 0; i < EventHandlerScriptProps.Num(); i++)
 			{
 				if (EventHandlerScriptProps[i].Script)
 				{
-					TempIdx = AddCompiledScriptForValidation(EventHandlerScriptProps[i].Script, TempIdx, Emitter);
+					if (EventHandlerScriptProps[i].ExecutionMode == EScriptExecutionMode::SpawnedParticles)
+					{
+						AddCompiledScriptForValidation(EventHandlerScriptProps[i].Script, ParticleSpawnVisitOrder, Emitter);
+					}
+						
+					else
+					{
+						AddCompiledScriptForValidation(EventHandlerScriptProps[i].Script, ParticleUpdateVisitOrder, Emitter);
+					}	
 				}
 			}
-
+			
+			// Gather up all the spawn stages first, since they could write to a value before anything else.
 			const TArray<UNiagaraSimulationStageBase*>& SimulationStages = Emitter->GetSimulationStages();
+			TArray<UNiagaraSimulationStageBase*> NonSpawnStages;
 			for (int32 i = 0; i < SimulationStages.Num(); i++)
 			{
 				if (SimulationStages[i] && SimulationStages[i]->Script)
@@ -2549,13 +2589,30 @@ void UNiagaraSystem::EvaluateCompileResultDependencies() const
 						continue;
 					}
 
-					TempIdx = AddCompiledScriptForValidation(SimulationStages[i]->Script, TempIdx, Emitter);
+					UNiagaraSimulationStageGeneric* GenericSimStage = Cast<UNiagaraSimulationStageGeneric>(SimulationStages[i]);
+					if (GenericSimStage && GenericSimStage->ExecuteBehavior == ENiagaraSimStageExecuteBehavior::OnSimulationReset)
+					{
+						int32 LastSpawnStageIdx = AddCompiledScriptForValidation(SimulationStages[i]->Script, ParticleSpawnVisitOrder, Emitter);
+						ParticleSpawnVisitOrder.Add(LastSpawnStageIdx);
+						ParticleUpdateVisitOrder.Add(LastSpawnStageIdx);
+					}
+					else
+					{
+						NonSpawnStages.Add(SimulationStages[i]);
+					}
 				}
 			}
 
+			// Now handle all the other stages
+			for (int32 i = 0; i < NonSpawnStages.Num(); i++)
+			{
+				int32 LastStageIdx = AddCompiledScriptForValidation(NonSpawnStages[i]->Script, ParticleUpdateVisitOrder, Emitter);
+				ParticleUpdateVisitOrder.Add(LastStageIdx);
+			}
+
 			if (Emitter->SimTarget == ENiagaraSimTarget::GPUComputeSim)
 			{
-				AddCompiledScriptForValidation(Emitter->GetGPUComputeScript(), TempIdx, Emitter);
+				AddCompiledScriptForValidation(Emitter->GetGPUComputeScript(), ParticleUpdateVisitOrder, Emitter);
 			}
 		}
 	}
@@ -2574,21 +2631,18 @@ void UNiagaraSystem::EvaluateCompileResultDependencies() const
 			}
 
 			bool bDependencyMet = false;
-			int32 TestIdx = ValidationInfo.ParentIndex;
-			while (TestIdx != INDEX_NONE && bDependencyMet == false)
+			for (const int32 TestIndex : ValidationInfo.ParentIndices)
 			{
-				if (ScriptCompilesToValidate.IsValidIndex(TestIdx))
+				if (ScriptCompilesToValidate.IsValidIndex(TestIndex) == false)
 				{
-					const FScriptCompileResultValidationInfo& TestInfo = ScriptCompilesToValidate[TestIdx];
-					if (TestInfo.CompileResults->AttributesWritten.Num() > 0)
-					{
-						if (TestInfo.CompileResults->AttributesWritten.Contains(TestVar))
-						{
-							bDependencyMet = true;
-							break;
-						}
-					}
-					TestIdx = TestInfo.ParentIndex;
+					ensureMsgf(false, TEXT("Encountered invalid index into script compiles to validate when evaluating fail if not set dependencies!"));
+				}
+
+				const FScriptCompileResultValidationInfo& TestInfo = ScriptCompilesToValidate[TestIndex];
+				if (TestInfo.CompileResults->AttributesWritten.Contains(TestVar))
+				{
+					bDependencyMet = true;
+					break;
 				}
 			}
 			if (!bDependencyMet)
diff --git a/Engine/Plugins/FX/Niagara/Source/NiagaraCore/Private/NiagaraCustomVersion.cpp b/Engine/Plugins/FX/Niagara/Source/NiagaraCore/Private/NiagaraCustomVersion.cpp
--- a/Engine/Plugins/FX/Niagara/Source/NiagaraCore/Private/NiagaraCustomVersion.cpp
+++ b/Engine/Plugins/FX/Niagara/Source/NiagaraCore/Private/NiagaraCustomVersion.cpp
@@ -8,4 +8,5 @@ const FGuid FNiagaraCustomVersion::GUID(0xFCF57AFA, 0x50764283, 0xB9A9E658, 0xFF
 // Register the custom version with core
 FCustomVersionRegistration GRegisterNiagaraCustomVersion(FNiagaraCustomVersion::GUID, FNiagaraCustomVersion::LatestVersion, TEXT("NiagaraVer"));
 
-const FGuid FNiagaraCustomVersion::LatestScriptCompileVersion(0xFC1F4DE2, 0x165D4E64, 0x9A023D80, 0x9D333FDB);
+const FGuid FNiagaraCustomVersion::LatestScriptCompileVersion(0xB2501D5E, 0x0FC54B9F, 0xB16B6767, 0xBF50D399);
+
diff --git a/Engine/Plugins/FX/Niagara/Source/NiagaraEditor/Private/NiagaraHlslTranslator.cpp b/Engine/Plugins/FX/Niagara/Source/NiagaraEditor/Private/NiagaraHlslTranslator.cpp
--- a/Engine/Plugins/FX/Niagara/Source/NiagaraEditor/Private/NiagaraHlslTranslator.cpp
+++ b/Engine/Plugins/FX/Niagara/Source/NiagaraEditor/Private/NiagaraHlslTranslator.cpp
@@ -681,15 +681,16 @@ FString FHlslNiagaraTranslator::BuildParameterMapHlslDefinitions(TArray<FNiagara
 	}
 
 	// Add in currently defined attribute vars.
-	ParamMapDefinedAttributesToNamespaceVars.GenerateValueArray(ValueArray);
-	for (FNiagaraVariable& Var : ValueArray)
+	TArray<FVarAndDefaultSource> VarAndDefaultSourceArray;
+	ParamMapDefinedAttributesToNamespaceVars.GenerateValueArray(VarAndDefaultSourceArray);
+	for (FVarAndDefaultSource& VarAndDefaultSource : VarAndDefaultSourceArray)
 	{
-		if (Var.GetType().GetClass() != nullptr)
+		if (VarAndDefaultSource.Variable.GetType().GetClass() != nullptr)
 		{
 			continue;
 		}
 
-		UniqueVariables.AddUnique(Var);
+		UniqueVariables.AddUnique(VarAndDefaultSource.Variable);
 	}
 
 	// Add in any bulk usage vars.
@@ -944,6 +945,27 @@ void FHlslNiagaraTranslator::BuildConstantBuffer(ENiagaraCodeChunkMode ChunkMode
 	}
 }
 
+void FHlslNiagaraTranslator::RecordParamMapDefinedAttributeToNamespaceVar(const FNiagaraVariable& VarToRecord, const UEdGraphPin* VarAssociatedDefaultPin)
+{
+	bool bDefaultPinExplicit = true;
+	if (VarAssociatedDefaultPin == nullptr || VarAssociatedDefaultPin->bHidden)
+	{
+		bDefaultPinExplicit = false;
+	}
+
+	if (FVarAndDefaultSource* VarAndDefaultSourcePtr = ParamMapDefinedAttributesToNamespaceVars.Find(VarToRecord.GetName()))
+	{
+		VarAndDefaultSourcePtr->bDefaultExplicit |= bDefaultPinExplicit;
+	}
+	else
+	{
+		FVarAndDefaultSource VarAndDefaultSource;
+		VarAndDefaultSource.Variable = VarToRecord;
+		VarAndDefaultSource.bDefaultExplicit = bDefaultPinExplicit;
+		ParamMapDefinedAttributesToNamespaceVars.Add(VarToRecord.GetName(), VarAndDefaultSource);
+	}
+}
+
 static void ConvertFloatToHalf(const FNiagaraCompileOptions& InCompileOptions, TArray<FNiagaraVariable>& Attributes)
 {
 	if (InCompileOptions.AdditionalDefines.Contains(TEXT("CompressAttributes")))// && UNiagaraScript::IsParticleScript(InCompileOptions.TargetUsage))
@@ -2907,12 +2929,13 @@ void FHlslNiagaraTranslator::DefineMainGPUFunctions(
 				if (i == 1 && TranslationStages[i].ScriptUsage == ENiagaraScriptUsage::ParticleUpdateScript) // The Update Phase might need previous parameters set.
 				{
 					// Put any gathered previous variables into the list here so that we can use them by recording the last value from the parent variable on transfer from previous stage if interpolated spawning.
-					TArray<FNiagaraVariable> Vars;
+					TArray<FVarAndDefaultSource> VarAndDefaultSourceArray;
 					TArray<FNiagaraVariable> GatheredPreviousVariables;
-					ParamMapDefinedAttributesToNamespaceVars.GenerateValueArray(Vars);
+					ParamMapDefinedAttributesToNamespaceVars.GenerateValueArray(VarAndDefaultSourceArray);
 
-					for (const FNiagaraVariable& Var : Vars)
+					for (const FVarAndDefaultSource& VarAndDefaultSource : VarAndDefaultSourceArray)
 					{
+						const FNiagaraVariable& Var = VarAndDefaultSource.Variable;
 						if (FNiagaraParameterMapHistory::IsPreviousValue(Var))
 						{
 							FNiagaraVariable SrcVar = FNiagaraParameterMapHistory::GetSourceForPreviousValue(Var);
@@ -3306,12 +3329,13 @@ void FHlslNiagaraTranslator::DefineMain(FString &OutHlslOutput,
 				if (StageIdx == 0 && UNiagaraScript::IsInterpolatedParticleSpawnScript(CompileOptions.TargetUsage)) // The Update Phase might need previous parameters set.
 				{
 					// Put any gathered previous variables into the list here so that we can use them by recording the last value from the parent variable on transfer from previous stage if interpolated spawning.
-					TArray<FNiagaraVariable> Vars;
+					TArray<FVarAndDefaultSource> VarAndDefaultSourceArray;
 					TArray<FNiagaraVariable> GatheredPreviousVariables;
-					ParamMapDefinedAttributesToNamespaceVars.GenerateValueArray(Vars);
+					ParamMapDefinedAttributesToNamespaceVars.GenerateValueArray(VarAndDefaultSourceArray);
 
-					for (const FNiagaraVariable& Var : Vars)
+					for (const FVarAndDefaultSource& VarAndDefaultSource : VarAndDefaultSourceArray)
 					{
+						const FNiagaraVariable& Var = VarAndDefaultSource.Variable;
 						if (FNiagaraParameterMapHistory::IsPreviousValue(Var))
 						{
 							FNiagaraVariable SrcVar = FNiagaraParameterMapHistory::GetSourceForPreviousValue(Var);
@@ -4555,7 +4579,6 @@ bool FHlslNiagaraTranslationStage::IsRelevantToSpawnForStage(const FNiagaraParam
 	return false;
 }
 
-
 void FHlslNiagaraTranslator::InitializeParameterMapDefaults(int32 ParamMapHistoryIdx)
 {
 	bInitializedDefaults = true;
@@ -4730,7 +4753,10 @@ void FHlslNiagaraTranslator::Output(UNiagaraNodeOutput* OutputNode, const TArray
 					if (ParamMapDefinedAttributesToUniformChunks.Find(Var.GetName()) == nullptr)
 					{
 						ParamMapDefinedAttributesToUniformChunks.Add(Var.GetName(), Input);
-						ParamMapDefinedAttributesToNamespaceVars.Add(Var.GetName(), VarNamespaced);
+						FVarAndDefaultSource VarAndDefaultSource;
+						VarAndDefaultSource.Variable = VarNamespaced;
+						VarAndDefaultSource.bDefaultExplicit = false;
+						ParamMapDefinedAttributesToNamespaceVars.Add(Var.GetName(), VarAndDefaultSource);
 					}
 
 					InstanceWrite.Variables.AddUnique(VarNamespaced);
@@ -4857,7 +4883,6 @@ void FHlslNiagaraTranslator::ParameterMapSet(UNiagaraNodeParameterMapSet* SetNod
 				Error(FText::Format(LOCTEXT("ParameterMapSetTypeError", "Cannot handle type {0}! Variable: {1}"), Var.GetType().GetNameText(), FText::FromName(Var.GetName())), nullptr, nullptr);
 			}
 
-			FString VarName = Var.GetName().ToString();
 			if (TranslationStages[ActiveStageIdx].IsExternalConstantNamespace(Var, CompileOptions.TargetUsage, CompileOptions.GetTargetUsageBitmask()))
 			{
 				Error(FText::Format(LOCTEXT("SetSystemConstantFail", "Cannot Set external constant, Type: {0} Variable: {1}"), Var.GetType().GetNameText(), FText::FromName(Var.GetName())), SetNode, nullptr);
@@ -4885,7 +4910,7 @@ void FHlslNiagaraTranslator::ParameterMapSet(UNiagaraNodeParameterMapSet* SetNod
 				if (VarIdx != INDEX_NONE && VarIdx < ParamMapSetVariablesToChunks[ParamMapHistoryIdx].Num())
 				{
 					ParamMapSetVariablesToChunks[ParamMapHistoryIdx][VarIdx] = Inputs[i].CompilationIndex;
-					ParamMapDefinedAttributesToNamespaceVars.FindOrAdd(Var.GetName()) = Var;
+					RecordParamMapDefinedAttributeToNamespaceVar(Var, Inputs[i].Pin);
 					if (ParamMapHistories[ParamMapHistoryIdx].IsPrimaryDataSetOutput(Var, GetTargetUsage())) // Note that data interfaces aren't ever in the primary data set even if the namespace matches.)
 						CompilationOutput.ScriptData.AttributesWritten.AddUnique(Var);
 				}
@@ -5266,7 +5291,10 @@ bool FHlslNiagaraTranslator::ParameterMapRegisterNamespaceAttributeVariable(cons
 			}
 
 			ParamMapDefinedAttributesToUniformChunks.Add(BasicVar.GetName(), UniformChunk);
-			ParamMapDefinedAttributesToNamespaceVars.Add(BasicVar.GetName(), NamespaceVar);
+			FVarAndDefaultSource VarAndDefaultSource;
+			VarAndDefaultSource.Variable = NamespaceVar;
+			VarAndDefaultSource.bDefaultExplicit = true;
+			ParamMapDefinedAttributesToNamespaceVars.Add(BasicVar.GetName(), VarAndDefaultSource);
 		}
 		Output = AddSourceChunk(ParameterMapInstanceName + TEXT(".") + SymbolNameNamespaced, NamespaceVar.GetType());
 		return true;
@@ -5499,7 +5527,7 @@ void FHlslNiagaraTranslator::HandleParameterRead(int32 ParamMapHistoryIdx, const
 
 	if (FNiagaraParameterMapHistory::IsExternalConstantNamespace(Var, CompileOptions.TargetUsage, CompileOptions.GetTargetUsageBitmask()))
 	{
-		if (Variable && Variable->DefaultMode == ENiagaraDefaultMode::FailIfPreviouslyNotSet && !bIgnoreDefaultSetFirst)
+		if (Variable && Variable->DefaultMode == ENiagaraDefaultMode::FailIfPreviouslyNotSet && !Variable->Variable.IsInNameSpace(FNiagaraConstants::UserNamespaceString) && !bIgnoreDefaultSetFirst)
 		{
 			RegisterCompileDependency(Var, FText::Format(LOCTEXT("UsedBeforeSet", "Variable {0} was read before being set. It's default mode is \"Fail If Previously Not Set\", so this isn't allowed."), FText::FromName(Var.GetName())), ErrorNode, nullptr, true);
 		}
@@ -5551,7 +5579,13 @@ void FHlslNiagaraTranslator::HandleParameterRead(int32 ParamMapHistoryIdx, const
 
 
 	bool bFailIfPreviouslyNotSetSentinel = false;
-	if (Variable && Variable->DefaultMode == ENiagaraDefaultMode::FailIfPreviouslyNotSet)
+	bool bValidateFailIfPreviouslyNotSet = false;
+	if (Variable != nullptr)
+	{
+		bValidateFailIfPreviouslyNotSet = Variable->DefaultMode == ENiagaraDefaultMode::FailIfPreviouslyNotSet;
+	}
+
+	if (bValidateFailIfPreviouslyNotSet)
 	{
 		FNiagaraVariable SearchVar = Var;
 		if (FNiagaraParameterMapHistory::IsInitialValue(Var))
@@ -5563,7 +5597,8 @@ void FHlslNiagaraTranslator::HandleParameterRead(int32 ParamMapHistoryIdx, const
 			SearchVar = FNiagaraParameterMapHistory::GetSourceForPreviousValue(Var);
 		}
 
-		bool bSetPreviously = nullptr != ParamMapDefinedAttributesToNamespaceVars.Find(SearchVar.GetName());
+		FVarAndDefaultSource* ParamMapDefinedVarAndDefaultSource = ParamMapDefinedAttributesToNamespaceVars.Find(SearchVar.GetName());
+		bool bSetPreviously = ParamMapDefinedVarAndDefaultSource != nullptr && ParamMapDefinedVarAndDefaultSource->bDefaultExplicit;
 		for (int32 OtherParamIdx = 0; OtherParamIdx < OtherOutputParamMapHistories.Num() && !bSetPreviously; OtherParamIdx++)
 		{
 			// Stop if this is already in our evaluation chain. Assume only indices above us are valid sourcers for this.
@@ -5573,10 +5608,13 @@ void FHlslNiagaraTranslator::HandleParameterRead(int32 ParamMapHistoryIdx, const
 			int32 FoundInParamIdx = OtherOutputParamMapHistories[OtherParamIdx].FindVariableByName(SearchVar.GetName());
 			if (INDEX_NONE != FoundInParamIdx)
 			{
-				if (OtherOutputParamMapHistories[OtherParamIdx].PerVariableWriteHistory[FoundInParamIdx].Num() != 0)
+				for (const FModuleScopedPin& ScopedPin : OtherOutputParamMapHistories[OtherParamIdx].PerVariableWriteHistory[FoundInParamIdx])
 				{
-					bSetPreviously = true;
-					break;
+					if (ScopedPin.Pin->Direction == EEdGraphPinDirection::EGPD_Input && ScopedPin.Pin->bHidden == false)
+					{
+						bSetPreviously = true;
+						break;
+					}
 				}
 			}
 		}
@@ -5584,7 +5622,6 @@ void FHlslNiagaraTranslator::HandleParameterRead(int32 ParamMapHistoryIdx, const
 			bFailIfPreviouslyNotSetSentinel = true;
 	}
 	
-
 	int32 LastSetChunkIdx = INDEX_NONE;
 	if (ParamMapHistoryIdx < ParamMapHistories.Num())
 	{
@@ -5613,28 +5650,28 @@ void FHlslNiagaraTranslator::HandleParameterRead(int32 ParamMapHistoryIdx, const
 
 			if (TranslationOptions.bParameterRapidIteration)
 			{
-			// Now try to look up with the new name.. we may have already made this an external variable before..
-			if (bVarChanged)
-			{
-				VarIdx = ParamMapHistories[ParamMapHistoryIdx].FindVariableByName(Var.GetName());
-				if (VarIdx != INDEX_NONE && VarIdx < ParamMapSetVariablesToChunks[ParamMapHistoryIdx].Num())
+				// Now try to look up with the new name.. we may have already made this an external variable before..
+				if (bVarChanged)
 				{
-					LastSetChunkIdx = ParamMapSetVariablesToChunks[ParamMapHistoryIdx][VarIdx];
+					VarIdx = ParamMapHistories[ParamMapHistoryIdx].FindVariableByName(Var.GetName());
+					if (VarIdx != INDEX_NONE && VarIdx < ParamMapSetVariablesToChunks[ParamMapHistoryIdx].Num())
+					{
+						LastSetChunkIdx = ParamMapSetVariablesToChunks[ParamMapHistoryIdx][VarIdx];
+					}
 				}
-			}
 
-			// If it isn't found yet.. go ahead and make it into a constant variable..
-			if (LastSetChunkIdx == INDEX_NONE && ParameterMapRegisterExternalConstantNamespaceVariable(Var, ErrorNode, ParamMapHistoryIdx, OutputChunkId, InputPin))
-			{
-				LastSetChunkIdx = OutputChunkId;
-				if (VarIdx != INDEX_NONE && VarIdx < ParamMapSetVariablesToChunks[ParamMapHistoryIdx].Num())
+				// If it isn't found yet.. go ahead and make it into a constant variable..
+				if (LastSetChunkIdx == INDEX_NONE && ParameterMapRegisterExternalConstantNamespaceVariable(Var, ErrorNode, ParamMapHistoryIdx, OutputChunkId, InputPin))
 				{
-					// Record that we wrote to it.
-					ParamMapSetVariablesToChunks[ParamMapHistoryIdx][VarIdx] = LastSetChunkIdx;
-					ParamMapDefinedAttributesToNamespaceVars.FindOrAdd(Var.GetName()) = Var;
+					LastSetChunkIdx = OutputChunkId;
+					if (VarIdx != INDEX_NONE && VarIdx < ParamMapSetVariablesToChunks[ParamMapHistoryIdx].Num())
+					{
+						// Record that we wrote to it.
+						ParamMapSetVariablesToChunks[ParamMapHistoryIdx][VarIdx] = LastSetChunkIdx;
+						RecordParamMapDefinedAttributeToNamespaceVar(Var, DefaultPin);
+					}
+					return;		
 				}
-				return;
-				}				
 			}
 			else
 			{
@@ -5674,14 +5711,16 @@ void FHlslNiagaraTranslator::HandleParameterRead(int32 ParamMapHistoryIdx, const
 		bool bIgnoreDefaultValue = ParamMapHistories[ParamMapHistoryIdx].ShouldIgnoreVariableDefault(Var);
 
 		// First check to see if this is defaulted to fail if not set previously. If so, then make sure we don't suck in defaults and error out.
-		if (LastSetChunkIdx == INDEX_NONE && Variable && Variable->DefaultMode == ENiagaraDefaultMode::FailIfPreviouslyNotSet && bFailIfPreviouslyNotSetSentinel && !bIgnoreDefaultSetFirst)
+		if (bValidateFailIfPreviouslyNotSet && bFailIfPreviouslyNotSetSentinel && !bIgnoreDefaultSetFirst)
 		{		
 			RegisterCompileDependency(Var, FText::Format(LOCTEXT("UsedBeforeSet", "Variable {0} was read before being set. It's default mode is \"Fail If Previously Not Set\", so this isn't allowed."), FText::FromName(Var.GetName())), ErrorNode, nullptr, false);
 		}
 
 		if (bIsPerInstanceAttribute)
 		{
-			FNiagaraVariable* ExistingVar = ParamMapDefinedAttributesToNamespaceVars.Find(Var.GetName());
+			FVarAndDefaultSource* ExistingVarAndDefaultSource = ParamMapDefinedAttributesToNamespaceVars.Find(Var.GetName());
+			FNiagaraVariable* ExistingVar = ExistingVarAndDefaultSource ? &ExistingVarAndDefaultSource->Variable : nullptr;
+
 			bool ExistsInAttribArrayAlready = ExistingVar != nullptr;
 			if (ExistsInAttribArrayAlready&& ExistingVar->GetType() != Var.GetType())
 			{
@@ -5715,9 +5754,13 @@ void FHlslNiagaraTranslator::HandleParameterRead(int32 ParamMapHistoryIdx, const
 
 				if (bFoundExistingSet)
 				{
-					LastSetChunkIdx = AddBodyChunk(ParameterMapInstanceName + TEXT(".") + GetSanitizedSymbolName(Var.GetName().ToString()),
-						ParameterMapInstanceName + TEXT(".") + GetSanitizedSymbolName(SourceForInitialValue.GetName().ToString()), Var.GetType(), false);
-					ParamMapDefinedAttributesToNamespaceVars.FindOrAdd(Var.GetName()) = Var;
+					LastSetChunkIdx = AddBodyChunk(
+						ParameterMapInstanceName + TEXT(".") + GetSanitizedSymbolName(Var.GetName().ToString()),
+						ParameterMapInstanceName + TEXT(".") + GetSanitizedSymbolName(SourceForInitialValue.GetName().ToString()), Var.GetType(),
+						false
+					);
+
+					RecordParamMapDefinedAttributeToNamespaceVar(Var, DefaultPin);
 				}
 				else
 				{
@@ -5831,11 +5874,11 @@ void FHlslNiagaraTranslator::HandleParameterRead(int32 ParamMapHistoryIdx, const
 					{
 						// Record that we wrote to it.
 						ParamMapSetVariablesToChunks[ParamMapHistoryIdx][VarIdx] = LastSetChunkIdx;
-						ParamMapDefinedAttributesToNamespaceVars.FindOrAdd(Var.GetName()) = Var;
+						RecordParamMapDefinedAttributeToNamespaceVar(Var, DefaultPin);
 					}
 					else if (VarIdx == INDEX_NONE && UniqueVars.Contains(Var))
 					{
-						ParamMapDefinedAttributesToNamespaceVars.FindOrAdd(Var.GetName()) = Var;
+						RecordParamMapDefinedAttributeToNamespaceVar(Var, DefaultPin);
 					}
 					else
 					{
@@ -5885,7 +5928,7 @@ void FHlslNiagaraTranslator::HandleParameterRead(int32 ParamMapHistoryIdx, const
 	else
 	{
 		OutputChunkId = AddSourceChunk(ParameterMapInstanceName + TEXT(".") + GetSanitizedSymbolName(Var.GetName().ToString()), Var.GetType());
-		ParamMapDefinedAttributesToNamespaceVars.FindOrAdd(Var.GetName()) = Var;
+		RecordParamMapDefinedAttributeToNamespaceVar(Var, DefaultPin);
 	}
 }
 
@@ -7152,7 +7195,11 @@ void FHlslNiagaraTranslator::RegisterFunctionCall(ENiagaraScriptUsage ScriptUsag
 									{
 										const UEdGraphPin* DefaultPin = History.GetDefaultValuePin(VarIdx);
 										UNiagaraScriptVariable* ScriptVariable = SourceGraph->GetScriptVariable(AliasedVar);
-										HandleParameterRead(ActiveStageIdx, AliasedVar, DefaultPin, ParamNode, LastSetChunkIdx, ScriptVariable);
+
+										// Do not error on defaults for parameter reads here; we may be entering a SetVariable function call which is setting the first default for a parameter.
+										const bool bTreatAsUnknownParameterMap = false;
+										const bool bIgnoreDefaultSetFirst = true;
+										HandleParameterRead(ActiveStageIdx, AliasedVar, DefaultPin, ParamNode, LastSetChunkIdx, ScriptVariable, bTreatAsUnknownParameterMap, bIgnoreDefaultSetFirst);
 										
 										// If this variable was in the pending defaults list, go ahead and remove it
 										// as we added it before first use...
@@ -8438,7 +8485,6 @@ void FHlslNiagaraTranslator::Warning(FText WarningText, const UNiagaraNode* InNo
 	Message(FNiagaraCompileEventSeverity::Warning, WarningText, InNode, Pin, ShortDescription, bDismissable);
 }
 
-
 void FHlslNiagaraTranslator::RegisterCompileDependency(const FNiagaraVariableBase& InVar, FText MessageText, const UNiagaraNode* Node, const UEdGraphPin* Pin, bool bEmitAsLinker)
 {
 	if (InVar.GetType().IsDataInterface() || InVar.GetType().IsUObject() || InVar.IsInNameSpace(FNiagaraConstants::UserNamespaceString) || InVar.IsInNameSpace(FNiagaraConstants::EngineNamespaceString) || InVar.IsInNameSpace(FNiagaraConstants::ParameterCollectionNamespaceString))
diff --git a/Engine/Plugins/FX/Niagara/Source/NiagaraEditor/Private/NiagaraHlslTranslator.h b/Engine/Plugins/FX/Niagara/Source/NiagaraEditor/Private/NiagaraHlslTranslator.h
--- a/Engine/Plugins/FX/Niagara/Source/NiagaraEditor/Private/NiagaraHlslTranslator.h
+++ b/Engine/Plugins/FX/Niagara/Source/NiagaraEditor/Private/NiagaraHlslTranslator.h
@@ -701,6 +701,9 @@ private:
 	void BuildConstantBuffer(ENiagaraCodeChunkMode ChunkMode);
 
 	void TrimAttributes(const FNiagaraCompileOptions& InCompileOptions, TArray<FNiagaraVariable>& Attributes);
+
+	// Specific method to reconcile whether the default value has been implicitly or explicitly set for a namespaced var added to the param map.
+	void RecordParamMapDefinedAttributeToNamespaceVar(const FNiagaraVariable& VarToRecord, const UEdGraphPin* VarAssociatedDefaultPin);
 	
 	/** Map of symbol names to count of times it's been used. Used for generating unique symbol names. */
 	TMap<FName, uint32> SymbolCounts;
@@ -738,6 +741,12 @@ private:
 		int32 ChunkMode;
 	};
 
+	struct FVarAndDefaultSource
+	{
+		FNiagaraVariable Variable;
+		bool bDefaultExplicit = false; //Whether or not the default value of the variable is explicit, e.g. there is an explicit value on a pin, explicit binding, or explicit custom initialization.
+	};
+
 	TMap<FName, UniformVariableInfo> ParamMapDefinedSystemVars; // Map from the defined constants to the uniform chunk expressing them (i.e. have we encountered before in this graph?)
 
 	// Synced to the EmitterParameter uniforms encountered for parameter maps thus far.
@@ -746,7 +755,7 @@ private:
 
 	// Synced to the Attributes encountered for parameter maps thus far.
 	TMap<FName, int32> ParamMapDefinedAttributesToUniformChunks; // Map from the variable name exposed as a attribute to the uniform chunk expressing it (i.e. have we encountered before in this graph?)
-	TMap<FName, FNiagaraVariable> ParamMapDefinedAttributesToNamespaceVars; // Map from defined parameter to the Namespaced variable expressing it.
+	TMap<FName, FVarAndDefaultSource> ParamMapDefinedAttributesToNamespaceVars; // Map from defined parameter to the Namespaced variable expressing it.
 
 	// Synced to the external variables used when bulk compiling system scripts.
 	TArray<FNiagaraVariable> ExternalVariablesForBulkUsage;
