Niagara LWC changes:
* added new position type
* changed core definition and some modules to use the position type
* auto-convert simulation inputs and outputs of positions from and to LWC space
* update existing graph logic to auto-convert pins to position type where possible and stay backwards compatible with vector type connections
* updated data interfaces to correctly convert position types

diff --git a/Engine/Plugins/Experimental/GeometryCache/Source/GeometryCache/Private/NiagaraRendererGeometryCache.cpp b/Engine/Plugins/Experimental/GeometryCache/Source/GeometryCache/Private/NiagaraRendererGeometryCache.cpp
--- a/Engine/Plugins/Experimental/GeometryCache/Source/GeometryCache/Private/NiagaraRendererGeometryCache.cpp
+++ b/Engine/Plugins/Experimental/GeometryCache/Source/GeometryCache/Private/NiagaraRendererGeometryCache.cpp
@@ -308,12 +308,13 @@ void FNiagaraRendererGeometryCache::PostSystemTick_GameThread(const UNiagaraRend
 			}
 		}
 
-		const auto PositionAccessor = FNiagaraDataSetAccessor<FVector3f>::CreateReader(Data, Properties->PositionBinding.GetDataSetBindableVariable().GetName());
+		const auto PositionAccessor = FNiagaraDataSetAccessor<FNiagaraPosition>::CreateReader(Data, Properties->PositionBinding.GetDataSetBindableVariable().GetName());
 		const auto RotationAccessor = FNiagaraDataSetAccessor<FVector3f>::CreateReader(Data, Properties->RotationBinding.GetDataSetBindableVariable().GetName());
 		const auto ScaleAccessor = FNiagaraDataSetAccessor<FVector3f>::CreateReader(Data, Properties->ScaleBinding.GetDataSetBindableVariable().GetName());
 		const auto ElapsedTimeAccessor = FNiagaraDataSetAccessor<float>::CreateReader(Data, Properties->ElapsedTimeBinding.GetDataSetBindableVariable().GetName());
+		FNiagaraLWCConverter LwcConverter = SystemInstance->GetLWCConverter(Emitter->GetCachedEmitter()->bLocalSpace);
 
-		FVector Position = PositionAccessor.GetSafe(ParticleIndex, FVector3f::ZeroVector);
+		FVector Position = LwcConverter.ConvertSimulationPositionToWorld(PositionAccessor.GetSafe(ParticleIndex, FNiagaraPosition(ForceInit)));
 		FVector Scale = ScaleAccessor.GetSafe(ParticleIndex, FVector3f::OneVector);
 		FVector RotationVector = RotationAccessor.GetSafe(ParticleIndex, FVector3f::ZeroVector);
 		FTransform Transform(FRotator(RotationVector.X, RotationVector.Y, RotationVector.Z), Position, Scale);
diff --git a/Engine/Plugins/FX/Niagara/Shaders/Private/NiagaraDataInterfaceCamera.ush b/Engine/Plugins/FX/Niagara/Shaders/Private/NiagaraDataInterfaceCamera.ush
new file mode 100644
--- /dev/null
+++ b/Engine/Plugins/FX/Niagara/Shaders/Private/NiagaraDataInterfaceCamera.ush
@@ -0,0 +1,69 @@
+// Copyright Epic Games, Inc. All Rights Reserved.
+
+#pragma once
+
+struct FNDICameraContext
+{
+	float3 SystemLWCTile;
+};
+
+void DICamera_GetCameraProperties(in FNDICameraContext DIContext, out float3 Out_CameraPositionWorld, out float3 Out_ViewForwardVector, out float3 Out_ViewUpVector, out float3 Out_ViewRightVector)
+{
+	FLWCVector3 CameraPos = PrimaryView.WorldCameraOrigin;
+	CameraPos.Tile -= DIContext.SystemLWCTile; // convert to simulation space
+	Out_CameraPositionWorld = LWCToFloat(CameraPos);
+	Out_ViewForwardVector = View.ViewForward;
+	Out_ViewUpVector = View.ViewUp;
+	Out_ViewRightVector = View.ViewRight;
+}
+
+void DICamera_GetViewProperties(in FNDICameraContext DIContext, out float3 Out_ViewPositionWorld, out float3 Out_ViewForwardVector, out float3 Out_ViewUpVector, out float3 Out_ViewRightVector, out float4 Out_ViewSizeAndInverseSize,
+	out float4 Out_ScreenToViewSpace, out float2 Out_Current_TAAJitter, out float2 Out_Previous_TAAJitter, out float3 Out_PreViewTranslation, out float4 Out_BufferSizeAndInverseSize, out float2 Out_ViewportOffset)
+{
+	FLWCVector3 ViewPos = PrimaryView.WorldViewOrigin;
+	ViewPos.Tile -= DIContext.SystemLWCTile; // convert to simulation space
+	Out_ViewPositionWorld.xyz = LWCToFloat(ViewPos);
+	Out_ViewForwardVector = View.ViewForward;
+	Out_ViewUpVector = View.ViewUp;
+	Out_ViewRightVector = View.ViewRight;
+	Out_ViewSizeAndInverseSize = View.ViewSizeAndInvSize;
+	Out_ScreenToViewSpace = View.ScreenToViewSpace;
+	Out_Current_TAAJitter = View.TemporalAAJitter.xy;
+	Out_Previous_TAAJitter = View.TemporalAAJitter.zw;
+	Out_PreViewTranslation = LWCHackToFloat(PrimaryView.PreViewTranslation);
+	Out_BufferSizeAndInverseSize = View.BufferSizeAndInvSize;
+	Out_ViewportOffset = View.ViewRectMin.xy;
+}
+
+void DICamera_GetClipSpaceTransforms(in FNDICameraContext DIContext, out float4x4 Out_WorldToClipTransform, out float4x4 Out_TranslatedWorldToClipTransform, out float4x4 Out_ClipToWorldTransform, out float4x4 Out_ClipToViewTransform,
+	out float4x4 Out_ClipToTranslatedWorldTransform, out float4x4 Out_ScreenToWorldTransform, out float4x4 Out_ScreenToTranslatedWorldTransform, out float4x4 Out_ClipToPreviousClipTransform)
+{
+	Out_WorldToClipTransform = LWCHackToFloat(PrimaryView.WorldToClip);
+	Out_TranslatedWorldToClipTransform = View.TranslatedWorldToClip;
+	Out_ClipToWorldTransform = LWCHackToFloat(PrimaryView.ClipToWorld);
+	Out_ClipToViewTransform = View.ClipToView;
+	Out_ClipToTranslatedWorldTransform = View.ClipToTranslatedWorld;
+	Out_ScreenToWorldTransform = LWCHackToFloat(PrimaryView.ScreenToWorld);
+	Out_ScreenToTranslatedWorldTransform = View.ScreenToTranslatedWorld;
+	Out_ClipToPreviousClipTransform = View.ClipToPrevClip;
+}
+
+void DICamera_GetViewSpaceTransforms(in FNDICameraContext DIContext, out float4x4 Out_TranslatedWorldToViewTransform, out float4x4 Out_ViewToTranslatedWorldTransform, out float4x4 Out_TranslatedWorldToCameraViewTransform,
+	out float4x4 Out_CameraViewToTranslatedWorldTransform, out float4x4 Out_ViewToClipTransform, out float4x4 Out_ViewToClipNoAATransform)
+{
+	Out_TranslatedWorldToViewTransform = View.TranslatedWorldToView;
+	Out_ViewToTranslatedWorldTransform = View.ViewToTranslatedWorld;
+	Out_TranslatedWorldToCameraViewTransform = View.TranslatedWorldToCameraView;
+	Out_CameraViewToTranslatedWorldTransform = View.CameraViewToTranslatedWorld;
+	Out_ViewToClipTransform = View.ViewToClip;
+	Out_ViewToClipNoAATransform = View.ViewToClipNoAA;
+}
+
+////////////////////////////////////////////////////////////////////////////////////////////////////////////////////////////////////////////////////////////////
+
+#define NDICAMERA_DECLARE_CONSTANTS(NAME) \
+	float3				SystemLWCTile_ ## NAME;
+
+#define NDICAMERA_MAKE_CONTEXT(NAME) \
+	FNDICameraContext DIContext; \
+	DIContext.SystemLWCTile = SystemLWCTile_ ## NAME;
diff --git a/Engine/Plugins/FX/Niagara/Shaders/Private/NiagaraDataInterfaceCollisionQuery.ush b/Engine/Plugins/FX/Niagara/Shaders/Private/NiagaraDataInterfaceCollisionQuery.ush
--- a/Engine/Plugins/FX/Niagara/Shaders/Private/NiagaraDataInterfaceCollisionQuery.ush
+++ b/Engine/Plugins/FX/Niagara/Shaders/Private/NiagaraDataInterfaceCollisionQuery.ush
@@ -15,15 +15,18 @@ float NDICollisionQuery_GetCustomSceneDepth(float2 ScreenUV)
 #endif
 }
 
-void NDICollisionQuery_QuerySceneDepthGPU(in float3 In_SamplePos, out float Out_SceneDepth, out float3 Out_CameraPosWorld, out bool Out_IsInsideView, out float3 Out_WorldPos, out float3 Out_WorldNormal)
+void NDICollisionQuery_QuerySceneDepthGPU(in float3 In_SamplePos, in float3 In_LWCTile, out float Out_SceneDepth, out float3 Out_CameraPosWorld, out bool Out_IsInsideView, out float3 Out_WorldPos, out float3 Out_WorldNormal)
 {
 	Out_SceneDepth = -1;
 	Out_WorldPos = float3(0.0, 0.0, 0.0);
 	Out_WorldNormal = float3(0.0, 0.0, 1.0);
 	Out_IsInsideView = true;
-	Out_CameraPosWorld.xyz = LWCHackToFloat(PrimaryView.WorldCameraOrigin);
+	FLWCVector3 CameraPos = PrimaryView.WorldCameraOrigin;
+	CameraPos.Tile -= In_LWCTile; // convert to simulation space
+	Out_CameraPosWorld = LWCToFloat(CameraPos);
 
-	float4 SamplePosition = float4(In_SamplePos + LWCHackToFloat(PrimaryView.PreViewTranslation), 1);
+	FLWCVector3 LwcSamplePos = MakeLWCVector3(In_LWCTile, In_SamplePos);
+	float4 SamplePosition = float4(LWCToFloat(LWCAdd(LwcSamplePos, PrimaryView.PreViewTranslation)), 1); // TODO[mg]: LWCToFloat here?
 	float4 ClipPosition = mul(SamplePosition, View.TranslatedWorldToClip);
 	float2 ScreenPosition = ClipPosition.xy / ClipPosition.w;
 
@@ -33,7 +36,7 @@ void NDICollisionQuery_QuerySceneDepthGPU(in float3 In_SamplePos, out float Out_
 		// Sample the depth buffer to get a world position near the sample position.
 		float2 ScreenUV = ScreenPosition * View.ScreenPositionScaleBias.xy + View.ScreenPositionScaleBias.wz;
 		float SceneDepth = CalcSceneDepth(ScreenUV);
-		float3 WorldPosition = WorldPositionFromSceneDepth(ScreenPosition.xy, SceneDepth);
+		FLWCVector3 WorldPosition = WorldPositionFromSceneDepth(ScreenPosition.xy, SceneDepth);
 
 #if NDICOLLISIONQUERY_USE_GBUFFER_NORMAL
 		float3 WorldNormal = Texture2DSampleLevel(SceneTexturesStruct.GBufferATexture, SceneTexturesStruct_GBufferATextureSampler, ScreenUV, 0).xyz * 2.0 - 1.0;
@@ -50,14 +53,17 @@ void NDICollisionQuery_QuerySceneDepthGPU(in float3 In_SamplePos, out float Out_
 			return;
 		}
 
-		float3 WorldPosition0 = WorldPositionFromSceneDepth(ScreenPosition.xy + float2(2 * View.ViewSizeAndInvSize.z, 0.0), SceneDepth0);
-		float3 WorldPosition1 = WorldPositionFromSceneDepth(ScreenPosition.xy - float2(0.0, 2 * View.ViewSizeAndInvSize.w), SceneDepth1);
-		float3 WorldNormal = normalize(cross(WorldPosition0 - WorldPosition, WorldPosition1 - WorldPosition));
+		FLWCVector3 WorldPosition0 = WorldPositionFromSceneDepth(ScreenPosition.xy + float2(2 * View.ViewSizeAndInvSize.z, 0.0), SceneDepth0);
+		FLWCVector3 WorldPosition1 = WorldPositionFromSceneDepth(ScreenPosition.xy - float2(0.0, 2 * View.ViewSizeAndInvSize.w), SceneDepth1);
+		float3 WorldNormal = normalize(cross(LWCSubtract(WorldPosition0, WorldPosition).Offset, LWCSubtract(WorldPosition1, WorldPosition).Offset));
 #endif
 
+		// convert LWC position back to simulation space
+		WorldPosition.Tile -= In_LWCTile;
+		
 		// Set outputs
 		Out_SceneDepth = SceneDepth;
-		Out_WorldPos = WorldPosition;
+		Out_WorldPos = LWCToFloat(WorldPosition);
 		Out_WorldNormal = WorldNormal;
 	}
 	else
@@ -66,15 +72,18 @@ void NDICollisionQuery_QuerySceneDepthGPU(in float3 In_SamplePos, out float Out_
 	}
 }
 
-void NDICollisionQuery_QueryCustomDepthGPU(in float3 In_SamplePos, out float Out_SceneDepth, out float3 Out_CameraPosWorld, out bool Out_IsInsideView, out float3 Out_WorldPos, out float3 Out_WorldNormal)
+void NDICollisionQuery_QueryCustomDepthGPU(in float3 In_SamplePos, in float3 In_LWCTile, out float Out_SceneDepth, out float3 Out_CameraPosWorld, out bool Out_IsInsideView, out float3 Out_WorldPos, out float3 Out_WorldNormal)
 {
 	Out_SceneDepth = -1;
 	Out_WorldPos = float3(0.0, 0.0, 0.0);
 	Out_WorldNormal = float3(0.0, 0.0, 1.0);
 	Out_IsInsideView = true;
-	Out_CameraPosWorld.xyz = LWCHackToFloat(PrimaryView.WorldCameraOrigin);
+	FLWCVector3 CameraPos = PrimaryView.WorldCameraOrigin;
+	CameraPos.Tile -= In_LWCTile; // convert to simulation space
+	Out_CameraPosWorld = LWCToFloat(CameraPos);
 
-	float4 SamplePosition = float4(In_SamplePos + LWCHackToFloat(PrimaryView.PreViewTranslation), 1);
+	FLWCVector3 LwcSamplePos = MakeLWCVector3(In_LWCTile, In_SamplePos);
+	float4 SamplePosition = float4(LWCToFloat(LWCAdd(LwcSamplePos, PrimaryView.PreViewTranslation)), 1); // TODO[mg]: LWCToFloat here?
 	float4 ClipPosition = mul(SamplePosition, View.TranslatedWorldToClip);
 	float2 ScreenPosition = ClipPosition.xy / ClipPosition.w;
 
@@ -84,7 +93,7 @@ void NDICollisionQuery_QueryCustomDepthGPU(in float3 In_SamplePos, out float Out
 		// Sample the depth buffer to get a world position near the sample position.
 		float2 ScreenUV = ScreenPosition * View.ScreenPositionScaleBias.xy + View.ScreenPositionScaleBias.wz;
 		float SceneDepth = NDICollisionQuery_GetCustomSceneDepth(ScreenUV);
-		float3 WorldPosition = WorldPositionFromSceneDepth(ScreenPosition.xy, SceneDepth);
+		FLWCVector3 WorldPosition = WorldPositionFromSceneDepth(ScreenPosition.xy, SceneDepth);
 
 #if NDICOLLISIONQUERY_USE_GBUFFER_NORMAL
 		float3 WorldNormal = Texture2DSampleLevel(SceneTexturesStruct.GBufferATexture, SceneTexturesStruct_GBufferATextureSampler, ScreenUV, 0).xyz * 2.0 - 1.0;
@@ -101,14 +110,17 @@ void NDICollisionQuery_QueryCustomDepthGPU(in float3 In_SamplePos, out float Out
 			return;
 		}
 
-		float3 WorldPosition0 = WorldPositionFromSceneDepth(ScreenPosition.xy + float2(2 * View.ViewSizeAndInvSize.z, 0.0), SceneDepth0);
-		float3 WorldPosition1 = WorldPositionFromSceneDepth(ScreenPosition.xy - float2(0.0, 2 * View.ViewSizeAndInvSize.w), SceneDepth1);
-		float3 WorldNormal = normalize(cross(WorldPosition0 - WorldPosition, WorldPosition1 - WorldPosition));
+		FLWCVector3 WorldPosition0 = WorldPositionFromSceneDepth(ScreenPosition.xy + float2(2 * View.ViewSizeAndInvSize.z, 0.0), SceneDepth0);
+		FLWCVector3 WorldPosition1 = WorldPositionFromSceneDepth(ScreenPosition.xy - float2(0.0, 2 * View.ViewSizeAndInvSize.w), SceneDepth1);
+		float3 WorldNormal = normalize(cross(LWCSubtract(WorldPosition0, WorldPosition).Offset, LWCSubtract(WorldPosition1, WorldPosition).Offset));
 #endif
 
+		// convert LWC position back to simulation space
+		WorldPosition.Tile = WorldPosition.Tile - In_LWCTile;
+		
 		// Set outputs
 		Out_SceneDepth = SceneDepth;
-		Out_WorldPos = WorldPosition;
+		Out_WorldPos = LWCToFloat(WorldPosition);
 		Out_WorldNormal = WorldNormal;
 	}
 	else
@@ -117,11 +129,14 @@ void NDICollisionQuery_QueryCustomDepthGPU(in float3 In_SamplePos, out float Out
 	}
 }
 
-void NDICollisionQuery_QueryMeshDistanceFieldGPU(in float3 In_SamplePos, out float Out_DistanceToNearestSurface, out float3 Out_FieldGradient, out bool Out_IsDistanceFieldValid)
+void NDICollisionQuery_QueryMeshDistanceFieldGPU(in float3 In_SamplePos, in float3 In_LWCTile, out float Out_DistanceToNearestSurface, out float3 Out_FieldGradient, out bool Out_IsDistanceFieldValid)
 {
 #if PLATFORM_SUPPORTS_DISTANCE_FIELDS && (FEATURE_LEVEL >= FEATURE_LEVEL_SM5)
-	Out_DistanceToNearestSurface = GetDistanceToNearestSurfaceGlobal(In_SamplePos);
-	Out_FieldGradient = GetDistanceFieldGradientGlobal(In_SamplePos);
+	FLWCVector3 LwcSamplePos = MakeLWCVector3(In_LWCTile, In_SamplePos);
+	float3 SamplePosition = LWCToFloat(LwcSamplePos); // TODO: remove LWCToFloat here when GetDistanceToNearestSurfaceGlobal supports LWC
+	
+	Out_DistanceToNearestSurface = GetDistanceToNearestSurfaceGlobal(SamplePosition);
+	Out_FieldGradient = GetDistanceFieldGradientGlobal(SamplePosition);
 	Out_IsDistanceFieldValid = (MaxGlobalDFAOConeDistance > 0) && !(Out_DistanceToNearestSurface > 0 && all(Out_FieldGradient == float3(0,0,0.001f)));
 #else
 	Out_DistanceToNearestSurface = 0;
@@ -138,18 +153,21 @@ uint BuildRayTraceMask(int TraceChannel)
 	return 0xFF;
 }
 
-void NDICollisionQuery_IssueAsyncRayTrace(uint MaxRayTraceCount, RWStructuredBuffer<FNiagaraRayData> RWRayRequests, uint RayRequestOffset, RWBuffer<uint> RWRayTraceCounts, uint RayTraceCountsOffset, int In_QueryID, float3 In_TraceStart, float3 In_TraceEnd, int In_TraceChannel, out bool Out_IsQueryValid)
+void NDICollisionQuery_IssueAsyncRayTrace(uint MaxRayTraceCount, RWStructuredBuffer<FNiagaraRayData> RWRayRequests, uint RayRequestOffset, RWBuffer<uint> RWRayTraceCounts, uint RayTraceCountsOffset, float3 In_LWCTile, int In_QueryID, float3 In_TraceStart, float3 In_TraceEnd, int In_TraceChannel, out bool Out_IsQueryValid)
 {
 	Out_IsQueryValid = false;
 
 	if (In_QueryID >= 0 && uint(In_QueryID) < MaxRayTraceCount)
 	{
-		float3 Trace = In_TraceEnd - In_TraceStart;
+		FLWCVector3 LwcStartPos = MakeLWCVector3(In_LWCTile, In_TraceStart);
+		FLWCVector3 LwcEndPos = MakeLWCVector3(In_LWCTile, In_TraceEnd);
+		
+		float3 Trace = LWCToFloat(LWCSubtract(LwcEndPos, LwcStartPos));
 		float TraceLength = length(Trace);
 
 		//Offset query into the actual ray buffer.
 		uint RayIndex = RayRequestOffset + In_QueryID;
-		RWRayRequests[RayIndex].Origin = In_TraceStart;
+		RWRayRequests[RayIndex].Origin = LWCToFloat(LwcStartPos); // TODO: remove LWCToFloat here when FNiagaraRayData supports LWC
 		RWRayRequests[RayIndex].Mask = BuildRayTraceMask(In_TraceChannel);
 		RWRayRequests[RayIndex].Direction = normalize(Trace);
 		RWRayRequests[RayIndex].TFar = TraceLength;
@@ -185,7 +203,7 @@ void NDICollisionQuery_ReserveRayTraceIndex(uint MaxRayTraceCount, RWBuffer<uint
 	}
 }
 
-void NDICollisionQuery_CreateAsyncRayTrace(	uint MaxRayTraceCount, RWStructuredBuffer<FNiagaraRayData> RWRayRequests, uint RayRequestsOffset, RWBuffer<uint> RWRayTraceCounts, uint RayTraceCountsOffset,
+void NDICollisionQuery_CreateAsyncRayTrace(	uint MaxRayTraceCount, RWStructuredBuffer<FNiagaraRayData> RWRayRequests, uint RayRequestsOffset, RWBuffer<uint> RWRayTraceCounts, uint RayTraceCountsOffset, float3 In_LWCTile,
 											float3 In_TraceStart, float3 In_TraceEnd, int In_TraceChannel, out int Out_QueryID, out bool Out_IsQueryValid)
 {
 	Out_IsQueryValid = false;
@@ -194,11 +212,14 @@ void NDICollisionQuery_CreateAsyncRayTrace(	uint MaxRayTraceCount, RWStructuredB
 
 	if (Out_IsQueryValid)
 	{
-		float3 Trace = In_TraceEnd - In_TraceStart;
+		FLWCVector3 LwcStartPos = MakeLWCVector3(In_LWCTile, In_TraceStart);
+		FLWCVector3 LwcEndPos = MakeLWCVector3(In_LWCTile, In_TraceEnd);
+		
+		float3 Trace = LWCToFloat(LWCSubtract(LwcEndPos, LwcStartPos));
 		float TraceLength = length(Trace);
 
 		uint RayIndex = Out_QueryID + RayRequestsOffset;
-		RWRayRequests[RayIndex].Origin = In_TraceStart;
+		RWRayRequests[RayIndex].Origin = LWCToFloat(LwcStartPos); // TODO: remove LWCToFloat here when FNiagaraRayData supports LWC
 		RWRayRequests[RayIndex].Mask = BuildRayTraceMask(In_TraceChannel);
 		RWRayRequests[RayIndex].Direction = normalize(Trace);
 		RWRayRequests[RayIndex].TFar = TraceLength;
@@ -208,7 +229,7 @@ void NDICollisionQuery_CreateAsyncRayTrace(	uint MaxRayTraceCount, RWStructuredB
 	}
 }
 
-void NDICollisionQuery_ReadAsyncRayTrace(uint MaxRayTraceCount, StructuredBuffer<FNiagaraRayTracingResult> IntersectionResults, uint IntersectionRestultsOffest, int In_PreviousFrameQueryID, out bool Out_CollisionValid, out float Out_CollisionDistance, out float3 Out_CollisionPosWorld, out float3 Out_CollisionNormal)
+void NDICollisionQuery_ReadAsyncRayTrace(uint MaxRayTraceCount, StructuredBuffer<FNiagaraRayTracingResult> IntersectionResults, uint IntersectionRestultsOffest, float3 In_LWCTile, int In_PreviousFrameQueryID, out bool Out_CollisionValid, out float Out_CollisionDistance, out float3 Out_CollisionPosWorld, out float3 Out_CollisionNormal)
 {
 	Out_CollisionValid = false;
 	Out_CollisionDistance = 0.0f;
@@ -223,7 +244,7 @@ void NDICollisionQuery_ReadAsyncRayTrace(uint MaxRayTraceCount, StructuredBuffer
 		if (Out_CollisionValid)
 		{
 			Out_CollisionDistance = HitResult.HitT;
-			Out_CollisionPosWorld = HitResult.WorldPosition;
+			Out_CollisionPosWorld = LWCToFloat(MakeLWCVector3(-In_LWCTile, HitResult.WorldPosition)); // TODO: HitResult.WorldPosition should support LWC properly
 			Out_CollisionNormal = HitResult.WorldNormal;
 		}
 	}
diff --git a/Engine/Plugins/FX/Niagara/Shaders/Private/NiagaraDataInterfaceCollisionQueryTemplate.ush b/Engine/Plugins/FX/Niagara/Shaders/Private/NiagaraDataInterfaceCollisionQueryTemplate.ush
--- a/Engine/Plugins/FX/Niagara/Shaders/Private/NiagaraDataInterfaceCollisionQueryTemplate.ush
+++ b/Engine/Plugins/FX/Niagara/Shaders/Private/NiagaraDataInterfaceCollisionQueryTemplate.ush
@@ -1,7 +1,8 @@
 // Copyright Epic Games, Inc. All Rights Reserved.
 
+float3										SystemLWCTile_{ParameterName};
 #if NIAGARA_USE_RAYTRACING
-/** Niagara ray tracing uses shared buffers between dispathces so each buffer has a buffer + and offset into that buffer for the current dispatch. */
+/** Niagara ray tracing uses shared buffers between dispatches so each buffer has a buffer + and offset into that buffer for the current dispatch. */
 bool										RayTracingEnabled_{ParameterName};
 uint										MaxRayTraceCount_{ParameterName};
 RWStructuredBuffer<FNiagaraRayData>			RWRayRequests_{ParameterName};
@@ -14,17 +15,17 @@ uint										RayTraceCountsOffset_{ ParameterName };
 
 void QuerySceneDepthGPU_{ParameterName}(in float3 In_SamplePos, out float Out_SceneDepth, out float3 Out_CameraPosWorld, out bool Out_IsInsideView, out float3 Out_WorldPos, out float3 Out_WorldNormal)
 {
-	NDICollisionQuery_QuerySceneDepthGPU(In_SamplePos, Out_SceneDepth, Out_CameraPosWorld, Out_IsInsideView, Out_WorldPos, Out_WorldNormal);
+	NDICollisionQuery_QuerySceneDepthGPU(In_SamplePos, SystemLWCTile_{ParameterName}, Out_SceneDepth, Out_CameraPosWorld, Out_IsInsideView, Out_WorldPos, Out_WorldNormal);
 }
 
 void QueryCustomDepthGPU_{ParameterName}(in float3 In_SamplePos, out float Out_SceneDepth, out float3 Out_CameraPosWorld, out bool Out_IsInsideView, out float3 Out_WorldPos, out float3 Out_WorldNormal)
 {
-	NDICollisionQuery_QueryCustomDepthGPU(In_SamplePos, Out_SceneDepth, Out_CameraPosWorld, Out_IsInsideView, Out_WorldPos, Out_WorldNormal);
+	NDICollisionQuery_QueryCustomDepthGPU(In_SamplePos, SystemLWCTile_{ParameterName}, Out_SceneDepth, Out_CameraPosWorld, Out_IsInsideView, Out_WorldPos, Out_WorldNormal);
 }
 
 void QueryMeshDistanceFieldGPU_{ParameterName}(in float3 In_SamplePos, out float Out_DistanceToNearestSurface, out float3 Out_FieldGradient, out bool Out_IsDistanceFieldValid)
 {
-	NDICollisionQuery_QueryMeshDistanceFieldGPU(In_SamplePos, Out_DistanceToNearestSurface, Out_FieldGradient, Out_IsDistanceFieldValid);
+	NDICollisionQuery_QueryMeshDistanceFieldGPU(In_SamplePos, SystemLWCTile_{ParameterName}, Out_DistanceToNearestSurface, Out_FieldGradient, Out_IsDistanceFieldValid);
 }
 
 void IssueAsyncRayTraceGpu_{ParameterName}_UEImpureCall(int In_QueryID, float3 In_TraceStart, float3 In_TraceEnd, int In_TraceChannel, out bool Out_IsQueryValid)
@@ -32,7 +33,7 @@ void IssueAsyncRayTraceGpu_{ParameterName}_UEImpureCall(int In_QueryID, float3 I
 #if NIAGARA_USE_RAYTRACING
 	if (RayTracingEnabled_{ParameterName})
 	{
-		NDICollisionQuery_IssueAsyncRayTrace(MaxRayTraceCount_{ParameterName}, RWRayRequests_{ParameterName}, RayRequestsOffset_{ ParameterName }, RWRayTraceCounts_{ParameterName}, RayTraceCountsOffset_{ ParameterName }, In_QueryID, In_TraceStart, In_TraceEnd, In_TraceChannel, Out_IsQueryValid);
+		NDICollisionQuery_IssueAsyncRayTrace(MaxRayTraceCount_{ParameterName}, RWRayRequests_{ParameterName}, RayRequestsOffset_{ ParameterName }, RWRayTraceCounts_{ParameterName}, RayTraceCountsOffset_{ ParameterName }, SystemLWCTile_{ParameterName}, In_QueryID, In_TraceStart, In_TraceEnd, In_TraceChannel, Out_IsQueryValid);
 		return;
 	}
 #endif
@@ -56,7 +57,7 @@ void CreateAsyncRayTraceGpu_{ParameterName}_UEImpureCall(float3 In_TraceStart, f
 #if NIAGARA_USE_RAYTRACING
 	if (RayTracingEnabled_{ParameterName})
 	{
-		NDICollisionQuery_CreateAsyncRayTrace(MaxRayTraceCount_{ParameterName}, RWRayRequests_{ParameterName}, RayRequestsOffset_{ ParameterName }, RWRayTraceCounts_{ParameterName}, RayTraceCountsOffset_{ ParameterName }, In_TraceStart, In_TraceEnd, In_TraceChannel, Out_QueryID, Out_IsQueryValid);
+		NDICollisionQuery_CreateAsyncRayTrace(MaxRayTraceCount_{ParameterName}, RWRayRequests_{ParameterName}, RayRequestsOffset_{ ParameterName }, RWRayTraceCounts_{ParameterName}, RayTraceCountsOffset_{ ParameterName }, SystemLWCTile_{ParameterName}, In_TraceStart, In_TraceEnd, In_TraceChannel, Out_QueryID, Out_IsQueryValid);
 		return;
 	}
 #endif
@@ -68,7 +69,7 @@ void ReadAsyncRayTraceGpu_{ParameterName}(int In_PreviousFrameQueryID, out bool
 #if NIAGARA_USE_RAYTRACING
 	if (RayTracingEnabled_{ParameterName})
 	{
-		NDICollisionQuery_ReadAsyncRayTrace(MaxRayTraceCount_{ParameterName}, IntersectionResults_{ParameterName}, IntersectionResultsOffset_{ParameterName}, In_PreviousFrameQueryID, Out_CollisionValid, Out_CollisionDistance, Out_CollisionPosWorld, Out_CollisionNormal);
+		NDICollisionQuery_ReadAsyncRayTrace(MaxRayTraceCount_{ParameterName}, IntersectionResults_{ParameterName}, IntersectionResultsOffset_{ParameterName}, SystemLWCTile_{ParameterName}, In_PreviousFrameQueryID, Out_CollisionValid, Out_CollisionDistance, Out_CollisionPosWorld, Out_CollisionNormal);
 		return;
 	}
 #endif
diff --git a/Engine/Plugins/FX/Niagara/Shaders/Private/NiagaraDataInterfaceLandscape.ush b/Engine/Plugins/FX/Niagara/Shaders/Private/NiagaraDataInterfaceLandscape.ush
--- a/Engine/Plugins/FX/Niagara/Shaders/Private/NiagaraDataInterfaceLandscape.ush
+++ b/Engine/Plugins/FX/Niagara/Shaders/Private/NiagaraDataInterfaceLandscape.ush
@@ -57,6 +57,8 @@ struct FNDILandscapeContext
 
 	Texture2D<int> CachedPhysMatTexture;
 	int2 CachedPhysMatTextureDimension;
+
+	float3 SystemLWCTile;
 };
 
 void NDILandscape_GetBaseColor(in FNDILandscapeContext DIContext, float3 InWorldPos, out float3 OutBaseColor, out bool OutIsValid)
@@ -96,6 +98,7 @@ void NDILandscape_GetBaseColor(in FNDILandscapeContext DIContext, float3 InWorld
 
 void NDILandscape_GetHeight(in FNDILandscapeContext DIContext, float3 InWorldPos, out float OutHeight, out bool OutIsValid)
 {
+	// todo[mg] convert world pos to LWC pos
 	if (DIContext.HeightVirtualTextureEnabled)
 	{
 		float3 UnpackHeightOrigin = DIContext.HeightVirtualTextureWorldToUvTransform[0].xyz;
@@ -139,6 +142,7 @@ void NDILandscape_GetHeight(in FNDILandscapeContext DIContext, float3 InWorldPos
 
 void NDILandscape_GetWorldNormal(in FNDILandscapeContext DIContext, float3 InWorldPos, out float3 OutNormal, out bool OutIsValid)
 {
+	// todo[mg] convert world pos to LWC pos
 	if (DIContext.NormalVirtualTextureEnabled)
 	{
 		float3 UnpackNormalOrigin = DIContext.NormalVirtualTextureWorldToUvTransform[0].xyz;
@@ -217,6 +221,7 @@ void NDILandscape_GetWorldNormal(in FNDILandscapeContext DIContext, float3 InWor
 
 void NDILandscape_GetPhysicalMaterialIndex(in FNDILandscapeContext DIContext, float3 InWorldPos, out int OutIndex, out bool OutIsValid)
 {
+	// todo[mg] convert world pos to LWC pos
 	float2 Uv = mul(float4(InWorldPos, 1.0f), DIContext.CachedHeightTextureWorldToUvTransform).xy;
 	Uv = Uv * DIContext.CachedHeightTextureUvScaleBias.xy + DIContext.CachedHeightTextureUvScaleBias.zw;
 
@@ -274,7 +279,8 @@ void NDILandscape_GetPhysicalMaterialIndex(in FNDILandscapeContext DIContext, fl
 	float2				CachedHeightTextureGridSize_ ## NAME; \
 	SamplerState		PointClampedSampler_ ## NAME; \
 	Texture2D<int>		CachedPhysMatTexture_ ## NAME; \
-	int2				CachedPhysMatTextureDimension_ ## NAME;
+	int2				CachedPhysMatTextureDimension_ ## NAME; \
+	float3				SystemLWCTile_ ## NAME;
 
 #define NDILANDSCAPE_MAKE_CONTEXT(NAME) \
 	FNDILandscapeContext DIContext; \
@@ -315,4 +321,5 @@ void NDILandscape_GetPhysicalMaterialIndex(in FNDILandscapeContext DIContext, fl
 	DIContext.CachedHeightTextureGridSize = CachedHeightTextureGridSize_ ## NAME; \
 	DIContext.PointClampedSampler = PointClampedSampler_ ## NAME; \
 	DIContext.CachedPhysMatTexture = CachedPhysMatTexture_ ## NAME; \
-	DIContext.CachedPhysMatTextureDimension = CachedPhysMatTextureDimension_ ## NAME;
+	DIContext.CachedPhysMatTextureDimension = CachedPhysMatTextureDimension_ ## NAME; \
+	DIContext.SystemLWCTile = SystemLWCTile_ ## NAME;
diff --git a/Engine/Plugins/FX/Niagara/Shaders/Private/NiagaraDataInterfaceOcclusion.ush b/Engine/Plugins/FX/Niagara/Shaders/Private/NiagaraDataInterfaceOcclusion.ush
--- a/Engine/Plugins/FX/Niagara/Shaders/Private/NiagaraDataInterfaceOcclusion.ush
+++ b/Engine/Plugins/FX/Niagara/Shaders/Private/NiagaraDataInterfaceOcclusion.ush
@@ -4,24 +4,31 @@
 
 #include "/Engine/Private/DeferredShadingCommon.ush"
 
-void DIOcclusion_Rectangle(in float3 In_SampleCenterWorldPos, in float In_SampleWindowWidthWorld, in float In_SampleWindowHeightWorld, in float In_SampleSteps, out float Out_VisibilityFraction, out float Out_SampleFraction)
+struct FNDIOcclusionContext
 {
-	// we need to project the camera vector to a plane to compare it with scene depth
-	float CameraDistance = abs(dot(In_SampleCenterWorldPos.xyz - LWCHackToFloat(PrimaryView.WorldViewOrigin), View.ViewForward.xyz));
-	float4 SamplePosition = float4(In_SampleCenterWorldPos + PrimaryView.PreViewTranslation, 1);
-	float4 ClipPosition = mul(SamplePosition, View.TranslatedWorldToClip);
+	float3 SystemLWCTile;
+};
+
+void DIOcclusion_Rectangle(in FNDIOcclusionContext DIContext, in float3 In_SampleCenterWorldPos, in float In_SampleWindowWidthWorld, in float In_SampleWindowHeightWorld, in float In_SampleSteps, out float Out_VisibilityFraction, out float Out_SampleFraction)
+{
+	FLWCVector3 LwcSamplePos = MakeLWCVector3(DIContext.SystemLWCTile, In_SampleCenterWorldPos);
+	float CameraDistance = LWCToFloat(LWCAbs(LWCDot(LWCSubtract(LwcSamplePos, PrimaryView.WorldViewOrigin), View.ViewForward.xyz)));
+
+	float3 ResolvedLwcSamplePos = LWCToFloat(LWCAdd(LwcSamplePos, PrimaryView.PreViewTranslation));
+	float4 SamplePosition = float4(ResolvedLwcSamplePos, 1);
+	float4 ClipPosition = mul(SamplePosition, ResolvedView.TranslatedWorldToClip);
 	float2 ScreenPosition = ClipPosition.xy / ClipPosition.w;
-	float2 ScreenUV = ScreenPosition * View.ScreenPositionScaleBias.xy + View.ScreenPositionScaleBias.wz;
+	float2 ScreenUV = ScreenPosition * ResolvedView.ScreenPositionScaleBias.xy + ResolvedView.ScreenPositionScaleBias.wz;
 
 	float Steps = In_SampleSteps <= 1 ? 0 : In_SampleSteps;
 	float TotalSamples = 0;
 	float OccludedSamples = 0;
 
-	float4 SampleWidthClip = mul(float4(View.ViewRight * In_SampleWindowWidthWorld, 0) + SamplePosition, View.TranslatedWorldToClip);
-	float4 SampleHeightClip = mul(float4(View.ViewUp * In_SampleWindowHeightWorld, 0) + SamplePosition, View.TranslatedWorldToClip);
+	float4 SampleWidthClip = mul(float4(ResolvedView.ViewRight * In_SampleWindowWidthWorld, 0) + SamplePosition, ResolvedView.TranslatedWorldToClip);
+	float4 SampleHeightClip = mul(float4(ResolvedView.ViewUp * In_SampleWindowHeightWorld, 0) + SamplePosition, ResolvedView.TranslatedWorldToClip);
 				
-	float2 SampleWidthUV = SampleWidthClip.xy / SampleWidthClip.w * View.ScreenPositionScaleBias.xy + View.ScreenPositionScaleBias.wz;
-	float2 SampleHeightUV = SampleHeightClip.xy / SampleHeightClip.w * View.ScreenPositionScaleBias.xy + View.ScreenPositionScaleBias.wz;
+	float2 SampleWidthUV = SampleWidthClip.xy / SampleWidthClip.w * ResolvedView.ScreenPositionScaleBias.xy + ResolvedView.ScreenPositionScaleBias.wz;
+	float2 SampleHeightUV = SampleHeightClip.xy / SampleHeightClip.w * ResolvedView.ScreenPositionScaleBias.xy + ResolvedView.ScreenPositionScaleBias.wz;
 				
 	float SampleWidth = ScreenUV.x > 1 ? 0 : SampleWidthUV.x - ScreenUV.x;
 	float SampleHeight = ScreenUV.y > 1 ? 0 : SampleHeightUV.y - ScreenUV.y;
@@ -56,15 +63,18 @@ void DIOcclusion_Rectangle(in float3 In_SampleCenterWorldPos, in float In_Sample
 	Out_SampleFraction = Steps == 0 ? 0 : (TotalSamples / (Steps * Steps));
 }
 
-void DIOcclusion_Circle(in float3 In_SampleCenterWorldPos, in float In_SampleWindowDiameterWorld, in float In_SampleRays, in float In_SampleStepsPerRay, out float Out_VisibilityFraction, out float Out_SampleFraction)
+void DIOcclusion_Circle(in FNDIOcclusionContext DIContext, in float3 In_SampleCenterWorldPos, in float In_SampleWindowDiameterWorld, in float In_SampleRays, in float In_SampleStepsPerRay, out float Out_VisibilityFraction, out float Out_SampleFraction)
 {
 	const float PI = 3.14159265;
 	const float SPIRAL_TURN = 2 * PI * 0.61803399; // use golden ratio to rotate sample pattern each ring so we get a spiral
-	float CameraDistance = abs(dot(In_SampleCenterWorldPos.xyz - LWCHackToFloat(PrimaryView.WorldViewOrigin), View.ViewForward.xyz)); // we need to project the camera vector to a plane to compare it with scene depth
-	float4 SamplePosition = float4(In_SampleCenterWorldPos + PrimaryView.PreViewTranslation, 1);
-	float4 ClipPosition = mul(SamplePosition, View.TranslatedWorldToClip);
+	FLWCVector3 LwcSamplePos = MakeLWCVector3(DIContext.SystemLWCTile, In_SampleCenterWorldPos);
+	float CameraDistance = LWCToFloat(LWCAbs(LWCDot(LWCSubtract(LwcSamplePos, PrimaryView.WorldViewOrigin), View.ViewForward.xyz)));
+
+	float3 ResolvedLwcSamplePos = LWCToFloat(LWCAdd(LwcSamplePos, PrimaryView.PreViewTranslation));
+	float4 SamplePosition = float4(ResolvedLwcSamplePos, 1);
+	float4 ClipPosition = mul(SamplePosition, ResolvedView.TranslatedWorldToClip);
 	float2 ScreenPosition = ClipPosition.xy / ClipPosition.w;
-	float2 ScreenUV = ScreenPosition * View.ScreenPositionScaleBias.xy + View.ScreenPositionScaleBias.wz;
+	float2 ScreenUV = ScreenPosition * ResolvedView.ScreenPositionScaleBias.xy + ResolvedView.ScreenPositionScaleBias.wz;
 
 	float Rays = In_SampleRays <= 1 ? 0 : In_SampleRays;
 	float Steps = In_SampleStepsPerRay < 1 ? 0 : In_SampleStepsPerRay;
@@ -90,9 +100,9 @@ void DIOcclusion_Circle(in float3 In_SampleCenterWorldPos, in float In_SampleWin
 			for (int ray = 0; ray < Rays; ray++)
 			{
 				// calc ray direction vector
-				float3 RayDirection = cos(Degrees) * View.ViewUp + sin(Degrees) * View.ViewRight;
-				float4 RayClip = mul(float4(RayDirection * In_SampleWindowDiameterWorld / 2, 0) + SamplePosition, View.TranslatedWorldToClip);
-				float2 RayUV = RayClip.xy / RayClip.w * View.ScreenPositionScaleBias.xy + View.ScreenPositionScaleBias.wz;
+				float3 RayDirection = cos(Degrees) * ResolvedView.ViewUp + sin(Degrees) * ResolvedView.ViewRight;
+				float4 RayClip = mul(float4(RayDirection * In_SampleWindowDiameterWorld / 2, 0) + SamplePosition, ResolvedView.TranslatedWorldToClip);
+				float2 RayUV = RayClip.xy / RayClip.w * ResolvedView.ScreenPositionScaleBias.xy + ResolvedView.ScreenPositionScaleBias.wz;
 
 				if ((ScreenUV.x > 1 && RayUV.x < 0) || (ScreenUV.y > 1 && RayUV.y < 0) || (ScreenUV.x < 0 && RayUV.x > 1) || (ScreenUV.y < 0 && RayUV.y > 1))
 				{
@@ -119,3 +129,12 @@ void DIOcclusion_Circle(in float3 In_SampleCenterWorldPos, in float In_SampleWin
 	Out_VisibilityFraction = TotalSamples > 0 ? 1 - OccludedSamples / TotalSamples : 0;
 	Out_SampleFraction = Steps == 0 ? 0 : (TotalSamples / (Rays * Steps + 1));
 }
+
+////////////////////////////////////////////////////////////////////////////////////////////////////////////////////////////////////////////////////////////////
+
+#define NDIOCCLUSION_DECLARE_CONSTANTS(NAME) \
+	float3				SystemLWCTile_ ## NAME;
+
+#define NDIOCCLUSION_MAKE_CONTEXT(NAME) \
+	FNDIOcclusionContext DIContext; \
+	DIContext.SystemLWCTile = SystemLWCTile_ ## NAME;
diff --git a/Engine/Plugins/FX/Niagara/Shaders/Private/NiagaraEmitterInstanceShader.usf b/Engine/Plugins/FX/Niagara/Shaders/Private/NiagaraEmitterInstanceShader.usf
--- a/Engine/Plugins/FX/Niagara/Shaders/Private/NiagaraEmitterInstanceShader.usf
+++ b/Engine/Plugins/FX/Niagara/Shaders/Private/NiagaraEmitterInstanceShader.usf
@@ -93,10 +93,12 @@ bool NiagaraAny(bool4 b) { return b.x || b.y || b.z || b.w; }
 	float Modulo(float x, float y){ return fmod(x,y); }
 
 	// utility function used for scene depth calculations
-	float3 WorldPositionFromSceneDepth(float2 ScreenPosition, float SceneDepth)
+	FLWCVector3 WorldPositionFromSceneDepth(float2 ScreenPosition, float SceneDepth)
 	{
-		float4 HomogeneousWorldPosition = mul(float4(ScreenPosition * SceneDepth, SceneDepth, 1), LWCHackToFloat(PrimaryView.ScreenToWorld));
-		return HomogeneousWorldPosition.xyz / HomogeneousWorldPosition.w;
+		FLWCVector4 HomogeneousWorldPosition = LWCMultiply(float4(ScreenPosition * SceneDepth, SceneDepth, 1), PrimaryView.ScreenToWorld);
+		FLWCVector3 WorldPosition = MakeLWCVector3(HomogeneousWorldPosition.Tile.xyz, HomogeneousWorldPosition.Offset.xyz);
+		FLWCScalar Scale = MakeLWCScalar(HomogeneousWorldPosition.Tile.w, HomogeneousWorldPosition.Offset.w);
+		return LWCDivide(WorldPosition, LWCToFloat(Scale));
 	}
 
 	// MASSIVE HACK - Tracked in JIRA UE-69298
diff --git a/Engine/Plugins/FX/Niagara/Shaders/Private/NiagaraMeshVertexFactory.ush b/Engine/Plugins/FX/Niagara/Shaders/Private/NiagaraMeshVertexFactory.ush
--- a/Engine/Plugins/FX/Niagara/Shaders/Private/NiagaraMeshVertexFactory.ush
+++ b/Engine/Plugins/FX/Niagara/Shaders/Private/NiagaraMeshVertexFactory.ush
@@ -66,16 +66,6 @@ float3 WorldToLocalVec(float3 WorldSpaceVec, float4x4 LocalToWorld)
 	return mul(WorldSpaceVec, InvRot);
 }
 
-float3 SimToWorldVec(float3 SimSpaceVec, float4x4 LocalToWorld)
-{
-	if (NiagaraMeshVF.bLocalSpace)
-	{
-		return LocalToWorldVec(SimSpaceVec, LocalToWorld);
-	}
-	
-	return SimSpaceVec;
-}
-
 float3 WorldToSimVec(float3 InVector, float4x4 LocalToWorld)
 {
 	if (NiagaraMeshVF.bLocalSpace)
@@ -600,7 +590,8 @@ void GetParticleSRT(out float3 Scale, out float3x3 Rotation, out float3 Translat
 	in int RotationDataOffset, in float4 DefaultRotation, in int PositionDataOffset, in float3 DefaultPos, in int CamOffsetDataOffset, in float DefaultCamOffset,
 	in float4 VelocityDirMag, in float3 CamViewPos, in float3 CamViewForward, in float3 CamViewUp, in float4x4 LocalToWorld)
 {
-	float3 ParticlePosition = SafeGetVec3(PositionDataOffset, InstanceId, DefaultPos);
+	float3 SimPosition = SafeGetVec3(PositionDataOffset, InstanceId, DefaultPos);
+	float3 ParticlePosition = NiagaraMeshVF.bLocalSpace ? SimPosition : LWCHackToFloat(MakeLWCVector3(NiagaraMeshVF.SystemLWCTile, SimPosition));
 	float4 ParticleRotation = SafeGetVec4(RotationDataOffset, InstanceId, DefaultRotation);
 	float3 ParticleScale = SafeGetVec3(ScaleDataOffset, InstanceId, DefaultScale);
 
diff --git a/Engine/Plugins/FX/Niagara/Shaders/Private/NiagaraRibbonVertexFactory.ush b/Engine/Plugins/FX/Niagara/Shaders/Private/NiagaraRibbonVertexFactory.ush
--- a/Engine/Plugins/FX/Niagara/Shaders/Private/NiagaraRibbonVertexFactory.ush
+++ b/Engine/Plugins/FX/Niagara/Shaders/Private/NiagaraRibbonVertexFactory.ush
@@ -63,12 +63,11 @@ float3 TransformPosition(float3 InPosition)
 {
 	if (NiagaraRibbonVF.bLocalSpace)
 	{
-		// TODO LWC
-		return LWCToFloat(TransformLocalToWorld(InPosition));
+		return LWCHackToFloat(TransformLocalToWorld(InPosition));
 	}
 	else
 	{
-		return InPosition;
+		return LWCHackToFloat(MakeLWCVector3(NiagaraRibbonVF.SystemLWCTile, InPosition));
 	}
 }
 
diff --git a/Engine/Plugins/FX/Niagara/Shaders/Private/NiagaraSortKeyGen.usf b/Engine/Plugins/FX/Niagara/Shaders/Private/NiagaraSortKeyGen.usf
--- a/Engine/Plugins/FX/Niagara/Shaders/Private/NiagaraSortKeyGen.usf
+++ b/Engine/Plugins/FX/Niagara/Shaders/Private/NiagaraSortKeyGen.usf
@@ -51,6 +51,7 @@ uint MeshIndexAttributeOffset;
 float2 CullDistanceRangeSquared;
 float4 LocalBoundingSphere;
 float3 CullingWorldSpaceOffset;
+float3 SystemLWCTile;
 
 RWBuffer<uint> OutKeys;
 RWBuffer<int> OutParticleIndices;
@@ -142,6 +143,7 @@ bool CullParticle(in uint ParticleIndex)
 	if (CullPositionAttributeOffset != -1)
 	{
 		Pos = GetVec3(CullPositionAttributeOffset, ParticleIndex);
+		Pos = LWCHackToFloat(MakeLWCVector3(SystemLWCTile, Pos));
 	}
 
 	float4 Quat = { 0, 0, 0, 1 };
diff --git a/Engine/Plugins/FX/Niagara/Shaders/Private/NiagaraSpriteVertexFactory.ush b/Engine/Plugins/FX/Niagara/Shaders/Private/NiagaraSpriteVertexFactory.ush
--- a/Engine/Plugins/FX/Niagara/Shaders/Private/NiagaraSpriteVertexFactory.ush
+++ b/Engine/Plugins/FX/Niagara/Shaders/Private/NiagaraSpriteVertexFactory.ush
@@ -628,14 +628,14 @@ float3 LocalToWorldPos(float3 LocalSpacePos, float4x4 LocalToWorld)
 		LocalToWorld[3].xyz;
 }
 
-float3 SimToWorldPos(float3 SimSpacePos, float4x4 LocalToWorld)
+float3 SimToWorldPos(float3 SimSpacePos, FLWCMatrix LocalToWorld)
 {
 	if (NiagaraSpriteVF.bLocalSpace)
 	{
-		return LocalToWorldPos(SimSpacePos, LocalToWorld);
+		return LocalToWorldPos(SimSpacePos, LWCHackToFloat(LocalToWorld));
 	}
 	
-	return SimSpacePos;
+	return LWCHackToFloat(MakeLWCVector3(NiagaraSpriteVF.SystemLWCTile, SimSpacePos));
 }
 
 float3 LocalToWorldVec(float3 LocalSpaceVec, float3x3 LocalToWorld)
@@ -663,12 +663,14 @@ float3 SimToWorldVec(float3 SimSpaceVec, float3x3 LocalToWorld)
 
 float3 GetNiagaraParticlePosition(uint InstanceID)
 {
-	return SimToWorldPos(SafeGetVec3(NiagaraSpriteVF.PositionDataOffset, InstanceID, NiagaraSpriteVF.DefaultPos.xyz), LWCHackToFloat(GetPrimitiveDataFromUniformBuffer().LocalToWorld));
+	float3 ParticlePos = SafeGetVec3(NiagaraSpriteVF.PositionDataOffset, InstanceID, NiagaraSpriteVF.DefaultPos.xyz);
+	return SimToWorldPos(ParticlePos, GetPrimitiveDataFromUniformBuffer().LocalToWorld);
 }
 
 float3 GetNiagaraPreviousParticlePosition(uint InstanceID)
 {
-	return SimToWorldPos(SafeGetVec3(NiagaraSpriteVF.PrevPositionDataOffset, InstanceID, NiagaraSpriteVF.DefaultPrevPos.xyz), LWCHackToFloat(GetPrimitiveDataFromUniformBuffer().PreviousLocalToWorld));
+	float3 PrevParticlePos = SafeGetVec3(NiagaraSpriteVF.PrevPositionDataOffset, InstanceID, NiagaraSpriteVF.DefaultPrevPos.xyz);
+	return SimToWorldPos(PrevParticlePos, GetPrimitiveDataFromUniformBuffer().PreviousLocalToWorld);
 }
 
 float3 GetNiagaraParticleVelocity(uint InstanceID)
diff --git a/Engine/Plugins/FX/Niagara/Source/Niagara/Classes/NiagaraCollision.h b/Engine/Plugins/FX/Niagara/Source/Niagara/Classes/NiagaraCollision.h
--- a/Engine/Plugins/FX/Niagara/Source/Niagara/Classes/NiagaraCollision.h
+++ b/Engine/Plugins/FX/Niagara/Source/Niagara/Classes/NiagaraCollision.h
@@ -87,7 +87,6 @@ public:
 		CurrBuffer = 0;
 	}
 
-	int32 SubmitQuery(FVector Position, FVector Direction, float CollisionSize, float DeltaSeconds);
 	int32 SubmitQuery(FVector StartPos, FVector EndPos, ECollisionChannel TraceChannel);
 	bool PerformQuery(FVector StartPos, FVector EndPos, FNiagaraDICollsionQueryResult &Result, ECollisionChannel TraceChannel);
 	bool GetQueryResult(uint32 TraceID, FNiagaraDICollsionQueryResult &Result);
@@ -103,11 +102,11 @@ private:
 
 	FRWLock CollisionTraceLock;
 	TArray<FNiagaraCollisionEventPayload> CollisionEvents;
-	FNiagaraDataSet *CollisionEventDataSet;
+	FNiagaraDataSet *CollisionEventDataSet = nullptr;
 
 	FNiagaraSystemInstanceID BatchID;
 	TArray<FNiagaraCollisionTrace> CollisionTraces[2];
 	TArray<FNiagaraDICollsionQueryResult> CollisionResults;
 	uint32 CurrBuffer;
-	UWorld *CollisionWorld;
+	UWorld *CollisionWorld = nullptr;
 };
\ No newline at end of file
diff --git a/Engine/Plugins/FX/Niagara/Source/Niagara/Classes/NiagaraConstants.h b/Engine/Plugins/FX/Niagara/Source/Niagara/Classes/NiagaraConstants.h
--- a/Engine/Plugins/FX/Niagara/Source/Niagara/Classes/NiagaraConstants.h
+++ b/Engine/Plugins/FX/Niagara/Source/Niagara/Classes/NiagaraConstants.h
@@ -44,6 +44,7 @@
 #define SYS_PARAM_ENGINE_Z_AXIS                          INiagaraModule::GetVar_Engine_Owner_ZAxis()
 #define SYS_PARAM_ENGINE_ROTATION                        INiagaraModule::GetVar_Engine_Owner_Rotation()
 #define SYS_PARAM_ENGINE_SCALE                           INiagaraModule::GetVar_Engine_Owner_Scale()
+#define SYS_PARAM_ENGINE_LWC_TILE                        INiagaraModule::GetVar_Engine_Owner_LWC_Tile()
 
 #define SYS_PARAM_ENGINE_LOCAL_TO_WORLD                  INiagaraModule::GetVar_Engine_Owner_SystemLocalToWorld()
 #define SYS_PARAM_ENGINE_WORLD_TO_LOCAL                  INiagaraModule::GetVar_Engine_Owner_SystemWorldToLocal()
@@ -146,6 +147,7 @@ struct NIAGARA_API FNiagaraConstants
 	static FNiagaraVariable UpdateEngineConstant(const FNiagaraVariable& InVar);
 	static const FNiagaraVariable *FindEngineConstant(const FNiagaraVariable& InVar);
 	static FText GetEngineConstantDescription(const FNiagaraVariable& InVar);
+	static const TArray<FNiagaraVariable>& GetOldPositionTypeVariables();
 
 	static const TArray<FNiagaraVariable>& GetCommonParticleAttributes();
 	static FText GetAttributeDescription(const FNiagaraVariable& InVar);
@@ -218,6 +220,7 @@ private:
 	static TArray<FNiagaraVariable> SystemParameters;
 	static TArray<FNiagaraVariable> TranslatorParameters;
 	static TArray<FNiagaraVariable> SwitchParameters;
+	static TArray<FNiagaraVariable> OldPositionTypes;
 	static TMap<FName, FNiagaraVariable> UpdatedSystemParameters;
 	static TMap<FNiagaraVariable, FText> SystemStrMap;
 	static TArray<FNiagaraVariable> Attributes;
diff --git a/Engine/Plugins/FX/Niagara/Source/Niagara/Classes/NiagaraDataInterface.h b/Engine/Plugins/FX/Niagara/Source/Niagara/Classes/NiagaraDataInterface.h
--- a/Engine/Plugins/FX/Niagara/Source/Niagara/Classes/NiagaraDataInterface.h
+++ b/Engine/Plugins/FX/Niagara/Source/Niagara/Classes/NiagaraDataInterface.h
@@ -787,6 +787,15 @@ struct FNDIInputParam<FVector3f>
 	FORCEINLINE bool IsConstant() const { return X.IsConstant() && Y.IsConstant() && Z.IsConstant(); }
 };
 
+template<>
+struct FNDIInputParam<FNiagaraPosition>
+{
+	VectorVM::FExternalFuncInputHandler<float> X;
+	VectorVM::FExternalFuncInputHandler<float> Y;
+	VectorVM::FExternalFuncInputHandler<float> Z;
+	FNDIInputParam(FVectorVMExternalFunctionContext& Context) : X(Context), Y(Context), Z(Context) {}
+	FORCEINLINE FNiagaraPosition GetAndAdvance() { return FNiagaraPosition(X.GetAndAdvance(), Y.GetAndAdvance(), Z.GetAndAdvance()); }
+};
 template<>
 struct FNDIInputParam<FVector4f>
 {
@@ -893,6 +902,21 @@ struct FNDIOutputParam<FVector3f>
 	}
 };
 
+template<>
+struct FNDIOutputParam<FNiagaraPosition>
+{
+	VectorVM::FExternalFuncRegisterHandler<float> X;
+	VectorVM::FExternalFuncRegisterHandler<float> Y;
+	VectorVM::FExternalFuncRegisterHandler<float> Z;
+	FNDIOutputParam(FVectorVMExternalFunctionContext& Context) : X(Context), Y(Context), Z(Context) {}
+	FORCEINLINE bool IsValid() const { return X.IsValid() || Y.IsValid() || Z.IsValid(); }
+	FORCEINLINE void SetAndAdvance(FNiagaraPosition Val)
+	{
+		*X.GetDestAndAdvance() = Val.X;
+		*Y.GetDestAndAdvance() = Val.Y;
+		*Z.GetDestAndAdvance() = Val.Z;
+	}
+};
 
 template<>
 struct FNDIOutputParam<FVector4f>
diff --git a/Engine/Plugins/FX/Niagara/Source/Niagara/Classes/NiagaraDataInterfaceArrayFloat.h b/Engine/Plugins/FX/Niagara/Source/Niagara/Classes/NiagaraDataInterfaceArrayFloat.h
--- a/Engine/Plugins/FX/Niagara/Source/Niagara/Classes/NiagaraDataInterfaceArrayFloat.h
+++ b/Engine/Plugins/FX/Niagara/Source/Niagara/Classes/NiagaraDataInterfaceArrayFloat.h
@@ -107,6 +107,15 @@ struct FNDIArrayImplHelper<FVector3d> : public FNDIArrayImplHelper<FVector3f>
 	}
 };
 
+template<>
+struct FNDIArrayImplHelper<FNiagaraPosition> : public FNDIArrayImplHelper<FVector3f>
+{
+	typedef FNiagaraPosition TVMArrayType;
+
+	static const FNiagaraTypeDefinition& GetTypeDefinition() { return FNiagaraTypeDefinition::GetPositionDef(); }
+	static const FNiagaraPosition GetDefaultValue() { return FVector3f::ZeroVector; }
+};
+
 template<>
 struct FNDIArrayImplHelper<FVector4f> : public FNDIArrayImplHelperBase<FVector4f>
 {
@@ -240,6 +249,17 @@ public:
 	TArray<FVector>& GetArrayReference() { return FloatData; }
 };
 
+UCLASS(EditInlineNew, Category = "Array", meta = (DisplayName = "Position Array", Experimental), Blueprintable, BlueprintType)
+class NIAGARA_API UNiagaraDataInterfaceArrayPosition : public UNiagaraDataInterfaceArray
+{
+	GENERATED_UCLASS_BODY()
+public:
+	UPROPERTY(VisibleAnywhere, Category = "Array")
+	TArray<FNiagaraPosition> PositionData;
+
+	TArray<FNiagaraPosition>& GetArrayReference() { return PositionData; }
+};
+
 UCLASS(EditInlineNew, Category = "Array", meta = (DisplayName = "Vector 4 Array", Experimental), Blueprintable, BlueprintType)
 class NIAGARA_API UNiagaraDataInterfaceArrayFloat4 : public UNiagaraDataInterfaceArray
 {
diff --git a/Engine/Plugins/FX/Niagara/Source/Niagara/Classes/NiagaraDataInterfaceArrayFunctionLibrary.h b/Engine/Plugins/FX/Niagara/Source/Niagara/Classes/NiagaraDataInterfaceArrayFunctionLibrary.h
--- a/Engine/Plugins/FX/Niagara/Source/Niagara/Classes/NiagaraDataInterfaceArrayFunctionLibrary.h
+++ b/Engine/Plugins/FX/Niagara/Source/Niagara/Classes/NiagaraDataInterfaceArrayFunctionLibrary.h
@@ -23,6 +23,9 @@ public:
 	/** Sets Niagara Array FVector Data. */
 	UFUNCTION(BlueprintCallable, Category = Niagara, meta = (DisplayName = "Niagara Set Vector Array"))
 	static void SetNiagaraArrayVector(UNiagaraComponent* NiagaraSystem, FName OverrideName, const TArray<FVector>& ArrayData);
+	/** Sets Niagara Array FVector Data. */
+	UFUNCTION(BlueprintCallable, Category = Niagara, meta = (DisplayName = "Niagara Set Vector Array"))
+	static void SetNiagaraArrayPosition(UNiagaraComponent* NiagaraSystem, FName OverrideName, const TArray<FVector>& ArrayData);
 	/** Sets Niagara Array FVector4 Data. */
 	UFUNCTION(BlueprintCallable, Category = Niagara, meta = (DisplayName = "Niagara Set Vector 4 Array"))
 	static void SetNiagaraArrayVector4(UNiagaraComponent* NiagaraSystem, FName OverrideName, const TArray<FVector4>& ArrayData);
@@ -48,6 +51,9 @@ public:
 	/** Gets a copy of Niagara FVector Data. */
 	UFUNCTION(BlueprintCallable, Category = Niagara, meta = (DisplayName = "Niagara Get Vector Array"))
 	static TArray<FVector> GetNiagaraArrayVector(UNiagaraComponent* NiagaraSystem, FName OverrideName);
+	/** Gets a copy of Niagara Position Data. */
+	UFUNCTION(BlueprintCallable, Category = Niagara, meta = (DisplayName = "Niagara Get Position Array"))
+	static TArray<FVector> GetNiagaraArrayPosition(UNiagaraComponent* NiagaraSystem, FName OverrideName);
 	/** Gets a copy of Niagara FVector4 Data. */
 	UFUNCTION(BlueprintCallable, Category = Niagara, meta = (DisplayName = "Niagara Get Vector 4 Array"))
 	static TArray<FVector4> GetNiagaraArrayVector4(UNiagaraComponent* NiagaraSystem, FName OverrideName);
@@ -73,6 +79,8 @@ public:
 	/** Sets a single value within a Niagara Array FVector. */
 	UFUNCTION(BlueprintCallable, Category = Niagara, meta = (DisplayName = "Niagara Set Vector Array Value"))
 	static void SetNiagaraArrayVectorValue(UNiagaraComponent* NiagaraSystem, FName OverrideName, int Index, const FVector& Value, bool bSizeToFit);
+	UFUNCTION(BlueprintCallable, Category = Niagara, meta = (DisplayName = "Niagara Set Position Array Value"))
+	static void SetNiagaraArrayPositionValue(UNiagaraComponent* NiagaraSystem, FName OverrideName, int Index, const FVector& Value, bool bSizeToFit);
 	/** Sets a single value within a Niagara Array FVector4. */
 	UFUNCTION(BlueprintCallable, Category = Niagara, meta = (DisplayName = "Niagara Set Vector 4 Array Value"))
 	static void SetNiagaraArrayVector4Value(UNiagaraComponent* NiagaraSystem, FName OverrideName, int Index, const FVector4& Value, bool bSizeToFit);
@@ -98,6 +106,9 @@ public:
 	/** Gets a single value within a Niagara Array FVector. */
 	UFUNCTION(BlueprintCallable, Category = Niagara, meta = (DisplayName = "Niagara Get Vector Array Value"))
 	static FVector GetNiagaraArrayVectorValue(UNiagaraComponent* NiagaraSystem, FName OverrideName, int Index);
+	/** Gets a single value within a Niagara Array Position. */
+	UFUNCTION(BlueprintCallable, Category = Niagara, meta = (DisplayName = "Niagara Get Position Array Value"))
+	static FVector GetNiagaraArrayPositionValue(UNiagaraComponent* NiagaraSystem, FName OverrideName, int Index);
 	/** Gets a single value within a Niagara Array FVector4. */
 	UFUNCTION(BlueprintCallable, Category = Niagara, meta = (DisplayName = "Niagara Get Vector 4 Array Value"))
 	static FVector4 GetNiagaraArrayVector4Value(UNiagaraComponent* NiagaraSystem, FName OverrideName, int Index);
diff --git a/Engine/Plugins/FX/Niagara/Source/Niagara/Classes/NiagaraDataInterfaceAudioPlayer.h b/Engine/Plugins/FX/Niagara/Source/Niagara/Classes/NiagaraDataInterfaceAudioPlayer.h
--- a/Engine/Plugins/FX/Niagara/Source/Niagara/Classes/NiagaraDataInterfaceAudioPlayer.h
+++ b/Engine/Plugins/FX/Niagara/Source/Niagara/Classes/NiagaraDataInterfaceAudioPlayer.h
@@ -41,6 +41,7 @@ struct FAudioPlayerInterface_InstanceData
 	TWeakObjectPtr<USoundConcurrency> Concurrency;
 	TArray<FName> ParameterNames;
 
+	FNiagaraLWCConverter LWCConverter;
 	int32 MaxPlaysPerTick = 0;
 	bool bStopWhenComponentIsDestroyed = true;
 #if WITH_EDITORONLY_DATA
diff --git a/Engine/Plugins/FX/Niagara/Source/Niagara/Classes/NiagaraDataInterfaceCamera.h b/Engine/Plugins/FX/Niagara/Source/Niagara/Classes/NiagaraDataInterfaceCamera.h
--- a/Engine/Plugins/FX/Niagara/Source/Niagara/Classes/NiagaraDataInterfaceCamera.h
+++ b/Engine/Plugins/FX/Niagara/Source/Niagara/Classes/NiagaraDataInterfaceCamera.h
@@ -9,14 +9,15 @@
 struct FDistanceData
 {
 	FNiagaraID ParticleID;
-	float DistanceSquared;
+	double DistanceSquared;
 };
 
 struct FCameraDataInterface_InstanceData
 {
-	FVector CameraLocation = FVector::ZeroVector;
+	FNiagaraPosition CameraLocation = FVector::ZeroVector;
 	FRotator CameraRotation = FRotator::ZeroRotator;
 	float CameraFOV = 0.0f;
+	FNiagaraLWCConverter LWCConverter;
 
 	TQueue<FDistanceData, EQueueMode::Mpsc> DistanceSortQueue;
 	TArray<FDistanceData> ParticlesSortedByDistance;	
@@ -28,6 +29,8 @@ class NIAGARA_API UNiagaraDataInterfaceCamera : public UNiagaraDataInterface
 	GENERATED_UCLASS_BODY()
 
 public:
+	DECLARE_NIAGARA_DI_PARAMETER();
+	
 	/** This is used to determine which camera position to query for cpu emitters. If no valid index is supplied, the first controller is used as camera reference. */
 	UPROPERTY(EditAnywhere, Category = "Camera")
 	int32 PlayerControllerIndex = 0;
@@ -51,6 +54,9 @@ public:
 #if WITH_EDITORONLY_DATA
 	virtual bool GetFunctionHLSL(const FNiagaraDataInterfaceGPUParamInfo& ParamInfo, const FNiagaraDataInterfaceGeneratedFunction& FunctionInfo, int FunctionInstanceIndex, FString& OutHLSL) override;
 	virtual bool UpgradeFunctionCall(FNiagaraFunctionSignature& FunctionSignature) override;
+	virtual bool AppendCompileHash(FNiagaraCompileHashVisitor* InVisitor) const override;
+	virtual void GetCommonHLSL(FString& OutHLSL) override;
+	virtual void GetParameterDefinitionHLSL(const FNiagaraDataInterfaceGPUParamInfo& ParamInfo, FString& OutHLSL) override;
 #endif
 	virtual bool CanExecuteOnTarget(ENiagaraSimTarget Target) const override { return true; }
 	virtual bool HasTickGroupPrereqs() const override { return true; }
@@ -68,10 +74,6 @@ public:
 	void GetClosestParticles(FVectorVMExternalFunctionContext& Context);
 	void GetCameraFOV(FVectorVMExternalFunctionContext& Context);
 	void GetCameraProperties(FVectorVMExternalFunctionContext& Context);
-	void GetViewPropertiesGPU(FVectorVMExternalFunctionContext& Context);
-	void GetClipSpaceTransformsGPU(FVectorVMExternalFunctionContext& Context);
-	void GetViewSpaceTransformsGPU(FVectorVMExternalFunctionContext& Context);
-	void GetTAAJitter(FVectorVMExternalFunctionContext& Context);
 
 protected:
 	virtual bool CopyToInternal(UNiagaraDataInterface* Destination) const override;
diff --git a/Engine/Plugins/FX/Niagara/Source/Niagara/Classes/NiagaraDataInterfaceCollisionQuery.h b/Engine/Plugins/FX/Niagara/Source/Niagara/Classes/NiagaraDataInterfaceCollisionQuery.h
--- a/Engine/Plugins/FX/Niagara/Source/Niagara/Classes/NiagaraDataInterfaceCollisionQuery.h
+++ b/Engine/Plugins/FX/Niagara/Source/Niagara/Classes/NiagaraDataInterfaceCollisionQuery.h
@@ -64,8 +64,6 @@ public:
 	// VM functions
 	void PerformQuerySyncCPU(FVectorVMExternalFunctionContext& Context);
 	void PerformQueryAsyncCPU(FVectorVMExternalFunctionContext& Context);
-	void QuerySceneDepth(FVectorVMExternalFunctionContext& Context);
-	void QueryMeshDistanceField(FVectorVMExternalFunctionContext& Context);
 
 	virtual bool CanExecuteOnTarget(ENiagaraSimTarget Target) const override { return true; }
 	virtual bool RequiresDistanceFieldData() const override { return true; }
diff --git a/Engine/Plugins/FX/Niagara/Source/Niagara/Classes/NiagaraDataInterfaceExport.h b/Engine/Plugins/FX/Niagara/Source/Niagara/Classes/NiagaraDataInterfaceExport.h
--- a/Engine/Plugins/FX/Niagara/Source/Niagara/Classes/NiagaraDataInterfaceExport.h
+++ b/Engine/Plugins/FX/Niagara/Source/Niagara/Classes/NiagaraDataInterfaceExport.h
@@ -4,7 +4,6 @@
 #include "NiagaraCommon.h"
 #include "NiagaraShared.h"
 #include "NiagaraDataInterface.h"
-#include "NiagaraParameterStore.h"
 #include "NiagaraDataInterfaceExport.generated.h"
 
 USTRUCT(BlueprintType, Blueprintable)
@@ -35,7 +34,7 @@ class INiagaraParticleCallbackHandler
 public:
 	/** This function is called once per tick with the gathered particle data. It will not be called if there is no particle data to call it with. */
 	UFUNCTION(BlueprintCallable, BlueprintNativeEvent, Category = "Niagara")
-	void ReceiveParticleData(const TArray<FBasicParticleData>& Data, UNiagaraSystem* NiagaraSystem);
+	void ReceiveParticleData(const TArray<FBasicParticleData>& Data, UNiagaraSystem* NiagaraSystem, const FVector& SimulationPositionOffset);
 };
 
 UENUM()
diff --git a/Engine/Plugins/FX/Niagara/Source/Niagara/Classes/NiagaraDataInterfaceMeshCommon.h b/Engine/Plugins/FX/Niagara/Source/Niagara/Classes/NiagaraDataInterfaceMeshCommon.h
--- a/Engine/Plugins/FX/Niagara/Source/Niagara/Classes/NiagaraDataInterfaceMeshCommon.h
+++ b/Engine/Plugins/FX/Niagara/Source/Niagara/Classes/NiagaraDataInterfaceMeshCommon.h
@@ -21,24 +21,23 @@ struct FMeshTriCoordinate
 
 	FMeshTriCoordinate()
 	: Tri(INDEX_NONE)
-	, BaryCoord(FVector::ZeroVector)
+	, BaryCoord(FVector3f::ZeroVector)
 	{}
 
-	FMeshTriCoordinate(int32 InTri, FVector InBaryCoord)
+	FMeshTriCoordinate(int32 InTri, FVector3f InBaryCoord)
 		: Tri(InTri)
 		, BaryCoord(InBaryCoord)
 	{}
 };
 
-FORCEINLINE FVector RandomBarycentricCoord(FRandomStream& RandStream)
+FORCEINLINE FVector3f RandomBarycentricCoord(FRandomStream& RandStream)
 {
 	//TODO: This is gonna be slooooow. Move to an LUT possibly or find faster method.
 	//Can probably handle lower quality randoms / uniformity for a decent speed win.
 	float r0 = RandStream.GetFraction();
 	float r1 = RandStream.GetFraction();
 	float sqrt0 = FMath::Sqrt(r0);
-	float sqrt1 = FMath::Sqrt(r1);
-	return FVector(1.0f - sqrt0, sqrt0 * (1.0 - r1), r1 * sqrt0);
+	return FVector3f(1.0f - sqrt0, sqrt0 * (1.0 - r1), r1 * sqrt0);
 }
 
 template<typename T>
diff --git a/Engine/Plugins/FX/Niagara/Source/Niagara/Classes/NiagaraDataInterfaceOcclusion.h b/Engine/Plugins/FX/Niagara/Source/Niagara/Classes/NiagaraDataInterfaceOcclusion.h
--- a/Engine/Plugins/FX/Niagara/Source/Niagara/Classes/NiagaraDataInterfaceOcclusion.h
+++ b/Engine/Plugins/FX/Niagara/Source/Niagara/Classes/NiagaraDataInterfaceOcclusion.h
@@ -12,24 +12,25 @@ class NIAGARA_API UNiagaraDataInterfaceOcclusion : public UNiagaraDataInterface
 	GENERATED_UCLASS_BODY()
 
 public:
+
+	DECLARE_NIAGARA_DI_PARAMETER();
+	
 	//UObject Interface
 	virtual void PostInitProperties() override;
 	//UObject Interface End
 
 	//UNiagaraDataInterface Interface
 	virtual void GetFunctions(TArray<FNiagaraFunctionSignature>& OutFunctions)override;
-	virtual void GetVMExternalFunction(const FVMExternalFunctionBindingInfo& BindingInfo, void* InstanceData, FVMExternalFunction &OutFunc) override;
 #if WITH_EDITORONLY_DATA
 	virtual bool AppendCompileHash(FNiagaraCompileHashVisitor* InVisitor) const override;
 	virtual void GetCommonHLSL(FString& OutHLSL) override;
 	virtual bool GetFunctionHLSL(const FNiagaraDataInterfaceGPUParamInfo& ParamInfo, const FNiagaraDataInterfaceGeneratedFunction& FunctionInfo, int FunctionInstanceIndex, FString& OutHLSL) override;
+	virtual bool UpgradeFunctionCall(FNiagaraFunctionSignature& FunctionSignature) override;
+	virtual void GetParameterDefinitionHLSL(const FNiagaraDataInterfaceGPUParamInfo& ParamInfo, FString& OutHLSL) override;
 #endif
 	virtual bool CanExecuteOnTarget(ENiagaraSimTarget Target) const override { return Target == ENiagaraSimTarget::GPUComputeSim; }
 	virtual bool RequiresDepthBuffer() const override { return true; }
 	//UNiagaraDataInterface Interface
-
-	void QueryOcclusionFactorGPU(FVectorVMExternalFunctionContext& Context);
-	void QueryOcclusionFactorCircleGPU(FVectorVMExternalFunctionContext& Context);
 	
 private:
 	static const FName GetCameraOcclusionRectangleName;
diff --git a/Engine/Plugins/FX/Niagara/Source/Niagara/Classes/NiagaraDataInterfaceParticleRead.h b/Engine/Plugins/FX/Niagara/Source/Niagara/Classes/NiagaraDataInterfaceParticleRead.h
--- a/Engine/Plugins/FX/Niagara/Source/Niagara/Classes/NiagaraDataInterfaceParticleRead.h
+++ b/Engine/Plugins/FX/Niagara/Source/Niagara/Classes/NiagaraDataInterfaceParticleRead.h
@@ -61,6 +61,7 @@ public:
 	void ReadVector3(FVectorVMExternalFunctionContext& Context, FName AttributeToRead);
 	void ReadVector4(FVectorVMExternalFunctionContext& Context, FName AttributeToRead);
 	void ReadColor(FVectorVMExternalFunctionContext& Context, FName AttributeToRead);
+	void ReadPosition(FVectorVMExternalFunctionContext& Context, FName AttributeToRead);
 	void ReadQuat(FVectorVMExternalFunctionContext& Context, FName AttributeToRead);
 	void ReadID(FVectorVMExternalFunctionContext& Context, FName AttributeToRead);
 	void ReadIntByIndex(FVectorVMExternalFunctionContext& Context, FName AttributeToRead);
@@ -69,6 +70,7 @@ public:
 	void ReadVector2ByIndex(FVectorVMExternalFunctionContext& Context, FName AttributeToRead);
 	void ReadVector3ByIndex(FVectorVMExternalFunctionContext& Context, FName AttributeToRead);
 	void ReadVector4ByIndex(FVectorVMExternalFunctionContext& Context, FName AttributeToRead);
+	void ReadPositionByIndex(FVectorVMExternalFunctionContext& Context, FName AttributeToRead);
 	void ReadColorByIndex(FVectorVMExternalFunctionContext& Context, FName AttributeToRead);
 	void ReadQuatByIndex(FVectorVMExternalFunctionContext& Context, FName AttributeToRead);
 	void ReadIDByIndex(FVectorVMExternalFunctionContext& Context, FName AttributeToRead);
diff --git a/Engine/Plugins/FX/Niagara/Source/Niagara/Classes/NiagaraDataInterfaceSpline.h b/Engine/Plugins/FX/Niagara/Source/Niagara/Classes/NiagaraDataInterfaceSpline.h
--- a/Engine/Plugins/FX/Niagara/Source/Niagara/Classes/NiagaraDataInterfaceSpline.h
+++ b/Engine/Plugins/FX/Niagara/Source/Niagara/Classes/NiagaraDataInterfaceSpline.h
@@ -9,6 +9,8 @@
 #include "Niagara/Private/NiagaraStats.h"
 #include "NiagaraDataInterfaceSpline.generated.h"
 
+class UNiagaraDataInterfaceSpline;
+
 /** Proxy data for splines */
 struct FNDISpline_InstanceData_RenderThread
 {
@@ -118,6 +120,7 @@ struct FNDISpline_InstanceData
 	//InverseTranspose of above for transforming normals/tangents.
 	FMatrix TransformInverseTransposed;
 	FTransform ComponentTransform;
+	FNiagaraLWCConverter LwcConverter;
 
 	FVector DefaultUpVector;
 	FSplineCurves SplineCurves;
@@ -230,7 +233,7 @@ public:
 	void SampleSplineDirectionByUnitDistance(FVectorVMExternalFunctionContext& Context);
 	template<typename UseLUT, typename TransformHandlerType, typename SplineSampleType>
 	void SampleSplineTangentByUnitDistance(FVectorVMExternalFunctionContext& Context);
-	template<typename UseLUT, typename PosXType, typename PosYType, typename PosZType>
+	template<typename UseLUT>
 	void FindClosestUnitDistanceFromPositionWS(FVectorVMExternalFunctionContext& Context);
 	
 	void GetLocalToWorld(FVectorVMExternalFunctionContext& Context);
@@ -245,6 +248,7 @@ protected:
 	virtual bool CopyToInternal(UNiagaraDataInterface* Destination) const override;
 
 #if WITH_EDITORONLY_DATA
+	virtual bool UpgradeFunctionCall(FNiagaraFunctionSignature& FunctionSignature) override;
 	virtual void GetCommonHLSL(FString& OutHLSL) override;
 	virtual void GetParameterDefinitionHLSL(const FNiagaraDataInterfaceGPUParamInfo& ParamInfo, FString& OutHLSL) override;
 
diff --git a/Engine/Plugins/FX/Niagara/Source/Niagara/Classes/NiagaraDataSetAccessor.h b/Engine/Plugins/FX/Niagara/Source/Niagara/Classes/NiagaraDataSetAccessor.h
--- a/Engine/Plugins/FX/Niagara/Source/Niagara/Classes/NiagaraDataSetAccessor.h
+++ b/Engine/Plugins/FX/Niagara/Source/Niagara/Classes/NiagaraDataSetAccessor.h
@@ -197,8 +197,21 @@ struct FNiagaraDataSetAccessorFloat
 		{
 			bIsFloat = true;
 			ComponentIndex = DataSetCompiledData.VariableLayouts[FloatVariableIndex].FloatComponentStart;
+			return;
 		}
-		else if (FNiagaraDataSetAccessorTypeInfo<TType>::bSupportsHalf)
+
+		if (FNiagaraDataSetAccessorTypeInfo<TType>::bSupportsAlternateType)
+		{
+			const int32 AlternateVariableIndex = DataSetCompiledData.Variables.IndexOfByKey(FNiagaraVariableBase(FNiagaraDataSetAccessorTypeInfo<TType>::GetAlternateType(), VariableName));
+			if (AlternateVariableIndex != INDEX_NONE)
+			{
+				bIsFloat = true;
+				ComponentIndex = DataSetCompiledData.VariableLayouts[AlternateVariableIndex].FloatComponentStart;
+				return;
+			}
+		}
+		
+		if (FNiagaraDataSetAccessorTypeInfo<TType>::bSupportsHalf)
 		{
 			const int32 HalfVariableIndex = DataSetCompiledData.Variables.IndexOfByKey(FNiagaraVariableBase(FNiagaraDataSetAccessorTypeInfo<TType>::GetHalfType(), VariableName));
 			if (HalfVariableIndex != INDEX_NONE)
@@ -228,9 +241,11 @@ struct FNiagaraDataSetAccessorTypeInfo<float>
 	using TAccessorBaseClass = FNiagaraDataSetAccessorFloat<float>;
 
 	static constexpr bool bSupportsHalf = true;
+	static constexpr bool bSupportsAlternateType = false;
 	static constexpr int32 NumElements = 1;
 	static const FNiagaraTypeDefinition& GetFloatType() { return FNiagaraTypeDefinition::GetFloatDef(); }
 	static const FNiagaraTypeDefinition& GetHalfType() { return FNiagaraTypeDefinition::GetHalfDef(); }
+	static const FNiagaraTypeDefinition& GetAlternateType() { return FNiagaraTypeDefinition::GetFloatDef(); }
 };
 
 template<>
@@ -239,9 +254,11 @@ struct FNiagaraDataSetAccessorTypeInfo<FVector2f>
 	using TAccessorBaseClass = FNiagaraDataSetAccessorFloat<FVector2f>;
 
 	static constexpr bool bSupportsHalf = true;
+	static constexpr bool bSupportsAlternateType = false;
 	static constexpr int32 NumElements = 2;
 	static const FNiagaraTypeDefinition& GetFloatType() { return FNiagaraTypeDefinition::GetVec2Def(); }
 	static const FNiagaraTypeDefinition& GetHalfType() { return FNiagaraTypeDefinition::GetHalfVec2Def(); }
+	static const FNiagaraTypeDefinition& GetAlternateType() { return FNiagaraTypeDefinition::GetVec2Def(); }
 };
 
 template<>
@@ -250,12 +267,25 @@ struct FNiagaraDataSetAccessorTypeInfo<FVector3f>
 	using TAccessorBaseClass = FNiagaraDataSetAccessorFloat<FVector3f>;
 
 	static constexpr bool bSupportsHalf = true;
+	static constexpr bool bSupportsAlternateType = false;
 	static constexpr int32 NumElements = 3;
 	static const FNiagaraTypeDefinition& GetFloatType() { return FNiagaraTypeDefinition::GetVec3Def(); }
 	static const FNiagaraTypeDefinition& GetHalfType() { return FNiagaraTypeDefinition::GetHalfVec3Def(); }
+	static const FNiagaraTypeDefinition& GetAlternateType() { return FNiagaraTypeDefinition::GetVec3Def(); }
 };
 
+template<>
+struct FNiagaraDataSetAccessorTypeInfo<FNiagaraPosition>
+{
+	using TAccessorBaseClass = FNiagaraDataSetAccessorFloat<FNiagaraPosition>;
 
+	static constexpr bool bSupportsHalf = true;
+	static constexpr bool bSupportsAlternateType = true;
+	static constexpr int32 NumElements = 3;
+	static const FNiagaraTypeDefinition& GetFloatType() { return FNiagaraTypeDefinition::GetPositionDef(); }
+	static const FNiagaraTypeDefinition& GetHalfType() { return FNiagaraTypeDefinition::GetHalfVec3Def(); }
+	static const FNiagaraTypeDefinition& GetAlternateType() { return FNiagaraTypeDefinition::GetVec3Def(); }
+};
 
 template<>
 struct FNiagaraDataSetAccessorTypeInfo<FVector4f>
@@ -263,9 +293,11 @@ struct FNiagaraDataSetAccessorTypeInfo<FVector4f>
 	using TAccessorBaseClass = FNiagaraDataSetAccessorFloat<FVector4f>;
 
 	static constexpr bool bSupportsHalf = true;
+	static constexpr bool bSupportsAlternateType = false;
 	static constexpr int32 NumElements = 4;
 	static const FNiagaraTypeDefinition& GetFloatType() { return FNiagaraTypeDefinition::GetVec4Def(); }
 	static const FNiagaraTypeDefinition& GetHalfType() { return FNiagaraTypeDefinition::GetHalfVec4Def(); }
+	static const FNiagaraTypeDefinition& GetAlternateType() { return FNiagaraTypeDefinition::GetVec4Def(); }
 };
 
 template<>
@@ -274,9 +306,11 @@ struct FNiagaraDataSetAccessorTypeInfo<FLinearColor>
 	using TAccessorBaseClass = FNiagaraDataSetAccessorFloat<FLinearColor>;
 
 	static constexpr bool bSupportsHalf = false;
+	static constexpr bool bSupportsAlternateType = false;
 	static constexpr int32 NumElements = 4;
 	static const FNiagaraTypeDefinition& GetFloatType() { return FNiagaraTypeDefinition::GetColorDef(); }
 	static const FNiagaraTypeDefinition& GetHalfType() { check(false); return FNiagaraTypeDefinition::GetHalfVec4Def(); }
+	static const FNiagaraTypeDefinition& GetAlternateType() { return FNiagaraTypeDefinition::GetColorDef(); }
 };
 
 template<>
@@ -285,9 +319,11 @@ struct FNiagaraDataSetAccessorTypeInfo<FQuat4f>
 	using TAccessorBaseClass = FNiagaraDataSetAccessorFloat<FQuat4f>;
 
 	static constexpr bool bSupportsHalf = false;
+	static constexpr bool bSupportsAlternateType = false;
 	static constexpr int32 NumElements = 4;
 	static const FNiagaraTypeDefinition& GetFloatType() { return FNiagaraTypeDefinition::GetQuatDef(); }
 	static const FNiagaraTypeDefinition& GetHalfType() { check(false); return FNiagaraTypeDefinition::GetQuatDef(); }
+	static const FNiagaraTypeDefinition& GetAlternateType() { return FNiagaraTypeDefinition::GetQuatDef(); }
 };
 
 //////////////////////////////////////////////////////////////////////////
diff --git a/Engine/Plugins/FX/Niagara/Source/Niagara/Classes/NiagaraEmitter.h b/Engine/Plugins/FX/Niagara/Source/Niagara/Classes/NiagaraEmitter.h
--- a/Engine/Plugins/FX/Niagara/Source/Niagara/Classes/NiagaraEmitter.h
+++ b/Engine/Plugins/FX/Niagara/Source/Niagara/Classes/NiagaraEmitter.h
@@ -498,9 +498,9 @@ public:
 	bool UsesScript(const UNiagaraScript* Script)const;
 	//bool UsesDataInterface(UNiagaraDataInterface* Interface);
 	bool UsesCollection(const class UNiagaraParameterCollection* Collection)const;
-	bool CanObtainParticleAttribute(const FNiagaraVariableBase& InVar) const;
-	bool CanObtainEmitterAttribute(const FNiagaraVariableBase& InVarWithUniqueNameNamespace) const;
-	bool CanObtainSystemAttribute(const FNiagaraVariableBase& InVar) const;
+	bool CanObtainParticleAttribute(const FNiagaraVariableBase& InVar, FNiagaraTypeDefinition& OutBoundType) const;
+	bool CanObtainEmitterAttribute(const FNiagaraVariableBase& InVarWithUniqueNameNamespace, FNiagaraTypeDefinition& OutBoundType) const;
+	bool CanObtainSystemAttribute(const FNiagaraVariableBase& InVar, FNiagaraTypeDefinition& OutBoundType) const;
 	bool CanObtainUserVariable(const FNiagaraVariableBase& InVar) const;
 
 #if WITH_NIAGARA_DEBUG_EMITTER_NAME
diff --git a/Engine/Plugins/FX/Niagara/Source/Niagara/Classes/NiagaraGPUSortInfo.h b/Engine/Plugins/FX/Niagara/Source/Niagara/Classes/NiagaraGPUSortInfo.h
--- a/Engine/Plugins/FX/Niagara/Source/Niagara/Classes/NiagaraGPUSortInfo.h
+++ b/Engine/Plugins/FX/Niagara/Source/Niagara/Classes/NiagaraGPUSortInfo.h
@@ -68,6 +68,7 @@ struct FNiagaraGPUSortInfo
 	int32 MeshIndex = 0;
 	FSphere LocalBSphere = FSphere(0);
 	FVector CullingWorldSpaceOffset = FVector(0, 0, 0);
+	FVector3f SystemLWCTile = FVector3f(0, 0, 0);
 	FVector2D DistanceCullRange { 0.0f, FLT_MAX };
 	TArray<FPlane, TFixedAllocator<MaxCullPlanes>> CullPlanes;
 
diff --git a/Engine/Plugins/FX/Niagara/Source/Niagara/Classes/NiagaraParameterCollection.h b/Engine/Plugins/FX/Niagara/Source/Niagara/Classes/NiagaraParameterCollection.h
--- a/Engine/Plugins/FX/Niagara/Source/Niagara/Classes/NiagaraParameterCollection.h
+++ b/Engine/Plugins/FX/Niagara/Source/Niagara/Classes/NiagaraParameterCollection.h
@@ -114,7 +114,7 @@ public:
 	void SetVector2DParameter(const FString& InVariableName, FVector2D InValue);
 
 	UFUNCTION(BlueprintCallable, Category = Niagara, meta = (DisplayName = "Set Vector Parameter"))
-	void SetVectorParameter(const FString& InVariableName, FVector InValue);
+	void SetVectorParameter(const FString& InVariableName, FVector InValue); // TODO[mg]: add position setter for LWC
 
 	UFUNCTION(BlueprintCallable, Category = Niagara, meta = (DisplayName = "Set Vector4 Parameter"))
 	void SetVector4Parameter(const FString& InVariableName, const FVector4& InValue);
diff --git a/Engine/Plugins/FX/Niagara/Source/Niagara/Classes/NiagaraScriptSourceBase.h b/Engine/Plugins/FX/Niagara/Source/Niagara/Classes/NiagaraScriptSourceBase.h
--- a/Engine/Plugins/FX/Niagara/Source/Niagara/Classes/NiagaraScriptSourceBase.h
+++ b/Engine/Plugins/FX/Niagara/Source/Niagara/Classes/NiagaraScriptSourceBase.h
@@ -165,6 +165,10 @@ class UNiagaraScriptSourceBase : public UObject
 	 */
 	virtual void RenameGraphAssignmentAndSetNodePins(const FName OldName, const FName NewName) {};
 
+	/** Checks if any of the provided variables are linked to function inputs of position type data
+	 */
+	virtual void GetLinkedPositionTypeInputs(const TArray<FNiagaraVariable>& ParametersToCheck, TSet<FNiagaraVariable>& OutLinkedParameters) {};
+
 protected:
 	FOnChanged OnChangedDelegate;
 #endif
diff --git a/Engine/Plugins/FX/Niagara/Source/Niagara/Classes/NiagaraSystem.h b/Engine/Plugins/FX/Niagara/Source/Niagara/Classes/NiagaraSystem.h
--- a/Engine/Plugins/FX/Niagara/Source/Niagara/Classes/NiagaraSystem.h
+++ b/Engine/Plugins/FX/Niagara/Source/Niagara/Classes/NiagaraSystem.h
@@ -348,6 +348,10 @@ public:
 
 private:
 	bool IsValidInternal() const;
+
+#if WITH_EDITORONLY_DATA
+	void FixupPositionUserParameters();
+#endif
 	
 public:
 	/** Returns true if this system is valid and can be instanced. False otherwise. */
@@ -531,6 +535,8 @@ public:
 	const FNiagaraSystemCompiledData& GetSystemCompiledData() const { return SystemCompiledData; };
 
 	bool UsesCollection(const UNiagaraParameterCollection* Collection)const;
+
+	bool SupportsLargeWorldCoordinates() const { return bSupportLargeWorldCoordinates && bLwcEnabledSettingCached; }
 #if WITH_EDITORONLY_DATA
 	bool UsesEmitter(const UNiagaraEmitter* Emitter) const;
 	bool UsesScript(const UNiagaraScript* Script)const; 
@@ -570,6 +576,10 @@ public:
 	TArray<FParameterDefinitionsSubscription> ParameterDefinitionsSubscriptions;
 #endif
 
+	/** If true then position type values will be rebased on system activation to fit into a float precision vector. This needs to be turned off when using a custom data interface or renderer that does not support the rebasing. */
+	UPROPERTY(EditAnywhere, AdvancedDisplay, Category = "Rendering")
+	uint8 bSupportLargeWorldCoordinates : 1;
+
 	//////////////////////////////////////////////////////////////////////////
 	// Various optional overrides for component properties when spawning a system
 	UPROPERTY(EditAnywhere, Category="Rendering", meta=(InlineEditConditionToggle="bCastShadow"))
@@ -707,8 +717,8 @@ public:
 	void AddToInstanceCountStat(int32 NumInstances, bool bSolo)const;
 
 	const FString& GetCrashReporterTag()const;
-	bool CanObtainEmitterAttribute(const FNiagaraVariableBase& InVarWithUniqueNameNamespace) const;
-	bool CanObtainSystemAttribute(const FNiagaraVariableBase& InVar) const;
+	bool CanObtainEmitterAttribute(const FNiagaraVariableBase& InVarWithUniqueNameNamespace, FNiagaraTypeDefinition& OutBoundType) const;
+	bool CanObtainSystemAttribute(const FNiagaraVariableBase& InVar, FNiagaraTypeDefinition& OutBoundType) const;
 	bool CanObtainUserVariable(const FNiagaraVariableBase& InVar) const;
 
 #if WITH_EDITORONLY_DATA
@@ -915,6 +925,7 @@ protected:
 	uint32 bIsValidCached : 1;
 	uint32 bIsReadyToRunCached : 1;
 	uint32 bNeedsAsyncOptimize : 1;
+	uint32 bLwcEnabledSettingCached : 1;
 
 	TOptional<float> MaxDeltaTime;
 	FNiagaraDataSetAccessor<ENiagaraExecutionState> SystemExecutionStateAccessor;
diff --git a/Engine/Plugins/FX/Niagara/Source/Niagara/Private/DataInterface/NiagaraDataInterfaceActorComponent.cpp b/Engine/Plugins/FX/Niagara/Source/Niagara/Private/DataInterface/NiagaraDataInterfaceActorComponent.cpp
--- a/Engine/Plugins/FX/Niagara/Source/Niagara/Private/DataInterface/NiagaraDataInterfaceActorComponent.cpp
+++ b/Engine/Plugins/FX/Niagara/Source/Niagara/Private/DataInterface/NiagaraDataInterfaceActorComponent.cpp
@@ -12,6 +12,18 @@
 
 #define LOCTEXT_NAMESPACE "NiagaraDataInterfaceActorComponent"
 
+struct FNiagaraActorDIFunctionVersion
+{
+	enum Type
+	{
+		InitialVersion = 0,
+		LWCConversion = 1,
+
+		VersionPlusOne,
+		LatestVersion = VersionPlusOne - 1
+	};
+};
+
 namespace NDIActorComponentLocal
 {
 	static const TCHAR* TemplateShaderFile = TEXT("/Plugin/FX/Niagara/Private/NiagaraDataInterfaceActorComponentTemplate.ush");
@@ -94,7 +106,7 @@ public:
 		FInstanceData_RenderThread* InstanceData = DataInterfaceProxy->SystemInstancesToInstanceData_RT.Find(Context.SystemInstanceID);
 		check(InstanceData != nullptr);
 
-		const FMatrix44f InstanceMatrix = (FMatrix44f)InstanceData->CachedTransform.ToMatrixWithScale();	// LWC_TODO: Precision loss
+		const FMatrix44f InstanceMatrix = (FMatrix44f)InstanceData->CachedTransform.ToMatrixWithScale();
 		const FQuat4f InstanceRotation = (FQuat4f)InstanceData->CachedTransform.GetRotation();
 		const FVector3f InstanceScale = (FVector3f)InstanceData->CachedTransform.GetScale3D();
 		SetShaderValue(RHICmdList, ComputeShaderRHI, ValidParam, InstanceData->bCachedValid ? 1 : 0);
@@ -163,6 +175,7 @@ void UNiagaraDataInterfaceActorComponent::GetFunctions(TArray<FNiagaraFunctionSi
 		FNiagaraFunctionSignature& FunctionSignature = OutFunctions.AddDefaulted_GetRef();
 		FunctionSignature.Name = GetMatrixName;
 		FunctionSignature.SetDescription(LOCTEXT("GetMatrix", "Returns the current matrix for the component if valid."));
+		FunctionSignature.SetFunctionVersion(FNiagaraActorDIFunctionVersion::LatestVersion);
 		FunctionSignature.bMemberFunction = true;
 		FunctionSignature.bRequiresContext = false;
 		FunctionSignature.bSupportsGPU = true;
@@ -174,12 +187,13 @@ void UNiagaraDataInterfaceActorComponent::GetFunctions(TArray<FNiagaraFunctionSi
 		FNiagaraFunctionSignature& FunctionSignature = OutFunctions.AddDefaulted_GetRef();
 		FunctionSignature.Name = GetTransformName;
 		FunctionSignature.SetDescription(LOCTEXT("GetTransform", "Returns the current transform for the component if valid."));
+		FunctionSignature.SetFunctionVersion(FNiagaraActorDIFunctionVersion::LatestVersion);
 		FunctionSignature.bMemberFunction = true;
 		FunctionSignature.bRequiresContext = false;
 		FunctionSignature.bSupportsGPU = true;
 		FunctionSignature.Inputs.Emplace(FNiagaraTypeDefinition(GetClass()), TEXT("ActorComponent"));
 		FunctionSignature.Outputs.Emplace(FNiagaraTypeDefinition::GetBoolDef(), TEXT("IsValid"));
-		FunctionSignature.Outputs.Emplace(FNiagaraTypeDefinition::GetVec3Def(), TEXT("Position"));
+		FunctionSignature.Outputs.Emplace(FNiagaraTypeDefinition::GetPositionDef(), TEXT("Position"));
 		FunctionSignature.Outputs.Emplace(FNiagaraTypeDefinition::GetQuatDef(), TEXT("Rotation"));
 		FunctionSignature.Outputs.Emplace(FNiagaraTypeDefinition::GetVec3Def(), TEXT("Scale"));
 	}
@@ -224,6 +238,25 @@ bool UNiagaraDataInterfaceActorComponent::GetFunctionHLSL(const FNiagaraDataInte
 	using namespace NDIActorComponentLocal;
 	return (FunctionInfo.DefinitionName == GetMatrixName) || (FunctionInfo.DefinitionName == GetTransformName);
 }
+
+bool UNiagaraDataInterfaceActorComponent::UpgradeFunctionCall(FNiagaraFunctionSignature& FunctionSignature)
+{
+	// LWC upgrades
+	if (FunctionSignature.FunctionVersion < FNiagaraActorDIFunctionVersion::LWCConversion)
+	{
+		TArray<FNiagaraFunctionSignature> AllFunctions;
+		GetFunctions(AllFunctions);
+		for (const FNiagaraFunctionSignature& Sig : AllFunctions)
+		{
+			if (FunctionSignature.Name == Sig.Name)
+			{
+				FunctionSignature = Sig;
+				return true;
+			}
+		}
+	}
+	return false;
+}
 #endif
 
 bool UNiagaraDataInterfaceActorComponent::InitPerInstanceData(void* PerInstanceData, FNiagaraSystemInstance* SystemInstance)
@@ -276,11 +309,13 @@ bool UNiagaraDataInterfaceActorComponent::PerInstanceTick(void* PerInstanceData,
 		{
 			InstanceData->bCachedValid = true;
 			InstanceData->CachedTransform = SceneComponent->GetComponentToWorld();
+			InstanceData->CachedTransform.AddToTranslation(FVector(SystemInstance->GetLWCTile()) * -FLargeWorldRenderScalar::GetTileSize());
 		}
 		else if (AActor* OwnerActor = ActorComponent->GetOwner())
 		{
 			InstanceData->bCachedValid = true;
 			InstanceData->CachedTransform = OwnerActor->GetTransform();
+			InstanceData->CachedTransform.AddToTranslation(FVector(SystemInstance->GetLWCTile()) * -FLargeWorldRenderScalar::GetTileSize());
 		}
 	}
 
diff --git a/Engine/Plugins/FX/Niagara/Source/Niagara/Private/DataInterface/NiagaraDataInterfaceStaticMesh.cpp b/Engine/Plugins/FX/Niagara/Source/Niagara/Private/DataInterface/NiagaraDataInterfaceStaticMesh.cpp
--- a/Engine/Plugins/FX/Niagara/Source/Niagara/Private/DataInterface/NiagaraDataInterfaceStaticMesh.cpp
+++ b/Engine/Plugins/FX/Niagara/Source/Niagara/Private/DataInterface/NiagaraDataInterfaceStaticMesh.cpp
@@ -56,6 +56,7 @@ namespace NDIStaticMeshLocal
 		{
 			InitialVersion = 0,
 			RefactoredV1 = 1,
+			LargeWorldCoordinates = 2,
 
 			VersionPlusOne,
 			LatestVersion = VersionPlusOne - 1
@@ -557,7 +558,8 @@ namespace NDIStaticMeshLocal
 
 			// Gather attached information
 			bComponentValid = SceneComponent != nullptr;
-			const FTransform& ComponentTransform = bComponentValid ? SceneComponent->GetComponentToWorld() : SystemInstance->GetWorldTransform();
+			FTransform ComponentTransform = bComponentValid ? SceneComponent->GetComponentToWorld() : SystemInstance->GetWorldTransform();
+			ComponentTransform.AddToTranslation(FVector(SystemInstance->GetLWCTile()) * -FLargeWorldRenderScalar::GetTileSize());
 
 			Transform = ComponentTransform.ToMatrixWithScale();
 			TransformInverseTransposed = ComponentTransform.Inverse().ToMatrixWithScale().GetTransposed();
@@ -745,7 +747,8 @@ namespace NDIStaticMeshLocal
 			DeltaSeconds = InDeltaSeconds;
 
 			USceneComponent* SceneComponent = SceneComponentWeakPtr.Get();
-			const FTransform& ComponentTransform = SceneComponent != nullptr ? SceneComponent->GetComponentToWorld() : SystemInstance->GetWorldTransform();
+			FTransform ComponentTransform = SceneComponent != nullptr ? SceneComponent->GetComponentToWorld() : SystemInstance->GetWorldTransform();
+			ComponentTransform.AddToTranslation(FVector(SystemInstance->GetLWCTile()) * -FLargeWorldRenderScalar::GetTileSize());
 
 			PrevTransform = Transform;
 			PrevTransformInverseTransposed = TransformInverseTransposed;
@@ -1600,7 +1603,10 @@ void UNiagaraDataInterfaceStaticMesh::GetFunctions(TArray<FNiagaraFunctionSignat
 		Sig.Outputs.Emplace(FNiagaraTypeDefinition::GetVec3Def(), TEXT("Bitangent"));
 		Sig.Outputs.Emplace(FNiagaraTypeDefinition::GetVec3Def(), TEXT("Tangent"));
 		OutFunctions.Add_GetRef(Sig).Name = GetVertexName;
-		OutFunctions.Add_GetRef(Sig).Name = GetVertexWSName;
+		
+		FNiagaraFunctionSignature& WsSig = OutFunctions.Add_GetRef(Sig);
+		WsSig.Name = GetVertexWSName;
+		WsSig.Outputs[0].SetType(FNiagaraTypeDefinition::GetPositionDef());
 	}
 	{
 		FNiagaraFunctionSignature& Sig = OutFunctions.Add_GetRef(BaseSignature);
@@ -1697,7 +1703,10 @@ void UNiagaraDataInterfaceStaticMesh::GetFunctions(TArray<FNiagaraFunctionSignat
 		Sig.Outputs.Emplace(FNiagaraTypeDefinition::GetVec3Def(), TEXT("Bitangent"));
 		Sig.Outputs.Emplace(FNiagaraTypeDefinition::GetVec3Def(), TEXT("Tangent"));
 		OutFunctions.Add_GetRef(Sig).Name = GetTriangleName;
-		OutFunctions.Add_GetRef(Sig).Name = GetTriangleWSName;
+
+		FNiagaraFunctionSignature WsSig = OutFunctions.Add_GetRef(Sig);
+		WsSig.Name = GetTriangleWSName;
+		WsSig.Outputs[0].SetType(FNiagaraTypeDefinition::GetPositionDef());
 	}
 	{
 		FNiagaraFunctionSignature& Sig = OutFunctions.Add_GetRef(BaseSignature);
@@ -1741,12 +1750,15 @@ void UNiagaraDataInterfaceStaticMesh::GetFunctions(TArray<FNiagaraFunctionSignat
 		Sig.Outputs.Add(FNiagaraVariable(FNiagaraTypeDefinition::GetQuatDef(), TEXT("Rotation")));
 		Sig.Outputs.Add(FNiagaraVariable(FNiagaraTypeDefinition::GetVec3Def(), TEXT("Scale")));
 
+		FNiagaraFunctionSignature WsSig = Sig;
+		WsSig.Outputs[0].SetType(FNiagaraTypeDefinition::GetPositionDef());
+
 		OutFunctions.Add_GetRef(Sig).Name = GetSocketTransformName;
-		OutFunctions.Add_GetRef(Sig).Name = GetSocketTransformWSName;
+		OutFunctions.Add_GetRef(WsSig).Name = GetSocketTransformWSName;
 		OutFunctions.Add_GetRef(Sig).Name = GetFilteredSocketTransformName;
-		OutFunctions.Add_GetRef(Sig).Name = GetFilteredSocketTransformWSName;
+		OutFunctions.Add_GetRef(WsSig).Name = GetFilteredSocketTransformWSName;
 		OutFunctions.Add_GetRef(Sig).Name = GetUnfilteredSocketTransformName;
-		OutFunctions.Add_GetRef(Sig).Name = GetUnfilteredSocketTransformWSName;
+		OutFunctions.Add_GetRef(WsSig).Name = GetUnfilteredSocketTransformWSName;
 	}
 
 	////////////////////////////////////////////////////////////////////////////
@@ -1955,6 +1967,20 @@ bool UNiagaraDataInterfaceStaticMesh::UpgradeFunctionCall(FNiagaraFunctionSignat
 		FunctionSignature.FunctionVersion = EDIFunctionVersion::RefactoredV1;
 	}
 
+	if (FunctionSignature.FunctionVersion < EDIFunctionVersion::LargeWorldCoordinates)
+	{
+		TArray<FNiagaraFunctionSignature> AllFunctions;
+		GetFunctions(AllFunctions);
+		for (const FNiagaraFunctionSignature& Sig : AllFunctions)
+		{
+			if (FunctionSignature.Name == Sig.Name)
+			{
+				FunctionSignature = Sig;
+				return true;
+			}
+		}
+	}
+
 	return true;
 }
 #endif //WITH_EDITORONLY_DATA
diff --git a/Engine/Plugins/FX/Niagara/Source/Niagara/Private/NiagaraBoundsCalculatorHelper.h b/Engine/Plugins/FX/Niagara/Source/Niagara/Private/NiagaraBoundsCalculatorHelper.h
--- a/Engine/Plugins/FX/Niagara/Source/Niagara/Private/NiagaraBoundsCalculatorHelper.h
+++ b/Engine/Plugins/FX/Niagara/Source/Niagara/Private/NiagaraBoundsCalculatorHelper.h
@@ -51,8 +51,8 @@ public:
 
 		constexpr float kDefaultSize = 50.0f;
 
-		FVector3f BoundsMin(ForceInitToZero);
-		FVector3f BoundsMax(ForceInitToZero);
+		FNiagaraPosition BoundsMin(ForceInitToZero);
+		FNiagaraPosition BoundsMax(ForceInitToZero);
 		PositionAccessor.GetReader(DataSet).GetMinMax(BoundsMin, BoundsMax);
 		FBox Bounds(BoundsMin, BoundsMax);
 
@@ -111,7 +111,7 @@ public:
 		return Bounds.ExpandBy(MaxSize);
 	}
 
-	FNiagaraDataSetAccessor<FVector3f> PositionAccessor;
+	FNiagaraDataSetAccessor<FNiagaraPosition> PositionAccessor;
 	FNiagaraDataSetAccessor<FVector2f> SpriteSizeAccessor;
 	FNiagaraDataSetAccessor<FVector3f> ScaleAccessor;
 	FNiagaraDataSetAccessor<float> RibbonWidthAccessor;
diff --git a/Engine/Plugins/FX/Niagara/Source/Niagara/Private/NiagaraCollision.cpp b/Engine/Plugins/FX/Niagara/Source/Niagara/Private/NiagaraCollision.cpp
--- a/Engine/Plugins/FX/Niagara/Source/Niagara/Private/NiagaraCollision.cpp
+++ b/Engine/Plugins/FX/Niagara/Source/Niagara/Private/NiagaraCollision.cpp
@@ -106,42 +106,6 @@ void FNiagaraDICollisionQueryBatch::CollectResults()
 }
 }
 
-int32 FNiagaraDICollisionQueryBatch::SubmitQuery(FVector StartPos, FVector Direction, float CollisionSize, float DeltaSeconds)
-{
-	if ( !GNiagaraCollisionCPUEnabled)
-	{
-		return INDEX_NONE;
-	}
-
-	FVector EndPos = StartPos + Direction * DeltaSeconds;
-
-	float Length;
-	FVector NormDir;
-	Direction.ToDirectionAndLength(NormDir, Length);
-	StartPos -= NormDir * (CollisionSize / 2);
-	EndPos += NormDir * (CollisionSize / 2);
-
-	FCollisionQueryParams CollisionQueryParams(SCENE_QUERY_STAT(NiagaraCollision), false);
-	CollisionQueryParams.OwnerTag = CollisionTagName;
-	CollisionQueryParams.bFindInitialOverlaps = false;
-	CollisionQueryParams.bReturnFaceIndex = false;
-	CollisionQueryParams.bReturnPhysicalMaterial = true;
-	CollisionQueryParams.bIgnoreTouches = true;
-
-	int32 TraceIdx = INDEX_NONE;
-
-	if (Length > SMALL_NUMBER)
-	{
-		FRWScopeLock WriteScope(CollisionTraceLock, SLT_Write);
-
-		const int32 WriteBufferIdx = GetWriteBufferIdx();
-
-		TraceIdx = CollisionTraces[WriteBufferIdx].Emplace(StartPos, EndPos, ECC_WorldStatic, CollisionQueryParams);
-	}
-
-	return TraceIdx;
-}
-
 int32 FNiagaraDICollisionQueryBatch::SubmitQuery(FVector StartPos, FVector EndPos, ECollisionChannel TraceChannel)
 {
 	if (!GNiagaraCollisionCPUEnabled)
diff --git a/Engine/Plugins/FX/Niagara/Source/Niagara/Private/NiagaraCommon.cpp b/Engine/Plugins/FX/Niagara/Source/Niagara/Private/NiagaraCommon.cpp
--- a/Engine/Plugins/FX/Niagara/Source/Niagara/Private/NiagaraCommon.cpp
+++ b/Engine/Plugins/FX/Niagara/Source/Niagara/Private/NiagaraCommon.cpp
@@ -161,6 +161,11 @@ UScriptStruct* FNiagaraTypeHelper::FindNiagaraFriendlyTopLevelStruct(UScriptStru
 		return FNiagaraTypeDefinition::GetQuatStruct();
 	}
 
+	if (InStruct->GetFName() == FName("NiagaraPosition"))
+	{
+		return FNiagaraTypeDefinition::GetPositionStruct();
+	}
+
 	return InStruct;
 }
 
@@ -843,22 +848,28 @@ void FNiagaraVariableAttributeBinding::CacheValues(const UNiagaraEmitter* InEmit
 			DataSetVariable = FNiagaraUtilities::ResolveAliases(DataSetVariable, ResolveAliasesContext);
 		}
 
+		FNiagaraTypeDefinition BoundVarType = ParamMapVariable.GetType();
 		if (BindingSourceMode == ENiagaraBindingSource::ExplicitParticles || (InSourceMode == ENiagaraRendererSourceDataMode::Particles && BindingSourceMode == ENiagaraBindingSource::ImplicitFromSource))
 		{
-			bBindingExistsOnSource = InEmitter->CanObtainParticleAttribute(DataSetVariable);
+			bBindingExistsOnSource = InEmitter->CanObtainParticleAttribute(DataSetVariable, BoundVarType);
 		}
 		else if (BindingSourceMode == ENiagaraBindingSource::ExplicitEmitter || (InSourceMode == ENiagaraRendererSourceDataMode::Emitter && BindingSourceMode == ENiagaraBindingSource::ImplicitFromSource))
 		{
-			bBindingExistsOnSource = InEmitter->CanObtainEmitterAttribute(ParamMapVariable);
+			bBindingExistsOnSource = InEmitter->CanObtainEmitterAttribute(ParamMapVariable, BoundVarType);
 		}
 		else if (BindingSourceMode == ENiagaraBindingSource::ExplicitSystem)
 		{
-			bBindingExistsOnSource = InEmitter->CanObtainSystemAttribute(ParamMapVariable);
+			bBindingExistsOnSource = InEmitter->CanObtainSystemAttribute(ParamMapVariable, BoundVarType);
 		}
 		else if (BindingSourceMode == ENiagaraBindingSource::ExplicitUser)
 		{
 			bBindingExistsOnSource = InEmitter->CanObtainUserVariable(ParamMapVariable);
 		}
+
+		if (bBindingExistsOnSource && BoundVarType != ParamMapVariable.GetType())
+		{
+			ParamMapVariable.SetType(BoundVarType);
+		}
 	}
 
 }
diff --git a/Engine/Plugins/FX/Niagara/Source/Niagara/Private/NiagaraComponent.cpp b/Engine/Plugins/FX/Niagara/Source/Niagara/Private/NiagaraComponent.cpp
--- a/Engine/Plugins/FX/Niagara/Source/Niagara/Private/NiagaraComponent.cpp
+++ b/Engine/Plugins/FX/Niagara/Source/Niagara/Private/NiagaraComponent.cpp
@@ -4,15 +4,11 @@
 #include "Engine/CollisionProfile.h"
 #include "EngineUtils.h"
 #include "Materials/MaterialInstanceDynamic.h"
-#include "MeshBatch.h"
 #include "NiagaraCommon.h"
 #include "NiagaraComponentSettings.h"
 #include "NiagaraConstants.h"
-#include "NiagaraCrashReporterHandler.h"
 #include "NiagaraCustomVersion.h"
 #include "NiagaraDataInterface.h"
-#include "NiagaraDataSetAccessor.h"
-#include "NiagaraEmitterInstance.h"
 #include "NiagaraFunctionLibrary.h"
 #include "NiagaraGpuComputeDispatchInterface.h"
 #include "NiagaraRenderer.h"
@@ -22,7 +18,6 @@
 #include "PrimitiveSceneInfo.h"
 #include "ProfilingDebugging/CsvProfiler.h"
 #include "UObject/NameTypes.h"
-#include "VectorVM.h"
 #include "Engine/StaticMesh.h"
 #include "NiagaraCullProxyComponent.h"
 
@@ -162,6 +157,14 @@ static FAutoConsoleVariableRef CVarNiagaraCanPreventScalabilityCullingOnPlayerFX
 	ECVF_Default
 );
 
+bool FitsIntoFloat(FVector InValue)
+{
+	// copied from PlayerCameraManager, 1048576 is HALF_WORLD_MAX for float precision vectors
+	return (InValue.X < 1048576.0f && InValue.X > -1048576.0f &&
+			InValue.Y < 1048576.0f && InValue.Y > -1048576.0f &&
+			InValue.Z < 1048576.0f && InValue.Z > -1048576.0f);
+}
+
 FNiagaraSceneProxy::FNiagaraSceneProxy(UNiagaraComponent* InComponent)
 	: FPrimitiveSceneProxy(InComponent, InComponent->GetAsset() ? InComponent->GetAsset()->GetFName() : FName())
 {
@@ -882,6 +885,7 @@ bool UNiagaraComponent::InitializeSystem()
 		check(Asset);
 
 		const bool bPooled = PoolingMethod != ENCPoolMethod::None;
+		OverrideParameters.MarkParametersDirty(); // new system instance means new lwc tile, so any position user params need to be re-evaluated
 
 		SystemInstanceController = MakeShared<FNiagaraSystemInstanceController, ESPMode::ThreadSafe>();
 		SystemInstanceController->Initialize(*World, *Asset, &OverrideParameters, this, TickBehavior, bPooled, RandomSeedOffset, bForceSolo);
@@ -1780,6 +1784,8 @@ void UNiagaraComponent::SendRenderDynamicData_Concurrent()
 	{
 		if (FNiagaraSystemRenderData* RenderData = NiagaraProxy->GetSystemRenderData())
 		{
+			RenderData->LWCRenderTile = SystemInstanceController->GetSystemInstance_Unsafe()->GetLWCTile();
+			
 			FNiagaraSceneProxy::FDynamicData NewProxyDynamicData;
 			NewProxyDynamicData.LODDistanceOverride = GetPreviewLODDistance();
 
@@ -2032,7 +2038,20 @@ void UNiagaraComponent::SetNiagaraVariableVec4(const FString& InVariableName, co
 
 void UNiagaraComponent::SetVariableVec3(FName InVariableName, FVector InValue)
 {
-	FVector3f AsFloat = (FVector3f)InValue;
+	// If the user parameter is of position type we use the position function instead, so it gets converted correctly 
+	if (OverrideParameters.GetPositionParameterValue(InVariableName)) 
+	{
+		SetVariablePosition(InVariableName, InValue);
+		return;
+	}
+	
+#if !(UE_BUILD_SHIPPING || UE_BUILD_TEST)
+	if (!FitsIntoFloat(InValue))
+	{
+		UE_LOG(LogNiagara, Warning, TEXT("Setting variable %s on asset %s with value (%f, %f, %f), which does not fit into a FVector3f without losing precision! Consider using a position type instead if the input is a location."), *InVariableName.ToString(), *GetPathName(), InValue.X, InValue.Y, InValue.Z);
+	}
+#endif
+	FVector3f AsFloat(InValue);
 	const FNiagaraVariable VariableDesc(FNiagaraTypeDefinition::GetVec3Def(), InVariableName);
 	if (SystemInstanceController.IsValid())
 	{
@@ -2052,6 +2071,34 @@ void UNiagaraComponent::SetNiagaraVariableVec3(const FString& InVariableName, FV
 	SetVariableVec3(FName(*InVariableName), InValue);
 }
 
+void UNiagaraComponent::SetVariablePosition(FName InVariableName, FVector InValue)
+{
+	// If the user parameter is not of position type we try to set it as vector instead, as users might confuse the two 
+	if (!OverrideParameters.GetPositionParameterValue(InVariableName)) 
+	{
+		SetVariableVec3(InVariableName, InValue);
+		return;
+	}
+	
+	const FNiagaraVariable VariableDesc(FNiagaraTypeDefinition::GetPositionDef(), InVariableName);
+	if (SystemInstanceController.IsValid())
+	{
+		SystemInstanceController->SetVariable_Deferred(InVariableName, InValue);
+	}
+	else
+	{
+		OverrideParameters.SetPositionParameterValue(InValue, InVariableName, true);
+	}
+#if WITH_EDITOR
+	SetParameterOverride(VariableDesc, FNiagaraVariant(&InValue, sizeof(FVector)));
+#endif
+}
+
+void UNiagaraComponent::SetNiagaraVariablePosition(const FString& InVariableName, FVector InValue)
+{
+	SetVariablePosition(FName(*InVariableName), InValue);
+}
+
 void UNiagaraComponent::SetVariableVec2(FName InVariableName, FVector2D InValue)
 {
 	FVector2f AsFloat = (FVector2f)InValue;
@@ -2300,7 +2347,6 @@ void FixInvalidUserParameters(FNiagaraUserRedirectionParameterStore& ParameterSt
 			UNiagaraDataInterface* DataInterfaceValue = nullptr;
 			UObject* ObjectValue = nullptr;
 			TArray<uint8> DataValue;
-			int32 ValueIndex = ParameterStore.IndexOf(IncorrectlyNamedParameter);
 			if (IncorrectlyNamedParameter.IsDataInterface())
 			{
 				DataInterfaceValue = ParameterStore.GetDataInterface(IncorrectlyNamedParameter);
@@ -2388,8 +2434,7 @@ void UNiagaraComponent::PostLoad()
 	{
 		Asset->ConditionalPostLoad();
 
-		FixInvalidUserParameters(OverrideParameters);
-		
+		FixInvalidUserParameters(OverrideParameters);		
 		UpgradeDeprecatedParameterOverrides();
 
 #if WITH_EDITORONLY_DATA
@@ -2713,6 +2758,14 @@ void UNiagaraComponent::EnsureOverrideParametersConsistent() const
 					ensureAlways(OverrideValue.GetUObject() == ActualObj);
 				}
 			}
+			else if (Key.GetType() == FNiagaraTypeDefinition::GetPositionDef())
+			{
+				const FVector* ActualData = OverrideParameters.GetPositionParameterValue(Key.GetName());
+				if (ActualData != nullptr)
+				{
+					ensureAlways(FMemory::Memcmp(ActualData, OverrideValue.GetBytes(), sizeof(FVector)) == 0);
+				}
+			}
 			else
 			{
 				const uint8* ActualData = OverrideParameters.GetParameterData(Key);
@@ -2825,6 +2878,8 @@ void UNiagaraComponent::CopyParametersFromAsset(bool bResetExistingOverrideParam
 
 	TArray<FNiagaraVariable> SourceVars;
 	Asset->GetExposedParameters().GetParameters(SourceVars);
+
+	// insert new asset parameters (and keep existing overrides)
 	for (FNiagaraVariable& Param : SourceVars)
 	{
 		bool bNewParam = OverrideParameters.AddParameter(Param, true);
@@ -2836,7 +2891,8 @@ void UNiagaraComponent::CopyParametersFromAsset(bool bResetExistingOverrideParam
 		}
 	}
 
-	for (FNiagaraVariable ExistingVar : ExistingVars)
+	// remove parameters that don't exist in the source asset
+	for (FNiagaraVariable& ExistingVar : ExistingVars)
 	{
 		if (!SourceVars.Contains(ExistingVar))
 		{
@@ -2898,6 +2954,35 @@ void UNiagaraComponent::FixInvalidUserParameterOverrideData()
 {
 	FixInvalidDataInterfaceOverrides(TemplateParameterOverrides, TEXT("Template"), this);
 	FixInvalidDataInterfaceOverrides(InstanceParameterOverrides, TEXT("Instance"), this);
+
+	TArray<FNiagaraVariable> ExistingVars;
+	TArray<FNiagaraVariable> SourceVars;
+	OverrideParameters.GetParameters(ExistingVars);
+	Asset->GetExposedParameters().GetParameters(SourceVars);
+
+	// check if one of the vector parameters was changed to a position type
+	for (FNiagaraVariable& ExistingVar : ExistingVars)
+	{
+		if (ExistingVar.GetType() == FNiagaraTypeDefinition::GetVec3Def() && !SourceVars.Contains(ExistingVar))
+		{
+			FNiagaraVariable PositionVar(FNiagaraTypeDefinition::GetPositionDef(), ExistingVar.GetName());
+			if (SourceVars.Contains(PositionVar))
+			{
+				OverrideParameters.ChangeParameterType(ExistingVar, FNiagaraTypeDefinition::GetPositionDef());
+				FNiagaraVariant VariantData;
+				if (InstanceParameterOverrides.RemoveAndCopyValue(ExistingVar, VariantData))
+				{
+					InstanceParameterOverrides.Add(PositionVar, VariantData);
+				}
+				if (TemplateParameterOverrides.RemoveAndCopyValue(ExistingVar, VariantData))
+				{
+					TemplateParameterOverrides.Add(PositionVar, VariantData);
+				}
+				
+				ExistingVar = PositionVar;
+			}
+		}
+	}
 }
 #endif
 
@@ -3029,14 +3114,19 @@ FNiagaraVariant UNiagaraComponent::GetCurrentParameterValue(const FNiagaraVariab
 	{
 		return FNiagaraVariant(OverrideParameters.GetDataInterface(ParameterOffset));
 	}
-	else if (InKey.GetType().IsUObject())
+	if (InKey.GetType().IsUObject())
 	{
 		return FNiagaraVariant(OverrideParameters.GetUObject(ParameterOffset));
 	}
-	else
+	if (InKey.GetType() == FNiagaraTypeDefinition::GetPositionDef())
 	{
-		return FNiagaraVariant(OverrideParameters.GetParameterData(ParameterOffset), InKey.GetSizeInBytes());
+		FNiagaraVariable RedirectedVar = InKey;
+		OverrideParameters.RedirectUserVariable(RedirectedVar);
+		const FVector* Value = OverrideParameters.GetPositionParameterValue(RedirectedVar.GetName());
+		return Value == nullptr ? FNiagaraVariant() : FNiagaraVariant(Value, sizeof(FVector));
 	}
+	
+	return FNiagaraVariant(OverrideParameters.GetParameterData(ParameterOffset), InKey.GetSizeInBytes());
 }
 
 void UNiagaraComponent::SetOverrideParameterStoreValue(const FNiagaraVariableBase& InKey, const FNiagaraVariant& InValue)
@@ -3053,6 +3143,25 @@ void UNiagaraComponent::SetOverrideParameterStoreValue(const FNiagaraVariableBas
 	{
 		OverrideParameters.SetUObject(InValue.GetUObject(), InKey);
 	}
+	else if (InKey.GetType() == FNiagaraTypeDefinition::GetPositionDef())
+	{
+		if (InValue.GetNumBytes() == sizeof(FVector))
+		{
+			FVector WorldPos;
+			FMemory::Memcpy(&WorldPos, InValue.GetBytes(), InValue.GetNumBytes());
+			OverrideParameters.SetPositionParameterValue(WorldPos, InKey.GetName(), true);
+		}
+		else if (InValue.GetNumBytes() == sizeof(FVector3f))
+		{
+			FVector3f WorldPos;
+			FMemory::Memcpy(&WorldPos, InValue.GetBytes(), InValue.GetNumBytes());
+			OverrideParameters.SetPositionParameterValue(WorldPos, InKey.GetName(), true);
+		}
+		else
+		{
+			checkf(false, TEXT("Invalid number of bytes for position data: %i"), InValue.GetNumBytes());
+		}
+	}
 	else
 	{
 		OverrideParameters.SetParameterData(InValue.GetBytes(), InKey, true);
