Niagara LWC changes:

diff --git a/Engine/Plugins/FX/Niagara/Source/Niagara/Private/NiagaraComponentRendererProperties.cpp b/Engine/Plugins/FX/Niagara/Source/Niagara/Private/NiagaraComponentRendererProperties.cpp
--- a/Engine/Plugins/FX/Niagara/Source/Niagara/Private/NiagaraComponentRendererProperties.cpp
+++ b/Engine/Plugins/FX/Niagara/Source/Niagara/Private/NiagaraComponentRendererProperties.cpp
@@ -14,7 +14,6 @@
 #include "Styling/SlateBrush.h"
 #include "AssetThumbnail.h"
 #include "Widgets/Text/STextBlock.h"
-#include "Editor.h"
 #endif
 #include "NiagaraSettings.h"
 
@@ -43,6 +42,8 @@ bool UNiagaraComponentRendererProperties::IsConvertible(const FNiagaraTypeDefini
 		(SourceType == FNiagaraTypeDefinition::GetVec4Def() && TargetType.GetStruct() == GetFVector4Def().GetStruct()) ||
 		(SourceType == FNiagaraTypeDefinition::GetVec4Def() && TargetType.GetStruct() == GetFColorDef().GetStruct()) ||
 		(SourceType == FNiagaraTypeDefinition::GetQuatDef() && TargetType.GetStruct() == GetFQuatDef().GetStruct()) ||
+		(SourceType == FNiagaraTypeDefinition::GetPositionDef() && TargetType.GetStruct() == GetFVectorDef().GetStruct()) ||
+		(SourceType == FNiagaraTypeDefinition::GetPositionDef() && TargetType.GetStruct() == GetFVector3fDef().GetStruct()) ||
 		(SourceType == FNiagaraTypeDefinition::GetQuatDef() && TargetType.GetStruct() == GetFRotatorDef().GetStruct()))
 	{
 		return true;
@@ -129,6 +130,13 @@ FNiagaraTypeDefinition UNiagaraComponentRendererProperties::GetFVector4Def()
 	return FNiagaraTypeDefinition(Vector4Struct, FNiagaraTypeDefinition::EAllowUnfriendlyStruct::Allow);
 }
 
+FNiagaraTypeDefinition UNiagaraComponentRendererProperties::GetFVector3fDef()
+{
+	static UPackage* CoreUObjectPkg = FindObjectChecked<UPackage>(nullptr, TEXT("/Script/CoreUObject"));
+	static UScriptStruct* VectorStruct = FindObjectChecked<UScriptStruct>(CoreUObjectPkg, TEXT("Vector3f"));
+	return FNiagaraTypeDefinition(VectorStruct, FNiagaraTypeDefinition::EAllowUnfriendlyStruct::Allow);
+}
+
 FNiagaraTypeDefinition UNiagaraComponentRendererProperties::GetFQuatDef()
 {
 	static UPackage* CoreUObjectPkg = FindObjectChecked<UPackage>(nullptr, TEXT("/Script/CoreUObject"));
@@ -166,9 +174,23 @@ void UNiagaraComponentRendererProperties::PostLoad()
 {
 	Super::PostLoad();
 	ENiagaraRendererSourceDataMode InSourceMode = ENiagaraRendererSourceDataMode::Particles;
+	
+	const TArray<FNiagaraVariable>& OldTypes = FNiagaraConstants::GetOldPositionTypeVariables();
 	for (FNiagaraComponentPropertyBinding& Binding : PropertyBindings)
 	{
 		Binding.AttributeBinding.PostLoad(InSourceMode);
+
+		// Move old bindings over to new position type 
+		for (FNiagaraVariable OldVarType : OldTypes)
+		{
+			if (Binding.AttributeBinding.GetParamMapBindableVariable() == OldVarType)
+			{
+				FNiagaraVariable NewVarType(FNiagaraTypeDefinition::GetPositionDef(), OldVarType.GetName());
+				Binding.AttributeBinding.Setup(NewVarType, NewVarType, InSourceMode);
+				Binding.PropertyType = FNiagaraTypeDefinition(FNiagaraTypeDefinition::GetVec3Struct());
+				break;
+			}
+		}
 	}
 	EnabledBinding.PostLoad(InSourceMode);
 	RendererVisibilityTagBinding.PostLoad(InSourceMode);
@@ -475,6 +497,7 @@ void UNiagaraComponentRendererProperties::PostEditChangeProperty(struct FPropert
 			FNiagaraComponentPropertyBinding PositionBinding;
 			PositionBinding.AttributeBinding.Setup(SYS_PARAM_PARTICLES_POSITION, SYS_PARAM_PARTICLES_POSITION);
 			PositionBinding.PropertyName = FName("RelativeLocation");
+			PositionBinding.PropertyType = FNiagaraTypeDefinition(FNiagaraTypeDefinition::GetVec3Struct());
 			PropertyBindings.Add(PositionBinding);
 
 			FNiagaraComponentPropertyBinding ScaleBinding;
diff --git a/Engine/Plugins/FX/Niagara/Source/Niagara/Private/NiagaraComputeExecutionContext.cpp b/Engine/Plugins/FX/Niagara/Source/Niagara/Private/NiagaraComputeExecutionContext.cpp
--- a/Engine/Plugins/FX/Niagara/Source/Niagara/Private/NiagaraComputeExecutionContext.cpp
+++ b/Engine/Plugins/FX/Niagara/Source/Niagara/Private/NiagaraComputeExecutionContext.cpp
@@ -147,6 +147,10 @@ bool FNiagaraComputeExecutionContext::Tick(FNiagaraSystemInstance* ParentSystemI
 			}
 		}
 #endif
+		if (ParentSystemInstance && CombinedParamStore.GetPositionDataDirty())
+		{
+			CombinedParamStore.ResolvePositions(ParentSystemInstance->GetLWCConverter());
+		}
 		CombinedParamStore.Tick();
 	}
 
diff --git a/Engine/Plugins/FX/Niagara/Source/Niagara/Private/NiagaraConstants.cpp b/Engine/Plugins/FX/Niagara/Source/Niagara/Private/NiagaraConstants.cpp
--- a/Engine/Plugins/FX/Niagara/Source/Niagara/Private/NiagaraConstants.cpp
+++ b/Engine/Plugins/FX/Niagara/Source/Niagara/Private/NiagaraConstants.cpp
@@ -10,6 +10,7 @@
 TArray<FNiagaraVariable> FNiagaraConstants::SystemParameters;
 TArray<FNiagaraVariable> FNiagaraConstants::TranslatorParameters;
 TArray<FNiagaraVariable> FNiagaraConstants::SwitchParameters;
+TArray<FNiagaraVariable> FNiagaraConstants::OldPositionTypes;
 TMap<FName, FNiagaraVariable> FNiagaraConstants::UpdatedSystemParameters;
 TMap<FNiagaraVariable, FText> FNiagaraConstants::SystemStrMap;
 TArray<FNiagaraVariable> FNiagaraConstants::Attributes;
@@ -85,6 +86,7 @@ void FNiagaraConstants::Init()
 		SystemParameters.Add(SYS_PARAM_ENGINE_X_AXIS);
 		SystemParameters.Add(SYS_PARAM_ENGINE_Y_AXIS);
 		SystemParameters.Add(SYS_PARAM_ENGINE_Z_AXIS);
+		SystemParameters.Add(SYS_PARAM_ENGINE_LWC_TILE);
 
 		SystemParameters.Add(SYS_PARAM_ENGINE_ROTATION);
 
@@ -141,6 +143,14 @@ void FNiagaraConstants::Init()
 		SwitchParameters.Add(SYS_PARAM_FUNCTION_DEBUG_STATE);
 	}
 
+	if (OldPositionTypes.Num() == 0)
+	{
+		OldPositionTypes.Add(FNiagaraVariable(FNiagaraTypeDefinition::GetVec3Def(),  SYS_PARAM_PARTICLES_POSITION.GetName()));
+		OldPositionTypes.Add(FNiagaraVariable(FNiagaraTypeDefinition::GetVec3Def(), FName("Particles.Initial.Position")));
+		OldPositionTypes.Add(FNiagaraVariable(FNiagaraTypeDefinition::GetVec3Def(),  SYS_PARAM_ENGINE_POSITION.GetName()));
+		OldPositionTypes.Add(FNiagaraVariable(FNiagaraTypeDefinition::GetVec3Def(),  SYS_PARAM_ENGINE_EMITTER_SIMULATION_POSITION.GetName()));
+	}
+
 	if (UpdatedSystemParameters.Num() == 0)
 	{
 		UpdatedSystemParameters.Add(FName(TEXT("System Delta Time")), SYS_PARAM_ENGINE_DELTA_TIME);
@@ -208,6 +218,7 @@ void FNiagaraConstants::Init()
 		SystemStrMap.Add(SYS_PARAM_ENGINE_Y_AXIS, LOCTEXT("YAxisDesc", "The Y-axis of the owning component."));
 		SystemStrMap.Add(SYS_PARAM_ENGINE_Z_AXIS, LOCTEXT("ZAxisDesc", "The Z-axis of the owning component."));
 		SystemStrMap.Add(SYS_PARAM_ENGINE_ROTATION, LOCTEXT("EngineRotationDesc", "The owning component's rotation in world space."));
+		SystemStrMap.Add(SYS_PARAM_ENGINE_LWC_TILE, LOCTEXT("EngineLWCTileDesc", "Due to large world coordinates, the simulation position can be shifted from the actual world position to allow for more accuracy. This is the tile the system is shifted by, so (SimulationPosition + Tile * TileSize) gives the original world position. The x,y,z components of this vector are the tile and the w component is the tile size."));
 
 		SystemStrMap.Add(SYS_PARAM_ENGINE_LOCAL_TO_WORLD, LOCTEXT("LocalToWorldDesc", "Owning component's local space to world space transform matrix."));
 		SystemStrMap.Add(SYS_PARAM_ENGINE_WORLD_TO_LOCAL, LOCTEXT("WorldToLocalDesc", "Owning component's world space to local space transform matrix."));
@@ -658,6 +669,12 @@ FText FNiagaraConstants::GetEngineConstantDescription(const FNiagaraVariable& In
 	return FText();
 }
 
+const TArray<FNiagaraVariable>& FNiagaraConstants::GetOldPositionTypeVariables()
+{
+	check(OldPositionTypes.Num() != 0);
+	return OldPositionTypes;
+}
+
 const TArray<FNiagaraVariable>& FNiagaraConstants::GetCommonParticleAttributes()
 {
 	check(Attributes.Num() != 0);
diff --git a/Engine/Plugins/FX/Niagara/Source/Niagara/Private/NiagaraDataInterfaceArrayFloat.cpp b/Engine/Plugins/FX/Niagara/Source/Niagara/Private/NiagaraDataInterfaceArrayFloat.cpp
--- a/Engine/Plugins/FX/Niagara/Source/Niagara/Private/NiagaraDataInterfaceArrayFloat.cpp
+++ b/Engine/Plugins/FX/Niagara/Source/Niagara/Private/NiagaraDataInterfaceArrayFloat.cpp
@@ -21,6 +21,12 @@ UNiagaraDataInterfaceArrayFloat3::UNiagaraDataInterfaceArrayFloat3(FObjectInitia
 	Proxy.Reset(new FNDIArrayProxyImpl<FVector, UNiagaraDataInterfaceArrayFloat3>(this));
 }
 
+UNiagaraDataInterfaceArrayPosition::UNiagaraDataInterfaceArrayPosition(FObjectInitializer const& ObjectInitializer)
+	: UNiagaraDataInterfaceArray(ObjectInitializer)
+{
+	Proxy.Reset(new FNDIArrayProxyImpl<FNiagaraPosition, UNiagaraDataInterfaceArrayPosition>(this));
+}
+
 UNiagaraDataInterfaceArrayFloat4::UNiagaraDataInterfaceArrayFloat4(FObjectInitializer const& ObjectInitializer)
 	: UNiagaraDataInterfaceArray(ObjectInitializer)
 {
diff --git a/Engine/Plugins/FX/Niagara/Source/Niagara/Private/NiagaraDataInterfaceArrayFunctionLibrary.cpp b/Engine/Plugins/FX/Niagara/Source/Niagara/Private/NiagaraDataInterfaceArrayFunctionLibrary.cpp
--- a/Engine/Plugins/FX/Niagara/Source/Niagara/Private/NiagaraDataInterfaceArrayFunctionLibrary.cpp
+++ b/Engine/Plugins/FX/Niagara/Source/Niagara/Private/NiagaraDataInterfaceArrayFunctionLibrary.cpp
@@ -1,6 +1,8 @@
 // Copyright Epic Games, Inc. All Rights Reserved.
 
 #include "NiagaraDataInterfaceArrayFunctionLibrary.h"
+
+#include "NiagaraComponent.h"
 #include "NiagaraDataInterfaceArrayFloat.h"
 #include "NiagaraDataInterfaceArrayInt.h"
 #include "NiagaraDataInterfaceArrayImpl.h"
@@ -52,6 +54,12 @@ TArrayType GetNiagaraArrayValue(UNiagaraComponent* NiagaraSystem, FName Override
 	return TArrayType();
 }
 
+FNiagaraLWCConverter GetLWCConverter(UNiagaraComponent* NiagaraSystem)
+{
+	FNiagaraSystemInstance* SystemInstance = NiagaraSystem->GetSystemInstanceController()->GetSystemInstance_Unsafe();
+	return SystemInstance->GetLWCConverter();
+}
+
 //////////////////////////////////////////////////////////////////////////
 
 void UNiagaraDataInterfaceArrayFunctionLibrary::SetNiagaraArrayFloat(UNiagaraComponent* NiagaraSystem, FName OverrideName, const TArray<float>& ArrayData)
@@ -69,6 +77,18 @@ void UNiagaraDataInterfaceArrayFunctionLibrary::SetNiagaraArrayVector(UNiagaraCo
 	SetNiagaraArray<FVector, UNiagaraDataInterfaceArrayFloat3>(NiagaraSystem, OverrideName, ArrayData);
 }
 
+void UNiagaraDataInterfaceArrayFunctionLibrary::SetNiagaraArrayPosition(UNiagaraComponent* NiagaraSystem, FName OverrideName, const TArray<FVector>& ArrayData)
+{
+	FNiagaraLWCConverter LwcConverter = GetLWCConverter(NiagaraSystem);
+	TArray<FNiagaraPosition> ConvertedData;
+	ConvertedData.SetNumUninitialized(ArrayData.Num());
+	for (int i = 0; i < ArrayData.Num(); i++)
+	{
+		ConvertedData[i] = LwcConverter.ConvertWorldToSimulationPosition(ArrayData[i]);
+	}
+	SetNiagaraArray<FNiagaraPosition, UNiagaraDataInterfaceArrayPosition>(NiagaraSystem, OverrideName, ConvertedData);
+}
+
 void UNiagaraDataInterfaceArrayFunctionLibrary::SetNiagaraArrayVector4(UNiagaraComponent* NiagaraSystem, FName OverrideName, const TArray<FVector4>& ArrayData)
 {
 	SetNiagaraArray<FVector4, UNiagaraDataInterfaceArrayFloat4>(NiagaraSystem, OverrideName, ArrayData);
@@ -111,6 +131,21 @@ TArray<FVector> UNiagaraDataInterfaceArrayFunctionLibrary::GetNiagaraArrayVector
 	return GetNiagaraArray<FVector, UNiagaraDataInterfaceArrayFloat3>(NiagaraSystem, OverrideName);
 }
 
+TArray<FVector> UNiagaraDataInterfaceArrayFunctionLibrary::GetNiagaraArrayPosition(UNiagaraComponent* NiagaraSystem, FName OverrideName)
+{
+	const TArray<FNiagaraPosition> SimData = GetNiagaraArray<FNiagaraPosition, UNiagaraDataInterfaceArrayPosition>(NiagaraSystem, OverrideName);
+	
+	FNiagaraLWCConverter LwcConverter = GetLWCConverter(NiagaraSystem);
+	TArray<FVector> ConvertedData;
+	ConvertedData.SetNumUninitialized(SimData.Num());
+	for (int i = 0; i < SimData.Num(); i++)
+	{
+		ConvertedData[i] = LwcConverter.ConvertSimulationPositionToWorld(SimData[i]);
+	}
+	
+	return ConvertedData;
+}
+
 TArray<FVector4> UNiagaraDataInterfaceArrayFunctionLibrary::GetNiagaraArrayVector4(UNiagaraComponent* NiagaraSystem, FName OverrideName)
 {
 	return GetNiagaraArray<FVector4, UNiagaraDataInterfaceArrayFloat4>(NiagaraSystem, OverrideName);
@@ -153,6 +188,13 @@ void UNiagaraDataInterfaceArrayFunctionLibrary::SetNiagaraArrayVectorValue(UNiag
 	SetNiagaraArrayValue<FVector, UNiagaraDataInterfaceArrayFloat3>(NiagaraSystem, OverrideName, Index, Value, bSizeToFit);
 }
 
+void UNiagaraDataInterfaceArrayFunctionLibrary::SetNiagaraArrayPositionValue(UNiagaraComponent* NiagaraSystem, FName OverrideName, int Index, const FVector& Value, bool bSizeToFit)
+{
+	FNiagaraLWCConverter LwcConverter = GetLWCConverter(NiagaraSystem);
+	FNiagaraPosition SimulationPosition = LwcConverter.ConvertWorldToSimulationPosition(Value);
+	SetNiagaraArrayValue<FNiagaraPosition, UNiagaraDataInterfaceArrayPosition>(NiagaraSystem, OverrideName, Index, SimulationPosition, bSizeToFit);
+}
+
 void UNiagaraDataInterfaceArrayFunctionLibrary::SetNiagaraArrayVector4Value(UNiagaraComponent* NiagaraSystem, FName OverrideName, int Index, const FVector4& Value, bool bSizeToFit)
 {
 	SetNiagaraArrayValue<FVector4, UNiagaraDataInterfaceArrayFloat4>(NiagaraSystem, OverrideName, Index, Value, bSizeToFit);
@@ -195,6 +237,13 @@ FVector UNiagaraDataInterfaceArrayFunctionLibrary::GetNiagaraArrayVectorValue(UN
 	return GetNiagaraArrayValue<FVector, UNiagaraDataInterfaceArrayFloat3>(NiagaraSystem, OverrideName, Index);
 }
 
+FVector UNiagaraDataInterfaceArrayFunctionLibrary::GetNiagaraArrayPositionValue(UNiagaraComponent* NiagaraSystem, FName OverrideName, int Index)
+{
+	FNiagaraPosition SimPosition = GetNiagaraArrayValue<FNiagaraPosition, UNiagaraDataInterfaceArrayPosition>(NiagaraSystem, OverrideName, Index);
+	FNiagaraLWCConverter LwcConverter = GetLWCConverter(NiagaraSystem);
+	return LwcConverter.ConvertSimulationPositionToWorld(SimPosition);
+}
+
 FVector4 UNiagaraDataInterfaceArrayFunctionLibrary::GetNiagaraArrayVector4Value(UNiagaraComponent* NiagaraSystem, FName OverrideName, int Index)
 {
 	return GetNiagaraArrayValue<FVector4, UNiagaraDataInterfaceArrayFloat4>(NiagaraSystem, OverrideName, Index);
diff --git a/Engine/Plugins/FX/Niagara/Source/Niagara/Private/NiagaraDataInterfaceAudioPlayer.cpp b/Engine/Plugins/FX/Niagara/Source/Niagara/Private/NiagaraDataInterfaceAudioPlayer.cpp
--- a/Engine/Plugins/FX/Niagara/Source/Niagara/Private/NiagaraDataInterfaceAudioPlayer.cpp
+++ b/Engine/Plugins/FX/Niagara/Source/Niagara/Private/NiagaraDataInterfaceAudioPlayer.cpp
@@ -27,6 +27,17 @@ const FName UNiagaraDataInterfaceAudioPlayer::SetPersistentAudioIntegerParamName
 const FName UNiagaraDataInterfaceAudioPlayer::SetPersistentAudioFloatParamName(TEXT("SetFloatParameter"));
 const FName UNiagaraDataInterfaceAudioPlayer::PausePersistentAudioName(TEXT("SetPaused"));
 
+struct FNiagaraAudioPlayerDIFunctionVersion
+{
+	enum Type
+	{
+		InitialVersion = 0,
+		LWCConversion = 1,
+
+		VersionPlusOne,
+		LatestVersion = VersionPlusOne - 1
+	};
+};
 
 /**
 Async task to play the audio on the game thread and isolate from the niagara tick
@@ -90,16 +101,22 @@ UNiagaraDataInterfaceAudioPlayer::UNiagaraDataInterfaceAudioPlayer(FObjectInitia
 #if WITH_EDITORONLY_DATA
 bool UNiagaraDataInterfaceAudioPlayer::UpgradeFunctionCall(FNiagaraFunctionSignature& FunctionSignature)
 {
-	bool bWasChanged = false;
-
-	// The play audio function was marked to have side effects, so older version need the exec pin 
-	if (FunctionSignature.Name == PlayAudioName && FunctionSignature.bRequiresExecPin == false)
+	// always upgrade to the latest version
+	if (FunctionSignature.FunctionVersion < FNiagaraAudioPlayerDIFunctionVersion::LatestVersion)
 	{
-		FunctionSignature.bRequiresExecPin = true;
-		bWasChanged = true;
+		TArray<FNiagaraFunctionSignature> AllFunctions;
+		GetFunctions(AllFunctions);
+		for (const FNiagaraFunctionSignature& Sig : AllFunctions)
+		{
+			if (FunctionSignature.Name == Sig.Name)
+			{
+				FunctionSignature = Sig;
+				return true;
+			}
+		}
 	}
 
-	return bWasChanged;
+	return false;
 }
 #endif
 
@@ -117,6 +134,7 @@ void UNiagaraDataInterfaceAudioPlayer::PostInitProperties()
 bool UNiagaraDataInterfaceAudioPlayer::InitPerInstanceData(void* PerInstanceData, FNiagaraSystemInstance* SystemInstance)
 {
 	FAudioPlayerInterface_InstanceData* PIData = new (PerInstanceData) FAudioPlayerInterface_InstanceData;
+	PIData->LWCConverter = SystemInstance->GetLWCConverter();
 	if (bLimitPlaysPerTick)
 	{
 		PIData->MaxPlaysPerTick = MaxPlaysPerTick;
@@ -253,6 +271,7 @@ void UNiagaraDataInterfaceAudioPlayer::GetFunctions(TArray<FNiagaraFunctionSigna
 	Sig.Name = PlayAudioName;
 #if WITH_EDITORONLY_DATA
 	Sig.Description = NSLOCTEXT("Niagara", "PlayAudioDIFunctionDescription", "This function plays a sound at the given location after the simulation has ticked.");
+	Sig.FunctionVersion = FNiagaraAudioPlayerDIFunctionVersion::LatestVersion;
 #endif
 	Sig.bMemberFunction = true;
 	Sig.bRequiresContext = false;
@@ -260,7 +279,7 @@ void UNiagaraDataInterfaceAudioPlayer::GetFunctions(TArray<FNiagaraFunctionSigna
 	Sig.bRequiresExecPin = true;
 	Sig.Inputs.Add(FNiagaraVariable(FNiagaraTypeDefinition(GetClass()), TEXT("Audio interface")));
 	Sig.Inputs.Add(FNiagaraVariable(FNiagaraTypeDefinition::GetBoolDef(), TEXT("Play Audio")));
-	Sig.Inputs.Add(FNiagaraVariable(FNiagaraTypeDefinition::GetVec3Def(), TEXT("PositionWS")));
+	Sig.Inputs.Add(FNiagaraVariable(FNiagaraTypeDefinition::GetPositionDef(), TEXT("PositionWS")));
 	Sig.Inputs.Add(FNiagaraVariable(FNiagaraTypeDefinition::GetVec3Def(), TEXT("RotationWS")));
 	Sig.Inputs.Add(FNiagaraVariable(FNiagaraTypeDefinition::GetFloatDef(), TEXT("VolumeFactor")));
 	Sig.Inputs.Add(FNiagaraVariable(FNiagaraTypeDefinition::GetFloatDef(), TEXT("PitchFactor")));
@@ -272,6 +291,7 @@ void UNiagaraDataInterfaceAudioPlayer::GetFunctions(TArray<FNiagaraFunctionSigna
 	Sig.Name = PlayPersistentAudioName;
 #if WITH_EDITORONLY_DATA
 	Sig.Description = NSLOCTEXT("Niagara", "PlayPersistentAudioDIFunctionDescription", "This function plays a sound at the given location after the simulation has ticked. The returned handle can be used to control the sound in subsequent ticks.");
+	Sig.FunctionVersion = FNiagaraAudioPlayerDIFunctionVersion::LatestVersion;
 #endif
 	Sig.bMemberFunction = true;
 	Sig.bRequiresContext = false;
@@ -280,7 +300,7 @@ void UNiagaraDataInterfaceAudioPlayer::GetFunctions(TArray<FNiagaraFunctionSigna
 	Sig.Inputs.Add(FNiagaraVariable(FNiagaraTypeDefinition(GetClass()), TEXT("Audio Interface")));
 	Sig.Inputs.Add(FNiagaraVariable(FNiagaraTypeDefinition::GetBoolDef(), TEXT("Play Audio")));
 	Sig.Inputs.Add(FNiagaraVariable(FNiagaraTypeDefinition::GetIntDef(), TEXT("Existing Audio Handle")));
-	Sig.Inputs.Add(FNiagaraVariable(FNiagaraTypeDefinition::GetVec3Def(), TEXT("Position WS")));
+	Sig.Inputs.Add(FNiagaraVariable(FNiagaraTypeDefinition::GetPositionDef(), TEXT("Position WS")));
 	Sig.Inputs.Add(FNiagaraVariable(FNiagaraTypeDefinition::GetVec3Def(), TEXT("Rotation WS")));
 	Sig.Inputs.Add(FNiagaraVariable(FNiagaraTypeDefinition::GetFloatDef(), TEXT("Volume Factor")));
 	Sig.Inputs.Add(FNiagaraVariable(FNiagaraTypeDefinition::GetFloatDef(), TEXT("Pitch Factor")));
@@ -294,6 +314,7 @@ void UNiagaraDataInterfaceAudioPlayer::GetFunctions(TArray<FNiagaraFunctionSigna
 	Sig.Name = SetPersistentAudioBoolParamName;
 #if WITH_EDITORONLY_DATA
 	Sig.Description = NSLOCTEXT("Niagara", "SetPersistentAudioBoolParamFunctionDescription", "If an active audio effect can be found for the given handle then the given sound cue parameter will be set on it.");
+	Sig.FunctionVersion = FNiagaraAudioPlayerDIFunctionVersion::LatestVersion;
 #endif
 	Sig.bMemberFunction = true;
 	Sig.bRequiresContext = false;
@@ -309,6 +330,7 @@ void UNiagaraDataInterfaceAudioPlayer::GetFunctions(TArray<FNiagaraFunctionSigna
 	Sig.Name = SetPersistentAudioIntegerParamName;
 #if WITH_EDITORONLY_DATA
 	Sig.Description = NSLOCTEXT("Niagara", "SetPersistentAudioIntegerParamFunctionDescription", "If an active audio effect can be found for the given handle then the given sound cue parameter will be set on it.");
+	Sig.FunctionVersion = FNiagaraAudioPlayerDIFunctionVersion::LatestVersion;
 #endif
 	Sig.bMemberFunction = true;
 	Sig.bRequiresContext = false;
@@ -324,6 +346,7 @@ void UNiagaraDataInterfaceAudioPlayer::GetFunctions(TArray<FNiagaraFunctionSigna
 	Sig.Name = SetPersistentAudioFloatParamName;
 #if WITH_EDITORONLY_DATA
 	Sig.Description = NSLOCTEXT("Niagara", "SetPersistentAudioFloatParamFunctionDescription", "If an active audio effect can be found for the given handle then the given sound cue parameter will be set on it.");
+	Sig.FunctionVersion = FNiagaraAudioPlayerDIFunctionVersion::LatestVersion;
 #endif
 	Sig.bMemberFunction = true;
 	Sig.bRequiresContext = false;
@@ -339,6 +362,7 @@ void UNiagaraDataInterfaceAudioPlayer::GetFunctions(TArray<FNiagaraFunctionSigna
 	Sig.Name = SetPersistentAudioVolumeName;
 #if WITH_EDITORONLY_DATA
 	Sig.Description = NSLOCTEXT("Niagara", "SetPersistentAudioVolumeFunctionDescription", "If an active audio effect can be found for the given handle then the this will adjusts its volume multiplier.");
+	Sig.FunctionVersion = FNiagaraAudioPlayerDIFunctionVersion::LatestVersion;
 #endif
 	Sig.bMemberFunction = true;
 	Sig.bRequiresContext = false;
@@ -353,6 +377,7 @@ void UNiagaraDataInterfaceAudioPlayer::GetFunctions(TArray<FNiagaraFunctionSigna
 	Sig.Name = SetPersistentAudioPitchName;
 #if WITH_EDITORONLY_DATA
 	Sig.Description = NSLOCTEXT("Niagara", "SetPersistentAudioPitchFunctionDescription", "If an active audio effect can be found for the given handle then the this will adjusts its pitch multiplier.");
+	Sig.FunctionVersion = FNiagaraAudioPlayerDIFunctionVersion::LatestVersion;
 #endif
 	Sig.bMemberFunction = true;
 	Sig.bRequiresContext = false;
@@ -367,6 +392,7 @@ void UNiagaraDataInterfaceAudioPlayer::GetFunctions(TArray<FNiagaraFunctionSigna
 	Sig.Name = SetPersistentAudioLocationName;
 #if WITH_EDITORONLY_DATA
 	Sig.Description = NSLOCTEXT("Niagara", "SetPersistentAudioLocationFunctionDescription", "If an active audio effect can be found for the given handle then the this will adjusts its world position.");
+	Sig.FunctionVersion = FNiagaraAudioPlayerDIFunctionVersion::LatestVersion;
 #endif
 	Sig.bMemberFunction = true;
 	Sig.bRequiresContext = false;
@@ -374,13 +400,14 @@ void UNiagaraDataInterfaceAudioPlayer::GetFunctions(TArray<FNiagaraFunctionSigna
 	Sig.bRequiresExecPin = true;
 	Sig.Inputs.Add(FNiagaraVariable(FNiagaraTypeDefinition(GetClass()), TEXT("Audio Interface")));
 	Sig.Inputs.Add(FNiagaraVariable(FNiagaraTypeDefinition::GetIntDef(), TEXT("Audio Handle")));
-	Sig.Inputs.Add(FNiagaraVariable(FNiagaraTypeDefinition::GetVec3Def(), TEXT("Position WS")));
+	Sig.Inputs.Add(FNiagaraVariable(FNiagaraTypeDefinition::GetPositionDef(), TEXT("Position WS")));
 	OutFunctions.Add(Sig);
 
 	Sig = FNiagaraFunctionSignature();
 	Sig.Name = SetPersistentAudioRotationName;
 #if WITH_EDITORONLY_DATA
 	Sig.Description = NSLOCTEXT("Niagara", "SetPersistentAudioRotationFunctionDescription", "If an active audio effect can be found for the given handle then the this will adjusts its rotation in the world.");
+	Sig.FunctionVersion = FNiagaraAudioPlayerDIFunctionVersion::LatestVersion;
 #endif
 	Sig.bMemberFunction = true;
 	Sig.bRequiresContext = false;
@@ -395,6 +422,7 @@ void UNiagaraDataInterfaceAudioPlayer::GetFunctions(TArray<FNiagaraFunctionSigna
 	Sig.Name = PausePersistentAudioName;
 #if WITH_EDITORONLY_DATA
 	Sig.Description = NSLOCTEXT("Niagara", "SetPersistentAudioPausedDescription", "If an active audio effect can be found for the given handle then the this will either pause or unpause the effect.");
+	Sig.FunctionVersion = FNiagaraAudioPlayerDIFunctionVersion::LatestVersion;
 #endif
 	Sig.bMemberFunction = true;
 	Sig.bRequiresContext = false;
@@ -635,7 +663,7 @@ void UNiagaraDataInterfaceAudioPlayer::UpdateLocation(FVectorVMExternalFunctionC
 	for (int32 i = 0; i < Context.GetNumInstances(); ++i)
 	{
 		int32 Handle = AudioHandleInParam.GetAndAdvance();
-		FVector Location = LocationParam.GetAndAdvance();
+		FVector Location = InstData->LWCConverter.ConvertSimulationVectorToWorld(LocationParam.GetAndAdvance());
 
 		if (Handle > 0)
 		{
@@ -741,7 +769,8 @@ void UNiagaraDataInterfaceAudioPlayer::PlayOneShotAudio(FVectorVMExternalFunctio
 	{
 		FNiagaraBool ShouldPlay = PlayDataParam.GetAndAdvance();
 		FAudioParticleData Data;
-		Data.Position = FVector(PositionParamX.GetAndAdvance(), PositionParamY.GetAndAdvance(), PositionParamZ.GetAndAdvance());
+		FNiagaraPosition SimulationPosition(PositionParamX.GetAndAdvance(), PositionParamY.GetAndAdvance(), PositionParamZ.GetAndAdvance());
+		Data.Position = InstData->LWCConverter.ConvertSimulationPositionToWorld(SimulationPosition);
 		Data.Rotation = FRotator(RotationParamX.GetAndAdvance(), RotationParamY.GetAndAdvance(), RotationParamZ.GetAndAdvance());
 		Data.Volume = VolumeParam.GetAndAdvance();
 		Data.Pitch = PitchParam.GetAndAdvance();
@@ -779,7 +808,7 @@ void UNiagaraDataInterfaceAudioPlayer::PlayPersistentAudio(FVectorVMExternalFunc
 	{
 		bool ShouldPlay = PlayAudioParam.GetAndAdvance();
 		int32 Handle = AudioHandleInParam.GetAndAdvance();
-		FVector Position = PositionParam.GetAndAdvance();
+		FVector Position = InstData->LWCConverter.ConvertSimulationVectorToWorld(PositionParam.GetAndAdvance());
 		FVector InRot = RotationParam.GetAndAdvance();
 		FRotator Rotation = FRotator(InRot.X, InRot.Y, InRot.Z);
 		float Volume = VolumeParam.GetAndAdvance();
diff --git a/Engine/Plugins/FX/Niagara/Source/Niagara/Private/NiagaraDataInterfaceCamera.cpp b/Engine/Plugins/FX/Niagara/Source/Niagara/Private/NiagaraDataInterfaceCamera.cpp
--- a/Engine/Plugins/FX/Niagara/Source/Niagara/Private/NiagaraDataInterfaceCamera.cpp
+++ b/Engine/Plugins/FX/Niagara/Source/Niagara/Private/NiagaraDataInterfaceCamera.cpp
@@ -6,6 +6,7 @@
 #include "Internationalization/Internationalization.h"
 #include "NiagaraSystemInstance.h"
 #include "GameFramework/PlayerController.h"
+#include "ShaderParameterUtils.h"
 
 #if WITH_EDITORONLY_DATA
 #include "EditorViewportClient.h"
@@ -21,6 +22,7 @@ struct FNiagaraCameraDIFunctionVersion
 		InitialVersion = 0,
         AddedPreViewTranslation = 1,
 		AddedBufferSizeInvSizeForGPU = 2,
+		LWCConversion = 3,
 
         VersionPlusOne,
         LatestVersion = VersionPlusOne - 1
@@ -77,6 +79,8 @@ bool UNiagaraDataInterfaceCamera::PerInstanceTick(void* PerInstanceData, FNiagar
 	PIData->ParticlesSortedByDistance.StableSort([](const FDistanceData& A, const FDistanceData& B) { return A.DistanceSquared < B.DistanceSquared; });
 
 	// grab the current camera data
+	FNiagaraLWCConverter LWCConverter = SystemInstance->GetLWCConverter();
+	PIData->LWCConverter = LWCConverter;
 	UWorld* World = SystemInstance->GetWorldManager()->GetWorld();
 	if (World && PlayerControllerIndex < World->GetNumPlayerControllers())
 	{
@@ -86,7 +90,7 @@ bool UNiagaraDataInterfaceCamera::PerInstanceTick(void* PerInstanceData, FNiagar
 			APlayerController* PlayerController = Iterator->Get();
 			if (i == PlayerControllerIndex && PlayerController)
 			{
-				PIData->CameraLocation = PlayerController->PlayerCameraManager->GetCameraLocation();
+				PIData->CameraLocation = LWCConverter.ConvertWorldToSimulationPosition(PlayerController->PlayerCameraManager->GetCameraLocation());
 				PIData->CameraRotation = PlayerController->PlayerCameraManager->GetCameraRotation();
 				PIData->CameraFOV = PlayerController->PlayerCameraManager->GetFOVAngle();
 				return false;
@@ -98,14 +102,14 @@ bool UNiagaraDataInterfaceCamera::PerInstanceTick(void* PerInstanceData, FNiagar
 	if (GCurrentLevelEditingViewportClient)
 	{
 		const FViewportCameraTransform& ViewTransform = GCurrentLevelEditingViewportClient->GetViewTransform();
-		PIData->CameraLocation = ViewTransform.GetLocation();
+		PIData->CameraLocation = LWCConverter.ConvertWorldToSimulationPosition(ViewTransform.GetLocation());
 		PIData->CameraRotation = ViewTransform.GetRotation();
 		PIData->CameraFOV = GCurrentLevelEditingViewportClient->ViewFOV;
 		return false;
 	}
 #endif
 
-	PIData->CameraLocation = FVector::ZeroVector;
+	PIData->CameraLocation = FNiagaraPosition(ForceInitToZero);
 	PIData->CameraRotation = FRotator(0);
 	PIData->CameraFOV = 0;
 	
@@ -124,7 +128,7 @@ void UNiagaraDataInterfaceCamera::GetFunctions(TArray<FNiagaraFunctionSignature>
 	Sig.bRequiresContext = false;
 	Sig.bSupportsCPU = false;
 	Sig.AddInput(FNiagaraVariable(FNiagaraTypeDefinition(GetClass()), TEXT("Camera interface")));
-	Sig.AddOutput(FNiagaraVariable(FNiagaraTypeDefinition::GetVec3Def(), TEXT("View Position World")), LOCTEXT("ViewPositionWorldDescription", "The camera position in world space."));
+	Sig.AddOutput(FNiagaraVariable(FNiagaraTypeDefinition::GetPositionDef(), TEXT("View Position World")), LOCTEXT("ViewPositionWorldDescription", "The camera position in world space."));
 	Sig.AddOutput(FNiagaraVariable(FNiagaraTypeDefinition::GetVec3Def(), TEXT("View Forward Vector")), LOCTEXT("ViewForwardVectorDescription", "The world space direction that the camera is pointing."));
 	Sig.AddOutput(FNiagaraVariable(FNiagaraTypeDefinition::GetVec3Def(), TEXT("View Up Vector")), LOCTEXT("ViewUpVectorDescription", "The camera's up direction in world space."));
 	Sig.AddOutput(FNiagaraVariable(FNiagaraTypeDefinition::GetVec3Def(), TEXT("View Right Vector")), LOCTEXT("ViewRightVectorDescription", "The camera's right direction in world space."));
@@ -200,7 +204,7 @@ void UNiagaraDataInterfaceCamera::GetFunctions(TArray<FNiagaraFunctionSignature>
 	Sig.bMemberFunction = true;
 	Sig.bRequiresContext = false;
 	Sig.AddInput(FNiagaraVariable(FNiagaraTypeDefinition(GetClass()), TEXT("Camera interface")));
-	Sig.AddOutput(FNiagaraVariable(FNiagaraTypeDefinition::GetVec3Def(), TEXT("Camera Position World")), LOCTEXT("CameraPositionWorldDescription", "The camera position in world space."));
+	Sig.AddOutput(FNiagaraVariable(FNiagaraTypeDefinition::GetPositionDef(), TEXT("Camera Position World")), LOCTEXT("CameraPositionWorldDescription", "The camera position in world space."));
 	Sig.AddOutput(FNiagaraVariable(FNiagaraTypeDefinition::GetVec3Def(), TEXT("Forward Vector World")), LOCTEXT("ForwardVectorWorldDescription", "The world space direction that the camera is pointing."));
 	Sig.AddOutput(FNiagaraVariable(FNiagaraTypeDefinition::GetVec3Def(), TEXT("Up Vector World")), LOCTEXT("UpVectorWorldDescription", "The camera's up direction in world space."));
 	Sig.AddOutput(FNiagaraVariable(FNiagaraTypeDefinition::GetVec3Def(), TEXT("Right Vector World")), LOCTEXT("RightVectorWorldDescription", "The camera's right direction in world space."));
@@ -233,13 +237,14 @@ void UNiagaraDataInterfaceCamera::GetFunctions(TArray<FNiagaraFunctionSignature>
 	Sig.bRequiresExecPin = true;
 	Sig.AddInput(FNiagaraVariable(FNiagaraTypeDefinition(GetClass()), TEXT("Camera interface")));
 	Sig.AddInput(FNiagaraVariable(FNiagaraTypeDefinition::GetIDDef(), TEXT("Particle ID")));
-	Sig.AddInput(FNiagaraVariable(FNiagaraTypeDefinition::GetVec3Def(), TEXT("Particle Position World")));
+	Sig.AddInput(FNiagaraVariable(FNiagaraTypeDefinition::GetPositionDef(), TEXT("Particle Position World")));
 	OutFunctions.Add(Sig);
 
 	Sig = FNiagaraFunctionSignature();
 	Sig.Name = GetTAAJitterName;
 #if WITH_EDITORONLY_DATA
 	Sig.Description = NSLOCTEXT("Niagara", "GetTAAJitterDescription", "This function returns the TAA jitter values of the currently active camera.");
+	Sig.FunctionVersion = FNiagaraCameraDIFunctionVersion::LatestVersion;
 #endif
 	Sig.bMemberFunction = true;
 	Sig.bRequiresContext = false;
@@ -252,26 +257,20 @@ void UNiagaraDataInterfaceCamera::GetFunctions(TArray<FNiagaraFunctionSignature>
 #if WITH_EDITORONLY_DATA
 bool UNiagaraDataInterfaceCamera::GetFunctionHLSL(const FNiagaraDataInterfaceGPUParamInfo& ParamInfo, const FNiagaraDataInterfaceGeneratedFunction& FunctionInfo, int FunctionInstanceIndex, FString& OutHLSL)
 {
-	TMap<FString, FStringFormatArg> ArgsSample;
-	ArgsSample.Add(TEXT("FunctionName"), FunctionInfo.InstanceName);
-
+	TMap<FString, FStringFormatArg> ArgsSample =
+	{
+		{TEXT("FunctionName"), FunctionInfo.InstanceName},
+		{TEXT("NDIGetContextName"), TEXT("NDICAMERA_MAKE_CONTEXT(") + ParamInfo.DataInterfaceHLSLSymbol + TEXT(")")},
+	};
+	
 	if (FunctionInfo.DefinitionName == GetViewPropertiesName)
 	{
 		static const TCHAR *FormatSample = TEXT(R"(
 			void {FunctionName}(out float3 Out_ViewPositionWorld, out float3 Out_ViewForwardVector, out float3 Out_ViewUpVector, out float3 Out_ViewRightVector, out float4 Out_ViewSizeAndInverseSize, out float4 Out_ScreenToViewSpace, out float2 Out_Current_TAAJitter, out float2 Out_Previous_TAAJitter, out float3 Out_PreViewTranslation, out float4 Out_BufferSizeAndInverseSize, out float2 Out_ViewportOffset)
 			{
-				Out_ViewPositionWorld.xyz = LWCHackToFloat(PrimaryView.WorldViewOrigin);
-				Out_ViewForwardVector.xyz = View.ViewForward.xyz;
-				Out_ViewUpVector.xyz = View.ViewUp.xyz;
-				Out_ViewRightVector.xyz = View.ViewRight.xyz;
-				Out_ViewSizeAndInverseSize = View.ViewSizeAndInvSize;
-				Out_ScreenToViewSpace = View.ScreenToViewSpace;
-				Out_Current_TAAJitter = View.TemporalAAJitter.xy;
-				Out_Previous_TAAJitter = View.TemporalAAJitter.zw;
-				Out_PreViewTranslation = LWCHackToFloat(PrimaryView.PreViewTranslation);
-				Out_BufferSizeAndInverseSize = View.BufferSizeAndInvSize;
-				Out_ViewportOffset = View.ViewRectMin.xy;
-			} 
+				{NDIGetContextName}
+				DICamera_GetViewProperties(DIContext, Out_ViewPositionWorld, Out_ViewForwardVector, Out_ViewUpVector, Out_ViewRightVector, Out_ViewSizeAndInverseSize, Out_ScreenToViewSpace, Out_Current_TAAJitter, Out_Previous_TAAJitter, Out_PreViewTranslation, Out_BufferSizeAndInverseSize, Out_ViewportOffset);
+			}
 		)");
 		OutHLSL += FString::Format(FormatSample, ArgsSample);
 		return true;
@@ -293,14 +292,8 @@ bool UNiagaraDataInterfaceCamera::GetFunctionHLSL(const FNiagaraDataInterfaceGPU
 			void {FunctionName}(out float4x4 Out_WorldToClipTransform, out float4x4 Out_TranslatedWorldToClipTransform, out float4x4 Out_ClipToWorldTransform, out float4x4 Out_ClipToViewTransform,
 				out float4x4 Out_ClipToTranslatedWorldTransform, out float4x4 Out_ScreenToWorldTransform, out float4x4 Out_ScreenToTranslatedWorldTransform, out float4x4 Out_ClipToPreviousClipTransform)
 			{
-				Out_WorldToClipTransform = LWCHackToFloat(PrimaryView.WorldToClip);
-				Out_TranslatedWorldToClipTransform = View.TranslatedWorldToClip;
-				Out_ClipToWorldTransform = LWCHackToFloat(PrimaryView.ClipToWorld);
-				Out_ClipToViewTransform = View.ClipToView;
-				Out_ClipToTranslatedWorldTransform = View.ClipToTranslatedWorld;
-				Out_ScreenToWorldTransform = LWCHackToFloat(PrimaryView.ScreenToWorld);
-				Out_ScreenToTranslatedWorldTransform = View.ScreenToTranslatedWorld;
-				Out_ClipToPreviousClipTransform = View.ClipToPrevClip;
+				{NDIGetContextName}
+				DICamera_GetClipSpaceTransforms(DIContext, Out_WorldToClipTransform, Out_TranslatedWorldToClipTransform, Out_ClipToWorldTransform, Out_ClipToViewTransform, Out_ClipToTranslatedWorldTransform, Out_ScreenToWorldTransform, Out_ScreenToTranslatedWorldTransform, Out_ClipToPreviousClipTransform);
 			}
 		)");
 		OutHLSL += FString::Format(FormatSample, ArgsSample);
@@ -312,12 +305,8 @@ bool UNiagaraDataInterfaceCamera::GetFunctionHLSL(const FNiagaraDataInterfaceGPU
 			void {FunctionName}(out float4x4 Out_TranslatedWorldToViewTransform, out float4x4 Out_ViewToTranslatedWorldTransform, out float4x4 Out_TranslatedWorldToCameraViewTransform,
 				out float4x4 Out_CameraViewToTranslatedWorldTransform, out float4x4 Out_ViewToClipTransform, out float4x4 Out_ViewToClipNoAATransform)
 			{
-				Out_TranslatedWorldToViewTransform = View.TranslatedWorldToView;
-				Out_ViewToTranslatedWorldTransform = View.ViewToTranslatedWorld;
-				Out_TranslatedWorldToCameraViewTransform = View.TranslatedWorldToCameraView;
-				Out_CameraViewToTranslatedWorldTransform = View.CameraViewToTranslatedWorld;
-				Out_ViewToClipTransform = View.ViewToClip;
-				Out_ViewToClipNoAATransform = View.ViewToClipNoAA;
+				{NDIGetContextName}
+				DICamera_GetViewSpaceTransforms(DIContext, Out_TranslatedWorldToViewTransform, Out_ViewToTranslatedWorldTransform, Out_TranslatedWorldToCameraViewTransform, Out_CameraViewToTranslatedWorldTransform, Out_ViewToClipTransform, Out_ViewToClipNoAATransform);
 			}
 		)");
 		OutHLSL += FString::Format(FormatSample, ArgsSample);
@@ -327,11 +316,9 @@ bool UNiagaraDataInterfaceCamera::GetFunctionHLSL(const FNiagaraDataInterfaceGPU
 	{
 		static const TCHAR *FormatSample = TEXT(R"(
 			void {FunctionName}(out float3 Out_CameraPositionWorld, out float3 Out_ViewForwardVector, out float3 Out_ViewUpVector, out float3 Out_ViewRightVector)
-			{				
-				Out_CameraPositionWorld.xyz = LWCHackToFloat(PrimaryView.WorldCameraOrigin);
-				Out_ViewForwardVector.xyz = View.ViewForward.xyz;
-				Out_ViewUpVector.xyz = View.ViewUp.xyz;
-				Out_ViewRightVector.xyz = View.ViewRight.xyz;
+			{
+				{NDIGetContextName}
+				DICamera_GetCameraProperties(DIContext, Out_CameraPositionWorld, Out_ViewForwardVector, Out_ViewUpVector, Out_ViewRightVector);
 			}
 		)");
 		OutHLSL += FString::Format(FormatSample, ArgsSample);
@@ -379,21 +366,54 @@ bool UNiagaraDataInterfaceCamera::UpgradeFunctionCall(FNiagaraFunctionSignature&
 		bWasChanged = true;
 	}
 
+	// LWC upgrades
+	if (FunctionSignature.FunctionVersion < FNiagaraCameraDIFunctionVersion::LWCConversion)
+	{
+		TArray<FNiagaraFunctionSignature> AllFunctions;
+		GetFunctions(AllFunctions);
+		for (const FNiagaraFunctionSignature& Sig : AllFunctions)
+		{
+			if (FunctionSignature.Name == Sig.Name)
+			{
+				FunctionSignature = Sig;
+				return true;
+			}
+		}
+	}
+
 	// Set latest version
     FunctionSignature.FunctionVersion = FNiagaraCameraDIFunctionVersion::LatestVersion;
 
 	return bWasChanged;
 }
+
+bool UNiagaraDataInterfaceCamera::AppendCompileHash(FNiagaraCompileHashVisitor* InVisitor) const
+{
+	if (!Super::AppendCompileHash(InVisitor))
+	{
+		return false;
+	}
+	FSHAHash Hash = GetShaderFileHash((TEXT("/Plugin/FX/Niagara/Private/NiagaraDataInterfaceCamera.ush")), EShaderPlatform::SP_PCD3D_SM5);
+	InVisitor->UpdateString(TEXT("NiagaraDataInterfaceCameraHLSLSource"), Hash.ToString());
+	return true;
+}
+
+void UNiagaraDataInterfaceCamera::GetCommonHLSL(FString& OutHLSL)
+{
+	OutHLSL += TEXT("#include \"/Plugin/FX/Niagara/Private/NiagaraDataInterfaceCamera.ush\"\n");
+}
+
+void UNiagaraDataInterfaceCamera::GetParameterDefinitionHLSL(const FNiagaraDataInterfaceGPUParamInfo& ParamInfo, FString& OutHLSL)
+{
+	Super::GetParameterDefinitionHLSL(ParamInfo, OutHLSL);
+	OutHLSL += TEXT("NDICAMERA_DECLARE_CONSTANTS(") + ParamInfo.DataInterfaceHLSLSymbol + TEXT(")\n");
+}
 #endif
 
 DEFINE_NDI_DIRECT_FUNC_BINDER(UNiagaraDataInterfaceCamera, GetClosestParticles);
 DEFINE_NDI_DIRECT_FUNC_BINDER(UNiagaraDataInterfaceCamera, CalculateParticleDistances);
 DEFINE_NDI_DIRECT_FUNC_BINDER(UNiagaraDataInterfaceCamera, GetCameraFOV);
 DEFINE_NDI_DIRECT_FUNC_BINDER(UNiagaraDataInterfaceCamera, GetCameraProperties);
-DEFINE_NDI_DIRECT_FUNC_BINDER(UNiagaraDataInterfaceCamera, GetViewPropertiesGPU);
-DEFINE_NDI_DIRECT_FUNC_BINDER(UNiagaraDataInterfaceCamera, GetClipSpaceTransformsGPU);
-DEFINE_NDI_DIRECT_FUNC_BINDER(UNiagaraDataInterfaceCamera, GetViewSpaceTransformsGPU);
-DEFINE_NDI_DIRECT_FUNC_BINDER(UNiagaraDataInterfaceCamera, GetTAAJitter);
 void UNiagaraDataInterfaceCamera::GetVMExternalFunction(const FVMExternalFunctionBindingInfo& BindingInfo, void* InstanceData, FVMExternalFunction &OutFunc)
 {
 	if (BindingInfo.Name == GetFieldOfViewName)
@@ -412,22 +432,6 @@ void UNiagaraDataInterfaceCamera::GetVMExternalFunction(const FVMExternalFunctio
 	{
 		NDI_FUNC_BINDER(UNiagaraDataInterfaceCamera, GetCameraProperties)::Bind(this, OutFunc);
 	}
-	else if (BindingInfo.Name == GetViewSpaceTransformsName)
-	{
-		NDI_FUNC_BINDER(UNiagaraDataInterfaceCamera, GetViewSpaceTransformsGPU)::Bind(this, OutFunc);
-	}
-	else if (BindingInfo.Name == GetClipSpaceTransformsName)
-	{
-		NDI_FUNC_BINDER(UNiagaraDataInterfaceCamera, GetClipSpaceTransformsGPU)::Bind(this, OutFunc);
-	}
-	else if (BindingInfo.Name == GetViewPropertiesName)
-	{
-		NDI_FUNC_BINDER(UNiagaraDataInterfaceCamera, GetViewPropertiesGPU)::Bind(this, OutFunc);
-	}
-	else if (BindingInfo.Name == GetTAAJitterName)
-	{
-		NDI_FUNC_BINDER(UNiagaraDataInterfaceCamera, GetTAAJitter)::Bind(this, OutFunc);
-	}
 	else
 	{
 		UE_LOG(LogNiagara, Display, TEXT("Could not find data interface external function in %s. Received Name: %s"), *GetPathNameSafe(this), *BindingInfo.Name.ToString());
@@ -535,13 +539,13 @@ void UNiagaraDataInterfaceCamera::CalculateParticleDistances(FVectorVMExternalFu
 	VectorVM::FUserPtrHandler<FCameraDataInterface_InstanceData> InstData(Context);
 
 	FNDIInputParam<FNiagaraID> IDParam(Context);
-	FNDIInputParam<FVector3f> ParticlePosParam(Context);
+	FNDIInputParam<FNiagaraPosition> ParticlePosParam(Context);
 
 	FVector CameraPos = InstData->CameraLocation;
 	for (int32 i = 0; i < Context.GetNumInstances(); ++i)
 	{
 		FDistanceData DistanceData;
-		FVector ParticlePos = ParticlePosParam.GetAndAdvance();
+		FVector ParticlePos = InstData->LWCConverter.ConvertSimulationPositionToWorld(ParticlePosParam.GetAndAdvance());
 		DistanceData.ParticleID = IDParam.GetAndAdvance();
 		DistanceData.DistanceSquared = (ParticlePos - CameraPos).SizeSquared();
 		InstData->DistanceSortQueue.Enqueue(DistanceData);
@@ -620,90 +624,6 @@ void UNiagaraDataInterfaceCamera::GetFeedback(UNiagaraSystem* Asset, UNiagaraCom
 
 #endif
 
-// ------- Dummy implementations for CPU execution ------------
-
-void UNiagaraDataInterfaceCamera::GetViewPropertiesGPU(FVectorVMExternalFunctionContext& Context)
-{
-	VectorVM::FUserPtrHandler<FCameraDataInterface_InstanceData> InstData(Context);
-	TArray<VectorVM::FExternalFuncRegisterHandler<float>> OutParams;
-	OutParams.Reserve(24);
-	for (int i = 0; i < 24; i++)
-	{
-		OutParams.Emplace(Context);
-	}
-
-	for (int32 k = 0; k < Context.GetNumInstances(); ++k)
-	{
-		for (int i = 0; i < 24; i++)
-		{
-			*OutParams[i].GetDestAndAdvance() = 0;
-		}
-	}
-}
-
-void UNiagaraDataInterfaceCamera::GetClipSpaceTransformsGPU(FVectorVMExternalFunctionContext& Context)
-{
-	VectorVM::FUserPtrHandler<FCameraDataInterface_InstanceData> InstData(Context);
-	TArray<VectorVM::FExternalFuncRegisterHandler<float>> OutParams;
-	OutParams.Reserve(128);
-	for (int i = 0; i < 128; i++)
-	{
-		OutParams.Emplace(Context);
-	}
-
-	for (int32 k = 0; k < Context.GetNumInstances(); ++k)
-	{
-		for (int i = 0; i < 128; i++)
-		{
-			*OutParams[i].GetDestAndAdvance() = 0;
-		}
-	}
-}
-
-void UNiagaraDataInterfaceCamera::GetViewSpaceTransformsGPU(FVectorVMExternalFunctionContext& Context)
-{
-	VectorVM::FUserPtrHandler<FCameraDataInterface_InstanceData> InstData(Context);
-	TArray<VectorVM::FExternalFuncRegisterHandler<float>> OutParams;
-	OutParams.Reserve(96);
-	for (int i = 0; i < 96; i++)
-	{
-		OutParams.Emplace(Context);
-	}
-
-	for (int32 k = 0; k < Context.GetNumInstances(); ++k)
-	{
-		for (int i = 0; i < 96; i++)
-		{
-			*OutParams[i].GetDestAndAdvance() = 0;
-		}
-	}
-}
-
-
-void UNiagaraDataInterfaceCamera::GetTAAJitter(FVectorVMExternalFunctionContext& Context)
-{
-	VectorVM::FUserPtrHandler<FCameraDataInterface_InstanceData> InstData(Context);
-	TArray<VectorVM::FExternalFuncRegisterHandler<float>> OutParams;
-
-	constexpr int32 ElementCount = 2 * sizeof(FVector2f) / sizeof(float);
-	OutParams.Reserve(ElementCount);
-	for (int i = 0; i < ElementCount; i++)
-	{
-		OutParams.Emplace(Context);
-	}
-
-	for (int32 k = 0; k < Context.GetNumInstances(); ++k)
-	{
-		for (int i = 0; i < ElementCount; i++)
-		{
-			*OutParams[i].GetDestAndAdvance() = 0;
-		}
-	}
-}
-
-// ------------------------------------------------------------
-
-
 bool UNiagaraDataInterfaceCamera::CopyToInternal(UNiagaraDataInterface* Destination) const
 {
 	if (!Super::CopyToInternal(Destination))
@@ -728,4 +648,34 @@ bool UNiagaraDataInterfaceCamera::Equals(const UNiagaraDataInterface* Other) con
 		OtherTyped->bRequireCurrentFrameData == bRequireCurrentFrameData ;
 }
 
+struct FNiagaraDataInterfaceParametersCS_Camera : public FNiagaraDataInterfaceParametersCS
+{
+	DECLARE_TYPE_LAYOUT(FNiagaraDataInterfaceParametersCS_Camera, NonVirtual);
+
+	void Bind(const FNiagaraDataInterfaceGPUParamInfo& ParameterInfo, const class FShaderParameterMap& ParameterMap)
+	{
+		SystemLWCTileParam.Bind(ParameterMap, *(SystemLWCTileName + ParameterInfo.DataInterfaceHLSLSymbol));
+	}
+
+	void Set(FRHICommandList& RHICmdList, const FNiagaraDataInterfaceSetArgs& Context) const
+	{
+		check(IsInRenderingThread());
+
+		FRHIComputeShader* ComputeShaderRHI = Context.Shader.GetComputeShader();
+		SetShaderValue(RHICmdList, ComputeShaderRHI, SystemLWCTileParam, Context.SystemLWCTile);
+	}
+
+private:
+	// LWC data
+	LAYOUT_FIELD(FShaderParameter, SystemLWCTileParam);
+
+	static const FString SystemLWCTileName;
+};
+
+// LWC
+const FString FNiagaraDataInterfaceParametersCS_Camera::SystemLWCTileName(TEXT("SystemLWCTile_"));
+
+IMPLEMENT_TYPE_LAYOUT(FNiagaraDataInterfaceParametersCS_Camera);
+IMPLEMENT_NIAGARA_DI_PARAMETER(UNiagaraDataInterfaceCamera, FNiagaraDataInterfaceParametersCS_Camera);
+
 #undef LOCTEXT_NAMESPACE
diff --git a/Engine/Plugins/FX/Niagara/Source/Niagara/Private/NiagaraDataInterfaceCollisionQuery.cpp b/Engine/Plugins/FX/Niagara/Source/Niagara/Private/NiagaraDataInterfaceCollisionQuery.cpp
--- a/Engine/Plugins/FX/Niagara/Source/Niagara/Private/NiagaraDataInterfaceCollisionQuery.cpp
+++ b/Engine/Plugins/FX/Niagara/Source/Niagara/Private/NiagaraDataInterfaceCollisionQuery.cpp
@@ -11,7 +11,6 @@
 #include "NiagaraStats.h"
 #include "NiagaraTypes.h"
 #include "NiagaraWorldManager.h"
-#include "RayTracingInstanceUtils.h"
 #include "RenderResource.h"
 #include "Shader.h"
 #include "ShaderCore.h"
@@ -59,6 +58,7 @@ namespace NDICollisionQueryLocal
 	static const FString IntersectionResultsOffsetParamName(TEXT("IntersectionResultsOffset_"));
 	static const FString RayTraceCountsParamName(TEXT("RayTraceCounts_"));
 	static const FString RayTraceCountsOffsetParamName(TEXT("RayTraceCountsOffset_"));
+	static const FString SystemLWCTileName(TEXT("SystemLWCTile_"));
 }
 
 FCriticalSection UNiagaraDataInterfaceCollisionQuery::CriticalSection;
@@ -71,6 +71,7 @@ struct FNiagaraCollisionDIFunctionVersion
 		AddedTraceSkip = 1,
 		AddedCustomDepthCollision = 2,
 		ReturnCollisionMaterialIdx = 3,
+		LargeWorldCoordinates = 4,
 
 		VersionPlusOne,
 		LatestVersion = VersionPlusOne - 1
@@ -270,70 +271,72 @@ void UNiagaraDataInterfaceCollisionQuery::GetAssetTagsForContext(const UObject*
 
 void UNiagaraDataInterfaceCollisionQuery::GetFunctions(TArray<FNiagaraFunctionSignature>& OutFunctions)
 {
-	FNiagaraFunctionSignature SigDepth;
-	SigDepth.Name = NDICollisionQueryLocal::SceneDepthName;
-	SigDepth.bMemberFunction = true;
-	SigDepth.bRequiresContext = false;
-	SigDepth.bSupportsCPU = false;
+	{
+		FNiagaraFunctionSignature SigDepth;
+		SigDepth.Name = NDICollisionQueryLocal::SceneDepthName;
+		SigDepth.bMemberFunction = true;
+		SigDepth.bSupportsCPU = false;
 #if WITH_EDITORONLY_DATA
-	SigDepth.FunctionVersion = FNiagaraCollisionDIFunctionVersion::LatestVersion;
-	SigDepth.Description = LOCTEXT("SceneDepthSignatureDescription", "Projects a given world position to view space and then queries the depth buffer with that position.");
+		SigDepth.FunctionVersion = FNiagaraCollisionDIFunctionVersion::LatestVersion;
+		SigDepth.Description = LOCTEXT("SceneDepthSignatureDescription", "Projects a given world position to view space and then queries the depth buffer with that position.");
 #endif
-	const FText DepthSamplePosWorldDescription = LOCTEXT("DepthSamplePosWorldDescription", "The world position where the depth should be queried. The position gets automatically transformed to view space to query the depth buffer.");
-	const FText SceneDepthDescription = LOCTEXT("SceneDepthDescription", "If the query was successful this returns the scene depth, otherwise -1.");
-	const FText CameraPosWorldDescription = LOCTEXT("CameraPosWorldDescription", "Returns the current camera position in world space.");
-	const FText IsInsideViewDescription = LOCTEXT("IsInsideViewDescription", "Returns true if the query position could be projected to valid screen coordinates.");
-	const FText SamplePosWorldDescription = LOCTEXT("SamplePosWorldDescription", "If the query was successful, this returns the world position that was recalculated from the scene depth. Otherwise returns (0, 0, 0).");
-	const FText SampleWorldNormalDescription = LOCTEXT("SampleWorldNormalDescription", "If the query was successful, this returns the world normal at the sample point. Otherwise returns (0, 0, 1).");
-
-	SigDepth.Inputs.Add(FNiagaraVariable(FNiagaraTypeDefinition(GetClass()), TEXT("CollisionQuery")));
-	SigDepth.AddInput(FNiagaraVariable(FNiagaraTypeDefinition::GetVec3Def(), TEXT("DepthSamplePosWorld")), DepthSamplePosWorldDescription);
-	SigDepth.AddOutput(FNiagaraVariable(FNiagaraTypeDefinition::GetFloatDef(), TEXT("SceneDepth")), SceneDepthDescription);
-	SigDepth.AddOutput(FNiagaraVariable(FNiagaraTypeDefinition::GetVec3Def(), TEXT("CameraPosWorld")), CameraPosWorldDescription);
-	SigDepth.AddOutput(FNiagaraVariable(FNiagaraTypeDefinition::GetBoolDef(), TEXT("IsInsideView")), IsInsideViewDescription);
-	SigDepth.AddOutput(FNiagaraVariable(FNiagaraTypeDefinition::GetVec3Def(), TEXT("SamplePosWorld")), SamplePosWorldDescription);
-	SigDepth.AddOutput(FNiagaraVariable(FNiagaraTypeDefinition::GetVec3Def(), TEXT("SampleWorldNormal")), SampleWorldNormalDescription);
+		const FText DepthSamplePosWorldDescription = LOCTEXT("DepthSamplePosWorldDescription", "The world position where the depth should be queried. The position gets automatically transformed to view space to query the depth buffer.");
+		const FText SceneDepthDescription = LOCTEXT("SceneDepthDescription", "If the query was successful this returns the scene depth, otherwise -1.");
+		const FText CameraPosWorldDescription = LOCTEXT("CameraPosWorldDescription", "Returns the current camera position in world space.");
+		const FText IsInsideViewDescription = LOCTEXT("IsInsideViewDescription", "Returns true if the query position could be projected to valid screen coordinates.");
+		const FText SamplePosWorldDescription = LOCTEXT("SamplePosWorldDescription", "If the query was successful, this returns the world position that was recalculated from the scene depth. Otherwise returns (0, 0, 0).");
+		const FText SampleWorldNormalDescription = LOCTEXT("SampleWorldNormalDescription", "If the query was successful, this returns the world normal at the sample point. Otherwise returns (0, 0, 1).");
+
+		SigDepth.Inputs.Add(FNiagaraVariable(FNiagaraTypeDefinition(GetClass()), TEXT("CollisionQuery")));
+		SigDepth.AddInput(FNiagaraVariable(FNiagaraTypeDefinition::GetPositionDef(), TEXT("DepthSamplePosWorld")), DepthSamplePosWorldDescription);
+		SigDepth.AddOutput(FNiagaraVariable(FNiagaraTypeDefinition::GetFloatDef(), TEXT("SceneDepth")), SceneDepthDescription);
+		SigDepth.AddOutput(FNiagaraVariable(FNiagaraTypeDefinition::GetPositionDef(), TEXT("CameraPosWorld")), CameraPosWorldDescription);
+		SigDepth.AddOutput(FNiagaraVariable(FNiagaraTypeDefinition::GetBoolDef(), TEXT("IsInsideView")), IsInsideViewDescription);
+		SigDepth.AddOutput(FNiagaraVariable(FNiagaraTypeDefinition::GetPositionDef(), TEXT("SamplePosWorld")), SamplePosWorldDescription);
+		SigDepth.AddOutput(FNiagaraVariable(FNiagaraTypeDefinition::GetVec3Def(), TEXT("SampleWorldNormal")), SampleWorldNormalDescription);
 	
-	OutFunctions.Add(SigDepth);
+		OutFunctions.Add(SigDepth);
 
-	FNiagaraFunctionSignature SigCustomDepth;
-	SigCustomDepth.Name = NDICollisionQueryLocal::CustomDepthName;
-	SigCustomDepth.bMemberFunction = true;
-	SigCustomDepth.bRequiresContext = false;
-	SigCustomDepth.bSupportsCPU = false;
+		
+		FNiagaraFunctionSignature SigCustomDepth;
+		SigCustomDepth.Name = NDICollisionQueryLocal::CustomDepthName;
+		SigCustomDepth.bMemberFunction = true;
+		SigCustomDepth.bSupportsCPU = false;
 #if WITH_EDITORONLY_DATA
-	SigCustomDepth.FunctionVersion = FNiagaraCollisionDIFunctionVersion::LatestVersion;
-	SigCustomDepth.Description = LOCTEXT("CustomDepthDescription", "Projects a given world position to view space and then queries the custom depth buffer with that position.");
+		SigCustomDepth.FunctionVersion = FNiagaraCollisionDIFunctionVersion::LatestVersion;
+		SigCustomDepth.Description = LOCTEXT("CustomDepthDescription", "Projects a given world position to view space and then queries the custom depth buffer with that position.");
 #endif
-	SigCustomDepth.Inputs.Add(FNiagaraVariable(FNiagaraTypeDefinition(GetClass()), TEXT("CollisionQuery")));
-	SigCustomDepth.AddInput(FNiagaraVariable(FNiagaraTypeDefinition::GetVec3Def(), TEXT("DepthSamplePosWorld")), DepthSamplePosWorldDescription);
-	SigCustomDepth.AddOutput(FNiagaraVariable(FNiagaraTypeDefinition::GetFloatDef(), TEXT("SceneDepth")), SceneDepthDescription);
-	SigCustomDepth.AddOutput(FNiagaraVariable(FNiagaraTypeDefinition::GetVec3Def(), TEXT("CameraPosWorld")), CameraPosWorldDescription);
-	SigCustomDepth.AddOutput(FNiagaraVariable(FNiagaraTypeDefinition::GetBoolDef(), TEXT("IsInsideView")), IsInsideViewDescription);
-	SigCustomDepth.AddOutput(FNiagaraVariable(FNiagaraTypeDefinition::GetVec3Def(), TEXT("SamplePosWorld")), SamplePosWorldDescription);
-	SigCustomDepth.AddOutput(FNiagaraVariable(FNiagaraTypeDefinition::GetVec3Def(), TEXT("SampleWorldNormal")), SampleWorldNormalDescription);
-	OutFunctions.Add(SigCustomDepth);
-
-	FNiagaraFunctionSignature SigMeshField;
-	SigMeshField.Name = NDICollisionQueryLocal::DistanceFieldName;
-	SigMeshField.bMemberFunction = true;
-	SigMeshField.bRequiresContext = false;
-	SigMeshField.bSupportsCPU = false;
+		SigCustomDepth.Inputs.Add(FNiagaraVariable(FNiagaraTypeDefinition(GetClass()), TEXT("CollisionQuery")));
+		SigCustomDepth.AddInput(FNiagaraVariable(FNiagaraTypeDefinition::GetPositionDef(), TEXT("DepthSamplePosWorld")), DepthSamplePosWorldDescription);
+		SigCustomDepth.AddOutput(FNiagaraVariable(FNiagaraTypeDefinition::GetFloatDef(), TEXT("SceneDepth")), SceneDepthDescription);
+		SigCustomDepth.AddOutput(FNiagaraVariable(FNiagaraTypeDefinition::GetPositionDef(), TEXT("CameraPosWorld")), CameraPosWorldDescription);
+		SigCustomDepth.AddOutput(FNiagaraVariable(FNiagaraTypeDefinition::GetBoolDef(), TEXT("IsInsideView")), IsInsideViewDescription);
+		SigCustomDepth.AddOutput(FNiagaraVariable(FNiagaraTypeDefinition::GetPositionDef(), TEXT("SamplePosWorld")), SamplePosWorldDescription);
+		SigCustomDepth.AddOutput(FNiagaraVariable(FNiagaraTypeDefinition::GetVec3Def(), TEXT("SampleWorldNormal")), SampleWorldNormalDescription);
+		OutFunctions.Add(SigCustomDepth);
+	}
+
+	{
+		FNiagaraFunctionSignature SigMeshField;
+		SigMeshField.Name = NDICollisionQueryLocal::DistanceFieldName;
+		SigMeshField.bMemberFunction = true;
+		SigMeshField.bSupportsCPU = false;
 #if WITH_EDITORONLY_DATA
-	SigMeshField.FunctionVersion = FNiagaraCollisionDIFunctionVersion::LatestVersion;
-	SigMeshField.Description = LOCTEXT("DistanceFieldDescription", "Queries the global distance field for a given world position.\nPlease note that the distance field resolution gets lower the farther away the queried position is from the camera.");
+		SigMeshField.FunctionVersion = FNiagaraCollisionDIFunctionVersion::LatestVersion;
+		SigMeshField.Description = LOCTEXT("DistanceFieldDescription", "Queries the global distance field for a given world position.\nPlease note that the distance field resolution gets lower the farther away the queried position is from the camera.");
 #endif
-	const FText FieldSamplePosWorldDescription = LOCTEXT("FieldSamplePosWorldDescription", "The world position where the distance field should be queried.");
-	const FText DistanceToNearestSurfaceDescription = LOCTEXT("DistanceToNearestSurfaceDescription", "If the query was successful this returns the distance to the nearest surface, otherwise returns 0.");
-	const FText FieldGradientDescription = LOCTEXT("FieldGradientDescription", "If the query was successful this returns the non-normalized direction to the nearest surface, otherwise returns (0, 0, 0).");
-	const FText IsDistanceFieldValidDescription = LOCTEXT("IsDistanceFieldValidDescription", "Returns true if the global distance field is available and there was a valid value retrieved for the given sample position.");
+		const FText FieldSamplePosWorldDescription = LOCTEXT("FieldSamplePosWorldDescription", "The world position where the distance field should be queried.");
+		const FText DistanceToNearestSurfaceDescription = LOCTEXT("DistanceToNearestSurfaceDescription", "If the query was successful this returns the distance to the nearest surface, otherwise returns 0.");
+		const FText FieldGradientDescription = LOCTEXT("FieldGradientDescription", "If the query was successful this returns the non-normalized direction to the nearest surface, otherwise returns (0, 0, 0).");
+		const FText IsDistanceFieldValidDescription = LOCTEXT("IsDistanceFieldValidDescription", "Returns true if the global distance field is available and there was a valid value retrieved for the given sample position.");
 	
-	SigMeshField.Inputs.Add(FNiagaraVariable(FNiagaraTypeDefinition(GetClass()), TEXT("CollisionQuery")));
-	SigMeshField.AddInput(FNiagaraVariable(FNiagaraTypeDefinition::GetVec3Def(), TEXT("FieldSamplePosWorld")), FieldSamplePosWorldDescription);
-	SigMeshField.AddOutput(FNiagaraVariable(FNiagaraTypeDefinition::GetFloatDef(), TEXT("DistanceToNearestSurface")), DistanceToNearestSurfaceDescription);
-	SigMeshField.AddOutput(FNiagaraVariable(FNiagaraTypeDefinition::GetVec3Def(), TEXT("FieldGradient")), FieldGradientDescription);
-	SigMeshField.AddOutput(FNiagaraVariable(FNiagaraTypeDefinition::GetBoolDef(), TEXT("IsDistanceFieldValid")), IsDistanceFieldValidDescription);
-    OutFunctions.Add(SigMeshField);
+		SigMeshField.Inputs.Add(FNiagaraVariable(FNiagaraTypeDefinition(GetClass()), TEXT("CollisionQuery")));
+		SigMeshField.AddInput(FNiagaraVariable(FNiagaraTypeDefinition::GetPositionDef(), TEXT("FieldSamplePosWorld")), FieldSamplePosWorldDescription);
+		SigMeshField.AddOutput(FNiagaraVariable(FNiagaraTypeDefinition::GetFloatDef(), TEXT("DistanceToNearestSurface")), DistanceToNearestSurfaceDescription);
+		SigMeshField.AddOutput(FNiagaraVariable(FNiagaraTypeDefinition::GetVec3Def(), TEXT("FieldGradient")), FieldGradientDescription);
+		SigMeshField.AddOutput(FNiagaraVariable(FNiagaraTypeDefinition::GetBoolDef(), TEXT("IsDistanceFieldValid")), IsDistanceFieldValidDescription);
+		OutFunctions.Add(SigMeshField);
+	}
 
 	{
 		const FText QueryIDDescription = LOCTEXT("QueryIDDescription", "Unique (for this frame) index of the query being enqueued (used in subsequent frames to retrieve results).  Must be less than MaxRayTraceCount");
@@ -346,7 +349,6 @@ void UNiagaraDataInterfaceCollisionQuery::GetFunctions(TArray<FNiagaraFunctionSi
 		IssueRayTrace.Name = NDICollisionQueryLocal::IssueAsyncRayTraceName;
 		IssueRayTrace.bRequiresExecPin = true;
 		IssueRayTrace.bMemberFunction = true;
-		IssueRayTrace.bRequiresContext = false;
 		IssueRayTrace.bSupportsCPU = false;
 #if WITH_EDITORONLY_DATA
 		IssueRayTrace.FunctionVersion = FNiagaraCollisionDIFunctionVersion::LatestVersion;
@@ -354,8 +356,8 @@ void UNiagaraDataInterfaceCollisionQuery::GetFunctions(TArray<FNiagaraFunctionSi
 #endif
 		IssueRayTrace.AddInput(FNiagaraVariable(FNiagaraTypeDefinition(GetClass()), TEXT("CollisionQuery")));
 		IssueRayTrace.AddInput(FNiagaraVariable(FNiagaraTypeDefinition::GetIntDef(), TEXT("QueryID")), QueryIDDescription);
-		IssueRayTrace.AddInput(FNiagaraVariable(FNiagaraTypeDefinition::GetVec3Def(), TEXT("TraceStartWorld")), TraceStartWorldDescription);
-		IssueRayTrace.AddInput(FNiagaraVariable(FNiagaraTypeDefinition::GetVec3Def(), TEXT("TraceEndWorld")), TraceEndWorldDescription);
+		IssueRayTrace.AddInput(FNiagaraVariable(FNiagaraTypeDefinition::GetPositionDef(), TEXT("TraceStartWorld")), TraceStartWorldDescription);
+		IssueRayTrace.AddInput(FNiagaraVariable(FNiagaraTypeDefinition::GetPositionDef(), TEXT("TraceEndWorld")), TraceEndWorldDescription);
 		IssueRayTrace.AddInput(FNiagaraVariable(FNiagaraTypeDefinition::GetIntDef(), TEXT("TraceChannel")), TraceChannelDescription);
 		IssueRayTrace.AddOutput(FNiagaraVariable(FNiagaraTypeDefinition::GetBoolDef(), TEXT("IsQueryValid")), IsQueryValidDescription);
 	}
@@ -371,15 +373,14 @@ void UNiagaraDataInterfaceCollisionQuery::GetFunctions(TArray<FNiagaraFunctionSi
 		IssueRayTrace.Name = NDICollisionQueryLocal::CreateAsyncRayTraceName;
 		IssueRayTrace.bRequiresExecPin = true;
 		IssueRayTrace.bMemberFunction = true;
-		IssueRayTrace.bRequiresContext = false;
 		IssueRayTrace.bSupportsCPU = false;
 #if WITH_EDITORONLY_DATA
 		IssueRayTrace.FunctionVersion = FNiagaraCollisionDIFunctionVersion::LatestVersion;
 		IssueRayTrace.Description = LOCTEXT("IssueAsyncRayTraceDescription", "Creates a GPU raytrace with the result being available the following frame (index is returned)");
 #endif
 		IssueRayTrace.AddInput(FNiagaraVariable(FNiagaraTypeDefinition(GetClass()), TEXT("CollisionQuery")));
-		IssueRayTrace.AddInput(FNiagaraVariable(FNiagaraTypeDefinition::GetVec3Def(), TEXT("TraceStartWorld")), TraceStartWorldDescription);
-		IssueRayTrace.AddInput(FNiagaraVariable(FNiagaraTypeDefinition::GetVec3Def(), TEXT("TraceEndWorld")), TraceEndWorldDescription);
+		IssueRayTrace.AddInput(FNiagaraVariable(FNiagaraTypeDefinition::GetPositionDef(), TEXT("TraceStartWorld")), TraceStartWorldDescription);
+		IssueRayTrace.AddInput(FNiagaraVariable(FNiagaraTypeDefinition::GetPositionDef(), TEXT("TraceEndWorld")), TraceEndWorldDescription);
 		IssueRayTrace.AddInput(FNiagaraVariable(FNiagaraTypeDefinition::GetIntDef(), TEXT("TraceChannel")), TraceChannelDescription);
 		IssueRayTrace.AddOutput(FNiagaraVariable(FNiagaraTypeDefinition::GetIntDef(), TEXT("QueryID")), QueryIDDescription);
 		IssueRayTrace.AddOutput(FNiagaraVariable(FNiagaraTypeDefinition::GetBoolDef(), TEXT("IsQueryValid")), IsQueryValidDescription);
@@ -394,7 +395,6 @@ void UNiagaraDataInterfaceCollisionQuery::GetFunctions(TArray<FNiagaraFunctionSi
 		IssueRayTrace.Name = NDICollisionQueryLocal::ReserveAsyncRayTraceName;
 		IssueRayTrace.bRequiresExecPin = true;
 		IssueRayTrace.bMemberFunction = true;
-		IssueRayTrace.bRequiresContext = false;
 		IssueRayTrace.bSupportsCPU = false;
 #if WITH_EDITORONLY_DATA
 		IssueRayTrace.FunctionVersion = FNiagaraCollisionDIFunctionVersion::LatestVersion;
@@ -416,7 +416,6 @@ void UNiagaraDataInterfaceCollisionQuery::GetFunctions(TArray<FNiagaraFunctionSi
 		FNiagaraFunctionSignature& ReadRayTrace = OutFunctions.AddDefaulted_GetRef();
 		ReadRayTrace.Name = NDICollisionQueryLocal::ReadAsyncRayTraceName;
 		ReadRayTrace.bMemberFunction = true;
-		ReadRayTrace.bRequiresContext = false;
 		ReadRayTrace.bSupportsCPU = false;
 #if WITH_EDITORONLY_DATA
 		ReadRayTrace.FunctionVersion = FNiagaraCollisionDIFunctionVersion::LatestVersion;
@@ -426,72 +425,72 @@ void UNiagaraDataInterfaceCollisionQuery::GetFunctions(TArray<FNiagaraFunctionSi
 		ReadRayTrace.AddInput(FNiagaraVariable(FNiagaraTypeDefinition::GetIntDef(), TEXT("PreviousFrameQueryID")), PreviousFrameQueryIDDescription);
 		ReadRayTrace.AddOutput(FNiagaraVariable(FNiagaraTypeDefinition::GetBoolDef(), TEXT("CollisionValid")), CollisionValidDescription);
 		ReadRayTrace.AddOutput(FNiagaraVariable(FNiagaraTypeDefinition::GetFloatDef(), TEXT("CollisionDistance")), CollisionDistanceDescription);
-		ReadRayTrace.AddOutput(FNiagaraVariable(FNiagaraTypeDefinition::GetVec3Def(), TEXT("CollisionPosWorld")), CollisionPosWorldDescription);
+		ReadRayTrace.AddOutput(FNiagaraVariable(FNiagaraTypeDefinition::GetPositionDef(), TEXT("CollisionPosWorld")), CollisionPosWorldDescription);
 		ReadRayTrace.AddOutput(FNiagaraVariable(FNiagaraTypeDefinition::GetVec3Def(), TEXT("CollisionNormal")), CollisionNormalDescription);
 	}
 
-	FNiagaraFunctionSignature SigCpuSync;
-	SigCpuSync.Name = NDICollisionQueryLocal::SyncTraceName;
-	SigCpuSync.bMemberFunction = true;
-	SigCpuSync.bRequiresContext = false;
-	SigCpuSync.bSupportsGPU = false;
+	{
+		FNiagaraFunctionSignature SigCpuSync;
+		SigCpuSync.Name = NDICollisionQueryLocal::SyncTraceName;
+		SigCpuSync.bMemberFunction = true;
+		SigCpuSync.bSupportsGPU = false;
 #if WITH_EDITORONLY_DATA
-	SigCpuSync.FunctionVersion = FNiagaraCollisionDIFunctionVersion::LatestVersion;
-	SigCpuSync.Description = LOCTEXT("SigCpuSyncDescription", "Traces a ray against the world using a specific channel and return the first blocking hit.");
+		SigCpuSync.FunctionVersion = FNiagaraCollisionDIFunctionVersion::LatestVersion;
+		SigCpuSync.Description = LOCTEXT("SigCpuSyncDescription", "Traces a ray against the world using a specific channel and return the first blocking hit.");
 #endif
-	const FText TraceStartWorldDescription = LOCTEXT("TraceStartWorldDescription", "The world position where the line trace should start.");
-	const FText TraceEndWorldDescription = LOCTEXT("TraceEndWorldDescription", "The world position where the line trace should end.");
-	const FText TraceChannelDescription = LOCTEXT("TraceChannelDescription", "The trace channel to collide against. Trace channels can be configured in the project settings.");
-	const FText SkipTraceDescription = LOCTEXT("SkipTraceDescription", "If true then the trace will be skipped completely.\nThis can be used as a performance optimization, as branch nodes in the graph still execute every path.");
-	const FText CollisionValidDescription = LOCTEXT("CollisionValidDescription", "Returns true if the trace was not skipped and the trace was blocked by some world geometry.");
-	const FText IsTraceInsideMeshDescription = LOCTEXT("IsTraceInsideMeshDescription", "If true then the trace started in penetration, i.e. with an initial blocking overlap.");
-	const FText CollisionPosWorldDescription = LOCTEXT("CollisionPosWorldDescription", "If the collision is valid, this returns the location of the blocking hit.");
-	const FText CollisionNormalDescription = LOCTEXT("CollisionNormalDescription", "If the collision is valid, this returns the normal at the position of the blocking hit.");
-	const FText CollisionMaterialFrictionDescription = LOCTEXT("CollisionMaterialFrictionDescription", "Friction value of surface, controls how easily things can slide on this surface (0 is frictionless, higher values increase the amount of friction).");
-	const FText CollisionMaterialRestitutionDescription = LOCTEXT("CollisionMaterialRestitutionDescription", "Restitution or 'bounciness' of this surface, between 0 (no bounce) and 1 (outgoing velocity is same as incoming)");
-	const FText CollisionMaterialIndexDescription = LOCTEXT("CollisionMaterialIndexDescription", "Returns the index of the surface as defined in the ProjectSettings/Physics/PhysicalSurface section");
+		const FText TraceStartWorldDescription = LOCTEXT("TraceStartWorldDescription", "The world position where the line trace should start.");
+		const FText TraceEndWorldDescription = LOCTEXT("TraceEndWorldDescription", "The world position where the line trace should end.");
+		const FText TraceChannelDescription = LOCTEXT("TraceChannelDescription", "The trace channel to collide against. Trace channels can be configured in the project settings.");
+		const FText SkipTraceDescription = LOCTEXT("SkipTraceDescription", "If true then the trace will be skipped completely.\nThis can be used as a performance optimization, as branch nodes in the graph still execute every path.");
+		const FText CollisionValidDescription = LOCTEXT("CollisionValidDescription", "Returns true if the trace was not skipped and the trace was blocked by some world geometry.");
+		const FText IsTraceInsideMeshDescription = LOCTEXT("IsTraceInsideMeshDescription", "If true then the trace started in penetration, i.e. with an initial blocking overlap.");
+		const FText CollisionPosWorldDescription = LOCTEXT("CollisionPosWorldDescription", "If the collision is valid, this returns the location of the blocking hit.");
+		const FText CollisionNormalDescription = LOCTEXT("CollisionNormalDescription", "If the collision is valid, this returns the normal at the position of the blocking hit.");
+		const FText CollisionMaterialFrictionDescription = LOCTEXT("CollisionMaterialFrictionDescription", "Friction value of surface, controls how easily things can slide on this surface (0 is frictionless, higher values increase the amount of friction).");
+		const FText CollisionMaterialRestitutionDescription = LOCTEXT("CollisionMaterialRestitutionDescription", "Restitution or 'bounciness' of this surface, between 0 (no bounce) and 1 (outgoing velocity is same as incoming)");
+		const FText CollisionMaterialIndexDescription = LOCTEXT("CollisionMaterialIndexDescription", "Returns the index of the surface as defined in the ProjectSettings/Physics/PhysicalSurface section");
 	
-	SigCpuSync.Inputs.Add(FNiagaraVariable(FNiagaraTypeDefinition(GetClass()), TEXT("CollisionQuery")));
-	SigCpuSync.AddInput(FNiagaraVariable(FNiagaraTypeDefinition::GetVec3Def(), TEXT("TraceStartWorld")), TraceStartWorldDescription);
-	SigCpuSync.AddInput(FNiagaraVariable(FNiagaraTypeDefinition::GetVec3Def(), TEXT("TraceEndWorld")), TraceEndWorldDescription);
-	SigCpuSync.AddInput(FNiagaraVariable(FNiagaraTypeDefinition(TraceChannelEnum), TEXT("TraceChannel")), TraceChannelDescription);
-	SigCpuSync.AddInput(FNiagaraVariable(FNiagaraTypeDefinition::GetBoolDef(), TEXT("SkipTrace")), SkipTraceDescription);
-	SigCpuSync.AddOutput(FNiagaraVariable(FNiagaraTypeDefinition::GetBoolDef(), TEXT("CollisionValid")), CollisionValidDescription);
-	SigCpuSync.AddOutput(FNiagaraVariable(FNiagaraTypeDefinition::GetBoolDef(), TEXT("IsTraceInsideMesh")), IsTraceInsideMeshDescription);
-	SigCpuSync.AddOutput(FNiagaraVariable(FNiagaraTypeDefinition::GetVec3Def(), TEXT("CollisionPosWorld")), CollisionPosWorldDescription);
-	SigCpuSync.AddOutput(FNiagaraVariable(FNiagaraTypeDefinition::GetVec3Def(), TEXT("CollisionNormal")), CollisionNormalDescription);
-	SigCpuSync.AddOutput(FNiagaraVariable(FNiagaraTypeDefinition::GetFloatDef(), TEXT("CollisionMaterialFriction")), CollisionMaterialFrictionDescription);
-	SigCpuSync.AddOutput(FNiagaraVariable(FNiagaraTypeDefinition::GetFloatDef(), TEXT("CollisionMaterialRestitution")), CollisionMaterialRestitutionDescription);
-	SigCpuSync.AddOutput(FNiagaraVariable(FNiagaraTypeDefinition::GetIntDef(), TEXT("CollisionMaterialIndex")), CollisionMaterialIndexDescription);
-	OutFunctions.Add(SigCpuSync);
-
-	FNiagaraFunctionSignature SigCpuAsync;
-	SigCpuAsync.Name = NDICollisionQueryLocal::AsyncTraceName;
-	SigCpuAsync.bMemberFunction = true;
-	SigCpuAsync.bRequiresContext = false;
-	SigCpuAsync.bSupportsGPU = false;
+		SigCpuSync.Inputs.Add(FNiagaraVariable(FNiagaraTypeDefinition(GetClass()), TEXT("CollisionQuery")));
+		SigCpuSync.AddInput(FNiagaraVariable(FNiagaraTypeDefinition::GetPositionDef(), TEXT("TraceStartWorld")), TraceStartWorldDescription);
+		SigCpuSync.AddInput(FNiagaraVariable(FNiagaraTypeDefinition::GetPositionDef(), TEXT("TraceEndWorld")), TraceEndWorldDescription);
+		SigCpuSync.AddInput(FNiagaraVariable(FNiagaraTypeDefinition(TraceChannelEnum), TEXT("TraceChannel")), TraceChannelDescription);
+		SigCpuSync.AddInput(FNiagaraVariable(FNiagaraTypeDefinition::GetBoolDef(), TEXT("SkipTrace")), SkipTraceDescription);
+		SigCpuSync.AddOutput(FNiagaraVariable(FNiagaraTypeDefinition::GetBoolDef(), TEXT("CollisionValid")), CollisionValidDescription);
+		SigCpuSync.AddOutput(FNiagaraVariable(FNiagaraTypeDefinition::GetBoolDef(), TEXT("IsTraceInsideMesh")), IsTraceInsideMeshDescription);
+		SigCpuSync.AddOutput(FNiagaraVariable(FNiagaraTypeDefinition::GetPositionDef(), TEXT("CollisionPosWorld")), CollisionPosWorldDescription);
+		SigCpuSync.AddOutput(FNiagaraVariable(FNiagaraTypeDefinition::GetVec3Def(), TEXT("CollisionNormal")), CollisionNormalDescription);
+		SigCpuSync.AddOutput(FNiagaraVariable(FNiagaraTypeDefinition::GetFloatDef(), TEXT("CollisionMaterialFriction")), CollisionMaterialFrictionDescription);
+		SigCpuSync.AddOutput(FNiagaraVariable(FNiagaraTypeDefinition::GetFloatDef(), TEXT("CollisionMaterialRestitution")), CollisionMaterialRestitutionDescription);
+		SigCpuSync.AddOutput(FNiagaraVariable(FNiagaraTypeDefinition::GetIntDef(), TEXT("CollisionMaterialIndex")), CollisionMaterialIndexDescription);
+		OutFunctions.Add(SigCpuSync);
+
+		FNiagaraFunctionSignature SigCpuAsync;
+		SigCpuAsync.Name = NDICollisionQueryLocal::AsyncTraceName;
+		SigCpuAsync.bMemberFunction = true;
+		SigCpuAsync.bSupportsGPU = false;
 #if WITH_EDITORONLY_DATA
-	SigCpuAsync.FunctionVersion = FNiagaraCollisionDIFunctionVersion::LatestVersion;
-	SigCpuAsync.Description = LOCTEXT("SigCpuAsyncDescription", "Traces a ray against the world using a specific channel and return the first blocking hit the next frame.\nNote that this is the ASYNC version of the trace function, meaning it will not returns the result right away, but with one frame latency.");
+		SigCpuAsync.FunctionVersion = FNiagaraCollisionDIFunctionVersion::LatestVersion;
+		SigCpuAsync.Description = LOCTEXT("SigCpuAsyncDescription", "Traces a ray against the world using a specific channel and return the first blocking hit the next frame.\nNote that this is the ASYNC version of the trace function, meaning it will not returns the result right away, but with one frame latency.");
 #endif
-	const FText PreviousFrameQueryIDDescription = LOCTEXT("PreviousFrameQueryIDDescription", "The query ID returned from the last frame's async trace call.\nRegardless if it is a valid ID or not this function call with issue a new async line trace, but it will only return results with a valid ID.");
-	const FText NextFrameQueryIDDescription = LOCTEXT("NextFrameQueryIDDescription", "The query ID to save and use as input to this function in the next frame.");
+		const FText PreviousFrameQueryIDDescription = LOCTEXT("PreviousFrameQueryIDDescription", "The query ID returned from the last frame's async trace call.\nRegardless if it is a valid ID or not this function call with issue a new async line trace, but it will only return results with a valid ID.");
+		const FText NextFrameQueryIDDescription = LOCTEXT("NextFrameQueryIDDescription", "The query ID to save and use as input to this function in the next frame.");
 	
-	SigCpuAsync.Inputs.Add(FNiagaraVariable(FNiagaraTypeDefinition(GetClass()), TEXT("CollisionQuery")));
-	SigCpuAsync.AddInput(FNiagaraVariable(FNiagaraTypeDefinition::GetIntDef(), TEXT("PreviousFrameQueryID")), PreviousFrameQueryIDDescription);
-	SigCpuAsync.AddInput(FNiagaraVariable(FNiagaraTypeDefinition::GetVec3Def(), TEXT("TraceStartWorld")), TraceStartWorldDescription);
-	SigCpuAsync.AddInput(FNiagaraVariable(FNiagaraTypeDefinition::GetVec3Def(), TEXT("TraceEndWorld")), TraceEndWorldDescription);
-	SigCpuAsync.AddInput(FNiagaraVariable(FNiagaraTypeDefinition(TraceChannelEnum), TEXT("TraceChannel")), TraceChannelDescription);
-	SigCpuAsync.AddInput(FNiagaraVariable(FNiagaraTypeDefinition::GetBoolDef(), TEXT("SkipTrace")), SkipTraceDescription);
-	SigCpuAsync.AddOutput(FNiagaraVariable(FNiagaraTypeDefinition::GetIntDef(), TEXT("NextFrameQueryID")), NextFrameQueryIDDescription);
-	SigCpuAsync.AddOutput(FNiagaraVariable(FNiagaraTypeDefinition::GetBoolDef(), TEXT("CollisionValid")), CollisionValidDescription);
-	SigCpuAsync.AddOutput(FNiagaraVariable(FNiagaraTypeDefinition::GetBoolDef(), TEXT("IsTraceInsideMesh")), IsTraceInsideMeshDescription);
-	SigCpuAsync.AddOutput(FNiagaraVariable(FNiagaraTypeDefinition::GetVec3Def(), TEXT("CollisionPosWorld")), CollisionPosWorldDescription);
-	SigCpuAsync.AddOutput(FNiagaraVariable(FNiagaraTypeDefinition::GetVec3Def(), TEXT("CollisionNormal")), CollisionNormalDescription);
-	SigCpuAsync.AddOutput(FNiagaraVariable(FNiagaraTypeDefinition::GetFloatDef(), TEXT("CollisionMaterialFriction")), CollisionMaterialFrictionDescription);
-	SigCpuAsync.AddOutput(FNiagaraVariable(FNiagaraTypeDefinition::GetFloatDef(), TEXT("CollisionMaterialRestitution")), CollisionMaterialRestitutionDescription);
-	SigCpuAsync.AddOutput(FNiagaraVariable(FNiagaraTypeDefinition::GetIntDef(), TEXT("CollisionMaterialIndex")), CollisionMaterialIndexDescription);
-	OutFunctions.Add(SigCpuAsync);
+		SigCpuAsync.Inputs.Add(FNiagaraVariable(FNiagaraTypeDefinition(GetClass()), TEXT("CollisionQuery")));
+		SigCpuAsync.AddInput(FNiagaraVariable(FNiagaraTypeDefinition::GetIntDef(), TEXT("PreviousFrameQueryID")), PreviousFrameQueryIDDescription);
+		SigCpuAsync.AddInput(FNiagaraVariable(FNiagaraTypeDefinition::GetPositionDef(), TEXT("TraceStartWorld")), TraceStartWorldDescription);
+		SigCpuAsync.AddInput(FNiagaraVariable(FNiagaraTypeDefinition::GetPositionDef(), TEXT("TraceEndWorld")), TraceEndWorldDescription);
+		SigCpuAsync.AddInput(FNiagaraVariable(FNiagaraTypeDefinition(TraceChannelEnum), TEXT("TraceChannel")), TraceChannelDescription);
+		SigCpuAsync.AddInput(FNiagaraVariable(FNiagaraTypeDefinition::GetBoolDef(), TEXT("SkipTrace")), SkipTraceDescription);
+		SigCpuAsync.AddOutput(FNiagaraVariable(FNiagaraTypeDefinition::GetIntDef(), TEXT("NextFrameQueryID")), NextFrameQueryIDDescription);
+		SigCpuAsync.AddOutput(FNiagaraVariable(FNiagaraTypeDefinition::GetBoolDef(), TEXT("CollisionValid")), CollisionValidDescription);
+		SigCpuAsync.AddOutput(FNiagaraVariable(FNiagaraTypeDefinition::GetBoolDef(), TEXT("IsTraceInsideMesh")), IsTraceInsideMeshDescription);
+		SigCpuAsync.AddOutput(FNiagaraVariable(FNiagaraTypeDefinition::GetPositionDef(), TEXT("CollisionPosWorld")), CollisionPosWorldDescription);
+		SigCpuAsync.AddOutput(FNiagaraVariable(FNiagaraTypeDefinition::GetVec3Def(), TEXT("CollisionNormal")), CollisionNormalDescription);
+		SigCpuAsync.AddOutput(FNiagaraVariable(FNiagaraTypeDefinition::GetFloatDef(), TEXT("CollisionMaterialFriction")), CollisionMaterialFrictionDescription);
+		SigCpuAsync.AddOutput(FNiagaraVariable(FNiagaraTypeDefinition::GetFloatDef(), TEXT("CollisionMaterialRestitution")), CollisionMaterialRestitutionDescription);
+		SigCpuAsync.AddOutput(FNiagaraVariable(FNiagaraTypeDefinition::GetIntDef(), TEXT("CollisionMaterialIndex")), CollisionMaterialIndexDescription);
+		OutFunctions.Add(SigCpuAsync);
+	}
 }
 
 // build the shader function HLSL; function name is passed in, as it's defined per-DI; that way, configuration could change
@@ -518,45 +517,22 @@ bool UNiagaraDataInterfaceCollisionQuery::GetFunctionHLSL(const FNiagaraDataInte
 
 bool UNiagaraDataInterfaceCollisionQuery::UpgradeFunctionCall(FNiagaraFunctionSignature& FunctionSignature)
 {
-	bool bWasChanged = false;
-
-	// The distance field query got a new output at some point, but there exists no custom version for it
-	if (FunctionSignature.Name == NDICollisionQueryLocal::DistanceFieldName && FunctionSignature.Outputs.Num() == 2)
-	{
-		FunctionSignature.Outputs.Add(FNiagaraVariable(FNiagaraTypeDefinition::GetBoolDef(), TEXT("IsDistanceFieldValid")));
-		bWasChanged = true;
-	}
-
-	// Early out for version matching
-	if (FunctionSignature.FunctionVersion == FNiagaraCollisionDIFunctionVersion::LatestVersion)
-	{
-		return bWasChanged;
-	}
-
-	// Added the possibility to skip a line trace to increase performance when only a fraction of particles wants to do a line trace
-	if (FunctionSignature.FunctionVersion < FNiagaraCollisionDIFunctionVersion::AddedTraceSkip)
-	{
-		if (FunctionSignature.Name == NDICollisionQueryLocal::SyncTraceName || FunctionSignature.Name == NDICollisionQueryLocal::AsyncTraceName)
-		{
-			FunctionSignature.Inputs.Add(FNiagaraVariable(FNiagaraTypeDefinition::GetBoolDef(), TEXT("SkipTrace")));
-			bWasChanged = true;
-		}
-	}
-
-	// Added the physical material ID as a result for line traces
-	if (FunctionSignature.FunctionVersion < FNiagaraCollisionDIFunctionVersion::ReturnCollisionMaterialIdx)
+	// always upgrade to the latest version
+	if (FunctionSignature.FunctionVersion < FNiagaraCollisionDIFunctionVersion::LatestVersion)
 	{
-		if (FunctionSignature.Name == NDICollisionQueryLocal::SyncTraceName || FunctionSignature.Name == NDICollisionQueryLocal::AsyncTraceName)
+		TArray<FNiagaraFunctionSignature> AllFunctions;
+		GetFunctions(AllFunctions);
+		for (const FNiagaraFunctionSignature& Sig : AllFunctions)
 		{
-			FunctionSignature.Outputs.Add(FNiagaraVariable(FNiagaraTypeDefinition::GetIntDef(), TEXT("CollisionMaterialIndex")));
-			bWasChanged = true;
+			if (FunctionSignature.Name == Sig.Name)
+			{
+				FunctionSignature = Sig;
+				return true;
+			}
 		}
 	}
 
-	// Set latest version
-	FunctionSignature.FunctionVersion = FNiagaraCollisionDIFunctionVersion::LatestVersion;
-
-	return bWasChanged;
+	return false;
 }
 #endif
 
@@ -600,8 +576,6 @@ void UNiagaraDataInterfaceCollisionQuery::GetParameterDefinitionHLSL(const FNiag
 
 DEFINE_NDI_DIRECT_FUNC_BINDER(UNiagaraDataInterfaceCollisionQuery, PerformQuerySyncCPU);
 DEFINE_NDI_DIRECT_FUNC_BINDER(UNiagaraDataInterfaceCollisionQuery, PerformQueryAsyncCPU);
-DEFINE_NDI_DIRECT_FUNC_BINDER(UNiagaraDataInterfaceCollisionQuery, QuerySceneDepth);
-DEFINE_NDI_DIRECT_FUNC_BINDER(UNiagaraDataInterfaceCollisionQuery, QueryMeshDistanceField);
 
 void UNiagaraDataInterfaceCollisionQuery::GetVMExternalFunction(const FVMExternalFunctionBindingInfo& BindingInfo, void* InstanceData, FVMExternalFunction &OutFunc)
 {
@@ -613,15 +587,6 @@ void UNiagaraDataInterfaceCollisionQuery::GetVMExternalFunction(const FVMExterna
 	{
 		NDI_FUNC_BINDER(UNiagaraDataInterfaceCollisionQuery, PerformQueryAsyncCPU)::Bind(this, OutFunc);
 	}
-	else if (BindingInfo.Name == NDICollisionQueryLocal::SceneDepthName ||
-			 BindingInfo.Name == NDICollisionQueryLocal::CustomDepthName)
-	{
-		NDI_FUNC_BINDER(UNiagaraDataInterfaceCollisionQuery, QuerySceneDepth)::Bind(this, OutFunc);
-	}
-	else if (BindingInfo.Name == NDICollisionQueryLocal::DistanceFieldName)
-	{
-		NDI_FUNC_BINDER(UNiagaraDataInterfaceCollisionQuery, QueryMeshDistanceField)::Bind(this, OutFunc);
-	}
 	else
 	{
 		UE_LOG(LogNiagara, Display, TEXT("Could not find data interface external function in %s. %s\n"),
@@ -662,67 +627,51 @@ void UNiagaraDataInterfaceCollisionQuery::PerformQuerySyncCPU(FVectorVMExternalF
 {
 	VectorVM::FUserPtrHandler<CQDIPerInstanceData> InstanceData(Context);
 
-	VectorVM::FExternalFuncInputHandler<float> StartPosParamX(Context);
-	VectorVM::FExternalFuncInputHandler<float> StartPosParamY(Context);
-	VectorVM::FExternalFuncInputHandler<float> StartPosParamZ(Context);
-
-	VectorVM::FExternalFuncInputHandler<float> EndPosParamX(Context);
-	VectorVM::FExternalFuncInputHandler<float> EndPosParamY(Context);
-	VectorVM::FExternalFuncInputHandler<float> EndPosParamZ(Context);
-
-	VectorVM::FExternalFuncInputHandler<ECollisionChannel> TraceChannelParam(Context);
-
-	VectorVM::FExternalFuncInputHandler<FNiagaraBool> IsSkipTrace(Context);
+	FNDIInputParam<FNiagaraPosition> StartPosParam(Context);
+	FNDIInputParam<FNiagaraPosition> EndPosParam(Context);
+	FNDIInputParam<ECollisionChannel> TraceChannelParam(Context);
+	FNDIInputParam<FNiagaraBool> IsSkipTrace(Context);
+	
+	FNDIOutputParam<FNiagaraBool> OutQueryValid(Context);
+	FNDIOutputParam<FNiagaraBool> OutInsideMesh(Context);
+	FNDIOutputParam<FNiagaraPosition> OutCollisionPos(Context);
+	FNDIOutputParam<FVector3f> OutCollisionNormal(Context);
+	FNDIOutputParam<float> OutFriction(Context);
+	FNDIOutputParam<float> OutRestitution(Context);
+	FNDIOutputParam<int32> OutPhysicalMaterialIdx(Context);
 
-	VectorVM::FExternalFuncRegisterHandler<FNiagaraBool> OutQueryValid(Context);
-	VectorVM::FExternalFuncRegisterHandler<FNiagaraBool> OutInsideMesh(Context);
-	VectorVM::FExternalFuncRegisterHandler<float> OutCollisionPosX(Context);
-	VectorVM::FExternalFuncRegisterHandler<float> OutCollisionPosY(Context);
-	VectorVM::FExternalFuncRegisterHandler<float> OutCollisionPosZ(Context);
-	VectorVM::FExternalFuncRegisterHandler<float> OutCollisionNormX(Context);
-	VectorVM::FExternalFuncRegisterHandler<float> OutCollisionNormY(Context);
-	VectorVM::FExternalFuncRegisterHandler<float> OutCollisionNormZ(Context);
-	VectorVM::FExternalFuncRegisterHandler<float> OutFriction(Context);
-	VectorVM::FExternalFuncRegisterHandler<float> OutRestitution(Context);
-	VectorVM::FExternalFuncRegisterHandler<int32> OutPhysicalMaterialIdx(Context);
+	FNiagaraLWCConverter LWCConverter = InstanceData->SystemInstance->GetLWCConverter();
 
 	FScopeLock ScopeLock(&CriticalSection);
 	for (int32 i = 0; i < Context.GetNumInstances(); ++i)
 	{
-		FVector Pos(StartPosParamX.GetAndAdvance(), StartPosParamY.GetAndAdvance(), StartPosParamZ.GetAndAdvance());
-		FVector Dir(EndPosParamX.GetAndAdvance(), EndPosParamY.GetAndAdvance(), EndPosParamZ.GetAndAdvance());
+		FNiagaraPosition TraceStart = StartPosParam.GetAndAdvance();
+		FNiagaraPosition TraceEnd = EndPosParam.GetAndAdvance();
 		ECollisionChannel TraceChannel = TraceChannelParam.GetAndAdvance();
-		bool Skip = IsSkipTrace.GetAndAdvance().GetValue();
-		ensure(!Pos.ContainsNaN());
+		bool Skip = IsSkipTrace.GetAndAdvance();
+		ensure(!TraceStart.ContainsNaN());
+		ensure(!TraceEnd.ContainsNaN());
 		FNiagaraDICollsionQueryResult Res;
 
-		if (!Skip && InstanceData->CollisionBatch.PerformQuery(Pos, Dir, Res, TraceChannel))
+		if (!Skip && InstanceData->CollisionBatch.PerformQuery(LWCConverter.ConvertSimulationPositionToWorld(TraceStart), LWCConverter.ConvertSimulationPositionToWorld(TraceEnd), Res, TraceChannel))
 		{
-			*OutQueryValid.GetDestAndAdvance() = FNiagaraBool(true);
-			*OutInsideMesh.GetDestAndAdvance() = FNiagaraBool(Res.IsInsideMesh);
-			*OutCollisionPosX.GetDestAndAdvance() = Res.CollisionPos.X;
-			*OutCollisionPosY.GetDestAndAdvance() = Res.CollisionPos.Y;
-			*OutCollisionPosZ.GetDestAndAdvance() = Res.CollisionPos.Z;
-			*OutCollisionNormX.GetDestAndAdvance() = Res.CollisionNormal.X;
-			*OutCollisionNormY.GetDestAndAdvance() = Res.CollisionNormal.Y;
-			*OutCollisionNormZ.GetDestAndAdvance() = Res.CollisionNormal.Z;
-			*OutFriction.GetDestAndAdvance() = Res.Friction;
-			*OutRestitution.GetDestAndAdvance() = Res.Restitution;
-			*OutPhysicalMaterialIdx.GetDestAndAdvance() = Res.PhysicalMaterialIdx;
+			OutQueryValid.SetAndAdvance(true);
+			OutInsideMesh.SetAndAdvance(Res.IsInsideMesh);
+			OutCollisionPos.SetAndAdvance(LWCConverter.ConvertWorldToSimulationPosition(Res.CollisionPos));
+			OutCollisionNormal.SetAndAdvance(Res.CollisionNormal);
+			OutFriction.SetAndAdvance(Res.Friction);
+			OutRestitution.SetAndAdvance(Res.Restitution);
+			OutPhysicalMaterialIdx.SetAndAdvance(Res.PhysicalMaterialIdx);
 		}
 		else
-		{
-			*OutQueryValid.GetDestAndAdvance() = FNiagaraBool();
-			*OutInsideMesh.GetDestAndAdvance() = FNiagaraBool();
-			*OutCollisionPosX.GetDestAndAdvance() = 0.0f;
-			*OutCollisionPosY.GetDestAndAdvance() = 0.0f;
-			*OutCollisionPosZ.GetDestAndAdvance() = 0.0f;
-			*OutCollisionNormX.GetDestAndAdvance() = 0.0f;
-			*OutCollisionNormY.GetDestAndAdvance() = 0.0f;
-			*OutCollisionNormZ.GetDestAndAdvance() = 0.0f;
-			*OutFriction.GetDestAndAdvance() = 0.0f;
-			*OutRestitution.GetDestAndAdvance() = 0.0f;
-			*OutPhysicalMaterialIdx.GetDestAndAdvance() = 0;
+		{			
+			OutQueryValid.SetAndAdvance(false);
+			OutInsideMesh.SetAndAdvance(false);
+			OutCollisionPos.SetAndAdvance(FVector3f::ZeroVector);
+			OutCollisionNormal.SetAndAdvance(FVector3f::ZeroVector);
+			OutFriction.SetAndAdvance(0);
+			OutRestitution.SetAndAdvance(0);
+			OutPhysicalMaterialIdx.SetAndAdvance(0);
 		}
 	}
 }
@@ -731,144 +680,62 @@ void UNiagaraDataInterfaceCollisionQuery::PerformQueryAsyncCPU(FVectorVMExternal
 {
 	VectorVM::FUserPtrHandler<CQDIPerInstanceData> InstanceData(Context);
 
-	VectorVM::FExternalFuncInputHandler<int32> InIDParam(Context);
-	VectorVM::FExternalFuncInputHandler<float> StartPosParamX(Context);
-	VectorVM::FExternalFuncInputHandler<float> StartPosParamY(Context);
-	VectorVM::FExternalFuncInputHandler<float> StartPosParamZ(Context);
-
-	VectorVM::FExternalFuncInputHandler<float> EndPosParamX(Context);
-	VectorVM::FExternalFuncInputHandler<float> EndPosParamY(Context);
-	VectorVM::FExternalFuncInputHandler<float> EndPosParamZ(Context);
-
-	VectorVM::FExternalFuncInputHandler<ECollisionChannel> TraceChannelParam(Context);
-
-	VectorVM::FExternalFuncInputHandler<FNiagaraBool> IsSkipTrace(Context);
-
-	VectorVM::FExternalFuncRegisterHandler<int32> OutQueryID(Context);
+	FNDIInputParam<int32> InIDParam(Context);
+	FNDIInputParam<FNiagaraPosition> StartPosParam(Context);
+	FNDIInputParam<FNiagaraPosition> EndPosParam(Context);
+	FNDIInputParam<ECollisionChannel> TraceChannelParam(Context);
+	FNDIInputParam<FNiagaraBool> IsSkipTrace(Context);
+	
+	FNDIOutputParam<int32> OutQueryID(Context);
+	FNDIOutputParam<FNiagaraBool> OutQueryValid(Context);
+	FNDIOutputParam<FNiagaraBool> OutInsideMesh(Context);
+	FNDIOutputParam<FNiagaraPosition> OutCollisionPos(Context);
+	FNDIOutputParam<FVector3f> OutCollisionNormal(Context);
+	FNDIOutputParam<float> OutFriction(Context);
+	FNDIOutputParam<float> OutRestitution(Context);
+	FNDIOutputParam<int32> OutPhysicalMaterialIdx(Context);
 
-	VectorVM::FExternalFuncRegisterHandler<FNiagaraBool> OutQueryValid(Context);
-	VectorVM::FExternalFuncRegisterHandler<FNiagaraBool> OutInsideMesh(Context);
-	VectorVM::FExternalFuncRegisterHandler<float> OutCollisionPosX(Context);
-	VectorVM::FExternalFuncRegisterHandler<float> OutCollisionPosY(Context);
-	VectorVM::FExternalFuncRegisterHandler<float> OutCollisionPosZ(Context);
-	VectorVM::FExternalFuncRegisterHandler<float> OutCollisionNormX(Context);
-	VectorVM::FExternalFuncRegisterHandler<float> OutCollisionNormY(Context);
-	VectorVM::FExternalFuncRegisterHandler<float> OutCollisionNormZ(Context);
-	VectorVM::FExternalFuncRegisterHandler<float> OutFriction(Context);
-	VectorVM::FExternalFuncRegisterHandler<float> OutRestitution(Context);
-	VectorVM::FExternalFuncRegisterHandler<int32> OutPhysicalMaterialIdx(Context);
+	FNiagaraLWCConverter LWCConverter = InstanceData->SystemInstance->GetLWCConverter();
 
 	FScopeLock ScopeLock(&CriticalSection);
 	for (int32 i = 0; i < Context.GetNumInstances(); ++i)
 	{
-		FVector Pos(StartPosParamX.GetAndAdvance(), StartPosParamY.GetAndAdvance(), StartPosParamZ.GetAndAdvance());
-		FVector End(EndPosParamX.GetAndAdvance(), EndPosParamY.GetAndAdvance(), EndPosParamZ.GetAndAdvance());
+		FNiagaraPosition TraceStart = StartPosParam.GetAndAdvance();
+		FNiagaraPosition TraceEnd = EndPosParam.GetAndAdvance();
 		ECollisionChannel TraceChannel = TraceChannelParam.GetAndAdvance();
-		bool Skip = IsSkipTrace.GetAndAdvance().GetValue();
-		ensure(!Pos.ContainsNaN());
+		bool Skip = IsSkipTrace.GetAndAdvance();
+		ensure(!TraceStart.ContainsNaN());
+		ensure(!TraceEnd.ContainsNaN());
 
-		*OutQueryID.GetDestAndAdvance() = Skip ? INDEX_NONE : InstanceData->CollisionBatch.SubmitQuery(Pos, End, TraceChannel);
+		int QueryID = Skip ? INDEX_NONE : InstanceData->CollisionBatch.SubmitQuery(LWCConverter.ConvertSimulationPositionToWorld(TraceStart), LWCConverter.ConvertSimulationPositionToWorld(TraceEnd), TraceChannel);
+		OutQueryID.SetAndAdvance(QueryID);
 
 		// try to retrieve a query with the supplied query ID
 		FNiagaraDICollsionQueryResult Res;
 		int32 ID = InIDParam.GetAndAdvance();
 		if (ID != INDEX_NONE && InstanceData->CollisionBatch.GetQueryResult(ID, Res))
 		{
-			*OutQueryValid.GetDestAndAdvance() = FNiagaraBool(true);
-			*OutInsideMesh.GetDestAndAdvance() = FNiagaraBool(Res.IsInsideMesh);
-			*OutCollisionPosX.GetDestAndAdvance() = Res.CollisionPos.X;
-			*OutCollisionPosY.GetDestAndAdvance() = Res.CollisionPos.Y;
-			*OutCollisionPosZ.GetDestAndAdvance() = Res.CollisionPos.Z;
-			*OutCollisionNormX.GetDestAndAdvance() = Res.CollisionNormal.X;
-			*OutCollisionNormY.GetDestAndAdvance() = Res.CollisionNormal.Y;
-			*OutCollisionNormZ.GetDestAndAdvance() = Res.CollisionNormal.Z;
-			*OutFriction.GetDestAndAdvance() = Res.Friction;
-			*OutRestitution.GetDestAndAdvance() = Res.Restitution;
-			*OutPhysicalMaterialIdx.GetDestAndAdvance() = Res.PhysicalMaterialIdx;
+			OutQueryValid.SetAndAdvance(true);
+			OutInsideMesh.SetAndAdvance(Res.IsInsideMesh);
+			OutCollisionPos.SetAndAdvance(LWCConverter.ConvertWorldToSimulationPosition(Res.CollisionPos));
+			OutCollisionNormal.SetAndAdvance(Res.CollisionNormal);
+			OutFriction.SetAndAdvance(Res.Friction);
+			OutRestitution.SetAndAdvance(Res.Restitution);
+			OutPhysicalMaterialIdx.SetAndAdvance(Res.PhysicalMaterialIdx);
 		}
 		else
 		{
-			*OutQueryValid.GetDestAndAdvance() = FNiagaraBool();
-			*OutInsideMesh.GetDestAndAdvance() = FNiagaraBool();
-			*OutCollisionPosX.GetDestAndAdvance() = 0.0f;
-			*OutCollisionPosY.GetDestAndAdvance() = 0.0f;
-			*OutCollisionPosZ.GetDestAndAdvance() = 0.0f;
-			*OutCollisionNormX.GetDestAndAdvance() = 0.0f;
-			*OutCollisionNormY.GetDestAndAdvance() = 0.0f;
-			*OutCollisionNormZ.GetDestAndAdvance() = 0.0f;
-			*OutFriction.GetDestAndAdvance() = 0.0f;
-			*OutRestitution.GetDestAndAdvance() = 0.0f;
-			*OutPhysicalMaterialIdx.GetDestAndAdvance() = 0;
+			OutQueryValid.SetAndAdvance(false);
+			OutInsideMesh.SetAndAdvance(false);
+			OutCollisionPos.SetAndAdvance(FVector3f::ZeroVector);
+			OutCollisionNormal.SetAndAdvance(FVector3f::ZeroVector);
+			OutFriction.SetAndAdvance(0);
+			OutRestitution.SetAndAdvance(0);
+			OutPhysicalMaterialIdx.SetAndAdvance(0);
 		}
 	}
 }
 
-void UNiagaraDataInterfaceCollisionQuery::QuerySceneDepth(FVectorVMExternalFunctionContext & Context)
-{
-	UE_LOG(LogNiagara, Error, TEXT("GPU only function 'QuerySceneDepthGPU' called on CPU VM, check your module code to fix."));
-
-	VectorVM::FUserPtrHandler<CQDIPerInstanceData> InstanceData(Context);
-
-	VectorVM::FExternalFuncInputHandler<float> SamplePosParamX(Context);
-	VectorVM::FExternalFuncInputHandler<float> SamplePosParamY(Context);
-	VectorVM::FExternalFuncInputHandler<float> SamplePosParamZ(Context);
-	
-	VectorVM::FExternalFuncRegisterHandler<float> OutSceneDepth(Context);
-	VectorVM::FExternalFuncRegisterHandler<float> OutCameraPosX(Context);
-	VectorVM::FExternalFuncRegisterHandler<float> OutCameraPosY(Context);
-	VectorVM::FExternalFuncRegisterHandler<float> OutCameraPosZ(Context);
-	VectorVM::FExternalFuncRegisterHandler<int32> OutIsInsideView(Context);
-	VectorVM::FExternalFuncRegisterHandler<float> OutWorldPosX(Context);
-	VectorVM::FExternalFuncRegisterHandler<float> OutWorldPosY(Context);
-	VectorVM::FExternalFuncRegisterHandler<float> OutWorldPosZ(Context);
-	VectorVM::FExternalFuncRegisterHandler<float> OutWorldNormX(Context);
-	VectorVM::FExternalFuncRegisterHandler<float> OutWorldNormY(Context);
-	VectorVM::FExternalFuncRegisterHandler<float> OutWorldNormZ(Context);
-
-	FScopeLock ScopeLock(&CriticalSection);
-	for (int32 i = 0; i < Context.GetNumInstances(); ++i)
-	{
-		*OutSceneDepth.GetDestAndAdvance() = -1;
-		*OutIsInsideView.GetDestAndAdvance() = 0;
-		*OutWorldPosX.GetDestAndAdvance() = 0.0f;
-		*OutWorldPosY.GetDestAndAdvance() = 0.0f;
-		*OutWorldPosZ.GetDestAndAdvance() = 0.0f;
-		*OutWorldNormX.GetDestAndAdvance() = 0.0f;
-		*OutWorldNormY.GetDestAndAdvance() = 0.0f;
-		*OutWorldNormZ.GetDestAndAdvance() = 1.0f;
-		*OutCameraPosX.GetDestAndAdvance() = 0.0f;
-		*OutCameraPosY.GetDestAndAdvance() = 0.0f;
-		*OutCameraPosZ.GetDestAndAdvance() = 0.0f;
-	}
-}
-
-void UNiagaraDataInterfaceCollisionQuery::QueryMeshDistanceField(FVectorVMExternalFunctionContext& Context)
-{
-	UE_LOG(LogNiagara, Error, TEXT("GPU only function 'QueryMeshDistanceFieldGPU' called on CPU VM, check your module code to fix."));
-
-	VectorVM::FUserPtrHandler<CQDIPerInstanceData> InstanceData(Context);
-
-	VectorVM::FExternalFuncInputHandler<float> SamplePosParamX(Context);
-	VectorVM::FExternalFuncInputHandler<float> SamplePosParamY(Context);
-	VectorVM::FExternalFuncInputHandler<float> SamplePosParamZ(Context);
-
-	VectorVM::FExternalFuncRegisterHandler<float> OutSurfaceDistance(Context);
-	VectorVM::FExternalFuncRegisterHandler<float> OutFieldGradientX(Context);
-	VectorVM::FExternalFuncRegisterHandler<float> OutFieldGradientY(Context);
-	VectorVM::FExternalFuncRegisterHandler<float> OutFieldGradientZ(Context);
-	FNDIOutputParam<FNiagaraBool> OutIsFieldValid(Context);
-
-	FScopeLock ScopeLock(&CriticalSection);
-	for (int32 i = 0; i < Context.GetNumInstances(); ++i)
-	{
-		*OutSurfaceDistance.GetDestAndAdvance() = -1;
-		*OutFieldGradientX.GetDestAndAdvance() = 0.0f;
-		*OutFieldGradientY.GetDestAndAdvance() = 0.0f;
-		*OutFieldGradientZ.GetDestAndAdvance() = 1.0f;
-		OutIsFieldValid.SetAndAdvance(FNiagaraBool());
-	}
-}
-
 bool UNiagaraDataInterfaceCollisionQuery::PerInstanceTick(void* PerInstanceData, FNiagaraSystemInstance* InSystemInstance, float DeltaSeconds)
 {
 	CQDIPerInstanceData* PIData = static_cast<CQDIPerInstanceData*>(PerInstanceData);
@@ -894,7 +761,6 @@ bool UNiagaraDataInterfaceCollisionQuery::Equals(const UNiagaraDataInterface* Ot
 
 	const UNiagaraDataInterfaceCollisionQuery* OtherTyped = CastChecked<const UNiagaraDataInterfaceCollisionQuery>(Other);
 	return OtherTyped->MaxTracesPerParticle == MaxTracesPerParticle;
-	return OtherTyped->MaxRetraces == MaxRetraces;
 }
 
 bool UNiagaraDataInterfaceCollisionQuery::CopyToInternal(UNiagaraDataInterface* Destination) const
@@ -944,6 +810,7 @@ public:
 	void Bind(const FNiagaraDataInterfaceGPUParamInfo& ParameterInfo, const class FShaderParameterMap& ParameterMap)
 	{
 		GlobalDistanceFieldParameters.Bind(ParameterMap);
+		SystemLWCTileParam.Bind(ParameterMap, *(NDICollisionQueryLocal::SystemLWCTileName + ParameterInfo.DataInterfaceHLSLSymbol));
 #if RHI_RAYTRACING
 		RayTracingEnabledParam.Bind(ParameterMap, *(NDICollisionQueryLocal::RayTracingEnabledParamName + ParameterInfo.DataInterfaceHLSLSymbol));
 		MaxRayTraceCountParam.Bind(ParameterMap, *(NDICollisionQueryLocal::MaxRayTraceCountParamName + ParameterInfo.DataInterfaceHLSLSymbol));
@@ -962,6 +829,8 @@ public:
 
 		FNiagaraDataIntefaceProxyCollisionQuery* QueryDI = (FNiagaraDataIntefaceProxyCollisionQuery*)Context.DataInterface;
 		FRHIComputeShader* ComputeShaderRHI = Context.Shader.GetComputeShader();
+		SetShaderValue(RHICmdList, ComputeShaderRHI, SystemLWCTileParam, Context.SystemLWCTile);
+		
 		// Bind distance field parameters
 		if (GlobalDistanceFieldParameters.IsBound())
 		{
@@ -1035,8 +904,6 @@ public:
 	void Unset(FRHICommandList& RHICmdList, const FNiagaraDataInterfaceSetArgs& Context) const
 	{
 		FRHIComputeShader* ComputeShaderRHI = Context.Shader.GetComputeShader();
-		FNiagaraDataIntefaceProxyCollisionQuery* QueryDI = (FNiagaraDataIntefaceProxyCollisionQuery*)Context.DataInterface;
-
 		if (RayRequestsParam.IsUAVBound())
 		{
 			RayRequestsParam.UnsetUAV(RHICmdList, ComputeShaderRHI);
@@ -1051,6 +918,7 @@ public:
 
 private:
 	LAYOUT_FIELD(FGlobalDistanceFieldParameters, GlobalDistanceFieldParameters);
+	LAYOUT_FIELD(FShaderParameter, SystemLWCTileParam);
 
 #if RHI_RAYTRACING
 	LAYOUT_FIELD(FShaderParameter, RayTracingEnabledParam);
