Niagara LWC changes:

diff --git a/Engine/Plugins/FX/Niagara/Source/Niagara/Private/NiagaraSystem.cpp b/Engine/Plugins/FX/Niagara/Source/Niagara/Private/NiagaraSystem.cpp
--- a/Engine/Plugins/FX/Niagara/Source/Niagara/Private/NiagaraSystem.cpp
+++ b/Engine/Plugins/FX/Niagara/Source/Niagara/Private/NiagaraSystem.cpp
@@ -112,6 +112,7 @@ UNiagaraSystem::UNiagaraSystem(const FObjectInitializer& ObjectInitializer)
 , bTrimAttributesOnCook(true)
 , bDisableAllDebugSwitches(false)
 #endif
+, bSupportLargeWorldCoordinates(true)
 , bFixedBounds(false)
 #if WITH_EDITORONLY_DATA
 , bIsolateEnabled(false)
@@ -142,6 +143,9 @@ UNiagaraSystem::UNiagaraSystem(const FObjectInitializer& ObjectInitializer)
 #if WITH_EDITORONLY_DATA
 	AssetGuid = FGuid::NewGuid();
 #endif
+
+	const UNiagaraSettings* Settings = GetDefault<UNiagaraSettings>();
+	bLwcEnabledSettingCached = Settings ? Settings->bSystemsSupportLargeWorldCoordinates : true;
 }
 
 #if WITH_EDITORONLY_DATA
@@ -1150,6 +1154,8 @@ void UNiagaraSystem::PostLoad()
 	}
 
 #if WITH_EDITORONLY_DATA
+	FixupPositionUserParameters(); //TODO[mg]: hide this behind a new custom version?
+	
 	if (EditorData == nullptr)
 	{
 		INiagaraModule& NiagaraModule = FModuleManager::GetModuleChecked<INiagaraModule>("Niagara");
@@ -2121,19 +2127,30 @@ void UNiagaraSystem::EnsureFullyLoaded() const
 	System->UpdateSystemAfterLoad();
 }
 
-
-bool UNiagaraSystem::CanObtainEmitterAttribute(const FNiagaraVariableBase& InVarWithUniqueNameNamespace) const
+bool UNiagaraSystem::CanObtainEmitterAttribute(const FNiagaraVariableBase& InVarWithUniqueNameNamespace, FNiagaraTypeDefinition& OutBoundType) const
 {
-	if (SystemSpawnScript)
-		return SystemSpawnScript->GetVMExecutableData().Attributes.Contains(InVarWithUniqueNameNamespace);
-	return false;
+	return CanObtainSystemAttribute(InVarWithUniqueNameNamespace, OutBoundType);
 }
-bool UNiagaraSystem::CanObtainSystemAttribute(const FNiagaraVariableBase& InVar) const
+
+bool UNiagaraSystem::CanObtainSystemAttribute(const FNiagaraVariableBase& InVar, FNiagaraTypeDefinition& OutBoundType) const
 {
 	if (SystemSpawnScript)
-		return SystemSpawnScript->GetVMExecutableData().Attributes.Contains(InVar);
+	{
+		OutBoundType = InVar.GetType();
+		bool bContainsAttribute = SystemSpawnScript->GetVMExecutableData().Attributes.Contains(InVar);
+		if (!bContainsAttribute && InVar.GetType() == FNiagaraTypeDefinition::GetPositionDef())
+		{
+			// if we don't find a position type var we check for a vec3 type for backwards compatibility
+			OutBoundType = FNiagaraTypeDefinition::GetVec3Def();
+			FNiagaraVariableBase VarCopy = InVar;
+			VarCopy.SetType(OutBoundType);
+			bContainsAttribute = SystemSpawnScript->GetVMExecutableData().Attributes.Contains(VarCopy);
+		}
+		return bContainsAttribute;
+	}
 	return false;
 }
+
 bool UNiagaraSystem::CanObtainUserVariable(const FNiagaraVariableBase& InVar) const
 {
 	return ExposedParameters.IndexOf(InVar) != INDEX_NONE;
@@ -3616,6 +3633,36 @@ FNiagaraEmitterCompiledData::FNiagaraEmitterCompiledData()
 }
 
 #if WITH_EDITORONLY_DATA
+void UNiagaraSystem::FixupPositionUserParameters()
+{
+	TArray<FNiagaraVariable> UserParameters;
+	ExposedParameters.GetUserParameters(UserParameters);
+	UserParameters.FilterByPredicate([](const FNiagaraVariable& Var) { return Var.GetType() == FNiagaraTypeDefinition::GetVec3Def(); });
+
+	if (UserParameters.Num() == 0)
+	{
+		return;
+	}
+	for (FNiagaraVariable& UserParameter : UserParameters)
+	{
+		ExposedParameters.MakeUserVariable(UserParameter);
+	}
+	
+	TSet<FNiagaraVariable> LinkedPositionInputs;
+	ForEachScript([&UserParameters, &LinkedPositionInputs](UNiagaraScript* Script)
+	{
+		Script->GetLatestSource()->GetLinkedPositionTypeInputs(UserParameters, LinkedPositionInputs);
+	});
+
+	// looks like we have a few FVector3f user parameters that are linked to position inputs in the stack.
+	// Most likely this is because core modules were changed to use position data. To ease the transition of old assets, we auto-convert those inputs to position types as well.
+	for (FNiagaraVariable& LinkedParameter : LinkedPositionInputs)
+	{
+		ExposedParameters.ConvertParameterType(LinkedParameter, FNiagaraTypeDefinition::GetPositionDef());
+		UE_LOG(LogNiagara, Log, TEXT("Converted parameter %s from vec3 to position type in asset %s"), *LinkedParameter.GetName().ToString(), *GetPathName());
+	}
+}
+
 void FNiagaraParameterDataSetBindingCollection::BuildInternal(const TArray<FNiagaraVariable>& ParameterVars, const FNiagaraDataSetCompiledData& DataSet, const FString& NamespaceBase, const FString& NamespaceReplacement)
 {
 	// be sure to reset the offsets first
@@ -3660,10 +3707,10 @@ void FNiagaraParameterDataSetBindingCollection::BuildInternal(const TArray<FNiag
 		}
 
 		// we need to take into account potential padding that is in the constant buffers based similar to what is done
-		// in the NiagaraHlslTranslator, where Vec2/Vec3 are treated as Vec4.
+		// in the NiagaraHlslTranslator, where Vec2/Vec3/Position are treated as Vec4.
 		int32 ParameterSize = Var.GetSizeInBytes();
 		const FNiagaraTypeDefinition& Type = Var.GetType();
-		if (Type == FNiagaraTypeDefinition::GetVec2Def() || Type == FNiagaraTypeDefinition::GetVec3Def())
+		if (Type == FNiagaraTypeDefinition::GetVec2Def() || Type == FNiagaraTypeDefinition::GetVec3Def() || Type == FNiagaraTypeDefinition::GetPositionDef())
 		{
 			ParameterSize = Align(ParameterSize, FNiagaraTypeDefinition::GetVec4Def().GetSize());
 		}
diff --git a/Engine/Plugins/FX/Niagara/Source/Niagara/Private/NiagaraSystemGpuComputeProxy.cpp b/Engine/Plugins/FX/Niagara/Source/Niagara/Private/NiagaraSystemGpuComputeProxy.cpp
--- a/Engine/Plugins/FX/Niagara/Source/Niagara/Private/NiagaraSystemGpuComputeProxy.cpp
+++ b/Engine/Plugins/FX/Niagara/Source/Niagara/Private/NiagaraSystemGpuComputeProxy.cpp
@@ -9,7 +9,7 @@
 
 FNiagaraSystemGpuComputeProxy::FNiagaraSystemGpuComputeProxy(FNiagaraSystemInstance* OwnerInstance)
 	: DebugOwnerInstance(OwnerInstance)
-	, DebugOwnerComputeDispatchInterface(nullptr)
+	, SystemLWCTile(OwnerInstance->GetLWCTile())
 	, SystemInstanceID(OwnerInstance->GetId())
 {
 	bRequiresDistanceFieldData = OwnerInstance->RequiresDistanceFieldData();
diff --git a/Engine/Plugins/FX/Niagara/Source/Niagara/Private/NiagaraSystemInstance.cpp b/Engine/Plugins/FX/Niagara/Source/Niagara/Private/NiagaraSystemInstance.cpp
--- a/Engine/Plugins/FX/Niagara/Source/Niagara/Private/NiagaraSystemInstance.cpp
+++ b/Engine/Plugins/FX/Niagara/Source/Niagara/Private/NiagaraSystemInstance.cpp
@@ -17,15 +17,10 @@
 #include "NiagaraCrashReporterHandler.h"
 
 #include "Async/Async.h"
-#include "Async/ParallelFor.h"
 #include "Components/SkeletalMeshComponent.h"
 #include "GameFramework/PlayerController.h"
 #include "Templates/AlignmentTemplates.h"
 
-#if WITH_EDITORONLY_DATA
-#include "Editor.h"
-#endif
-
 
 DECLARE_CYCLE_STAT(TEXT("System Activate [GT]"), STAT_NiagaraSystemActivate, STATGROUP_Niagara);
 DECLARE_CYCLE_STAT(TEXT("System Deactivate [GT]"), STAT_NiagaraSystemDeactivate, STATGROUP_Niagara);
@@ -87,7 +82,7 @@ static FAutoConsoleVariableRef CVarNiagaraAllowDeferredReset(
 );
 
 FNiagaraSystemInstance::FNiagaraSystemInstance(UWorld& InWorld, UNiagaraSystem& InAsset, FNiagaraUserRedirectionParameterStore* InOverrideParameters,
-	USceneComponent* InAttachComponent, ENiagaraTickBehavior InTickBehavior, bool bInPooled)
+                                               USceneComponent* InAttachComponent, ENiagaraTickBehavior InTickBehavior, bool bInPooled)
 	: SystemInstanceIndex(INDEX_NONE)
 	, SignificanceIndex(INDEX_NONE)
 	, World(&InWorld)
@@ -792,6 +787,13 @@ void FNiagaraSystemInstance::Reset(FNiagaraSystemInstance::EResetMode Mode)
 		}
 	}
 
+	// Set tile for LWC offset
+	LWCTile = FVector3f::ZeroVector;
+	if (AttachComponent.IsValid() && GetSystem()->SupportsLargeWorldCoordinates())
+	{
+		LWCTile = FLargeWorldRenderScalar::GetTileFor(AttachComponent->GetComponentLocation());
+	}
+
 	// Depending on the rest mode we may need to bind or can possibly skip it
 	// We must bind if we were previously complete as unbind will have been called, we can not get here if the system was disabled
 	bool bBindParams = IsComplete();
@@ -907,6 +909,7 @@ void FNiagaraSystemInstance::ResetInternal(bool bResetSimulations)
 	check(World);
 	if (OverrideParameters && World->WorldType == EWorldType::Editor)
 	{
+		OverrideParameters->ResolvePositions(GetLWCConverter());
 		OverrideParameters->Tick();
 	}
 #endif
@@ -1239,7 +1242,16 @@ void FNiagaraSystemInstance::UnbindParameters(bool bFromComplete)
 	}
 }
 
-FNiagaraWorldManager* FNiagaraSystemInstance::GetWorldManager()const
+FNiagaraLWCConverter FNiagaraSystemInstance::GetLWCConverter(bool bLocalSpaceEmitter) const
+{
+	if (bLocalSpaceEmitter)
+	{
+		return FNiagaraLWCConverter();
+	}
+	return FNiagaraLWCConverter(FVector(LWCTile) * FLargeWorldRenderScalar::GetTileSize());
+}
+
+FNiagaraWorldManager* FNiagaraSystemInstance::GetWorldManager() const
 {
 	check(World);
 	return FNiagaraWorldManager::Get(World);
@@ -1385,6 +1397,7 @@ void FNiagaraSystemInstance::InitDataInterfaces()
 
 	if (OverrideParameters != nullptr)
 	{
+		OverrideParameters->ResolvePositions(GetLWCConverter());
 		OverrideParameters->Tick();
 	}
 
@@ -1594,7 +1607,6 @@ void FNiagaraSystemInstance::InitDataInterfaces()
 			}
 		}
 	}
-	
 }
 
 void FNiagaraSystemInstance::TickDataInterfaces(float DeltaSeconds, bool bPostSimulate)
@@ -1855,6 +1867,7 @@ void FNiagaraSystemInstance::TickInstanceParameters_GameThread(float DeltaSecond
 	if (AttachComponent.IsValid())
 	{
 		WorldTransform = AttachComponent->GetComponentToWorld();
+		WorldTransform.AddToTranslation(FVector(LWCTile) * -FLargeWorldRenderScalar::GetTileSize());
 	}
 	const bool TransformMatches = GatheredInstanceParameters.ComponentTrans.Equals(WorldTransform);
 	if (TransformMatches)
@@ -1915,6 +1928,7 @@ void FNiagaraSystemInstance::TickInstanceParameters_GameThread(float DeltaSecond
 
 	if (OverrideParameters)
 	{
+		OverrideParameters->ResolvePositions(GetLWCConverter());
 		OverrideParameters->Tick();
 	}
 }
@@ -1947,6 +1961,8 @@ void FNiagaraSystemInstance::TickInstanceParameters_Concurrent()
 		CurrentOwnerParameters.EngineYAxis = CurrentOwnerParameters.EngineRotation.GetAxisY();
 		CurrentOwnerParameters.EngineZAxis = CurrentOwnerParameters.EngineRotation.GetAxisZ();
 		CurrentOwnerParameters.EngineScale = GatheredInstanceParameters.ComponentTrans.GetScale3D();
+		CurrentOwnerParameters.EngineLWCTile = LWCTile;
+		CurrentOwnerParameters.EngineLWCTile.W = FLargeWorldRenderScalar::GetTileSize();
 	}
 
 	CurrentSystemParameters.EngineEmitterCount = GatheredInstanceParameters.EmitterCount;
diff --git a/Engine/Plugins/FX/Niagara/Source/Niagara/Private/NiagaraSystemInstanceController.cpp b/Engine/Plugins/FX/Niagara/Source/Niagara/Private/NiagaraSystemInstanceController.cpp
--- a/Engine/Plugins/FX/Niagara/Source/Niagara/Private/NiagaraSystemInstanceController.cpp
+++ b/Engine/Plugins/FX/Niagara/Source/Niagara/Private/NiagaraSystemInstanceController.cpp
@@ -71,6 +71,7 @@ void FNiagaraSystemInstanceController::PostTickRenderers(FNiagaraSystemRenderDat
 		if (bNeedsOverrideParametersTicked)
 		{
 			bNeedsOverrideParametersTicked = false;
+			OverrideParameters->ResolvePositions(SystemInst->GetLWCConverter());
 			OverrideParameters->Tick();
 		}
 		if (bNeedsUpdateEmitterMaterials)
@@ -372,6 +373,16 @@ void FNiagaraSystemInstanceController::SetVariable(FName InVariableName, FVector
 	}
 }
 
+#if !UE_LARGE_WORLD_COORDINATES_DISABLED
+void FNiagaraSystemInstanceController::SetVariable(FName InVariableName, FVector InValue)
+{
+	if (OverrideParameters)
+	{
+		OverrideParameters->SetPositionParameterValue(InValue, InVariableName, true);
+	}
+}
+#endif
+
 void FNiagaraSystemInstanceController::SetVariable(FName InVariableName, FVector4f InValue)
 {
 	if (OverrideParameters)
diff --git a/Engine/Plugins/FX/Niagara/Source/Niagara/Private/NiagaraTypes.cpp b/Engine/Plugins/FX/Niagara/Source/Niagara/Private/NiagaraTypes.cpp
--- a/Engine/Plugins/FX/Niagara/Source/Niagara/Private/NiagaraTypes.cpp
+++ b/Engine/Plugins/FX/Niagara/Source/Niagara/Private/NiagaraTypes.cpp
@@ -22,3 +22,28 @@ void FNiagaraVariableMetaData::CopyUserEditableMetaData(const FNiagaraVariableMe
 		};
 	}
 }
+
+FNiagaraLWCConverter::FNiagaraLWCConverter(FVector InSystemWorldPos)
+{
+	SystemWorldPos = InSystemWorldPos;
+}
+
+FVector3f FNiagaraLWCConverter::ConvertWorldToSimulationVector(FVector WorldPosition) const
+{
+	return FVector3f(WorldPosition - SystemWorldPos);
+}
+
+FNiagaraPosition FNiagaraLWCConverter::ConvertWorldToSimulationPosition(FVector WorldPosition) const
+{
+	return FNiagaraPosition(ConvertWorldToSimulationVector(WorldPosition));
+}
+
+FVector FNiagaraLWCConverter::ConvertSimulationPositionToWorld(FNiagaraPosition SimulationPosition) const
+{
+	return ConvertSimulationVectorToWorld(SimulationPosition);
+}
+
+FVector FNiagaraLWCConverter::ConvertSimulationVectorToWorld(FVector3f SimulationPosition) const
+{
+	return FVector(SimulationPosition) + SystemWorldPos;
+}
diff --git a/Engine/Plugins/FX/Niagara/Source/Niagara/Private/NiagaraUserRedirectionParameterStore.cpp b/Engine/Plugins/FX/Niagara/Source/Niagara/Private/NiagaraUserRedirectionParameterStore.cpp
--- a/Engine/Plugins/FX/Niagara/Source/Niagara/Private/NiagaraUserRedirectionParameterStore.cpp
+++ b/Engine/Plugins/FX/Niagara/Source/Niagara/Private/NiagaraUserRedirectionParameterStore.cpp
@@ -1,8 +1,6 @@
 // Copyright Epic Games, Inc. All Rights Reserved.
 
 #include "NiagaraUserRedirectionParameterStore.h"
-#include "NiagaraStats.h"
-#include "NiagaraDataInterface.h"
 #include "NiagaraSystemInstance.h"
 
 FNiagaraUserRedirectionParameterStore::FNiagaraUserRedirectionParameterStore() : FNiagaraParameterStore()
@@ -110,7 +108,6 @@ void FNiagaraUserRedirectionParameterStore::SanityCheckData(bool bInitInterfaces
 			UE_LOG(LogNiagara, Warning, TEXT("User parameter redirect exists but no real value! Param: %s Owner:%s"), *Var.GetName().ToString(), GetOwner() != nullptr ? *GetOwner()->GetPathName() : TEXT("Unknown owner"));
 		}
 	}
-
 }
 
 bool FNiagaraUserRedirectionParameterStore::RemoveParameter(const FNiagaraVariableBase& InVar)
@@ -143,6 +140,20 @@ void FNiagaraUserRedirectionParameterStore::Reset(bool bClearBindings /*= true*/
 	UserParameterRedirects.Reset();
 }
 
+#if WITH_EDITORONLY_DATA
+void FNiagaraUserRedirectionParameterStore::ConvertParameterType(const FNiagaraVariable& ExistingParam,	const FNiagaraTypeDefinition& NewType)
+{
+	FNiagaraParameterStore::ConvertParameterType(ExistingParam, NewType);
+
+	FNiagaraVariable OldVar = GetUserRedirection(ExistingParam);
+	FNiagaraVariable OldRedirect = ExistingParam;
+	if (UserParameterRedirects.RemoveAndCopyValue(OldVar, OldRedirect))
+	{
+		UserParameterRedirects.Add(FNiagaraVariable(NewType, OldVar.GetName()), FNiagaraVariable(NewType, OldRedirect.GetName()));
+	}
+}
+#endif
+
 bool FNiagaraUserRedirectionParameterStore::SerializeFromMismatchedTag(const FPropertyTag & Tag, FStructuredArchive::FSlot Slot)
 {
 	static FName StoreDataName("NiagaraParameterStore");
diff --git a/Engine/Plugins/FX/Niagara/Source/Niagara/Public/NiagaraCommon.h b/Engine/Plugins/FX/Niagara/Source/Niagara/Public/NiagaraCommon.h
--- a/Engine/Plugins/FX/Niagara/Source/Niagara/Public/NiagaraCommon.h
+++ b/Engine/Plugins/FX/Niagara/Source/Niagara/Public/NiagaraCommon.h
@@ -524,6 +524,13 @@ struct NIAGARA_API FNiagaraFunctionSignature
 	#endif
 	}
 
+	void SetFunctionVersion(uint32 Version)
+	{
+	#if WITH_EDITORONLY_DATA
+		FunctionVersion = Version;
+	#endif
+	}
+
 	void SetDescription(const FText& Desc)
 	{
 	#if WITH_EDITORONLY_DATA
diff --git a/Engine/Plugins/FX/Niagara/Source/Niagara/Public/NiagaraComponent.h b/Engine/Plugins/FX/Niagara/Source/Niagara/Public/NiagaraComponent.h
--- a/Engine/Plugins/FX/Niagara/Source/Niagara/Public/NiagaraComponent.h
+++ b/Engine/Plugins/FX/Niagara/Source/Niagara/Public/NiagaraComponent.h
@@ -419,6 +419,14 @@ public:
 	UFUNCTION(BlueprintCallable, Category = Niagara, meta = (DisplayName = "Set Niagara Variable (Vector3)"))
 	void SetVariableVec3(FName InVariableName, FVector InValue);
 
+	/** Sets a Niagara Position parameter by name, overriding locally if necessary.*/
+	UFUNCTION(BlueprintCallable, Category = Niagara, meta = (DisplayName = "Set Niagara Variable By String (Position)"))
+	void SetNiagaraVariablePosition(const FString& InVariableName, FVector InValue);
+
+	/** Sets a Niagara Position parameter by name, overriding locally if necessary.*/
+	UFUNCTION(BlueprintCallable, Category = Niagara, meta = (DisplayName = "Set Niagara Variable (Position)"))
+	void SetVariablePosition(FName InVariableName, FVector InValue);
+
 	/** Sets a Niagara Vector3 parameter by name, overriding locally if necessary.*/
 	UFUNCTION(BlueprintCallable, Category = Niagara, meta = (DisplayName = "Set Niagara Variable By String (Vector2)"))
 	void SetNiagaraVariableVec2(const FString& InVariableName, FVector2D InValue);
@@ -796,6 +804,8 @@ public:
 
 	FORCEINLINE const FMatrix& GetLocalToWorldInverse() const { return LocalToWorldInverse; }
 
+	const FVector3f& GetLWCRenderTile() const { return RenderData ? RenderData->LWCRenderTile : FVector3f::ZeroVector; }
+
 	FRHIUniformBuffer* GetCustomUniformBuffer(bool bHasVelocity, const FBox& PreSkinnedBounds = FBox(ForceInitToZero)) const;
 
 	virtual FPrimitiveViewRelevance GetViewRelevance(const FSceneView* View) const override;
diff --git a/Engine/Plugins/FX/Niagara/Source/Niagara/Public/NiagaraComponentRendererProperties.h b/Engine/Plugins/FX/Niagara/Source/Niagara/Public/NiagaraComponentRendererProperties.h
--- a/Engine/Plugins/FX/Niagara/Source/Niagara/Public/NiagaraComponentRendererProperties.h
+++ b/Engine/Plugins/FX/Niagara/Source/Niagara/Public/NiagaraComponentRendererProperties.h
@@ -136,6 +136,7 @@ public:
 	static FNiagaraTypeDefinition GetFVector2DDef();
 	static FNiagaraTypeDefinition GetFVectorDef();
 	static FNiagaraTypeDefinition GetFVector4Def();
+	static FNiagaraTypeDefinition GetFVector3fDef();
 	static FNiagaraTypeDefinition GetFQuatDef();
 
 	virtual void CacheFromCompiledData(const FNiagaraDataSetCompiledData* CompiledData) override;
diff --git a/Engine/Plugins/FX/Niagara/Source/Niagara/Public/NiagaraLightRendererProperties.h b/Engine/Plugins/FX/Niagara/Source/Niagara/Public/NiagaraLightRendererProperties.h
--- a/Engine/Plugins/FX/Niagara/Source/Niagara/Public/NiagaraLightRendererProperties.h
+++ b/Engine/Plugins/FX/Niagara/Source/Niagara/Public/NiagaraLightRendererProperties.h
@@ -102,7 +102,7 @@ public:
 	UPROPERTY(EditAnywhere, AdvancedDisplay, Category = "Bindings")
 	FNiagaraVariableAttributeBinding RendererVisibilityTagBinding;
 
-	FNiagaraDataSetAccessor<FVector3f> PositionDataSetAccessor;
+	FNiagaraDataSetAccessor<FNiagaraPosition> PositionDataSetAccessor;
 	FNiagaraDataSetAccessor<FLinearColor> ColorDataSetAccessor;
 	FNiagaraDataSetAccessor<float> RadiusDataSetAccessor;
 	FNiagaraDataSetAccessor<float> ExponentDataSetAccessor;
diff --git a/Engine/Plugins/FX/Niagara/Source/Niagara/Public/NiagaraModule.h b/Engine/Plugins/FX/Niagara/Source/Niagara/Public/NiagaraModule.h
--- a/Engine/Plugins/FX/Niagara/Source/Niagara/Public/NiagaraModule.h
+++ b/Engine/Plugins/FX/Niagara/Source/Niagara/Public/NiagaraModule.h
@@ -129,6 +129,7 @@ public:
 	FORCEINLINE static const FNiagaraVariable&  GetVar_Engine_Owner_ZAxis() { return Engine_Owner_ZAxis; }
 	FORCEINLINE static const FNiagaraVariable&  GetVar_Engine_Owner_Scale() { return Engine_Owner_Scale; }
 	FORCEINLINE static const FNiagaraVariable&  GetVar_Engine_Owner_Rotation() { return Engine_Owner_Rotation; }
+	FORCEINLINE static const FNiagaraVariable&  GetVar_Engine_Owner_LWC_Tile() { return Engine_Owner_LWC_Tile; }
 
 	FORCEINLINE static const FNiagaraVariable&  GetVar_Engine_Owner_SystemLocalToWorld() { return Engine_Owner_SystemLocalToWorld; }
 	FORCEINLINE static const FNiagaraVariable&  GetVar_Engine_Owner_SystemWorldToLocal() { return Engine_Owner_SystemWorldToLocal; }
@@ -249,6 +250,7 @@ private:
 	static FNiagaraVariable Engine_Owner_ZAxis;
 	static FNiagaraVariable Engine_Owner_Scale;
 	static FNiagaraVariable Engine_Owner_Rotation;
+	static FNiagaraVariable Engine_Owner_LWC_Tile;
 
 	static FNiagaraVariable Engine_Owner_SystemLocalToWorld;
 	static FNiagaraVariable Engine_Owner_SystemWorldToLocal;
diff --git a/Engine/Plugins/FX/Niagara/Source/Niagara/Public/NiagaraParameterStore.h b/Engine/Plugins/FX/Niagara/Source/Niagara/Public/NiagaraParameterStore.h
--- a/Engine/Plugins/FX/Niagara/Source/Niagara/Public/NiagaraParameterStore.h
+++ b/Engine/Plugins/FX/Niagara/Source/Niagara/Public/NiagaraParameterStore.h
@@ -25,6 +25,21 @@ struct FNiagaraBoundParameter
 
 };
 
+USTRUCT()
+struct FNiagaraPositionSource
+{
+	GENERATED_USTRUCT_BODY()
+
+	FNiagaraPositionSource() : Value(FVector::ZeroVector) {}
+	FNiagaraPositionSource(FName InName, FVector InValue) : Name(InName), Value(InValue) {}
+	
+	UPROPERTY()
+	FName Name;
+	
+	UPROPERTY()
+	FVector Value;
+};
+
 typedef TArray<FNiagaraBoundParameter> FNiagaraBoundParameterArray;
 
 //Binding from one parameter store to another.
@@ -112,7 +127,6 @@ private:
 	bool BindParameters(FNiagaraParameterStore* DestStore, FNiagaraParameterStore* SrcStore, const FNiagaraBoundParameterArray* BoundParameters = nullptr);
 };
 
-
 USTRUCT()
 struct FNiagaraVariableWithOffset : public FNiagaraVariableBase
 {
@@ -166,7 +180,7 @@ private:
 	
 #if WITH_EDITORONLY_DATA
 	/** Map from parameter defs to their offset in the data table or the data interface. TODO: Separate out into a layout and instance class to reduce duplicated data for this?  */
-	UPROPERTY()
+	UPROPERTY(meta=(DeprecatedProperty))
 	TMap<FNiagaraVariable, int32> ParameterOffsets;
 #endif // WITH_EDITORONLY_DATA
 
@@ -187,6 +201,10 @@ private:
 	UPROPERTY()
 	TArray<TObjectPtr<UObject>> UObjects;
 
+	/** Holds position type source data to be later converted to LWC format. We use an array here instead of a map to save some memory and because linear search is faster with the few elements in here. */
+	UPROPERTY()
+	TArray<FNiagaraPositionSource> OriginalPositionData;	
+	
 	/** Bindings between this parameter store and others we push data into when we tick. */
 	typedef TPair<FNiagaraParameterStore*, FNiagaraParameterStoreBinding> BindingPair;
 	TArray<BindingPair> Bindings;
@@ -201,6 +219,8 @@ private:
 	/** Marks our UObjects as dirty. They will be pushed to any bound stores on tick if true. */
 	uint32 bUObjectsDirty : 1;
 
+	uint32 bPositionDataDirty : 1; 
+
 	/** Uniquely identifies the current layout of this parameter store for detecting layout changes. */
 	uint32 LayoutVersion;
 
@@ -208,6 +228,11 @@ private:
 	FOnChanged OnChangedDelegate;
 #endif
 
+	void SetPositionData(const FName& Name, const FVector& Position);
+	bool HasPositionData(const FName& Name) const;
+	const FVector* GetPositionData(const FName& Name) const;
+	void RemovePositionData(const FName& Name);
+
 public:
 	FNiagaraParameterStore();
 	FNiagaraParameterStore(const FNiagaraParameterStore& Other);
@@ -233,10 +258,12 @@ public:
 	FORCEINLINE uint32 GetParametersDirty() const { return bParametersDirty; }
 	FORCEINLINE uint32 GetInterfacesDirty() const { return bInterfacesDirty; }
 	FORCEINLINE uint32 GetUObjectsDirty() const { return bUObjectsDirty; }
+	FORCEINLINE uint32 GetPositionDataDirty() const { return bPositionDataDirty; }
 
 	FORCEINLINE void MarkParametersDirty() { bParametersDirty = true; }
 	FORCEINLINE void MarkInterfacesDirty() { bInterfacesDirty = true; }
 	FORCEINLINE void MarkUObjectsDirty() { bUObjectsDirty = true; }
+	FORCEINLINE void MarkPositionDataDirty() { bPositionDataDirty = true; }
 
 	FORCEINLINE uint32 GetLayoutVersion() const { return LayoutVersion; }
 
@@ -267,6 +294,8 @@ public:
 	virtual bool AddParameter(const FNiagaraVariable& Param, bool bInitialize = true, bool bTriggerRebind = true, int32* OutOffset = nullptr);
 
 #if WITH_EDITORONLY_DATA
+	virtual void ConvertParameterType(const FNiagaraVariable& ExistingParam, const FNiagaraTypeDefinition& NewType);
+	
 	template<typename BufferType>
 	void AddConstantBuffer()
 	{
@@ -283,6 +312,9 @@ public:
 	/** Renames the passed parameter. */
 	virtual void RenameParameter(const FNiagaraVariableBase& Param, FName NewName);
 
+	/** Changes the type of the passed parameter. */
+	virtual void ChangeParameterType(const FNiagaraVariableBase& Param, const FNiagaraTypeDefinition& NewType);
+
 	/** Removes all parameters from this store and releases any data. */
 	virtual void Empty(bool bClearBindings = true);
 
@@ -439,6 +471,11 @@ public:
 			{
 				DestStore.SetUObject(GetUObject(SrcIndex), DestIndex);
 			}
+			else if (Parameter.GetType() == FNiagaraTypeDefinition::GetPositionDef())
+			{
+				const FVector* SourceVector = GetPositionParameterValue(Parameter.GetName());
+				DestStore.SetPositionParameterValue(SourceVector ? *SourceVector : FVector::ZeroVector, Parameter.GetName());
+			}
 			else
 			{
 				DestStore.SetParameterData(GetParameterData(SrcIndex), DestIndex, Parameter.GetSizeInBytes());
@@ -465,10 +502,20 @@ public:
 
 	FString ToString() const;
 
+	bool SetPositionParameterValue(const FVector& InValue, const FName& ParamName, bool bAdd=false);
+	const FVector* GetPositionParameterValue(const FName& ParamName) const;
+	void ResolvePositions(FNiagaraLWCConverter LwcConverter);
+
 	template<typename T>
 	FORCEINLINE_DEBUGGABLE bool SetParameterValue(const T& InValue, const FNiagaraVariable& Param, bool bAdd=false)
 	{
 		check(Param.GetSizeInBytes() == sizeof(T));
+#if WITH_EDITOR
+		if (Param.GetType() == FNiagaraTypeDefinition::GetPositionDef())
+		{
+			check(HasPositionData(Param.GetName()));
+		}
+#endif
 		int32 Offset = IndexOf(Param);
 		if (Offset != INDEX_NONE)
 		{
@@ -510,19 +557,15 @@ public:
 		OnParameterChange();
 	}
 
-	template<typename T>
-	FORCEINLINE_DEBUGGABLE void SetParameterDataTyped(const uint8* Data, int32 Offset)
-	{
-		checkSlow(Data != nullptr);
-		checkSlow((Offset + sizeof(T)) <= ParameterData.Num());
-		uint8* Dest = GetParameterData_Internal(Offset);
-		*((T*)Dest) = *((T*)Data);
-		OnParameterChange();
-	}
-
 	FORCEINLINE_DEBUGGABLE bool SetParameterData(const uint8* Data, const FNiagaraVariable& Param, bool bAdd = false)
 	{
 		checkSlow(Data != nullptr);
+		if (Param.GetType() == FNiagaraTypeDefinition::GetPositionDef())
+		{
+			FNiagaraPosition Value = *reinterpret_cast<const FNiagaraPosition*>(Data);
+			return SetPositionParameterValue(Value, Param.GetName(), bAdd);
+		}
+		
 		int32 Offset = IndexOf(Param);
 		if (Offset != INDEX_NONE)
 		{
@@ -653,6 +696,7 @@ protected:
 	void SetParameterDataArray(const TArray<uint8>& InParameterDataArray, bool bNotifyAsDirty = true);
 	void SetDataInterfaces(const TArray<UNiagaraDataInterface*>& InDataInterfaces, bool bNotifyAsDirty = true);
 	void SetUObjects(const TArray<UObject*>& InUObjects, bool bNotifyAsDirty = true);
+	void SetOriginalPositionData(const TArray<FNiagaraPositionSource>& InOriginalPositionData);
 
 	friend struct FNiagaraParameterStoreToDataSetBinding;    // this should be the only class calling SetParameterByOffset
 };
diff --git a/Engine/Plugins/FX/Niagara/Source/Niagara/Public/NiagaraRendererComponents.h b/Engine/Plugins/FX/Niagara/Source/Niagara/Public/NiagaraRendererComponents.h
--- a/Engine/Plugins/FX/Niagara/Source/Niagara/Public/NiagaraRendererComponents.h
+++ b/Engine/Plugins/FX/Niagara/Source/Niagara/Public/NiagaraRendererComponents.h
@@ -66,7 +66,8 @@ private:
 		USceneComponent* Component,
 		FNiagaraDataSet& Data,
 		int32 ParticleIndex,
-		FComponentPoolEntry& PoolEntry);
+		FComponentPoolEntry& PoolEntry,
+		const FNiagaraLWCConverter& LwcConverter);
 
 	// These property accessor methods are largely copied over from MovieSceneCommonHelpers.h
 	static FComponentPropertyAddress FindPropertyRecursive(void* BasePointer, UStruct* InStruct, TArray<FString>& InPropertyNames, uint32 Index);
diff --git a/Engine/Plugins/FX/Niagara/Source/Niagara/Public/NiagaraRibbonRendererProperties.h b/Engine/Plugins/FX/Niagara/Source/Niagara/Public/NiagaraRibbonRendererProperties.h
--- a/Engine/Plugins/FX/Niagara/Source/Niagara/Public/NiagaraRibbonRendererProperties.h
+++ b/Engine/Plugins/FX/Niagara/Source/Niagara/Public/NiagaraRibbonRendererProperties.h
@@ -416,7 +416,7 @@ public:
 
 	bool								bSortKeyDataSetAccessorIsAge = false;
 	FNiagaraDataSetAccessor<float>		SortKeyDataSetAccessor;
-	FNiagaraDataSetAccessor<FVector3f>	PositionDataSetAccessor;
+	FNiagaraDataSetAccessor<FNiagaraPosition>	PositionDataSetAccessor;
 	FNiagaraDataSetAccessor<float>		NormalizedAgeAccessor;
 	FNiagaraDataSetAccessor<float>		SizeDataSetAccessor;
 	FNiagaraDataSetAccessor<float>		TwistDataSetAccessor;
diff --git a/Engine/Plugins/FX/Niagara/Source/Niagara/Public/NiagaraSettings.h b/Engine/Plugins/FX/Niagara/Source/Niagara/Public/NiagaraSettings.h
--- a/Engine/Plugins/FX/Niagara/Source/Niagara/Public/NiagaraSettings.h
+++ b/Engine/Plugins/FX/Niagara/Source/Niagara/Public/NiagaraSettings.h
@@ -72,6 +72,10 @@ class NIAGARA_API UNiagaraSettings : public UDeveloperSettings
 	bool bSystemViewportInOrbitMode = true;
 #endif // WITH_EDITORONLY_DATA
 
+	/** If true then active effects rebase the simulation positions to not lose precision. Can be turned off if not needed to skip unnecessary rebasing calculations. */
+	UPROPERTY(config, EditAnywhere, Category = Niagara, meta = ( ConfigRestartRequired = true ))
+	bool bSystemsSupportLargeWorldCoordinates = true;
+
 	/** Default effect type to use for effects that don't define their own. Can be null. */
 	UPROPERTY(config, EditAnywhere, Category = Niagara, meta = (AllowedClasses = "NiagaraEffectType"))
 	FSoftObjectPath DefaultEffectType;
diff --git a/Engine/Plugins/FX/Niagara/Source/Niagara/Public/NiagaraSystemGpuComputeProxy.h b/Engine/Plugins/FX/Niagara/Source/Niagara/Public/NiagaraSystemGpuComputeProxy.h
--- a/Engine/Plugins/FX/Niagara/Source/Niagara/Public/NiagaraSystemGpuComputeProxy.h
+++ b/Engine/Plugins/FX/Niagara/Source/Niagara/Public/NiagaraSystemGpuComputeProxy.h
@@ -31,11 +31,13 @@ public:
 	bool RequiresEarlyViewData() const { return bRequiresEarlyViewData; }
 	bool RequiresViewUniformBuffer() const { return bRequiresViewUniformBuffer; }
 	bool RequiresRayTracingScene() const { return bRequiresRayTracingScene; }
+	FVector3f GetSystemLWCTile() const { return SystemLWCTile; }
 
 private:
 	FNiagaraSystemInstance*						DebugOwnerInstance = nullptr;
 	FNiagaraGpuComputeDispatchInterface*		DebugOwnerComputeDispatchInterface = nullptr;
 	int32										ComputeDispatchIndex = INDEX_NONE;
+	FVector3f									SystemLWCTile;
 
 	FNiagaraSystemInstanceID					SystemInstanceID = FNiagaraSystemInstanceID();
 	ENiagaraGpuComputeTickStage::Type			ComputeTickStage = ENiagaraGpuComputeTickStage::PostOpaqueRender;
diff --git a/Engine/Plugins/FX/Niagara/Source/Niagara/Public/NiagaraSystemInstance.h b/Engine/Plugins/FX/Niagara/Source/Niagara/Public/NiagaraSystemInstance.h
--- a/Engine/Plugins/FX/Niagara/Source/Niagara/Public/NiagaraSystemInstance.h
+++ b/Engine/Plugins/FX/Niagara/Source/Niagara/Public/NiagaraSystemInstance.h
@@ -133,6 +133,7 @@ public:
 	void UnbindParameters(bool bFromComplete = false);
 
 	FORCEINLINE FNiagaraParameterStore& GetInstanceParameters() { return InstanceParameters; }
+	FNiagaraLWCConverter GetLWCConverter(bool bLocalSpaceEmitter = false) const;
 
 	FORCEINLINE uint32 GetParameterIndex(bool PreviousFrame = false) const
 	{
@@ -231,6 +232,7 @@ public:
 	FORCEINLINE TArray<TSharedRef<FNiagaraEmitterInstance, ESPMode::ThreadSafe> > &GetEmitters() { return Emitters; }
 	FORCEINLINE const TArray<TSharedRef<FNiagaraEmitterInstance, ESPMode::ThreadSafe> >& GetEmitters() const { return Emitters; }
 	FORCEINLINE const FBox& GetLocalBounds() const { return LocalBounds;  }
+	FORCEINLINE const FVector3f& GetLWCTile() const { return LWCTile;  }
 	TConstArrayView<FNiagaraEmitterExecutionIndex> GetEmitterExecutionOrder() const;
 
 	FORCEINLINE void SetSystemFixedBounds(const FBox& InLocalBounds) { FixedBounds_GT = InLocalBounds; }
@@ -417,6 +419,9 @@ private:
 	/** Call PrepareForSImulation on each data source from the simulations and determine which need per-tick updates.*/
 	void InitDataInterfaces();
 
+	/** The LWC tile of this system instance, used to offset all local simulation relative to the origin */
+	FVector3f LWCTile = FVector3f::ZeroVector;
+
 	/** Index of this instance in the system simulation. */
 	int32 SystemInstanceIndex;
 
diff --git a/Engine/Plugins/FX/Niagara/Source/Niagara/Public/NiagaraSystemInstanceController.h b/Engine/Plugins/FX/Niagara/Source/Niagara/Public/NiagaraSystemInstanceController.h
--- a/Engine/Plugins/FX/Niagara/Source/Niagara/Public/NiagaraSystemInstanceController.h
+++ b/Engine/Plugins/FX/Niagara/Source/Niagara/Public/NiagaraSystemInstanceController.h
@@ -153,6 +153,9 @@ private:
 	void SetVariable(FName InVariableName, float InValue);
 	void SetVariable(FName InVariableName, FVector2f InValue);
 	void SetVariable(FName InVariableName, FVector3f InValue);
+#if !UE_LARGE_WORLD_COORDINATES_DISABLED
+	void SetVariable(FName InVariableName, FVector InValue);
+#endif
 	void SetVariable(FName InVariableName, FVector4f InValue);
 	void SetVariable(FName InVariableName, FLinearColor InValue);
 	void SetVariable(FName InVariableName, FQuat4f InValue);
diff --git a/Engine/Plugins/FX/Niagara/Source/Niagara/Public/NiagaraSystemRenderData.h b/Engine/Plugins/FX/Niagara/Source/Niagara/Public/NiagaraSystemRenderData.h
--- a/Engine/Plugins/FX/Niagara/Source/Niagara/Public/NiagaraSystemRenderData.h
+++ b/Engine/Plugins/FX/Niagara/Source/Niagara/Public/NiagaraSystemRenderData.h
@@ -76,6 +76,7 @@ public:
 	FORCEINLINE void SetRenderingEnabled(bool bInEnabled) { bRenderingEnabled = bInEnabled; }
 	FORCEINLINE bool HasAnyMotionBlurEnabled() const { return bAnyMotionBlurEnabled; }
 
+	FVector3f LWCRenderTile = FVector3f::ZeroVector;
 private:
 	/** Emitter Renderers in the order they appear in the emitters. To be accessed by the GameThread */
 	TArray<FNiagaraRenderer*> EmitterRenderers_GT;
diff --git a/Engine/Plugins/FX/Niagara/Source/Niagara/Public/NiagaraTypes.h b/Engine/Plugins/FX/Niagara/Source/Niagara/Public/NiagaraTypes.h
--- a/Engine/Plugins/FX/Niagara/Source/Niagara/Public/NiagaraTypes.h
+++ b/Engine/Plugins/FX/Niagara/Source/Niagara/Public/NiagaraTypes.h
@@ -72,6 +72,23 @@ private:
 	int32 Value = FNiagaraBool::False;
 };
 
+USTRUCT(meta = (DisplayName = "Position"))
+struct FNiagaraPosition : public FVector3f
+{
+	GENERATED_USTRUCT_BODY()
+
+	FNiagaraPosition() {}
+	
+	explicit FORCEINLINE FNiagaraPosition(EForceInit Init) : Super(Init) {}
+
+	FORCEINLINE FNiagaraPosition(const float& X, const float& Y, const float& Z) : Super(X, Y, Z) {}
+
+	FORCEINLINE FNiagaraPosition(const FVector3f& Other) : Super(Other) {}
+#if !UE_LARGE_WORLD_COORDINATES_DISABLED
+	FORCEINLINE FNiagaraPosition(const FVector& Other) : Super(Other) {}
+#endif
+};
+
 USTRUCT(meta = (DisplayName = "Half", NiagaraInternalType = "true"))
 struct FNiagaraHalf
 {
@@ -185,6 +202,20 @@ struct FNiagaraAssetVersion
 	bool operator<=(const FNiagaraAssetVersion& Other) const { return *this < Other || *this == Other; }
 };
 
+struct NIAGARA_API FNiagaraLWCConverter
+{
+	FNiagaraLWCConverter(FVector InSystemWorldPos = FVector::ZeroVector);
+
+	FVector3f ConvertWorldToSimulationVector(FVector WorldPosition) const;
+	FNiagaraPosition ConvertWorldToSimulationPosition(FVector WorldPosition) const;
+	
+	FVector ConvertSimulationPositionToWorld(FNiagaraPosition SimulationPosition) const;
+	FVector ConvertSimulationVectorToWorld(FVector3f SimulationPosition) const;
+
+private:
+	FVector SystemWorldPos;
+};
+
 FORCEINLINE uint32 GetTypeHash(const FNiagaraAssetVersion& Version)
 {
 	return HashCombine(GetTypeHash(Version.MajorVersion), GetTypeHash(Version.MinorVersion));
@@ -831,7 +862,7 @@ public:
 	bool IsFloatPrimitive() const
 	{
 		return ClassStructOrEnum == FNiagaraTypeDefinition::GetFloatStruct() || ClassStructOrEnum == FNiagaraTypeDefinition::GetVec2Struct() || ClassStructOrEnum == FNiagaraTypeDefinition::GetVec3Struct() || ClassStructOrEnum == FNiagaraTypeDefinition::GetVec4Struct() ||
-			ClassStructOrEnum == FNiagaraTypeDefinition::GetMatrix4Struct() || ClassStructOrEnum == FNiagaraTypeDefinition::GetColorStruct() || ClassStructOrEnum == FNiagaraTypeDefinition::GetQuatStruct();
+			ClassStructOrEnum == FNiagaraTypeDefinition::GetMatrix4Struct() || ClassStructOrEnum == FNiagaraTypeDefinition::GetColorStruct() || ClassStructOrEnum == FNiagaraTypeDefinition::GetQuatStruct() || ClassStructOrEnum == FNiagaraTypeDefinition::GetPositionStruct();
  	}
 
 	bool IsIndexType() const
@@ -888,6 +919,7 @@ public:
 	static const FNiagaraTypeDefinition& GetVec3Def() { return Vec3Def; }
 	static const FNiagaraTypeDefinition& GetVec4Def() { return Vec4Def; }
 	static const FNiagaraTypeDefinition& GetColorDef() { return ColorDef; }
+	static const FNiagaraTypeDefinition& GetPositionDef() { return PositionDef; }
 	static const FNiagaraTypeDefinition& GetQuatDef() { return QuatDef; }
 	static const FNiagaraTypeDefinition& GetMatrix4Def() { return Matrix4Def; }
 	static const FNiagaraTypeDefinition& GetGenericNumericDef() { return NumericDef; }
@@ -915,6 +947,7 @@ public:
 	static UScriptStruct* GetVec3Struct() { return Vec3Struct; }
 	static UScriptStruct* GetVec4Struct() { return Vec4Struct; }
 	static UScriptStruct* GetColorStruct() { return ColorStruct; }
+	static UScriptStruct* GetPositionStruct() { return PositionStruct; }
 	static UScriptStruct* GetQuatStruct() { return QuatStruct; }
 	static UScriptStruct* GetMatrix4Struct() { return Matrix4Struct; }
 	static UScriptStruct* GetGenericNumericStruct() { return NumericStruct; }
@@ -972,6 +1005,7 @@ private:
 	static FNiagaraTypeDefinition Vec3Def;
 	static FNiagaraTypeDefinition Vec4Def;
 	static FNiagaraTypeDefinition ColorDef;
+	static FNiagaraTypeDefinition PositionDef;
 	static FNiagaraTypeDefinition QuatDef;
 	static FNiagaraTypeDefinition Matrix4Def;
 	static FNiagaraTypeDefinition NumericDef;
@@ -1001,6 +1035,7 @@ private:
 	static UScriptStruct* Matrix4Struct;
 	static UScriptStruct* NumericStruct;
 	static UScriptStruct* WildcardStruct;
+	static UScriptStruct* PositionStruct;
 
 	static UScriptStruct* HalfStruct;
 	static UScriptStruct* HalfVec2Struct;
@@ -1059,6 +1094,7 @@ const FNiagaraTypeDefinition& FNiagaraTypeDefinition::Get()
 	if (TIsSame<T, float>::Value) { return FNiagaraTypeDefinition::GetFloatDef(); }
 	if (TIsSame<T, FVector2f>::Value) { return FNiagaraTypeDefinition::GetVec2Def(); }
 	if (TIsSame<T, FVector3f>::Value) { return FNiagaraTypeDefinition::GetVec3Def(); }	
+	if (TIsSame<T, FNiagaraPosition>::Value) { return FNiagaraTypeDefinition::GetPositionDef(); }	
 	if (TIsSame<T, FVector4f>::Value) { return FNiagaraTypeDefinition::GetVec4Def(); }
 	if (TIsSame<T, int32>::Value) { return FNiagaraTypeDefinition::GetIntDef(); }
 	if (TIsSame<T, FNiagaraBool>::Value) { return FNiagaraTypeDefinition::GetBoolDef(); }
@@ -1371,7 +1407,7 @@ struct FNiagaraVariableBase
 		, TypeDef_DEPRECATED(InType)
 #endif
 		{}
-	
+
 	/** Check if Name and Type definition are the same. The actual stored value is not checked here.*/
 	bool operator==(const FNiagaraVariableBase& Other)const
 	{
@@ -1729,7 +1765,7 @@ struct alignas(16) FNiagaraGlobalParameters
 
 // Any change to this structure, or it's GetVariables implementation will require a bump in the CustomNiagaraVersion so that we
 // properly rebuild the scripts
-// You must pad this struct and the results of GetVariables() to a 16 byte boundry.
+// You must pad this struct and the results of GetVariables() to a 16 byte boundary.
 struct alignas(16) FNiagaraSystemParameters
 {
 #if WITH_EDITOR
@@ -1753,7 +1789,7 @@ struct alignas(16) FNiagaraSystemParameters
 
 // Any change to this structure, or it's GetVariables implementation will require a bump in the CustomNiagaraVersion so that we
 // properly rebuild the scripts
-// You must pad this struct and the results of GetVariables() to a 16 byte boundry.
+// You must pad this struct and the results of GetVariables() to a 16 byte boundary.
 struct alignas(16) FNiagaraOwnerParameters
 {
 #if WITH_EDITOR
@@ -1773,11 +1809,12 @@ struct alignas(16) FNiagaraOwnerParameters
 	FVector4f EngineYAxis = FVector4f(0.0f, 1.0f, 0.0f, 0.0f);
 	FVector4f EngineZAxis = FVector4f(0.0f, 0.0f, 1.0f, 0.0f);
 	FVector4f EngineScale = FVector4f(1.0f, 1.0f, 1.0f, 0.0f);
+	FVector4f EngineLWCTile = FVector4f(EForceInit::ForceInitToZero);
 };
 
 // Any change to this structure, or it's GetVariables implementation will require a bump in the CustomNiagaraVersion so that we
 // properly rebuild the scripts
-// You must pad this struct and the results of GetVariables() to a 16 byte boundry.
+// You must pad this struct and the results of GetVariables() to a 16 byte boundary.
 struct alignas(16) FNiagaraEmitterParameters
 {
 #if WITH_EDITOR
diff --git a/Engine/Plugins/FX/Niagara/Source/Niagara/Public/NiagaraUserRedirectionParameterStore.h b/Engine/Plugins/FX/Niagara/Source/Niagara/Public/NiagaraUserRedirectionParameterStore.h
--- a/Engine/Plugins/FX/Niagara/Source/Niagara/Public/NiagaraUserRedirectionParameterStore.h
+++ b/Engine/Plugins/FX/Niagara/Source/Niagara/Public/NiagaraUserRedirectionParameterStore.h
@@ -53,12 +53,17 @@ public:
 			return FNiagaraParameterStore::FindParameterOffset(Redirection ? *Redirection : Parameter);
 		}
 	}
+	
 	virtual bool AddParameter(const FNiagaraVariable& Param, bool bInitialize = true, bool bTriggerRebind = true, int32* OutOffset = nullptr) override;
 	virtual bool RemoveParameter(const FNiagaraVariableBase& InVar) override;
 	virtual void InitFromSource(const FNiagaraParameterStore* SrcStore, bool bNotifyAsDirty) override;
 	virtual void Empty(bool bClearBindings = true) override;
 	virtual void Reset(bool bClearBindings = true) override;
 
+#if WITH_EDITORONLY_DATA
+	virtual void ConvertParameterType(const FNiagaraVariable& ExistingParam, const FNiagaraTypeDefinition& NewType) override;
+#endif
+
 	virtual void SanityCheckData(bool bInitInterfaces = true) override;
 	// ~ End FNiagaraParameterStore overrides
 
@@ -68,7 +73,6 @@ public:
 	/** Turn the input NiagaraVariable into the User namespaced version if needed, independent of whether or not it is in a redirection table.*/
 	NIAGARA_API static void MakeUserVariable(FNiagaraVariableBase& InVar);
 	static bool IsUserParameter(const FNiagaraVariableBase& InVar);
-
 private:
 
 	/** Map from the variables with shortened display names to the original variables with the full namespace */
diff --git a/Engine/Plugins/FX/Niagara/Source/NiagaraCore/Private/NiagaraCustomVersion.cpp b/Engine/Plugins/FX/Niagara/Source/NiagaraCore/Private/NiagaraCustomVersion.cpp
--- a/Engine/Plugins/FX/Niagara/Source/NiagaraCore/Private/NiagaraCustomVersion.cpp
+++ b/Engine/Plugins/FX/Niagara/Source/NiagaraCore/Private/NiagaraCustomVersion.cpp
@@ -8,5 +8,4 @@ const FGuid FNiagaraCustomVersion::GUID(0xFCF57AFA, 0x50764283, 0xB9A9E658, 0xFF
 // Register the custom version with core
 FCustomVersionRegistration GRegisterNiagaraCustomVersion(FNiagaraCustomVersion::GUID, FNiagaraCustomVersion::LatestVersion, TEXT("NiagaraVer"));
 
-const FGuid FNiagaraCustomVersion::LatestScriptCompileVersion(0xB2501D5E, 0x0FC54B9F, 0xB16B6767, 0xBF50D399);
-
+const FGuid FNiagaraCustomVersion::LatestScriptCompileVersion(0x47F836A4, 0x880FE04D, 0x8E8AFB04, 0x6D8CA5C7);
diff --git a/Engine/Plugins/FX/Niagara/Source/NiagaraCore/Public/NiagaraDataInterfaceBase.h b/Engine/Plugins/FX/Niagara/Source/NiagaraCore/Public/NiagaraDataInterfaceBase.h
--- a/Engine/Plugins/FX/Niagara/Source/NiagaraCore/Public/NiagaraDataInterfaceBase.h
+++ b/Engine/Plugins/FX/Niagara/Source/NiagaraCore/Public/NiagaraDataInterfaceBase.h
@@ -28,22 +28,24 @@ DECLARE_EXPORTED_TEMPLATE_INTRINSIC_TYPE_LAYOUT(template<>, TIndexedPtr<UNiagara
 
 struct FNiagaraDataInterfaceArgs
 {
-	FNiagaraDataInterfaceArgs(FNiagaraDataInterfaceProxy* InDataInterface, FNiagaraSystemInstanceID InSystemInstanceID, const FNiagaraGpuComputeDispatchInterface* InComputeDispatchInterface)
+	FNiagaraDataInterfaceArgs(FNiagaraDataInterfaceProxy* InDataInterface, FNiagaraSystemInstanceID InSystemInstanceID, FVector3f InSystemLWCTile, const FNiagaraGpuComputeDispatchInterface* InComputeDispatchInterface)
 		: DataInterface(InDataInterface)
 		, SystemInstanceID(InSystemInstanceID)
 		, ComputeDispatchInterface(InComputeDispatchInterface)
+		, SystemLWCTile(InSystemLWCTile)
 	{
 	}
 
 	FNiagaraDataInterfaceProxy*	DataInterface;
 	FNiagaraSystemInstanceID SystemInstanceID;
 	const FNiagaraGpuComputeDispatchInterface* ComputeDispatchInterface;
+	FVector3f SystemLWCTile;
 };
 
 struct FNiagaraDataInterfaceStageArgs : public FNiagaraDataInterfaceArgs
 {
-	FNiagaraDataInterfaceStageArgs(FNiagaraDataInterfaceProxy* InDataInterface, FNiagaraSystemInstanceID InSystemInstanceID, const FNiagaraGpuComputeDispatchInterface* InComputeDispatchInterface, const FNiagaraComputeInstanceData* InComputeInstanceData, const FNiagaraSimStageData* InSimStageData, bool InIsOutputStage, bool InIsIterationStage)
-		: FNiagaraDataInterfaceArgs(InDataInterface, InSystemInstanceID, InComputeDispatchInterface)
+	FNiagaraDataInterfaceStageArgs(FNiagaraDataInterfaceProxy* InDataInterface, FNiagaraSystemInstanceID InSystemInstanceID, FVector3f InSystemLWCTile, const FNiagaraGpuComputeDispatchInterface* InComputeDispatchInterface, const FNiagaraComputeInstanceData* InComputeInstanceData, const FNiagaraSimStageData* InSimStageData, bool InIsOutputStage, bool InIsIterationStage)
+		: FNiagaraDataInterfaceArgs(InDataInterface, InSystemInstanceID, InSystemLWCTile, InComputeDispatchInterface)
 		, ComputeInstanceData(InComputeInstanceData)
 		, SimStageData(InSimStageData)
 		, IsOutputStage(InIsOutputStage)
@@ -62,8 +64,8 @@ struct FNiagaraDataInterfaceSetArgs : public FNiagaraDataInterfaceArgs
 	typedef TShaderRefBase<FNiagaraShader, FNiagaraShaderMapPointerTable> FShaderReference;
 
 	FNiagaraDataInterfaceSetArgs(
-		FNiagaraDataInterfaceProxy* InDataInterface, FNiagaraSystemInstanceID InSystemInstanceID, const FNiagaraGpuComputeDispatchInterface* InComputeDispatchInterface, const FShaderReference& InShader, const FNiagaraComputeInstanceData* InComputeInstanceData, const FNiagaraSimStageData* InSimStageData, bool InIsOutputStage, bool InIsIterationStage)
-		: FNiagaraDataInterfaceArgs(InDataInterface, InSystemInstanceID, InComputeDispatchInterface)
+		FNiagaraDataInterfaceProxy* InDataInterface, FNiagaraSystemInstanceID InSystemInstanceID, FVector3f InSystemLWCTile, const FNiagaraGpuComputeDispatchInterface* InComputeDispatchInterface, const FShaderReference& InShader, const FNiagaraComputeInstanceData* InComputeInstanceData, const FNiagaraSimStageData* InSimStageData, bool InIsOutputStage, bool InIsIterationStage)
+		: FNiagaraDataInterfaceArgs(InDataInterface, InSystemInstanceID, InSystemLWCTile, InComputeDispatchInterface)
 		, Shader(InShader)
 		, ComputeInstanceData(InComputeInstanceData)
 		, SimStageData(InSimStageData)
diff --git a/Engine/Plugins/FX/Niagara/Source/NiagaraEditor/Private/Customizations/NiagaraComponentDetails.cpp b/Engine/Plugins/FX/Niagara/Source/NiagaraEditor/Private/Customizations/NiagaraComponentDetails.cpp
--- a/Engine/Plugins/FX/Niagara/Source/NiagaraEditor/Private/Customizations/NiagaraComponentDetails.cpp
+++ b/Engine/Plugins/FX/Niagara/Source/NiagaraEditor/Private/Customizations/NiagaraComponentDetails.cpp
@@ -72,6 +72,12 @@ static FNiagaraVariant GetParameterValueFromAsset(const FNiagaraVariableBase& Pa
 		}
 	}
 
+	if (Parameter.GetType() == FNiagaraTypeDefinition::GetPositionDef())
+	{
+		const FVector* Value = UserParameterStore.GetPositionParameterValue(Parameter.GetName());
+		return Value == nullptr ? FNiagaraVariant() : FNiagaraVariant(Value, sizeof(FVector));
+	}
+
 	const uint8* ParameterData = UserParameterStore.GetParameterData(Parameter);
 	if (ParameterData == nullptr)
 	{
@@ -389,7 +395,15 @@ public:
 			}
 			else
 			{
-				TSharedPtr<FStructOnScope> StructOnScope = MakeShareable(new FStructOnScope(Parameter.GetType().GetStruct(), ParameterProxy->Value().GetBytes()));
+				FNiagaraTypeDefinition Type = Parameter.GetType();
+				UStruct* Struct = Type.GetStruct();
+				if (Type == FNiagaraTypeDefinition::GetPositionDef())
+				{
+					static UPackage* CoreUObjectPkg = FindObjectChecked<UPackage>(nullptr, TEXT("/Script/CoreUObject"));
+					static UScriptStruct* VectorStruct = FindObjectChecked<UScriptStruct>(CoreUObjectPkg, TEXT("Vector"));
+					Struct = VectorStruct;
+				}
+				TSharedPtr<FStructOnScope> StructOnScope = MakeShareable(new FStructOnScope(Struct, ParameterProxy->Value().GetBytes()));
 
 				FAddPropertyParams Params = FAddPropertyParams()
 					.UniqueId(Parameter.GetName());
@@ -488,7 +502,14 @@ private:
 					if (DisplayStructPtr != nullptr && DisplayStructPtr->IsValid())
 					{
 						TSharedPtr<FStructOnScope> DisplayStruct = DisplayStructPtr->Pin();
-						FMemory::Memcpy(DisplayStruct->GetStructMemory(), OverrideParameters.GetParameterData(UserParameter), UserParameter.GetSizeInBytes());
+						if (UserParameter.GetType() == FNiagaraTypeDefinition::GetPositionDef())
+						{
+							FMemory::Memcpy(DisplayStruct->GetStructMemory(), OverrideParameters.GetPositionParameterValue(UserParameter.GetName()), UserParameter.GetSizeInBytes());
+						}
+						else
+						{
+							FMemory::Memcpy(DisplayStruct->GetStructMemory(), OverrideParameters.GetParameterData(UserParameter), UserParameter.GetSizeInBytes());
+						}
 					}
 				}
 			}
diff --git a/Engine/Plugins/FX/Niagara/Source/NiagaraEditor/Private/EdGraphSchema_Niagara.cpp b/Engine/Plugins/FX/Niagara/Source/NiagaraEditor/Private/EdGraphSchema_Niagara.cpp
--- a/Engine/Plugins/FX/Niagara/Source/NiagaraEditor/Private/EdGraphSchema_Niagara.cpp
+++ b/Engine/Plugins/FX/Niagara/Source/NiagaraEditor/Private/EdGraphSchema_Niagara.cpp
@@ -1729,7 +1729,7 @@ UNiagaraParameterCollection* UEdGraphSchema_Niagara::VariableIsFromParameterColl
 
 FNiagaraTypeDefinition UEdGraphSchema_Niagara::GetTypeDefForProperty(const FProperty* Property)const
 {
-	if (Property->IsA(FFloatProperty::StaticClass()))
+	if (Property->IsA(FFloatProperty::StaticClass()) || Property->IsA(FDoubleProperty::StaticClass()))
 	{
 		return FNiagaraTypeDefinition::GetFloatDef();
 	}
diff --git a/Engine/Plugins/FX/Niagara/Source/NiagaraEditor/Private/NiagaraEditorModule.cpp b/Engine/Plugins/FX/Niagara/Source/NiagaraEditor/Private/NiagaraEditorModule.cpp
--- a/Engine/Plugins/FX/Niagara/Source/NiagaraEditor/Private/NiagaraEditorModule.cpp
+++ b/Engine/Plugins/FX/Niagara/Source/NiagaraEditor/Private/NiagaraEditorModule.cpp
@@ -944,6 +944,9 @@ void FNiagaraEditorModule::StartupModule()
 	GraphPanelPinFactory->RegisterTypePin(FNiagaraTypeDefinition::GetVec3Struct(), FNiagaraScriptGraphPanelPinFactory::FCreateGraphPin::CreateLambda(
 		[](UEdGraphPin* GraphPin) -> TSharedRef<SGraphPin> { return SNew(TNiagaraGraphPinEditableName<SGraphPinVector>, GraphPin); }));
 
+	GraphPanelPinFactory->RegisterTypePin(FNiagaraTypeDefinition::GetPositionStruct(), FNiagaraScriptGraphPanelPinFactory::FCreateGraphPin::CreateLambda(
+		[](UEdGraphPin* GraphPin) -> TSharedRef<SGraphPin> { return SNew(TNiagaraGraphPinEditableName<SGraphPinVector>, GraphPin); }));
+
 	GraphPanelPinFactory->RegisterTypePin(FNiagaraTypeDefinition::GetVec4Struct(), FNiagaraScriptGraphPanelPinFactory::FCreateGraphPin::CreateLambda(
 		[](UEdGraphPin* GraphPin) -> TSharedRef<SGraphPin> { return SNew(TNiagaraGraphPinEditableName<SGraphPinVector4>, GraphPin); }));
 
@@ -971,6 +974,7 @@ void FNiagaraEditorModule::StartupModule()
 	RegisterTypeUtilities(FNiagaraTypeDefinition::GetVec3Def(), MakeShareable(new FNiagaraEditorVector3TypeUtilities()));
 	RegisterTypeUtilities(FNiagaraTypeDefinition::GetVec4Def(), MakeShareable(new FNiagaraEditorVector4TypeUtilities()));
 	RegisterTypeUtilities(FNiagaraTypeDefinition::GetQuatDef(), MakeShareable(new FNiagaraEditorQuatTypeUtilities()));
+	RegisterTypeUtilities(FNiagaraTypeDefinition::GetPositionDef(), MakeShareable(new FNiagaraEditorVector3TypeUtilities()));
 	RegisterTypeUtilities(FNiagaraTypeDefinition::GetColorDef(), MakeShareable(new FNiagaraEditorColorTypeUtilities()));
 	RegisterTypeUtilities(FNiagaraTypeDefinition::GetMatrix4Def(), MakeShareable(new FNiagaraEditorMatrixTypeUtilities()));
 	RegisterTypeUtilities(FNiagaraTypeDefinition::GetIDDef(), MakeShareable(new FNiagaraEditorNiagaraIDTypeUtilities()));
@@ -1029,6 +1033,12 @@ void FNiagaraEditorModule::StartupModule()
 		VectorTrack->SetChannelsUsed(4);
 		return VectorTrack;
 	}));
+	RegisterParameterTrackCreatorForType(*FNiagaraTypeDefinition::GetPositionStruct(), FOnCreateMovieSceneTrackForParameter::CreateLambda([](FNiagaraVariable InParameter)
+	{
+		UMovieSceneNiagaraVectorParameterTrack* VectorTrack = NewObject<UMovieSceneNiagaraVectorParameterTrack>();
+		VectorTrack->SetChannelsUsed(3);
+		return VectorTrack;
+	}));
 	RegisterParameterTrackCreatorForType(*FNiagaraTypeDefinition::GetColorStruct(), FOnCreateMovieSceneTrackForParameter::CreateLambda([](FNiagaraVariable InParameter) {
 		return NewObject<UMovieSceneNiagaraColorParameterTrack>(); }));
 
diff --git a/Engine/Plugins/FX/Niagara/Source/NiagaraEditor/Private/NiagaraGraph.cpp b/Engine/Plugins/FX/Niagara/Source/NiagaraEditor/Private/NiagaraGraph.cpp
--- a/Engine/Plugins/FX/Niagara/Source/NiagaraEditor/Private/NiagaraGraph.cpp
+++ b/Engine/Plugins/FX/Niagara/Source/NiagaraEditor/Private/NiagaraGraph.cpp
@@ -10,7 +10,6 @@
 #include "NiagaraConstants.h"
 #include "NiagaraCustomVersion.h"
 #include "NiagaraEditorModule.h"
-#include "NiagaraEditorModule.h"
 #include "NiagaraEditorUtilities.h"
 #include "NiagaraHlslTranslator.h"
 #include "NiagaraModule.h"
@@ -198,8 +197,6 @@ void UNiagaraGraph::PostLoad()
 		}
 	}
 
-
-
 	if (bAllZeroes && UniqueNames.Num() > 1)
 	{
 		// Just do the lexicographic sort and assign the call order to their ordered index value.
@@ -245,16 +242,6 @@ void UNiagaraGraph::PostLoad()
 	{
 		SetFlags(RF_Transactional);
 	}
-	
-	FString FullPathName = GetPathName();
-	FString PathName;
-	int ColonPos;
-	if (FullPathName.FindChar(TCHAR('.'), ColonPos))
-	{
-		// GetPathName() returns something similar to "/Path/To/ScriptName.ScriptName:NiagaraScriptSource_N.NiagaraGraph_N"
-		// so this will extract "/Path/To/ScriptName"
-		PathName = FullPathName.Left(ColonPos);
-	}
 
 	// Migrate input condition metadata
 	if (NiagaraVer < FNiagaraCustomVersion::MetaDataAndParametersUpdate)
@@ -345,6 +332,58 @@ void UNiagaraGraph::PostLoad()
 		StandardizeParameterNames();
 	}
 
+	// Fix LWC position type mapping; before LWC, all position were FVector3f, but now that they are a separate type, we need to upgrade the pins and metadata to match
+	const UEdGraphSchema_Niagara* NiagaraSchema = GetDefault<UEdGraphSchema_Niagara>();
+	TArray<UNiagaraNodeParameterMapBase*> NiagaraParameterNodes;
+	GetNodesOfClass<UNiagaraNodeParameterMapBase>(NiagaraParameterNodes);
+	TArray<FNiagaraVariable> OldTypes = FNiagaraConstants::GetOldPositionTypeVariables();
+	for (UNiagaraNodeParameterMapBase* NiagaraNode : NiagaraParameterNodes)
+	{
+		for (UEdGraphPin* Pin : NiagaraNode->Pins)
+		{
+			if (NiagaraNode->IsAddPin(Pin))
+			{
+				continue;
+			}
+			FNiagaraVariable Variable = FNiagaraVariable(NiagaraSchema->PinToTypeDefinition(Pin), Pin->PinName);
+			for (const FNiagaraVariable& OldVarType : OldTypes)
+			{
+				if (Variable == OldVarType)
+				{
+					Pin->PinType = NiagaraSchema->TypeDefinitionToPinType(FNiagaraTypeDefinition::GetPositionDef());
+					NiagaraNode->MarkNodeRequiresSynchronization(__FUNCTION__, true);
+					break;
+				}
+			}
+		}
+	}
+	for (FNiagaraVariable OldVarType : OldTypes)
+	{
+		TObjectPtr<UNiagaraScriptVariable>* ScriptVariablePtr = VariableToScriptVariable.Find(OldVarType);
+		if (ScriptVariablePtr != nullptr && !ScriptVariablePtr->IsNull())
+		{
+			UNiagaraScriptVariable& ScriptVar = *ScriptVariablePtr->Get();
+			ScriptVar.Variable.SetType(FNiagaraTypeDefinition::GetPositionDef());
+			VariableToScriptVariable.Remove(OldVarType);
+
+			FNiagaraVariable NewVarType(FNiagaraTypeDefinition::GetPositionDef(), OldVarType.GetName());
+			VariableToScriptVariable.Add(NewVarType, *ScriptVariablePtr);
+			ScriptVariableChanged(NewVarType);
+		}
+
+		// Also fix stale parameter reference map entries. This should usually happen in RefreshParameterReferences, but since we just changed VariableToScriptVariable we need to fix that as well. 
+		FNiagaraGraphParameterReferenceCollection* ReferenceCollection = ParameterToReferencesMap.Find(OldVarType);
+		if (ReferenceCollection != nullptr)
+		{
+			if (ReferenceCollection->ParameterReferences.Num() > 0)
+			{
+				FNiagaraVariable NewVarType(FNiagaraTypeDefinition::GetPositionDef(), OldVarType.GetName());
+				ParameterToReferencesMap.Add(NewVarType, *ReferenceCollection);
+			}
+			ParameterToReferencesMap.Remove(OldVarType);
+		}
+	}
+
 	InvalidateCachedParameterData();
 }
 
@@ -2907,7 +2946,14 @@ void UNiagaraGraph::RefreshParameterReferences() const
 			continue;
 		}
 
-		AddBindingParameterReference(FNiagaraVariable(Variable->Variable.GetType(), Variable->DefaultBinding.GetName()));
+		FNiagaraTypeDefinition LinkedType = Variable->Variable.GetType();
+		FName BindingName = Variable->DefaultBinding.GetName();
+		if (FNiagaraConstants::GetOldPositionTypeVariables().Contains(FNiagaraVariable(LinkedType, BindingName)))
+		{
+			// it is not uncommon that old assets have vector inputs that default bind to what is now a position type. If we detect that, we change the type to prevent a compiler error.
+			LinkedType = FNiagaraTypeDefinition::GetPositionDef();
+		}
+		AddBindingParameterReference(FNiagaraVariable(LinkedType, BindingName));
 	}
 
 	// If there were any previous parameters which didn't have any references added, remove them here.
diff --git a/Engine/Plugins/FX/Niagara/Source/NiagaraEditor/Private/NiagaraHlslTranslator.cpp b/Engine/Plugins/FX/Niagara/Source/NiagaraEditor/Private/NiagaraHlslTranslator.cpp
--- a/Engine/Plugins/FX/Niagara/Source/NiagaraEditor/Private/NiagaraHlslTranslator.cpp
+++ b/Engine/Plugins/FX/Niagara/Source/NiagaraEditor/Private/NiagaraHlslTranslator.cpp
@@ -975,6 +975,7 @@ static void ConvertFloatToHalf(const FNiagaraCompileOptions& InCompileOptions, T
 			{ FNiagaraTypeDefinition::GetFloatDef(), FNiagaraTypeDefinition::GetHalfDef() },
 			{ FNiagaraTypeDefinition::GetVec2Def(), FNiagaraTypeDefinition::GetHalfVec2Def() },
 			{ FNiagaraTypeDefinition::GetVec3Def(), FNiagaraTypeDefinition::GetHalfVec3Def() },
+			{ FNiagaraTypeDefinition::GetPositionDef(), FNiagaraTypeDefinition::GetHalfVec3Def() },
 			{ FNiagaraTypeDefinition::GetVec4Def(), FNiagaraTypeDefinition::GetHalfVec4Def() },
 			{ FNiagaraTypeDefinition::GetColorDef(), FNiagaraTypeDefinition::GetHalfVec4Def() },
 			{ FNiagaraTypeDefinition::GetQuatDef(), FNiagaraTypeDefinition::GetHalfVec4Def() },
@@ -3980,7 +3981,7 @@ int32 FHlslNiagaraTranslator::AddUniformChunk(FString SymbolName, const FNiagara
 				Chunk.Type = FNiagaraTypeDefinition::GetVec4Def();
 				Chunk.ComponentMask = TEXT(".xy");
 			}
-			else if (Type == FNiagaraTypeDefinition::GetVec3Def())
+			else if (Type == FNiagaraTypeDefinition::GetVec3Def() || Type == FNiagaraTypeDefinition::GetPositionDef())
 			{
 				Chunk.Type = FNiagaraTypeDefinition::GetVec4Def();
 				Chunk.ComponentMask = TEXT(".xyz");
@@ -4364,7 +4365,7 @@ int32 FHlslNiagaraTranslator::GetConstant(const FNiagaraVariable& Constant, FStr
 	FNiagaraVariable LiteralConstant = Constant;
 	if (GetLiteralConstantVariable(LiteralConstant))
 	{
-		checkf(LiteralConstant.GetType() == FNiagaraTypeDefinition::GetBoolDef() || LiteralConstant.GetType() == FNiagaraTypeDefinition::GetVec3Def(), TEXT("Only boolean and vec3 types are currently supported for literal constants."));
+		checkf(LiteralConstant.GetType() == FNiagaraTypeDefinition::GetBoolDef() || LiteralConstant.GetType() == FNiagaraTypeDefinition::GetVec3Def() || LiteralConstant.GetType() == FNiagaraTypeDefinition::GetPositionDef(), TEXT("Only boolean and vec3 types are currently supported for literal constants."));
 		ConstantStr = GenerateConstantString(LiteralConstant);
 	}
 	else
@@ -4468,7 +4469,7 @@ FString FHlslNiagaraTranslator::GenerateConstantString(const FNiagaraVariable& C
 			float* ValuePtr = (float*)Constant.GetData();
 			ConstantStr = FString::Printf(TEXT("float2(%g,%g)"), *ValuePtr, *(ValuePtr + 1));
 		}
-		else if (Type == FNiagaraTypeDefinition::GetVec3Def())
+		else if (Type == FNiagaraTypeDefinition::GetVec3Def() || Type == FNiagaraTypeDefinition::GetPositionDef())
 		{
 			float* ValuePtr = (float*)Constant.GetData();
 			ConstantStr = FString::Printf(TEXT("float3(%g,%g,%g)"), *ValuePtr, *(ValuePtr + 1), *(ValuePtr + 2));
@@ -5786,6 +5787,11 @@ void FHlslNiagaraTranslator::HandleParameterRead(int32 ParamMapHistoryIdx, const
 
 				int32 Out = INDEX_NONE;
 				FNiagaraVariable BindVar = FNiagaraVariable(InVar.GetType(), Bind.GetName());
+				if (FNiagaraConstants::GetOldPositionTypeVariables().Contains(BindVar))
+				{
+					// Old assets often have vector inputs that default bind to what is now a position type. If we detect that, we change the type to prevent a compiler error.
+					BindVar.SetType(FNiagaraTypeDefinition::GetPositionDef());
+				}
 				HandleParameterRead(ActiveStageIdx, BindVar, nullptr, ErrorNode, Out, nullptr);
 
 				if (Out != INDEX_NONE)
@@ -8605,7 +8611,7 @@ FString FHlslNiagaraTranslator::GetHlslDefaultForType(const FNiagaraTypeDefiniti
 	{
 		return "float2(0.0, 0.0)";
 	}
-	else if (Type == FNiagaraTypeDefinition::GetVec3Def())
+	else if (Type == FNiagaraTypeDefinition::GetVec3Def() || Type == FNiagaraTypeDefinition::GetPositionDef())
 	{
 		return "float3(0.0, 0.0, 0.0)";
 	}
@@ -8643,6 +8649,7 @@ bool FHlslNiagaraTranslator::IsBuiltInHlslType(const FNiagaraTypeDefinition& Typ
 		Type == FNiagaraTypeDefinition::GetVec3Def() ||
 		Type == FNiagaraTypeDefinition::GetVec4Def() ||
 		Type == FNiagaraTypeDefinition::GetColorDef() ||
+		Type == FNiagaraTypeDefinition::GetPositionDef() ||
 		Type == FNiagaraTypeDefinition::GetQuatDef() ||
 		Type == FNiagaraTypeDefinition::GetMatrix4Def() ||
 		Type == FNiagaraTypeDefinition::GetIntDef() ||
@@ -8664,7 +8671,7 @@ FString FHlslNiagaraTranslator::GetStructHlslTypeName(const FNiagaraTypeDefiniti
 	{
 		return "float2";
 	}
-	else if (Type == FNiagaraTypeDefinition::GetVec3Def())
+	else if (Type == FNiagaraTypeDefinition::GetVec3Def() || Type == FNiagaraTypeDefinition::GetPositionDef())
 	{
 		return "float3";
 	}
@@ -8763,6 +8770,7 @@ bool FHlslNiagaraTranslator::IsHlslBuiltinVector(const FNiagaraTypeDefinition& T
 		(Type == FNiagaraTypeDefinition::GetVec3Def()) ||
 		(Type == FNiagaraTypeDefinition::GetVec4Def()) ||
 		(Type == FNiagaraTypeDefinition::GetQuatDef()) ||
+		(Type == FNiagaraTypeDefinition::GetPositionDef()) ||
 		(Type == FNiagaraTypeDefinition::GetColorDef()))
 	{
 		return true;
diff --git a/Engine/Plugins/FX/Niagara/Source/NiagaraEditor/Private/NiagaraNodeDataSetBase.cpp b/Engine/Plugins/FX/Niagara/Source/NiagaraEditor/Private/NiagaraNodeDataSetBase.cpp
--- a/Engine/Plugins/FX/Niagara/Source/NiagaraEditor/Private/NiagaraNodeDataSetBase.cpp
+++ b/Engine/Plugins/FX/Niagara/Source/NiagaraEditor/Private/NiagaraNodeDataSetBase.cpp
@@ -260,6 +260,11 @@ bool UNiagaraNodeDataSetBase::GetSupportedNiagaraTypeDef(const FProperty* Proper
 		TypeDef = FNiagaraTypeDefinition::GetColorDef();
 		return true;
 	}
+	else if (StructProp && StructProp->Struct == FNiagaraTypeDefinition::GetPositionStruct())
+	{
+		TypeDef = FNiagaraTypeDefinition::GetPositionDef();
+		return true;
+	}
 	else if (StructProp && StructProp->Struct == FNiagaraTypeDefinition::GetQuatStruct())
 	{
 		TypeDef = FNiagaraTypeDefinition::GetQuatDef();
diff --git a/Engine/Plugins/FX/Niagara/Source/NiagaraEditor/Private/NiagaraNodeFunctionCall.cpp b/Engine/Plugins/FX/Niagara/Source/NiagaraEditor/Private/NiagaraNodeFunctionCall.cpp
--- a/Engine/Plugins/FX/Niagara/Source/NiagaraEditor/Private/NiagaraNodeFunctionCall.cpp
+++ b/Engine/Plugins/FX/Niagara/Source/NiagaraEditor/Private/NiagaraNodeFunctionCall.cpp
@@ -388,7 +388,7 @@ bool UNiagaraNodeFunctionCall::IsDeprecated() const
 bool UNiagaraNodeFunctionCall::FixupPinNames()
 {
 	// This checks the override pins to see if any of the inputs were renamed since we last checked.
-	// Note that this does *not* fix inputs set via rapid iteration parameters, they are fixed by UNiagaraScriptSource::UpdateRenamedParameters.
+	// Note that this does *not* fix inputs set via rapid iteration parameters, they are fixed by UNiagaraScriptSource::FixupRenamedParameters.
 	if (UNiagaraNodeParameterMapSet* OverrideNode = FNiagaraStackGraphUtilities::GetStackFunctionOverrideNode(*this))
 	{
 		FixupFunctionScriptVersion();
@@ -404,7 +404,7 @@ bool UNiagaraNodeFunctionCall::FixupPinNames()
 		TArray<const UEdGraphPin*> ModuleInputPins;
 		FCompileConstantResolver ConstantResolver;
 		GetStackFunctionInputPins(*this, ModuleInputPins, HiddenModulePins, ConstantResolver, FNiagaraStackGraphUtilities::ENiagaraGetStackFunctionInputPinsOptions::ModuleInputsOnly);
-		TSet<FName> ModulePinNames;
+		TMap<FName, const UEdGraphPin*> ModulePinNames;
 		TMap<FGuid, FNiagaraVariable> ModuleInputGuidMapping;
 		for (const UEdGraphPin* ModulePin : ModuleInputPins)
 		{
@@ -414,7 +414,7 @@ bool UNiagaraNodeFunctionCall::FixupPinNames()
 			if (VariableMetaData.IsSet()) // check if there is metadata with a guid available
 			{
 				ModuleInputGuidMapping.Add(VariableMetaData->GetVariableGuid(), InputVar);
-				ModulePinNames.Add(InputName);
+				ModulePinNames.Add(InputName, ModulePin);
 			}
 		}
 
@@ -430,6 +430,20 @@ bool UNiagaraNodeFunctionCall::FixupPinNames()
 				continue;
 			}
 
+			// If the type of our input was changed from vec3 to position we update the override pin type as well
+			FNiagaraVariable InputVar = NiagaraSchema->PinToNiagaraVariable(OverridePin);
+			FNiagaraParameterHandle CurrentPinHandle = FNiagaraParameterHandle(OverridePin->PinName);
+			FName InputName = CurrentPinHandle.GetName();
+			if (ModulePinNames.Contains(InputName) && InputVar.GetType() == FNiagaraTypeDefinition::GetVec3Def())
+			{
+				FNiagaraVariable ModuleInputVar = NiagaraSchema->PinToNiagaraVariable(ModulePinNames[InputName]);
+				if (ModuleInputVar.GetType() == FNiagaraTypeDefinition::GetPositionDef())
+				{
+					OverridePin->PinType = ModulePinNames[InputName]->PinType;
+					bNodeChanged = true;
+				}
+			}
+
 			// Gather module output linked to our inputs
 			if (UEdGraphPin* LinkedValuePin = FNiagaraStackGraphUtilities::GetLinkedValueHandleForFunctionInput(*OverridePin))
 			{
@@ -439,12 +453,16 @@ bool UNiagaraNodeFunctionCall::FixupPinNames()
 				{
 					LinkedModuleOutputs.Add(LinkedValuePin);
 				}
+
+				// when changing the module input type to position, changes the linked value pin as well
+				if (InputVar.GetType() == FNiagaraTypeDefinition::GetPositionDef() && OverridePin->PinType != LinkedValuePin->PinType)
+				{
+					LinkedValuePin->PinType = OverridePin->PinType;
+					bNodeChanged = true;
+				}
 			}
 
 			// check if the pin is still a valid input name from the module
-			FNiagaraVariable InputVar = NiagaraSchema->PinToNiagaraVariable(OverridePin);
-			FNiagaraParameterHandle CurrentPinHandle = FNiagaraParameterHandle(OverridePin->PinName);
-			FName InputName = CurrentPinHandle.GetName();
 			if (ModulePinNames.Contains(InputName))
 			{
 				// update the existing binding data
@@ -575,6 +593,12 @@ bool UNiagaraNodeFunctionCall::FixupPinNames()
 				}
 			}
 		}
+
+		if (bNodeChanged)
+		{
+			UNiagaraGraph* LinkedPinGraph = Cast<UNiagaraGraph>(OverrideNode->GetGraph());
+			LinkedPinGraph->NotifyGraphNeedsRecompile();
+		}
 		
 		return bNodeChanged;
 	}
diff --git a/Engine/Plugins/FX/Niagara/Source/NiagaraEditor/Private/NiagaraParameterMapHistory.cpp b/Engine/Plugins/FX/Niagara/Source/NiagaraEditor/Private/NiagaraParameterMapHistory.cpp
--- a/Engine/Plugins/FX/Niagara/Source/NiagaraEditor/Private/NiagaraParameterMapHistory.cpp
+++ b/Engine/Plugins/FX/Niagara/Source/NiagaraEditor/Private/NiagaraParameterMapHistory.cpp
@@ -1303,6 +1303,11 @@ int32 FNiagaraParameterMapHistoryBuilder::HandleVariableRead(int32 ParamMapIdx,
 				if (Variable && Variable->DefaultMode == ENiagaraDefaultMode::Binding && Variable->DefaultBinding.IsValid())
 				{
 					FNiagaraVariable TempVar = FNiagaraVariable(Var.GetType(), Variable->DefaultBinding.GetName());
+					if (FNiagaraConstants::GetOldPositionTypeVariables().Contains(TempVar))
+					{
+						// Old assets often have vector inputs that default bind to what is now a position type. If we detect that, we change the type to prevent a compiler error.
+						TempVar.SetType(FNiagaraTypeDefinition::GetPositionDef());
+					}
 					int32 FoundIdxBinding = AddVariableToHistory(History, TempVar, TempVar, nullptr);
 
 					History.PerVariableReadHistory[FoundIdxBinding].AddDefaulted_GetRef().ReadPin = FModuleScopedPin(InDefaultPin, ModuleName);
diff --git a/Engine/Plugins/FX/Niagara/Source/NiagaraEditor/Private/NiagaraScriptSource.cpp b/Engine/Plugins/FX/Niagara/Source/NiagaraEditor/Private/NiagaraScriptSource.cpp
--- a/Engine/Plugins/FX/Niagara/Source/NiagaraEditor/Private/NiagaraScriptSource.cpp
+++ b/Engine/Plugins/FX/Niagara/Source/NiagaraEditor/Private/NiagaraScriptSource.cpp
@@ -16,6 +16,7 @@
 #include "NiagaraNodeFunctionCall.h"
 #include "NiagaraNodeInput.h"
 #include "NiagaraNodeOutput.h"
+#include "NiagaraNodeParameterMapSet.h"
 #include "NiagaraScript.h"
 #include "Logging/LogMacros.h"
 #include "ViewModels/Stack/NiagaraStackGraphUtilities.h"
@@ -298,7 +299,7 @@ bool UNiagaraScriptSource::AddModuleIfMissing(FString ModulePath, ENiagaraScript
 	return false;
 }
 
-void UNiagaraScriptSource::FixupRenamedParameters(UNiagaraNode* Node, FNiagaraParameterStore& RapidIterationParameters, const TArray<FNiagaraVariable>& OldRapidIterationVariables, const TSet<FName>& ValidRapidIterationParameterNames, const UNiagaraEmitter* Emitter, ENiagaraScriptUsage ScriptUsage) const
+void UNiagaraScriptSource::FixupRenamedParameters(UNiagaraNode* Node, FNiagaraParameterStore& RapidIterationParameters, const TArray<FNiagaraVariable>& OldRapidIterationVariables, const UNiagaraEmitter* Emitter, ENiagaraScriptUsage ScriptUsage) const
 {
 	UNiagaraNodeFunctionCall* FunctionCallNode = Cast<UNiagaraNodeFunctionCall>(Node);
 	if (FunctionCallNode == nullptr || FunctionCallNode->FunctionScript == nullptr)
@@ -320,45 +321,58 @@ void UNiagaraScriptSource::FixupRenamedParameters(UNiagaraNode* Node, FNiagaraPa
 	// go through the existing rapid iteration params to see if they are either still valid or were renamed
 	for (FNiagaraVariable OldRapidIterationVar : OldRapidIterationVariables)
 	{
-		if (ValidRapidIterationParameterNames.Contains(OldRapidIterationVar.GetName()))
-		{
-			continue;
-		}
-
-		const FGuid* BoundGuid = RapidIterationParameters.ParameterGuidMapping.Find(OldRapidIterationVar);
-		if (BoundGuid)
+		if (FGuid* BoundGuid = RapidIterationParameters.ParameterGuidMapping.Find(OldRapidIterationVar))
 		{
 			for (const UEdGraphPin* ModulePin : ModuleInputPins)
 			{
 				FNiagaraVariable InputVar = NiagaraSchema->PinToNiagaraVariable(ModulePin);
 				TOptional<FNiagaraVariableMetaData> VariableMetaData = Graph->GetMetaData(InputVar);
 
-				if (VariableMetaData.IsSet() && VariableMetaData->GetVariableGuid() == *BoundGuid)
+				FNiagaraParameterHandle AliasedFunctionInputHandle = FNiagaraParameterHandle::CreateAliasedModuleParameterHandle(FNiagaraParameterHandle(ModulePin->PinName), FunctionCallNode);
+				const FString ConstantName = FNiagaraUtilities::CreateRapidIterationConstantName(AliasedFunctionInputHandle.GetParameterHandleString(), *UniqueEmitterName, ScriptUsage);
+				bool bNameMatches = OldRapidIterationVar.GetName() == *ConstantName;
+				
+				// move on if the namespaces don't match
+				int32 NameSpaceEnd;
+				if (ConstantName.FindLastChar(TEXT('.'), NameSpaceEnd))
 				{
-					FNiagaraParameterHandle AliasedFunctionInputHandle = FNiagaraParameterHandle::CreateAliasedModuleParameterHandle(FNiagaraParameterHandle(ModulePin->PinName), FunctionCallNode);
-
-					const FString ConstantName = FNiagaraUtilities::CreateRapidIterationConstantName(AliasedFunctionInputHandle.GetParameterHandleString(), *UniqueEmitterName, ScriptUsage);
-
-					// if the names match move on
-					if (OldRapidIterationVar.GetName() == *ConstantName)
+					if (!OldRapidIterationVar.IsInNameSpace(ConstantName.Left(NameSpaceEnd)))
 					{
 						continue;
 					}
-
-					// before we try to rename make sure we check the namespace so that we're dealing with the right variable
-					int32 NameSpaceEnd;
-					if (ConstantName.FindLastChar(TEXT('.'), NameSpaceEnd))
+				}
+				
+				if (VariableMetaData.IsSet() && VariableMetaData->GetVariableGuid() == *BoundGuid)
+				{
+					// if the names match we have nothing to do
+					if (bNameMatches)
 					{
-						if (!OldRapidIterationVar.IsInNameSpace(ConstantName.Left(NameSpaceEnd)))
-						{
-							continue;
-						}
+						continue;
 					}
 
 					// if the guid matches but the names differ then the parameter was renamed, so lets update the rapid iteration parameter
 					FNiagaraTypeDefinition InputType = NiagaraSchema->PinToTypeDefinition(ModulePin);
 					FNiagaraVariable NewRapidIterationVar = FNiagaraStackGraphUtilities::CreateRapidIterationParameter(UniqueEmitterName, ScriptUsage, AliasedFunctionInputHandle.GetParameterHandleString(), InputType);
 					RapidIterationParameters.RenameParameter(OldRapidIterationVar, NewRapidIterationVar.GetName());
+					break;
+				}
+
+				// check if maybe the variable type changed from vector to position and update the rapid iteration param accordingly
+				if (OldRapidIterationVar.GetType() == FNiagaraTypeDefinition::GetVec3Def() && InputVar.GetType() == FNiagaraTypeDefinition::GetPositionDef())
+				{
+					// the names have to match here, as we don't rename anything but change the type
+					if (!bNameMatches)
+					{
+						continue;
+					}
+
+					// if the guid matches but the types differ then the parameter was changed from vector to position, so lets update the existing rapid iteration parameter
+					if (RapidIterationParameters.IndexOf(OldRapidIterationVar) != INDEX_NONE)
+					{
+						FNiagaraVariable RapidIterationAlias(FNiagaraTypeDefinition::GetPositionDef(), OldRapidIterationVar.GetName());
+						RapidIterationParameters.RemoveParameter(RapidIterationAlias);
+						RapidIterationParameters.ChangeParameterType(OldRapidIterationVar, FNiagaraTypeDefinition::GetPositionDef());
+					}
 				}
 			}
 		}
@@ -455,7 +469,7 @@ void UNiagaraScriptSource::CleanUpOldAndInitializeNewRapidIterationParameters(co
 
 			for(UNiagaraNode* Node : Nodes)
 			{
-				FixupRenamedParameters(Node, RapidIterationParameters, OldRapidIterationVariables, ValidRapidIterationParameterNames, Emitter, ScriptUsage);
+				FixupRenamedParameters(Node, RapidIterationParameters, OldRapidIterationVariables, Emitter, ScriptUsage);
 				InitializeNewRapidIterationParametersForNode(Schema, Node, Emitter, OutputNode->GetUsage(), RapidIterationParameters, ValidRapidIterationParameterNames);
 			}
 		}
@@ -576,3 +590,62 @@ void UNiagaraScriptSource::RenameGraphAssignmentAndSetNodePins(const FName OldNa
 {
 	NodeGraph->RenameAssignmentAndSetNodePins(OldName, NewName);
 }
+
+void UNiagaraScriptSource::GetLinkedPositionTypeInputs(const TArray<FNiagaraVariable>& ParametersToCheck, TSet<FNiagaraVariable>& OutLinkedParameters)
+{
+	if (!NodeGraph)
+	{
+		return;
+	}
+	const UEdGraphSchema_Niagara* Schema = GetDefault<UEdGraphSchema_Niagara>();
+	TArray<UNiagaraNodeOutput*> OutputNodes;
+	NodeGraph->FindOutputNodes(OutputNodes);
+	
+	for (UNiagaraNodeOutput* OutputNode : OutputNodes)
+	{
+		TArray<UNiagaraNode*> TraversalNodes;
+		NodeGraph->BuildTraversal(TraversalNodes, OutputNode, true);
+
+		for (UNiagaraNode* TraversalNode : TraversalNodes)
+		{
+			UNiagaraNodeFunctionCall* FunctionCallNode = Cast<UNiagaraNodeFunctionCall>(TraversalNode);
+			if (!FunctionCallNode)
+			{
+				continue;
+			}
+			UNiagaraNodeParameterMapSet* OverrideNode = FNiagaraStackGraphUtilities::GetStackFunctionOverrideNode(*FunctionCallNode);
+			if (!OverrideNode)
+			{
+				continue;
+			}
+			
+			// find inputs that are overriden
+			TArray<UEdGraphPin*> OverrideInputPins;
+			TArray<UEdGraphPin*> LinkedModuleOutputs;
+			OverrideNode->GetInputPins(OverrideInputPins);
+			for (UEdGraphPin* OverridePin : OverrideInputPins)
+			{
+				if (OverridePin->Direction != EGPD_Input || !OverridePin->PinName.ToString().StartsWith(FunctionCallNode->GetFunctionName()))
+				{
+					continue;
+				}
+
+				// Gather linked value inputs
+				if (UEdGraphPin* LinkedValuePin = FNiagaraStackGraphUtilities::GetLinkedValueHandleForFunctionInput(*OverridePin))
+				{
+					for (const FNiagaraVariable& UserVar : ParametersToCheck)
+					{
+						if (LinkedValuePin->PinName == UserVar.GetName())
+						{
+							FNiagaraVariable LinkedInputVar = Schema->PinToNiagaraVariable(OverridePin);
+							if (LinkedInputVar.GetType() == FNiagaraTypeDefinition::GetPositionDef())
+							{
+								OutLinkedParameters.Add(UserVar);
+							}
+						}
+					}
+				}
+			}
+		}
+	}
+}
diff --git a/Engine/Plugins/FX/Niagara/Source/NiagaraEditor/Private/TypeEditorUtilities/NiagaraVectorTypeEditorUtilities.cpp b/Engine/Plugins/FX/Niagara/Source/NiagaraEditor/Private/TypeEditorUtilities/NiagaraVectorTypeEditorUtilities.cpp
--- a/Engine/Plugins/FX/Niagara/Source/NiagaraEditor/Private/TypeEditorUtilities/NiagaraVectorTypeEditorUtilities.cpp
+++ b/Engine/Plugins/FX/Niagara/Source/NiagaraEditor/Private/TypeEditorUtilities/NiagaraVectorTypeEditorUtilities.cpp
@@ -201,13 +201,13 @@ public:
 
 	virtual void UpdateInternalValueFromStruct(TSharedRef<FStructOnScope> Struct) override
 	{
-		checkf(Struct->GetStruct() == FNiagaraTypeDefinition::GetVec3Struct(), TEXT("Struct type not supported."));	// LWC_TODO: Support for FVector3d likely required here.
+		checkf(Struct->GetStruct() == FNiagaraTypeDefinition::GetVec3Struct() || Struct->GetStruct() == FNiagaraTypeDefinition::GetPositionStruct(), TEXT("Struct type not supported."));	// LWC_TODO: Support for FVector3d likely required here.
 		VectorValue = *((FVector3f*)Struct->GetStructMemory());
 	}
 
 	virtual void UpdateStructFromInternalValue(TSharedRef<FStructOnScope> Struct) override
 	{
-		checkf(Struct->GetStruct() == FNiagaraTypeDefinition::GetVec3Struct(), TEXT("Struct type not supported."));
+		checkf(Struct->GetStruct() == FNiagaraTypeDefinition::GetVec3Struct() || Struct->GetStruct() == FNiagaraTypeDefinition::GetPositionStruct(), TEXT("Struct type not supported."));
 		*((FVector3f*)Struct->GetStructMemory()) = VectorValue;
 	}
 
diff --git a/Engine/Plugins/FX/Niagara/Source/NiagaraEditor/Public/NiagaraScriptSource.h b/Engine/Plugins/FX/Niagara/Source/NiagaraEditor/Public/NiagaraScriptSource.h
--- a/Engine/Plugins/FX/Niagara/Source/NiagaraEditor/Public/NiagaraScriptSource.h
+++ b/Engine/Plugins/FX/Niagara/Source/NiagaraEditor/Public/NiagaraScriptSource.h
@@ -49,7 +49,7 @@ class UNiagaraScriptSource : public UNiagaraScriptSourceBase
 
 	NIAGARAEDITOR_API virtual bool AddModuleIfMissing(FString ModulePath, ENiagaraScriptUsage Usage, bool& bOutFoundModule)override;
 
-	virtual void FixupRenamedParameters(UNiagaraNode* Node, FNiagaraParameterStore& RapidIterationParameters, const TArray<FNiagaraVariable>& OldRapidIterationVariables, const TSet<FName>& ValidRapidIterationParameterNames, const UNiagaraEmitter* Emitter, ENiagaraScriptUsage ScriptUsage) const;
+	virtual void FixupRenamedParameters(UNiagaraNode* Node, FNiagaraParameterStore& RapidIterationParameters, const TArray<FNiagaraVariable>& OldRapidIterationVariables, const UNiagaraEmitter* Emitter, ENiagaraScriptUsage ScriptUsage) const;
 	virtual void CleanUpOldAndInitializeNewRapidIterationParameters(const UNiagaraEmitter* Emitter, ENiagaraScriptUsage ScriptUsage, FGuid ScriptUsageId, FNiagaraParameterStore& RapidIterationParameters) const override;
 	virtual void ForceGraphToRecompileOnNextCheck() override;
 	virtual void RefreshFromExternalChanges() override;
@@ -75,8 +75,9 @@ class UNiagaraScriptSource : public UNiagaraScriptSourceBase
 	/** Rename all graph assignment and map set node pins.
 	 *  Used when synchronizing definitions with source scripts of systems and emitters.
 	 */
-	virtual void RenameGraphAssignmentAndSetNodePins(const FName OldName, const FName NewName);
+	virtual void RenameGraphAssignmentAndSetNodePins(const FName OldName, const FName NewName) override;
 
+	virtual void GetLinkedPositionTypeInputs(const TArray<FNiagaraVariable>& ParametersToCheck, TSet<FNiagaraVariable>& OutLinkedParameters) override;
 private:
 	void OnGraphChanged(const FEdGraphEditAction &Action);
 	void OnGraphDataInterfaceChanged();
diff --git a/Engine/Plugins/FX/Niagara/Source/NiagaraEditorWidgets/Private/Stack/SNiagaraStackParameterStoreEntryValue.cpp b/Engine/Plugins/FX/Niagara/Source/NiagaraEditorWidgets/Private/Stack/SNiagaraStackParameterStoreEntryValue.cpp
--- a/Engine/Plugins/FX/Niagara/Source/NiagaraEditorWidgets/Private/Stack/SNiagaraStackParameterStoreEntryValue.cpp
+++ b/Engine/Plugins/FX/Niagara/Source/NiagaraEditorWidgets/Private/Stack/SNiagaraStackParameterStoreEntryValue.cpp
@@ -3,21 +3,17 @@
 #include "Stack/SNiagaraStackParameterStoreEntryValue.h"
 #include "NiagaraEditorModule.h"
 #include "NiagaraEditorWidgetsStyle.h"
-#include "NiagaraEditorStyle.h"
 #include "INiagaraEditorTypeUtilities.h"
 #include "SNiagaraParameterEditor.h"
 #include "ViewModels/Stack/NiagaraStackParameterStoreEntry.h"
 #include "Widgets/Images/SImage.h"
 #include "Widgets/Layout/SBox.h"
 #include "Widgets/Input/SButton.h"
-#include "Widgets/Input/SComboButton.h"
 #include "Widgets/Notifications/SNotificationList.h"
 #include "Framework/Notifications/NotificationManager.h"
-#include "Framework/MultiBox/MultiBoxBuilder.h"
 #include "PropertyEditorModule.h"
 #include "IStructureDetailsView.h"
 #include "Modules/ModuleManager.h"
-#include "PropertyEditorModule.h"
 #include "PropertyCustomizationHelpers.h"
 
 #define LOCTEXT_NAMESPACE "NiagaraStackParameterStoreEntryValue"
diff --git a/Engine/Plugins/FX/Niagara/Source/NiagaraVertexFactories/Public/NiagaraRibbonVertexFactory.h b/Engine/Plugins/FX/Niagara/Source/NiagaraVertexFactories/Public/NiagaraRibbonVertexFactory.h
--- a/Engine/Plugins/FX/Niagara/Source/NiagaraVertexFactories/Public/NiagaraRibbonVertexFactory.h
+++ b/Engine/Plugins/FX/Niagara/Source/NiagaraVertexFactories/Public/NiagaraRibbonVertexFactory.h
@@ -58,6 +58,7 @@ BEGIN_GLOBAL_SHADER_PARAMETER_STRUCT(FNiagaraRibbonUniformParameters, NIAGARAVER
 	SHADER_PARAMETER(int, ShouldFlipNormalToView)
 	SHADER_PARAMETER(int, U0DistributionMode)
 	SHADER_PARAMETER(int, U1DistributionMode)
+	SHADER_PARAMETER(FVector3f, SystemLWCTile)
 	SHADER_PARAMETER(FVector4f, PackedVData)
 	SHADER_PARAMETER(uint32, bLocalSpace)
 	SHADER_PARAMETER_EX(float, DeltaSeconds, EShaderPrecisionModifier::Half)
diff --git a/Engine/Source/Runtime/Core/Private/Misc/LargeWorldRenderPosition.cpp b/Engine/Source/Runtime/Core/Private/Misc/LargeWorldRenderPosition.cpp
--- a/Engine/Source/Runtime/Core/Private/Misc/LargeWorldRenderPosition.cpp
+++ b/Engine/Source/Runtime/Core/Private/Misc/LargeWorldRenderPosition.cpp
@@ -32,6 +32,22 @@ double FLargeWorldRenderScalar::GetTileSize()
 	return UE_LWC_RENDER_TILE_SIZE;
 }
 
+FVector3f FLargeWorldRenderScalar::GetTileFor(FVector InPosition)
+{
+	if constexpr (UE_LWC_RENDER_TILE_SIZE == 0)
+	{
+		return FVector3f::ZeroVector;
+	}
+	FVector3f LWCTile = InPosition / UE_LWC_RENDER_TILE_SIZE + 0.5;
+
+	// normalize the tile
+	LWCTile.X = FMath::FloorToFloat(LWCTile.X);
+	LWCTile.Y = FMath::FloorToFloat(LWCTile.Y);
+	LWCTile.Z = FMath::FloorToFloat(LWCTile.Z);
+
+	return LWCTile;
+}
+
 FMatrix44f FLargeWorldRenderPosition::MakeToRelativeWorldMatrix(const FMatrix& ToWorld) const
 {
 	return SafeCastMatrix(ToWorld * FTranslationMatrix(-GetTileOffset()));
diff --git a/Engine/Source/Runtime/Core/Public/Misc/LargeWorldRenderPosition.h b/Engine/Source/Runtime/Core/Public/Misc/LargeWorldRenderPosition.h
--- a/Engine/Source/Runtime/Core/Public/Misc/LargeWorldRenderPosition.h
+++ b/Engine/Source/Runtime/Core/Public/Misc/LargeWorldRenderPosition.h
@@ -12,6 +12,7 @@ struct FLargeWorldRenderScalar
 public:
 	// This should become inline at some point, but keeping it in cpp file for now to make iteration faster when experementing with tile size
 	CORE_API static double GetTileSize();
+	CORE_API static FVector3f GetTileFor(FVector InPosition);
 
 	float GetTile() const { return (float)Tile; }
 	float GetOffset() const { return (float)Offset; }
