Fixed the debug icon in Niagara using the UE5 one instead of an old UE4 one
Also changed the scratch pad icon color back to white
Added the parameters icon to legacy parameters as well

diff --git a/Engine/Plugins/FX/Niagara/Source/NiagaraEditor/Private/NiagaraEditorStyle.cpp b/Engine/Plugins/FX/Niagara/Source/NiagaraEditor/Private/NiagaraEditorStyle.cpp
--- a/Engine/Plugins/FX/Niagara/Source/NiagaraEditor/Private/NiagaraEditorStyle.cpp
+++ b/Engine/Plugins/FX/Niagara/Source/NiagaraEditor/Private/NiagaraEditorStyle.cpp
@@ -365,7 +365,6 @@ void InitToolbarIcons(TSharedRef< FSlateStyleSet > Style)
 	Style->Set("NiagaraEditor.LockToChanges.Small", new IMAGE_BRUSH("Icons/icon_levels_LockedReadOnly_40x", Icon20x20));
 	Style->Set("NiagaraEditor.SimulationOptions", new IMAGE_PLUGIN_BRUSH("Icons/Commands/icon_simulationOptions_40x", Icon40x40));
 	Style->Set("NiagaraEditor.SimulationOptions.Small", new IMAGE_PLUGIN_BRUSH("Icons/Commands/icon_simulationOptions_40x", Icon20x20));
-	Style->Set("NiagaraEditor.DebugOptions", new IMAGE_CORE_BRUSH("Icons/icon_tab_DebugTools_16x", Icon20x20));
 
 	Style->Set("Niagara.Asset.ReimportAsset.Needed", new IMAGE_BRUSH("Icons/icon_Reimport_Needed_40x", Icon40x40));
 	Style->Set("Niagara.Asset.ReimportAsset.Default", new IMAGE_BRUSH("Icons/icon_Reimport_40x", Icon40x40));
@@ -626,7 +625,7 @@ void InitTabIcons(TSharedRef<FSlateStyleSet> Style)
 	Style->Set("Tab.Log", new IMAGE_PLUGIN_BRUSH_SVG("Icons/Tabs/Log", Icon16x16));
 	Style->Set("Tab.Debugger", new IMAGE_PLUGIN_BRUSH_SVG("Icons/Tabs/NiagaraDebugger", Icon16x16));
 	Style->Set("Tab.Parameters", new IMAGE_PLUGIN_BRUSH_SVG("Icons/Tabs/Parameters", Icon16x16));
-	Style->Set("Tab.ScratchPad", new IMAGE_PLUGIN_BRUSH_SVG("Icons/Tabs/ScratchPad", Icon16x16, FLinearColor::Yellow));
+	Style->Set("Tab.ScratchPad", new IMAGE_PLUGIN_BRUSH_SVG("Icons/Tabs/ScratchPad", Icon16x16));
 	Style->Set("Tab.ScriptStats", new IMAGE_PLUGIN_BRUSH_SVG("Icons/Tabs/ScriptStats", Icon16x16));
 	Style->Set("Tab.Settings", new IMAGE_PLUGIN_BRUSH_SVG("Icons/Tabs/Settings", Icon16x16));
 	Style->Set("Tab.Spreadsheet", new IMAGE_PLUGIN_BRUSH_SVG("Icons/Tabs/Spreadsheet", Icon16x16));
diff --git a/Engine/Plugins/FX/Niagara/Source/NiagaraEditor/Private/Toolkits/NiagaraSystemToolkit.cpp b/Engine/Plugins/FX/Niagara/Source/NiagaraEditor/Private/Toolkits/NiagaraSystemToolkit.cpp
--- a/Engine/Plugins/FX/Niagara/Source/NiagaraEditor/Private/Toolkits/NiagaraSystemToolkit.cpp
+++ b/Engine/Plugins/FX/Niagara/Source/NiagaraEditor/Private/Toolkits/NiagaraSystemToolkit.cpp
@@ -132,7 +132,8 @@ void FNiagaraSystemToolkit::RegisterTabSpawners(const TSharedRef<class FTabManag
 
 	InTabManager->RegisterTabSpawner(SystemParametersTabID2, FOnSpawnTab::CreateSP(this, &FNiagaraSystemToolkit::SpawnTab_SystemParameters2))
 		.SetDisplayName(LOCTEXT("SystemParameters2", "Legacy Parameters"))
-		.SetGroup(WorkspaceMenuCategory.ToSharedRef());
+		.SetGroup(WorkspaceMenuCategory.ToSharedRef())
+		.SetIcon(FSlateIcon(FNiagaraEditorStyle::GetStyleSetName(), "Tab.Parameters"));
 
 //@todo(ng) disable parameter definitions panel pending bug fixes
 // 	InTabManager->RegisterTabSpawner(SystemParameterDefinitionsTabID, FOnSpawnTab::CreateSP(this, &FNiagaraSystemToolkit::SpawnTab_SystemParameterDefinitions)) 
@@ -1233,8 +1234,7 @@ void FNiagaraSystemToolkit::ExtendToolbar()
 					FOnGetContent::CreateStatic(Local::FillDebugOptionsMenu, Toolkit),
 					LOCTEXT("DebugOptions", "Debug"),
 					LOCTEXT("DebugOptionsTooltip", "Debug options"),
-					FSlateIcon(FNiagaraEditorStyle::GetStyleSetName(), "NiagaraEditor.DebugOptions")
-				);
+					FSlateIcon(FAppStyle::Get().GetStyleSetName(), "Debug"));
 			}
 			ToolbarBuilder.EndSection();
 #endif
