Fix for collecting Niagraa mesh elements running immediately for RT not behaving well with RDG passes for simulation

diff --git a/Engine/Plugins/FX/Niagara/Source/Niagara/Private/NiagaraRendererMeshes.cpp b/Engine/Plugins/FX/Niagara/Source/Niagara/Private/NiagaraRendererMeshes.cpp
--- a/Engine/Plugins/FX/Niagara/Source/Niagara/Private/NiagaraRendererMeshes.cpp
+++ b/Engine/Plugins/FX/Niagara/Source/Niagara/Private/NiagaraRendererMeshes.cpp
@@ -1385,46 +1385,59 @@ void FNiagaraRendererMeshes::GetDynamicRayTracingInstances(FRayTracingMaterialGa
 			// Gpu Target
 			else if (FNiagaraUtilities::AllowComputeShaders(GShaderPlatformForFeatureLevel[FeatureLevel]) && FDataDrivenShaderPlatformInfo::GetSupportsRayTracingIndirectInstanceData(GShaderPlatformForFeatureLevel[FeatureLevel]) )
 			{
-				FRHICommandListImmediate& RHICmdList = Context.RHICmdList;
+				FRHICommandListImmediate& RHICmdList = Context.GraphBuilder.RHICmdList;
 
 				RayTracingInstance.NumTransforms = NumInstances;
 
-				FRWBufferStructured InstanceGPUTransformsBuffer;
-				InstanceGPUTransformsBuffer.Initialize(TEXT("InstanceGPUTransformsBuffer"), 4 * sizeof(float), 3 * NumInstances, BUF_Static);
-				RayTracingInstance.InstanceGPUTransformsSRV = InstanceGPUTransformsBuffer.SRV;
+				FRDGBufferRef InstanceGPUTransformsBufferRef = Context.GraphBuilder.CreateBuffer(
+					FRDGBufferDesc::CreateStructuredDesc(4 * sizeof(float), 3 * NumInstances),
+					TEXT("InstanceGPUTransformsBuffer"));
+
+				const TRefCountPtr<FRDGPooledBuffer>& ExternalBuffer = Context.GraphBuilder.ConvertToExternalBuffer(InstanceGPUTransformsBufferRef);
+				RayTracingInstance.InstanceGPUTransformsSRV = ExternalBuffer->GetOrCreateSRV(FRHIBufferSRVCreateInfo());
+
+				FNiagaraGPURayTracingTransformsCS::FParameters* PassParameters = Context.GraphBuilder.AllocParameters< FNiagaraGPURayTracingTransformsCS::FParameters>();
 
-				FNiagaraGPURayTracingTransformsCS::FParameters PassParameters;
 				{
-					PassParameters.ParticleDataFloatStride		= ParticleMeshRenderData.ParticleFloatDataStride;
+					PassParameters->ParticleDataFloatStride		= ParticleMeshRenderData.ParticleFloatDataStride;
 					//PassParameters.ParticleDataHalfStride		= ParticleMeshRenderData.ParticleHalfDataStride;
-					PassParameters.ParticleDataIntStride		= ParticleMeshRenderData.ParticleIntDataStride;
-					PassParameters.CPUNumInstances				= NumInstances;
-					PassParameters.InstanceCountOffset			= ParticleMeshRenderData.SourceParticleData->GetGPUInstanceCountBufferOffset();
-					PassParameters.PositionDataOffset			= VFVariables[ENiagaraMeshVFLayout::Position].GetGPUOffset();
-					PassParameters.RotationDataOffset			= VFVariables[ENiagaraMeshVFLayout::Rotation].GetGPUOffset();
-					PassParameters.ScaleDataOffset				= VFVariables[ENiagaraMeshVFLayout::Scale].GetGPUOffset();
-					PassParameters.bLocalSpace					= bUseLocalSpace ? 1 : 0;
-					PassParameters.RenderVisibilityOffset		= RendererVisTagOffset;
-					PassParameters.MeshIndexOffset				= MeshIndexOffset;
-					PassParameters.RenderVisibilityValue		= RendererVisibility;
-					PassParameters.MeshIndexValue				= MeshData.SourceMeshIndex;
-					PassParameters.LocalTransform				= LocalTransform;
-					PassParameters.DefaultPosition				= FVector::ZeroVector;
-					PassParameters.DefaultRotation				= FVector4(0.0f, 0.0f, 0.0f, 1.0f);
-					PassParameters.DefaultScale					= FVector(1.0f, 1.0f, 1.0f);
-					PassParameters.MeshScale					= MeshData.Scale;
-					PassParameters.ParticleDataFloatBuffer		= ParticleMeshRenderData.ParticleFloatSRV;
+					PassParameters->ParticleDataIntStride		= ParticleMeshRenderData.ParticleIntDataStride;
+					PassParameters->CPUNumInstances				= NumInstances;
+					PassParameters->InstanceCountOffset			= ParticleMeshRenderData.SourceParticleData->GetGPUInstanceCountBufferOffset();
+					PassParameters->PositionDataOffset			= VFVariables[ENiagaraMeshVFLayout::Position].GetGPUOffset();
+					PassParameters->RotationDataOffset			= VFVariables[ENiagaraMeshVFLayout::Rotation].GetGPUOffset();
+					PassParameters->ScaleDataOffset				= VFVariables[ENiagaraMeshVFLayout::Scale].GetGPUOffset();
+					PassParameters->bLocalSpace					= bUseLocalSpace ? 1 : 0;
+					PassParameters->RenderVisibilityOffset		= RendererVisTagOffset;
+					PassParameters->MeshIndexOffset				= MeshIndexOffset;
+					PassParameters->RenderVisibilityValue		= RendererVisibility;
+					PassParameters->MeshIndexValue				= MeshData.SourceMeshIndex;
+					PassParameters->LocalTransform				= LocalTransform;
+					PassParameters->DefaultPosition				= FVector::ZeroVector;
+					PassParameters->DefaultRotation				= FVector4(0.0f, 0.0f, 0.0f, 1.0f);
+					PassParameters->DefaultScale					= FVector(1.0f, 1.0f, 1.0f);
+					PassParameters->MeshScale					= MeshData.Scale;
+					PassParameters->ParticleDataFloatBuffer		= ParticleMeshRenderData.ParticleFloatSRV;
 					//PassParameters.ParticleDataHalfBuffer		= ParticleMeshRenderData.ParticleHalfSRV;
-					PassParameters.ParticleDataIntBuffer		= ParticleMeshRenderData.ParticleIntSRV;
-					PassParameters.GPUInstanceCountBuffer		= ComputeDispatchInterface->GetGPUInstanceCounterManager().GetInstanceCountBuffer().SRV;
-					PassParameters.TLASTransforms				= InstanceGPUTransformsBuffer.UAV;
+					PassParameters->ParticleDataIntBuffer		= ParticleMeshRenderData.ParticleIntSRV;
+					PassParameters->GPUInstanceCountBuffer		= ComputeDispatchInterface->GetGPUInstanceCounterManager().GetInstanceCountBuffer().SRV;
+					PassParameters->TLASTransforms				= ExternalBuffer->GetOrCreateUAV(FRHIBufferUAVCreateInfo());
 				}
 
-				FNiagaraGPURayTracingTransformsCS::FPermutationDomain PermutationVector;
-				TShaderMapRef<FNiagaraGPURayTracingTransformsCS> ComputeShader(GetGlobalShaderMap(FeatureLevel), PermutationVector);
-				FComputeShaderUtils::Dispatch(RHICmdList, ComputeShader, PassParameters, FComputeShaderUtils::GetGroupCount(int32(NumInstances), int32(FNiagaraGPURayTracingTransformsCS::ThreadGroupSize)));
-			
-				RHICmdList.Transition(FRHITransitionInfo(InstanceGPUTransformsBuffer.UAV, ERHIAccess::Unknown, ERHIAccess::SRVCompute));
+				Context.GraphBuilder.AddPass(
+					RDG_EVENT_NAME("NiagaraGPURayTracingTransforms"),
+					PassParameters,
+					ERDGPassFlags::Compute,
+					[PassParameters, Pass_FeatureLevel = FeatureLevel, NumInstances](FRHICommandList& RHICmdList)
+				{
+					RHICmdList.Transition(FRHITransitionInfo(PassParameters->TLASTransforms, ERHIAccess::Unknown, ERHIAccess::UAVCompute));
+
+					FNiagaraGPURayTracingTransformsCS::FPermutationDomain PermutationVector;
+					TShaderMapRef<FNiagaraGPURayTracingTransformsCS> ComputeShader(GetGlobalShaderMap(Pass_FeatureLevel), PermutationVector);
+					FComputeShaderUtils::Dispatch(RHICmdList, ComputeShader, *PassParameters, FComputeShaderUtils::GetGroupCount(int32(NumInstances), int32(FNiagaraGPURayTracingTransformsCS::ThreadGroupSize)));
+
+					RHICmdList.Transition(FRHITransitionInfo(PassParameters->TLASTransforms, ERHIAccess::Unknown, ERHIAccess::SRVCompute));
+				});
 			}
 		}
 
diff --git a/Engine/Source/Runtime/Engine/Public/SceneManagement.h b/Engine/Source/Runtime/Engine/Public/SceneManagement.h
--- a/Engine/Source/Runtime/Engine/Public/SceneManagement.h
+++ b/Engine/Source/Runtime/Engine/Public/SceneManagement.h
@@ -2543,8 +2543,8 @@ struct FRayTracingMaterialGatheringContext
 	const class FScene* Scene;
 	const FSceneView* ReferenceView;
 	const FSceneViewFamily& ReferenceViewFamily;
-	FRHICommandListImmediate& RHICmdList;
 
+	FRDGBuilder& GraphBuilder;
 	FRayTracingMeshResourceCollector& RayTracingMeshResourceCollector;
 	TArray<FRayTracingDynamicGeometryUpdateParams> DynamicRayTracingGeometriesToUpdate;
 };
diff --git a/Engine/Source/Runtime/Renderer/Private/DeferredShadingRenderer.cpp b/Engine/Source/Runtime/Renderer/Private/DeferredShadingRenderer.cpp
--- a/Engine/Source/Runtime/Renderer/Private/DeferredShadingRenderer.cpp
+++ b/Engine/Source/Runtime/Renderer/Private/DeferredShadingRenderer.cpp
@@ -471,7 +471,7 @@ static void AddDebugRayTracingInstanceFlags(ERayTracingInstanceFlags& InOutFlags
 	}
 }
 
-bool FDeferredShadingSceneRenderer::GatherRayTracingWorldInstancesForView(FRHICommandListImmediate& RHICmdList, FViewInfo& View, FRayTracingScene& RayTracingScene)
+bool FDeferredShadingSceneRenderer::GatherRayTracingWorldInstancesForView(FRDGBuilder& GraphBuilder, FViewInfo& View, FRayTracingScene& RayTracingScene)
 {
 	if (!IsRayTracingEnabled())
 	{
@@ -530,7 +530,7 @@ bool FDeferredShadingSceneRenderer::GatherRayTracingWorldInstancesForView(FRHICo
 		Scene,
 		&View,
 		ViewFamily,
-		RHICmdList,
+		GraphBuilder,
 		*View.RayTracingMeshResourceCollector
 	};
 
@@ -1052,7 +1052,7 @@ bool FDeferredShadingSceneRenderer::GatherRayTracingWorldInstancesForView(FRHICo
 						// We are only doing this for dynamic geometries now
 						SceneInfo->LastRenderTime = CurrentWorldTime;
 						SceneInfo->UpdateComponentLastRenderTime(CurrentWorldTime, /*bUpdateLastRenderTimeOnScreen=*/true);
-						SceneInfo->ConditionalUpdateUniformBuffer(RHICmdList);
+						SceneInfo->ConditionalUpdateUniformBuffer(GraphBuilder.RHICmdList);
 					}
 				}
 			}
@@ -2045,7 +2045,7 @@ void FDeferredShadingSceneRenderer::Render(FRDGBuilder& GraphBuilder)
 	FViewInfo& ReferenceView = Views[ReferenceViewIndex];
 
 	// Prepare the scene for rendering this frame.
-	GatherRayTracingWorldInstancesForView(GraphBuilder.RHICmdList, ReferenceView, RayTracingScene);
+	GatherRayTracingWorldInstancesForView(GraphBuilder, ReferenceView, RayTracingScene);
 
 #endif // RHI_RAYTRACING
 
diff --git a/Engine/Source/Runtime/Renderer/Private/DeferredShadingRenderer.h b/Engine/Source/Runtime/Renderer/Private/DeferredShadingRenderer.h
--- a/Engine/Source/Runtime/Renderer/Private/DeferredShadingRenderer.h
+++ b/Engine/Source/Runtime/Renderer/Private/DeferredShadingRenderer.h
@@ -927,7 +927,7 @@ private:
 	void RenderRayTracingBarycentrics(FRDGBuilder& GraphBuilder, const FViewInfo& View, FRDGTextureRef SceneColorOutputTexture);
 
 	/** Fills RayTracingScene instance list for the given View and adds relevant ray tracing data to the view. Does not reset previous scene contents. */
-	bool GatherRayTracingWorldInstancesForView(FRHICommandListImmediate& RHICmdList, FViewInfo& View, FRayTracingScene& RayTracingScene);
+	bool GatherRayTracingWorldInstancesForView(FRDGBuilder& GraphBuilder, FViewInfo& View, FRayTracingScene& RayTracingScene);
 
 	bool SetupRayTracingPipelineStates(FRHICommandListImmediate& RHICmdList);
 	bool DispatchRayTracingWorldUpdates(FRDGBuilder& GraphBuilder);
