Fixed occlusion DI shader

diff --git a/Engine/Plugins/FX/Niagara/Shaders/Private/NiagaraDataInterfaceOcclusion.ush b/Engine/Plugins/FX/Niagara/Shaders/Private/NiagaraDataInterfaceOcclusion.ush
--- a/Engine/Plugins/FX/Niagara/Shaders/Private/NiagaraDataInterfaceOcclusion.ush
+++ b/Engine/Plugins/FX/Niagara/Shaders/Private/NiagaraDataInterfaceOcclusion.ush
@@ -16,19 +16,19 @@ void DIOcclusion_Rectangle(in FNDIOcclusionContext DIContext, in float3 In_Sampl
 
 	float3 ResolvedLwcSamplePos = LWCToFloat(LWCAdd(LwcSamplePos, PrimaryView.PreViewTranslation));
 	float4 SamplePosition = float4(ResolvedLwcSamplePos, 1);
-	float4 ClipPosition = mul(SamplePosition, ResolvedView.TranslatedWorldToClip);
+	float4 ClipPosition = mul(SamplePosition, PrimaryView.TranslatedWorldToClip);
 	float2 ScreenPosition = ClipPosition.xy / ClipPosition.w;
-	float2 ScreenUV = ScreenPosition * ResolvedView.ScreenPositionScaleBias.xy + ResolvedView.ScreenPositionScaleBias.wz;
+	float2 ScreenUV = ScreenPosition * PrimaryView.ScreenPositionScaleBias.xy + PrimaryView.ScreenPositionScaleBias.wz;
 
 	float Steps = In_SampleSteps <= 1 ? 0 : In_SampleSteps;
 	float TotalSamples = 0;
 	float OccludedSamples = 0;
 
-	float4 SampleWidthClip = mul(float4(ResolvedView.ViewRight * In_SampleWindowWidthWorld, 0) + SamplePosition, ResolvedView.TranslatedWorldToClip);
-	float4 SampleHeightClip = mul(float4(ResolvedView.ViewUp * In_SampleWindowHeightWorld, 0) + SamplePosition, ResolvedView.TranslatedWorldToClip);
+	float4 SampleWidthClip = mul(float4(PrimaryView.ViewRight * In_SampleWindowWidthWorld, 0) + SamplePosition, PrimaryView.TranslatedWorldToClip);
+	float4 SampleHeightClip = mul(float4(PrimaryView.ViewUp * In_SampleWindowHeightWorld, 0) + SamplePosition, PrimaryView.TranslatedWorldToClip);
 				
-	float2 SampleWidthUV = SampleWidthClip.xy / SampleWidthClip.w * ResolvedView.ScreenPositionScaleBias.xy + ResolvedView.ScreenPositionScaleBias.wz;
-	float2 SampleHeightUV = SampleHeightClip.xy / SampleHeightClip.w * ResolvedView.ScreenPositionScaleBias.xy + ResolvedView.ScreenPositionScaleBias.wz;
+	float2 SampleWidthUV = SampleWidthClip.xy / SampleWidthClip.w * PrimaryView.ScreenPositionScaleBias.xy + PrimaryView.ScreenPositionScaleBias.wz;
+	float2 SampleHeightUV = SampleHeightClip.xy / SampleHeightClip.w * PrimaryView.ScreenPositionScaleBias.xy + PrimaryView.ScreenPositionScaleBias.wz;
 				
 	float SampleWidth = ScreenUV.x > 1 ? 0 : SampleWidthUV.x - ScreenUV.x;
 	float SampleHeight = ScreenUV.y > 1 ? 0 : SampleHeightUV.y - ScreenUV.y;
@@ -72,9 +72,9 @@ void DIOcclusion_Circle(in FNDIOcclusionContext DIContext, in float3 In_SampleCe
 
 	float3 ResolvedLwcSamplePos = LWCToFloat(LWCAdd(LwcSamplePos, PrimaryView.PreViewTranslation));
 	float4 SamplePosition = float4(ResolvedLwcSamplePos, 1);
-	float4 ClipPosition = mul(SamplePosition, ResolvedView.TranslatedWorldToClip);
+	float4 ClipPosition = mul(SamplePosition, PrimaryView.TranslatedWorldToClip);
 	float2 ScreenPosition = ClipPosition.xy / ClipPosition.w;
-	float2 ScreenUV = ScreenPosition * ResolvedView.ScreenPositionScaleBias.xy + ResolvedView.ScreenPositionScaleBias.wz;
+	float2 ScreenUV = ScreenPosition * PrimaryView.ScreenPositionScaleBias.xy + PrimaryView.ScreenPositionScaleBias.wz;
 
 	float Rays = In_SampleRays <= 1 ? 0 : In_SampleRays;
 	float Steps = In_SampleStepsPerRay < 1 ? 0 : In_SampleStepsPerRay;
@@ -100,9 +100,9 @@ void DIOcclusion_Circle(in FNDIOcclusionContext DIContext, in float3 In_SampleCe
 			for (int ray = 0; ray < Rays; ray++)
 			{
 				// calc ray direction vector
-				float3 RayDirection = cos(Degrees) * ResolvedView.ViewUp + sin(Degrees) * ResolvedView.ViewRight;
-				float4 RayClip = mul(float4(RayDirection * In_SampleWindowDiameterWorld / 2, 0) + SamplePosition, ResolvedView.TranslatedWorldToClip);
-				float2 RayUV = RayClip.xy / RayClip.w * ResolvedView.ScreenPositionScaleBias.xy + ResolvedView.ScreenPositionScaleBias.wz;
+				float3 RayDirection = cos(Degrees) * PrimaryView.ViewUp + sin(Degrees) * PrimaryView.ViewRight;
+				float4 RayClip = mul(float4(RayDirection * In_SampleWindowDiameterWorld / 2, 0) + SamplePosition, PrimaryView.TranslatedWorldToClip);
+				float2 RayUV = RayClip.xy / RayClip.w * PrimaryView.ScreenPositionScaleBias.xy + PrimaryView.ScreenPositionScaleBias.wz;
 
 				if ((ScreenUV.x > 1 && RayUV.x < 0) || (ScreenUV.y > 1 && RayUV.y < 0) || (ScreenUV.x < 0 && RayUV.x > 1) || (ScreenUV.y < 0 && RayUV.y > 1))
 				{
