- Fix debug draw cook out to work correctly

#rnx

diff --git a/Engine/Plugins/FX/Niagara/Source/Niagara/Classes/NiagaraScript.h b/Engine/Plugins/FX/Niagara/Source/Niagara/Classes/NiagaraScript.h
--- a/Engine/Plugins/FX/Niagara/Source/Niagara/Classes/NiagaraScript.h
+++ b/Engine/Plugins/FX/Niagara/Source/Niagara/Classes/NiagaraScript.h
@@ -190,6 +190,10 @@ public:
 	UPROPERTY()
 	uint32 bUsesRapidIterationParams : 1;
 
+	/** Whether or not we allow debug switches to be used. */
+	UPROPERTY()
+	uint32 bDisableDebugSwitches : 1;
+
 	/** Do we require interpolated spawning */
 	UPROPERTY()
 	uint32 bInterpolatedSpawn : 1;
@@ -227,6 +231,7 @@ public:
 		: CompilerVersionID()
 		, ScriptUsageType(ENiagaraScriptUsage::Function)
 		, bUsesRapidIterationParams(true)
+		, bDisableDebugSwitches(false)
 		, bInterpolatedSpawn(false)
 		, bRequiresPersistentIDs(false)
 		, BaseScriptID_DEPRECATED(0, 0, 0, 0)
diff --git a/Engine/Plugins/FX/Niagara/Source/Niagara/Classes/NiagaraScriptSourceBase.h b/Engine/Plugins/FX/Niagara/Source/Niagara/Classes/NiagaraScriptSourceBase.h
--- a/Engine/Plugins/FX/Niagara/Source/Niagara/Classes/NiagaraScriptSourceBase.h
+++ b/Engine/Plugins/FX/Niagara/Source/Niagara/Classes/NiagaraScriptSourceBase.h
@@ -36,6 +36,7 @@ public:
 	virtual const FNiagaraCompileRequestDataBase* GetDependentRequest(int32 Index) const = 0;
 	virtual FName ResolveEmitterAlias(FName VariableName) const = 0;
 	virtual bool GetUseRapidIterationParams() const = 0;
+	virtual bool GetDisableDebugSwitches() const = 0;
 };
 
 /** External reference to the compile request data generated.*/
diff --git a/Engine/Plugins/FX/Niagara/Source/Niagara/Classes/NiagaraSystem.h b/Engine/Plugins/FX/Niagara/Source/Niagara/Classes/NiagaraSystem.h
--- a/Engine/Plugins/FX/Niagara/Source/Niagara/Classes/NiagaraSystem.h
+++ b/Engine/Plugins/FX/Niagara/Source/Niagara/Classes/NiagaraSystem.h
@@ -567,9 +567,13 @@ public:
 	UPROPERTY(EditAnywhere, AdvancedDisplay, Category = "Performance", meta = (SkipSystemResetOnChange = "true"))
 	uint32 bTrimAttributesOnCook : 1;
 
-	/** If true, forcefully disables all debug switches */
-	UPROPERTY(meta = (SkipSystemResetOnChange = "true"))
-	uint32 bDisableAllDebugSwitches : 1;
+	/** If true all debug switches will be ignored in the editor. */
+	UPROPERTY(EditAnywhere, AdvancedDisplay, Category = "Performance")
+	uint32 bDisableDebugSwitches : 1;
+
+	/** If true all debug switches will be ignored for cooked builds */
+	UPROPERTY(EditAnywhere, AdvancedDisplay, Category = "Performance", meta = (SkipSystemResetOnChange = "true"))
+	uint32 bDisableDebugSwitchesOnCook : 1;
 
 	/** Subscriptions to definitions of parameters. */
 	UPROPERTY()
diff --git a/Engine/Plugins/FX/Niagara/Source/Niagara/Private/NiagaraScript.cpp b/Engine/Plugins/FX/Niagara/Source/Niagara/Private/NiagaraScript.cpp
--- a/Engine/Plugins/FX/Niagara/Source/Niagara/Private/NiagaraScript.cpp
+++ b/Engine/Plugins/FX/Niagara/Source/Niagara/Private/NiagaraScript.cpp
@@ -569,6 +569,7 @@ bool FNiagaraVMExecutableDataId::operator==(const FNiagaraVMExecutableDataId& Re
 		BaseScriptCompileHash != ReferenceSet.BaseScriptCompileHash ||
 #endif
 		bUsesRapidIterationParams != ReferenceSet.bUsesRapidIterationParams ||
+		bDisableDebugSwitches != ReferenceSet.bDisableDebugSwitches ||
 		bInterpolatedSpawn != ReferenceSet.bInterpolatedSpawn ||
 		bRequiresPersistentIDs != ReferenceSet.bRequiresPersistentIDs ||
 		ScriptVersionID != ReferenceSet.ScriptVersionID)
@@ -660,6 +661,15 @@ void FNiagaraVMExecutableDataId::AppendKeyString(FString& KeyString, const FStri
 		KeyString += TEXT("NORI") + Delimiter;
 	}
 
+	if (bDisableDebugSwitches)
+	{
+		KeyString += TEXT("DISBALEDEBUGSWITCH") + Delimiter;
+	}
+	else
+	{
+		KeyString += TEXT("ALLOWDEBUGSWITCH") + Delimiter;
+	}
+
 	if (bAppendObjectForDebugging)
 	{
 		KeyString += TEXT("[AdditionalDefines]") + Delimiter;
@@ -786,6 +796,7 @@ void UNiagaraScript::ComputeVMCompilationId(FNiagaraVMExecutableDataId& Id, FGui
 	Id = FNiagaraVMExecutableDataId();
 
 	Id.bUsesRapidIterationParams = true;
+	Id.bDisableDebugSwitches = false;
 	Id.bInterpolatedSpawn = false;
 	Id.bRequiresPersistentIDs = false;
 	Id.ScriptVersionID = IsVersioningEnabled() ? (VersionGuid.IsValid() ? VersionGuid : ExposedVersion) : FGuid();
@@ -803,6 +814,10 @@ void UNiagaraScript::ComputeVMCompilationId(FNiagaraVMExecutableDataId& Id, FGui
 			{
 				Id.bUsesRapidIterationParams = false;
 			}
+			if (EmitterOwner->bDisableDebugSwitches)
+			{
+				Id.bDisableDebugSwitches = true;
+			}
 			if (EmitterOwner->bCompressAttributes)
 			{
 				Id.AdditionalDefines.Add(TEXT("CompressAttributes"));
@@ -965,6 +980,10 @@ void UNiagaraScript::ComputeVMCompilationId(FNiagaraVMExecutableDataId& Id, FGui
 		{
 			Id.bUsesRapidIterationParams = false;
 		}
+		if (System->bDisableDebugSwitches)
+		{
+			Id.bDisableDebugSwitches = true;
+		}
 		if (System->bCompressAttributes)
 		{
 			Id.AdditionalDefines.Add(TEXT("CompressAttributes"));
@@ -1421,19 +1440,14 @@ void UNiagaraScript::Serialize(FArchive& Ar)
 		if (UNiagaraEmitter* Emitter = GetTypedOuter<UNiagaraEmitter>())
 		{
 			UNiagaraSystem* EmitterOwner = Emitter->GetTypedOuter<UNiagaraSystem>();
-			if (EmitterOwner && EmitterOwner->bBakeOutRapidIteration)
+			if (EmitterOwner->bBakeOutRapidIteration)
 			{
 				bUsesRapidIterationParams = false;
 			}
-
-			if (!Emitter->bBakeOutRapidIteration)
-			{
-				bUsesRapidIterationParams = true;
-			}
 		}
 		else if (UNiagaraSystem* System = GetTypedOuter<UNiagaraSystem>())
 		{
-			if (System && System->bBakeOutRapidIteration)
+			if (System->bBakeOutRapidIteration)
 			{
 				bUsesRapidIterationParams = false;
 			}
diff --git a/Engine/Plugins/FX/Niagara/Source/Niagara/Private/NiagaraSystem.cpp b/Engine/Plugins/FX/Niagara/Source/Niagara/Private/NiagaraSystem.cpp
--- a/Engine/Plugins/FX/Niagara/Source/Niagara/Private/NiagaraSystem.cpp
+++ b/Engine/Plugins/FX/Niagara/Source/Niagara/Private/NiagaraSystem.cpp
@@ -110,7 +110,8 @@ UNiagaraSystem::UNiagaraSystem(const FObjectInitializer& ObjectInitializer)
 , bBakeOutRapidIterationOnCook(true)
 , bTrimAttributes(false)
 , bTrimAttributesOnCook(true)
-, bDisableAllDebugSwitches(false)
+, bDisableDebugSwitches(false)
+, bDisableDebugSwitchesOnCook(true)
 #endif
 , bSupportLargeWorldCoordinates(true)
 , bFixedBounds(false)
@@ -1010,8 +1011,7 @@ void UNiagaraSystem::Serialize(FArchive& Ar)
 	{
 		bBakeOutRapidIteration = bBakeOutRapidIteration || bBakeOutRapidIterationOnCook;
 		bTrimAttributes = bTrimAttributes || bTrimAttributesOnCook;
-
-		bDisableAllDebugSwitches = true;
+		bDisableDebugSwitches = bDisableDebugSwitches || bDisableDebugSwitchesOnCook;
 	}
 #endif
 }
diff --git a/Engine/Plugins/FX/Niagara/Source/NiagaraCore/Private/NiagaraCustomVersion.cpp b/Engine/Plugins/FX/Niagara/Source/NiagaraCore/Private/NiagaraCustomVersion.cpp
--- a/Engine/Plugins/FX/Niagara/Source/NiagaraCore/Private/NiagaraCustomVersion.cpp
+++ b/Engine/Plugins/FX/Niagara/Source/NiagaraCore/Private/NiagaraCustomVersion.cpp
@@ -8,4 +8,4 @@ const FGuid FNiagaraCustomVersion::GUID(0xFCF57AFA, 0x50764283, 0xB9A9E658, 0xFF
 // Register the custom version with core
 FCustomVersionRegistration GRegisterNiagaraCustomVersion(FNiagaraCustomVersion::GUID, FNiagaraCustomVersion::LatestVersion, TEXT("NiagaraVer"));
 
-const FGuid FNiagaraCustomVersion::LatestScriptCompileVersion(0x47F836A4, 0x880FE04D, 0x8E8AFB04, 0x6D8CA5C7);
+const FGuid FNiagaraCustomVersion::LatestScriptCompileVersion(0xD8BFE375, 0xCAF76740, 0x9E0FD451, 0xA3B5984D);
diff --git a/Engine/Plugins/FX/Niagara/Source/NiagaraEditor/Private/NiagaraCompiler.cpp b/Engine/Plugins/FX/Niagara/Source/NiagaraEditor/Private/NiagaraCompiler.cpp
--- a/Engine/Plugins/FX/Niagara/Source/NiagaraEditor/Private/NiagaraCompiler.cpp
+++ b/Engine/Plugins/FX/Niagara/Source/NiagaraEditor/Private/NiagaraCompiler.cpp
@@ -767,6 +767,7 @@ TSharedPtr<FNiagaraCompileRequestDataBase, ESPMode::ThreadSafe> FNiagaraEditorMo
 		check(System->GetSystemSpawnScript()->GetLatestSource() == System->GetSystemUpdateScript()->GetLatestSource());
 		BasePtr->Source = Cast<UNiagaraScriptSource>(System->GetSystemSpawnScript()->GetLatestSource());
 		BasePtr->bUseRapidIterationParams = !System->bBakeOutRapidIteration;
+		BasePtr->bDisableDebugSwitches = System->bDisableDebugSwitches;
 		
 		TArray<FString> EmitterNames;
 
@@ -787,6 +788,7 @@ TSharedPtr<FNiagaraCompileRequestDataBase, ESPMode::ThreadSafe> FNiagaraEditorMo
 			EmitterPtr->SourceName = BasePtr->SourceName;
 			EmitterPtr->Source = Cast<UNiagaraScriptSource>(Handle.GetInstance()->GraphSource);
 			EmitterPtr->bUseRapidIterationParams = BasePtr->bUseRapidIterationParams || (!Handle.GetInstance()->bBakeOutRapidIteration);
+			EmitterPtr->bDisableDebugSwitches = BasePtr->bDisableDebugSwitches;
 			//EmitterPtr->bSimulationStagesEnabled = Handle.GetInstance()->bSimulationStagesEnabled;
 			EmitterPtr->SharedCompileDataInterfaceData = BasePtr->SharedCompileDataInterfaceData;
 			BasePtr->EmitterData.Add(EmitterPtr);
@@ -1128,6 +1130,7 @@ int32 FNiagaraEditorModule::CompileScript(const FNiagaraCompileRequestDataBase*
 	}
 	TranslateOptions.OverrideModuleConstants = CookedRapidIterationParams;
 	TranslateOptions.bParameterRapidIteration = InCompileRequest->GetUseRapidIterationParams();
+	TranslateOptions.bDisableDebugSwitches = InCompileRequest->GetDisableDebugSwitches();
 
 
 	double TranslationStartTime = FPlatformTime::Seconds();
diff --git a/Engine/Plugins/FX/Niagara/Source/NiagaraEditor/Private/NiagaraHlslTranslator.h b/Engine/Plugins/FX/Niagara/Source/NiagaraEditor/Private/NiagaraHlslTranslator.h
--- a/Engine/Plugins/FX/Niagara/Source/NiagaraEditor/Private/NiagaraHlslTranslator.h
+++ b/Engine/Plugins/FX/Niagara/Source/NiagaraEditor/Private/NiagaraHlslTranslator.h
@@ -84,6 +84,8 @@ public:
 	void AddRapidIterationParameters(const FNiagaraParameterStore& InParamStore, FCompileConstantResolver InResolver);
 	virtual bool GetUseRapidIterationParams() const override { return bUseRapidIterationParams; }
 
+	virtual bool GetDisableDebugSwitches() const override { return bDisableDebugSwitches; }
+
 	// Simulation Stage Variables. Sim stage of 0 is always Spawn/Update
 	struct FCompileSimStageData
 	{
@@ -114,6 +116,7 @@ public:
 	TWeakObjectPtr<UNiagaraScriptSource> Source;
 	FString SourceName;
 	bool bUseRapidIterationParams = true;
+	bool bDisableDebugSwitches = false;
 
 	UEnum* ENiagaraScriptCompileStatusEnum;
 	UEnum* ENiagaraScriptUsageEnum;
@@ -305,11 +308,17 @@ public:
 	FHlslNiagaraTranslatorOptions()
 		: SimTarget(ENiagaraSimTarget::CPUSim)
 		, bParameterRapidIteration(true)
+		, bDisableDebugSwitches(false)
 	{
 
 	}
 
-	FHlslNiagaraTranslatorOptions(const FHlslNiagaraTranslatorOptions& InOptions) : SimTarget(InOptions.SimTarget), InstanceParameterNamespaces(InOptions.InstanceParameterNamespaces), bParameterRapidIteration(InOptions.bParameterRapidIteration), OverrideModuleConstants(InOptions.OverrideModuleConstants)
+	FHlslNiagaraTranslatorOptions(const FHlslNiagaraTranslatorOptions& InOptions)
+		: SimTarget(InOptions.SimTarget)
+		, InstanceParameterNamespaces(InOptions.InstanceParameterNamespaces)
+		, bParameterRapidIteration(InOptions.bParameterRapidIteration)
+		, bDisableDebugSwitches(InOptions.bDisableDebugSwitches)
+		, OverrideModuleConstants(InOptions.OverrideModuleConstants)
 	{
 	}
 
@@ -321,6 +330,9 @@ public:
 	/** Whether or not to treat top-level module variables as external values for rapid iteration without need for compilation.*/
 	bool bParameterRapidIteration;
 
+	/** Should we disable debug switches during translation. */
+	bool bDisableDebugSwitches;
+
 	/** Whether or not to override top-level module variables with values from the constant override table. This is only used for variables that were candidates for rapid iteration.*/
 	TArray<FNiagaraVariable> OverrideModuleConstants;
 
@@ -528,6 +540,8 @@ public:
 	virtual int32 GetParameter(const FNiagaraVariable& Parameter);
 	virtual int32 GetRapidIterationParameter(const FNiagaraVariable& Parameter);
 
+	bool DisableDebugSwitches() const { return TranslationOptions.bDisableDebugSwitches; }
+
 	bool IsCompileOptionDefined(const TCHAR* InDefineStr);
 
 	virtual int32 GetAttribute(const FNiagaraVariable& Attribute);
diff --git a/Engine/Plugins/FX/Niagara/Source/NiagaraEditor/Private/NiagaraParameterMapHistory.cpp b/Engine/Plugins/FX/Niagara/Source/NiagaraEditor/Private/NiagaraParameterMapHistory.cpp
--- a/Engine/Plugins/FX/Niagara/Source/NiagaraEditor/Private/NiagaraParameterMapHistory.cpp
+++ b/Engine/Plugins/FX/Niagara/Source/NiagaraEditor/Private/NiagaraParameterMapHistory.cpp
@@ -1632,9 +1632,20 @@ bool FCompileConstantResolver::ResolveConstant(FNiagaraVariable& OutConstant) co
 
 ENiagaraFunctionDebugState FCompileConstantResolver::CalculateDebugState() const
 {
-	const UNiagaraSystem* CurrentSystem =  System? System : (Emitter ? Cast<UNiagaraSystem>(Emitter->GetOuter()) : nullptr);
-	bool bDisableDebug = CurrentSystem ? CurrentSystem->bDisableAllDebugSwitches : false;
-	return bDisableDebug ? ENiagaraFunctionDebugState::NoDebug : DebugState;
+	// System or Emitter
+	const UNiagaraSystem* OwnerSystem = System ? System : (Emitter ? Cast<const UNiagaraSystem>(Emitter->GetOuter()) : nullptr);
+	if ( OwnerSystem != nullptr )
+	{
+		return OwnerSystem->bDisableDebugSwitches ? ENiagaraFunctionDebugState::NoDebug : DebugState;
+	}
+
+	// Translator
+	if ( Translator != nullptr )
+	{
+		return Translator->DisableDebugSwitches() ? ENiagaraFunctionDebugState::NoDebug : DebugState;
+	}
+
+	return DebugState;
 }
 
 
