A few changes to prevent NPC from being static loaded during the PostLoad of UNiagaraScript and UNiagaraGraph
-add preloading from the editor module
-allow history builder to bypass processing NPC if it's being created during PostLoad (an error will be generated if the NPC are queried from the hsitory in this case)

diff --git a/Engine/Plugins/FX/Niagara/Source/NiagaraEditor/Private/EdGraphSchema_Niagara.cpp b/Engine/Plugins/FX/Niagara/Source/NiagaraEditor/Private/EdGraphSchema_Niagara.cpp
--- a/Engine/Plugins/FX/Niagara/Source/NiagaraEditor/Private/EdGraphSchema_Niagara.cpp
+++ b/Engine/Plugins/FX/Niagara/Source/NiagaraEditor/Private/EdGraphSchema_Niagara.cpp
@@ -1666,20 +1666,10 @@ UNiagaraParameterCollection* UEdGraphSchema_Niagara::VariableIsFromParameterColl
 	FString VarName = Var.GetName().ToString();
 	if (VarName.StartsWith(TEXT("NPC.")))
 	{
-		FAssetRegistryModule& AssetRegistryModule = FModuleManager::LoadModuleChecked<FAssetRegistryModule>("AssetRegistry");
-		TArray<FAssetData> CollectionAssets;
-		AssetRegistryModule.Get().GetAssetsByClass(UNiagaraParameterCollection::StaticClass()->GetFName(), CollectionAssets);
-		TArray<FName> ExistingNames;
-		for (FAssetData& CollectionAsset : CollectionAssets)
-		{
-			// asset may not have been fully loaded so give it a chance to do it's PostLoad
-			if (UNiagaraParameterCollection* Collection = EnsureCollectionLoaded(CollectionAsset))
-			{
-				if (VarName.StartsWith(Collection->GetFullNamespace()))
-				{
-					return Collection;
-				}
-			}
+		FNiagaraEditorModule& NiagaraEditorModule = FModuleManager::GetModuleChecked<FNiagaraEditorModule>("NiagaraEditor");
+		if (UNiagaraParameterCollection* Collection = NiagaraEditorModule.FindCollectionForVariable(VarName))
+		{
+			return Collection;
 		}
 	}
 	return nullptr;
@@ -1691,37 +1681,27 @@ UNiagaraParameterCollection* UEdGraphSchema_Niagara::VariableIsFromParameterColl
 
 	if (VarName.StartsWith(TEXT("NPC.")))
 	{
-		FAssetRegistryModule& AssetRegistryModule = FModuleManager::LoadModuleChecked<FAssetRegistryModule>("AssetRegistry");
-		TArray<FAssetData> CollectionAssets;
-		AssetRegistryModule.Get().GetAssetsByClass(UNiagaraParameterCollection::StaticClass()->GetFName(), CollectionAssets);
-		TArray<FName> ExistingNames;
-		for (FAssetData& CollectionAsset : CollectionAssets)
+		FNiagaraEditorModule& NiagaraEditorModule = FModuleManager::GetModuleChecked<FNiagaraEditorModule>("NiagaraEditor");
+		if (UNiagaraParameterCollection* Collection = NiagaraEditorModule.FindCollectionForVariable(VarName))
 		{
-			// asset may not have been fully loaded so give it a chance to do it's PostLoad
-			if (UNiagaraParameterCollection* Collection = EnsureCollectionLoaded(CollectionAsset))
+			const TArray<FNiagaraVariable>& CollectionVariables = Collection->GetParameters();
+			FString BestMatchSoFar;
+
+			for (const FNiagaraVariable& CollVar : CollectionVariables)
 			{
-				if (VarName.StartsWith(Collection->GetFullNamespace()))
+				FString CollVarName = CollVar.GetName().ToString();
+				if (CollVarName == VarName)
 				{
-					const TArray<FNiagaraVariable>& CollectionVariables = Collection->GetParameters();
-					FString BestMatchSoFar;
-
-					for (const FNiagaraVariable& CollVar : CollectionVariables)
-					{
-						FString CollVarName = CollVar.GetName().ToString();
-						if (CollVarName == VarName)
-						{
-							OutVar = CollVar;
-							break;
-						}
-						else if (bAllowPartialMatch && VarName.StartsWith(CollVarName + TEXT(".")) && (BestMatchSoFar.Len() == 0 || CollVarName.Len() > BestMatchSoFar.Len()))
-						{
-							OutVar = CollVar;
-							BestMatchSoFar = CollVarName;
-						}
-					}
-					return Collection;
+					OutVar = CollVar;
+					break;
+				}
+				else if (bAllowPartialMatch && VarName.StartsWith(CollVarName + TEXT(".")) && (BestMatchSoFar.Len() == 0 || CollVarName.Len() > BestMatchSoFar.Len()))
+				{
+					OutVar = CollVar;
+					BestMatchSoFar = CollVarName;
 				}
 			}
+			return Collection;
 		}
 	}
 	return nullptr;
diff --git a/Engine/Plugins/FX/Niagara/Source/NiagaraEditor/Private/NiagaraEditorModule.cpp b/Engine/Plugins/FX/Niagara/Source/NiagaraEditor/Private/NiagaraEditorModule.cpp
--- a/Engine/Plugins/FX/Niagara/Source/NiagaraEditor/Private/NiagaraEditorModule.cpp
+++ b/Engine/Plugins/FX/Niagara/Source/NiagaraEditor/Private/NiagaraEditorModule.cpp
@@ -14,6 +14,7 @@
 #include "ThumbnailRendering/ThumbnailManager.h"
 #include "Stats/Stats.h"
 #include "Subsystems/ImportSubsystem.h"
+#include "UObject/UObjectThreadContext.h"
 
 #include "AssetTypeActions/AssetTypeActions_NiagaraSystem.h"
 #include "AssetTypeActions/AssetTypeActions_NiagaraEmitter.h"
@@ -790,6 +791,15 @@ void FNiagaraEditorModule::StartupModule()
 	RegisterAssetTypeAction(AssetTools, MakeShareable(new FAssetTypeActions_NiagaraParameterDefinitions()));
 	RegisterAssetTypeAction(AssetTools, MakeShareable(new FAssetTypeActions_NiagaraEffectType()));
 
+	// Preload all parameter definition & collection assets so that they will be postloaded before postload calls to scripts/emitters/systems that rely on them.
+	{
+		FAssetRegistryModule& AssetRegistryModule = FModuleManager::LoadModuleChecked<FAssetRegistryModule>("AssetRegistry");
+		AssetRegistryModule.Get().OnFilesLoaded().AddLambda([this]()
+		{
+			RefreshParameterCollections(true /*AllowLoading*/);
+		});
+	}
+
 	UNiagaraSettings::OnSettingsChanged().AddRaw(this, &FNiagaraEditorModule::OnNiagaraSettingsChangedEvent);
 	FCoreUObjectDelegates::GetPreGarbageCollectDelegate().AddRaw(this, &FNiagaraEditorModule::OnPreGarbageCollection);
 	
@@ -1139,6 +1149,13 @@ void FNiagaraEditorModule::ShutdownModule()
 	}
 	CreatedAssetTypeActions.Empty();
 
+	// Clean up asset registry callbacks
+	if (FModuleManager::Get().IsModuleLoaded("AssetRegistry"))
+	{
+		FAssetRegistryModule& AssetRegistryModule = FModuleManager::GetModuleChecked<FAssetRegistryModule>("AssetRegistry");
+		AssetRegistryModule.Get().OnFilesLoaded().RemoveAll(this);
+	}
+
 	UNiagaraSettings::OnSettingsChanged().RemoveAll(this);
 
 	FCoreUObjectDelegates::GetPreGarbageCollectDelegate().RemoveAll(this);
@@ -1279,6 +1296,9 @@ void FNiagaraEditorModule::OnPostEngineInit()
 		UE_LOG(LogNiagaraEditor, Log, TEXT("GEditor isn't valid! Particle reset commands will not work for Niagara components!"));
 	}
 
+	// ensure that the known parameter collections are loaded 
+	RefreshParameterCollections(true /*AllowLoading*/);
+
 #if WITH_NIAGARA_DEBUGGER
 	Debugger = MakeShared<FNiagaraDebugger>();
 	Debugger->Init();
@@ -1793,6 +1813,53 @@ bool FNiagaraEditorModule::DeferredDestructObjects(float InDeltaTime)
 	return false;
 }
 
+void FNiagaraEditorModule::RefreshParameterCollections(bool AllowLoading)
+{
+	const FAssetRegistryModule& AssetRegistryModule = FModuleManager::LoadModuleChecked<FAssetRegistryModule>(TEXT("AssetRegistry"));
+	TArray<FAssetData> ParameterCollectionAssetData;
+	AssetRegistryModule.GetRegistry().GetAssetsByClass(UNiagaraParameterCollection::StaticClass()->GetFName(), ParameterCollectionAssetData);
+
+	CachedParameterCollections.Reset(ParameterCollectionAssetData.Num());
+	for (const FAssetData& ParameterCollectionAssetDatum : ParameterCollectionAssetData)
+	{
+		if (ParameterCollectionAssetDatum.IsAssetLoaded() || AllowLoading)
+		{
+			if (UNiagaraParameterCollection* Collection = Cast<UNiagaraParameterCollection>(ParameterCollectionAssetDatum.GetAsset()))
+			{
+				CachedParameterCollections.Add(MakeWeakObjectPtr(Collection));
+			}
+		}
+	}
+}
+
+UNiagaraParameterCollection* FNiagaraEditorModule::FindCollectionForVariable(const FString& VariableName)
+{
+	auto FindCachedCollectionByPrefix = [this](const FString& Prefix)
+	{
+		for (TWeakObjectPtr<UNiagaraParameterCollection>& CollectionPtr : CachedParameterCollections)
+		{
+			if (UNiagaraParameterCollection* Collection = CollectionPtr.Get())
+			{
+				if (Prefix.StartsWith(Collection->GetFullNamespace()))
+				{
+					return Collection;
+				}
+			}
+		}
+
+		return (UNiagaraParameterCollection*)nullptr;
+	};
+
+	if (UNiagaraParameterCollection* Collection = FindCachedCollectionByPrefix(VariableName))
+	{
+		return Collection;
+	}
+
+	RefreshParameterCollections(!FUObjectThreadContext::Get().IsRoutingPostLoad /*AllowLoading*/);
+
+	return FindCachedCollectionByPrefix(VariableName);
+}
+
 #if NIAGARA_PERF_BASELINES
 void FNiagaraEditorModule::GeneratePerfBaselines(TArray<UNiagaraEffectType*>& BaselinesToGenerate)
 {
diff --git a/Engine/Plugins/FX/Niagara/Source/NiagaraEditor/Private/NiagaraNodeParameterMapBase.cpp b/Engine/Plugins/FX/Niagara/Source/NiagaraEditor/Private/NiagaraNodeParameterMapBase.cpp
--- a/Engine/Plugins/FX/Niagara/Source/NiagaraEditor/Private/NiagaraNodeParameterMapBase.cpp
+++ b/Engine/Plugins/FX/Niagara/Source/NiagaraEditor/Private/NiagaraNodeParameterMapBase.cpp
@@ -40,19 +40,7 @@ UNiagaraNodeParameterMapBase::UNiagaraNodeParameterMapBase()
 
 }
 
-TArray<FNiagaraParameterMapHistory> UNiagaraNodeParameterMapBase::GetParameterMaps(UNiagaraScriptSourceBase* InSource, FString EmitterNameOverride, const TArray<FNiagaraVariable>& EncounterableVariables)
-{
-	TArray<FNiagaraParameterMapHistory> OutputParameterMapHistories;
-	UNiagaraScriptSource* Base = Cast<UNiagaraScriptSource>(InSource);
-	if (Base != nullptr)
-	{
-		OutputParameterMapHistories = GetParameterMaps(Base->NodeGraph, EmitterNameOverride, EncounterableVariables);
-	}
-	return OutputParameterMapHistories;
-
-}
-
-TArray<FNiagaraParameterMapHistory> UNiagaraNodeParameterMapBase::GetParameterMaps(const UNiagaraGraph* InGraph, FString EmitterNameOverride, const TArray<FNiagaraVariable>& EncounterableVariables)
+TArray<FNiagaraParameterMapHistory> UNiagaraNodeParameterMapBase::GetParameterMaps(const UNiagaraGraph* InGraph)
 {
 	TArray<UNiagaraNodeOutput*> OutputNodes;
 	InGraph->FindOutputNodes(OutputNodes);
@@ -60,39 +48,15 @@ TArray<FNiagaraParameterMapHistory> UNiagaraNodeParameterMapBase::GetParameterMa
 
 	for (UNiagaraNodeOutput* FoundOutputNode : OutputNodes)
 	{
-		OutputParameterMapHistories.Append(GetParameterMaps(FoundOutputNode, false, EmitterNameOverride,EncounterableVariables));
-	}
-
-	return OutputParameterMapHistories;
-}
+		FNiagaraParameterMapHistoryBuilder Builder;
+		Builder.BuildParameterMaps(FoundOutputNode);
 
-TArray<FNiagaraParameterMapHistory> UNiagaraNodeParameterMapBase::GetParameterMaps(const UNiagaraNodeOutput* InGraphEnd, bool bLimitToOutputScriptType, FString EmitterNameOverride, const TArray<FNiagaraVariable>& EncounterableVariables)
-{
-	const UEdGraphSchema_Niagara* Schema = Cast<UEdGraphSchema_Niagara>(InGraphEnd->GetSchema());
-
-	FNiagaraParameterMapHistoryBuilder Builder;
-	Builder.RegisterEncounterableVariables(EncounterableVariables);
-	if (!EmitterNameOverride.IsEmpty())
-	{
-		Builder.EnterEmitter(EmitterNameOverride, InGraphEnd->GetNiagaraGraph(), nullptr);
+		OutputParameterMapHistories.Append(Builder.Histories);
 	}
 
-	if (bLimitToOutputScriptType)
-	{
-		Builder.EnableScriptAllowList(true, InGraphEnd->GetUsage());
-	}
-	
-	Builder.BuildParameterMaps(InGraphEnd);
-	
-	if (!EmitterNameOverride.IsEmpty())
-	{
-		Builder.ExitEmitter(EmitterNameOverride, nullptr);
-	}
-	
-	return Builder.Histories;
+	return OutputParameterMapHistories;
 }
 
-
 bool UNiagaraNodeParameterMapBase::AllowNiagaraTypeForAddPin(const FNiagaraTypeDefinition& InType) const
 {
 	return InType != FNiagaraTypeDefinition::GetGenericNumericDef()
diff --git a/Engine/Plugins/FX/Niagara/Source/NiagaraEditor/Private/NiagaraNodeParameterMapBase.h b/Engine/Plugins/FX/Niagara/Source/NiagaraEditor/Private/NiagaraNodeParameterMapBase.h
--- a/Engine/Plugins/FX/Niagara/Source/NiagaraEditor/Private/NiagaraNodeParameterMapBase.h
+++ b/Engine/Plugins/FX/Niagara/Source/NiagaraEditor/Private/NiagaraNodeParameterMapBase.h
@@ -21,9 +21,7 @@ public:
 	/** Traverse the graph looking for the history of the parameter map specified by the input pin. This will return the list of variables discovered, any per-variable warnings (type mismatches, etc)
 		encountered per variable, and an array of pins encountered in order of traversal outward from the input pin.
 	*/
-	static TArray<FNiagaraParameterMapHistory> GetParameterMaps(const UNiagaraNodeOutput* InGraphEnd, bool bLimitToOutputScriptType = false, FString EmitterNameOverride = TEXT(""), const TArray<FNiagaraVariable>& EncounterableVariables = TArray<FNiagaraVariable>());
-	static TArray<FNiagaraParameterMapHistory> GetParameterMaps(const UNiagaraGraph* InGraph, FString EmitterNameOverride = TEXT(""), const TArray<FNiagaraVariable>& EncounterableVariables = TArray<FNiagaraVariable>());
-	static TArray<FNiagaraParameterMapHistory> GetParameterMaps(class UNiagaraScriptSourceBase* InSource, FString EmitterNameOverride = TEXT(""), const TArray<FNiagaraVariable>& EncounterableVariables = TArray<FNiagaraVariable>());
+	static TArray<FNiagaraParameterMapHistory> GetParameterMaps(const UNiagaraGraph* InGraph);
 
 	virtual bool AllowNiagaraTypeForAddPin(const FNiagaraTypeDefinition& InType) const override;
 	
diff --git a/Engine/Plugins/FX/Niagara/Source/NiagaraEditor/Private/NiagaraParameterMapHistory.cpp b/Engine/Plugins/FX/Niagara/Source/NiagaraEditor/Private/NiagaraParameterMapHistory.cpp
--- a/Engine/Plugins/FX/Niagara/Source/NiagaraEditor/Private/NiagaraParameterMapHistory.cpp
+++ b/Engine/Plugins/FX/Niagara/Source/NiagaraEditor/Private/NiagaraParameterMapHistory.cpp
@@ -17,9 +17,9 @@
 #include "NiagaraParameterCollection.h"
 #include "NiagaraConstants.h"
 #include "NiagaraNodeStaticSwitch.h"
-
 #include "NiagaraScriptVariable.h"
 #include "NiagaraNodeParameterMapSet.h"
+#include "UObject/UObjectThreadContext.h"
 
 #define LOCTEXT_NAMESPACE "NiagaraEditor"
 
@@ -674,6 +674,12 @@ FNiagaraVariable FNiagaraParameterMapHistory::ConvertVariableToRapidIterationCon
 
 UNiagaraParameterCollection* FNiagaraParameterMapHistory::IsParameterCollectionParameter(FNiagaraVariable& InVar, bool& bMissingParameter)
 {
+	if (bParameterCollectionsSkipped)
+	{
+		UE_LOG(LogNiagaraEditor, Error, TEXT("NiagaraParameterCollection was skipped during history building.  History which require NPC data can not be generated during PostLoad()."));
+		return nullptr;
+	}
+
 	bMissingParameter = false;
 	FString VarName = InVar.GetName().ToString();
 	for (int32 i = 0; i < ParameterCollections.Num(); ++i)
@@ -718,11 +724,16 @@ FNiagaraParameterMapHistoryBuilder::FNiagaraParameterMapHistoryBuilder()
 	bFilterByScriptAllowList = false;
 	bIgnoreDisabled = true;
 	FilterScriptType = ENiagaraScriptUsage::Function;
+
+	// depending on when this builder is created, we could be doing a PostLoad.  In this scenario
+	// we need to avoid triggering any loads for dependent NiagaraParameterCollection.  Mark the builder
+	// to skip those steps if necessary and we'll ensure that the generated histories don't require the
+	// collection information
+	bIncludeParameterCollectionInfo = !FUObjectThreadContext::Get().IsRoutingPostLoad;
 }
 
 void FNiagaraParameterMapHistoryBuilder::BuildParameterMaps(const UNiagaraNodeOutput* OutputNode, bool bRecursive)
 {
-	
 	TOptional<FName> StackContextAlias = OutputNode->GetStackContextOverride();
 	BeginUsage(OutputNode->GetUsage(), StackContextAlias.Get(ScriptUsageContextNameStack.Num() > 0 ? ScriptUsageContextNameStack.Top() : NAME_None));
 
@@ -773,12 +784,17 @@ ENiagaraScriptUsage FNiagaraParameterMapHistoryBuilder::GetBaseUsageContext()con
 int32 FNiagaraParameterMapHistoryBuilder::CreateParameterMap()
 {
 	int32 RetValue = Histories.AddDefaulted(1);
+	FNiagaraParameterMapHistory& History = Histories[RetValue];
+
 	FName StackName = ScriptUsageContextNameStack.Num() > 0 ? ScriptUsageContextNameStack.Top() : NAME_None;
 
 	if (StackName.IsValid())
 	{
-		Histories[RetValue].IterationNamespaceOverridesEncountered.AddUnique(StackName);
+		History.IterationNamespaceOverridesEncountered.AddUnique(StackName);
 	}
+
+	History.bParameterCollectionsSkipped = !bIncludeParameterCollectionInfo;
+
 	return RetValue;
 }
 
@@ -1269,13 +1285,16 @@ int32 FNiagaraParameterMapHistoryBuilder::HandleVariableRead(int32 ParamMapIdx,
 	Var = ResolveAliases(Var);
 
 	//Track any parameter collections we're referencing.
-	if (UNiagaraParameterCollection* Collection = Schema->VariableIsFromParameterCollection(Var))
+	if (bIncludeParameterCollectionInfo)
 	{
-		int32 Index = History.ParameterCollections.AddUnique(Collection);
-		History.ParameterCollectionNamespaces.SetNum(History.ParameterCollections.Num());
-		History.ParameterCollectionVariables.SetNum(History.ParameterCollections.Num());
-		History.ParameterCollectionNamespaces[Index] = Collection->GetFullNamespace();
-		History.ParameterCollectionVariables[Index] = Collection->GetParameters();
+		if (UNiagaraParameterCollection* Collection = Schema->VariableIsFromParameterCollection(Var))
+		{
+			int32 Index = History.ParameterCollections.AddUnique(Collection);
+			History.ParameterCollectionNamespaces.SetNum(History.ParameterCollections.Num());
+			History.ParameterCollectionVariables.SetNum(History.ParameterCollections.Num());
+			History.ParameterCollectionNamespaces[Index] = Collection->GetFullNamespace();
+			History.ParameterCollectionVariables[Index] = Collection->GetParameters();
+		}
 	}
 
 	const FString* ModuleAlias = GetModuleAlias();
@@ -1394,16 +1413,19 @@ int32 FNiagaraParameterMapHistoryBuilder::HandleExternalVariableRead(int32 Param
 	FNiagaraVariable OriginalUnaliasedVar = Var;
 
 	//Track any parameter collections we're referencing.
-	FNiagaraVariable FoundTempVar;
-	if (UNiagaraParameterCollection* Collection = Schema->VariableIsFromParameterCollection(Name.ToString(), true, FoundTempVar))
-	{
-		int32 Index = History.ParameterCollections.AddUnique(Collection);
-		History.ParameterCollectionNamespaces.SetNum(History.ParameterCollections.Num());
-		History.ParameterCollectionVariables.SetNum(History.ParameterCollections.Num());
-		History.ParameterCollectionNamespaces[Index] = Collection->GetFullNamespace();
-		History.ParameterCollectionVariables[Index] = Collection->GetParameters();
-		Var = FoundTempVar;
-		UpdateAliasedVariable(AliasedVar, OriginalUnaliasedVar, Var);
+	if (bIncludeParameterCollectionInfo)
+	{
+		FNiagaraVariable FoundTempVar;
+		if (UNiagaraParameterCollection* Collection = Schema->VariableIsFromParameterCollection(Name.ToString(), true, FoundTempVar))
+		{
+			int32 Index = History.ParameterCollections.AddUnique(Collection);
+			History.ParameterCollectionNamespaces.SetNum(History.ParameterCollections.Num());
+			History.ParameterCollectionVariables.SetNum(History.ParameterCollections.Num());
+			History.ParameterCollectionNamespaces[Index] = Collection->GetFullNamespace();
+			History.ParameterCollectionVariables[Index] = Collection->GetParameters();
+			Var = FoundTempVar;
+			UpdateAliasedVariable(AliasedVar, OriginalUnaliasedVar, Var);
+		}
 	}
 
 	int32 FoundIdx = Histories[ParamMapIdx].FindVariableByName(Name, true);
diff --git a/Engine/Plugins/FX/Niagara/Source/NiagaraEditor/Public/NiagaraEditorModule.h b/Engine/Plugins/FX/Niagara/Source/NiagaraEditor/Public/NiagaraEditorModule.h
--- a/Engine/Plugins/FX/Niagara/Source/NiagaraEditor/Public/NiagaraEditorModule.h
+++ b/Engine/Plugins/FX/Niagara/Source/NiagaraEditor/Public/NiagaraEditorModule.h
@@ -42,6 +42,7 @@ class FNiagaraDebugger;
 class UNiagaraParameterDefinitions;
 class UNiagaraReservedParametersManager;
 class FNiagaraGraphDataCache;
+class UNiagaraParameterCollection;
 
 DECLARE_STATS_GROUP(TEXT("Niagara Editor"), STATGROUP_NiagaraEditor, STATCAT_Advanced);
 
@@ -207,6 +208,8 @@ public:
 
 	FNiagaraGraphDataCache& GetGraphDataCache() const { return *GraphDataCache.Get(); }
 
+	NIAGARAEDITOR_API UNiagaraParameterCollection* FindCollectionForVariable(const FString& VariableName);
+
 private:
 	class FDeferredDestructionContainerBase
 	{
@@ -261,6 +264,8 @@ private:
 		StackIssueGenerators.Add(StructName) = Generator;
 	}
 
+	void RefreshParameterCollections(bool AllowLoading);
+
 private:
 	TSharedPtr<FExtensibilityManager> MenuExtensibilityManager;
 	TSharedPtr<FExtensibilityManager> ToolBarExtensibilityManager;
@@ -348,6 +353,8 @@ private:
 	TSet<FGuid> ReservedDefinitionIds;
 
 	TUniquePtr<FNiagaraGraphDataCache> GraphDataCache;
+
+	TArray<TWeakObjectPtr<UNiagaraParameterCollection>> CachedParameterCollections;
 };
 
 USTRUCT()
diff --git a/Engine/Plugins/FX/Niagara/Source/NiagaraEditor/Public/NiagaraParameterMapHistory.h b/Engine/Plugins/FX/Niagara/Source/NiagaraEditor/Public/NiagaraParameterMapHistory.h
--- a/Engine/Plugins/FX/Niagara/Source/NiagaraEditor/Public/NiagaraParameterMapHistory.h
+++ b/Engine/Plugins/FX/Niagara/Source/NiagaraEditor/Public/NiagaraParameterMapHistory.h
@@ -117,6 +117,9 @@ public:
 	/** List of additional DataSets to be written that were encountered during traversal. */
 	TArray<FNiagaraDataSetID> AdditionalDataSetWrites;
 
+	/** Flag indicating that the history building skipped processing parameter collections */
+	bool bParameterCollectionsSkipped = false;
+
 	bool IsVariableFromCustomIterationNamespaceOverride(const FNiagaraVariable& InVar) const;
 	
 	/**
@@ -476,6 +479,9 @@ protected:
 	/** Whether or not to ignore disabled nodes.*/
 	bool bIgnoreDisabled;
 
+	/** Whether we want to include ParameterCollection information */
+	bool bIncludeParameterCollectionInfo;
+
 	TArray<FNiagaraVariable> EncounterableExternalVariables;
 };
 
