- Add support for getting a bones parent index
- Add warning if filtered bones are not required by all LODs, this can mean that the effect may not play correctly if a lower LOD is used

diff --git a/Engine/Plugins/FX/Niagara/Shaders/Private/NiagaraDataInterfaceSkeletalMesh.ush b/Engine/Plugins/FX/Niagara/Shaders/Private/NiagaraDataInterfaceSkeletalMesh.ush
--- a/Engine/Plugins/FX/Niagara/Shaders/Private/NiagaraDataInterfaceSkeletalMesh.ush
+++ b/Engine/Plugins/FX/Niagara/Shaders/Private/NiagaraDataInterfaceSkeletalMesh.ush
@@ -902,6 +902,16 @@ void DISkelMesh_GetBoneCount(in FDISkelMeshContext DIContext, out int Count)
 	Count = DIContext.NumBones;
 }
 
+void DISkelMesh_GetParentBone(in FDISkelMeshContext DIContext, int BoneIndex, out int ParentIndex)
+{
+	ParentIndex = -1;
+	if ( BoneIndex >= 0 && BoneIndex < DIContext.NumBones )
+	{
+		ParentIndex = asint(DIContext.MeshCurrSamplingBonesBuffer[BoneIndex * 3 + 2].w);
+		
+	}
+}
+
 void DISkelMesh_GetFilteredBoneCount(in FDISkelMeshContext DIContext, out int Count)
 {
     Count = DIContext.NumFilteredBones;
diff --git a/Engine/Plugins/FX/Niagara/Source/Niagara/Classes/NiagaraDataInterfaceSkeletalMesh.h b/Engine/Plugins/FX/Niagara/Source/Niagara/Classes/NiagaraDataInterfaceSkeletalMesh.h
--- a/Engine/Plugins/FX/Niagara/Source/Niagara/Classes/NiagaraDataInterfaceSkeletalMesh.h
+++ b/Engine/Plugins/FX/Niagara/Source/Niagara/Classes/NiagaraDataInterfaceSkeletalMesh.h
@@ -992,6 +992,7 @@ public:
 	void IsValidBone(FVectorVMExternalFunctionContext& Context);
 	void RandomBone(FVectorVMExternalFunctionContext& Context);
 	void GetBoneCount(FVectorVMExternalFunctionContext& Context);
+	void GetParentBone(FVectorVMExternalFunctionContext& Context);
 
 	void GetFilteredBoneCount(FVectorVMExternalFunctionContext& Context);
 	void GetFilteredBoneAt(FVectorVMExternalFunctionContext& Context);
@@ -1046,6 +1047,7 @@ public:
 	static const FName IsValidBoneName;
 	static const FName RandomBoneName;
 	static const FName GetBoneCountName;
+	static const FName GetParentBoneName;
 
 	static const FName RandomFilteredBoneName;
 	static const FName GetFilteredBoneCountName;
diff --git a/Engine/Plugins/FX/Niagara/Source/Niagara/Private/NDISkeletalMesh_BoneSampling.cpp b/Engine/Plugins/FX/Niagara/Source/Niagara/Private/NDISkeletalMesh_BoneSampling.cpp
--- a/Engine/Plugins/FX/Niagara/Source/Niagara/Private/NDISkeletalMesh_BoneSampling.cpp
+++ b/Engine/Plugins/FX/Niagara/Source/Niagara/Private/NDISkeletalMesh_BoneSampling.cpp
@@ -28,6 +28,7 @@ const FName FSkeletalMeshInterfaceHelper::GetSkinnedBoneDataWSInterpolatedName("
 const FName FSkeletalMeshInterfaceHelper::IsValidBoneName("IsValidBone");
 const FName FSkeletalMeshInterfaceHelper::RandomBoneName("RandomBone");
 const FName FSkeletalMeshInterfaceHelper::GetBoneCountName("GetBoneCount");
+const FName FSkeletalMeshInterfaceHelper::GetParentBoneName("GetParentBone");
 
 const FName FSkeletalMeshInterfaceHelper::RandomFilteredBoneName("RandomFilteredBone");
 const FName FSkeletalMeshInterfaceHelper::GetFilteredBoneCountName("GetFilteredBoneCount");
@@ -153,6 +154,19 @@ void UNiagaraDataInterfaceSkeletalMesh::GetSkeletonSamplingFunctions(TArray<FNia
 #endif
 		OutFunctions.Add(Sig);
 	}
+	{
+		FNiagaraFunctionSignature Sig;
+		Sig.Name = FSkeletalMeshInterfaceHelper::GetParentBoneName;
+		Sig.Inputs.Add(FNiagaraVariable(FNiagaraTypeDefinition(GetClass()), TEXT("SkeletalMesh")));
+		Sig.Inputs.Add(FNiagaraVariable(FNiagaraTypeDefinition::GetIntDef(), TEXT("BoneIndex")));
+		Sig.Outputs.Add(FNiagaraVariable(FNiagaraTypeDefinition::GetIntDef(), TEXT("ParentIndex")));
+		Sig.bMemberFunction = true;
+		Sig.bRequiresContext = false;
+#if WITH_EDITORONLY_DATA
+		Sig.Description = LOCTEXT("GetParentBoneDesc", "Returns the parent bone index for the given bone, -1 if no parent exists.");
+#endif
+		OutFunctions.Add(Sig);
+	}
 
 	{
 		FNiagaraFunctionSignature Sig;
@@ -373,6 +387,11 @@ void UNiagaraDataInterfaceSkeletalMesh::BindSkeletonSamplingFunction(const FVMEx
 		check(BindingInfo.GetNumInputs() == 1 && BindingInfo.GetNumOutputs() == 1);
 		OutFunc = FVMExternalFunction::CreateLambda([this](FVectorVMExternalFunctionContext& Context) { this->GetBoneCount(Context); });
 	}
+	else if (BindingInfo.Name == FSkeletalMeshInterfaceHelper::GetParentBoneName)
+	{
+		check(BindingInfo.GetNumInputs() == 2 && BindingInfo.GetNumOutputs() == 1);
+		OutFunc = FVMExternalFunction::CreateLambda([this](FVectorVMExternalFunctionContext& Context) { this->GetParentBone(Context); });
+	}
 	else if (BindingInfo.Name == FSkeletalMeshInterfaceHelper::RandomFilteredBoneName)
 	{
 		check(BindingInfo.GetNumInputs() == 4 && BindingInfo.GetNumOutputs() == 1);
@@ -699,6 +718,37 @@ void UNiagaraDataInterfaceSkeletalMesh::GetBoneCount(FVectorVMExternalFunctionCo
 	}
 }
 
+void UNiagaraDataInterfaceSkeletalMesh::GetParentBone(FVectorVMExternalFunctionContext& Context)
+{
+	SCOPE_CYCLE_COUNTER(STAT_NiagaraSkel_Bone_Sample);
+
+	VectorVM::FUserPtrHandler<FNDISkeletalMesh_InstanceData> InstData(Context);
+	FNDIInputParam<int32> InBoneIndex(Context);
+	FNDIOutputParam<int32> OutParentIndex(Context);
+
+	FSkeletalMeshAccessorHelper MeshAccessor;
+	MeshAccessor.Init<TIntegralConstant<int32, 0>, TIntegralConstant<int32, 0>>(InstData);
+
+	if (MeshAccessor.AreBonesAccessible())
+	{
+		const FReferenceSkeleton& RefSkeleton = MeshAccessor.Mesh->GetRefSkeleton();
+		const int32 NumBones = RefSkeleton.GetNum();
+		for (int32 i = 0; i < Context.GetNumInstances(); ++i)
+		{
+			const int32 BoneIndex = InBoneIndex.GetAndAdvance();
+			const int32 ParentIndex = BoneIndex >= 0 && BoneIndex < NumBones ? RefSkeleton.GetParentIndex(BoneIndex) : INDEX_NONE;
+			OutParentIndex.SetAndAdvance(ParentIndex);
+		}
+	}
+	else
+	{
+		for (int32 i = 0; i < Context.GetNumInstances(); ++i)
+		{
+			OutParentIndex.SetAndAdvance(INDEX_NONE);
+		}
+	}
+}
+
 struct FBoneSocketSkinnedDataOutputHandler
 {
 	FBoneSocketSkinnedDataOutputHandler(FVectorVMExternalFunctionContext& Context)
diff --git a/Engine/Plugins/FX/Niagara/Source/Niagara/Private/NiagaraDataInterfaceSkeletalMesh.cpp b/Engine/Plugins/FX/Niagara/Source/Niagara/Private/NiagaraDataInterfaceSkeletalMesh.cpp
--- a/Engine/Plugins/FX/Niagara/Source/Niagara/Private/NiagaraDataInterfaceSkeletalMesh.cpp
+++ b/Engine/Plugins/FX/Niagara/Source/Niagara/Private/NiagaraDataInterfaceSkeletalMesh.cpp
@@ -1068,7 +1068,7 @@ void FSkeletalMeshGpuDynamicBufferProxy::NewFrame(const FNDISkeletalMesh_Instanc
 	TArray<FVector4f> BoneSamplingData;
 
 	auto FillBuffers =
-		[&](const TArray<FTransform>& BoneTransforms)
+		[&](const TArray<FTransform>& BoneTransforms, const FReferenceSkeleton* ReferenceSkeleton)
 		{
 			check(BoneTransforms.Num() == SamplingBoneCount);
 
@@ -1105,21 +1105,24 @@ void FSkeletalMeshGpuDynamicBufferProxy::NewFrame(const FNDISkeletalMesh_Instanc
 
 			// Fill BoneSamplingData
 			BoneSamplingData.Reserve((SamplingBoneCount + SamplingSocketCount) * 2);
-			for (const FTransform& BoneTransform : BoneTransforms)
+			for (int i=0; i < BoneTransforms.Num(); ++i )
 			{
+				const FTransform& BoneTransform = BoneTransforms[i];
 				const FQuat Rotation = BoneTransform.GetRotation();
-				BoneSamplingData.Add(BoneTransform.GetLocation());
-				BoneSamplingData.Add(FVector4f(Rotation.X, Rotation.Y, Rotation.Z, Rotation.W));
-				BoneSamplingData.Add(BoneTransform.GetScale3D());
+				const int32 ParentIndex = ReferenceSkeleton ? ReferenceSkeleton->GetParentIndex(i) : -1;
+				BoneSamplingData.Emplace(BoneTransform.GetLocation());
+				BoneSamplingData.Emplace(Rotation.X, Rotation.Y, Rotation.Z, Rotation.W);
+				BoneSamplingData.Emplace(BoneTransform.GetScale3D(), reinterpret_cast<const float&>(ParentIndex));
 			}
 
 			// Append sockets
 			for (const FTransform3f& SocketTransform : InstanceData->GetFilteredSocketsCurrBuffer())
 			{
 				const FQuat4f Rotation = SocketTransform.GetRotation();
-				BoneSamplingData.Add(SocketTransform.GetLocation());
-				BoneSamplingData.Add(FVector4f(Rotation.X, Rotation.Y, Rotation.Z, Rotation.W));
-				BoneSamplingData.Add(SocketTransform.GetScale3D());
+				const int32 ParentIndex = -1;
+				BoneSamplingData.Emplace(SocketTransform.GetLocation());
+				BoneSamplingData.Emplace(Rotation.X, Rotation.Y, Rotation.Z, Rotation.W);
+				BoneSamplingData.Emplace(SocketTransform.GetScale3D(), reinterpret_cast<const float&>(ParentIndex));
 			}
 		};
 
@@ -1128,6 +1131,7 @@ void FSkeletalMeshGpuDynamicBufferProxy::NewFrame(const FNDISkeletalMesh_Instanc
 	{
 		if (USkinnedMeshComponent* MasterComponent = SkelComp->MasterPoseComponent.Get())
 		{
+			const FReferenceSkeleton* ReferenceSkeleton = nullptr;
 			const TArray<int32>& MasterBoneMap = SkelComp->GetMasterBoneMap();
 			const int32 NumBones = MasterBoneMap.Num();
 
@@ -1141,6 +1145,7 @@ void FSkeletalMeshGpuDynamicBufferProxy::NewFrame(const FNDISkeletalMesh_Instanc
 			}
 			else
 			{
+				ReferenceSkeleton = &SkelMesh->GetRefSkeleton();
 				const TArray<FTransform>& MasterTransforms = MasterComponent->GetComponentSpaceTransforms();
 				for (int32 BoneIndex=0; BoneIndex < NumBones; ++BoneIndex)
 				{
@@ -1154,8 +1159,8 @@ void FSkeletalMeshGpuDynamicBufferProxy::NewFrame(const FNDISkeletalMesh_Instanc
 						}
 					}
 
-					const int32 ParentIndex = SkelMesh->GetRefSkeleton().GetParentIndex(BoneIndex);
-					FTransform BoneTransform =SkelMesh->GetRefSkeleton().GetRefBonePose()[BoneIndex];
+					const int32 ParentIndex  = ReferenceSkeleton->GetParentIndex(BoneIndex);
+					FTransform BoneTransform = ReferenceSkeleton->GetRefBonePose()[BoneIndex];
 					if (TempBoneTransforms.IsValidIndex(ParentIndex))
 					{
 						BoneTransform = BoneTransform * TempBoneTransforms[ParentIndex];
@@ -1163,14 +1168,14 @@ void FSkeletalMeshGpuDynamicBufferProxy::NewFrame(const FNDISkeletalMesh_Instanc
 					TempBoneTransforms.Add(BoneTransform);
 				}
 			}
-			FillBuffers(TempBoneTransforms);
+			FillBuffers(TempBoneTransforms, ReferenceSkeleton);
 		}
 		else
 		{
 			const TArray<FTransform>& ComponentTransforms = SkelComp->GetComponentSpaceTransforms();
 			if (ComponentTransforms.Num() > 0)
 			{
-				FillBuffers(ComponentTransforms);
+				FillBuffers(ComponentTransforms, &SkelMesh->GetRefSkeleton());
 			}
 			else
 			{
@@ -1180,13 +1185,14 @@ void FSkeletalMeshGpuDynamicBufferProxy::NewFrame(const FNDISkeletalMesh_Instanc
 
 				TArray<FTransform> TempBoneTransforms;
 				TempBoneTransforms.AddDefaulted(SamplingBoneCount);
-				FillBuffers(TempBoneTransforms);
+				FillBuffers(TempBoneTransforms, nullptr);
 			}
 		}
 	}
 	else
 	{
 		//-TODO: Opt and combine with MaterPoseComponent
+		const FReferenceSkeleton* ReferenceSkeleton = &SkelMesh->GetRefSkeleton();
 		TArray<FTransform> TempBoneTransforms;
 		TempBoneTransforms.Reserve(SamplingBoneCount);
 
@@ -1194,7 +1200,7 @@ void FSkeletalMeshGpuDynamicBufferProxy::NewFrame(const FNDISkeletalMesh_Instanc
 		for (int32 i=0; i < RefTransforms.Num(); ++i)
 		{
 			FTransform BoneTransform = RefTransforms[i];
-			const int32 ParentIndex = SkelMesh->GetRefSkeleton().GetParentIndex(i);
+			const int32 ParentIndex = ReferenceSkeleton->GetParentIndex(i);
 			if (TempBoneTransforms.IsValidIndex(ParentIndex))
 			{
 				BoneTransform = BoneTransform * TempBoneTransforms[ParentIndex];
@@ -1202,7 +1208,7 @@ void FSkeletalMeshGpuDynamicBufferProxy::NewFrame(const FNDISkeletalMesh_Instanc
 			TempBoneTransforms.Add(BoneTransform);
 		}
 
-		FillBuffers(TempBoneTransforms);
+		FillBuffers(TempBoneTransforms, ReferenceSkeleton);
 	}
 
 	FSkeletalMeshGpuDynamicBufferProxy* ThisProxy = this;
@@ -2670,11 +2676,10 @@ void UNiagaraDataInterfaceSkeletalMesh::GetFeedback(UNiagaraSystem* Asset, UNiag
 
 	bool bHasCPUAccessWarning = false;
 	bool bHasNoMeshAssignedWarning = false;
+	USkeletalMesh* SkelMesh = GetSkeletalMesh(Component);
 
 	// Collect Errors
 #if WITH_EDITORONLY_DATA
-	USkeletalMesh* SkelMesh = GetSkeletalMesh(Component);
-
 	if (SkelMesh != nullptr)
 	{
 		bool bHasCPUAccess = true;
@@ -2769,6 +2774,52 @@ void UNiagaraDataInterfaceSkeletalMesh::GetFeedback(UNiagaraSystem* Asset, UNiag
 
 		OutWarnings.Add(NoMeshAssignedError);
 	}
+
+	// Look for bones being used that are LOD'ed out
+	if ( (SkelMesh != nullptr) && (SkelMesh->GetResourceForRendering() != nullptr) )
+	{
+		FSkeletalMeshRenderData* SkelResource = SkelMesh->GetResourceForRendering();
+		auto IsBoneRequiredInAllLODs =
+			[&](const FName BoneName)
+			{
+				const int32 BoneIndex = SkelMesh->GetRefSkeleton().FindBoneIndex(BoneName);
+				if (BoneIndex == INDEX_NONE)
+				{
+					return false;
+				}
+
+				for ( const FSkeletalMeshLODRenderData& LODData : SkelResource->LODRenderData )
+				{
+					if ( !LODData.RequiredBones.Contains(BoneIndex) )
+					{
+						return false;
+					}
+				}
+				return true;
+			};
+
+		if (FilteredBones.Num() > 0)
+		{
+			FString MissingBoneList;
+			for (FName Bone : FilteredBones)
+			{
+				if ( !IsBoneRequiredInAllLODs(Bone) )
+				{
+					MissingBoneList.Append(TEXT("\n"));
+					Bone.AppendString(MissingBoneList);
+				}
+			}
+
+			if (MissingBoneList.Len() > 0 )
+			{
+				OutWarnings.Emplace(
+					FText::Format(LOCTEXT("BonesLODOutError", "Filtered Bones may not animate in all LODs, this can lead to incorrect results when animating at those LOD levels.\n{0}"), FText::FromString(MissingBoneList)),
+					LOCTEXT("BonesLODOutErrorSummary", "Filtered bones may not animate in all LODs."),
+					FNiagaraDataInterfaceFix()
+				);
+			}
+		}
+	}
 }
 
 //Deprecated functions we check for and advise on updates in ValidateFunction
@@ -3151,6 +3202,11 @@ bool UNiagaraDataInterfaceSkeletalMesh::GetFunctionHLSL(const FNiagaraDataInterf
 		static const TCHAR* FormatSample = TEXT("void {InstanceFunctionName} (out int Count) { {GetDISkelMeshContextName} DISkelMesh_GetBoneCount(DIContext, Count); }");
 		OutHLSL += FString::Format(FormatSample, ArgsSample);
 	}
+	else if (FunctionInfo.DefinitionName == FSkeletalMeshInterfaceHelper::GetParentBoneName)
+	{
+		static const TCHAR* FormatSample = TEXT("void {InstanceFunctionName} (int BoneIndex, out int ParentIndex) { {GetDISkelMeshContextName} DISkelMesh_GetParentBone(DIContext, BoneIndex, ParentIndex); }");
+		OutHLSL += FString::Format(FormatSample, ArgsSample);
+	}
 	else if (FunctionInfo.DefinitionName == FSkeletalMeshInterfaceHelper::GetFilteredBoneCountName)
 	{
 		static const TCHAR* FormatSample = TEXT("void {InstanceFunctionName} (out int Count) { {GetDISkelMeshContextName} DISkelMesh_GetFilteredBoneCount(DIContext, Count); }");
