Fixed a compilation problem when using scene depth in non-lwc builds

diff --git a/Engine/Plugins/FX/Niagara/Shaders/Private/NiagaraEmitterInstanceShader.usf b/Engine/Plugins/FX/Niagara/Shaders/Private/NiagaraEmitterInstanceShader.usf
--- a/Engine/Plugins/FX/Niagara/Shaders/Private/NiagaraEmitterInstanceShader.usf
+++ b/Engine/Plugins/FX/Niagara/Shaders/Private/NiagaraEmitterInstanceShader.usf
@@ -96,8 +96,8 @@ bool NiagaraAny(bool4 b) { return b.x || b.y || b.z || b.w; }
 	FLWCVector3 WorldPositionFromSceneDepth(float2 ScreenPosition, float SceneDepth)
 	{
 		FLWCVector4 HomogeneousWorldPosition = LWCMultiply(float4(ScreenPosition * SceneDepth, SceneDepth, 1), PrimaryView.ScreenToWorld);
-		FLWCVector3 WorldPosition = MakeLWCVector3(HomogeneousWorldPosition.Tile.xyz, HomogeneousWorldPosition.Offset.xyz);
-		FLWCScalar Scale = MakeLWCScalar(HomogeneousWorldPosition.Tile.w, HomogeneousWorldPosition.Offset.w);
+		FLWCVector3 WorldPosition = MakeLWCVector3(LWCGetTile(HomogeneousWorldPosition).xyz, HomogeneousWorldPosition.Offset.xyz);
+		FLWCScalar Scale = MakeLWCScalar(LWCGetTile(HomogeneousWorldPosition).w, HomogeneousWorldPosition.Offset.w);
 		return LWCDivide(WorldPosition, LWCToFloat(Scale));
 	}
 
