Fixed a translator problem where old assets could end up in a state with both position types and vector types for the same parameter.
The scripts are still valid in that case, so we should not error out, just emit a warning.

diff --git a/Engine/Plugins/FX/Niagara/Source/NiagaraCore/Private/NiagaraCustomVersion.cpp b/Engine/Plugins/FX/Niagara/Source/NiagaraCore/Private/NiagaraCustomVersion.cpp
--- a/Engine/Plugins/FX/Niagara/Source/NiagaraCore/Private/NiagaraCustomVersion.cpp
+++ b/Engine/Plugins/FX/Niagara/Source/NiagaraCore/Private/NiagaraCustomVersion.cpp
@@ -8,4 +8,4 @@ const FGuid FNiagaraCustomVersion::GUID(0xFCF57AFA, 0x50764283, 0xB9A9E658, 0xFF
 // Register the custom version with core
 FCustomVersionRegistration GRegisterNiagaraCustomVersion(FNiagaraCustomVersion::GUID, FNiagaraCustomVersion::LatestVersion, TEXT("NiagaraVer"));
 
-const FGuid FNiagaraCustomVersion::LatestScriptCompileVersion(0xD8BFE375, 0xCAF76740, 0x9E0FD451, 0xA3B5984D);
+const FGuid FNiagaraCustomVersion::LatestScriptCompileVersion(0xD8BFE375, 0xCAF76740, 0x9E0FD451, 0xA3B5984E);
diff --git a/Engine/Plugins/FX/Niagara/Source/NiagaraEditor/Private/NiagaraHlslTranslator.cpp b/Engine/Plugins/FX/Niagara/Source/NiagaraEditor/Private/NiagaraHlslTranslator.cpp
--- a/Engine/Plugins/FX/Niagara/Source/NiagaraEditor/Private/NiagaraHlslTranslator.cpp
+++ b/Engine/Plugins/FX/Niagara/Source/NiagaraEditor/Private/NiagaraHlslTranslator.cpp
@@ -5723,10 +5723,18 @@ void FHlslNiagaraTranslator::HandleParameterRead(int32 ParamMapHistoryIdx, const
 			FNiagaraVariable* ExistingVar = ExistingVarAndDefaultSource ? &ExistingVarAndDefaultSource->Variable : nullptr;
 
 			bool ExistsInAttribArrayAlready = ExistingVar != nullptr;
-			if (ExistsInAttribArrayAlready&& ExistingVar->GetType() != Var.GetType())
+			if (ExistsInAttribArrayAlready && ExistingVar->GetType() != Var.GetType())
 			{
-				Error(FText::Format(LOCTEXT("Mismatched Types", "Variable {0} was defined earlier, but its type is different! {1} != {2}"), FText::FromName(Var.GetName()),
-					ExistingVar->GetType().GetNameText(), Var.GetType().GetNameText()), ErrorNode, nullptr);
+				if ((ExistingVar->GetType() == FNiagaraTypeDefinition::GetVec3Def() && Var.GetType() == FNiagaraTypeDefinition::GetPositionDef())
+					|| (ExistingVar->GetType() == FNiagaraTypeDefinition::GetPositionDef() && Var.GetType() == FNiagaraTypeDefinition::GetVec3Def()))
+				{
+					Warning(FText::Format(LOCTEXT("Mismatched Types", "Variable {0} was defined both as position and vector, please check your modules and linked values for compatibility."), FText::FromName(Var.GetName())), ErrorNode, nullptr);
+				}
+				else
+				{
+					Error(FText::Format(LOCTEXT("Mismatched Types", "Variable {0} was defined earlier, but its type is different! {1} != {2}"), FText::FromName(Var.GetName()),
+						ExistingVar->GetType().GetNameText(), Var.GetType().GetNameText()), ErrorNode, nullptr);
+				}
 			}
 
 
