Localization key fix

diff --git a/Engine/Plugins/FX/Niagara/Source/NiagaraEditor/Private/NiagaraHlslTranslator.cpp b/Engine/Plugins/FX/Niagara/Source/NiagaraEditor/Private/NiagaraHlslTranslator.cpp
--- a/Engine/Plugins/FX/Niagara/Source/NiagaraEditor/Private/NiagaraHlslTranslator.cpp
+++ b/Engine/Plugins/FX/Niagara/Source/NiagaraEditor/Private/NiagaraHlslTranslator.cpp
@@ -5728,7 +5728,7 @@ void FHlslNiagaraTranslator::HandleParameterRead(int32 ParamMapHistoryIdx, const
 				if ((ExistingVar->GetType() == FNiagaraTypeDefinition::GetVec3Def() && Var.GetType() == FNiagaraTypeDefinition::GetPositionDef())
 					|| (ExistingVar->GetType() == FNiagaraTypeDefinition::GetPositionDef() && Var.GetType() == FNiagaraTypeDefinition::GetVec3Def()))
 				{
-					Warning(FText::Format(LOCTEXT("Mismatched Types", "Variable {0} was defined both as position and vector, please check your modules and linked values for compatibility."), FText::FromName(Var.GetName())), ErrorNode, nullptr);
+					Warning(FText::Format(LOCTEXT("MismatchedPositionTypes", "Variable {0} was defined both as position and vector, please check your modules and linked values for compatibility."), FText::FromName(Var.GetName())), ErrorNode, nullptr);
 				}
 				else
 				{
