- Add back code to inject SkipSet on SetArrayElem which was accidentally removed

#rnx

diff --git a/Engine/Plugins/FX/Niagara/Source/Niagara/Private/NiagaraDataInterfaceArrayImpl.cpp b/Engine/Plugins/FX/Niagara/Source/Niagara/Private/NiagaraDataInterfaceArrayImpl.cpp
--- a/Engine/Plugins/FX/Niagara/Source/Niagara/Private/NiagaraDataInterfaceArrayImpl.cpp
+++ b/Engine/Plugins/FX/Niagara/Source/Niagara/Private/NiagaraDataInterfaceArrayImpl.cpp
@@ -56,6 +56,11 @@ bool FNiagaraDataInterfaceArrayImplHelper::UpgradeFunctionCall(FNiagaraFunctionS
 		}
 
 		FunctionSignature.bExperimental = false;
+
+		if (FunctionSignature.Name == FNiagaraDataInterfaceArrayImplHelper::Function_SetArrayElemName)
+		{
+			FunctionSignature.Inputs.EmplaceAt(1, FNiagaraTypeDefinition::GetBoolDef(), TEXT("SkipSet"));
+		}
 	}
 
 	FunctionSignature.FunctionVersion = FFunctionVersion::LatestVersion;
