- Add UserParameter binding to all texture samples rather than using BP method to directly poke data into the data interface
- Support setting render targests for each type
- Update test with new functionality and to cover all permutations

#rnx

diff --git a/Engine/Plugins/FX/Niagara/Source/Niagara/Classes/NiagaraDataInterface2DArrayTexture.h b/Engine/Plugins/FX/Niagara/Source/Niagara/Classes/NiagaraDataInterface2DArrayTexture.h
--- a/Engine/Plugins/FX/Niagara/Source/Niagara/Classes/NiagaraDataInterface2DArrayTexture.h
+++ b/Engine/Plugins/FX/Niagara/Source/Niagara/Classes/NiagaraDataInterface2DArrayTexture.h
@@ -21,12 +21,11 @@ public:
 	UPROPERTY(EditAnywhere, Category = "Texture")
 	TObjectPtr<UTexture2DArray> Texture;
 
+	UPROPERTY(EditAnywhere, Category = "Texture", meta = (ToolTip = "When valid the user parameter is used as the texture rather than the one on the data interface"))
+	FNiagaraUserParameterBinding TextureUserParameter;
+
 	//UObject Interface
 	virtual void PostInitProperties()override;
-	virtual void PostLoad() override;
-#if WITH_EDITOR
-	virtual void PostEditChangeProperty(struct FPropertyChangedEvent& PropertyChangedEvent) override;
-#endif
 	//UObject Interface End
 
 	//UNiagaraDataInterface Interface
@@ -34,6 +33,10 @@ public:
 	virtual void GetVMExternalFunction(const FVMExternalFunctionBindingInfo& BindingInfo, void* InstanceData, FVMExternalFunction &OutFunc) override;
 	virtual bool CanExecuteOnTarget(ENiagaraSimTarget Target)const override { return Target == ENiagaraSimTarget::GPUComputeSim; }
 
+	virtual int32 PerInstanceDataSize() const override;
+	virtual bool InitPerInstanceData(void* PerInstanceData, FNiagaraSystemInstance* SystemInstance) override;
+	virtual void DestroyPerInstanceData(void* PerInstanceData, FNiagaraSystemInstance* SystemInstance) override;
+
 	virtual bool HasPreSimulateTick() const override { return true; }
 	virtual bool PerInstanceTick(void* PerInstanceData, FNiagaraSystemInstance* SystemInstance, float DeltaSeconds) override;
 	//UNiagaraDataInterface Interface End
@@ -58,24 +61,7 @@ public:
 protected:
 	virtual bool CopyToInternal(UNiagaraDataInterface* Destination) const override;
 
-	virtual void PushToRenderThreadImpl() override;
-
 protected:
-	FIntVector TextureSize = FIntVector::ZeroValue;
-
 	static const FName SampleTextureName;
 	static const FName TextureDimsName;
 };
-
-struct FNiagaraDataInterfaceProxy2DArrayTexture : public FNiagaraDataInterfaceProxy
-{
-	FSamplerStateRHIRef SamplerStateRHI;
-	FTextureRHIRef TextureRHI;
-	FVector TexDims;
-
-	virtual void ConsumePerInstanceDataFromGameThread(void* PerInstanceData, const FNiagaraSystemInstanceID& Instance) override { check(false); }
-	virtual int32 PerInstanceDataPassedToRenderThreadSize() const override
-	{
-		return 0;
-	}
-};
diff --git a/Engine/Plugins/FX/Niagara/Source/Niagara/Classes/NiagaraDataInterfaceCubeTexture.h b/Engine/Plugins/FX/Niagara/Source/Niagara/Classes/NiagaraDataInterfaceCubeTexture.h
--- a/Engine/Plugins/FX/Niagara/Source/Niagara/Classes/NiagaraDataInterfaceCubeTexture.h
+++ b/Engine/Plugins/FX/Niagara/Source/Niagara/Classes/NiagaraDataInterfaceCubeTexture.h
@@ -21,12 +21,11 @@ public:
 	UPROPERTY(EditAnywhere, Category = "Texture")
 	TObjectPtr<UTextureCube> Texture;
 
+	UPROPERTY(EditAnywhere, Category = "Texture", meta = (ToolTip = "When valid the user parameter is used as the texture rather than the one on the data interface"))
+	FNiagaraUserParameterBinding TextureUserParameter;
+
 	//UObject Interface
 	virtual void PostInitProperties()override;
-	virtual void PostLoad() override;
-#if WITH_EDITOR
-	virtual void PostEditChangeProperty(struct FPropertyChangedEvent& PropertyChangedEvent) override;
-#endif
 	//UObject Interface End
 
 	//UNiagaraDataInterface Interface
@@ -34,6 +33,10 @@ public:
 	virtual void GetVMExternalFunction(const FVMExternalFunctionBindingInfo& BindingInfo, void* InstanceData, FVMExternalFunction &OutFunc) override;
 	virtual bool CanExecuteOnTarget(ENiagaraSimTarget Target)const override { return Target == ENiagaraSimTarget::GPUComputeSim; }
 
+	virtual int32 PerInstanceDataSize() const override;
+	virtual bool InitPerInstanceData(void* PerInstanceData, FNiagaraSystemInstance* SystemInstance) override;
+	virtual void DestroyPerInstanceData(void* PerInstanceData, FNiagaraSystemInstance* SystemInstance) override;
+
 	virtual bool HasPreSimulateTick() const override { return true; }
 	virtual bool PerInstanceTick(void* PerInstanceData, FNiagaraSystemInstance* SystemInstance, float DeltaSeconds) override;
 	//UNiagaraDataInterface Interface End
@@ -59,11 +62,7 @@ public:
 protected:
 	virtual bool CopyToInternal(UNiagaraDataInterface* Destination) const override;
 
-	virtual void PushToRenderThreadImpl() override;
-
 protected:
-	FIntPoint TextureSize = FIntPoint::ZeroValue;
-
 	static const FName SampleCubeTextureName;
 	static const FName TextureDimsName;
 };
diff --git a/Engine/Plugins/FX/Niagara/Source/Niagara/Classes/NiagaraDataInterfaceTexture.h b/Engine/Plugins/FX/Niagara/Source/Niagara/Classes/NiagaraDataInterfaceTexture.h
--- a/Engine/Plugins/FX/Niagara/Source/Niagara/Classes/NiagaraDataInterfaceTexture.h
+++ b/Engine/Plugins/FX/Niagara/Source/Niagara/Classes/NiagaraDataInterfaceTexture.h
@@ -19,12 +19,12 @@ public:
 	UPROPERTY(EditAnywhere, Category = "Texture")
 	TObjectPtr<UTexture> Texture;
 
+	UPROPERTY(EditAnywhere, Category = "Texture", meta = (ToolTip = "When valid the user parameter is used as the texture rather than the one on the data interface"))
+	FNiagaraUserParameterBinding TextureUserParameter;
+
 	//UObject Interface
 	virtual void PostInitProperties()override;
 	virtual void PostLoad() override;
-#if WITH_EDITOR
-	virtual void PostEditChangeProperty(struct FPropertyChangedEvent& PropertyChangedEvent) override;
-#endif
 	virtual void Serialize(FArchive& Ar) override;
 	//UObject Interface End
 
@@ -33,6 +33,10 @@ public:
 	virtual void GetVMExternalFunction(const FVMExternalFunctionBindingInfo& BindingInfo, void* InstanceData, FVMExternalFunction &OutFunc) override;
 	virtual bool CanExecuteOnTarget(ENiagaraSimTarget Target)const override { return Target==ENiagaraSimTarget::GPUComputeSim; }
 
+	virtual int32 PerInstanceDataSize() const override;
+	virtual bool InitPerInstanceData(void* PerInstanceData, FNiagaraSystemInstance* SystemInstance) override;
+	virtual void DestroyPerInstanceData(void* PerInstanceData, FNiagaraSystemInstance* SystemInstance) override;
+
 	virtual bool HasPreSimulateTick() const override { return true; }
 	virtual bool PerInstanceTick(void* PerInstanceData, FNiagaraSystemInstance* SystemInstance, float DeltaSeconds) override;
 	//UNiagaraDataInterface Interface End
@@ -59,11 +63,7 @@ public:
 protected:
 	virtual bool CopyToInternal(UNiagaraDataInterface* Destination) const override;
 
-	virtual void PushToRenderThreadImpl() override;
-
 protected:
-	FIntPoint TextureSize = FIntPoint::ZeroValue;
-
 	static const FName SampleTexture2DName;
 	static const FName SampleVolumeTextureName;
 	static const FName SamplePseudoVolumeTextureName;
diff --git a/Engine/Plugins/FX/Niagara/Source/Niagara/Classes/NiagaraDataInterfaceVolumeTexture.h b/Engine/Plugins/FX/Niagara/Source/Niagara/Classes/NiagaraDataInterfaceVolumeTexture.h
--- a/Engine/Plugins/FX/Niagara/Source/Niagara/Classes/NiagaraDataInterfaceVolumeTexture.h
+++ b/Engine/Plugins/FX/Niagara/Source/Niagara/Classes/NiagaraDataInterfaceVolumeTexture.h
@@ -21,12 +21,11 @@ public:
 	UPROPERTY(EditAnywhere, Category = "Texture")
 	TObjectPtr<UVolumeTexture> Texture;
 
+	UPROPERTY(EditAnywhere, Category = "Texture", meta = (ToolTip = "When valid the user parameter is used as the texture rather than the one on the data interface"))
+	FNiagaraUserParameterBinding TextureUserParameter;
+
 	//UObject Interface
 	virtual void PostInitProperties()override;
-	virtual void PostLoad() override;
-#if WITH_EDITOR
-	virtual void PostEditChangeProperty(struct FPropertyChangedEvent& PropertyChangedEvent) override;
-#endif
 	//UObject Interface End
 
 	//UNiagaraDataInterface Interface
@@ -34,6 +33,10 @@ public:
 	virtual void GetVMExternalFunction(const FVMExternalFunctionBindingInfo& BindingInfo, void* InstanceData, FVMExternalFunction &OutFunc) override;
 	virtual bool CanExecuteOnTarget(ENiagaraSimTarget Target)const override { return Target == ENiagaraSimTarget::GPUComputeSim; }
 
+	virtual int32 PerInstanceDataSize() const override;
+	virtual bool InitPerInstanceData(void* PerInstanceData, FNiagaraSystemInstance* SystemInstance) override;
+	virtual void DestroyPerInstanceData(void* PerInstanceData, FNiagaraSystemInstance* SystemInstance) override;
+
 	virtual bool HasPreSimulateTick() const override { return true; }
 	virtual bool PerInstanceTick(void* PerInstanceData, FNiagaraSystemInstance* SystemInstance, float DeltaSeconds) override;
 	//UNiagaraDataInterface Interface End
@@ -59,24 +62,7 @@ public:
 protected:
 	virtual bool CopyToInternal(UNiagaraDataInterface* Destination) const override;
 
-	virtual void PushToRenderThreadImpl() override;
-
 protected:
-	FIntVector TextureSize = FIntVector::ZeroValue;
-
 	static const FName SampleVolumeTextureName;
 	static const FName TextureDimsName;
 };
-
-struct FNiagaraDataInterfaceProxyVolumeTexture : public FNiagaraDataInterfaceProxy
-{
-	FSamplerStateRHIRef SamplerStateRHI;
-	FTextureRHIRef TextureRHI;
-	FVector TexDims;
-
-	virtual void ConsumePerInstanceDataFromGameThread(void* PerInstanceData, const FNiagaraSystemInstanceID& Instance) override { check(false); }
-	virtual int32 PerInstanceDataPassedToRenderThreadSize() const override
-	{
-		return 0;
-	}
-};
\ No newline at end of file
diff --git a/Engine/Plugins/FX/Niagara/Source/Niagara/Private/NiagaraDataInterface2DArrayTexture.cpp b/Engine/Plugins/FX/Niagara/Source/Niagara/Private/NiagaraDataInterface2DArrayTexture.cpp
--- a/Engine/Plugins/FX/Niagara/Source/Niagara/Private/NiagaraDataInterface2DArrayTexture.cpp
+++ b/Engine/Plugins/FX/Niagara/Source/Niagara/Private/NiagaraDataInterface2DArrayTexture.cpp
@@ -5,7 +5,9 @@
 #include "ShaderParameterUtils.h"
 #include "NiagaraCustomVersion.h"
 #include "Engine/Texture2DArray.h"
-
+#include "Engine/TextureRenderTarget2DArray.h"
+#include "NiagaraSystemInstance.h"
+#include "NiagaraComputeExecutionContext.h"
 
 #define LOCTEXT_NAMESPACE "UNiagaraDataInterface2DArrayTexture"
 
@@ -15,12 +17,61 @@ const FString UNiagaraDataInterface2DArrayTexture::TextureName(TEXT("Texture_"))
 const FString UNiagaraDataInterface2DArrayTexture::SamplerName(TEXT("Sampler_"));
 const FString UNiagaraDataInterface2DArrayTexture::DimensionsBaseName(TEXT("Dimensions_"));
 
+struct FNDITexture2DArrayInstanceData_GameThread
+{
+	TWeakObjectPtr<UTexture> CurrentTexture = nullptr;
+	FIntVector CurrentTextureSize = FIntVector::ZeroValue;
+	FNiagaraParameterDirectBinding<UObject*> UserParamBinding;
+};
+
+struct FNDITexture2DArrayInstanceData_RenderThread
+{
+	FSamplerStateRHIRef		SamplerStateRHI;
+	FTextureReferenceRHIRef	TextureReferenceRHI;
+	FTextureRHIRef			ResolvedTextureRHI;
+	FVector3f				TextureSize;
+};
+
+struct FNiagaraDataInterfaceProxyTexture2DArray : public FNiagaraDataInterfaceProxy
+{
+	virtual void ConsumePerInstanceDataFromGameThread(void* PerInstanceData, const FNiagaraSystemInstanceID& Instance) override { check(false); }
+	virtual int32 PerInstanceDataPassedToRenderThreadSize() const override { return 0; }
+
+	virtual void PreStage(FRHICommandList& RHICmdList, const FNiagaraDataInterfaceStageArgs& Context) override
+	{
+		if (FNDITexture2DArrayInstanceData_RenderThread* InstanceData = InstanceData_RT.Find(Context.SystemInstanceID))
+		{
+			// Because the underlying reference can have a switch in flight on the RHI we get the referenced texture
+			// here, ensure it's valid (as it could be queued for delete) and cache until next round.  If we were
+			// to release the reference in PostStage / PostSimulate we still stand a chance the the transition we
+			// queue will be invalid by the time it is processed on the RHI thread.
+			if (Context.SimStageData->bFirstStage && InstanceData->TextureReferenceRHI.IsValid())
+			{
+				InstanceData->ResolvedTextureRHI = InstanceData->TextureReferenceRHI->GetReferencedTexture();
+				if (InstanceData->ResolvedTextureRHI && !InstanceData->ResolvedTextureRHI->IsValid())
+				{
+					InstanceData->ResolvedTextureRHI = nullptr;
+				}
+			}
+			if (InstanceData->TextureReferenceRHI.IsValid())
+			{
+				// Make sure the texture is readable, we don't know where it's coming from.
+				RHICmdList.Transition(FRHITransitionInfo(InstanceData->TextureReferenceRHI->GetReferencedTexture(), ERHIAccess::Unknown, ERHIAccess::SRVMask));
+			}
+		}
+	}
+
+	TMap<FNiagaraSystemInstanceID, FNDITexture2DArrayInstanceData_RenderThread> InstanceData_RT;
+};
+
 UNiagaraDataInterface2DArrayTexture::UNiagaraDataInterface2DArrayTexture(FObjectInitializer const& ObjectInitializer)
 	: Super(ObjectInitializer)
 	, Texture(nullptr)
 {
-	Proxy.Reset(new FNiagaraDataInterfaceProxy2DArrayTexture());
-	MarkRenderDataDirty();
+	Proxy.Reset(new FNiagaraDataInterfaceProxyTexture2DArray());
+
+	FNiagaraTypeDefinition Def(UTexture::StaticClass());
+	TextureUserParameter.Parameter.SetType(Def);
 }
 
 void UNiagaraDataInterface2DArrayTexture::PostInitProperties()
@@ -32,29 +83,8 @@ void UNiagaraDataInterface2DArrayTexture::PostInitProperties()
 		ENiagaraTypeRegistryFlags Flags = ENiagaraTypeRegistryFlags::AllowAnyVariable | ENiagaraTypeRegistryFlags::AllowParameter;
 		FNiagaraTypeRegistry::Register(FNiagaraTypeDefinition(GetClass()), Flags);
 	}
-
-	MarkRenderDataDirty();
-}
-
-void UNiagaraDataInterface2DArrayTexture::PostLoad()
-{
-	Super::PostLoad();
-
-	// Not safe since the UTexture might not have yet PostLoad() called and so UpdateResource() called.
-	// This will affect whether the SamplerStateRHI will be available or not.
-	MarkRenderDataDirty();
-}
-
-#if WITH_EDITOR
-
-void UNiagaraDataInterface2DArrayTexture::PostEditChangeProperty(struct FPropertyChangedEvent& PropertyChangedEvent)
-{
-	Super::PostEditChangeProperty(PropertyChangedEvent);
-	MarkRenderDataDirty();
 }
 
-#endif
-
 bool UNiagaraDataInterface2DArrayTexture::CopyToInternal(UNiagaraDataInterface* Destination) const
 {
 	if (!Super::CopyToInternal(Destination))
@@ -63,7 +93,7 @@ bool UNiagaraDataInterface2DArrayTexture::CopyToInternal(UNiagaraDataInterface*
 	}
 	UNiagaraDataInterface2DArrayTexture* DestinationTexture = CastChecked<UNiagaraDataInterface2DArrayTexture>(Destination);
 	DestinationTexture->Texture = Texture;
-	DestinationTexture->MarkRenderDataDirty();
+	DestinationTexture->TextureUserParameter = TextureUserParameter;
 
 	return true;
 }
@@ -75,7 +105,9 @@ bool UNiagaraDataInterface2DArrayTexture::Equals(const UNiagaraDataInterface* Ot
 		return false;
 	}
 	const UNiagaraDataInterface2DArrayTexture* OtherTexture = CastChecked<const UNiagaraDataInterface2DArrayTexture>(Other);
-	return OtherTexture->Texture == Texture;
+	return
+		OtherTexture->Texture == Texture &&
+		OtherTexture->TextureUserParameter == TextureUserParameter;
 }
 
 void UNiagaraDataInterface2DArrayTexture::GetFunctions(TArray<FNiagaraFunctionSignature>& OutFunctions)
@@ -92,7 +124,7 @@ void UNiagaraDataInterface2DArrayTexture::GetFunctions(TArray<FNiagaraFunctionSi
 		Sig.Inputs.Add(FNiagaraVariable(FNiagaraTypeDefinition(GetClass()), TEXT("Texture")));
 		Sig.Inputs.Add(FNiagaraVariable(FNiagaraTypeDefinition::GetVec3Def(), TEXT("UVW")));
 		Sig.Inputs.Add(FNiagaraVariable(FNiagaraTypeDefinition::GetFloatDef(), TEXT("MipLevel")));
-		Sig.SetDescription(LOCTEXT("TextureSampleVolumeTextureDesc", "Sample the specified mip level of the input texture at the specified UVW coordinates. Where W is the slice to sample (0,1,2, etc) and UV are the coordinates into the slice."));
+		Sig.SetDescription(LOCTEXT("TextureSample2DArrayTextureDesc", "Sample the specified mip level of the input texture at the specified UVW coordinates. Where W is the slice to sample (0,1,2, etc) and UV are the coordinates into the slice."));
 		Sig.Outputs.Add(FNiagaraVariable(FNiagaraTypeDefinition::GetVec4Def(), TEXT("Value")));
 	}
 
@@ -111,27 +143,84 @@ void UNiagaraDataInterface2DArrayTexture::GetVMExternalFunction(const FVMExterna
 {
 	if (BindingInfo.Name == TextureDimsName)
 	{
-		check(BindingInfo.GetNumInputs() == 0 && BindingInfo.GetNumOutputs() == 3);
+		check(BindingInfo.GetNumInputs() == 1 && BindingInfo.GetNumOutputs() == 3);
 		OutFunc = FVMExternalFunction::CreateUObject(this, &UNiagaraDataInterface2DArrayTexture::GetTextureDimensions);
 	}
 }
 
+int32 UNiagaraDataInterface2DArrayTexture::PerInstanceDataSize() const
+{
+	return sizeof(FNDITexture2DArrayInstanceData_GameThread);
+}
+
+bool UNiagaraDataInterface2DArrayTexture::InitPerInstanceData(void* PerInstanceData, FNiagaraSystemInstance* SystemInstance)
+{
+	FNDITexture2DArrayInstanceData_GameThread* InstanceData = new (PerInstanceData) FNDITexture2DArrayInstanceData_GameThread();
+	InstanceData->UserParamBinding.Init(SystemInstance->GetInstanceParameters(), TextureUserParameter.Parameter);
+	return true;
+}
+
+void UNiagaraDataInterface2DArrayTexture::DestroyPerInstanceData(void* PerInstanceData, FNiagaraSystemInstance* SystemInstance)
+{
+	FNDITexture2DArrayInstanceData_GameThread* InstanceData = static_cast<FNDITexture2DArrayInstanceData_GameThread*>(PerInstanceData);
+	InstanceData->~FNDITexture2DArrayInstanceData_GameThread();
+
+	ENQUEUE_RENDER_COMMAND(NDITexture_RemoveInstance)
+	(
+		[RT_Proxy = GetProxyAs<FNiagaraDataInterfaceProxyTexture2DArray>(), RT_InstanceID = SystemInstance->GetId()](FRHICommandListImmediate&)
+		{
+			RT_Proxy->InstanceData_RT.Remove(RT_InstanceID);
+		}
+	);
+}
+
 bool UNiagaraDataInterface2DArrayTexture::PerInstanceTick(void* PerInstanceData, FNiagaraSystemInstance* SystemInstance, float DeltaSeconds)
 {
-	const FIntVector CurrentTextureSize = Texture != nullptr ? FIntVector(Texture->GetSizeX(), Texture->GetSizeY(), Texture->GetArraySize()) : FIntVector::ZeroValue;
-	if ( CurrentTextureSize != TextureSize )
+	FNDITexture2DArrayInstanceData_GameThread* InstanceData = static_cast<FNDITexture2DArrayInstanceData_GameThread*>(PerInstanceData);
+
+	UTexture* CurrentTexture = InstanceData->UserParamBinding.GetValueOrDefault<UTexture>(Texture);
+	if ( InstanceData->CurrentTexture != CurrentTexture )
 	{
-		TextureSize = CurrentTextureSize;
-		MarkRenderDataDirty();
+		UTexture2DArray* CurrentTextureArray = Cast<UTexture2DArray>(CurrentTexture);
+		UTextureRenderTarget2DArray* CurrentTextureRT = Cast<UTextureRenderTarget2DArray>(CurrentTexture);
+		if (CurrentTextureArray || CurrentTextureRT)
+		{
+			const FIntVector CurrentTextureSize = CurrentTextureArray ?
+				FIntVector(CurrentTextureArray->GetSizeX(), CurrentTextureArray->GetSizeY(), CurrentTextureArray->GetArraySize()) :
+				FIntVector(CurrentTextureRT->SizeX, CurrentTextureRT->SizeY, CurrentTextureRT->Slices);
+
+			InstanceData->CurrentTexture = CurrentTexture;
+			InstanceData->CurrentTextureSize = CurrentTextureSize;
+
+			ENQUEUE_RENDER_COMMAND(NDITexture_UpdateInstance)
+			(
+				[RT_Proxy=GetProxyAs<FNiagaraDataInterfaceProxyTexture2DArray>(), RT_InstanceID=SystemInstance->GetId(), RT_Texture=CurrentTexture, RT_TextureSize=CurrentTextureSize](FRHICommandListImmediate&)
+				{
+					FNDITexture2DArrayInstanceData_RenderThread& InstanceData = RT_Proxy->InstanceData_RT.FindOrAdd(RT_InstanceID);
+					if (RT_Texture)
+					{
+						InstanceData.TextureReferenceRHI = RT_Texture->TextureReference.TextureReferenceRHI;
+						InstanceData.SamplerStateRHI = RT_Texture->GetResource() ? RT_Texture->GetResource()->SamplerStateRHI : nullptr;
+					}
+					else
+					{
+						InstanceData.TextureReferenceRHI = nullptr;
+						InstanceData.SamplerStateRHI = nullptr;
+					}
+					InstanceData.TextureSize = FVector(RT_TextureSize.X, RT_TextureSize.Y, RT_TextureSize.Z);
+				}
+			);
+		}
 	}
 	return false;
 }
 
 void UNiagaraDataInterface2DArrayTexture::GetTextureDimensions(FVectorVMExternalFunctionContext& Context)
 {
+	VectorVM::FUserPtrHandler<FNDITexture2DArrayInstanceData_GameThread> InstData(Context);
 	FNDIOutputParam<FVector3f> DimensionsOut(Context);
 
-	const FVector FloatTextureSize(TextureSize.X, TextureSize.Y, TextureSize.Z);
+	const FVector FloatTextureSize(InstData->CurrentTextureSize.X, InstData->CurrentTextureSize.Y, InstData->CurrentTextureSize.Z);
 	for (int32 i = 0; i < Context.GetNumInstances(); ++i)
 	{
 		DimensionsOut.SetAndAdvance(FloatTextureSize);
@@ -191,26 +280,21 @@ public:
 		check(IsInRenderingThread());
 
 		FRHIComputeShader* ComputeShaderRHI = Context.Shader.GetComputeShader();
-		FNiagaraDataInterfaceProxy2DArrayTexture* TextureDI = static_cast<FNiagaraDataInterfaceProxy2DArrayTexture*>(Context.DataInterface);
+		FNiagaraDataInterfaceProxyTexture2DArray* TextureDI = static_cast<FNiagaraDataInterfaceProxyTexture2DArray*>(Context.DataInterface);
+		FNDITexture2DArrayInstanceData_RenderThread* InstanceData = TextureDI->InstanceData_RT.Find(Context.SystemInstanceID);
 
-		if (TextureDI && TextureDI->TextureRHI)
+		if (InstanceData && InstanceData->ResolvedTextureRHI.IsValid())
 		{
-			FRHISamplerState* SamplerStateRHI = TextureDI->SamplerStateRHI;
-			if (!SamplerStateRHI)
-			{
-				// Fallback required because PostLoad() order affects whether RHI resources 
-				// are initalized in UNiagaraDataInterface2DArrayTexture::PushToRenderThreadImpl().
-				SamplerStateRHI = TStaticSamplerState<SF_Point, AM_Clamp, AM_Clamp, AM_Clamp>::GetRHI();
-			}
+			FRHISamplerState* SamplerStateRHI = InstanceData->SamplerStateRHI ? InstanceData->SamplerStateRHI : GBlackTexture->SamplerStateRHI;
 			SetTextureParameter(
 				RHICmdList,
 				ComputeShaderRHI,
 				TextureParam,
 				SamplerParam,
 				SamplerStateRHI,
-				TextureDI->TextureRHI
+				InstanceData->ResolvedTextureRHI
 			);
-			SetShaderValue(RHICmdList, ComputeShaderRHI, Dimensions, (FVector3f)TextureDI->TexDims);
+			SetShaderValue(RHICmdList, ComputeShaderRHI, Dimensions, InstanceData->TextureSize);
 		}
 		else
 		{
@@ -235,35 +319,11 @@ IMPLEMENT_TYPE_LAYOUT(FNiagaraDataInterfaceParametersCS_2DArrayTexture);
 
 IMPLEMENT_NIAGARA_DI_PARAMETER(UNiagaraDataInterface2DArrayTexture, FNiagaraDataInterfaceParametersCS_2DArrayTexture);
 
-void UNiagaraDataInterface2DArrayTexture::PushToRenderThreadImpl()
-{
-	FNiagaraDataInterfaceProxy2DArrayTexture* RT_Proxy = GetProxyAs<FNiagaraDataInterfaceProxy2DArrayTexture>();
-
-	TextureSize = FIntVector::ZeroValue;
-	if (Texture)
-	{
-		TextureSize.X = Texture->GetSizeX();
-		TextureSize.Y = Texture->GetSizeY();
-		TextureSize.Z = Texture->GetArraySize();
-	}
-
-	ENQUEUE_RENDER_COMMAND(FPushDITextureToRT)
-	(
-		[RT_Proxy, RT_Resource=Texture ? Texture->GetResource() : nullptr, RT_TexDims=TextureSize](FRHICommandListImmediate& RHICmdList)
-		{
-			RT_Proxy->TextureRHI = RT_Resource ? RT_Resource->TextureRHI : nullptr;
-			RT_Proxy->SamplerStateRHI = RT_Resource ? RT_Resource->SamplerStateRHI : nullptr;
-			RT_Proxy->TexDims = FVector(RT_TexDims.X, RT_TexDims.Y, RT_TexDims.Z);
-		}
-	);
-}
-
 void UNiagaraDataInterface2DArrayTexture::SetTexture(UTexture2DArray* InTexture)
 {
 	if (InTexture)
 	{
 		Texture = InTexture;
-		MarkRenderDataDirty();
 	}
 }
 
diff --git a/Engine/Plugins/FX/Niagara/Source/Niagara/Private/NiagaraDataInterfaceCubeTexture.cpp b/Engine/Plugins/FX/Niagara/Source/Niagara/Private/NiagaraDataInterfaceCubeTexture.cpp
--- a/Engine/Plugins/FX/Niagara/Source/Niagara/Private/NiagaraDataInterfaceCubeTexture.cpp
+++ b/Engine/Plugins/FX/Niagara/Source/Niagara/Private/NiagaraDataInterfaceCubeTexture.cpp
@@ -5,7 +5,9 @@
 #include "ShaderParameterUtils.h"
 #include "NiagaraCustomVersion.h"
 #include "Engine/TextureCube.h"
-
+#include "Engine/TextureRenderTargetCube.h"
+#include "NiagaraSystemInstance.h"
+#include "NiagaraComputeExecutionContext.h"
 
 #define LOCTEXT_NAMESPACE "UNiagaraDataInterfaceCubeTexture"
 
@@ -15,22 +17,61 @@ const FString UNiagaraDataInterfaceCubeTexture::TextureName(TEXT("Texture_"));
 const FString UNiagaraDataInterfaceCubeTexture::SamplerName(TEXT("Sampler_"));
 const FString UNiagaraDataInterfaceCubeTexture::DimensionsBaseName(TEXT("Dimensions_"));
 
-struct FNiagaraDataInterfaceProxyCubeTexture : public FNiagaraDataInterfaceProxy
+struct FNDICubeTextureInstanceData_GameThread
+{
+	TWeakObjectPtr<UTexture> CurrentTexture = nullptr;
+	FIntPoint CurrentTextureSize = FIntPoint::ZeroValue;
+	FNiagaraParameterDirectBinding<UObject*> UserParamBinding;
+};
+
+struct FNDICubeTextureInstanceData_RenderThread
 {
-	FSamplerStateRHIRef SamplerStateRHI;
-	FTextureRHIRef TextureRHI;
-	FIntPoint TexDims;
+	FSamplerStateRHIRef		SamplerStateRHI;
+	FTextureReferenceRHIRef	TextureReferenceRHI;
+	FTextureRHIRef			ResolvedTextureRHI;
+	FIntPoint				TextureSize;
+};
 
+struct FNiagaraDataInterfaceProxyCubeTexture : public FNiagaraDataInterfaceProxy
+{
 	virtual void ConsumePerInstanceDataFromGameThread(void* PerInstanceData, const FNiagaraSystemInstanceID& Instance) override { check(false); }
 	virtual int32 PerInstanceDataPassedToRenderThreadSize() const override { return 0; }
-};
+
+	virtual void PreStage(FRHICommandList& RHICmdList, const FNiagaraDataInterfaceStageArgs& Context) override
+	{
+		if (FNDICubeTextureInstanceData_RenderThread* InstanceData = InstanceData_RT.Find(Context.SystemInstanceID))
+		{
+			// Because the underlying reference can have a switch in flight on the RHI we get the referenced texture
+			// here, ensure it's valid (as it could be queued for delete) and cache until next round.  If we were
+			// to release the reference in PostStage / PostSimulate we still stand a chance the the transition we
+			// queue will be invalid by the time it is processed on the RHI thread.
+			if (Context.SimStageData->bFirstStage && InstanceData->TextureReferenceRHI.IsValid())
+			{
+				InstanceData->ResolvedTextureRHI = InstanceData->TextureReferenceRHI->GetReferencedTexture();
+				if (InstanceData->ResolvedTextureRHI && !InstanceData->ResolvedTextureRHI->IsValid())
+				{
+					InstanceData->ResolvedTextureRHI = nullptr;
+				}
+			}
+			if (InstanceData->TextureReferenceRHI.IsValid())
+			{
+				// Make sure the texture is readable, we don't know where it's coming from.
+				RHICmdList.Transition(FRHITransitionInfo(InstanceData->TextureReferenceRHI->GetReferencedTexture(), ERHIAccess::Unknown, ERHIAccess::SRVMask));
+			}
+		}
+	}
+
+	TMap<FNiagaraSystemInstanceID, FNDICubeTextureInstanceData_RenderThread> InstanceData_RT;
+}; 
 
 UNiagaraDataInterfaceCubeTexture::UNiagaraDataInterfaceCubeTexture(FObjectInitializer const& ObjectInitializer)
 	: Super(ObjectInitializer)
 	, Texture(nullptr)
 {
 	Proxy.Reset(new FNiagaraDataInterfaceProxyCubeTexture());
-	MarkRenderDataDirty();
+
+	FNiagaraTypeDefinition Def(UTexture::StaticClass());
+	TextureUserParameter.Parameter.SetType(Def);
 }
 
 void UNiagaraDataInterfaceCubeTexture::PostInitProperties()
@@ -42,29 +83,8 @@ void UNiagaraDataInterfaceCubeTexture::PostInitProperties()
 		ENiagaraTypeRegistryFlags Flags = ENiagaraTypeRegistryFlags::AllowAnyVariable | ENiagaraTypeRegistryFlags::AllowParameter;
 		FNiagaraTypeRegistry::Register(FNiagaraTypeDefinition(GetClass()), Flags);
 	}
-
-	MarkRenderDataDirty();
-}
-
-void UNiagaraDataInterfaceCubeTexture::PostLoad()
-{
-	Super::PostLoad();
-
-	// Not safe since the UTexture might not have yet PostLoad() called and so UpdateResource() called.
-	// This will affect whether the SamplerStateRHI will be available or not.
-	MarkRenderDataDirty();
-}
-
-#if WITH_EDITOR
-
-void UNiagaraDataInterfaceCubeTexture::PostEditChangeProperty(struct FPropertyChangedEvent& PropertyChangedEvent)
-{
-	Super::PostEditChangeProperty(PropertyChangedEvent);
-	MarkRenderDataDirty();
 }
 
-#endif
-
 bool UNiagaraDataInterfaceCubeTexture::CopyToInternal(UNiagaraDataInterface* Destination) const
 {
 	if (!Super::CopyToInternal(Destination))
@@ -73,7 +93,7 @@ bool UNiagaraDataInterfaceCubeTexture::CopyToInternal(UNiagaraDataInterface* Des
 	}
 	UNiagaraDataInterfaceCubeTexture* DestinationTexture = CastChecked<UNiagaraDataInterfaceCubeTexture>(Destination);
 	DestinationTexture->Texture = Texture;
-	DestinationTexture->MarkRenderDataDirty();
+	DestinationTexture->TextureUserParameter = TextureUserParameter;
 
 	return true;
 }
@@ -85,7 +105,9 @@ bool UNiagaraDataInterfaceCubeTexture::Equals(const UNiagaraDataInterface* Other
 		return false;
 	}
 	const UNiagaraDataInterfaceCubeTexture* OtherTexture = CastChecked<const UNiagaraDataInterfaceCubeTexture>(Other);
-	return OtherTexture->Texture == Texture;
+	return
+		OtherTexture->Texture == Texture &&
+		OtherTexture->TextureUserParameter == TextureUserParameter;
 }
 
 void UNiagaraDataInterfaceCubeTexture::GetFunctions(TArray<FNiagaraFunctionSignature>& OutFunctions)
@@ -121,41 +143,99 @@ void UNiagaraDataInterfaceCubeTexture::GetVMExternalFunction(const FVMExternalFu
 {
 	if (BindingInfo.Name == SampleCubeTextureName)
 	{
-		check(BindingInfo.GetNumInputs() == 4 && BindingInfo.GetNumOutputs() == 4);
+		check(BindingInfo.GetNumInputs() == 5 && BindingInfo.GetNumOutputs() == 4);
 		NDI_FUNC_BINDER(UNiagaraDataInterfaceCubeTexture, SampleCubeTexture)::Bind(this, OutFunc);
 	}
 	else if (BindingInfo.Name == TextureDimsName)
 	{
-		check(BindingInfo.GetNumInputs() == 0 && BindingInfo.GetNumOutputs() == 2);
+		check(BindingInfo.GetNumInputs() == 1 && BindingInfo.GetNumOutputs() == 2);
 		OutFunc = FVMExternalFunction::CreateUObject(this, &UNiagaraDataInterfaceCubeTexture::GetTextureDimensions);
 	}
 }
 
+int32 UNiagaraDataInterfaceCubeTexture::PerInstanceDataSize() const
+{
+	return sizeof(FNDICubeTextureInstanceData_GameThread);
+}
+
+bool UNiagaraDataInterfaceCubeTexture::InitPerInstanceData(void* PerInstanceData, FNiagaraSystemInstance* SystemInstance)
+{
+	FNDICubeTextureInstanceData_GameThread* InstanceData = new (PerInstanceData) FNDICubeTextureInstanceData_GameThread();
+	InstanceData->UserParamBinding.Init(SystemInstance->GetInstanceParameters(), TextureUserParameter.Parameter);
+	return true;
+}
+
+void UNiagaraDataInterfaceCubeTexture::DestroyPerInstanceData(void* PerInstanceData, FNiagaraSystemInstance* SystemInstance)
+{
+	FNDICubeTextureInstanceData_GameThread* InstanceData = static_cast<FNDICubeTextureInstanceData_GameThread*>(PerInstanceData);
+	InstanceData->~FNDICubeTextureInstanceData_GameThread();
+
+	ENQUEUE_RENDER_COMMAND(NDITexture_RemoveInstance)
+	(
+		[RT_Proxy = GetProxyAs<FNiagaraDataInterfaceProxyCubeTexture>(), RT_InstanceID = SystemInstance->GetId()](FRHICommandListImmediate&)
+		{
+			RT_Proxy->InstanceData_RT.Remove(RT_InstanceID);
+		}
+	);
+}
+
 bool UNiagaraDataInterfaceCubeTexture::PerInstanceTick(void* PerInstanceData, FNiagaraSystemInstance* SystemInstance, float DeltaSeconds)
 {
-	const FIntPoint CurrentTextureSize = Texture != nullptr ? FIntPoint(Texture->GetSizeX(), Texture->GetSizeY()) : FIntPoint::ZeroValue;
-	if ( CurrentTextureSize != TextureSize )
+	FNDICubeTextureInstanceData_GameThread* InstanceData = static_cast<FNDICubeTextureInstanceData_GameThread*>(PerInstanceData);
+
+	UTexture* CurrentTexture = InstanceData->UserParamBinding.GetValueOrDefault<UTexture>(Texture);
+	if ( InstanceData->CurrentTexture != CurrentTexture )
 	{
-		TextureSize = CurrentTextureSize;
-		MarkRenderDataDirty();
+		UTextureCube* CurrentTextureCube = Cast<UTextureCube>(CurrentTexture);
+		UTextureRenderTargetCube* CurrentTextureRT = Cast<UTextureRenderTargetCube>(CurrentTexture);
+		if (CurrentTextureCube || CurrentTextureRT)
+		{
+			const FIntPoint CurrentTextureSize = CurrentTextureCube != nullptr ?
+				FIntPoint(CurrentTextureCube->GetSizeX(), CurrentTextureCube->GetSizeY()) :
+				FIntPoint(CurrentTextureRT->SizeX, CurrentTextureRT->SizeX);
+
+			InstanceData->CurrentTexture = CurrentTexture;
+			InstanceData->CurrentTextureSize = CurrentTextureSize;
+
+			ENQUEUE_RENDER_COMMAND(NDITexture_UpdateInstance)
+			(
+				[RT_Proxy=GetProxyAs<FNiagaraDataInterfaceProxyCubeTexture>(), RT_InstanceID=SystemInstance->GetId(), RT_Texture=CurrentTexture, RT_TextureSize=CurrentTextureSize](FRHICommandListImmediate&)
+				{
+					FNDICubeTextureInstanceData_RenderThread& InstanceData = RT_Proxy->InstanceData_RT.FindOrAdd(RT_InstanceID);
+					if (RT_Texture)
+					{
+						InstanceData.TextureReferenceRHI = RT_Texture->TextureReference.TextureReferenceRHI;
+						InstanceData.SamplerStateRHI = RT_Texture->GetResource() ? RT_Texture->GetResource()->SamplerStateRHI : nullptr;
+					}
+					else
+					{
+						InstanceData.TextureReferenceRHI = nullptr;
+						InstanceData.SamplerStateRHI = nullptr;
+					}
+					InstanceData.TextureSize = RT_TextureSize;
+				}
+			);
+		}
 	}
 	return false;
 }
 
 void UNiagaraDataInterfaceCubeTexture::GetTextureDimensions(FVectorVMExternalFunctionContext& Context)
 {
+	VectorVM::FUserPtrHandler<FNDICubeTextureInstanceData_GameThread> InstData(Context);
 	FNDIOutputParam<int32> OutWidth(Context);
 	FNDIOutputParam<int32> OutHeight(Context);
 
 	for (int32 i = 0; i < Context.GetNumInstances(); ++i)
 	{
-		OutWidth.SetAndAdvance(TextureSize.X);
-		OutHeight.SetAndAdvance(TextureSize.Y);
+		OutWidth.SetAndAdvance(InstData->CurrentTextureSize.X);
+		OutHeight.SetAndAdvance(InstData->CurrentTextureSize.Y);
 	}
 }
 
 void UNiagaraDataInterfaceCubeTexture::SampleCubeTexture(FVectorVMExternalFunctionContext& Context)
 {
+	VectorVM::FUserPtrHandler<FNDICubeTextureInstanceData_GameThread> InstData(Context);
 	FNDIInputParam<FVector3f> InCoord(Context);
 	FNDIInputParam<float> InMipLevel(Context);
 	FNDIOutputParam<FVector4f> OutColor(Context);
@@ -228,23 +308,20 @@ public:
 
 		FRHIComputeShader* ComputeShaderRHI = Context.Shader.GetComputeShader();
 		FNiagaraDataInterfaceProxyCubeTexture* TextureDI = static_cast<FNiagaraDataInterfaceProxyCubeTexture*>(Context.DataInterface);
+		FNDICubeTextureInstanceData_RenderThread* InstanceData = TextureDI->InstanceData_RT.Find(Context.SystemInstanceID);
 
-		if (TextureDI && TextureDI->TextureRHI)
+		if (InstanceData && InstanceData->ResolvedTextureRHI.IsValid())
 		{
-			FRHISamplerState* SamplerStateRHI = TextureDI->SamplerStateRHI;
-			if (!SamplerStateRHI)
-			{
-				SamplerStateRHI = TStaticSamplerState<SF_Point, AM_Clamp, AM_Clamp, AM_Clamp>::GetRHI();
-			}
+			FRHISamplerState* SamplerStateRHI = InstanceData->SamplerStateRHI ? InstanceData->SamplerStateRHI : GBlackTexture->SamplerStateRHI;
 			SetTextureParameter(
 				RHICmdList,
 				ComputeShaderRHI,
 				TextureParam,
 				SamplerParam,
 				SamplerStateRHI,
-				TextureDI->TextureRHI
+				InstanceData->ResolvedTextureRHI
 			);
-			SetShaderValue(RHICmdList, ComputeShaderRHI, Dimensions, TextureDI->TexDims);
+			SetShaderValue(RHICmdList, ComputeShaderRHI, Dimensions, InstanceData->TextureSize);
 		}
 		else
 		{
@@ -268,34 +345,11 @@ private:
 IMPLEMENT_TYPE_LAYOUT(FNiagaraDataInterfaceParametersCS_CubeTexture);
 IMPLEMENT_NIAGARA_DI_PARAMETER(UNiagaraDataInterfaceCubeTexture, FNiagaraDataInterfaceParametersCS_CubeTexture);
 
-void UNiagaraDataInterfaceCubeTexture::PushToRenderThreadImpl()
-{
-	FNiagaraDataInterfaceProxyCubeTexture* RT_Proxy = GetProxyAs<FNiagaraDataInterfaceProxyCubeTexture>();
-
-	TextureSize = FIntPoint::ZeroValue;
-	if (Texture)
-	{
-		TextureSize.X = Texture->GetSizeX();
-		TextureSize.Y = Texture->GetSizeY();
-	}
-
-	ENQUEUE_RENDER_COMMAND(FPushDITextureToRT)
-	(
-		[RT_Proxy, RT_Resource=Texture ? Texture->GetResource() : nullptr, RT_TexDims=TextureSize](FRHICommandListImmediate& RHICmdList)
-		{
-			RT_Proxy->TextureRHI = RT_Resource ? RT_Resource->TextureRHI : nullptr;
-			RT_Proxy->SamplerStateRHI = RT_Resource ? RT_Resource->SamplerStateRHI : nullptr;
-			RT_Proxy->TexDims = RT_TexDims;
-		}
-	);
-}
-
 void UNiagaraDataInterfaceCubeTexture::SetTexture(UTextureCube* InTexture)
 {
 	if (InTexture)
 	{
 		Texture = InTexture;
-		MarkRenderDataDirty();
 	}
 }
 
diff --git a/Engine/Plugins/FX/Niagara/Source/Niagara/Private/NiagaraDataInterfaceTexture.cpp b/Engine/Plugins/FX/Niagara/Source/Niagara/Private/NiagaraDataInterfaceTexture.cpp
--- a/Engine/Plugins/FX/Niagara/Source/Niagara/Private/NiagaraDataInterfaceTexture.cpp
+++ b/Engine/Plugins/FX/Niagara/Source/Niagara/Private/NiagaraDataInterfaceTexture.cpp
@@ -3,6 +3,7 @@
 #include "NiagaraDataInterfaceTexture.h"
 #include "NiagaraShader.h"
 #include "ShaderParameterUtils.h"
+#include "NiagaraSystemInstance.h"
 #include "NiagaraCustomVersion.h"
 #include "NiagaraComputeExecutionContext.h"
 
@@ -16,13 +17,23 @@ const FString UNiagaraDataInterfaceTexture::TextureName(TEXT("Texture_"));
 const FString UNiagaraDataInterfaceTexture::SamplerName(TEXT("Sampler_"));
 const FString UNiagaraDataInterfaceTexture::DimensionsBaseName(TEXT("Dimensions_"));
 
-struct FNiagaraDataInterfaceProxyTexture : public FNiagaraDataInterfaceProxy
+struct FNDITextureInstanceData_GameThread
+{
+	TWeakObjectPtr<UTexture> CurrentTexture = nullptr;
+	FIntPoint CurrentTextureSize = FIntPoint::ZeroValue;
+	FNiagaraParameterDirectBinding<UObject*> UserParamBinding;
+};
+
+struct FNDITextureInstanceData_RenderThread
 {
-	FSamplerStateRHIRef SamplerStateRHI;
-	FTextureReferenceRHIRef TextureReferenceRHI;
-	FTextureRHIRef ResolvedTextureRHI;
-	FVector2D TexDims;
+	FSamplerStateRHIRef		SamplerStateRHI;
+	FTextureReferenceRHIRef	TextureReferenceRHI;
+	FTextureRHIRef			ResolvedTextureRHI;
+	FVector2f				TextureSize;
+};
 
+struct FNiagaraDataInterfaceProxyTexture : public FNiagaraDataInterfaceProxy
+{
 	virtual void ConsumePerInstanceDataFromGameThread(void* PerInstanceData, const FNiagaraSystemInstanceID& Instance) override
 	{
 		checkNoEntry();
@@ -35,24 +46,29 @@ struct FNiagaraDataInterfaceProxyTexture : public FNiagaraDataInterfaceProxy
 
 	virtual void PreStage(FRHICommandList& RHICmdList, const FNiagaraDataInterfaceStageArgs& Context) override
 	{
-		// Because the underlying reference can have a switch in flight on the RHI we get the referenced texture
-		// here, ensure it's valid (as it could be queued for delete) and cache until next round.  If we were
-		// to release the reference in PostStage / PostSimulate we still stand a chance the the transition we
-		// queue will be invalid by the time it is processed on the RHI thread.
-		if (Context.SimStageData->bFirstStage && TextureReferenceRHI.IsValid())
+		if ( FNDITextureInstanceData_RenderThread* InstanceData = InstanceData_RT.Find(Context.SystemInstanceID) )
 		{
-			ResolvedTextureRHI = TextureReferenceRHI->GetReferencedTexture();
-			if (ResolvedTextureRHI && !ResolvedTextureRHI->IsValid())
+			// Because the underlying reference can have a switch in flight on the RHI we get the referenced texture
+			// here, ensure it's valid (as it could be queued for delete) and cache until next round.  If we were
+			// to release the reference in PostStage / PostSimulate we still stand a chance the the transition we
+			// queue will be invalid by the time it is processed on the RHI thread.
+			if (Context.SimStageData->bFirstStage && InstanceData->TextureReferenceRHI.IsValid())
 			{
-				ResolvedTextureRHI = nullptr;
+				InstanceData->ResolvedTextureRHI = InstanceData->TextureReferenceRHI->GetReferencedTexture();
+				if (InstanceData->ResolvedTextureRHI && !InstanceData->ResolvedTextureRHI->IsValid())
+				{
+					InstanceData->ResolvedTextureRHI = nullptr;
+				}
+			}
+			if (InstanceData->TextureReferenceRHI.IsValid())
+			{
+				// Make sure the texture is readable, we don't know where it's coming from.
+				RHICmdList.Transition(FRHITransitionInfo(InstanceData->TextureReferenceRHI->GetReferencedTexture(), ERHIAccess::Unknown, ERHIAccess::SRVMask));
 			}
-		}
-		if (ResolvedTextureRHI)
-		{
-			// Make sure the texture is readable, we don't know where it's coming from.
-			RHICmdList.Transition(FRHITransitionInfo(ResolvedTextureRHI, ERHIAccess::Unknown, ERHIAccess::SRVMask));
 		}
 	}
+
+	TMap<FNiagaraSystemInstanceID, FNDITextureInstanceData_RenderThread> InstanceData_RT;
 };
 
 UNiagaraDataInterfaceTexture::UNiagaraDataInterfaceTexture(FObjectInitializer const& ObjectInitializer)
@@ -60,7 +76,9 @@ UNiagaraDataInterfaceTexture::UNiagaraDataInterfaceTexture(FObjectInitializer co
 	, Texture(nullptr)
 {
 	Proxy.Reset(new FNiagaraDataInterfaceProxyTexture());
-	MarkRenderDataDirty();
+
+	FNiagaraTypeDefinition Def(UTexture::StaticClass());
+	TextureUserParameter.Parameter.SetType(Def);
 }
 
 void UNiagaraDataInterfaceTexture::PostInitProperties()
@@ -72,8 +90,6 @@ void UNiagaraDataInterfaceTexture::PostInitProperties()
 		ENiagaraTypeRegistryFlags Flags = ENiagaraTypeRegistryFlags::AllowAnyVariable | ENiagaraTypeRegistryFlags::AllowParameter;
 		FNiagaraTypeRegistry::Register(FNiagaraTypeDefinition(GetClass()), Flags);
 	}
-
-	MarkRenderDataDirty();
 }
 
 void UNiagaraDataInterfaceTexture::PostLoad()
@@ -89,23 +105,8 @@ void UNiagaraDataInterfaceTexture::PostLoad()
 		}
 	}
 #endif
-	// Not safe since the UTexture might not have yet PostLoad() called and so UpdateResource() called.
-	// This will affect whether the SamplerStateRHI will be available or not.
-	TextureSize = Texture != nullptr ? FIntPoint(Texture->GetSurfaceWidth(), Texture->GetSurfaceHeight()) : FIntPoint::ZeroValue;
-	MarkRenderDataDirty();
 }
 
-#if WITH_EDITOR
-
-void UNiagaraDataInterfaceTexture::PostEditChangeProperty(struct FPropertyChangedEvent& PropertyChangedEvent)
-{
-	Super::PostEditChangeProperty(PropertyChangedEvent);
-	TextureSize = Texture != nullptr ? FIntPoint(Texture->GetSurfaceWidth(), Texture->GetSurfaceHeight()) : FIntPoint::ZeroValue;
-	MarkRenderDataDirty();
-}
-
-#endif
-
 void UNiagaraDataInterfaceTexture::Serialize(FArchive& Ar)
 {
 	Super::Serialize(Ar);
@@ -125,8 +126,7 @@ bool UNiagaraDataInterfaceTexture::CopyToInternal(UNiagaraDataInterface* Destina
 	}
 	UNiagaraDataInterfaceTexture* DestinationTexture = CastChecked<UNiagaraDataInterfaceTexture>(Destination);
 	DestinationTexture->Texture = Texture;
-	DestinationTexture->TextureSize = Texture != nullptr ? FIntPoint(Texture->GetSurfaceWidth(), Texture->GetSurfaceHeight()) : FIntPoint::ZeroValue;
-	DestinationTexture->MarkRenderDataDirty();
+	DestinationTexture->TextureUserParameter = TextureUserParameter;
 
 	return true;
 }
@@ -138,7 +138,9 @@ bool UNiagaraDataInterfaceTexture::Equals(const UNiagaraDataInterface* Other) co
 		return false;
 	}
 	const UNiagaraDataInterfaceTexture* OtherTexture = CastChecked<const UNiagaraDataInterfaceTexture>(Other);
-	return OtherTexture->Texture == Texture;
+	return
+		OtherTexture->Texture == Texture &&
+		OtherTexture->TextureUserParameter == TextureUserParameter;
 }
 
 void UNiagaraDataInterfaceTexture::GetFunctions(TArray<FNiagaraFunctionSignature>& OutFunctions)
@@ -216,40 +218,97 @@ void UNiagaraDataInterfaceTexture::GetVMExternalFunction(const FVMExternalFuncti
 {
 	if (BindingInfo.Name == SampleTexture2DName)
 	{
-		check(BindingInfo.GetNumInputs() == 2 && BindingInfo.GetNumOutputs() == 4);
+		check(BindingInfo.GetNumInputs() == 3 && BindingInfo.GetNumOutputs() == 4);
 		NDI_FUNC_BINDER(UNiagaraDataInterfaceTexture, SampleTexture)::Bind(this, OutFunc);
 	}
 	else if (BindingInfo.Name == SamplePseudoVolumeTextureName)
 	{
-		check(BindingInfo.GetNumInputs() == 12 && BindingInfo.GetNumOutputs() == 4);
+		check(BindingInfo.GetNumInputs() == 13 && BindingInfo.GetNumOutputs() == 4);
 		NDI_FUNC_BINDER(UNiagaraDataInterfaceTexture, SamplePseudoVolumeTexture)::Bind(this, OutFunc);
 	}
 	else if (BindingInfo.Name == TextureDimsName)
 	{
-		check(BindingInfo.GetNumInputs() == 0 && BindingInfo.GetNumOutputs() == 2);
+		check(BindingInfo.GetNumInputs() == 1 && BindingInfo.GetNumOutputs() == 2);
 		OutFunc = FVMExternalFunction::CreateUObject(this, &UNiagaraDataInterfaceTexture::GetTextureDimensions);
 	}
 }
 
+int32 UNiagaraDataInterfaceTexture::PerInstanceDataSize() const
+{
+	return sizeof(FNDITextureInstanceData_GameThread);
+}
+
+bool UNiagaraDataInterfaceTexture::InitPerInstanceData(void* PerInstanceData, FNiagaraSystemInstance* SystemInstance)
+{
+	FNDITextureInstanceData_GameThread* InstanceData = new (PerInstanceData) FNDITextureInstanceData_GameThread();
+	InstanceData->UserParamBinding.Init(SystemInstance->GetInstanceParameters(), TextureUserParameter.Parameter);
+	return true;
+}
+
+void UNiagaraDataInterfaceTexture::DestroyPerInstanceData(void* PerInstanceData, FNiagaraSystemInstance* SystemInstance)
+{
+	FNDITextureInstanceData_GameThread* InstanceData = static_cast<FNDITextureInstanceData_GameThread*>(PerInstanceData);
+	InstanceData->~FNDITextureInstanceData_GameThread();
+
+	ENQUEUE_RENDER_COMMAND(NDITexture_RemoveInstance)
+	(
+		[RT_Proxy=GetProxyAs<FNiagaraDataInterfaceProxyTexture>(), RT_InstanceID=SystemInstance->GetId()](FRHICommandListImmediate&)
+		{
+			RT_Proxy->InstanceData_RT.Remove(RT_InstanceID);
+		}
+	);
+}
+
 bool UNiagaraDataInterfaceTexture::PerInstanceTick(void* PerInstanceData, FNiagaraSystemInstance* SystemInstance, float DeltaSeconds)
 {
+	FNDITextureInstanceData_GameThread* InstanceData = static_cast<FNDITextureInstanceData_GameThread*>(PerInstanceData);
+
+	UTexture* CurrentTexture = InstanceData->UserParamBinding.GetValueOrDefault<UTexture>(Texture);
+	const FIntPoint CurrentTextureSize = CurrentTexture != nullptr ? FIntPoint(CurrentTexture->GetSurfaceWidth(), CurrentTexture->GetSurfaceHeight()) : FIntPoint::ZeroValue;
+	if ( (InstanceData->CurrentTexture != CurrentTexture) || (InstanceData->CurrentTextureSize != CurrentTextureSize) )
+	{
+		InstanceData->CurrentTexture = CurrentTexture;
+		InstanceData->CurrentTextureSize = CurrentTextureSize;
+
+		ENQUEUE_RENDER_COMMAND(NDITexture_UpdateInstance)
+		(
+			[RT_Proxy=GetProxyAs<FNiagaraDataInterfaceProxyTexture>(), RT_InstanceID=SystemInstance->GetId(), RT_Texture=CurrentTexture, RT_TextureSize=CurrentTextureSize](FRHICommandListImmediate&)
+			{
+				FNDITextureInstanceData_RenderThread& InstanceData = RT_Proxy->InstanceData_RT.FindOrAdd(RT_InstanceID);
+				if (RT_Texture)
+				{
+					InstanceData.TextureReferenceRHI = RT_Texture->TextureReference.TextureReferenceRHI;
+					InstanceData.SamplerStateRHI = RT_Texture->GetResource() ? RT_Texture->GetResource()->SamplerStateRHI : nullptr;
+				}
+				else
+				{
+					InstanceData.TextureReferenceRHI = nullptr;
+					InstanceData.SamplerStateRHI = nullptr;
+				}
+				InstanceData.TextureSize = FVector2f(RT_TextureSize.X, RT_TextureSize.Y);
+			}
+		);
+	}
+
 	return false;
 }
 
 void UNiagaraDataInterfaceTexture::GetTextureDimensions(FVectorVMExternalFunctionContext& Context)
 {
+	VectorVM::FUserPtrHandler<FNDITextureInstanceData_GameThread> InstData(Context);
 	FNDIOutputParam<float> OutWidth(Context);
 	FNDIOutputParam<float> OutHeight(Context);
 
 	for (int32 i = 0; i < Context.GetNumInstances(); ++i)
 	{
-		OutWidth.SetAndAdvance(TextureSize.X);
-		OutHeight.SetAndAdvance(TextureSize.Y);
+		OutWidth.SetAndAdvance(InstData->CurrentTextureSize.X);
+		OutHeight.SetAndAdvance(InstData->CurrentTextureSize.Y);
 	}
 }
 
 void UNiagaraDataInterfaceTexture::SampleTexture(FVectorVMExternalFunctionContext& Context)
 {
+	VectorVM::FUserPtrHandler<FNDITextureInstanceData_GameThread> InstData(Context);
 	VectorVM::FExternalFuncInputHandler<float> XParam(Context);
 	VectorVM::FExternalFuncInputHandler<float> YParam(Context);
 	VectorVM::FExternalFuncRegisterHandler<float> OutSampleR(Context);
@@ -272,6 +331,7 @@ void UNiagaraDataInterfaceTexture::SampleTexture(FVectorVMExternalFunctionContex
 void UNiagaraDataInterfaceTexture::SamplePseudoVolumeTexture(FVectorVMExternalFunctionContext& Context)
 {
 	// Noop handler which just returns magenta since this doesn't run on CPU.
+	VectorVM::FUserPtrHandler<FNDITextureInstanceData_GameThread> InstData(Context);
 	VectorVM::FExternalFuncInputHandler<float> UVW_UParam(Context);
 	VectorVM::FExternalFuncInputHandler<float> UVW_VParam(Context);
 	VectorVM::FExternalFuncInputHandler<float> UVW_WParam(Context);
@@ -395,17 +455,10 @@ public:
 
 		FRHIComputeShader* ComputeShaderRHI = Context.Shader.GetComputeShader();
 		FNiagaraDataInterfaceProxyTexture* TextureDI = static_cast<FNiagaraDataInterfaceProxyTexture*>(Context.DataInterface);
-
-		
-		if (TextureDI && TextureDI->ResolvedTextureRHI.IsValid())
+		FNDITextureInstanceData_RenderThread* InstanceData = TextureDI->InstanceData_RT.Find(Context.SystemInstanceID);
+		if ( InstanceData && InstanceData->ResolvedTextureRHI.IsValid() )
 		{
-			FRHISamplerState* SamplerStateRHI = TextureDI->SamplerStateRHI;
-			if (!SamplerStateRHI)
-			{
-				// Fallback required because PostLoad() order affects whether RHI resources 
-				// are initalized in UNiagaraDataInterfaceTexture::PushToRenderThreadImpl().
-				SamplerStateRHI = TStaticSamplerState<SF_Point, AM_Clamp, AM_Clamp, AM_Clamp>::GetRHI();
-			}
+			FRHISamplerState* SamplerStateRHI = InstanceData->SamplerStateRHI ? InstanceData->SamplerStateRHI : GBlackTexture->SamplerStateRHI;
 
 			SetTextureParameter(
 				RHICmdList,
@@ -413,9 +466,9 @@ public:
 				TextureParam,
 				SamplerParam,
 				SamplerStateRHI,
-				TextureDI->ResolvedTextureRHI
+				InstanceData->ResolvedTextureRHI
 			);
-			SetShaderValue(RHICmdList, ComputeShaderRHI, Dimensions, FVector2f(TextureDI->TexDims.X, TextureDI->TexDims.Y));
+			SetShaderValue(RHICmdList, ComputeShaderRHI, Dimensions, InstanceData->TextureSize);
 		}
 		else
 		{
@@ -427,7 +480,6 @@ public:
 				GBlackTexture->SamplerStateRHI,
 				GBlackTexture->TextureRHI
 			);
-			FVector2D TexDims(EForceInit::ForceInitToZero);
 			SetShaderValue(RHICmdList, ComputeShaderRHI, Dimensions, FVector2f::ZeroVector);
 		}
 	}
@@ -442,41 +494,9 @@ IMPLEMENT_TYPE_LAYOUT(FNiagaraDataInterfaceParametersCS_Texture);
 
 IMPLEMENT_NIAGARA_DI_PARAMETER(UNiagaraDataInterfaceTexture, FNiagaraDataInterfaceParametersCS_Texture);
 
-void UNiagaraDataInterfaceTexture::PushToRenderThreadImpl()
-{
-	FNiagaraDataInterfaceProxyTexture* RT_Proxy = GetProxyAs<FNiagaraDataInterfaceProxyTexture>();
-
-	TextureSize = FIntPoint::ZeroValue;
-	if (Texture)
-	{
-		TextureSize.X = Texture->GetSurfaceWidth();
-		TextureSize.Y = Texture->GetSurfaceHeight();
-	}
-
-	ENQUEUE_RENDER_COMMAND(FPushDITextureToRT)
-	(
-		[RT_Proxy, RT_Texture=Texture, RT_TexDims=TextureSize](FRHICommandListImmediate& RHICmdList)
-		{
-			if (RT_Texture && RT_Texture->GetResource() && RT_Texture->TextureReference.TextureReferenceRHI)
-			{
-				RT_Proxy->TextureReferenceRHI = RT_Texture->TextureReference.TextureReferenceRHI;
-				RT_Proxy->SamplerStateRHI = RT_Texture->GetResource()->SamplerStateRHI;
-			}
-			else
-			{
-				RT_Proxy->TextureReferenceRHI = nullptr;
-				RT_Proxy->SamplerStateRHI = nullptr;
-			}
-			RT_Proxy->TexDims = FVector2D(RT_TexDims.X, RT_TexDims.Y);
-		}
-	);
-}
-
 void UNiagaraDataInterfaceTexture::SetTexture(UTexture* InTexture)
 {
 	Texture = InTexture;
-	TextureSize = Texture != nullptr ? FIntPoint(Texture->GetSurfaceWidth(), Texture->GetSurfaceHeight()) : FIntPoint::ZeroValue;
-	MarkRenderDataDirty();
 }
 
-#undef LOCTEXT_NAMESPACE
+#undef LOCTEXT_NAMESPACE
\ No newline at end of file
diff --git a/Engine/Plugins/FX/Niagara/Source/Niagara/Private/NiagaraDataInterfaceVolumeTexture.cpp b/Engine/Plugins/FX/Niagara/Source/Niagara/Private/NiagaraDataInterfaceVolumeTexture.cpp
--- a/Engine/Plugins/FX/Niagara/Source/Niagara/Private/NiagaraDataInterfaceVolumeTexture.cpp
+++ b/Engine/Plugins/FX/Niagara/Source/Niagara/Private/NiagaraDataInterfaceVolumeTexture.cpp
@@ -5,7 +5,9 @@
 #include "ShaderParameterUtils.h"
 #include "NiagaraCustomVersion.h"
 #include "Engine/VolumeTexture.h"
-
+#include "Engine/TextureRenderTargetVolume.h"
+#include "NiagaraSystemInstance.h"
+#include "NiagaraComputeExecutionContext.h"
 
 #define LOCTEXT_NAMESPACE "UNiagaraDataInterfaceVolumeTexture"
 
@@ -15,12 +17,61 @@ const FString UNiagaraDataInterfaceVolumeTexture::TextureName(TEXT("Texture_"));
 const FString UNiagaraDataInterfaceVolumeTexture::SamplerName(TEXT("Sampler_"));
 const FString UNiagaraDataInterfaceVolumeTexture::DimensionsBaseName(TEXT("Dimensions_"));
 
+struct FNDIVolumeTextureInstanceData_GameThread
+{
+	TWeakObjectPtr<UTexture> CurrentTexture = nullptr;
+	FIntVector CurrentTextureSize = FIntVector::ZeroValue;
+	FNiagaraParameterDirectBinding<UObject*> UserParamBinding;
+};
+
+struct FNDIVolumeTextureInstanceData_RenderThread
+{
+	FSamplerStateRHIRef		SamplerStateRHI;
+	FTextureReferenceRHIRef	TextureReferenceRHI;
+	FTextureRHIRef			ResolvedTextureRHI;
+	FVector3f				TextureSize;
+};
+
+struct FNiagaraDataInterfaceProxyVolumeTexture : public FNiagaraDataInterfaceProxy
+{
+	virtual void ConsumePerInstanceDataFromGameThread(void* PerInstanceData, const FNiagaraSystemInstanceID& Instance) override { check(false); }
+	virtual int32 PerInstanceDataPassedToRenderThreadSize() const override { return 0; }
+
+	virtual void PreStage(FRHICommandList& RHICmdList, const FNiagaraDataInterfaceStageArgs& Context) override
+	{
+		if (FNDIVolumeTextureInstanceData_RenderThread* InstanceData = InstanceData_RT.Find(Context.SystemInstanceID))
+		{
+			// Because the underlying reference can have a switch in flight on the RHI we get the referenced texture
+			// here, ensure it's valid (as it could be queued for delete) and cache until next round.  If we were
+			// to release the reference in PostStage / PostSimulate we still stand a chance the the transition we
+			// queue will be invalid by the time it is processed on the RHI thread.
+			if (Context.SimStageData->bFirstStage && InstanceData->TextureReferenceRHI.IsValid())
+			{
+				InstanceData->ResolvedTextureRHI = InstanceData->TextureReferenceRHI->GetReferencedTexture();
+				if (InstanceData->ResolvedTextureRHI && !InstanceData->ResolvedTextureRHI->IsValid())
+				{
+					InstanceData->ResolvedTextureRHI = nullptr;
+				}
+			}
+			if (InstanceData->TextureReferenceRHI.IsValid())
+			{
+				// Make sure the texture is readable, we don't know where it's coming from.
+				RHICmdList.Transition(FRHITransitionInfo(InstanceData->TextureReferenceRHI->GetReferencedTexture(), ERHIAccess::Unknown, ERHIAccess::SRVMask));
+			}
+		}
+	}
+
+	TMap<FNiagaraSystemInstanceID, FNDIVolumeTextureInstanceData_RenderThread> InstanceData_RT;
+};
+
 UNiagaraDataInterfaceVolumeTexture::UNiagaraDataInterfaceVolumeTexture(FObjectInitializer const& ObjectInitializer)
 	: Super(ObjectInitializer)
 	, Texture(nullptr)
 {
 	Proxy.Reset(new FNiagaraDataInterfaceProxyVolumeTexture());
-	MarkRenderDataDirty();
+
+	FNiagaraTypeDefinition Def(UTexture::StaticClass());
+	TextureUserParameter.Parameter.SetType(Def);
 }
 
 void UNiagaraDataInterfaceVolumeTexture::PostInitProperties()
@@ -32,29 +83,8 @@ void UNiagaraDataInterfaceVolumeTexture::PostInitProperties()
 		ENiagaraTypeRegistryFlags Flags = ENiagaraTypeRegistryFlags::AllowAnyVariable | ENiagaraTypeRegistryFlags::AllowParameter;
 		FNiagaraTypeRegistry::Register(FNiagaraTypeDefinition(GetClass()), Flags);
 	}
-
-	MarkRenderDataDirty();
-}
-
-void UNiagaraDataInterfaceVolumeTexture::PostLoad()
-{
-	Super::PostLoad();
-
-	// Not safe since the UTexture might not have yet PostLoad() called and so UpdateResource() called.
-	// This will affect whether the SamplerStateRHI will be available or not.
-	MarkRenderDataDirty();
-}
-
-#if WITH_EDITOR
-
-void UNiagaraDataInterfaceVolumeTexture::PostEditChangeProperty(struct FPropertyChangedEvent& PropertyChangedEvent)
-{
-	Super::PostEditChangeProperty(PropertyChangedEvent);
-	MarkRenderDataDirty();
 }
 
-#endif
-
 bool UNiagaraDataInterfaceVolumeTexture::CopyToInternal(UNiagaraDataInterface* Destination) const
 {
 	if (!Super::CopyToInternal(Destination))
@@ -63,7 +93,7 @@ bool UNiagaraDataInterfaceVolumeTexture::CopyToInternal(UNiagaraDataInterface* D
 	}
 	UNiagaraDataInterfaceVolumeTexture* DestinationTexture = CastChecked<UNiagaraDataInterfaceVolumeTexture>(Destination);
 	DestinationTexture->Texture = Texture;
-	DestinationTexture->MarkRenderDataDirty();
+	DestinationTexture->TextureUserParameter = TextureUserParameter;
 
 	return true;
 }
@@ -75,7 +105,9 @@ bool UNiagaraDataInterfaceVolumeTexture::Equals(const UNiagaraDataInterface* Oth
 		return false;
 	}
 	const UNiagaraDataInterfaceVolumeTexture* OtherTexture = CastChecked<const UNiagaraDataInterfaceVolumeTexture>(Other);
-	return OtherTexture->Texture == Texture;
+	return	
+		OtherTexture->Texture == Texture &&
+		OtherTexture->TextureUserParameter == TextureUserParameter;
 }
 
 void UNiagaraDataInterfaceVolumeTexture::GetFunctions(TArray<FNiagaraFunctionSignature>& OutFunctions)
@@ -116,34 +148,91 @@ void UNiagaraDataInterfaceVolumeTexture::GetVMExternalFunction(const FVMExternal
 {
 	if (BindingInfo.Name == SampleVolumeTextureName)
 	{
-		check(BindingInfo.GetNumInputs() == 4 && BindingInfo.GetNumOutputs() == 4);
+		check(BindingInfo.GetNumInputs() == 5 && BindingInfo.GetNumOutputs() == 4);
 		NDI_FUNC_BINDER(UNiagaraDataInterfaceVolumeTexture, SampleVolumeTexture)::Bind(this, OutFunc);
 	}
 	else if (BindingInfo.Name == TextureDimsName)
 	{
-		check(BindingInfo.GetNumInputs() == 0 && BindingInfo.GetNumOutputs() == 3);
+		check(BindingInfo.GetNumInputs() == 1 && BindingInfo.GetNumOutputs() == 3);
 		OutFunc = FVMExternalFunction::CreateUObject(this, &UNiagaraDataInterfaceVolumeTexture::GetTextureDimensions);
 	}
 }
 
+int32 UNiagaraDataInterfaceVolumeTexture::PerInstanceDataSize() const
+{
+	return sizeof(FNDIVolumeTextureInstanceData_GameThread);
+}
+
+bool UNiagaraDataInterfaceVolumeTexture::InitPerInstanceData(void* PerInstanceData, FNiagaraSystemInstance* SystemInstance)
+{
+	FNDIVolumeTextureInstanceData_GameThread* InstanceData = new (PerInstanceData) FNDIVolumeTextureInstanceData_GameThread();
+	InstanceData->UserParamBinding.Init(SystemInstance->GetInstanceParameters(), TextureUserParameter.Parameter);
+	return true;
+}
+
+void UNiagaraDataInterfaceVolumeTexture::DestroyPerInstanceData(void* PerInstanceData, FNiagaraSystemInstance* SystemInstance)
+{
+	FNDIVolumeTextureInstanceData_GameThread* InstanceData = static_cast<FNDIVolumeTextureInstanceData_GameThread*>(PerInstanceData);
+	InstanceData->~FNDIVolumeTextureInstanceData_GameThread();
+
+	ENQUEUE_RENDER_COMMAND(NDITexture_RemoveInstance)
+	(
+		[RT_Proxy = GetProxyAs<FNiagaraDataInterfaceProxyVolumeTexture>(), RT_InstanceID = SystemInstance->GetId()](FRHICommandListImmediate&)
+		{
+			RT_Proxy->InstanceData_RT.Remove(RT_InstanceID);
+		}
+	);
+}
+
 bool UNiagaraDataInterfaceVolumeTexture::PerInstanceTick(void* PerInstanceData, FNiagaraSystemInstance* SystemInstance, float DeltaSeconds)
 {
-	const FIntVector CurrentTextureSize = Texture != nullptr ? FIntVector(Texture->GetSizeX(), Texture->GetSizeY(), Texture->GetSizeZ()) : FIntVector::ZeroValue;
-	if ( CurrentTextureSize !=  TextureSize )
+	FNDIVolumeTextureInstanceData_GameThread* InstanceData = static_cast<FNDIVolumeTextureInstanceData_GameThread*>(PerInstanceData);
+
+	UTexture* CurrentTexture = InstanceData->UserParamBinding.GetValueOrDefault<UTexture>(Texture);
+	if ( InstanceData->CurrentTexture != CurrentTexture )
 	{
-		TextureSize = CurrentTextureSize;
-		MarkRenderDataDirty();
+		UVolumeTexture* CurrentTextureVolume = Cast<UVolumeTexture>(CurrentTexture);
+		UTextureRenderTargetVolume* CurrentTextureRT = Cast<UTextureRenderTargetVolume>(CurrentTexture);
+		if (CurrentTextureVolume || CurrentTextureRT)
+		{
+			const FIntVector CurrentTextureSize = CurrentTextureVolume != nullptr ?
+				FIntVector(CurrentTextureVolume->GetSizeX(), CurrentTextureVolume->GetSizeY(), CurrentTextureVolume->GetSizeZ()) :
+				FIntVector(CurrentTextureRT->SizeX, CurrentTextureRT->SizeY, CurrentTextureRT->SizeZ);
+
+			InstanceData->CurrentTexture = CurrentTexture;
+			InstanceData->CurrentTextureSize = CurrentTextureSize;
+
+			ENQUEUE_RENDER_COMMAND(NDITexture_UpdateInstance)
+			(
+				[RT_Proxy=GetProxyAs<FNiagaraDataInterfaceProxyVolumeTexture>(), RT_InstanceID=SystemInstance->GetId(), RT_Texture=CurrentTexture, RT_TextureSize=CurrentTextureSize](FRHICommandListImmediate&)
+				{
+					FNDIVolumeTextureInstanceData_RenderThread& InstanceData = RT_Proxy->InstanceData_RT.FindOrAdd(RT_InstanceID);
+					if (RT_Texture)
+					{
+						InstanceData.TextureReferenceRHI = RT_Texture->TextureReference.TextureReferenceRHI;
+						InstanceData.SamplerStateRHI = RT_Texture->GetResource() ? RT_Texture->GetResource()->SamplerStateRHI : nullptr;
+					}
+					else
+					{
+						InstanceData.TextureReferenceRHI = nullptr;
+						InstanceData.SamplerStateRHI = nullptr;
+					}
+					InstanceData.TextureSize = FVector(RT_TextureSize.X, RT_TextureSize.Y, RT_TextureSize.Z);
+				}
+			);
+		}
 	}
 	return false;
 }
 
 void UNiagaraDataInterfaceVolumeTexture::GetTextureDimensions(FVectorVMExternalFunctionContext& Context)
 {
+	VectorVM::FUserPtrHandler<FNDIVolumeTextureInstanceData_GameThread> InstData(Context);
 	FNDIOutputParam<float> OutWidth(Context);
 	FNDIOutputParam<float> OutHeight(Context);
 	FNDIOutputParam<float> OutDepth(Context);
 
-	FVector FloatTextureSize(TextureSize.X, TextureSize.Y, TextureSize.Z);
+	FVector3f FloatTextureSize(InstData->CurrentTextureSize.X, InstData->CurrentTextureSize.Y, InstData->CurrentTextureSize.Z);
 	for (int32 i = 0; i < Context.GetNumInstances(); ++i)
 	{
 		OutWidth.SetAndAdvance(FloatTextureSize.X);
@@ -154,6 +243,7 @@ void UNiagaraDataInterfaceVolumeTexture::GetTextureDimensions(FVectorVMExternalF
 
 void UNiagaraDataInterfaceVolumeTexture::SampleVolumeTexture(FVectorVMExternalFunctionContext& Context)
 {
+	VectorVM::FUserPtrHandler<FNDIVolumeTextureInstanceData_GameThread> InstData(Context);
 	VectorVM::FExternalFuncInputHandler<float> XParam(Context);
 	VectorVM::FExternalFuncInputHandler<float> YParam(Context);
 	VectorVM::FExternalFuncInputHandler<float> ZParam(Context);
@@ -232,25 +322,20 @@ public:
 
 		FRHIComputeShader* ComputeShaderRHI = Context.Shader.GetComputeShader();
 		FNiagaraDataInterfaceProxyVolumeTexture* TextureDI = static_cast<FNiagaraDataInterfaceProxyVolumeTexture*>(Context.DataInterface);
+		FNDIVolumeTextureInstanceData_RenderThread* InstanceData = TextureDI->InstanceData_RT.Find(Context.SystemInstanceID);
 
-		if (TextureDI && TextureDI->TextureRHI)
+		if (InstanceData && InstanceData->ResolvedTextureRHI.IsValid())
 		{
-			FRHISamplerState* SamplerStateRHI = TextureDI->SamplerStateRHI;
-			if (!SamplerStateRHI)
-			{
-				// Fallback required because PostLoad() order affects whether RHI resources 
-				// are initalized in UNiagaraDataInterfaceVolumeTexture::PushToRenderThreadImpl().
-				SamplerStateRHI = TStaticSamplerState<SF_Point, AM_Clamp, AM_Clamp, AM_Clamp>::GetRHI();
-			}
+			FRHISamplerState* SamplerStateRHI = InstanceData->SamplerStateRHI ? InstanceData->SamplerStateRHI : GBlackTexture->SamplerStateRHI;
 			SetTextureParameter(
 				RHICmdList,
 				ComputeShaderRHI,
 				TextureParam,
 				SamplerParam,
 				SamplerStateRHI,
-				TextureDI->TextureRHI
+				InstanceData->ResolvedTextureRHI
 			);
-			SetShaderValue(RHICmdList, ComputeShaderRHI, Dimensions, (FVector3f)TextureDI->TexDims);
+			SetShaderValue(RHICmdList, ComputeShaderRHI, Dimensions, InstanceData->TextureSize);
 		}
 		else
 		{
@@ -274,36 +359,12 @@ private:
 IMPLEMENT_TYPE_LAYOUT(FNiagaraDataInterfaceParametersCS_VolumeTexture);
 IMPLEMENT_NIAGARA_DI_PARAMETER(UNiagaraDataInterfaceVolumeTexture, FNiagaraDataInterfaceParametersCS_VolumeTexture);
 
-void UNiagaraDataInterfaceVolumeTexture::PushToRenderThreadImpl()
-{
-	FNiagaraDataInterfaceProxyVolumeTexture* RT_Proxy = GetProxyAs<FNiagaraDataInterfaceProxyVolumeTexture>();
-
-	TextureSize = FIntVector::ZeroValue;
-	if (Texture)
-	{
-		TextureSize.X = Texture->GetSizeX();
-		TextureSize.Y = Texture->GetSizeY();
-		TextureSize.Z = Texture->GetSizeZ();
-	}
-
-	ENQUEUE_RENDER_COMMAND(FPushDITextureToRT)
-	(
-		[RT_Proxy, RT_Resource=Texture ? Texture->GetResource() : nullptr, RT_TexDims=TextureSize](FRHICommandListImmediate& RHICmdList)
-		{
-			RT_Proxy->TextureRHI = RT_Resource ? RT_Resource->TextureRHI : nullptr;
-			RT_Proxy->SamplerStateRHI = RT_Resource ? RT_Resource->SamplerStateRHI : nullptr;
-			RT_Proxy->TexDims = FVector(RT_TexDims.X, RT_TexDims.Y, RT_TexDims.Z);
-		}
-	);
-}
-
 void UNiagaraDataInterfaceVolumeTexture::SetTexture(UVolumeTexture* InTexture)
 {
 	if (InTexture)
 	{
 		Texture = InTexture;
-		MarkRenderDataDirty();
 	}
 }
 
-#undef LOCTEXT_NAMESPACE
\ No newline at end of file
+#undef LOCTEXT_NAMESPACE
diff --git a/Engine/Plugins/FX/Niagara/Source/Niagara/Private/NiagaraFunctionLibrary.cpp b/Engine/Plugins/FX/Niagara/Source/Niagara/Private/NiagaraFunctionLibrary.cpp
--- a/Engine/Plugins/FX/Niagara/Source/Niagara/Private/NiagaraFunctionLibrary.cpp
+++ b/Engine/Plugins/FX/Niagara/Source/Niagara/Private/NiagaraFunctionLibrary.cpp
@@ -460,9 +460,6 @@ void UNiagaraFunctionLibrary::SetTextureObject(UNiagaraComponent* NiagaraSystem,
 		}
 
 		TextureDI->SetTexture(Texture);
-#if WITH_EDITOR
-		NiagaraSystem->SetParameterOverride(Variable, FNiagaraVariant(TextureDI));
-#endif
 	}
 	else if (UTexture2DArray* Texture2DArray = Cast<UTexture2DArray>(Texture))
 	{
@@ -482,9 +479,6 @@ void UNiagaraFunctionLibrary::SetTextureObject(UNiagaraComponent* NiagaraSystem,
 		}
 
 		TextureDI->SetTexture(Texture2DArray);
-#if WITH_EDITOR
-		NiagaraSystem->SetParameterOverride(Variable, FNiagaraVariant(TextureDI));
-#endif
 	}
 	else if (UVolumeTexture* TextureVolume = Cast<UVolumeTexture>(Texture))
 	{
@@ -504,9 +498,6 @@ void UNiagaraFunctionLibrary::SetTextureObject(UNiagaraComponent* NiagaraSystem,
 		}
 
 		TextureDI->SetTexture(TextureVolume);
-#if WITH_EDITOR
-		NiagaraSystem->SetParameterOverride(Variable, FNiagaraVariant(TextureDI));
-#endif
 	}
 	else if (UTextureCube* TextureCube = Cast<UTextureCube>(Texture))
 	{
@@ -526,9 +517,6 @@ void UNiagaraFunctionLibrary::SetTextureObject(UNiagaraComponent* NiagaraSystem,
 		}
 
 		TextureDI->SetTexture(TextureCube);
-#if WITH_EDITOR
-		NiagaraSystem->SetParameterOverride(Variable, FNiagaraVariant(TextureDI));
-#endif
 	}
 	else
 	{
diff --git a/Engine/Plugins/FX/Niagara/Source/Niagara/Public/NiagaraParameterStore.h b/Engine/Plugins/FX/Niagara/Source/Niagara/Public/NiagaraParameterStore.h
--- a/Engine/Plugins/FX/Niagara/Source/Niagara/Public/NiagaraParameterStore.h
+++ b/Engine/Plugins/FX/Niagara/Source/Niagara/Public/NiagaraParameterStore.h
@@ -1328,7 +1328,7 @@ struct FNiagaraParameterDirectBinding<UObject*>
 	}
 
 	template<class TObjectType>
-	FORCEINLINE TObjectType* GetValue() const
+	FORCEINLINE TObjectType* GetValueOrDefault(TObjectType* DefaultValue) const
 	{
 		if (UObjectOffset != INDEX_NONE)
 		{
@@ -1337,6 +1337,12 @@ struct FNiagaraParameterDirectBinding<UObject*>
 
 			return Cast<TObjectType>(BoundStore->GetUObject(UObjectOffset));
 		}
-		return nullptr;
+		return DefaultValue;
+	}
+
+	template<class TObjectType>
+	FORCEINLINE TObjectType* GetValue() const
+	{
+		return GetValueOrDefault<TObjectType>(nullptr);
 	}
 };
