- Fix missing check for nullptr if the emitter doesn't have a system owner

#rnx

diff --git a/Engine/Plugins/FX/Niagara/Source/Niagara/Private/NiagaraScript.cpp b/Engine/Plugins/FX/Niagara/Source/Niagara/Private/NiagaraScript.cpp
--- a/Engine/Plugins/FX/Niagara/Source/Niagara/Private/NiagaraScript.cpp
+++ b/Engine/Plugins/FX/Niagara/Source/Niagara/Private/NiagaraScript.cpp
@@ -1439,10 +1439,12 @@ void UNiagaraScript::Serialize(FArchive& Ar)
 #if WITH_EDITORONLY_DATA
 		if (UNiagaraEmitter* Emitter = GetTypedOuter<UNiagaraEmitter>())
 		{
-			UNiagaraSystem* EmitterOwner = Emitter->GetTypedOuter<UNiagaraSystem>();
-			if (EmitterOwner->bBakeOutRapidIteration)
+			if ( UNiagaraSystem* System = Emitter->GetTypedOuter<UNiagaraSystem>() )
 			{
-				bUsesRapidIterationParams = false;
+				if (System->bBakeOutRapidIteration)
+				{
+					bUsesRapidIterationParams = false;
+				}
 			}
 		}
 		else if (UNiagaraSystem* System = GetTypedOuter<UNiagaraSystem>())
