- Fix static analysis warning

#rnx

diff --git a/Engine/Plugins/FX/Niagara/Source/Niagara/Private/NiagaraDataInterfaceSkeletalMesh.cpp b/Engine/Plugins/FX/Niagara/Source/Niagara/Private/NiagaraDataInterfaceSkeletalMesh.cpp
--- a/Engine/Plugins/FX/Niagara/Source/Niagara/Private/NiagaraDataInterfaceSkeletalMesh.cpp
+++ b/Engine/Plugins/FX/Niagara/Source/Niagara/Private/NiagaraDataInterfaceSkeletalMesh.cpp
@@ -2747,7 +2747,7 @@ void UNiagaraDataInterfaceSkeletalMesh::GetFeedback(UNiagaraSystem* Asset, UNiag
 	}
 
 	// Report Errors/Warnings
-	if (bHasCPUAccessWarning)
+	if (SkelMesh && bHasCPUAccessWarning)
 	{
 		FNiagaraDataInterfaceFeedback CPUAccessNotAllowedWarning(
 			FText::Format(LOCTEXT("CPUAccessNotAllowedError", "This mesh may need CPU access in order to be used properly (even when used by GPU emitters). ({0})"), FText::FromString(SkelMesh->GetName())),
