Fixed a shader compilation problem when using lwc tiles in non-lwc builds

diff --git a/Engine/Plugins/FX/Niagara/Shaders/Private/NiagaraDataInterfaceCamera.ush b/Engine/Plugins/FX/Niagara/Shaders/Private/NiagaraDataInterfaceCamera.ush
--- a/Engine/Plugins/FX/Niagara/Shaders/Private/NiagaraDataInterfaceCamera.ush
+++ b/Engine/Plugins/FX/Niagara/Shaders/Private/NiagaraDataInterfaceCamera.ush
@@ -10,7 +10,7 @@ struct FNDICameraContext
 void DICamera_GetCameraProperties(in FNDICameraContext DIContext, out float3 Out_CameraPositionWorld, out float3 Out_ViewForwardVector, out float3 Out_ViewUpVector, out float3 Out_ViewRightVector)
 {
 	FLWCVector3 CameraPos = PrimaryView.WorldCameraOrigin;
-	CameraPos.Tile -= DIContext.SystemLWCTile; // convert to simulation space
+	LWCSetTile(CameraPos, LWCGetTile(CameraPos) - DIContext.SystemLWCTile); // convert to simulation space
 	Out_CameraPositionWorld = LWCToFloat(CameraPos);
 	Out_ViewForwardVector = View.ViewForward;
 	Out_ViewUpVector = View.ViewUp;
@@ -21,7 +21,7 @@ void DICamera_GetViewProperties(in FNDICameraContext DIContext, out float3 Out_V
 	out float4 Out_ScreenToViewSpace, out float2 Out_Current_TAAJitter, out float2 Out_Previous_TAAJitter, out float3 Out_PreViewTranslation, out float4 Out_BufferSizeAndInverseSize, out float2 Out_ViewportOffset)
 {
 	FLWCVector3 ViewPos = PrimaryView.WorldViewOrigin;
-	ViewPos.Tile -= DIContext.SystemLWCTile; // convert to simulation space
+	LWCSetTile(ViewPos, LWCGetTile(ViewPos) - DIContext.SystemLWCTile); // convert to simulation space
 	Out_ViewPositionWorld.xyz = LWCToFloat(ViewPos);
 	Out_ViewForwardVector = View.ViewForward;
 	Out_ViewUpVector = View.ViewUp;
diff --git a/Engine/Plugins/FX/Niagara/Shaders/Private/NiagaraDataInterfaceCollisionQuery.ush b/Engine/Plugins/FX/Niagara/Shaders/Private/NiagaraDataInterfaceCollisionQuery.ush
--- a/Engine/Plugins/FX/Niagara/Shaders/Private/NiagaraDataInterfaceCollisionQuery.ush
+++ b/Engine/Plugins/FX/Niagara/Shaders/Private/NiagaraDataInterfaceCollisionQuery.ush
@@ -22,7 +22,7 @@ void NDICollisionQuery_QuerySceneDepthGPU(in float3 In_SamplePos, in float3 In_L
 	Out_WorldNormal = float3(0.0, 0.0, 1.0);
 	Out_IsInsideView = true;
 	FLWCVector3 CameraPos = PrimaryView.WorldCameraOrigin;
-	CameraPos.Tile -= In_LWCTile; // convert to simulation space
+	LWCSetTile(CameraPos, LWCGetTile(CameraPos) - In_LWCTile); // convert to simulation space
 	Out_CameraPosWorld = LWCToFloat(CameraPos);
 
 	FLWCVector3 LwcSamplePos = MakeLWCVector3(In_LWCTile, In_SamplePos);
@@ -59,7 +59,7 @@ void NDICollisionQuery_QuerySceneDepthGPU(in float3 In_SamplePos, in float3 In_L
 #endif
 
 		// convert LWC position back to simulation space
-		WorldPosition.Tile -= In_LWCTile;
+		LWCSetTile(WorldPosition, LWCGetTile(WorldPosition) - In_LWCTile);
 		
 		// Set outputs
 		Out_SceneDepth = SceneDepth;
@@ -79,7 +79,7 @@ void NDICollisionQuery_QueryCustomDepthGPU(in float3 In_SamplePos, in float3 In_
 	Out_WorldNormal = float3(0.0, 0.0, 1.0);
 	Out_IsInsideView = true;
 	FLWCVector3 CameraPos = PrimaryView.WorldCameraOrigin;
-	CameraPos.Tile -= In_LWCTile; // convert to simulation space
+	LWCSetTile(CameraPos, LWCGetTile(CameraPos) - In_LWCTile); // convert to simulation space
 	Out_CameraPosWorld = LWCToFloat(CameraPos);
 
 	FLWCVector3 LwcSamplePos = MakeLWCVector3(In_LWCTile, In_SamplePos);
@@ -116,7 +116,7 @@ void NDICollisionQuery_QueryCustomDepthGPU(in float3 In_SamplePos, in float3 In_
 #endif
 
 		// convert LWC position back to simulation space
-		WorldPosition.Tile = WorldPosition.Tile - In_LWCTile;
+		LWCSetTile(WorldPosition, LWCGetTile(WorldPosition) - In_LWCTile);
 		
 		// Set outputs
 		Out_SceneDepth = SceneDepth;
