- Include layout and UB changes as part of hash generation

diff --git a/Engine/Plugins/FX/Niagara/Source/NiagaraShader/Private/NiagaraShader.cpp b/Engine/Plugins/FX/Niagara/Source/NiagaraShader/Private/NiagaraShader.cpp
--- a/Engine/Plugins/FX/Niagara/Source/NiagaraShader/Private/NiagaraShader.cpp
+++ b/Engine/Plugins/FX/Niagara/Source/NiagaraShader/Private/NiagaraShader.cpp
@@ -343,25 +343,41 @@ void FNiagaraShaderMapId::AppendKeyString(FString& KeyString) const
 		}
 	}
 
-	for (const FShaderTypeDependency& Dependency : ShaderTypeDependencies)
+	TSortedMap<const TCHAR*, FCachedUniformBufferDeclaration, FDefaultAllocator, FUniformBufferNameSortOrder> ReferencedUniformBuffers;
+	for (const FShaderTypeDependency& ShaderTypeDependency : ShaderTypeDependencies)
 	{
+		const FShaderType* ShaderType = FindShaderTypeByName(ShaderTypeDependency.ShaderTypeName);
+
 		KeyString += TEXT("_");
-		KeyString += Dependency.SourceHash.ToString();
-	}
+		KeyString += ShaderTypeDependency.SourceHash.ToString();
 
-	/*
-	ParameterSet.AppendKeyString(KeyString);
+		if (const FShaderParametersMetadata* ParameterStructMetadata = ShaderType->GetRootParametersMetadata())
+		{
+			KeyString += FString::Printf(TEXT("%08x"), ParameterStructMetadata->GetLayoutHash());
+		}
 
-	KeyString += TEXT("_");
-	KeyString += FString::FromInt(Usage);
-	KeyString += TEXT("_");
+		// Add the serialization history to the key string so that we can detect changes to shader serialization without a corresponding .usf change
+		const FSHAHash LayoutHash = Freeze::HashLayout(ShaderType->GetLayout(), LayoutParams);
+		KeyString += LayoutHash.ToString();
+
+		const TMap<const TCHAR*, FCachedUniformBufferDeclaration>& ReferencedUniformBufferStructsCache = ShaderType->GetReferencedUniformBufferStructsCache();
+		for (TMap<const TCHAR*, FCachedUniformBufferDeclaration>::TConstIterator It(ReferencedUniformBufferStructsCache); It; ++It)
+		{
+			ReferencedUniformBuffers.Add(It.Key(), It.Value());
+		}
+	}
 
-	// Add any referenced functions to the key so that we will recompile when they are changed
-	for (int32 FunctionIndex = 0; FunctionIndex < ReferencedFunctions.Num(); FunctionIndex++)
 	{
-		KeyString += ReferencedFunctions[FunctionIndex].ToString();
+		TArray<uint8> TempData;
+		FSerializationHistory SerializationHistory;
+		FMemoryWriter Ar(TempData, true);
+		FShaderSaveArchive SaveArchive(Ar, SerializationHistory);
+
+		// Save uniform buffer member info so we can detect when layout has changed
+		SerializeUniformBufferInfo(SaveArchive, ReferencedUniformBuffers);
+
+		SerializationHistory.AppendKeyString(KeyString);
 	}
-	*/
 }
 
 
