- Disable Niagara's auto flushing as it's handled inside the RHI now
- Ensure that draw indirect buffers are updated on RDG timeline

#rnx

diff --git a/Engine/Plugins/FX/Niagara/Source/Niagara/Private/NiagaraGpuComputeDispatch.cpp b/Engine/Plugins/FX/Niagara/Source/Niagara/Private/NiagaraGpuComputeDispatch.cpp
--- a/Engine/Plugins/FX/Niagara/Source/Niagara/Private/NiagaraGpuComputeDispatch.cpp
+++ b/Engine/Plugins/FX/Niagara/Source/Niagara/Private/NiagaraGpuComputeDispatch.cpp
@@ -57,7 +57,7 @@ uint32 FNiagaraComputeExecutionContext::TickCounter = 0;
 const FName FNiagaraGpuComputeDispatch::TemporalEffectName("FNiagaraGpuComputeDispatch");
 #endif // WITH_MGPU
 
-int32 GNiagaraGpuSubmitCommandHint = WITH_EDITOR ? 10 : 0;
+int32 GNiagaraGpuSubmitCommandHint = 0;
 static FAutoConsoleVariableRef CVarNiagaraGpuSubmitCommandHint(
 	TEXT("fx.NiagaraGpuSubmitCommandHint"),
 	GNiagaraGpuSubmitCommandHint,
@@ -485,10 +485,20 @@ void FNiagaraGpuComputeDispatch::ProcessPendingTicksFlush(FRHICommandListImmedia
 			FRDGBuilder GraphBuilder(RHICmdList);
 			CreateSystemTextures(GraphBuilder);
 			PreInitViews(GraphBuilder, bAllowGPUParticleUpdate);
-			GPUInstanceCounterManager.UpdateDrawIndirectBuffers(this, RHICmdList, ENiagaraGPUCountUpdatePhase::PreOpaque);
+			AddPass(GraphBuilder, RDG_EVENT_NAME("UpdateDrawIndirectBuffers - PreOpaque"),
+				[this](FRHICommandListImmediate& RHICmdList)
+				{
+					GPUInstanceCounterManager.UpdateDrawIndirectBuffers(this, RHICmdList, ENiagaraGPUCountUpdatePhase::PreOpaque);
+				}
+			);
 			PostInitViews(GraphBuilder, DummyViews, bAllowGPUParticleUpdate);
 			PostRenderOpaque(GraphBuilder, DummyViews, bAllowGPUParticleUpdate);
-			GPUInstanceCounterManager.UpdateDrawIndirectBuffers(this, RHICmdList, ENiagaraGPUCountUpdatePhase::PostOpaque);
+			AddPass(GraphBuilder, RDG_EVENT_NAME("UpdateDrawIndirectBuffers - PostOpaque"),
+				[this](FRHICommandListImmediate& RHICmdList)
+				{
+					GPUInstanceCounterManager.UpdateDrawIndirectBuffers(this, RHICmdList, ENiagaraGPUCountUpdatePhase::PostOpaque);
+				}
+			);
 			GraphBuilder.Execute();
 			break;
 		}
