Fixed postload order problem with niagara systems

diff --git a/Engine/Plugins/FX/Niagara/Source/Niagara/Private/NiagaraSystem.cpp b/Engine/Plugins/FX/Niagara/Source/Niagara/Private/NiagaraSystem.cpp
--- a/Engine/Plugins/FX/Niagara/Source/Niagara/Private/NiagaraSystem.cpp
+++ b/Engine/Plugins/FX/Niagara/Source/Niagara/Private/NiagaraSystem.cpp
@@ -3651,6 +3651,7 @@ void UNiagaraSystem::FixupPositionUserParameters()
 	TSet<FNiagaraVariable> LinkedPositionInputs;
 	ForEachScript([&UserParameters, &LinkedPositionInputs](UNiagaraScript* Script)
 	{
+		Script->ConditionalPostLoad();
 		Script->GetLatestSource()->GetLinkedPositionTypeInputs(UserParameters, LinkedPositionInputs);
 	});
 
