Changed displayed hlsl text to a monospace font

diff --git a/Engine/Plugins/FX/Niagara/Source/NiagaraEditor/Private/NiagaraEditorStyle.cpp b/Engine/Plugins/FX/Niagara/Source/NiagaraEditor/Private/NiagaraEditorStyle.cpp
--- a/Engine/Plugins/FX/Niagara/Source/NiagaraEditor/Private/NiagaraEditorStyle.cpp
+++ b/Engine/Plugins/FX/Niagara/Source/NiagaraEditor/Private/NiagaraEditorStyle.cpp
@@ -310,7 +310,7 @@ void InitCodeView(TSharedRef< FSlateStyleSet > Style)
 		.SetShadowOffset(FVector2D(1, 1))
 		.SetShadowColorAndOpacity(FLinearColor(0, 0, 0, 0.9f)));
 
-	const int32 LogFontSize = 9;
+	constexpr int32 LogFontSize = 9;
 	FSlateFontInfo LogFont = DEFAULT_FONT("Mono", LogFontSize);
 	FTextBlockStyle NormalLogText = FTextBlockStyle(NormalText)
 		.SetFont(LogFont)
@@ -325,15 +325,23 @@ void InitCodeView(TSharedRef< FSlateStyleSet > Style)
 		
 	Style->Set("TextEditor.NormalText", NormalText);
 
-	Style->Set("SyntaxHighlight.HLSL.Normal", FTextBlockStyle(NormalText).SetColorAndOpacity(FLinearColor(FColor(189, 183, 107))));
-	Style->Set("SyntaxHighlight.HLSL.Operator", FTextBlockStyle(NormalText).SetColorAndOpacity(FLinearColor(FColor(220, 220, 220))));
-	Style->Set("SyntaxHighlight.HLSL.Keyword", FTextBlockStyle(NormalText).SetColorAndOpacity(FLinearColor(FColor(86, 156, 214))));
-	Style->Set("SyntaxHighlight.HLSL.String", FTextBlockStyle(NormalText).SetColorAndOpacity(FLinearColor(FColor(214, 157, 133))));
-	Style->Set("SyntaxHighlight.HLSL.Number", FTextBlockStyle(NormalText).SetColorAndOpacity(FLinearColor(FColor(181, 206, 168))));
-	Style->Set("SyntaxHighlight.HLSL.Comment", FTextBlockStyle(NormalText).SetColorAndOpacity(FLinearColor(FColor(87, 166, 74))));
-	Style->Set("SyntaxHighlight.HLSL.PreProcessorKeyword", FTextBlockStyle(NormalText).SetColorAndOpacity(FLinearColor(FColor(188, 98, 171))));
-
-	Style->Set("SyntaxHighlight.HLSL.Error", ErrorText); 
+	constexpr int32 HlslFontSize = 9;
+	FSlateFontInfo HlslFont = DEFAULT_FONT("Mono", HlslFontSize);
+	FTextBlockStyle NormalHlslText = FTextBlockStyle(NormalText)
+		.SetFont(HlslFont);
+	const FTextBlockStyle HlslErrorText = FTextBlockStyle(NormalHlslText)
+		.SetUnderlineBrush(IMAGE_BRUSH("White", Icon8x8, FLinearColor::Red, ESlateBrushTileType::Both))
+		.SetColorAndOpacity(FLinearColor::Red);
+	
+	Style->Set("SyntaxHighlight.HLSL.Normal", FTextBlockStyle(NormalHlslText).SetColorAndOpacity(FLinearColor(FColor(189, 183, 107))));
+	Style->Set("SyntaxHighlight.HLSL.Operator", FTextBlockStyle(NormalHlslText).SetColorAndOpacity(FLinearColor(FColor(220, 220, 220))));
+	Style->Set("SyntaxHighlight.HLSL.Keyword", FTextBlockStyle(NormalHlslText).SetColorAndOpacity(FLinearColor(FColor(86, 156, 214))));
+	Style->Set("SyntaxHighlight.HLSL.String", FTextBlockStyle(NormalHlslText).SetColorAndOpacity(FLinearColor(FColor(214, 157, 133))));
+	Style->Set("SyntaxHighlight.HLSL.Number", FTextBlockStyle(NormalHlslText).SetColorAndOpacity(FLinearColor(FColor(181, 206, 168))));
+	Style->Set("SyntaxHighlight.HLSL.Comment", FTextBlockStyle(NormalHlslText).SetColorAndOpacity(FLinearColor(FColor(87, 166, 74))));
+	Style->Set("SyntaxHighlight.HLSL.PreProcessorKeyword", FTextBlockStyle(NormalHlslText).SetColorAndOpacity(FLinearColor(FColor(188, 98, 171))));
+
+	Style->Set("SyntaxHighlight.HLSL.Error", HlslErrorText); 
 		
 }
 
@@ -369,8 +377,8 @@ void InitToolbarIcons(TSharedRef< FSlateStyleSet > Style)
 	Style->Set("Niagara.Asset.ReimportAsset.Needed", new IMAGE_BRUSH("Icons/icon_Reimport_Needed_40x", Icon40x40));
 	Style->Set("Niagara.Asset.ReimportAsset.Default", new IMAGE_BRUSH("Icons/icon_Reimport_40x", Icon40x40));
 	
-	Style->Set("Niagaraeditor.OverviewNode.IsolatedColor", FLinearColor::Yellow);
-	Style->Set("Niagaraeditor.OverviewNode.NotIsolatedColor", FLinearColor::Transparent);
+	Style->Set("NiagaraEditor.OverviewNode.IsolatedColor", FLinearColor::Yellow);
+	Style->Set("NiagaraEditor.OverviewNode.NotIsolatedColor", FLinearColor::Transparent);
 }
 
 void InitIcons(TSharedRef< FSlateStyleSet > Style)
