Fixed a crash when postloading a system with invalid script data

diff --git a/Engine/Plugins/FX/Niagara/Source/Niagara/Private/NiagaraSystem.cpp b/Engine/Plugins/FX/Niagara/Source/Niagara/Private/NiagaraSystem.cpp
--- a/Engine/Plugins/FX/Niagara/Source/Niagara/Private/NiagaraSystem.cpp
+++ b/Engine/Plugins/FX/Niagara/Source/Niagara/Private/NiagaraSystem.cpp
@@ -3651,8 +3651,11 @@ void UNiagaraSystem::FixupPositionUserParameters()
 	TSet<FNiagaraVariable> LinkedPositionInputs;
 	ForEachScript([&UserParameters, &LinkedPositionInputs](UNiagaraScript* Script)
 	{
-		Script->ConditionalPostLoad();
-		Script->GetLatestSource()->GetLinkedPositionTypeInputs(UserParameters, LinkedPositionInputs);
+		if (Script)
+		{
+			Script->ConditionalPostLoad();
+			Script->GetLatestSource()->GetLinkedPositionTypeInputs(UserParameters, LinkedPositionInputs);
+		}
 	});
 
 	// looks like we have a few FVector3f user parameters that are linked to position inputs in the stack.
