Fix memory stomp on loading UNiagaraGraph

diff --git a/Engine/Plugins/FX/Niagara/Source/NiagaraEditor/Private/NiagaraGraph.cpp b/Engine/Plugins/FX/Niagara/Source/NiagaraEditor/Private/NiagaraGraph.cpp
--- a/Engine/Plugins/FX/Niagara/Source/NiagaraEditor/Private/NiagaraGraph.cpp
+++ b/Engine/Plugins/FX/Niagara/Source/NiagaraEditor/Private/NiagaraGraph.cpp
@@ -367,7 +367,7 @@ void UNiagaraGraph::PostLoad()
 			VariableToScriptVariable.Remove(OldVarType);
 
 			FNiagaraVariable NewVarType(FNiagaraTypeDefinition::GetPositionDef(), OldVarType.GetName());
-			VariableToScriptVariable.Add(NewVarType, *ScriptVariablePtr);
+			VariableToScriptVariable.Add(NewVarType, TObjectPtr<UNiagaraScriptVariable>(&ScriptVar));
 			ScriptVariableChanged(NewVarType);
 		}
 
@@ -378,7 +378,8 @@ void UNiagaraGraph::PostLoad()
 			if (ReferenceCollection->ParameterReferences.Num() > 0)
 			{
 				FNiagaraVariable NewVarType(FNiagaraTypeDefinition::GetPositionDef(), OldVarType.GetName());
-				ParameterToReferencesMap.Add(NewVarType, *ReferenceCollection);
+				FNiagaraGraphParameterReferenceCollection NewReferenceCollection = *ReferenceCollection;
+				ParameterToReferencesMap.Add(NewVarType, NewReferenceCollection);
 			}
 			ParameterToReferencesMap.Remove(OldVarType);
 		}
