Optimize cook time and fix indeterminism in niagara curve generated textures by using content based guid

diff --git a/Engine/Plugins/FX/Niagara/Source/Niagara/Private/NiagaraDataInterfaceCurveBase.cpp b/Engine/Plugins/FX/Niagara/Source/Niagara/Private/NiagaraDataInterfaceCurveBase.cpp
--- a/Engine/Plugins/FX/Niagara/Source/Niagara/Private/NiagaraDataInterfaceCurveBase.cpp
+++ b/Engine/Plugins/FX/Niagara/Source/Niagara/Private/NiagaraDataInterfaceCurveBase.cpp
@@ -289,7 +289,7 @@ void UNiagaraDataInterfaceCurveBase::UpdateExposedTexture()
 	}
 
 	ExposedTexture->Source.UnlockMip(0);
-
+	ExposedTexture->Source.UseHashAsGuid();
 	ExposedTexture->PostEditChange();
 
 	// PostEditChange() will assign a random GUID to the texture, which leads to non-deterministic builds.
