- Destroy components that fail to activate on servers that are auto destroy

#rnx

diff --git a/Engine/Plugins/FX/Niagara/Source/Niagara/Private/NiagaraComponent.cpp b/Engine/Plugins/FX/Niagara/Source/Niagara/Private/NiagaraComponent.cpp
--- a/Engine/Plugins/FX/Niagara/Source/Niagara/Private/NiagaraComponent.cpp
+++ b/Engine/Plugins/FX/Niagara/Source/Niagara/Private/NiagaraComponent.cpp
@@ -452,6 +452,7 @@ UNiagaraComponent::UNiagaraComponent(const FObjectInitializer& ObjectInitializer
 	PrimaryComponentTick.TickGroup = GNiagaraSoloTickEarly ? TG_PrePhysics : TG_DuringPhysics;
 	PrimaryComponentTick.EndTickGroup = GNiagaraSoloAllowAsyncWorkToEndOfFrame ? TG_LastDemotable : ETickingGroup(PrimaryComponentTick.TickGroup);
 	PrimaryComponentTick.bStartWithTickEnabled = false;
+	PrimaryComponentTick.bAllowTickOnDedicatedServer = false;
 	PrimaryComponentTick.SetTickFunctionEnable(false);
 	bTickInEditor = true;
 	bAutoActivate = true;
@@ -981,6 +982,10 @@ void UNiagaraComponent::ActivateInternal(bool bReset /* = false */, bool bIsScal
 	// If the particle system can't ever render (ie on dedicated server or in a commandlet) than do not activate...
 	if (!FApp::CanEverRender() || !World || World->IsNetMode(NM_DedicatedServer) || World->bIsTearingDown)
 	{
+		if ( bAutoDestroy )
+		{
+			DestroyComponent();
+		}
 		return;
 	}
 
