Some code cleanup, mainly in Niagara data interfaces

diff --git a/Engine/Plugins/FX/Niagara/Source/Niagara/Classes/NiagaraDataInterface.h b/Engine/Plugins/FX/Niagara/Source/Niagara/Classes/NiagaraDataInterface.h
--- a/Engine/Plugins/FX/Niagara/Source/Niagara/Classes/NiagaraDataInterface.h
+++ b/Engine/Plugins/FX/Niagara/Source/Niagara/Classes/NiagaraDataInterface.h
@@ -598,7 +598,7 @@ struct FNDIRandomInfoPolicy
 	FORCEINLINE bool IsDeterministic() const { return RandInfo.Seed3 != INDEX_NONE; }
 
 	//////////////////////////////////////////////////////////////////////////
-	FORCEINLINE_DEBUGGABLE FVector4 Rand4(int32 InstanceIndex) const
+	FORCEINLINE_DEBUGGABLE FVector4f Rand4(int32 InstanceIndex) const
 	{
 		if (IsDeterministic())
 		{
@@ -620,16 +620,16 @@ struct FNDIRandomInfoPolicy
 			//           We use the upper 24 bits as they tend to be higher quality.
 
 			// NOTE(mv): The divide can often be folded with the range scale in the rand functions
-			return FVector4((v >> 8) & 0x00ffffff) / 16777216.0; // 0x01000000 == 16777216
+			return FVector4f((v >> 8) & 0x00ffffff) / 16777216.0; // 0x01000000 == 16777216
 			// return float4((v >> 8) & 0x00ffffff) * (1.0/16777216.0); // bugged, see UE-67738
 		}
 		else
 		{
-			return FVector4(Context.GetRandStream().GetFraction(), Context.GetRandStream().GetFraction(), Context.GetRandStream().GetFraction(), Context.GetRandStream().GetFraction());
+			return FVector4f(Context.GetRandStream().GetFraction(), Context.GetRandStream().GetFraction(), Context.GetRandStream().GetFraction(), Context.GetRandStream().GetFraction());
 		}
 	}
 
-	FORCEINLINE_DEBUGGABLE FVector Rand3(int32 InstanceIndex) const
+	FORCEINLINE_DEBUGGABLE FVector3f Rand3(int32 InstanceIndex) const
 	{
 		if (IsDeterministic())
 		{
@@ -644,24 +644,24 @@ struct FNDIRandomInfoPolicy
 			v.Y += v.Z * v.X;
 			v.Z += v.X * v.Y;
 
-			return FVector((v >> 8) & 0x00ffffff) / 16777216.0; // 0x01000000 == 16777216
+			return FVector3f((v >> 8) & 0x00ffffff) / 16777216.0; // 0x01000000 == 16777216
 		}
 		else
 		{
-			return FVector(Context.GetRandStream().GetFraction(), Context.GetRandStream().GetFraction(), Context.GetRandStream().GetFraction());
+			return FVector3f(Context.GetRandStream().GetFraction(), Context.GetRandStream().GetFraction(), Context.GetRandStream().GetFraction());
 		}
 	}
 
-	FORCEINLINE_DEBUGGABLE FVector2D Rand2(int32 InstanceIndex) const
+	FORCEINLINE_DEBUGGABLE FVector2f Rand2(int32 InstanceIndex) const
 	{
 		if (IsDeterministic())
 		{
 			FVector Rand3D = Rand3(InstanceIndex);
-			return FVector2D(Rand3D.X, Rand3D.Y);
+			return FVector2f(Rand3D.X, Rand3D.Y);
 		}
 		else
 		{
-			return FVector2D(Context.GetRandStream().GetFraction(), Context.GetRandStream().GetFraction());
+			return FVector2f(Context.GetRandStream().GetFraction(), Context.GetRandStream().GetFraction());
 		}
 	}
 
@@ -723,14 +723,13 @@ struct TNDIRandomHelper : public TRandomPolicy
 		return Min + (int(TRandomPolicy::Rand(InstanceIndex) * (Range + 1)));
 	}
 
-	FORCEINLINE_DEBUGGABLE FVector RandomBarycentricCoord(int32 InstanceIndex) const
+	FORCEINLINE_DEBUGGABLE FVector3f RandomBarycentricCoord(int32 InstanceIndex) const
 	{
 		//TODO: This is gonna be slooooow. Move to an LUT possibly or find faster method.
 		//Can probably handle lower quality randoms / uniformity for a decent speed win.
-		FVector2D r = TRandomPolicy::Rand2(InstanceIndex);
+		FVector2f r = TRandomPolicy::Rand2(InstanceIndex);
 		float sqrt0 = FMath::Sqrt(r.X);
-		float sqrt1 = FMath::Sqrt(r.Y);
-		return FVector(1.0f - sqrt0, sqrt0 * (1.0 - r.Y), r.Y * sqrt0);
+		return FVector3f(1.0f - sqrt0, sqrt0 * (1.0 - r.Y), r.Y * sqrt0);
 	}
 };
 
diff --git a/Engine/Plugins/FX/Niagara/Source/Niagara/Classes/NiagaraDataInterfaceSkeletalMesh.h b/Engine/Plugins/FX/Niagara/Source/Niagara/Classes/NiagaraDataInterfaceSkeletalMesh.h
--- a/Engine/Plugins/FX/Niagara/Source/Niagara/Classes/NiagaraDataInterfaceSkeletalMesh.h
+++ b/Engine/Plugins/FX/Niagara/Source/Niagara/Classes/NiagaraDataInterfaceSkeletalMesh.h
@@ -54,11 +54,11 @@ struct FSkeletalMeshSkinningDataHandle
 	FSkeletalMeshSkinningDataHandle();
 	FSkeletalMeshSkinningDataHandle(FSkeletalMeshSkinningDataUsage InUsage, const TSharedPtr<struct FSkeletalMeshSkinningData>& InSkinningData, bool bNeedsDataImmediately);
 	FSkeletalMeshSkinningDataHandle(const FSkeletalMeshSkinningDataHandle& Other) = delete;
-	FSkeletalMeshSkinningDataHandle(FSkeletalMeshSkinningDataHandle&& Other);
+	FSkeletalMeshSkinningDataHandle(FSkeletalMeshSkinningDataHandle&& Other) noexcept;
 	~FSkeletalMeshSkinningDataHandle();
 
 	FSkeletalMeshSkinningDataHandle& operator=(const FSkeletalMeshSkinningDataHandle& Other) = delete;
-	FSkeletalMeshSkinningDataHandle& operator=(FSkeletalMeshSkinningDataHandle&& Other);
+	FSkeletalMeshSkinningDataHandle& operator=(FSkeletalMeshSkinningDataHandle&& Other) noexcept;
 
 	FSkeletalMeshSkinningDataUsage Usage;
 	TSharedPtr<FSkeletalMeshSkinningData> SkinningData;
@@ -260,11 +260,11 @@ struct NIAGARA_API FSkeletalMeshUvMappingHandle
 	FSkeletalMeshUvMappingHandle();
 	FSkeletalMeshUvMappingHandle(FSkeletalMeshUvMappingUsage InUsage, const TSharedPtr<struct FSkeletalMeshUvMapping>& InUvMappingData, bool bNeedsDataImmediately);
 	FSkeletalMeshUvMappingHandle(const FSkeletalMeshUvMappingHandle& Other) = delete;
-	FSkeletalMeshUvMappingHandle(FSkeletalMeshUvMappingHandle&& Other);
+	FSkeletalMeshUvMappingHandle(FSkeletalMeshUvMappingHandle&& Other) noexcept;
 	~FSkeletalMeshUvMappingHandle();
 
 	FSkeletalMeshUvMappingHandle& operator=(const FSkeletalMeshUvMappingHandle& Other) = delete;
-	FSkeletalMeshUvMappingHandle& operator=(FSkeletalMeshUvMappingHandle&& Other);
+	FSkeletalMeshUvMappingHandle& operator=(FSkeletalMeshUvMappingHandle&& Other) noexcept;
 	explicit operator bool() const;
 
 	FSkeletalMeshUvMappingUsage Usage;
@@ -301,11 +301,11 @@ struct NIAGARA_API FSkeletalMeshConnectivityHandle
 	FSkeletalMeshConnectivityHandle();
 	FSkeletalMeshConnectivityHandle(FSkeletalMeshConnectivityUsage InUsage, const TSharedPtr<struct FSkeletalMeshConnectivity>& InConnectivityData, bool bNeedsDataImmediately);
 	FSkeletalMeshConnectivityHandle(const FSkeletalMeshConnectivityHandle& Other) = delete;
-	FSkeletalMeshConnectivityHandle(FSkeletalMeshConnectivityHandle&& Other);
+	FSkeletalMeshConnectivityHandle(FSkeletalMeshConnectivityHandle&& Other) noexcept;
 	~FSkeletalMeshConnectivityHandle();
 
 	FSkeletalMeshConnectivityHandle& operator=(const FSkeletalMeshConnectivityHandle& Other) = delete;
-	FSkeletalMeshConnectivityHandle& operator=(FSkeletalMeshConnectivityHandle&& Other);
+	FSkeletalMeshConnectivityHandle& operator=(FSkeletalMeshConnectivityHandle&& Other) noexcept;
 	explicit operator bool() const;
 
 	FSkeletalMeshConnectivityUsage Usage;
@@ -433,7 +433,7 @@ class FSkeletalMeshGpuSpawnStaticBuffers : public FRenderResource
 {
 public:
 
-	virtual ~FSkeletalMeshGpuSpawnStaticBuffers();
+	virtual ~FSkeletalMeshGpuSpawnStaticBuffers() override;
 
 	FORCEINLINE_DEBUGGABLE void Initialise(struct FNDISkeletalMesh_InstanceData* InstData, const FSkeletalMeshLODRenderData& SkeletalMeshLODRenderData,const FSkeletalMeshSamplingLODBuiltData& SkeletalMeshSamplingLODBuiltData, FNiagaraSystemInstance* SystemInstance);
 
@@ -535,7 +535,7 @@ class FSkeletalMeshGpuDynamicBufferProxy : public FRenderResource
 public:
 
 	FSkeletalMeshGpuDynamicBufferProxy();
-	virtual ~FSkeletalMeshGpuDynamicBufferProxy();
+	virtual ~FSkeletalMeshGpuDynamicBufferProxy() override;
 
 	void Initialise(const FReferenceSkeleton& RefSkel, const FSkeletalMeshLODRenderData& SkeletalMeshLODRenderData, uint32 InSamplingSocketCount);
 
diff --git a/Engine/Plugins/FX/Niagara/Source/Niagara/Classes/NiagaraGPUSortInfo.h b/Engine/Plugins/FX/Niagara/Source/Niagara/Classes/NiagaraGPUSortInfo.h
--- a/Engine/Plugins/FX/Niagara/Source/Niagara/Classes/NiagaraGPUSortInfo.h
+++ b/Engine/Plugins/FX/Niagara/Source/Niagara/Classes/NiagaraGPUSortInfo.h
@@ -69,7 +69,7 @@ struct FNiagaraGPUSortInfo
 	FSphere LocalBSphere = FSphere(0);
 	FVector CullingWorldSpaceOffset = FVector(0, 0, 0);
 	FVector3f SystemLWCTile = FVector3f(0, 0, 0);
-	FVector2D DistanceCullRange { 0.0f, FLT_MAX };
+	FVector2f DistanceCullRange { 0.0f, FLT_MAX };
 	TArray<FPlane, TFixedAllocator<MaxCullPlanes>> CullPlanes;
 
 	// The GPUSortManager bindings for this sort task.
diff --git a/Engine/Plugins/FX/Niagara/Source/Niagara/Private/DataInterface/NiagaraDataInterfaceStaticMesh.cpp b/Engine/Plugins/FX/Niagara/Source/Niagara/Private/DataInterface/NiagaraDataInterfaceStaticMesh.cpp
--- a/Engine/Plugins/FX/Niagara/Source/Niagara/Private/DataInterface/NiagaraDataInterfaceStaticMesh.cpp
+++ b/Engine/Plugins/FX/Niagara/Source/Niagara/Private/DataInterface/NiagaraDataInterfaceStaticMesh.cpp
@@ -2757,7 +2757,7 @@ void UNiagaraDataInterfaceStaticMesh::VMGetVertexUV(FVectorVMExternalFunctionCon
 	{
 		for (int32 i = 0; i < Context.GetNumInstances(); ++i)
 		{
-			OutUV.SetAndAdvance(FVector2D::ZeroVector);
+			OutUV.SetAndAdvance(FVector2f::ZeroVector);
 		}
 	}
 }
@@ -3113,7 +3113,7 @@ void UNiagaraDataInterfaceStaticMesh::VMGetTriangleUV(FVectorVMExternalFunctionC
 	{
 		for (int32 i = 0; i < Context.GetNumInstances(); ++i)
 		{
-			OutUVParam.SetAndAdvance(FVector2D::ZeroVector);
+			OutUVParam.SetAndAdvance(FVector2f::ZeroVector);
 		}
 	}
 }
diff --git a/Engine/Plugins/FX/Niagara/Source/Niagara/Private/NDISkeletalMesh_BoneSampling.cpp b/Engine/Plugins/FX/Niagara/Source/Niagara/Private/NDISkeletalMesh_BoneSampling.cpp
--- a/Engine/Plugins/FX/Niagara/Source/Niagara/Private/NDISkeletalMesh_BoneSampling.cpp
+++ b/Engine/Plugins/FX/Niagara/Source/Niagara/Private/NDISkeletalMesh_BoneSampling.cpp
@@ -2,11 +2,9 @@
 
 #include "NiagaraDataInterfaceSkeletalMesh.h"
 #include "NiagaraComponent.h"
-#include "Animation/SkeletalMeshActor.h"
 #include "Components/SkeletalMeshComponent.h"
 #include "SkeletalMeshTypes.h"
 #include "NiagaraStats.h"
-#include "Templates/AlignmentTemplates.h"
 #include "NDISkeletalMeshCommon.h"
 
 #define LOCTEXT_NAMESPACE "NiagaraDataInterfaceSkeletalMesh_BoneSampling"
@@ -793,10 +791,10 @@ void UNiagaraDataInterfaceSkeletalMesh::GetSkinnedBoneData(FVectorVMExternalFunc
 
 	//TODO: Replace this by storing off FTransforms and doing a proper lerp to get a final transform.
 	//Also need to pull in a per particle interpolation factor.
-	const FMatrix44f InstanceTransform = (FMatrix44f)InstData->Transform;
-	const FMatrix44f PrevInstanceTransform = (FMatrix44f)InstData->PrevTransform;
-	const FQuat4f InstanceRotation = Output.bNeedsRotation ? InstanceTransform.GetMatrixWithoutScale().ToQuat() : FQuat4f::Identity;
-	const FQuat4f PrevInstanceRotation = Output.bNeedsRotation ? PrevInstanceTransform.GetMatrixWithoutScale().ToQuat() : FQuat4f::Identity;
+	const FMatrix InstanceTransform = InstData->Transform;
+	const FMatrix PrevInstanceTransform = InstData->PrevTransform;
+	const FQuat4f InstanceRotation = Output.bNeedsRotation ? FQuat4f(InstanceTransform.GetMatrixWithoutScale().ToQuat()) : FQuat4f::Identity;
+	const FQuat4f PrevInstanceRotation = Output.bNeedsRotation ? FQuat4f(PrevInstanceTransform.GetMatrixWithoutScale().ToQuat()) : FQuat4f::Identity;
 
 	FSkeletalMeshAccessorHelper Accessor;
 	Accessor.Init<TNDISkelMesh_FilterModeNone, TNDISkelMesh_AreaWeightingOff>(InstData);
@@ -1008,15 +1006,15 @@ void UNiagaraDataInterfaceSkeletalMesh::GetFilteredSocketTransform(FVectorVMExte
 	if (SocketMax >= 0)
 	{
 		const bool bNeedsRotation = OutSocketRotation.IsValid();
-		const FMatrix44f InstanceTransform = (FMatrix44f)InstData->Transform;
-		const FQuat4f InstanceRotation = bNeedsRotation ? InstanceTransform.GetMatrixWithoutScale().ToQuat() : FQuat4f::Identity;
+		const FMatrix InstanceTransform = InstData->Transform;
+		const FQuat4f InstanceRotation = bNeedsRotation ? FQuat4f(InstanceTransform.GetMatrixWithoutScale().ToQuat()) : FQuat4f::Identity;
 
 		for (int32 i = 0; i < Context.GetNumInstances(); ++i)
 		{
 			const int32 SocketIndex = FMath::Clamp(SocketParam.GetAndAdvance(), 0, SocketMax);
-			FVector3f SocketTranslation = CurrentFilteredSockets[SocketIndex].GetTranslation();
+			FVector SocketTranslation = CurrentFilteredSockets[SocketIndex].GetTranslation();
 			FQuat4f SocketRotation = CurrentFilteredSockets[SocketIndex].GetRotation();
-			FVector3f SocketScale = CurrentFilteredSockets[SocketIndex].GetScale3D();
+			FVector SocketScale = CurrentFilteredSockets[SocketIndex].GetScale3D();
 
 			const bool bApplyTransform = ApplyWorldTransform.GetAndAdvance();
 			if (bApplyTransform)
diff --git a/Engine/Plugins/FX/Niagara/Source/Niagara/Private/NDISkeletalMesh_TriangleSampling.cpp b/Engine/Plugins/FX/Niagara/Source/Niagara/Private/NDISkeletalMesh_TriangleSampling.cpp
--- a/Engine/Plugins/FX/Niagara/Source/Niagara/Private/NDISkeletalMesh_TriangleSampling.cpp
+++ b/Engine/Plugins/FX/Niagara/Source/Niagara/Private/NDISkeletalMesh_TriangleSampling.cpp
@@ -1041,7 +1041,7 @@ void UNiagaraDataInterfaceSkeletalMesh::GetTriCoordUV(FVectorVMExternalFunctionC
 	// Fall-though for failure conditions
 	for (int32 i=0; i < Context.GetNumInstances(); ++i)
 	{
-		OutUV.SetAndAdvance(FVector2D::ZeroVector);
+		OutUV.SetAndAdvance(FVector2f::ZeroVector);
 	}
 }
 
diff --git a/Engine/Plugins/FX/Niagara/Source/Niagara/Private/NDISkeletalMesh_VertexSampling.cpp b/Engine/Plugins/FX/Niagara/Source/Niagara/Private/NDISkeletalMesh_VertexSampling.cpp
--- a/Engine/Plugins/FX/Niagara/Source/Niagara/Private/NDISkeletalMesh_VertexSampling.cpp
+++ b/Engine/Plugins/FX/Niagara/Source/Niagara/Private/NDISkeletalMesh_VertexSampling.cpp
@@ -652,7 +652,7 @@ void UNiagaraDataInterfaceSkeletalMesh::GetVertexUV(FVectorVMExternalFunctionCon
 	// Fall though for bad data
 	for (int32 i = 0; i < Context.GetNumInstances(); ++i)
 	{
-		OutUV.SetAndAdvance(FVector2D::ZeroVector);
+		OutUV.SetAndAdvance(FVector2f::ZeroVector);
 	}
 }
 
@@ -668,7 +668,7 @@ void UNiagaraDataInterfaceSkeletalMesh::GetVertexUV<FSkelMeshVertexAccessorNoop>
 
 	for (int32 i = 0; i < Context.GetNumInstances(); ++i)
 	{
-		OutUV.SetAndAdvance(FVector2D::ZeroVector);
+		OutUV.SetAndAdvance(FVector2f::ZeroVector);
 	}
 }
 
diff --git a/Engine/Plugins/FX/Niagara/Source/Niagara/Private/NiagaraDataInterface.cpp b/Engine/Plugins/FX/Niagara/Source/Niagara/Private/NiagaraDataInterface.cpp
--- a/Engine/Plugins/FX/Niagara/Source/Niagara/Private/NiagaraDataInterface.cpp
+++ b/Engine/Plugins/FX/Niagara/Source/Niagara/Private/NiagaraDataInterface.cpp
@@ -3,7 +3,6 @@
 #include "NiagaraDataInterface.h"
 #include "Curves/CurveVector.h"
 #include "Curves/CurveLinearColor.h"
-#include "Curves/CurveFloat.h"
 #include "NiagaraTypes.h"
 #include "ShaderParameterUtils.h"
 #include "NiagaraShader.h"
diff --git a/Engine/Plugins/FX/Niagara/Source/Niagara/Private/NiagaraDataInterfaceSkeletalMesh.cpp b/Engine/Plugins/FX/Niagara/Source/Niagara/Private/NiagaraDataInterfaceSkeletalMesh.cpp
--- a/Engine/Plugins/FX/Niagara/Source/Niagara/Private/NiagaraDataInterfaceSkeletalMesh.cpp
+++ b/Engine/Plugins/FX/Niagara/Source/Niagara/Private/NiagaraDataInterfaceSkeletalMesh.cpp
@@ -175,7 +175,6 @@ float FSkeletalMeshSamplingRegionAreaWeightedSampler::GetWeights(TArray<float>&
 	for (int32 i = 0; i < NumUsedRegions; ++i)
 	{
 		int32 RegionIdx = Owner->SamplingRegionIndices[i];
-		const FSkeletalMeshSamplingRegion& Region = SamplingInfo.GetRegion(RegionIdx);
 		float T = SamplingInfo.GetRegionBuiltData(RegionIdx).AreaWeightedSampler.GetTotalWeight();
 		OutWeights.Add(T);
 		Total += T;
@@ -207,14 +206,14 @@ FSkeletalMeshSkinningDataHandle::~FSkeletalMeshSkinningDataHandle()
 	}
 }
 
-FSkeletalMeshSkinningDataHandle::FSkeletalMeshSkinningDataHandle(FSkeletalMeshSkinningDataHandle&& Other)
+FSkeletalMeshSkinningDataHandle::FSkeletalMeshSkinningDataHandle(FSkeletalMeshSkinningDataHandle&& Other) noexcept
 {
 	Usage = Other.Usage;
 	SkinningData = Other.SkinningData;
 	Other.SkinningData = nullptr;
 }
 
-FSkeletalMeshSkinningDataHandle& FSkeletalMeshSkinningDataHandle::operator=(FSkeletalMeshSkinningDataHandle&& Other)
+FSkeletalMeshSkinningDataHandle& FSkeletalMeshSkinningDataHandle::operator=(FSkeletalMeshSkinningDataHandle&& Other) noexcept
 {
 	if (this != &Other)
 	{
@@ -364,7 +363,6 @@ void FSkeletalMeshSkinningData::UpdateBoneTransforms()
 		return;
 	}
 
-	const TArray<FTransform>& BaseCompSpaceTransforms = SkelComp->GetComponentSpaceTransforms();
 	TArray<FMatrix44f>& CurrBones = CurrBoneRefToLocals();
 	TArray<FTransform3f>& CurrTransforms = CurrComponentTransforms();
 
@@ -1077,7 +1075,6 @@ void FSkeletalMeshGpuDynamicBufferProxy::NewFrame(const FNDISkeletalMesh_Instanc
 			check(0 <= LODIndex && LODIndex < LODRenderDataArray.Num());
 			FSkeletalMeshLODRenderData& LODRenderData = LODRenderDataArray[LODIndex];
 			TArray<FSkelMeshRenderSection>& Sections = LODRenderData.RenderSections;
-			uint32 SectionCount = Sections.Num();
 
 			// Count number of matrices we want before appending all of them according to the per section mapping from BoneMap
 			uint32 Float4Count = 0;
@@ -2356,13 +2353,6 @@ bool FNDISkeletalMesh_InstanceData::Tick(UNiagaraDataInterfaceSkeletalMesh* Inte
 
 		if (MeshGpuSpawnDynamicBuffers)
 		{
-			USkeletalMeshComponent* Comp = Cast<USkeletalMeshComponent>(SceneComponent.Get());
-			const USkinnedMeshComponent* BaseComp = nullptr;
-			if (Comp)
-			{
-				BaseComp = Comp->GetBaseComponent();
-			}
-
 			MeshGpuSpawnDynamicBuffers->NewFrame(this, GetLODIndex());
 		}
 
@@ -2421,16 +2411,16 @@ void FNDISkeletalMesh_InstanceData::Release()
 
 UNiagaraDataInterfaceSkeletalMesh::UNiagaraDataInterfaceSkeletalMesh(FObjectInitializer const& ObjectInitializer)
 	: Super(ObjectInitializer)
-	, SourceMode(ENDISkeletalMesh_SourceMode::Default)
+	  , SourceMode(ENDISkeletalMesh_SourceMode::Default)
 #if WITH_EDITORONLY_DATA
-	, PreviewMesh(nullptr)
+	  , PreviewMesh(nullptr)
 #endif
-	, Source(nullptr)
-	, SkinningMode(ENDISkeletalMesh_SkinningMode::SkinOnTheFly)
-	, WholeMeshLOD(INDEX_NONE)
-	, ChangeId(0)
+	  , Source(nullptr)
+	  , SourceComponent(nullptr)
+      , SkinningMode(ENDISkeletalMesh_SkinningMode::SkinOnTheFly)
+	  , WholeMeshLOD(INDEX_NONE)
+	  , ChangeId(0)
 {
-
 	FNiagaraTypeDefinition Def(UObject::StaticClass());
 	MeshUserParameter.Parameter.SetType(Def);
 
@@ -2825,8 +2815,6 @@ void UNiagaraDataInterfaceSkeletalMesh::GetFeedback(UNiagaraSystem* Asset, UNiag
 //Deprecated functions we check for and advise on updates in ValidateFunction
 static const FName GetTriPositionName_DEPRECATED("GetTriPosition");
 static const FName GetTriPositionWSName_DEPRECATED("GetTriPositionWS");
-static const FName GetTriNormalName_DEPRECATED("GetTriNormal");
-static const FName GetTriNormalWSName_DEPRECATED("GetTriNormalWS");
 static const FName GetTriPositionVelocityAndNormalName_DEPRECATED("GetTriPositionVelocityAndNormal");
 static const FName GetTriPositionVelocityAndNormalWSName_DEPRECATED("GetTriPositionVelocityAndNormalWS");
 static const FName GetTriPositionVelocityAndNormalBinormalTangentName_DEPRECATED("GetTriPositionVelocityAndNormalBinormalTangent");
diff --git a/Engine/Plugins/FX/Niagara/Source/Niagara/Private/NiagaraDataInterfaceSkeletalMeshConnectivity.cpp b/Engine/Plugins/FX/Niagara/Source/Niagara/Private/NiagaraDataInterfaceSkeletalMeshConnectivity.cpp
--- a/Engine/Plugins/FX/Niagara/Source/Niagara/Private/NiagaraDataInterfaceSkeletalMeshConnectivity.cpp
+++ b/Engine/Plugins/FX/Niagara/Source/Niagara/Private/NiagaraDataInterfaceSkeletalMeshConnectivity.cpp
@@ -3,8 +3,6 @@
 #include "NiagaraDataInterfaceSkeletalMeshConnectivity.h"
 
 #include "Algo/StableSort.h"
-#include "Algo/Unique.h"
-#include "NDISkeletalMeshCommon.h"
 #include "NiagaraResourceArrayWriter.h"
 #include "NiagaraSettings.h"
 #include "NiagaraStats.h"
@@ -36,14 +34,14 @@ FSkeletalMeshConnectivityHandle::~FSkeletalMeshConnectivityHandle()
 	}
 }
 
-FSkeletalMeshConnectivityHandle::FSkeletalMeshConnectivityHandle(FSkeletalMeshConnectivityHandle&& Other)
+FSkeletalMeshConnectivityHandle::FSkeletalMeshConnectivityHandle(FSkeletalMeshConnectivityHandle&& Other) noexcept
 {
 	Usage = Other.Usage;
 	ConnectivityData = Other.ConnectivityData;
 	Other.ConnectivityData = nullptr;
 }
 
-FSkeletalMeshConnectivityHandle& FSkeletalMeshConnectivityHandle::operator=(FSkeletalMeshConnectivityHandle&& Other)
+FSkeletalMeshConnectivityHandle& FSkeletalMeshConnectivityHandle::operator=(FSkeletalMeshConnectivityHandle&& Other) noexcept
 {
 	if (this != &Other)
 	{
diff --git a/Engine/Plugins/FX/Niagara/Source/Niagara/Private/NiagaraDataInterfaceSkeletalMeshUvMapping.cpp b/Engine/Plugins/FX/Niagara/Source/Niagara/Private/NiagaraDataInterfaceSkeletalMeshUvMapping.cpp
--- a/Engine/Plugins/FX/Niagara/Source/Niagara/Private/NiagaraDataInterfaceSkeletalMeshUvMapping.cpp
+++ b/Engine/Plugins/FX/Niagara/Source/Niagara/Private/NiagaraDataInterfaceSkeletalMeshUvMapping.cpp
@@ -100,7 +100,7 @@ struct FQuadTreeQueryHelper
 			float TriangleSegmentMin = FMath::Min3(FVector2D::DotProduct(A, SeparatingAxis), FVector2D::DotProduct(B, SeparatingAxis), FVector2D::DotProduct(C, SeparatingAxis));
 			float TriangleSegmentMax = FMath::Max3(FVector2D::DotProduct(A, SeparatingAxis), FVector2D::DotProduct(B, SeparatingAxis), FVector2D::DotProduct(C, SeparatingAxis));
 
-			if (AabbSegmentMin > TriangleSegmentMax || AabbSegmentMax < TriangleSegmentMax)
+			if (AabbSegmentMin > TriangleSegmentMax || AabbSegmentMax < TriangleSegmentMin)
 			{
 				return false;
 			}
@@ -160,14 +160,14 @@ FSkeletalMeshUvMappingHandle::~FSkeletalMeshUvMappingHandle()
 	}
 }
 
-FSkeletalMeshUvMappingHandle::FSkeletalMeshUvMappingHandle(FSkeletalMeshUvMappingHandle&& Other)
+FSkeletalMeshUvMappingHandle::FSkeletalMeshUvMappingHandle(FSkeletalMeshUvMappingHandle&& Other) noexcept
 {
 	Usage = Other.Usage;
 	UvMappingData = Other.UvMappingData;
 	Other.UvMappingData = nullptr;
 }
 
-FSkeletalMeshUvMappingHandle& FSkeletalMeshUvMappingHandle::operator=(FSkeletalMeshUvMappingHandle&& Other)
+FSkeletalMeshUvMappingHandle& FSkeletalMeshUvMappingHandle::operator=(FSkeletalMeshUvMappingHandle&& Other) noexcept
 {
 	if (this != &Other)
 	{
@@ -183,7 +183,7 @@ FSkeletalMeshUvMappingHandle::operator bool() const
 	return UvMappingData.IsValid();
 }
 
-bool FSkeletalMeshUvMapping::IsValidMeshObject(TWeakObjectPtr<USkeletalMesh>& MeshObject, int32 InLodIndex, int32 InUvSetIndex)
+bool FSkeletalMeshUvMapping::IsValidMeshObject(const TWeakObjectPtr<USkeletalMesh>& MeshObject, int32 InLodIndex, int32 InUvSetIndex)
 {
 	USkeletalMesh* Mesh = MeshObject.Get();
 
diff --git a/Engine/Plugins/FX/Niagara/Source/Niagara/Private/NiagaraDataInterfaceSkeletalMeshUvMapping.h b/Engine/Plugins/FX/Niagara/Source/Niagara/Private/NiagaraDataInterfaceSkeletalMeshUvMapping.h
--- a/Engine/Plugins/FX/Niagara/Source/Niagara/Private/NiagaraDataInterfaceSkeletalMeshUvMapping.h
+++ b/Engine/Plugins/FX/Niagara/Source/Niagara/Private/NiagaraDataInterfaceSkeletalMeshUvMapping.h
@@ -41,7 +41,7 @@ struct FSkeletalMeshUvMapping
 	void RegisterUser(FSkeletalMeshUvMappingUsage Usage, bool bNeedsDataImmediately);
 	void UnregisterUser(FSkeletalMeshUvMappingUsage Usage);
 
-	static bool IsValidMeshObject(TWeakObjectPtr<USkeletalMesh>& MeshObject, int32 InLodIndex, int32 InUvSetIndex);
+	static bool IsValidMeshObject(const TWeakObjectPtr<USkeletalMesh>& MeshObject, int32 InLodIndex, int32 InUvSetIndex);
 
 	bool Matches(const TWeakObjectPtr<USkeletalMesh>& InMeshObject, int32 InLodIndex, int32 InUvSetIndex) const;
 
diff --git a/Engine/Plugins/FX/Niagara/Source/Niagara/Private/NiagaraDataInterfaceSpline.cpp b/Engine/Plugins/FX/Niagara/Source/Niagara/Private/NiagaraDataInterfaceSpline.cpp
--- a/Engine/Plugins/FX/Niagara/Source/Niagara/Private/NiagaraDataInterfaceSpline.cpp
+++ b/Engine/Plugins/FX/Niagara/Source/Niagara/Private/NiagaraDataInterfaceSpline.cpp
@@ -5,7 +5,6 @@
 #include "NiagaraComponent.h"
 #include "NiagaraSystemInstance.h"
 #include "Internationalization/Internationalization.h"
-#include "NiagaraShader.h"
 #include "ShaderParameterUtils.h"
 #include "ShaderCompilerCore.h"
 
@@ -520,7 +519,7 @@ void UNiagaraDataInterfaceSpline::DestroyPerInstanceData(void* PerInstanceData,
 {
 	SystemInstancesToProxyData_GT.Remove(SystemInstance->GetId());
 	
-	FNDISpline_InstanceData* InstData = (FNDISpline_InstanceData*)PerInstanceData;
+	FNDISpline_InstanceData* InstData = static_cast<FNDISpline_InstanceData*>(PerInstanceData);
 	InstData->~FNDISpline_InstanceData();
 
 	FNiagaraDataInterfaceProxySpline* RT_Proxy = GetProxyAs<FNiagaraDataInterfaceProxySpline>();
@@ -571,7 +570,7 @@ bool UNiagaraDataInterfaceSpline::PerInstanceTick(void* PerInstanceData, FNiagar
 						SplineComponent = UserSplineComp;
 					}
 				}
-				else if (AActor* Actor = Cast<AActor>(UserParamObject))
+				else if (Cast<AActor>(UserParamObject))
 				{
 					SplineComponent = Source->FindComponentByClass<USplineComponent>();
 				}
diff --git a/Engine/Plugins/FX/Niagara/Source/Niagara/Private/NiagaraGpuComputeDispatch.cpp b/Engine/Plugins/FX/Niagara/Source/Niagara/Private/NiagaraGpuComputeDispatch.cpp
--- a/Engine/Plugins/FX/Niagara/Source/Niagara/Private/NiagaraGpuComputeDispatch.cpp
+++ b/Engine/Plugins/FX/Niagara/Source/Niagara/Private/NiagaraGpuComputeDispatch.cpp
@@ -1988,14 +1988,14 @@ void FNiagaraGpuComputeDispatch::GenerateSortKeys(FRHICommandListImmediate& RHIC
 			Params.MeshIndex = SortInfo.MeshIndex;
 			Params.MeshIndexAttributeOffset = SortInfo.MeshIndexAttributeOffset;
 			Params.CullDistanceRangeSquared = SortInfo.DistanceCullRange * SortInfo.DistanceCullRange;
-			Params.LocalBoundingSphere = FVector4(SortInfo.LocalBSphere.Center, SortInfo.LocalBSphere.W);
+			Params.LocalBoundingSphere = FVector4f(SortInfo.LocalBSphere.Center, SortInfo.LocalBSphere.W);
 			Params.CullingWorldSpaceOffset = SortInfo.CullingWorldSpaceOffset;
 			Params.SystemLWCTile = SortInfo.SystemLWCTile;
 
 			Params.NumCullPlanes = 0;
 			for (const FPlane& Plane : SortInfo.CullPlanes)
 			{
-				Params.CullPlanes[Params.NumCullPlanes++] = FVector4(Plane.X, Plane.Y, Plane.Z, Plane.W);
+				Params.CullPlanes[Params.NumCullPlanes++] = FVector4f(Plane.X, Plane.Y, Plane.Z, Plane.W);
 			}
 
 			// Choose the shader to bind
diff --git a/Engine/Plugins/FX/Niagara/Source/Niagara/Private/NiagaraRendererMeshes.cpp b/Engine/Plugins/FX/Niagara/Source/Niagara/Private/NiagaraRendererMeshes.cpp
--- a/Engine/Plugins/FX/Niagara/Source/Niagara/Private/NiagaraRendererMeshes.cpp
+++ b/Engine/Plugins/FX/Niagara/Source/Niagara/Private/NiagaraRendererMeshes.cpp
@@ -89,14 +89,14 @@ FNiagaraRendererMeshes::FNiagaraRendererMeshes(ERHIFeatureLevel::Type FeatureLev
 	bSubImageBlend = Properties->bSubImageBlend;
 	bEnableFrustumCulling = Properties->bEnableFrustumCulling;
 	bEnableCulling = bEnableFrustumCulling;
-	DistanceCullRange = FVector2D(0, FLT_MAX);
+	DistanceCullRange = FVector2f(0, FLT_MAX);
 	RendererVisibility = Properties->RendererVisibility;
 	bAccurateMotionVectors = Properties->NeedsPreciseMotionVectors();
 	MaxSectionCount = 0;
 
 	if (Properties->bEnableCameraDistanceCulling)
 	{
-		DistanceCullRange = FVector2D(Properties->MinCameraDistance, Properties->MaxCameraDistance);
+		DistanceCullRange = FVector2f(Properties->MinCameraDistance, Properties->MaxCameraDistance);
 		bEnableCulling = true;
 	}
 
diff --git a/Engine/Plugins/FX/Niagara/Source/Niagara/Private/NiagaraRendererSprites.cpp b/Engine/Plugins/FX/Niagara/Source/Niagara/Private/NiagaraRendererSprites.cpp
--- a/Engine/Plugins/FX/Niagara/Source/Niagara/Private/NiagaraRendererSprites.cpp
+++ b/Engine/Plugins/FX/Niagara/Source/Niagara/Private/NiagaraRendererSprites.cpp
@@ -107,7 +107,7 @@ FNiagaraRendererSprites::FNiagaraRendererSprites(ERHIFeatureLevel::Type FeatureL
 	bEnableDistanceCulling = Properties->bEnableCameraDistanceCulling;
 	if (Properties->bEnableCameraDistanceCulling)
 	{
-		DistanceCullRange = FVector2D(Properties->MinCameraDistance, Properties->MaxCameraDistance);
+		DistanceCullRange = FVector2f(Properties->MinCameraDistance, Properties->MaxCameraDistance);
 	}
 
 	// Get the offset of visibility tag in either particle data or parameter store
diff --git a/Engine/Plugins/FX/Niagara/Source/Niagara/Public/NiagaraRendererMeshes.h b/Engine/Plugins/FX/Niagara/Source/Niagara/Public/NiagaraRendererMeshes.h
--- a/Engine/Plugins/FX/Niagara/Source/Niagara/Public/NiagaraRendererMeshes.h
+++ b/Engine/Plugins/FX/Niagara/Source/Niagara/Public/NiagaraRendererMeshes.h
@@ -143,12 +143,12 @@ private:
 	uint32 bAccurateMotionVectors : 1;
 
 	uint32 bSubImageBlend : 1;
-	FVector2D SubImageSize;
+	FVector2f SubImageSize;
 
 	FVector LockedAxis;
 	ENiagaraMeshLockedAxisSpace LockedAxisSpace;
 
-	FVector2D DistanceCullRange;
+	FVector2f DistanceCullRange;
 	int32 RendererVisTagOffset;
 	int32 RendererVisibility;
 	int32 MeshIndexOffset;
diff --git a/Engine/Plugins/FX/Niagara/Source/Niagara/Public/NiagaraRendererSprites.h b/Engine/Plugins/FX/Niagara/Source/Niagara/Public/NiagaraRendererSprites.h
--- a/Engine/Plugins/FX/Niagara/Source/Niagara/Public/NiagaraRendererSprites.h
+++ b/Engine/Plugins/FX/Niagara/Source/Niagara/Public/NiagaraRendererSprites.h
@@ -104,9 +104,9 @@ private:
 	ENiagaraSpriteAlignment Alignment;
 	ENiagaraSpriteFacingMode FacingMode;
 	ENiagaraSortMode SortMode;
-	FVector2D PivotInUVSpace;
+	FVector2f PivotInUVSpace;
 	float MacroUVRadius;
-	FVector2D SubImageSize;
+	FVector2f SubImageSize;
 
 	uint32 NumIndicesPerInstance;
 
@@ -122,7 +122,7 @@ private:
 
 	float MinFacingCameraBlendDistance;
 	float MaxFacingCameraBlendDistance;
-	FVector2D DistanceCullRange;
+	FVector2f DistanceCullRange;
 	FNiagaraCutoutVertexBuffer CutoutVertexBuffer;
 	int32 NumCutoutVertexPerSubImage = 0;
 	uint32 MaterialParamValidMask = 0;
diff --git a/Engine/Source/Runtime/VectorVM/Public/VectorVM.h b/Engine/Source/Runtime/VectorVM/Public/VectorVM.h
--- a/Engine/Source/Runtime/VectorVM/Public/VectorVM.h
+++ b/Engine/Source/Runtime/VectorVM/Public/VectorVM.h
@@ -227,7 +227,9 @@ struct FDataSetMeta
 	FORCEINLINE void UnlockFreeTable();
 
 	FDataSetMeta()
-		: DataSetAccessIndex(INDEX_NONE)
+		: InputRegisterTypeOffsets{}
+		, OutputRegisterTypeOffsets{}
+		, DataSetAccessIndex(INDEX_NONE)
 		, InstanceOffset(INDEX_NONE)
 		, IDTable(nullptr)
 		, FreeIDTable(nullptr)
@@ -569,7 +571,7 @@ namespace VectorVM
 			UserPtrIdx = *Context.GetConstant<int32>(ConstantTableOffset);
 			check(UserPtrIdx != INDEX_NONE);
 			
-			Ptr = reinterpret_cast<T*>(Context.GetUserPtrTable(UserPtrIdx));
+			Ptr = static_cast<T*>(Context.GetUserPtrTable(UserPtrIdx));
 #endif
 		}
 
@@ -712,6 +714,3 @@ namespace VectorVM
 		FORCEINLINE void Advance() { }
 	};
 } // namespace VectorVM
-
-
-
