Niagara - When setting an input to a user parameter, make sure the parameter exists and if it doesn't, add it.  Also when adding the user parameter, transfer the current value if it was set localy. This fixes an issue where user parameters wouldn't be added in cases where the system was able to get the compilation data from the ddc, and also makes the experience more consistent when dealing with long compile times since the user parameter is available immediately.

diff --git a/Engine/Plugins/FX/Niagara/Source/NiagaraEditor/Private/ViewModels/Stack/NiagaraStackFunctionInput.cpp b/Engine/Plugins/FX/Niagara/Source/NiagaraEditor/Private/ViewModels/Stack/NiagaraStackFunctionInput.cpp
--- a/Engine/Plugins/FX/Niagara/Source/NiagaraEditor/Private/ViewModels/Stack/NiagaraStackFunctionInput.cpp
+++ b/Engine/Plugins/FX/Niagara/Source/NiagaraEditor/Private/ViewModels/Stack/NiagaraStackFunctionInput.cpp
@@ -1113,6 +1113,7 @@ void UNiagaraStackFunctionInput::SetLinkedValueHandle(const FNiagaraParameterHan
 		return;
 	}
 
+	TGuardValue<bool> UpdateGuard(bUpdatingLocalValueDirectly, true);
 	FScopedTransaction ScopedTransaction(LOCTEXT("UpdateLinkedInputValue", "Update linked input value"));
 	RemoveOverridePin();
 
@@ -1123,6 +1124,27 @@ void UNiagaraStackFunctionInput::SetLinkedValueHandle(const FNiagaraParameterHan
 
 	if (InParameterHandle != DefaultInputValues.LinkedHandle)
 	{
+		if (InParameterHandle.IsUserHandle())
+		{
+			// If the handle is a user parameter, make sure the system has it exposed.  If it's not exposed add it directly here
+			// rather than waiting on the compile results so that it's immediately available.
+			FNiagaraVariable LinkedParameterVariable = FNiagaraVariable(InputType, InParameterHandle.GetParameterHandleString()); 
+			FNiagaraParameterStore& UserParameters = GetSystemViewModel()->GetSystem().GetExposedParameters();
+			if (UserParameters.IndexOf(LinkedParameterVariable) == INDEX_NONE)
+			{
+				if (InputValues.Mode == EValueMode::Local && InputValues.LocalStruct.IsValid())
+				{
+					// If the current value is local, and valid transfer that value to the user parameter.
+					LinkedParameterVariable.SetData(InputValues.LocalStruct->GetStructMemory());
+				}
+				else
+				{
+					FNiagaraEditorUtilities::ResetVariableToDefaultValue(LinkedParameterVariable);
+				}
+				UserParameters.SetParameterData(LinkedParameterVariable.GetData(), LinkedParameterVariable, true);
+			}
+		}
+
 		// Only set the linked value if it's actually different from the default.
 		FNiagaraStackGraphUtilities::SetLinkedValueHandleForFunctionInput(GetOrCreateOverridePin(), InParameterHandle);
 		OwningFunctionCallNode->FixupPinNames(); // refresh the input guids and update the bound names
