allow grids with 0 attributes to not error out

diff --git a/Engine/Plugins/FX/Niagara/Source/Niagara/Private/NiagaraDataInterfaceGrid3DCollection.cpp b/Engine/Plugins/FX/Niagara/Source/Niagara/Private/NiagaraDataInterfaceGrid3DCollection.cpp
--- a/Engine/Plugins/FX/Niagara/Source/Niagara/Private/NiagaraDataInterfaceGrid3DCollection.cpp
+++ b/Engine/Plugins/FX/Niagara/Source/Niagara/Private/NiagaraDataInterfaceGrid3DCollection.cpp
@@ -1909,7 +1909,7 @@ bool UNiagaraDataInterfaceGrid3DCollection::InitPerInstanceData(void* PerInstanc
 		InstanceData->NumCells.Y <= 0 ||
 		InstanceData->NumCells.Z <= 0)
 	{
-		UE_LOG(LogNiagara, Error, TEXT("Zero grid resolution defined on %s"), *FNiagaraUtilities::SystemInstanceIDToString(SystemInstance->GetId()));
+		UE_LOG(LogNiagara, Error, TEXT("Zero grid resolution defined on %s"), *Proxy->SourceDIName.ToString());
 		return false;
 	}
 
@@ -1930,20 +1930,32 @@ bool UNiagaraDataInterfaceGrid3DCollection::InitPerInstanceData(void* PerInstanc
 
 	if (InstanceData->NumCells.X > MaxDim || InstanceData->NumCells.Y > MaxDim || InstanceData->NumCells.Z > MaxDim)
 	{
-		UE_LOG(LogNiagara, Error, TEXT("Resolution is too high on %s... max is 2048, defined as %i, %i %i"), *FNiagaraUtilities::SystemInstanceIDToString(SystemInstance->GetId()), InstanceData->NumCells.X, InstanceData->NumCells.Y, InstanceData->NumCells.Z);
+		UE_LOG(LogNiagara, Error, TEXT("Resolution is too high on %s... max is 2048, defined as %d, %d %d"), *Proxy->SourceDIName.ToString(), InstanceData->NumCells.X, InstanceData->NumCells.Y, InstanceData->NumCells.Z);
 		return false;
 	}
 
 	if (NumAttribChannelsFound > MaxAttributes && MaxAttributes > 0)
 	{
-		UE_LOG(LogNiagara, Error, TEXT("Invalid number of attributes defined on %s... max is %i, num defined is %i"), *FNiagaraUtilities::SystemInstanceIDToString(SystemInstance->GetId()), MaxAttributes, NumAttribChannelsFound);
+		UE_LOG(LogNiagara, Error, TEXT("Invalid number of attributes defined on %s... max is %d, num defined is %d"), *Proxy->SourceDIName.ToString(), MaxAttributes, NumAttribChannelsFound);
 		return false;
 	}
 
 	if (NumAttribChannelsFound == 0)
 	{
-		UE_LOG(LogNiagara, Warning, TEXT("Invalid number of attributes defined on %s... max is %i, num defined is %i"), *FNiagaraUtilities::SystemInstanceIDToString(SystemInstance->GetId()), MaxAttributes, NumAttribChannelsFound);
-		return false;
+		UE_LOG(LogNiagara, Warning, TEXT("Invalid number of attributes defined on %s... max is %d, num defined is %d"), *Proxy->SourceDIName.ToString(), MaxAttributes, NumAttribChannelsFound);
+		
+
+		// Push Updates to Proxy.
+		FNiagaraDataInterfaceProxyGrid3DCollectionProxy* RT_Proxy = GetProxyAs<FNiagaraDataInterfaceProxyGrid3DCollectionProxy>();
+		ENQUEUE_RENDER_COMMAND(FUpdateData)(
+			[RT_Proxy, InstanceID = SystemInstance->GetId()](FRHICommandListImmediate& RHICmdList)
+		{
+			check(!RT_Proxy->SystemInstancesToProxyData_RT.Contains(InstanceID));
+			FGrid3DCollectionRWInstanceData_RenderThread* TargetData = &RT_Proxy->SystemInstancesToProxyData_RT.Add(InstanceID);
+
+		});
+
+		return true;
 	}
 
 	// need to determine number of tiles in x and y based on number of attributes and max dimension size
@@ -2526,13 +2538,19 @@ bool UNiagaraDataInterfaceGrid3DCollection::PerInstanceTickPostSimulate(void* Pe
 		
 		if (InstanceData->NumCells.X > MaxDim || InstanceData->NumCells.Y > MaxDim || InstanceData->NumCells.Z > MaxDim)
 		{
-			UE_LOG(LogNiagara, Error, TEXT("Resolution is too high on %s... max is 2048, defined as %i, %i %i"), *FNiagaraUtilities::SystemInstanceIDToString(SystemInstance->GetId()), InstanceData->NumCells.X, InstanceData->NumCells.Y, InstanceData->NumCells.Z);
+			UE_LOG(LogNiagara, Error, TEXT("Resolution is too high on %s... max is 2048, defined as %d, %d %d"), *Proxy->SourceDIName.ToString(), InstanceData->NumCells.X, InstanceData->NumCells.Y, InstanceData->NumCells.Z);
 			return false;
 		}
 		
-		if ((InstanceData->TotalNumAttributes > MaxAttributes && MaxAttributes > 0) || InstanceData->TotalNumAttributes <= 0)
+		if ((InstanceData->TotalNumAttributes > MaxAttributes && MaxAttributes > 0))
+		{
+			UE_LOG(LogNiagara, Error, TEXT("Invalid number of attributes defined on %s... max is %d, num defined is %d"), *Proxy->SourceDIName.ToString(), MaxAttributes, InstanceData->TotalNumAttributes);
+			return false;
+		}
+
+		if (InstanceData->TotalNumAttributes <= 0)
 		{
-			UE_LOG(LogNiagara, Error, TEXT("Invalid number of attributes defined on %s... max is %i, num defined is %i"), *FNiagaraUtilities::SystemInstanceIDToString(SystemInstance->GetId()), MaxAttributes, InstanceData->TotalNumAttributes);
+			UE_LOG(LogNiagara, Warning, TEXT("Invalid number of attributes defined on %s... max is %d, num defined is %d"), *Proxy->SourceDIName.ToString(), MaxAttributes, InstanceData->TotalNumAttributes);
 			return false;
 		}
 
