Niagara - Fix a crash issue when making a user parameter from a data interface input.

diff --git a/Engine/Plugins/FX/Niagara/Source/NiagaraEditor/Private/ViewModels/Stack/NiagaraStackFunctionInput.cpp b/Engine/Plugins/FX/Niagara/Source/NiagaraEditor/Private/ViewModels/Stack/NiagaraStackFunctionInput.cpp
--- a/Engine/Plugins/FX/Niagara/Source/NiagaraEditor/Private/ViewModels/Stack/NiagaraStackFunctionInput.cpp
+++ b/Engine/Plugins/FX/Niagara/Source/NiagaraEditor/Private/ViewModels/Stack/NiagaraStackFunctionInput.cpp
@@ -1132,16 +1132,30 @@ void UNiagaraStackFunctionInput::SetLinkedValueHandle(const FNiagaraParameterHan
 			FNiagaraParameterStore& UserParameters = GetSystemViewModel()->GetSystem().GetExposedParameters();
 			if (UserParameters.IndexOf(LinkedParameterVariable) == INDEX_NONE)
 			{
-				if (InputValues.Mode == EValueMode::Local && InputValues.LocalStruct.IsValid())
+				if (InputType.IsDataInterface())
 				{
-					// If the current value is local, and valid transfer that value to the user parameter.
-					LinkedParameterVariable.SetData(InputValues.LocalStruct->GetStructMemory());
+					int32 DataInterfaceOffset;
+					bool bInitialize = true;
+					bool bTriggerRebind = true;
+					UserParameters.AddParameter(LinkedParameterVariable, bInitialize, bTriggerRebind, &DataInterfaceOffset);
+					if (InputValues.Mode == EValueMode::Data && InputValues.DataObject.IsValid())
+					{
+						InputValues.DataObject->CopyTo(UserParameters.GetDataInterface(DataInterfaceOffset));
+					}
 				}
 				else
 				{
-					FNiagaraEditorUtilities::ResetVariableToDefaultValue(LinkedParameterVariable);
+					if (InputValues.Mode == EValueMode::Local && InputValues.LocalStruct.IsValid())
+					{
+						// If the current value is local, and valid transfer that value to the user parameter.
+						LinkedParameterVariable.SetData(InputValues.LocalStruct->GetStructMemory());
+					}
+					else
+					{
+						FNiagaraEditorUtilities::ResetVariableToDefaultValue(LinkedParameterVariable);
+					}
+					UserParameters.SetParameterData(LinkedParameterVariable.GetData(), LinkedParameterVariable, true);
 				}
-				UserParameters.SetParameterData(LinkedParameterVariable.GetData(), LinkedParameterVariable, true);
 			}
 		}
 
