Fix for issue with component renderers breaking scalalbility culling by constantly reactivating it's components
Also adding refresh path for scalability managers to clear some cached data when system settings change.
Occasionally some edits not being applied correctly.

diff --git a/Engine/Plugins/FX/Niagara/Source/Niagara/Private/NiagaraComponentRendererProperties.cpp b/Engine/Plugins/FX/Niagara/Source/Niagara/Private/NiagaraComponentRendererProperties.cpp
--- a/Engine/Plugins/FX/Niagara/Source/Niagara/Private/NiagaraComponentRendererProperties.cpp
+++ b/Engine/Plugins/FX/Niagara/Source/Niagara/Private/NiagaraComponentRendererProperties.cpp
@@ -147,7 +147,7 @@ FNiagaraTypeDefinition UNiagaraComponentRendererProperties::GetFQuatDef()
 TArray<TWeakObjectPtr<UNiagaraComponentRendererProperties>> UNiagaraComponentRendererProperties::ComponentRendererPropertiesToDeferredInit;
 
 UNiagaraComponentRendererProperties::UNiagaraComponentRendererProperties()
-	: ComponentCountLimit(15), bAssignComponentsOnParticleID(true), bOnlyCreateComponentsOnParticleSpawn(true)
+	: ComponentCountLimit(15), bAssignComponentsOnParticleID(true), bOnlyCreateComponentsOnParticleSpawn(true), bOnlyActivateNewlyAquiredComponents(true)
 #if WITH_EDITORONLY_DATA
 	, bVisualizeComponents(true)
 #endif
diff --git a/Engine/Plugins/FX/Niagara/Source/Niagara/Private/NiagaraRendererComponents.cpp b/Engine/Plugins/FX/Niagara/Source/Niagara/Private/NiagaraRendererComponents.cpp
--- a/Engine/Plugins/FX/Niagara/Source/Niagara/Private/NiagaraRendererComponents.cpp
+++ b/Engine/Plugins/FX/Niagara/Source/Niagara/Private/NiagaraRendererComponents.cpp
@@ -447,6 +447,7 @@ void FNiagaraRendererComponents::PostSystemTick_GameThread(const UNiagaraRendere
 		}
 
 		// Acquire a component for this particle
+		bool bNewlyAcquiredComponent = false;
 		USceneComponent* SceneComponent = nullptr;
 		if (PoolIndex == -1)
 		{
@@ -460,6 +461,7 @@ void FNiagaraRendererComponents::PostSystemTick_GameThread(const UNiagaraRendere
 			{
 				PoolIndex = FreeList.Pop(false);
 			}
+			bNewlyAcquiredComponent = true;
 		}
 
 		if (PoolIndex >= 0)
@@ -513,11 +515,15 @@ void FNiagaraRendererComponents::PostSystemTick_GameThread(const UNiagaraRendere
 		FNiagaraLWCConverter LwcConverter = SystemInstance->GetLWCConverter(Emitter->GetCachedEmitter()->bLocalSpace);
 		TickPropertyBindings(Properties, SceneComponent, Data, ParticleIndex, PoolEntry, LwcConverter);
 
-		// activate the component
-		if (!SceneComponent->IsActive())
+		// Activate the component.
+		// If components are assigned by ID then we can optionally do this only on the first frame the component is acquired by a particle.
+		if(!Properties->bAssignComponentsOnParticleID || bNewlyAcquiredComponent || !Properties->bOnlyActivateNewlyAquiredComponents)
 		{
-			SceneComponent->SetVisibility(true, true);
-			SceneComponent->Activate(false);
+			if (!SceneComponent->IsActive())
+			{
+				SceneComponent->SetVisibility(true, true);
+				SceneComponent->Activate(false);
+			}
 		}
 
 		PoolEntry.LastAssignedToParticleID = ParticleID;
@@ -762,4 +768,4 @@ void FNiagaraRendererComponents::ResetComponentPool(bool bResetOwner)
 			OwnerActor->Destroy();
 		}
 	}
-}
\ No newline at end of file
+}
diff --git a/Engine/Plugins/FX/Niagara/Source/Niagara/Private/NiagaraScalabilityManager.cpp b/Engine/Plugins/FX/Niagara/Source/Niagara/Private/NiagaraScalabilityManager.cpp
--- a/Engine/Plugins/FX/Niagara/Source/Niagara/Private/NiagaraScalabilityManager.cpp
+++ b/Engine/Plugins/FX/Niagara/Source/Niagara/Private/NiagaraScalabilityManager.cpp
@@ -568,6 +568,12 @@ FNiagaraScalabilitySystemData& FNiagaraScalabilityManager::GetSystemData(int32 C
 
 #if WITH_EDITOR
 void FNiagaraScalabilityManager::OnSystemPostChange(UNiagaraSystem* System)
+{
+	RefreshSystemScalabilitySettings(System);
+}
+#endif//WITH_EDITOR
+
+void FNiagaraScalabilityManager::RefreshSystemScalabilitySettings(UNiagaraSystem* System)
 {
 	//Clear and refresh all cached system data.
 	SystemDataIndexMap.Reset();
@@ -578,7 +584,6 @@ void FNiagaraScalabilityManager::OnSystemPostChange(UNiagaraSystem* System)
 		GetSystemData(CompIdx, true);
 	}
 }
-#endif//WITH_EDITOR
 
 #if DEBUG_SCALABILITY_STATE
 
diff --git a/Engine/Plugins/FX/Niagara/Source/Niagara/Private/NiagaraSystem.cpp b/Engine/Plugins/FX/Niagara/Source/Niagara/Private/NiagaraSystem.cpp
--- a/Engine/Plugins/FX/Niagara/Source/Niagara/Private/NiagaraSystem.cpp
+++ b/Engine/Plugins/FX/Niagara/Source/Niagara/Private/NiagaraSystem.cpp
@@ -3581,6 +3581,8 @@ void UNiagaraSystem::ResolveScalabilitySettings()
 		};
 		ForEachScript(ScriptUsesSigIndex);
 	}
+
+	FNiagaraWorldManager::RefreshSystemScalabilitySettingsForAllWorlds(this);
 }
 
 void UNiagaraSystem::OnScalabilityCVarChanged()
diff --git a/Engine/Plugins/FX/Niagara/Source/Niagara/Private/NiagaraWorldManager.cpp b/Engine/Plugins/FX/Niagara/Source/Niagara/Private/NiagaraWorldManager.cpp
--- a/Engine/Plugins/FX/Niagara/Source/Niagara/Private/NiagaraWorldManager.cpp
+++ b/Engine/Plugins/FX/Niagara/Source/Niagara/Private/NiagaraWorldManager.cpp
@@ -1521,6 +1521,26 @@ bool FNiagaraWorldManager::GetScalabilityState(UNiagaraComponent* Component, FNi
 	return false;
 }
 
+void FNiagaraWorldManager::RefreshSystemScalabilitySettingsForAllWorlds(UNiagaraSystem* System)
+{
+	for (auto& Pair : WorldManagers)
+	{
+		if (Pair.Value)
+		{
+			Pair.Value->RefreshSystemScalabilitySettings(System);
+		}
+	}
+}
+
+void FNiagaraWorldManager::RefreshSystemScalabilitySettings(UNiagaraSystem* System)
+{
+	for (auto& Pair : ScalabilityManagers)
+	{
+		FNiagaraScalabilityManager& ScalabilityMan = Pair.Value;
+		ScalabilityMan.RefreshSystemScalabilitySettings(System);
+	}
+}
+
 void FNiagaraWorldManager::PrimePoolForAllWorlds(UNiagaraSystem* System)
 {
 	if (GNigaraAllowPrimedPools)
diff --git a/Engine/Plugins/FX/Niagara/Source/Niagara/Public/NiagaraComponentRendererProperties.h b/Engine/Plugins/FX/Niagara/Source/Niagara/Public/NiagaraComponentRendererProperties.h
--- a/Engine/Plugins/FX/Niagara/Source/Niagara/Public/NiagaraComponentRendererProperties.h
+++ b/Engine/Plugins/FX/Niagara/Source/Niagara/Public/NiagaraComponentRendererProperties.h
@@ -107,6 +107,16 @@ public:
 	 * limit was reached) then it will be blocked from spawning a component on subsequent frames. */
 	UPROPERTY(EditAnywhere, AdvancedDisplay, Category = "Component Rendering", meta = (EditCondition = "bAssignComponentsOnParticleID"))
 	bool bOnlyCreateComponentsOnParticleSpawn;
+	
+	/** 
+	If true then components will only be activated when newly acquired. e.g. on particle spawn or when the particle enables/disables the component.
+	If false, components will be always kept active while they are used by an enabled particle.
+	This can be useful for component types that can internally activate and deactivate during their lifetimes while still being used by an enabled particle.
+	For example NiagaraComponents that deactivate/reactivate according to scalability settings.
+	This setting is only valid when bAssignComponentsOnParticleID is true.
+	*/
+	UPROPERTY(EditAnywhere, AdvancedDisplay, Category = "Component Rendering", meta = (EditCondition = "bAssignComponentsOnParticleID"))
+	bool bOnlyActivateNewlyAquiredComponents;
 
 #if WITH_EDITORONLY_DATA
 
diff --git a/Engine/Plugins/FX/Niagara/Source/Niagara/Public/NiagaraScalabilityManager.h b/Engine/Plugins/FX/Niagara/Source/Niagara/Public/NiagaraScalabilityManager.h
--- a/Engine/Plugins/FX/Niagara/Source/Niagara/Public/NiagaraScalabilityManager.h
+++ b/Engine/Plugins/FX/Niagara/Source/Niagara/Public/NiagaraScalabilityManager.h
@@ -86,6 +86,9 @@ struct FNiagaraScalabilityManager
 #if WITH_EDITOR
 	void OnSystemPostChange(UNiagaraSystem* System);
 #endif//WITH_EDITOR
+
+	void RefreshSystemScalabilitySettings(UNiagaraSystem* System);
+
 private: 
 	void UnregisterAt(int32 IndexToRemove);
 	bool HasPendingUpdates() const { return DefaultContext.ComponentRequiresUpdate.Num() > 0; }
diff --git a/Engine/Plugins/FX/Niagara/Source/Niagara/Public/NiagaraWorldManager.h b/Engine/Plugins/FX/Niagara/Source/Niagara/Public/NiagaraWorldManager.h
--- a/Engine/Plugins/FX/Niagara/Source/Niagara/Public/NiagaraWorldManager.h
+++ b/Engine/Plugins/FX/Niagara/Source/Niagara/Public/NiagaraWorldManager.h
@@ -174,6 +174,9 @@ public:
 	void PrimePoolForAllSystems();
 	void PrimePool(UNiagaraSystem* System);
 
+	static void RefreshSystemScalabilitySettingsForAllWorlds(UNiagaraSystem* System);
+	void RefreshSystemScalabilitySettings(UNiagaraSystem* System);
+
 	void SetDebugPlaybackMode(ENiagaraDebugPlaybackMode Mode) { RequestedDebugPlaybackMode = Mode; }
 	ENiagaraDebugPlaybackMode GetDebugPlaybackMode() const { return DebugPlaybackMode; }
 
