Some additional fixes for attribute trimming
-prevents disabled renderers or emitters from contributing to the attribute trimming process
-adds a new flag to the system that allows trimming to happen even though an emitter is the source for particle reads (it's up to the user to manually preserve attributes for now)
-Fix attribute trimming for impure function inputs

diff --git a/Engine/Plugins/FX/Niagara/Source/Niagara/Classes/NiagaraSystem.h b/Engine/Plugins/FX/Niagara/Source/Niagara/Classes/NiagaraSystem.h
--- a/Engine/Plugins/FX/Niagara/Source/Niagara/Classes/NiagaraSystem.h
+++ b/Engine/Plugins/FX/Niagara/Source/Niagara/Classes/NiagaraSystem.h
@@ -567,6 +567,10 @@ public:
 	UPROPERTY(EditAnywhere, AdvancedDisplay, Category = "Performance", meta = (SkipSystemResetOnChange = "true"))
 	uint32 bTrimAttributesOnCook : 1;
 
+	/** If true ParticleReads will not absolutely prevent attribute trimming - User must ensure that the appropriate attributes are preserved on the source emitter! */
+	UPROPERTY(EditAnywhere, AdvancedDisplay, Category = "Performance")
+	uint32 bIgnoreParticleReadsForAttributeTrim : 1;
+
 	/** If true all debug switches will be ignored in the editor. */
 	UPROPERTY(EditAnywhere, AdvancedDisplay, Category = "Performance")
 	uint32 bDisableDebugSwitches : 1;
diff --git a/Engine/Plugins/FX/Niagara/Source/Niagara/Private/NiagaraScript.cpp b/Engine/Plugins/FX/Niagara/Source/Niagara/Private/NiagaraScript.cpp
--- a/Engine/Plugins/FX/Niagara/Source/Niagara/Private/NiagaraScript.cpp
+++ b/Engine/Plugins/FX/Niagara/Source/Niagara/Private/NiagaraScript.cpp
@@ -822,6 +822,10 @@ void UNiagaraScript::ComputeVMCompilationId(FNiagaraVMExecutableDataId& Id, FGui
 			{
 				Id.AdditionalDefines.Add(TEXT("CompressAttributes"));
 			}
+			if (EmitterOwner->bIgnoreParticleReadsForAttributeTrim)
+			{
+				Id.AdditionalDefines.Add(TEXT("IgnoreParticleReadsForAttributeTrim"));
+			}
 
 			bool TrimAttributes = EmitterOwner->bTrimAttributes;
 			if (TrimAttributes)
@@ -849,12 +853,15 @@ void UNiagaraScript::ComputeVMCompilationId(FNiagaraVMExecutableDataId& Id, FGui
 				};
 
 				// if this emitter is being referenced by another emitter (PartilceRead) then don't worry about trimming attributes
-				for (const FNiagaraEmitterHandle& EmitterHandle : EmitterOwner->GetEmitterHandles())
+				if (!EmitterOwner->bIgnoreParticleReadsForAttributeTrim)
 				{
-					if (!TrimAttributesSupported(EmitterHandle.GetInstance()))
+					for (const FNiagaraEmitterHandle& EmitterHandle : EmitterOwner->GetEmitterHandles())
 					{
-						TrimAttributes = false;
-						break;
+						if (EmitterHandle.GetIsEnabled() && !TrimAttributesSupported(EmitterHandle.GetInstance()))
+						{
+							TrimAttributes = false;
+							break;
+						}
 					}
 				}
 			}
@@ -883,9 +890,12 @@ void UNiagaraScript::ComputeVMCompilationId(FNiagaraVMExecutableDataId& Id, FGui
 				// Now preserve the attributes that have been defined on the renderers in use
 				for (UNiagaraRendererProperties* RendererProperty : Emitter->GetRenderers())
 				{
-					for (const FNiagaraVariable& BoundAttribute : RendererProperty->GetBoundAttributes())
+					if (RendererProperty->bIsEnabled)
 					{
-						AddAttributeToPreserve(BoundAttribute.GetName());
+						for (const FNiagaraVariable& BoundAttribute : RendererProperty->GetBoundAttributes())
+						{
+							AddAttributeToPreserve(BoundAttribute.GetName());
+						}
 					}
 				}
 
diff --git a/Engine/Plugins/FX/Niagara/Source/Niagara/Private/NiagaraSystem.cpp b/Engine/Plugins/FX/Niagara/Source/Niagara/Private/NiagaraSystem.cpp
--- a/Engine/Plugins/FX/Niagara/Source/Niagara/Private/NiagaraSystem.cpp
+++ b/Engine/Plugins/FX/Niagara/Source/Niagara/Private/NiagaraSystem.cpp
@@ -110,6 +110,7 @@ UNiagaraSystem::UNiagaraSystem(const FObjectInitializer& ObjectInitializer)
 , bBakeOutRapidIterationOnCook(true)
 , bTrimAttributes(false)
 , bTrimAttributesOnCook(true)
+, bIgnoreParticleReadsForAttributeTrim(false)
 , bDisableDebugSwitches(false)
 , bDisableDebugSwitchesOnCook(true)
 #endif
