Add an optional format of the MapFor node that includes support for a continue (MapForWithContinue)
-conditional logic is supplied with a new pin
-new node is also available for returning the current index in the MapFor iteration

diff --git a/Engine/Plugins/FX/Niagara/Source/NiagaraEditor/Private/EdGraphSchema_Niagara.cpp b/Engine/Plugins/FX/Niagara/Source/NiagaraEditor/Private/EdGraphSchema_Niagara.cpp
--- a/Engine/Plugins/FX/Niagara/Source/NiagaraEditor/Private/EdGraphSchema_Niagara.cpp
+++ b/Engine/Plugins/FX/Niagara/Source/NiagaraEditor/Private/EdGraphSchema_Niagara.cpp
@@ -755,6 +755,16 @@ TArray<TSharedPtr<FNiagaraAction_NewNode>> UEdGraphSchema_Niagara::GetGraphActio
 			UNiagaraNodeParameterMapFor* BaseNode = NewObject<UNiagaraNodeParameterMapFor>(OwnerOfTemporaries);
 			AddNewNodeMenuAction(NewActions, BaseNode, FText::FromString(Name), ENiagaraMenuSections::General, {MenuCat.ToString()}, FText::GetEmpty(), FText::GetEmpty());
 		}
+		{
+			FString Name = TEXT("Parameter Map For With Continue");
+			UNiagaraNodeParameterMapForWithContinue* BaseNode = NewObject<UNiagaraNodeParameterMapForWithContinue>(OwnerOfTemporaries);
+			AddNewNodeMenuAction(NewActions, BaseNode, FText::FromString(Name), ENiagaraMenuSections::General, { MenuCat.ToString() }, FText::GetEmpty(), FText::GetEmpty());
+		}
+		{
+			FString Name = TEXT("Parameter Map For Current Index");
+			UNiagaraNodeParameterMapForIndex* BaseNode = NewObject<UNiagaraNodeParameterMapForIndex>(OwnerOfTemporaries);
+			AddNewNodeMenuAction(NewActions, BaseNode, FText::FromString(Name), ENiagaraMenuSections::General, { MenuCat.ToString() }, FText::GetEmpty(), FText::GetEmpty());
+		}
 	}
 
 	// Handle comment nodes
diff --git a/Engine/Plugins/FX/Niagara/Source/NiagaraEditor/Private/NiagaraHlslTranslator.cpp b/Engine/Plugins/FX/Niagara/Source/NiagaraEditor/Private/NiagaraHlslTranslator.cpp
--- a/Engine/Plugins/FX/Niagara/Source/NiagaraEditor/Private/NiagaraHlslTranslator.cpp
+++ b/Engine/Plugins/FX/Niagara/Source/NiagaraEditor/Private/NiagaraHlslTranslator.cpp
@@ -4828,7 +4828,21 @@ void FHlslNiagaraTranslator::ParameterMapForBegin(UNiagaraNodeParameterMapFor* F
 {
 	NIAGARA_SCOPE_CYCLE_COUNTER(STAT_NiagaraEditor_HlslTranslator_MapForBegin);
 
-	AddBodyChunk(TEXT(""), TEXT("for(int index = 0; index < {0}; ++index)\n\t{"), FNiagaraTypeDefinition::GetIntDef(), IterationCount, false, false);
+	const int32 IndexChunkIndex = AddBodyChunk(GetUniqueSymbolName(TEXT("Index")), TEXT(""), FNiagaraTypeDefinition::GetIntDef(), true);
+	ParameterMapForIndexStack.Push(IndexChunkIndex);
+
+	TArray<int32> SourceChunks;
+	SourceChunks.Add(IndexChunkIndex);
+	SourceChunks.Add(IterationCount);
+
+	AddBodyChunk(TEXT(""), TEXT("for({0} = 0; {0} < {1}; ++{0})\n\t{"), FNiagaraTypeDefinition::GetIntDef(), SourceChunks, false, false);
+}
+
+void FHlslNiagaraTranslator::ParameterMapForContinue(UNiagaraNodeParameterMapFor* ForNode, int32 IterationEnabled)
+{
+	NIAGARA_SCOPE_CYCLE_COUNTER(STAT_NiagaraEditor_HlslTranslator_MapForBegin);
+
+	AddBodyChunk(TEXT(""), TEXT("if (!{0}) continue;"), FNiagaraTypeDefinition::GetBoolDef(), IterationEnabled, false, false);
 }
 
 void FHlslNiagaraTranslator::ParameterMapForEnd(UNiagaraNodeParameterMapFor* ForNode)
@@ -4836,6 +4850,18 @@ void FHlslNiagaraTranslator::ParameterMapForEnd(UNiagaraNodeParameterMapFor* For
 	NIAGARA_SCOPE_CYCLE_COUNTER(STAT_NiagaraEditor_HlslTranslator_MapForEnd);
 
 	AddBodyChunk(TEXT(""), TEXT("}"), FNiagaraTypeDefinition::GetIntDef(), false, false);
+
+	ParameterMapForIndexStack.Pop();
+}
+
+int32 FHlslNiagaraTranslator::ParameterMapForInnerIndex() const
+{
+	if (ParameterMapForIndexStack.Num())
+	{
+		return ParameterMapForIndexStack.Last();
+	}
+
+	return INDEX_NONE;
 }
 
 void FHlslNiagaraTranslator::ParameterMapSet(UNiagaraNodeParameterMapSet* SetNode, TArrayView<const FCompiledPin> Inputs, TArray<int32>& Outputs)
diff --git a/Engine/Plugins/FX/Niagara/Source/NiagaraEditor/Private/NiagaraHlslTranslator.h b/Engine/Plugins/FX/Niagara/Source/NiagaraEditor/Private/NiagaraHlslTranslator.h
--- a/Engine/Plugins/FX/Niagara/Source/NiagaraEditor/Private/NiagaraHlslTranslator.h
+++ b/Engine/Plugins/FX/Niagara/Source/NiagaraEditor/Private/NiagaraHlslTranslator.h
@@ -555,7 +555,9 @@ public:
 	virtual void Emitter(class UNiagaraNodeEmitter* GetNode, TArray<int32>& Inputs, TArray<int32>& Outputs);
 
 	virtual void ParameterMapForBegin(class UNiagaraNodeParameterMapFor* ForNode, int32 IterationCount);
+	virtual void ParameterMapForContinue(class UNiagaraNodeParameterMapFor* ForNode, int32 IterationEnabled);
 	virtual void ParameterMapForEnd(class UNiagaraNodeParameterMapFor* ForNode);
+	virtual int32 ParameterMapForInnerIndex() const;
 
 	void DefineInterpolatedParametersFunction(FString &HlslOutput);
 	void DefinePreviousParametersFunction(FString& HlslOutput, TArray<TArray<FNiagaraVariable>>& DataSetVariables, TMap<FNiagaraDataSetID, int32>& DataSetReads, TMap<FNiagaraDataSetID, int32>& DataSetWrites);
@@ -789,6 +791,8 @@ private:
 	// Map ofthe primary output variable description to the actual chunk id that wrote to it.
 	TMap<FNiagaraVariable, int32> UniqueVarToChunk;
 
+	TArray<int32> ParameterMapForIndexStack;
+
 	// Strings to be inserted within the main function
 	TArray<TArray<FString>> PerStageMainPreSimulateChunks;
 
diff --git a/Engine/Plugins/FX/Niagara/Source/NiagaraEditor/Private/NiagaraNodeParameterMapFor.cpp b/Engine/Plugins/FX/Niagara/Source/NiagaraEditor/Private/NiagaraNodeParameterMapFor.cpp
--- a/Engine/Plugins/FX/Niagara/Source/NiagaraEditor/Private/NiagaraNodeParameterMapFor.cpp
+++ b/Engine/Plugins/FX/Niagara/Source/NiagaraEditor/Private/NiagaraNodeParameterMapFor.cpp
@@ -61,4 +61,87 @@ FText UNiagaraNodeParameterMapFor::GetNodeTitle(ENodeTitleType::Type TitleType)
 	return LOCTEXT("UNiagaraNodeParameterMapForName", "Map For");
 }
 
+UNiagaraNodeParameterMapForWithContinue::UNiagaraNodeParameterMapForWithContinue() : UNiagaraNodeParameterMapFor()
+{
+	UEdGraphNode::NodeUpgradeMessage = LOCTEXT("NodeExperimental", "This node is marked as experimental, use with care!");
+}
+
+void UNiagaraNodeParameterMapForWithContinue::Compile(FHlslNiagaraTranslator* Translator, TArray<int32>& Outputs)
+{
+	FPinCollectorArray InputPins;
+	GetInputPins(InputPins);
+	if (Translator)
+	{
+		if (Translator->GetSimulationTarget() == ENiagaraSimTarget::GPUComputeSim)
+		{
+			const int32 IterationCount = Translator->CompilePin(InputPins[1]);
+
+			Translator->ParameterMapForBegin(this, IterationCount);
+			if (InputPins[2])
+			{
+				const int32 IterationEnabledChunk = Translator->CompilePin(InputPins[2]);
+				Translator->ParameterMapForContinue(this, IterationEnabledChunk);
+			}
+
+			UNiagaraNodeParameterMapSet::Compile(Translator, Outputs);
+
+			Translator->ParameterMapForEnd(this);
+		}
+		else
+		{
+			UNiagaraNodeParameterMapSet::Compile(Translator, Outputs);
+			//Translator->Message(FNiagaraCompileEventSeverity::Log,LOCTEXT("UnsupportedParamMapFor", "Parameter map for is not yet supported on cpu."), this, nullptr);
+		}
+	}
+}
+
+void UNiagaraNodeParameterMapForWithContinue::AllocateDefaultPins()
+{
+	PinPendingRename = nullptr;
+	const UEdGraphSchema_Niagara* Schema = GetDefault<UEdGraphSchema_Niagara>();
+	CreatePin(EGPD_Input, Schema->TypeDefinitionToPinType(FNiagaraTypeDefinition::GetParameterMapDef()), *UNiagaraNodeParameterMapBase::SourcePinName.ToString());
+	CreatePin(EGPD_Input, Schema->TypeDefinitionToPinType(FNiagaraTypeDefinition::GetIntDef()), TEXT("Module.Iteration Count"));
+	CreatePin(EGPD_Input, Schema->TypeDefinitionToPinType(FNiagaraTypeDefinition::GetBoolDef()), TEXT("Module.Iteration Enabled"));
+	CreatePin(EGPD_Output, Schema->TypeDefinitionToPinType(FNiagaraTypeDefinition::GetParameterMapDef()), *UNiagaraNodeParameterMapBase::DestPinName.ToString());
+	CreateAddPin(EGPD_Input);
+}
+
+FText UNiagaraNodeParameterMapForWithContinue::GetNodeTitle(ENodeTitleType::Type TitleType) const
+{
+	return LOCTEXT("UNiagaraNodeParameterMapForName", "Map For With Continue");
+}
+
+UNiagaraNodeParameterMapForIndex::UNiagaraNodeParameterMapForIndex() : UNiagaraNode()
+{
+	UEdGraphNode::NodeUpgradeMessage = LOCTEXT("NodeExperimental", "This node is marked as experimental, use with care!");
+}
+
+void UNiagaraNodeParameterMapForIndex::AllocateDefaultPins()
+{
+	const UEdGraphSchema_Niagara* Schema = GetDefault<UEdGraphSchema_Niagara>();
+	CreatePin(EGPD_Output, Schema->TypeDefinitionToPinType(FNiagaraTypeDefinition::GetIntDef()), TEXT("Module.Current Iteration"));
+}
+
+FText UNiagaraNodeParameterMapForIndex::GetNodeTitle(ENodeTitleType::Type TitleType) const
+{
+	return LOCTEXT("UNiagaraNodeParameterMapForIndex", "Map For Index");
+}
+
+void UNiagaraNodeParameterMapForIndex::Compile(class FHlslNiagaraTranslator* Translator, TArray<int32>& Outputs)
+{
+	if (Translator)
+	{
+		if (Translator->GetSimulationTarget() == ENiagaraSimTarget::GPUComputeSim)
+		{
+			Outputs.Add(Translator->ParameterMapForInnerIndex());
+		}
+		else
+		{
+			FNiagaraVariable Constant(FNiagaraTypeDefinition::GetIntDef(), TEXT("Constant"));
+			Constant.SetValue(0);
+			Outputs.Add(Translator->GetConstant(Constant));
+		}
+	}
+}
+
 #undef LOCTEXT_NAMESPACE
diff --git a/Engine/Plugins/FX/Niagara/Source/NiagaraEditor/Private/NiagaraNodeParameterMapFor.h b/Engine/Plugins/FX/Niagara/Source/NiagaraEditor/Private/NiagaraNodeParameterMapFor.h
--- a/Engine/Plugins/FX/Niagara/Source/NiagaraEditor/Private/NiagaraNodeParameterMapFor.h
+++ b/Engine/Plugins/FX/Niagara/Source/NiagaraEditor/Private/NiagaraNodeParameterMapFor.h
@@ -21,3 +21,29 @@ public:
 
 	virtual FText GetNodeTitle(ENodeTitleType::Type TitleType) const override;
 };
+
+UCLASS()
+class UNiagaraNodeParameterMapForWithContinue : public UNiagaraNodeParameterMapFor
+{
+public:
+	GENERATED_BODY()
+
+	UNiagaraNodeParameterMapForWithContinue();
+
+	virtual void AllocateDefaultPins() override;
+	virtual void Compile(class FHlslNiagaraTranslator* Translator, TArray<int32>& Outputs) override;
+	virtual FText GetNodeTitle(ENodeTitleType::Type TitleType) const override;
+};
+
+UCLASS()
+class UNiagaraNodeParameterMapForIndex : public UNiagaraNode
+{
+public:
+	GENERATED_BODY()
+
+	UNiagaraNodeParameterMapForIndex();
+
+	virtual void AllocateDefaultPins() override;
+	virtual FText GetNodeTitle(ENodeTitleType::Type TitleType) const override;
+	virtual void Compile(class FHlslNiagaraTranslator* Translator, TArray<int32>& Outputs) override;
+};
