Fixed threading issue with scalabiltiy manager refresh of cached system data.

diff --git a/Engine/Plugins/FX/Niagara/Source/Niagara/Private/NiagaraScalabilityManager.cpp b/Engine/Plugins/FX/Niagara/Source/Niagara/Private/NiagaraScalabilityManager.cpp
--- a/Engine/Plugins/FX/Niagara/Source/Niagara/Private/NiagaraScalabilityManager.cpp
+++ b/Engine/Plugins/FX/Niagara/Source/Niagara/Private/NiagaraScalabilityManager.cpp
@@ -73,6 +73,7 @@ static int32 GetMaxUpdatesPerFrame(const UNiagaraEffectType* EffectType, int32 I
 FNiagaraScalabilityManager::FNiagaraScalabilityManager()
 	: EffectType(nullptr)
 	, LastUpdateTime(0.0f)
+	, bRefreshCachedSystemData(false)
 {
 
 }
@@ -448,6 +449,20 @@ void FNiagaraScalabilityManager::UpdateInternal(FNiagaraWorldManager* WorldMan,
 
 void FNiagaraScalabilityManager::Update(FNiagaraWorldManager* WorldMan, float DeltaSeconds, bool bNewOnly)
 {
+	if (bRefreshCachedSystemData)
+	{	
+		bRefreshCachedSystemData = false;
+
+		//Clear and refresh all cached system data.
+		SystemDataIndexMap.Reset();
+		SystemData.Reset();
+
+		for (int32 CompIdx = 0; CompIdx < ManagedComponents.Num(); ++CompIdx)
+		{
+			GetSystemData(CompIdx, true);
+		}
+	}
+
 	//Paranoia code in case the EffectType is GCd from under us.
 	if (EffectType == nullptr)
 	{
@@ -569,20 +584,13 @@ FNiagaraScalabilitySystemData& FNiagaraScalabilityManager::GetSystemData(int32 C
 #if WITH_EDITOR
 void FNiagaraScalabilityManager::OnSystemPostChange(UNiagaraSystem* System)
 {
-	RefreshSystemScalabilitySettings(System);
+	InvalidateCachedSystemData();
 }
 #endif//WITH_EDITOR
 
-void FNiagaraScalabilityManager::RefreshSystemScalabilitySettings(UNiagaraSystem* System)
+void FNiagaraScalabilityManager::InvalidateCachedSystemData()
 {
-	//Clear and refresh all cached system data.
-	SystemDataIndexMap.Reset();
-	SystemData.Reset();
-
-	for (int32 CompIdx = 0; CompIdx < ManagedComponents.Num(); ++CompIdx)
-	{
-		GetSystemData(CompIdx, true);
-	}
+	bRefreshCachedSystemData = true;
 }
 
 #if DEBUG_SCALABILITY_STATE
diff --git a/Engine/Plugins/FX/Niagara/Source/Niagara/Private/NiagaraSystem.cpp b/Engine/Plugins/FX/Niagara/Source/Niagara/Private/NiagaraSystem.cpp
--- a/Engine/Plugins/FX/Niagara/Source/Niagara/Private/NiagaraSystem.cpp
+++ b/Engine/Plugins/FX/Niagara/Source/Niagara/Private/NiagaraSystem.cpp
@@ -3583,7 +3583,7 @@ void UNiagaraSystem::ResolveScalabilitySettings()
 		ForEachScript(ScriptUsesSigIndex);
 	}
 
-	FNiagaraWorldManager::RefreshSystemScalabilitySettingsForAllWorlds(this);
+	FNiagaraWorldManager::InvalidateCachedSystemScalabilityDataForAllWorlds();
 }
 
 void UNiagaraSystem::OnScalabilityCVarChanged()
diff --git a/Engine/Plugins/FX/Niagara/Source/Niagara/Private/NiagaraWorldManager.cpp b/Engine/Plugins/FX/Niagara/Source/Niagara/Private/NiagaraWorldManager.cpp
--- a/Engine/Plugins/FX/Niagara/Source/Niagara/Private/NiagaraWorldManager.cpp
+++ b/Engine/Plugins/FX/Niagara/Source/Niagara/Private/NiagaraWorldManager.cpp
@@ -1521,23 +1521,23 @@ bool FNiagaraWorldManager::GetScalabilityState(UNiagaraComponent* Component, FNi
 	return false;
 }
 
-void FNiagaraWorldManager::RefreshSystemScalabilitySettingsForAllWorlds(UNiagaraSystem* System)
+void FNiagaraWorldManager::InvalidateCachedSystemScalabilityDataForAllWorlds()
 {
 	for (auto& Pair : WorldManagers)
 	{
 		if (Pair.Value)
 		{
-			Pair.Value->RefreshSystemScalabilitySettings(System);
+			Pair.Value->InvalidateCachedSystemScalabilityData();
 		}
 	}
 }
 
-void FNiagaraWorldManager::RefreshSystemScalabilitySettings(UNiagaraSystem* System)
+void FNiagaraWorldManager::InvalidateCachedSystemScalabilityData()
 {
 	for (auto& Pair : ScalabilityManagers)
 	{
 		FNiagaraScalabilityManager& ScalabilityMan = Pair.Value;
-		ScalabilityMan.RefreshSystemScalabilitySettings(System);
+		ScalabilityMan.InvalidateCachedSystemData();
 	}
 }
 
diff --git a/Engine/Plugins/FX/Niagara/Source/Niagara/Public/NiagaraScalabilityManager.h b/Engine/Plugins/FX/Niagara/Source/Niagara/Public/NiagaraScalabilityManager.h
--- a/Engine/Plugins/FX/Niagara/Source/Niagara/Public/NiagaraScalabilityManager.h
+++ b/Engine/Plugins/FX/Niagara/Source/Niagara/Public/NiagaraScalabilityManager.h
@@ -87,7 +87,7 @@ struct FNiagaraScalabilityManager
 	void OnSystemPostChange(UNiagaraSystem* System);
 #endif//WITH_EDITOR
 
-	void RefreshSystemScalabilitySettings(UNiagaraSystem* System);
+	void InvalidateCachedSystemData();
 
 private: 
 	void UnregisterAt(int32 IndexToRemove);
@@ -101,4 +101,6 @@ private:
 	FNiagaraScalabilitySystemData& GetSystemData(int32 ComponentIndex, bool bForceRefresh=false);
 
 	FComponentIterationContext DefaultContext;
+
+	bool bRefreshCachedSystemData;
 };
diff --git a/Engine/Plugins/FX/Niagara/Source/Niagara/Public/NiagaraWorldManager.h b/Engine/Plugins/FX/Niagara/Source/Niagara/Public/NiagaraWorldManager.h
--- a/Engine/Plugins/FX/Niagara/Source/Niagara/Public/NiagaraWorldManager.h
+++ b/Engine/Plugins/FX/Niagara/Source/Niagara/Public/NiagaraWorldManager.h
@@ -174,8 +174,8 @@ public:
 	void PrimePoolForAllSystems();
 	void PrimePool(UNiagaraSystem* System);
 
-	static void RefreshSystemScalabilitySettingsForAllWorlds(UNiagaraSystem* System);
-	void RefreshSystemScalabilitySettings(UNiagaraSystem* System);
+	static void InvalidateCachedSystemScalabilityDataForAllWorlds();
+	void InvalidateCachedSystemScalabilityData();
 
 	void SetDebugPlaybackMode(ENiagaraDebugPlaybackMode Mode) { RequestedDebugPlaybackMode = Mode; }
 	ENiagaraDebugPlaybackMode GetDebugPlaybackMode() const { return DebugPlaybackMode; }
