Fix for compile issue introduced in CL-17962735
This was an attempt to preserve the Thumbnail Rendering issues from CL-17948832 / CL-17962735.

diff --git a/Engine/Plugins/FX/Niagara/Source/Niagara/Private/NiagaraGpuComputeDispatch.cpp b/Engine/Plugins/FX/Niagara/Source/Niagara/Private/NiagaraGpuComputeDispatch.cpp
--- a/Engine/Plugins/FX/Niagara/Source/Niagara/Private/NiagaraGpuComputeDispatch.cpp
+++ b/Engine/Plugins/FX/Niagara/Source/Niagara/Private/NiagaraGpuComputeDispatch.cpp
@@ -458,7 +458,7 @@ void FNiagaraGpuComputeDispatch::ProcessPendingTicksFlush(FRHICommandListImmedia
 			FMemMark Mark(FMemStack::Get());
 
 			FSceneViewFamily ViewFamily(FSceneViewFamily::ConstructionValues(nullptr, nullptr, FEngineShowFlags(ESFIM_Game))
-				.SetTime(UThumbnailRenderer::GetTime())
+				.SetTime(FGameTime())
 				.SetGammaCorrection(1.0f));
 
 			FSceneViewInitOptions ViewInitOptions;
