Uses DECLARE_GLOBAL_SHADER on shader class that uses a shader parameter structure

This allows to list quickly in code what shader class are left to refactor to RDG

diff --git a/Engine/Plugins/Experimental/NeuralNetworkInference/Source/Private/NeuralNetworkInferenceShaders/Internal/BatchNormalizationCS.h b/Engine/Plugins/Experimental/NeuralNetworkInference/Source/Private/NeuralNetworkInferenceShaders/Internal/BatchNormalizationCS.h
--- a/Engine/Plugins/Experimental/NeuralNetworkInference/Source/Private/NeuralNetworkInferenceShaders/Internal/BatchNormalizationCS.h
+++ b/Engine/Plugins/Experimental/NeuralNetworkInference/Source/Private/NeuralNetworkInferenceShaders/Internal/BatchNormalizationCS.h
@@ -14,10 +14,9 @@
 
 class NEURALNETWORKINFERENCESHADERS_API FBatchNormalizationCS : public FGlobalShader
 {
-	DECLARE_SHADER_TYPE(FBatchNormalizationCS, Global);
+	DECLARE_GLOBAL_SHADER(FBatchNormalizationCS);
 	SHADER_USE_PARAMETER_STRUCT(FBatchNormalizationCS, FGlobalShader)
 
-public:
 	class FBatchNormalizationMode : SHADER_PERMUTATION_ENUM_CLASS("BATCH_NORMALIZATION_MODE", EBatchNormalizationMode);
 	class FIsInlined : SHADER_PERMUTATION_BOOL("IS_INLINED");
 	using FPermutationDomain = TShaderPermutationDomain<FBatchNormalizationMode, FIsInlined>;
diff --git a/Engine/Plugins/Experimental/NeuralNetworkInference/Source/Private/NeuralNetworkInferenceShaders/Internal/ConvBaseCS.h b/Engine/Plugins/Experimental/NeuralNetworkInference/Source/Private/NeuralNetworkInferenceShaders/Internal/ConvBaseCS.h
--- a/Engine/Plugins/Experimental/NeuralNetworkInference/Source/Private/NeuralNetworkInferenceShaders/Internal/ConvBaseCS.h
+++ b/Engine/Plugins/Experimental/NeuralNetworkInference/Source/Private/NeuralNetworkInferenceShaders/Internal/ConvBaseCS.h
@@ -14,10 +14,9 @@
 
 class NEURALNETWORKINFERENCESHADERS_API FConvBaseCS : public FGlobalShader
 {
-	DECLARE_SHADER_TYPE(FConvBaseCS, Global);
+	DECLARE_GLOBAL_SHADER(FConvBaseCS);
 	SHADER_USE_PARAMETER_STRUCT(FConvBaseCS, FGlobalShader)
 
-public:
 	class FHasBias : SHADER_PERMUTATION_BOOL("HAS_BIAS");
 	class FConvMode : SHADER_PERMUTATION_ENUM_CLASS("CONV_MODE", EConvMode);
 	using FPermutationDomain = TShaderPermutationDomain<FHasBias, FConvMode>;
diff --git a/Engine/Plugins/Experimental/NeuralNetworkInference/Source/Private/NeuralNetworkInferenceShaders/Internal/ConvTransposeCS.h b/Engine/Plugins/Experimental/NeuralNetworkInference/Source/Private/NeuralNetworkInferenceShaders/Internal/ConvTransposeCS.h
--- a/Engine/Plugins/Experimental/NeuralNetworkInference/Source/Private/NeuralNetworkInferenceShaders/Internal/ConvTransposeCS.h
+++ b/Engine/Plugins/Experimental/NeuralNetworkInference/Source/Private/NeuralNetworkInferenceShaders/Internal/ConvTransposeCS.h
@@ -14,10 +14,9 @@
 
 class NEURALNETWORKINFERENCESHADERS_API FConvTransposeCS : public FGlobalShader
 {
-	DECLARE_SHADER_TYPE(FConvTransposeCS, Global);
+	DECLARE_GLOBAL_SHADER(FConvTransposeCS);
 	SHADER_USE_PARAMETER_STRUCT(FConvTransposeCS, FGlobalShader)
 
-public:
 	class FConvMode : SHADER_PERMUTATION_ENUM_CLASS("CONV_MODE", EConvMode);
 	using FPermutationDomain = TShaderPermutationDomain<FConvMode>;
 
diff --git a/Engine/Plugins/Experimental/NeuralNetworkInference/Source/Private/NeuralNetworkInferenceShaders/Internal/CopyCS.h b/Engine/Plugins/Experimental/NeuralNetworkInference/Source/Private/NeuralNetworkInferenceShaders/Internal/CopyCS.h
--- a/Engine/Plugins/Experimental/NeuralNetworkInference/Source/Private/NeuralNetworkInferenceShaders/Internal/CopyCS.h
+++ b/Engine/Plugins/Experimental/NeuralNetworkInference/Source/Private/NeuralNetworkInferenceShaders/Internal/CopyCS.h
@@ -14,10 +14,9 @@
 
 class NEURALNETWORKINFERENCESHADERS_API FCopyCS : public FGlobalShader
 {
-	DECLARE_SHADER_TYPE(FCopyCS, Global);
+	DECLARE_GLOBAL_SHADER(FCopyCS);
 	SHADER_USE_PARAMETER_STRUCT(FCopyCS, FGlobalShader)
 
-public:
 	static const uint32 THREADGROUP_SIZE_X;
 
 	static void ModifyCompilationEnvironment(const FGlobalShaderPermutationParameters& InParameters, FShaderCompilerEnvironment& OutEnvironment);
diff --git a/Engine/Plugins/Experimental/NeuralNetworkInference/Source/Private/NeuralNetworkInferenceShaders/Internal/ElementWiseCS.h b/Engine/Plugins/Experimental/NeuralNetworkInference/Source/Private/NeuralNetworkInferenceShaders/Internal/ElementWiseCS.h
--- a/Engine/Plugins/Experimental/NeuralNetworkInference/Source/Private/NeuralNetworkInferenceShaders/Internal/ElementWiseCS.h
+++ b/Engine/Plugins/Experimental/NeuralNetworkInference/Source/Private/NeuralNetworkInferenceShaders/Internal/ElementWiseCS.h
@@ -14,10 +14,9 @@
 
 class NEURALNETWORKINFERENCESHADERS_API FElementWiseCS : public FGlobalShader
 {
-	DECLARE_SHADER_TYPE(FElementWiseCS, Global);
+	DECLARE_GLOBAL_SHADER(FElementWiseCS);
 	SHADER_USE_PARAMETER_STRUCT(FElementWiseCS, FGlobalShader)
 
-public:
 	static const uint32 THREADGROUP_SIZE_X;
 
 	class FShaderType : SHADER_PERMUTATION_ENUM_CLASS("SHADER_FUNCTION", EElementWiseOperator);
diff --git a/Engine/Plugins/Experimental/NeuralNetworkInference/Source/Private/NeuralNetworkInferenceShaders/Internal/GemmCS.h b/Engine/Plugins/Experimental/NeuralNetworkInference/Source/Private/NeuralNetworkInferenceShaders/Internal/GemmCS.h
--- a/Engine/Plugins/Experimental/NeuralNetworkInference/Source/Private/NeuralNetworkInferenceShaders/Internal/GemmCS.h
+++ b/Engine/Plugins/Experimental/NeuralNetworkInference/Source/Private/NeuralNetworkInferenceShaders/Internal/GemmCS.h
@@ -14,10 +14,9 @@
 
 class NEURALNETWORKINFERENCESHADERS_API FGemmCS : public FGlobalShader
 {
-	DECLARE_SHADER_TYPE(FGemmCS, Global);
+	DECLARE_GLOBAL_SHADER(FGemmCS);
 	SHADER_USE_PARAMETER_STRUCT(FGemmCS, FGlobalShader)
 
-public:
 	class FGemmMode : SHADER_PERMUTATION_ENUM_CLASS("GEMM_MODE", EGemmMode);
 	using FPermutationDomain = TShaderPermutationDomain<FGemmMode>;
 
diff --git a/Engine/Plugins/Experimental/NeuralNetworkInference/Source/Private/NeuralNetworkInferenceShaders/Internal/MultidirectionalBroadcastCS.h b/Engine/Plugins/Experimental/NeuralNetworkInference/Source/Private/NeuralNetworkInferenceShaders/Internal/MultidirectionalBroadcastCS.h
--- a/Engine/Plugins/Experimental/NeuralNetworkInference/Source/Private/NeuralNetworkInferenceShaders/Internal/MultidirectionalBroadcastCS.h
+++ b/Engine/Plugins/Experimental/NeuralNetworkInference/Source/Private/NeuralNetworkInferenceShaders/Internal/MultidirectionalBroadcastCS.h
@@ -14,10 +14,9 @@
 
 class NEURALNETWORKINFERENCESHADERS_API FMultidirectionalBroadcastCS : public FGlobalShader
 {
-	DECLARE_SHADER_TYPE(FMultidirectionalBroadcastCS, Global);
+	DECLARE_GLOBAL_SHADER(FMultidirectionalBroadcastCS);
 	SHADER_USE_PARAMETER_STRUCT(FMultidirectionalBroadcastCS, FGlobalShader)
 
-public:
 	static const uint32 THREADGROUP_SIZE_X;
 	static const uint32 MAX_NUMBER_DIMENSIONS;
 
diff --git a/Engine/Plugins/Experimental/NeuralNetworkInference/Source/Private/NeuralNetworkInferenceShaders/Private/BatchNormalizationCS.cpp b/Engine/Plugins/Experimental/NeuralNetworkInference/Source/Private/NeuralNetworkInferenceShaders/Private/BatchNormalizationCS.cpp
--- a/Engine/Plugins/Experimental/NeuralNetworkInference/Source/Private/NeuralNetworkInferenceShaders/Private/BatchNormalizationCS.cpp
+++ b/Engine/Plugins/Experimental/NeuralNetworkInference/Source/Private/NeuralNetworkInferenceShaders/Private/BatchNormalizationCS.cpp
@@ -22,4 +22,4 @@ const uint32 FBatchNormalizationCS::THREADGROUP_SIZE_X = 128;
 /* Shader implementation
  *****************************************************************************/
 
-IMPLEMENT_SHADER_TYPE(, FBatchNormalizationCS, TEXT("/Plugins/NeuralNetworkInference/Private/BatchNormalizationOperator.usf"), TEXT("BatchNormalizationCS"), SF_Compute) // Path defined in NeuralNetworkInferenceShadersModule.cpp
+IMPLEMENT_GLOBAL_SHADER(FBatchNormalizationCS, "/Plugins/NeuralNetworkInference/Private/BatchNormalizationOperator.usf", "BatchNormalizationCS", SF_Compute); // Path defined in NeuralNetworkInferenceShadersModule.cpp
diff --git a/Engine/Plugins/Experimental/NeuralNetworkInference/Source/Private/NeuralNetworkInferenceShaders/Private/ConvBaseCS.cpp b/Engine/Plugins/Experimental/NeuralNetworkInference/Source/Private/NeuralNetworkInferenceShaders/Private/ConvBaseCS.cpp
--- a/Engine/Plugins/Experimental/NeuralNetworkInference/Source/Private/NeuralNetworkInferenceShaders/Private/ConvBaseCS.cpp
+++ b/Engine/Plugins/Experimental/NeuralNetworkInference/Source/Private/NeuralNetworkInferenceShaders/Private/ConvBaseCS.cpp
@@ -27,4 +27,4 @@ void FConvBaseCS::ModifyCompilationEnvironment(const FGlobalShaderPermutationPar
 /* Shader implementation
  *****************************************************************************/
 
-IMPLEMENT_SHADER_TYPE(, FConvBaseCS, TEXT("/Plugins/NeuralNetworkInference/Private/ConvBaseOperator.usf"), TEXT("ConvBaseCS"), SF_Compute) // Path defined in NeuralNetworkInferenceShadersModule.cpp
+IMPLEMENT_GLOBAL_SHADER(FConvBaseCS, "/Plugins/NeuralNetworkInference/Private/ConvBaseOperator.usf", "ConvBaseCS", SF_Compute); // Path defined in NeuralNetworkInferenceShadersModule.cpp
diff --git a/Engine/Plugins/Experimental/NeuralNetworkInference/Source/Private/NeuralNetworkInferenceShaders/Private/ConvTransposeCS.cpp b/Engine/Plugins/Experimental/NeuralNetworkInference/Source/Private/NeuralNetworkInferenceShaders/Private/ConvTransposeCS.cpp
--- a/Engine/Plugins/Experimental/NeuralNetworkInference/Source/Private/NeuralNetworkInferenceShaders/Private/ConvTransposeCS.cpp
+++ b/Engine/Plugins/Experimental/NeuralNetworkInference/Source/Private/NeuralNetworkInferenceShaders/Private/ConvTransposeCS.cpp
@@ -25,4 +25,4 @@ void FConvTransposeCS::ModifyCompilationEnvironment(const FGlobalShaderPermutati
 /* Shader implementation
  *****************************************************************************/
 
-IMPLEMENT_SHADER_TYPE(, FConvTransposeCS, TEXT("/Plugins/NeuralNetworkInference/Private/ConvTransposeOperator.usf"), TEXT("XToXWithZerosCS"), SF_Compute) // Path defined in NeuralNetworkInferenceShadersModule.cpp
+IMPLEMENT_GLOBAL_SHADER(FConvTransposeCS, "/Plugins/NeuralNetworkInference/Private/ConvTransposeOperator.usf", "XToXWithZerosCS", SF_Compute); // Path defined in NeuralNetworkInferenceShadersModule.cpp
diff --git a/Engine/Plugins/Experimental/NeuralNetworkInference/Source/Private/NeuralNetworkInferenceShaders/Private/CopyCS.cpp b/Engine/Plugins/Experimental/NeuralNetworkInference/Source/Private/NeuralNetworkInferenceShaders/Private/CopyCS.cpp
--- a/Engine/Plugins/Experimental/NeuralNetworkInference/Source/Private/NeuralNetworkInferenceShaders/Private/CopyCS.cpp
+++ b/Engine/Plugins/Experimental/NeuralNetworkInference/Source/Private/NeuralNetworkInferenceShaders/Private/CopyCS.cpp
@@ -22,4 +22,4 @@ const uint32 FCopyCS::THREADGROUP_SIZE_X = 128;
 /* Shader implementation
  *****************************************************************************/
 
-IMPLEMENT_SHADER_TYPE(, FCopyCS, TEXT("/Plugins/NeuralNetworkInference/Private/Copy.usf"), TEXT("CopyCS"), SF_Compute) // Path defined in NeuralNetworkInferenceShadersModule.cpp
+IMPLEMENT_GLOBAL_SHADER(FCopyCS, "/Plugins/NeuralNetworkInference/Private/Copy.usf", "CopyCS", SF_Compute); // Path defined in NeuralNetworkInferenceShadersModule.cpp
diff --git a/Engine/Plugins/Experimental/NeuralNetworkInference/Source/Private/NeuralNetworkInferenceShaders/Private/ElementWiseCS.cpp b/Engine/Plugins/Experimental/NeuralNetworkInference/Source/Private/NeuralNetworkInferenceShaders/Private/ElementWiseCS.cpp
--- a/Engine/Plugins/Experimental/NeuralNetworkInference/Source/Private/NeuralNetworkInferenceShaders/Private/ElementWiseCS.cpp
+++ b/Engine/Plugins/Experimental/NeuralNetworkInference/Source/Private/NeuralNetworkInferenceShaders/Private/ElementWiseCS.cpp
@@ -112,4 +112,4 @@ uint32 FElementWiseCS::GetNumberAttributes(const EElementWiseOperator InType)
 /* Shader implementation
  *****************************************************************************/
 
-IMPLEMENT_SHADER_TYPE(, FElementWiseCS, TEXT("/Plugins/NeuralNetworkInference/Private/ElementWiseOperator.usf"), TEXT("ElementWiseCS"), SF_Compute) // Path defined in NeuralNetworkInferenceShadersModule.cpp
+IMPLEMENT_GLOBAL_SHADER(FElementWiseCS, "/Plugins/NeuralNetworkInference/Private/ElementWiseOperator.usf", "ElementWiseCS", SF_Compute); // Path defined in NeuralNetworkInferenceShadersModule.cpp
diff --git a/Engine/Plugins/Experimental/NeuralNetworkInference/Source/Private/NeuralNetworkInferenceShaders/Private/GemmCS.cpp b/Engine/Plugins/Experimental/NeuralNetworkInference/Source/Private/NeuralNetworkInferenceShaders/Private/GemmCS.cpp
--- a/Engine/Plugins/Experimental/NeuralNetworkInference/Source/Private/NeuralNetworkInferenceShaders/Private/GemmCS.cpp
+++ b/Engine/Plugins/Experimental/NeuralNetworkInference/Source/Private/NeuralNetworkInferenceShaders/Private/GemmCS.cpp
@@ -24,4 +24,4 @@ const uint32 FGemmCS::THREADGROUP_SIZE_Y = 1;
 /* Shader implementation
  *****************************************************************************/
 
-IMPLEMENT_SHADER_TYPE(, FGemmCS, TEXT("/Plugins/NeuralNetworkInference/Private/GemmOperator.usf"), TEXT("GemmCS"), SF_Compute) // Path defined in NeuralNetworkInferenceShadersModule.cpp
+IMPLEMENT_GLOBAL_SHADER(FGemmCS, "/Plugins/NeuralNetworkInference/Private/GemmOperator.usf", "GemmCS", SF_Compute); // Path defined in NeuralNetworkInferenceShadersModule.cpp
diff --git a/Engine/Plugins/Experimental/NeuralNetworkInference/Source/Private/NeuralNetworkInferenceShaders/Private/MultidirectionalBroadcastCS.cpp b/Engine/Plugins/Experimental/NeuralNetworkInference/Source/Private/NeuralNetworkInferenceShaders/Private/MultidirectionalBroadcastCS.cpp
--- a/Engine/Plugins/Experimental/NeuralNetworkInference/Source/Private/NeuralNetworkInferenceShaders/Private/MultidirectionalBroadcastCS.cpp
+++ b/Engine/Plugins/Experimental/NeuralNetworkInference/Source/Private/NeuralNetworkInferenceShaders/Private/MultidirectionalBroadcastCS.cpp
@@ -55,4 +55,4 @@ FString FMultidirectionalBroadcastCS::GetFunctionString(const EMultidirectionalB
 /* Shader implementation
  *****************************************************************************/
 
-IMPLEMENT_SHADER_TYPE(, FMultidirectionalBroadcastCS, TEXT("/Plugins/NeuralNetworkInference/Private/MultidirectionalBroadcastOperator.usf"), TEXT("MultidirectionalBroadcastCS"), SF_Compute) // Path defined in NeuralNetworkInferenceShadersModule.cpp
+IMPLEMENT_GLOBAL_SHADER(FMultidirectionalBroadcastCS, "/Plugins/NeuralNetworkInference/Private/MultidirectionalBroadcastOperator.usf", "MultidirectionalBroadcastCS", SF_Compute); // Path defined in NeuralNetworkInferenceShadersModule.cpp
diff --git a/Engine/Plugins/FX/Niagara/Source/NiagaraShader/Private/NiagaraBatchedElements.cpp b/Engine/Plugins/FX/Niagara/Source/NiagaraShader/Private/NiagaraBatchedElements.cpp
--- a/Engine/Plugins/FX/Niagara/Source/NiagaraShader/Private/NiagaraBatchedElements.cpp
+++ b/Engine/Plugins/FX/Niagara/Source/NiagaraShader/Private/NiagaraBatchedElements.cpp
@@ -10,8 +10,7 @@
 
 class FNiagaraSimpleElement2DArrayAttribute : public FGlobalShader
 {
-public:
-	DECLARE_SHADER_TYPE(FNiagaraSimpleElement2DArrayAttribute, Global);
+	DECLARE_GLOBAL_SHADER(FNiagaraSimpleElement2DArrayAttribute);
 	SHADER_USE_PARAMETER_STRUCT(FNiagaraSimpleElement2DArrayAttribute, FGlobalShader);
 
 	static bool ShouldCompilePermutation(const FGlobalShaderPermutationParameters& Parameters)
@@ -77,8 +76,7 @@ void FBatchedElementNiagara2DArrayAttribute::BindShaders(FRHICommandList& RHICmd
 
 class FNiagaraSimpleElementPS : public FGlobalShader
 {
-public:
-	DECLARE_SHADER_TYPE(FNiagaraSimpleElementPS, Global);
+	DECLARE_GLOBAL_SHADER(FNiagaraSimpleElementPS);
 	SHADER_USE_PARAMETER_STRUCT(FNiagaraSimpleElementPS, FGlobalShader);
 
 	static bool ShouldCompilePermutation(const FGlobalShaderPermutationParameters& Parameters)
diff --git a/Engine/Plugins/FX/Niagara/Source/NiagaraShader/Private/NiagaraClearCounts.cpp b/Engine/Plugins/FX/Niagara/Source/NiagaraShader/Private/NiagaraClearCounts.cpp
--- a/Engine/Plugins/FX/Niagara/Source/NiagaraShader/Private/NiagaraClearCounts.cpp
+++ b/Engine/Plugins/FX/Niagara/Source/NiagaraShader/Private/NiagaraClearCounts.cpp
@@ -14,8 +14,7 @@
 
 class NIAGARASHADER_API FNiagaraClearCountsIntCS : public FGlobalShader
 {
-public:
-	DECLARE_SHADER_TYPE(FNiagaraClearCountsIntCS, Global);
+	DECLARE_GLOBAL_SHADER(FNiagaraClearCountsIntCS);
 	SHADER_USE_PARAMETER_STRUCT(FNiagaraClearCountsIntCS, FGlobalShader);
 
 	static constexpr uint32 ThreadGroupSize = 64;
diff --git a/Engine/Plugins/FX/Niagara/Source/NiagaraShader/Private/NiagaraDebugShaders.cpp b/Engine/Plugins/FX/Niagara/Source/NiagaraShader/Private/NiagaraDebugShaders.cpp
--- a/Engine/Plugins/FX/Niagara/Source/NiagaraShader/Private/NiagaraDebugShaders.cpp
+++ b/Engine/Plugins/FX/Niagara/Source/NiagaraShader/Private/NiagaraDebugShaders.cpp
@@ -33,8 +33,7 @@ static FAutoConsoleVariableRef CVarNiagaraGpuComputeDebug_FourComponentMode(
 
 class NIAGARASHADER_API FNiagaraVisualizeTexturePS : public FGlobalShader
 {
-public:
-	DECLARE_SHADER_TYPE(FNiagaraVisualizeTexturePS, Global);
+	DECLARE_GLOBAL_SHADER(FNiagaraVisualizeTexturePS);
 	SHADER_USE_PARAMETER_STRUCT(FNiagaraVisualizeTexturePS, FGlobalShader);
 
 	class FIntegerTexture : SHADER_PERMUTATION_BOOL("TEXTURE_INTEGER");
@@ -75,8 +74,7 @@ IMPLEMENT_GLOBAL_SHADER(FNiagaraVisualizeTexturePS, "/Plugin/FX/Niagara/Private/
 
 class NIAGARASHADER_API FNiagaraClearUAVCS : public FGlobalShader
 {
-public:
-	DECLARE_SHADER_TYPE(FNiagaraClearUAVCS, Global);
+	DECLARE_GLOBAL_SHADER(FNiagaraClearUAVCS);
 	SHADER_USE_PARAMETER_STRUCT(FNiagaraClearUAVCS, FGlobalShader);
 
 	static constexpr uint32 ThreadGroupSize = 32;
@@ -111,8 +109,7 @@ IMPLEMENT_GLOBAL_SHADER(FNiagaraClearUAVCS, "/Plugin/FX/Niagara/Private/NiagaraD
 
 class NIAGARASHADER_API FNiagaraDebugDrawLineVS : public FGlobalShader
 {
-public:
-	DECLARE_SHADER_TYPE(FNiagaraDebugDrawLineVS, Global);
+	DECLARE_GLOBAL_SHADER(FNiagaraDebugDrawLineVS);
 	SHADER_USE_PARAMETER_STRUCT(FNiagaraDebugDrawLineVS, FGlobalShader);
 
 	static bool ShouldCompilePermutation(const FGlobalShaderPermutationParameters& Parameters)
@@ -134,8 +131,7 @@ public:
 
 class NIAGARASHADER_API FNiagaraDebugDrawLinePS : public FGlobalShader
 {
-public:
-	DECLARE_SHADER_TYPE(FNiagaraDebugDrawLinePS, Global);
+	DECLARE_GLOBAL_SHADER(FNiagaraDebugDrawLinePS);
 	SHADER_USE_PARAMETER_STRUCT(FNiagaraDebugDrawLinePS, FGlobalShader);
 
 	static bool ShouldCompilePermutation(const FGlobalShaderPermutationParameters& Parameters)
diff --git a/Engine/Plugins/FX/Niagara/Source/NiagaraShader/Private/NiagaraRayTracingHelper.cpp b/Engine/Plugins/FX/Niagara/Source/NiagaraShader/Private/NiagaraRayTracingHelper.cpp
--- a/Engine/Plugins/FX/Niagara/Source/NiagaraShader/Private/NiagaraRayTracingHelper.cpp
+++ b/Engine/Plugins/FX/Niagara/Source/NiagaraShader/Private/NiagaraRayTracingHelper.cpp
@@ -113,7 +113,7 @@ class FNiagaraCollisionRayTraceRG : public FGlobalShader
 
 class FNiagaraCollisionRayTraceCH : public FGlobalShader
 {
-	DECLARE_SHADER_TYPE(FNiagaraCollisionRayTraceCH, Global);
+	DECLARE_GLOBAL_SHADER(FNiagaraCollisionRayTraceCH);
 
 	static bool ShouldCompilePermutation(const FGlobalShaderPermutationParameters& Parameters)
 	{
@@ -134,7 +134,7 @@ class FNiagaraCollisionRayTraceCH : public FGlobalShader
 
 class FNiagaraCollisionRayTraceMiss : public FGlobalShader
 {
-	DECLARE_SHADER_TYPE(FNiagaraCollisionRayTraceMiss, Global);
+	DECLARE_GLOBAL_SHADER(FNiagaraCollisionRayTraceMiss);
 
 	static bool ShouldCompilePermutation(const FGlobalShaderPermutationParameters& Parameters)
 	{
diff --git a/Engine/Source/Runtime/Renderer/Private/CapsuleShadowRendering.cpp b/Engine/Source/Runtime/Renderer/Private/CapsuleShadowRendering.cpp
--- a/Engine/Source/Runtime/Renderer/Private/CapsuleShadowRendering.cpp
+++ b/Engine/Source/Runtime/Renderer/Private/CapsuleShadowRendering.cpp
@@ -301,8 +301,7 @@ IMPLEMENT_GLOBAL_SHADER(FCapsuleShadowingCS, "/Engine/Private/CapsuleShadowShade
 
 class FCapsuleShadowingUpsampleVS : public FGlobalShader
 {
-	DECLARE_SHADER_TYPE(FCapsuleShadowingUpsampleVS, Global);
-public:
+	DECLARE_GLOBAL_SHADER(FCapsuleShadowingUpsampleVS);
 	SHADER_USE_PARAMETER_STRUCT(FCapsuleShadowingUpsampleVS, FGlobalShader);
 
 	BEGIN_SHADER_PARAMETER_STRUCT(FParameters, )
diff --git a/Engine/Source/Runtime/Renderer/Private/PlanarReflectionRendering.cpp b/Engine/Source/Runtime/Renderer/Private/PlanarReflectionRendering.cpp
--- a/Engine/Source/Runtime/Renderer/Private/PlanarReflectionRendering.cpp
+++ b/Engine/Source/Runtime/Renderer/Private/PlanarReflectionRendering.cpp
@@ -135,8 +135,7 @@ IMPLEMENT_GLOBAL_SHADER_PARAMETER_STRUCT(FPlanarReflectionUniformParameters, "Pl
 template< bool bEnablePlanarReflectionPrefilter >
 class TPrefilterPlanarReflectionPS : public FGlobalShader
 {
-public:
-	DECLARE_SHADER_TYPE(TPrefilterPlanarReflectionPS, Global);
+	DECLARE_GLOBAL_SHADER(TPrefilterPlanarReflectionPS);
 	SHADER_USE_PARAMETER_STRUCT(TPrefilterPlanarReflectionPS, FGlobalShader);
 
 	BEGIN_SHADER_PARAMETER_STRUCT(FParameters, )
diff --git a/Engine/Source/Runtime/Renderer/Private/PostProcess/PostProcessMobile.cpp b/Engine/Source/Runtime/Renderer/Private/PostProcess/PostProcessMobile.cpp
--- a/Engine/Source/Runtime/Renderer/Private/PostProcess/PostProcessMobile.cpp
+++ b/Engine/Source/Runtime/Renderer/Private/PostProcess/PostProcessMobile.cpp
@@ -1785,7 +1785,7 @@ class TMobileHistogramCS : public FMobileHistogram
 {
 public:
 	typedef TMobileHistogramCS< LowSharedComputeMemory > ClassName; // typedef is only so that we can use in DECLARE_SHADER_TYPE macro
-	DECLARE_SHADER_TYPE(ClassName, Global);
+	DECLARE_GLOBAL_SHADER(ClassName);
 
 	// The number of texels on each axis processed by a single thread group.
 	static const FIntPoint TexelsPerThreadGroup;
diff --git a/Engine/Source/Runtime/Renderer/Private/RayTracing/RayTracingBarycentrics.cpp b/Engine/Source/Runtime/Renderer/Private/RayTracing/RayTracingBarycentrics.cpp
--- a/Engine/Source/Runtime/Renderer/Private/RayTracing/RayTracingBarycentrics.cpp
+++ b/Engine/Source/Runtime/Renderer/Private/RayTracing/RayTracingBarycentrics.cpp
@@ -32,15 +32,13 @@ IMPLEMENT_GLOBAL_SHADER(FRayTracingBarycentricsRGS, "/Engine/Private/RayTracing/
 // Example closest hit shader
 class FRayTracingBarycentricsCHS : public FGlobalShader
 {
-	DECLARE_SHADER_TYPE(FRayTracingBarycentricsCHS, Global);
+	DECLARE_GLOBAL_SHADER(FRayTracingBarycentricsCHS);
 
 	static void ModifyCompilationEnvironment(const FGlobalShaderPermutationParameters& Parameters, FShaderCompilerEnvironment& OutEnvironment)
 	{
 		FGlobalShader::ModifyCompilationEnvironment(Parameters, OutEnvironment);
 	}
 
-public:
-
 	static bool ShouldCompilePermutation(const FGlobalShaderPermutationParameters& Parameters)
 	{
 		return ShouldCompileRayTracingShadersForProject(Parameters.Platform);
diff --git a/Engine/Source/Runtime/Renderer/Private/RayTracing/RayTracingDebug.cpp b/Engine/Source/Runtime/Renderer/Private/RayTracing/RayTracingDebug.cpp
--- a/Engine/Source/Runtime/Renderer/Private/RayTracing/RayTracingDebug.cpp
+++ b/Engine/Source/Runtime/Renderer/Private/RayTracing/RayTracingDebug.cpp
@@ -65,7 +65,7 @@ IMPLEMENT_GLOBAL_SHADER(FRayTracingDebugRGS, "/Engine/Private/RayTracing/RayTrac
 
 class FRayTracingDebugCHS : public FGlobalShader
 {
-	DECLARE_SHADER_TYPE(FRayTracingDebugCHS, Global);
+	DECLARE_GLOBAL_SHADER(FRayTracingDebugCHS);
 
 	static void ModifyCompilationEnvironment(const FGlobalShaderPermutationParameters& Parameters, FShaderCompilerEnvironment& OutEnvironment)
 	{
diff --git a/Engine/Source/Runtime/Renderer/Private/VT/VirtualTextureTest.cpp b/Engine/Source/Runtime/Renderer/Private/VT/VirtualTextureTest.cpp
--- a/Engine/Source/Runtime/Renderer/Private/VT/VirtualTextureTest.cpp
+++ b/Engine/Source/Runtime/Renderer/Private/VT/VirtualTextureTest.cpp
@@ -115,7 +115,7 @@ bool FVirtualTextureTest::RequestPageData( uint8 vLevel, uint64 vAddress, void*
 
 class FVirtualTextureTestPS : public FGlobalShader
 {
-	DECLARE_SHADER_TYPE(FVirtualTextureTestPS, Global);
+	DECLARE_GLOBAL_SHADER(FVirtualTextureTestPS);
 
 	static bool ShouldCache( EShaderPlatform Platform )
 	{
