Fix inconsistent preprocessor state warnings

diff --git a/Engine/Plugins/FX/Niagara/Source/Niagara/Private/NiagaraGPUSceneUtils.h b/Engine/Plugins/FX/Niagara/Source/Niagara/Private/NiagaraGPUSceneUtils.h
--- a/Engine/Plugins/FX/Niagara/Source/Niagara/Private/NiagaraGPUSceneUtils.h
+++ b/Engine/Plugins/FX/Niagara/Source/Niagara/Private/NiagaraGPUSceneUtils.h
@@ -3,8 +3,8 @@
 #pragma once
 
 #include "CoreMinimal.h"
+#include "RHI.h"
 #include "RHIDefinitions.h"
-#include "RHICommandList.h"
 #include "ShaderParameterStruct.h"
 #include "ShaderParameterUtils.h"
 #include "NiagaraMeshVertexFactory.h"
