Allow current SP in Niagara script stats settings.

diff --git a/Engine/Plugins/FX/Niagara/Source/NiagaraEditor/Private/ViewModels/NiagaraScriptStatsViewModel.cpp b/Engine/Plugins/FX/Niagara/Source/NiagaraEditor/Private/ViewModels/NiagaraScriptStatsViewModel.cpp
--- a/Engine/Plugins/FX/Niagara/Source/NiagaraEditor/Private/ViewModels/NiagaraScriptStatsViewModel.cpp
+++ b/Engine/Plugins/FX/Niagara/Source/NiagaraEditor/Private/ViewModels/NiagaraScriptStatsViewModel.cpp
@@ -246,11 +246,7 @@ namespace NiagaraScriptStatsLocal
 UNiagaraScripStatsViewModelSettings::UNiagaraScripStatsViewModelSettings(const FObjectInitializer& ObjectInitializer)
 	: Super(ObjectInitializer)
 {
-#if PLATFORM_WINDOWS
-	EnabledPlatforms.AddUnique(SP_PCD3D_SM5);
-#elif PLATFORM_IOS
-	EnabledPlatforms.AddUnique(SP_METAL);
-#endif
+	EnabledPlatforms.Add(static_cast<int32>(GMaxRHIShaderPlatform));
 }
 
 /////////////////////////////////////////////////////////////////////////////////////
