Fix errors on some shader platforms related to niagara particle access in shaders

diff --git a/Engine/Plugins/FX/Niagara/Shaders/Private/NiagaraParticleAccess.ush b/Engine/Plugins/FX/Niagara/Shaders/Private/NiagaraParticleAccess.ush
--- a/Engine/Plugins/FX/Niagara/Shaders/Private/NiagaraParticleAccess.ush
+++ b/Engine/Plugins/FX/Niagara/Shaders/Private/NiagaraParticleAccess.ush
@@ -32,15 +32,16 @@
 	Buffer<float> NiagaraParticleDataHalf;
 	uint NiagaraFloatDataStride;
 
-	Buffer<float> NiagaraGetFloatBuffer() { return NiagaraParticleDataFloat; }
-	Buffer<float> NiagaraGetHalfBuffer() { return NiagaraParticleDataHalf; }
+	// NOTE: NiagarGet*Buffer have to be macros (not functions) to appease certain shader platforms
+	#define NiagaraGetFloatBuffer() (NiagaraParticleDataFloat)
+	#define NiagaraGetHalfBuffer() (NiagaraParticleDataHalf)
 	uint NiagaraGetFloatDataStride() { return NiagaraFloatDataStride; }
 	
 	#if ENABLE_NIAGARA_INT_DATA_ACCESS
 		Buffer<int> NiagaraParticleDataInt;
 		uint NiagaraIntDataStride;
 
-		Buffer<int> NiagaraGetIntBuffer() { return NiagaraParticleDataInt; }
+		#define NiagaraGetIntBuffer() (NiagaraParticleDataInt)
 		uint NiagaraGetIntDataStride() { return NiagaraIntDataStride; }
 	#endif
 #endif
diff --git a/Engine/Plugins/FX/Niagara/Shaders/Private/NiagaraVFParticleAccess.usf b/Engine/Plugins/FX/Niagara/Shaders/Private/NiagaraVFParticleAccess.usf
--- a/Engine/Plugins/FX/Niagara/Shaders/Private/NiagaraVFParticleAccess.usf
+++ b/Engine/Plugins/FX/Niagara/Shaders/Private/NiagaraVFParticleAccess.usf
@@ -5,8 +5,9 @@
  *********************************************************************************/
 
 // These lines set up access to Niagara particle data buffers for Vertex Factories
-Buffer<float> NiagaraGetFloatBuffer() { return NiagaraVFLooseParameters.NiagaraParticleDataFloat; }
-Buffer<float> NiagaraGetHalfBuffer() { return NiagaraVFLooseParameters.NiagaraParticleDataHalf; }
+// NOTE: the buffer accessors have to be macros to appease certain shader platforms
+#define NiagaraGetFloatBuffer() (NiagaraVFLooseParameters.NiagaraParticleDataFloat)
+#define NiagaraGetHalfBuffer() (NiagaraVFLooseParameters.NiagaraParticleDataHalf)
 uint NiagaraGetFloatDataStride() { return NiagaraVFLooseParameters.NiagaraFloatDataStride; }
 #define USE_GLOBAL_NIAGARA_DATA_BUFFERS 0 // We provided the buffer accessors above
 #define ENABLE_NIAGARA_INT_DATA_ACCESS 0 // Currently, ints not supported in VFs (but could be)
