- Fix crash when changing a struct binding to value

#rnx

diff --git a/Engine/Plugins/FX/Niagara/Source/NiagaraEditor/Private/NiagaraGraph.cpp b/Engine/Plugins/FX/Niagara/Source/NiagaraEditor/Private/NiagaraGraph.cpp
--- a/Engine/Plugins/FX/Niagara/Source/NiagaraEditor/Private/NiagaraGraph.cpp
+++ b/Engine/Plugins/FX/Niagara/Source/NiagaraEditor/Private/NiagaraGraph.cpp
@@ -1832,31 +1832,40 @@ bool UNiagaraGraph::RenameParameter(const FNiagaraVariable& Parameter, FName New
 void UNiagaraGraph::ScriptVariableChanged(FNiagaraVariable Variable)
 {
 	TObjectPtr<UNiagaraScriptVariable>* ScriptVariable = GetAllMetaData().Find(Variable);
-	if (!ScriptVariable || !*ScriptVariable || (*ScriptVariable)->GetIsStaticSwitch()) {
+	if (!ScriptVariable || !*ScriptVariable || (*ScriptVariable)->GetIsStaticSwitch())
+	{
 		return;
 	}
 
 	FNiagaraEditorModule& EditorModule = FNiagaraEditorModule::Get();
-	auto TypeUtilityValue = EditorModule.GetTypeUtilities(Variable.GetType());
 
 	TArray<UEdGraphPin*> Pins = FindParameterMapDefaultValuePins(Variable.GetName());
-	for (UEdGraphPin* Pin : Pins) {
+	for (UEdGraphPin* Pin : Pins)
+	{
 		Pin->bHidden = ((*ScriptVariable)->DefaultMode == ENiagaraDefaultMode::Binding) || ((*ScriptVariable)->DefaultMode == ENiagaraDefaultMode::FailIfPreviouslyNotSet);
-		if ((*ScriptVariable)->DefaultMode == ENiagaraDefaultMode::Custom) {
+		if ((*ScriptVariable)->DefaultMode == ENiagaraDefaultMode::Custom)
+		{
 			Pin->bNotConnectable = false;
 			Pin->bDefaultValueIsReadOnly = false;
 		}
-		else {
+		else
+		{
 			Pin->BreakAllPinLinks(true);
 			Pin->bNotConnectable = true;
 			Pin->bDefaultValueIsReadOnly = true;
 
-			if ((*ScriptVariable)->DefaultMode == ENiagaraDefaultMode::Value && !Variable.GetType().IsDataInterface()) {
-				if (!Variable.IsDataAllocated()) {
-					Variable.AllocateData();
+			if ((*ScriptVariable)->DefaultMode == ENiagaraDefaultMode::Value && !Variable.GetType().IsDataInterface())
+			{
+				auto TypeUtilityValue = EditorModule.GetTypeUtilities(Variable.GetType());
+				if (TypeUtilityValue.IsValid() && TypeUtilityValue->CanHandlePinDefaults())
+				{
+					if (!Variable.IsDataAllocated())
+					{
+						Variable.AllocateData();
+					}
+					FString NewDefaultValue = TypeUtilityValue->GetPinDefaultStringFromValue(Variable);
+					GetDefault<UEdGraphSchema_Niagara>()->TrySetDefaultValue(*Pin, NewDefaultValue, true);
 				}
-				FString NewDefaultValue = TypeUtilityValue->GetPinDefaultStringFromValue(Variable);
-				GetDefault<UEdGraphSchema_Niagara>()->TrySetDefaultValue(*Pin, NewDefaultValue, true);
 			}
 		}
 	}
