Move niagara mesh VF to use MANUAL_VERTEX_FETCH to address limited functionality when using ray tracing

diff --git a/Engine/Plugins/FX/Niagara/Shaders/Private/NiagaraMeshVertexFactory.ush b/Engine/Plugins/FX/Niagara/Shaders/Private/NiagaraMeshVertexFactory.ush
--- a/Engine/Plugins/FX/Niagara/Shaders/Private/NiagaraMeshVertexFactory.ush
+++ b/Engine/Plugins/FX/Niagara/Shaders/Private/NiagaraMeshVertexFactory.ush
@@ -27,47 +27,59 @@
 #define ENABLE_PRECISE_MOTION_VECTORS 0
 #endif
 
+#ifndef MANUAL_VERTEX_FETCH
+#define MANUAL_VERTEX_FETCH 0
+#endif
+
+#if MANUAL_VERTEX_FETCH
+	#define VF_ColorIndexMask_Index 0
+	#define VF_NumTexcoords_Index 1
+#endif // MANUAL_VERTEX_FETCH
+
 struct FVertexFactoryInput
 {
 	float4	Position	: ATTRIBUTE0;
+
+#if !MANUAL_VERTEX_FETCH
 	half3	TangentX	: ATTRIBUTE1;
 	// TangentZ.w contains sign of tangent basis determinant
 	half4	TangentZ	: ATTRIBUTE2;
 	half4	VertexColor : ATTRIBUTE3;
 
-#if NUM_MATERIAL_TEXCOORDS_VERTEX
-	float2	TexCoords0 : ATTRIBUTE4;
-	#if NUM_MATERIAL_TEXCOORDS_VERTEX > 1
-		float2	TexCoords1 : ATTRIBUTE5;
-	#endif
-	#if NUM_MATERIAL_TEXCOORDS_VERTEX > 2
-		float2	TexCoords2 : ATTRIBUTE6;
-	#endif
-	#if NUM_MATERIAL_TEXCOORDS_VERTEX > 3
-		float2	TexCoords3 : ATTRIBUTE7;
-	#endif
-	#if NUM_MATERIAL_TEXCOORDS_VERTEX > 4
-		float2	TexCoords4 : ATTRIBUTE8;
-	#endif
-	#if NUM_MATERIAL_TEXCOORDS_VERTEX > 5
-		float2	TexCoords5 : ATTRIBUTE9;
-	#endif
-	#if NUM_MATERIAL_TEXCOORDS_VERTEX > 6
-		float2	TexCoords6 : ATTRIBUTE10;
-	#endif
-	#if NUM_MATERIAL_TEXCOORDS_VERTEX > 7
-		float2	TexCoords7 : ATTRIBUTE11;
-	#endif
-#elif USE_PARTICLE_SUBUVS
-	float2	TexCoords0 : ATTRIBUTE4;
-#endif
+	#if NUM_MATERIAL_TEXCOORDS_VERTEX
+		float2	TexCoords0 : ATTRIBUTE4;
+		#if NUM_MATERIAL_TEXCOORDS_VERTEX > 1
+			float2	TexCoords1 : ATTRIBUTE5;
+		#endif
+		#if NUM_MATERIAL_TEXCOORDS_VERTEX > 2
+			float2	TexCoords2 : ATTRIBUTE6;
+		#endif
+		#if NUM_MATERIAL_TEXCOORDS_VERTEX > 3
+			float2	TexCoords3 : ATTRIBUTE7;
+		#endif
+		#if NUM_MATERIAL_TEXCOORDS_VERTEX > 4
+			float2	TexCoords4 : ATTRIBUTE8;
+		#endif
+		#if NUM_MATERIAL_TEXCOORDS_VERTEX > 5
+			float2	TexCoords5 : ATTRIBUTE9;
+		#endif
+		#if NUM_MATERIAL_TEXCOORDS_VERTEX > 6
+			float2	TexCoords6 : ATTRIBUTE10;
+		#endif
+		#if NUM_MATERIAL_TEXCOORDS_VERTEX > 7
+			float2	TexCoords7 : ATTRIBUTE11;
+		#endif
+	#elif USE_PARTICLE_SUBUVS
+		float2	TexCoords0 : ATTRIBUTE4;
+	#endif // NUM_MATERIAL_TEXCOORDS_VERTEX
+#endif // !MANUAL_VERTEX_FETCH
 
 #if (VF_USE_PRIMITIVE_SCENE_DATA == 1)
 	uint InstanceIdOffset : ATTRIBUTE13;
 #endif
 
 	uint InstanceId	: SV_InstanceID;
-#if PASS_NEEDS_VERTEX_ID
+#if PASS_NEEDS_VERTEX_ID || MANUAL_VERTEX_FETCH
 	uint VertexId : SV_VertexID;
 #endif // #if PASS_NEEDS_VERTEX_ID
 };
@@ -197,6 +209,9 @@ struct FVertexFactoryIntermediates
 	float4x4 PrevParticleToWorld;	
 	float3x3 ParticleToWorldNoScale;
 
+	float3x3 TangentToLocal;
+	float TangentDeterminant;
+
 	/** Cached primitive and instance data */
 	FSceneDataIntermediates SceneData;
 };
@@ -343,7 +358,7 @@ FMaterialVertexParameters GetMaterialVertexParameters(FVertexFactoryInput Input,
 	Result.WorldPosition = WorldPosition;
 	Result.VertexColor = Intermediates.VertexColor;
 	Result.PreSkinnedPosition = Input.Position.xyz;
-	Result.PreSkinnedNormal = TangentBias(Input.TangentZ.xyz);
+	Result.PreSkinnedNormal = Intermediates.TangentToLocal[2];
 
 #if NEEDS_PARTICLE_POSITION
 	Result.Particle.TranslatedWorldPositionAndSize.xyz = Intermediates.ParticleTranslatedWorldPosition.xyz;
@@ -403,6 +418,34 @@ FMaterialVertexParameters GetMaterialVertexParameters(FVertexFactoryInput Input,
 	return Result;
 }
 
+float3x3 CalculateTangentToLocal(FVertexFactoryInput Input, out float TangentDeterminant)
+{
+	float3x3 Result=0;
+#if MANUAL_VERTEX_FETCH
+	float3 TangentX = NiagaraMeshVF.VertexFetch_PackedTangentsBuffer[2 * Input.VertexId + 0].xyz;
+	float4 TangentZ = NiagaraMeshVF.VertexFetch_PackedTangentsBuffer[2 * Input.VertexId + 1].xyzw;
+#else
+	float3 TangentX = TangentBias(Input.TangentX);
+	float4 TangentZ = TangentBias(Input.TangentZ);
+#endif // MANUAL_VERTEX_FETCH
+
+	// pass-thru the tangent
+	Result[0] = TangentX;
+	// pass-thru the normal
+	Result[2] = TangentZ.xyz;
+
+	// derive the binormal by getting the cross product of the normal and tangent
+	Result[1] = cross(Result[2], Result[0]) * TangentZ.w;
+	// Recalculate TangentX off of the other two vectors
+	// This corrects quantization error since TangentX was passed in as a quantized vertex input
+	// The error shows up most in specular off of a mesh with a smoothed UV seam (normal is smooth, but tangents vary across the seam)
+	Result[0] = cross(Result[1], Result[2]) * TangentZ.w;
+
+	TangentDeterminant = TangentZ.w;
+
+	return Result;
+}
+
 FVertexFactoryIntermediates GetVertexFactoryIntermediates(FVertexFactoryInput Input)
 {
 	FVertexFactoryIntermediates Intermediates = (FVertexFactoryIntermediates)0;
@@ -451,9 +494,6 @@ FVertexFactoryIntermediates GetVertexFactoryIntermediates(FVertexFactoryInput In
 	Intermediates.ParticleRandom = SafeGetFloat(NiagaraMeshVF.MaterialRandomDataOffset, Intermediates.ParticleIndex, NiagaraMeshVF.DefaultMatRandom);	
 #endif
 
-	// Swizzle vertex color.
-	Intermediates.VertexColor = Input.VertexColor FCOLOR_COMPONENT_SWIZZLE;
-
 	// Compute transforms
 #if (VF_USE_PRIMITIVE_SCENE_DATA == 1)
 	Intermediates.ParticleToWorld = LWCHackToFloat(Intermediates.SceneData.InstanceData.LocalToWorld);
@@ -526,29 +566,40 @@ FVertexFactoryIntermediates GetVertexFactoryIntermediates(FVertexFactoryInput In
 	Intermediates.ParticleTranslatedWorldPosition = Intermediates.ParticleToWorld[3].xyz + LWCHackToFloat(ResolvedView.PreViewTranslation);
 
 #if NUM_MATERIAL_TEXCOORDS_VERTEX
-	Intermediates.TexCoords[0] = Input.TexCoords0;
-	#if NUM_MATERIAL_TEXCOORDS_VERTEX > 1
-		Intermediates.TexCoords[1] = Input.TexCoords1;
-	#endif
-	#if NUM_MATERIAL_TEXCOORDS_VERTEX > 2
-		Intermediates.TexCoords[2] = Input.TexCoords2;
-	#endif
-	#if NUM_MATERIAL_TEXCOORDS_VERTEX > 3
-		Intermediates.TexCoords[3] = Input.TexCoords3;
-	#endif
-	#if NUM_MATERIAL_TEXCOORDS_VERTEX > 4
-		Intermediates.TexCoords[4] = Input.TexCoords4;
-	#endif
-	#if NUM_MATERIAL_TEXCOORDS_VERTEX > 5
-		Intermediates.TexCoords[5] = Input.TexCoords5;
-	#endif
-	#if NUM_MATERIAL_TEXCOORDS_VERTEX > 6
-		Intermediates.TexCoords[6] = Input.TexCoords6;
-	#endif
-	#if NUM_MATERIAL_TEXCOORDS_VERTEX > 7
-		Intermediates.TexCoords[7] = Input.TexCoords7;
-	#endif
-#endif
+	#if MANUAL_VERTEX_FETCH
+		const uint NumFetchTexCoords = NiagaraMeshVF.VertexFetch_Parameters[VF_NumTexcoords_Index];
+		UNROLL
+		for (uint CoordinateIndex = 0; CoordinateIndex < NUM_MATERIAL_TEXCOORDS_VERTEX; CoordinateIndex++)
+		{
+			// Clamp coordinates to mesh's maximum as materials can request more than are available
+			uint ClampedCoordinateIndex = min(CoordinateIndex, NumFetchTexCoords-1);
+			Intermediates.TexCoords[CoordinateIndex] = NiagaraMeshVF.VertexFetch_TexCoordBuffer[NumFetchTexCoords * Input.VertexId + ClampedCoordinateIndex];
+		}
+	#else
+		Intermediates.TexCoords[0] = Input.TexCoords0;
+		#if NUM_MATERIAL_TEXCOORDS_VERTEX > 1
+			Intermediates.TexCoords[1] = Input.TexCoords1;
+		#endif
+		#if NUM_MATERIAL_TEXCOORDS_VERTEX > 2
+			Intermediates.TexCoords[2] = Input.TexCoords2;
+		#endif
+		#if NUM_MATERIAL_TEXCOORDS_VERTEX > 3
+			Intermediates.TexCoords[3] = Input.TexCoords3;
+		#endif
+		#if NUM_MATERIAL_TEXCOORDS_VERTEX > 4
+			Intermediates.TexCoords[4] = Input.TexCoords4;
+		#endif
+		#if NUM_MATERIAL_TEXCOORDS_VERTEX > 5
+			Intermediates.TexCoords[5] = Input.TexCoords5;
+		#endif
+		#if NUM_MATERIAL_TEXCOORDS_VERTEX > 6
+			Intermediates.TexCoords[6] = Input.TexCoords6;
+		#endif
+		#if NUM_MATERIAL_TEXCOORDS_VERTEX > 7
+			Intermediates.TexCoords[7] = Input.TexCoords7;
+		#endif
+	#endif // MANUAL_VERTEX_FETCH
+#endif // NUM_MATERIAL_TEXCOORDS_VERTEX
 
 #if USE_PARTICLE_SUBUVS
 	// SubUV.
@@ -560,10 +611,27 @@ FVertexFactoryIntermediates GetVertexFactoryIntermediates(FVertexFactoryInput In
 	float SubImageBH = fmod(SubImageB, NiagaraMeshVF.SubImageSize.x);
 	float SubImageAV = floor(SubImageA * NiagaraMeshVF.SubImageSize.z);
 	float SubImageBV = floor(SubImageB * NiagaraMeshVF.SubImageSize.z);
-	Intermediates.SubUVCoords.xy = (float2(SubImageAH, SubImageAV) + Input.TexCoords0.xy) * NiagaraMeshVF.SubImageSize.zw;
-	Intermediates.SubUVCoords.zw = (float2(SubImageBH, SubImageBV) + Input.TexCoords0.xy) * NiagaraMeshVF.SubImageSize.zw;
 	Intermediates.SubUVLerp = (NiagaraMeshVF.SubImageBlendMode == 1) ? SubImageLerp : 0.0f;
-#endif
+
+	#if MANUAL_VERTEX_FETCH
+		const uint NumFetchTexCoords = NiagaraMeshVF.VertexFetch_Parameters[VF_NumTexcoords_Index];
+		const float2 SubUvTexCoord = NiagaraMeshVF.VertexFetch_TexCoordBuffer[NumFetchTexCoords * Input.VertexId];
+	#else
+		const float2 SubUvTexCoord = Input.TexCoords0.xy;
+	#endif // MANUAL_VERTEX_FETCH
+
+	Intermediates.SubUVCoords.xy = (float2(SubImageAH, SubImageAV) + SubUvTexCoord) * NiagaraMeshVF.SubImageSize.zw;
+	Intermediates.SubUVCoords.zw = (float2(SubImageBH, SubImageBV) + SubUvTexCoord) * NiagaraMeshVF.SubImageSize.zw;
+#endif // USE_PARTICLE_SUBUVS
+
+#if MANUAL_VERTEX_FETCH
+	Intermediates.VertexColor = NiagaraMeshVF.VertexFetch_ColorComponentsBuffer[Input.VertexId & NiagaraMeshVF.VertexFetch_Parameters[VF_ColorIndexMask_Index]] FMANUALFETCH_COLOR_COMPONENT_SWIZZLE; // Swizzle vertex color.
+#else
+	// Swizzle vertex color.
+	Intermediates.VertexColor = Input.VertexColor FCOLOR_COMPONENT_SWIZZLE;
+#endif // MANUAL_VERTEX_FETCH
+
+	Intermediates.TangentToLocal = CalculateTangentToLocal(Input, Intermediates.TangentDeterminant);
 
 	return Intermediates;
 }
@@ -577,24 +645,7 @@ FVertexFactoryIntermediates GetVertexFactoryIntermediates(FVertexFactoryInput In
 */
 float3x3 VertexFactoryGetTangentToLocal( FVertexFactoryInput Input, FVertexFactoryIntermediates Intermediates )
 {
-	float3x3 Result=0;
-	float3 TangentX = TangentBias(Input.TangentX);
-	float4 TangentZ = TangentBias(Input.TangentZ);
-
-	// pass-thru the tangent
-	Result[0] = TangentX;
-	// pass-thru the normal
-	Result[2] = TangentZ.xyz;
-
-
-	// derive the binormal by getting the cross product of the normal and tangent
-	Result[1] = cross(Result[2], Result[0]) * TangentZ.w;
-	// Recalculate TangentX off of the other two vectors
-	// This corrects quantization error since TangentX was passed in as a quantized vertex input
-	// The error shows up most in specular off of a mesh with a smoothed UV seam (normal is smooth, but tangents vary across the seam)
-	Result[0] = cross(Result[1], Result[2]) * TangentZ.w;
-
-	return Result;
+	return Intermediates.TangentToLocal;
 }
 
 // @return translated world position
@@ -630,7 +681,7 @@ void CalcTangentToWorld(FVertexFactoryInput Input, FVertexFactoryIntermediates I
 	const float DetSign = Intermediates.SceneData.InstanceData.DeterminantSign;
 
 	TangentToWorld0 = normalize(TangentToWorld[0]);
-	TangentToWorld2 = float4(normalize(TangentToWorld[2]), TangentBias(Input.TangentZ.w) * DetSign);
+	TangentToWorld2 = float4(normalize(TangentToWorld[2]), TangentBias(Intermediates.TangentDeterminant) * GetPrimitive_DeterminantSign_FromFlags(GetPrimitiveDataFromUniformBuffer().Flags));
 }
 
 FVertexFactoryInterpolantsVSToPS VertexFactoryGetInterpolantsVSToPS(FVertexFactoryInput Input, FVertexFactoryIntermediates Intermediates, FMaterialVertexParameters VertexParameters)
@@ -861,38 +912,8 @@ float4 VertexFactoryGetTranslatedPrimitiveVolumeBounds(FVertexFactoryInterpolant
 
 		// Fill FVertexFactoryInput with dummy values, the real ones will be fetched later from NiagaraMeshVF using InstanceId  
 		Input.Position = float4(Tri.LocalPositions[VertexIndex], 1.0f);
-		Input.TangentX = 1.0;
-		Input.TangentZ = 0.0;
-		Input.VertexColor = float4(1.0, 1.0, 1.0, 0.5);
-
-#if NUM_MATERIAL_TEXCOORDS_VERTEX
-		Input.TexCoords0.xy = float2(0.5f, 0.5f);
-	#if NUM_MATERIAL_TEXCOORDS_VERTEX > 1
-		Input.TexCoords1.xy = float2(0.5f, 0.5f);
-	#endif
-	#if NUM_MATERIAL_TEXCOORDS_VERTEX > 2
-		Input.TexCoords2.xy = float2(0.5f, 0.5f);
-	#endif
-	#if NUM_MATERIAL_TEXCOORDS_VERTEX > 3
-		Input.TexCoords3.xy = float2(0.5f, 0.5f);
-	#endif
-	#if NUM_MATERIAL_TEXCOORDS_VERTEX > 4
-		Input.TexCoords4.xy = float2(0.5f, 0.5f);
-	#endif
-	#if NUM_MATERIAL_TEXCOORDS_VERTEX > 5
-		Input.TexCoords5.xy = float2(0.5f, 0.5f);
-	#endif
-	#if NUM_MATERIAL_TEXCOORDS_VERTEX > 6
-		Input.TexCoords6.xy = float2(0.5f, 0.5f);
-	#endif
-	#if NUM_MATERIAL_TEXCOORDS_VERTEX > 7
-		Input.TexCoords7.xy = float2(0.5f, 0.5f);
-	#endif
-#elif USE_PARTICLE_SUBUVS
-		Input.TexCoords0.xy = float2(0.5f, 0.5f);
-#endif
 
-#if PASS_NEEDS_VERTEX_ID
+#if PASS_NEEDS_VERTEX_ID || MANUAL_VERTEX_FETCH
 		Input.VertexId = VertexId;
 #endif // #if PASS_NEEDS_VERTEX_ID
 
diff --git a/Engine/Plugins/FX/Niagara/Source/Niagara/Private/NiagaraRendererMeshes.cpp b/Engine/Plugins/FX/Niagara/Source/Niagara/Private/NiagaraRendererMeshes.cpp
--- a/Engine/Plugins/FX/Niagara/Source/Niagara/Private/NiagaraRendererMeshes.cpp
+++ b/Engine/Plugins/FX/Niagara/Source/Niagara/Private/NiagaraRendererMeshes.cpp
@@ -592,7 +592,7 @@ void FNiagaraRendererMeshes::InitializeSortInfo(const FParticleMeshRenderData& P
 	}
 }
 
-void FNiagaraRendererMeshes::PreparePerMeshData(FParticleMeshRenderData& ParticleMeshRenderData, const FNiagaraSceneProxy& SceneProxy, const FMeshData& MeshData) const
+void FNiagaraRendererMeshes::PreparePerMeshData(FParticleMeshRenderData& ParticleMeshRenderData, const FNiagaraMeshVertexFactory& VertexFactory, const FNiagaraSceneProxy& SceneProxy, const FMeshData& MeshData) const
 {
 	// Calculate pivot offset / culling sphere
 	FBox MeshLocalBounds = MeshData.LocalBounds;
@@ -616,6 +616,24 @@ void FNiagaraRendererMeshes::PreparePerMeshData(FParticleMeshRenderData& Particl
 			ParticleMeshRenderData.WorldSpacePivotOffset = SceneProxy.GetLocalToWorld().TransformVector(ParticleMeshRenderData.WorldSpacePivotOffset);
 		}
 	}
+
+	if (RHISupportsManualVertexFetch(GShaderPlatformForFeatureLevel[FeatureLevel]))
+	{
+		const int32 NumTexCoords = VertexFactory.GetNumTexcoords();
+		const int32 ColorIndexMask = VertexFactory.GetColorIndexMask();
+
+		ParticleMeshRenderData.VertexFetch_Parameters = { ColorIndexMask, NumTexCoords, INDEX_NONE, INDEX_NONE };
+		ParticleMeshRenderData.TexCoordBufferSrv = VertexFactory.GetTextureCoordinatesSRV();
+		ParticleMeshRenderData.PackedTangentsBufferSrv = VertexFactory.GetTangentsSRV();
+		ParticleMeshRenderData.ColorComponentsBufferSrv = VertexFactory.GetColorComponentsSRV();
+	}
+	else
+	{
+		ParticleMeshRenderData.VertexFetch_Parameters = FIntVector4(INDEX_NONE);
+		ParticleMeshRenderData.TexCoordBufferSrv = FNiagaraRenderer::GetDummyFloat2Buffer();
+		ParticleMeshRenderData.PackedTangentsBufferSrv = FNiagaraRenderer::GetDummyFloat4Buffer();
+		ParticleMeshRenderData.ColorComponentsBufferSrv = FNiagaraRenderer::GetDummyFloat4Buffer();
+	}
 }
 
 uint32 FNiagaraRendererMeshes::PerformSortAndCull(FParticleMeshRenderData& ParticleMeshRenderData, FGlobalDynamicReadBuffer& ReadBuffer, FNiagaraGPUSortInfo& SortInfo, FNiagaraGpuComputeDispatchInterface* ComputeDispatchInterface, int32 MeshIndex) const
@@ -872,6 +890,11 @@ FNiagaraMeshUniformBufferRef FNiagaraRendererMeshes::CreateVFUniformBuffer(const
 		check(SourceMode <= ENiagaraRendererSourceDataMode::Emitter);
 	}
 
+	Params.VertexFetch_Parameters = ParticleMeshRenderData.VertexFetch_Parameters;
+	Params.VertexFetch_TexCoordBuffer = ParticleMeshRenderData.TexCoordBufferSrv;
+	Params.VertexFetch_PackedTangentsBuffer = ParticleMeshRenderData.PackedTangentsBufferSrv;
+	Params.VertexFetch_ColorComponentsBuffer = ParticleMeshRenderData.ColorComponentsBufferSrv;
+
 	if (bSetAnyBoundVars)
 	{
 		const FNiagaraDynamicDataMesh* DynamicDataMesh = ParticleMeshRenderData.DynamicDataMesh;
@@ -1238,7 +1261,7 @@ void FNiagaraRendererMeshes::GetDynamicMeshElements(const TArray<const FSceneVie
 				VertexFactory.InitResource();
 				SetupVertexFactory(VertexFactory, LODModel);
 
-				PreparePerMeshData(ParticleMeshRenderData, *SceneProxy, MeshData);
+				PreparePerMeshData(ParticleMeshRenderData, VertexFactory, *SceneProxy, MeshData);
 
 				// Sort/Cull particles if needed.
 				const uint32 NumInstances = PerformSortAndCull(ParticleMeshRenderData, Collector.GetDynamicReadBuffer(), SortInfo, ComputeDispatchInterface, MeshData.SourceMeshIndex);
@@ -1398,7 +1421,7 @@ void FNiagaraRendererMeshes::GetDynamicRayTracingInstances(FRayTracingMaterialGa
 		VertexFactory.InitResource();
 		SetupVertexFactory(VertexFactory, LODModel);
 
-		PreparePerMeshData(ParticleMeshRenderData, *SceneProxy, MeshData);
+		PreparePerMeshData(ParticleMeshRenderData, VertexFactory, *SceneProxy, MeshData);
 
 		// Sort/Cull particles if needed.
 		FNiagaraGpuComputeDispatchInterface* ComputeDispatchInterface = SceneProxy->GetComputeDispatchInterface();
diff --git a/Engine/Plugins/FX/Niagara/Source/Niagara/Public/NiagaraRendererMeshes.h b/Engine/Plugins/FX/Niagara/Source/Niagara/Public/NiagaraRendererMeshes.h
--- a/Engine/Plugins/FX/Niagara/Source/Niagara/Public/NiagaraRendererMeshes.h
+++ b/Engine/Plugins/FX/Niagara/Source/Niagara/Public/NiagaraRendererMeshes.h
@@ -63,6 +63,11 @@ protected:
 		FRHIShaderResourceView*			ParticleSortedIndicesSRV = nullptr;
 		uint32							ParticleSortedIndicesOffset = 0xffffffff;
 
+		FIntVector4						VertexFetch_Parameters = FIntVector4(INDEX_NONE);
+		FRHIShaderResourceView*			TexCoordBufferSrv = nullptr;
+		FRHIShaderResourceView*			PackedTangentsBufferSrv = nullptr;
+		FRHIShaderResourceView*			ColorComponentsBufferSrv = nullptr;
+
 		uint32							RendererVisTagOffset = INDEX_NONE;
 		uint32							MeshIndexOffset = INDEX_NONE;
 
@@ -129,7 +134,7 @@ protected:
 	void PrepareParticleMeshRenderData(FParticleMeshRenderData& ParticleMeshRenderData, FMeshElementCollector* Collector, FNiagaraDynamicDataBase* InDynamicData, const FNiagaraSceneProxy* SceneProxy) const;
 	void PrepareParticleRenderBuffers(FParticleMeshRenderData& ParticleMeshRenderData, FGlobalDynamicReadBuffer& DynamicReadBuffer) const;
 	void InitializeSortInfo(const FParticleMeshRenderData& ParticleMeshRenderData, const FNiagaraSceneProxy& SceneProxy, const FSceneView& View, int32 ViewIndex, bool bIsInstancedStereo, FNiagaraGPUSortInfo& OutSortInfo) const;
-	void PreparePerMeshData(FParticleMeshRenderData& ParticleMeshRenderData, const FNiagaraSceneProxy& SceneProxy, const FMeshData& MeshData) const;
+	void PreparePerMeshData(FParticleMeshRenderData& ParticleMeshRenderData, const FNiagaraMeshVertexFactory& VertexFactory, const FNiagaraSceneProxy& SceneProxy, const FMeshData& MeshData) const;
 	uint32 PerformSortAndCull(FParticleMeshRenderData& ParticleMeshRenderData, FGlobalDynamicReadBuffer& ReadBuffer, FNiagaraGPUSortInfo& SortInfo, class FNiagaraGpuComputeDispatchInterface* ComputeDispatchInterface, int32 MeshIndex) const;
 	FNiagaraMeshCommonParameters CreateCommonShaderParams(const FParticleMeshRenderData& ParticleMeshRenderData, const FSceneView& View, const FMeshData& MeshData, const FNiagaraSceneProxy& SceneProxy) const;
 	FNiagaraMeshUniformBufferRef CreateVFUniformBuffer(const FParticleMeshRenderData& ParticleMeshRenderData, const FNiagaraMeshCommonParameters& CommonParams) const;
diff --git a/Engine/Plugins/FX/Niagara/Source/NiagaraVertexFactories/Private/NiagaraMeshVertexFactory.cpp b/Engine/Plugins/FX/Niagara/Source/NiagaraVertexFactories/Private/NiagaraMeshVertexFactory.cpp
--- a/Engine/Plugins/FX/Niagara/Source/NiagaraVertexFactories/Private/NiagaraMeshVertexFactory.cpp
+++ b/Engine/Plugins/FX/Niagara/Source/NiagaraVertexFactories/Private/NiagaraMeshVertexFactory.cpp
@@ -164,6 +164,12 @@ void FNiagaraMeshVertexFactory::ModifyCompilationEnvironment(const FVertexFactor
 	OutEnvironment.SetDefine(TEXT("VF_SUPPORTS_PRIMITIVE_SCENE_DATA"), bSupportsPrimitiveIdStream && bUseGPUScene);
 	OutEnvironment.SetDefine(TEXT("VF_REQUIRES_PER_INSTANCE_CUSTOM_DATA"), bSupportsPrimitiveIdStream && bUseGPUScene);
 #endif
+
+	const bool ContainsManualVertexFetch = OutEnvironment.GetDefinitions().Contains("MANUAL_VERTEX_FETCH");
+	if (!ContainsManualVertexFetch && RHISupportsManualVertexFetch(Parameters.Platform))
+	{
+		OutEnvironment.SetDefine(TEXT("MANUAL_VERTEX_FETCH"), TEXT("1"));
+	}
 }
 
 void FNiagaraMeshVertexFactory::SetData(const FStaticMeshDataType& InData)
diff --git a/Engine/Plugins/FX/Niagara/Source/NiagaraVertexFactories/Public/NiagaraMeshVertexFactory.h b/Engine/Plugins/FX/Niagara/Source/NiagaraVertexFactories/Public/NiagaraMeshVertexFactory.h
--- a/Engine/Plugins/FX/Niagara/Source/NiagaraVertexFactories/Public/NiagaraMeshVertexFactory.h
+++ b/Engine/Plugins/FX/Niagara/Source/NiagaraVertexFactories/Public/NiagaraMeshVertexFactory.h
@@ -79,6 +79,11 @@ END_SHADER_PARAMETER_STRUCT()
 BEGIN_GLOBAL_SHADER_PARAMETER_STRUCT(FNiagaraMeshUniformParameters, NIAGARAVERTEXFACTORIES_API)
 	SHADER_PARAMETER_STRUCT_INCLUDE(FNiagaraMeshCommonParameters, Common)
 
+	SHADER_PARAMETER_SRV(Buffer<float2>, VertexFetch_TexCoordBuffer)
+	SHADER_PARAMETER_SRV(Buffer<float4>, VertexFetch_PackedTangentsBuffer)
+	SHADER_PARAMETER_SRV(Buffer<float4>, VertexFetch_ColorComponentsBuffer)
+	SHADER_PARAMETER(FIntVector4, VertexFetch_Parameters)
+
 	SHADER_PARAMETER(FVector4f, SubImageSize)
 	SHADER_PARAMETER(uint32, TexCoordWeightA)
 	SHADER_PARAMETER(uint32, TexCoordWeightB)
@@ -169,6 +174,31 @@ public:
 		return MeshParticleUniformBuffer;
 	}
 
+	FORCEINLINE FRHIShaderResourceView* GetTangentsSRV() const
+	{
+		return Data.TangentsSRV;
+	}
+
+	FORCEINLINE FRHIShaderResourceView* GetTextureCoordinatesSRV() const
+	{
+		return Data.TextureCoordinatesSRV;
+	}
+
+	FORCEINLINE FRHIShaderResourceView* GetColorComponentsSRV() const
+	{
+		return Data.ColorComponentsSRV;
+	}
+
+	FORCEINLINE uint32 GetColorIndexMask() const
+	{
+		return Data.ColorIndexMask;
+	}
+
+	FORCEINLINE int GetNumTexcoords() const
+	{
+		return Data.NumTexCoords;
+	}
+	
 	/**
 	* Copy the data from another vertex factory
 	* @param Other - factory to copy from
