Move NiagaraGPUSceneUtils to NiagaraShader module to fix load order issues with the global shader it uses

diff --git a/Engine/Plugins/FX/Niagara/Source/Niagara/Private/NiagaraGPUSceneUtils.cpp b/Engine/Plugins/FX/Niagara/Source/Niagara/Private/NiagaraGPUSceneUtils.cpp
deleted file mode 100644
diff --git a/Engine/Plugins/FX/Niagara/Source/Niagara/Private/NiagaraGPUSceneUtils.h b/Engine/Plugins/FX/Niagara/Source/Niagara/Private/NiagaraGPUSceneUtils.h
deleted file mode 100644
diff --git a/Engine/Plugins/FX/Niagara/Source/NiagaraShader/NiagaraShader.Build.cs b/Engine/Plugins/FX/Niagara/Source/NiagaraShader/NiagaraShader.Build.cs
--- a/Engine/Plugins/FX/Niagara/Source/NiagaraShader/NiagaraShader.Build.cs
+++ b/Engine/Plugins/FX/Niagara/Source/NiagaraShader/NiagaraShader.Build.cs
@@ -12,6 +12,7 @@ public class NiagaraShader : ModuleRules
                 "CoreUObject",
                 "Engine",
                 "NiagaraCore",
+                "NiagaraVertexFactories",
                 "Renderer",
             }
         );
@@ -23,7 +24,8 @@ public class NiagaraShader : ModuleRules
                 "VectorVM",
                 "RHI",
                 "Projects",
-                "NiagaraCore"
+                "NiagaraCore",
+                "NiagaraVertexFactories",
             }
         );
 
diff --git a/Engine/Plugins/FX/Niagara/Source/NiagaraShader/Private/NiagaraGPUSceneUtils.cpp b/Engine/Plugins/FX/Niagara/Source/NiagaraShader/Private/NiagaraGPUSceneUtils.cpp
new file mode 100644
--- /dev/null
+++ b/Engine/Plugins/FX/Niagara/Source/NiagaraShader/Private/NiagaraGPUSceneUtils.cpp
@@ -0,0 +1,55 @@
+// Copyright Epic Games, Inc. All Rights Reserved.
+
+#include "NiagaraGPUSceneUtils.h"
+#include "GlobalShader.h"
+#include "RenderGraphUtils.h"
+
+class FNiagaraUpdateMeshGPUSceneInstancesCS : public FGlobalShader
+{
+public:
+	static constexpr uint32 ThreadGroupSize = 64;
+
+	DECLARE_GLOBAL_SHADER(FNiagaraUpdateMeshGPUSceneInstancesCS);
+	SHADER_USE_PARAMETER_STRUCT(FNiagaraUpdateMeshGPUSceneInstancesCS, FGlobalShader);
+
+	using FParameters = FNiagaraGPUSceneUtils::FUpdateMeshParticleInstancesParams;
+
+	class FPreciseMotionVectors : SHADER_PERMUTATION_BOOL("ENABLE_PRECISE_MOTION_VECTORS");
+	using FPermutationDomain = TShaderPermutationDomain<FPreciseMotionVectors>;
+
+	static bool ShouldCompilePermutation(const FGlobalShaderPermutationParameters& Parameters)
+	{
+		return UseGPUScene(Parameters.Platform, GetMaxSupportedFeatureLevel(Parameters.Platform));
+	}
+
+	static void ModifyCompilationEnvironment(const FGlobalShaderPermutationParameters& Parameters, FShaderCompilerEnvironment& OutEnvironment)
+	{
+		FGlobalShader::ModifyCompilationEnvironment(Parameters, OutEnvironment);
+
+		OutEnvironment.SetDefine(TEXT("VF_SUPPORTS_PRIMITIVE_SCENE_DATA"), 1);
+		OutEnvironment.SetDefine(TEXT("THREAD_GROUP_SIZE"), ThreadGroupSize);
+	}
+};
+
+IMPLEMENT_GLOBAL_SHADER(FNiagaraUpdateMeshGPUSceneInstancesCS, "/Plugin/FX/Niagara/Private/NiagaraUpdateMeshGPUSceneInstances.usf", "UpdateMeshInstancesCS", SF_Compute);
+
+void FNiagaraGPUSceneUtils::AddUpdateMeshParticleInstancesPass(
+	FRDGBuilder& GraphBuilder, 
+	FUpdateMeshParticleInstancesParams& Params,
+	ERHIFeatureLevel::Type FeatureLevel,
+	bool bPreciseMotionVectors
+)
+{
+	FNiagaraUpdateMeshGPUSceneInstancesCS::FPermutationDomain PermutationVector;
+	PermutationVector.Set<FNiagaraUpdateMeshGPUSceneInstancesCS::FPreciseMotionVectors>(bPreciseMotionVectors);
+
+	TShaderMapRef<FNiagaraUpdateMeshGPUSceneInstancesCS> UpdateMeshInstancesCS(GetGlobalShaderMap(FeatureLevel), PermutationVector);
+
+	FComputeShaderUtils::AddPass(
+		GraphBuilder,
+		RDG_EVENT_NAME("Niagara.UpdateMeshParticleInstances (%u Instances)", Params.ParticleCount),
+		UpdateMeshInstancesCS,
+		&Params,
+		FComputeShaderUtils::GetGroupCountWrapped(Params.ParticleCount, FNiagaraUpdateMeshGPUSceneInstancesCS::ThreadGroupSize)
+	);
+}
\ No newline at end of file
diff --git a/Engine/Plugins/FX/Niagara/Source/NiagaraShader/Public/NiagaraGPUSceneUtils.h b/Engine/Plugins/FX/Niagara/Source/NiagaraShader/Public/NiagaraGPUSceneUtils.h
new file mode 100644
--- /dev/null
+++ b/Engine/Plugins/FX/Niagara/Source/NiagaraShader/Public/NiagaraGPUSceneUtils.h
@@ -0,0 +1,43 @@
+// Copyright Epic Games, Inc. All Rights Reserved.
+
+#pragma once
+
+#include "CoreMinimal.h"
+#include "RHI.h"
+#include "RHIDefinitions.h"
+#include "ShaderParameterStruct.h"
+#include "ShaderParameterUtils.h"
+#include "NiagaraMeshVertexFactory.h"
+#include "SceneView.h"
+
+class NIAGARASHADER_API FNiagaraGPUSceneUtils
+{
+public:	
+	BEGIN_SHADER_PARAMETER_STRUCT(FUpdateMeshParticleInstancesParams, NIAGARASHADER_API)
+		SHADER_PARAMETER_STRUCT_REF(FViewUniformShaderParameters, View)
+		SHADER_PARAMETER_STRUCT_INCLUDE(FGPUSceneWriterParameters, GPUSceneWriterParameters)
+		SHADER_PARAMETER_STRUCT_INCLUDE(FNiagaraMeshCommonParameters, Common)
+	
+		SHADER_PARAMETER(uint32, PrimitiveId)
+		SHADER_PARAMETER(uint32, ParticleCount)
+		SHADER_PARAMETER_SRV(Buffer<uint>, GPUParticleCountBuffer)
+		SHADER_PARAMETER(uint32, GPUParticleCountOffset)
+		
+		SHADER_PARAMETER(int, MeshIndex)
+		SHADER_PARAMETER(int, MeshIndexDataOffset)
+		SHADER_PARAMETER(int, RendererVisibility)
+		SHADER_PARAMETER(int, VisibilityTagDataOffset)
+
+		SHADER_PARAMETER(FVector3f, LocalBoundingCenter)
+		SHADER_PARAMETER(FVector2f, DistanceCullRangeSquared)
+
+		SHADER_PARAMETER(int, bNeedsPrevTransform)
+	END_SHADER_PARAMETER_STRUCT()
+
+	static void AddUpdateMeshParticleInstancesPass(
+		FRDGBuilder& GraphBuilder,
+		FUpdateMeshParticleInstancesParams& Params,
+		ERHIFeatureLevel::Type FeatureLevel,
+		bool bPreciseMotionVectors
+	);
+};
\ No newline at end of file
