Fixes to reduce the number of UObjects that get created during precompile for niagara systems (reduces object counts by up 50-80% depending on emitter count)
-flattens reroute nodes from the graph during duplication
-ensure duplicate data is properly cached so that it can be reused by related scripts
-defer CDO creation so that it's done once rather than per emitter in a system
-avoid duplicating system graphs unless necessary

diff --git a/Engine/Plugins/FX/Niagara/Source/NiagaraEditor/Private/NiagaraCompiler.cpp b/Engine/Plugins/FX/Niagara/Source/NiagaraEditor/Private/NiagaraCompiler.cpp
--- a/Engine/Plugins/FX/Niagara/Source/NiagaraEditor/Private/NiagaraCompiler.cpp
+++ b/Engine/Plugins/FX/Niagara/Source/NiagaraEditor/Private/NiagaraCompiler.cpp
@@ -677,9 +677,12 @@ void FNiagaraCompileRequestDuplicateData::FinishPrecompileDuplicate(const TArray
 			PrecompiledHistories.Emplace();
 		}
 	}
+}
 
+void FNiagaraCompileRequestDuplicateData::CreateDataInterfaceCDO(TArrayView<UClass*> VariableDataInterfaces)
+{
 	// Collect classes for any data interfaces found in the duplicated graphs
-	TArray<UClass*> DataInterfaceClasses;
+	TArray<UClass*> DataInterfaceClasses(VariableDataInterfaces);
 	for (const TPair<const UNiagaraGraph*, TArray<FDuplicatedGraphData>>& SourceGraphToDuplicatedGraphs : *(SharedSourceGraphToDuplicatedGraphsMap.Get()))
 	{
 		for (const FDuplicatedGraphData& DuplicatedGraphData : SourceGraphToDuplicatedGraphs.Value)
@@ -696,15 +699,6 @@ void FNiagaraCompileRequestDuplicateData::FinishPrecompileDuplicate(const TArray
 		}
 	}
 
-	// Collect classes for external encounterable variables
-	for (const FNiagaraVariable& EncounterableVariable : EncounterableVariables)
-	{
-		if (EncounterableVariable.IsDataInterface())
-		{
-			DataInterfaceClasses.AddUnique(EncounterableVariable.GetType().GetClass());
-		}
-	}
-
 	// Generate copies of the CDOs for any encountered data interfaces.
 	for (UClass* DataInterfaceClass : DataInterfaceClasses)
 	{
@@ -1029,6 +1023,22 @@ TSharedPtr<FNiagaraCompileRequestDuplicateDataBase, ESPMode::ThreadSafe> FNiagar
 	BasePtr->SharedSourceGraphToDuplicatedGraphsMap = MakeShared<TMap<const UNiagaraGraph*, TArray<FNiagaraCompileRequestDuplicateData::FDuplicatedGraphData>>>();
 	BasePtr->SharedNameToDuplicatedDataInterfaceMap = MakeShared<TMap<FName, UNiagaraDataInterface*>>();
 	BasePtr->SharedDataInterfaceClassToDuplicatedCDOMap = MakeShared<TMap<UClass*, UNiagaraDataInterface*>>();
+	BasePtr->OwningSystem = OwningSystem;
+	BasePtr->OwningEmitter = OwningEmitter;
+
+	TArray<UClass*> DataInterfaceClasses;
+
+	auto CollectDataInterfaceClasses = [&](TConstArrayView<FNiagaraVariable> Variables)
+	{
+		// Collect classes for external encounterable variables
+		for (const FNiagaraVariable& EncounterableVariable : Variables)
+		{
+			if (EncounterableVariable.IsDataInterface())
+			{
+				DataInterfaceClasses.AddUnique(EncounterableVariable.GetType().GetClass());
+			}
+		}
+	};
 
 	if (OwningSystem == nullptr)
 	{
@@ -1038,6 +1048,7 @@ TSharedPtr<FNiagaraCompileRequestDuplicateDataBase, ESPMode::ThreadSafe> FNiagar
 		TArray<FNiagaraVariable> EncounterableScriptVariables;
 		OwningSystemRequestData->GatherPreCompiledVariables(FString(), EncounterableScriptVariables);
 		BasePtr->FinishPrecompileDuplicate(EncounterableScriptVariables, EmptyResolver, nullptr);
+		CollectDataInterfaceClasses(EncounterableScriptVariables);
 	}
 	else
 	{
@@ -1045,13 +1056,6 @@ TSharedPtr<FNiagaraCompileRequestDuplicateDataBase, ESPMode::ThreadSafe> FNiagar
 
 		check(OwningSystem->GetSystemSpawnScript()->GetLatestSource() == OwningSystem->GetSystemUpdateScript()->GetLatestSource());
 
-		// Store off the current variables in the exposed parameters list.
-		TArray<FNiagaraVariable> OriginalExposedParams;
-		OwningSystem->GetExposedParameters().GetParameters(OriginalExposedParams);
-
-		// Create an array of variables that we might encounter when traversing the graphs (include the originally exposed vars above)
-		TArray<FNiagaraVariable> EncounterableVars(OriginalExposedParams);
-
 		// First deep copy all the emitter graphs referenced by the system so that we can later hook up emitter handles in the system traversal.
 		BasePtr->EmitterData.Empty();
 		for (int32 i = 0; i < OwningSystem->GetEmitterHandles().Num(); i++)
@@ -1074,28 +1078,42 @@ TSharedPtr<FNiagaraCompileRequestDuplicateDataBase, ESPMode::ThreadSafe> FNiagar
 
 		// Now deep copy the system graphs, skipping traversal into any emitter references.
 		{
-			FCompileConstantResolver ConstantResolver(OwningSystem, ENiagaraScriptUsage::SystemSpawnScript);
 			UNiagaraScriptSource* Source = Cast<UNiagaraScriptSource>(OwningSystem->GetSystemSpawnScript()->GetLatestSource());
-			BasePtr->DeepCopyGraphs(Source, ENiagaraScriptUsage::SystemSpawnScript, ConstantResolver);
+
 			TArray<FNiagaraVariable> EncounterableSystemVariables;
 			OwningSystemRequestData->GatherPreCompiledVariables(FString(), EncounterableSystemVariables);
-			BasePtr->FinishPrecompileDuplicate(EncounterableSystemVariables, ConstantResolver, nullptr);
+
+			// skip the deep copy if we're not compiling the system scripts
+			if (BasePtr->ValidUsages.Contains(ENiagaraScriptUsage::SystemSpawnScript))
+			{
+				FCompileConstantResolver ConstantResolver(OwningSystem, ENiagaraScriptUsage::SystemSpawnScript);
+				BasePtr->DeepCopyGraphs(Source, ENiagaraScriptUsage::SystemSpawnScript, ConstantResolver);
+				BasePtr->FinishPrecompileDuplicate(EncounterableSystemVariables, ConstantResolver, nullptr);
+			}
+
+			CollectDataInterfaceClasses(EncounterableSystemVariables);
 		}
 
 		// Now we can finish off the emitters.
 		for (int32 i = 0; i < OwningSystem->GetEmitterHandles().Num(); i++)
 		{
 			const FNiagaraEmitterHandle& Handle = OwningSystem->GetEmitterHandle(i);
-			FCompileConstantResolver ConstantResolver(Handle.GetInstance(), ENiagaraScriptUsage::EmitterSpawnScript);
-			if (Handle.GetIsEnabled())
+
+			TArray<FNiagaraVariable> EncounterableEmitterVariables;
+			OwningSystemRequestData->GetDependentRequest(i)->GatherPreCompiledVariables(FString(), EncounterableEmitterVariables);
+
+			if (Handle.GetIsEnabled() && (OwningEmitter == nullptr || OwningEmitter == Handle.GetInstance()))
 			{
-				TArray<FNiagaraVariable> EncounterableEmitterVariables;
-				OwningSystemRequestData->GetDependentRequest(i)->GatherPreCompiledVariables(FString(), EncounterableEmitterVariables);
+				FCompileConstantResolver ConstantResolver(Handle.GetInstance(), ENiagaraScriptUsage::EmitterSpawnScript);
 				BasePtr->EmitterData[i]->FinishPrecompileDuplicate(EncounterableEmitterVariables, ConstantResolver, &Handle.GetInstance()->GetSimulationStages());
 			}
+
+			CollectDataInterfaceClasses(EncounterableEmitterVariables);
 		}
 	}
 
+	BasePtr->CreateDataInterfaceCDO(DataInterfaceClasses);
+
 	UE_LOG(LogNiagaraEditor, Verbose, TEXT("'%s' PrecompileDuplicate took %f sec."), *LogName,
 		(float)(FPlatformTime::Seconds() - StartTime));
 
diff --git a/Engine/Plugins/FX/Niagara/Source/NiagaraEditor/Private/NiagaraGraph.cpp b/Engine/Plugins/FX/Niagara/Source/NiagaraEditor/Private/NiagaraGraph.cpp
--- a/Engine/Plugins/FX/Niagara/Source/NiagaraEditor/Private/NiagaraGraph.cpp
+++ b/Engine/Plugins/FX/Niagara/Source/NiagaraEditor/Private/NiagaraGraph.cpp
@@ -1028,9 +1028,9 @@ UNiagaraGraph* UNiagaraGraph::CreateCompilationCopy(const TArray<ENiagaraScriptU
 			continue;
 		}
 		const uint8* SourceAddr = Property->ContainerPtrToValuePtr<uint8>(this);
-        uint8* DestinationAddr = Property->ContainerPtrToValuePtr<uint8>(Result);
+		uint8* DestinationAddr = Property->ContainerPtrToValuePtr<uint8>(Result);
 
-        Property->CopyCompleteValue(DestinationAddr, SourceAddr);
+		Property->CopyCompleteValue(DestinationAddr, SourceAddr);
 	}
 	Result->bIsForCompilationOnly = true;
 
@@ -1040,7 +1040,7 @@ UNiagaraGraph* UNiagaraGraph::CreateCompilationCopy(const TArray<ENiagaraScriptU
 		It.Value = CastChecked<UNiagaraScriptVariable>(NiagaraEditorModule.GetPooledDuplicateObject(It.Value));
 	}
 
-	// Only add references to nodes which match a complie usage.
+	// Only add references to nodes which match a compile usage.
 	TMap<UEdGraphNode*, UEdGraphNode*> DuplicationMapping;
 	TArray<UNiagaraNodeOutput*> OutputNodes;
 	GetNodesOfClass(OutputNodes);
@@ -1054,9 +1054,16 @@ UNiagaraGraph* UNiagaraGraph::CreateCompilationCopy(const TArray<ENiagaraScriptU
 		}
 	}
 
+	TSet<const UEdGraphNode*> RerouteNodesToFixup;
 	TArray<UEdGraphNode*> NewNodes;
 	for (UEdGraphNode* Node : Result->Nodes)
 	{
+		if (const UNiagaraNodeReroute* RerouteNode = Cast<const UNiagaraNodeReroute>(Node))
+		{
+			RerouteNodesToFixup.Add(RerouteNode);
+			continue;
+		}
+
 		Node->ClearFlags(RF_Transactional);
 		UEdGraphNode* DupNode = NewNodes.Add_GetRef(DuplicateObject(Node, Result));
 		Node->SetFlags(RF_Transactional);
@@ -1070,7 +1077,32 @@ UNiagaraGraph* UNiagaraGraph::CreateCompilationCopy(const TArray<ENiagaraScriptU
 			for (int i = 0; i < Pin->LinkedTo.Num(); i++)
 			{
 				UEdGraphPin* LinkedPin = Pin->LinkedTo[i];
-				UEdGraphNode** NewNodePtr = DuplicationMapping.Find(LinkedPin->GetOwningNode());
+				UEdGraphNode* OwningNode = LinkedPin->GetOwningNode();
+
+				// if the owning node is already in the right graph then no need to adjust it (this can
+				// happen if the pin connection is added as a part of removing the reroute nodes)
+				if (OwningNode->GetGraph() == Result)
+				{
+					continue;
+				}
+
+				UEdGraphNode** NewNodePtr = DuplicationMapping.Find(OwningNode);
+				bool bAddReciprocalLink = false;
+
+				if (NewNodePtr == nullptr)
+				{
+					// check to see if we're dealing with an input pin that is coming from
+					// a reroute node that was trimmed above
+					if (Pin->Direction == EGPD_Input && RerouteNodesToFixup.Contains(OwningNode))
+					{
+						check(LinkedPin->Direction == EGPD_Output);
+						LinkedPin = CastChecked<UNiagaraNodeReroute>(OwningNode)->GetTracedOutputPin(LinkedPin, false);
+						OwningNode = LinkedPin->GetOwningNode();
+						NewNodePtr = DuplicationMapping.Find(OwningNode);
+						bAddReciprocalLink = true;
+					}
+				}
+
 				if (NewNodePtr == nullptr)
 				{
 					// If output pins were connected to another node which wasn't encountered in traversal then it's possible it won't be found here.
@@ -1082,6 +1114,11 @@ UNiagaraGraph* UNiagaraGraph::CreateCompilationCopy(const TArray<ENiagaraScriptU
 					UEdGraphPin** NodePin = (*NewNodePtr)->Pins.FindByPredicate([LinkedPin](UEdGraphPin* Pin) { return Pin->PinId == LinkedPin->PinId; });
 					check(NodePin);
 					Pin->LinkedTo[i] = *NodePin;
+
+					if (bAddReciprocalLink)
+					{
+						(*NodePin)->LinkedTo.AddUnique(Pin);
+					}
 				}
 			}
 			// Remove any null linked pins.
diff --git a/Engine/Plugins/FX/Niagara/Source/NiagaraEditor/Private/NiagaraHlslTranslator.h b/Engine/Plugins/FX/Niagara/Source/NiagaraEditor/Private/NiagaraHlslTranslator.h
--- a/Engine/Plugins/FX/Niagara/Source/NiagaraEditor/Private/NiagaraHlslTranslator.h
+++ b/Engine/Plugins/FX/Niagara/Source/NiagaraEditor/Private/NiagaraHlslTranslator.h
@@ -144,6 +144,7 @@ public:
 	void DeepCopyGraphs(UNiagaraEmitter* Emitter);
 	
 	void FinishPrecompileDuplicate(const TArray<FNiagaraVariable>& EncounterableVariables, FCompileConstantResolver ConstantResolver, const TArray<class UNiagaraSimulationStageBase*>* SimStages);
+	void CreateDataInterfaceCDO(TArrayView<UClass*> VariableDataInterfaces);
 
 	virtual int32 GetDependentRequestCount() const override	{ return EmitterData.Num();	}
 	virtual TSharedPtr<FNiagaraCompileRequestDuplicateDataBase, ESPMode::ThreadSafe> GetDependentRequest(int32 Index) override { return EmitterData[Index];	}
