- Add per mesh rotation to mesh renderer

diff --git a/Engine/Plugins/FX/Niagara/Shaders/Private/NiagaraGPURayTracingTransforms.usf b/Engine/Plugins/FX/Niagara/Shaders/Private/NiagaraGPURayTracingTransforms.usf
--- a/Engine/Plugins/FX/Niagara/Shaders/Private/NiagaraGPURayTracingTransforms.usf
+++ b/Engine/Plugins/FX/Niagara/Shaders/Private/NiagaraGPURayTracingTransforms.usf
@@ -30,6 +30,7 @@ float3	DefaultPosition;
 float4	DefaultRotation;
 float3	DefaultScale;
 float3	MeshScale;
+float4	MeshRotation;
 
 float4x4	LocalTransform;
 
@@ -87,7 +88,7 @@ void NiagaraGPURayTracingTransformsCS(uint3 DispatchThreadId : SV_DispatchThread
 	if (InstanceIndex < GPUNumInstances)
 	{
 		float3 InstancePosition = SafeGetVec3(PositionDataOffset, InstanceIndex, DefaultPosition);
-		float4 InstanceRotation = SafeGetVec4(RotationDataOffset, InstanceIndex, DefaultRotation);
+		float4 InstanceRotation = NiagaraQuatMul(SafeGetVec4(RotationDataOffset, InstanceIndex, DefaultRotation), MeshRotation);
 		float3 InstanceScale    = SafeGetVec3(ScaleDataOffset, InstanceIndex, DefaultScale) * MeshScale;
 
 		float4x4 InstanceTransform = NiagaraComposeTransformMatrix(InstanceScale, NiagaraQuatTo3x3(InstanceRotation), InstancePosition);
diff --git a/Engine/Plugins/FX/Niagara/Shaders/Private/NiagaraMeshParticleUtils.ush b/Engine/Plugins/FX/Niagara/Shaders/Private/NiagaraMeshParticleUtils.ush
--- a/Engine/Plugins/FX/Niagara/Shaders/Private/NiagaraMeshParticleUtils.ush
+++ b/Engine/Plugins/FX/Niagara/Shaders/Private/NiagaraMeshParticleUtils.ush
@@ -53,6 +53,7 @@ struct NiagaraMeshParticleSRTParams
 
 	// Additional mesh scale and offset to apply to mesh
 	float3 MeshScale;
+	float4 MeshRotation;
 	float3 MeshOffset;
 	bool bMeshOffsetIsWorldSpace;
 
@@ -104,6 +105,7 @@ struct NiagaraMeshParticleTransformsParams
 
 	// Additional scale and offset to apply to mesh
 	float3 MeshScale;
+	float4 MeshRotation;
 	float3 MeshOffset;
 	bool bMeshOffsetIsWorldSpace;
 
@@ -352,7 +354,7 @@ NiagaraMeshParticleSRT NiagaraCalculateMeshParticleSRT(NiagaraMeshParticleSRTPar
 	
 	////////////////////////////////////////////////////////////////////////////////
 	// Rotation
-	SRT.Rotation = NiagaraQuatTo3x3(normalize(ParticleRotation));
+	SRT.Rotation = NiagaraQuatTo3x3(NiagaraQuatMul(normalize(ParticleRotation), Params.MeshRotation));
 	if (Params.FacingMode != MESH_FACING_DEFAULT)
 	{
 		// Factor in facing rotation
@@ -390,6 +392,7 @@ NiagaraMeshParticleTransforms NiagaraCalculateMeshParticleTransforms(NiagaraMesh
 	SRTParams.bLocalSpace 					= Params.bLocalSpace;
 	SRTParams.FacingMode					= Params.FacingMode;
 	SRTParams.MeshScale 					= Params.MeshScale;
+	SRTParams.MeshRotation					= Params.MeshRotation;
 	SRTParams.MeshOffset 					= Params.MeshOffset;
 	SRTParams.bMeshOffsetIsWorldSpace		= Params.bMeshOffsetIsWorldSpace;
 	SRTParams.bLockedAxisEnable				= Params.bLockedAxisEnable;
diff --git a/Engine/Plugins/FX/Niagara/Shaders/Private/NiagaraMeshVertexFactory.ush b/Engine/Plugins/FX/Niagara/Shaders/Private/NiagaraMeshVertexFactory.ush
--- a/Engine/Plugins/FX/Niagara/Shaders/Private/NiagaraMeshVertexFactory.ush
+++ b/Engine/Plugins/FX/Niagara/Shaders/Private/NiagaraMeshVertexFactory.ush
@@ -520,6 +520,7 @@ FVertexFactoryIntermediates GetVertexFactoryIntermediates(FVertexFactoryInput In
 	Params.FacingMode					= NiagaraMeshVF.FacingMode;
 	Params.DeltaSeconds					= NiagaraMeshVF.DeltaSeconds;
 	Params.MeshScale 					= NiagaraMeshVF.MeshScale;
+	Params.MeshRotation					= NiagaraMeshVF.MeshRotation;
 	Params.MeshOffset 					= NiagaraMeshVF.MeshOffset;
 	Params.bMeshOffsetIsWorldSpace		= NiagaraMeshVF.bMeshOffsetIsWorldSpace;
 	Params.bLockedAxisEnable			= NiagaraMeshVF.bLockedAxisEnable;
diff --git a/Engine/Plugins/FX/Niagara/Shaders/Private/NiagaraUpdateMeshGPUSceneInstances.usf b/Engine/Plugins/FX/Niagara/Shaders/Private/NiagaraUpdateMeshGPUSceneInstances.usf
--- a/Engine/Plugins/FX/Niagara/Shaders/Private/NiagaraUpdateMeshGPUSceneInstances.usf
+++ b/Engine/Plugins/FX/Niagara/Shaders/Private/NiagaraUpdateMeshGPUSceneInstances.usf
@@ -36,6 +36,7 @@ float3 DefaultPrevVelocity;
 float DefaultPrevCameraOffset;
 
 float3 MeshScale;
+float4 MeshRotation;
 int bMeshOffsetIsWorldSpace;
 float3 MeshOffset;
 
@@ -130,6 +131,7 @@ void UpdateMeshInstancesCS(uint3 GroupId : SV_GroupID, uint GroupIndex : SV_Grou
 	Params.DeltaSeconds					= DeltaSeconds;
 	Params.MeshScale 					= MeshScale;
 	Params.MeshOffset 					= MeshOffset;
+	Params.MeshRotation					= MeshRotation;
 	Params.bMeshOffsetIsWorldSpace		= bMeshOffsetIsWorldSpace;
 	Params.bLockedAxisEnable			= bLockedAxisEnable;
 	Params.LockedAxis					= LockedAxis;
diff --git a/Engine/Plugins/FX/Niagara/Source/Niagara/Private/NiagaraMeshRendererProperties.cpp b/Engine/Plugins/FX/Niagara/Source/Niagara/Private/NiagaraMeshRendererProperties.cpp
--- a/Engine/Plugins/FX/Niagara/Source/Niagara/Private/NiagaraMeshRendererProperties.cpp
+++ b/Engine/Plugins/FX/Niagara/Source/Niagara/Private/NiagaraMeshRendererProperties.cpp
@@ -52,6 +52,7 @@ FNiagaraMeshRendererMeshProperties::FNiagaraMeshRendererMeshProperties()
 	: Mesh(nullptr)
 	, UserParamBinding(FNiagaraTypeDefinition(UStaticMesh::StaticClass()))
 	, Scale(1.0f, 1.0f, 1.0f)
+	, Rotation(FRotator::ZeroRotator)
 	, PivotOffset(ForceInitToZero)
 	, PivotOffsetSpace(ENiagaraMeshPivotOffsetSpace::Mesh)
 {
diff --git a/Engine/Plugins/FX/Niagara/Source/Niagara/Private/NiagaraRendererMeshes.cpp b/Engine/Plugins/FX/Niagara/Source/Niagara/Private/NiagaraRendererMeshes.cpp
--- a/Engine/Plugins/FX/Niagara/Source/Niagara/Private/NiagaraRendererMeshes.cpp
+++ b/Engine/Plugins/FX/Niagara/Source/Niagara/Private/NiagaraRendererMeshes.cpp
@@ -187,6 +187,7 @@ void FNiagaraRendererMeshes::Initialize(const UNiagaraRendererProperties* InProp
 			MeshData.PivotOffset = MeshProperties.PivotOffset;
 			MeshData.PivotOffsetSpace = MeshProperties.PivotOffsetSpace;
 			MeshData.Scale = MeshProperties.Scale;
+			MeshData.Rotation = FQuat4f(MeshProperties.Rotation.Quaternion());
 			MeshData.LocalBounds = Mesh->GetExtendedBounds().GetBox();
 
 			// Create an index remap from mesh material index to it's index in the master material list
@@ -698,6 +699,7 @@ FNiagaraMeshCommonParameters FNiagaraRendererMeshes::CreateCommonShaderParams(co
 	Params.FacingMode 				= (uint32)FacingMode;	
 
 	Params.MeshScale				= MeshData.Scale;	
+	Params.MeshRotation				= FVector4f(MeshData.Rotation.X, MeshData.Rotation.Y, MeshData.Rotation.Z, MeshData.Rotation.W);
 	if (MeshData.PivotOffsetSpace == ENiagaraMeshPivotOffsetSpace::Mesh)
 	{
 		Params.MeshOffset 				= MeshData.PivotOffset;
@@ -1617,6 +1619,8 @@ void FNiagaraRendererMeshes::GetDynamicRayTracingInstances(FRayTracingMaterialGa
 				const bool bHasScale = ScaleBaseCompOffset > 0;
 
 				const FMatrix NullInstanceTransform(FVector::ZeroVector, FVector::ZeroVector, FVector::ZeroVector, FVector::ZeroVector);
+				const FQuat MeshRotation = FQuat(MeshData.Rotation);
+				const FVector MeshScale = FVector(MeshData.Scale);
 				for (uint32 InstanceIndex = 0; InstanceIndex < NumInstances; InstanceIndex++)
 				{
 					if ( RenderVisibilityData && (RenderVisibilityData[InstanceIndex] != RendererVisibility) )
@@ -1631,10 +1635,10 @@ void FNiagaraRendererMeshes::GetDynamicRayTracingInstances(FRayTracingMaterialGa
 					}
 
 					const FVector InstancePosition = bHasPosition ? GetInstancePosition(InstanceIndex) : FVector::ZeroVector;
-					const FQuat InstanceRotation = bHasRotation ? GetInstanceQuat(InstanceIndex).GetNormalized() : FQuat::Identity;
+					const FQuat InstanceRotation = bHasRotation ? GetInstanceQuat(InstanceIndex).GetNormalized() * MeshRotation : MeshRotation;
 					FMatrix InstanceTransform = FQuatRotationTranslationMatrix::Make(InstanceRotation, InstancePosition);
 
-					const FVector InstanceScale = bHasScale ? GetInstanceScale(InstanceIndex) * MeshData.Scale : MeshData.Scale;
+					const FVector InstanceScale = bHasScale ? GetInstanceScale(InstanceIndex) * MeshScale : MeshScale;
 					InstanceTransform = FScaleMatrix(InstanceScale) * InstanceTransform;
 
 					if (bLocalSpace)
@@ -1680,6 +1684,7 @@ void FNiagaraRendererMeshes::GetDynamicRayTracingInstances(FRayTracingMaterialGa
 					PassParameters->DefaultRotation				= FVector4f(0.0f, 0.0f, 0.0f, 1.0f);
 					PassParameters->DefaultScale				= FVector3f(1.0f, 1.0f, 1.0f);
 					PassParameters->MeshScale					= MeshData.Scale;
+					PassParameters->MeshRotation				= FVector4f(MeshData.Rotation.X, MeshData.Rotation.Y, MeshData.Rotation.Z, MeshData.Rotation.W);
 					PassParameters->ParticleDataFloatBuffer		= ParticleMeshRenderData.ParticleFloatSRV;
 					//PassParameters.ParticleDataHalfBuffer		= ParticleMeshRenderData.ParticleHalfSRV;
 					PassParameters->ParticleDataIntBuffer		= ParticleMeshRenderData.ParticleIntSRV;
diff --git a/Engine/Plugins/FX/Niagara/Source/Niagara/Public/NiagaraMeshRendererProperties.h b/Engine/Plugins/FX/Niagara/Source/Niagara/Public/NiagaraMeshRendererProperties.h
--- a/Engine/Plugins/FX/Niagara/Source/Niagara/Public/NiagaraMeshRendererProperties.h
+++ b/Engine/Plugins/FX/Niagara/Source/Niagara/Public/NiagaraMeshRendererProperties.h
@@ -134,6 +134,10 @@ struct NIAGARA_API FNiagaraMeshRendererMeshProperties
 	UPROPERTY(EditAnywhere, Category = "Mesh")
 	FVector Scale;
 
+	/** Rotation of the mesh */
+	UPROPERTY(EditAnywhere, Category = "Mesh")
+	FRotator Rotation;
+
 	/** Offset of the mesh pivot */
 	UPROPERTY(EditAnywhere, Category = "Mesh")
 	FVector PivotOffset;
diff --git a/Engine/Plugins/FX/Niagara/Source/Niagara/Public/NiagaraRendererMeshes.h b/Engine/Plugins/FX/Niagara/Source/Niagara/Public/NiagaraRendererMeshes.h
--- a/Engine/Plugins/FX/Niagara/Source/Niagara/Public/NiagaraRendererMeshes.h
+++ b/Engine/Plugins/FX/Niagara/Source/Niagara/Public/NiagaraRendererMeshes.h
@@ -80,9 +80,10 @@ protected:
 		FStaticMeshRenderData* RenderData = nullptr;
 		int32 MinimumLOD = 0;
 		uint32 SourceMeshIndex = INDEX_NONE;
-		FVector PivotOffset = FVector(ForceInitToZero);
+		FVector3f PivotOffset = FVector3f::ZeroVector;
 		ENiagaraMeshPivotOffsetSpace PivotOffsetSpace = ENiagaraMeshPivotOffsetSpace::Mesh;
-		FVector Scale = FVector(1.0f, 1.0f, 1.0f);
+		FVector3f Scale = FVector3f(1.0f, 1.0f, 1.0f);
+		FQuat4f Rotation = FQuat4f::Identity;
 		FBox LocalBounds = FBox(ForceInitToZero);
 		TArray<uint32, TInlineAllocator<4>> MaterialRemapTable;
 	};
diff --git a/Engine/Plugins/FX/Niagara/Source/NiagaraVertexFactories/Public/NiagaraGPURayTracingTransformsShader.h b/Engine/Plugins/FX/Niagara/Source/NiagaraVertexFactories/Public/NiagaraGPURayTracingTransformsShader.h
--- a/Engine/Plugins/FX/Niagara/Source/NiagaraVertexFactories/Public/NiagaraGPURayTracingTransformsShader.h
+++ b/Engine/Plugins/FX/Niagara/Source/NiagaraVertexFactories/Public/NiagaraGPURayTracingTransformsShader.h
@@ -47,6 +47,7 @@ class NIAGARAVERTEXFACTORIES_API FNiagaraGPURayTracingTransformsCS : public FGlo
 		SHADER_PARAMETER(FVector4f,	DefaultRotation)
 		SHADER_PARAMETER(FVector3f,	DefaultScale)
 		SHADER_PARAMETER(FVector3f,	MeshScale)
+		SHADER_PARAMETER(FVector4f, MeshRotation)
 
 		SHADER_PARAMETER(FMatrix44f,	LocalTransform)
 
diff --git a/Engine/Plugins/FX/Niagara/Source/NiagaraVertexFactories/Public/NiagaraMeshVertexFactory.h b/Engine/Plugins/FX/Niagara/Source/NiagaraVertexFactories/Public/NiagaraMeshVertexFactory.h
--- a/Engine/Plugins/FX/Niagara/Source/NiagaraVertexFactories/Public/NiagaraMeshVertexFactory.h
+++ b/Engine/Plugins/FX/Niagara/Source/NiagaraVertexFactories/Public/NiagaraMeshVertexFactory.h
@@ -44,6 +44,7 @@ BEGIN_SHADER_PARAMETER_STRUCT(FNiagaraMeshCommonParameters, NIAGARAVERTEXFACTORI
 
 	SHADER_PARAMETER(FVector3f, MeshScale)
 	SHADER_PARAMETER(FVector3f, MeshOffset)
+	SHADER_PARAMETER(FVector4f, MeshRotation)
 	SHADER_PARAMETER(int, bMeshOffsetIsWorldSpace)
 
 	SHADER_PARAMETER(uint32, bLockedAxisEnable)
