supress extraneous error

diff --git a/Engine/Plugins/FX/Niagara/Source/Niagara/Private/NiagaraDataInterfaceRigidMeshCollisionQuery.cpp b/Engine/Plugins/FX/Niagara/Source/Niagara/Private/NiagaraDataInterfaceRigidMeshCollisionQuery.cpp
--- a/Engine/Plugins/FX/Niagara/Source/Niagara/Private/NiagaraDataInterfaceRigidMeshCollisionQuery.cpp
+++ b/Engine/Plugins/FX/Niagara/Source/Niagara/Private/NiagaraDataInterfaceRigidMeshCollisionQuery.cpp
@@ -713,7 +713,7 @@ public:
 
 				if (DistanceFieldSceneData == nullptr)
 				{
-					UE_LOG(LogRigidMeshCollision, Error, TEXT("Distance fields are not available for use"));
+					// UE_LOG(LogRigidMeshCollision, Error, TEXT("Distance fields are not available for use"));
 					// #todo(dmp): should we set something here in the case where distance field data is bound but we don't have it?
 					// there is no trivial constructor
 					//FDistanceFieldSceneData DummyDistanceFieldSceneData(Context.Shader->GetShaderPlatform());
