Niagara - Update the niagara stack property row to override the IsExpandedByDefault method to return the information provided by the property row generator.  This fixes an issue where expanded state set in details customizations was ignored when an object was displayed in the niagara selection stack.

diff --git a/Engine/Plugins/FX/Niagara/Source/NiagaraEditor/Private/ViewModels/Stack/NiagaraStackPropertyRow.cpp b/Engine/Plugins/FX/Niagara/Source/NiagaraEditor/Private/ViewModels/Stack/NiagaraStackPropertyRow.cpp
--- a/Engine/Plugins/FX/Niagara/Source/NiagaraEditor/Private/ViewModels/Stack/NiagaraStackPropertyRow.cpp
+++ b/Engine/Plugins/FX/Niagara/Source/NiagaraEditor/Private/ViewModels/Stack/NiagaraStackPropertyRow.cpp
@@ -53,6 +53,11 @@ bool UNiagaraStackPropertyRow::HasOverridenContent() const
 	return false;
 }
 
+bool UNiagaraStackPropertyRow::IsExpandedByDefault() const
+{
+	return DetailTreeNode->GetInitiallyCollapsed() == false;
+}
+
 void UNiagaraStackPropertyRow::FinalizeInternal()
 {
 	Super::FinalizeInternal();
diff --git a/Engine/Plugins/FX/Niagara/Source/NiagaraEditor/Public/ViewModels/Stack/NiagaraStackPropertyRow.h b/Engine/Plugins/FX/Niagara/Source/NiagaraEditor/Public/ViewModels/Stack/NiagaraStackPropertyRow.h
--- a/Engine/Plugins/FX/Niagara/Source/NiagaraEditor/Public/ViewModels/Stack/NiagaraStackPropertyRow.h
+++ b/Engine/Plugins/FX/Niagara/Source/NiagaraEditor/Public/ViewModels/Stack/NiagaraStackPropertyRow.h
@@ -22,6 +22,8 @@ public:
 
 	virtual bool HasOverridenContent() const override;
 
+	virtual bool IsExpandedByDefault() const override;
+
 protected:
 	virtual void FinalizeInternal() override;
 
