expose override for max number of debug lines

diff --git a/Engine/Plugins/FX/Niagara/Source/Niagara/Private/DataInterface/NiagaraDataInterfaceDebugDraw.cpp b/Engine/Plugins/FX/Niagara/Source/Niagara/Private/DataInterface/NiagaraDataInterfaceDebugDraw.cpp
--- a/Engine/Plugins/FX/Niagara/Source/Niagara/Private/DataInterface/NiagaraDataInterfaceDebugDraw.cpp
+++ b/Engine/Plugins/FX/Niagara/Source/Niagara/Private/DataInterface/NiagaraDataInterfaceDebugDraw.cpp
@@ -1412,6 +1412,7 @@ struct FNDIDebugDrawInstanceData_GameThread
 struct FNDIDebugDrawInstanceData_RenderThread
 {
 	FNiagaraGpuComputeDebug* GpuComputeDebug = nullptr;
+	uint32 OverrideMaxLineInstances = 0;
 };
 #endif //NIAGARA_COMPUTEDEBUG_ENABLED
 
@@ -2293,7 +2294,7 @@ public:
 		FNiagaraSimulationDebugDrawData* DebugDraw = nullptr;
 		if ( InstanceData->GpuComputeDebug )
 		{
-			DebugDraw = InstanceData->GpuComputeDebug->GetSimulationDebugDrawData(Context.SystemInstanceID, true);
+			DebugDraw = InstanceData->GpuComputeDebug->GetSimulationDebugDrawData(Context.SystemInstanceID, true, InstanceData->OverrideMaxLineInstances);
 		}
 
 		const bool bIsValid =
@@ -2335,7 +2336,7 @@ public:
 		FNiagaraSimulationDebugDrawData* DebugDraw = nullptr;
 		if (InstanceData->GpuComputeDebug)
 		{
-			DebugDraw = InstanceData->GpuComputeDebug->GetSimulationDebugDrawData(Context.SystemInstanceID, true);
+			DebugDraw = InstanceData->GpuComputeDebug->GetSimulationDebugDrawData(Context.SystemInstanceID, true, InstanceData->OverrideMaxLineInstances);
 		}
 
 		const bool bIsValid =
@@ -3214,7 +3215,7 @@ bool UNiagaraDataInterfaceDebugDraw::PerInstanceTickPostSimulate(void* PerInstan
 
 					if (RT_InstanceData && RT_InstanceData->GpuComputeDebug)
 					{
-						if (FNiagaraSimulationDebugDrawData* DebugDraw = RT_InstanceData->GpuComputeDebug->GetSimulationDebugDrawData(RT_InstanceID, false))
+						if (FNiagaraSimulationDebugDrawData* DebugDraw = RT_InstanceData->GpuComputeDebug->GetSimulationDebugDrawData(RT_InstanceID, false, RT_InstanceData->OverrideMaxLineInstances))
 						{
 							if (DebugDraw->LastUpdateTickCount != RT_TickCount)
 							{
@@ -3247,11 +3248,12 @@ bool UNiagaraDataInterfaceDebugDraw::InitPerInstanceData(void* PerInstanceData,
 
 #if NIAGARA_COMPUTEDEBUG_ENABLED
 	ENQUEUE_RENDER_COMMAND(NDIDebugDrawInit)(
-		[RT_Proxy=GetProxyAs<FNDIDebugDrawProxy>(), RT_InstanceID=SystemInstance->GetId(), RT_ComputeDispatchInterface=SystemInstance->GetComputeDispatchInterface()](FRHICommandListImmediate& RHICmdList)
+		[RT_Proxy=GetProxyAs<FNDIDebugDrawProxy>(), RT_InstanceID=SystemInstance->GetId(), RT_ComputeDispatchInterface=SystemInstance->GetComputeDispatchInterface(), RT_OverrideMaxLineInstances = OverrideMaxLineInstances](FRHICommandListImmediate& RHICmdList)
 		{
 			check(!RT_Proxy->SystemInstancesToProxyData_RT.Contains(RT_InstanceID));
 			FNDIDebugDrawInstanceData_RenderThread* RT_InstanceData = &RT_Proxy->SystemInstancesToProxyData_RT.Add(RT_InstanceID);
 			RT_InstanceData->GpuComputeDebug = RT_ComputeDispatchInterface->GetGpuComputeDebug();
+			RT_InstanceData->OverrideMaxLineInstances = RT_OverrideMaxLineInstances;
 		}
 	);
 #endif
@@ -3279,3 +3281,28 @@ void UNiagaraDataInterfaceDebugDraw::DestroyPerInstanceData(void* PerInstanceDat
 	);
 #endif
 }
+
+bool UNiagaraDataInterfaceDebugDraw::Equals(const UNiagaraDataInterface* Other) const
+{
+	if (!Super::Equals(Other))
+	{
+		return false;
+	}
+	const UNiagaraDataInterfaceDebugDraw* OtherTyped = CastChecked<const UNiagaraDataInterfaceDebugDraw>(Other);
+
+	return OtherTyped->OverrideMaxLineInstances == OverrideMaxLineInstances;
+}
+
+bool UNiagaraDataInterfaceDebugDraw::CopyToInternal(UNiagaraDataInterface* Destination) const
+{
+	if (!Super::CopyToInternal(Destination))
+	{
+		return false;
+	}
+
+	UNiagaraDataInterfaceDebugDraw* OtherTyped = CastChecked<UNiagaraDataInterfaceDebugDraw>(Destination);
+
+	OtherTyped->OverrideMaxLineInstances = OverrideMaxLineInstances;	
+
+	return true;
+}
\ No newline at end of file
diff --git a/Engine/Plugins/FX/Niagara/Source/Niagara/Private/DataInterface/NiagaraDataInterfaceDebugDraw.h b/Engine/Plugins/FX/Niagara/Source/Niagara/Private/DataInterface/NiagaraDataInterfaceDebugDraw.h
--- a/Engine/Plugins/FX/Niagara/Source/Niagara/Private/DataInterface/NiagaraDataInterfaceDebugDraw.h
+++ b/Engine/Plugins/FX/Niagara/Source/Niagara/Private/DataInterface/NiagaraDataInterfaceDebugDraw.h
@@ -13,6 +13,10 @@ class NIAGARA_API UNiagaraDataInterfaceDebugDraw : public UNiagaraDataInterface
 
 public:
 	DECLARE_NIAGARA_DI_PARAMETER();
+		
+	// Number of debug lines is set as the max of OverrideMaxLineInstances and fx.Niagara.GpuComputeDebug.MaxLineInstances
+	UPROPERTY(EditAnywhere, Category = "Debug Draw")
+	uint32 OverrideMaxLineInstances = 0;
 
 	//UObject Interface
 	virtual void PostInitProperties() override;
@@ -62,5 +66,11 @@ public:
 	virtual bool NeedsGPUContextInit() const override { return true; }
 	virtual bool GPUContextInit(const FNiagaraScriptDataInterfaceCompileInfo& InInfo, void* PerInstanceData, FNiagaraSystemInstance* SystemInstance) const override;
 
+	virtual bool Equals(const UNiagaraDataInterface* Other) const override;
+
+protected:
+	/** Copy one niagara DI to this */
+	virtual bool CopyToInternal(UNiagaraDataInterface* Destination) const override;
+
 	//UNiagaraDataInterface Interface
 };
diff --git a/Engine/Plugins/FX/Niagara/Source/Niagara/Private/NiagaraGpuComputeDebug.cpp b/Engine/Plugins/FX/Niagara/Source/Niagara/Private/NiagaraGpuComputeDebug.cpp
--- a/Engine/Plugins/FX/Niagara/Source/Niagara/Private/NiagaraGpuComputeDebug.cpp
+++ b/Engine/Plugins/FX/Niagara/Source/Niagara/Private/NiagaraGpuComputeDebug.cpp
@@ -222,7 +222,7 @@ void FNiagaraGpuComputeDebug::AddAttributeTexture(FRHICommandList& RHICmdList, F
 	}
 }
 
-FNiagaraSimulationDebugDrawData* FNiagaraGpuComputeDebug::GetSimulationDebugDrawData(FNiagaraSystemInstanceID SystemInstanceID, bool bRequiresGpuBuffers)
+FNiagaraSimulationDebugDrawData* FNiagaraGpuComputeDebug::GetSimulationDebugDrawData(FNiagaraSystemInstanceID SystemInstanceID, bool bRequiresGpuBuffers, uint32 OverrideMaxDebugLines)
 {
 	TUniquePtr<FNiagaraSimulationDebugDrawData>& DebugDrawDataPtr = DebugDrawBuffers.FindOrAdd(SystemInstanceID);
 	if (!DebugDrawDataPtr.IsValid())
@@ -230,12 +230,14 @@ FNiagaraSimulationDebugDrawData* FNiagaraGpuComputeDebug::GetSimulationDebugDraw
 		DebugDrawDataPtr.Reset(new FNiagaraSimulationDebugDrawData());
 	}
 
-	if (bRequiresGpuBuffers && (DebugDrawDataPtr->GpuLineMaxInstances != GNiagaraGpuComputeDebug_MaxLineInstances))
+	int MaxLineInstancesToUse = FMath::Max3(DebugDrawDataPtr->GpuLineMaxInstances, (uint32) GNiagaraGpuComputeDebug_MaxLineInstances, OverrideMaxDebugLines);
+
+	if (bRequiresGpuBuffers && (DebugDrawDataPtr->GpuLineMaxInstances != MaxLineInstancesToUse))
 	{
 		check(IsInRenderingThread());
 		DebugDrawDataPtr->GpuLineBufferArgs.Release();
 		DebugDrawDataPtr->GpuLineVertexBuffer.Release();
-		DebugDrawDataPtr->GpuLineMaxInstances = GNiagaraGpuComputeDebug_MaxLineInstances;
+		DebugDrawDataPtr->GpuLineMaxInstances = MaxLineInstancesToUse;
 		if (DebugDrawDataPtr->GpuLineMaxInstances > 0)
 		{
 			DebugDrawDataPtr->GpuLineBufferArgs.Initialize(TEXT("NiagaraGpuComputeDebug::DrawLineBufferArgs"), sizeof(uint32), 4, EPixelFormat::PF_R32_UINT, BUF_Static | BUF_DrawIndirect);
diff --git a/Engine/Plugins/FX/Niagara/Source/Niagara/Private/NiagaraGpuComputeDebug.h b/Engine/Plugins/FX/Niagara/Source/Niagara/Private/NiagaraGpuComputeDebug.h
--- a/Engine/Plugins/FX/Niagara/Source/Niagara/Private/NiagaraGpuComputeDebug.h
+++ b/Engine/Plugins/FX/Niagara/Source/Niagara/Private/NiagaraGpuComputeDebug.h
@@ -73,7 +73,7 @@ public:
 	void AddAttributeTexture(FRHICommandList& RHICmdList, FNiagaraSystemInstanceID SystemInstanceID, FName SourceName, FRHITexture* Texture, FIntVector4 NumTextureAttributes, FIntVector4 AttributeIndices, FVector2D PreviewDisplayRange = FVector2D::ZeroVector);
 
 	// Get Debug draw buffers for a system instance
-	FNiagaraSimulationDebugDrawData* GetSimulationDebugDrawData(FNiagaraSystemInstanceID SystemInstanceID, bool bRequiresGpuBuffers);
+	FNiagaraSimulationDebugDrawData* GetSimulationDebugDrawData(FNiagaraSystemInstanceID SystemInstanceID, bool bRequiresGpuBuffers, uint32 OverrideMaxDebugLines = 0);
 
 	// Force remove debug draw data
 	void RemoveSimulationDebugDrawData(FNiagaraSystemInstanceID SystemInstanceID);
