- Fix re-enterant compilation checking

#rnx

diff --git a/Engine/Plugins/FX/Niagara/Source/Niagara/Classes/NiagaraSystem.h b/Engine/Plugins/FX/Niagara/Source/Niagara/Classes/NiagaraSystem.h
--- a/Engine/Plugins/FX/Niagara/Source/Niagara/Classes/NiagaraSystem.h
+++ b/Engine/Plugins/FX/Niagara/Source/Niagara/Classes/NiagaraSystem.h
@@ -810,11 +810,12 @@ protected:
 	TArray<TObjectPtr<UNiagaraParameterCollectionInstance>> ParameterCollectionOverrides;
 
 #if WITH_EDITORONLY_DATA
+	bool bCompilationReentrantGuard = false;
+
 	UPROPERTY(Transient)
 	TArray<FNiagaraSystemCompileRequest> ActiveCompilations;
 #endif
 
-
 	/** The script which defines the System parameters, and which generates the bindings from System
 		parameter to emitter parameter. */
 	UPROPERTY()
diff --git a/Engine/Plugins/FX/Niagara/Source/Niagara/Private/NiagaraSystem.cpp b/Engine/Plugins/FX/Niagara/Source/Niagara/Private/NiagaraSystem.cpp
--- a/Engine/Plugins/FX/Niagara/Source/Niagara/Private/NiagaraSystem.cpp
+++ b/Engine/Plugins/FX/Niagara/Source/Niagara/Private/NiagaraSystem.cpp
@@ -23,6 +23,7 @@
 #include "Async/Async.h"
 #include "Interfaces/ITargetPlatform.h"
 #include "Misc/ScopedSlowTask.h"
+#include "Misc/ScopeExit.h"
 #include "Modules/ModuleManager.h"
 #include "ProfilingDebugging/CookStats.h"
 #include "UObject/ObjectSaveContext.h"
@@ -35,7 +36,6 @@
 
 #if WITH_EDITOR
 #include "DerivedDataCacheInterface.h"
-static uint32 CompileGuardSlot = 0;
 #endif
 
 DECLARE_CYCLE_STAT(TEXT("Niagara - System - Precompile"), STAT_Niagara_System_Precompile, STATGROUP_Niagara);
@@ -2368,21 +2368,6 @@ bool UNiagaraSystem::PollForCompilationComplete()
 	return true;
 }
 
-bool InternalCompileGuardCheck(void* TestValue)
-{
-	// We need to make sure that we don't re-enter this function on the same thread as it might update things behind our backs.
-	// Am slightly concerened about PostLoad happening on a worker thread, so am not using a generic static variable here, just
-	// a thread local storage variable. The initialized TLS value should be nullptr. When we are doing a compile request, we 
-	// will set the TLS to our this pointer. If the TLS is already this when requesting a compile, we will just early out.
-	if (!CompileGuardSlot)
-	{
-		CompileGuardSlot = FPlatformTLS::AllocTlsSlot();
-	}
-	check(CompileGuardSlot != 0);
-	bool bCompileGuardInProgress = FPlatformTLS::GetTlsValue(CompileGuardSlot) == TestValue;
-	return bCompileGuardInProgress;
-}
-
 bool UNiagaraSystem::CompilationResultsValid(FNiagaraSystemCompileRequest& CompileRequest) const
 {
 	// for now the only thing we're concerned about is if we've got results for SystemSpawn and SystemUpdate scripts
@@ -2702,10 +2687,12 @@ void UNiagaraSystem::PreProcessWaitingDDCTasks(bool bProcessForWait)
 
 bool UNiagaraSystem::QueryCompileComplete(bool bWait, bool bDoPost, bool bDoNotApply)
 {
-	bool bCompileGuardInProgress = InternalCompileGuardCheck(this);
-
-	if (ActiveCompilations.Num() > 0 && !bCompileGuardInProgress)
+	check(IsInGameThread());
+	if (ActiveCompilations.Num() > 0 && !bCompilationReentrantGuard)
 	{
+		bCompilationReentrantGuard = true;
+		ON_SCOPE_EXIT { bCompilationReentrantGuard = false; };
+
 		PreProcessWaitingDDCTasks(bWait);
 		
 		FNiagaraSystemCompileRequest& CompileRequest = ActiveCompilations[0];
@@ -3024,6 +3011,7 @@ void UNiagaraSystem::PrepareRapidIterationParametersForCompilation()
 
 bool UNiagaraSystem::RequestCompile(bool bForce, FNiagaraSystemUpdateContext* OptionalUpdateContext)
 {
+	check(IsInGameThread());
 	TRACE_CPUPROFILER_EVENT_SCOPE(UNiagaraSystem::RequestCompile)
 
 	// We remove emitters and scripts on dedicated servers, so skip further work.
@@ -3038,14 +3026,12 @@ bool UNiagaraSystem::RequestCompile(bool bForce, FNiagaraSystemUpdateContext* Op
 		return false;
 	}
 
-	bool bCompileGuardInProgress = InternalCompileGuardCheck(this);
-
 	if (bForce)
 	{
 		ForceGraphToRecompileOnNextCheck();
 	}
 
-	if (bCompileGuardInProgress)
+	if (bCompilationReentrantGuard)
 	{
 		return false;
 	}
@@ -3063,7 +3049,8 @@ bool UNiagaraSystem::RequestCompile(bool bForce, FNiagaraSystemUpdateContext* Op
 
 	// Record that we entered this function already.
 	SCOPE_CYCLE_COUNTER(STAT_Niagara_System_CompileScript);
-	FPlatformTLS::SetTlsValue(CompileGuardSlot, this);
+	bCompilationReentrantGuard = true;
+	ON_SCOPE_EXIT { check(bCompilationReentrantGuard == false); };
 
 	FNiagaraSystemCompileRequest& ActiveCompilation = ActiveCompilations.AddDefaulted_GetRef();
 	ActiveCompilation.bForced = bForce;
@@ -3208,10 +3195,8 @@ bool UNiagaraSystem::RequestCompile(bool bForce, FNiagaraSystemUpdateContext* Op
 		}
 	}
 
-
 	// Now record that we are done with this function.
-	FPlatformTLS::SetTlsValue(CompileGuardSlot, nullptr);
-
+	bCompilationReentrantGuard = false;
 
 	// We might be able to just complete compilation right now if nothing needed compilation.
 	if (ScriptsNeedingCompile.Num() == 0 && ActiveCompilations.Num() == 1)
