Overview stack row inline module data
- Parameter meta data "Display in Overview"
- For local inputs (inputs & staticswitch inputs) only
- Creates a button per valid parameter that reads the button's current value (makes more sense with enums than with floats, unless the module is very explicit and low complexity)
- For enums: an entry can have an overriden name or an icon. Icon takes priority.
- Default border color for text-based entries is based on type. Can be overridden.
- button's tooltip shows the actual parameter behind the value
- clicking the button navigates the stack to the input
- Many layout tweaks and fixes (like debug button's minimum size being too large and offset, using negative padding counters it)

diff --git a/Engine/Plugins/FX/Niagara/Source/NiagaraEditor/Private/Widgets/SNiagaraPinTypeSelector.h b/Engine/Plugins/FX/Niagara/Source/NiagaraEditor/Private/Widgets/SNiagaraPinTypeSelector.h
deleted file mode 100644
diff --git a/Engine/Plugins/FX/Niagara/Source/NiagaraEditor/Public/Widgets/SNiagaraPinTypeSelector.h b/Engine/Plugins/FX/Niagara/Source/NiagaraEditor/Public/Widgets/SNiagaraPinTypeSelector.h
new file mode 100644
--- /dev/null
+++ b/Engine/Plugins/FX/Niagara/Source/NiagaraEditor/Public/Widgets/SNiagaraPinTypeSelector.h
@@ -0,0 +1,66 @@
+// Copyright Epic Games, Inc. All Rights Reserved.
+
+#pragma once
+
+#include "CoreMinimal.h"
+#include "EdGraph/EdGraphPin.h"
+#include "Input/Reply.h"
+#include "Misc/Attribute.h"
+#include "NiagaraTypes.h"
+#include "SGraphPalette.h"
+#include "SlateFwd.h"
+#include "Styling/SlateColor.h"
+#include "Widgets/DeclarativeSyntaxSupport.h"
+#include "Widgets/Input/SComboButton.h"
+#include "Widgets/SCompoundWidget.h"
+#include "Widgets/SWidget.h"
+
+class NIAGARAEDITOR_API SNiagaraPinTypeSelector : public SCompoundWidget
+{
+public:
+	SLATE_BEGIN_ARGS(SNiagaraPinTypeSelector)
+	{}
+	SLATE_END_ARGS()
+public:
+	void Construct(const FArguments& InArgs, UEdGraphPin * InGraphPin);
+
+protected:
+	virtual TSharedRef<SWidget>	GetMenuContent();
+	FText GetTooltipText() const;
+
+private:
+	UEdGraphPin* Pin = nullptr;
+	TSharedPtr<SComboButton> SelectorButton;
+};
+
+class NIAGARAEDITOR_API SNiagaraIconWidget : public SGraphPaletteItem
+{
+public:
+
+	SLATE_BEGIN_ARGS(SNiagaraIconWidget)
+	{}
+		SLATE_ARGUMENT(FText, IconToolTip)
+		SLATE_ARGUMENT(const FSlateBrush*, IconBrush)
+		SLATE_ARGUMENT(FSlateColor, IconColor)
+		SLATE_ARGUMENT(FString, DocLink)
+		SLATE_ARGUMENT(FString, DocExcerpt)
+		SLATE_ARGUMENT(const FSlateBrush*, SecondaryIconBrush)
+		SLATE_ARGUMENT(FSlateColor, SecondaryIconColor)
+	SLATE_END_ARGS();
+
+	void Construct(const FArguments& InArgs)
+	{
+		ChildSlot
+		[
+			CreateIconWidget(
+				InArgs._IconToolTip
+				, InArgs._IconBrush
+				, InArgs._IconColor
+				, InArgs._DocLink
+				, InArgs._DocExcerpt
+				, InArgs._SecondaryIconBrush
+				, InArgs._SecondaryIconColor
+			)
+		];
+	}
+};
\ No newline at end of file
