Allow sequencer to disable scalability on Niagara Components.
Should allow us to better reuse game FX in events/cinematics without having sequencer and scalability fight each other.
There is more that could be done to that may allow these to coexist better but for now keeping them separate when needed should work.

diff --git a/Engine/Plugins/FX/Niagara/Source/Niagara/Private/MovieScene/MovieSceneNiagaraSystemSpawnSection.cpp b/Engine/Plugins/FX/Niagara/Source/Niagara/Private/MovieScene/MovieSceneNiagaraSystemSpawnSection.cpp
--- a/Engine/Plugins/FX/Niagara/Source/Niagara/Private/MovieScene/MovieSceneNiagaraSystemSpawnSection.cpp
+++ b/Engine/Plugins/FX/Niagara/Source/Niagara/Private/MovieScene/MovieSceneNiagaraSystemSpawnSection.cpp
@@ -9,6 +9,7 @@ UMovieSceneNiagaraSystemSpawnSection::UMovieSceneNiagaraSystemSpawnSection()
 	SectionEvaluateBehavior = ENiagaraSystemSpawnSectionEvaluateBehavior::ActivateIfInactive;
 	SectionEndBehavior = ENiagaraSystemSpawnSectionEndBehavior::SetSystemInactive;
 	AgeUpdateMode = ENiagaraAgeUpdateMode::TickDeltaTime;
+	bAllowScalability = false;
 }
 
 ENiagaraSystemSpawnSectionStartBehavior UMovieSceneNiagaraSystemSpawnSection::GetSectionStartBehavior() const
@@ -31,4 +32,7 @@ ENiagaraAgeUpdateMode UMovieSceneNiagaraSystemSpawnSection::GetAgeUpdateMode() c
 	return AgeUpdateMode;
 }
 
-
+bool UMovieSceneNiagaraSystemSpawnSection::GetAllowScalability()const
+{
+	return bAllowScalability;
+}
\ No newline at end of file
diff --git a/Engine/Plugins/FX/Niagara/Source/Niagara/Private/MovieScene/MovieSceneNiagaraSystemTrack.cpp b/Engine/Plugins/FX/Niagara/Source/Niagara/Private/MovieScene/MovieSceneNiagaraSystemTrack.cpp
--- a/Engine/Plugins/FX/Niagara/Source/Niagara/Private/MovieScene/MovieSceneNiagaraSystemTrack.cpp
+++ b/Engine/Plugins/FX/Niagara/Source/Niagara/Private/MovieScene/MovieSceneNiagaraSystemTrack.cpp
@@ -37,7 +37,7 @@ void UMovieSceneNiagaraSystemTrack::PostCompile(FMovieSceneEvaluationTrack& OutT
 		OutTrack.SetTrackImplementation(FMovieSceneNiagaraSystemTrackImplementation(
 			SpawnSection->GetInclusiveStartFrame(), SpawnSection->GetExclusiveEndFrame(),
 			SpawnSection->GetSectionStartBehavior(), SpawnSection->GetSectionEvaluateBehavior(),
-			SpawnSection->GetSectionEndBehavior(), SpawnSection->GetAgeUpdateMode()));
+			SpawnSection->GetSectionEndBehavior(), SpawnSection->GetAgeUpdateMode(), SpawnSection->GetAllowScalability()));
 	}
 }
 
diff --git a/Engine/Plugins/FX/Niagara/Source/Niagara/Private/MovieScene/MovieSceneNiagaraSystemTrackTemplate.cpp b/Engine/Plugins/FX/Niagara/Source/Niagara/Private/MovieScene/MovieSceneNiagaraSystemTrackTemplate.cpp
--- a/Engine/Plugins/FX/Niagara/Source/Niagara/Private/MovieScene/MovieSceneNiagaraSystemTrackTemplate.cpp
+++ b/Engine/Plugins/FX/Niagara/Source/Niagara/Private/MovieScene/MovieSceneNiagaraSystemTrackTemplate.cpp
@@ -17,6 +17,7 @@ struct FPreAnimatedNiagaraComponentToken : IMovieScenePreAnimatedToken
 		bool bInComponentRenderingEnabled,
 		TOptional<ENiagaraExecutionState> InSystemInstanceExecutionState,
 		ENiagaraAgeUpdateMode InComponentAgeUpdateMode,
+		bool bInComponentAllowScalability,
 		float InComponentSeekDelta,
 		float InComponentDesiredAge,
 		bool bInComponentLockDesiredAgeDeltaTimeToSeekDelta
@@ -26,6 +27,7 @@ struct FPreAnimatedNiagaraComponentToken : IMovieScenePreAnimatedToken
 		, bComponentRenderingEnabled(bInComponentRenderingEnabled)
 		, SystemInstanceExecutionState(InSystemInstanceExecutionState)
 		, ComponentAgeUpdateMode(InComponentAgeUpdateMode)
+		, bComponentAllowScalability(bInComponentAllowScalability)
 		, ComponentSeekDelta(InComponentSeekDelta)
 		, ComponentDesiredAge(InComponentDesiredAge)
 		, bComponentLockDesiredAgeDeltaTimeToSeekDelta(bInComponentLockDesiredAgeDeltaTimeToSeekDelta)
@@ -50,8 +52,10 @@ struct FPreAnimatedNiagaraComponentToken : IMovieScenePreAnimatedToken
 		NiagaraComponent->SetForceSolo(bComponentForceSolo);
 		NiagaraComponent->SetRenderingEnabled(bComponentRenderingEnabled);
 		NiagaraComponent->SetAgeUpdateMode(ComponentAgeUpdateMode);
+		NiagaraComponent->SetAllowScalability(bComponentAllowScalability);
 		NiagaraComponent->SetSeekDelta(ComponentSeekDelta);
 		NiagaraComponent->SetDesiredAge(ComponentDesiredAge);
+		NiagaraComponent->SetAllowScalability(bComponentAllowScalability);
 		NiagaraComponent->SetLockDesiredAgeDeltaTimeToSeekDelta(bComponentLockDesiredAgeDeltaTimeToSeekDelta);
 
 		// TODO: When this action is ACTUALLY deferred, just expose it to the component		
@@ -70,6 +74,7 @@ struct FPreAnimatedNiagaraComponentToken : IMovieScenePreAnimatedToken
 	bool bComponentRenderingEnabled;
 	TOptional<ENiagaraExecutionState> SystemInstanceExecutionState;
 	ENiagaraAgeUpdateMode ComponentAgeUpdateMode;
+	bool bComponentAllowScalability;
 	float ComponentSeekDelta;
 	float ComponentDesiredAge;
 	bool bComponentLockDesiredAgeDeltaTimeToSeekDelta;
@@ -88,6 +93,7 @@ struct FPreAnimatedNiagaraComponentTokenProducer : IMovieScenePreAnimatedTokenPr
 			NiagaraComponent->GetRenderingEnabled(),
 			SystemInstanceController.IsValid() ? SystemInstanceController->GetRequestedExecutionState() : TOptional<ENiagaraExecutionState>(),
 			NiagaraComponent->GetAgeUpdateMode(),
+			NiagaraComponent->GetAllowScalability(),
 			NiagaraComponent->GetSeekDelta(),
 			NiagaraComponent->GetDesiredAge(),
 			NiagaraComponent->GetLockDesiredAgeDeltaTimeToSeekDelta());
@@ -99,13 +105,14 @@ struct FNiagaraSystemUpdateDesiredAgeExecutionToken : IMovieSceneExecutionToken
 	FNiagaraSystemUpdateDesiredAgeExecutionToken(
 		FFrameNumber InSpawnSectionStartFrame, FFrameNumber InSpawnSectionEndFrame,
 		ENiagaraSystemSpawnSectionStartBehavior InSpawnSectionStartBehavior, ENiagaraSystemSpawnSectionEvaluateBehavior InSpawnSectionEvaluateBehavior,
-		ENiagaraSystemSpawnSectionEndBehavior InSpawnSectionEndBehavior, ENiagaraAgeUpdateMode InAgeUpdateMode)
+		ENiagaraSystemSpawnSectionEndBehavior InSpawnSectionEndBehavior, ENiagaraAgeUpdateMode InAgeUpdateMode, bool bInAllowScalability)
 		: SpawnSectionStartFrame(InSpawnSectionStartFrame)
 		, SpawnSectionEndFrame(InSpawnSectionEndFrame)
 		, SpawnSectionStartBehavior(InSpawnSectionStartBehavior)
 		, SpawnSectionEvaluateBehavior(InSpawnSectionEvaluateBehavior)
 		, SpawnSectionEndBehavior(InSpawnSectionEndBehavior)
 		, AgeUpdateMode(InAgeUpdateMode)
+		, bAllowScalability(bInAllowScalability)
 	{}
 
 	virtual void Execute(const FMovieSceneContext& Context, const FMovieSceneEvaluationOperand& Operand, FPersistentEvaluationData& PersistentData, IMovieScenePlayer& Player) override
@@ -124,6 +131,7 @@ struct FNiagaraSystemUpdateDesiredAgeExecutionToken : IMovieSceneExecutionToken
 
 			NiagaraComponent->SetForceSolo(true);
 			NiagaraComponent->SetAgeUpdateMode(AgeUpdateMode);
+			NiagaraComponent->SetAllowScalability(bAllowScalability);
 
 			UMovieSceneSequence* MovieSceneSequence = Player.GetEvaluationTemplate().GetSequence(Operand.SequenceID);
 			if (MovieSceneSequence != nullptr)
@@ -235,18 +243,20 @@ struct FNiagaraSystemUpdateDesiredAgeExecutionToken : IMovieSceneExecutionToken
 	ENiagaraSystemSpawnSectionEvaluateBehavior SpawnSectionEvaluateBehavior;
 	ENiagaraSystemSpawnSectionEndBehavior SpawnSectionEndBehavior;
 	ENiagaraAgeUpdateMode AgeUpdateMode;
+	bool bAllowScalability;
 };
 
 FMovieSceneNiagaraSystemTrackImplementation::FMovieSceneNiagaraSystemTrackImplementation(
 	FFrameNumber InSpawnSectionStartFrame, FFrameNumber InSpawnSectionEndFrame,
 	ENiagaraSystemSpawnSectionStartBehavior InSpawnSectionStartBehavior, ENiagaraSystemSpawnSectionEvaluateBehavior InSpawnSectionEvaluateBehavior,
-	ENiagaraSystemSpawnSectionEndBehavior InSpawnSectionEndBehavior, ENiagaraAgeUpdateMode InAgeUpdateMode)
+	ENiagaraSystemSpawnSectionEndBehavior InSpawnSectionEndBehavior, ENiagaraAgeUpdateMode InAgeUpdateMode, bool bInAllowScalability)
 	: SpawnSectionStartFrame(InSpawnSectionStartFrame)
 	, SpawnSectionEndFrame(InSpawnSectionEndFrame)
 	, SpawnSectionStartBehavior(InSpawnSectionStartBehavior)
 	, SpawnSectionEvaluateBehavior(InSpawnSectionEvaluateBehavior)
 	, SpawnSectionEndBehavior(InSpawnSectionEndBehavior)
 	, AgeUpdateMode(InAgeUpdateMode)
+	, bAllowScalability(bInAllowScalability)
 
 {
 }
@@ -258,6 +268,7 @@ FMovieSceneNiagaraSystemTrackImplementation::FMovieSceneNiagaraSystemTrackImplem
 	, SpawnSectionEvaluateBehavior(ENiagaraSystemSpawnSectionEvaluateBehavior::None)
 	, SpawnSectionEndBehavior(ENiagaraSystemSpawnSectionEndBehavior::SetSystemInactive)
 	, AgeUpdateMode(ENiagaraAgeUpdateMode::TickDeltaTime)
+	, bAllowScalability(false)
 {
 }
 
@@ -267,5 +278,5 @@ void FMovieSceneNiagaraSystemTrackImplementation::Evaluate(const FMovieSceneEval
 	ExecutionTokens.Add(FNiagaraSystemUpdateDesiredAgeExecutionToken(
 		SpawnSectionStartFrame, SpawnSectionEndFrame,
 		SpawnSectionStartBehavior, SpawnSectionEvaluateBehavior,
-		SpawnSectionEndBehavior, AgeUpdateMode));
+		SpawnSectionEndBehavior, AgeUpdateMode, bAllowScalability));
 }
\ No newline at end of file
diff --git a/Engine/Plugins/FX/Niagara/Source/Niagara/Private/MovieScene/MovieSceneNiagaraSystemTrackTemplate.h b/Engine/Plugins/FX/Niagara/Source/Niagara/Private/MovieScene/MovieSceneNiagaraSystemTrackTemplate.h
--- a/Engine/Plugins/FX/Niagara/Source/Niagara/Private/MovieScene/MovieSceneNiagaraSystemTrackTemplate.h
+++ b/Engine/Plugins/FX/Niagara/Source/Niagara/Private/MovieScene/MovieSceneNiagaraSystemTrackTemplate.h
@@ -30,7 +30,7 @@ public:
 	FMovieSceneNiagaraSystemTrackImplementation(
 		FFrameNumber InSpawnSectionStartFrame, FFrameNumber InSpawnSectionEndFrame,
 		ENiagaraSystemSpawnSectionStartBehavior InSectionStartBehavior, ENiagaraSystemSpawnSectionEvaluateBehavior InSectionEvaluateBehavior,
-		ENiagaraSystemSpawnSectionEndBehavior InSectionEndBehavior, ENiagaraAgeUpdateMode InAgeUpdateMode);
+		ENiagaraSystemSpawnSectionEndBehavior InSectionEndBehavior, ENiagaraAgeUpdateMode InAgeUpdateMode, bool bInAllowScalability);
 
 	virtual void SetupOverrides() override
 	{
@@ -54,4 +54,7 @@ private:
 	ENiagaraSystemSpawnSectionEndBehavior SpawnSectionEndBehavior;
 	UPROPERTY()
 	ENiagaraAgeUpdateMode AgeUpdateMode;
+	UPROPERTY()
+	bool bAllowScalability;
+
 };
\ No newline at end of file
diff --git a/Engine/Plugins/FX/Niagara/Source/Niagara/Private/NiagaraComponent.cpp b/Engine/Plugins/FX/Niagara/Source/Niagara/Private/NiagaraComponent.cpp
--- a/Engine/Plugins/FX/Niagara/Source/Niagara/Private/NiagaraComponent.cpp
+++ b/Engine/Plugins/FX/Niagara/Source/Niagara/Private/NiagaraComponent.cpp
@@ -3343,7 +3343,11 @@ void UNiagaraComponent::SetPreviewLODDistance(bool bInEnablePreviewLODDistance,
 void UNiagaraComponent::SetAllowScalability(bool bAllow)
 {
 	bAllowScalability = bAllow; 
-	if (!bAllow)
+	if (bAllow)
+	{
+		RegisterWithScalabilityManager();
+	}
+	else
 	{
 		UnregisterWithScalabilityManager();
 	}
diff --git a/Engine/Plugins/FX/Niagara/Source/Niagara/Public/MovieScene/MovieSceneNiagaraSystemSpawnSection.h b/Engine/Plugins/FX/Niagara/Source/Niagara/Public/MovieScene/MovieSceneNiagaraSystemSpawnSection.h
--- a/Engine/Plugins/FX/Niagara/Source/Niagara/Public/MovieScene/MovieSceneNiagaraSystemSpawnSection.h
+++ b/Engine/Plugins/FX/Niagara/Source/Niagara/Public/MovieScene/MovieSceneNiagaraSystemSpawnSection.h
@@ -53,6 +53,8 @@ public:
 
 	ENiagaraAgeUpdateMode GetAgeUpdateMode() const;
 
+	bool GetAllowScalability()const;
+
 private:
 	/** Specifies what should happen to the niagara system from before the section evaluates up until the first frame of the section. */
 	UPROPERTY(EditAnywhere, Category = "Life Cycle")
@@ -69,4 +71,7 @@ private:
 	/** Specifies how sequencer should update the age of the controlled niagara system. */
 	UPROPERTY(EditAnywhere, Category = "Life Cycle")
 	ENiagaraAgeUpdateMode AgeUpdateMode;
+	
+	UPROPERTY(EditAnywhere, Category = "Life Cycle")
+	bool bAllowScalability;
 };
\ No newline at end of file
