- Fix for false NaN detection

#rnx

diff --git a/Engine/Plugins/FX/Niagara/Source/Niagara/Private/NiagaraDataInterfaceSkeletalMesh.cpp b/Engine/Plugins/FX/Niagara/Source/Niagara/Private/NiagaraDataInterfaceSkeletalMesh.cpp
--- a/Engine/Plugins/FX/Niagara/Source/Niagara/Private/NiagaraDataInterfaceSkeletalMesh.cpp
+++ b/Engine/Plugins/FX/Niagara/Source/Niagara/Private/NiagaraDataInterfaceSkeletalMesh.cpp
@@ -1109,7 +1109,7 @@ void FSkeletalMeshGpuDynamicBufferProxy::NewFrame(const FNDISkeletalMesh_Instanc
 				const int32 ParentIndex = ReferenceSkeleton ? ReferenceSkeleton->GetParentIndex(i) : -1;
 				BoneSamplingData.Emplace(BoneTransform.GetLocation());
 				BoneSamplingData.Emplace(Rotation.X, Rotation.Y, Rotation.Z, Rotation.W);
-				BoneSamplingData.Emplace(BoneTransform.GetScale3D(), reinterpret_cast<const float&>(ParentIndex));
+				BoneSamplingData.Emplace_GetRef(BoneTransform.GetScale3D()).W = reinterpret_cast<const float&>(ParentIndex);
 			}
 
 			// Append sockets
@@ -1119,7 +1119,7 @@ void FSkeletalMeshGpuDynamicBufferProxy::NewFrame(const FNDISkeletalMesh_Instanc
 				const int32 ParentIndex = -1;
 				BoneSamplingData.Emplace(SocketTransform.GetLocation());
 				BoneSamplingData.Emplace(Rotation.X, Rotation.Y, Rotation.Z, Rotation.W);
-				BoneSamplingData.Emplace(SocketTransform.GetScale3D(), reinterpret_cast<const float&>(ParentIndex));
+				BoneSamplingData.Emplace_GetRef(SocketTransform.GetScale3D()).W = reinterpret_cast<const float&>(ParentIndex);
 			}
 		};
 
