Removed highlights feature

diff --git a/Engine/Plugins/FX/Niagara/Source/Niagara/Classes/NiagaraScript.h b/Engine/Plugins/FX/Niagara/Source/Niagara/Classes/NiagaraScript.h
--- a/Engine/Plugins/FX/Niagara/Source/Niagara/Classes/NiagaraScript.h
+++ b/Engine/Plugins/FX/Niagara/Source/Niagara/Classes/NiagaraScript.h
@@ -594,9 +594,6 @@ public:
 	UPROPERTY(EditAnywhere, Category = Script)
 	uint32 bCanBeUsedForTypeConversions : 1;
 
-	UPROPERTY(EditAnywhere, Category = Script)
-	TArray<FNiagaraScriptHighlight> Highlights;
-
 	UPROPERTY(EditAnywhere, Category = Script, DisplayName = "Script Metadata", meta = (ToolTip = "Script Metadata"))
 	TMap<FName, FString> ScriptMetaData;
 	
@@ -796,10 +793,6 @@ public:
 	UPROPERTY(meta = (DeprecatedProperty))
 	FText CollapsedViewFormat_DEPRECATED;
 
-	/* Deprecated, use LibraryVisibility instead. */
-	UPROPERTY(meta = (DeprecatedProperty))
-	TArray<FNiagaraScriptHighlight> Highlights_DEPRECATED;
-
 	/* Deprecated, use LibraryVisibility instead. */
 	UPROPERTY(meta = (DeprecatedProperty))
 	TMap<FName, FString> ScriptMetaData_DEPRECATED;
diff --git a/Engine/Plugins/FX/Niagara/Source/Niagara/Private/NiagaraScript.cpp b/Engine/Plugins/FX/Niagara/Source/Niagara/Private/NiagaraScript.cpp
--- a/Engine/Plugins/FX/Niagara/Source/Niagara/Private/NiagaraScript.cpp
+++ b/Engine/Plugins/FX/Niagara/Source/Niagara/Private/NiagaraScript.cpp
@@ -480,7 +480,6 @@ void UNiagaraScript::CheckVersionDataAvailable()
 	Data.Source = Source_DEPRECATED;
 	Data.Keywords = Keywords_DEPRECATED;
 	Data.Category = Category_DEPRECATED;
-	Data.Highlights = Highlights_DEPRECATED;
 	Data.Description = Description_DEPRECATED;
 	Data.bDeprecated = bDeprecated_DEPRECATED;
 	Data.NoteMessage = NoteMessage_DEPRECATED;
@@ -2559,21 +2558,6 @@ void UNiagaraScript::GetAssetRegistryTags(TArray<FAssetRegistryTag>& OutTags) co
 		OutTags.Add(FAssetRegistryTag(ProvidedDependenciesName, *ProvidedDependenciesTags, FAssetRegistryTag::TT_Hidden));
 	}
 
-	// Highlights
-	TArray<FNiagaraScriptHighlight> Highlights = ScriptData ? ScriptData->Highlights : Highlights_DEPRECATED; 
-	if (Highlights.Num() > 0)
-	{
-		FName HighlightsName = GET_MEMBER_NAME_CHECKED(FVersionedNiagaraScriptData, Highlights);
-		FString* HighlightsTags = CustomAssetRegistryTagCache->Find(HighlightsName);
-		if (HighlightsTags == nullptr)
-		{
-			HighlightsTags = &CustomAssetRegistryTagCache->Add(HighlightsName);
-			FNiagaraScriptHighlight::ArrayToJson(Highlights, *HighlightsTags);
-		}
-
-		OutTags.Add(FAssetRegistryTag(HighlightsName, *HighlightsTags, FAssetRegistryTag::TT_Hidden));
-	}
-
 	// Category
 	FText CategoryText = ScriptData ? ScriptData->Category : Category_DEPRECATED;
 	if (!CategoryText.IsEmpty())
diff --git a/Engine/Plugins/FX/Niagara/Source/NiagaraEditor/Private/Customizations/NiagaraScriptDetails.cpp b/Engine/Plugins/FX/Niagara/Source/NiagaraEditor/Private/Customizations/NiagaraScriptDetails.cpp
--- a/Engine/Plugins/FX/Niagara/Source/NiagaraEditor/Private/Customizations/NiagaraScriptDetails.cpp
+++ b/Engine/Plugins/FX/Niagara/Source/NiagaraEditor/Private/Customizations/NiagaraScriptDetails.cpp
@@ -275,610 +275,4 @@ class SNiagaraDetailsButton : public SButton
 	}
 };
 
-class SNiagaraScriptHighlightPicker : public SCompoundWidget
-{
-private:
-	enum class EEditMode
-	{
-		Picker,
-		EditCurrent,
-		EditNew
-	};
-
-public:
-	SLATE_BEGIN_ARGS(SNiagaraScriptHighlightPicker)
-	{}
-	SLATE_END_ARGS()
-
-	void Construct(const FArguments& InArgs, TSharedRef<IPropertyHandle> InPropertyHandle)
-	{
-		PropertyHandle = InPropertyHandle;
-		PropertyHandle->SetOnPropertyValueChanged(FSimpleDelegate::CreateSP(this, &SNiagaraScriptHighlightPicker::RefreshDisplayedValue));
-		EditMode = EEditMode::Picker;
-		RefreshDisplayedValue();
-		ChildSlot
-		[
-			SNew(SOverlay)
-			// Highlight picker
-			+ SOverlay::Slot()
-			[
-				ConstructPickerWidgets()
-			]
-			// Highlight Editor
-			+ SOverlay::Slot()
-			[
-				ConstructEditWidgets()
-			]
-		];
-	}
-
-private:
-
-	TSharedRef<SWidget> ConstructPickerWidgets()
-	{
-		return SNew(SHorizontalBox)
-			// Hightlight picker combo
-			+ SHorizontalBox::Slot()
-			[
-				SAssignNew(HighlightPickerCombo, SComboButton)
-				.Visibility(this, &SNiagaraScriptHighlightPicker::GetPickerVisibility)
-				.OnGetMenuContent(this, &SNiagaraScriptHighlightPicker::OnGetMenuContent)
-				.ContentPadding(FMargin(4, 2, 4, 2))
-				.ButtonContent()
-				[
-					SNew(SHorizontalBox)
-					+ SHorizontalBox::Slot()
-					.AutoWidth()
-					.Padding(0, 0, 5, 0)
-					[
-						SNew(SBorder)
-						.BorderImage(FEditorStyle::GetBrush("WhiteBrush"))
-						.BorderBackgroundColor(this, &SNiagaraScriptHighlightPicker::GetHighlightSlateColor)
-						.Padding(0)
-						[
-							SNew(SBox)
-							.WidthOverride(DisplayColorWidth)
-							.HeightOverride(DisplayColorHeight)
-						]
-					]
-					+ SHorizontalBox::Slot()
-					.VAlign(VAlign_Center)
-					[
-						SNew(STextBlock)
-						.TextStyle(FNiagaraEditorStyle::Get(), "NiagaraEditor.ParameterText")
-						.Text(this, &SNiagaraScriptHighlightPicker::GetHighlightText)
-					]
-				]
-			]
-			+ SHorizontalBox::Slot()
-			.AutoWidth()
-			.VAlign(VAlign_Center)
-			[
-				PropertyHandle->CreateDefaultPropertyButtonWidgets()
-			];
-	}
-
-	TSharedRef<SWidget> ConstructEditWidgets()
-	{
-		return SNew(SBorder)
-			.BorderImage(FEditorStyle::GetBrush("ToolPanel.DarkGroupBorder"))
-			.Visibility(this, &SNiagaraScriptHighlightPicker::GetEditVisibility)
-			.Padding(10)
-			[
-				SNew(SVerticalBox)
-
-				// Edit Header
-				+ SVerticalBox::Slot()
-				.AutoHeight()
-				[
-					SNew(STextBlock)
-					.TextStyle(FNiagaraEditorStyle::Get(), "NiagaraEditor.DetailsHeadingText")
-					.Text(this, &SNiagaraScriptHighlightPicker::GetEditModeText)
-				]
-				// Edit controls.
-				+ SVerticalBox::Slot()
-				.AutoHeight()
-				.Padding(0, 10, 0, 0)
-				[
-					ConstructEditHighlightWidgets()
-				]
-				// Error display.
-				+ SVerticalBox::Slot()
-				.AutoHeight()
-				.Padding(0, 5, 0, 0)
-				[
-					ConstructEditErrorWidgets()
-				]
-				// Edit buttons
-				+ SVerticalBox::Slot()
-				.AutoHeight()
-				.HAlign(HAlign_Right)
-				.Padding(0, 10, 0, 0)
-				[
-					ConstructEditButtons()
-				]
-			];
-	}
-
-	TSharedRef<SWidget> ConstructEditHighlightWidgets()
-	{
-		return SNew(SHorizontalBox)
-			+ SHorizontalBox::Slot()
-			.AutoWidth()
-			.Padding(0, 0, 5, 0)
-			[
-				SAssignNew(ColorBlock, SColorBlock)
-				.Color(this, &SNiagaraScriptHighlightPicker::GetHighlightColor)
-				.ShowBackgroundForAlpha(true)
-				.OnMouseButtonDown(this, &SNiagaraScriptHighlightPicker::OnMouseButtonDownColorBlock)
-				.Size(FVector2D(EditColorWidth, EditColorHeight))
-			]
-			+ SHorizontalBox::Slot()
-			[
-				SNew(SEditableTextBox)
-				.Text(this, &SNiagaraScriptHighlightPicker::GetHighlightText)
-				.OnTextCommitted(this, &SNiagaraScriptHighlightPicker::HighlightTextCommitted)
-			];
-	};
-
-	TSharedRef<SWidget> ConstructEditButtons()
-	{
-		return SNew(SHorizontalBox)
-			// Add new button.
-			+ SHorizontalBox::Slot()
-			[
-				SNew(SNiagaraDetailsButton)
-				.Visibility(this, &SNiagaraScriptHighlightPicker::GetAddNewButtonVisibility)
-				.Text(LOCTEXT("AddNewHighlight", "Add"))
-				.OnClicked(this, &SNiagaraScriptHighlightPicker::AddNewRequested)
-			]
-			// Apply button
-			+ SHorizontalBox::Slot()
-			.Padding(5, 0, 0, 0)
-			[
-				SNew(SNiagaraDetailsButton)
-				.Visibility(this, &SNiagaraScriptHighlightPicker::GetApplyButtonVisibility)
-				.Text(LOCTEXT("ApplyToCurrent", "Apply"))
-				.ToolTipText(LOCTEXT("ApplyToCurrentToolTip", "Applies the current highlight changes to this script."))
-				.OnClicked(this, &SNiagaraScriptHighlightPicker::ApplyCurrentRequested)
-			]
-			// Apply all button
-			+ SHorizontalBox::Slot()
-			.Padding(5, 0, 0, 0)
-			[
-				SNew(SNiagaraDetailsButton)
-				.Visibility(this, &SNiagaraScriptHighlightPicker::GetApplyButtonVisibility)
-				.Text(LOCTEXT("ApplyToAll", "Apply All"))
-				.ToolTipText(LOCTEXT("ApplyToAllCurrentToolTip", "Applies the current highlight changes to this script,\nand any other scripts that use this same highlight."))
-				.OnClicked(this, &SNiagaraScriptHighlightPicker::ApplyAllRequested)
-			]
-			// Cancel button
-			+ SHorizontalBox::Slot()
-			.Padding(5, 0, 0, 0)
-			[
-				SNew(SNiagaraDetailsButton)
-				.Text(LOCTEXT("Cancel", "Cancel"))
-				.OnClicked(this, &SNiagaraScriptHighlightPicker::AddOrEditCancelled)
-			];
-	}
-
-	TSharedRef<SWidget> ConstructEditErrorWidgets()
-	{
-		return SNew(SHorizontalBox)
-			.Visibility(this, &SNiagaraScriptHighlightPicker::GetErrorVisibility)
-			+ SHorizontalBox::Slot()
-			.AutoWidth()
-			.Padding(0, 0, 5, 0)
-			.VAlign(VAlign_Center)
-			[
-				SNew(SImage)
-				.Image(FEditorStyle::GetBrush("Icons.Error"))
-			]
-			+ SHorizontalBox::Slot()
-			[
-				SNew(STextBlock)
-				.Text(this, &SNiagaraScriptHighlightPicker::GetErrorText)
-				.AutoWrapText(true)
-			];
-	}
-
-	FNiagaraScriptHighlight GetCurrentPropertyValue()
-	{
-		FNiagaraScriptHighlight CurrentPropertyValue;
-		void* PropertyValuePtr;
-		if (PropertyHandle->GetValueData(PropertyValuePtr) == FPropertyAccess::Success)
-		{
-			CurrentPropertyValue = *(FNiagaraScriptHighlight*)PropertyValuePtr;
-		}
-		return CurrentPropertyValue;
-	}
-
-	void RefreshDisplayedValue()
-	{
-		DisplayedHighlight = GetCurrentPropertyValue();
-	}
-
-	EVisibility GetPickerVisibility() const
-	{
-		return EditMode == EEditMode::Picker ? EVisibility::Visible : EVisibility::Collapsed;
-	}
-
-	EVisibility GetEditVisibility() const
-	{
-		return EditMode == EEditMode::EditCurrent || EditMode == EEditMode::EditNew ? EVisibility::Visible : EVisibility::Collapsed;
-	}
-
-	FText GetEditModeText() const
-	{
-		if (EditMode == EEditMode::EditNew)
-		{
-			return LOCTEXT("EditNewLabel", "Add a new highlight value");
-		}
-		else if (EditMode == EEditMode::EditCurrent)
-		{
-			return LOCTEXT("EditCurrentLabel", "Edit an existing highlight value");
-		}
-		else
-		{
-			return FText();
-		}
-	}
-
-	void DeleteItem()
-	{
-		TSharedPtr<IPropertyHandleArray> ParentArray = PropertyHandle->GetParentHandle()->AsArray();
-		if (ParentArray.IsValid())
-		{
-			int32 PropertyIndex = PropertyHandle->GetIndexInArray();
-			if (PropertyIndex != INDEX_NONE)
-			{
-				ParentArray->DeleteItem(PropertyHandle->GetIndexInArray());
-			}
-		}
-	}
-
-	typedef SItemSelector<FText, FNiagaraScriptHighlight> SNiagaraScriptHighlightSelector;
-
-	TSharedRef<SWidget> OnGetMenuContent()
-	{
-		return SNew(SBox)
-			.MaxDesiredHeight(200)
-			[
-				SNew(SVerticalBox)
-				// Existing Items
-				+ SVerticalBox::Slot()
-				.Padding(5, 10, 5, 5)
-				[
-					SNew(SNiagaraScriptHighlightSelector)
-					.Items(FNiagaraEditorModule::Get().GetCachedScriptAssetHighlights())
-					.ClickActivateMode(EItemSelectorClickActivateMode::SingleClick)
-					.OnCompareItemsForSorting(this, &SNiagaraScriptHighlightPicker::CompareHighlights)
-					.OnDoesItemMatchFilterText(this, &SNiagaraScriptHighlightPicker::DoesHightlightMatchFilterText)
-					.OnGenerateWidgetForItem(this, &SNiagaraScriptHighlightPicker::ConstructHighlightWidget)
-					.OnItemActivated(this, &SNiagaraScriptHighlightPicker::HighlightActivated)
-				]
-				// Add and edit buttons
-				+ SVerticalBox::Slot()
-				.AutoHeight()
-				.Padding(5, 0, 5, 5)
-				.HAlign(HAlign_Right)
-				[
-					SNew(SHorizontalBox)
-					+ SHorizontalBox::Slot()
-					.AutoWidth()
-					.Padding(0, 0, 5, 0)
-					[
-						SNew(SNiagaraDetailsButton)
-						.Text(LOCTEXT("AddNewText", "Add New"))
-						.ToolTipText(LOCTEXT("AddNewToolTipText", "Add a new highlight."))
-						.OnClicked(this, &SNiagaraScriptHighlightPicker::RequestNewHighlight)
-					]
-					+ SHorizontalBox::Slot()
-					.AutoWidth()
-					[
-						SNew(SNiagaraDetailsButton)
-						.Text(LOCTEXT("EditCurrentText", "Edit"))
-						.ToolTipText(LOCTEXT("EditCurrentToolTipText", "Edit the current highlight."))
-						.OnClicked(this, &SNiagaraScriptHighlightPicker::RequestEditHighlight)
-					]
-				]
-			];
-	}
-
-	bool CompareHighlights(const FNiagaraScriptHighlight& HighlightA, const FNiagaraScriptHighlight& HighlightB) const
-	{
-		return HighlightA.DisplayName.CompareTo(HighlightB.DisplayName) > 0;
-	}
-
-	bool DoesHightlightMatchFilterText(const FText& FilterText, const FNiagaraScriptHighlight& Highlight)
-	{
-		return Highlight.DisplayName.ToString().Contains(FilterText.ToString());
-	}
-
-	TSharedRef<SWidget> ConstructHighlightWidget(const FNiagaraScriptHighlight& Highlight) const
-	{
-		return SNew(SHorizontalBox)
-			+ SHorizontalBox::Slot()
-			.AutoWidth()
-			.Padding(0, 5, 5, 5)
-			[
-				SNew(SBorder)
-				.BorderImage(FEditorStyle::GetBrush("WhiteBrush"))
-				.BorderBackgroundColor(Highlight.Color)
-				.Padding(0)
-				[
-					SNew(SBox)
-					.WidthOverride(DisplayColorWidth)
-					.HeightOverride(DisplayColorHeight)
-				]
-			]
-			+ SHorizontalBox::Slot()
-			.VAlign(VAlign_Center)
-			[
-				SNew(STextBlock)
-				.Text(Highlight.DisplayName)
-			];
-	}
-
-	void HighlightActivated(const FNiagaraScriptHighlight& ActivatedHighlight)
-	{
-		HighlightPickerCombo->SetIsOpen(false);
-		ApplyHighlightChange(ActivatedHighlight);
-	}
-
-	void ApplyHighlightChange(const FNiagaraScriptHighlight& NewHighlight)
-	{
-		void* PropertyValuePtr;
-		if (PropertyHandle->GetValueData(PropertyValuePtr) == FPropertyAccess::Success)
-		{
-			TArray<UObject*> OuterObjects;
-			PropertyHandle->GetOuterObjects(OuterObjects);
-			FScopedTransaction EditHightlightTransaction(LOCTEXT("EditHighlight", "Edit Highlight"));
-			for (UObject* OuterObject : OuterObjects)
-			{
-				OuterObject->Modify();
-			}
-			PropertyHandle->NotifyPreChange();
-			FNiagaraScriptHighlight* CurrentScriptHighlight = (FNiagaraScriptHighlight*)PropertyValuePtr;
-			*CurrentScriptHighlight = NewHighlight;
-			DisplayedHighlight = NewHighlight;
-			PropertyHandle->NotifyPostChange(EPropertyChangeType::ValueSet);
-			PropertyHandle->NotifyFinishedChangingProperties();
-		}
-	}
-
-	FReply RequestNewHighlight()
-	{
-		EditMode = EEditMode::EditNew;
-		HighlightPickerCombo->SetIsOpen(false);
-		return FReply::Handled();
-	}
-
-	FReply RequestEditHighlight()
-	{
-		EditMode = EEditMode::EditCurrent;
-		HighlightPickerCombo->SetIsOpen(false);
-		return FReply::Handled();
-	}
-
-	FSlateColor GetHighlightSlateColor() const
-	{
-		return DisplayedHighlight.Color;
-	}
-
-	FLinearColor GetHighlightColor() const
-	{
-		return DisplayedHighlight.Color;
-	}
-
-	FText GetHighlightText() const
-	{
-		return DisplayedHighlight.DisplayName;
-	}
-
-	void HighlightTextCommitted(const FText& Text, ETextCommit::Type CommitType)
-	{
-		DisplayedHighlight.DisplayName = Text;
-	}
-
-	FReply OnMouseButtonDownColorBlock(const FGeometry& MyGeometry, const FPointerEvent& MouseEvent)
-	{
-		if (MouseEvent.GetEffectingButton() != EKeys::LeftMouseButton)
-		{
-			return FReply::Unhandled();
-		}
-
-		FColorPickerArgs PickerArgs;
-		{
-			PickerArgs.bUseAlpha = true;
-			PickerArgs.bOnlyRefreshOnMouseUp = false;
-			PickerArgs.bOnlyRefreshOnOk = false;
-			PickerArgs.DisplayGamma = TAttribute<float>::Create(TAttribute<float>::FGetter::CreateUObject(GEngine, &UEngine::GetDisplayGamma));
-			PickerArgs.OnColorCommitted = FOnLinearColorValueChanged::CreateSP(this, &SNiagaraScriptHighlightPicker::HighlightColorCommitted);
-			PickerArgs.OnColorPickerCancelled = FOnColorPickerCancelled::CreateSP(this, &SNiagaraScriptHighlightPicker::HighlightColorCancelled);
-			PickerArgs.InitialColorOverride = DisplayedHighlight.Color;
-			PickerArgs.ParentWidget = ColorBlock;
-		}
-
-		OpenColorPicker(PickerArgs);
-		return FReply::Handled();
-	}
-
-	void HighlightColorCommitted(FLinearColor CommittedColor)
-	{
-		DisplayedHighlight.Color = CommittedColor;
-	}
-
-	void HighlightColorCancelled(FLinearColor OriginalColor)
-	{
-		DisplayedHighlight.Color = OriginalColor;
-	}
-
-	EVisibility GetAddNewButtonVisibility() const
-	{
-		return EditMode == EEditMode::EditNew ? EVisibility::Visible : EVisibility::Collapsed;
-	}
-
-	EVisibility GetApplyButtonVisibility() const
-	{
-		return EditMode == EEditMode::EditCurrent ? EVisibility::Visible : EVisibility::Collapsed;
-	}
-
-	bool ValidateNewHighlight(const FNiagaraScriptHighlight& HighlightToCheck, TOptional<FNiagaraScriptHighlight> CurrentHightlightToIgnore, TOptional<FNiagaraScriptHighlight>& OutFoundExistingHighlight, FText& OutErrorMessage)
-	{
-		if (HighlightToCheck.DisplayName.IsEmptyOrWhitespace())
-		{
-			OutErrorMessage = LOCTEXT("HighlightWithEmptyDisplay", "Display name can not be empty or white space.");
-			return false;
-		}
-
-		for (const FNiagaraScriptHighlight& ExistingHighlight : FNiagaraEditorModule::Get().GetCachedScriptAssetHighlights())
-		{
-			if (CurrentHightlightToIgnore.IsSet() && ExistingHighlight == CurrentHightlightToIgnore)
-			{
-				continue;
-			}
-
-			bool bColorMatches = ExistingHighlight.Color.Equals(HighlightToCheck.Color);
-			bool bDisplayNameMatches = ExistingHighlight.DisplayName.CompareTo(HighlightToCheck.DisplayName) == 0;
-			if (bColorMatches && bDisplayNameMatches)
-			{
-				// A match was found so return it so the color will match exactly.
-				OutFoundExistingHighlight = ExistingHighlight;
-				return true;
-			}
-			if (bColorMatches && bDisplayNameMatches == false)
-			{
-				// Only the color or display name matched which is invalid.
-				OutErrorMessage = LOCTEXT("HightlightMismatchDisplayName", "This new highlight matches an existing highlight's color but has a different display name.");
-				return false;
-			}
-			if (bColorMatches == false && bDisplayNameMatches)
-			{
-				// Only the color or display name matched which is invalid.
-				OutErrorMessage = LOCTEXT("HightlightMismatchColor", "This new highlight matches an existing highlight's display name but has a different color.");
-				return false;
-			}
-		}
-		// No matches were found which is also valid.
-		OutFoundExistingHighlight.Reset();
-		return true;
-	}
-
-	FReply AddNewRequested()
-	{
-		TOptional<FNiagaraScriptHighlight> ExistingHighlight;
-		FText ValidateErrorMessage;
-		if (ValidateNewHighlight(DisplayedHighlight, TOptional<FNiagaraScriptHighlight>(), ExistingHighlight, ValidateErrorMessage))
-		{
-			ApplyHighlightChange(ExistingHighlight.IsSet() ? ExistingHighlight.GetValue() : DisplayedHighlight);
-			EditMode = EEditMode::Picker;
-			bShowError = false;
-			ErrorMessage = FText();
-		}
-		else
-		{
-			bShowError = true;
-			ErrorMessage = ValidateErrorMessage;
-		}
-		return FReply::Handled();
-	}
-
-	FReply ApplyCurrentRequested()
-	{
-		TOptional<FNiagaraScriptHighlight> ExistingHighlight;
-		FText ValidateErrorMessage;
-		if (ValidateNewHighlight(DisplayedHighlight, TOptional<FNiagaraScriptHighlight>(), ExistingHighlight, ValidateErrorMessage))
-		{
-			ApplyHighlightChange(ExistingHighlight.IsSet() ? ExistingHighlight.GetValue() : DisplayedHighlight);
-			EditMode = EEditMode::Picker;
-			bShowError = false;
-			ErrorMessage = FText();
-		}
-		else
-		{
-			bShowError = true;
-			ErrorMessage = ValidateErrorMessage;
-		}
-		return FReply::Handled();
-	}
-
-	FReply ApplyAllRequested()
-	{
-		TOptional<FNiagaraScriptHighlight> ExistingHighlight;
-		FText ValidateErrorMessage;
-		FNiagaraScriptHighlight CurrentHighlight = GetCurrentPropertyValue();
-		if (ValidateNewHighlight(DisplayedHighlight, CurrentHighlight, ExistingHighlight, ValidateErrorMessage))
-		{
-			TArray<FAssetData> MatchingScriptAssets;
-			FNiagaraEditorModule::Get().GetScriptAssetsMatchingHighlight(CurrentHighlight, MatchingScriptAssets);
-
-			FScopedTransaction ApplyAllTransaction(LOCTEXT("ApplyAllTransaction", "Apply highlight change to all scripts."));
-
-			FNiagaraScriptHighlight NewHighlight = ExistingHighlight.IsSet() ? ExistingHighlight.GetValue() : DisplayedHighlight;
-			ApplyHighlightChange(NewHighlight);
-
-			for (const FAssetData& MatchingScriptAsset : MatchingScriptAssets)
-			{
-				UNiagaraScript* Script = CastChecked<UNiagaraScript>(MatchingScriptAsset.GetAsset());
-				Script->Modify();
-				for (FNiagaraScriptHighlight& Highlight : Script->GetLatestScriptData()->Highlights)
-				{
-					if (Highlight == CurrentHighlight)
-					{
-						Highlight = NewHighlight;
-					}
-				}
-			}
-
-			EditMode = EEditMode::Picker;
-			bShowError = false;
-			ErrorMessage = FText();
-		}
-		else
-		{
-			bShowError = true;
-			ErrorMessage = ValidateErrorMessage;
-		}
-		return FReply::Handled();
-	}
-
-	FReply AddOrEditCancelled()
-	{
-		EditMode = EEditMode::Picker;
-		bShowError = false;
-		ErrorMessage = FText();
-		RefreshDisplayedValue();
-		return FReply::Handled();
-	}
-
-	EVisibility GetErrorVisibility() const
-	{
-		return bShowError ? EVisibility::Visible : EVisibility::Collapsed;
-	}
-
-	FText GetErrorText() const
-	{
-		return ErrorMessage;
-	}
-
-private:
-	TSharedPtr<IPropertyHandle> PropertyHandle;
-	FNiagaraScriptHighlight DisplayedHighlight;
-	EEditMode EditMode;
-	TSharedPtr<SComboButton> HighlightPickerCombo;
-	TSharedPtr<SColorBlock> ColorBlock;
-	bool bShowError;
-	FText ErrorMessage;
-};
-
-void FNiagaraScriptHighlightDetails::CustomizeHeader(TSharedRef<IPropertyHandle> PropertyHandle, FDetailWidgetRow& HeaderRow, IPropertyTypeCustomizationUtils& CustomizationUtils)
-{
-	HeaderRow.WholeRowContent()
-		[
-			SNew(SNiagaraScriptHighlightPicker, PropertyHandle)
-		];
-}
-
-
 #undef LOCTEXT_NAMESPACE
diff --git a/Engine/Plugins/FX/Niagara/Source/NiagaraEditor/Private/Customizations/NiagaraScriptDetails.h b/Engine/Plugins/FX/Niagara/Source/NiagaraEditor/Private/Customizations/NiagaraScriptDetails.h
--- a/Engine/Plugins/FX/Niagara/Source/NiagaraEditor/Private/Customizations/NiagaraScriptDetails.h
+++ b/Engine/Plugins/FX/Niagara/Source/NiagaraEditor/Private/Customizations/NiagaraScriptDetails.h
@@ -30,17 +30,3 @@ private:
 	TSharedPtr<FNiagaraScriptViewModel> ScriptViewModel;
 	TSharedPtr<FNiagaraMetaDataCustomNodeBuilder> MetaDataBuilder;
 };
-
-class FNiagaraScriptHighlightDetails : public IPropertyTypeCustomization
-{
-public:
-	static TSharedRef<IPropertyTypeCustomization> MakeInstance()
-	{
-		return MakeShared<FNiagaraScriptHighlightDetails>();
-	}
-
-	virtual void CustomizeHeader(TSharedRef<IPropertyHandle> PropertyHandle, FDetailWidgetRow& HeaderRow, IPropertyTypeCustomizationUtils& CustomizationUtils) override;
-
-	virtual void CustomizeChildren(TSharedRef<IPropertyHandle> PropertyHandle, IDetailChildrenBuilder& ChildBuilder, IPropertyTypeCustomizationUtils& CustomizationUtils) override { }
-};
-
diff --git a/Engine/Plugins/FX/Niagara/Source/NiagaraEditor/Private/NiagaraEditorModule.cpp b/Engine/Plugins/FX/Niagara/Source/NiagaraEditor/Private/NiagaraEditorModule.cpp
--- a/Engine/Plugins/FX/Niagara/Source/NiagaraEditor/Private/NiagaraEditorModule.cpp
+++ b/Engine/Plugins/FX/Niagara/Source/NiagaraEditor/Private/NiagaraEditorModule.cpp
@@ -879,16 +879,11 @@ void FNiagaraEditorModule::StartupModule()
 	PropertyModule.RegisterCustomPropertyTypeLayout(
 		FNiagaraUserParameterBinding::StaticStruct()->GetFName(),
 		FOnGetPropertyTypeCustomizationInstance::CreateStatic(&FNiagaraUserParameterBindingCustomization::MakeInstance));
-
-
+	
 	PropertyModule.RegisterCustomPropertyTypeLayout(
 		FNiagaraMaterialAttributeBinding::StaticStruct()->GetFName(),
 		FOnGetPropertyTypeCustomizationInstance::CreateStatic(&FNiagaraMaterialAttributeBindingCustomization::MakeInstance));
 
-	PropertyModule.RegisterCustomPropertyTypeLayout(
-		FNiagaraScriptHighlight::StaticStruct()->GetFName(),
-		FOnGetPropertyTypeCustomizationInstance::CreateStatic(&FNiagaraScriptHighlightDetails::MakeInstance));
-
 	PropertyModule.RegisterCustomPropertyTypeLayout(
 	    FNiagaraVariableDataInterfaceBinding::StaticStruct()->GetFName(),
 		FOnGetPropertyTypeCustomizationInstance::CreateStatic(&FNiagaraDataInterfaceBindingCustomization::MakeInstance));
@@ -1487,48 +1482,6 @@ void FNiagaraEditorModule::InvalidateCachedScriptAssetData()
 	TypeConversionScriptCache.Reset();
 }
 
-const TArray<FNiagaraScriptHighlight>& FNiagaraEditorModule::GetCachedScriptAssetHighlights() const
-{
-	if (CachedScriptAssetHighlights.IsSet() == false)
-	{
-		CachedScriptAssetHighlights = TArray<FNiagaraScriptHighlight>();
-		FAssetRegistryModule& AssetRegistryModule = FModuleManager::LoadModuleChecked<FAssetRegistryModule>("AssetRegistry");
-		TArray<FAssetData> ScriptAssets;
-		AssetRegistryModule.Get().GetAssetsByClass(UNiagaraScript::StaticClass()->GetFName(), ScriptAssets);
-		for (const FAssetData& ScriptAsset : ScriptAssets)
-		{
-			if (ScriptAsset.IsAssetLoaded())
-			{
-				UNiagaraScript* Script = CastChecked<UNiagaraScript>(ScriptAsset.GetAsset());
-				for (const FNiagaraScriptHighlight& Highlight : Script->GetLatestScriptData()->Highlights)
-				{
-					if (Highlight.IsValid())
-					{
-						CachedScriptAssetHighlights->AddUnique(Highlight);
-					}
-				}
-			}
-			else
-			{
-				FString HighlightsString;
-				if (ScriptAsset.GetTagValue(GET_MEMBER_NAME_CHECKED(FVersionedNiagaraScriptData, Highlights), HighlightsString))
-				{
-					TArray<FNiagaraScriptHighlight> Highlights;
-					FNiagaraScriptHighlight::JsonToArray(HighlightsString, Highlights);
-					for (const FNiagaraScriptHighlight& Highlight : Highlights)
-					{
-						if (Highlight.IsValid())
-						{
-							CachedScriptAssetHighlights->AddUnique(Highlight);
-						}
-					}
-				}
-			}
-		}
-	}
-	return CachedScriptAssetHighlights.GetValue();
-}
-
 const TArray<UNiagaraScript*>& FNiagaraEditorModule::GetCachedTypeConversionScripts() const
 {
 	if (!TypeConversionScriptCache.IsSet())
@@ -1556,45 +1509,6 @@ const TArray<UNiagaraScript*>& FNiagaraEditorModule::GetCachedTypeConversionScri
 	return TypeConversionScriptCache.GetValue();
 }
 
-void FNiagaraEditorModule::GetScriptAssetsMatchingHighlight(const FNiagaraScriptHighlight& InHighlight, TArray<FAssetData>& OutMatchingScriptAssets) const
-{
-	FAssetRegistryModule& AssetRegistryModule = FModuleManager::LoadModuleChecked<FAssetRegistryModule>("AssetRegistry");
-	TArray<FAssetData> ScriptAssets;
-	AssetRegistryModule.Get().GetAssetsByClass(UNiagaraScript::StaticClass()->GetFName(), ScriptAssets);
-	for (const FAssetData& ScriptAsset : ScriptAssets)
-	{
-		if (ScriptAsset.IsAssetLoaded())
-		{
-			UNiagaraScript* Script = CastChecked<UNiagaraScript>(ScriptAsset.GetAsset());
-			for (const FNiagaraScriptHighlight& Highlight : Script->GetLatestScriptData()->Highlights)
-			{
-				if (Highlight == InHighlight)
-				{
-					OutMatchingScriptAssets.Add(ScriptAsset);
-					break;
-				}
-			}
-		}
-		else
-		{
-			FString HighlightsString;
-			if (ScriptAsset.GetTagValue(GET_MEMBER_NAME_CHECKED(FVersionedNiagaraScriptData, Highlights), HighlightsString))
-			{
-				TArray<FNiagaraScriptHighlight> Highlights;
-				FNiagaraScriptHighlight::JsonToArray(HighlightsString, Highlights);
-				for (const FNiagaraScriptHighlight& Highlight : Highlights)
-				{
-					if (Highlight == InHighlight)
-					{
-						OutMatchingScriptAssets.Add(ScriptAsset);
-						break;
-					}
-				}
-			}
-		}
-	}
-}
-
 FNiagaraClipboard& FNiagaraEditorModule::GetClipboard() const
 {
 	return Clipboard.Get();
diff --git a/Engine/Plugins/FX/Niagara/Source/NiagaraEditor/Private/ViewModels/Stack/NiagaraStackItem.cpp b/Engine/Plugins/FX/Niagara/Source/NiagaraEditor/Private/ViewModels/Stack/NiagaraStackItem.cpp
--- a/Engine/Plugins/FX/Niagara/Source/NiagaraEditor/Private/ViewModels/Stack/NiagaraStackItem.cpp
+++ b/Engine/Plugins/FX/Niagara/Source/NiagaraEditor/Private/ViewModels/Stack/NiagaraStackItem.cpp
@@ -47,12 +47,6 @@ void UNiagaraStackItem::SetIsEnabled(bool bInIsEnabled)
 	SetIsEnabledInternal(bInIsEnabled);
 }
 
-const TArray<FNiagaraScriptHighlight>& UNiagaraStackItem::GetHighlights() const
-{
-	static TArray<FNiagaraScriptHighlight> Empty;
-	return Empty;
-}
-
 const FSlateBrush* UNiagaraStackItem::GetIconBrush() const
 {
 	return FEditorStyle::GetBrush("NoBrush");
diff --git a/Engine/Plugins/FX/Niagara/Source/NiagaraEditor/Private/ViewModels/Stack/NiagaraStackModuleItem.cpp b/Engine/Plugins/FX/Niagara/Source/NiagaraEditor/Private/ViewModels/Stack/NiagaraStackModuleItem.cpp
--- a/Engine/Plugins/FX/Niagara/Source/NiagaraEditor/Private/ViewModels/Stack/NiagaraStackModuleItem.cpp
+++ b/Engine/Plugins/FX/Niagara/Source/NiagaraEditor/Private/ViewModels/Stack/NiagaraStackModuleItem.cpp
@@ -1100,6 +1100,11 @@ void UNiagaraStackModuleItem::Refresh()
 	}
 }
 
+TOptional<float> UNiagaraStackModuleItem::GetCustomOverviewIndent()
+{
+	return 8.f;
+}
+
 bool UNiagaraStackModuleItem::GetIsEnabled() const
 {
 	return bIsEnabled;
@@ -1126,16 +1131,6 @@ void UNiagaraStackModuleItem::SetDebugDrawEnabled(bool bInEnabled)
 	OnRequestFullRefreshDeferred().Broadcast();
 }
 
-bool UNiagaraStackModuleItem::SupportsHighlights() const
-{
-	return FunctionCallNode != nullptr && FunctionCallNode->FunctionScript != nullptr;
-}
-
-const TArray<FNiagaraScriptHighlight>& UNiagaraStackModuleItem::GetHighlights() const
-{
-	return FunctionCallNode->GetScriptData()->Highlights;
-}
-
 int32 UNiagaraStackModuleItem::GetModuleIndex() const
 {
 	TArray<FNiagaraStackGraphUtilities::FStackNodeGroup> StackGroups;
diff --git a/Engine/Plugins/FX/Niagara/Source/NiagaraEditor/Public/NiagaraEditorModule.h b/Engine/Plugins/FX/Niagara/Source/NiagaraEditor/Public/NiagaraEditorModule.h
--- a/Engine/Plugins/FX/Niagara/Source/NiagaraEditor/Public/NiagaraEditorModule.h
+++ b/Engine/Plugins/FX/Niagara/Source/NiagaraEditor/Public/NiagaraEditorModule.h
@@ -172,12 +172,9 @@ public:
 	NIAGARAEDITOR_API const FNiagaraEditorCommands& GetCommands() const;
 
 	void InvalidateCachedScriptAssetData();
-
-	const TArray<FNiagaraScriptHighlight>& GetCachedScriptAssetHighlights() const;
+	
 	const TArray<UNiagaraScript*>& GetCachedTypeConversionScripts() const;
 
-	void GetScriptAssetsMatchingHighlight(const FNiagaraScriptHighlight& InHighlight, TArray<FAssetData>& OutMatchingScriptAssets) const;
-
 	NIAGARAEDITOR_API FNiagaraClipboard& GetClipboard() const;
 
 	template<typename T>
diff --git a/Engine/Plugins/FX/Niagara/Source/NiagaraEditor/Public/ViewModels/Stack/NiagaraStackItem.h b/Engine/Plugins/FX/Niagara/Source/NiagaraEditor/Public/ViewModels/Stack/NiagaraStackItem.h
--- a/Engine/Plugins/FX/Niagara/Source/NiagaraEditor/Public/ViewModels/Stack/NiagaraStackItem.h
+++ b/Engine/Plugins/FX/Niagara/Source/NiagaraEditor/Public/ViewModels/Stack/NiagaraStackItem.h
@@ -34,11 +34,13 @@ public:
 	virtual bool SupportsChangeEnabled() const { return false; }
 	void SetIsEnabled(bool bInIsEnabled);
 
-	virtual bool SupportsHighlights() const { return false; }
-	virtual const TArray<FNiagaraScriptHighlight>& GetHighlights() const;
-
+	
 	virtual bool SupportsIcon() const { return false; }
 	virtual const FSlateBrush* GetIconBrush() const;
+	/** The custom overview indent can be used as a backup indent in the overview stack.
+	 *  If an icon exists, the icon will be used instead, and if no custom overview indent is specified and no icon is set, we use a default indent
+	 */
+	virtual TOptional<float> GetCustomOverviewIndent() { return TOptional<float>();}
 
 	virtual bool SupportsResetToBase() const { return false; }
 	virtual bool TestCanResetToBaseWithMessage(FText& OutCanResetToBaseMessage) const { return false; }
diff --git a/Engine/Plugins/FX/Niagara/Source/NiagaraEditor/Public/ViewModels/Stack/NiagaraStackModuleItem.h b/Engine/Plugins/FX/Niagara/Source/NiagaraEditor/Public/ViewModels/Stack/NiagaraStackModuleItem.h
--- a/Engine/Plugins/FX/Niagara/Source/NiagaraEditor/Public/ViewModels/Stack/NiagaraStackModuleItem.h
+++ b/Engine/Plugins/FX/Niagara/Source/NiagaraEditor/Public/ViewModels/Stack/NiagaraStackModuleItem.h
@@ -43,12 +43,11 @@ public:
 
 	virtual bool SupportsRename() const override { return true; }
 
+	virtual TOptional<float> GetCustomOverviewIndent() override;
+
 	virtual bool SupportsChangeEnabled() const override { return true; }
 	virtual bool GetIsEnabled() const override;
 
-	virtual bool SupportsHighlights() const override;
-	virtual const TArray<FNiagaraScriptHighlight>& GetHighlights() const override;
-
 	int32 GetModuleIndex() const;
 	
 	UObject* GetExternalAsset() const override;
diff --git a/Engine/Plugins/FX/Niagara/Source/NiagaraEditorWidgets/Private/SNiagaraOverviewStack.cpp b/Engine/Plugins/FX/Niagara/Source/NiagaraEditorWidgets/Private/SNiagaraOverviewStack.cpp
--- a/Engine/Plugins/FX/Niagara/Source/NiagaraEditorWidgets/Private/SNiagaraOverviewStack.cpp
+++ b/Engine/Plugins/FX/Niagara/Source/NiagaraEditorWidgets/Private/SNiagaraOverviewStack.cpp
@@ -720,91 +720,9 @@ TSharedRef<ITableRow> SNiagaraOverviewStack::OnGenerateRowForEntry(UNiagaraStack
 	{
 		UNiagaraStackItem* StackItem = CastChecked<UNiagaraStackItem>(Item);
 		TSharedPtr<SWidget> IndentContent;
-		if (StackItem->SupportsHighlights())
-		{
-			TArray<FNiagaraScriptHighlight> ScriptHighlights;
-			for (const FNiagaraScriptHighlight& ScriptHighlight : StackItem->GetHighlights())
-			{
-				if (ScriptHighlight.IsValid())
-				{
-					ScriptHighlights.Add(ScriptHighlight);
-				}
-			}
-
-			TSharedRef<SHorizontalBox> ToolTipBox = SNew(SHorizontalBox);
-			if (ScriptHighlights.Num() > 0)
-			{
-				for (const FNiagaraScriptHighlight& ScriptHighlight : ScriptHighlights)
-				{
-					ToolTipBox->AddSlot()
-						.AutoWidth()
-						.VAlign(VAlign_Center)
-						.Padding(0, 0, 5, 0)
-						[
-							SNew(SImage)
-							.Image(FNiagaraEditorWidgetsStyle::Get().GetBrush("NiagaraEditor.Stack.ModuleHighlightLarge"))
-							.ColorAndOpacity(ScriptHighlight.Color)
-						];
-					ToolTipBox->AddSlot()
-						.AutoWidth()
-						.VAlign(VAlign_Center)
-						.Padding(0, 0, 10, 0)
-						[
-							SNew(STextBlock)
-							.Text(ScriptHighlight.DisplayName)
-						];
-				}
-			}
-
-			TSharedRef<SBox> IndentBox = SNew(SBox)
-				.ToolTip(SNew(SToolTip) [ ToolTipBox ])
-				.IsEnabled_UObject(Item, &UNiagaraStackEntry::GetIsEnabledAndOwnerIsEnabled)
-				.Visibility(EVisibility::Visible)
-				.WidthOverride(IconSize.X)
-				.HeightOverride(IconSize.Y);
-			
-			if (ScriptHighlights.Num() > 0)
-			{
-				IndentBox->SetHAlign(HAlign_Left);
-				IndentBox->SetVAlign(VAlign_Top);
 
-				TSharedPtr<SGridPanel> HighlightsGrid;
-				IndentBox->SetContent(SAssignNew(HighlightsGrid, SGridPanel));
-
-				int32 HighlightsAdded = 0;
-				for (int32 HighlightIndex = 0; HighlightIndex < ScriptHighlights.Num() && HighlightsAdded < 4; HighlightIndex++)
-				{
-					if (ScriptHighlights[HighlightIndex].IsValid())
-					{
-						FName HighlightImageBrushName;
-						FLinearColor HighlightImageColor;
-						if (HighlightIndex < 3)
-						{
-							HighlightImageBrushName = "NiagaraEditor.Stack.ModuleHighlight";
-							HighlightImageColor = ScriptHighlights[HighlightIndex].Color;
-						}
-						else
-						{
-							HighlightImageBrushName = "NiagaraEditor.Stack.ModuleHighlightMore";
-							HighlightImageColor = FLinearColor::White;
-						}
-
-						int32 Column = HighlightIndex % 2;
-						int32 Row = HighlightIndex / 2;
-						HighlightsGrid->AddSlot(Column, Row)
-							.Padding(0, 0, 1 - Column, 1 - Row)
-							[
-								SNew(SImage)
-								.Image(FNiagaraEditorWidgetsStyle::Get().GetBrush(HighlightImageBrushName))
-								.ColorAndOpacity(HighlightImageColor)
-							];
-						HighlightsAdded++;
-					}
-				}
-			}
-			IndentContent = IndentBox;
-		}
-		else if (StackItem->SupportsIcon())
+		// we give icons priority for the indent content
+		if (StackItem->SupportsIcon())
 		{
 			IndentContent = SNew(SBox)
 				.IsEnabled_UObject(Item, &UNiagaraStackEntry::GetIsEnabledAndOwnerIsEnabled)
@@ -817,6 +735,15 @@ TSharedRef<ITableRow> SNiagaraOverviewStack::OnGenerateRowForEntry(UNiagaraStack
 					.Image_UObject(StackItem, &UNiagaraStackItem::GetIconBrush)
 				];
 		}
+		// if we didn't have an icon but we specified a custom indent, we use that instead
+		else if(StackItem->GetCustomOverviewIndent().IsSet())
+		{
+			float IndentValue = StackItem->GetCustomOverviewIndent().GetValue();
+			IndentContent = SNew(SBox)
+				.WidthOverride(IndentValue)
+				.HeightOverride(1.f);
+		}
+		// otherwise, use a default
 		else
 		{
 			IndentContent = SNew(SBox)
