Repair perforce history on SNiagaraOverviewStack.cpp

diff --git a/Engine/Plugins/FX/Niagara/Source/NiagaraEditorWidgets/Private/SNiagaraOverviewStack.cpp b/Engine/Plugins/FX/Niagara/Source/NiagaraEditorWidgets/Private/SNiagaraOverviewStack.cpp
--- a/Engine/Plugins/FX/Niagara/Source/NiagaraEditorWidgets/Private/SNiagaraOverviewStack.cpp
+++ b/Engine/Plugins/FX/Niagara/Source/NiagaraEditorWidgets/Private/SNiagaraOverviewStack.cpp
@@ -26,6 +26,7 @@
 #include "NiagaraEditorModule.h"
 #include "NiagaraNodeAssignment.h"
 #include "Widgets/SNiagaraParameterName.h"
+#include "SNiagaraOverviewInlineParameterBox.h"
 #include "Widgets/Layout/SScrollBox.h"
 #include "Widgets/Layout/SGridPanel.h"
 #include "Widgets/Layout/SScaleBox.h"
@@ -763,13 +764,15 @@ TSharedRef<ITableRow> SNiagaraOverviewStack::OnGenerateRowForEntry(UNiagaraStack
 			// Name
 			+ SHorizontalBox::Slot()
 			.VAlign(VAlign_Center)
+			.AutoWidth()
+			.MaxWidth(150.f)
 			.Padding(3, 2, 0, 2)
 			[
 				SNew(SNiagaraSystemOverviewItemName, StackItem)
 			]
 			+ SHorizontalBox::Slot()
             .AutoWidth()
-            .Padding(6, 0, 9, 0)
+            .Padding(2, 0, 2, 0)
             [
                 SNew(SNiagaraStackRowPerfWidget, Item)
             ];
@@ -779,6 +782,12 @@ TSharedRef<ITableRow> SNiagaraOverviewStack::OnGenerateRowForEntry(UNiagaraStack
 			
 		if (StackModuleItem)
 		{
+			ContentBox->AddSlot()
+			.HAlign(HAlign_Fill)
+			[
+				SNew(SNiagaraOverviewInlineParameterBox, *StackModuleItem)
+			];
+			
 			if(StackModuleItem->IsScratchModule())
 			{
 				ContentBox->AddSlot()
@@ -793,27 +802,30 @@ TSharedRef<ITableRow> SNiagaraOverviewStack::OnGenerateRowForEntry(UNiagaraStack
 			}
 			
 			if(StackModuleItem->GetModuleNode().ContainsDebugSwitch())
-		{
-			ContentBox->AddSlot()
-				.VAlign(VAlign_Center)
-				.AutoWidth()
-				.Padding(3, 0, 0, 0)
-				[
-					SNew(SButton)
-					.ButtonColorAndOpacity(FLinearColor::Transparent)
-					.ForegroundColor(FLinearColor::Transparent)
-					.ToolTipText(LOCTEXT("EnableDebugDrawCheckBoxToolTip", "Enable or disable debug drawing for this item."))
-					.OnClicked(this, &SNiagaraOverviewStack::ToggleModuleDebugDraw, StackItem)
-					.ContentPadding(FMargin(0, 0, 0, 0))
+			{
+				ContentBox->AddSlot()
+					.VAlign(VAlign_Center)
+					.HAlign(HAlign_Center)
+					.AutoWidth()
+					.Padding(3, 0)
 					[
-						SNew(SImage)
-						.Image(this, &SNiagaraOverviewStack::GetDebugIconBrush, StackItem)
-					]
-				];
-		}
+						SNew(SButton)
+						.HAlign(HAlign_Center)
+						.VAlign(VAlign_Center)
+						.ButtonColorAndOpacity(FLinearColor::Transparent)
+						.ForegroundColor(FLinearColor::Transparent)
+						.ToolTipText(LOCTEXT("EnableDebugDrawCheckBoxToolTip", "Enable or disable debug drawing for this item."))
+						.OnClicked(this, &SNiagaraOverviewStack::ToggleModuleDebugDraw, StackItem)
+						// @Todo the debug icons aren't centered, so we correct this using padding. We also shrink the button size
+						.ContentPadding(FMargin(-11, 2, -7, 2))
+						[
+							SNew(SImage)
+							.Image(this, &SNiagaraOverviewStack::GetDebugIconBrush, StackItem)
+						]
+					];
+			}
 		}
 
-
 		// Enabled checkbox
 		ContentBox->AddSlot()
 			.VAlign(VAlign_Center)
@@ -827,7 +839,7 @@ TSharedRef<ITableRow> SNiagaraOverviewStack::OnGenerateRowForEntry(UNiagaraStack
 				.OnCheckedChanged_UObject(StackItem, &UNiagaraStackItem::SetIsEnabled)
 			];
 
-		Content = ContentBox;
+		Content = ContentBox;		
 	}
 	else if (Item->IsA<UNiagaraStackItemGroup>())
 	{
@@ -1129,4 +1141,4 @@ FReply SNiagaraOverviewStack::OnItemGroupDeleteClicked(UNiagaraStackItemGroup* G
 	return FReply::Handled();
 }
 
-#undef LOCTEXT_NAMESPACE
\ No newline at end of file
+#undef LOCTEXT_NAMESPACE
