Manual merge
Overview stack row inline module data
- Parameter meta data "Display in Overview"
- For local inputs (inputs & staticswitch inputs) only
- Creates a button per valid parameter that reads the button's current value (makes more sense with enums than with floats, unless the module is very explicit and low complexity)
- For enums: an entry can have an overriden name or an icon. Icon takes priority.
- Default border color for text-based entries is based on type. Can be overridden.
- button's tooltip shows the actual parameter behind the value
- clicking the button navigates the stack to the input
- Many layout tweaks and fixes (like debug button's minimum size being too large and offset, using negative padding counters it)

diff --git a/Engine/Plugins/FX/Niagara/Source/Niagara/Private/NiagaraCommon.cpp b/Engine/Plugins/FX/Niagara/Source/Niagara/Private/NiagaraCommon.cpp
--- a/Engine/Plugins/FX/Niagara/Source/Niagara/Private/NiagaraCommon.cpp
+++ b/Engine/Plugins/FX/Niagara/Source/Niagara/Private/NiagaraCommon.cpp
@@ -135,6 +135,73 @@ FString FNiagaraTypeHelper::ToString(const uint8* ValueData, const UObject* Stru
 	return Ret;
 }
 
+FString FNiagaraTypeHelper::ToStringCosmetic(const uint8* ValueData, const UObject* StructOrEnum)
+{
+	FString Ret;
+	if (const UEnum* Enum = Cast<const UEnum>(StructOrEnum))
+	{
+		Ret += Enum->GetDisplayNameTextByValue(*(int32*)ValueData).ToString();
+	}
+	else if (const UScriptStruct* Struct = Cast<const UScriptStruct>(StructOrEnum))
+	{
+		if (Struct == FNiagaraTypeDefinition::GetFloatStruct())
+		{
+			Ret += FString::Printf(TEXT("%g "), *(float*)ValueData);
+		}
+		else if (Struct == FNiagaraTypeDefinition::GetIntStruct())
+		{
+			Ret += FString::Printf(TEXT("%d "), *(int32*)ValueData);
+		}
+		else if (Struct == FNiagaraTypeDefinition::GetBoolStruct())
+		{
+			int32 Val = *(int32*)ValueData;
+			Ret += Val == 0xFFFFFFFF ? (TEXT("True")) : (Val == 0x0 ? TEXT("False") : TEXT("Invalid"));
+		}
+		else
+		{
+			for (TFieldIterator<FProperty> PropertyIt(Struct, EFieldIteratorFlags::IncludeSuper); PropertyIt; ++PropertyIt)
+			{
+				const FProperty* Property = *PropertyIt;
+				const uint8* PropPtr = ValueData + PropertyIt->GetOffset_ForInternal();
+				if (Property->IsA(FFloatProperty::StaticClass()))
+				{
+					Ret += FString::Printf(TEXT("%s: %g "), *Property->GetNameCPP(), *(float*)PropPtr);
+				}
+				else if (Property->IsA(FDoubleProperty::StaticClass()))
+				{
+					Ret += FString::Printf(TEXT("%s: %g "), *Property->GetNameCPP(), *(double*)PropPtr);
+				}
+				else if (Property->IsA(FUInt16Property::StaticClass()))
+				{
+					FFloat16 Val = *(FFloat16*)PropPtr;
+					Ret += FString::Printf(TEXT("%s: %f "), *Property->GetNameCPP(), Val.GetFloat());
+				}
+				else if (Property->IsA(FIntProperty::StaticClass()))
+				{
+					Ret += FString::Printf(TEXT("%s: %d "), *Property->GetNameCPP(), *(int32*)PropPtr);
+				}
+				else if (Property->IsA(FBoolProperty::StaticClass()))
+				{
+					int32 Val = *(int32*)ValueData;
+					FString BoolStr = Val == 0xFFFFFFFF ? (TEXT("True")) : (Val == 0x0 ? TEXT("False") : TEXT("Invalid"));
+					Ret += FString::Printf(TEXT("%s: %d "), *Property->GetNameCPP(), *BoolStr);
+				}
+				else if (const FStructProperty* StructProp = CastFieldChecked<const FStructProperty>(Property))
+				{
+					Ret += FString::Printf(TEXT("%s: (%s) "), *Property->GetNameCPP(), *FNiagaraTypeHelper::ToStringCosmetic(PropPtr, StructProp->Struct));
+				}
+				else
+				{
+					check(false);
+					Ret += TEXT("Unknown Type");
+				}
+			}
+		}
+	}
+	
+	return Ret;
+}
+
 UScriptStruct* FNiagaraTypeHelper::FindNiagaraFriendlyTopLevelStruct(UScriptStruct* InStruct)
 {
 	if (!InStruct)
diff --git a/Engine/Plugins/FX/Niagara/Source/Niagara/Public/NiagaraTypes.h b/Engine/Plugins/FX/Niagara/Source/Niagara/Public/NiagaraTypes.h
--- a/Engine/Plugins/FX/Niagara/Source/Niagara/Public/NiagaraTypes.h
+++ b/Engine/Plugins/FX/Niagara/Source/Niagara/Public/NiagaraTypes.h
@@ -7,6 +7,7 @@
 #include "Misc/SecureHash.h"
 #include "UObject/GCObject.h"
 #include "UObject/UnrealType.h"
+#include "Engine/Texture2D.h"
 
 #include "NiagaraTypes.generated.h"
 
@@ -302,6 +303,7 @@ class NIAGARA_API FNiagaraTypeHelper
 {
 public:
 	static FString ToString(const uint8* ValueData, const UObject* StructOrEnum);
+	static FString ToStringCosmetic(const uint8* ValueData, const UObject* StructOrEnum);
 	static UScriptStruct* FindNiagaraFriendlyTopLevelStruct(UScriptStruct* InStruct);
 	static bool IsNiagaraFriendlyTopLevelStruct(UScriptStruct* InStruct);
 };
@@ -610,6 +612,21 @@ public:
 	TArray<FString> TargetValues;
 };
 
+/** Defines override data for enum parameters displayed in the UI. */
+USTRUCT()
+struct NIAGARA_API FNiagaraEnumParameterMetaData
+{
+	GENERATED_BODY()
+	
+	/** If specified, this name will be used for the given enum entry. Useful for shortening names. */
+	UPROPERTY(EditAnywhere, Category="Enum Override")
+	FName OverrideName;
+
+	/** If specified, this icon will be used for the given enum entry. If OverrideName isn't empty, the icon takes priority. */
+	UPROPERTY(EditAnywhere, Category="Enum Override")
+	UTexture2D* IconOverride;
+};
+
 USTRUCT()
 struct NIAGARA_API FNiagaraVariableMetaData
 {
@@ -617,6 +634,10 @@ struct NIAGARA_API FNiagaraVariableMetaData
 
 	FNiagaraVariableMetaData()
 		: bAdvancedDisplay(false)
+		, bDisplayInOverviewStack(false)
+		, InlineParameterSortPriority(0)
+		, bOverrideColor(false)
+		, InlineParameterColorOverride(FLinearColor(EForceInit::ForceInit))
 		, EditorSortPriority(0)
 		, bInlineEditConditionToggle(false)
 		, bIsStaticSwitch_DEPRECATED(false)
@@ -633,6 +654,22 @@ struct NIAGARA_API FNiagaraVariableMetaData
 	UPROPERTY(EditAnywhere, Category = "Variable", meta = (SkipForCompileHash = "true"))
 	bool bAdvancedDisplay;
 
+	/** Declares that this parameter's value will be shown in the overview node if it's set to a local value. */
+	UPROPERTY(EditAnywhere, Category = "Variable", meta = (SkipForCompileHash = "true"))
+	bool bDisplayInOverviewStack;
+
+	UPROPERTY(EditAnywhere, Category = "Variable", meta = (EditCondition="bDisplayInOverviewStack", ToolTip = "Affects the sort order for parameters shown inline in the overview. Use a smaller number to push it to the top. Defaults to zero.", SkipForCompileHash = "true"))
+	int32 InlineParameterSortPriority;
+
+	UPROPERTY(EditAnywhere, Category = "Variable", meta = (InlineEditConditionToggle, ToolTip = "The color used to display a parameter in the overview. If no color is specified, the type color is used.", SkipForCompileHash = "true"))
+	bool bOverrideColor;
+	
+	UPROPERTY(EditAnywhere, Category = "Variable", meta = (EditCondition="bOverrideColor", ToolTip = "The color used to display a parameter in the overview. If no color is specified, the type color is used.", SkipForCompileHash = "true"))
+	FLinearColor InlineParameterColorOverride;
+
+	UPROPERTY(EditAnywhere, Category = "Variable", meta = (EditCondition="bDisplayInOverviewStack", ToolTip = "The index of the entry maps to the index of an enum value. Useful for overriding how an enum parameter is displayed in the overview.", SkipForCompileHash = "true"))
+	TArray<FNiagaraEnumParameterMetaData> InlineParameterEnumOverrides; 
+	
 	UPROPERTY(EditAnywhere, Category = "Variable", meta = (ToolTip = "Affects the sort order in the editor stacks. Use a smaller number to push it to the top. Defaults to zero.", SkipForCompileHash = "true"))
 	int32 EditorSortPriority;
 
@@ -987,6 +1024,16 @@ public:
 		return FNiagaraTypeHelper::ToString(ValueData, ClassStructOrEnum);
 	}
 
+	FString ToStringCosmetic(const uint8* ValueData)const
+	{
+		checkf(IsValid(), TEXT("Type definition is not valid."));
+		if (ValueData == nullptr)
+		{
+			return TEXT("(null)");
+		}
+		return FNiagaraTypeHelper::ToStringCosmetic(ValueData, ClassStructOrEnum);
+	}
+
 	static bool TypesAreAssignable(const FNiagaraTypeDefinition& TypeA, const FNiagaraTypeDefinition& TypeB);
 	static bool IsLossyConversion(const FNiagaraTypeDefinition& TypeA, const FNiagaraTypeDefinition& TypeB);
 	static FNiagaraTypeDefinition GetNumericOutputType(const TArray<FNiagaraTypeDefinition> TypeDefinintions, ENiagaraNumericOutputTypeSelectionMode SelectionMode);
diff --git a/Engine/Plugins/FX/Niagara/Source/NiagaraEditor/Private/Customizations/NiagaraScriptVariableCustomization.cpp b/Engine/Plugins/FX/Niagara/Source/NiagaraEditor/Private/Customizations/NiagaraScriptVariableCustomization.cpp
--- a/Engine/Plugins/FX/Niagara/Source/NiagaraEditor/Private/Customizations/NiagaraScriptVariableCustomization.cpp
+++ b/Engine/Plugins/FX/Niagara/Source/NiagaraEditor/Private/Customizations/NiagaraScriptVariableCustomization.cpp
@@ -11,6 +11,7 @@
 
 #include "EdGraphSchema_Niagara.h"
 #include "INiagaraEditorTypeUtilities.h"
+#include "NiagaraConstants.h"
 #include "NiagaraEditorCommon.h"
 #include "NiagaraEditorModule.h"
 #include "NiagaraEditorStyle.h"
@@ -177,6 +178,22 @@ void FNiagaraScriptVariableDetails::CustomizeDetails(IDetailLayoutBuilder& Detai
 
 	// Always hide the default value variant. We generate a node for this property to enable PostEditChangeProperty(...) events but do not modify it via the default generated customization.
 	DetailBuilder.HideProperty("DefaultValueVariant", UNiagaraScriptVariable::StaticClass());
+
+	// we hide the enum overrides if the variable type isn't an enum
+	if(!Variable->Variable.GetType().GetEnum())
+	{
+		DetailBuilder.HideProperty(GET_MEMBER_NAME_CHECKED(UNiagaraScriptVariable, Metadata.InlineParameterEnumOverrides));
+	}
+
+	// @TODO It seems like all the HideProperties aren't working for values within metadata
+	// generally we want all the inline parameters only available for inputs & static switches
+	if(!Variable->Variable.IsInNameSpace(FNiagaraConstants::ModuleNamespace) && !Variable->Variable.IsInNameSpace(FNiagaraConstants::StaticSwitchNamespace))
+	{
+		DetailBuilder.HideProperty(GET_MEMBER_NAME_CHECKED(UNiagaraScriptVariable, Metadata.bDisplayInOverviewStack));
+		DetailBuilder.HideProperty(GET_MEMBER_NAME_CHECKED(UNiagaraScriptVariable, Metadata.InlineParameterColorOverride));
+		DetailBuilder.HideProperty(GET_MEMBER_NAME_CHECKED(UNiagaraScriptVariable, Metadata.InlineParameterEnumOverrides));
+		DetailBuilder.HideProperty(GET_MEMBER_NAME_CHECKED(UNiagaraScriptVariable, Metadata.InlineParameterSortPriority));
+	}
 }
 
 void FNiagaraScriptVariableDetails::CustomizeDetailsGenericScriptVariable(IDetailLayoutBuilder& DetailBuilder)
diff --git a/Engine/Plugins/FX/Niagara/Source/NiagaraEditor/Private/ViewModels/Stack/NiagaraStackFunctionInputCollection.cpp b/Engine/Plugins/FX/Niagara/Source/NiagaraEditor/Private/ViewModels/Stack/NiagaraStackFunctionInputCollection.cpp
--- a/Engine/Plugins/FX/Niagara/Source/NiagaraEditor/Private/ViewModels/Stack/NiagaraStackFunctionInputCollection.cpp
+++ b/Engine/Plugins/FX/Niagara/Source/NiagaraEditor/Private/ViewModels/Stack/NiagaraStackFunctionInputCollection.cpp
@@ -145,6 +145,70 @@ void UNiagaraStackFunctionInputCollection::GetChildInputs(TArray<UNiagaraStackFu
 	}
 }
 
+void UNiagaraStackFunctionInputCollection::GetFilteredChildInputs(TArray<UNiagaraStackFunctionInput*>& OutResult) const
+{
+	TArray<UNiagaraStackInputCategory*> ChildCategories;
+	GetFilteredChildrenOfType(ChildCategories);
+	for (UNiagaraStackInputCategory* ChildCategory : ChildCategories)
+	{
+		ChildCategory->GetFilteredChildrenOfType(OutResult);
+	}
+}
+
+TOptional<FNiagaraLocalInputValueData> UNiagaraStackFunctionInputCollection::GetLocalInput(FNiagaraVariable InputParameter, bool bFiltered) const
+{
+	TArray<UNiagaraStackFunctionInput*> FunctionInputs;
+
+	if(bFiltered)
+	{
+		GetFilteredChildInputs(FunctionInputs);
+	}
+	else
+	{
+		GetChildInputs(FunctionInputs);
+	}
+
+	for(UNiagaraStackFunctionInput* FunctionInput : FunctionInputs)
+	{
+		if(FunctionInput->GetValueMode() == UNiagaraStackFunctionInput::EValueMode::Local && FunctionInput->GetInputParameterHandle().GetParameterHandleString() == InputParameter.GetName())
+		{
+			FNiagaraLocalInputValueData InputValueData(FunctionInput->GetLocalValueStruct(), FunctionInput->OnValueChanged());
+			FunctionInput->GetSearchItems(InputValueData.StackSearchItems);
+			return InputValueData;
+		}
+	}
+
+	TOptional<FNiagaraLocalInputValueData> UnsetValueData;
+	return UnsetValueData;
+}
+
+TOptional<FNiagaraDataInterfaceInput> UNiagaraStackFunctionInputCollection::GetDataInterfaceForInput(FNiagaraVariable InputParameter, bool bFiltered) const
+{
+	TArray<UNiagaraStackFunctionInput*> FunctionInputs;
+
+	if(bFiltered)
+	{
+		GetFilteredChildInputs(FunctionInputs);
+	}
+	else
+	{
+		GetChildInputs(FunctionInputs);
+	}
+
+	for(UNiagaraStackFunctionInput* FunctionInput : FunctionInputs)
+	{
+		if(FunctionInput->GetValueMode() == UNiagaraStackFunctionInput::EValueMode::Data && FunctionInput->GetInputParameterHandle().GetParameterHandleString() == InputParameter.GetName())
+		{
+			FNiagaraDataInterfaceInput InputValueData(FunctionInput->GetDataValueObject(), FunctionInput->OnValueChanged());
+			FunctionInput->GetSearchItems(InputValueData.StackSearchItems);
+			return TOptional<FNiagaraDataInterfaceInput>(InputValueData);
+		}
+	}
+
+	TOptional<FNiagaraDataInterfaceInput> UnsetValueData;
+	return UnsetValueData;
+}
+
 void UNiagaraStackFunctionInputCollection::RefreshChildrenInternal(const TArray<UNiagaraStackEntry*>& CurrentChildren, TArray<UNiagaraStackEntry*>& NewChildren, TArray<FStackIssue>& NewIssues)
 {
 	RefreshChildrenForFunctionCall(ModuleNode, InputFunctionCallNode, CurrentChildren, NewChildren, NewIssues, false, UncategorizedName);
diff --git a/Engine/Plugins/FX/Niagara/Source/NiagaraEditor/Private/ViewModels/Stack/NiagaraStackModuleItem.cpp b/Engine/Plugins/FX/Niagara/Source/NiagaraEditor/Private/ViewModels/Stack/NiagaraStackModuleItem.cpp
--- a/Engine/Plugins/FX/Niagara/Source/NiagaraEditor/Private/ViewModels/Stack/NiagaraStackModuleItem.cpp
+++ b/Engine/Plugins/FX/Niagara/Source/NiagaraEditor/Private/ViewModels/Stack/NiagaraStackModuleItem.cpp
@@ -1275,6 +1275,16 @@ void UNiagaraStackModuleItem::SetInputValuesFromClipboardFunctionInputs(const TA
 	InputCollection->SetValuesFromClipboardFunctionInputs(ClipboardFunctionInputs);
 }
 
+TOptional<FNiagaraLocalInputValueData> UNiagaraStackModuleItem::GetLocalInputData(FNiagaraVariable InputParameter, bool bFiltered) const
+{
+	return InputCollection->GetLocalInput(InputParameter, bFiltered);
+}
+
+TOptional<FNiagaraDataInterfaceInput> UNiagaraStackModuleItem::GetDataInterfaceForInput(FNiagaraVariable InputParameter, bool bFiltered) const
+{
+	return InputCollection->GetDataInterfaceForInput(InputParameter, bFiltered);	
+}
+
 bool UNiagaraStackModuleItem::TestCanCutWithMessage(FText& OutMessage) const
 {
 	FText CanCopyMessage;
diff --git a/Engine/Plugins/FX/Niagara/Source/NiagaraEditor/Private/ViewModels/Stack/NiagaraStackViewModel.cpp b/Engine/Plugins/FX/Niagara/Source/NiagaraEditor/Private/ViewModels/Stack/NiagaraStackViewModel.cpp
--- a/Engine/Plugins/FX/Niagara/Source/NiagaraEditor/Private/ViewModels/Stack/NiagaraStackViewModel.cpp
+++ b/Engine/Plugins/FX/Niagara/Source/NiagaraEditor/Private/ViewModels/Stack/NiagaraStackViewModel.cpp
@@ -212,6 +212,15 @@ void UNiagaraStackViewModel::ResetSearchText()
 	bRestartSearch = true;	
 }
 
+void UNiagaraStackViewModel::SetSearchTextExternal(const FText& NewSearchText)
+{
+	if(!CurrentSearchText.EqualTo(NewSearchText))
+	{
+		// this will forward a change in search text in a search box in a widget that registered itself, which in turn will call OnSearchTextChanged
+		OnChangeSearchTextExternalDelegate.ExecuteIfBound(NewSearchText);
+	}
+}
+
 void UNiagaraStackViewModel::OnSearchTextChanged(const FText& SearchText)
 {
 	if (RootEntry && !CurrentSearchText.EqualTo(SearchText))
@@ -484,6 +493,11 @@ TArray<UNiagaraStackEntry*>& UNiagaraStackViewModel::GetRootEntryAsArray()
 	return RootEntries;
 }
 
+UNiagaraStackViewModel::FOnChangeSearchTextExternal& UNiagaraStackViewModel::OnChangeSearchTextExternal()
+{
+	return OnChangeSearchTextExternalDelegate;
+}
+
 UNiagaraStackViewModel::FOnExpansionChanged& UNiagaraStackViewModel::OnExpansionChanged()
 {
 	return ExpansionChangedDelegate;
diff --git a/Engine/Plugins/FX/Niagara/Source/NiagaraEditor/Private/Widgets/SNiagaraParameterPanel.cpp b/Engine/Plugins/FX/Niagara/Source/NiagaraEditor/Private/Widgets/SNiagaraParameterPanel.cpp
--- a/Engine/Plugins/FX/Niagara/Source/NiagaraEditor/Private/Widgets/SNiagaraParameterPanel.cpp
+++ b/Engine/Plugins/FX/Niagara/Source/NiagaraEditor/Private/Widgets/SNiagaraParameterPanel.cpp
@@ -15,7 +15,7 @@
 #include "NiagaraScriptVariable.h"
 #include "NiagaraTypes.h"
 #include "SDropTarget.h"
-#include "SNiagaraPinTypeSelector.h"
+#include "Widgets/SNiagaraPinTypeSelector.h"
 #include "ViewModels/NiagaraParameterPanelViewModel.h"
 #include "ViewModels/Stack/NiagaraStackGraphUtilities.h"
 #include "Widgets/SToolTip.h"
diff --git a/Engine/Plugins/FX/Niagara/Source/NiagaraEditor/Private/Widgets/SNiagaraPinTypeSelector.cpp b/Engine/Plugins/FX/Niagara/Source/NiagaraEditor/Private/Widgets/SNiagaraPinTypeSelector.cpp
--- a/Engine/Plugins/FX/Niagara/Source/NiagaraEditor/Private/Widgets/SNiagaraPinTypeSelector.cpp
+++ b/Engine/Plugins/FX/Niagara/Source/NiagaraEditor/Private/Widgets/SNiagaraPinTypeSelector.cpp
@@ -1,6 +1,6 @@
 // Copyright Epic Games, Inc. All Rights Reserved.
 
-#include "SNiagaraPinTypeSelector.h"
+#include "Widgets/SNiagaraPinTypeSelector.h"
 
 #include "NiagaraEditorStyle.h"
 #include "Widgets/Images/SImage.h"
diff --git a/Engine/Plugins/FX/Niagara/Source/NiagaraEditor/Public/NiagaraGraph.h b/Engine/Plugins/FX/Niagara/Source/NiagaraEditor/Public/NiagaraGraph.h
--- a/Engine/Plugins/FX/Niagara/Source/NiagaraEditor/Public/NiagaraGraph.h
+++ b/Engine/Plugins/FX/Niagara/Source/NiagaraEditor/Public/NiagaraGraph.h
@@ -274,8 +274,8 @@ class UNiagaraGraph : public UEdGraph
 	/** Sets the meta-data associated with this variable. Creates a new UNiagaraScriptVariable if the target variable cannot be found. Illegal to call on FNiagaraVariables that are Niagara Constants. */
 	void SetMetaData(const FNiagaraVariable& InVar, const FNiagaraVariableMetaData& MetaData);
 
-	const TMap<FNiagaraVariable, TObjectPtr<UNiagaraScriptVariable>>& GetAllMetaData() const;
-	TMap<FNiagaraVariable, TObjectPtr<UNiagaraScriptVariable>>& GetAllMetaData();
+	NIAGARAEDITOR_API const TMap<FNiagaraVariable, TObjectPtr<UNiagaraScriptVariable>>& GetAllMetaData() const;
+	NIAGARAEDITOR_API TMap<FNiagaraVariable, TObjectPtr<UNiagaraScriptVariable>>& GetAllMetaData();
 
 	const TMap<FNiagaraVariable, FNiagaraGraphParameterReferenceCollection>& GetParameterReferenceMap() const; // NOTE: The const is a lie! (This indirectly calls RefreshParameterReferences, which can recreate the entire map)
 
diff --git a/Engine/Plugins/FX/Niagara/Source/NiagaraEditor/Public/NiagaraScriptVariable.h b/Engine/Plugins/FX/Niagara/Source/NiagaraEditor/Public/NiagaraScriptVariable.h
--- a/Engine/Plugins/FX/Niagara/Source/NiagaraEditor/Public/NiagaraScriptVariable.h
+++ b/Engine/Plugins/FX/Niagara/Source/NiagaraEditor/Public/NiagaraScriptVariable.h
@@ -11,7 +11,7 @@
 * Used to store variable data and metadata per graph. 
 */
 UCLASS()
-class UNiagaraScriptVariable : public UObject
+class NIAGARAEDITOR_API UNiagaraScriptVariable : public UObject
 {
 public:
 	GENERATED_UCLASS_BODY()
diff --git a/Engine/Plugins/FX/Niagara/Source/NiagaraEditor/Public/ViewModels/Stack/NiagaraStackEntry.h b/Engine/Plugins/FX/Niagara/Source/NiagaraEditor/Public/ViewModels/Stack/NiagaraStackEntry.h
--- a/Engine/Plugins/FX/Niagara/Source/NiagaraEditor/Public/ViewModels/Stack/NiagaraStackEntry.h
+++ b/Engine/Plugins/FX/Niagara/Source/NiagaraEditor/Public/ViewModels/Stack/NiagaraStackEntry.h
@@ -332,6 +332,21 @@ public:
 		}
 	}
 
+	template<typename T>
+	void GetFilteredChildrenOfType(TArray<T*>& OutFilteredChildrenOfType) const
+	{
+		TArray<UNiagaraStackEntry*> OutFilteredChildren;
+		GetFilteredChildren(OutFilteredChildren);
+		for (UNiagaraStackEntry* UnfilteredChild : OutFilteredChildren)
+		{
+			T* UnfilteredChildOfType = Cast<T>(UnfilteredChild);
+			if (UnfilteredChildOfType != nullptr)
+			{
+				OutFilteredChildrenOfType.Add(UnfilteredChildOfType);
+			}
+		}
+	}
+
 	FOnExpansionChanged& OnExpansionChanged();
 
 	FOnExpansionChanged& OnExpansionInOverviewChanged();
diff --git a/Engine/Plugins/FX/Niagara/Source/NiagaraEditor/Public/ViewModels/Stack/NiagaraStackFunctionInputCollection.h b/Engine/Plugins/FX/Niagara/Source/NiagaraEditor/Public/ViewModels/Stack/NiagaraStackFunctionInputCollection.h
--- a/Engine/Plugins/FX/Niagara/Source/NiagaraEditor/Public/ViewModels/Stack/NiagaraStackFunctionInputCollection.h
+++ b/Engine/Plugins/FX/Niagara/Source/NiagaraEditor/Public/ViewModels/Stack/NiagaraStackFunctionInputCollection.h
@@ -11,6 +11,26 @@ class UNiagaraStackFunctionInput;
 class UNiagaraClipboardFunctionInput;
 class UEdGraphPin;
 
+struct NIAGARAEDITOR_API FNiagaraLocalInputValueData
+{
+	FNiagaraLocalInputValueData(TSharedPtr<const FStructOnScope> InLocalData, TMulticastDelegate<void()>& InOnValueChanged) :
+	LocalData(InLocalData), OnValueChanged(InOnValueChanged)
+	{}
+	TSharedPtr<const FStructOnScope> LocalData;
+	TMulticastDelegate<void()>& OnValueChanged;
+	TArray<UNiagaraStackEntry::FStackSearchItem> StackSearchItems;
+};
+
+struct NIAGARAEDITOR_API FNiagaraDataInterfaceInput
+{
+	FNiagaraDataInterfaceInput(TWeakObjectPtr<const UNiagaraDataInterface> InDataInterface, TMulticastDelegate<void()>& InOnValueChanged) :
+	DataInterface(InDataInterface), OnValueChanged(InOnValueChanged)
+	{}
+	TWeakObjectPtr<const UNiagaraDataInterface> DataInterface;
+	TMulticastDelegate<void()>& OnValueChanged;
+	TArray<UNiagaraStackEntry::FStackSearchItem> StackSearchItems;
+};
+
 
 
 UCLASS()
@@ -99,6 +119,12 @@ public:
 
 	void GetChildInputs(TArray<UNiagaraStackFunctionInput*>& OutResult) const;
 
+	void GetFilteredChildInputs(TArray<UNiagaraStackFunctionInput*>& OutResult) const;
+
+	TOptional<FNiagaraLocalInputValueData> GetLocalInput(FNiagaraVariable InputParameter, bool bFiltered = false) const;
+
+	TOptional<FNiagaraDataInterfaceInput> GetDataInterfaceForInput(FNiagaraVariable InputParameter, bool bFiltered = false) const;
+
 protected:
 	virtual void FinalizeInternal() override;
 
diff --git a/Engine/Plugins/FX/Niagara/Source/NiagaraEditor/Public/ViewModels/Stack/NiagaraStackModuleItem.h b/Engine/Plugins/FX/Niagara/Source/NiagaraEditor/Public/ViewModels/Stack/NiagaraStackModuleItem.h
--- a/Engine/Plugins/FX/Niagara/Source/NiagaraEditor/Public/ViewModels/Stack/NiagaraStackModuleItem.h
+++ b/Engine/Plugins/FX/Niagara/Source/NiagaraEditor/Public/ViewModels/Stack/NiagaraStackModuleItem.h
@@ -74,6 +74,9 @@ public:
 
 	void SetInputValuesFromClipboardFunctionInputs(const TArray<const UNiagaraClipboardFunctionInput*>& ClipboardFunctionInputs);
 
+	TOptional<struct FNiagaraLocalInputValueData> GetLocalInputData(FNiagaraVariable InputParameter, bool bFiltered = false) const;
+	TOptional<struct FNiagaraDataInterfaceInput> GetDataInterfaceForInput(FNiagaraVariable InputParameter, bool bFiltered = false) const;
+
 	virtual bool SupportsCut() const override { return true; }
 	virtual bool TestCanCutWithMessage(FText& OutMessage) const override;
 	virtual FText GetCutTransactionText() const override;
diff --git a/Engine/Plugins/FX/Niagara/Source/NiagaraEditor/Public/ViewModels/Stack/NiagaraStackViewModel.h b/Engine/Plugins/FX/Niagara/Source/NiagaraEditor/Public/ViewModels/Stack/NiagaraStackViewModel.h
--- a/Engine/Plugins/FX/Niagara/Source/NiagaraEditor/Public/ViewModels/Stack/NiagaraStackViewModel.h
+++ b/Engine/Plugins/FX/Niagara/Source/NiagaraEditor/Public/ViewModels/Stack/NiagaraStackViewModel.h
@@ -40,6 +40,7 @@ class NIAGARAEDITOR_API UNiagaraStackViewModel : public UObject
 	GENERATED_BODY()
 
 public:
+	DECLARE_DELEGATE_OneParam(FOnChangeSearchTextExternal, FText)
 	DECLARE_MULTICAST_DELEGATE(FOnExpansionChanged);
 	DECLARE_MULTICAST_DELEGATE_OneParam(FOnStructureChanged, ENiagaraStructureChangedFlags);
 	DECLARE_MULTICAST_DELEGATE(FOnSearchCompleted);
@@ -88,6 +89,8 @@ public:
 
 	TArray<UNiagaraStackEntry*>& GetRootEntryAsArray();
 
+	FOnChangeSearchTextExternal& OnChangeSearchTextExternal();
+
 	FOnExpansionChanged& OnExpansionChanged();
 	FOnExpansionChanged& OnExpansionInOverviewChanged();
 	FOnStructureChanged& OnStructureChanged();
@@ -112,8 +115,9 @@ public:
 	virtual void Tick();
 	//~ stack search stuff
 	void ResetSearchText();
-	void OnSearchTextChanged(const FText& SearchText);
 	FText GetCurrentSearchText() const { return CurrentSearchText; };
+	void OnSearchTextChanged(const FText& SearchText);
+	void SetSearchTextExternal(const FText& NewSearchText);
 	bool IsSearching();
 	const TArray<FSearchResult>& GetCurrentSearchResults();
 	int GetCurrentFocusedMatchIndex() const { return CurrentFocusedSearchMatchIndex; }
@@ -185,6 +189,9 @@ private:
 
 	bool bExternalRootEntry;
 
+	/** Used to forward a programmatic change in search text to a widget representing the search. The widget is responsible for binding to this delegate. */
+	FOnChangeSearchTextExternal OnChangeSearchTextExternalDelegate;
+	
 	FOnExpansionChanged ExpansionChangedDelegate;
 	FOnExpansionChanged ExpansionInOverviewChangedDelegate;
 	FOnStructureChanged StructureChangedDelegate;
diff --git a/Engine/Plugins/FX/Niagara/Source/NiagaraEditorWidgets/Private/NiagaraEditorWidgetsStyle.cpp b/Engine/Plugins/FX/Niagara/Source/NiagaraEditorWidgets/Private/NiagaraEditorWidgetsStyle.cpp
--- a/Engine/Plugins/FX/Niagara/Source/NiagaraEditorWidgets/Private/NiagaraEditorWidgetsStyle.cpp
+++ b/Engine/Plugins/FX/Niagara/Source/NiagaraEditorWidgets/Private/NiagaraEditorWidgetsStyle.cpp
@@ -149,8 +149,17 @@ TSharedRef< FSlateStyleSet > FNiagaraEditorWidgetsStyle::Create()
 	FSlateFontInfo SystemOverviewAlternateItemFont = DEFAULT_FONT("Italic", 9);
 	FTextBlockStyle SystemOverviewAlternateItemText = FTextBlockStyle(NormalText)
 		.SetFont(SystemOverviewAlternateItemFont);
+
+	FSlateFontInfo SystemOverviewInlineParametersFont = DEFAULT_FONT("Regular", 7);
+	FTextBlockStyle SystemOverviewInlineParametersText = FTextBlockStyle(NormalText)
+		.SetFont(SystemOverviewInlineParametersFont);
+	
+	Style->Set("NiagaraEditor.SystemOverview.ItemText", SystemOverviewItemText);
+	
 	Style->Set("NiagaraEditor.SystemOverview.AlternateItemText", SystemOverviewAlternateItemText);
 
+	Style->Set("NiagaraEditor.SystemOverview.InlineParameterText", SystemOverviewInlineParametersText);
+
 	Style->Set("NiagaraEditor.SystemOverview.Item.BackgroundColor", FLinearColor(FColor(62, 62, 62)));
 	Style->Set("NiagaraEditor.SystemOverview.Group.BackgroundColor", FLinearColor::Transparent);
 	Style->Set("NiagaraEditor.SystemOverview.CheckBoxColor", FLinearColor(FColor(160, 160, 160)));
@@ -307,7 +316,7 @@ TSharedRef< FSlateStyleSet > FNiagaraEditorWidgetsStyle::Create()
 
 	Style->Set("NiagaraEditor.Stack.ModuleHighlight", new IMAGE_PLUGIN_BRUSH("Icons/ModuleHighlight", Icon6x6, FLinearColor::White));
 	Style->Set("NiagaraEditor.Stack.ModuleHighlightMore", new IMAGE_PLUGIN_BRUSH("Icons/ModuleHighlightMore", Icon6x6, FLinearColor::White));
-	Style->Set("NiagaraEditor.Stack.ModuleHighlightLarge", new IMAGE_PLUGIN_BRUSH("Icons/ModuleHighlightLarge", Icon8x8, FLinearColor::White));
+	Style->Set("NiagaraEditor.Stack.ModuleHighlightLarge", new IMAGE_PLUGIN_BRUSH("Icons/ModuleHighlightLarge", Icon8x8, FLinearColor::White));	
 
 	// Scratch pad
 	FSlateFontInfo ScratchPadEditorHeaderFont = DEFAULT_FONT("Bold", 11);
diff --git a/Engine/Plugins/FX/Niagara/Source/NiagaraEditorWidgets/Private/SNiagaraOverviewInlineParameterBox.cpp b/Engine/Plugins/FX/Niagara/Source/NiagaraEditorWidgets/Private/SNiagaraOverviewInlineParameterBox.cpp
new file mode 100644
--- /dev/null
+++ b/Engine/Plugins/FX/Niagara/Source/NiagaraEditorWidgets/Private/SNiagaraOverviewInlineParameterBox.cpp
@@ -0,0 +1,292 @@
+// Copyright Epic Games, Inc. All Rights Reserved.
+
+#include "SNiagaraOverviewInlineParameterBox.h"
+
+#include "NiagaraEditorWidgetsStyle.h"
+#include "NiagaraNodeFunctionCall.h"
+#include "NiagaraScriptSource.h"
+#include "NiagaraScriptVariable.h"
+#include "NiagaraSystem.h"
+#include "NiagaraComponent.h"
+#include "NiagaraStackEditorData.h"
+#include "Styling/StyleColors.h"
+#include "ViewModels/NiagaraEmitterViewModel.h"
+#include "ViewModels/NiagaraSystemSelectionViewModel.h"
+#include "ViewModels/NiagaraSystemViewModel.h"
+#include "Widgets/Layout/SWrapBox.h"
+#include "Widgets/Text/STextBlock.h"
+#include "ViewModels/Stack/NiagaraStackFunctionInput.h"
+#include "ViewModels/Stack/NiagaraStackFunctionInputCollection.h"
+#include "Widgets/SNiagaraParameterName.h"
+#include "Widgets/SNiagaraPinTypeSelector.h"
+#include "Widgets/Layout/SScrollBox.h"
+
+#define LOCTEXT_NAMESPACE "SNiagaraOverviewInlineParameterBox"
+
+void SNiagaraOverviewInlineParameterBox::Construct(const FArguments& InArgs, UNiagaraStackModuleItem& InStackModuleItem)
+{
+	ModuleItem = &InStackModuleItem;
+
+	ChildSlot
+	.HAlign(HAlign_Left)
+	.VAlign(VAlign_Center)
+	[
+		SAssignNew(Container, SScrollBox)
+		.Orientation(EOrientation::Orient_Horizontal)
+		.ScrollBarThickness(FVector2D(2.f, 2.f))
+		.ScrollBarVisibility(EVisibility::Collapsed)
+		.ConsumeMouseWheel(EConsumeMouseWheel::Always)
+	];
+
+	ConstructChildren();
+}
+
+SNiagaraOverviewInlineParameterBox::~SNiagaraOverviewInlineParameterBox()
+{
+	for(FNiagaraLocalInputValueData& BoundInput : BoundFunctionInputs)
+	{
+		BoundInput.OnValueChanged.RemoveAll(this);
+	}
+}
+
+FReply SNiagaraOverviewInlineParameterBox::NavigateToStack(FNiagaraLocalInputValueData LocalInputValueData)
+{
+	// even if we can't navigate, the button should consume the input
+	if(!ModuleItem.IsValid())
+	{
+		return FReply::Handled();
+	}
+	
+	ModuleItem->GetSystemViewModel()->GetSelectionViewModel()->UpdateSelectedEntries({ModuleItem.Get()}, {}, true);
+	
+	if(ensure(LocalInputValueData.StackSearchItems.Num() > 0))
+	{
+		ModuleItem->GetSystemViewModel()->GetSelectionViewModel()->GetSelectionStackViewModel()->SetSearchTextExternal(LocalInputValueData.StackSearchItems[0].Value);
+		return FReply::Handled();
+	}
+
+	return FReply::Handled();
+}
+
+void SNiagaraOverviewInlineParameterBox::ConstructChildren()
+{
+	Container->ClearChildren();
+
+	TArray<TSharedRef<SWidget>> ParameterWidgets = GenerateParameterWidgets();
+	
+	for (TSharedRef<SWidget> ParameterWidget : ParameterWidgets)
+	{
+		Container->AddSlot()
+		.VAlign(VAlign_Center)
+		.HAlign(HAlign_Right)
+		.Padding(FMargin(2.f, 1.f))
+		[
+			ParameterWidget
+		];
+	}
+
+	Container->SetVisibility(ParameterWidgets.Num() > 0 ? EVisibility::Visible : EVisibility::Collapsed);
+}
+
+TArray<TSharedRef<SWidget>> SNiagaraOverviewInlineParameterBox::GenerateParameterWidgets()
+{
+	TArray<TSharedRef<SWidget>> ParameterWidgets;
+
+	// the stack module item should be valid at this point, but we check to make sure
+	if(!ModuleItem.IsValid())
+	{
+		return ParameterWidgets;
+	}
+	
+	UNiagaraScriptSource* ScriptSource = CastChecked<UNiagaraScriptSource>(ModuleItem->GetModuleNode().GetFunctionScriptSource());
+	// in case the script source is no longer valid (i.e. the asset was deleted), we return immediately
+	if(ScriptSource == nullptr)
+	{
+		return ParameterWidgets;
+	}
+
+	// we cache the sort order of the widgets we create to sort them after all widgets have been created
+	TMap<TSharedRef<SWidget>, int32> SortOrder;
+
+	// we go through all script variables of the source graph and check if a parameter is supposed to be displayed in the overview
+	const TMap<FNiagaraVariable, TObjectPtr<UNiagaraScriptVariable>>& AllMetaData = ScriptSource->NodeGraph->GetAllMetaData();
+	for(const auto& ParameterMetaData : AllMetaData)
+	{
+		const UNiagaraScriptVariable* ScriptVariable = ParameterMetaData.Value;
+
+		if(ScriptVariable == nullptr || !ScriptVariable->Metadata.bDisplayInOverviewStack)
+		{
+			continue;
+		}
+
+		TOptional<FNiagaraLocalInputValueData> LocalInputValueData = ModuleItem->GetLocalInputData(ParameterMetaData.Key, true);
+		TOptional<FNiagaraDataInterfaceInput> InputDataInterface = ModuleItem->GetDataInterfaceForInput(ParameterMetaData.Key, true);
+
+		// we only support inline display of local values and data interfaces currently
+		// @TODO use the commented out line when DIs support this too.
+		// if(!LocalInputValueData.IsSet() && !InputDataInterface.IsSet())
+		if(!LocalInputValueData.IsSet())
+		{
+			continue;
+		}
+
+		TSharedPtr<SWidget> Widget;
+		if(LocalInputValueData.IsSet())
+		{
+			Widget = GenerateParameterWidgetFromLocalValue(ScriptVariable, LocalInputValueData.GetValue());
+		}
+		else if(InputDataInterface.IsSet() && InputDataInterface->DataInterface.IsValid())
+		{
+			Widget = GenerateParameterWidgetFromDataInterface(ScriptVariable, InputDataInterface.GetValue());
+		}		
+
+		SortOrder.Add(Widget.ToSharedRef(), ScriptVariable->Metadata.InlineParameterSortPriority);
+		ParameterWidgets.Add(Widget.ToSharedRef());
+	}
+	
+
+	ParameterWidgets.Sort([&](TSharedRef<SWidget> WidgetA, TSharedRef<SWidget> WidgetB)
+	{
+		return SortOrder[WidgetA] < SortOrder[WidgetB];
+	});
+	
+	return ParameterWidgets;
+}
+
+TSharedRef<SWidget> SNiagaraOverviewInlineParameterBox::GenerateParameterWidgetFromLocalValue(const UNiagaraScriptVariable* ScriptVariable, FNiagaraLocalInputValueData& LocalInputValueData)
+{
+	LocalInputValueData.OnValueChanged.RemoveAll(this);
+	LocalInputValueData.OnValueChanged.AddSP(this, &SNiagaraOverviewInlineParameterBox::ConstructChildren);
+	BoundFunctionInputs.Add(LocalInputValueData);
+	FNiagaraTypeDefinition Type = ScriptVariable->Variable.GetType();
+	FString Value = Type.ToStringCosmetic(LocalInputValueData.LocalData->GetStructMemory());
+	const FLinearColor TypeColor = UEdGraphSchema_Niagara::GetTypeColor(Type);
+
+	UTexture2D* Icon = nullptr;
+	if (const UEnum* Enum = Type.GetEnum())
+	{
+		int32 EnumIndex = Enum->GetIndexByValue(*(int32*)LocalInputValueData.LocalData->GetStructMemory());
+		if(ScriptVariable->Metadata.InlineParameterEnumOverrides.IsValidIndex(EnumIndex))
+		{
+			Value = ScriptVariable->Metadata.InlineParameterEnumOverrides[EnumIndex].OverrideName.ToString();
+			Icon = ScriptVariable->Metadata.InlineParameterEnumOverrides[EnumIndex].IconOverride;
+		}
+	}
+
+	TSharedPtr<SWidget> ParameterWidget;
+
+	TSharedRef<SButton> ParameterWidgetButton = SNew(SButton)
+	.Text(FText::FromString(Value))
+	.ContentPadding(FMargin(-5, 0))
+	.TextStyle(&FNiagaraEditorWidgetsStyle::Get().GetWidgetStyle<FTextBlockStyle>("NiagaraEditor.SystemOverview.InlineParameterText"))
+	.ButtonStyle(&FAppStyle::Get().GetWidgetStyle<FButtonStyle>("Button"))
+	.OnClicked(this, &SNiagaraOverviewInlineParameterBox::NavigateToStack, LocalInputValueData);
+
+	if(Icon != nullptr)
+	{
+		FSlateBrush& ImageBrush = ImageBrushes.Emplace_GetRef();
+		ImageBrush.SetResourceObject(Icon);
+		ImageBrush.ImageSize.X = 16;
+		ImageBrush.ImageSize.Y = 16;
+	
+		TSharedRef<SImage> ParameterWidgetIcon= SNew(SImage).Image(&ImageBrush);
+		ParameterWidgetButton->SetContent(ParameterWidgetIcon);
+		//ParameterWidgetButton->SetBorderBackgroundColor(FStyleColors::Transparent);
+		ParameterWidgetButton->SetButtonStyle(&FAppStyle::Get().GetWidgetStyle<FButtonStyle>("SimpleButton"));
+	}
+	else
+	{
+		FLinearColor BorderColor = ScriptVariable->Metadata.bOverrideColor ? ScriptVariable->Metadata.InlineParameterColorOverride : TypeColor;
+		ParameterWidgetButton->SetBorderBackgroundColor(BorderColor);				
+	}			
+
+	ParameterWidget = ParameterWidgetButton;
+
+	FText			   IconToolTip = FText::GetEmpty();
+	FSlateBrush const* IconBrush = FAppStyle::Get().GetBrush(TEXT("Kismet.AllClasses.VariableIcon"));
+	FSlateColor        IconColor = FSlateColor(TypeColor);
+	FString			   IconDocLink, IconDocExcerpt;
+	FSlateBrush const* SecondaryIconBrush = FEditorStyle::GetBrush(TEXT("NoBrush"));
+	FSlateColor        SecondaryIconColor = IconColor;
+
+	// we construct a tooltip widget that shows the parameter the value is associated with
+	TSharedRef<SToolTip> TooltipWidget = SNew(SToolTip).Content()
+	[
+		SNew(SHorizontalBox)
+		+ SHorizontalBox::Slot()
+		.AutoWidth()
+		.HAlign(HAlign_Center)
+		.VAlign(VAlign_Center)
+		.Padding(3.f)
+		[
+			SNew(SNiagaraParameterNameTextBlock)
+			.IsReadOnly(true)
+			.ParameterText(LOCTEXT("ParameterTooltipText", "Value for"))
+		]
+		+ SHorizontalBox::Slot()
+		.AutoWidth()
+		.HAlign(HAlign_Center)
+		.VAlign(VAlign_Center)
+		[
+			SNew(SNiagaraIconWidget)
+			.IconToolTip(IconToolTip)
+			.IconBrush(IconBrush)
+			.IconColor(IconColor)
+			.DocLink(IconDocLink)
+			.DocExcerpt(IconDocExcerpt)
+			.SecondaryIconBrush(SecondaryIconBrush) 
+			.SecondaryIconColor(SecondaryIconColor)
+		]
+		+ SHorizontalBox::Slot()
+		.AutoWidth()
+		.HAlign(HAlign_Center)
+		.VAlign(VAlign_Center)
+		.Padding(3.f)
+		[
+			SNew(SNiagaraParameterNameTextBlock)
+			.IsReadOnly(true)
+			.ParameterText(FText::FromName(ScriptVariable->Variable.GetName()))
+		]
+	];
+
+	ParameterWidget->SetToolTip(TooltipWidget);
+	return ParameterWidget.ToSharedRef();
+}
+
+TSharedRef<SWidget> SNiagaraOverviewInlineParameterBox::GenerateParameterWidgetFromDataInterface(const UNiagaraScriptVariable* ScriptVariable, FNiagaraDataInterfaceInput& DataInterfaceInput)
+{
+	if(DataInterfaceInput.DataInterface.IsValid())
+	{
+		DataInterfaceInput.OnValueChanged.RemoveAll(this);
+		DataInterfaceInput.OnValueChanged.AddSP(this, &SNiagaraOverviewInlineParameterBox::ConstructChildren);
+		TArray<FNiagaraDataInterfaceError> Errors;
+		TArray<FNiagaraDataInterfaceFeedback> Warnings;
+		TArray<FNiagaraDataInterfaceFeedback> Infos;
+
+		// these can be nullptr depending on context
+		UNiagaraSystem* System = DataInterfaceInput.DataInterface->GetTypedOuter<UNiagaraSystem>();
+		UNiagaraComponent* Component = DataInterfaceInput.DataInterface->GetTypedOuter<UNiagaraComponent>();
+
+		//@TODO Finish this up.
+		//DataInterfaceInput.DataInterface->GetInlineFeedback(System, Component, Errors, Warnings, Infos);
+
+		TSharedPtr<SHorizontalBox> InfoBox = SNew(SHorizontalBox);
+		for(FNiagaraDataInterfaceFeedback& Feedback : Infos)
+		{
+			InfoBox->AddSlot()
+			.AutoWidth()
+			[
+				SNew(STextBlock)
+				.Text(Feedback.GetFeedbackSummaryText())
+				.ToolTipText(Feedback.GetFeedbackText())
+				.TextStyle(&FNiagaraEditorWidgetsStyle::Get().GetWidgetStyle<FTextBlockStyle>("NiagaraEditor.SystemOverview.InlineParameterText"))
+			];			
+		}
+
+		return InfoBox.ToSharedRef();
+	}
+
+	return SNullWidget::NullWidget;
+}
+
+#undef LOCTEXT_NAMESPACE
diff --git a/Engine/Plugins/FX/Niagara/Source/NiagaraEditorWidgets/Private/SNiagaraOverviewInlineParameterBox.h b/Engine/Plugins/FX/Niagara/Source/NiagaraEditorWidgets/Private/SNiagaraOverviewInlineParameterBox.h
new file mode 100644
--- /dev/null
+++ b/Engine/Plugins/FX/Niagara/Source/NiagaraEditorWidgets/Private/SNiagaraOverviewInlineParameterBox.h
@@ -0,0 +1,52 @@
+// Copyright Epic Games, Inc. All Rights Reserved.
+
+#pragma once
+
+#include "Brushes/SlateImageBrush.h"
+#include "Widgets/SCompoundWidget.h"
+#include "Widgets/DeclarativeSyntaxSupport.h"
+#include "ViewModels/Stack/NiagaraStackViewModel.h"
+#include "ViewModels/Stack/NiagaraStackFunctionInput.h"
+#include "ViewModels/Stack/NiagaraStackModuleItem.h"
+#include "ViewModels/Stack/NiagaraStackFunctionInputCollection.h"
+
+class SWrapBox;
+class SScrollBox;
+class UNiagaraStackViewModel;
+class UNiagaraStackEntry;
+class SBox;
+class SInlineEditableTextBlock;
+
+class SNiagaraOverviewInlineParameterBox : public SCompoundWidget
+{
+public:
+	SLATE_BEGIN_ARGS(SNiagaraOverviewInlineParameterBox)
+	{
+		_Clipping = EWidgetClipping::OnDemand;
+	}
+	
+	SLATE_END_ARGS();
+
+	void Construct(const FArguments& InArgs, UNiagaraStackModuleItem& InStackModuleItem);
+
+	virtual ~SNiagaraOverviewInlineParameterBox() override;
+
+private:
+	/** Selects the module item in the stack and searches for the function input */
+	FReply NavigateToStack(FNiagaraLocalInputValueData LocalInputValueData);
+	void ConstructChildren();
+	TArray<TSharedRef<SWidget>> GenerateParameterWidgets();
+	TSharedRef<SWidget> GenerateParameterWidgetFromLocalValue(const UNiagaraScriptVariable* ScriptVariable, FNiagaraLocalInputValueData& LocalInputValueData);
+	TSharedRef<SWidget> GenerateParameterWidgetFromDataInterface(const UNiagaraScriptVariable* ScriptVariable, FNiagaraDataInterfaceInput& DataInterface);
+
+private:
+	/** The module item whose parameters this box is representing */
+	TWeakObjectPtr<UNiagaraStackModuleItem> ModuleItem;
+	
+	TSharedPtr<SScrollBox> Container;
+	
+	/** We keep track of the function inputs we are observing through this parameter box so we can unbind the delegates later */
+	TArray<FNiagaraLocalInputValueData> BoundFunctionInputs;
+
+	TArray<FSlateBrush> ImageBrushes;
+};
diff --git a/Engine/Plugins/FX/Niagara/Source/NiagaraEditorWidgets/Private/SNiagaraStack.cpp b/Engine/Plugins/FX/Niagara/Source/NiagaraEditorWidgets/Private/SNiagaraStack.cpp
--- a/Engine/Plugins/FX/Niagara/Source/NiagaraEditorWidgets/Private/SNiagaraStack.cpp
+++ b/Engine/Plugins/FX/Niagara/Source/NiagaraEditorWidgets/Private/SNiagaraStack.cpp
@@ -268,6 +268,7 @@ const float SpacerHeight = 6;
 void SNiagaraStack::Construct(const FArguments& InArgs, UNiagaraStackViewModel* InStackViewModel)
 {
 	StackViewModel = InStackViewModel;
+	StackViewModel->OnChangeSearchTextExternal().BindSP(this, &SNiagaraStack::UpdateSearchTextFromExternal);
 	StackViewModel->OnStructureChanged().AddSP(this, &SNiagaraStack::StackStructureChanged);
 	StackViewModel->OnExpansionChanged().AddSP(this, &SNiagaraStack::OnStackExpansionChanged);
 	StackViewModel->OnSearchCompleted().AddSP(this, &SNiagaraStack::OnStackSearchComplete); 
@@ -386,6 +387,11 @@ TSharedRef<SWidget> SNiagaraStack::ConstructHeaderWidget()
 		];
 }
 
+void SNiagaraStack::UpdateSearchTextFromExternal(FText NewSearchText)
+{
+	SearchBox->SetText(NewSearchText);
+}
+
 void SNiagaraStack::OnSearchTextChanged(const FText& SearchText)
 {
 	StackViewModel->OnSearchTextChanged(SearchText);
diff --git a/Engine/Plugins/FX/Niagara/Source/NiagaraEditorWidgets/Private/SNiagaraStack.h b/Engine/Plugins/FX/Niagara/Source/NiagaraEditorWidgets/Private/SNiagaraStack.h
--- a/Engine/Plugins/FX/Niagara/Source/NiagaraEditorWidgets/Private/SNiagaraStack.h
+++ b/Engine/Plugins/FX/Niagara/Source/NiagaraEditorWidgets/Private/SNiagaraStack.h
@@ -86,6 +86,7 @@ private:
 	TSharedRef<SWidget> ConstructHeaderWidget();
 
 	// ~stack search stuff
+	void UpdateSearchTextFromExternal(FText NewSearchText);
 	void OnSearchTextChanged(const FText& SearchText);
 	FReply ScrollToNextMatch();
 	FReply ScrollToPreviousMatch();
