Fixed lack of initializationof new texture ptr

diff --git a/Engine/Plugins/FX/Niagara/Source/Niagara/Public/NiagaraTypes.h b/Engine/Plugins/FX/Niagara/Source/Niagara/Public/NiagaraTypes.h
--- a/Engine/Plugins/FX/Niagara/Source/Niagara/Public/NiagaraTypes.h
+++ b/Engine/Plugins/FX/Niagara/Source/Niagara/Public/NiagaraTypes.h
@@ -624,7 +624,7 @@ struct NIAGARA_API FNiagaraEnumParameterMetaData
 
 	/** If specified, this icon will be used for the given enum entry. If OverrideName isn't empty, the icon takes priority. */
 	UPROPERTY(EditAnywhere, Category="Enum Override")
-	UTexture2D* IconOverride;
+	UTexture2D* IconOverride = nullptr;
 };
 
 USTRUCT()
