Fixed missing include for non-unity build

diff --git a/Engine/Plugins/FX/Niagara/Source/NiagaraEditorWidgets/Private/SNiagaraOverviewInlineParameterBox.cpp b/Engine/Plugins/FX/Niagara/Source/NiagaraEditorWidgets/Private/SNiagaraOverviewInlineParameterBox.cpp
--- a/Engine/Plugins/FX/Niagara/Source/NiagaraEditorWidgets/Private/SNiagaraOverviewInlineParameterBox.cpp
+++ b/Engine/Plugins/FX/Niagara/Source/NiagaraEditorWidgets/Private/SNiagaraOverviewInlineParameterBox.cpp
@@ -15,6 +15,7 @@
 #include "ViewModels/NiagaraSystemViewModel.h"
 #include "Widgets/Layout/SWrapBox.h"
 #include "Widgets/Text/STextBlock.h"
+#include "Widgets/Input/SButton.h"
 #include "ViewModels/Stack/NiagaraStackFunctionInput.h"
 #include "ViewModels/Stack/NiagaraStackFunctionInputCollection.h"
 #include "Widgets/SNiagaraParameterName.h"
