- Disable some log warnings in shipping unless we enable verbose

#rnx

diff --git a/Engine/Plugins/FX/Niagara/Source/Niagara/Private/NiagaraComponent.cpp b/Engine/Plugins/FX/Niagara/Source/Niagara/Private/NiagaraComponent.cpp
--- a/Engine/Plugins/FX/Niagara/Source/Niagara/Private/NiagaraComponent.cpp
+++ b/Engine/Plugins/FX/Niagara/Source/Niagara/Private/NiagaraComponent.cpp
@@ -591,7 +591,7 @@ void UNiagaraComponent::ReleaseToPool()
 	if (PoolingMethod != ENCPoolMethod::ManualRelease)
 	{		
 		// Only emit this warning if pooling is enabled. If it's not, all components will have PoolingMethod none.
-		if (UNiagaraComponentPool::Enabled())
+		if (UNiagaraComponentPool::Enabled() && FNiagaraUtilities::LogVerboseWarnings())
 		{	
 			UE_LOG(LogNiagara, Warning, TEXT("ReleaseToPool called but PoolingMethod is (%s) when it should be ENCPoolMethod::ManualRelease. Asset=%s Component=%s"), *StaticEnum<ENCPoolMethod>()->GetNameStringByValue(int64(PoolingMethod)), *GetPathNameSafe(Asset), *GetPathName());
 		}
@@ -1622,7 +1622,10 @@ void UNiagaraComponent::OnComponentDestroyed(bool bDestroyingHierarchy)
 	{
 		if (UWorld* World = GetWorld())
 		{
-			UE_LOG(LogNiagara, Warning, TEXT("UNiagaraComponent::OnComponentDestroyed: Component (%p - %s) Asset (%s) is still pooled (%d) while destroying!"), this, *GetFullNameSafe(this), *GetFullNameSafe(Asset), PoolingMethod);
+			if (FNiagaraUtilities::LogVerboseWarnings())
+			{
+				UE_LOG(LogNiagara, Warning, TEXT("UNiagaraComponent::OnComponentDestroyed: Component (%p - %s) Asset (%s) is still pooled (%d) while destroying!"), this, *GetFullNameSafe(this), *GetFullNameSafe(Asset), PoolingMethod);
+			}
 			if (FNiagaraWorldManager* WorldManager = FNiagaraWorldManager::Get(World))
 			{
 				if (UNiagaraComponentPool* ComponentPool = WorldManager->GetComponentPool())
@@ -1633,7 +1636,10 @@ void UNiagaraComponent::OnComponentDestroyed(bool bDestroyingHierarchy)
 		}
 		else
 		{
-			UE_LOG(LogNiagara, Warning, TEXT("UNiagaraComponent::OnComponentDestroyed: Component (%p - %s) Asset (%s) is still pooled (%d) while destroying and world it nullptr!"), this, *GetFullNameSafe(this), *GetFullNameSafe(Asset), PoolingMethod);
+			if (FNiagaraUtilities::LogVerboseWarnings())
+			{
+				UE_LOG(LogNiagara, Warning, TEXT("UNiagaraComponent::OnComponentDestroyed: Component (%p - %s) Asset (%s) is still pooled (%d) while destroying and world it nullptr!"), this, *GetFullNameSafe(this), *GetFullNameSafe(Asset), PoolingMethod);
+			}
 		}
 
 		// Set pooling method to none as we are destroyed and can not go into the pool after this point
@@ -1694,7 +1700,10 @@ void UNiagaraComponent::BeginDestroy()
 	{
 		if (UWorld* World = GetWorld())
 		{
-			UE_LOG(LogNiagara, Warning, TEXT("UNiagaraComponent::BeginDestroy: Component (%p - %s) Asset (%s) is still pooled (%d) while destroying!"), this, *GetFullNameSafe(this), *GetFullNameSafe(Asset), PoolingMethod);
+			if (FNiagaraUtilities::LogVerboseWarnings())
+			{
+				UE_LOG(LogNiagara, Warning, TEXT("UNiagaraComponent::BeginDestroy: Component (%p - %s) Asset (%s) is still pooled (%d) while destroying!"), this, *GetFullNameSafe(this), *GetFullNameSafe(Asset), PoolingMethod);
+			}
 
 			if (FNiagaraWorldManager* WorldManager = FNiagaraWorldManager::Get(World))
 			{
@@ -1706,7 +1715,10 @@ void UNiagaraComponent::BeginDestroy()
 		}
 		else
 		{
-			UE_LOG(LogNiagara, Warning, TEXT("UNiagaraComponent::BeginDestroy: Component (%p - %s) Asset (%s) is still pooled (%d) while destroying and world is nullptr!"), this, *GetFullNameSafe(this), *GetFullNameSafe(Asset), PoolingMethod);
+			if (FNiagaraUtilities::LogVerboseWarnings())
+			{
+				UE_LOG(LogNiagara, Warning, TEXT("UNiagaraComponent::BeginDestroy: Component (%p - %s) Asset (%s) is still pooled (%d) while destroying and world is nullptr!"), this, *GetFullNameSafe(this), *GetFullNameSafe(Asset), PoolingMethod);
+			}
 		}
 
 		// Set pooling method to none as we are destroyed and can not go into the pool after this point
