- Remove ViewTargetChangedHandle from server

#rnx

diff --git a/Engine/Plugins/FX/Niagara/Source/Niagara/Private/NiagaraWorldManager.cpp b/Engine/Plugins/FX/Niagara/Source/Niagara/Private/NiagaraWorldManager.cpp
--- a/Engine/Plugins/FX/Niagara/Source/Niagara/Private/NiagaraWorldManager.cpp
+++ b/Engine/Plugins/FX/Niagara/Source/Niagara/Private/NiagaraWorldManager.cpp
@@ -316,12 +316,15 @@ void FNiagaraWorldManager::OnStartup()
 	PostGCHandle = FCoreUObjectDelegates::GetPostGarbageCollect().AddStatic(&FNiagaraWorldManager::OnPostGarbageCollect);
 	PreGCBeginDestroyHandle = FCoreUObjectDelegates::PreGarbageCollectConditionalBeginDestroy.AddStatic(&FNiagaraWorldManager::OnPreGarbageCollectBeginDestroy);
 
-	ViewTargetChangedHandle = FGameDelegates::Get().GetViewTargetChangedDelegate().AddLambda(
-		[](APlayerController* PC, AActor* OldTarget, AActor* NewTarget)
-		{
-			OnRefreshOwnerAllowsScalability();
-		}
-	);
+	if ( FApp::CanEverRender() )
+	{
+		ViewTargetChangedHandle = FGameDelegates::Get().GetViewTargetChangedDelegate().AddLambda(
+			[](APlayerController* PC, AActor* OldTarget, AActor* NewTarget)
+			{
+				OnRefreshOwnerAllowsScalability();
+			}
+		);
+	}
 }
 
 void FNiagaraWorldManager::OnShutdown()
@@ -339,7 +342,10 @@ void FNiagaraWorldManager::OnShutdown()
 	FCoreUObjectDelegates::GetPostGarbageCollect().Remove(PostGCHandle);
 	FCoreUObjectDelegates::PreGarbageCollectConditionalBeginDestroy.Remove(PreGCBeginDestroyHandle);
 	
-	FGameDelegates::Get().GetViewTargetChangedDelegate().Remove(ViewTargetChangedHandle);
+	if ( ViewTargetChangedHandle.IsValid() )
+	{
+		FGameDelegates::Get().GetViewTargetChangedDelegate().Remove(ViewTargetChangedHandle);
+	}
 
 	//Should have cleared up all world managers by now.
 	check(WorldManagers.Num() == 0);
