- Add global override for render target format

#rnx

diff --git a/Engine/Plugins/FX/Niagara/Source/Niagara/Private/NiagaraDataInterfaceRenderTarget2D.cpp b/Engine/Plugins/FX/Niagara/Source/Niagara/Private/NiagaraDataInterfaceRenderTarget2D.cpp
--- a/Engine/Plugins/FX/Niagara/Source/Niagara/Private/NiagaraDataInterfaceRenderTarget2D.cpp
+++ b/Engine/Plugins/FX/Niagara/Source/Niagara/Private/NiagaraDataInterfaceRenderTarget2D.cpp
@@ -71,6 +71,48 @@ static FAutoConsoleVariableRef CVarNiagaraRenderTargetAllowReads(
 
 /*--------------------------------------------------------------------------------------------------------------------------*/
 
+static bool GNiagaraRenderTargetOverrideFormatEnabled = false;
+static ETextureRenderTargetFormat GNiagaraRenderTargetOverrideFormat = ETextureRenderTargetFormat::RTF_RGBA32f;
+
+static FAutoConsoleCommandWithWorldAndArgs GCommandNiagaraRenderTargetOverrideFormat(
+	TEXT("fx.Niagara.RenderTarget.OverrideFormat"),
+	TEXT("Optional global format override for all Niagara render targets"),
+	FConsoleCommandWithWorldAndArgsDelegate::CreateLambda(
+		[](const TArray<FString>& Args, UWorld*)
+		{
+			static UEnum* TextureRenderTargetFormatEnum = StaticEnum<ETextureRenderTargetFormat>();
+			if ( ensure(TextureRenderTargetFormatEnum) )
+			{
+				if (Args.Num() == 1)
+				{
+					const int32 EnumIndex = TextureRenderTargetFormatEnum->GetIndexByNameString(Args[0]);
+					if (EnumIndex != INDEX_NONE)
+					{
+						GNiagaraRenderTargetOverrideFormatEnabled = true;
+						GNiagaraRenderTargetOverrideFormat = ETextureRenderTargetFormat(EnumIndex);
+					}
+					else
+					{
+						GNiagaraRenderTargetOverrideFormatEnabled = false;
+					}
+				}
+				UE_LOG(LogNiagara, Log, TEXT("Niagara RenderTarget Override is '%s' with format '%s'."), GNiagaraRenderTargetOverrideFormatEnabled ? TEXT("Enabled") : TEXT("Disabled"), *TextureRenderTargetFormatEnum->GetNameStringByIndex(int32(GNiagaraRenderTargetOverrideFormat)));
+			}
+		}
+	)
+);
+
+ETextureRenderTargetFormat GetRenderTargetFormat(bool bOverrideFormat, ETextureRenderTargetFormat OverrideFormat)
+{
+	if (GNiagaraRenderTargetOverrideFormatEnabled)
+	{
+		return GNiagaraRenderTargetOverrideFormat;
+	}
+	return bOverrideFormat ? OverrideFormat : GetDefault<UNiagaraSettings>()->DefaultRenderTargetFormat.GetValue();
+}
+
+/*--------------------------------------------------------------------------------------------------------------------------*/
+
 struct FNDIRenderTarget2DFunctionVersion
 {
 	enum Type
@@ -554,7 +596,7 @@ bool UNiagaraDataInterfaceRenderTarget2D::InitPerInstanceData(void* PerInstanceD
 	InstanceData->Size.Y = FMath::Clamp<int>(int(float(Size.Y) * GNiagaraRenderTargetResolutionMultiplier), 1, GMaxTextureDimensions);
 	InstanceData->MipMapGeneration = MipMapGeneration;
 	InstanceData->MipMapGenerationType = MipMapGenerationType;
-	InstanceData->Format = bOverrideFormat ? OverrideRenderTargetFormat : GetDefault<UNiagaraSettings>()->DefaultRenderTargetFormat;
+	InstanceData->Format = GetRenderTargetFormat(bOverrideFormat, OverrideRenderTargetFormat);
 	InstanceData->RTUserParamBinding.Init(SystemInstance->GetInstanceParameters(), RenderTargetUserParameter.Parameter);
 #if WITH_EDITORONLY_DATA
 	InstanceData->bPreviewTexture = bPreviewRenderTarget;
diff --git a/Engine/Plugins/FX/Niagara/Source/Niagara/Private/NiagaraDataInterfaceRenderTarget2DArray.cpp b/Engine/Plugins/FX/Niagara/Source/Niagara/Private/NiagaraDataInterfaceRenderTarget2DArray.cpp
--- a/Engine/Plugins/FX/Niagara/Source/Niagara/Private/NiagaraDataInterfaceRenderTarget2DArray.cpp
+++ b/Engine/Plugins/FX/Niagara/Source/Niagara/Private/NiagaraDataInterfaceRenderTarget2DArray.cpp
@@ -508,12 +508,14 @@ bool UNiagaraDataInterfaceRenderTarget2DArray::InitPerInstanceData(void* PerInst
 {
 	check(Proxy);
 
+	extern ETextureRenderTargetFormat GetRenderTargetFormat(bool bOverrideFormat, ETextureRenderTargetFormat OverrideFormat);
 	extern float GNiagaraRenderTargetResolutionMultiplier;
+
 	FRenderTarget2DArrayRWInstanceData_GameThread* InstanceData = new (PerInstanceData) FRenderTarget2DArrayRWInstanceData_GameThread();
 	InstanceData->Size.X = FMath::Clamp<int>(int(float(Size.X) * GNiagaraRenderTargetResolutionMultiplier), 1, GMaxTextureDimensions);
 	InstanceData->Size.Y = FMath::Clamp<int>(int(float(Size.Y) * GNiagaraRenderTargetResolutionMultiplier), 1, GMaxTextureDimensions);
 	InstanceData->Size.Z = FMath::Clamp<int>(Size.Z, 1, GMaxTextureDimensions);
-	InstanceData->Format = GetPixelFormatFromRenderTargetFormat(bOverrideFormat ? OverrideRenderTargetFormat : GetDefault<UNiagaraSettings>()->DefaultRenderTargetFormat);
+	InstanceData->Format = GetPixelFormatFromRenderTargetFormat(GetRenderTargetFormat(bOverrideFormat, OverrideRenderTargetFormat));
 	InstanceData->RTUserParamBinding.Init(SystemInstance->GetInstanceParameters(), RenderTargetUserParameter.Parameter);
 #if WITH_EDITORONLY_DATA
 	InstanceData->bPreviewTexture = bPreviewRenderTarget;
@@ -522,7 +524,6 @@ bool UNiagaraDataInterfaceRenderTarget2DArray::InitPerInstanceData(void* PerInst
 	return true;
 }
 
-
 void UNiagaraDataInterfaceRenderTarget2DArray::DestroyPerInstanceData(void* PerInstanceData, FNiagaraSystemInstance* SystemInstance)
 {
 	FRenderTarget2DArrayRWInstanceData_GameThread* InstanceData = static_cast<FRenderTarget2DArrayRWInstanceData_GameThread*>(PerInstanceData);
diff --git a/Engine/Plugins/FX/Niagara/Source/Niagara/Private/NiagaraDataInterfaceRenderTargetCube.cpp b/Engine/Plugins/FX/Niagara/Source/Niagara/Private/NiagaraDataInterfaceRenderTargetCube.cpp
--- a/Engine/Plugins/FX/Niagara/Source/Niagara/Private/NiagaraDataInterfaceRenderTargetCube.cpp
+++ b/Engine/Plugins/FX/Niagara/Source/Niagara/Private/NiagaraDataInterfaceRenderTargetCube.cpp
@@ -498,10 +498,12 @@ bool UNiagaraDataInterfaceRenderTargetCube::InitPerInstanceData(void* PerInstanc
 {
 	check(Proxy);
 
+	extern ETextureRenderTargetFormat GetRenderTargetFormat(bool bOverrideFormat, ETextureRenderTargetFormat OverrideFormat);
 	extern float GNiagaraRenderTargetResolutionMultiplier;
+
 	FRenderTargetCubeRWInstanceData_GameThread* InstanceData = new (PerInstanceData) FRenderTargetCubeRWInstanceData_GameThread();
 	InstanceData->Size = FMath::Clamp<int>(int(float(Size) * GNiagaraRenderTargetResolutionMultiplier), 1, GMaxCubeTextureDimensions);
-	InstanceData->Format = GetPixelFormatFromRenderTargetFormat(bOverrideFormat ? OverrideRenderTargetFormat : GetDefault<UNiagaraSettings>()->DefaultRenderTargetFormat);
+	InstanceData->Format = GetPixelFormatFromRenderTargetFormat(GetRenderTargetFormat(bOverrideFormat, OverrideRenderTargetFormat));
 	InstanceData->RTUserParamBinding.Init(SystemInstance->GetInstanceParameters(), RenderTargetUserParameter.Parameter);
 #if WITH_EDITORONLY_DATA
 	InstanceData->bPreviewTexture = bPreviewRenderTarget;
diff --git a/Engine/Plugins/FX/Niagara/Source/Niagara/Private/NiagaraDataInterfaceRenderTargetVolume.cpp b/Engine/Plugins/FX/Niagara/Source/Niagara/Private/NiagaraDataInterfaceRenderTargetVolume.cpp
--- a/Engine/Plugins/FX/Niagara/Source/Niagara/Private/NiagaraDataInterfaceRenderTargetVolume.cpp
+++ b/Engine/Plugins/FX/Niagara/Source/Niagara/Private/NiagaraDataInterfaceRenderTargetVolume.cpp
@@ -505,12 +505,14 @@ bool UNiagaraDataInterfaceRenderTargetVolume::InitPerInstanceData(void* PerInsta
 {
 	check(Proxy);
 
+	extern ETextureRenderTargetFormat GetRenderTargetFormat(bool bOverrideFormat, ETextureRenderTargetFormat OverrideFormat);
 	extern float GNiagaraRenderTargetResolutionMultiplier;
+
 	FRenderTargetVolumeRWInstanceData_GameThread* InstanceData = new (PerInstanceData) FRenderTargetVolumeRWInstanceData_GameThread();
 	InstanceData->Size.X = FMath::Clamp<int>(int(float(Size.X) * GNiagaraRenderTargetResolutionMultiplier), 1, GMaxVolumeTextureDimensions);
 	InstanceData->Size.Y = FMath::Clamp<int>(int(float(Size.Y) * GNiagaraRenderTargetResolutionMultiplier), 1, GMaxVolumeTextureDimensions);
 	InstanceData->Size.Z = FMath::Clamp<int>(int(float(Size.Z) * GNiagaraRenderTargetResolutionMultiplier), 1, GMaxVolumeTextureDimensions);
-	InstanceData->Format = GetPixelFormatFromRenderTargetFormat(bOverrideFormat ? OverrideRenderTargetFormat : GetDefault<UNiagaraSettings>()->DefaultRenderTargetFormat);
+	InstanceData->Format = GetPixelFormatFromRenderTargetFormat(GetRenderTargetFormat(bOverrideFormat, OverrideRenderTargetFormat));
 	InstanceData->RTUserParamBinding.Init(SystemInstance->GetInstanceParameters(), RenderTargetUserParameter.Parameter);
 #if WITH_EDITORONLY_DATA
 	InstanceData->bPreviewTexture = bPreviewRenderTarget;
