Fixed script variable meta data properties that were supposed to be hidden not being hidden.
Reason: the meta data customization re-added the properties back. By skipping child property handles that are already marked customized (which means hidden) we also skip adding them back.

diff --git a/Engine/Plugins/FX/Niagara/Source/NiagaraEditor/Private/Customizations/NiagaraTypeCustomizations.cpp b/Engine/Plugins/FX/Niagara/Source/NiagaraEditor/Private/Customizations/NiagaraTypeCustomizations.cpp
--- a/Engine/Plugins/FX/Niagara/Source/NiagaraEditor/Private/Customizations/NiagaraTypeCustomizations.cpp
+++ b/Engine/Plugins/FX/Niagara/Source/NiagaraEditor/Private/Customizations/NiagaraTypeCustomizations.cpp
@@ -1872,6 +1872,15 @@ void FNiagaraVariableMetaDataCustomization::CustomizeChildren(TSharedRef<IProper
 
 	for (const TSharedRef<IPropertyHandle>& ChildPropertyHandle : ChildPropertyHandles)
 	{
+		/** The customized flag is used for hiding properties.
+		 *  The script variable's customization which contains meta data is hiding these properties already,
+		 *  but the meta data customization would stomp these by readding them. Hence, we skip them here if they were marked as customized.
+		*/
+		if(ChildPropertyHandle->IsCustomized())
+		{
+			continue;
+		}
+		
 		const FProperty* ChildProperty = ChildPropertyHandle->GetProperty();
 		const FName ChildPropertyName = ChildProperty->GetFName();
 
