Fixed the script var customization hiding the inline parameter data when it shouldn't for staticswitch variables

diff --git a/Engine/Plugins/FX/Niagara/Source/NiagaraEditor/Private/Customizations/NiagaraScriptVariableCustomization.cpp b/Engine/Plugins/FX/Niagara/Source/NiagaraEditor/Private/Customizations/NiagaraScriptVariableCustomization.cpp
--- a/Engine/Plugins/FX/Niagara/Source/NiagaraEditor/Private/Customizations/NiagaraScriptVariableCustomization.cpp
+++ b/Engine/Plugins/FX/Niagara/Source/NiagaraEditor/Private/Customizations/NiagaraScriptVariableCustomization.cpp
@@ -187,7 +187,7 @@ void FNiagaraScriptVariableDetails::CustomizeDetails(IDetailLayoutBuilder& Detai
 
 	// @TODO It seems like all the HideProperties aren't working for values within metadata
 	// generally we want all the inline parameters only available for inputs & static switches
-	if(!Variable->Variable.IsInNameSpace(FNiagaraConstants::ModuleNamespace) && !Variable->Variable.IsInNameSpace(FNiagaraConstants::StaticSwitchNamespace))
+	if(!Variable->Variable.IsInNameSpace(FNiagaraConstants::ModuleNamespace) && !Variable->GetIsStaticSwitch())
 	{
 		DetailBuilder.HideProperty(GET_MEMBER_NAME_CHECKED(UNiagaraScriptVariable, Metadata.bDisplayInOverviewStack));
 		DetailBuilder.HideProperty(GET_MEMBER_NAME_CHECKED(UNiagaraScriptVariable, Metadata.InlineParameterColorOverride));
