- Fix crash in scalability manager

#rnx

diff --git a/Engine/Plugins/FX/Niagara/Source/Niagara/Private/NiagaraComponent.cpp b/Engine/Plugins/FX/Niagara/Source/Niagara/Private/NiagaraComponent.cpp
--- a/Engine/Plugins/FX/Niagara/Source/Niagara/Private/NiagaraComponent.cpp
+++ b/Engine/Plugins/FX/Niagara/Source/Niagara/Private/NiagaraComponent.cpp
@@ -3357,7 +3357,10 @@ void UNiagaraComponent::SetAllowScalability(bool bAllow)
 	bAllowScalability = bAllow; 
 	if (bAllow)
 	{
-		RegisterWithScalabilityManager();
+		if ( IsActive() )
+		{
+			RegisterWithScalabilityManager();
+		}
 	}
 	else
 	{
