- Ensure we don't create a RT format that can not handle typed UAV store

#rnx

diff --git a/Engine/Plugins/FX/Niagara/Source/Niagara/Private/NiagaraDataInterfaceRenderTarget2D.cpp b/Engine/Plugins/FX/Niagara/Source/Niagara/Private/NiagaraDataInterfaceRenderTarget2D.cpp
--- a/Engine/Plugins/FX/Niagara/Source/Niagara/Private/NiagaraDataInterfaceRenderTarget2D.cpp
+++ b/Engine/Plugins/FX/Niagara/Source/Niagara/Private/NiagaraDataInterfaceRenderTarget2D.cpp
@@ -102,13 +102,51 @@ static FAutoConsoleCommandWithWorldAndArgs GCommandNiagaraRenderTargetOverrideFo
 	)
 );
 
-ETextureRenderTargetFormat GetRenderTargetFormat(bool bOverrideFormat, ETextureRenderTargetFormat OverrideFormat)
+bool GetRenderTargetFormat(bool bOverrideFormat, ETextureRenderTargetFormat OverrideFormat, ETextureRenderTargetFormat& OutRenderTargetFormat)
 {
+	OutRenderTargetFormat = bOverrideFormat ? OverrideFormat : GetDefault<UNiagaraSettings>()->DefaultRenderTargetFormat.GetValue();
+	EPixelFormat PixelFormat = GetPixelFormatFromRenderTargetFormat(OutRenderTargetFormat);
 	if (GNiagaraRenderTargetOverrideFormatEnabled)
 	{
-		return GNiagaraRenderTargetOverrideFormat;
+		OutRenderTargetFormat = GNiagaraRenderTargetOverrideFormat;
 	}
-	return bOverrideFormat ? OverrideFormat : GetDefault<UNiagaraSettings>()->DefaultRenderTargetFormat.GetValue();
+
+	// If the format does not support typed store we need to find one that will
+	while ( GDynamicRHI->RHIIsTypedUAVStoreSupported(GetPixelFormatFromRenderTargetFormat(OutRenderTargetFormat)) == false )
+	{
+		static const TPair<ETextureRenderTargetFormat, ETextureRenderTargetFormat> FormatRemapTable[] =
+		{
+			TPairInitializer(RTF_R8,			RTF_R16f),
+			TPairInitializer(RTF_RG8,			RTF_RG16f),
+			TPairInitializer(RTF_RGBA8,			RTF_RGBA16f),
+			TPairInitializer(RTF_RGBA8_SRGB,	RTF_RGBA16f),
+			TPairInitializer(RTF_R16f,			RTF_R32f),
+			TPairInitializer(RTF_RG16f,			RTF_RG32f),
+			TPairInitializer(RTF_RGBA16f,		RTF_RGBA32f),
+			TPairInitializer(RTF_R32f,			RTF_RGBA32f),
+			TPairInitializer(RTF_RG32f,			RTF_RGBA32f),
+			TPairInitializer(RTF_RGBA32f,		RTF_RGBA32f),
+			TPairInitializer(RTF_RGB10A2,		RTF_RGBA32f),
+		};
+
+		const ETextureRenderTargetFormat PreviousFormat = OutRenderTargetFormat;
+		for ( const auto& FormatRemap : FormatRemapTable)
+		{
+			if (FormatRemap.Key == OutRenderTargetFormat)
+			{
+				OutRenderTargetFormat = FormatRemap.Value;
+				break;
+			}
+		}
+		if ( PreviousFormat == OutRenderTargetFormat)
+		{
+			// This is fatal as we failed to find any format that supports typed UAV stores
+			UE_LOG(LogNiagara, Warning, TEXT("Failed to find a render target format that supports UAV store"));
+			return false;
+		}
+	}
+
+	return true;
 }
 
 /*--------------------------------------------------------------------------------------------------------------------------*/
@@ -592,11 +630,17 @@ bool UNiagaraDataInterfaceRenderTarget2D::InitPerInstanceData(void* PerInstanceD
 	FRenderTarget2DRWInstanceData_GameThread* InstanceData = new (PerInstanceData) FRenderTarget2DRWInstanceData_GameThread();
 	SystemInstancesToProxyData_GT.Emplace(SystemInstance->GetId(), InstanceData);
 
+	ETextureRenderTargetFormat RenderTargetFormat;
+	if ( GetRenderTargetFormat(bOverrideFormat, OverrideRenderTargetFormat, RenderTargetFormat) == false )
+	{
+		return false;
+	}
+
 	InstanceData->Size.X = FMath::Clamp<int>(int(float(Size.X) * GNiagaraRenderTargetResolutionMultiplier), 1, GMaxTextureDimensions);
 	InstanceData->Size.Y = FMath::Clamp<int>(int(float(Size.Y) * GNiagaraRenderTargetResolutionMultiplier), 1, GMaxTextureDimensions);
 	InstanceData->MipMapGeneration = MipMapGeneration;
 	InstanceData->MipMapGenerationType = MipMapGenerationType;
-	InstanceData->Format = GetRenderTargetFormat(bOverrideFormat, OverrideRenderTargetFormat);
+	InstanceData->Format = RenderTargetFormat;
 	InstanceData->RTUserParamBinding.Init(SystemInstance->GetInstanceParameters(), RenderTargetUserParameter.Parameter);
 #if WITH_EDITORONLY_DATA
 	InstanceData->bPreviewTexture = bPreviewRenderTarget;
diff --git a/Engine/Plugins/FX/Niagara/Source/Niagara/Private/NiagaraDataInterfaceRenderTarget2DArray.cpp b/Engine/Plugins/FX/Niagara/Source/Niagara/Private/NiagaraDataInterfaceRenderTarget2DArray.cpp
--- a/Engine/Plugins/FX/Niagara/Source/Niagara/Private/NiagaraDataInterfaceRenderTarget2DArray.cpp
+++ b/Engine/Plugins/FX/Niagara/Source/Niagara/Private/NiagaraDataInterfaceRenderTarget2DArray.cpp
@@ -508,14 +508,20 @@ bool UNiagaraDataInterfaceRenderTarget2DArray::InitPerInstanceData(void* PerInst
 {
 	check(Proxy);
 
-	extern ETextureRenderTargetFormat GetRenderTargetFormat(bool bOverrideFormat, ETextureRenderTargetFormat OverrideFormat);
+	extern bool GetRenderTargetFormat(bool bOverrideFormat, ETextureRenderTargetFormat OverrideFormat, ETextureRenderTargetFormat & OutRenderTargetFormat);
 	extern float GNiagaraRenderTargetResolutionMultiplier;
 
+	ETextureRenderTargetFormat RenderTargetFormat;
+	if (GetRenderTargetFormat(bOverrideFormat, OverrideRenderTargetFormat, RenderTargetFormat) == false)
+	{
+		return false;
+	}
+
 	FRenderTarget2DArrayRWInstanceData_GameThread* InstanceData = new (PerInstanceData) FRenderTarget2DArrayRWInstanceData_GameThread();
 	InstanceData->Size.X = FMath::Clamp<int>(int(float(Size.X) * GNiagaraRenderTargetResolutionMultiplier), 1, GMaxTextureDimensions);
 	InstanceData->Size.Y = FMath::Clamp<int>(int(float(Size.Y) * GNiagaraRenderTargetResolutionMultiplier), 1, GMaxTextureDimensions);
 	InstanceData->Size.Z = FMath::Clamp<int>(Size.Z, 1, GMaxTextureDimensions);
-	InstanceData->Format = GetPixelFormatFromRenderTargetFormat(GetRenderTargetFormat(bOverrideFormat, OverrideRenderTargetFormat));
+	InstanceData->Format = GetPixelFormatFromRenderTargetFormat(RenderTargetFormat);
 	InstanceData->RTUserParamBinding.Init(SystemInstance->GetInstanceParameters(), RenderTargetUserParameter.Parameter);
 #if WITH_EDITORONLY_DATA
 	InstanceData->bPreviewTexture = bPreviewRenderTarget;
diff --git a/Engine/Plugins/FX/Niagara/Source/Niagara/Private/NiagaraDataInterfaceRenderTargetCube.cpp b/Engine/Plugins/FX/Niagara/Source/Niagara/Private/NiagaraDataInterfaceRenderTargetCube.cpp
--- a/Engine/Plugins/FX/Niagara/Source/Niagara/Private/NiagaraDataInterfaceRenderTargetCube.cpp
+++ b/Engine/Plugins/FX/Niagara/Source/Niagara/Private/NiagaraDataInterfaceRenderTargetCube.cpp
@@ -498,12 +498,18 @@ bool UNiagaraDataInterfaceRenderTargetCube::InitPerInstanceData(void* PerInstanc
 {
 	check(Proxy);
 
-	extern ETextureRenderTargetFormat GetRenderTargetFormat(bool bOverrideFormat, ETextureRenderTargetFormat OverrideFormat);
+	extern bool GetRenderTargetFormat(bool bOverrideFormat, ETextureRenderTargetFormat OverrideFormat, ETextureRenderTargetFormat & OutRenderTargetFormat);
 	extern float GNiagaraRenderTargetResolutionMultiplier;
 
+	ETextureRenderTargetFormat RenderTargetFormat;
+	if (GetRenderTargetFormat(bOverrideFormat, OverrideRenderTargetFormat, RenderTargetFormat) == false)
+	{
+		return false;
+	}
+
 	FRenderTargetCubeRWInstanceData_GameThread* InstanceData = new (PerInstanceData) FRenderTargetCubeRWInstanceData_GameThread();
 	InstanceData->Size = FMath::Clamp<int>(int(float(Size) * GNiagaraRenderTargetResolutionMultiplier), 1, GMaxCubeTextureDimensions);
-	InstanceData->Format = GetPixelFormatFromRenderTargetFormat(GetRenderTargetFormat(bOverrideFormat, OverrideRenderTargetFormat));
+	InstanceData->Format = GetPixelFormatFromRenderTargetFormat(RenderTargetFormat);
 	InstanceData->RTUserParamBinding.Init(SystemInstance->GetInstanceParameters(), RenderTargetUserParameter.Parameter);
 #if WITH_EDITORONLY_DATA
 	InstanceData->bPreviewTexture = bPreviewRenderTarget;
diff --git a/Engine/Plugins/FX/Niagara/Source/Niagara/Private/NiagaraDataInterfaceRenderTargetVolume.cpp b/Engine/Plugins/FX/Niagara/Source/Niagara/Private/NiagaraDataInterfaceRenderTargetVolume.cpp
--- a/Engine/Plugins/FX/Niagara/Source/Niagara/Private/NiagaraDataInterfaceRenderTargetVolume.cpp
+++ b/Engine/Plugins/FX/Niagara/Source/Niagara/Private/NiagaraDataInterfaceRenderTargetVolume.cpp
@@ -505,14 +505,20 @@ bool UNiagaraDataInterfaceRenderTargetVolume::InitPerInstanceData(void* PerInsta
 {
 	check(Proxy);
 
-	extern ETextureRenderTargetFormat GetRenderTargetFormat(bool bOverrideFormat, ETextureRenderTargetFormat OverrideFormat);
+	extern bool GetRenderTargetFormat(bool bOverrideFormat, ETextureRenderTargetFormat OverrideFormat, ETextureRenderTargetFormat & OutRenderTargetFormat);
 	extern float GNiagaraRenderTargetResolutionMultiplier;
 
+	ETextureRenderTargetFormat RenderTargetFormat;
+	if (GetRenderTargetFormat(bOverrideFormat, OverrideRenderTargetFormat, RenderTargetFormat) == false)
+	{
+		return false;
+	}
+
 	FRenderTargetVolumeRWInstanceData_GameThread* InstanceData = new (PerInstanceData) FRenderTargetVolumeRWInstanceData_GameThread();
 	InstanceData->Size.X = FMath::Clamp<int>(int(float(Size.X) * GNiagaraRenderTargetResolutionMultiplier), 1, GMaxVolumeTextureDimensions);
 	InstanceData->Size.Y = FMath::Clamp<int>(int(float(Size.Y) * GNiagaraRenderTargetResolutionMultiplier), 1, GMaxVolumeTextureDimensions);
 	InstanceData->Size.Z = FMath::Clamp<int>(int(float(Size.Z) * GNiagaraRenderTargetResolutionMultiplier), 1, GMaxVolumeTextureDimensions);
-	InstanceData->Format = GetPixelFormatFromRenderTargetFormat(GetRenderTargetFormat(bOverrideFormat, OverrideRenderTargetFormat));
+	InstanceData->Format = GetPixelFormatFromRenderTargetFormat(RenderTargetFormat);
 	InstanceData->RTUserParamBinding.Init(SystemInstance->GetInstanceParameters(), RenderTargetUserParameter.Parameter);
 #if WITH_EDITORONLY_DATA
 	InstanceData->bPreviewTexture = bPreviewRenderTarget;
