- Add option to dump data interfaces and usage

#rnx

diff --git a/Engine/Plugins/FX/Niagara/Source/NiagaraEditor/Private/Commandlets/NiagaraSystemAuditCommandlet.cpp b/Engine/Plugins/FX/Niagara/Source/NiagaraEditor/Private/Commandlets/NiagaraSystemAuditCommandlet.cpp
--- a/Engine/Plugins/FX/Niagara/Source/NiagaraEditor/Private/Commandlets/NiagaraSystemAuditCommandlet.cpp
+++ b/Engine/Plugins/FX/Niagara/Source/NiagaraEditor/Private/Commandlets/NiagaraSystemAuditCommandlet.cpp
@@ -86,6 +86,8 @@ int32 UNiagaraSystemAuditCommandlet::Main(const FString& Params)
 
 	FParse::Bool(*Params, TEXT("RendererDetailed="), bRendererDetailed);
 
+	FParse::Bool(*Params, TEXT("CaptureDataInterfaceUsage="), bCaptureDataInterfaceUsage);
+
 	// Package Paths
 	FString PackagePathsString;
 	if (FParse::Value(*Params, TEXT("PackagePaths="), PackagePathsString, false))
@@ -98,7 +100,7 @@ int32 UNiagaraSystemAuditCommandlet::Main(const FString& Params)
 		}
 	}
 
-	if (PackagePaths.Num() == 0)
+	if (PackagePaths.Num() == 0 && FParse::Param(*Params, TEXT("GameContentOnly")))
 	{
 		PackagePaths.Add(FName(TEXT("/Game")));
 	}
@@ -181,6 +183,13 @@ bool UNiagaraSystemAuditCommandlet::ProcessNiagaraSystems()
 		TSet<FName> SystemUserDataInterfaces;
 		for (UNiagaraDataInterface* DataInterface : GetDataInterfaces(NiagaraSystem))
 		{
+			if (bCaptureDataInterfaceUsage)
+			{
+				FDataInterfaceUsage& DataInterfaceUsage = NiagaraDataInterfaceUsage.FindOrAdd(DataInterface->GetClass()->GetFName());
+				++DataInterfaceUsage.UsageCount;
+				DataInterfaceUsage.Systems.Add(NiagaraSystem->GetFName());
+			}
+
 			if (DataInterface->HasTickGroupPrereqs())
 			{
 				SystemDataInterfacesWihPrereqs.Add(DataInterface->GetClass()->GetFName());
@@ -388,6 +397,28 @@ void UNiagaraSystemAuditCommandlet::DumpResults()
 		HeaderString.Append(TEXT(",CVar Conditions,System Path"));
 		DumpSimpleSet(NiagaraSystemsWithGPUEmitters, TEXT("NiagaraSystemsWithGPUEmitters"), HeaderString.ToString());
 	}
+
+	if (NiagaraDataInterfaceUsage.Num() > 0)
+	{
+		if ( FArchive* OutputStream = GetOutputFile(TEXT("NiagaraDataInterfaceUsage")) )
+		{
+			OutputStream->Logf(TEXT("Name,Usage Count,Systems"));
+			for ( auto UsageIt=NiagaraDataInterfaceUsage.CreateConstIterator(); UsageIt; ++UsageIt)
+			{
+				FString SystemList;
+				for ( FName SystemName : UsageIt.Value().Systems )
+				{
+					SystemList.Append(*SystemName.ToString());
+					SystemList.AppendChar(' ');
+				}
+
+				OutputStream->Logf(TEXT("%s,%d,%s"), *UsageIt.Key().ToString(), UsageIt.Value().UsageCount, *SystemList);
+			}
+			OutputStream->Close();
+			delete OutputStream;
+		}
+	}
+
 	if (bRendererDetailed)
 	{
 		DumpSimpleSet(NiagaraRibbonRenderers, TEXT("NiagaraRibbonRenderers"), TEXT("Name,TessellationMode"));
diff --git a/Engine/Plugins/FX/Niagara/Source/NiagaraEditor/Public/Commandlets/NiagaraSystemAuditCommandlet.h b/Engine/Plugins/FX/Niagara/Source/NiagaraEditor/Public/Commandlets/NiagaraSystemAuditCommandlet.h
--- a/Engine/Plugins/FX/Niagara/Source/NiagaraEditor/Public/Commandlets/NiagaraSystemAuditCommandlet.h
+++ b/Engine/Plugins/FX/Niagara/Source/NiagaraEditor/Public/Commandlets/NiagaraSystemAuditCommandlet.h
@@ -28,6 +28,15 @@ class UNiagaraSystemAuditCommandlet : public UCommandlet
 	/** All Niagara systems that use a user requested data interface */
 	TSet<FString> NiagaraSystemsWithUserDataInterface;
 
+	/** Optional list of data interface usage */
+	struct FDataInterfaceUsage
+	{
+		int32			UsageCount = 0;
+		TSet<FName>		Systems;
+	};
+	bool bCaptureDataInterfaceUsage = false;
+	TMap<FName, FDataInterfaceUsage>  NiagaraDataInterfaceUsage;
+
 	/** Detailed renderer information. */
 	bool bRendererDetailed = false;
 	TSet<FString> NiagaraRibbonRenderers;
