- Ensure we always have dummy buffers rather than leaving nullptr
- Ensure we keep hold of the mesh buffers until we get a notification about stream out

#rnx

diff --git a/Engine/Plugins/FX/Niagara/Source/Niagara/Classes/NiagaraDataInterfaceSkeletalMesh.h b/Engine/Plugins/FX/Niagara/Source/Niagara/Classes/NiagaraDataInterfaceSkeletalMesh.h
--- a/Engine/Plugins/FX/Niagara/Source/Niagara/Classes/NiagaraDataInterfaceSkeletalMesh.h
+++ b/Engine/Plugins/FX/Niagara/Source/Niagara/Classes/NiagaraDataInterfaceSkeletalMesh.h
@@ -455,11 +455,11 @@ public:
 	FRHIShaderResourceView* GetSampleRegionsTriangleIndicesSRV() const { return SampleRegionsTriangleIndicesSRV; }
 	FRHIShaderResourceView* GetSampleRegionsVerticesSRV() const { return SampleRegionsVerticesSRV; }
 
-	FRHIShaderResourceView* GetBufferPositionSRV() const { return MeshVertexBufferSrv; }
-	FRHIShaderResourceView* GetBufferIndexSRV() const { return MeshIndexBufferSrv; }
+	FRHIShaderResourceView* GetBufferPositionSRV() const { return MeshVertexBufferSRV; }
+	FRHIShaderResourceView* GetBufferIndexSRV() const { return MeshIndexBufferSRV; }
 	FRHIShaderResourceView* GetBufferTangentSRV() const { return MeshTangentBufferSRV; }
-	FRHIShaderResourceView* GetBufferTexCoordSRV() const { return MeshTexCoordBufferSrv; }
-	FRHIShaderResourceView* GetBufferColorSRV() const { return MeshColorBufferSrv; }
+	FRHIShaderResourceView* GetBufferTexCoordSRV() const { return MeshTexCoordBufferSRV; }
+	FRHIShaderResourceView* GetBufferColorSRV() const { return MeshColorBufferSRV; }
 
 	uint32 GetNumTexCoord() const { return NumTexCoord; }
 	uint32 GetNumWeights() const { return NumWeights; }
@@ -503,11 +503,11 @@ protected:
 	int32 FilteredSocketBoneOffset = 0;
 
 	/** Cached SRV to gpu buffers of the mesh we spawn from */
-	FRHIShaderResourceView* MeshVertexBufferSrv;
-	FRHIShaderResourceView* MeshIndexBufferSrv;
-	FRHIShaderResourceView* MeshTangentBufferSRV;
-	FRHIShaderResourceView* MeshTexCoordBufferSrv;
-	FRHIShaderResourceView* MeshColorBufferSrv;
+	FShaderResourceViewRHIRef MeshVertexBufferSRV;
+	FShaderResourceViewRHIRef MeshIndexBufferSRV;
+	FShaderResourceViewRHIRef MeshTangentBufferSRV;
+	FShaderResourceViewRHIRef MeshTexCoordBufferSRV;
+	FShaderResourceViewRHIRef MeshColorBufferSRV;
 
 	uint32 NumTexCoord = 0;
 	uint32 NumWeights = 0;
diff --git a/Engine/Plugins/FX/Niagara/Source/Niagara/Private/NiagaraDataInterfaceSkeletalMesh.cpp b/Engine/Plugins/FX/Niagara/Source/Niagara/Private/NiagaraDataInterfaceSkeletalMesh.cpp
--- a/Engine/Plugins/FX/Niagara/Source/Niagara/Private/NiagaraDataInterfaceSkeletalMesh.cpp
+++ b/Engine/Plugins/FX/Niagara/Source/Niagara/Private/NiagaraDataInterfaceSkeletalMesh.cpp
@@ -827,25 +827,13 @@ void FSkeletalMeshGpuSpawnStaticBuffers::InitRHI()
 	// Also see https://jira.it.epicgames.net/browse/UE-69376 : we would need to know if GPU sampling of the mesh surface is needed or not on the mesh to be able to do that.
 	// ALso today we do not know if an interface is create from a CPU or GPU emitter. So always allocate for now. Follow up in https://jira.it.epicgames.net/browse/UE-69375.
 
-	const FMultiSizeIndexContainer& IndexBuffer = LODRenderData->MultiSizeIndexContainer;
-	MeshIndexBufferSrv = IndexBuffer.GetIndexBuffer()->GetSRV();
-
-	// TODO: Disable this for now to suppress log spam. Revive when we know for sure that the user is trying to sample from it
-	//if (!MeshIndexBufferSrv)
-	//{
-	//	UE_LOG(LogNiagara, Warning, TEXT("Skeletal Mesh does not have an SRV for the index buffer, if you are using triangle sampling it will not work."));
-	//}
-
-	MeshVertexBufferSrv = LODRenderData->StaticVertexBuffers.PositionVertexBuffer.GetSRV();
-
-	MeshTangentBufferSRV = LODRenderData->StaticVertexBuffers.StaticMeshVertexBuffer.GetTangentsSRV();
-	//check(MeshTangentBufferSRV->IsValid()); // not available in this stream
-
-	MeshTexCoordBufferSrv = LODRenderData->StaticVertexBuffers.StaticMeshVertexBuffer.GetTexCoordsSRV();
+	MeshIndexBufferSRV = FNiagaraRenderer::GetSrvOrDefaultUInt(LODRenderData->MultiSizeIndexContainer.GetIndexBuffer()->GetSRV());
+	MeshVertexBufferSRV = FNiagaraRenderer::GetSrvOrDefaultFloat(LODRenderData->StaticVertexBuffers.PositionVertexBuffer.GetSRV());
+	MeshTangentBufferSRV = FNiagaraRenderer::GetSrvOrDefaultFloat4(LODRenderData->StaticVertexBuffers.StaticMeshVertexBuffer.GetTangentsSRV());
+	MeshTexCoordBufferSRV = FNiagaraRenderer::GetSrvOrDefaultFloat2(LODRenderData->StaticVertexBuffers.StaticMeshVertexBuffer.GetTexCoordsSRV());
+	MeshColorBufferSRV = FNiagaraRenderer::GetSrvOrDefaultFloat4(LODRenderData->StaticVertexBuffers.ColorVertexBuffer.GetColorComponentsSRV());
 	NumTexCoord = LODRenderData->StaticVertexBuffers.StaticMeshVertexBuffer.GetNumTexCoords();
 
-	MeshColorBufferSrv = LODRenderData->StaticVertexBuffers.ColorVertexBuffer.GetColorComponentsSRV();
-
 	NumWeights = LODRenderData->SkinWeightVertexBuffer.GetMaxBoneInfluences();
 
 	uint32 SectionCount = LODRenderData->RenderSections.Num();
@@ -870,8 +858,15 @@ void FSkeletalMeshGpuSpawnStaticBuffers::InitRHI()
 		GPUMemoryUsage += SizeByte;
 #endif
 	}
+	else
+	{
+		BufferTriangleUniformSamplerProbAliasSRV = FNiagaraRenderer::GetDummyUIntBuffer();
+	}
 
 	// Prepare sampling regions (if we have any)
+	SampleRegionsProbAliasSRV = FNiagaraRenderer::GetDummyUIntBuffer();
+	SampleRegionsTriangleIndicesSRV = FNiagaraRenderer::GetDummyUIntBuffer();
+	SampleRegionsVerticesSRV = FNiagaraRenderer::GetDummyUIntBuffer();
 	if (NumSamplingRegionTriangles > 0)
 	{
 		FRHIResourceCreateInfo CreateInfo(TEXT("SampleRegionsProbAliasBuffer"));
@@ -936,6 +931,10 @@ void FSkeletalMeshGpuSpawnStaticBuffers::InitRHI()
 		FilteredAndUnfilteredBonesBuffer = RHICreateVertexBuffer(FilteredAndUnfilteredBonesArray.Num() * FilteredAndUnfilteredBonesArray.GetTypeSize(), BUF_Static | BUF_ShaderResource, CreateInfo);
 		FilteredAndUnfilteredBonesSRV = RHICreateShaderResourceView(FilteredAndUnfilteredBonesBuffer, sizeof(uint16), PF_R16_UINT);
 	}
+	else
+	{
+		FilteredAndUnfilteredBonesSRV = FNiagaraRenderer::GetDummyUIntBuffer();
+	}
 #if STATS
 	GPUMemoryUsage += FilteredAndUnfilteredBonesArray.Num() * FilteredAndUnfilteredBonesArray.GetTypeSize();
 #endif
@@ -952,8 +951,8 @@ void FSkeletalMeshGpuSpawnStaticBuffers::ReleaseRHI()
 	GPUMemoryUsage = 0;
 #endif
 
-	FilteredAndUnfilteredBonesBuffer.SafeRelease();
-	FilteredAndUnfilteredBonesSRV.SafeRelease();
+	BufferTriangleMatricesOffsetRHI.SafeRelease();
+	BufferTriangleMatricesOffsetSRV.SafeRelease();
 
 	BufferTriangleUniformSamplerProbAliasRHI.SafeRelease();
 	BufferTriangleUniformSamplerProbAliasSRV.SafeRelease();
@@ -965,11 +964,14 @@ void FSkeletalMeshGpuSpawnStaticBuffers::ReleaseRHI()
 	SampleRegionsVerticesBuffer.SafeRelease();
 	SampleRegionsVerticesSRV.SafeRelease();
 
-	MeshVertexBufferSrv = nullptr;
-	MeshIndexBufferSrv = nullptr;
-	MeshTangentBufferSRV = nullptr;
-	MeshTexCoordBufferSrv = nullptr;
-	MeshColorBufferSrv = nullptr;
+	FilteredAndUnfilteredBonesBuffer.SafeRelease();
+	FilteredAndUnfilteredBonesSRV.SafeRelease();
+
+	MeshVertexBufferSRV.SafeRelease();
+	MeshIndexBufferSRV.SafeRelease();
+	MeshTangentBufferSRV.SafeRelease();
+	MeshTexCoordBufferSRV.SafeRelease();
+	MeshColorBufferSRV.SafeRelease();
 }
 
 //////////////////////////////////////////////////////////////////////////
@@ -1404,20 +1406,20 @@ public:
 			SetSRVParameter(RHICmdList, ComputeShaderRHI, MeshTangentBuffer, StaticBuffers->GetBufferTangentSRV());
 
 			SetShaderValue(RHICmdList, ComputeShaderRHI, MeshNumTexCoord, StaticBuffers->GetNumTexCoord());
-			SetSRVParameter(RHICmdList, ComputeShaderRHI, MeshTexCoordBuffer, StaticBuffers->GetBufferTexCoordSRV() != nullptr ? StaticBuffers->GetBufferTexCoordSRV() : FNiagaraRenderer::GetDummyFloatBuffer());
-			SetSRVParameter(RHICmdList, ComputeShaderRHI, MeshColorBuffer, StaticBuffers->GetBufferColorSRV() != nullptr ? StaticBuffers->GetBufferColorSRV() : FNiagaraRenderer::GetDummyFloat4Buffer());
+			SetSRVParameter(RHICmdList, ComputeShaderRHI, MeshTexCoordBuffer, StaticBuffers->GetBufferTexCoordSRV());
+			SetSRVParameter(RHICmdList, ComputeShaderRHI, MeshColorBuffer, StaticBuffers->GetBufferColorSRV());
 			SetShaderValue(RHICmdList, ComputeShaderRHI, MeshTriangleCount, StaticBuffers->GetTriangleCount());
 			SetShaderValue(RHICmdList, ComputeShaderRHI, MeshVertexCount, StaticBuffers->GetVertexCount());
 
 			// Set triangle sampling buffer
-			SetSRVParameter(RHICmdList, ComputeShaderRHI, MeshTriangleSamplerProbAliasBuffer, StaticBuffers->GetBufferTriangleUniformSamplerProbAliasSRV() != nullptr ? StaticBuffers->GetBufferTriangleUniformSamplerProbAliasSRV() : FNiagaraRenderer::GetDummyUIntBuffer());
+			SetSRVParameter(RHICmdList, ComputeShaderRHI, MeshTriangleSamplerProbAliasBuffer, StaticBuffers->GetBufferTriangleUniformSamplerProbAliasSRV());
 
 			// Set triangle sampling region buffer
 			SetShaderValue(RHICmdList, ComputeShaderRHI, MeshNumSamplingRegionTriangles, StaticBuffers->GetNumSamplingRegionTriangles());
 			SetShaderValue(RHICmdList, ComputeShaderRHI, MeshNumSamplingRegionVertices, StaticBuffers->GetNumSamplingRegionVertices());
-			SetSRVParameter(RHICmdList, ComputeShaderRHI, MeshSamplingRegionsProbAliasBuffer, StaticBuffers->GetSampleRegionsProbAliasSRV() != nullptr ? StaticBuffers->GetSampleRegionsProbAliasSRV() : FNiagaraRenderer::GetDummyUIntBuffer());
-			SetSRVParameter(RHICmdList, ComputeShaderRHI, MeshSampleRegionsTriangleIndices, StaticBuffers->GetSampleRegionsTriangleIndicesSRV() != nullptr ? StaticBuffers->GetSampleRegionsTriangleIndicesSRV() : FNiagaraRenderer::GetDummyUIntBuffer());
-			SetSRVParameter(RHICmdList, ComputeShaderRHI, MeshSampleRegionsVertices, StaticBuffers->GetSampleRegionsVerticesSRV() != nullptr ? StaticBuffers->GetSampleRegionsVerticesSRV() : FNiagaraRenderer::GetDummyUIntBuffer());
+			SetSRVParameter(RHICmdList, ComputeShaderRHI, MeshSamplingRegionsProbAliasBuffer, StaticBuffers->GetSampleRegionsProbAliasSRV());
+			SetSRVParameter(RHICmdList, ComputeShaderRHI, MeshSampleRegionsTriangleIndices, StaticBuffers->GetSampleRegionsTriangleIndicesSRV());
+			SetSRVParameter(RHICmdList, ComputeShaderRHI, MeshSampleRegionsVertices, StaticBuffers->GetSampleRegionsVerticesSRV());
 
 			FRHIShaderResourceView* MeshSkinWeightBufferSRV = InstanceData->MeshSkinWeightBuffer->GetSRV();
 			SetSRVParameter(RHICmdList, ComputeShaderRHI, MeshSkinWeightBuffer, MeshSkinWeightBufferSRV ? MeshSkinWeightBufferSRV : FNiagaraRenderer::GetDummyUIntBuffer());
@@ -1456,12 +1458,11 @@ public:
 
 			SetShaderValue(RHICmdList, ComputeShaderRHI, NumBones, DynamicBuffers->GetNumBones());
 
-			FRHIShaderResourceView* FilteredAndUnfilteredBonesSRV = StaticBuffers->GetFilteredAndUnfilteredBonesSRV() != nullptr ? StaticBuffers->GetFilteredAndUnfilteredBonesSRV() : FNiagaraRenderer::GetDummyUIntBuffer();
 			SetShaderValue(RHICmdList, ComputeShaderRHI, NumFilteredBones, StaticBuffers->GetNumFilteredBones());
 			SetShaderValue(RHICmdList, ComputeShaderRHI, NumUnfilteredBones, StaticBuffers->GetNumUnfilteredBones());
 			SetShaderValue(RHICmdList, ComputeShaderRHI, RandomMaxBone, StaticBuffers->GetExcludedBoneIndex() >= 0 ? DynamicBuffers->GetNumBones() - 2 : DynamicBuffers->GetNumBones() - 1);
 			SetShaderValue(RHICmdList, ComputeShaderRHI, ExcludeBoneIndex, StaticBuffers->GetExcludedBoneIndex());
-			SetSRVParameter(RHICmdList, ComputeShaderRHI, FilteredAndUnfilteredBones, FilteredAndUnfilteredBonesSRV);
+			SetSRVParameter(RHICmdList, ComputeShaderRHI, FilteredAndUnfilteredBones, StaticBuffers->GetFilteredAndUnfilteredBonesSRV());
 
 			SetShaderValue(RHICmdList, ComputeShaderRHI, NumFilteredSockets, StaticBuffers->GetNumFilteredSockets());
 			SetShaderValue(RHICmdList, ComputeShaderRHI, FilteredSocketBoneOffset, StaticBuffers->GetFilteredSocketBoneOffset());
diff --git a/Engine/Plugins/FX/Niagara/Source/Niagara/Public/NiagaraRenderer.h b/Engine/Plugins/FX/Niagara/Source/Niagara/Public/NiagaraRenderer.h
--- a/Engine/Plugins/FX/Niagara/Source/Niagara/Public/NiagaraRenderer.h
+++ b/Engine/Plugins/FX/Niagara/Source/Niagara/Public/NiagaraRenderer.h
@@ -120,6 +120,12 @@ public:
 	FORCEINLINE static FRHIShaderResourceView* GetSrvOrDefaultUInt(FGlobalDynamicReadBuffer::FAllocation& Allocation) { return Allocation.IsValid() ? (FRHIShaderResourceView*)Allocation.SRV : GetDummyUIntBuffer(); }
 	FORCEINLINE static FRHIShaderResourceView* GetSrvOrDefaultInt(FGlobalDynamicReadBuffer::FAllocation& Allocation) { return Allocation.IsValid() ? (FRHIShaderResourceView*)Allocation.SRV : GetDummyIntBuffer(); }
 
+	FORCEINLINE static FRHIShaderResourceView* GetSrvOrDefaultFloat(FRHIShaderResourceView* InSRV) { return InSRV ? InSRV : GetDummyFloatBuffer(); }
+	FORCEINLINE static FRHIShaderResourceView* GetSrvOrDefaultFloat2(FRHIShaderResourceView* InSRV) { return InSRV ? InSRV : GetDummyFloat2Buffer(); }
+	FORCEINLINE static FRHIShaderResourceView* GetSrvOrDefaultFloat4(FRHIShaderResourceView* InSRV) { return InSRV ? InSRV : GetDummyFloat4Buffer(); }
+	FORCEINLINE static FRHIShaderResourceView* GetSrvOrDefaultUInt(FRHIShaderResourceView* InSRV) { return InSRV ? InSRV : GetDummyUIntBuffer(); }
+	FORCEINLINE static FRHIShaderResourceView* GetSrvOrDefaultInt(FRHIShaderResourceView* InSRV) { return InSRV ? InSRV : GetDummyIntBuffer(); }
+
 	static FRHIShaderResourceView* GetDummyFloatBuffer();
 	static FRHIShaderResourceView* GetDummyFloat2Buffer();
 	static FRHIShaderResourceView* GetDummyFloat4Buffer();
