Manual merge

Improved overview inline parameter support, namely for bools and override colors for enum entries

diff --git a/Engine/Plugins/FX/Niagara/Source/Niagara/Private/NiagaraCommon.cpp b/Engine/Plugins/FX/Niagara/Source/Niagara/Private/NiagaraCommon.cpp
--- a/Engine/Plugins/FX/Niagara/Source/Niagara/Private/NiagaraCommon.cpp
+++ b/Engine/Plugins/FX/Niagara/Source/Niagara/Private/NiagaraCommon.cpp
@@ -135,73 +135,6 @@ FString FNiagaraTypeHelper::ToString(const uint8* ValueData, const UObject* Stru
 	return Ret;
 }
 
-FString FNiagaraTypeHelper::ToStringCosmetic(const uint8* ValueData, const UObject* StructOrEnum)
-{
-	FString Ret;
-	if (const UEnum* Enum = Cast<const UEnum>(StructOrEnum))
-	{
-		Ret += Enum->GetDisplayNameTextByValue(*(int32*)ValueData).ToString();
-	}
-	else if (const UScriptStruct* Struct = Cast<const UScriptStruct>(StructOrEnum))
-	{
-		if (Struct == FNiagaraTypeDefinition::GetFloatStruct())
-		{
-			Ret += FString::Printf(TEXT("%g "), *(float*)ValueData);
-		}
-		else if (Struct == FNiagaraTypeDefinition::GetIntStruct())
-		{
-			Ret += FString::Printf(TEXT("%d "), *(int32*)ValueData);
-		}
-		else if (Struct == FNiagaraTypeDefinition::GetBoolStruct())
-		{
-			int32 Val = *(int32*)ValueData;
-			Ret += Val == 0xFFFFFFFF ? (TEXT("True")) : (Val == 0x0 ? TEXT("False") : TEXT("Invalid"));
-		}
-		else
-		{
-			for (TFieldIterator<FProperty> PropertyIt(Struct, EFieldIteratorFlags::IncludeSuper); PropertyIt; ++PropertyIt)
-			{
-				const FProperty* Property = *PropertyIt;
-				const uint8* PropPtr = ValueData + PropertyIt->GetOffset_ForInternal();
-				if (Property->IsA(FFloatProperty::StaticClass()))
-				{
-					Ret += FString::Printf(TEXT("%s: %g "), *Property->GetNameCPP(), *(float*)PropPtr);
-				}
-				else if (Property->IsA(FDoubleProperty::StaticClass()))
-				{
-					Ret += FString::Printf(TEXT("%s: %g "), *Property->GetNameCPP(), *(double*)PropPtr);
-				}
-				else if (Property->IsA(FUInt16Property::StaticClass()))
-				{
-					FFloat16 Val = *(FFloat16*)PropPtr;
-					Ret += FString::Printf(TEXT("%s: %f "), *Property->GetNameCPP(), Val.GetFloat());
-				}
-				else if (Property->IsA(FIntProperty::StaticClass()))
-				{
-					Ret += FString::Printf(TEXT("%s: %d "), *Property->GetNameCPP(), *(int32*)PropPtr);
-				}
-				else if (Property->IsA(FBoolProperty::StaticClass()))
-				{
-					int32 Val = *(int32*)ValueData;
-					FString BoolStr = Val == 0xFFFFFFFF ? (TEXT("True")) : (Val == 0x0 ? TEXT("False") : TEXT("Invalid"));
-					Ret += FString::Printf(TEXT("%s: %d "), *Property->GetNameCPP(), *BoolStr);
-				}
-				else if (const FStructProperty* StructProp = CastFieldChecked<const FStructProperty>(Property))
-				{
-					Ret += FString::Printf(TEXT("%s: (%s) "), *Property->GetNameCPP(), *FNiagaraTypeHelper::ToStringCosmetic(PropPtr, StructProp->Struct));
-				}
-				else
-				{
-					check(false);
-					Ret += TEXT("Unknown Type");
-				}
-			}
-		}
-	}
-	
-	return Ret;
-}
-
 UScriptStruct* FNiagaraTypeHelper::FindNiagaraFriendlyTopLevelStruct(UScriptStruct* InStruct)
 {
 	if (!InStruct)
diff --git a/Engine/Plugins/FX/Niagara/Source/Niagara/Public/NiagaraTypes.h b/Engine/Plugins/FX/Niagara/Source/Niagara/Public/NiagaraTypes.h
--- a/Engine/Plugins/FX/Niagara/Source/Niagara/Public/NiagaraTypes.h
+++ b/Engine/Plugins/FX/Niagara/Source/Niagara/Public/NiagaraTypes.h
@@ -303,7 +303,6 @@ class NIAGARA_API FNiagaraTypeHelper
 {
 public:
 	static FString ToString(const uint8* ValueData, const UObject* StructOrEnum);
-	static FString ToStringCosmetic(const uint8* ValueData, const UObject* StructOrEnum);
 	static UScriptStruct* FindNiagaraFriendlyTopLevelStruct(UScriptStruct* InStruct);
 	static bool IsNiagaraFriendlyTopLevelStruct(UScriptStruct* InStruct);
 };
@@ -625,6 +624,49 @@ struct NIAGARA_API FNiagaraEnumParameterMetaData
 	/** If specified, this icon will be used for the given enum entry. If OverrideName isn't empty, the icon takes priority. */
 	UPROPERTY(EditAnywhere, Category="Enum Override")
 	UTexture2D* IconOverride = nullptr;
+
+	UPROPERTY(EditAnywhere, Category="Enum Override", meta=(InlineEditConditionToggle))
+	bool bUseColorOverride = false;
+	
+	UPROPERTY(EditAnywhere, Category="Enum Override", meta=(EditCondition="bUseColorOverride"))
+	FLinearColor ColorOverride;
+};
+
+UENUM()
+enum class ENiagaraBoolDisplayMode : uint8
+{
+	DisplayAlways,
+	DisplayIfTrue,
+	DisplayIfFalse
+};
+
+USTRUCT()
+struct NIAGARA_API FNiagaraBoolParameterMetaData
+{
+	GENERATED_BODY()
+
+	/** The mode used determines the cases in which a bool parameter is displayed.
+	 *  If set to DisplayAlways, both True and False cases will display. 
+	 *  If set to DisplayIfTrue, it will only display if the bool evaluates to True.
+	 */
+	UPROPERTY(EditAnywhere, Category="Bool Override")
+	ENiagaraBoolDisplayMode DisplayMode = ENiagaraBoolDisplayMode::DisplayAlways;
+	
+	/** If specified, this name will be used for the given bool if it evaluates to True. */
+	UPROPERTY(EditAnywhere, Category="Bool Override")
+	FName OverrideNameTrue;
+
+	/** If specified, this name will be used for the given bool if it evaluates to False. */
+	UPROPERTY(EditAnywhere, Category="Bool Override")
+	FName OverrideNameFalse;
+
+	/** If specified, this icon will be used for the given bool if it evaluates to True. If OverrideName isn't empty, the icon takes priority. */
+	UPROPERTY(EditAnywhere, Category="Bool Override")
+	UTexture2D* IconOverrideTrue = nullptr;
+
+	/** If specified, this icon will be used for the given bool if it evaluates to False. If OverrideName isn't empty, the icon takes priority. */
+	UPROPERTY(EditAnywhere, Category="Bool Override")
+	UTexture2D* IconOverrideFalse = nullptr;
 };
 
 USTRUCT()
@@ -638,6 +680,7 @@ struct NIAGARA_API FNiagaraVariableMetaData
 		, InlineParameterSortPriority(0)
 		, bOverrideColor(false)
 		, InlineParameterColorOverride(FLinearColor(EForceInit::ForceInit))
+		, bEnableBoolOverride(false)
 		, EditorSortPriority(0)
 		, bInlineEditConditionToggle(false)
 		, bIsStaticSwitch_DEPRECATED(false)
@@ -668,7 +711,13 @@ struct NIAGARA_API FNiagaraVariableMetaData
 	FLinearColor InlineParameterColorOverride;
 
 	UPROPERTY(EditAnywhere, Category = "Variable", meta = (EditCondition="bDisplayInOverviewStack", ToolTip = "The index of the entry maps to the index of an enum value. Useful for overriding how an enum parameter is displayed in the overview.", SkipForCompileHash = "true"))
-	TArray<FNiagaraEnumParameterMetaData> InlineParameterEnumOverrides; 
+	TArray<FNiagaraEnumParameterMetaData> InlineParameterEnumOverrides;
+
+	UPROPERTY(EditAnywhere, Category = "Variable", meta = (InlineEditConditionToggle, ToolTip = "Useful to override inline bool visualization in the overview.", SkipForCompileHash = "true"))
+	bool bEnableBoolOverride;
+	
+	UPROPERTY(EditAnywhere, Category = "Variable", meta = (EditCondition="bEnableBoolOverride", ToolTip = "Useful to override inline bool visualization in the overview.", SkipForCompileHash = "true"))
+	FNiagaraBoolParameterMetaData InlineParameterBoolOverride;
 	
 	UPROPERTY(EditAnywhere, Category = "Variable", meta = (ToolTip = "Affects the sort order in the editor stacks. Use a smaller number to push it to the top. Defaults to zero.", SkipForCompileHash = "true"))
 	int32 EditorSortPriority;
@@ -1024,16 +1073,6 @@ public:
 		return FNiagaraTypeHelper::ToString(ValueData, ClassStructOrEnum);
 	}
 
-	FString ToStringCosmetic(const uint8* ValueData)const
-	{
-		checkf(IsValid(), TEXT("Type definition is not valid."));
-		if (ValueData == nullptr)
-		{
-			return TEXT("(null)");
-		}
-		return FNiagaraTypeHelper::ToStringCosmetic(ValueData, ClassStructOrEnum);
-	}
-
 	static bool TypesAreAssignable(const FNiagaraTypeDefinition& TypeA, const FNiagaraTypeDefinition& TypeB);
 	static bool IsLossyConversion(const FNiagaraTypeDefinition& TypeA, const FNiagaraTypeDefinition& TypeB);
 	static FNiagaraTypeDefinition GetNumericOutputType(const TArray<FNiagaraTypeDefinition> TypeDefinintions, ENiagaraNumericOutputTypeSelectionMode SelectionMode);
diff --git a/Engine/Plugins/FX/Niagara/Source/NiagaraEditor/Private/Customizations/NiagaraScriptVariableCustomization.cpp b/Engine/Plugins/FX/Niagara/Source/NiagaraEditor/Private/Customizations/NiagaraScriptVariableCustomization.cpp
--- a/Engine/Plugins/FX/Niagara/Source/NiagaraEditor/Private/Customizations/NiagaraScriptVariableCustomization.cpp
+++ b/Engine/Plugins/FX/Niagara/Source/NiagaraEditor/Private/Customizations/NiagaraScriptVariableCustomization.cpp
@@ -171,28 +171,29 @@ void FNiagaraScriptVariableDetails::CustomizeDetails(IDetailLayoutBuilder& Detai
 		DetailBuilder.HideProperty(GET_MEMBER_NAME_CHECKED(UNiagaraScriptVariable, DefaultBinding));
 	}
 
+	// Always hide the default value variant. We generate a node for this property to enable PostEditChangeProperty(...) events but do not modify it via the default generated customization.
+	DetailBuilder.HideProperty("DefaultValueVariant", UNiagaraScriptVariable::StaticClass());
+
 	if (Variable->Variable.GetType() != FNiagaraTypeDefinition::GetBoolDef())
 	{
 		DetailBuilder.HideProperty(GET_MEMBER_NAME_CHECKED(UNiagaraScriptVariable, Metadata.bInlineEditConditionToggle));
-	}
-
-	// Always hide the default value variant. We generate a node for this property to enable PostEditChangeProperty(...) events but do not modify it via the default generated customization.
-	DetailBuilder.HideProperty("DefaultValueVariant", UNiagaraScriptVariable::StaticClass());
+		DetailBuilder.HideProperty(GET_MEMBER_NAME_CHECKED(UNiagaraScriptVariable, Metadata.InlineParameterBoolOverride));
+	}	
 
 	// we hide the enum overrides if the variable type isn't an enum
 	if(!Variable->Variable.GetType().GetEnum())
 	{
 		DetailBuilder.HideProperty(GET_MEMBER_NAME_CHECKED(UNiagaraScriptVariable, Metadata.InlineParameterEnumOverrides));
 	}
-
-	// @TODO It seems like all the HideProperties aren't working for values within metadata
+	
 	// generally we want all the inline parameters only available for inputs & static switches
 	if(!Variable->Variable.IsInNameSpace(FNiagaraConstants::ModuleNamespace) && !Variable->GetIsStaticSwitch())
 	{
 		DetailBuilder.HideProperty(GET_MEMBER_NAME_CHECKED(UNiagaraScriptVariable, Metadata.bDisplayInOverviewStack));
+		DetailBuilder.HideProperty(GET_MEMBER_NAME_CHECKED(UNiagaraScriptVariable, Metadata.InlineParameterSortPriority));
 		DetailBuilder.HideProperty(GET_MEMBER_NAME_CHECKED(UNiagaraScriptVariable, Metadata.InlineParameterColorOverride));
 		DetailBuilder.HideProperty(GET_MEMBER_NAME_CHECKED(UNiagaraScriptVariable, Metadata.InlineParameterEnumOverrides));
-		DetailBuilder.HideProperty(GET_MEMBER_NAME_CHECKED(UNiagaraScriptVariable, Metadata.InlineParameterSortPriority));
+		DetailBuilder.HideProperty(GET_MEMBER_NAME_CHECKED(UNiagaraScriptVariable, Metadata.InlineParameterBoolOverride));
 	}
 }
 
diff --git a/Engine/Plugins/FX/Niagara/Source/NiagaraEditor/Private/TypeEditorUtilities/NiagaraColorTypeEditorUtilities.cpp b/Engine/Plugins/FX/Niagara/Source/NiagaraEditor/Private/TypeEditorUtilities/NiagaraColorTypeEditorUtilities.cpp
--- a/Engine/Plugins/FX/Niagara/Source/NiagaraEditor/Private/TypeEditorUtilities/NiagaraColorTypeEditorUtilities.cpp
+++ b/Engine/Plugins/FX/Niagara/Source/NiagaraEditor/Private/TypeEditorUtilities/NiagaraColorTypeEditorUtilities.cpp
@@ -235,7 +235,7 @@ FText FNiagaraEditorColorTypeUtilities::GetSearchTextFromValue(const FNiagaraVar
 	return FText::FromString(GetPinDefaultStringFromValue(AllocatedVariable));
 }
 
-FText FNiagaraEditorColorTypeUtilities::GetStackDisplayText(FNiagaraVariable& Variable) const
+FText FNiagaraEditorColorTypeUtilities::GetStackDisplayText(const FNiagaraVariable& Variable) const
 {
 	FLinearColor Value = Variable.GetValue<FLinearColor>();
 	return FText::Format(FText::FromString("({0}, {1}, {2}, {3})"), Value.R, Value.G, Value.B, Value.A);
diff --git a/Engine/Plugins/FX/Niagara/Source/NiagaraEditor/Private/TypeEditorUtilities/NiagaraColorTypeEditorUtilities.h b/Engine/Plugins/FX/Niagara/Source/NiagaraEditor/Private/TypeEditorUtilities/NiagaraColorTypeEditorUtilities.h
--- a/Engine/Plugins/FX/Niagara/Source/NiagaraEditor/Private/TypeEditorUtilities/NiagaraColorTypeEditorUtilities.h
+++ b/Engine/Plugins/FX/Niagara/Source/NiagaraEditor/Private/TypeEditorUtilities/NiagaraColorTypeEditorUtilities.h
@@ -19,5 +19,5 @@ public:
 	virtual FString GetPinDefaultStringFromValue(const FNiagaraVariable& AllocatedVariable) const override;
 	virtual bool SetValueFromPinDefaultString(const FString& StringValue, FNiagaraVariable& Variable) const override;
 	virtual FText GetSearchTextFromValue(const FNiagaraVariable& AllocatedVariable) const override;
-	virtual FText GetStackDisplayText(FNiagaraVariable& Variable) const override;
+	virtual FText GetStackDisplayText(const FNiagaraVariable& Variable) const override;
 };
\ No newline at end of file
diff --git a/Engine/Plugins/FX/Niagara/Source/NiagaraEditor/Private/TypeEditorUtilities/NiagaraEnumTypeEditorUtilities.cpp b/Engine/Plugins/FX/Niagara/Source/NiagaraEditor/Private/TypeEditorUtilities/NiagaraEnumTypeEditorUtilities.cpp
--- a/Engine/Plugins/FX/Niagara/Source/NiagaraEditor/Private/TypeEditorUtilities/NiagaraEnumTypeEditorUtilities.cpp
+++ b/Engine/Plugins/FX/Niagara/Source/NiagaraEditor/Private/TypeEditorUtilities/NiagaraEnumTypeEditorUtilities.cpp
@@ -165,7 +165,7 @@ FText FNiagaraEditorEnumTypeUtilities::GetSearchTextFromValue(const FNiagaraVari
 	return Enum->GetDisplayNameTextByIndex(EnumNameIndex);
 }
 
-FText FNiagaraEditorEnumTypeUtilities::GetStackDisplayText(FNiagaraVariable& Variable) const
+FText FNiagaraEditorEnumTypeUtilities::GetStackDisplayText(const FNiagaraVariable& Variable) const
 {
-	return Variable.GetType().GetEnum()->GetDisplayNameTextByIndex(Variable.GetValue<int32>());
+	return Variable.GetType().GetEnum()->GetDisplayNameTextByValue(Variable.GetValue<int32>());
 }
diff --git a/Engine/Plugins/FX/Niagara/Source/NiagaraEditor/Private/TypeEditorUtilities/NiagaraEnumTypeEditorUtilities.h b/Engine/Plugins/FX/Niagara/Source/NiagaraEditor/Private/TypeEditorUtilities/NiagaraEnumTypeEditorUtilities.h
--- a/Engine/Plugins/FX/Niagara/Source/NiagaraEditor/Private/TypeEditorUtilities/NiagaraEnumTypeEditorUtilities.h
+++ b/Engine/Plugins/FX/Niagara/Source/NiagaraEditor/Private/TypeEditorUtilities/NiagaraEnumTypeEditorUtilities.h
@@ -20,5 +20,5 @@ public:
 	virtual bool CanSetValueFromDisplayName() const override;
 	virtual bool SetValueFromDisplayName(const FText& TextValue, FNiagaraVariable& Variable) const override;
 	virtual FText GetSearchTextFromValue(const FNiagaraVariable& AllocatedVariable) const override;
-	virtual FText GetStackDisplayText(FNiagaraVariable& Variable) const override;
+	virtual FText GetStackDisplayText(const FNiagaraVariable& Variable) const override;
 };
\ No newline at end of file
diff --git a/Engine/Plugins/FX/Niagara/Source/NiagaraEditor/Private/TypeEditorUtilities/NiagaraFloatTypeEditorUtilities.h b/Engine/Plugins/FX/Niagara/Source/NiagaraEditor/Private/TypeEditorUtilities/NiagaraFloatTypeEditorUtilities.h
--- a/Engine/Plugins/FX/Niagara/Source/NiagaraEditor/Private/TypeEditorUtilities/NiagaraFloatTypeEditorUtilities.h
+++ b/Engine/Plugins/FX/Niagara/Source/NiagaraEditor/Private/TypeEditorUtilities/NiagaraFloatTypeEditorUtilities.h
@@ -18,5 +18,5 @@ public:
 	virtual FString GetPinDefaultStringFromValue(const FNiagaraVariable& AllocatedVariable) const override;
 	virtual bool SetValueFromPinDefaultString(const FString& StringValue, FNiagaraVariable& Variable) const override;
 	virtual FText GetSearchTextFromValue(const FNiagaraVariable& AllocatedVariable) const override;
-	virtual FText GetStackDisplayText(FNiagaraVariable& Variable) const override { return FText::AsNumber(Variable.GetValue<float>()); };
+	virtual FText GetStackDisplayText(const FNiagaraVariable& Variable) const override { return FText::AsNumber(Variable.GetValue<float>()); };
 };
\ No newline at end of file
diff --git a/Engine/Plugins/FX/Niagara/Source/NiagaraEditor/Private/TypeEditorUtilities/NiagaraVectorTypeEditorUtilities.cpp b/Engine/Plugins/FX/Niagara/Source/NiagaraEditor/Private/TypeEditorUtilities/NiagaraVectorTypeEditorUtilities.cpp
--- a/Engine/Plugins/FX/Niagara/Source/NiagaraEditor/Private/TypeEditorUtilities/NiagaraVectorTypeEditorUtilities.cpp
+++ b/Engine/Plugins/FX/Niagara/Source/NiagaraEditor/Private/TypeEditorUtilities/NiagaraVectorTypeEditorUtilities.cpp
@@ -180,7 +180,7 @@ FText FNiagaraEditorVector2TypeUtilities::GetSearchTextFromValue(const FNiagaraV
 	return FText::FromString(GetPinDefaultStringFromValue(AllocatedVariable));
 }
 
-FText FNiagaraEditorVector2TypeUtilities::GetStackDisplayText(FNiagaraVariable& Variable) const
+FText FNiagaraEditorVector2TypeUtilities::GetStackDisplayText(const FNiagaraVariable& Variable) const
 {
 	FVector2f Value = Variable.GetValue<FVector2f>();
 	return FText::Format(FText::FromString("({0}, {1})"), Value.X, Value.Y);
@@ -262,7 +262,7 @@ FText FNiagaraEditorVector3TypeUtilities::GetSearchTextFromValue(const FNiagaraV
 	return FText::FromString(GetPinDefaultStringFromValue(AllocatedVariable));
 }
 
-FText FNiagaraEditorVector3TypeUtilities::GetStackDisplayText(FNiagaraVariable& Variable) const
+FText FNiagaraEditorVector3TypeUtilities::GetStackDisplayText(const FNiagaraVariable& Variable) const
 {
 	FVector3f Value = Variable.GetValue<FVector3f>();
 	return FText::Format(FText::FromString("({0}, {1}, {2})"), Value.X, Value.Y, Value.Z);
@@ -344,7 +344,7 @@ FText FNiagaraEditorVector4TypeUtilities::GetSearchTextFromValue(const FNiagaraV
 	return FText::FromString(GetPinDefaultStringFromValue(AllocatedVariable));
 }
 
-FText FNiagaraEditorVector4TypeUtilities::GetStackDisplayText(FNiagaraVariable& Variable) const
+FText FNiagaraEditorVector4TypeUtilities::GetStackDisplayText(const FNiagaraVariable& Variable) const
 {
 	FVector4f Value = Variable.GetValue<FVector4f>();
 	return FText::Format(FText::FromString("({0}, {1}, {2}, {3})"), Value.X, Value.Y, Value.Z, Value.W);
@@ -455,7 +455,7 @@ FText FNiagaraEditorQuatTypeUtilities::GetSearchTextFromValue(const FNiagaraVari
 	return FText::FromString(GetPinDefaultStringFromValue(AllocatedVariable));
 }
 
-FText FNiagaraEditorQuatTypeUtilities::GetStackDisplayText(FNiagaraVariable& Variable) const
+FText FNiagaraEditorQuatTypeUtilities::GetStackDisplayText(const FNiagaraVariable& Variable) const
 {
 	FQuat4f Value = Variable.GetValue<FQuat4f>();
 	return FText::Format(FText::FromString("({0}, {1}, {2}, {3})"), Value.X, Value.Y, Value.Z, Value.W);
diff --git a/Engine/Plugins/FX/Niagara/Source/NiagaraEditor/Private/TypeEditorUtilities/NiagaraVectorTypeEditorUtilities.h b/Engine/Plugins/FX/Niagara/Source/NiagaraEditor/Private/TypeEditorUtilities/NiagaraVectorTypeEditorUtilities.h
--- a/Engine/Plugins/FX/Niagara/Source/NiagaraEditor/Private/TypeEditorUtilities/NiagaraVectorTypeEditorUtilities.h
+++ b/Engine/Plugins/FX/Niagara/Source/NiagaraEditor/Private/TypeEditorUtilities/NiagaraVectorTypeEditorUtilities.h
@@ -17,7 +17,7 @@ public:
 	virtual FString GetPinDefaultStringFromValue(const FNiagaraVariable& AllocatedVariable) const override;
 	virtual bool SetValueFromPinDefaultString(const FString& StringValue, FNiagaraVariable& Variable) const override;
 	virtual FText GetSearchTextFromValue(const FNiagaraVariable& AllocatedVariable) const override;
-	virtual FText GetStackDisplayText(FNiagaraVariable& Variable) const override;
+	virtual FText GetStackDisplayText(const FNiagaraVariable& Variable) const override;
 };
 
 /** Niagara editor utilities for the FVector type. */
@@ -31,7 +31,7 @@ public:
 	virtual FString GetPinDefaultStringFromValue(const FNiagaraVariable& AllocatedVariable) const override;
 	virtual bool SetValueFromPinDefaultString(const FString& StringValue, FNiagaraVariable& Variable) const override;
 	virtual FText GetSearchTextFromValue(const FNiagaraVariable& AllocatedVariable) const override;
-	virtual FText GetStackDisplayText(FNiagaraVariable& Variable) const override;
+	virtual FText GetStackDisplayText(const FNiagaraVariable& Variable) const override;
 };
 
 /** Niagara editor utilities for the FVector4 type. */
@@ -45,7 +45,7 @@ public:
 	virtual FString GetPinDefaultStringFromValue(const FNiagaraVariable& AllocatedVariable) const override;
 	virtual bool SetValueFromPinDefaultString(const FString& StringValue, FNiagaraVariable& Variable) const override;
 	virtual FText GetSearchTextFromValue(const FNiagaraVariable& AllocatedVariable) const override;
-	virtual FText GetStackDisplayText(FNiagaraVariable& Variable) const override;
+	virtual FText GetStackDisplayText(const FNiagaraVariable& Variable) const override;
 };
 
 /** Niagara editor utilities for the FVector4 type. */
@@ -61,7 +61,7 @@ public:
 	virtual FString GetPinDefaultStringFromValue(const FNiagaraVariable& AllocatedVariable) const override;
 	virtual bool SetValueFromPinDefaultString(const FString& StringValue, FNiagaraVariable& Variable) const override;
 	virtual FText GetSearchTextFromValue(const FNiagaraVariable& AllocatedVariable) const override;
-	virtual FText GetStackDisplayText(FNiagaraVariable& Variable) const override;
+	virtual FText GetStackDisplayText(const FNiagaraVariable& Variable) const override;
 };
 
 /** Niagara editor utilities for the FNiagaraID type. */
diff --git a/Engine/Plugins/FX/Niagara/Source/NiagaraEditor/Private/ViewModels/Stack/NiagaraStackFunctionInput.cpp b/Engine/Plugins/FX/Niagara/Source/NiagaraEditor/Private/ViewModels/Stack/NiagaraStackFunctionInput.cpp
--- a/Engine/Plugins/FX/Niagara/Source/NiagaraEditor/Private/ViewModels/Stack/NiagaraStackFunctionInput.cpp
+++ b/Engine/Plugins/FX/Niagara/Source/NiagaraEditor/Private/ViewModels/Stack/NiagaraStackFunctionInput.cpp
@@ -215,7 +215,7 @@ UNiagaraScript* UNiagaraStackFunctionInput::GetInputFunctionCallInitialScript()
 	return OwningFunctionCallInitialScript.Get();
 }
 
-UNiagaraStackFunctionInput::EValueMode UNiagaraStackFunctionInput::GetValueMode()
+UNiagaraStackFunctionInput::EValueMode UNiagaraStackFunctionInput::GetValueMode() const
 {
 	return InputValues.Mode;
 }
@@ -381,6 +381,11 @@ TArray<UNiagaraStackFunctionInput*> UNiagaraStackFunctionInput::GetChildInputs()
 	return ChildInputs;
 }
 
+TOptional<FNiagaraVariableMetaData> UNiagaraStackFunctionInput::GetInputMetaData() const
+{
+	return InputMetaData;
+}
+
 void UNiagaraStackFunctionInput::GetCurrentChangeIds(FGuid& OutOwningGraphChangeId, FGuid& OutFunctionGraphChangeId) const
 {
 	OutOwningGraphChangeId = OwningFunctionCallNode->GetNiagaraGraph()->GetChangeID();
@@ -2379,6 +2384,17 @@ bool UNiagaraStackFunctionInput::IsScratchDynamicInput() const
 	return bIsScratchDynamicInputCache.GetValue();
 }
 
+bool UNiagaraStackFunctionInput::ShouldDisplayInline() const
+{
+	// we don't allow advanced values to be inlined as they can introduce problems with filtering
+	if(InputValues.Mode == EValueMode::Local && InputMetaData.IsSet() && !GetIsAdvanced() && InputMetaData->bDisplayInOverviewStack)
+	{
+		return true;
+	}
+
+	return false;
+}
+
 bool UNiagaraStackFunctionInput::IsSemanticChild() const
 {
 	return bIsSemanticChild;
diff --git a/Engine/Plugins/FX/Niagara/Source/NiagaraEditor/Private/ViewModels/Stack/NiagaraStackFunctionInputCollection.cpp b/Engine/Plugins/FX/Niagara/Source/NiagaraEditor/Private/ViewModels/Stack/NiagaraStackFunctionInputCollection.cpp
--- a/Engine/Plugins/FX/Niagara/Source/NiagaraEditor/Private/ViewModels/Stack/NiagaraStackFunctionInputCollection.cpp
+++ b/Engine/Plugins/FX/Niagara/Source/NiagaraEditor/Private/ViewModels/Stack/NiagaraStackFunctionInputCollection.cpp
@@ -155,58 +155,22 @@ void UNiagaraStackFunctionInputCollection::GetFilteredChildInputs(TArray<UNiagar
 	}
 }
 
-TOptional<FNiagaraLocalInputValueData> UNiagaraStackFunctionInputCollection::GetLocalInput(FNiagaraVariable InputParameter, bool bFiltered) const
-{
-	TArray<UNiagaraStackFunctionInput*> FunctionInputs;
-
-	if(bFiltered)
-	{
-		GetFilteredChildInputs(FunctionInputs);
-	}
-	else
-	{
-		GetChildInputs(FunctionInputs);
-	}
-
-	for(UNiagaraStackFunctionInput* FunctionInput : FunctionInputs)
-	{
-		if(FunctionInput->GetValueMode() == UNiagaraStackFunctionInput::EValueMode::Local && FunctionInput->GetInputParameterHandle().GetParameterHandleString() == InputParameter.GetName())
-		{
-			FNiagaraLocalInputValueData InputValueData(FunctionInput->GetLocalValueStruct(), FunctionInput->OnValueChanged());
-			FunctionInput->GetSearchItems(InputValueData.StackSearchItems);
-			return InputValueData;
-		}
-	}
-
-	TOptional<FNiagaraLocalInputValueData> UnsetValueData;
-	return UnsetValueData;
-}
-
-TOptional<FNiagaraDataInterfaceInput> UNiagaraStackFunctionInputCollection::GetDataInterfaceForInput(FNiagaraVariable InputParameter, bool bFiltered) const
+TArray<UNiagaraStackFunctionInput*> UNiagaraStackFunctionInputCollection::GetInlineParameterInputs() const
 {
+	TArray<UNiagaraStackFunctionInput*> OutArray;
+	
 	TArray<UNiagaraStackFunctionInput*> FunctionInputs;
-
-	if(bFiltered)
-	{
-		GetFilteredChildInputs(FunctionInputs);
-	}
-	else
-	{
-		GetChildInputs(FunctionInputs);
-	}
+	GetFilteredChildInputs(FunctionInputs);
 
 	for(UNiagaraStackFunctionInput* FunctionInput : FunctionInputs)
 	{
-		if(FunctionInput->GetValueMode() == UNiagaraStackFunctionInput::EValueMode::Data && FunctionInput->GetInputParameterHandle().GetParameterHandleString() == InputParameter.GetName())
+		if(FunctionInput->ShouldDisplayInline())
 		{
-			FNiagaraDataInterfaceInput InputValueData(FunctionInput->GetDataValueObject(), FunctionInput->OnValueChanged());
-			FunctionInput->GetSearchItems(InputValueData.StackSearchItems);
-			return TOptional<FNiagaraDataInterfaceInput>(InputValueData);
+			OutArray.Add(FunctionInput);
 		}
 	}
 
-	TOptional<FNiagaraDataInterfaceInput> UnsetValueData;
-	return UnsetValueData;
+	return OutArray;
 }
 
 void UNiagaraStackFunctionInputCollection::RefreshChildrenInternal(const TArray<UNiagaraStackEntry*>& CurrentChildren, TArray<UNiagaraStackEntry*>& NewChildren, TArray<FStackIssue>& NewIssues)
diff --git a/Engine/Plugins/FX/Niagara/Source/NiagaraEditor/Private/ViewModels/Stack/NiagaraStackModuleItem.cpp b/Engine/Plugins/FX/Niagara/Source/NiagaraEditor/Private/ViewModels/Stack/NiagaraStackModuleItem.cpp
--- a/Engine/Plugins/FX/Niagara/Source/NiagaraEditor/Private/ViewModels/Stack/NiagaraStackModuleItem.cpp
+++ b/Engine/Plugins/FX/Niagara/Source/NiagaraEditor/Private/ViewModels/Stack/NiagaraStackModuleItem.cpp
@@ -1275,14 +1275,9 @@ void UNiagaraStackModuleItem::SetInputValuesFromClipboardFunctionInputs(const TA
 	InputCollection->SetValuesFromClipboardFunctionInputs(ClipboardFunctionInputs);
 }
 
-TOptional<FNiagaraLocalInputValueData> UNiagaraStackModuleItem::GetLocalInputData(FNiagaraVariable InputParameter, bool bFiltered) const
+TArray<UNiagaraStackFunctionInput*> UNiagaraStackModuleItem::GetInlineParameterInputs() const
 {
-	return InputCollection->GetLocalInput(InputParameter, bFiltered);
-}
-
-TOptional<FNiagaraDataInterfaceInput> UNiagaraStackModuleItem::GetDataInterfaceForInput(FNiagaraVariable InputParameter, bool bFiltered) const
-{
-	return InputCollection->GetDataInterfaceForInput(InputParameter, bFiltered);	
+	return InputCollection->GetInlineParameterInputs();
 }
 
 bool UNiagaraStackModuleItem::TestCanCutWithMessage(FText& OutMessage) const
diff --git a/Engine/Plugins/FX/Niagara/Source/NiagaraEditor/Public/INiagaraEditorTypeUtilities.h b/Engine/Plugins/FX/Niagara/Source/NiagaraEditor/Public/INiagaraEditorTypeUtilities.h
--- a/Engine/Plugins/FX/Niagara/Source/NiagaraEditor/Public/INiagaraEditorTypeUtilities.h
+++ b/Engine/Plugins/FX/Niagara/Source/NiagaraEditor/Public/INiagaraEditorTypeUtilities.h
@@ -42,7 +42,7 @@ public:
 
 	virtual FText GetSearchTextFromValue(const FNiagaraVariable& AllocatedVariable) const = 0;
 
-	virtual FText GetStackDisplayText(FNiagaraVariable& Variable) const = 0;
+	virtual FText GetStackDisplayText(const FNiagaraVariable& Variable) const = 0;
 };
 
 class FNiagaraEditorTypeUtilities : public INiagaraEditorTypeUtilities, public TSharedFromThis<FNiagaraEditorTypeUtilities, ESPMode::ThreadSafe>
@@ -63,7 +63,7 @@ public:
 	virtual bool CanSetValueFromDisplayName() const override { return false; }
 	virtual bool SetValueFromDisplayName(const FText& TextValue, FNiagaraVariable& Variable) const override { return false; }
 	virtual FText GetSearchTextFromValue(const FNiagaraVariable& AllocatedVariable) const override { return FText(); }
-	virtual FText GetStackDisplayText(FNiagaraVariable& Variable) const override
+	virtual FText GetStackDisplayText(const FNiagaraVariable& Variable) const override
 	{
 		FString DefaultString = GetPinDefaultStringFromValue(Variable);
 		return FText::FromString(DefaultString.IsEmpty() ? "[?]" : DefaultString);
diff --git a/Engine/Plugins/FX/Niagara/Source/NiagaraEditor/Public/ViewModels/Stack/NiagaraStackFunctionInput.h b/Engine/Plugins/FX/Niagara/Source/NiagaraEditor/Public/ViewModels/Stack/NiagaraStackFunctionInput.h
--- a/Engine/Plugins/FX/Niagara/Source/NiagaraEditor/Public/ViewModels/Stack/NiagaraStackFunctionInput.h
+++ b/Engine/Plugins/FX/Niagara/Source/NiagaraEditor/Public/ViewModels/Stack/NiagaraStackFunctionInput.h
@@ -86,7 +86,7 @@ public:
 	UNiagaraScript* GetInputFunctionCallInitialScript() const;
 
 	/** Gets the current value mode */
-	EValueMode GetValueMode();
+	EValueMode GetValueMode() const;
 
 	/** Gets the type of this input. */
 	const FNiagaraTypeDefinition& GetInputType() const;
@@ -260,8 +260,12 @@ public:
 
 	bool IsScratchDynamicInput() const;
 
+	bool ShouldDisplayInline() const;
+	
 	TArray<UNiagaraStackFunctionInput*> GetChildInputs() const;
 
+	TOptional<FNiagaraVariableMetaData> GetInputMetaData() const;
+	
 	virtual bool IsSemanticChild() const;
 	void SetSemanticChild(bool IsSemanticChild);
 
diff --git a/Engine/Plugins/FX/Niagara/Source/NiagaraEditor/Public/ViewModels/Stack/NiagaraStackFunctionInputCollection.h b/Engine/Plugins/FX/Niagara/Source/NiagaraEditor/Public/ViewModels/Stack/NiagaraStackFunctionInputCollection.h
--- a/Engine/Plugins/FX/Niagara/Source/NiagaraEditor/Public/ViewModels/Stack/NiagaraStackFunctionInputCollection.h
+++ b/Engine/Plugins/FX/Niagara/Source/NiagaraEditor/Public/ViewModels/Stack/NiagaraStackFunctionInputCollection.h
@@ -11,26 +11,6 @@ class UNiagaraStackFunctionInput;
 class UNiagaraClipboardFunctionInput;
 class UEdGraphPin;
 
-struct NIAGARAEDITOR_API FNiagaraLocalInputValueData
-{
-	FNiagaraLocalInputValueData(TSharedPtr<const FStructOnScope> InLocalData, TMulticastDelegate<void()>& InOnValueChanged) :
-	LocalData(InLocalData), OnValueChanged(InOnValueChanged)
-	{}
-	TSharedPtr<const FStructOnScope> LocalData;
-	TMulticastDelegate<void()>& OnValueChanged;
-	TArray<UNiagaraStackEntry::FStackSearchItem> StackSearchItems;
-};
-
-struct NIAGARAEDITOR_API FNiagaraDataInterfaceInput
-{
-	FNiagaraDataInterfaceInput(TWeakObjectPtr<const UNiagaraDataInterface> InDataInterface, TMulticastDelegate<void()>& InOnValueChanged) :
-	DataInterface(InDataInterface), OnValueChanged(InOnValueChanged)
-	{}
-	TWeakObjectPtr<const UNiagaraDataInterface> DataInterface;
-	TMulticastDelegate<void()>& OnValueChanged;
-	TArray<UNiagaraStackEntry::FStackSearchItem> StackSearchItems;
-};
-
 
 
 UCLASS()
@@ -121,9 +101,7 @@ public:
 
 	void GetFilteredChildInputs(TArray<UNiagaraStackFunctionInput*>& OutResult) const;
 
-	TOptional<FNiagaraLocalInputValueData> GetLocalInput(FNiagaraVariable InputParameter, bool bFiltered = false) const;
-
-	TOptional<FNiagaraDataInterfaceInput> GetDataInterfaceForInput(FNiagaraVariable InputParameter, bool bFiltered = false) const;
+	TArray<UNiagaraStackFunctionInput*> GetInlineParameterInputs() const;
 
 protected:
 	virtual void FinalizeInternal() override;
diff --git a/Engine/Plugins/FX/Niagara/Source/NiagaraEditor/Public/ViewModels/Stack/NiagaraStackModuleItem.h b/Engine/Plugins/FX/Niagara/Source/NiagaraEditor/Public/ViewModels/Stack/NiagaraStackModuleItem.h
--- a/Engine/Plugins/FX/Niagara/Source/NiagaraEditor/Public/ViewModels/Stack/NiagaraStackModuleItem.h
+++ b/Engine/Plugins/FX/Niagara/Source/NiagaraEditor/Public/ViewModels/Stack/NiagaraStackModuleItem.h
@@ -74,8 +74,7 @@ public:
 
 	void SetInputValuesFromClipboardFunctionInputs(const TArray<const UNiagaraClipboardFunctionInput*>& ClipboardFunctionInputs);
 
-	TOptional<struct FNiagaraLocalInputValueData> GetLocalInputData(FNiagaraVariable InputParameter, bool bFiltered = false) const;
-	TOptional<struct FNiagaraDataInterfaceInput> GetDataInterfaceForInput(FNiagaraVariable InputParameter, bool bFiltered = false) const;
+	TArray<class UNiagaraStackFunctionInput*> GetInlineParameterInputs() const;
 
 	virtual bool SupportsCut() const override { return true; }
 	virtual bool TestCanCutWithMessage(FText& OutMessage) const override;
diff --git a/Engine/Plugins/FX/Niagara/Source/NiagaraEditorWidgets/Private/NiagaraEditorWidgetsStyle.cpp b/Engine/Plugins/FX/Niagara/Source/NiagaraEditorWidgets/Private/NiagaraEditorWidgetsStyle.cpp
--- a/Engine/Plugins/FX/Niagara/Source/NiagaraEditorWidgets/Private/NiagaraEditorWidgetsStyle.cpp
+++ b/Engine/Plugins/FX/Niagara/Source/NiagaraEditorWidgets/Private/NiagaraEditorWidgetsStyle.cpp
@@ -153,12 +153,54 @@ TSharedRef< FSlateStyleSet > FNiagaraEditorWidgetsStyle::Create()
 	FSlateFontInfo SystemOverviewInlineParametersFont = DEFAULT_FONT("Regular", 7);
 	FTextBlockStyle SystemOverviewInlineParametersText = FTextBlockStyle(NormalText)
 		.SetFont(SystemOverviewInlineParametersFont);
+
+	FButtonStyle InlineParameterButtonStyle = FButtonStyle()
+		.SetNormal(FSlateRoundedBoxBrush(FStyleColors::AccentGray, 4.0f, FStyleColors::Input, 1.f))
+		.SetHovered(FSlateRoundedBoxBrush(FStyleColors::White, 4.0f, FStyleColors::Input, 1.f))
+		.SetPressed(FSlateRoundedBoxBrush(FStyleColors::Header, 4.0f, FStyleColors::Input, 1.f))
+		.SetDisabled(FSlateRoundedBoxBrush(FStyleColors::Dropdown, 4.0f, FStyleColors::Recessed, 1.f))
+		.SetNormalForeground(FStyleColors::ForegroundHover)
+		.SetHoveredForeground(FStyleColors::ForegroundInverted)
+		.SetPressedForeground(FStyleColors::ForegroundHover)
+		.SetDisabledForeground(FStyleColors::Foreground)
+		.SetNormalPadding(FMargin(2.f, 2.f, 2.f, 2.f))
+		.SetPressedPadding(FMargin(2.f, 1.0f, 2.f, 0.0f));
+
+	// this style is used when the user has specified a custom color for an inline parameter. By default we'd tint it a bit more gray.
+	FButtonStyle InlineParameterButtonStyleNoTint = FButtonStyle()
+		.SetNormal(FSlateRoundedBoxBrush(FStyleColors::White, 4.0f, FStyleColors::Input, 1.f))
+		.SetHovered(FSlateRoundedBoxBrush(FStyleColors::White, 4.0f, FStyleColors::Input, 1.f))
+		.SetPressed(FSlateRoundedBoxBrush(FStyleColors::Header, 4.0f, FStyleColors::Input, 1.f))
+		.SetDisabled(FSlateRoundedBoxBrush(FStyleColors::Dropdown, 4.0f, FStyleColors::Recessed, 1.f))
+		.SetNormalForeground(FStyleColors::ForegroundHover)
+		.SetHoveredForeground(FStyleColors::ForegroundInverted)
+		.SetPressedForeground(FStyleColors::ForegroundHover)
+		.SetDisabledForeground(FStyleColors::Foreground)
+		.SetNormalPadding(FMargin(2.f, 2.f, 2.f, 2.f))
+		.SetPressedPadding(FMargin(2.f, 1.0f, 2.f, 0.0f));
+
+	// this style is used when the user has specified an icon to represent a value. In that case, we no longer need a proper border.
+	FButtonStyle InlineParameterButtonStyleTransparent = FButtonStyle()
+		.SetNormal(FSlateRoundedBoxBrush(FStyleColors::Transparent, 4.0f, FStyleColors::Input, 0.f))
+		.SetHovered(FSlateRoundedBoxBrush(FStyleColors::Hover, 4.0f, FStyleColors::Input, 0.f))
+		.SetPressed(FSlateRoundedBoxBrush(FStyleColors::Header, 4.0f, FStyleColors::Input, 0.f))
+		.SetDisabled(FSlateRoundedBoxBrush(FStyleColors::Dropdown, 4.0f, FStyleColors::Recessed, 0.f))
+		.SetNormalForeground(FStyleColors::ForegroundHover)
+		.SetHoveredForeground(FStyleColors::ForegroundHover)
+		.SetPressedForeground(FStyleColors::ForegroundHover)
+		.SetDisabledForeground(FStyleColors::Foreground)
+		.SetNormalPadding(FMargin(2.f))
+		.SetPressedPadding(FMargin(2.f, 1.0f, 2.f, 0.0f));
 	
 	Style->Set("NiagaraEditor.SystemOverview.ItemText", SystemOverviewItemText);
 	
 	Style->Set("NiagaraEditor.SystemOverview.AlternateItemText", SystemOverviewAlternateItemText);
 
 	Style->Set("NiagaraEditor.SystemOverview.InlineParameterText", SystemOverviewInlineParametersText);
+	Style->Set("NiagaraEditor.SystemOverview.InlineParameterButton", InlineParameterButtonStyle);
+	Style->Set("NiagaraEditor.SystemOverview.InlineParameterButton.Transparent", InlineParameterButtonStyleTransparent);
+	Style->Set("NiagaraEditor.SystemOverview.InlineParameterButton.NoTint", InlineParameterButtonStyleNoTint);
+
 
 	Style->Set("NiagaraEditor.SystemOverview.Item.BackgroundColor", FLinearColor(FColor(62, 62, 62)));
 	Style->Set("NiagaraEditor.SystemOverview.Group.BackgroundColor", FLinearColor::Transparent);
diff --git a/Engine/Plugins/FX/Niagara/Source/NiagaraEditorWidgets/Private/SNiagaraOverviewInlineParameterBox.cpp b/Engine/Plugins/FX/Niagara/Source/NiagaraEditorWidgets/Private/SNiagaraOverviewInlineParameterBox.cpp
--- a/Engine/Plugins/FX/Niagara/Source/NiagaraEditorWidgets/Private/SNiagaraOverviewInlineParameterBox.cpp
+++ b/Engine/Plugins/FX/Niagara/Source/NiagaraEditorWidgets/Private/SNiagaraOverviewInlineParameterBox.cpp
@@ -2,25 +2,23 @@
 
 #include "SNiagaraOverviewInlineParameterBox.h"
 
+#include "INiagaraEditorTypeUtilities.h"
 #include "NiagaraEditorWidgetsStyle.h"
 #include "NiagaraNodeFunctionCall.h"
 #include "NiagaraScriptSource.h"
 #include "NiagaraScriptVariable.h"
 #include "NiagaraSystem.h"
-#include "NiagaraComponent.h"
-#include "NiagaraStackEditorData.h"
+#include "NiagaraEditorModule.h"
 #include "Styling/StyleColors.h"
-#include "ViewModels/NiagaraEmitterViewModel.h"
 #include "ViewModels/NiagaraSystemSelectionViewModel.h"
 #include "ViewModels/NiagaraSystemViewModel.h"
-#include "Widgets/Layout/SWrapBox.h"
-#include "Widgets/Text/STextBlock.h"
-#include "Widgets/Input/SButton.h"
-#include "Widgets/SToolTip.h"
 #include "ViewModels/Stack/NiagaraStackFunctionInput.h"
-#include "ViewModels/Stack/NiagaraStackFunctionInputCollection.h"
 #include "Widgets/SNiagaraParameterName.h"
 #include "Widgets/SNiagaraPinTypeSelector.h"
+#include "Widgets/SToolTip.h"
+#include "Widgets/Text/STextBlock.h"
+#include "Widgets/Input/SButton.h"
+#include "Widgets/Layout/SWrapBox.h"
 #include "Widgets/Layout/SScrollBox.h"
 
 #define LOCTEXT_NAMESPACE "SNiagaraOverviewInlineParameterBox"
@@ -45,25 +43,37 @@ void SNiagaraOverviewInlineParameterBox::Construct(const FArguments& InArgs, UNi
 
 SNiagaraOverviewInlineParameterBox::~SNiagaraOverviewInlineParameterBox()
 {
-	for(FNiagaraLocalInputValueData& BoundInput : BoundFunctionInputs)
+	for(TWeakObjectPtr<UNiagaraStackFunctionInput> BoundInput : BoundFunctionInputs)
 	{
-		BoundInput.OnValueChanged.RemoveAll(this);
+		if(BoundInput.IsValid())
+		{
+			BoundInput->OnValueChanged().RemoveAll(this);
+		}		
 	}
+	
+	BoundFunctionInputs.Empty();
 }
 
-FReply SNiagaraOverviewInlineParameterBox::NavigateToStack(FNiagaraLocalInputValueData LocalInputValueData)
+int32 SNiagaraOverviewInlineParameterBox::GetNumChildWidgets()
+{
+	return NumParameterWidgets;
+}
+
+FReply SNiagaraOverviewInlineParameterBox::NavigateToStack(TWeakObjectPtr<UNiagaraStackFunctionInput> FunctionInput)
 {
 	// even if we can't navigate, the button should consume the input
-	if(!ModuleItem.IsValid())
+	if(!ModuleItem.IsValid() || !FunctionInput.IsValid() || FunctionInput->IsFinalized())
 	{
 		return FReply::Handled();
 	}
 	
 	ModuleItem->GetSystemViewModel()->GetSelectionViewModel()->UpdateSelectedEntries({ModuleItem.Get()}, {}, true);
-	
-	if(ensure(LocalInputValueData.StackSearchItems.Num() > 0))
+
+	TArray<UNiagaraStackEntry::FStackSearchItem> SearchItems;
+	FunctionInput->GetSearchItems(SearchItems);
+	if(ensure(SearchItems.Num() > 0))
 	{
-		ModuleItem->GetSystemViewModel()->GetSelectionViewModel()->GetSelectionStackViewModel()->SetSearchTextExternal(LocalInputValueData.StackSearchItems[0].Value);
+		ModuleItem->GetSystemViewModel()->GetSelectionViewModel()->GetSelectionStackViewModel()->SetSearchTextExternal(SearchItems[0].Value);
 		return FReply::Handled();
 	}
 
@@ -73,21 +83,34 @@ FReply SNiagaraOverviewInlineParameterBox::NavigateToStack(FNiagaraLocalInputVal
 void SNiagaraOverviewInlineParameterBox::ConstructChildren()
 {
 	Container->ClearChildren();
+	NumParameterWidgets = 0;
 
+	for(TWeakObjectPtr<UNiagaraStackFunctionInput> BoundInput : BoundFunctionInputs)
+	{
+		if(BoundInput.IsValid())
+		{
+			BoundInput->OnValueChanged().RemoveAll(this);
+		}
+	}
+	BoundFunctionInputs.Empty();
+	
 	TArray<TSharedRef<SWidget>> ParameterWidgets = GenerateParameterWidgets();
 	
 	for (TSharedRef<SWidget> ParameterWidget : ParameterWidgets)
 	{
 		Container->AddSlot()
 		.VAlign(VAlign_Center)
-		.HAlign(HAlign_Right)
-		.Padding(FMargin(2.f, 1.f))
+		.HAlign(HAlign_Center)
+		.Padding(FMargin(2.f, 0.f))
 		[
 			ParameterWidget
 		];
+
+		NumParameterWidgets++;
 	}
 
-	Container->SetVisibility(ParameterWidgets.Num() > 0 ? EVisibility::Visible : EVisibility::Collapsed);
+	SetVisibility(ParameterWidgets.Num() > 0 ? EVisibility::Visible : EVisibility::Collapsed);
+	Container->SetEnabled(ModuleItem->GetIsEnabled());
 }
 
 TArray<TSharedRef<SWidget>> SNiagaraOverviewInlineParameterBox::GenerateParameterWidgets()
@@ -109,44 +132,34 @@ TArray<TSharedRef<SWidget>> SNiagaraOverviewInlineParameterBox::GenerateParamete
 
 	// we cache the sort order of the widgets we create to sort them after all widgets have been created
 	TMap<TSharedRef<SWidget>, int32> SortOrder;
-
-	// we go through all script variables of the source graph and check if a parameter is supposed to be displayed in the overview
-	const TMap<FNiagaraVariable, TObjectPtr<UNiagaraScriptVariable>>& AllMetaData = ScriptSource->NodeGraph->GetAllMetaData();
-	for(const auto& ParameterMetaData : AllMetaData)
+	
+	TArray<UNiagaraStackFunctionInput*> FunctionInputs = ModuleItem->GetInlineParameterInputs();
+	for(UNiagaraStackFunctionInput* FunctionInput : FunctionInputs)
 	{
-		const UNiagaraScriptVariable* ScriptVariable = ParameterMetaData.Value;
-
-		if(ScriptVariable == nullptr || !ScriptVariable->Metadata.bDisplayInOverviewStack)
+		FunctionInput->OnValueChanged().RemoveAll(this);
+		FunctionInput->OnValueChanged().AddSP(this, &SNiagaraOverviewInlineParameterBox::ConstructChildren);
+		BoundFunctionInputs.Add(FunctionInput);
+		
+		TSharedPtr<SWidget> Widget;
+		if(FunctionInput->GetValueMode() == UNiagaraStackFunctionInput::EValueMode::Local)
 		{
-			continue;
+			Widget = GenerateParameterWidgetFromLocalValue(FunctionInput);
 		}
-
-		TOptional<FNiagaraLocalInputValueData> LocalInputValueData = ModuleItem->GetLocalInputData(ParameterMetaData.Key, true);
-		TOptional<FNiagaraDataInterfaceInput> InputDataInterface = ModuleItem->GetDataInterfaceForInput(ParameterMetaData.Key, true);
-
-		// we only support inline display of local values and data interfaces currently
-		// @TODO use the commented out line when DIs support this too.
-		// if(!LocalInputValueData.IsSet() && !InputDataInterface.IsSet())
-		if(!LocalInputValueData.IsSet())
+		// @todo currently data mode inputs won't be retrieved by GetInlineParameterInputs
+		else if(FunctionInput->GetValueMode() == UNiagaraStackFunctionInput::EValueMode::Data)
 		{
-			continue;
+			Widget = GenerateParameterWidgetFromDataInterface(FunctionInput);
 		}
 
-		TSharedPtr<SWidget> Widget;
-		if(LocalInputValueData.IsSet())
+		if(Widget == SNullWidget::NullWidget)
 		{
-			Widget = GenerateParameterWidgetFromLocalValue(ScriptVariable, LocalInputValueData.GetValue());
+			continue;
 		}
-		else if(InputDataInterface.IsSet() && InputDataInterface->DataInterface.IsValid())
-		{
-			Widget = GenerateParameterWidgetFromDataInterface(ScriptVariable, InputDataInterface.GetValue());
-		}		
 
-		SortOrder.Add(Widget.ToSharedRef(), ScriptVariable->Metadata.InlineParameterSortPriority);
+		SortOrder.Add(Widget.ToSharedRef(), FunctionInput->GetInputMetaData()->InlineParameterSortPriority);
 		ParameterWidgets.Add(Widget.ToSharedRef());
 	}
 	
-
 	ParameterWidgets.Sort([&](TSharedRef<SWidget> WidgetA, TSharedRef<SWidget> WidgetB)
 	{
 		return SortOrder[WidgetA] < SortOrder[WidgetB];
@@ -155,54 +168,117 @@ TArray<TSharedRef<SWidget>> SNiagaraOverviewInlineParameterBox::GenerateParamete
 	return ParameterWidgets;
 }
 
-TSharedRef<SWidget> SNiagaraOverviewInlineParameterBox::GenerateParameterWidgetFromLocalValue(const UNiagaraScriptVariable* ScriptVariable, FNiagaraLocalInputValueData& LocalInputValueData)
+TSharedRef<SWidget> SNiagaraOverviewInlineParameterBox::GenerateParameterWidgetFromLocalValue(UNiagaraStackFunctionInput* FunctionInput)
 {
-	LocalInputValueData.OnValueChanged.RemoveAll(this);
-	LocalInputValueData.OnValueChanged.AddSP(this, &SNiagaraOverviewInlineParameterBox::ConstructChildren);
-	BoundFunctionInputs.Add(LocalInputValueData);
-	FNiagaraTypeDefinition Type = ScriptVariable->Variable.GetType();
-	FString Value = Type.ToStringCosmetic(LocalInputValueData.LocalData->GetStructMemory());
+	bool bGenerateProperWidget = true;
+	
+	FNiagaraTypeDefinition Type = FunctionInput->GetInputType();
+	TOptional<FNiagaraVariableMetaData> InputMetaData = FunctionInput->GetInputMetaData();
+	TSharedPtr<INiagaraEditorTypeUtilities> TypeUtilities = FNiagaraEditorModule::Get().GetTypeUtilities(Type);
+
+	// we construct a variable with the data from our local struct memory because type utilities don't work on raw data
+	FNiagaraVariable Variable(Type, NAME_None);
+	Variable.SetData(FunctionInput->GetLocalValueStruct()->GetStructMemory());
+	
+	// by default we display type information. In some cases we have overrides, but even then we want to keep the type information as we continue using it in the tooltips
+	const FText ValueText = TypeUtilities->GetStackDisplayText(Variable);
+	FText DisplayedText = ValueText;
 	const FLinearColor TypeColor = UEdGraphSchema_Niagara::GetTypeColor(Type);
+	FLinearColor DisplayedColor = TypeColor;
 
 	UTexture2D* Icon = nullptr;
 	if (const UEnum* Enum = Type.GetEnum())
 	{
-		int32 EnumIndex = Enum->GetIndexByValue(*(int32*)LocalInputValueData.LocalData->GetStructMemory());
-		if(ScriptVariable->Metadata.InlineParameterEnumOverrides.IsValidIndex(EnumIndex))
+		int32 EnumIndex = Enum->GetIndexByValue(*(int32*)FunctionInput->GetLocalValueStruct()->GetStructMemory());
+		if(InputMetaData->InlineParameterEnumOverrides.IsValidIndex(EnumIndex))
 		{
-			Value = ScriptVariable->Metadata.InlineParameterEnumOverrides[EnumIndex].OverrideName.ToString();
-			Icon = ScriptVariable->Metadata.InlineParameterEnumOverrides[EnumIndex].IconOverride;
+			const FNiagaraEnumParameterMetaData& EnumParameterMetaData = InputMetaData->InlineParameterEnumOverrides[EnumIndex];
+
+			// we only want to use the override name if it has been set
+			DisplayedText = !EnumParameterMetaData.OverrideName.IsNone() ? FText::FromName(EnumParameterMetaData.OverrideName) : DisplayedText;			
+			Icon = EnumParameterMetaData.IconOverride;
+			DisplayedColor = EnumParameterMetaData.bUseColorOverride ? EnumParameterMetaData.ColorOverride : DisplayedColor;
 		}
 	}
+	else if(Type == FNiagaraTypeDefinition::GetBoolDef() && InputMetaData->bEnableBoolOverride)
+	{
+		const FNiagaraBoolParameterMetaData& BoolParameterMetaData = InputMetaData->InlineParameterBoolOverride;
+		DisplayedColor = InputMetaData->bOverrideColor ? InputMetaData->InlineParameterColorOverride : DisplayedColor;
 
-	TSharedPtr<SWidget> ParameterWidget;
+		int32 Val = *(int32*)FunctionInput->GetLocalValueStruct()->GetStructMemory();
+		bool bParameterValue = Val == 0 ? false : true;
 
-	TSharedRef<SButton> ParameterWidgetButton = SNew(SButton)
-	.Text(FText::FromString(Value))
-	.ContentPadding(FMargin(-5, 0))
-	.TextStyle(&FNiagaraEditorWidgetsStyle::Get().GetWidgetStyle<FTextBlockStyle>("NiagaraEditor.SystemOverview.InlineParameterText"))
-	.ButtonStyle(&FAppStyle::Get().GetWidgetStyle<FButtonStyle>("Button"))
-	.OnClicked(this, &SNiagaraOverviewInlineParameterBox::NavigateToStack, LocalInputValueData);
+		if(bParameterValue == true)
+		{
+			DisplayedText = !BoolParameterMetaData.OverrideNameTrue.IsNone() ? FText::FromName(BoolParameterMetaData.OverrideNameTrue) : DisplayedText;			
+			Icon = BoolParameterMetaData.IconOverrideTrue;
+		}
+		else if(bParameterValue == false)
+		{
+			DisplayedText = !BoolParameterMetaData.OverrideNameFalse.IsNone() ? FText::FromName(BoolParameterMetaData.OverrideNameFalse) : DisplayedText;			
+			Icon = BoolParameterMetaData.IconOverrideFalse;
+		}
+		
+		if(BoolParameterMetaData.DisplayMode != ENiagaraBoolDisplayMode::DisplayAlways)
+		{
+			// in case we only want to display the variable in one case, we make sure to test against the correct value
+			bool bDisplayBool = BoolParameterMetaData.DisplayMode == ENiagaraBoolDisplayMode::DisplayIfTrue ? true : false;  
+			if(bDisplayBool != bParameterValue)
+			{
+				bGenerateProperWidget = false;
+			}			
+		}		
+	}
+	else if(Type == FNiagaraTypeDefinition::GetColorDef())
+	{
+		DisplayedText = FText::FromName(FunctionInput->GetInputParameterHandle().GetName());
+		DisplayedColor = *(FLinearColor*)FunctionInput->GetLocalValueStruct()->GetStructMemory();
+	}
+
+	TSharedPtr<SWidget> ParameterWidget;
 
-	if(Icon != nullptr)
+	// we might have cases in which we don't want to display a parameter's data, such as if we only want to display a bool parameter when it is set to True
+	if(bGenerateProperWidget)
 	{
-		FSlateBrush& ImageBrush = ImageBrushes.Emplace_GetRef();
-		ImageBrush.SetResourceObject(Icon);
-		ImageBrush.ImageSize.X = 16;
-		ImageBrush.ImageSize.Y = 16;
+		TWeakObjectPtr<UNiagaraStackFunctionInput> FunctionInputWeak = FunctionInput;
+		TSharedRef<SButton> ParameterWidgetButton = SNew(SButton)
+		.Text(DisplayedText)
+		.TextStyle(&FNiagaraEditorWidgetsStyle::Get().GetWidgetStyle<FTextBlockStyle>("NiagaraEditor.SystemOverview.InlineParameterText"))
+		.ButtonStyle(&FNiagaraEditorWidgetsStyle::Get().GetWidgetStyle<FButtonStyle>("NiagaraEditor.SystemOverview.InlineParameterButton"))
+		.ContentPadding(FMargin(0.f))
+		.OnClicked(this, &SNiagaraOverviewInlineParameterBox::NavigateToStack, FunctionInputWeak);
+
+		// if we have an icon available, we use it instead of any text. Border color is irrelevant here.
+		if(Icon != nullptr)
+		{
+			FSlateBrush& ImageBrush = ImageBrushes.Emplace_GetRef();
+			ImageBrush.SetResourceObject(Icon);
+			ImageBrush.ImageSize.X = 16;
+			ImageBrush.ImageSize.Y = 16;
 	
-		TSharedRef<SImage> ParameterWidgetIcon= SNew(SImage).Image(&ImageBrush);
-		ParameterWidgetButton->SetContent(ParameterWidgetIcon);
-		//ParameterWidgetButton->SetBorderBackgroundColor(FStyleColors::Transparent);
-		ParameterWidgetButton->SetButtonStyle(&FAppStyle::Get().GetWidgetStyle<FButtonStyle>("SimpleButton"));
+			TSharedRef<SImage> ParameterWidgetIcon= SNew(SImage).Image(&ImageBrush);
+			ParameterWidgetButton->SetContent(ParameterWidgetIcon);
+			ParameterWidgetButton->SetButtonStyle(&FNiagaraEditorWidgetsStyle::Get().GetWidgetStyle<FButtonStyle>("NiagaraEditor.SystemOverview.InlineParameterButton.Transparent"));
+		}
+		// otherwise, we use text and set the border to its specified color (type, override, instance override color)
+		else
+		{
+			// if we are using a custom color, we set the button style to use a white base in order to get rid of a tint
+			if(DisplayedColor != TypeColor)
+			{
+				ParameterWidgetButton->SetButtonStyle(&FNiagaraEditorWidgetsStyle::Get().GetWidgetStyle<FButtonStyle>("NiagaraEditor.SystemOverview.InlineParameterButton.NoTint"));
+			}
+			
+			ParameterWidgetButton->SetBorderBackgroundColor(DisplayedColor);
+		}			
+
+		ParameterWidget = ParameterWidgetButton;
 	}
 	else
 	{
-		FLinearColor BorderColor = ScriptVariable->Metadata.bOverrideColor ? ScriptVariable->Metadata.InlineParameterColorOverride : TypeColor;
-		ParameterWidgetButton->SetBorderBackgroundColor(BorderColor);				
-	}			
-
-	ParameterWidget = ParameterWidgetButton;
+		return SNullWidget::NullWidget;
+	}
+	
 
 	FText			   IconToolTip = FText::GetEmpty();
 	FSlateBrush const* IconBrush = FAppStyle::Get().GetBrush(TEXT("Kismet.AllClasses.VariableIcon"));
@@ -214,40 +290,58 @@ TSharedRef<SWidget> SNiagaraOverviewInlineParameterBox::GenerateParameterWidgetF
 	// we construct a tooltip widget that shows the parameter the value is associated with
 	TSharedRef<SToolTip> TooltipWidget = SNew(SToolTip).Content()
 	[
-		SNew(SHorizontalBox)
-		+ SHorizontalBox::Slot()
-		.AutoWidth()
-		.HAlign(HAlign_Center)
-		.VAlign(VAlign_Center)
-		.Padding(3.f)
+		SNew(SVerticalBox)
+		+ SVerticalBox::Slot()
+		.AutoHeight()
 		[
-			SNew(SNiagaraParameterNameTextBlock)
-			.IsReadOnly(true)
-			.ParameterText(LOCTEXT("ParameterTooltipText", "Value for"))
-		]
-		+ SHorizontalBox::Slot()
-		.AutoWidth()
-		.HAlign(HAlign_Center)
-		.VAlign(VAlign_Center)
-		[
-			SNew(SNiagaraIconWidget)
-			.IconToolTip(IconToolTip)
-			.IconBrush(IconBrush)
-			.IconColor(IconColor)
-			.DocLink(IconDocLink)
-			.DocExcerpt(IconDocExcerpt)
-			.SecondaryIconBrush(SecondaryIconBrush) 
-			.SecondaryIconColor(SecondaryIconColor)
+			SNew(SHorizontalBox)
+			+ SHorizontalBox::Slot()
+			.AutoWidth()
+			.HAlign(HAlign_Center)
+			.VAlign(VAlign_Center)
+			[
+				SNew(SNiagaraIconWidget)
+				.IconToolTip(IconToolTip)
+				.IconBrush(IconBrush)
+				.IconColor(IconColor)
+				.DocLink(IconDocLink)
+				.DocExcerpt(IconDocExcerpt)
+				.SecondaryIconBrush(SecondaryIconBrush) 
+				.SecondaryIconColor(SecondaryIconColor)
+			]
+			+ SHorizontalBox::Slot()
+			.AutoWidth()
+			.HAlign(HAlign_Center)
+			.VAlign(VAlign_Center)
+			.Padding(3.f)
+			[
+				SNew(SNiagaraParameterNameTextBlock)
+				.IsReadOnly(true)
+				.ParameterText(FText::FromName(FunctionInput->GetInputParameterHandle().GetParameterHandleString()))
+			]
 		]
-		+ SHorizontalBox::Slot()
-		.AutoWidth()
-		.HAlign(HAlign_Center)
-		.VAlign(VAlign_Center)
-		.Padding(3.f)
+		+ SVerticalBox::Slot()
+		.AutoHeight()
 		[
-			SNew(SNiagaraParameterNameTextBlock)
-			.IsReadOnly(true)
-			.ParameterText(FText::FromName(ScriptVariable->Variable.GetName()))
+			SNew(SHorizontalBox)
+			+ SHorizontalBox::Slot()
+			.AutoWidth()
+			.HAlign(HAlign_Center)
+			.VAlign(VAlign_Center)
+			.Padding(3.f)
+			[
+				SNew(STextBlock)
+				.Text(LOCTEXT("ParameterTooltipText", "Value: "))
+			]
+			+ SHorizontalBox::Slot()
+			.AutoWidth()
+			.HAlign(HAlign_Center)
+			.VAlign(VAlign_Center)
+			.Padding(3.f)
+			[
+				SNew(STextBlock)
+				.Text(ValueText)
+			]
 		]
 	];
 
@@ -255,38 +349,38 @@ TSharedRef<SWidget> SNiagaraOverviewInlineParameterBox::GenerateParameterWidgetF
 	return ParameterWidget.ToSharedRef();
 }
 
-TSharedRef<SWidget> SNiagaraOverviewInlineParameterBox::GenerateParameterWidgetFromDataInterface(const UNiagaraScriptVariable* ScriptVariable, FNiagaraDataInterfaceInput& DataInterfaceInput)
+TSharedRef<SWidget> SNiagaraOverviewInlineParameterBox::GenerateParameterWidgetFromDataInterface(UNiagaraStackFunctionInput* FunctionInput)
 {
-	if(DataInterfaceInput.DataInterface.IsValid())
-	{
-		DataInterfaceInput.OnValueChanged.RemoveAll(this);
-		DataInterfaceInput.OnValueChanged.AddSP(this, &SNiagaraOverviewInlineParameterBox::ConstructChildren);
-		TArray<FNiagaraDataInterfaceError> Errors;
-		TArray<FNiagaraDataInterfaceFeedback> Warnings;
-		TArray<FNiagaraDataInterfaceFeedback> Infos;
-
-		// these can be nullptr depending on context
-		UNiagaraSystem* System = DataInterfaceInput.DataInterface->GetTypedOuter<UNiagaraSystem>();
-		UNiagaraComponent* Component = DataInterfaceInput.DataInterface->GetTypedOuter<UNiagaraComponent>();
-
-		//@TODO Finish this up.
-		//DataInterfaceInput.DataInterface->GetInlineFeedback(System, Component, Errors, Warnings, Infos);
-
-		TSharedPtr<SHorizontalBox> InfoBox = SNew(SHorizontalBox);
-		for(FNiagaraDataInterfaceFeedback& Feedback : Infos)
-		{
-			InfoBox->AddSlot()
-			.AutoWidth()
-			[
-				SNew(STextBlock)
-				.Text(Feedback.GetFeedbackSummaryText())
-				.ToolTipText(Feedback.GetFeedbackText())
-				.TextStyle(&FNiagaraEditorWidgetsStyle::Get().GetWidgetStyle<FTextBlockStyle>("NiagaraEditor.SystemOverview.InlineParameterText"))
-			];			
-		}
-
-		return InfoBox.ToSharedRef();
-	}
+	// if(DataInterfaceInput.DataInterface.IsValid())
+	// {
+	// 	DataInterfaceInput.OnValueChanged.RemoveAll(this);
+	// 	DataInterfaceInput.OnValueChanged.AddSP(this, &SNiagaraOverviewInlineParameterBox::ConstructChildren);
+	// 	TArray<FNiagaraDataInterfaceError> Errors;
+	// 	TArray<FNiagaraDataInterfaceFeedback> Warnings;
+	// 	TArray<FNiagaraDataInterfaceFeedback> Infos;
+	//
+	// 	// these can be nullptr depending on context
+	// 	UNiagaraSystem* System = DataInterfaceInput.DataInterface->GetTypedOuter<UNiagaraSystem>();
+	// 	UNiagaraComponent* Component = DataInterfaceInput.DataInterface->GetTypedOuter<UNiagaraComponent>();
+	//
+	// 	//@TODO Finish this up.
+	// 	//DataInterfaceInput.DataInterface->GetInlineFeedback(System, Component, Errors, Warnings, Infos);
+	//
+	// 	TSharedPtr<SHorizontalBox> InfoBox = SNew(SHorizontalBox);
+	// 	for(FNiagaraDataInterfaceFeedback& Feedback : Infos)
+	// 	{
+	// 		InfoBox->AddSlot()
+	// 		.AutoWidth()
+	// 		[
+	// 			SNew(STextBlock)
+	// 			.Text(Feedback.GetFeedbackSummaryText())
+	// 			.ToolTipText(Feedback.GetFeedbackText())
+	// 			.TextStyle(&FNiagaraEditorWidgetsStyle::Get().GetWidgetStyle<FTextBlockStyle>("NiagaraEditor.SystemOverview.InlineParameterText"))
+	// 		];			
+	// 	}
+	//
+	// 	return InfoBox.ToSharedRef();
+	//}
 
 	return SNullWidget::NullWidget;
 }
diff --git a/Engine/Plugins/FX/Niagara/Source/NiagaraEditorWidgets/Private/SNiagaraOverviewInlineParameterBox.h b/Engine/Plugins/FX/Niagara/Source/NiagaraEditorWidgets/Private/SNiagaraOverviewInlineParameterBox.h
--- a/Engine/Plugins/FX/Niagara/Source/NiagaraEditorWidgets/Private/SNiagaraOverviewInlineParameterBox.h
+++ b/Engine/Plugins/FX/Niagara/Source/NiagaraEditorWidgets/Private/SNiagaraOverviewInlineParameterBox.h
@@ -31,13 +31,15 @@ public:
 
 	virtual ~SNiagaraOverviewInlineParameterBox() override;
 
+	int32 GetNumChildWidgets();
+
 private:
 	/** Selects the module item in the stack and searches for the function input */
-	FReply NavigateToStack(FNiagaraLocalInputValueData LocalInputValueData);
+	FReply NavigateToStack(TWeakObjectPtr<UNiagaraStackFunctionInput> FunctionInput);
 	void ConstructChildren();
 	TArray<TSharedRef<SWidget>> GenerateParameterWidgets();
-	TSharedRef<SWidget> GenerateParameterWidgetFromLocalValue(const UNiagaraScriptVariable* ScriptVariable, FNiagaraLocalInputValueData& LocalInputValueData);
-	TSharedRef<SWidget> GenerateParameterWidgetFromDataInterface(const UNiagaraScriptVariable* ScriptVariable, FNiagaraDataInterfaceInput& DataInterface);
+	TSharedRef<SWidget> GenerateParameterWidgetFromLocalValue(UNiagaraStackFunctionInput* FunctionInput);
+	TSharedRef<SWidget> GenerateParameterWidgetFromDataInterface(UNiagaraStackFunctionInput* FunctionInput);
 
 private:
 	/** The module item whose parameters this box is representing */
@@ -46,7 +48,9 @@ private:
 	TSharedPtr<SScrollBox> Container;
 	
 	/** We keep track of the function inputs we are observing through this parameter box so we can unbind the delegates later */
-	TArray<FNiagaraLocalInputValueData> BoundFunctionInputs;
+	TArray<TWeakObjectPtr<UNiagaraStackFunctionInput>> BoundFunctionInputs;
 
 	TArray<FSlateBrush> ImageBrushes;
+
+	int32 NumParameterWidgets = 0;
 };
diff --git a/Engine/Plugins/FX/Niagara/Source/NiagaraEditorWidgets/Private/SNiagaraOverviewStack.cpp b/Engine/Plugins/FX/Niagara/Source/NiagaraEditorWidgets/Private/SNiagaraOverviewStack.cpp
--- a/Engine/Plugins/FX/Niagara/Source/NiagaraEditorWidgets/Private/SNiagaraOverviewStack.cpp
+++ b/Engine/Plugins/FX/Niagara/Source/NiagaraEditorWidgets/Private/SNiagaraOverviewStack.cpp
@@ -43,6 +43,7 @@
 #include "ScopedTransaction.h"
 #include "Styling/StyleColors.h"
 #include "Widgets/Input/SCheckBox.h"
+#include "Widgets/Layout/SWrapBox.h"
 
 #define LOCTEXT_NAMESPACE "NiagaraOverviewStack"
 
@@ -751,7 +752,7 @@ TSharedRef<ITableRow> SNiagaraOverviewStack::OnGenerateRowForEntry(UNiagaraStack
 				.WidthOverride(IconSize.X)
 				.HeightOverride(IconSize.Y);
 		}
-
+		
 		TSharedRef<SHorizontalBox> ContentBox = SNew(SHorizontalBox)
 			// Indent content
 			+ SHorizontalBox::Slot()
@@ -764,8 +765,6 @@ TSharedRef<ITableRow> SNiagaraOverviewStack::OnGenerateRowForEntry(UNiagaraStack
 			// Name
 			+ SHorizontalBox::Slot()
 			.VAlign(VAlign_Center)
-			.AutoWidth()
-			.MaxWidth(150.f)
 			.Padding(3, 2, 0, 2)
 			[
 				SNew(SNiagaraSystemOverviewItemName, StackItem)
@@ -777,20 +776,16 @@ TSharedRef<ITableRow> SNiagaraOverviewStack::OnGenerateRowForEntry(UNiagaraStack
                 SNew(SNiagaraStackRowPerfWidget, Item)
             ];
 
-		// Debug draw 
+		TSharedRef<SHorizontalBox> OptionsBox = SNew(SHorizontalBox);
+		
 		UNiagaraStackModuleItem* StackModuleItem = Cast<UNiagaraStackModuleItem>(StackItem);
-			
+
 		if (StackModuleItem)
 		{
-			ContentBox->AddSlot()
-			.HAlign(HAlign_Fill)
-			[
-				SNew(SNiagaraOverviewInlineParameterBox, *StackModuleItem)
-			];
-			
+			// Scratch icon
 			if(StackModuleItem->IsScratchModule())
 			{
-				ContentBox->AddSlot()
+				OptionsBox->AddSlot()
 					.VAlign(VAlign_Center)
 					.AutoWidth()
 					.Padding(3, 0, 3, 1)
@@ -800,12 +795,12 @@ TSharedRef<ITableRow> SNiagaraOverviewStack::OnGenerateRowForEntry(UNiagaraStack
 						.Image(FNiagaraEditorStyle::Get().GetBrush("Tab.ScratchPad"))
 					];
 			}
-			
+
+			// Debug draw 
 			if(StackModuleItem->GetModuleNode().ContainsDebugSwitch())
 			{
-				ContentBox->AddSlot()
+				OptionsBox->AddSlot()
 					.VAlign(VAlign_Center)
-					.HAlign(HAlign_Center)
 					.AutoWidth()
 					.Padding(3, 0)
 					[
@@ -816,8 +811,8 @@ TSharedRef<ITableRow> SNiagaraOverviewStack::OnGenerateRowForEntry(UNiagaraStack
 						.ForegroundColor(FLinearColor::Transparent)
 						.ToolTipText(LOCTEXT("EnableDebugDrawCheckBoxToolTip", "Enable or disable debug drawing for this item."))
 						.OnClicked(this, &SNiagaraOverviewStack::ToggleModuleDebugDraw, StackItem)
-						// @Todo the debug icons aren't centered, so we correct this using padding. We also shrink the button size
-						.ContentPadding(FMargin(-11, 2, -7, 2))
+						.ContentPadding(FMargin(2.f))
+						.ButtonStyle(&FAppStyle::Get().GetWidgetStyle<FButtonStyle>("SimpleButton"))
 						[
 							SNew(SImage)
 							.Image(this, &SNiagaraOverviewStack::GetDebugIconBrush, StackItem)
@@ -827,7 +822,7 @@ TSharedRef<ITableRow> SNiagaraOverviewStack::OnGenerateRowForEntry(UNiagaraStack
 		}
 
 		// Enabled checkbox
-		ContentBox->AddSlot()
+		OptionsBox->AddSlot()
 			.VAlign(VAlign_Center)
 			.AutoWidth()
 			.Padding(0)
@@ -839,6 +834,65 @@ TSharedRef<ITableRow> SNiagaraOverviewStack::OnGenerateRowForEntry(UNiagaraStack
 				.OnCheckedChanged_UObject(StackItem, &UNiagaraStackItem::SetIsEnabled)
 			];
 
+		// in case we have at least one inline input, we add a wrap box to make sure we have enough space for both the module name and parameters
+		if(StackModuleItem && StackModuleItem->GetInlineParameterInputs().Num() > 0)
+		{
+			ContentBox = SNew(SHorizontalBox)
+			+ SHorizontalBox::Slot()
+			.VAlign(VAlign_Center)
+			[
+				SNew(SWrapBox)
+				// the wrap box will put the inline box into a second row if content box + parameter box would exceed the preferred size
+				.PreferredSize(250.f)
+				+ SWrapBox::Slot()
+				.VAlign(VAlign_Center)
+				[
+					SNew(SBox)
+					.VAlign(VAlign_Center)
+					// despite that, we also have to cap the content box size, since the wrapbox can't wrap a single widget
+					.MaxDesiredWidth(250.f)
+					[
+						ContentBox
+					]
+				]
+				+ SWrapBox::Slot()
+				.VAlign(VAlign_Center)
+				[
+					SNew(SBox)
+					.VAlign(VAlign_Center)
+					.MaxDesiredWidth(150.f)
+					[
+						SNew(SNiagaraOverviewInlineParameterBox, *StackModuleItem)
+					]
+				]
+			]
+			+ SHorizontalBox::Slot()
+			.VAlign(VAlign_Center)
+			.AutoWidth()
+			[
+				OptionsBox
+			];
+		}
+		// if we don't, we don't use the wrap box for perf reasons
+		else
+		{
+			ContentBox = SNew(SHorizontalBox)
+			+ SHorizontalBox::Slot()
+			[
+				SNew(SBox)
+				.VAlign(VAlign_Center)
+				.MaxDesiredWidth(250.f)
+				[
+					ContentBox
+				]			
+			]
+			+ SHorizontalBox::Slot()
+			.AutoWidth()
+			[
+				OptionsBox
+			];
+		}
+		
 		Content = ContentBox;		
 	}
 	else if (Item->IsA<UNiagaraStackItemGroup>())
