Fixed ShooterGame compile error due to hard wrapped-pointer requirement.

#rnx

diff --git a/Engine/Plugins/FX/Niagara/Source/Niagara/Public/NiagaraTypes.h b/Engine/Plugins/FX/Niagara/Source/Niagara/Public/NiagaraTypes.h
--- a/Engine/Plugins/FX/Niagara/Source/Niagara/Public/NiagaraTypes.h
+++ b/Engine/Plugins/FX/Niagara/Source/Niagara/Public/NiagaraTypes.h
@@ -623,13 +623,13 @@ struct NIAGARA_API FNiagaraEnumParameterMetaData
 
 	/** If specified, this icon will be used for the given enum entry. If OverrideName isn't empty, the icon takes priority. */
 	UPROPERTY(EditAnywhere, Category="Enum Override")
-	UTexture2D* IconOverride = nullptr;
+	TObjectPtr<UTexture2D> IconOverride = nullptr;
 
 	UPROPERTY(EditAnywhere, Category="Enum Override", meta=(InlineEditConditionToggle))
 	bool bUseColorOverride = false;
 	
 	UPROPERTY(EditAnywhere, Category="Enum Override", meta=(EditCondition="bUseColorOverride"))
-	FLinearColor ColorOverride;
+	FLinearColor ColorOverride = FLinearColor::White;
 };
 
 UENUM()
@@ -662,11 +662,11 @@ struct NIAGARA_API FNiagaraBoolParameterMetaData
 
 	/** If specified, this icon will be used for the given bool if it evaluates to True. If OverrideName isn't empty, the icon takes priority. */
 	UPROPERTY(EditAnywhere, Category="Bool Override")
-	UTexture2D* IconOverrideTrue = nullptr;
+	TObjectPtr<UTexture2D> IconOverrideTrue = nullptr;
 
 	/** If specified, this icon will be used for the given bool if it evaluates to False. If OverrideName isn't empty, the icon takes priority. */
 	UPROPERTY(EditAnywhere, Category="Bool Override")
-	UTexture2D* IconOverrideFalse = nullptr;
+	TObjectPtr<UTexture2D> IconOverrideFalse = nullptr;
 };
 
 USTRUCT()
