Record parameter definition links to scripts/emitters/systems as a uproperty ptr instead of a guid so that we can resolve them at load time.
Refactor parameter definition synchronization on postload to deferred synchronize when we have not discovered all parameter definition assets in the project.

diff --git a/Engine/Plugins/FX/Niagara/Source/Niagara/Classes/NiagaraScript.h b/Engine/Plugins/FX/Niagara/Source/Niagara/Classes/NiagaraScript.h
--- a/Engine/Plugins/FX/Niagara/Source/Niagara/Classes/NiagaraScript.h
+++ b/Engine/Plugins/FX/Niagara/Source/Niagara/Classes/NiagaraScript.h
@@ -1208,6 +1208,8 @@ public:
 	{
 	};
 
+	~FVersionedNiagaraScript();
+
 	//~ Begin INiagaraParameterDefinitionsSubscriber Interface
 	virtual const TArray<FParameterDefinitionsSubscription>& GetParameterDefinitionsSubscriptions() const override { return GetScriptData()->ParameterDefinitionsSubscriptions; };
 	virtual TArray<FParameterDefinitionsSubscription>& GetParameterDefinitionsSubscriptions() override { return GetScriptData()->ParameterDefinitionsSubscriptions; };
@@ -1240,6 +1242,8 @@ public:
 	{
 	};
 
+	~FVersionedNiagaraScriptWeakPtr();
+
 	//~ Begin INiagaraParameterDefinitionsSubscriber Interface
 	virtual const TArray<FParameterDefinitionsSubscription>& GetParameterDefinitionsSubscriptions() const override { return Pin().GetScriptData()->ParameterDefinitionsSubscriptions; };
 	virtual TArray<FParameterDefinitionsSubscription>& GetParameterDefinitionsSubscriptions() override { return Pin().GetScriptData()->ParameterDefinitionsSubscriptions; };
diff --git a/Engine/Plugins/FX/Niagara/Source/Niagara/Private/NiagaraScript.cpp b/Engine/Plugins/FX/Niagara/Source/Niagara/Private/NiagaraScript.cpp
--- a/Engine/Plugins/FX/Niagara/Source/Niagara/Private/NiagaraScript.cpp
+++ b/Engine/Plugins/FX/Niagara/Source/Niagara/Private/NiagaraScript.cpp
@@ -2232,6 +2232,11 @@ FVersionedNiagaraScript FVersionedNiagaraScriptWeakPtr::Pin()
 	return FVersionedNiagaraScript();
 }
 
+FVersionedNiagaraScriptWeakPtr::~FVersionedNiagaraScriptWeakPtr()
+{
+	CleanupDefinitionsSubscriptions();
+}
+
 TArray<UNiagaraScriptSourceBase*> FVersionedNiagaraScriptWeakPtr::GetAllSourceScripts()
 {
 	if (Script.IsValid())
@@ -2246,6 +2251,11 @@ FString FVersionedNiagaraScriptWeakPtr::GetSourceObjectPathName() const
 	return Script.IsValid() ? Script.Get()->GetPathName() : FString();
 }
 
+FVersionedNiagaraScript::~FVersionedNiagaraScript()
+{
+	CleanupDefinitionsSubscriptions();
+}
+
 TArray<UNiagaraScriptSourceBase*> FVersionedNiagaraScript::GetAllSourceScripts()
 {
 	if (Script != nullptr)
