Rearrange cleanup of scripts waiting to sync with definitions to resolve on BeginDestroy instead of destructor as we could crash during engine shutdown after the asset registry is unloaded.

diff --git a/Engine/Plugins/FX/Niagara/Source/Niagara/Classes/NiagaraScript.h b/Engine/Plugins/FX/Niagara/Source/Niagara/Classes/NiagaraScript.h
--- a/Engine/Plugins/FX/Niagara/Source/Niagara/Classes/NiagaraScript.h
+++ b/Engine/Plugins/FX/Niagara/Source/Niagara/Classes/NiagaraScript.h
@@ -1208,8 +1208,6 @@ public:
 	{
 	};
 
-	~FVersionedNiagaraScript();
-
 	//~ Begin INiagaraParameterDefinitionsSubscriber Interface
 	virtual const TArray<FParameterDefinitionsSubscription>& GetParameterDefinitionsSubscriptions() const override { return GetScriptData()->ParameterDefinitionsSubscriptions; };
 	virtual TArray<FParameterDefinitionsSubscription>& GetParameterDefinitionsSubscriptions() override { return GetScriptData()->ParameterDefinitionsSubscriptions; };
@@ -1242,8 +1240,6 @@ public:
 	{
 	};
 
-	~FVersionedNiagaraScriptWeakPtr();
-
 	//~ Begin INiagaraParameterDefinitionsSubscriber Interface
 	virtual const TArray<FParameterDefinitionsSubscription>& GetParameterDefinitionsSubscriptions() const override { return Pin().GetScriptData()->ParameterDefinitionsSubscriptions; };
 	virtual TArray<FParameterDefinitionsSubscription>& GetParameterDefinitionsSubscriptions() override { return Pin().GetScriptData()->ParameterDefinitionsSubscriptions; };
diff --git a/Engine/Plugins/FX/Niagara/Source/Niagara/Private/NiagaraScript.cpp b/Engine/Plugins/FX/Niagara/Source/Niagara/Private/NiagaraScript.cpp
--- a/Engine/Plugins/FX/Niagara/Source/Niagara/Private/NiagaraScript.cpp
+++ b/Engine/Plugins/FX/Niagara/Source/Niagara/Private/NiagaraScript.cpp
@@ -2232,11 +2232,6 @@ FVersionedNiagaraScript FVersionedNiagaraScriptWeakPtr::Pin()
 	return FVersionedNiagaraScript();
 }
 
-FVersionedNiagaraScriptWeakPtr::~FVersionedNiagaraScriptWeakPtr()
-{
-	CleanupDefinitionsSubscriptions();
-}
-
 TArray<UNiagaraScriptSourceBase*> FVersionedNiagaraScriptWeakPtr::GetAllSourceScripts()
 {
 	if (Script.IsValid())
@@ -2251,11 +2246,6 @@ FString FVersionedNiagaraScriptWeakPtr::GetSourceObjectPathName() const
 	return Script.IsValid() ? Script.Get()->GetPathName() : FString();
 }
 
-FVersionedNiagaraScript::~FVersionedNiagaraScript()
-{
-	CleanupDefinitionsSubscriptions();
-}
-
 TArray<UNiagaraScriptSourceBase*> FVersionedNiagaraScript::GetAllSourceScripts()
 {
 	if (Script != nullptr)
