Fix for niagara shader generation where the name of a stage/module ending in 'Particles' would confuse the translator

diff --git a/Engine/Plugins/FX/Niagara/Source/NiagaraCore/Private/NiagaraCustomVersion.cpp b/Engine/Plugins/FX/Niagara/Source/NiagaraCore/Private/NiagaraCustomVersion.cpp
--- a/Engine/Plugins/FX/Niagara/Source/NiagaraCore/Private/NiagaraCustomVersion.cpp
+++ b/Engine/Plugins/FX/Niagara/Source/NiagaraCore/Private/NiagaraCustomVersion.cpp
@@ -8,4 +8,4 @@ const FGuid FNiagaraCustomVersion::GUID(0xFCF57AFA, 0x50764283, 0xB9A9E658, 0xFF
 // Register the custom version with core
 FCustomVersionRegistration GRegisterNiagaraCustomVersion(FNiagaraCustomVersion::GUID, FNiagaraCustomVersion::LatestVersion, TEXT("NiagaraVer"));
 
-const FGuid FNiagaraCustomVersion::LatestScriptCompileVersion(0xAD773C8E, 0x9F7A1544, 0xAE4598F6, 0xC73F6125);
+const FGuid FNiagaraCustomVersion::LatestScriptCompileVersion(0xDDD7441D, 0x64BB4007, 0x94DB2190, 0x75AFD92A);
diff --git a/Engine/Plugins/FX/Niagara/Source/NiagaraEditor/Private/NiagaraHlslTranslator.cpp b/Engine/Plugins/FX/Niagara/Source/NiagaraEditor/Private/NiagaraHlslTranslator.cpp
--- a/Engine/Plugins/FX/Niagara/Source/NiagaraEditor/Private/NiagaraHlslTranslator.cpp
+++ b/Engine/Plugins/FX/Niagara/Source/NiagaraEditor/Private/NiagaraHlslTranslator.cpp
@@ -291,6 +291,25 @@ private:
 typedef TSimStagePermutationContext<EPermutationScopeContext::Declaration> FDeclarationPermutationContext;
 typedef TSimStagePermutationContext<EPermutationScopeContext::Expression> FExpressionPermutationContext;
 
+// utility function to replace a namespace (potentially internal within a fully qualified name) with another.
+// As an example MyParticlesValue.Particles.CurrentValue with Source <Particles> and Replace <Array> will
+// target the first 'Particles' namespace and ignore the leading namespace qualifier that included 'Particles'
+static void ReplaceNamespaceInline(FString& FullName, FStringView Source, FStringView Replace)
+{
+	const int32 SourceLength = Source.Len();
+	int32 SourceIdx = FullName.Find(Source.GetData());
+	while (SourceIdx != INDEX_NONE)
+	{
+		if ((SourceIdx == 0 || (SourceIdx > 0 && FullName[SourceIdx - 1] == '.')))
+		{
+			FullName.RemoveAt(SourceIdx, SourceLength);
+			FullName.InsertAt(SourceIdx, Replace.GetData());
+			break;
+		}
+		SourceIdx = FullName.Find(Source.GetData(), ESearchCase::IgnoreCase, ESearchDir::FromStart, SourceIdx + SourceLength);
+	}
+}
+
 FString FHlslNiagaraTranslator::GetCode(int32 ChunkIdx)
 {
 	FNiagaraCodeChunk& Chunk = CodeChunks[ChunkIdx];
@@ -772,8 +791,9 @@ FString FHlslNiagaraTranslator::BuildParameterMapHlslDefinitions(TArray<FNiagara
 		const FNiagaraVariable& NiagaraVariable = UniqueVariables[UniqueVarIdx];
 		if (FNiagaraParameterMapHistory::IsAttribute(NiagaraVariable))
 		{
-			const FString VariableName = GetSanitizedSymbolName(NiagaraVariable.GetName().ToString());
-			RegisterNames.Add(VariableName.Replace(PARAM_MAP_ATTRIBUTE_STR, PARAM_MAP_INDICES_STR));
+			FString VariableName = GetSanitizedSymbolName(NiagaraVariable.GetName().ToString());
+			ReplaceNamespaceInline(VariableName, PARAM_MAP_ATTRIBUTE_STR, PARAM_MAP_INDICES_STR);
+			RegisterNames.Add(VariableName);
 		}
 	}
 	for (const FString& RegisterName : RegisterNames)
@@ -2851,21 +2871,7 @@ void FHlslNiagaraTranslator::DefineMainGPUFunctions(
 						if (FNiagaraParameterMapHistory::IsAttribute(Var))
 						{
 							FString RegisterName = VarName;
-							const FString AttribStr = PARAM_MAP_ATTRIBUTE_STR;
-
-							//-TODO: If a simulation stage ends in the name "Particles." (i.e. ResolveParticles) a single search can fail so we must loop to find the correct namespace
-							//       Ideally we replace this with some more robust namespace functionality.
-							int32 ParticlesIdx = RegisterName.Find(AttribStr);
-							while (ParticlesIdx != INDEX_NONE)
-							{
-								if ( (ParticlesIdx == 0 || (ParticlesIdx > 0 && RegisterName[ParticlesIdx - 1] == '.')) )
-								{
-									RegisterName.RemoveAt(ParticlesIdx, AttribStr.Len());
-									RegisterName.InsertAt(ParticlesIdx, PARAM_MAP_INDICES_STR);
-									break;
-								}
-								ParticlesIdx = RegisterName.Find(AttribStr, ESearchCase::IgnoreCase, ESearchDir::FromStart, ParticlesIdx + AttribStr.Len());
-							}
+							ReplaceNamespaceInline(RegisterName, PARAM_MAP_ATTRIBUTE_STR, PARAM_MAP_INDICES_STR);
 
 							const int32 RegisterValue = Var.GetType().IsFloatPrimitive() ? FloatCounter : IntCounter;
 							HlslOutput += RegisterName + FString::Printf(TEXT(" = %d;\n"), RegisterValue);
@@ -3555,7 +3561,7 @@ void FHlslNiagaraTranslator::DefineDataSetVariableReads(FString &OutHlslOutput,
 				if (FNiagaraParameterMapHistory::IsAttribute(Var))
 				{
 					FString RegisterName = VariableName;
-					RegisterName.ReplaceInline(PARAM_MAP_ATTRIBUTE_STR, PARAM_MAP_INDICES_STR);
+					ReplaceNamespaceInline(RegisterName, PARAM_MAP_ATTRIBUTE_STR, PARAM_MAP_INDICES_STR);
 
 					Fmt += RegisterName + TEXT(" = {3};\n");
 				}
