Fixed inline parameters not being disabled when the entire emitter gets disabled

diff --git a/Engine/Plugins/FX/Niagara/Source/NiagaraEditorWidgets/Private/SNiagaraOverviewInlineParameterBox.cpp b/Engine/Plugins/FX/Niagara/Source/NiagaraEditorWidgets/Private/SNiagaraOverviewInlineParameterBox.cpp
--- a/Engine/Plugins/FX/Niagara/Source/NiagaraEditorWidgets/Private/SNiagaraOverviewInlineParameterBox.cpp
+++ b/Engine/Plugins/FX/Niagara/Source/NiagaraEditorWidgets/Private/SNiagaraOverviewInlineParameterBox.cpp
@@ -110,7 +110,7 @@ void SNiagaraOverviewInlineParameterBox::ConstructChildren()
 	}
 
 	SetVisibility(ParameterWidgets.Num() > 0 ? EVisibility::Visible : EVisibility::Collapsed);
-	Container->SetEnabled(ModuleItem->GetIsEnabled());
+	Container->SetEnabled(ModuleItem->GetIsEnabledAndOwnerIsEnabled());
 }
 
 TArray<TSharedRef<SWidget>> SNiagaraOverviewInlineParameterBox::GenerateParameterWidgets()
