Fix for multiply declared symbol when subuvs were used for niagara assets

diff --git a/Engine/Plugins/FX/Niagara/Shaders/Private/NiagaraMeshVertexFactory.ush b/Engine/Plugins/FX/Niagara/Shaders/Private/NiagaraMeshVertexFactory.ush
--- a/Engine/Plugins/FX/Niagara/Shaders/Private/NiagaraMeshVertexFactory.ush
+++ b/Engine/Plugins/FX/Niagara/Shaders/Private/NiagaraMeshVertexFactory.ush
@@ -615,8 +615,8 @@ FVertexFactoryIntermediates GetVertexFactoryIntermediates(FVertexFactoryInput In
 	Intermediates.SubUVLerp = (NiagaraMeshVF.SubImageBlendMode == 1) ? SubImageLerp : 0.0f;
 
 	#if MANUAL_VERTEX_FETCH
-		const uint NumFetchTexCoords = NiagaraMeshVF.VertexFetch_Parameters[VF_NumTexcoords_Index];
-		const float2 SubUvTexCoord = NiagaraMeshVF.VertexFetch_TexCoordBuffer[NumFetchTexCoords * Input.VertexId];
+		const uint NumFetchSubUvTexCoords = NiagaraMeshVF.VertexFetch_Parameters[VF_NumTexcoords_Index];
+		const float2 SubUvTexCoord = NiagaraMeshVF.VertexFetch_TexCoordBuffer[NumFetchSubUvTexCoords * Input.VertexId];
 	#else
 		const float2 SubUvTexCoord = Input.TexCoords0.xy;
 	#endif // MANUAL_VERTEX_FETCH
