Implementing refresh for Convert Nodes. Convert nodes will now fix up the paths in their connections to reflect changes made to BP structs.

Also fixing case where swizzle inputs didn't have a typedef, so the wiring was not show on the nodes.
Niagara bools and ints are now assignable.

diff --git a/Engine/Plugins/FX/Niagara/Source/Niagara/Private/NiagaraModule.cpp b/Engine/Plugins/FX/Niagara/Source/Niagara/Private/NiagaraModule.cpp
--- a/Engine/Plugins/FX/Niagara/Source/Niagara/Private/NiagaraModule.cpp
+++ b/Engine/Plugins/FX/Niagara/Source/Niagara/Private/NiagaraModule.cpp
@@ -1100,7 +1100,7 @@ bool FNiagaraTypeDefinition::TypesAreAssignable(const FNiagaraTypeDefinition& Ty
 	bool bIsSupportedConversion = false;
 	if (IsScalarDefinition(TypeA) && IsScalarDefinition(TypeB))
 	{
-		bIsSupportedConversion = (TypeA == IntDef && TypeB == FloatDef) || (TypeB == IntDef && TypeA == FloatDef);
+		bIsSupportedConversion = (TypeA == IntDef && TypeB == FloatDef) || (TypeB == IntDef && TypeA == FloatDef) || (TypeA == IntDef && TypeB == BoolDef) || (TypeB == IntDef && TypeA == BoolDef);
 	}
 	else
 	{
diff --git a/Engine/Plugins/FX/Niagara/Source/NiagaraEditor/Private/NiagaraNodeConvert.cpp b/Engine/Plugins/FX/Niagara/Source/NiagaraEditor/Private/NiagaraNodeConvert.cpp
--- a/Engine/Plugins/FX/Niagara/Source/NiagaraEditor/Private/NiagaraNodeConvert.cpp
+++ b/Engine/Plugins/FX/Niagara/Source/NiagaraEditor/Private/NiagaraNodeConvert.cpp
@@ -6,6 +6,7 @@
 #include "SNiagaraGraphNodeConvert.h"
 #include "NiagaraHlslTranslator.h"
 #include "UObject/UnrealType.h"
+#include "Kismet2/StructureEditorUtils.h"
 
 #include "ScopedTransaction.h"
 
@@ -214,6 +215,289 @@ FString FNiagaraConvertConnection::ToString() const
 	return SrcName + TEXT(" to ") + DestName;
 }
 
+namespace FNiagaraConvertRefreshHelpers
+{
+	// Helper function to reduce book-keeping around searching by predicate
+	UEdGraphPin* FindPinByID(FGuid SearchPinId, TArray<UEdGraphPin*> PinsToSearch)
+	{
+		UEdGraphPin** FoundPin = PinsToSearch.FindByPredicate([&SearchPinId](UEdGraphPin* Pin) {return Pin->PinId == SearchPinId; });
+		if (FoundPin)
+		{
+			return *FoundPin;
+		}
+
+		return nullptr;
+	}
+
+	// Helper function that renames a pin if it matches a member in the given struct
+	void RenamePinIfInStruct(UEdGraphPin* Pin, const UScriptStruct* Struct, FGuid& ConnectionGuid)
+	{
+
+		const UUserDefinedStruct* UserDefinedStruct = Cast<const UUserDefinedStruct>(Struct);
+
+		for (FProperty* Property : TFieldRange<FProperty>(Struct, EFieldIteratorFlags::IncludeSuper))
+		{
+
+			FString PropertyName;
+			FGuid PropertyGuid = FStructureEditorUtils::GetGuidForProperty(Property);
+
+			// User defined structs will have GUIDs as part of the name, so the friendly name is needed instead
+			if (UserDefinedStruct)
+			{
+				PropertyName = FStructureEditorUtils::GetVariableFriendlyName(UserDefinedStruct, PropertyGuid);
+			}
+			else
+			{
+				PropertyName = Property->GetName();
+			}
+
+			// Rename and store GUID
+			if (Pin->GetName() == PropertyName || (ConnectionGuid == PropertyGuid && ConnectionGuid != FGuid()))
+			{
+				ConnectionGuid = PropertyGuid;
+				Pin->PinName = FName(*PropertyName);
+
+				break;
+			}
+		}
+	}
+
+	// Recurses paths of a convert connection, and renames if the path is valid to conform to changes in the underlying struct
+	// Returns true if the path is valid, if a property can be found, its pointer is stored in OutProperty.
+	bool RecursePaths(const UScriptStruct* ParentStruct, TArray<FName>& Paths, int32 PathIndex, FProperty*& OutProperty)
+	{
+		bool bResult = false;
+		if (Paths.IsValidIndex(PathIndex))
+		{
+			FName CurrentPath = Paths[PathIndex];
+			// Paths will end with "Value" for scalars, which is a special case. The path is valid, but there's nothing more to search for
+			if (CurrentPath.ToString().Equals("Value"))
+			{
+				return true;
+			}
+
+			FGuid CurrentGUID = FStructureEditorUtils::GetGuidFromPropertyName(CurrentPath);
+			FProperty* FoundProperty = nullptr;
+
+			// Search for path by name or GUID in struct
+			for (FProperty* Property : TFieldRange<FProperty>(ParentStruct))
+			{
+				FGuid PropertyGuid = FStructureEditorUtils::GetGuidFromPropertyName(Property->GetFName());
+
+				if (CurrentPath == Property->GetFName() || (PropertyGuid != FGuid() && PropertyGuid == CurrentGUID))
+				{
+					FoundProperty = Property;
+					break;
+				}
+			}
+
+
+			if (FoundProperty)
+			{
+				// Path is valid up until this point
+				bResult = true;
+				const FStructProperty* StructProperty = CastField<const FStructProperty>(FoundProperty);
+
+				// For nested structs update the parent struct to this property
+				if (StructProperty)
+				{
+					ParentStruct = StructProperty->Struct;
+				}
+
+				// Continue if there are more entries
+				if (Paths.IsValidIndex(PathIndex + 1))
+				{
+					bResult = RecursePaths(ParentStruct, Paths, PathIndex + 1, OutProperty);
+				}
+				if (bResult)
+				{
+					// If no other calls have assigned the out property, do so here
+					if (!OutProperty)
+					{
+						OutProperty = FoundProperty;
+					}
+
+					// Fix up the name here.
+					Paths[PathIndex] = FoundProperty->GetFName();
+				}
+			}
+
+		}
+		else if (Paths.Num() == 0)
+		{
+			// Top-level struct connections will have an empty paths array. 
+			return true;
+		}
+
+		return bResult;
+	}
+}
+
+
+
+bool UNiagaraNodeConvert::RefreshFromExternalChanges()
+{
+	bool bHasChanged = false;
+
+	// Used to get struct definition
+	const UEdGraphSchema_Niagara* Schema = CastChecked<UEdGraphSchema_Niagara>(GetSchema());
+	check(Schema);
+	
+	const UScriptStruct* SourceScriptStruct;
+	const UScriptStruct* DestinationScriptStruct;
+
+	// Cache connections before clean-up for comparison to check if anything changed/a recompile is necessary.
+	const TArray<FNiagaraConvertConnection> OldConnections = Connections;
+
+	// Clean up paths in each connection. Remove or mark orphaned pins, and remove connections as needed.
+	for (int32 ConnectionIndex = 0; ConnectionIndex < Connections.Num(); ConnectionIndex++)
+	{
+		FNiagaraConvertConnection Connection = Connections[ConnectionIndex];
+
+		UEdGraphPin* SourcePin = FNiagaraConvertRefreshHelpers::FindPinByID(Connection.SourcePinId, Pins);
+		UEdGraphPin* DestinationPin = FNiagaraConvertRefreshHelpers::FindPinByID(Connection.DestinationPinId, Pins);
+
+		// Get type definitions for the pins in this connection
+		// PinToTypeDefinition handles nullptr case, so no need to null check them here
+		FNiagaraTypeDefinition SourcePinTypeDef = Schema->PinToTypeDefinition(SourcePin);
+		FNiagaraTypeDefinition DestinationPinTypeDef = Schema->PinToTypeDefinition(DestinationPin);
+
+		SourceScriptStruct = SourcePinTypeDef.GetScriptStruct();
+		DestinationScriptStruct = DestinationPinTypeDef.GetScriptStruct();
+
+		bool bIsSourceScalar = FNiagaraTypeDefinition::IsScalarDefinition(SourcePinTypeDef);
+		bool bIsDestinationScalar = FNiagaraTypeDefinition::IsScalarDefinition(DestinationPinTypeDef);
+		
+		// Rename pins and fix up stored GUID if found in struct
+		// If both pins are structs, search both for the other pin, rename pin, and set GUID as necessary. 
+		if (!bIsSourceScalar && !bIsDestinationScalar)
+		{
+			FGuid InOutSourceGuidToSearch = Connection.SourcePropertyId;
+			FGuid InOutDestGuidToSearch = Connection.DestinationPropertyId;
+
+			FNiagaraConvertRefreshHelpers::RenamePinIfInStruct(SourcePin, DestinationScriptStruct, InOutSourceGuidToSearch);
+			FNiagaraConvertRefreshHelpers::RenamePinIfInStruct(DestinationPin, SourceScriptStruct, InOutDestGuidToSearch);
+
+			if (InOutSourceGuidToSearch != FGuid())
+			{
+				Connection.SourcePropertyId = InOutSourceGuidToSearch;
+			}
+
+			if (InOutDestGuidToSearch != FGuid())
+			{
+				Connection.DestinationPropertyId = InOutDestGuidToSearch;
+			}
+		}
+		else if (bIsSourceScalar && !bIsDestinationScalar)
+		{
+			FGuid InOutGuidToSearch = Connection.SourcePropertyId;
+
+			FNiagaraConvertRefreshHelpers::RenamePinIfInStruct(SourcePin, DestinationScriptStruct, InOutGuidToSearch);
+
+			if (InOutGuidToSearch != FGuid())
+			{	
+				Connection.SourcePropertyId = InOutGuidToSearch;
+			}
+		}
+		else if (!bIsSourceScalar && bIsDestinationScalar)
+		{
+			FGuid InOutGuidToSearch = Connection.DestinationPropertyId;
+
+			FNiagaraConvertRefreshHelpers::RenamePinIfInStruct(DestinationPin, SourceScriptStruct, InOutGuidToSearch);
+
+			if (InOutGuidToSearch != FGuid())
+			{
+				Connection.DestinationPropertyId = InOutGuidToSearch;
+			}
+		}
+
+		// Need to initialize as nullptr for RecursePaths
+		FProperty* SourceProperty = nullptr;
+		FProperty* DestinationProperty = nullptr;
+
+		bool bSourceFound = FNiagaraConvertRefreshHelpers::RecursePaths(SourceScriptStruct, Connection.SourcePath, 0, SourceProperty);
+		bool bDestinationFound = FNiagaraConvertRefreshHelpers::RecursePaths(DestinationScriptStruct, Connection.DestinationPath, 0, DestinationProperty);
+
+		bool bTypesAreAssignable = false;
+
+		// Ensure types are still assignable.
+		if (bSourceFound && bDestinationFound)
+		{
+			// Use the property found by recursing paths if possible
+			// A null Source or Destination property means the path was empty, and we should use the pin's type. 
+			FNiagaraTypeDefinition SourcePropTypeDef = SourceProperty ? Schema->GetTypeDefForProperty(SourceProperty) : SourcePinTypeDef;
+			FNiagaraTypeDefinition DestinationPropTypeDef = DestinationProperty ? Schema->GetTypeDefForProperty(DestinationProperty) : DestinationPinTypeDef;
+
+			bTypesAreAssignable = FNiagaraTypeDefinition::TypesAreAssignable(SourcePropTypeDef, DestinationPropTypeDef);
+		}
+
+
+		// Scalar pins that are not connected should be orphaned
+		if (bIsSourceScalar && (!bDestinationFound || !bTypesAreAssignable))
+		{
+			SourcePin->bOrphanedPin = true;
+		}
+		// If they were orphaned, and now have a valid connection they are no longer orphan pins
+		else if (SourcePin->bOrphanedPin && bDestinationFound && bTypesAreAssignable)
+		{
+			SourcePin->bOrphanedPin = false;
+		}
+
+		// Repeat for destination pins
+		if (bIsDestinationScalar && (!bSourceFound || !bTypesAreAssignable))
+		{
+			DestinationPin->bOrphanedPin = true;
+		}
+		else if (DestinationPin->bOrphanedPin && bSourceFound && bTypesAreAssignable)
+		{
+			DestinationPin->bOrphanedPin = false;
+		}
+
+		// Remove invalid connections
+		if (!bDestinationFound || !bSourceFound || !bTypesAreAssignable)
+		{
+			Connections.RemoveAt(ConnectionIndex);
+			ConnectionIndex--;
+		} 
+		else
+		{
+			// write updated connection to array.
+			Connections[ConnectionIndex] = Connection;
+		}
+	}
+
+	// Compare new and old connections to see if data has changed.
+	if (OldConnections.Num() == Connections.Num())
+	{
+		for (int32 ConnectionIndex = 0; ConnectionIndex < Connections.Num(); ConnectionIndex++)
+		{
+			FNiagaraConvertConnection OldConneciton = OldConnections[ConnectionIndex];
+			FNiagaraConvertConnection NewConnection = Connections[ConnectionIndex];
+
+			if (!(NewConnection == OldConneciton))
+			{
+				bHasChanged = true;
+				break;
+			}
+		}
+	}
+	else
+	{
+		bHasChanged = true;
+	}
+
+	// There are changes to the underlying struct that may not impact the connections stored, so we have no way of testing when a visual refresh is necessary.
+	// So assume the UI always needs to be refreshed if this function is called, since the user is triggering a node refresh explicitly.
+	// If the connections haven't changed, just update the UI, and do not recompile. 
+	// Any necessary recompile and visual update will be handled by the schema, if the data has changed (i.e. this function returns true). See: UEdGraphSchema_Niagara::RefreshNode
+	if (!bHasChanged)
+	{
+		MarkNodeRequiresSynchronization(__FUNCTION__, false);
+	}
+	
+	return bHasChanged;
+}
+
 void UNiagaraNodeConvert::AutowireNewNode(UEdGraphPin* FromPin)
 {
 	const UEdGraphSchema_Niagara* Schema = CastChecked<UEdGraphSchema_Niagara>(GetSchema());
@@ -302,7 +586,8 @@ void UNiagaraNodeConvert::AutowireNewNode(UEdGraphPin* FromPin)
 						SrcPath.Add(TEXT("Value"));
 					}
 					DestPath.Add(*Property->GetName());
-					Connections.Add(FNiagaraConvertConnection(NewPin->PinId, SrcPath, ConnectPin->PinId, DestPath));
+					FGuid PropertyGuid = FStructureEditorUtils::GetGuidForProperty(Property);
+					Connections.Add(FNiagaraConvertConnection(NewPin->PinId, SrcPath, ConnectPin->PinId, DestPath, PropertyGuid, FGuid()));
 					if (SrcPath.Num())
 					{
 						AddExpandedRecord(FNiagaraConvertPinRecord(NewPin->PinId, SrcPath).GetParent());
@@ -319,7 +604,8 @@ void UNiagaraNodeConvert::AutowireNewNode(UEdGraphPin* FromPin)
 					{
 						DestPath.Add(TEXT("Value"));
 					}
-					Connections.Add(FNiagaraConvertConnection(ConnectPin->PinId, SrcPath, NewPin->PinId, DestPath));
+					FGuid PropertyGuid = FStructureEditorUtils::GetGuidForProperty(Property);
+					Connections.Add(FNiagaraConvertConnection(ConnectPin->PinId, SrcPath, NewPin->PinId, DestPath, FGuid(), PropertyGuid));
 					if (DestPath.Num())
 					{
 						AddExpandedRecord(FNiagaraConvertPinRecord(NewPin->PinId, DestPath).GetParent());
@@ -351,6 +637,9 @@ void UNiagaraNodeConvert::AutowireNewNode(UEdGraphPin* FromPin)
 			TEXT("W")
 		};
 
+		// TypeDef won't be initialized for swizzles yet
+		TypeDef = Schema->PinToTypeDefinition(FromPin);
+
 		UEdGraphPin* ConnectPin = RequestNewTypedPin(OppositeDir, TypeDef);
 		check(ConnectPin);
 		ConnectPin->MakeLinkTo(FromPin);
diff --git a/Engine/Plugins/FX/Niagara/Source/NiagaraEditor/Private/NiagaraNodeConvert.h b/Engine/Plugins/FX/Niagara/Source/NiagaraEditor/Private/NiagaraNodeConvert.h
--- a/Engine/Plugins/FX/Niagara/Source/NiagaraEditor/Private/NiagaraNodeConvert.h
+++ b/Engine/Plugins/FX/Niagara/Source/NiagaraEditor/Private/NiagaraNodeConvert.h
@@ -62,22 +62,32 @@ struct FNiagaraConvertConnection
 	UPROPERTY()
 	TArray<FName> DestinationPath;
 
+	// Optional Guids for storing the property's Guid within a struct to help with determining if a rename has occurred.
+	UPROPERTY()
+	FGuid SourcePropertyId;
+
+	UPROPERTY()
+	FGuid DestinationPropertyId;
+
 	FNiagaraConvertConnection()
 	{
 	}
 
-	FNiagaraConvertConnection(FGuid InSourcePinId, const TArray<FName>& InSourcePath, FGuid InDestinationPinId, const TArray<FName>& InDestinationPath)
+	FNiagaraConvertConnection(FGuid InSourcePinId, const TArray<FName>& InSourcePath, FGuid InDestinationPinId, const TArray<FName>& InDestinationPath, FGuid InSourcePropertyId = FGuid(), FGuid InDestinationPropertyId = FGuid())
 		: SourcePinId(InSourcePinId)
 		, SourcePath(InSourcePath)
 		, DestinationPinId(InDestinationPinId)
 		, DestinationPath(InDestinationPath)
+		, SourcePropertyId(InSourcePropertyId)
+		, DestinationPropertyId(InDestinationPropertyId)
 	{
 	}
 
+	bool operator ==(const FNiagaraConvertConnection& B) { return SourcePinId == B.SourcePinId && DestinationPinId == B.DestinationPinId && DestinationPath == B.DestinationPath && SourcePath == B.SourcePath; }
+
 	FString ToString() const;
 };
 
-
 /** A node which allows the user to build a set of arbitrary output types from an arbitrary set of input types by connecting their inner components. */
 UCLASS()
 class UNiagaraNodeConvert : public UNiagaraNodeWithDynamicPins
@@ -97,6 +107,8 @@ public:
 	//~ UNiagaraNode interface
 	virtual void Compile(class FHlslNiagaraTranslator* Translator, TArray<int32>& Outputs) override;
 
+	virtual bool RefreshFromExternalChanges() override;
+
 	/** Gets the nodes inner connection. */
 	TArray<FNiagaraConvertConnection>& GetConnections();
 
