Changed the previous compile severity "Log" to "Display" and added a new"Log" compile message severity that does as the name implies (logs in the log, but doesn't display otherwise)
Useful for messages that we want to show in the log, but not in the stack to help with decluttering
Message severities are all over the place with EMessageSeverity, EStackIssueSeverity, FNiagaraMessageSeverity. Should try to consolidate these eventually but will need some post load fixup for the serialized messages.

diff --git a/Engine/Plugins/FX/Niagara/Source/NiagaraEditor/Private/NiagaraHlslTranslator.cpp b/Engine/Plugins/FX/Niagara/Source/NiagaraEditor/Private/NiagaraHlslTranslator.cpp
--- a/Engine/Plugins/FX/Niagara/Source/NiagaraEditor/Private/NiagaraHlslTranslator.cpp
+++ b/Engine/Plugins/FX/Niagara/Source/NiagaraEditor/Private/NiagaraHlslTranslator.cpp
@@ -6415,7 +6415,7 @@ void FHlslNiagaraTranslator::FunctionCall(UNiagaraNodeFunctionCall* FunctionNode
 				FNiagaraTypeDefinition TypeDef = Schema->PinToTypeDefinition(CallInputs[i]);
 				if (!(TypeDef.IsDataInterface() || TypeDef == FNiagaraTypeDefinition::GetParameterMapDef()))
 				{
-					WriteCompilerTag(Inputs[i], CallInputs[i], false, FNiagaraCompileEventSeverity::Log, ResolvedVariable.GetName().ToString());
+					WriteCompilerTag(Inputs[i], CallInputs[i], false, FNiagaraCompileEventSeverity::Display, ResolvedVariable.GetName().ToString());
 				}
 			}
 		}
diff --git a/Engine/Plugins/FX/Niagara/Source/NiagaraEditor/Private/NiagaraMessageUtilities.cpp b/Engine/Plugins/FX/Niagara/Source/NiagaraEditor/Private/NiagaraMessageUtilities.cpp
--- a/Engine/Plugins/FX/Niagara/Source/NiagaraEditor/Private/NiagaraMessageUtilities.cpp
+++ b/Engine/Plugins/FX/Niagara/Source/NiagaraEditor/Private/NiagaraMessageUtilities.cpp
@@ -202,45 +202,6 @@ UNiagaraStackEntry::FStackIssue FNiagaraMessageUtilities::StackMessageToStackIss
 		FixLinks);
 }
 
-void FNiagaraMessageUtilities::GetCompileMessageData(TSharedRef<const FNiagaraMessageCompileEvent> InMessage, bool& bOutDismissable, FText& OutShortDescription)
-{
-	TSharedRef<FTokenizedMessage> TokenizedMessage = InMessage->GenerateTokenizedMessage();
-	EStackIssueSeverity StackIssueSeverity;
-	switch (TokenizedMessage->GetSeverity())
-	{
-	case EMessageSeverity::CriticalError:
-	case EMessageSeverity::Error:
-		StackIssueSeverity = EStackIssueSeverity::Error;
-		break;
-	case EMessageSeverity::PerformanceWarning:
-	case EMessageSeverity::Warning:
-		StackIssueSeverity = EStackIssueSeverity::Warning;
-		break;
-	case EMessageSeverity::Info:
-		StackIssueSeverity = EStackIssueSeverity::Info;
-		break;
-	default:
-		StackIssueSeverity = EStackIssueSeverity::Info;
-		break;
-	}
-
-	const FName& MessageTopic = InMessage->GetMessageTopic();
-
-	if ((int32)EStackIssueSeverity::None >= (int32)InMessage->GetCompileEvent().Severity)
-	{
-		if (!InMessage->GetCompileEvent().ShortDescription.IsEmpty())
-		{
-			OutShortDescription = FText::FromString(InMessage->GetCompileEvent().ShortDescription);
-		}
-		else if (MessageTopic == FNiagaraMessageTopics::CompilerTopicName)
-		{
-			OutShortDescription = GetShortDescriptionFromSeverity(StackIssueSeverity);
-		}
-	}
-
-	bOutDismissable = InMessage->GetCompileEvent().bDismissable;
-}
-
 FText FNiagaraMessageUtilities::GetShortDescriptionFromSeverity(EStackIssueSeverity Severity)
 {
 	if(Severity == EStackIssueSeverity::Error)
diff --git a/Engine/Plugins/FX/Niagara/Source/NiagaraEditor/Private/NiagaraMessages.cpp b/Engine/Plugins/FX/Niagara/Source/NiagaraEditor/Private/NiagaraMessages.cpp
--- a/Engine/Plugins/FX/Niagara/Source/NiagaraEditor/Private/NiagaraMessages.cpp
+++ b/Engine/Plugins/FX/Niagara/Source/NiagaraEditor/Private/NiagaraMessages.cpp
@@ -49,6 +49,10 @@ TSharedRef<FTokenizedMessage> FNiagaraMessageCompileEvent::GenerateTokenizedMess
 	case FNiagaraCompileEventSeverity::Warning:
 		MessageSeverity = EMessageSeverity::Warning;
 		break;
+	case FNiagaraCompileEventSeverity::Display:
+		MessageSeverity = EMessageSeverity::Info;
+		break;
+	// log is still treated like an info
 	case FNiagaraCompileEventSeverity::Log:
 		MessageSeverity = EMessageSeverity::Info;
 		break;
@@ -153,7 +157,7 @@ FText FNiagaraMessageCompileEvent::GenerateMessageTitle() const
 	case FNiagaraCompileEventSeverity::Warning:
 		StackIssueSeverity = EStackIssueSeverity::Warning;
 		break;
-	case FNiagaraCompileEventSeverity::Log:
+	case FNiagaraCompileEventSeverity::Display:
 		StackIssueSeverity = EStackIssueSeverity::Info;
 		break;
 	default:
diff --git a/Engine/Plugins/FX/Niagara/Source/NiagaraEditor/Private/NiagaraNodeOutputTag.h b/Engine/Plugins/FX/Niagara/Source/NiagaraEditor/Private/NiagaraNodeOutputTag.h
--- a/Engine/Plugins/FX/Niagara/Source/NiagaraEditor/Private/NiagaraNodeOutputTag.h
+++ b/Engine/Plugins/FX/Niagara/Source/NiagaraEditor/Private/NiagaraNodeOutputTag.h
@@ -34,7 +34,7 @@ public:
 	bool bEmitMessageOnFailure = true;
 
 	UPROPERTY(EditAnywhere, Category = Tag)
-	FNiagaraCompileEventSeverity FailureSeverity = FNiagaraCompileEventSeverity::Log;
+	FNiagaraCompileEventSeverity FailureSeverity = FNiagaraCompileEventSeverity::Display;
 
 protected:
 
diff --git a/Engine/Plugins/FX/Niagara/Source/NiagaraEditor/Private/NiagaraNodeStaticSwitch.cpp b/Engine/Plugins/FX/Niagara/Source/NiagaraEditor/Private/NiagaraNodeStaticSwitch.cpp
--- a/Engine/Plugins/FX/Niagara/Source/NiagaraEditor/Private/NiagaraNodeStaticSwitch.cpp
+++ b/Engine/Plugins/FX/Niagara/Source/NiagaraEditor/Private/NiagaraNodeStaticSwitch.cpp
@@ -430,7 +430,7 @@ void UNiagaraNodeStaticSwitch::CheckForOutdatedEnum(FHlslNiagaraTranslator* Tran
 		{
 			if(NewPinNames.Num() == 0)
 			{
-				Translator->Message(FNiagaraCompileEventSeverity::Log, LOCTEXT("EmptyEnumEntries", "A static switch either has no enum assigned or the enum has no entries. A node refresh is advised."), this, nullptr, FString(), true);
+				Translator->Message(FNiagaraCompileEventSeverity::Display, LOCTEXT("EmptyEnumEntries", "A static switch either has no enum assigned or the enum has no entries. A node refresh is advised."), this, nullptr, FString(), true);
 			}
 			else
 			{
diff --git a/Engine/Plugins/FX/Niagara/Source/NiagaraEditor/Private/ViewModels/Stack/NiagaraStackFunctionInput.cpp b/Engine/Plugins/FX/Niagara/Source/NiagaraEditor/Private/ViewModels/Stack/NiagaraStackFunctionInput.cpp
--- a/Engine/Plugins/FX/Niagara/Source/NiagaraEditor/Private/ViewModels/Stack/NiagaraStackFunctionInput.cpp
+++ b/Engine/Plugins/FX/Niagara/Source/NiagaraEditor/Private/ViewModels/Stack/NiagaraStackFunctionInput.cpp
@@ -1858,6 +1858,11 @@ void UNiagaraStackFunctionInput::OnMessageManagerRefresh(const TArray<TSharedRef
 		{
 			for (TSharedRef<const INiagaraMessage> Message : NewMessages)
 			{
+				if(Message->ShouldOnlyLog())
+				{
+					continue;
+				}
+				
 				FStackIssue Issue = FNiagaraMessageUtilities::MessageToStackIssue(Message, GetStackEditorDataKey());
 				if (MessageManagerIssues.ContainsByPredicate([&Issue](const FStackIssue& NewIssue)
 					{ return NewIssue.GetUniqueIdentifier() == Issue.GetUniqueIdentifier(); }) == false)
diff --git a/Engine/Plugins/FX/Niagara/Source/NiagaraEditor/Private/ViewModels/Stack/NiagaraStackModuleItem.cpp b/Engine/Plugins/FX/Niagara/Source/NiagaraEditor/Private/ViewModels/Stack/NiagaraStackModuleItem.cpp
--- a/Engine/Plugins/FX/Niagara/Source/NiagaraEditor/Private/ViewModels/Stack/NiagaraStackModuleItem.cpp
+++ b/Engine/Plugins/FX/Niagara/Source/NiagaraEditor/Private/ViewModels/Stack/NiagaraStackModuleItem.cpp
@@ -1045,6 +1045,12 @@ void UNiagaraStackModuleItem::OnMessageManagerRefresh(const TArray<TSharedRef<co
 		MessageManagerIssues.Reset();
 		for (TSharedRef<const INiagaraMessage> Message : NewMessages)
 		{
+			// we skip issues that do not have relevance for the stack, such as notes that should only appear in the log despite being related to a module
+			if(Message->ShouldOnlyLog())
+			{
+				continue;
+			}
+			
 			// Sometimes compile errors with the same info are generated, so guard against duplicates here.
 			FStackIssue Issue = FNiagaraMessageUtilities::MessageToStackIssue(Message, GetStackEditorDataKey());
 			if (MessageManagerIssues.ContainsByPredicate([&Issue](const FStackIssue& NewIssue)
diff --git a/Engine/Plugins/FX/Niagara/Source/NiagaraEditor/Public/NiagaraMessageUtilities.h b/Engine/Plugins/FX/Niagara/Source/NiagaraEditor/Public/NiagaraMessageUtilities.h
--- a/Engine/Plugins/FX/Niagara/Source/NiagaraEditor/Public/NiagaraMessageUtilities.h
+++ b/Engine/Plugins/FX/Niagara/Source/NiagaraEditor/Public/NiagaraMessageUtilities.h
@@ -16,7 +16,5 @@ namespace FNiagaraMessageUtilities
 
 	UNiagaraStackEntry::FStackIssue StackMessageToStackIssue(const FNiagaraStackMessage& InMessage, FString InStackEditorDataKey, const TArray<FLinkNameAndDelegate>& InLinks);
 
-	void GetCompileMessageData(TSharedRef<const FNiagaraMessageCompileEvent> InMessage, bool& bOutDismissable, FText& OutShortDescription);
-
 	FText GetShortDescriptionFromSeverity(EStackIssueSeverity Severity);
 }
diff --git a/Engine/Plugins/FX/Niagara/Source/NiagaraEditor/Public/NiagaraMessages.h b/Engine/Plugins/FX/Niagara/Source/NiagaraEditor/Public/NiagaraMessages.h
--- a/Engine/Plugins/FX/Niagara/Source/NiagaraEditor/Public/NiagaraMessages.h
+++ b/Engine/Plugins/FX/Niagara/Source/NiagaraEditor/Public/NiagaraMessages.h
@@ -88,6 +88,9 @@ public:
 
 	const uint32 GetMessageTopicBitflag() const;
 
+	/** Can optionally be overridden to allow a message to only be logged instead of appearing elsewhere, like the stack. */
+	virtual bool ShouldOnlyLog() const { return false; }
+
 protected:
 	const TArray<FObjectKey> AssociatedObjectKeys;
 	mutable uint32 MessageTopicBitflag;
@@ -115,7 +118,9 @@ public:
 	virtual void GenerateLinks(TArray<FText>& OutLinkDisplayNames, TArray<FSimpleDelegate>& OutLinkNavigationActions) const override;
 
 	virtual const FName GetMessageTopic() const override { return FNiagaraMessageTopics::CompilerTopicName; };
- 
+
+	virtual bool ShouldOnlyLog() const override { return CompileEvent.Severity == FNiagaraCompileEventSeverity::Log; }
+	
 	const FNiagaraCompileEvent& GetCompileEvent() const { return CompileEvent; }
 private:
 	const FNiagaraCompileEvent CompileEvent;
diff --git a/Engine/Plugins/FX/Niagara/Source/NiagaraShader/Public/NiagaraShared.h b/Engine/Plugins/FX/Niagara/Source/NiagaraShader/Public/NiagaraShared.h
--- a/Engine/Plugins/FX/Niagara/Source/NiagaraShader/Public/NiagaraShared.h
+++ b/Engine/Plugins/FX/Niagara/Source/NiagaraShader/Public/NiagaraShared.h
@@ -39,8 +39,9 @@ UENUM()
 enum class FNiagaraCompileEventSeverity : uint8
 {
 	Log = 0,
-	Warning = 1,
-	Error = 2
+	Display = 1,
+	Warning = 2,
+	Error = 3
 };
 
 UENUM()
@@ -58,7 +59,7 @@ struct FNiagaraCompileEvent
 public:
 	FNiagaraCompileEvent()
 	{
-		Severity = FNiagaraCompileEventSeverity::Log;
+		Severity = FNiagaraCompileEventSeverity::Display;
 		Message = FString();
 		ShortDescription = FString();
 		bDismissable = true;
