Removed Select node override for the allowed pin type functions
This effectively adds that the parameter map isn't allowed anymore

diff --git a/Engine/Plugins/FX/Niagara/Source/NiagaraEditor/Private/NiagaraNodeSelect.cpp b/Engine/Plugins/FX/Niagara/Source/NiagaraEditor/Private/NiagaraNodeSelect.cpp
--- a/Engine/Plugins/FX/Niagara/Source/NiagaraEditor/Private/NiagaraNodeSelect.cpp
+++ b/Engine/Plugins/FX/Niagara/Source/NiagaraEditor/Private/NiagaraNodeSelect.cpp
@@ -484,13 +484,6 @@ bool UNiagaraNodeSelect::CanRenamePin(const UEdGraphPin* Pin) const
 	return Super::CanRenamePin(Pin) && Pin->Direction == EGPD_Output && Pin->bOrphanedPin == false && !UEdGraphSchema_Niagara::IsPinWildcard(Pin);
 }
 
-bool UNiagaraNodeSelect::AllowNiagaraTypeForAddPin(const FNiagaraTypeDefinition& InType) const
-{
-	return InType.GetScriptStruct() != nullptr
-		&& InType != FNiagaraTypeDefinition::GetGenericNumericDef()
-		&& !InType.IsInternalType();
-}
-
 FString UNiagaraNodeSelect::GetInputCaseName(int32 Case) const
 {
 	if(SelectorPinType == FNiagaraTypeDefinition::GetBoolDef())
diff --git a/Engine/Plugins/FX/Niagara/Source/NiagaraEditor/Private/NiagaraNodeSelect.h b/Engine/Plugins/FX/Niagara/Source/NiagaraEditor/Private/NiagaraNodeSelect.h
--- a/Engine/Plugins/FX/Niagara/Source/NiagaraEditor/Private/NiagaraNodeSelect.h
+++ b/Engine/Plugins/FX/Niagara/Source/NiagaraEditor/Private/NiagaraNodeSelect.h
@@ -67,7 +67,6 @@ private:
 	virtual bool CanMovePin(const UEdGraphPin* Pin, int32 DirectionToMove) const override;
 	virtual void MoveDynamicPin(UEdGraphPin* Pin, int32 DirectionToMove) override;
 	virtual bool CanRenamePin(const UEdGraphPin* Pin) const override;
-	virtual bool AllowNiagaraTypeForAddPin(const FNiagaraTypeDefinition& InType) const override;
 
 	/** UNiagaraNodeUsageSelector interface */
 	virtual FString GetInputCaseName(int32 Case) const override;
