- Reactive components who were culled by scalability when we resolve the owner status

#rnx

diff --git a/Engine/Plugins/FX/Niagara/Source/Niagara/Private/NiagaraComponent.cpp b/Engine/Plugins/FX/Niagara/Source/Niagara/Private/NiagaraComponent.cpp
--- a/Engine/Plugins/FX/Niagara/Source/Niagara/Private/NiagaraComponent.cpp
+++ b/Engine/Plugins/FX/Niagara/Source/Niagara/Private/NiagaraComponent.cpp
@@ -1933,13 +1933,23 @@ bool UNiagaraComponent::ResolveOwnerAllowsScalability(bool bRegister)
 	//Update manager registration if needed.
 	if (bRegister)
 	{
-		if (IsActive() && bOwnerAllowsScalabiltiy)
+		if (bOwnerAllowsScalabiltiy)
 		{
-			RegisterWithScalabilityManager();
+			if ( IsActive() )
+			{
+				RegisterWithScalabilityManager();
+			}
 		}
 		else if (!bOwnerAllowsScalabiltiy)
 		{
-			UnregisterWithScalabilityManager();
+			if (!IsActive() && bIsCulledByScalability)
+			{
+				ActivateInternal(false, true);
+			}
+			else
+			{
+				UnregisterWithScalabilityManager();
+			}
 		}
 	}
 	return bOwnerAllowsScalabiltiy;
