Adding support for Static Variables to Niagara stack.
Sub features:
+ Static vars for custom hlsl
+ handle basic boolean operations, !, &&, ||
+ handle basic math ops
+ Added "Set From Pin" to static switch
+ unnamed static switch text changes when wired in
+ Default linkage of static variables settable
+ bools vs static bools visual distinction
+ iconic stack representation of static variables, an S adjacent
+ custom hlsl can input static vars but cannot output static vars
+ make sure we only have 1 static var write in graph
+ make sure static vars don't make it into payloads
+ make sure static vars can generate errors if fail on not set
+ make sure that static variables hide/show properly

diff --git a/Engine/Plugins/FX/Niagara/Source/Niagara/Classes/NiagaraDataSetDebugAccessor.h b/Engine/Plugins/FX/Niagara/Source/Niagara/Classes/NiagaraDataSetDebugAccessor.h
--- a/Engine/Plugins/FX/Niagara/Source/Niagara/Classes/NiagaraDataSetDebugAccessor.h
+++ b/Engine/Plugins/FX/Niagara/Source/Niagara/Classes/NiagaraDataSetDebugAccessor.h
@@ -29,7 +29,7 @@ struct FNiagaraDataSetDebugAccessor
 		}
 		else if (IsInt())
 		{
-			if (NiagaraType == FNiagaraTypeDefinition::GetBoolDef())
+			if (NiagaraType.IsSameBaseDefinition(FNiagaraTypeDefinition::GetBoolDef()))
 			{
 				const TCHAR* TrueText = TEXT("true");
 				const TCHAR* FalseText = TEXT("false");
diff --git a/Engine/Plugins/FX/Niagara/Source/Niagara/Classes/NiagaraEmitter.h b/Engine/Plugins/FX/Niagara/Source/Niagara/Classes/NiagaraEmitter.h
--- a/Engine/Plugins/FX/Niagara/Source/Niagara/Classes/NiagaraEmitter.h
+++ b/Engine/Plugins/FX/Niagara/Source/Niagara/Classes/NiagaraEmitter.h
@@ -254,6 +254,8 @@ public:
 
 	/** Helper method for binding the notifications needed for proper editor integration. */
 	NIAGARA_API void BindNotifications();
+
+	NIAGARA_API void GatherStaticVariables(TArray<FNiagaraVariable>& OutVars) const;
 #endif
 	virtual bool NeedsLoadForTargetPlatform(const ITargetPlatform* TargetPlatform) const override;
 	void Serialize(FArchive& Ar)override;
@@ -277,6 +279,9 @@ public:
 	/** Get All adapters to editor only script vars owned directly by this subscriber. */
 	virtual TArray<UNiagaraEditorParametersAdapterBase*> GetEditorOnlyParametersAdapters() override;
 	//~ End INiagaraParameterDefinitionsSubscriber Interface
+
+	/** Get the cached parameter map traversal for this emitter.  */
+	NIAGARA_API const TSharedPtr<FNiagaraGraphCachedDataBase, ESPMode::ThreadSafe>& GetCachedTraversalData() const;
 #endif
 
 	bool IsEnabledOnPlatform(const FString& PlatformName)const;
@@ -325,6 +330,7 @@ public:
 	    Used in conjunction with UNiagaraSystem::bTrimAttributes */
 	UPROPERTY(EditAnywhere, AdvancedDisplay, Category = "Emitter")
 	TArray<FString> AttributesToPreserve;
+
 #endif
 
 	UPROPERTY(EditAnywhere, Category = "Emitter")
@@ -590,6 +596,11 @@ public:
 
 	TConstArrayView<TUniquePtr<FNiagaraBoundsCalculator>> GetBoundsCalculators() const { return MakeArrayView(BoundsCalculators); }
 
+	UPROPERTY()
+	FNiagaraParameterStore RendererBindings;
+
+	void RebuildRendererBindings();
+
 protected:
 	virtual void BeginDestroy() override;
 
@@ -613,6 +624,7 @@ private:
 
 	void PersistentEditorDataChanged();
 
+	mutable TSharedPtr<FNiagaraGraphCachedDataBase, ESPMode::ThreadSafe> CachedTraversalData;
 private:
 	/** Adjusted every time that we compile this emitter. Lets us know that we might differ from any cached versions.*/
 	UPROPERTY()
diff --git a/Engine/Plugins/FX/Niagara/Source/Niagara/Classes/NiagaraEmitterHandle.h b/Engine/Plugins/FX/Niagara/Source/Niagara/Classes/NiagaraEmitterHandle.h
--- a/Engine/Plugins/FX/Niagara/Source/Niagara/Classes/NiagaraEmitterHandle.h
+++ b/Engine/Plugins/FX/Niagara/Source/Niagara/Classes/NiagaraEmitterHandle.h
@@ -7,6 +7,7 @@
 #include "UObject/Object.h"
 #include "Misc/Guid.h"
 #include "NiagaraModule.h"
+#include "NiagaraScript.h"
 #include "NiagaraEmitterHandle.generated.h"
 
 class UNiagaraSystem;
@@ -107,6 +108,7 @@ private:
 
 	UPROPERTY(Transient)
 	bool bIsolated;
+
 #endif
 
 	/** The copied instance of the emitter this handle references. */
diff --git a/Engine/Plugins/FX/Niagara/Source/Niagara/Classes/NiagaraScript.h b/Engine/Plugins/FX/Niagara/Source/Niagara/Classes/NiagaraScript.h
--- a/Engine/Plugins/FX/Niagara/Source/Niagara/Classes/NiagaraScript.h
+++ b/Engine/Plugins/FX/Niagara/Source/Niagara/Classes/NiagaraScript.h
@@ -474,6 +474,10 @@ public:
 	UPROPERTY()
 	TArray<FNiagaraVariableBase> AttributesWritten;
 
+	/** List of all attributes explicitly written by this VM script graph. Used to verify external dependencies.*/
+	UPROPERTY()
+	TArray<FNiagaraVariable> StaticVariablesWritten;
+
 	UPROPERTY()
 	FString ErrorMsg;
 
@@ -504,6 +508,12 @@ public:
 #endif
 };
 
+struct NIAGARA_API FNiagaraGraphCachedDataBase
+{
+	virtual ~FNiagaraGraphCachedDataBase() {}
+	virtual void GetStaticVariables(TArray<FNiagaraVariable>& OutVars) {}
+};
+
 /** Struct containing all of the data that can be different between different script versions.*/
 USTRUCT()
 struct NIAGARA_API FVersionedNiagaraScriptData
@@ -866,6 +876,8 @@ public:
 	static bool IsSystemSpawnScript(ENiagaraScriptUsage Usage) { return Usage == ENiagaraScriptUsage::SystemSpawnScript; }
 	static bool IsSystemUpdateScript(ENiagaraScriptUsage Usage) { return Usage == ENiagaraScriptUsage::SystemUpdateScript; }
 	static bool IsSystemScript(ENiagaraScriptUsage Usage) { return IsSystemSpawnScript(Usage) || IsSystemUpdateScript(Usage);}
+
+	static bool IsEmitterScript(ENiagaraScriptUsage Usage) { return IsEmitterSpawnScript(Usage) || IsEmitterUpdateScript(Usage); }
 	static bool IsEmitterSpawnScript(ENiagaraScriptUsage Usage) { return Usage == ENiagaraScriptUsage::EmitterSpawnScript; }
 	static bool IsEmitterUpdateScript(ENiagaraScriptUsage Usage) { return Usage == ENiagaraScriptUsage::EmitterUpdateScript; }
 	static bool IsStandaloneScript(ENiagaraScriptUsage Usage)  { return IsDynamicInputScript(Usage) || IsFunctionScript(Usage) || IsModuleScript(Usage); }
@@ -1085,6 +1097,8 @@ public:
 	/** Kicks off an async job to convert the ByteCode into an optimized version for the platform we are running on. */
 	FGraphEventRef HandleByteCodeOptimization(bool bShouldForceNow = false);
 
+	static int32 NIAGARA_API LogCompileStaticVars;
+
 private:
 	/** Generates all of the function bindings for DI that don't require user data */
 	void GenerateDefaultFunctionBindings();
diff --git a/Engine/Plugins/FX/Niagara/Source/Niagara/Classes/NiagaraSystem.h b/Engine/Plugins/FX/Niagara/Source/Niagara/Classes/NiagaraSystem.h
--- a/Engine/Plugins/FX/Niagara/Source/Niagara/Classes/NiagaraSystem.h
+++ b/Engine/Plugins/FX/Niagara/Source/Niagara/Classes/NiagaraSystem.h
@@ -708,6 +708,9 @@ public:
 
 	FORCEINLINE bool GetOverrideScalabilitySettings()const { return bOverrideScalabilitySettings; }
 	FORCEINLINE void SetOverrideScalabilitySettings(bool bOverride) { bOverrideScalabilitySettings = bOverride; }
+
+
+	void GatherStaticVariables(TArray<FNiagaraVariable>& OutVars, TArray<FNiagaraVariable>& OutEmitterVars) const;
 #endif
 	UNiagaraEffectType* GetEffectType()const;
 	FORCEINLINE const FNiagaraSystemScalabilitySettings& GetScalabilitySettings()const { return CurrentScalabilitySettings; }
@@ -745,11 +748,18 @@ public:
 	UNiagaraBakerSettings* GetBakerSettings();
 	const UNiagaraBakerSettings* GetBakerGeneratedSettings() const { return BakerGeneratedSettings; }
 	void SetBakerGeneratedSettings(UNiagaraBakerSettings* Settings) { BakerGeneratedSettings = Settings; }
+
+
+	/** Get the cached parameter map traversal for this emitter.  */
+	const TSharedPtr<FNiagaraGraphCachedDataBase, ESPMode::ThreadSafe>& GetCachedTraversalData() const;
+	void InvalidateCachedData();
+	void GraphSourceChanged();
 #endif
 
 private:
 #if WITH_EDITORONLY_DATA
 
+
 	/** Since the shader compilation is done in another process, this is used to check if the result for any ongoing compilations is done.
 	*   If bWait is true then this *blocks* the game thread (and ui) until all shader compilations are finished.
 	*/
@@ -774,6 +784,7 @@ private:
 	/** Helper for filling in attribute datasets per emitter. */
 	void InitEmitterDataSetCompiledData(FNiagaraDataSetCompiledData& DataSetToInit, const UNiagaraEmitter* InAssociatedEmitter, const FNiagaraEmitterHandle& InAssociatedEmitterHandle);
 	void PrepareRapidIterationParametersForCompilation();
+
 #endif
 
 	void ResolveScalabilitySettings();
@@ -814,6 +825,8 @@ protected:
 
 	UPROPERTY(Transient)
 	TArray<FNiagaraSystemCompileRequest> ActiveCompilations;
+
+	mutable TSharedPtr<FNiagaraGraphCachedDataBase, ESPMode::ThreadSafe> CachedTraversalData;
 #endif
 
 	/** The script which defines the System parameters, and which generates the bindings from System
diff --git a/Engine/Plugins/FX/Niagara/Source/Niagara/Private/NiagaraEmitter.cpp b/Engine/Plugins/FX/Niagara/Source/Niagara/Private/NiagaraEmitter.cpp
--- a/Engine/Plugins/FX/Niagara/Source/Niagara/Private/NiagaraEmitter.cpp
+++ b/Engine/Plugins/FX/Niagara/Source/Niagara/Private/NiagaraEmitter.cpp
@@ -975,6 +975,8 @@ void UNiagaraEmitter::HandleVariableRenamed(const FNiagaraVariable& InOldVariabl
 		}
 	}
 
+	RebuildRendererBindings();
+
 	if (bUpdateContexts)
 	{
 		FNiagaraSystemUpdateContext UpdateCtx(this, true);
@@ -990,6 +992,8 @@ void UNiagaraEmitter::HandleVariableRemoved(const FNiagaraVariable& InOldVariabl
 		Prop->RemoveVariable(InOldVariable, this);
 	}
 
+	RebuildRendererBindings();
+
 	if (bUpdateContexts)
 	{
 		FNiagaraSystemUpdateContext UpdateCtx(this, true);
@@ -1026,6 +1030,16 @@ TArray<UNiagaraEditorParametersAdapterBase*> UNiagaraEmitter::GetEditorOnlyParam
 {
 	return { GetEditorParameters() };
 }
+
+const TSharedPtr<FNiagaraGraphCachedDataBase, ESPMode::ThreadSafe>& UNiagaraEmitter::GetCachedTraversalData() const
+{
+	if (CachedTraversalData.IsValid())
+		return CachedTraversalData;
+
+	INiagaraModule& NiagaraModule = FModuleManager::Get().LoadModuleChecked<INiagaraModule>(TEXT("Niagara"));
+	CachedTraversalData = NiagaraModule.CacheGraphTraversal(this, FGuid());
+	return CachedTraversalData;
+}
 #endif // WITH_EDITORONLY_DATA
 
 bool UNiagaraEmitter::IsEnabledOnPlatform(const FString& PlatformName)const
@@ -1232,41 +1246,41 @@ void UNiagaraEmitter::CacheFromCompiledData(const FNiagaraDataSetCompiledData* C
 		}
 
 		// See how many particles we can fit in a GPU buffer. This number can be quite small on some platforms.
-		const uint64 MaxBufferElements = (GNiagaraEmitterMaxGPUBufferElements > 0) ? uint64(GNiagaraEmitterMaxGPUBufferElements) : GetMaxBufferDimension();
+const uint64 MaxBufferElements = (GNiagaraEmitterMaxGPUBufferElements > 0) ? uint64(GNiagaraEmitterMaxGPUBufferElements) : GetMaxBufferDimension();
 
-		// Don't just cast the result of the division to 32-bit, since that will produce garbage if MaxNumInstances is larger than UINT_MAX. Saturate instead.
-		MaxAllocationCount = (uint32)FMath::Min(MaxBufferElements / MaxGPUBufferComponents, (uint64)UINT_MAX);
-		MaxInstanceCount = MaxAllocationCount;
+// Don't just cast the result of the division to 32-bit, since that will produce garbage if MaxNumInstances is larger than UINT_MAX. Saturate instead.
+MaxAllocationCount = (uint32)FMath::Min(MaxBufferElements / MaxGPUBufferComponents, (uint64)UINT_MAX);
+MaxInstanceCount = MaxAllocationCount;
 
-		if ( SimTarget == ENiagaraSimTarget::GPUComputeSim )
-		{
-			// Round down to nearest thread group size
-			MaxAllocationCount = FMath::DivideAndRoundDown(MaxAllocationCount, NiagaraComputeMaxThreadGroupSize) * NiagaraComputeMaxThreadGroupSize;
+if (SimTarget == ENiagaraSimTarget::GPUComputeSim)
+{
+	// Round down to nearest thread group size
+	MaxAllocationCount = FMath::DivideAndRoundDown(MaxAllocationCount, NiagaraComputeMaxThreadGroupSize) * NiagaraComputeMaxThreadGroupSize;
 
-			// We can't go below a thread group size, the rest of the code does not handle this and we should only hit this when buffer size forcing is enabled
-			if ( MaxAllocationCount < NiagaraComputeMaxThreadGroupSize )
-			{
-				MaxAllocationCount = NiagaraComputeMaxThreadGroupSize;
-				check(MaxAllocationCount*MaxGPUBufferComponents <= GetMaxBufferDimension());
-			}
+	// We can't go below a thread group size, the rest of the code does not handle this and we should only hit this when buffer size forcing is enabled
+	if (MaxAllocationCount < NiagaraComputeMaxThreadGroupSize)
+	{
+		MaxAllocationCount = NiagaraComputeMaxThreadGroupSize;
+		check(MaxAllocationCount * MaxGPUBufferComponents <= GetMaxBufferDimension());
+	}
 
-			// -1 because we need a scratch instance on the GPU
-			MaxInstanceCount = MaxAllocationCount - 1;
-		}
+	// -1 because we need a scratch instance on the GPU
+	MaxInstanceCount = MaxAllocationCount - 1;
+}
 
-		if (AllocationMode == EParticleAllocationMode::FixedCount)
-		{
-			MaxInstanceCount = FMath::Min(MaxInstanceCount, uint32(FMath::Max(PreAllocationCount, 0)));
-		}
+if (AllocationMode == EParticleAllocationMode::FixedCount)
+{
+	MaxInstanceCount = FMath::Min(MaxInstanceCount, uint32(FMath::Max(PreAllocationCount, 0)));
+}
 	}
 	else
 	{
-		MaxInstanceCount = 0;
+	MaxInstanceCount = 0;
 	}
 
 	// Cache shaders for all GPU scripts
 #if WITH_EDITORONLY_DATA
-	if ( AreAllScriptAndSourcesSynchronized() )
+	if (AreAllScriptAndSourcesSynchronized())
 	{
 		ForEachScript(
 			[](UNiagaraScript* Script)
@@ -1280,7 +1294,7 @@ void UNiagaraEmitter::CacheFromCompiledData(const FNiagaraDataSetCompiledData* C
 #if WITH_NIAGARA_DEBUG_EMITTER_NAME
 	// Ensure our debug simulation name is up to date
 	// Only required on uncooked as it can change due to compilation
-	if ( !FPlatformProperties::RequiresCookedData() )
+	if (!FPlatformProperties::RequiresCookedData())
 	{
 		DebugSimName.Empty();
 		if (const UNiagaraSystem* SystemOwner = Cast<const UNiagaraSystem>(GetOuter()))
@@ -1291,6 +1305,62 @@ void UNiagaraEmitter::CacheFromCompiledData(const FNiagaraDataSetCompiledData* C
 		DebugSimName.Append(GetName());
 	}
 #endif
+
+	RebuildRendererBindings();
+}
+
+void UNiagaraEmitter::RebuildRendererBindings()
+{
+#if WITH_EDITORONLY_DATA
+
+	// For right now we only capture the values of the static variables here. So we create a temp
+	// parameter store for ALL possible values, then fill in the real one with the static variables only.
+	bool bAnyRendererBindingsAdded = false;
+	FNiagaraParameterStore TempStore;
+	for (UNiagaraRendererProperties* Props : GetRenderers())
+	{
+		if (Props && Props->bIsEnabled)
+		{
+			bAnyRendererBindingsAdded |= Props->PopulateRequiredBindings(TempStore);
+		}
+	}
+
+	TArray<UNiagaraScript*> TargetScripts;
+	if (UNiagaraSystem* SystemOwner = Cast<UNiagaraSystem>(GetOuter()))
+	{
+		TargetScripts.Add(SystemOwner->GetSystemSpawnScript());
+		TargetScripts.Add(SystemOwner->GetSystemUpdateScript());
+	}
+	GetScripts(TargetScripts, false, true);
+
+	if (bAnyRendererBindingsAdded)
+	{
+		RendererBindings.Empty();
+		TArrayView<const FNiagaraVariableWithOffset> Vars = TempStore.ReadParameterVariables();
+		for (const FNiagaraVariableWithOffset& Var : Vars)
+		{
+			bool bFound = false;
+			for (UNiagaraScript* Script : TargetScripts)
+			{
+				if (Script)
+				{
+					for (const FNiagaraVariable& StaticVar : Script->GetVMExecutableData().StaticVariablesWritten)
+					{
+						if (StaticVar.GetType().IsSameBaseDefinition(Var.GetType()) && StaticVar.GetName() == Var.GetName())
+						{
+							RendererBindings.AddParameter(StaticVar);
+							bFound = true;
+							break;
+						}
+					}
+				}
+
+				if (bFound)
+					break;
+			}
+		}
+	}
+#endif
 }
 
 void UNiagaraEmitter::CacheFromShaderCompiled()
@@ -1643,6 +1713,33 @@ void UNiagaraEmitter::BindNotifications()
 	}
 }
 
+void UNiagaraEmitter::GatherStaticVariables(TArray<FNiagaraVariable>& OutVars) const
+{
+	TArray<UNiagaraScript*> OutScripts;
+	GetScripts(OutScripts, false, true);
+
+	for (UNiagaraScript* Script : OutScripts)
+	{
+		TArray<FNiagaraVariable> StoreParams;
+		Script->RapidIterationParameters.GetParameters(StoreParams);
+
+		for (int32 i = 0; i < StoreParams.Num(); i++)
+		{
+			if (StoreParams[i].GetType().IsStatic())
+			{
+				const int32* Index = Script->RapidIterationParameters.FindParameterOffset(StoreParams[i]);
+				if (Index != nullptr)
+				{
+					StoreParams[i].SetData(Script->RapidIterationParameters.GetParameterData(*Index)); // This will memcopy the data in.			
+					OutVars.AddUnique(StoreParams[i]);					
+					//UE_LOG(LogNiagara, Log, TEXT("UNiagaraEmitter::GatherStaticVariables Added %s"), *StoreParams[i].ToString());
+				}
+			}
+		}
+	}
+}
+
+
 #endif
 
 bool UNiagaraEmitter::UsesScript(const UNiagaraScript* Script)const
@@ -1961,6 +2058,8 @@ void UNiagaraEmitter::AddRenderer(UNiagaraRendererProperties* Renderer)
 	UpdateChangeId(TEXT("Renderer added"));
 	OnRenderersChangedDelegate.Broadcast();
 #endif
+	RebuildRendererBindings();
+
 	if (UNiagaraSystem* Owner = GetTypedOuter<UNiagaraSystem>())
 	{
 		Owner->ComputeRenderersDrawOrder();
@@ -1977,6 +2076,8 @@ void UNiagaraEmitter::RemoveRenderer(UNiagaraRendererProperties* Renderer)
 	UpdateChangeId(TEXT("Renderer removed"));
 	OnRenderersChangedDelegate.Broadcast();
 #endif
+	RebuildRendererBindings();
+
 	if (UNiagaraSystem* Owner = GetTypedOuter<UNiagaraSystem>())
 	{
 		Owner->ComputeRenderersDrawOrder();
@@ -2111,6 +2212,7 @@ void UNiagaraEmitter::UpdateChangeId(const FString& Reason)
 #if STATS
 	StatDatabase.ClearStatCaptures();
 #endif
+	CachedTraversalData.Reset();
 }
 
 void UNiagaraEmitter::ScriptRapidIterationParameterChanged()
@@ -2131,6 +2233,10 @@ void UNiagaraEmitter::RendererChanged()
 void UNiagaraEmitter::GraphSourceChanged()
 {
 	UpdateChangeId(TEXT("Graph source changed."));
+
+	UNiagaraSystem* Sys = Cast<UNiagaraSystem>(GetOuter());
+	if (Sys)
+		Sys->GraphSourceChanged();
 }
 
 void UNiagaraEmitter::RaiseOnEmitterGPUCompiled(UNiagaraScript* InScript, const FGuid& ScriptVersion)
diff --git a/Engine/Plugins/FX/Niagara/Source/Niagara/Private/NiagaraEmitterInstance.cpp b/Engine/Plugins/FX/Niagara/Source/Niagara/Private/NiagaraEmitterInstance.cpp
--- a/Engine/Plugins/FX/Niagara/Source/Niagara/Private/NiagaraEmitterInstance.cpp
+++ b/Engine/Plugins/FX/Niagara/Source/Niagara/Private/NiagaraEmitterInstance.cpp
@@ -514,6 +514,9 @@ void FNiagaraEmitterInstance::Init(int32 InEmitterIdx, FNiagaraSystemInstanceID
 			{
 				GPUExecContext->CombinedParamStore.Bind(&RendererBindings);
 			}
+
+			// Handle populating static variable values into the list
+			CachedEmitter->RendererBindings.Bind(&RendererBindings);
 		}
 	}	
 
diff --git a/Engine/Plugins/FX/Niagara/Source/Niagara/Private/NiagaraModule.cpp b/Engine/Plugins/FX/Niagara/Source/Niagara/Private/NiagaraModule.cpp
--- a/Engine/Plugins/FX/Niagara/Source/Niagara/Private/NiagaraModule.cpp
+++ b/Engine/Plugins/FX/Niagara/Source/Niagara/Private/NiagaraModule.cpp
@@ -552,6 +552,12 @@ TSharedPtr<FNiagaraCompileRequestDuplicateDataBase, ESPMode::ThreadSafe> INiagar
 	return PrecompileDuplicateDelegate.Execute(OwningSystemRequestData, OwningSystem, OwningEmitter, TargetScript, TargetVersion);
 }
 
+TSharedPtr<FNiagaraGraphCachedDataBase, ESPMode::ThreadSafe> INiagaraModule::CacheGraphTraversal(const UObject* Obj, FGuid Version)
+{
+	checkf(GraphTraversalCacheDelegate.IsBound(), TEXT("GraphTraversalCache delegate not bound."));
+	return GraphTraversalCacheDelegate.Execute(Obj, Version);
+}
+
 FDelegateHandle INiagaraModule::RegisterPrecompiler(FOnPrecompile PreCompiler)
 {
 	checkf(PrecompileDelegate.IsBound() == false, TEXT("Only one handler is allowed for the Precompiler delegate"));
@@ -566,6 +572,8 @@ void INiagaraModule::UnregisterPrecompiler(FDelegateHandle DelegateHandle)
 	PrecompileDelegate.Unbind();
 }
 
+
+
 FDelegateHandle INiagaraModule::RegisterPrecompileDuplicator(FOnPrecompileDuplicate PreCompileDuplicator)
 {
 	checkf(PrecompileDuplicateDelegate.IsBound() == false, TEXT("Only one handler is allowed for the PrecompileDuplicator delegate"));
@@ -580,6 +588,23 @@ void INiagaraModule::UnregisterPrecompileDuplicator(FDelegateHandle DelegateHand
 	PrecompileDuplicateDelegate.Unbind();
 }
 
+
+FDelegateHandle INiagaraModule::RegisterGraphTraversalCacher(FOnCacheGraphTraversal GraphTraversalCacher)
+{
+	checkf(GraphTraversalCacheDelegate.IsBound() == false, TEXT("Only one handler is allowed for the GraphTraversalCacher delegate"));
+	GraphTraversalCacheDelegate = GraphTraversalCacher;
+	return GraphTraversalCacheDelegate.GetHandle();
+}
+
+void INiagaraModule::UnregisterGraphTraversalCacher(FDelegateHandle DelegateHandle)
+{
+	checkf(GraphTraversalCacheDelegate.IsBound(), TEXT("ObjectGraphTraversalCacher is not registered"));
+	checkf(GraphTraversalCacheDelegate.GetHandle() == DelegateHandle, TEXT("Can only unregister the GraphTraversalCacher delegate with the handle it was registered with."));
+	GraphTraversalCacheDelegate.Unbind();
+}
+
+
+
 void INiagaraModule::OnAssetLoaded(UObject* Asset)
 {
 	if (Asset->IsA<UNiagaraSystem>())
@@ -876,12 +901,22 @@ bool FNiagaraTypeDefinition::AppendCompileHash(FNiagaraCompileHashVisitor* InVis
 		InVisitor->UpdateString(TEXT("\tTDName"), InvalidStr);
 	}
 
+	InVisitor->UpdatePOD(TEXT("\tFlags"), Flags);
+
 	return true;
 #else
 	return false;
 #endif
 }
 
+FNiagaraTypeDefinition FNiagaraTypeDefinition::ToStaticDef() const
+{
+	FNiagaraTypeDefinition Def(*this);
+	Def.Flags |= TF_Static;
+	return Def;
+}
+
+
 #if WITH_EDITORONLY_DATA
 bool FNiagaraTypeDefinition::IsInternalType() const
 {
@@ -922,6 +957,12 @@ void FNiagaraTypeDefinition::RecreateUserDefinedTypeRegistry()
 	FNiagaraTypeRegistry::Register(PositionDef, ParamFlags | PayloadFlags);
 	FNiagaraTypeRegistry::Register(QuatDef, ParamFlags | PayloadFlags);
 	FNiagaraTypeRegistry::Register(Matrix4Def, ParamFlags);
+
+	// Register static versions
+	ENiagaraTypeRegistryFlags StaticParamFlags = ENiagaraTypeRegistryFlags::AllowSystemVariable | ENiagaraTypeRegistryFlags::AllowEmitterVariable | ENiagaraTypeRegistryFlags::AllowParticleVariable | ENiagaraTypeRegistryFlags::AllowParameter;
+	FNiagaraTypeRegistry::Register(IntDef.ToStaticDef(), StaticParamFlags);
+	FNiagaraTypeRegistry::Register(BoolDef.ToStaticDef(), StaticParamFlags);
+
 	// @todo wildcard shouldn't be available for parameters etc., so just don't register here?
 	//FNiagaraTypeRegistry::Register(WildcardDef, VarFlags);
 
@@ -1040,7 +1081,9 @@ void FNiagaraTypeDefinition::RecreateUserDefinedTypeRegistry()
 					{
 						Flags |= ENiagaraTypeRegistryFlags::AllowPayload;
 					}
-					FNiagaraTypeRegistry::Register(Enum, Flags);
+					FNiagaraTypeDefinition Def(Enum);
+					FNiagaraTypeRegistry::Register(Def, Flags);
+					FNiagaraTypeRegistry::Register(Def.ToStaticDef(), StaticParamFlags);
 				}
 
 				UObjectRedirector* Redirector = FindObject<UObjectRedirector>(nullptr, *AssetRefPathNamePreResolve.ToString());
@@ -1064,58 +1107,70 @@ bool FNiagaraTypeDefinition::IsScalarDefinition(const FNiagaraTypeDefinition& Ty
 	return ScalarStructs.Contains(Type.GetScriptStruct()) || (Type.GetScriptStruct() == IntStruct && Type.GetEnum() != nullptr);
 }
 
-bool FNiagaraTypeDefinition::TypesAreAssignable(const FNiagaraTypeDefinition& TypeA, const FNiagaraTypeDefinition& TypeB)
+bool FNiagaraTypeDefinition::TypesAreAssignable(const FNiagaraTypeDefinition& TypeInput, const FNiagaraTypeDefinition& TypeOutput)
 {
-	if (const UClass* AClass = TypeA.GetClass())
+	if (const UClass* AClass = TypeInput.GetClass())
 	{
-		if (const UClass* BClass = TypeB.GetClass())
+		if (const UClass* BClass = TypeOutput.GetClass())
 		{
 			return AClass == BClass;
 			return true;
 		}
 	}
 	
-	if (const UClass* BClass = TypeB.GetClass())
+	if (const UClass* BClass = TypeOutput.GetClass())
 	{
 		return false;
 	}
 
-	if (const UClass* AClass = TypeA.GetClass())
+	if (const UClass* AClass = TypeInput.GetClass())
 	{
 		return false;
 	}
 
 	// Make sure that enums are not assignable to enums of different types or just plain ints
-	if (TypeA.GetStruct() == TypeB.GetStruct() &&
-		TypeA.GetEnum() != TypeB.GetEnum())
+	if (TypeInput.GetStruct() == TypeOutput.GetStruct() &&
+		TypeInput.GetEnum() != TypeOutput.GetEnum())
+	{
+		return false;
+	}
+
+
+	// Static can go to static or static can go to nonstatic, but nonstatic can't go to static
+	if (((!TypeInput.IsStatic() && TypeOutput.IsStatic()) || (TypeInput.IsStatic() && TypeOutput.IsStatic())) && TypeInput.GetStruct() == TypeOutput.GetStruct() && TypeInput.GetEnum() == TypeOutput.GetEnum())
+	{
+		return true;
+	}
+	else if ((TypeInput.IsStatic() || TypeOutput.IsStatic()) && TypeInput.GetStruct() == TypeOutput.GetStruct() && TypeInput.GetEnum() == TypeOutput.GetEnum())
 	{
 		return false;
 	}
 
-	if (TypeA.GetStruct() == TypeB.GetStruct())
+	if (TypeInput.GetStruct() == TypeOutput.GetStruct())
 	{
 		return true;
 	}
 
 	bool bIsSupportedConversion = false;
-	if (IsScalarDefinition(TypeA) && IsScalarDefinition(TypeB))
+	if (IsScalarDefinition(TypeInput) && IsScalarDefinition(TypeOutput))
 	{
-		bIsSupportedConversion = (TypeA == IntDef && TypeB == FloatDef) || (TypeB == IntDef && TypeA == FloatDef) || (TypeA == IntDef && TypeB == BoolDef) || (TypeB == IntDef && TypeA == BoolDef);
+		bIsSupportedConversion = (TypeInput == IntDef && TypeOutput == FloatDef) || (TypeOutput == IntDef && TypeInput == FloatDef) || (TypeInput == IntDef && TypeOutput == BoolDef) || (TypeOutput == IntDef && TypeInput == BoolDef);
 	}
 	else
 	{
-		bIsSupportedConversion = (TypeA == ColorDef && TypeB == Vec4Def) || (TypeB == ColorDef && TypeA == Vec4Def);
+		bIsSupportedConversion = (TypeInput == ColorDef && TypeOutput == Vec4Def) || (TypeOutput == ColorDef && TypeInput == Vec4Def);
+	}
+
+	if (bIsSupportedConversion)
+	{
+		return true;
 	}
 
-if (bIsSupportedConversion)
-{
-	return true;
-}
 
-return	(TypeA == NumericDef && NumericStructs.Contains(TypeB.GetScriptStruct())) ||
-(TypeB == NumericDef && NumericStructs.Contains(TypeA.GetScriptStruct())) ||
-(TypeA == NumericDef && (TypeB.GetStruct() == GetIntStruct()) && TypeB.GetEnum() != nullptr) ||
-(TypeB == NumericDef && (TypeA.GetStruct() == GetIntStruct()) && TypeA.GetEnum() != nullptr);
+	return	(TypeInput == NumericDef && NumericStructs.Contains(TypeOutput.GetScriptStruct())) ||
+		(TypeOutput == NumericDef && NumericStructs.Contains(TypeInput.GetScriptStruct())) ||
+		(TypeInput == NumericDef && (TypeOutput.GetStruct() == GetIntStruct()) && TypeOutput.GetEnum() != nullptr) ||
+		(TypeOutput == NumericDef && (TypeInput.GetStruct() == GetIntStruct()) && TypeInput.GetEnum() != nullptr);
 }
 
 bool FNiagaraTypeDefinition::IsLossyConversion(const FNiagaraTypeDefinition& TypeA, const FNiagaraTypeDefinition& TypeB)
diff --git a/Engine/Plugins/FX/Niagara/Source/Niagara/Private/NiagaraScript.cpp b/Engine/Plugins/FX/Niagara/Source/Niagara/Private/NiagaraScript.cpp
--- a/Engine/Plugins/FX/Niagara/Source/Niagara/Private/NiagaraScript.cpp
+++ b/Engine/Plugins/FX/Niagara/Source/Niagara/Private/NiagaraScript.cpp
@@ -86,6 +86,14 @@ static FAutoConsoleVariableRef CVarNiagaraDelayScriptAsyncOptimization(
 	ECVF_Default
 );
 
+int32 UNiagaraScript::LogCompileStaticVars =0;
+static FAutoConsoleVariableRef CVarLogCompileIdGeneration(
+	TEXT("fx.LogCompileStaticVars"),
+	UNiagaraScript::LogCompileStaticVars,
+	TEXT("If > 0 all compile id generation dealing with static variables will be logged.  \n"),
+	ECVF_Default
+);
+
 FNiagaraScriptDebuggerInfo::FNiagaraScriptDebuggerInfo() : bWaitForGPU(false), FrameLastWriteId(-1), bWritten(false)
 {
 }
@@ -800,15 +808,25 @@ void UNiagaraScript::ComputeVMCompilationId(FNiagaraVMExecutableDataId& Id, FGui
 	Id.bRequiresPersistentIDs = false;
 	Id.ScriptVersionID = IsVersioningEnabled() ? (VersionGuid.IsValid() ? VersionGuid : ExposedVersion) : FGuid();
 	
+	TArray<FNiagaraVariable> ReferencedStaticVars;
+	TArray<const uint8*> ReferencedStaticVarValues;
 	ENiagaraSimTarget SimTargetToBuild = ENiagaraSimTarget::CPUSim;
 	// Ideally we wouldn't want to do this but rather than push the data down
 	// from the emitter.  Checking all outers here to pick up simulation stages too.
 	UNiagaraEmitter* OuterEmitter = GetTypedOuter<UNiagaraEmitter>();
+	int32 SystemSpawnIdx = INDEX_NONE;
+	int32 SystemUpdateIdx = INDEX_NONE;
+	TArray<UNiagaraScript*> Scripts;
+
 	if (OuterEmitter != nullptr)
 	{
 		UNiagaraEmitter* Emitter = OuterEmitter;
+
 		if (UNiagaraSystem* EmitterOwner = Cast<UNiagaraSystem>(Emitter->GetOuter()))
 		{
+			SystemSpawnIdx = Scripts.Add(EmitterOwner->GetSystemSpawnScript());
+			SystemUpdateIdx = Scripts.Add(EmitterOwner->GetSystemUpdateScript());
+
 			if (EmitterOwner->bBakeOutRapidIteration)
 			{
 				Id.bUsesRapidIterationParams = false;
@@ -916,6 +934,9 @@ void UNiagaraScript::ComputeVMCompilationId(FNiagaraVMExecutableDataId& Id, FGui
 			ComputeVMCompilationId_EmitterShared(Id, Emitter, EmitterOwner, ENiagaraRendererSourceDataMode::Particles);
 		}
 
+		// Gather up the scripts we will use to iterate for static rapid iteration vars later...
+		Emitter->GetScripts(Scripts, false, true);
+
 		if ((Emitter->bInterpolatedSpawning && Usage == ENiagaraScriptUsage::ParticleGPUComputeScript) ||
 			(Emitter->bInterpolatedSpawning && Usage == ENiagaraScriptUsage::ParticleSpawnScript) ||
 			Usage == ENiagaraScriptUsage::ParticleSpawnScriptInterpolated)
@@ -985,6 +1006,9 @@ void UNiagaraScript::ComputeVMCompilationId(FNiagaraVMExecutableDataId& Id, FGui
 	UObject* Obj = GetOuter();
 	if (UNiagaraSystem* System = Cast<UNiagaraSystem>(Obj))
 	{
+		SystemSpawnIdx = Scripts.Add(System->GetSystemSpawnScript());
+		SystemUpdateIdx = Scripts.Add(System->GetSystemUpdateScript());
+
 		if (System->bBakeOutRapidIteration)
 		{
 			Id.bUsesRapidIterationParams = false;
@@ -1030,33 +1054,106 @@ void UNiagaraScript::ComputeVMCompilationId(FNiagaraVMExecutableDataId& Id, FGui
 	}
 
 
+	for (int32 ScriptIdx = 0; ScriptIdx < Scripts.Num(); ScriptIdx++)
+	{
+		UNiagaraScript* Script = Scripts[ScriptIdx];
+		TArray<FNiagaraVariable> Vars;
+		Script->RapidIterationParameters.GetParameters(Vars);		
+		bool bIsThisEmitterUsage = UNiagaraScript::IsEmitterScript(Usage);
+		bool bIsThisParticleUsage = UNiagaraScript::IsParticleScript(Usage);
+	
+		bool bIsOtherSystemUsage = UNiagaraScript::IsSystemScript(Script->Usage);
+		bool bIsOtherParticleUsage = UNiagaraScript::IsParticleScript(Script->Usage);
+		bool bIsOtherEmitterUsage = UNiagaraScript::IsEmitterScript(Script->Usage);
+
+		// Emitter scripts don't depend on static variables from particle scripts.
+		if (bIsThisEmitterUsage && bIsOtherParticleUsage)
+			continue;
+
+		for (const FNiagaraVariable& Var : Vars)
+		{
+			if (Var.GetType().IsStatic())
+			{
+				if ((bIsThisEmitterUsage || bIsThisParticleUsage) && bIsOtherSystemUsage)
+				{
+					TArray<FString> SplitName;
+					Var.GetName().ToString().ParseIntoArray(SplitName, TEXT("."));
+
+					// Only include System based rapid iteration static variables if we're an emitter/particle script as system scripts can contain the amalgam of 
+					// it's own and any child emitter script rapid iteration variables.
+					if (SplitName.Num() == 3)
+					{
+						//Constants.Module.Variable input, this would match up for system inputs to modules that could produce new outputs that are static.
+					}
+					else if (SplitName.Num() > 3 && SplitName[SplitName.Num()-2] == TEXT("System"))
+					{
+						// Example: Constants.SetVariables_C872846946BD27D8655B60892D599F11.System.New bool
+					}
+					else
+					{
+						continue;
+					}
+				}
+				ReferencedStaticVars.Add(Var);
+				ReferencedStaticVarValues.Add(Script->RapidIterationParameters.GetParameterData(Var));
+			}
+		}
+	}
+
 	// If we aren't using rapid iteration parameters, we need to bake them into the hashstate for the compile id. This
 	// makes their values part of the lookup.
-	if (false == Id.bUsesRapidIterationParams)
+	
 	{
 		FSHA1 HashState;
-		TArray<FNiagaraVariable> Vars;
-		RapidIterationParameters.GetParameters(Vars);
-		//UE_LOG(LogNiagara, Display, TEXT("AreScriptAndSourceSynchronized %s ======================== "), *GetFullName());
-		for (int32 i = 0; i < Vars.Num(); i++)
+		if (false == Id.bUsesRapidIterationParams)
 		{
-			if (Vars[i].IsDataInterface() || Vars[i].IsUObject())
+			TArray<FNiagaraVariable> Vars;
+			RapidIterationParameters.GetParameters(Vars);
+			//UE_LOG(LogNiagara, Display, TEXT("AreScriptAndSourceSynchronized %s ======================== "), *GetFullName());
+			for (int32 i = 0; i < Vars.Num(); i++)
 			{
-				// Skip these types as they don't bake out, just normal parameters get baked.
+				if (Vars[i].IsDataInterface() || Vars[i].IsUObject())
+				{
+					// Skip these types as they don't bake out, just normal parameters get baked.
+				}
+				else
+				{
+					// Hash the name, type, and value of each parameter..
+					FString VarName = Vars[i].GetName().ToString();
+					FString VarTypeName = Vars[i].GetType().GetName();
+					HashState.UpdateWithString(*VarName, VarName.Len());
+					HashState.UpdateWithString(*VarTypeName, VarTypeName.Len());
+					const uint8* VarData = RapidIterationParameters.GetParameterData(Vars[i]);
+					if (VarData)
+					{
+						//UE_LOG(LogNiagara, Display, TEXT("Param %s %s %d"), *VarTypeName, *VarName, (uint32)VarData[0]);
+						HashState.Update(VarData, Vars[i].GetType().GetSize());
+					}
+				}
 			}
-			else
+		}
+
+		for (int32 i = 0; i < ReferencedStaticVars.Num(); i++)
+		{
+			if (i == 0 && UNiagaraScript::LogCompileStaticVars > 0)
 			{
-				// Hash the name, type, and value of each parameter..
-				FString VarName = Vars[i].GetName().ToString();
-				FString VarTypeName = Vars[i].GetType().GetName();
-				HashState.UpdateWithString(*VarName, VarName.Len());
-				HashState.UpdateWithString(*VarTypeName, VarTypeName.Len());
-				const uint8* VarData = RapidIterationParameters.GetParameterData(Vars[i]);
-				if (VarData)
+				UE_LOG(LogNiagara, Display, TEXT("***** Referenced Static Vars %d %s"), (int32)Usage, *GetPathName());
+			}
+
+			// Hash the name, type, and value of each parameter..
+			FString VarName = ReferencedStaticVars[i].GetName().ToString();
+			FString VarTypeName = ReferencedStaticVars[i].GetType().GetName();
+			HashState.UpdateWithString(*VarName, VarName.Len());
+			HashState.UpdateWithString(*VarTypeName, VarTypeName.Len());
+			const uint8* VarData = ReferencedStaticVarValues[i];
+
+			if (VarData)
+			{
+				if (UNiagaraScript::LogCompileStaticVars > 0)
 				{
-					//UE_LOG(LogNiagara, Display, TEXT("Param %s %s %s"), *VarTypeName, *VarName, *ByteStr);
-					HashState.Update(VarData, Vars[i].GetType().GetSize());
+					UE_LOG(LogNiagara, Display, TEXT("Param %s %s %d"), *VarTypeName, *VarName, (uint32)VarData[0]);
 				}
+				HashState.Update(VarData, ReferencedStaticVars[i].GetType().GetSize());
 			}
 		}
 		HashState.Final();
@@ -2924,6 +3021,7 @@ void UNiagaraScript::SyncAliases(const FNiagaraAliasContext& ResolveAliasesConte
 		RapidIterationParameters.GetParameters(Params);
 		for (FNiagaraVariable Var : Params)
 		{
+
 			FNiagaraVariable NewVar = FNiagaraUtilities::ResolveAliases(Var, ResolveAliasesContext);
 			if (NewVar.GetName() != Var.GetName())
 			{
diff --git a/Engine/Plugins/FX/Niagara/Source/Niagara/Private/NiagaraSystem.cpp b/Engine/Plugins/FX/Niagara/Source/Niagara/Private/NiagaraSystem.cpp
--- a/Engine/Plugins/FX/Niagara/Source/Niagara/Private/NiagaraSystem.cpp
+++ b/Engine/Plugins/FX/Niagara/Source/Niagara/Private/NiagaraSystem.cpp
@@ -666,6 +666,14 @@ void UNiagaraSystem::PostInitProperties()
 	ResolveScalabilitySettings();
 	UpdateDITickFlags();
 	UpdateHasGPUEmitters();
+
+#if WITH_EDITORONLY_DATA
+	if (SystemSpawnScript && SystemSpawnScript->GetLatestSource() != nullptr)
+	{
+		ensure(SystemSpawnScript->GetLatestSource() == SystemUpdateScript->GetLatestSource());
+		SystemSpawnScript->GetLatestSource()->OnChanged().AddUObject(this, &UNiagaraSystem::GraphSourceChanged);
+	}
+#endif
 }
 
 bool UNiagaraSystem::IsLooping() const
@@ -701,6 +709,15 @@ void UNiagaraSystem::UpdateSystemAfterLoad()
 	}
 	bFullyLoaded = true;
 
+#if WITH_EDITORONLY_DATA
+	if (SystemSpawnScript)
+	{
+		ensure(SystemSpawnScript->GetLatestSource() != nullptr);
+		ensure(SystemSpawnScript->GetLatestSource() == SystemUpdateScript->GetLatestSource());
+		SystemSpawnScript->GetLatestSource()->OnChanged().AddUObject(this, &UNiagaraSystem::GraphSourceChanged);
+	}
+#endif
+
 	for (FNiagaraEmitterHandle& EmitterHandle : EmitterHandles)
 	{
 		if (EmitterHandle.GetInstance() != nullptr)
@@ -890,6 +907,12 @@ void UNiagaraSystem::UpdateSystemAfterLoad()
 }
 
 #if WITH_EDITORONLY_DATA
+
+void UNiagaraSystem::GraphSourceChanged()
+{
+	InvalidateCachedData();
+}
+
 bool UNiagaraSystem::UsesScript(const UNiagaraScript* Script) const
 {
 	EnsureFullyLoaded();
@@ -899,7 +922,7 @@ bool UNiagaraSystem::UsesScript(const UNiagaraScript* Script) const
 		return true;
 	}
 
-	for (FNiagaraEmitterHandle EmitterHandle : GetEmitterHandles())
+	for (const FNiagaraEmitterHandle& EmitterHandle : GetEmitterHandles())
 	{
 		if (EmitterHandle.GetInstance() && EmitterHandle.GetInstance()->UsesScript(Script))
 		{
@@ -2642,7 +2665,12 @@ void UNiagaraSystem::EvaluateCompileResultDependencies() const
 				}
 
 				const FScriptCompileResultValidationInfo& TestInfo = ScriptCompilesToValidate[TestIndex];
-				if (TestInfo.CompileResults->AttributesWritten.Contains(TestVar))
+				if (TestVar.GetType().IsStatic() && TestInfo.CompileResults->StaticVariablesWritten.Contains(TestVar))
+				{
+					bDependencyMet = true;
+					break;
+				}
+				else if (TestInfo.CompileResults->AttributesWritten.Contains(TestVar))
 				{
 					bDependencyMet = true;
 					break;
@@ -3009,6 +3037,22 @@ void UNiagaraSystem::PrepareRapidIterationParametersForCompilation()
 	FNiagaraUtilities::PrepareRapidIterationParameters(Scripts, ScriptDependencyMap, ScriptToEmitterMap);
 }
 
+
+const TSharedPtr<FNiagaraGraphCachedDataBase, ESPMode::ThreadSafe>& UNiagaraSystem::GetCachedTraversalData() const
+{
+	if (CachedTraversalData.IsValid())
+		return CachedTraversalData;
+
+	INiagaraModule& NiagaraModule = FModuleManager::Get().LoadModuleChecked<INiagaraModule>(TEXT("Niagara"));
+	CachedTraversalData = NiagaraModule.CacheGraphTraversal(this, FGuid());
+	return CachedTraversalData;
+}
+
+void UNiagaraSystem::InvalidateCachedData()
+{
+	CachedTraversalData.Reset();
+}
+
 bool UNiagaraSystem::RequestCompile(bool bForce, FNiagaraSystemUpdateContext* OptionalUpdateContext)
 {
 	check(IsInGameThread());
@@ -3484,6 +3528,54 @@ void UNiagaraSystem::SetEffectType(UNiagaraEffectType* InEffectType)
 		UpdateCtx.Add(this, true);
 	}	
 }
+
+
+void UNiagaraSystem::GatherStaticVariables(TArray<FNiagaraVariable>& OutVars, TArray<FNiagaraVariable>& OutEmitterVars) const
+{
+	TArray<UNiagaraScript*> OutScripts;
+	OutScripts.Add(SystemSpawnScript);
+	OutScripts.Add(SystemUpdateScript);
+
+	auto GatherFromScripts = [&](TArray<UNiagaraScript*> Scripts, TArray<FNiagaraVariable>& Vars)
+	{
+		for (UNiagaraScript* Script : Scripts)
+		{
+			TArray<FNiagaraVariable> StoreParams;
+			Script->RapidIterationParameters.GetParameters(StoreParams);
+
+			for (int32 i = 0; i < StoreParams.Num(); i++)
+			{
+				if (StoreParams[i].GetType().IsStatic())
+				{
+					const int32* Index = Script->RapidIterationParameters.FindParameterOffset(StoreParams[i]);
+					if (Index != nullptr)
+					{
+						StoreParams[i].SetData(Script->RapidIterationParameters.GetParameterData(*Index)); // This will memcopy the data in.
+						Vars.AddUnique(StoreParams[i]);
+
+						//UE_LOG(LogNiagara, Log, TEXT("UNiagaraSystem::GatherStaticVariables Added %s"), *StoreParams[i].ToString());
+					}
+				}
+			}
+		}
+	};
+
+	GatherFromScripts(OutScripts, OutVars);
+
+	TArray<UNiagaraScript*> EmitterScripts; 
+	//const TArray<FNiagaraEmitterHandle>& EmitterHandles = GetEmitterHandles();
+	for (int32 i = 0; i < EmitterHandles.Num(); i++)
+	{
+		const FNiagaraEmitterHandle& Handle = EmitterHandles[i];
+		if (Handle.GetInstance() && Handle.GetIsEnabled())
+		{
+			EmitterScripts.Add(Handle.GetInstance()->EmitterSpawnScriptProps.Script);
+			EmitterScripts.Add(Handle.GetInstance()->EmitterUpdateScriptProps.Script);
+		}		
+	}
+
+	GatherFromScripts(EmitterScripts, OutEmitterVars);
+}
 #endif
 
 void UNiagaraSystem::ResolveScalabilitySettings()
diff --git a/Engine/Plugins/FX/Niagara/Source/Niagara/Private/NiagaraSystemInstance.cpp b/Engine/Plugins/FX/Niagara/Source/Niagara/Private/NiagaraSystemInstance.cpp
--- a/Engine/Plugins/FX/Niagara/Source/Niagara/Private/NiagaraSystemInstance.cpp
+++ b/Engine/Plugins/FX/Niagara/Source/Niagara/Private/NiagaraSystemInstance.cpp
@@ -2029,7 +2029,7 @@ bool FNiagaraSystemInstance::UsesScript(const UNiagaraScript* Script)const
 {
 	if (GetSystem())
 	{
-		for (FNiagaraEmitterHandle EmitterHandle : GetSystem()->GetEmitterHandles())
+		for (const FNiagaraEmitterHandle& EmitterHandle : GetSystem()->GetEmitterHandles())
 		{
 			if (EmitterHandle.GetInstance() && EmitterHandle.GetInstance()->UsesScript(Script))
 			{
diff --git a/Engine/Plugins/FX/Niagara/Source/Niagara/Public/NiagaraModule.h b/Engine/Plugins/FX/Niagara/Source/Niagara/Public/NiagaraModule.h
--- a/Engine/Plugins/FX/Niagara/Source/Niagara/Public/NiagaraModule.h
+++ b/Engine/Plugins/FX/Niagara/Source/Niagara/Public/NiagaraModule.h
@@ -10,6 +10,7 @@
 #include "Templates/SharedPointer.h"
 #include "NiagaraPerfBaseline.h"
 #include "NiagaraDebuggerCommon.h"
+#include "NiagaraScript.h"
 #include "Templates/PimplPtr.h"
 
 class FNiagaraWorldManager;
@@ -36,10 +37,13 @@ public:
 #if WITH_EDITOR
 	typedef TSharedPtr<FNiagaraCompileRequestDataBase, ESPMode::ThreadSafe> CompileRequestPtr;
 	typedef TSharedPtr<FNiagaraCompileRequestDuplicateDataBase, ESPMode::ThreadSafe> CompileRequestDuplicatePtr;
+	typedef TSharedPtr<FNiagaraGraphCachedDataBase, ESPMode::ThreadSafe> GraphCachedDataPtr;
+
 	DECLARE_DELEGATE_RetVal_ThreeParams(int32, FScriptCompiler, const FNiagaraCompileRequestDataBase*, const FNiagaraCompileRequestDuplicateDataBase*, const FNiagaraCompileOptions&);
 	DECLARE_DELEGATE_RetVal_TwoParams(TSharedPtr<FNiagaraVMExecutableData>, FCheckCompilationResult, int32, bool);
 	DECLARE_DELEGATE_RetVal_TwoParams(CompileRequestPtr, FOnPrecompile, UObject*, FGuid);
 	DECLARE_DELEGATE_RetVal_FiveParams(CompileRequestDuplicatePtr, FOnPrecompileDuplicate, const FNiagaraCompileRequestDataBase* /*OwningSystemRequestData*/, UNiagaraSystem* /*OwningSystem*/, UNiagaraEmitter* /*OwningEmitter*/, UNiagaraScript* /*TargetScript*/, FGuid /*Version*/);
+	DECLARE_DELEGATE_RetVal_TwoParams(GraphCachedDataPtr, FOnCacheGraphTraversal, const UObject*, FGuid);
 #endif
 	DECLARE_DELEGATE_RetVal(void, FOnProcessQueue);
 
@@ -96,6 +100,10 @@ public:
 	FDelegateHandle RegisterPrecompileDuplicator(FOnPrecompileDuplicate PreCompileDuplicator);
 	void UnregisterPrecompileDuplicator(FDelegateHandle DelegateHandle);
 
+	TSharedPtr<FNiagaraGraphCachedDataBase, ESPMode::ThreadSafe> CacheGraphTraversal(const UObject* InObj, FGuid Version);
+	FDelegateHandle RegisterGraphTraversalCacher(FOnCacheGraphTraversal PreCompiler);
+	void UnregisterGraphTraversalCacher(FDelegateHandle DelegateHandle);
+
 	void OnAssetLoaded(UObject* Asset);
 
 #endif
@@ -229,6 +237,7 @@ public:
 	FCheckCompilationResult CompilationResultDelegate;
 	FOnPrecompile PrecompileDelegate;
 	FOnPrecompileDuplicate PrecompileDuplicateDelegate;
+	FOnCacheGraphTraversal GraphTraversalCacheDelegate;
 #endif
 
 	static int32 EngineEffectsQuality;
diff --git a/Engine/Plugins/FX/Niagara/Source/Niagara/Public/NiagaraTypes.h b/Engine/Plugins/FX/Niagara/Source/Niagara/Public/NiagaraTypes.h
--- a/Engine/Plugins/FX/Niagara/Source/Niagara/Public/NiagaraTypes.h
+++ b/Engine/Plugins/FX/Niagara/Source/Niagara/Public/NiagaraTypes.h
@@ -90,6 +90,7 @@ struct FNiagaraPosition : public FVector3f
 #endif
 };
 
+
 USTRUCT(meta = (DisplayName = "Half", NiagaraInternalType = "true"))
 struct FNiagaraHalf
 {
@@ -786,11 +787,17 @@ struct NIAGARA_API FNiagaraTypeDefinition
 		UT_Enum
 	};
 
+	enum FTypeFlags
+	{
+		TF_None = 0,
+		TF_Static
+	};
+
 public:
 
 	// Construct blank raw type definition 
 	FORCEINLINE FNiagaraTypeDefinition(UClass *ClassDef)
-		: ClassStructOrEnum(ClassDef), UnderlyingType(UT_Class), Size(INDEX_NONE), Alignment(INDEX_NONE)
+		: ClassStructOrEnum(ClassDef), UnderlyingType(UT_Class), Flags(TF_None), Size(INDEX_NONE), Alignment(INDEX_NONE)
 #if WITH_EDITORONLY_DATA
 		, Struct_DEPRECATED(nullptr), Enum_DEPRECATED(nullptr)
 #endif
@@ -799,7 +806,7 @@ public:
 	}
 
 	FORCEINLINE FNiagaraTypeDefinition(UEnum *EnumDef)
-		: ClassStructOrEnum(EnumDef), UnderlyingType(UT_Enum), Size(INDEX_NONE), Alignment(INDEX_NONE)
+		: ClassStructOrEnum(EnumDef), UnderlyingType(UT_Enum), Flags(TF_None), Size(INDEX_NONE), Alignment(INDEX_NONE)
 #if WITH_EDITORONLY_DATA
 		, Struct_DEPRECATED(nullptr), Enum_DEPRECATED(nullptr)
 #endif
@@ -809,7 +816,7 @@ public:
 
 	enum class EAllowUnfriendlyStruct : uint8 { Deny, Allow };
 	FORCEINLINE FNiagaraTypeDefinition(UScriptStruct* StructDef, EAllowUnfriendlyStruct AllowUnfriendlyStruct)
-		: ClassStructOrEnum(StructDef), UnderlyingType(UT_Struct), Size(INDEX_NONE), Alignment(INDEX_NONE)
+		: ClassStructOrEnum(StructDef), UnderlyingType(UT_Struct), Flags(TF_None), Size(INDEX_NONE), Alignment(INDEX_NONE)
 #if WITH_EDITORONLY_DATA
 		, Struct_DEPRECATED(nullptr), Enum_DEPRECATED(nullptr)
 #endif
@@ -820,7 +827,7 @@ public:
 	FORCEINLINE FNiagaraTypeDefinition(UScriptStruct* StructDef) : FNiagaraTypeDefinition(StructDef, EAllowUnfriendlyStruct::Deny) {}
 
 	FORCEINLINE FNiagaraTypeDefinition(const FNiagaraTypeDefinition &Other)
-		: ClassStructOrEnum(Other.ClassStructOrEnum), UnderlyingType(Other.UnderlyingType), Size(INDEX_NONE), Alignment(INDEX_NONE)
+		: ClassStructOrEnum(Other.ClassStructOrEnum), UnderlyingType(Other.UnderlyingType), Flags(Other.Flags), Size(INDEX_NONE), Alignment(INDEX_NONE)
 #if WITH_EDITORONLY_DATA
 		, Struct_DEPRECATED(nullptr), Enum_DEPRECATED(nullptr)
 #endif
@@ -829,7 +836,7 @@ public:
 
 	// Construct a blank raw type definition
 	FORCEINLINE FNiagaraTypeDefinition()
-		: ClassStructOrEnum(nullptr), UnderlyingType(UT_None), Size(INDEX_NONE), Alignment(INDEX_NONE)
+		: ClassStructOrEnum(nullptr), UnderlyingType(UT_None), Flags(TF_None), Size(INDEX_NONE), Alignment(INDEX_NONE)
 #if WITH_EDITORONLY_DATA
 		, Struct_DEPRECATED(nullptr), Enum_DEPRECATED(nullptr)
 #endif
@@ -841,6 +848,11 @@ public:
 	}
 
 	FORCEINLINE bool operator == (const FNiagaraTypeDefinition &Other) const
+	{
+		return ClassStructOrEnum == Other.ClassStructOrEnum && UnderlyingType == Other.UnderlyingType && Flags == Other.Flags;
+	}
+
+	FORCEINLINE bool IsSameBaseDefinition(const FNiagaraTypeDefinition& Other) const
 	{
 		return ClassStructOrEnum == Other.ClassStructOrEnum && UnderlyingType == Other.UnderlyingType;
 	}
@@ -852,13 +864,32 @@ public:
 			return NSLOCTEXT("NiagaraTypeDefinition", "InvalidNameText", "Invalid (null type)");
 		}
 
+		FText NameText;
 #if WITH_EDITOR
 		if ( UnderlyingType != UT_Enum )
 		{
-			return GetStruct()->GetDisplayNameText();
+			NameText = GetStruct()->GetDisplayNameText();
 		}
+		else
 #endif
-		return FText::FromString(ClassStructOrEnum->GetName());
+		{
+			NameText = FText::FromString(ClassStructOrEnum->GetName());
+		}
+
+		if (Flags != 0)
+		{
+			if (IsStatic())
+			{
+				NameText = FText::Format(NSLOCTEXT("NiagaraTypeDefinition", "FlagFormatNameText", "Static {0}"), NameText);
+			}
+			else
+			{
+				NameText = FText::Format(NSLOCTEXT("NiagaraTypeDefinition", "UnknownNameText", "Unknown Flag Type {0}"), NameText);
+			}
+		}
+		
+		return NameText;
+		
 	}
 
 	FName GetFName() const
@@ -909,7 +940,21 @@ public:
 
 	bool IsEnum() const { return UnderlyingType == UT_Enum; }
 
+	bool IsStatic() const {
+		return (GetFlags() & TF_Static) != 0;
+	}
+
+	void SetFlags(FTypeFlags InFlags)
+	{
+		Flags = InFlags;
+	}
+
+	uint32 GetFlags()const { return Flags; }
+	
+
 	bool IsIndexWildcard() const { return ClassStructOrEnum == FNiagaraTypeDefinition::GetWildcardStruct(); }
+
+	FNiagaraTypeDefinition ToStaticDef() const;
 	
 	int32 GetSize() const
 	{
@@ -982,6 +1027,9 @@ public:
 	void PostSerialize(const FArchive& Ar);
 
 private:
+	UPROPERTY(EditAnywhere, Category = Type)
+	uint8 Flags;
+
 	mutable int16 Size;
 	mutable int16 Alignment;
 
@@ -1193,7 +1241,9 @@ const FNiagaraTypeDefinition& FNiagaraTypeDefinition::Get()
 
 FORCEINLINE uint32 GetTypeHash(const FNiagaraTypeDefinition& Type)
 {
-	return HashCombine(GetTypeHash(Type.GetStruct()), GetTypeHash(Type.GetEnum()));
+	uint32 Result= HashCombine(GetTypeHash(Type.GetStruct()), GetTypeHash(Type.GetEnum()));
+	Result = HashCombine(Result, GetTypeHash(Type.GetFlags()));
+	return Result;
 }
 
 //////////////////////////////////////////////////////////////////////////
@@ -1382,6 +1432,20 @@ public:
 		}
 	}
 
+	static bool IsStaticPossible(const FNiagaraTypeDefinition& InSrc) 
+	{
+		if (InSrc.IsStatic())
+			return true;
+
+		for (const FNiagaraTypeDefinition& TypeDef : Get().RegisteredTypes)
+		{
+			if (InSrc.IsSameBaseDefinition(TypeDef) && TypeDef.IsStatic())
+				return true;
+		}
+
+		return false;
+	}
+
 	static int32 RegisterIndexed(const FNiagaraTypeDefinition& NewType)
 	{
 		FNiagaraTypeRegistry& Registry = Get();
@@ -1448,6 +1512,11 @@ struct FNiagaraTypeDefinitionHandle
 		return RegisteredTypeIndex != Other.RegisteredTypeIndex;
 	}
 
+	bool IsSameBase(const FNiagaraTypeDefinition& Other) const
+	{
+		return Resolve().IsSameBaseDefinition(Other);
+	}
+
 	bool AppendCompileHash(FNiagaraCompileHashVisitor* InVisitor) const
 	{
 		return Resolve().AppendCompileHash(InVisitor);
@@ -1812,7 +1881,7 @@ struct TStructOpsTypeTraits<FNiagaraVariable> : public TStructOpsTypeTraitsBase2
 template<>
 inline bool FNiagaraVariable::GetValue<bool>() const
 {
-	check(*TypeDefHandle == FNiagaraTypeDefinition::GetBoolDef());
+	check(TypeDefHandle.IsSameBase(FNiagaraTypeDefinition::GetBoolDef()));
 	check(IsDataAllocated());
 	FNiagaraBool* BoolStruct = (FNiagaraBool*)GetData();
 	return BoolStruct->GetValue();
@@ -1821,7 +1890,7 @@ inline bool FNiagaraVariable::GetValue<bool>() const
 template<>
 inline void FNiagaraVariable::SetValue<bool>(const bool& Data)
 {
-	check(*TypeDefHandle == FNiagaraTypeDefinition::GetBoolDef());
+	check(TypeDefHandle.IsSameBase(FNiagaraTypeDefinition::GetBoolDef()));
 	AllocateData();
 	FNiagaraBool* BoolStruct = (FNiagaraBool*)GetData();
 	BoolStruct->SetValue(Data);
diff --git a/Engine/Plugins/FX/Niagara/Source/NiagaraCore/Private/NiagaraCustomVersion.cpp b/Engine/Plugins/FX/Niagara/Source/NiagaraCore/Private/NiagaraCustomVersion.cpp
--- a/Engine/Plugins/FX/Niagara/Source/NiagaraCore/Private/NiagaraCustomVersion.cpp
+++ b/Engine/Plugins/FX/Niagara/Source/NiagaraCore/Private/NiagaraCustomVersion.cpp
@@ -8,4 +8,4 @@ const FGuid FNiagaraCustomVersion::GUID(0xFCF57AFA, 0x50764283, 0xB9A9E658, 0xFF
 // Register the custom version with core
 FCustomVersionRegistration GRegisterNiagaraCustomVersion(FNiagaraCustomVersion::GUID, FNiagaraCustomVersion::LatestVersion, TEXT("NiagaraVer"));
 
-const FGuid FNiagaraCustomVersion::LatestScriptCompileVersion(0xAD773C8E, 0x9F7A1542, 0xAE4598B6, 0xA73F6125);
+const FGuid FNiagaraCustomVersion::LatestScriptCompileVersion(0xEF798E15, 0x4D684C92, 0xA6FF404A, 0x64ED9651);
diff --git a/Engine/Plugins/FX/Niagara/Source/NiagaraEditor/Private/Customizations/NiagaraStaticSwitchNodeDetails.cpp b/Engine/Plugins/FX/Niagara/Source/NiagaraEditor/Private/Customizations/NiagaraStaticSwitchNodeDetails.cpp
--- a/Engine/Plugins/FX/Niagara/Source/NiagaraEditor/Private/Customizations/NiagaraStaticSwitchNodeDetails.cpp
+++ b/Engine/Plugins/FX/Niagara/Source/NiagaraEditor/Private/Customizations/NiagaraStaticSwitchNodeDetails.cpp
@@ -43,7 +43,10 @@ FNiagaraStaticSwitchNodeDetails::FNiagaraStaticSwitchNodeDetails()
 
 void FNiagaraStaticSwitchNodeDetails::CustomizeDetails(IDetailLayoutBuilder& DetailBuilder)
 {
+	 SwitchTypeDataProperty = DetailBuilder.GetProperty(GET_MEMBER_NAME_CHECKED(UNiagaraNodeStaticSwitch, SwitchTypeData));
+
 	static const FName SwitchCategoryName = TEXT("Static Switch");
+	static const FName SwitchHiddenCategoryName = TEXT("HiddenMetaData");
 
 	TArray<TWeakObjectPtr<UObject>> ObjectsCustomized;
 	DetailBuilder.GetObjectsBeingCustomized(ObjectsCustomized);
@@ -67,12 +70,16 @@ void FNiagaraStaticSwitchNodeDetails::CustomizeDetails(IDetailLayoutBuilder& Det
 				DropdownOptions.Add(MakeShareable(new SwitchDropdownOption(TypeDefOnNode.GetEnum()->GetName(), TypeDefOnNode.GetEnum())));
 		}
 		UpdateSelectionFromNode();
-		
+	
+		DetailBuilder.EditCategory(SwitchHiddenCategoryName).SetCategoryVisibility(false);
+
 		IDetailCategoryBuilder& CategoryBuilder = DetailBuilder.EditCategory(SwitchCategoryName);		
 		FDetailWidgetRow& NameWidget = CategoryBuilder.AddCustomRow(LOCTEXT("NiagaraSwitchNodeNameFilterText", "Input parameter name"));
 		FDetailWidgetRow& DropdownWidget = CategoryBuilder.AddCustomRow(LOCTEXT("NiagaraSwitchNodeTypeFilterText", "Input parameter type"));
 		FDetailWidgetRow& DefaultValueOption = CategoryBuilder.AddCustomRow(LOCTEXT("NiagaraSwitchNodeDefaultFilterText", "Default value"));
+		FDetailWidgetRow& ExternalPinOption = CategoryBuilder.AddCustomRow(LOCTEXT("NiagaraSwitchNodeExteranPinText", "Create Pin"));
 		FDetailWidgetRow& ConstantSelection = CategoryBuilder.AddCustomRow(LOCTEXT("NiagaraSwitchNodeConstantFilterText", "Compiler constant"));
+		
 
 		NameWidget
 		.NameContent()
@@ -187,6 +194,28 @@ void FNiagaraStaticSwitchNodeDetails::CustomizeDetails(IDetailLayoutBuilder& Det
 			]
 		];
 
+		ExternalPinOption.NameContent()
+			[
+				SNew(SBox)
+				.Padding(FMargin(0.0f, 2.0f))
+				[
+					SNew(STextBlock)
+					.IsEnabled(this, &FNiagaraStaticSwitchNodeDetails::GetExposeAsPinEnabled)
+					.TextStyle(FNiagaraEditorStyle::Get(), "NiagaraEditor.ParameterText")
+					.Text(LOCTEXT("NiagaraSwitchNodeExposeAsPinText", "Expose As Pin"))
+				]
+			]
+		.ValueContent()
+			[
+				SNew(SBox)
+				.Padding(FMargin(0.0f, 2.0f))
+				[
+					SNew(SCheckBox)
+					.IsChecked_Lambda([this]() { return GetIsPinExposed() ? ECheckBoxState::Checked : ECheckBoxState::Unchecked; })
+					.OnCheckStateChanged(this, &FNiagaraStaticSwitchNodeDetails::ExposeAsPinCommitted)
+				]
+			];
+
 		ConstantSelection
 		.NameContent()
 		[
@@ -261,6 +290,30 @@ void FNiagaraStaticSwitchNodeDetails::DefaultBoolValueCommitted(ECheckBoxState N
 	}
 }
 
+
+void FNiagaraStaticSwitchNodeDetails::ExposeAsPinCommitted(ECheckBoxState NewState)
+{
+	if (Node.IsValid())
+	{
+		Node->Modify();
+		Node->SwitchTypeData.bExposeAsPin = (NewState == ECheckBoxState::Checked);
+
+		if (SwitchTypeDataProperty.IsValid())
+		{
+			SwitchTypeDataProperty->NotifyPostChange(EPropertyChangeType::ValueSet);
+		}
+	}
+}
+
+bool FNiagaraStaticSwitchNodeDetails::GetIsPinExposed() const
+{
+	if (Node.IsValid())
+	{
+		return Node->SwitchTypeData.bExposeAsPin;
+	}
+	return false;
+}
+
 TSharedRef<SWidget> FNiagaraStaticSwitchNodeDetails::CreateWidgetForDropdownOption(TSharedPtr<SwitchDropdownOption> InOption)
 {
 	return SNew(STextBlock).Text(FText::FromString(*InOption->Name));
@@ -305,6 +358,10 @@ void FNiagaraStaticSwitchNodeDetails::OnSelectionChanged(TSharedPtr<DefaultEnumO
 
 	ScriptVar->SetStaticSwitchDefaultValue(SelectedDefaultValue->EnumIndex);
 	SetSwitchParameterMetadata(ScriptVar->Metadata);
+	if (SwitchTypeDataProperty.IsValid())
+	{
+		SwitchTypeDataProperty->NotifyPostChange(EPropertyChangeType::ValueSet);
+	}
 }
 
 void FNiagaraStaticSwitchNodeDetails::OnSelectionChanged(TSharedPtr<SwitchDropdownOption> NewValue, ESelectInfo::Type)
@@ -336,6 +393,10 @@ void FNiagaraStaticSwitchNodeDetails::OnSelectionChanged(TSharedPtr<SwitchDropdo
 	}
 	RefreshDropdownValues();
 	Node->OnSwitchParameterTypeChanged(OldType);
+	if (SwitchTypeDataProperty.IsValid())
+	{
+		SwitchTypeDataProperty->NotifyPostChange(EPropertyChangeType::ValueSet);
+	}
 }
 
 void FNiagaraStaticSwitchNodeDetails::OnSelectionChanged(TSharedPtr<ConstantDropdownOption> NewValue, ESelectInfo::Type)
@@ -357,6 +418,10 @@ void FNiagaraStaticSwitchNodeDetails::OnSelectionChanged(TSharedPtr<ConstantDrop
 		Node->SwitchTypeData.Enum = nullptr;
 	}
 	Node->RefreshFromExternalChanges();
+	if (SwitchTypeDataProperty.IsValid())
+	{
+		SwitchTypeDataProperty->NotifyPostChange(EPropertyChangeType::ValueSet);
+	}
 }
 
 FText FNiagaraStaticSwitchNodeDetails::GetDropdownItemLabel() const
@@ -472,6 +537,11 @@ void FNiagaraStaticSwitchNodeDetails::SetSwitchParameterMetadata(const FNiagaraV
 	}
 	Node->GetNiagaraGraph()->SetMetaData(FNiagaraVariable(Node->GetInputType(), Node->InputParameterName), MetaData);
 	Node->GetNiagaraGraph()->NotifyGraphNeedsRecompile();
+
+	if (SwitchTypeDataProperty.IsValid())
+	{
+		SwitchTypeDataProperty->NotifyPostChange(EPropertyChangeType::ValueSet);
+	}
 }
 
 void FNiagaraStaticSwitchNodeDetails::UpdateSelectionFromNode()
@@ -514,6 +584,12 @@ bool FNiagaraStaticSwitchNodeDetails::IsConstantSelection() const
 	return SelectedDropdownItem == DropdownOptions[1] || SelectedDropdownItem == DropdownOptions[3] || SelectedDropdownItem == DropdownOptions[4];
 }
 
+bool FNiagaraStaticSwitchNodeDetails::GetExposeAsPinEnabled() const
+{
+	return !IsConstantSelection();
+
+}
+
 bool FNiagaraStaticSwitchNodeDetails::GetIntOptionEnabled() const
 {
 	return Node.IsValid() && Node->SwitchTypeData.SwitchType == ENiagaraStaticSwitchType::Integer && (!IsConstantSelection() || SelectedDropdownItem == DropdownOptions[3]);
@@ -582,7 +658,7 @@ bool FNiagaraStaticSwitchNodeDetails::OnVerifyParameterNameChanged(const FText&
 
 bool FNiagaraStaticSwitchNodeDetails::GetDefaultOptionEnabled() const
 {
-	return !IsConstantSelection();
+	return !IsConstantSelection() && !Node->SwitchTypeData.bExposeAsPin;
 }
 
 #undef LOCTEXT_NAMESPACE
diff --git a/Engine/Plugins/FX/Niagara/Source/NiagaraEditor/Private/Customizations/NiagaraStaticSwitchNodeDetails.h b/Engine/Plugins/FX/Niagara/Source/NiagaraEditor/Private/Customizations/NiagaraStaticSwitchNodeDetails.h
--- a/Engine/Plugins/FX/Niagara/Source/NiagaraEditor/Private/Customizations/NiagaraStaticSwitchNodeDetails.h
+++ b/Engine/Plugins/FX/Niagara/Source/NiagaraEditor/Private/Customizations/NiagaraStaticSwitchNodeDetails.h
@@ -104,6 +104,11 @@ private:
 	TSharedRef<SWidget> CreateWidgetForDropdownOption(TSharedPtr<ConstantDropdownOption> InOption);
 	FText GetConstantSelectionItemLabel() const;
 
+	bool GetExposeAsPinEnabled() const;
+	void ExposeAsPinCommitted(ECheckBoxState NewState);
+	bool GetIsPinExposed() const;
+	
+
 	TWeakObjectPtr<class UNiagaraNodeStaticSwitch> Node;
 	TArray<TSharedPtr<SwitchDropdownOption>> DropdownOptions;
 	TSharedPtr<SwitchDropdownOption> SelectedDropdownItem;
@@ -111,4 +116,5 @@ private:
 	TSharedPtr<DefaultEnumOption> SelectedDefaultValue;
 	TArray<TSharedPtr<ConstantDropdownOption>> ConstantOptions;
 	TSharedPtr<ConstantDropdownOption> SelectedConstantItem;
+	TSharedPtr<IPropertyHandle> SwitchTypeDataProperty;
 };
diff --git a/Engine/Plugins/FX/Niagara/Source/NiagaraEditor/Private/Customizations/NiagaraTypeCustomizations.cpp b/Engine/Plugins/FX/Niagara/Source/NiagaraEditor/Private/Customizations/NiagaraTypeCustomizations.cpp
--- a/Engine/Plugins/FX/Niagara/Source/NiagaraEditor/Private/Customizations/NiagaraTypeCustomizations.cpp
+++ b/Engine/Plugins/FX/Niagara/Source/NiagaraEditor/Private/Customizations/NiagaraTypeCustomizations.cpp
@@ -120,7 +120,7 @@ void FNiagaraMatrixCustomization::CustomizeChildren(TSharedRef<IPropertyHandle>
 	}
 }
 
-TArray<FNiagaraVariableBase> FNiagaraStackAssetAction_VarBind::FindVariables(UNiagaraEmitter* InEmitter, bool bSystem, bool bEmitter, bool bParticles, bool bUser)
+TArray<FNiagaraVariableBase> FNiagaraStackAssetAction_VarBind::FindVariables(UNiagaraEmitter* InEmitter, bool bSystem, bool bEmitter, bool bParticles, bool bUser, bool bAllowStatic)
 {
 	TArray<FNiagaraVariableBase> Bindings;
 	TArray<FNiagaraParameterMapHistory> Histories;
@@ -150,6 +150,9 @@ TArray<FNiagaraVariableBase> FNiagaraStackAssetAction_VarBind::FindVariables(UNi
 	{
 		for (const FNiagaraVariable& Var : History.Variables)
 		{
+			if (Var.GetType().IsStatic() && !bAllowStatic)
+				continue;
+
 			if (FNiagaraParameterMapHistory::IsAttribute(Var) && bParticles)
 			{
 				Bindings.AddUnique(Var);
@@ -265,13 +268,19 @@ TArray<FName> FNiagaraVariableAttributeBindingCustomization::GetNames(UNiagaraEm
 	bool bEmitter = true;
 	bool bParticles = true;
 	bool bUser = true;
-	TArray<FNiagaraVariableBase> Vars = FNiagaraStackAssetAction_VarBind::FindVariables(InEmitter, bSystem, bEmitter, bParticles, bUser);
+	bool bAllowStatic = true;
+	TArray<FNiagaraVariableBase> Vars = FNiagaraStackAssetAction_VarBind::FindVariables(InEmitter, bSystem, bEmitter, bParticles, bUser, bAllowStatic);
 	for (const FNiagaraVariableBase& Var : Vars)
 	{
-		if ( Var.GetType() != TargetVariableBinding->GetType() )
+		if (!bAllowStatic && Var.GetType() != TargetVariableBinding->GetType() )
+		{
+			continue;
+		}
+		else if (bAllowStatic && !Var.GetType().IsSameBaseDefinition(TargetVariableBinding->GetType()))
 		{
 			continue;
 		}
+
 		if ( RenderProps && RenderProps->IsSupportedVariableForBinding(Var, *PropertyHandle->GetProperty()->GetName()) )
 		{
 			Names.AddUnique(Var.GetName());
@@ -828,7 +837,8 @@ TArray<TPair<FNiagaraVariableBase, FNiagaraVariableBase> > FNiagaraMaterialAttri
 		bool bEmitter = true;
 		bool bParticles = false;
 		bool bUser = true;
-		BaseVars = FNiagaraStackAssetAction_VarBind::FindVariables(BaseEmitter, bSystem, bEmitter, bParticles, bUser);
+		bool bStatic = false;
+		BaseVars = FNiagaraStackAssetAction_VarBind::FindVariables(BaseEmitter, bSystem, bEmitter, bParticles, bUser, bStatic);
 
 		TArray<UNiagaraScript*> Scripts;
 		Scripts.Add(BaseSystem->GetSystemUpdateScript());
diff --git a/Engine/Plugins/FX/Niagara/Source/NiagaraEditor/Private/Customizations/NiagaraTypeCustomizations.h b/Engine/Plugins/FX/Niagara/Source/NiagaraEditor/Private/Customizations/NiagaraTypeCustomizations.h
--- a/Engine/Plugins/FX/Niagara/Source/NiagaraEditor/Private/Customizations/NiagaraTypeCustomizations.h
+++ b/Engine/Plugins/FX/Niagara/Source/NiagaraEditor/Private/Customizations/NiagaraTypeCustomizations.h
@@ -151,7 +151,7 @@ struct FNiagaraStackAssetAction_VarBind : public FEdGraphSchemaAction
 	}
 	//~ End FEdGraphSchemaAction Interface
 
-	static TArray<FNiagaraVariableBase> FindVariables(UNiagaraEmitter* InEmitter, bool bSystem, bool bEmitter, bool bParticles, bool bUser);
+	static TArray<FNiagaraVariableBase> FindVariables(UNiagaraEmitter* InEmitter, bool bSystem, bool bEmitter, bool bParticles, bool bUser, bool bAllowStatic);
 };
 
 class FNiagaraVariableAttributeBindingCustomization : public IPropertyTypeCustomization
diff --git a/Engine/Plugins/FX/Niagara/Source/NiagaraEditor/Private/EdGraphSchema_Niagara.cpp b/Engine/Plugins/FX/Niagara/Source/NiagaraEditor/Private/EdGraphSchema_Niagara.cpp
--- a/Engine/Plugins/FX/Niagara/Source/NiagaraEditor/Private/EdGraphSchema_Niagara.cpp
+++ b/Engine/Plugins/FX/Niagara/Source/NiagaraEditor/Private/EdGraphSchema_Niagara.cpp
@@ -61,6 +61,11 @@ const FName UEdGraphSchema_Niagara::PinCategoryMisc("Misc");
 const FName UEdGraphSchema_Niagara::PinCategoryClass("Class");
 const FName UEdGraphSchema_Niagara::PinCategoryEnum("Enum");
 
+
+const FName UEdGraphSchema_Niagara::PinCategoryStaticType("StaticType");
+const FName UEdGraphSchema_Niagara::PinCategoryStaticClass("StaticClass");
+const FName UEdGraphSchema_Niagara::PinCategoryStaticEnum("StaticEnum");
+
 namespace NiagaraNodeNumbers
 {
 	// Maximum distance a drag can be off a node edge to require 'push off' from node
@@ -622,11 +627,12 @@ TArray<TSharedPtr<FNiagaraAction_NewNode>> UEdGraphSchema_Niagara::GetGraphActio
 						for (TFieldIterator<FProperty> CandidatePropertyIt(CandidateType.GetStruct(), EFieldIteratorFlags::IncludeSuper); CandidatePropertyIt; ++CandidatePropertyIt)
 						{
 							FNiagaraTypeDefinition CandidatePropertyType = GetTypeDefForProperty(*CandidatePropertyIt);
-							if(FNiagaraTypeDefinition::TypesAreAssignable(PinType, CandidatePropertyType))
+							if(FNiagaraTypeDefinition::TypesAreAssignable(CandidatePropertyType, PinType))
 							{
 								bAddMake = true;
 								break;
 						}
+						}
 					}
 				}
 					else
@@ -634,7 +640,7 @@ TArray<TSharedPtr<FNiagaraAction_NewNode>> UEdGraphSchema_Niagara::GetGraphActio
 						// if our "from pin" is an input pin, we generally only allow that exact type for making, with some exceptions like NiagaraGenerics
 						if(PinType == CandidateType)
 					{
-							bAddMake = true;
+								bAddMake = true;
 						}
 
 						if(PinType == FNiagaraTypeDefinition::GetGenericNumericDef() && FNiagaraTypeDefinition::IsValidNumericInput(CandidateType))
@@ -959,10 +965,14 @@ const FPinConnectionResponse UEdGraphSchema_Niagara::CanCreateConnection(const U
 	auto GetPinsAreInvalidAddPinCombination = [](const UEdGraphPin* A, const UEdGraphPin* B)->bool {
 		if (A->PinType.PinSubCategory == UNiagaraNodeWithDynamicPins::AddPinSubCategory)
 		{
-			if (B->PinType.PinCategory != PinCategoryType)
+			if (B->PinType.PinCategory == PinCategoryMisc)
 			{
 				return true;
 			}
+			else
+			{
+				return false;
+			}
 		}
 		return false;
 	};
@@ -982,28 +992,28 @@ const FPinConnectionResponse UEdGraphSchema_Niagara::CanCreateConnection(const U
 	if (!IsPinWildcard(PinA) && !IsPinWildcard(PinB))
 	{
 		// Check for compatible type pins.
-		if (PinA->PinType.PinCategory == PinCategoryType &&
-			PinB->PinType.PinCategory == PinCategoryType &&
+		if ((PinA->PinType.PinCategory == PinCategoryType || PinA->PinType.PinCategory == PinCategoryStaticType) &&
+			(PinB->PinType.PinCategory == PinCategoryType || PinB->PinType.PinCategory == PinCategoryStaticType) &&
 			PinA->PinType != PinB->PinType)
 		{
-			FNiagaraTypeDefinition PinTypeA = PinToTypeDefinition(PinA);
-			FNiagaraTypeDefinition PinTypeB = PinToTypeDefinition(PinB);
+			FNiagaraTypeDefinition PinTypeInput = PinToTypeDefinition(InputPin);
+			FNiagaraTypeDefinition PinTypeOutput = PinToTypeDefinition(OutputPin);
 
-			if (PinTypeA == FNiagaraTypeDefinition::GetParameterMapDef() || PinTypeB == FNiagaraTypeDefinition::GetParameterMapDef())
+			if (PinTypeInput == FNiagaraTypeDefinition::GetParameterMapDef() || PinTypeOutput == FNiagaraTypeDefinition::GetParameterMapDef())
 			{
 				return FPinConnectionResponse(CONNECT_RESPONSE_DISALLOW, TEXT("Types are not compatible"));
 			}
 
-			else if (FNiagaraTypeDefinition::TypesAreAssignable(PinTypeA, PinTypeB) == false)
+			else if (FNiagaraTypeDefinition::TypesAreAssignable(PinTypeInput, PinTypeOutput) == false)
 			{
 				//Do some limiting on auto conversions here?
-				if (PinTypeA.GetClass())
+				if (PinTypeInput.GetClass() || PinTypeInput.IsStatic() || PinTypeOutput.IsStatic())
 				{
 					return FPinConnectionResponse(CONNECT_RESPONSE_DISALLOW, TEXT("Types are not compatible"));
 				}
 				else
 				{
-					return FPinConnectionResponse(CONNECT_RESPONSE_MAKE_WITH_CONVERSION_NODE, FString::Printf(TEXT("Convert %s to %s"), *(PinToTypeDefinition(PinA).GetNameText().ToString()), *(PinToTypeDefinition(PinB).GetNameText().ToString())));
+					return FPinConnectionResponse(CONNECT_RESPONSE_MAKE_WITH_CONVERSION_NODE, FString::Printf(TEXT("Convert %s to %s"), *(PinToTypeDefinition(InputPin).GetNameText().ToString()), *(PinToTypeDefinition(OutputPin).GetNameText().ToString())));
 				}
 			}
 		}
@@ -1039,26 +1049,33 @@ const FPinConnectionResponse UEdGraphSchema_Niagara::CanCreateConnection(const U
 				return FPinConnectionResponse(CONNECT_RESPONSE_DISALLOW, TEXT("Types are not compatible"));
 			}
 		}
-
-		if (PinA->PinType.PinCategory == PinCategoryClass || PinB->PinType.PinCategory == PinCategoryClass)
+		else
 		{
-			FNiagaraTypeDefinition AType = PinToTypeDefinition(PinA);
-			FNiagaraTypeDefinition BType = PinToTypeDefinition(PinB);
-			if (AType != BType)
+
+			if (PinA->PinType.PinCategory == PinCategoryClass || PinB->PinType.PinCategory == PinCategoryClass ||
+				PinA->PinType.PinCategory == PinCategoryStaticClass || PinB->PinType.PinCategory == PinCategoryStaticClass)
 			{
-				return FPinConnectionResponse(CONNECT_RESPONSE_DISALLOW, TEXT("Types are not compatible"));
+				FNiagaraTypeDefinition AType = PinToTypeDefinition(PinA);
+				FNiagaraTypeDefinition BType = PinToTypeDefinition(PinB);
+
+				if (AType != BType)
+				{
+					return FPinConnectionResponse(CONNECT_RESPONSE_DISALLOW, TEXT("Types are not compatible"));
+				}
 			}
-		}
 
-		if (PinA->PinType.PinCategory == PinCategoryEnum || PinB->PinType.PinCategory == PinCategoryEnum)
-		{
-			FNiagaraTypeDefinition PinTypeA = PinToTypeDefinition(PinA);
-			FNiagaraTypeDefinition PinTypeB = PinToTypeDefinition(PinB);
-			if (FNiagaraTypeDefinition::TypesAreAssignable(PinTypeA, PinTypeB) == false)
+			if (PinA->PinType.PinCategory == PinCategoryEnum || PinB->PinType.PinCategory == PinCategoryEnum ||
+				PinA->PinType.PinCategory == PinCategoryStaticEnum || PinB->PinType.PinCategory == PinCategoryStaticEnum)
 			{
-				return FPinConnectionResponse(CONNECT_RESPONSE_DISALLOW, TEXT("Types are not compatible"));
+				FNiagaraTypeDefinition PinTypeInput = PinToTypeDefinition(InputPin);
+				FNiagaraTypeDefinition PinTypeOutput = PinToTypeDefinition(OutputPin);
+				if (FNiagaraTypeDefinition::TypesAreAssignable(PinTypeInput, PinTypeOutput) == false)
+				{
+					return FPinConnectionResponse(CONNECT_RESPONSE_DISALLOW, TEXT("Types are not compatible"));
+				}
 			}
 		}
+
 	}
 	// at least one pin is a wildcard
 	else
@@ -1320,7 +1337,7 @@ bool UEdGraphSchema_Niagara::TryCreateConnection(UEdGraphPin* PinA, UEdGraphPin*
 
 FLinearColor UEdGraphSchema_Niagara::GetPinTypeColor(const FEdGraphPinType& PinType) const
 {
-	if (PinType.PinCategory == PinCategoryType)
+	if (PinType.PinCategory == PinCategoryType || PinType.PinCategory == PinCategoryStaticType)
 	{
 		FNiagaraTypeDefinition Type(CastChecked<UScriptStruct>(PinType.PinSubCategoryObject.Get()));
 		return GetTypeColor(Type);
@@ -1337,11 +1354,11 @@ FLinearColor UEdGraphSchema_Niagara::GetTypeColor(const FNiagaraTypeDefinition&
 	{
 		return Settings->FloatPinTypeColor;
 	}
-	else if (Type == FNiagaraTypeDefinition::GetIntDef())
+	else if (Type.IsSameBaseDefinition(FNiagaraTypeDefinition::GetIntDef()))
 	{
 		return Settings->IntPinTypeColor;
 	}
-	else if (Type == FNiagaraTypeDefinition::GetBoolDef())
+	else if (Type.IsSameBaseDefinition(FNiagaraTypeDefinition::GetBoolDef()))
 	{
 		return Settings->BooleanPinTypeColor;
 	}
@@ -1458,7 +1475,7 @@ FNiagaraTypeDefinition UEdGraphSchema_Niagara::PinToTypeDefinition(const UEdGrap
 		return FNiagaraTypeDefinition();
 	}
 	UEdGraphNode* OwningNode = Pin->GetOwningNodeUnchecked();
-	if (Pin->PinType.PinCategory == PinCategoryType && Pin->PinType.PinSubCategoryObject.IsValid())
+	if ((Pin->PinType.PinCategory == PinCategoryType || Pin->PinType.PinCategory == PinCategoryStaticType) && Pin->PinType.PinSubCategoryObject.IsValid())
 	{
 		UScriptStruct* Struct = Cast<UScriptStruct>(Pin->PinType.PinSubCategoryObject.Get());
 		if (Struct == nullptr)
@@ -1481,9 +1498,13 @@ FNiagaraTypeDefinition UEdGraphSchema_Niagara::PinToTypeDefinition(const UEdGrap
 			}
 			return FNiagaraTypeDefinition(FNiagaraTypeHelper::FindNiagaraFriendlyTopLevelStruct(Struct));
 		}
-		return FNiagaraTypeDefinition(Struct);
+
+		if (Pin->PinType.PinCategory == PinCategoryType)
+			return FNiagaraTypeDefinition(Struct);
+		else if (Pin->PinType.PinCategory == PinCategoryStaticType)
+			return FNiagaraTypeDefinition(Struct).ToStaticDef();
 	}
-	else if (Pin->PinType.PinCategory == PinCategoryClass)
+	else if (Pin->PinType.PinCategory == PinCategoryClass || Pin->PinType.PinCategory == PinCategoryStaticClass)
 	{
 		UClass* Class = Cast<UClass>(Pin->PinType.PinSubCategoryObject.Get());
 		if (Class == nullptr)
@@ -1492,9 +1513,13 @@ FNiagaraTypeDefinition UEdGraphSchema_Niagara::PinToTypeDefinition(const UEdGrap
 				*Pin->PinName.ToString(), OwningNode ? *OwningNode->GetFullName() : TEXT("Invalid"));
 			return FNiagaraTypeDefinition();
 		}
-		return FNiagaraTypeDefinition(Class);
+
+		if (Pin->PinType.PinCategory == PinCategoryClass)	
+			return FNiagaraTypeDefinition(Class);
+		else if (Pin->PinType.PinCategory == PinCategoryStaticClass)
+			return FNiagaraTypeDefinition(Class).ToStaticDef();
 	}
-	else if (Pin->PinType.PinCategory == PinCategoryEnum)
+	else if (Pin->PinType.PinCategory == PinCategoryEnum || Pin->PinType.PinCategory == PinCategoryStaticEnum)
 	{
 		UEnum* Enum = Cast<UEnum>(Pin->PinType.PinSubCategoryObject.Get());
 		if (Enum == nullptr)
@@ -1503,14 +1528,18 @@ FNiagaraTypeDefinition UEdGraphSchema_Niagara::PinToTypeDefinition(const UEdGrap
 				OwningNode ? *OwningNode->GetFullName() : TEXT("Invalid"));
 			return FNiagaraTypeDefinition(FNiagaraTypeDefinition::GetIntDef());
 		}
-		return FNiagaraTypeDefinition(Enum);
+
+		if (Pin->PinType.PinCategory == PinCategoryEnum)
+			return FNiagaraTypeDefinition(Enum);
+		else if (Pin->PinType.PinCategory == PinCategoryStaticEnum)
+			return FNiagaraTypeDefinition(Enum).ToStaticDef();
 	}
 	return FNiagaraTypeDefinition();
 }
 
 FNiagaraTypeDefinition UEdGraphSchema_Niagara::PinTypeToTypeDefinition(const FEdGraphPinType& PinType)
 {
-	if (PinType.PinCategory == PinCategoryType && PinType.PinSubCategoryObject.IsValid())
+	if ((PinType.PinCategory == PinCategoryType || PinType.PinCategory == PinCategoryStaticType) && PinType.PinSubCategoryObject.IsValid())
 	{
 		UScriptStruct* Struct = Cast<UScriptStruct>(PinType.PinSubCategoryObject.Get());
 		if (Struct == nullptr)
@@ -1518,9 +1547,13 @@ FNiagaraTypeDefinition UEdGraphSchema_Niagara::PinTypeToTypeDefinition(const FEd
 			UE_LOG(LogNiagaraEditor, Error, TEXT("Pin states that it is of struct type, but is missing its struct object. This is usually the result of a registered type going away."));
 			return FNiagaraTypeDefinition();
 		}
-		return FNiagaraTypeDefinition(Struct);
+
+		if (PinType.PinCategory == PinCategoryType)
+			return FNiagaraTypeDefinition(Struct);
+		else if (PinType.PinCategory == PinCategoryStaticType)
+			return FNiagaraTypeDefinition(Struct).ToStaticDef();
 	}
-	else if (PinType.PinCategory == PinCategoryClass)
+	else if (PinType.PinCategory == PinCategoryClass || PinType.PinCategory == PinCategoryStaticClass)
 	{
 		UClass* Class = Cast<UClass>(PinType.PinSubCategoryObject.Get());
 		if (Class == nullptr)
@@ -1528,9 +1561,13 @@ FNiagaraTypeDefinition UEdGraphSchema_Niagara::PinTypeToTypeDefinition(const FEd
 			UE_LOG(LogNiagaraEditor, Warning, TEXT("Pin states that it is of class type, but is missing its class object. This is usually the result of a registered type going away."));
 			return FNiagaraTypeDefinition();
 		}
-		return FNiagaraTypeDefinition(Class);
+
+		if (PinType.PinCategory == PinCategoryClass)
+			return FNiagaraTypeDefinition(Class);
+		else if (PinType.PinCategory == PinCategoryStaticClass)
+			return FNiagaraTypeDefinition(Class).ToStaticDef();
 	}
-	else if (PinType.PinCategory == PinCategoryEnum)
+	else if (PinType.PinCategory == PinCategoryEnum || PinType.PinCategory == PinCategoryStaticEnum)
 	{
 		UEnum* Enum = Cast<UEnum>(PinType.PinSubCategoryObject.Get());
 		if (Enum == nullptr)
@@ -1538,8 +1575,13 @@ FNiagaraTypeDefinition UEdGraphSchema_Niagara::PinTypeToTypeDefinition(const FEd
 			UE_LOG(LogNiagaraEditor, Warning, TEXT("Pin states that it is of Enum type, but is missing its Enum! Turning into standard int definition!"));
 			return FNiagaraTypeDefinition(FNiagaraTypeDefinition::GetIntDef());
 		}
-		return FNiagaraTypeDefinition(Enum);
+
+		if (PinType.PinCategory == PinCategoryEnum)
+			return FNiagaraTypeDefinition(Enum);
+		else if (PinType.PinCategory == PinCategoryEnum)
+			return FNiagaraTypeDefinition(Enum).ToStaticDef();
 	}
+	
 	return FNiagaraTypeDefinition();
 }
 
@@ -1547,16 +1589,26 @@ FEdGraphPinType UEdGraphSchema_Niagara::TypeDefinitionToPinType(FNiagaraTypeDefi
 {
 	if (TypeDef.GetClass())
 	{
-		return FEdGraphPinType(PinCategoryClass, NAME_None, const_cast<UClass*>(TypeDef.GetClass()), EPinContainerType::None, false, FEdGraphTerminalType());
+		if (TypeDef.IsStatic())
+			return FEdGraphPinType(PinCategoryStaticClass, NAME_None, const_cast<UClass*>(TypeDef.GetClass()), EPinContainerType::None, false, FEdGraphTerminalType());
+		else
+			return FEdGraphPinType(PinCategoryClass, NAME_None, const_cast<UClass*>(TypeDef.GetClass()), EPinContainerType::None, false, FEdGraphTerminalType());
 	}
 	else if (TypeDef.GetEnum())
 	{
-		return FEdGraphPinType(PinCategoryEnum, NAME_None, const_cast<UEnum*>(TypeDef.GetEnum()), EPinContainerType::None, false, FEdGraphTerminalType());
-	}
+		if (TypeDef.IsStatic())
+			return FEdGraphPinType(PinCategoryStaticEnum, NAME_None, const_cast<UEnum*>(TypeDef.GetEnum()), EPinContainerType::None, false, FEdGraphTerminalType());
+		else
+			return FEdGraphPinType(PinCategoryEnum, NAME_None, const_cast<UEnum*>(TypeDef.GetEnum()), EPinContainerType::None, false, FEdGraphTerminalType());
+	}	
 	else
 	{
 		//TODO: Are base types better as structs or done like BPS as a special name?
-		return FEdGraphPinType(PinCategoryType, NAME_None, const_cast<UScriptStruct*>(TypeDef.GetScriptStruct()), EPinContainerType::None, false, FEdGraphTerminalType());
+		if (TypeDef.IsStatic())
+			return FEdGraphPinType(PinCategoryStaticType, NAME_None, const_cast<UScriptStruct*>(TypeDef.GetScriptStruct()), EPinContainerType::None, false, FEdGraphTerminalType());
+		else
+			return FEdGraphPinType(PinCategoryType, NAME_None, const_cast<UScriptStruct*>(TypeDef.GetScriptStruct()), EPinContainerType::None, false, FEdGraphTerminalType());
+
 	}
 }
 
@@ -1753,6 +1805,7 @@ void UEdGraphSchema_Niagara::ConvertNumericPinToTypeAll(UNiagaraNode* InNode, FN
 {
 	if (InNode)
 	{
+		FScopedTransaction AllTransaction(NSLOCTEXT("UnrealEd", "NiagaraEditorChangeNumericPinTypeAll", "Change Pin Type For All"));
 		for (auto Pin : InNode->Pins)
 		{
 			if (PinToTypeDefinition(Pin) == FNiagaraTypeDefinition::GetGenericNumericDef())
@@ -1796,6 +1849,37 @@ void UEdGraphSchema_Niagara::ConvertNumericPinToType(UEdGraphPin* InGraphPin, FN
 	}
 }
 
+bool UEdGraphSchema_Niagara::PinTypesValidForNumericConversion(FEdGraphPinType AType, FEdGraphPinType BType) const
+{
+	if (AType == BType)
+		return true;
+	else if (AType.PinCategory == PinCategoryStaticType && BType.PinCategory == PinCategoryType)
+		return true;
+	else if (BType.PinCategory == PinCategoryStaticType && AType.PinCategory == PinCategoryType)
+		return true;
+	else if (AType.PinCategory == PinCategoryType && BType.PinCategory == PinCategoryType)
+		return true;
+	else if (BType.PinCategory == PinCategoryType && AType.PinCategory == PinCategoryType)
+		return true;
+
+
+	return false;
+}
+
+bool UEdGraphSchema_Niagara::IsStaticPin(const UEdGraphPin* Pin)
+{
+	if (Pin)
+	{
+		if (Pin->PinType.PinCategory == PinCategoryStaticType)
+			return true;
+		if (Pin->PinType.PinCategory == PinCategoryStaticClass)
+			return true;
+		if (Pin->PinType.PinCategory == PinCategoryStaticEnum)
+			return true;
+	}
+	return false;
+}
+
 bool UEdGraphSchema_Niagara::CheckCircularConnection(TSet<const UEdGraphNode*>& VisitedNodes, const UEdGraphNode* InNode, const UEdGraphNode* InTestNode)
 {
 	bool AlreadyAdded = false;
diff --git a/Engine/Plugins/FX/Niagara/Source/NiagaraEditor/Private/NiagaraActions.cpp b/Engine/Plugins/FX/Niagara/Source/NiagaraEditor/Private/NiagaraActions.cpp
--- a/Engine/Plugins/FX/Niagara/Source/NiagaraEditor/Private/NiagaraActions.cpp
+++ b/Engine/Plugins/FX/Niagara/Source/NiagaraEditor/Private/NiagaraActions.cpp
@@ -584,11 +584,11 @@ void FNiagaraParameterGraphDragOperation::MakeStaticSwitch(FNiagaraParameterNode
 	SwitchNode->InputParameterName = InParams.Parameter.GetName();
 	const FNiagaraTypeDefinition& Type = InParams.Parameter.GetType();
 	
-	if (Type == FNiagaraTypeDefinition::GetBoolDef())
+	if (Type.IsSameBaseDefinition(FNiagaraTypeDefinition::GetBoolDef()))
 	{
 		SwitchNode->SwitchTypeData.SwitchType = ENiagaraStaticSwitchType::Bool;
 	}
-	else if (Type == FNiagaraTypeDefinition::GetIntDef())
+	else if (Type.IsSameBaseDefinition(FNiagaraTypeDefinition::GetIntDef()))
 	{
 		SwitchNode->SwitchTypeData.SwitchType = ENiagaraStaticSwitchType::Integer;
 	}
diff --git a/Engine/Plugins/FX/Niagara/Source/NiagaraEditor/Private/NiagaraCompiler.cpp b/Engine/Plugins/FX/Niagara/Source/NiagaraEditor/Private/NiagaraCompiler.cpp
--- a/Engine/Plugins/FX/Niagara/Source/NiagaraEditor/Private/NiagaraCompiler.cpp
+++ b/Engine/Plugins/FX/Niagara/Source/NiagaraEditor/Private/NiagaraCompiler.cpp
@@ -30,6 +30,8 @@
 #include "NiagaraTrace.h"
 #include "Serialization/MemoryReader.h"
 #include "../../Niagara/Private/NiagaraPrecompileContainer.h"
+#include "Modules/ModuleManager.h"
+#include "NiagaraEditorModule.h"
 
 #define LOCTEXT_NAMESPACE "NiagaraCompiler"
 
@@ -68,6 +70,14 @@ static FAutoConsoleVariableRef CVarForceNiagaraVMBinaryDump(
 	ECVF_Default
 );
 
+static int32 GbForceNiagaraCacheDump = 0;
+static FAutoConsoleVariableRef CVarForceNiagaraCacheDump(
+	TEXT("fx.ForceNiagaraCacheDump"),
+	GbForceNiagaraCacheDump,
+	TEXT("If > 0 all cached graph traversal data will be dumped \n"),
+	ECVF_Default
+);
+
 static int32 GNiagaraEnablePrecompilerNamespaceFixup = 0;
 static FAutoConsoleVariableRef CVarNiagaraEnablePrecompilerNamespaceFixup(
 	TEXT("fx.NiagaraEnablePrecompilerNamespaceFixup"),
@@ -442,8 +452,49 @@ void FNiagaraCompileRequestDuplicateData::ReleaseCompilationCopies()
 	SourceDeepCopy = nullptr;
 }
 
-void FNiagaraCompileRequestData::FinishPrecompile(const TArray<FNiagaraVariable>& EncounterableVariables, FCompileConstantResolver ConstantResolver, const TArray<ENiagaraScriptUsage>& UsagesToProcess, const TArray<class UNiagaraSimulationStageBase*>* SimStages, const TArray<FString> EmitterNames)
+void FNiagaraCompileRequestData::CompareAgainst(FNiagaraGraphCachedBuiltHistory* InCachedDataBase)
+{
+	if (InCachedDataBase)
+	{
+		bool bDumpVars = false;
+		if (StaticVariables.Num() != InCachedDataBase->StaticVariables.Num())
+		{
+			UE_LOG(LogNiagaraEditor, Warning, TEXT("FNiagaraCompileRequestData::CompareAgainst> StaticVariables.Num() != InCachedDataBase->StaticVariables.Num()"));
+			bDumpVars = true;
+		}
+		
+		for (const FNiagaraVariable& Var : StaticVariables)
+		{
+			bool bFound = false;
+			for (const FNiagaraVariable& OtherVar : InCachedDataBase->StaticVariables)
+			{
+				if (OtherVar == Var)
+				{
+					bFound = true;
+					break;
+				}
+			}
+
+			if (!bFound)
+			{
+				UE_LOG(LogNiagaraEditor, Warning, TEXT("FNiagaraCompileRequestData::CompareAgainst> Could not find %s"), *Var.ToString());
+				bDumpVars = true;
+			}
+		}
+
+		if (bDumpVars)
+		{
+			for (const FNiagaraVariable& Var : InCachedDataBase->StaticVariables)
+			{
+				UE_LOG(LogNiagaraEditor, Warning, TEXT("%s"), *Var.ToString());
+			}
+		}
+	}
+}
+
+void FNiagaraCompileRequestData::FinishPrecompile(const TArray<FNiagaraVariable>& EncounterableVariables, const TArray<FNiagaraVariable>& InStaticVariables, FCompileConstantResolver ConstantResolver, const TArray<ENiagaraScriptUsage>& UsagesToProcess, const TArray<class UNiagaraSimulationStageBase*>* SimStages, const TArray<FString> EmitterNames)
 {
+	FNiagaraEditorModule& NiagaraEditorModule = FModuleManager::GetModuleChecked<FNiagaraEditorModule>("NiagaraEditor");
 	{
 		ENiagaraScriptCompileStatusEnum = StaticEnum<ENiagaraScriptCompileStatus>();
 		ENiagaraScriptUsageEnum = StaticEnum<ENiagaraScriptUsage>();
@@ -454,6 +505,18 @@ void FNiagaraCompileRequestData::FinishPrecompile(const TArray<FNiagaraVariable>
 			Source->NodeGraph->FindOutputNodes(OutputNodes);
 		}
 
+		bool bFilterByEmitterAlias = true;
+		for (UNiagaraNodeOutput* FoundOutputNode : OutputNodes)
+		{
+			if (UNiagaraScript::IsSystemScript(FoundOutputNode->GetUsage()))
+			{
+				bFilterByEmitterAlias = false;
+			}
+		}
+		
+		// Only use the static variables that match up with our expectations for this script. IE for emitters, filter things out for resolution.
+		FNiagaraParameterUtilities::FilterToRelevantStaticVariables(InStaticVariables, StaticVariables, *GetUniqueEmitterName(), TEXT("Emitter"), bFilterByEmitterAlias);
+
 		int32 NumSimStageNodes = 0;
 		for (UNiagaraNodeOutput* FoundOutputNode : OutputNodes)
 		{
@@ -489,8 +552,9 @@ void FNiagaraCompileRequestData::FinishPrecompile(const TArray<FNiagaraVariable>
 				Builder.ConstantResolver = ConstantResolver;
 				Builder.AddGraphToCallingGraphContextStack(Source->NodeGraph);
 				Builder.RegisterEncounterableVariables(EncounterableVariables);
+				Builder.RegisterExternalStaticVariables(StaticVariables);
 
-				FString TranslationName = TEXT("Emitter");
+				FString TranslationName = TEXT("Emitter");// Note that this cannot be GetUniqueEmitterName() as it would break downstream logic for some reason for data interfaces.
 				Builder.BeginTranslation(TranslationName);
 				Builder.BeginUsage(FoundOutputNode->GetUsage(), SimStageName);
 				Builder.EnableScriptAllowList(true, FoundOutputNode->GetUsage());
@@ -502,14 +566,24 @@ void FNiagaraCompileRequestData::FinishPrecompile(const TArray<FNiagaraVariable>
 				for (FNiagaraParameterMapHistory& History : Builder.Histories)
 				{
 					History.OriginatingScriptUsage = FoundOutputNode->GetUsage();
-					for (FNiagaraVariable& Var : History.Variables)
+					for (int32 i = 0; i < History.Variables.Num(); i++)
 					{
+						FNiagaraVariable& Var = History.Variables[i];
 						EncounteredVariables.AddUnique(Var);
 						if (Var.GetType() == FNiagaraTypeDefinition::GetGenericNumericDef())
 						{
 							UE_LOG(LogNiagaraEditor, Log, TEXT("Invalid numeric parameter found! %s"), *Var.GetName().ToString())
 						}
+
+						if (Var.GetType().IsStatic())
+						{
+							if (History.PerVariableConstantValue[i].Num() > 1)
+							{
+								StaticVariablesWithMultipleWrites.AddUnique(Var);
+							}
+						}
 					}
+					PinToConstantValues.Append(History.PinToConstantValues);
 				}
 
 				if (FoundOutputNode->GetUsage() == ENiagaraScriptUsage::ParticleSimulationStageScript)
@@ -517,6 +591,21 @@ void FNiagaraCompileRequestData::FinishPrecompile(const TArray<FNiagaraVariable>
 					NumSimStageNodes++;
 				}
 
+				if (UNiagaraScript::LogCompileStaticVars > 0)
+				{
+					UE_LOG(LogNiagaraEditor, Log, TEXT("Builder.StaticVariables After Param Map Traversal............................"));
+				}
+
+				for (const FNiagaraVariable& Var : Builder.StaticVariables)
+				{
+					StaticVariables.AddUnique(Var);
+
+					if (UNiagaraScript::LogCompileStaticVars > 0)
+					{
+						UE_LOG(LogNiagaraEditor, Log, TEXT("%s"), *Var.ToString());
+					}
+				}
+
 				Builder.EndTranslation(TranslationName);
 
 				// Collect data interface information.
@@ -607,10 +696,13 @@ void FNiagaraCompileRequestData::FinishPrecompile(const TArray<FNiagaraVariable>
 	}
 }
 
-void FNiagaraCompileRequestDuplicateData::FinishPrecompileDuplicate(const TArray<FNiagaraVariable>& EncounterableVariables, FCompileConstantResolver ConstantResolver, const TArray<class UNiagaraSimulationStageBase*>* SimStages)
+void FNiagaraCompileRequestDuplicateData::FinishPrecompileDuplicate(const TArray<FNiagaraVariable>& EncounterableVariables, const TArray<FNiagaraVariable>& InStaticVariables, FCompileConstantResolver ConstantResolver, const TArray<class UNiagaraSimulationStageBase*>* SimStages, const TArray<FNiagaraVariable>& InParamStore)
 {
+
+	FNiagaraEditorModule& NiagaraEditorModule = FModuleManager::GetModuleChecked<FNiagaraEditorModule>("NiagaraEditor");
 	PrecompiledHistories.Empty();
 
+
 	TArray<UNiagaraNodeOutput*> OutputNodes;
 	if (NodeGraphDeepCopy.IsValid())
 	{
@@ -646,6 +738,7 @@ void FNiagaraCompileRequestDuplicateData::FinishPrecompileDuplicate(const TArray
 			Builder.ConstantResolver = ConstantResolver;
 			Builder.AddGraphToCallingGraphContextStack(NodeGraphDeepCopy.Get());
 			Builder.RegisterEncounterableVariables(EncounterableVariables);
+			Builder.RegisterExternalStaticVariables(InStaticVariables);
 
 			FString TranslationName = TEXT("Emitter");
 			Builder.BeginTranslation(TranslationName);
@@ -659,8 +752,9 @@ void FNiagaraCompileRequestDuplicateData::FinishPrecompileDuplicate(const TArray
 			for (FNiagaraParameterMapHistory& History : Builder.Histories)
 			{
 				History.OriginatingScriptUsage = FoundOutputNode->GetUsage();
-				for (FNiagaraVariable& Var : History.Variables)
+				for (int32 VarIdx = 0; VarIdx < History.Variables.Num(); VarIdx++)
 				{
+					const FNiagaraVariable& Var = History.Variables[VarIdx];
 					if (Var.GetType() == FNiagaraTypeDefinition::GetGenericNumericDef())
 					{
 						UE_LOG(LogNiagaraEditor, Log, TEXT("Invalid numeric parameter found! %s"), *Var.GetName().ToString())
@@ -754,10 +848,15 @@ TSharedPtr<FNiagaraCompileRequestDataBase, ESPMode::ThreadSafe> FNiagaraEditorMo
 		const TArray<FNiagaraVariable> EncounterableVariables;
 		const TArray<ENiagaraScriptUsage> ValidUsages = { ENiagaraScriptUsage::Function, ENiagaraScriptUsage::Module, ENiagaraScriptUsage::DynamicInput };
 		const TArray<FString> EmitterNames;
-		BasePtr->FinishPrecompile(EncounterableVariables, EmptyResolver, ValidUsages, nullptr, EmitterNames);
+		TArray<FNiagaraVariable> StaticVars;
+		BasePtr->FinishPrecompile(EncounterableVariables, StaticVars, EmptyResolver, ValidUsages, nullptr, EmitterNames);
+		
 	}
 	else if (System)
 	{
+
+		TSharedPtr<FNiagaraGraphCachedDataBase, ESPMode::ThreadSafe> CachedTraversalSystemData = System->GetCachedTraversalData();
+
 		check(System->GetSystemSpawnScript()->GetLatestSource() == System->GetSystemUpdateScript()->GetLatestSource());
 		BasePtr->Source = Cast<UNiagaraScriptSource>(System->GetSystemSpawnScript()->GetLatestSource());
 		BasePtr->bUseRapidIterationParams = !System->bBakeOutRapidIteration;
@@ -790,11 +889,15 @@ TSharedPtr<FNiagaraCompileRequestDataBase, ESPMode::ThreadSafe> FNiagaraEditorMo
 		}
 
 		// Now deep copy the system graphs, skipping traversal into any emitter references.
+		TArray<FNiagaraVariable> StaticVariablesFromSystem = ((FNiagaraGraphCachedBuiltHistory*)CachedTraversalSystemData.Get())->StaticVariables;
 		{
 			FCompileConstantResolver ConstantResolver(System, ENiagaraScriptUsage::SystemSpawnScript);
 			UNiagaraScriptSource* Source = Cast<UNiagaraScriptSource>(System->GetSystemSpawnScript()->GetLatestSource());
 			static TArray<ENiagaraScriptUsage> SystemUsages = { ENiagaraScriptUsage::SystemSpawnScript, ENiagaraScriptUsage::SystemUpdateScript };
-			BasePtr->FinishPrecompile(EncounterableVars, ConstantResolver, SystemUsages, nullptr, EmitterNames);
+			
+			BasePtr->FinishPrecompile(EncounterableVars, StaticVariablesFromSystem, ConstantResolver, SystemUsages, nullptr, EmitterNames);
+
+			BasePtr->CompareAgainst(((FNiagaraGraphCachedBuiltHistory*)CachedTraversalSystemData.Get()));
 		}
 
 		// Add the User and System variables that we did encounter to the list that emitters might also encounter.
@@ -813,6 +916,10 @@ TSharedPtr<FNiagaraCompileRequestDataBase, ESPMode::ThreadSafe> FNiagaraEditorMo
 			FCompileConstantResolver ConstantResolver(Handle.GetInstance(), ENiagaraScriptUsage::EmitterSpawnScript);
 			if (Handle.GetIsEnabled()) // Don't pull in the emitter if it isn't going to be used.
 			{
+				TSharedPtr<FNiagaraGraphCachedDataBase, ESPMode::ThreadSafe> CachedTraversalEmitterData = Handle.GetInstance()->GetCachedTraversalData();
+				TArray<FNiagaraVariable> StaticVariablesFromEmitter = StaticVariablesFromSystem;
+				StaticVariablesFromEmitter.Append(((FNiagaraGraphCachedBuiltHistory*)CachedTraversalEmitterData.Get())->StaticVariables);
+
 				TArray<UNiagaraScript*> EmitterScripts;
 				Handle.GetInstance()->GetScripts(EmitterScripts, false, true);
 
@@ -824,9 +931,11 @@ TSharedPtr<FNiagaraCompileRequestDataBase, ESPMode::ThreadSafe> FNiagaraEditorMo
 						ScriptToEmitterMap.Add(EmitterScripts[ScriptIdx], Handle.GetInstance());
 					}
 				}
-				// First finish the precompile for the emitters so that the encountered emitter variables can be found.
 				static TArray<ENiagaraScriptUsage> EmitterUsages = { ENiagaraScriptUsage::EmitterSpawnScript, ENiagaraScriptUsage::EmitterUpdateScript };
-				BasePtr->EmitterData[i]->FinishPrecompile(EncounterableVars, ConstantResolver, EmitterUsages, nullptr, EmitterNames);
+
+				BasePtr->EmitterData[i]->FinishPrecompile(EncounterableVars, StaticVariablesFromEmitter, ConstantResolver, EmitterUsages, nullptr, EmitterNames);
+				BasePtr->EmitterData[i]->CompareAgainst(((FNiagaraGraphCachedBuiltHistory*)CachedTraversalEmitterData.Get()));
+				
 
 				// Then finish the precompile for the particle scripts once we've gathered the emitter vars which might be referenced.
 				TArray<FNiagaraVariable> ParticleEncounterableVars = EncounterableVars;
@@ -838,7 +947,9 @@ TSharedPtr<FNiagaraCompileRequestDataBase, ESPMode::ThreadSafe> FNiagaraEditorMo
 					ENiagaraScriptUsage::ParticleEventScript,
 					ENiagaraScriptUsage::ParticleGPUComputeScript,
 					ENiagaraScriptUsage::ParticleSimulationStageScript };
-				BasePtr->EmitterData[i]->FinishPrecompile(ParticleEncounterableVars, ConstantResolver, ParticleUsages, &Handle.GetInstance()->GetSimulationStages(), EmitterNames);
+
+				TArray<FNiagaraVariable> OldStaticVars = BasePtr->EmitterData[i]->StaticVariables;
+				BasePtr->EmitterData[i]->FinishPrecompile(ParticleEncounterableVars, OldStaticVars, ConstantResolver, ParticleUsages, &Handle.GetInstance()->GetSimulationStages(), EmitterNames);
 			}
 		}
 
@@ -1047,11 +1158,12 @@ TSharedPtr<FNiagaraCompileRequestDuplicateDataBase, ESPMode::ThreadSafe> FNiagar
 		BasePtr->DeepCopyGraphs(Source, TargetScript->GetUsage(), EmptyResolver);
 		TArray<FNiagaraVariable> EncounterableScriptVariables;
 		OwningSystemRequestData->GatherPreCompiledVariables(FString(), EncounterableScriptVariables);
-		BasePtr->FinishPrecompileDuplicate(EncounterableScriptVariables, EmptyResolver, nullptr);
+		BasePtr->FinishPrecompileDuplicate(EncounterableScriptVariables, TArray<FNiagaraVariable>(), EmptyResolver, nullptr, ((FNiagaraCompileRequestData*)OwningSystemRequestData)->RapidIterationParams);
 		CollectDataInterfaceClasses(EncounterableScriptVariables);
 	}
 	else
 	{
+		TSharedPtr<FNiagaraGraphCachedDataBase, ESPMode::ThreadSafe> CachedTraversalSystemData = OwningSystem->GetCachedTraversalData();
 		GetUsagesToDuplicate(TargetScript->GetUsage(), BasePtr->ValidUsages);
 
 		check(OwningSystem->GetSystemSpawnScript()->GetLatestSource() == OwningSystem->GetSystemUpdateScript()->GetLatestSource());
@@ -1077,6 +1189,7 @@ TSharedPtr<FNiagaraCompileRequestDuplicateDataBase, ESPMode::ThreadSafe> FNiagar
 		}
 
 		// Now deep copy the system graphs, skipping traversal into any emitter references.
+		TArray<FNiagaraVariable> StaticVariablesFromSystem = ((FNiagaraGraphCachedBuiltHistory*)CachedTraversalSystemData.Get())->StaticVariables;
 		{
 			UNiagaraScriptSource* Source = Cast<UNiagaraScriptSource>(OwningSystem->GetSystemSpawnScript()->GetLatestSource());
 
@@ -1088,7 +1201,7 @@ TSharedPtr<FNiagaraCompileRequestDuplicateDataBase, ESPMode::ThreadSafe> FNiagar
 			{
 				FCompileConstantResolver ConstantResolver(OwningSystem, ENiagaraScriptUsage::SystemSpawnScript);
 				BasePtr->DeepCopyGraphs(Source, ENiagaraScriptUsage::SystemSpawnScript, ConstantResolver);
-				BasePtr->FinishPrecompileDuplicate(EncounterableSystemVariables, ConstantResolver, nullptr);
+				BasePtr->FinishPrecompileDuplicate(EncounterableSystemVariables, StaticVariablesFromSystem, ConstantResolver, nullptr, ((FNiagaraCompileRequestData*)OwningSystemRequestData)->RapidIterationParams);
 			}
 
 			CollectDataInterfaceClasses(EncounterableSystemVariables);
@@ -1099,13 +1212,19 @@ TSharedPtr<FNiagaraCompileRequestDuplicateDataBase, ESPMode::ThreadSafe> FNiagar
 		{
 			const FNiagaraEmitterHandle& Handle = OwningSystem->GetEmitterHandle(i);
 
+
 			TArray<FNiagaraVariable> EncounterableEmitterVariables;
 			OwningSystemRequestData->GetDependentRequest(i)->GatherPreCompiledVariables(FString(), EncounterableEmitterVariables);
 
 			if (Handle.GetIsEnabled() && (OwningEmitter == nullptr || OwningEmitter == Handle.GetInstance()))
 			{
+				TSharedPtr<FNiagaraGraphCachedDataBase, ESPMode::ThreadSafe> CachedTraversalEmitterData = Handle.GetInstance()->GetCachedTraversalData();
+				TArray<FNiagaraVariable> StaticVariablesFromEmitter = StaticVariablesFromSystem;
+				StaticVariablesFromEmitter.Append(((FNiagaraGraphCachedBuiltHistory*)CachedTraversalEmitterData.Get())->StaticVariables);
+
 				FCompileConstantResolver ConstantResolver(Handle.GetInstance(), ENiagaraScriptUsage::EmitterSpawnScript);
-				BasePtr->EmitterData[i]->FinishPrecompileDuplicate(EncounterableEmitterVariables, ConstantResolver, &Handle.GetInstance()->GetSimulationStages());
+				BasePtr->EmitterData[i]->FinishPrecompileDuplicate(EncounterableEmitterVariables, StaticVariablesFromEmitter, ConstantResolver, &Handle.GetInstance()->GetSimulationStages(), ((FNiagaraCompileRequestData*)OwningSystemRequestData)->EmitterData[i]->RapidIterationParams);
+
 			}
 
 			CollectDataInterfaceClasses(EncounterableEmitterVariables);
@@ -1178,6 +1297,145 @@ int32 FNiagaraEditorModule::CompileScript(const FNiagaraCompileRequestDataBase*
 	
 }
 
+TSharedPtr<FNiagaraGraphCachedDataBase, ESPMode::ThreadSafe> FNiagaraEditorModule::CacheGraphTraversal(const UObject* Obj, FGuid Version)
+{
+	TSharedPtr<FNiagaraGraphCachedBuiltHistory, ESPMode::ThreadSafe> CachedPtr = MakeShared<FNiagaraGraphCachedBuiltHistory>();
+
+	const UNiagaraSystem* System = Cast<UNiagaraSystem>(Obj);
+	const UNiagaraEmitter* Emitter = Cast<UNiagaraEmitter>(Obj);
+
+	/*
+	* const TArray<FNiagaraVariable>& EncounterableVariables, const TArray<FNiagaraVariable>& InStaticVariables, FCompileConstantResolver ConstantResolver, const TArray<ENiagaraScriptUsage>& UsagesToProcess, const 
+	*/
+	const TArray<class UNiagaraSimulationStageBase*>* SimStages = nullptr;
+	const UNiagaraScriptSource* ScriptSource = nullptr;
+	FCompileConstantResolver ConstantResolver;
+	TArray<FNiagaraVariable> EncounterableVariables;
+	TArray<FNiagaraVariable> StaticVariablesFromSystem;
+	TArray<FNiagaraVariable> StaticVariablesFromSystemAndEmitters;
+	TArray<FNiagaraVariable> SrcStaticVariables;
+	TSharedPtr<FNiagaraGraphCachedDataBase, ESPMode::ThreadSafe> ParentCachedData;
+	FString SrcUniqueEmitterName;
+
+	if (System)
+	{
+		ScriptSource = CastChecked<UNiagaraScriptSource>(System->GetSystemSpawnScript()->GetLatestSource());
+		ConstantResolver = FCompileConstantResolver(System, ENiagaraScriptUsage::SystemSpawnScript);
+		System->GatherStaticVariables(StaticVariablesFromSystem, StaticVariablesFromSystemAndEmitters);
+
+		SrcStaticVariables = StaticVariablesFromSystem;
+
+		for (const FNiagaraVariable& Var : StaticVariablesFromSystemAndEmitters)
+			SrcStaticVariables.AddUnique(Var);
+	}
+	else if (Emitter)
+	{
+		UNiagaraSystem* SysParent = Cast<UNiagaraSystem>(Emitter->GetOuter());
+		if (SysParent)
+		{
+			ParentCachedData = SysParent->GetCachedTraversalData();
+
+			if (ParentCachedData.IsValid())
+			{
+				SrcStaticVariables = ((FNiagaraGraphCachedBuiltHistory*)ParentCachedData.Get())->StaticVariables;
+			}
+			/*for (const FNiagaraEmitterHandle& Handle : SysParent->GetEmitterHandles())
+			{
+				if (Handle.GetInstance() == Emitter)
+				{
+					SrcUniqueEmitterName = Handle.GetUniqueEmitterName().ToString();
+				}
+			}*/
+		}
+		SrcUniqueEmitterName = Emitter->GetUniqueEmitterName();
+		Emitter->GatherStaticVariables(SrcStaticVariables);
+
+		ScriptSource = CastChecked<UNiagaraScriptSource>(Emitter->GraphSource);
+		SimStages = &Emitter->GetSimulationStages();
+		ConstantResolver = FCompileConstantResolver(Emitter, ENiagaraScriptUsage::EmitterSpawnScript);
+	}
+
+
+	TArray<UNiagaraNodeOutput*> OutputNodes;
+	if (ScriptSource != nullptr && ScriptSource->NodeGraph != nullptr)
+	{
+		ScriptSource->NodeGraph->FindOutputNodes(OutputNodes);
+	}
+
+	bool bFilterByEmitterAlias = true;
+	for (UNiagaraNodeOutput* FoundOutputNode : OutputNodes)
+	{
+		if (UNiagaraScript::IsSystemScript(FoundOutputNode->GetUsage()))
+		{
+			bFilterByEmitterAlias = false;
+		}
+	}
+
+	// Only use the static variables that match up with our expectations for this script. IE for emitters, filter things out for resolution.
+	TArray<FNiagaraVariable> StaticVariables;
+	FNiagaraParameterUtilities::FilterToRelevantStaticVariables(SrcStaticVariables, StaticVariables, *SrcUniqueEmitterName, TEXT("Emitter"), bFilterByEmitterAlias);
+
+	int32 NumSimStageNodes = 0;
+	for (UNiagaraNodeOutput* FoundOutputNode : OutputNodes)
+	{
+		FName SimStageName;
+		bool bStageEnabled = true;
+		if (FoundOutputNode->GetUsage() == ENiagaraScriptUsage::ParticleSimulationStageScript && SimStages)
+		{
+			// Find the simulation stage for this output node.
+			const FGuid& UsageId = FoundOutputNode->GetUsageId();
+			UNiagaraSimulationStageBase* const* MatchingStagePtr = SimStages->FindByPredicate([UsageId](UNiagaraSimulationStageBase* SimStage)
+				{
+					return SimStage != nullptr && SimStage->Script != nullptr && SimStage->Script->GetUsageId() == UsageId;
+				});
+
+			// Set whether or not the stage is enabled, and get the iteration source name if available.
+			bStageEnabled = MatchingStagePtr != nullptr && (*MatchingStagePtr)->bEnabled;
+			if (bStageEnabled && (*MatchingStagePtr)->IsA<UNiagaraSimulationStageGeneric>())
+			{
+				UNiagaraSimulationStageGeneric* GenericStage = CastChecked<UNiagaraSimulationStageGeneric>(*MatchingStagePtr);
+				SimStageName = GenericStage->IterationSource == ENiagaraIterationSource::DataInterface ? GenericStage->DataInterface.BoundVariable.GetName() : FName();
+			}
+		}
+
+		if (bStageEnabled)
+		{
+
+			// Map all for this output node
+			FNiagaraParameterMapHistoryWithMetaDataBuilder Builder;
+			Builder.ConstantResolver = ConstantResolver;
+			Builder.AddGraphToCallingGraphContextStack(ScriptSource->NodeGraph);
+			Builder.RegisterEncounterableVariables(EncounterableVariables);
+			Builder.RegisterExternalStaticVariables(StaticVariables);
+
+			FString TranslationName = TEXT("Emitter");
+			Builder.BeginTranslation(TranslationName);
+			Builder.BeginUsage(FoundOutputNode->GetUsage(), SimStageName);
+			Builder.EnableScriptAllowList(true, FoundOutputNode->GetUsage());
+			Builder.BuildParameterMaps(FoundOutputNode, true);
+			Builder.EndUsage();
+
+			for (const FNiagaraVariable& Var : Builder.StaticVariables)
+				StaticVariables.AddUnique(Var);
+		}
+	}
+
+
+	CachedPtr->StaticVariables = StaticVariables;
+	if (GbForceNiagaraCacheDump != 0 && Obj )
+	{
+		UE_LOG(LogNiagaraEditor, Log, TEXT("==================================================================\nCacheGraphTraversal %s\n=================================================================="), *Obj->GetPathName());
+		int32 i = 0;
+		UE_LOG(LogNiagaraEditor, Log, TEXT("Static Variables: %d"), StaticVariables.Num());
+		for (const FNiagaraVariable& Var : StaticVariables)
+		{
+			UE_LOG(LogNiagaraEditor, Log, TEXT("[%d] %s"), i, *Var.ToString());
+			++i;
+		}
+	}
+	return CachedPtr;
+}
+
 TSharedPtr<FNiagaraVMExecutableData> FNiagaraEditorModule::GetCompilationResult(int32 JobID, bool bWait)
 {
 	TSharedPtr<FHlslNiagaraCompiler>* MapEntry = ActiveCompilations.Find(JobID);
diff --git a/Engine/Plugins/FX/Niagara/Source/NiagaraEditor/Private/NiagaraEditorCommon.cpp b/Engine/Plugins/FX/Niagara/Source/NiagaraEditor/Private/NiagaraEditorCommon.cpp
--- a/Engine/Plugins/FX/Niagara/Source/NiagaraEditor/Private/NiagaraEditorCommon.cpp
+++ b/Engine/Plugins/FX/Niagara/Source/NiagaraEditor/Private/NiagaraEditorCommon.cpp
@@ -131,6 +131,22 @@ void FNiagaraOpInfo::Init()
 		Op->bSupportsAddedInputs = true;
 		Op->AddedInputTypeRestrictions.Add(Type);
 		Op->AddedInputFormatting = TEXT("{A} + {B}");
+		Op->bSupportsStaticResolution = true;
+		Op->StaticVariableResolveFunction.BindLambda([=](const TArray<int32>& InPinValues)
+			{
+				if (InPinValues.Num() == 0)
+					return 0;
+				if (InPinValues.Num() == 1)
+					return InPinValues[0];
+
+				int32 Return = InPinValues[0];
+				for (int32 i = 1; i < InPinValues.Num(); i++)
+				{
+					Return = Return + InPinValues[i];
+				}
+
+				return Return;
+			});
 		OpInfoMap.Add(Op->Name) = Idx;
 
 		Idx = OpInfos.AddDefaulted();
@@ -146,6 +162,22 @@ void FNiagaraOpInfo::Init()
 		Op->bSupportsAddedInputs = true;
 		Op->AddedInputTypeRestrictions.Add(Type);
 		Op->AddedInputFormatting = TEXT("{A} - {B}");
+		Op->bSupportsStaticResolution = true;
+		Op->StaticVariableResolveFunction.BindLambda([=](const TArray<int32>& InPinValues)
+			{
+				if (InPinValues.Num() == 0)
+					return 0;
+				if (InPinValues.Num() == 1)
+					return InPinValues[0];
+
+				int32 Return = InPinValues[0];
+				for (int32 i = 1; i < InPinValues.Num(); i++)
+				{
+					Return = Return - InPinValues[i];
+				}
+
+				return Return;
+			});
 		OpInfoMap.Add(Op->Name) = Idx;
 
 		Idx = OpInfos.AddDefaulted();
@@ -161,6 +193,22 @@ void FNiagaraOpInfo::Init()
 		Op->bSupportsAddedInputs = true;
 		Op->AddedInputTypeRestrictions.Add(Type);
 		Op->AddedInputFormatting = TEXT("{A} * {B}");
+		Op->bSupportsStaticResolution = true;
+		Op->StaticVariableResolveFunction.BindLambda([=](const TArray<int32>& InPinValues)
+			{
+				if (InPinValues.Num() == 0)
+					return 0;
+				if (InPinValues.Num() == 1)
+					return InPinValues[0];
+
+				int32 Return = InPinValues[0];
+				for (int32 i = 1; i < InPinValues.Num(); i++)
+				{
+					Return = Return * InPinValues[i];
+				}
+
+				return Return;
+			});
 		OpInfoMap.Add(Op->Name) = Idx;
 
 		Idx = OpInfos.AddDefaulted();
@@ -173,6 +221,25 @@ void FNiagaraOpInfo::Init()
 		Op->Inputs.Add(FNiagaraOpInOutInfo(B, Type, BText, BText, DefaultStr_One));
 		Op->Outputs.Add(FNiagaraOpInOutInfo(Result, Type, ResultText, ResultText, DefaultStr_One, TEXT("{0} / {1}")));
 		Op->BuildName(TEXT("Div"), CategoryName);
+		Op->bSupportsStaticResolution = true;
+		Op->StaticVariableResolveFunction.BindLambda([=](const TArray<int32>& InPinValues)
+			{
+				if (InPinValues.Num() == 0)
+					return 0;
+				if (InPinValues.Num() == 1)
+					return InPinValues[0];
+
+				int32 Return = InPinValues[0];
+				for (int32 i = 1; i < InPinValues.Num(); i++)
+				{
+					if (InPinValues[i] != 0)
+						Return = Return / InPinValues[i];
+					else
+						Return = 0;
+				}
+
+				return Return;
+			});
 		OpInfoMap.Add(Op->Name) = Idx;
 
 		Idx = OpInfos.AddDefaulted();
@@ -251,6 +318,14 @@ void FNiagaraOpInfo::Init()
 		Op->Inputs.Add(FNiagaraOpInOutInfo(A, Type, AText, AText, DefaultStr_One));
 		Op->Outputs.Add(FNiagaraOpInOutInfo(Result, Type, ResultText, ResultText, DefaultStr_One, TEXT("1 - {0}")));
 		Op->BuildName(TEXT("OneMinus"), CategoryName);
+		Op->bSupportsStaticResolution = true;
+		Op->StaticVariableResolveFunction.BindLambda([=](const TArray<int32>& InPinValues)
+			{
+				if (InPinValues.Num() == 1)
+					return 1 - InPinValues[0];
+				else
+					return 0;
+			});
 		OpInfoMap.Add(Op->Name) = Idx;
 
 		Idx = OpInfos.AddDefaulted();
@@ -261,6 +336,14 @@ void FNiagaraOpInfo::Init()
 		Op->Inputs.Add(FNiagaraOpInOutInfo(A, Type, AText, AText, DefaultStr_One));
 		Op->Outputs.Add(FNiagaraOpInOutInfo(Result, Type, ResultText, ResultText, DefaultStr_One, TEXT("-({0})")));
 		Op->BuildName(TEXT("Negate"), CategoryName);
+		Op->bSupportsStaticResolution = true;
+		Op->StaticVariableResolveFunction.BindLambda([=](const TArray<int32>& InPinValues)
+			{
+				if (InPinValues.Num() == 1)
+					return -InPinValues[0];
+				else
+					return 0;
+			});
 		OpInfoMap.Add(Op->Name) = Idx;
 
 		Idx = OpInfos.AddDefaulted();
@@ -922,8 +1005,19 @@ void FNiagaraOpInfo::Init()
 		Op->Inputs.Add(FNiagaraOpInOutInfo(B, Type, BText, BText, DefaultStr_One));
 		Op->Outputs.Add(FNiagaraOpInOutInfo(Result, FNiagaraTypeDefinition::GetBoolDef(), ResultText, ResultText, DefaultStr_One, TEXT("NiagaraAll({0} < {1})")));
 		Op->BuildName(TEXT("CmpLT"), CategoryName);
+		Op->bSupportsStaticResolution = true;
+		Op->StaticVariableResolveFunction.BindLambda([=](const TArray<int32>& InPinValues)
+		{
+			if (InPinValues.Num() != 2)
+				return 0;
+
+			int32 Return = InPinValues[0] < InPinValues[1];
+			return Return;
+		});
 		OpInfoMap.Add(Op->Name) = Idx;
 
+		
+
 		Idx = OpInfos.AddDefaulted();
 		Op = &OpInfos[Idx];
 		Op->Category = NumericCategory;
@@ -934,8 +1028,18 @@ void FNiagaraOpInfo::Init()
 		Op->Inputs.Add(FNiagaraOpInOutInfo(B, Type, BText, BText, DefaultStr_One));
 		Op->Outputs.Add(FNiagaraOpInOutInfo(Result, FNiagaraTypeDefinition::GetBoolDef(), ResultText, ResultText, DefaultStr_One, TEXT("NiagaraAll({0} <= {1})")));
 		Op->BuildName(TEXT("CmpLE"), CategoryName);
+		Op->bSupportsStaticResolution = true;
+		Op->StaticVariableResolveFunction.BindLambda([=](const TArray<int32>& InPinValues)
+		{
+			if (InPinValues.Num() != 2)
+				return 0;
+
+			int32 Return = InPinValues[0] <= InPinValues[1];
+			return Return;
+		});
 		OpInfoMap.Add(Op->Name) = Idx;
 
+
 		Idx = OpInfos.AddDefaulted();
 		Op = &OpInfos[Idx];
 		Op->Category = NumericCategory;
@@ -946,6 +1050,15 @@ void FNiagaraOpInfo::Init()
 		Op->Inputs.Add(FNiagaraOpInOutInfo(B, Type, BText, BText, DefaultStr_One));
 		Op->Outputs.Add(FNiagaraOpInOutInfo(Result, FNiagaraTypeDefinition::GetBoolDef(), ResultText, ResultText, DefaultStr_One, TEXT("NiagaraAll({0} > {1})")));
 		Op->BuildName(TEXT("CmpGT"), CategoryName);
+		Op->bSupportsStaticResolution = true;
+		Op->StaticVariableResolveFunction.BindLambda([=](const TArray<int32>& InPinValues)
+		{
+			if (InPinValues.Num() != 2)
+				return 0;
+
+			int32 Return = InPinValues[0] > InPinValues[1];
+			return Return;
+		});
 		OpInfoMap.Add(Op->Name) = Idx;
 
 		Idx = OpInfos.AddDefaulted();
@@ -958,6 +1071,15 @@ void FNiagaraOpInfo::Init()
 		Op->Inputs.Add(FNiagaraOpInOutInfo(B, Type, BText, BText, DefaultStr_One));
 		Op->Outputs.Add(FNiagaraOpInOutInfo(Result, FNiagaraTypeDefinition::GetBoolDef(), ResultText, ResultText, DefaultStr_One, TEXT("NiagaraAll({0} >= {1})")));
 		Op->BuildName(TEXT("CmpGE"), CategoryName);
+		Op->bSupportsStaticResolution = true;
+		Op->StaticVariableResolveFunction.BindLambda([=](const TArray<int32>& InPinValues)
+		{
+			if (InPinValues.Num() != 2)
+				return 0;
+
+			int32 Return = InPinValues[0] >= InPinValues[1];
+			return Return;
+		});
 		OpInfoMap.Add(Op->Name) = Idx;
 
 		Idx = OpInfos.AddDefaulted();
@@ -970,6 +1092,17 @@ void FNiagaraOpInfo::Init()
 		Op->Inputs.Add(FNiagaraOpInOutInfo(B, Type, BText, BText, DefaultStr_One));
 		Op->Outputs.Add(FNiagaraOpInOutInfo(Result, FNiagaraTypeDefinition::GetBoolDef(), ResultText, ResultText, DefaultStr_One, TEXT("NiagaraAll({0} == {1})")));
 		Op->BuildName(TEXT("CmpEQ"), CategoryName);
+		Op->bSupportsStaticResolution = true;
+		Op->StaticVariableResolveFunction.BindLambda([=](const TArray<int32>& InPinValues)
+		{
+			if (InPinValues.Num() != 2)
+				return 0;
+
+			int32 Return = InPinValues[0] == InPinValues[1];
+			return Return;
+		});
+		OpInfoMap.Add(Op->Name) = Idx;
+
 		OpInfoMap.Add(Op->Name) = Idx;
 		
 		Idx = OpInfos.AddDefaulted();
@@ -982,6 +1115,18 @@ void FNiagaraOpInfo::Init()
 		Op->Inputs.Add(FNiagaraOpInOutInfo(B, Type, BText, BText, DefaultStr_One));
 		Op->Outputs.Add(FNiagaraOpInOutInfo(Result, FNiagaraTypeDefinition::GetBoolDef(), ResultText, ResultText, DefaultStr_One, TEXT("NiagaraAll({0} != {1})")));
 		Op->BuildName(TEXT("CmpNEQ"), CategoryName);
+		Op->bSupportsStaticResolution = true;
+		Op->StaticVariableResolveFunction.BindLambda([=](const TArray<int32>& InPinValues)
+		{
+			if (InPinValues.Num() != 2)
+				return 0;
+
+			int32 Return = InPinValues[0] != InPinValues[1];
+			return Return;
+		});
+		OpInfoMap.Add(Op->Name) = Idx;
+
+		
 		OpInfoMap.Add(Op->Name) = Idx;
 	}
 
@@ -1101,6 +1246,24 @@ void FNiagaraOpInfo::Init()
 	Op->bSupportsAddedInputs = true;
 	Op->AddedInputTypeRestrictions.Add(BoolType);
 	Op->AddedInputFormatting = TEXT("{A} && {B}");
+	Op->bSupportsStaticResolution = true;
+	Op->StaticVariableResolveFunction.BindLambda([=](const TArray<int32>& InPinValues)
+	{
+			if (InPinValues.Num() == 0)
+				return 0;
+			if (InPinValues.Num() == 1)
+				return InPinValues[0];
+
+			int32 Return = InPinValues[0];
+			for (int32 i = 1; i < InPinValues.Num(); i++)
+			{
+				Return = Return && InPinValues[i];
+			}
+
+			return Return;
+	});
+	OpInfoMap.Add(Op->Name) = Idx;
+
 	OpInfoMap.Add(Op->Name) = Idx;
 
 	Idx = OpInfos.AddDefaulted();
@@ -1116,8 +1279,26 @@ void FNiagaraOpInfo::Init()
 	Op->bSupportsAddedInputs = true;
 	Op->AddedInputTypeRestrictions.Add(BoolType);
 	Op->AddedInputFormatting = TEXT("{A} || {B}");
+	Op->bSupportsStaticResolution = true;
+	Op->StaticVariableResolveFunction.BindLambda([=](const TArray<int32>& InPinValues)
+	{
+		if (InPinValues.Num() == 0)
+			return 0;
+		if (InPinValues.Num() == 1)
+			return InPinValues[0];
+
+		int32 Return = InPinValues[0];
+		for (int32 i = 1; i < InPinValues.Num(); i++)
+		{
+			Return = Return || InPinValues[i];
+		}
+
+		return Return;
+	});
 	OpInfoMap.Add(Op->Name) = Idx;
 
+	
+
 	Idx = OpInfos.AddDefaulted();
 	Op = &OpInfos[Idx];
 	Op->Category = BoolCategory;
@@ -1127,8 +1308,18 @@ void FNiagaraOpInfo::Init()
 	Op->Inputs.Add(FNiagaraOpInOutInfo(A, BoolType, AText, AText, Default_BoolOne));
 	Op->Outputs.Add(FNiagaraOpInOutInfo(Result, BoolType, ResultText, ResultText, Default_BoolOne, TEXT("!{0}")));
 	Op->BuildName(TEXT("LogicNot"), BoolCategoryName);
+	Op->bSupportsStaticResolution = true;
+	Op->StaticVariableResolveFunction.BindLambda([=](const TArray<int32>& InPinValues)
+	{
+		if (InPinValues.Num() != 1)
+			return 0;
+
+		int32 Return = !InPinValues[0];
+		return Return;
+	});
 	OpInfoMap.Add(Op->Name) = Idx;
 
+
 	Idx = OpInfos.AddDefaulted();
 	Op = &OpInfos[Idx];
 	Op->Category = BoolCategory;
@@ -1139,6 +1330,15 @@ void FNiagaraOpInfo::Init()
 	Op->Inputs.Add(FNiagaraOpInOutInfo(B, BoolType, BText, BText, Default_BoolOne));
 	Op->Outputs.Add(FNiagaraOpInOutInfo(Result, BoolType, ResultText, ResultText, Default_BoolOne, TEXT("NiagaraAll({0} == {1})")));
 	Op->BuildName(TEXT("LogicEq"), BoolCategoryName);
+	Op->bSupportsStaticResolution = true;
+	Op->StaticVariableResolveFunction.BindLambda([=](const TArray<int32>& InPinValues)
+	{
+		if (InPinValues.Num() != 2)
+			return 0;
+
+		int32 Return = InPinValues[0] == InPinValues[1];
+		return Return;
+	});
 	OpInfoMap.Add(Op->Name) = Idx;
 
 	Idx = OpInfos.AddDefaulted();
@@ -1151,8 +1351,19 @@ void FNiagaraOpInfo::Init()
 	Op->Inputs.Add(FNiagaraOpInOutInfo(B, BoolType, BText, BText, Default_BoolOne));
 	Op->Outputs.Add(FNiagaraOpInOutInfo(Result, BoolType, ResultText, ResultText, Default_BoolOne, TEXT("NiagaraAll({0} != {1})")));
 	Op->BuildName(TEXT("LogicNEq"), BoolCategoryName);
+	Op->bSupportsStaticResolution = true;
+	Op->StaticVariableResolveFunction.BindLambda([=](const TArray<int32>& InPinValues)
+	{
+		if (InPinValues.Num() != 2)
+			return 0;
+
+		int32 Return = InPinValues[0] != InPinValues[1];
+		return Return;
+	});
 	OpInfoMap.Add(Op->Name) = Idx;
 
+
+
 	//////////////////////////////////////////////////////////////////////////
 	//Matrix only ops
 	FString Default_MatrixZero(TEXT(
diff --git a/Engine/Plugins/FX/Niagara/Source/NiagaraEditor/Private/NiagaraEditorModule.cpp b/Engine/Plugins/FX/Niagara/Source/NiagaraEditor/Private/NiagaraEditorModule.cpp
--- a/Engine/Plugins/FX/Niagara/Source/NiagaraEditor/Private/NiagaraEditorModule.cpp
+++ b/Engine/Plugins/FX/Niagara/Source/NiagaraEditor/Private/NiagaraEditorModule.cpp
@@ -214,10 +214,19 @@ public:
 
 	//~ FGraphPanelPinFactory interface
 	virtual TSharedPtr<class SGraphPin> CreatePin(class UEdGraphPin* InPin) const override
+	{
+		TSharedPtr<class SGraphPin> OutPin = InternalCreatePin(InPin);
+
+		if (OutPin.IsValid() && UEdGraphSchema_Niagara::IsStaticPin(InPin))
+			OutPin->SetCustomPinIcon(FNiagaraEditorStyle::Get().GetBrush(TEXT("NiagaraEditor.Pins.StaticConnected")), FNiagaraEditorStyle::Get().GetBrush(TEXT("NiagaraEditor.Pins.StaticDisconnected")));
+		return OutPin;
+	}
+
+	TSharedPtr<class SGraphPin> InternalCreatePin(class UEdGraphPin* InPin) const
 	{
 		if (const UEdGraphSchema_Niagara* NSchema = Cast<UEdGraphSchema_Niagara>(InPin->GetSchema()))
 		{
-			if (InPin->PinType.PinCategory == UEdGraphSchema_Niagara::PinCategoryType)
+			if (InPin->PinType.PinCategory == UEdGraphSchema_Niagara::PinCategoryType || InPin->PinType.PinCategory == UEdGraphSchema_Niagara::PinCategoryStaticType)
 			{
 				if (InPin->PinType.PinSubCategoryObject != nullptr && InPin->PinType.PinSubCategoryObject->IsA<UScriptStruct>())
 				{
@@ -241,7 +250,7 @@ public:
 					return CreatePin(InPin);
 				}
 			}
-			else if (InPin->PinType.PinCategory == UEdGraphSchema_Niagara::PinCategoryEnum)
+			else if (InPin->PinType.PinCategory == UEdGraphSchema_Niagara::PinCategoryEnum || InPin->PinType.PinCategory == UEdGraphSchema_Niagara::PinCategoryStaticEnum)
 			{
 				const UEnum* Enum = Cast<const UEnum>(InPin->PinType.PinSubCategoryObject.Get());
 				if (Enum == nullptr)
@@ -974,7 +983,9 @@ void FNiagaraEditorModule::StartupModule()
 	EnumTypeUtilities = MakeShareable(new FNiagaraEditorEnumTypeUtilities());
 	RegisterTypeUtilities(FNiagaraTypeDefinition::GetFloatDef(), MakeShareable(new FNiagaraEditorFloatTypeUtilities()));
 	RegisterTypeUtilities(FNiagaraTypeDefinition::GetIntDef(), MakeShareable(new FNiagaraEditorIntegerTypeUtilities()));
+	RegisterTypeUtilities(FNiagaraTypeDefinition::GetIntDef().ToStaticDef(), MakeShareable(new FNiagaraEditorIntegerTypeUtilities()));
 	RegisterTypeUtilities(FNiagaraTypeDefinition::GetBoolDef(), MakeShareable(new FNiagaraEditorBoolTypeUtilities()));
+	RegisterTypeUtilities(FNiagaraTypeDefinition::GetBoolDef().ToStaticDef(), MakeShareable(new FNiagaraEditorBoolTypeUtilities()));
 	RegisterTypeUtilities(FNiagaraTypeDefinition::GetVec2Def(), MakeShareable(new FNiagaraEditorVector2TypeUtilities()));
 	RegisterTypeUtilities(FNiagaraTypeDefinition::GetVec3Def(), MakeShareable(new FNiagaraEditorVector3TypeUtilities()));
 	RegisterTypeUtilities(FNiagaraTypeDefinition::GetVec4Def(), MakeShareable(new FNiagaraEditorVector4TypeUtilities()));
@@ -1088,6 +1099,12 @@ void FNiagaraEditorModule::StartupModule()
 		return PrecompileDuplicate(OwningSystemRequestData, OwningSystem, OwningEmitter, TargetScript, Version);
 	}));
 
+	GraphCacheTraversalHandle = NiagaraModule.RegisterGraphTraversalCacher(INiagaraModule::FOnCacheGraphTraversal::CreateLambda([this](const UObject* InObj, FGuid Version)
+		{
+			return CacheGraphTraversal(InObj, Version);
+		}));
+
+
 	TestCompileScriptCommand = IConsoleManager::Get().RegisterConsoleCommand(
 		TEXT("fx.TestCompileNiagaraScript"),
 		TEXT("Compiles the specified script on disk for the niagara vector vm"),
@@ -1197,6 +1214,8 @@ void FNiagaraEditorModule::ShutdownModule()
 		NiagaraModule->UnregisterScriptCompiler(ScriptCompilerHandle);
 		NiagaraModule->UnregisterCompileResultDelegate(CompileResultHandle);
 		NiagaraModule->UnregisterPrecompiler(PrecompilerHandle);
+		NiagaraModule->UnregisterPrecompileDuplicator(PrecompileDuplicatorHandle);
+		NiagaraModule->UnregisterGraphTraversalCacher(GraphCacheTraversalHandle);
 	}
 
 	// Verify that we've cleaned up all the view models in the world.
diff --git a/Engine/Plugins/FX/Niagara/Source/NiagaraEditor/Private/NiagaraEditorStyle.cpp b/Engine/Plugins/FX/Niagara/Source/NiagaraEditor/Private/NiagaraEditorStyle.cpp
--- a/Engine/Plugins/FX/Niagara/Source/NiagaraEditor/Private/NiagaraEditorStyle.cpp
+++ b/Engine/Plugins/FX/Niagara/Source/NiagaraEditor/Private/NiagaraEditorStyle.cpp
@@ -127,6 +127,8 @@ void InitAssetPicker(TSharedRef< FSlateStyleSet > Style)
 	Style->Set("NiagaraEditor.NewAssetDialog.ActiveOptionBorderColor", FLinearColor(FColor(96, 96, 96)));
 
 	Style->Set("NiagaraEditor.NewAssetDialog.SubBorder", new BOX_CORE_BRUSH("Common/GroupBorderLight", FMargin(4.0f / 16.0f)));
+
+
 }
 
 void InitActionMenu(TSharedRef< FSlateStyleSet > Style)
@@ -278,6 +280,13 @@ void InitParameters(TSharedRef< FSlateStyleSet > Style)
 	Style->Set("NiagaraEditor.ParameterName.TypeText", FTextBlockStyle(NormalText)
 		.SetFont(DEFAULT_FONT("Regular", 8))
 		.SetColorAndOpacity(FLinearColor(1.0f, 1.0f, 1.0f, 0.5f)));
+
+
+	Style->Set("NiagaraEditor.StaticIcon", new IMAGE_PLUGIN_BRUSH("Icons/staticpill_16x", Icon16x16));
+	Style->Set("NiagaraEditor.Pins.StaticConnected", new IMAGE_PLUGIN_BRUSH("Icons/StaticPin_Connected", Icon16x16));
+	Style->Set("NiagaraEditor.Pins.StaticDisconnected", new IMAGE_PLUGIN_BRUSH("Icons/StaticPin_Disconnected", Icon16x16));
+	Style->Set("NiagaraEditor.Pins.StaticConnectedHovered", new IMAGE_BRUSH("Icons/StaticPin_Connected", Icon16x16, FLinearColor(0.8f, 0.8f, 0.8f)));
+	Style->Set("NiagaraEditor.Pins.StaticDisconnectedHovered", new IMAGE_BRUSH("Icons/StaticPin_Disconnected", Icon16x16, FLinearColor(0.8f, 0.8f, 0.8f)));
 }
 
 void InitParameterMapView(TSharedRef< FSlateStyleSet > Style)
diff --git a/Engine/Plugins/FX/Niagara/Source/NiagaraEditor/Private/NiagaraEditorUtilities.cpp b/Engine/Plugins/FX/Niagara/Source/NiagaraEditor/Private/NiagaraEditorUtilities.cpp
--- a/Engine/Plugins/FX/Niagara/Source/NiagaraEditor/Private/NiagaraEditorUtilities.cpp
+++ b/Engine/Plugins/FX/Niagara/Source/NiagaraEditor/Private/NiagaraEditorUtilities.cpp
@@ -1083,6 +1083,10 @@ void FNiagaraEditorUtilities::SetStaticSwitchConstants(UNiagaraGraph* Graph, TAr
 			{
 				SwitchNode->SetSwitchValue(ConstantResolver);
 			}
+			else if (SwitchNode->IsSetByPin())
+			{
+				SwitchNode->SetSwitchValue(ConstantResolver);
+			}
 			else
 			{
 				FEdGraphPinType VarType = Schema->TypeDefinitionToPinType(SwitchNode->GetInputType());
@@ -1168,7 +1172,7 @@ bool FNiagaraEditorUtilities::ResolveConstantValue(UEdGraphPin* Pin, int32& Valu
 
 TSharedPtr<FStructOnScope> FNiagaraEditorUtilities::StaticSwitchDefaultIntToStructOnScope(int32 InStaticSwitchDefaultValue, FNiagaraTypeDefinition InSwitchType)
 {
-	if (InSwitchType == FNiagaraTypeDefinition::GetBoolDef())
+	if (InSwitchType.IsSameBaseDefinition(FNiagaraTypeDefinition::GetBoolDef()))
 	{
 		checkf(FNiagaraBool::StaticStruct()->GetStructureSize() == InSwitchType.GetSize(), TEXT("Value to type def size mismatch."));
 
@@ -1180,7 +1184,7 @@ TSharedPtr<FStructOnScope> FNiagaraEditorUtilities::StaticSwitchDefaultIntToStru
 
 		return StructValue;
 	}
-	else if (InSwitchType == FNiagaraTypeDefinition::GetIntDef() || InSwitchType.IsEnum())
+	else if (InSwitchType.IsSameBaseDefinition(FNiagaraTypeDefinition::GetIntDef()) || InSwitchType.IsEnum())
 	{
 		checkf(FNiagaraInt32::StaticStruct()->GetStructureSize() == InSwitchType.GetSize(), TEXT("Value to type def size mismatch."));
 
@@ -2012,7 +2016,7 @@ TArray<UNiagaraComponent*> FNiagaraEditorUtilities::GetComponentsThatReferenceSy
 		UNiagaraComponent* Component = *ComponentIt;
 		if (Component && Component->GetAsset())
 		{
-			for (auto EmitterHandle : ReferencedSystemViewModel.GetSystem().GetEmitterHandles())
+			for (const auto& EmitterHandle : ReferencedSystemViewModel.GetSystem().GetEmitterHandles())
 			{
 				if (Component->GetAsset()->UsesEmitter(EmitterHandle.GetInstance()->GetParent()))
 				{
@@ -3511,4 +3515,35 @@ bool FNiagaraParameterUtilities::TestCanRenameWithMessage(FName ParameterName, F
 	return false;
 }
 
+void FNiagaraParameterUtilities::FilterToRelevantStaticVariables(const TArray<FNiagaraVariable>& InVars, TArray<FNiagaraVariable>& OutVars, FName InOldEmitterAlias, FName InNewEmitterAlias, bool bFilterByEmitterAliasAndConvertToUnaliased)
+{
+	FNiagaraAliasContext RenameContext(ENiagaraScriptUsage::ParticleSpawnScript);
+	RenameContext.ChangeEmitterName(InOldEmitterAlias.ToString(), InNewEmitterAlias.ToString());
+
+	for (const FNiagaraVariable& InVar : InVars)
+	{
+		if (InVar.GetType().IsStatic())
+		{
+			if (bFilterByEmitterAliasAndConvertToUnaliased)
+			{
+				FNiagaraVariable NewVar = FNiagaraUtilities::ResolveAliases(InVar, RenameContext);
+
+				if (NewVar.GetName() != InVar.GetName())
+				{
+					NewVar.SetData(InVar.GetData());
+					OutVars.AddUnique(NewVar);
+				}
+				else
+				{
+					OutVars.AddUnique(InVar);
+				}
+			}
+			else
+			{
+				OutVars.AddUnique(InVar);
+			}
+		}
+	}
+}
+
 #undef LOCTEXT_NAMESPACE
diff --git a/Engine/Plugins/FX/Niagara/Source/NiagaraEditor/Private/NiagaraGraph.cpp b/Engine/Plugins/FX/Niagara/Source/NiagaraEditor/Private/NiagaraGraph.cpp
--- a/Engine/Plugins/FX/Niagara/Source/NiagaraEditor/Private/NiagaraGraph.cpp
+++ b/Engine/Plugins/FX/Niagara/Source/NiagaraEditor/Private/NiagaraGraph.cpp
@@ -1186,20 +1186,27 @@ UNiagaraNodeOutput* UNiagaraGraph::FindEquivalentOutputNode(ENiagaraScriptUsage
 
 void BuildTraversalHelper(TArray<class UNiagaraNode*>& OutNodesTraversed, UNiagaraNode* CurrentNode, bool bEvaluateStaticSwitches)
 {
-	if (CurrentNode == nullptr || (bEvaluateStaticSwitches && CurrentNode->IsA<UNiagaraNodeStaticSwitch>()))
-	{
-		return;
-	}
-
-	for (UEdGraphPin* Pin : CurrentNode->GetAllPins())
-	{
+	auto VisitPin = [](const UEdGraphPin* Pin, TArray<class UNiagaraNode*>& OutNodesTraversed, bool bEvaluateStaticSwitches) {
 		if (Pin->Direction == EEdGraphPinDirection::EGPD_Input && Pin->LinkedTo.Num() > 0)
 		{
 			for (UEdGraphPin* LinkedPin : Pin->LinkedTo)
 			{
-				UEdGraphPin* TracedPin = bEvaluateStaticSwitches ? UNiagaraNode::TraceOutputPin(LinkedPin) : LinkedPin;
+				bool bFilterForCompilation = true;
+				TArray<const UNiagaraNode*> NodesVisitedDuringTrace;
+
+				UEdGraphPin* TracedPin = bEvaluateStaticSwitches ? UNiagaraNode::TraceOutputPin(LinkedPin, bFilterForCompilation,  &NodesVisitedDuringTrace) : LinkedPin;
 				if (TracedPin != nullptr)
 				{
+					// Deal with static switches driven by a pin...
+					for (const UNiagaraNode* VisitedNode : NodesVisitedDuringTrace)
+					{
+						if (OutNodesTraversed.Contains(VisitedNode))
+						{
+							continue;
+						}
+						BuildTraversalHelper(OutNodesTraversed, (UNiagaraNode * )VisitedNode, bEvaluateStaticSwitches);
+					}
+
 					UNiagaraNode* Node = Cast<UNiagaraNode>(TracedPin->GetOwningNode());
 					if (OutNodesTraversed.Contains(Node))
 					{
@@ -1209,6 +1216,33 @@ void BuildTraversalHelper(TArray<class UNiagaraNode*>& OutNodesTraversed, UNiaga
 				}
 			}
 		}
+	};
+
+
+	if (CurrentNode == nullptr)
+	{
+		return;
+	}
+
+	if (bEvaluateStaticSwitches)
+	{
+		UNiagaraNodeStaticSwitch* StaticSwitch = Cast< UNiagaraNodeStaticSwitch>(CurrentNode);
+		if (StaticSwitch && StaticSwitch->IsSetByPin())
+		{
+			// The selector pin is never traversed directly below by the TracedPin route, so we explicitly visit it here.
+			UEdGraphPin* Pin = StaticSwitch->GetSelectorPin();
+			if (Pin)
+				VisitPin(Pin, OutNodesTraversed, bEvaluateStaticSwitches);
+		}
+		else if (StaticSwitch)
+		{
+			return;
+		}
+	}
+
+	for (UEdGraphPin* Pin : CurrentNode->GetAllPins())
+	{
+		VisitPin(Pin, OutNodesTraversed, bEvaluateStaticSwitches);
 	}
 
 	OutNodesTraversed.Add(CurrentNode);
@@ -1322,7 +1356,7 @@ TArray<FNiagaraVariable> UNiagaraGraph::FindStaticSwitchInputs(bool bReachableOn
 	for (UEdGraphNode* Node : NodesToProcess)
 	{
 		UNiagaraNodeStaticSwitch* SwitchNode = Cast<UNiagaraNodeStaticSwitch>(Node);
-		if (SwitchNode && !SwitchNode->IsSetByCompiler())
+		if (SwitchNode && !SwitchNode->IsSetByCompiler() && !SwitchNode->IsSetByPin())
 		{
 			FNiagaraVariable Variable(SwitchNode->GetInputType(), SwitchNode->InputParameterName);
 			Result.AddUnique(Variable);
@@ -2718,13 +2752,16 @@ void UNiagaraGraph::ResolveNumerics(TMap<UNiagaraNode*, bool>& VisitedNodes, UEd
 			{
 				for (int32 j = 0; j < InputPins[i]->LinkedTo.Num(); j++)
 				{
-					UNiagaraNode* FoundNode = Cast<UNiagaraNode>(InputPins[i]->LinkedTo[j]->GetOwningNode());
-					if (!FoundNode || VisitedNodes.Contains(FoundNode))
+					if (InputPins[i]->LinkedTo[j])
 					{
-						continue;
+						UNiagaraNode* FoundNode = Cast<UNiagaraNode>(InputPins[i]->LinkedTo[j]->GetOwningNode());
+						if (!FoundNode || VisitedNodes.Contains(FoundNode))
+						{
+							continue;
+						}
+						VisitedNodes.Add(FoundNode, true);
+						ResolveNumerics(VisitedNodes, FoundNode);
 					}
-					VisitedNodes.Add(FoundNode, true);
-					ResolveNumerics(VisitedNodes, FoundNode);
 				}
 			}
 		}
@@ -2961,7 +2998,7 @@ void UNiagaraGraph::RefreshParameterReferences() const
 		}
 
 		UNiagaraNodeStaticSwitch* SwitchNode = Cast<UNiagaraNodeStaticSwitch>(Node);
-		if (SwitchNode && !SwitchNode->IsSetByCompiler())
+		if (SwitchNode && !SwitchNode->IsSetByCompiler() && !SwitchNode->IsSetByPin())
 		{
 			FNiagaraVariable Variable(SwitchNode->GetInputType(), SwitchNode->InputParameterName);
 			AddStaticParameterReference(Variable, SwitchNode);
diff --git a/Engine/Plugins/FX/Niagara/Source/NiagaraEditor/Private/NiagaraHlslTranslator.cpp b/Engine/Plugins/FX/Niagara/Source/NiagaraEditor/Private/NiagaraHlslTranslator.cpp
--- a/Engine/Plugins/FX/Niagara/Source/NiagaraEditor/Private/NiagaraHlslTranslator.cpp
+++ b/Engine/Plugins/FX/Niagara/Source/NiagaraEditor/Private/NiagaraHlslTranslator.cpp
@@ -373,7 +373,7 @@ FString FHlslNiagaraTranslator::GetCodeAsSource(int32 ChunkIdx)
 	return "Undefined";
 }
 
-bool FHlslNiagaraTranslator::ValidateTypePins(UNiagaraNode* NodeToValidate)
+bool FHlslNiagaraTranslator::ValidateTypePins(const UNiagaraNode* NodeToValidate)
 {
 	bool bPinsAreValid = true;
 	for (UEdGraphPin* Pin : NodeToValidate->GetAllPins())
@@ -383,7 +383,7 @@ bool FHlslNiagaraTranslator::ValidateTypePins(UNiagaraNode* NodeToValidate)
 			Error(LOCTEXT("InvalidPinTypeError", "Node pin has an undefined type."), NodeToValidate, Pin);
 			bPinsAreValid = false;
 		}
-		else if (Pin->PinType.PinCategory == UEdGraphSchema_Niagara::PinCategoryType)
+		else if (Pin->PinType.PinCategory == UEdGraphSchema_Niagara::PinCategoryType || Pin->PinType.PinCategory == UEdGraphSchema_Niagara::PinCategoryStaticType)
 		{
 			FNiagaraTypeDefinition Type = Schema->PinToTypeDefinition(Pin);
 			if (Type.IsValid() == false)
@@ -452,7 +452,7 @@ void FHlslNiagaraTranslator::GenerateFunctionSignature(ENiagaraScriptUsage Scrip
 				FNiagaraVariable InputVar = InputsNodes[i]->Input;
 				if (GetLiteralConstantVariable(InputVar))
 				{
-					checkf(InputVar.GetType() == FNiagaraTypeDefinition::GetBoolDef(), TEXT("Only boolean types are currently supported for literal constants."));
+					checkf(InputVar.GetType().IsSameBaseDefinition(FNiagaraTypeDefinition::GetBoolDef()), TEXT("Only boolean types are currently supported for literal constants."));
 					FString LiteralConstantAlias = InputVar.GetName().ToString() + TEXT("_") + (InputVar.GetValue<bool>() ? TEXT("true") : TEXT("false"));
 					InName += TEXT("_") + GetSanitizedSymbolName(LiteralConstantAlias.Replace(TEXT("."), TEXT("_")));
 					ConstantInputIndicesToRemove.Add(i);
@@ -551,7 +551,7 @@ void FHlslNiagaraTranslator::GenerateFunctionSignature(ENiagaraScriptUsage Scrip
 //////////////////////////////////////////////////////////////////////////
 
 FHlslNiagaraTranslator::FHlslNiagaraTranslator()
-	: Schema(nullptr)
+	: Schema(GetDefault<UEdGraphSchema_Niagara>())
 	, TranslateResults()
 	, CurrentBodyChunkMode(ENiagaraCodeChunkMode::Body)
 	, ActiveStageIdx(-1)
@@ -1721,7 +1721,7 @@ const FNiagaraTranslateResults &FHlslNiagaraTranslator::Translate(const FNiagara
 		TArray<FNiagaraVariable> BasicAttributes;
 		for (FNiagaraVariable& Var : InstanceRead.Variables)
 		{
-			if (Var.GetType().GetClass() != nullptr)
+			if (Var.GetType().GetClass() != nullptr || Var.GetType().IsStatic())
 			{
 				continue;
 			}
@@ -1729,7 +1729,7 @@ const FNiagaraTranslateResults &FHlslNiagaraTranslator::Translate(const FNiagara
 		}
 		for (FNiagaraVariable& Var : InstanceWrite.Variables)
 		{
-			if (Var.GetType().GetClass() != nullptr)
+			if (Var.GetType().GetClass() != nullptr || Var.GetType().IsStatic())
 			{
 				continue;
 			}
@@ -1913,6 +1913,29 @@ const FNiagaraTranslateResults &FHlslNiagaraTranslator::Translate(const FNiagara
 
 		}
 
+		// Display the computed compile tags in the source hlsl to make checking easier.
+		if (CompileData->StaticVariables.Num() != 0)
+		{
+			Preamble.Appendf(TEXT("\n// Compile Data> Static Variables Input: \n"));
+
+			for (const FNiagaraVariable& StaticVar: CompileData->StaticVariables)
+			{
+				Preamble.Appendf(TEXT("//\tVariable: %s \n"), *StaticVar.ToString());
+			}
+
+		}
+
+		if (CompilationOutput.ScriptData.StaticVariablesWritten.Num() != 0)
+		{
+			Preamble.Appendf(TEXT("\n// Static Variables Written: \n"));
+
+			for (const FNiagaraVariable& StaticVar : CompilationOutput.ScriptData.StaticVariablesWritten)
+			{
+				Preamble.Appendf(TEXT("//\tVariable: %s \n"), *StaticVar.ToString());
+			}
+
+		}
+
 		HlslOutput = Preamble + TEXT("\n\n") +  HlslOutput;
 
 		// We may have created some transient data interfaces. This cleans up the ones that we created.
@@ -1949,7 +1972,7 @@ void FHlslNiagaraTranslator::HandleSimStageSetupAndTeardown(int32 InWhichStage,
 	{
 		IterationSourceVar = *FoundVar;
 	}
-
+	
 	if (!IterationSourceVar.IsValid())
 	{
 		Error(FText::Format(LOCTEXT("CannotFindIterationSourceInParamMap", "Variable {0} missing in graphs referenced during compile!"), FText::FromName(TranslationStage.IterationSource)), nullptr, nullptr);
@@ -4144,7 +4167,7 @@ void FHlslNiagaraTranslator::UpdateStaticSwitchConstants(UEdGraphNode* Node)
 	if (UNiagaraNodeStaticSwitch* SwitchNode = Cast<UNiagaraNodeStaticSwitch>(Node))
 	{
 		SwitchNode->CheckForOutdatedEnum(this);
-		
+
 		TArray<UNiagaraNodeStaticSwitch*> NodesToUpdate;
 		NodesToUpdate.Add(SwitchNode);
 
@@ -4371,7 +4394,7 @@ int32 FHlslNiagaraTranslator::GetConstant(const FNiagaraVariable& Constant, FStr
 	FNiagaraVariable LiteralConstant = Constant;
 	if (GetLiteralConstantVariable(LiteralConstant))
 	{
-		checkf(LiteralConstant.GetType() == FNiagaraTypeDefinition::GetBoolDef() || LiteralConstant.GetType() == FNiagaraTypeDefinition::GetVec3Def() || LiteralConstant.GetType() == FNiagaraTypeDefinition::GetPositionDef(), TEXT("Only boolean and vec3 types are currently supported for literal constants."));
+		checkf(LiteralConstant.GetType().IsSameBaseDefinition(FNiagaraTypeDefinition::GetBoolDef()) || LiteralConstant.GetType() == FNiagaraTypeDefinition::GetVec3Def() || LiteralConstant.GetType() == FNiagaraTypeDefinition::GetPositionDef(), TEXT("Only boolean and vec3 types are currently supported for literal constants."));
 		ConstantStr = GenerateConstantString(LiteralConstant);
 	}
 	else
@@ -4412,6 +4435,14 @@ int32 FHlslNiagaraTranslator::GetConstantDirect(bool InConstantValue)
 	return GetConstant(Constant);
 }
 
+int32 FHlslNiagaraTranslator::GetConstantDirect(int InConstantValue)
+{
+	FNiagaraVariable Constant(FNiagaraTypeDefinition::GetIntDef(), TEXT("Constant"));
+	Constant.SetValue(InConstantValue);
+
+	return GetConstant(Constant);
+}
+
 bool FHlslNiagaraTranslator::GenerateStructInitializer(TStringBuilder<128>& InitializerString, UStruct* UserDefinedStruct, const void* StructData)
 {
 	InitializerString.Append('{');
@@ -4511,7 +4542,7 @@ FString FHlslNiagaraTranslator::GenerateConstantString(const FNiagaraVariable& C
 				ValuePtr[12], ValuePtr[13], ValuePtr[14], ValuePtr[15]
 			);
 		}
-		else if (Type == FNiagaraTypeDefinition::GetBoolDef())
+		else if (Type.IsSameBaseDefinition(FNiagaraTypeDefinition::GetBoolDef()))
 		{
 			if (Constant.GetValue<FNiagaraBool>().IsValid() == false)
 			{
@@ -4937,6 +4968,7 @@ void FHlslNiagaraTranslator::ParameterMapSet(UNiagaraNodeParameterMapSet* SetNod
 				continue;
 			}
 
+
 			if (ParamMapHistoryIdx < ParamMapHistories.Num())
 			{
 				int32 VarIdx = ParamMapHistories[ParamMapHistoryIdx].FindVariableByName(Var.GetName());
@@ -4944,11 +4976,68 @@ void FHlslNiagaraTranslator::ParameterMapSet(UNiagaraNodeParameterMapSet* SetNod
 				{
 					ParamMapSetVariablesToChunks[ParamMapHistoryIdx][VarIdx] = Inputs[i].CompilationIndex;
 					RecordParamMapDefinedAttributeToNamespaceVar(Var, Inputs[i].Pin);
-					if (ParamMapHistories[ParamMapHistoryIdx].IsPrimaryDataSetOutput(Var, GetTargetUsage())) // Note that data interfaces aren't ever in the primary data set even if the namespace matches.)
-						CompilationOutput.ScriptData.AttributesWritten.AddUnique(Var);
+					if (!Var.GetType().IsStatic())
+					{
+						if (ParamMapHistories[ParamMapHistoryIdx].IsPrimaryDataSetOutput(Var, GetTargetUsage())) // Note that data interfaces aren't ever in the primary data set even if the namespace matches.)
+							CompilationOutput.ScriptData.AttributesWritten.AddUnique(Var);
+					}
+					else 
+					{
+						if (Var.GetType().IsStatic())
+						{
+							if (Var.GetName().ToString().Contains("MathFun"))
+							{
+								UE_LOG(LogNiagaraEditor, Log, TEXT("A"));
+							}
+
+							bool bAllSame = true;//CompileDataStaticVariablesWithMultipleWrites;
+							int32 FoundOverrideIdx = CompileData->StaticVariablesWithMultipleWrites.IndexOfByPredicate([&](const FNiagaraVariable& InObj) -> bool
+								{
+									return (InObj.GetName() == Var.GetName());
+								});;
+
+							if (FoundOverrideIdx != INDEX_NONE)
+							{
+								bAllSame = false;
+							}
+							else if (FoundOverrideIdx == INDEX_NONE && FNiagaraParameterMapHistory::IsAliasedEmitterParameter(Var))
+							{
+								FNiagaraAliasContext ResolveAliasesContext(FNiagaraAliasContext::ERapidIterationParameterMode::EmitterOrParticleScript);
+								ResolveAliasesContext.ChangeEmitterToEmitterName(CompileData->GetUniqueEmitterName());
+								FNiagaraVariable TestEmitterResolvedVar = FNiagaraUtilities::ResolveAliases(Var, ResolveAliasesContext);
+
+								FoundOverrideIdx = CompileData->StaticVariablesWithMultipleWrites.IndexOfByPredicate([&](const FNiagaraVariable& InObj) -> bool
+									{
+										return (InObj.GetName() == TestEmitterResolvedVar.GetName());
+									});;
+
+								if (FoundOverrideIdx != INDEX_NONE)
+								{
+									bAllSame = false;
+								}
+							}
+
+							if (!bAllSame)
+							{
+								Error(FText::Format(LOCTEXT("ParameterMapStaticMultipleWriteErrorFormat", "Static variable is not set to a consistent value. Please ensure that all values are equal.  Parameter: {0}"), FText::FromName(Var.GetName())), SetNode, Inputs[i].Pin);
+							}
+						}
+
+						if (ParamMapHistories[ParamMapHistoryIdx].IsPrimaryDataSetOutput(Var, GetTargetUsage(), true, true)) // Note that data interfaces aren't ever in the primary data set even if the namespace matches.)
+						{
+							FString DebugStr;
+							FNiagaraVariable StaticVersionOfVar = Var;
+							SetConstantByStaticVariable(StaticVersionOfVar, Inputs[i].Pin, &DebugStr);
+							AddBodyComment(TEXT("//SetConstantByStaticVariable \"") + DebugStr + TEXT("\""));
+							CompilationOutput.ScriptData.StaticVariablesWritten.AddUnique(StaticVersionOfVar);
+						}
+					}
 				}
 			}
 
+
+			
+
 			if (Var.IsDataInterface())
 			{
 				if (CompileOptions.TargetUsage == ENiagaraScriptUsage::ParticleSpawnScriptInterpolated && TranslationStages[ActiveStageIdx].ScriptUsage == ENiagaraScriptUsage::ParticleUpdateScript)
@@ -5161,7 +5250,10 @@ bool FHlslNiagaraTranslator::ParameterMapRegisterExternalConstantNamespaceVariab
 		bool bIsDataInterface = InVariable.GetType().IsDataInterface();
 		const FString* EmitterAlias = ActiveHistoryForFunctionCalls.GetEmitterAlias();
 		bool bIsPerInstanceBulkSystemParam = IsBulkSystemScript() && !bIsDataInterface && (FNiagaraParameterMapHistory::IsUserParameter(InVariable) || FNiagaraParameterMapHistory::IsPerInstanceEngineParameter(InVariable, EmitterAlias != nullptr ? *EmitterAlias : TEXT("Emitter")));
-		const bool bIsExternalConstantParameter = FNiagaraParameterMapHistory::IsRapidIterationParameter(InVariable);
+		const bool bIsExternalConstantParameter = FNiagaraParameterMapHistory::IsRapidIterationParameter(InVariable) && !InVariable.GetType().IsStatic();
+
+		if (InVariable.GetType().IsStatic())
+			return false;
 		
 		if (!bIsPerInstanceBulkSystemParam)
 		{
@@ -5201,12 +5293,7 @@ bool FHlslNiagaraTranslator::ParameterMapRegisterExternalConstantNamespaceVariab
 				}
 				else if (!InVariable.IsDataAllocated())
 				{
-					FNiagaraVariable Var = Schema->PinToNiagaraVariable(InDefaultPin, true);
-					FNiagaraEditorUtilities::ResetVariableToDefaultValue(InVariable);
-					if (Var.IsDataAllocated() && Var.GetData() != nullptr)
-					{
-						InVariable.SetData(Var.GetData());
-					}
+					FillVariableWithDefaultValue(InVariable, InDefaultPin);
 				}
 
 				if (InVariable.GetAllocatedSizeInBytes() != InVariable.GetSizeInBytes())
@@ -5265,6 +5352,192 @@ bool FHlslNiagaraTranslator::ParameterMapRegisterExternalConstantNamespaceVariab
 	return false;
 }
 
+void FHlslNiagaraTranslator::FillVariableWithDefaultValue(FNiagaraVariable& InVariable, const UEdGraphPin* InDefaultPin)
+{
+	FNiagaraVariable Var = Schema->PinToNiagaraVariable(InDefaultPin, true);
+	FNiagaraEditorUtilities::ResetVariableToDefaultValue(InVariable);
+	if (Var.IsDataAllocated() && Var.GetData() != nullptr)
+	{
+		InVariable.SetData(Var.GetData());
+	}
+}
+
+void FHlslNiagaraTranslator::FillVariableWithDefaultValue(int32& OutValue, const UEdGraphPin* InDefaultPin)
+{
+	FNiagaraVariable Var = Schema->PinToNiagaraVariable(InDefaultPin, true);
+	FNiagaraVariable VarFinal = Var;
+	FNiagaraEditorUtilities::ResetVariableToDefaultValue(VarFinal); // Do this to handle non-zero defaults
+	if (Var.IsDataAllocated() && Var.GetData() != nullptr)
+	{
+		VarFinal.SetData(Var.GetData());
+	}
+
+	if (VarFinal.GetType().IsSameBaseDefinition(FNiagaraTypeDefinition::GetBoolDef()))
+	{
+		OutValue = VarFinal.GetValue<bool>();
+	}
+	else if (VarFinal.GetType().IsSameBaseDefinition(FNiagaraTypeDefinition::GetIntDef()) || VarFinal.GetType().IsEnum())
+	{
+		OutValue = VarFinal.GetValue<int32>();
+	}
+}
+
+
+void FHlslNiagaraTranslator::SetConstantByStaticVariable(int32& OutValue, const UEdGraphPin* InDefaultPin, FString* OutDebugString)
+{
+	FNiagaraEditorModule& NiagaraEditorModule = FModuleManager::GetModuleChecked<FNiagaraEditorModule>("NiagaraEditor");
+	OutValue = 0;
+	FNiagaraVariable Var = Schema->PinToNiagaraVariable(InDefaultPin, true);
+	FNiagaraVariable VarDefault = Var;
+	FNiagaraEditorUtilities::ResetVariableToDefaultValue(VarDefault);// Do this to handle non-zero defaults
+	if (VarDefault.GetType().IsStatic())
+	{
+		FNiagaraVariable VarWithValue = FNiagaraVariable(Var.GetType(), Var.GetName());
+		FString Value;
+		int32 Input = INDEX_NONE;
+		const UEdGraphPin* PinToTest = InDefaultPin;
+		if (InDefaultPin && InDefaultPin->Direction == EEdGraphPinDirection::EGPD_Input && InDefaultPin->LinkedTo.Num() != 0)
+		{
+			PinToTest = InDefaultPin->LinkedTo[0];
+		}
+		FGraphTraversalHandle PinHandle(ActiveHistoryForFunctionCalls.ActivePath);
+		PinHandle.Push(PinToTest);
+
+		const FString* ConstantPtr = CompileData->PinToConstantValues.Find(PinHandle);
+		if (ConstantPtr != nullptr)
+		{
+			TSharedPtr<INiagaraEditorTypeUtilities, ESPMode::ThreadSafe> TypeEditorUtilities = NiagaraEditorModule.GetTypeUtilities(Var.GetType());
+			if (TypeEditorUtilities.IsValid() && TypeEditorUtilities->CanHandlePinDefaults())
+			{
+				TypeEditorUtilities->SetValueFromPinDefaultString(*ConstantPtr, VarWithValue);
+			}
+		}
+
+		if (OutDebugString != nullptr)
+		{
+			*OutDebugString = PinHandle.ToString();
+		}
+
+		if (VarWithValue.IsDataAllocated())
+		{
+			if (VarWithValue.GetType().IsSameBaseDefinition(FNiagaraTypeDefinition::GetBoolDef()))
+			{
+				OutValue = VarWithValue.GetValue<bool>();
+			}
+			else if (VarWithValue.GetType().IsSameBaseDefinition(FNiagaraTypeDefinition::GetIntDef()) || VarWithValue.GetType().IsEnum())
+			{
+				OutValue = VarWithValue.GetValue<int32>();
+			}
+		}
+		else
+		{
+			Error(LOCTEXT("CouldNotResolveStaticVarByPin", "Could not resolve static variable through pin."), Cast<UNiagaraNode>(InDefaultPin->GetOwningNode()), InDefaultPin);
+		}
+	}
+}
+
+void FHlslNiagaraTranslator::SetConstantByStaticVariable(FNiagaraVariable& OutValue, const UEdGraphPin* InDefaultPin, FString* DebugString )
+{
+	OutValue.AllocateData();
+	int32 Constant = 0;
+	SetConstantByStaticVariable(Constant, InDefaultPin, DebugString);
+
+	if (OutValue.GetType().IsSameBaseDefinition(FNiagaraTypeDefinition::GetBoolDef()))
+	{
+		OutValue.SetValue<bool>(Constant != 0 ? true : false);
+	}
+	else if (OutValue.GetType().IsSameBaseDefinition(FNiagaraTypeDefinition::GetIntDef()) || OutValue.GetType().IsEnum())
+	{
+		OutValue.SetValue<int32>(Constant);
+	}
+}
+
+void FHlslNiagaraTranslator::SetConstantByStaticVariable(FNiagaraVariable& OutValue, const FNiagaraVariable& Var, FString* DebugString)
+{
+	OutValue = Var;
+	OutValue.AllocateData();
+	int32 Constant = 0;
+	SetConstantByStaticVariable(Constant, Var, DebugString);
+
+	if (Var.GetType().IsSameBaseDefinition(FNiagaraTypeDefinition::GetBoolDef()))
+	{
+		OutValue.SetValue<bool>(Constant != 0 ? true : false);
+	}
+	else if (Var.GetType().IsSameBaseDefinition(FNiagaraTypeDefinition::GetIntDef()) || Var.GetType().IsEnum())
+	{
+		OutValue.SetValue<int32>(Constant);
+	}
+}
+
+void FHlslNiagaraTranslator::SetConstantByStaticVariable(int32& OutValue, const FNiagaraVariable& Var, FString* DebugString)
+{
+	OutValue = 0;
+	FNiagaraVariable VarDefault = Var;
+	FNiagaraEditorUtilities::ResetVariableToDefaultValue(VarDefault);// Do this to handle non-zero defaults
+	if (VarDefault.GetType().IsStatic())
+	{
+		if (DebugString != nullptr)
+		{
+			*DebugString = Var.GetName().ToString();
+		}
+		FNiagaraVariable VarWithValue = FNiagaraVariable(Var.GetType(), Var.GetName());
+		FString Value = Var.GetName().ToString();
+		int32 Input = INDEX_NONE;
+
+		// If we found a string, we should try and map to the actual value of that variable..
+		if (Value.Len() != 0)
+		{
+			if (!VarWithValue.IsDataAllocated())
+			{
+				int32 FoundOverrideIdx = CompileData->StaticVariables.IndexOfByPredicate([&](const FNiagaraVariable& InObj) -> bool
+					{
+						return (InObj.GetName() == *Value);
+					});;
+
+				if (FoundOverrideIdx != INDEX_NONE)
+				{
+					VarWithValue.SetData(CompileData->StaticVariables[FoundOverrideIdx].GetData());
+				}
+				else if (FNiagaraParameterMapHistory::IsAliasedEmitterParameter(Var))
+				{
+					FNiagaraAliasContext ResolveAliasesContext(FNiagaraAliasContext::ERapidIterationParameterMode::EmitterOrParticleScript);
+					ResolveAliasesContext.ChangeEmitterToEmitterName(CompileData->GetUniqueEmitterName());
+					FNiagaraVariable TestEmitterResolvedVar = FNiagaraUtilities::ResolveAliases(Var, ResolveAliasesContext);
+
+					FoundOverrideIdx = CompileData->StaticVariables.IndexOfByPredicate([&](const FNiagaraVariable& InObj) -> bool
+						{
+							return (InObj.GetName() == TestEmitterResolvedVar.GetName());
+						});;
+
+					if (FoundOverrideIdx != INDEX_NONE)
+					{
+						VarWithValue.SetData(CompileData->StaticVariables[FoundOverrideIdx].GetData());
+					}
+				}
+			}
+		}
+
+		if (VarWithValue.IsDataAllocated())
+		{
+			if (VarWithValue.GetType().IsSameBaseDefinition(FNiagaraTypeDefinition::GetBoolDef()))
+			{
+				OutValue = VarWithValue.GetValue<bool>();
+			}
+			else if (VarWithValue.GetType().IsSameBaseDefinition(FNiagaraTypeDefinition::GetIntDef()) || VarWithValue.GetType().IsEnum())
+			{
+				OutValue = VarWithValue.GetValue<int32>();
+			}
+
+			//UE_LOG(LogNiagaraEditor, Log, TEXT("SetConstantByStaticVariable Added %s (%d)"), *Var.GetName().ToString(), OutValue);
+		}
+		else
+		{
+			Error(FText::Format(LOCTEXT("CouldNotResolveStaticVar", "Could not resolve static variable \"{0}\". Default type value used instead."),  FText::FromName(Var.GetName())), nullptr, nullptr);
+		}
+	}
+}
+
+
 
 bool FHlslNiagaraTranslator::ParameterMapRegisterUniformAttributeVariable(const FNiagaraVariable& InVariable, UNiagaraNode* InNode, int32 InParamMapHistoryIdx, int32& Output)
 {
@@ -5511,11 +5784,36 @@ void FHlslNiagaraTranslator::ParameterMapGet(UNiagaraNodeParameterMapGet* GetNod
 			FNiagaraVariable Var = Schema->PinToNiagaraVariable(OutputPins[i], bNeedsValue);
 
 			UNiagaraScriptVariable* Variable = GetNode->GetNiagaraGraph()->GetScriptVariable(Var);
-			HandleParameterRead(ParamMapHistoryIdx, Var, GetNode->GetDefaultPin(OutputPins[i]), GetNode, Outputs[i], Variable);
+			if (Var.GetType().IsStatic())
+			{
+				if (FNiagaraParameterMapHistory::IsExternalConstantNamespace(Var, CompileOptions.TargetUsage, CompileOptions.GetTargetUsageBitmask()))
+				{
+					if (Variable && Variable->DefaultMode == ENiagaraDefaultMode::FailIfPreviouslyNotSet && !Variable->Variable.IsInNameSpace(FNiagaraConstants::UserNamespaceString))
+					{
+						RegisterCompileDependency(Var, FText::Format(LOCTEXT("UsedBeforeSet", "Variable {0} was read before being set. It's default mode is \"Fail If Previously Not Set\", so this isn't allowed."), FText::FromName(Var.GetName())), GetNode, OutputPins[i], true);
+					}
+				}
+
+				Outputs[i] = MakeStaticVariableDirect(OutputPins[i]);
+			}
+			else
+			{
+				HandleParameterRead(ParamMapHistoryIdx, Var, GetNode->GetDefaultPin(OutputPins[i]), GetNode, Outputs[i], Variable);
+			}
 		}
 	}
 }
 
+int32 FHlslNiagaraTranslator::MakeStaticVariableDirect(const UEdGraphPin* InDefaultPin)
+{
+	int32 Constant = INDEX_NONE;
+	FString DebugStr;
+	SetConstantByStaticVariable(Constant, InDefaultPin, &DebugStr);
+
+	AddBodyComment(TEXT("//SetConstantByStaticVariable \"") + DebugStr + TEXT("\""));
+	return GetConstantDirect(Constant);
+}
+
 void FHlslNiagaraTranslator::HandleParameterRead(int32 ParamMapHistoryIdx, const FNiagaraVariable& InVar, const UEdGraphPin* DefaultPin, UNiagaraNode* ErrorNode, int32& OutputChunkId, UNiagaraScriptVariable* Variable, bool bTreatAsUnknownParameterMap, bool bIgnoreDefaultSetFirst)
 {
 	FString ParameterMapInstanceName = GetParameterMapInstanceName(ParamMapHistoryIdx);
@@ -5558,16 +5856,22 @@ void FHlslNiagaraTranslator::HandleParameterRead(int32 ParamMapHistoryIdx, const
 		return;
 	}
 
+	
 	if (FNiagaraParameterMapHistory::IsExternalConstantNamespace(Var, CompileOptions.TargetUsage, CompileOptions.GetTargetUsageBitmask()))
 	{
 		if (Variable && Variable->DefaultMode == ENiagaraDefaultMode::FailIfPreviouslyNotSet && !Variable->Variable.IsInNameSpace(FNiagaraConstants::UserNamespaceString) && !bIgnoreDefaultSetFirst)
 		{
 			RegisterCompileDependency(Var, FText::Format(LOCTEXT("UsedBeforeSet", "Variable {0} was read before being set. It's default mode is \"Fail If Previously Not Set\", so this isn't allowed."), FText::FromName(Var.GetName())), ErrorNode, nullptr, true);
 		}
-		if (ParameterMapRegisterExternalConstantNamespaceVariable(Var, ErrorNode, ParamMapHistoryIdx, OutputChunkId, DefaultPin))
+		if (Var.GetType().IsStatic())
 		{
+			OutputChunkId = MakeStaticVariableDirect(DefaultPin);
 			return;
 		}
+		else if (ParameterMapRegisterExternalConstantNamespaceVariable(Var, ErrorNode, ParamMapHistoryIdx, OutputChunkId, DefaultPin))
+		{
+			return;
+		}			
 	}
 	else if (FNiagaraParameterMapHistory::IsAliasedModuleParameter(Var) && ActiveHistoryForFunctionCalls.InTopLevelFunctionCall(CompileOptions.TargetUsage))
 	{
@@ -5654,6 +5958,7 @@ void FHlslNiagaraTranslator::HandleParameterRead(int32 ParamMapHistoryIdx, const
 		if (!bSetPreviously && !UNiagaraScript::IsStandaloneScript(CompileOptions.TargetUsage))
 			bFailIfPreviouslyNotSetSentinel = true;
 	}
+
 	
 	int32 LastSetChunkIdx = INDEX_NONE;
 	if (ParamMapHistoryIdx < ParamMapHistories.Num())
@@ -5665,6 +5970,13 @@ void FHlslNiagaraTranslator::HandleParameterRead(int32 ParamMapHistoryIdx, const
 			LastSetChunkIdx = ParamMapSetVariablesToChunks[ParamMapHistoryIdx][VarIdx];
 		}
 
+		bool bIsStaticConstant = Var.GetType().IsStatic();
+		if (LastSetChunkIdx == INDEX_NONE && bIsStaticConstant)
+		{
+			OutputChunkId = MakeStaticVariableDirect(DefaultPin);
+			return;
+		}
+
 		// Check to see if this is the first time we've encountered this node and it is a viable candidate for rapid iteration
 		if (LastSetChunkIdx == INDEX_NONE && bIsCandidateForRapidIteration)
 		{
@@ -5681,6 +5993,7 @@ void FHlslNiagaraTranslator::HandleParameterRead(int32 ParamMapHistoryIdx, const
 				bVarChanged = true;
 			}
 
+			
 			if (TranslationOptions.bParameterRapidIteration)
 			{
 				// Now try to look up with the new name.. we may have already made this an external variable before..
@@ -5703,11 +6016,11 @@ void FHlslNiagaraTranslator::HandleParameterRead(int32 ParamMapHistoryIdx, const
 						ParamMapSetVariablesToChunks[ParamMapHistoryIdx][VarIdx] = LastSetChunkIdx;
 						RecordParamMapDefinedAttributeToNamespaceVar(Var, DefaultPin);
 					}
-					return;		
+				return;
 				}
 			}
 			else
-			{
+			{		
 				int32 FoundIdx = TranslationOptions.OverrideModuleConstants.Find(Var);
 				if (FoundIdx == INDEX_NONE)
 				{
@@ -6193,6 +6506,80 @@ void FHlslNiagaraTranslator::Operation(class UNiagaraNodeOp* Operation, TArray<i
 			return Pin->bOrphanedPin == true;
 		});
 
+	TArray<UEdGraphPin*> InputPins;
+	Operation->GetInputPins(InputPins);
+
+	bool bAllPinsStatic = true;
+	{
+		for (int32 InputIdx = 0; InputIdx < InputPins.Num(); InputIdx++)
+		{
+			if (Operation->IsAddPin(InputPins[InputIdx]))
+				continue;
+			FNiagaraTypeDefinition InputType = Schema->PinToTypeDefinition(InputPins[InputIdx]);
+			if (!InputType.IsStatic())
+				bAllPinsStatic = false;
+		}
+
+		for (int32 OutputIdx = 0; OutputIdx < OutputPins.Num(); OutputIdx++)
+		{
+			if (Operation->IsAddPin(OutputPins[OutputIdx]))
+				continue;
+			FNiagaraTypeDefinition OutputType = Schema->PinToTypeDefinition(OutputPins[OutputIdx]);
+			if (!OutputType.IsStatic())
+				bAllPinsStatic = false;
+		}
+	}
+
+
+	if (OpInfo && OpInfo->StaticVariableResolveFunction.IsBound() && bAllPinsStatic)
+	{
+		if (OpInfo->Outputs.Num() != 1 || OutputPins.Num() != OpInfo->Outputs.Num())
+		{
+			FText PinNameText = OutputPins[0]->PinFriendlyName.IsEmpty() ? FText::FromName(OutputPins[0]->PinName) : OutputPins[0]->PinFriendlyName;
+			Error(LOCTEXT("InvalidOutputPinCount", "Only one output pin is supported for static variables"), Operation, OutputPins[0]);
+			Outputs.Add(INDEX_NONE);
+			return;
+		}
+
+		FNiagaraTypeDefinition OutputType = Schema->PinToTypeDefinition(OutputPins[0]);
+		if(!OutputType.IsStatic())
+		{
+			Error(LOCTEXT("InvalidOutputPinType", "Only static types are supported for this operation!"), Operation, OutputPins[0]);
+			Outputs.Add(INDEX_NONE);
+			return;
+		}
+		
+		int32 NumVars = 0;
+		
+		for (int32 InputIdx = 0; InputIdx < InputPins.Num(); InputIdx++)
+		{
+			if (Operation->IsAddPin(InputPins[InputIdx]))
+				continue;			
+
+			FNiagaraTypeDefinition InputType = Schema->PinToTypeDefinition(InputPins[InputIdx]);
+			if (!InputType.IsStatic())
+			{
+				Error(LOCTEXT("InvalidInputPinType", "Only static types are supported for this operation!"), Operation, InputPins[InputIdx]);
+				Outputs.Add(INDEX_NONE);
+				return;
+			}
+			NumVars++;
+		}
+
+		const FNiagaraOpInOutInfo& IOInfo = OpInfo->Outputs[0];
+
+		if (NumVars > 0)
+		{
+			int32 OutputChunkId = MakeStaticVariableDirect(OutputPins[0]);
+			Outputs.Add(OutputChunkId);
+		}
+		else
+		{
+			Outputs.Add(INDEX_NONE);
+		}
+		return;
+	}
+
 	for (int32 OutputIndex = 0; OutputIndex < OutputPins.Num(); OutputIndex++)
 	{
 		UEdGraphPin* OutputPin = OutputPins[OutputIndex];
@@ -7234,7 +7621,7 @@ void FHlslNiagaraTranslator::RegisterFunctionCall(ENiagaraScriptUsage ScriptUsag
 								// for aliased values we check to see if they're used in the current context by resolving the alias and checking against the current resolved variable
 								// name since aliased values can only be resolved for reading in the correct context.
 								bool bIsValidForCurrentCallingContext = (bIsInTopLevelFunction && bIsAliased == false) || (bIsAliased && ActiveHistoryForFunctionCalls.ResolveAliases(AliasedVar).GetName() == Var.GetName());
-								if (bIsValidForCurrentCallingContext)
+								if (bIsValidForCurrentCallingContext && !Var.GetType().IsStatic() )
 								{
 									int32 LastSetChunkIdx = ParamMapSetVariablesToChunks[ActiveStageIdx][VarIdx];
 									if (LastSetChunkIdx == INDEX_NONE)
@@ -8002,7 +8389,9 @@ void FHlslNiagaraTranslator::Convert(class UNiagaraNodeConvert* Convert, TArrayV
 	for (UEdGraphPin* InputPin : InputPins)
 	{
 		if (InputPin->PinType.PinCategory == UEdGraphSchema_Niagara::PinCategoryType ||
-			InputPin->PinType.PinCategory == UEdGraphSchema_Niagara::PinCategoryEnum)
+			InputPin->PinType.PinCategory == UEdGraphSchema_Niagara::PinCategoryStaticType || 
+			InputPin->PinType.PinCategory == UEdGraphSchema_Niagara::PinCategoryEnum || 
+			InputPin->PinType.PinCategory == UEdGraphSchema_Niagara::PinCategoryStaticEnum)
 		{
 			FNiagaraTypeDefinition Type = Schema->PinToTypeDefinition(InputPin);
 			if (!AddStructToDefinitionSet(Type))
@@ -8017,7 +8406,9 @@ void FHlslNiagaraTranslator::Convert(class UNiagaraNodeConvert* Convert, TArrayV
 	for (UEdGraphPin* OutputPin : OutputPins)
 	{
 		if (OutputPin->PinType.PinCategory == UEdGraphSchema_Niagara::PinCategoryType ||
-			OutputPin->PinType.PinCategory == UEdGraphSchema_Niagara::PinCategoryEnum)
+			OutputPin->PinType.PinCategory == UEdGraphSchema_Niagara::PinCategoryStaticType ||
+			OutputPin->PinType.PinCategory == UEdGraphSchema_Niagara::PinCategoryEnum ||
+			OutputPin->PinType.PinCategory == UEdGraphSchema_Niagara::PinCategoryStaticEnum)
 		{
 			FNiagaraTypeDefinition Type = Schema->PinToTypeDefinition(OutputPin);
 			if (!AddStructToDefinitionSet(Type))
@@ -8177,28 +8568,10 @@ void FHlslNiagaraTranslator::Select(UNiagaraNodeSelect* SelectNode, int32 Select
 	Outputs.Add(INDEX_NONE);
 }
 
-void FHlslNiagaraTranslator::WriteCompilerTag(int32 InputCompileResult, const UEdGraphPin* Pin, bool bEmitMessageOnFailure, FNiagaraCompileEventSeverity FailureSeverity, const FString& Prefix)
+void FHlslNiagaraTranslator::FindConstantValue(int32 InputCompileResult, const FNiagaraTypeDefinition& TypeDef, FString& Value, FNiagaraVariable& Variable)
 {
-	FString Value;
-	FNiagaraTypeDefinition TypeDef = Schema->PinToTypeDefinition(Pin);
-	FNiagaraVariable Variable(TypeDef, Prefix.Len() ? *(Prefix + TEXT(".") + Pin->GetName()) : *Pin->GetName());
 	FNiagaraEditorModule& NiagaraEditorModule = FModuleManager::GetModuleChecked<FNiagaraEditorModule>("NiagaraEditor");
-
-
-	//If we're in an emitter script then the tag needs to be made per emitter with EmitterName.Tag
-	bool bIsSystemOrEmitterScript =
-		UNiagaraScript::IsEmitterSpawnScript(CompileOptions.TargetUsage) || UNiagaraScript::IsSystemSpawnScript(CompileOptions.TargetUsage) ||
-		UNiagaraScript::IsEmitterUpdateScript(CompileOptions.TargetUsage) || UNiagaraScript::IsSystemUpdateScript(CompileOptions.TargetUsage);
-
-	if (bIsSystemOrEmitterScript)
 	{
-		if (const FString* EmitterAliasStr = ActiveHistoryForFunctionCalls.GetEmitterAlias())
-		{
-			Variable.SetName(*(*EmitterAliasStr + TEXT(".") + Variable.GetName().ToString()));
-		}
-	}
-
-	{		
 		bool bSearch = true;
 		FString SourceName;
 		while (bSearch)
@@ -8228,7 +8601,7 @@ void FHlslNiagaraTranslator::WriteCompilerTag(int32 InputCompileResult, const UE
 						{
 							TSharedPtr<INiagaraEditorTypeUtilities, ESPMode::ThreadSafe> TypeEditorUtilities = NiagaraEditorModule.GetTypeUtilities(TypeDef);
 							if (TypeEditorUtilities.IsValid() && TypeEditorUtilities->CanHandlePinDefaults() && CodeChunks[InputCompileResult].Definition.IsEmpty() == false)
-							{							
+							{
 								// Note that this might fail due to string not being properly formatted for the type. If so, we just take the definition string altogether.
 								bool bHasValue = TypeEditorUtilities->SetValueFromPinDefaultString(CodeChunks[InputCompileResult].Definition, Variable);
 								if (bHasValue == false)
@@ -8246,7 +8619,7 @@ void FHlslNiagaraTranslator::WriteCompilerTag(int32 InputCompileResult, const UE
 					else  if (CodeChunks[InputCompileResult].Mode == ENiagaraCodeChunkMode::Uniform)
 					{
 						//UE_LOG(LogNiagaraEditor, Log, TEXT("CompilerTag %s looking for uniform with symbol %s"), *Variable.GetName().ToString(), *CodeChunks[InputCompileResult].SymbolName);
-						FString VarSymbol;						
+						FString VarSymbol;
 						for (const FNiagaraVariable& UniformVar : CompilationOutput.ScriptData.Parameters.Parameters)
 						{
 							VarSymbol.Reset();
@@ -8279,16 +8652,55 @@ void FHlslNiagaraTranslator::WriteCompilerTag(int32 InputCompileResult, const UE
 						}
 						else if (CodeChunks[InputCompileResult].SourceChunks.Num() == 0) // Search through parent chunks for a name match
 						{
-							for (int32 i = InputCompileResult - 1; i >= 0 && i < InputCompileResult; i--)
+							bool bFoundAlternate = false;
+
+							// First see if this is output from a function call variable, if so we need to check to see when it was last written and what the chunk was that happened in.
+							int32 ParamMapHistoryIdx = ActiveStageIdx; // Should match the order in Translation stages! 
+							if (ParamMapHistories.IsValidIndex(ParamMapHistoryIdx))
 							{
-								if (CodeChunks[i].SymbolName == SourceName)
+								TArray<FString> SplitName;
+								CodeChunks[InputCompileResult].SymbolName.ParseIntoArray(SplitName, TEXT("."));
+								FString NewName;
+
+								if (SplitName.Num() > 2 && SplitName[0] == TEXT("Context") &&  SplitName[1] == TranslationStages[ParamMapHistoryIdx].PassNamespace)
 								{
-									InputCompileResult = i;
-									break;
+									for (int32 SplitIdx = 2; SplitIdx < SplitName.Num(); SplitIdx++)
+									{
+										if (NewName.Len() > 0)
+											NewName += TEXT(".") + SplitName[SplitIdx];
+										else
+											NewName += SplitName[SplitIdx];
+									}
 								}
-								if (i == 0)
+
+								if (NewName.Len() > 0)
 								{
-									bSearch = false;
+									int32 VarIdx = ParamMapHistories[ParamMapHistoryIdx].FindVariableByName(*NewName);
+									if (VarIdx != INDEX_NONE && VarIdx < ParamMapSetVariablesToChunks[ParamMapHistoryIdx].Num())
+									{
+										int32 PossibleIndex = ParamMapSetVariablesToChunks[ParamMapHistoryIdx][VarIdx];
+										if (PossibleIndex < InputCompileResult - 1)
+										{
+											InputCompileResult = PossibleIndex;
+											bFoundAlternate = true;
+										}
+									}
+								}
+							}
+
+							if (!bFoundAlternate)
+							{
+								for (int32 i = InputCompileResult - 1; i >= 0 && i < InputCompileResult; i--)
+								{
+									if (CodeChunks[i].SymbolName == SourceName)
+									{
+										InputCompileResult = i;
+										break;
+									}
+									if (i == 0)
+									{
+										bSearch = false;
+									}
 								}
 							}
 						}
@@ -8309,6 +8721,29 @@ void FHlslNiagaraTranslator::WriteCompilerTag(int32 InputCompileResult, const UE
 			}
 		}
 	}
+}
+
+void FHlslNiagaraTranslator::WriteCompilerTag(int32 InputCompileResult, const UEdGraphPin* Pin, bool bEmitMessageOnFailure, FNiagaraCompileEventSeverity FailureSeverity, const FString& Prefix)
+{
+	FString Value;
+	FNiagaraTypeDefinition TypeDef = Schema->PinToTypeDefinition(Pin);
+	FNiagaraVariable Variable(TypeDef, Prefix.Len() ? *(Prefix + TEXT(".") + Pin->GetName()) : *Pin->GetName());
+
+
+	//If we're in an emitter script then the tag needs to be made per emitter with EmitterName.Tag
+	bool bIsSystemOrEmitterScript =
+		UNiagaraScript::IsEmitterSpawnScript(CompileOptions.TargetUsage) || UNiagaraScript::IsSystemSpawnScript(CompileOptions.TargetUsage) ||
+		UNiagaraScript::IsEmitterUpdateScript(CompileOptions.TargetUsage) || UNiagaraScript::IsSystemUpdateScript(CompileOptions.TargetUsage);
+
+	if (bIsSystemOrEmitterScript)
+	{
+		if (const FString* EmitterAliasStr = ActiveHistoryForFunctionCalls.GetEmitterAlias())
+		{
+			Variable.SetName(*(*EmitterAliasStr + TEXT(".") + Variable.GetName().ToString()));
+		}
+	}
+
+	FindConstantValue(InputCompileResult, TypeDef, Value, Variable);
 
 	if (Value.Len() == 0 && Variable.IsDataAllocated() == false && bEmitMessageOnFailure)
 	{
@@ -8340,9 +8775,20 @@ int32 FHlslNiagaraTranslator::CompilePin(const UEdGraphPin* Pin)
 	{
 		if (Pin->LinkedTo.Num() > 0)
 		{
+			if (Pin->LinkedTo[0])
+			{
+				FPinConnectionResponse ConnectionResponse = Schema->CanCreateConnection(Pin->LinkedTo[0], Pin);
+				if (ConnectionResponse.Response == CONNECT_RESPONSE_DISALLOW)
+				{
+					FText OutErrorMessage = FText::Format(LOCTEXT("InputConnectionDisallowed", "Input connection is not allowed! Reason: {0}"),
+						ConnectionResponse.Message);
+
+					Error(OutErrorMessage, Cast<UNiagaraNode>(Pin->GetOwningNode()), Pin);
+				}
+			}
 			Ret = CompileOutputPin(Pin->LinkedTo[0]);
 		}
-		else if (!Pin->bDefaultValueIsIgnored && Pin->PinType.PinCategory == UEdGraphSchema_Niagara::PinCategoryType)
+		else if (!Pin->bDefaultValueIsIgnored && (Pin->PinType.PinCategory == UEdGraphSchema_Niagara::PinCategoryType || Pin->PinType.PinCategory == UEdGraphSchema_Niagara::PinCategoryStaticType))
 		{
 			if (TypeDef == FNiagaraTypeDefinition::GetParameterMapDef())
 			{
@@ -8356,7 +8802,7 @@ int32 FHlslNiagaraTranslator::CompilePin(const UEdGraphPin* Pin)
 				return GetConstant(PinVar);
 			}
 		}
-		else if (!Pin->bDefaultValueIsIgnored && Pin->PinType.PinCategory == UEdGraphSchema_Niagara::PinCategoryEnum)
+		else if (!Pin->bDefaultValueIsIgnored && (Pin->PinType.PinCategory == UEdGraphSchema_Niagara::PinCategoryEnum || Pin->PinType.PinCategory == UEdGraphSchema_Niagara::PinCategoryStaticEnum))
 		{
 			//No connections to this input so add the default as a const expression.			
 			FNiagaraVariable PinVar = Schema->PinToNiagaraVariable(Pin, true);
@@ -8385,17 +8831,18 @@ int32 FHlslNiagaraTranslator::CompileOutputPin(const UEdGraphPin* InPin)
 
 	// The incoming pin to compile may be pointing to a reroute node. If so, we just jump over it
 	// to where it really came from.
-	UEdGraphPin* Pin = UNiagaraNode::TraceOutputPin(const_cast<UEdGraphPin*>(InPin));
+	const UEdGraphPin* Pin = InPin;
 	if (Pin == nullptr || Pin->Direction != EGPD_Output)
 	{
 		Error(LOCTEXT("TraceOutputPinFailed", "Failed to trace pin to an output!"), Cast<UNiagaraNode>(InPin->GetOwningNode()), InPin);
 		return INDEX_NONE;
 	}
 
-	// The node can also replace our pin with another pin (e.g. in the case of static switches), so we need to make sure we don't run into a circular dependency
-	TSet<UEdGraphPin*> SeenPins;
 	UNiagaraNode* Node = Cast<UNiagaraNode>(Pin->GetOwningNode());
-	UEdGraphPin* OriginalPin = Pin;
+	const UEdGraphPin* OriginalPin = Pin;
+
+	// The node can also replace our pin with another pin (e.g. in the case of static switches), so we need to make sure we don't run into a circular dependency
+	/*TSet<UEdGraphPin*> SeenPins;
 	while (Node->SubstituteCompiledPin(this, &Pin))
 	{
 		bool bIsAlreadyInSet = false;
@@ -8406,7 +8853,7 @@ int32 FHlslNiagaraTranslator::CompileOutputPin(const UEdGraphPin* InPin)
 			Error(LOCTEXT("CircularGraphSubstitutionError", "Circular dependency detected, please check your module graph."), Node, Pin);
 			return INDEX_NONE;
 		}
-	}
+	}*/
 
 	// It is possible that the output pin was substituted by an input pin (e.g. the default value pin on a node).
 	// If that is the case we try to compile that pin directly.
@@ -8671,7 +9118,7 @@ FString FHlslNiagaraTranslator::GetHlslDefaultForType(const FNiagaraTypeDefiniti
 	{
 		return "(0)";
 	}
-	else if (Type == FNiagaraTypeDefinition::GetBoolDef())
+	else if (Type.IsSameBaseDefinition(FNiagaraTypeDefinition::GetBoolDef()))
 	{
 		return "(false)";
 	}
@@ -8692,9 +9139,9 @@ bool FHlslNiagaraTranslator::IsBuiltInHlslType(const FNiagaraTypeDefinition& Typ
 		Type == FNiagaraTypeDefinition::GetPositionDef() ||
 		Type == FNiagaraTypeDefinition::GetQuatDef() ||
 		Type == FNiagaraTypeDefinition::GetMatrix4Def() ||
-		Type == FNiagaraTypeDefinition::GetIntDef() ||
+		Type.IsSameBaseDefinition(FNiagaraTypeDefinition::GetIntDef()) ||
 		Type.GetStruct() == FNiagaraTypeDefinition::GetIntStruct() ||
-		Type == FNiagaraTypeDefinition::GetBoolDef();
+		Type.IsSameBaseDefinition(FNiagaraTypeDefinition::GetBoolDef());
 }
 
 FString FHlslNiagaraTranslator::GetStructHlslTypeName(const FNiagaraTypeDefinition& Type)
@@ -8723,11 +9170,11 @@ FString FHlslNiagaraTranslator::GetStructHlslTypeName(const FNiagaraTypeDefiniti
 	{
 		return "float4x4";
 	}
-	else if (Type == FNiagaraTypeDefinition::GetIntDef() || Type.GetEnum() != nullptr)
+	else if (Type.IsSameBaseDefinition(FNiagaraTypeDefinition::GetIntDef()) || Type.GetEnum() != nullptr)
 	{
 		return "int";
 	}
-	else if (Type == FNiagaraTypeDefinition::GetBoolDef())
+	else if (Type.IsSameBaseDefinition(FNiagaraTypeDefinition::GetBoolDef()))
 	{
 		return "bool";
 	}
diff --git a/Engine/Plugins/FX/Niagara/Source/NiagaraEditor/Private/NiagaraHlslTranslator.h b/Engine/Plugins/FX/Niagara/Source/NiagaraEditor/Private/NiagaraHlslTranslator.h
--- a/Engine/Plugins/FX/Niagara/Source/NiagaraEditor/Private/NiagaraHlslTranslator.h
+++ b/Engine/Plugins/FX/Niagara/Source/NiagaraEditor/Private/NiagaraHlslTranslator.h
@@ -70,7 +70,7 @@ public:
 	virtual FName ResolveEmitterAlias(FName VariableName) const override;
 
 	const FString& GetUniqueEmitterName() const { return EmitterUniqueName; }
-	void FinishPrecompile(const TArray<FNiagaraVariable>& EncounterableVariables, FCompileConstantResolver ConstantResolver, const TArray<ENiagaraScriptUsage>& UsagesToProcess, const TArray<class UNiagaraSimulationStageBase*>* SimStages, const TArray<FString> EmitterNames);
+	void FinishPrecompile(const TArray<FNiagaraVariable>& EncounterableVariables, const TArray<FNiagaraVariable>& InStaticVariables, FCompileConstantResolver ConstantResolver, const TArray<ENiagaraScriptUsage>& UsagesToProcess, const TArray<class UNiagaraSimulationStageBase*>* SimStages, const TArray<FString> EmitterNames);
 	virtual int32 GetDependentRequestCount() const override {
 		return EmitterData.Num();
 	};
@@ -122,6 +122,12 @@ public:
 	UEnum* ENiagaraScriptUsageEnum;
 
 	TArray<FNiagaraVariable> RapidIterationParams;
+
+	TMap<FGraphTraversalHandle, FString> PinToConstantValues;
+	TArray<FNiagaraVariable> StaticVariables;
+	TArray<FNiagaraVariable> StaticVariablesWithMultipleWrites;
+
+	void CompareAgainst(FNiagaraGraphCachedBuiltHistory* InCachedDataBase);
 };
 
 class FNiagaraCompileRequestDuplicateData : public FNiagaraCompileRequestDuplicateDataBase
@@ -143,7 +149,7 @@ public:
 	void DeepCopyGraphs(UNiagaraScriptSource* ScriptSource, ENiagaraScriptUsage InUsage, FCompileConstantResolver ConstantResolver);
 	void DeepCopyGraphs(UNiagaraEmitter* Emitter);
 	
-	void FinishPrecompileDuplicate(const TArray<FNiagaraVariable>& EncounterableVariables, FCompileConstantResolver ConstantResolver, const TArray<class UNiagaraSimulationStageBase*>* SimStages);
+	void FinishPrecompileDuplicate(const TArray<FNiagaraVariable>& EncounterableVariables, const TArray<FNiagaraVariable>& InStaticVariables, FCompileConstantResolver ConstantResolver, const TArray<class UNiagaraSimulationStageBase*>* SimStages, const TArray<FNiagaraVariable>& InParamStore);
 	void CreateDataInterfaceCDO(TArrayView<UClass*> VariableDataInterfaces);
 
 	virtual int32 GetDependentRequestCount() const override	{ return EmitterData.Num();	}
@@ -369,6 +375,7 @@ public:
 	TArray<FNiagaraVariable> SetParticleAttributes;
 	FString CustomReadFunction;
 	FString CustomWriteFunction;
+	int32 ParamMapHistoryIndex = INDEX_NONE;
 
 	bool ShouldDoSpawnOnlyLogic() const;
 	bool IsRelevantToSpawnForStage(const FNiagaraParameterMapHistory& InHistory, const FNiagaraVariable& InAliasedVar, const FNiagaraVariable& InVar) const;
@@ -421,7 +428,7 @@ protected:
 	Map of Pins to compiled code chunks. Allows easy reuse of previously compiled pins.
 	A stack so that we can track pin reuse within function calls but not have cached pins cross talk with subsequent calls to the same function.
 	*/
-	TArray<TMap<UEdGraphPin*, int32>> PinToCodeChunks;
+	TArray<TMap<const UEdGraphPin*, int32>> PinToCodeChunks;
 
 	/** The combined output of the compilation of this script. This is temporary and will be reworked soon. */
 	FNiagaraTranslatorOutput CompilationOutput;
@@ -645,6 +652,17 @@ public:
 	bool HandleBoundConstantVariableToDataSetRead(FNiagaraVariable InVariable, UNiagaraNode* InNode, int32 InParamMapHistoryIdx, int32& Output, const UEdGraphPin* InDefaultPin);
 
 	static FString GenerateFunctionHlslPrototype(FStringView InVariableName, const FNiagaraFunctionSignature& FunctionSignature);
+	void FillVariableWithDefaultValue(FNiagaraVariable& InVar, const UEdGraphPin* InDefaultPin);
+	void FillVariableWithDefaultValue(int32& OutValue, const UEdGraphPin* InDefaultPin);
+
+	void SetConstantByStaticVariable(int32& OutValue, const UEdGraphPin* InDefaultPin, FString* DebugString = nullptr);
+	void SetConstantByStaticVariable(int32& OutValue, const FNiagaraVariable& Var, FString* DebugString = nullptr);
+	void SetConstantByStaticVariable(FNiagaraVariable& OutValue, const FNiagaraVariable& Var, FString* DebugString = nullptr);
+	void SetConstantByStaticVariable(FNiagaraVariable& OutValue, const UEdGraphPin* InDefaultPin, FString* DebugString = nullptr);
+
+	int32 MakeStaticVariableDirect(const UEdGraphPin* InDefaultPin);
+
+
 private:
 	void InitializeParameterMapDefaults(int32 ParamMapHistoryIdx);
 	void HandleParameterRead(int32 ParamMapHistoryIdx, const FNiagaraVariable& Var, const UEdGraphPin* DefaultPin, UNiagaraNode* ErrorNode, int32& OutputChunkId, UNiagaraScriptVariable* Variable, bool bTreatAsUnknownParameterMap = false, bool bIgnoreDefaultSetFirst = false);
@@ -655,6 +673,8 @@ private:
 
 	void HandleNamespacedExternalVariablesToDataSetRead(TArray<FNiagaraVariable>& InDataSetVars, FString InNamespaceStr);
 
+	void FindConstantValue(int32 InputCompileResult, const FNiagaraTypeDefinition& TypeDef, FString& Value, FNiagaraVariable& Variable);
+
 	// For GPU simulations we have to special case some variables and pass them view shader parameters rather than the uniform buffer as they vary from CPU simulations
 	bool IsVariableInUniformBuffer(const FNiagaraVariable& Variable) const;
 
@@ -670,6 +690,7 @@ private:
 	// Add a raw float constant chunk
 	int32 GetConstantDirect(float InValue);
 	int32 GetConstantDirect(bool InValue);
+	int32 GetConstantDirect(int InValue);
 
 	FNiagaraTypeDefinition GetChildType(const FNiagaraTypeDefinition& BaseType, const FName& PropertyName);
 	FString NamePathToString(const FString& Prefix, const FNiagaraTypeDefinition& RootType, const TArray<FName>& NamePath);
@@ -700,7 +721,7 @@ private:
 
     // Checks that the Partices.ID parameter is only used if persistent IDs are active
     void ValidateParticleIDUsage();
-	bool ValidateTypePins(UNiagaraNode* NodeToValidate);
+	bool ValidateTypePins(const UNiagaraNode* NodeToValidate);
 	void GenerateFunctionSignature(ENiagaraScriptUsage ScriptUsage, FString InName, const FString& InFullName, const FString& InFunctionNameSuffix, UNiagaraGraph* FuncGraph, TArray<int32>& Inputs, 
 		bool bHadNumericInputs, bool bHasParameterMapParameters, TArray<UEdGraphPin*> StaticSwitchValues, FNiagaraFunctionSignature& OutSig) const;
 
diff --git a/Engine/Plugins/FX/Niagara/Source/NiagaraEditor/Private/NiagaraNode.cpp b/Engine/Plugins/FX/Niagara/Source/NiagaraEditor/Private/NiagaraNode.cpp
--- a/Engine/Plugins/FX/Niagara/Source/NiagaraEditor/Private/NiagaraNode.cpp
+++ b/Engine/Plugins/FX/Niagara/Source/NiagaraEditor/Private/NiagaraNode.cpp
@@ -311,6 +311,8 @@ bool UNiagaraNode::ReallocatePins(bool bMarkNeedsResynchronizeOnChange)
 		}
 	}
 
+	OnPostSynchronizationInReallocatePins();
+
 	//GetGraph()->NotifyGraphChanged();
 	if (bMarkNeedsResynchronizeOnChange && !bAllSame)
 	{
@@ -435,7 +437,7 @@ bool UNiagaraNode::ConvertNumericPinToType(UEdGraphPin* InGraphPin, FNiagaraType
 	const UEdGraphSchema_Niagara* Schema = CastChecked<UEdGraphSchema_Niagara>(GetSchema());
 	check(Schema);
 	FEdGraphPinType PinType = Schema->TypeDefinitionToPinType(TypeDef);
-	if (PinType == InGraphPin->PinType)
+	if (!Schema->PinTypesValidForNumericConversion(PinType, InGraphPin->PinType))
 	{
 		return false;
 	}
@@ -736,13 +738,14 @@ void UNiagaraNode::NumericResolutionByPins(const UEdGraphSchema_Niagara* Schema,
 	for (UEdGraphPin* InputPin : InputPins)
 	{
 		FNiagaraTypeDefinition InputPinType = Schema->PinToTypeDefinition(InputPin);
-		if (InputPin->PinType.PinCategory == UEdGraphSchema_Niagara::PinCategoryType)
+		if (InputPin->PinType.PinCategory == UEdGraphSchema_Niagara::PinCategoryType ||
+			InputPin->PinType.PinCategory == UEdGraphSchema_Niagara::PinCategoryStaticType)
 		{
 
 			// If the input pin is the generic numeric type set it to the type of the linked output pin which should have been processed already.
 			if (InputPinType == FNiagaraTypeDefinition::GetGenericNumericDef() && InputPin->LinkedTo.Num() == 1)
 			{
-				UEdGraphPin* InputPinLinkedPin = UNiagaraNode::TraceOutputPin(InputPin->LinkedTo[0]);
+				UEdGraphPin* InputPinLinkedPin = InputPin->LinkedTo[0];
 				FNiagaraTypeDefinition InputPinLinkedPinType = Schema->PinToTypeDefinition(InputPinLinkedPin);
 				if (InputPinLinkedPinType.IsValid())
 				{
@@ -814,12 +817,12 @@ void UNiagaraNode::NumericResolutionByPins(const UEdGraphSchema_Niagara* Schema,
 
 void UNiagaraNode::ResolveNumerics(const UEdGraphSchema_Niagara* Schema, bool bSetInline, TMap<TPair<FGuid, UEdGraphNode*>, FNiagaraTypeDefinition>* PinCache)
 {
-	// Fix up numeric input pins and keep track of numeric types to decide the output type.
-	FPinCollectorArray InputPins;
-	GetInputPins(InputPins);
-	FPinCollectorArray OutputPins;
-	GetOutputPins(OutputPins);
-	NumericResolutionByPins(Schema, InputPins, OutputPins, bSetInline, PinCache);
+// Fix up numeric input pins and keep track of numeric types to decide the output type.
+FPinCollectorArray InputPins;
+GetInputPins(InputPins);
+FPinCollectorArray OutputPins;
+GetOutputPins(OutputPins);
+NumericResolutionByPins(Schema, InputPins, OutputPins, bSetInline, PinCache);
 }
 
 void UNiagaraNode::GetWildcardPinHoverConnectionTextAddition(const UEdGraphPin* WildcardPin, const UEdGraphPin* OtherPin, ECanCreateConnectionResponse ConnectionResponse, FString& OutString) const
@@ -832,7 +835,7 @@ ENiagaraNumericOutputTypeSelectionMode UNiagaraNode::GetNumericOutputTypeSelecti
 	return ENiagaraNumericOutputTypeSelectionMode::None;
 }
 
-UEdGraphPin* UNiagaraNode::TraceOutputPin(UEdGraphPin* LocallyOwnedOutputPin, bool bFilterForCompilation)
+UEdGraphPin* UNiagaraNode::TraceOutputPin(UEdGraphPin* LocallyOwnedOutputPin, bool bFilterForCompilation, TArray<const UNiagaraNode*>* OutNodesVisitedDuringTrace)
 {
 	if (LocallyOwnedOutputPin == nullptr)
 	{
@@ -840,9 +843,19 @@ UEdGraphPin* UNiagaraNode::TraceOutputPin(UEdGraphPin* LocallyOwnedOutputPin, bo
 	}
 
 	UNiagaraNode* LinkedNode = CastChecked<UNiagaraNode>(LocallyOwnedOutputPin->GetOwningNode());
-	return LinkedNode->GetTracedOutputPin(LocallyOwnedOutputPin, bFilterForCompilation);
+	return LinkedNode->GetTracedOutputPin(LocallyOwnedOutputPin, bFilterForCompilation, OutNodesVisitedDuringTrace);
 }
 
+UEdGraphPin* UNiagaraNode::GetTracedOutputPin(UEdGraphPin* LocallyOwnedOutputPin, bool bFilterForCompilation, TArray<const UNiagaraNode*>* OutNodesVisitedDuringTrace) const
+{
+	if (OutNodesVisitedDuringTrace != nullptr)
+	{
+		OutNodesVisitedDuringTrace->Add(this);
+	}
+	return LocallyOwnedOutputPin;
+}
+
+
 bool UNiagaraNode::SubstituteCompiledPin(FHlslNiagaraTranslator* Translator, UEdGraphPin** LocallyOwnedPin)
 {
 	return false;
@@ -885,6 +898,41 @@ void UNiagaraNode::RouteParameterMapAroundMe(FNiagaraParameterMapHistoryBuilder&
 	}
 }
 
+void UNiagaraNode::RegisterPassthroughPin(FNiagaraParameterMapHistoryBuilder& OutHistory, UEdGraphPin* InputPin, UEdGraphPin* OutputPin, bool bFilterForCompilation, bool bVisitInputPin) const
+{
+	const UEdGraphSchema_Niagara* Schema = CastChecked<UEdGraphSchema_Niagara>(GetSchema());
+	if (bVisitInputPin)
+		OutHistory.VisitInputPin(InputPin, this, bFilterForCompilation);
+
+	FNiagaraTypeDefinition InDef = Schema->PinToTypeDefinition(InputPin);
+	FNiagaraTypeDefinition OutDef = Schema->PinToTypeDefinition(OutputPin);
+
+	if (InputPin && InDef == FNiagaraTypeDefinition::GetParameterMapDef() && OutDef == FNiagaraTypeDefinition::GetParameterMapDef() && InputPin->LinkedTo.Num() > 0)
+	{
+		int32 PMIdx = OutHistory.TraceParameterMapOutputPin(InputPin->LinkedTo[0]);
+		OutHistory.RegisterParameterMapPin(PMIdx, OutputPin);
+	}
+	else if (InputPin && InDef.IsStatic() && InputPin->LinkedTo.Num() > 0)
+	{
+		int32 ConstantIdx = OutHistory.GetConstantFromOutputPin(InputPin->LinkedTo[0]);
+		OutHistory.RegisterConstantPin(ConstantIdx, InputPin);
+
+		if (OutputPin && OutDef == InDef)
+		{
+			OutHistory.RegisterConstantPin(ConstantIdx, OutputPin);
+		}
+	}
+	else if (InputPin && InDef.IsStatic() && InputPin->LinkedTo.Num() == 0)
+	{
+		int32 ConstantIdx = OutHistory.AddOrGetConstantFromValue(InputPin->DefaultValue);
+		OutHistory.RegisterConstantPin(ConstantIdx, InputPin);
+		if (OutputPin && OutDef == InDef)
+		{
+			OutHistory.RegisterConstantPin(ConstantIdx, OutputPin);
+		}
+	}
+}
+
 bool UNiagaraNode::GetTooltipTextForKnownPin(const UEdGraphPin& Pin, FText& OutTooltip) const
 {
 	const UNiagaraGraph* NiagaraGraph = Cast<UNiagaraGraph>(GetGraph());
diff --git a/Engine/Plugins/FX/Niagara/Source/NiagaraEditor/Private/NiagaraNodeConvert.cpp b/Engine/Plugins/FX/Niagara/Source/NiagaraEditor/Private/NiagaraNodeConvert.cpp
--- a/Engine/Plugins/FX/Niagara/Source/NiagaraEditor/Private/NiagaraNodeConvert.cpp
+++ b/Engine/Plugins/FX/Niagara/Source/NiagaraEditor/Private/NiagaraNodeConvert.cpp
@@ -143,7 +143,9 @@ void UNiagaraNodeConvert::Compile(class FHlslNiagaraTranslator* Translator, TArr
 	for(UEdGraphPin* InputPin : InputPins)
 	{
 		if (InputPin->PinType.PinCategory == UEdGraphSchema_Niagara::PinCategoryType || 
-			InputPin->PinType.PinCategory == UEdGraphSchema_Niagara::PinCategoryEnum)
+			InputPin->PinType.PinCategory == UEdGraphSchema_Niagara::PinCategoryEnum || 
+			InputPin->PinType.PinCategory == UEdGraphSchema_Niagara::PinCategoryStaticType ||
+			InputPin->PinType.PinCategory == UEdGraphSchema_Niagara::PinCategoryStaticEnum)
 		{
 			int32 CompiledInput = Translator->CompilePin(InputPin);
 			if (CompiledInput == INDEX_NONE)
@@ -574,7 +576,7 @@ void UNiagaraNodeConvert::AutowireNewNode(UEdGraphPin* FromPin)
 						FNiagaraTypeDefinition FromPinType = UEdGraphSchema_Niagara::PinToTypeDefinition(FromPin);
 						FNiagaraTypeDefinition InputPinType = UEdGraphSchema_Niagara::PinToTypeDefinition(NewPin);
 
-						if(FNiagaraTypeDefinition::TypesAreAssignable(FromPinType, InputPinType) && GetSchema()->CanCreateConnection(FromPin, NewPin).Response != ECanCreateConnectionResponse::CONNECT_RESPONSE_DISALLOW)
+						if(FNiagaraTypeDefinition::TypesAreAssignable(InputPinType, FromPinType) && GetSchema()->CanCreateConnection(FromPin, NewPin).Response != ECanCreateConnectionResponse::CONNECT_RESPONSE_DISALLOW)
 						{
 							FromPin->MakeLinkTo(NewPin);
 							bConnectionMade = true;
diff --git a/Engine/Plugins/FX/Niagara/Source/NiagaraEditor/Private/NiagaraNodeCustomHlsl.cpp b/Engine/Plugins/FX/Niagara/Source/NiagaraEditor/Private/NiagaraNodeCustomHlsl.cpp
--- a/Engine/Plugins/FX/Niagara/Source/NiagaraEditor/Private/NiagaraNodeCustomHlsl.cpp
+++ b/Engine/Plugins/FX/Niagara/Source/NiagaraEditor/Private/NiagaraNodeCustomHlsl.cpp
@@ -449,6 +449,18 @@ bool UNiagaraNodeCustomHlsl::AllowNiagaraTypeForAddPin(const FNiagaraTypeDefinit
 		return false;
 }
 
+bool UNiagaraNodeCustomHlsl::AllowNiagaraTypeForAddPin(const FNiagaraTypeDefinition& InType, EEdGraphPinDirection InDirection) const
+{
+	if (AllowNiagaraTypeForAddPin(InType))
+	{
+		if (InType.IsStatic() && InDirection == EEdGraphPinDirection::EGPD_Output)
+			return false;
+		else
+			return true;
+	}
+	return false;
+}
+
 bool UNiagaraNodeCustomHlsl::ReferencesVariable(const FNiagaraVariableBase& InVar) const
 {
 	// for now we'll just do a text search through the non-comment code strings to see if we can find
diff --git a/Engine/Plugins/FX/Niagara/Source/NiagaraEditor/Private/NiagaraNodeEmitter.cpp b/Engine/Plugins/FX/Niagara/Source/NiagaraEditor/Private/NiagaraNodeEmitter.cpp
--- a/Engine/Plugins/FX/Niagara/Source/NiagaraEditor/Private/NiagaraNodeEmitter.cpp
+++ b/Engine/Plugins/FX/Niagara/Source/NiagaraEditor/Private/NiagaraNodeEmitter.cpp
@@ -363,6 +363,11 @@ void UNiagaraNodeEmitter::BuildParameterMapHistory(FNiagaraParameterMapHistoryBu
 			ChildBuilder.ConstantResolver = OutHistory.ConstantResolver;
 			ChildBuilder.RegisterEncounterableVariables(OutHistory.GetEncounterableVariables());
 			ChildBuilder.EnableScriptAllowList(true, GetUsage());
+
+			TArray<FNiagaraVariable> LocalStaticVars;
+			FNiagaraParameterUtilities::FilterToRelevantStaticVariables(OutHistory.StaticVariables, LocalStaticVars, *EmitterUniqueName, TEXT("Emitter"), true);
+			ChildBuilder.RegisterExternalStaticVariables(LocalStaticVars);
+
 			FString LocalEmitterName = TEXT("Emitter");
 			ChildBuilder.EnterEmitter(LocalEmitterName, Graph, this);
 			for (UNiagaraNodeOutput* OutputNode : OutputNodes)
@@ -392,10 +397,16 @@ void UNiagaraNodeEmitter::BuildParameterMapHistory(FNiagaraParameterMapHistoryBu
 					OutHistory.Histories[ParamMapIdx].PerVariableReadHistory[ExistingIdx].Append(History.PerVariableReadHistory[SrcVarIdx]);
 					OutHistory.Histories[ParamMapIdx].PerVariableWriteHistory[ExistingIdx].Append(History.PerVariableWriteHistory[SrcVarIdx]);
 					OutHistory.Histories[ParamMapIdx].PerVariableWarnings[ExistingIdx].Append(History.PerVariableWarnings[SrcVarIdx]);	
+					for (int32 PerConstantIdx = 0; PerConstantIdx < History.PerVariableConstantValue[SrcVarIdx].Num(); PerConstantIdx++)
+					{	
+						const FString& ConstantStr = History.PerVariableConstantValue[SrcVarIdx][PerConstantIdx];
+						OutHistory.Histories[ParamMapIdx].PerVariableConstantValue[ExistingIdx].AddUnique(ConstantStr);	
+					}
 				}
 				OutHistory.Histories[ParamMapIdx].ParameterCollections.Append(History.ParameterCollections);
 				OutHistory.Histories[ParamMapIdx].ParameterCollectionNamespaces.Append(History.ParameterCollectionNamespaces);
 				OutHistory.Histories[ParamMapIdx].ParameterCollectionVariables.Append(History.ParameterCollectionVariables);
+				OutHistory.Histories[ParamMapIdx].PinToConstantValues.Append(History.PinToConstantValues);
 			}
 		}
 
diff --git a/Engine/Plugins/FX/Niagara/Source/NiagaraEditor/Private/NiagaraNodeFunctionCall.cpp b/Engine/Plugins/FX/Niagara/Source/NiagaraEditor/Private/NiagaraNodeFunctionCall.cpp
--- a/Engine/Plugins/FX/Niagara/Source/NiagaraEditor/Private/NiagaraNodeFunctionCall.cpp
+++ b/Engine/Plugins/FX/Niagara/Source/NiagaraEditor/Private/NiagaraNodeFunctionCall.cpp
@@ -1272,6 +1272,8 @@ void UNiagaraNodeFunctionCall::BuildParameterMapHistory(FNiagaraParameterMapHist
 		// Looks like this function call is not yet hooked up. Skip it to prevent cascading errors in the compilation
 		return;
 	}
+	UNiagaraScriptSource* Source = GetFunctionScriptSource();
+
 	
 	Super::BuildParameterMapHistory(OutHistory, bRecursive, bFilterForCompilation);
 	if (!IsNodeEnabled() && OutHistory.GetIgnoreDisabled())
@@ -1281,7 +1283,6 @@ void UNiagaraNodeFunctionCall::BuildParameterMapHistory(FNiagaraParameterMapHist
 	}
 
 	
-	UNiagaraScriptSource* Source = GetFunctionScriptSource();
 	if (Source)
 	{
 		UNiagaraGraph* FunctionGraph = CastChecked<UNiagaraGraph>(Source->NodeGraph);
@@ -1328,6 +1329,9 @@ void UNiagaraNodeFunctionCall::BuildParameterMapHistory(FNiagaraParameterMapHist
 		GetOutputPins(OutputPins);
 
 		TArray<TPair<UEdGraphPin*, int32>, TInlineAllocator<16> > MatchedPairs;
+		TArray<TPair<UEdGraphPin*, int32>, TInlineAllocator<16> > MatchedConstants;
+		TArray<bool, TInlineAllocator<16> > OutputMatched;
+		OutputMatched.AddDefaulted(OutputPins.Num());
 
 		// Find the matches of names and types of the sub-graph output pins and this function call nodes' outputs.
 		for (UEdGraphPin* ChildOutputNodePin : OutputNode->GetAllPins())
@@ -1343,8 +1347,24 @@ void UNiagaraNodeFunctionCall::BuildParameterMapHistory(FNiagaraParameterMapHist
 					{
 						TPair<UEdGraphPin*, int32> Pair;
 						Pair.Key = OutputPins[i];
-						Pair.Value = OutHistory.TraceParameterMapOutputPin(UNiagaraNode::TraceOutputPin(ChildOutputNodePin->LinkedTo[0]));
+						Pair.Value = OutHistory.TraceParameterMapOutputPin(ChildOutputNodePin->LinkedTo[0]);
 						MatchedPairs.Add(Pair);
+						OutputMatched[i] = true;
+					}
+				}
+			}
+			else if (ChildOutputNodePin && VarChild.GetType().IsStatic())
+			{
+				for (int32 i = 0; i < OutputPins.Num(); i++)
+				{
+					FNiagaraVariable OutputVar = Schema->PinToNiagaraVariable(OutputPins[i]);
+					if (OutputVar.IsEquivalent(VarChild) && OutputMatched[i] == false)
+					{
+						TPair<UEdGraphPin*, int32> Pair;
+						Pair.Key = OutputPins[i];
+						Pair.Value = OutHistory.GetConstantFromInputPin(ChildOutputNodePin);
+						MatchedConstants.Add(Pair);
+						OutputMatched[i] = true;
 					}
 				}
 			}
@@ -1361,6 +1381,11 @@ void UNiagaraNodeFunctionCall::BuildParameterMapHistory(FNiagaraParameterMapHist
 		{
 			OutHistory.RegisterParameterMapPin(MatchedPairs[i].Value, MatchedPairs[i].Key);
 		}
+
+		for (int32 i = 0; i < MatchedConstants.Num(); i++)
+		{
+			OutHistory.RegisterConstantPin(MatchedConstants[i].Value, MatchedConstants[i].Key);
+		}
 	}
 	else if (!ScriptIsValid() || Signature.bRequiresExecPin)
 	{
diff --git a/Engine/Plugins/FX/Niagara/Source/NiagaraEditor/Private/NiagaraNodeInput.cpp b/Engine/Plugins/FX/Niagara/Source/NiagaraEditor/Private/NiagaraNodeInput.cpp
--- a/Engine/Plugins/FX/Niagara/Source/NiagaraEditor/Private/NiagaraNodeInput.cpp
+++ b/Engine/Plugins/FX/Niagara/Source/NiagaraEditor/Private/NiagaraNodeInput.cpp
@@ -105,6 +105,15 @@ void UNiagaraNodeInput::BuildParameterMapHistory(FNiagaraParameterMapHistoryBuil
 			OutHistory.RegisterParameterMapPin(ParamMapIdx, GetOutputPin(0));
 		}
 	}
+	else if (Input.GetType().IsStatic())
+	{
+		int32 ConstantIdx = OutHistory.FindMatchingStaticFromContextInputs(Input);
+
+		if (ConstantIdx != INDEX_NONE)
+		{
+			OutHistory.RegisterConstantPin(ConstantIdx, GetOutputPin(0));
+		}
+	}
 }
 
 void UNiagaraNodeInput::AppendFunctionAliasForContext(const FNiagaraGraphFunctionAliasContext& InFunctionAliasContext, FString& InOutFunctionAlias, bool& OutOnlyOncePerNodeType)
diff --git a/Engine/Plugins/FX/Niagara/Source/NiagaraEditor/Private/NiagaraNodeOp.cpp b/Engine/Plugins/FX/Niagara/Source/NiagaraEditor/Private/NiagaraNodeOp.cpp
--- a/Engine/Plugins/FX/Niagara/Source/NiagaraEditor/Private/NiagaraNodeOp.cpp
+++ b/Engine/Plugins/FX/Niagara/Source/NiagaraEditor/Private/NiagaraNodeOp.cpp
@@ -10,7 +10,7 @@
 #define LOCTEXT_NAMESPACE "NiagaraNodeOp"
 
 UNiagaraNodeOp::UNiagaraNodeOp(const FObjectInitializer& ObjectInitializer)
-	: Super(ObjectInitializer)
+	: Super(ObjectInitializer), bAllStatic(false)
 {
 }
 
@@ -62,6 +62,7 @@ void UNiagaraNodeOp::AllocateDefaultPins()
 	{
 		CreateAddPin(EGPD_Input);
 	}
+
 }
 
 FText UNiagaraNodeOp::GetNodeTitle(ENodeTitleType::Type TitleType) const
@@ -72,6 +73,14 @@ FText UNiagaraNodeOp::GetNodeTitle(ENodeTitleType::Type TitleType) const
 		FString Strn = "Unknown";
 		return FText::FromString(Strn);
 	}
+
+	if (OpInfo && OpInfo->StaticVariableResolveFunction.IsBound() && OpInfo->bSupportsStaticResolution && bAllStatic)
+	{
+		FFormatNamedArguments Args;
+		Args.Add(TEXT("OpName"), OpInfo->FriendlyName);
+		FText Format = LOCTEXT("OpNodeNameStatic", "{OpName} (Static)");
+		return FText::Format(Format, Args);
+	}
 	return OpInfo->FriendlyName;
 }
 
@@ -83,6 +92,14 @@ FText UNiagaraNodeOp::GetTooltipText()const
 		FString Strn = "Unknown";
 		return FText::FromString(Strn);
 	}
+
+	if (OpInfo && OpInfo->StaticVariableResolveFunction.IsBound() && OpInfo->bSupportsStaticResolution && bAllStatic)
+	{
+		FFormatNamedArguments Args;
+		Args.Add(TEXT("OpDesc"), OpInfo->Description);
+		FText Format = LOCTEXT("OpNodeDescStatic", "{OpDesc}\r\n(To convert into a static version, connect all input pins to other static pins.)");
+		return FText::Format(Format, Args);
+	}
 	return OpInfo->Description;
 }
 
@@ -91,6 +108,30 @@ FLinearColor UNiagaraNodeOp::GetNodeTitleColor() const
 	return GetDefault<UGraphEditorSettings>()->FunctionCallNodeTitleColor;
 }
 
+void UNiagaraNodeOp::OnPostSynchronizationInReallocatePins() 
+{
+	FPinCollectorArray InputPins;
+	GetInputPins(InputPins);
+
+	bool bOutputPinsNeedUpdate = false;
+	for (int32 i = 0; i < InputPins.Num(); ++i)
+	{
+		UEdGraphPin* Pin = InputPins[i];
+
+		FEdGraphPinType StartType = Pin->PinType;
+		HandleStaticInputPinUpgrade(Pin);
+
+		// Type was changed... now make sure that we can fix up the output pin.
+		if (StartType != Pin->PinType)
+		{
+			bOutputPinsNeedUpdate = true;
+		}
+	}
+
+	HandleStaticOutputPinUpgrade();		
+}
+
+
 void UNiagaraNodeOp::Compile(class FHlslNiagaraTranslator* Translator, TArray<int32>& Outputs)
 {
 	const FNiagaraOpInfo* OpInfo = FNiagaraOpInfo::GetOpInfo(OpName);
@@ -208,6 +249,7 @@ void UNiagaraNodeOp::PostLoad()
 	}
 }
 
+
 bool UNiagaraNodeOp::RefreshFromExternalChanges()
 {
 	// TODO - Leverage code in reallocate pins to determine if any pins have changed...
@@ -258,7 +300,15 @@ bool UNiagaraNodeOp::AllowNiagaraTypeForAddPin(const FNiagaraTypeDefinition& InT
 	{
 		return true;
 	}
-	return OpInfo->AddedInputTypeRestrictions.Contains(InType);
+
+	auto FindPredicate = [=](const FNiagaraTypeDefinition& PinType)
+	{
+		if (bAllStatic)
+			return PinType.ToStaticDef() == InType;
+		else 
+			return PinType == InType;
+	};
+	return OpInfo->AddedInputTypeRestrictions.IndexOfByPredicate(FindPredicate) != INDEX_NONE;
 }
 
 void UNiagaraNodeOp::OnPinRemoved(UEdGraphPin* PinToRemove)
@@ -286,6 +336,14 @@ void UNiagaraNodeOp::OnNewTypedPinAdded(UEdGraphPin*& NewPin)
 	PinData.PinName = UniqueName;
 	NewPin->PinName = UniqueName;
 	AddedPins.Add(PinData);
+
+	if (bAllStatic)
+	{
+		const UEdGraphSchema_Niagara* Schema = GetDefault<UEdGraphSchema_Niagara>();
+		FNiagaraTypeDefinition PinType = Schema->PinToTypeDefinition(NewPin);
+		if (PinType.IsStatic() == false)
+			NewPin->PinType = Schema->TypeDefinitionToPinType(PinType.ToStaticDef());
+	}
 }
 
 void UNiagaraNodeOp::OnPinRenamed(UEdGraphPin* RenamedPin, const FString& OldName)
@@ -352,4 +410,246 @@ FName UNiagaraNodeOp::GetUniqueAdditionalPinName() const
 	return FNiagaraUtilities::GetUniqueName(*Name, ExistingNames);
 }
 
+void UNiagaraNodeOp::BuildParameterMapHistory(FNiagaraParameterMapHistoryBuilder& OutHistory, bool bRecursive /*= true*/, bool bFilterForCompilation /*= true*/) const
+{
+	if (bRecursive)
+	{
+		OutHistory.VisitInputPins(this, bFilterForCompilation);
+	}
+
+	if (!IsNodeEnabled() && OutHistory.GetIgnoreDisabled())
+	{
+		RouteParameterMapAroundMe(OutHistory, bRecursive);
+		return;
+	}
+	else 
+	{		
+		const FNiagaraOpInfo* OpInfo = FNiagaraOpInfo::GetOpInfo(OpName);
+		if (OpInfo && OpInfo->StaticVariableResolveFunction.IsBound())
+		{
+			const UEdGraphSchema_Niagara* Schema = GetDefault<UEdGraphSchema_Niagara>();
+			TArray<UEdGraphPin*> InputPins;
+			GetInputPins(InputPins);
+			
+			bool bAllPinsStatic = true;
+			UEdGraphPin* OutputPin = nullptr;
+			{
+				for (int32 PinIdx = 0; PinIdx < Pins.Num(); PinIdx++)
+				{
+					if (IsAddPin(Pins[PinIdx]))
+						continue;
+					FNiagaraTypeDefinition InputType = Schema->PinToTypeDefinition(Pins[PinIdx]);
+					if (!InputType.IsStatic())
+						bAllPinsStatic = false;
+
+					if (OutputPin == nullptr && Pins[PinIdx]->Direction == EEdGraphPinDirection::EGPD_Output)
+					{
+						OutputPin = Pins[PinIdx];
+					}
+				}
+			}
+
+			if (bAllPinsStatic)
+			{
+				TArray<int32> Vars;
+
+				for (int32 InputIdx = 0; InputIdx < InputPins.Num(); InputIdx++)
+				{
+					if (IsAddPin(InputPins[InputIdx]))
+						continue;
+
+					FNiagaraTypeDefinition InputType = Schema->PinToTypeDefinition(InputPins[InputIdx]);
+					if (!InputType.IsStatic())
+					{
+						return;
+					}
+					int32 Value = 0;
+					OutHistory.SetConstantByStaticVariable(Value, InputPins[InputIdx]);
+					Vars.Add(Value);
+				}
+
+				if (Vars.Num() > 0)
+				{
+					int32 Result = OpInfo->StaticVariableResolveFunction.Execute(Vars);
+					int32 ConstantIdx = OutHistory.AddOrGetConstantFromValue(FString::FromInt(Result));
+
+					if (UNiagaraScript::LogCompileStaticVars > 0)
+					{
+						UE_LOG(LogNiagaraEditor, Log, TEXT("Inputs Static Node Op: %s"), *GetNodeTitle(ENodeTitleType::FullTitle).ToString());
+						for (int32 i = 0; i < Vars.Num(); i++)
+						{
+							UE_LOG(LogNiagaraEditor, Log, TEXT("[%d] %d"), i, Vars[i]);
+						}
+						UE_LOG(LogNiagaraEditor, Log, TEXT("Result: %d"), Result);
+					}
+
+					OutHistory.RegisterConstantPin(ConstantIdx, OutputPin);
+				}
+			}
+		}
+	}
+}
+
+void UNiagaraNodeOp::PinConnectionListChanged(UEdGraphPin* Pin)
+{
+	Super::PinConnectionListChanged(Pin);
+
+	FEdGraphPinType StartType = Pin->PinType;
+	HandleStaticInputPinUpgrade(Pin);
+
+	// Type was changed... now make sure that we can fix up the output pin.
+	if (StartType != Pin->PinType)
+	{
+		HandleStaticOutputPinUpgrade();
+	}
+}
+
+void UNiagaraNodeOp::HandleStaticInputPinUpgrade(UEdGraphPin* Pin)
+{
+	const FNiagaraOpInfo* OpInfo = FNiagaraOpInfo::GetOpInfo(OpName);
+	if (Pin && OpInfo && OpInfo->StaticVariableResolveFunction.IsBound() && OpInfo->bSupportsStaticResolution)
+	{
+		const UEdGraphSchema_Niagara* Schema = GetDefault<UEdGraphSchema_Niagara>();
+		FNiagaraVariable PinVar = Schema->PinToNiagaraVariable(Pin);
+
+		if (Pin->Direction == EEdGraphPinDirection::EGPD_Input && !IsAddPin(Pin))
+		{
+			// Handle linkage...
+			bool bHandled = false;
+			if (Pin->LinkedTo.Num() > 0)
+			{
+				FNiagaraTypeDefinition LinkedPinType = Schema->PinToTypeDefinition(Pin->LinkedTo[0]);
+				Pin->PinType = Schema->TypeDefinitionToPinType(LinkedPinType);
+			}
+			else // Handle unlinkage...
+			{
+				// Downcast if the types don't match...
+				if (OpInfo->bSupportsAddedInputs && OpInfo->Inputs.Num() > 0)
+				{
+					if (PinVar.GetType() != OpInfo->Inputs[0].DataType && OpInfo->Inputs[0].DataType != FNiagaraTypeDefinition::GetGenericNumericDef())
+					{
+						Pin->PinType = Schema->TypeDefinitionToPinType(OpInfo->Inputs[0].DataType);
+					}
+				}
+				else
+				{
+					for (int32 i = 0; i < OpInfo->Inputs.Num(); i++)
+					{
+						if (OpInfo->Inputs[i].Name == PinVar.GetName())
+						{
+							if (PinVar.GetType() != OpInfo->Inputs[i].DataType && OpInfo->Inputs[0].DataType != FNiagaraTypeDefinition::GetGenericNumericDef())
+							{
+								Pin->PinType = Schema->TypeDefinitionToPinType(OpInfo->Inputs[i].DataType);
+								break;
+							}
+							break;
+						}
+					}
+				}
+			}
+		}
+	}
+}
+
+void UNiagaraNodeOp::HandleStaticOutputPinUpgrade()
+{
+	const UEdGraphSchema_Niagara* Schema = GetDefault<UEdGraphSchema_Niagara>();
+	const FNiagaraOpInfo* OpInfo = FNiagaraOpInfo::GetOpInfo(OpName);
+	if (OpInfo)
+	{
+		bAllStatic = false;
+
+		bool bAllInputStatic = true;
+		bool bAnyInputStatic = false;
+		for (int32 i = 0; i < Pins.Num(); i++)
+		{
+			if (IsAddPin(Pins[i]) || Pins[i]->Direction != EEdGraphPinDirection::EGPD_Input)
+				continue;
+
+			FNiagaraTypeDefinition PinType = Schema->PinToTypeDefinition(Pins[i]);
+			if (PinType.IsStatic())
+			{
+				bAnyInputStatic = true;
+			}
+			else
+			{
+				bAllInputStatic = false;
+				continue;
+			}
+		}
+
+		if (OpInfo->StaticVariableResolveFunction.IsBound() && OpInfo->bSupportsStaticResolution)
+		{
+			// Synchronize the output pin either way...
+			bool bAllOutputStatic = true;
+			bool bAnyOutputStatic = false;
+			for (int32 PinIdx = 0; PinIdx < Pins.Num(); PinIdx++)
+			{
+				if (Pins[PinIdx]->Direction == EEdGraphPinDirection::EGPD_Output)
+				{
+					FNiagaraVariable PinVar = Schema->PinToNiagaraVariable(Pins[PinIdx]);
+
+					if (IsAddPin(Pins[PinIdx]))
+						continue;
+
+					if (!PinVar.GetType().IsStatic())
+					{
+						bAllOutputStatic = false;
+					}
+					else
+					{
+						bAnyOutputStatic = true;
+					}
+				}
+			}
+
+			bool bCanAutoConvertOutputs = bAllInputStatic;
+			for (int32 OutputIdx = 0; OutputIdx < OpInfo->Outputs.Num(); OutputIdx++)
+			{
+				if (!FNiagaraTypeRegistry::IsStaticPossible(OpInfo->Outputs[OutputIdx].DataType))
+				{
+					bCanAutoConvertOutputs = false;
+				}
+			}
+
+			if (bCanAutoConvertOutputs)
+			{
+				for (int32 PinIdx = 0; PinIdx < Pins.Num(); PinIdx++)
+				{
+					if (IsAddPin(Pins[PinIdx]))
+						continue;
+
+					if (Pins[PinIdx]->Direction == EEdGraphPinDirection::EGPD_Output)
+					{
+						FNiagaraTypeDefinition PinStaticType = Schema->PinToTypeDefinition(Pins[PinIdx]).ToStaticDef();
+						Pins[PinIdx]->PinType = Schema->TypeDefinitionToPinType(PinStaticType);
+					}
+				}
+				bAllOutputStatic = true;
+			}
+
+			if (bAllOutputStatic && bAllInputStatic)
+			{
+				bAllStatic = true;
+				NodeUpgradeMessage = FText();
+			}
+			else if (bAnyOutputStatic && !bAllInputStatic)
+			{
+				NodeUpgradeMessage = LOCTEXT("AllStaticInvalidInput", "All inputs must be static to have static outputs. Static will not be preserved.");
+				MarkNodeRequiresSynchronization(__FUNCTION__, true);
+			}
+			else if (bAnyInputStatic && !bAllOutputStatic)
+			{
+				NodeUpgradeMessage = LOCTEXT("AllStaticInvalidOutput", "All outputs must be static to have static inputs. Static will not be preserved.");
+				MarkNodeRequiresSynchronization(__FUNCTION__, true);
+			}
+			
+		}
+		else if (bAllStatic)
+		{
+			NodeUpgradeMessage = LOCTEXT("AllStaticNotSupported", "This op doesn't fully support static outputs. Static will not be preserved.");
+			MarkNodeRequiresSynchronization(__FUNCTION__, true);
+		}
+	}
+}
 #undef LOCTEXT_NAMESPACE
diff --git a/Engine/Plugins/FX/Niagara/Source/NiagaraEditor/Private/NiagaraNodeOutput.cpp b/Engine/Plugins/FX/Niagara/Source/NiagaraEditor/Private/NiagaraNodeOutput.cpp
--- a/Engine/Plugins/FX/Niagara/Source/NiagaraEditor/Private/NiagaraNodeOutput.cpp
+++ b/Engine/Plugins/FX/Niagara/Source/NiagaraEditor/Private/NiagaraNodeOutput.cpp
@@ -293,6 +293,35 @@ void UNiagaraNodeOutput::Compile(FHlslNiagaraTranslator *Translator, TArray<int3
 	}
 }
 
+void UNiagaraNodeOutput::BuildParameterMapHistory(FNiagaraParameterMapHistoryBuilder& OutHistory, bool bRecursive /*= true*/, bool bFilterForCompilation /*= true*/) const
+{
+	Super::BuildParameterMapHistory(OutHistory, bRecursive, bFilterForCompilation);
+	const UEdGraphSchema_Niagara* Schema = GetDefault<UEdGraphSchema_Niagara>();
+	for (const UEdGraphPin* Pin : Pins)
+	{
+		if (Pin && Pin->Direction == EGPD_Input)
+		{
+			FNiagaraTypeDefinition InputType = Schema->PinToTypeDefinition(Pin);
+			if (InputType.IsStatic())
+			{
+				OutHistory.RegisterConstantFromInputPin(Pin);
+			}
+			else if (InputType == FNiagaraTypeDefinition::GetParameterMapDef())
+			{
+				if (Pin->LinkedTo.Num() != 0)
+				{
+					int32 ParamMapIdx = OutHistory.TraceParameterMapOutputPin(Pin->LinkedTo[0]);			
+					if (UNiagaraScript::LogCompileStaticVars > 0)
+					{
+						UE_LOG(LogNiagaraEditor, Log, TEXT("Build Parameter Map History: %s %s PMapIdx: %d"), *GetClass()->GetName(), *GetNodeTitle(ENodeTitleType::FullTitle).ToString(), ParamMapIdx);
+					}
+					OutHistory.RegisterParameterMapPin(ParamMapIdx, Pin);
+				}
+			}
+		}
+	}
+}
+
 
 void UNiagaraNodeOutput::PostLoad()
 {
diff --git a/Engine/Plugins/FX/Niagara/Source/NiagaraEditor/Private/NiagaraNodeParameterMapGet.cpp b/Engine/Plugins/FX/Niagara/Source/NiagaraEditor/Private/NiagaraNodeParameterMapGet.cpp
--- a/Engine/Plugins/FX/Niagara/Source/NiagaraEditor/Private/NiagaraNodeParameterMapGet.cpp
+++ b/Engine/Plugins/FX/Niagara/Source/NiagaraEditor/Private/NiagaraNodeParameterMapGet.cpp
@@ -436,7 +436,9 @@ void UNiagaraNodeParameterMapGet::Compile(class FHlslNiagaraTranslator* Translat
 		UEdGraphPin* InputPin = InputPins[i];
 
 		if (InputPin->PinType.PinCategory == UEdGraphSchema_Niagara::PinCategoryType || 
-			InputPin->PinType.PinCategory == UEdGraphSchema_Niagara::PinCategoryEnum)
+			InputPin->PinType.PinCategory == UEdGraphSchema_Niagara::PinCategoryStaticType ||
+			InputPin->PinType.PinCategory == UEdGraphSchema_Niagara::PinCategoryEnum || 
+			InputPin->PinType.PinCategory == UEdGraphSchema_Niagara::PinCategoryStaticEnum )
 		{
 			int32 CompiledInput = INDEX_NONE;
 			if (i == 0) // Only the zeroth item is not an default value pin.
diff --git a/Engine/Plugins/FX/Niagara/Source/NiagaraEditor/Private/NiagaraNodeReroute.cpp b/Engine/Plugins/FX/Niagara/Source/NiagaraEditor/Private/NiagaraNodeReroute.cpp
--- a/Engine/Plugins/FX/Niagara/Source/NiagaraEditor/Private/NiagaraNodeReroute.cpp
+++ b/Engine/Plugins/FX/Niagara/Source/NiagaraEditor/Private/NiagaraNodeReroute.cpp
@@ -106,7 +106,32 @@ bool UNiagaraNodeReroute::ShouldDrawNodeAsControlPointOnly(int32& OutInputPinInd
 
 void UNiagaraNodeReroute::Compile(class FHlslNiagaraTranslator* Translator, TArray<int32>& Outputs)
 {
-	Super::Compile(Translator, Outputs);
+	FPinCollectorArray InputPins;
+	GetInputPins(InputPins);
+
+	FPinCollectorArray OutputPins;
+	GetOutputPins(OutputPins);
+
+	// Initialize the outputs to invalid values.
+	check(Outputs.Num() == 0);
+	Outputs.Reserve(OutputPins.Num());
+	for (int32 i = 0; i < OutputPins.Num(); i++)
+	{
+		if (InputPins.IsValidIndex(i))
+		{
+			int32 CompiledInput = Translator->CompilePin(InputPins[i]);
+			Outputs.Add(CompiledInput);
+		}
+		else
+		{
+			Outputs.Add(INDEX_NONE);
+		}
+	}
+
+	if (InputPins.Num() != OutputPins.Num())
+	{
+		Translator->Error(LOCTEXT("IncorrectNumOutputsError", "Input and Output pin counts must match."), this, nullptr);
+	}
 }
 
 bool UNiagaraNodeReroute::RefreshFromExternalChanges()
@@ -124,12 +149,11 @@ void UNiagaraNodeReroute::PinConnectionListChanged(UEdGraphPin* Pin)
 
 void UNiagaraNodeReroute::BuildParameterMapHistory(FNiagaraParameterMapHistoryBuilder& OutHistory, bool bRecursive /*= true*/, bool bFilterForCompilation /*= true*/) const
 {
-	// Should never hit here...
-	check(false);
+	RegisterPassthroughPin(OutHistory, GetInputPin(0), GetOutputPin(0), bFilterForCompilation, true);
 }
 
 /** Traces one of this node's output pins to its source output pin if it is a reroute node output pin.*/
-UEdGraphPin* UNiagaraNodeReroute::GetTracedOutputPin(UEdGraphPin* LocallyOwnedOutputPin, bool bFilterForCompilation) const
+UEdGraphPin* UNiagaraNodeReroute::GetTracedOutputPin(UEdGraphPin* LocallyOwnedOutputPin, bool bFilterForCompilation, TArray<const UNiagaraNode*>* OutNodesVisitedDuringTrace) const
 {
 	check(Pins.Contains(LocallyOwnedOutputPin) && LocallyOwnedOutputPin->Direction == EGPD_Output);
 	UEdGraphPin* InputPin = GetInputPin(0);
@@ -137,7 +161,7 @@ UEdGraphPin* UNiagaraNodeReroute::GetTracedOutputPin(UEdGraphPin* LocallyOwnedOu
 	{
 		UEdGraphPin* LinkedPin = InputPin->LinkedTo[0];
 		UNiagaraNode* LinkedNode = CastChecked<UNiagaraNode>(LinkedPin->GetOwningNode());
-		return LinkedNode->GetTracedOutputPin(LinkedPin, bFilterForCompilation);
+		return LinkedNode->GetTracedOutputPin(LinkedPin, bFilterForCompilation, OutNodesVisitedDuringTrace);
 	}
 	return nullptr;
 }
diff --git a/Engine/Plugins/FX/Niagara/Source/NiagaraEditor/Private/NiagaraNodeReroute.h b/Engine/Plugins/FX/Niagara/Source/NiagaraEditor/Private/NiagaraNodeReroute.h
--- a/Engine/Plugins/FX/Niagara/Source/NiagaraEditor/Private/NiagaraNodeReroute.h
+++ b/Engine/Plugins/FX/Niagara/Source/NiagaraEditor/Private/NiagaraNodeReroute.h
@@ -38,7 +38,7 @@ public:
 	virtual bool AllowExternalPinTypeChanges(const UEdGraphPin* InGraphPin) const override;
 	
 	/** Trace to an output pin that is not a reroute node output pin. If the reroute nodes resulted in a dead end and no output pin was found return nullptr.*/
-	virtual UEdGraphPin* GetTracedOutputPin(UEdGraphPin* LocallyOwnedOutputPin, bool bFilterForCompilation) const override;
+	virtual UEdGraphPin* GetTracedOutputPin(UEdGraphPin* LocallyOwnedOutputPin, bool bFilterForCompilation, TArray<const UNiagaraNode*>* OutNodesVisitedDuringTrace = nullptr) const override;
 	virtual UEdGraphPin* GetPassThroughPin(const UEdGraphPin* LocallyOwnedOutputPin, ENiagaraScriptUsage MasterUsage) const override { return GetPassThroughPin(LocallyOwnedOutputPin); }
 	//~ End UNiagaraNode Interface
 
diff --git a/Engine/Plugins/FX/Niagara/Source/NiagaraEditor/Private/NiagaraNodeSelect.cpp b/Engine/Plugins/FX/Niagara/Source/NiagaraEditor/Private/NiagaraNodeSelect.cpp
--- a/Engine/Plugins/FX/Niagara/Source/NiagaraEditor/Private/NiagaraNodeSelect.cpp
+++ b/Engine/Plugins/FX/Niagara/Source/NiagaraEditor/Private/NiagaraNodeSelect.cpp
@@ -79,7 +79,7 @@ void UNiagaraNodeSelect::ChangeSelectorPinType(FNiagaraTypeDefinition Type)
 	}
 	// ensure we have at least a minimum amount of entries if we change to integer mode
 	const int32 MinimumEntries = 2;
-	if (SelectorPinType == FNiagaraTypeDefinition::GetIntDef() && NumOptionsPerVariable < MinimumEntries)
+	if (SelectorPinType.IsSameBaseDefinition(FNiagaraTypeDefinition::GetIntDef()) && NumOptionsPerVariable < MinimumEntries)
 	{
 		NumOptionsPerVariable = MinimumEntries;
 	}
@@ -486,7 +486,7 @@ bool UNiagaraNodeSelect::CanRenamePin(const UEdGraphPin* Pin) const
 
 FString UNiagaraNodeSelect::GetInputCaseName(int32 Case) const
 {
-	if(SelectorPinType == FNiagaraTypeDefinition::GetBoolDef())
+	if( FNiagaraTypeDefinition::GetBoolDef().IsSameBaseDefinition(SelectorPinType))
 	{
 		if(Case == 0)
 		{
@@ -497,7 +497,7 @@ FString UNiagaraNodeSelect::GetInputCaseName(int32 Case) const
 			return TEXT("True");
 		}
 	}
-	else if (SelectorPinType == FNiagaraTypeDefinition::GetIntDef())
+	else if ( FNiagaraTypeDefinition::GetIntDef().IsSameBaseDefinition(SelectorPinType))
 	{
 		return FString::FromInt(Case);
 	}
@@ -661,11 +661,11 @@ TArray<int32> UNiagaraNodeSelect::GetOptionValues() const
 {
 	TArray<int32> SelectorValues;
 
-	if(SelectorPinType == FNiagaraTypeDefinition::GetBoolDef())
+	if(FNiagaraTypeDefinition::GetBoolDef().IsSameBaseDefinition(SelectorPinType))
 	{
 		SelectorValues = { 1, 0 };
 	}
-	else if(SelectorPinType == FNiagaraTypeDefinition::GetIntDef())
+	else if(FNiagaraTypeDefinition::GetIntDef().IsSameBaseDefinition(SelectorPinType))
 	{
 		int32 NewOptionsCount = FMath::Max(2, NumOptionsPerVariable);
 		for(int32 Index = 0; Index < NewOptionsCount; Index++)
@@ -712,12 +712,12 @@ FText UNiagaraNodeSelect::GetIntegerRemoveButtonTooltipText() const
 
 EVisibility UNiagaraNodeSelect::ShowAddIntegerButton() const
 {
-	return SelectorPinType == FNiagaraTypeDefinition::GetIntDef() ? EVisibility::Visible : EVisibility::Collapsed;
+	return SelectorPinType.IsSameBaseDefinition(FNiagaraTypeDefinition::GetIntDef()) ? EVisibility::Visible : EVisibility::Collapsed;
 }
 
 EVisibility UNiagaraNodeSelect::ShowRemoveIntegerButton() const
 {
-	return SelectorPinType == FNiagaraTypeDefinition::GetIntDef() && NumOptionsPerVariable > 2 ? EVisibility::Visible : EVisibility::Collapsed;
+	return SelectorPinType.IsSameBaseDefinition(FNiagaraTypeDefinition::GetIntDef()) && NumOptionsPerVariable > 2 ? EVisibility::Visible : EVisibility::Collapsed;
 }
 
 #undef LOCTEXT_NAMESPACE
diff --git a/Engine/Plugins/FX/Niagara/Source/NiagaraEditor/Private/NiagaraNodeStaticSwitch.cpp b/Engine/Plugins/FX/Niagara/Source/NiagaraEditor/Private/NiagaraNodeStaticSwitch.cpp
--- a/Engine/Plugins/FX/Niagara/Source/NiagaraEditor/Private/NiagaraNodeStaticSwitch.cpp
+++ b/Engine/Plugins/FX/Niagara/Source/NiagaraEditor/Private/NiagaraNodeStaticSwitch.cpp
@@ -77,6 +77,20 @@ FString UNiagaraNodeStaticSwitch::GetInputCaseName(int32 Case) const
 	return TEXT("");
 }
 
+#if WITH_EDITOR
+void UNiagaraNodeStaticSwitch::PostEditChangeProperty(struct FPropertyChangedEvent& PropertyChangedEvent)
+{
+	Super::PostEditChangeProperty(PropertyChangedEvent);
+
+	const FName MemberPropertyName = (PropertyChangedEvent.MemberProperty != nullptr) ? PropertyChangedEvent.MemberProperty->GetFName() : NAME_None;
+	
+	if (MemberPropertyName != NAME_None)
+	{
+		ReallocatePins();
+	}
+}
+#endif
+
 TArray<int32> UNiagaraNodeStaticSwitch::GetOptionValues() const
 {
 	TArray<int32> OptionValues;
@@ -182,7 +196,7 @@ void UNiagaraNodeStaticSwitch::SetSwitchValue(int Value)
 
 void UNiagaraNodeStaticSwitch::SetSwitchValue(const FCompileConstantResolver& ConstantResolver)
 {
-	if (!IsSetByCompiler())
+	if (!IsSetByCompiler() && !IsSetByPin())
 	{
 		return;
 	}
@@ -217,6 +231,12 @@ bool UNiagaraNodeStaticSwitch::IsSetByCompiler() const
 	return !SwitchTypeData.SwitchConstant.IsNone();
 }
 
+bool UNiagaraNodeStaticSwitch::IsSetByPin() const
+{
+	return SwitchTypeData.bExposeAsPin && !IsSetByCompiler() && !IsDebugSwitch();
+}
+
+
 bool UNiagaraNodeStaticSwitch::IsDebugSwitch() const
 {
 	return SwitchTypeData.SwitchType == ENiagaraStaticSwitchType::Enum && SwitchTypeData.SwitchConstant == TEXT("Function.DebugState") && SwitchTypeData.Enum == FNiagaraTypeDefinition::GetFunctionDebugStateEnum();
@@ -256,7 +276,15 @@ void UNiagaraNodeStaticSwitch::AllocateDefaultPins()
 	}
 		
 	GetNiagaraGraph()->AddParameter(FNiagaraVariable(GetInputType(), InputParameterName), true);
-	
+
+	if (IsSetByPin())
+	{		
+		FNiagaraTypeDefinition Type = GetInputType().ToStaticDef();
+		UEdGraphPin* NewPin  = CreatePin(EGPD_Input, Schema->TypeDefinitionToPinType(Type), TEXT("Selector"));
+		if (SelectorGuid.IsValid() == false)
+			SelectorGuid = FGuid::NewGuid();
+		NewPin->PersistentGuid = SelectorGuid;
+	}
 	// create the output pins
 	for (int32 Index = 0; Index < OutputVars.Num(); Index++)
 	{
@@ -282,8 +310,63 @@ bool UNiagaraNodeStaticSwitch::GetVarIndex(FHlslNiagaraTranslator* Translator, i
 	return GetVarIndex(Translator, InputPinCount, SwitchValue, VarIndexOut);
 }
 
+UEdGraphPin* UNiagaraNodeStaticSwitch::GetSelectorPin() const
+{
+	if (IsSetByPin() == false)
+		return nullptr;
+
+	for (int32 i = 0; i < Pins.Num(); i++)
+	{
+		if (Pins[i]->PersistentGuid == SelectorGuid)
+		{
+			return  Pins[i];
+		}
+	}
+
+	return nullptr;
+}
+
+
 void UNiagaraNodeStaticSwitch::UpdateCompilerConstantValue(FHlslNiagaraTranslator* Translator)
 {
+	if (IsSetByPin() && Translator)
+	{
+		ClearSwitchValue();
+
+		const UEdGraphPin* SelectorPin = nullptr;
+		for (int32 i = 0; i < Pins.Num(); i++)
+		{
+			if (Pins[i]->PersistentGuid == SelectorGuid)
+			{
+				SelectorPin = Pins[i];
+				break;
+			}
+		}
+
+		if (SelectorPin)
+		{
+			// Use the default value if disconnected...
+			if (SelectorPin->LinkedTo.Num() == 0)
+			{
+				Translator->FillVariableWithDefaultValue(SwitchValue, SelectorPin);
+				IsValueSet = true;
+			}
+			else 
+			{
+				Translator->SetConstantByStaticVariable(SwitchValue, SelectorPin);
+				IsValueSet = true;
+			}
+
+			if (UNiagaraScript::LogCompileStaticVars != 0)
+			{
+				UE_LOG(LogNiagaraEditor, Log, TEXT("Static Switch Value: %d  Name: \"%s\""), SwitchValue, *GetPathName());
+			}
+			return;
+		}
+	}
+
+
+
 	if (!IsSetByCompiler() || IsDebugSwitch() || !Translator)
 	{
 		return;
@@ -420,13 +503,14 @@ void UNiagaraNodeStaticSwitch::CheckForOutdatedEnum(FHlslNiagaraTranslator* Tran
 		}
 
 		// this checks the current, possibly outdated, state for consistency. Should never be unequal but we are making sure.
-		if (InputPins.Num() != NumOptionsPerVariable * OutputVars.Num())
+		int NumChoices = IsSetByPin() ? InputPins.Num() - 1 : InputPins.Num();
+		if (NumChoices != NumOptionsPerVariable * OutputVars.Num())
 		{
 			Translator->Message(FNiagaraCompileEventSeverity::Error, LOCTEXT("InternalCalculationPinCountError", "An internal calculation error has occured. Please refresh the node."), this, nullptr);
 		}
 
 		// this checks the new, up to date pin count against the old pin count
-		if(NewPinNames.Num() != InputPins.Num())
+		if(NewPinNames.Num() != NumChoices)
 		{
 			if(NewPinNames.Num() == 0)
 			{
@@ -459,9 +543,88 @@ void UNiagaraNodeStaticSwitch::CheckForOutdatedEnum(FHlslNiagaraTranslator* Tran
 
 void UNiagaraNodeStaticSwitch::Compile(FHlslNiagaraTranslator* Translator, TArray<int32>& Outputs)
 {	
-	UNiagaraNode::Compile(Translator, Outputs);
+	FPinCollectorArray InputPins;
+	GetInputPins(InputPins);
+	int32 SelectorValue = SwitchValue;
+
+	if (IsSetByPin())
+	{
+		UEdGraphPin* SelectorPin = GetSelectorPin();	
+		if (SelectorPin)
+		{
+			Translator->CompilePin(SelectorPin);
+		}
+	}
+	FPinCollectorArray OutputPins;
+	GetOutputPins(OutputPins);
+
+	// Initialize the outputs to invalid values.
+	check(Outputs.Num() == 0);
+	Outputs.Reserve(OutputPins.Num());
+	for (int32 i = 0; i < OutputPins.Num(); i++)
+	{
+		Outputs.Add(INDEX_NONE);
+	}
+
+	for (int i = 0; i < OutputPins.Num(); i++)
+	{
+		UEdGraphPin* OutPin = OutputPins[i];
+		if (IsAddPin(OutPin))
+		{
+			continue;
+		}
+		int32 VarIdx;
+		int NumChoices = IsSetByPin() ? InputPins.Num() - 1 : InputPins.Num();
+		if (GetVarIndex(nullptr, NumChoices, SelectorValue, VarIdx))
+		{
+			UEdGraphPin* InputPin = InputPins[VarIdx + i];
+			if (InputPin)
+			{
+				int32 CompiledInput = Translator->CompilePin(InputPin);
+				Outputs[i] = CompiledInput;
+			}
+		}		
+	}
 }
 
+ENiagaraNumericOutputTypeSelectionMode UNiagaraNodeStaticSwitch::GetNumericOutputTypeSelectionMode() const
+{
+	return ENiagaraNumericOutputTypeSelectionMode::Largest;
+}
+
+void UNiagaraNodeStaticSwitch::ResolveNumerics(const UEdGraphSchema_Niagara* Schema, bool bSetInline, TMap<TPair<FGuid, UEdGraphNode*>, FNiagaraTypeDefinition>* PinCache)
+{	
+	FPinCollectorArray MasterInputPins;
+	GetInputPins(MasterInputPins);
+	FPinCollectorArray MasterOutputPins;
+	GetOutputPins(MasterOutputPins);
+
+	TArray<int32> OptionValues = GetOptionValues();
+	int32 VarIdx = 0;
+	for (int i = 0; i < MasterOutputPins.Num(); i++)
+	{
+		UEdGraphPin* OutPin = MasterOutputPins[i];
+		if (IsAddPin(OutPin))
+		{
+			continue;
+		}
+
+		// Fix up numeric input pins and keep track of numeric types to decide the output type.
+		TArray<UEdGraphPin*> InputPins;
+		TArray<UEdGraphPin*> OutputPins; 
+		
+		for (int32 j = 0; j < OptionValues.Num(); j++)
+		{
+			UEdGraphPin* InputPin = MasterInputPins[OutputVars.Num()*j + VarIdx];
+			InputPins.Add(InputPin);
+		}
+
+		OutputPins.Add(OutPin);
+
+		NumericResolutionByPins(Schema, InputPins, OutputPins, bSetInline, PinCache);
+		VarIdx++;
+	}
+}
 bool UNiagaraNodeStaticSwitch::SubstituteCompiledPin(FHlslNiagaraTranslator* Translator, UEdGraphPin** LocallyOwnedPin)
 {
 	// if we compile the standalone module or function we don't have any valid input yet, so we just take the first option to satisfy the compiler
@@ -491,7 +654,8 @@ bool UNiagaraNodeStaticSwitch::SubstituteCompiledPin(FHlslNiagaraTranslator* Tra
 	{
 		UEdGraphPin* OutPin = OutputPins[i];
 		int32 VarIdx;
-		if (OutPin == *LocallyOwnedPin && GetVarIndex(Translator, InputPins.Num(), IsDryRun ? 0 : SwitchValue, VarIdx))
+		int NumChoices = IsSetByPin() ? InputPins.Num() - 1 : InputPins.Num();
+		if (OutPin == *LocallyOwnedPin && GetVarIndex(Translator, NumChoices, IsDryRun ? 0 : SwitchValue, VarIdx))
 		{
 			UEdGraphPin* InputPin = InputPins[VarIdx + i];
 			if (InputPin->LinkedTo.Num() == 1)
@@ -572,15 +736,15 @@ void UNiagaraNodeStaticSwitch::PostLoad()
 		{
 			RefreshFromExternalChanges();
 		}
-	}
+	}	
 }
 
-UEdGraphPin* UNiagaraNodeStaticSwitch::GetTracedOutputPin(UEdGraphPin* LocallyOwnedOutputPin, bool bFilterForCompilation) const
+UEdGraphPin* UNiagaraNodeStaticSwitch::GetTracedOutputPin(UEdGraphPin* LocallyOwnedOutputPin, bool bFilterForCompilation, TArray<const UNiagaraNode*>* OutNodesVisitedDuringTrace) const
 {
-	return GetTracedOutputPin(LocallyOwnedOutputPin, true, bFilterForCompilation);
+	return GetTracedOutputPin(LocallyOwnedOutputPin, true, bFilterForCompilation, OutNodesVisitedDuringTrace);
 }
 
-UEdGraphPin* UNiagaraNodeStaticSwitch::GetTracedOutputPin(UEdGraphPin* LocallyOwnedOutputPin, bool bRecursive, bool bFilterForCompilation) const
+UEdGraphPin* UNiagaraNodeStaticSwitch::GetTracedOutputPin(UEdGraphPin* LocallyOwnedOutputPin, bool bRecursive, bool bFilterForCompilation, TArray<const UNiagaraNode*>* OutNodesVisitedDuringTrace) const
 {
 	if (!bFilterForCompilation)
 	{
@@ -593,6 +757,11 @@ UEdGraphPin* UNiagaraNodeStaticSwitch::GetTracedOutputPin(UEdGraphPin* LocallyOw
 	FPinCollectorArray OutputPins;
 	GetOutputPins(OutputPins);
 
+	if (OutNodesVisitedDuringTrace != nullptr && IsSetByPin())
+	{
+		OutNodesVisitedDuringTrace->Add(this);
+	}
+
 	for (int i = 0; i < OutputPins.Num(); i++)
 	{
 		UEdGraphPin* OutPin = OutputPins[i];
@@ -601,12 +770,13 @@ UEdGraphPin* UNiagaraNodeStaticSwitch::GetTracedOutputPin(UEdGraphPin* LocallyOw
 			continue;
 		}
 		int32 VarIdx;
-		if (OutPin == LocallyOwnedOutputPin && GetVarIndex(nullptr, InputPins.Num(), SwitchValue, VarIdx))
+		int NumChoices = IsSetByPin() ? InputPins.Num() - 1 : InputPins.Num();
+		if (OutPin == LocallyOwnedOutputPin && GetVarIndex(nullptr, NumChoices, SwitchValue, VarIdx))
 		{
 			UEdGraphPin* InputPin = InputPins[VarIdx + i];
 			if (InputPin->LinkedTo.Num() == 1)
 			{
-				return bRecursive ? UNiagaraNode::TraceOutputPin(InputPin->LinkedTo[0], bFilterForCompilation) : InputPin->LinkedTo[0];
+				return bRecursive ? UNiagaraNode::TraceOutputPin(InputPin->LinkedTo[0], bFilterForCompilation, OutNodesVisitedDuringTrace) : InputPin->LinkedTo[0];
 			}
 		}
 	}
@@ -625,7 +795,70 @@ UEdGraphPin* UNiagaraNodeStaticSwitch::GetPassThroughPin(const UEdGraphPin* Loca
 
 void UNiagaraNodeStaticSwitch::BuildParameterMapHistory(FNiagaraParameterMapHistoryBuilder& OutHistory, bool bRecursive /*= true*/, bool bFilterForCompilation /*= true*/) const
 {
-	UNiagaraNode::BuildParameterMapHistory(OutHistory, bRecursive, bFilterForCompilation);
+	FPinCollectorArray InputPins;
+	GetInputPins(InputPins);
+	int32 SelectorValue = SwitchValue;
+
+	if (!bFilterForCompilation)
+	{
+		UNiagaraNode::BuildParameterMapHistory(OutHistory, bRecursive, bFilterForCompilation);
+	}
+	else
+	{
+		if (IsSetByPin())
+		{
+			UEdGraphPin* SelectorPin = GetSelectorPin();
+			OutHistory.VisitInputPin(SelectorPin, this, bFilterForCompilation);
+
+			OutHistory.SetConstantByStaticVariable(SelectorValue, SelectorPin);
+
+			if (UNiagaraScript::LogCompileStaticVars != 0)
+			{
+				UE_LOG(LogNiagaraEditor, Log, TEXT("Static Switch Value: %d  Name: \"%s\"  SwitchInputParameterName:\"%s\""), SelectorValue, *GetPathName(), *InputParameterName.ToString());
+			}
+		}
+	}
+
+	FPinCollectorArray OutputPins;
+	GetOutputPins(OutputPins);
+
+	for (int i = 0; i < OutputPins.Num(); i++)
+	{
+		UEdGraphPin* OutPin = OutputPins[i];
+		if (IsAddPin(OutPin))
+		{
+			continue;
+		}
+		int32 VarIdx;
+		int NumChoices = IsSetByPin() ? InputPins.Num() - 1 : InputPins.Num();
+		if (bFilterForCompilation && GetVarIndex(nullptr, NumChoices, SelectorValue, VarIdx))
+		{
+			UEdGraphPin* InputPin = InputPins[VarIdx + i];
+			RegisterPassthroughPin(OutHistory, InputPin, OutPin, bFilterForCompilation, true);
+			/*if (InputPin->LinkedTo.Num() == 1)
+			{
+				return bRecursive ? UNiagaraNode::TraceOutputPin(InputPin->LinkedTo[0], bFilterForCompilation, OutNodesVisitedDuringTrace) : InputPin->LinkedTo[0];
+			}*/
+		}
+		else if (!bFilterForCompilation && InputPins.Num() > 0)
+		{
+			const UEdGraphSchema_Niagara* Schema = CastChecked<UEdGraphSchema_Niagara>(GetSchema());
+
+			FNiagaraTypeDefinition OutDef = Schema->PinToTypeDefinition(OutPin);
+			if (OutDef != FNiagaraTypeDefinition::GetParameterMapDef())
+				continue;
+
+			for (UEdGraphPin* InputPin : InputPins)
+			{
+				FNiagaraTypeDefinition InDef = Schema->PinToTypeDefinition(InputPin);
+				if (InDef == FNiagaraTypeDefinition::GetParameterMapDef())
+				{
+					RegisterPassthroughPin(OutHistory, InputPin, OutPin, bFilterForCompilation, false);
+					return;
+				}
+			}			
+		}
+	}
 }
 
 void UNiagaraNodeStaticSwitch::AddWidgetsToOutputBox(TSharedPtr<SVerticalBox> OutputBox)
@@ -733,6 +966,8 @@ FText UNiagaraNodeStaticSwitch::GetTooltipText() const
 
 FText UNiagaraNodeStaticSwitch::GetNodeTitle(ENodeTitleType::Type TitleType) const
 {
+	if (IsSetByPin())
+		return LOCTEXT("NiagaraStaticSwitchNodePinName", "Static Switch");
 	return FText::FormatOrdered(LOCTEXT("StaticSwitchTitle", "Static Switch ({0})"), FText::FromName(IsSetByCompiler() ? SwitchTypeData.SwitchConstant : InputParameterName));
 }
 
diff --git a/Engine/Plugins/FX/Niagara/Source/NiagaraEditor/Private/NiagaraNodeStaticSwitch.h b/Engine/Plugins/FX/Niagara/Source/NiagaraEditor/Private/NiagaraNodeStaticSwitch.h
--- a/Engine/Plugins/FX/Niagara/Source/NiagaraEditor/Private/NiagaraNodeStaticSwitch.h
+++ b/Engine/Plugins/FX/Niagara/Source/NiagaraEditor/Private/NiagaraNodeStaticSwitch.h
@@ -37,6 +37,9 @@ struct FStaticSwitchTypeData
 	UPROPERTY()
 	bool bAutoRefreshEnabled = false;
 
+	UPROPERTY()
+	bool bExposeAsPin = false;
+
 	FStaticSwitchTypeData() : SwitchType(ENiagaraStaticSwitchType::Bool), Enum(nullptr)
 	{ }
 };
@@ -51,7 +54,7 @@ public:
 	UPROPERTY()
 	FName InputParameterName;
 	
-	UPROPERTY()
+	UPROPERTY(EditAnywhere, Category = "HiddenMetaData")
 	FStaticSwitchTypeData SwitchTypeData;
 
 	FNiagaraTypeDefinition GetInputType() const;
@@ -67,6 +70,9 @@ public:
 
 	bool IsDebugSwitch() const;
 
+	bool IsSetByPin() const;
+	UEdGraphPin* GetSelectorPin() const;
+
 	/** This is a hack used in the translator to check for inconsistencies with old static switches before auto refresh was a thing */
 	void CheckForOutdatedEnum(FHlslNiagaraTranslator* Translator);
 
@@ -84,15 +90,21 @@ public:
 	//~ Begin UNiagaraNode Interface
 	virtual void Compile(class FHlslNiagaraTranslator* Translator, TArray<int32>& Outputs) override;
 	virtual bool SubstituteCompiledPin(FHlslNiagaraTranslator* Translator, UEdGraphPin** LocallyOwnedPin) override;
-	virtual UEdGraphPin* GetTracedOutputPin(UEdGraphPin* LocallyOwnedOutputPin, bool bFilterForCompilation) const override;
-	virtual UEdGraphPin* GetTracedOutputPin(UEdGraphPin* LocallyOwnedOutputPin, bool bRecursive, bool bFilterForCompilation) const;
+	virtual UEdGraphPin* GetTracedOutputPin(UEdGraphPin* LocallyOwnedOutputPin, bool bFilterForCompilation, TArray<const UNiagaraNode*>* OutNodesVisitedDuringTrace = nullptr) const override;
+	virtual UEdGraphPin* GetTracedOutputPin(UEdGraphPin* LocallyOwnedOutputPin, bool bRecursive, bool bFilterForCompilation, TArray<const UNiagaraNode*>* OutNodesVisitedDuringTrace = nullptr) const;
 	virtual UEdGraphPin* GetPassThroughPin(const UEdGraphPin* LocallyOwnedOutputPin, ENiagaraScriptUsage MasterUsage) const override;
 	virtual bool AllowNiagaraTypeForAddPin(const FNiagaraTypeDefinition& InType) const override;
 	virtual void BuildParameterMapHistory(FNiagaraParameterMapHistoryBuilder& OutHistory, bool bRecursive = true, bool bFilterForCompilation = true) const override;
 	virtual void AddWidgetsToOutputBox(TSharedPtr<SVerticalBox> OutputBox) override;
 	virtual void GetNodeContextMenuActions(UToolMenu* Menu, UGraphNodeContextMenuContext* Context) const override;
+	virtual void ResolveNumerics(const UEdGraphSchema_Niagara* Schema, bool bSetInline, TMap<TPair<FGuid, UEdGraphNode*>, FNiagaraTypeDefinition>* PinCache) override;
+	virtual ENiagaraNumericOutputTypeSelectionMode GetNumericOutputTypeSelectionMode() const override;
 	//~ End UNiagaraNode Interface
 
+#if WITH_EDITOR
+	virtual void PostEditChangeProperty(struct FPropertyChangedEvent& PropertyChangedEvent) override;
+#endif // WITH_EDITOR
+
 protected:
 	//~ Begin UNiagaraNodeUsageSelector Interface
 	virtual FString GetInputCaseName(int32 Case) const override;
diff --git a/Engine/Plugins/FX/Niagara/Source/NiagaraEditor/Private/NiagaraNodeUsageSelector.cpp b/Engine/Plugins/FX/Niagara/Source/NiagaraEditor/Private/NiagaraNodeUsageSelector.cpp
--- a/Engine/Plugins/FX/Niagara/Source/NiagaraEditor/Private/NiagaraNodeUsageSelector.cpp
+++ b/Engine/Plugins/FX/Niagara/Source/NiagaraEditor/Private/NiagaraNodeUsageSelector.cpp
@@ -429,7 +429,8 @@ void UNiagaraNodeUsageSelector::BuildParameterMapHistory(FNiagaraParameterMapHis
 
 			for (int32 i = 0; i < OutputVars.Num(); i++)
 			{
-				OutHistory.VisitInputPin(InputPins[VarIdx + i], this, bFilterForCompilation);
+				//OutHistory.VisitInputPin(InputPins[VarIdx + i], this, bFilterForCompilation);
+				RegisterPassthroughPin(OutHistory, InputPins[VarIdx + i], GetOutputPin(i), bFilterForCompilation, true);
 			}
 		}
 	}
diff --git a/Engine/Plugins/FX/Niagara/Source/NiagaraEditor/Private/NiagaraNodeUsageSelector.h b/Engine/Plugins/FX/Niagara/Source/NiagaraEditor/Private/NiagaraNodeUsageSelector.h
--- a/Engine/Plugins/FX/Niagara/Source/NiagaraEditor/Private/NiagaraNodeUsageSelector.h
+++ b/Engine/Plugins/FX/Niagara/Source/NiagaraEditor/Private/NiagaraNodeUsageSelector.h
@@ -17,6 +17,9 @@ public:
 	UPROPERTY()
 	TArray<FGuid> OutputVarGuids;
 
+	UPROPERTY(meta = (SkipForCompileHash = "true"))
+	FGuid SelectorGuid;
+
 	UPROPERTY()
 	int32 NumOptionsPerVariable = 0;
 	
diff --git a/Engine/Plugins/FX/Niagara/Source/NiagaraEditor/Private/NiagaraParameterMapHistory.cpp b/Engine/Plugins/FX/Niagara/Source/NiagaraEditor/Private/NiagaraParameterMapHistory.cpp
--- a/Engine/Plugins/FX/Niagara/Source/NiagaraEditor/Private/NiagaraParameterMapHistory.cpp
+++ b/Engine/Plugins/FX/Niagara/Source/NiagaraEditor/Private/NiagaraParameterMapHistory.cpp
@@ -17,9 +17,14 @@
 #include "NiagaraParameterCollection.h"
 #include "NiagaraConstants.h"
 #include "NiagaraNodeStaticSwitch.h"
+#include "NiagaraEditorModule.h"
+#include "INiagaraEditorTypeUtilities.h"
+#include "NiagaraEditorModule.h"
 #include "NiagaraScriptVariable.h"
 #include "NiagaraNodeParameterMapSet.h"
 #include "UObject/UObjectThreadContext.h"
+#include "Modules/ModuleManager.h"
+#include "NiagaraEditorModule.h"
 
 #define LOCTEXT_NAMESPACE "NiagaraEditor"
 
@@ -39,6 +44,24 @@ static FAutoConsoleVariableRef CVarNiagaraForcePrecompilerCullDataset(
 	ECVF_Default
 );
 
+void FGraphTraversalHandle::Push(const UEdGraphNode* Node)
+{
+	Path.Push(Node->NodeGuid);
+	FriendlyPath.Push(Node->GetNodeTitle(ENodeTitleType::FullTitle).ToString());
+}
+
+void FGraphTraversalHandle::Push(const UEdGraphPin* Pin)
+{
+	Path.Push(Pin->PersistentGuid);
+	FriendlyPath.Push(Pin->GetOwningNode()->GetNodeTitle(ENodeTitleType::FullTitle).ToString() + TEXT("->") + Pin->PinName.ToString());
+}
+
+void FGraphTraversalHandle::Pop()
+{
+	Path.Pop();
+	FriendlyPath.Pop();
+}
+
 FNiagaraParameterMapHistory::FNiagaraParameterMapHistory() 
 {
 	OriginatingScriptUsage = ENiagaraScriptUsage::Function;
@@ -143,6 +166,50 @@ int32 FNiagaraParameterMapHistory::RegisterParameterMapPin(const UEdGraphPin* Pi
 	return RetIdx;
 }
 
+void FNiagaraParameterMapHistory::RegisterConstantPin(const FGraphTraversalHandle& InTraversalPath, const UEdGraphPin* InPin, const FString& InValue)
+{
+	if (UEdGraphSchema_Niagara::IsStaticPin(InPin))
+	{
+		FGraphTraversalHandle Handle = InTraversalPath;
+		Handle.Push(InPin);
+		
+		PinToConstantValues.Add(Handle, InValue);
+
+		if (UNiagaraScript::LogCompileStaticVars > 0)
+		{
+			UE_LOG(LogNiagaraEditor, Log, TEXT("RegisterConstantPin \"%s\" Pin: \"%s\""), *InValue, *Handle.ToString());
+		}
+	}
+}
+
+
+void FNiagaraParameterMapHistory::RegisterConstantVariableWrite(const FString& InValue, int32 VarIdx, bool bIsSettingDefault)
+{
+	
+	FNiagaraEditorModule& NiagaraEditorModule = FModuleManager::GetModuleChecked<FNiagaraEditorModule>("NiagaraEditor");
+
+	if (bIsSettingDefault == false && PerVariableConstantValue.Num() > VarIdx && VarIdx >= 0)
+	{
+		PerVariableConstantValue[VarIdx].AddUnique(InValue);
+
+		if (UNiagaraScript::LogCompileStaticVars > 0)
+		{
+			UE_LOG(LogNiagaraEditor, Log, TEXT("RegisterConstantVariableWrite Value \"%s\" Var: \"%s\""), *InValue, *Variables[VarIdx].GetName().ToString());
+		}
+	}
+	else if (bIsSettingDefault == true && PerVariableConstantDefaultValue.Num() > VarIdx && VarIdx >= 0)
+	{
+		PerVariableConstantDefaultValue[VarIdx].AddUnique(InValue);
+
+		if (UNiagaraScript::LogCompileStaticVars > 0)
+		{
+			UE_LOG(LogNiagaraEditor, Log, TEXT("RegisterConstantVariableWrite Default \"%s\" Var: \"%s\""), *InValue, *Variables[VarIdx].GetName().ToString());
+		}
+	}
+
+	
+}
+
 uint32 FNiagaraParameterMapHistory::BeginNodeVisitation(const UNiagaraNode* Node)
 {
 	uint32 AddedIndex = MapNodeVisitations.Add(Node);
@@ -204,6 +271,9 @@ int32 FNiagaraParameterMapHistory::AddVariable(
 		PerVariableWarnings.AddDefaulted(1);
 		PerVariableWriteHistory.AddDefaulted(1);
 		PerVariableReadHistory.AddDefaulted(1);
+		PerVariableConstantValue.AddDefaulted(1);
+		PerVariableConstantDefaultValue.AddDefaulted(1);
+
 		if (InMetaData.IsSet())
 		{
 			VariableMetaData.Add(InMetaData.GetValue());
@@ -223,6 +293,8 @@ int32 FNiagaraParameterMapHistory::AddVariable(
 	check(Variables.Num() == PerVariableWarnings.Num());
 	check(Variables.Num() == PerVariableReadHistory.Num());
 	check(Variables.Num() == PerVariableWriteHistory.Num());
+	check(Variables.Num() == PerVariableConstantValue.Num());
+	check(Variables.Num() == PerVariableConstantDefaultValue.Num());
 
 	return FoundIdx;
 }
@@ -512,18 +584,23 @@ FName FNiagaraParameterMapHistory::GetSourceForPreviousValue(const FName& InVari
 	return ReplaceVariableNamespace(InVariableName, PARAM_MAP_PREVIOUS_BASE_STR);
 }
 
-bool FNiagaraParameterMapHistory::IsPrimaryDataSetOutput(const FNiagaraVariable& InVar, const UNiagaraScript* InScript,  bool bAllowDataInterfaces) const
+bool FNiagaraParameterMapHistory::IsPrimaryDataSetOutput(const FNiagaraVariable& InVar, const UNiagaraScript* InScript,  bool bAllowDataInterfaces, bool bAllowStatics) const
 {
 	return IsPrimaryDataSetOutput(InVar, InScript->GetUsage(),  bAllowDataInterfaces);
 }
 
-bool FNiagaraParameterMapHistory::IsPrimaryDataSetOutput(const FNiagaraVariable& InVar, ENiagaraScriptUsage Usage, bool bAllowDataInterfaces) const
+bool FNiagaraParameterMapHistory::IsPrimaryDataSetOutput(const FNiagaraVariable& InVar, ENiagaraScriptUsage Usage, bool bAllowDataInterfaces, bool bAllowStatics) const
 {
 	if (bAllowDataInterfaces == false && InVar.GetType().GetClass() != nullptr)
 	{
 		return false;
 	}
 
+	if (bAllowStatics == false && InVar.GetType().IsStatic())
+	{
+		return false;
+	}
+
 	if (Usage == ENiagaraScriptUsage::EmitterSpawnScript || Usage == ENiagaraScriptUsage::EmitterUpdateScript || 
 		Usage == ENiagaraScriptUsage::SystemSpawnScript || Usage == ENiagaraScriptUsage::SystemUpdateScript)
 	{
@@ -721,6 +798,8 @@ FNiagaraParameterMapHistoryBuilder::FNiagaraParameterMapHistoryBuilder()
 {
 	ContextuallyVisitedNodes.AddDefaulted(1);
 	PinToParameterMapIndices.AddDefaulted(1);
+	PinToConstantIndices.AddDefaulted(1);
+	VariableToConstantIndices.AddDefaulted(1);
 	bFilterByScriptAllowList = false;
 	bIgnoreDisabled = true;
 	FilterScriptType = ENiagaraScriptUsage::Function;
@@ -734,6 +813,10 @@ FNiagaraParameterMapHistoryBuilder::FNiagaraParameterMapHistoryBuilder()
 
 void FNiagaraParameterMapHistoryBuilder::BuildParameterMaps(const UNiagaraNodeOutput* OutputNode, bool bRecursive)
 {
+	if (UNiagaraScript::LogCompileStaticVars > 0)
+	{
+		UE_LOG(LogNiagaraEditor, Log, TEXT("************** BuildParameterMaps %d  **************"), OutputNode->GetUsage());
+	}
 	TOptional<FName> StackContextAlias = OutputNode->GetStackContextOverride();
 	BeginUsage(OutputNode->GetUsage(), StackContextAlias.Get(ScriptUsageContextNameStack.Num() > 0 ? ScriptUsageContextNameStack.Top() : NAME_None));
 
@@ -828,11 +911,23 @@ void FNiagaraParameterMapHistoryBuilder::RegisterDataSetWrite(int32 WhichParamet
 
 int32 FNiagaraParameterMapHistoryBuilder::RegisterParameterMapPin(int32 WhichParameterMap, const UEdGraphPin* Pin)
 {
+	if (Pin && UNiagaraScript::LogCompileStaticVars > 0)
+	{
+		UE_LOG(LogNiagaraEditor, Log, TEXT("FNiagaraParameterMapHistoryBuilder::RegisterParameterMapPin %d > %s .... %p"), WhichParameterMap, *Pin->GetOwningNode()->GetNodeTitle(ENodeTitleType::FullTitle).ToString(), Pin);
+	}
 	if (WhichParameterMap != INDEX_NONE)
 	{
 		if (Pin->Direction == EEdGraphPinDirection::EGPD_Output)
 		{
 			PinToParameterMapIndices.Last().Add(Pin, WhichParameterMap);
+			if (UNiagaraScript::LogCompileStaticVars > 0)
+			{
+				UE_LOG(LogNiagaraEditor, Log, TEXT("FNiagaraParameterMapHistoryBuilder::RegisterParameterMapPin Added %d"), WhichParameterMap);
+			}
+		}
+		else if (UNiagaraScript::LogCompileStaticVars > 0)
+		{
+			UE_LOG(LogNiagaraEditor, Log, TEXT("FNiagaraParameterMapHistoryBuilder::RegisterParameterMapPin Not Output %p"), Pin);
 		}
 
 		return Histories[WhichParameterMap].RegisterParameterMapPin(Pin);
@@ -843,11 +938,175 @@ int32 FNiagaraParameterMapHistoryBuilder::RegisterParameterMapPin(int32 WhichPar
 	}
 }
 
+int32 FNiagaraParameterMapHistoryBuilder::RegisterConstantPin(int32 WhichConstant, const UEdGraphPin* Pin)
+{
+	if (WhichConstant != INDEX_NONE)
+	{
+		PinToConstantIndices.Last().Add(Pin, WhichConstant);
+
+		if (Histories.Num() > 0)
+			Histories.Last().RegisterConstantPin(ActivePath, Pin, Constants[WhichConstant]);
+
+		return WhichConstant;
+	}
+	else
+	{
+		return INDEX_NONE;
+	}
+}
+
+int32 FNiagaraParameterMapHistoryBuilder::RegisterConstantVariableWrite(int32 WhichParamMapIdx, int32 WhichConstant,  int32 WhichVarIdx, bool bIsSettingDefault)
+{
+
+	FNiagaraEditorModule& NiagaraEditorModule = FModuleManager::GetModuleChecked<FNiagaraEditorModule>("NiagaraEditor");
+	if (WhichConstant != INDEX_NONE)
+	{
+		const FNiagaraVariableBase& Var = Histories[WhichParamMapIdx].Variables[WhichVarIdx];
+		if (Var.IsValid())
+		{
+			VariableToConstantIndices.Last().Add(Var, WhichConstant);
+
+			if (WhichParamMapIdx != INDEX_NONE)
+			{
+				Histories[WhichParamMapIdx].RegisterConstantVariableWrite(Constants[WhichConstant], WhichVarIdx, bIsSettingDefault);
+			}
+
+			if (Constants[WhichConstant].Len() != 0 && Var.GetType().IsStatic())
+			{
+				TSharedPtr<INiagaraEditorTypeUtilities, ESPMode::ThreadSafe> TypeEditorUtilities = NiagaraEditorModule.GetTypeUtilities(Var.GetType());
+				if (TypeEditorUtilities.IsValid() && TypeEditorUtilities->CanHandlePinDefaults())
+				{
+					FNiagaraVariable VarWithValue = Var;
+					VarWithValue.AllocateData();
+					if (TypeEditorUtilities->SetValueFromPinDefaultString(Constants[WhichConstant], VarWithValue))
+					{
+						StaticVariables.AddUnique(VarWithValue);
+					}
+				}
+			}
+		}
+
+		return WhichConstant;
+	}
+	else
+	{
+		return INDEX_NONE;
+	}
+}
+
+
+FString FNiagaraParameterMapHistoryBuilder::GetConstant(int32 InIndex)
+{
+	if (Constants.IsValidIndex(InIndex))
+		return Constants[InIndex];
+	return FString();
+}
+
+int32 FNiagaraParameterMapHistoryBuilder::AddOrGetConstantFromValue(const FString& Value)
+{
+	return Constants.AddUnique(Value);
+}
+
+int32 FNiagaraParameterMapHistoryBuilder::RegisterConstantFromInputPin(const UEdGraphPin* InputPin) 
+{
+	int32 ConstantIdx = INDEX_NONE;
+	if (InputPin && InputPin->Direction == EEdGraphPinDirection::EGPD_Input)
+	{
+		FString Value;
+		if (InputPin && InputPin->LinkedTo.Num() > 0 && InputPin->LinkedTo[0] != nullptr)
+		{
+			UEdGraphPin* ConnectedOutputPin = InputPin->LinkedTo[0];
+			ConstantIdx = GetConstantFromOutputPin(ConnectedOutputPin);
+			RegisterConstantPin(ConstantIdx, InputPin);
+		}
+		else if (InputPin && InputPin->Direction == EEdGraphPinDirection::EGPD_Input && InputPin->LinkedTo.Num() == 0 && InputPin->DefaultValue.Len() != 0)
+		{
+			Value = InputPin->DefaultValue;
+			ConstantIdx = AddOrGetConstantFromValue(Value);
+			RegisterConstantPin(ConstantIdx, InputPin);
+		}
+	}
+
+	if (UNiagaraScript::LogCompileStaticVars > 0)
+	{
+		FGraphTraversalHandle TestPath = ActivePath;
+		TestPath.Push(InputPin);
+		UE_LOG(LogNiagaraEditor, Log, TEXT("RegisterConstantFromInputPin Value \"%s\" Var: \"%s\""), ConstantIdx != INDEX_NONE ? *Constants[ConstantIdx] : TEXT("INDEX_NONE"), *TestPath.ToString());
+	}
+
+	return ConstantIdx;
+}
+
+int32 FNiagaraParameterMapHistoryBuilder::GetConstantFromInputPin(const UEdGraphPin* InputPin) const
+{
+	if (InputPin && InputPin->Direction == EEdGraphPinDirection::EGPD_Input)
+	{
+		if (InputPin)
+		{
+			const int32* IdxPtr = PinToConstantIndices.Last().Find(InputPin);
+			if (IdxPtr != nullptr)
+			{
+				if (UNiagaraScript::LogCompileStaticVars > 0)
+				{
+					FGraphTraversalHandle TestPath = ActivePath;
+					TestPath.Push(InputPin);
+					UE_LOG(LogNiagaraEditor, Log, TEXT("GetConstantFromInputPin Value \"%s\" Var: \"%s\""), *Constants[*IdxPtr], *TestPath.ToString());
+				}
+				return *IdxPtr;
+			}
+		}
+	}
+	return INDEX_NONE;
+}
+
+int32 FNiagaraParameterMapHistoryBuilder::GetConstantFromOutputPin(const UEdGraphPin* OutputPin) const
+{
+	if (OutputPin && OutputPin->Direction == EEdGraphPinDirection::EGPD_Output)
+	{
+		if (OutputPin)
+		{
+			const int32* IdxPtr = PinToConstantIndices.Last().Find(OutputPin);
+			if (IdxPtr != nullptr)
+			{
+				if (UNiagaraScript::LogCompileStaticVars > 0)
+				{
+					FGraphTraversalHandle TestPath = ActivePath;
+					TestPath.Push(OutputPin);
+					UE_LOG(LogNiagaraEditor, Log, TEXT("GetConstantFromOutputPin Value \"%s\" Var: \"%s\""), *Constants[*IdxPtr], *TestPath.ToString());
+				}
+				return *IdxPtr;
+			}
+		}
+	}
+	return INDEX_NONE;
+}
+
+
+int32 FNiagaraParameterMapHistoryBuilder::GetConstantFromVariableRead(const FNiagaraVariableBase& InVar)
+{
+	if (InVar.IsValid())
+	{
+		for (int32 i = VariableToConstantIndices.Num() - 1; i >= 0; i--)
+		{
+			const int32* IdxPtr = VariableToConstantIndices[i].Find(InVar);
+			if (IdxPtr != nullptr)
+			{
+				if (UNiagaraScript::LogCompileStaticVars > 0)
+				{
+					UE_LOG(LogNiagaraEditor, Log, TEXT("GetConstantFromVariableRead Value \"%s\" Var: \"%s\" Path: \"%s\""), *Constants[*IdxPtr], *InVar.GetName().ToString(),  *ActivePath.ToString());
+				}
+				return *IdxPtr;
+			}
+		}
+		
+	}
+	return INDEX_NONE;
+}
+
 int32 FNiagaraParameterMapHistoryBuilder::TraceParameterMapOutputPin(const UEdGraphPin* OutputPin)
 {
 	if (OutputPin && OutputPin->Direction == EEdGraphPinDirection::EGPD_Output)
 	{
-		OutputPin = UNiagaraNode::TraceOutputPin(const_cast<UEdGraphPin*>(OutputPin));
 		if (OutputPin)
 		{
 			const int32* IdxPtr = PinToParameterMapIndices.Last().Find(OutputPin);
@@ -892,7 +1151,7 @@ int32 FNiagaraParameterMapHistoryBuilder::FindMatchingParameterMapFromContextInp
 		{
 			if (Inputs[i]->LinkedTo.Num() != 0 && PinToParameterMapIndices.Num() >= 2)
 			{
-				const UEdGraphPin* OutputPin = UNiagaraNode::TraceOutputPin(const_cast<UEdGraphPin*>(Inputs[i]->LinkedTo[0]));
+				const UEdGraphPin* OutputPin = Inputs[i]->LinkedTo[0];
 				const int32* ParamMapIdxPtr = PinToParameterMapIndices[PinToParameterMapIndices.Num() - 2].Find(OutputPin);
 				if (ParamMapIdxPtr != nullptr)
 				{
@@ -934,6 +1193,35 @@ int32 FNiagaraParameterMapHistoryBuilder::FindMatchingParameterMapFromContextInp
 	return INDEX_NONE;
 }
 
+int32 FNiagaraParameterMapHistoryBuilder::FindMatchingStaticFromContextInputs(const FNiagaraVariable& InVar) const
+{
+	if (CallingContext.Num() == 0)
+	{
+		return INDEX_NONE;
+	}
+	const UNiagaraNode* Node = CallingContext.Last();
+	FPinCollectorArray Inputs;
+	Node->GetInputPins(Inputs);
+	const UEdGraphSchema_Niagara* Schema = CastChecked<UEdGraphSchema_Niagara>(Node->GetSchema());
+
+	for (int32 i = 0; i < Inputs.Num(); i++)
+	{
+		FNiagaraVariable CallInputVar = Schema->PinToNiagaraVariable(Inputs[i]);
+		if (CallInputVar.IsEquivalent(InVar))
+		{
+			if (Inputs[i]->LinkedTo.Num() != 0 && PinToConstantIndices.Num() >= 2)
+			{
+				const UEdGraphPin* OutputPin = Inputs[i]->LinkedTo[0];
+				//int32 ConstantIdx = GetConstantFromOutputPin(OutputPin);
+				const int32* IdxPtr = PinToConstantIndices[PinToConstantIndices.Num() - 2].Find(OutputPin);
+				if (IdxPtr != nullptr)
+					return *IdxPtr;
+			}
+		}
+	}
+	return INDEX_NONE;
+}
+
 void FNiagaraParameterMapHistoryBuilder::BeginTranslation(const UNiagaraScript* Script)
 {
 	//For now this will just tell particle scripts what emitter they're being compiled as part of but maybe we want to do more here.
@@ -1043,8 +1331,10 @@ void FNiagaraParameterMapHistoryBuilder::EnterFunction(const FString& InNodeName
 	{
 		RegisterNodeVisitation(Node);
 		CallingContext.Push(Node);
+		ActivePath.Push(Node);
 		CallingGraphContext.Push(InGraph);
 		PinToParameterMapIndices.Emplace();
+		PinToConstantIndices.Emplace();
 		FunctionNameContextStack.Emplace(*InNodeName);
 		BuildCurrentAliases();
 		if (EncounteredFunctionNames.Num() != 0)
@@ -1061,8 +1351,10 @@ void FNiagaraParameterMapHistoryBuilder::ExitFunction(const FString& InNodeName,
 	{
 		CallingContext.Pop();
 		CallingGraphContext.Pop();
+		ActivePath.Pop();
 		PinToParameterMapIndices.Pop();
 		FunctionNameContextStack.Pop();
+		PinToConstantIndices.Pop();
 		BuildCurrentAliases();
 		ContextuallyVisitedNodes.Pop();
 	}
@@ -1072,6 +1364,7 @@ void FNiagaraParameterMapHistoryBuilder::EnterEmitter(const FString& InEmitterNa
 {
 	RegisterNodeVisitation(Node);
 	CallingContext.Push(Node);
+	ActivePath.Push(Node);
 	CallingGraphContext.Push(InGraph);
 	EmitterNameContextStack.Emplace(*InEmitterName);
 	BuildCurrentAliases();
@@ -1104,6 +1397,7 @@ void FNiagaraParameterMapHistoryBuilder::ExitEmitter(const FString& InEmitterNam
 	CallingContext.Pop();
 	CallingGraphContext.Pop();
 	EmitterNameContextStack.Pop();
+	ActivePath.Pop();
 	BuildCurrentAliases();
 	ContextuallyVisitedNodes.Pop();
 	EncounteredFunctionNames.Pop();
@@ -1175,36 +1469,37 @@ void FNiagaraParameterMapHistoryBuilder::VisitInputPin(const UEdGraphPin* Pin, c
 	{
 		for (int32 j = 0; j < Pin->LinkedTo.Num(); j++)
 		{
-			UEdGraphPin* LinkedPin = Pin->LinkedTo[j];
-			const UEdGraphPin* OutputPin;
-			if (!bFilterForCompilation && LinkedPin && Cast<UNiagaraNodeStaticSwitch>(LinkedPin->GetOwningNode()))
-			{
-				// if we do not filter the parameter maps we want all nodes connected to the static switch, regardless of the switch value
-				OutputPin = LinkedPin;
-			}
-			else
-			{
-				OutputPin = UNiagaraNode::TraceOutputPin(LinkedPin, bFilterForCompilation);
-				if (bFilterForCompilation && OutputPin && Cast<UNiagaraNodeStaticSwitch>(OutputPin->GetOwningNode()))
-				{
-					// this means there is no further node connected to this pin of the static switch, so it is not relevant
-					continue;
-				}
-			}
+			UEdGraphPin* LinkedPin = Pin->LinkedTo[j];	
 
-			if (OutputPin)
+			if (LinkedPin)
 			{
-				UNiagaraNode* Node = CastChecked<UNiagaraNode>(OutputPin->GetOwningNode());
+				UNiagaraNode* Node = CastChecked<UNiagaraNode>(LinkedPin->GetOwningNode());
 
 				if (!GetNodePreviouslyVisited(Node))
 				{
+				
+					if (UNiagaraScript::LogCompileStaticVars > 0)
+					{
+						UE_LOG(LogNiagaraEditor, Log, TEXT("Build Parameter Map History: %s %s"), *Node->GetClass()->GetName(), *Node->GetNodeTitle(ENodeTitleType::FullTitle).ToString());
+					}
 					Node->BuildParameterMapHistory(*this, true, bFilterForCompilation);
 					RegisterNodeVisitation(Node);
 				}
 
 				if (Schema->PinToTypeDefinition(Pin) == FNiagaraTypeDefinition::GetParameterMapDef())
 				{
-					int32 ParamMapIdx = TraceParameterMapOutputPin(OutputPin);
+					int32 ParamMapIdx = TraceParameterMapOutputPin(LinkedPin);
+					
+					if (UNiagaraScript::LogCompileStaticVars > 0)
+					{
+						if (ParamMapIdx == INDEX_NONE)
+						{
+							UE_LOG(LogNiagaraEditor, Log, TEXT("Found bad pin!"));
+						}
+
+						UE_LOG(LogNiagaraEditor, Log, TEXT("Build Parameter Map History: %s %s PMapIdx: %d"), *Node->GetClass()->GetName(), *Node->GetNodeTitle(ENodeTitleType::FullTitle).ToString(), ParamMapIdx);
+					}
+
 					RegisterParameterMapPin(ParamMapIdx, Pin);
 				}
 			}
@@ -1253,7 +1548,29 @@ int32 FNiagaraParameterMapHistoryBuilder::HandleVariableWrite(int32 ParamMapIdx,
 	FNiagaraVariable AliasedVar = Var;
 	Var = ResolveAliases(Var);
 
-	return AddVariableToHistory(Histories[ParamMapIdx], Var, AliasedVar, InPin);
+	int32 FoundIdx = AddVariableToHistory(Histories[ParamMapIdx], Var, AliasedVar, InPin);
+	if (InPin && Var.GetType().IsStatic() && InPin->LinkedTo.Num() > 0)
+	{
+		int32 ConstantIdx = GetConstantFromOutputPin(InPin->LinkedTo[0]);
+		RegisterConstantPin(ConstantIdx, InPin);
+
+		if (FoundIdx != INDEX_NONE && ParamMapIdx != INDEX_NONE && ConstantIdx != INDEX_NONE)
+		{
+			RegisterConstantVariableWrite(ParamMapIdx, ConstantIdx, FoundIdx, false);
+		}
+	}
+	else if (InPin && Var.GetType().IsStatic() && InPin->LinkedTo.Num() == 0)
+	{
+		int32 ConstantIdx = AddOrGetConstantFromValue(InPin->DefaultValue);
+		RegisterConstantPin(ConstantIdx, InPin);
+
+		if (FoundIdx != INDEX_NONE && ParamMapIdx != INDEX_NONE && ConstantIdx != INDEX_NONE)
+		{
+			RegisterConstantVariableWrite(ParamMapIdx, ConstantIdx, FoundIdx, false);
+		}
+	}
+
+	return FoundIdx;
 }
 
 int32 FNiagaraParameterMapHistoryBuilder::HandleVariableWrite(int32 ParameterMapIndex, const FNiagaraVariable& Var)
@@ -1266,11 +1583,13 @@ int32 FNiagaraParameterMapHistoryBuilder::HandleVariableWrite(int32 ParameterMap
 	}
 	FNiagaraVariable ResolvedVar = ResolveAliases(Var);
 
-	return AddVariableToHistory(Histories[ParameterMapIndex], ResolvedVar, Var, nullptr);
+	int32 FoundIdx = AddVariableToHistory(Histories[ParameterMapIndex], ResolvedVar, Var, nullptr);
+	return FoundIdx;
 }
 
 int32 FNiagaraParameterMapHistoryBuilder::HandleVariableRead(int32 ParamMapIdx, const UEdGraphPin* InPin, bool RegisterReadsAsVariables, const UEdGraphPin* InDefaultPin, bool bFilterForCompilation, bool& OutUsedDefault)
 {
+	FString DefaultValue;
 	OutUsedDefault = false;
 	const UEdGraphSchema_Niagara* Schema = GetDefault<UEdGraphSchema_Niagara>();
 	FNiagaraVariable Var = Schema->PinToNiagaraVariable(InPin);
@@ -1301,15 +1620,19 @@ int32 FNiagaraParameterMapHistoryBuilder::HandleVariableRead(int32 ParamMapIdx,
 	FName ModuleName = ModuleAlias ? FName(**ModuleAlias) : NAME_None;
 
 	bool AddWriteHistory = true;
+	bool bUseRapidIterationParam = true;
 
 	int32 FoundIdx = History.FindVariable(Var.GetName(), Var.GetType());
 	if (FoundIdx == -1)
 	{
 		if (RegisterReadsAsVariables)
 		{
+			OutUsedDefault = true;
+
 			if (InDefaultPin)
 			{
 				VisitInputPin(InDefaultPin, Cast<UNiagaraNode>(InDefaultPin->GetOwningNode()), bFilterForCompilation);
+				DefaultValue = InDefaultPin->DefaultValue;
 			}
 
 			FoundIdx = AddVariableToHistory(History, Var, AliasedVar, InDefaultPin);
@@ -1329,6 +1652,9 @@ int32 FNiagaraParameterMapHistoryBuilder::HandleVariableRead(int32 ParamMapIdx,
 					}
 					int32 FoundIdxBinding = AddVariableToHistory(History, TempVar, TempVar, nullptr);
 
+					DefaultValue = History.Variables[FoundIdxBinding].GetName().ToString();
+					bUseRapidIterationParam = false;
+
 					History.PerVariableReadHistory[FoundIdxBinding].AddDefaulted_GetRef().ReadPin = FModuleScopedPin(InDefaultPin, ModuleName);
 				}
 			}
@@ -1352,6 +1678,103 @@ int32 FNiagaraParameterMapHistoryBuilder::HandleVariableRead(int32 ParamMapIdx,
 	check(History.Variables.Num() == History.PerVariableWarnings.Num());
 	check(History.Variables.Num() == History.PerVariableWriteHistory.Num());
 	check(History.Variables.Num() == History.PerVariableReadHistory.Num());
+	check(History.Variables.Num() == History.PerVariableConstantValue.Num());
+	check(History.Variables.Num() == History.PerVariableConstantDefaultValue.Num());
+
+	if (OutUsedDefault && InDefaultPin && UEdGraphSchema_Niagara::IsStaticPin(InDefaultPin))
+	{
+		int32 ConstantIdx = INDEX_NONE;
+		ENiagaraScriptUsage Usage = GetCurrentUsageContext();
+
+		FString EmitterNamespace;
+		if (!UNiagaraScript::IsSystemScript(Usage) && Histories[ParamMapIdx].EmitterNamespacesEncountered.Num() > 0)
+		{
+			EmitterNamespace = Histories[ParamMapIdx].EmitterNamespacesEncountered.Last();
+		}
+		else if (!UNiagaraScript::IsSystemScript(Usage) )
+		{
+			if (EmitterNamespace.IsEmpty() && EmitterNameContextStack.Num() > 0)
+				EmitterNamespace = EmitterNameContextStack.Last().ToString();
+			else
+				EmitterNamespace = TEXT("Emitter");
+		}
+
+		if (bUseRapidIterationParam && /*InTopLevelFunctionCall(Usage) && */FNiagaraParameterMapHistory::IsAliasedModuleParameter(Histories[ParamMapIdx].VariablesWithOriginalAliasesIntact[FoundIdx]))
+		{					
+			FNiagaraVariable VarRapid = FNiagaraParameterMapHistory::ConvertVariableToRapidIterationConstantName(Histories[ParamMapIdx].Variables[FoundIdx], EmitterNamespace.Len() == 0 ? nullptr : *EmitterNamespace, Usage);
+			
+			int32 StaticIdx = StaticVariables.IndexOfByPredicate([&](const FNiagaraVariable& InObj) -> bool
+				{
+					return (InObj.GetName() == VarRapid.GetName());
+				});;
+
+			if (ConstantResolver.ResolveConstant(VarRapid))
+			{
+				DefaultValue = VarRapid.GetType().ToString(VarRapid.GetData());
+			}
+			else if (StaticIdx != INDEX_NONE)
+			{
+				DefaultValue = StaticVariables[StaticIdx].GetType().ToString(StaticVariables[StaticIdx].GetData());	
+			}
+			else if (DefaultValue.Len() == 0)
+			{
+				//DefaultValue = VarRapid.GetName().ToString();
+				FNiagaraVariable VarDefault = Var;
+				FNiagaraEditorUtilities::ResetVariableToDefaultValue(VarDefault);// Do this to handle non-zero defaults
+				DefaultValue = VarDefault.GetType().ToString(VarDefault.GetData());
+			}
+			
+			if (UNiagaraScript::LogCompileStaticVars > 0)
+			{
+				FGraphTraversalHandle TestPath = ActivePath;
+				TestPath.Push(InPin);
+				UE_LOG(LogNiagaraEditor, Log, TEXT("HandleVariableRead RapidIt Value \"%s\" Var: \"%s\" Path: \"%s\""), *DefaultValue, *VarRapid.GetName().ToString(),  *TestPath.ToString());
+			}
+			
+
+		}
+		else if (FNiagaraParameterMapHistory::IsExternalConstantNamespace(Var, Usage, 0) || Var.IsInNameSpace(EmitterNamespace) || Var.IsInNameSpace(FNiagaraConstants::SystemNamespace))
+		{
+			const FNiagaraVariable& VarRapid = Histories[ParamMapIdx].Variables[FoundIdx];
+			int32 StaticIdx = StaticVariables.IndexOfByPredicate([&](const FNiagaraVariable& InObj) -> bool
+				{
+					return (InObj.GetName() == VarRapid.GetName());
+				});;
+
+			if (ConstantResolver.ResolveConstant(Var))
+			{
+				DefaultValue = Var.GetType().ToString(Var.GetData());
+			}
+			else if (StaticIdx != INDEX_NONE)
+			{
+				DefaultValue = StaticVariables[StaticIdx].GetType().ToString(StaticVariables[StaticIdx].GetData());
+			}
+			else
+			{
+				DefaultValue = Var.GetName().ToString();
+			}
+			
+			if (UNiagaraScript::LogCompileStaticVars > 0)
+			{
+				FGraphTraversalHandle TestPath = ActivePath;
+				TestPath.Push(InPin);
+				UE_LOG(LogNiagaraEditor, Log, TEXT("HandleVariableRead Ext Value \"%s\" Var: \"%s\" Path: \"%s\""), *DefaultValue, *VarRapid.GetName().ToString(),  *TestPath.ToString());
+			}
+		}
+
+		ConstantIdx = AddOrGetConstantFromValue(DefaultValue);
+
+		RegisterConstantPin(ConstantIdx, InPin);
+		if (FoundIdx != INDEX_NONE && ParamMapIdx != INDEX_NONE && ConstantIdx != INDEX_NONE)
+		{
+			RegisterConstantVariableWrite(ParamMapIdx, ConstantIdx, FoundIdx, true);
+		}
+	}
+	else if (FoundIdx != INDEX_NONE && ParamMapIdx != INDEX_NONE && UEdGraphSchema_Niagara::IsStaticPin(InPin))
+	{
+		int32 ConstantIdx = GetConstantFromVariableRead(Histories[ParamMapIdx].Variables[FoundIdx]);
+		RegisterConstantPin(ConstantIdx, InPin);
+	}
 
 	return FoundIdx;
 }
@@ -1361,6 +1784,21 @@ void FNiagaraParameterMapHistoryBuilder::RegisterEncounterableVariables(const TA
 	EncounterableExternalVariables.Append(Variables);
 }
 
+void FNiagaraParameterMapHistoryBuilder::RegisterExternalStaticVariables(const TArray<FNiagaraVariable>& Variables)
+{
+	StaticVariables.Append(Variables);
+
+	if (UNiagaraScript::LogCompileStaticVars > 0)
+	{
+		UE_LOG(LogNiagaraEditor, Log, TEXT("FNiagaraParameterMapHistoryBuilder::RegisterExternalStaticVariables ==================================="));
+		for (int32 i = 0; i < StaticVariables.Num(); i++)
+		{
+			UE_LOG(LogNiagaraEditor, Log, TEXT("[%d] %s"), i, *StaticVariables[i].ToString());
+		}
+		UE_LOG(LogNiagaraEditor, Log, TEXT("  ==================================="));
+	}
+}
+
 void UpdateAliasedVariable(FNiagaraVariable& AliasedVar, const FNiagaraVariable& OriginalUnaliasedVar, const FNiagaraVariable& UpdatedUnaliasedVar)
 {
 	AliasedVar.SetType(UpdatedUnaliasedVar.GetType());
@@ -1471,6 +1909,104 @@ int32 FNiagaraParameterMapHistoryBuilder::HandleExternalVariableRead(int32 Param
 }
 
 
+void FNiagaraParameterMapHistoryBuilder::SetConstantByStaticVariable(int32& OutValue, const UEdGraphPin* InDefaultPin)
+{
+	const UEdGraphSchema_Niagara* Schema = GetDefault<UEdGraphSchema_Niagara>();
+
+	FNiagaraEditorModule& NiagaraEditorModule = FModuleManager::GetModuleChecked<FNiagaraEditorModule>("NiagaraEditor");
+	OutValue = 0;
+	FNiagaraVariable Var = Schema->PinToNiagaraVariable(InDefaultPin, true);
+	FNiagaraVariable VarDefault = Var;
+	FNiagaraEditorUtilities::ResetVariableToDefaultValue(VarDefault);// Do this to handle non-zero defaults
+	if (VarDefault.GetType().IsStatic())
+	{
+		FNiagaraVariable VarWithValue = FNiagaraVariable(Var.GetType(), Var.GetName());
+		FString Value;
+		if (InDefaultPin && InDefaultPin->Direction == EEdGraphPinDirection::EGPD_Input && InDefaultPin->LinkedTo.Num() > 0 && InDefaultPin->LinkedTo[0] != nullptr)
+		{
+			UEdGraphPin* ConnectedOutputPin = InDefaultPin->LinkedTo[0];
+			int32 InputConstant = GetConstantFromOutputPin(ConnectedOutputPin);
+			if (InputConstant != INDEX_NONE)
+			{
+				Value = Constants[InputConstant];
+			}
+		}
+		else if (InDefaultPin && InDefaultPin->Direction == EEdGraphPinDirection::EGPD_Input && InDefaultPin->LinkedTo.Num() == 0 && InDefaultPin->DefaultValue.Len() != 0)
+		{
+			Value = InDefaultPin->DefaultValue;
+		}
+
+		// If we found a string, we should try and map to the actual value of that variable..
+		if (Value.Len() != 0)
+		{
+			TSharedPtr<INiagaraEditorTypeUtilities, ESPMode::ThreadSafe> TypeEditorUtilities = NiagaraEditorModule.GetTypeUtilities(Var.GetType());
+			if (TypeEditorUtilities.IsValid() && TypeEditorUtilities->CanHandlePinDefaults() )
+			{
+				TypeEditorUtilities->SetValueFromPinDefaultString(Value, VarWithValue);
+			}
+		}
+
+		if (VarWithValue.IsDataAllocated())
+		{
+			if (VarWithValue.GetType().IsSameBaseDefinition(FNiagaraTypeDefinition::GetBoolDef()))
+			{
+				OutValue = VarWithValue.GetValue<bool>();
+			}
+			else if (VarWithValue.GetType().IsSameBaseDefinition(FNiagaraTypeDefinition::GetIntDef()) || VarWithValue.GetType().IsEnum())
+			{
+				OutValue = VarWithValue.GetValue<int32>();
+			}
+		}
+	}
+}
+
+void FNiagaraParameterMapHistoryBuilder::SetConstantByStaticVariable(int32& OutValue, const FNiagaraVariable& Var)
+{
+	
+	const UEdGraphSchema_Niagara* Schema = GetDefault<UEdGraphSchema_Niagara>();
+
+	OutValue = 0;
+	FNiagaraVariable VarDefault = Var;
+	FNiagaraEditorUtilities::ResetVariableToDefaultValue(VarDefault);// Do this to handle non-zero defaults
+	if (VarDefault.GetType().IsStatic())
+	{
+		FNiagaraVariable VarWithValue = FNiagaraVariable(Var.GetType(), Var.GetName());
+		FString Value = Var.GetName().ToString();
+		int32 Input = INDEX_NONE;
+
+		// If we found a string, we should try and map to the actual value of that variable..
+		if (Value.Len() != 0)
+		{
+			if (!VarWithValue.IsDataAllocated())
+			{
+				int32 FoundOverrideIdx = StaticVariables.IndexOfByPredicate([&](const FNiagaraVariable& InObj) -> bool
+					{
+						return (InObj.GetName() == *Value);
+					});;
+
+				if (FoundOverrideIdx != INDEX_NONE)
+				{
+					VarWithValue.SetData(StaticVariables[FoundOverrideIdx].GetData());
+				}
+			}
+		}
+
+		if (VarWithValue.IsDataAllocated())
+		{
+			if (VarWithValue.GetType().IsSameBaseDefinition(FNiagaraTypeDefinition::GetBoolDef()))
+			{
+				OutValue = VarWithValue.GetValue<bool>();
+			}
+			else if (VarWithValue.GetType().IsSameBaseDefinition(FNiagaraTypeDefinition::GetIntDef()) || VarWithValue.GetType().IsEnum())
+			{
+				OutValue = VarWithValue.GetValue<int32>();
+			}
+		}
+	}
+	
+}
+
+
 int32 FNiagaraParameterMapHistoryBuilder::AddVariableToHistory(FNiagaraParameterMapHistory& History, const FNiagaraVariable& InVar, const FNiagaraVariable& InAliasedVar, const UEdGraphPin* InPin)
 {
 	const FString* ModuleAlias = GetModuleAlias();
@@ -1596,26 +2132,26 @@ bool FCompileConstantResolver::ResolveConstant(FNiagaraVariable& OutConstant) co
 		OutConstant.SetValue(Emitter->bLocalSpace ? FNiagaraBool(true) : FNiagaraBool(false));
 		return true;
 	}
-	if (Emitter && OutConstant == FNiagaraVariable(FNiagaraTypeDefinition::GetBoolDef(), TEXT("Emitter.Determinism")))
+	else if (Emitter && OutConstant == FNiagaraVariable(FNiagaraTypeDefinition::GetBoolDef(), TEXT("Emitter.Determinism")))
 	{
 		OutConstant.SetValue(Emitter->bDeterminism ? FNiagaraBool(true) : FNiagaraBool(false));
 		return true;
 	}
-	if (Emitter && OutConstant == FNiagaraVariable(FNiagaraTypeDefinition::GetSimulationTargetEnum(), TEXT("Emitter.SimulationTarget")))
+	else if (Emitter && OutConstant == FNiagaraVariable(FNiagaraTypeDefinition::GetSimulationTargetEnum(), TEXT("Emitter.SimulationTarget")))
 	{
 		FNiagaraInt32 EnumValue;
 		EnumValue.Value = (uint8)Emitter->SimTarget;
 		OutConstant.SetValue(EnumValue);
 		return true;
 	}
-	if (Emitter && OutConstant == FNiagaraVariable(FNiagaraTypeDefinition::GetScriptUsageEnum(), TEXT("Script.Usage")))
+	else if (Emitter && OutConstant == FNiagaraVariable(FNiagaraTypeDefinition::GetScriptUsageEnum(), TEXT("Script.Usage")))
 	{
 		FNiagaraInt32 EnumValue;
 		EnumValue.Value = (uint8)FNiagaraUtilities::ConvertScriptUsageToStaticSwitchUsage(Usage);
 		OutConstant.SetValue(EnumValue);
 		return true;
 	}
-	if (Emitter && OutConstant == FNiagaraVariable(FNiagaraTypeDefinition::GetScriptContextEnum(), TEXT("Script.Context")))
+	else if (Emitter && OutConstant == FNiagaraVariable(FNiagaraTypeDefinition::GetScriptContextEnum(), TEXT("Script.Context")))
 	{
 		FNiagaraInt32 EnumValue;
 		EnumValue.Value = (uint8)FNiagaraUtilities::ConvertScriptUsageToStaticSwitchContext(Usage);
@@ -1629,19 +2165,19 @@ bool FCompileConstantResolver::ResolveConstant(FNiagaraVariable& OutConstant) co
 		OutConstant.SetValue(FNiagaraBool(false));
 		return true;
 	}
-	if (System && OutConstant == FNiagaraVariable(FNiagaraTypeDefinition::GetBoolDef(), TEXT("Emitter.Determinism")))
+	else if (System && OutConstant == FNiagaraVariable(FNiagaraTypeDefinition::GetBoolDef(), TEXT("Emitter.Determinism")))
 	{
 		OutConstant.SetValue(FNiagaraBool(true));
 		return true;
 	}
-	if (System && OutConstant == FNiagaraVariable(FNiagaraTypeDefinition::GetScriptUsageEnum(), TEXT("Script.Usage")))
+	else if (System && OutConstant == FNiagaraVariable(FNiagaraTypeDefinition::GetScriptUsageEnum(), TEXT("Script.Usage")))
 	{
 		FNiagaraInt32 EnumValue;
 		EnumValue.Value = (uint8)FNiagaraUtilities::ConvertScriptUsageToStaticSwitchUsage(Usage);
 		OutConstant.SetValue(EnumValue);
 		return true;
 	}
-	if (System && OutConstant == FNiagaraVariable(FNiagaraTypeDefinition::GetScriptContextEnum(), TEXT("Script.Context")))
+	else if (System && OutConstant == FNiagaraVariable(FNiagaraTypeDefinition::GetScriptContextEnum(), TEXT("Script.Context")))
 	{
 		FNiagaraInt32 EnumValue;
 		EnumValue.Value = (uint8)FNiagaraUtilities::ConvertScriptUsageToStaticSwitchContext(Usage);
diff --git a/Engine/Plugins/FX/Niagara/Source/NiagaraEditor/Private/ViewModels/NiagaraConvertNodeViewModel.cpp b/Engine/Plugins/FX/Niagara/Source/NiagaraEditor/Private/ViewModels/NiagaraConvertNodeViewModel.cpp
--- a/Engine/Plugins/FX/Niagara/Source/NiagaraEditor/Private/ViewModels/NiagaraConvertNodeViewModel.cpp
+++ b/Engine/Plugins/FX/Niagara/Source/NiagaraEditor/Private/ViewModels/NiagaraConvertNodeViewModel.cpp
@@ -296,7 +296,9 @@ void FNiagaraConvertNodeViewModel::RefreshPinViewModels()
 	for (UEdGraphPin* InputPin : InputPins)
 	{
 		if (InputPin->PinType.PinCategory == UEdGraphSchema_Niagara::PinCategoryType || 
-			InputPin->PinType.PinCategory == UEdGraphSchema_Niagara::PinCategoryEnum)
+			InputPin->PinType.PinCategory == UEdGraphSchema_Niagara::PinCategoryStaticType || 
+			InputPin->PinType.PinCategory == UEdGraphSchema_Niagara::PinCategoryEnum || 
+			InputPin->PinType.PinCategory == UEdGraphSchema_Niagara::PinCategoryStaticEnum)
 		{
 			InputPinViewModels.Add(MakeShareable(new FNiagaraConvertPinViewModel(this->AsShared(), *InputPin)));
 		}
@@ -304,9 +306,11 @@ void FNiagaraConvertNodeViewModel::RefreshPinViewModels()
 
 	for (UEdGraphPin* OutputPin : OutputPins)
 	{
-		if (OutputPin->PinType.PinCategory == UEdGraphSchema_Niagara::PinCategoryType ||
+		if (OutputPin->PinType.PinCategory == UEdGraphSchema_Niagara::PinCategoryStaticType ||
+			OutputPin->PinType.PinCategory == UEdGraphSchema_Niagara::PinCategoryStaticEnum || 
+			OutputPin->PinType.PinCategory == UEdGraphSchema_Niagara::PinCategoryType ||
 			OutputPin->PinType.PinCategory == UEdGraphSchema_Niagara::PinCategoryEnum)
-		{
+		{																			 
 			OutputPinViewModels.Add(MakeShareable(new FNiagaraConvertPinViewModel(this->AsShared(), *OutputPin)));
 		}
 	}
diff --git a/Engine/Plugins/FX/Niagara/Source/NiagaraEditor/Private/ViewModels/NiagaraScriptInputCollectionViewModel.cpp b/Engine/Plugins/FX/Niagara/Source/NiagaraEditor/Private/ViewModels/NiagaraScriptInputCollectionViewModel.cpp
--- a/Engine/Plugins/FX/Niagara/Source/NiagaraEditor/Private/ViewModels/NiagaraScriptInputCollectionViewModel.cpp
+++ b/Engine/Plugins/FX/Niagara/Source/NiagaraEditor/Private/ViewModels/NiagaraScriptInputCollectionViewModel.cpp
@@ -316,6 +316,11 @@ bool FNiagaraScriptInputCollectionViewModel::SupportsType(const FNiagaraTypeDefi
 		return false;
 	}
 
+	if (Type.IsStatic())
+	{
+		return true;
+	}
+
 	if (Type.IsInternalType())
 	{
 		return false;
diff --git a/Engine/Plugins/FX/Niagara/Source/NiagaraEditor/Private/ViewModels/Stack/NiagaraStackFunctionInput.cpp b/Engine/Plugins/FX/Niagara/Source/NiagaraEditor/Private/ViewModels/Stack/NiagaraStackFunctionInput.cpp
--- a/Engine/Plugins/FX/Niagara/Source/NiagaraEditor/Private/ViewModels/Stack/NiagaraStackFunctionInput.cpp
+++ b/Engine/Plugins/FX/Niagara/Source/NiagaraEditor/Private/ViewModels/Stack/NiagaraStackFunctionInput.cpp
@@ -1042,7 +1042,7 @@ void UNiagaraStackFunctionInput::RefreshFromMetaData(TArray<FStackIssue>& NewIss
 
 		FText EditConditionError;
 		EditCondition.Refresh(InputMetaData->EditCondition, EditConditionError);
-		if (EditCondition.IsValid() && EditCondition.GetConditionInputType() == FNiagaraTypeDefinition::GetBoolDef())
+		if (EditCondition.IsValid() && EditCondition.GetConditionInputType().IsSameBaseDefinition(FNiagaraTypeDefinition::GetBoolDef()))
 		{
 			TOptional<FNiagaraVariableMetaData> EditConditionInputMetadata = EditCondition.GetConditionInputMetaData();
 			if (EditConditionInputMetadata.IsSet())
@@ -1084,7 +1084,7 @@ void UNiagaraStackFunctionInput::RefreshFromMetaData(TArray<FStackIssue>& NewIss
 				true));
 		}
 
-		bIsInlineEditConditionToggle = InputType == FNiagaraTypeDefinition::GetBoolDef() && 
+		bIsInlineEditConditionToggle = InputType.IsSameBaseDefinition(FNiagaraTypeDefinition::GetBoolDef()) && 
 			InputMetaData->bInlineEditConditionToggle;
 	}
 
@@ -1590,6 +1590,10 @@ void UNiagaraStackFunctionInput::SetLocalValue(TSharedRef<FStructOnScope> InLoca
 			{
 				Script->RapidIterationParameters.ParameterGuidMapping.Add(RapidIterationParameter, InputMetaData->GetVariableGuid());
 			}
+			if (InputType.IsStatic()) // Need to potentially trigger a recompile.
+			{
+				Script->MarkScriptAndSourceDesynchronized(TEXT("Static variable changed!"), FGuid::NewGuid());
+			}
 		}
 	}
 	else
diff --git a/Engine/Plugins/FX/Niagara/Source/NiagaraEditor/Private/ViewModels/Stack/NiagaraStackFunctionInputCondition.cpp b/Engine/Plugins/FX/Niagara/Source/NiagaraEditor/Private/ViewModels/Stack/NiagaraStackFunctionInputCondition.cpp
--- a/Engine/Plugins/FX/Niagara/Source/NiagaraEditor/Private/ViewModels/Stack/NiagaraStackFunctionInputCondition.cpp
+++ b/Engine/Plugins/FX/Niagara/Source/NiagaraEditor/Private/ViewModels/Stack/NiagaraStackFunctionInputCondition.cpp
@@ -104,7 +104,7 @@ bool FNiagaraStackFunctionInputCondition::GetConditionIsEnabled() const
 
 bool FNiagaraStackFunctionInputCondition::CanSetConditionIsEnabled() const
 {
-	return IsValid() && InputBinder.GetInputType() == FNiagaraTypeDefinition::GetBoolDef();
+	return IsValid() && InputBinder.GetInputType().IsSameBaseDefinition(FNiagaraTypeDefinition::GetBoolDef());
 }
 
 void FNiagaraStackFunctionInputCondition::SetConditionIsEnabled(bool bInIsEnabled)
diff --git a/Engine/Plugins/FX/Niagara/Source/NiagaraEditor/Private/ViewModels/Stack/NiagaraStackGraphUtilities.cpp b/Engine/Plugins/FX/Niagara/Source/NiagaraEditor/Private/ViewModels/Stack/NiagaraStackGraphUtilities.cpp
--- a/Engine/Plugins/FX/Niagara/Source/NiagaraEditor/Private/ViewModels/Stack/NiagaraStackGraphUtilities.cpp
+++ b/Engine/Plugins/FX/Niagara/Source/NiagaraEditor/Private/ViewModels/Stack/NiagaraStackGraphUtilities.cpp
@@ -763,7 +763,27 @@ void FNiagaraStackGraphUtilities::GetStackFunctionInputPinsWithoutCache(UNiagara
 	FNiagaraParameterMapHistoryBuilder Builder;
 	Builder.SetIgnoreDisabled(bIgnoreDisabled);
 	Builder.ConstantResolver = ConstantResolver;
-	
+
+	UNiagaraEmitter* Emitter = FunctionCallNode.GetTypedOuter<UNiagaraEmitter>();
+	UNiagaraSystem* System = FunctionCallNode.GetTypedOuter<UNiagaraSystem>();
+
+	TArray<FNiagaraVariable> StaticVars;
+	TSharedPtr<FNiagaraGraphCachedDataBase, ESPMode::ThreadSafe> CachedTraversalData;
+	if (Emitter)
+	{
+		CachedTraversalData = Emitter->GetCachedTraversalData();
+	}
+	else if (System)
+	{
+		CachedTraversalData = System->GetCachedTraversalData();
+	}
+	if (CachedTraversalData.IsValid())
+	{
+		CachedTraversalData.Get()->GetStaticVariables(StaticVars);
+	}
+
+	Builder.RegisterExternalStaticVariables(StaticVars);
+
 	FunctionCallNode.BuildParameterMapHistory(Builder, false, false);
 	
 	if (Builder.Histories.Num() == 1)
@@ -773,6 +793,8 @@ void FNiagaraStackGraphUtilities::GetStackFunctionInputPinsWithoutCache(UNiagara
 		FNiagaraParameterMapHistoryBuilder BuilderCompiled;
 		BuilderCompiled.ConstantResolver = ConstantResolver;
 		BuilderCompiled.SetIgnoreDisabled(bIgnoreDisabled);
+		BuilderCompiled.RegisterExternalStaticVariables(StaticVars);
+
 		FunctionCallNode.BuildParameterMapHistory(BuilderCompiled, false, true);
 		TArray<const UEdGraphPin*> CompilationPins;
 		if (BuilderCompiled.Histories.Num() == 1)
@@ -780,6 +802,8 @@ void FNiagaraStackGraphUtilities::GetStackFunctionInputPinsWithoutCache(UNiagara
 			ExtractInputPinsFromHistory(BuilderCompiled.Histories[0], FunctionCallNode.GetCalledGraph(), Options, CompilationPins);
 		}
 
+		//UE_LOG(LogNiagaraEditor, Log, TEXT("OutInputPins:"));
+		int32 Idx = 0;
 		for (const UEdGraphPin* Pin : OutInputPins)
 		{
 			bool bFoundPin = false;
@@ -796,6 +820,9 @@ void FNiagaraStackGraphUtilities::GetStackFunctionInputPinsWithoutCache(UNiagara
 			{
 				OutHiddenPins.Add(Pin);
 			}
+
+			//UE_LOG(LogNiagaraEditor, Log, TEXT("[%d]:%s %s"), Idx, *Pin->PinName.ToString(), bFoundPin ? TEXT("") : TEXT("HIDDEN"));
+			++Idx;
 		}
 	}
 }
@@ -1782,7 +1809,9 @@ void GetFunctionNamesForOutputNode(UNiagaraNodeOutput& OutputNode, TArray<FStrin
 bool FNiagaraStackGraphUtilities::IsRapidIterationType(const FNiagaraTypeDefinition& InputType)
 {
 	checkf(InputType.IsValid(), TEXT("Type is invalid."));
-	return InputType != FNiagaraTypeDefinition::GetBoolDef() && !InputType.IsEnum() &&
+	if (InputType.IsStatic())
+		return true;
+	return InputType != FNiagaraTypeDefinition::GetBoolDef() && (!InputType.IsEnum()) &&
 		InputType != FNiagaraTypeDefinition::GetParameterMapDef() && !InputType.IsUObject();
 }
 
diff --git a/Engine/Plugins/FX/Niagara/Source/NiagaraEditor/Private/ViewModels/Stack/NiagaraStackParameterStoreEntry.cpp b/Engine/Plugins/FX/Niagara/Source/NiagaraEditor/Private/ViewModels/Stack/NiagaraStackParameterStoreEntry.cpp
--- a/Engine/Plugins/FX/Niagara/Source/NiagaraEditor/Private/ViewModels/Stack/NiagaraStackParameterStoreEntry.cpp
+++ b/Engine/Plugins/FX/Niagara/Source/NiagaraEditor/Private/ViewModels/Stack/NiagaraStackParameterStoreEntry.cpp
@@ -191,7 +191,7 @@ TArray<UEdGraphPin*> UNiagaraStackParameterStoreEntry::GetOwningPins()
 
 	// search emitter graphs
 	auto EmitterHandles = GetSystemViewModel()->GetSystem().GetEmitterHandles();
-	for (FNiagaraEmitterHandle Handle : EmitterHandles)
+	for (const FNiagaraEmitterHandle& Handle : EmitterHandles)
 	{
 		UNiagaraGraph* EmitterGraph = CastChecked<UNiagaraScriptSource>(Handle.GetInstance()->GraphSource)->NodeGraph;
 		GraphsToCheck.Add(EmitterGraph);
diff --git a/Engine/Plugins/FX/Niagara/Source/NiagaraEditor/Private/Widgets/SNiagaraParameterMapView.cpp b/Engine/Plugins/FX/Niagara/Source/NiagaraEditor/Private/Widgets/SNiagaraParameterMapView.cpp
--- a/Engine/Plugins/FX/Niagara/Source/NiagaraEditor/Private/Widgets/SNiagaraParameterMapView.cpp
+++ b/Engine/Plugins/FX/Niagara/Source/NiagaraEditor/Private/Widgets/SNiagaraParameterMapView.cpp
@@ -987,7 +987,7 @@ void SNiagaraParameterMapView::GetAllGraphsInSystem(TArray<UNiagaraGraph*>& OutR
 
 	for (int i = 0; i < CachedSystem->GetNumEmitters(); i++)
 	{
-		FNiagaraEmitterHandle Handle = CachedSystem->GetEmitterHandle(i);
+		const FNiagaraEmitterHandle& Handle = CachedSystem->GetEmitterHandle(i);
 		UNiagaraScriptSource* EmitterSource = Cast<UNiagaraScriptSource>(Handle.GetInstance()->GraphSource);
 		OutResult.Add(EmitterSource->NodeGraph);
 	}
diff --git a/Engine/Plugins/FX/Niagara/Source/NiagaraEditor/Private/Widgets/SNiagaraParameterPanel.cpp b/Engine/Plugins/FX/Niagara/Source/NiagaraEditor/Private/Widgets/SNiagaraParameterPanel.cpp
--- a/Engine/Plugins/FX/Niagara/Source/NiagaraEditor/Private/Widgets/SNiagaraParameterPanel.cpp
+++ b/Engine/Plugins/FX/Niagara/Source/NiagaraEditor/Private/Widgets/SNiagaraParameterPanel.cpp
@@ -238,7 +238,7 @@ TSharedRef<SWidget> SNiagaraParameterPanel::OnGenerateWidgetForItem(const FNiaga
 {
 	// Generate the icon widget.
 	FText			   IconToolTip = Item.ScriptVariable->Variable.GetType().GetNameText();
-	FSlateBrush const* IconBrush = FEditorStyle::GetBrush(TEXT("Kismet.AllClasses.VariableIcon"));
+	FSlateBrush const* IconBrush = Item.GetVariable().GetType().IsStatic() ? FNiagaraEditorStyle::Get().GetBrush(TEXT("NiagaraEditor.StaticIcon")) : FEditorStyle::GetBrush(TEXT("Kismet.AllClasses.VariableIcon"));
 	const FLinearColor TypeColor = UEdGraphSchema_Niagara::GetTypeColor(Item.GetVariable().GetType());
 	FSlateColor        IconColor = FSlateColor(TypeColor);
 	FString			   IconDocLink, IconDocExcerpt;
diff --git a/Engine/Plugins/FX/Niagara/Source/NiagaraEditor/Public/EdGraphSchema_Niagara.h b/Engine/Plugins/FX/Niagara/Source/NiagaraEditor/Public/EdGraphSchema_Niagara.h
--- a/Engine/Plugins/FX/Niagara/Source/NiagaraEditor/Public/EdGraphSchema_Niagara.h
+++ b/Engine/Plugins/FX/Niagara/Source/NiagaraEditor/Public/EdGraphSchema_Niagara.h
@@ -90,6 +90,11 @@ class NIAGARAEDITOR_API UEdGraphSchema_Niagara : public UEdGraphSchema
 	static const FName PinCategoryMisc;
 	static const FName PinCategoryClass;
 	static const FName PinCategoryEnum;
+	static const FName PinCategoryStaticType;
+	static const FName PinCategoryStaticClass;
+	static const FName PinCategoryStaticEnum;
+
+	static bool IsStaticPin(const UEdGraphPin* Pin);
 
 	//~ Begin EdGraphSchema Interface
 	virtual void GetContextMenuActions(class UToolMenu* Menu, class UGraphNodeContextMenuContext* Context) const override;
@@ -164,6 +169,8 @@ class NIAGARAEDITOR_API UEdGraphSchema_Niagara : public UEdGraphSchema
 	static const FLinearColor NodeTitleColor_Event; 
 	static const FLinearColor NodeTitleColor_TranslatorConstant;
 	static const FLinearColor NodeTitleColor_RapidIteration;
+
+	bool PinTypesValidForNumericConversion(FEdGraphPinType AType, FEdGraphPinType BType) const;
 	
 private:
 	void GetNumericConversionToSubMenuActions(class UToolMenu* Menu, const FName SectionName, UEdGraphPin* InGraphPin);
diff --git a/Engine/Plugins/FX/Niagara/Source/NiagaraEditor/Public/NiagaraEditorCommon.h b/Engine/Plugins/FX/Niagara/Source/NiagaraEditor/Public/NiagaraEditorCommon.h
--- a/Engine/Plugins/FX/Niagara/Source/NiagaraEditor/Public/NiagaraEditorCommon.h
+++ b/Engine/Plugins/FX/Niagara/Source/NiagaraEditor/Public/NiagaraEditorCommon.h
@@ -3,6 +3,7 @@
 #pragma once
 
 #include "NiagaraCommon.h"
+#include "Delegates/Delegate.h"
 #include "NiagaraEditorCommon.generated.h"
 
 NIAGARAEDITOR_API DECLARE_LOG_CATEGORY_EXTERN(LogNiagaraEditor, Log, All);
@@ -26,6 +27,7 @@ public:
 		, bSupportsAddedInputs(false)
 		, bNumericsCanBeIntegers(true)
 		, bNumericsCanBeFloats(true)
+		, bSupportsStaticResolution(false)
 	{}
 
 	FName Name;
@@ -37,6 +39,7 @@ public:
 	TArray<FNiagaraOpInOutInfo> Inputs;
 	TArray<FNiagaraOpInOutInfo> Outputs;
 
+
 	/** If true then this operation supports a variable number of inputs */
 	bool bSupportsAddedInputs;
 
@@ -68,6 +71,13 @@ public:
 	void BuildName(FString InName, FString InCategory);
 
 	bool CreateHlslForAddedInputs(int32 InputCount, FString& HlslResult) const;
+
+	DECLARE_DELEGATE_RetVal_OneParam(int32, FStaticVariableResolve, const TArray<int32>& );
+
+	FStaticVariableResolve StaticVariableResolveFunction;
+
+	/** Whether or not you can upgrade type to static on connection.*/
+	bool bSupportsStaticResolution;
 };
 
 /** Interface for struct representing information about where to focus in a Niagara Script Graph after opening the editor for it. */
diff --git a/Engine/Plugins/FX/Niagara/Source/NiagaraEditor/Public/NiagaraEditorModule.h b/Engine/Plugins/FX/Niagara/Source/NiagaraEditor/Public/NiagaraEditorModule.h
--- a/Engine/Plugins/FX/Niagara/Source/NiagaraEditor/Public/NiagaraEditorModule.h
+++ b/Engine/Plugins/FX/Niagara/Source/NiagaraEditor/Public/NiagaraEditorModule.h
@@ -121,6 +121,7 @@ public:
 		UNiagaraEmitter* OwningEmitter,
 		UNiagaraScript* TargetScript,
 		FGuid TargetVersion);
+	TSharedPtr<FNiagaraGraphCachedDataBase, ESPMode::ThreadSafe> CacheGraphTraversal(const UObject* Obj, FGuid Version);
 
 
 	/** Gets the extensibility managers for outside entities to extend static mesh editor's menus and toolbars */
@@ -290,6 +291,7 @@ private:
 	FDelegateHandle CompileResultHandle;
 	FDelegateHandle PrecompilerHandle;
 	FDelegateHandle PrecompileDuplicatorHandle;
+	FDelegateHandle GraphCacheTraversalHandle;
 
 	FDelegateHandle DeviceProfileManagerUpdatedHandle;
 
diff --git a/Engine/Plugins/FX/Niagara/Source/NiagaraEditor/Public/NiagaraEditorUtilities.h b/Engine/Plugins/FX/Niagara/Source/NiagaraEditor/Public/NiagaraEditorUtilities.h
--- a/Engine/Plugins/FX/Niagara/Source/NiagaraEditor/Public/NiagaraEditorUtilities.h
+++ b/Engine/Plugins/FX/Niagara/Source/NiagaraEditor/Public/NiagaraEditorUtilities.h
@@ -394,4 +394,6 @@ namespace FNiagaraParameterUtilities
 	NIAGARAEDITOR_API FName SetCustomNamespaceModifier(FName InParameterName, TSet<FName>& CurrentParameterNames);
 
 	NIAGARAEDITOR_API bool TestCanRenameWithMessage(FName ParameterName, FText& OutMessage);
+
+	NIAGARAEDITOR_API void FilterToRelevantStaticVariables(const TArray<FNiagaraVariable>& InVars, TArray<FNiagaraVariable>& OutVars, FName InOldEmitterAlias, FName InNewEmitterAlias, bool bFilterByEmitterAliasAndConvertToUnaliased);
 };
diff --git a/Engine/Plugins/FX/Niagara/Source/NiagaraEditor/Public/NiagaraNode.h b/Engine/Plugins/FX/Niagara/Source/NiagaraEditor/Public/NiagaraNode.h
--- a/Engine/Plugins/FX/Niagara/Source/NiagaraEditor/Public/NiagaraNode.h
+++ b/Engine/Plugins/FX/Niagara/Source/NiagaraEditor/Public/NiagaraNode.h
@@ -31,6 +31,8 @@ protected:
 	bool ReallocatePins(bool bMarkNeedsResynchronizeOnChange = true);
 
 	bool CompileInputPins(class FHlslNiagaraTranslator *Translator, TArray<int32>& OutCompiledInputs);
+
+	virtual void OnPostSynchronizationInReallocatePins() {}
 	
 public:
 
@@ -148,14 +150,14 @@ public:
 	 *  @Param	LocallyOwnedOutputPin	The pin to trace, must be a child of this node
 	 *  @Param	bFilterForCompilation	If true, some nodes like static switches or reroute nodes are jumped over. The ui usually sets this to false to follow all possible paths.  
 	 */
-	virtual UEdGraphPin* GetTracedOutputPin(UEdGraphPin* LocallyOwnedOutputPin, bool bFilterForCompilation) const {return LocallyOwnedOutputPin;}
+	virtual UEdGraphPin* GetTracedOutputPin(UEdGraphPin* LocallyOwnedOutputPin, bool bFilterForCompilation, TArray<const UNiagaraNode*>* OutNodesVisitedDuringTrace = nullptr) const;
 
 	/** Traces a node's output pins to its source output pin.
 	*
 	*  @Param	LocallyOwnedOutputPin	The pin to trace
 	*  @Param	bFilterForCompilation	If true, some nodes like static switches or reroute nodes are jumped over. The ui usually sets this to false to follow all possible paths.  
 	*/
-	static UEdGraphPin* TraceOutputPin(UEdGraphPin* LocallyOwnedOutputPin, bool bFilterForCompilation = true);
+	static UEdGraphPin* TraceOutputPin(UEdGraphPin* LocallyOwnedOutputPin, bool bFilterForCompilation = true, TArray<const UNiagaraNode*>* OutNodesVisitedDuringTrace = nullptr);
 
 	/** Allows a node to replace a pin that is about to be compiled with another pin. This can be used for either optimizations or features such as the static switch. Returns true if the pin was successfully replaced, false otherwise. */
 	virtual bool SubstituteCompiledPin(FHlslNiagaraTranslator* Translator, UEdGraphPin** LocallyOwnedPin);
@@ -213,6 +215,9 @@ protected:
 		the input pins should have been visited already.*/
 	virtual void RouteParameterMapAroundMe(FNiagaraParameterMapHistoryBuilder& OutHistory, bool bRecursive) const;
 
+	/** Basically routes the pin through the parameter map builder so that it looks like a regular pin. Handles visiting the node connected as well to the input pin.*/
+	virtual void RegisterPassthroughPin(FNiagaraParameterMapHistoryBuilder& OutHistory, UEdGraphPin* InputPin, UEdGraphPin* OutputPin, bool bFilterForCompilation, bool bVisitInputPin) const;
+
 	/** If the pin is a known name (like Engine.DeltaTime) this tries to return a default tooltip for it. */
 	bool GetTooltipTextForKnownPin(const UEdGraphPin& Pin, FText& OutTooltip) const;
 	
diff --git a/Engine/Plugins/FX/Niagara/Source/NiagaraEditor/Public/NiagaraNodeCustomHlsl.h b/Engine/Plugins/FX/Niagara/Source/NiagaraEditor/Public/NiagaraNodeCustomHlsl.h
--- a/Engine/Plugins/FX/Niagara/Source/NiagaraEditor/Public/NiagaraNodeCustomHlsl.h
+++ b/Engine/Plugins/FX/Niagara/Source/NiagaraEditor/Public/NiagaraNodeCustomHlsl.h
@@ -38,6 +38,7 @@ public:
 	static uint32 ReplaceExactMatchTokens(TArray<FString>& Tokens, const FString& SrcString, const FString& ReplaceString, bool bAllowNamespaceSeparation);
 	static FNiagaraVariable StripVariableToBaseType(const FNiagaraVariable& InVar);
 	virtual bool AllowNiagaraTypeForAddPin(const FNiagaraTypeDefinition& InType) const;
+	virtual bool AllowNiagaraTypeForAddPin(const FNiagaraTypeDefinition& InType, EEdGraphPinDirection InDirection) const;
 
 	virtual bool ReferencesVariable(const FNiagaraVariableBase& InVar) const;
 
diff --git a/Engine/Plugins/FX/Niagara/Source/NiagaraEditor/Public/NiagaraNodeOp.h b/Engine/Plugins/FX/Niagara/Source/NiagaraEditor/Public/NiagaraNodeOp.h
--- a/Engine/Plugins/FX/Niagara/Source/NiagaraEditor/Public/NiagaraNodeOp.h
+++ b/Engine/Plugins/FX/Niagara/Source/NiagaraEditor/Public/NiagaraNodeOp.h
@@ -44,6 +44,9 @@ public:
 	UPROPERTY(meta = (SkipForCompileHash = "true"))
 	TArray<FAddedPinData> AddedPins;
 
+	UPROPERTY(meta = (SkipForCompileHash = "true"))
+	bool bAllStatic;
+
 	//~ Begin UObject interface
 	virtual void PostLoad() override;
 	//~ End UObject interface
@@ -60,12 +63,16 @@ public:
 	virtual bool RefreshFromExternalChanges() override;
 	virtual ENiagaraNumericOutputTypeSelectionMode GetNumericOutputTypeSelectionMode() const override;
 	virtual bool GenerateCompileHashForClassMembers(const UClass* InClass, FNiagaraCompileHashVisitor* InVisitor) const override;
+	virtual void BuildParameterMapHistory(FNiagaraParameterMapHistoryBuilder& OutHistory, bool bRecursive = true, bool bFilterForCompilation = true) const;
+
 	//~ End UNiagaraNode Interface
 
 	//~ Begin UNiagaraNodeWithDynamicPins Interface
 	virtual bool AllowNiagaraTypeForAddPin(const FNiagaraTypeDefinition& InType) const override;
+	virtual void PinConnectionListChanged(UEdGraphPin* Pin) override;
 	//~ End UNiagaraNodeWithDynamicPins Interface
 
+
 protected:
 	//~ Begin EdGraphNode Interface
 	virtual void OnPinRemoved(UEdGraphPin* PinToRemove) override;
@@ -78,6 +85,11 @@ protected:
 	virtual void OnPinRenamed(UEdGraphPin* RenamedPin, const FString& OldName) override;
 	virtual bool CanRemovePin(const UEdGraphPin* Pin) const override;
 	//~ End UNiagaraNodeWithDynamicPins Interface
+
+	virtual void OnPostSynchronizationInReallocatePins() override;
+
+	void HandleStaticInputPinUpgrade(UEdGraphPin* InputPin);
+	void HandleStaticOutputPinUpgrade();
 	
 private:
 	FName GetUniqueAdditionalPinName() const;
diff --git a/Engine/Plugins/FX/Niagara/Source/NiagaraEditor/Public/NiagaraNodeOutput.h b/Engine/Plugins/FX/Niagara/Source/NiagaraEditor/Public/NiagaraNodeOutput.h
--- a/Engine/Plugins/FX/Niagara/Source/NiagaraEditor/Public/NiagaraNodeOutput.h
+++ b/Engine/Plugins/FX/Niagara/Source/NiagaraEditor/Public/NiagaraNodeOutput.h
@@ -45,6 +45,8 @@ public:
 	virtual bool IncludeParentNodeContextMenu() const { return true; }
 	//~ End EdGraphNode Interface
 
+	virtual void BuildParameterMapHistory(FNiagaraParameterMapHistoryBuilder& OutHistory, bool bRecursive = true, bool bFilterForCompilation = true) const override;
+
 	/** Notifies the node that it's output variables have been modified externally. */
 	void NotifyOutputVariablesChanged();
 
diff --git a/Engine/Plugins/FX/Niagara/Source/NiagaraEditor/Public/NiagaraNodeWithDynamicPins.h b/Engine/Plugins/FX/Niagara/Source/NiagaraEditor/Public/NiagaraNodeWithDynamicPins.h
--- a/Engine/Plugins/FX/Niagara/Source/NiagaraEditor/Public/NiagaraNodeWithDynamicPins.h
+++ b/Engine/Plugins/FX/Niagara/Source/NiagaraEditor/Public/NiagaraNodeWithDynamicPins.h
@@ -36,6 +36,7 @@ public:
 
 	/** Determine whether or not a Niagara type is supported for an Add Pin possibility.*/
 	virtual bool AllowNiagaraTypeForAddPin(const FNiagaraTypeDefinition& InType) const;
+	virtual bool AllowNiagaraTypeForAddPin(const FNiagaraTypeDefinition& InType, EEdGraphPinDirection InDirection) const { return AllowNiagaraTypeForAddPin(InType); };
 
 	/** Used in to gather the actions for selecting the pin to add. */
 	virtual void CollectAddPinActions(FGraphActionListBuilderBase& OutActions, bool& bOutCreateRemainingActions, UEdGraphPin* Pin);
diff --git a/Engine/Plugins/FX/Niagara/Source/NiagaraEditor/Public/NiagaraParameterMapHistory.h b/Engine/Plugins/FX/Niagara/Source/NiagaraEditor/Public/NiagaraParameterMapHistory.h
--- a/Engine/Plugins/FX/Niagara/Source/NiagaraEditor/Public/NiagaraParameterMapHistory.h
+++ b/Engine/Plugins/FX/Niagara/Source/NiagaraEditor/Public/NiagaraParameterMapHistory.h
@@ -6,6 +6,7 @@
 #include "NiagaraCommon.h"
 #include "Templates/Tuple.h"
 #include "ViewModels/Stack/NiagaraStackGraphUtilities.h"
+#include "NiagaraScript.h"
 
 class UNiagaraNodeOutput;
 class UEdGraphPin;
@@ -57,6 +58,48 @@ struct FModuleScopedPin
 	FName ModuleName = NAME_None;
 };
 
+struct FGraphTraversalHandle
+{
+public:
+	FGraphTraversalHandle() {}
+	FGraphTraversalHandle(const TArray<FGuid>& InputPath, TArray<FString> InFriendlyPath) : Path(InputPath), FriendlyPath(InFriendlyPath) {}
+
+	void Append(const FGuid& Guid, const FString& FriendlyName)
+	{
+		FriendlyPath.Add(FriendlyName);
+		Path.Add(Guid);
+	}
+
+	void Push(const UEdGraphNode* Node);
+	void Push(const UEdGraphPin* Pin);
+	void Pop();
+
+	bool operator==(const FGraphTraversalHandle& Rhs) const
+	{
+		return Path == Rhs.Path;
+	}
+
+	friend uint32 GetTypeHash(const FGraphTraversalHandle& Ref)
+	{
+		uint32 Value = 0;
+		for (int32 i = 0; i < Ref.Path.Num(); i++)
+			Value = HashCombine(Value, GetTypeHash(Ref.Path[i]));
+		return Value;
+	}
+
+	FString ToString() const
+	{
+		FString Output;
+		for (const FString& Str : FriendlyPath)
+			Output += TEXT("/") + Str;
+		return Output;
+	}
+private:
+	TArray<FGuid > Path;
+	TArray<FString> FriendlyPath;
+};
+
+
 /** Traverses a Niagara node graph to identify the variables that have been written and read from a parameter map. 
 * 	This class is meant to aid in UI and compilation of the graph. There are several main script types and each one interacts
 *	slightly differently with the history depending on context.
@@ -74,6 +117,9 @@ public:
 	/** The metadata associated with each variable identified during the traversal. Only gathered by FNiagaraParameterMapHistoryWithMetaDataBuilder. */
 	TArray<FNiagaraVariableMetaData> VariableMetaData;
 
+	TArray<TArray<FString> > PerVariableConstantValue;
+	TArray<TArray<FString> > PerVariableConstantDefaultValue;
+
 	TArray<FNiagaraVariable> VariablesWithOriginalAliasesIntact;
 
 
@@ -117,8 +163,9 @@ public:
 	/** List of additional DataSets to be written that were encountered during traversal. */
 	TArray<FNiagaraDataSetID> AdditionalDataSetWrites;
 
-	/** Flag indicating that the history building skipped processing parameter collections */
+	TMap<FGraphTraversalHandle, FString> PinToConstantValues;
 	bool bParameterCollectionsSkipped = false;
+	void RegisterConstantPin(const FGraphTraversalHandle& InTraversalPath, const UEdGraphPin* InPin, const FString& InValue);
 
 	bool IsVariableFromCustomIterationNamespaceOverride(const FNiagaraVariable& InVar) const;
 	
@@ -127,6 +174,8 @@ public:
 	*/
 	int32 RegisterParameterMapPin(const UEdGraphPin* Pin);
 
+	void RegisterConstantVariableWrite(const FString& InValue, int32 InVarIdx, bool bIsSettingDefault);
+
 	uint32 BeginNodeVisitation(const UNiagaraNode* Node);
 	void EndNodeVisitation(uint32 IndexFromBeginNode);
 
@@ -215,8 +264,8 @@ public:
 	static bool IsValidNamespaceForReading(ENiagaraScriptUsage InScriptUsage, int32 InUsageBitmask, FString Namespace);
 
 	/** Called to determine if a given variable should be output from a script. It is not static as it requires the overall context to include emitter namespaces visited for system scripts.*/
-	bool IsPrimaryDataSetOutput(const FNiagaraVariable& InVar, const UNiagaraScript* InScript,   bool bAllowDataInterfaces = false) const;
-	bool IsPrimaryDataSetOutput(const FNiagaraVariable& InVar, ENiagaraScriptUsage InUsage, bool bAllowDataInterfaces = false) const;
+	bool IsPrimaryDataSetOutput(const FNiagaraVariable& InVar, const UNiagaraScript* InScript,   bool bAllowDataInterfaces = false,  bool bAllowStatics = false) const;
+	bool IsPrimaryDataSetOutput(const FNiagaraVariable& InVar, ENiagaraScriptUsage InUsage, bool bAllowDataInterfaces = false, bool bAllowStatics = false) const;
 	static bool IsWrittenToScriptUsage(const FNiagaraVariable& InVar, ENiagaraScriptUsage InUsage, bool bAllowDataInterfaces);
 
 	/** Are we required to export this variable as an external constant?*/
@@ -306,6 +355,17 @@ public:
 	/** Important. Must be called for each routing of the parameter map. This feeds the list used by TraceParameterMapOutputPin.*/
 	int32 RegisterParameterMapPin(int32 WhichParameterMap, const UEdGraphPin* Pin);
 
+	int32 RegisterConstantPin(int32 WhichConstant, const UEdGraphPin* Pin);
+	int32 RegisterConstantVariableWrite(int32 WhichParamMapIdx, int32 WhichConstant, int32 WhichVarIdx, bool bIsSettingDefault);
+
+	int32  RegisterConstantFromInputPin(const UEdGraphPin* InputPin);
+
+	int32 GetConstantFromInputPin(const UEdGraphPin* InputPin) const;
+	int32 GetConstantFromOutputPin(const UEdGraphPin* OutputPin) const;
+	int32 GetConstantFromVariableRead(const FNiagaraVariableBase& InVar);
+	int32 AddOrGetConstantFromValue(const FString& Value);
+	FString GetConstant(int32 InIndex);
+
 	/** Records a write to a DataSet in the appropriate parameter map history */
 	void RegisterDataSetWrite(int32 WhichParameterMap, const FNiagaraDataSetID& DataSet);
 
@@ -329,6 +389,9 @@ public:
 	void BeginUsage(ENiagaraScriptUsage InUsage, FName InStageName = FName());
 	void EndUsage();
 
+
+	FGraphTraversalHandle ActivePath;
+
 	/**
 	* Record that we have entered a new function scope.
 	*/
@@ -411,6 +474,8 @@ public:
 	/** Helper method to identify any matching input nodes from the calling context node to the input variable.*/
 	int32 FindMatchingParameterMapFromContextInputs(const FNiagaraVariable& InVar) const;
 
+	int32 FindMatchingStaticFromContextInputs(const FNiagaraVariable& InVar) const;
+
 	/** In some cases, we don't want all the variables encountered in a traversal. In this case, you can filter 
 	*  the map history to only include variables that are relevant to the specific script type. For instance, a System 
 	*  script doesn't really care about the Particles namespace.
@@ -432,10 +497,18 @@ public:
 	void RegisterEncounterableVariables(const TArray<FNiagaraVariable>& Variables);
 	const TArray<FNiagaraVariable>& GetEncounterableVariables() const {	return EncounterableExternalVariables;}
 
+	void RegisterExternalStaticVariables(const TArray<FNiagaraVariable>& Variables);
+
 	FCompileConstantResolver ConstantResolver;
 
 	bool bShouldBuildSubHistories = true;
 
+
+	void SetConstantByStaticVariable(int32& OutValue, const UEdGraphPin* InDefaultPin);
+	void SetConstantByStaticVariable(int32& OutValue, const FNiagaraVariable& Var);
+
+
+	TArray<FNiagaraVariable> StaticVariables;
 protected:
 	/**
 	* Generate the internal alias map from the current traversal state.
@@ -454,6 +527,10 @@ protected:
 	/** Contains the hierarchy of graphs leading to and including the current graph being processed. Is not in sync with CallingContext as there is an additional 0th entry for the NodeGraphDeepCopy. */
 	TArray<const UNiagaraGraph*> CallingGraphContext;
 
+	TArray< TMap<const UEdGraphPin*, int32> > PinToConstantIndices;
+	TArray<FString> Constants;
+	TArray< TMap<FNiagaraVariableBase, int32> > VariableToConstantIndices;
+
 	/** Tracker for each context level of the parameter map index associated with a given pin. Used to trace parameter maps through the graph.*/
 	TArray<TMap<const UEdGraphPin*, int32> > PinToParameterMapIndices;
 	/** List of previously visited nodes per context. Note that the same node may be visited multiple times across all graph traversals, but only one time per context level.*/
@@ -470,6 +547,7 @@ protected:
 
 	TArray<TArray<FString> > EncounteredFunctionNames;
 	TArray<FString> EncounteredEmitterNames;
+
 	
 	/** Whether or not the script allow list is active.*/
 	bool bFilterByScriptAllowList;
@@ -496,3 +574,11 @@ public:
 	*/
 	virtual int32 AddVariableToHistory(FNiagaraParameterMapHistory& History, const FNiagaraVariable& InVar, const FNiagaraVariable& InAliasedVar, const UEdGraphPin* InPin);
 };
+
+struct FNiagaraGraphCachedBuiltHistory : public FNiagaraGraphCachedDataBase
+{
+	virtual ~FNiagaraGraphCachedBuiltHistory() {};
+	virtual void GetStaticVariables(TArray<FNiagaraVariable>& OutVars) { OutVars = StaticVariables; }
+
+	TArray<FNiagaraVariable> StaticVariables;
+};
\ No newline at end of file
diff --git a/Engine/Plugins/FX/Niagara/Source/NiagaraEditor/Public/NiagaraStackFunctionInputBinder.h b/Engine/Plugins/FX/Niagara/Source/NiagaraEditor/Public/NiagaraStackFunctionInputBinder.h
--- a/Engine/Plugins/FX/Niagara/Source/NiagaraEditor/Public/NiagaraStackFunctionInputBinder.h
+++ b/Engine/Plugins/FX/Niagara/Source/NiagaraEditor/Public/NiagaraStackFunctionInputBinder.h
@@ -118,7 +118,7 @@ private:
 template<>
 inline bool FNiagaraStackFunctionInputBinder::GetValue<bool>() const
 {
-	check(InputType == FNiagaraTypeDefinition::GetBoolDef());
+	check(InputType.IsSameBaseDefinition(FNiagaraTypeDefinition::GetBoolDef()));
 	TArray<uint8> ValueData = GetData();
 	FNiagaraBool* BoolStruct = (FNiagaraBool*)ValueData.GetData();
 	return BoolStruct->GetValue();
@@ -127,7 +127,7 @@ inline bool FNiagaraStackFunctionInputBinder::GetValue<bool>() const
 template<>
 inline void FNiagaraStackFunctionInputBinder::SetValue<bool>(const bool& bInValue)
 {
-	check(InputType == FNiagaraTypeDefinition::GetBoolDef());
+	check(InputType.IsSameBaseDefinition(FNiagaraTypeDefinition::GetBoolDef()));
 	FNiagaraBool BoolStruct(bInValue);
 	SetData((uint8*)(&BoolStruct), sizeof(FNiagaraBool));
 }
\ No newline at end of file
diff --git a/Engine/Plugins/FX/Niagara/Source/NiagaraEditor/Public/Widgets/SNiagaraParameterMenu.cpp b/Engine/Plugins/FX/Niagara/Source/NiagaraEditor/Public/Widgets/SNiagaraParameterMenu.cpp
--- a/Engine/Plugins/FX/Niagara/Source/NiagaraEditor/Public/Widgets/SNiagaraParameterMenu.cpp
+++ b/Engine/Plugins/FX/Niagara/Source/NiagaraEditor/Public/Widgets/SNiagaraParameterMenu.cpp
@@ -804,7 +804,7 @@ void SNiagaraAddParameterFromPinMenu::CollectAllActions(FGraphActionListBuilderB
 	for (const FNiagaraTypeDefinition& RegisteredType : Types)
 	{
 		bool bAllowType = false;
-		bAllowType = NiagaraNode->AllowNiagaraTypeForAddPin(RegisteredType);
+		bAllowType = NiagaraNode->AllowNiagaraTypeForAddPin(RegisteredType, bIsParameterReadNode ? EEdGraphPinDirection::EGPD_Output : EEdGraphPinDirection::EGPD_Input) ;
 
 		if (bAllowType)
 		{
diff --git a/Engine/Plugins/FX/Niagara/Source/NiagaraEditorWidgets/Private/Stack/SNiagaraStackFunctionInputValue.cpp b/Engine/Plugins/FX/Niagara/Source/NiagaraEditorWidgets/Private/Stack/SNiagaraStackFunctionInputValue.cpp
--- a/Engine/Plugins/FX/Niagara/Source/NiagaraEditorWidgets/Private/Stack/SNiagaraStackFunctionInputValue.cpp
+++ b/Engine/Plugins/FX/Niagara/Source/NiagaraEditorWidgets/Private/Stack/SNiagaraStackFunctionInputValue.cpp
@@ -50,7 +50,7 @@
 #include "Widgets/Layout/SSpacer.h"
 #include "Widgets/Layout/SSeparator.h"
 #include "Widgets/Input/SSpinBox.h"
-
+#include "NiagaraCommon.h"
 #define LOCTEXT_NAMESPACE "NiagaraStackFunctionInputValue"
 
 const float TextIconSize = 16;
@@ -98,6 +98,23 @@ void SNiagaraStackFunctionInputValue::Construct(const FArguments& InArgs, UNiaga
 			]
 			+ SHorizontalBox::Slot()
 			.VAlign(VAlign_Center)
+			.AutoWidth()
+			.Padding(0, 0, 3, 0)
+			[
+				// Type Modifier Icon
+				SNew(SBox)
+				.WidthOverride(TextIconSize)
+				.VAlign(VAlign_Center)
+				.Visibility(this, &SNiagaraStackFunctionInputValue::GetTypeModifierIconVisibility)
+				[
+					SNew(STextBlock)
+					.Text(this, &SNiagaraStackFunctionInputValue::GetTypeModifierIconText)
+					.ToolTipText(this, &SNiagaraStackFunctionInputValue::GetTypeModifierIconToolTip)
+					.ColorAndOpacity(this, &SNiagaraStackFunctionInputValue::GetTypeModifierIconColor)
+				]
+			]
+			+ SHorizontalBox::Slot()
+			.VAlign(VAlign_Center)
 			[
 				// Value container and widgets.
 				SAssignNew(ValueContainer, SBox)
@@ -778,6 +795,26 @@ FSlateColor SNiagaraStackFunctionInputValue::GetInputIconColor() const
 	return FNiagaraEditorWidgetsStyle::Get().GetColor(FNiagaraStackEditorWidgetsUtilities::GetIconColorNameForInputMode(FunctionInput->GetValueMode()));
 }
 
+EVisibility SNiagaraStackFunctionInputValue::GetTypeModifierIconVisibility() const
+{
+	return FunctionInput->GetInputType().IsStatic() ? EVisibility::Visible : EVisibility::Hidden;
+}
+
+FText SNiagaraStackFunctionInputValue::GetTypeModifierIconText() const
+{
+	return FText::FromString(FString(TEXT("S")));
+}
+
+FText SNiagaraStackFunctionInputValue::GetTypeModifierIconToolTip() const
+{
+	return LOCTEXT("TypeModifierTooltip", "Static variables can only be set once in the graph and are meant to communicate inputs to static switches across the emitter.");
+}
+
+FSlateColor SNiagaraStackFunctionInputValue::GetTypeModifierIconColor() const
+{
+	return UEdGraphSchema_Niagara::GetTypeColor(FunctionInput->GetInputType());
+}
+
 FReply SNiagaraStackFunctionInputValue::OnFunctionInputDrop(const FGeometry& InGeometry, const FDragDropEvent& InDragDropEvent)
 {
 	TSharedPtr<FNiagaraParameterDragOperation> InputDragDropOperation = InDragDropEvent.GetOperationAs<FNiagaraParameterDragOperation>();
diff --git a/Engine/Plugins/FX/Niagara/Source/NiagaraEditorWidgets/Private/Stack/SNiagaraStackFunctionInputValue.h b/Engine/Plugins/FX/Niagara/Source/NiagaraEditorWidgets/Private/Stack/SNiagaraStackFunctionInputValue.h
--- a/Engine/Plugins/FX/Niagara/Source/NiagaraEditorWidgets/Private/Stack/SNiagaraStackFunctionInputValue.h
+++ b/Engine/Plugins/FX/Niagara/Source/NiagaraEditorWidgets/Private/Stack/SNiagaraStackFunctionInputValue.h
@@ -134,6 +134,14 @@ private:
 
 	FSlateColor GetInputIconColor() const;
 
+	EVisibility GetTypeModifierIconVisibility() const;
+
+	FText GetTypeModifierIconText() const;
+
+	FText GetTypeModifierIconToolTip() const;
+
+	FSlateColor GetTypeModifierIconColor() const;
+
 	FReply OnFunctionInputDrop(const FGeometry& InGeometry, const FDragDropEvent& InDragDropEvent);
 
 	bool OnFunctionInputAllowDrop(TSharedPtr<FDragDropOperation> DragDropOperation);
