Fixing CIS issue due to bad robomerge

diff --git a/Engine/Plugins/FX/Niagara/Source/NiagaraEditor/Private/EdGraphSchema_Niagara.cpp b/Engine/Plugins/FX/Niagara/Source/NiagaraEditor/Private/EdGraphSchema_Niagara.cpp
--- a/Engine/Plugins/FX/Niagara/Source/NiagaraEditor/Private/EdGraphSchema_Niagara.cpp
+++ b/Engine/Plugins/FX/Niagara/Source/NiagaraEditor/Private/EdGraphSchema_Niagara.cpp
@@ -576,8 +576,8 @@ TArray<TSharedPtr<FNiagaraAction_NewNode>> UEdGraphSchema_Niagara::GetGraphActio
 		if (FromPin)
 		{
 			PinType = PinToTypeDefinition(FromPin);
-			}
-
+		}
+		
 		/** Make & Break */
 		if (PinType.GetScriptStruct())
 		{
@@ -614,12 +614,12 @@ TArray<TSharedPtr<FNiagaraAction_NewNode>> UEdGraphSchema_Niagara::GetGraphActio
 				
 				if(!CandidateType.IsValid() || CandidateType.IsUObject() || NoMakeAllowed.Contains(CandidateType))
 				{
-							continue;
-						}
-
+					continue;
+				}
+				
 				// if we have a "from pin", we generally only want applicable make types
 				if(FromPin)
-						{
+				{
 					// if our "from pin" is an output pin, we require at least one input pin type of the candidate type to be the same or assignable
 					if(FromPin->Direction == EGPD_Output)
 					{
@@ -634,15 +634,14 @@ TArray<TSharedPtr<FNiagaraAction_NewNode>> UEdGraphSchema_Niagara::GetGraphActio
 						}
 						}
 					}
-				}
 					else
-				{
+					{
 						// if our "from pin" is an input pin, we generally only allow that exact type for making, with some exceptions like NiagaraGenerics
 						if(PinType == CandidateType)
-					{
+						{
 								bAddMake = true;
 						}
-
+						
 						if(PinType == FNiagaraTypeDefinition::GetGenericNumericDef() && FNiagaraTypeDefinition::IsValidNumericInput(CandidateType))
 						{
 							bAddMake = true;
@@ -650,8 +649,8 @@ TArray<TSharedPtr<FNiagaraAction_NewNode>> UEdGraphSchema_Niagara::GetGraphActio
 					}
 				}
 				// if we don't, just generally allow make types. We are stricter about break types however
-			else
-			{
+				else
+				{
 					bAddMake = true;
 				}
 
@@ -659,8 +658,8 @@ TArray<TSharedPtr<FNiagaraAction_NewNode>> UEdGraphSchema_Niagara::GetGraphActio
 				{
 					MakeBreakType(CandidateType, true);
 				}
-				}
-				
+			}
+			
 			// Break for the current "from pin"
 			{
 				// Don't break scalars. Allow makes for now as a convenient method of getting internal script constants when dealing with numeric pins.
@@ -668,12 +667,12 @@ TArray<TSharedPtr<FNiagaraAction_NewNode>> UEdGraphSchema_Niagara::GetGraphActio
 				if (FromPin && !PinType.IsInternalType() && !PinType.IsUObject() && !FNiagaraTypeDefinition::IsScalarDefinition(PinType) && PinType != FNiagaraTypeDefinition::GetGenericNumericDef() && PinType != FNiagaraTypeDefinition::GetParameterMapDef())
 				{
 					if(FromPin->Direction == EGPD_Output)
-				{
+					{
 						MakeBreakType(PinType, false);
+					}					
 				}
-				}
-			}
-
+			}				
+			
 
 			//Always add generic convert as an option.
 			FText Desc = LOCTEXT("NiagaraConvert", "Convert");
