New static variable engine test

diff --git a/Engine/Plugins/FX/Niagara/Source/Niagara/Private/NiagaraRendererProperties.cpp b/Engine/Plugins/FX/Niagara/Source/Niagara/Private/NiagaraRendererProperties.cpp
--- a/Engine/Plugins/FX/Niagara/Source/Niagara/Private/NiagaraRendererProperties.cpp
+++ b/Engine/Plugins/FX/Niagara/Source/Niagara/Private/NiagaraRendererProperties.cpp
@@ -386,7 +386,11 @@ void UNiagaraRendererProperties::PostEditChangeProperty(struct FPropertyChangedE
 		{
 			UNiagaraSystem::RequestCompileForEmitter(Emitter);
 		}
+
+		// Just in case we changed something that needs static params, refresh that cached list.
+		Emitter->RebuildRendererBindings();
 	}
+
 }
 
 #endif
