Fix CIS error

diff --git a/Engine/Plugins/FX/Niagara/Source/Niagara/Private/NiagaraScript.cpp b/Engine/Plugins/FX/Niagara/Source/Niagara/Private/NiagaraScript.cpp
--- a/Engine/Plugins/FX/Niagara/Source/Niagara/Private/NiagaraScript.cpp
+++ b/Engine/Plugins/FX/Niagara/Source/Niagara/Private/NiagaraScript.cpp
@@ -87,7 +87,7 @@ static FAutoConsoleVariableRef CVarNiagaraDelayScriptAsyncOptimization(
 );
 
 int32 UNiagaraScript::LogCompileStaticVars =0;
-static FAutoConsoleVariableRef CVarLogCompileIdGeneration(
+static FAutoConsoleVariableRef CVarLogCompileStaticVars(
 	TEXT("fx.LogCompileStaticVars"),
 	UNiagaraScript::LogCompileStaticVars,
 	TEXT("If > 0 all compile id generation dealing with static variables will be logged.  \n"),
