Trying to narrow down cooking crash on cook  in UNiagaraStaticSwitch Resolve Numerics to failing asset. Preventing crash.

diff --git a/Engine/Plugins/FX/Niagara/Source/NiagaraEditor/Private/NiagaraNodeStaticSwitch.cpp b/Engine/Plugins/FX/Niagara/Source/NiagaraEditor/Private/NiagaraNodeStaticSwitch.cpp
--- a/Engine/Plugins/FX/Niagara/Source/NiagaraEditor/Private/NiagaraNodeStaticSwitch.cpp
+++ b/Engine/Plugins/FX/Niagara/Source/NiagaraEditor/Private/NiagaraNodeStaticSwitch.cpp
@@ -615,8 +615,16 @@ void UNiagaraNodeStaticSwitch::ResolveNumerics(const UEdGraphSchema_Niagara* Sch
 		
 		for (int32 j = 0; j < OptionValues.Num(); j++)
 		{
-			UEdGraphPin* InputPin = MasterInputPins[OutputVars.Num()*j + VarIdx];
-			InputPins.Add(InputPin);
+			const int TargetIdx = OutputVars.Num() * j + VarIdx;
+			if (MasterInputPins.IsValidIndex(TargetIdx))
+			{
+				UEdGraphPin* InputPin = MasterInputPins[TargetIdx];
+				InputPins.Add(InputPin);
+			}
+			else
+			{
+				UE_LOG(LogNiagaraEditor, Warning, TEXT("Invalid index on UNiagaraNodeStaticSwitch::ResolveNumerics %s, OptionIdx: %d VarIdx: %d MaxIndex: %d"), *GetPathName(), j, VarIdx, MasterInputPins.Num())
+			}
 		}
 
 		OutputPins.Add(OutPin);
