Manual merge + non unity fix

Inline parameters: Support for inline bool parameters
This works by assigning the edit inline function inputs a "substitute item for UI" during PostRefreshChildrenInternal in UNiagaraStackFunctionInputCollection.
Added a custom filtering option for accessing child entries so we can only filter out the hidden children for inline inputs instead of advanced, inline inputs etc.
Also minor style changes:
- Enums have a dark gray background now by default instead of blue and a new button style keeps the white text intact even when hovered
- Reduced padding to save space

diff --git a/Engine/Plugins/FX/Niagara/Source/NiagaraEditor/Private/NiagaraEditorUtilities.cpp b/Engine/Plugins/FX/Niagara/Source/NiagaraEditor/Private/NiagaraEditorUtilities.cpp
--- a/Engine/Plugins/FX/Niagara/Source/NiagaraEditor/Private/NiagaraEditorUtilities.cpp
+++ b/Engine/Plugins/FX/Niagara/Source/NiagaraEditor/Private/NiagaraEditorUtilities.cpp
@@ -63,6 +63,7 @@
 #include "Widgets/SWidget.h"
 #include "Widgets/Text/STextBlock.h"
 #include "UObject/TextProperty.h"
+#include "Widgets/SNiagaraPinTypeSelector.h"
 #include "ViewModels/NiagaraEmitterHandleViewModel.h"
 #include "ViewModels/NiagaraEmitterViewModel.h"
 
@@ -3515,6 +3516,91 @@ bool FNiagaraParameterUtilities::TestCanRenameWithMessage(FName ParameterName, F
 	return false;
 }
 
+TSharedRef<SToolTip> FNiagaraParameterUtilities::GetTooltipWidget(FNiagaraVariable Variable, bool bShowValue)
+{
+	FNiagaraTypeDefinition Type = Variable.GetType();
+	const FLinearColor TypeColor = UEdGraphSchema_Niagara::GetTypeColor(Type);
+	TSharedPtr<INiagaraEditorTypeUtilities> TypeUtilities = FNiagaraEditorModule::Get().GetTypeUtilities(Type);
+
+	FText			   IconToolTip = FText::GetEmpty();
+	FSlateBrush const* IconBrush = FAppStyle::Get().GetBrush(TEXT("Kismet.AllClasses.VariableIcon"));
+	FSlateColor        IconColor = FSlateColor(TypeColor);
+	FString			   IconDocLink, IconDocExcerpt;
+	FSlateBrush const* SecondaryIconBrush = FEditorStyle::GetBrush(TEXT("NoBrush"));
+	FSlateColor        SecondaryIconColor = IconColor;
+
+
+	TSharedRef<SVerticalBox> TooltipContent = SNew(SVerticalBox)
+		+ SVerticalBox::Slot()
+		.AutoHeight()
+		[
+			SNew(SHorizontalBox)
+			+ SHorizontalBox::Slot()
+			.AutoWidth()
+			.HAlign(HAlign_Center)
+			.VAlign(VAlign_Center)
+			[
+				SNew(SNiagaraIconWidget)
+				.IconToolTip(IconToolTip)
+				.IconBrush(IconBrush)
+				.IconColor(IconColor)
+				.DocLink(IconDocLink)
+				.DocExcerpt(IconDocExcerpt)
+				.SecondaryIconBrush(SecondaryIconBrush) 
+				.SecondaryIconColor(SecondaryIconColor)
+			]
+			+ SHorizontalBox::Slot()
+			.AutoWidth()
+			.HAlign(HAlign_Center)
+			.VAlign(VAlign_Center)
+			.Padding(3.f)
+			[
+				SNew(SNiagaraParameterNameTextBlock)
+				.IsReadOnly(true)
+				.ParameterText(FText::FromName(Variable.GetName()))
+			]
+		];
+
+	if(bShowValue && Variable.IsDataAllocated())
+	{
+		const FText ValueText = TypeUtilities->GetStackDisplayText(Variable);
+
+		TooltipContent->AddSlot()
+		.AutoHeight()
+		[
+			SNew(SHorizontalBox)
+			+ SHorizontalBox::Slot()
+			.AutoWidth()
+			.HAlign(HAlign_Center)
+			.VAlign(VAlign_Center)
+			.Padding(3.f)
+			[
+				SNew(STextBlock)
+				.Text(LOCTEXT("ParameterTooltipText", "Value: "))
+			]
+			+ SHorizontalBox::Slot()
+			.AutoWidth()
+			.HAlign(HAlign_Center)
+			.VAlign(VAlign_Center)
+			.Padding(3.f)
+			[
+				SNew(STextBlock)
+				.Text(ValueText)
+			]
+		];
+	}
+	
+	
+	// we construct a tooltip widget that shows the parameter the value is associated with
+	TSharedRef<SToolTip> TooltipWidget = SNew(SToolTip)
+	.Content()
+	[
+		TooltipContent
+	];
+
+	return TooltipWidget;
+}
+
 void FNiagaraParameterUtilities::FilterToRelevantStaticVariables(const TArray<FNiagaraVariable>& InVars, TArray<FNiagaraVariable>& OutVars, FName InOldEmitterAlias, FName InNewEmitterAlias, bool bFilterByEmitterAliasAndConvertToUnaliased)
 {
 	FNiagaraAliasContext RenameContext(ENiagaraScriptUsage::ParticleSpawnScript);
diff --git a/Engine/Plugins/FX/Niagara/Source/NiagaraEditor/Private/ViewModels/Stack/NiagaraStackEntry.cpp b/Engine/Plugins/FX/Niagara/Source/NiagaraEditor/Private/ViewModels/Stack/NiagaraStackEntry.cpp
--- a/Engine/Plugins/FX/Niagara/Source/NiagaraEditor/Private/ViewModels/Stack/NiagaraStackEntry.cpp
+++ b/Engine/Plugins/FX/Niagara/Source/NiagaraEditor/Private/ViewModels/Stack/NiagaraStackEntry.cpp
@@ -389,6 +389,27 @@ void UNiagaraStackEntry::GetFilteredChildren(TArray<UNiagaraStackEntry*>& OutFil
 	OutFilteredChildren.Append(FilteredChildren);
 }
 
+void UNiagaraStackEntry::GetCustomFilteredChildren(TArray<UNiagaraStackEntry*>& OutFilteredChildren, const TArray<FOnFilterChild>& CustomChildFilters) const
+{
+	for (UNiagaraStackEntry* Child : Children)
+	{
+		bool bPassesFilter = true;
+		for (const FOnFilterChild& ChildFilter : CustomChildFilters)
+		{
+			if (ChildFilter.Execute(*Child) == false)
+			{
+				bPassesFilter = false;
+				break;
+			}
+		}
+
+		if (bPassesFilter)
+		{
+			OutFilteredChildren.Add(Child);
+		}
+	}
+}
+
 void UNiagaraStackEntry::GetUnfilteredChildren(TArray<UNiagaraStackEntry*>& OutUnfilteredChildren) const
 {
 	OutUnfilteredChildren.Append(ErrorChildren);
diff --git a/Engine/Plugins/FX/Niagara/Source/NiagaraEditor/Private/ViewModels/Stack/NiagaraStackFunctionInput.cpp b/Engine/Plugins/FX/Niagara/Source/NiagaraEditor/Private/ViewModels/Stack/NiagaraStackFunctionInput.cpp
--- a/Engine/Plugins/FX/Niagara/Source/NiagaraEditor/Private/ViewModels/Stack/NiagaraStackFunctionInput.cpp
+++ b/Engine/Plugins/FX/Niagara/Source/NiagaraEditor/Private/ViewModels/Stack/NiagaraStackFunctionInput.cpp
@@ -2013,6 +2013,16 @@ UNiagaraStackFunctionInput::FOnValueChanged& UNiagaraStackFunctionInput::OnValue
 	return ValueChangedDelegate;
 }
 
+TOptional<FNiagaraVariable> UNiagaraStackFunctionInput::GetEditConditionVariable() const
+{
+	if(GetHasEditCondition())
+	{
+		return FNiagaraVariable(EditCondition.GetConditionInputType(), EditCondition.GetConditionInputName());
+	}
+
+	return TOptional<FNiagaraVariable>();
+}
+
 bool UNiagaraStackFunctionInput::GetHasEditCondition() const
 {
 	return EditCondition.IsValid();
@@ -2395,10 +2405,11 @@ bool UNiagaraStackFunctionInput::IsScratchDynamicInput() const
 
 bool UNiagaraStackFunctionInput::ShouldDisplayInline() const
 {
+	// we only allow values to show up that are supposed to be visible
 	// we don't allow advanced values to be inlined as they can introduce problems with filtering
-	if(InputValues.Mode == EValueMode::Local && InputMetaData.IsSet() && !GetIsAdvanced() && InputMetaData->bDisplayInOverviewStack)
+	if(InputMetaData.IsSet() && InputMetaData->bDisplayInOverviewStack && InputValues.Mode == EValueMode::Local && GetShouldPassFilterForVisibleCondition())
 	{
-		return true;
+		return true;				
 	}
 
 	return false;
diff --git a/Engine/Plugins/FX/Niagara/Source/NiagaraEditor/Private/ViewModels/Stack/NiagaraStackFunctionInputCollection.cpp b/Engine/Plugins/FX/Niagara/Source/NiagaraEditor/Private/ViewModels/Stack/NiagaraStackFunctionInputCollection.cpp
--- a/Engine/Plugins/FX/Niagara/Source/NiagaraEditor/Private/ViewModels/Stack/NiagaraStackFunctionInputCollection.cpp
+++ b/Engine/Plugins/FX/Niagara/Source/NiagaraEditor/Private/ViewModels/Stack/NiagaraStackFunctionInputCollection.cpp
@@ -145,13 +145,13 @@ void UNiagaraStackFunctionInputCollection::GetChildInputs(TArray<UNiagaraStackFu
 	}
 }
 
-void UNiagaraStackFunctionInputCollection::GetFilteredChildInputs(TArray<UNiagaraStackFunctionInput*>& OutResult) const
+void UNiagaraStackFunctionInputCollection::GetCustomFilteredChildInputs(TArray<UNiagaraStackFunctionInput*>& OutResult, const TArray<FOnFilterChild>& CustomFilters) const
 {
 	TArray<UNiagaraStackInputCategory*> ChildCategories;
-	GetFilteredChildrenOfType(ChildCategories);
+	GetUnfilteredChildrenOfType(ChildCategories);
 	for (UNiagaraStackInputCategory* ChildCategory : ChildCategories)
 	{
-		ChildCategory->GetFilteredChildrenOfType(OutResult);
+		ChildCategory->GetCustomFilteredChildrenOfType(OutResult, CustomFilters);
 	}
 }
 
@@ -160,7 +160,9 @@ TArray<UNiagaraStackFunctionInput*> UNiagaraStackFunctionInputCollection::GetInl
 	TArray<UNiagaraStackFunctionInput*> OutArray;
 	
 	TArray<UNiagaraStackFunctionInput*> FunctionInputs;
-	GetFilteredChildInputs(FunctionInputs);
+	TArray<FOnFilterChild> CustomChildFilters;
+	CustomChildFilters.Add(FOnFilterChild::CreateUObject(this, &UNiagaraStackFunctionInput::FilterHiddenChildren));
+	GetCustomFilteredChildInputs(FunctionInputs, CustomChildFilters);
 
 	for(UNiagaraStackFunctionInput* FunctionInput : FunctionInputs)
 	{
diff --git a/Engine/Plugins/FX/Niagara/Source/NiagaraEditor/Private/ViewModels/Stack/NiagaraStackModuleItem.cpp b/Engine/Plugins/FX/Niagara/Source/NiagaraEditor/Private/ViewModels/Stack/NiagaraStackModuleItem.cpp
--- a/Engine/Plugins/FX/Niagara/Source/NiagaraEditor/Private/ViewModels/Stack/NiagaraStackModuleItem.cpp
+++ b/Engine/Plugins/FX/Niagara/Source/NiagaraEditor/Private/ViewModels/Stack/NiagaraStackModuleItem.cpp
@@ -1281,6 +1281,11 @@ void UNiagaraStackModuleItem::SetInputValuesFromClipboardFunctionInputs(const TA
 	InputCollection->SetValuesFromClipboardFunctionInputs(ClipboardFunctionInputs);
 }
 
+void UNiagaraStackModuleItem::GetParameterInputs(TArray<UNiagaraStackFunctionInput*>& OutResult) const
+{
+	return InputCollection->GetChildInputs(OutResult);
+}
+
 TArray<UNiagaraStackFunctionInput*> UNiagaraStackModuleItem::GetInlineParameterInputs() const
 {
 	return InputCollection->GetInlineParameterInputs();
diff --git a/Engine/Plugins/FX/Niagara/Source/NiagaraEditor/Public/NiagaraEditorUtilities.h b/Engine/Plugins/FX/Niagara/Source/NiagaraEditor/Public/NiagaraEditorUtilities.h
--- a/Engine/Plugins/FX/Niagara/Source/NiagaraEditor/Public/NiagaraEditorUtilities.h
+++ b/Engine/Plugins/FX/Niagara/Source/NiagaraEditor/Public/NiagaraEditorUtilities.h
@@ -395,5 +395,8 @@ namespace FNiagaraParameterUtilities
 
 	NIAGARAEDITOR_API bool TestCanRenameWithMessage(FName ParameterName, FText& OutMessage);
 
+	/** Creates a tooltip based on a parameter. Also shows the value, if allocated and enabled. */
+	NIAGARAEDITOR_API TSharedRef<SToolTip> GetTooltipWidget(FNiagaraVariable Variable, bool bShowValue = true);
+
 	NIAGARAEDITOR_API void FilterToRelevantStaticVariables(const TArray<FNiagaraVariable>& InVars, TArray<FNiagaraVariable>& OutVars, FName InOldEmitterAlias, FName InNewEmitterAlias, bool bFilterByEmitterAliasAndConvertToUnaliased);
 };
diff --git a/Engine/Plugins/FX/Niagara/Source/NiagaraEditor/Public/ViewModels/Stack/NiagaraParameterHandle.h b/Engine/Plugins/FX/Niagara/Source/NiagaraEditor/Public/ViewModels/Stack/NiagaraParameterHandle.h
--- a/Engine/Plugins/FX/Niagara/Source/NiagaraEditor/Public/ViewModels/Stack/NiagaraParameterHandle.h
+++ b/Engine/Plugins/FX/Niagara/Source/NiagaraEditor/Public/ViewModels/Stack/NiagaraParameterHandle.h
@@ -74,4 +74,9 @@ private:
 	FName Name;
 	FName Namespace;
 	mutable TArray<FName> HandlePartsCache;
-};
\ No newline at end of file
+};
+
+FORCEINLINE uint32 GetTypeHash(const FNiagaraParameterHandle& Var)
+{
+	return GetTypeHash(Var.GetParameterHandleString());
+}
diff --git a/Engine/Plugins/FX/Niagara/Source/NiagaraEditor/Public/ViewModels/Stack/NiagaraStackEntry.h b/Engine/Plugins/FX/Niagara/Source/NiagaraEditor/Public/ViewModels/Stack/NiagaraStackEntry.h
--- a/Engine/Plugins/FX/Niagara/Source/NiagaraEditor/Public/ViewModels/Stack/NiagaraStackEntry.h
+++ b/Engine/Plugins/FX/Niagara/Source/NiagaraEditor/Public/ViewModels/Stack/NiagaraStackEntry.h
@@ -294,6 +294,8 @@ public:
 	
 	void GetUnfilteredChildren(TArray<UNiagaraStackEntry*>& OutUnfilteredChildren) const;
 
+	void GetCustomFilteredChildren(TArray<UNiagaraStackEntry*>& OutFilteredChildren, const TArray<FOnFilterChild>& ChildFilters) const;
+
 	void GetFilteredChildrenOfTypes(TArray<UNiagaraStackEntry*>& OutFilteredChildren, const TSet<UClass*>& AllowedClasses) const
 	{
 		TArray<UNiagaraStackEntry*> FilteredChildrenTmp;
@@ -347,6 +349,31 @@ public:
 		}
 	}
 
+	template<typename T>
+	void GetCustomFilteredChildrenOfType(TArray<T*>& OutFilteredChildrenOfType, const TArray<FOnFilterChild>& CustomChildFilters) const
+	{
+		for (UNiagaraStackEntry* Child : Children)
+		{
+			if(T* CastChild = Cast<T>(Child))
+			{
+				bool bPassesFilter = true;
+				for (const FOnFilterChild& ChildFilter : CustomChildFilters)
+				{
+					if (ChildFilter.Execute(*CastChild) == false)
+					{
+						bPassesFilter = false;
+						break;
+					}
+				}
+
+				if (bPassesFilter)
+				{
+					OutFilteredChildrenOfType.Add(CastChild);
+				}				
+			}
+		}
+	}
+
 	FOnExpansionChanged& OnExpansionChanged();
 
 	FOnExpansionChanged& OnExpansionInOverviewChanged();
diff --git a/Engine/Plugins/FX/Niagara/Source/NiagaraEditor/Public/ViewModels/Stack/NiagaraStackFunctionInput.h b/Engine/Plugins/FX/Niagara/Source/NiagaraEditor/Public/ViewModels/Stack/NiagaraStackFunctionInput.h
--- a/Engine/Plugins/FX/Niagara/Source/NiagaraEditor/Public/ViewModels/Stack/NiagaraStackFunctionInput.h
+++ b/Engine/Plugins/FX/Niagara/Source/NiagaraEditor/Public/ViewModels/Stack/NiagaraStackFunctionInput.h
@@ -109,7 +109,7 @@ public:
 	virtual FText GetPasteTransactionText(const UNiagaraClipboardContent* ClipboardContent) const override;
 	virtual void Paste(const UNiagaraClipboardContent* ClipboardContent, FText& OutPasteWarning) override;
 	virtual bool HasOverridenContent() const override;
-
+	
 	/** Gets the tooltip that should be shown for the value of this input. */
 	FText GetValueToolTip() const;
 
@@ -214,6 +214,9 @@ public:
 	/** Gets a multicast delegate which is called whenever the value on this input changes. */
 	FOnValueChanged& OnValueChanged();
 
+	/** Gets the variable that serves as an edit condition for this input. */
+	TOptional<FNiagaraVariable> GetEditConditionVariable() const;
+	
 	/** Gets whether or not this input has an associated edit condition input. */
 	bool GetHasEditCondition() const;
 
@@ -470,7 +473,7 @@ private:
 
 	/** An input condition handler for the visible condition. */
 	FNiagaraStackFunctionInputCondition VisibleCondition;
-
+	
 	/** Whether or not to show an inline control for the edit condition input. */
 	bool bShowEditConditionInline;
 
diff --git a/Engine/Plugins/FX/Niagara/Source/NiagaraEditor/Public/ViewModels/Stack/NiagaraStackFunctionInputCollection.h b/Engine/Plugins/FX/Niagara/Source/NiagaraEditor/Public/ViewModels/Stack/NiagaraStackFunctionInputCollection.h
--- a/Engine/Plugins/FX/Niagara/Source/NiagaraEditor/Public/ViewModels/Stack/NiagaraStackFunctionInputCollection.h
+++ b/Engine/Plugins/FX/Niagara/Source/NiagaraEditor/Public/ViewModels/Stack/NiagaraStackFunctionInputCollection.h
@@ -99,7 +99,7 @@ public:
 
 	void GetChildInputs(TArray<UNiagaraStackFunctionInput*>& OutResult) const;
 
-	void GetFilteredChildInputs(TArray<UNiagaraStackFunctionInput*>& OutResult) const;
+	void GetCustomFilteredChildInputs(TArray<UNiagaraStackFunctionInput*>& OutResult, const TArray<FOnFilterChild>& CustomFilters) const;
 
 	TArray<UNiagaraStackFunctionInput*> GetInlineParameterInputs() const;
 
diff --git a/Engine/Plugins/FX/Niagara/Source/NiagaraEditor/Public/ViewModels/Stack/NiagaraStackModuleItem.h b/Engine/Plugins/FX/Niagara/Source/NiagaraEditor/Public/ViewModels/Stack/NiagaraStackModuleItem.h
--- a/Engine/Plugins/FX/Niagara/Source/NiagaraEditor/Public/ViewModels/Stack/NiagaraStackModuleItem.h
+++ b/Engine/Plugins/FX/Niagara/Source/NiagaraEditor/Public/ViewModels/Stack/NiagaraStackModuleItem.h
@@ -74,6 +74,7 @@ public:
 
 	void SetInputValuesFromClipboardFunctionInputs(const TArray<const UNiagaraClipboardFunctionInput*>& ClipboardFunctionInputs);
 
+	void GetParameterInputs(TArray<class UNiagaraStackFunctionInput*>& OutResult) const;
 	TArray<class UNiagaraStackFunctionInput*> GetInlineParameterInputs() const;
 
 	virtual bool SupportsCut() const override { return true; }
diff --git a/Engine/Plugins/FX/Niagara/Source/NiagaraEditorWidgets/Private/NiagaraEditorWidgetsStyle.cpp b/Engine/Plugins/FX/Niagara/Source/NiagaraEditorWidgets/Private/NiagaraEditorWidgetsStyle.cpp
--- a/Engine/Plugins/FX/Niagara/Source/NiagaraEditorWidgets/Private/NiagaraEditorWidgetsStyle.cpp
+++ b/Engine/Plugins/FX/Niagara/Source/NiagaraEditorWidgets/Private/NiagaraEditorWidgetsStyle.cpp
@@ -166,6 +166,18 @@ TSharedRef< FSlateStyleSet > FNiagaraEditorWidgetsStyle::Create()
 		.SetNormalPadding(FMargin(2.f, 2.f, 2.f, 2.f))
 		.SetPressedPadding(FMargin(2.f, 1.0f, 2.f, 0.0f));
 
+	FButtonStyle InlineParameterButtonStyleEnum = FButtonStyle()
+		.SetNormal(FSlateRoundedBoxBrush(FStyleColors::AccentGray, 4.0f, FStyleColors::Input, 1.f))
+		.SetHovered(FSlateRoundedBoxBrush(FStyleColors::White, 4.0f, FStyleColors::Input, 1.f))
+		.SetPressed(FSlateRoundedBoxBrush(FStyleColors::Header, 4.0f, FStyleColors::Input, 1.f))
+		.SetDisabled(FSlateRoundedBoxBrush(FStyleColors::Dropdown, 4.0f, FStyleColors::Recessed, 1.f))
+		.SetNormalForeground(FStyleColors::ForegroundHover)
+		.SetHoveredForeground(FStyleColors::ForegroundHover)
+		.SetPressedForeground(FStyleColors::ForegroundHover)
+		.SetDisabledForeground(FStyleColors::Foreground)
+		.SetNormalPadding(FMargin(2.f, 2.f, 2.f, 2.f))
+		.SetPressedPadding(FMargin(2.f, 1.0f, 2.f, 0.0f));
+
 	// this style is used when the user has specified a custom color for an inline parameter. By default we'd tint it a bit more gray.
 	FButtonStyle InlineParameterButtonStyleNoTint = FButtonStyle()
 		.SetNormal(FSlateRoundedBoxBrush(FStyleColors::White, 4.0f, FStyleColors::Input, 1.f))
@@ -179,18 +191,17 @@ TSharedRef< FSlateStyleSet > FNiagaraEditorWidgetsStyle::Create()
 		.SetNormalPadding(FMargin(2.f, 2.f, 2.f, 2.f))
 		.SetPressedPadding(FMargin(2.f, 1.0f, 2.f, 0.0f));
 
-	// this style is used when the user has specified an icon to represent a value. In that case, we no longer need a proper border.
-	FButtonStyle InlineParameterButtonStyleTransparent = FButtonStyle()
-		.SetNormal(FSlateRoundedBoxBrush(FStyleColors::Transparent, 4.0f, FStyleColors::Input, 0.f))
-		.SetHovered(FSlateRoundedBoxBrush(FStyleColors::Hover, 4.0f, FStyleColors::Input, 0.f))
-		.SetPressed(FSlateRoundedBoxBrush(FStyleColors::Header, 4.0f, FStyleColors::Input, 0.f))
-		.SetDisabled(FSlateRoundedBoxBrush(FStyleColors::Dropdown, 4.0f, FStyleColors::Recessed, 0.f))
-		.SetNormalForeground(FStyleColors::ForegroundHover)
+	FButtonStyle InlineParameterButtonStyleIcon = FButtonStyle()
+		.SetNormal(FSlateNoResource())
+		.SetHovered(FSlateRoundedBoxBrush(FStyleColors::Dropdown, 4.0f))
+		.SetPressed(FSlateRoundedBoxBrush(FStyleColors::Dropdown, 4.0f))
+		.SetDisabled(FSlateNoResource())
+		.SetNormalForeground(FStyleColors::Foreground)
 		.SetHoveredForeground(FStyleColors::ForegroundHover)
 		.SetPressedForeground(FStyleColors::ForegroundHover)
 		.SetDisabledForeground(FStyleColors::Foreground)
-		.SetNormalPadding(FMargin(2.f))
-		.SetPressedPadding(FMargin(2.f, 1.0f, 2.f, 0.0f));
+		.SetNormalPadding(FMargin(1.f))
+		.SetPressedPadding(FMargin(1.f, 2.f, 1.f, 0.f));
 	
 	Style->Set("NiagaraEditor.SystemOverview.ItemText", SystemOverviewItemText);
 	
@@ -198,8 +209,10 @@ TSharedRef< FSlateStyleSet > FNiagaraEditorWidgetsStyle::Create()
 
 	Style->Set("NiagaraEditor.SystemOverview.InlineParameterText", SystemOverviewInlineParametersText);
 	Style->Set("NiagaraEditor.SystemOverview.InlineParameterButton", InlineParameterButtonStyle);
-	Style->Set("NiagaraEditor.SystemOverview.InlineParameterButton.Transparent", InlineParameterButtonStyleTransparent);
+	Style->Set("NiagaraEditor.SystemOverview.InlineParameterButton.Enum", InlineParameterButtonStyleEnum);
 	Style->Set("NiagaraEditor.SystemOverview.InlineParameterButton.NoTint", InlineParameterButtonStyleNoTint);
+	Style->Set("NiagaraEditor.SystemOverview.InlineParameterButton.Icon", InlineParameterButtonStyleIcon);
+
 
 
 	Style->Set("NiagaraEditor.SystemOverview.Item.BackgroundColor", FLinearColor(FColor(62, 62, 62)));
diff --git a/Engine/Plugins/FX/Niagara/Source/NiagaraEditorWidgets/Private/SNiagaraOverviewInlineParameterBox.cpp b/Engine/Plugins/FX/Niagara/Source/NiagaraEditorWidgets/Private/SNiagaraOverviewInlineParameterBox.cpp
--- a/Engine/Plugins/FX/Niagara/Source/NiagaraEditorWidgets/Private/SNiagaraOverviewInlineParameterBox.cpp
+++ b/Engine/Plugins/FX/Niagara/Source/NiagaraEditorWidgets/Private/SNiagaraOverviewInlineParameterBox.cpp
@@ -9,6 +9,7 @@
 #include "NiagaraScriptVariable.h"
 #include "NiagaraSystem.h"
 #include "NiagaraEditorModule.h"
+#include "NiagaraEditorUtilities.h"
 #include "Styling/StyleColors.h"
 #include "ViewModels/NiagaraSystemSelectionViewModel.h"
 #include "ViewModels/NiagaraSystemViewModel.h"
@@ -54,12 +55,7 @@ SNiagaraOverviewInlineParameterBox::~SNiagaraOverviewInlineParameterBox()
 	BoundFunctionInputs.Empty();
 }
 
-int32 SNiagaraOverviewInlineParameterBox::GetNumChildWidgets()
-{
-	return NumParameterWidgets;
-}
-
-FReply SNiagaraOverviewInlineParameterBox::NavigateToStack(TWeakObjectPtr<UNiagaraStackFunctionInput> FunctionInput)
+FReply SNiagaraOverviewInlineParameterBox::NavigateToStack(TWeakObjectPtr<const UNiagaraStackFunctionInput> FunctionInput)
 {
 	// even if we can't navigate, the button should consume the input
 	if(!ModuleItem.IsValid() || !FunctionInput.IsValid() || FunctionInput->IsFinalized())
@@ -69,6 +65,13 @@ FReply SNiagaraOverviewInlineParameterBox::NavigateToStack(TWeakObjectPtr<UNiaga
 	
 	ModuleItem->GetSystemViewModel()->GetSelectionViewModel()->UpdateSelectedEntries({ModuleItem.Get()}, {}, true);
 
+	// we attempt to find a substitute entry, i.e. if our current input is an inline edit toggle, we try to find the input that the toggle is handling instead
+	TWeakObjectPtr<const UNiagaraStackFunctionInput> SubstituteEntry = FindSubstituteEntry(FunctionInput.Get());
+	if(SubstituteEntry.IsValid())
+	{
+		FunctionInput = SubstituteEntry;
+	}
+	
 	TArray<UNiagaraStackEntry::FStackSearchItem> SearchItems;
 	FunctionInput->GetSearchItems(SearchItems);
 	if(ensure(SearchItems.Num() > 0))
@@ -83,8 +86,8 @@ FReply SNiagaraOverviewInlineParameterBox::NavigateToStack(TWeakObjectPtr<UNiaga
 void SNiagaraOverviewInlineParameterBox::ConstructChildren()
 {
 	Container->ClearChildren();
-	NumParameterWidgets = 0;
-
+	ImageBrushes.Empty();
+	
 	for(TWeakObjectPtr<UNiagaraStackFunctionInput> BoundInput : BoundFunctionInputs)
 	{
 		if(BoundInput.IsValid())
@@ -101,12 +104,10 @@ void SNiagaraOverviewInlineParameterBox::ConstructChildren()
 		Container->AddSlot()
 		.VAlign(VAlign_Center)
 		.HAlign(HAlign_Center)
-		.Padding(FMargin(2.f, 0.f))
+		.Padding(FMargin(1.f, 0.f))
 		[
 			ParameterWidget
 		];
-
-		NumParameterWidgets++;
 	}
 
 	SetVisibility(ParameterWidgets.Num() > 0 ? EVisibility::Visible : EVisibility::Collapsed);
@@ -177,14 +178,22 @@ TSharedRef<SWidget> SNiagaraOverviewInlineParameterBox::GenerateParameterWidgetF
 	TSharedPtr<INiagaraEditorTypeUtilities> TypeUtilities = FNiagaraEditorModule::Get().GetTypeUtilities(Type);
 
 	// we construct a variable with the data from our local struct memory because type utilities don't work on raw data
-	FNiagaraVariable Variable(Type, NAME_None);
+	FNiagaraVariable Variable(Type, FunctionInput->GetInputParameterHandle().GetParameterHandleString());
 	Variable.SetData(FunctionInput->GetLocalValueStruct()->GetStructMemory());
 	
 	// by default we display type information. In some cases we have overrides, but even then we want to keep the type information as we continue using it in the tooltips
 	const FText ValueText = TypeUtilities->GetStackDisplayText(Variable);
 	FText DisplayedText = ValueText;
-	const FLinearColor TypeColor = UEdGraphSchema_Niagara::GetTypeColor(Type);
-	FLinearColor DisplayedColor = TypeColor;
+	const FLinearColor OriginalTypeColor = UEdGraphSchema_Niagara::GetTypeColor(Type);
+	FLinearColor ActualTypeColor = OriginalTypeColor;
+
+	// if we have an enum type, we want to override the type color to be gray instead
+	if(Type.GetEnum())
+	{
+		ActualTypeColor = FLinearColor(0.02f, 0.02f, 0.02f, 1.f);
+	}
+	
+	FLinearColor DisplayedColor = InputMetaData->bOverrideColor ? InputMetaData->InlineParameterColorOverride : ActualTypeColor;
 
 	UTexture2D* Icon = nullptr;
 	if (const UEnum* Enum = Type.GetEnum())
@@ -203,7 +212,6 @@ TSharedRef<SWidget> SNiagaraOverviewInlineParameterBox::GenerateParameterWidgetF
 	else if(Type == FNiagaraTypeDefinition::GetBoolDef() && InputMetaData->bEnableBoolOverride)
 	{
 		const FNiagaraBoolParameterMetaData& BoolParameterMetaData = InputMetaData->InlineParameterBoolOverride;
-		DisplayedColor = InputMetaData->bOverrideColor ? InputMetaData->InlineParameterColorOverride : DisplayedColor;
 
 		int32 Val = *(int32*)FunctionInput->GetLocalValueStruct()->GetStructMemory();
 		bool bParameterValue = Val == 0 ? false : true;
@@ -240,7 +248,8 @@ TSharedRef<SWidget> SNiagaraOverviewInlineParameterBox::GenerateParameterWidgetF
 	// we might have cases in which we don't want to display a parameter's data, such as if we only want to display a bool parameter when it is set to True
 	if(bGenerateProperWidget)
 	{
-		TWeakObjectPtr<UNiagaraStackFunctionInput> FunctionInputWeak = FunctionInput;
+		TWeakObjectPtr<const UNiagaraStackFunctionInput> FunctionInputWeak = FunctionInput;
+		
 		TSharedRef<SButton> ParameterWidgetButton = SNew(SButton)
 		.Text(DisplayedText)
 		.TextStyle(&FNiagaraEditorWidgetsStyle::Get().GetWidgetStyle<FTextBlockStyle>("NiagaraEditor.SystemOverview.InlineParameterText"))
@@ -251,23 +260,25 @@ TSharedRef<SWidget> SNiagaraOverviewInlineParameterBox::GenerateParameterWidgetF
 		// if we have an icon available, we use it instead of any text. Border color is irrelevant here.
 		if(Icon != nullptr)
 		{
-			FSlateBrush& ImageBrush = ImageBrushes.Emplace_GetRef();
-			ImageBrush.SetResourceObject(Icon);
-			ImageBrush.ImageSize.X = 16;
-			ImageBrush.ImageSize.Y = 16;
-	
+			FSlateImageBrush& ImageBrush = ImageBrushes.Emplace_GetRef(Icon, FVector2D(16.f, 16.f));
 			TSharedRef<SImage> ParameterWidgetIcon= SNew(SImage).Image(&ImageBrush);
 			ParameterWidgetButton->SetContent(ParameterWidgetIcon);
-			ParameterWidgetButton->SetButtonStyle(&FNiagaraEditorWidgetsStyle::Get().GetWidgetStyle<FButtonStyle>("NiagaraEditor.SystemOverview.InlineParameterButton.Transparent"));
+			ParameterWidgetButton->SetButtonStyle(&FNiagaraEditorWidgetsStyle::Get().GetWidgetStyle<FButtonStyle>("NiagaraEditor.SystemOverview.InlineParameterButton.Icon"));
 		}
 		// otherwise, we use text and set the border to its specified color (type, override, instance override color)
 		else
 		{
 			// if we are using a custom color, we set the button style to use a white base in order to get rid of a tint
-			if(DisplayedColor != TypeColor)
+			if(DisplayedColor != ActualTypeColor)
 			{
 				ParameterWidgetButton->SetButtonStyle(&FNiagaraEditorWidgetsStyle::Get().GetWidgetStyle<FButtonStyle>("NiagaraEditor.SystemOverview.InlineParameterButton.NoTint"));
 			}
+
+			// if we are an enum, we use a style that better fits the text
+			if(Type.GetEnum())
+			{
+				ParameterWidgetButton->SetButtonStyle(&FNiagaraEditorWidgetsStyle::Get().GetWidgetStyle<FButtonStyle>("NiagaraEditor.SystemOverview.InlineParameterButton.Enum"));
+			}
 			
 			ParameterWidgetButton->SetBorderBackgroundColor(DisplayedColor);
 		}			
@@ -279,72 +290,8 @@ TSharedRef<SWidget> SNiagaraOverviewInlineParameterBox::GenerateParameterWidgetF
 		return SNullWidget::NullWidget;
 	}
 	
-
-	FText			   IconToolTip = FText::GetEmpty();
-	FSlateBrush const* IconBrush = FAppStyle::Get().GetBrush(TEXT("Kismet.AllClasses.VariableIcon"));
-	FSlateColor        IconColor = FSlateColor(TypeColor);
-	FString			   IconDocLink, IconDocExcerpt;
-	FSlateBrush const* SecondaryIconBrush = FEditorStyle::GetBrush(TEXT("NoBrush"));
-	FSlateColor        SecondaryIconColor = IconColor;
-
 	// we construct a tooltip widget that shows the parameter the value is associated with
-	TSharedRef<SToolTip> TooltipWidget = SNew(SToolTip).Content()
-	[
-		SNew(SVerticalBox)
-		+ SVerticalBox::Slot()
-		.AutoHeight()
-		[
-			SNew(SHorizontalBox)
-			+ SHorizontalBox::Slot()
-			.AutoWidth()
-			.HAlign(HAlign_Center)
-			.VAlign(VAlign_Center)
-			[
-				SNew(SNiagaraIconWidget)
-				.IconToolTip(IconToolTip)
-				.IconBrush(IconBrush)
-				.IconColor(IconColor)
-				.DocLink(IconDocLink)
-				.DocExcerpt(IconDocExcerpt)
-				.SecondaryIconBrush(SecondaryIconBrush) 
-				.SecondaryIconColor(SecondaryIconColor)
-			]
-			+ SHorizontalBox::Slot()
-			.AutoWidth()
-			.HAlign(HAlign_Center)
-			.VAlign(VAlign_Center)
-			.Padding(3.f)
-			[
-				SNew(SNiagaraParameterNameTextBlock)
-				.IsReadOnly(true)
-				.ParameterText(FText::FromName(FunctionInput->GetInputParameterHandle().GetParameterHandleString()))
-			]
-		]
-		+ SVerticalBox::Slot()
-		.AutoHeight()
-		[
-			SNew(SHorizontalBox)
-			+ SHorizontalBox::Slot()
-			.AutoWidth()
-			.HAlign(HAlign_Center)
-			.VAlign(VAlign_Center)
-			.Padding(3.f)
-			[
-				SNew(STextBlock)
-				.Text(LOCTEXT("ParameterTooltipText", "Value: "))
-			]
-			+ SHorizontalBox::Slot()
-			.AutoWidth()
-			.HAlign(HAlign_Center)
-			.VAlign(VAlign_Center)
-			.Padding(3.f)
-			[
-				SNew(STextBlock)
-				.Text(ValueText)
-			]
-		]
-	];
-
+	TSharedRef<SToolTip> TooltipWidget = FNiagaraParameterUtilities::GetTooltipWidget(Variable);
 	ParameterWidget->SetToolTip(TooltipWidget);
 	return ParameterWidget.ToSharedRef();
 }
@@ -385,4 +332,35 @@ TSharedRef<SWidget> SNiagaraOverviewInlineParameterBox::GenerateParameterWidgetF
 	return SNullWidget::NullWidget;
 }
 
+TWeakObjectPtr<const UNiagaraStackFunctionInput> SNiagaraOverviewInlineParameterBox::FindSubstituteEntry(const UNiagaraStackFunctionInput* InInput)
+{
+	// there can only be substitute entries for inline edit condition toggles
+	if(!InInput->GetInputMetaData().IsSet() || !InInput->GetInputMetaData()->bInlineEditConditionToggle)
+	{
+		return nullptr;
+	}
+	
+	TArray<UNiagaraStackFunctionInput*> AllInputs;
+	ModuleItem->GetParameterInputs(AllInputs);
+
+	// we gather all edit conditions here so we can later map them back
+	for(const UNiagaraStackFunctionInput* Input : AllInputs)
+	{
+		if(Input->GetHasEditCondition())
+		{
+			if(Input->GetEditConditionVariable().IsSet())
+			{
+				FNiagaraParameterHandle EditHandle(Input->GetEditConditionVariable().GetValue().GetName());
+
+				if(InInput->GetInputParameterHandle() == EditHandle)
+				{
+					return Input;
+				}
+			}
+		}
+	}
+
+	return nullptr;
+}
+
 #undef LOCTEXT_NAMESPACE
diff --git a/Engine/Plugins/FX/Niagara/Source/NiagaraEditorWidgets/Private/SNiagaraOverviewInlineParameterBox.h b/Engine/Plugins/FX/Niagara/Source/NiagaraEditorWidgets/Private/SNiagaraOverviewInlineParameterBox.h
--- a/Engine/Plugins/FX/Niagara/Source/NiagaraEditorWidgets/Private/SNiagaraOverviewInlineParameterBox.h
+++ b/Engine/Plugins/FX/Niagara/Source/NiagaraEditorWidgets/Private/SNiagaraOverviewInlineParameterBox.h
@@ -35,11 +35,14 @@ public:
 
 private:
 	/** Selects the module item in the stack and searches for the function input */
-	FReply NavigateToStack(TWeakObjectPtr<UNiagaraStackFunctionInput> FunctionInput);
+	FReply NavigateToStack(TWeakObjectPtr<const UNiagaraStackFunctionInput> FunctionInput);
 	void ConstructChildren();
 	TArray<TSharedRef<SWidget>> GenerateParameterWidgets();
 	TSharedRef<SWidget> GenerateParameterWidgetFromLocalValue(UNiagaraStackFunctionInput* FunctionInput);
 	TSharedRef<SWidget> GenerateParameterWidgetFromDataInterface(UNiagaraStackFunctionInput* FunctionInput);
+	
+	/** We build a substitute map of entries that we want to */
+	TWeakObjectPtr<const UNiagaraStackFunctionInput> FindSubstituteEntry(const UNiagaraStackFunctionInput* Input);
 
 private:
 	/** The module item whose parameters this box is representing */
@@ -50,7 +53,5 @@ private:
 	/** We keep track of the function inputs we are observing through this parameter box so we can unbind the delegates later */
 	TArray<TWeakObjectPtr<UNiagaraStackFunctionInput>> BoundFunctionInputs;
 
-	TArray<FSlateBrush> ImageBrushes;
-
-	int32 NumParameterWidgets = 0;
+	TArray<FSlateImageBrush> ImageBrushes;
 };
diff --git a/Engine/Plugins/FX/Niagara/Source/NiagaraEditorWidgets/Private/Stack/SNiagaraStackFunctionInputName.cpp b/Engine/Plugins/FX/Niagara/Source/NiagaraEditorWidgets/Private/Stack/SNiagaraStackFunctionInputName.cpp
--- a/Engine/Plugins/FX/Niagara/Source/NiagaraEditorWidgets/Private/Stack/SNiagaraStackFunctionInputName.cpp
+++ b/Engine/Plugins/FX/Niagara/Source/NiagaraEditorWidgets/Private/Stack/SNiagaraStackFunctionInputName.cpp
@@ -15,6 +15,7 @@
 #include "HAL/PlatformApplicationMisc.h"
 #include "Framework/MultiBox/MultiBoxBuilder.h"
 #include "ScopedTransaction.h"
+#include "Widgets/SToolTip.h"
 
 #define LOCTEXT_NAMESPACE "NiagaraStackFunctionInputName"
 
@@ -32,15 +33,25 @@ void SNiagaraStackFunctionInputName::Construct(const FArguments& InArgs, UNiagar
 		.IsEnabled_UObject(FunctionInput, &UNiagaraStackEntry::GetOwnerIsEnabled)
 	];
 
+	TOptional<FNiagaraVariable> EditVariable = FunctionInput->GetEditConditionVariable();
+
+	TSharedRef<SCheckBox> InlineEnableToggle = SNew(SCheckBox)
+		.Visibility(this, &SNiagaraStackFunctionInputName::GetEditConditionCheckBoxVisibility)
+		.IsChecked(this, &SNiagaraStackFunctionInputName::GetEditConditionCheckState)
+		.OnCheckStateChanged(this, &SNiagaraStackFunctionInputName::OnEditConditionCheckStateChanged);
+		
+	if(FunctionInput->GetHasEditCondition() && EditVariable.IsSet())
+	{
+		TSharedRef<SToolTip> Tooltip = FNiagaraParameterUtilities::GetTooltipWidget(EditVariable.GetValue(), false);
+		InlineEnableToggle->SetToolTip(Tooltip);
+	}
+	
 	// Edit condition checkbox
 	NameBox->AddSlot()
 	.AutoWidth()
 	.Padding(0, 0, 3, 0)
 	[
-		SNew(SCheckBox)
-		.Visibility(this, &SNiagaraStackFunctionInputName::GetEditConditionCheckBoxVisibility)
-		.IsChecked(this, &SNiagaraStackFunctionInputName::GetEditConditionCheckState)
-		.OnCheckStateChanged(this, &SNiagaraStackFunctionInputName::OnEditConditionCheckStateChanged)
+		InlineEnableToggle
 	];
 	
 	// Name Label
