Rename for script variables with broken UObject names. This fixes an issue where script variables can't be properly copy pasted

diff --git a/Engine/Plugins/FX/Niagara/Source/NiagaraEditor/Private/NiagaraGraph.cpp b/Engine/Plugins/FX/Niagara/Source/NiagaraEditor/Private/NiagaraGraph.cpp
--- a/Engine/Plugins/FX/Niagara/Source/NiagaraEditor/Private/NiagaraGraph.cpp
+++ b/Engine/Plugins/FX/Niagara/Source/NiagaraEditor/Private/NiagaraGraph.cpp
@@ -156,6 +156,11 @@ void UNiagaraGraph::PostLoad()
 		}
 		else
 		{
+			if(ScriptVar->GetName().Contains(TEXT(".")) || ScriptVar->GetName().Contains(TEXT(" ")))
+			{
+				ScriptVar->Rename(*MakeUniqueObjectName(this, UNiagaraScriptVariable::StaticClass(), NAME_None).ToString(), nullptr, REN_NonTransactional | REN_ForceNoResetLoaders);
+			}
+			
 			// Conditional postload all ScriptVars to ensure static switch default values are allocated as these are required when postloading all graph nodes later.
 			ScriptVar->ConditionalPostLoad();
 		}
