Fix from UDN for niagara component handling of instance data when used within blueprints

diff --git a/Engine/Plugins/FX/Niagara/Source/Niagara/Private/NiagaraComponent.cpp b/Engine/Plugins/FX/Niagara/Source/Niagara/Private/NiagaraComponent.cpp
--- a/Engine/Plugins/FX/Niagara/Source/Niagara/Private/NiagaraComponent.cpp
+++ b/Engine/Plugins/FX/Niagara/Source/Niagara/Private/NiagaraComponent.cpp
@@ -1667,6 +1667,19 @@ void UNiagaraComponent::OnComponentCreated()
 #endif
 }
 
+// for handling cases where niagara components are stored in blueprints we can get into situations where
+// the instancing code will get confused with objects stored in the different parameter override variables
+// (TemplateParameterOverrides, InstanceParameterOverrides, OverrideParameters), so we make sure that
+// the outers are properly fixed up
+void UNiagaraComponent::PostApplyToComponent()
+{
+	Super::PostApplyToComponent();
+
+#if WITH_EDITOR
+	FixDataInterfaceOuters();
+#endif
+}
+
 void UNiagaraComponent::OnUnregister()
 {
 	Super::OnUnregister();
diff --git a/Engine/Plugins/FX/Niagara/Source/Niagara/Public/NiagaraComponent.h b/Engine/Plugins/FX/Niagara/Source/Niagara/Public/NiagaraComponent.h
--- a/Engine/Plugins/FX/Niagara/Source/Niagara/Public/NiagaraComponent.h
+++ b/Engine/Plugins/FX/Niagara/Source/Niagara/Public/NiagaraComponent.h
@@ -207,6 +207,7 @@ public:
 	virtual bool IsReadyForOwnerToAutoDestroy() const override;
 	virtual void OnComponentDestroyed(bool bDestroyingHierarchy) override;
 	virtual void OnComponentCreated() override;
+	virtual void PostApplyToComponent() override;
 	//~ End UActorComponent Interface.
 
 	//~ Begin UPrimitiveComponent Interface
