Fixed SA issue, the else if would always evaluate to true, which was desired, so I change it into an else

diff --git a/Engine/Plugins/FX/Niagara/Source/NiagaraEditorWidgets/Private/SNiagaraOverviewInlineParameterBox.cpp b/Engine/Plugins/FX/Niagara/Source/NiagaraEditorWidgets/Private/SNiagaraOverviewInlineParameterBox.cpp
--- a/Engine/Plugins/FX/Niagara/Source/NiagaraEditorWidgets/Private/SNiagaraOverviewInlineParameterBox.cpp
+++ b/Engine/Plugins/FX/Niagara/Source/NiagaraEditorWidgets/Private/SNiagaraOverviewInlineParameterBox.cpp
@@ -221,7 +221,7 @@ TSharedRef<SWidget> SNiagaraOverviewInlineParameterBox::GenerateParameterWidgetF
 			DisplayedText = !BoolParameterMetaData.OverrideNameTrue.IsNone() ? FText::FromName(BoolParameterMetaData.OverrideNameTrue) : DisplayedText;			
 			Icon = BoolParameterMetaData.IconOverrideTrue;
 		}
-		else if(bParameterValue == false)
+		else
 		{
 			DisplayedText = !BoolParameterMetaData.OverrideNameFalse.IsNone() ? FText::FromName(BoolParameterMetaData.OverrideNameFalse) : DisplayedText;			
 			Icon = BoolParameterMetaData.IconOverrideFalse;
