Added an is valid check to keep a garbage collected DI from crashing the engine during non-critical function calls

diff --git a/Engine/Plugins/FX/Niagara/Source/NiagaraEditor/Private/ViewModels/Stack/NiagaraStackFunctionInput.cpp b/Engine/Plugins/FX/Niagara/Source/NiagaraEditor/Private/ViewModels/Stack/NiagaraStackFunctionInput.cpp
--- a/Engine/Plugins/FX/Niagara/Source/NiagaraEditor/Private/ViewModels/Stack/NiagaraStackFunctionInput.cpp
+++ b/Engine/Plugins/FX/Niagara/Source/NiagaraEditor/Private/ViewModels/Stack/NiagaraStackFunctionInput.cpp
@@ -1655,7 +1655,11 @@ bool UNiagaraStackFunctionInput::CanReset() const
 			switch (InputValues.Mode)
 			{
 			case EValueMode::Data:
-				bNewCanReset = InputValues.DataObject->Equals(DefaultInputValues.DataObject.Get()) == false;
+				// the DI could have been garbage collected at this point
+				if(InputValues.DataObject.IsValid())
+				{
+					bNewCanReset = InputValues.DataObject->Equals(DefaultInputValues.DataObject.Get()) == false;
+				}
 				break;
 			case EValueMode::Dynamic:
 				// For now assume that default dynamic inputs can always be reset to default since they're not currently supported properly.
