Fixed the inline button style being overridden twice if the color is customized + it's an enum.

diff --git a/Engine/Plugins/FX/Niagara/Source/NiagaraEditorWidgets/Private/SNiagaraOverviewInlineParameterBox.cpp b/Engine/Plugins/FX/Niagara/Source/NiagaraEditorWidgets/Private/SNiagaraOverviewInlineParameterBox.cpp
--- a/Engine/Plugins/FX/Niagara/Source/NiagaraEditorWidgets/Private/SNiagaraOverviewInlineParameterBox.cpp
+++ b/Engine/Plugins/FX/Niagara/Source/NiagaraEditorWidgets/Private/SNiagaraOverviewInlineParameterBox.cpp
@@ -273,9 +273,8 @@ TSharedRef<SWidget> SNiagaraOverviewInlineParameterBox::GenerateParameterWidgetF
 			{
 				ParameterWidgetButton->SetButtonStyle(&FNiagaraEditorWidgetsStyle::Get().GetWidgetStyle<FButtonStyle>("NiagaraEditor.SystemOverview.InlineParameterButton.NoTint"));
 			}
-
-			// if we are an enum, we use a style that better fits the text
-			if(Type.GetEnum())
+			// if we are an enum and haven't overridden the color, we use a style that better fits the text
+			else if(Type.GetEnum())
 			{
 				ParameterWidgetButton->SetButtonStyle(&FNiagaraEditorWidgetsStyle::Get().GetWidgetStyle<FButtonStyle>("NiagaraEditor.SystemOverview.InlineParameterButton.Enum"));
 			}
