Simple null check for static analysis cleanup

diff --git a/Engine/Plugins/FX/Niagara/Source/NiagaraEditor/Private/NiagaraParameterMapHistory.cpp b/Engine/Plugins/FX/Niagara/Source/NiagaraEditor/Private/NiagaraParameterMapHistory.cpp
--- a/Engine/Plugins/FX/Niagara/Source/NiagaraEditor/Private/NiagaraParameterMapHistory.cpp
+++ b/Engine/Plugins/FX/Niagara/Source/NiagaraEditor/Private/NiagaraParameterMapHistory.cpp
@@ -915,7 +915,7 @@ int32 FNiagaraParameterMapHistoryBuilder::RegisterParameterMapPin(int32 WhichPar
 	{
 		UE_LOG(LogNiagaraEditor, Log, TEXT("FNiagaraParameterMapHistoryBuilder::RegisterParameterMapPin %d > %s .... %p"), WhichParameterMap, *Pin->GetOwningNode()->GetNodeTitle(ENodeTitleType::FullTitle).ToString(), Pin);
 	}
-	if (WhichParameterMap != INDEX_NONE)
+	if (Pin && WhichParameterMap != INDEX_NONE)
 	{
 		if (Pin->Direction == EEdGraphPinDirection::EGPD_Output)
 		{
