- Ensure data interface param info is set

#rnx

diff --git a/Engine/Plugins/FX/Niagara/Source/Niagara/Private/NiagaraScript.cpp b/Engine/Plugins/FX/Niagara/Source/Niagara/Private/NiagaraScript.cpp
--- a/Engine/Plugins/FX/Niagara/Source/Niagara/Private/NiagaraScript.cpp
+++ b/Engine/Plugins/FX/Niagara/Source/Niagara/Private/NiagaraScript.cpp
@@ -2985,6 +2985,8 @@ void UNiagaraScript::CacheResourceShadersForRendering(bool bRegenerateId, bool b
 				CachedScriptVMId.BaseScriptCompileHash, CachedScriptVMId.ReferencedCompileHashes,
 				CachedScriptVMId.bUsesRapidIterationParams, GetFriendlyName());
 
+			ScriptResource->SetDataInterfaceParamInfo(CachedScriptVM.DIParamInfo);
+
 			if (FNiagaraUtilities::SupportsComputeShaders(ShaderPlatform))
 			{
 				CacheShadersForResources(ScriptResource.Get(), true);
