- Avoid making a reference on loop

#rnx

diff --git a/Engine/Plugins/FX/Niagara/Source/Niagara/Private/NiagaraSystemInstance.cpp b/Engine/Plugins/FX/Niagara/Source/Niagara/Private/NiagaraSystemInstance.cpp
--- a/Engine/Plugins/FX/Niagara/Source/Niagara/Private/NiagaraSystemInstance.cpp
+++ b/Engine/Plugins/FX/Niagara/Source/Niagara/Private/NiagaraSystemInstance.cpp
@@ -1488,7 +1488,7 @@ void FNiagaraSystemInstance::InitDataInterfaces()
 	}
 
 	//Iterate over interfaces to get size for table and clear their interface bindings.
-	for (TSharedRef<FNiagaraEmitterInstance, ESPMode::ThreadSafe> Simulation : Emitters)
+	for (const TSharedRef<FNiagaraEmitterInstance, ESPMode::ThreadSafe>& Simulation : Emitters)
 	{
 		FNiagaraEmitterInstance& Sim = Simulation.Get();
 		if (Sim.IsDisabled())
