Various fixes for emitter summary view

diff --git a/Engine/Plugins/FX/Niagara/Source/NiagaraEditor/Private/NiagaraEmitterEditorData.cpp b/Engine/Plugins/FX/Niagara/Source/NiagaraEditor/Private/NiagaraEmitterEditorData.cpp
--- a/Engine/Plugins/FX/Niagara/Source/NiagaraEditor/Private/NiagaraEmitterEditorData.cpp
+++ b/Engine/Plugins/FX/Niagara/Source/NiagaraEditor/Private/NiagaraEmitterEditorData.cpp
@@ -65,10 +65,8 @@ FFunctionInputSummaryViewMetadata UNiagaraEmitterEditorData::GetSummaryViewMetaD
 }
 
 void UNiagaraEmitterEditorData::SetSummaryViewMetaData(const FFunctionInputSummaryViewKey& Key, const FFunctionInputSummaryViewMetadata& NewMetadata)
-{
-	FScopedTransaction ScopedTransaction(NSLOCTEXT("NiagaraEmitter", "EmitterModuleNodeMetaDataChanged", "MetaData for summary view node changed."));
+{	
 	static const FFunctionInputSummaryViewMetadata Empty;
-
 	if (NewMetadata == Empty)
 	{
 		SummaryViewFunctionInputMetadata.Remove(Key);
@@ -77,15 +75,18 @@ void UNiagaraEmitterEditorData::SetSummaryViewMetaData(const FFunctionInputSumma
 	{
 		SummaryViewFunctionInputMetadata.FindOrAdd(Key) = NewMetadata;
 	}
+	
 	OnPersistentDataChanged().Broadcast();
-	// GraphSource->MarkNotSynchronized(TEXT("MetaData for summary view node changed."));
 	OnSummaryViewStateChangedDelegate.Broadcast();
-	// UNiagaraSystem::RequestCompileForEmitter(this);
 }
 
 void UNiagaraEmitterEditorData::ToggleShowSummaryView()
 {
+	FScopedTransaction ScopedTransaction(NSLOCTEXT("NiagaraEmitter", "EmitterModuleShowSummaryChanged", "Emitter summary view enabled/disabled."));
+	Modify();
+	
 	bShowSummaryView = !bShowSummaryView;
+	
 	OnPersistentDataChanged().Broadcast();
 	OnSummaryViewStateChangedDelegate.Broadcast();
 }
diff --git a/Engine/Plugins/FX/Niagara/Source/NiagaraEditor/Private/ViewModels/Stack/NiagaraStackFunctionInput.cpp b/Engine/Plugins/FX/Niagara/Source/NiagaraEditor/Private/ViewModels/Stack/NiagaraStackFunctionInput.cpp
--- a/Engine/Plugins/FX/Niagara/Source/NiagaraEditor/Private/ViewModels/Stack/NiagaraStackFunctionInput.cpp
+++ b/Engine/Plugins/FX/Niagara/Source/NiagaraEditor/Private/ViewModels/Stack/NiagaraStackFunctionInput.cpp
@@ -2478,6 +2478,11 @@ bool UNiagaraStackFunctionInput::HasFrontDivider() const
 	return IsSemanticChild() || Super::HasFrontDivider();
 }
 
+TOptional<FNiagaraVariableMetaData> UNiagaraStackFunctionInput::GetMetadata() const
+{
+	return InputMetaData;
+}
+
 TOptional<FGuid> UNiagaraStackFunctionInput::GetMetadataGuid() const
 {
 	return InputMetaData.IsSet() ? InputMetaData->GetVariableGuid() : TOptional<FGuid>();
diff --git a/Engine/Plugins/FX/Niagara/Source/NiagaraEditor/Private/ViewModels/Stack/NiagaraStackFunctionInputCollection.cpp b/Engine/Plugins/FX/Niagara/Source/NiagaraEditor/Private/ViewModels/Stack/NiagaraStackFunctionInputCollection.cpp
--- a/Engine/Plugins/FX/Niagara/Source/NiagaraEditor/Private/ViewModels/Stack/NiagaraStackFunctionInputCollection.cpp
+++ b/Engine/Plugins/FX/Niagara/Source/NiagaraEditor/Private/ViewModels/Stack/NiagaraStackFunctionInputCollection.cpp
@@ -34,14 +34,32 @@ static FText GetUserFriendlyFunctionName(UNiagaraNodeFunctionCall* Node)
 	return FText::FromString(Node->GetFunctionName());
 }
 
-static bool ShouldShowInSummaryView(UNiagaraEmitter* Emitter, UNiagaraNodeFunctionCall* InputFunctionCallNode, TOptional<FNiagaraVariableMetaData>& VariableMetadata)
+static TOptional<FFunctionInputSummaryViewKey> GetSummaryViewInputKeyForFunctionInput(UNiagaraNodeFunctionCall* FunctionCall, const FNiagaraVariable& InputVariable, const TOptional<FNiagaraVariableMetaData>& InputMetaData)
 {
-	if (Emitter && VariableMetadata.IsSet())
+	FGuid NodeGuid = FunctionCall->NodeGuid;
+	
+	if (FunctionCall->IsA<UNiagaraNodeAssignment>())
+	{
+		return FFunctionInputSummaryViewKey(NodeGuid, InputVariable.GetName());
+	}
+	
+	if (InputMetaData.IsSet())
+	{
+		return FFunctionInputSummaryViewKey(NodeGuid, InputMetaData->GetVariableGuid());
+	}
+	
+	return TOptional<FFunctionInputSummaryViewKey>();
+}
+
+static bool ShouldShowInSummaryView(UNiagaraEmitter* Emitter, UNiagaraNodeFunctionCall* FunctionCall, const FNiagaraVariable& InputVariable, const TOptional<FNiagaraVariableMetaData>& InputMetaData)
+{
+	UNiagaraEmitterEditorData* EditorData = Emitter? Cast<UNiagaraEmitterEditorData>(Emitter->GetEditorData()) : nullptr;
+	if (EditorData && InputMetaData.IsSet())
 	{
-		UNiagaraEmitterEditorData* EditorData = Cast<UNiagaraEmitterEditorData>(Emitter->GetEditorData());
-		if (EditorData)
+		TOptional<FFunctionInputSummaryViewKey> SummaryViewKey = GetSummaryViewInputKeyForFunctionInput(FunctionCall, InputVariable, InputMetaData);
+		if (SummaryViewKey.IsSet())
 		{
-			return EditorData->GetSummaryViewMetaData(FFunctionInputSummaryViewKey(InputFunctionCallNode->NodeGuid, VariableMetadata->GetVariableGuid())).bVisible;
+			return EditorData->GetSummaryViewMetaData(SummaryViewKey.GetValue()).bVisible;
 		}
 	}
 	return false;
@@ -249,31 +267,30 @@ void UNiagaraStackFunctionInputCollectionBase::RefreshChildrenForFunctionCall(UN
 		int32 EditorSortPriority = InputMetaData.IsSet() ? InputMetaData->EditorSortPriority : 0;
 		TOptional<FText> DisplayName;
 
-		if (Emitter && InputMetaData.IsSet() && bShouldApplySummaryFilter)
+		UNiagaraEmitterEditorData* EditorData = Emitter? Cast<UNiagaraEmitterEditorData>(Emitter->GetEditorData()) : nullptr;
+
+		const auto& SummaryViewKey = GetSummaryViewInputKeyForFunctionInput(InputFunctionCallNode, InputVariable, InputMetaData);
+		if (EditorData && bShouldApplySummaryFilter && SummaryViewKey.IsSet())
 		{
-			UNiagaraEmitterEditorData* EditorData = Cast<UNiagaraEmitterEditorData>(Emitter->GetEditorData());
-			if (EditorData)
+			FFunctionInputSummaryViewMetadata SummaryViewData = EditorData->GetSummaryViewMetaData(SummaryViewKey.GetValue());
+			if (SummaryViewData.Category != NAME_None)
 			{
-				FFunctionInputSummaryViewMetadata SummaryViewData = EditorData->GetSummaryViewMetaData(FFunctionInputSummaryViewKey(FFunctionInputSummaryViewKey(InputFunctionCallNode->NodeGuid, InputMetaData->GetVariableGuid())));
-				if (SummaryViewData.Category != NAME_None)
-				{
-					InputCategory = FText::FromName(SummaryViewData.Category);
-				}
-				else
-				{
-					InputCategory = InputCategory.EqualTo(UncategorizedName)? FText::FromString(*InputFunctionCallNode->GetFunctionName()) : FText::FromString(*(InputFunctionCallNode->GetFunctionName() + TEXT(" - ") + InputCategory.ToString()));
-				}
-				EditorSortPriority = SummaryViewData.SortIndex;
-				DisplayName = (SummaryViewData.DisplayName != NAME_None) ? FText::FromName(SummaryViewData.DisplayName) : TOptional<FText>();
+				InputCategory = FText::FromName(SummaryViewData.Category);
 			}
+			else
+			{
+				InputCategory = InputCategory.EqualTo(UncategorizedName)? FText::FromString(*InputFunctionCallNode->GetFunctionName()) : FText::FromString(*(InputFunctionCallNode->GetFunctionName() + TEXT(" - ") + InputCategory.ToString()));
+			}
+			EditorSortPriority = SummaryViewData.SortIndex != INDEX_NONE? SummaryViewData.SortIndex : EditorSortPriority;
+			DisplayName = (SummaryViewData.DisplayName != NAME_None) ? FText::FromName(SummaryViewData.DisplayName) : TOptional<FText>();		
 		}
 
-		bool bShouldShowInSummary = ShouldShowInSummaryView(Emitter, InputFunctionCallNode, InputMetaData);
+		bool bShouldShowInSummary = ShouldShowInSummaryView(Emitter, InputFunctionCallNode, InputVariable, InputMetaData);
 		if (bShouldShowInSummary)
 		{
 			SummaryViewPins.Add(InputPin);
 		}		
-		bool bIsInputHidden = HiddenPins.Contains(InputPin) || (bShouldApplySummaryFilter && !bShouldShowInSummary);
+		bool bIsInputHidden = HiddenPins.Contains(InputPin);
 		FInputData InputData = { InputPin, InputVariable.GetType(), EditorSortPriority, DisplayName, InputCategory, false, bIsInputHidden, bShouldShowInSummary };
 		int32 Index = InputDataCollection.Add(InputData);
 
@@ -327,43 +344,37 @@ void UNiagaraStackFunctionInputCollectionBase::RefreshChildrenForFunctionCall(UN
 			InputMetaData = InputFunctionGraph->GetMetaData(InputVariable);
 		}
 
-		FText InputCategory = InputMetaData.IsSet() && InputMetaData->CategoryName.IsEmptyOrWhitespace() == false
+		FText InputCategory = (InputMetaData.IsSet() && InputMetaData->CategoryName.IsEmptyOrWhitespace() == false)
 			? InputMetaData->CategoryName
 			: UncategorizedName;
 
 		int32 EditorSortPriority = InputMetaData.IsSet() ? InputMetaData->EditorSortPriority : 0;
 		TOptional<FText> DisplayName;
+		
+		UNiagaraEmitterEditorData* EditorData = Emitter? Cast<UNiagaraEmitterEditorData>(Emitter->GetEditorData()) : nullptr;
 
-		if (Emitter && InputMetaData.IsSet() && bShouldApplySummaryFilter)
+		const auto& SummaryViewKey = GetSummaryViewInputKeyForFunctionInput(InputFunctionCallNode, InputVariable, InputMetaData);
+		if (EditorData && bShouldApplySummaryFilter && SummaryViewKey.IsSet())
 		{
-			UNiagaraEmitterEditorData* EditorData = Cast<UNiagaraEmitterEditorData>(Emitter->GetEditorData());
-			if (EditorData)
+			FFunctionInputSummaryViewMetadata SummaryViewData = EditorData->GetSummaryViewMetaData(SummaryViewKey.GetValue());
+			if (SummaryViewData.Category != NAME_None)
 			{
-				FFunctionInputSummaryViewMetadata SummaryViewData = EditorData->GetSummaryViewMetaData(FFunctionInputSummaryViewKey(FFunctionInputSummaryViewKey(InputFunctionCallNode->NodeGuid, InputMetaData->GetVariableGuid())));
-				if (SummaryViewData.Category != NAME_None)
-				{
-					InputCategory = FText::FromName(SummaryViewData.Category);
-				}
-				else
-				{
-					InputCategory = InputCategory.EqualTo(UncategorizedName)? FText::FromString(*InputFunctionCallNode->GetFunctionName()) : FText::FromString(*(InputFunctionCallNode->GetFunctionName() + TEXT(" - ") + InputCategory.ToString()));
-				}
-				EditorSortPriority = SummaryViewData.SortIndex;
-				DisplayName = (SummaryViewData.DisplayName != NAME_None) ? FText::FromName(SummaryViewData.DisplayName) : TOptional<FText>();
+				InputCategory = FText::FromName(SummaryViewData.Category);
 			}
-		}
-
-		if (DisplayName.IsSet())
-		{
-			check(true);
+			else
+			{
+				InputCategory = InputCategory.EqualTo(UncategorizedName)? FText::FromString(*InputFunctionCallNode->GetFunctionName()) : FText::FromString(*(InputFunctionCallNode->GetFunctionName() + TEXT(" - ") + InputCategory.ToString()));
+			}
+			EditorSortPriority = SummaryViewData.SortIndex != INDEX_NONE? SummaryViewData.SortIndex : EditorSortPriority;
+			DisplayName = (SummaryViewData.DisplayName != NAME_None) ? FText::FromName(SummaryViewData.DisplayName) : TOptional<FText>();		
 		}
 		
-		bool bShouldShowInSummary = ShouldShowInSummaryView(Emitter, InputFunctionCallNode, InputMetaData);
+		bool bShouldShowInSummary = ShouldShowInSummaryView(Emitter, InputFunctionCallNode, InputVariable, InputMetaData);
 		if (bShouldShowInSummary)
 		{
 			SummaryViewPins.Add(InputPin);
 		}		
-		bool bIsInputHidden = HiddenSwitchPins.Contains(InputPin) || (bShouldApplySummaryFilter && !bShouldShowInSummary);
+		bool bIsInputHidden = HiddenSwitchPins.Contains(InputPin);
 		FInputData InputData = { InputPin, InputVariable.GetType(), EditorSortPriority, DisplayName, InputCategory, true, bIsInputHidden, bShouldShowInSummary };
 		int32 Index = InputDataCollection.Add(InputData);
 
@@ -415,29 +426,28 @@ void UNiagaraStackFunctionInputCollectionBase::RefreshChildrenForFunctionCall(UN
 	}
 
 
-
-	static TFunction<bool(const FInputData*)> HasChildShownInSummary = [](const FInputData* Input) -> bool
+	if (bShouldApplySummaryFilter)
 	{
-		if (Input->bShouldShowInSummary)
-		{
-			return true;
-		}
-		for (const FInputData* Child : Input->Children)
-		{
-			if (HasChildShownInSummary(Child))
+		static TFunction<void(FInputData*)> PropagateSummaryViewShowToChildren = [](FInputData* Input)
+		{		
+			bool bShouldShowChildren = !Input->bIsHidden && Input->bShouldShowInSummary;
+
+			Input->bIsHidden = !bShouldShowChildren;
+
+			for (FInputData* Child : Input->Children)
 			{
-				return true;
-			}
-		}
-		return false;		
-	};
+				Child->bShouldShowInSummary = bShouldShowChildren;
+				PropagateSummaryViewShowToChildren(Child);
+			}		
+		};
 
-	// Propagate summary visibility up the parents
-	for (FInputData& InputData : InputDataCollection)
-	{		
-		if (HasChildShownInSummary(&InputData))
+		// Propagate summary visibility to children
+		for (FInputData& InputData : InputDataCollection)
 		{
-			InputData.bIsHidden = false;
+			if (!InputData.bIsChild)
+			{
+				PropagateSummaryViewShowToChildren(&InputData);
+			}
 		}
 	}
 
diff --git a/Engine/Plugins/FX/Niagara/Source/NiagaraEditor/Private/ViewModels/Stack/NiagaraStackRoot.cpp b/Engine/Plugins/FX/Niagara/Source/NiagaraEditor/Private/ViewModels/Stack/NiagaraStackRoot.cpp
--- a/Engine/Plugins/FX/Niagara/Source/NiagaraEditor/Private/ViewModels/Stack/NiagaraStackRoot.cpp
+++ b/Engine/Plugins/FX/Niagara/Source/NiagaraEditor/Private/ViewModels/Stack/NiagaraStackRoot.cpp
@@ -55,6 +55,15 @@ void UNiagaraStackRoot::Initialize(FRequiredEntryData InRequiredEntryData, bool
 	}	
 }
 
+void UNiagaraStackRoot::FinalizeInternal()
+{
+	if (bIncludeEmitterInformation && GetEmitterViewModel())
+	{
+		GetEmitterViewModel()->GetOrCreateEditorData().OnSummaryViewStateChanged().RemoveAll(this);
+	}	
+	Super::FinalizeInternal();
+}
+
 bool UNiagaraStackRoot::GetCanExpand() const
 {
 	return false;
diff --git a/Engine/Plugins/FX/Niagara/Source/NiagaraEditor/Private/ViewModels/Stack/NiagaraStackSummaryViewInputCollection.cpp b/Engine/Plugins/FX/Niagara/Source/NiagaraEditor/Private/ViewModels/Stack/NiagaraStackSummaryViewInputCollection.cpp
--- a/Engine/Plugins/FX/Niagara/Source/NiagaraEditor/Private/ViewModels/Stack/NiagaraStackSummaryViewInputCollection.cpp
+++ b/Engine/Plugins/FX/Niagara/Source/NiagaraEditor/Private/ViewModels/Stack/NiagaraStackSummaryViewInputCollection.cpp
@@ -42,6 +42,12 @@ void UNiagaraStackSummaryViewObject::Initialize(FRequiredEntryData InRequiredEnt
 	GetEmitterViewModel()->GetOrCreateEditorData().OnSummaryViewStateChanged().AddUObject(this, &UNiagaraStackSummaryViewObject::OnViewStateChanged);
 }
 
+void UNiagaraStackSummaryViewObject::FinalizeInternal()
+{
+	GetEmitterViewModel()->GetOrCreateEditorData().OnSummaryViewStateChanged().RemoveAll(this);
+	Super::FinalizeInternal();
+}
+
 FText UNiagaraStackSummaryViewObject::GetDisplayName() const
 {
 	return LOCTEXT("FilteredInputCollectionDisplayName", "Filtered Inputs");
@@ -84,7 +90,7 @@ void UNiagaraStackSummaryViewObject::RefreshChildrenInternal(const TArray<UNiaga
 
 	if (NewChildren.Num() == 0)
 	{
-		FText EmptyAssignmentNodeMessageText = LOCTEXT("EmptySummaryNodeMessage", "No Parameters in Summary View.\n\nTo add a parameter use the icon to the right of an available parameter in the stack. You can also set the DisplayName, SortIndex, and Category from the same dropdown, as well as here in the summary view after adding the parameters.");
+		FText EmptyAssignmentNodeMessageText = LOCTEXT("EmptySummaryNodeMessage", "No Parameters in Emitter Summary.\n\nTo add a parameter use the context menu for an available parameter in the stack. Once in the summary view you can use the drop down menu from the icon next to a parameter to set the DisplayName, SortIndex, and Category.");
 		UNiagaraStackItemTextContent* EmtpySummaryMessage = FindCurrentChildOfTypeByPredicate<UNiagaraStackItemTextContent>(CurrentChildren,
 			[&](UNiagaraStackItemTextContent* CurrentStackItemTextContent) { return CurrentStackItemTextContent->GetDisplayName().IdenticalTo(EmptyAssignmentNodeMessageText); });
 		
diff --git a/Engine/Plugins/FX/Niagara/Source/NiagaraEditor/Private/ViewModels/Stack/NiagaraStackViewModel.cpp b/Engine/Plugins/FX/Niagara/Source/NiagaraEditor/Private/ViewModels/Stack/NiagaraStackViewModel.cpp
--- a/Engine/Plugins/FX/Niagara/Source/NiagaraEditor/Private/ViewModels/Stack/NiagaraStackViewModel.cpp
+++ b/Engine/Plugins/FX/Niagara/Source/NiagaraEditor/Private/ViewModels/Stack/NiagaraStackViewModel.cpp
@@ -135,6 +135,13 @@ void UNiagaraStackViewModel::InitializeWithRootEntry(UNiagaraStackEntry* InRootE
 	bExternalRootEntry = true;
 
 	StructureChangedDelegate.Broadcast(ENiagaraStructureChangedFlags::StructureChanged);
+
+	TSharedPtr<FNiagaraEmitterHandleViewModel> EmitterHandleViewModelPinned = EmitterHandleViewModel.Pin();
+	if (EmitterHandleViewModelPinned.IsValid())
+	{
+		TSharedPtr<FNiagaraEmitterViewModel> EmitterViewModel = EmitterHandleViewModelPinned->GetEmitterViewModel();
+		EmitterViewModel->GetOrCreateEditorData().OnSummaryViewStateChanged().RemoveAll(this);
+	}
 }
 
 void UNiagaraStackViewModel::Reset()
diff --git a/Engine/Plugins/FX/Niagara/Source/NiagaraEditor/Public/NiagaraEditorCommon.h b/Engine/Plugins/FX/Niagara/Source/NiagaraEditor/Public/NiagaraEditorCommon.h
--- a/Engine/Plugins/FX/Niagara/Source/NiagaraEditor/Public/NiagaraEditorCommon.h
+++ b/Engine/Plugins/FX/Niagara/Source/NiagaraEditor/Public/NiagaraEditorCommon.h
@@ -197,22 +197,36 @@ private:
 	UPROPERTY()
 	FGuid InputGuid;
 
+	UPROPERTY()
+	FName InputName;
 public:
 
 	FFunctionInputSummaryViewKey() { }
 
 	FFunctionInputSummaryViewKey(const FGuid& InFunctionGuid, const FGuid& InInputGuid)
 		: FunctionGuid(InFunctionGuid), InputGuid(InInputGuid)
-	{ }
+	{
+		check(InFunctionGuid.IsValid());
+		check(InInputGuid.IsValid());		
+	}
+
+	FFunctionInputSummaryViewKey(const FGuid& InFunctionGuid, const FName& InInputName)
+		: FunctionGuid(InFunctionGuid), InputName(InInputName)
+	{
+		check(InFunctionGuid.IsValid());
+		check(!InputName.IsNone());		
+	}
 
 	bool operator==(const FFunctionInputSummaryViewKey& Other) const
 	{
-		return FunctionGuid == Other.FunctionGuid && InputGuid == Other.InputGuid;
+		return FunctionGuid == Other.FunctionGuid && InputGuid == Other.InputGuid && InputName == Other.InputName;
 	}
 
 	friend uint32 GetTypeHash(const FFunctionInputSummaryViewKey& Key)
 	{
-		return HashCombine(GetTypeHash(Key.FunctionGuid), GetTypeHash(Key.InputGuid));
+		return HashCombine(
+			GetTypeHash(Key.FunctionGuid),
+			Key.InputGuid.IsValid()? GetTypeHash(Key.InputGuid) : GetTypeHash(Key.InputName));
 	}
 };
 
@@ -233,7 +247,7 @@ public:
 	
 	FFunctionInputSummaryViewMetadata()
 		: bVisible(false)
-		, SortIndex(0)
+		, SortIndex(INDEX_NONE)
 	{		
 	}
 
diff --git a/Engine/Plugins/FX/Niagara/Source/NiagaraEditor/Public/ViewModels/Stack/NiagaraStackEmitterSettingsGroup.h b/Engine/Plugins/FX/Niagara/Source/NiagaraEditor/Public/ViewModels/Stack/NiagaraStackEmitterSettingsGroup.h
--- a/Engine/Plugins/FX/Niagara/Source/NiagaraEditor/Public/ViewModels/Stack/NiagaraStackEmitterSettingsGroup.h
+++ b/Engine/Plugins/FX/Niagara/Source/NiagaraEditor/Public/ViewModels/Stack/NiagaraStackEmitterSettingsGroup.h
@@ -134,23 +134,4 @@ public:
 protected:
 	virtual void RefreshChildrenInternal(const TArray<UNiagaraStackEntry*>& CurrentChildren, TArray<UNiagaraStackEntry*>& NewChildren, TArray<FStackIssue>& NewIssues) override;
 
-
-
-private:
-	// UPROPERTY()
-	// UNiagaraStackItemGroupFooter* GroupFooter;
-
-};
-
-// UCLASS()
-// class NIAGARAEDITOR_API UNiagaraStackItemGroupFooter : public UNiagaraStackEntry
-// {
-// 	GENERATED_BODY()
-//
-// public:
-// 	void Initialize(FRequiredEntryData InRequiredEntryData);
-//
-// 	virtual EStackRowStyle GetStackRowStyle() const override;
-//
-// 	virtual bool GetCanExpand() const;
-// };
\ No newline at end of file
+};
\ No newline at end of file
diff --git a/Engine/Plugins/FX/Niagara/Source/NiagaraEditor/Public/ViewModels/Stack/NiagaraStackFunctionInput.h b/Engine/Plugins/FX/Niagara/Source/NiagaraEditor/Public/ViewModels/Stack/NiagaraStackFunctionInput.h
--- a/Engine/Plugins/FX/Niagara/Source/NiagaraEditor/Public/ViewModels/Stack/NiagaraStackFunctionInput.h
+++ b/Engine/Plugins/FX/Niagara/Source/NiagaraEditor/Public/ViewModels/Stack/NiagaraStackFunctionInput.h
@@ -276,6 +276,7 @@ public:
 	virtual void GetSearchItems(TArray<FStackSearchItem>& SearchItems) const override;
 	virtual bool HasFrontDivider() const override;
 
+	TOptional<FNiagaraVariableMetaData> GetMetadata() const;
 	TOptional<FGuid> GetMetadataGuid() const;
 
 protected:
diff --git a/Engine/Plugins/FX/Niagara/Source/NiagaraEditor/Public/ViewModels/Stack/NiagaraStackRoot.h b/Engine/Plugins/FX/Niagara/Source/NiagaraEditor/Public/ViewModels/Stack/NiagaraStackRoot.h
--- a/Engine/Plugins/FX/Niagara/Source/NiagaraEditor/Public/ViewModels/Stack/NiagaraStackRoot.h
+++ b/Engine/Plugins/FX/Niagara/Source/NiagaraEditor/Public/ViewModels/Stack/NiagaraStackRoot.h
@@ -24,6 +24,7 @@ public:
 	UNiagaraStackRoot();
 	
 	void Initialize(FRequiredEntryData InRequiredEntryData, bool bInIncludeSystemInformation, bool bInIncludeEmitterInformation);
+	void FinalizeInternal() override;
 
 	virtual bool GetCanExpand() const override;
 	virtual bool GetShouldShowInStack() const override;
diff --git a/Engine/Plugins/FX/Niagara/Source/NiagaraEditor/Public/ViewModels/Stack/NiagaraStackSummaryViewInputCollection.h b/Engine/Plugins/FX/Niagara/Source/NiagaraEditor/Public/ViewModels/Stack/NiagaraStackSummaryViewInputCollection.h
--- a/Engine/Plugins/FX/Niagara/Source/NiagaraEditor/Public/ViewModels/Stack/NiagaraStackSummaryViewInputCollection.h
+++ b/Engine/Plugins/FX/Niagara/Source/NiagaraEditor/Public/ViewModels/Stack/NiagaraStackSummaryViewInputCollection.h
@@ -26,11 +26,12 @@ public:
 	UNiagaraStackSummaryViewObject();
 
 	void Initialize(FRequiredEntryData InRequiredEntryData, UNiagaraEmitter* InEmitter, FString InOwningStackItemEditorDataKey);
-
+	virtual void FinalizeInternal() override;
 
 	virtual FText GetDisplayName() const override;
 	virtual bool GetShouldShowInStack() const override;
 	virtual bool GetIsEnabled() const;
+	bool GetCanExpand() const override { return false; }
 protected:
 
 	virtual void RefreshChildrenInternal(const TArray<UNiagaraStackEntry*>& CurrentChildren, TArray<UNiagaraStackEntry*>& NewChildren, TArray<FStackIssue>& NewIssues) override;
diff --git a/Engine/Plugins/FX/Niagara/Source/NiagaraEditorWidgets/Private/NiagaraEditorWidgetsUtilities.cpp b/Engine/Plugins/FX/Niagara/Source/NiagaraEditorWidgets/Private/NiagaraEditorWidgetsUtilities.cpp
--- a/Engine/Plugins/FX/Niagara/Source/NiagaraEditorWidgets/Private/NiagaraEditorWidgetsUtilities.cpp
+++ b/Engine/Plugins/FX/Niagara/Source/NiagaraEditorWidgets/Private/NiagaraEditorWidgetsUtilities.cpp
@@ -6,6 +6,7 @@
 #include "ViewModels/Stack/NiagaraStackModuleItem.h"
 #include "ViewModels/NiagaraSystemViewModel.h"
 #include "Stack/SNiagaraStackItemGroupAddMenu.h"
+#include "NiagaraNodeAssignment.h"
 #include "NiagaraEditorWidgetsStyle.h"
 #include "NiagaraEditorCommon.h"
 #include "NiagaraClipboard.h"
@@ -20,9 +21,13 @@
 #include "Subsystems/AssetEditorSubsystem.h"
 #include "Editor.h"
 #include "EditorFontGlyphs.h"
+#include "NiagaraEmitter.h"
+#include "NiagaraEmitterEditorData.h"
 #include "NiagaraMessages.h"
+#include "ViewModels/NiagaraEmitterViewModel.h"
 #include "ViewModels/NiagaraSystemSelectionViewModel.h"
 #include "ViewModels/NiagaraSystemViewModel.h"
+#include "ViewModels/Stack/NiagaraStackInputCategory.h"
 
 #define LOCTEXT_NAMESPACE "NiagaraStackEditorWidgetsUtilities"
 
@@ -511,30 +516,55 @@ FString FNiagaraStackEditorWidgetsUtilities::StackEntryToStringForListDebug(UNia
 	return FString::Printf(TEXT("0x%08x - %s - %s"), StackEntry, *StackEntry->GetClass()->GetName(), *StackEntry->GetDisplayName().ToString());
 }
 
-UNiagaraStackFunctionInput* FNiagaraStackEditorWidgetsUtilities::FindTopMostParentFunctionInput(UNiagaraStackFunctionInput* FunctionInput)
+TOptional<FFunctionInputSummaryViewKey> FNiagaraStackEditorWidgetsUtilities::GetSummaryViewInputKeyForFunctionInput(UNiagaraStackFunctionInput* FunctionInput)
 {
-	UNiagaraStackFunctionInput* CurrentTopMost = FunctionInput;
-	UObject* CurrentParent = FunctionInput->GetOuter();
-	while (CurrentParent)
+	const FGuid NodeGuid = FunctionInput->GetInputFunctionCallNode().NodeGuid;
+
+	const UNiagaraNodeFunctionCall& InputCallNode = FunctionInput->GetInputFunctionCallNode();
+	if (InputCallNode.IsA<UNiagaraNodeAssignment>())
+	{		
+		return FFunctionInputSummaryViewKey(NodeGuid, FunctionInput->GetInputParameterHandle().GetParameterHandleString());
+	}
+	
+	const TOptional<FGuid> VariableGuid = FunctionInput->GetMetadataGuid();
+	if (VariableGuid.IsSet() && VariableGuid.GetValue().IsValid())
 	{
-		if (CurrentParent->IsA<UNiagaraStackFunctionInput>())
-		{
-			CurrentTopMost = CastChecked<UNiagaraStackFunctionInput>(CurrentParent);
-		}
-		
-		CurrentParent = CurrentParent->GetOuter();
+		return FFunctionInputSummaryViewKey(NodeGuid, VariableGuid.GetValue());
 	}
-
-	return CurrentTopMost;
+	
+	return TOptional<FFunctionInputSummaryViewKey>();
 }
 
-TOptional<FFunctionInputSummaryViewKey> FNiagaraStackEditorWidgetsUtilities::GetSummaryViewInputKeyForFunctionInput(UNiagaraStackFunctionInput* FunctionInput)
+UNiagaraStackFunctionInput* FNiagaraStackEditorWidgetsUtilities::GetParentInputForSummaryView(UNiagaraStackFunctionInput* FunctionInput)
 {
-	FGuid NodeGuid = FunctionInput->GetInputFunctionCallNode().NodeGuid;
-	TOptional<FGuid> VariableGuid = FunctionInput->GetMetadataGuid();
-	return VariableGuid.IsSet()
-		       ? FFunctionInputSummaryViewKey(NodeGuid, VariableGuid.GetValue())
-		       : TOptional<FFunctionInputSummaryViewKey>();
+	UNiagaraEmitter* Emitter = (FunctionInput && FunctionInput->GetEmitterViewModel())? FunctionInput->GetEmitterViewModel()->GetEmitter() : nullptr;
+	UNiagaraEmitterEditorData* EditorData = Emitter? Cast<UNiagaraEmitterEditorData>(Emitter->GetEditorData()) : nullptr;
+	
+	if (EditorData)
+	{		
+		TOptional<FFunctionInputSummaryViewKey> Key = FNiagaraStackEditorWidgetsUtilities::GetSummaryViewInputKeyForFunctionInput(FunctionInput);
+
+		TOptional<FNiagaraVariableMetaData> Metadata = FunctionInput->GetMetadata();
+		if (Metadata.IsSet() && !Metadata->ParentAttribute.IsNone())
+		{
+			UNiagaraStackInputCategory* Category = CastChecked<UNiagaraStackInputCategory>(FunctionInput->GetOuter());
+
+			TArray<UNiagaraStackEntry*> Children;
+			Category->GetFilteredChildrenOfTypes(Children, { UNiagaraStackFunctionInput::StaticClass() });
+
+			for (UNiagaraStackEntry* Entry : Children)
+			{
+				UNiagaraStackFunctionInput* Input = CastChecked<UNiagaraStackFunctionInput>(Entry);
+
+				if (Input->GetInputParameterHandle().GetName() == Metadata->ParentAttribute)
+				{
+					return Input;
+				}				
+			}
+		}
+	}
+
+	return FunctionInput;	
 }
 
 #undef LOCTEXT_NAMESPACE
diff --git a/Engine/Plugins/FX/Niagara/Source/NiagaraEditorWidgets/Private/NiagaraEditorWidgetsUtilities.h b/Engine/Plugins/FX/Niagara/Source/NiagaraEditorWidgets/Private/NiagaraEditorWidgetsUtilities.h
--- a/Engine/Plugins/FX/Niagara/Source/NiagaraEditorWidgets/Private/NiagaraEditorWidgetsUtilities.h
+++ b/Engine/Plugins/FX/Niagara/Source/NiagaraEditorWidgets/Private/NiagaraEditorWidgetsUtilities.h
@@ -43,8 +43,8 @@ namespace FNiagaraStackEditorWidgetsUtilities
 	bool HandleDropForStackEntry(const FDragDropEvent& InDragDropEvent, EItemDropZone InDropZone, UNiagaraStackEntry* InTargetEntry, UNiagaraStackEntry::EDropOptions DropOptions);
 
 	FString StackEntryToStringForListDebug(UNiagaraStackEntry* StackEntry);
-
-	UNiagaraStackFunctionInput* FindTopMostParentFunctionInput(UNiagaraStackFunctionInput* FunctionInput);
 	
 	TOptional<FFunctionInputSummaryViewKey> GetSummaryViewInputKeyForFunctionInput(UNiagaraStackFunctionInput* FunctionInput);
+
+	UNiagaraStackFunctionInput* GetParentInputForSummaryView(UNiagaraStackFunctionInput* FunctionInput);
 }
diff --git a/Engine/Plugins/FX/Niagara/Source/NiagaraEditorWidgets/Private/SNiagaraOverviewStackNode.cpp b/Engine/Plugins/FX/Niagara/Source/NiagaraEditorWidgets/Private/SNiagaraOverviewStackNode.cpp
--- a/Engine/Plugins/FX/Niagara/Source/NiagaraEditorWidgets/Private/SNiagaraOverviewStackNode.cpp
+++ b/Engine/Plugins/FX/Niagara/Source/NiagaraEditorWidgets/Private/SNiagaraOverviewStackNode.cpp
@@ -308,36 +308,18 @@ void SNiagaraOverviewStackNode::CreateBottomSummaryExpander()
 	.VAlign(VAlign_Fill)
 	[
 		SNew(SButton)
-		.ForegroundColor(FSlateColor::UseForeground())
-		.ButtonStyle(FCoreStyle::Get(), "NoBorder")
-		.HAlign(HAlign_Fill)
-		.OnClicked(this, &SNiagaraOverviewStackNode::OnSummaryExpanderClicked)
+		.ButtonStyle(FEditorStyle::Get(), "NoBorder")
+		.HAlign(HAlign_Center)
+		.ContentPadding(2)
+		.ToolTipText(this, &SNiagaraOverviewStackNode::GetSummaryViewCollapseTooltipText)
+		.OnClicked(this, &SNiagaraOverviewStackNode::ExpandSummaryViewClicked)
+		.IsFocusable(false)
+		.Content()
 		[
-			SNew(SHorizontalBox)
-			+SHorizontalBox::Slot()
-			.AutoWidth()
-			.HAlign(HAlign_Left)
-			[
-				// LEFT
-				SNullWidget::NullWidget
-			]
-			+SHorizontalBox::Slot()
-			.HAlign(HAlign_Center)
-			[
-				SNew(STextBlock)
-				//.Font(FEditorStyle::Get().GetFontStyle("FontAwesome.10"))
-				.TextStyle(FNiagaraEditorWidgetsStyle::Get(), "NiagaraEditor.SystemOverview.GroupHeaderText")
-				.Text(this, &SNiagaraOverviewStackNode::GetSummaryExpanderButtonText)
-				
-			]
-			+SHorizontalBox::Slot()
-			.AutoWidth()
-			.HAlign(HAlign_Right)
-			[
-				// RIGHT
-				SNullWidget::NullWidget
-			]
-		]
+			// add the dropdown button for advanced properties 
+			SNew(SImage)
+				.Image(this, &SNiagaraOverviewStackNode::GetSummaryViewButtonBrush)
+		]		
 	];
 }
 
@@ -626,17 +608,31 @@ EVisibility SNiagaraOverviewStackNode::GetOpenParentEmitterVisibility() const
 		: EVisibility::Collapsed;
 }
 
-FText SNiagaraOverviewStackNode::GetSummaryExpanderButtonText() const
+
+const FSlateBrush* SNiagaraOverviewStackNode::GetSummaryViewButtonBrush() const
 {
 	UNiagaraEmitterEditorData* EditorData = EmitterHandleViewModelWeak.IsValid()? &EmitterHandleViewModelWeak.Pin()->GetEmitterViewModel()->GetOrCreateEditorData() : nullptr;
-	if (EditorData)
+	if (BottomSummaryExpander->IsHovered())
 	{
-		return EditorData->ShouldShowSummaryView()? LOCTEXT("SummaryViewExpanderButton_Expand", "Show Emitter Avanced") : LOCTEXT("SummaryViewExpanderButton_Hide", "Show Emitter Summary");
+		return EditorData && EditorData->ShouldShowSummaryView()
+			? FEditorStyle::GetBrush("DetailsView.PulldownArrow.Down.Hovered")
+			: FEditorStyle::GetBrush("DetailsView.PulldownArrow.Up.Hovered");
 	}
-	return LOCTEXT("SummaryViewExpanderButton_Error", "Invalid Emitter");
+	else
+	{
+		return EditorData && EditorData->ShouldShowSummaryView()
+			? FEditorStyle::GetBrush("DetailsView.PulldownArrow.Down")
+			: FEditorStyle::GetBrush("DetailsView.PulldownArrow.Up");
+	}	
+}
+
+FText SNiagaraOverviewStackNode::GetSummaryViewCollapseTooltipText() const
+{
+	UNiagaraEmitterEditorData* EditorData = EmitterHandleViewModelWeak.IsValid()? &EmitterHandleViewModelWeak.Pin()->GetEmitterViewModel()->GetOrCreateEditorData() : nullptr;
+	return EditorData && EditorData->ShouldShowSummaryView()? LOCTEXT("HideAdvancedToolTip", "Show Full Emitter") : LOCTEXT("ShowAdvancedToolTip", "Show Emitter Summary");
 }
 
-FReply SNiagaraOverviewStackNode::OnSummaryExpanderClicked()
+FReply SNiagaraOverviewStackNode::ExpandSummaryViewClicked()
 {
 	UNiagaraEmitterEditorData* EditorData = EmitterHandleViewModelWeak.IsValid()? &EmitterHandleViewModelWeak.Pin()->GetEmitterViewModel()->GetOrCreateEditorData() : nullptr;
 	if (EditorData)
diff --git a/Engine/Plugins/FX/Niagara/Source/NiagaraEditorWidgets/Private/SNiagaraOverviewStackNode.h b/Engine/Plugins/FX/Niagara/Source/NiagaraEditorWidgets/Private/SNiagaraOverviewStackNode.h
--- a/Engine/Plugins/FX/Niagara/Source/NiagaraEditorWidgets/Private/SNiagaraOverviewStackNode.h
+++ b/Engine/Plugins/FX/Niagara/Source/NiagaraEditorWidgets/Private/SNiagaraOverviewStackNode.h
@@ -62,8 +62,9 @@ private:
 	FReply OpenParentEmitter();
 	EVisibility GetOpenParentEmitterVisibility() const;
 
-	FText GetSummaryExpanderButtonText() const;
-	FReply OnSummaryExpanderClicked();
+	const FSlateBrush* GetSummaryViewButtonBrush() const;
+	FText GetSummaryViewCollapseTooltipText() const;
+	FReply ExpandSummaryViewClicked();
 private:
 	UNiagaraOverviewNode* OverviewStackNode;
 	UNiagaraStackViewModel* StackViewModel;
diff --git a/Engine/Plugins/FX/Niagara/Source/NiagaraEditorWidgets/Private/Stack/SNiagaraStackFunctionInputValue.cpp b/Engine/Plugins/FX/Niagara/Source/NiagaraEditorWidgets/Private/Stack/SNiagaraStackFunctionInputValue.cpp
--- a/Engine/Plugins/FX/Niagara/Source/NiagaraEditorWidgets/Private/Stack/SNiagaraStackFunctionInputValue.cpp
+++ b/Engine/Plugins/FX/Niagara/Source/NiagaraEditorWidgets/Private/Stack/SNiagaraStackFunctionInputValue.cpp
@@ -1,9 +1,6 @@
 // Copyright Epic Games, Inc. All Rights Reserved.
 
 #include "Stack/SNiagaraStackFunctionInputValue.h"
-
-#include <ViewModels/Stack/NiagaraStackSummaryViewInputCollection.h>
-
 #include "AssetRegistryModule.h"
 #include "EditorFontGlyphs.h"
 #include "INiagaraEditorTypeUtilities.h"
@@ -37,6 +34,7 @@
 #include "ViewModels/Stack/NiagaraStackFunctionInputCollection.h"
 #include "ViewModels/Stack/NiagaraStackGraphUtilities.h"
 #include "ViewModels/Stack/NiagaraStackInputCategory.h"
+#include "ViewModels/Stack/NiagaraStackSummaryViewInputCollection.h"
 #include "Widgets/SBoxPanel.h"
 #include "Widgets/Text/SRichTextBlock.h"
 #include "Widgets/SNiagaraParameterName.h"
@@ -50,7 +48,10 @@
 #include "Widgets/Layout/SSpacer.h"
 #include "Widgets/Layout/SSeparator.h"
 #include "Widgets/Input/SSpinBox.h"
+#include "Widgets/Input/SNumericEntryBox.h"
+#include "NiagaraNodeAssignment.h"
 #include "NiagaraCommon.h"
+
 #define LOCTEXT_NAMESPACE "NiagaraStackFunctionInputValue"
 
 const float TextIconSize = 16;
@@ -67,121 +68,125 @@ void SNiagaraStackFunctionInputValue::Construct(const FArguments& InArgs, UNiaga
 	
 	ChildSlot
 	[
-		SNew(SDropTarget)
-		.OnAllowDrop(this, &SNiagaraStackFunctionInputValue::OnFunctionInputAllowDrop)
-		.OnDropped(this, &SNiagaraStackFunctionInputValue::OnFunctionInputDrop)
-		.HorizontalImage(FNiagaraEditorWidgetsStyle::Get().GetBrush("NiagaraEditor.Stack.DropTarget.BorderHorizontal"))
-		.VerticalImage(FNiagaraEditorWidgetsStyle::Get().GetBrush("NiagaraEditor.Stack.DropTarget.BorderVertical"))
-		.IsEnabled_UObject(FunctionInput, &UNiagaraStackEntry::GetOwnerIsEnabled)
-		.Content()
+		SAssignNew(ChildrenBox, SHorizontalBox)
+		+ SHorizontalBox::Slot()
 		[
-			// Values
-			SAssignNew(ChildrenBox, SHorizontalBox)
-			.IsEnabled(this, &SNiagaraStackFunctionInputValue::GetInputEnabled)
-			+ SHorizontalBox::Slot()
-			.VAlign(VAlign_Center)
-			.AutoWidth()
-			.Padding(0, 0, 3, 0)
+			SNew(SDropTarget)
+			.OnAllowDrop(this, &SNiagaraStackFunctionInputValue::OnFunctionInputAllowDrop)
+			.OnDropped(this, &SNiagaraStackFunctionInputValue::OnFunctionInputDrop)
+			.HorizontalImage(FNiagaraEditorWidgetsStyle::Get().GetBrush("NiagaraEditor.Stack.DropTarget.BorderHorizontal"))
+			.VerticalImage(FNiagaraEditorWidgetsStyle::Get().GetBrush("NiagaraEditor.Stack.DropTarget.BorderVertical"))
+			.IsEnabled_UObject(FunctionInput, &UNiagaraStackEntry::GetOwnerIsEnabled)
+			.Content()
 			[
-				// Value Icon
-				SNew(SBox)
-				.WidthOverride(TextIconSize)
+				// Values
+				SNew(SHorizontalBox)
+				.IsEnabled(this, &SNiagaraStackFunctionInputValue::GetInputEnabled)
+				+ SHorizontalBox::Slot()
 				.VAlign(VAlign_Center)
-				.Visibility(this, &SNiagaraStackFunctionInputValue::GetInputIconVisibility)
+				.AutoWidth()
+				.Padding(0, 0, 3, 0)
 				[
-					SNew(STextBlock)
-					.Font(FEditorStyle::Get().GetFontStyle("FontAwesome.10"))
-					.Text(this, &SNiagaraStackFunctionInputValue::GetInputIconText)
-					.ToolTipText(this, &SNiagaraStackFunctionInputValue::GetInputIconToolTip)
-					.ColorAndOpacity(this, &SNiagaraStackFunctionInputValue::GetInputIconColor)
+					// Value Icon
+					SNew(SBox)
+					.WidthOverride(TextIconSize)
+					.VAlign(VAlign_Center)
+					.Visibility(this, &SNiagaraStackFunctionInputValue::GetInputIconVisibility)
+					[
+						SNew(STextBlock)
+						.Font(FEditorStyle::Get().GetFontStyle("FontAwesome.10"))
+						.Text(this, &SNiagaraStackFunctionInputValue::GetInputIconText)
+						.ToolTipText(this, &SNiagaraStackFunctionInputValue::GetInputIconToolTip)
+						.ColorAndOpacity(this, &SNiagaraStackFunctionInputValue::GetInputIconColor)
+					]
 				]
-			]
-			+ SHorizontalBox::Slot()
-			.VAlign(VAlign_Center)
-			.AutoWidth()
-			.Padding(0, 0, 3, 0)
-			[
-				// Type Modifier Icon
-				SNew(SBox)
-				.WidthOverride(TextIconSize)
+				+ SHorizontalBox::Slot()
 				.VAlign(VAlign_Center)
-				.Visibility(this, &SNiagaraStackFunctionInputValue::GetTypeModifierIconVisibility)
+				.AutoWidth()
+				.Padding(0, 0, 3, 0)
 				[
-					SNew(STextBlock)
-					.Text(this, &SNiagaraStackFunctionInputValue::GetTypeModifierIconText)
-					.ToolTipText(this, &SNiagaraStackFunctionInputValue::GetTypeModifierIconToolTip)
-					.ColorAndOpacity(this, &SNiagaraStackFunctionInputValue::GetTypeModifierIconColor)
+					// Type Modifier Icon
+					SNew(SBox)
+					.WidthOverride(TextIconSize)
+					.VAlign(VAlign_Center)
+					.Visibility(this, &SNiagaraStackFunctionInputValue::GetTypeModifierIconVisibility)
+					[
+						SNew(STextBlock)
+						.Text(this, &SNiagaraStackFunctionInputValue::GetTypeModifierIconText)
+						.ToolTipText(this, &SNiagaraStackFunctionInputValue::GetTypeModifierIconToolTip)
+						.ColorAndOpacity(this, &SNiagaraStackFunctionInputValue::GetTypeModifierIconColor)
+					]
 				]
-			]
-			+ SHorizontalBox::Slot()
-			.VAlign(VAlign_Center)
-			[
-				// Value container and widgets.
-				SAssignNew(ValueContainer, SBox)
-				.ToolTipText_UObject(FunctionInput, &UNiagaraStackFunctionInput::GetValueToolTip)
+				+ SHorizontalBox::Slot()
+				.VAlign(VAlign_Center)
 				[
-					ConstructValueWidgets()
+					// Value container and widgets.
+					SAssignNew(ValueContainer, SBox)
+					.ToolTipText_UObject(FunctionInput, &UNiagaraStackFunctionInput::GetValueToolTip)
+					[
+						ConstructValueWidgets()
+					]
 				]
-			]
 
-			// Handle drop-down button
-			+ SHorizontalBox::Slot()
-			.AutoWidth()
-			.VAlign(VAlign_Center)
-			.Padding(3, 0, 0, 0)
-			[
-				SAssignNew(SetFunctionInputButton, SComboButton)
-				.ButtonStyle(FEditorStyle::Get(), "HoverHintOnly")
-				.IsFocusable(false)
-				.ForegroundColor(FSlateColor::UseForeground())
-				.OnGetMenuContent(this, &SNiagaraStackFunctionInputValue::OnGetAvailableHandleMenu)
-				.ContentPadding(FMargin(2))
-				.Visibility(this, &SNiagaraStackFunctionInputValue::GetDropdownButtonVisibility)
-				.MenuPlacement(MenuPlacement_BelowRightAnchor)
-				.HAlign(HAlign_Center)
+				// Handle drop-down button
+				+ SHorizontalBox::Slot()
+				.AutoWidth()
 				.VAlign(VAlign_Center)
-			]
-
-			// Reset Button
-			+ SHorizontalBox::Slot()
-			.AutoWidth()
-			.VAlign(VAlign_Center)
-			.Padding(3, 0, 0, 0)
-			[
-				SNew(SButton)
-				.IsFocusable(false)
-				.ToolTipText(LOCTEXT("ResetToolTip", "Reset to the default value"))
-				.ButtonStyle(FEditorStyle::Get(), "NoBorder")
-				.ContentPadding(0)
-				.Visibility(this, &SNiagaraStackFunctionInputValue::GetResetButtonVisibility)
-				.OnClicked(this, &SNiagaraStackFunctionInputValue::ResetButtonPressed)
-				.Content()
+				.Padding(3, 0, 0, 0)
 				[
-					SNew(SImage)
-					.Image(FEditorStyle::GetBrush("PropertyWindow.DiffersFromDefault"))
+					SAssignNew(SetFunctionInputButton, SComboButton)
+					.ButtonStyle(FEditorStyle::Get(), "HoverHintOnly")
+					.IsFocusable(false)
+					.ForegroundColor(FSlateColor::UseForeground())
+					.OnGetMenuContent(this, &SNiagaraStackFunctionInputValue::OnGetAvailableHandleMenu)
+					.ContentPadding(FMargin(2))
+					.Visibility(this, &SNiagaraStackFunctionInputValue::GetDropdownButtonVisibility)
+					.MenuPlacement(MenuPlacement_BelowRightAnchor)
+					.HAlign(HAlign_Center)
+					.VAlign(VAlign_Center)
 				]
-			]
 
-			// Reset to base Button
-			+ SHorizontalBox::Slot()
-			.AutoWidth()
-			.VAlign(VAlign_Center)
-			.Padding(3, 0, 0, 0)
-			[
-				SNew(SButton)
-				.IsFocusable(false)
-				.ToolTipText(LOCTEXT("ResetToBaseToolTip", "Reset this input to the value defined by the parent emitter"))
-				.ButtonStyle(FEditorStyle::Get(), "NoBorder")
-				.ContentPadding(0)
-				.Visibility(this, &SNiagaraStackFunctionInputValue::GetResetToBaseButtonVisibility)
-				.OnClicked(this, &SNiagaraStackFunctionInputValue::ResetToBaseButtonPressed)
-				.Content()
+				// Reset Button
+				+ SHorizontalBox::Slot()
+				.AutoWidth()
+				.VAlign(VAlign_Center)
+				.Padding(3, 0, 0, 0)
 				[
-					SNew(SImage)
-					.Image(FEditorStyle::GetBrush("PropertyWindow.DiffersFromDefault"))
-					.ColorAndOpacity(FSlateColor(FLinearColor::Green))
+					SNew(SButton)
+					.IsFocusable(false)
+					.ToolTipText(LOCTEXT("ResetToolTip", "Reset to the default value"))
+					.ButtonStyle(FEditorStyle::Get(), "NoBorder")
+					.ContentPadding(0)
+					.Visibility(this, &SNiagaraStackFunctionInputValue::GetResetButtonVisibility)
+					.OnClicked(this, &SNiagaraStackFunctionInputValue::ResetButtonPressed)
+					.Content()
+					[
+						SNew(SImage)
+						.Image(FEditorStyle::GetBrush("PropertyWindow.DiffersFromDefault"))
+					]
 				]
-			]			
+
+				// Reset to base Button
+				+ SHorizontalBox::Slot()
+				.AutoWidth()
+				.VAlign(VAlign_Center)
+				.Padding(3, 0, 0, 0)
+				[
+					SNew(SButton)
+					.IsFocusable(false)
+					.ToolTipText(LOCTEXT("ResetToBaseToolTip", "Reset this input to the value defined by the parent emitter"))
+					.ButtonStyle(FEditorStyle::Get(), "NoBorder")
+					.ContentPadding(0)
+					.Visibility(this, &SNiagaraStackFunctionInputValue::GetResetToBaseButtonVisibility)
+					.OnClicked(this, &SNiagaraStackFunctionInputValue::ResetToBaseButtonPressed)
+					.Content()
+					[
+						SNew(SImage)
+						.Image(FEditorStyle::GetBrush("PropertyWindow.DiffersFromDefault"))
+						.ColorAndOpacity(FSlateColor(FLinearColor::Green))
+					]
+				]			
+			]
 		]
 	];
 
@@ -199,6 +204,7 @@ void SNiagaraStackFunctionInputValue::Construct(const FArguments& InArgs, UNiaga
 			.ButtonStyle(FEditorStyle::Get(), "NoBorder")
 			.OnGetMenuContent(this, &SNiagaraStackFunctionInputValue::GetFilteredViewPropertiesContent)
 			.Visibility(this, &SNiagaraStackFunctionInputValue::GetFilteredViewContextButtonVisibility)
+			.IsEnabled(true)
 			.ButtonContent()
 			[
 				SNew(SImage)
@@ -1317,33 +1323,45 @@ void SNiagaraStackFunctionInputValue::TriggerRefresh(const TMap<EScriptSource, b
 
 const FSlateBrush* SNiagaraStackFunctionInputValue::GetFilteredViewIcon() const
 {
-	return GetFilteredViewVisibility() == ECheckBoxState::Checked ? FEditorStyle::GetBrush(TEXT("TextureEditor.GreenChannel.Small")) : FEditorStyle::GetBrush(TEXT("TextureEditor.AlphaChannel.Small"));
+	return FEditorStyle::GetBrush(TEXT("TextureEditor.GreenChannel.Small"));
 }
 
 EVisibility SNiagaraStackFunctionInputValue::GetFilteredViewContextButtonVisibility() const
 {
-	UNiagaraEmitterEditorData* EditorData = FunctionInput->GetEmitterViewModel()? &FunctionInput->GetEmitterViewModel()->GetOrCreateEditorData() : nullptr;
-		
-	if (!EditorData || EditorData->ShouldShowSummaryView())
+	UNiagaraEmitter* Emitter = (FunctionInput && FunctionInput->GetEmitterViewModel())? FunctionInput->GetEmitterViewModel()->GetEmitter() : nullptr;
+	UNiagaraEmitterEditorData* EditorData = Emitter? Cast<UNiagaraEmitterEditorData>(Emitter->GetEditorData()) : nullptr;
+	
+	if (!EditorData || EditorData->ShouldShowSummaryView() || FunctionInput->GetTypedOuter<UNiagaraStackSummaryViewObject>() == nullptr)
 	{
 		return EVisibility::Hidden;
 	}
+	
+	UNiagaraStackFunctionInput* ParentInput = FNiagaraStackEditorWidgetsUtilities::GetParentInputForSummaryView(FunctionInput);
+	TOptional<FFunctionInputSummaryViewKey> Key = FNiagaraStackEditorWidgetsUtilities::GetSummaryViewInputKeyForFunctionInput(ParentInput);
+	return (Key.IsSet() && EditorData->GetSummaryViewMetaData(Key.GetValue()).bVisible)? EVisibility::Visible : EVisibility::Hidden;
+}
 
+bool SNiagaraStackFunctionInputValue::GetSummaryViewChangeEnabledStateAllowed() const
+{
 	// Only show if we aren't a child input, and if we're in the summary view category
-	UNiagaraStackFunctionInput* ParentInput = FNiagaraStackEditorWidgetsUtilities::FindTopMostParentFunctionInput(FunctionInput);
-	if (FunctionInput == ParentInput)
-	{
-		UObject* Outer = FunctionInput->GetOuter();
-		while (Outer)
-		{
-			if (Outer->IsA<UNiagaraStackSummaryViewObject>())
-			{
-				return EVisibility::Visible;
-			}
-			Outer = Outer->GetOuter();
-		}
-	}
-	return EVisibility::Collapsed;
+	UNiagaraStackFunctionInput* ParentInput = FNiagaraStackEditorWidgetsUtilities::GetParentInputForSummaryView(FunctionInput);
+	return FunctionInput == ParentInput;	
+}
+
+bool SNiagaraStackFunctionInputValue::GetSummaryViewChangeDisplayNameAllowed() const
+{
+	// Only show if we aren't a child input, and if we're in the summary view category
+	// Also disallow changing display name on a direct attribute set.
+	UNiagaraStackFunctionInput* ParentInput = FNiagaraStackEditorWidgetsUtilities::GetParentInputForSummaryView(FunctionInput);	
+	const UNiagaraNodeFunctionCall& InputCallNode = FunctionInput->GetInputFunctionCallNode();
+	return FunctionInput == ParentInput && !InputCallNode.IsA<UNiagaraNodeAssignment>();
+}
+
+bool SNiagaraStackFunctionInputValue::GetSummaryViewChangeCategoryAllowed() const
+{
+	// Only show if we aren't a child input, and if we're in the summary view category
+	UNiagaraStackFunctionInput* ParentInput = FNiagaraStackEditorWidgetsUtilities::GetParentInputForSummaryView(FunctionInput);
+	return FunctionInput == ParentInput;	
 }
 
 TSharedRef<SWidget> SNiagaraStackFunctionInputValue::GetFilteredViewPropertiesContent()
@@ -1351,18 +1369,6 @@ TSharedRef<SWidget> SNiagaraStackFunctionInputValue::GetFilteredViewPropertiesCo
 	FMenuBuilder MenuBuilder(false, MakeShared<FUICommandList>());
 	MenuBuilder.BeginSection("SummaryViewDetails", LOCTEXT("SummaryViewDetails", "Emitter Summary Details"));
 
-	// Is Enabled
-	TSharedRef<SWidget> IsEnabledWidget =
-		SNew(SBox)
-		.WidthOverride(100)
-		.Padding(FMargin(5, 0, 0, 0))
-		[
-			SNew(SCheckBox)
-			.IsChecked(this, &SNiagaraStackFunctionInputValue::GetFilteredViewVisibility)
-		.OnCheckStateChanged(this, &SNiagaraStackFunctionInputValue::OnFilteredViewVisibilityChanged)
-		];
-	MenuBuilder.AddWidget(IsEnabledWidget, LOCTEXT("FilteredViewIsEnabled", "Is Enabled"));
-
 
 	// Display Name
 	TSharedRef<SWidget> NameWidget =
@@ -1372,8 +1378,9 @@ TSharedRef<SWidget> SNiagaraStackFunctionInputValue::GetFilteredViewPropertiesCo
 		[
 			SNew(SEditableTextBox)
 			.Text(this, &SNiagaraStackFunctionInputValue::GetFilteredViewDisplayName)
-		.OnVerifyTextChanged(this, &SNiagaraStackFunctionInputValue::VerifyFilteredViewDisplayName)
-		.OnTextCommitted(this, &SNiagaraStackFunctionInputValue::FilteredViewDisplayNameTextCommitted)
+			.IsEnabled(this, &SNiagaraStackFunctionInputValue::GetSummaryViewChangeDisplayNameAllowed)
+			.OnVerifyTextChanged(this, &SNiagaraStackFunctionInputValue::VerifyFilteredViewDisplayName)
+			.OnTextCommitted(this, &SNiagaraStackFunctionInputValue::FilteredViewDisplayNameTextCommitted)
 		];
 	MenuBuilder.AddWidget(NameWidget, LOCTEXT("FilteredViewDisplayName", "Display Name"));
 
@@ -1385,6 +1392,7 @@ TSharedRef<SWidget> SNiagaraStackFunctionInputValue::GetFilteredViewPropertiesCo
 		.Padding(FMargin(5, 0, 0, 0))
 		[
 			SNew(SEditableTextBox)
+			.IsEnabled(this, &SNiagaraStackFunctionInputValue::GetSummaryViewChangeCategoryAllowed)
 			.Text(this, &SNiagaraStackFunctionInputValue::GetFilteredViewCategory)
 			.OnVerifyTextChanged(this, &SNiagaraStackFunctionInputValue::VerifyFilteredViewCategory)
 			.OnTextCommitted(this, &SNiagaraStackFunctionInputValue::FilteredViewCategoryTextCommitted)
@@ -1398,14 +1406,11 @@ TSharedRef<SWidget> SNiagaraStackFunctionInputValue::GetFilteredViewPropertiesCo
 		.WidthOverride(100)
 		.Padding(FMargin(5, 0, 0, 0))
 		[
-			SNew(SSpinBox<int32>)
+			SNew(SNumericEntryBox<int32>)
 			//.Style(FNiagaraEditorStyle::Get(), "NiagaraEditor.ParameterSpinBox")
 			//.Font(FNiagaraEditorStyle::Get().GetFontStyle("NiagaraEditor.ParameterFont"))
 			.MinValue(0)
 			.MaxValue(TOptional<int32>())
-			.MaxSliderValue(0)
-			.MinSliderValue(100)
-			.Delta(1)
 			.Value(this, &SNiagaraStackFunctionInputValue::GetFilteredViewSortIndex)
 			.OnValueChanged(this, &SNiagaraStackFunctionInputValue::FilteredViewSortIndexChanged)
 			.OnValueCommitted(this, &SNiagaraStackFunctionInputValue::FilteredViewSortIndexCommitted)
@@ -1417,56 +1422,20 @@ TSharedRef<SWidget> SNiagaraStackFunctionInputValue::GetFilteredViewPropertiesCo
 	return MenuBuilder.MakeWidget();
 }
 
-ECheckBoxState SNiagaraStackFunctionInputValue::GetFilteredViewVisibility() const
-{
-	UNiagaraEmitterEditorData* EditorData = FunctionInput->GetEmitterViewModel()? &FunctionInput->GetEmitterViewModel()->GetOrCreateEditorData() : nullptr;
-	if (EditorData)
-	{
-		TOptional<FFunctionInputSummaryViewKey> Key = FNiagaraStackEditorWidgetsUtilities::GetSummaryViewInputKeyForFunctionInput(FunctionInput);
-		if (Key.IsSet())
-		{
-			return EditorData->GetSummaryViewMetaData(Key.GetValue()).bVisible ? ECheckBoxState::Checked : ECheckBoxState::Unchecked;
-		}
-	}
-
-	return ECheckBoxState::Unchecked;
-}
-
-void SNiagaraStackFunctionInputValue::OnFilteredViewVisibilityChanged(ECheckBoxState NewState)
-{
-	UNiagaraEmitterEditorData* EditorData = FunctionInput->GetEmitterViewModel()? &FunctionInput->GetEmitterViewModel()->GetOrCreateEditorData() : nullptr;
-	if (EditorData)
-	{
-		TOptional<FFunctionInputSummaryViewKey> Key = FNiagaraStackEditorWidgetsUtilities::GetSummaryViewInputKeyForFunctionInput(FunctionInput);
-		if (Key.IsSet())
-		{
-			FFunctionInputSummaryViewMetadata SummaryViewMetaData = EditorData->GetSummaryViewMetaData(Key.GetValue());
-			SummaryViewMetaData.bVisible = NewState == ECheckBoxState::Checked;				
-			EditorData->SetSummaryViewMetaData(Key.GetValue(), SummaryViewMetaData);	
-		}
-	}
-}
-
 FText SNiagaraStackFunctionInputValue::GetFilteredViewDisplayName() const
 {
-	UNiagaraEmitterEditorData* EditorData = FunctionInput->GetEmitterViewModel()? &FunctionInput->GetEmitterViewModel()->GetOrCreateEditorData() : nullptr;
-	if (EditorData)
-	{
-		TOptional<FFunctionInputSummaryViewKey> Key = FNiagaraStackEditorWidgetsUtilities::GetSummaryViewInputKeyForFunctionInput(FunctionInput);
-		if (Key.IsSet())
-		{
-			FName DisplayName = EditorData->GetSummaryViewMetaData(Key.GetValue()).DisplayName;
+	UNiagaraEmitterEditorData* EditorData = FunctionInput->GetEmitterViewModel()? &FunctionInput->GetEmitterViewModel()->GetOrCreateEditorData() : nullptr;	
+	UNiagaraStackFunctionInput* ParentInput = FNiagaraStackEditorWidgetsUtilities::GetParentInputForSummaryView(FunctionInput);
+	TOptional<FFunctionInputSummaryViewKey> Key = FNiagaraStackEditorWidgetsUtilities::GetSummaryViewInputKeyForFunctionInput(ParentInput);
 
-			return (DisplayName != NAME_None) ? FText::FromName(DisplayName) : FunctionInput->GetDisplayName();
-		}
-	}
+	const FName SummaryViewName = (EditorData && Key.IsSet())? EditorData->GetSummaryViewMetaData(Key.GetValue()).DisplayName : NAME_None;
 
-	return FText();
+	return (SummaryViewName != NAME_None) ? FText::FromName(SummaryViewName) : FunctionInput->GetDisplayName();
 }
 
 bool SNiagaraStackFunctionInputValue::VerifyFilteredViewDisplayName(const FText& InText, FText& OutErrorMessage) const
 {
-	FString DisplayName = InText.ToString();
+	const FString DisplayName = InText.ToString();
 
 	if (DisplayName.Len() >= FNiagaraConstants::MaxCategoryNameLength)
 	{
@@ -1479,21 +1448,22 @@ bool SNiagaraStackFunctionInputValue::VerifyFilteredViewDisplayName(const FText&
 
 void SNiagaraStackFunctionInputValue::FilteredViewDisplayNameTextCommitted(const FText& Text, ETextCommit::Type CommitType)
 {
-	UNiagaraEmitterEditorData* EditorData = FunctionInput->GetEmitterViewModel()? &FunctionInput->GetEmitterViewModel()->GetOrCreateEditorData() : nullptr;
-	if (EditorData)
+	UNiagaraEmitterEditorData* EditorData = FunctionInput->GetEmitterViewModel()? &FunctionInput->GetEmitterViewModel()->GetOrCreateEditorData() : nullptr;	
+	UNiagaraStackFunctionInput* ParentInput = FNiagaraStackEditorWidgetsUtilities::GetParentInputForSummaryView(FunctionInput);
+	TOptional<FFunctionInputSummaryViewKey> Key = FNiagaraStackEditorWidgetsUtilities::GetSummaryViewInputKeyForFunctionInput(ParentInput);
+	
+	if (EditorData && Key.IsSet())
 	{
-		TOptional<FFunctionInputSummaryViewKey> Key = FNiagaraStackEditorWidgetsUtilities::GetSummaryViewInputKeyForFunctionInput(FunctionInput);
-		if (Key.IsSet())
+		const FName NewName = Text.IsEmptyOrWhitespace()? NAME_None : FName(Text.ToString());
+		FFunctionInputSummaryViewMetadata SummaryViewMetaData = EditorData->GetSummaryViewMetaData(Key.GetValue());		
+		if (NewName != SummaryViewMetaData.DisplayName)
 		{
-			FName NewName = (Text.IsEmptyOrWhitespace() || Text.EqualTo(FunctionInput->GetDisplayName()))? NAME_None : FName(Text.ToString());
+			FScopedTransaction ScopedTransaction(FText::Format(LOCTEXT("SummaryViewChangedInputDisplayName", "Changed summary view display name for {0}"), FunctionInput->GetDisplayName()));
+			EditorData->Modify();
 			
-			FFunctionInputSummaryViewMetadata SummaryViewMetaData = EditorData->GetSummaryViewMetaData(Key.GetValue());
-			if (NewName != SummaryViewMetaData.DisplayName)
-			{
-				SummaryViewMetaData.DisplayName = NewName;
-				EditorData->SetSummaryViewMetaData(Key.GetValue(), SummaryViewMetaData);
-			}
-		}
+			SummaryViewMetaData.DisplayName = NewName;
+			EditorData->SetSummaryViewMetaData(Key.GetValue(), SummaryViewMetaData);
+		}		
 	}
 }
 
@@ -1506,24 +1476,18 @@ FText GetFunctionInputCategory(UNiagaraStackFunctionInput* Input)
 
 FText SNiagaraStackFunctionInputValue::GetFilteredViewCategory() const
 {
-	UNiagaraEmitterEditorData* EditorData = FunctionInput->GetEmitterViewModel()? &FunctionInput->GetEmitterViewModel()->GetOrCreateEditorData() : nullptr;
-	if (EditorData)
-	{
-		TOptional<FFunctionInputSummaryViewKey> Key = FNiagaraStackEditorWidgetsUtilities::GetSummaryViewInputKeyForFunctionInput(FunctionInput);
-		if (Key.IsSet())
-		{
-			FName CategoryName = EditorData->GetSummaryViewMetaData(Key.GetValue()).Category;
+	UNiagaraEmitterEditorData* EditorData = FunctionInput->GetEmitterViewModel()? &FunctionInput->GetEmitterViewModel()->GetOrCreateEditorData() : nullptr;	
+	UNiagaraStackFunctionInput* ParentInput = FNiagaraStackEditorWidgetsUtilities::GetParentInputForSummaryView(FunctionInput);
+	TOptional<FFunctionInputSummaryViewKey> Key = FNiagaraStackEditorWidgetsUtilities::GetSummaryViewInputKeyForFunctionInput(ParentInput);
 
-			return (CategoryName != NAME_None) ? FText::FromName(CategoryName) : GetFunctionInputCategory(FunctionInput);
-		}
-	}
+	const FName SummaryViewCategory = (EditorData && Key.IsSet())? EditorData->GetSummaryViewMetaData(Key.GetValue()).Category : NAME_None;
 
-	return FText();
+	return (SummaryViewCategory != NAME_None) ? FText::FromName(SummaryViewCategory) : GetFunctionInputCategory(FunctionInput);
 }
 
 bool SNiagaraStackFunctionInputValue::VerifyFilteredViewCategory(const FText& InText, FText& OutErrorMessage) const
 {
-	FString CategoryName = InText.ToString();
+	const FString CategoryName = InText.ToString();
 
 	if (CategoryName.Len() >= FNiagaraConstants::MaxCategoryNameLength)
 	{
@@ -1536,56 +1500,58 @@ bool SNiagaraStackFunctionInputValue::VerifyFilteredViewCategory(const FText& In
 
 void SNiagaraStackFunctionInputValue::FilteredViewCategoryTextCommitted(const FText& Text, ETextCommit::Type CommitType)
 {
-	UNiagaraEmitterEditorData* EditorData = FunctionInput->GetEmitterViewModel()? &FunctionInput->GetEmitterViewModel()->GetOrCreateEditorData() : nullptr;
-	if (EditorData)
+	UNiagaraEmitterEditorData* EditorData = FunctionInput->GetEmitterViewModel()? &FunctionInput->GetEmitterViewModel()->GetOrCreateEditorData() : nullptr;	
+	UNiagaraStackFunctionInput* ParentInput = FNiagaraStackEditorWidgetsUtilities::GetParentInputForSummaryView(FunctionInput);
+	TOptional<FFunctionInputSummaryViewKey> Key = FNiagaraStackEditorWidgetsUtilities::GetSummaryViewInputKeyForFunctionInput(ParentInput);
+	
+	if (EditorData && Key.IsSet())
 	{
-		TOptional<FFunctionInputSummaryViewKey> Key = FNiagaraStackEditorWidgetsUtilities::GetSummaryViewInputKeyForFunctionInput(FunctionInput);
-		if (Key.IsSet())
+		const FName NewCategory = Text.IsEmptyOrWhitespace()? NAME_None : FName(Text.ToString());
+		FFunctionInputSummaryViewMetadata SummaryViewMetaData = EditorData->GetSummaryViewMetaData(Key.GetValue());		
+		if (NewCategory != SummaryViewMetaData.Category)
 		{
-			FName NewCategory = (Text.IsEmptyOrWhitespace() || Text.EqualTo(GetFunctionInputCategory(FunctionInput)))? NAME_None : FName(Text.ToString());
+			FScopedTransaction ScopedTransaction(FText::Format(LOCTEXT("SummaryViewChangedInputCategory", "Changed summary view category for {0}"), FunctionInput->GetDisplayName()));
+			EditorData->Modify();
 			
-			FFunctionInputSummaryViewMetadata SummaryViewMetaData = EditorData->GetSummaryViewMetaData(Key.GetValue());
-			if (NewCategory != SummaryViewMetaData.Category)
-			{
-				SummaryViewMetaData.Category = NewCategory;
-				EditorData->SetSummaryViewMetaData(Key.GetValue(), SummaryViewMetaData);
-			}
-		}
+			SummaryViewMetaData.Category = NewCategory;
+			EditorData->SetSummaryViewMetaData(Key.GetValue(), SummaryViewMetaData);
+		}		
 	}
 }
 
 
-int32 SNiagaraStackFunctionInputValue::GetFilteredViewSortIndex() const
+TOptional<int32> SNiagaraStackFunctionInputValue::GetFilteredViewSortIndex() const
 {
-	UNiagaraEmitterEditorData* EditorData = FunctionInput->GetEmitterViewModel()? &FunctionInput->GetEmitterViewModel()->GetOrCreateEditorData() : nullptr;
-	if (EditorData)
+	UNiagaraEmitterEditorData* EditorData = FunctionInput->GetEmitterViewModel()? &FunctionInput->GetEmitterViewModel()->GetOrCreateEditorData() : nullptr;	
+	UNiagaraStackFunctionInput* ParentInput = FNiagaraStackEditorWidgetsUtilities::GetParentInputForSummaryView(FunctionInput);
+	TOptional<FFunctionInputSummaryViewKey> Key = FNiagaraStackEditorWidgetsUtilities::GetSummaryViewInputKeyForFunctionInput(ParentInput);
+
+	if (EditorData && Key.IsSet())
 	{
-		TOptional<FFunctionInputSummaryViewKey> Key = FNiagaraStackEditorWidgetsUtilities::GetSummaryViewInputKeyForFunctionInput(FunctionInput);
-		if (Key.IsSet())
-		{
-			return EditorData->GetSummaryViewMetaData(Key.GetValue()).SortIndex;
-		}
+		int32 SortIndex = EditorData->GetSummaryViewMetaData(Key.GetValue()).SortIndex;
+		return SortIndex == INDEX_NONE? TOptional<int32>() : SortIndex;
 	}
-
-	return 0;
+	
+	return TOptional<int32>();
 }
 
 void SNiagaraStackFunctionInputValue::FilteredViewSortIndexChanged(int32 Value)
 {
-	UNiagaraEmitterEditorData* EditorData = FunctionInput->GetEmitterViewModel()? &FunctionInput->GetEmitterViewModel()->GetOrCreateEditorData() : nullptr;
-	if (EditorData)
+	UNiagaraEmitterEditorData* EditorData = FunctionInput->GetEmitterViewModel()? &FunctionInput->GetEmitterViewModel()->GetOrCreateEditorData() : nullptr;	
+	UNiagaraStackFunctionInput* ParentInput = FNiagaraStackEditorWidgetsUtilities::GetParentInputForSummaryView(FunctionInput);
+	TOptional<FFunctionInputSummaryViewKey> Key = FNiagaraStackEditorWidgetsUtilities::GetSummaryViewInputKeyForFunctionInput(ParentInput);
+	
+	if (EditorData && Key.IsSet())
 	{
-		TOptional<FFunctionInputSummaryViewKey> Key = FNiagaraStackEditorWidgetsUtilities::GetSummaryViewInputKeyForFunctionInput(FunctionInput);
-		if (Key.IsSet())
+		FFunctionInputSummaryViewMetadata SummaryViewMetaData = EditorData->GetSummaryViewMetaData(Key.GetValue());		
+		if (Value != SummaryViewMetaData.SortIndex)
 		{
-			FFunctionInputSummaryViewMetadata SummaryViewMetaData = EditorData->GetSummaryViewMetaData(Key.GetValue());
-
-			if (SummaryViewMetaData.SortIndex != Value)
-			{
-				SummaryViewMetaData.SortIndex = Value;			
-				EditorData->SetSummaryViewMetaData(Key.GetValue(), SummaryViewMetaData);
-			}
-		}
+			FScopedTransaction ScopedTransaction(FText::Format(LOCTEXT("SummaryViewChangedInputSortIndex", "Changed summary view sorty index for {0}"), FunctionInput->GetDisplayName()));
+			EditorData->Modify();
+			
+			SummaryViewMetaData.SortIndex = Value;
+			EditorData->SetSummaryViewMetaData(Key.GetValue(), SummaryViewMetaData);
+		}		
 	}
 }
 
diff --git a/Engine/Plugins/FX/Niagara/Source/NiagaraEditorWidgets/Private/Stack/SNiagaraStackFunctionInputValue.h b/Engine/Plugins/FX/Niagara/Source/NiagaraEditorWidgets/Private/Stack/SNiagaraStackFunctionInputValue.h
--- a/Engine/Plugins/FX/Niagara/Source/NiagaraEditorWidgets/Private/Stack/SNiagaraStackFunctionInputValue.h
+++ b/Engine/Plugins/FX/Niagara/Source/NiagaraEditorWidgets/Private/Stack/SNiagaraStackFunctionInputValue.h
@@ -158,11 +158,11 @@ private:
 
 	const FSlateBrush* GetFilteredViewIcon() const;
 	EVisibility GetFilteredViewContextButtonVisibility() const;
+	bool GetSummaryViewChangeEnabledStateAllowed() const;
+	bool GetSummaryViewChangeDisplayNameAllowed() const;
+	bool GetSummaryViewChangeCategoryAllowed() const;
 	TSharedRef<SWidget> GetFilteredViewPropertiesContent();
-
-	ECheckBoxState GetFilteredViewVisibility() const;
-	void OnFilteredViewVisibilityChanged(ECheckBoxState NewState);
-
+	
 	FText GetFilteredViewDisplayName() const;
 	bool VerifyFilteredViewDisplayName(const FText& InText, FText& OutErrorMessage) const;
 	void FilteredViewDisplayNameTextCommitted(const FText& Text, ETextCommit::Type CommitType);
@@ -171,7 +171,7 @@ private:
 	bool VerifyFilteredViewCategory(const FText& InText, FText& OutErrorMessage) const;
 	void FilteredViewCategoryTextCommitted(const FText& Text, ETextCommit::Type CommitType);
 
-	int32 GetFilteredViewSortIndex() const;
+	TOptional<int32> GetFilteredViewSortIndex() const;
 	void FilteredViewSortIndexChanged(int32 Value);
 	void FilteredViewSortIndexCommitted(int32 Value, ETextCommit::Type CommitInfo);
 
diff --git a/Engine/Plugins/FX/Niagara/Source/NiagaraEditorWidgets/Private/Stack/SNiagaraStackTableRow.cpp b/Engine/Plugins/FX/Niagara/Source/NiagaraEditorWidgets/Private/Stack/SNiagaraStackTableRow.cpp
--- a/Engine/Plugins/FX/Niagara/Source/NiagaraEditorWidgets/Private/Stack/SNiagaraStackTableRow.cpp
+++ b/Engine/Plugins/FX/Niagara/Source/NiagaraEditorWidgets/Private/Stack/SNiagaraStackTableRow.cpp
@@ -19,6 +19,8 @@
 #include "ViewModels/NiagaraEmitterViewModel.h"
 #include "Styling/StyleColors.h"
 #include "NiagaraEmitterEditorData.h"
+#include "ViewModels/Stack/NiagaraStackFunctionInputCollection.h"
+#include "ViewModels/Stack/NiagaraStackInputCategory.h"
 
 #define LOCTEXT_NAMESPACE "NiagaraStackTableRow"
 
@@ -466,10 +468,10 @@ FReply SNiagaraStackTableRow::OnMouseButtonUp(const FGeometry& MyGeometry, const
 				FCanExecuteAction(),
 				FIsActionChecked::CreateSP(this, &SNiagaraStackTableRow::ShouldShowInSummaryView));
 		
-			MenuBuilder.BeginSection("SummaryViewActions", LOCTEXT("SummaryViewActions", "Summary View"));
+			MenuBuilder.BeginSection("SummaryViewActions", LOCTEXT("SummaryViewActions", "Emitter Summary"));
 			MenuBuilder.AddMenuEntry(
-				LOCTEXT("SummaryViewShow", "Show In Summary View"),
-				LOCTEXT("SummaryViewShowTooltip", "Should this parameter be visible in the summary view?"),
+				LOCTEXT("SummaryViewShow", "Show in Emitter Summary"),
+				LOCTEXT("SummaryViewShowTooltip", "Should this parameter be visible in the emitter summary?"),
 				FSlateIcon(),
 				ShowHideSummaryViewAction,
 				NAME_None, EUserInterfaceActionType::ToggleButton);
@@ -614,61 +616,57 @@ void SNiagaraStackTableRow::NavigateTo(UNiagaraStackEntry* Item)
 bool SNiagaraStackTableRow::IsValidForSummaryView() const
 {
 	UNiagaraStackFunctionInput* FunctionInput = Cast<UNiagaraStackFunctionInput>(StackEntry);
-	if (FunctionInput && FunctionInput->GetEmitterViewModel())
+
+	UNiagaraEmitter* Emitter = (FunctionInput && FunctionInput->GetEmitterViewModel())? FunctionInput->GetEmitterViewModel()->GetEmitter() : nullptr;
+	UNiagaraEmitterEditorData* EditorData = Emitter? Cast<UNiagaraEmitterEditorData>(Emitter->GetEditorData()) : nullptr;
+
+	if (EditorData)
 	{
-		UNiagaraEmitter* Emitter = FunctionInput->GetEmitterViewModel()->GetEmitter();
-		UNiagaraStackFunctionInput* ParentInput = FNiagaraStackEditorWidgetsUtilities::FindTopMostParentFunctionInput(FunctionInput);			
-		if (Emitter && FNiagaraStackEditorWidgetsUtilities::GetSummaryViewInputKeyForFunctionInput(ParentInput).IsSet())
-		{
-			return true;	
-		}
-	}	
+		UNiagaraStackFunctionInput* ParentInput = FNiagaraStackEditorWidgetsUtilities::GetParentInputForSummaryView(FunctionInput);
+		return FNiagaraStackEditorWidgetsUtilities::GetSummaryViewInputKeyForFunctionInput(ParentInput).IsSet();
+	}
 	return false;
 }
 
 void SNiagaraStackTableRow::ToggleShowInSummaryView()
 {
 	UNiagaraStackFunctionInput* FunctionInput = Cast<UNiagaraStackFunctionInput>(StackEntry);
-	if (FunctionInput && FunctionInput->GetEmitterViewModel())
+	UNiagaraEmitter* Emitter = (FunctionInput && FunctionInput->GetEmitterViewModel())? FunctionInput->GetEmitterViewModel()->GetEmitter() : nullptr;
+	UNiagaraEmitterEditorData* EditorData = Emitter? Cast<UNiagaraEmitterEditorData>(Emitter->GetEditorData()) : nullptr;
+	
+	if (EditorData)
 	{
-		UNiagaraEmitter* Emitter = FunctionInput->GetEmitterViewModel()->GetEmitter();
-		if (Emitter)
+		UNiagaraStackFunctionInput* ParentInput = FNiagaraStackEditorWidgetsUtilities::GetParentInputForSummaryView(FunctionInput);
+		TOptional<FFunctionInputSummaryViewKey> Key = FNiagaraStackEditorWidgetsUtilities::GetSummaryViewInputKeyForFunctionInput(ParentInput);
+	
+		if (Key.IsSet())
 		{
-			UNiagaraStackFunctionInput* ParentInput = FNiagaraStackEditorWidgetsUtilities::FindTopMostParentFunctionInput(FunctionInput);		
-			TOptional<FFunctionInputSummaryViewKey> Key = FNiagaraStackEditorWidgetsUtilities::GetSummaryViewInputKeyForFunctionInput(ParentInput);
-			if (Key.IsSet())
-			{
-				UNiagaraEmitterEditorData* EditorData = Cast<UNiagaraEmitterEditorData>(Emitter->GetEditorData());
-				if (EditorData)
-				{
-					FFunctionInputSummaryViewMetadata SummaryViewMetaData = EditorData->GetSummaryViewMetaData(Key.GetValue());
-					SummaryViewMetaData.bVisible = !SummaryViewMetaData.bVisible;			
-					EditorData->SetSummaryViewMetaData(Key.GetValue(), SummaryViewMetaData);
-				}
-			}			
+			FFunctionInputSummaryViewMetadata SummaryViewMetaData = EditorData->GetSummaryViewMetaData(Key.GetValue());
+			
+			FScopedTransaction ScopedTransaction(FText::Format(LOCTEXT("SummaryViewChangedInputVisibility", "Changed summary view visibility for {0}"), FunctionInput->GetDisplayName()));
+			EditorData->Modify();
+			
+			SummaryViewMetaData.bVisible = !SummaryViewMetaData.bVisible;
+			EditorData->SetSummaryViewMetaData(Key.GetValue(), SummaryViewMetaData);
 		}
-	}
+	}	
 }
 
 bool SNiagaraStackTableRow::ShouldShowInSummaryView() const
 {
 	UNiagaraStackFunctionInput* FunctionInput = Cast<UNiagaraStackFunctionInput>(StackEntry);
-	if (FunctionInput && FunctionInput->GetEmitterViewModel())
+	UNiagaraEmitter* Emitter = (FunctionInput && FunctionInput->GetEmitterViewModel())? FunctionInput->GetEmitterViewModel()->GetEmitter() : nullptr;
+	UNiagaraEmitterEditorData* EditorData = Emitter? Cast<UNiagaraEmitterEditorData>(Emitter->GetEditorData()) : nullptr;
+	
+	if (EditorData)
 	{
-		UNiagaraEmitter* Emitter = FunctionInput->GetEmitterViewModel()->GetEmitter();
-		if (Emitter)
+		UNiagaraStackFunctionInput* ParentInput = FNiagaraStackEditorWidgetsUtilities::GetParentInputForSummaryView(FunctionInput);
+		TOptional<FFunctionInputSummaryViewKey> Key = FNiagaraStackEditorWidgetsUtilities::GetSummaryViewInputKeyForFunctionInput(ParentInput);
+	
+		if (Key.IsSet())
 		{
-			UNiagaraStackFunctionInput* ParentInput = FNiagaraStackEditorWidgetsUtilities::FindTopMostParentFunctionInput(FunctionInput);			
-			TOptional<FFunctionInputSummaryViewKey> Key = FNiagaraStackEditorWidgetsUtilities::GetSummaryViewInputKeyForFunctionInput(ParentInput);
-			if (Key.IsSet())
-			{
-				UNiagaraEmitterEditorData* EditorData = Cast<UNiagaraEmitterEditorData>(Emitter->GetEditorData());
-				if (EditorData)
-				{
-					return EditorData->GetSummaryViewMetaData(Key.GetValue()).bVisible;
-				}
-			}			
-		}
+			return EditorData->GetSummaryViewMetaData(Key.GetValue()).bVisible;
+		}			
 	}
 
 	return false;
