- Modify how we write GPU profiling data to fix an issue where it would stop if we caught up
- Debug HUD uses smoothed timers for it's main display to be easier to read

#rnx

diff --git a/Engine/Plugins/FX/Niagara/Source/Niagara/Private/NiagaraDebugHud.cpp b/Engine/Plugins/FX/Niagara/Source/Niagara/Private/NiagaraDebugHud.cpp
--- a/Engine/Plugins/FX/Niagara/Source/Niagara/Private/NiagaraDebugHud.cpp
+++ b/Engine/Plugins/FX/Niagara/Source/Niagara/Private/NiagaraDebugHud.cpp
@@ -539,6 +539,13 @@ FNiagaraDebugHud::FNiagaraDebugHud(UWorld* World)
 				GpuResults = InGpuResults;
 				GpuResultsGameFrameCounter = GFrameCounter;
 
+				for (const auto& StageResults : GpuResults->StageResults )
+				{
+					GpuTotalDispatches.Accumulate(GpuResultsGameFrameCounter, StageResults.NumDispatches);
+					GpuTotalDispatchGroups.Accumulate(GpuResultsGameFrameCounter, StageResults.NumDispatchGroups);
+					GpuTotalMicroseconds.Accumulate(GpuResultsGameFrameCounter, StageResults.DurationMicroseconds);
+				}
+
 				for ( const auto& DispatchResult : GpuResults->DispatchResults )
 				{
 					if ( DispatchResult.OwnerEmitter.IsExplicitlyNull() )
@@ -1754,25 +1761,15 @@ void FNiagaraDebugHud::DrawGpuComputeOverriew(class FNiagaraWorldManager* WorldM
 
 	// Display overview of results
 	{
-		int32 TotalDispatches = 0;
-		int32 TotalDispatchGroups = 0;
-		uint64 TotalMicroseconds = 0;
-		for (const auto& StageResult : GpuResults->StageResults)
-		{
-			TotalDispatches += StageResult.NumDispatches;
-			TotalDispatchGroups += StageResult.NumDispatchGroups;
-			TotalMicroseconds += StageResult.DurationMicroseconds;
-		}
-
 		FSimpleTableDraw SimpleTable;
 		SimpleTable.AddColumns(1, 250.0f);
 		SimpleTable.AddColumns(2, 150.0f);
 		SimpleTable.GetColumnText(0).Append(TEXT("Gpu Compute Overview"));
 		SimpleTable.RowComplete();
 
-		SimpleTable.GetColumnText(0).Appendf(TEXT("TotalMicroseconds : %llu"), TotalMicroseconds);
-		SimpleTable.GetColumnText(1).Appendf(TEXT("TotalDispatches : %d"), TotalDispatches);
-		SimpleTable.GetColumnText(2).Appendf(TEXT("TotalDispatchGroups : %d"), TotalDispatchGroups);
+		SimpleTable.GetColumnText(0).Appendf(TEXT("TotalMicroseconds : %llu"), GpuTotalMicroseconds.GetAverage());
+		SimpleTable.GetColumnText(1).Appendf(TEXT("TotalDispatches : %d"), GpuTotalDispatches.GetAverage());
+		SimpleTable.GetColumnText(2).Appendf(TEXT("TotalDispatchGroups : %d"), GpuTotalDispatchGroups.GetAverage());
 
 		SimpleTable.Draw(Font, DrawCanvas, TextLocation, DetailColor, BackgroundColor);
 		TextLocation.Y += fAdvanceHeight;
diff --git a/Engine/Plugins/FX/Niagara/Source/Niagara/Private/NiagaraDebugHud.h b/Engine/Plugins/FX/Niagara/Source/Niagara/Private/NiagaraDebugHud.h
--- a/Engine/Plugins/FX/Niagara/Source/Niagara/Private/NiagaraDebugHud.h
+++ b/Engine/Plugins/FX/Niagara/Source/Niagara/Private/NiagaraDebugHud.h
@@ -272,6 +272,10 @@ private:
 	FNiagaraGpuFrameResultsPtr	GpuResults;
 	uint64						GpuResultsGameFrameCounter = 0;
 
+	FSmoothedCounter<int32>		GpuTotalDispatches;
+	FSmoothedCounter<int32>		GpuTotalDispatchGroups;
+	FSmoothedCounter<uint64>	GpuTotalMicroseconds;
+
 	struct FGpuUsagePerStage
 	{
 		FSmoothedCounter<uint32> InstanceCount;
diff --git a/Engine/Plugins/FX/Niagara/Source/Niagara/Private/NiagaraGPUProfiler.cpp b/Engine/Plugins/FX/Niagara/Source/Niagara/Private/NiagaraGPUProfiler.cpp
--- a/Engine/Plugins/FX/Niagara/Source/Niagara/Private/NiagaraGPUProfiler.cpp
+++ b/Engine/Plugins/FX/Niagara/Source/Niagara/Private/NiagaraGPUProfiler.cpp
@@ -43,14 +43,11 @@ FNiagaraGPUProfiler::~FNiagaraGPUProfiler()
 
 void FNiagaraGPUProfiler::BeginFrame(FRHICommandListImmediate& RHICmdList)
 {
-	bProfilingFrame = false;
-	bProfilingDispatches = false;
-
-	// Process any frames that are potentially complete
-	while ( CurrentReadFrame != CurrentWriteFrame )
+	// Process all frames until we run out
+	while (FGpuFrameData* ReadFrame = GetReadFrame())
 	{
-		FGpuFrameData& ReadFrame = GetReadFrame();
-		if ( !ProcessFrame(RHICmdList, ReadFrame) )
+		// Attempt to process, might not be read
+		if ( !ProcessFrame(RHICmdList, *ReadFrame) )
 		{
 			break;
 		}
@@ -65,42 +62,32 @@ void FNiagaraGPUProfiler::BeginFrame(FRHICommandListImmediate& RHICmdList)
 		return;
 	}
 
-	// Have a processed all previous data?
-	if (GetWriteFrame().EndQuery.GetQuery() != nullptr)
-	{
-		return;
-	}
-
-	// Latch values for frame
-	bProfilingFrame = true;
-	bProfilingDispatches = true;
+	// Get frame to write into
+	ActiveWriteFrame = GetWriteFrame();
 }
 
 void FNiagaraGPUProfiler::EndFrame(FRHICommandList& RHICmdList)
 {
-	if ( !bProfilingFrame )
+	if (ActiveWriteFrame == nullptr)
 	{
 		return;
 	}
 
 	// Inject end marker so we know if all dispatches are complete
-	FGpuFrameData& WriteFrame = GetWriteFrame();
-
-	WriteFrame.EndQuery = QueryPool->AllocateQuery();
-	RHICmdList.EndRenderQuery(WriteFrame.EndQuery.GetQuery());
+	ActiveWriteFrame->EndQuery = QueryPool->AllocateQuery();
+	RHICmdList.EndRenderQuery(ActiveWriteFrame->EndQuery.GetQuery());
 
 	CurrentWriteFrame = (CurrentWriteFrame + 1) % NumBufferFrames;
 }
 
 void FNiagaraGPUProfiler::BeginStage(FRHICommandList& RHICmdList, ENiagaraGpuComputeTickStage::Type TickStage, int32 NumDispatchGroups)
 {
-	if (!bProfilingFrame)
+	if (ActiveWriteFrame == nullptr)
 	{
 		return;
 	}
 
-	FGpuFrameData& WriteFrame = GetWriteFrame();
-	FGpuStageTimer& StageTimer = WriteFrame.StageTimers[TickStage];
+	FGpuStageTimer& StageTimer = ActiveWriteFrame->StageTimers[TickStage];
 	StageTimer.NumDispatchGroups = NumDispatchGroups;
 	StageTimer.StartQuery = QueryPool->AllocateQuery();
 	RHICmdList.EndRenderQuery(StageTimer.StartQuery.GetQuery());
@@ -108,13 +95,12 @@ void FNiagaraGPUProfiler::BeginStage(FRHICommandList& RHICmdList, ENiagaraGpuCom
 
 void FNiagaraGPUProfiler::EndStage(FRHICommandList& RHICmdList, ENiagaraGpuComputeTickStage::Type TickStage, int32 NumDispatches)
 {
-	if (!bProfilingFrame)
+	if (ActiveWriteFrame == nullptr)
 	{
 		return;
 	}
 
-	FGpuFrameData& WriteFrame = GetWriteFrame();
-	FGpuStageTimer& StageTimer = WriteFrame.StageTimers[TickStage];
+	FGpuStageTimer& StageTimer = ActiveWriteFrame->StageTimers[TickStage];
 	StageTimer.NumDispatches = NumDispatches;
 	StageTimer.EndQuery = QueryPool->AllocateQuery();
 	RHICmdList.EndRenderQuery(StageTimer.EndQuery.GetQuery());
@@ -122,15 +108,14 @@ void FNiagaraGPUProfiler::EndStage(FRHICommandList& RHICmdList, ENiagaraGpuCompu
 
 void FNiagaraGPUProfiler::BeginDispatch(FRHICommandList& RHICmdList, const FNiagaraGpuDispatchInstance& DispatchInstance)
 {
-	if (!bProfilingDispatches)
+	if (ActiveWriteFrame == nullptr)
 	{
 		return;
 	}
 	check(bDispatchRecursionGuard == false);
 	bDispatchRecursionGuard = true;
 
-	FGpuFrameData& WriteFrame = GetWriteFrame();
-	FGpuDispatchTimer& DispatchTimer = WriteFrame.DispatchTimers.AddDefaulted_GetRef();
+	FGpuDispatchTimer& DispatchTimer = ActiveWriteFrame->DispatchTimers.AddDefaulted_GetRef();
 
 	DispatchTimer.bUniqueInstance = DispatchInstance.SimStageData.bSetDataToRender && (&DispatchInstance.InstanceData == &DispatchInstance.Tick.GetInstances()[0]);
 	DispatchTimer.OwnerComponent = DispatchInstance.InstanceData.Context->ProfilingComponentPtr;
@@ -142,15 +127,14 @@ void FNiagaraGPUProfiler::BeginDispatch(FRHICommandList& RHICmdList, const FNiag
 
 void FNiagaraGPUProfiler::BeginDispatch(FRHICommandList& RHICmdList, const struct FNiagaraComputeInstanceData& InstanceData, FName StageName)
 {
-	if (!bProfilingDispatches)
+	if (ActiveWriteFrame == nullptr)
 	{
 		return;
 	}
 	check(bDispatchRecursionGuard == false);
 	bDispatchRecursionGuard = true;
 
-	FGpuFrameData& WriteFrame = GetWriteFrame();
-	FGpuDispatchTimer& DispatchTimer = WriteFrame.DispatchTimers.AddDefaulted_GetRef();
+	FGpuDispatchTimer& DispatchTimer = ActiveWriteFrame->DispatchTimers.AddDefaulted_GetRef();
 	DispatchTimer.bUniqueInstance = false;
 	DispatchTimer.OwnerComponent = InstanceData.Context->ProfilingComponentPtr;
 	DispatchTimer.OwnerEmitter = InstanceData.Context->ProfilingEmitterPtr;
@@ -161,15 +145,14 @@ void FNiagaraGPUProfiler::BeginDispatch(FRHICommandList& RHICmdList, const struc
 
 void FNiagaraGPUProfiler::BeginDispatch(FRHICommandList& RHICmdList, FName StageName)
 {
-	if (!bProfilingDispatches)
+	if (ActiveWriteFrame == nullptr)
 	{
 		return;
 	}
 	check(bDispatchRecursionGuard == false);
 	bDispatchRecursionGuard = true;
 
-	FGpuFrameData& WriteFrame = GetWriteFrame();
-	FGpuDispatchTimer& DispatchTimer = WriteFrame.DispatchTimers.AddDefaulted_GetRef();
+	FGpuDispatchTimer& DispatchTimer = ActiveWriteFrame->DispatchTimers.AddDefaulted_GetRef();
 	DispatchTimer.bUniqueInstance = false;
 	DispatchTimer.OwnerComponent = nullptr;
 	DispatchTimer.OwnerEmitter = nullptr;
@@ -180,15 +163,14 @@ void FNiagaraGPUProfiler::BeginDispatch(FRHICommandList& RHICmdList, FName Stage
 
 void FNiagaraGPUProfiler::EndDispatch(FRHICommandList& RHICmdList)
 {
-	if (!bProfilingDispatches)
+	if (ActiveWriteFrame == nullptr)
 	{
 		return;
 	}
 	check(bDispatchRecursionGuard == true);
 	bDispatchRecursionGuard = false;
 
-	FGpuFrameData& WriteFrame = GetWriteFrame();
-	FGpuDispatchTimer& DispatchTimer = WriteFrame.DispatchTimers.Last();
+	FGpuDispatchTimer& DispatchTimer = ActiveWriteFrame->DispatchTimers.Last();
 	DispatchTimer.EndQuery = QueryPool->AllocateQuery();
 	RHICmdList.EndRenderQuery(DispatchTimer.EndQuery.GetQuery());
 }
diff --git a/Engine/Plugins/FX/Niagara/Source/Niagara/Private/NiagaraGPUProfiler.h b/Engine/Plugins/FX/Niagara/Source/Niagara/Private/NiagaraGPUProfiler.h
--- a/Engine/Plugins/FX/Niagara/Source/Niagara/Private/NiagaraGPUProfiler.h
+++ b/Engine/Plugins/FX/Niagara/Source/Niagara/Private/NiagaraGPUProfiler.h
@@ -39,6 +39,9 @@ class FNiagaraGPUProfiler : public FNiagaraGPUProfilerInterface
 
 	struct FGpuFrameData
 	{
+		bool CanWrite() const { return EndQuery.GetQuery() == nullptr; }
+		bool CanRead() const { return EndQuery.GetQuery() != nullptr; }
+
 		FRHIPooledRenderQuery		EndQuery;
 		FGpuStageTimer				StageTimers[ENiagaraGpuComputeTickStage::Max];
 		TArray<FGpuDispatchTimer>	DispatchTimers;
@@ -62,22 +65,20 @@ public:
 	void EndDispatch(FRHICommandList& RHICmdList);
 
 private:
-	FGpuFrameData& GetReadFrame() { check(CurrentReadFrame >= 0 && CurrentReadFrame < UE_ARRAY_COUNT(GpuFrames)); return GpuFrames[CurrentReadFrame]; }
-	FGpuFrameData& GetWriteFrame() { check(CurrentWriteFrame >= 0 && CurrentWriteFrame < UE_ARRAY_COUNT(GpuFrames)); return GpuFrames[CurrentWriteFrame]; }
+	FGpuFrameData* GetReadFrame() { check(CurrentReadFrame >= 0 && CurrentReadFrame < UE_ARRAY_COUNT(GpuFrames)); return GpuFrames[CurrentReadFrame].CanRead() ? &GpuFrames[CurrentReadFrame] : nullptr; }
+	FGpuFrameData* GetWriteFrame() { check(CurrentWriteFrame >= 0 && CurrentWriteFrame < UE_ARRAY_COUNT(GpuFrames)); return GpuFrames[CurrentWriteFrame].CanWrite() ? &GpuFrames[CurrentWriteFrame] : nullptr; }
 	bool ProcessFrame(FRHICommandListImmediate& RHICmdList, FGpuFrameData& ReadFrame);
 
 private:
 	uintptr_t				OwnerContext = 0;
 
-	bool					bProfilingFrame = false;
-	bool					bProfilingDispatches = false;
-
-	bool					bDispatchRecursionGuard = false;
-
-	int32					CurrentReadFrame = 0;
-	int32					CurrentWriteFrame = 0;
+	int32					CurrentReadFrame = 0;					// Index of the next frame to read from
+	int32					CurrentWriteFrame = 0;					// Index of the next frame to write into
 	FGpuFrameData			GpuFrames[NumBufferFrames];
 
+	FGpuFrameData*			ActiveWriteFrame = nullptr;				// Not null while we are generating a frame of data, otherwise null
+	bool					bDispatchRecursionGuard = false;		// We don't support timing dispatches inside one another
+
 	FRenderQueryPoolRHIRef	QueryPool;
 };
 
