- Add missing swizzle on color sample

diff --git a/Engine/Plugins/FX/Niagara/Shaders/Private/NiagaraDataInterfaceStaticMeshTemplate.ush b/Engine/Plugins/FX/Niagara/Shaders/Private/NiagaraDataInterfaceStaticMeshTemplate.ush
--- a/Engine/Plugins/FX/Niagara/Shaders/Private/NiagaraDataInterfaceStaticMeshTemplate.ush
+++ b/Engine/Plugins/FX/Niagara/Shaders/Private/NiagaraDataInterfaceStaticMeshTemplate.ush
@@ -108,7 +108,7 @@ void GetVertexWS_{ParameterName}(int Vertex, out float3 Position, out float3 Vel
 void GetVertexColor_{ParameterName}(int Vertex, out float4 Color)
 {
 	Vertex = clamp(Vertex, 0, NumVertices_{ParameterName} - 1);
-	Color = HasColors_{ParameterName} ? ColorBuffer_{ParameterName}[Vertex] : float4(1, 1, 1, 1);
+	Color = HasColors_{ParameterName} ? ColorBuffer_{ParameterName}[Vertex] FMANUALFETCH_COLOR_COMPONENT_SWIZZLE : float4(1, 1, 1, 1);
 }
 
 void GetVertexUV_{ParameterName}(int Vertex, int UVSet, out float2 UV)
