- Handle individual functions rather than using GetFunctions as the API may change

#rnx

diff --git a/Engine/Plugins/FX/Niagara/Source/Niagara/Private/DataInterface/NiagaraDataInterfaceStaticMesh.cpp b/Engine/Plugins/FX/Niagara/Source/Niagara/Private/DataInterface/NiagaraDataInterfaceStaticMesh.cpp
--- a/Engine/Plugins/FX/Niagara/Source/Niagara/Private/DataInterface/NiagaraDataInterfaceStaticMesh.cpp
+++ b/Engine/Plugins/FX/Niagara/Source/Niagara/Private/DataInterface/NiagaraDataInterfaceStaticMesh.cpp
@@ -1963,24 +1963,28 @@ bool UNiagaraDataInterfaceStaticMesh::UpgradeFunctionCall(FNiagaraFunctionSignat
 				break;
 			}
 		}
-
-		FunctionSignature.FunctionVersion = EDIFunctionVersion::RefactoredV1;
 	}
 
 	if (FunctionSignature.FunctionVersion < EDIFunctionVersion::LargeWorldCoordinates)
 	{
-		TArray<FNiagaraFunctionSignature> AllFunctions;
-		GetFunctions(AllFunctions);
-		for (const FNiagaraFunctionSignature& Sig : AllFunctions)
+		if (
+			(FunctionSignature.Name == GetVertexWSName) ||
+			(FunctionSignature.Name == GetTriangleWSName) ||
+			(FunctionSignature.Name == GetSocketTransformName) ||
+			(FunctionSignature.Name == GetSocketTransformWSName) ||
+			(FunctionSignature.Name == GetFilteredSocketTransformName) ||
+			(FunctionSignature.Name == GetFilteredSocketTransformWSName) ||
+			(FunctionSignature.Name == GetUnfilteredSocketTransformName) ||
+			(FunctionSignature.Name == GetUnfilteredSocketTransformWSName) )
 		{
-			if (FunctionSignature.Name == Sig.Name)
-			{
-				FunctionSignature = Sig;
-				return true;
-			}
+			check(FunctionSignature.Outputs[0].GetName() == TEXT("Position"));
+			check(FunctionSignature.Outputs[0].GetType() == FNiagaraTypeDefinition::GetVec3Def());
+			FunctionSignature.Outputs[0].SetType(FNiagaraTypeDefinition::GetPositionDef());
 		}
 	}
 
+	FunctionSignature.FunctionVersion = EDIFunctionVersion::LatestVersion;
+
 	return true;
 }
 #endif //WITH_EDITORONLY_DATA
diff --git a/Engine/Plugins/FX/Niagara/Source/Niagara/Private/NiagaraDataInterfaceSkeletalMesh.cpp b/Engine/Plugins/FX/Niagara/Source/Niagara/Private/NiagaraDataInterfaceSkeletalMesh.cpp
--- a/Engine/Plugins/FX/Niagara/Source/Niagara/Private/NiagaraDataInterfaceSkeletalMesh.cpp
+++ b/Engine/Plugins/FX/Niagara/Source/Niagara/Private/NiagaraDataInterfaceSkeletalMesh.cpp
@@ -3471,15 +3471,17 @@ bool UNiagaraDataInterfaceSkeletalMesh::UpgradeFunctionCall(FNiagaraFunctionSign
 
 	if (FunctionSignature.FunctionVersion < FNiagaraSkelMeshDIFunctionVersion::LargeWorldCoordinates)
 	{
-		TArray<FNiagaraFunctionSignature> AllFunctions;
-		GetFunctions(AllFunctions);
-		for (const FNiagaraFunctionSignature& Sig : AllFunctions)
-		{
-			if (FunctionSignature.Name == Sig.Name)
-			{
-				FunctionSignature = Sig;
-				return true;
-			}
+		if (
+			(FunctionSignature.Name == FSkeletalMeshInterfaceHelper::GetSkinnedTriangleDataWSInterpName) ||
+			(FunctionSignature.Name == FSkeletalMeshInterfaceHelper::GetSkinnedTriangleDataWSName) ||
+			(FunctionSignature.Name == FSkeletalMeshInterfaceHelper::GetSkinnedVertexDataWSName) ||
+			(FunctionSignature.Name == FSkeletalMeshInterfaceHelper::GetSkinnedBoneDataWSName) ||
+			(FunctionSignature.Name == FSkeletalMeshInterfaceHelper::GetSkinnedBoneDataWSInterpolatedName))
+		{
+			check(FunctionSignature.Outputs[0].GetName() == TEXT("Position"));
+			check(FunctionSignature.Outputs[0].GetType() == FNiagaraTypeDefinition::GetVec3Def());
+			FunctionSignature.Outputs[0].SetType(FNiagaraTypeDefinition::GetPositionDef());
+			bWasChanged = true;
 		}
 	}
 
