- Add better information for bounds with NaNs
- Take an explicit path when we have no calculated, or fixed bounds

#rnx

diff --git a/Engine/Plugins/FX/Niagara/Source/Niagara/Private/NiagaraComponent.cpp b/Engine/Plugins/FX/Niagara/Source/Niagara/Private/NiagaraComponent.cpp
--- a/Engine/Plugins/FX/Niagara/Source/Niagara/Private/NiagaraComponent.cpp
+++ b/Engine/Plugins/FX/Niagara/Source/Niagara/Private/NiagaraComponent.cpp
@@ -1898,18 +1898,29 @@ FBoxSphereBounds UNiagaraComponent::CalcBounds(const FTransform& LocalToWorld) c
 	FBoxSphereBounds SystemBounds;
 	if (CurrLocalBounds.IsValid)
 	{
+		ensureMsgf(!CurrLocalBounds.Min.ContainsNaN() && !CurrLocalBounds.Max.ContainsNaN(), TEXT("CurrLocalBounds contains NaN for System(%d)"), *GetNameSafe(Asset));
+
 		SystemBounds = CurrLocalBounds;
 		SystemBounds.BoxExtent *= BoundsScale;
 		SystemBounds.SphereRadius *= BoundsScale;
 	}
 	else
 	{
-		FBox SimBounds = SystemFixedBounds;
-		if (!SimBounds.IsValid && Asset && Asset->bFixedBounds)
+		if (SystemFixedBounds.IsValid)
+		{
+			ensureMsgf(!SystemFixedBounds.Min.ContainsNaN() && !SystemFixedBounds.Max.ContainsNaN(), TEXT("SystemFixedBounds contains NaN for System(%d)"), *GetNameSafe(Asset));
+			SystemBounds = FBoxSphereBounds(SystemFixedBounds);
+		}
+		else if ( Asset && Asset->bFixedBounds )
 		{
-			SimBounds = Asset->GetFixedBounds();
+			const FBox AssetFixedBounds = Asset->GetFixedBounds();
+			ensureMsgf(!AssetFixedBounds.Min.ContainsNaN() && !AssetFixedBounds.Max.ContainsNaN(), TEXT("AssetFixedBounds contains NaN for System(%d)"), *GetNameSafe(Asset));
+			SystemBounds = FBoxSphereBounds(AssetFixedBounds);
+		}
+		else
+		{
+			SystemBounds = FBoxSphereBounds(FVector::ZeroVector, FVector(1.0f), 1.0f);
 		}
-		SystemBounds = FBoxSphereBounds(SimBounds);
 	}
 
 	return SystemBounds.TransformBy(LocalToWorld);
@@ -1953,7 +1964,7 @@ bool UNiagaraComponent::ResolveOwnerAllowsScalability(bool bRegister)
 				RegisterWithScalabilityManager();
 			}
 		}
-		else if (!bOwnerAllowsScalabiltiy)
+		else
 		{
 			if (!IsActive() && bIsCulledByScalability)
 			{
