- Data interface for reading properties from UObjects
- This allows parameters to be automatically pulled from objects rather than BP needing to push them manually per frame
- An explicit object user parameter binding can be used or a source actor
- If the object is an actor we will look at the root component or for an component of type 'source actor component class' to bind properties to

diff --git a/Engine/Plugins/FX/Niagara/Source/Niagara/Private/DataInterface/NiagaraDataInterfaceUObjectPropertyReader.cpp b/Engine/Plugins/FX/Niagara/Source/Niagara/Private/DataInterface/NiagaraDataInterfaceUObjectPropertyReader.cpp
new file mode 100644
--- /dev/null
+++ b/Engine/Plugins/FX/Niagara/Source/Niagara/Private/DataInterface/NiagaraDataInterfaceUObjectPropertyReader.cpp
@@ -0,0 +1,779 @@
+// Copyright Epic Games, Inc. All Rights Reserved.
+
+#include "NiagaraDataInterfaceUObjectPropertyReader.h"
+#include "NiagaraComputeExecutionContext.h"
+#include "NiagaraFunctionLibrary.h"
+#include "NiagaraSystemInstance.h"
+
+#include "Internationalization/Internationalization.h"
+#include "ShaderParameterUtils.h"
+
+#define LOCTEXT_NAMESPACE "NiagaraDataInterfaceUObjectPropertyReader"
+
+//////////////////////////////////////////////////////////////////////////
+
+namespace NDIUObjectPropertyReaderLocal
+{
+	struct FNDIPropertyGetter
+	{
+		FNiagaraVariableBase		Variable;
+		uint32						DataOffset = INDEX_NONE;
+		TWeakFieldPtr<FProperty>	WeakProperty;
+		TFunction<void(void*)>		PropertyCopyFunction;
+	};
+
+	struct FInstanceData_GameToRender
+	{
+		TArray<uint32>	PropertyOffsets;
+		TArray<uint8>	PropertyData;
+	};
+
+	struct FInstanceData_GameThread
+	{
+		FNiagaraParameterDirectBinding<UObject*>	UObjectBinding;
+		TWeakObjectPtr<UObject>						WeakUObject;
+		TArray<FNDIPropertyGetter>					PropertyGetters;
+		TArray<uint8>								PropertyData;
+		uint32										ChangeId = 0;
+
+		uint32 AddProperty(FNiagaraVariableBase PropertyType)
+		{
+			for ( int i=0; i < PropertyGetters.Num(); ++i )
+			{
+				if ( PropertyGetters[i].Variable == PropertyType )
+				{
+					return i;
+				}
+			}
+
+			FNDIPropertyGetter& NewGetter = PropertyGetters.AddDefaulted_GetRef();
+			NewGetter.Variable = PropertyType;
+			NewGetter.DataOffset = PropertyData.AddZeroed(PropertyType.GetSizeInBytes());
+
+			return PropertyGetters.Num() - 1;
+		}
+	};
+
+	struct FInstanceData_RenderThread
+	{
+		FReadBuffer		PropertyData;
+		TArray<uint32>	PropertyOffsets;
+	};
+
+	struct FNDIProxy : public FNiagaraDataInterfaceProxy
+	{
+		void ProvidePerInstanceDataForRenderThread(void* DataForRenderThread, void* PerInstanceData, const FNiagaraSystemInstanceID& InstanceID)
+		{
+			FInstanceData_GameToRender* InstanceData_ForRT = new(DataForRenderThread) FInstanceData_GameToRender();
+			FInstanceData_GameThread* InstanceData_GT = reinterpret_cast<FInstanceData_GameThread*>(PerInstanceData);
+			if (InstanceData_GT->PropertyData.Num() == 0)
+			{
+				InstanceData_ForRT->PropertyOffsets.AddZeroed(4);
+				InstanceData_ForRT->PropertyOffsets[0] = INDEX_NONE;
+				InstanceData_ForRT->PropertyData.AddZeroed(4);
+			}
+			else
+			{
+				InstanceData_ForRT->PropertyOffsets.AddZeroed(Align(InstanceData_GT->PropertyGetters.Num(), 4));
+				for ( int32 i=0; i < InstanceData_GT->PropertyGetters.Num(); ++i )
+				{
+					const FNDIPropertyGetter& PropertyGetter = InstanceData_GT->PropertyGetters[i];
+					InstanceData_ForRT->PropertyOffsets.Add(PropertyGetter.PropertyCopyFunction ? PropertyGetter.DataOffset : INDEX_NONE);
+				}
+				InstanceData_ForRT->PropertyData = InstanceData_GT->PropertyData;
+			}
+		}
+
+		virtual void ConsumePerInstanceDataFromGameThread(void* PerInstanceData, const FNiagaraSystemInstanceID& InstanceID) override
+		{
+			FInstanceData_GameToRender* InstanceData_FromGT = static_cast<FInstanceData_GameToRender*>(PerInstanceData);
+			if ( FInstanceData_RenderThread* InstanceData_RT = PerInstanceData_RenderThread.Find(InstanceID) )
+			{
+				InstanceData_RT->PropertyOffsets = InstanceData_FromGT->PropertyOffsets;
+				if ( InstanceData_RT->PropertyData.NumBytes != InstanceData_FromGT->PropertyData.Num())
+				{
+					InstanceData_RT->PropertyData.Release();
+					const uint32 NumElements = InstanceData_FromGT->PropertyData.Num() / 4;
+					InstanceData_RT->PropertyData.Initialize(TEXT("NiagaraUObjectPropertyReader"), sizeof(float), NumElements, EPixelFormat::PF_R32_FLOAT);
+				}
+				void* GpuMemory = RHILockBuffer(InstanceData_RT->PropertyData.Buffer, 0, InstanceData_RT->PropertyData.NumBytes, RLM_WriteOnly);
+				FMemory::Memcpy(GpuMemory, InstanceData_FromGT->PropertyData.GetData(), InstanceData_FromGT->PropertyData.Num());
+				RHIUnlockBuffer(InstanceData_RT->PropertyData.Buffer);
+			}
+			InstanceData_FromGT->~FInstanceData_GameToRender();
+		}
+
+		virtual int32 PerInstanceDataPassedToRenderThreadSize() const override
+		{
+			return sizeof(FInstanceData_GameToRender);
+		}
+	
+		virtual void PreStage(FRHICommandList& RHICmdList, const FNiagaraDataInterfaceStageArgs& Context) override {}
+		virtual void PostSimulate(FRHICommandList& RHICmdList, const FNiagaraDataInterfaceArgs& Context) override {}
+
+		TMap<FNiagaraSystemInstanceID, FInstanceData_RenderThread>	PerInstanceData_RenderThread;
+	};
+
+	template<typename TInputType, typename TOutputType>
+	void VMReadData(FVectorVMExternalFunctionContext& Context, const uint32 GetterIndex, const TOutputType DefaultValue)
+	{
+		using namespace NDIUObjectPropertyReaderLocal;
+
+		VectorVM::FUserPtrHandler<FInstanceData_GameThread> InstanceData_GT(Context);
+		FNDIOutputParam<bool> SuccessValue(Context);
+		FNDIOutputParam<TOutputType> OutValue(Context);
+		
+		if (InstanceData_GT->PropertyGetters[GetterIndex].PropertyCopyFunction != nullptr)
+		{
+			const uint32 DataOffset = InstanceData_GT->PropertyGetters[GetterIndex].DataOffset;
+			const TInputType* Value = reinterpret_cast<const TInputType*>(InstanceData_GT->PropertyData.GetData() + DataOffset);
+			for (int32 i = 0; i < Context.GetNumInstances(); ++i)
+			{
+				SuccessValue.SetAndAdvance(true);
+				OutValue.SetAndAdvance(*Value);
+			}
+		}
+		else
+		{
+			for (int32 i = 0; i < Context.GetNumInstances(); ++i)
+			{
+				SuccessValue.SetAndAdvance(false);
+				OutValue.SetAndAdvance(DefaultValue);
+			}
+		}
+	}
+
+	template<typename TType>
+	struct FTypeHelper : std::false_type
+	{
+	};
+
+	template<>
+	struct FTypeHelper<float> : std::true_type
+	{
+		static FName GetFunctionName() { return FName("GetFloatProperty"); }
+		static FNiagaraTypeDefinition GetTypeDef() { return FNiagaraTypeDefinition::GetFloatDef(); }
+		static void VMFunction(FVectorVMExternalFunctionContext& Context, const uint32 GetterIndex) { VMReadData<float, float>(Context, GetterIndex, 0.0f); }
+		static TFunction<void(void*)> GetCopyFunction(const FProperty* InProperty, const void* PropertyAddress)
+		{
+			if (const FFloatProperty* FloatProperty = CastField<const FFloatProperty>(InProperty))
+			{
+				return [PropertyAddress](void* DestAddress) { *reinterpret_cast<float*>(DestAddress) = *reinterpret_cast<const float*>(PropertyAddress); };
+			}
+			else if (const FDoubleProperty* DoubleProperty = CastField<const FDoubleProperty>(InProperty))
+			{
+				return [PropertyAddress](void* DestAddress) { *reinterpret_cast<float*>(DestAddress) = *reinterpret_cast<const double*>(PropertyAddress); };
+			}
+			return nullptr;
+		}
+		static constexpr TCHAR const* HlslBufferType = TEXT("float");
+		static constexpr TCHAR const* HlslBufferRead = TEXT("BUFFER[OFFSET] : 0.0f");
+	};
+
+	template<>
+	struct FTypeHelper<FVector2f> : std::true_type
+	{
+		static FName GetFunctionName() { return FName("GetVec2Property"); }
+		static FNiagaraTypeDefinition GetTypeDef() { return FNiagaraTypeDefinition::GetVec2Def(); }
+		static void VMFunction(FVectorVMExternalFunctionContext& Context, const uint32 GetterIndex) { VMReadData<FVector2f, FVector2f>(Context, GetterIndex, FVector2f::ZeroVector); }
+		static TFunction<void(void*)> GetCopyFunction(const FProperty* InProperty, const void* PropertyAddress)
+		{
+			if (const FStructProperty* StructProperty = CastField<FStructProperty>(InProperty))
+			{
+				if ( StructProperty->Struct == TBaseStructure<FVector2D>::Get() )
+				{
+					return [PropertyAddress](void* DestAddress) { *reinterpret_cast<FVector2f*>(DestAddress) = FVector2f(*reinterpret_cast<const FVector2D*>(PropertyAddress)); };
+				}
+			}
+			return nullptr;
+		}
+		static constexpr TCHAR const* HlslBufferType = TEXT("float2");
+		static constexpr TCHAR const* HlslBufferRead = TEXT("float2(BUFFER[OFFSET + 0], BUFFER[OFFSET + 1]) : float2(0.0f, 0.0f)");
+	};
+
+	template<>
+	struct FTypeHelper<FVector3f> : std::true_type
+	{
+		static FName GetFunctionName() { return FName("GetVec3Property"); }
+		static FNiagaraTypeDefinition GetTypeDef() { return FNiagaraTypeDefinition::GetVec3Def(); }
+		static void VMFunction(FVectorVMExternalFunctionContext& Context, const uint32 GetterIndex) { VMReadData<FVector3f, FVector3f>(Context, GetterIndex, FVector3f::ZeroVector); }
+		static TFunction<void(void*)> GetCopyFunction(const FProperty* InProperty, const void* PropertyAddress)
+		{
+			if (const FStructProperty* StructProperty = CastField<FStructProperty>(InProperty))
+			{
+				if (StructProperty->Struct == TBaseStructure<FVector>::Get())
+				{
+					return [PropertyAddress](void* DestAddress) { *reinterpret_cast<FVector3f*>(DestAddress) = FVector3f(*reinterpret_cast<const FVector*>(PropertyAddress)); };
+				}
+			}
+			return nullptr;
+		}
+		static constexpr TCHAR const* HlslBufferType = TEXT("float3");
+		static constexpr TCHAR const* HlslBufferRead = TEXT("float3(BUFFER[OFFSET + 0], BUFFER[OFFSET + 1], BUFFER[OFFSET + 2]) : float3(0.0f, 0.0f, 0.0f)");
+	};
+
+	template<>
+	struct FTypeHelper<FVector4f> : std::true_type
+	{
+		static FName GetFunctionName() { return FName("GetVec4Property"); }
+		static FNiagaraTypeDefinition GetTypeDef() { return FNiagaraTypeDefinition::GetVec4Def(); }
+		static void VMFunction(FVectorVMExternalFunctionContext& Context, const uint32 GetterIndex) { VMReadData<FVector4f, FVector4f>(Context, GetterIndex, FVector3f::ZeroVector); }
+		static TFunction<void(void*)> GetCopyFunction(const FProperty* InProperty, const void* PropertyAddress)
+		{
+			if (const FStructProperty* StructProperty = CastField<FStructProperty>(InProperty))
+			{
+				if (StructProperty->Struct == TBaseStructure<FVector4>::Get())
+				{
+					return [PropertyAddress](void* DestAddress) { *reinterpret_cast<FVector4f*>(DestAddress) = FVector4f(*reinterpret_cast<const FVector4*>(PropertyAddress)); };
+				}
+			}
+			return nullptr;
+		}
+		static constexpr TCHAR const* HlslBufferType = TEXT("float4");
+		static constexpr TCHAR const* HlslBufferRead = TEXT("float4(BUFFER[OFFSET + 0], BUFFER[OFFSET + 1], BUFFER[OFFSET + 2], BUFFER[OFFSET + 3]) : float4(0.0f, 0.0f, 0.0f, 0.0f)");
+	};
+
+	template<>
+	struct FTypeHelper<FLinearColor> : std::true_type
+	{
+		static FName GetFunctionName() { return FName("GetColorProperty"); }
+		static FNiagaraTypeDefinition GetTypeDef() { return FNiagaraTypeDefinition::GetColorDef(); }
+		static void VMFunction(FVectorVMExternalFunctionContext& Context, const uint32 GetterIndex) { VMReadData<FLinearColor, FLinearColor>(Context, GetterIndex, FVector3f::ZeroVector); }
+		static TFunction<void(void*)> GetCopyFunction(const FProperty* InProperty, const void* PropertyAddress)
+		{
+			if (const FStructProperty* StructProperty = CastField<FStructProperty>(InProperty))
+			{
+				if (StructProperty->Struct == TBaseStructure<FColor>::Get())
+				{
+					return [PropertyAddress](void* DestAddress) { *reinterpret_cast<FLinearColor*>(DestAddress) = FLinearColor(*reinterpret_cast<const FColor*>(PropertyAddress)); };
+				}
+				else if (StructProperty->Struct == TBaseStructure<FLinearColor>::Get())
+				{
+					return [PropertyAddress](void* DestAddress) { *reinterpret_cast<FLinearColor*>(DestAddress) = *reinterpret_cast<const FLinearColor*>(PropertyAddress); };
+				}
+			}
+			return nullptr;
+		}
+		static constexpr TCHAR const* HlslBufferType = TEXT("float4");
+		static constexpr TCHAR const* HlslBufferRead = TEXT("float4(BUFFER[OFFSET + 0], BUFFER[OFFSET + 1], BUFFER[OFFSET + 2], BUFFER[OFFSET + 3]) : float4(0.0f, 0.0f, 0.0f, 0.0f)");
+	};
+
+	template<>
+	struct FTypeHelper<int32> : std::true_type
+	{
+		static FName GetFunctionName() { return FName("GetIntProperty"); }
+		static FNiagaraTypeDefinition GetTypeDef() { return FNiagaraTypeDefinition::GetIntDef(); }
+		static void VMFunction(FVectorVMExternalFunctionContext& Context, const uint32 GetterIndex) { VMReadData<int32, int32>(Context, GetterIndex, 0.0f); }
+		static TFunction<void(void*)> GetCopyFunction(const FProperty* InProperty, const void* PropertyAddress)
+		{
+			if (const FEnumProperty* EnumProperty = CastField<FEnumProperty>(InProperty))
+			{
+				InProperty = EnumProperty->GetUnderlyingProperty();
+			}
+
+			if (InProperty->IsA<const FByteProperty>())
+			{
+				return [PropertyAddress](void* DestAddress) { *reinterpret_cast<int32*>(DestAddress) = *reinterpret_cast<const uint8*>(PropertyAddress); };
+			}
+			else if (InProperty->IsA<const FUInt16Property>())
+			{
+				return [PropertyAddress](void* DestAddress) { *reinterpret_cast<int32*>(DestAddress) = *reinterpret_cast<const uint16*>(PropertyAddress); };
+			}
+			else if (InProperty->IsA<const FUInt32Property>())
+			{
+				return [PropertyAddress](void* DestAddress) { *reinterpret_cast<int32*>(DestAddress) = *reinterpret_cast<const uint32*>(PropertyAddress); };
+			}
+			else if (InProperty->IsA<const FInt16Property>())
+			{
+				return [PropertyAddress](void* DestAddress) { *reinterpret_cast<int32*>(DestAddress) = *reinterpret_cast<const int16*>(PropertyAddress); };
+			}
+			else if (InProperty->IsA<const FIntProperty>())
+			{
+				return [PropertyAddress](void* DestAddress) { *reinterpret_cast<int32*>(DestAddress) = *reinterpret_cast<const int32*>(PropertyAddress); };
+			}
+			return nullptr;
+		}
+		static constexpr TCHAR const* HlslBufferType = TEXT("int");
+		static constexpr TCHAR const* HlslBufferRead = TEXT("asint(BUFFER[OFFSET]) : 0");
+	};
+
+	template<>
+	struct FTypeHelper<bool> : std::true_type
+	{
+		static FName GetFunctionName() { return FName("GetBoolProperty"); }
+		static FNiagaraTypeDefinition GetTypeDef() { return FNiagaraTypeDefinition::GetBoolDef(); }
+		static void VMFunction(FVectorVMExternalFunctionContext& Context, const uint32 GetterIndex) { VMReadData<FNiagaraBool, FNiagaraBool>(Context, GetterIndex, 0.0f); }
+		static TFunction<void(void*)> GetCopyFunction(const FProperty* InProperty, const void* PropertyAddress)
+		{
+			if (InProperty->IsA<const FBoolProperty>())
+			{
+				return [PropertyAddress](void* DestAddress) { *reinterpret_cast<FNiagaraBool*>(DestAddress) = *reinterpret_cast<const bool*>(PropertyAddress); };
+			}
+			return nullptr;
+		}
+		static constexpr TCHAR const* HlslBufferType = TEXT("bool");
+		static constexpr TCHAR const* HlslBufferRead = TEXT("asint(BUFFER[OFFSET]) != 0 : 0");
+	};
+
+	#define NDI_PROPERTY_TYPES \
+		NDI_PROPERTY_TYPE(float) \
+		NDI_PROPERTY_TYPE(FVector2f) \
+		NDI_PROPERTY_TYPE(FVector3f) \
+		NDI_PROPERTY_TYPE(FVector4f) \
+		NDI_PROPERTY_TYPE(FLinearColor) \
+		/*NDI_PROPERTY_TYPE(FNiagaraPosition)*/ \
+		NDI_PROPERTY_TYPE(int32) \
+		NDI_PROPERTY_TYPE(bool) \
+		/*NDI_PROPERTY_TYPE(FQuat4f)*/ \
+		/*NDI_PROPERTY_TYPE(FMatrix44f)*/ \
+
+	void BindPropertyGetter(FNDIPropertyGetter& PropertyGetter, UObject* ObjectBinding)
+	{
+		PropertyGetter.WeakProperty.Reset();
+		PropertyGetter.PropertyCopyFunction = nullptr;
+
+		TStringBuilder<128> VariablePathSB;
+		PropertyGetter.Variable.GetName().ToString(VariablePathSB);
+
+		FStringView VariablePath = VariablePathSB.ToView();
+
+		UStruct* ObjectClass = ObjectBinding->GetClass();
+		void* ObjectAddress = ObjectBinding;
+
+		while ( true )
+		{
+			FStringView PropertyName = VariablePath;
+			int32 DotIndex = INDEX_NONE;
+			if (VariablePath.FindChar('.', DotIndex))
+			{
+				PropertyName = VariablePath.Mid(0, DotIndex);
+				VariablePath = VariablePath.Mid(DotIndex + 1);
+			}
+
+			// Attempt to find property
+			FName PropertyFName(PropertyName, FNAME_Find);
+			FProperty* Property = FindFProperty<FProperty>(ObjectClass, PropertyFName);
+			if ( Property == nullptr )
+			{
+				if ( FNiagaraUtilities::LogVerboseWarnings() )
+				{
+					UE_LOG(LogNiagara, Log, TEXT("Failed attempting to bind property '%s' for object '%s' fullpath '%s'"), *PropertyFName.ToString(), *GetNameSafe(ObjectClass), *PropertyGetter.Variable.GetName().ToString());
+				}
+				return;
+			}
+
+			// Sub property?  Must be a structure of we are done
+			if ( DotIndex != INDEX_NONE )
+			{
+				const FStructProperty* StructProperty = CastField<FStructProperty>(Property);
+				if (StructProperty == nullptr)
+				{
+					return;
+				}
+
+				ObjectClass = StructProperty->Struct;
+				ObjectAddress = StructProperty->ContainerPtrToValuePtr<void>(ObjectAddress);
+				continue;
+			}
+
+			// This is the final property
+			void* PropertyAddress = Property->ContainerPtrToValuePtr<void>(ObjectAddress);
+
+			#define NDI_PROPERTY_TYPE(TYPE) \
+				else if (PropertyGetter.Variable.GetType() == FTypeHelper<TYPE>::GetTypeDef()) { PropertyGetter.PropertyCopyFunction = FTypeHelper<TYPE>::GetCopyFunction(Property, PropertyAddress); }
+
+				if (false) { } NDI_PROPERTY_TYPES
+			#undef NDI_PROPERTY_TYPE
+
+			// Complete
+			return;
+		}
+	}
+}
+
+//////////////////////////////////////////////////////////////////////////
+
+struct FNDIUObjectPropertyReaderCS : public FNiagaraDataInterfaceParametersCS
+{
+	DECLARE_TYPE_LAYOUT(FNDIUObjectPropertyReaderCS, NonVirtual);
+
+public:
+	void Bind(const FNiagaraDataInterfaceGPUParamInfo& ParameterInfo, const class FShaderParameterMap& ParameterMap)
+	{
+		using namespace NDIUObjectPropertyReaderLocal;
+
+		PropertiesToRead.Reserve(ParameterInfo.GeneratedFunctions.Num());
+		for ( const FNiagaraDataInterfaceGeneratedFunction& FunctionInfo : ParameterInfo.GeneratedFunctions )
+		{
+			#define NDI_PROPERTY_TYPE(TYPE) \
+				else if (FunctionInfo.DefinitionName == FTypeHelper<TYPE>::GetFunctionName()) { PropertiesToRead.Emplace(FTypeHelper<TYPE>::GetTypeDef(), FunctionInfo.Specifiers[0].Value); }
+
+				if (false) { } NDI_PROPERTY_TYPES
+			#undef NDI_PROPERTY_TYPE
+		}
+
+		PropertyDataParam.Bind(ParameterMap, *(TEXT("PropertyData_") + ParameterInfo.DataInterfaceHLSLSymbol));
+		PropertyOffsetsParam.Bind(ParameterMap, *(TEXT("PropertyOffsets_") + ParameterInfo.DataInterfaceHLSLSymbol));
+	}
+
+	void Set(FRHICommandList& RHICmdList, const FNiagaraDataInterfaceSetArgs& Context) const
+	{
+		using namespace NDIUObjectPropertyReaderLocal;
+
+		check(IsInRenderingThread());
+
+		FRHIComputeShader* ComputeShaderRHI = Context.Shader.GetComputeShader();
+		const FNDIProxy* DIProxy = static_cast<FNDIProxy*>(Context.DataInterface);
+		const FInstanceData_RenderThread* InstanceData_RT = &DIProxy->PerInstanceData_RenderThread.FindChecked(Context.SystemInstanceID);
+		SetSRVParameter(RHICmdList, ComputeShaderRHI, PropertyDataParam, InstanceData_RT->PropertyData.SRV);
+		SetShaderValueArray(RHICmdList, ComputeShaderRHI, PropertyOffsetsParam, InstanceData_RT->PropertyOffsets.GetData(), InstanceData_RT->PropertyOffsets.Num());
+
+	}
+
+private:
+	LAYOUT_FIELD(TMemoryImageArray<FNiagaraVariableBase>,	PropertiesToRead);
+	LAYOUT_FIELD(FShaderResourceParameter,					PropertyDataParam);
+	LAYOUT_FIELD(FShaderParameter,							PropertyOffsetsParam);
+};
+
+IMPLEMENT_TYPE_LAYOUT(FNDIUObjectPropertyReaderCS);
+IMPLEMENT_NIAGARA_DI_PARAMETER(UNiagaraDataInterfaceUObjectPropertyReader, FNDIUObjectPropertyReaderCS);
+
+//////////////////////////////////////////////////////////////////////////
+
+UNiagaraDataInterfaceUObjectPropertyReader::UNiagaraDataInterfaceUObjectPropertyReader(FObjectInitializer const& ObjectInitializer)
+	: Super(ObjectInitializer)
+{
+	using namespace NDIUObjectPropertyReaderLocal;
+
+	FNiagaraTypeDefinition Def(UObject::StaticClass());
+	UObjectParameterBinding.Parameter.SetType(Def);
+
+	Proxy.Reset(new FNDIProxy());
+}
+
+void UNiagaraDataInterfaceUObjectPropertyReader::PostInitProperties()
+{
+	Super::PostInitProperties();
+
+	//Can we register data interfaces as regular types and fold them into the FNiagaraVariable framework for UI and function calls etc?
+	if (HasAnyFlags(RF_ClassDefaultObject))
+	{
+		ENiagaraTypeRegistryFlags Flags = ENiagaraTypeRegistryFlags::AllowAnyVariable | ENiagaraTypeRegistryFlags::AllowParameter;
+		FNiagaraTypeRegistry::Register(FNiagaraTypeDefinition(GetClass()), Flags);
+	}
+}
+
+bool UNiagaraDataInterfaceUObjectPropertyReader::Equals(const UNiagaraDataInterface* Other) const
+{
+	if (!Super::Equals(Other))
+	{
+		return false;
+	}
+
+	auto* OtherTyped = CastChecked<const UNiagaraDataInterfaceUObjectPropertyReader>(Other);
+	return
+		OtherTyped->UObjectParameterBinding == UObjectParameterBinding &&
+		OtherTyped->PropertyRemap == PropertyRemap &&
+		OtherTyped->SourceActor == SourceActor &&
+		OtherTyped->SourceActorComponentClass == SourceActorComponentClass;
+}
+
+bool UNiagaraDataInterfaceUObjectPropertyReader::CopyToInternal(UNiagaraDataInterface* Destination) const
+{
+	if (!Super::CopyToInternal(Destination))
+	{
+		return false;
+	}
+
+	auto* DestinationTyped = CastChecked<UNiagaraDataInterfaceUObjectPropertyReader>(Destination);
+	DestinationTyped->UObjectParameterBinding = UObjectParameterBinding;
+	DestinationTyped->PropertyRemap = PropertyRemap;
+	DestinationTyped->SourceActor = SourceActor;
+	DestinationTyped->SourceActorComponentClass = SourceActorComponentClass;
+
+	return true;
+}
+
+bool UNiagaraDataInterfaceUObjectPropertyReader::InitPerInstanceData(void* PerInstanceData, FNiagaraSystemInstance* SystemInstance)
+{
+	using namespace NDIUObjectPropertyReaderLocal;
+
+	FInstanceData_GameThread* InstanceData_GT = new(PerInstanceData) FInstanceData_GameThread();
+	InstanceData_GT->UObjectBinding.Init(SystemInstance->GetInstanceParameters(), UObjectParameterBinding.Parameter);
+	InstanceData_GT->ChangeId = ChangeId;
+
+	if ( IsUsedByGPUEmitter() )
+	{
+		// We shouldn't need to do this per init, we should be able to cache once and once only
+		for (const TSharedRef<FNiagaraEmitterInstance, ESPMode::ThreadSafe>& EmitterInstance : SystemInstance->GetEmitters())
+		{
+			if (EmitterInstance->IsDisabled() || EmitterInstance->GetCachedEmitter() == nullptr || EmitterInstance->GetGPUContext() == nullptr)
+			{
+				continue;
+			}
+			const FNiagaraScriptInstanceParameterStore& ParameterStore = EmitterInstance->GetGPUContext()->CombinedParamStore;
+			const TArray<UNiagaraDataInterface*>& DataInterfaces = ParameterStore.GetDataInterfaces();
+			const TArray<FNiagaraDataInterfaceGPUParamInfo>& DataInterfaceParamInfo = EmitterInstance->GetGPUContext()->GPUScript_RT->GetDataInterfaceParamInfo();
+			for ( int32 i=0; i < DataInterfaces.Num(); ++i )
+			{
+				if ( (DataInterfaces[i] == this) && DataInterfaceParamInfo.IsValidIndex(i) )
+				{
+					for (const FNiagaraDataInterfaceGeneratedFunction& FunctionInfo : DataInterfaceParamInfo[i].GeneratedFunctions )
+					{
+						#define NDI_PROPERTY_TYPE(TYPE) \
+							else if (FunctionInfo.DefinitionName == FTypeHelper<TYPE>::GetFunctionName()) { InstanceData_GT->AddProperty(FNiagaraVariableBase(FTypeHelper<TYPE>::GetTypeDef(), GetRemappedPropertyName(FunctionInfo.Specifiers[0].Value))); }
+
+							if (false) {} NDI_PROPERTY_TYPES
+						#undef NDI_PROPERTY_TYPE
+					}
+				}
+			}
+		}
+
+		// Initialize render side instance data
+		ENQUEUE_RENDER_COMMAND(NDIUObjectPropertyReader_InitRT)
+		(
+			[Proxy_RT=GetProxyAs<FNDIProxy>(), InstanceID=SystemInstance->GetId()](FRHICommandList& CmdList)
+			{
+				Proxy_RT->PerInstanceData_RenderThread.Add(InstanceID);
+			}
+		);
+	}
+
+	return true;
+}
+
+void UNiagaraDataInterfaceUObjectPropertyReader::DestroyPerInstanceData(void* PerInstanceData, FNiagaraSystemInstance* SystemInstance)
+{
+	using namespace NDIUObjectPropertyReaderLocal;
+
+	FInstanceData_GameThread* InstanceData_GT = reinterpret_cast<FInstanceData_GameThread*>(PerInstanceData);
+	InstanceData_GT->~FInstanceData_GameThread();
+
+	if ( IsUsedByGPUEmitter() )
+	{
+		ENQUEUE_RENDER_COMMAND(NDIUObjectPropertyReader_InitRT)
+		(
+			[Proxy_RT=GetProxyAs<FNDIProxy>(), InstanceID=SystemInstance->GetId()](FRHICommandList& CmdList)
+			{
+				Proxy_RT->PerInstanceData_RenderThread.Remove(InstanceID);
+			}
+		);
+	}
+}
+
+int32 UNiagaraDataInterfaceUObjectPropertyReader::PerInstanceDataSize() const
+{
+	using namespace NDIUObjectPropertyReaderLocal;
+	return sizeof(FInstanceData_GameThread);
+}
+
+void UNiagaraDataInterfaceUObjectPropertyReader::ProvidePerInstanceDataForRenderThread(void* DataForRenderThread, void* PerInstanceData, const FNiagaraSystemInstanceID& InstanceID)
+{
+	using namespace NDIUObjectPropertyReaderLocal;
+	FNDIProxy* DIProxy = GetProxyAs<FNDIProxy>();
+	DIProxy->ProvidePerInstanceDataForRenderThread(DataForRenderThread, PerInstanceData, InstanceID);
+}
+
+bool UNiagaraDataInterfaceUObjectPropertyReader::PerInstanceTick(void* PerInstanceData, FNiagaraSystemInstance* SystemInstance, float DeltaSeconds)
+{
+	using namespace NDIUObjectPropertyReaderLocal;
+
+	check(PerInstanceData && SystemInstance);
+
+	FInstanceData_GameThread* InstanceData_GT = reinterpret_cast<FInstanceData_GameThread*>(PerInstanceData);
+
+	// Do we need to reinitialize the data interface?
+	if ( InstanceData_GT->ChangeId != ChangeId )
+	{
+		return true;
+	}
+
+	// Do we need to rebind our properties?
+	UObject* ObjectBinding = InstanceData_GT->UObjectBinding.GetValue();
+	if ( ObjectBinding == nullptr )
+	{
+		ObjectBinding = SourceActor.Get();
+	}
+
+	if ( InstanceData_GT->WeakUObject.Get() != ObjectBinding )
+	{
+		InstanceData_GT->WeakUObject = ObjectBinding;
+
+		if (ObjectBinding != nullptr)
+		{
+			// If the object we are binding to is an actor we will get the root component / specified component class
+			UObject* ActorComponent = nullptr;
+			if ( AActor* ObjectActor = Cast<AActor>(ObjectBinding) )
+			{
+				if ( SourceActorComponentClass != nullptr )
+				{
+					ActorComponent = ObjectActor->FindComponentByClass(SourceActorComponentClass);
+				}
+				else
+				{
+					ActorComponent = ObjectActor->GetRootComponent();
+				}
+			}
+
+			for (FNDIPropertyGetter& PropertyGetter : InstanceData_GT->PropertyGetters)
+			{
+				BindPropertyGetter(PropertyGetter, ObjectBinding);
+
+				if ( PropertyGetter.PropertyCopyFunction == nullptr )
+				{
+					BindPropertyGetter(PropertyGetter, ActorComponent);
+				}
+			}
+		}
+		else
+		{
+			for (FNDIPropertyGetter& PropertyGetter : InstanceData_GT->PropertyGetters)
+			{
+				PropertyGetter.WeakProperty.Reset();
+				PropertyGetter.PropertyCopyFunction = nullptr;
+			}
+		}
+	}
+
+	// Update our data store as we can not read object's async it's unsafe
+	if (ObjectBinding != nullptr)
+	{
+		for (FNDIPropertyGetter& PropertyGetter : InstanceData_GT->PropertyGetters )
+		{
+			if (PropertyGetter.PropertyCopyFunction != nullptr)
+			{
+				PropertyGetter.PropertyCopyFunction(InstanceData_GT->PropertyData.GetData() + PropertyGetter.DataOffset);
+			}
+		}
+	}
+
+	return false;
+}
+
+void UNiagaraDataInterfaceUObjectPropertyReader::GetFunctions(TArray<FNiagaraFunctionSignature>& OutFunctions)
+{
+	using namespace NDIUObjectPropertyReaderLocal;
+
+	// Reserve space
+	{
+		int32 NumFunctions = 0;
+		#define NDI_PROPERTY_TYPE(TYPE) ++NumFunctions;
+			NDI_PROPERTY_TYPES
+		#undef NDI_PROPERTY_TYPE
+		OutFunctions.Reserve(OutFunctions.Num() + NumFunctions);
+	}
+
+	// Build default signature
+	FNiagaraFunctionSignature DefaultSignature;
+	DefaultSignature.bMemberFunction = true;
+	DefaultSignature.bRequiresContext = false;
+	DefaultSignature.bSupportsCPU = true;
+	DefaultSignature.bSupportsGPU = true;
+	DefaultSignature.Inputs.Emplace(FNiagaraTypeDefinition(GetClass()), TEXT("UObjectReader"));
+	DefaultSignature.Outputs.Emplace(FNiagaraTypeDefinition::GetBoolDef(), TEXT("Success"));
+	DefaultSignature.FunctionSpecifiers.Emplace("PropertyName");
+
+	// Build function list
+	#define NDI_PROPERTY_TYPE(TYPE) \
+		{ \
+			FNiagaraFunctionSignature& Sig = OutFunctions.Add_GetRef(DefaultSignature); \
+			Sig.Name = FTypeHelper<TYPE>::GetFunctionName(); \
+			Sig.Outputs.Emplace(FTypeHelper<TYPE>::GetTypeDef(), TEXT("Value")); \
+			/*Sig.SetDescription(FTypeHelper<TYPE>::GetFunctionDesc());*/ \
+		} \
+
+		NDI_PROPERTY_TYPES
+	#undef NDI_PROPERTY_TYPE
+}
+
+void UNiagaraDataInterfaceUObjectPropertyReader::GetVMExternalFunction(const FVMExternalFunctionBindingInfo& BindingInfo, void* PerInstanceData, FVMExternalFunction& OutFunc)
+{
+	using namespace NDIUObjectPropertyReaderLocal;
+
+	// Ensure we force a rebind
+	FInstanceData_GameThread* InstanceData_GT = reinterpret_cast<FInstanceData_GameThread*>(PerInstanceData);
+	InstanceData_GT->WeakUObject = nullptr;
+
+	#define NDI_PROPERTY_TYPE(TYPE) \
+		else if (BindingInfo.Name == FTypeHelper<TYPE>::GetFunctionName()) \
+		{ \
+			const uint32 GetterIndex = InstanceData_GT->AddProperty(FNiagaraVariableBase(FTypeHelper<TYPE>::GetTypeDef(), GetRemappedPropertyName(BindingInfo.FunctionSpecifiers[0].Value))); \
+			OutFunc = FVMExternalFunction::CreateLambda([GetterIndex](FVectorVMExternalFunctionContext& Context) { FTypeHelper<TYPE>::VMFunction(Context, GetterIndex); }); \
+		} \
+
+		if (false) { } NDI_PROPERTY_TYPES
+	#undef NDI_PROPERTY_TYPE
+}
+
+#if WITH_EDITORONLY_DATA
+//bool UNiagaraDataInterfaceUObjectPropertyReader::AppendCompileHash(FNiagaraCompileHashVisitor* InVisitor) const
+//{
+//	PropertyDataParam.Bind(ParameterMap, *(TEXT("PropertyData_") + ParameterInfo.DataInterfaceHLSLSymbol));
+//	PropertyOffsetsParam.Bind(ParameterMap, *(TEXT("PropertyOffsets_") + ParameterInfo.DataInterfaceHLSLSymbol));
+//}
+
+void UNiagaraDataInterfaceUObjectPropertyReader::GetParameterDefinitionHLSL(const FNiagaraDataInterfaceGPUParamInfo& ParameterInfo, FString& OutHLSL)
+{
+	OutHLSL.Appendf(TEXT("Buffer<float> PropertyData_%s;\n"), *ParameterInfo.DataInterfaceHLSLSymbol);
+	OutHLSL.Appendf(TEXT("uint4 PropertyOffsets_%s[%d];\n"), *ParameterInfo.DataInterfaceHLSLSymbol, FMath::DivideAndRoundUp(ParameterInfo.GeneratedFunctions.Num(), 4));
+}
+
+bool UNiagaraDataInterfaceUObjectPropertyReader::GetFunctionHLSL(const FNiagaraDataInterfaceGPUParamInfo& ParamInfo, const FNiagaraDataInterfaceGeneratedFunction& FunctionInfo, int FunctionInstanceIndex, FString& OutHLSL)
+{
+	using namespace NDIUObjectPropertyReaderLocal;
+
+	const TCHAR* HlslBufferType = nullptr;
+	const TCHAR* HlslBufferRead = nullptr;
+	#define NDI_PROPERTY_TYPE(TYPE) \
+		else if (FunctionInfo.DefinitionName == FTypeHelper<TYPE>::GetFunctionName()) \
+		{ \
+			HlslBufferType = FTypeHelper<TYPE>::HlslBufferType; \
+			HlslBufferRead = FTypeHelper<TYPE>::HlslBufferRead; \
+		} \
+
+		if (false) { } NDI_PROPERTY_TYPES
+	#undef NDI_PROPERTY_TYPE
+
+	if ( HlslBufferType && HlslBufferRead )
+	{
+		FString HlslBufferReadString(HlslBufferRead);
+		HlslBufferReadString.ReplaceInline(TEXT("BUFFER"), *(TEXT("PropertyData_") + ParamInfo.DataInterfaceHLSLSymbol));
+		HlslBufferReadString.ReplaceInline(TEXT("OFFSET"), TEXT("BufferOffset"));
+
+		OutHLSL.Appendf(TEXT("void %s(out bool bSuccess, out %s Value)\n"), *FunctionInfo.InstanceName, HlslBufferType);
+		OutHLSL.Append(TEXT("{\n"));
+		OutHLSL.Appendf(TEXT("	uint BufferOffset = PropertyOffsets_%s[%d][%d];\n"), *ParamInfo.DataInterfaceHLSLSymbol, FunctionInstanceIndex / 4, FunctionInstanceIndex % 4);
+		OutHLSL.Append(TEXT("	bSuccess = BufferOffset != 0xffffffff;\n"));
+		OutHLSL.Appendf(TEXT("	Value = bSuccess ? %s;\n"), *HlslBufferReadString);
+		OutHLSL.Append(TEXT("}\n"));
+		return true;
+	}
+
+	return false;
+}
+#endif
+
+void UNiagaraDataInterfaceUObjectPropertyReader::SetUObjectReaderPropertyRemap(UNiagaraComponent* NiagaraComponent, FName UserParameterName, FName GraphName, FName RemapName)
+{
+	UNiagaraDataInterfaceUObjectPropertyReader* ReaderDI = UNiagaraFunctionLibrary::GetDataInterface<UNiagaraDataInterfaceUObjectPropertyReader>(NiagaraComponent, UserParameterName);
+	if ( ReaderDI == nullptr )
+	{
+		return;
+	}
+
+	FNiagaraUObjectPropertyReaderRemap* RemapEntry = ReaderDI->PropertyRemap.FindByPredicate([GraphName](const FNiagaraUObjectPropertyReaderRemap& Entry) { return Entry.GraphName == GraphName; });
+	if (RemapEntry == nullptr)
+	{
+		RemapEntry = &ReaderDI->PropertyRemap.AddDefaulted_GetRef();
+		RemapEntry->GraphName = GraphName;
+	}
+	RemapEntry->RemapName = RemapName;
+
+	// Notify changed so all user recache
+	++ReaderDI->ChangeId;
+}
+
+#undef NDI_PROPERTY_TYPES
+#undef LOCTEXT_NAMESPACE
diff --git a/Engine/Plugins/FX/Niagara/Source/Niagara/Private/DataInterface/NiagaraDataInterfaceUObjectPropertyReader.h b/Engine/Plugins/FX/Niagara/Source/Niagara/Private/DataInterface/NiagaraDataInterfaceUObjectPropertyReader.h
new file mode 100644
--- /dev/null
+++ b/Engine/Plugins/FX/Niagara/Source/Niagara/Private/DataInterface/NiagaraDataInterfaceUObjectPropertyReader.h
@@ -0,0 +1,107 @@
+// Copyright Epic Games, Inc. All Rights Reserved.
+#pragma once
+
+#include "NiagaraCommon.h"
+#include "NiagaraShared.h"
+#include "NiagaraDataInterface.h"
+#include "NiagaraParameterStore.h"
+#include "NiagaraDataInterfaceUObjectPropertyReader.generated.h"
+
+USTRUCT()
+struct FNiagaraUObjectPropertyReaderRemap
+{
+	GENERATED_BODY()
+
+	UPROPERTY(EditAnywhere, Category = "Remap")
+	FName GraphName = NAME_None;
+
+	UPROPERTY(EditAnywhere, Category = "Remap")
+	FName RemapName = NAME_None;
+
+	bool operator==(const FNiagaraUObjectPropertyReaderRemap& Other) const
+	{
+		return GraphName == Other.GraphName && RemapName == Other.RemapName;
+	}
+};
+
+/**
+Data interface to read properties from UObjects.
+Rather than having BP tick functions that push data into Niagara this data interface will instead pull them.
+*/
+UCLASS(EditInlineNew, Category = "DataInterface", meta=(DisplayName="Object Reader", Experimental))
+class UNiagaraDataInterfaceUObjectPropertyReader : public UNiagaraDataInterface
+{
+	GENERATED_UCLASS_BODY()
+
+public:
+	DECLARE_NIAGARA_DI_PARAMETER();
+
+	/** User parameter Object binding to read properties from. */
+	UPROPERTY(EditAnywhere, Category = "UObjectPropertyReader")
+	FNiagaraUserParameterBinding UObjectParameterBinding;
+
+	UPROPERTY(EditAnywhere, Category = "UObjectPropertyReader")
+	TArray<FNiagaraUObjectPropertyReaderRemap> PropertyRemap;
+
+	/** Optional source actor to use, if the user parameter binding is valid this will be ignored. */
+	UPROPERTY(EditAnywhere, Category = "UObjectPropertyReader")
+	TLazyObjectPtr<AActor> SourceActor;
+
+	/**
+	When an actor is bound as the object we will also search for a component of this type to bind properties to.
+	For example, setting this to a UPointLightComponent when binding properties we will first look at the actor
+	then look for a component of UPointLightComponent and look at properties on that also.
+	If no class is specified here we look at the RootComponent instead.
+	*/
+	UPROPERTY(EditAnywhere, Category = "UObjectPropertyReader")
+	TObjectPtr<UClass> SourceActorComponentClass;
+
+	//UNiagaraDataInterface Interface
+	virtual void PostInitProperties() override;
+	// UObject Interface End
+
+	// UNiagaraDataInterface Interface Begin
+	virtual bool CanExecuteOnTarget(ENiagaraSimTarget Target) const override { return true; }
+
+	virtual bool Equals(const UNiagaraDataInterface* Other) const override;
+	virtual bool CopyToInternal(UNiagaraDataInterface* Destination) const override;
+
+	virtual bool InitPerInstanceData(void* PerInstanceData, FNiagaraSystemInstance* SystemInstance) override;
+	virtual void DestroyPerInstanceData(void* PerInstanceData, FNiagaraSystemInstance* SystemInstance) override;
+	virtual int32 PerInstanceDataSize() const override;
+
+	virtual void ProvidePerInstanceDataForRenderThread(void* DataForRenderThread, void* PerInstanceData, const FNiagaraSystemInstanceID& InstanceID) override;
+
+	virtual bool HasPreSimulateTick() const override { return true; }
+	virtual bool PerInstanceTick(void* PerInstanceData, FNiagaraSystemInstance* SystemInstance, float DeltaSeconds) override;
+
+	virtual void GetFunctions(TArray<FNiagaraFunctionSignature>& OutFunctions) override;
+	virtual void GetVMExternalFunction(const FVMExternalFunctionBindingInfo& BindingInfo, void* InstanceData, FVMExternalFunction& OutFunc) override;
+#if WITH_EDITORONLY_DATA
+	//virtual bool AppendCompileHash(FNiagaraCompileHashVisitor* InVisitor) const override;
+	virtual void GetParameterDefinitionHLSL(const FNiagaraDataInterfaceGPUParamInfo& ParamInfo, FString& OutHLSL) override;
+	virtual bool GetFunctionHLSL(const FNiagaraDataInterfaceGPUParamInfo& ParamInfo, const FNiagaraDataInterfaceGeneratedFunction& FunctionInfo, int FunctionInstanceIndex, FString& OutHLSL) override;
+#endif
+	//UNiagaraDataInterface Interface
+
+	/** Remaps a property reader, where the  */
+	UFUNCTION(BlueprintCallable, Category = Niagara)
+	static void SetUObjectReaderPropertyRemap(UNiagaraComponent* NiagaraComponent, FName UserParameterName, FName GraphName, FName RemapName);
+
+private:
+	//void VMGetProperty(FVectorVMExternalFunctionContext& Context);
+
+	FName GetRemappedPropertyName(FName InName) const
+	{
+		for ( const FNiagaraUObjectPropertyReaderRemap& RemapEntry : PropertyRemap )
+		{
+			if ( RemapEntry.GraphName == InName )
+			{
+				return RemapEntry.RemapName;
+			}
+		}
+		return InName;
+	};
+
+	uint32 ChangeId = 0;
+};
