- Add support for shared static float buffers that data interfaces can push data into

diff --git a/Engine/Plugins/FX/Niagara/Shaders/Private/NiagaraEmitterInstanceShader.usf b/Engine/Plugins/FX/Niagara/Shaders/Private/NiagaraEmitterInstanceShader.usf
--- a/Engine/Plugins/FX/Niagara/Shaders/Private/NiagaraEmitterInstanceShader.usf
+++ b/Engine/Plugins/FX/Niagara/Shaders/Private/NiagaraEmitterInstanceShader.usf
@@ -510,6 +510,8 @@ int rand_int(int x, int Seed1, int Seed2, int Seed3)
 	RWBuffer<float> RWOutputFloat;
 	RWBuffer<half> RWOutputHalf;
 
+	Buffer<float> StaticInputFloat;
+
 	RWBuffer<uint> RWInstanceCounts;
 	uint ReadInstanceCountOffset;
 	uint WriteInstanceCountOffset;
diff --git a/Engine/Plugins/FX/Niagara/Source/Niagara/Classes/NiagaraDataInterface.h b/Engine/Plugins/FX/Niagara/Source/Niagara/Classes/NiagaraDataInterface.h
--- a/Engine/Plugins/FX/Niagara/Source/Niagara/Classes/NiagaraDataInterface.h
+++ b/Engine/Plugins/FX/Niagara/Source/Niagara/Classes/NiagaraDataInterface.h
@@ -291,9 +291,11 @@ public:
 	virtual bool GenerateIterationSourceNamespaceWriteAttributesHLSL(FNiagaraDataInterfaceGPUParamInfo& DIInstanceInfo, const FNiagaraVariable& InIterationSourceVariable, TConstArrayView<FNiagaraVariable> InArguments, TConstArrayView<FNiagaraVariable> InAttributes, TConstArrayView<FString> InAttributeHLSLNames, bool bSpawnOnly, bool bPartialWrites, TArray<FText>& OutErrors, FString& OutHLSL) const { return false; };
 	/** Used by the translator when dealing with signatures that turn into compiler tags to figure out the precise compiler tag. */
 	virtual bool GenerateCompilerTagPrefix(const FNiagaraFunctionSignature& InSignature, FString& OutPrefix) const  { return false; }
-
 #endif
 
+	/** Allows data interfaces to cache any static data that may be shared between instances */
+	virtual void CacheStaticBuffers(struct FNiagaraSystemStaticBuffers& StaticBuffers, const struct FNiagaraScriptDataInterfaceInfo& DataInterfaceInfo, ENiagaraSimTarget SimTarget) {}
+
 	virtual bool NeedsGPUContextInit() const { return false; }
 	virtual bool GPUContextInit(const FNiagaraScriptDataInterfaceCompileInfo& InInfo, void* PerInstanceData, FNiagaraSystemInstance* SystemInstance) const { return false; }
 
diff --git a/Engine/Plugins/FX/Niagara/Source/Niagara/Classes/NiagaraDataInterfaceCurveBase.h b/Engine/Plugins/FX/Niagara/Source/Niagara/Classes/NiagaraDataInterfaceCurveBase.h
--- a/Engine/Plugins/FX/Niagara/Source/Niagara/Classes/NiagaraDataInterfaceCurveBase.h
+++ b/Engine/Plugins/FX/Niagara/Source/Niagara/Classes/NiagaraDataInterfaceCurveBase.h
@@ -20,10 +20,11 @@ struct FNiagaraDataInterfaceProxyCurveBase : public FNiagaraDataInterfaceProxy
 		CurveLUT.Release();
 	}
 
-	float LUTMinTime;
-	float LUTMaxTime;
-	float LUTInvTimeRange;
-	float CurveLUTNumMinusOne;
+	float LUTMinTime = 0.0f;
+	float LUTMaxTime = 0.0f;
+	float LUTInvTimeRange = 0.0f;
+	float CurveLUTNumMinusOne = 0.0f;
+	uint32 LUTOffset = INDEX_NONE;
 	FReadBuffer CurveLUT;
 
 	virtual int32 PerInstanceDataPassedToRenderThreadSize() const override
@@ -60,6 +61,8 @@ protected:
 	UPROPERTY()
 	float LUTNumSamplesMinusOne;
 
+	uint32 LUTOffset = INDEX_NONE;
+
 	/** Remap a sample time for this curve to 0 to 1 between first and last keys for LUT access.*/
 	FORCEINLINE float NormalizeTime(float T) const
 	{
@@ -179,6 +182,7 @@ public:
 	/** Gets information for all of the curves owned by this curve data interface. */
 	virtual void GetCurveData(TArray<FCurveData>& OutCurveData) { }
 
+	virtual void CacheStaticBuffers(struct FNiagaraSystemStaticBuffers& StaticBuffers, const struct FNiagaraScriptDataInterfaceInfo& DataInterfaceInfo, ENiagaraSimTarget SimTarget) override;
 #if WITH_EDITORONLY_DATA
 	virtual void GetParameterDefinitionHLSL(const FNiagaraDataInterfaceGPUParamInfo& ParamInfo, FString& OutHLSL) override;
 #endif
@@ -244,5 +248,6 @@ struct FNiagaraDataInterfaceParametersCS_Curve : public FNiagaraDataInterfacePar
 	LAYOUT_FIELD(FShaderParameter, MaxTime);
 	LAYOUT_FIELD(FShaderParameter, InvTimeRange);
 	LAYOUT_FIELD(FShaderParameter, CurveLUTNumMinusOne);
+	LAYOUT_FIELD(FShaderParameter, LUTOffset);
 	LAYOUT_FIELD(FShaderResourceParameter, CurveLUT);
 };
diff --git a/Engine/Plugins/FX/Niagara/Source/Niagara/Classes/NiagaraSystem.h b/Engine/Plugins/FX/Niagara/Source/Niagara/Classes/NiagaraSystem.h
--- a/Engine/Plugins/FX/Niagara/Source/Niagara/Classes/NiagaraSystem.h
+++ b/Engine/Plugins/FX/Niagara/Source/Niagara/Classes/NiagaraSystem.h
@@ -13,6 +13,7 @@
 #include "NiagaraEmitterInstance.h"
 #include "NiagaraParameterCollection.h"
 #include "NiagaraParameterDefinitionsSubscriber.h"
+#include "NiagaraSystemStaticBuffers.h"
 #include "NiagaraUserRedirectionParameterStore.h"
 #include "Particles/ParticlePerfStats.h"
 #include "Particles/ParticleSystem.h"
@@ -972,10 +973,11 @@ public:
 		return nullptr;
 	}
 
-	bool AllowCullingForLocalPlayers()const{ return bAllowCullingForLocalPlayers; }
+	bool AllowCullingForLocalPlayers() const { return bAllowCullingForLocalPlayers; }
 
-protected:
+	const FNiagaraSystemStaticBuffers* GetStaticBuffers() const { return StaticBuffers.Get(); }
 
+protected:
 	void GenerateStatID()const;
 #if STATS
 	mutable TStatId StatID_GT;
@@ -1009,7 +1011,10 @@ protected:
 
 	/** Total active instances of this system. */
 	int32 ActiveInstances;
+
+	TUniquePtr<FNiagaraSystemStaticBuffers, FNiagaraSystemStaticBuffers::FDeletor> StaticBuffers;
 };
+
 FORCEINLINE void UNiagaraSystem::RegisterActiveInstance()
 {
 	++ActiveInstances;
diff --git a/Engine/Plugins/FX/Niagara/Source/Niagara/Private/NiagaraCommon.cpp b/Engine/Plugins/FX/Niagara/Source/Niagara/Private/NiagaraCommon.cpp
--- a/Engine/Plugins/FX/Niagara/Source/Niagara/Private/NiagaraCommon.cpp
+++ b/Engine/Plugins/FX/Niagara/Source/Niagara/Private/NiagaraCommon.cpp
@@ -1106,6 +1106,13 @@ void FNiagaraUtilities::CollectScriptDataInterfaceParameters(const UObject& Owne
 	}
 }
 
+bool FNiagaraScriptDataInterfaceInfo::IsUserDataInterface() const
+{
+	TStringBuilder<128> NameBuilder;
+	Name.ToString(NameBuilder);
+	return FCString::Strnicmp(NameBuilder.ToString(), TEXT("user."), 5) == 0;
+}
+
 bool FNiagaraScriptDataInterfaceCompileInfo::CanExecuteOnTarget(ENiagaraSimTarget SimTarget) const
 {
 	// Note that this can be called on non-game threads. We ensure that the data interface CDO object is already in existence at application init time.
diff --git a/Engine/Plugins/FX/Niagara/Source/Niagara/Private/NiagaraDataInterfaceCurveBase.cpp b/Engine/Plugins/FX/Niagara/Source/Niagara/Private/NiagaraDataInterfaceCurveBase.cpp
--- a/Engine/Plugins/FX/Niagara/Source/Niagara/Private/NiagaraDataInterfaceCurveBase.cpp
+++ b/Engine/Plugins/FX/Niagara/Source/Niagara/Private/NiagaraDataInterfaceCurveBase.cpp
@@ -10,7 +10,9 @@
 
 #include "NiagaraCustomVersion.h"
 #include "NiagaraEmitter.h"
+#include "NiagaraRenderer.h"
 #include "NiagaraShader.h"
+#include "NiagaraSystemStaticBuffers.h"
 #include "NiagaraTypes.h"
 
 #define LOCTEXT_NAMESPACE "NiagaraDataInterfaceCurveBase"
@@ -25,6 +27,7 @@ void FNiagaraDataInterfaceParametersCS_Curve::Bind(const FNiagaraDataInterfaceGP
 	MaxTime.Bind(ParameterMap, *(TEXT("MaxTime_") + ParameterInfo.DataInterfaceHLSLSymbol));
 	InvTimeRange.Bind(ParameterMap, *(TEXT("InvTimeRange_") + ParameterInfo.DataInterfaceHLSLSymbol));
 	CurveLUTNumMinusOne.Bind(ParameterMap, *(TEXT("CurveLUTNumMinusOne_") + ParameterInfo.DataInterfaceHLSLSymbol));
+	LUTOffset.Bind(ParameterMap, *(TEXT("LUTOffset_") + ParameterInfo.DataInterfaceHLSLSymbol));
 	CurveLUT.Bind(ParameterMap, *(TEXT("CurveLUT_") + ParameterInfo.DataInterfaceHLSLSymbol));
 }
 
@@ -34,17 +37,17 @@ void FNiagaraDataInterfaceParametersCS_Curve::Set(FRHICommandList& RHICmdList, c
 
 	FRHIComputeShader* ComputeShaderRHI = RHICmdList.GetBoundComputeShader();
 	FNiagaraDataInterfaceProxyCurveBase* CurveDI = (FNiagaraDataInterfaceProxyCurveBase*)Context.DataInterface;
-	FReadBuffer& CurveLUTBuffer = CurveDI->CurveLUT;
+
+	ensure((CurveDI->LUTOffset != INDEX_NONE) || CurveDI->CurveLUT.SRV.IsValid());
 
 	SetShaderValue(RHICmdList, ComputeShaderRHI, MinTime, CurveDI->LUTMinTime);
 	SetShaderValue(RHICmdList, ComputeShaderRHI, MaxTime, CurveDI->LUTMaxTime);
 	SetShaderValue(RHICmdList, ComputeShaderRHI, InvTimeRange, CurveDI->LUTInvTimeRange);
 	SetShaderValue(RHICmdList, ComputeShaderRHI, CurveLUTNumMinusOne, CurveDI->CurveLUTNumMinusOne);
-	SetSRVParameter(RHICmdList, ComputeShaderRHI, CurveLUT, CurveLUTBuffer.SRV);
+	SetShaderValue(RHICmdList, ComputeShaderRHI, LUTOffset, CurveDI->LUTOffset);
+	SetSRVParameter(RHICmdList, ComputeShaderRHI, CurveLUT, CurveDI->CurveLUT.SRV.IsValid() ? CurveDI->CurveLUT.SRV.GetReference() : FNiagaraRenderer::GetDummyFloatBuffer());
 }
 
-
-
 float GNiagaraLUTOptimizeThreshold = UNiagaraDataInterfaceCurveBase::DefaultOptimizeThreshold;
 static FAutoConsoleVariableRef CVarNiagaraLUTOptimizeThreshold(
 	TEXT("fx.Niagara.LUT.OptimizeThreshold"),
@@ -324,6 +327,14 @@ bool UNiagaraDataInterfaceCurveBase::Equals(const UNiagaraDataInterface* Other)
 	return bEqual;
 }
 
+void UNiagaraDataInterfaceCurveBase::CacheStaticBuffers(struct FNiagaraSystemStaticBuffers& StaticBuffers, const struct FNiagaraScriptDataInterfaceInfo& DataInterfaceInfo, ENiagaraSimTarget SimTarget)
+{
+	LUTOffset = INDEX_NONE;
+	if (DataInterfaceInfo.IsUserDataInterface() == false && SimTarget == ENiagaraSimTarget::GPUComputeSim)
+	{
+		LUTOffset = StaticBuffers.AddGpuData(ShaderLUT);
+	}
+}
 
 #if WITH_EDITORONLY_DATA
 void UNiagaraDataInterfaceCurveBase::GetParameterDefinitionHLSL(const FNiagaraDataInterfaceGPUParamInfo& ParamInfo, FString& OutHLSL)
@@ -333,19 +344,36 @@ void UNiagaraDataInterfaceCurveBase::GetParameterDefinitionHLSL(const FNiagaraDa
 	FString MaxTimeStr = TEXT("MaxTime_") + ParamInfo.DataInterfaceHLSLSymbol;
 	FString InvTimeRangeStr = TEXT("InvTimeRange_") + ParamInfo.DataInterfaceHLSLSymbol;
 	FString CurveLUTNumMinusOneStr = TEXT("CurveLUTNumMinusOne_") + ParamInfo.DataInterfaceHLSLSymbol;
+	FString LUTOffsetStr = TEXT("LUTOffset_") + ParamInfo.DataInterfaceHLSLSymbol;
 
-	FString BufferName = "CurveLUT_" + ParamInfo.DataInterfaceHLSLSymbol;
-	OutHLSL += TEXT("Buffer<float> ") + BufferName + TEXT(";\n");
+	FString BufferName;
+	if (ParamInfo.IsUserParameter() == true)
+	{
+		BufferName = "CurveLUT_" + ParamInfo.DataInterfaceHLSLSymbol;
+		OutHLSL += TEXT("Buffer<float> ") + BufferName + TEXT(";\n");
+	}
+	else
+	{
+		BufferName = "StaticInputFloat";
+	}
 
 	OutHLSL += TEXT("float ") + MinTimeStr + TEXT(";\n");
 	OutHLSL += TEXT("float ") + MaxTimeStr + TEXT(";\n");
 	OutHLSL += TEXT("float ") + InvTimeRangeStr + TEXT(";\n");
 	OutHLSL += TEXT("float ") + CurveLUTNumMinusOneStr + TEXT(";\n");
+	OutHLSL += TEXT("uint ") + LUTOffsetStr + TEXT(";\n");
 	OutHLSL += TEXT("\n");
 
 	//TODO: Create a Unitiliy/Common funcitons hlsl def shared between all instances of the same data interface class for these.
 	OutHLSL += FString::Printf(TEXT("float TimeToLUTFraction_%s(float T)\n{\n\treturn saturate((T - %s) * %s);\n}\n"), *ParamInfo.DataInterfaceHLSLSymbol, *MinTimeStr, *InvTimeRangeStr);
-	OutHLSL += FString::Printf(TEXT("float SampleCurve_%s(float T)\n{\n\treturn %s[(uint)T];\n}\n"), *ParamInfo.DataInterfaceHLSLSymbol, *BufferName);
+	if (ParamInfo.IsUserParameter() == true)
+	{
+		OutHLSL += FString::Printf(TEXT("float SampleCurve_%s(float T)\n{\n\treturn %s[(uint)T];\n}\n"), *ParamInfo.DataInterfaceHLSLSymbol, *BufferName);
+	}
+	else
+	{
+		OutHLSL += FString::Printf(TEXT("float SampleCurve_%s(float T)\n{\n\treturn %s[%s + (uint)T];\n}\n"), *ParamInfo.DataInterfaceHLSLSymbol, *BufferName, *LUTOffsetStr);
+	}
 	OutHLSL += TEXT("\n");
 }
 #endif
@@ -389,36 +417,39 @@ bool UNiagaraDataInterfaceCurveBase::GetExposedVariableValue(const FNiagaraVaria
 
 void UNiagaraDataInterfaceCurveBase::PushToRenderThreadImpl()
 {
-	FNiagaraDataInterfaceProxyCurveBase* RT_Proxy = GetProxyAs<FNiagaraDataInterfaceProxyCurveBase>();
+	ENQUEUE_RENDER_COMMAND(FUpdateDICurve)(
+		[
+			RT_Proxy=GetProxyAs<FNiagaraDataInterfaceProxyCurveBase>(),
+			RT_LUTMinTime=LUTMinTime,
+			RT_LUTMaxTime=LUTMaxTime,
+			RT_LUTInvTimeRange=LUTInvTimeRange,
+			RT_CurveLUTNumMinusOne=LUTNumSamplesMinusOne,
+			RT_LUTOffset=LUTOffset,
+			RT_ShaderLUT=(LUTOffset == INDEX_NONE) ? ShaderLUT : TArray<float>()
+		](FRHICommandListImmediate& RHICmdList)
+		{
+			RT_Proxy->LUTMinTime = RT_LUTMinTime;
+			RT_Proxy->LUTMaxTime = RT_LUTMaxTime;
+			RT_Proxy->LUTInvTimeRange = RT_LUTInvTimeRange;
+			RT_Proxy->CurveLUTNumMinusOne = RT_CurveLUTNumMinusOne;
+			RT_Proxy->LUTOffset = RT_LUTOffset;
 
-	const float rtLUTMinTime = this->LUTMinTime;
-	const float rtLUTMaxTime = this->LUTMaxTime;
-	const float rtLUTInvTimeRange = this->LUTInvTimeRange;
-	const float rtCurveLUTNumMinusOne = this->LUTNumSamplesMinusOne;
+			DEC_MEMORY_STAT_BY(STAT_NiagaraGPUDataInterfaceMemory, RT_Proxy->CurveLUT.NumBytes);
+			RT_Proxy->CurveLUT.Release();
 
-	// Push Updates to Proxy.
-	ENQUEUE_RENDER_COMMAND(FUpdateDIColorCurve)(
-		[RT_Proxy, rtLUTMinTime, rtLUTMaxTime, rtLUTInvTimeRange, rtCurveLUTNumMinusOne, rtShaderLUT = ShaderLUT](FRHICommandListImmediate& RHICmdList)
-	{
-		RT_Proxy->LUTMinTime = rtLUTMinTime;
-		RT_Proxy->LUTMaxTime = rtLUTMaxTime;
-		RT_Proxy->LUTInvTimeRange = rtLUTInvTimeRange;
-		RT_Proxy->CurveLUTNumMinusOne = rtCurveLUTNumMinusOne;
-
-		DEC_MEMORY_STAT_BY(STAT_NiagaraGPUDataInterfaceMemory, RT_Proxy->CurveLUT.NumBytes);
-		RT_Proxy->CurveLUT.Release();
-
-		check(rtShaderLUT.Num());
-		RT_Proxy->CurveLUT.Initialize(TEXT("CurveLUT"), sizeof(float), rtShaderLUT.Num(), EPixelFormat::PF_R32_FLOAT, BUF_Static);
-		uint32 BufferSize = rtShaderLUT.Num() * sizeof(float);
-		int32 *BufferData = static_cast<int32*>(RHILockBuffer(RT_Proxy->CurveLUT.Buffer, 0, BufferSize, EResourceLockMode::RLM_WriteOnly));
-		FPlatformMemory::Memcpy(BufferData, rtShaderLUT.GetData(), BufferSize);
-		RHIUnlockBuffer(RT_Proxy->CurveLUT.Buffer);
-		INC_MEMORY_STAT_BY(STAT_NiagaraGPUDataInterfaceMemory, RT_Proxy->CurveLUT.NumBytes);
-	});
+			if (RT_ShaderLUT.Num() > 0)
+			{
+				RT_Proxy->CurveLUT.Initialize(TEXT("CurveLUT"), sizeof(float), RT_ShaderLUT.Num(), EPixelFormat::PF_R32_FLOAT, BUF_Static);
+				const uint32 BufferSize = RT_ShaderLUT.Num() * sizeof(float);
+				void* BufferData = RHILockBuffer(RT_Proxy->CurveLUT.Buffer, 0, BufferSize, EResourceLockMode::RLM_WriteOnly);
+				FPlatformMemory::Memcpy(BufferData, RT_ShaderLUT.GetData(), BufferSize);
+				RHIUnlockBuffer(RT_Proxy->CurveLUT.Buffer);
+				INC_MEMORY_STAT_BY(STAT_NiagaraGPUDataInterfaceMemory, RT_Proxy->CurveLUT.NumBytes);
+			}
+		}
+	);
 }
 
-
 #if WITH_EDITOR	
 /** Refreshes and returns the errors detected with the corresponding data, if any.*/
 TArray<FNiagaraDataInterfaceError> UNiagaraDataInterfaceCurveBase::GetErrors()
diff --git a/Engine/Plugins/FX/Niagara/Source/Niagara/Private/NiagaraGpuComputeDispatch.cpp b/Engine/Plugins/FX/Niagara/Source/Niagara/Private/NiagaraGpuComputeDispatch.cpp
--- a/Engine/Plugins/FX/Niagara/Source/Niagara/Private/NiagaraGpuComputeDispatch.cpp
+++ b/Engine/Plugins/FX/Niagara/Source/Niagara/Private/NiagaraGpuComputeDispatch.cpp
@@ -1570,6 +1570,9 @@ void FNiagaraGpuComputeDispatch::DispatchStage(FRHICommandList& RHICmdList, FRHI
 		SetShaderValue(RHICmdList, RHIComputeShader, ComputeShader->SimulationStageNormalizedIterationIndexParam, SimulationStageNormalizedIterationIndex);
 	}
 
+	// Set static input buffers
+	SetSRVParameter(RHICmdList, RHIComputeShader, ComputeShader->StaticInputFloatParam, Tick.SystemGpuComputeProxy->StaticFloatBuffer);
+
 	// Set Input & Output buffers
 	FNiagaraDataBuffer::SetInputShaderParams(RHICmdList, ComputeShader.GetShader(), SimStageData.Source);
 	FNiagaraDataBuffer::SetOutputShaderParams(RHICmdList, ComputeShader.GetShader(), SimStageData.Destination);
diff --git a/Engine/Plugins/FX/Niagara/Source/Niagara/Private/NiagaraSystem.cpp b/Engine/Plugins/FX/Niagara/Source/Niagara/Private/NiagaraSystem.cpp
--- a/Engine/Plugins/FX/Niagara/Source/Niagara/Private/NiagaraSystem.cpp
+++ b/Engine/Plugins/FX/Niagara/Source/Niagara/Private/NiagaraSystem.cpp
@@ -1941,6 +1941,9 @@ void UNiagaraSystem::CacheFromCompiledData()
 
 	bNeedsAsyncOptimize = true;
 
+	// Reset static buffers
+	StaticBuffers.Reset(new FNiagaraSystemStaticBuffers());
+
 	// Cache system data accessors
 	static const FName NAME_System_ExecutionState = "System.ExecutionState";
 	SystemExecutionStateAccessor.Init(SystemDataSet, NAME_System_ExecutionState);
@@ -1986,12 +1989,38 @@ void UNiagaraSystem::CacheFromCompiledData()
 			}
 			NiagaraEmitter->ConditionalPostLoad();
 			NiagaraEmitter->CacheFromCompiledData(DataSetCompiledData);
+
+			// Allow data interfaces to cache static buffers
+			UNiagaraScript* NiagaraScripts[] =
+			{
+				NiagaraEmitter->SpawnScriptProps.Script,
+				NiagaraEmitter->UpdateScriptProps.Script,
+				NiagaraEmitter->GetGPUComputeScript(),
+			};
+			for (UNiagaraScript* NiagaraScript : NiagaraScripts)
+			{
+				if (NiagaraScript == nullptr)
+				{
+					continue;
+				}
+
+				for (const FNiagaraScriptDataInterfaceInfo& DataInterfaceInfo : NiagaraScript->GetCachedDefaultDataInterfaces())
+				{
+					if (UNiagaraDataInterface* DataInterface = DataInterfaceInfo.DataInterface)
+					{
+						DataInterface->CacheStaticBuffers(*StaticBuffers.Get(), DataInterfaceInfo, NiagaraEmitter->SimTarget);
+					}
+				}
+			}
 		}
 		else
 		{
 			EmitterExecutionStateAccessors.AddDefaulted();
 		}
 	}
+
+	// Finalize static buffers
+	StaticBuffers->Finalize();
 }
 
 bool UNiagaraSystem::HasSystemScriptDIsWithPerInstanceData() const
diff --git a/Engine/Plugins/FX/Niagara/Source/Niagara/Private/NiagaraSystemGpuComputeProxy.cpp b/Engine/Plugins/FX/Niagara/Source/Niagara/Private/NiagaraSystemGpuComputeProxy.cpp
--- a/Engine/Plugins/FX/Niagara/Source/Niagara/Private/NiagaraSystemGpuComputeProxy.cpp
+++ b/Engine/Plugins/FX/Niagara/Source/Niagara/Private/NiagaraSystemGpuComputeProxy.cpp
@@ -6,6 +6,7 @@
 #include "NiagaraComputeExecutionContext.h"
 #include "NiagaraGpuComputeDispatchInterface.h"
 #include "NiagaraGPUSystemTick.h"
+#include "NiagaraRenderer.h"
 
 FNiagaraSystemGpuComputeProxy::FNiagaraSystemGpuComputeProxy(FNiagaraSystemInstance* OwnerInstance)
 	: DebugOwnerInstance(OwnerInstance)
@@ -44,6 +45,18 @@ FNiagaraSystemGpuComputeProxy::FNiagaraSystemGpuComputeProxy(FNiagaraSystemInsta
 	{
 		ComputeTickStage = ENiagaraGpuComputeTickStage::PreInitViews;
 	}
+
+	// Static buffers
+	ENQUEUE_RENDER_COMMAND(SetStaticBuffers)(
+		[RT_Context=this, RT_StaticBuffers=OwnerInstance->GetSystem()->GetStaticBuffers()](FRHICommandListImmediate& RHICmdList)
+		{
+			RT_Context->StaticFloatBuffer = RT_StaticBuffers->GetGpuFloatBuffer();
+			if (!RT_Context->StaticFloatBuffer.IsValid())
+			{
+				RT_Context->StaticFloatBuffer = FNiagaraRenderer::GetDummyFloatBuffer();
+			}
+		}
+	);
 }
 
 FNiagaraSystemGpuComputeProxy::~FNiagaraSystemGpuComputeProxy()
diff --git a/Engine/Plugins/FX/Niagara/Source/Niagara/Private/NiagaraSystemStaticBuffers.cpp b/Engine/Plugins/FX/Niagara/Source/Niagara/Private/NiagaraSystemStaticBuffers.cpp
new file mode 100644
--- /dev/null
+++ b/Engine/Plugins/FX/Niagara/Source/Niagara/Private/NiagaraSystemStaticBuffers.cpp
@@ -0,0 +1,65 @@
+// Copyright Epic Games, Inc. All Rights Reserved.
+
+#include "NiagaraSystemStaticBuffers.h"
+
+FNiagaraSystemStaticBuffers::~FNiagaraSystemStaticBuffers()
+{
+	check(IsInRenderingThread());
+	GpuFloatBuffer.Release();
+}
+
+void FNiagaraSystemStaticBuffers::Finalize()
+{
+	check(bFinalized == false);
+	bFinalized = true;
+
+	if (GpuFloatResource.Num() > 0)
+	{
+		ENQUEUE_RENDER_COMMAND(InitBuffers)(
+			[this](FRHICommandListImmediate&)
+			{
+				check(GpuFloatBuffer.NumBytes == 0);
+
+				GpuFloatBuffer.Initialize(TEXT("NiagaraSystemStaticBuffers"), sizeof(float), GpuFloatResource.Num(), EPixelFormat::PF_R32_FLOAT, BUF_None, &GpuFloatResource);
+			}
+		);
+	}
+}
+
+uint32 FNiagaraSystemStaticBuffers::AddCpuData(TConstArrayView<float> InFloatData)
+{
+	check(bFinalized == false);
+
+	// Search for existing
+	for (int32 i = 0; i <= CpuFloatBuffer.Num() - InFloatData.Num(); ++i)
+	{
+		if (FMemory::Memcmp(CpuFloatBuffer.GetData() + i, InFloatData.GetData(), InFloatData.Num() * sizeof(float)) == 0)
+		{
+			return i;
+		}
+	}
+
+	// Add new
+	const uint32 OutIndex = CpuFloatBuffer.AddUninitialized(InFloatData.Num());
+	FMemory::Memcpy(CpuFloatBuffer.GetData() + OutIndex, InFloatData.GetData(), InFloatData.Num() * sizeof(float));
+	return OutIndex;
+}
+
+uint32 FNiagaraSystemStaticBuffers::AddGpuData(TConstArrayView<float> InFloatData)
+{
+	check(bFinalized == false);
+
+	// Search for existing
+	for (int32 i = 0; i <= GpuFloatResource.Num() - InFloatData.Num(); ++i)
+	{
+		if (FMemory::Memcmp(GpuFloatResource.GetData() + i, InFloatData.GetData(), InFloatData.Num() * sizeof(float)) == 0)
+		{
+			return i;
+		}
+	}
+
+	// Add new
+	const uint32 OutIndex = GpuFloatResource.AddUninitialized(InFloatData.Num());
+	FMemory::Memcpy(GpuFloatResource.GetData() + OutIndex, InFloatData.GetData(), InFloatData.Num() * sizeof(float));
+	return OutIndex;
+}
diff --git a/Engine/Plugins/FX/Niagara/Source/Niagara/Public/NiagaraCommon.h b/Engine/Plugins/FX/Niagara/Source/Niagara/Public/NiagaraCommon.h
--- a/Engine/Plugins/FX/Niagara/Source/Niagara/Public/NiagaraCommon.h
+++ b/Engine/Plugins/FX/Niagara/Source/Niagara/Public/NiagaraCommon.h
@@ -582,8 +582,7 @@ public:
 	UPROPERTY()
 	FName RegisteredParameterMapWrite;
 
-	//TODO: Allow data interfaces to own datasets
-	void CopyTo(FNiagaraScriptDataInterfaceInfo* Destination, UObject* Outer) const;
+	bool IsUserDataInterface() const;
 };
 
 USTRUCT()
diff --git a/Engine/Plugins/FX/Niagara/Source/Niagara/Public/NiagaraSystemGpuComputeProxy.h b/Engine/Plugins/FX/Niagara/Source/Niagara/Public/NiagaraSystemGpuComputeProxy.h
--- a/Engine/Plugins/FX/Niagara/Source/Niagara/Public/NiagaraSystemGpuComputeProxy.h
+++ b/Engine/Plugins/FX/Niagara/Source/Niagara/Public/NiagaraSystemGpuComputeProxy.h
@@ -47,6 +47,8 @@ private:
 	uint32										bRequiresViewUniformBuffer : 1;
 	uint32										bRequiresRayTracingScene : 1;
 
+	FShaderResourceViewRHIRef					StaticFloatBuffer;
+
 	TArray<FNiagaraComputeExecutionContext*>	ComputeContexts;
 	TArray<FNiagaraGPUSystemTick>				PendingTicks;
 };
diff --git a/Engine/Plugins/FX/Niagara/Source/Niagara/Public/NiagaraSystemStaticBuffers.h b/Engine/Plugins/FX/Niagara/Source/Niagara/Public/NiagaraSystemStaticBuffers.h
new file mode 100644
--- /dev/null
+++ b/Engine/Plugins/FX/Niagara/Source/Niagara/Public/NiagaraSystemStaticBuffers.h
@@ -0,0 +1,54 @@
+// Copyright Epic Games, Inc. All Rights Reserved.
+
+#pragma once
+
+#include "CoreMinimal.h"
+#include "RHIUtilities.h"
+
+// Static buffers shared between all instances built once on load
+struct FNiagaraSystemStaticBuffers
+{
+public:
+	struct FDeletor
+	{
+		void operator()(FNiagaraSystemStaticBuffers* Ptr) const
+		{
+			ENQUEUE_RENDER_COMMAND(ScriptSafeDelete)(
+				[RT_Ptr=Ptr](FRHICommandListImmediate& RHICmdList)
+				{
+					delete RT_Ptr;
+				}
+			);
+		}
+	};
+
+public:
+	FNiagaraSystemStaticBuffers() {}
+	~FNiagaraSystemStaticBuffers();
+
+	FNiagaraSystemStaticBuffers(const FNiagaraSystemStaticBuffers&) = delete;
+	FNiagaraSystemStaticBuffers(FNiagaraSystemStaticBuffers&&) = delete;
+	FNiagaraSystemStaticBuffers& operator=(const FNiagaraSystemStaticBuffers&) = delete;
+	FNiagaraSystemStaticBuffers& operator=(FNiagaraSystemStaticBuffers&&) = delete;
+
+	// Called to finalize the data (done on load by the system)
+	void Finalize();
+
+	// Adds float data for Cpu access, returns the index into the buffer in elements
+	uint32 AddCpuData(TConstArrayView<float> InFloatData);
+	// Adds float data for Gpu access, returns the index into the buffer in elements
+	uint32 AddGpuData(TConstArrayView<float> InFloatData);
+
+	// Get float buffer for Cpu Access
+	TConstArrayView<float> GetCpuFloatBuffer() const { return MakeArrayView(CpuFloatBuffer); }
+
+	// Get float buffer for Gpu Access
+	FRHIShaderResourceView* GetGpuFloatBuffer() const { check(IsInRenderingThread()); return GpuFloatBuffer.SRV; }
+
+private:
+	bool					bFinalized = false;
+	TArray<float>			CpuFloatBuffer;
+
+	TResourceArray<float>	GpuFloatResource;
+	FReadBuffer				GpuFloatBuffer;
+};
diff --git a/Engine/Plugins/FX/Niagara/Source/NiagaraCore/Private/NiagaraCustomVersion.cpp b/Engine/Plugins/FX/Niagara/Source/NiagaraCore/Private/NiagaraCustomVersion.cpp
--- a/Engine/Plugins/FX/Niagara/Source/NiagaraCore/Private/NiagaraCustomVersion.cpp
+++ b/Engine/Plugins/FX/Niagara/Source/NiagaraCore/Private/NiagaraCustomVersion.cpp
@@ -8,4 +8,4 @@ const FGuid FNiagaraCustomVersion::GUID(0xFCF57AFA, 0x50764283, 0xB9A9E658, 0xFF
 // Register the custom version with core
 FCustomVersionRegistration GRegisterNiagaraCustomVersion(FNiagaraCustomVersion::GUID, FNiagaraCustomVersion::LatestVersion, TEXT("NiagaraVer"));
 
-const FGuid FNiagaraCustomVersion::LatestScriptCompileVersion(0xEF798E15, 0x4D684C92, 0xA6FF404A, 0x64ED9651);
+const FGuid FNiagaraCustomVersion::LatestScriptCompileVersion(0xF353115E, 0x540A7748, 0xAA9B2D94, 0x4DD17DBB);
diff --git a/Engine/Plugins/FX/Niagara/Source/NiagaraShader/Private/NiagaraShader.cpp b/Engine/Plugins/FX/Niagara/Source/NiagaraShader/Private/NiagaraShader.cpp
--- a/Engine/Plugins/FX/Niagara/Source/NiagaraShader/Private/NiagaraShader.cpp
+++ b/Engine/Plugins/FX/Niagara/Source/NiagaraShader/Private/NiagaraShader.cpp
@@ -1197,6 +1197,7 @@ void FNiagaraShader::BindParams(const TArray<FNiagaraDataInterfaceGPUParamInfo>&
 	FloatInputBufferParam.Bind(ParameterMap, TEXT("InputFloat"));
 	IntInputBufferParam.Bind(ParameterMap, TEXT("InputInt"));
 	HalfInputBufferParam.Bind(ParameterMap, TEXT("InputHalf"));
+	StaticInputFloatParam.Bind(ParameterMap, TEXT("StaticInputFloat"));
 	FloatOutputBufferParam.Bind(ParameterMap, TEXT("OutputFloat"));
 	IntOutputBufferParam.Bind(ParameterMap, TEXT("OutputInt"));
 	HalfOutputBufferParam.Bind(ParameterMap, TEXT("OutputHalf"));
@@ -1321,6 +1322,11 @@ void FNiagaraDataInterfaceParamRef::WriteFrozenParameters(FMemoryImageWriter& Wr
 	InParameters.WriteMemoryImageWithDerivedType(Writer, Base->GetComputeParametersTypeDesc());
 }
 
+bool FNiagaraDataInterfaceGPUParamInfo::IsUserParameter() const
+{
+	return DataInterfaceHLSLSymbol.StartsWith(TEXT("User_"));
+}
+
 bool FNiagaraDataInterfaceGPUParamInfo::Serialize(FArchive& Ar)
 {
 	Ar.UsingCustomVersion(FNiagaraCustomVersion::GUID);
diff --git a/Engine/Plugins/FX/Niagara/Source/NiagaraShader/Public/NiagaraShader.h b/Engine/Plugins/FX/Niagara/Source/NiagaraShader/Public/NiagaraShader.h
--- a/Engine/Plugins/FX/Niagara/Source/NiagaraShader/Public/NiagaraShader.h
+++ b/Engine/Plugins/FX/Niagara/Source/NiagaraShader/Public/NiagaraShader.h
@@ -90,6 +90,7 @@ public:
 	LAYOUT_FIELD(FShaderResourceParameter, FloatInputBufferParam);
 	LAYOUT_FIELD(FShaderResourceParameter, IntInputBufferParam);
 	LAYOUT_FIELD(FShaderResourceParameter, HalfInputBufferParam);
+	LAYOUT_FIELD(FShaderResourceParameter, StaticInputFloatParam);
 	LAYOUT_FIELD(FRWShaderParameter, FloatOutputBufferParam);
 	LAYOUT_FIELD(FRWShaderParameter, IntOutputBufferParam);
 	LAYOUT_FIELD(FRWShaderParameter, HalfOutputBufferParam);
diff --git a/Engine/Plugins/FX/Niagara/Source/NiagaraShader/Public/NiagaraShared.h b/Engine/Plugins/FX/Niagara/Source/NiagaraShader/Public/NiagaraShared.h
--- a/Engine/Plugins/FX/Niagara/Source/NiagaraShader/Public/NiagaraShared.h
+++ b/Engine/Plugins/FX/Niagara/Source/NiagaraShader/Public/NiagaraShared.h
@@ -216,6 +216,8 @@ struct NIAGARASHADER_API FNiagaraDataInterfaceGPUParamInfo
 	UPROPERTY()
 	TArray<FNiagaraDataInterfaceGeneratedFunction> GeneratedFunctions;
 
+	bool IsUserParameter() const;
+
 	bool Serialize(FArchive& Ar);
 };
 
