- IWYU fix

#rnx

diff --git a/Engine/Plugins/FX/Niagara/Source/Niagara/Public/NiagaraSystemStaticBuffers.h b/Engine/Plugins/FX/Niagara/Source/Niagara/Public/NiagaraSystemStaticBuffers.h
--- a/Engine/Plugins/FX/Niagara/Source/Niagara/Public/NiagaraSystemStaticBuffers.h
+++ b/Engine/Plugins/FX/Niagara/Source/Niagara/Public/NiagaraSystemStaticBuffers.h
@@ -3,6 +3,9 @@
 #pragma once
 
 #include "CoreMinimal.h"
+#include "RHI.h"
+#include "RenderResource.h"
+#include "RenderingThread.h"
 #include "RHIUtilities.h"
 
 // Static buffers shared between all instances built once on load
