- Add support for position type
- Add support for pulling the component transform when attached to an actor

#rnx

diff --git a/Engine/Plugins/FX/Niagara/Source/Niagara/Private/DataInterface/NiagaraDataInterfaceUObjectPropertyReader.cpp b/Engine/Plugins/FX/Niagara/Source/Niagara/Private/DataInterface/NiagaraDataInterfaceUObjectPropertyReader.cpp
--- a/Engine/Plugins/FX/Niagara/Source/Niagara/Private/DataInterface/NiagaraDataInterfaceUObjectPropertyReader.cpp
+++ b/Engine/Plugins/FX/Niagara/Source/Niagara/Private/DataInterface/NiagaraDataInterfaceUObjectPropertyReader.cpp
@@ -14,18 +14,21 @@
 
 namespace NDIUObjectPropertyReaderLocal
 {
+	static const FName GetComponentTransformName("GetComponentTransform");
+
 	struct FNDIPropertyGetter
 	{
 		FNiagaraVariableBase		Variable;
 		uint32						DataOffset = INDEX_NONE;
 		TWeakFieldPtr<FProperty>	WeakProperty;
-		TFunction<void(void*)>		PropertyCopyFunction;
+		TFunction<void(const FNiagaraLWCConverter&, void*)>		PropertyCopyFunction;
 	};
 
 	struct FInstanceData_GameToRender
 	{
-		TArray<uint32>	PropertyOffsets;
-		TArray<uint8>	PropertyData;
+		TOptional<FTransform>	CachedTransform;
+		TArray<uint32>			PropertyOffsets;
+		TArray<uint8>			PropertyData;
 	};
 
 	struct FInstanceData_GameThread
@@ -36,6 +39,8 @@ namespace NDIUObjectPropertyReaderLocal
 		TArray<uint8>								PropertyData;
 		uint32										ChangeId = 0;
 
+		TOptional<FTransform>						CachedTransform;
+
 		uint32 AddProperty(FNiagaraVariableBase PropertyType)
 		{
 			for ( int i=0; i < PropertyGetters.Num(); ++i )
@@ -56,8 +61,9 @@ namespace NDIUObjectPropertyReaderLocal
 
 	struct FInstanceData_RenderThread
 	{
-		FReadBuffer		PropertyData;
-		TArray<uint32>	PropertyOffsets;
+		TOptional<FTransform>	CachedTransform;
+		FReadBuffer				PropertyData;
+		TArray<uint32>			PropertyOffsets;
 	};
 
 	struct FNDIProxy : public FNiagaraDataInterfaceProxy
@@ -66,6 +72,7 @@ namespace NDIUObjectPropertyReaderLocal
 		{
 			FInstanceData_GameToRender* InstanceData_ForRT = new(DataForRenderThread) FInstanceData_GameToRender();
 			FInstanceData_GameThread* InstanceData_GT = reinterpret_cast<FInstanceData_GameThread*>(PerInstanceData);
+			InstanceData_ForRT->CachedTransform = InstanceData_GT->CachedTransform;
 			if (InstanceData_GT->PropertyData.Num() == 0)
 			{
 				InstanceData_ForRT->PropertyOffsets.AddZeroed(4);
@@ -89,6 +96,7 @@ namespace NDIUObjectPropertyReaderLocal
 			FInstanceData_GameToRender* InstanceData_FromGT = static_cast<FInstanceData_GameToRender*>(PerInstanceData);
 			if ( FInstanceData_RenderThread* InstanceData_RT = PerInstanceData_RenderThread.Find(InstanceID) )
 			{
+				InstanceData_RT->CachedTransform = InstanceData_FromGT->CachedTransform;
 				InstanceData_RT->PropertyOffsets = InstanceData_FromGT->PropertyOffsets;
 				if ( InstanceData_RT->PropertyData.NumBytes != InstanceData_FromGT->PropertyData.Num())
 				{
@@ -154,15 +162,15 @@ namespace NDIUObjectPropertyReaderLocal
 		static FName GetFunctionName() { return FName("GetFloatProperty"); }
 		static FNiagaraTypeDefinition GetTypeDef() { return FNiagaraTypeDefinition::GetFloatDef(); }
 		static void VMFunction(FVectorVMExternalFunctionContext& Context, const uint32 GetterIndex) { VMReadData<float, float>(Context, GetterIndex, 0.0f); }
-		static TFunction<void(void*)> GetCopyFunction(const FProperty* InProperty, const void* PropertyAddress)
+		static TFunction<void(const FNiagaraLWCConverter&, void*)> GetCopyFunction(const FProperty* InProperty, const void* PropertyAddress)
 		{
 			if (const FFloatProperty* FloatProperty = CastField<const FFloatProperty>(InProperty))
 			{
-				return [PropertyAddress](void* DestAddress) { *reinterpret_cast<float*>(DestAddress) = *reinterpret_cast<const float*>(PropertyAddress); };
+				return [PropertyAddress](const FNiagaraLWCConverter& LwcConverter, void* DestAddress) { *reinterpret_cast<float*>(DestAddress) = *reinterpret_cast<const float*>(PropertyAddress); };
 			}
 			else if (const FDoubleProperty* DoubleProperty = CastField<const FDoubleProperty>(InProperty))
 			{
-				return [PropertyAddress](void* DestAddress) { *reinterpret_cast<float*>(DestAddress) = *reinterpret_cast<const double*>(PropertyAddress); };
+				return [PropertyAddress](const FNiagaraLWCConverter& LwcConverter, void* DestAddress) { *reinterpret_cast<float*>(DestAddress) = *reinterpret_cast<const double*>(PropertyAddress); };
 			}
 			return nullptr;
 		}
@@ -176,13 +184,13 @@ namespace NDIUObjectPropertyReaderLocal
 		static FName GetFunctionName() { return FName("GetVec2Property"); }
 		static FNiagaraTypeDefinition GetTypeDef() { return FNiagaraTypeDefinition::GetVec2Def(); }
 		static void VMFunction(FVectorVMExternalFunctionContext& Context, const uint32 GetterIndex) { VMReadData<FVector2f, FVector2f>(Context, GetterIndex, FVector2f::ZeroVector); }
-		static TFunction<void(void*)> GetCopyFunction(const FProperty* InProperty, const void* PropertyAddress)
+		static TFunction<void(const FNiagaraLWCConverter&, void*)> GetCopyFunction(const FProperty* InProperty, const void* PropertyAddress)
 		{
 			if (const FStructProperty* StructProperty = CastField<FStructProperty>(InProperty))
 			{
 				if ( StructProperty->Struct == TBaseStructure<FVector2D>::Get() )
 				{
-					return [PropertyAddress](void* DestAddress) { *reinterpret_cast<FVector2f*>(DestAddress) = FVector2f(*reinterpret_cast<const FVector2D*>(PropertyAddress)); };
+					return [PropertyAddress](const FNiagaraLWCConverter& LwcConverter, void* DestAddress) { *reinterpret_cast<FVector2f*>(DestAddress) = FVector2f(*reinterpret_cast<const FVector2D*>(PropertyAddress)); };
 				}
 			}
 			return nullptr;
@@ -197,13 +205,13 @@ namespace NDIUObjectPropertyReaderLocal
 		static FName GetFunctionName() { return FName("GetVec3Property"); }
 		static FNiagaraTypeDefinition GetTypeDef() { return FNiagaraTypeDefinition::GetVec3Def(); }
 		static void VMFunction(FVectorVMExternalFunctionContext& Context, const uint32 GetterIndex) { VMReadData<FVector3f, FVector3f>(Context, GetterIndex, FVector3f::ZeroVector); }
-		static TFunction<void(void*)> GetCopyFunction(const FProperty* InProperty, const void* PropertyAddress)
+		static TFunction<void(const FNiagaraLWCConverter&, void*)> GetCopyFunction(const FProperty* InProperty, const void* PropertyAddress)
 		{
 			if (const FStructProperty* StructProperty = CastField<FStructProperty>(InProperty))
 			{
 				if (StructProperty->Struct == TBaseStructure<FVector>::Get())
 				{
-					return [PropertyAddress](void* DestAddress) { *reinterpret_cast<FVector3f*>(DestAddress) = FVector3f(*reinterpret_cast<const FVector*>(PropertyAddress)); };
+					return [PropertyAddress](const FNiagaraLWCConverter& LwcConverter, void* DestAddress) { *reinterpret_cast<FVector3f*>(DestAddress) = FVector3f(*reinterpret_cast<const FVector*>(PropertyAddress)); };
 				}
 			}
 			return nullptr;
@@ -218,13 +226,13 @@ namespace NDIUObjectPropertyReaderLocal
 		static FName GetFunctionName() { return FName("GetVec4Property"); }
 		static FNiagaraTypeDefinition GetTypeDef() { return FNiagaraTypeDefinition::GetVec4Def(); }
 		static void VMFunction(FVectorVMExternalFunctionContext& Context, const uint32 GetterIndex) { VMReadData<FVector4f, FVector4f>(Context, GetterIndex, FVector3f::ZeroVector); }
-		static TFunction<void(void*)> GetCopyFunction(const FProperty* InProperty, const void* PropertyAddress)
+		static TFunction<void(const FNiagaraLWCConverter&, void*)> GetCopyFunction(const FProperty* InProperty, const void* PropertyAddress)
 		{
 			if (const FStructProperty* StructProperty = CastField<FStructProperty>(InProperty))
 			{
 				if (StructProperty->Struct == TBaseStructure<FVector4>::Get())
 				{
-					return [PropertyAddress](void* DestAddress) { *reinterpret_cast<FVector4f*>(DestAddress) = FVector4f(*reinterpret_cast<const FVector4*>(PropertyAddress)); };
+					return [PropertyAddress](const FNiagaraLWCConverter& LwcConverter, void* DestAddress) { *reinterpret_cast<FVector4f*>(DestAddress) = FVector4f(*reinterpret_cast<const FVector4*>(PropertyAddress)); };
 				}
 			}
 			return nullptr;
@@ -239,17 +247,17 @@ namespace NDIUObjectPropertyReaderLocal
 		static FName GetFunctionName() { return FName("GetColorProperty"); }
 		static FNiagaraTypeDefinition GetTypeDef() { return FNiagaraTypeDefinition::GetColorDef(); }
 		static void VMFunction(FVectorVMExternalFunctionContext& Context, const uint32 GetterIndex) { VMReadData<FLinearColor, FLinearColor>(Context, GetterIndex, FVector3f::ZeroVector); }
-		static TFunction<void(void*)> GetCopyFunction(const FProperty* InProperty, const void* PropertyAddress)
+		static TFunction<void(const FNiagaraLWCConverter&, void*)> GetCopyFunction(const FProperty* InProperty, const void* PropertyAddress)
 		{
 			if (const FStructProperty* StructProperty = CastField<FStructProperty>(InProperty))
 			{
 				if (StructProperty->Struct == TBaseStructure<FColor>::Get())
 				{
-					return [PropertyAddress](void* DestAddress) { *reinterpret_cast<FLinearColor*>(DestAddress) = FLinearColor(*reinterpret_cast<const FColor*>(PropertyAddress)); };
+					return [PropertyAddress](const FNiagaraLWCConverter& LwcConverter, void* DestAddress) { *reinterpret_cast<FLinearColor*>(DestAddress) = FLinearColor(*reinterpret_cast<const FColor*>(PropertyAddress)); };
 				}
 				else if (StructProperty->Struct == TBaseStructure<FLinearColor>::Get())
 				{
-					return [PropertyAddress](void* DestAddress) { *reinterpret_cast<FLinearColor*>(DestAddress) = *reinterpret_cast<const FLinearColor*>(PropertyAddress); };
+					return [PropertyAddress](const FNiagaraLWCConverter& LwcConverter, void* DestAddress) { *reinterpret_cast<FLinearColor*>(DestAddress) = *reinterpret_cast<const FLinearColor*>(PropertyAddress); };
 				}
 			}
 			return nullptr;
@@ -258,13 +266,34 @@ namespace NDIUObjectPropertyReaderLocal
 		static constexpr TCHAR const* HlslBufferRead = TEXT("float4(BUFFER[OFFSET + 0], BUFFER[OFFSET + 1], BUFFER[OFFSET + 2], BUFFER[OFFSET + 3]) : float4(0.0f, 0.0f, 0.0f, 0.0f)");
 	};
 
+	template<>
+	struct FTypeHelper<FNiagaraPosition> : std::true_type
+	{
+		static FName GetFunctionName() { return FName("GetPositionProperty"); }
+		static FNiagaraTypeDefinition GetTypeDef() { return FNiagaraTypeDefinition::GetPositionDef(); }
+		static void VMFunction(FVectorVMExternalFunctionContext& Context, const uint32 GetterIndex) { VMReadData<FNiagaraPosition, FNiagaraPosition>(Context, GetterIndex, FNiagaraPosition(ForceInit)); }
+		static TFunction<void(const FNiagaraLWCConverter&, void*)> GetCopyFunction(const FProperty* InProperty, const void* PropertyAddress)
+		{
+			if (const FStructProperty* StructProperty = CastField<FStructProperty>(InProperty))
+			{
+				if (StructProperty->Struct == TBaseStructure<FVector>::Get())
+				{
+					return [PropertyAddress](const FNiagaraLWCConverter& LwcConverter, void* DestAddress) { *reinterpret_cast<FNiagaraPosition*>(DestAddress) = LwcConverter.ConvertWorldToSimulationPosition(*reinterpret_cast<const FVector*>(PropertyAddress)); };
+				}
+			}
+			return nullptr;
+		}
+		static constexpr TCHAR const* HlslBufferType = TEXT("float3");
+		static constexpr TCHAR const* HlslBufferRead = TEXT("float3(BUFFER[OFFSET + 0], BUFFER[OFFSET + 1], BUFFER[OFFSET + 2]) : float3(0.0f, 0.0f, 0.0f)");
+	};
+
 	template<>
 	struct FTypeHelper<int32> : std::true_type
 	{
 		static FName GetFunctionName() { return FName("GetIntProperty"); }
 		static FNiagaraTypeDefinition GetTypeDef() { return FNiagaraTypeDefinition::GetIntDef(); }
 		static void VMFunction(FVectorVMExternalFunctionContext& Context, const uint32 GetterIndex) { VMReadData<int32, int32>(Context, GetterIndex, 0.0f); }
-		static TFunction<void(void*)> GetCopyFunction(const FProperty* InProperty, const void* PropertyAddress)
+		static TFunction<void(const FNiagaraLWCConverter&, void*)> GetCopyFunction(const FProperty* InProperty, const void* PropertyAddress)
 		{
 			if (const FEnumProperty* EnumProperty = CastField<FEnumProperty>(InProperty))
 			{
@@ -273,23 +302,23 @@ namespace NDIUObjectPropertyReaderLocal
 
 			if (InProperty->IsA<const FByteProperty>())
 			{
-				return [PropertyAddress](void* DestAddress) { *reinterpret_cast<int32*>(DestAddress) = *reinterpret_cast<const uint8*>(PropertyAddress); };
+				return [PropertyAddress](const FNiagaraLWCConverter& LwcConverter, void* DestAddress) { *reinterpret_cast<int32*>(DestAddress) = *reinterpret_cast<const uint8*>(PropertyAddress); };
 			}
 			else if (InProperty->IsA<const FUInt16Property>())
 			{
-				return [PropertyAddress](void* DestAddress) { *reinterpret_cast<int32*>(DestAddress) = *reinterpret_cast<const uint16*>(PropertyAddress); };
+				return [PropertyAddress](const FNiagaraLWCConverter& LwcConverter, void* DestAddress) { *reinterpret_cast<int32*>(DestAddress) = *reinterpret_cast<const uint16*>(PropertyAddress); };
 			}
 			else if (InProperty->IsA<const FUInt32Property>())
 			{
-				return [PropertyAddress](void* DestAddress) { *reinterpret_cast<int32*>(DestAddress) = *reinterpret_cast<const uint32*>(PropertyAddress); };
+				return [PropertyAddress](const FNiagaraLWCConverter& LwcConverter, void* DestAddress) { *reinterpret_cast<int32*>(DestAddress) = *reinterpret_cast<const uint32*>(PropertyAddress); };
 			}
 			else if (InProperty->IsA<const FInt16Property>())
 			{
-				return [PropertyAddress](void* DestAddress) { *reinterpret_cast<int32*>(DestAddress) = *reinterpret_cast<const int16*>(PropertyAddress); };
+				return [PropertyAddress](const FNiagaraLWCConverter& LwcConverter, void* DestAddress) { *reinterpret_cast<int32*>(DestAddress) = *reinterpret_cast<const int16*>(PropertyAddress); };
 			}
 			else if (InProperty->IsA<const FIntProperty>())
 			{
-				return [PropertyAddress](void* DestAddress) { *reinterpret_cast<int32*>(DestAddress) = *reinterpret_cast<const int32*>(PropertyAddress); };
+				return [PropertyAddress](const FNiagaraLWCConverter& LwcConverter, void* DestAddress) { *reinterpret_cast<int32*>(DestAddress) = *reinterpret_cast<const int32*>(PropertyAddress); };
 			}
 			return nullptr;
 		}
@@ -303,11 +332,11 @@ namespace NDIUObjectPropertyReaderLocal
 		static FName GetFunctionName() { return FName("GetBoolProperty"); }
 		static FNiagaraTypeDefinition GetTypeDef() { return FNiagaraTypeDefinition::GetBoolDef(); }
 		static void VMFunction(FVectorVMExternalFunctionContext& Context, const uint32 GetterIndex) { VMReadData<FNiagaraBool, FNiagaraBool>(Context, GetterIndex, 0.0f); }
-		static TFunction<void(void*)> GetCopyFunction(const FProperty* InProperty, const void* PropertyAddress)
+		static TFunction<void(const FNiagaraLWCConverter&, void*)> GetCopyFunction(const FProperty* InProperty, const void* PropertyAddress)
 		{
 			if (InProperty->IsA<const FBoolProperty>())
 			{
-				return [PropertyAddress](void* DestAddress) { *reinterpret_cast<FNiagaraBool*>(DestAddress) = *reinterpret_cast<const bool*>(PropertyAddress); };
+				return [PropertyAddress](const FNiagaraLWCConverter& LwcConverter, void* DestAddress) { *reinterpret_cast<FNiagaraBool*>(DestAddress) = *reinterpret_cast<const bool*>(PropertyAddress); };
 			}
 			return nullptr;
 		}
@@ -321,7 +350,7 @@ namespace NDIUObjectPropertyReaderLocal
 		NDI_PROPERTY_TYPE(FVector3f) \
 		NDI_PROPERTY_TYPE(FVector4f) \
 		NDI_PROPERTY_TYPE(FLinearColor) \
-		/*NDI_PROPERTY_TYPE(FNiagaraPosition)*/ \
+		NDI_PROPERTY_TYPE(FNiagaraPosition) \
 		NDI_PROPERTY_TYPE(int32) \
 		NDI_PROPERTY_TYPE(bool) \
 		/*NDI_PROPERTY_TYPE(FQuat4f)*/ \
@@ -412,6 +441,7 @@ public:
 			#undef NDI_PROPERTY_TYPE
 		}
 
+		TransformInfoParam.Bind(ParameterMap, *(TEXT("TransformInfo_") + ParameterInfo.DataInterfaceHLSLSymbol));
 		PropertyDataParam.Bind(ParameterMap, *(TEXT("PropertyData_") + ParameterInfo.DataInterfaceHLSLSymbol));
 		PropertyOffsetsParam.Bind(ParameterMap, *(TEXT("PropertyOffsets_") + ParameterInfo.DataInterfaceHLSLSymbol));
 	}
@@ -428,9 +458,25 @@ public:
 		SetSRVParameter(RHICmdList, ComputeShaderRHI, PropertyDataParam, InstanceData_RT->PropertyData.SRV);
 		SetShaderValueArray(RHICmdList, ComputeShaderRHI, PropertyOffsetsParam, InstanceData_RT->PropertyOffsets.GetData(), InstanceData_RT->PropertyOffsets.Num());
 
+		FVector4f TransformData[3];
+		if ( InstanceData_RT->CachedTransform.IsSet() )
+		{
+			const FQuat4f TransformRotation = FQuat4f(InstanceData_RT->CachedTransform->GetRotation());
+			TransformData[0] = FVector4(FVector3f(InstanceData_RT->CachedTransform->GetLocation()), 1.0f);
+			TransformData[1] = FVector4(TransformRotation.X, TransformRotation.Y, TransformRotation.Z, TransformRotation.W);
+			TransformData[2] = FVector4(FVector3f(InstanceData_RT->CachedTransform->GetScale3D()));
+		}
+		else
+		{
+			TransformData[0] = FVector4(0.0f, 0.0f, 0.0f, 0.0f);
+			TransformData[1] = FVector4(0.0f, 0.0f, 0.0f, 1.0f);
+			TransformData[2] = FVector4(1.0f, 1.0f, 1.0f, 0.0f);
+		}
+		SetShaderValueArray(RHICmdList, ComputeShaderRHI, TransformInfoParam, TransformData, UE_ARRAY_COUNT(TransformData));
 	}
 
 private:
+	LAYOUT_FIELD(FShaderParameter,							TransformInfoParam);
 	LAYOUT_FIELD(TMemoryImageArray<FNiagaraVariableBase>,	PropertiesToRead);
 	LAYOUT_FIELD(FShaderResourceParameter,					PropertyDataParam);
 	LAYOUT_FIELD(FShaderParameter,							PropertyOffsetsParam);
@@ -615,6 +661,20 @@ bool UNiagaraDataInterfaceUObjectPropertyReader::PerInstanceTick(void* PerInstan
 				{
 					ActorComponent = ObjectActor->GetRootComponent();
 				}
+
+				USceneComponent* SceneComponent = Cast<USceneComponent>(ActorComponent);
+				InstanceData_GT->CachedTransform = SceneComponent ? SceneComponent->GetComponentToWorld() : ObjectActor->GetTransform();
+			}
+			else
+			{
+				if ( USceneComponent* SceneComponent = Cast<USceneComponent>(ObjectBinding) )
+				{
+					InstanceData_GT->CachedTransform = SceneComponent->GetComponentToWorld();
+				}
+				else
+				{
+					InstanceData_GT->CachedTransform.Reset();
+				}
 			}
 
 			for (FNDIPropertyGetter& PropertyGetter : InstanceData_GT->PropertyGetters)
@@ -640,11 +700,12 @@ bool UNiagaraDataInterfaceUObjectPropertyReader::PerInstanceTick(void* PerInstan
 	// Update our data store as we can not read object's async it's unsafe
 	if (ObjectBinding != nullptr)
 	{
+		FNiagaraLWCConverter LwcConverter = SystemInstance->GetLWCConverter(false);
 		for (FNDIPropertyGetter& PropertyGetter : InstanceData_GT->PropertyGetters )
 		{
 			if (PropertyGetter.PropertyCopyFunction != nullptr)
 			{
-				PropertyGetter.PropertyCopyFunction(InstanceData_GT->PropertyData.GetData() + PropertyGetter.DataOffset);
+				PropertyGetter.PropertyCopyFunction(LwcConverter, InstanceData_GT->PropertyData.GetData() + PropertyGetter.DataOffset);
 			}
 		}
 	}
@@ -675,7 +736,18 @@ void UNiagaraDataInterfaceUObjectPropertyReader::GetFunctions(TArray<FNiagaraFun
 	DefaultSignature.Outputs.Emplace(FNiagaraTypeDefinition::GetBoolDef(), TEXT("Success"));
 	DefaultSignature.FunctionSpecifiers.Emplace("PropertyName");
 
-	// Build function list
+	// Utility functions
+	{
+		FNiagaraFunctionSignature& Sig = OutFunctions.Add_GetRef(DefaultSignature);
+		Sig.Name = GetComponentTransformName;
+		Sig.FunctionSpecifiers.Empty();
+		Sig.Outputs.Emplace(FNiagaraTypeDefinition::GetPositionDef(), TEXT("Position"));
+		Sig.Outputs.Emplace(FNiagaraTypeDefinition::GetQuatDef(), TEXT("Rotation"));
+		Sig.Outputs.Emplace(FNiagaraTypeDefinition::GetVec3Def(), TEXT("Scale"));
+		Sig.SetDescription(LOCTEXT("GetComponentTransformDesc", "If the object we are bound to is an actor it will return root component transform, or the component class we bound to"));
+	}
+
+	// Build property function list
 	#define NDI_PROPERTY_TYPE(TYPE) \
 		{ \
 			FNiagaraFunctionSignature& Sig = OutFunctions.Add_GetRef(DefaultSignature); \
@@ -696,6 +768,11 @@ void UNiagaraDataInterfaceUObjectPropertyReader::GetVMExternalFunction(const FVM
 	FInstanceData_GameThread* InstanceData_GT = reinterpret_cast<FInstanceData_GameThread*>(PerInstanceData);
 	InstanceData_GT->WeakUObject = nullptr;
 
+	if ( BindingInfo.Name == GetComponentTransformName )
+	{
+		OutFunc = FVMExternalFunction::CreateLambda([this](FVectorVMExternalFunctionContext& Context) { VMGetComponentTransform(Context); });
+	}
+	// Bind property functions
 	#define NDI_PROPERTY_TYPE(TYPE) \
 		else if (BindingInfo.Name == FTypeHelper<TYPE>::GetFunctionName()) \
 		{ \
@@ -703,19 +780,14 @@ void UNiagaraDataInterfaceUObjectPropertyReader::GetVMExternalFunction(const FVM
 			OutFunc = FVMExternalFunction::CreateLambda([GetterIndex](FVectorVMExternalFunctionContext& Context) { FTypeHelper<TYPE>::VMFunction(Context, GetterIndex); }); \
 		} \
 
-		if (false) { } NDI_PROPERTY_TYPES
+		NDI_PROPERTY_TYPES
 	#undef NDI_PROPERTY_TYPE
 }
 
 #if WITH_EDITORONLY_DATA
-//bool UNiagaraDataInterfaceUObjectPropertyReader::AppendCompileHash(FNiagaraCompileHashVisitor* InVisitor) const
-//{
-//	PropertyDataParam.Bind(ParameterMap, *(TEXT("PropertyData_") + ParameterInfo.DataInterfaceHLSLSymbol));
-//	PropertyOffsetsParam.Bind(ParameterMap, *(TEXT("PropertyOffsets_") + ParameterInfo.DataInterfaceHLSLSymbol));
-//}
-
 void UNiagaraDataInterfaceUObjectPropertyReader::GetParameterDefinitionHLSL(const FNiagaraDataInterfaceGPUParamInfo& ParameterInfo, FString& OutHLSL)
 {
+	OutHLSL.Appendf(TEXT("float4 TransformInfo_%s[3];\n"), *ParameterInfo.DataInterfaceHLSLSymbol);
 	OutHLSL.Appendf(TEXT("Buffer<float> PropertyData_%s;\n"), *ParameterInfo.DataInterfaceHLSLSymbol);
 	OutHLSL.Appendf(TEXT("uint4 PropertyOffsets_%s[%d];\n"), *ParameterInfo.DataInterfaceHLSLSymbol, FMath::DivideAndRoundUp(ParameterInfo.GeneratedFunctions.Num(), 4));
 }
@@ -726,6 +798,18 @@ bool UNiagaraDataInterfaceUObjectPropertyReader::GetFunctionHLSL(const FNiagaraD
 
 	const TCHAR* HlslBufferType = nullptr;
 	const TCHAR* HlslBufferRead = nullptr;
+
+	if (FunctionInfo.DefinitionName == GetComponentTransformName )
+	{
+		OutHLSL.Appendf(TEXT("void %s(out bool bSuccess, out float3 Position, out float4 Rotation, out float3 Scale)\n"), *FunctionInfo.InstanceName);
+		OutHLSL.Append(TEXT("{\n"));
+		OutHLSL.Appendf(TEXT("	bSuccess = TransformInfo_%s[0].w > 0.0f;\n"), *ParamInfo.DataInterfaceHLSLSymbol);
+		OutHLSL.Appendf(TEXT("	Position = TransformInfo_%s[0].xyz;\n"), *ParamInfo.DataInterfaceHLSLSymbol);
+		OutHLSL.Appendf(TEXT("	Rotation = TransformInfo_%s[1];\n"), *ParamInfo.DataInterfaceHLSLSymbol);
+		OutHLSL.Appendf(TEXT("	Scale = TransformInfo_%s[2].xyz;\n"), *ParamInfo.DataInterfaceHLSLSymbol);
+		OutHLSL.Append(TEXT("}\n"));
+		return true;
+	}
 	#define NDI_PROPERTY_TYPE(TYPE) \
 		else if (FunctionInfo.DefinitionName == FTypeHelper<TYPE>::GetFunctionName()) \
 		{ \
@@ -733,7 +817,7 @@ bool UNiagaraDataInterfaceUObjectPropertyReader::GetFunctionHLSL(const FNiagaraD
 			HlslBufferRead = FTypeHelper<TYPE>::HlslBufferRead; \
 		} \
 
-		if (false) { } NDI_PROPERTY_TYPES
+		NDI_PROPERTY_TYPES
 	#undef NDI_PROPERTY_TYPE
 
 	if ( HlslBufferType && HlslBufferRead )
@@ -775,5 +859,29 @@ void UNiagaraDataInterfaceUObjectPropertyReader::SetUObjectReaderPropertyRemap(U
 	++ReaderDI->ChangeId;
 }
 
+void UNiagaraDataInterfaceUObjectPropertyReader::VMGetComponentTransform(FVectorVMExternalFunctionContext& Context)
+{
+	using namespace NDIUObjectPropertyReaderLocal;
+
+	VectorVM::FUserPtrHandler<FInstanceData_GameThread> InstanceData_GT(Context);
+	FNDIOutputParam<bool> OutValid(Context);
+	FNDIOutputParam<FNiagaraPosition> OutPosition(Context);
+	FNDIOutputParam<FQuat4f> OutRotation(Context);
+	FNDIOutputParam<FVector3f> OutScale(Context);
+
+	const bool bTransformValid			= InstanceData_GT->CachedTransform.IsSet();
+	const FVector3f TransformPosition	= bTransformValid ? FVector3f(InstanceData_GT->CachedTransform->GetLocation()) : FVector3f::ZeroVector;
+	const FQuat4f TransformRotation		= bTransformValid ? FQuat4f(InstanceData_GT->CachedTransform->GetRotation()) : FQuat4f::Identity;
+	const FVector3f TransformScale		= bTransformValid ? FVector3f(InstanceData_GT->CachedTransform->GetScale3D()) : FVector3f::OneVector;
+
+	for (int32 i = 0; i < Context.GetNumInstances(); ++i)
+	{
+		OutValid.SetAndAdvance(bTransformValid);
+		OutPosition.SetAndAdvance(TransformPosition);
+		OutRotation.SetAndAdvance(TransformRotation);
+		OutScale.SetAndAdvance(TransformScale);
+	}
+}
+
 #undef NDI_PROPERTY_TYPES
 #undef LOCTEXT_NAMESPACE
diff --git a/Engine/Plugins/FX/Niagara/Source/Niagara/Private/DataInterface/NiagaraDataInterfaceUObjectPropertyReader.h b/Engine/Plugins/FX/Niagara/Source/Niagara/Private/DataInterface/NiagaraDataInterfaceUObjectPropertyReader.h
--- a/Engine/Plugins/FX/Niagara/Source/Niagara/Private/DataInterface/NiagaraDataInterfaceUObjectPropertyReader.h
+++ b/Engine/Plugins/FX/Niagara/Source/Niagara/Private/DataInterface/NiagaraDataInterfaceUObjectPropertyReader.h
@@ -89,7 +89,7 @@ public:
 	static void SetUObjectReaderPropertyRemap(UNiagaraComponent* NiagaraComponent, FName UserParameterName, FName GraphName, FName RemapName);
 
 private:
-	//void VMGetProperty(FVectorVMExternalFunctionContext& Context);
+	void VMGetComponentTransform(FVectorVMExternalFunctionContext& Context);
 
 	FName GetRemappedPropertyName(FName InName) const
 	{
