- Cache transform per frame

#rnx

diff --git a/Engine/Plugins/FX/Niagara/Source/Niagara/Private/DataInterface/NiagaraDataInterfaceUObjectPropertyReader.cpp b/Engine/Plugins/FX/Niagara/Source/Niagara/Private/DataInterface/NiagaraDataInterfaceUObjectPropertyReader.cpp
--- a/Engine/Plugins/FX/Niagara/Source/Niagara/Private/DataInterface/NiagaraDataInterfaceUObjectPropertyReader.cpp
+++ b/Engine/Plugins/FX/Niagara/Source/Niagara/Private/DataInterface/NiagaraDataInterfaceUObjectPropertyReader.cpp
@@ -661,20 +661,6 @@ bool UNiagaraDataInterfaceUObjectPropertyReader::PerInstanceTick(void* PerInstan
 				{
 					ActorComponent = ObjectActor->GetRootComponent();
 				}
-
-				USceneComponent* SceneComponent = Cast<USceneComponent>(ActorComponent);
-				InstanceData_GT->CachedTransform = SceneComponent ? SceneComponent->GetComponentToWorld() : ObjectActor->GetTransform();
-			}
-			else
-			{
-				if ( USceneComponent* SceneComponent = Cast<USceneComponent>(ObjectBinding) )
-				{
-					InstanceData_GT->CachedTransform = SceneComponent->GetComponentToWorld();
-				}
-				else
-				{
-					InstanceData_GT->CachedTransform.Reset();
-				}
 			}
 
 			for (FNDIPropertyGetter& PropertyGetter : InstanceData_GT->PropertyGetters)
@@ -698,8 +684,21 @@ bool UNiagaraDataInterfaceUObjectPropertyReader::PerInstanceTick(void* PerInstan
 	}
 
 	// Update our data store as we can not read object's async it's unsafe
+	InstanceData_GT->CachedTransform.Reset();
 	if (ObjectBinding != nullptr)
 	{
+		// Update transform
+		if (AActor* ObjectActor = Cast<AActor>(ObjectBinding))
+		{
+			USceneComponent* ActorComponent = Cast<USceneComponent>(SourceActorComponentClass ? ObjectActor->FindComponentByClass(SourceActorComponentClass) : ObjectActor->GetRootComponent());
+			InstanceData_GT->CachedTransform = ActorComponent ? ActorComponent->GetComponentToWorld() : ObjectActor->GetTransform();
+		}
+		else if ( USceneComponent* SceneComponent = Cast<USceneComponent>(ObjectBinding) )
+		{
+			InstanceData_GT->CachedTransform = SceneComponent->GetComponentToWorld();
+		}
+
+		// Update properties
 		FNiagaraLWCConverter LwcConverter = SystemInstance->GetLWCConverter(false);
 		for (FNDIPropertyGetter& PropertyGetter : InstanceData_GT->PropertyGetters )
 		{
