Improved Niagara pin type colors

diff --git a/Engine/Plugins/FX/Niagara/Source/Niagara/Private/NiagaraSettings.cpp b/Engine/Plugins/FX/Niagara/Source/Niagara/Private/NiagaraSettings.cpp
--- a/Engine/Plugins/FX/Niagara/Source/Niagara/Private/NiagaraSettings.cpp
+++ b/Engine/Plugins/FX/Niagara/Source/Niagara/Private/NiagaraSettings.cpp
@@ -8,7 +8,7 @@ UNiagaraSettings::UNiagaraSettings(const FObjectInitializer& ObjectInitlaizer)
 	, NDISkelMesh_GpuUniformSamplingFormat(ENDISkelMesh_GpuUniformSamplingFormat::Full)
 	, NDISkelMesh_AdjacencyTriangleIndexFormat(ENDISkelMesh_AdjacencyTriangleIndexFormat::Full)
 {
-
+	PositionPinTypeColor = FLinearColor(1.0f, 0.3f, 1.0f, 1.0f);
 }
 
 FName UNiagaraSettings::GetCategoryName() const
diff --git a/Engine/Plugins/FX/Niagara/Source/Niagara/Public/NiagaraSettings.h b/Engine/Plugins/FX/Niagara/Source/Niagara/Public/NiagaraSettings.h
--- a/Engine/Plugins/FX/Niagara/Source/Niagara/Public/NiagaraSettings.h
+++ b/Engine/Plugins/FX/Niagara/Source/Niagara/Public/NiagaraSettings.h
@@ -80,6 +80,10 @@ class NIAGARA_API UNiagaraSettings : public UDeveloperSettings
 	UPROPERTY(config, EditAnywhere, Category = Niagara, meta = (AllowedClasses = "NiagaraEffectType"))
 	FSoftObjectPath DefaultEffectType;
 
+	/** Position pin type color. The other pin colors are defined in the general editor settings. */
+	UPROPERTY(config, EditAnywhere, Category=Niagara)
+	FLinearColor PositionPinTypeColor;
+
 	/** The quality levels Niagara uses. */
 	UPROPERTY(config, EditAnywhere, Category = Scalability)
 	TArray<FText> QualityLevels;
diff --git a/Engine/Plugins/FX/Niagara/Source/NiagaraEditor/Private/EdGraphSchema_Niagara.cpp b/Engine/Plugins/FX/Niagara/Source/NiagaraEditor/Private/EdGraphSchema_Niagara.cpp
--- a/Engine/Plugins/FX/Niagara/Source/NiagaraEditor/Private/EdGraphSchema_Niagara.cpp
+++ b/Engine/Plugins/FX/Niagara/Source/NiagaraEditor/Private/EdGraphSchema_Niagara.cpp
@@ -35,6 +35,7 @@
 #include "NiagaraParameterCollection.h"
 #include "NiagaraScript.h"
 #include "NiagaraScriptSource.h"
+#include "NiagaraSettings.h"
 #include "ObjectEditorUtils.h"
 #include "ScopedTransaction.h"
 #include "ToolMenus.h"
@@ -1343,12 +1344,21 @@ FLinearColor UEdGraphSchema_Niagara::GetPinTypeColor(const FEdGraphPinType& PinT
 	}
 		
 	const UGraphEditorSettings* Settings = GetDefault<UGraphEditorSettings>();
+	if (PinType.PinCategory == PinCategoryEnum || PinType.PinCategory == PinCategoryStaticEnum)
+	{
+		return Settings->BytePinTypeColor;
+	}
+	if (PinType.PinCategory == PinCategoryClass || PinType.PinCategory == PinCategoryStaticClass)
+	{
+		return Settings->ObjectPinTypeColor;
+	}
 	return Settings->WildcardPinTypeColor;
 }
 
 FLinearColor UEdGraphSchema_Niagara::GetTypeColor(const FNiagaraTypeDefinition& Type)
 {
 	const UGraphEditorSettings* Settings = GetDefault<UGraphEditorSettings>();
+	const UNiagaraSettings* NiagaraSettings = GetDefault<UNiagaraSettings>();
 	if (Type == FNiagaraTypeDefinition::GetFloatDef())
 	{
 		return Settings->FloatPinTypeColor;
@@ -1365,10 +1375,22 @@ FLinearColor UEdGraphSchema_Niagara::GetTypeColor(const FNiagaraTypeDefinition&
 	{
 		return Settings->VectorPinTypeColor;
 	}
+	else if (Type == FNiagaraTypeDefinition::GetPositionDef())
+	{
+		return NiagaraSettings->PositionPinTypeColor;
+	}
 	else if (Type == FNiagaraTypeDefinition::GetParameterMapDef())
 	{
 		return Settings->ExecutionPinTypeColor;
 	}
+	else if (Type.IsUObject() || Type.IsDataInterface())
+	{
+		return Settings->ObjectPinTypeColor;
+	}
+	else if (Type.IsEnum())
+	{
+		return Settings->BytePinTypeColor;
+	}
 	else if(Type == FNiagaraTypeDefinition::GetWildcardDef())
 	{
 		return Settings->WildcardPinTypeColor;
