Fixed a bug where setting an engine parameter in the stack that had a default binding in a parameter definition caused a compile error

diff --git a/Engine/Plugins/FX/Niagara/Source/Niagara/Private/NiagaraConstants.cpp b/Engine/Plugins/FX/Niagara/Source/Niagara/Private/NiagaraConstants.cpp
--- a/Engine/Plugins/FX/Niagara/Source/Niagara/Private/NiagaraConstants.cpp
+++ b/Engine/Plugins/FX/Niagara/Source/Niagara/Private/NiagaraConstants.cpp
@@ -336,8 +336,8 @@ void FNiagaraConstants::Init()
 	if (AttrDefaultsStrMap.Num() == 0)
 	{
 		FNiagaraVariable Var;
-		AttrDefaultsStrMap.Add(SYS_PARAM_PARTICLES_POSITION, SYS_PARAM_ENGINE_POSITION.GetName().ToString());
-		AttrDefaultsValueMap.Add(SYS_PARAM_PARTICLES_POSITION, SYS_PARAM_ENGINE_POSITION);
+		AttrDefaultsStrMap.Add(SYS_PARAM_PARTICLES_POSITION, SYS_PARAM_ENGINE_EMITTER_SIMULATION_POSITION.GetName().ToString());
+		AttrDefaultsValueMap.Add(SYS_PARAM_PARTICLES_POSITION, SYS_PARAM_ENGINE_EMITTER_SIMULATION_POSITION);
 		
 		AttrDefaultsStrMap.Add(SYS_PARAM_PARTICLES_VELOCITY, TEXT("0.0,0.0,0.0"));
 		Var = SYS_PARAM_PARTICLES_VELOCITY;
diff --git a/Engine/Plugins/FX/Niagara/Source/NiagaraEditor/Private/NiagaraNodeAssignment.cpp b/Engine/Plugins/FX/Niagara/Source/NiagaraEditor/Private/NiagaraNodeAssignment.cpp
--- a/Engine/Plugins/FX/Niagara/Source/NiagaraEditor/Private/NiagaraNodeAssignment.cpp
+++ b/Engine/Plugins/FX/Niagara/Source/NiagaraEditor/Private/NiagaraNodeAssignment.cpp
@@ -6,14 +6,11 @@
 #include "NiagaraScriptSource.h"
 #include "NiagaraScript.h"
 #include "NiagaraSystem.h"
-#include "NiagaraSystemEditorData.h"
 #include "NiagaraNodeInput.h"
 #include "NiagaraNodeOutput.h"
 #include "EdGraphSchema_Niagara.h"
 #include "Modules/ModuleManager.h"
-#include "AssetRegistryModule.h"
 #include "NiagaraComponent.h"
-#include "NiagaraHlslTranslator.h"
 #include "NiagaraEditorUtilities.h"
 #include "NiagaraNodeParameterMapGet.h"
 #include "NiagaraNodeParameterMapSet.h"
@@ -22,11 +19,8 @@
 #include "Framework/MultiBox/MultiBoxBuilder.h"
 #include "ScopedTransaction.h"
 #include "ViewModels/NiagaraParameterPanelViewModel.h"
-#include "ViewModels/Stack/INiagaraStackItemGroupAddUtilities.h"
 #include "ViewModels/Stack/NiagaraStackGraphUtilities.h"
 #include "ViewModels/TNiagaraViewModelManager.h"
-#include "Widgets/Layout/SBox.h"
-#include "Widgets/SNiagaraParameterName.h"
 #include "NiagaraCustomVersion.h"
 
 #define LOCTEXT_NAMESPACE "NiagaraNodeAssigment"
@@ -474,7 +468,6 @@ void UNiagaraNodeAssignment::InitializeScript(UNiagaraScript* NewScript)
 				{
 					UEdGraphPin* DefaultInputPin = GetNodes[0]->GetDefaultPin(GetPin);
 
-					bool isEngineConstant = false;
 					FNiagaraVariable SeekVar = FNiagaraVariable(Type, FName(*DefaultValue));
 					const FNiagaraVariable* FoundVar = FNiagaraConstants::FindEngineConstant(SeekVar);
 					if (FoundVar != nullptr)
@@ -504,6 +497,12 @@ void UNiagaraNodeAssignment::InitializeScript(UNiagaraScript* NewScript)
 						CreatedGraph->SetMetaData(TargetVar, NewMetaData);
 					}
 				}
+				
+				if (UNiagaraScriptVariable* ScriptVar = CreatedGraph->GetScriptVariable(TargetVar))
+				{
+					ScriptVar->DefaultMode = ENiagaraDefaultMode::Custom;
+					CreatedGraph->ScriptVariableChanged(TargetVar);
+				}
 			}
 		}
 	}
@@ -571,10 +570,7 @@ int32 UNiagaraNodeAssignment::FindAssignmentTarget(const FName& InName)
 {
 	for (int32 i = 0; i < AssignmentTargets.Num(); i++)
 	{
-		FName Name = AssignmentTargets[i].GetName();
-		FNiagaraTypeDefinition Type = AssignmentTargets[i].GetType();
-
-		if (InName == Name)
+		if (InName == AssignmentTargets[i].GetName())
 		{
 			return i;
 		}
