- Add BUF_Static to fix issue with some RHIs

#rnx

diff --git a/Engine/Plugins/FX/Niagara/Source/Niagara/Private/NiagaraSystemStaticBuffers.cpp b/Engine/Plugins/FX/Niagara/Source/Niagara/Private/NiagaraSystemStaticBuffers.cpp
--- a/Engine/Plugins/FX/Niagara/Source/Niagara/Private/NiagaraSystemStaticBuffers.cpp
+++ b/Engine/Plugins/FX/Niagara/Source/Niagara/Private/NiagaraSystemStaticBuffers.cpp
@@ -20,7 +20,7 @@ void FNiagaraSystemStaticBuffers::Finalize()
 			{
 				check(GpuFloatBuffer.NumBytes == 0);
 
-				GpuFloatBuffer.Initialize(TEXT("NiagaraSystemStaticBuffers"), sizeof(float), GpuFloatResource.Num(), EPixelFormat::PF_R32_FLOAT, BUF_None, &GpuFloatResource);
+				GpuFloatBuffer.Initialize(TEXT("NiagaraSystemStaticBuffers"), sizeof(float), GpuFloatResource.Num(), EPixelFormat::PF_R32_FLOAT, BUF_Static, &GpuFloatResource);
 			}
 		);
 	}
