- Added some more information about visibility state to debug hud

#rnx

diff --git a/Engine/Plugins/FX/Niagara/Source/Niagara/Private/NiagaraDebugHud.cpp b/Engine/Plugins/FX/Niagara/Source/Niagara/Private/NiagaraDebugHud.cpp
--- a/Engine/Plugins/FX/Niagara/Source/Niagara/Private/NiagaraDebugHud.cpp
+++ b/Engine/Plugins/FX/Niagara/Source/Niagara/Private/NiagaraDebugHud.cpp
@@ -2361,6 +2361,21 @@ void FNiagaraDebugHud::DrawComponents(FNiagaraWorldManager* WorldManager, UCanva
 					{
 						StringBuilder.Appendf(TEXT("Pooled - %s\n"), *PoolingMethodEnum->GetNameStringByIndex((int32)NiagaraComponent->PoolingMethod));
 					}
+
+					if ( NiagaraComponent->bHiddenInGame || !NiagaraComponent->GetVisibleFlag() || (OwnerActor && OwnerActor->IsHidden()) )
+					{
+						StringBuilder.Appendf(TEXT("HiddenInGame(%d) Visible(%d)"), NiagaraComponent->bHiddenInGame, NiagaraComponent->GetVisibleFlag());
+						if ( OwnerActor )
+						{
+							StringBuilder.Appendf(TEXT(" OwnerActorHidden(%d)"), OwnerActor->IsHidden());
+						}
+						StringBuilder.Append(TEXT("\n"));
+					}
+					if ( !NiagaraComponent->GetVisibleFlag() )
+					{
+						StringBuilder.Append(TEXT("Visibile - false\n"));
+					}
+
 					if (bIsActive)
 					{
 						if (NiagaraComponent->IsRegisteredWithScalabilityManager())
