- Align GPU tick data to 16 bytes
- Data passed needs to be revisited to address alignment / type safety in a better way

#rnx

diff --git a/Engine/Plugins/FX/Niagara/Source/Niagara/Private/NiagaraGPUSystemTick.cpp b/Engine/Plugins/FX/Niagara/Source/Niagara/Private/NiagaraGPUSystemTick.cpp
--- a/Engine/Plugins/FX/Niagara/Source/Niagara/Private/NiagaraGPUSystemTick.cpp
+++ b/Engine/Plugins/FX/Niagara/Source/Niagara/Private/NiagaraGPUSystemTick.cpp
@@ -40,7 +40,7 @@ void FNiagaraGPUSystemTick::Init(FNiagaraSystemInstance* InSystemInstance)
 			FNiagaraDataInterfaceProxy* Proxy = Interface->GetProxy();
 			const int32 Offset = Pair.Value;
 
-			const int32 RTDataSize = Interface->PerInstanceDataPassedToRenderThreadSize();
+			const int32 RTDataSize = Align(Interface->PerInstanceDataPassedToRenderThreadSize(), 16);
 			ensure(RTDataSize > 0);
 			check(Proxy);
 
diff --git a/Engine/Plugins/FX/Niagara/Source/Niagara/Private/NiagaraSystemInstance.cpp b/Engine/Plugins/FX/Niagara/Source/Niagara/Private/NiagaraSystemInstance.cpp
--- a/Engine/Plugins/FX/Niagara/Source/Niagara/Private/NiagaraSystemInstance.cpp
+++ b/Engine/Plugins/FX/Niagara/Source/Niagara/Private/NiagaraSystemInstance.cpp
@@ -1546,7 +1546,7 @@ void FNiagaraSystemInstance::InitDataInterfaces()
 				if (GPUDataSize > 0)
 				{
 					GPUDataInterfaces.Emplace(Interface, Pair.Value);
-					GPUDataInterfaceInstanceDataSize += GPUDataSize;
+					GPUDataInterfaceInstanceDataSize += Align(GPUDataSize, 16);
 				}
 			}
 
