Fix for non-deterministic cook issues related to serialized placeholder data interfaces being different based on whether they were generated from ddc or from an active compilation

diff --git a/Engine/Plugins/FX/Niagara/Source/Niagara/Classes/NiagaraScriptSourceBase.h b/Engine/Plugins/FX/Niagara/Source/Niagara/Classes/NiagaraScriptSourceBase.h
--- a/Engine/Plugins/FX/Niagara/Source/Niagara/Classes/NiagaraScriptSourceBase.h
+++ b/Engine/Plugins/FX/Niagara/Source/Niagara/Classes/NiagaraScriptSourceBase.h
@@ -12,6 +12,7 @@
 
 class INiagaraParameterDefinitionsSubscriber;
 class UNiagaraParameterDefinitionsBase;
+class UNiagaraScriptSourceBase;
 
 struct EditorExposedVectorConstant
 {
@@ -47,6 +48,7 @@ public:
 	virtual bool IsDuplicateDataFor(UNiagaraSystem* InSystem, UNiagaraEmitter* InEmitter, UNiagaraScript* InScript) const = 0;
 	virtual void GetDuplicatedObjects(TArray<UObject*>& Objects) = 0;
 	virtual const TMap<FName, UNiagaraDataInterface*>& GetObjectNameMap() = 0;
+	virtual const UNiagaraScriptSourceBase* GetScriptSource() const =  0;
 	virtual int32 GetDependentRequestCount() const = 0;
 	virtual TSharedPtr<FNiagaraCompileRequestDuplicateDataBase, ESPMode::ThreadSafe> GetDependentRequest(int32 Index) = 0;
 	virtual void ReleaseCompilationCopies() = 0;
diff --git a/Engine/Plugins/FX/Niagara/Source/Niagara/Private/NiagaraSystem.cpp b/Engine/Plugins/FX/Niagara/Source/Niagara/Private/NiagaraSystem.cpp
--- a/Engine/Plugins/FX/Niagara/Source/Niagara/Private/NiagaraSystem.cpp
+++ b/Engine/Plugins/FX/Niagara/Source/Niagara/Private/NiagaraSystem.cpp
@@ -2821,21 +2821,43 @@ bool UNiagaraSystem::QueryCompileComplete(bool bWait, bool bDoPost, bool bDoNotA
 			{
 				TSharedPtr<FNiagaraVMExecutableData> ExeData = EmitterCompiledScriptPair.CompileResults;
 				CompileRequest.CombinedCompileTime += ExeData->CompileTime;
+				UNiagaraScript* CompiledScript = EmitterCompiledScriptPair.CompiledScript;
+
+				// generate the ObjectNameMap from the source (from the duplicated data if available).
 				TMap<FName, UNiagaraDataInterface*> ObjectNameMap;
-				if (AsyncTask->ComputedPrecompileData.IsValid())
+				if (AsyncTask->ComputedPrecompileDuplicateData.IsValid())
 				{
-					ObjectNameMap = AsyncTask->ComputedPrecompileDuplicateData->GetObjectNameMap(); 
+					if (const UNiagaraScriptSourceBase* ScriptSource = AsyncTask->ComputedPrecompileDuplicateData->GetScriptSource())
+					{
+						ObjectNameMap = ScriptSource->ComputeObjectNameMap(*this, CompiledScript->GetUsage(), CompiledScript->GetUsageId(), AsyncTask->UniqueEmitterName);
+					}
+					else
+					{
+						checkf(false, TEXT("Failed to get ScriptSource from ComputedPrecompileDuplicateData"));
+					}
+
+					// we need to replace the DI in the above generated map with the duplicates that we've created as a part of the duplication process
+					TMap<FName, UNiagaraDataInterface*> DuplicatedObjectNameMap = AsyncTask->ComputedPrecompileDuplicateData->GetObjectNameMap();
+					for (auto ObjectMapIt = ObjectNameMap.CreateIterator(); ObjectMapIt; ++ObjectMapIt)
+					{
+						if (UNiagaraDataInterface** Replacement = DuplicatedObjectNameMap.Find(ObjectMapIt.Key()))
+						{
+							ObjectMapIt.Value() = *Replacement;
+						}
+						else
+						{
+							ObjectMapIt.RemoveCurrent();
+						}
+					}
 				}
 				else
 				{
-					UNiagaraScript* CompiledScript = EmitterCompiledScriptPair.CompiledScript;
-					ObjectNameMap = CompiledScript->GetLatestSource()->ComputeObjectNameMap(*this, CompiledScript->GetUsage(), CompiledScript->GetUsageId(), AsyncTask->UniqueEmitterName);
+					const UNiagaraScriptSourceBase* ScriptSource = CompiledScript->GetLatestSource();
+					ObjectNameMap = ScriptSource->ComputeObjectNameMap(*this, CompiledScript->GetUsage(), CompiledScript->GetUsageId(), AsyncTask->UniqueEmitterName);
 				}
+
 				EmitterCompiledScriptPair.CompiledScript->SetVMCompilationResults(EmitterCompiledScriptPair.CompileId, *ExeData, AsyncTask->UniqueEmitterName, ObjectNameMap);
-			}
-			// clean up the precompile data
-			if (AsyncTask->ComputedPrecompileData.IsValid())
-			{
+
 				// Synchronize the variables that we actually encountered during precompile so that we can expose them to the end user.
 				TArray<FNiagaraVariable> OriginalExposedParams;
 				ExposedParameters.GetParameters(OriginalExposedParams);
@@ -2848,10 +2870,17 @@ bool UNiagaraSystem::QueryCompileComplete(bool bWait, bool bDoPost, bool bDoNotA
 						ExposedParameters.AddParameter(EncounteredExposedVars[i], true, false);
 					}
 				}
-				
+			}
+		}
+
+		// clean up the precompile data
+		for (auto& AsyncTask : CompileRequest.DDCTasks)
+		{
+			AsyncTask->ComputedPrecompileData.Reset();
+			if (AsyncTask->ComputedPrecompileDuplicateData.IsValid())
+			{
 				AsyncTask->ComputedPrecompileDuplicateData->ReleaseCompilationCopies();
 				AsyncTask->ComputedPrecompileDuplicateData.Reset();
-				AsyncTask->ComputedPrecompileData.Reset();
 			}
 		}
 
diff --git a/Engine/Plugins/FX/Niagara/Source/NiagaraCore/Private/NiagaraCustomVersion.cpp b/Engine/Plugins/FX/Niagara/Source/NiagaraCore/Private/NiagaraCustomVersion.cpp
--- a/Engine/Plugins/FX/Niagara/Source/NiagaraCore/Private/NiagaraCustomVersion.cpp
+++ b/Engine/Plugins/FX/Niagara/Source/NiagaraCore/Private/NiagaraCustomVersion.cpp
@@ -8,4 +8,4 @@ const FGuid FNiagaraCustomVersion::GUID(0xFCF57AFA, 0x50764283, 0xB9A9E658, 0xFF
 // Register the custom version with core
 FCustomVersionRegistration GRegisterNiagaraCustomVersion(FNiagaraCustomVersion::GUID, FNiagaraCustomVersion::LatestVersion, TEXT("NiagaraVer"));
 
-const FGuid FNiagaraCustomVersion::LatestScriptCompileVersion(0xF353115E, 0x540A7748, 0xAA9B2D94, 0x4DD17DBB);
+const FGuid FNiagaraCustomVersion::LatestScriptCompileVersion(0x5B30E7F5, 0xB45C4092, 0xB60C84DB, 0x8D7C5416);
diff --git a/Engine/Plugins/FX/Niagara/Source/NiagaraEditor/Private/NiagaraCompiler.cpp b/Engine/Plugins/FX/Niagara/Source/NiagaraEditor/Private/NiagaraCompiler.cpp
--- a/Engine/Plugins/FX/Niagara/Source/NiagaraEditor/Private/NiagaraCompiler.cpp
+++ b/Engine/Plugins/FX/Niagara/Source/NiagaraEditor/Private/NiagaraCompiler.cpp
@@ -246,6 +246,11 @@ const TMap<FName, UNiagaraDataInterface*>& FNiagaraCompileRequestDuplicateData::
 	return *SharedNameToDuplicatedDataInterfaceMap.Get();
 }
 
+const UNiagaraScriptSourceBase* FNiagaraCompileRequestDuplicateData::GetScriptSource() const
+{
+	return SourceDeepCopy.Get();
+}
+
 UNiagaraDataInterface* FNiagaraCompileRequestDuplicateData::GetDuplicatedDataInterfaceCDOForClass(UClass* Class) const
 {
 	if (SharedDataInterfaceClassToDuplicatedCDOMap.IsValid())
diff --git a/Engine/Plugins/FX/Niagara/Source/NiagaraEditor/Private/NiagaraHlslTranslator.h b/Engine/Plugins/FX/Niagara/Source/NiagaraEditor/Private/NiagaraHlslTranslator.h
--- a/Engine/Plugins/FX/Niagara/Source/NiagaraEditor/Private/NiagaraHlslTranslator.h
+++ b/Engine/Plugins/FX/Niagara/Source/NiagaraEditor/Private/NiagaraHlslTranslator.h
@@ -140,6 +140,7 @@ public:
 	virtual bool IsDuplicateDataFor(UNiagaraSystem* InSystem, UNiagaraEmitter* InEmitter, UNiagaraScript* InScript) const override;
 	virtual void GetDuplicatedObjects(TArray<UObject*>& Objects) override;
 	virtual const TMap<FName, UNiagaraDataInterface*>& GetObjectNameMap() override;
+	virtual const UNiagaraScriptSourceBase* GetScriptSource() const override;
 	UNiagaraDataInterface* GetDuplicatedDataInterfaceCDOForClass(UClass* Class) const;
 
 	TArray<FNiagaraParameterMapHistory>& GetPrecomputedHistories() { return PrecompiledHistories; }
