Fixed missing scratch pad icon for scratch dynamic input scripts

diff --git a/Engine/Plugins/FX/Niagara/Source/NiagaraEditorWidgets/Private/Stack/SNiagaraStackFunctionInputValue.cpp b/Engine/Plugins/FX/Niagara/Source/NiagaraEditorWidgets/Private/Stack/SNiagaraStackFunctionInputValue.cpp
--- a/Engine/Plugins/FX/Niagara/Source/NiagaraEditorWidgets/Private/Stack/SNiagaraStackFunctionInputValue.cpp
+++ b/Engine/Plugins/FX/Niagara/Source/NiagaraEditorWidgets/Private/Stack/SNiagaraStackFunctionInputValue.cpp
@@ -275,7 +275,7 @@ TSharedRef<SWidget> SNiagaraStackFunctionInputValue::ConstructValueWidgets()
 					.Content()
 					[
 						SNew(SImage)
-						.Image(FNiagaraEditorStyle::Get().GetBrush("NiagaraEditor.Scratch"))
+						.Image(FNiagaraEditorStyle::Get().GetBrush("Tab.ScratchPad"))
 					]
 				];
 		}
