Removed CPU access warning for camera DIs

diff --git a/Engine/Plugins/FX/Niagara/Source/Niagara/Private/NiagaraDataInterfaceCamera.cpp b/Engine/Plugins/FX/Niagara/Source/Niagara/Private/NiagaraDataInterfaceCamera.cpp
--- a/Engine/Plugins/FX/Niagara/Source/Niagara/Private/NiagaraDataInterfaceCamera.cpp
+++ b/Engine/Plugins/FX/Niagara/Source/Niagara/Private/NiagaraDataInterfaceCamera.cpp
@@ -568,60 +568,7 @@ void UNiagaraDataInterfaceCamera::GetFeedback(UNiagaraSystem* Asset, UNiagaraCom
 	{
 		return;
 	}
-
-	// we need to check if the DI is used to access camera properties in a cpu context to warn the user that
-	// the Niagara viewport does not support cpu cameras and it only works correctly in the level viewport and PIE
-
-	TArray<UNiagaraScript*> Scripts;
-	Scripts.Add(Asset->GetSystemSpawnScript());
-	Scripts.Add(Asset->GetSystemUpdateScript());
-	for (auto&& EmitterHandle : Asset->GetEmitterHandles())
-	{
-		if (EmitterHandle.GetInstance()->SimTarget == ENiagaraSimTarget::GPUComputeSim)
-		{
-			// Ignore gpu emitters
-			continue;
-		}
-		TArray<UNiagaraScript*> OutScripts;
-		EmitterHandle.GetInstance()->GetScripts(OutScripts, false);
-		Scripts.Append(OutScripts);
-	}
-
-	// Check if any CPU script uses camera functions
-	//TODO: This is the same as in the skel mesh DI, it doesn't guarantee that the DI used by these functions are THIS DI.
-	// Has a possibility of false positives
-	bool bHasCameraAccessWarning = [this, &Scripts]()
-	{
-		for (const auto Script : Scripts)
-		{
-			for (const auto& Info : Script->GetVMExecutableData().DataInterfaceInfo)
-			{
-				if (Info.MatchesClass(GetClass()))
-				{
-					for (const auto& Func : Info.RegisteredFunctions)
-					{
-						if (Func.Name == GetCameraPropertiesName || Func.Name == GetFieldOfViewName)
-						{
-							return true;
-						}
-					}
-				}
-			}
-		}
-		return false;
-	}();
-
-	if (bHasCameraAccessWarning)
-	{
-		FNiagaraDataInterfaceFeedback CPUAccessNotAllowedWarning(
-         LOCTEXT("CPUCameraAccessWarning", "The cpu camera is bound to a player controller and will therefore not work correctly in the Niagara viewport.\nTo correctly preview the effect, use it in the level editor or switch to a GPU emitter."),
-         LOCTEXT("CPUCameraAccessWarningSummary", "Camera properties cannot be previewed on CPU emitters!"),
-         FNiagaraDataInterfaceFix());
-
-		Warnings.Add(CPUAccessNotAllowedWarning);
-	}
 }
-
 #endif
 
 bool UNiagaraDataInterfaceCamera::CopyToInternal(UNiagaraDataInterface* Destination) const
