Fixed issue with CVar condition suggestion box not closing correctly.

#rnx

diff --git a/Engine/Plugins/FX/Niagara/Source/NiagaraEditor/Private/Customizations/NiagaraPlatformSetCustomization.cpp b/Engine/Plugins/FX/Niagara/Source/NiagaraEditor/Private/Customizations/NiagaraPlatformSetCustomization.cpp
--- a/Engine/Plugins/FX/Niagara/Source/NiagaraEditor/Private/Customizations/NiagaraPlatformSetCustomization.cpp
+++ b/Engine/Plugins/FX/Niagara/Source/NiagaraEditor/Private/Customizations/NiagaraPlatformSetCustomization.cpp
@@ -1253,6 +1253,12 @@ void SNiagaraConsoleInputBox::SuggestionSelectionChanged(TSharedPtr<FString> New
 		return;
 	}
 
+	if (NewValue.IsValid() == false)
+	{
+		SuggestionBox->SetIsOpen(false);
+		return;
+	}
+
 	Suggestions.SelectedSuggestion = Suggestions.SuggestionsList.IndexOfByPredicate([&NewValue](const TSharedPtr<FString>& InSuggestion)
 		{
 			return InSuggestion == NewValue;
@@ -1378,6 +1384,7 @@ void SNiagaraConsoleInputBox::OnTextCommitted(const FText& InText, ETextCommit::
 		WorkingText = InText;
 		IConsoleManager::Get().AddConsoleHistoryEntry(NiagaraCVarHistoryKey, *InText.ToString());
 		OnTextCommittedEvent.ExecuteIfBound(InText);
+		SuggestionBox->SetIsOpen(false);
 	}
 }
 
