- Fix for property reads on the GPU not using the correct index

#rnx

diff --git a/Engine/Plugins/FX/Niagara/Source/Niagara/Private/DataInterface/NiagaraDataInterfaceUObjectPropertyReader.cpp b/Engine/Plugins/FX/Niagara/Source/Niagara/Private/DataInterface/NiagaraDataInterfaceUObjectPropertyReader.cpp
--- a/Engine/Plugins/FX/Niagara/Source/Niagara/Private/DataInterface/NiagaraDataInterfaceUObjectPropertyReader.cpp
+++ b/Engine/Plugins/FX/Niagara/Source/Niagara/Private/DataInterface/NiagaraDataInterfaceUObjectPropertyReader.cpp
@@ -85,7 +85,7 @@ namespace NDIUObjectPropertyReaderLocal
 				for ( int32 i=0; i < InstanceData_GT->PropertyGetters.Num(); ++i )
 				{
 					const FNDIPropertyGetter& PropertyGetter = InstanceData_GT->PropertyGetters[i];
-					InstanceData_ForRT->PropertyOffsets.Add(PropertyGetter.PropertyCopyFunction ? PropertyGetter.DataOffset : INDEX_NONE);
+					InstanceData_ForRT->PropertyOffsets[i] = PropertyGetter.PropertyCopyFunction ? PropertyGetter.DataOffset : INDEX_NONE;
 				}
 				InstanceData_ForRT->PropertyData = InstanceData_GT->PropertyData;
 			}
