Fixed a crash when the underlying script for inline parameter widget generation was removed

diff --git a/Engine/Plugins/FX/Niagara/Source/NiagaraEditor/Public/NiagaraNodeFunctionCall.h b/Engine/Plugins/FX/Niagara/Source/NiagaraEditor/Public/NiagaraNodeFunctionCall.h
--- a/Engine/Plugins/FX/Niagara/Source/NiagaraEditor/Public/NiagaraNodeFunctionCall.h
+++ b/Engine/Plugins/FX/Niagara/Source/NiagaraEditor/Public/NiagaraNodeFunctionCall.h
@@ -97,7 +97,7 @@ public:
 	UPROPERTY(EditAnywhere, Category = "Debug")
 	bool bInheritDebugStatus = true;
 
-	bool ScriptIsValid() const;
+	NIAGARAEDITOR_API bool ScriptIsValid() const;
 
 	//Begin UObject interface
 	virtual void PostLoad()override;
diff --git a/Engine/Plugins/FX/Niagara/Source/NiagaraEditorWidgets/Private/SNiagaraOverviewStack.cpp b/Engine/Plugins/FX/Niagara/Source/NiagaraEditorWidgets/Private/SNiagaraOverviewStack.cpp
--- a/Engine/Plugins/FX/Niagara/Source/NiagaraEditorWidgets/Private/SNiagaraOverviewStack.cpp
+++ b/Engine/Plugins/FX/Niagara/Source/NiagaraEditorWidgets/Private/SNiagaraOverviewStack.cpp
@@ -835,7 +835,7 @@ TSharedRef<ITableRow> SNiagaraOverviewStack::OnGenerateRowForEntry(UNiagaraStack
 			];
 
 		// in case we have at least one inline input, we add a wrap box to make sure we have enough space for both the module name and parameters
-		if(StackModuleItem && StackModuleItem->GetInlineParameterInputs().Num() > 0)
+		if(StackModuleItem && StackModuleItem->GetModuleNode().ScriptIsValid() && StackModuleItem->GetInlineParameterInputs().Num() > 0)
 		{
 			ContentBox = SNew(SHorizontalBox)
 			+ SHorizontalBox::Slot()
