- Ensure we look for data interface on the system scripts for emitters when caching data interfaces

#rnx

diff --git a/Engine/Plugins/FX/Niagara/Source/Niagara/Classes/NiagaraDataInterface.h b/Engine/Plugins/FX/Niagara/Source/Niagara/Classes/NiagaraDataInterface.h
--- a/Engine/Plugins/FX/Niagara/Source/Niagara/Classes/NiagaraDataInterface.h
+++ b/Engine/Plugins/FX/Niagara/Source/Niagara/Classes/NiagaraDataInterface.h
@@ -294,7 +294,7 @@ public:
 #endif
 
 	/** Allows data interfaces to cache any static data that may be shared between instances */
-	virtual void CacheStaticBuffers(struct FNiagaraSystemStaticBuffers& StaticBuffers, const struct FNiagaraScriptDataInterfaceInfo& DataInterfaceInfo, ENiagaraSimTarget SimTarget) {}
+	virtual void CacheStaticBuffers(struct FNiagaraSystemStaticBuffers& StaticBuffers, const struct FNiagaraScriptDataInterfaceInfo& DataInterfaceInfo, bool bUsedByCPU, bool bUsedByGPU) {}
 
 	virtual bool NeedsGPUContextInit() const { return false; }
 	virtual bool GPUContextInit(const FNiagaraScriptDataInterfaceCompileInfo& InInfo, void* PerInstanceData, FNiagaraSystemInstance* SystemInstance) const { return false; }
diff --git a/Engine/Plugins/FX/Niagara/Source/Niagara/Classes/NiagaraDataInterfaceCurveBase.h b/Engine/Plugins/FX/Niagara/Source/Niagara/Classes/NiagaraDataInterfaceCurveBase.h
--- a/Engine/Plugins/FX/Niagara/Source/Niagara/Classes/NiagaraDataInterfaceCurveBase.h
+++ b/Engine/Plugins/FX/Niagara/Source/Niagara/Classes/NiagaraDataInterfaceCurveBase.h
@@ -182,7 +182,7 @@ public:
 	/** Gets information for all of the curves owned by this curve data interface. */
 	virtual void GetCurveData(TArray<FCurveData>& OutCurveData) { }
 
-	virtual void CacheStaticBuffers(struct FNiagaraSystemStaticBuffers& StaticBuffers, const struct FNiagaraScriptDataInterfaceInfo& DataInterfaceInfo, ENiagaraSimTarget SimTarget) override;
+	virtual void CacheStaticBuffers(struct FNiagaraSystemStaticBuffers& StaticBuffers, const struct FNiagaraScriptDataInterfaceInfo& DataInterfaceInfo, bool bUsedByCPU, bool bUsedByGPU) override;
 #if WITH_EDITORONLY_DATA
 	virtual void GetParameterDefinitionHLSL(const FNiagaraDataInterfaceGPUParamInfo& ParamInfo, FString& OutHLSL) override;
 #endif
diff --git a/Engine/Plugins/FX/Niagara/Source/Niagara/Private/NiagaraDataInterfaceCurveBase.cpp b/Engine/Plugins/FX/Niagara/Source/Niagara/Private/NiagaraDataInterfaceCurveBase.cpp
--- a/Engine/Plugins/FX/Niagara/Source/Niagara/Private/NiagaraDataInterfaceCurveBase.cpp
+++ b/Engine/Plugins/FX/Niagara/Source/Niagara/Private/NiagaraDataInterfaceCurveBase.cpp
@@ -327,10 +327,10 @@ bool UNiagaraDataInterfaceCurveBase::Equals(const UNiagaraDataInterface* Other)
 	return bEqual;
 }
 
-void UNiagaraDataInterfaceCurveBase::CacheStaticBuffers(struct FNiagaraSystemStaticBuffers& StaticBuffers, const struct FNiagaraScriptDataInterfaceInfo& DataInterfaceInfo, ENiagaraSimTarget SimTarget)
+void UNiagaraDataInterfaceCurveBase::CacheStaticBuffers(struct FNiagaraSystemStaticBuffers& StaticBuffers, const struct FNiagaraScriptDataInterfaceInfo& DataInterfaceInfo, bool bUsedByCPU, bool bUsedByGPU)
 {
 	LUTOffset = INDEX_NONE;
-	if (DataInterfaceInfo.IsUserDataInterface() == false && SimTarget == ENiagaraSimTarget::GPUComputeSim)
+	if (bUsedByGPU && DataInterfaceInfo.IsUserDataInterface() == false)
 	{
 		LUTOffset = StaticBuffers.AddGpuData(ShaderLUT);
 	}
diff --git a/Engine/Plugins/FX/Niagara/Source/Niagara/Private/NiagaraSystem.cpp b/Engine/Plugins/FX/Niagara/Source/Niagara/Private/NiagaraSystem.cpp
--- a/Engine/Plugins/FX/Niagara/Source/Niagara/Private/NiagaraSystem.cpp
+++ b/Engine/Plugins/FX/Niagara/Source/Niagara/Private/NiagaraSystem.cpp
@@ -1956,6 +1956,7 @@ void UNiagaraSystem::CacheFromCompiledData()
 	// reset the MaxDeltaTime so we get the most up to date values from the emitters
 	MaxDeltaTime.Reset();
 
+	TSet<FName> DataInterfaceGpuUsage;
 	TStringBuilder<128> ExecutionStateNameBuilder;
 	for (int32 i=0; i < EmitterHandles.Num(); ++i)
 	{
@@ -1991,24 +1992,35 @@ void UNiagaraSystem::CacheFromCompiledData()
 			NiagaraEmitter->CacheFromCompiledData(DataSetCompiledData);
 
 			// Allow data interfaces to cache static buffers
-			UNiagaraScript* NiagaraScripts[] =
+			UNiagaraScript* NiagaraEmitterScripts[] =
 			{
 				NiagaraEmitter->SpawnScriptProps.Script,
 				NiagaraEmitter->UpdateScriptProps.Script,
 				NiagaraEmitter->GetGPUComputeScript(),
 			};
-			for (UNiagaraScript* NiagaraScript : NiagaraScripts)
+
+			for (UNiagaraScript* NiagaraScript : NiagaraEmitterScripts)
 			{
 				if (NiagaraScript == nullptr)
 				{
 					continue;
 				}
 
+				const bool bUsedByCPU = NiagaraEmitter->SimTarget == ENiagaraSimTarget::CPUSim;
+				const bool bUsedByGPU = NiagaraEmitter->SimTarget == ENiagaraSimTarget::GPUComputeSim;
 				for (const FNiagaraScriptDataInterfaceInfo& DataInterfaceInfo : NiagaraScript->GetCachedDefaultDataInterfaces())
 				{
 					if (UNiagaraDataInterface* DataInterface = DataInterfaceInfo.DataInterface)
 					{
-						DataInterface->CacheStaticBuffers(*StaticBuffers.Get(), DataInterfaceInfo, NiagaraEmitter->SimTarget);
+						DataInterface->CacheStaticBuffers(*StaticBuffers.Get(), DataInterfaceInfo, bUsedByCPU, bUsedByGPU);
+
+						if ( bUsedByGPU )
+						{
+							if (!DataInterfaceInfo.RegisteredParameterMapRead.IsNone())
+							{
+								DataInterfaceGpuUsage.Add(DataInterfaceInfo.RegisteredParameterMapRead);
+							}
+						}
 					}
 				}
 			}
@@ -2019,6 +2031,24 @@ void UNiagaraSystem::CacheFromCompiledData()
 		}
 	}
 
+	// Loop over system scripts these are more awkward because we need to determine the usage
+	for (UNiagaraScript* NiagaraScript : { SystemSpawnScript, SystemUpdateScript })
+	{
+		if (NiagaraScript == nullptr)
+		{
+			continue;
+		}
+
+		for (const FNiagaraScriptDataInterfaceInfo& DataInterfaceInfo : NiagaraScript->GetCachedDefaultDataInterfaces())
+		{
+			if (UNiagaraDataInterface* DataInterface = DataInterfaceInfo.DataInterface)
+			{
+				const bool bUsedByGPU = DataInterfaceGpuUsage.Contains(DataInterfaceInfo.RegisteredParameterMapWrite);
+				DataInterface->CacheStaticBuffers(*StaticBuffers.Get(), DataInterfaceInfo, true, bUsedByGPU);
+			}
+		}
+	}
+
 	// Finalize static buffers
 	StaticBuffers->Finalize();
 }
