- Add toggle to enable / disable particle index display
- Disabled by default as most of the time you don't need this

#rnx

diff --git a/Engine/Plugins/FX/Niagara/Source/Niagara/Private/NiagaraDebugHud.cpp b/Engine/Plugins/FX/Niagara/Source/Niagara/Private/NiagaraDebugHud.cpp
--- a/Engine/Plugins/FX/Niagara/Source/Niagara/Private/NiagaraDebugHud.cpp
+++ b/Engine/Plugins/FX/Niagara/Source/Niagara/Private/NiagaraDebugHud.cpp
@@ -2305,7 +2305,10 @@ void FNiagaraDebugHud::DrawComponents(FNiagaraWorldManager* WorldManager, UCanva
 					const FVector ParticleScreenLocation = Canvas->Project(ParticleWorldPosition);
 
 					TStringBuilder<1024> StringBuilder;
-					StringBuilder.Appendf(TEXT("Particle(%u)"), iInstance);
+					if ( Settings.bShowParticleIndex )
+					{
+						StringBuilder.Appendf(TEXT("Particle(%u)"), iInstance);
+					}
 					for (const auto& ParticleVariable : CachedVariables.ParticleVariables[iEmitter])
 					{
 						StringBuilder.AppendChar(Settings.bShowParticleVariablesVertical ? TEXT('\n') : TEXT(' '));
diff --git a/Engine/Plugins/FX/Niagara/Source/Niagara/Public/NiagaraDebuggerCommon.h b/Engine/Plugins/FX/Niagara/Source/Niagara/Public/NiagaraDebuggerCommon.h
--- a/Engine/Plugins/FX/Niagara/Source/Niagara/Public/NiagaraDebuggerCommon.h
+++ b/Engine/Plugins/FX/Niagara/Source/Niagara/Public/NiagaraDebuggerCommon.h
@@ -530,6 +530,13 @@ struct NIAGARA_API FNiagaraDebugHUDSettingsData
 	UPROPERTY(Config, EditAnywhere, Category = "Debug Particles", meta = (EditCondition = "bShowParticleVariables"))
 	bool bEnableGpuParticleReadback = false;
 
+	/**
+	When enabled the particle index will be displayed along with any attributes.
+	Note: This is the index in the particle data buffer and not the persistent ID index.
+	*/
+	UPROPERTY(Config, EditAnywhere, Category = "Debug Particles", meta = (EditCondition = "bShowParticleVariables"))
+	bool bShowParticleIndex = false;
+
 	/**
 	List of attributes to show per particle, each entry uses wildcard matching.
 	For example, "*Position" would match all attributes that end in Position.
