Improve dynamic material parameter name resolution in niagara stack; now evaluate whether dynamic parameter nodes in material graphs have pin links when determining name string.

diff --git a/Engine/Plugins/FX/Niagara/Source/NiagaraEditor/Private/ViewModels/Stack/NiagaraStackFunctionInput.cpp b/Engine/Plugins/FX/Niagara/Source/NiagaraEditor/Private/ViewModels/Stack/NiagaraStackFunctionInput.cpp
--- a/Engine/Plugins/FX/Niagara/Source/NiagaraEditor/Private/ViewModels/Stack/NiagaraStackFunctionInput.cpp
+++ b/Engine/Plugins/FX/Niagara/Source/NiagaraEditor/Private/ViewModels/Stack/NiagaraStackFunctionInput.cpp
@@ -3,6 +3,7 @@
 #include "ViewModels/Stack/NiagaraStackFunctionInput.h"
 
 #include "AssetRegistryModule.h"
+#include "Containers/StaticBitArray.h"
 #include "EdGraphSchema_Niagara.h"
 #include "Editor.h"
 #include "INiagaraEditorTypeUtilities.h"
@@ -789,35 +790,10 @@ void UNiagaraStackFunctionInput::RefreshChildrenInternal(const TArray<UNiagaraSt
 	LastFunctionGraphChangeId = CurrentFunctionGraphChangeId;
 }
 
-class UNiagaraStackFunctionInputUtilities
+FString UNiagaraStackFunctionInput::ResolveDisplayNameArgument(const FString& InArg) const
 {
-public:
-	static bool GetMaterialExpressionDynamicParameter(UNiagaraEmitter* InEmitter, FNiagaraEmitterInstance* InEmitterInstance, TArray<UMaterialExpressionDynamicParameter*>& OutDynamicParameterExpressions)
-	{
-		TArray<UMaterial*> Materials = GetMaterialFromEmitter(InEmitter, InEmitterInstance);
-		
-		OutDynamicParameterExpressions.Reset();
-
-		// Find the dynamic parameters expression from the material.
-		// @YannickLange todo: Notify user that the material did not have any dynamic parameters. Maybe add them from code?
-		for (UMaterial* Material : Materials)
-		{
-			if (Material != nullptr)
-			{
-				TArray<UMaterialExpressionDynamicParameter*> Expressions;
-				Material->GetAllExpressionsInMaterialAndFunctionsOfType(Expressions);
-				for (UMaterialExpressionDynamicParameter* DynParamExpFound : Expressions)
-				{
-					OutDynamicParameterExpressions.Add(DynParamExpFound);
-				}
-			}
-		}
-
-		return OutDynamicParameterExpressions.Num() > 0;
-	}
-
-	static TArray<UMaterial*> GetMaterialFromEmitter(UNiagaraEmitter* InEmitter, FNiagaraEmitterInstance* InEmitterInstance)
-	{
+	//~ Begin helper functions
+	auto GetMaterialsFromEmitter = [](UNiagaraEmitter* InEmitter, FNiagaraEmitterInstance* InEmitterInstance)->TArray<UMaterial*> {
 		TArray<UMaterial*> ResultMaterials;
 		if (InEmitter != nullptr)
 		{
@@ -840,65 +816,171 @@ public:
 			}
 		}
 		return ResultMaterials;
-	}
-};
+	};
 
-FString UNiagaraStackFunctionInput::ResolveDisplayNameArgument(const FString& InArg) const
-{
-	if (InArg.StartsWith(TEXT("MaterialDynamicParam")))
-	{
-		TSharedPtr<FNiagaraEmitterViewModel> ThisEmitterViewModel = GetEmitterViewModel();
-		TArray<UMaterialExpressionDynamicParameter*> ExpressionParams;
-		FNiagaraEmitterInstance* Instance = ThisEmitterViewModel->GetSimulation().Pin().Get();
-		if (ThisEmitterViewModel.IsValid() == false || !UNiagaraStackFunctionInputUtilities::GetMaterialExpressionDynamicParameter(ThisEmitterViewModel->GetEmitter(), Instance, ExpressionParams))
-		{
-			return InArg.Replace(TEXT("MaterialDynamic"), TEXT("")) + TEXT(" (No material found using dynamic params)");
-		}
-		
-		FString Suffix = InArg.Right(3);
-		int32 ParamIdx;
-		LexFromString(ParamIdx, *Suffix.Left(1));
-		int32 ParamSlotIdx;
-		LexFromString(ParamSlotIdx, *Suffix.Right(1));
-		
-		if (ParamIdx < 0 || ParamIdx > 3 || ParamSlotIdx < 0 || ParamSlotIdx > 3)
-		{
-			return InArg.Replace(TEXT("MaterialDynamic"), TEXT("")) + TEXT(" (error parsing parameter name)");
-		}
+	// Duplicated from UMaterialGraph::GetValidOutputIndex()
+	auto GetValidOutputIndex = [](FExpressionInput* Input)->int32 /*OutputIndex*/ {
+		int32 OutputIndex = 0;
 
-		FName ParamName = NAME_None;
-		bool bAllSame = true;
-		for (UMaterialExpressionDynamicParameter* Expression : ExpressionParams)
+		if (Input->Expression)
 		{
-			const FExpressionOutput& Output = Expression->GetOutputs()[ParamIdx];
-			if (ParamSlotIdx == Expression->ParameterIndex)
+			TArray<FExpressionOutput>& Outputs = Input->Expression->GetOutputs();
+
+			if (Outputs.Num() > 0)
 			{
-				if (ParamName == NAME_None)
+				const bool bOutputIndexIsValid = Outputs.IsValidIndex(Input->OutputIndex)
+					// Attempt to handle legacy connections before OutputIndex was used that had a mask
+					&& (Input->OutputIndex != 0 || Input->Mask == 0);
+
+				for (; OutputIndex < Outputs.Num(); ++OutputIndex)
 				{
-					ParamName = Output.OutputName;
+					const FExpressionOutput& Output = Outputs[OutputIndex];
+
+					if ((bOutputIndexIsValid && OutputIndex == Input->OutputIndex)
+						|| (!bOutputIndexIsValid
+							&& Output.Mask == Input->Mask
+							&& Output.MaskR == Input->MaskR
+							&& Output.MaskG == Input->MaskG
+							&& Output.MaskB == Input->MaskB
+							&& Output.MaskA == Input->MaskA))
+					{
+						break;
+					}
 				}
-				else if (ParamName != Output.OutputName)
+
+				if (OutputIndex >= Outputs.Num())
 				{
-					bAllSame = false;
+					// Work around for non-reproducible crash where OutputIndex would be out of bounds
+					OutputIndex = Outputs.Num() - 1;
 				}
 			}
 		}
 
-		if (ParamName != NAME_None)
+		return OutputIndex;
+	};
+	//~ End helper functions
+
+	// If the DisplayNameArgument to resolve is not a MaterialDynamicParam, early out.
+	if (InArg.StartsWith(TEXT("MaterialDynamicParam")) == false)
+	{
+		return FString();
+	}
+
+	// Get the target indices of the MaterialDynamicParam. Early out if they are invalid.
+	FString Suffix = InArg.Right(3);
+	int32 ParamIdx;
+	LexFromString(ParamIdx, *Suffix.Left(1));
+	int32 ParamSlotIdx;
+	LexFromString(ParamSlotIdx, *Suffix.Right(1));
+
+	if (ParamIdx < 0 || ParamIdx > 3 || ParamSlotIdx < 0 || ParamSlotIdx > 3)
+	{
+		return InArg.Replace(TEXT("MaterialDynamic"), TEXT("")) + TEXT(" (error parsing parameter name)");
+	}
+
+	TArray<UMaterial*> Materials;
+	TSharedPtr<FNiagaraEmitterViewModel> ThisEmitterViewModel = GetEmitterViewModel();
+	if (ThisEmitterViewModel.IsValid())
+	{
+		FNiagaraEmitterInstance* ThisEmitterInstance = ThisEmitterViewModel->GetSimulation().Pin().Get();
+		Materials = GetMaterialsFromEmitter(ThisEmitterViewModel->GetEmitter(), ThisEmitterInstance);
+	} 
+
+	// Determine the MaterialDynamicParam name to display inline based on the friendly names for each UMaterialExpressionDynamicParameter in the graph, and whether each UMaterialExpressionDynamicParameter pin is linked.
+	// For each graph; traverse every expression, and for each expression check ExpressionInputs; if an input of ExpressionInputs is recording a link to a UMaterialExpressionDynamicParameter, record this in DynamicParameterExpressionToOutputMaskMap.
+	// After recording all UMaterialExpressionDynamicParameter output pins that are linked at least once, iterate DynamicParameterExpressionToOutputMaskMap to determine the final name.
+	// NOTE: This check does not constitute a true "reachability analysis" as we are only recording if each pin of the UMaterialExpressionDynamicParameter is linked, and not whether that pin is connected to a route of expressions that would eventually output.
+	TMap<UMaterialExpressionDynamicParameter*, TStaticBitArray<4>> DynamicParameterExpressionToOutputMaskMap;
+	TArray<FExpressionInput*> ExpressionInputsToProcess;
+
+	// Visit each material and gather all expression inputs for each expression.
+	for (UMaterial* Material : Materials)
+	{
+		if (Material == nullptr)
+		{
+			continue;
+		}
+
+		for (int32 MaterialPropertyIndex = 0; MaterialPropertyIndex < MP_MAX; ++MaterialPropertyIndex)
 		{
-			if (bAllSame)
+			FExpressionInput* ExpressionInput = Material->GetExpressionInputForProperty(EMaterialProperty(MaterialPropertyIndex));
+			if (ExpressionInput != nullptr)
 			{
-				return ParamName.ToString();
-			}
-			else
-			{
-				return ParamName.ToString() + TEXT(" (Multiple Aliases Found)");
+				ExpressionInputsToProcess.Add(ExpressionInput);
 			}
 		}
 
-		return InArg.Replace(TEXT("MaterialDynamic"), TEXT("")) + TEXT(" (Parameter not used in materials.)");
+		TArray<UMaterialExpression*> Expressions;
+		Material->GetAllExpressionsInMaterialAndFunctionsOfType<UMaterialExpression>(Expressions);
+		for (UMaterialExpression* Expression : Expressions)
+		{
+			const TArray<FExpressionInput*>& ExpressionInputs = Expression->GetInputs();
+			ExpressionInputsToProcess.Append(ExpressionInputs);
+		}
 	}
-	return FString();
+
+	// Visit each expression input and record which inputs are associated with UMaterialExpressionDynamicParameter outputs.
+	bool bAnyDynamicParametersFound = false;
+	for (FExpressionInput* ExpressionInput : ExpressionInputsToProcess)
+	{
+		if (ExpressionInput->Expression == nullptr)
+		{
+			continue;
+		}
+
+		UMaterialExpressionDynamicParameter* DynamicParameterExpression = Cast<UMaterialExpressionDynamicParameter>(ExpressionInput->Expression);
+		if (DynamicParameterExpression == nullptr)
+		{
+			continue;
+		}
+
+		bAnyDynamicParametersFound = true;
+		TStaticBitArray<4>& DynamicParameterExpressionOutputMask = DynamicParameterExpressionToOutputMaskMap.FindOrAdd(DynamicParameterExpression);
+		DynamicParameterExpressionOutputMask[GetValidOutputIndex(ExpressionInput)] = true;
+	}
+
+	// Construct the final dynamic param UI name. Visit each UMaterialExpressionDynamicParameter and for those which have an output which is used, consider them for the param name.
+	FName ParamName = NAME_None;
+	bool bMultipleAliasesFound = false;
+	const FString DefaultDynamicParameterNameString = "Param" + FString::Printf(TEXT("%d"), ParamIdx + 1);
+	for (auto It = DynamicParameterExpressionToOutputMaskMap.CreateConstIterator(); It; ++It)
+	{
+		UMaterialExpressionDynamicParameter* ExpressionDynamicParameter = It.Key();
+		const TStaticBitArray<4>& ExpressionOutputMask = It.Value();
+		if (ExpressionDynamicParameter->ParameterIndex != ParamSlotIdx || ExpressionOutputMask[ParamIdx] == false)
+		{
+			continue;
+		}
+
+		const FExpressionOutput& Output = ExpressionDynamicParameter->GetOutputs()[ParamIdx];
+		if (ParamName == NAME_None)
+		{
+			ParamName = Output.OutputName;
+		}
+		else if (ParamName != Output.OutputName)
+		{
+			bMultipleAliasesFound = true;
+		}
+	}
+
+	// Return the final dynamic param UI name.
+	if (bAnyDynamicParametersFound == false)
+	{
+		return InArg.Replace(TEXT("MaterialDynamic"), TEXT("")) + TEXT(" (No material found using dynamic params)");
+	}
+	else if (ParamName != NAME_None)
+	{
+		if (bMultipleAliasesFound == false)
+		{
+			return ParamName.ToString();
+		}
+		else
+		{
+			return ParamName.ToString() + TEXT(" (Multiple Aliases Found)");
+		}
+	}
+
+	return InArg.Replace(TEXT("MaterialDynamic"), TEXT("")) + TEXT(" (Parameter not used in materials.)");
 }
 
 void UNiagaraStackFunctionInput::RefreshValues()
