- Fix for function ordering issue on GPU

#rnx

diff --git a/Engine/Plugins/FX/Niagara/Source/Niagara/Private/DataInterface/NiagaraDataInterfaceUObjectPropertyReader.cpp b/Engine/Plugins/FX/Niagara/Source/Niagara/Private/DataInterface/NiagaraDataInterfaceUObjectPropertyReader.cpp
--- a/Engine/Plugins/FX/Niagara/Source/Niagara/Private/DataInterface/NiagaraDataInterfaceUObjectPropertyReader.cpp
+++ b/Engine/Plugins/FX/Niagara/Source/Niagara/Private/DataInterface/NiagaraDataInterfaceUObjectPropertyReader.cpp
@@ -64,6 +64,7 @@ namespace NDIUObjectPropertyReaderLocal
 		TOptional<FTransform>	CachedTransform;
 		FReadBuffer				PropertyData;
 		TArray<uint32>			PropertyOffsets;
+		TArray<uint32>			GpuFunctionToPropertyRemap;
 	};
 
 	struct FNDIProxy : public FNiagaraDataInterfaceProxy
@@ -73,15 +74,9 @@ namespace NDIUObjectPropertyReaderLocal
 			FInstanceData_GameToRender* InstanceData_ForRT = new(DataForRenderThread) FInstanceData_GameToRender();
 			FInstanceData_GameThread* InstanceData_GT = reinterpret_cast<FInstanceData_GameThread*>(PerInstanceData);
 			InstanceData_ForRT->CachedTransform = InstanceData_GT->CachedTransform;
-			if (InstanceData_GT->PropertyData.Num() == 0)
+			if (InstanceData_GT->PropertyData.Num() > 0)
 			{
-				InstanceData_ForRT->PropertyOffsets.AddZeroed(4);
-				InstanceData_ForRT->PropertyOffsets[0] = INDEX_NONE;
-				InstanceData_ForRT->PropertyData.AddZeroed(4);
-			}
-			else
-			{
-				InstanceData_ForRT->PropertyOffsets.AddZeroed(Align(InstanceData_GT->PropertyGetters.Num(), 4));
+				InstanceData_ForRT->PropertyOffsets.AddZeroed(InstanceData_GT->PropertyGetters.Num());
 				for ( int32 i=0; i < InstanceData_GT->PropertyGetters.Num(); ++i )
 				{
 					const FNDIPropertyGetter& PropertyGetter = InstanceData_GT->PropertyGetters[i];
@@ -97,7 +92,13 @@ namespace NDIUObjectPropertyReaderLocal
 			if ( FInstanceData_RenderThread* InstanceData_RT = PerInstanceData_RenderThread.Find(InstanceID) )
 			{
 				InstanceData_RT->CachedTransform = InstanceData_FromGT->CachedTransform;
-				InstanceData_RT->PropertyOffsets = InstanceData_FromGT->PropertyOffsets;
+
+				for ( int32 i=0; i < InstanceData_RT->GpuFunctionToPropertyRemap.Num(); ++i )
+				{
+					const uint32 RemapIndex = InstanceData_RT->GpuFunctionToPropertyRemap[i];
+					InstanceData_RT->PropertyOffsets[i] = RemapIndex == INDEX_NONE ? INDEX_NONE : InstanceData_FromGT->PropertyOffsets[RemapIndex];
+				}
+
 				if ( InstanceData_RT->PropertyData.NumBytes != InstanceData_FromGT->PropertyData.Num())
 				{
 					InstanceData_RT->PropertyData.Release();
@@ -551,6 +552,8 @@ bool UNiagaraDataInterfaceUObjectPropertyReader::InitPerInstanceData(void* PerIn
 
 	if ( IsUsedByGPUEmitter() )
 	{
+		TArray<uint32> GpuFunctionToPropertyRemap;
+
 		// We shouldn't need to do this per init, we should be able to cache once and once only
 		for (const TSharedRef<FNiagaraEmitterInstance, ESPMode::ThreadSafe>& EmitterInstance : SystemInstance->GetEmitters())
 		{
@@ -561,14 +564,18 @@ bool UNiagaraDataInterfaceUObjectPropertyReader::InitPerInstanceData(void* PerIn
 			const FNiagaraScriptInstanceParameterStore& ParameterStore = EmitterInstance->GetGPUContext()->CombinedParamStore;
 			const TArray<UNiagaraDataInterface*>& DataInterfaces = ParameterStore.GetDataInterfaces();
 			const TArray<FNiagaraDataInterfaceGPUParamInfo>& DataInterfaceParamInfo = EmitterInstance->GetGPUContext()->GPUScript_RT->GetDataInterfaceParamInfo();
-			for ( int32 i=0; i < DataInterfaces.Num(); ++i )
+			for ( int32 iDataInterface=0; iDataInterface < DataInterfaces.Num(); ++iDataInterface)
 			{
-				if ( (DataInterfaces[i] == this) && DataInterfaceParamInfo.IsValidIndex(i) )
+				if ( (DataInterfaces[iDataInterface] == this) && DataInterfaceParamInfo.IsValidIndex(iDataInterface) )
 				{
-					for (const FNiagaraDataInterfaceGeneratedFunction& FunctionInfo : DataInterfaceParamInfo[i].GeneratedFunctions )
+					const TArray<FNiagaraDataInterfaceGeneratedFunction>& GeneratedFunctions = DataInterfaceParamInfo[iDataInterface].GeneratedFunctions;
+					GpuFunctionToPropertyRemap.AddZeroed(GeneratedFunctions.Num());
+					for ( int32 iFunctionInfo=0; iFunctionInfo < GeneratedFunctions.Num(); ++iFunctionInfo )
 					{
+						const FNiagaraDataInterfaceGeneratedFunction& FunctionInfo = GeneratedFunctions[iFunctionInfo];
+						GpuFunctionToPropertyRemap[iFunctionInfo] = INDEX_NONE;
 						#define NDI_PROPERTY_TYPE(TYPE) \
-							else if (FunctionInfo.DefinitionName == FTypeHelper<TYPE>::GetFunctionName()) { InstanceData_GT->AddProperty(FNiagaraVariableBase(FTypeHelper<TYPE>::GetTypeDef(), GetRemappedPropertyName(FunctionInfo.Specifiers[0].Value))); }
+							else if (FunctionInfo.DefinitionName == FTypeHelper<TYPE>::GetFunctionName()) { GpuFunctionToPropertyRemap[iFunctionInfo] = InstanceData_GT->AddProperty(FNiagaraVariableBase(FTypeHelper<TYPE>::GetTypeDef(), GetRemappedPropertyName(FunctionInfo.Specifiers[0].Value))); }
 
 							if (false) {} NDI_PROPERTY_TYPES
 						#undef NDI_PROPERTY_TYPE
@@ -580,9 +587,12 @@ bool UNiagaraDataInterfaceUObjectPropertyReader::InitPerInstanceData(void* PerIn
 		// Initialize render side instance data
 		ENQUEUE_RENDER_COMMAND(NDIUObjectPropertyReader_InitRT)
 		(
-			[Proxy_RT=GetProxyAs<FNDIProxy>(), InstanceID=SystemInstance->GetId()](FRHICommandList& CmdList)
+			[Proxy_RT=GetProxyAs<FNDIProxy>(), InstanceID=SystemInstance->GetId(), GpuFunctionToPropertyRemap_RT=MoveTemp(GpuFunctionToPropertyRemap)](FRHICommandList& CmdList)
 			{
-				Proxy_RT->PerInstanceData_RenderThread.Add(InstanceID);
+				FInstanceData_RenderThread* InstanceData_RT = &Proxy_RT->PerInstanceData_RenderThread.Add(InstanceID);
+				InstanceData_RT->PropertyOffsets.AddUninitialized(Align(GpuFunctionToPropertyRemap_RT.Num(), 4));
+				FMemory::Memset(InstanceData_RT->PropertyOffsets.GetData(), 0xff, InstanceData_RT->PropertyOffsets.GetAllocatedSize());
+				InstanceData_RT->GpuFunctionToPropertyRemap = GpuFunctionToPropertyRemap_RT;
 			}
 		);
 	}
