Removed unnecessary widget that was causing missing style logs

diff --git a/Engine/Plugins/FX/Niagara/Source/NiagaraEditor/Private/Widgets/SNiagaraParameterPanel.cpp b/Engine/Plugins/FX/Niagara/Source/NiagaraEditor/Private/Widgets/SNiagaraParameterPanel.cpp
--- a/Engine/Plugins/FX/Niagara/Source/NiagaraEditor/Private/Widgets/SNiagaraParameterPanel.cpp
+++ b/Engine/Plugins/FX/Niagara/Source/NiagaraEditor/Private/Widgets/SNiagaraParameterPanel.cpp
@@ -96,39 +96,6 @@ void SNiagaraParameterPanel::Construct(const FArguments& InArgs, const TSharedPt
 	.OnGetKeyForItem(this, &SNiagaraParameterPanel::OnGetKeyForItem)
 	.OnGetKeyForCategory(this, &SNiagaraParameterPanel::OnGetKeyForCategory);
 
-	// View options
-	TSharedRef<SWidget> ViewOptionsWidget = SNew(SBorder)
-		.Padding(0)
-		.BorderImage_Static(&SNiagaraParameterPanel::GetViewOptionsBorderBrush)
-		[
-			SNew(SComboButton)
-			.ContentPadding(0)
-			.ForegroundColor(FSlateColor::UseForeground())
-			.ButtonStyle(FEditorStyle::Get(), "ToggleButton")
-			.ToolTipText(LOCTEXT("ViewOptionsToolTip", "View Options"))
-			.OnGetMenuContent(this, &SNiagaraParameterPanel::GetViewOptionsMenu)
-			.ButtonContent()
-			[
-				SNew(SOverlay)
-				// drop shadow
-				+ SOverlay::Slot()
-				.VAlign(VAlign_Top)
-				.Padding(ViewOptionsShadowOffset.X, ViewOptionsShadowOffset.Y, 0, 0)
-				[
-					SNew(SImage)
-					.Image(FEditorStyle::GetBrush("GenericViewButton"))
-					.ColorAndOpacity(FNiagaraEditorStyle::Get().GetColor("NiagaraEditor.Stack.ViewOptionsShadowColor"))
-				]
-				+ SOverlay::Slot()
-				.VAlign(VAlign_Top)
-				[
-					SNew(SImage)
-					.Image(FEditorStyle::GetBrush("GenericViewButton"))
-					.ColorAndOpacity(FNiagaraEditorStyle::Get().GetColor("NiagaraEditor.Stack.FlatButtonColor"))
-				]
-			]
-		];
-
 	// Finalize the widget
 	ChildSlot
 	[
