Removed unused style access that also didn't exist

diff --git a/Engine/Plugins/FX/Niagara/Source/NiagaraEditor/Private/Widgets/SNiagaraParameterPanel.cpp b/Engine/Plugins/FX/Niagara/Source/NiagaraEditor/Private/Widgets/SNiagaraParameterPanel.cpp
--- a/Engine/Plugins/FX/Niagara/Source/NiagaraEditor/Private/Widgets/SNiagaraParameterPanel.cpp
+++ b/Engine/Plugins/FX/Niagara/Source/NiagaraEditor/Private/Widgets/SNiagaraParameterPanel.cpp
@@ -58,7 +58,6 @@ void SNiagaraParameterPanel::Construct(const FArguments& InArgs, const TSharedPt
 	ToolkitCommands = InToolkitCommands;
 	bPendingRefresh = false;
 	bParameterItemsPendingChange = false;
-	const FVector2D ViewOptionsShadowOffset = FNiagaraEditorStyle::Get().GetVector("NiagaraEditor.Stack.ViewOptionsShadowOffset");
 	ParametersWithNamespaceModifierRenamePending = MakeShared<TArray<FName>>();
 
 	ParameterPanelViewModel->GetOnRequestRefreshDelegate().BindSP(this, &SNiagaraParameterPanel::Refresh);
