Static Analysis Warning Fixes

diff --git a/Engine/Plugins/FX/Niagara/Source/NiagaraEditor/Private/NiagaraCompiler.cpp b/Engine/Plugins/FX/Niagara/Source/NiagaraEditor/Private/NiagaraCompiler.cpp
--- a/Engine/Plugins/FX/Niagara/Source/NiagaraEditor/Private/NiagaraCompiler.cpp
+++ b/Engine/Plugins/FX/Niagara/Source/NiagaraEditor/Private/NiagaraCompiler.cpp
@@ -1409,7 +1409,10 @@ TSharedPtr<FNiagaraGraphCachedDataBase, ESPMode::ThreadSafe> FNiagaraEditorModul
 			// Map all for this output node
 			FNiagaraParameterMapHistoryWithMetaDataBuilder Builder;
 			Builder.ConstantResolver = ConstantResolver;
-			Builder.AddGraphToCallingGraphContextStack(ScriptSource->NodeGraph);
+			if (ScriptSource != nullptr)
+			{
+				Builder.AddGraphToCallingGraphContextStack(ScriptSource->NodeGraph);
+			}
 			Builder.RegisterEncounterableVariables(EncounterableVariables);
 			Builder.RegisterExternalStaticVariables(StaticVariables);
 
diff --git a/Engine/Plugins/FX/Niagara/Source/NiagaraEditor/Private/NiagaraHlslTranslator.cpp b/Engine/Plugins/FX/Niagara/Source/NiagaraEditor/Private/NiagaraHlslTranslator.cpp
--- a/Engine/Plugins/FX/Niagara/Source/NiagaraEditor/Private/NiagaraHlslTranslator.cpp
+++ b/Engine/Plugins/FX/Niagara/Source/NiagaraEditor/Private/NiagaraHlslTranslator.cpp
@@ -5385,6 +5385,9 @@ void FHlslNiagaraTranslator::FillVariableWithDefaultValue(int32& OutValue, const
 
 void FHlslNiagaraTranslator::SetConstantByStaticVariable(int32& OutValue, const UEdGraphPin* InDefaultPin, FString* OutDebugString)
 {
+	if (InDefaultPin == nullptr)
+		return;
+
 	FNiagaraEditorModule& NiagaraEditorModule = FModuleManager::GetModuleChecked<FNiagaraEditorModule>("NiagaraEditor");
 	OutValue = 0;
 	FNiagaraVariable Var = Schema->PinToNiagaraVariable(InDefaultPin, true);
@@ -6498,7 +6501,6 @@ void FHlslNiagaraTranslator::Operation(class UNiagaraNodeOp* Operation, TArray<i
 
 	// Use the pins to determine the output type here since they may have been changed due to numeric pin fix up.
 	const FNiagaraOpInfo* OpInfo = FNiagaraOpInfo::GetOpInfo(Operation->OpName);
-
 	FPinCollectorArray OutputPins;
 	Operation->GetOutputPins(OutputPins);
 	OutputPins.RemoveAll([](UEdGraphPin* Pin)
@@ -6590,23 +6592,25 @@ void FHlslNiagaraTranslator::Operation(class UNiagaraNodeOp* Operation, TArray<i
 			FText PinNameText = OutputPin->PinFriendlyName.IsEmpty() ? FText::FromName(OutputPin->PinName) : OutputPin->PinFriendlyName;
 			Error(FText::Format(LOCTEXT("GetConstantFailTypePin", "Cannot handle type {0}! Output Pin: {1}"), OutputType.GetNameText(), PinNameText), Operation, OutputPin);
 		}
-
-		const FNiagaraOpInOutInfo& IOInfo = OpInfo->Outputs[OutputIndex];
-		FString OutputHlsl;
-		if (OpInfo->bSupportsAddedInputs)
+		if (OpInfo != nullptr)
 		{
-			if (!OpInfo->CreateHlslForAddedInputs(Inputs.Num(), OutputHlsl))
+			const FNiagaraOpInOutInfo& IOInfo = OpInfo->Outputs[OutputIndex];
+			FString OutputHlsl;
+			if (OpInfo->bSupportsAddedInputs)
 			{
-				FText PinNameText = OutputPin->PinFriendlyName.IsEmpty() ? FText::FromName(OutputPin->PinName) : OutputPin->PinFriendlyName;
-				Error(FText::Format(LOCTEXT("AggregateInputFailTypePin", "Cannot create hlsl output for type {0}! Output Pin: {1}"), OutputType.GetNameText(), PinNameText), Operation, OutputPin);
+				if (!OpInfo->CreateHlslForAddedInputs(Inputs.Num(), OutputHlsl))
+				{
+					FText PinNameText = OutputPin->PinFriendlyName.IsEmpty() ? FText::FromName(OutputPin->PinName) : OutputPin->PinFriendlyName;
+					Error(FText::Format(LOCTEXT("AggregateInputFailTypePin", "Cannot create hlsl output for type {0}! Output Pin: {1}"), OutputType.GetNameText(), PinNameText), Operation, OutputPin);
+					OutputHlsl = IOInfo.HlslSnippet;
+				}
+			}
+			else {
 				OutputHlsl = IOInfo.HlslSnippet;
 			}
+			check(!OutputHlsl.IsEmpty());
+			Outputs.Add(AddBodyChunk(GetUniqueSymbolName(IOInfo.Name), OutputHlsl, OutputType, Inputs));
 		}
-		else {
-			OutputHlsl = IOInfo.HlslSnippet;
-		}
-		check(!OutputHlsl.IsEmpty());
-		Outputs.Add(AddBodyChunk(GetUniqueSymbolName(IOInfo.Name), OutputHlsl, OutputType, Inputs));
 	}
 }
 
