- Fix for using incorrect RW type

#rnx

diff --git a/Engine/Plugins/FX/Niagara/Source/Niagara/Private/NiagaraDataInterfaceRenderTargetCube.cpp b/Engine/Plugins/FX/Niagara/Source/Niagara/Private/NiagaraDataInterfaceRenderTargetCube.cpp
--- a/Engine/Plugins/FX/Niagara/Source/Niagara/Private/NiagaraDataInterfaceRenderTargetCube.cpp
+++ b/Engine/Plugins/FX/Niagara/Source/Niagara/Private/NiagaraDataInterfaceRenderTargetCube.cpp
@@ -399,7 +399,7 @@ void UNiagaraDataInterfaceRenderTargetCube::GetParameterDefinitionHLSL(const FNi
 	Super::GetParameterDefinitionHLSL(ParamInfo, OutHLSL);
 
 	static const TCHAR *FormatDeclarations = TEXT(R"(				
-		RWTexture2DArray<float4> {OutputName};
+		RWTextureCube<float4> {OutputName};
 		TextureCube<float4> {InputName};
 		SamplerState {InputName}SamplerState;
 		int {SizeName};
diff --git a/Engine/Plugins/FX/Niagara/Source/NiagaraCore/Private/NiagaraCustomVersion.cpp b/Engine/Plugins/FX/Niagara/Source/NiagaraCore/Private/NiagaraCustomVersion.cpp
--- a/Engine/Plugins/FX/Niagara/Source/NiagaraCore/Private/NiagaraCustomVersion.cpp
+++ b/Engine/Plugins/FX/Niagara/Source/NiagaraCore/Private/NiagaraCustomVersion.cpp
@@ -8,4 +8,4 @@ const FGuid FNiagaraCustomVersion::GUID(0xFCF57AFA, 0x50764283, 0xB9A9E658, 0xFF
 // Register the custom version with core
 FCustomVersionRegistration GRegisterNiagaraCustomVersion(FNiagaraCustomVersion::GUID, FNiagaraCustomVersion::LatestVersion, TEXT("NiagaraVer"));
 
-const FGuid FNiagaraCustomVersion::LatestScriptCompileVersion(0x5B30E7F5, 0xB45C4092, 0xB60C84DB, 0x8D7C5416);
+const FGuid FNiagaraCustomVersion::LatestScriptCompileVersion(0x9316EE10, 0x4B7FF347, 0xB781A33C, 0xC4E24B97);
