- Allow data interface to specify how they wish to dispatch (1d / 2d / 3d)
- Clean up shader code generation between particle / iteration interfaces to improve performance

diff --git a/Engine/Plugins/FX/Niagara/Shaders/Private/NiagaraEmitterInstanceShader.usf b/Engine/Plugins/FX/Niagara/Shaders/Private/NiagaraEmitterInstanceShader.usf
--- a/Engine/Plugins/FX/Niagara/Shaders/Private/NiagaraEmitterInstanceShader.usf
+++ b/Engine/Plugins/FX/Niagara/Shaders/Private/NiagaraEmitterInstanceShader.usf
@@ -63,24 +63,29 @@ bool NiagaraAny(bool4 b) { return b.x || b.y || b.z || b.w; }
 		#define NIAGARA_PARTICLE_PARTIAL_ENABLED 0
 	#endif
 
-	// Used to convert from DispatchThreadId to a linear thread id
-	uint DispatchThreadIdToLinear;
+	// Type of dispatch mode we are in this allows us to customize some HLSL code to avoid / % operations
+	#define NIAGARA_DISPATCH_TYPE_ONE_D		0
+	#define NIAGARA_DISPATCH_TYPE_TWO_D		1
+	#define NIAGARA_DISPATCH_TYPE_THREE_D	2
+	#define NIAGARA_DISPATCH_TYPE_CUSTOM	3
+	#ifndef NIAGARA_DISPATCH_TYPE
+		#define NIAGARA_DISPATCH_TYPE NIAGARA_DISPATCH_TYPE_ONE_D
+	#endif
+
+	uint3 DispatchThreadIdBounds;		// Used to track thread bounds
+	uint3 DispatchThreadIdToLinear;		// Used to convert from DispatchThreadId to a linear thread id
 
 	static uint GLinearThreadId;
 	static uint3 GDispatchThreadId;
 	static uint3 GGroupThreadId;
-	static uint GCurrentPhase;
 	static uint GEmitterTickCounter;
-	static uint GSimStart;
 	static uint GRandomSeedOffset = 0;
 	
-	const static uint GSpawnPhase = 0;
-	const static uint GUpdatePhase = 1;
-
 	// To be removed as well but don't know who is using it
 	#include "/Engine/Private/SceneTexturesCommon.ush"
 
 	// Physics common 
+	//-TOFIX: This should not be included here
 	#if USE_SIMULATION_STAGES == 1
 		#include "NiagaraPhysicsCommon.ush"
 	#endif
@@ -128,17 +133,16 @@ bool NiagaraAny(bool4 b) { return b.x || b.y || b.z || b.w; }
 
 	static uint GSpawnStartInstance;
 	uint SpawnedInstances;
-	uint UpdateStartInstance;
 
 	void SetupExecIndexForGPU()
 	{
-		GGPUExecIndex = UpdateStartInstance + GLinearThreadId;
+		GGPUExecIndex = GLinearThreadId;
 		Engine_ExecutionCount = GSpawnStartInstance;
 	}
 
 	void SetupExecIndexAndSpawnInfoForGPU()
 	{
-		GGPUExecIndex = UpdateStartInstance + GLinearThreadId - GSpawnStartInstance;
+		GGPUExecIndex = GLinearThreadId - GSpawnStartInstance;
 
 		int SpawnInfoIndex = 0;
 		UNROLL
@@ -540,8 +544,8 @@ int rand_int(int x, int Seed1, int Seed2, int Seed3)
 		// End static assert
 
 		// This is the same as ExecIndex() right now, but that function may change in the future to accommodate multiple
-		// spawn infos. Revisit this computation if the change affects the meaning of UpdateStartInstance or GSpawnStartInstance.
-		int SpawnIndex = UpdateStartInstance + GLinearThreadId - GSpawnStartInstance;
+		// spawn infos. Revisit this computation if the change affects the meaning of GSpawnStartInstance.
+		int SpawnIndex = GLinearThreadId - GSpawnStartInstance;
 		IDIndex = FreeIDList[SpawnIndex];
 		IDAcquireTag = EmitterTickCounter;
 	}
@@ -736,22 +740,22 @@ int rand_int(int x, int Seed1, int Seed2, int Seed3)
 	 */
 	void OutputDataFloat(int DataSetIndex, int RegisterIndex, int InstanceIndex, float Value)
 	{
-		RWOutputFloat[RegisterIndex*ComponentBufferSizeWrite + InstanceIndex + UpdateStartInstance] = Value;
+		RWOutputFloat[RegisterIndex*ComponentBufferSizeWrite + InstanceIndex] = Value;
 	}
 
 	void OutputDataInt(int DataSetIndex, int RegisterIndex, int InstanceIndex, int Value)
 	{
-		RWOutputInt[RegisterIndex*ComponentBufferSizeWrite + InstanceIndex + UpdateStartInstance] = Value;
+		RWOutputInt[RegisterIndex*ComponentBufferSizeWrite + InstanceIndex] = Value;
 	}
 
 	void OutputDataBool(int DataSetIndex, int RegisterIndex, int InstanceIndex, bool Value)
 	{
-		RWOutputInt[RegisterIndex*ComponentBufferSizeWrite + InstanceIndex + UpdateStartInstance] = Value ? -1 : 0;
+		RWOutputInt[RegisterIndex*ComponentBufferSizeWrite + InstanceIndex] = Value ? -1 : 0;
 	}
 
 	void OutputDataHalf(int DataSetIndex, int RegisterIndex, int InstanceIndex, float Value)
 	{
-		RWOutputHalf[RegisterIndex*ComponentBufferSizeWrite + InstanceIndex + UpdateStartInstance] = Value;
+		RWOutputHalf[RegisterIndex*ComponentBufferSizeWrite + InstanceIndex] = Value;
 	}
 
 	void EnterStatScope(int ID)	{}
@@ -769,7 +773,6 @@ int rand_int(int x, int Seed1, int Seed2, int Seed3)
 			// other branchless construct), because we don't want to call AcquireIndex() at all, since
 			// that manipulates the RWInstanceCounts UAV. The generated code will copy the source count
 			// at the end of the shader.
-			BRANCH
 			if (!bStageKillsParticles)
 			{
 				return GLinearThreadId;
diff --git a/Engine/Plugins/FX/Niagara/Source/Niagara/Classes/NiagaraDataInterface.h b/Engine/Plugins/FX/Niagara/Source/Niagara/Classes/NiagaraDataInterface.h
--- a/Engine/Plugins/FX/Niagara/Source/Niagara/Classes/NiagaraDataInterface.h
+++ b/Engine/Plugins/FX/Niagara/Source/Niagara/Classes/NiagaraDataInterface.h
@@ -8,6 +8,7 @@
 #include "Curves/RichCurve.h"
 #include "NiagaraMergeable.h"
 #include "NiagaraDataInterfaceBase.h"
+#include "NiagaraScriptBase.h"
 #include "NiagaraDataInterface.generated.h"
 
 class INiagaraCompiler;
@@ -291,6 +292,9 @@ public:
 	virtual bool GenerateIterationSourceNamespaceWriteAttributesHLSL(FNiagaraDataInterfaceGPUParamInfo& DIInstanceInfo, const FNiagaraVariable& InIterationSourceVariable, TConstArrayView<FNiagaraVariable> InArguments, TConstArrayView<FNiagaraVariable> InAttributes, TConstArrayView<FString> InAttributeHLSLNames, bool bSpawnOnly, bool bPartialWrites, TArray<FText>& OutErrors, FString& OutHLSL) const { return false; };
 	/** Used by the translator when dealing with signatures that turn into compiler tags to figure out the precise compiler tag. */
 	virtual bool GenerateCompilerTagPrefix(const FNiagaraFunctionSignature& InSignature, FString& OutPrefix) const  { return false; }
+
+	virtual ENiagaraGpuDispatchType GetGpuDispatchType() const { return ENiagaraGpuDispatchType::OneD; }
+	virtual FIntVector GetGpuDispatchNumThreads() const { return FIntVector(64, 1, 1); }
 #endif
 
 	/** Allows data interfaces to cache any static data that may be shared between instances */
diff --git a/Engine/Plugins/FX/Niagara/Source/Niagara/Private/NiagaraDataInterface.cpp b/Engine/Plugins/FX/Niagara/Source/Niagara/Private/NiagaraDataInterface.cpp
--- a/Engine/Plugins/FX/Niagara/Source/Niagara/Private/NiagaraDataInterface.cpp
+++ b/Engine/Plugins/FX/Niagara/Source/Niagara/Private/NiagaraDataInterface.cpp
@@ -34,6 +34,15 @@ UNiagaraDataInterface::~UNiagaraDataInterface()
 #if WITH_EDITORONLY_DATA
 bool UNiagaraDataInterface::AppendCompileHash(FNiagaraCompileHashVisitor* InVisitor) const
 {
+	//-TODO: Currently applied to all, but we only need to hash this in for the iteration source
+	const UNiagaraDataInterface* BaseDataInterface = GetDefault<UNiagaraDataInterface>();
+	if (BaseDataInterface->GetGpuDispatchType() != GetGpuDispatchType())
+	{
+		const FString DataInterfaceName = GetClass()->GetName();
+		InVisitor->UpdatePOD(*FString::Printf(TEXT("%s_GpuDispatchType"), *DataInterfaceName), (int32)GetGpuDispatchType());
+		InVisitor->UpdatePOD(*FString::Printf(TEXT("%s_GpuDispatchNumThreads"), *DataInterfaceName), *FString::Printf(TEXT("%dx%dx%d"), GetGpuDispatchNumThreads().X, GetGpuDispatchNumThreads().Y, GetGpuDispatchNumThreads().Z));
+	}
+
 	return true;
 }
 #endif
diff --git a/Engine/Plugins/FX/Niagara/Source/Niagara/Private/NiagaraDataInterfaceRW.cpp b/Engine/Plugins/FX/Niagara/Source/Niagara/Private/NiagaraDataInterfaceRW.cpp
--- a/Engine/Plugins/FX/Niagara/Source/Niagara/Private/NiagaraDataInterfaceRW.cpp
+++ b/Engine/Plugins/FX/Niagara/Source/Niagara/Private/NiagaraDataInterfaceRW.cpp
@@ -401,35 +401,45 @@ bool UNiagaraDataInterfaceGrid3D::GetFunctionHLSL(const FNiagaraDataInterfaceGPU
 	}
 	else if (FunctionInfo.DefinitionName == UNiagaraDataInterfaceRWBase::ExecutionIndexToUnitFunctionName)
 	{
-	static const TCHAR* FormatSample = TEXT(R"(
+		static const TCHAR* FormatSample = TEXT(R"(
 			void {FunctionName}(out float3 Out_Unit)
 			{
-				const uint Linear = GLinearThreadId;
-				const uint IndexX = Linear % {NumCellsName}.x;
-				const uint IndexY = (Linear / {NumCellsName}.x) % {NumCellsName}.y;
-				const uint IndexZ = Linear / ({NumCellsName}.x * {NumCellsName}.y);				
-
-				Out_Unit = (float3(IndexX, IndexY, IndexZ) + .5) / {NumCellsName};				
+				#if NIAGARA_DISPATCH_TYPE == NIAGARA_DISPATCH_TYPE_THREE_D
+					Out_Unit = (float3(GDispatchThreadId.x, GDispatchThreadId.y, GDispatchThreadId.z) + .5) / {NumCellsName};
+				#else
+					const uint Linear = GLinearThreadId;
+					const uint IndexX = Linear % {NumCellsName}.x;
+					const uint IndexY = (Linear / {NumCellsName}.x) % {NumCellsName}.y;
+					const uint IndexZ = Linear / ({NumCellsName}.x * {NumCellsName}.y);				
+
+					Out_Unit = (float3(IndexX, IndexY, IndexZ) + .5) / {NumCellsName};				
+				#endif
 			}
 		)");
 
-	OutHLSL += FString::Format(FormatSample, ArgsDeclarations);
-	return true;
+		OutHLSL += FString::Format(FormatSample, ArgsDeclarations);
+		return true;
 	}
 	else if (FunctionInfo.DefinitionName == UNiagaraDataInterfaceRWBase::ExecutionIndexToGridIndexFunctionName)
 	{
-	static const TCHAR* FormatSample = TEXT(R"(
+		static const TCHAR* FormatSample = TEXT(R"(
 			void {FunctionName}(out int Out_IndexX, out int Out_IndexY, out int Out_IndexZ)
 			{
-				const uint Linear = GLinearThreadId;
-				Out_IndexX = Linear % {NumCellsName}.x;
-				Out_IndexY = (Linear / {NumCellsName}.x) % {NumCellsName}.y;
-				Out_IndexZ = Linear / ({NumCellsName}.x * {NumCellsName}.y);
+				#if NIAGARA_DISPATCH_TYPE == NIAGARA_DISPATCH_TYPE_THREE_D
+					Out_IndexX = GDispatchThreadId.x;
+					Out_IndexY = GDispatchThreadId.y;
+					Out_IndexZ = GDispatchThreadId.z;
+				#else
+					const uint Linear = GLinearThreadId;
+					Out_IndexX = Linear % {NumCellsName}.x;
+					Out_IndexY = (Linear / {NumCellsName}.x) % {NumCellsName}.y;
+					Out_IndexZ = Linear / ({NumCellsName}.x * {NumCellsName}.y);
+				#endif
 			}
 		)");
 
-	OutHLSL += FString::Format(FormatSample, ArgsDeclarations);
-	return true;
+		OutHLSL += FString::Format(FormatSample, ArgsDeclarations);
+		return true;
 	}
 	else if (FunctionInfo.DefinitionName == UNiagaraDataInterfaceRWBase::CellSizeFunctionName)
 	{
@@ -882,33 +892,41 @@ bool UNiagaraDataInterfaceGrid2D::GetFunctionHLSL(const FNiagaraDataInterfaceGPU
 	}
 	else if (FunctionInfo.DefinitionName == UNiagaraDataInterfaceRWBase::ExecutionIndexToUnitFunctionName)
 	{
-	static const TCHAR* FormatSample = TEXT(R"(
+		static const TCHAR* FormatSample = TEXT(R"(
 			void {FunctionName}(out float2 Out_Unit)
 			{
-				const uint Linear = GLinearThreadId;
-				const uint IndexX = Linear % {NumCellsName}.x;
-				const uint IndexY = Linear / {NumCellsName}.x;				
-
-				Out_Unit = (float2(IndexX, IndexY) + .5) * {UnitToUVName};			
+				#if NIAGARA_DISPATCH_TYPE == NIAGARA_DISPATCH_TYPE_TWO_D
+					Out_Unit = (float2(GDispatchThreadId.x, GDispatchThreadId.y) + .5) * {UnitToUVName};			
+				#else
+					const uint Linear = GLinearThreadId;
+					const uint IndexX = Linear % {NumCellsName}.x;
+					const uint IndexY = Linear / {NumCellsName}.x;				
+					Out_Unit = (float2(IndexX, IndexY) + .5) * {UnitToUVName};			
+				#endif
 			}
 		)");
 
-	OutHLSL += FString::Format(FormatSample, ArgsDeclarations);
-	return true;
+		OutHLSL += FString::Format(FormatSample, ArgsDeclarations);
+		return true;
 	}
 	else if (FunctionInfo.DefinitionName == UNiagaraDataInterfaceRWBase::ExecutionIndexToGridIndexFunctionName)
 	{
-	static const TCHAR* FormatSample = TEXT(R"(
+		static const TCHAR* FormatSample = TEXT(R"(
 			void {FunctionName}(out int Out_IndexX, out int Out_IndexY)
 			{
-				const uint Linear = GLinearThreadId;
-				Out_IndexX = Linear % {NumCellsName}.x;
-				Out_IndexY = Linear / {NumCellsName}.x;				
+				#if NIAGARA_DISPATCH_TYPE == NIAGARA_DISPATCH_TYPE_TWO_D
+					Out_IndexX = GDispatchThreadId.x;
+					Out_IndexY = GDispatchThreadId.y;
+				#else
+					const uint Linear = GLinearThreadId;
+					Out_IndexX = Linear % {NumCellsName}.x;
+					Out_IndexY = Linear / {NumCellsName}.x;				
+				#endif
 			}
 		)");
 
-	OutHLSL += FString::Format(FormatSample, ArgsDeclarations);
-	return true;
+		OutHLSL += FString::Format(FormatSample, ArgsDeclarations);
+		return true;
 	}
 	else if (FunctionInfo.DefinitionName == UNiagaraDataInterfaceRWBase::CellSizeFunctionName)
 	{
diff --git a/Engine/Plugins/FX/Niagara/Source/Niagara/Private/NiagaraGpuComputeDispatch.cpp b/Engine/Plugins/FX/Niagara/Source/Niagara/Private/NiagaraGpuComputeDispatch.cpp
--- a/Engine/Plugins/FX/Niagara/Source/Niagara/Private/NiagaraGpuComputeDispatch.cpp
+++ b/Engine/Plugins/FX/Niagara/Source/Niagara/Private/NiagaraGpuComputeDispatch.cpp
@@ -1467,34 +1467,52 @@ void FNiagaraGpuComputeDispatch::DispatchStage(FRHICommandList& RHICmdList, FRHI
 	}
 
 	// Get dispatch count
-	uint32 DispatchCount = 0;
-	if ( SimStageData.AlternateIterationSource )
+	ENiagaraGpuDispatchType DispatchType = ENiagaraGpuDispatchType::OneD;
+	FIntVector DispatchCount = FIntVector(0, 0, 0);
+	FIntVector DispatchNumThreads = FIntVector(0, 0, 0);
+	if (SimStageData.AlternateIterationSource)
 	{
-		const FIntVector ElementCount = SimStageData.AlternateIterationSource->GetElementCount(Tick.SystemInstanceID);
+		DispatchType = SimStageData.StageMetaData->GpuDispatchType;
+		DispatchCount = SimStageData.AlternateIterationSource->GetElementCount(Tick.SystemInstanceID);
+		DispatchNumThreads = (DispatchType == ENiagaraGpuDispatchType::Custom) ? SimStageData.StageMetaData->GpuDispatchNumThreads : FNiagaraShader::GetDefaultThreadGroupSize(DispatchType);
+
 		// Verify the number of elements isn't higher that what we can handle
-		checkf(ElementCount.X * ElementCount.Y * ElementCount.Z < uint64(TNumericLimits<uint32>::Max()), TEXT("ElementCount(%d, %d, %d) for IterationInterface(%s) overflows an int32 this is not allowed"), ElementCount.X, ElementCount.Y, ElementCount.Z, *SimStageData.AlternateIterationSource->SourceDIName.ToString());
+		checkf(uint64(DispatchCount.X) * uint64(DispatchCount.Y) * uint64(DispatchCount.Z) < uint64(TNumericLimits<int32>::Max()), TEXT("DispatchCount(%d, %d, %d) for IterationInterface(%s) overflows an int32 this is not allowed"), DispatchCount.X, DispatchCount.Y, DispatchCount.Z, *SimStageData.AlternateIterationSource->SourceDIName.ToString());
 
-		DispatchCount = ElementCount.X * ElementCount.Y * ElementCount.Z;
+		// Data interfaces such as grids / render targets can choose to dispatch in either the correct dimensionality for the target (i.e. RT2D would choose 2D)
+		// or run in linear mode if performance is not beneficial due to increased waves.  It is also possible the we may choose to override on the simulation stage.
+		// Therefore we need to special case OneD and convert our element count back to linear.
+		if (DispatchType == ENiagaraGpuDispatchType::OneD)
+		{
+			DispatchCount.X = DispatchCount.X * DispatchCount.Y * DispatchCount.Z;
+			DispatchCount.Y = 1;
+			DispatchCount.Z = 1;
+		}
 	}
 	else
 	{
-		DispatchCount = SimStageData.DestinationNumInstances;
+		DispatchType = ENiagaraGpuDispatchType::OneD;
+		DispatchCount = FIntVector(SimStageData.DestinationNumInstances, 1, 1);
+		DispatchNumThreads = FNiagaraShader::GetDefaultThreadGroupSize(ENiagaraGpuDispatchType::OneD);
 	}
 
-	if ( DispatchCount == 0 )
+	const int32 TotalDispatchCount = DispatchCount.X * DispatchCount.Y * DispatchCount.Z;
+	if (TotalDispatchCount == 0)
 	{
 		return;
 	}
 
+	checkf(DispatchNumThreads.X * DispatchNumThreads.Y * DispatchNumThreads.Z > 0, TEXT("DispatchNumThreads(%d, %d, %d) is invalid"), DispatchNumThreads.X, DispatchNumThreads.Y, DispatchNumThreads.Z);
+
 	// Get Shader
 	const FNiagaraShaderRef ComputeShader = InstanceData.Context->GPUScript_RT->GetShader(SimStageData.StageIndex);
 	FRHIComputeShader* RHIComputeShader = ComputeShader.GetComputeShader();
 	SetComputePipelineState(RHICmdList, RHIComputeShader);
 
 	SCOPED_DRAW_EVENTF(RHICmdList, NiagaraGPUSimulationCS,
-		TEXT("NiagaraGpuSim(%s) DispatchCount(%u) Stage(%s %u) Iteration(%d) NumInstructions(%u)"),
+		TEXT("NiagaraGpuSim(%s) DispatchCount(%d x %d x %d) Stage(%s %u) NumInstructions(%u)"),
 		InstanceData.Context->GetDebugSimName(),
-		DispatchCount,
+		DispatchCount.X, DispatchCount.Y, DispatchCount.Z,
 		*SimStageData.StageMetaData->SimulationStageName.ToString(),
 		SimStageData.StageIndex,
 		SimStageData.IterationIndex,
@@ -1507,7 +1525,6 @@ void FNiagaraGpuComputeDispatch::DispatchStage(FRHICommandList& RHICmdList, FRHI
 
 	SetShaderValue(RHICmdList, RHIComputeShader, ComputeShader->SimStartParam, InstanceData.bResetData ? 1U : 0U);
 	SetShaderValue(RHICmdList, RHIComputeShader, ComputeShader->EmitterTickCounterParam, FNiagaraComputeExecutionContext::TickCounter);
-	SetShaderValue(RHICmdList, RHIComputeShader, ComputeShader->UpdateStartInstanceParam, 0);
 	SetShaderValue(RHICmdList, RHIComputeShader, ComputeShader->NumSpawnedInstancesParam, InstancesToSpawn);
 	SetSRVParameter(RHICmdList, RHIComputeShader, ComputeShader->FreeIDBufferParam, bRequiresPersistentIDs ? InstanceData.Context->MainDataSet->GetGPUFreeIDs().SRV.GetReference() : FNiagaraRenderer::GetDummyIntBuffer());
 
@@ -1557,7 +1574,7 @@ void FNiagaraGpuComputeDispatch::DispatchStage(FRHICommandList& RHICmdList, FRHI
 		{
 			const uint32 IterationInstanceCountOffset = SimStageData.AlternateIterationSource->GetGPUInstanceCountOffset(Tick.SystemInstanceID);
 			SimulationStageIterationInfo.X = IterationInstanceCountOffset;
-			SimulationStageIterationInfo.Y = IterationInstanceCountOffset == INDEX_NONE ? DispatchCount : 0;
+			SimulationStageIterationInfo.Y = IterationInstanceCountOffset == INDEX_NONE ? TotalDispatchCount : 0;
 		}
 
 		const int32 NumIterations = InstanceData.NumIterationsPerStage[SimStageData.StageIndex];
@@ -1582,19 +1599,25 @@ void FNiagaraGpuComputeDispatch::DispatchStage(FRHICommandList& RHICmdList, FRHI
 
 	// Execute the dispatch
 	{
-		const uint32 ThreadGroupSize = FNiagaraShader::GetGroupSize(ShaderPlatform);
-		checkf(ThreadGroupSize <= NiagaraComputeMaxThreadGroupSize, TEXT("ShaderPlatform(%d) is set to ThreadGroup size (%d) which is > the NiagaraComputeMaxThreadGroupSize(%d)"), ShaderPlatform, ThreadGroupSize, NiagaraComputeMaxThreadGroupSize);
-		const uint32 TotalThreadGroups = FMath::DivideAndRoundUp(DispatchCount, ThreadGroupSize);
+		// In the OneD case we can use the Y dimension to get higher particle counts
+		if (DispatchType == ENiagaraGpuDispatchType::OneD)
+		{
+			const int32 TotalThreadGroups = FMath::DivideAndRoundUp(DispatchCount.X, DispatchNumThreads.X);
+			DispatchCount.Y = FMath::DivideAndRoundUp<int32>(TotalThreadGroups, NiagaraMaxThreadGroupCountPerDimension);
+			DispatchCount.X = FMath::DivideAndRoundUp(DispatchCount.X, DispatchNumThreads.Y);
+		}
 
-		const uint32 ThreadGroupCountZ = 1;
-		const uint32 ThreadGroupCountY = FMath::DivideAndRoundUp(TotalThreadGroups, NiagaraMaxThreadGroupCountPerDimension);
-		const uint32 ThreadGroupCountX = FMath::DivideAndRoundUp(TotalThreadGroups, ThreadGroupCountY);
+		FIntVector ThreadGroupCount;
+		ThreadGroupCount.X = FMath::DivideAndRoundUp(DispatchCount.X, DispatchNumThreads.X);
+		ThreadGroupCount.Y = FMath::DivideAndRoundUp(DispatchCount.Y, DispatchNumThreads.Y);
+		ThreadGroupCount.Z = FMath::DivideAndRoundUp(DispatchCount.Z, DispatchNumThreads.Z);
 
-		SetShaderValue(RHICmdList, RHIComputeShader, ComputeShader->DispatchThreadIdToLinearParam, ThreadGroupCountY > 1 ? ThreadGroupCountX * ThreadGroupSize : 0);
+		SetShaderValue(RHICmdList, RHIComputeShader, ComputeShader->DispatchThreadIdToLinearParam, FIntVector(1, DispatchCount.X, DispatchCount.X * DispatchCount.Y));
+		SetShaderValue(RHICmdList, RHIComputeShader, ComputeShader->DispatchThreadIdBoundsParam, DispatchCount);
 
 		SetConstantBuffers(RHICmdList, ComputeShader, Tick, &InstanceData);
 
-		DispatchComputeShader(RHICmdList, ComputeShader, ThreadGroupCountX, ThreadGroupCountY, ThreadGroupCountZ);
+		DispatchComputeShader(RHICmdList, ComputeShader, ThreadGroupCount.X, ThreadGroupCount.Y, ThreadGroupCount.Z);
 
 		INC_DWORD_STAT(STAT_NiagaraGPUDispatches);
 	}
diff --git a/Engine/Plugins/FX/Niagara/Source/NiagaraCore/Private/NiagaraCustomVersion.cpp b/Engine/Plugins/FX/Niagara/Source/NiagaraCore/Private/NiagaraCustomVersion.cpp
--- a/Engine/Plugins/FX/Niagara/Source/NiagaraCore/Private/NiagaraCustomVersion.cpp
+++ b/Engine/Plugins/FX/Niagara/Source/NiagaraCore/Private/NiagaraCustomVersion.cpp
@@ -8,4 +8,4 @@ const FGuid FNiagaraCustomVersion::GUID(0xFCF57AFA, 0x50764283, 0xB9A9E658, 0xFF
 // Register the custom version with core
 FCustomVersionRegistration GRegisterNiagaraCustomVersion(FNiagaraCustomVersion::GUID, FNiagaraCustomVersion::LatestVersion, TEXT("NiagaraVer"));
 
-const FGuid FNiagaraCustomVersion::LatestScriptCompileVersion(0x9316EE10, 0x4B7FF347, 0xB781A33C, 0xC4E24B97);
+const FGuid FNiagaraCustomVersion::LatestScriptCompileVersion(0xA73D5647, 0xED387740, 0x93DF7180, 0xD1FD586C);
diff --git a/Engine/Plugins/FX/Niagara/Source/NiagaraEditor/Private/NiagaraHlslTranslator.cpp b/Engine/Plugins/FX/Niagara/Source/NiagaraEditor/Private/NiagaraHlslTranslator.cpp
--- a/Engine/Plugins/FX/Niagara/Source/NiagaraEditor/Private/NiagaraHlslTranslator.cpp
+++ b/Engine/Plugins/FX/Niagara/Source/NiagaraEditor/Private/NiagaraHlslTranslator.cpp
@@ -1286,6 +1286,19 @@ const FNiagaraTranslateResults &FHlslNiagaraTranslator::Translate(const FNiagara
 					SimulationStageMetaData.bWritesParticles = TranslationStages[Index].bWritesParticles;
 					SimulationStageMetaData.bPartialParticleUpdate = TranslationStages[Index].bPartialParticleUpdate;
 
+					// Determine dispatch information from iteration source (if we have one)
+					if ( !SimulationStageMetaData.IterationSource.IsNone() )
+					{
+						if (const FNiagaraVariable* IterationSourceVar = CompileData->EncounteredVariables.FindByPredicate([&](const FNiagaraVariable& VarInfo) { return VarInfo.GetName() == SimulationStageMetaData.IterationSource; }))
+						{
+							if ( UNiagaraDataInterface* IteratinoSourceCDO = CompileDuplicateData->GetDuplicatedDataInterfaceCDOForClass(IterationSourceVar->GetType().GetClass()) )
+							{
+								SimulationStageMetaData.GpuDispatchType = IteratinoSourceCDO->GetGpuDispatchType();
+								SimulationStageMetaData.GpuDispatchNumThreads = IteratinoSourceCDO->GetGpuDispatchNumThreads();
+							}
+						}
+					}
+
 					// Increment source stage index
 					++SourceSimStageIndex;
 
@@ -2984,7 +2997,7 @@ void FHlslNiagaraTranslator::DefineMainGPUFunctions(
 	/////////////////////////
 	// StoreUpdateVariables()
 
-	HlslOutput += TEXT("void StoreUpdateVariables(in FSimulationContext Context)\n{\n");
+	HlslOutput += TEXT("void StoreUpdateVariables(in FSimulationContext Context, bool bIsValidInstance)\n{\n");
 	{
 		FExpressionPermutationContext PermutationContext(HlslOutput);
 
@@ -3027,13 +3040,13 @@ void FHlslNiagaraTranslator::DefineMainGPUFunctions(
 				// under dynamic flow control, which is not allowed. Therefore, we must always use the more expensive path when the spawn phase
 				// can kill particles.
 				bWriteInstanceCount = false;
-				HlslOutput += TEXT("\t\tconst bool bValid = Context.") + TranslationStages[i].PassNamespace + TEXT(".DataInstance.Alive;\n");
+				HlslOutput += TEXT("\t\tconst bool bValid = bIsValidInstance && Context.") + TranslationStages[i].PassNamespace + TEXT(".DataInstance.Alive;\n");
 				HlslOutput += TEXT("\t\tconst int WriteIndex = OutputIndex(0, true, bValid);\n");
 			}
 			else
 			{
 				// The stage doesn't kill particles, we can take the simpler path which doesn't need to manage the particle count.
-				HlslOutput += TEXT("\t\tconst bool bValid = GCurrentPhase != -1;\n");
+				HlslOutput += TEXT("\t\tconst bool bValid = bIsValidInstance;\n");
 				HlslOutput += TEXT("\t\tconst int WriteIndex = OutputIndex(0, false, bValid);\n");
 			}
 
@@ -3068,15 +3081,14 @@ void FHlslNiagaraTranslator::DefineMainGPUFunctions(
 
 			if (bWriteInstanceCount)
 			{
+				//-TODO: This path should be deprecated if we ever remove the ability to disable partial writes
 				HlslOutput += TEXT(
 					"\t\t// If a stage doesn't kill particles, StoreUpdateVariables() never calls AcquireIndex(), so the\n"
 					"\t\t// count isn't updated. In that case we must manually copy the original count here.\n"
-					"\t\t#if USE_SIMULATION_STAGES\n"
 					"\t\tif (WriteInstanceCountOffset != 0xFFFFFFFF && GLinearThreadId == 0) \n"
 					"\t\t{\n"
 					"\t\t	RWInstanceCounts[WriteInstanceCountOffset] = GSpawnStartInstance + SpawnedInstances; \n"
 					"\t\t}\n"
-					"\t\t#endif\n"
 				);
 			}
 		}
@@ -3149,91 +3161,123 @@ void FHlslNiagaraTranslator::DefineMainGPUFunctions(
 				"		\n");
 	}
 	
+	//////////////////////////////////////////////////////////////////////////
+	// Generate per translation stage logic parts
+	// i.e. ones who vary based upon what we iterate over
+	FString SimRunSpawnUpdateLogicString;
+	FString SimSetupSpawnExecIndexLogicString;
+	for (int32 i=1; i < TranslationStages.Num(); i++)
+	{
+		const FHlslNiagaraTranslationStage& TranslationStage = TranslationStages[i];
+		const TCHAR* IfOrElifText = (i == 1) ? TEXT("#if") : TEXT("#elif");
+
+		SimRunSpawnUpdateLogicString.Appendf(TEXT("%s SimulationStageIndex == %d	// %s\n"), IfOrElifText, TranslationStage.SimulationStageIndex, *TranslationStage.PassNamespace);
+		SimSetupSpawnExecIndexLogicString.Appendf(TEXT("%s SimulationStageIndex == %d	// %s\n"), IfOrElifText, TranslationStage.SimulationStageIndex, *TranslationStage.PassNamespace);
+
+		// Particle iteration stage
+		if ( TranslationStage.IterationSource.IsNone() )
+		{
+			// We combine the update & spawn scripts together on GPU so we only need to check for spawning on the first translation stage
+			// Note: Depending on how spawning inside stages works we may need to enable the spawn logic for those stages *only*			
+			const bool bParticleSpawnStage = i == 1;
+			if (bParticleSpawnStage)
+			{
+				SimRunSpawnUpdateLogicString += TEXT(
+					"	if (ReadInstanceCountOffset == 0xFFFFFFFF)\n"
+					"	{\n"
+					"		GSpawnStartInstance = 0;\n"
+					"	}\n"
+					"	else\n"
+					"	{\n"
+					"		GSpawnStartInstance = RWInstanceCounts[ReadInstanceCountOffset];\n"
+					"	}\n"
+					"	const int MaxInstances = GSpawnStartInstance + SpawnedInstances;\n"
+					"	const bool bRunUpdateLogic = GLinearThreadId < GSpawnStartInstance && GLinearThreadId < MaxInstances;\n"
+					"	const bool bRunSpawnLogic = GLinearThreadId >= GSpawnStartInstance && GLinearThreadId < MaxInstances;\n"
+				);
+			}
+			else
+			{
+				SimRunSpawnUpdateLogicString += TEXT(
+					"	GSpawnStartInstance = RWInstanceCounts[ReadInstanceCountOffset];\n"
+					"	const bool bRunUpdateLogic = GLinearThreadId < GSpawnStartInstance;\n"
+					"	const bool bRunSpawnLogic = false;\n"
+				);
+			}
+
+			SimSetupSpawnExecIndexLogicString += TEXT("		SetupExecIndexAndSpawnInfoForGPU();\n");
+		}
+		// Iteration interface stage
+		else
+		{
+			//	if (const FNiagaraVariable* IterationSourceVar = CompileData->EncounteredVariables.FindByPredicate([&](const FNiagaraVariable& VarInfo) { return VarInfo.GetName() == SimulationStageMetaData.IterationSource; }))
+			//	{
+			//		if (UNiagaraDataInterface* IteratinoSourceCDO = CompileDuplicateData->GetDuplicatedDataInterfaceCDOForClass(IterationSourceVar->GetType().GetClass()))
+			//		{
+			//			SimulationStageMetaData.GpuDispatchType = IteratinoSourceCDO->GetGpuDispatchType();
+			//			SimulationStageMetaData.GpuDispatchNumThreads = IteratinoSourceCDO->GetGpuDispatchNumThreads();
+			//		}
+			//	}
+
+			//-TODO: We can simplify the logic here with SimulationStage_GetInstanceCount() as only things that can provide an instance count offset
+			//		 can really be variable, everything else is driven from CPU code.
+			SimRunSpawnUpdateLogicString += TEXT(
+				"	const uint MaxInstances = SimulationStage_GetInstanceCount();\n"
+				"	GLinearThreadId = all(DispatchThreadId < DispatchThreadIdBounds) ? GLinearThreadId : MaxInstances;\n"
+				"	GSpawnStartInstance = MaxInstances;\n"
+				"	const bool bRunUpdateLogic = (GLinearThreadId < GSpawnStartInstance) && (SimStart != 1);\n"
+				"	const bool bRunSpawnLogic = (GLinearThreadId < GSpawnStartInstance) && (SimStart == 1);\n"
+			);
+
+			SimSetupSpawnExecIndexLogicString += TEXT("		SetupExecIndexForGPU();\n");
+		}
+	}
+	SimRunSpawnUpdateLogicString.Append(TEXT("#endif\n"));
+	SimSetupSpawnExecIndexLogicString.Append(TEXT("#endif\n"));
+
+	//////////////////////////////////////////////////////////////////////////
+	// Generate main body
 	HlslOutput +=
 		TEXT("\n\n/*\n"
 			"*	CS wrapper for our generated code; calls spawn and update functions on the corresponding instances in the buffer\n"
 			" */\n"
 			"\n"
-			"[numthreads(THREADGROUP_SIZE, 1, 1)]\n"
+			"[numthreads(THREADGROUP_SIZE_X, THREADGROUP_SIZE_Y, THREADGROUP_SIZE_Z)]\n"
 			"void SimulateMainComputeCS(\n"
 			"	uint3 DispatchThreadId : SV_DispatchThreadID,\n"
 			"	uint3 GroupThreadId : SV_GroupThreadID)\n"
 			"{\n"
-			"	GLinearThreadId = DispatchThreadId.x + (DispatchThreadId.y * DispatchThreadIdToLinear);\n"
+			"	GLinearThreadId = DispatchThreadId.x + (DispatchThreadId.y * DispatchThreadIdToLinear.y);\n"
+			"	#if NIAGARA_DISPATCH_TYPE >= NIAGARA_DISPATCH_TYPE_THREE_D\n"
+			"		GLinearThreadId += DispatchThreadId.z * DispatchThreadIdToLinear.z;\n"
+			"	#endif\n"
 			"	GDispatchThreadId = DispatchThreadId;\n"
 			"	GGroupThreadId = GroupThreadId;\n"
-			"	GCurrentPhase = -1;\n"
 			"	GEmitterTickCounter = EmitterTickCounter;\n"
-			"	GSimStart = SimStart;\n"
 			"	GRandomSeedOffset = 0;\n"
+			"\n"
+			) + SimRunSpawnUpdateLogicString + TEXT(
 			"	\n"
-			"	/*\n"
-			"	if(CopyInstancesBeforeStart == 1)\n"
-			"	{\n"
-			"		UpdateStartInstance = 0;\n"
-			"	}\n"
-			"	*/\n"
-			"	\n"
-			"   // The CPU code will set UpdateStartInstance to 0 and ReadInstanceCountOffset to -1 for stages.\n"
-			"	const uint InstanceID = UpdateStartInstance + GLinearThreadId;\n"
-			"	if (ReadInstanceCountOffset == 0xFFFFFFFF)\n"
-			"	{\n"
-			"		GSpawnStartInstance = 0;\n"
-			"	}\n"
-			"	else\n"
-			"	{\n"
-			"		GSpawnStartInstance = RWInstanceCounts[ReadInstanceCountOffset];				// needed by ExecIndex()\n"
-			"	}\n"
-			"	bool bRunUpdateLogic, bRunSpawnLogic;\n"
-			"	#if USE_SIMULATION_STAGES\n"
-			"	int IterationInterfaceInstanceCount = SimulationStage_GetInstanceCount();\n"
-			"	if (IterationInterfaceInstanceCount > 0)\n"
-			"	{\n"
-			"		bRunUpdateLogic = InstanceID < IterationInterfaceInstanceCount && GSimStart != 1;\n"
-			"		bRunSpawnLogic = InstanceID < IterationInterfaceInstanceCount && GSimStart == 1;\n"
-			"		GSpawnStartInstance = IterationInterfaceInstanceCount;\n"
-			"	}\n"
-			"	else\n"
-			"	#endif // USE_SIMULATION_STAGES\n"
-			"	{\n"
-			"	    const int MaxInstances = GSpawnStartInstance + SpawnedInstances;\n"
-			"		bRunUpdateLogic = InstanceID < GSpawnStartInstance && InstanceID < UpdateStartInstance + MaxInstances;\n"
-			"		bRunSpawnLogic = InstanceID >= GSpawnStartInstance && InstanceID < UpdateStartInstance + MaxInstances;\n"
-			"	}\n"
-			"	\n"
-			"	const float RandomSeedInitialisation = NiagaraInternalNoise(InstanceID * 16384, 0 * 8196, (bRunUpdateLogic ? 4096 : 0) + EmitterTickCounter);	// initialise the random state seed\n"
+			"	const float RandomSeedInitialisation = NiagaraInternalNoise(GLinearThreadId * 16384, 0 * 8196, (bRunUpdateLogic ? 4096 : 0) + EmitterTickCounter);	// initialise the random state seed\n"
 			"	\n"
 			"	FSimulationContext Context = (FSimulationContext)0;\n"
 			"	\n"
 			"	BRANCH\n"
 			"	if (bRunUpdateLogic)\n"
 			"	{\n"
-			"		GCurrentPhase = GUpdatePhase;\n"
 			"		SetupExecIndexForGPU();\n"
 			"		InitConstants(Context);\n"
-			"		LoadUpdateVariables(Context, InstanceID);\n"
+			"		LoadUpdateVariables(Context, GLinearThreadId);\n"
 			"		ReadDataSets(Context);\n"
 			"	}\n"
 			"	else if (bRunSpawnLogic)\n"
 			"	{\n"
-			"		GCurrentPhase = GSpawnPhase;\n"
-			"	#if USE_SIMULATION_STAGES\n"
-			"		// Only process the spawn info for particle-based stages. Stages with an iteration interface expect the exec index to simply be the thread index.\n"
-			"		if (IterationInterfaceInstanceCount > 0)\n"
-			"		{\n"
-			"			SetupExecIndexForGPU();\n"
-			"		}\n"
-			"		else\n"
-			"	#endif\n"
-			"		{\n"
-			"			SetupExecIndexAndSpawnInfoForGPU();\n"
-			"		}\n"
+		) + SimSetupSpawnExecIndexLogicString + TEXT(
 			"		InitConstants(Context);\n"
 			"		InitSpawnVariables(Context);\n"
 			"		ReadDataSets(Context);\n"
 			"		\n") + SpawnLogicString
 			+ TEXT(
-			"		GCurrentPhase = GUpdatePhase;\n"
 			"		\n"
 			"		TransferAttributes(Context);\n"
 			"		\n"
@@ -3241,12 +3285,11 @@ void FHlslNiagaraTranslator::DefineMainGPUFunctions(
 			"\n"
 			"	if (bRunUpdateLogic || bRunSpawnLogic)\n"
 			"	{\n"
-		) + SimDoWorkString
-		+ TEXT(
+		) + SimDoWorkString + TEXT(
 			"		WriteDataSets(Context);\n"
 			"	}\n"
 			"	\n"
-			"	StoreUpdateVariables(Context);\n\n"
+			"	StoreUpdateVariables(Context, bRunUpdateLogic || bRunSpawnLogic);\n\n"
 		);
 
 	HlslOutput += TEXT("}\n");
diff --git a/Engine/Plugins/FX/Niagara/Source/NiagaraShader/Private/NiagaraScriptBase.cpp b/Engine/Plugins/FX/Niagara/Source/NiagaraShader/Private/NiagaraScriptBase.cpp
--- a/Engine/Plugins/FX/Niagara/Source/NiagaraShader/Private/NiagaraScriptBase.cpp
+++ b/Engine/Plugins/FX/Niagara/Source/NiagaraShader/Private/NiagaraScriptBase.cpp
@@ -10,5 +10,7 @@ UNiagaraScriptBase::UNiagaraScriptBase(const FObjectInitializer& ObjectInitializ
 FSimulationStageMetaData::FSimulationStageMetaData()
 	: bWritesParticles(0)
 	, bPartialParticleUpdate(0)
+	, GpuDispatchType(ENiagaraGpuDispatchType::OneD)
+	, GpuDispatchNumThreads(0, 0, 0)
 {
 }
diff --git a/Engine/Plugins/FX/Niagara/Source/NiagaraShader/Private/NiagaraShader.cpp b/Engine/Plugins/FX/Niagara/Source/NiagaraShader/Private/NiagaraShader.cpp
--- a/Engine/Plugins/FX/Niagara/Source/NiagaraShader/Private/NiagaraShader.cpp
+++ b/Engine/Plugins/FX/Niagara/Source/NiagaraShader/Private/NiagaraShader.cpp
@@ -432,6 +432,22 @@ void FNiagaraShaderType::BeginCompileShader(
 
 	NewJob->Input.Environment.SetDefine(TEXT("NIAGARA_COMPRESSED_ATTRIBUTES_ENABLED"), Script->GetUsesCompressedAttributes() ? 1 : 0);
 
+	FIntVector ThreadGroupSize = FNiagaraShader::GetDefaultThreadGroupSize(ENiagaraGpuDispatchType::OneD);
+	if (SimStageMetaData.GpuDispatchType != ENiagaraGpuDispatchType::Custom)
+	{
+		ThreadGroupSize = FNiagaraShader::GetDefaultThreadGroupSize(SimStageMetaData.GpuDispatchType);
+	}
+	else
+	{
+		ThreadGroupSize = SimStageMetaData.GpuDispatchNumThreads;
+	}
+
+	NewJob->Input.Environment.SetDefine(TEXT("NIAGARA_DISPATCH_TYPE"), int(SimStageMetaData.GpuDispatchType));
+	NewJob->Input.Environment.SetDefine(TEXT("THREADGROUP_SIZE"), ThreadGroupSize.X * ThreadGroupSize.Y * ThreadGroupSize.Z);
+	NewJob->Input.Environment.SetDefine(TEXT("THREADGROUP_SIZE_X"), ThreadGroupSize.X);
+	NewJob->Input.Environment.SetDefine(TEXT("THREADGROUP_SIZE_Y"), ThreadGroupSize.Y);
+	NewJob->Input.Environment.SetDefine(TEXT("THREADGROUP_SIZE_Z"), ThreadGroupSize.Z);
+
 	Script->ModifyCompilationEnvironment(Platform, NewJob->Input.Environment);
 
 	AddReferencedUniformBufferIncludes(NewJob->Input.Environment, NewJob->Input.SourceFilePrefix, (EShaderPlatform)Target.Platform);
@@ -1219,10 +1235,10 @@ void FNiagaraShader::BindParams(const TArray<FNiagaraDataInterfaceGPUParamInfo>&
 	CopyInstancesBeforeStartParam.Bind(ParameterMap, TEXT("CopyInstancesBeforeStart"));
 
 	NumSpawnedInstancesParam.Bind(ParameterMap, TEXT("SpawnedInstances"));
-	UpdateStartInstanceParam.Bind(ParameterMap, TEXT("UpdateStartInstance"));
 
 	SimulationStageIterationInfoParam.Bind(ParameterMap, TEXT("SimulationStageIterationInfo"));
 	SimulationStageNormalizedIterationIndexParam.Bind(ParameterMap, TEXT("SimulationStageNormalizedIterationIndex"));
+	DispatchThreadIdBoundsParam.Bind(ParameterMap, TEXT("DispatchThreadIdBounds"));
 	DispatchThreadIdToLinearParam.Bind(ParameterMap, TEXT("DispatchThreadIdToLinear"));
 
 	ComponentBufferSizeReadParam.Bind(ParameterMap, TEXT("ComponentBufferSizeRead"));
diff --git a/Engine/Plugins/FX/Niagara/Source/NiagaraShader/Public/NiagaraScriptBase.h b/Engine/Plugins/FX/Niagara/Source/NiagaraShader/Public/NiagaraScriptBase.h
--- a/Engine/Plugins/FX/Niagara/Source/NiagaraShader/Public/NiagaraScriptBase.h
+++ b/Engine/Plugins/FX/Niagara/Source/NiagaraShader/Public/NiagaraScriptBase.h
@@ -7,6 +7,19 @@
 
 #include "NiagaraScriptBase.generated.h"
 
+UENUM()
+enum class ENiagaraGpuDispatchType : uint8
+{
+	/* Instances will distribute along X using platforms specific thread counts. */
+	OneD,
+	/* Instances will distribute along X & Y using platforms specific thread counts. */
+	TwoD,
+	/* Instances will distribute along X, Y & Z using platforms specific thread counts. */
+	ThreeD,
+	/* NumThreads will be determined manually. */
+	Custom,
+};
+
 UENUM()
 enum class ENiagaraSimStageExecuteBehavior : uint8
 {
@@ -66,6 +79,14 @@ public:
 		const bool bResetOnly = ExecuteBehavior == ENiagaraSimStageExecuteBehavior::OnSimulationReset;
 		return bAlways || (bResetOnly == bResetData);
 	}
+
+	/** Dispatch type set for this stage. */
+	UPROPERTY()
+	ENiagaraGpuDispatchType GpuDispatchType;
+
+	/** When in custom mode this is the num threads. */
+	UPROPERTY()
+	FIntVector GpuDispatchNumThreads;
 };
 
 UCLASS(MinimalAPI, abstract)
diff --git a/Engine/Plugins/FX/Niagara/Source/NiagaraShader/Public/NiagaraShader.h b/Engine/Plugins/FX/Niagara/Source/NiagaraShader/Public/NiagaraShader.h
--- a/Engine/Plugins/FX/Niagara/Source/NiagaraShader/Public/NiagaraShader.h
+++ b/Engine/Plugins/FX/Niagara/Source/NiagaraShader/Public/NiagaraShader.h
@@ -13,6 +13,7 @@
 #include "Shader.h"
 #include "GlobalShader.h"
 #include "NiagaraCommon.h"
+#include "NiagaraScriptBase.h"
 #include "NiagaraShared.h"
 #include "NiagaraShaderType.h"
 #include "SceneRenderTargetParameters.h"
@@ -35,17 +36,22 @@ public:
 	{
 	}
 
-	static uint32 GetGroupSize(EShaderPlatform Platform)
+	static FIntVector GetDefaultThreadGroupSize(ENiagaraGpuDispatchType DispatchType)
 	{
-		const bool bUseWaveIntrinsics = false; // TODO: Some content breaks with this - FDataDrivenShaderPlatformInfo::GetInfo(Platform).bSupportsIntrinsicWaveOnce;
-		return 64; // TODO: Force 64 until we fix the intrinsics // bUseWaveIntrinsics ? 64 : 32;
+		//-TODO: Grab this from FDataDrivenShaderPlatformInfo
+		switch (DispatchType)
+		{
+			case ENiagaraGpuDispatchType::OneD:		return FIntVector(64, 1, 1);
+			case ENiagaraGpuDispatchType::TwoD:		return FIntVector(8, 8, 1);
+			case ENiagaraGpuDispatchType::ThreeD:	return FIntVector(4, 4, 2);
+			default:								return FIntVector(64, 1, 1);
+		}
 	}
 
 	static void ModifyCompilationEnvironment(const FNiagaraShaderPermutationParameters& Parameters, FShaderCompilerEnvironment& OutEnvironment)
 	{
 		FShader::ModifyCompilationEnvironment(Parameters, OutEnvironment);
 		const bool bUseWaveIntrinsics = false; // TODO: Some content breaks with this - FDataDrivenShaderPlatformInfo::GetInfo(Platform).bSupportsIntrinsicWaveOnce;
-		OutEnvironment.SetDefine(TEXT("THREADGROUP_SIZE"), GetGroupSize(Parameters.Platform));
 		OutEnvironment.SetDefine(TEXT("USE_WAVE_INTRINSICS"), bUseWaveIntrinsics ? 1 : 0);
 	}
 
@@ -57,28 +63,6 @@ public:
 
 	FNiagaraShader(const FNiagaraShaderType::CompiledShaderInitializerType& Initializer);
 
-//	FRHIUniformBuffer* GetParameterCollectionBuffer(const FGuid& Id, const FSceneInterface* SceneInterface) const;
-	/*
-	template<typename ShaderRHIParamRef>
-	FORCEINLINE_DEBUGGABLE void SetViewParameters(FRHICommandList& RHICmdList, const ShaderRHIParamRef ShaderRHI, const FSceneView& View, const TUniformBufferRef<FViewUniformShaderParameters>& ViewUniformBuffer)
-	{
-		const auto& ViewUniformBufferParameter = GetUniformBufferParameter<FViewUniformShaderParameters>();
-		const auto& BuiltinSamplersUBParameter = GetUniformBufferParameter<FBuiltinSamplersParameters>();
-		CheckShaderIsValid();
-		SetUniformBufferParameter(RHICmdList, ShaderRHI, ViewUniformBufferParameter, ViewUniformBuffer);
-
-		if (View.bShouldBindInstancedViewUB && View.Family->Views.Num() > 0)
-		{
-			// When drawing the left eye in a stereo scene, copy the right eye view values into the instanced view uniform buffer.
-			const EStereoscopicPass StereoPassIndex = (View.StereoPass != EStereoscopicPass::eSSP_FULL) ? EStereoscopicPass::eSSP_RIGHT_EYE : EStereoscopicPass::eSSP_FULL;
-
-			const FSceneView& InstancedView = View.Family->GetStereoEyeView(StereoPassIndex);
-			const auto& InstancedViewUniformBufferParameter = GetUniformBufferParameter<FInstancedViewUniformShaderParameters>();
-			SetUniformBufferParameter(RHICmdList, ShaderRHI, InstancedViewUniformBufferParameter, InstancedView.ViewUniformBuffer);
-		}
-	}
-	*/
-
 	// Bind parameters
 	void BindParams(const TArray<FNiagaraDataInterfaceGPUParamInfo>& InDIParamInfo, const FShaderParameterMap &ParameterMap);
 
@@ -113,10 +97,10 @@ public:
 	LAYOUT_FIELD(FShaderParameter, NumParticlesPerEventParam);
 	LAYOUT_FIELD(FShaderParameter, CopyInstancesBeforeStartParam);
 	LAYOUT_FIELD(FShaderParameter, NumSpawnedInstancesParam);
-	LAYOUT_FIELD(FShaderParameter, UpdateStartInstanceParam);
 
 	LAYOUT_FIELD(FShaderParameter, SimulationStageIterationInfoParam);
 	LAYOUT_FIELD(FShaderParameter, SimulationStageNormalizedIterationIndexParam);
+	LAYOUT_FIELD(FShaderParameter, DispatchThreadIdBoundsParam);
 	LAYOUT_FIELD(FShaderParameter, DispatchThreadIdToLinearParam);
 
 	LAYOUT_FIELD(FShaderParameter, ComponentBufferSizeReadParam);
