- Handle case where we get zero property data from the GT

#rnx

diff --git a/Engine/Plugins/FX/Niagara/Source/Niagara/Private/DataInterface/NiagaraDataInterfaceUObjectPropertyReader.cpp b/Engine/Plugins/FX/Niagara/Source/Niagara/Private/DataInterface/NiagaraDataInterfaceUObjectPropertyReader.cpp
--- a/Engine/Plugins/FX/Niagara/Source/Niagara/Private/DataInterface/NiagaraDataInterfaceUObjectPropertyReader.cpp
+++ b/Engine/Plugins/FX/Niagara/Source/Niagara/Private/DataInterface/NiagaraDataInterfaceUObjectPropertyReader.cpp
@@ -103,11 +103,17 @@ namespace NDIUObjectPropertyReaderLocal
 				{
 					InstanceData_RT->PropertyData.Release();
 					const uint32 NumElements = InstanceData_FromGT->PropertyData.Num() / 4;
-					InstanceData_RT->PropertyData.Initialize(TEXT("NiagaraUObjectPropertyReader"), sizeof(float), NumElements, EPixelFormat::PF_R32_FLOAT);
+					if (NumElements > 0)
+					{
+						InstanceData_RT->PropertyData.Initialize(TEXT("NiagaraUObjectPropertyReader"), sizeof(float), NumElements, EPixelFormat::PF_R32_FLOAT);
+					}
+				}
+				if (InstanceData_RT->PropertyData.NumBytes > 0)
+				{
+					void* GpuMemory = RHILockBuffer(InstanceData_RT->PropertyData.Buffer, 0, InstanceData_RT->PropertyData.NumBytes, RLM_WriteOnly);
+					FMemory::Memcpy(GpuMemory, InstanceData_FromGT->PropertyData.GetData(), InstanceData_FromGT->PropertyData.Num());
+					RHIUnlockBuffer(InstanceData_RT->PropertyData.Buffer);
 				}
-				void* GpuMemory = RHILockBuffer(InstanceData_RT->PropertyData.Buffer, 0, InstanceData_RT->PropertyData.NumBytes, RLM_WriteOnly);
-				FMemory::Memcpy(GpuMemory, InstanceData_FromGT->PropertyData.GetData(), InstanceData_FromGT->PropertyData.Num());
-				RHIUnlockBuffer(InstanceData_RT->PropertyData.Buffer);
 			}
 			InstanceData_FromGT->~FInstanceData_GameToRender();
 		}
