- Remove output name from baker, use generated texture and if it isn't set pop a dialog up

diff --git a/Engine/Plugins/FX/Niagara/Source/Niagara/Classes/NiagaraBakerSettings.h b/Engine/Plugins/FX/Niagara/Source/Niagara/Classes/NiagaraBakerSettings.h
--- a/Engine/Plugins/FX/Niagara/Source/Niagara/Classes/NiagaraBakerSettings.h
+++ b/Engine/Plugins/FX/Niagara/Source/Niagara/Classes/NiagaraBakerSettings.h
@@ -36,10 +36,6 @@ USTRUCT()
 struct FNiagaraBakerTextureSettings
 {
 	GENERATED_BODY()
-
-	/** Optional output name, if left empty a name will be auto generated using the index of the texture/ */
-	UPROPERTY(EditAnywhere, Category = "Texture")
-	FName OutputName;
 	
 	/** Source visualization we should capture, i.e. Scene Color, World Normal, etc */
 	UPROPERTY(EditAnywhere, Category = "Texture")
diff --git a/Engine/Plugins/FX/Niagara/Source/Niagara/Private/NiagaraBakerSettings.cpp b/Engine/Plugins/FX/Niagara/Source/Niagara/Private/NiagaraBakerSettings.cpp
--- a/Engine/Plugins/FX/Niagara/Source/Niagara/Private/NiagaraBakerSettings.cpp
+++ b/Engine/Plugins/FX/Niagara/Source/Niagara/Private/NiagaraBakerSettings.cpp
@@ -8,7 +8,6 @@
 bool FNiagaraBakerTextureSettings::Equals(const FNiagaraBakerTextureSettings& Other) const
 {
 	return
-		OutputName == Other.OutputName &&
 		SourceBinding.SourceName == Other.SourceBinding.SourceName &&
 		bUseFrameSize == Other.bUseFrameSize &&
 		FrameSize == Other.FrameSize &&
diff --git a/Engine/Plugins/FX/Niagara/Source/NiagaraEditor/Private/ViewModels/NiagaraBakerViewModel.cpp b/Engine/Plugins/FX/Niagara/Source/NiagaraEditor/Private/ViewModels/NiagaraBakerViewModel.cpp
--- a/Engine/Plugins/FX/Niagara/Source/NiagaraEditor/Private/ViewModels/NiagaraBakerViewModel.cpp
+++ b/Engine/Plugins/FX/Niagara/Source/NiagaraEditor/Private/ViewModels/NiagaraBakerViewModel.cpp
@@ -22,6 +22,8 @@
 #include "PackageTools.h"
 #include "AdvancedPreviewScene.h"
 
+#include "Factories/Texture2dFactoryNew.h"
+
 #define LOCTEXT_NAMESPACE "NiagaraBakerViewModel"
 
 FNiagaraBakerViewModel::FNiagaraBakerViewModel()
@@ -229,30 +231,31 @@ void FNiagaraBakerViewModel::RenderBaker()
 		}
 	}
 
+	// Remove existing generated settings to avoid holding references to textures
+	if (UNiagaraSystem* Asset = PreviewComponent->GetAsset())
+	{
+		Asset->SetBakerGeneratedSettings(nullptr);
+	}
+
 	// Send data to generated textures
 	for (int32 iOutputTexture = 0; iOutputTexture < BakerSettings->OutputTextures.Num(); ++iOutputTexture)
 	{
 		FNiagaraBakerTextureSettings& OutputTexture = BakerSettings->OutputTextures[iOutputTexture];
 
-		// If we don't have a texture create one
+		// If we don't already have a generated texture ask the user to create one
 		if (OutputTexture.GeneratedTexture == nullptr)
 		{
+			IAssetTools& AssetTools = FModuleManager::Get().LoadModuleChecked<FAssetToolsModule>("AssetTools").Get();
+			UTexture2DFactoryNew* Texture2DFactory = NewObject<UTexture2DFactoryNew>();
+
 			FString PackagePath = FPaths::GetPath(PreviewComponent->GetAsset()->GetPackage()->GetPathName());
-			if (OutputTexture.OutputName.IsNone())
+			FString AssetName = PreviewComponent->GetAsset()->GetName() + TEXT("_Baker") + FString::FromInt(iOutputTexture) + TEXT("Texture");
+			OutputTexture.GeneratedTexture = Cast<UTexture2D>(AssetTools.CreateAssetWithDialog(AssetName, PackagePath, UTexture2D::StaticClass(), Texture2DFactory));
+			if (OutputTexture.GeneratedTexture == nullptr)
 			{
-				PackagePath = PackagePath / PreviewComponent->GetAsset()->GetName() + TEXT("_Baker") + FString::FromInt(iOutputTexture);
-			}
-			else
-			{
-				PackagePath = PackagePath / PreviewComponent->GetAsset()->GetName() + OutputTexture.OutputName.ToString();
+				// User skipped for some reason
+				continue;
 			}
-
-			IAssetTools& AssetTools = FModuleManager::Get().LoadModuleChecked<FAssetToolsModule>("AssetTools").Get();
-			FString PackageName;
-			FString AssetName;
-			AssetTools.CreateUniqueAssetName(UPackageTools::SanitizePackageName(PackagePath), TEXT(""), PackageName, AssetName);
-
-			OutputTexture.GeneratedTexture = NewObject<UTexture2D>(CreatePackage(*PackagePath), FName(*AssetName), RF_Standalone | RF_Public);
 		}
 
 		const bool bIsPoT = FMath::IsPowerOfTwo(OutputTexture.TextureSize.X) && FMath::IsPowerOfTwo(OutputTexture.TextureSize.Y);
diff --git a/Engine/Plugins/FX/Niagara/Source/NiagaraEditor/Private/Widgets/SNiagaraBakerViewportToolbar.cpp b/Engine/Plugins/FX/Niagara/Source/NiagaraEditor/Private/Widgets/SNiagaraBakerViewportToolbar.cpp
--- a/Engine/Plugins/FX/Niagara/Source/NiagaraEditor/Private/Widgets/SNiagaraBakerViewportToolbar.cpp
+++ b/Engine/Plugins/FX/Niagara/Source/NiagaraEditor/Private/Widgets/SNiagaraBakerViewportToolbar.cpp
@@ -256,18 +256,15 @@ TSharedRef<SWidget> SNiagaraBakerViewportToolbar::GenerateOptionsMenu() const
 		const UNiagaraBakerSettings* BakerSettings = ViewModel->GetBakerSettings();
 		for ( int32 i=0; i < BakerSettings->OutputTextures.Num(); ++i )
 		{
-			FString TextureName;
-			if ( BakerSettings->OutputTextures[i].OutputName.IsNone() )
+			TStringBuilder<128> TextureName;
+			TextureName.Appendf(TEXT("Texture(%d)"), i);
+			if (UTexture2D* GeneratedTexture = BakerSettings->OutputTextures[i].GeneratedTexture)
 			{
-				TextureName = FString::Printf(TEXT("Output Texture %d"), i);
+				TextureName.Appendf(TEXT(" - %s"), *BakerSettings->OutputTextures[i].GeneratedTexture->GetName());
 			}
-			else
-			{
-				TextureName = BakerSettings->OutputTextures[i].OutputName.ToString();
-			}
-
+			
 			MenuBuilder.AddMenuEntry(
-				FText::FromString(TextureName),
+				FText::FromStringView(TextureName.ToView()),
 				FText(),
 				FSlateIcon(),
 				FUIAction(
diff --git a/Engine/Plugins/FX/Niagara/Source/NiagaraEditor/Private/Widgets/SNiagaraBakerWidget.cpp b/Engine/Plugins/FX/Niagara/Source/NiagaraEditor/Private/Widgets/SNiagaraBakerWidget.cpp
--- a/Engine/Plugins/FX/Niagara/Source/NiagaraEditor/Private/Widgets/SNiagaraBakerWidget.cpp
+++ b/Engine/Plugins/FX/Niagara/Source/NiagaraEditor/Private/Widgets/SNiagaraBakerWidget.cpp
@@ -121,16 +121,13 @@ FText SNiagaraBakerWidget::GetSelectedTextureAsText() const
 			const int32 PrevewTextureIndex = ViewModel->GetPreviewTextureIndex();
 			if ( BakerSettings->OutputTextures.IsValidIndex(PrevewTextureIndex) )
 			{
-				FString TextureName;
-				if (BakerSettings->OutputTextures[PrevewTextureIndex].OutputName.IsNone())
+				TStringBuilder<128> TextureName;
+				TextureName.Appendf(TEXT("Texture(%d)"), PrevewTextureIndex);
+				if (UTexture2D* GeneratedTexture = BakerSettings->OutputTextures[PrevewTextureIndex].GeneratedTexture)
 				{
-					TextureName = FString::Printf(TEXT("Output Texture %d"), PrevewTextureIndex);
+					TextureName.Appendf(TEXT(" - %s"), *BakerSettings->OutputTextures[PrevewTextureIndex].GeneratedTexture->GetName());
 				}
-				else
-				{
-					TextureName = BakerSettings->OutputTextures[PrevewTextureIndex].OutputName.ToString();
-				}
-				return FText::FromString(TextureName);
+				return FText::FromStringView(TextureName.ToView());
 			}
 		}
 	}
