- Fix bool component reading
- Add a clear message if we attempt to bind a property and the type is not what we expected from the graph

#rnx

diff --git a/Engine/Plugins/FX/Niagara/Source/Niagara/Private/DataInterface/NiagaraDataInterfaceUObjectPropertyReader.cpp b/Engine/Plugins/FX/Niagara/Source/Niagara/Private/DataInterface/NiagaraDataInterfaceUObjectPropertyReader.cpp
--- a/Engine/Plugins/FX/Niagara/Source/Niagara/Private/DataInterface/NiagaraDataInterfaceUObjectPropertyReader.cpp
+++ b/Engine/Plugins/FX/Niagara/Source/Niagara/Private/DataInterface/NiagaraDataInterfaceUObjectPropertyReader.cpp
@@ -341,9 +341,9 @@ namespace NDIUObjectPropertyReaderLocal
 		static void VMFunction(FVectorVMExternalFunctionContext& Context, const uint32 GetterIndex) { VMReadData<FNiagaraBool, FNiagaraBool>(Context, GetterIndex, 0.0f); }
 		static TFunction<void(const FNiagaraLWCConverter&, void*)> GetCopyFunction(const FProperty* InProperty, const void* PropertyAddress)
 		{
-			if (InProperty->IsA<const FBoolProperty>())
+			if (const FBoolProperty* BoolProperty = CastField<const FBoolProperty>(InProperty))
 			{
-				return [PropertyAddress](const FNiagaraLWCConverter& LwcConverter, void* DestAddress) { *reinterpret_cast<FNiagaraBool*>(DestAddress) = *reinterpret_cast<const bool*>(PropertyAddress); };
+				return [PropertyAddress, BoolProperty](const FNiagaraLWCConverter& LwcConverter, void* DestAddress) { *reinterpret_cast<FNiagaraBool*>(DestAddress) = BoolProperty->GetPropertyValue(PropertyAddress); };
 			}
 			return nullptr;
 		}
@@ -391,10 +391,6 @@ namespace NDIUObjectPropertyReaderLocal
 			FProperty* Property = FindFProperty<FProperty>(ObjectClass, PropertyFName);
 			if ( Property == nullptr )
 			{
-				if ( FNiagaraUtilities::LogVerboseWarnings() )
-				{
-					UE_LOG(LogNiagara, Log, TEXT("Failed attempting to bind property '%s' for object '%s' fullpath '%s'"), *PropertyFName.ToString(), *GetNameSafe(ObjectClass), *PropertyGetter.Variable.GetName().ToString());
-				}
 				return;
 			}
 
@@ -416,7 +412,16 @@ namespace NDIUObjectPropertyReaderLocal
 			void* PropertyAddress = Property->ContainerPtrToValuePtr<void>(ObjectAddress);
 
 			#define NDI_PROPERTY_TYPE(TYPE) \
-				else if (PropertyGetter.Variable.GetType() == FTypeHelper<TYPE>::GetTypeDef()) { PropertyGetter.PropertyCopyFunction = FTypeHelper<TYPE>::GetCopyFunction(Property, PropertyAddress); }
+				else if (PropertyGetter.Variable.GetType() == FTypeHelper<TYPE>::GetTypeDef()) \
+				{ \
+					PropertyGetter.PropertyCopyFunction = FTypeHelper<TYPE>::GetCopyFunction(Property, PropertyAddress); \
+					if (PropertyGetter.PropertyCopyFunction == nullptr && FNiagaraUtilities::LogVerboseWarnings()) \
+					{ \
+						const FStructProperty* StructProperty = CastField<FStructProperty>(Property); \
+						const FString PropertyType = StructProperty && StructProperty->Struct ? StructProperty->Struct->GetName() : Property->GetClass()->GetName(); \
+						UE_LOG(LogNiagara, Warning, TEXT("Could not copy property '%s' type '%s' into expected Niagara type '%s'"), *Property->GetName(), *PropertyType, *PropertyGetter.Variable.GetType().GetName()); \
+					} \
+				} \
 
 				if (false) { } NDI_PROPERTY_TYPES
 			#undef NDI_PROPERTY_TYPE
@@ -683,10 +688,15 @@ bool UNiagaraDataInterfaceUObjectPropertyReader::PerInstanceTick(void* PerInstan
 			{
 				BindPropertyGetter(PropertyGetter, ObjectBinding);
 
-				if ( PropertyGetter.PropertyCopyFunction == nullptr )
+				if ( PropertyGetter.PropertyCopyFunction == nullptr && ActorComponent != nullptr )
 				{
 					BindPropertyGetter(PropertyGetter, ActorComponent);
 				}
+
+				if (PropertyGetter.PropertyCopyFunction == nullptr && FNiagaraUtilities::LogVerboseWarnings())
+				{
+					UE_LOG(LogNiagara, Warning, TEXT("Could not find property '%s' inside object '%s' or component '%s'"), *PropertyGetter.Variable.GetName().ToString(), *GetNameSafe(ObjectBinding), *GetNameSafe(ActorComponent));
+				}
 			}
 		}
 		else
