Change include path for NiagaraStats.h to avoid issues when including from non Niagara files.

#rnx

diff --git a/Engine/Plugins/FX/Niagara/Source/Niagara/Classes/NiagaraDataInterfaceNeighborGrid3D.h b/Engine/Plugins/FX/Niagara/Source/Niagara/Classes/NiagaraDataInterfaceNeighborGrid3D.h
--- a/Engine/Plugins/FX/Niagara/Source/Niagara/Classes/NiagaraDataInterfaceNeighborGrid3D.h
+++ b/Engine/Plugins/FX/Niagara/Source/Niagara/Classes/NiagaraDataInterfaceNeighborGrid3D.h
@@ -5,7 +5,7 @@
 #include "NiagaraCommon.h"
 #include "NiagaraDataInterfaceRW.h"
 #include "ClearQuad.h"
-#include "NiagaraStats.h"
+#include "Niagara/Private/NiagaraStats.h"
 
 #include "NiagaraDataInterfaceNeighborGrid3D.generated.h"
 
