Renamed select node in the search menu to "select / if"

diff --git a/Engine/Plugins/FX/Niagara/Source/NiagaraEditor/Private/EdGraphSchema_Niagara.cpp b/Engine/Plugins/FX/Niagara/Source/NiagaraEditor/Private/EdGraphSchema_Niagara.cpp
--- a/Engine/Plugins/FX/Niagara/Source/NiagaraEditor/Private/EdGraphSchema_Niagara.cpp
+++ b/Engine/Plugins/FX/Niagara/Source/NiagaraEditor/Private/EdGraphSchema_Niagara.cpp
@@ -917,10 +917,10 @@ TArray<TSharedPtr<FNiagaraAction_NewNode>> UEdGraphSchema_Niagara::GetGraphActio
 	// Note: Data Interfaces are not supported for this type
 	if ( bAllowSelectNodes )
 	{
-		const FText SelectMenuDesc = LOCTEXT("NiagaraSelectMenuDesc", "Select");
+		const FText SelectMenuDesc = LOCTEXT("NiagaraSelectMenuDesc", "Select / If");
 		
 		UNiagaraNodeSelect* Node = NewObject<UNiagaraNodeSelect>(OwnerOfTemporaries);
-		AddNewNodeMenuAction(NewActions, Node, SelectMenuDesc, ENiagaraMenuSections::Suggested, {UtilMenuCat.ToString()}, FText::GetEmpty(), FText::FromString(TEXT("If Branch Bool")));
+		AddNewNodeMenuAction(NewActions, Node, SelectMenuDesc, ENiagaraMenuSections::Suggested, {UtilMenuCat.ToString()}, FText::GetEmpty(), FText::FromString(TEXT("If Branch Bool Select")));
 	}
 
 	// Add static switch node
