Fix for changing comment node titles not dirtying scratch pad

diff --git a/Engine/Plugins/FX/Niagara/Source/NiagaraEditor/Private/Widgets/SNiagaraScriptGraph.cpp b/Engine/Plugins/FX/Niagara/Source/NiagaraEditor/Private/Widgets/SNiagaraScriptGraph.cpp
--- a/Engine/Plugins/FX/Niagara/Source/NiagaraEditor/Private/Widgets/SNiagaraScriptGraph.cpp
+++ b/Engine/Plugins/FX/Niagara/Source/NiagaraEditor/Private/Widgets/SNiagaraScriptGraph.cpp
@@ -307,6 +307,12 @@ void SNiagaraScriptGraph::OnNodeTitleCommitted(const FText& NewText, ETextCommit
 			const FScopedTransaction Transaction(LOCTEXT("RenameNode", "Rename Node"));
 			NodeBeingChanged->Modify();
 			NodeBeingChanged->OnRenameNode(NewText.ToString());
+
+			// If the node isn't a NiagarNode, like say a comment, go ahead and mark it dirty
+			if (!NodeBeingChanged->IsA(UNiagaraNode::StaticClass()))
+			{
+				ViewModel->GetGraph()->NotifyGraphNeedsRecompile();
+			}
 		}
 	}
 }
