Fixed includes

diff --git a/Engine/Plugins/FX/Niagara/Source/NiagaraEditorWidgets/Private/Stack/SNiagaraStackFunctionInputValue.cpp b/Engine/Plugins/FX/Niagara/Source/NiagaraEditorWidgets/Private/Stack/SNiagaraStackFunctionInputValue.cpp
--- a/Engine/Plugins/FX/Niagara/Source/NiagaraEditorWidgets/Private/Stack/SNiagaraStackFunctionInputValue.cpp
+++ b/Engine/Plugins/FX/Niagara/Source/NiagaraEditorWidgets/Private/Stack/SNiagaraStackFunctionInputValue.cpp
@@ -17,6 +17,7 @@
 #include "NiagaraNodeFunctionCall.h"
 #include "NiagaraParameterCollection.h"
 #include "NiagaraSystem.h"
+#include "ScopedTransaction.h"
 #include "SDropTarget.h"
 #include "SNiagaraGraphActionWidget.h"
 #include "SNiagaraParameterEditor.h"
@@ -29,7 +30,6 @@
 #include "ViewModels/NiagaraScratchPadScriptViewModel.h"
 #include "ViewModels/NiagaraScratchPadViewModel.h"
 #include "ViewModels/NiagaraSystemViewModel.h"
-#include "ViewModels/NiagaraEmitterViewModel.h"
 #include "ViewModels/Stack/NiagaraStackFunctionInput.h"
 #include "ViewModels/Stack/NiagaraStackFunctionInputCollection.h"
 #include "ViewModels/Stack/NiagaraStackGraphUtilities.h"
@@ -45,8 +45,6 @@
 #include "Widgets/Layout/SBox.h"
 #include "Widgets/NiagaraHLSLSyntaxHighlighter.h"
 #include "Widgets/Text/SMultiLineEditableText.h"
-#include "Widgets/Layout/SSpacer.h"
-#include "Widgets/Layout/SSeparator.h"
 #include "Widgets/Input/SSpinBox.h"
 #include "Widgets/Input/SNumericEntryBox.h"
 #include "NiagaraNodeAssignment.h"
diff --git a/Engine/Plugins/FX/Niagara/Source/NiagaraEditorWidgets/Private/Stack/SNiagaraStackTableRow.cpp b/Engine/Plugins/FX/Niagara/Source/NiagaraEditorWidgets/Private/Stack/SNiagaraStackTableRow.cpp
--- a/Engine/Plugins/FX/Niagara/Source/NiagaraEditorWidgets/Private/Stack/SNiagaraStackTableRow.cpp
+++ b/Engine/Plugins/FX/Niagara/Source/NiagaraEditorWidgets/Private/Stack/SNiagaraStackTableRow.cpp
@@ -18,9 +18,8 @@
 #include "Framework/Application/SlateApplication.h"
 #include "ViewModels/NiagaraEmitterViewModel.h"
 #include "Styling/StyleColors.h"
+#include "ScopedTransaction.h"
 #include "NiagaraEmitterEditorData.h"
-#include "ViewModels/Stack/NiagaraStackFunctionInputCollection.h"
-#include "ViewModels/Stack/NiagaraStackInputCategory.h"
 
 #define LOCTEXT_NAMESPACE "NiagaraStackTableRow"
 
