- Debug dump properties moved to transient to avoid accidentally being left on

#rnx

diff --git a/Engine/Plugins/FX/Niagara/Source/Niagara/Classes/NiagaraSystem.h b/Engine/Plugins/FX/Niagara/Source/Niagara/Classes/NiagaraSystem.h
--- a/Engine/Plugins/FX/Niagara/Source/Niagara/Classes/NiagaraSystem.h
+++ b/Engine/Plugins/FX/Niagara/Source/Niagara/Classes/NiagaraSystem.h
@@ -680,11 +680,11 @@ public:
 		return Found ? *Found : nullptr;
 	}
 	
-	UPROPERTY(EditAnywhere, Category = "Debug")
-	bool bDumpDebugSystemInfo;
+	UPROPERTY(EditAnywhere, Category = "Debug", Transient)
+	bool bDumpDebugSystemInfo = false;
 
-	UPROPERTY(EditAnywhere, Category = "Debug")
-	bool bDumpDebugEmitterInfo;
+	UPROPERTY(EditAnywhere, Category = "Debug", Transient)
+	bool bDumpDebugEmitterInfo = false;
 
 	bool bFullyLoaded = false;
 
