Fixed a bug with the component renderer not using the correct type with lwc changes

diff --git a/Engine/Plugins/FX/Niagara/Source/Niagara/Private/NiagaraComponentRendererProperties.cpp b/Engine/Plugins/FX/Niagara/Source/Niagara/Private/NiagaraComponentRendererProperties.cpp
--- a/Engine/Plugins/FX/Niagara/Source/Niagara/Private/NiagaraComponentRendererProperties.cpp
+++ b/Engine/Plugins/FX/Niagara/Source/Niagara/Private/NiagaraComponentRendererProperties.cpp
@@ -11,7 +11,6 @@
 #include "Widgets/Images/SImage.h"
 #include "Styling/SlateIconFinder.h"
 #include "Widgets/SWidget.h"
-#include "Styling/SlateBrush.h"
 #include "AssetThumbnail.h"
 #include "Widgets/Text/STextBlock.h"
 #endif
@@ -71,7 +70,32 @@ NIAGARA_API FNiagaraTypeDefinition UNiagaraComponentRendererProperties::ToNiagar
 		FStructProperty* StructProperty = (FStructProperty*)Property;
 		if (StructProperty->Struct)
 		{
-			return FNiagaraTypeDefinition(FNiagaraTypeHelper::FindNiagaraFriendlyTopLevelStruct(StructProperty->Struct));
+			if (StructProperty->Struct->GetFName() == NAME_Vector2D || StructProperty->Struct->GetFName() == NAME_Vector2d) 
+			{
+				return GetFVector2DDef();
+			}
+
+			if (StructProperty->Struct->GetFName() == NAME_Vector || StructProperty->Struct->GetFName() == NAME_Vector3d)
+			{
+				return GetFVectorDef();
+			}
+
+			if (StructProperty->Struct->GetFName() == NAME_Vector4 || StructProperty->Struct->GetFName() == NAME_Vector4d)
+			{
+				return GetFVector4Def();
+			}
+
+			if (StructProperty->Struct->GetFName() == NAME_Quat || StructProperty->Struct->GetFName() == NAME_Quat4d)
+			{
+				return GetFQuatDef();
+			}
+
+			if (StructProperty->Struct->GetFName() == FName("NiagaraPosition"))
+			{
+				return FNiagaraTypeDefinition::GetPositionStruct();
+			}
+
+			return FNiagaraTypeDefinition(StructProperty->Struct);
 		}
 	}
 
@@ -191,6 +215,23 @@ void UNiagaraComponentRendererProperties::PostLoad()
 				break;
 			}
 		}
+
+#if WITH_EDITOR
+		// check if any of the bound properties was changed (e.g. to a lwc type) and regenerate the bindings if necessary
+		if (!TemplateComponent)
+		{
+			continue;
+		}
+		for (TFieldIterator<FProperty> PropertyIt(TemplateComponent->GetClass()); PropertyIt; ++PropertyIt)
+		{
+			FProperty* Property = *PropertyIt;
+			if (Property && Property->GetFName() == Binding.PropertyName)
+			{
+				Binding.PropertyType = ToNiagaraType(Property);
+				break;
+			}
+		}
+#endif
 	}
 	EnabledBinding.PostLoad(InSourceMode);
 	RendererVisibilityTagBinding.PostLoad(InSourceMode);
diff --git a/Engine/Plugins/FX/Niagara/Source/Niagara/Private/NiagaraRendererComponents.cpp b/Engine/Plugins/FX/Niagara/Source/Niagara/Private/NiagaraRendererComponents.cpp
--- a/Engine/Plugins/FX/Niagara/Source/Niagara/Private/NiagaraRendererComponents.cpp
+++ b/Engine/Plugins/FX/Niagara/Source/Niagara/Private/NiagaraRendererComponents.cpp
@@ -642,6 +642,38 @@ void FNiagaraRendererComponents::TickPropertyBindings(
 						PropertyType = FNiagaraTypeDefinition(StructProp->Struct, FNiagaraTypeDefinition::EAllowUnfriendlyStruct::Allow);
 						break;
 					}
+
+					// we can't determine the struct type, but since it's the same size we assume it's the same type as the source var
+					FNiagaraTypeDefinition VarType = DataVariable.GetType();
+					if (Property->GetSize() == VarType.GetSize())
+					{
+						break;
+					}
+
+					// the property is missing the struct info and has a different size. In case of LWC types we try to guess the type if it fits the size.
+					bool bDoubleSize = Property->GetSize() == VarType.GetSize() * 2;
+					if (VarType == FNiagaraTypeDefinition::GetVec2Def() && bDoubleSize)
+					{
+						PropertyType = UNiagaraComponentRendererProperties::GetFVector2DDef();
+					}
+					else if ((VarType == FNiagaraTypeDefinition::GetVec3Def() || VarType == FNiagaraTypeDefinition::GetPositionDef()) && bDoubleSize)
+					{
+						PropertyType = UNiagaraComponentRendererProperties::GetFVectorDef();
+					}
+					else if (VarType == FNiagaraTypeDefinition::GetVec4Def() && bDoubleSize)
+					{
+						PropertyType = UNiagaraComponentRendererProperties::GetFVector4Def();
+					}
+					else if (VarType == FNiagaraTypeDefinition::GetQuatDef() && bDoubleSize)
+					{
+						PropertyType = UNiagaraComponentRendererProperties::GetFQuatDef();
+					}
+					else
+					{
+						// unable to resolve the type, invalidate everything to skip the tick
+						PropertyType = FNiagaraTypeDefinition();
+						DataVariable = FNiagaraVariable();
+					}
 				}
 			}
 		}
