Setting fixed bounds for system/emitter from toolbar correctly updates details view. Also fixed bounds in emitter is now hidden when system fixed bounds is enabled

diff --git a/Engine/Plugins/FX/Niagara/Source/Niagara/Classes/NiagaraEmitter.h b/Engine/Plugins/FX/Niagara/Source/Niagara/Classes/NiagaraEmitter.h
--- a/Engine/Plugins/FX/Niagara/Source/Niagara/Classes/NiagaraEmitter.h
+++ b/Engine/Plugins/FX/Niagara/Source/Niagara/Classes/NiagaraEmitter.h
@@ -596,6 +596,8 @@ public:
 
 	TConstArrayView<TUniquePtr<FNiagaraBoundsCalculator>> GetBoundsCalculators() const { return MakeArrayView(BoundsCalculators); }
 
+	void SetFixedBounds(const FBox& NewBounds);
+
 	UPROPERTY()
 	FNiagaraParameterStore RendererBindings;
 
diff --git a/Engine/Plugins/FX/Niagara/Source/Niagara/Classes/NiagaraSystem.h b/Engine/Plugins/FX/Niagara/Source/Niagara/Classes/NiagaraSystem.h
--- a/Engine/Plugins/FX/Niagara/Source/Niagara/Classes/NiagaraSystem.h
+++ b/Engine/Plugins/FX/Niagara/Source/Niagara/Classes/NiagaraSystem.h
@@ -294,7 +294,8 @@ class NIAGARA_API UNiagaraSystem : public UFXSystemAsset, public INiagaraParamet
 public:
 #if WITH_EDITOR
 	DECLARE_MULTICAST_DELEGATE_OneParam(FOnSystemCompiled, UNiagaraSystem*);
-	DECLARE_MULTICAST_DELEGATE_OneParam(FOnSystemPostEditChange, UNiagaraSystem*);
+	DECLARE_MULTICAST_DELEGATE_OneParam(FOnSystemPostEditChange, UNiagaraSystem*);	
+	DECLARE_MULTICAST_DELEGATE(FOnPropertiesChanged);
 #endif
 	//TestChange
 
@@ -472,6 +473,9 @@ public:
 	/** Delegate called on PostEditChange.*/
 	FOnSystemPostEditChange& OnSystemPostEditChange();
 
+	/** Delegate called when a property is changed **/
+	FOnPropertiesChanged& OnPropertiesChanged();
+	
 	/** Gets editor specific data stored with this system. */
 	UNiagaraEditorDataBase* GetEditorData();
 
@@ -702,7 +706,7 @@ public:
 	const TArray<FName>& GetUserDINamesReadInSystemScripts() const;
 
 	FBox GetFixedBounds() const;
-	FORCEINLINE void SetFixedBounds(const FBox& Box) { FixedBounds = Box;  }
+	void SetFixedBounds(const FBox& Box);
 
 #if WITH_EDITOR
 	void SetEffectType(UNiagaraEffectType* EffectType);
@@ -867,6 +871,8 @@ protected:
 
 	/** A multicast delegate which is called whenever this system's properties are changed. */
 	FOnSystemPostEditChange OnSystemPostEditChangeDelegate;
+
+	FOnPropertiesChanged OnPropertiesChangedDelegate;
 #endif
 
 	/** The fixed bounding box value. bFixedBounds is the condition whether the fixed bounds can be edited. */
diff --git a/Engine/Plugins/FX/Niagara/Source/Niagara/Private/NiagaraEmitter.cpp b/Engine/Plugins/FX/Niagara/Source/Niagara/Private/NiagaraEmitter.cpp
--- a/Engine/Plugins/FX/Niagara/Source/Niagara/Private/NiagaraEmitter.cpp
+++ b/Engine/Plugins/FX/Niagara/Source/Niagara/Private/NiagaraEmitter.cpp
@@ -1309,6 +1309,20 @@ if (AllocationMode == EParticleAllocationMode::FixedCount)
 	RebuildRendererBindings();
 }
 
+void UNiagaraEmitter::SetFixedBounds(const FBox& NewBounds)
+{
+	if (!bFixedBounds || FixedBounds != NewBounds)
+	{
+		bFixedBounds = true;
+		FixedBounds = NewBounds;
+
+#if WITH_EDITOR		
+		FPropertyChangedEvent FixedBoundsPropertyChangedEvent(FindFProperty<FProperty>(StaticClass(), "FixedBounds"), EPropertyChangeType::ValueSet);
+		PostEditChangeProperty(FixedBoundsPropertyChangedEvent);
+#endif
+	}
+}
+
 void UNiagaraEmitter::RebuildRendererBindings()
 {
 #if WITH_EDITORONLY_DATA
diff --git a/Engine/Plugins/FX/Niagara/Source/Niagara/Private/NiagaraEmitterInstance.cpp b/Engine/Plugins/FX/Niagara/Source/Niagara/Private/NiagaraEmitterInstance.cpp
--- a/Engine/Plugins/FX/Niagara/Source/Niagara/Private/NiagaraEmitterInstance.cpp
+++ b/Engine/Plugins/FX/Niagara/Source/Niagara/Private/NiagaraEmitterInstance.cpp
@@ -980,14 +980,13 @@ void FNiagaraEmitterInstance::CalculateFixedBounds(const FTransform& ToWorldSpac
 		return;
 
 	CachedEmitter->Modify();
-	CachedEmitter->bFixedBounds = true;
 	if (CachedEmitter->bLocalSpace)
 	{
-		CachedEmitter->FixedBounds = Bounds;
+		CachedEmitter->SetFixedBounds(Bounds);
 	}
 	else
 	{
-		CachedEmitter->FixedBounds = Bounds.TransformBy(ToWorldSpace);
+		CachedEmitter->SetFixedBounds(Bounds.TransformBy(ToWorldSpace));
 	}
 
 	CachedBounds = Bounds;
diff --git a/Engine/Plugins/FX/Niagara/Source/Niagara/Private/NiagaraSystem.cpp b/Engine/Plugins/FX/Niagara/Source/Niagara/Private/NiagaraSystem.cpp
--- a/Engine/Plugins/FX/Niagara/Source/Niagara/Private/NiagaraSystem.cpp
+++ b/Engine/Plugins/FX/Niagara/Source/Niagara/Private/NiagaraSystem.cpp
@@ -1093,6 +1093,8 @@ void UNiagaraSystem::PostEditChangeProperty(struct FPropertyChangedEvent& Proper
 
 	UpdateContext.CommitUpdate();
 
+	OnPropertiesChangedDelegate.Broadcast();
+	
 	static FName SkipReset = TEXT("SkipSystemResetOnChange");
 	bool bPropertyHasSkip = PropertyChangedEvent.Property && PropertyChangedEvent.Property->HasMetaData(SkipReset);
 	bool bMemberHasSkip = PropertyChangedEvent.MemberProperty && PropertyChangedEvent.MemberProperty->HasMetaData(SkipReset);
@@ -2068,6 +2070,20 @@ FBox UNiagaraSystem::GetFixedBounds() const
 	return FixedBounds;
 }
 
+void UNiagaraSystem::SetFixedBounds(const FBox& NewBounds)
+{
+	if (!bFixedBounds || FixedBounds != NewBounds)
+	{
+		bFixedBounds = true;
+		FixedBounds = NewBounds;
+
+#if WITH_EDITOR		
+		FPropertyChangedEvent FixedBoundsPropertyChangedEvent(FindFProperty<FProperty>(StaticClass(), "FixedBounds"), EPropertyChangeType::ValueSet);
+		PostEditChangeProperty(FixedBoundsPropertyChangedEvent);
+#endif
+	}
+}
+
 void CheckDICompileInfo(const TArray<FNiagaraScriptDataInterfaceCompileInfo>& ScriptDICompileInfos, bool& bOutbHasSystemDIsWithPerInstanceData, TArray<FName>& OutUserDINamesReadInSystemScripts)
 {
 	for (const FNiagaraScriptDataInterfaceCompileInfo& ScriptDICompileInfo : ScriptDICompileInfos)
@@ -2348,6 +2364,11 @@ UNiagaraSystem::FOnSystemPostEditChange& UNiagaraSystem::OnSystemPostEditChange(
 	return OnSystemPostEditChangeDelegate;
 }
 
+UNiagaraSystem::FOnPropertiesChanged& UNiagaraSystem::OnPropertiesChanged()
+{
+	return OnPropertiesChangedDelegate;
+}
+
 void UNiagaraSystem::ForceGraphToRecompileOnNextCheck()
 {
 	check(SystemSpawnScript->GetLatestSource() == SystemUpdateScript->GetLatestSource());
diff --git a/Engine/Plugins/FX/Niagara/Source/NiagaraEditor/Private/NiagaraEmitterDetailsCustomization.cpp b/Engine/Plugins/FX/Niagara/Source/NiagaraEditor/Private/NiagaraEmitterDetailsCustomization.cpp
--- a/Engine/Plugins/FX/Niagara/Source/NiagaraEditor/Private/NiagaraEmitterDetailsCustomization.cpp
+++ b/Engine/Plugins/FX/Niagara/Source/NiagaraEditor/Private/NiagaraEmitterDetailsCustomization.cpp
@@ -6,15 +6,59 @@
 #include "DetailCategoryBuilder.h"
 #include "NiagaraEmitter.h"
 #include "NiagaraEditorModule.h"
+#include "DetailWidgetRow.h"
 
-TSharedRef<IDetailCustomization> FNiagaraEmitterDetails::MakeInstance()
+#define LOCTEXT_NAMESPACE "FNiagaraEmitterDetails"
+
+TSharedRef<IDetailCustomization> FNiagaraEmitterDetails::MakeInstance(UNiagaraSystem* System)
 {
-	return MakeShared<FNiagaraEmitterDetails>();
+	return MakeShareable(new FNiagaraEmitterDetails(System));
 }
 
-void FNiagaraEmitterDetails::CustomizeDetails(IDetailLayoutBuilder& InDetailLayout)
+PRAGMA_DISABLE_OPTIMIZATION
+void FNiagaraEmitterDetails::CustomizeDetails(IDetailLayoutBuilder& DetailBuilder)
 {
-	TSharedPtr<IPropertyHandle> EventHandlersPropertyHandle = InDetailLayout.GetProperty(UNiagaraEmitter::PrivateMemberNames::EventHandlerScriptProps);
+	TSharedPtr<IPropertyHandle> EventHandlersPropertyHandle = DetailBuilder.GetProperty(UNiagaraEmitter::PrivateMemberNames::EventHandlerScriptProps);
 	EventHandlersPropertyHandle->MarkHiddenByCustomization();
+
+ 	if (System && System->bFixedBounds)
+ 	{
+ 		static const FName EmitterCategoryName = TEXT("Emitter");
+ 		IDetailCategoryBuilder* SourceCategory = &DetailBuilder.EditCategory(EmitterCategoryName);
+ 		{
+ 			TArray<TSharedRef<IPropertyHandle>> Properties;
+ 			SourceCategory->GetDefaultProperties(Properties, true, true);
+ 			
+ 			for (TSharedPtr<IPropertyHandle> Property : Properties)
+ 			{
+ 				FProperty* PropertyPtr = Property->GetProperty();
+ 				if (PropertyPtr->GetName()== "FixedBounds")
+ 				{
+					SourceCategory->AddCustomRow(Property->GetPropertyDisplayName())
+ 					.NameContent()
+ 					[
+ 						SNew(STextBlock)
+						.Text(Property->GetPropertyDisplayName())
+						.Font(DetailBuilder.GetDetailFont())
+ 					]
+ 					.ValueContent()
+ 					[
+ 						SNew(STextBlock)
+						.Text(LOCTEXT("FixedBoundsOverridenBySystem", "Fixed bounds cannot be set here while overridden by system FixedBounds."))
+						.AutoWrapText(true)
+						.Font(DetailBuilder.GetDetailFontItalic())
+ 					];
+
+ 					DetailBuilder.HideProperty(Property);
+ 				}
+ 				else
+ 				{
+ 					SourceCategory->AddProperty(Property);
+ 				}
+ 			}
+ 		}
+ 	}
 }
 
+#undef LOCTEXT_NAMESPACE
+
diff --git a/Engine/Plugins/FX/Niagara/Source/NiagaraEditor/Private/NiagaraEmitterDetailsCustomization.h b/Engine/Plugins/FX/Niagara/Source/NiagaraEditor/Private/NiagaraEmitterDetailsCustomization.h
--- a/Engine/Plugins/FX/Niagara/Source/NiagaraEditor/Private/NiagaraEmitterDetailsCustomization.h
+++ b/Engine/Plugins/FX/Niagara/Source/NiagaraEditor/Private/NiagaraEmitterDetailsCustomization.h
@@ -3,6 +3,7 @@
 #pragma once
 
 #include "IDetailCustomization.h"
+#include "NiagaraSystem.h"
 
 class IDetailLayoutBuilder;
 
@@ -14,9 +15,14 @@ class FNiagaraEmitterDetails : public IDetailCustomization
 {
 public:
 	/** Makes a new instance of this detail layout class for a specific detail view requesting it */
-	static TSharedRef<IDetailCustomization> MakeInstance();
+	static TSharedRef<IDetailCustomization> MakeInstance(UNiagaraSystem* System);
 
 	/** IDetailCustomization interface */
 	virtual void CustomizeDetails(IDetailLayoutBuilder& InDetailLayout) override;
+
+private:
+	FNiagaraEmitterDetails(UNiagaraSystem* InSystem) : System(InSystem) { }
+
+	UNiagaraSystem* System;
 };
 
diff --git a/Engine/Plugins/FX/Niagara/Source/NiagaraEditor/Private/ViewModels/NiagaraSystemViewModel.cpp b/Engine/Plugins/FX/Niagara/Source/NiagaraEditor/Private/ViewModels/NiagaraSystemViewModel.cpp
--- a/Engine/Plugins/FX/Niagara/Source/NiagaraEditor/Private/ViewModels/NiagaraSystemViewModel.cpp
+++ b/Engine/Plugins/FX/Niagara/Source/NiagaraEditor/Private/ViewModels/NiagaraSystemViewModel.cpp
@@ -2214,7 +2214,6 @@ void FNiagaraSystemViewModel::UpdateSystemFixedBounds()
 	{
 		GetSystem().Modify();
 		
-		GetSystem().bFixedBounds = true;
 		GetSystem().SetFixedBounds(SystemInstance->GetLocalBounds());
 
 		PreviewComponent->MarkRenderTransformDirty();
diff --git a/Engine/Plugins/FX/Niagara/Source/NiagaraEditor/Private/ViewModels/Stack/NiagaraStackEmitterSettingsGroup.cpp b/Engine/Plugins/FX/Niagara/Source/NiagaraEditor/Private/ViewModels/Stack/NiagaraStackEmitterSettingsGroup.cpp
--- a/Engine/Plugins/FX/Niagara/Source/NiagaraEditor/Private/ViewModels/Stack/NiagaraStackEmitterSettingsGroup.cpp
+++ b/Engine/Plugins/FX/Niagara/Source/NiagaraEditor/Private/ViewModels/Stack/NiagaraStackEmitterSettingsGroup.cpp
@@ -21,6 +21,9 @@ void UNiagaraStackEmitterPropertiesItem::Initialize(FRequiredEntryData InRequire
 	Super::Initialize(InRequiredEntryData, TEXT("EmitterProperties"));
 	Emitter = GetEmitterViewModel()->GetEmitter();
 	Emitter->OnPropertiesChanged().AddUObject(this, &UNiagaraStackEmitterPropertiesItem::EmitterPropertiesChanged);
+
+	// We tie into the system properties changed event to know when the fixed bounds changes so we update our state accordingly
+	GetSystemViewModel()->GetSystem().OnPropertiesChanged().AddUObject(this, &UNiagaraStackEmitterPropertiesItem::EmitterPropertiesChanged);
 }
 
 void UNiagaraStackEmitterPropertiesItem::FinalizeInternal()
@@ -28,6 +31,7 @@ void UNiagaraStackEmitterPropertiesItem::FinalizeInternal()
 	if (Emitter.IsValid())
 	{
 		Emitter->OnPropertiesChanged().RemoveAll(this);
+	GetSystemViewModel()->GetSystem().OnPropertiesChanged().RemoveAll(this);
 	}
 	Super::FinalizeInternal();
 }
@@ -105,7 +109,7 @@ void UNiagaraStackEmitterPropertiesItem::RefreshChildrenInternal(const TArray<UN
 		EmitterObject = NewObject<UNiagaraStackObject>(this);
 		FRequiredEntryData RequiredEntryData(GetSystemViewModel(), GetEmitterViewModel(), FExecutionCategoryNames::Emitter, NAME_None, GetStackEditorData());
 		EmitterObject->Initialize(RequiredEntryData, Emitter.Get(), GetStackEditorDataKey());
-		EmitterObject->RegisterInstancedCustomPropertyLayout(UNiagaraEmitter::StaticClass(), FOnGetDetailCustomizationInstance::CreateStatic(&FNiagaraEmitterDetails::MakeInstance));
+		EmitterObject->RegisterInstancedCustomPropertyLayout(UNiagaraEmitter::StaticClass(), FOnGetDetailCustomizationInstance::CreateStatic(&FNiagaraEmitterDetails::MakeInstance, &GetSystemViewModel()->GetSystem()));
 	}
 
 	NewChildren.Add(EmitterObject);
@@ -151,6 +155,11 @@ void UNiagaraStackEmitterPropertiesItem::EmitterPropertiesChanged()
 		// so guard against receiving an event when finalized here.
 		bCanResetToBaseCache.Reset();
 		RefreshChildren();
+
+		if (EmitterObject)
+		{
+			EmitterObject->InvalidateDetailRows();
+		}
 	}
 }
 
diff --git a/Engine/Plugins/FX/Niagara/Source/NiagaraEditor/Private/ViewModels/Stack/NiagaraStackObject.cpp b/Engine/Plugins/FX/Niagara/Source/NiagaraEditor/Private/ViewModels/Stack/NiagaraStackObject.cpp
--- a/Engine/Plugins/FX/Niagara/Source/NiagaraEditor/Private/ViewModels/Stack/NiagaraStackObject.cpp
+++ b/Engine/Plugins/FX/Niagara/Source/NiagaraEditor/Private/ViewModels/Stack/NiagaraStackObject.cpp
@@ -72,6 +72,13 @@ UObject* UNiagaraStackObject::GetDisplayedObject() const
 	return WeakObject.Get();
 }
 
+void UNiagaraStackObject::InvalidateDetailRows()
+{
+	PropertyRowGenerator.Reset();
+
+	RefreshChildren();
+}
+
 void UNiagaraStackObject::FinalizeInternal()
 {
 	if (PropertyRowGenerator.IsValid())
diff --git a/Engine/Plugins/FX/Niagara/Source/NiagaraEditor/Private/ViewModels/Stack/NiagaraStackSystemPropertiesItem.cpp b/Engine/Plugins/FX/Niagara/Source/NiagaraEditor/Private/ViewModels/Stack/NiagaraStackSystemPropertiesItem.cpp
--- a/Engine/Plugins/FX/Niagara/Source/NiagaraEditor/Private/ViewModels/Stack/NiagaraStackSystemPropertiesItem.cpp
+++ b/Engine/Plugins/FX/Niagara/Source/NiagaraEditor/Private/ViewModels/Stack/NiagaraStackSystemPropertiesItem.cpp
@@ -12,6 +12,8 @@ void UNiagaraStackSystemPropertiesItem::Initialize(FRequiredEntryData InRequired
 {
 	Super::Initialize(InRequiredEntryData, TEXT("SystemProperties"));
 	System = &GetSystemViewModel()->GetSystem();
+	
+	System->OnPropertiesChanged().AddUObject(this, &UNiagaraStackSystemPropertiesItem::SystemPropertiesChanged);
 }
 
 FText UNiagaraStackSystemPropertiesItem::GetDisplayName() const
@@ -75,7 +77,18 @@ void UNiagaraStackSystemPropertiesItem::RefreshChildrenInternal(const TArray<UNi
 
 void UNiagaraStackSystemPropertiesItem::SystemPropertiesChanged()
 {
-	bCanResetToBase.Reset();
+	if (IsFinalized() == false)
+	{
+		// Undo/redo can cause objects to disappear and reappear which can prevent safe removal of delegates
+		// so guard against receiving an event when finalized here.
+		bCanResetToBase.Reset();
+		RefreshChildren();
+
+		if (SystemObject)
+		{
+			SystemObject->InvalidateDetailRows();
+		}
+	}
 }
 
 #undef LOCTEXT_NAMESPACE
diff --git a/Engine/Plugins/FX/Niagara/Source/NiagaraEditor/Public/ViewModels/Stack/NiagaraStackObject.h b/Engine/Plugins/FX/Niagara/Source/NiagaraEditor/Public/ViewModels/Stack/NiagaraStackObject.h
--- a/Engine/Plugins/FX/Niagara/Source/NiagaraEditor/Public/ViewModels/Stack/NiagaraStackObject.h
+++ b/Engine/Plugins/FX/Niagara/Source/NiagaraEditor/Public/ViewModels/Stack/NiagaraStackObject.h
@@ -39,6 +39,8 @@ public:
 	virtual bool GetShouldShowInStack() const override;
 	virtual UObject* GetDisplayedObject() const override;
 
+	void InvalidateDetailRows();
+
 protected:
 	virtual void FinalizeInternal() override;
 
