Removed optimization debug flag

diff --git a/Engine/Plugins/FX/Niagara/Source/NiagaraEditor/Private/NiagaraEmitterDetailsCustomization.cpp b/Engine/Plugins/FX/Niagara/Source/NiagaraEditor/Private/NiagaraEmitterDetailsCustomization.cpp
--- a/Engine/Plugins/FX/Niagara/Source/NiagaraEditor/Private/NiagaraEmitterDetailsCustomization.cpp
+++ b/Engine/Plugins/FX/Niagara/Source/NiagaraEditor/Private/NiagaraEmitterDetailsCustomization.cpp
@@ -15,7 +15,6 @@ TSharedRef<IDetailCustomization> FNiagaraEmitterDetails::MakeInstance(UNiagaraSy
 	return MakeShareable(new FNiagaraEmitterDetails(System));
 }
 
-PRAGMA_DISABLE_OPTIMIZATION
 void FNiagaraEmitterDetails::CustomizeDetails(IDetailLayoutBuilder& DetailBuilder)
 {
 	TSharedPtr<IPropertyHandle> EventHandlersPropertyHandle = DetailBuilder.GetProperty(UNiagaraEmitter::PrivateMemberNames::EventHandlerScriptProps);
