- Order code to fix false PVS issue

#rnx

diff --git a/Engine/Plugins/FX/Niagara/Source/Niagara/Private/DataInterface/NiagaraDataInterfaceUObjectPropertyReader.cpp b/Engine/Plugins/FX/Niagara/Source/Niagara/Private/DataInterface/NiagaraDataInterfaceUObjectPropertyReader.cpp
--- a/Engine/Plugins/FX/Niagara/Source/Niagara/Private/DataInterface/NiagaraDataInterfaceUObjectPropertyReader.cpp
+++ b/Engine/Plugins/FX/Niagara/Source/Niagara/Private/DataInterface/NiagaraDataInterfaceUObjectPropertyReader.cpp
@@ -668,28 +668,29 @@ bool UNiagaraDataInterfaceUObjectPropertyReader::PerInstanceTick(void* PerInstan
 	{
 		InstanceData_GT->WeakUObject = ObjectBinding;
 
-		if (ObjectBinding != nullptr)
+		// If the object we are binding to is an actor we will get the root component / specified component class
+		UObject* ActorComponent = nullptr;
+		if ( AActor* ObjectActor = Cast<AActor>(ObjectBinding) )
 		{
-			// If the object we are binding to is an actor we will get the root component / specified component class
-			UObject* ActorComponent = nullptr;
-			if ( AActor* ObjectActor = Cast<AActor>(ObjectBinding) )
+			if ( SourceActorComponentClass != nullptr )
 			{
-				if ( SourceActorComponentClass != nullptr )
-				{
-					ActorComponent = ObjectActor->FindComponentByClass(SourceActorComponentClass);
-				}
-				else
-				{
-					ActorComponent = ObjectActor->GetRootComponent();
-				}
+				ActorComponent = ObjectActor->FindComponentByClass(SourceActorComponentClass);
 			}
+			else
+			{
+				ActorComponent = ObjectActor->GetRootComponent();
+			}
+		}
 
-			check(ObjectBinding != nullptr);	// Fixes CIS warning
-			for (FNDIPropertyGetter& PropertyGetter : InstanceData_GT->PropertyGetters)
+		for (FNDIPropertyGetter& PropertyGetter : InstanceData_GT->PropertyGetters)
+		{
+			PropertyGetter.WeakProperty.Reset();
+			PropertyGetter.PropertyCopyFunction = nullptr;
+
+			if ( ObjectBinding != nullptr )
 			{
 				BindPropertyGetter(PropertyGetter, ObjectBinding);
-
-				if ( PropertyGetter.PropertyCopyFunction == nullptr && ActorComponent != nullptr )
+				if (PropertyGetter.PropertyCopyFunction == nullptr && ActorComponent != nullptr)
 				{
 					BindPropertyGetter(PropertyGetter, ActorComponent);
 				}
@@ -700,14 +701,6 @@ bool UNiagaraDataInterfaceUObjectPropertyReader::PerInstanceTick(void* PerInstan
 				}
 			}
 		}
-		else
-		{
-			for (FNDIPropertyGetter& PropertyGetter : InstanceData_GT->PropertyGetters)
-			{
-				PropertyGetter.WeakProperty.Reset();
-				PropertyGetter.PropertyCopyFunction = nullptr;
-			}
-		}
 	}
 
 	// Update our data store as we can not read object's async it's unsafe
