- Fixed force solo to not kill the component
- When setting custom time dilation Niagara will change into solo mode, and back out when it is set to 1.0

diff --git a/Engine/Plugins/FX/Niagara/Source/Niagara/Private/NiagaraComponent.cpp b/Engine/Plugins/FX/Niagara/Source/Niagara/Private/NiagaraComponent.cpp
--- a/Engine/Plugins/FX/Niagara/Source/Niagara/Private/NiagaraComponent.cpp
+++ b/Engine/Plugins/FX/Niagara/Source/Niagara/Private/NiagaraComponent.cpp
@@ -899,7 +899,7 @@ bool UNiagaraComponent::InitializeSystem()
 		OverrideParameters.MarkParametersDirty(); // new system instance means new lwc tile, so any position user params need to be re-evaluated
 
 		SystemInstanceController = MakeShared<FNiagaraSystemInstanceController, ESPMode::ThreadSafe>();
-		SystemInstanceController->Initialize(*World, *Asset, &OverrideParameters, this, TickBehavior, bPooled, RandomSeedOffset, bForceSolo);
+		SystemInstanceController->Initialize(*World, *Asset, &OverrideParameters, this, TickBehavior, bPooled, RandomSeedOffset, RequiresSoloMode());
 		SystemInstanceController->SetOnPostTick(FNiagaraSystemInstance::FOnPostTick::CreateUObject(this, &UNiagaraComponent::PostSystemTick_GameThread));
 		SystemInstanceController->SetOnComplete(FNiagaraSystemInstance::FOnComplete::CreateUObject(this, &UNiagaraComponent::OnSystemComplete));
 
@@ -3501,11 +3501,25 @@ void UNiagaraComponent::SetForceSolo(bool bInForceSolo)
 	if (bForceSolo != bInForceSolo)
 	{
 		bForceSolo = bInForceSolo;
-		DestroyInstance();
-		SetComponentTickEnabled(bInForceSolo);
+		UpdateInstanceSoloMode();
 	}
 }
 
+bool UNiagaraComponent::RequiresSoloMode() const
+{
+	return bForceSolo || !FMath::IsNearlyEqual(CustomTimeDilation, 1.0f);
+}
+
+void UNiagaraComponent::UpdateInstanceSoloMode()
+{
+	const bool bRequireSoloMode = RequiresSoloMode();
+	if (SystemInstanceController)
+	{
+		SystemInstanceController->SetForceSolo(bRequireSoloMode);
+	}
+	SetComponentTickEnabled(bRequireSoloMode);
+}
+
 void UNiagaraComponent::SetGpuComputeDebug(bool bEnableDebug)
 {
 	if (bEnableGpuComputeDebug != bEnableDebug)
@@ -3520,12 +3534,11 @@ void UNiagaraComponent::SetGpuComputeDebug(bool bEnableDebug)
 
 void UNiagaraComponent::SetCustomTimeDilation(float Dilation)
 {
-	if ( !bForceSolo )
+	if ( !FMath::IsNearlyEqual(CustomTimeDilation, Dilation) )
 	{
-		UE_LOG(LogNiagara, Warning, TEXT("Attempting to set custom time dilation on a NiagaraComponent(%s) that is not in solo mode, value will be ignored."), *GetNameSafe(this));
-		return;
+		CustomTimeDilation = Dilation;
+		UpdateInstanceSoloMode();
 	}
-	CustomTimeDilation = Dilation;
 }
 
 void UNiagaraComponent::SetAutoAttachmentParameters(USceneComponent* Parent, FName SocketName, EAttachmentRule LocationRule, EAttachmentRule RotationRule, EAttachmentRule ScaleRule)
diff --git a/Engine/Plugins/FX/Niagara/Source/Niagara/Private/NiagaraDebugHud.cpp b/Engine/Plugins/FX/Niagara/Source/Niagara/Private/NiagaraDebugHud.cpp
--- a/Engine/Plugins/FX/Niagara/Source/Niagara/Private/NiagaraDebugHud.cpp
+++ b/Engine/Plugins/FX/Niagara/Source/Niagara/Private/NiagaraDebugHud.cpp
@@ -2365,6 +2365,15 @@ void FNiagaraDebugHud::DrawComponents(FNiagaraWorldManager* WorldManager, UCanva
 						StringBuilder.Appendf(TEXT("Pooled - %s\n"), *PoolingMethodEnum->GetNameStringByIndex((int32)NiagaraComponent->PoolingMethod));
 					}
 
+					if (SystemInstanceController->IsSolo())
+					{
+						if ( NiagaraComponent->GetForceSolo() )
+						{
+							StringBuilder.Append(TEXT("ForceSolo "));
+						}
+						StringBuilder.Append(TEXT("IsSolo\n"));
+					}
+
 					if ( NiagaraComponent->bHiddenInGame || !NiagaraComponent->GetVisibleFlag() || (OwnerActor && OwnerActor->IsHidden()) )
 					{
 						StringBuilder.Appendf(TEXT("HiddenInGame(%d) Visible(%d)"), NiagaraComponent->bHiddenInGame, NiagaraComponent->GetVisibleFlag());
diff --git a/Engine/Plugins/FX/Niagara/Source/Niagara/Private/NiagaraSystemInstance.cpp b/Engine/Plugins/FX/Niagara/Source/Niagara/Private/NiagaraSystemInstance.cpp
--- a/Engine/Plugins/FX/Niagara/Source/Niagara/Private/NiagaraSystemInstance.cpp
+++ b/Engine/Plugins/FX/Niagara/Source/Niagara/Private/NiagaraSystemInstance.cpp
@@ -81,6 +81,32 @@ static FAutoConsoleVariableRef CVarNiagaraAllowDeferredReset(
 	ECVF_Default
 );
 
+bool DoSystemDataInterfacesRequireSolo(const UNiagaraSystem& System, const FNiagaraUserRedirectionParameterStore* OverrideParameters)
+{
+	if (FNiagaraSystemSimulation::UseLegacySystemSimulationContexts())
+	{
+		if (System.HasSystemScriptDIsWithPerInstanceData())
+		{
+			return true;
+		}
+
+		const TArray<FName>& UserDINamesReadInSystemScripts = System.GetUserDINamesReadInSystemScripts();
+		if (OverrideParameters != nullptr && UserDINamesReadInSystemScripts.Num() > 0)
+		{
+			TArray<FNiagaraVariable> OverrideParameterVariables;
+			OverrideParameters->GetParameters(OverrideParameterVariables);
+			for (const FNiagaraVariable& OverrideParameterVariable : OverrideParameterVariables)
+			{
+				if (OverrideParameterVariable.IsDataInterface() && UserDINamesReadInSystemScripts.Contains(OverrideParameterVariable.GetName()))
+				{
+					return true;
+				}
+			}
+		}
+	}
+	return false;
+}
+
 FNiagaraSystemInstance::FNiagaraSystemInstance(UWorld& InWorld, UNiagaraSystem& InAsset, FNiagaraUserRedirectionParameterStore* InOverrideParameters,
                                                USceneComponent* InAttachComponent, ENiagaraTickBehavior InTickBehavior, bool bInPooled)
 	: SystemInstanceIndex(INDEX_NONE)
@@ -473,6 +499,8 @@ void FNiagaraSystemInstance::SetSolo(bool bInSolo)
 	}
 	else
 	{
+		ensureMsgf(DoSystemDataInterfacesRequireSolo(*GetSystem(), OverrideParameters) == false, TEXT("Disabling solo mode but data interfaces require it"));
+
 		const ETickingGroup TickGroup = CalculateTickGroup();
 		TSharedPtr<FNiagaraSystemSimulation, ESPMode::ThreadSafe> NewSim = GetWorldManager()->GetSystemSimulation(TickGroup, System);
 
@@ -995,32 +1023,6 @@ bool FNiagaraSystemInstance::IsReadyToRun() const
 	return bAllReadyToRun;
 }
 
-bool DoSystemDataInterfacesRequireSolo(const UNiagaraSystem& System, const FNiagaraUserRedirectionParameterStore* OverrideParameters)
-{
-	if (FNiagaraSystemSimulation::UseLegacySystemSimulationContexts())
-	{
-		if (System.HasSystemScriptDIsWithPerInstanceData())
-		{
-			return true;
-		}
-
-		const TArray<FName>& UserDINamesReadInSystemScripts = System.GetUserDINamesReadInSystemScripts();
-		if (OverrideParameters != nullptr && UserDINamesReadInSystemScripts.Num() > 0)
-		{
-			TArray<FNiagaraVariable> OverrideParameterVariables;
-			OverrideParameters->GetParameters(OverrideParameterVariables);
-			for (const FNiagaraVariable& OverrideParameterVariable : OverrideParameterVariables)
-			{
-				if (OverrideParameterVariable.IsDataInterface() && UserDINamesReadInSystemScripts.Contains(OverrideParameterVariable.GetName()))
-				{
-					return true;
-				}
-			}
-		}
-	}
-	return false;
-}
-
 void FNiagaraSystemInstance::ReInitInternal()
 {
 	check(SystemInstanceIndex == INDEX_NONE);
@@ -2654,6 +2656,22 @@ FNiagaraEmitterInstance* FNiagaraSystemInstance::GetEmitterByID(FGuid InID)
 	return nullptr;
 }
 
+void FNiagaraSystemInstance::SetForceSolo(bool bInForceSolo)
+{
+	// We may be forced into solo mode so if we match nothing to do here
+	bForceSolo = bInForceSolo;
+	if ( bSolo == bInForceSolo )
+	{
+		return;
+	}
+
+	const bool bNewSolo = bForceSolo || DoSystemDataInterfacesRequireSolo(*GetSystem(), OverrideParameters);
+	if (bNewSolo != bSolo)
+	{
+		SetSolo(bNewSolo);
+	}
+}
+
 void FNiagaraSystemInstance::EvaluateBoundFunction(FName FunctionName, bool& UsedOnCpu, bool& UsedOnGpu) const
 {
 	auto ScriptUsesFunction = [&](UNiagaraScript* Script)
diff --git a/Engine/Plugins/FX/Niagara/Source/Niagara/Public/NiagaraComponent.h b/Engine/Plugins/FX/Niagara/Source/Niagara/Public/NiagaraComponent.h
--- a/Engine/Plugins/FX/Niagara/Source/Niagara/Public/NiagaraComponent.h
+++ b/Engine/Plugins/FX/Niagara/Source/Niagara/Public/NiagaraComponent.h
@@ -250,6 +250,11 @@ public:
 	UFUNCTION(BlueprintCallable, Category = Niagara, meta = (DisplayName = "Is In Forced Solo Mode"))
 	bool GetForceSolo()const { return bForceSolo; }
 
+private:
+	bool RequiresSoloMode() const;
+	void UpdateInstanceSoloMode();
+public:
+
 	UFUNCTION(BlueprintCallable, Category = Niagara)
 	void SetGpuComputeDebug(bool bEnableDebug);
 
diff --git a/Engine/Plugins/FX/Niagara/Source/Niagara/Public/NiagaraSystemInstance.h b/Engine/Plugins/FX/Niagara/Source/Niagara/Public/NiagaraSystemInstance.h
--- a/Engine/Plugins/FX/Niagara/Source/Niagara/Public/NiagaraSystemInstance.h
+++ b/Engine/Plugins/FX/Niagara/Source/Niagara/Public/NiagaraSystemInstance.h
@@ -242,6 +242,7 @@ public:
 
 	FNiagaraEmitterInstance* GetEmitterByID(FGuid InID);
 
+	void SetForceSolo(bool bForceSolo);
 	FORCEINLINE bool IsSolo() const { return bSolo; }
 
 	FORCEINLINE bool NeedsGPUTick() const { return ActiveGPUEmitterCount > 0 /*&& Component->IsRegistered()*/ && !IsComplete();}
diff --git a/Engine/Plugins/FX/Niagara/Source/Niagara/Public/NiagaraSystemInstanceController.h b/Engine/Plugins/FX/Niagara/Source/Niagara/Public/NiagaraSystemInstanceController.h
--- a/Engine/Plugins/FX/Niagara/Source/Niagara/Public/NiagaraSystemInstanceController.h
+++ b/Engine/Plugins/FX/Niagara/Source/Niagara/Public/NiagaraSystemInstanceController.h
@@ -88,7 +88,6 @@ public:
 	/** Dumps system instance state and info to the log (Used by fx.Niagara.DumpComponents commandlet) */
 	void DebugDump(bool bFullDump);
 
-
 	// Deferred methods - With NIAGARA_SYSTEM_INSTANCE_CONTROLLER_ASYNC, these are executed in a deferred queue
 	NIAGARA_SYSTEM_INSTANCE_CONTROLLER_DEFFERED_METHOD(Reset)
 	NIAGARA_SYSTEM_INSTANCE_CONTROLLER_DEFFERED_METHOD(AdvanceSimulation)
@@ -114,6 +113,7 @@ public:
 	NIAGARA_SYSTEM_INSTANCE_CONTROLLER_SHIM(GetAge, const)
 	NIAGARA_SYSTEM_INSTANCE_CONTROLLER_SHIM(GetComputeDispatchInterface, const)
 	NIAGARA_SYSTEM_INSTANCE_CONTROLLER_SHIM(IsSolo, const)
+	NIAGARA_SYSTEM_INSTANCE_CONTROLLER_SHIM(SetForceSolo,)
 	NIAGARA_SYSTEM_INSTANCE_CONTROLLER_SHIM(IsComplete, const)
 	NIAGARA_SYSTEM_INSTANCE_CONTROLLER_SHIM(IsPendingSpawn, const)
 	NIAGARA_SYSTEM_INSTANCE_CONTROLLER_SHIM(Reset,)
