- Use TypedOuter rather than system view model to get UNiagaraSystem

#rnx

diff --git a/Engine/Plugins/FX/Niagara/Source/NiagaraEditor/Private/ViewModels/Stack/NiagaraStackEmitterSettingsGroup.cpp b/Engine/Plugins/FX/Niagara/Source/NiagaraEditor/Private/ViewModels/Stack/NiagaraStackEmitterSettingsGroup.cpp
--- a/Engine/Plugins/FX/Niagara/Source/NiagaraEditor/Private/ViewModels/Stack/NiagaraStackEmitterSettingsGroup.cpp
+++ b/Engine/Plugins/FX/Niagara/Source/NiagaraEditor/Private/ViewModels/Stack/NiagaraStackEmitterSettingsGroup.cpp
@@ -23,7 +23,10 @@ void UNiagaraStackEmitterPropertiesItem::Initialize(FRequiredEntryData InRequire
 	Emitter->OnPropertiesChanged().AddUObject(this, &UNiagaraStackEmitterPropertiesItem::EmitterPropertiesChanged);
 
 	// We tie into the system properties changed event to know when the fixed bounds changes so we update our state accordingly
-	GetSystemViewModel()->GetSystem().OnPropertiesChanged().AddUObject(this, &UNiagaraStackEmitterPropertiesItem::EmitterPropertiesChanged);
+	if (UNiagaraSystem* NiagaraSystem = Emitter->GetTypedOuter<UNiagaraSystem>())
+	{
+		NiagaraSystem->OnPropertiesChanged().AddUObject(this, &UNiagaraStackEmitterPropertiesItem::EmitterPropertiesChanged);
+	}
 }
 
 void UNiagaraStackEmitterPropertiesItem::FinalizeInternal()
@@ -31,7 +34,10 @@ void UNiagaraStackEmitterPropertiesItem::FinalizeInternal()
 	if (Emitter.IsValid())
 	{
 		Emitter->OnPropertiesChanged().RemoveAll(this);
-	GetSystemViewModel()->GetSystem().OnPropertiesChanged().RemoveAll(this);
+		if ( UNiagaraSystem* NiagaraSystem = Emitter->GetTypedOuter<UNiagaraSystem>() )
+		{
+			NiagaraSystem->OnPropertiesChanged().RemoveAll(this);
+		}
 	}
 	Super::FinalizeInternal();
 }
