- Fix SetSolo function to handle instances that have no simulation
- Fix TransferInstance to change the simulation pointer when in None mode

#rnx

diff --git a/Engine/Plugins/FX/Niagara/Source/Niagara/Private/NiagaraSystemInstance.cpp b/Engine/Plugins/FX/Niagara/Source/Niagara/Private/NiagaraSystemInstance.cpp
--- a/Engine/Plugins/FX/Niagara/Source/Niagara/Private/NiagaraSystemInstance.cpp
+++ b/Engine/Plugins/FX/Niagara/Source/Niagara/Private/NiagaraSystemInstance.cpp
@@ -487,28 +487,30 @@ void FNiagaraSystemInstance::SetSolo(bool bInSolo)
 	// Wait for any async operations
 	WaitForConcurrentTickDoNotFinalize();
 
-	UNiagaraSystem* System = GetSystem();
-	if (bInSolo)
+	// We only need to transfer the instance into a new simulation if this one is within a simulation
+	if ( SystemSimulation.IsValid() )
 	{
-		TSharedPtr<FNiagaraSystemSimulation, ESPMode::ThreadSafe> NewSoloSim = MakeShared<FNiagaraSystemSimulation, ESPMode::ThreadSafe>();
-		NewSoloSim->Init(System, World, true, TG_MAX);
-
-		NewSoloSim->TransferInstance(this);
-
-		bSolo = true;
-	}
-	else
-	{
-		ensureMsgf(DoSystemDataInterfacesRequireSolo(*GetSystem(), OverrideParameters) == false, TEXT("Disabling solo mode but data interfaces require it"));
+		UNiagaraSystem* System = GetSystem();
+		if (bInSolo)
+		{
+			TSharedPtr<FNiagaraSystemSimulation, ESPMode::ThreadSafe> NewSoloSim = MakeShared<FNiagaraSystemSimulation, ESPMode::ThreadSafe>();
+			NewSoloSim->Init(System, World, true, TG_MAX);
 
-		const ETickingGroup TickGroup = CalculateTickGroup();
-		TSharedPtr<FNiagaraSystemSimulation, ESPMode::ThreadSafe> NewSim = GetWorldManager()->GetSystemSimulation(TickGroup, System);
+			NewSoloSim->TransferInstance(this);
+		}
+		else
+		{
+			ensureMsgf(DoSystemDataInterfacesRequireSolo(*GetSystem(), OverrideParameters) == false, TEXT("Disabling solo mode but data interfaces require it"));
 
-		NewSim->TransferInstance(this);
+			const ETickingGroup TickGroup = CalculateTickGroup();
+			TSharedPtr<FNiagaraSystemSimulation, ESPMode::ThreadSafe> NewSim = GetWorldManager()->GetSystemSimulation(TickGroup, System);
 
-		bSolo = false;
+			NewSim->TransferInstance(this);
+		}
 	}
 
+	bSolo = bInSolo;
+
 	// Execute any pending finalize
 	if (FinalizeRef.IsPending())
 	{
diff --git a/Engine/Plugins/FX/Niagara/Source/Niagara/Private/NiagaraSystemSimulation.cpp b/Engine/Plugins/FX/Niagara/Source/Niagara/Private/NiagaraSystemSimulation.cpp
--- a/Engine/Plugins/FX/Niagara/Source/Niagara/Private/NiagaraSystemSimulation.cpp
+++ b/Engine/Plugins/FX/Niagara/Source/Niagara/Private/NiagaraSystemSimulation.cpp
@@ -832,6 +832,10 @@ void FNiagaraSystemSimulation::TransferInstance(FNiagaraSystemInstance* Instance
 	{
 		// Nothing to do as nothing to transfer
 		case ENiagaraSystemInstanceState::None:
+			if (Instance->SystemSimulation.IsValid())
+			{
+				Instance->SystemSimulation = this->AsShared();
+			}
 			break;
 
 		case ENiagaraSystemInstanceState::PendingSpawn:
