Remove FNiagaraEditorModule::UNiagaraReservedParametersManager and consolidate caching logic with existing handling for parameter collections.

diff --git a/Engine/Plugins/FX/Niagara/Source/NiagaraEditor/Private/NiagaraEditorModule.cpp b/Engine/Plugins/FX/Niagara/Source/NiagaraEditor/Private/NiagaraEditorModule.cpp
--- a/Engine/Plugins/FX/Niagara/Source/NiagaraEditor/Private/NiagaraEditorModule.cpp
+++ b/Engine/Plugins/FX/Niagara/Source/NiagaraEditor/Private/NiagaraEditorModule.cpp
@@ -10,7 +10,6 @@
 #include "ISettingsModule.h"
 #include "SequencerChannelInterface.h"
 #include "SequencerSettings.h"
-#include "AssetRegistryModule.h"
 #include "ThumbnailRendering/ThumbnailManager.h"
 #include "Stats/Stats.h"
 #include "Subsystems/ImportSubsystem.h"
@@ -805,7 +804,8 @@ void FNiagaraEditorModule::StartupModule()
 		FAssetRegistryModule& AssetRegistryModule = FModuleManager::LoadModuleChecked<FAssetRegistryModule>("AssetRegistry");
 		AssetRegistryModule.Get().OnFilesLoaded().AddLambda([this]()
 		{
-			RefreshParameterCollections(true /*AllowLoading*/);
+			ParameterCollectionAssetCache.RefreshCache(true /*bAllowLoading*/);
+			ParameterDefinitionsAssetCache.RefreshCache(true /*bAllowLoading*/);
 		});
 	}
 
@@ -1310,8 +1310,9 @@ void FNiagaraEditorModule::OnPostEngineInit()
 		UE_LOG(LogNiagaraEditor, Log, TEXT("GEditor isn't valid! Particle reset commands will not work for Niagara components!"));
 	}
 
-	// ensure that the known parameter collections are loaded 
-	RefreshParameterCollections(true /*AllowLoading*/);
+	// ensure that all cached asset types are fully loaded.
+	ParameterCollectionAssetCache.RefreshCache(true /*bAllowLoading*/);
+	ParameterDefinitionsAssetCache.RefreshCache(true /*bAllowLoading*/);
 
 #if WITH_NIAGARA_DEBUGGER
 	Debugger = MakeShared<FNiagaraDebugger>();
@@ -1533,20 +1534,6 @@ FNiagaraClipboard& FNiagaraEditorModule::GetClipboard() const
 	return Clipboard.Get();
 }
 
-UNiagaraReservedParametersManager* FNiagaraEditorModule::GetReservedParametersManager()
-{
-	if (ReservedParametersManagerSingleton != nullptr)
-	{
-		return ReservedParametersManagerSingleton;
-	}
-
-	ReservedParametersManagerSingleton = NewObject<UNiagaraReservedParametersManager>(GetTransientPackage());
-	ReservedParametersManagerSingleton->AddToRoot();
-	ReservedParametersManagerSingleton->SetFlags(RF_Transactional);
-	ReservedParametersManagerSingleton->InitReservedParameters();
-	return ReservedParametersManagerSingleton;
-}
-
 void FNiagaraEditorModule::GetDataInterfaceFeedbackSafe(UNiagaraDataInterface* InDataInterface, TArray<FNiagaraDataInterfaceError>& OutErrors, TArray<FNiagaraDataInterfaceFeedback>& OutWarnings, TArray<FNiagaraDataInterfaceFeedback>& OutInfo)
 {
 	if (!InDataInterface)
@@ -1597,6 +1584,11 @@ void FNiagaraEditorModule::EnsureReservedDefinitionUnique(FGuid& UniqueId)
 	ReservedDefinitionIds.Add(UniqueId);
 }
 
+const TArray<TWeakObjectPtr<UNiagaraParameterDefinitions>>& FNiagaraEditorModule::GetCachedParameterDefinitionsAssets()
+{
+	return ParameterDefinitionsAssetCache.Get();
+}
+
 void FNiagaraEditorModule::GetTargetSystemAndEmitterForDataInterface(UNiagaraDataInterface* InDataInterface, UNiagaraSystem*& OutOwningSystem, UNiagaraEmitter*& OutOwningEmitter)
 {
 	OutOwningSystem = InDataInterface->GetTypedOuter<UNiagaraSystem>();
@@ -1746,30 +1738,11 @@ bool FNiagaraEditorModule::DeferredDestructObjects(float InDeltaTime)
 	return false;
 }
 
-void FNiagaraEditorModule::RefreshParameterCollections(bool AllowLoading)
-{
-	const FAssetRegistryModule& AssetRegistryModule = FModuleManager::LoadModuleChecked<FAssetRegistryModule>(TEXT("AssetRegistry"));
-	TArray<FAssetData> ParameterCollectionAssetData;
-	AssetRegistryModule.GetRegistry().GetAssetsByClass(UNiagaraParameterCollection::StaticClass()->GetFName(), ParameterCollectionAssetData);
-
-	CachedParameterCollections.Reset(ParameterCollectionAssetData.Num());
-	for (const FAssetData& ParameterCollectionAssetDatum : ParameterCollectionAssetData)
-	{
-		if (ParameterCollectionAssetDatum.IsAssetLoaded() || AllowLoading)
-		{
-			if (UNiagaraParameterCollection* Collection = Cast<UNiagaraParameterCollection>(ParameterCollectionAssetDatum.GetAsset()))
-			{
-				CachedParameterCollections.Add(MakeWeakObjectPtr(Collection));
-			}
-		}
-	}
-}
-
 UNiagaraParameterCollection* FNiagaraEditorModule::FindCollectionForVariable(const FString& VariableName)
 {
 	auto FindCachedCollectionByPrefix = [this](const FString& Prefix)
 	{
-		for (TWeakObjectPtr<UNiagaraParameterCollection>& CollectionPtr : CachedParameterCollections)
+		for (const TWeakObjectPtr<UNiagaraParameterCollection>& CollectionPtr : ParameterCollectionAssetCache.Get())
 		{
 			if (UNiagaraParameterCollection* Collection = CollectionPtr.Get())
 			{
@@ -1788,7 +1761,7 @@ UNiagaraParameterCollection* FNiagaraEditorModule::FindCollectionForVariable(con
 		return Collection;
 	}
 
-	RefreshParameterCollections(!FUObjectThreadContext::Get().IsRoutingPostLoad /*AllowLoading*/);
+	ParameterCollectionAssetCache.RefreshCache(!FUObjectThreadContext::Get().IsRoutingPostLoad /*bAllowLoading*/);
 
 	return FindCachedCollectionByPrefix(VariableName);
 }
@@ -1847,94 +1820,6 @@ void FNiagaraEditorModule::OnPerfBaselineWindowClosed(const TSharedRef<SWindow>&
 	ClosedWindow->SetContent(SNullWidget::NullWidget);
 	BaselineViewport.Reset();
 }
-
 #endif
 
-const TArray<FReservedParameter>* UNiagaraReservedParametersManager::FindReservedParametersByName(const FName ParameterName) const
-{
-	const FReservedParameterArray* ReservedParameterArrWrapperPtr = ReservedParameters.Find(ParameterName);
-	if (ReservedParameterArrWrapperPtr != nullptr)
-	{
-		return &(ReservedParameterArrWrapperPtr->Arr);
-	}
-	return nullptr;
-}
-
-int32 UNiagaraReservedParametersManager::GetNumReservedParametersByName(const FName ParameterName) const
-{
-	if (const FReservedParameterArray* FoundReservedParametersWrapperPtr = ReservedParameters.Find(ParameterName))
-	{
-		return FoundReservedParametersWrapperPtr->Arr.Num();
-	}
-	return 0;
-}
-
-EParameterDefinitionMatchState UNiagaraReservedParametersManager::GetDefinitionMatchStateForParameter(const FNiagaraVariableBase& Parameter) const
-{
-	const TArray<FReservedParameter>* ReservedParametersForNamePtr = FindReservedParametersByName(Parameter.GetName());
-	if (ReservedParametersForNamePtr == nullptr || ReservedParametersForNamePtr->Num() == 0)
-	{
-		return EParameterDefinitionMatchState::NoMatchingDefinitions;
-	}
-
-	const TArray<FReservedParameter>& ReservedParametersForName = *ReservedParametersForNamePtr;
-	if (ReservedParametersForNamePtr->Num() == 1)
-	{
-		if (ReservedParametersForName[0].GetParameter().GetType() == Parameter.GetType())
-		{
-			return EParameterDefinitionMatchState::MatchingOneDefinition;
-		}
-		return EParameterDefinitionMatchState::MatchingDefinitionNameButNotType;
-	}
-	else
-	{
-		return EParameterDefinitionMatchState::MatchingMoreThanOneDefinition;
-	}
-}
-
-void UNiagaraReservedParametersManager::AddReservedParameter(const FNiagaraVariableBase& Parameter, const UNiagaraParameterDefinitions* ReservingDefinitionsAsset)
-{
-	Modify();
-	ReservedParameters.FindOrAdd(Parameter.GetName()).Arr.AddUnique(FReservedParameter(Parameter, ReservingDefinitionsAsset));
-}
-
-void UNiagaraReservedParametersManager::RemoveReservedParameter(const FNiagaraVariableBase& Parameter, const UNiagaraParameterDefinitions* ReservingDefinitionsAsset)
-{
-	Modify();
-	FReservedParameterArray* ReservedParametersForNameWrapperPtr = ReservedParameters.Find(Parameter.GetName());
-	if (ReservedParametersForNameWrapperPtr != nullptr)
-	{
-		TArray<FReservedParameter>& ReservedParametersForName = ReservedParametersForNameWrapperPtr->Arr;
-		for (int32 Idx = ReservedParametersForName.Num() - 1; Idx > -1; --Idx)
-		{
-			const FReservedParameter& ReservedParameter = ReservedParametersForName[Idx];
-			if (ReservedParameter.GetParameter() == Parameter && ReservedParameter.GetReservingDefinitionsAsset()->GetDefinitionsUniqueId() == ReservingDefinitionsAsset->GetDefinitionsUniqueId())
-			{
-				ReservedParametersForName.RemoveAtSwap(Idx);
-			}
-		}
-	}
-}
-
-void UNiagaraReservedParametersManager::InitReservedParameters()
-{
-	TArray<FAssetData> ParameterDefinitionsAssetData;
-	FAssetRegistryModule& AssetRegistryModule = FModuleManager::LoadModuleChecked<FAssetRegistryModule>("AssetRegistry");
-	AssetRegistryModule.GetRegistry().GetAssetsByClass(UNiagaraParameterDefinitions::StaticClass()->GetFName(), ParameterDefinitionsAssetData);
-	for (const FAssetData& ParameterDefinitionsAssetDatum : ParameterDefinitionsAssetData)
-	{
-		const UNiagaraParameterDefinitions* ParameterDefinitions = Cast<UNiagaraParameterDefinitions>(ParameterDefinitionsAssetDatum.GetAsset());
-		if (ParameterDefinitions == nullptr)
-		{
-			ensureMsgf(false, TEXT("Failed to load parameter definitions when initializing reserved parameters"));
-			continue;
-		}
-
-		for (const UNiagaraScriptVariable* ScriptVar : ParameterDefinitions->GetParametersConst())
-		{
-			ReservedParameters.FindOrAdd(ScriptVar->Variable.GetName()).Arr.AddUnique(FReservedParameter(ScriptVar->Variable, ParameterDefinitions));
-		}
-	}
-}
-
 #undef LOCTEXT_NAMESPACE
diff --git a/Engine/Plugins/FX/Niagara/Source/NiagaraEditor/Private/NiagaraEditorUtilities.cpp b/Engine/Plugins/FX/Niagara/Source/NiagaraEditor/Private/NiagaraEditorUtilities.cpp
--- a/Engine/Plugins/FX/Niagara/Source/NiagaraEditor/Private/NiagaraEditorUtilities.cpp
+++ b/Engine/Plugins/FX/Niagara/Source/NiagaraEditor/Private/NiagaraEditorUtilities.cpp
@@ -3632,4 +3632,48 @@ void FNiagaraParameterUtilities::FilterToRelevantStaticVariables(const TArray<FN
 	}
 }
 
+TArray<const UNiagaraScriptVariable*> FNiagaraParameterDefinitionsUtilities::FindReservedParametersByName(const FName ParameterName)
+{
+	TArray<const UNiagaraScriptVariable*> OutScriptVars;
+	FNiagaraEditorModule& NiagaraEditorModule = FModuleManager::GetModuleChecked<FNiagaraEditorModule>("NiagaraEditor");
+	const TArray<TWeakObjectPtr<UNiagaraParameterDefinitions>>& CachedParameterDefinitionsAssets = NiagaraEditorModule.GetCachedParameterDefinitionsAssets();
+	for (const TWeakObjectPtr<UNiagaraParameterDefinitions>& CachedParameterDefinitionsAsset : CachedParameterDefinitionsAssets)
+	{
+		for (const UNiagaraScriptVariable* ScriptVar : CachedParameterDefinitionsAsset->GetParametersConst())
+		{
+			if (ScriptVar->Variable.GetName() == ParameterName)
+			{
+				OutScriptVars.Add(ScriptVar);
+			}
+		}
+	}
+	return OutScriptVars;
+}
+
+int32 FNiagaraParameterDefinitionsUtilities::GetNumParametersReservedForName(const FName ParameterName)
+{
+	return FindReservedParametersByName(ParameterName).Num();
+}
+
+EParameterDefinitionMatchState FNiagaraParameterDefinitionsUtilities::GetDefinitionMatchStateForParameter(const FNiagaraVariableBase& Parameter)
+{
+	const TArray<const UNiagaraScriptVariable*> ReservedScriptVarsForName = FindReservedParametersByName(Parameter.GetName());
+	if (ReservedScriptVarsForName.Num() == 0)
+	{
+		return EParameterDefinitionMatchState::NoMatchingDefinitions;
+	}
+	else if (ReservedScriptVarsForName.Num() == 1)
+	{
+		if (ReservedScriptVarsForName[0]->Variable.GetType() == Parameter.GetType())
+		{
+			return EParameterDefinitionMatchState::MatchingOneDefinition;
+		}
+		return EParameterDefinitionMatchState::MatchingDefinitionNameButNotType;
+	}
+	else
+	{
+		return EParameterDefinitionMatchState::MatchingMoreThanOneDefinition;
+	}
+}
+
 #undef LOCTEXT_NAMESPACE
diff --git a/Engine/Plugins/FX/Niagara/Source/NiagaraEditor/Private/NiagaraParameterDefinitions.cpp b/Engine/Plugins/FX/Niagara/Source/NiagaraEditor/Private/NiagaraParameterDefinitions.cpp
--- a/Engine/Plugins/FX/Niagara/Source/NiagaraEditor/Private/NiagaraParameterDefinitions.cpp
+++ b/Engine/Plugins/FX/Niagara/Source/NiagaraEditor/Private/NiagaraParameterDefinitions.cpp
@@ -48,7 +48,6 @@ void UNiagaraParameterDefinitions::AddParameter(const FNiagaraVariable& NewVaria
 	NewScriptVariable->Init(NewVariable, FNiagaraVariableMetaData());
 	NewScriptVariable->SetIsStaticSwitch(false);
 	NewScriptVariable->SetIsSubscribedToParameterDefinitions(true);
-	FNiagaraEditorModule::Get().GetReservedParametersManager()->AddReservedParameter(NewVariable, this);
 	NotifyParameterDefinitionsChanged();
 }
 
@@ -62,7 +61,6 @@ void UNiagaraParameterDefinitions::RemoveParameter(const FNiagaraVariable& Varia
 		// Make sure to remove any links to binding name subscriptions to external parameter libraries.
 		UnsubscribeBindingNameFromExternalParameterDefinitions(RemovedScriptVarGuid);
 		ScriptVariables.RemoveAtSwap(Idx, 1, false);
-		FNiagaraEditorModule::Get().GetReservedParametersManager()->RemoveReservedParameter(VariableToRemove, this);
 		NotifyParameterDefinitionsChanged();
 	}
 }
@@ -77,9 +75,6 @@ void UNiagaraParameterDefinitions::RenameParameter(const FNiagaraVariable& Varia
 		ScriptVariable->Modify();
 		ScriptVariable->Variable.SetName(NewName);
 		ScriptVariable->UpdateChangeId();
-		UNiagaraReservedParametersManager* ReservedParametersManager = FNiagaraEditorModule::Get().GetReservedParametersManager();
-		ReservedParametersManager->RemoveReservedParameter(OldVariable, this);
-		ReservedParametersManager->AddReservedParameter(ScriptVariable->Variable, this);
 		NotifyParameterDefinitionsChanged();
 	}
 }
diff --git a/Engine/Plugins/FX/Niagara/Source/NiagaraEditor/Private/ViewModels/NiagaraParameterPanelViewModel.cpp b/Engine/Plugins/FX/Niagara/Source/NiagaraEditor/Private/ViewModels/NiagaraParameterPanelViewModel.cpp
--- a/Engine/Plugins/FX/Niagara/Source/NiagaraEditor/Private/ViewModels/NiagaraParameterPanelViewModel.cpp
+++ b/Engine/Plugins/FX/Niagara/Source/NiagaraEditor/Private/ViewModels/NiagaraParameterPanelViewModel.cpp
@@ -702,9 +702,9 @@ void INiagaraParameterPanelViewModel::SelectParameterItemByName(const FName Para
 
 void INiagaraParameterPanelViewModel::SubscribeParameterToLibraryIfMatchingDefinition(const UNiagaraScriptVariable* ScriptVarToModify, const FName ScriptVarName) const
 {
-	const TArray<FReservedParameter>* ReservedParametersForName = FNiagaraEditorModule::Get().GetReservedParametersManager()->FindReservedParametersByName(ScriptVarName);
+	const TArray<const UNiagaraScriptVariable*> ReservedParametersForName = FNiagaraParameterDefinitionsUtilities::FindReservedParametersByName(ScriptVarName);
 
-	if (ReservedParametersForName == nullptr || ReservedParametersForName->Num() == 0)
+	if (ReservedParametersForName.Num() == 0)
 	{
 		// Do not call SetParameterIsSubscribedToLibrary() unless ScriptVarToModify is already marked as SubscribedToParameterDefinitions.
 		if (ScriptVarToModify->GetIsSubscribedToParameterDefinitions())
@@ -712,18 +712,18 @@ void INiagaraParameterPanelViewModel::SubscribeParameterToLibraryIfMatchingDefin
 			SetParameterIsSubscribedToLibrary(ScriptVarToModify, false);
 		}
 	}
-	else if (ReservedParametersForName->Num() == 1)
+	else if (ReservedParametersForName.Num() == 1)
 	{
-		const FReservedParameter& ReservedParameterDefinition = (*ReservedParametersForName)[0];
-		if (ReservedParameterDefinition.GetParameter().GetType() != ScriptVarToModify->Variable.GetType())
+		const UNiagaraScriptVariable* ReservedParameterDefinition = (ReservedParametersForName)[0];
+		if (ReservedParameterDefinition->Variable.GetType() != ScriptVarToModify->Variable.GetType())
 		{
 			const FText TypeMismatchWarningTemplate = LOCTEXT("RenameParameter_DefinitionTypeMismatch", "Renamed parameter \"{0}\" with type {1}. Type does not match existing parameter definition \"{0}\" with type {2} from {3}! ");
 			FText TypeMismatchWarning = FText::Format(
 				TypeMismatchWarningTemplate,
 				FText::FromName(ScriptVarName),
 				ScriptVarToModify->Variable.GetType().GetNameText(),
-				ReservedParameterDefinition.GetParameter().GetType().GetNameText(),
-				FText::FromString(ReservedParameterDefinition.GetReservingDefinitionsAsset()->GetName()));
+				ReservedParameterDefinition->Variable.GetType().GetNameText(),
+				FText::FromString(ReservedParameterDefinition->GetTypedOuter<UNiagaraParameterDefinitions>()->GetName()));
 			FNiagaraEditorUtilities::WarnWithToastAndLog(TypeMismatchWarning);
 
 			// Do not call SetParameterIsSubscribedToLibrary() unless ScriptVarToModify is already marked as SubscribedToParameterDefinitions.
@@ -1506,7 +1506,6 @@ TArray<FNiagaraParameterPanelItem> FNiagaraSystemToolkitParameterPanelViewModel:
 	TMap<FNiagaraVariable, TObjectPtr<UNiagaraScriptVariable>> ParameterToScriptVariableMap;
 	const TArray<UNiagaraGraph*> EditableGraphs = GetEditableGraphsConst();
 	const TArray<UNiagaraGraph*> AllGraphs = GetAllGraphsConst();
-	const UNiagaraReservedParametersManager* ReservedParametersManager = FNiagaraEditorModule::Get().GetReservedParametersManager();
 
 	// Collect all metadata to be packaged with the FNiagaraParameterPanelItems.
 	for (const UNiagaraGraph* EditableGraph : EditableGraphs)
@@ -1516,7 +1515,7 @@ TArray<FNiagaraParameterPanelItem> FNiagaraSystemToolkitParameterPanelViewModel:
 	ParameterToScriptVariableMap.Append(TransientParameterToScriptVarMap);
 
 	// Helper lambda to get all FNiagaraVariable parameters from a UNiagaraParameterCollection as FNiagaraParameterPanelItemArgs.
-	auto CollectParamStore = [this, &ParameterToScriptVariableMap, &VisitedParameterToItemMap, ReservedParametersManager](const FNiagaraParameterStore* ParamStore){
+	auto CollectParamStore = [this, &ParameterToScriptVariableMap, &VisitedParameterToItemMap](const FNiagaraParameterStore* ParamStore){
 		TArray<FNiagaraVariable> Vars;
 		ParamStore->GetParameters(Vars);
 		for (const FNiagaraVariable& Var : Vars)
@@ -1542,8 +1541,8 @@ TArray<FNiagaraParameterPanelItem> FNiagaraSystemToolkitParameterPanelViewModel:
 			Item.bSourcedFromCustomStackContext = false;
 			Item.ReferenceCount = 0;
 			
-			// Determine whether the item is name aliasing a parameter definitions's parameter.
-			Item.DefinitionMatchState = ReservedParametersManager->GetDefinitionMatchStateForParameter(ScriptVar->Variable);
+			// Determine whether the item is name aliasing a parameter definition's parameter.
+			Item.DefinitionMatchState = FNiagaraParameterDefinitionsUtilities::GetDefinitionMatchStateForParameter(ScriptVar->Variable);
 
 			VisitedParameterToItemMap.Add(Var, Item);
 		}
@@ -1565,8 +1564,8 @@ TArray<FNiagaraParameterPanelItem> FNiagaraSystemToolkitParameterPanelViewModel:
 		Item.bSourcedFromCustomStackContext = false;
 		Item.ReferenceCount = 0;
 
-		// Determine whether the item is name aliasing a parameter definitions's parameter.
-		Item.DefinitionMatchState = ReservedParametersManager->GetDefinitionMatchStateForParameter(EditorOnlyScriptVar->Variable);
+		// Determine whether the item is name aliasing a parameter definition's parameter.
+		Item.DefinitionMatchState = FNiagaraParameterDefinitionsUtilities::GetDefinitionMatchStateForParameter(EditorOnlyScriptVar->Variable);
 
 		VisitedParameterToItemMap.Add(Var, Item);
 	}
@@ -1716,8 +1715,8 @@ TArray<FNiagaraParameterPanelItem> FNiagaraSystemToolkitParameterPanelViewModel:
 						}
 					}
 
-					// Determine whether the item is name aliasing a parameter definitions's parameter.
-					Item.DefinitionMatchState = ReservedParametersManager->GetDefinitionMatchStateForParameter(Item.ScriptVariable->Variable);
+					// Determine whether the item is name aliasing a parameter definition's parameter.
+					Item.DefinitionMatchState = FNiagaraParameterDefinitionsUtilities::GetDefinitionMatchStateForParameter(Item.ScriptVariable->Variable);
 
 					// -Increment the reference count.
 					Item.ReferenceCount += Builder.Histories[0].PerVariableReadHistory[VariableIndex].Num() + Builder.Histories[0].PerVariableWriteHistory[VariableIndex].Num();
@@ -2579,7 +2578,6 @@ TArray<FNiagaraParameterPanelItem> FNiagaraScriptToolkitParameterPanelViewModel:
 
 	TMap<FNiagaraVariable, FNiagaraParameterPanelItem> VisitedParameterToItemMap;
 	TArray<FNiagaraVariable> VisitedInvalidParameters;
-	const UNiagaraReservedParametersManager* ReservedParametersManager = FNiagaraEditorModule::Get().GetReservedParametersManager();
 
 	// For scripts we use the reference maps cached in the graph to collect parameters.
 	for (const UNiagaraGraph* Graph : GetEditableGraphsConst())
@@ -2651,8 +2649,8 @@ TArray<FNiagaraParameterPanelItem> FNiagaraScriptToolkitParameterPanelViewModel:
 				Item.bExternallyReferenced = false;
 				Item.bSourcedFromCustomStackContext = false;
 
-				// Determine whether the item is name aliasing a parameter definitions's parameter.
-				Item.DefinitionMatchState = ReservedParametersManager->GetDefinitionMatchStateForParameter(ScriptVar->Variable);
+				// Determine whether the item is name aliasing a parameter definition's parameter.
+				Item.DefinitionMatchState = FNiagaraParameterDefinitionsUtilities::GetDefinitionMatchStateForParameter(ScriptVar->Variable);
 
 				// -Increment the reference count.
 				Item.ReferenceCount += ParameterElement.Value.ParameterReferences.Num();
@@ -2913,7 +2911,7 @@ bool FNiagaraParameterDefinitionsToolkitParameterPanelViewModel::GetCanRenamePar
 			OutCanRenameParameterToolTip = LOCTEXT("ParameterPanelLibraryViewModel_RenameParameter_NameAlias", "Cannot Rename Parameter: A Parameter with this name already exists in this definition asset.");
 			return false;
 		}
-		else if (FNiagaraEditorModule::Get().GetReservedParametersManager()->GetNumReservedParametersByName(NewVariableName) > 0)
+		else if (FNiagaraParameterDefinitionsUtilities::GetNumParametersReservedForName(NewVariableName) > 0)
 		{
 			OutCanRenameParameterToolTip = LOCTEXT("ParameterPanelLibraryViewModel_RenameParameter_ReservedDefinitionName", "Cannot Rename Parameter: A Parameter with this name already exists in another definition asset.");
 			return false;
diff --git a/Engine/Plugins/FX/Niagara/Source/NiagaraEditor/Public/NiagaraEditorModule.h b/Engine/Plugins/FX/Niagara/Source/NiagaraEditor/Public/NiagaraEditorModule.h
--- a/Engine/Plugins/FX/Niagara/Source/NiagaraEditor/Public/NiagaraEditorModule.h
+++ b/Engine/Plugins/FX/Niagara/Source/NiagaraEditor/Public/NiagaraEditorModule.h
@@ -3,6 +3,7 @@
 #pragma once
 
 #include "CoreMinimal.h"
+#include "AssetRegistryModule.h"
 #include "Modules/ModuleInterface.h"
 #include "Toolkits/AssetEditorToolkit.h"
 #include "NiagaraTypes.h"
@@ -198,7 +199,7 @@ public:
 	TSharedPtr<FNiagaraDebugger> GetDebugger(){ return Debugger; }
 #endif
 
-	UNiagaraReservedParametersManager* GetReservedParametersManager();
+	const TArray<TWeakObjectPtr<UNiagaraParameterDefinitions>>& GetCachedParameterDefinitionsAssets();
 
 	NIAGARAEDITOR_API void GetTargetSystemAndEmitterForDataInterface(UNiagaraDataInterface* InDataInterface, UNiagaraSystem*& OutOwningSystem, UNiagaraEmitter*& OutOwningEmitter);
 	NIAGARAEDITOR_API void GetDataInterfaceFeedbackSafe(UNiagaraDataInterface* InDataInterface, TArray<FNiagaraDataInterfaceError>& OutErrors, TArray<FNiagaraDataInterfaceFeedback>& Warnings, TArray<FNiagaraDataInterfaceFeedback>& Info);
@@ -235,6 +236,35 @@ private:
 		TSharedPtr<const T> ObjectToDestuct;
 	};
 
+	template<class AssetType>
+	class TAssetPreloadCache
+	{
+	public:
+		void RefreshCache(bool bAllowLoading)
+		{
+			const FAssetRegistryModule& AssetRegistryModule = FModuleManager::LoadModuleChecked<FAssetRegistryModule>(TEXT("AssetRegistry"));
+			TArray<FAssetData> AssetData;
+			AssetRegistryModule.GetRegistry().GetAssetsByClass(AssetType::StaticClass()->GetFName(), AssetData);
+
+			CachedAssets.Reset(AssetData.Num());
+			for (const FAssetData& AssetDatum : AssetData)
+			{
+				if (AssetDatum.IsAssetLoaded() || bAllowLoading)
+				{
+					if (AssetType* Asset = Cast<AssetType>(AssetDatum.GetAsset()))
+					{
+						CachedAssets.Add(MakeWeakObjectPtr(Asset));
+					}
+				}
+			}
+		};
+
+		const TArray<TWeakObjectPtr<AssetType>>& Get() const { return CachedAssets; };
+
+	private:
+		TArray<TWeakObjectPtr<AssetType>> CachedAssets;
+	};
+
 	void RegisterAssetTypeAction(IAssetTools& AssetTools, TSharedRef<IAssetTypeActions> Action);
 	void OnNiagaraSettingsChangedEvent(const FName& PropertyName, const UNiagaraSettings* Settings);
 	void OnPreGarbageCollection();
@@ -263,8 +293,6 @@ private:
 		StackIssueGenerators.Add(StructName) = Generator;
 	}
 
-	void RefreshParameterCollections(bool AllowLoading);
-
 private:
 	TSharedPtr<FExtensibilityManager> MenuExtensibilityManager;
 	TSharedPtr<FExtensibilityManager> ToolBarExtensibilityManager;
@@ -354,42 +382,6 @@ private:
 
 	TUniquePtr<FNiagaraGraphDataCache> GraphDataCache;
 
-	TArray<TWeakObjectPtr<UNiagaraParameterCollection>> CachedParameterCollections;
-};
-
-USTRUCT()
-struct FReservedParameterArray
-{
-	GENERATED_BODY()
-
-public:
-	UPROPERTY(Transient)
-	TArray<FReservedParameter> Arr;
-};
-
-/** Manager singleton for tracking parameters that are reserved by parameter definitions assets.
- *  Implements UObject to support undo/redo transactions on the map of definitions asset ptrs to parameter names.
- */
-UCLASS()
-class UNiagaraReservedParametersManager : public UObject
-{
-	GENERATED_BODY()
-
-public:
-	// UNiagaraReservedParametersManager is lazy initialized inside FNiagaraEditorModule::GetReservedParametersByName(), so delegate this method as a friend so that it may call InitReservedParameters().
-	friend UNiagaraReservedParametersManager* FNiagaraEditorModule::GetReservedParametersManager();
-
-public:
-	const TArray<FReservedParameter>* FindReservedParametersByName(const FName ParameterName) const;
-	int32 GetNumReservedParametersByName(const FName ParameterName) const;
-	EParameterDefinitionMatchState GetDefinitionMatchStateForParameter(const FNiagaraVariableBase& Parameter) const;
-	void AddReservedParameter(const FNiagaraVariableBase& Parameter, const UNiagaraParameterDefinitions* ReservingDefinitionsAsset);
-	void RemoveReservedParameter(const FNiagaraVariableBase& Parameter, const UNiagaraParameterDefinitions* ReservingDefinitionsAsset);
-
-private:
-	UPROPERTY(Transient)
-	TMap<FName, FReservedParameterArray> ReservedParameters;
-
-private:
-	void InitReservedParameters();
+	TAssetPreloadCache<UNiagaraParameterCollection> ParameterCollectionAssetCache;
+	TAssetPreloadCache<UNiagaraParameterDefinitions> ParameterDefinitionsAssetCache;
 };
diff --git a/Engine/Plugins/FX/Niagara/Source/NiagaraEditor/Public/NiagaraEditorUtilities.h b/Engine/Plugins/FX/Niagara/Source/NiagaraEditor/Public/NiagaraEditorUtilities.h
--- a/Engine/Plugins/FX/Niagara/Source/NiagaraEditor/Public/NiagaraEditorUtilities.h
+++ b/Engine/Plugins/FX/Niagara/Source/NiagaraEditor/Public/NiagaraEditorUtilities.h
@@ -400,3 +400,10 @@ namespace FNiagaraParameterUtilities
 
 	NIAGARAEDITOR_API void FilterToRelevantStaticVariables(const TArray<FNiagaraVariable>& InVars, TArray<FNiagaraVariable>& OutVars, FName InOldEmitterAlias, FName InNewEmitterAlias, bool bFilterByEmitterAliasAndConvertToUnaliased);
 };
+
+namespace FNiagaraParameterDefinitionsUtilities
+{
+	TArray<const UNiagaraScriptVariable*>  FindReservedParametersByName(const FName ParameterName);
+	int32 GetNumParametersReservedForName(const FName ParameterName);
+	EParameterDefinitionMatchState GetDefinitionMatchStateForParameter(const FNiagaraVariableBase& Parameter);
+};
