Automatically choose a location in the stack that satisfies all dependencies when adding a module script.

diff --git a/Engine/Plugins/FX/Niagara/Source/NiagaraEditor/Private/ViewModels/Stack/NiagaraStackGraphUtilities.cpp b/Engine/Plugins/FX/Niagara/Source/NiagaraEditor/Private/ViewModels/Stack/NiagaraStackGraphUtilities.cpp
--- a/Engine/Plugins/FX/Niagara/Source/NiagaraEditor/Private/ViewModels/Stack/NiagaraStackGraphUtilities.cpp
+++ b/Engine/Plugins/FX/Niagara/Source/NiagaraEditor/Private/ViewModels/Stack/NiagaraStackGraphUtilities.cpp
@@ -13,8 +13,10 @@
 #include "NiagaraNodeParameterMapGet.h"
 #include "NiagaraNodeParameterMapSet.h"
 #include "NiagaraNodeEmitter.h"
+#include "NiagaraScript.h"
 #include "NiagaraScriptSource.h"
 #include "EdGraphSchema_Niagara.h"
+#include "ViewModels/NiagaraEmitterHandleViewModel.h"
 #include "ViewModels/NiagaraScriptViewModel.h"
 #include "ViewModels/Stack/NiagaraStackEntry.h"
 #include "ViewModels/Stack/NiagaraStackFunctionInputCollection.h"
@@ -1547,6 +1549,93 @@ bool FNiagaraStackGraphUtilities::FindScriptModulesInStack(FAssetData ModuleScri
 	return OutFunctionCalls.Num() > 0;
 }
 
+UNiagaraNodeFunctionCall* FNiagaraStackGraphUtilities::AddScriptModuleToStack(const FAddScriptModuleToStackArgs& Args)
+{
+	// Unpack args struct
+	const FAssetData& ModuleScriptAsset = Args.ModuleScriptAsset;
+	UNiagaraScript* const ModuleScript = Args.ModuleScript;
+	UNiagaraNodeOutput* const TargetOutputNode = Args.TargetOutputNode;
+	const int32 TargetIndex = Args.TargetIndex;
+	FString SuggestedName = Args.SuggestedName;
+	const bool bFixupTargetIndex = Args.bFixupTargetIndex;
+	const FGuid& VersionGuid = Args.VersionGuid;
+
+	// Get the stack graph via the TargetOutputNode.
+	UEdGraph* Graph = TargetOutputNode->GetGraph();
+	Graph->Modify();
+
+	// Create the UNiagaraNodeFunctionCall to represent the script in the stack.
+	FGraphNodeCreator<UNiagaraNodeFunctionCall> ModuleNodeCreator(*Graph);
+	UNiagaraNodeFunctionCall* NewModuleNode = ModuleNodeCreator.CreateNode();
+
+	// Set the script on the UNiagaraNodeFunctionCall, and set the script version if specified.
+	if (ModuleScript != nullptr)
+	{
+		NewModuleNode->FunctionScript = ModuleScript;
+	}
+	else if(ModuleScriptAsset.IsValid())
+	{
+		NewModuleNode->FunctionScript = CastChecked<UNiagaraScript>(ModuleScriptAsset.GetAsset());
+	}
+	else
+	{
+		ensureMsgf(false, TEXT("Encountered invalid FAddScriptModuleToStackArgs! ModuleScript or ModuleScriptAsset must be valid!"));
+		return nullptr;
+	}
+
+	if (NewModuleNode->FunctionScript->IsVersioningEnabled())
+	{
+		NewModuleNode->SelectedScriptVersion = VersionGuid.IsValid() ? VersionGuid : NewModuleNode->FunctionScript->GetExposedVersion().VersionGuid;
+	}
+	else
+	{
+		NewModuleNode->SelectedScriptVersion = FGuid();
+	}
+	ModuleNodeCreator.Finalize();
+
+	// Ensure there are input and output pins on the UNiagaraNodeFunctionCall to wire into the stack graph.
+	const UEdGraphSchema_Niagara* NiagaraSchema = GetDefault<UEdGraphSchema_Niagara>();
+	if (NewModuleNode->FunctionScript == nullptr)
+	{
+		// If the module script is null, add parameter map inputs and outputs so that the node can be wired into the graph correctly.
+		NewModuleNode->CreatePin(EGPD_Input, NiagaraSchema->TypeDefinitionToPinType(FNiagaraTypeDefinition::GetParameterMapDef()), TEXT("InputMap"));
+		NewModuleNode->CreatePin(EGPD_Output, NiagaraSchema->TypeDefinitionToPinType(FNiagaraTypeDefinition::GetParameterMapDef()), TEXT("OutputMap"));
+		if (SuggestedName.IsEmpty())
+		{
+			SuggestedName = TEXT("InvalidScript");
+		}
+	}
+	else
+	{
+		// Make sure that the input and output pins are available to prevent failures in the connect module node.  Once the node is
+		// connected these missing pins will generate compile errors which the user can find and fix.
+		if (FNiagaraStackGraphUtilities::GetParameterMapInputPin(*NewModuleNode) == nullptr)
+		{
+			NewModuleNode->CreatePin(EGPD_Input, NiagaraSchema->TypeDefinitionToPinType(FNiagaraTypeDefinition::GetParameterMapDef()), TEXT("InputMap"));
+		}
+		if (FNiagaraStackGraphUtilities::GetParameterMapOutputPin(*NewModuleNode) == nullptr)
+		{
+			NewModuleNode->CreatePin(EGPD_Output, NiagaraSchema->TypeDefinitionToPinType(FNiagaraTypeDefinition::GetParameterMapDef()), TEXT("OutputMap"));
+		}
+	}
+
+	// If specified, suggest the name for the script.
+	if (SuggestedName.IsEmpty() == false)
+	{
+		NewModuleNode->SuggestName(SuggestedName);
+	}
+
+	// If specified, find the nearest index to TargetIndex that satisfies the new module script's order dependencies.
+	int32 FinalTargetIndex = TargetIndex;
+	if (bFixupTargetIndex)
+	{
+		FinalTargetIndex = DependencyUtilities::FindBestIndexForModuleInStack(*NewModuleNode, *Graph);
+	}
+
+	ConnectModuleNode(*NewModuleNode, *TargetOutputNode, FinalTargetIndex);
+	return NewModuleNode;
+}
+
 UNiagaraNodeFunctionCall* FNiagaraStackGraphUtilities::AddScriptModuleToStack(FAssetData ModuleScriptAsset, UNiagaraNodeOutput& TargetOutputNode, int32 TargetIndex, FString SuggestedName)
 {
 	UEdGraph* Graph = TargetOutputNode.GetGraph();
@@ -1911,35 +2000,6 @@ void FNiagaraStackGraphUtilities::GetNewParameterAvailableTypes(TArray<FNiagaraT
 	}
 }
 
-void FNiagaraStackGraphUtilities::GetModuleScriptAssetsByDependencyProvided(FName DependencyName, TOptional<ENiagaraScriptUsage> RequiredUsage, TArray<FAssetData>& OutAssets)
-{
-	FNiagaraEditorUtilities::FGetFilteredScriptAssetsOptions ScriptFilterOptions;
-	ScriptFilterOptions.bIncludeDeprecatedScripts = false;
-	ScriptFilterOptions.bIncludeNonLibraryScripts = true;
-	ScriptFilterOptions.ScriptUsageToInclude = ENiagaraScriptUsage::Module;
-	ScriptFilterOptions.TargetUsageToMatch = RequiredUsage;
-	TArray<FAssetData> ModuleAssets;
-	FNiagaraEditorUtilities::GetFilteredScriptAssets(ScriptFilterOptions, ModuleAssets);
-
-	for (const FAssetData& ModuleAsset : ModuleAssets)
-	{
-		FString ProvidedDependenciesString;
-		if (ModuleAsset.GetTagValue(GET_MEMBER_NAME_CHECKED(FVersionedNiagaraScriptData, ProvidedDependencies), ProvidedDependenciesString) && ProvidedDependenciesString.IsEmpty() == false)
-		{
-			TArray<FString> DependencyStrings;
-			ProvidedDependenciesString.ParseIntoArray(DependencyStrings, TEXT(","));
-			for (FString DependencyString : DependencyStrings)
-			{
-				if (FName(*DependencyString) == DependencyName)
-				{
-					OutAssets.Add(ModuleAsset);
-					break;
-				}
-			}
-		}
-	}
-}
-
 TOptional<FName> FNiagaraStackGraphUtilities::GetNamespaceForOutputNode(const UNiagaraNodeOutput* OutputNode)
 {
 	if (OutputNode)
@@ -3154,4 +3214,215 @@ void FNiagaraStackGraphUtilities::SynchronizeVariableToLibraryAndApplyToGraph(UN
 #endif
 }
 
+bool FNiagaraStackGraphUtilities::DependencyUtilities::DoesStackModuleProvideDependency(const FNiagaraStackModuleData& StackModuleData, const FNiagaraModuleDependency& SourceModuleRequiredDependency, const UNiagaraNodeOutput& SourceOutputNode)
+{
+	if (StackModuleData.ModuleNode != nullptr && StackModuleData.ModuleNode->FunctionScript != nullptr)
+	{
+		FVersionedNiagaraScriptData* ScriptData = StackModuleData.ModuleNode->FunctionScript->GetScriptData(StackModuleData.ModuleNode->SelectedScriptVersion);
+		if (ScriptData && ScriptData->ProvidedDependencies.Contains(SourceModuleRequiredDependency.Id))
+		{
+			if (SourceModuleRequiredDependency.ScriptConstraint == ENiagaraModuleDependencyScriptConstraint::AllScripts)
+			{
+				return true;
+			}
+
+			if (SourceModuleRequiredDependency.ScriptConstraint == ENiagaraModuleDependencyScriptConstraint::SameScript)
+			{
+				UNiagaraNodeOutput* OutputNode = FNiagaraStackGraphUtilities::GetEmitterOutputNodeForStackNode(*StackModuleData.ModuleNode);
+				return OutputNode != nullptr && UNiagaraScript::IsEquivalentUsage(OutputNode->GetUsage(), SourceOutputNode.GetUsage()) && OutputNode->GetUsageId() == SourceOutputNode.GetUsageId();
+			}
+		}
+	}
+	return false;
+}
+
+void FNiagaraStackGraphUtilities::DependencyUtilities::GetModuleScriptAssetsByDependencyProvided(FName DependencyName, TOptional<ENiagaraScriptUsage> RequiredUsage, TArray<FAssetData>& OutAssets)
+{
+	FNiagaraEditorUtilities::FGetFilteredScriptAssetsOptions ScriptFilterOptions;
+	ScriptFilterOptions.bIncludeDeprecatedScripts = false;
+	ScriptFilterOptions.bIncludeNonLibraryScripts = true;
+	ScriptFilterOptions.ScriptUsageToInclude = ENiagaraScriptUsage::Module;
+	ScriptFilterOptions.TargetUsageToMatch = RequiredUsage;
+	TArray<FAssetData> ModuleAssets;
+	FNiagaraEditorUtilities::GetFilteredScriptAssets(ScriptFilterOptions, ModuleAssets);
+
+	for (const FAssetData& ModuleAsset : ModuleAssets)
+	{
+		FString ProvidedDependenciesString;
+		if (ModuleAsset.GetTagValue(GET_MEMBER_NAME_CHECKED(FVersionedNiagaraScriptData, ProvidedDependencies), ProvidedDependenciesString) && ProvidedDependenciesString.IsEmpty() == false)
+		{
+			TArray<FString> DependencyStrings;
+			ProvidedDependenciesString.ParseIntoArray(DependencyStrings, TEXT(","));
+			for (FString DependencyString : DependencyStrings)
+			{
+				if (FName(*DependencyString) == DependencyName)
+				{
+					OutAssets.Add(ModuleAsset);
+					break;
+				}
+			}
+		}
+	}
+}
+
+int32 FNiagaraStackGraphUtilities::DependencyUtilities::FindBestIndexForModuleInStack(UNiagaraNodeFunctionCall& ModuleNode, UEdGraph& EmitterScriptGraph)
+{
+	// Check if the new module node has any dependencies to begin with. If not, early exit.
+	FVersionedNiagaraScriptData* ScriptData = ModuleNode.GetScriptData();
+	if (ScriptData == nullptr || (ScriptData->RequiredDependencies.Num() == 0 && ScriptData->ProvidedDependencies.Num() == 0))
+	{
+		return INDEX_NONE;
+	}
+
+	// Get the Emitter and System the emitter script script graph is outered to.
+	UNiagaraSystem* System = EmitterScriptGraph.GetTypedOuter<UNiagaraSystem>();
+	UNiagaraEmitter* Emitter = EmitterScriptGraph.GetTypedOuter<UNiagaraEmitter>();
+	if (System == nullptr || Emitter == nullptr)
+	{
+		return INDEX_NONE;
+	}
+
+	// Get the stack module data for the emitter stack to find dependencies.
+	TSharedPtr<FNiagaraSystemViewModel> SystemViewModel = TNiagaraViewModelManager<UNiagaraSystem, FNiagaraSystemViewModel>::GetExistingViewModelForObject(System);
+	if (SystemViewModel->IsValid() == false)
+	{
+		ensureMsgf(false, TEXT("Failed to get systemviewmodel for valid system when getting best index in stack for module!"));
+		return INDEX_NONE;
+	}
+	TSharedPtr<FNiagaraEmitterHandleViewModel> EmitterHandleViewModel = SystemViewModel->GetEmitterHandleViewModelForEmitter(Emitter);
+	if (EmitterHandleViewModel->IsValid() == false)
+	{
+		ensureMsgf(false, TEXT("Failed to get emitterhandleviewmodel for valid emitter when getting best index in stack for module!"));
+		return INDEX_NONE;
+	}
+	const TArray<FNiagaraStackModuleData>& StackModuleData = SystemViewModel->GetStackModuleDataByEmitterHandleId(EmitterHandleViewModel->GetId());
+
+
+	// Find the greatest and least indices for the stack first.
+	int32 LeastIndex = INT_MAX;
+	int32 GreatestIndex = INDEX_NONE;
+	for (const FNiagaraStackModuleData& CurrentStackModuleData : StackModuleData)
+	{
+		int32 Index = CurrentStackModuleData.Index;
+		LeastIndex = LeastIndex > Index ? Index : LeastIndex;
+		GreatestIndex = GreatestIndex < Index ? Index : GreatestIndex;
+	}
+
+	// Find the greatest and least indices to satisfy the pre and post dependencies of the new module script being added, respectively.
+	int32 LeastDependencyIndex = LeastIndex;
+	int32 GreatestDependencyIndex = GreatestIndex;
+	bool bDependencyLowerBoundFound = false;
+	bool bDependencyUpperBoundFound = false;
+
+	const TArray<FNiagaraModuleDependency>& NewModuleScriptRequiredDependencies = ScriptData->RequiredDependencies;
+
+	TMap<ENiagaraScriptUsage, UNiagaraNodeOutput*> ScriptUsageToOutputNode;
+	auto GetOutputNodeForStackModuleData = [&ScriptUsageToOutputNode](const FNiagaraStackModuleData& StackModuleData)->UNiagaraNodeOutput* {
+		if (UNiagaraNodeOutput** OutputNodePtr = ScriptUsageToOutputNode.Find(StackModuleData.Usage))
+		{
+			return *OutputNodePtr;
+		}
+		UNiagaraNodeOutput* OutputNode = FNiagaraStackGraphUtilities::GetEmitterOutputNodeForStackNode(*StackModuleData.ModuleNode);
+		ScriptUsageToOutputNode.Add(StackModuleData.Usage, OutputNode);
+		return OutputNode;
+	};
+
+	for (const FNiagaraStackModuleData& CurrentStackModuleData : StackModuleData)
+	{
+		for (const FNiagaraModuleDependency& RequiredDependency : NewModuleScriptRequiredDependencies)
+		{
+			if (DoesStackModuleProvideDependency(CurrentStackModuleData, RequiredDependency, *GetOutputNodeForStackModuleData(CurrentStackModuleData)))
+			{
+				if (RequiredDependency.Type == ENiagaraModuleDependencyType::PreDependency)
+				{
+					int32 SuggestedIndex = CurrentStackModuleData.Index + 1;
+					LeastDependencyIndex = LeastDependencyIndex < SuggestedIndex ? SuggestedIndex : LeastDependencyIndex;
+					bDependencyLowerBoundFound = true;
+				}
+				else if (RequiredDependency.Type == ENiagaraModuleDependencyType::PostDependency)
+				{
+					int32 SuggestedIndex = CurrentStackModuleData.Index;
+					GreatestDependencyIndex = GreatestDependencyIndex > SuggestedIndex ? SuggestedIndex : GreatestDependencyIndex;
+					bDependencyUpperBoundFound = true;
+				}
+			}
+		}
+	}
+
+	// Check that there is valid index which satisfies all dependencies for the new module script.
+	int32 TargetIndex = INDEX_NONE;
+	if (bDependencyLowerBoundFound)
+	{
+		if (bDependencyUpperBoundFound && GreatestDependencyIndex < LeastDependencyIndex)
+		{
+			// It is impossible to satisfy both target indices: do nothing.
+			return INDEX_NONE;
+		}
+		TargetIndex = LeastDependencyIndex;
+	}
+	else if (bDependencyUpperBoundFound)
+	{
+		TargetIndex = GreatestDependencyIndex;
+	}
+
+	// Do another pass to find a target index that satisfies dependencies for other modules in the stack by providing as many dependencies as possible.
+	int32 LeastRequirementIndex = LeastIndex;
+	int32 GreatestRequirementIndex = GreatestIndex;
+	bool bRequirementLowerBoundFound = false;
+	bool bRequirementUpperBoundFound = false;
+
+	const TArray<FName>& NewModuleScriptProvidedDependencies = ScriptData->ProvidedDependencies;
+
+	auto GetStackModuleDataRequiredDependenciesBeingProvided = [&NewModuleScriptProvidedDependencies](const FNiagaraStackModuleData& StackModuleData)->const TArray<FNiagaraModuleDependency> /*OutDependencies*/ {
+		TArray<FNiagaraModuleDependency> OutDependencies;
+		for (FNiagaraModuleDependency& RequiredDependency : StackModuleData.ModuleNode->GetScriptData()->RequiredDependencies)
+		{
+			if (NewModuleScriptProvidedDependencies.Contains(RequiredDependency.Id))
+			{
+				OutDependencies.Add(RequiredDependency);
+			}
+		}
+		return OutDependencies;
+	};
+
+	for (const FNiagaraStackModuleData& CurrentStackModuleData : StackModuleData)
+	{
+		for (const FNiagaraModuleDependency& ProvidedDependency : GetStackModuleDataRequiredDependenciesBeingProvided(CurrentStackModuleData))
+		{
+			if (ProvidedDependency.Type == ENiagaraModuleDependencyType::PreDependency)
+			{
+				int32 SuggestedIndex = CurrentStackModuleData.Index;
+				LeastRequirementIndex = LeastRequirementIndex < SuggestedIndex ? SuggestedIndex : LeastRequirementIndex;
+				bRequirementLowerBoundFound = true;
+			}
+			else if (ProvidedDependency.Type == ENiagaraModuleDependencyType::PostDependency)
+			{
+				int32 SuggestedIndex = CurrentStackModuleData.Index + 1;
+				GreatestRequirementIndex = GreatestRequirementIndex > SuggestedIndex ? SuggestedIndex : GreatestRequirementIndex;
+				bRequirementUpperBoundFound = true;
+			}
+		}
+	}
+
+	LeastRequirementIndex = LeastDependencyIndex < LeastRequirementIndex ? LeastRequirementIndex : LeastDependencyIndex;
+	GreatestRequirementIndex = GreatestDependencyIndex > GreatestRequirementIndex ? GreatestRequirementIndex : GreatestDependencyIndex;
+
+	// Do another validity test for requirement indices.
+	if (bRequirementLowerBoundFound)
+	{
+		if (bRequirementUpperBoundFound && GreatestRequirementIndex < LeastRequirementIndex)
+		{
+			// It is impossible to satisfy both target indices: do nothing.
+			return TargetIndex;
+		}
+		TargetIndex = LeastRequirementIndex;
+	}
+	else if (bRequirementUpperBoundFound)
+	{
+		TargetIndex = GreatestRequirementIndex;
+	}
+
+	return TargetIndex;
+}
+
 #undef LOCTEXT_NAMESPACE
diff --git a/Engine/Plugins/FX/Niagara/Source/NiagaraEditor/Private/ViewModels/Stack/NiagaraStackModuleItem.cpp b/Engine/Plugins/FX/Niagara/Source/NiagaraEditor/Private/ViewModels/Stack/NiagaraStackModuleItem.cpp
--- a/Engine/Plugins/FX/Niagara/Source/NiagaraEditor/Private/ViewModels/Stack/NiagaraStackModuleItem.cpp
+++ b/Engine/Plugins/FX/Niagara/Source/NiagaraEditor/Private/ViewModels/Stack/NiagaraStackModuleItem.cpp
@@ -426,7 +426,7 @@ void GenerateFixesForAddingDependencyProviders(
 		? DependentUsage
 		: TOptional<ENiagaraScriptUsage>();
 	TArray<FAssetData> ModuleAssetsForDependency;
-	FNiagaraStackGraphUtilities::GetModuleScriptAssetsByDependencyProvided(RequiredDependency.Id, RequiredUsage, ModuleAssetsForDependency);
+	FNiagaraStackGraphUtilities::DependencyUtilities::GetModuleScriptAssetsByDependencyProvided(RequiredDependency.Id, RequiredUsage, ModuleAssetsForDependency);
 
 	// Gather duplicate module names so their fixes can be disambiguated.
 	TSet<FName> ModuleNames;
diff --git a/Engine/Plugins/FX/Niagara/Source/NiagaraEditor/Private/ViewModels/Stack/NiagaraStackScriptItemGroup.cpp b/Engine/Plugins/FX/Niagara/Source/NiagaraEditor/Private/ViewModels/Stack/NiagaraStackScriptItemGroup.cpp
--- a/Engine/Plugins/FX/Niagara/Source/NiagaraEditor/Private/ViewModels/Stack/NiagaraStackScriptItemGroup.cpp
+++ b/Engine/Plugins/FX/Niagara/Source/NiagaraEditor/Private/ViewModels/Stack/NiagaraStackScriptItemGroup.cpp
@@ -288,7 +288,10 @@ public:
 		UNiagaraNodeFunctionCall* NewModuleNode = nullptr;
 		if (ScriptGroupAddAction->GetModuleAssetData().IsValid())
 		{
-			NewModuleNode = FNiagaraStackGraphUtilities::AddScriptModuleToStack(ScriptGroupAddAction->GetModuleAssetData(), *OutputNode, TargetIndex);
+			FNiagaraStackGraphUtilities::FAddScriptModuleToStackArgs AddScriptModuleToStackArgs(ScriptGroupAddAction->GetModuleAssetData(), *OutputNode);
+			AddScriptModuleToStackArgs.TargetIndex = TargetIndex;
+			AddScriptModuleToStackArgs.bFixupTargetIndex = true;
+			NewModuleNode = FNiagaraStackGraphUtilities::AddScriptModuleToStack(AddScriptModuleToStackArgs);
 		}
 		else if (ScriptGroupAddAction->GetModuleParameterVariable().IsValid())
 		{
@@ -323,7 +326,10 @@ public:
 private:
 	UNiagaraNodeFunctionCall* AddScriptAssetModule(const FAssetData& AssetData, int32 TargetIndex)
 	{
-		return FNiagaraStackGraphUtilities::AddScriptModuleToStack(AssetData, *OutputNode, TargetIndex);
+		FNiagaraStackGraphUtilities::FAddScriptModuleToStackArgs AddScriptModuleToStackArgs(AssetData, *OutputNode);
+		AddScriptModuleToStackArgs.TargetIndex = TargetIndex;
+		AddScriptModuleToStackArgs.bFixupTargetIndex = true;
+		return FNiagaraStackGraphUtilities::AddScriptModuleToStack(AddScriptModuleToStackArgs);
 	}
 
 	UNiagaraNodeFunctionCall* AddParameterModule(const FNiagaraVariable& ParameterVariable, bool bRenameParameterOnAdd, int32 TargetIndex)
@@ -1302,7 +1308,12 @@ void UNiagaraStackScriptItemGroup::PasteModules(const UNiagaraClipboardContent*
 						// Otherwise it's a scratch pad script from another asset so we need to add a duplicate scratch pad script to this asset.
 						NewFunctionScript = GetSystemViewModel()->GetScriptScratchPadViewModel()->CreateNewScriptAsDuplicate(ClipboardFunctionScript)->GetOriginalScript();
 					}
-					NewFunctionCallNode = FNiagaraStackGraphUtilities::AddScriptModuleToStack(NewFunctionScript, *OutputNode, CurrentPasteIndex, ClipboardFunction->FunctionName, ClipboardFunction->ScriptVersion);
+					FNiagaraStackGraphUtilities::FAddScriptModuleToStackArgs AddScriptModuleToStackArgs(NewFunctionScript, *OutputNode);
+					AddScriptModuleToStackArgs.TargetIndex = CurrentPasteIndex;
+					AddScriptModuleToStackArgs.bFixupTargetIndex = true;
+					AddScriptModuleToStackArgs.SuggestedName = ClipboardFunction->FunctionName;
+					AddScriptModuleToStackArgs.VersionGuid = ClipboardFunction->ScriptVersion;
+					NewFunctionCallNode = FNiagaraStackGraphUtilities::AddScriptModuleToStack(AddScriptModuleToStackArgs);
 				}
 				else
 				{
diff --git a/Engine/Plugins/FX/Niagara/Source/NiagaraEditor/Public/ViewModels/Stack/NiagaraStackGraphUtilities.h b/Engine/Plugins/FX/Niagara/Source/NiagaraEditor/Public/ViewModels/Stack/NiagaraStackGraphUtilities.h
--- a/Engine/Plugins/FX/Niagara/Source/NiagaraEditor/Public/ViewModels/Stack/NiagaraStackGraphUtilities.h
+++ b/Engine/Plugins/FX/Niagara/Source/NiagaraEditor/Public/ViewModels/Stack/NiagaraStackGraphUtilities.h
@@ -29,6 +29,8 @@ class UNiagaraStackErrorItem;
 class FCompileConstantResolver;
 class INiagaraMessage;
 class FNiagaraParameterMapHistoryBuilder;
+struct FNiagaraStackModuleData;
+struct FNiagaraModuleDependency;
 
 namespace FNiagaraStackGraphUtilities
 {
@@ -144,6 +146,47 @@ namespace FNiagaraStackGraphUtilities
 
 	bool RemoveModuleFromStack(UNiagaraScript& OwningScript, UNiagaraNodeFunctionCall& ModuleNode, TArray<TWeakObjectPtr<UNiagaraNodeInput>>& OutRemovedInputNodes);
 
+	struct FAddScriptModuleToStackArgs
+	{
+	public:
+		FAddScriptModuleToStackArgs(FAssetData InModuleScriptAsset, UNiagaraNodeOutput& InTargetOutputNode)
+			: ModuleScriptAsset(InModuleScriptAsset)
+			, TargetOutputNode(&InTargetOutputNode)
+		{};
+
+		FAddScriptModuleToStackArgs(UNiagaraScript* InModuleScript, UNiagaraNodeOutput& InTargetOutputNode)
+			: ModuleScript(InModuleScript)
+			, TargetOutputNode(&InTargetOutputNode)
+		{};
+
+	public:
+		//~ Begin required args
+		// ModuleScriptAsset OR ModuleScript is required. The script to add to the stack.
+		const FAssetData ModuleScriptAsset = FAssetData();
+		UNiagaraScript* const ModuleScript = nullptr;
+
+		// The output node in the stack graph to add to.
+		UNiagaraNodeOutput* const TargetOutputNode = nullptr;
+		//~End required args
+
+		// The index in the stack list in which to instert the script.
+		int32 TargetIndex = INDEX_NONE;
+
+		// The suggested name for the new script.
+		FString SuggestedName = FString();
+
+		// Whether or not to modify the final TargetIndex to be the closest value to TargetIndex which satisfies all script dependencies.
+		bool bFixupTargetIndex = false;
+
+		// The version to use for the script to add.
+		FGuid VersionGuid = FGuid();
+
+	private:
+		FAddScriptModuleToStackArgs() = default;
+	};
+
+	UNiagaraNodeFunctionCall* AddScriptModuleToStack(const FAddScriptModuleToStackArgs& Args);
+
 	UNiagaraNodeFunctionCall* AddScriptModuleToStack(FAssetData ModuleScriptAsset, UNiagaraNodeOutput& TargetOutputNode, int32 TargetIndex = INDEX_NONE, FString SuggestedName = FString());
 
 	UNiagaraNodeFunctionCall* AddScriptModuleToStack(UNiagaraScript* ModuleScript, UNiagaraNodeOutput& TargetOutputNode, int32 TargetIndex = INDEX_NONE, FString SuggestedName = FString(), const FGuid& VersionGuid = FGuid());
@@ -170,8 +213,6 @@ namespace FNiagaraStackGraphUtilities
 
 	void GetNewParameterAvailableTypes(TArray<FNiagaraTypeDefinition>& OutAvailableTypes, FName Namespace);
 
-	void GetModuleScriptAssetsByDependencyProvided(FName DependencyName, TOptional<ENiagaraScriptUsage> RequiredUsage, TArray<FAssetData>& OutAssets);
-
 	TOptional<FName> GetNamespaceForScriptUsage(ENiagaraScriptUsage ScriptUsage);
 	TOptional<FName> GetNamespaceForOutputNode(const UNiagaraNodeOutput* OutputNode);
 	
@@ -243,4 +284,13 @@ namespace FNiagaraStackGraphUtilities
 	void AddNewVariableToParameterMapNode(UNiagaraNodeParameterMapBase* MapBaseNode, bool bCreateInputPin, const UNiagaraScriptVariable* NewScriptVar);
 
 	void SynchronizeVariableToLibraryAndApplyToGraph(UNiagaraScriptVariable* ScriptVarToSync);
+
+	namespace DependencyUtilities
+	{
+		bool DoesStackModuleProvideDependency(const FNiagaraStackModuleData& StackModuleData, const FNiagaraModuleDependency& SourceModuleRequiredDependency, const UNiagaraNodeOutput& SourceOutputNode);
+
+		void GetModuleScriptAssetsByDependencyProvided(FName DependencyName, TOptional<ENiagaraScriptUsage> RequiredUsage, TArray<FAssetData>& OutAssets);
+
+		int32 FindBestIndexForModuleInStack(UNiagaraNodeFunctionCall& ModuleNode, UEdGraph& EmitterScriptGraph);
+	}
 }
