- Simulation stage particle iteration over a sub set of particles
- Further work is required to make this more optimal, we need to understand usage cases a bit more before premature optimization

diff --git a/Engine/Plugins/FX/Niagara/Shaders/Private/NiagaraEmitterInstanceShader.usf b/Engine/Plugins/FX/Niagara/Shaders/Private/NiagaraEmitterInstanceShader.usf
--- a/Engine/Plugins/FX/Niagara/Shaders/Private/NiagaraEmitterInstanceShader.usf
+++ b/Engine/Plugins/FX/Niagara/Shaders/Private/NiagaraEmitterInstanceShader.usf
@@ -497,10 +497,6 @@ int rand_int(int x, int Seed1, int Seed2, int Seed3)
 #if GPU_SIMULATION
 	uint ComponentBufferSizeRead;
 	uint ComponentBufferSizeWrite;
-	int	NumEventsPerParticle;
-	int NumParticlesPerEvent;
-	uint CopyInstancesBeforeStart;
-
 	uint SimStart;
 
 	/* Buffers for particle data and DrawIndirect calls
@@ -536,6 +532,9 @@ int rand_int(int x, int Seed1, int Seed2, int Seed3)
 	int SimulationStage_GetNumIterations() { return SimulationStageIterationInfo.w; }
 	float SimulationStage_GetNormalizedIterationIndex() { return SimulationStageNormalizedIterationIndex; }
 
+	// Where X = Parameter Binding, YZ = Inclusive Range
+	uint3 ParticleIterationStateInfo;
+
 	void AcquireID(int DataSetID, out int IDIndex, out int IDAcquireTag)
 	{
 		// Begin static assert : GPU particles only support DataSetID 0
diff --git a/Engine/Plugins/FX/Niagara/Source/Niagara/Classes/NiagaraDataSet.h b/Engine/Plugins/FX/Niagara/Source/Niagara/Classes/NiagaraDataSet.h
--- a/Engine/Plugins/FX/Niagara/Source/Niagara/Classes/NiagaraDataSet.h
+++ b/Engine/Plugins/FX/Niagara/Source/Niagara/Classes/NiagaraDataSet.h
@@ -299,6 +299,7 @@ struct NIAGARA_API FNiagaraDataSetCompiledData
 	ENiagaraSimTarget SimTarget;
 
 	FNiagaraDataSetCompiledData();
+	const FNiagaraVariableLayoutInfo* FindVariableLayoutInfo(const FNiagaraVariableBase& VariableDef) const;
 	void BuildLayout();
 	void Empty();
 
diff --git a/Engine/Plugins/FX/Niagara/Source/Niagara/Private/NiagaraDataSet.cpp b/Engine/Plugins/FX/Niagara/Source/Niagara/Private/NiagaraDataSet.cpp
--- a/Engine/Plugins/FX/Niagara/Source/Niagara/Private/NiagaraDataSet.cpp
+++ b/Engine/Plugins/FX/Niagara/Source/Niagara/Private/NiagaraDataSet.cpp
@@ -628,9 +628,6 @@ bool FNiagaraDataBuffer::CheckForNaNs()const
 
 void FNiagaraDataBuffer::Allocate(uint32 InNumInstances, bool bMaintainExisting)
 {
-	//CheckUsage(false);
-	checkSlow(Owner->GetSimTarget() == ENiagaraSimTarget::CPUSim);
-
 	NumInstances = 0;
 	if (!bMaintainExisting)
 	{
@@ -1453,6 +1450,18 @@ void FNiagaraDataBuffer::BuildRegisterTable()
 
 FNiagaraDataSetCompiledData FNiagaraDataSetCompiledData::DummyCompiledData;
 
+const FNiagaraVariableLayoutInfo* FNiagaraDataSetCompiledData::FindVariableLayoutInfo(const FNiagaraVariableBase& VariableDef) const
+{
+	for ( int i=0; i < Variables.Num(); ++i )
+	{
+		if ( Variables[i].GetName() == VariableDef.GetName() && Variables[i].GetType() == VariableDef.GetType() )
+		{
+			return &VariableLayouts[i];
+		}
+	}
+	return nullptr;
+}
+
 // note, this method is also implemented in FNiagaraDataSet::BuildLayout()
 void FNiagaraDataSetCompiledData::BuildLayout()
 {
diff --git a/Engine/Plugins/FX/Niagara/Source/Niagara/Private/NiagaraEmitter.cpp b/Engine/Plugins/FX/Niagara/Source/Niagara/Private/NiagaraEmitter.cpp
--- a/Engine/Plugins/FX/Niagara/Source/Niagara/Private/NiagaraEmitter.cpp
+++ b/Engine/Plugins/FX/Niagara/Source/Niagara/Private/NiagaraEmitter.cpp
@@ -1246,36 +1246,36 @@ void UNiagaraEmitter::CacheFromCompiledData(const FNiagaraDataSetCompiledData* C
 		}
 
 		// See how many particles we can fit in a GPU buffer. This number can be quite small on some platforms.
-const uint64 MaxBufferElements = (GNiagaraEmitterMaxGPUBufferElements > 0) ? uint64(GNiagaraEmitterMaxGPUBufferElements) : GetMaxBufferDimension();
+		const uint64 MaxBufferElements = (GNiagaraEmitterMaxGPUBufferElements > 0) ? uint64(GNiagaraEmitterMaxGPUBufferElements) : GetMaxBufferDimension();
 
-// Don't just cast the result of the division to 32-bit, since that will produce garbage if MaxNumInstances is larger than UINT_MAX. Saturate instead.
-MaxAllocationCount = (uint32)FMath::Min(MaxBufferElements / MaxGPUBufferComponents, (uint64)UINT_MAX);
-MaxInstanceCount = MaxAllocationCount;
+		// Don't just cast the result of the division to 32-bit, since that will produce garbage if MaxNumInstances is larger than UINT_MAX. Saturate instead.
+		MaxAllocationCount = (uint32)FMath::Min(MaxBufferElements / MaxGPUBufferComponents, (uint64)UINT_MAX);
+		MaxInstanceCount = MaxAllocationCount;
 
-if (SimTarget == ENiagaraSimTarget::GPUComputeSim)
-{
-	// Round down to nearest thread group size
-	MaxAllocationCount = FMath::DivideAndRoundDown(MaxAllocationCount, NiagaraComputeMaxThreadGroupSize) * NiagaraComputeMaxThreadGroupSize;
+		if (SimTarget == ENiagaraSimTarget::GPUComputeSim)
+		{
+			// Round down to nearest thread group size
+			MaxAllocationCount = FMath::DivideAndRoundDown(MaxAllocationCount, NiagaraComputeMaxThreadGroupSize) * NiagaraComputeMaxThreadGroupSize;
 
-	// We can't go below a thread group size, the rest of the code does not handle this and we should only hit this when buffer size forcing is enabled
-	if (MaxAllocationCount < NiagaraComputeMaxThreadGroupSize)
-	{
-		MaxAllocationCount = NiagaraComputeMaxThreadGroupSize;
-		check(MaxAllocationCount * MaxGPUBufferComponents <= GetMaxBufferDimension());
-	}
+			// We can't go below a thread group size, the rest of the code does not handle this and we should only hit this when buffer size forcing is enabled
+			if (MaxAllocationCount < NiagaraComputeMaxThreadGroupSize)
+			{
+				MaxAllocationCount = NiagaraComputeMaxThreadGroupSize;
+				check(MaxAllocationCount * MaxGPUBufferComponents <= GetMaxBufferDimension());
+			}
 
-	// -1 because we need a scratch instance on the GPU
-	MaxInstanceCount = MaxAllocationCount - 1;
-}
+			// -1 because we need a scratch instance on the GPU
+			MaxInstanceCount = MaxAllocationCount - 1;
+		}
 
-if (AllocationMode == EParticleAllocationMode::FixedCount)
-{
-	MaxInstanceCount = FMath::Min(MaxInstanceCount, uint32(FMath::Max(PreAllocationCount, 0)));
-}
+		if (AllocationMode == EParticleAllocationMode::FixedCount)
+		{
+			MaxInstanceCount = FMath::Min(MaxInstanceCount, uint32(FMath::Max(PreAllocationCount, 0)));
+		}
 	}
 	else
 	{
-	MaxInstanceCount = 0;
+		MaxInstanceCount = 0;
 	}
 
 	// Cache shaders for all GPU scripts
@@ -1307,6 +1307,20 @@ if (AllocationMode == EParticleAllocationMode::FixedCount)
 #endif
 
 	RebuildRendererBindings();
+
+	if ( GPUComputeScript )
+	{
+		for (FSimulationStageMetaData& SimStageMetaData : GPUComputeScript->GetVMExecutableData().SimulationStageMetaData)
+		{
+			if (SimStageMetaData.bParticleIterationStateEnabled)
+			{
+				if (const FNiagaraVariableLayoutInfo* VariableInfo = CompiledData->FindVariableLayoutInfo(FNiagaraVariableBase(FNiagaraTypeDefinition::GetIntDef(), SimStageMetaData.ParticleIterationStateBinding)))
+				{
+					SimStageMetaData.ParticleIterationStateComponentIndex = VariableInfo->Int32ComponentStart;
+				}
+			}
+		}
+	}
 }
 
 void UNiagaraEmitter::RebuildRendererBindings()
diff --git a/Engine/Plugins/FX/Niagara/Source/Niagara/Private/NiagaraGpuComputeDispatch.cpp b/Engine/Plugins/FX/Niagara/Source/Niagara/Private/NiagaraGpuComputeDispatch.cpp
--- a/Engine/Plugins/FX/Niagara/Source/Niagara/Private/NiagaraGpuComputeDispatch.cpp
+++ b/Engine/Plugins/FX/Niagara/Source/Niagara/Private/NiagaraGpuComputeDispatch.cpp
@@ -1587,6 +1587,13 @@ void FNiagaraGpuComputeDispatch::DispatchStage(FRHICommandList& RHICmdList, FRHI
 		SetShaderValue(RHICmdList, RHIComputeShader, ComputeShader->SimulationStageNormalizedIterationIndexParam, SimulationStageNormalizedIterationIndex);
 	}
 
+	// Set particle iteration state info
+	// Where X = Parameter Binding, YZ = Inclusive Range
+	{
+		int ParticleIterationStateInfo[3] = { SimStageData.StageMetaData->ParticleIterationStateComponentIndex, SimStageData.StageMetaData->ParticleIterationStateRange.X, SimStageData.StageMetaData->ParticleIterationStateRange.Y };
+		SetShaderValue(RHICmdList, RHIComputeShader, ComputeShader->ParticleIterationStateInfoParam, ParticleIterationStateInfo);
+	}
+
 	// Set static input buffers
 	SetSRVParameter(RHICmdList, RHIComputeShader, ComputeShader->StaticInputFloatParam, Tick.SystemGpuComputeProxy->StaticFloatBuffer);
 
diff --git a/Engine/Plugins/FX/Niagara/Source/Niagara/Private/NiagaraSimulationStageBase.cpp b/Engine/Plugins/FX/Niagara/Source/Niagara/Private/NiagaraSimulationStageBase.cpp
--- a/Engine/Plugins/FX/Niagara/Source/Niagara/Private/NiagaraSimulationStageBase.cpp
+++ b/Engine/Plugins/FX/Niagara/Source/Niagara/Private/NiagaraSimulationStageBase.cpp
@@ -44,6 +44,8 @@ void UNiagaraSimulationStageBase::RequestRecompile()
 
 void UNiagaraSimulationStageBase::PostEditChangeProperty(struct FPropertyChangedEvent& PropertyChangedEvent)
 {
+	Super::PostEditChangeProperty(PropertyChangedEvent);
+
 	FName PropertyName;
 	if (PropertyChangedEvent.Property)
 	{
@@ -69,6 +71,9 @@ bool UNiagaraSimulationStageGeneric::AppendCompileHash(FNiagaraCompileHashVisito
 	InVisitor->UpdatePOD(TEXT("bDisablePartialParticleUpdate"), bDisablePartialParticleUpdate ? 1 : 0);
 	InVisitor->UpdateString(TEXT("DataInterface"), DataInterface.BoundVariable.GetName().ToString());
 	InVisitor->UpdateString(TEXT("SimulationStageName"), SimulationStageName.ToString());
+	InVisitor->UpdatePOD(TEXT("bParticleIterationStateEnabled"), bParticleIterationStateEnabled ? 1 : 0);
+	InVisitor->UpdateString(TEXT("ParticleIterationStateBinding"), ParticleIterationStateBinding.GetDataSetBindableVariable().GetName().ToString());
+	InVisitor->UpdateString(TEXT("ParticleIterationStateRange"), FString::Printf(TEXT("%d,%d"), ParticleIterationStateRange.X, ParticleIterationStateRange.Y));
 	return true;
 }
 
@@ -89,6 +94,13 @@ void UNiagaraSimulationStageGeneric::PostInitProperties()
 			FNiagaraVariableBase(FNiagaraTypeDefinition::GetIntDef(), NAME_None),
 			ENiagaraRendererSourceDataMode::Emitter
 		);
+
+		static const FName ParticleStateIndex("Particles.StateIndex");
+		ParticleIterationStateBinding.Setup(
+			FNiagaraVariableBase(FNiagaraTypeDefinition::GetIntDef(), ParticleStateIndex),
+			FNiagaraVariableBase(FNiagaraTypeDefinition::GetIntDef(), ParticleStateIndex),
+			ENiagaraRendererSourceDataMode::Particles
+		);
 	}
 }
 
@@ -136,6 +148,18 @@ void UNiagaraSimulationStageGeneric::PostEditChangeProperty(struct FPropertyChan
 	{
 		bNeedsRecompile = true;
 	}
+	else if (PropertyName == GET_MEMBER_NAME_CHECKED(UNiagaraSimulationStageGeneric, bParticleIterationStateEnabled))
+	{
+		bNeedsRecompile = true;
+	}
+	else if (PropertyName == GET_MEMBER_NAME_CHECKED(UNiagaraSimulationStageGeneric, ParticleIterationStateBinding))
+	{
+		bNeedsRecompile = true;
+	}
+	else if (PropertyName == GET_MEMBER_NAME_CHECKED(UNiagaraSimulationStageGeneric, ParticleIterationStateRange))
+	{
+		bNeedsRecompile = true;
+	}
 
 	if (bNeedsRecompile)
 	{
diff --git a/Engine/Plugins/FX/Niagara/Source/Niagara/Public/NiagaraSimulationStageBase.h b/Engine/Plugins/FX/Niagara/Source/Niagara/Public/NiagaraSimulationStageBase.h
--- a/Engine/Plugins/FX/Niagara/Source/Niagara/Public/NiagaraSimulationStageBase.h
+++ b/Engine/Plugins/FX/Niagara/Source/Niagara/Public/NiagaraSimulationStageBase.h
@@ -52,11 +52,6 @@ class NIAGARA_API UNiagaraSimulationStageGeneric : public UNiagaraSimulationStag
 	GENERATED_BODY()
 
 public:
-	UNiagaraSimulationStageGeneric()
-		: Iterations(1)
-	{
-	}
-
 	/** Binding to a bool parameter which dynamically controls if the simulation stage is enabled or not. */
 	UPROPERTY(EditAnywhere, Category = "Simulation Stage")
 	FNiagaraVariableAttributeBinding EnabledBinding;
@@ -65,8 +60,8 @@ public:
 	UPROPERTY(EditAnywhere, Category = "Simulation Stage")
 	ENiagaraIterationSource IterationSource;
 
-	UPROPERTY(EditAnywhere, Category = "Simulation Stage", meta = (NoSpinbox = "true", ClampMin = 1, Tooltip = "The number of times we run this simulation stage before moving to the next stage."))
-	int32 Iterations;
+	UPROPERTY(EditAnywhere, Category = "Simulation Stage", meta = (NoSpinbox = "true", ClampMin = 1, DisplayName = "Num Iterations", Tooltip = "The number of times we run this simulation stage before moving to the next stage."))
+	int32 Iterations = 1;
 
 	/** Binding to an int parameter which dynamically controls the number of times the simulation stage runs. */
 	UPROPERTY(EditAnywhere, Category = "Simulation Stage")
@@ -85,6 +80,17 @@ public:
 	UPROPERTY(EditAnywhere, Category = "Simulation Stage", meta = (editcondition = "IterationSource == ENiagaraIterationSource::DataInterface"))
 	FNiagaraVariableDataInterfaceBinding DataInterface;
 
+	UPROPERTY(EditAnywhere, Category = "Simulation Stage", meta = (editcondition = "IterationSource == ENiagaraIterationSource::Particles"))
+	uint32 bParticleIterationStateEnabled : 1;
+
+	/** Particle state attribute binding, when enabled we will only allow particles who pass the state range check to be processed. */
+	UPROPERTY(EditAnywhere, Category = "Simulation Stage", meta = (editcondition = "IterationSource == ENiagaraIterationSource::Particles"))
+	FNiagaraVariableAttributeBinding ParticleIterationStateBinding;
+
+	/** The inclusive range used to check particle state binding against when enabled. */
+	UPROPERTY(EditAnywhere, Category = "Simulation Stage", meta = (editcondition = "IterationSource == ENiagaraIterationSource::Particles"))
+	FIntPoint ParticleIterationStateRange = FIntPoint(0, 0);
+
 	virtual void PostInitProperties() override;
 	virtual bool AppendCompileHash(FNiagaraCompileHashVisitor* InVisitor) const override;
 #if WITH_EDITOR
diff --git a/Engine/Plugins/FX/Niagara/Source/NiagaraEditor/Private/Customizations/NiagaraSimulationStageCustomization.cpp b/Engine/Plugins/FX/Niagara/Source/NiagaraEditor/Private/Customizations/NiagaraSimulationStageCustomization.cpp
new file mode 100644
--- /dev/null
+++ b/Engine/Plugins/FX/Niagara/Source/NiagaraEditor/Private/Customizations/NiagaraSimulationStageCustomization.cpp
@@ -0,0 +1,109 @@
+// Copyright Epic Games, Inc. All Rights Reserved.
+
+#include "NiagaraSimulationStageCustomization.h"
+
+#include "Modules/ModuleManager.h"
+
+//Customization
+#include "DetailCategoryBuilder.h"
+#include "DetailLayoutBuilder.h"
+#include "DetailWidgetRow.h"
+#include "IDetailChildrenBuilder.h"
+//Widgets
+#include "Widgets/Input/SCheckBox.h"
+#include "Widgets/Input/SEditableTextBox.h"
+#include "Widgets/SBoxPanel.h"
+#include "Widgets/SNiagaraDebugger.h"
+#include "Widgets/Views/STreeView.h"
+#include "Widgets/Input/SSearchBox.h"
+#include "Widgets/Layout/SScrollBorder.h"
+#include "Widgets/Input/SMultiLineEditableTextBox.h"
+///Niagara
+#include "NiagaraEditorModule.h"
+
+#define LOCTEXT_NAMESPACE "NiagaraSimulationStageCustomization"
+
+FNiagaraSimulationStageGenericCustomization::~FNiagaraSimulationStageGenericCustomization()
+{
+	if (UNiagaraSimulationStageGeneric* SimStage = WeakSimStage.Get())
+	{
+		SimStage->OnChanged().RemoveAll(this);
+	}
+}
+
+TSharedRef<IDetailCustomization> FNiagaraSimulationStageGenericCustomization::MakeInstance()
+{
+	return MakeShareable(new FNiagaraSimulationStageGenericCustomization());
+}
+
+void FNiagaraSimulationStageGenericCustomization::CustomizeDetails(const TSharedPtr<IDetailLayoutBuilder>& DetailBuilder)
+{
+	WeakDetailBuilder = DetailBuilder;
+	CustomizeDetails(*DetailBuilder);
+}
+
+void FNiagaraSimulationStageGenericCustomization::CustomizeDetails(class IDetailLayoutBuilder& DetailBuilder)
+{
+	// We only support customization on 1 object
+	TArray<TWeakObjectPtr<UObject>> ObjectsCustomized;
+	DetailBuilder.GetObjectsBeingCustomized(ObjectsCustomized);
+	if (ObjectsCustomized.Num() != 1 || !ObjectsCustomized[0]->IsA<UNiagaraSimulationStageGeneric>())
+	{
+		return;
+	}
+
+	UNiagaraSimulationStageGeneric* SimStage = CastChecked<UNiagaraSimulationStageGeneric>(ObjectsCustomized[0]);
+	WeakSimStage = SimStage;
+
+	SimStage->OnChanged().AddRaw(this, &FNiagaraSimulationStageGenericCustomization::OnPropertyChanged);
+
+	static const FName CategoryName = TEXT("Simulation Stage");
+	static const FName ParticleStageName = TEXT("Particle Parameters");
+	static const FName DataInterfaceStageName = TEXT("DataInterface Parameters");
+	IDetailCategoryBuilder& SimStageCategory = DetailBuilder.EditCategory(CategoryName);
+	IDetailCategoryBuilder& ParticleStageCategory = DetailBuilder.EditCategory(ParticleStageName);
+	IDetailCategoryBuilder& DataInterfaceStageCategory = DetailBuilder.EditCategory(DataInterfaceStageName);
+
+	// Hide all
+	DetailBuilder.HideProperty(GET_MEMBER_NAME_CHECKED(UNiagaraSimulationStageGeneric, EnabledBinding));
+	DetailBuilder.HideProperty(GET_MEMBER_NAME_CHECKED(UNiagaraSimulationStageGeneric, IterationSource));
+	DetailBuilder.HideProperty(GET_MEMBER_NAME_CHECKED(UNiagaraSimulationStageGeneric, Iterations));
+	DetailBuilder.HideProperty(GET_MEMBER_NAME_CHECKED(UNiagaraSimulationStageGeneric, NumIterationsBinding));
+	DetailBuilder.HideProperty(GET_MEMBER_NAME_CHECKED(UNiagaraSimulationStageGeneric, ExecuteBehavior));
+	DetailBuilder.HideProperty(GET_MEMBER_NAME_CHECKED(UNiagaraSimulationStageGeneric, bDisablePartialParticleUpdate));
+	DetailBuilder.HideProperty(GET_MEMBER_NAME_CHECKED(UNiagaraSimulationStageGeneric, DataInterface));
+	DetailBuilder.HideProperty(GET_MEMBER_NAME_CHECKED(UNiagaraSimulationStageGeneric, bParticleIterationStateEnabled));
+	DetailBuilder.HideProperty(GET_MEMBER_NAME_CHECKED(UNiagaraSimulationStageGeneric, ParticleIterationStateBinding));
+	DetailBuilder.HideProperty(GET_MEMBER_NAME_CHECKED(UNiagaraSimulationStageGeneric, ParticleIterationStateRange));
+
+	// Show properties in the order we want them to appear
+	SimStageCategory.AddProperty(DetailBuilder.GetProperty(GET_MEMBER_NAME_CHECKED(UNiagaraSimulationStageGeneric, EnabledBinding)));
+	SimStageCategory.AddProperty(DetailBuilder.GetProperty(GET_MEMBER_NAME_CHECKED(UNiagaraSimulationStageGeneric, Iterations)));
+	SimStageCategory.AddProperty(DetailBuilder.GetProperty(GET_MEMBER_NAME_CHECKED(UNiagaraSimulationStageGeneric, NumIterationsBinding)));
+	SimStageCategory.AddProperty(DetailBuilder.GetProperty(GET_MEMBER_NAME_CHECKED(UNiagaraSimulationStageGeneric, IterationSource)));
+	if ( SimStage->IterationSource == ENiagaraIterationSource::Particles )
+	{
+		ParticleStageCategory.AddProperty(DetailBuilder.GetProperty(GET_MEMBER_NAME_CHECKED(UNiagaraSimulationStageGeneric, bDisablePartialParticleUpdate)));
+		ParticleStageCategory.AddProperty(DetailBuilder.GetProperty(GET_MEMBER_NAME_CHECKED(UNiagaraSimulationStageGeneric, bParticleIterationStateEnabled)));
+		if (SimStage->bParticleIterationStateEnabled)
+		{
+			ParticleStageCategory.AddProperty(DetailBuilder.GetProperty(GET_MEMBER_NAME_CHECKED(UNiagaraSimulationStageGeneric, ParticleIterationStateBinding)));
+			ParticleStageCategory.AddProperty(DetailBuilder.GetProperty(GET_MEMBER_NAME_CHECKED(UNiagaraSimulationStageGeneric, ParticleIterationStateRange)));
+		}
+	}
+	else
+	{
+		DataInterfaceStageCategory.AddProperty(DetailBuilder.GetProperty(GET_MEMBER_NAME_CHECKED(UNiagaraSimulationStageGeneric, DataInterface)));
+		DataInterfaceStageCategory.AddProperty(DetailBuilder.GetProperty(GET_MEMBER_NAME_CHECKED(UNiagaraSimulationStageGeneric, ExecuteBehavior)));
+	}
+}
+
+void FNiagaraSimulationStageGenericCustomization::OnPropertyChanged()
+{
+	if (IDetailLayoutBuilder* DetailBuilder =  WeakDetailBuilder.Pin().Get() )
+	{
+		DetailBuilder->ForceRefreshDetails();
+	}
+}
+
+#undef LOCTEXT_NAMESPACE
diff --git a/Engine/Plugins/FX/Niagara/Source/NiagaraEditor/Private/Customizations/NiagaraSimulationStageCustomization.h b/Engine/Plugins/FX/Niagara/Source/NiagaraEditor/Private/Customizations/NiagaraSimulationStageCustomization.h
new file mode 100644
--- /dev/null
+++ b/Engine/Plugins/FX/Niagara/Source/NiagaraEditor/Private/Customizations/NiagaraSimulationStageCustomization.h
@@ -0,0 +1,35 @@
+// Copyright Epic Games, Inc. All Rights Reserved.
+
+#pragma once
+
+#include "CoreMinimal.h"
+#include "IDetailCustomization.h"
+#include "IPropertyTypeCustomization.h"
+#include "NiagaraCommon.h"
+#include "EdGraph/EdGraphSchema.h"
+#include "Layout/Visibility.h"
+#include "NiagaraSimulationStageBase.h"
+
+class FDetailWidgetRow;
+class IPropertyHandle;
+class IPropertyHandleArray;
+enum class ECheckBoxState : uint8;
+
+class FNiagaraSimulationStageGenericCustomization : public IDetailCustomization
+{
+public:
+	static TSharedRef<IDetailCustomization> MakeInstance();
+
+	virtual ~FNiagaraSimulationStageGenericCustomization();
+
+	// Begin IDetailCustomization interface
+	virtual void CustomizeDetails(const TSharedPtr<IDetailLayoutBuilder>& DetailBuilder) override;
+	virtual void CustomizeDetails(class IDetailLayoutBuilder& DetailBuilder) override;
+	// End IDetailCustomization interface
+
+	void OnPropertyChanged();
+
+private:
+	TWeakPtr<IDetailLayoutBuilder> WeakDetailBuilder;
+	TWeakObjectPtr<UNiagaraSimulationStageGeneric> WeakSimStage;
+};
diff --git a/Engine/Plugins/FX/Niagara/Source/NiagaraEditor/Private/Customizations/NiagaraTypeCustomizations.cpp b/Engine/Plugins/FX/Niagara/Source/NiagaraEditor/Private/Customizations/NiagaraTypeCustomizations.cpp
--- a/Engine/Plugins/FX/Niagara/Source/NiagaraEditor/Private/Customizations/NiagaraTypeCustomizations.cpp
+++ b/Engine/Plugins/FX/Niagara/Source/NiagaraEditor/Private/Customizations/NiagaraTypeCustomizations.cpp
@@ -264,12 +264,13 @@ TArray<FName> FNiagaraVariableAttributeBindingCustomization::GetNames(UNiagaraEm
 		return Names;
 	}
 
-	bool bSystem = true;
-	bool bEmitter = true;
-	bool bParticles = true;
-	bool bUser = true;
-	bool bAllowStatic = true;
+	const bool bSystem = true;
+	const bool bEmitter = true;
+	const bool bParticles = true;
+	const bool bUser = true;
+	const bool bAllowStatic = true;
 	TArray<FNiagaraVariableBase> Vars = FNiagaraStackAssetAction_VarBind::FindVariables(InEmitter, bSystem, bEmitter, bParticles, bUser, bAllowStatic);
+
 	for (const FNiagaraVariableBase& Var : Vars)
 	{
 		if (!bAllowStatic && Var.GetType() != TargetVariableBinding->GetType() )
@@ -285,9 +286,14 @@ TArray<FName> FNiagaraVariableAttributeBindingCustomization::GetNames(UNiagaraEm
 		{
 			Names.AddUnique(Var.GetName());
 		}
-		else if ( SimulationStage && !Var.IsInNameSpace(FNiagaraConstants::ParticleAttributeNamespace) )
+		else if ( SimulationStage )
 		{
-			Names.AddUnique(Var.GetName());
+			// Unless we have an explicit "particles." binding we are not in the particle namespace
+			const bool IsParticleNamespace = !TargetVariableBinding->GetName().IsNone() && TargetVariableBinding->IsParticleBinding();
+			if ( IsParticleNamespace == Var.IsInNameSpace(FNiagaraConstants::ParticleAttributeNamespace) )
+			{
+				Names.AddUnique(Var.GetName());
+			}
 		}
 	}
 
diff --git a/Engine/Plugins/FX/Niagara/Source/NiagaraEditor/Private/NiagaraCompiler.cpp b/Engine/Plugins/FX/Niagara/Source/NiagaraEditor/Private/NiagaraCompiler.cpp
--- a/Engine/Plugins/FX/Niagara/Source/NiagaraEditor/Private/NiagaraCompiler.cpp
+++ b/Engine/Plugins/FX/Niagara/Source/NiagaraEditor/Private/NiagaraCompiler.cpp
@@ -693,6 +693,17 @@ void FNiagaraCompileRequestData::FinishPrecompile(const TArray<FNiagaraVariable>
 					SimStageData.IterationSource		= GenericStage->IterationSource == ENiagaraIterationSource::DataInterface ? GenericStage->DataInterface.BoundVariable.GetName() : FName();
 					SimStageData.ExecuteBehavior		= GenericStage->ExecuteBehavior;
 					SimStageData.PartialParticleUpdate	= GenericStage->bDisablePartialParticleUpdate == false;
+					SimStageData.bParticleIterationStateEnabled	= GenericStage->bParticleIterationStateEnabled;
+					SimStageData.ParticleIterationStateRange	= GenericStage->ParticleIterationStateRange;
+
+					if (SimStageData.bParticleIterationStateEnabled)
+					{
+						FString AttributeName = GenericStage->ParticleIterationStateBinding.GetName().ToString();
+						if ( ensureMsgf(AttributeName.RemoveFromStart(TEXT("Particles.")), TEXT("Attribute '%s' is not in particles namespace"), *AttributeName) )
+						{
+							SimStageData.ParticleIterationStateBinding = FName(AttributeName);
+						}
+					}
 
 					++ActiveStageCount;
 				}
diff --git a/Engine/Plugins/FX/Niagara/Source/NiagaraEditor/Private/NiagaraEditorModule.cpp b/Engine/Plugins/FX/Niagara/Source/NiagaraEditor/Private/NiagaraEditorModule.cpp
--- a/Engine/Plugins/FX/Niagara/Source/NiagaraEditor/Private/NiagaraEditorModule.cpp
+++ b/Engine/Plugins/FX/Niagara/Source/NiagaraEditor/Private/NiagaraEditorModule.cpp
@@ -109,6 +109,7 @@
 #include "Customizations/NiagaraDebugHUDCustomization.h"
 #include "Customizations/NiagaraBakerSettingsDetails.h"
 #include "Customizations/NiagaraOutlinerCustomization.h"
+#include "Customizations/NiagaraSimulationStageCustomization.h"
 
 #include "NiagaraComponent.h"
 #include "NiagaraNodeStaticSwitch.h"
@@ -905,6 +906,10 @@ void FNiagaraEditorModule::StartupModule()
 		FNiagaraVariableMetaData::StaticStruct()->GetFName(),
 		FOnGetPropertyTypeCustomizationInstance::CreateStatic(&FNiagaraVariableMetaDataCustomization::MakeInstance));
 
+	PropertyModule.RegisterCustomClassLayout(
+		UNiagaraSimulationStageGeneric::StaticClass()->GetFName(),
+		FOnGetDetailCustomizationInstance::CreateStatic(&FNiagaraSimulationStageGenericCustomization::MakeInstance));
+
 #if WITH_NIAGARA_DEBUGGER
 	PropertyModule.RegisterCustomPropertyTypeLayout(
 		FNiagaraDebugHUDVariable::StaticStruct()->GetFName(),
diff --git a/Engine/Plugins/FX/Niagara/Source/NiagaraEditor/Private/NiagaraHlslTranslator.cpp b/Engine/Plugins/FX/Niagara/Source/NiagaraEditor/Private/NiagaraHlslTranslator.cpp
--- a/Engine/Plugins/FX/Niagara/Source/NiagaraEditor/Private/NiagaraHlslTranslator.cpp
+++ b/Engine/Plugins/FX/Niagara/Source/NiagaraEditor/Private/NiagaraHlslTranslator.cpp
@@ -1221,6 +1221,9 @@ const FNiagaraTranslateResults &FHlslNiagaraTranslator::Translate(const FNiagara
 					TranslationStages[Index].bPartialParticleUpdate = CompileSimStageData.PartialParticleUpdate;
 					TranslationStages[Index].IterationSource = CompileSimStageData.IterationSource;
 					TranslationStages[Index].NumIterationsBinding = CompileSimStageData.NumIterationsBinding;
+					TranslationStages[Index].bParticleIterationStateEnabled = CompileSimStageData.bParticleIterationStateEnabled;
+					TranslationStages[Index].ParticleIterationStateBinding = CompileSimStageData.ParticleIterationStateBinding;
+					TranslationStages[Index].ParticleIterationStateRange = CompileSimStageData.ParticleIterationStateRange;
 
 					ParamMapHistories.AddDefaulted(1);
 					ParamMapHistoriesSourceInOtherHistories.AddDefaulted(1);
@@ -1285,6 +1288,9 @@ const FNiagaraTranslateResults &FHlslNiagaraTranslator::Translate(const FNiagara
 					SimulationStageMetaData.NumIterations = TranslationStages[Index].NumIterations;
 					SimulationStageMetaData.bWritesParticles = TranslationStages[Index].bWritesParticles;
 					SimulationStageMetaData.bPartialParticleUpdate = TranslationStages[Index].bPartialParticleUpdate;
+					SimulationStageMetaData.bParticleIterationStateEnabled = TranslationStages[Index].bParticleIterationStateEnabled;
+					SimulationStageMetaData.ParticleIterationStateBinding = TranslationStages[Index].ParticleIterationStateBinding;
+					SimulationStageMetaData.ParticleIterationStateRange = TranslationStages[Index].ParticleIterationStateRange;
 
 					// Determine dispatch information from iteration source (if we have one)
 					if ( !SimulationStageMetaData.IterationSource.IsNone() )
@@ -1908,6 +1914,11 @@ const FNiagaraTranslateResults &FHlslNiagaraTranslator::Translate(const FNiagara
 			Preamble.Appendf(TEXT("//\tWritesParticles = %s\n"), SimStageMetaData.bWritesParticles ? TEXT("True") : TEXT("False"));
 			Preamble.Appendf(TEXT("//\tPartialParticleUpdate = %s\n"), SimStageMetaData.bPartialParticleUpdate ? TEXT("True") : TEXT("False"));
 
+			if ( SimStageMetaData.bParticleIterationStateEnabled )
+			{
+				Preamble.Appendf(TEXT("//\tParticleIterationStage = Attribute(%s) Range(%d ... %d)\n"), *SimStageMetaData.ParticleIterationStateBinding.ToString(), SimStageMetaData.ParticleIterationStateRange.X, SimStageMetaData.ParticleIterationStateRange.Y);
+			}
+
 			for (const FName& Dest : SimStageMetaData.OutputDestinations)
 			{
 				Preamble.Appendf(TEXT("//\tOutputs to: \"%s\"\n"), *Dest.ToString());
@@ -3177,6 +3188,29 @@ void FHlslNiagaraTranslator::DefineMainGPUFunctions(
 		// Particle iteration stage
 		if ( TranslationStage.IterationSource.IsNone() )
 		{
+			// Do we have particle state iteration enable
+			SimRunSpawnUpdateLogicString += TEXT("	bool bRunSpawnUpdateLogic = true;\n");
+			if ( TranslationStage.bParticleIterationStateEnabled )
+			{
+				if ( TranslationStage.bPartialParticleUpdate == false )
+				{
+					FName StageName;
+					if (CompilationOutput.ScriptData.SimulationStageMetaData.IsValidIndex(TranslationStage.SimulationStageIndex))
+					{
+						StageName = CompilationOutput.ScriptData.SimulationStageMetaData[TranslationStage.SimulationStageIndex].SimulationStageName;
+					}
+
+					Error(FText::Format(LOCTEXT("ParticleIterationState_Invalid", "Simulation stage '{0}' is incompatible with particle state iteration due to killing particles or disabling particle updates."), FText::FromName(StageName)), nullptr, nullptr);
+				}
+				SimRunSpawnUpdateLogicString += TEXT(
+					"	if ( ParticleIterationStateInfo.x != -1 )\n"
+					"	{\n"
+					"		int ParticleStateValue = InputDataInt(0, uint(ParticleIterationStateInfo.x), GLinearThreadId);\n"
+					"		bRunSpawnUpdateLogic = (ParticleStateValue >= ParticleIterationStateInfo.y) && (ParticleStateValue <= ParticleIterationStateInfo.z);\n"
+					"	}\n"
+				);
+			}
+
 			// We combine the update & spawn scripts together on GPU so we only need to check for spawning on the first translation stage
 			// Note: Depending on how spawning inside stages works we may need to enable the spawn logic for those stages *only*			
 			const bool bParticleSpawnStage = i == 1;
@@ -3192,15 +3226,15 @@ void FHlslNiagaraTranslator::DefineMainGPUFunctions(
 					"		GSpawnStartInstance = RWInstanceCounts[ReadInstanceCountOffset];\n"
 					"	}\n"
 					"	const int MaxInstances = GSpawnStartInstance + SpawnedInstances;\n"
-					"	const bool bRunUpdateLogic = GLinearThreadId < GSpawnStartInstance && GLinearThreadId < MaxInstances;\n"
-					"	const bool bRunSpawnLogic = GLinearThreadId >= GSpawnStartInstance && GLinearThreadId < MaxInstances;\n"
+					"	const bool bRunUpdateLogic = bRunSpawnUpdateLogic && GLinearThreadId < GSpawnStartInstance && GLinearThreadId < MaxInstances;\n"
+					"	const bool bRunSpawnLogic = bRunSpawnUpdateLogic && GLinearThreadId >= GSpawnStartInstance && GLinearThreadId < MaxInstances;\n"
 				);
 			}
 			else
 			{
 				SimRunSpawnUpdateLogicString += TEXT(
 					"	GSpawnStartInstance = RWInstanceCounts[ReadInstanceCountOffset];\n"
-					"	const bool bRunUpdateLogic = GLinearThreadId < GSpawnStartInstance;\n"
+					"	const bool bRunUpdateLogic = bRunSpawnUpdateLogic && GLinearThreadId < GSpawnStartInstance;\n"
 					"	const bool bRunSpawnLogic = false;\n"
 				);
 			}
diff --git a/Engine/Plugins/FX/Niagara/Source/NiagaraEditor/Private/NiagaraHlslTranslator.h b/Engine/Plugins/FX/Niagara/Source/NiagaraEditor/Private/NiagaraHlslTranslator.h
--- a/Engine/Plugins/FX/Niagara/Source/NiagaraEditor/Private/NiagaraHlslTranslator.h
+++ b/Engine/Plugins/FX/Niagara/Source/NiagaraEditor/Private/NiagaraHlslTranslator.h
@@ -97,6 +97,9 @@ public:
 		FName							IterationSource;
 		ENiagaraSimStageExecuteBehavior	ExecuteBehavior = ENiagaraSimStageExecuteBehavior::Always;
 		mutable bool					PartialParticleUpdate = false;
+		bool							bParticleIterationStateEnabled = false;
+		FName							ParticleIterationStateBinding;
+		FIntPoint						ParticleIterationStateRange = FIntPoint::ZeroValue;
 	};
 	TArray<FCompileSimStageData> CompileSimStageData;
 
@@ -369,6 +372,9 @@ public:
 	int32 NumIterations = 1;
 	FName NumIterationsBinding;
 	ENiagaraSimStageExecuteBehavior ExecuteBehavior = ENiagaraSimStageExecuteBehavior::Always;
+	bool bParticleIterationStateEnabled = false;
+	FName ParticleIterationStateBinding;
+	FIntPoint ParticleIterationStateRange = FIntPoint::ZeroValue;
 	bool bUsesAlive = false;
 	bool bWritesAlive = false;
 	bool bWritesParticles = false;
diff --git a/Engine/Plugins/FX/Niagara/Source/NiagaraShader/Private/NiagaraScriptBase.cpp b/Engine/Plugins/FX/Niagara/Source/NiagaraShader/Private/NiagaraScriptBase.cpp
--- a/Engine/Plugins/FX/Niagara/Source/NiagaraShader/Private/NiagaraScriptBase.cpp
+++ b/Engine/Plugins/FX/Niagara/Source/NiagaraShader/Private/NiagaraScriptBase.cpp
@@ -10,6 +10,7 @@ UNiagaraScriptBase::UNiagaraScriptBase(const FObjectInitializer& ObjectInitializ
 FSimulationStageMetaData::FSimulationStageMetaData()
 	: bWritesParticles(0)
 	, bPartialParticleUpdate(0)
+	, bParticleIterationStateEnabled(0)
 	, GpuDispatchType(ENiagaraGpuDispatchType::OneD)
 	, GpuDispatchNumThreads(0, 0, 0)
 {
diff --git a/Engine/Plugins/FX/Niagara/Source/NiagaraShader/Private/NiagaraShader.cpp b/Engine/Plugins/FX/Niagara/Source/NiagaraShader/Private/NiagaraShader.cpp
--- a/Engine/Plugins/FX/Niagara/Source/NiagaraShader/Private/NiagaraShader.cpp
+++ b/Engine/Plugins/FX/Niagara/Source/NiagaraShader/Private/NiagaraShader.cpp
@@ -1230,14 +1230,11 @@ void FNiagaraShader::BindParams(const TArray<FNiagaraDataInterfaceGPUParamInfo>&
 	EmitterSpawnInfoOffsetsParam.Bind(ParameterMap, TEXT("EmitterSpawnInfoOffsets"));
 	EmitterSpawnInfoParamsParam.Bind(ParameterMap, TEXT("EmitterSpawnInfoParams"));
 
-	NumEventsPerParticleParam.Bind(ParameterMap, TEXT("NumEventsPerParticle"));
-	NumParticlesPerEventParam.Bind(ParameterMap, TEXT("NumParticlesPerEvent"));
-	CopyInstancesBeforeStartParam.Bind(ParameterMap, TEXT("CopyInstancesBeforeStart"));
-
 	NumSpawnedInstancesParam.Bind(ParameterMap, TEXT("SpawnedInstances"));
 
 	SimulationStageIterationInfoParam.Bind(ParameterMap, TEXT("SimulationStageIterationInfo"));
 	SimulationStageNormalizedIterationIndexParam.Bind(ParameterMap, TEXT("SimulationStageNormalizedIterationIndex"));
+	ParticleIterationStateInfoParam.Bind(ParameterMap, TEXT("ParticleIterationStateInfo"));
 	DispatchThreadIdBoundsParam.Bind(ParameterMap, TEXT("DispatchThreadIdBounds"));
 	DispatchThreadIdToLinearParam.Bind(ParameterMap, TEXT("DispatchThreadIdToLinear"));
 
diff --git a/Engine/Plugins/FX/Niagara/Source/NiagaraShader/Public/NiagaraScriptBase.h b/Engine/Plugins/FX/Niagara/Source/NiagaraShader/Public/NiagaraScriptBase.h
--- a/Engine/Plugins/FX/Niagara/Source/NiagaraShader/Public/NiagaraScriptBase.h
+++ b/Engine/Plugins/FX/Niagara/Source/NiagaraShader/Public/NiagaraScriptBase.h
@@ -61,6 +61,20 @@ public:
 	UPROPERTY()
 	uint32 bPartialParticleUpdate : 1;
 
+	UPROPERTY()
+	uint32 bParticleIterationStateEnabled : 1;
+
+	/** When the value is not none this is the binding used for particle state iteration stages. */
+	UPROPERTY()
+	FName ParticleIterationStateBinding;
+
+	/** Cached component index for particle iteration stage. */
+	int32 ParticleIterationStateComponentIndex = INDEX_NONE;
+
+	/** Inclusive range to compare the particle state value with. */
+	UPROPERTY()
+	FIntPoint ParticleIterationStateRange;
+
 	/** DataInterfaces that we write to in this stage.*/
 	UPROPERTY()
 	TArray<FName> OutputDestinations;
diff --git a/Engine/Plugins/FX/Niagara/Source/NiagaraShader/Public/NiagaraShader.h b/Engine/Plugins/FX/Niagara/Source/NiagaraShader/Public/NiagaraShader.h
--- a/Engine/Plugins/FX/Niagara/Source/NiagaraShader/Public/NiagaraShader.h
+++ b/Engine/Plugins/FX/Niagara/Source/NiagaraShader/Public/NiagaraShader.h
@@ -93,13 +93,11 @@ public:
 	LAYOUT_FIELD(FShaderParameter, EmitterTickCounterParam);
 	LAYOUT_FIELD(FShaderParameter, EmitterSpawnInfoOffsetsParam);
 	LAYOUT_FIELD(FShaderParameter, EmitterSpawnInfoParamsParam);
-	LAYOUT_FIELD(FShaderParameter, NumEventsPerParticleParam);
-	LAYOUT_FIELD(FShaderParameter, NumParticlesPerEventParam);
-	LAYOUT_FIELD(FShaderParameter, CopyInstancesBeforeStartParam);
 	LAYOUT_FIELD(FShaderParameter, NumSpawnedInstancesParam);
 
 	LAYOUT_FIELD(FShaderParameter, SimulationStageIterationInfoParam);
 	LAYOUT_FIELD(FShaderParameter, SimulationStageNormalizedIterationIndexParam);
+	LAYOUT_FIELD(FShaderParameter, ParticleIterationStateInfoParam);
 	LAYOUT_FIELD(FShaderParameter, DispatchThreadIdBoundsParam);
 	LAYOUT_FIELD(FShaderParameter, DispatchThreadIdToLinearParam);
 
