- Initialize struct member ParticleIterationStateRange

#rnx

diff --git a/Engine/Plugins/FX/Niagara/Source/NiagaraShader/Public/NiagaraScriptBase.h b/Engine/Plugins/FX/Niagara/Source/NiagaraShader/Public/NiagaraScriptBase.h
--- a/Engine/Plugins/FX/Niagara/Source/NiagaraShader/Public/NiagaraScriptBase.h
+++ b/Engine/Plugins/FX/Niagara/Source/NiagaraShader/Public/NiagaraScriptBase.h
@@ -73,7 +73,7 @@ public:
 
 	/** Inclusive range to compare the particle state value with. */
 	UPROPERTY()
-	FIntPoint ParticleIterationStateRange;
+	FIntPoint ParticleIterationStateRange = FIntPoint(0, 0);
 
 	/** DataInterfaces that we write to in this stage.*/
 	UPROPERTY()
