Fixed issue where spawning empty systems would leak.

diff --git a/Engine/Plugins/FX/Niagara/Source/Niagara/Private/NiagaraComponent.cpp b/Engine/Plugins/FX/Niagara/Source/Niagara/Private/NiagaraComponent.cpp
--- a/Engine/Plugins/FX/Niagara/Source/Niagara/Private/NiagaraComponent.cpp
+++ b/Engine/Plugins/FX/Niagara/Source/Niagara/Private/NiagaraComponent.cpp
@@ -1146,7 +1146,7 @@ void UNiagaraComponent::ActivateInternal(bool bReset /* = false */, bool bIsScal
 		ResetMode = FNiagaraSystemInstance::EResetMode::None;//Already done a reinit sete
 	}
 
-	if (!SystemInstanceController)
+	if (!SystemInstanceController || SystemInstanceController->IsComplete())
 	{
 		OnSystemComplete(true);
 		return;
