- Fix for ParticleIterationStateRange not recompiling

#rnx

diff --git a/Engine/Plugins/FX/Niagara/Source/Niagara/Private/NiagaraSimulationStageBase.cpp b/Engine/Plugins/FX/Niagara/Source/Niagara/Private/NiagaraSimulationStageBase.cpp
--- a/Engine/Plugins/FX/Niagara/Source/Niagara/Private/NiagaraSimulationStageBase.cpp
+++ b/Engine/Plugins/FX/Niagara/Source/Niagara/Private/NiagaraSimulationStageBase.cpp
@@ -110,10 +110,15 @@ void UNiagaraSimulationStageGeneric::PostEditChangeProperty(struct FPropertyChan
 	Super::PostEditChangeProperty(PropertyChangedEvent);
 
 	FName PropertyName;
+	FName MemberPropertyName;
 	if (PropertyChangedEvent.Property)
 	{
 		PropertyName = PropertyChangedEvent.Property->GetFName();
 	}
+	if (PropertyChangedEvent.MemberProperty)
+	{
+		MemberPropertyName = PropertyChangedEvent.MemberProperty->GetFName();
+	}
 
 	bool bNeedsRecompile = false;
 	if (PropertyName == GET_MEMBER_NAME_CHECKED(UNiagaraSimulationStageGeneric, EnabledBinding))
@@ -156,7 +161,7 @@ void UNiagaraSimulationStageGeneric::PostEditChangeProperty(struct FPropertyChan
 	{
 		bNeedsRecompile = true;
 	}
-	else if (PropertyName == GET_MEMBER_NAME_CHECKED(UNiagaraSimulationStageGeneric, ParticleIterationStateRange))
+	else if (MemberPropertyName == GET_MEMBER_NAME_CHECKED(UNiagaraSimulationStageGeneric, ParticleIterationStateRange))
 	{
 		bNeedsRecompile = true;
 	}
@@ -176,4 +181,4 @@ FName UNiagaraSimulationStageGeneric::GetStackContextReplacementName() const
 	ensureMsgf(false, TEXT("Should not get here! Need to handle unknown case!"));
 	return NAME_None;
 }
-#endif
\ No newline at end of file
+#endif
