Various dependency fixes;
Remove duplicated method.
Auto fixup dependency order when using fix now button.
Do not auto fixup dependency order when target index in stack is specified (insert above/below.)
Fix logic error in choosing best index with off by one for provided dependencies.

diff --git a/Engine/Plugins/FX/Niagara/Source/NiagaraEditor/Private/ViewModels/Stack/NiagaraStackGraphUtilities.cpp b/Engine/Plugins/FX/Niagara/Source/NiagaraEditor/Private/ViewModels/Stack/NiagaraStackGraphUtilities.cpp
--- a/Engine/Plugins/FX/Niagara/Source/NiagaraEditor/Private/ViewModels/Stack/NiagaraStackGraphUtilities.cpp
+++ b/Engine/Plugins/FX/Niagara/Source/NiagaraEditor/Private/ViewModels/Stack/NiagaraStackGraphUtilities.cpp
@@ -3309,10 +3309,10 @@ int32 FNiagaraStackGraphUtilities::DependencyUtilities::FindBestIndexForModuleIn
 	}
 
 	// Find the greatest and least indices to satisfy the pre and post dependencies of the new module script being added, respectively.
-	int32 LeastDependencyIndex = LeastIndex;
-	int32 GreatestDependencyIndex = GreatestIndex;
-	bool bDependencyLowerBoundFound = false;
-	bool bDependencyUpperBoundFound = false;
+	int32 GreatestRequiredDependencyLowerBoundIdx = LeastIndex;
+	int32 LeastRequiredDependencyUpperBoundIdx = GreatestIndex;
+	bool bRequiredDependencyLowerBoundFound = false;
+	bool bRequiredDependencyUpperBoundFound = false;
 
 	const TArray<FNiagaraModuleDependency>& NewModuleScriptRequiredDependencies = ScriptData->RequiredDependencies;
 
@@ -3336,14 +3336,14 @@ int32 FNiagaraStackGraphUtilities::DependencyUtilities::FindBestIndexForModuleIn
 				if (RequiredDependency.Type == ENiagaraModuleDependencyType::PreDependency)
 				{
 					int32 SuggestedIndex = CurrentStackModuleData.Index + 1;
-					LeastDependencyIndex = LeastDependencyIndex < SuggestedIndex ? SuggestedIndex : LeastDependencyIndex;
-					bDependencyLowerBoundFound = true;
+					GreatestRequiredDependencyLowerBoundIdx = GreatestRequiredDependencyLowerBoundIdx < SuggestedIndex ? SuggestedIndex : GreatestRequiredDependencyLowerBoundIdx;
+					bRequiredDependencyLowerBoundFound = true;
 				}
 				else if (RequiredDependency.Type == ENiagaraModuleDependencyType::PostDependency)
 				{
 					int32 SuggestedIndex = CurrentStackModuleData.Index;
-					GreatestDependencyIndex = GreatestDependencyIndex > SuggestedIndex ? SuggestedIndex : GreatestDependencyIndex;
-					bDependencyUpperBoundFound = true;
+					LeastRequiredDependencyUpperBoundIdx = LeastRequiredDependencyUpperBoundIdx > SuggestedIndex ? SuggestedIndex : LeastRequiredDependencyUpperBoundIdx;
+					bRequiredDependencyUpperBoundFound = true;
 				}
 			}
 		}
@@ -3351,28 +3351,27 @@ int32 FNiagaraStackGraphUtilities::DependencyUtilities::FindBestIndexForModuleIn
 
 	// Check that there is valid index which satisfies all dependencies for the new module script.
 	int32 TargetIndex = INDEX_NONE;
-	if (bDependencyLowerBoundFound)
+	if (bRequiredDependencyLowerBoundFound)
 	{
-		if (bDependencyUpperBoundFound && GreatestDependencyIndex < LeastDependencyIndex)
+		if (bRequiredDependencyUpperBoundFound && LeastRequiredDependencyUpperBoundIdx < GreatestRequiredDependencyLowerBoundIdx)
 		{
-			// It is impossible to satisfy both target indices: do nothing.
+			// It is impossible to satisfy both target indices as the upper bound is less than the lower bound: do nothing.
 			return INDEX_NONE;
 		}
-		TargetIndex = LeastDependencyIndex;
+		TargetIndex = GreatestRequiredDependencyLowerBoundIdx;
 	}
-	else if (bDependencyUpperBoundFound)
+	else if (bRequiredDependencyUpperBoundFound)
 	{
-		TargetIndex = GreatestDependencyIndex;
+		TargetIndex = LeastRequiredDependencyUpperBoundIdx;
 	}
 
 	// Do another pass to find a target index that satisfies dependencies for other modules in the stack by providing as many dependencies as possible.
-	int32 LeastRequirementIndex = LeastIndex;
-	int32 GreatestRequirementIndex = GreatestIndex;
-	bool bRequirementLowerBoundFound = false;
-	bool bRequirementUpperBoundFound = false;
+	int32 GreatestProvidedDependencyLowerBoundIdx = GreatestIndex;
+	int32 LeastProvidedDependencyUpperBoundIdx = LeastIndex;
+	bool bProvidedDependencyLowerBoundFound = false;
+	bool bProvidedDependencyUpperBoundFound = false;
 
 	const TArray<FName>& NewModuleScriptProvidedDependencies = ScriptData->ProvidedDependencies;
-
 	auto GetStackModuleDataRequiredDependenciesBeingProvided = [&NewModuleScriptProvidedDependencies](const FNiagaraStackModuleData& StackModuleData)->const TArray<FNiagaraModuleDependency> /*OutDependencies*/ {
 		TArray<FNiagaraModuleDependency> OutDependencies;
 		for (FNiagaraModuleDependency& RequiredDependency : StackModuleData.ModuleNode->GetScriptData()->RequiredDependencies)
@@ -3389,37 +3388,44 @@ int32 FNiagaraStackGraphUtilities::DependencyUtilities::FindBestIndexForModuleIn
 	{
 		for (const FNiagaraModuleDependency& ProvidedDependency : GetStackModuleDataRequiredDependenciesBeingProvided(CurrentStackModuleData))
 		{
-			if (ProvidedDependency.Type == ENiagaraModuleDependencyType::PreDependency)
+			if (ProvidedDependency.Type == ENiagaraModuleDependencyType::PostDependency)
 			{
-				int32 SuggestedIndex = CurrentStackModuleData.Index;
-				LeastRequirementIndex = LeastRequirementIndex < SuggestedIndex ? SuggestedIndex : LeastRequirementIndex;
-				bRequirementLowerBoundFound = true;
+				int32 SuggestedIndex = CurrentStackModuleData.Index + 1;
+				GreatestProvidedDependencyLowerBoundIdx = GreatestProvidedDependencyLowerBoundIdx < SuggestedIndex ? SuggestedIndex : GreatestProvidedDependencyLowerBoundIdx;
+				bProvidedDependencyLowerBoundFound = true;
 			}
-			else if (ProvidedDependency.Type == ENiagaraModuleDependencyType::PostDependency)
+			else if (ProvidedDependency.Type == ENiagaraModuleDependencyType::PreDependency)
 			{
-				int32 SuggestedIndex = CurrentStackModuleData.Index + 1;
-				GreatestRequirementIndex = GreatestRequirementIndex > SuggestedIndex ? SuggestedIndex : GreatestRequirementIndex;
-				bRequirementUpperBoundFound = true;
+				int32 SuggestedIndex = CurrentStackModuleData.Index;
+				LeastProvidedDependencyUpperBoundIdx = LeastProvidedDependencyUpperBoundIdx > SuggestedIndex ? SuggestedIndex : LeastProvidedDependencyUpperBoundIdx;
+				bProvidedDependencyUpperBoundFound = true;
 			}
 		}
 	}
 
-	LeastRequirementIndex = LeastDependencyIndex < LeastRequirementIndex ? LeastRequirementIndex : LeastDependencyIndex;
-	GreatestRequirementIndex = GreatestDependencyIndex > GreatestRequirementIndex ? GreatestRequirementIndex : GreatestDependencyIndex;
+	// Alias bounds without least and greatest prefixes for legibility.
+	int32& RequiredDependencyLowerBoundIdx = GreatestRequiredDependencyLowerBoundIdx;
+	int32& RequiredDependencyUpperBoundIdx = LeastRequiredDependencyUpperBoundIdx;
+	int32& ProvidedDependencyLowerBoundIdx = GreatestProvidedDependencyLowerBoundIdx;
+	int32& ProvidedDependencyUpperBoundIdx = LeastProvidedDependencyUpperBoundIdx;
+
+	// Provided dependency lower and upper bound must be inside required dependency lower and upper bound respectively.
+	ProvidedDependencyLowerBoundIdx = (bRequiredDependencyLowerBoundFound && ProvidedDependencyLowerBoundIdx < RequiredDependencyLowerBoundIdx) ? RequiredDependencyLowerBoundIdx : ProvidedDependencyLowerBoundIdx;
+	ProvidedDependencyUpperBoundIdx = (bRequiredDependencyUpperBoundFound && ProvidedDependencyUpperBoundIdx > RequiredDependencyUpperBoundIdx) ? RequiredDependencyUpperBoundIdx : ProvidedDependencyUpperBoundIdx;
 
-	// Do another validity test for requirement indices.
-	if (bRequirementLowerBoundFound)
+	// Do another validity test for provided dependency lower and upper bound indices.
+	if (bProvidedDependencyLowerBoundFound)
 	{
-		if (bRequirementUpperBoundFound && GreatestRequirementIndex < LeastRequirementIndex)
+		if (bProvidedDependencyUpperBoundFound && ProvidedDependencyUpperBoundIdx < ProvidedDependencyLowerBoundIdx)
 		{
-			// It is impossible to satisfy both target indices: do nothing.
+			// It is impossible to satisfy both target indices as the upper bound is less than the lower bound: do nothing.
 			return TargetIndex;
 		}
-		TargetIndex = LeastRequirementIndex;
+		TargetIndex = ProvidedDependencyLowerBoundIdx;
 	}
-	else if (bRequirementUpperBoundFound)
+	else if (bProvidedDependencyUpperBoundFound)
 	{
-		TargetIndex = GreatestRequirementIndex;
+		TargetIndex = ProvidedDependencyUpperBoundIdx;
 	}
 
 	return TargetIndex;
diff --git a/Engine/Plugins/FX/Niagara/Source/NiagaraEditor/Private/ViewModels/Stack/NiagaraStackModuleItem.cpp b/Engine/Plugins/FX/Niagara/Source/NiagaraEditor/Private/ViewModels/Stack/NiagaraStackModuleItem.cpp
--- a/Engine/Plugins/FX/Niagara/Source/NiagaraEditor/Private/ViewModels/Stack/NiagaraStackModuleItem.cpp
+++ b/Engine/Plugins/FX/Niagara/Source/NiagaraEditor/Private/ViewModels/Stack/NiagaraStackModuleItem.cpp
@@ -411,7 +411,10 @@ void AddModuleToFixDependencyIssue(
 	}
 
 	FScopedTransaction ScopedTransaction(LOCTEXT("AddDependencyFixTransaction", "Add a module to fix a dependency"));
-	FNiagaraStackGraphUtilities::AddScriptModuleToStack(DependencyProviderScript, *TargetOutputNode, TargetIndex.GetValue());
+	FNiagaraStackGraphUtilities::FAddScriptModuleToStackArgs AddScriptModuleToStackArgs(DependencyProviderScript, *TargetOutputNode);
+	AddScriptModuleToStackArgs.TargetIndex = TargetIndex.GetValue();
+	AddScriptModuleToStackArgs.bFixupTargetIndex = true;
+	FNiagaraStackGraphUtilities::AddScriptModuleToStack(AddScriptModuleToStackArgs);
 }
 
 void GenerateFixesForAddingDependencyProviders(
@@ -554,28 +557,6 @@ void GenerateFixesForEnablingModules(
 	}
 }
 
-bool DoesStackModuleProvideDependency(const FNiagaraStackModuleData& StackModuleData, const FNiagaraModuleDependency& SourceModuleRequiredDependency, const UNiagaraNodeOutput& SourceOutputNode)
-{
-	if (StackModuleData.ModuleNode != nullptr && StackModuleData.ModuleNode->FunctionScript != nullptr)
-	{
-		FVersionedNiagaraScriptData* ScriptData = StackModuleData.ModuleNode->FunctionScript->GetScriptData(StackModuleData.ModuleNode->SelectedScriptVersion);
-		if (ScriptData && ScriptData->ProvidedDependencies.Contains(SourceModuleRequiredDependency.Id))
-		{
-			if (SourceModuleRequiredDependency.ScriptConstraint == ENiagaraModuleDependencyScriptConstraint::AllScripts)
-			{
-				return true;
-			}
-
-			if (SourceModuleRequiredDependency.ScriptConstraint == ENiagaraModuleDependencyScriptConstraint::SameScript)
-			{
-				UNiagaraNodeOutput* OutputNode = FNiagaraStackGraphUtilities::GetEmitterOutputNodeForStackNode(*StackModuleData.ModuleNode);
-				return OutputNode != nullptr && UNiagaraScript::IsEquivalentUsage(OutputNode->GetUsage(), SourceOutputNode.GetUsage()) && OutputNode->GetUsageId() == SourceOutputNode.GetUsageId();
-			}
-		}
-	}
-	return false;
-}
-
 void GenerateDependencyIssues(
 	TSharedRef<FNiagaraSystemViewModel> SourceSystemViewModel,
 	FGuid SourceEmitterHandleId,
@@ -601,7 +582,7 @@ void GenerateDependencyIssues(
 			TArray<int32> DependencyProviderIndices;
 			for (int32 StackModuleDataIndex = 0; StackModuleDataIndex < SourceStackModuleData.Num(); StackModuleDataIndex++)
 			{
-				if (DoesStackModuleProvideDependency(SourceStackModuleData[StackModuleDataIndex], SourceRequiredDependency, SourceOutputNode))
+				if (FNiagaraStackGraphUtilities::DependencyUtilities::DoesStackModuleProvideDependency(SourceStackModuleData[StackModuleDataIndex], SourceRequiredDependency, SourceOutputNode))
 				{
 					DependencyProviderIndices.Add(StackModuleDataIndex);
 				}
diff --git a/Engine/Plugins/FX/Niagara/Source/NiagaraEditor/Private/ViewModels/Stack/NiagaraStackScriptItemGroup.cpp b/Engine/Plugins/FX/Niagara/Source/NiagaraEditor/Private/ViewModels/Stack/NiagaraStackScriptItemGroup.cpp
--- a/Engine/Plugins/FX/Niagara/Source/NiagaraEditor/Private/ViewModels/Stack/NiagaraStackScriptItemGroup.cpp
+++ b/Engine/Plugins/FX/Niagara/Source/NiagaraEditor/Private/ViewModels/Stack/NiagaraStackScriptItemGroup.cpp
@@ -290,7 +290,13 @@ public:
 		{
 			FNiagaraStackGraphUtilities::FAddScriptModuleToStackArgs AddScriptModuleToStackArgs(ScriptGroupAddAction->GetModuleAssetData(), *OutputNode);
 			AddScriptModuleToStackArgs.TargetIndex = TargetIndex;
-			AddScriptModuleToStackArgs.bFixupTargetIndex = true;
+
+			// If the user has specified a target index, do not try to fixup the target index.
+			if (TargetIndex == INDEX_NONE)
+			{
+				AddScriptModuleToStackArgs.bFixupTargetIndex = true;
+			}
+
 			NewModuleNode = FNiagaraStackGraphUtilities::AddScriptModuleToStack(AddScriptModuleToStackArgs);
 		}
 		else if (ScriptGroupAddAction->GetModuleParameterVariable().IsValid())
