- Use correct ViewRect

#rnx

diff --git a/Engine/Plugins/FX/Niagara/Source/NiagaraShader/Private/NiagaraDebugShaders.cpp b/Engine/Plugins/FX/Niagara/Source/NiagaraShader/Private/NiagaraDebugShaders.cpp
--- a/Engine/Plugins/FX/Niagara/Source/NiagaraShader/Private/NiagaraDebugShaders.cpp
+++ b/Engine/Plugins/FX/Niagara/Source/NiagaraShader/Private/NiagaraDebugShaders.cpp
@@ -201,7 +201,7 @@ void NiagaraDebugShaders::DrawDebugLines(
 		RDG_EVENT_NAME("NiagaraDrawDebugLines"),
 		PassParameters,
 		ERDGPassFlags::Raster,
-		[VertexShader, PixelShader, LineInstanceCount, PassParameters, ViewRect=View.ViewRect](FRHICommandListImmediate& RHICmdList)
+		[VertexShader, PixelShader, LineInstanceCount, PassParameters, ViewRect=View.UnscaledViewRect](FRHICommandListImmediate& RHICmdList)
 		{
 			FGraphicsPipelineStateInitializer GraphicsPSOInit;
 			RHICmdList.ApplyCachedRenderTargets(GraphicsPSOInit);
@@ -241,7 +241,7 @@ void NiagaraDebugShaders::DrawDebugLines(
 		RDG_EVENT_NAME("NiagaraDrawDebugLines"),
 		PassParameters,
 		ERDGPassFlags::Raster,
-		[VertexShader, PixelShader, ArgsBuffer, PassParameters, ViewRect=View.ViewRect](FRHICommandListImmediate& RHICmdList)
+		[VertexShader, PixelShader, ArgsBuffer, PassParameters, ViewRect=View.UnscaledViewRect](FRHICommandListImmediate& RHICmdList)
 		{
 			FGraphicsPipelineStateInitializer GraphicsPSOInit;
 			RHICmdList.ApplyCachedRenderTargets(GraphicsPSOInit);
