- Fix static analysis issue

#rnx

diff --git a/Engine/Plugins/FX/Niagara/Source/Niagara/Private/NiagaraEmitter.cpp b/Engine/Plugins/FX/Niagara/Source/Niagara/Private/NiagaraEmitter.cpp
--- a/Engine/Plugins/FX/Niagara/Source/Niagara/Private/NiagaraEmitter.cpp
+++ b/Engine/Plugins/FX/Niagara/Source/Niagara/Private/NiagaraEmitter.cpp
@@ -1308,7 +1308,7 @@ void UNiagaraEmitter::CacheFromCompiledData(const FNiagaraDataSetCompiledData* C
 
 	RebuildRendererBindings();
 
-	if ( GPUComputeScript )
+	if ( GPUComputeScript && CompiledData )
 	{
 		for (FSimulationStageMetaData& SimStageMetaData : GPUComputeScript->GetVMExecutableData().SimulationStageMetaData)
 		{
