Multiple tweaks and additions for making scalabiltiy smarter wrt culling important player FX.

- New master mode for scalability culling allowing us to pause all culling when in the FE etc.
- Adding reset for FXBudget on these mode transitions.
- Tidied up library spawn functions with an initialization struct version. Adding more parms was becoming so cumbersome.
- NiagaraComponents can now be told that they are a player effect explicitly.

diff --git a/Engine/Plugins/FX/Niagara/Source/Niagara/Private/NiagaraCommon.cpp b/Engine/Plugins/FX/Niagara/Source/Niagara/Private/NiagaraCommon.cpp
--- a/Engine/Plugins/FX/Niagara/Source/Niagara/Private/NiagaraCommon.cpp
+++ b/Engine/Plugins/FX/Niagara/Source/Niagara/Private/NiagaraCommon.cpp
@@ -1573,3 +1573,30 @@ FSynchronizeWithParameterDefinitionsArgs::FSynchronizeWithParameterDefinitionsAr
 	, AdditionalOldToNewNames()
 {
 }
+
+
+ENCPoolMethod ToNiagaraPooling(EPSCPoolMethod PoolingMethod)
+{
+	if (PoolingMethod == EPSCPoolMethod::AutoRelease)
+	{
+		return ENCPoolMethod::AutoRelease;
+	}
+	if (PoolingMethod == EPSCPoolMethod::ManualRelease)
+	{
+		return ENCPoolMethod::ManualRelease;
+	}
+	return ENCPoolMethod::None;
+}
+
+EPSCPoolMethod ToPSCPoolMethod(ENCPoolMethod PoolingMethod)
+{
+	if (PoolingMethod == ENCPoolMethod::AutoRelease)
+	{
+		return EPSCPoolMethod::AutoRelease;
+	}
+	if (PoolingMethod == ENCPoolMethod::ManualRelease)
+	{
+		return EPSCPoolMethod::ManualRelease;
+	}
+	return EPSCPoolMethod::None;
+}
\ No newline at end of file
diff --git a/Engine/Plugins/FX/Niagara/Source/Niagara/Private/NiagaraComponent.cpp b/Engine/Plugins/FX/Niagara/Source/Niagara/Private/NiagaraComponent.cpp
--- a/Engine/Plugins/FX/Niagara/Source/Niagara/Private/NiagaraComponent.cpp
+++ b/Engine/Plugins/FX/Niagara/Source/Niagara/Private/NiagaraComponent.cpp
@@ -448,6 +448,7 @@ UNiagaraComponent::UNiagaraComponent(const FObjectInitializer& ObjectInitializer
 	, bActivateShouldResetWhenReady(false)
 	, bDidAutoAttach(false)
 	, bAllowScalability(true)
+	, bForceLocalPlayerEffect(false)
 	, bIsCulledByScalability(false)
 	, bDuringUpdateContextReset(false)
 	//, bIsChangingAutoAttachment(false)
@@ -2749,7 +2750,7 @@ void UNiagaraComponent::SetUserParametersToDefaultValues()
 
 bool UNiagaraComponent::IsLocalPlayerEffect()const
 {
-	return FNiagaraWorldManager::IsComponentLocalPlayerLinked(this);
+	return GetForceLocalPlayerEffect() || FNiagaraWorldManager::IsComponentLocalPlayerLinked(this);
 }
 
 #if WITH_EDITOR
@@ -3407,6 +3408,18 @@ bool UNiagaraComponent::GetAllowScalability()const
 	return bAllowScalability;
 }
 
+
+void UNiagaraComponent::SetForceLocalPlayerEffect(bool bIsPlayerEffect)
+{
+	bForceLocalPlayerEffect = bIsPlayerEffect;
+	ResolveOwnerAllowsScalability();
+}
+
+bool UNiagaraComponent::GetForceLocalPlayerEffect()const
+{
+	return bForceLocalPlayerEffect;
+}
+
 #if WITH_EDITOR
 
 void UNiagaraComponent::PostLoadNormalizeOverrideNames()
diff --git a/Engine/Plugins/FX/Niagara/Source/Niagara/Private/NiagaraFunctionLibrary.cpp b/Engine/Plugins/FX/Niagara/Source/Niagara/Private/NiagaraFunctionLibrary.cpp
--- a/Engine/Plugins/FX/Niagara/Source/Niagara/Private/NiagaraFunctionLibrary.cpp
+++ b/Engine/Plugins/FX/Niagara/Source/Niagara/Private/NiagaraFunctionLibrary.cpp
@@ -90,30 +90,55 @@ UNiagaraComponent* CreateNiagaraSystem(UNiagaraSystem* SystemTemplate, UWorld* W
 }
 
 
+UNiagaraComponent* UNiagaraFunctionLibrary::SpawnSystemAtLocation(const UObject* WorldContextObject, UNiagaraSystem* SystemTemplate, FVector SpawnLocation, FRotator SpawnRotation, FVector Scale, bool bAutoDestroy, bool bAutoActivate, ENCPoolMethod PoolingMethod, bool bPreCullCheck)
+{
+	FFXSystemSpawnParameters SpawnParams;
+	SpawnParams.WorldContextObject = WorldContextObject;
+	SpawnParams.SystemTemplate = SystemTemplate;
+	SpawnParams.Location = SpawnLocation;
+	SpawnParams.Rotation = SpawnRotation;
+	SpawnParams.Scale = Scale;
+	SpawnParams.bAutoDestroy = bAutoDestroy;
+	SpawnParams.bAutoActivate = bAutoActivate;
+	SpawnParams.PoolingMethod = ToPSCPoolMethod(PoolingMethod);
+	SpawnParams.bPreCullCheck = bPreCullCheck;
+	return SpawnSystemAtLocationWithParams(SpawnParams);
+}
+
 /**
 * Spawns a Niagara System at the specified world location/rotation
 * @return			The spawned UNiagaraComponent
 */
-UNiagaraComponent* UNiagaraFunctionLibrary::SpawnSystemAtLocation(const UObject* WorldContextObject, UNiagaraSystem* SystemTemplate, FVector SpawnLocation, FRotator SpawnRotation, FVector Scale, bool bAutoDestroy, bool bAutoActivate, ENCPoolMethod PoolingMethod, bool bPreCullCheck)
+UNiagaraComponent* UNiagaraFunctionLibrary::SpawnSystemAtLocationWithParams(FFXSystemSpawnParameters& SpawnParams)
 {
 	UNiagaraComponent* PSC = NULL;
-	if (SystemTemplate)
+	UNiagaraSystem* NiagaraSystem = Cast<UNiagaraSystem>(SpawnParams.SystemTemplate);
+	if (NiagaraSystem)
 	{
-		UWorld* World = GEngine->GetWorldFromContextObject(WorldContextObject, EGetWorldErrorMode::LogAndReturnNull);
+		UWorld* World = GEngine->GetWorldFromContextObject(SpawnParams.WorldContextObject, EGetWorldErrorMode::LogAndReturnNull);
 		if (World != nullptr)
 		{
 			bool bShouldCull = false;
-			if (bPreCullCheck)
+			if (SpawnParams.bPreCullCheck)
 			{
-				FNiagaraWorldManager* WorldManager = FNiagaraWorldManager::Get(World);
-				bShouldCull = WorldManager->ShouldPreCull(SystemTemplate, SpawnLocation);
+				bool bIsPlayerEffect = SpawnParams.bIsPlayerEffect;
+				if (NiagaraSystem->AllowCullingForLocalPlayers() || bIsPlayerEffect == false)
+				{
+					FNiagaraWorldManager* WorldManager = FNiagaraWorldManager::Get(World);
+					bShouldCull = WorldManager->ShouldPreCull(NiagaraSystem, SpawnParams.Location);
+				}
 			}
 
 			if (!bShouldCull)
 			{
-				PSC = CreateNiagaraSystem(SystemTemplate, World, World->GetWorldSettings(), bAutoDestroy, PoolingMethod);
+				PSC = CreateNiagaraSystem(NiagaraSystem, World, World->GetWorldSettings(), SpawnParams.bAutoDestroy, ToNiagaraPooling(SpawnParams.PoolingMethod));
 				if(PSC)
 				{
+					if (SpawnParams.bIsPlayerEffect)
+					{
+						PSC->SetForceLocalPlayerEffect(true);
+					}
+
 #if WITH_EDITORONLY_DATA
 					PSC->bWaitForCompilationOnActivate = GIsAutomationTesting;
 #endif
@@ -124,9 +149,9 @@ UNiagaraComponent* UNiagaraFunctionLibrary::SpawnSystemAtLocation(const UObject*
 					}
 
 					PSC->SetAbsolute(true, true, true);
-					PSC->SetWorldLocationAndRotation(SpawnLocation, SpawnRotation);
-					PSC->SetRelativeScale3D(Scale);
-					if (bAutoActivate)
+					PSC->SetWorldLocationAndRotation(SpawnParams.Location, SpawnParams.Rotation);
+					PSC->SetRelativeScale3D(SpawnParams.Scale);
+					if (SpawnParams.bAutoActivate)
 					{
 						PSC->Activate(true);
 					}
@@ -138,9 +163,6 @@ UNiagaraComponent* UNiagaraFunctionLibrary::SpawnSystemAtLocation(const UObject*
 }
 
 
-
-
-
 /**
 * Spawns a Niagara System attached to a component
 * @return			The spawned UNiagaraComponent
@@ -150,11 +172,6 @@ UNiagaraComponent* UNiagaraFunctionLibrary::SpawnSystemAttached(UNiagaraSystem*
 	return SpawnSystemAttached(SystemTemplate, AttachToComponent, AttachPointName, Location, Rotation, FVector(1, 1, 1), LocationType, bAutoDestroy, PoolingMethod, bAutoActivate, bPreCullCheck);
 }
 
-/**
-* Spawns a Niagara System attached to a component
-* @return			The spawned UNiagaraComponent
-*/
-
 UNiagaraComponent* UNiagaraFunctionLibrary::SpawnSystemAttached(
 	UNiagaraSystem* SystemTemplate,
 	USceneComponent* AttachToComponent,
@@ -168,38 +185,67 @@ UNiagaraComponent* UNiagaraFunctionLibrary::SpawnSystemAttached(
 	bool bAutoActivate,
 	bool bPreCullCheck
 )
+{
+	FFXSystemSpawnParameters SpawnParams;
+	SpawnParams.SystemTemplate = SystemTemplate;
+	SpawnParams.AttachToComponent = AttachToComponent;
+	SpawnParams.AttachPointName = AttachPointName;
+	SpawnParams.Location = Location;
+	SpawnParams.Rotation = Rotation;
+	SpawnParams.Scale = Scale;
+	SpawnParams.LocationType = LocationType;
+	SpawnParams.bAutoDestroy = bAutoDestroy;
+	SpawnParams.bAutoActivate = bAutoActivate;
+	SpawnParams.PoolingMethod = ToPSCPoolMethod(PoolingMethod);
+	SpawnParams.bPreCullCheck = bPreCullCheck;
+	return SpawnSystemAttachedWithParams(SpawnParams);
+}
+
+/**
+* Spawns a Niagara System attached to a component
+* @return			The spawned UNiagaraComponent
+*/
+
+UNiagaraComponent* UNiagaraFunctionLibrary::SpawnSystemAttachedWithParams(FFXSystemSpawnParameters& SpawnParams)
 {
 	LLM_SCOPE(ELLMTag::Niagara);
 	
 	UNiagaraComponent* PSC = nullptr;
-	if (SystemTemplate)
+	UNiagaraSystem* NiagaraSystem = Cast<UNiagaraSystem>(SpawnParams.SystemTemplate);
+	if (NiagaraSystem)
 	{
-		if (!AttachToComponent)
+		if (!SpawnParams.AttachToComponent)
 		{
 			UE_LOG(LogScript, Warning, TEXT("UGameplayStatics::SpawnNiagaraEmitterAttached: NULL AttachComponent specified!"));
 		}
 		else
 		{
-			UWorld* const World = AttachToComponent->GetWorld();
+			UWorld* const World = SpawnParams.AttachToComponent->GetWorld();
 			if (World && !World->IsNetMode(NM_DedicatedServer))
 			{
 				bool bShouldCull = false;
-				if (bPreCullCheck)
+				if (SpawnParams.bPreCullCheck)
 				{
 					//Don't precull if this is a local player linked effect and the system doesn't allow us to cull those.
-					if (SystemTemplate->AllowCullingForLocalPlayers() || FNiagaraWorldManager::IsComponentLocalPlayerLinked(AttachToComponent) == false)
+					bool bIsPlayerEffect = SpawnParams.bIsPlayerEffect || FNiagaraWorldManager::IsComponentLocalPlayerLinked(SpawnParams.AttachToComponent);
+					if (NiagaraSystem->AllowCullingForLocalPlayers() || bIsPlayerEffect == false)
 					{
 						FNiagaraWorldManager* WorldManager = FNiagaraWorldManager::Get(World);
 						//TODO: For now using the attach parent location and ignoring the emitters relative location which is clearly going to be a bit wrong in some cases.
-						bShouldCull = WorldManager->ShouldPreCull(SystemTemplate, AttachToComponent->GetComponentLocation());
+						bShouldCull = WorldManager->ShouldPreCull(NiagaraSystem, SpawnParams.AttachToComponent->GetComponentLocation());
 					}
 				}
 
 				if (!bShouldCull)
 				{
-					PSC = CreateNiagaraSystem(SystemTemplate, World, AttachToComponent->GetOwner(), bAutoDestroy, PoolingMethod);
+					PSC = CreateNiagaraSystem(NiagaraSystem, World, SpawnParams.AttachToComponent->GetOwner(), SpawnParams.bAutoDestroy, ToNiagaraPooling(SpawnParams.PoolingMethod));
 					if (PSC)
 					{
+						if (SpawnParams.bIsPlayerEffect)
+						{
+							PSC->SetForceLocalPlayerEffect(true);
+						}
+						
 #if WITH_EDITOR
 						if (GForceNiagaraSpawnAttachedSolo > 0)
 						{
@@ -208,10 +254,10 @@ UNiagaraComponent* UNiagaraFunctionLibrary::SpawnSystemAttached(
 #endif
 						auto SetupRelativeTransforms = [&]()
 						{
-							if (LocationType == EAttachLocation::KeepWorldPosition)
+							if (SpawnParams.LocationType == EAttachLocation::KeepWorldPosition)
 							{
-								const FTransform ParentToWorld = AttachToComponent->GetSocketTransform(AttachPointName);
-								const FTransform ComponentToWorld(Rotation, Location, Scale);
+								const FTransform ParentToWorld = SpawnParams.AttachToComponent->GetSocketTransform(SpawnParams.AttachPointName);
+								const FTransform ComponentToWorld(SpawnParams.Rotation, SpawnParams.Location, SpawnParams.Scale);
 								const FTransform RelativeTM = ComponentToWorld.GetRelativeTransform(ParentToWorld);
 								PSC->SetRelativeLocation_Direct(RelativeTM.GetLocation());
 								PSC->SetRelativeRotation_Direct(RelativeTM.GetRotation().Rotator());
@@ -219,19 +265,19 @@ UNiagaraComponent* UNiagaraFunctionLibrary::SpawnSystemAttached(
 							}
 							else
 							{
-								PSC->SetRelativeLocation_Direct(Location);
-								PSC->SetRelativeRotation_Direct(Rotation);
+								PSC->SetRelativeLocation_Direct(SpawnParams.Location);
+								PSC->SetRelativeRotation_Direct(SpawnParams.Rotation);
 
-								if (LocationType == EAttachLocation::SnapToTarget)
+								if (SpawnParams.LocationType == EAttachLocation::SnapToTarget)
 								{
 									// SnapToTarget indicates we "keep world scale", this indicates we we want the inverse of the parent-to-world scale 
 									// to calculate world scale at Scale 1, and then apply the passed in Scale
-									const FTransform ParentToWorld = AttachToComponent->GetSocketTransform(AttachPointName);
-									PSC->SetRelativeScale3D_Direct(Scale * ParentToWorld.GetSafeScaleReciprocal(ParentToWorld.GetScale3D()));
+									const FTransform ParentToWorld = SpawnParams.AttachToComponent->GetSocketTransform(SpawnParams.AttachPointName);
+									PSC->SetRelativeScale3D_Direct(SpawnParams.Scale * ParentToWorld.GetSafeScaleReciprocal(ParentToWorld.GetScale3D()));
 								}
 								else
 								{
-									PSC->SetRelativeScale3D_Direct(Scale);
+									PSC->SetRelativeScale3D_Direct(SpawnParams.Scale);
 								}
 							}
 						};
@@ -240,16 +286,16 @@ UNiagaraComponent* UNiagaraFunctionLibrary::SpawnSystemAttached(
 						{
 							//It is now possible for us to be already regisetered so we must use AttachToComponent() instead.
 							SetupRelativeTransforms();
-							PSC->AttachToComponent(AttachToComponent, FAttachmentTransformRules::KeepRelativeTransform, AttachPointName);
+							PSC->AttachToComponent(SpawnParams.AttachToComponent, FAttachmentTransformRules::KeepRelativeTransform, SpawnParams.AttachPointName);
 						}
 						else
 						{
-							PSC->SetupAttachment(AttachToComponent, AttachPointName);
+							PSC->SetupAttachment(SpawnParams.AttachToComponent, SpawnParams.AttachPointName);
 							SetupRelativeTransforms();
 							PSC->RegisterComponentWithWorld(World);
 						}
 
-						if (bAutoActivate)
+						if (SpawnParams.bAutoActivate)
 						{
 							PSC->Activate(true);
 						}
diff --git a/Engine/Plugins/FX/Niagara/Source/Niagara/Private/NiagaraScalabilityManager.cpp b/Engine/Plugins/FX/Niagara/Source/Niagara/Private/NiagaraScalabilityManager.cpp
--- a/Engine/Plugins/FX/Niagara/Source/Niagara/Private/NiagaraScalabilityManager.cpp
+++ b/Engine/Plugins/FX/Niagara/Source/Niagara/Private/NiagaraScalabilityManager.cpp
@@ -463,11 +463,16 @@ void FNiagaraScalabilityManager::Update(FNiagaraWorldManager* WorldMan, float De
 		}
 	}
 
-	//Paranoia code in case the EffectType is GCd from under us.
-	if (EffectType == nullptr)
+	bool bShutdown = EffectType == nullptr || FNiagaraWorldManager::GetScalabilityCullingMode() == ENiagaraScalabilityCullingMode::Disabled;
+	if (bShutdown)
 	{
-		ManagedComponents.Empty();
-		State.Empty();
+		while (ManagedComponents.Num())
+		{
+			ManagedComponents.Last()->UnregisterWithScalabilityManager();
+		}
+
+		check(ManagedComponents.Num() == 0);
+		check(State.Num() == 0);
 		LastUpdateTime = 0.0f;
 		return;
 	}
diff --git a/Engine/Plugins/FX/Niagara/Source/Niagara/Private/NiagaraWorldManager.cpp b/Engine/Plugins/FX/Niagara/Source/Niagara/Private/NiagaraWorldManager.cpp
--- a/Engine/Plugins/FX/Niagara/Source/Niagara/Private/NiagaraWorldManager.cpp
+++ b/Engine/Plugins/FX/Niagara/Source/Niagara/Private/NiagaraWorldManager.cpp
@@ -240,6 +240,25 @@ FName FNiagaraWorldManagerTickFunction::DiagnosticContext(bool bDetailed)
 
 //////////////////////////////////////////////////////////////////////////
 
+ENiagaraScalabilityCullingMode FNiagaraWorldManager::ScalabilityCullingMode = ENiagaraScalabilityCullingMode::Enabled;
+
+void FNiagaraWorldManager::SetScalabilityCullingMode(ENiagaraScalabilityCullingMode NewMode) 
+{
+	//also reset fx budgets on state change.
+	if (ScalabilityCullingMode != NewMode)
+	{
+		FFXBudget::Reset();
+		ScalabilityCullingMode = NewMode;
+
+		auto UpdateScalability = [](FNiagaraWorldManager& WorldMan)
+		{
+			WorldMan.UpdateScalabilityManagers(0.01f, false);
+		};
+		ForAllWorldManagers(UpdateScalability);
+	}
+}
+
+
 FNiagaraWorldManager::FNiagaraWorldManager()
 	: World(nullptr)
 	, ActiveNiagaraTickGroup(-1)
@@ -1153,6 +1172,11 @@ void FNiagaraWorldManager::UpdateScalabilityManagers(float DeltaSeconds, bool bN
 
 void FNiagaraWorldManager::RegisterWithScalabilityManager(UNiagaraComponent* Component)
 {
+	if (GetScalabilityCullingMode() == ENiagaraScalabilityCullingMode::Disabled)
+	{
+		return;
+	}
+
 	check(Component);
 	if (UNiagaraEffectType* EffectType = Component->GetAsset()->GetEffectType())
 	{
@@ -1188,7 +1212,7 @@ void FNiagaraWorldManager::UnregisterWithScalabilityManager(UNiagaraComponent* C
 
 bool FNiagaraWorldManager::ShouldPreCull(UNiagaraSystem* System, UNiagaraComponent* Component)
 {
-	if (System)
+	if (GetScalabilityCullingMode() == ENiagaraScalabilityCullingMode::Enabled && System)
 	{
 		if (UNiagaraEffectType* EffectType = System->GetEffectType())
 		{
@@ -1206,7 +1230,7 @@ bool FNiagaraWorldManager::ShouldPreCull(UNiagaraSystem* System, UNiagaraCompone
 
 bool FNiagaraWorldManager::ShouldPreCull(UNiagaraSystem* System, FVector Location)
 {
-	if (System)
+	if (GetScalabilityCullingMode() == ENiagaraScalabilityCullingMode::Enabled && System)
 	{
 		if (UNiagaraEffectType* EffectType = System->GetEffectType())
 		{
@@ -1225,6 +1249,11 @@ bool FNiagaraWorldManager::ShouldPreCull(UNiagaraSystem* System, FVector Locatio
 
 void FNiagaraWorldManager::CalculateScalabilityState(UNiagaraSystem* System, const FNiagaraSystemScalabilitySettings& ScalabilitySettings, UNiagaraEffectType* EffectType, FVector Location, bool bIsPreCull, float WorstGlobalBudgetUse, FNiagaraScalabilityState& OutState)
 {
+	if (GetScalabilityCullingMode() == ENiagaraScalabilityCullingMode::Disabled)
+	{
+		return;
+	}
+
 	OutState.bCulled = false;
 
 	DistanceCull(EffectType, ScalabilitySettings, Location, OutState);
@@ -1247,6 +1276,11 @@ void FNiagaraWorldManager::CalculateScalabilityState(UNiagaraSystem* System, con
 
 void FNiagaraWorldManager::CalculateScalabilityState(UNiagaraSystem* System, const FNiagaraSystemScalabilitySettings& ScalabilitySettings, UNiagaraEffectType* EffectType, UNiagaraComponent* Component, bool bIsPreCull, float WorstGlobalBudgetUse, FNiagaraScalabilityState& OutState)
 {
+	if (GetScalabilityCullingMode() == ENiagaraScalabilityCullingMode::Disabled)
+	{
+		return;
+	}
+
 	OutState.bCulled = false;
 
 	DistanceCull(EffectType, ScalabilitySettings, Component, OutState);
@@ -1282,7 +1316,7 @@ void FNiagaraWorldManager::SortedSignificanceCull(UNiagaraEffectType* EffectType
 	//Cull all but the N most significance FX.
 	bool bCull = false;
 	
-	if(GEnableNiagaraInstanceCountCulling)
+	if(GetScalabilityCullingMode() == ENiagaraScalabilityCullingMode::Enabled && GEnableNiagaraInstanceCountCulling)
 	{
 		int32 SystemInstanceMax = ScalabilitySettings.MaxSystemInstances;
 		int32 EffectTypeInstanceMax = ScalabilitySettings.MaxInstances;
@@ -1337,6 +1371,11 @@ void FNiagaraWorldManager::SortedSignificanceCull(UNiagaraEffectType* EffectType
 
 void FNiagaraWorldManager::VisibilityCull(UNiagaraEffectType* EffectType, const FNiagaraSystemScalabilitySettings& ScalabilitySettings, UNiagaraComponent* Component, FNiagaraScalabilityState& OutState)
 {
+	if (GetScalabilityCullingMode() != ENiagaraScalabilityCullingMode::Enabled)
+	{
+		return;
+	}
+
 	float TimeSinceRendered = Component->GetSafeTimeSinceRendered(World->TimeSeconds);
 	bool bCull = Component->GetLastRenderTime() >= 0.0f && ScalabilitySettings.bCullByMaxTimeWithoutRender && TimeSinceRendered > ScalabilitySettings.MaxTimeWithoutRender;
 
@@ -1348,6 +1387,11 @@ void FNiagaraWorldManager::VisibilityCull(UNiagaraEffectType* EffectType, const
 
 void FNiagaraWorldManager::InstanceCountCull(UNiagaraEffectType* EffectType, UNiagaraSystem* System, const FNiagaraSystemScalabilitySettings& ScalabilitySettings, FNiagaraScalabilityState& OutState)
 {
+	if (GetScalabilityCullingMode() != ENiagaraScalabilityCullingMode::Enabled)
+	{
+		return;
+	}
+
 	int32 SystemInstanceMax = ScalabilitySettings.MaxSystemInstances;
 	int32 EffectTypeInstanceMax = ScalabilitySettings.MaxInstances;
 
@@ -1415,7 +1459,7 @@ void FNiagaraWorldManager::DistanceCull(UNiagaraEffectType* EffectType, const FN
 	float MaxDist = ScalabilitySettings.MaxDistance;
 	Component->SetLODDistance(LODDistance, FMath::Max(MaxDist, 1.0f));
 
-	if (GEnableNiagaraDistanceCulling && ScalabilitySettings.bCullByDistance)
+	if (GetScalabilityCullingMode() == ENiagaraScalabilityCullingMode::Enabled && GEnableNiagaraDistanceCulling && ScalabilitySettings.bCullByDistance)
 	{
 		bool bCull = LODDistance > MaxDist;
 		OutState.bCulled |= bCull;
@@ -1451,7 +1495,7 @@ void FNiagaraWorldManager::DistanceCull(UNiagaraEffectType* EffectType, const FN
 
 void FNiagaraWorldManager::DistanceCull(UNiagaraEffectType* EffectType, const FNiagaraSystemScalabilitySettings& ScalabilitySettings, FVector Location, FNiagaraScalabilityState& OutState)
 {
-	if (bCachedPlayerViewLocationsValid)
+	if (GetScalabilityCullingMode() == ENiagaraScalabilityCullingMode::Enabled && bCachedPlayerViewLocationsValid)
 	{
 		float ClosestDistSq = FLT_MAX;
 		for (FVector ViewLocation : CachedPlayerViewLocations)
@@ -1497,7 +1541,7 @@ void FNiagaraWorldManager::DistanceCull(UNiagaraEffectType* EffectType, const FN
 
 void FNiagaraWorldManager::GlobalBudgetCull(const FNiagaraSystemScalabilitySettings& ScalabilitySettings, float WorstGlobalBudgetUse, FNiagaraScalabilityState& OutState)
 {
- 	bool bCull = WorstGlobalBudgetUse >= ScalabilitySettings.BudgetScaling.MaxGlobalBudgetUsage;
+ 	bool bCull = GetScalabilityCullingMode() == ENiagaraScalabilityCullingMode::Enabled && WorstGlobalBudgetUse >= ScalabilitySettings.BudgetScaling.MaxGlobalBudgetUsage;
  	OutState.bCulled |= bCull;
 #if DEBUG_SCALABILITY_STATE
 	OutState.bCulledByGlobalBudget |= bCull;
@@ -1591,7 +1635,7 @@ bool FNiagaraWorldManager::IsComponentLocalPlayerLinked(const USceneComponent* I
 	//Is our owner or instigator a locally viewed pawn?
 	if (const AActor* Owner = InComponent->GetOwner())
 	{
-		if (const APawn* OwnerPawn = Cast<const APawn>(InComponent))
+		if (const APawn* OwnerPawn = Cast<const APawn>(Owner))
 		{
 			if (OwnerPawn->IsLocallyViewed())
 			{
diff --git a/Engine/Plugins/FX/Niagara/Source/Niagara/Public/NiagaraCommon.h b/Engine/Plugins/FX/Niagara/Source/Niagara/Public/NiagaraCommon.h
--- a/Engine/Plugins/FX/Niagara/Source/Niagara/Public/NiagaraCommon.h
+++ b/Engine/Plugins/FX/Niagara/Source/Niagara/Public/NiagaraCommon.h
@@ -1514,6 +1514,9 @@ enum class ENCPoolMethod : uint8
 	FreeInPool UMETA(Hidden),
 };
 
+NIAGARA_API ENCPoolMethod ToNiagaraPooling(EPSCPoolMethod PoolingMethod);
+NIAGARA_API EPSCPoolMethod ToPSCPoolMethod(ENCPoolMethod PoolingMethod);
+
 UENUM()
 enum class ENiagaraSystemInstanceState : uint8
 {
diff --git a/Engine/Plugins/FX/Niagara/Source/Niagara/Public/NiagaraComponent.h b/Engine/Plugins/FX/Niagara/Source/Niagara/Public/NiagaraComponent.h
--- a/Engine/Plugins/FX/Niagara/Source/Niagara/Public/NiagaraComponent.h
+++ b/Engine/Plugins/FX/Niagara/Source/Niagara/Public/NiagaraComponent.h
@@ -702,6 +702,12 @@ public:
 	
 	UFUNCTION(BlueprintGetter)
 	bool GetAllowScalability()const;
+	
+	UFUNCTION(BlueprintSetter, Category = Scalability, meta = (Keywords = "LOD scalability"))
+	void SetForceLocalPlayerEffect(bool bIsPlayerEffect);
+	
+	UFUNCTION(BlueprintGetter)
+	bool GetForceLocalPlayerEffect()const;
 
 	FORCEINLINE bool IsRegisteredWithScalabilityManager()const { return ScalabilityManagerHandle != INDEX_NONE; }
 	FORCEINLINE int32 GetScalabilityManagerHandle()const { return ScalabilityManagerHandle; }
@@ -731,6 +737,10 @@ private:
 	UPROPERTY(EditAnywhere, Category=Niagara, BlueprintGetter=GetAllowScalability, BlueprintSetter=SetAllowScalability)
 	uint32 bAllowScalability : 1;
 
+	/** Flag allowing us to force this Effect to be considered a LocalPlayer Effect. This can prevent it being scalabiltiy culled in some situations etc. */
+	//UPROPERTY(EditAnywhere, Category = Niagara, BlueprintGetter = GetForceLocalPlayerEffect, BlueprintSetter = SetForceLocalPlayerEffect)
+	uint32 bForceLocalPlayerEffect : 1;
+
 	/** Whether the owner of this component allows it to be scalability culled. We optionally do not allow scalability culling for components on or attached to the local player. */
 	uint32 bOwnerAllowsScalabiltiy : 1;
 
diff --git a/Engine/Plugins/FX/Niagara/Source/Niagara/Public/NiagaraFunctionLibrary.h b/Engine/Plugins/FX/Niagara/Source/Niagara/Public/NiagaraFunctionLibrary.h
--- a/Engine/Plugins/FX/Niagara/Source/Niagara/Public/NiagaraFunctionLibrary.h
+++ b/Engine/Plugins/FX/Niagara/Source/Niagara/Public/NiagaraFunctionLibrary.h
@@ -9,10 +9,11 @@
 #include "NiagaraComponentPool.h"
 #include "NiagaraCommon.h"
 #include "VectorVM.h"
+#include "NiagaraSystem.h"
+#include "Particles/ParticleSystemComponent.h"
 #include "NiagaraFunctionLibrary.generated.h"
 
 class UNiagaraComponent;
-class UNiagaraSystem;
 class USceneComponent;
 class UVolumeTexture;
 
@@ -26,6 +27,13 @@ class NIAGARA_API UNiagaraFunctionLibrary : public UBlueprintFunctionLibrary
 	GENERATED_UCLASS_BODY()
 
 public:
+
+	UFUNCTION(BlueprintCallable, Category = Niagara, meta = (Keywords = "niagara System", WorldContext = "WorldContextObject", UnsafeDuringActorConstruction = "true"))
+	static UNiagaraComponent* SpawnSystemAtLocationWithParams(FFXSystemSpawnParameters& SpawnParams);
+
+	UFUNCTION(BlueprintCallable, Category = Niagara, meta = (Keywords = "niagara System", UnsafeDuringActorConstruction = "true"))
+	static UNiagaraComponent* SpawnSystemAttachedWithParams(FFXSystemSpawnParameters& SpawnParams);
+
 	/**
 	* Spawns a Niagara System at the specified world location/rotation
 	* @return			The spawned UNiagaraComponent
diff --git a/Engine/Plugins/FX/Niagara/Source/Niagara/Public/NiagaraWorldManager.h b/Engine/Plugins/FX/Niagara/Source/Niagara/Public/NiagaraWorldManager.h
--- a/Engine/Plugins/FX/Niagara/Source/Niagara/Public/NiagaraWorldManager.h
+++ b/Engine/Plugins/FX/Niagara/Source/Niagara/Public/NiagaraWorldManager.h
@@ -55,6 +55,16 @@ struct TStructOpsTypeTraits<FNiagaraWorldManagerTickFunction> : public TStructOp
 
 using FNiagaraSystemSimulationRef = TSharedRef<FNiagaraSystemSimulation, ESPMode::ThreadSafe>;
 
+enum class NIAGARA_API ENiagaraScalabilityCullingMode : uint8
+{
+	/* Scalability culling is enabled as normal. */
+	Enabled,
+	/* No scalability culling is done but FX stay registered with the managers etc for when it is resumed. */
+	Paused,
+	/* Scalability is disabled entirely and all tracking is dropped. */
+	Disabled,
+};
+
 /**
 * Manager class for any data relating to a particular world.
 */
@@ -192,6 +202,9 @@ public:
 
 	static void OnRefreshOwnerAllowsScalability();
 
+	NIAGARA_API static void SetScalabilityCullingMode(ENiagaraScalabilityCullingMode NewMode);
+	NIAGARA_API static ENiagaraScalabilityCullingMode GetScalabilityCullingMode() { return ScalabilityCullingMode; }
+
 private:
 	// Callback function registered with global world delegates to instantiate world manager when a game world is created
 	static void OnWorldInit(UWorld* World, const UWorld::InitializationValues IVS);
@@ -296,6 +309,9 @@ private:
 	TUniquePtr<class FNiagaraDebugHud> NiagaraDebugHud;
 
 	TMap<UNiagaraSystem*, UNiagaraCullProxyComponent*> CullProxyMap;
+
+	/** A global flag for all scalability culling */
+	static ENiagaraScalabilityCullingMode ScalabilityCullingMode;
 };
 
 
diff --git a/Engine/Source/Runtime/Engine/Classes/Particles/FXBudget.h b/Engine/Source/Runtime/Engine/Classes/Particles/FXBudget.h
--- a/Engine/Source/Runtime/Engine/Classes/Particles/FXBudget.h
+++ b/Engine/Source/Runtime/Engine/Classes/Particles/FXBudget.h
@@ -41,6 +41,8 @@ public:
 	FORCEINLINE static float GetWorstAdjustedUsage() { return WorstAdjustedUsage; }
 	FORCEINLINE static void SetWorstAdjustedUsage(float NewAdjustedUsage){ WorstAdjustedUsage = NewAdjustedUsage; }
 
+	static void Reset();
+
 	static TSharedPtr<FParticlePerfStatsListener_FXBudget, ESPMode::ThreadSafe> StatsListener;
 
 	static void OnEnabledCVarChanged(IConsoleVariable* CVar);
@@ -66,5 +68,7 @@ public:
 	FORCEINLINE static float GetWorstAdjustedUsage() { return 0.0f; }
 	FORCEINLINE static bool Enabled() { return false; }
 	FORCEINLINE static void SetEnabled(bool bInEnabled) { }
+
+	FORCEINLINE static void Reset(){}
 };
 #endif
\ No newline at end of file
diff --git a/Engine/Source/Runtime/Engine/Classes/Particles/ParticleSystemComponent.h b/Engine/Source/Runtime/Engine/Classes/Particles/ParticleSystemComponent.h
--- a/Engine/Source/Runtime/Engine/Classes/Particles/ParticleSystemComponent.h
+++ b/Engine/Source/Runtime/Engine/Classes/Particles/ParticleSystemComponent.h
@@ -328,6 +328,29 @@ struct FParticleEventKismetData : public FParticleEventData
 {
 };
 
+/** Parameters controlling the spawning behavior of FX systems via the SpawnSystemAtLocation and SpawnSystemAttached. */
+USTRUCT(BlueprintType)
+struct FFXSystemSpawnParameters
+{
+	GENERATED_BODY()
+
+	const UObject* WorldContextObject = nullptr;
+	UFXSystemAsset* SystemTemplate = nullptr;
+	FVector Location = FVector::ZeroVector;
+	FRotator Rotation = FRotator::ZeroRotator;
+	FVector Scale = FVector(1.f);
+
+	USceneComponent* AttachToComponent = nullptr;
+	FName AttachPointName = NAME_None;
+	EAttachLocation::Type LocationType = EAttachLocation::KeepWorldPosition;
+
+	bool bAutoDestroy = true;
+	bool bAutoActivate = true;
+	EPSCPoolMethod PoolingMethod = EPSCPoolMethod::None;
+	bool bPreCullCheck = true;
+	bool bIsPlayerEffect = false;
+};
+
 UCLASS(Abstract)
 class ENGINE_API UFXSystemComponent : public UPrimitiveComponent
 {
diff --git a/Engine/Source/Runtime/Engine/Private/Particles/FXBudget.cpp b/Engine/Source/Runtime/Engine/Private/Particles/FXBudget.cpp
--- a/Engine/Source/Runtime/Engine/Private/Particles/FXBudget.cpp
+++ b/Engine/Source/Runtime/Engine/Private/Particles/FXBudget.cpp
@@ -328,6 +328,19 @@ FFXTimeData FFXBudget::GetAdjustedUsage()
 	}
 }
 
+void FFXBudget::Reset()
+{
+	if (StatsListener.IsValid())
+	{
+		if (StatsListener.IsValid())
+		{
+			FParticlePerfStatsManager::RemoveListener(StatsListener);
+			StatsListener = MakeShared<FParticlePerfStatsListener_FXBudget, ESPMode::ThreadSafe>();
+			FParticlePerfStatsManager::AddListener(StatsListener);
+		}
+	}
+}
+
 void FFXBudget::OnEnabledCVarChanged(IConsoleVariable* CVar)
 {
 	OnEnabledChangedInternal();
