Compile fix.

#rnx

diff --git a/Engine/Plugins/FX/Niagara/Source/Niagara/Public/NiagaraCommon.h b/Engine/Plugins/FX/Niagara/Source/Niagara/Public/NiagaraCommon.h
--- a/Engine/Plugins/FX/Niagara/Source/Niagara/Public/NiagaraCommon.h
+++ b/Engine/Plugins/FX/Niagara/Source/Niagara/Public/NiagaraCommon.h
@@ -1514,8 +1514,8 @@ enum class ENCPoolMethod : uint8
 	FreeInPool UMETA(Hidden),
 };
 
-NIAGARA_API ENCPoolMethod ToNiagaraPooling(EPSCPoolMethod PoolingMethod);
-NIAGARA_API EPSCPoolMethod ToPSCPoolMethod(ENCPoolMethod PoolingMethod);
+extern NIAGARA_API ENCPoolMethod ToNiagaraPooling(EPSCPoolMethod PoolingMethod);
+extern NIAGARA_API EPSCPoolMethod ToPSCPoolMethod(ENCPoolMethod PoolingMethod);
 
 UENUM()
 enum class ENiagaraSystemInstanceState : uint8
