Initialize new UNiagaraScriptVariables added to graphs with default mode Value when the variable is in the module namespace.

diff --git a/Engine/Plugins/FX/Niagara/Source/NiagaraEditor/Private/NiagaraGraph.cpp b/Engine/Plugins/FX/Niagara/Source/NiagaraEditor/Private/NiagaraGraph.cpp
--- a/Engine/Plugins/FX/Niagara/Source/NiagaraEditor/Private/NiagaraGraph.cpp
+++ b/Engine/Plugins/FX/Niagara/Source/NiagaraEditor/Private/NiagaraGraph.cpp
@@ -1540,7 +1540,17 @@ UNiagaraScriptVariable* UNiagaraGraph::AddParameter(const FNiagaraVariable& Para
 		UNiagaraScriptVariable* NewScriptVariable = NewObject<UNiagaraScriptVariable>(this, FName(), RF_Transactional);
 		NewScriptVariable->Init(Parameter, ParameterMetaData);
 		NewScriptVariable->SetIsStaticSwitch(bIsStaticSwitch);
-		NewScriptVariable->DefaultMode = ENiagaraDefaultMode::FailIfPreviouslyNotSet;
+
+		// Inputs in graphs (module namespace) are intrinsically written to, so set the default mode to value instead of fail if not set.
+		if (NewScriptVariable->Variable.IsInNameSpace(FNiagaraConstants::ModuleNamespace))
+		{
+			NewScriptVariable->DefaultMode = ENiagaraDefaultMode::Value;
+		}
+		else
+		{ 
+			NewScriptVariable->DefaultMode = ENiagaraDefaultMode::FailIfPreviouslyNotSet;
+		}
+
 		VariableToScriptVariable.Add(Parameter, NewScriptVariable);
 		if (bNotifyChanged)
 		{
