Fixing static variables in the particles namespace.

diff --git a/Engine/Plugins/FX/Niagara/Source/NiagaraEditor/Private/NiagaraParameterMapHistory.cpp b/Engine/Plugins/FX/Niagara/Source/NiagaraEditor/Private/NiagaraParameterMapHistory.cpp
--- a/Engine/Plugins/FX/Niagara/Source/NiagaraEditor/Private/NiagaraParameterMapHistory.cpp
+++ b/Engine/Plugins/FX/Niagara/Source/NiagaraEditor/Private/NiagaraParameterMapHistory.cpp
@@ -1587,6 +1587,7 @@ int32 FNiagaraParameterMapHistoryBuilder::HandleVariableWrite(int32 ParameterMap
 	return FoundIdx;
 }
 
+
 int32 FNiagaraParameterMapHistoryBuilder::HandleVariableRead(int32 ParamMapIdx, const UEdGraphPin* InPin, bool RegisterReadsAsVariables, const UEdGraphPin* InDefaultPin, bool bFilterForCompilation, bool& OutUsedDefault)
 {
 	FString DefaultValue;
@@ -1733,7 +1734,7 @@ int32 FNiagaraParameterMapHistoryBuilder::HandleVariableRead(int32 ParamMapIdx,
 			
 
 		}
-		else if (FNiagaraParameterMapHistory::IsExternalConstantNamespace(Var, Usage, 0) || Var.IsInNameSpace(EmitterNamespace) || Var.IsInNameSpace(FNiagaraConstants::SystemNamespace))
+		else if (FNiagaraParameterMapHistory::IsExternalConstantNamespace(Var, Usage, 0) || Var.IsInNameSpace(EmitterNamespace) || Var.IsInNameSpace(FNiagaraConstants::SystemNamespace) || Var.IsInNameSpace(FNiagaraConstants::ParticleAttributeNamespace))
 		{
 			const FNiagaraVariable& VarRapid = Histories[ParamMapIdx].Variables[FoundIdx];
 			int32 StaticIdx = StaticVariables.IndexOfByPredicate([&](const FNiagaraVariable& InObj) -> bool
@@ -1960,6 +1961,17 @@ void FNiagaraParameterMapHistoryBuilder::SetConstantByStaticVariable(int32& OutV
 	}
 }
 
+int32 FNiagaraParameterMapHistoryBuilder::FindStaticVariable(const FNiagaraVariable& Var) const
+{
+	FString Value = Var.GetName().ToString();
+	int32 FoundOverrideIdx = StaticVariables.IndexOfByPredicate([&](const FNiagaraVariable& InObj) -> bool
+		{
+			return (InObj.GetName() == *Value);
+		});;
+
+	return FoundOverrideIdx;
+}
+
 void FNiagaraParameterMapHistoryBuilder::SetConstantByStaticVariable(int32& OutValue, const FNiagaraVariable& Var)
 {
 	
diff --git a/Engine/Plugins/FX/Niagara/Source/NiagaraEditor/Public/NiagaraParameterMapHistory.h b/Engine/Plugins/FX/Niagara/Source/NiagaraEditor/Public/NiagaraParameterMapHistory.h
--- a/Engine/Plugins/FX/Niagara/Source/NiagaraEditor/Public/NiagaraParameterMapHistory.h
+++ b/Engine/Plugins/FX/Niagara/Source/NiagaraEditor/Public/NiagaraParameterMapHistory.h
@@ -506,6 +506,7 @@ public:
 
 	void SetConstantByStaticVariable(int32& OutValue, const UEdGraphPin* InDefaultPin);
 	void SetConstantByStaticVariable(int32& OutValue, const FNiagaraVariable& Var);
+	int32 FindStaticVariable(const FNiagaraVariable& Var) const;
 
 
 	TArray<FNiagaraVariable> StaticVariables;
