Correct fix for UE-134040

diff --git a/Engine/Plugins/FX/Niagara/Source/Niagara/Private/NiagaraComponent.cpp b/Engine/Plugins/FX/Niagara/Source/Niagara/Private/NiagaraComponent.cpp
--- a/Engine/Plugins/FX/Niagara/Source/Niagara/Private/NiagaraComponent.cpp
+++ b/Engine/Plugins/FX/Niagara/Source/Niagara/Private/NiagaraComponent.cpp
@@ -1167,7 +1167,13 @@ void UNiagaraComponent::ActivateInternal(bool bReset /* = false */, bool bIsScal
 
 	// NOTE: This call can cause SystemInstance itself to get destroyed with auto destroy systems
 	SystemInstanceController->Activate(ResetMode);
-	
+
+	if (SystemInstanceController->IsComplete())
+	{
+		OnSystemComplete(true);
+		return;
+	}
+
 	if (SystemInstanceController && SystemInstanceController->IsSolo())
 	{
 		const ETickingGroup SoloTickGroup = SystemInstanceController->CalculateTickGroup();
