- When setting an array from BP we need to leave it in the modified state otherwise get / set will get confused

#rnx

diff --git a/Engine/Plugins/FX/Niagara/Source/Niagara/Classes/NiagaraDataInterfaceArrayImpl.h b/Engine/Plugins/FX/Niagara/Source/Niagara/Classes/NiagaraDataInterfaceArrayImpl.h
--- a/Engine/Plugins/FX/Niagara/Source/Niagara/Classes/NiagaraDataInterfaceArrayImpl.h
+++ b/Engine/Plugins/FX/Niagara/Source/Niagara/Classes/NiagaraDataInterfaceArrayImpl.h
@@ -458,10 +458,9 @@ struct FNDIArrayProxyImpl : public INDIArrayProxyBase
 			ensure(PerInstanceData_GameThread.Num() == 1);
 			FNDIArrayInstanceData_GameThread<TArrayType>* InstanceData = PerInstanceData_GameThread.CreateConstIterator().Value();
 			FWriteScopeLock	ScopeLock(InstanceData->ArrayRWGuard);
-			InstanceData->bIsModified = false;
+			InstanceData->bIsModified = true;
 			InstanceData->bIsRenderDirty = bShouldSyncToGpu;
-			InstanceData->ArrayData.Empty();
-			Owner->GetArrayReference() = InArrayData;
+			InstanceData->ArrayData = InArrayData;
 		}
 	}
 
