Fix for crash deleting renderer

#rnx

diff --git a/Engine/Plugins/FX/Niagara/Source/NiagaraEditor/Private/ViewModels/NiagaraSystemViewModel.cpp b/Engine/Plugins/FX/Niagara/Source/NiagaraEditor/Private/ViewModels/NiagaraSystemViewModel.cpp
--- a/Engine/Plugins/FX/Niagara/Source/NiagaraEditor/Private/ViewModels/NiagaraSystemViewModel.cpp
+++ b/Engine/Plugins/FX/Niagara/Source/NiagaraEditor/Private/ViewModels/NiagaraSystemViewModel.cpp
@@ -1543,9 +1543,6 @@ void FNiagaraSystemViewModel::ResetSystem(ETimeResetMode TimeResetMode, EMultiRe
 
 	FNiagaraSystemUpdateContext UpdateContext;
 
-	//TODO: Some path through the system sim init code is causing this to break running systems.
-	UpdateContext.SetDestroySystemSim(false);
-
 	UpdateContext.GetPostWork().BindLambda(
 		[ReinitMode, bResetAge](UNiagaraComponent* Component)
 		{
