- Temporary fix for DI BakeOut / Trimming / Debug remove when cooking shaders

#rnx

diff --git a/Engine/Plugins/FX/Niagara/Source/Niagara/Private/NiagaraSystem.cpp b/Engine/Plugins/FX/Niagara/Source/Niagara/Private/NiagaraSystem.cpp
--- a/Engine/Plugins/FX/Niagara/Source/Niagara/Private/NiagaraSystem.cpp
+++ b/Engine/Plugins/FX/Niagara/Source/Niagara/Private/NiagaraSystem.cpp
@@ -1031,7 +1031,7 @@ void UNiagaraSystem::Serialize(FArchive& Ar)
 	}
 
 #if WITH_EDITOR
-	if (GIsCookerLoadingPackage && Ar.IsLoading())
+	if ((IsRunningCookCommandlet() || GIsCookerLoadingPackage) && Ar.IsLoading())
 	{
 		bBakeOutRapidIteration = bBakeOutRapidIteration || bBakeOutRapidIterationOnCook;
 		bTrimAttributes = bTrimAttributes || bTrimAttributesOnCook;
diff --git a/Engine/Plugins/FX/Niagara/Source/NiagaraCore/Private/NiagaraCustomVersion.cpp b/Engine/Plugins/FX/Niagara/Source/NiagaraCore/Private/NiagaraCustomVersion.cpp
--- a/Engine/Plugins/FX/Niagara/Source/NiagaraCore/Private/NiagaraCustomVersion.cpp
+++ b/Engine/Plugins/FX/Niagara/Source/NiagaraCore/Private/NiagaraCustomVersion.cpp
@@ -8,6 +8,6 @@ const FGuid FNiagaraCustomVersion::GUID(0xFCF57AFA, 0x50764283, 0xB9A9E658, 0xFF
 // Register the custom version with core
 FCustomVersionRegistration GRegisterNiagaraCustomVersion(FNiagaraCustomVersion::GUID, FNiagaraCustomVersion::LatestVersion, TEXT("NiagaraVer"));
 
-const FGuid FNiagaraCustomVersion::LatestScriptCompileVersion(0xDB4455C6, 0x96344711, 0xD5A3076, 0x9F35A1B5);
+const FGuid FNiagaraCustomVersion::LatestScriptCompileVersion(0x48B7DC79, 0x72FA0147, 0x8676F62F, 0xFD209DA5);
 
 
diff --git a/Engine/Plugins/FX/Niagara/Source/NiagaraEditor/Private/NiagaraHlslTranslator.cpp b/Engine/Plugins/FX/Niagara/Source/NiagaraEditor/Private/NiagaraHlslTranslator.cpp
--- a/Engine/Plugins/FX/Niagara/Source/NiagaraEditor/Private/NiagaraHlslTranslator.cpp
+++ b/Engine/Plugins/FX/Niagara/Source/NiagaraEditor/Private/NiagaraHlslTranslator.cpp
@@ -128,7 +128,6 @@ void FNiagaraShaderQueueTickable::ProcessQueue()
 			// Create a shader compiler environment for the script that will be shared by all jobs from this script
 			TRefCountPtr<FSharedShaderCompilerEnvironment> CompilerEnvironment = new FSharedShaderCompilerEnvironment();
 
-			FString ShaderCode = CompilableScript->GetVMExecutableData().LastHlslTranslationGPU;
 			// Shaders are created in-sync in the postload when running the automated tests.
 			const bool bSynchronousCompile = GIsAutomationTesting;
 
