Fixed static analysis warning.Removed duplicate tests keeping the existing logic the same.

diff --git a/Engine/Plugins/FX/Niagara/Source/NiagaraEditor/Private/EdGraphSchema_Niagara.cpp b/Engine/Plugins/FX/Niagara/Source/NiagaraEditor/Private/EdGraphSchema_Niagara.cpp
--- a/Engine/Plugins/FX/Niagara/Source/NiagaraEditor/Private/EdGraphSchema_Niagara.cpp
+++ b/Engine/Plugins/FX/Niagara/Source/NiagaraEditor/Private/EdGraphSchema_Niagara.cpp
@@ -1599,8 +1599,8 @@ FNiagaraTypeDefinition UEdGraphSchema_Niagara::PinTypeToTypeDefinition(const FEd
 
 		if (PinType.PinCategory == PinCategoryEnum)
 			return FNiagaraTypeDefinition(Enum);
-		else if (PinType.PinCategory == PinCategoryEnum)
-			return FNiagaraTypeDefinition(Enum).ToStaticDef();
+		//else if (PinType.PinCategory == PinCategoryEnum)
+		//	return FNiagaraTypeDefinition(Enum).ToStaticDef();
 	}
 	
 	return FNiagaraTypeDefinition();
@@ -1876,8 +1876,6 @@ bool UEdGraphSchema_Niagara::PinTypesValidForNumericConversion(FEdGraphPinType A
 		return true;
 	else if (AType.PinCategory == PinCategoryStaticType && BType.PinCategory == PinCategoryType)
 		return true;
-	else if (BType.PinCategory == PinCategoryStaticType && AType.PinCategory == PinCategoryType)
-		return true;
 	else if (AType.PinCategory == PinCategoryType && BType.PinCategory == PinCategoryType)
 		return true;
 	else if (BType.PinCategory == PinCategoryType && AType.PinCategory == PinCategoryType)
