Fix lwc type naming in Niagara

diff --git a/Engine/Plugins/FX/Niagara/Source/Niagara/Public/NiagaraTypes.h b/Engine/Plugins/FX/Niagara/Source/Niagara/Public/NiagaraTypes.h
--- a/Engine/Plugins/FX/Niagara/Source/Niagara/Public/NiagaraTypes.h
+++ b/Engine/Plugins/FX/Niagara/Source/Niagara/Public/NiagaraTypes.h
@@ -857,7 +857,7 @@ public:
 		return ClassStructOrEnum == Other.ClassStructOrEnum && UnderlyingType == Other.UnderlyingType;
 	}
 
-	FText GetNameText()const
+	FText GetNameText() const
 	{
 		if (IsValid() == false)
 		{
@@ -868,7 +868,27 @@ public:
 #if WITH_EDITOR
 		if ( UnderlyingType != UT_Enum )
 		{
-			NameText = GetStruct()->GetDisplayNameText();
+			// fix names for common types after LWC changes
+			if (*this == GetVec2Def())
+			{
+				NameText = FText::FromString("Vector 2");
+			}
+			else if (*this == GetVec3Def())
+			{
+				NameText = FText::FromString("Vector 3");
+			}
+			else if (*this == GetVec4Def())
+			{
+				NameText = FText::FromString("Vector 4");
+			}
+			else if (*this == GetQuatDef())
+			{
+				NameText = FText::FromString("Quat");
+			}
+			else
+			{
+				NameText = GetStruct()->GetDisplayNameText();
+			}
 		}
 		else
 #endif
diff --git a/Engine/Plugins/FX/Niagara/Source/NiagaraEditor/Private/Widgets/SNiagaraParameterMapView.cpp b/Engine/Plugins/FX/Niagara/Source/NiagaraEditor/Private/Widgets/SNiagaraParameterMapView.cpp
--- a/Engine/Plugins/FX/Niagara/Source/NiagaraEditor/Private/Widgets/SNiagaraParameterMapView.cpp
+++ b/Engine/Plugins/FX/Niagara/Source/NiagaraEditor/Private/Widgets/SNiagaraParameterMapView.cpp
@@ -2415,7 +2415,7 @@ void SNiagaraAddParameterMenu::AddParameterSelected(FNiagaraVariable NewVariable
 		}
 		else if (NewVariable.GetType().GetStruct() != nullptr)
 		{
-			TypeDisplayName = NewVariable.GetType().GetStruct()->GetDisplayNameText().ToString();
+			TypeDisplayName = NewVariable.GetType().GetNameText().ToString();
 		}
 		else if (NewVariable.GetType().GetClass() != nullptr)
 		{
diff --git a/Engine/Plugins/FX/Niagara/Source/NiagaraEditor/Public/Widgets/SNiagaraParameterMenu.cpp b/Engine/Plugins/FX/Niagara/Source/NiagaraEditor/Public/Widgets/SNiagaraParameterMenu.cpp
--- a/Engine/Plugins/FX/Niagara/Source/NiagaraEditor/Public/Widgets/SNiagaraParameterMenu.cpp
+++ b/Engine/Plugins/FX/Niagara/Source/NiagaraEditor/Public/Widgets/SNiagaraParameterMenu.cpp
@@ -703,7 +703,7 @@ void SNiagaraAddParameterFromPanelMenu::ParameterSelected(FNiagaraVariable NewVa
 		}
 		else if (NewVariable.GetType().GetStruct() != nullptr)
 		{
-			TypeDisplayName = NewVariable.GetType().GetStruct()->GetDisplayNameText().ToString();
+			TypeDisplayName = NewVariable.GetType().GetNameText().ToString();
 		}
 		else if (NewVariable.GetType().GetClass() != nullptr)
 		{
