Changed the parameter types of the transform node to not be as confusing for vector/position types

diff --git a/Engine/Plugins/FX/Niagara/Source/NiagaraEditor/Private/NiagaraEditorCommon.cpp b/Engine/Plugins/FX/Niagara/Source/NiagaraEditor/Private/NiagaraEditorCommon.cpp
--- a/Engine/Plugins/FX/Niagara/Source/NiagaraEditor/Private/NiagaraEditorCommon.cpp
+++ b/Engine/Plugins/FX/Niagara/Source/NiagaraEditor/Private/NiagaraEditorCommon.cpp
@@ -1,7 +1,6 @@
 // Copyright Epic Games, Inc. All Rights Reserved.
 
 #include "NiagaraEditorCommon.h"
-#include "AssetData.h"
 #include "INiagaraCompiler.h"
 #include "NiagaraSystem.h"
 
@@ -82,8 +81,6 @@ void FNiagaraOpInfo::Init()
 	//Common input and output names.
  	static FName Result(TEXT("Result"));
 	static FText ResultText = NSLOCTEXT("NiagaraOpInfo", "Operation Result", "Result");
-		
-	static FName In(TEXT("In"));
 	static FText InText = NSLOCTEXT("NiagaraOpInfo", "Single Function Param", "In"); 
 
 	static FName A(TEXT("A"));
@@ -96,13 +93,8 @@ void FNiagaraOpInfo::Init()
 	static FText CText = NSLOCTEXT("NiagaraOpInfo", "Third Function Param", "C");
 
 	FName X(TEXT("X"));	FText XText = NSLOCTEXT("NiagaraOpInfo", "First Vector Component", "X");
-	FName Y(TEXT("Y"));	FText YText = NSLOCTEXT("NiagaraOpInfo", "Second Vector Component", "Y");
-	FName Z(TEXT("Z"));	FText ZText = NSLOCTEXT("NiagaraOpInfo", "Third Vector Component", "Z");
-	FName W(TEXT("W"));	FText WText = NSLOCTEXT("NiagaraOpInfo", "Fourth Vector Component", "W");
 
-	static FName Selection(TEXT("Selection"));	
 	static FText SelectionText = NSLOCTEXT("NiagaraOpInfo", "Selection", "Selection");
-	static FName Mask(TEXT("Mask"));	
 	static FText MaskText = NSLOCTEXT("NiagaraOpInfo", "Mask", "Mask");
 	
 	static FName Min(TEXT("Min"));
@@ -1484,10 +1476,10 @@ void FNiagaraOpInfo::Init()
 	Op = &OpInfos[Idx];
 	Op->Category = MatrixCategory;
 	Op->FriendlyName = NSLOCTEXT("NiagaraOpInfo", "TransformPosition Name", "Transform Position");
-	Op->Description = NSLOCTEXT("NiagaraOpInfo", "TransformPosition Desc", "Transforms a Vector3 as a position.");
+	Op->Description = NSLOCTEXT("NiagaraOpInfo", "TransformPosition Desc", "Transforms a Vector3 as a position, meaning it multiplies it with the full matrix, including translation.");
 	Op->Inputs.Add(FNiagaraOpInOutInfo(M, MatrixType, MText, MText, Default_MatrixOne));
-	Op->Inputs.Add(FNiagaraOpInOutInfo(V, FNiagaraTypeDefinition::GetVec3Def(), VText, VText, Default_VectorOne));
-	Op->Outputs.Add(FNiagaraOpInOutInfo(Result, FNiagaraTypeDefinition::GetVec3Def(), ResultText, ResultText, Default_VectorOne, TEXT("mul(float4({1},1.0),{0}).xyz")));
+	Op->Inputs.Add(FNiagaraOpInOutInfo(V, FNiagaraTypeDefinition::GetGenericNumericDef(), VText, VText, Default_VectorOne));
+	Op->Outputs.Add(FNiagaraOpInOutInfo(Result, FNiagaraTypeDefinition::GetGenericNumericDef(), ResultText, ResultText, Default_VectorOne, TEXT("mul(float4({1},1.0),{0}).xyz")));
 	Op->BuildName(TEXT("TransformPosition"), MatrixCategoryName);
 	OpInfoMap.Add(Op->Name) = Idx;
 
@@ -1495,7 +1487,7 @@ void FNiagaraOpInfo::Init()
 	Op = &OpInfos[Idx];
 	Op->Category = MatrixCategory;
 	Op->FriendlyName = NSLOCTEXT("NiagaraOpInfo", "TransformVector Name", "Transform Vector");
-	Op->Description = NSLOCTEXT("NiagaraOpInfo", "TransformVector Desc", "Transforms a Vector3 as a vector.");
+	Op->Description = NSLOCTEXT("NiagaraOpInfo", "TransformVector Desc", "Transforms a Vector3 as a vector, meaning it does not apply the translation part of the matrix.");
 	Op->Inputs.Add(FNiagaraOpInOutInfo(M, MatrixType, MText, MText, Default_MatrixOne));
 	Op->Inputs.Add(FNiagaraOpInOutInfo(V, FNiagaraTypeDefinition::GetVec3Def(), VText, VText, Default_VectorOne));
 	Op->Outputs.Add(FNiagaraOpInOutInfo(Result, FNiagaraTypeDefinition::GetVec3Def(), ResultText, ResultText, Default_VectorOne, TEXT("mul(float4({1},0.0),{0}).xyz")));
