Fixed user redirection for position user parameters

diff --git a/Engine/Plugins/FX/Niagara/Source/Niagara/Private/NiagaraUserRedirectionParameterStore.cpp b/Engine/Plugins/FX/Niagara/Source/Niagara/Private/NiagaraUserRedirectionParameterStore.cpp
--- a/Engine/Plugins/FX/Niagara/Source/Niagara/Private/NiagaraUserRedirectionParameterStore.cpp
+++ b/Engine/Plugins/FX/Niagara/Source/Niagara/Private/NiagaraUserRedirectionParameterStore.cpp
@@ -140,6 +140,20 @@ void FNiagaraUserRedirectionParameterStore::Reset(bool bClearBindings /*= true*/
 	UserParameterRedirects.Reset();
 }
 
+bool FNiagaraUserRedirectionParameterStore::SetPositionParameterValue(const FVector& InValue, const FName& ParamName, bool bAdd)
+{
+	FNiagaraVariable PositionVar = FNiagaraVariable(FNiagaraTypeDefinition::GetPositionDef(), ParamName);
+	const FNiagaraVariable* Redirection = UserParameterRedirects.Find(PositionVar);
+	return FNiagaraParameterStore::SetPositionParameterValue(InValue, Redirection ? Redirection->GetName() : ParamName, bAdd);
+}
+
+const FVector* FNiagaraUserRedirectionParameterStore::GetPositionParameterValue(const FName& ParamName) const
+{
+	FNiagaraVariable PositionVar = FNiagaraVariable(FNiagaraTypeDefinition::GetPositionDef(), ParamName);
+	const FNiagaraVariable* Redirection = UserParameterRedirects.Find(PositionVar);
+	return FNiagaraParameterStore::GetPositionParameterValue(Redirection ? Redirection->GetName() : ParamName);
+}
+
 #if WITH_EDITORONLY_DATA
 void FNiagaraUserRedirectionParameterStore::ConvertParameterType(const FNiagaraVariable& ExistingParam,	const FNiagaraTypeDefinition& NewType)
 {
diff --git a/Engine/Plugins/FX/Niagara/Source/Niagara/Public/NiagaraParameterStore.h b/Engine/Plugins/FX/Niagara/Source/Niagara/Public/NiagaraParameterStore.h
--- a/Engine/Plugins/FX/Niagara/Source/Niagara/Public/NiagaraParameterStore.h
+++ b/Engine/Plugins/FX/Niagara/Source/Niagara/Public/NiagaraParameterStore.h
@@ -502,8 +502,8 @@ public:
 
 	FString ToString() const;
 
-	bool SetPositionParameterValue(const FVector& InValue, const FName& ParamName, bool bAdd=false);
-	const FVector* GetPositionParameterValue(const FName& ParamName) const;
+	virtual bool SetPositionParameterValue(const FVector& InValue, const FName& ParamName, bool bAdd=false);
+	virtual const FVector* GetPositionParameterValue(const FName& ParamName) const;
 	void ResolvePositions(FNiagaraLWCConverter LwcConverter);
 
 	template<typename T>
diff --git a/Engine/Plugins/FX/Niagara/Source/Niagara/Public/NiagaraUserRedirectionParameterStore.h b/Engine/Plugins/FX/Niagara/Source/Niagara/Public/NiagaraUserRedirectionParameterStore.h
--- a/Engine/Plugins/FX/Niagara/Source/Niagara/Public/NiagaraUserRedirectionParameterStore.h
+++ b/Engine/Plugins/FX/Niagara/Source/Niagara/Public/NiagaraUserRedirectionParameterStore.h
@@ -59,6 +59,8 @@ public:
 	virtual void InitFromSource(const FNiagaraParameterStore* SrcStore, bool bNotifyAsDirty) override;
 	virtual void Empty(bool bClearBindings = true) override;
 	virtual void Reset(bool bClearBindings = true) override;
+	virtual bool SetPositionParameterValue(const FVector& InValue, const FName& ParamName, bool bAdd=false) override;
+	virtual const FVector* GetPositionParameterValue(const FName& ParamName) const override;
 
 #if WITH_EDITORONLY_DATA
 	virtual void ConvertParameterType(const FNiagaraVariable& ExistingParam, const FNiagaraTypeDefinition& NewType) override;
