Fixed issue with ribbon indexing breaking in large ribbons due to selection between 16 and 32 bit indices. Fix should save some memory with small/normal size ribbons, and correctly handle large complicated ribbons.

diff --git a/Engine/Plugins/FX/Niagara/Source/Niagara/Private/NiagaraRendererRibbons.cpp b/Engine/Plugins/FX/Niagara/Source/Niagara/Private/NiagaraRendererRibbons.cpp
--- a/Engine/Plugins/FX/Niagara/Source/Niagara/Private/NiagaraRendererRibbons.cpp
+++ b/Engine/Plugins/FX/Niagara/Source/Niagara/Private/NiagaraRendererRibbons.cpp
@@ -1368,7 +1368,7 @@ void FNiagaraRendererRibbons::CreatePerViewResources(
 	FGlobalDynamicIndexBuffer& DynamicIndexBuffer = Collector.GetDynamicIndexBuffer();
 
 	const uint32 NumIndices = NumSegments * TrianglesPerSegment * 3;
-	if (IndexBufferOffsets.TotalBitCount > 16/* Number of bits in a ushort*/)
+	if (IndexBufferOffsets.TotalBitCount <= 16/* Number of bits in a ushort*/)
 	{
 		InOutIndexAllocation = DynamicIndexBuffer.Allocate<uint16>(NumIndices);
 		GenerateIndexBuffer<uint16>(InOutIndexAllocation, IndexBufferOffsets, SegmentTessellation, View->GetViewDirection(), ViewOriginForDistanceCulling, DynamicDataRibbon);
