Fix Mesh Renderer Info to show only renderers in current emitter.

diff --git a/Engine/Plugins/FX/Niagara/Source/NiagaraEditorWidgets/Private/DetailCustomizations/NiagaraDataInterfaceMeshRendererInfoDetails.cpp b/Engine/Plugins/FX/Niagara/Source/NiagaraEditorWidgets/Private/DetailCustomizations/NiagaraDataInterfaceMeshRendererInfoDetails.cpp
--- a/Engine/Plugins/FX/Niagara/Source/NiagaraEditorWidgets/Private/DetailCustomizations/NiagaraDataInterfaceMeshRendererInfoDetails.cpp
+++ b/Engine/Plugins/FX/Niagara/Source/NiagaraEditorWidgets/Private/DetailCustomizations/NiagaraDataInterfaceMeshRendererInfoDetails.cpp
@@ -106,26 +106,20 @@ void FNiagaraDataInterfaceMeshRendererInfoDetails::GenerateRendererList()
 	RendererList.SetNum(0, false);
 	RendererLabels.SetNum(0, false);
 
-	// Get a list of all mesh renderers in the system
+	// Get a list of all mesh renderers in the emitter
 	if (auto Interface = DataInterface.Get())
 	{
 		UNiagaraSystem* System = nullptr;
 		UNiagaraEmitter* Emitter = nullptr;
 		FNiagaraEditorModule::Get().GetTargetSystemAndEmitterForDataInterface(Interface, System, Emitter);
-		if (System)
+		if (System && Emitter)
 		{
-			for (auto& EmitterHandle : System->GetEmitterHandles())
+			for (auto RendererProps : Emitter->GetRenderers())
 			{
-				if (EmitterHandle.IsValid() && EmitterHandle.GetInstance())
+				if (auto MeshProps = Cast<UNiagaraMeshRendererProperties>(RendererProps))
 				{
-					for (auto RendererProps : EmitterHandle.GetInstance()->GetRenderers())
-					{
-						if (auto MeshProps = Cast<UNiagaraMeshRendererProperties>(RendererProps))
-						{
-							RendererList.Add(MeshProps);
-							RendererLabels.Add(CreateRendererTextLabel(MeshProps));
-						}
-					}
+					RendererList.Add(MeshProps);
+					RendererLabels.Add(CreateRendererTextLabel(MeshProps));
 				}
 			}
 
@@ -192,8 +186,6 @@ FText FNiagaraDataInterfaceMeshRendererInfoDetails::CreateRendererTextLabel(cons
 	UNiagaraEmitter* Emitter = Properties->GetTypedOuter<UNiagaraEmitter>();
 	check(Emitter);
 
-	const FText EmitterText = FText::FromString(Emitter->GetUniqueEmitterName());
-
 	UNiagaraEmitterEditorData* EmitterEditorData = static_cast<UNiagaraEmitterEditorData*>(Emitter->GetEditorData());
 	if (EmitterEditorData == nullptr)
 	{
@@ -203,8 +195,7 @@ FText FNiagaraDataInterfaceMeshRendererInfoDetails::CreateRendererTextLabel(cons
 	FString RendererStackEditorDataKey = FString::Printf(TEXT("Renderer-%s"), *Properties->GetName());
 	const FText* RendererDisplayName = EmitterStackEditorData.GetStackEntryDisplayName(RendererStackEditorDataKey);
 
-	const FText RendererText = RendererDisplayName ? *RendererDisplayName : Properties->GetWidgetDisplayName();
-	return FText::Format(FTextFormat::FromString(TEXT("{0} ({1})")), RendererText, EmitterText);
+	return RendererDisplayName ? *RendererDisplayName : Properties->GetWidgetDisplayName();
 }
 
 #undef LOCTEXT_NAMESPACE
\ No newline at end of file
