Remove necessity of separately clearing the package custom versions array when resetting a package linker
This is in preparation of the deprecation of the usage of public data members in UPackage
ResetLoaders now also empty a package custom versions array when detaching a valid linker from it.

diff --git a/Engine/Plugins/FX/Niagara/Source/NiagaraEditor/Private/Toolkits/NiagaraParameterDefinitionsToolkit.cpp b/Engine/Plugins/FX/Niagara/Source/NiagaraEditor/Private/Toolkits/NiagaraParameterDefinitionsToolkit.cpp
--- a/Engine/Plugins/FX/Niagara/Source/NiagaraEditor/Private/Toolkits/NiagaraParameterDefinitionsToolkit.cpp
+++ b/Engine/Plugins/FX/Niagara/Source/NiagaraEditor/Private/Toolkits/NiagaraParameterDefinitionsToolkit.cpp
@@ -81,8 +81,7 @@ void FNiagaraParameterDefinitionsToolkit::Initialize(const EToolkitMode::Type Mo
 	check(InParameterDefinitions != nullptr);
 	ParameterDefinitionsSource = InParameterDefinitions;
 
-	ResetLoaders(GetTransientPackage()); // Make sure that we're not going to get invalid version number linkers into the package we are going into. 
-	GetTransientPackage()->LinkerCustomVersion.Empty();
+	// No need to reset loader or versioning on the transient package, there should never be any set
 	ParameterDefinitionsInstance = Cast<UNiagaraParameterDefinitions>(StaticDuplicateObject(ParameterDefinitionsSource, GetTransientPackage(), NAME_None, ~RF_Standalone, UNiagaraParameterDefinitions::StaticClass()));
 	ParameterDefinitionsInstance->ClearFlags(RF_Standalone | RF_Public);
 	LastSyncedDefinitionsChangeIdHash = ParameterDefinitionsInstance->GetChangeIdHash();
@@ -319,7 +318,6 @@ void FNiagaraParameterDefinitionsToolkit::OnApply()
 	}
 
 	ResetLoaders(ParameterDefinitionsSource->GetOutermost()); // Make sure that we're not going to get invalid version number linkers into the package we are going into. 
-	ParameterDefinitionsSource->GetOutermost()->LinkerCustomVersion.Empty();
 
 	ParameterDefinitionsSource->PreEditChange(nullptr);
 	// overwrite the original parameter definitions in place by constructing a new one with the same name
diff --git a/Engine/Plugins/FX/Niagara/Source/NiagaraEditor/Private/Toolkits/NiagaraScriptToolkit.cpp b/Engine/Plugins/FX/Niagara/Source/NiagaraEditor/Private/Toolkits/NiagaraScriptToolkit.cpp
--- a/Engine/Plugins/FX/Niagara/Source/NiagaraEditor/Private/Toolkits/NiagaraScriptToolkit.cpp
+++ b/Engine/Plugins/FX/Niagara/Source/NiagaraEditor/Private/Toolkits/NiagaraScriptToolkit.cpp
@@ -178,8 +178,7 @@ void FNiagaraScriptToolkit::Initialize( const EToolkitMode::Type Mode, const TSh
 	FGuid ScriptVersion = InputScript->IsVersioningEnabled() && InputScript->VersionToOpenInEditor.IsValid() ? InputScript->VersionToOpenInEditor : InputScript->GetExposedVersion().VersionGuid;
 	OriginalNiagaraScript.Script = InputScript;
 	OriginalNiagaraScript.Version = ScriptVersion;
-	ResetLoaders(GetTransientPackage()); // Make sure that we're not going to get invalid version number linkers into the package we are going into. 
-	GetTransientPackage()->LinkerCustomVersion.Empty();
+	// No need to reset loader or versioning on the transient package, there should never be any set
 	EditedNiagaraScript.Script = (UNiagaraScript*)StaticDuplicateObject(InputScript, GetTransientPackage(), NAME_None, ~RF_Standalone, UNiagaraScript::StaticClass());
 	EditedNiagaraScript.Script->OnVMScriptCompiled().AddSP(this, &FNiagaraScriptToolkit::OnVMScriptCompiled);
 	EditedNiagaraScript.Version = ScriptVersion;
@@ -908,7 +907,6 @@ void FNiagaraScriptToolkit::UpdateOriginalNiagaraScript()
 	}
 
 	ResetLoaders(OriginalNiagaraScript.Script->GetOutermost()); // Make sure that we're not going to get invalid version number linkers into the package we are going into. 
-	OriginalNiagaraScript.Script->GetOutermost()->LinkerCustomVersion.Empty();
 
 	// Compile and then overwrite the original script in place by constructing a new one with the same name
 	ScriptViewModel->CompileStandaloneScript();
diff --git a/Engine/Plugins/FX/Niagara/Source/NiagaraEditor/Private/Toolkits/NiagaraSystemToolkit.cpp b/Engine/Plugins/FX/Niagara/Source/NiagaraEditor/Private/Toolkits/NiagaraSystemToolkit.cpp
--- a/Engine/Plugins/FX/Niagara/Source/NiagaraEditor/Private/Toolkits/NiagaraSystemToolkit.cpp
+++ b/Engine/Plugins/FX/Niagara/Source/NiagaraEditor/Private/Toolkits/NiagaraSystemToolkit.cpp
@@ -329,8 +329,7 @@ void FNiagaraSystemToolkit::InitializeWithEmitter(const EToolkitMode::Type Mode,
 
 	FNiagaraUtilities::PrepareRapidIterationParameters(Scripts, ScriptDependencyMap, ScriptToEmitterMap);
 
-	ResetLoaders(GetTransientPackage()); // Make sure that we're not going to get invalid version number linkers into the package we are going into. 
-	GetTransientPackage()->LinkerCustomVersion.Empty();
+	// No need to reset loader or versioning on the transient package, there should never be any set 
 	
 	bEmitterThumbnailUpdated = false;
 
@@ -1611,7 +1610,6 @@ void FNiagaraSystemToolkit::UpdateOriginalEmitter()
 		}
 
 		ResetLoaders(Emitter->GetOutermost()); // Make sure that we're not going to get invalid version number linkers into the package we are going into. 
-		Emitter->GetOutermost()->LinkerCustomVersion.Empty();
 
 		TArray<UNiagaraScript*> AllScripts;
 		EditableEmitter->GetScripts(AllScripts, true);
diff --git a/Engine/Plugins/FX/Niagara/Source/NiagaraEditor/Private/ViewModels/NiagaraScratchPadScriptViewModel.cpp b/Engine/Plugins/FX/Niagara/Source/NiagaraEditor/Private/ViewModels/NiagaraScratchPadScriptViewModel.cpp
--- a/Engine/Plugins/FX/Niagara/Source/NiagaraEditor/Private/ViewModels/NiagaraScratchPadScriptViewModel.cpp
+++ b/Engine/Plugins/FX/Niagara/Source/NiagaraEditor/Private/ViewModels/NiagaraScratchPadScriptViewModel.cpp
@@ -171,7 +171,6 @@ bool FNiagaraScratchPadScriptViewModel::HasUnappliedChanges() const
 void FNiagaraScratchPadScriptViewModel::ApplyChanges()
 {
 	ResetLoaders(OriginalScript->GetOutermost()); // Make sure that we're not going to get invalid version number linkers into the package we are going into. 
-	OriginalScript->GetOutermost()->LinkerCustomVersion.Empty();
 
 	OriginalScript = (UNiagaraScript*)StaticDuplicateObject(EditScript.Script, OriginalScript->GetOuter(), OriginalScript->GetFName(),
 		RF_AllFlags,
diff --git a/Engine/Source/Runtime/CoreUObject/Private/UObject/LinkerLoad.cpp b/Engine/Source/Runtime/CoreUObject/Private/UObject/LinkerLoad.cpp
--- a/Engine/Source/Runtime/CoreUObject/Private/UObject/LinkerLoad.cpp
+++ b/Engine/Source/Runtime/CoreUObject/Private/UObject/LinkerLoad.cpp
@@ -5521,6 +5521,13 @@ void FLinkerLoad::Detach()
 	if (LinkerRoot)
 	{
 		LinkerRoot->LinkerLoad = nullptr;
+		// When detaching the linker from its package, also empty its stored list of custom versions. 
+		// This is so that object post loaded in the editor in a package that has no associated linker consider that all package custom versions as latest
+		// (i.e. when duplicating an object in the package)
+		// The runtime *may* use the stored version in the package since there are never any linker associated with it when using iostore
+		PRAGMA_DISABLE_DEPRECATION_WARNINGS
+		LinkerRoot->LinkerCustomVersion.Empty();
+		PRAGMA_ENABLE_DEPRECATION_WARNINGS
 		LinkerRoot = nullptr;
 	}
 
