- Add option to override simulation stage num threads

#rnx

diff --git a/Engine/Plugins/FX/Niagara/Source/Niagara/Private/NiagaraDataInterfaceRW.cpp b/Engine/Plugins/FX/Niagara/Source/Niagara/Private/NiagaraDataInterfaceRW.cpp
--- a/Engine/Plugins/FX/Niagara/Source/Niagara/Private/NiagaraDataInterfaceRW.cpp
+++ b/Engine/Plugins/FX/Niagara/Source/Niagara/Private/NiagaraDataInterfaceRW.cpp
@@ -404,7 +404,7 @@ bool UNiagaraDataInterfaceGrid3D::GetFunctionHLSL(const FNiagaraDataInterfaceGPU
 		static const TCHAR* FormatSample = TEXT(R"(
 			void {FunctionName}(out float3 Out_Unit)
 			{
-				#if NIAGARA_DISPATCH_TYPE == NIAGARA_DISPATCH_TYPE_THREE_D
+				#if NIAGARA_DISPATCH_TYPE == NIAGARA_DISPATCH_TYPE_THREE_D || NIAGARA_DISPATCH_TYPE == NIAGARA_DISPATCH_TYPE_CUSTOM
 					Out_Unit = (float3(GDispatchThreadId.x, GDispatchThreadId.y, GDispatchThreadId.z) + .5) / {NumCellsName};
 				#else
 					const uint Linear = GLinearThreadId;
@@ -425,7 +425,7 @@ bool UNiagaraDataInterfaceGrid3D::GetFunctionHLSL(const FNiagaraDataInterfaceGPU
 		static const TCHAR* FormatSample = TEXT(R"(
 			void {FunctionName}(out int Out_IndexX, out int Out_IndexY, out int Out_IndexZ)
 			{
-				#if NIAGARA_DISPATCH_TYPE == NIAGARA_DISPATCH_TYPE_THREE_D
+				#if NIAGARA_DISPATCH_TYPE == NIAGARA_DISPATCH_TYPE_THREE_D || NIAGARA_DISPATCH_TYPE == NIAGARA_DISPATCH_TYPE_CUSTOM
 					Out_IndexX = GDispatchThreadId.x;
 					Out_IndexY = GDispatchThreadId.y;
 					Out_IndexZ = GDispatchThreadId.z;
@@ -895,7 +895,7 @@ bool UNiagaraDataInterfaceGrid2D::GetFunctionHLSL(const FNiagaraDataInterfaceGPU
 		static const TCHAR* FormatSample = TEXT(R"(
 			void {FunctionName}(out float2 Out_Unit)
 			{
-				#if NIAGARA_DISPATCH_TYPE == NIAGARA_DISPATCH_TYPE_TWO_D
+				#if NIAGARA_DISPATCH_TYPE == NIAGARA_DISPATCH_TYPE_TWO_D || NIAGARA_DISPATCH_TYPE == NIAGARA_DISPATCH_TYPE_CUSTOM
 					Out_Unit = (float2(GDispatchThreadId.x, GDispatchThreadId.y) + .5) * {UnitToUVName};			
 				#else
 					const uint Linear = GLinearThreadId;
@@ -914,7 +914,7 @@ bool UNiagaraDataInterfaceGrid2D::GetFunctionHLSL(const FNiagaraDataInterfaceGPU
 		static const TCHAR* FormatSample = TEXT(R"(
 			void {FunctionName}(out int Out_IndexX, out int Out_IndexY)
 			{
-				#if NIAGARA_DISPATCH_TYPE == NIAGARA_DISPATCH_TYPE_TWO_D
+				#if NIAGARA_DISPATCH_TYPE == NIAGARA_DISPATCH_TYPE_TWO_D || NIAGARA_DISPATCH_TYPE == NIAGARA_DISPATCH_TYPE_CUSTOM
 					Out_IndexX = GDispatchThreadId.x;
 					Out_IndexY = GDispatchThreadId.y;
 				#else
diff --git a/Engine/Plugins/FX/Niagara/Source/Niagara/Private/NiagaraGpuComputeDispatch.cpp b/Engine/Plugins/FX/Niagara/Source/Niagara/Private/NiagaraGpuComputeDispatch.cpp
--- a/Engine/Plugins/FX/Niagara/Source/Niagara/Private/NiagaraGpuComputeDispatch.cpp
+++ b/Engine/Plugins/FX/Niagara/Source/Niagara/Private/NiagaraGpuComputeDispatch.cpp
@@ -1474,7 +1474,7 @@ void FNiagaraGpuComputeDispatch::DispatchStage(FRHICommandList& RHICmdList, FRHI
 	{
 		DispatchType = SimStageData.StageMetaData->GpuDispatchType;
 		DispatchCount = SimStageData.AlternateIterationSource->GetElementCount(Tick.SystemInstanceID);
-		DispatchNumThreads = (DispatchType == ENiagaraGpuDispatchType::Custom) ? SimStageData.StageMetaData->GpuDispatchNumThreads : FNiagaraShader::GetDefaultThreadGroupSize(DispatchType);
+		DispatchNumThreads = SimStageData.StageMetaData->GpuDispatchNumThreads;
 
 		// Verify the number of elements isn't higher that what we can handle
 		checkf(uint64(DispatchCount.X) * uint64(DispatchCount.Y) * uint64(DispatchCount.Z) < uint64(TNumericLimits<int32>::Max()), TEXT("DispatchCount(%d, %d, %d) for IterationInterface(%s) overflows an int32 this is not allowed"), DispatchCount.X, DispatchCount.Y, DispatchCount.Z, *SimStageData.AlternateIterationSource->SourceDIName.ToString());
@@ -1510,13 +1510,13 @@ void FNiagaraGpuComputeDispatch::DispatchStage(FRHICommandList& RHICmdList, FRHI
 	SetComputePipelineState(RHICmdList, RHIComputeShader);
 
 	SCOPED_DRAW_EVENTF(RHICmdList, NiagaraGPUSimulationCS,
-		TEXT("NiagaraGpuSim(%s) DispatchCount(%d x %d x %d) Stage(%s %u) NumInstructions(%u)"),
+		TEXT("NiagaraGpuSim(%s) DispatchCount(%dx%dx%d) Stage(%s %u) Iteration(%u) NumThreads(%dx%dx%d)"),
 		InstanceData.Context->GetDebugSimName(),
 		DispatchCount.X, DispatchCount.Y, DispatchCount.Z,
 		*SimStageData.StageMetaData->SimulationStageName.ToString(),
 		SimStageData.StageIndex,
 		SimStageData.IterationIndex,
-		ComputeShader->GetNumInstructions()
+		DispatchNumThreads.X, DispatchNumThreads.Y, DispatchNumThreads.Z
 	);
 	FNiagaraEmptyUAVPoolScopedAccess UAVPoolAccessScope(GetEmptyUAVPool());
 
diff --git a/Engine/Plugins/FX/Niagara/Source/Niagara/Private/NiagaraSimulationStageBase.cpp b/Engine/Plugins/FX/Niagara/Source/Niagara/Private/NiagaraSimulationStageBase.cpp
--- a/Engine/Plugins/FX/Niagara/Source/Niagara/Private/NiagaraSimulationStageBase.cpp
+++ b/Engine/Plugins/FX/Niagara/Source/Niagara/Private/NiagaraSimulationStageBase.cpp
@@ -74,6 +74,10 @@ bool UNiagaraSimulationStageGeneric::AppendCompileHash(FNiagaraCompileHashVisito
 	InVisitor->UpdatePOD(TEXT("bParticleIterationStateEnabled"), bParticleIterationStateEnabled ? 1 : 0);
 	InVisitor->UpdateString(TEXT("ParticleIterationStateBinding"), ParticleIterationStateBinding.GetDataSetBindableVariable().GetName().ToString());
 	InVisitor->UpdateString(TEXT("ParticleIterationStateRange"), FString::Printf(TEXT("%d,%d"), ParticleIterationStateRange.X, ParticleIterationStateRange.Y));
+	InVisitor->UpdatePOD(TEXT("bGpuDispatchForceLinear"), bGpuDispatchForceLinear ? 1 : 0);
+	InVisitor->UpdatePOD(TEXT("bOverrideGpuDispatchNumThreads"), bOverrideGpuDispatchNumThreads ? 1 : 0);
+	InVisitor->UpdateString(TEXT("OverrideGpuDispatchNumThreads"), FString::Printf(TEXT("%d,%d,%d"), OverrideGpuDispatchNumThreads.X, OverrideGpuDispatchNumThreads.Y, OverrideGpuDispatchNumThreads.Z));
+
 	return true;
 }
 
@@ -165,6 +169,21 @@ void UNiagaraSimulationStageGeneric::PostEditChangeProperty(struct FPropertyChan
 	{
 		bNeedsRecompile = true;
 	}
+	else if (PropertyName == GET_MEMBER_NAME_CHECKED(UNiagaraSimulationStageGeneric, bGpuDispatchForceLinear))
+	{
+		bNeedsRecompile = true;
+	}
+	else if (PropertyName == GET_MEMBER_NAME_CHECKED(UNiagaraSimulationStageGeneric, bOverrideGpuDispatchNumThreads))
+	{
+		bNeedsRecompile = true;
+	}
+	else if (MemberPropertyName == GET_MEMBER_NAME_CHECKED(UNiagaraSimulationStageGeneric, OverrideGpuDispatchNumThreads))
+	{
+		OverrideGpuDispatchNumThreads.X = FMath::Max(OverrideGpuDispatchNumThreads.X, 1);
+		OverrideGpuDispatchNumThreads.Y = FMath::Max(OverrideGpuDispatchNumThreads.Y, 1);
+		OverrideGpuDispatchNumThreads.Z = FMath::Max(OverrideGpuDispatchNumThreads.Z, 1);
+		bNeedsRecompile = true;
+	}
 
 	if (bNeedsRecompile)
 	{
diff --git a/Engine/Plugins/FX/Niagara/Source/Niagara/Public/NiagaraSimulationStageBase.h b/Engine/Plugins/FX/Niagara/Source/Niagara/Public/NiagaraSimulationStageBase.h
--- a/Engine/Plugins/FX/Niagara/Source/Niagara/Public/NiagaraSimulationStageBase.h
+++ b/Engine/Plugins/FX/Niagara/Source/Niagara/Public/NiagaraSimulationStageBase.h
@@ -91,6 +91,17 @@ public:
 	UPROPERTY(EditAnywhere, Category = "Simulation Stage", meta = (editcondition = "IterationSource == ENiagaraIterationSource::Particles"))
 	FIntPoint ParticleIterationStateRange = FIntPoint(0, 0);
 
+	/** When enabled we force the dispatch to be linear (i.e. one dimension is used). */
+	UPROPERTY(EditAnywhere, Category = "Simulation Stage")
+	uint32 bGpuDispatchForceLinear : 1;
+
+	/** When enabled we use a custom number of threads for the dispatch. */
+	UPROPERTY(EditAnywhere, Category = "Simulation Stage")
+	uint32 bOverrideGpuDispatchNumThreads : 1;
+
+	UPROPERTY(EditAnywhere, Category = "Simulation Stage")
+	FIntVector OverrideGpuDispatchNumThreads = FIntVector(64, 1, 1);
+
 	virtual void PostInitProperties() override;
 	virtual bool AppendCompileHash(FNiagaraCompileHashVisitor* InVisitor) const override;
 #if WITH_EDITOR
diff --git a/Engine/Plugins/FX/Niagara/Source/NiagaraCore/Private/NiagaraCustomVersion.cpp b/Engine/Plugins/FX/Niagara/Source/NiagaraCore/Private/NiagaraCustomVersion.cpp
--- a/Engine/Plugins/FX/Niagara/Source/NiagaraCore/Private/NiagaraCustomVersion.cpp
+++ b/Engine/Plugins/FX/Niagara/Source/NiagaraCore/Private/NiagaraCustomVersion.cpp
@@ -8,6 +8,6 @@ const FGuid FNiagaraCustomVersion::GUID(0xFCF57AFA, 0x50764283, 0xB9A9E658, 0xFF
 // Register the custom version with core
 FCustomVersionRegistration GRegisterNiagaraCustomVersion(FNiagaraCustomVersion::GUID, FNiagaraCustomVersion::LatestVersion, TEXT("NiagaraVer"));
 
-const FGuid FNiagaraCustomVersion::LatestScriptCompileVersion(0x48B7DC79, 0x72FA0147, 0x8676F62F, 0xFD209DA5);
+const FGuid FNiagaraCustomVersion::LatestScriptCompileVersion(0xC71A1029, 0x6E7E894B, 0xB639619F, 0xE083EBF9);
 
 
diff --git a/Engine/Plugins/FX/Niagara/Source/NiagaraEditor/Private/Customizations/NiagaraSimulationStageCustomization.cpp b/Engine/Plugins/FX/Niagara/Source/NiagaraEditor/Private/Customizations/NiagaraSimulationStageCustomization.cpp
--- a/Engine/Plugins/FX/Niagara/Source/NiagaraEditor/Private/Customizations/NiagaraSimulationStageCustomization.cpp
+++ b/Engine/Plugins/FX/Niagara/Source/NiagaraEditor/Private/Customizations/NiagaraSimulationStageCustomization.cpp
@@ -75,6 +75,9 @@ void FNiagaraSimulationStageGenericCustomization::CustomizeDetails(class IDetail
 	DetailBuilder.HideProperty(GET_MEMBER_NAME_CHECKED(UNiagaraSimulationStageGeneric, bParticleIterationStateEnabled));
 	DetailBuilder.HideProperty(GET_MEMBER_NAME_CHECKED(UNiagaraSimulationStageGeneric, ParticleIterationStateBinding));
 	DetailBuilder.HideProperty(GET_MEMBER_NAME_CHECKED(UNiagaraSimulationStageGeneric, ParticleIterationStateRange));
+	DetailBuilder.HideProperty(GET_MEMBER_NAME_CHECKED(UNiagaraSimulationStageGeneric, bGpuDispatchForceLinear));
+	DetailBuilder.HideProperty(GET_MEMBER_NAME_CHECKED(UNiagaraSimulationStageGeneric, bOverrideGpuDispatchNumThreads));
+	DetailBuilder.HideProperty(GET_MEMBER_NAME_CHECKED(UNiagaraSimulationStageGeneric, OverrideGpuDispatchNumThreads));
 
 	// Show properties in the order we want them to appear
 	SimStageCategory.AddProperty(DetailBuilder.GetProperty(GET_MEMBER_NAME_CHECKED(UNiagaraSimulationStageGeneric, EnabledBinding)));
@@ -95,6 +98,13 @@ void FNiagaraSimulationStageGenericCustomization::CustomizeDetails(class IDetail
 	{
 		DataInterfaceStageCategory.AddProperty(DetailBuilder.GetProperty(GET_MEMBER_NAME_CHECKED(UNiagaraSimulationStageGeneric, DataInterface)));
 		DataInterfaceStageCategory.AddProperty(DetailBuilder.GetProperty(GET_MEMBER_NAME_CHECKED(UNiagaraSimulationStageGeneric, ExecuteBehavior)));
+
+		DataInterfaceStageCategory.AddProperty(DetailBuilder.GetProperty(GET_MEMBER_NAME_CHECKED(UNiagaraSimulationStageGeneric, bGpuDispatchForceLinear)));
+		DataInterfaceStageCategory.AddProperty(DetailBuilder.GetProperty(GET_MEMBER_NAME_CHECKED(UNiagaraSimulationStageGeneric, bOverrideGpuDispatchNumThreads)));
+		if (SimStage->bOverrideGpuDispatchNumThreads)
+		{
+			DataInterfaceStageCategory.AddProperty(DetailBuilder.GetProperty(GET_MEMBER_NAME_CHECKED(UNiagaraSimulationStageGeneric, OverrideGpuDispatchNumThreads)));
+		}
 	}
 }
 
diff --git a/Engine/Plugins/FX/Niagara/Source/NiagaraEditor/Private/NiagaraCompiler.cpp b/Engine/Plugins/FX/Niagara/Source/NiagaraEditor/Private/NiagaraCompiler.cpp
--- a/Engine/Plugins/FX/Niagara/Source/NiagaraEditor/Private/NiagaraCompiler.cpp
+++ b/Engine/Plugins/FX/Niagara/Source/NiagaraEditor/Private/NiagaraCompiler.cpp
@@ -695,6 +695,9 @@ void FNiagaraCompileRequestData::FinishPrecompile(const TArray<FNiagaraVariable>
 					SimStageData.PartialParticleUpdate	= GenericStage->bDisablePartialParticleUpdate == false;
 					SimStageData.bParticleIterationStateEnabled	= GenericStage->bParticleIterationStateEnabled;
 					SimStageData.ParticleIterationStateRange	= GenericStage->ParticleIterationStateRange;
+					SimStageData.bGpuDispatchForceLinear		= GenericStage->bGpuDispatchForceLinear;
+					SimStageData.bOverrideGpuDispatchNumThreads	= GenericStage->bOverrideGpuDispatchNumThreads;
+					SimStageData.OverrideGpuDispatchNumThreads	= GenericStage->OverrideGpuDispatchNumThreads;
 
 					if (SimStageData.bParticleIterationStateEnabled)
 					{
diff --git a/Engine/Plugins/FX/Niagara/Source/NiagaraEditor/Private/NiagaraHlslTranslator.cpp b/Engine/Plugins/FX/Niagara/Source/NiagaraEditor/Private/NiagaraHlslTranslator.cpp
--- a/Engine/Plugins/FX/Niagara/Source/NiagaraEditor/Private/NiagaraHlslTranslator.cpp
+++ b/Engine/Plugins/FX/Niagara/Source/NiagaraEditor/Private/NiagaraHlslTranslator.cpp
@@ -1161,6 +1161,8 @@ const FNiagaraTranslateResults &FHlslNiagaraTranslator::Translate(const FNiagara
 			SimulationStageMetaData.NumIterations = 1;
 			SimulationStageMetaData.bWritesParticles = true;
 			SimulationStageMetaData.bPartialParticleUpdate = false;
+			SimulationStageMetaData.GpuDispatchType = ENiagaraGpuDispatchType::OneD;
+			SimulationStageMetaData.GpuDispatchNumThreads = FNiagaraShader::GetDefaultThreadGroupSize(ENiagaraGpuDispatchType::OneD);
 		}
 
 		{
@@ -1223,6 +1225,9 @@ const FNiagaraTranslateResults &FHlslNiagaraTranslator::Translate(const FNiagara
 					TranslationStages[Index].bParticleIterationStateEnabled = CompileSimStageData.bParticleIterationStateEnabled;
 					TranslationStages[Index].ParticleIterationStateBinding = CompileSimStageData.ParticleIterationStateBinding;
 					TranslationStages[Index].ParticleIterationStateRange = CompileSimStageData.ParticleIterationStateRange;
+					TranslationStages[Index].bGpuDispatchForceLinear = CompileSimStageData.bGpuDispatchForceLinear;
+					TranslationStages[Index].bOverrideGpuDispatchNumThreads = CompileSimStageData.bOverrideGpuDispatchNumThreads;
+					TranslationStages[Index].OverrideGpuDispatchNumThreads = CompileSimStageData.OverrideGpuDispatchNumThreads;
 
 					ParamMapHistories.AddDefaulted(1);
 					ParamMapHistoriesSourceInOtherHistories.AddDefaulted(1);
@@ -1292,14 +1297,24 @@ const FNiagaraTranslateResults &FHlslNiagaraTranslator::Translate(const FNiagara
 					SimulationStageMetaData.ParticleIterationStateRange = TranslationStages[Index].ParticleIterationStateRange;
 
 					// Determine dispatch information from iteration source (if we have one)
+					SimulationStageMetaData.GpuDispatchType = ENiagaraGpuDispatchType::OneD;
+					SimulationStageMetaData.GpuDispatchNumThreads = FNiagaraShader::GetDefaultThreadGroupSize(ENiagaraGpuDispatchType::OneD);
 					if ( !SimulationStageMetaData.IterationSource.IsNone() )
 					{
 						if (const FNiagaraVariable* IterationSourceVar = CompileData->EncounteredVariables.FindByPredicate([&](const FNiagaraVariable& VarInfo) { return VarInfo.GetName() == SimulationStageMetaData.IterationSource; }))
 						{
 							if ( UNiagaraDataInterface* IteratinoSourceCDO = CompileDuplicateData->GetDuplicatedDataInterfaceCDOForClass(IterationSourceVar->GetType().GetClass()) )
 							{
-								SimulationStageMetaData.GpuDispatchType = IteratinoSourceCDO->GetGpuDispatchType();
-								SimulationStageMetaData.GpuDispatchNumThreads = IteratinoSourceCDO->GetGpuDispatchNumThreads();
+								if ( TranslationStages[Index].bOverrideGpuDispatchNumThreads )
+								{
+									SimulationStageMetaData.GpuDispatchType = TranslationStages[Index].bGpuDispatchForceLinear ? ENiagaraGpuDispatchType::OneD : ENiagaraGpuDispatchType::Custom;
+									SimulationStageMetaData.GpuDispatchNumThreads = TranslationStages[Index].OverrideGpuDispatchNumThreads;
+								}
+								else
+								{
+									SimulationStageMetaData.GpuDispatchType = TranslationStages[Index].bGpuDispatchForceLinear ? ENiagaraGpuDispatchType::OneD : IteratinoSourceCDO->GetGpuDispatchType();
+									SimulationStageMetaData.GpuDispatchNumThreads = SimulationStageMetaData.GpuDispatchType == ENiagaraGpuDispatchType::Custom ? IteratinoSourceCDO->GetGpuDispatchNumThreads() : FNiagaraShader::GetDefaultThreadGroupSize(SimulationStageMetaData.GpuDispatchType);
+								}
 							}
 						}
 					}
diff --git a/Engine/Plugins/FX/Niagara/Source/NiagaraEditor/Private/NiagaraHlslTranslator.h b/Engine/Plugins/FX/Niagara/Source/NiagaraEditor/Private/NiagaraHlslTranslator.h
--- a/Engine/Plugins/FX/Niagara/Source/NiagaraEditor/Private/NiagaraHlslTranslator.h
+++ b/Engine/Plugins/FX/Niagara/Source/NiagaraEditor/Private/NiagaraHlslTranslator.h
@@ -100,6 +100,9 @@ public:
 		bool							bParticleIterationStateEnabled = false;
 		FName							ParticleIterationStateBinding;
 		FIntPoint						ParticleIterationStateRange = FIntPoint::ZeroValue;
+		bool							bGpuDispatchForceLinear = false;
+		bool							bOverrideGpuDispatchNumThreads = false;
+		FIntVector						OverrideGpuDispatchNumThreads = FIntVector(1, 1, 1);
 	};
 	TArray<FCompileSimStageData> CompileSimStageData;
 
@@ -379,6 +382,9 @@ public:
 	bool bWritesAlive = false;
 	bool bWritesParticles = false;
 	bool bPartialParticleUpdate = false;
+	bool bGpuDispatchForceLinear = false;
+	bool bOverrideGpuDispatchNumThreads = false;
+	FIntVector OverrideGpuDispatchNumThreads = FIntVector(1, 1, 1);
 	TArray<FNiagaraVariable> SetParticleAttributes;
 	FString CustomReadFunction;
 	FString CustomWriteFunction;
diff --git a/Engine/Plugins/FX/Niagara/Source/NiagaraShader/Private/NiagaraShader.cpp b/Engine/Plugins/FX/Niagara/Source/NiagaraShader/Private/NiagaraShader.cpp
--- a/Engine/Plugins/FX/Niagara/Source/NiagaraShader/Private/NiagaraShader.cpp
+++ b/Engine/Plugins/FX/Niagara/Source/NiagaraShader/Private/NiagaraShader.cpp
@@ -432,16 +432,7 @@ void FNiagaraShaderType::BeginCompileShader(
 
 	NewJob->Input.Environment.SetDefine(TEXT("NIAGARA_COMPRESSED_ATTRIBUTES_ENABLED"), Script->GetUsesCompressedAttributes() ? 1 : 0);
 
-	FIntVector ThreadGroupSize = FNiagaraShader::GetDefaultThreadGroupSize(ENiagaraGpuDispatchType::OneD);
-	if (SimStageMetaData.GpuDispatchType != ENiagaraGpuDispatchType::Custom)
-	{
-		ThreadGroupSize = FNiagaraShader::GetDefaultThreadGroupSize(SimStageMetaData.GpuDispatchType);
-	}
-	else
-	{
-		ThreadGroupSize = SimStageMetaData.GpuDispatchNumThreads;
-	}
-
+	const FIntVector ThreadGroupSize = SimStageMetaData.GpuDispatchNumThreads;
 	NewJob->Input.Environment.SetDefine(TEXT("NIAGARA_DISPATCH_TYPE"), int(SimStageMetaData.GpuDispatchType));
 	NewJob->Input.Environment.SetDefine(TEXT("THREADGROUP_SIZE"), ThreadGroupSize.X * ThreadGroupSize.Y * ThreadGroupSize.Z);
 	NewJob->Input.Environment.SetDefine(TEXT("THREADGROUP_SIZE_X"), ThreadGroupSize.X);
