Fixed coloring of the source filter in Niagara menus

diff --git a/Engine/Plugins/FX/Niagara/Source/NiagaraEditor/Private/NiagaraEditorStyle.cpp b/Engine/Plugins/FX/Niagara/Source/NiagaraEditor/Private/NiagaraEditorStyle.cpp
--- a/Engine/Plugins/FX/Niagara/Source/NiagaraEditor/Private/NiagaraEditorStyle.cpp
+++ b/Engine/Plugins/FX/Niagara/Source/NiagaraEditor/Private/NiagaraEditorStyle.cpp
@@ -195,6 +195,8 @@ void InitActionMenu(TSharedRef< FSlateStyleSet > Style)
 
 	Style->Set("ActionMenu.ActionTextBlock", ActionMenuActionText);
 
+	Style->Set("GraphActionMenu.Background", new FSlateRoundedBoxBrush(FStyleColors::Background, 3.f));
+	
 	Style->Set("GraphActionMenu.ActionSourceTextBlock", ActionMenuSourceText);
 
 	Style->Set("GraphActionMenu.ActionFilterTextBlock", ActionMenuFilterText);
@@ -202,6 +204,8 @@ void InitActionMenu(TSharedRef< FSlateStyleSet > Style)
 	Style->Set("GraphActionMenu.TemplateTabTextBlock", TemplateTabText);
 	
 	Style->Set("GraphActionMenu.FilterCheckBox", NiagaraGraphActionMenuFilterCheckBox);
+
+	Style->Set("GraphActionMenu.FilterCheckBox.Border", new FSlateRoundedBoxBrush(FStyleColors::Secondary, 3.f));
 }
 
 void InitEmitterHeader(TSharedRef< FSlateStyleSet > Style)
diff --git a/Engine/Plugins/FX/Niagara/Source/NiagaraEditor/Private/Widgets/SNiagaraFilterBox.cpp b/Engine/Plugins/FX/Niagara/Source/NiagaraEditor/Private/Widgets/SNiagaraFilterBox.cpp
--- a/Engine/Plugins/FX/Niagara/Source/NiagaraEditor/Private/Widgets/SNiagaraFilterBox.cpp
+++ b/Engine/Plugins/FX/Niagara/Source/NiagaraEditor/Private/Widgets/SNiagaraFilterBox.cpp
@@ -3,9 +3,12 @@
 #include "Widgets/SNiagaraFilterBox.h"
 #include "NiagaraEditorStyle.h"
 #include "NiagaraEditorUtilities.h"
+#include "Styling/StyleColors.h"
 #include "ViewModels/NiagaraCurveSelectionViewModel.h"
 #include "Widgets/Layout/SGridPanel.h"
 #include "Widgets/Layout/SExpandableArea.h"
+#include "Widgets/Images/SImage.h"
+#include "Widgets/SBoxPanel.h"
 
 #define LOCTEXT_NAMESPACE "NiagaraFilter"
 
@@ -32,7 +35,6 @@ void SNiagaraSourceFilterCheckBox::Construct(const FArguments& Args, EScriptSour
 	
 	SCheckBox::FArguments ParentArgs;
 	ParentArgs
-	.ForegroundColor(this, &SNiagaraSourceFilterCheckBox::GetBackgroundColor)
 	.Style(FEditorStyle::Get(), "ContentBrowser.FilterButton")
 	.IsChecked(Args._IsChecked)
 	.OnCheckStateChanged(FOnCheckStateChanged::CreateLambda([=](ECheckBoxState NewState)
@@ -43,11 +45,31 @@ void SNiagaraSourceFilterCheckBox::Construct(const FArguments& Args, EScriptSour
 
 	SetToolTipText(ToolTipText);
 	
-	SetContent(		
-        SNew(STextBlock)
-        .Text(DisplayName)
-        .ColorAndOpacity_Raw(this, &SNiagaraSourceFilterCheckBox::GetTextColor)
-        .TextStyle(FNiagaraEditorStyle::Get(), "GraphActionMenu.ActionFilterTextBlock")	     
+	SetContent(
+		SNew(SBorder)
+		 .Padding(1.0f)
+		.BorderImage(FNiagaraEditorStyle::Get().GetBrush("GraphActionMenu.FilterCheckBox.Border"))
+		[
+			SNew(SHorizontalBox)
+			+SHorizontalBox::Slot()
+			.VAlign(VAlign_Center)
+			.AutoWidth()
+			[
+				SNew(SImage)
+				.Image(FAppStyle::Get().GetBrush("ContentBrowser.FilterImage"))
+				.ColorAndOpacity(this, &SNiagaraSourceFilterCheckBox::GetScriptSourceColor)
+			]
+			+SHorizontalBox::Slot()
+			.Padding(TAttribute<FMargin>(this, &SNiagaraSourceFilterCheckBox::GetFilterNamePadding))
+			.VAlign(VAlign_Center)
+			.AutoWidth()
+			[
+				SNew(STextBlock)
+				.Text(DisplayName)
+				.ColorAndOpacity_Raw(this, &SNiagaraSourceFilterCheckBox::GetTextColor)
+				.TextStyle(FNiagaraEditorStyle::Get(), "GraphActionMenu.ActionFilterTextBlock")
+			]
+		]
     );
 }
 
@@ -70,16 +92,21 @@ FSlateColor SNiagaraSourceFilterCheckBox::GetTextColor() const
 	return IsHovered() ? FLinearColor(DimFactor, DimFactor, DimFactor, 1.0f) : IsChecked() ? FLinearColor::White : FLinearColor::Gray;
 }
 
-FSlateColor SNiagaraSourceFilterCheckBox::GetBackgroundColor() const
+FSlateColor SNiagaraSourceFilterCheckBox::GetScriptSourceColor() const
 {
-	if(IsChecked())
+	FLinearColor ScriptSourceColor = FNiagaraEditorUtilities::GetScriptSourceColor(Source);
+	if(!IsChecked())
 	{
-		return FNiagaraEditorUtilities::GetScriptSourceColor(Source);
+		ScriptSourceColor.A = 0.1f;
 	}
 
-	return FLinearColor::Gray;
+	return ScriptSourceColor;
 }
 
+FMargin SNiagaraSourceFilterCheckBox::GetFilterNamePadding() const
+{
+	return bIsPressed ? FMargin(2,2,1,0) : FMargin(2,1,1,1);	
+}
 
 void SNiagaraSourceFilterBox::Construct(const FArguments& Args)
 {
@@ -94,7 +121,7 @@ void SNiagaraSourceFilterBox::Construct(const FArguments& Args)
     .Padding(5.f)
     [
         SNew(SBorder)
-        .BorderImage(FEditorStyle::GetBrush(TEXT("NoBorder")))
+        .BorderImage(FNiagaraEditorStyle::Get().GetBrush("GraphActionMenu.FilterCheckBox.Border"))
         .ToolTipText(LOCTEXT("ShowAllToolTip", "Show all"))
         .Padding(3.f)
         [
diff --git a/Engine/Plugins/FX/Niagara/Source/NiagaraEditor/Public/Widgets/SNiagaraFilterBox.h b/Engine/Plugins/FX/Niagara/Source/NiagaraEditor/Public/Widgets/SNiagaraFilterBox.h
--- a/Engine/Plugins/FX/Niagara/Source/NiagaraEditor/Public/Widgets/SNiagaraFilterBox.h
+++ b/Engine/Plugins/FX/Niagara/Source/NiagaraEditor/Public/Widgets/SNiagaraFilterBox.h
@@ -36,7 +36,8 @@ private:
 	FOnShiftClicked OnShiftClicked;	
 private:
 	FSlateColor GetTextColor() const;
-	FSlateColor GetBackgroundColor() const;
+	FSlateColor GetScriptSourceColor() const;
+	FMargin GetFilterNamePadding() const;
 };
 
 /**
