- Remove unused shader bindings

#rnx

diff --git a/Engine/Plugins/FX/Niagara/Source/NiagaraShader/Private/NiagaraShader.cpp b/Engine/Plugins/FX/Niagara/Source/NiagaraShader/Private/NiagaraShader.cpp
--- a/Engine/Plugins/FX/Niagara/Source/NiagaraShader/Private/NiagaraShader.cpp
+++ b/Engine/Plugins/FX/Niagara/Source/NiagaraShader/Private/NiagaraShader.cpp
@@ -1245,31 +1245,6 @@ void FNiagaraShader::BindParams(const TArray<FNiagaraDataInterfaceGPUParamInfo>&
 	EmitterConstantBufferParam[1].Bind(ParameterMap, TEXT("PREV_FNiagaraEmitterParameters"));
 	ExternalConstantBufferParam[1].Bind(ParameterMap, TEXT("PREV_FNiagaraExternalParameters"));
 
-	// params for event buffers
-	// this is horrendous; need to do this in a uniform buffer instead.
-	//
-	for (uint32 i = 0; i < MAX_CONCURRENT_EVENT_DATASETS; i++)
-	{
-		FString VarName = TEXT("WriteDataSetFloat") + FString::FromInt(i + 1);
-		EventFloatUAVParams[i].Bind(ParameterMap, *VarName);
-		VarName = TEXT("WriteDataSetInt") + FString::FromInt(i + 1);
-		EventIntUAVParams[i].Bind(ParameterMap, *VarName);
-
-		VarName = TEXT("ReadDataSetFloat") + FString::FromInt(i + 1);
-		EventFloatSRVParams[i].Bind(ParameterMap, *VarName);
-		VarName = TEXT("ReadDataSetInt") + FString::FromInt(i + 1);
-		EventIntSRVParams[i].Bind(ParameterMap, *VarName);
-
-		VarName = TEXT("DSComponentBufferSizeReadFloat") + FString::FromInt(i + 1);
-		EventReadFloatStrideParams[i].Bind(ParameterMap, *VarName);
-		VarName = TEXT("DSComponentBufferSizeWriteFloat") + FString::FromInt(i + 1);
-		EventWriteFloatStrideParams[i].Bind(ParameterMap, *VarName);
-		VarName = TEXT("DSComponentBufferSizeReadInt") + FString::FromInt(i + 1);
-		EventReadIntStrideParams[i].Bind(ParameterMap, *VarName);
-		VarName = TEXT("DSComponentBufferSizeWriteInt") + FString::FromInt(i + 1);
-		EventWriteIntStrideParams[i].Bind(ParameterMap, *VarName);
-	}
-
 	DataInterfaceParameters.Empty(InDIParamInfo.Num());
 	for(const FNiagaraDataInterfaceGPUParamInfo& DIParamInfo : InDIParamInfo)
 	{
diff --git a/Engine/Plugins/FX/Niagara/Source/NiagaraShader/Public/NiagaraShader.h b/Engine/Plugins/FX/Niagara/Source/NiagaraShader/Public/NiagaraShader.h
--- a/Engine/Plugins/FX/Niagara/Source/NiagaraShader/Public/NiagaraShader.h
+++ b/Engine/Plugins/FX/Niagara/Source/NiagaraShader/Public/NiagaraShader.h
@@ -103,14 +103,6 @@ public:
 
 	LAYOUT_FIELD(FShaderParameter, ComponentBufferSizeReadParam);
 	LAYOUT_FIELD(FShaderParameter, ComponentBufferSizeWriteParam);
-	LAYOUT_ARRAY(FRWShaderParameter, EventIntUAVParams, MAX_CONCURRENT_EVENT_DATASETS);
-	LAYOUT_ARRAY(FRWShaderParameter, EventFloatUAVParams, MAX_CONCURRENT_EVENT_DATASETS);
-	LAYOUT_ARRAY(FShaderResourceParameter, EventIntSRVParams, MAX_CONCURRENT_EVENT_DATASETS);
-	LAYOUT_ARRAY(FShaderResourceParameter, EventFloatSRVParams, MAX_CONCURRENT_EVENT_DATASETS);
-	LAYOUT_ARRAY(FShaderParameter, EventWriteFloatStrideParams, MAX_CONCURRENT_EVENT_DATASETS);
-	LAYOUT_ARRAY(FShaderParameter, EventWriteIntStrideParams, MAX_CONCURRENT_EVENT_DATASETS);
-	LAYOUT_ARRAY(FShaderParameter, EventReadFloatStrideParams, MAX_CONCURRENT_EVENT_DATASETS);
-	LAYOUT_ARRAY(FShaderParameter, EventReadIntStrideParams, MAX_CONCURRENT_EVENT_DATASETS);
 
 private:
 	LAYOUT_FIELD(FShaderUniformBufferParameter, NiagaraUniformBuffer);
diff --git a/Engine/Plugins/FX/Niagara/Source/NiagaraShader/Public/NiagaraShared.h b/Engine/Plugins/FX/Niagara/Source/NiagaraShader/Public/NiagaraShared.h
--- a/Engine/Plugins/FX/Niagara/Source/NiagaraShader/Public/NiagaraShared.h
+++ b/Engine/Plugins/FX/Niagara/Source/NiagaraShader/Public/NiagaraShared.h
@@ -32,8 +32,6 @@ class FNiagaraShaderMapId;
 class UNiagaraScriptBase;
 struct FNiagaraVMExecutableDataId;
 
-#define MAX_CONCURRENT_EVENT_DATASETS 4
-
 /** Defines the compile event types for translation/compilation.*/
 UENUM()
 enum class FNiagaraCompileEventSeverity : uint8
