Reinitializing systems correctly on undo to fix crash.

diff --git a/Engine/Plugins/FX/Niagara/Source/NiagaraEditor/Private/ViewModels/NiagaraSystemViewModel.cpp b/Engine/Plugins/FX/Niagara/Source/NiagaraEditor/Private/ViewModels/NiagaraSystemViewModel.cpp
--- a/Engine/Plugins/FX/Niagara/Source/NiagaraEditor/Private/ViewModels/NiagaraSystemViewModel.cpp
+++ b/Engine/Plugins/FX/Niagara/Source/NiagaraEditor/Private/ViewModels/NiagaraSystemViewModel.cpp
@@ -642,6 +642,12 @@ void FNiagaraSystemViewModel::AddReferencedObjects(FReferenceCollector& Collecto
 void FNiagaraSystemViewModel::PostUndo(bool bSuccess)
 {
 	// Reset system stack and emitter handle view models to prevent accessing invalid data if they were in the undo operation.
+
+	//We have to force a reinit of all systems in the event of an undo.
+	FNiagaraSystemUpdateContext UpdateContext;
+	UpdateContext.SetDestroyOnAdd(true);
+	UpdateContext.Add(System, true);
+
 	SystemStackViewModel->Reset();
 	ResetEmitterHandleViewModelsAndTracks();
 	RefreshAll();
