Changed add emitter to be a button instead of a submenu
Works via command now. Default is E(mitter). Auto focus on the search box.
Fixed position of new emitter nodes being defaulted; instead it gets added where the user clicked or pressed E at.

diff --git a/Engine/Plugins/FX/Niagara/Source/NiagaraEditor/Private/NiagaraEditorCommands.cpp b/Engine/Plugins/FX/Niagara/Source/NiagaraEditor/Private/NiagaraEditorCommands.cpp
--- a/Engine/Plugins/FX/Niagara/Source/NiagaraEditor/Private/NiagaraEditorCommands.cpp
+++ b/Engine/Plugins/FX/Niagara/Source/NiagaraEditor/Private/NiagaraEditorCommands.cpp
@@ -17,6 +17,8 @@ void FNiagaraEditorCommands::RegisterCommands()
 	UI_COMMAND(SelectNextUsage, "Select Next Usage", "Selects the next usage of the selected item.", EUserInterfaceActionType::Button, FInputChord());
 	UI_COMMAND(CreateAssetFromSelection, "Create Asset...", "Creates an asset from the current selection.", EUserInterfaceActionType::Button, FInputChord());
 
+	UI_COMMAND(OpenAddEmitterMenu, "Add Emitter", "Adds an emitter to the system.", EUserInterfaceActionType::Button, FInputChord(EModifierKey::None, EKeys::E));
+
 	UI_COMMAND(ResetSimulation, "Reset", "Resets the current simulation", EUserInterfaceActionType::Button, FInputChord(EKeys::SpaceBar));
 	UI_COMMAND(TogglePreviewGrid, "Grid", "Toggles the preview pane's grid.", EUserInterfaceActionType::ToggleButton, FInputChord() );
 	UI_COMMAND(ToggleInstructionCounts, "InstructionCounts", "Display Instruction Counts", EUserInterfaceActionType::ToggleButton, FInputChord());
diff --git a/Engine/Plugins/FX/Niagara/Source/NiagaraEditor/Private/NiagaraSystemEditorData.cpp b/Engine/Plugins/FX/Niagara/Source/NiagaraEditor/Private/NiagaraSystemEditorData.cpp
--- a/Engine/Plugins/FX/Niagara/Source/NiagaraEditor/Private/NiagaraSystemEditorData.cpp
+++ b/Engine/Plugins/FX/Niagara/Source/NiagaraEditor/Private/NiagaraSystemEditorData.cpp
@@ -280,9 +280,16 @@ void UNiagaraSystemEditorData::SynchronizeOverviewGraphWithSystem(UNiagaraSystem
 			UNiagaraOverviewNode* EmitterOverviewNode = EmitterOverviewNodeCreator.CreateNode(false);
 			EmitterOverviewNode->Initialize(&OwnerSystem, EmitterHandle.GetId());
 
-			FVector2D EmitterOverviewNodeLocation = RightNode != nullptr
-				? FVector2D(RightNode->NodePosX + SystemOverviewNodePadding, RightNode->NodePosY)
-				: SystemOverviewGraph->GetGoodPlaceForNewNode();
+			FVector2D CurrentPos(EmitterOverviewNode->NodePosX, EmitterOverviewNode->NodePosY);
+			FVector2D EmitterOverviewNodeLocation = CurrentPos;
+
+			// callbacks in CreateNode might have set an explicit position already. If so, we don't update the position
+			if(EmitterOverviewNodeLocation.IsZero())
+			{
+				EmitterOverviewNodeLocation = RightNode != nullptr
+					? FVector2D(RightNode->NodePosX + SystemOverviewNodePadding, RightNode->NodePosY)
+					: SystemOverviewGraph->GetGoodPlaceForNewNode();
+			}
 
 			checkf(EmitterOverviewNode != nullptr, TEXT("Emitter overview node creation failed!"));
 			EmitterOverviewNode->NodePosX = EmitterOverviewNodeLocation.X;
diff --git a/Engine/Plugins/FX/Niagara/Source/NiagaraEditor/Public/NiagaraEditorCommands.h b/Engine/Plugins/FX/Niagara/Source/NiagaraEditor/Public/NiagaraEditorCommands.h
--- a/Engine/Plugins/FX/Niagara/Source/NiagaraEditor/Public/NiagaraEditorCommands.h
+++ b/Engine/Plugins/FX/Niagara/Source/NiagaraEditor/Public/NiagaraEditorCommands.h
@@ -32,6 +32,8 @@ public:
 	TSharedPtr<FUICommandInfo> SelectNextUsage;
 	TSharedPtr<FUICommandInfo> CreateAssetFromSelection;
 
+	TSharedPtr<FUICommandInfo> OpenAddEmitterMenu;
+	
 	/** Toggles the preview pane's grid */
 	TSharedPtr<FUICommandInfo> TogglePreviewGrid;
 	TSharedPtr<FUICommandInfo> ToggleInstructionCounts;
diff --git a/Engine/Plugins/FX/Niagara/Source/NiagaraEditor/Public/Widgets/SNiagaraScriptGraph.h b/Engine/Plugins/FX/Niagara/Source/NiagaraEditor/Public/Widgets/SNiagaraScriptGraph.h
--- a/Engine/Plugins/FX/Niagara/Source/NiagaraEditor/Public/Widgets/SNiagaraScriptGraph.h
+++ b/Engine/Plugins/FX/Niagara/Source/NiagaraEditor/Public/Widgets/SNiagaraScriptGraph.h
@@ -17,7 +17,7 @@ class SNiagaraScriptGraph : public SCompoundWidget
 {
 public:
 	SLATE_BEGIN_ARGS(SNiagaraScriptGraph)
-		: _ForegroundColor(FEditorStyle::GetColor("DefaultForeground"))
+		: _ForegroundColor(FEditorStyle::GetColor("Graph.ForegroundColor"))
 		, _ZoomToFitOnLoad(false)
 	{}
 		/** The text displayed in the title bar of the graph. */
diff --git a/Engine/Plugins/FX/Niagara/Source/NiagaraEditorWidgets/Private/SNiagaraOverviewGraph.cpp b/Engine/Plugins/FX/Niagara/Source/NiagaraEditorWidgets/Private/SNiagaraOverviewGraph.cpp
--- a/Engine/Plugins/FX/Niagara/Source/NiagaraEditorWidgets/Private/SNiagaraOverviewGraph.cpp
+++ b/Engine/Plugins/FX/Niagara/Source/NiagaraEditorWidgets/Private/SNiagaraOverviewGraph.cpp
@@ -13,6 +13,7 @@
 #include "NiagaraObjectSelection.h"
 #include "ViewModels/NiagaraOverviewGraphViewModel.h"
 #include "EdGraphSchema_Niagara.h"
+#include "GraphEditAction.h"
 #include "NiagaraEditorCommands.h"
 #include "NiagaraEditorModule.h"
 #include "GraphEditorActions.h"
@@ -103,7 +104,13 @@
 		FExecuteAction::CreateSP(this, &SNiagaraOverviewGraph::OnDistributeNodesV)
 	);
 
-
+	Commands->MapAction(FNiagaraEditorCommands::Get().OpenAddEmitterMenu,
+		FExecuteAction::CreateSP(this, &SNiagaraOverviewGraph::OpenAddEmitterMenu),
+		FCanExecuteAction::CreateSP(this, &SNiagaraOverviewGraph::CanAddEmitters),
+		FIsActionChecked(),
+		FIsActionButtonVisible::CreateSP(this, &SNiagaraOverviewGraph::CanAddEmitters)
+	);
+	
 	GraphEditor = SNew(SGraphEditor)
 		.AdditionalCommands(Commands)
 		.Appearance(AppearanceInfo)
@@ -114,6 +121,8 @@
 
 	GraphEditor->SetNodeFactory(MakeShared<FNiagaraOverviewGraphNodeFactory>());
 
+	GraphEditor->GetCurrentGraph()->AddOnGraphChangedHandler(FOnGraphChanged::FDelegate::CreateSP(this, &SNiagaraOverviewGraph::OnNodesCreated));
+	
 	FNiagaraGraphViewSettings ViewSettings = ViewModel->GetViewSettings();
 	if (ViewSettings.IsValid())
 	{
@@ -193,6 +202,54 @@
 	}
 }
 
+void SNiagaraOverviewGraph::OpenAddEmitterMenu()
+{
+	FNiagaraAssetPickerListViewOptions ViewOptions;
+	ViewOptions.SetCategorizeUserDefinedCategory(true);
+	ViewOptions.SetCategorizeLibraryAssets(true);
+
+	SNiagaraTemplateTabBox::FNiagaraTemplateTabOptions TabOptions;
+	TabOptions.ChangeTabState(ENiagaraScriptTemplateSpecification::Template, true);
+	TabOptions.ChangeTabState(ENiagaraScriptTemplateSpecification::None, true);
+	TabOptions.ChangeTabState(ENiagaraScriptTemplateSpecification::Behavior, true);	
+
+	TSharedPtr<SNiagaraAssetPickerList> AssetPickerList;
+	
+	SAssignNew(AssetPickerList, SNiagaraAssetPickerList, UNiagaraEmitter::StaticClass())
+	.ClickActivateMode(EItemSelectorClickActivateMode::SingleClick)
+	.ViewOptions(ViewOptions)
+	.TabOptions(TabOptions)
+	.OnTemplateAssetActivated_Lambda([this](const FAssetData& AssetData) {
+		FSlateApplication::Get().DismissAllMenus();
+		ViewModel->GetSystemViewModel()->AddEmitterFromAssetData(AssetData);
+	});
+	
+	TSharedRef<SWidget> EmitterAddSubMenu =
+		SNew(SBorder)
+		.BorderImage(FNiagaraEditorStyle::Get().GetBrush("GraphActionMenu.Background"))
+		.Padding(0.0f)
+		[
+			SNew(SBox)
+			.WidthOverride(450.0f)
+			.HeightOverride(500.0f)
+			[
+				AssetPickerList.ToSharedRef()
+			]
+		];
+
+	FSlateApplication::Get().PushMenu(SharedThis(this), FWidgetPath(), EmitterAddSubMenu, FSlateApplication::Get().GetCursorPos(), FPopupTransitionEffect::None);
+
+	GEditor->GetTimerManager()->SetTimerForNextTick(FTimerDelegate::CreateLambda([AssetPickerList]()
+	{
+		FSlateApplication::Get().SetKeyboardFocus(AssetPickerList->GetSearchBox());
+	}));
+}
+
+bool SNiagaraOverviewGraph::CanAddEmitters() const
+{
+	return ViewModel->GetSystemViewModel()->GetEditMode() == ENiagaraSystemViewModelEditMode::SystemAsset ? true : false;
+}
+
 FActionMenuContent SNiagaraOverviewGraph::OnCreateGraphActionMenu(UEdGraph* InGraph, const FVector2D& InNodePosition, const TArray<UEdGraphPin*>& InDraggedPins, bool bAutoExpand, SGraphEditor::FActionMenuClosed InOnMenuClosed)
 {
 	if (ViewModel->GetSystemViewModel()->GetEditMode() == ENiagaraSystemViewModelEditMode::SystemAsset)
@@ -201,10 +258,7 @@
 
 		MenuBuilder.BeginSection(TEXT("NiagaraOverview_EditGraph"), LOCTEXT("EditGraph", "Edit Graph"));
 		{
-			MenuBuilder.AddSubMenu(
-				LOCTEXT("EmitterAddLabel", "Add Emitter"),
-				LOCTEXT("EmitterAddToolTip", "Add an existing emitter"),
-				FNewMenuDelegate::CreateSP(this, &SNiagaraOverviewGraph::CreateAddEmitterMenuContent, InGraph));
+			MenuBuilder.AddMenuEntry(FNiagaraEditorCommands::Get().OpenAddEmitterMenu);
 
 			MenuBuilder.AddMenuEntry(
 				LOCTEXT("CommentsLabel", "Add Comment"),
@@ -335,48 +389,6 @@
 	}
 }
 
-void SNiagaraOverviewGraph::CreateAddEmitterMenuContent(FMenuBuilder& MenuBuilder, UEdGraph* InGraph)
-{
-	TArray<FRefreshItemSelectorDelegate*> RefreshItemSelectorDelegates;
-	RefreshItemSelectorDelegates.Add(&RefreshItemSelector);
-	FNiagaraAssetPickerListViewOptions ViewOptions;
-	ViewOptions.SetCategorizeUserDefinedCategory(true);
-	ViewOptions.SetCategorizeLibraryAssets(true);
-
-	SNiagaraTemplateTabBox::FNiagaraTemplateTabOptions TabOptions;
-	TabOptions.ChangeTabState(ENiagaraScriptTemplateSpecification::Template, true);
-	TabOptions.ChangeTabState(ENiagaraScriptTemplateSpecification::None, true);
-	TabOptions.ChangeTabState(ENiagaraScriptTemplateSpecification::Behavior, true);	
-
-	TSharedPtr<SNiagaraAssetPickerList> AssetPickerList;
-	
-	SAssignNew(AssetPickerList, SNiagaraAssetPickerList, UNiagaraEmitter::StaticClass())
-	.ClickActivateMode(EItemSelectorClickActivateMode::SingleClick)
-	.ViewOptions(ViewOptions)
-	.TabOptions(TabOptions)
-	.RefreshItemSelectorDelegates(RefreshItemSelectorDelegates)
-	.OnTemplateAssetActivated_Lambda([this, InGraph](const FAssetData& AssetData) {
-		FSlateApplication::Get().DismissAllMenus();
-		ViewModel->GetSystemViewModel()->AddEmitterFromAssetData(AssetData);
-	});
-	
-	TSharedRef<SWidget> EmitterAddSubMenu =
-		SNew(SBorder)
-		.BorderImage(FEditorStyle::GetBrush("Menu.Background"))
-		.Padding(0.0f)
-		[
-			SNew(SBox)
-			.WidthOverride(450.0f)
-			.HeightOverride(500.0f)
-			[
-				AssetPickerList.ToSharedRef()
-			]
-		];
-
-	MenuBuilder.AddWidget(EmitterAddSubMenu, FText());
-	FSlateApplication::Get().SetKeyboardFocus(AssetPickerList->GetSearchBox());
-}
-
 void SNiagaraOverviewGraph::ZoomToFit()
 {
 	GraphEditor->ZoomToFit(true);
@@ -428,4 +440,26 @@
 	}
 }
 
+void SNiagaraOverviewGraph::OnNodesCreated(const FEdGraphEditAction& Action)
+{
+	if((Action.Action & GRAPHACTION_AddNode) != 0)
+	{
+		FVector2D PasteLocation = GraphEditor->GetPasteLocation();
+		int32 NodeCounter = 0;
+
+		for(const UEdGraphNode* NewNode : Action.Nodes)
+		{
+			UEdGraphNode* EditableNode = const_cast<UEdGraphNode*>(NewNode);
+			EditableNode->NodePosX = PasteLocation.X + NodeCounter * 300.f;
+			EditableNode->NodePosY = PasteLocation.Y;
+			NodeCounter++;
+		}
+
+		if(Action.Nodes.Num() == 1)
+		{
+			GraphEditor->JumpToNode(Action.Nodes.Array()[0]);
+		}
+	}
+}
+
 #undef LOCTEXT_NAMESPACE // "NiagaraOverviewGraph"
diff --git a/Engine/Plugins/FX/Niagara/Source/NiagaraEditorWidgets/Private/SNiagaraOverviewGraph.h b/Engine/Plugins/FX/Niagara/Source/NiagaraEditorWidgets/Private/SNiagaraOverviewGraph.h
--- a/Engine/Plugins/FX/Niagara/Source/NiagaraEditorWidgets/Private/SNiagaraOverviewGraph.h
+++ b/Engine/Plugins/FX/Niagara/Source/NiagaraEditorWidgets/Private/SNiagaraOverviewGraph.h
@@ -22,15 +22,15 @@ public:
 
 	virtual void Tick(const FGeometry& AllottedGeometry, const double InCurrentTime, const float InDeltaTime) override;
 
-public:
-	FRefreshItemSelectorDelegate RefreshItemSelector;
-
 private:
 	void ViewModelSelectionChanged();
 
 	void GraphSelectionChanged(const TSet<UObject*>& SelectedNodes);
 
 	void PreClose();
+	
+	void OpenAddEmitterMenu();
+	bool CanAddEmitters() const;
 
 	/** Called to create context menu when right-clicking on graph */
 	FActionMenuContent OnCreateGraphActionMenu(UEdGraph* InGraph, const FVector2D& InNodePosition, const TArray<UEdGraphPin*>& InDraggedPins, bool bAutoExpand, SGraphEditor::FActionMenuClosed InOnMenuClosed);
@@ -44,11 +44,12 @@ private:
 	/** Called when nodes are pasted in the System Overview View Model. */
 	void NodesPasted(const TSet<UEdGraphNode*>& PastedNodes);
 
+	/** A callback function to position nodes after they've been created */
+	void OnNodesCreated(const FEdGraphEditAction& Action);
+	
 	/** Sets the position on a group of newly pasted nodes. */
 	void PositionPastedNodes(const TSet<UEdGraphNode*>& PastedNodes);
 
-	void CreateAddEmitterMenuContent(FMenuBuilder& MenuBuilder, UEdGraph* InGraph);
-
 	void ZoomToFit();
 	void ZoomToFitAll();
 
