- Fix for uninitialized bounds

#rnx

diff --git a/Engine/Plugins/FX/Niagara/Source/Niagara/Private/NiagaraSystemInstance.cpp b/Engine/Plugins/FX/Niagara/Source/Niagara/Private/NiagaraSystemInstance.cpp
--- a/Engine/Plugins/FX/Niagara/Source/Niagara/Private/NiagaraSystemInstance.cpp
+++ b/Engine/Plugins/FX/Niagara/Source/Niagara/Private/NiagaraSystemInstance.cpp
@@ -2069,7 +2069,7 @@ void FNiagaraSystemInstance::InitEmitters()
 
 	bHasGPUEmitters = false;
 
-	LocalBounds = FBox();
+	LocalBounds = FBox(FVector::ZeroVector, FVector::ZeroVector);
 
 	Emitters.Empty(false);
 	UNiagaraSystem* System = GetSystem();
