Event writes from event handler scripts now throws a compile error instead of crashing.
Also removing Niagara's own event generator modules from the menus for event handlers.

diff --git a/Engine/Plugins/FX/Niagara/Source/NiagaraEditor/Private/NiagaraNodeWriteDataSet.cpp b/Engine/Plugins/FX/Niagara/Source/NiagaraEditor/Private/NiagaraNodeWriteDataSet.cpp
--- a/Engine/Plugins/FX/Niagara/Source/NiagaraEditor/Private/NiagaraNodeWriteDataSet.cpp
+++ b/Engine/Plugins/FX/Niagara/Source/NiagaraEditor/Private/NiagaraNodeWriteDataSet.cpp
@@ -82,7 +82,6 @@ void UNiagaraNodeWriteDataSet::Compile(class FHlslNiagaraTranslator* Translator,
 	TArray<int32> Inputs;
 	CompileInputPins(Translator, Inputs);
 
-
 	bool bGPUSim = Translator->IsCompileOptionDefined(TEXT("GPUComputeSim"));
 
 	if (bGPUSim)
@@ -90,6 +89,12 @@ void UNiagaraNodeWriteDataSet::Compile(class FHlslNiagaraTranslator* Translator,
 		Translator->Error(LOCTEXT("CannotRunWriteDataSetGPU", "Cannot use an event write node on GPU sims!"), this, nullptr);
 	}
 
+	bool bIsEventScript = UNiagaraScript::IsParticleEventScript(Translator->GetTargetUsage());
+	if (bIsEventScript)
+	{
+		Translator->Error(LOCTEXT("CannotRunWriteDataSetFromEventScript", "Cannot use an event write node in an Event Handler!"), this, nullptr);
+	}
+
 	FString IssuesWithStruct;
 	if (!IsSynchronizedWithStruct(true, &IssuesWithStruct,false))
 	{
