- Enable / disable the debug HUD and it's settings when you close the widget

#rnx

diff --git a/Engine/Plugins/FX/Niagara/Source/Niagara/Private/NiagaraDebuggerClient.cpp b/Engine/Plugins/FX/Niagara/Source/Niagara/Private/NiagaraDebuggerClient.cpp
--- a/Engine/Plugins/FX/Niagara/Source/Niagara/Private/NiagaraDebuggerClient.cpp
+++ b/Engine/Plugins/FX/Niagara/Source/Niagara/Private/NiagaraDebuggerClient.cpp
@@ -158,20 +158,24 @@ void FNiagaraDebuggerClient::HandleDebugHUDSettingsMessage(const FNiagaraDebugHU
 		UE_LOG(LogNiagaraDebuggerClient, Log, TEXT("Received updated DebugHUD settings. | Session: %s | Instance: %s (%s)."), *SessionId.ToString(), *InstanceId.ToString(), *InstanceName);
 
 		//Pass along the new settings.
-		auto ApplySettingsToWorldMan = [&Message](FNiagaraWorldManager& WorldMan)
-		{
-			WorldMan.GetNiagaraDebugHud()->UpdateSettings(Message);
+		auto ApplySettingsToWorldMan =
+			[&Message](FNiagaraWorldManager& WorldMan)
+			{
+				WorldMan.GetNiagaraDebugHud()->UpdateSettings(Message);
 
-			//TODO: Move these to just take direct from the debug hud per worldman?
-			//Possbly move the debug hud itself to the debugger client rather than having one per world man and they all share global state.
-			WorldMan.SetDebugPlaybackMode(Message.PlaybackMode);
-			WorldMan.SetDebugPlaybackRate(Message.bPlaybackRateEnabled ? Message.PlaybackRate : 1.0f);
-		};
+				//TODO: Move these to just take direct from the debug hud per worldman?
+				//Possibly move the debug hud itself to the debugger client rather than having one per world man and they all share global state.
+				const ENiagaraDebugPlaybackMode PlaybackMode = Message.bEnabled ? Message.PlaybackMode : ENiagaraDebugPlaybackMode::Play;
+				const float PlaybackRate = Message.bEnabled && Message.bPlaybackRateEnabled ? Message.PlaybackRate : 1.0f;
+				WorldMan.SetDebugPlaybackMode(PlaybackMode);
+				WorldMan.SetDebugPlaybackRate(PlaybackRate);
+			};
 
 		FNiagaraWorldManager::ForAllWorldManagers(ApplySettingsToWorldMan);
 			
 		//TODO: Move usage to come direct from settings struct instead of this CVar.
-		ExecuteConsoleCommand(*FString::Printf(TEXT("fx.Niagara.Debug.GlobalLoopTime %.3f"), Message.bLoopTimeEnabled && Message.PlaybackMode == ENiagaraDebugPlaybackMode::Loop ? Message.LoopTime : 0.0f), true);
+		const float GlobalLoopTime = Message.bEnabled && Message.bLoopTimeEnabled && Message.PlaybackMode == ENiagaraDebugPlaybackMode::Loop ? Message.LoopTime : 0.0f;
+		ExecuteConsoleCommand(*FString::Printf(TEXT("fx.Niagara.Debug.GlobalLoopTime %.3f"), GlobalLoopTime), true);
 	}
 }
 void FNiagaraDebuggerClient::HandleRequestSimpleClientInfoMessage(const FNiagaraRequestSimpleClientInfoMessage& Message, const TSharedRef<IMessageContext, ESPMode::ThreadSafe>& Context)
diff --git a/Engine/Plugins/FX/Niagara/Source/NiagaraEditor/Private/Widgets/SNiagaraDebugger.cpp b/Engine/Plugins/FX/Niagara/Source/NiagaraEditor/Private/Widgets/SNiagaraDebugger.cpp
--- a/Engine/Plugins/FX/Niagara/Source/NiagaraEditor/Private/Widgets/SNiagaraDebugger.cpp
+++ b/Engine/Plugins/FX/Niagara/Source/NiagaraEditor/Private/Widgets/SNiagaraDebugger.cpp
@@ -825,6 +825,10 @@ TSharedRef<SDockTab> SNiagaraDebugger::SpawnNiagaraDebugger(const FSpawnTabArgs&
 		SDockTab::FOnTabClosedCallback::CreateStatic(
 			[](TSharedRef<SDockTab> Self, TWeakPtr<FTabManager> TabManager)
 			{
+				UNiagaraDebugHUDSettings* HudSettings = GetMutableDefault<UNiagaraDebugHUDSettings>();
+				HudSettings->Data.bEnabled = false;
+				HudSettings->NotifyPropertyChanged();
+
 				TSharedPtr<FTabManager> OwningTabManager = TabManager.Pin();
 				if (OwningTabManager.IsValid())
 				{
@@ -836,6 +840,10 @@ TSharedRef<SDockTab> SNiagaraDebugger::SpawnNiagaraDebugger(const FSpawnTabArgs&
 		)
 	);
 
+	UNiagaraDebugHUDSettings* HudSettings = GetMutableDefault<UNiagaraDebugHUDSettings>();
+	HudSettings->Data.bEnabled = true;
+	HudSettings->NotifyPropertyChanged();
+
 	FNiagaraEditorModule& NiagaraEditorModule = FModuleManager::GetModuleChecked<FNiagaraEditorModule>("NiagaraEditor");
 
 	auto MainWidget = SNew(SNiagaraDebugger)
