Fix a bug when creating niagara clipboard function inputs for vector2 and vectors due to memcpy doubles to floats.

diff --git a/Engine/Plugins/FX/Niagara/Source/NiagaraEditor/Private/NiagaraClipboard.cpp b/Engine/Plugins/FX/Niagara/Source/NiagaraEditor/Private/NiagaraClipboard.cpp
--- a/Engine/Plugins/FX/Niagara/Source/NiagaraEditor/Private/NiagaraClipboard.cpp
+++ b/Engine/Plugins/FX/Niagara/Source/NiagaraEditor/Private/NiagaraClipboard.cpp
@@ -268,21 +268,23 @@ UNiagaraClipboardFunctionInput* UNiagaraClipboardEditorScriptingUtilities::Creat
 UNiagaraClipboardFunctionInput* UNiagaraClipboardEditorScriptingUtilities::CreateVec2LocalValueInput(UObject* InOuter, FName InInputName, bool bInHasEditCondition, bool bInEditConditionValue, FVector2D InVec2Value)
 {
 	FNiagaraTypeDefinition InputType = FNiagaraTypeDefinition::GetVec2Def();
-	TArray<uint8> Vec2Value;
-	Vec2Value.AddUninitialized(InputType.GetSize());
-	FMemory::Memcpy(Vec2Value.GetData(), &InVec2Value, InputType.GetSize());
+	const FVector2f InVec2ValueFloat = InVec2Value;
+	TArray<uint8> Vec2ValueFloat;
+	Vec2ValueFloat.AddUninitialized(InputType.GetSize());
+	FMemory::Memcpy(Vec2ValueFloat.GetData(), &InVec2ValueFloat, InputType.GetSize());
 
-	return const_cast<UNiagaraClipboardFunctionInput*>(CreateLocalValue(InOuter, InInputName, InputType, bInHasEditCondition, bInEditConditionValue, Vec2Value));
+	return const_cast<UNiagaraClipboardFunctionInput*>(CreateLocalValue(InOuter, InInputName, InputType, bInHasEditCondition, bInEditConditionValue, Vec2ValueFloat));
 }
 
 UNiagaraClipboardFunctionInput* UNiagaraClipboardEditorScriptingUtilities::CreateVec3LocalValueInput(UObject* InOuter, FName InInputName, bool bInHasEditCondition, bool bInEditConditionValue, FVector InVec3Value)
 {
 	FNiagaraTypeDefinition InputType = FNiagaraTypeDefinition::GetVec3Def();
-	TArray<uint8> Vec3Value;
-	Vec3Value.AddUninitialized(InputType.GetSize());
-	FMemory::Memcpy(Vec3Value.GetData(), &InVec3Value, InputType.GetSize());
+	const FVector3f InVec3ValueFloat = InVec3Value;
+	TArray<uint8> Vec3ValueFloat;
+	Vec3ValueFloat.AddUninitialized(InputType.GetSize());
+	FMemory::Memcpy(Vec3ValueFloat.GetData(), &InVec3ValueFloat, InputType.GetSize());
 
-	return const_cast<UNiagaraClipboardFunctionInput*>(CreateLocalValue(InOuter, InInputName, InputType, bInHasEditCondition, bInEditConditionValue, Vec3Value));
+	return const_cast<UNiagaraClipboardFunctionInput*>(CreateLocalValue(InOuter, InInputName, InputType, bInHasEditCondition, bInEditConditionValue, Vec3ValueFloat));
 }
 
 UNiagaraClipboardFunctionInput* UNiagaraClipboardEditorScriptingUtilities::CreateIntLocalValueInput(UObject* InOuter, FName InInputName, bool bInHasEditCondition, bool bInEditConditionValue, int32 InIntValue)
