Resetting all systems when we enter Looping mode in the Niagara Debugger.
Fixes issue where completed burst systems do not begin looping as expected.

diff --git a/Engine/Plugins/FX/Niagara/Source/NiagaraEditor/Private/Widgets/SNiagaraDebugger.cpp b/Engine/Plugins/FX/Niagara/Source/NiagaraEditor/Private/Widgets/SNiagaraDebugger.cpp
--- a/Engine/Plugins/FX/Niagara/Source/NiagaraEditor/Private/Widgets/SNiagaraDebugger.cpp
+++ b/Engine/Plugins/FX/Niagara/Source/NiagaraEditor/Private/Widgets/SNiagaraDebugger.cpp
@@ -922,7 +922,15 @@ TSharedRef<SWidget> SNiagaraDebugger::MakeToolbar()
 		{
 			ToolbarBuilder.AddToolBarButton(
 				FUIAction(
-					FExecuteAction::CreateLambda([=]() {Settings->Data.PlaybackMode = ENiagaraDebugPlaybackMode::Loop; Settings->NotifyPropertyChanged(); }),
+					FExecuteAction::CreateLambda([=]() 
+					{
+						Settings->Data.PlaybackMode = ENiagaraDebugPlaybackMode::Loop;
+						Settings->NotifyPropertyChanged(); 
+
+						//Reset all systems so that currently inactive systems will begin to loop.
+						FNiagaraSystemUpdateContext UpdateContext;
+						UpdateContext.AddAll(false);
+					}),
 					FCanExecuteAction(),
 					FIsActionChecked::CreateLambda([=]() { return Settings->Data.PlaybackMode == ENiagaraDebugPlaybackMode::Loop; })
 				),
