- Fix for missing & when adjust signature

#rnx

diff --git a/Engine/Plugins/FX/Niagara/Source/Niagara/Private/DataInterface/NiagaraDataInterfaceStaticMesh.cpp b/Engine/Plugins/FX/Niagara/Source/Niagara/Private/DataInterface/NiagaraDataInterfaceStaticMesh.cpp
--- a/Engine/Plugins/FX/Niagara/Source/Niagara/Private/DataInterface/NiagaraDataInterfaceStaticMesh.cpp
+++ b/Engine/Plugins/FX/Niagara/Source/Niagara/Private/DataInterface/NiagaraDataInterfaceStaticMesh.cpp
@@ -1704,7 +1704,7 @@ void UNiagaraDataInterfaceStaticMesh::GetFunctions(TArray<FNiagaraFunctionSignat
 		Sig.Outputs.Emplace(FNiagaraTypeDefinition::GetVec3Def(), TEXT("Tangent"));
 		OutFunctions.Add_GetRef(Sig).Name = GetTriangleName;
 
-		FNiagaraFunctionSignature WsSig = OutFunctions.Add_GetRef(Sig);
+		FNiagaraFunctionSignature& WsSig = OutFunctions.Add_GetRef(Sig);
 		WsSig.Name = GetTriangleWSName;
 		WsSig.Outputs[0].SetType(FNiagaraTypeDefinition::GetPositionDef());
 	}
