- Fix debug HUD and positions with LWC
- Fix LODDistance calculation with LWC

diff --git a/Engine/Plugins/FX/Niagara/Source/Niagara/Private/NiagaraDebugHud.cpp b/Engine/Plugins/FX/Niagara/Source/Niagara/Private/NiagaraDebugHud.cpp
--- a/Engine/Plugins/FX/Niagara/Source/Niagara/Private/NiagaraDebugHud.cpp
+++ b/Engine/Plugins/FX/Niagara/Source/Niagara/Private/NiagaraDebugHud.cpp
@@ -2286,9 +2286,10 @@ void FNiagaraDebugHud::DrawComponents(FNiagaraWorldManager* WorldManager, UCanva
 				const uint32 MaxDisplayParticles = Settings.bUseMaxParticlesToDisplay ? Settings.MaxParticlesToDisplay : 0xFFFFFFFF;
 				uint32 NumDisplayedParticles = 0;
 
+				FNiagaraLWCConverter LwcConverter = SystemInstance->GetLWCConverter(false);
 				for (uint32 iInstance = 0; iInstance < DataBuffer->GetNumInstances(); ++iInstance)
 				{
-					const FVector ParticalLocalPosition = PositionReader.Get(iInstance);
+					const FVector ParticalLocalPosition = LwcConverter.ConvertSimulationPositionToWorld(PositionReader.Get(iInstance));
 					const FVector ParticleWorldPosition = bParticlesLocalSpace ? SystemTransform.TransformPosition(ParticalLocalPosition) : ParticalLocalPosition;
 
 					const FPlane ClipPosition = SceneView->Project(ParticleWorldPosition);
diff --git a/Engine/Plugins/FX/Niagara/Source/Niagara/Private/NiagaraSystemInstance.cpp b/Engine/Plugins/FX/Niagara/Source/Niagara/Private/NiagaraSystemInstance.cpp
--- a/Engine/Plugins/FX/Niagara/Source/Niagara/Private/NiagaraSystemInstance.cpp
+++ b/Engine/Plugins/FX/Niagara/Source/Niagara/Private/NiagaraSystemInstance.cpp
@@ -1714,7 +1714,7 @@ float FNiagaraSystemInstance::GetLODDistance()
 	}
 
 	check(World);
-	const FVector EffectLocation = WorldTransform.GetLocation();
+	const FVector EffectLocation = WorldTransform.GetLocation() + (FVector(LWCTile) * FLargeWorldRenderScalar::GetTileSize());
 	LODDistance = DefaultLODDistance;
 
 	// If we are inside the WorldManager tick we will use the cache player view locations as we can be ticked on different threads
