Fixed crash where pasting invalid/empty clipboard to niagara scriptgraph would cause a crash.

diff --git a/Engine/Plugins/FX/Niagara/Source/NiagaraEditor/Private/ViewModels/NiagaraScriptGraphViewModel.cpp b/Engine/Plugins/FX/Niagara/Source/NiagaraEditor/Private/ViewModels/NiagaraScriptGraphViewModel.cpp
--- a/Engine/Plugins/FX/Niagara/Source/NiagaraEditor/Private/ViewModels/NiagaraScriptGraphViewModel.cpp
+++ b/Engine/Plugins/FX/Niagara/Source/NiagaraEditor/Private/ViewModels/NiagaraScriptGraphViewModel.cpp
@@ -332,6 +332,11 @@ void FNiagaraScriptGraphViewModel::PasteNodes()
 
 	const UNiagaraClipboardContent* ClipboardContent = FNiagaraEditorModule::Get().GetClipboard().GetClipboardContent();
 	
+	if (ClipboardContent == nullptr)
+	{
+		return;
+	}
+	
 	TSet<UEdGraphNode*> PastedNodes;
 	FEdGraphUtilities::ImportNodesFromText(Graph, ClipboardContent->ExportedNodes, PastedNodes);
 	
@@ -446,7 +451,7 @@ bool FNiagaraScriptGraphViewModel::CanPasteNodes() const
 	}
 
 	const UNiagaraClipboardContent* ClipboardContent = FNiagaraEditorModule::Get().GetClipboard().GetClipboardContent();
-	return FEdGraphUtilities::CanImportNodesFromText(Graph, ClipboardContent->ExportedNodes);
+	return ClipboardContent? FEdGraphUtilities::CanImportNodesFromText(Graph, ClipboardContent->ExportedNodes) : false;
 }
 
 void FNiagaraScriptGraphViewModel::DuplicateNodes()
diff --git a/Engine/Plugins/FX/Niagara/Source/NiagaraEditor/Private/ViewModels/Stack/NiagaraStackClipboardUtilities.cpp b/Engine/Plugins/FX/Niagara/Source/NiagaraEditor/Private/ViewModels/Stack/NiagaraStackClipboardUtilities.cpp
--- a/Engine/Plugins/FX/Niagara/Source/NiagaraEditor/Private/ViewModels/Stack/NiagaraStackClipboardUtilities.cpp
+++ b/Engine/Plugins/FX/Niagara/Source/NiagaraEditor/Private/ViewModels/Stack/NiagaraStackClipboardUtilities.cpp
@@ -189,7 +189,7 @@ void FNiagaraStackClipboardUtilities::PasteSelection(const TArray<UNiagaraStackE
 	{
 		const UNiagaraClipboardContent* ClipboardContent = FNiagaraEditorModule::Get().GetClipboard().GetClipboardContent();
 
-		FText TransactionMessage = SelectedEntries[0]->GetPasteTransactionText(ClipboardContent);
+		FText TransactionMessage = ClipboardContent? SelectedEntries[0]->GetPasteTransactionText(ClipboardContent) : FText();
 		if (TransactionMessage.IsEmptyOrWhitespace())
 		{
 			TransactionMessage = LOCTEXT("PasteSelectionTransaction", "Paste to niagara stack.");
