Changed the parameter types of the transform node to not be as confusing for vector/position types

diff --git a/Engine/Plugins/FX/Niagara/Source/NiagaraEditor/Private/NiagaraEditorCommon.cpp b/Engine/Plugins/FX/Niagara/Source/NiagaraEditor/Private/NiagaraEditorCommon.cpp
--- a/Engine/Plugins/FX/Niagara/Source/NiagaraEditor/Private/NiagaraEditorCommon.cpp
+++ b/Engine/Plugins/FX/Niagara/Source/NiagaraEditor/Private/NiagaraEditorCommon.cpp
@@ -1489,8 +1489,8 @@ void FNiagaraOpInfo::Init()
 	Op->FriendlyName = NSLOCTEXT("NiagaraOpInfo", "TransformVector Name", "Transform Vector");
 	Op->Description = NSLOCTEXT("NiagaraOpInfo", "TransformVector Desc", "Transforms a Vector3 as a vector, meaning it does not apply the translation part of the matrix.");
 	Op->Inputs.Add(FNiagaraOpInOutInfo(M, MatrixType, MText, MText, Default_MatrixOne));
-	Op->Inputs.Add(FNiagaraOpInOutInfo(V, FNiagaraTypeDefinition::GetVec3Def(), VText, VText, Default_VectorOne));
-	Op->Outputs.Add(FNiagaraOpInOutInfo(Result, FNiagaraTypeDefinition::GetVec3Def(), ResultText, ResultText, Default_VectorOne, TEXT("mul(float4({1},0.0),{0}).xyz")));
+	Op->Inputs.Add(FNiagaraOpInOutInfo(V, FNiagaraTypeDefinition::GetGenericNumericDef(), VText, VText, Default_VectorOne));
+	Op->Outputs.Add(FNiagaraOpInOutInfo(Result, FNiagaraTypeDefinition::GetGenericNumericDef(), ResultText, ResultText, Default_VectorOne, TEXT("mul(float4({1},0.0),{0}).xyz")));
 	Op->BuildName(TEXT("TransformVector"), MatrixCategoryName);
 	OpInfoMap.Add(Op->Name) = Idx;
 
