- Refactored Niagara editor style to user better standards
- Moved slate icons from Content/Icons to Content/Slate/Icons
- Removed a bunch of icons that we are definitely no longer using. Might be some others left we can safely remove
- Updated a bunch of style icons that were using outdated ue4 icons to new ue5 icons

diff --git a/Engine/Plugins/FX/Niagara/Source/NiagaraEditor/Private/NiagaraEditorModule.cpp b/Engine/Plugins/FX/Niagara/Source/NiagaraEditor/Private/NiagaraEditorModule.cpp
--- a/Engine/Plugins/FX/Niagara/Source/NiagaraEditor/Private/NiagaraEditorModule.cpp
+++ b/Engine/Plugins/FX/Niagara/Source/NiagaraEditor/Private/NiagaraEditorModule.cpp
@@ -938,7 +938,7 @@ void FNiagaraEditorModule::StartupModule()
 	RegisterStackIssueGenerator(FNiagaraPlatformSet::StaticStruct()->GetFName(), new FNiagaraPlatformSetIssueGenerator());
 
 
-	FNiagaraEditorStyle::Initialize();
+	FNiagaraEditorStyle::Register();
 	ReinitializeStyleCommand = IConsoleManager::Get().RegisterConsoleCommand(
 		TEXT("fx.NiagaraEditor.ReinitializeStyle"),
 		TEXT("Reinitializes the style for the niagara editor module.  Used in conjuction with live coding for UI tweaks.  May crash the editor if style objects are in use."),
@@ -1723,8 +1723,7 @@ void FNiagaraEditorModule::OnExecParticleInvoked(const TCHAR* Str)
 
 void FNiagaraEditorModule::ReinitializeStyle()
 {
-	FNiagaraEditorStyle::Shutdown();
-	FNiagaraEditorStyle::Initialize();
+	FNiagaraEditorStyle::ReinitializeStyle();
 }
 
 void FNiagaraEditorModule::EnqueueObjectForDeferredDestructionInternal(FDeferredDestructionContainerBase* InObjectToDestruct)
diff --git a/Engine/Plugins/FX/Niagara/Source/NiagaraEditor/Private/NiagaraEditorStyle.cpp b/Engine/Plugins/FX/Niagara/Source/NiagaraEditor/Private/NiagaraEditorStyle.cpp
--- a/Engine/Plugins/FX/Niagara/Source/NiagaraEditor/Private/NiagaraEditorStyle.cpp
+++ b/Engine/Plugins/FX/Niagara/Source/NiagaraEditor/Private/NiagaraEditorStyle.cpp
@@ -3,6 +3,7 @@
 #include "NiagaraEditorStyle.h"
 
 #include "Framework/Application/SlateApplication.h"
+#include "Styling/SlateStyleMacros.h"
 #include "EditorStyleSet.h"
 #include "Slate/SlateGameResources.h"
 #include "Styling/SlateStyleRegistry.h"
@@ -12,67 +13,44 @@
 #include "Classes/EditorStyleSettings.h"
 #include "Styling/StarshipCoreStyle.h"
 #include "Styling/StyleColors.h"
+#include "Styling/ToolBarStyle.h"
 
-TSharedPtr< FSlateStyleSet > FNiagaraEditorStyle::NiagaraEditorStyleInstance = NULL;
+TSharedPtr<FNiagaraEditorStyle> FNiagaraEditorStyle::NiagaraEditorStyle = nullptr;
 
-void FNiagaraEditorStyle::Initialize()
+void FNiagaraEditorStyle::Register()
 {
-	if (!NiagaraEditorStyleInstance.IsValid())
-	{
-		NiagaraEditorStyleInstance = Create();
-		FSlateStyleRegistry::RegisterSlateStyle(*NiagaraEditorStyleInstance);
-	}
-}
-
-void FNiagaraEditorStyle::Shutdown()
-{
-	FSlateStyleRegistry::UnRegisterSlateStyle(*NiagaraEditorStyleInstance);
-	ensure(NiagaraEditorStyleInstance.IsUnique());
-	NiagaraEditorStyleInstance.Reset();
+	FSlateStyleRegistry::RegisterSlateStyle(Get());
 }
 
-FName FNiagaraEditorStyle::GetStyleSetName()
+void FNiagaraEditorStyle::Unregister()
 {
-	static FName StyleSetName(TEXT("NiagaraEditorStyle"));
-	return StyleSetName;
+	FSlateStyleRegistry::UnRegisterSlateStyle(Get());
 }
 
-NIAGARAEDITOR_API FString RelativePathToPluginPath(const FString& RelativePath, const ANSICHAR* Extension)
+void FNiagaraEditorStyle::Shutdown()
 {
-	static FString ContentDir = IPluginManager::Get().FindPlugin(TEXT("Niagara"))->GetContentDir();
-	return (ContentDir / RelativePath) + Extension;
+	Unregister();
+	NiagaraEditorStyle.Reset();
 }
 
-#define IMAGE_BRUSH( RelativePath, ... ) FSlateImageBrush( Style->RootToContentDir( RelativePath, TEXT(".png") ), __VA_ARGS__ )
-#define IMAGE_CORE_BRUSH( RelativePath, ... ) FSlateImageBrush( FPaths::EngineContentDir() / "Editor/Slate" / RelativePath + TEXT(".png"), __VA_ARGS__ )
-#define IMAGE_PLUGIN_BRUSH( RelativePath, ... ) FSlateImageBrush( RelativePathToPluginPath( RelativePath, ".png" ), __VA_ARGS__ )
-#define BOX_BRUSH( RelativePath, ... ) FSlateBoxBrush( Style->RootToContentDir( RelativePath, TEXT(".png") ), __VA_ARGS__ )
-#define BOX_CORE_BRUSH( RelativePath, ... ) FSlateBoxBrush( FPaths::EngineContentDir() / "Editor/Slate" / RelativePath + TEXT(".png"), __VA_ARGS__ )
-#define BOX_PLUGIN_BRUSH( RelativePath, ... ) FSlateBoxBrush( RelativePathToPluginPath( RelativePath, ".png"), __VA_ARGS__ )
-#define BORDER_CORE_BRUSH( RelativePath, ... ) FSlateBorderBrush( FPaths::EngineContentDir() / "Editor/Slate" / RelativePath + TEXT(".png"), __VA_ARGS__ )
-#define BORDER_BRUSH( RelativePath, ... ) FSlateBorderBrush( Style->RootToContentDir( RelativePath, TEXT(".png") ), __VA_ARGS__ )
-#define DEFAULT_FONT(...) FStarshipCoreStyle::GetDefaultFontStyle(__VA_ARGS__)
-#define IMAGE_PLUGIN_BRUSH_SVG( RelativePath, ... ) FSlateVectorImageBrush( RelativePathToPluginPath( RelativePath, ".svg" ), __VA_ARGS__ )
-
 const FVector2D Icon8x8(8.0f, 8.0f);
 const FVector2D Icon12x12(12.0f, 12.0f);
 const FVector2D Icon16x16(16.0f, 16.0f);
 const FVector2D Icon20x20(20.0f, 20.0f);
+const FVector2D Icon32x32(32.0f, 32.0f);
 const FVector2D Icon40x40(40.0f, 40.0f);
 const FVector2D Icon64x64(64.0f, 64.0f);
 
-namespace NiagaraEditorStyleImpl {
-
-void InitStats(TSharedRef< FSlateStyleSet > Style)
+void FNiagaraEditorStyle::InitStats()
 {
 	const FTextBlockStyle CategoryText = FAppStyle::Get().GetWidgetStyle<FTextBlockStyle>("DetailsView.CategoryTextStyle");
 	const FSlateColor SelectionColor = FEditorStyle::GetSlateColor("SelectionColor");
 	const FSlateColor SelectionColor_Pressed = FEditorStyle::GetSlateColor("SelectionColor_Pressed");
 	
-	Style->Set("NiagaraEditor.StatsText", CategoryText);
+	Set("NiagaraEditor.StatsText", CategoryText);
 }
 
-void InitAssetPicker(TSharedRef< FSlateStyleSet > Style)
+void FNiagaraEditorStyle::InitAssetPicker()
 {
 	const FTextBlockStyle NormalText = FEditorStyle::GetWidgetStyle<FTextBlockStyle>("NormalText");
 
@@ -80,58 +58,58 @@ void InitAssetPicker(TSharedRef< FSlateStyleSet > Style)
 		.SetColorAndOpacity(FLinearColor::White)
 		.SetFont(DEFAULT_FONT("Bold", 9));
 
-	Style->Set("NiagaraEditor.AssetPickerBoldAssetNameText", AssetPickerBoldAssetNameText);
+	Set("NiagaraEditor.AssetPickerBoldAssetNameText", AssetPickerBoldAssetNameText);
 
 	FTextBlockStyle AssetPickerAssetNameText = FTextBlockStyle(NormalText)
 		.SetColorAndOpacity(FLinearColor::White)
 		.SetFont(DEFAULT_FONT("Regular", 9));
 
-	Style->Set("NiagaraEditor.AssetPickerAssetNameText", AssetPickerAssetNameText);
+	Set("NiagaraEditor.AssetPickerAssetNameText", AssetPickerAssetNameText);
 
 	FTextBlockStyle AssetPickerAssetCategoryText = FTextBlockStyle(NormalText)
 		.SetFont(DEFAULT_FONT("Bold", 11));
 
-	Style->Set("NiagaraEditor.AssetPickerAssetCategoryText", AssetPickerAssetCategoryText);
+	Set("NiagaraEditor.AssetPickerAssetCategoryText", AssetPickerAssetCategoryText);
 
 
 	FTextBlockStyle AssetPickerAssetSubcategoryText = FTextBlockStyle(NormalText)
 		.SetFont(DEFAULT_FONT("Bold", 10));
 
-	Style->Set("NiagaraEditor.AssetPickerAssetSubcategoryText", AssetPickerAssetSubcategoryText);
+	Set("NiagaraEditor.AssetPickerAssetSubcategoryText", AssetPickerAssetSubcategoryText);
 
 	// New Asset Dialog
 	FTextBlockStyle NewAssetDialogOptionText = FTextBlockStyle(NormalText)
 		.SetFont(DEFAULT_FONT("Bold", 11));
 
-	Style->Set("NiagaraEditor.NewAssetDialog.OptionText", NewAssetDialogOptionText);
+	Set("NiagaraEditor.NewAssetDialog.OptionText", NewAssetDialogOptionText);
 
 	FTextBlockStyle NewAssetDialogHeaderText = FTextBlockStyle(NormalText)
 		.SetColorAndOpacity(FLinearColor::White)
 		.SetFont(DEFAULT_FONT("Bold", 10));
 
-	Style->Set("NiagaraEditor.NewAssetDialog.HeaderText", NewAssetDialogHeaderText);
+	Set("NiagaraEditor.NewAssetDialog.HeaderText", NewAssetDialogHeaderText);
 
 	FTextBlockStyle NewAssetDialogSubHeaderText = FTextBlockStyle(NormalText)
 		.SetColorAndOpacity(FLinearColor::White)
 		.SetFont(DEFAULT_FONT("Bold", 11));
 
-	Style->Set("NiagaraEditor.NewAssetDialog.SubHeaderText", NewAssetDialogSubHeaderText);
+	Set("NiagaraEditor.NewAssetDialog.SubHeaderText", NewAssetDialogSubHeaderText);
 
-	Style->Set("NiagaraEditor.NewAssetDialog.AddButton", FButtonStyle()
-		.SetNormal(BOX_CORE_BRUSH("Common/FlatButton", 2.0f / 8.0f, FLinearColor(0, 0, 0, .25f)))
-		.SetHovered(BOX_CORE_BRUSH("Common/FlatButton", 2.0f / 8.0f, FEditorStyle::GetSlateColor("SelectionColor")))
-		.SetPressed(BOX_CORE_BRUSH("Common/FlatButton", 2.0f / 8.0f, FEditorStyle::GetSlateColor("SelectionColor_Pressed")))
+	Set("NiagaraEditor.NewAssetDialog.AddButton", FButtonStyle()
+		.SetNormal(CORE_BOX_BRUSH("Common/FlatButton", 2.0f / 8.0f, FLinearColor(0, 0, 0, .25f)))
+		.SetHovered(CORE_BOX_BRUSH("Common/FlatButton", 2.0f / 8.0f, FEditorStyle::GetSlateColor("SelectionColor")))
+		.SetPressed(CORE_BOX_BRUSH("Common/FlatButton", 2.0f / 8.0f, FEditorStyle::GetSlateColor("SelectionColor_Pressed")))
 	);
 
-	Style->Set("NiagaraEditor.NewAssetDialog.SubBorderColor", FLinearColor(FColor(48, 48, 48)));
-	Style->Set("NiagaraEditor.NewAssetDialog.ActiveOptionBorderColor", FLinearColor(FColor(96, 96, 96)));
+	Set("NiagaraEditor.NewAssetDialog.SubBorderColor", FLinearColor(FColor(48, 48, 48)));
+	Set("NiagaraEditor.NewAssetDialog.ActiveOptionBorderColor", FLinearColor(FColor(96, 96, 96)));
 
-	Style->Set("NiagaraEditor.NewAssetDialog.SubBorder", new BOX_CORE_BRUSH("Common/GroupBorderLight", FMargin(4.0f / 16.0f)));
+	Set("NiagaraEditor.NewAssetDialog.SubBorder", new CORE_BOX_BRUSH("Common/GroupBorderLight", FMargin(4.0f / 16.0f)));
 
 
 }
 
-void InitActionMenu(TSharedRef< FSlateStyleSet > Style)
+void FNiagaraEditorStyle::InitActionMenu()
 {
 	const FTextBlockStyle NormalText = FEditorStyle::GetWidgetStyle<FTextBlockStyle>("NormalText");
 
@@ -164,18 +142,18 @@ void InitActionMenu(TSharedRef< FSlateStyleSet > Style)
 	const FCheckBoxStyle NiagaraGraphActionMenuFilterCheckBox = FCheckBoxStyle()
             .SetCheckBoxType(ESlateCheckBoxType::ToggleButton)
             .SetUncheckedImage( FSlateNoResource() )
-            .SetUncheckedHoveredImage(BOX_CORE_BRUSH("Common/RoundedSelection_16x", 4.0f/16.0f, FLinearColor(0.7f, 0.7f, 0.7f) ))
-            .SetUncheckedPressedImage(BOX_CORE_BRUSH("Common/RoundedSelection_16x", 4.0f/16.0f, FLinearColor(0.8f, 0.8f, 0.8f) ))
-            .SetCheckedImage(BOX_CORE_BRUSH("Common/RoundedSelection_16x",  4.0f/16.0f, FLinearColor(0.9f, 0.9f, 0.9f) ))
-            .SetCheckedHoveredImage(BOX_CORE_BRUSH("Common/RoundedSelection_16x",  4.0f/16.0f, FLinearColor(1.f, 1.f, 1.f) ))
-            .SetCheckedPressedImage(BOX_CORE_BRUSH("Common/RoundedSelection_16x",  4.0f/16.0f, FLinearColor(1.f, 1.f, 1.f) ));
+            .SetUncheckedHoveredImage(CORE_BOX_BRUSH("Common/RoundedSelection_16x", 4.0f/16.0f, FLinearColor(0.7f, 0.7f, 0.7f) ))
+            .SetUncheckedPressedImage(CORE_BOX_BRUSH("Common/RoundedSelection_16x", 4.0f/16.0f, FLinearColor(0.8f, 0.8f, 0.8f) ))
+            .SetCheckedImage(CORE_BOX_BRUSH("Common/RoundedSelection_16x",  4.0f/16.0f, FLinearColor(0.9f, 0.9f, 0.9f) ))
+            .SetCheckedHoveredImage(CORE_BOX_BRUSH("Common/RoundedSelection_16x",  4.0f/16.0f, FLinearColor(1.f, 1.f, 1.f) ))
+            .SetCheckedPressedImage(CORE_BOX_BRUSH("Common/RoundedSelection_16x",  4.0f/16.0f, FLinearColor(1.f, 1.f, 1.f) ));
 
 	const FTableRowStyle ActionMenuRowStyle = FTableRowStyle()
             .SetEvenRowBackgroundBrush(FSlateNoResource())
             .SetOddRowBackgroundBrush(FSlateNoResource())
-			.SetEvenRowBackgroundHoveredBrush(IMAGE_CORE_BRUSH("Common/Selection", Icon8x8, FLinearColor(1.0f, 1.0f, 1.0f, 0.1f)))
-			.SetOddRowBackgroundHoveredBrush(IMAGE_CORE_BRUSH("Common/Selection", Icon8x8, FLinearColor(1.0f, 1.0f, 1.0f, 0.1f)))
-            .SetSelectorFocusedBrush(BORDER_CORE_BRUSH("Common/Selector", FMargin(4.f / 16.f), FStarshipCoreStyle::GetCoreStyle().GetSlateColor("SelectorColor")))
+			.SetEvenRowBackgroundHoveredBrush(CORE_IMAGE_BRUSH("Common/Selection", Icon8x8, FLinearColor(1.0f, 1.0f, 1.0f, 0.1f)))
+			.SetOddRowBackgroundHoveredBrush(CORE_IMAGE_BRUSH("Common/Selection", Icon8x8, FLinearColor(1.0f, 1.0f, 1.0f, 0.1f)))
+            .SetSelectorFocusedBrush(CORE_BORDER_BRUSH("Common/Selector", FMargin(4.f / 16.f), FStarshipCoreStyle::GetCoreStyle().GetSlateColor("SelectorColor")))
             .SetActiveBrush(FSlateColorBrush(FStyleColors::Select))
             .SetActiveHoveredBrush(FSlateColorBrush(FStyleColors::Select))
             .SetInactiveBrush(FSlateColorBrush(FStyleColors::SelectInactive))
@@ -185,30 +163,31 @@ void InitActionMenu(TSharedRef< FSlateStyleSet > Style)
 			.SetTextColor(FStyleColors::Foreground)
 			.SetSelectedTextColor(FStyleColors::ForegroundInverted)
 
-            .SetDropIndicator_Above(BOX_CORE_BRUSH("Common/DropZoneIndicator_Above", FMargin(10.0f / 16.0f, 10.0f / 16.0f, 0, 0), FStarshipCoreStyle::GetCoreStyle().GetSlateColor("SelectorColor")))
-            .SetDropIndicator_Onto(BOX_CORE_BRUSH("Common/DropZoneIndicator_Onto", FMargin(4.0f / 16.0f), FStarshipCoreStyle::GetCoreStyle().GetSlateColor("SelectorColor")))
-            .SetDropIndicator_Below(BOX_CORE_BRUSH("Common/DropZoneIndicator_Below", FMargin(10.0f / 16.0f, 0, 0, 10.0f / 16.0f), FStarshipCoreStyle::GetCoreStyle().GetSlateColor("SelectorColor")));
+            .SetDropIndicator_Above(CORE_BOX_BRUSH("Common/DropZoneIndicator_Above", FMargin(10.0f / 16.0f, 10.0f / 16.0f, 0, 0), FStarshipCoreStyle::GetCoreStyle().GetSlateColor("SelectorColor")))
+            .SetDropIndicator_Onto(CORE_BOX_BRUSH("Common/DropZoneIndicator_Onto", FMargin(4.0f / 16.0f), FStarshipCoreStyle::GetCoreStyle().GetSlateColor("SelectorColor")))
+            .SetDropIndicator_Below(CORE_BOX_BRUSH("Common/DropZoneIndicator_Below", FMargin(10.0f / 16.0f, 0, 0, 10.0f / 16.0f), FStarshipCoreStyle::GetCoreStyle().GetSlateColor("SelectorColor")));
 	
-	Style->Set("ActionMenu.Row", ActionMenuRowStyle);
+	Set("ActionMenu.Row", ActionMenuRowStyle);
 	
-	Style->Set("ActionMenu.HeadingTextBlock", ActionMenuHeadingText);
+	Set("ActionMenu.HeadingTextBlock", ActionMenuHeadingText);
 
-	Style->Set("ActionMenu.ActionTextBlock", ActionMenuActionText);
+	Set("ActionMenu.ActionTextBlock", ActionMenuActionText);
 
-	Style->Set("GraphActionMenu.Background", new FSlateRoundedBoxBrush(FStyleColors::Background, 3.f));
+	Set("GraphActionMenu.Background", new FSlateRoundedBoxBrush(FStyleColors::Background, 3.f));
 	
-	Style->Set("GraphActionMenu.ActionSourceTextBlock", ActionMenuSourceText);
+	Set("GraphActionMenu.ActionSourceTextBlock", ActionMenuSourceText);
 
-	Style->Set("GraphActionMenu.ActionFilterTextBlock", ActionMenuFilterText);
+	Set("GraphActionMenu.ActionFilterTextBlock", ActionMenuFilterText);
 
-	Style->Set("GraphActionMenu.TemplateTabTextBlock", TemplateTabText);
+	Set("GraphActionMenu.TemplateTabTextBlock", TemplateTabText);
 	
-	Style->Set("GraphActionMenu.FilterCheckBox", NiagaraGraphActionMenuFilterCheckBox);
+	Set("GraphActionMenu.FilterCheckBox", NiagaraGraphActionMenuFilterCheckBox);
 
-	Style->Set("GraphActionMenu.FilterCheckBox.Border", new FSlateRoundedBoxBrush(FStyleColors::Secondary, 3.f));
+	Set("GraphActionMenu.FilterCheckBox.Border", new FSlateRoundedBoxBrush(FStyleColors::Secondary, 3.f));
+	
 }
 
-void InitEmitterHeader(TSharedRef< FSlateStyleSet > Style)
+void FNiagaraEditorStyle::InitEmitterHeader()
 {
 	const FTextBlockStyle NormalText = FEditorStyle::GetWidgetStyle<FTextBlockStyle>("NormalText");
 	
@@ -217,35 +196,35 @@ void InitEmitterHeader(TSharedRef< FSlateStyleSet > Style)
 	StackHeaderText.SetFont(DEFAULT_FONT("Regular", 11))
 		.SetColorAndOpacity(FSlateColor(EStyleColor::White));
 
-	Style->Set("NiagaraEditor.HeadingTextBlock", StackHeaderText);
+	Set("NiagaraEditor.HeadingTextBlock", StackHeaderText);
 
 	FTextBlockStyle StackHeaderTextSubdued = FTextBlockStyle(NormalText);
 	StackHeaderTextSubdued.SetFont(DEFAULT_FONT("Regular", 11))
 		.SetColorAndOpacity(FStyleColors::Foreground);
 
-	Style->Set("NiagaraEditor.HeadingTextBlockSubdued", StackHeaderTextSubdued);
+	Set("NiagaraEditor.HeadingTextBlockSubdued", StackHeaderTextSubdued);
 
 	FTextBlockStyle TabText = FTextBlockStyle(NormalText)
 		.SetFont(DEFAULT_FONT("Regular", 12))
 		.SetShadowOffset(FVector2D(0, 1))
 		.SetShadowColorAndOpacity(FLinearColor(0, 0, 0, 0.9f));
 	
-	Style->Set("NiagaraEditor.AttributeSpreadsheetTabText", TabText);
+	Set("NiagaraEditor.AttributeSpreadsheetTabText", TabText);
 
 	FTextBlockStyle SubduedHeadingText = FTextBlockStyle(NormalText)
 		.SetFont(DEFAULT_FONT("Regular", 14))
 		.SetColorAndOpacity(FSlateColor(FLinearColor(0.5f, 0.5f, 0.5f, 1.0f)));
 	
-	Style->Set("NiagaraEditor.SubduedHeadingTextBox", SubduedHeadingText);
+	Set("NiagaraEditor.SubduedHeadingTextBox", SubduedHeadingText);
 
 	// Details
 	FTextBlockStyle DetailsHeadingText = FTextBlockStyle(NormalText)
 		.SetFont(DEFAULT_FONT("Bold", 9));
 
-	Style->Set("NiagaraEditor.DetailsHeadingText", DetailsHeadingText);
+	Set("NiagaraEditor.DetailsHeadingText", DetailsHeadingText);
 }
 
-void InitParameters(TSharedRef< FSlateStyleSet > Style)
+void FNiagaraEditorStyle::InitParameters()
 {
 	const FTextBlockStyle NormalText = FEditorStyle::GetWidgetStyle<FTextBlockStyle>("NormalText");
 	const FEditableTextBoxStyle NormalEditableTextBox = FStarshipCoreStyle::GetCoreStyle().GetWidgetStyle<FEditableTextBoxStyle>("NormalEditableTextBox");
@@ -256,68 +235,69 @@ void InitParameters(TSharedRef< FSlateStyleSet > Style)
 	FTextBlockStyle ParameterText = FTextBlockStyle(NormalText)
 		.SetFont(NormalFont);
 
-	Style->Set("NiagaraEditor.ParameterText", ParameterText);
+	Set("NiagaraEditor.ParameterText", ParameterText);
 
 	FEditableTextBoxStyle ParameterEditableText = FEditableTextBoxStyle(NormalEditableTextBox)
-		.SetFont(NormalFont);
+		.SetFont(NormalFont)
+		.SetForegroundColor(FStyleColors::Black);
 
 	FInlineEditableTextBlockStyle ParameterEditableTextBox = FInlineEditableTextBlockStyle()
 		.SetEditableTextBoxStyle(ParameterEditableText)
 		.SetTextStyle(ParameterText);
-	Style->Set("NiagaraEditor.ParameterInlineEditableText", ParameterEditableTextBox);
+	Set("NiagaraEditor.ParameterInlineEditableText", ParameterEditableTextBox);
 
 
-	Style->Set("NiagaraEditor.ParameterName.NamespaceBorder", new BOX_PLUGIN_BRUSH("Icons/NamespaceBorder", FMargin(4.0f / 16.0f)));
+	Set("NiagaraEditor.ParameterName.NamespaceBorder", new BOX_BRUSH("Icons/NamespaceBorder", FMargin(4.0f / 16.0f)));
 
-	Style->Set("NiagaraEditor.ParameterName.NamespaceText", FTextBlockStyle(NormalText)
+	Set("NiagaraEditor.ParameterName.NamespaceText", FTextBlockStyle(NormalText)
 		.SetFont(DEFAULT_FONT("Bold", 8))
 		.SetColorAndOpacity(FLinearColor(1.0f, 1.0f, 1.0f, 0.9f))
 		.SetShadowOffset(FVector2D(1, 1))
 		.SetShadowColorAndOpacity(FLinearColor(0, 0, 0, 0.7f)));
 
-	Style->Set("NiagaraEditor.ParameterName.NamespaceTextDark", FTextBlockStyle(NormalText)
+	Set("NiagaraEditor.ParameterName.NamespaceTextDark", FTextBlockStyle(NormalText)
 		.SetFont(DEFAULT_FONT("Bold", 8))
 		.SetColorAndOpacity(FLinearColor(0.0f, 0.0f, 0.0f, 0.9f))
 		.SetShadowOffset(FVector2D(1, 1))
 		.SetShadowColorAndOpacity(FLinearColor(1.0, 1.0, 1.0, 0.25f)));
 
-	Style->Set("NiagaraEditor.ParameterName.TypeText", FTextBlockStyle(NormalText)
+	Set("NiagaraEditor.ParameterName.TypeText", FTextBlockStyle(NormalText)
 		.SetFont(DEFAULT_FONT("Regular", 8))
 		.SetColorAndOpacity(FLinearColor(1.0f, 1.0f, 1.0f, 0.5f)));
 
 
-	Style->Set("NiagaraEditor.StaticIcon", new IMAGE_PLUGIN_BRUSH("Icons/staticpill_16x", Icon16x16));
-	Style->Set("NiagaraEditor.Pins.StaticConnected", new IMAGE_PLUGIN_BRUSH("Icons/StaticPin_Connected", Icon16x16));
-	Style->Set("NiagaraEditor.Pins.StaticDisconnected", new IMAGE_PLUGIN_BRUSH("Icons/StaticPin_Disconnected", Icon16x16));
-	Style->Set("NiagaraEditor.Pins.StaticConnectedHovered", new IMAGE_PLUGIN_BRUSH("Icons/StaticPin_Connected", Icon16x16, FLinearColor(0.8f, 0.8f, 0.8f)));
-	Style->Set("NiagaraEditor.Pins.StaticDisconnectedHovered", new IMAGE_PLUGIN_BRUSH("Icons/StaticPin_Disconnected", Icon16x16, FLinearColor(0.8f, 0.8f, 0.8f)));
+	Set("NiagaraEditor.StaticIcon", new IMAGE_BRUSH("Icons/staticpill_16x", Icon16x16));
+	Set("NiagaraEditor.Pins.StaticConnected", new IMAGE_BRUSH("Icons/StaticPin_Connected", Icon16x16));
+	Set("NiagaraEditor.Pins.StaticDisconnected", new IMAGE_BRUSH("Icons/StaticPin_Disconnected", Icon16x16));
+	Set("NiagaraEditor.Pins.StaticConnectedHovered", new IMAGE_BRUSH("Icons/StaticPin_Connected", Icon16x16, FLinearColor(0.8f, 0.8f, 0.8f)));
+	Set("NiagaraEditor.Pins.StaticDisconnectedHovered", new IMAGE_BRUSH("Icons/StaticPin_Disconnected", Icon16x16, FLinearColor(0.8f, 0.8f, 0.8f)));
 }
 
-void InitParameterMapView(TSharedRef< FSlateStyleSet > Style)
+void FNiagaraEditorStyle::InitParameterMapView()
 {
 	// Parameter Map View
-	Style->Set("NiagaraEditor.Stack.DepressedHighlightedButtonBrush", new BOX_CORE_BRUSH("Common/ButtonHoverHint", FMargin(4 / 16.0f), FStyleColors::PrimaryPress));
-	Style->Set("NiagaraEditor.Stack.FlatButtonColor", FLinearColor(FColor(205, 205, 205)));
+	Set("NiagaraEditor.Stack.DepressedHighlightedButtonBrush", new CORE_BOX_BRUSH("Common/ButtonHoverHint", FMargin(4 / 16.0f), FStyleColors::PrimaryPress));
+	Set("NiagaraEditor.Stack.FlatButtonColor", FLinearColor(FColor(205, 205, 205)));
 
 	//Parameters panel
 	const FTableRowStyle TreeViewStyle = FEditorStyle::GetWidgetStyle<FTableRowStyle>("DetailsView.TreeView.TableRow");
 	FTableRowStyle ParameterPanelRowStyle = FTableRowStyle(TreeViewStyle)
 		.SetTextColor(FLinearColor::White)
 		.SetSelectedTextColor(FLinearColor::White);
-	Style->Set("NiagaraEditor.Parameters.TableRow", ParameterPanelRowStyle);
+	Set("NiagaraEditor.Parameters.TableRow", ParameterPanelRowStyle);
 	
 	const FTextBlockStyle CategoryTextStyle = FEditorStyle::GetWidgetStyle<FTextBlockStyle>("DetailsView.CategoryTextStyle");
 	FTextBlockStyle ParameterSectionStyle = FTextBlockStyle(CategoryTextStyle)
 		.SetColorAndOpacity(FLinearColor::White);
-	Style->Set("NiagaraEditor.Parameters.HeaderText", ParameterSectionStyle);
+	Set("NiagaraEditor.Parameters.HeaderText", ParameterSectionStyle);
 }
 
-void InitCodeView(TSharedRef< FSlateStyleSet > Style)
+void FNiagaraEditorStyle::InitCodeView()
 {
 	const FTextBlockStyle NormalText = FEditorStyle::GetWidgetStyle<FTextBlockStyle>("NormalText");
 
 	// Code View
-	Style->Set("NiagaraEditor.CodeView.Checkbox.Text", FTextBlockStyle(NormalText)
+	Set("NiagaraEditor.CodeView.Checkbox.Text", FTextBlockStyle(NormalText)
 		.SetFont(DEFAULT_FONT("Bold", 12))
 		.SetColorAndOpacity(FLinearColor(1.0f, 1.0f, 1.0f, 0.9f))
 		.SetShadowOffset(FVector2D(1, 1))
@@ -329,36 +309,36 @@ void InitCodeView(TSharedRef< FSlateStyleSet > Style)
 		.SetFont(LogFont)
 		.SetColorAndOpacity(FLinearColor(FColor(0xffffffff)))
 		.SetSelectedBackgroundColor(FLinearColor(FColor(0xff666666)));
-	Style->Set("NiagaraEditor.CodeView.Hlsl.Normal", NormalLogText);
+	Set("NiagaraEditor.CodeView.Hlsl.Normal", NormalLogText);
 
 
 	const FTextBlockStyle ErrorText = FTextBlockStyle(NormalText)
-		.SetUnderlineBrush(IMAGE_CORE_BRUSH("Old/White", Icon8x8, FLinearColor::Red, ESlateBrushTileType::Both))
+		.SetUnderlineBrush(IMAGE_BRUSH("Icons/White", Icon8x8, FLinearColor::Red, ESlateBrushTileType::Both))
 		.SetColorAndOpacity(FLinearColor::Red);
 		
-	Style->Set("TextEditor.NormalText", NormalText);
+	Set("TextEditor.NormalText", NormalText);
 
 	constexpr int32 HlslFontSize = 9;
 	FSlateFontInfo HlslFont = DEFAULT_FONT("Mono", HlslFontSize);
 	FTextBlockStyle NormalHlslText = FTextBlockStyle(NormalText)
 		.SetFont(HlslFont);
 	const FTextBlockStyle HlslErrorText = FTextBlockStyle(NormalHlslText)
-		.SetUnderlineBrush(IMAGE_CORE_BRUSH("Old/White", Icon8x8, FLinearColor::Red, ESlateBrushTileType::Both))
+		.SetUnderlineBrush(IMAGE_BRUSH("Icons/White", Icon8x8, FLinearColor::Red, ESlateBrushTileType::Both))
 		.SetColorAndOpacity(FLinearColor::Red);
 	
-	Style->Set("SyntaxHighlight.HLSL.Normal", FTextBlockStyle(NormalHlslText).SetColorAndOpacity(FLinearColor(FColor(189, 183, 107))));
-	Style->Set("SyntaxHighlight.HLSL.Operator", FTextBlockStyle(NormalHlslText).SetColorAndOpacity(FLinearColor(FColor(220, 220, 220))));
-	Style->Set("SyntaxHighlight.HLSL.Keyword", FTextBlockStyle(NormalHlslText).SetColorAndOpacity(FLinearColor(FColor(86, 156, 214))));
-	Style->Set("SyntaxHighlight.HLSL.String", FTextBlockStyle(NormalHlslText).SetColorAndOpacity(FLinearColor(FColor(214, 157, 133))));
-	Style->Set("SyntaxHighlight.HLSL.Number", FTextBlockStyle(NormalHlslText).SetColorAndOpacity(FLinearColor(FColor(181, 206, 168))));
-	Style->Set("SyntaxHighlight.HLSL.Comment", FTextBlockStyle(NormalHlslText).SetColorAndOpacity(FLinearColor(FColor(87, 166, 74))));
-	Style->Set("SyntaxHighlight.HLSL.PreProcessorKeyword", FTextBlockStyle(NormalHlslText).SetColorAndOpacity(FLinearColor(FColor(188, 98, 171))));
-
-	Style->Set("SyntaxHighlight.HLSL.Error", HlslErrorText); 
+	Set("SyntaxHighlight.HLSL.Normal", FTextBlockStyle(NormalHlslText).SetColorAndOpacity(FLinearColor(FColor(189, 183, 107))));
+	Set("SyntaxHighlight.HLSL.Operator", FTextBlockStyle(NormalHlslText).SetColorAndOpacity(FLinearColor(FColor(220, 220, 220))));
+	Set("SyntaxHighlight.HLSL.Keyword", FTextBlockStyle(NormalHlslText).SetColorAndOpacity(FLinearColor(FColor(86, 156, 214))));
+	Set("SyntaxHighlight.HLSL.String", FTextBlockStyle(NormalHlslText).SetColorAndOpacity(FLinearColor(FColor(214, 157, 133))));
+	Set("SyntaxHighlight.HLSL.Number", FTextBlockStyle(NormalHlslText).SetColorAndOpacity(FLinearColor(FColor(181, 206, 168))));
+	Set("SyntaxHighlight.HLSL.Comment", FTextBlockStyle(NormalHlslText).SetColorAndOpacity(FLinearColor(FColor(87, 166, 74))));
+	Set("SyntaxHighlight.HLSL.PreProcessorKeyword", FTextBlockStyle(NormalHlslText).SetColorAndOpacity(FLinearColor(FColor(188, 98, 171))));
+
+	Set("SyntaxHighlight.HLSL.Error", HlslErrorText); 
 		
 }
 
-void InitSelectedEmitter(TSharedRef< FSlateStyleSet > Style)
+void FNiagaraEditorStyle::InitSelectedEmitter()
 {
 	const FTextBlockStyle NormalText = FEditorStyle::GetWidgetStyle<FTextBlockStyle>("NormalText");
 
@@ -367,409 +347,360 @@ void InitSelectedEmitter(TSharedRef< FSlateStyleSet > Style)
 	FTextBlockStyle SelectedEmitterUnsupportedSelectionText = FTextBlockStyle(NormalText)
 		.SetFont(SelectedEmitterUnsupportedSelectionFont)
 		.SetColorAndOpacity(FLinearColor(1.0f, 1.0f, 1.0f, 1.0f));
-	Style->Set("NiagaraEditor.SelectedEmitter.UnsupportedSelectionText", SelectedEmitterUnsupportedSelectionText);
-}
-
-void InitToolbarIcons(TSharedRef< FSlateStyleSet > Style)
-{
-	Style->Set("NiagaraEditor.Apply", new IMAGE_BRUSH("Icons/icon_Niagara_Apply_40x", Icon40x40));
-	Style->Set("NiagaraEditor.Apply.Small", new IMAGE_BRUSH("Icons/icon_Niagara_Apply_40x", Icon20x20));
-	Style->Set("NiagaraEditor.ApplyScratchPadChanges", new IMAGE_PLUGIN_BRUSH("Icons/Commands/icon_ApplyScratchPadChanges_40x", Icon40x40));
-	Style->Set("NiagaraEditor.ApplyScratchPadChanges.Small", new IMAGE_PLUGIN_BRUSH("Icons/Commands/icon_ApplyScratchPadChanges_40x", Icon20x20));
-	Style->Set("NiagaraEditor.Compile", new IMAGE_BRUSH("Icons/icon_compile_40x", Icon40x40));
-	Style->Set("NiagaraEditor.Compile.Small", new IMAGE_BRUSH("Icons/icon_compile_40x", Icon20x20));
-	Style->Set("NiagaraEditor.AddEmitter", new IMAGE_BRUSH("Icons/icon_AddObject_40x", Icon40x40));
-	Style->Set("NiagaraEditor.AddEmitter.Small", new IMAGE_BRUSH("Icons/icon_AddObject_40x", Icon20x20));
-	Style->Set("NiagaraEditor.UnlockToChanges", new IMAGE_BRUSH("Icons/icon_levels_unlocked_40x", Icon40x40));
-	Style->Set("NiagaraEditor.UnlockToChanges.Small", new IMAGE_BRUSH("Icons/icon_levels_unlocked_40x", Icon20x20));
-	Style->Set("NiagaraEditor.LockToChanges", new IMAGE_BRUSH("Icons/icon_levels_LockedReadOnly_40x", Icon40x40));
-	Style->Set("NiagaraEditor.LockToChanges.Small", new IMAGE_BRUSH("Icons/icon_levels_LockedReadOnly_40x", Icon20x20));
-	Style->Set("NiagaraEditor.SimulationOptions", new IMAGE_PLUGIN_BRUSH("Icons/Commands/icon_simulationOptions_40x", Icon40x40));
-	Style->Set("NiagaraEditor.SimulationOptions.Small", new IMAGE_PLUGIN_BRUSH("Icons/Commands/icon_simulationOptions_40x", Icon20x20));
-
-	Style->Set("Niagara.Asset.ReimportAsset.Needed", new IMAGE_BRUSH("Icons/icon_Reimport_Needed_40x", Icon40x40));
-	Style->Set("Niagara.Asset.ReimportAsset.Default", new IMAGE_BRUSH("Icons/icon_Reimport_40x", Icon40x40));
-	
-	Style->Set("NiagaraEditor.OverviewNode.IsolatedColor", FLinearColor::Yellow);
-	Style->Set("NiagaraEditor.OverviewNode.NotIsolatedColor", FLinearColor::Transparent);
+	Set("NiagaraEditor.SelectedEmitter.UnsupportedSelectionText", SelectedEmitterUnsupportedSelectionText);
 }
 
-void InitIcons(TSharedRef< FSlateStyleSet > Style)
+void FNiagaraEditorStyle::InitToolbarIcons()
+{
+	Set("NiagaraEditor.Refresh", new CORE_IMAGE_BRUSH_SVG("Starship/Common/Update", Icon20x20));
+	Set("NiagaraEditor.ApplyScratchPadChanges", new IMAGE_BRUSH_SVG("Icons/Commands/ApplyScratch", Icon40x40));		
+	Set("NiagaraEditor.OverviewNode.IsolatedColor", FLinearColor::Yellow);
+	Set("NiagaraEditor.OverviewNode.NotIsolatedColor", FLinearColor::Transparent);
+}
+
+void FNiagaraEditorStyle::InitIcons()
 {
 	// Icons
-	Style->Set("NiagaraEditor.Isolate", new IMAGE_PLUGIN_BRUSH("Icons/Isolate", Icon16x16));
-	Style->Set("NiagaraEditor.Module.Pin.TypeSelector", new IMAGE_PLUGIN_BRUSH("Icons/Scratch", Icon16x16, FLinearColor::Gray));
-	Style->Set("NiagaraEditor.Module.AddPin", new IMAGE_PLUGIN_BRUSH("Icons/PlusSymbol_12x", Icon12x12, FLinearColor::Gray));
-	Style->Set("NiagaraEditor.Module.RemovePin", new IMAGE_PLUGIN_BRUSH("Icons/MinusSymbol_12x", Icon12x12, FLinearColor::Gray));
-	Style->Set("NiagaraEditor.Message.CustomNote", new IMAGE_PLUGIN_BRUSH("Icons/icon_custom_note_16x", Icon16x16));
+	Set("NiagaraEditor.Isolate", new IMAGE_BRUSH_SVG("Icons/Commands/IsolateIcon", Icon16x16));
+	Set("NiagaraEditor.Module.Pin.TypeSelector", new IMAGE_BRUSH("Icons/Scratch", Icon16x16, FLinearColor::Gray));
+	Set("NiagaraEditor.Module.AddPin", new IMAGE_BRUSH("Icons/PlusSymbol_12x", Icon12x12, FLinearColor::Gray));
+	Set("NiagaraEditor.Module.RemovePin", new IMAGE_BRUSH("Icons/MinusSymbol_12x", Icon12x12, FLinearColor::Gray));
+	Set("NiagaraEditor.Message.CustomNote", new IMAGE_BRUSH("Icons/icon_custom_note_16x", Icon16x16));
 }
 
-void InitOverview(TSharedRef< FSlateStyleSet > Style)
+void FNiagaraEditorStyle::InitOverview()
 {
 	// Overview debug icons
-	Style->Set("NiagaraEditor.Overview.DebugActive", new IMAGE_PLUGIN_BRUSH("Icons/OverviewDebugActive", Icon16x16));
-	Style->Set("NiagaraEditor.Overview.DebugInactive", new IMAGE_PLUGIN_BRUSH("Icons/OverviewDebugInactive", Icon16x16));
+	Set("NiagaraEditor.Overview.DebugActive", new IMAGE_BRUSH("Icons/OverviewDebugActive", Icon16x16));
+	Set("NiagaraEditor.Overview.DebugInactive", new IMAGE_BRUSH("Icons/OverviewDebugInactive", Icon16x16));
 }
 
-void InitEmitterDetails(TSharedRef< FSlateStyleSet > Style)
+void FNiagaraEditorStyle::InitEmitterDetails()
 {
 	// Emitter details customization
-	Style->Set("NiagaraEditor.MaterialWarningBorder", new BOX_CORE_BRUSH("Common/GroupBorderLight", FMargin(4.0f / 16.0f)));
+	Set("NiagaraEditor.MaterialWarningBorder", new CORE_BOX_BRUSH("Common/GroupBorderLight", FMargin(4.0f / 16.0f)));
 }
 
-void InitAssetColors(TSharedRef< FSlateStyleSet > Style)
+void FNiagaraEditorStyle::InitAssetColors()
 {
 	// Asset colors
-	Style->Set("NiagaraEditor.AssetColors.System", FLinearColor(1.0f, 0.0f, 0.0f));
-	Style->Set("NiagaraEditor.AssetColors.Emitter", FLinearColor(1.0f, 0.3f, 0.0f));
-	Style->Set("NiagaraEditor.AssetColors.Script", FLinearColor(1.0f, 1.0f, 0.0f));
-	Style->Set("NiagaraEditor.AssetColors.ParameterCollection", FLinearColor(1.0f, 1.0f, 0.3f));
-	Style->Set("NiagaraEditor.AssetColors.ParameterCollectionInstance", FLinearColor(1.0f, 1.0f, 0.7f));
-	Style->Set("NiagaraEditor.AssetColors.ParameterDefinitions", FLinearColor(0.57f, 0.82f, 0.06f));
-	Style->Set("NiagaraEditor.AssetColors.EffectType", FLinearColor(1.f, 1.f, 1.f));
+	Set("NiagaraEditor.AssetColors.System", FLinearColor(1.0f, 0.0f, 0.0f));
+	Set("NiagaraEditor.AssetColors.Emitter", FLinearColor(1.0f, 0.3f, 0.0f));
+	Set("NiagaraEditor.AssetColors.Script", FLinearColor(1.0f, 1.0f, 0.0f));
+	Set("NiagaraEditor.AssetColors.ParameterCollection", FLinearColor(1.0f, 1.0f, 0.3f));
+	Set("NiagaraEditor.AssetColors.ParameterCollectionInstance", FLinearColor(1.0f, 1.0f, 0.7f));
+	Set("NiagaraEditor.AssetColors.ParameterDefinitions", FLinearColor(0.57f, 0.82f, 0.06f));
+	Set("NiagaraEditor.AssetColors.EffectType", FLinearColor(1.f, 1.f, 1.f));
 }
 
-void InitThumbnails(TSharedRef< FSlateStyleSet > Style)
+void FNiagaraEditorStyle::InitThumbnails()
 {
 	// Script factory thumbnails
-	Style->Set("NiagaraEditor.Thumbnails.DynamicInputs", new IMAGE_BRUSH("Icons/NiagaraScriptDynamicInputs_64x", Icon64x64));
-	Style->Set("NiagaraEditor.Thumbnails.Functions", new IMAGE_BRUSH("Icons/NiagaraScriptFunction_64x", Icon64x64));
-	Style->Set("NiagaraEditor.Thumbnails.Modules", new IMAGE_BRUSH("Icons/NiagaraScriptModules_64x", Icon64x64));
+	Set("NiagaraEditor.Thumbnails.DynamicInputs", new IMAGE_BRUSH("Icons/NiagaraScriptDynamicInputs_64x", Icon64x64));
+	Set("NiagaraEditor.Thumbnails.Functions", new IMAGE_BRUSH("Icons/NiagaraScriptFunction_64x", Icon64x64));
+	Set("NiagaraEditor.Thumbnails.Modules", new IMAGE_BRUSH("Icons/NiagaraScriptModules_64x", Icon64x64));
 }
 
-void InitClassIcon(TSharedRef< FSlateStyleSet > Style)
+void FNiagaraEditorStyle::InitClassIcon()
 {
 	// Renderer class icons
-	Style->Set("ClassIcon.NiagaraSpriteRendererProperties", new IMAGE_PLUGIN_BRUSH("Icons/Renderers/renderer_sprite", Icon16x16));
-	Style->Set("ClassIcon.NiagaraMeshRendererProperties", new IMAGE_PLUGIN_BRUSH("Icons/Renderers/renderer_mesh", Icon16x16));
-	Style->Set("ClassIcon.NiagaraRibbonRendererProperties", new IMAGE_PLUGIN_BRUSH("Icons/Renderers/renderer_ribbon", Icon16x16));
-	Style->Set("ClassIcon.NiagaraLightRendererProperties", new IMAGE_PLUGIN_BRUSH("Icons/Renderers/renderer_light", Icon16x16));
-	Style->Set("ClassIcon.NiagaraRendererProperties", new IMAGE_PLUGIN_BRUSH("Icons/Renderers/renderer_default", Icon16x16));
+	Set("ClassIcon.NiagaraSpriteRendererProperties", new IMAGE_BRUSH("Icons/Renderers/renderer_sprite", Icon16x16));
+	Set("ClassIcon.NiagaraMeshRendererProperties", new IMAGE_BRUSH("Icons/Renderers/renderer_mesh", Icon16x16));
+	Set("ClassIcon.NiagaraRibbonRendererProperties", new IMAGE_BRUSH("Icons/Renderers/renderer_ribbon", Icon16x16));
+	Set("ClassIcon.NiagaraLightRendererProperties", new IMAGE_BRUSH("Icons/Renderers/renderer_light", Icon16x16));
+	Set("ClassIcon.NiagaraRendererProperties", new IMAGE_BRUSH("Icons/Renderers/renderer_default", Icon16x16));
 }
 
-void InitStackIcons(TSharedRef< FSlateStyleSet > Style)
+void FNiagaraEditorStyle::InitStackIcons()
 {
 	//GPU/CPU icons
-	Style->Set("NiagaraEditor.Stack.GPUIcon", new IMAGE_PLUGIN_BRUSH("Icons/Simulate_GPU_x40", Icon16x16));
-	Style->Set("NiagaraEditor.Stack.CPUIcon", new IMAGE_PLUGIN_BRUSH("Icons/Simulate_CPU_x40", Icon16x16));
+	Set("NiagaraEditor.Stack.GPUIcon", new IMAGE_BRUSH("Icons/Simulate_GPU_x40", Icon16x16));
+	Set("NiagaraEditor.Stack.CPUIcon", new IMAGE_BRUSH("Icons/Simulate_CPU_x40", Icon16x16));
 }
 
-void InitNiagaraSequence(TSharedRef< FSlateStyleSet > Style)
+void FNiagaraEditorStyle::InitNiagaraSequence()
 {
 	// Niagara sequence
-	Style->Set("NiagaraEditor.NiagaraSequence.DefaultTrackColor", FLinearColor(0, .25f, 0));
+	Set("NiagaraEditor.NiagaraSequence.DefaultTrackColor", FLinearColor(0, .25f, 0));
 }
 
-void InitPlatformSet(TSharedRef< FSlateStyleSet > Style)
+void FNiagaraEditorStyle::InitPlatformSet()
 {
 	const FTextBlockStyle NormalText = FEditorStyle::GetWidgetStyle<FTextBlockStyle>("NormalText");
 
 	// Niagara platform set customization
-	Style->Set("NiagaraEditor.PlatformSet.DropdownButton", new IMAGE_CORE_BRUSH("Common/ComboArrow", Icon8x8));
+	Set("NiagaraEditor.PlatformSet.DropdownButton", new CORE_IMAGE_BRUSH("Common/ComboArrow", Icon8x8));
 
-	Style->Set("NiagaraEditor.PlatformSet.ButtonText", FTextBlockStyle(NormalText)
+	Set("NiagaraEditor.PlatformSet.ButtonText", FTextBlockStyle(NormalText)
 		.SetFont(DEFAULT_FONT("Bold", 10))
 		.SetColorAndOpacity(FLinearColor(0.72f, 0.72f, 0.72f))
 		.SetHighlightColor(FLinearColor(1, 1, 1)));
 
 	// Separator in the action menus
-	Style->Set( "MenuSeparator", new BOX_CORE_BRUSH( "Common/Separator", 1/4.0f, FLinearColor(1,1,1,0.2f) ) );
+	Set( "MenuSeparator", new CORE_BOX_BRUSH( "Common/Separator", 1/4.0f, FLinearColor(1,1,1,0.2f) ) );
 	
 	const FString SmallRoundedButtonStart(TEXT("Common/SmallRoundedButtonLeft"));
 	const FString SmallRoundedButtonMiddle(TEXT("Common/SmallRoundedButtonCentre"));
 	const FString SmallRoundedButtonEnd(TEXT("Common/SmallRoundedButtonRight"));
 
-	Style->Set("NiagaraEditor.Module.Pin.TypeSelector.Button", FButtonStyle()
+	Set("NiagaraEditor.Module.Pin.TypeSelector.Button", FButtonStyle()
 		.SetNormal(FSlateNoResource())
-		.SetPressed(BOX_CORE_BRUSH("Common/Button_Pressed", 8.0f / 32.0f, FStyleColors::PrimaryPress))
-		.SetHovered(BOX_CORE_BRUSH("Common/Button_Hovered", 8.0f / 32.0f, FStyleColors::PrimaryHover))
+		.SetPressed(CORE_BOX_BRUSH("Common/Button_Pressed", 8.0f / 32.0f, FStyleColors::PrimaryPress))
+		.SetHovered(CORE_BOX_BRUSH("Common/Button_Hovered", 8.0f / 32.0f, FStyleColors::PrimaryHover))
 		.SetNormalPadding(FMargin(0, 0, 0, 0))
 		.SetPressedPadding(FMargin(0, 0, 0, 0)));
 	{
 		const FLinearColor NormalColor(0.15, 0.15, 0.15, 1);
 
-		Style->Set("NiagaraEditor.PlatformSet.StartButton", FCheckBoxStyle()
+		Set("NiagaraEditor.PlatformSet.StartButton", FCheckBoxStyle()
 			.SetCheckBoxType(ESlateCheckBoxType::ToggleButton)
-			.SetUncheckedImage(BOX_CORE_BRUSH(*SmallRoundedButtonStart, FMargin(7.f / 16.f), NormalColor))
-			.SetUncheckedPressedImage(BOX_CORE_BRUSH(*SmallRoundedButtonStart, FMargin(7.f / 16.f), FStyleColors::PrimaryPress))
-			.SetUncheckedHoveredImage(BOX_CORE_BRUSH(*SmallRoundedButtonStart, FMargin(7.f / 16.f), FStyleColors::PrimaryHover))
-			.SetCheckedHoveredImage(BOX_CORE_BRUSH(*SmallRoundedButtonStart, FMargin(7.f / 16.f), FStyleColors::SelectHover))
-			.SetCheckedPressedImage(BOX_CORE_BRUSH(*SmallRoundedButtonStart, FMargin(7.f / 16.f), FStyleColors::Select))
-			.SetCheckedImage(BOX_CORE_BRUSH(*SmallRoundedButtonStart, FMargin(7.f / 16.f), FStyleColors::Select)));
-
-		Style->Set("NiagaraEditor.PlatformSet.MiddleButton", FCheckBoxStyle()
+			.SetUncheckedImage(CORE_BOX_BRUSH(*SmallRoundedButtonStart, FMargin(7.f / 16.f), NormalColor))
+			.SetUncheckedPressedImage(CORE_BOX_BRUSH(*SmallRoundedButtonStart, FMargin(7.f / 16.f), FStyleColors::PrimaryPress))
+			.SetUncheckedHoveredImage(CORE_BOX_BRUSH(*SmallRoundedButtonStart, FMargin(7.f / 16.f), FStyleColors::PrimaryHover))
+			.SetCheckedHoveredImage(CORE_BOX_BRUSH(*SmallRoundedButtonStart, FMargin(7.f / 16.f), FStyleColors::SelectHover))
+			.SetCheckedPressedImage(CORE_BOX_BRUSH(*SmallRoundedButtonStart, FMargin(7.f / 16.f), FStyleColors::Select))
+			.SetCheckedImage(CORE_BOX_BRUSH(*SmallRoundedButtonStart, FMargin(7.f / 16.f), FStyleColors::Select)));
+
+		Set("NiagaraEditor.PlatformSet.MiddleButton", FCheckBoxStyle()
 			.SetCheckBoxType(ESlateCheckBoxType::ToggleButton)
-			.SetUncheckedImage(BOX_CORE_BRUSH(*SmallRoundedButtonMiddle, FMargin(7.f / 16.f), NormalColor))
-			.SetUncheckedPressedImage(BOX_CORE_BRUSH(*SmallRoundedButtonMiddle, FMargin(7.f / 16.f), FStyleColors::PrimaryPress))
-			.SetUncheckedHoveredImage(BOX_CORE_BRUSH(*SmallRoundedButtonMiddle, FMargin(7.f / 16.f), FStyleColors::PrimaryHover))
-			.SetCheckedHoveredImage(BOX_CORE_BRUSH(*SmallRoundedButtonMiddle, FMargin(7.f / 16.f), FStyleColors::SelectHover))
-			.SetCheckedPressedImage(BOX_CORE_BRUSH(*SmallRoundedButtonMiddle, FMargin(7.f / 16.f), FStyleColors::Select))
-			.SetCheckedImage(BOX_CORE_BRUSH(*SmallRoundedButtonMiddle, FMargin(7.f / 16.f), FStyleColors::Select)));
-
-		Style->Set("NiagaraEditor.PlatformSet.EndButton", FCheckBoxStyle()
+			.SetUncheckedImage(CORE_BOX_BRUSH(*SmallRoundedButtonMiddle, FMargin(7.f / 16.f), NormalColor))
+			.SetUncheckedPressedImage(CORE_BOX_BRUSH(*SmallRoundedButtonMiddle, FMargin(7.f / 16.f), FStyleColors::PrimaryPress))
+			.SetUncheckedHoveredImage(CORE_BOX_BRUSH(*SmallRoundedButtonMiddle, FMargin(7.f / 16.f), FStyleColors::PrimaryHover))
+			.SetCheckedHoveredImage(CORE_BOX_BRUSH(*SmallRoundedButtonMiddle, FMargin(7.f / 16.f), FStyleColors::SelectHover))
+			.SetCheckedPressedImage(CORE_BOX_BRUSH(*SmallRoundedButtonMiddle, FMargin(7.f / 16.f), FStyleColors::Select))
+			.SetCheckedImage(CORE_BOX_BRUSH(*SmallRoundedButtonMiddle, FMargin(7.f / 16.f), FStyleColors::Select)));
+
+		Set("NiagaraEditor.PlatformSet.EndButton", FCheckBoxStyle()
 			.SetCheckBoxType(ESlateCheckBoxType::ToggleButton)
-			.SetUncheckedImage(BOX_CORE_BRUSH(*SmallRoundedButtonEnd, FMargin(7.f / 16.f), NormalColor))
-			.SetUncheckedPressedImage(BOX_CORE_BRUSH(*SmallRoundedButtonEnd, FMargin(7.f / 16.f), FStyleColors::SelectHover))
-			.SetUncheckedHoveredImage(BOX_CORE_BRUSH(*SmallRoundedButtonEnd, FMargin(7.f / 16.f), FStyleColors::SelectHover))
-			.SetCheckedHoveredImage(BOX_CORE_BRUSH(*SmallRoundedButtonEnd, FMargin(7.f / 16.f), FStyleColors::SelectHover))
-			.SetCheckedPressedImage(BOX_CORE_BRUSH(*SmallRoundedButtonEnd, FMargin(7.f / 16.f), FStyleColors::Select))
-			.SetCheckedImage(BOX_CORE_BRUSH(*SmallRoundedButtonEnd, FMargin(7.f / 16.f), FStyleColors::Select)));
+			.SetUncheckedImage(CORE_BOX_BRUSH(*SmallRoundedButtonEnd, FMargin(7.f / 16.f), NormalColor))
+			.SetUncheckedPressedImage(CORE_BOX_BRUSH(*SmallRoundedButtonEnd, FMargin(7.f / 16.f), FStyleColors::SelectHover))
+			.SetUncheckedHoveredImage(CORE_BOX_BRUSH(*SmallRoundedButtonEnd, FMargin(7.f / 16.f), FStyleColors::SelectHover))
+			.SetCheckedHoveredImage(CORE_BOX_BRUSH(*SmallRoundedButtonEnd, FMargin(7.f / 16.f), FStyleColors::SelectHover))
+			.SetCheckedPressedImage(CORE_BOX_BRUSH(*SmallRoundedButtonEnd, FMargin(7.f / 16.f), FStyleColors::Select))
+			.SetCheckedImage(CORE_BOX_BRUSH(*SmallRoundedButtonEnd, FMargin(7.f / 16.f), FStyleColors::Select)));
 	}
 
-	Style->Set("NiagaraEditor.PlatformSet.Include", new IMAGE_CORE_BRUSH("Icons/PlusSymbol_12x", Icon12x12));
-	Style->Set("NiagaraEditor.PlatformSet.Exclude", new IMAGE_CORE_BRUSH("Icons/MinusSymbol_12x", Icon12x12));
-	Style->Set("NiagaraEditor.PlatformSet.Remove", new IMAGE_CORE_BRUSH("Icons/Cross_12x", Icon12x12));
+	Set("NiagaraEditor.PlatformSet.Include", new CORE_IMAGE_BRUSH("Icons/PlusSymbol_12x", Icon12x12));
+	Set("NiagaraEditor.PlatformSet.Exclude", new CORE_IMAGE_BRUSH("Icons/MinusSymbol_12x", Icon12x12));
+	Set("NiagaraEditor.PlatformSet.Remove", new CORE_IMAGE_BRUSH("Icons/Cross_12x", Icon12x12));
 
 	const FSlateColor SelectionColor_Inactive = FEditorStyle::GetSlateColor("SelectionColor_Inactive");
 
-	Style->Set("NiagaraEditor.PlatformSet.TreeView", FTableRowStyle()
+	Set("NiagaraEditor.PlatformSet.TreeView", FTableRowStyle()
 		.SetEvenRowBackgroundBrush(FSlateNoResource())
-		.SetEvenRowBackgroundHoveredBrush(IMAGE_CORE_BRUSH("Common/Selection", Icon8x8, SelectionColor_Inactive))
+		.SetEvenRowBackgroundHoveredBrush(CORE_IMAGE_BRUSH("Common/Selection", Icon8x8, SelectionColor_Inactive))
 		.SetOddRowBackgroundBrush(FSlateNoResource())
-		.SetOddRowBackgroundHoveredBrush(IMAGE_CORE_BRUSH("Common/Selection", Icon8x8, SelectionColor_Inactive))
+		.SetOddRowBackgroundHoveredBrush(CORE_IMAGE_BRUSH("Common/Selection", Icon8x8, SelectionColor_Inactive))
 		.SetSelectorFocusedBrush(FSlateNoResource())
-		.SetActiveBrush(IMAGE_CORE_BRUSH("Common/Selection", Icon8x8, FStyleColors::Select))
-		.SetActiveHoveredBrush(IMAGE_CORE_BRUSH("Common/Selection", Icon8x8, FStyleColors::Select))
-		.SetInactiveBrush(IMAGE_CORE_BRUSH("Common/Selection", Icon8x8, SelectionColor_Inactive))
-		.SetInactiveHoveredBrush(IMAGE_CORE_BRUSH("Common/Selection", Icon8x8, SelectionColor_Inactive)));
+		.SetActiveBrush(CORE_IMAGE_BRUSH("Common/Selection", Icon8x8, FStyleColors::Select))
+		.SetActiveHoveredBrush(CORE_IMAGE_BRUSH("Common/Selection", Icon8x8, FStyleColors::Select))
+		.SetInactiveBrush(CORE_IMAGE_BRUSH("Common/Selection", Icon8x8, SelectionColor_Inactive))
+		.SetInactiveHoveredBrush(CORE_IMAGE_BRUSH("Common/Selection", Icon8x8, SelectionColor_Inactive)));
 
 }
 
-void InitDropTarget(TSharedRef< FSlateStyleSet > Style)
+void FNiagaraEditorStyle::InitDropTarget()
 {
-	Style->Set("NiagaraEditor.DropTarget.BackgroundColor", FLinearColor(1.0f, 1.0f, 1.0f, 0.25f));
-	Style->Set("NiagaraEditor.DropTarget.BackgroundColorHover", FLinearColor(1.0f, 1.0f, 1.0f, 0.1f));
-	Style->Set("NiagaraEditor.DropTarget.BorderVertical", new IMAGE_PLUGIN_BRUSH("Icons/StackDropTargetBorder_Vertical", FVector2D(2, 8), FLinearColor::White, ESlateBrushTileType::Vertical));
-	Style->Set("NiagaraEditor.DropTarget.BorderHorizontal", new IMAGE_PLUGIN_BRUSH("Icons/StackDropTargetBorder_Horizontal", FVector2D(8, 2), FLinearColor::White, ESlateBrushTileType::Horizontal));
+	Set("NiagaraEditor.DropTarget.BackgroundColor", FLinearColor(1.0f, 1.0f, 1.0f, 0.25f));
+	Set("NiagaraEditor.DropTarget.BackgroundColorHover", FLinearColor(1.0f, 1.0f, 1.0f, 0.1f));
+	Set("NiagaraEditor.DropTarget.BorderVertical", new IMAGE_BRUSH("Icons/StackDropTargetBorder_Vertical", FVector2D(2, 8), FLinearColor::White, ESlateBrushTileType::Vertical));
+	Set("NiagaraEditor.DropTarget.BorderHorizontal", new IMAGE_BRUSH("Icons/StackDropTargetBorder_Horizontal", FVector2D(8, 2), FLinearColor::White, ESlateBrushTileType::Horizontal));
 }
 
-void InitScriptGraph(TSharedRef< FSlateStyleSet > Style)
+void FNiagaraEditorStyle::InitScriptGraph()
 {
-	Style->Set("NiagaraEditor.ScriptGraph.SearchBorderColor", FLinearColor(.1f, .1f, .1f, 1.f));
+	Set("NiagaraEditor.ScriptGraph.SearchBorderColor", FLinearColor(.1f, .1f, .1f, 1.f));
 }
 
-void InitDebuggerStyle(TSharedRef< FSlateStyleSet > Style)
+void FNiagaraEditorStyle::InitDebuggerStyle()
 {
 	const FVector2D Icon24x24(24.0f, 24.0f);
 	
-	Style->Set("NiagaraEditor.Debugger.PlayIcon", new IMAGE_PLUGIN_BRUSH("Icons/Debugger/Play", Icon24x24));
-	Style->Set("NiagaraEditor.Debugger.SpeedIcon", new IMAGE_PLUGIN_BRUSH("Icons/Debugger/Speed", Icon24x24));
-	Style->Set("NiagaraEditor.Debugger.PauseIcon", new IMAGE_PLUGIN_BRUSH("Icons/Debugger/Pause", Icon24x24));
-	Style->Set("NiagaraEditor.Debugger.LoopIcon", new IMAGE_PLUGIN_BRUSH("Icons/Debugger/Loop", Icon24x24));
-	Style->Set("NiagaraEditor.Debugger.StepIcon", new IMAGE_PLUGIN_BRUSH("Icons/Debugger/Step", Icon24x24));
+	Set("NiagaraEditor.Debugger.PlayIcon", new IMAGE_BRUSH("Icons/Debugger/Play", Icon24x24));
+	Set("NiagaraEditor.Debugger.SpeedIcon", new IMAGE_BRUSH("Icons/Debugger/Speed", Icon24x24));
+	Set("NiagaraEditor.Debugger.PauseIcon", new IMAGE_BRUSH("Icons/Debugger/Pause", Icon24x24));
+	Set("NiagaraEditor.Debugger.LoopIcon", new IMAGE_BRUSH("Icons/Debugger/Loop", Icon24x24));
+	Set("NiagaraEditor.Debugger.StepIcon", new IMAGE_BRUSH("Icons/Debugger/Step", Icon24x24));
 
-	Style->Set("NiagaraEditor.Debugger.Outliner.Capture", new IMAGE_PLUGIN_BRUSH("Icons/Debugger/Capture", Icon24x24));
-	Style->Set("NiagaraEditor.Debugger.Outliner.CapturePerf", new IMAGE_PLUGIN_BRUSH("Icons/Debugger/Perf_40x", Icon24x24));
-	Style->Set("NiagaraEditor.Debugger.Outliner.Filter", new IMAGE_PLUGIN_BRUSH("Icons/Debugger/Filter_24x", Icon24x24));
+	Set("NiagaraEditor.Debugger.Outliner.Capture", new IMAGE_BRUSH("Icons/Debugger/Capture", Icon24x24));
+	Set("NiagaraEditor.Debugger.Outliner.Filter", new IMAGE_BRUSH("Icons/Debugger/Filter_24x", Icon24x24));
 
 	const FSlateColor SelectionColor = FEditorStyle::GetSlateColor("SelectionColor");
 	const FSlateColor SelectionColor_Pressed = FEditorStyle::GetSlateColor("SelectionColor_Pressed");
 
 	FButtonStyle OutlinerToolBarButton = FButtonStyle()
-		.SetNormal(BOX_CORE_BRUSH("Common/FlatButton", FMargin(4 / 16.0f), FLinearColor(0.0f, 0.0f, 0.0f, 0.0f)))
-		.SetHovered(BOX_CORE_BRUSH("Common/FlatButton", FMargin(4 / 16.0f), SelectionColor))
-		.SetPressed(BOX_CORE_BRUSH("Common/FlatButton", FMargin(4 / 16.0f), SelectionColor_Pressed))
+		.SetNormal(CORE_BOX_BRUSH("Common/FlatButton", FMargin(4 / 16.0f), FLinearColor(0.0f, 0.0f, 0.0f, 0.0f)))
+		.SetHovered(CORE_BOX_BRUSH("Common/FlatButton", FMargin(4 / 16.0f), SelectionColor))
+		.SetPressed(CORE_BOX_BRUSH("Common/FlatButton", FMargin(4 / 16.0f), SelectionColor_Pressed))
 		.SetNormalPadding(FMargin(0, 0, 0, 1))
 		.SetPressedPadding(FMargin(0, 1, 0, 0));
-	Style->Set("NiagaraEditor.Debugger.Outliner.Toolbar", OutlinerToolBarButton);
-}
-
-void InitBakerStyle(TSharedRef< FSlateStyleSet > Style)
-{
-	Style->Set("NiagaraEditor.Baker", new IMAGE_PLUGIN_BRUSH("Icons/Baker/BakerIcon", Icon40x40));
+	Set("NiagaraEditor.Debugger.Outliner.Toolbar", OutlinerToolBarButton);
 }
 
-void InitCommonColors(TSharedRef< FSlateStyleSet > Style)
+void FNiagaraEditorStyle::InitBakerStyle()
 {
-	Style->Set("NiagaraEditor.CommonColors.System", FLinearColor(FColor(1, 202, 252)));
-	Style->Set("NiagaraEditor.CommonColors.Emitter", FLinearColor(FColor(241, 99, 6)));
-	Style->Set("NiagaraEditor.CommonColors.Particle", FLinearColor(FColor(131, 218, 9)));
+	Set("NiagaraEditor.Baker", new IMAGE_BRUSH("Icons/Baker/BakerIcon", Icon40x40));
 }
 
-void InitToolbar(TSharedRef< FSlateStyleSet > Style)
+void FNiagaraEditorStyle::InitCommonColors()
 {
-	const FTextBlockStyle NormalText = FEditorStyle::GetWidgetStyle<FTextBlockStyle>("NormalText");
-
-	// Overview Thumbnail toolbar
-	// The ToolbarBuilder requires a specific set of resources with specific names. Do not change names lightly.
-	{
-		Style->Set( "OverviewStackNodeThumbnailToolBar.Label", FTextBlockStyle(NormalText) .SetFont( DEFAULT_FONT( "Regular", 9 ) ) );
-		Style->Set( "OverviewStackNodeThumbnailToolBar.Background", new FSlateNoResource() );
-		Style->Set( "OverviewStackNodeThumbnailToolBar.Icon", new IMAGE_CORE_BRUSH( "Icons/icon_tab_toolbar_16px", Icon16x16 ) );
-		Style->Set( "OverviewStackNodeThumbnailToolBar.Expand", new IMAGE_CORE_BRUSH( "Icons/toolbar_expand_16x", Icon16x16) );
-		Style->Set( "OverviewStackNodeThumbnailToolBar.SubMenuIndicator", new IMAGE_CORE_BRUSH( "Common/SubmenuArrow", Icon8x8 ) );
-		Style->Set( "OverviewStackNodeThumbnailToolBar.SToolBarComboButtonBlock.Padding", FMargin(4.0f,0.0f));
-		Style->Set( "OverviewStackNodeThumbnailToolBar.SToolBarButtonBlock.Padding", FMargin(4.0f,0.0f));
-		Style->Set( "OverviewStackNodeThumbnailToolBar.SToolBarCheckComboButtonBlock.Padding", FMargin(4.0f,0.0f));
-		Style->Set( "OverviewStackNodeThumbnailToolBar.SToolBarButtonBlock.CheckBox.Padding", FMargin(4.0f,0.0f) );
-		Style->Set( "OverviewStackNodeThumbnailToolBar.SToolBarComboButtonBlock.ComboButton.Color", FCoreStyle::Get().GetSlateColor("DefaultForeground") );
-
-		Style->Set( "OverviewStackNodeThumbnailToolBar.Block.IndentedPadding", FMargin( 18.0f, 2.0f, 4.0f, 4.0f ) );
-		Style->Set( "OverviewStackNodeThumbnailToolBar.Block.Padding", FMargin( 2.0f, 2.0f, 4.0f, 4.0f ) );
-
-		Style->Set( "OverviewStackNodeThumbnailToolBar.Separator", new FSlateColorBrush( FLinearColor(FColor(96, 96, 96)) ) );
-		Style->Set( "OverviewStackNodeThumbnailToolBar.Separator.Padding", FMargin( 1.f, 0.f, 1.f, 0.f) );
-
-		const FButtonStyle Button = FButtonStyle()
-			.SetNormal( BOX_CORE_BRUSH( "Common/Button", FVector2D(32,32), 8.0f/32.0f ) )
-			.SetHovered( BOX_CORE_BRUSH( "Common/Button_Hovered", FVector2D(32,32), 8.0f/32.0f ) )
-			.SetPressed( BOX_CORE_BRUSH( "Common/Button_Pressed", FVector2D(32,32), 8.0f/32.0f ) )
-			.SetNormalPadding( FMargin( 2,2,2,2 ) )
-			.SetPressedPadding( FMargin( 2,3,2,1 ) );
-		
-		Style->Set( "OverviewStackNodeThumbnailToolBar.Button", FButtonStyle(Button)
-			.SetNormal ( FSlateNoResource() )
-			.SetPressed( BOX_CORE_BRUSH( "Common/RoundedSelection_16x", 4.0f/16.0f, FStyleColors::PrimaryPress ) )
-			.SetHovered( BOX_CORE_BRUSH( "Common/RoundedSelection_16x", 4.0f/16.0f, FStyleColors::PrimaryHover ) )
-		);
-		
-		Style->Set( "OverviewStackNodeThumbnailToolBar.Button.Normal", new FSlateNoResource() );
-		Style->Set( "OverviewStackNodeThumbnailToolBar.Button.Pressed", new BOX_CORE_BRUSH( "Common/RoundedSelection_16x", 4.0f/16.0f, FStyleColors::PrimaryPress ) );
-		Style->Set( "OverviewStackNodeThumbnailToolBar.Button.Hovered", new BOX_CORE_BRUSH( "Common/RoundedSelection_16x", 4.0f/16.0f, FStyleColors::PrimaryHover ) );
-
-		Style->Set( "OverviewStackNodeThumbnailToolBar.Button.Checked", new BOX_CORE_BRUSH( "Common/RoundedSelection_16x",  4.0f/16.0f, FStyleColors::PrimaryPress ) );
-		Style->Set( "OverviewStackNodeThumbnailToolBar.Button.Checked_Hovered", new BOX_CORE_BRUSH( "Common/RoundedSelection_16x",  4.0f/16.0f, FStyleColors::PrimaryPress ) );
-		Style->Set( "OverviewStackNodeThumbnailToolBar.Button.Checked_Pressed", new BOX_CORE_BRUSH( "Common/RoundedSelection_16x",  4.0f/16.0f, FStyleColors::Select ) );
-
-		// The Wrap combo button of the toolbar requires these to be set with these names
-		Style->Set("Menu.Background", new BOX_CORE_BRUSH( "Common/GroupBorder", FMargin(4.0f/16.0f) ));
-		Style->Set( "Menu.Block.IndentedPadding", FMargin( 18.0f, 2.0f, 4.0f, 4.0f ) );
-		Style->Set( "Menu.Block.Padding", FMargin( 2.0f, 2.0f, 4.0f, 4.0f ) );
-		Style->Set( "Menu.Label", FTextBlockStyle(NormalText) .SetFont( DEFAULT_FONT( "Regular", 9 ) ) );
-	}
+	Set("NiagaraEditor.CommonColors.System", FLinearColor(FColor(1, 202, 252)));
+	Set("NiagaraEditor.CommonColors.Emitter", FLinearColor(FColor(241, 99, 6)));
+	Set("NiagaraEditor.CommonColors.Particle", FLinearColor(FColor(131, 218, 9)));
 }
 
-void InitTabIcons(TSharedRef<FSlateStyleSet> Style)
+void FNiagaraEditorStyle::InitTabIcons()
 {
-	Style->Set("Tab.Curves", new IMAGE_PLUGIN_BRUSH_SVG("Icons/Tabs/Curves", Icon16x16));
-	Style->Set("Tab.GeneratedCode", new IMAGE_PLUGIN_BRUSH_SVG("Icons/Tabs/GeneratedCode", Icon16x16));
-	Style->Set("Tab.Log", new IMAGE_PLUGIN_BRUSH_SVG("Icons/Tabs/Log", Icon16x16));
-	Style->Set("Tab.Debugger", new IMAGE_PLUGIN_BRUSH_SVG("Icons/Tabs/NiagaraDebugger", Icon16x16));
-	Style->Set("Tab.Parameters", new IMAGE_PLUGIN_BRUSH_SVG("Icons/Tabs/Parameters", Icon16x16));
-	Style->Set("Tab.ScratchPad", new IMAGE_PLUGIN_BRUSH_SVG("Icons/Tabs/ScratchPad", Icon16x16));
-	Style->Set("Tab.ScriptStats", new IMAGE_PLUGIN_BRUSH_SVG("Icons/Tabs/ScriptStats", Icon16x16));
-	Style->Set("Tab.Settings", new IMAGE_PLUGIN_BRUSH_SVG("Icons/Tabs/Settings", Icon16x16));
-	Style->Set("Tab.Spreadsheet", new IMAGE_PLUGIN_BRUSH_SVG("Icons/Tabs/Spreadsheet", Icon16x16));
-	Style->Set("Tab.SystemOverview", new IMAGE_PLUGIN_BRUSH_SVG("Icons/Tabs/SystemOverview", Icon16x16));
-	Style->Set("Tab.Timeline", new IMAGE_PLUGIN_BRUSH_SVG("Icons/Tabs/Timeline", Icon16x16));
-	Style->Set("Tab.Viewport", new IMAGE_PLUGIN_BRUSH_SVG("Icons/Tabs/Viewport", Icon16x16));
-	Style->Set("Tab.VisualEffects", new IMAGE_PLUGIN_BRUSH_SVG("Icons/Tabs/VisualEffects", Icon16x16));
+	Set("Tab.Curves", new IMAGE_BRUSH_SVG("Icons/Tabs/Curves", Icon20x20));
+	Set("Tab.GeneratedCode", new IMAGE_BRUSH_SVG("Icons/Tabs/GeneratedCode", Icon20x20));
+	Set("Tab.Log", new IMAGE_BRUSH_SVG("Icons/Tabs/Log", Icon20x20));
+	Set("Tab.Debugger", new IMAGE_BRUSH_SVG("Icons/Tabs/NiagaraDebugger", Icon20x20));
+	Set("Tab.Parameters", new IMAGE_BRUSH_SVG("Icons/Tabs/Parameters", Icon20x20));
+	Set("Tab.ScratchPad", new IMAGE_BRUSH_SVG("Icons/Tabs/ScratchPad", Icon20x20));
+	Set("Tab.ScriptStats", new IMAGE_BRUSH_SVG("Icons/Tabs/ScriptStats", Icon20x20));
+	Set("Tab.Settings", new IMAGE_BRUSH_SVG("Icons/Tabs/Settings", Icon20x20));
+	Set("Tab.Spreadsheet", new IMAGE_BRUSH_SVG("Icons/Tabs/Spreadsheet", Icon20x20));
+	Set("Tab.SystemOverview", new IMAGE_BRUSH_SVG("Icons/Tabs/SystemOverview", Icon20x20));
+	Set("Tab.Timeline", new IMAGE_BRUSH_SVG("Icons/Tabs/Timeline", Icon20x20));
+	Set("Tab.Viewport", new IMAGE_BRUSH_SVG("Icons/Tabs/Viewport", Icon20x20));
+	Set("Tab.VisualEffects", new IMAGE_BRUSH_SVG("Icons/Tabs/VisualEffects", Icon20x20));
 }
 
-void InitOutlinerStyle(TSharedRef< FSlateStyleSet > Style)
+void FNiagaraEditorStyle::InitOutlinerStyle()
 {
 	const FSlateColor SelectionColor = FEditorStyle::GetSlateColor("SelectionColor");
 	const FSlateColor SelectionColor_Inactive = FEditorStyle::GetSlateColor("SelectionColor_Inactive");
 
-	Style->Set("NiagaraEditor.Outliner.WorldItem", FTableRowStyle()
+	Set("NiagaraEditor.Outliner.WorldItem", FTableRowStyle()
 		.SetEvenRowBackgroundBrush(FSlateNoResource())
-		.SetEvenRowBackgroundHoveredBrush(IMAGE_CORE_BRUSH("Common/Selection", Icon8x8, SelectionColor_Inactive))
+		.SetEvenRowBackgroundHoveredBrush(CORE_IMAGE_BRUSH("Common/Selection", Icon8x8, SelectionColor_Inactive))
 		.SetOddRowBackgroundBrush(FSlateNoResource())
-		.SetOddRowBackgroundHoveredBrush(IMAGE_CORE_BRUSH("Common/Selection", Icon8x8, SelectionColor_Inactive))
+		.SetOddRowBackgroundHoveredBrush(CORE_IMAGE_BRUSH("Common/Selection", Icon8x8, SelectionColor_Inactive))
 		.SetSelectorFocusedBrush(FSlateNoResource())
-		.SetActiveBrush(IMAGE_CORE_BRUSH("Common/Selection", Icon8x8, SelectionColor))
-		.SetActiveHoveredBrush(IMAGE_CORE_BRUSH("Common/Selection", Icon8x8, SelectionColor))
-		.SetInactiveBrush(IMAGE_CORE_BRUSH("Common/Selection", Icon8x8, SelectionColor_Inactive))
-		.SetInactiveHoveredBrush(IMAGE_CORE_BRUSH("Common/Selection", Icon8x8, SelectionColor_Inactive)));
+		.SetActiveBrush(CORE_IMAGE_BRUSH("Common/Selection", Icon8x8, SelectionColor))
+		.SetActiveHoveredBrush(CORE_IMAGE_BRUSH("Common/Selection", Icon8x8, SelectionColor))
+		.SetInactiveBrush(CORE_IMAGE_BRUSH("Common/Selection", Icon8x8, SelectionColor_Inactive))
+		.SetInactiveHoveredBrush(CORE_IMAGE_BRUSH("Common/Selection", Icon8x8, SelectionColor_Inactive)));
 
 
-	const FLinearColor SystemColor = Style->GetColor("NiagaraEditor.AssetColors.System") * 0.6f;
+	const FLinearColor SystemColor = GetColor("NiagaraEditor.AssetColors.System") * 0.6f;
 	const FLinearColor SystemColorEven = SystemColor * 0.85f;
 	const FLinearColor SystemColorOdd = SystemColor * 0.7f;
-	Style->Set("NiagaraEditor.Outliner.SystemItem", FTableRowStyle()
-		.SetEvenRowBackgroundBrush(IMAGE_CORE_BRUSH("Common/Selection", Icon8x8, SystemColorEven))
-		.SetEvenRowBackgroundHoveredBrush(IMAGE_CORE_BRUSH("Common/Selection", Icon8x8, SystemColor))
-		.SetOddRowBackgroundBrush(IMAGE_CORE_BRUSH("Common/Selection", Icon8x8, SystemColorOdd))
-		.SetOddRowBackgroundHoveredBrush(IMAGE_CORE_BRUSH("Common/Selection", Icon8x8, SystemColor))
+	Set("NiagaraEditor.Outliner.SystemItem", FTableRowStyle()
+		.SetEvenRowBackgroundBrush(CORE_IMAGE_BRUSH("Common/Selection", Icon8x8, SystemColorEven))
+		.SetEvenRowBackgroundHoveredBrush(CORE_IMAGE_BRUSH("Common/Selection", Icon8x8, SystemColor))
+		.SetOddRowBackgroundBrush(CORE_IMAGE_BRUSH("Common/Selection", Icon8x8, SystemColorOdd))
+		.SetOddRowBackgroundHoveredBrush(CORE_IMAGE_BRUSH("Common/Selection", Icon8x8, SystemColor))
 		.SetSelectorFocusedBrush(FSlateNoResource())
-		.SetActiveBrush(IMAGE_CORE_BRUSH("Common/Selection", Icon8x8, SelectionColor))
-		.SetActiveHoveredBrush(IMAGE_CORE_BRUSH("Common/Selection", Icon8x8, SelectionColor))
-		.SetInactiveBrush(IMAGE_CORE_BRUSH("Common/Selection", Icon8x8, SelectionColor_Inactive))
-		.SetInactiveHoveredBrush(IMAGE_CORE_BRUSH("Common/Selection", Icon8x8, SelectionColor_Inactive)));
+		.SetActiveBrush(CORE_IMAGE_BRUSH("Common/Selection", Icon8x8, SelectionColor))
+		.SetActiveHoveredBrush(CORE_IMAGE_BRUSH("Common/Selection", Icon8x8, SelectionColor))
+		.SetInactiveBrush(CORE_IMAGE_BRUSH("Common/Selection", Icon8x8, SelectionColor_Inactive))
+		.SetInactiveHoveredBrush(CORE_IMAGE_BRUSH("Common/Selection", Icon8x8, SelectionColor_Inactive)));
 
 
-	const FLinearColor SystemInstanceColor = Style->GetColor("NiagaraEditor.CommonColors.System") * 0.6f;
+	const FLinearColor SystemInstanceColor = GetColor("NiagaraEditor.CommonColors.System") * 0.6f;
 	const FLinearColor SystemInstanceColorEven = SystemInstanceColor * 0.85f;
 	const FLinearColor SystemInstanceColorOdd = SystemInstanceColor * 0.7f;
-	Style->Set("NiagaraEditor.Outliner.ComponentItem", FTableRowStyle()
-		.SetEvenRowBackgroundBrush(IMAGE_CORE_BRUSH("Common/Selection", Icon8x8, SystemInstanceColorEven))
-		.SetEvenRowBackgroundHoveredBrush(IMAGE_CORE_BRUSH("Common/Selection", Icon8x8, SystemInstanceColor))
-		.SetOddRowBackgroundBrush(IMAGE_CORE_BRUSH("Common/Selection", Icon8x8, SystemInstanceColorOdd))
-		.SetOddRowBackgroundHoveredBrush(IMAGE_CORE_BRUSH("Common/Selection", Icon8x8, SystemInstanceColor))
+	Set("NiagaraEditor.Outliner.ComponentItem", FTableRowStyle()
+		.SetEvenRowBackgroundBrush(CORE_IMAGE_BRUSH("Common/Selection", Icon8x8, SystemInstanceColorEven))
+		.SetEvenRowBackgroundHoveredBrush(CORE_IMAGE_BRUSH("Common/Selection", Icon8x8, SystemInstanceColor))
+		.SetOddRowBackgroundBrush(CORE_IMAGE_BRUSH("Common/Selection", Icon8x8, SystemInstanceColorOdd))
+		.SetOddRowBackgroundHoveredBrush(CORE_IMAGE_BRUSH("Common/Selection", Icon8x8, SystemInstanceColor))
 		.SetSelectorFocusedBrush(FSlateNoResource())
-		.SetActiveBrush(IMAGE_CORE_BRUSH("Common/Selection", Icon8x8, SelectionColor))
-		.SetActiveHoveredBrush(IMAGE_CORE_BRUSH("Common/Selection", Icon8x8, SelectionColor))
-		.SetInactiveBrush(IMAGE_CORE_BRUSH("Common/Selection", Icon8x8, SelectionColor_Inactive))
-		.SetInactiveHoveredBrush(IMAGE_CORE_BRUSH("Common/Selection", Icon8x8, SelectionColor_Inactive)));
+		.SetActiveBrush(CORE_IMAGE_BRUSH("Common/Selection", Icon8x8, SelectionColor))
+		.SetActiveHoveredBrush(CORE_IMAGE_BRUSH("Common/Selection", Icon8x8, SelectionColor))
+		.SetInactiveBrush(CORE_IMAGE_BRUSH("Common/Selection", Icon8x8, SelectionColor_Inactive))
+		.SetInactiveHoveredBrush(CORE_IMAGE_BRUSH("Common/Selection", Icon8x8, SelectionColor_Inactive)));
 
-	const FLinearColor EmitterInstanceColor = Style->GetColor("NiagaraEditor.CommonColors.Emitter") * 0.6f;
+	const FLinearColor EmitterInstanceColor = GetColor("NiagaraEditor.CommonColors.Emitter") * 0.6f;
 	const FLinearColor EmitterInstanceColorEven = EmitterInstanceColor * 0.85f;
 	const FLinearColor EmitterInstanceColorOdd = EmitterInstanceColor * 0.7f;
-	Style->Set("NiagaraEditor.Outliner.EmitterItem", FTableRowStyle()
-		.SetEvenRowBackgroundBrush(IMAGE_CORE_BRUSH("Common/Selection", Icon8x8, EmitterInstanceColorEven))
-		.SetEvenRowBackgroundHoveredBrush(IMAGE_CORE_BRUSH("Common/Selection", Icon8x8, EmitterInstanceColor))
-		.SetOddRowBackgroundBrush(IMAGE_CORE_BRUSH("Common/Selection", Icon8x8, EmitterInstanceColorOdd))
-		.SetOddRowBackgroundHoveredBrush(IMAGE_CORE_BRUSH("Common/Selection", Icon8x8, EmitterInstanceColor))
+	Set("NiagaraEditor.Outliner.EmitterItem", FTableRowStyle()
+		.SetEvenRowBackgroundBrush(CORE_IMAGE_BRUSH("Common/Selection", Icon8x8, EmitterInstanceColorEven))
+		.SetEvenRowBackgroundHoveredBrush(CORE_IMAGE_BRUSH("Common/Selection", Icon8x8, EmitterInstanceColor))
+		.SetOddRowBackgroundBrush(CORE_IMAGE_BRUSH("Common/Selection", Icon8x8, EmitterInstanceColorOdd))
+		.SetOddRowBackgroundHoveredBrush(CORE_IMAGE_BRUSH("Common/Selection", Icon8x8, EmitterInstanceColor))
 		.SetSelectorFocusedBrush(FSlateNoResource())
-		.SetActiveBrush(IMAGE_CORE_BRUSH("Common/Selection", Icon8x8, SelectionColor))
-		.SetActiveHoveredBrush(IMAGE_CORE_BRUSH("Common/Selection", Icon8x8, SelectionColor))
-		.SetInactiveBrush(IMAGE_CORE_BRUSH("Common/Selection", Icon8x8, SelectionColor_Inactive))
-		.SetInactiveHoveredBrush(IMAGE_CORE_BRUSH("Common/Selection", Icon8x8, SelectionColor_Inactive)));
-}
-
-} // namespace NiagaraEditorStyleImpl
-
-TSharedRef< FSlateStyleSet > FNiagaraEditorStyle::Create()
-{
-	using namespace NiagaraEditorStyleImpl;
-
-	TSharedRef< FSlateStyleSet > Style = MakeShareable(new FSlateStyleSet("NiagaraEditorStyle"));
-	Style->SetContentRoot(FPaths::EngineContentDir() / TEXT("Editor/Slate/Niagara"));
-
-	InitStats(Style);
-	InitAssetPicker(Style);
-	InitActionMenu(Style);
-	InitEmitterHeader(Style);
-	InitParameters(Style);
-	InitParameterMapView(Style);
-	InitCodeView(Style);
-	InitSelectedEmitter(Style);
-	InitToolbarIcons(Style);
-	InitToolbar(Style);
-	InitTabIcons(Style);
-	InitIcons(Style);
-	InitOverview(Style);
-	InitEmitterDetails(Style);
-	InitAssetColors(Style);
-	InitThumbnails(Style);
-	InitClassIcon(Style);
-	InitStackIcons(Style);
-	InitNiagaraSequence(Style);
-	InitPlatformSet(Style);
-	InitDropTarget(Style);
-	InitScriptGraph(Style);
-	InitDebuggerStyle(Style);
-	InitBakerStyle(Style);
-	InitCommonColors(Style);
-	InitOutlinerStyle(Style);
-
-	return Style;
-}
-
-#undef IMAGE_PLUGIN_BRUSH
-#undef IMAGE_CORE_BRUSH
-#undef IMAGE_BRUSH
-#undef BOX_BRUSH
-#undef BORDER_BRUSH
-#undef BOX_CORE_BRUSH
-#undef DEFAULT_FONT
+		.SetActiveBrush(CORE_IMAGE_BRUSH("Common/Selection", Icon8x8, SelectionColor))
+		.SetActiveHoveredBrush(CORE_IMAGE_BRUSH("Common/Selection", Icon8x8, SelectionColor))
+		.SetInactiveBrush(CORE_IMAGE_BRUSH("Common/Selection", Icon8x8, SelectionColor_Inactive))
+		.SetInactiveHoveredBrush(CORE_IMAGE_BRUSH("Common/Selection", Icon8x8, SelectionColor_Inactive)));
+}
+
+void FNiagaraEditorStyle::InitScalabilityColors()
+{
+	Set("NiagaraEditor.Scalability.Included", FLinearColor(0.12, 0.62, 0.85f, 1));
+	Set("NiagaraEditor.Scalability.Excluded", FLinearColor(0.15, 0.15, 0.15, 1));
+
+	Set("NiagaraEditor.Scalability.System.Feature.Source.EffectType", FLinearColor(1.f, 1.f, 1.f, 1));
+	Set("NiagaraEditor.Scalability.System.Feature.Source.Override", FLinearColor(0.8, 0.8, 0.15, 1));
+	Set("NiagaraEditor.Scalability.System.Feature.Active", FLinearColor(1.f, 1.f, 1.f, 1));
+	Set("NiagaraEditor.Scalability.System.Feature.Inactive", FLinearColor(0.15, 0.15, 0.15, 1));
+
+	Set("NiagaraEditor.SystemOverview.ExcludedFromScalability", FLinearColor(0.7, 0.1, 0.1f, 0.5));
+	Set("NiagaraEditor.SystemOverview.ExcludedFromScalability.NodeBody", new FSlateRoundedBoxBrush(GetColor("NiagaraEditor.SystemOverview.ExcludedFromScalability"), FVector4(5, 5, 5, 5)));
+	Set("NiagaraEditor.SystemOverview.ExcludedFromScalability.RendererItem", new FSlateRoundedBoxBrush(GetColor("NiagaraEditor.SystemOverview.ExcludedFromScalability"), 5.f));
+	Set("NiagaraEditor.SystemOverview.ExcludedFromScalability.EmitterTrack", new FSlateRoundedBoxBrush(GetColor("NiagaraEditor.SystemOverview.ExcludedFromScalability"), 1.f));
+}
+
+void FNiagaraEditorStyle::InitScalabilityIcons()
+{
+	Set("NiagaraEditor.Scalability.System.Feature.Override", new IMAGE_BRUSH_SVG("Icons/Scalability/Override", Icon20x20));
+	Set("NiagaraEditor.Scalability.Preview.ResetPlatform", new IMAGE_BRUSH("Icons/Scalability/ResetPreviewPlatform", Icon20x20));
+}
+
+FNiagaraEditorStyle::FNiagaraEditorStyle() : FSlateStyleSet("NiagaraEditorStyle")
+{	
+	SetContentRoot(FPaths::EnginePluginsDir() / TEXT("FX/Niagara/Content/Slate"));
+	SetCoreContentRoot(FPaths::EngineContentDir() / TEXT("Editor/Slate"));
+	
+	InitStats();
+	InitAssetPicker();
+	InitActionMenu();
+	InitEmitterHeader();
+	InitParameters();
+	InitParameterMapView();
+	InitCodeView();
+	InitSelectedEmitter();
+	InitToolbarIcons();
+	InitTabIcons();
+	InitIcons();
+	InitOverview();
+	InitEmitterDetails();
+	InitAssetColors();
+	InitThumbnails();
+	InitClassIcon();
+	InitStackIcons();
+	InitNiagaraSequence();
+	InitPlatformSet();
+	InitDropTarget();
+	InitScriptGraph();
+	InitDebuggerStyle();
+	InitBakerStyle();
+	InitCommonColors();
+	InitOutlinerStyle();
+	InitScalabilityColors();
+	InitScalabilityIcons();
+}
 
 void FNiagaraEditorStyle::ReloadTextures()
 {
 	FSlateApplication::Get().GetRenderer()->ReloadTextureResources();
 }
 
-const ISlateStyle& FNiagaraEditorStyle::Get()
+const FNiagaraEditorStyle& FNiagaraEditorStyle::Get()
+{
+	if(!NiagaraEditorStyle.IsValid())
+	{
+		NiagaraEditorStyle = MakeShareable(new FNiagaraEditorStyle());
+	}
+	
+	return *NiagaraEditorStyle;
+}
+
+void FNiagaraEditorStyle::ReinitializeStyle()
 {
-	return *NiagaraEditorStyleInstance;
+	Unregister();
+	NiagaraEditorStyle.Reset();
+	Register();	
 }
diff --git a/Engine/Plugins/FX/Niagara/Source/NiagaraEditor/Private/Toolkits/NiagaraParameterCollectionToolkit.cpp b/Engine/Plugins/FX/Niagara/Source/NiagaraEditor/Private/Toolkits/NiagaraParameterCollectionToolkit.cpp
--- a/Engine/Plugins/FX/Niagara/Source/NiagaraEditor/Private/Toolkits/NiagaraParameterCollectionToolkit.cpp
+++ b/Engine/Plugins/FX/Niagara/Source/NiagaraEditor/Private/Toolkits/NiagaraParameterCollectionToolkit.cpp
@@ -229,7 +229,7 @@ void FNiagaraParameterCollectionToolkit::ExtendToolbar()
 	// 					GetEmitterMenuContent,
 	// 					LOCTEXT("AddEmitterButtonText", "Add Emitter"),
 	// 					LOCTEXT("AddEmitterButtonTextToolTip", "Adds an emitter to the system from an existing emitter asset."),
-	// 					FSlateIcon(FNiagaraEditorStyle::GetStyleSetName(), "NiagaraEditor.AddEmitter"));
+	// 					FSlateIcon(FNiagaraEditorStyle::Get().GetStyleSetName(), "NiagaraEditor.AddEmitter"));
 	// 			}
 	// 			ToolbarBuilder.EndSection();
 	// 		}
diff --git a/Engine/Plugins/FX/Niagara/Source/NiagaraEditor/Private/Toolkits/NiagaraParameterDefinitionsToolkit.cpp b/Engine/Plugins/FX/Niagara/Source/NiagaraEditor/Private/Toolkits/NiagaraParameterDefinitionsToolkit.cpp
--- a/Engine/Plugins/FX/Niagara/Source/NiagaraEditor/Private/Toolkits/NiagaraParameterDefinitionsToolkit.cpp
+++ b/Engine/Plugins/FX/Niagara/Source/NiagaraEditor/Private/Toolkits/NiagaraParameterDefinitionsToolkit.cpp
@@ -274,7 +274,7 @@ void FNiagaraParameterDefinitionsToolkit::ExtendToolbar()
 				{
 					ToolbarBuilder.AddToolBarButton(FNiagaraEditorCommands::Get().Apply,
 						NAME_None, TAttribute<FText>(), TAttribute<FText>(),
-						FSlateIcon(FNiagaraEditorStyle::GetStyleSetName(), "NiagaraEditor.Apply"),
+						FSlateIcon(FNiagaraEditorStyle::Get().GetStyleSetName(), "NiagaraEditor.Apply"),
 						FName(TEXT("ApplyNiagaraEmitter")));
 				}
 				ToolbarBuilder.EndSection();
diff --git a/Engine/Plugins/FX/Niagara/Source/NiagaraEditor/Private/Toolkits/NiagaraScriptToolkit.cpp b/Engine/Plugins/FX/Niagara/Source/NiagaraEditor/Private/Toolkits/NiagaraScriptToolkit.cpp
--- a/Engine/Plugins/FX/Niagara/Source/NiagaraEditor/Private/Toolkits/NiagaraScriptToolkit.cpp
+++ b/Engine/Plugins/FX/Niagara/Source/NiagaraEditor/Private/Toolkits/NiagaraScriptToolkit.cpp
@@ -627,7 +627,7 @@ void FNiagaraScriptToolkit::ExtendToolbar()
 			ToolbarBuilder.BeginSection("Apply");
 			ToolbarBuilder.AddToolBarButton(FNiagaraEditorCommands::Get().Apply,
 				NAME_None, TAttribute<FText>(), TAttribute<FText>(), 
-				FSlateIcon(FNiagaraEditorStyle::GetStyleSetName(), "NiagaraEditor.Apply"),
+				FSlateIcon(FAppStyle::Get().GetStyleSetName(), "AssetEditor.Apply"),
 				FName(TEXT("ApplyNiagaraScript")));
 			ToolbarBuilder.EndSection();
 
@@ -731,7 +731,7 @@ FText FNiagaraScriptToolkit::GetCompileStatusTooltip() const
 
 FSlateIcon FNiagaraScriptToolkit::GetRefreshStatusImage() const
 {
-	return FSlateIcon(FNiagaraEditorStyle::GetStyleSetName(), "Niagara.Asset.ReimportAsset.Default");
+	return FSlateIcon(FNiagaraEditorStyle::Get().GetStyleSetName(), "NiagaraEditor.Refresh");
 }
 
 FText FNiagaraScriptToolkit::GetRefreshStatusTooltip() const
diff --git a/Engine/Plugins/FX/Niagara/Source/NiagaraEditor/Private/Toolkits/NiagaraSystemToolkit.cpp b/Engine/Plugins/FX/Niagara/Source/NiagaraEditor/Private/Toolkits/NiagaraSystemToolkit.cpp
--- a/Engine/Plugins/FX/Niagara/Source/NiagaraEditor/Private/Toolkits/NiagaraSystemToolkit.cpp
+++ b/Engine/Plugins/FX/Niagara/Source/NiagaraEditor/Private/Toolkits/NiagaraSystemToolkit.cpp
@@ -108,17 +108,17 @@ void FNiagaraSystemToolkit::RegisterTabSpawners(const TSharedRef<class FTabManag
 	InTabManager->RegisterTabSpawner(ViewportTabID, FOnSpawnTab::CreateSP(this, &FNiagaraSystemToolkit::SpawnTab_Viewport))
 		.SetDisplayName(LOCTEXT("Preview", "Preview"))
 		.SetGroup(WorkspaceMenuCategory.ToSharedRef())
-		.SetIcon(FSlateIcon(FNiagaraEditorStyle::GetStyleSetName(), "Tab.Viewport"));
+		.SetIcon(FSlateIcon(FNiagaraEditorStyle::Get().GetStyleSetName(), "Tab.Viewport"));
 
 	InTabManager->RegisterTabSpawner(CurveEditorTabID, FOnSpawnTab::CreateSP(this, &FNiagaraSystemToolkit::SpawnTab_CurveEd))
 		.SetDisplayName(LOCTEXT("Curves", "Curves"))
 		.SetGroup(WorkspaceMenuCategory.ToSharedRef())
-		.SetIcon(FSlateIcon(FNiagaraEditorStyle::GetStyleSetName(), "Tab.Curves"));
+		.SetIcon(FSlateIcon(FNiagaraEditorStyle::Get().GetStyleSetName(), "Tab.Curves"));
 
 	InTabManager->RegisterTabSpawner(SequencerTabID, FOnSpawnTab::CreateSP(this, &FNiagaraSystemToolkit::SpawnTab_Sequencer))
 		.SetDisplayName(LOCTEXT("Timeline", "Timeline"))
 		.SetGroup(WorkspaceMenuCategory.ToSharedRef())
-		.SetIcon(FSlateIcon(FNiagaraEditorStyle::GetStyleSetName(), "Tab.Timeline"));
+		.SetIcon(FSlateIcon(FNiagaraEditorStyle::Get().GetStyleSetName(), "Tab.Timeline"));
 
 	InTabManager->RegisterTabSpawner(SystemScriptTabID, FOnSpawnTab::CreateSP(this, &FNiagaraSystemToolkit::SpawnTab_SystemScript))
 		.SetDisplayName(LOCTEXT("SystemScript", "System Script"))
@@ -128,12 +128,12 @@ void FNiagaraSystemToolkit::RegisterTabSpawners(const TSharedRef<class FTabManag
 	InTabManager->RegisterTabSpawner(SystemParametersTabID, FOnSpawnTab::CreateSP(this, &FNiagaraSystemToolkit::SpawnTab_SystemParameters))
 		.SetDisplayName(LOCTEXT("SystemParameters", "Parameters"))
 		.SetGroup(WorkspaceMenuCategory.ToSharedRef())
-		.SetIcon(FSlateIcon(FNiagaraEditorStyle::GetStyleSetName(), "Tab.Parameters"));
+		.SetIcon(FSlateIcon(FNiagaraEditorStyle::Get().GetStyleSetName(), "Tab.Parameters"));
 
 	InTabManager->RegisterTabSpawner(SystemParametersTabID2, FOnSpawnTab::CreateSP(this, &FNiagaraSystemToolkit::SpawnTab_SystemParameters2))
 		.SetDisplayName(LOCTEXT("SystemParameters2", "Legacy Parameters"))
 		.SetGroup(WorkspaceMenuCategory.ToSharedRef())
-		.SetIcon(FSlateIcon(FNiagaraEditorStyle::GetStyleSetName(), "Tab.Parameters"));
+		.SetIcon(FSlateIcon(FNiagaraEditorStyle::Get().GetStyleSetName(), "Tab.Parameters"));
 
 //@todo(ng) disable parameter definitions panel pending bug fixes
 // 	InTabManager->RegisterTabSpawner(SystemParameterDefinitionsTabID, FOnSpawnTab::CreateSP(this, &FNiagaraSystemToolkit::SpawnTab_SystemParameterDefinitions)) 
@@ -143,7 +143,7 @@ void FNiagaraSystemToolkit::RegisterTabSpawners(const TSharedRef<class FTabManag
 	InTabManager->RegisterTabSpawner(SelectedEmitterStackTabID, FOnSpawnTab::CreateSP(this, &FNiagaraSystemToolkit::SpawnTab_SelectedEmitterStack))
 		.SetDisplayName(LOCTEXT("SelectedEmitterStacks", "Selected Emitters"))
 		.SetGroup(WorkspaceMenuCategory.ToSharedRef())
-		.SetIcon(FSlateIcon(FNiagaraEditorStyle::GetStyleSetName(), "Tab.VisualEffects"));
+		.SetIcon(FSlateIcon(FNiagaraEditorStyle::Get().GetStyleSetName(), "Tab.VisualEffects"));
 
 	InTabManager->RegisterTabSpawner(SelectedEmitterGraphTabID, FOnSpawnTab::CreateSP(this, &FNiagaraSystemToolkit::SpawnTab_SelectedEmitterGraph))
 		.SetDisplayName(LOCTEXT("SelectedEmitterGraph", "Selected Emitter Graph"))
@@ -153,37 +153,37 @@ void FNiagaraSystemToolkit::RegisterTabSpawners(const TSharedRef<class FTabManag
 	InTabManager->RegisterTabSpawner(DebugSpreadsheetTabID, FOnSpawnTab::CreateSP(this, &FNiagaraSystemToolkit::SpawnTab_DebugSpreadsheet))
 		.SetDisplayName(LOCTEXT("DebugSpreadsheet", "Attribute Spreadsheet"))
 		.SetGroup(WorkspaceMenuCategory.ToSharedRef())
-		.SetIcon(FSlateIcon(FNiagaraEditorStyle::GetStyleSetName(), "Tab.Spreadsheet"));
+		.SetIcon(FSlateIcon(FNiagaraEditorStyle::Get().GetStyleSetName(), "Tab.Spreadsheet"));
 
 	InTabManager->RegisterTabSpawner(PreviewSettingsTabId, FOnSpawnTab::CreateSP(this, &FNiagaraSystemToolkit::SpawnTab_PreviewSettings))
 		.SetDisplayName(LOCTEXT("PreviewSceneSettingsTab", "Preview Scene Settings"))
 		.SetGroup(WorkspaceMenuCategory.ToSharedRef())
-		.SetIcon(FSlateIcon(FNiagaraEditorStyle::GetStyleSetName(), "Tab.Settings"));
+		.SetIcon(FSlateIcon(FNiagaraEditorStyle::Get().GetStyleSetName(), "Tab.Settings"));
 
 	InTabManager->RegisterTabSpawner(GeneratedCodeTabID, FOnSpawnTab::CreateSP(this, &FNiagaraSystemToolkit::SpawnTab_GeneratedCode))
 		.SetDisplayName(LOCTEXT("GeneratedCode", "Generated Code"))
 		.SetGroup(WorkspaceMenuCategory.ToSharedRef())
-		.SetIcon(FSlateIcon(FNiagaraEditorStyle::GetStyleSetName(), "Tab.GeneratedCode"));
+		.SetIcon(FSlateIcon(FNiagaraEditorStyle::Get().GetStyleSetName(), "Tab.GeneratedCode"));
 
 	InTabManager->RegisterTabSpawner(MessageLogTabID, FOnSpawnTab::CreateSP(this, &FNiagaraSystemToolkit::SpawnTab_MessageLog))
 		.SetDisplayName(LOCTEXT("NiagaraMessageLog", "Niagara Log"))
 		.SetGroup(WorkspaceMenuCategory.ToSharedRef())
-		.SetIcon(FSlateIcon(FNiagaraEditorStyle::GetStyleSetName(), "Tab.Log"));
+		.SetIcon(FSlateIcon(FNiagaraEditorStyle::Get().GetStyleSetName(), "Tab.Log"));
 
 	InTabManager->RegisterTabSpawner(SystemOverviewTabID, FOnSpawnTab::CreateSP(this, &FNiagaraSystemToolkit::SpawnTab_SystemOverview))
 		.SetDisplayName(LOCTEXT("SystemOverviewTabName", "System Overview"))
 		.SetGroup(WorkspaceMenuCategory.ToSharedRef())
-		.SetIcon(FSlateIcon(FNiagaraEditorStyle::GetStyleSetName(), "Tab.SystemOverview"));
+		.SetIcon(FSlateIcon(FNiagaraEditorStyle::Get().GetStyleSetName(), "Tab.SystemOverview"));
 
 	InTabManager->RegisterTabSpawner(ScratchPadTabID, FOnSpawnTab::CreateSP(this, &FNiagaraSystemToolkit::SpawnTab_ScratchPad))
 		.SetDisplayName(LOCTEXT("ScratchPadTabName", "Scratch Pad"))
 		.SetGroup(WorkspaceMenuCategory.ToSharedRef())
-		.SetIcon(FSlateIcon(FNiagaraEditorStyle::GetStyleSetName(), "Tab.ScratchPad"));
+		.SetIcon(FSlateIcon(FNiagaraEditorStyle::Get().GetStyleSetName(), "Tab.ScratchPad"));
 
 	InTabManager->RegisterTabSpawner(ScriptStatsTabID, FOnSpawnTab::CreateSP(this, &FNiagaraSystemToolkit::SpawnTab_ScriptStats))
 		.SetDisplayName(LOCTEXT("NiagaraScriptsStatsTab", "Script Stats"))
 		.SetGroup(WorkspaceMenuCategory.ToSharedRef())
-		.SetIcon(FSlateIcon(FNiagaraEditorStyle::GetStyleSetName(), "Tab.ScriptStats"));
+		.SetIcon(FSlateIcon(FNiagaraEditorStyle::Get().GetStyleSetName(), "Tab.ScriptStats"));
 
 	if (GetDefault<UNiagaraEditorSettings>()->bEnableBaker)
 	{
@@ -1174,7 +1174,7 @@ void FNiagaraSystemToolkit::ExtendToolbar()
 					NAME_None,
 					TAttribute<FText>(),
 					TAttribute<FText>(),
-					FSlateIcon(FNiagaraEditorStyle::GetStyleSetName(), "NiagaraEditor.ApplyScratchPadChanges"),
+					FSlateIcon(FNiagaraEditorStyle::Get().GetStyleSetName(), "NiagaraEditor.ApplyScratchPadChanges"),
 					FName(TEXT("ApplyScratchPadChanges")));
 				ToolbarBuilder.AddToolBarButton(FNiagaraEditorCommands::Get().Compile,
 					NAME_None,
@@ -1259,7 +1259,7 @@ void FNiagaraSystemToolkit::ExtendToolbar()
 						FOnGetContent::CreateStatic(Local::GenerateBakerMenu, Toolkit),
 						LOCTEXT("Baker", "Baker"),
 						LOCTEXT("BakerTooltip", "Options for Baker rendering."),
-						FSlateIcon(FNiagaraEditorStyle::GetStyleSetName(), "NiagaraEditor.Baker")
+						FSlateIcon(FNiagaraEditorStyle::Get().GetStyleSetName(), "NiagaraEditor.Baker")
 					);
 				}
 				ToolbarBuilder.EndSection();
diff --git a/Engine/Plugins/FX/Niagara/Source/NiagaraEditor/Private/Widgets/SNiagaraDebugger.cpp b/Engine/Plugins/FX/Niagara/Source/NiagaraEditor/Private/Widgets/SNiagaraDebugger.cpp
--- a/Engine/Plugins/FX/Niagara/Source/NiagaraEditor/Private/Widgets/SNiagaraDebugger.cpp
+++ b/Engine/Plugins/FX/Niagara/Source/NiagaraEditor/Private/Widgets/SNiagaraDebugger.cpp
@@ -301,7 +301,7 @@ namespace NiagaraOutlinerTab
 					NAME_None,
 					LOCTEXT("Capture", "Capture"),
 					LOCTEXT("CaptureTooltip", "Triggers a capture of outliner info from the connected session."),
-					FSlateIcon(FNiagaraEditorStyle::GetStyleSetName(), "NiagaraEditor.Debugger.Outliner.Capture"),
+					FSlateIcon(FNiagaraEditorStyle::Get().GetStyleSetName(), "NiagaraEditor.Debugger.Outliner.Capture"),
 					EUserInterfaceActionType::Button
 				);
 			}
@@ -324,7 +324,7 @@ namespace NiagaraOutlinerTab
 					NAME_None,
 					LOCTEXT("GatherOutlinerPerfData", "Perf"),
 					LOCTEXT("GatherOutlinerPerfDataTooltip", "Gather Performance data during outliner capture."),
-					FSlateIcon(FNiagaraEditorStyle::GetStyleSetName(), "NiagaraEditor.Debugger.Outliner.CapturePerf"),
+					FSlateIcon(FAppStyle::Get().GetStyleSetName(), "AssetEditor.ToggleStats"),
 					EUserInterfaceActionType::ToggleButton
 				);
 			}
@@ -791,7 +791,7 @@ void SNiagaraDebugger::RegisterTabSpawner()
 	.SetDisplayName(NSLOCTEXT("UnrealEditor", "NiagaraDebuggerTab", "Niagara Debugger"))
 	.SetTooltipText(NSLOCTEXT("UnrealEditor", "NiagaraDebuggerTooltipText", "Open the Niagara Debugger Tab."))
 	.SetGroup(WorkspaceMenu::GetMenuStructure().GetDeveloperToolsDebugCategory())
-	.SetIcon(FSlateIcon(FNiagaraEditorStyle::GetStyleSetName(), "Tab.Debugger"));
+	.SetIcon(FSlateIcon(FNiagaraEditorStyle::Get().GetStyleSetName(), "Tab.Debugger"));
 }
 
 void SNiagaraDebugger::UnregisterTabSpawner()
@@ -899,7 +899,7 @@ TSharedRef<SWidget> SNiagaraDebugger::MakeToolbar()
 				NAME_None,
 				LOCTEXT("Play", "Play"),
 				LOCTEXT("PlayTooltip", "Simulations will play as normal"),
-				FSlateIcon(FNiagaraEditorStyle::GetStyleSetName(), "NiagaraEditor.Debugger.PlayIcon"),
+				FSlateIcon(FNiagaraEditorStyle::Get().GetStyleSetName(), "NiagaraEditor.Debugger.PlayIcon"),
 				EUserInterfaceActionType::ToggleButton
 			);
 		}
@@ -914,7 +914,7 @@ TSharedRef<SWidget> SNiagaraDebugger::MakeToolbar()
 				NAME_None,
 				LOCTEXT("Pause", "Pause"),
 				LOCTEXT("PauseTooltip", "Pause all simulations"),
-				FSlateIcon(FNiagaraEditorStyle::GetStyleSetName(), "NiagaraEditor.Debugger.PauseIcon"),
+				FSlateIcon(FNiagaraEditorStyle::Get().GetStyleSetName(), "NiagaraEditor.Debugger.PauseIcon"),
 				EUserInterfaceActionType::ToggleButton
 			);
 		}
@@ -937,7 +937,7 @@ TSharedRef<SWidget> SNiagaraDebugger::MakeToolbar()
 				NAME_None,
 				CreateTAttribute<FText>([=]() { return Settings->Data.bLoopTimeEnabled ? FText::Format(LOCTEXT("PlaybackLoopFormat", "Loop Every\n{0} Seconds"), FText::AsNumber(Settings->Data.LoopTime)) : LOCTEXT("Loop", "Loop"); }),
 				LOCTEXT("LoopTooltip", "Loop all simulations, i.e. one shot effects will loop"),
-				FSlateIcon(FNiagaraEditorStyle::GetStyleSetName(), "NiagaraEditor.Debugger.LoopIcon"),
+				FSlateIcon(FNiagaraEditorStyle::Get().GetStyleSetName(), "NiagaraEditor.Debugger.LoopIcon"),
 				EUserInterfaceActionType::ToggleButton
 			);
 		}
@@ -952,7 +952,7 @@ TSharedRef<SWidget> SNiagaraDebugger::MakeToolbar()
 				NAME_None,
 				LOCTEXT("Step", "Step"),
 				LOCTEXT("StepTooltip", "Step all simulations a single frame then pause them"),
-				FSlateIcon(FNiagaraEditorStyle::GetStyleSetName(), "NiagaraEditor.Debugger.StepIcon"),
+				FSlateIcon(FNiagaraEditorStyle::Get().GetStyleSetName(), "NiagaraEditor.Debugger.StepIcon"),
 				EUserInterfaceActionType::ToggleButton
 			);
 		}
@@ -967,7 +967,7 @@ TSharedRef<SWidget> SNiagaraDebugger::MakeToolbar()
 				NAME_None,
 				CreateTAttribute<FText>([=]() { return  FText::Format(LOCTEXT("PlaybackSpeedFormat", "Speed\n{0} x"), FText::AsNumber(Settings->Data.PlaybackRate)); }),
 				LOCTEXT("SlowTooltip", "When enabled adjusts the playback speed for simulations."),
-				FSlateIcon(FNiagaraEditorStyle::GetStyleSetName(), "NiagaraEditor.Debugger.SpeedIcon"),
+				FSlateIcon(FNiagaraEditorStyle::Get().GetStyleSetName(), "NiagaraEditor.Debugger.SpeedIcon"),
 				EUserInterfaceActionType::ToggleButton
 			);
 		}
diff --git a/Engine/Plugins/FX/Niagara/Source/NiagaraEditor/Public/NiagaraEditorStyle.h b/Engine/Plugins/FX/Niagara/Source/NiagaraEditor/Public/NiagaraEditorStyle.h
--- a/Engine/Plugins/FX/Niagara/Source/NiagaraEditor/Public/NiagaraEditorStyle.h
+++ b/Engine/Plugins/FX/Niagara/Source/NiagaraEditor/Public/NiagaraEditorStyle.h
@@ -1,30 +1,54 @@
 // Copyright Epic Games, Inc. All Rights Reserved.
 
 #pragma once
-#include "Styling/ISlateStyle.h"
+
+#include "Styling/SlateStyle.h"
 
 /** Manages the style which provides resources for niagara editor widgets. */
-class NIAGARAEDITOR_API FNiagaraEditorStyle
+class NIAGARAEDITOR_API FNiagaraEditorStyle : public FSlateStyleSet
 {
 public:
-
-	static void Initialize();
-
+	static void Register();
+	static void Unregister();
 	static void Shutdown();
 
 	/** reloads textures used by slate renderer */
 	static void ReloadTextures();
 
 	/** @return The Slate style set for niagara editor widgets */
-	static const ISlateStyle& Get();
-
-	static FName GetStyleSetName();
-
-private:
-
-	static TSharedRef< class FSlateStyleSet > Create();
-
-private:
-
-	static TSharedPtr< class FSlateStyleSet > NiagaraEditorStyleInstance;
+	static const FNiagaraEditorStyle& Get();
+
+	static void ReinitializeStyle();
+
+private:	
+	FNiagaraEditorStyle();
+	void InitStats();
+	void InitAssetPicker();
+	void InitActionMenu();
+	void InitEmitterHeader();
+	void InitParameters();
+	void InitParameterMapView();
+	void InitCodeView();
+	void InitSelectedEmitter();
+	void InitToolbarIcons();
+	void InitTabIcons();
+	void InitIcons();
+	void InitOverview();
+	void InitEmitterDetails();
+	void InitAssetColors();
+	void InitThumbnails();
+	void InitClassIcon();
+	void InitStackIcons();
+	void InitNiagaraSequence();
+	void InitPlatformSet();
+	void InitDropTarget();
+	void InitScriptGraph();
+	void InitDebuggerStyle();
+	void InitBakerStyle();
+	void InitCommonColors();
+	void InitOutlinerStyle();
+	void InitScalabilityColors();
+	void InitScalabilityIcons();
+
+	static TSharedPtr<FNiagaraEditorStyle> NiagaraEditorStyle;
 };
diff --git a/Engine/Plugins/FX/Niagara/Source/NiagaraEditorWidgets/Private/DetailCustomizations/NiagaraDataInterfaceCurveDetails.cpp b/Engine/Plugins/FX/Niagara/Source/NiagaraEditorWidgets/Private/DetailCustomizations/NiagaraDataInterfaceCurveDetails.cpp
--- a/Engine/Plugins/FX/Niagara/Source/NiagaraEditorWidgets/Private/DetailCustomizations/NiagaraDataInterfaceCurveDetails.cpp
+++ b/Engine/Plugins/FX/Niagara/Source/NiagaraEditorWidgets/Private/DetailCustomizations/NiagaraDataInterfaceCurveDetails.cpp
@@ -1407,14 +1407,14 @@ void FNiagaraDataInterfaceCurveDetailsBase::CustomizeDetails(IDetailLayoutBuilde
 		FOnGetContent::CreateSP(this, &FNiagaraDataInterfaceCurveDetailsBase::GetCurveToCopyMenu),
 		NSLOCTEXT("NiagaraDataInterfaceCurveDetails", "Import", "Import"),
 		NSLOCTEXT("NiagaraDataInterfaceCurveDetails", "CopyCurveAsset", "Import data from another Curve asset"),
-		FSlateIcon(FNiagaraEditorWidgetsStyle::GetStyleSetName(), "NiagaraEditor.CurveDetails.Import"));
+		FSlateIcon(FNiagaraEditorWidgetsStyle::Get().GetStyleSetName(), "NiagaraEditor.CurveDetails.Import"));
 
 	ToolBarBuilder.AddToolBarButton(
 		FUIAction(FExecuteAction::CreateSP(this, &FNiagaraDataInterfaceCurveDetailsBase::OnShowInCurveEditor)),
 		NAME_None,
 		LOCTEXT("CurveOverviewLabel", "Curve Overview"),
 		LOCTEXT("ShowInCurveOverviewToolTip", "Show this curve in the curve overview."),
-		FSlateIcon(FNiagaraEditorWidgetsStyle::GetStyleSetName(), "NiagaraEditor.CurveDetails.ShowInOverview"));
+		FSlateIcon(FNiagaraEditorWidgetsStyle::Get().GetStyleSetName(), "NiagaraEditor.CurveDetails.ShowInOverview"));
 
 	ToolBarBuilder.AddSeparator();
 
diff --git a/Engine/Plugins/FX/Niagara/Source/NiagaraEditorWidgets/Private/NiagaraEditorWidgetsModule.cpp b/Engine/Plugins/FX/Niagara/Source/NiagaraEditorWidgets/Private/NiagaraEditorWidgetsModule.cpp
--- a/Engine/Plugins/FX/Niagara/Source/NiagaraEditorWidgets/Private/NiagaraEditorWidgetsModule.cpp
+++ b/Engine/Plugins/FX/Niagara/Source/NiagaraEditorWidgets/Private/NiagaraEditorWidgetsModule.cpp
@@ -113,7 +113,7 @@ void FNiagaraEditorWidgetsModule::StartupModule()
 	WidgetProvider = MakeShared<FNiagaraEditorWidgetProvider>();
 	NiagaraEditorModule.RegisterWidgetProvider(WidgetProvider.ToSharedRef());
 
-	FNiagaraEditorWidgetsStyle::Initialize();
+	FNiagaraEditorWidgetsStyle::Register();
 
 	FPropertyEditorModule& PropertyModule = FModuleManager::LoadModuleChecked<FPropertyEditorModule>("PropertyEditor");
 	PropertyModule.RegisterCustomClassLayout("NiagaraDataInterface", FOnGetDetailCustomizationInstance::CreateStatic(&FNiagaraDataInterfaceDetailsBase::MakeInstance));
@@ -174,8 +174,7 @@ FNiagaraEditorWidgetsModule& FNiagaraEditorWidgetsModule::Get()
 
 void FNiagaraEditorWidgetsModule::ReinitializeStyle()
 {
-	FNiagaraEditorWidgetsStyle::Shutdown();
-	FNiagaraEditorWidgetsStyle::Initialize();
+	FNiagaraEditorWidgetsStyle::ReinitializeStyle();
 }
 
 TSharedRef<FNiagaraStackCurveEditorOptions> FNiagaraEditorWidgetsModule::GetOrCreateStackCurveEditorOptionsForObject(UObject* Object, float DefaultHeight)
diff --git a/Engine/Plugins/FX/Niagara/Source/NiagaraEditorWidgets/Private/NiagaraEditorWidgetsStyle.cpp b/Engine/Plugins/FX/Niagara/Source/NiagaraEditorWidgets/Private/NiagaraEditorWidgetsStyle.cpp
--- a/Engine/Plugins/FX/Niagara/Source/NiagaraEditorWidgets/Private/NiagaraEditorWidgetsStyle.cpp
+++ b/Engine/Plugins/FX/Niagara/Source/NiagaraEditorWidgets/Private/NiagaraEditorWidgetsStyle.cpp
@@ -2,71 +2,76 @@
 
 #include "NiagaraEditorWidgetsStyle.h"
 
+#include "Styling/SlateStyleMacros.h"
 #include "Framework/Application/SlateApplication.h"
 #include "Slate/SlateGameResources.h"
 #include "Styling/SlateStyleRegistry.h"
 #include "Styling/SlateTypes.h"
 #include "Styling/CoreStyle.h"
-#include "Interfaces/IPluginManager.h"
 #include "Classes/EditorStyleSettings.h"
 #include "Styling/AppStyle.h"
 #include "Styling/StyleColors.h"
 
-TSharedPtr< FSlateStyleSet > FNiagaraEditorWidgetsStyle::NiagaraEditorWidgetsStyleInstance = NULL;
- 
-void FNiagaraEditorWidgetsStyle::Initialize()
+const FVector2D Icon6x6(6.0f, 6.0f);
+const FVector2D Icon8x8(8.0f, 8.0f);
+const FVector2D Icon8x16(8.0f, 16.0f);
+const FVector2D Icon12x12(12.0f, 12.0f);
+const FVector2D Icon16x16(16.0f, 16.0f);
+const FVector2D Icon20x20(20.0f, 20.0f);
+const FVector2D Icon30x30(30.0f, 30.0f);
+const FVector2D Icon40x40(40.0f, 40.0f);
+
+TSharedPtr<FNiagaraEditorWidgetsStyle> FNiagaraEditorWidgetsStyle::NiagaraEditorWidgetsStyle = nullptr;
+
+void FNiagaraEditorWidgetsStyle::Register()
 {
-	if (!NiagaraEditorWidgetsStyleInstance.IsValid())
-	{
-		NiagaraEditorWidgetsStyleInstance = Create();
-		FSlateStyleRegistry::RegisterSlateStyle(*NiagaraEditorWidgetsStyleInstance);
-	}
+	FSlateStyleRegistry::RegisterSlateStyle(Get());
 }
 
-void FNiagaraEditorWidgetsStyle::Shutdown()
+void FNiagaraEditorWidgetsStyle::Unregister()
 {
-	FSlateStyleRegistry::UnRegisterSlateStyle(*NiagaraEditorWidgetsStyleInstance);
-	ensure(NiagaraEditorWidgetsStyleInstance.IsUnique());
-	NiagaraEditorWidgetsStyleInstance.Reset();
+	FSlateStyleRegistry::UnRegisterSlateStyle(Get());
 }
 
-FName FNiagaraEditorWidgetsStyle::GetStyleSetName()
+void FNiagaraEditorWidgetsStyle::Shutdown()
 {
-	static FName StyleSetName(TEXT("NiagaraEditorWidgetsStyle"));
-	return StyleSetName;
+	Unregister();
+	NiagaraEditorWidgetsStyle.Reset();
 }
 
-NIAGARAEDITOR_API FString RelativePathToPluginPath(const FString& RelativePath, const ANSICHAR* Extension);
+void FNiagaraEditorWidgetsStyle::ReloadTextures()
+{
+	FSlateApplication::Get().GetRenderer()->ReloadTextureResources();
+}
 
-#define IMAGE_PLUGIN_BRUSH( RelativePath, ... ) FSlateImageBrush( RelativePathToPluginPath( RelativePath, ".png" ), __VA_ARGS__ )
-#define IMAGE_BRUSH( RelativePath, ... ) FSlateImageBrush( Style->RootToContentDir( RelativePath, TEXT(".png") ), __VA_ARGS__ )
-#define BOX_PLUGIN_BRUSH( RelativePath, ... ) FSlateBoxBrush( RelativePathToPluginPath( RelativePath, ".png"), __VA_ARGS__ )
-#define BOX_BRUSH( RelativePath, ... ) FSlateBoxBrush( Style->RootToContentDir( RelativePath, TEXT(".png") ), __VA_ARGS__ )
-#define BORDER_BRUSH( RelativePath, ... ) FSlateBorderBrush( Style->RootToContentDir( RelativePath, TEXT(".png") ), __VA_ARGS__ )
-#define DEFAULT_FONT(...) FCoreStyle::GetDefaultFontStyle(__VA_ARGS__)
-#define IMAGE_CORE_BRUSH( RelativePath, ... ) FSlateImageBrush( FPaths::EngineContentDir() / "Editor/Slate" / RelativePath + TEXT(".png") , __VA_ARGS__ )
-#define BOX_CORE_BRUSH( RelativePath, ... ) FSlateBoxBrush( FPaths::EngineContentDir() / "Editor/Slate" / RelativePath + TEXT(".png"), __VA_ARGS__ )
+const FNiagaraEditorWidgetsStyle& FNiagaraEditorWidgetsStyle::Get()
+{
+	if(!NiagaraEditorWidgetsStyle.IsValid())
+	{
+		NiagaraEditorWidgetsStyle = MakeShareable(new FNiagaraEditorWidgetsStyle());
+	}
+	
+	return *NiagaraEditorWidgetsStyle;
+}
 
-const FVector2D Icon6x6(6.0f, 6.0f);
-const FVector2D Icon8x8(8.0f, 8.0f);
-const FVector2D Icon8x16(8.0f, 16.0f);
-const FVector2D Icon12x12(12.0f, 12.0f);
-const FVector2D Icon16x16(16.0f, 16.0f);
-const FVector2D Icon20x20(20.0f, 20.0f);
-const FVector2D Icon30x30(30.0f, 30.0f);
-const FVector2D Icon40x40(40.0f, 40.0f);
+void FNiagaraEditorWidgetsStyle::ReinitializeStyle()
+{
+	Unregister();
+	NiagaraEditorWidgetsStyle.Reset();
+	Register();
+}
 
-TSharedRef< FSlateStyleSet > FNiagaraEditorWidgetsStyle::Create()
+FNiagaraEditorWidgetsStyle::FNiagaraEditorWidgetsStyle() : FSlateStyleSet("NiagaraEditorWidgetsStyle")
 {
+	SetContentRoot(FPaths::EnginePluginsDir() / TEXT("FX/Niagara/Content/Slate"));
+	SetCoreContentRoot(FPaths::EngineContentDir() / TEXT("Editor/Slate"));
+	
 	const FTextBlockStyle NormalText = FAppStyle::Get().GetWidgetStyle<FTextBlockStyle>("NormalText");
 	const FEditableTextBoxStyle NormalEditableTextBox = FAppStyle::Get().GetWidgetStyle<FEditableTextBoxStyle>("NormalEditableTextBox");
 	const FSpinBoxStyle NormalSpinBox = FAppStyle::Get().GetWidgetStyle<FSpinBoxStyle>("SpinBox");
 
-	TSharedRef< FSlateStyleSet > Style = MakeShareable(new FSlateStyleSet("NiagaraEditorWidgetsStyle"));
-	Style->SetContentRoot(FPaths::EngineContentDir() / TEXT("Editor/Slate/Niagara"));
-
 	// Stack
-	Style->Set("NiagaraEditor.Stack.IconSize", FVector2D(18.0f, 18.0f));
+	Set("NiagaraEditor.Stack.IconSize", FVector2D(18.0f, 18.0f));
 
 	const FTextBlockStyle CategoryText = FAppStyle::Get().GetWidgetStyle<FTextBlockStyle>("DetailsView.CategoryTextStyle");
 	FSlateFontInfo CategoryFont = FAppStyle::Get().GetFontStyle(TEXT("PropertyWindow.BoldFont"));
@@ -74,7 +79,7 @@ TSharedRef< FSlateStyleSet > FNiagaraEditorWidgetsStyle::Create()
 		.SetFont(CategoryFont)
 		.SetTransformPolicy(ETextTransformPolicy::ToUpper);
 
-	Style->Set("NiagaraEditor.Stack.GroupText", StackGroupText);
+	Set("NiagaraEditor.Stack.GroupText", StackGroupText);
 
 	FEditableTextBoxStyle StackEditableGroupText = FEditableTextBoxStyle(NormalEditableTextBox)
 		.SetFont(CategoryFont);
@@ -83,18 +88,18 @@ TSharedRef< FSlateStyleSet > FNiagaraEditorWidgetsStyle::Create()
 		.SetEditableTextBoxStyle(StackEditableGroupText)
 		.SetTextStyle(StackGroupText);
 
-	Style->Set("NiagaraEditor.Stack.EditableGroupText", StackInlineEditableGroupText);
+	Set("NiagaraEditor.Stack.EditableGroupText", StackInlineEditableGroupText);
 
 	FSlateFontInfo NormalFont = FAppStyle::Get().GetFontStyle(TEXT("PropertyWindow.NormalFont"));
 	FTextBlockStyle StackItemText = FTextBlockStyle(CategoryText)
 		.SetFont(NormalFont)
 		.SetTransformPolicy(ETextTransformPolicy::ToUpper);
 
-	Style->Set("NiagaraEditor.Stack.ItemText", StackItemText);
+	Set("NiagaraEditor.Stack.ItemText", StackItemText);
 	
 	FTextBlockStyle OverviewStackItemText = FTextBlockStyle(CategoryText)
 		.SetFont(NormalFont);
-	Style->Set("NiagaraEditor.OverviewStack.ItemText", OverviewStackItemText);
+	Set("NiagaraEditor.OverviewStack.ItemText", OverviewStackItemText);
 
 	FEditableTextBoxStyle StackEditableItemText = FEditableTextBoxStyle(NormalEditableTextBox)
 		.SetFont(NormalFont);
@@ -102,49 +107,49 @@ TSharedRef< FSlateStyleSet > FNiagaraEditorWidgetsStyle::Create()
 	FInlineEditableTextBlockStyle StackInlineEditableItemText = FInlineEditableTextBlockStyle()
 		.SetEditableTextBoxStyle(StackEditableItemText)
 		.SetTextStyle(StackItemText);
-	Style->Set("NiagaraEditor.Stack.EditableItemText", StackInlineEditableItemText);
+	Set("NiagaraEditor.Stack.EditableItemText", StackInlineEditableItemText);
 
 	FTextBlockStyle StackDefaultText = FTextBlockStyle(NormalText)
 		.SetFont(NormalFont);
-	Style->Set("NiagaraEditor.Stack.DefaultText", StackDefaultText);
+	Set("NiagaraEditor.Stack.DefaultText", StackDefaultText);
 
 	FTextBlockStyle StackCategoryText = FTextBlockStyle(NormalText);
-	Style->Set("NiagaraEditor.Stack.CategoryText", StackCategoryText);
-	Style->Set("NiagaraEditor.SystemOverview.GroupHeaderText", StackCategoryText);
+	Set("NiagaraEditor.Stack.CategoryText", StackCategoryText);
+	Set("NiagaraEditor.SystemOverview.GroupHeaderText", StackCategoryText);
 
 	FSlateFontInfo ParameterFont = DEFAULT_FONT("Regular", 8);
 	FTextBlockStyle ParameterText = FTextBlockStyle(NormalText)
 		.SetFont(ParameterFont);
-	Style->Set("NiagaraEditor.Stack.ParameterText", ParameterText);
+	Set("NiagaraEditor.Stack.ParameterText", ParameterText);
 
 	FSlateFontInfo TextContentFont = DEFAULT_FONT("Regular", 9);
 	FTextBlockStyle TextContentText = FTextBlockStyle(NormalText)
 		.SetFont(TextContentFont);
-	Style->Set("NiagaraEditor.Stack.TextContentText", TextContentText);
+	Set("NiagaraEditor.Stack.TextContentText", TextContentText);
 
 	FSlateFontInfo PerfWidgetDetailFont = DEFAULT_FONT("Regular", 7);
-	Style->Set("NiagaraEditor.Stack.Stats.DetailFont", PerfWidgetDetailFont);
+	Set("NiagaraEditor.Stack.Stats.DetailFont", PerfWidgetDetailFont);
 	FSlateFontInfo PerfWidgetGroupFont = DEFAULT_FONT("Regular", 8);
-	Style->Set("NiagaraEditor.Stack.Stats.GroupFont", PerfWidgetGroupFont);
+	Set("NiagaraEditor.Stack.Stats.GroupFont", PerfWidgetGroupFont);
 	FSlateFontInfo PerfWidgetEvalTypeFont = DEFAULT_FONT("Regular", 7);
-	Style->Set("NiagaraEditor.Stack.Stats.EvalTypeFont", PerfWidgetEvalTypeFont);
+	Set("NiagaraEditor.Stack.Stats.EvalTypeFont", PerfWidgetEvalTypeFont);
 	
 
 
 	FSlateFontInfo StackSubduedItemFont = DEFAULT_FONT("Regular", 9);
 	FTextBlockStyle StackSubduedItemText = FTextBlockStyle(NormalText)
 		.SetFont(StackSubduedItemFont);
-	Style->Set("NiagaraEditor.Stack.SubduedItemText", StackSubduedItemText);
+	Set("NiagaraEditor.Stack.SubduedItemText", StackSubduedItemText);
 
 	FSlateFontInfo SystemOverviewListHeaderFont = DEFAULT_FONT("Bold", 12);
 	FTextBlockStyle SystemOverviewListHeaderText = FTextBlockStyle(NormalText)
 		.SetFont(SystemOverviewListHeaderFont);
-	Style->Set("NiagaraEditor.SystemOverview.ListHeaderText", SystemOverviewListHeaderText);
+	Set("NiagaraEditor.SystemOverview.ListHeaderText", SystemOverviewListHeaderText);
 
 	FSlateFontInfo SystemOverviewItemFont = DEFAULT_FONT("Regular", 9);
 	FTextBlockStyle SystemOverviewItemText = FTextBlockStyle(NormalText)
 		.SetFont(SystemOverviewItemFont);
-	Style->Set("NiagaraEditor.SystemOverview.ItemText", SystemOverviewItemText);
+	Set("NiagaraEditor.SystemOverview.ItemText", SystemOverviewItemText);
 
 	FSlateFontInfo SystemOverviewAlternateItemFont = DEFAULT_FONT("Italic", 9);
 	FTextBlockStyle SystemOverviewAlternateItemText = FTextBlockStyle(NormalText)
@@ -203,80 +208,96 @@ TSharedRef< FSlateStyleSet > FNiagaraEditorWidgetsStyle::Create()
 		.SetNormalPadding(FMargin(1.f))
 		.SetPressedPadding(FMargin(1.f, 2.f, 1.f, 0.f));
 	
-	Style->Set("NiagaraEditor.SystemOverview.ItemText", SystemOverviewItemText);
+	Set("NiagaraEditor.SystemOverview.ItemText", SystemOverviewItemText);
 	
-	Style->Set("NiagaraEditor.SystemOverview.AlternateItemText", SystemOverviewAlternateItemText);
+	Set("NiagaraEditor.SystemOverview.AlternateItemText", SystemOverviewAlternateItemText);
 
-	Style->Set("NiagaraEditor.SystemOverview.InlineParameterText", SystemOverviewInlineParametersText);
-	Style->Set("NiagaraEditor.SystemOverview.InlineParameterButton", InlineParameterButtonStyle);
-	Style->Set("NiagaraEditor.SystemOverview.InlineParameterButton.Enum", InlineParameterButtonStyleEnum);
-	Style->Set("NiagaraEditor.SystemOverview.InlineParameterButton.NoTint", InlineParameterButtonStyleNoTint);
-	Style->Set("NiagaraEditor.SystemOverview.InlineParameterButton.Icon", InlineParameterButtonStyleIcon);
+	Set("NiagaraEditor.SystemOverview.InlineParameterText", SystemOverviewInlineParametersText);
+	Set("NiagaraEditor.SystemOverview.InlineParameterButton", InlineParameterButtonStyle);
+	Set("NiagaraEditor.SystemOverview.InlineParameterButton.Enum", InlineParameterButtonStyleEnum);
+	Set("NiagaraEditor.SystemOverview.InlineParameterButton.NoTint", InlineParameterButtonStyleNoTint);
+	Set("NiagaraEditor.SystemOverview.InlineParameterButton.Icon", InlineParameterButtonStyleIcon);
 
 
 
-	Style->Set("NiagaraEditor.SystemOverview.Item.BackgroundColor", FLinearColor(FColor(62, 62, 62)));
-	Style->Set("NiagaraEditor.SystemOverview.Group.BackgroundColor", FLinearColor::Transparent);
-	Style->Set("NiagaraEditor.SystemOverview.CheckBoxColor", FLinearColor(FColor(160, 160, 160)));
-	Style->Set("NiagaraEditor.SystemOverview.CheckBoxBorder", new BOX_CORE_BRUSH("Common/GroupBorderLight", FMargin(4.0f / 16.0f)));
-	Style->Set("NiagaraEditor.SystemOverview.NodeBackgroundBorder", new BOX_PLUGIN_BRUSH("Icons/SystemOverviewNodeBackground", FMargin(1.0f / 4.0f)));
-	Style->Set("NiagaraEditor.SystemOverview.NodeBackgroundColor", FLinearColor(FColor(48, 48, 48)));
+	Set("NiagaraEditor.SystemOverview.Item.BackgroundColor", FLinearColor(FColor(62, 62, 62)));
+	Set("NiagaraEditor.SystemOverview.Group.BackgroundColor", FLinearColor::Transparent);
+	Set("NiagaraEditor.SystemOverview.CheckBoxColor", FLinearColor(FColor(160, 160, 160)));
+	Set("NiagaraEditor.SystemOverview.CheckBoxBorder", new CORE_BOX_BRUSH("Common/GroupBorderLight", FMargin(4.0f / 16.0f)));
+	Set("NiagaraEditor.SystemOverview.NodeBackgroundBorder", new BOX_BRUSH("Icons/SystemOverviewNodeBackground", FMargin(1.0f / 4.0f)));
+	Set("NiagaraEditor.SystemOverview.NodeBackgroundColor", FLinearColor(FColor(48, 48, 48)));
+
+	Set("NiagaraEditor.Scalability.EffectType.bAllowCullingForLocalPlayers", new IMAGE_BRUSH_SVG("Icons/Scalability/CullDistance_20", Icon20x20));
+	Set("NiagaraEditor.Scalability.EffectType.UpdateFrequency", new IMAGE_BRUSH_SVG("Icons/Scalability/Update", Icon20x20));
+	Set("NiagaraEditor.Scalability.EffectType.CullReaction",  new IMAGE_BRUSH_SVG("Icons/Scalability/CullReaction_20", Icon20x20));
+	Set("NiagaraEditor.Scalability.EffectType.SignificanceHandler",  new IMAGE_BRUSH_SVG("Icons/Scalability/SignificanceHandler_20", Icon20x20));
+
+	FTextBlockStyle EffectTypeScalabilityPropertyTextStyle = FTextBlockStyle(NormalText)
+		.SetFont(DEFAULT_FONT("Regular", 8))
+		.SetColorAndOpacity(FLinearColor(1.f, 1.f, 1.f, 1.f));
+
+	FTextBlockStyle EffectTypeScalabilityValueTextStyle = FTextBlockStyle(NormalText)
+		.SetFont(DEFAULT_FONT("Regular", 9))
+		.SetColorAndOpacity(FLinearColor(1.f, 1.f, 1.f, 1.f));
+	
+	Set("NiagaraEditor.Scalability.EffectType.Property", EffectTypeScalabilityPropertyTextStyle);
+	Set("NiagaraEditor.Scalability.EffectType.Value", EffectTypeScalabilityPropertyTextStyle);
 
 	FTextBlockStyle NamePropertySelectionEntryTextBlockStyle = FTextBlockStyle(NormalText)
 		.SetFont(DEFAULT_FONT("Regular", 10))
 		.SetColorAndOpacity(FLinearColor(1.f, 1.f, 1.f, 1.f));
-	Style->Set("NiagaraEditor.NamePropertySelectionEntry", NamePropertySelectionEntryTextBlockStyle);
+	Set("NiagaraEditor.NamePropertySelectionEntry", NamePropertySelectionEntryTextBlockStyle);
 	
 	const FTableRowStyle& NormalTableRowStyle = FAppStyle::Get().GetWidgetStyle<FTableRowStyle>("TableView.Row");
 	FSlateFontInfo CurveOverviewTopLevelFont = DEFAULT_FONT("Bold", 10);
 	FTextBlockStyle CurveOverviewTopLevelText = FTextBlockStyle(NormalText)
 		.SetFont(CurveOverviewTopLevelFont);
-	Style->Set("NiagaraEditor.CurveOverview.TopLevelText", CurveOverviewTopLevelText);
+	Set("NiagaraEditor.CurveOverview.TopLevelText", CurveOverviewTopLevelText);
 
 	FSlateFontInfo CurveOverviewScriptFont = DEFAULT_FONT("Bold", 10);
 	FTextBlockStyle CurveOverviewScriptText = FTextBlockStyle(NormalText)
 		.SetFont(CurveOverviewScriptFont);
-	Style->Set("NiagaraEditor.CurveOverview.ScriptText", CurveOverviewScriptText);
+	Set("NiagaraEditor.CurveOverview.ScriptText", CurveOverviewScriptText);
 
 	FSlateFontInfo CurveOverviewModuleFont = DEFAULT_FONT("Bold", 10);
 	FTextBlockStyle CurveOverviewModuleText = FTextBlockStyle(NormalText)
 		.SetFont(CurveOverviewModuleFont);
-	Style->Set("NiagaraEditor.CurveOverview.ModuleText", CurveOverviewModuleText);
+	Set("NiagaraEditor.CurveOverview.ModuleText", CurveOverviewModuleText);
 
 	FSlateFontInfo CurveOverviewInputFont = DEFAULT_FONT("Bold", 9);
 	FTextBlockStyle CurveOverviewInputText = FTextBlockStyle(NormalText)
 		.SetFont(CurveOverviewInputFont);
-	Style->Set("NiagaraEditor.CurveOverview.InputText", CurveOverviewInputText);
+	Set("NiagaraEditor.CurveOverview.InputText", CurveOverviewInputText);
 
 	FSlateFontInfo CurveOverviewDataInterfaceFont = DEFAULT_FONT("Bold", 10);
 	FTextBlockStyle CurveOverviewDataInterfaceText = FTextBlockStyle(NormalText)
 		.SetFont(CurveOverviewDataInterfaceFont);
-	Style->Set("NiagaraEditor.CurveOverview.DataInterfaceText", CurveOverviewDataInterfaceText);
+	Set("NiagaraEditor.CurveOverview.DataInterfaceText", CurveOverviewDataInterfaceText);
 
 	FSlateFontInfo CurveOverviewSecondaryFont = DEFAULT_FONT("Italic", 9);
 	FTextBlockStyle CurveOverviewSecondaryText = FTextBlockStyle(NormalText)
 		.SetFont(CurveOverviewSecondaryFont);
-	Style->Set("NiagaraEditor.CurveOverview.SecondaryText", CurveOverviewSecondaryText);
+	Set("NiagaraEditor.CurveOverview.SecondaryText", CurveOverviewSecondaryText);
 
 	FSlateFontInfo CurveOverviewCurveComponentFont = DEFAULT_FONT("Regular", 9);
 	FTextBlockStyle CurveOverviewCurveComponentText = FTextBlockStyle(NormalText)
 		.SetFont(CurveOverviewCurveComponentFont);
-	Style->Set("NiagaraEditor.CurveOverview.CurveComponentText", CurveOverviewCurveComponentText);
+	Set("NiagaraEditor.CurveOverview.CurveComponentText", CurveOverviewCurveComponentText);
 
 	FSlateFontInfo CurveOverviewDefaultFont = DEFAULT_FONT("Bold", 10);
 	FTextBlockStyle CurveOverviewDefaultText = FTextBlockStyle(NormalText)
 		.SetFont(CurveOverviewDefaultFont);
-	Style->Set("NiagaraEditor.CurveOverview.DefaultText", CurveOverviewDefaultText);
+	Set("NiagaraEditor.CurveOverview.DefaultText", CurveOverviewDefaultText);
 
-	Style->Set("NiagaraEditor.CurveDetails.TextButtonForeground", FLinearColor::White);
+	Set("NiagaraEditor.CurveDetails.TextButtonForeground", FLinearColor::White);
 
-	Style->Set("NiagaraEditor.CurveDetails.Import.Small", new IMAGE_CORE_BRUSH("Icons/GeneralTools/Import_40x", Icon20x20));
+	Set("NiagaraEditor.CurveDetails.Import.Small", new CORE_IMAGE_BRUSH("Icons/GeneralTools/Import_40x", Icon20x20));
 
-	Style->Set("NiagaraEditor.CurveDetails.ShowInOverview.Small", new IMAGE_CORE_BRUSH("Common/GoToSource", Icon12x12, FLinearColor(.9f, .9f, .9f, 1.0f)));
+	Set("NiagaraEditor.CurveDetails.ShowInOverview.Small", new CORE_IMAGE_BRUSH("Common/GoToSource", Icon12x12, FLinearColor(.9f, .9f, .9f, 1.0f)));
 
-	FSlateBrush StackRowSelectionBrush = BOX_PLUGIN_BRUSH("Icons/StackSelectionBorder", FMargin(2.0f / 8.0f), FStyleColors::Select);
-	FSlateBrush StackRowSubduedSelectionBrush = BOX_PLUGIN_BRUSH("Icons/StackSelectionBorder", FMargin(2.0f / 8.0f), FStyleColors::SelectInactive);
-	Style->Set("NiagaraEditor.Stack.TableViewRow", FTableRowStyle(NormalTableRowStyle)
+	FSlateBoxBrush StackRowSelectionBrush = BOX_BRUSH("Icons/StackSelectionBorder", FMargin(2.0f / 8.0f), FStyleColors::Select);
+	FSlateBoxBrush StackRowSubduedSelectionBrush = BOX_BRUSH("Icons/StackSelectionBorder", FMargin(2.0f / 8.0f), FStyleColors::SelectInactive);
+	Set("NiagaraEditor.Stack.TableViewRow", FTableRowStyle(NormalTableRowStyle)
 		.SetEvenRowBackgroundBrush(FSlateColorBrush(FStyleColors::White))
 		.SetOddRowBackgroundBrush(FSlateColorBrush(FStyleColors::White))
 		.SetActiveBrush(StackRowSelectionBrush)
@@ -284,139 +305,121 @@ TSharedRef< FSlateStyleSet > FNiagaraEditorWidgetsStyle::Create()
 		.SetInactiveBrush(StackRowSubduedSelectionBrush)
 		.SetInactiveHoveredBrush(StackRowSelectionBrush));
 
-	Style->Set("NiagaraEditor.SystemOverview.TableViewRow", FTableRowStyle(NormalTableRowStyle)
+	Set("NiagaraEditor.SystemOverview.TableViewRow", FTableRowStyle(NormalTableRowStyle)
 		.SetEvenRowBackgroundBrush(FSlateColorBrush(FStyleColors::Panel))
 		.SetOddRowBackgroundBrush(FSlateColorBrush(FStyleColors::Panel))
 		.SetSelectedTextColor(FStyleColors::ForegroundInverted)
-		.SetInactiveBrush(IMAGE_CORE_BRUSH("Common/Selection", Icon8x8, GetDefault<UEditorStyleSettings>()->GetSubduedSelectionColor())));
-
-	Style->Set("NiagaraEditor.Stack.BackgroundColor", FLinearColor(FColor(96, 96, 96)));
-	Style->Set("NiagaraEditor.Stack.Item.HeaderBackgroundColor", FLinearColor(FColor(48, 48, 48)));
-	Style->Set("NiagaraEditor.Stack.Item.ContentBackgroundColor", FLinearColor(FColor(62, 62, 62)));
-	Style->Set("NiagaraEditor.Stack.Item.ContentAdvancedBackgroundColor", FLinearColor(FColor(53, 53, 53)));
-	Style->Set("NiagaraEditor.Stack.Item.FooterBackgroundColor", FLinearColor(FColor(75, 75, 75)));
-	Style->Set("NiagaraEditor.Stack.Item.CustomNoteBackgroundColor", FLinearColor(FColor(56, 111, 75)));
-	Style->Set("NiagaraEditor.Stack.Item.InfoBackgroundColor", FLinearColor(FColor(68, 100, 106)));
-	Style->Set("NiagaraEditor.Stack.Item.WarningBackgroundColor", FLinearColor(FColor(97, 97, 68)));
-	Style->Set("NiagaraEditor.Stack.Item.ErrorBackgroundColor", FLinearColor(FColor(126, 78, 68)));
-
-	Style->Set("NiagaraEditor.Stack.UnknownColor", FLinearColor(1, 0, 1));
-
-	Style->Set("NiagaraEditor.Stack.ItemHeaderFooter.BackgroundBrush", new FSlateColorBrush(FLinearColor(FColor(20, 20, 20))));
-
-	Style->Set("NiagaraEditor.Stack.ForegroundColor", FLinearColor(FColor(220, 220, 220)));
-	Style->Set("NiagaraEditor.Stack.GroupForegroundColor", FLinearColor(FColor(220, 220, 220)));
-	Style->Set("NiagaraEditor.Stack.FlatButtonColor", FLinearColor(FColor(205, 205, 205)));
-	Style->Set("NiagaraEditor.Stack.DividerColor", FLinearColor(FColor(92, 92, 92)));
+		.SetInactiveBrush(CORE_IMAGE_BRUSH("Common/Selection", Icon8x8, GetDefault<UEditorStyleSettings>()->GetSubduedSelectionColor())));
+
+	Set("NiagaraEditor.Stack.BackgroundColor", FLinearColor(FColor(96, 96, 96)));
+	Set("NiagaraEditor.Stack.Item.HeaderBackgroundColor", FLinearColor(FColor(48, 48, 48)));
+	Set("NiagaraEditor.Stack.Item.ContentBackgroundColor", FLinearColor(FColor(62, 62, 62)));
+	Set("NiagaraEditor.Stack.Item.ContentAdvancedBackgroundColor", FLinearColor(FColor(53, 53, 53)));
+	Set("NiagaraEditor.Stack.Item.FooterBackgroundColor", FLinearColor(FColor(75, 75, 75)));
+	Set("NiagaraEditor.Stack.Item.CustomNoteBackgroundColor", FLinearColor(FColor(56, 111, 75)));
+	Set("NiagaraEditor.Stack.Item.InfoBackgroundColor", FLinearColor(FColor(68, 100, 106)));
+	Set("NiagaraEditor.Stack.Item.WarningBackgroundColor", FLinearColor(FColor(97, 97, 68)));
+	Set("NiagaraEditor.Stack.Item.ErrorBackgroundColor", FLinearColor(FColor(126, 78, 68)));
+
+	Set("NiagaraEditor.Stack.UnknownColor", FLinearColor(1, 0, 1));
+
+	Set("NiagaraEditor.Stack.ItemHeaderFooter.BackgroundBrush", new FSlateColorBrush(FLinearColor(FColor(20, 20, 20))));
+
+	Set("NiagaraEditor.Stack.ForegroundColor", FLinearColor(FColor(220, 220, 220)));
+	Set("NiagaraEditor.Stack.GroupForegroundColor", FLinearColor(FColor(220, 220, 220)));
+	Set("NiagaraEditor.Stack.FlatButtonColor", FLinearColor(FColor(205, 205, 205)));
+	Set("NiagaraEditor.Stack.DividerColor", FLinearColor(FColor(92, 92, 92)));
 	
-	Style->Set("NiagaraEditor.Stack.Stats.EvalTypeColor", FLinearColor(FColor(168, 168, 168)));
-	Style->Set("NiagaraEditor.Stack.Stats.RuntimePlaceholderColor", FLinearColor(FColor(86, 86, 86)));
-	Style->Set("NiagaraEditor.Stack.Stats.RuntimeUsageColorDefault", FLinearColor(FColor(200, 60, 60)));
-	Style->Set("NiagaraEditor.Stack.Stats.RuntimeUsageColorParticleUpdate", FLinearColor(FColor(246, 3, 142)));
-	Style->Set("NiagaraEditor.Stack.Stats.RuntimeUsageColorParticleSpawn", FLinearColor(FColor(255, 181, 0)));
-	Style->Set("NiagaraEditor.Stack.Stats.RuntimeUsageColorSystem", FLinearColor(FColor(20, 161, 255)));
-	Style->Set("NiagaraEditor.Stack.Stats.RuntimeUsageColorEmitter", FLinearColor(FColor(241, 99, 6)));
-	Style->Set("NiagaraEditor.Stack.Stats.LowCostColor", FLinearColor(FColor(143, 185, 130)));
-	Style->Set("NiagaraEditor.Stack.Stats.MediumCostColor", FLinearColor(FColor(220, 210, 86)));
-	Style->Set("NiagaraEditor.Stack.Stats.HighCostColor", FLinearColor(FColor(205, 114, 69)));
-	Style->Set("NiagaraEditor.Stack.Stats.MaxCostColor", FLinearColor(FColor(200, 60, 60)));
-
-	Style->Set("NiagaraEditor.Stack.AccentColor.System", FLinearColor(FColor(67, 105, 124)));
-	Style->Set("NiagaraEditor.Stack.AccentColor.Emitter", FLinearColor(FColor(126, 87, 67)));
-	Style->Set("NiagaraEditor.Stack.AccentColor.Particle", FLinearColor(FColor(87, 107, 61)));
-	Style->Set("NiagaraEditor.Stack.AccentColor.Render", FLinearColor(FColor(134, 80, 80)));
-	Style->Set("NiagaraEditor.Stack.AccentColor.None", FLinearColor::Transparent);
-
-	Style->Set("NiagaraEditor.Stack.IconColor.System", FLinearColor(FColor(1, 202, 252)));
-	Style->Set("NiagaraEditor.Stack.IconColor.Emitter", FLinearColor(FColor(241, 99, 6)));
-	Style->Set("NiagaraEditor.Stack.IconColor.Particle", FLinearColor(FColor(131, 218, 9)));
-	Style->Set("NiagaraEditor.Stack.IconColor.Render", FLinearColor(FColor(230, 102, 102)));
-	Style->Set("NiagaraEditor.Stack.IconColor.VersionUpgrade", FLinearColor(FColor(255, 181, 0)));
-
-	Style->Set("NiagaraEditor.Stack.InputValueIconColor.Linked", FLinearColor(FColor::Purple));
-	Style->Set("NiagaraEditor.Stack.InputValueIconColor.Data", FLinearColor(FColor::Yellow));
-	Style->Set("NiagaraEditor.Stack.InputValueIconColor.Dynamic", FLinearColor(FColor::Cyan));
-	Style->Set("NiagaraEditor.Stack.InputValueIconColor.Expression", FLinearColor(FColor::Green));
-	Style->Set("NiagaraEditor.Stack.InputValueIconColor.Default", FLinearColor(FColor::White));
-
- 	Style->Set("NiagaraEditor.Stack.DropTarget.BackgroundColor", FLinearColor(1.0f, 1.0f, 1.0f, 0.25f));
- 	Style->Set("NiagaraEditor.Stack.DropTarget.BackgroundColorHover", FLinearColor(1.0f, 1.0f, 1.0f, 0.1f));
-	Style->Set("NiagaraEditor.Stack.DropTarget.BorderVertical", new IMAGE_PLUGIN_BRUSH("Icons/StackDropTargetBorder_Vertical", FVector2D(2, 8), FLinearColor::White, ESlateBrushTileType::Vertical));
-	Style->Set("NiagaraEditor.Stack.DropTarget.BorderHorizontal", new IMAGE_PLUGIN_BRUSH("Icons/StackDropTargetBorder_Horizontal", FVector2D(8, 2), FLinearColor::White, ESlateBrushTileType::Horizontal));
-
-	Style->Set("NiagaraEditor.Stack.GoToSourceIcon", new IMAGE_CORE_BRUSH("Common/GoToSource", Icon30x30, FLinearColor::White));
-	Style->Set("NiagaraEditor.Stack.ParametersIcon", new IMAGE_PLUGIN_BRUSH("Icons/SystemParams", Icon12x12, FLinearColor::White));
-	Style->Set("NiagaraEditor.Stack.SpawnIcon", new IMAGE_PLUGIN_BRUSH("Icons/Spawn", Icon12x12, FLinearColor::White));
-	Style->Set("NiagaraEditor.Stack.UpdateIcon", new IMAGE_PLUGIN_BRUSH("Icons/Update", Icon12x12, FLinearColor::White));
-	Style->Set("NiagaraEditor.Stack.EventIcon", new IMAGE_PLUGIN_BRUSH("Icons/Event", Icon12x12, FLinearColor::White));
-	Style->Set("NiagaraEditor.Stack.SimulationStageIcon", new IMAGE_PLUGIN_BRUSH("Icons/SimulationStage", Icon12x12, FLinearColor::White));
-	Style->Set("NiagaraEditor.Stack.RenderIcon", new IMAGE_PLUGIN_BRUSH("Icons/Render", Icon12x12, FLinearColor::White));
-
-	Style->Set("NiagaraEditor.Stack.ParametersIconHighlighted", new IMAGE_PLUGIN_BRUSH("Icons/SystemParams", Icon16x16, FLinearColor::White));
-	Style->Set("NiagaraEditor.Stack.SpawnIconHighlighted", new IMAGE_PLUGIN_BRUSH("Icons/Spawn", Icon16x16, FLinearColor::White));
-	Style->Set("NiagaraEditor.Stack.UpdateIconHighlighted", new IMAGE_PLUGIN_BRUSH("Icons/Update", Icon16x16, FLinearColor::White));
-	Style->Set("NiagaraEditor.Stack.EventIconHighlighted", new IMAGE_PLUGIN_BRUSH("Icons/Event", Icon16x16, FLinearColor::White));
-	Style->Set("NiagaraEditor.Stack.SimulationStageIconHighlighted", new IMAGE_PLUGIN_BRUSH("Icons/SimulationStage", Icon16x16, FLinearColor::White));
-	Style->Set("NiagaraEditor.Stack.RenderIconHighlighted", new IMAGE_PLUGIN_BRUSH("Icons/Render", Icon16x16, FLinearColor::White));
-
-	Style->Set("NiagaraEditor.Stack.IconHighlightedSize", 16.0f);
-
-	Style->Set("NiagaraEditor.Stack.Splitter", FSplitterStyle()
-		.SetHandleNormalBrush(IMAGE_CORE_BRUSH("Common/SplitterHandleHighlight", Icon8x8, FLinearColor(.1f, .1f, .1f, 1.0f)))
-		.SetHandleHighlightBrush(IMAGE_CORE_BRUSH("Common/SplitterHandleHighlight", Icon8x8, FLinearColor::White))
+	Set("NiagaraEditor.Stack.Stats.EvalTypeColor", FLinearColor(FColor(168, 168, 168)));
+	Set("NiagaraEditor.Stack.Stats.RuntimePlaceholderColor", FLinearColor(FColor(86, 86, 86)));
+	Set("NiagaraEditor.Stack.Stats.RuntimeUsageColorDefault", FLinearColor(FColor(200, 60, 60)));
+	Set("NiagaraEditor.Stack.Stats.RuntimeUsageColorParticleUpdate", FLinearColor(FColor(246, 3, 142)));
+	Set("NiagaraEditor.Stack.Stats.RuntimeUsageColorParticleSpawn", FLinearColor(FColor(255, 181, 0)));
+	Set("NiagaraEditor.Stack.Stats.RuntimeUsageColorSystem", FLinearColor(FColor(20, 161, 255)));
+	Set("NiagaraEditor.Stack.Stats.RuntimeUsageColorEmitter", FLinearColor(FColor(241, 99, 6)));
+	Set("NiagaraEditor.Stack.Stats.LowCostColor", FLinearColor(FColor(143, 185, 130)));
+	Set("NiagaraEditor.Stack.Stats.MediumCostColor", FLinearColor(FColor(220, 210, 86)));
+	Set("NiagaraEditor.Stack.Stats.HighCostColor", FLinearColor(FColor(205, 114, 69)));
+	Set("NiagaraEditor.Stack.Stats.MaxCostColor", FLinearColor(FColor(200, 60, 60)));
+
+	Set("NiagaraEditor.Stack.AccentColor.System", FLinearColor(FColor(67, 105, 124)));
+	Set("NiagaraEditor.Stack.AccentColor.Emitter", FLinearColor(FColor(126, 87, 67)));
+	Set("NiagaraEditor.Stack.AccentColor.Particle", FLinearColor(FColor(87, 107, 61)));
+	Set("NiagaraEditor.Stack.AccentColor.Render", FLinearColor(FColor(134, 80, 80)));
+	Set("NiagaraEditor.Stack.AccentColor.None", FLinearColor::Transparent);
+
+	Set("NiagaraEditor.Stack.IconColor.System", FLinearColor(FColor(1, 202, 252)));
+	Set("NiagaraEditor.Stack.IconColor.Emitter", FLinearColor(FColor(241, 99, 6)));
+	Set("NiagaraEditor.Stack.IconColor.Particle", FLinearColor(FColor(131, 218, 9)));
+	Set("NiagaraEditor.Stack.IconColor.Render", FLinearColor(FColor(230, 102, 102)));
+	Set("NiagaraEditor.Stack.IconColor.VersionUpgrade", FLinearColor(FColor(255, 181, 0)));
+
+	Set("NiagaraEditor.Stack.InputValueIconColor.Linked", FLinearColor(FColor::Purple));
+	Set("NiagaraEditor.Stack.InputValueIconColor.Data", FLinearColor(FColor::Yellow));
+	Set("NiagaraEditor.Stack.InputValueIconColor.Dynamic", FLinearColor(FColor::Cyan));
+	Set("NiagaraEditor.Stack.InputValueIconColor.Expression", FLinearColor(FColor::Green));
+	Set("NiagaraEditor.Stack.InputValueIconColor.Default", FLinearColor(FColor::White));
+
+ 	Set("NiagaraEditor.Stack.DropTarget.BackgroundColor", FLinearColor(1.0f, 1.0f, 1.0f, 0.25f));
+ 	Set("NiagaraEditor.Stack.DropTarget.BackgroundColorHover", FLinearColor(1.0f, 1.0f, 1.0f, 0.1f));
+	Set("NiagaraEditor.Stack.DropTarget.BorderVertical", new IMAGE_BRUSH("Icons/StackDropTargetBorder_Vertical", FVector2D(2, 8), FLinearColor::White, ESlateBrushTileType::Vertical));
+	Set("NiagaraEditor.Stack.DropTarget.BorderHorizontal", new IMAGE_BRUSH("Icons/StackDropTargetBorder_Horizontal", FVector2D(8, 2), FLinearColor::White, ESlateBrushTileType::Horizontal));
+
+	Set("NiagaraEditor.Stack.GoToSourceIcon", new CORE_IMAGE_BRUSH("Common/GoToSource", Icon30x30, FLinearColor::White));
+	Set("NiagaraEditor.Stack.ParametersIcon", new IMAGE_BRUSH("Icons/SystemParams", Icon12x12, FLinearColor::White));
+	Set("NiagaraEditor.Stack.SpawnIcon", new IMAGE_BRUSH("Icons/Spawn", Icon12x12, FLinearColor::White));
+	Set("NiagaraEditor.Stack.UpdateIcon", new IMAGE_BRUSH("Icons/Update", Icon12x12, FLinearColor::White));
+	Set("NiagaraEditor.Stack.EventIcon", new IMAGE_BRUSH("Icons/Event", Icon12x12, FLinearColor::White));
+	Set("NiagaraEditor.Stack.SimulationStageIcon", new IMAGE_BRUSH("Icons/SimulationStage", Icon12x12, FLinearColor::White));
+	Set("NiagaraEditor.Stack.RenderIcon", new IMAGE_BRUSH("Icons/Render", Icon12x12, FLinearColor::White));
+
+	Set("NiagaraEditor.Stack.ParametersIconHighlighted", new IMAGE_BRUSH("Icons/SystemParams", Icon16x16, FLinearColor::White));
+	Set("NiagaraEditor.Stack.SpawnIconHighlighted", new IMAGE_BRUSH("Icons/Spawn", Icon16x16, FLinearColor::White));
+	Set("NiagaraEditor.Stack.UpdateIconHighlighted", new IMAGE_BRUSH("Icons/Update", Icon16x16, FLinearColor::White));
+	Set("NiagaraEditor.Stack.EventIconHighlighted", new IMAGE_BRUSH("Icons/Event", Icon16x16, FLinearColor::White));
+	Set("NiagaraEditor.Stack.SimulationStageIconHighlighted", new IMAGE_BRUSH("Icons/SimulationStage", Icon16x16, FLinearColor::White));
+	Set("NiagaraEditor.Stack.RenderIconHighlighted", new IMAGE_BRUSH("Icons/Render", Icon16x16, FLinearColor::White));
+
+	Set("NiagaraEditor.Stack.IconHighlightedSize", 16.0f);
+
+	Set("NiagaraEditor.Stack.Splitter", FSplitterStyle()
+		.SetHandleNormalBrush(CORE_IMAGE_BRUSH("Common/SplitterHandleHighlight", Icon8x8, FLinearColor(.1f, .1f, .1f, 1.0f)))
+		.SetHandleHighlightBrush(CORE_IMAGE_BRUSH("Common/SplitterHandleHighlight", Icon8x8, FLinearColor::White))
 	);
 
-	Style->Set("NiagaraEditor.Stack.SearchHighlightColor", FStyleColors::Highlight);
-	Style->Set("NiagaraEditor.Stack.SearchResult", new BOX_PLUGIN_BRUSH("Icons/SearchResultBorder", FMargin(1.f/8.f)));
+	Set("NiagaraEditor.Stack.SearchHighlightColor", FStyleColors::Highlight);
+	Set("NiagaraEditor.Stack.SearchResult", new BOX_BRUSH("Icons/SearchResultBorder", FMargin(1.f/8.f)));
 
-	Style->Set("NiagaraEditor.Stack.ModuleHighlight", new IMAGE_PLUGIN_BRUSH("Icons/ModuleHighlight", Icon6x6, FLinearColor::White));
-	Style->Set("NiagaraEditor.Stack.ModuleHighlightMore", new IMAGE_PLUGIN_BRUSH("Icons/ModuleHighlightMore", Icon6x6, FLinearColor::White));
-	Style->Set("NiagaraEditor.Stack.ModuleHighlightLarge", new IMAGE_PLUGIN_BRUSH("Icons/ModuleHighlightLarge", Icon8x8, FLinearColor::White));	
+	Set("NiagaraEditor.Stack.ModuleHighlight", new IMAGE_BRUSH("Icons/ModuleHighlight", Icon6x6, FLinearColor::White));
+	Set("NiagaraEditor.Stack.ModuleHighlightMore", new IMAGE_BRUSH("Icons/ModuleHighlightMore", Icon6x6, FLinearColor::White));
+	Set("NiagaraEditor.Stack.ModuleHighlightLarge", new IMAGE_BRUSH("Icons/ModuleHighlightLarge", Icon8x8, FLinearColor::White));	
 
 	// Scratch pad
 	FSlateFontInfo ScratchPadEditorHeaderFont = DEFAULT_FONT("Bold", 11);
 	FTextBlockStyle ScratchPadEditorHeaderText = FTextBlockStyle(NormalText)
 		.SetFont(ScratchPadEditorHeaderFont);
-	Style->Set("NiagaraEditor.ScratchPad.EditorHeaderText", ScratchPadEditorHeaderText);
+	Set("NiagaraEditor.ScratchPad.EditorHeaderText", ScratchPadEditorHeaderText);
 
 	FSlateFontInfo ScratchPadSubSectionHeaderFont = DEFAULT_FONT("Bold", 9);
 	FTextBlockStyle ScratchPadSubSectionHeaderText = FTextBlockStyle(NormalText)
 		.SetFont(ScratchPadSubSectionHeaderFont);
-	Style->Set("NiagaraEditor.ScratchPad.SubSectionHeaderText", ScratchPadSubSectionHeaderText);
+	Set("NiagaraEditor.ScratchPad.SubSectionHeaderText", ScratchPadSubSectionHeaderText);
 
-	FSlateBrush ScratchPadCategoryBrush = BOX_PLUGIN_BRUSH("Icons/CategoryRow", FMargin(2.0f / 8.0f), FLinearColor(FColor(48, 48, 48)));
-	FSlateBrush ScratchPadHoveredCategoryBrush = BOX_PLUGIN_BRUSH("Icons/CategoryRow", FMargin(2.0f / 8.0f), FLinearColor(FColor(38, 38, 38)));
-	Style->Set("NiagaraEditor.ScratchPad.CategoryRow", FTableRowStyle(NormalTableRowStyle)
+	FSlateBoxBrush ScratchPadCategoryBrush = BOX_BRUSH("Icons/CategoryRow", FMargin(2.0f / 8.0f), FLinearColor(FColor(48, 48, 48)));
+	FSlateBoxBrush ScratchPadHoveredCategoryBrush = BOX_BRUSH("Icons/CategoryRow", FMargin(2.0f / 8.0f), FLinearColor(FColor(38, 38, 38)));
+	Set("NiagaraEditor.ScratchPad.CategoryRow", FTableRowStyle(NormalTableRowStyle)
 		.SetEvenRowBackgroundBrush(ScratchPadCategoryBrush)
 		.SetOddRowBackgroundBrush(ScratchPadCategoryBrush)
 		.SetEvenRowBackgroundHoveredBrush(ScratchPadHoveredCategoryBrush)
 		.SetOddRowBackgroundHoveredBrush(ScratchPadHoveredCategoryBrush));
 
-	Style->Set("NiagaraEditor.Scope.Engine", FLinearColor(FColor(230, 102, 102)));
-	Style->Set("NiagaraEditor.Scope.Owner", FLinearColor(FColor(210, 112, 112)));
-	Style->Set("NiagaraEditor.Scope.User", FLinearColor(FColor(114, 226, 254)));
-	Style->Set("NiagaraEditor.Scope.System", FLinearColor(FColor(1, 202, 252)));
-	Style->Set("NiagaraEditor.Scope.Emitter", FLinearColor(FColor(241, 99, 6)));
-	Style->Set("NiagaraEditor.Scope.Particles", FLinearColor(FColor(131, 218, 9)));
-	Style->Set("NiagaraEditor.Scope.ScriptPersistent", FLinearColor(FColor(255, 247, 77)));
-	Style->Set("NiagaraEditor.Scope.ScriptTransient", FLinearColor(FColor(255, 247, 77)));
-
-	return Style;
-}
-
-#undef IMAGE_BRUSH
-#undef BOX_BRUSH
-#undef BORDER_BRUSH
-#undef DEFAULT_FONT
-#undef BOX_CORE_BRUSH
-#undef IMAGE_CORE_BRUSH
-
-void FNiagaraEditorWidgetsStyle::ReloadTextures()
-{
-	FSlateApplication::Get().GetRenderer()->ReloadTextureResources();
+	Set("NiagaraEditor.Scope.Engine", FLinearColor(FColor(230, 102, 102)));
+	Set("NiagaraEditor.Scope.Owner", FLinearColor(FColor(210, 112, 112)));
+	Set("NiagaraEditor.Scope.User", FLinearColor(FColor(114, 226, 254)));
+	Set("NiagaraEditor.Scope.System", FLinearColor(FColor(1, 202, 252)));
+	Set("NiagaraEditor.Scope.Emitter", FLinearColor(FColor(241, 99, 6)));
+	Set("NiagaraEditor.Scope.Particles", FLinearColor(FColor(131, 218, 9)));
+	Set("NiagaraEditor.Scope.ScriptPersistent", FLinearColor(FColor(255, 247, 77)));
+	Set("NiagaraEditor.Scope.ScriptTransient", FLinearColor(FColor(255, 247, 77)));
 }
 
-const ISlateStyle& FNiagaraEditorWidgetsStyle::Get()
-{
-	return *NiagaraEditorWidgetsStyleInstance;
-}
diff --git a/Engine/Plugins/FX/Niagara/Source/NiagaraEditorWidgets/Private/NiagaraEditorWidgetsStyle.h b/Engine/Plugins/FX/Niagara/Source/NiagaraEditorWidgets/Private/NiagaraEditorWidgetsStyle.h
--- a/Engine/Plugins/FX/Niagara/Source/NiagaraEditorWidgets/Private/NiagaraEditorWidgetsStyle.h
+++ b/Engine/Plugins/FX/Niagara/Source/NiagaraEditorWidgets/Private/NiagaraEditorWidgetsStyle.h
@@ -1,30 +1,27 @@
 // Copyright Epic Games, Inc. All Rights Reserved.
 
 #pragma once
-#include "Styling/ISlateStyle.h"
+
+#include "Styling/SlateStyle.h"
 
 /** Manages the style which provides resources for niagara editor widgets. */
-class FNiagaraEditorWidgetsStyle
+class FNiagaraEditorWidgetsStyle : public FSlateStyleSet
 {
 public:
 
-	static void Initialize();
-
+	static void Register();
+	static void Unregister();
 	static void Shutdown();
 
 	/** reloads textures used by slate renderer */
 	static void ReloadTextures();
 
 	/** @return The Slate style set for niagara editor widgets */
-	static const ISlateStyle& Get();
-
-	static FName GetStyleSetName();
-
-private:
-
-	static TSharedRef< class FSlateStyleSet > Create();
+	static const FNiagaraEditorWidgetsStyle& Get();
 
+	static void ReinitializeStyle();
 private:
+	FNiagaraEditorWidgetsStyle();
 
-	static TSharedPtr< class FSlateStyleSet > NiagaraEditorWidgetsStyleInstance;
+	static TSharedPtr<FNiagaraEditorWidgetsStyle> NiagaraEditorWidgetsStyle;
 };
diff --git a/Engine/Plugins/FX/Niagara/Source/NiagaraEditorWidgets/Private/SNiagaraOverviewStackNode.cpp b/Engine/Plugins/FX/Niagara/Source/NiagaraEditorWidgets/Private/SNiagaraOverviewStackNode.cpp
--- a/Engine/Plugins/FX/Niagara/Source/NiagaraEditorWidgets/Private/SNiagaraOverviewStackNode.cpp
+++ b/Engine/Plugins/FX/Niagara/Source/NiagaraEditorWidgets/Private/SNiagaraOverviewStackNode.cpp
@@ -466,7 +466,6 @@ void SNiagaraOverviewStackNode::FillTopContentBar()
 		FNiagaraEmitterInstance* InInstance = EmitterHandleViewModelWeak.Pin()->GetEmitterViewModel()->GetSimulation().IsValid() ? EmitterHandleViewModelWeak.Pin()->GetEmitterViewModel()->GetSimulation().Pin().Get() : nullptr;
 
 		FToolBarBuilder ToolBarBuilder(nullptr, FMultiBoxCustomization::None, nullptr, true);
-		ToolBarBuilder.SetStyle(&FNiagaraEditorStyle::Get(), FName("OverviewStackNodeThumbnailToolBar"));
 		ToolBarBuilder.SetLabelVisibility(EVisibility::Collapsed);
 		
 		for (int32 StackEntryIndex = 0; StackEntryIndex < PreviewStackEntries.Num(); StackEntryIndex++)
diff --git a/Engine/Plugins/FX/Niagara/Source/NiagaraEditorWidgets/Private/SNiagaraScratchPad.cpp b/Engine/Plugins/FX/Niagara/Source/NiagaraEditorWidgets/Private/SNiagaraScratchPad.cpp
--- a/Engine/Plugins/FX/Niagara/Source/NiagaraEditorWidgets/Private/SNiagaraScratchPad.cpp
+++ b/Engine/Plugins/FX/Niagara/Source/NiagaraEditorWidgets/Private/SNiagaraScratchPad.cpp
@@ -527,7 +527,7 @@ class SNiagaraScratchPadScriptEditor : public SCompoundWidget
 						.Padding(2.0f, 2.0f, 2.0f, 2.0f)
 						[
 							SNew(SImage)
-							.Image(FNiagaraEditorStyle::Get().GetBrush("NiagaraEditor.Apply.Small"))
+							.Image(FAppStyle::Get().GetBrush("AssetEditor.Apply"))
 						]
 						+ SHorizontalBox::Slot()
 						.AutoWidth()
diff --git a/Engine/Plugins/FX/Niagara/Source/NiagaraEditorWidgets/Private/Stack/SNiagaraStackModuleItem.cpp b/Engine/Plugins/FX/Niagara/Source/NiagaraEditorWidgets/Private/Stack/SNiagaraStackModuleItem.cpp
--- a/Engine/Plugins/FX/Niagara/Source/NiagaraEditorWidgets/Private/Stack/SNiagaraStackModuleItem.cpp
+++ b/Engine/Plugins/FX/Niagara/Source/NiagaraEditorWidgets/Private/Stack/SNiagaraStackModuleItem.cpp
@@ -323,7 +323,7 @@ TSharedRef<SWidget> SNiagaraStackModuleItem::AddContainerForRowWidgets(TSharedRe
 		[
 			SNew(SButton)
 			.Text(LOCTEXT("CreateNote_Yes", "Add Note"))
-			.ButtonStyle(FEditorStyle::Get(), "FlatButton.Success")
+			.ButtonStyle(FAppStyle::Get(), "FlatButton.Success")
 			.IsEnabled_Lambda([&]()
 			{
 				return !DescriptionTextBox->GetText().IsEmpty();
@@ -337,7 +337,7 @@ TSharedRef<SWidget> SNiagaraStackModuleItem::AddContainerForRowWidgets(TSharedRe
 		.Padding(5.f)
 		[
 			SNew(SButton)
-			.ButtonStyle(FEditorStyle::Get(), "FlatButton.Light")
+			.ButtonStyle(FAppStyle::Get(), "FlatButton.Dark")
 			.Text(LOCTEXT("CreateNote_No", "Cancel"))
 			.OnClicked_Lambda(CancelPressed)
 		]
