Added unit metadata to system/emitter properties

diff --git a/Engine/Plugins/FX/Niagara/Source/Niagara/Classes/NiagaraEmitter.h b/Engine/Plugins/FX/Niagara/Source/Niagara/Classes/NiagaraEmitter.h
--- a/Engine/Plugins/FX/Niagara/Source/Niagara/Classes/NiagaraEmitter.h
+++ b/Engine/Plugins/FX/Niagara/Source/Niagara/Classes/NiagaraEmitter.h
@@ -382,7 +382,7 @@ public:
 	uint32 bCombineEventSpawn : 1;
 
 	/** Limits the delta time per tick to prevent simulation spikes due to frame lags. */
-	UPROPERTY(EditAnywhere, AdvancedDisplay, Category = "Emitter", meta = (EditCondition = "bLimitDeltaTime"))
+	UPROPERTY(EditAnywhere, AdvancedDisplay, Category = "Emitter", meta = (EditCondition = "bLimitDeltaTime", ForceUnits=s))
 	float MaxDeltaTimePerTick;
 
 
diff --git a/Engine/Plugins/FX/Niagara/Source/Niagara/Classes/NiagaraSystem.h b/Engine/Plugins/FX/Niagara/Source/Niagara/Classes/NiagaraSystem.h
--- a/Engine/Plugins/FX/Niagara/Source/Niagara/Classes/NiagaraSystem.h
+++ b/Engine/Plugins/FX/Niagara/Source/Niagara/Classes/NiagaraSystem.h
@@ -888,7 +888,7 @@ protected:
 	int32 RandomSeed = 0;
 
 	/** Warm up time in seconds. Used to calculate WarmupTickCount. Rounds down to the nearest multiple of WarmupTickDelta. */
-	UPROPERTY(EditAnywhere, Category = Warmup)
+	UPROPERTY(EditAnywhere, Category = Warmup, meta = (ForceUnits=s))
 	float WarmupTime;
 
 	/** Number of ticks to process for warmup. You can set by this or by time via WarmupTime. */
@@ -896,7 +896,7 @@ protected:
 	int32 WarmupTickCount;
 
 	/** Delta time to use for warmup ticks. */
-	UPROPERTY(EditAnywhere, Category = Warmup)
+	UPROPERTY(EditAnywhere, Category = Warmup, meta = (ForceUnits=s))
 	float WarmupTickDelta;
 
 	UPROPERTY(EditAnywhere, Category = "System", meta = (InlineEditConditionToggle))
@@ -909,7 +909,7 @@ protected:
 
 	The max number of substeps per frame can be set via fx.Niagara.SystemSimulation.MaxTickSubsteps
 	*/
-	UPROPERTY(EditAnywhere, Category = "System", meta = (EditCondition = "bFixedTickDelta"))
+	UPROPERTY(EditAnywhere, Category = "System", meta = (EditCondition = "bFixedTickDelta", ForceUnits=s))
 	float FixedTickDeltaTime = 0.01667;
 
 #if WITH_EDITORONLY_DATA
