Changed the parameter types of the transform nodes to not be as confusing for vector/position types.
Removed duplicate op definition.

diff --git a/Engine/Plugins/FX/Niagara/Source/NiagaraEditor/Private/NiagaraEditorCommon.cpp b/Engine/Plugins/FX/Niagara/Source/NiagaraEditor/Private/NiagaraEditorCommon.cpp
--- a/Engine/Plugins/FX/Niagara/Source/NiagaraEditor/Private/NiagaraEditorCommon.cpp
+++ b/Engine/Plugins/FX/Niagara/Source/NiagaraEditor/Private/NiagaraEditorCommon.cpp
@@ -1464,33 +1464,22 @@ void FNiagaraOpInfo::Init()
 	Idx = OpInfos.AddDefaulted();
 	Op = &OpInfos[Idx];
 	Op->Category = MatrixCategory;
-	Op->FriendlyName = NSLOCTEXT("NiagaraOpInfo", "MatrixVector Mul Name", "Multiply (Matrix * Vector4)");
-	Op->Description = NSLOCTEXT("NiagaraOpInfo", "MatrixVector Mul Desc", "Multiplies a matrix by a vector4.");
-	Op->Inputs.Add(FNiagaraOpInOutInfo(M, MatrixType, MText, MText, Default_MatrixOne));
-	Op->Inputs.Add(FNiagaraOpInOutInfo(V, FNiagaraTypeDefinition::GetVec4Def(), VText, VText, Default_Vector4One));
-	Op->Outputs.Add(FNiagaraOpInOutInfo(Result, FNiagaraTypeDefinition::GetVec4Def(), ResultText, ResultText, Default_Vector4One, TEXT("mul({1},{0})")));
-	Op->BuildName(TEXT("MatrixVectorMultiply"), MatrixCategoryName);
-	OpInfoMap.Add(Op->Name) = Idx;
-
-	Idx = OpInfos.AddDefaulted();
-	Op = &OpInfos[Idx];
-	Op->Category = MatrixCategory;
-	Op->FriendlyName = NSLOCTEXT("NiagaraOpInfo", "TransformPosition Name", "Transform Position");
-	Op->Description = NSLOCTEXT("NiagaraOpInfo", "TransformPosition Desc", "Transforms a Vector3 as a position, meaning it multiplies it with the full matrix, including translation.");
+	Op->FriendlyName = NSLOCTEXT("NiagaraOpInfo", "TransformPosition Name", "Matrix Transform Position");
+	Op->Description = NSLOCTEXT("NiagaraOpInfo", "TransformPosition Desc", "Transforms a position vector, meaning it multiplies it with the full matrix, including translation.");
 	Op->Inputs.Add(FNiagaraOpInOutInfo(M, MatrixType, MText, MText, Default_MatrixOne));
-	Op->Inputs.Add(FNiagaraOpInOutInfo(V, FNiagaraTypeDefinition::GetGenericNumericDef(), VText, VText, Default_VectorOne));
-	Op->Outputs.Add(FNiagaraOpInOutInfo(Result, FNiagaraTypeDefinition::GetGenericNumericDef(), ResultText, ResultText, Default_VectorOne, TEXT("mul(float4({1},1.0),{0}).xyz")));
+	Op->Inputs.Add(FNiagaraOpInOutInfo(V, FNiagaraTypeDefinition::GetPositionDef(), VText, VText, Default_VectorOne));
+	Op->Outputs.Add(FNiagaraOpInOutInfo(Result, FNiagaraTypeDefinition::GetPositionDef(), ResultText, ResultText, Default_VectorOne, TEXT("mul(float4({1},1.0),{0}).xyz")));
 	Op->BuildName(TEXT("TransformPosition"), MatrixCategoryName);
 	OpInfoMap.Add(Op->Name) = Idx;
 
 	Idx = OpInfos.AddDefaulted();
 	Op = &OpInfos[Idx];
 	Op->Category = MatrixCategory;
-	Op->FriendlyName = NSLOCTEXT("NiagaraOpInfo", "TransformVector Name", "Transform Vector");
+	Op->FriendlyName = NSLOCTEXT("NiagaraOpInfo", "TransformVector Name", "Matrix Transform Vector");
 	Op->Description = NSLOCTEXT("NiagaraOpInfo", "TransformVector Desc", "Transforms a Vector3 as a vector, meaning it does not apply the translation part of the matrix.");
 	Op->Inputs.Add(FNiagaraOpInOutInfo(M, MatrixType, MText, MText, Default_MatrixOne));
-	Op->Inputs.Add(FNiagaraOpInOutInfo(V, FNiagaraTypeDefinition::GetGenericNumericDef(), VText, VText, Default_VectorOne));
-	Op->Outputs.Add(FNiagaraOpInOutInfo(Result, FNiagaraTypeDefinition::GetGenericNumericDef(), ResultText, ResultText, Default_VectorOne, TEXT("mul(float4({1},0.0),{0}).xyz")));
+	Op->Inputs.Add(FNiagaraOpInOutInfo(V, FNiagaraTypeDefinition::GetVec3Def(), VText, VText, Default_VectorOne));
+	Op->Outputs.Add(FNiagaraOpInOutInfo(Result, FNiagaraTypeDefinition::GetVec3Def(), ResultText, ResultText, Default_VectorOne, TEXT("mul(float4({1},0.0),{0}).xyz")));
 	Op->BuildName(TEXT("TransformVector"), MatrixCategoryName);
 	OpInfoMap.Add(Op->Name) = Idx;
 
