Show convertible linked inputs in the ui for stack input bindings

diff --git a/Engine/Plugins/FX/Niagara/Source/NiagaraEditor/Private/ViewModels/Stack/NiagaraStackFunctionInput.cpp b/Engine/Plugins/FX/Niagara/Source/NiagaraEditor/Private/ViewModels/Stack/NiagaraStackFunctionInput.cpp
--- a/Engine/Plugins/FX/Niagara/Source/NiagaraEditor/Private/ViewModels/Stack/NiagaraStackFunctionInput.cpp
+++ b/Engine/Plugins/FX/Niagara/Source/NiagaraEditor/Private/ViewModels/Stack/NiagaraStackFunctionInput.cpp
@@ -30,7 +30,6 @@
 #include "NiagaraScriptMergeManager.h"
 #include "NiagaraScriptSource.h"
 #include "NiagaraScriptVariable.h"
-#include "NiagaraStackEditorData.h"
 #include "NiagaraSystemScriptViewModel.h"
 #include "ScopedTransaction.h"
 #include "EdGraph/EdGraphPin.h"
@@ -600,7 +599,6 @@ void UNiagaraStackFunctionInput::RefreshChildrenInternal(const TArray<UNiagaraSt
 
 	if (InputValues.Mode == EValueMode::Dynamic && InputValues.DynamicNode.IsValid())
 	{
-		UNiagaraScript* FunctionScript = InputValues.DynamicNode->FunctionScript;
 		FVersionedNiagaraScriptData* ScriptData = InputValues.DynamicNode->GetScriptData();
 		if (ScriptData != nullptr)
 		{
@@ -1330,8 +1328,21 @@ FName GetNamespaceForUsage(ENiagaraScriptUsage Usage)
 	}
 }
 
-void UNiagaraStackFunctionInput::GetAvailableParameterHandles(TArray<FNiagaraParameterHandle>& AvailableParameterHandles) const
+UNiagaraScript* UNiagaraStackFunctionInput::FindConversionScript(const FNiagaraTypeDefinition& FromType, TMap<FNiagaraTypeDefinition, UNiagaraScript*>& ConversionScriptCache) const
 {
+	if (UNiagaraScript** CacheEntry = ConversionScriptCache.Find(FromType))
+	{
+		return *CacheEntry;
+	}
+	TArray<UNiagaraScript*> Scripts = GetPossibleConversionScripts(FromType);
+	ConversionScriptCache.Add(FromType, Scripts.Num() == 0? nullptr : Scripts[0]);
+	return ConversionScriptCache[FromType];
+}
+
+void UNiagaraStackFunctionInput::GetAvailableParameterHandles(TArray<FNiagaraParameterHandle>& AvailableParameterHandles, TMap<FNiagaraVariable, UNiagaraScript*>& AvailableConversionHandles) const
+{
+	TMap<FNiagaraTypeDefinition, UNiagaraScript*> ConversionScriptCache;
+	
 	// Engine Handles.
 	for (const FNiagaraVariable& SystemVariable : FNiagaraConstants::GetEngineConstants())
 	{
@@ -1339,11 +1350,29 @@ void UNiagaraStackFunctionInput::GetAvailableParameterHandles(TArray<FNiagaraPar
 		{
 			AvailableParameterHandles.Add(FNiagaraParameterHandle::CreateEngineParameterHandle(SystemVariable));
 		}
+		else if (UNiagaraScript* ConversionScript = FindConversionScript(SystemVariable.GetType(), ConversionScriptCache))
+		{
+			AvailableConversionHandles.Add(SystemVariable, ConversionScript);
+		}
 	}
 
-	UNiagaraNodeOutput* CurrentOutputNode = FNiagaraStackGraphUtilities::GetEmitterOutputNodeForStackNode(*OwningModuleNode);
-	UNiagaraEmitter* Emitter = nullptr;
+	// user parameters
+	TArray<FNiagaraVariable> ExposedVars;
+	GetSystemViewModel()->GetSystem().GetExposedParameters().GetParameters(ExposedVars);
+	for (const FNiagaraVariable& ExposedVar : ExposedVars)
+	{
+		if (ExposedVar.GetType() == InputType)
+		{
+			AvailableParameterHandles.Add(FNiagaraParameterHandle::CreateEngineParameterHandle(ExposedVar));
+		}
+		else if (UNiagaraScript* ConversionScript = FindConversionScript(ExposedVar.GetType(), ConversionScriptCache))
+		{
+			AvailableConversionHandles.Add(ExposedVar, ConversionScript);
+		}
+	}
 
+	// gather variables from the parameter map history
+	UNiagaraEmitter* Emitter = nullptr;
 	TArray<UNiagaraNodeOutput*> AllOutputNodes;
 	if (GetEmitterViewModel().IsValid())
 	{
@@ -1354,19 +1383,8 @@ void UNiagaraStackFunctionInput::GetAvailableParameterHandles(TArray<FNiagaraPar
 	{
 		GetSystemViewModel()->GetSystemScriptViewModel()->GetGraphViewModel()->GetGraph()->GetNodesOfClass<UNiagaraNodeOutput>(AllOutputNodes);
 	}
-	
+	UNiagaraNodeOutput* CurrentOutputNode = FNiagaraStackGraphUtilities::GetEmitterOutputNodeForStackNode(*OwningModuleNode);
 	TArray<FName> StackContextRoots = FNiagaraStackGraphUtilities::StackContextResolution(Emitter, CurrentOutputNode);
-
-	TArray<FNiagaraVariable> ExposedVars;
-	GetSystemViewModel()->GetSystem().GetExposedParameters().GetParameters(ExposedVars);
-	for (const FNiagaraVariable& ExposedVar : ExposedVars)
-	{
-		if (ExposedVar.GetType() == InputType)
-		{
-			AvailableParameterHandles.Add(FNiagaraParameterHandle::CreateEngineParameterHandle(ExposedVar));
-		}
-	}
-
 	for (UNiagaraNodeOutput* OutputNode : AllOutputNodes)
 	{
 		// Check if this is in a spawn event handler and the emitter is not using interpolated spawn so we
@@ -1407,33 +1425,44 @@ void UNiagaraStackFunctionInput::GetAvailableParameterHandles(TArray<FNiagaraPar
 					for (int32 j = 0; j < Builder.Histories[0].Variables.Num(); j++)
 					{
 						FNiagaraVariable& HistoryVariable = Builder.Histories[0].Variables[j];
-						FNiagaraVariable& AliasedHistoryVariable = Builder.Histories[0].VariablesWithOriginalAliasesIntact[j];
 						FNiagaraParameterHandle AvailableHandle = FNiagaraParameterHandle(HistoryVariable.GetName());
-						if (HistoryVariable.GetType() == InputType)
+
+						// check if the variable was written to
+						bool bWritten = false;
+						for (const FModuleScopedPin& WritePin : Builder.Histories[0].PerVariableWriteHistory[j])
 						{
-							for (const FModuleScopedPin& WritePin : Builder.Histories[0].PerVariableWriteHistory[j])
+							if (Cast<UNiagaraNodeParameterMapSet>(WritePin.Pin->GetOwningNode()) != nullptr)
 							{
-								if (Cast<UNiagaraNodeParameterMapSet>(WritePin.Pin->GetOwningNode()) != nullptr)
-								{
-									AvailableParameterHandles.AddUnique(AvailableHandle);
-									AvailableParameterHandlesForThisOutput.AddUnique(AvailableHandle);
+								bWritten = true;
+								break;
+							}
+						}
 
-									// Check to see if any variables can be converted to StackContext. This may be more portable for people to setup.
-									for (int32 StackRootIdx = 0; StackRootIdx < StackContextRoots.Num(); StackRootIdx++)
+						if (bWritten)
+						{
+							if (HistoryVariable.GetType() == InputType)
+							{
+								AvailableParameterHandles.AddUnique(AvailableHandle);
+								AvailableParameterHandlesForThisOutput.AddUnique(AvailableHandle);
+
+								// Check to see if any variables can be converted to StackContext. This may be more portable for people to setup.
+								for (int32 StackRootIdx = 0; StackRootIdx < StackContextRoots.Num(); StackRootIdx++)
+								{
+									if (HistoryVariable.IsInNameSpace(StackContextRoots[StackRootIdx]))
 									{
-										if (HistoryVariable.IsInNameSpace(StackContextRoots[StackRootIdx]))
-										{
-											// We do a replace here so that we can leave modifiers and other parts intact that might also be aliased.
-											FString NewName = HistoryVariable.GetName().ToString().Replace(*StackContextRoots[StackRootIdx].ToString(), *FNiagaraConstants::StackContextNamespace.ToString());
-
-											FNiagaraParameterHandle AvailableAliasedHandle = FNiagaraParameterHandle(*NewName);
-											AvailableParameterHandles.AddUnique(AvailableAliasedHandle);
-											AvailableParameterHandlesForThisOutput.AddUnique(AvailableAliasedHandle);
-										}
+										// We do a replace here so that we can leave modifiers and other parts intact that might also be aliased.
+										FString NewName = HistoryVariable.GetName().ToString().Replace(*StackContextRoots[StackRootIdx].ToString(), *FNiagaraConstants::StackContextNamespace.ToString());
+
+										FNiagaraParameterHandle AvailableAliasedHandle = FNiagaraParameterHandle(*NewName);
+										AvailableParameterHandles.AddUnique(AvailableAliasedHandle);
+										AvailableParameterHandlesForThisOutput.AddUnique(AvailableAliasedHandle);
 									}
-									break;
 								}
 							}
+							else if (UNiagaraScript* ConversionScript = FindConversionScript(HistoryVariable.GetType(), ConversionScriptCache))
+							{
+								AvailableConversionHandles.Add(HistoryVariable, ConversionScript);
+							}
 						}
 					}
 				}
@@ -1471,6 +1500,10 @@ void UNiagaraStackFunctionInput::GetAvailableParameterHandles(TArray<FNiagaraPar
 				{
 					AvailableParameterHandles.AddUnique(FNiagaraParameterHandle(CollectionParam.GetName()));
 				}
+				else if (UNiagaraScript* ConversionScript = FindConversionScript(CollectionParam.GetType(), ConversionScriptCache))
+				{
+					AvailableConversionHandles.Add(CollectionParam, ConversionScript);
+				}
 			}
 		}
 		else
@@ -2312,6 +2345,24 @@ void UNiagaraStackFunctionInput::SetLinkedInputViaConversionScript(const FName&
 	}
 }
 
+void UNiagaraStackFunctionInput::SetLinkedInputViaConversionScript(const FNiagaraVariable& LinkedInput, UNiagaraScript* ConversionScript)
+{
+	if (ConversionScript == nullptr)
+	{
+		return;
+	}
+	FScopedTransaction ScopedTransaction(LOCTEXT("SetConversionInput", "Make auto-convert dynamic input"));
+	SetDynamicInput(ConversionScript);
+	for (UNiagaraStackFunctionInput* ChildInput : GetChildInputs())
+	{
+		if (LinkedInput.GetType() == ChildInput->GetInputType())
+		{
+			ChildInput->SetLinkedValueHandle(FNiagaraParameterHandle(LinkedInput.GetName()));
+			break;
+		}
+	}
+}
+
 void UNiagaraStackFunctionInput::SetClipboardContentViaConversionScript(const UNiagaraClipboardFunctionInput& ClipboardFunctionInput)
 {
 	TArray<UNiagaraScript*> NiagaraScripts = GetPossibleConversionScripts(ClipboardFunctionInput.InputType);
@@ -2741,7 +2792,8 @@ void UNiagaraStackFunctionInput::GetDefaultLinkedHandleOrLinkedFunctionFromDefau
 	{
 		// If there are a chain of linked values use the first one that's available, otherwise just use the last one.
 		TArray<FNiagaraParameterHandle> AvailableHandles;
-		GetAvailableParameterHandles(AvailableHandles);
+		TMap<FNiagaraVariable, UNiagaraScript*> ConversionHandles;
+		GetAvailableParameterHandles(AvailableHandles, ConversionHandles);
 		for (FLinkedHandleOrFunctionNode& LinkedValue : LinkedValues)
 		{
 			if (LinkedValue.LinkedFunctionCallNode.IsValid() ||
diff --git a/Engine/Plugins/FX/Niagara/Source/NiagaraEditor/Public/ViewModels/Stack/NiagaraStackFunctionInput.h b/Engine/Plugins/FX/Niagara/Source/NiagaraEditor/Public/ViewModels/Stack/NiagaraStackFunctionInput.h
--- a/Engine/Plugins/FX/Niagara/Source/NiagaraEditor/Public/ViewModels/Stack/NiagaraStackFunctionInput.h
+++ b/Engine/Plugins/FX/Niagara/Source/NiagaraEditor/Public/ViewModels/Stack/NiagaraStackFunctionInput.h
@@ -131,7 +131,7 @@ public:
 	void SetLinkedValueHandle(const FNiagaraParameterHandle& InParameterHandle);
 
 	/** Gets the current set of available parameter handles which can be assigned to this input. */
-	void GetAvailableParameterHandles(TArray<FNiagaraParameterHandle>& AvailableParameterHandles) const;
+	void GetAvailableParameterHandles(TArray<FNiagaraParameterHandle>& AvailableParameterHandles, TMap<FNiagaraVariable, UNiagaraScript*>& AvailableConversionHandles) const;
 
 	/** Gets the function node form the script graph if the current value mode is DefaultFunction. */
 	UNiagaraNodeFunctionCall* GetDefaultFunctionNode() const;
@@ -253,6 +253,7 @@ public:
 	TArray<UNiagaraScript*> GetPossibleConversionScripts(const FNiagaraTypeDefinition& FromType) const;
 
 	void SetLinkedInputViaConversionScript(const FName& LinkedInputName, const FNiagaraTypeDefinition& FromType);
+	void SetLinkedInputViaConversionScript(const FNiagaraVariable& LinkedInput, UNiagaraScript* ConversionScript);
 	void SetClipboardContentViaConversionScript(const UNiagaraClipboardFunctionInput& ClipboardFunctionInput);
 
 	void ChangeScriptVersion(FGuid NewScriptVersion);
@@ -367,6 +368,8 @@ private:
 
 	void GetCurrentChangeIds(FGuid& OutOwningGraphChangeId, FGuid& OutFunctionGraphChangeId) const;
 
+	UNiagaraScript* FindConversionScript(const FNiagaraTypeDefinition& FromType, TMap<FNiagaraTypeDefinition, UNiagaraScript*>& ConversionScriptCache) const;
+
 private:
 	/** The module function call which owns this input entry. NOTE: This input might not be an input to the module function
 		call, it may be an input to a dynamic input function call which is owned by the module. */
diff --git a/Engine/Plugins/FX/Niagara/Source/NiagaraEditorWidgets/Private/Stack/SNiagaraStackFunctionInputValue.cpp b/Engine/Plugins/FX/Niagara/Source/NiagaraEditorWidgets/Private/Stack/SNiagaraStackFunctionInputValue.cpp
--- a/Engine/Plugins/FX/Niagara/Source/NiagaraEditorWidgets/Private/Stack/SNiagaraStackFunctionInputValue.cpp
+++ b/Engine/Plugins/FX/Niagara/Source/NiagaraEditorWidgets/Private/Stack/SNiagaraStackFunctionInputValue.cpp
@@ -743,6 +743,11 @@ void SNiagaraStackFunctionInputValue::ParameterHandleSelected(FNiagaraParameterH
 	FunctionInput->SetLinkedValueHandle(Handle);
 }
 
+void SNiagaraStackFunctionInputValue::ConversionHandleSelected(FNiagaraVariable Handle, UNiagaraScript* ConversionScript)
+{
+	FunctionInput->SetLinkedInputViaConversionScript(Handle, ConversionScript);
+}
+
 EVisibility SNiagaraStackFunctionInputValue::GetResetButtonVisibility() const
 {
 	return FunctionInput->CanReset() ? EVisibility::Visible : EVisibility::Hidden;
@@ -1076,7 +1081,8 @@ TArray<TSharedPtr<FNiagaraMenuAction_Generic>> SNiagaraStackFunctionInputValue::
 
 	// Link existing attribute
 	TArray<FNiagaraParameterHandle> AvailableHandles;
-	FunctionInput->GetAvailableParameterHandles(AvailableHandles);
+	TMap<FNiagaraVariable, UNiagaraScript*> AvailableConversionHandles;
+	FunctionInput->GetAvailableParameterHandles(AvailableHandles, AvailableConversionHandles);
 
 	const FString RootCategoryName = FString("Link Inputs");
 	const FText MapInputFormat = LOCTEXT("LinkInputFormat", "Link this input to {0}");
@@ -1103,6 +1109,36 @@ TArray<TSharedPtr<FNiagaraMenuAction_Generic>> SNiagaraStackFunctionInputValue::
 		}
 	}
 
+	const FText ConvertInputFormat = LOCTEXT("ConvertInputFormat", "Link this input to {0} via a conversion script");
+	for (const auto& Entry : AvailableConversionHandles)
+	{
+		FNiagaraVariable ParameterVariable = Entry.Key;
+		UNiagaraScript* ConversionScript = Entry.Value;
+		const FNiagaraParameterHandle& AvailableHandle(ParameterVariable.GetName());
+		TArray<FName> HandleParts = AvailableHandle.GetHandleParts();
+		FNiagaraNamespaceMetadata NamespaceMetadata = GetDefault<UNiagaraEditorSettings>()->GetMetaDataForNamespaces(HandleParts);
+		if (NamespaceMetadata.IsValid())
+		{			
+			// Only add handles which are in known namespaces to prevent collecting parameter handles
+			// which are being used to configure modules and dynamic inputs in the stack graphs.
+			const FText Category = NamespaceMetadata.DisplayName;
+			const FText DisplayName = FText::FromName(AvailableHandle.GetParameterHandleString());
+			const FText Tooltip = FText::Format(ConvertInputFormat, FText::FromName(AvailableHandle.GetParameterHandleString()));
+
+			TSharedPtr<FNiagaraMenuAction_Generic> LinkAction(new FNiagaraMenuAction_Generic(
+				FNiagaraMenuAction_Generic::FOnExecuteAction::CreateSP(this, &SNiagaraStackFunctionInputValue::ConversionHandleSelected, ParameterVariable, ConversionScript),
+				DisplayName, ENiagaraMenuSections::General, {RootCategoryName, Category.ToString()}, Tooltip, FText()));
+
+			LinkAction->SetParameterVariable(ParameterVariable);
+
+			// set the source data from the script
+			TTuple<EScriptSource, FText> Source = FNiagaraEditorUtilities::GetScriptSource(ConversionScript);
+			LinkAction->SourceData = FNiagaraActionSourceData(Source.Key, Source.Value, true);;
+
+			OutAllActions.Add(LinkAction);
+		}
+	}
+
 	// Read from new attribute
 	{
 		const FText CategoryName = LOCTEXT("MakeCategory", "Make");
diff --git a/Engine/Plugins/FX/Niagara/Source/NiagaraEditorWidgets/Private/Stack/SNiagaraStackFunctionInputValue.h b/Engine/Plugins/FX/Niagara/Source/NiagaraEditorWidgets/Private/Stack/SNiagaraStackFunctionInputValue.h
--- a/Engine/Plugins/FX/Niagara/Source/NiagaraEditorWidgets/Private/Stack/SNiagaraStackFunctionInputValue.h
+++ b/Engine/Plugins/FX/Niagara/Source/NiagaraEditorWidgets/Private/Stack/SNiagaraStackFunctionInputValue.h
@@ -115,6 +115,7 @@ private:
 	void CreateScratchSelected();
 
 	void ParameterHandleSelected(FNiagaraParameterHandle Handle);
+	void ConversionHandleSelected(FNiagaraVariable Handle, UNiagaraScript* ConversionScript);
 
 	EVisibility GetResetButtonVisibility() const;
 
