- Niagara shaders now default the cooked out versions in the editor
- During editing we enable the quicker version we will include rapid iteration parameters, won't trim attributes, etc
- Also fixes editing scripts when RI bake out is enabled

#rnx

diff --git a/Engine/Plugins/FX/Niagara/Source/Niagara/Classes/NiagaraEmitter.h b/Engine/Plugins/FX/Niagara/Source/Niagara/Classes/NiagaraEmitter.h
--- a/Engine/Plugins/FX/Niagara/Source/Niagara/Classes/NiagaraEmitter.h
+++ b/Engine/Plugins/FX/Niagara/Source/Niagara/Classes/NiagaraEmitter.h
@@ -411,10 +411,6 @@ public:
 	UPROPERTY()
 	TObjectPtr<class UNiagaraScriptSourceBase>	GraphSource;
 
-	/** Should we enable rapid iteration removal if the system is also set to remove rapid iteration parameters on compile? This value defaults to true.*/
-	UPROPERTY(EditAnywhere, AdvancedDisplay, Category = "Emitter", meta = (DisplayName = "Supports Baked Rapid Iteration"))
-	uint32 bBakeOutRapidIteration : 1;
-
 	NIAGARA_API bool AreAllScriptAndSourcesSynchronized() const;
 	NIAGARA_API void OnPostCompile();
 
diff --git a/Engine/Plugins/FX/Niagara/Source/Niagara/Classes/NiagaraSystem.h b/Engine/Plugins/FX/Niagara/Source/Niagara/Classes/NiagaraSystem.h
--- a/Engine/Plugins/FX/Niagara/Source/Niagara/Classes/NiagaraSystem.h
+++ b/Engine/Plugins/FX/Niagara/Source/Niagara/Classes/NiagaraSystem.h
@@ -547,12 +547,26 @@ public:
 	static void RecomputeExecutionOrderForEmitter(UNiagaraEmitter* InEmitter);
 	static void RecomputeExecutionOrderForDataInterface(class UNiagaraDataInterface* DataInterface);
 
-	/** Experimental feature that allows us to bake out rapid iteration parameters during the normal compile process. */
-	UPROPERTY(EditAnywhere, AdvancedDisplay, Category = "Performance")
+	FORCEINLINE bool ShouldUseRapidIterationParameters() const { return bCompileForEdit ? !bBakeOutRapidIteration : !bBakeOutRapidIterationOnCook; }
+	FORCEINLINE bool ShouldTrimAttributes() const { return bCompileForEdit ? bTrimAttributes : bTrimAttributesOnCook; }
+	FORCEINLINE bool ShouldIgnoreParticleReadsForAttributeTrim() const { return bIgnoreParticleReadsForAttributeTrim; }
+	FORCEINLINE bool ShouldDisableDebugSwitches() const { return bCompileForEdit ? bDisableDebugSwitches : bDisableDebugSwitchesOnCook; }
+	FORCEINLINE bool ShouldCompressAttributes() const { return bCompressAttributes; }
+
+	FORCEINLINE void SetBakeOutRapidIterationOnCook(bool bBakeOut) { bBakeOutRapidIteration = bBakeOut; bBakeOutRapidIterationOnCook = bBakeOut; }
+	FORCEINLINE void SetTrimAttributesOnCook(bool bTrim) { bTrimAttributes = bTrim; bTrimAttributesOnCook = bTrim; }
+
+	/** When enabled we compile for the edit path, which can result in slower system performance but faster editor responsiveness. */
+	UPROPERTY(transient)
+	uint32 bCompileForEdit : 1;
+
+protected:
+	/** When enable constant values are baked into the scripts while editing the system, this will increase iteration times but improve performance. */
+	UPROPERTY(EditAnywhere, AdvancedDisplay, Category = "Performance", meta=(DisplayName="Bake Rapid Iteration Parameters During Edit"))
 	uint32 bBakeOutRapidIteration : 1;
 
-	/** If true bBakeOutRapidIteration will be made to be true during cooks  */
-	UPROPERTY(EditAnywhere, AdvancedDisplay, Category = "Performance", meta = (SkipSystemResetOnChange = "true"))
+	/** When enabled constant values are baked into scripts to improve performance. */
+	UPROPERTY(EditAnywhere, AdvancedDisplay, Category = "Performance", meta=(DisplayName="Bake Rapid Iteration Parameters"))
 	uint32 bBakeOutRapidIterationOnCook : 1;
 
 	/** Toggles whether or not emitters within this system will try and compress their particle attributes.
@@ -560,26 +574,25 @@ public:
 	UPROPERTY(EditAnywhere, AdvancedDisplay, Category = "Performance")
 	uint32 bCompressAttributes : 1;
 
-	/** If true Particle attributes will be removed from the DataSet if they are unnecessary (are never read by ParameterMap) */
-	UPROPERTY(EditAnywhere, AdvancedDisplay, Category = "Performance")
+	/** When enabled we trim particle attributes while editing the system. */
+	UPROPERTY(EditAnywhere, AdvancedDisplay, Category = "Performance", meta=(DisplayName="Trim Attributes During Edit"))
 	uint32 bTrimAttributes : 1;
-
-	/** If true bTrimAttributes will be made to be true during cooks */
-	UPROPERTY(EditAnywhere, AdvancedDisplay, Category = "Performance", meta = (SkipSystemResetOnChange = "true"))
+	/** If true Particle attributes will be removed from the DataSet if they are unnecessary (are never read by ParameterMap) */
+	UPROPERTY(EditAnywhere, AdvancedDisplay, Category = "Performance", meta = (DisplayName="Trim Attributes"))
 	uint32 bTrimAttributesOnCook : 1;
 
 	/** If true ParticleReads will not absolutely prevent attribute trimming - User must ensure that the appropriate attributes are preserved on the source emitter! */
 	UPROPERTY(EditAnywhere, AdvancedDisplay, Category = "Performance")
 	uint32 bIgnoreParticleReadsForAttributeTrim : 1;
 
-	/** If true all debug switches will be ignored in the editor. */
-	UPROPERTY(EditAnywhere, AdvancedDisplay, Category = "Performance")
+	/** When enable debug switches are disabled while editing the system. */
+	UPROPERTY(EditAnywhere, AdvancedDisplay, Category = "Performance", meta = (DisplayName="Disable Debug Switches During Edit"))
 	uint32 bDisableDebugSwitches : 1;
-
-	/** If true all debug switches will be ignored for cooked builds */
-	UPROPERTY(EditAnywhere, AdvancedDisplay, Category = "Performance", meta = (SkipSystemResetOnChange = "true"))
+	/** When enabled debug switches are disabled when compiling the system. */
+	UPROPERTY(EditAnywhere, AdvancedDisplay, Category = "Performance", meta = (DisplayName="Disable Debug Switches"))
 	uint32 bDisableDebugSwitchesOnCook : 1;
 
+public:
 	/** Subscriptions to definitions of parameters. */
 	UPROPERTY()
 	TArray<FParameterDefinitionsSubscription> ParameterDefinitionsSubscriptions;
diff --git a/Engine/Plugins/FX/Niagara/Source/Niagara/Private/NiagaraEmitter.cpp b/Engine/Plugins/FX/Niagara/Source/Niagara/Private/NiagaraEmitter.cpp
--- a/Engine/Plugins/FX/Niagara/Source/Niagara/Private/NiagaraEmitter.cpp
+++ b/Engine/Plugins/FX/Niagara/Source/Niagara/Private/NiagaraEmitter.cpp
@@ -157,7 +157,6 @@ UNiagaraEmitter::UNiagaraEmitter(const FObjectInitializer& Initializer)
 , bLimitDeltaTime(true)
 , MaxGPUParticlesSpawnPerFrame(0)
 #if WITH_EDITORONLY_DATA
-, bBakeOutRapidIteration(true)
 , ThumbnailImageOutOfDate(true)
 , IsCooked(false)
 #endif
diff --git a/Engine/Plugins/FX/Niagara/Source/Niagara/Private/NiagaraScript.cpp b/Engine/Plugins/FX/Niagara/Source/Niagara/Private/NiagaraScript.cpp
--- a/Engine/Plugins/FX/Niagara/Source/Niagara/Private/NiagaraScript.cpp
+++ b/Engine/Plugins/FX/Niagara/Source/Niagara/Private/NiagaraScript.cpp
@@ -827,24 +827,18 @@ void UNiagaraScript::ComputeVMCompilationId(FNiagaraVMExecutableDataId& Id, FGui
 			SystemSpawnIdx = Scripts.Add(EmitterOwner->GetSystemSpawnScript());
 			SystemUpdateIdx = Scripts.Add(EmitterOwner->GetSystemUpdateScript());
 
-			if (EmitterOwner->bBakeOutRapidIteration)
-			{
-				Id.bUsesRapidIterationParams = false;
-			}
-			if (EmitterOwner->bDisableDebugSwitches)
-			{
-				Id.bDisableDebugSwitches = true;
-			}
-			if (EmitterOwner->bCompressAttributes)
+			Id.bUsesRapidIterationParams = EmitterOwner->ShouldUseRapidIterationParameters();
+			Id.bDisableDebugSwitches = EmitterOwner->ShouldDisableDebugSwitches();
+			if (EmitterOwner->ShouldCompressAttributes())
 			{
 				Id.AdditionalDefines.Add(TEXT("CompressAttributes"));
 			}
-			if (EmitterOwner->bIgnoreParticleReadsForAttributeTrim)
+			if (EmitterOwner->ShouldIgnoreParticleReadsForAttributeTrim())
 			{
 				Id.AdditionalDefines.Add(TEXT("IgnoreParticleReadsForAttributeTrim"));
 			}
 
-			bool TrimAttributes = EmitterOwner->bTrimAttributes;
+			bool TrimAttributes = EmitterOwner->ShouldTrimAttributes();
 			if (TrimAttributes)
 			{
 				auto TrimAttributesSupported = [=](const UNiagaraEmitter* OtherEmitter)
@@ -870,7 +864,7 @@ void UNiagaraScript::ComputeVMCompilationId(FNiagaraVMExecutableDataId& Id, FGui
 				};
 
 				// if this emitter is being referenced by another emitter (PartilceRead) then don't worry about trimming attributes
-				if (!EmitterOwner->bIgnoreParticleReadsForAttributeTrim)
+				if (!EmitterOwner->ShouldIgnoreParticleReadsForAttributeTrim())
 				{
 					for (const FNiagaraEmitterHandle& EmitterHandle : EmitterOwner->GetEmitterHandles())
 					{
@@ -962,11 +956,6 @@ void UNiagaraScript::ComputeVMCompilationId(FNiagaraVMExecutableDataId& Id, FGui
 			Id.AdditionalDefines.Add(TEXT("Emitter.Determinism"));
 		}
 
-		if (!Emitter->bBakeOutRapidIteration)
-		{
-			Id.bUsesRapidIterationParams = true;
-		}
-
 		// Has simulation stages
 		{
 
@@ -1009,15 +998,9 @@ void UNiagaraScript::ComputeVMCompilationId(FNiagaraVMExecutableDataId& Id, FGui
 		SystemSpawnIdx = Scripts.Add(System->GetSystemSpawnScript());
 		SystemUpdateIdx = Scripts.Add(System->GetSystemUpdateScript());
 
-		if (System->bBakeOutRapidIteration)
-		{
-			Id.bUsesRapidIterationParams = false;
-		}
-		if (System->bDisableDebugSwitches)
-		{
-			Id.bDisableDebugSwitches = true;
-		}
-		if (System->bCompressAttributes)
+		Id.bUsesRapidIterationParams = System->ShouldUseRapidIterationParameters();
+		Id.bDisableDebugSwitches = System->ShouldDisableDebugSwitches();
+		if (System->ShouldCompressAttributes())
 		{
 			Id.AdditionalDefines.Add(TEXT("CompressAttributes"));
 		}
@@ -1543,22 +1526,9 @@ void UNiagaraScript::Serialize(FArchive& Ar)
 		bool bUsesRapidIterationParams = true;
 
 #if WITH_EDITORONLY_DATA
-		if (UNiagaraEmitter* Emitter = GetTypedOuter<UNiagaraEmitter>())
+		if (UNiagaraSystem* System = GetTypedOuter<UNiagaraSystem>())
 		{
-			if ( UNiagaraSystem* System = Emitter->GetTypedOuter<UNiagaraSystem>() )
-			{
-				if (System->bBakeOutRapidIteration)
-				{
-					bUsesRapidIterationParams = false;
-				}
-			}
-		}
-		else if (UNiagaraSystem* System = GetTypedOuter<UNiagaraSystem>())
-		{
-			if (System->bBakeOutRapidIteration)
-			{
-				bUsesRapidIterationParams = false;
-			}
+			bUsesRapidIterationParams = System->ShouldUseRapidIterationParameters();
 		}
 #endif
 
diff --git a/Engine/Plugins/FX/Niagara/Source/Niagara/Private/NiagaraSystem.cpp b/Engine/Plugins/FX/Niagara/Source/Niagara/Private/NiagaraSystem.cpp
--- a/Engine/Plugins/FX/Niagara/Source/Niagara/Private/NiagaraSystem.cpp
+++ b/Engine/Plugins/FX/Niagara/Source/Niagara/Private/NiagaraSystem.cpp
@@ -54,14 +54,6 @@ namespace NiagaraScriptCookStats
 int32 GEnableNiagaraRuntimeCycleCounts = 0;
 static FAutoConsoleVariableRef CVarEnableNiagaraRuntimeCycleCounts(TEXT("fx.EnableNiagaraRuntimeCycleCounts"), GEnableNiagaraRuntimeCycleCounts, TEXT("Toggle for runtime cylce counts tracking Niagara's frame time. \n"), ECVF_ReadOnly);
 
-static int GNiagaraForceSystemsToCookOutRapidIterationOnLoad = 0;
-static FAutoConsoleVariableRef CVarNiagaraForceSystemsToCookOutRapidIterationOnLoad(
-	TEXT("fx.NiagaraForceSystemsToCookOutRapidIterationOnLoad"),
-	GNiagaraForceSystemsToCookOutRapidIterationOnLoad,
-	TEXT("When enabled UNiagaraSystem's bBakeOutRapidIteration will be forced to true on PostLoad of the system."),
-	ECVF_Default
-);
-
 static int GNiagaraLogDDCStatusForSystems = 0;
 static FAutoConsoleVariableRef CVarLogDDCStatusForSystems(
 	TEXT("fx.NiagaraLogDDCStatusForSystems"),
@@ -107,6 +99,8 @@ static FAutoConsoleVariableRef CVarNiagaraCompileDDCWaitTimeout(
 UNiagaraSystem::UNiagaraSystem(const FObjectInitializer& ObjectInitializer)
 : Super(ObjectInitializer)
 #if WITH_EDITORONLY_DATA
+, bCompileForEdit(false)
+, bBakeOutRapidIteration(false)
 , bBakeOutRapidIterationOnCook(true)
 , bTrimAttributes(false)
 , bTrimAttributesOnCook(true)
@@ -886,12 +880,6 @@ void UNiagaraSystem::UpdateSystemAfterLoad()
 			UE_LOG(LogNiagara, Log, TEXT("System %s being compiled because there were changes to an emitter script Change ID."), *GetPathName());
 		}
 
-		if (GNiagaraForceSystemsToCookOutRapidIterationOnLoad == 1 && !bBakeOutRapidIteration)
-		{
-			bSystemScriptsAreSynchronized = false;
-			bBakeOutRapidIteration = true;
-		}
-		
 		if (bSystemScriptsAreSynchronized == false || bEmitterScriptsAreSynchronized == false)
 		{
 			if (IsRunningCommandlet())
@@ -1029,15 +1017,6 @@ void UNiagaraSystem::Serialize(FArchive& Ar)
 			NiagaraEmitterCompiledDataStruct->SerializeTaggedProperties(Ar, (uint8*)&ConstCastSharedRef<FNiagaraEmitterCompiledData>(EmitterCompiledData[EmitterIndex]).Get(), NiagaraEmitterCompiledDataStruct, nullptr);
 		}
 	}
-
-#if WITH_EDITOR
-	if ((IsRunningCookCommandlet() || GIsCookerLoadingPackage) && Ar.IsLoading())
-	{
-		bBakeOutRapidIteration = bBakeOutRapidIteration || bBakeOutRapidIterationOnCook;
-		bTrimAttributes = bTrimAttributes || bTrimAttributesOnCook;
-		bDisableDebugSwitches = bDisableDebugSwitches || bDisableDebugSwitchesOnCook;
-	}
-#endif
 }
 
 #if WITH_EDITOR
diff --git a/Engine/Plugins/FX/Niagara/Source/NiagaraEditor/Private/Commandlets/NiagaraDumpBytecodeCommandlet.cpp b/Engine/Plugins/FX/Niagara/Source/NiagaraEditor/Private/Commandlets/NiagaraDumpBytecodeCommandlet.cpp
--- a/Engine/Plugins/FX/Niagara/Source/NiagaraEditor/Private/Commandlets/NiagaraDumpBytecodeCommandlet.cpp
+++ b/Engine/Plugins/FX/Niagara/Source/NiagaraEditor/Private/Commandlets/NiagaraDumpBytecodeCommandlet.cpp
@@ -146,8 +146,14 @@ void UNiagaraDumpByteCodeCommandlet::ProcessNiagaraScripts()
 
 		if (ForceBakedRapidIteration || ForceAttributeTrimming)
 		{
-			NiagaraSystem->bBakeOutRapidIteration = ForceBakedRapidIteration;
-			NiagaraSystem->bTrimAttributes = ForceAttributeTrimming;
+			if (ForceBakedRapidIteration)
+			{
+				NiagaraSystem->SetBakeOutRapidIterationOnCook(true);
+			}
+			if (ForceAttributeTrimming)
+			{
+				NiagaraSystem->SetTrimAttributesOnCook(true);
+			}
 
 			NiagaraSystem->RequestCompile(true);
 			NiagaraSystem->WaitForCompilationComplete(true);
diff --git a/Engine/Plugins/FX/Niagara/Source/NiagaraEditor/Private/NiagaraCompiler.cpp b/Engine/Plugins/FX/Niagara/Source/NiagaraEditor/Private/NiagaraCompiler.cpp
--- a/Engine/Plugins/FX/Niagara/Source/NiagaraEditor/Private/NiagaraCompiler.cpp
+++ b/Engine/Plugins/FX/Niagara/Source/NiagaraEditor/Private/NiagaraCompiler.cpp
@@ -873,13 +873,12 @@ TSharedPtr<FNiagaraCompileRequestDataBase, ESPMode::ThreadSafe> FNiagaraEditorMo
 	}
 	else if (System)
 	{
-
 		TSharedPtr<FNiagaraGraphCachedDataBase, ESPMode::ThreadSafe> CachedTraversalSystemData = System->GetCachedTraversalData();
 
 		check(System->GetSystemSpawnScript()->GetLatestSource() == System->GetSystemUpdateScript()->GetLatestSource());
 		BasePtr->Source = Cast<UNiagaraScriptSource>(System->GetSystemSpawnScript()->GetLatestSource());
-		BasePtr->bUseRapidIterationParams = !System->bBakeOutRapidIteration;
-		BasePtr->bDisableDebugSwitches = System->bDisableDebugSwitches;
+		BasePtr->bUseRapidIterationParams = System->ShouldUseRapidIterationParameters();
+		BasePtr->bDisableDebugSwitches = System->ShouldDisableDebugSwitches();
 		
 		TArray<FString> EmitterNames;
 
@@ -899,7 +898,7 @@ TSharedPtr<FNiagaraCompileRequestDataBase, ESPMode::ThreadSafe> FNiagaraEditorMo
 			EmitterPtr->EmitterUniqueName = Handle.GetInstance()->GetUniqueEmitterName();
 			EmitterPtr->SourceName = BasePtr->SourceName;
 			EmitterPtr->Source = Cast<UNiagaraScriptSource>(Handle.GetInstance()->GraphSource);
-			EmitterPtr->bUseRapidIterationParams = BasePtr->bUseRapidIterationParams || (!Handle.GetInstance()->bBakeOutRapidIteration);
+			EmitterPtr->bUseRapidIterationParams = BasePtr->bUseRapidIterationParams;
 			EmitterPtr->bDisableDebugSwitches = BasePtr->bDisableDebugSwitches;
 			//EmitterPtr->bSimulationStagesEnabled = Handle.GetInstance()->bSimulationStagesEnabled;
 			EmitterPtr->SharedCompileDataInterfaceData = BasePtr->SharedCompileDataInterfaceData;
diff --git a/Engine/Plugins/FX/Niagara/Source/NiagaraEditor/Private/NiagaraParameterMapHistory.cpp b/Engine/Plugins/FX/Niagara/Source/NiagaraEditor/Private/NiagaraParameterMapHistory.cpp
--- a/Engine/Plugins/FX/Niagara/Source/NiagaraEditor/Private/NiagaraParameterMapHistory.cpp
+++ b/Engine/Plugins/FX/Niagara/Source/NiagaraEditor/Private/NiagaraParameterMapHistory.cpp
@@ -2206,7 +2206,7 @@ ENiagaraFunctionDebugState FCompileConstantResolver::CalculateDebugState() const
 	const UNiagaraSystem* OwnerSystem = System ? System : (Emitter ? Cast<const UNiagaraSystem>(Emitter->GetOuter()) : nullptr);
 	if ( OwnerSystem != nullptr )
 	{
-		return OwnerSystem->bDisableDebugSwitches ? ENiagaraFunctionDebugState::NoDebug : DebugState;
+		return OwnerSystem->ShouldDisableDebugSwitches() ? ENiagaraFunctionDebugState::NoDebug : DebugState;
 	}
 
 	// Translator
diff --git a/Engine/Plugins/FX/Niagara/Source/NiagaraEditor/Private/ViewModels/NiagaraSystemViewModel.cpp b/Engine/Plugins/FX/Niagara/Source/NiagaraEditor/Private/ViewModels/NiagaraSystemViewModel.cpp
--- a/Engine/Plugins/FX/Niagara/Source/NiagaraEditor/Private/ViewModels/NiagaraSystemViewModel.cpp
+++ b/Engine/Plugins/FX/Niagara/Source/NiagaraEditor/Private/ViewModels/NiagaraSystemViewModel.cpp
@@ -95,6 +95,8 @@ FNiagaraSystemViewModel::FNiagaraSystemViewModel()
 void FNiagaraSystemViewModel::Initialize(UNiagaraSystem& InSystem, FNiagaraSystemViewModelOptions InOptions)
 {
 	System = &InSystem;
+	System->bCompileForEdit = true;
+
 	RegisteredHandle = RegisterViewModelWithMap(System, this);
 
 	bCanModifyEmittersFromTimeline = InOptions.bCanModifyEmittersFromTimeline;
@@ -176,6 +178,7 @@ void FNiagaraSystemViewModel::Cleanup()
 
 	if (System)
 	{
+		System->bCompileForEdit = false;
 		System->OnSystemPostEditChange().Remove(SystemChangedDelegateHandle);
 	}
 
diff --git a/Engine/Plugins/FX/Niagara/Source/NiagaraEditor/Private/ViewModels/Stack/NiagaraStackFunctionInput.cpp b/Engine/Plugins/FX/Niagara/Source/NiagaraEditor/Private/ViewModels/Stack/NiagaraStackFunctionInput.cpp
--- a/Engine/Plugins/FX/Niagara/Source/NiagaraEditor/Private/ViewModels/Stack/NiagaraStackFunctionInput.cpp
+++ b/Engine/Plugins/FX/Niagara/Source/NiagaraEditor/Private/ViewModels/Stack/NiagaraStackFunctionInput.cpp
@@ -1710,6 +1710,12 @@ void UNiagaraStackFunctionInput::SetLocalValue(TSharedRef<FStructOnScope> InLoca
 				Script->MarkScriptAndSourceDesynchronized(TEXT("Static variable changed!"), FGuid::NewGuid());
 			}
 		}
+
+		UNiagaraSystem& NiagaraSystem = GetSystemViewModel()->GetSystem();
+		if (NiagaraSystem.ShouldUseRapidIterationParameters() == false)
+		{
+			NiagaraSystem.RequestCompile(false);
+		}
 	}
 	else
 	{
