- Separate out widget being enabled or not so we don't reactive the HUD on editor launch all the time

#rnx

diff --git a/Engine/Plugins/FX/Niagara/Source/Niagara/Private/NiagaraDebugHud.cpp b/Engine/Plugins/FX/Niagara/Source/Niagara/Private/NiagaraDebugHud.cpp
--- a/Engine/Plugins/FX/Niagara/Source/Niagara/Private/NiagaraDebugHud.cpp
+++ b/Engine/Plugins/FX/Niagara/Source/Niagara/Private/NiagaraDebugHud.cpp
@@ -91,7 +91,15 @@ namespace NiagaraDebugLocal
 	static TTuple<const TCHAR*, const TCHAR*, TFunction<void(FString)>> GDebugConsoleCommands[] =
 	{
 		// Main HUD commands
-		MakeTuple(TEXT("Enabled="), TEXT("Enable or disable the HUD"), [](FString Arg) {Settings.bEnabled = FCString::Atoi(*Arg) != 0; }),
+		MakeTuple(TEXT("Enabled="), TEXT("Enable or disable the HUD"),
+			[](FString Arg)
+			{
+				Settings.bHudEnabled = FCString::Atoi(*Arg) != 0;
+				#if WITH_EDITORONLY_DATA
+				Settings.bWidgetEnabled = true;
+				#endif
+			}
+		),
 		MakeTuple(TEXT("ValidateSystemSimulationDataBuffers="), TEXT("Enable or disable validation on system data buffers"), [](FString Arg) {Settings.bValidateSystemSimulationDataBuffers = FCString::Atoi(*Arg) != 0; }),
 		MakeTuple(TEXT("bValidateParticleDataBuffers="), TEXT("Enable or disable validation on particle data buffers"), [](FString Arg) {Settings.bValidateParticleDataBuffers = FCString::Atoi(*Arg) != 0; }),
 
@@ -683,7 +691,7 @@ void FNiagaraDebugHud::GatherSystemInfo()
 	}
 
 	// When not enabled do nothing
-	if (!Settings.bEnabled)
+	if (!Settings.IsEnabled())
 	{
 		return;
 	}
@@ -995,7 +1003,7 @@ void FNiagaraDebugHud::DebugDrawCallback(UCanvas* Canvas, APlayerController* PC)
 {
 	using namespace NiagaraDebugLocal;
 
-	if (!Settings.bEnabled)
+	if (!Settings.IsEnabled())
 	{
 		return;
 	}
diff --git a/Engine/Plugins/FX/Niagara/Source/Niagara/Private/NiagaraDebuggerClient.cpp b/Engine/Plugins/FX/Niagara/Source/Niagara/Private/NiagaraDebuggerClient.cpp
--- a/Engine/Plugins/FX/Niagara/Source/Niagara/Private/NiagaraDebuggerClient.cpp
+++ b/Engine/Plugins/FX/Niagara/Source/Niagara/Private/NiagaraDebuggerClient.cpp
@@ -165,8 +165,8 @@ void FNiagaraDebuggerClient::HandleDebugHUDSettingsMessage(const FNiagaraDebugHU
 
 				//TODO: Move these to just take direct from the debug hud per worldman?
 				//Possibly move the debug hud itself to the debugger client rather than having one per world man and they all share global state.
-				const ENiagaraDebugPlaybackMode PlaybackMode = Message.bEnabled ? Message.PlaybackMode : ENiagaraDebugPlaybackMode::Play;
-				const float PlaybackRate = Message.bEnabled && Message.bPlaybackRateEnabled ? Message.PlaybackRate : 1.0f;
+				const ENiagaraDebugPlaybackMode PlaybackMode = Message.IsEnabled() ? Message.PlaybackMode : ENiagaraDebugPlaybackMode::Play;
+				const float PlaybackRate = Message.IsEnabled() && Message.bPlaybackRateEnabled ? Message.PlaybackRate : 1.0f;
 				WorldMan.SetDebugPlaybackMode(PlaybackMode);
 				WorldMan.SetDebugPlaybackRate(PlaybackRate);
 			};
@@ -174,7 +174,7 @@ void FNiagaraDebuggerClient::HandleDebugHUDSettingsMessage(const FNiagaraDebugHU
 		FNiagaraWorldManager::ForAllWorldManagers(ApplySettingsToWorldMan);
 			
 		//TODO: Move usage to come direct from settings struct instead of this CVar.
-		const float GlobalLoopTime = Message.bEnabled && Message.bLoopTimeEnabled && Message.PlaybackMode == ENiagaraDebugPlaybackMode::Loop ? Message.LoopTime : 0.0f;
+		const float GlobalLoopTime = Message.IsEnabled() && Message.bLoopTimeEnabled && Message.PlaybackMode == ENiagaraDebugPlaybackMode::Loop ? Message.LoopTime : 0.0f;
 		ExecuteConsoleCommand(*FString::Printf(TEXT("fx.Niagara.Debug.GlobalLoopTime %.3f"), GlobalLoopTime), true);
 	}
 }
diff --git a/Engine/Plugins/FX/Niagara/Source/Niagara/Public/NiagaraDebuggerCommon.h b/Engine/Plugins/FX/Niagara/Source/Niagara/Public/NiagaraDebuggerCommon.h
--- a/Engine/Plugins/FX/Niagara/Source/Niagara/Public/NiagaraDebuggerCommon.h
+++ b/Engine/Plugins/FX/Niagara/Source/Niagara/Public/NiagaraDebuggerCommon.h
@@ -412,9 +412,23 @@ struct NIAGARA_API FNiagaraDebugHUDSettingsData
 
 	FNiagaraDebugHUDSettingsData();
 
+	bool IsEnabled() const
+	{
+	#if WITH_EDITORONLY_DATA
+		return bWidgetEnabled && bHudEnabled;
+	#else
+		return bHudEnabled;
+	#endif
+	}
+
+#if WITH_EDITORONLY_DATA
+	UPROPERTY(transient)
+	bool bWidgetEnabled = false;
+#endif
+
 	/** Master control for all HUD features. */
 	UPROPERTY(EditAnywhere, Category = "Debug General", meta = (DisplayName = "Debug HUD Enabled"))
-	bool bEnabled = true;
+	bool bHudEnabled = true;
 
 	/**
 	When enabled all Niagara systems that pass the filter will have the simulation data buffers validation.
diff --git a/Engine/Plugins/FX/Niagara/Source/NiagaraEditor/Private/NiagaraDebugger.cpp b/Engine/Plugins/FX/Niagara/Source/NiagaraEditor/Private/NiagaraDebugger.cpp
--- a/Engine/Plugins/FX/Niagara/Source/NiagaraEditor/Private/NiagaraDebugger.cpp
+++ b/Engine/Plugins/FX/Niagara/Source/NiagaraEditor/Private/NiagaraDebugger.cpp
@@ -124,9 +124,16 @@ void FNiagaraDebugger::TriggerOutlinerCapture()
 			{
 				if (UNiagaraDebugHUDSettings* Settings = GetMutableDefault<UNiagaraDebugHUDSettings>())
 				{
-					if (!Settings->Data.bEnabled)
+					if (
+#if WITH_EDITORONLY_DATA
+						!Settings->Data.bWidgetEnabled ||
+#endif
+						!Settings->Data.bHudEnabled )
 					{
-						Settings->Data.bEnabled = true;
+#if WITH_EDITORONLY_DATA
+						Settings->Data.bWidgetEnabled = true;
+#endif
+						Settings->Data.bHudEnabled = true;
 						Settings->PostEditChange();
 					}
 				}
diff --git a/Engine/Plugins/FX/Niagara/Source/NiagaraEditor/Private/Widgets/SNiagaraDebugger.cpp b/Engine/Plugins/FX/Niagara/Source/NiagaraEditor/Private/Widgets/SNiagaraDebugger.cpp
--- a/Engine/Plugins/FX/Niagara/Source/NiagaraEditor/Private/Widgets/SNiagaraDebugger.cpp
+++ b/Engine/Plugins/FX/Niagara/Source/NiagaraEditor/Private/Widgets/SNiagaraDebugger.cpp
@@ -826,7 +826,7 @@ TSharedRef<SDockTab> SNiagaraDebugger::SpawnNiagaraDebugger(const FSpawnTabArgs&
 			[](TSharedRef<SDockTab> Self, TWeakPtr<FTabManager> TabManager)
 			{
 				UNiagaraDebugHUDSettings* HudSettings = GetMutableDefault<UNiagaraDebugHUDSettings>();
-				HudSettings->Data.bEnabled = false;
+				HudSettings->Data.bWidgetEnabled = false;
 				HudSettings->NotifyPropertyChanged();
 
 				TSharedPtr<FTabManager> OwningTabManager = TabManager.Pin();
@@ -841,7 +841,7 @@ TSharedRef<SDockTab> SNiagaraDebugger::SpawnNiagaraDebugger(const FSpawnTabArgs&
 	);
 
 	UNiagaraDebugHUDSettings* HudSettings = GetMutableDefault<UNiagaraDebugHUDSettings>();
-	HudSettings->Data.bEnabled = true;
+	HudSettings->Data.bWidgetEnabled = true;
 	HudSettings->NotifyPropertyChanged();
 
 	FNiagaraEditorModule& NiagaraEditorModule = FModuleManager::GetModuleChecked<FNiagaraEditorModule>("NiagaraEditor");
