multiple attributes for rasterization grid
update various flip modules to use new rasterization method

diff --git a/Engine/Plugins/FX/Niagara/Source/Niagara/Classes/NiagaraDataInterfaceRasterizationGrid3D.h b/Engine/Plugins/FX/Niagara/Source/Niagara/Classes/NiagaraDataInterfaceRasterizationGrid3D.h
--- a/Engine/Plugins/FX/Niagara/Source/Niagara/Classes/NiagaraDataInterfaceRasterizationGrid3D.h
+++ b/Engine/Plugins/FX/Niagara/Source/Niagara/Classes/NiagaraDataInterfaceRasterizationGrid3D.h
@@ -27,8 +27,9 @@ public:
 	}
 
 	void ResizeBuffers();
-
+	int32 TotalNumAttributes = 0;
 	FIntVector NumCells = FIntVector::ZeroValue;
+	FIntVector NumTiles = FIntVector::ZeroValue;
 	float Precision = 0;
 	int ResetValue = 0;
 
@@ -36,6 +37,8 @@ public:
 	
 	FTextureRWBuffer3D RasterizationBuffer;
 
+	FReadBuffer PerAttributeData;
+
 #if STATS
 	int32 GPUMemory = 0;
 #endif
@@ -61,6 +64,10 @@ class NIAGARA_API UNiagaraDataInterfaceRasterizationGrid3D : public UNiagaraData
 
 public:	
 
+	// Number of attributes stored on the grid
+	UPROPERTY(EditAnywhere, Category = "Grid")
+	int32 NumAttributes;
+
 	UPROPERTY(EditAnywhere, Category = "Grid")
 	float Precision;
 
@@ -113,6 +120,8 @@ public:
 	static const FName SetNumCellsFunctionName;
 	static const FName SetFloatResetValueFunctionName;	
 
+	static const FString PerAttributeDataName;
+
 protected:
 	//~ UNiagaraDataInterface interface
 	virtual bool CopyToInternal(UNiagaraDataInterface* Destination) const override;
diff --git a/Engine/Plugins/FX/Niagara/Source/Niagara/Private/NiagaraDataInterfaceRasterizationGrid3D.cpp b/Engine/Plugins/FX/Niagara/Source/Niagara/Private/NiagaraDataInterfaceRasterizationGrid3D.cpp
--- a/Engine/Plugins/FX/Niagara/Source/Niagara/Private/NiagaraDataInterfaceRasterizationGrid3D.cpp
+++ b/Engine/Plugins/FX/Niagara/Source/Niagara/Private/NiagaraDataInterfaceRasterizationGrid3D.cpp
@@ -16,7 +16,7 @@ static const FString PrecisionName(TEXT("Precision_"));
 const FName UNiagaraDataInterfaceRasterizationGrid3D::SetNumCellsFunctionName("SetNumCells");
 const FName UNiagaraDataInterfaceRasterizationGrid3D::SetFloatResetValueFunctionName("SetFloatResetValue");
 
-
+const FString UNiagaraDataInterfaceRasterizationGrid3D::PerAttributeDataName(TEXT("PerAttributeDataName_"));
 
 // Global VM function names, also used by the shaders code generation methods.
 
@@ -34,7 +34,7 @@ static const FName SetIntValueFunctionName("SetIntGridValue");
 static const FName GetIntValueFunctionName("GetIntGridValue");
 
 
-static int32 GMaxNiagaraRasterizationGridCells = (1000*1000*1000);
+static int32 GMaxNiagaraRasterizationGridCells = (1024 * 1024 * 1024);
 static FAutoConsoleVariableRef CVarMaxNiagaraRasterizationGridCells(
 	TEXT("fx.MaxNiagaraRasterizationGridCells"),
 	GMaxNiagaraRasterizationGridCells,
@@ -50,6 +50,7 @@ struct FNiagaraDataInterfaceParametersCS_RasterizationGrid3D : public FNiagaraDa
 public:
 	void Bind(const FNiagaraDataInterfaceGPUParamInfo& ParameterInfo, const class FShaderParameterMap& ParameterMap)
 	{
+		NumAttributesParam.Bind(ParameterMap, *(UNiagaraDataInterfaceRWBase::NumAttributesName + ParameterInfo.DataInterfaceHLSLSymbol));
 		NumCellsParam.Bind(ParameterMap, *(UNiagaraDataInterfaceRWBase::NumCellsName + ParameterInfo.DataInterfaceHLSLSymbol));
 		UnitToUVParam.Bind(ParameterMap, *(UNiagaraDataInterfaceRWBase::UnitToUVName + ParameterInfo.DataInterfaceHLSLSymbol));		
 		PrecisionParam.Bind(ParameterMap, *(PrecisionName + ParameterInfo.DataInterfaceHLSLSymbol));
@@ -57,6 +58,8 @@ public:
 		IntGridParam.Bind(ParameterMap,  *(IntGridName + ParameterInfo.DataInterfaceHLSLSymbol));
 		
 		OutputIntGridParam.Bind(ParameterMap, *(OutputIntGridName + ParameterInfo.DataInterfaceHLSLSymbol));
+
+		PerAttributeDataParam.Bind(ParameterMap, *(UNiagaraDataInterfaceRasterizationGrid3D::PerAttributeDataName + ParameterInfo.DataInterfaceHLSLSymbol));
 	}
 
 	// #todo(dmp): make resource transitions batched
@@ -72,11 +75,12 @@ public:
 
 		if (!(ProxyData && ProxyData->RasterizationBuffer.Buffer.IsValid()))
 		{			
+			SetShaderValue(RHICmdList, ComputeShaderRHI, NumAttributesParam, 0);
 			SetShaderValue(RHICmdList, ComputeShaderRHI, NumCellsParam, 0);
 			SetShaderValue(RHICmdList, ComputeShaderRHI, UnitToUVParam, FVector3f::ZeroVector);						
 			SetShaderValue(RHICmdList, ComputeShaderRHI, PrecisionParam, 0);
 			SetSRVParameter(RHICmdList, ComputeShaderRHI, IntGridParam, FNiagaraRenderer::GetDummyIntBuffer());
-
+			SetSRVParameter(RHICmdList, ComputeShaderRHI, PerAttributeDataParam, FNiagaraRenderer::GetDummyIntBuffer());
 			if (OutputIntGridParam.IsUAVBound())
 			{
 				RHICmdList.SetUAVParameter(ComputeShaderRHI, OutputIntGridParam.GetUAVIndex(), Context.ComputeDispatchInterface->GetEmptyUAVFromPool(RHICmdList, PF_R32_SINT, ENiagaraEmptyUAVType::Buffer));
@@ -85,10 +89,12 @@ public:
 			return;
 		}
 
+		SetShaderValue(RHICmdList, ComputeShaderRHI, NumAttributesParam, ProxyData->TotalNumAttributes);
 		SetShaderValue(RHICmdList, ComputeShaderRHI, NumCellsParam, ProxyData->NumCells);
 		SetShaderValue(RHICmdList, ComputeShaderRHI, UnitToUVParam, FVector3f(1.0f) / FVector3f(ProxyData->NumCells));
 		SetShaderValue(RHICmdList, ComputeShaderRHI, PrecisionParam, ProxyData->Precision);
-		
+		SetSRVParameter(RHICmdList, ComputeShaderRHI, PerAttributeDataParam, ProxyData->PerAttributeData.SRV);
+
 		if (!Context.IsOutputStage)
 		{			
 			if (IntGridParam.IsBound())
@@ -127,11 +133,13 @@ public:
 	}
 
 private:
+	LAYOUT_FIELD(FShaderParameter, NumAttributesParam);
 	LAYOUT_FIELD(FShaderParameter, NumCellsParam);
 	LAYOUT_FIELD(FShaderParameter, UnitToUVParam);		
 	LAYOUT_FIELD(FShaderParameter, PrecisionParam);
 	LAYOUT_FIELD(FShaderResourceParameter, IntGridParam);
 	LAYOUT_FIELD(FRWShaderParameter, OutputIntGridParam);	
+	LAYOUT_FIELD(FShaderResourceParameter, PerAttributeDataParam);
 };
 
 void RasterizationGrid3DRWInstanceData::ResizeBuffers()
@@ -143,7 +151,7 @@ void RasterizationGrid3DRWInstanceData::ResizeBuffers()
 	if (NumTotalCells > (uint32)GMaxNiagaraRasterizationGridCells)
 		return;
 
-	RasterizationBuffer.Initialize(TEXT("NiagaraRasterizationGrid3D::IntGrid"), sizeof(int32), NumCells.X, NumCells.Y, NumCells.Z, EPixelFormat::PF_R32_SINT);
+	RasterizationBuffer.Initialize(TEXT("NiagaraRasterizationGrid3D::IntGrid"), sizeof(int32), NumCells.X * NumTiles.X, NumCells.Y * NumTiles.Y, NumCells.Z * NumTiles.Z, EPixelFormat::PF_R32_SINT);
 
 	#if STATS
 		DEC_MEMORY_STAT_BY(STAT_NiagaraGPUDataInterfaceMemory, GPUMemory);
@@ -158,6 +166,7 @@ IMPLEMENT_NIAGARA_DI_PARAMETER(UNiagaraDataInterfaceRasterizationGrid3D, FNiagar
 
 UNiagaraDataInterfaceRasterizationGrid3D::UNiagaraDataInterfaceRasterizationGrid3D(FObjectInitializer const& ObjectInitializer)
 	: Super(ObjectInitializer)
+	, NumAttributes(1)
 	, Precision(1.f)
 	, ResetValue(0)
 {	
@@ -211,7 +220,8 @@ void UNiagaraDataInterfaceRasterizationGrid3D::GetFunctions(TArray<FNiagaraFunct
 		Sig.Inputs.Add(FNiagaraVariable(FNiagaraTypeDefinition(GetClass()), TEXT("Grid")));
 		Sig.Inputs.Add(FNiagaraVariable(FNiagaraTypeDefinition::GetIntDef(), TEXT("IndexX")));
 		Sig.Inputs.Add(FNiagaraVariable(FNiagaraTypeDefinition::GetIntDef(), TEXT("IndexY")));
-		Sig.Inputs.Add(FNiagaraVariable(FNiagaraTypeDefinition::GetIntDef(), TEXT("IndexZ")));		
+		Sig.Inputs.Add(FNiagaraVariable(FNiagaraTypeDefinition::GetIntDef(), TEXT("IndexZ")));
+		Sig.Inputs.Add(FNiagaraVariable(FNiagaraTypeDefinition::GetIntDef(), TEXT("AttributeIndex")));
 		Sig.Inputs.Add(FNiagaraVariable(FNiagaraTypeDefinition::GetFloatDef(), TEXT("Value")));
 		Sig.Outputs.Add(FNiagaraVariable(FNiagaraTypeDefinition::GetIntDef(), TEXT("IGNORE")));
 
@@ -236,6 +246,7 @@ void UNiagaraDataInterfaceRasterizationGrid3D::GetFunctions(TArray<FNiagaraFunct
 		Sig.Inputs.Add(FNiagaraVariable(FNiagaraTypeDefinition::GetIntDef(), TEXT("IndexX")));
 		Sig.Inputs.Add(FNiagaraVariable(FNiagaraTypeDefinition::GetIntDef(), TEXT("IndexY")));
 		Sig.Inputs.Add(FNiagaraVariable(FNiagaraTypeDefinition::GetIntDef(), TEXT("IndexZ")));
+		Sig.Inputs.Add(FNiagaraVariable(FNiagaraTypeDefinition::GetIntDef(), TEXT("AttributeIndex")));
 		Sig.Inputs.Add(FNiagaraVariable(FNiagaraTypeDefinition::GetIntDef(), TEXT("Value")));		
 		Sig.Outputs.Add(FNiagaraVariable(FNiagaraTypeDefinition::GetIntDef(), TEXT("IGNORE")));
 
@@ -260,6 +271,7 @@ void UNiagaraDataInterfaceRasterizationGrid3D::GetFunctions(TArray<FNiagaraFunct
 		Sig.Inputs.Add(FNiagaraVariable(FNiagaraTypeDefinition::GetIntDef(), TEXT("IndexX")));
 		Sig.Inputs.Add(FNiagaraVariable(FNiagaraTypeDefinition::GetIntDef(), TEXT("IndexY")));
 		Sig.Inputs.Add(FNiagaraVariable(FNiagaraTypeDefinition::GetIntDef(), TEXT("IndexZ")));
+		Sig.Inputs.Add(FNiagaraVariable(FNiagaraTypeDefinition::GetIntDef(), TEXT("AttributeIndex")));
 		Sig.Inputs.Add(FNiagaraVariable(FNiagaraTypeDefinition::GetFloatDef(), TEXT("Value")));
 		Sig.Outputs.Add(FNiagaraVariable(FNiagaraTypeDefinition::GetIntDef(), TEXT("IGNORE")));
 
@@ -284,6 +296,7 @@ void UNiagaraDataInterfaceRasterizationGrid3D::GetFunctions(TArray<FNiagaraFunct
 		Sig.Inputs.Add(FNiagaraVariable(FNiagaraTypeDefinition::GetIntDef(), TEXT("IndexX")));
 		Sig.Inputs.Add(FNiagaraVariable(FNiagaraTypeDefinition::GetIntDef(), TEXT("IndexY")));
 		Sig.Inputs.Add(FNiagaraVariable(FNiagaraTypeDefinition::GetIntDef(), TEXT("IndexZ")));
+		Sig.Inputs.Add(FNiagaraVariable(FNiagaraTypeDefinition::GetIntDef(), TEXT("AttributeIndex")));
 		Sig.Inputs.Add(FNiagaraVariable(FNiagaraTypeDefinition::GetFloatDef(), TEXT("Value")));
 		Sig.Outputs.Add(FNiagaraVariable(FNiagaraTypeDefinition::GetIntDef(), TEXT("IGNORE")));
 
@@ -308,6 +321,7 @@ void UNiagaraDataInterfaceRasterizationGrid3D::GetFunctions(TArray<FNiagaraFunct
 		Sig.Inputs.Add(FNiagaraVariable(FNiagaraTypeDefinition::GetIntDef(), TEXT("IndexX")));
 		Sig.Inputs.Add(FNiagaraVariable(FNiagaraTypeDefinition::GetIntDef(), TEXT("IndexY")));
 		Sig.Inputs.Add(FNiagaraVariable(FNiagaraTypeDefinition::GetIntDef(), TEXT("IndexZ")));
+		Sig.Inputs.Add(FNiagaraVariable(FNiagaraTypeDefinition::GetIntDef(), TEXT("AttributeIndex")));
 		Sig.Inputs.Add(FNiagaraVariable(FNiagaraTypeDefinition::GetFloatDef(), TEXT("Value")));
 		Sig.Outputs.Add(FNiagaraVariable(FNiagaraTypeDefinition::GetIntDef(), TEXT("IGNORE")));
 
@@ -332,7 +346,8 @@ void UNiagaraDataInterfaceRasterizationGrid3D::GetFunctions(TArray<FNiagaraFunct
 		Sig.Inputs.Add(FNiagaraVariable(FNiagaraTypeDefinition(GetClass()), TEXT("Grid")));
 		Sig.Inputs.Add(FNiagaraVariable(FNiagaraTypeDefinition::GetIntDef(), TEXT("IndexX")));
 		Sig.Inputs.Add(FNiagaraVariable(FNiagaraTypeDefinition::GetIntDef(), TEXT("IndexY")));
-		Sig.Inputs.Add(FNiagaraVariable(FNiagaraTypeDefinition::GetIntDef(), TEXT("IndexZ")));		
+		Sig.Inputs.Add(FNiagaraVariable(FNiagaraTypeDefinition::GetIntDef(), TEXT("IndexZ")));
+		Sig.Inputs.Add(FNiagaraVariable(FNiagaraTypeDefinition::GetIntDef(), TEXT("AttributeIndex")));
 		Sig.Outputs.Add(FNiagaraVariable(FNiagaraTypeDefinition::GetFloatDef(), TEXT("Value")));
 
 		Sig.bExperimental = true;
@@ -356,7 +371,8 @@ void UNiagaraDataInterfaceRasterizationGrid3D::GetFunctions(TArray<FNiagaraFunct
 		Sig.Inputs.Add(FNiagaraVariable(FNiagaraTypeDefinition(GetClass()), TEXT("Grid")));
 		Sig.Inputs.Add(FNiagaraVariable(FNiagaraTypeDefinition::GetIntDef(), TEXT("IndexX")));
 		Sig.Inputs.Add(FNiagaraVariable(FNiagaraTypeDefinition::GetIntDef(), TEXT("IndexY")));
-		Sig.Inputs.Add(FNiagaraVariable(FNiagaraTypeDefinition::GetIntDef(), TEXT("IndexZ")));		
+		Sig.Inputs.Add(FNiagaraVariable(FNiagaraTypeDefinition::GetIntDef(), TEXT("IndexZ")));
+		Sig.Inputs.Add(FNiagaraVariable(FNiagaraTypeDefinition::GetIntDef(), TEXT("AttributeIndex")));
 		Sig.Outputs.Add(FNiagaraVariable(FNiagaraTypeDefinition::GetIntDef(), TEXT("Value")));
 
 		Sig.bExperimental = true;
@@ -424,7 +440,7 @@ bool UNiagaraDataInterfaceRasterizationGrid3D::Equals(const UNiagaraDataInterfac
 	}
 	const UNiagaraDataInterfaceRasterizationGrid3D* OtherTyped = CastChecked<const UNiagaraDataInterfaceRasterizationGrid3D>(Other);
 
-	return OtherTyped->Precision == Precision && OtherTyped->ResetValue == ResetValue;
+	return OtherTyped->NumAttributes == NumAttributes && OtherTyped->Precision == Precision && OtherTyped->ResetValue == ResetValue;
 }
 
 #if WITH_EDITORONLY_DATA
@@ -436,11 +452,15 @@ void UNiagaraDataInterfaceRasterizationGrid3D::GetParameterDefinitionHLSL(const
 		Texture3D<int> {IntGridName};		
 		RWTexture3D<int> RW{OutputIntGridName};
 		float {Precision};
+		Buffer<float4> {PerAttributeDataName};
+		int {NumAttributesName};
 	)");
 	TMap<FString, FStringFormatArg> ArgsDeclarations = {				
 		{ TEXT("IntGridName"),    IntGridName + ParamInfo.DataInterfaceHLSLSymbol },		
 		{ TEXT("OutputIntGridName"),    OutputIntGridName + ParamInfo.DataInterfaceHLSLSymbol },
 		{ TEXT("Precision"), PrecisionName + ParamInfo.DataInterfaceHLSLSymbol},
+		{ TEXT("PerAttributeDataName"), PerAttributeDataName + ParamInfo.DataInterfaceHLSLSymbol},
+		{ TEXT("NumAttributesName"), UNiagaraDataInterfaceRWBase::NumAttributesName + ParamInfo.DataInterfaceHLSLSymbol},
 	};
 	OutHLSL += FString::Format(FormatDeclarations, ArgsDeclarations);
 
@@ -490,11 +510,13 @@ bool UNiagaraDataInterfaceRasterizationGrid3D::GetFunctionHLSL(const FNiagaraDat
 		{ TEXT("FunctionName"), FunctionInfo.InstanceName},
 		{ TEXT("IntGrid"), IntGridName + ParamInfo.DataInterfaceHLSLSymbol},
 		{ TEXT("OutputIntGrid"), OutputIntGridName + ParamInfo.DataInterfaceHLSLSymbol},		
+		{TEXT("NumAttributesName"), UNiagaraDataInterfaceRWBase::NumAttributesName + ParamInfo.DataInterfaceHLSLSymbol},
 		{ TEXT("NumCellsName"), UNiagaraDataInterfaceRWBase::NumCellsName + ParamInfo.DataInterfaceHLSLSymbol},
 		{ TEXT("UnitToUVName"), UNiagaraDataInterfaceRWBase::UnitToUVName + ParamInfo.DataInterfaceHLSLSymbol},		
 		{ TEXT("IntToFloatFunctionName"), IntToFloatFunctionName.ToString() + ParamInfo.DataInterfaceHLSLSymbol},
 		{ TEXT("FloatToIntFunctionName"), FloatToIntFunctionName.ToString() + ParamInfo.DataInterfaceHLSLSymbol},
 		{ TEXT("Precision"), PrecisionName + ParamInfo.DataInterfaceHLSLSymbol},
+		{TEXT("PerAttributeDataName"), PerAttributeDataName + ParamInfo.DataInterfaceHLSLSymbol},
 	};
 
 	if (ParentRet)
@@ -517,10 +539,15 @@ bool UNiagaraDataInterfaceRasterizationGrid3D::GetFunctionHLSL(const FNiagaraDat
 	else if (FunctionInfo.DefinitionName == SetFloatValueFunctionName)
 	{
 		static const TCHAR* FormatBounds = TEXT(R"(
-			void {FunctionName}(int In_IndexX, int In_IndexY, int In_IndexZ, float In_Value, out int val)
+			void {FunctionName}(int In_IndexX, int In_IndexY, int In_IndexZ, int In_AttributeIndex, float In_Value, out int val)
 			{			
-				val = 0;				
-				RW{OutputIntGrid}[int3(In_IndexX, In_IndexY, In_IndexZ)] = {FloatToIntFunctionName}(In_Value);
+				val = 0;
+				Out_Val = 0;
+				if ( In_AttributeIndex < {NumAttributesName} )
+				{
+					int3 TileOffset = {PerAttributeDataName}[In_AttributeIndex].xyz;
+					RW{OutputIntGrid}[int3(In_IndexX + TileOffset.x, In_IndexY + TileOffset.y, In_IndexZ + TileOffset.z)] = {FloatToIntFunctionName}(In_Value);
+				}
 			}
 		)");
 
@@ -530,10 +557,14 @@ bool UNiagaraDataInterfaceRasterizationGrid3D::GetFunctionHLSL(const FNiagaraDat
 	else if (FunctionInfo.DefinitionName == SetIntValueFunctionName)
 	{
 		static const TCHAR* FormatBounds = TEXT(R"(
-			void {FunctionName}(int In_IndexX, int In_IndexY, int In_IndexZ, int In_Value, out int val)
+			void {FunctionName}(int In_IndexX, int In_IndexY, int In_IndexZ, int In_AttributeIndex, int In_Value, out int val)
 			{			
-				val = 0;			
-				RW{OutputIntGrid}[int3(In_IndexX, In_IndexY, In_IndexZ)] = In_Value;				
+				val = 0;				
+				if ( In_AttributeIndex < {NumAttributesName} )
+				{	
+					int3 TileOffset = {PerAttributeDataName}[In_AttributeIndex].xyz;			
+					RW{OutputIntGrid}[int3(In_IndexX + TileOffset.x, In_IndexY + TileOffset.y, In_IndexZ + TileOffset.z)] = In_Value;				
+				}
 			}
 		)");
 		OutHLSL += FString::Format(FormatBounds, ArgsBounds);
@@ -542,10 +573,14 @@ bool UNiagaraDataInterfaceRasterizationGrid3D::GetFunctionHLSL(const FNiagaraDat
 	else if (FunctionInfo.DefinitionName == InterlockedAddFloatGridValueFunctionName)
 	{
 		static const TCHAR* FormatBounds = TEXT(R"(
-			void {FunctionName}(int In_IndexX, int In_IndexY, int In_IndexZ, float In_Value, out int val)
+			void {FunctionName}(int In_IndexX, int In_IndexY, int In_IndexZ, int In_AttributeIndex, float In_Value, out int val)
 			{							
-				val = 0;			
-				InterlockedAdd(RW{OutputIntGrid}[int3(In_IndexX, In_IndexY, In_IndexZ)], {FloatToIntFunctionName}(In_Value));
+				val = 0;					
+				if ( In_AttributeIndex < {NumAttributesName} )
+				{
+					int3 TileOffset = {PerAttributeDataName}[In_AttributeIndex].xyz;		
+					InterlockedAdd(RW{OutputIntGrid}[int3(In_IndexX + TileOffset.x, In_IndexY + TileOffset.y, In_IndexZ + TileOffset.z)], {FloatToIntFunctionName}(In_Value));
+				}
 			}
 		)");
 		OutHLSL += FString::Format(FormatBounds, ArgsBounds);
@@ -554,10 +589,15 @@ bool UNiagaraDataInterfaceRasterizationGrid3D::GetFunctionHLSL(const FNiagaraDat
 	else if (FunctionInfo.DefinitionName == InterlockedMinFloatGridValueFunctionName)
 	{
 		static const TCHAR* FormatBounds = TEXT(R"(
-			void {FunctionName}(int In_IndexX, int In_IndexY, int In_IndexZ, float In_Value, out int val)
+			void {FunctionName}(int In_IndexX, int In_IndexY, int In_IndexZ, int In_AttributeIndex, float In_Value, out int val)
 			{							
-				val = 0;			
-				InterlockedMin(RW{OutputIntGrid}[int3(In_IndexX, In_IndexY, In_IndexZ)], {FloatToIntFunctionName}(In_Value));
+				val = 0;
+				
+				if ( In_AttributeIndex < {NumAttributesName} )
+				{
+					int3 TileOffset = {PerAttributeDataName}[In_AttributeIndex].xyz;			
+					InterlockedMin(RW{OutputIntGrid}[int3(In_IndexX + TileOffset.x, In_IndexY + TileOffset.y, In_IndexZ + TileOffset.z)], {FloatToIntFunctionName}(In_Value));
+				}
 			}
 		)");
 		OutHLSL += FString::Format(FormatBounds, ArgsBounds);
@@ -566,10 +606,14 @@ bool UNiagaraDataInterfaceRasterizationGrid3D::GetFunctionHLSL(const FNiagaraDat
 	else if (FunctionInfo.DefinitionName == InterlockedMaxFloatGridValueFunctionName)
 	{
 		static const TCHAR* FormatBounds = TEXT(R"(
-			void {FunctionName}(int In_IndexX, int In_IndexY, int In_IndexZ, float In_Value, out int val)
+			void {FunctionName}(int In_IndexX, int In_IndexY, int In_IndexZ, int In_AttributeIndex, float In_Value, out int val)
 			{							
-				val = 0;			
-				InterlockedMax(RW{OutputIntGrid}[int3(In_IndexX, In_IndexY, In_IndexZ)], {FloatToIntFunctionName}(In_Value));
+				val = 0;				
+				if ( In_AttributeIndex < {NumAttributesName} )
+				{
+					int3 TileOffset = {PerAttributeDataName}[In_AttributeIndex].xyz;			
+					InterlockedMax(RW{OutputIntGrid}[int3(In_IndexX + TileOffset.x, In_IndexY + TileOffset.y, In_IndexZ + TileOffset.z)], {FloatToIntFunctionName}(In_Value));
+				}
 			}
 		)");
 		OutHLSL += FString::Format(FormatBounds, ArgsBounds);
@@ -578,9 +622,14 @@ bool UNiagaraDataInterfaceRasterizationGrid3D::GetFunctionHLSL(const FNiagaraDat
 	else if (FunctionInfo.DefinitionName == GetFloatValueFunctionName)
 	{
 		static const TCHAR* FormatBounds = TEXT(R"(
-			void {FunctionName}(int In_IndexX, int In_IndexY, int In_IndexZ, out float Out_Val)
-			{				
-				Out_Val =  {IntToFloatFunctionName}({IntGrid}.Load(int4(In_IndexX, In_IndexY, In_IndexZ, 0)));
+			void {FunctionName}(int In_IndexX, int In_IndexY, int In_IndexZ, int In_AttributeIndex, out float Out_Val)
+			{		
+				Out_Val = 0;
+				if ( In_AttributeIndex < {NumAttributesName} )
+				{
+					int3 TileOffset = {PerAttributeDataName}[In_AttributeIndex].xyz;		
+					Out_Val =  {IntToFloatFunctionName}({IntGrid}.Load(int4(In_IndexX + TileOffset.x, In_IndexY + TileOffset.y, In_IndexZ + TileOffset.z, 0)));
+				}
 			}
 		)");
 		OutHLSL += FString::Format(FormatBounds, ArgsBounds);
@@ -589,9 +638,14 @@ bool UNiagaraDataInterfaceRasterizationGrid3D::GetFunctionHLSL(const FNiagaraDat
 	else if (FunctionInfo.DefinitionName == GetIntValueFunctionName)
 	{
 		static const TCHAR* FormatBounds = TEXT(R"(
-			void {FunctionName}(int In_IndexX, int In_IndexY, int In_IndexZ, out int Out_Val)
+			void {FunctionName}(int In_IndexX, int In_IndexY, int In_IndexZ, int In_AttributeIndex, out int Out_Val)
 			{
-				Out_Val = Out_Val = {IntGrid}.Load(int4(In_IndexX, In_IndexY, In_IndexZ, 0));				
+				Out_Val = 0;
+				if ( In_AttributeIndex < {NumAttributesName} )
+				{
+					int3 TileOffset = {PerAttributeDataName}[In_AttributeIndex].xyz;
+					Out_Val = {IntGrid}.Load(int4(In_IndexX + TileOffset.x, In_IndexY + TileOffset.y, In_IndexZ + TileOffset.z, 0));				
+				}
 			}
 		)");
 		OutHLSL += FString::Format(FormatBounds, ArgsBounds);
@@ -608,6 +662,15 @@ bool UNiagaraDataInterfaceRasterizationGrid3D::InitPerInstanceData(void* PerInst
 
 	FNiagaraDataInterfaceProxyRasterizationGrid3D* RT_Proxy = GetProxyAs<FNiagaraDataInterfaceProxyRasterizationGrid3D>();
 
+	int32 NumAttribChannelsFound = 0;
+
+	int32 NumNamedAttribChannelsFound = 0;
+	// #todo(dmp): implement named attributes
+	//FindAttributes(InstanceData->Vars, InstanceData->Offsets, NumNamedAttribChannelsFound);
+
+	NumAttribChannelsFound = NumAttributes + NumNamedAttribChannelsFound;
+	int32 RT_TotalNumAttributes = NumAttribChannelsFound;
+
 	FIntVector RT_NumCells = NumCells;		
 
 	float RT_Precision = Precision;
@@ -616,7 +679,8 @@ bool UNiagaraDataInterfaceRasterizationGrid3D::InitPerInstanceData(void* PerInst
 	RT_NumCells.X = FMath::Max(RT_NumCells.X, 1);
 	RT_NumCells.Y = FMath::Max(RT_NumCells.Y, 1);
 	RT_NumCells.Z = FMath::Max(RT_NumCells.Z, 1);
-		
+	
+	InstanceData->TotalNumAttributes = RT_TotalNumAttributes;
 	InstanceData->NumCells = RT_NumCells;
 	InstanceData->Precision = RT_Precision;
 	InstanceData->ResetValue = RT_ResetValue;
@@ -627,15 +691,37 @@ bool UNiagaraDataInterfaceRasterizationGrid3D::InitPerInstanceData(void* PerInst
 			return false;
 	}
 
+	// Compute number of tiles based on resolution of individual attributes
+	// #todo(dmp): refactor
+	int32 MaxDim = 2048;
+	int32 MaxTilesX = floor(MaxDim / InstanceData->NumCells.X);
+	int32 MaxTilesY = floor(MaxDim / InstanceData->NumCells.Y);
+	int32 MaxTilesZ = floor(MaxDim / InstanceData->NumCells.Z);
+	int32 MaxAttributes = MaxTilesX * MaxTilesY * MaxTilesZ;
+
+	// need to determine number of tiles in x and y based on number of attributes and max dimension size
+	int32 NumTilesX = FMath::Min<int32>(MaxTilesX, NumAttribChannelsFound);
+	int32 NumTilesY = FMath::Min<int32>(MaxTilesY, ceil(1.0 * NumAttribChannelsFound / NumTilesX));
+	int32 NumTilesZ = FMath::Min<int32>(MaxTilesZ, ceil(1.0 * NumAttribChannelsFound / (NumTilesX * NumTilesY)));
+
+	InstanceData->NumTiles.X = NumTilesX;
+	InstanceData->NumTiles.Y = NumTilesY;
+	InstanceData->NumTiles.Z = NumTilesZ;
+
+	check(InstanceData->NumTiles.X > 0);
+	check(InstanceData->NumTiles.Y > 0);
+	check(InstanceData->NumTiles.Z > 0);
+
 	// @todo-threadsafety. This would be a race but I'm taking a ref here. Not ideal in the long term.
 	// Push Updates to Proxy.
 	ENQUEUE_RENDER_COMMAND(FUpdateData)(
-		[RT_Proxy, RT_NumCells, RT_Precision, RT_ResetValue, InstanceID = SystemInstance->GetId()](FRHICommandListImmediate& RHICmdList)
+		[RT_Proxy, RT_TotalNumAttributes, RT_NumCells, RT_Precision, RT_ResetValue, InstanceID = SystemInstance->GetId(), RT_InstanceData = *InstanceData](FRHICommandListImmediate& RHICmdList)
 	{
 		check(!RT_Proxy->SystemInstancesToProxyData.Contains(InstanceID));
 		RasterizationGrid3DRWInstanceData* TargetData = &RT_Proxy->SystemInstancesToProxyData.Add(InstanceID);
-
-		TargetData->NumCells = RT_NumCells;				
+		TargetData->TotalNumAttributes = RT_TotalNumAttributes;
+		TargetData->NumCells = RT_NumCells;		
+		TargetData->NumTiles = RT_InstanceData.NumTiles;
 		TargetData->Precision = RT_Precision;
 		TargetData->ResetValue = RT_ResetValue;
 
@@ -702,17 +788,41 @@ bool UNiagaraDataInterfaceRasterizationGrid3D::PerInstanceTickPostSimulate(void*
 	{
 		InstanceData->NeedsRealloc = false;
 		
+		// Compute number of tiles based on resolution of individual attributes
+		// #todo(dmp): refactor
+		int32 MaxDim = 2048;
+		int32 MaxTilesX = floor(MaxDim / InstanceData->NumCells.X);
+		int32 MaxTilesY = floor(MaxDim / InstanceData->NumCells.Y);
+		int32 MaxTilesZ = floor(MaxDim / InstanceData->NumCells.Z);
+		int32 MaxAttributes = MaxTilesX * MaxTilesY * MaxTilesZ;
+
+		// need to determine number of tiles in x and y based on number of attributes and max dimension size
+		int32 NumTilesX = FMath::Min<int32>(MaxTilesX, InstanceData->TotalNumAttributes);
+		int32 NumTilesY = FMath::Min<int32>(MaxTilesY, ceil(1.0 * InstanceData->TotalNumAttributes / NumTilesX));
+		int32 NumTilesZ = FMath::Min<int32>(MaxTilesZ, ceil(1.0 * InstanceData->TotalNumAttributes / (NumTilesX * NumTilesY)));
+
+		InstanceData->NumTiles.X = NumTilesX;
+		InstanceData->NumTiles.Y = NumTilesY;
+		InstanceData->NumTiles.Z = NumTilesZ;
+
+		check(InstanceData->NumTiles.X > 0);
+		check(InstanceData->NumTiles.Y > 0);
+		check(InstanceData->NumTiles.Z > 0);
+
 		FNiagaraDataInterfaceProxyRasterizationGrid3D* RT_Proxy = GetProxyAs<FNiagaraDataInterfaceProxyRasterizationGrid3D>();
 		ENQUEUE_RENDER_COMMAND(FUpdateData)(
-			[RT_Proxy, RT_NumCells = InstanceData->NumCells, RT_Precision = InstanceData->Precision, RT_ResetValue = InstanceData->ResetValue, InstanceID = SystemInstance->GetId()](FRHICommandListImmediate& RHICmdList)
+			[RT_Proxy, RT_NumCells = InstanceData->NumCells, RT_Precision = InstanceData->Precision, RT_ResetValue = InstanceData->ResetValue, InstanceID = SystemInstance->GetId(), RT_InstanceData = *InstanceData](FRHICommandListImmediate& RHICmdList)
 		{
 			check(RT_Proxy->SystemInstancesToProxyData.Contains(InstanceID));
 			RasterizationGrid3DRWInstanceData* TargetData = &RT_Proxy->SystemInstancesToProxyData.Add(InstanceID);
 
+			TargetData->NumTiles = RT_InstanceData.NumTiles;
 			TargetData->NumCells = RT_NumCells;					
 			TargetData->Precision = RT_Precision;
+			TargetData->TotalNumAttributes = RT_InstanceData.TotalNumAttributes;
 			TargetData->ResetValue = RT_ResetValue;
 			TargetData->ResizeBuffers();
+			TargetData->PerAttributeData.Release();
 		});
 	}
 
@@ -740,9 +850,36 @@ void UNiagaraDataInterfaceRasterizationGrid3D::DestroyPerInstanceData(void* PerI
 
 void FNiagaraDataInterfaceProxyRasterizationGrid3D::PreStage(FRHICommandList& RHICmdList, const FNiagaraDataInterfaceStageArgs& Context)
 {
+	RasterizationGrid3DRWInstanceData* ProxyData = SystemInstancesToProxyData.Find(Context.SystemInstanceID);
+
+	// Generate per-attribute data
+	if (ProxyData->PerAttributeData.NumBytes == 0)
+	{
+		TResourceArray<FVector4f> PerAttributeData;
+		PerAttributeData.AddUninitialized((ProxyData->TotalNumAttributes * 2) + 1);
+		for (int32 iAttribute = 0; iAttribute < ProxyData->TotalNumAttributes; ++iAttribute)
+		{
+			const FIntVector AttributeTileIndex(iAttribute % ProxyData->NumTiles.X, (iAttribute / ProxyData->NumTiles.X) % ProxyData->NumTiles.Y, iAttribute / (ProxyData->NumTiles.X * ProxyData->NumTiles.Y));
+			PerAttributeData[iAttribute] = FVector4f(
+				AttributeTileIndex.X * ProxyData->NumCells.X,
+				AttributeTileIndex.Y * ProxyData->NumCells.Y,
+				AttributeTileIndex.Z * ProxyData->NumCells.Z,
+				0
+			);
+			PerAttributeData[iAttribute + ProxyData->TotalNumAttributes] = FVector4f(
+				(1.0f / ProxyData->NumTiles.X) * float(AttributeTileIndex.X),
+				(1.0f / ProxyData->NumTiles.Y) * float(AttributeTileIndex.Y),
+				(1.0f / ProxyData->NumTiles.Z) * float(AttributeTileIndex.Z),
+				0.0f
+			);
+		}
+		PerAttributeData[ProxyData->TotalNumAttributes * 2] = FVector4f(65535, 65535, 65535, 65535);
+		ProxyData->PerAttributeData.Initialize(TEXT("Grid3D::PerAttributeData"), sizeof(FVector4f), PerAttributeData.Num(), EPixelFormat::PF_A32B32G32R32F, BUF_Static, &PerAttributeData);
+	}
+
 	if (Context.IsOutputStage)
 	{
-		RasterizationGrid3DRWInstanceData* ProxyData = SystemInstancesToProxyData.Find(Context.SystemInstanceID);
+		
 
 		SCOPED_DRAW_EVENT(RHICmdList, NiagaraRasterizationGrid3DClearNeighborInfo);
 		ERHIFeatureLevel::Type FeatureLevel = Context.ComputeDispatchInterface->GetFeatureLevel();
@@ -775,7 +912,7 @@ bool UNiagaraDataInterfaceRasterizationGrid3D::CopyToInternal(UNiagaraDataInterf
 
 	UNiagaraDataInterfaceRasterizationGrid3D* OtherTyped = CastChecked<UNiagaraDataInterfaceRasterizationGrid3D>(Destination);
 
-	
+	OtherTyped->NumAttributes = NumAttributes;
 	OtherTyped->Precision = Precision;
 	OtherTyped->ResetValue = ResetValue;
 
diff --git a/Engine/Plugins/FX/Niagara/Source/NiagaraCore/Private/NiagaraCustomVersion.cpp b/Engine/Plugins/FX/Niagara/Source/NiagaraCore/Private/NiagaraCustomVersion.cpp
--- a/Engine/Plugins/FX/Niagara/Source/NiagaraCore/Private/NiagaraCustomVersion.cpp
+++ b/Engine/Plugins/FX/Niagara/Source/NiagaraCore/Private/NiagaraCustomVersion.cpp
@@ -8,6 +8,4 @@ const FGuid FNiagaraCustomVersion::GUID(0xFCF57AFA, 0x50764283, 0xB9A9E658, 0xFF
 // Register the custom version with core
 FCustomVersionRegistration GRegisterNiagaraCustomVersion(FNiagaraCustomVersion::GUID, FNiagaraCustomVersion::LatestVersion, TEXT("NiagaraVer"));
 
-const FGuid FNiagaraCustomVersion::LatestScriptCompileVersion(0xC71A1029, 0x6E7E894B, 0xB639619F, 0xE083EBF9);
-
-
+const FGuid FNiagaraCustomVersion::LatestScriptCompileVersion(0x14D5F635, 0x89F04AC8, 0xBE4C1F5F, 0x91E1134D);
\ No newline at end of file
