Fix for static analysis error in PVS.
Removed duplicate logically equivalent, but reversed, condition

diff --git a/Engine/Plugins/FX/Niagara/Source/NiagaraEditor/Private/EdGraphSchema_Niagara.cpp b/Engine/Plugins/FX/Niagara/Source/NiagaraEditor/Private/EdGraphSchema_Niagara.cpp
--- a/Engine/Plugins/FX/Niagara/Source/NiagaraEditor/Private/EdGraphSchema_Niagara.cpp
+++ b/Engine/Plugins/FX/Niagara/Source/NiagaraEditor/Private/EdGraphSchema_Niagara.cpp
@@ -1878,9 +1878,6 @@ bool UEdGraphSchema_Niagara::PinTypesValidForNumericConversion(FEdGraphPinType A
 		return true;
 	else if (AType.PinCategory == PinCategoryType && BType.PinCategory == PinCategoryType)
 		return true;
-	else if (BType.PinCategory == PinCategoryType && AType.PinCategory == PinCategoryType)
-		return true;
-
 
 	return false;
 }
