Moved versions image to starshipstyle to use in Niagara

diff --git a/Engine/Plugins/FX/Niagara/Source/NiagaraEditor/Private/Toolkits/NiagaraScriptToolkit.cpp b/Engine/Plugins/FX/Niagara/Source/NiagaraEditor/Private/Toolkits/NiagaraScriptToolkit.cpp
--- a/Engine/Plugins/FX/Niagara/Source/NiagaraEditor/Private/Toolkits/NiagaraScriptToolkit.cpp
+++ b/Engine/Plugins/FX/Niagara/Source/NiagaraEditor/Private/Toolkits/NiagaraScriptToolkit.cpp
@@ -554,7 +554,6 @@ TSharedRef<SDockTab> FNiagaraScriptToolkit::SpawnTabVersioning(const FSpawnTabAr
             ]
         ];
 
-	SpawnedTab->SetTabIcon(TAttribute<const FSlateBrush*>(this, &FNiagaraScriptToolkit::GetVersioningBrush));
 	return SpawnedTab;
 }
 
@@ -651,7 +650,7 @@ void FNiagaraScriptToolkit::ExtendToolbar()
 				ToolbarBuilder.AddToolBarButton(FNiagaraEditorCommands::Get().ModuleVersioning, NAME_None,
                     TAttribute<FText>(ScriptToolkit, &FNiagaraScriptToolkit::GetVersionButtonLabel),
                     LOCTEXT("NiagaraShowModuleVersionsTooltip", "Manage different versions of this module script."),
-                    TAttribute<FSlateIcon>(ScriptToolkit, &FNiagaraScriptToolkit::GetVersioningImage));
+                    FSlateIcon(FAppStyle::GetAppStyleSetName(), "Versions"));
 
 				FUIAction DropdownAction;
 				DropdownAction.IsActionVisibleDelegate = FIsActionButtonVisible::CreateLambda([ScriptToolkit]() { return ScriptToolkit->EditedNiagaraScript.Script->GetAllAvailableVersions().Num() > 1; });
@@ -739,16 +738,6 @@ FText FNiagaraScriptToolkit::GetRefreshStatusTooltip() const
 	return LOCTEXT("Refresh_Status", "Currently dependencies up-to-date. Consider refreshing if status isn't accurate.");
 }
 
-FSlateIcon FNiagaraScriptToolkit::GetVersioningImage() const
-{
-	if (EditedNiagaraScript.Script && EditedNiagaraScript.Script->IsVersioningEnabled())
-	{
-		return FSlateIcon(FAppStyle::GetAppStyleSetName(), "SourceControl.StatusIcon.On");
-	}
-	
-	return FSlateIcon(FAppStyle::GetAppStyleSetName(), "SourceControl.StatusIcon.Off");
-}
-
 FText FNiagaraScriptToolkit::GetVersionButtonLabel() const
 {
 	FText BaseLabel = LOCTEXT("NiagaraShowModuleVersions", "Versioning");
@@ -763,16 +752,6 @@ FText FNiagaraScriptToolkit::GetVersionButtonLabel() const
 	return BaseLabel;
 }
 
-const FSlateBrush* FNiagaraScriptToolkit::GetVersioningBrush() const
-{
-	if (EditedNiagaraScript.Script && EditedNiagaraScript.Script->IsVersioningEnabled())
-	{
-		return FAppStyle::Get().GetBrush("SourceControl.StatusIcon.On");
-	}
-	
-	return FAppStyle::Get().GetBrush("SourceControl.StatusIcon.Off");
-}
-
 void FNiagaraScriptToolkit::CompileScript(bool bForce)
 {
 	ScriptViewModel->CompileStandaloneScript();
diff --git a/Engine/Plugins/FX/Niagara/Source/NiagaraEditor/Private/Toolkits/NiagaraScriptToolkit.h b/Engine/Plugins/FX/Niagara/Source/NiagaraEditor/Private/Toolkits/NiagaraScriptToolkit.h
--- a/Engine/Plugins/FX/Niagara/Source/NiagaraEditor/Private/Toolkits/NiagaraScriptToolkit.h
+++ b/Engine/Plugins/FX/Niagara/Source/NiagaraEditor/Private/Toolkits/NiagaraScriptToolkit.h
@@ -141,9 +141,7 @@ private:
 
 	FSlateIcon GetRefreshStatusImage() const;
 	FText GetRefreshStatusTooltip() const;
-	FSlateIcon GetVersioningImage() const;
 	FText GetVersionButtonLabel() const;
-	const FSlateBrush* GetVersioningBrush() const;
 
 	bool IsEditScriptDifferentFromOriginalScript() const;
 
diff --git a/Engine/Source/Editor/EditorStyle/Private/StarshipStyle.cpp b/Engine/Source/Editor/EditorStyle/Private/StarshipStyle.cpp
--- a/Engine/Source/Editor/EditorStyle/Private/StarshipStyle.cpp
+++ b/Engine/Source/Editor/EditorStyle/Private/StarshipStyle.cpp
@@ -2587,6 +2587,7 @@ void FStarshipEditorStyle::FStyle::SetupGeneralIcons()
 	
 	Set("Debug", new IMAGE_BRUSH_SVG( "Starship/Common/Debug", Icon16x16 ) );
 	Set("Modules", new IMAGE_BRUSH_SVG( "Starship/Common/Modules", Icon16x16 ) );
+	Set("Versions", new IMAGE_BRUSH_SVG("Starship/Common/Versions", Icon20x20));
 }
 
 void FStarshipEditorStyle::FStyle::SetupWindowStyles()
