- All position types return a position

diff --git a/Engine/Plugins/FX/Niagara/Source/Niagara/Private/DataInterface/NiagaraDataInterfaceStaticMesh.cpp b/Engine/Plugins/FX/Niagara/Source/Niagara/Private/DataInterface/NiagaraDataInterfaceStaticMesh.cpp
--- a/Engine/Plugins/FX/Niagara/Source/Niagara/Private/DataInterface/NiagaraDataInterfaceStaticMesh.cpp
+++ b/Engine/Plugins/FX/Niagara/Source/Niagara/Private/DataInterface/NiagaraDataInterfaceStaticMesh.cpp
@@ -57,6 +57,7 @@ namespace NDIStaticMeshLocal
 			InitialVersion = 0,
 			RefactoredV1 = 1,
 			LargeWorldCoordinates = 2,
+			LargeWorldCoordinates2 = 2,
 
 			VersionPlusOne,
 			LatestVersion = VersionPlusOne - 1
@@ -1597,7 +1598,7 @@ void UNiagaraDataInterfaceStaticMesh::GetFunctions(TArray<FNiagaraFunctionSignat
 	{
 		FNiagaraFunctionSignature Sig = BaseSignature;
 		Sig.Inputs.Emplace(FNiagaraTypeDefinition::GetIntDef(), TEXT("Vertex"));
-		Sig.Outputs.Emplace(FNiagaraTypeDefinition::GetVec3Def(), TEXT("Position"));
+		Sig.Outputs.Emplace(FNiagaraTypeDefinition::GetPositionDef(), TEXT("Position"));
 		Sig.Outputs.Emplace(FNiagaraTypeDefinition::GetVec3Def(), TEXT("Velocity"));
 		Sig.Outputs.Emplace(FNiagaraTypeDefinition::GetVec3Def(), TEXT("Normal"));
 		Sig.Outputs.Emplace(FNiagaraTypeDefinition::GetVec3Def(), TEXT("Bitangent"));
@@ -1697,7 +1698,7 @@ void UNiagaraDataInterfaceStaticMesh::GetFunctions(TArray<FNiagaraFunctionSignat
 		FNiagaraFunctionSignature Sig = BaseSignature;
 		Sig.Inputs.Emplace(FNiagaraTypeDefinition::GetIntDef(), TEXT("Triangle"));
 		Sig.Inputs.Emplace_GetRef(FNiagaraVariable(FNiagaraTypeDefinition::GetVec3Def(), TEXT("BaryCoord"))).SetValue(FVector3f(1.0f / 3.0f));
-		Sig.Outputs.Emplace(FNiagaraTypeDefinition::GetVec3Def(), TEXT("Position"));
+		Sig.Outputs.Emplace(FNiagaraTypeDefinition::GetPositionDef(), TEXT("Position"));
 		Sig.Outputs.Emplace(FNiagaraTypeDefinition::GetVec3Def(), TEXT("Velocity"));
 		Sig.Outputs.Emplace(FNiagaraTypeDefinition::GetVec3Def(), TEXT("Normal"));
 		Sig.Outputs.Emplace(FNiagaraTypeDefinition::GetVec3Def(), TEXT("Bitangent"));
@@ -1746,7 +1747,7 @@ void UNiagaraDataInterfaceStaticMesh::GetFunctions(TArray<FNiagaraFunctionSignat
 	{
 		FNiagaraFunctionSignature Sig = BaseSignature;
 		Sig.Inputs.Add(FNiagaraVariable(FNiagaraTypeDefinition::GetIntDef(), TEXT("Index")));
-		Sig.Outputs.Add(FNiagaraVariable(FNiagaraTypeDefinition::GetVec3Def(), TEXT("Position")));
+		Sig.Outputs.Add(FNiagaraVariable(FNiagaraTypeDefinition::GetPositionDef(), TEXT("Position")));
 		Sig.Outputs.Add(FNiagaraVariable(FNiagaraTypeDefinition::GetQuatDef(), TEXT("Rotation")));
 		Sig.Outputs.Add(FNiagaraVariable(FNiagaraTypeDefinition::GetVec3Def(), TEXT("Scale")));
 
@@ -1965,10 +1966,12 @@ bool UNiagaraDataInterfaceStaticMesh::UpgradeFunctionCall(FNiagaraFunctionSignat
 		}
 	}
 
-	if (FunctionSignature.FunctionVersion < EDIFunctionVersion::LargeWorldCoordinates)
+	if (FunctionSignature.FunctionVersion < EDIFunctionVersion::LargeWorldCoordinates2)
 	{
 		if (
+			(FunctionSignature.Name == GetVertexName) ||
 			(FunctionSignature.Name == GetVertexWSName) ||
+			(FunctionSignature.Name == GetTriangleName) ||
 			(FunctionSignature.Name == GetTriangleWSName) ||
 			(FunctionSignature.Name == GetSocketTransformName) ||
 			(FunctionSignature.Name == GetSocketTransformWSName) ||
@@ -1978,7 +1981,7 @@ bool UNiagaraDataInterfaceStaticMesh::UpgradeFunctionCall(FNiagaraFunctionSignat
 			(FunctionSignature.Name == GetUnfilteredSocketTransformWSName) )
 		{
 			check(FunctionSignature.Outputs[0].GetName() == TEXT("Position"));
-			check(FunctionSignature.Outputs[0].GetType() == FNiagaraTypeDefinition::GetVec3Def());
+			check(FunctionSignature.Outputs[0].GetType() == FNiagaraTypeDefinition::GetVec3Def() || FunctionSignature.Outputs[0].GetType() == FNiagaraTypeDefinition::GetPositionDef());
 			FunctionSignature.Outputs[0].SetType(FNiagaraTypeDefinition::GetPositionDef());
 		}
 	}
diff --git a/Engine/Plugins/FX/Niagara/Source/Niagara/Private/NDISkeletalMesh_BoneSampling.cpp b/Engine/Plugins/FX/Niagara/Source/Niagara/Private/NDISkeletalMesh_BoneSampling.cpp
--- a/Engine/Plugins/FX/Niagara/Source/Niagara/Private/NDISkeletalMesh_BoneSampling.cpp
+++ b/Engine/Plugins/FX/Niagara/Source/Niagara/Private/NDISkeletalMesh_BoneSampling.cpp
@@ -54,7 +54,7 @@ void UNiagaraDataInterfaceSkeletalMesh::GetSkeletonSamplingFunctions(TArray<FNia
 		Sig.Name = FSkeletalMeshInterfaceHelper::GetSkinnedBoneDataName;
 		Sig.Inputs.Add(FNiagaraVariable(FNiagaraTypeDefinition(GetClass()), TEXT("SkeletalMesh")));
 		Sig.Inputs.Add(FNiagaraVariable(FNiagaraTypeDefinition::GetIntDef(), TEXT("Bone")));
-		Sig.Outputs.Add(FNiagaraVariable(FNiagaraTypeDefinition::GetVec3Def(), TEXT("Position")));
+		Sig.Outputs.Add(FNiagaraVariable(FNiagaraTypeDefinition::GetPositionDef(), TEXT("Position")));
 		Sig.Outputs.Add(FNiagaraVariable(FNiagaraTypeDefinition::GetQuatDef(), TEXT("Rotation")));
 		Sig.Outputs.Add(FNiagaraVariable(FNiagaraTypeDefinition::GetVec3Def(), TEXT("Velocity")));
 		Sig.bMemberFunction = true;
@@ -88,7 +88,7 @@ void UNiagaraDataInterfaceSkeletalMesh::GetSkeletonSamplingFunctions(TArray<FNia
 		Sig.Inputs.Add(FNiagaraVariable(FNiagaraTypeDefinition(GetClass()), TEXT("SkeletalMesh")));
 		Sig.Inputs.Add(FNiagaraVariable(FNiagaraTypeDefinition::GetIntDef(), TEXT("Bone")));
 		Sig.Inputs.Add(FNiagaraVariable(FNiagaraTypeDefinition::GetFloatDef(), TEXT("Interpolation")));
-		Sig.Outputs.Add(FNiagaraVariable(FNiagaraTypeDefinition::GetVec3Def(), TEXT("Position")));
+		Sig.Outputs.Add(FNiagaraVariable(FNiagaraTypeDefinition::GetPositionDef(), TEXT("Position")));
 		Sig.Outputs.Add(FNiagaraVariable(FNiagaraTypeDefinition::GetQuatDef(), TEXT("Rotation")));
 		Sig.Outputs.Add(FNiagaraVariable(FNiagaraTypeDefinition::GetVec3Def(), TEXT("Velocity")));
 		Sig.bMemberFunction = true;
diff --git a/Engine/Plugins/FX/Niagara/Source/Niagara/Private/NDISkeletalMesh_TriangleSampling.cpp b/Engine/Plugins/FX/Niagara/Source/Niagara/Private/NDISkeletalMesh_TriangleSampling.cpp
--- a/Engine/Plugins/FX/Niagara/Source/Niagara/Private/NDISkeletalMesh_TriangleSampling.cpp
+++ b/Engine/Plugins/FX/Niagara/Source/Niagara/Private/NDISkeletalMesh_TriangleSampling.cpp
@@ -80,7 +80,7 @@ void UNiagaraDataInterfaceSkeletalMesh::GetTriangleSamplingFunctions(TArray<FNia
 		Sig.Name = FSkeletalMeshInterfaceHelper::GetTriangleDataName;
 		Sig.Inputs.Add(FNiagaraVariable(FNiagaraTypeDefinition(GetClass()), TEXT("SkeletalMesh")));
 		Sig.Inputs.Add(FNiagaraVariable(FNiagaraTypeDefinition(FMeshTriCoordinate::StaticStruct()), TEXT("Coord")));
-		Sig.Outputs.Add(FNiagaraVariable(FNiagaraTypeDefinition::GetVec3Def(), TEXT("Position")));
+		Sig.Outputs.Add(FNiagaraVariable(FNiagaraTypeDefinition::GetPositionDef(), TEXT("Position")));
 		Sig.Outputs.Add(FNiagaraVariable(FNiagaraTypeDefinition::GetVec3Def(), TEXT("Normal")));
 		Sig.Outputs.Add(FNiagaraVariable(FNiagaraTypeDefinition::GetVec3Def(), TEXT("Binormal")));
 		Sig.Outputs.Add(FNiagaraVariable(FNiagaraTypeDefinition::GetVec3Def(), TEXT("Tangent")));
@@ -96,7 +96,7 @@ void UNiagaraDataInterfaceSkeletalMesh::GetTriangleSamplingFunctions(TArray<FNia
 		Sig.Name = FSkeletalMeshInterfaceHelper::GetSkinnedTriangleDataName;
 		Sig.Inputs.Add(FNiagaraVariable(FNiagaraTypeDefinition(GetClass()), TEXT("SkeletalMesh")));
 		Sig.Inputs.Add(FNiagaraVariable(FNiagaraTypeDefinition(FMeshTriCoordinate::StaticStruct()), TEXT("Coord")));
-		Sig.Outputs.Add(FNiagaraVariable(FNiagaraTypeDefinition::GetVec3Def(), TEXT("Position")));
+		Sig.Outputs.Add(FNiagaraVariable(FNiagaraTypeDefinition::GetPositionDef(), TEXT("Position")));
 		Sig.Outputs.Add(FNiagaraVariable(FNiagaraTypeDefinition::GetVec3Def(), TEXT("Velocity")));
 		Sig.Outputs.Add(FNiagaraVariable(FNiagaraTypeDefinition::GetVec3Def(), TEXT("Normal")));
 		Sig.Outputs.Add(FNiagaraVariable(FNiagaraTypeDefinition::GetVec3Def(), TEXT("Binormal")));
@@ -131,7 +131,7 @@ void UNiagaraDataInterfaceSkeletalMesh::GetTriangleSamplingFunctions(TArray<FNia
 		Sig.Inputs.Add(FNiagaraVariable(FNiagaraTypeDefinition(GetClass()), TEXT("SkeletalMesh")));
 		Sig.Inputs.Add(FNiagaraVariable(FNiagaraTypeDefinition(FMeshTriCoordinate::StaticStruct()), TEXT("Coord")));
 		Sig.Inputs.Add(FNiagaraVariable(FNiagaraTypeDefinition::GetFloatDef(), TEXT("Interp")));
-		Sig.Outputs.Add(FNiagaraVariable(FNiagaraTypeDefinition::GetVec3Def(), TEXT("Position")));
+		Sig.Outputs.Add(FNiagaraVariable(FNiagaraTypeDefinition::GetPositionDef(), TEXT("Position")));
 		Sig.Outputs.Add(FNiagaraVariable(FNiagaraTypeDefinition::GetVec3Def(), TEXT("Velocity")));
 		Sig.Outputs.Add(FNiagaraVariable(FNiagaraTypeDefinition::GetVec3Def(), TEXT("Normal")));
 		Sig.Outputs.Add(FNiagaraVariable(FNiagaraTypeDefinition::GetVec3Def(), TEXT("Binormal")));
diff --git a/Engine/Plugins/FX/Niagara/Source/Niagara/Private/NDISkeletalMesh_VertexSampling.cpp b/Engine/Plugins/FX/Niagara/Source/Niagara/Private/NDISkeletalMesh_VertexSampling.cpp
--- a/Engine/Plugins/FX/Niagara/Source/Niagara/Private/NDISkeletalMesh_VertexSampling.cpp
+++ b/Engine/Plugins/FX/Niagara/Source/Niagara/Private/NDISkeletalMesh_VertexSampling.cpp
@@ -47,7 +47,7 @@ void UNiagaraDataInterfaceSkeletalMesh::GetVertexSamplingFunctions(TArray<FNiaga
 		Sig.Name = FSkeletalMeshInterfaceHelper::GetVertexDataName;
 		Sig.Inputs.Add(FNiagaraVariable(FNiagaraTypeDefinition(GetClass()), TEXT("SkeletalMesh")));
 		Sig.Inputs.Add(FNiagaraVariable(FNiagaraTypeDefinition::GetIntDef(), TEXT("Vertex")));
-		Sig.Outputs.Add(FNiagaraVariable(FNiagaraTypeDefinition::GetVec3Def(), TEXT("Position")));
+		Sig.Outputs.Add(FNiagaraVariable(FNiagaraTypeDefinition::GetPositionDef(), TEXT("Position")));
 		Sig.Outputs.Add(FNiagaraVariable(FNiagaraTypeDefinition::GetVec3Def(), TEXT("Normal")));
 		Sig.Outputs.Add(FNiagaraVariable(FNiagaraTypeDefinition::GetVec3Def(), TEXT("Binormal")));
 		Sig.Outputs.Add(FNiagaraVariable(FNiagaraTypeDefinition::GetVec3Def(), TEXT("Tangent")));
@@ -63,7 +63,7 @@ void UNiagaraDataInterfaceSkeletalMesh::GetVertexSamplingFunctions(TArray<FNiaga
 		Sig.Name = FSkeletalMeshInterfaceHelper::GetSkinnedVertexDataName;
 		Sig.Inputs.Add(FNiagaraVariable(FNiagaraTypeDefinition(GetClass()), TEXT("SkeletalMesh")));
 		Sig.Inputs.Add(FNiagaraVariable(FNiagaraTypeDefinition::GetIntDef(), TEXT("Vertex")));
-		Sig.Outputs.Add(FNiagaraVariable(FNiagaraTypeDefinition::GetVec3Def(), TEXT("Position")));
+		Sig.Outputs.Add(FNiagaraVariable(FNiagaraTypeDefinition::GetPositionDef(), TEXT("Position")));
 		Sig.Outputs.Add(FNiagaraVariable(FNiagaraTypeDefinition::GetVec3Def(), TEXT("Velocity")));
 		Sig.Outputs.Add(FNiagaraVariable(FNiagaraTypeDefinition::GetVec3Def(), TEXT("Normal")));
 		Sig.Outputs.Add(FNiagaraVariable(FNiagaraTypeDefinition::GetVec3Def(), TEXT("Binormal")));
diff --git a/Engine/Plugins/FX/Niagara/Source/Niagara/Private/NiagaraDataInterfaceSkeletalMesh.cpp b/Engine/Plugins/FX/Niagara/Source/Niagara/Private/NiagaraDataInterfaceSkeletalMesh.cpp
--- a/Engine/Plugins/FX/Niagara/Source/Niagara/Private/NiagaraDataInterfaceSkeletalMesh.cpp
+++ b/Engine/Plugins/FX/Niagara/Source/Niagara/Private/NiagaraDataInterfaceSkeletalMesh.cpp
@@ -41,6 +41,7 @@ struct FNiagaraSkelMeshDIFunctionVersion
 		AddedValidConnectivity = 7,
 		AddedInputBardCoordToGetFilteredTriangleAt = 8,
 		LargeWorldCoordinates = 9,
+		LargeWorldCoordinates2 = 9,
 
 		VersionPlusOne,
 		LatestVersion = VersionPlusOne - 1
@@ -3469,17 +3470,24 @@ bool UNiagaraDataInterfaceSkeletalMesh::UpgradeFunctionCall(FNiagaraFunctionSign
 		}
 	}
 
-	if (FunctionSignature.FunctionVersion < FNiagaraSkelMeshDIFunctionVersion::LargeWorldCoordinates)
+	if (FunctionSignature.FunctionVersion < FNiagaraSkelMeshDIFunctionVersion::LargeWorldCoordinates2)
 	{
 		if (
-			(FunctionSignature.Name == FSkeletalMeshInterfaceHelper::GetSkinnedTriangleDataWSInterpName) ||
+			(FunctionSignature.Name == FSkeletalMeshInterfaceHelper::GetTriangleDataName) ||
+			(FunctionSignature.Name == FSkeletalMeshInterfaceHelper::GetSkinnedTriangleDataName) ||
 			(FunctionSignature.Name == FSkeletalMeshInterfaceHelper::GetSkinnedTriangleDataWSName) ||
-			(FunctionSignature.Name == FSkeletalMeshInterfaceHelper::GetSkinnedVertexDataWSName) ||
+			(FunctionSignature.Name == FSkeletalMeshInterfaceHelper::GetSkinnedTriangleDataInterpName) ||
+			(FunctionSignature.Name == FSkeletalMeshInterfaceHelper::GetSkinnedTriangleDataWSInterpName) ||
+			(FunctionSignature.Name == FSkeletalMeshInterfaceHelper::GetSkinnedBoneDataName) ||
 			(FunctionSignature.Name == FSkeletalMeshInterfaceHelper::GetSkinnedBoneDataWSName) ||
-			(FunctionSignature.Name == FSkeletalMeshInterfaceHelper::GetSkinnedBoneDataWSInterpolatedName))
+			(FunctionSignature.Name == FSkeletalMeshInterfaceHelper::GetSkinnedBoneDataInterpolatedName) ||
+			(FunctionSignature.Name == FSkeletalMeshInterfaceHelper::GetSkinnedBoneDataWSInterpolatedName) ||
+			(FunctionSignature.Name == FSkeletalMeshInterfaceHelper::GetVertexDataName) ||
+			(FunctionSignature.Name == FSkeletalMeshInterfaceHelper::GetSkinnedVertexDataName) ||
+			(FunctionSignature.Name == FSkeletalMeshInterfaceHelper::GetSkinnedVertexDataWSName) )
 		{
 			check(FunctionSignature.Outputs[0].GetName() == TEXT("Position"));
-			check(FunctionSignature.Outputs[0].GetType() == FNiagaraTypeDefinition::GetVec3Def());
+			check(FunctionSignature.Outputs[0].GetType() == FNiagaraTypeDefinition::GetVec3Def() || FunctionSignature.Outputs[0].GetType() == FNiagaraTypeDefinition::GetPositionDef());
 			FunctionSignature.Outputs[0].SetType(FNiagaraTypeDefinition::GetPositionDef());
 			bWasChanged = true;
 		}
