Updating the Niagara Stack Style to correctly show selection

diff --git a/Engine/Plugins/FX/Niagara/Source/NiagaraEditorWidgets/Private/NiagaraEditorWidgetsStyle.cpp b/Engine/Plugins/FX/Niagara/Source/NiagaraEditorWidgets/Private/NiagaraEditorWidgetsStyle.cpp
--- a/Engine/Plugins/FX/Niagara/Source/NiagaraEditorWidgets/Private/NiagaraEditorWidgetsStyle.cpp
+++ b/Engine/Plugins/FX/Niagara/Source/NiagaraEditorWidgets/Private/NiagaraEditorWidgetsStyle.cpp
@@ -298,10 +298,14 @@ FNiagaraEditorWidgetsStyle::FNiagaraEditorWidgetsStyle() : FSlateStyleSet("Niaga
 	FSlateBoxBrush StackRowSelectionBrush = BOX_BRUSH("Icons/StackSelectionBorder", FMargin(2.0f / 8.0f), FStyleColors::Select);
 	FSlateBoxBrush StackRowSubduedSelectionBrush = BOX_BRUSH("Icons/StackSelectionBorder", FMargin(2.0f / 8.0f), FStyleColors::SelectInactive);
 	Set("NiagaraEditor.Stack.TableViewRow", FTableRowStyle(NormalTableRowStyle)
-		.SetEvenRowBackgroundBrush(FSlateColorBrush(FStyleColors::White))
-		.SetOddRowBackgroundBrush(FSlateColorBrush(FStyleColors::White))
+		.SetUseParentRowBrush(true)
+ 		.SetEvenRowBackgroundBrush(FSlateColorBrush(FStyleColors::Transparent))
+		.SetOddRowBackgroundBrush(FSlateColorBrush(FStyleColors::Transparent))
 		.SetActiveBrush(StackRowSelectionBrush)
 		.SetActiveHoveredBrush(StackRowSelectionBrush)
+		.SetParentRowBackgroundHoveredBrush(FSlateColorBrush(FStyleColors::Header))
+		.SetOddRowBackgroundHoveredBrush(FSlateColorBrush(FStyleColors::Header))
+		.SetEvenRowBackgroundHoveredBrush(FSlateColorBrush(FStyleColors::Header))
 		.SetInactiveBrush(StackRowSubduedSelectionBrush)
 		.SetInactiveHoveredBrush(StackRowSelectionBrush));
 
diff --git a/Engine/Plugins/FX/Niagara/Source/NiagaraEditorWidgets/Private/Stack/SNiagaraStackTableRow.cpp b/Engine/Plugins/FX/Niagara/Source/NiagaraEditorWidgets/Private/Stack/SNiagaraStackTableRow.cpp
--- a/Engine/Plugins/FX/Niagara/Source/NiagaraEditorWidgets/Private/Stack/SNiagaraStackTableRow.cpp
+++ b/Engine/Plugins/FX/Niagara/Source/NiagaraEditorWidgets/Private/Stack/SNiagaraStackTableRow.cpp
@@ -354,7 +354,12 @@ void SNiagaraStackTableRow::SetNameAndValueContent(TSharedRef<SWidget> InNameWid
  								.BorderBackgroundColor(FStyleColors::Select)
 								.Padding(0)
 								[
-									ChildContent.ToSharedRef()
+									SNew(SBorder)
+									.BorderImage(this, &SNiagaraStackTableRow::GetSelectionBorderBrush)
+									.Padding(0)
+									[
+										ChildContent.ToSharedRef()
+									]
 								]
 							]
 						]
