- Fix incorrect version numbers

#rnx

diff --git a/Engine/Plugins/FX/Niagara/Source/Niagara/Private/DataInterface/NiagaraDataInterfaceStaticMesh.cpp b/Engine/Plugins/FX/Niagara/Source/Niagara/Private/DataInterface/NiagaraDataInterfaceStaticMesh.cpp
--- a/Engine/Plugins/FX/Niagara/Source/Niagara/Private/DataInterface/NiagaraDataInterfaceStaticMesh.cpp
+++ b/Engine/Plugins/FX/Niagara/Source/Niagara/Private/DataInterface/NiagaraDataInterfaceStaticMesh.cpp
@@ -57,7 +57,7 @@ namespace NDIStaticMeshLocal
 			InitialVersion = 0,
 			RefactoredV1 = 1,
 			LargeWorldCoordinates = 2,
-			LargeWorldCoordinates2 = 2,
+			LargeWorldCoordinates2 = 3,
 
 			VersionPlusOne,
 			LatestVersion = VersionPlusOne - 1
diff --git a/Engine/Plugins/FX/Niagara/Source/Niagara/Private/NiagaraDataInterfaceSkeletalMesh.cpp b/Engine/Plugins/FX/Niagara/Source/Niagara/Private/NiagaraDataInterfaceSkeletalMesh.cpp
--- a/Engine/Plugins/FX/Niagara/Source/Niagara/Private/NiagaraDataInterfaceSkeletalMesh.cpp
+++ b/Engine/Plugins/FX/Niagara/Source/Niagara/Private/NiagaraDataInterfaceSkeletalMesh.cpp
@@ -41,7 +41,7 @@ struct FNiagaraSkelMeshDIFunctionVersion
 		AddedValidConnectivity = 7,
 		AddedInputBardCoordToGetFilteredTriangleAt = 8,
 		LargeWorldCoordinates = 9,
-		LargeWorldCoordinates2 = 9,
+		LargeWorldCoordinates2 = 10,
 
 		VersionPlusOne,
 		LatestVersion = VersionPlusOne - 1
