- Adjust check to give more details

#rnx

diff --git a/Engine/Plugins/FX/Niagara/Source/NiagaraEditor/Private/NiagaraHlslTranslator.cpp b/Engine/Plugins/FX/Niagara/Source/NiagaraEditor/Private/NiagaraHlslTranslator.cpp
--- a/Engine/Plugins/FX/Niagara/Source/NiagaraEditor/Private/NiagaraHlslTranslator.cpp
+++ b/Engine/Plugins/FX/Niagara/Source/NiagaraEditor/Private/NiagaraHlslTranslator.cpp
@@ -5370,7 +5370,7 @@ bool FHlslNiagaraTranslator::ParameterMapRegisterExternalConstantNamespaceVariab
 					else
 					{
 						DataInterface = CompileDuplicateData->GetDuplicatedDataInterfaceCDOForClass(const_cast<UClass*>(InVariable.GetType().GetClass()));
-						check(DataInterface != nullptr);
+						checkf(DataInterface != nullptr, TEXT("GetDuplicatedDataInterfaceCDOForClass failed for Variable(%s) Class(%s)"), *InVariable.GetName().ToString(), *InVariable.GetType().GetClass()->GetName());
 					}
 					if (ensure(DataInterface))
 					{
