Niagara - Overview stack UI overhaul.
+ Fix style issues from initial UE5 reskin.
+ Remove the "Emitter Settings" group and move Emitter Properties to it's own top level group with no sub items.
+ Remove the redundant sub item from "Emitter Summary"
+ Remove the "System Settings" group and move System Properties and User Parameters to their own top level groups with no sub items.
+ Add an add button to the "User Parameters" so that you can add parameters directly from the selection stack when user parameters are selected.
+ Remove the rows for "Add Event Handler" and "Add Simulation Stage" and add a button to the emitter properties row for adding additional stages.
+ Remove colored icons from the row and add a colored highlight to the sub items.
+ Create more compact version of the issue markers, and don't show issue markers on group headers unless the group has a direct issue, or it's collapsed and it's items have issues.

diff --git a/Engine/Plugins/FX/Niagara/Source/NiagaraEditor/Private/ViewModels/NiagaraSystemSelectionViewModel.cpp b/Engine/Plugins/FX/Niagara/Source/NiagaraEditor/Private/ViewModels/NiagaraSystemSelectionViewModel.cpp
--- a/Engine/Plugins/FX/Niagara/Source/NiagaraEditor/Private/ViewModels/NiagaraSystemSelectionViewModel.cpp
+++ b/Engine/Plugins/FX/Niagara/Source/NiagaraEditor/Private/ViewModels/NiagaraSystemSelectionViewModel.cpp
@@ -219,6 +219,11 @@ void UNiagaraSystemSelectionViewModel::AddEntryToSelectionByDisplayedObjectKeyDe
 	DeferredDisplayedObjectKeysToAddToSelection.Add(InObjectKey);
 }
 
+void UNiagaraSystemSelectionViewModel::AddEntryToSelectionBySelectionIdDeferred(const FGuid& InSelectionId)
+{
+	DeferredSelectionIdsToAddToSelection.Add(InSelectionId);
+}
+
 void CollectGroupAndItemEntries(TSharedRef<FNiagaraSystemViewModel> SystemViewModel, TMap<FGuid, TMap<FString, UNiagaraStackEntry*>>& OutEmitterGuidToEntryKeyToEntryMap)
 {
 	auto CollectFromStackViewModel = [](FGuid EmitterGuid, UNiagaraStackViewModel* StackViewModel, TMap<FGuid, TMap<FString, UNiagaraStackEntry*>>& OutEmitterGuidToEntryKeyToEntryMap)
@@ -313,37 +318,39 @@ UNiagaraSystemSelectionViewModel::FOnSelectionChanged& UNiagaraSystemSelectionVi
 	return OnSystemIsSelectedChangedDelegate;
 }
 
-void GatherAllObjectKeysForEntry(UNiagaraStackEntry& StackEntry, TArray<FObjectKey>& OutObjectKeys)
+template<typename TEntryKey, typename TTryGetKeyFromEntry>
+void GatherAllKeysForEntry(UNiagaraStackEntry* StackEntry, TTryGetKeyFromEntry TryGetKeyFromEntry, TArray<TEntryKey>& OutEntryKeys)
 {
-	UObject* DisplayedObject = StackEntry.GetDisplayedObject();
-	if (DisplayedObject != nullptr)
+	TEntryKey EntryKey;
+	if (TryGetKeyFromEntry(StackEntry, EntryKey))
 	{
-		OutObjectKeys.Add(FObjectKey(DisplayedObject));
+		OutEntryKeys.Add(EntryKey);
 	}
 
 	TArray<UNiagaraStackEntry*> Children;
-	StackEntry.GetUnfilteredChildren(Children);
+	StackEntry->GetUnfilteredChildren(Children);
 	for (UNiagaraStackEntry* ChildEntry : Children)
 	{
-		GatherAllObjectKeysForEntry(*ChildEntry, OutObjectKeys);
+		GatherAllKeysForEntry(ChildEntry, TryGetKeyFromEntry, OutEntryKeys);
 	}
 }
 
-void FindStackGroupsAndItemsForDisplayedObjectKeysRecursive(UNiagaraStackEntry* StackEntry, const TArray<FObjectKey>& ObjectKeys, TArray<UNiagaraStackEntry*>& OutFoundStackEntries)
+template<typename TEntryKey, typename TTryGetKeyFromEntry>
+void FindStackGroupsAndItemsForEntryKeysRecursive(UNiagaraStackEntry* StackEntry, const TArray<TEntryKey>& EntryKeysToFind, TTryGetKeyFromEntry TryGetKeyFromEntry, TArray<UNiagaraStackEntry*>& OutFoundStackEntries)
 {
 	if (StackEntry->IsA<UNiagaraStackItemGroup>())
 	{
-		UObject* DisplayedObject = StackEntry->GetDisplayedObject();
-		if (DisplayedObject != nullptr && ObjectKeys.Contains(FObjectKey(DisplayedObject)))
+		TEntryKey EntryKey;
+		if (TryGetKeyFromEntry(StackEntry, EntryKey) && EntryKeysToFind.Contains(EntryKey))
 		{
 			OutFoundStackEntries.Add(StackEntry);
 		}
 	}
 	else if (StackEntry->IsA<UNiagaraStackItem>())
 	{
-		TArray<FObjectKey> ItemObjectKeys;
-		GatherAllObjectKeysForEntry(*StackEntry, ItemObjectKeys);
-		if (ObjectKeys.ContainsByPredicate([&ItemObjectKeys](const FObjectKey& ObjectKey) { return ItemObjectKeys.Contains(ObjectKey); }))
+		TArray<TEntryKey> ItemEntryKeys;
+		GatherAllKeysForEntry(StackEntry, TryGetKeyFromEntry, ItemEntryKeys);
+		if (EntryKeysToFind.ContainsByPredicate([&ItemEntryKeys](const TEntryKey& EntryKey) { return ItemEntryKeys.Contains(EntryKey); }))
 		{
 			OutFoundStackEntries.Add(StackEntry);
 		}
@@ -355,7 +362,7 @@ void FindStackGroupsAndItemsForDisplayedObjectKeysRecursive(UNiagaraStackEntry*
 		StackEntry->GetUnfilteredChildren(Children);
 		for (UNiagaraStackEntry* ChildEntry : Children)
 		{
-			FindStackGroupsAndItemsForDisplayedObjectKeysRecursive(ChildEntry, ObjectKeys, OutFoundStackEntries);
+			FindStackGroupsAndItemsForEntryKeysRecursive(ChildEntry, EntryKeysToFind, TryGetKeyFromEntry, OutFoundStackEntries);
 		}
 	}
 }
@@ -364,26 +371,61 @@ void UNiagaraSystemSelectionViewModel::Tick()
 {
 	bool bSelectionChanged = false;
 
+	TArray<UNiagaraStackEntry*> FoundStackEntries;
+	TSharedRef<FNiagaraSystemViewModel> SystemViewModel = GetSystemViewModel();
+
 	if(DeferredDisplayedObjectKeysToAddToSelection.Num() > 0)
 	{
-		TArray<UNiagaraStackEntry*> FoundStackEntries;
-		TSharedRef<FNiagaraSystemViewModel> SystemViewModel = GetSystemViewModel();
-		FindStackGroupsAndItemsForDisplayedObjectKeysRecursive(SystemViewModel->GetSystemStackViewModel()->GetRootEntry(), DeferredDisplayedObjectKeysToAddToSelection, FoundStackEntries);
+		auto TryGetObjectKeyFromStackEntry = [](const UNiagaraStackEntry* StackEntry, FObjectKey& OutObjectKey)
+		{
+			if (StackEntry->GetDisplayedObject() != nullptr)
+			{
+				OutObjectKey = FObjectKey(StackEntry->GetDisplayedObject());
+				return true;
+			}
+			OutObjectKey = FObjectKey();
+			return false;
+		};
+		
+		FindStackGroupsAndItemsForEntryKeysRecursive(SystemViewModel->GetSystemStackViewModel()->GetRootEntry(),
+			DeferredDisplayedObjectKeysToAddToSelection, TryGetObjectKeyFromStackEntry, FoundStackEntries);
+
 		for (TSharedRef<FNiagaraEmitterHandleViewModel> EmitterViewModel : SystemViewModel->GetEmitterHandleViewModels())
 		{
-			FindStackGroupsAndItemsForDisplayedObjectKeysRecursive(EmitterViewModel->GetEmitterStackViewModel()->GetRootEntry(), DeferredDisplayedObjectKeysToAddToSelection, FoundStackEntries);
+			FindStackGroupsAndItemsForEntryKeysRecursive(EmitterViewModel->GetEmitterStackViewModel()->GetRootEntry(),
+				DeferredDisplayedObjectKeysToAddToSelection, TryGetObjectKeyFromStackEntry, FoundStackEntries);
 		}
 
-		for (UNiagaraStackEntry* FoundStackEntry : FoundStackEntries)
+		DeferredDisplayedObjectKeysToAddToSelection.Empty();
+	}
+
+	if (DeferredSelectionIdsToAddToSelection.Num() > 0)
+	{
+		auto TryGetSelectionIdFromStackEntry = [](const UNiagaraStackEntry* StackEntry, FGuid& OutSelectionId)
+		{ 
+			OutSelectionId = StackEntry->GetSelectionId();
+			return OutSelectionId.IsValid();
+		};
+
+		FindStackGroupsAndItemsForEntryKeysRecursive(SystemViewModel->GetSystemStackViewModel()->GetRootEntry(),
+			DeferredSelectionIdsToAddToSelection, TryGetSelectionIdFromStackEntry, FoundStackEntries);
+
+		for (TSharedRef<FNiagaraEmitterHandleViewModel> EmitterViewModel : SystemViewModel->GetEmitterHandleViewModels())
 		{
-			if (ContainsEntry(FoundStackEntry) == false)
-			{
-				SelectionEntries.Add(FSelectionEntry(FoundStackEntry));
-				bSelectionChanged = true;
-			}
+			FindStackGroupsAndItemsForEntryKeysRecursive(EmitterViewModel->GetEmitterStackViewModel()->GetRootEntry(),
+				DeferredSelectionIdsToAddToSelection, TryGetSelectionIdFromStackEntry, FoundStackEntries);
 		}
 
-		DeferredDisplayedObjectKeysToAddToSelection.Empty();
+		DeferredSelectionIdsToAddToSelection.Empty();
+	}
+
+	for (UNiagaraStackEntry* FoundStackEntry : FoundStackEntries)
+	{
+		if (ContainsEntry(FoundStackEntry) == false)
+		{
+			SelectionEntries.Add(FSelectionEntry(FoundStackEntry));
+			bSelectionChanged = true;
+		}
 	}
 
 	bool bHasBeenRefreshed = false;
diff --git a/Engine/Plugins/FX/Niagara/Source/NiagaraEditor/Private/ViewModels/Stack/NiagaraStackEmitterPropertiesGroup.cpp b/Engine/Plugins/FX/Niagara/Source/NiagaraEditor/Private/ViewModels/Stack/NiagaraStackEmitterPropertiesGroup.cpp
new file mode 100644
--- /dev/null
+++ b/Engine/Plugins/FX/Niagara/Source/NiagaraEditor/Private/ViewModels/Stack/NiagaraStackEmitterPropertiesGroup.cpp
@@ -0,0 +1,232 @@
+// Copyright Epic Games, Inc. All Rights Reserved.
+
+#include "ViewModels/Stack/NiagaraStackEmitterPropertiesGroup.h"
+
+#include "NiagaraEditorStyle.h"
+#include "NiagaraEditorUtilities.h"
+#include "NiagaraEmitter.h"
+#include "NiagaraScriptSource.h"
+#include "ViewModels/NiagaraEmitterViewModel.h"
+#include "ViewModels/NiagaraScriptGraphViewModel.h"
+#include "ViewModels/NiagaraScriptViewModel.h"
+#include "ViewModels/NiagaraSystemSelectionViewModel.h"
+#include "ViewModels/NiagaraSystemViewModel.h"
+#include "ViewModels/Stack/NiagaraStackEmitterSettingsGroup.h"
+#include "ViewModels/Stack/NiagaraStackGraphUtilities.h"
+#include "NiagaraSimulationStageBase.h"
+
+#define LOCTEXT_NAMESPACE "StackEmitterProperties"
+
+class FNiagaraStackEmitterPropertiesGroupAddUtilities : public FNiagaraStackItemGroupAddUtilities
+{
+private:
+	enum class EEmitterAddMode
+	{
+		Event,
+		SimulationStage
+	};
+
+	class FAddEmitterStageAction : public INiagaraStackItemGroupAddAction
+	{
+	public:
+		FAddEmitterStageAction(EEmitterAddMode InAddMode, UClass* InSimulationStageClass, const TArray<FString>& InCategories, FText InDisplayName, FText InDescription, FText InKeywords)
+			: AddMode(InAddMode), SimulationStageClass(InSimulationStageClass), Categories(InCategories), DisplayName(InDisplayName), Description(InDescription), Keywords(InKeywords)
+		{
+		}
+
+		virtual TArray<FString> GetCategories() const override
+		{
+			return Categories;
+		}
+
+		virtual FText GetDisplayName() const override
+		{
+			return DisplayName;
+		}
+
+		virtual FText GetDescription() const override
+		{
+			return Description;
+		}
+
+		virtual FText GetKeywords() const override
+		{
+			return Keywords;
+		}
+
+		EEmitterAddMode AddMode;
+		UClass* SimulationStageClass;
+		TArray<FString> Categories;
+		FText DisplayName;
+		FText Description;
+		FText Keywords;
+	};
+
+public:
+	DECLARE_DELEGATE_TwoParams(FOnItemAdded, FGuid /* AddedEventHandlerId */, UNiagaraSimulationStageBase* /* AddedSimulationStage */);
+
+public:
+	FNiagaraStackEmitterPropertiesGroupAddUtilities(TSharedPtr<FNiagaraEmitterViewModel> InEmitterViewModel, FOnItemAdded InOnItemAdded)
+		: FNiagaraStackItemGroupAddUtilities(LOCTEXT("AddUtilitiesName", "Stage"), EAddMode::AddFromAction, true, true)
+		, EmitterViewModelWeak(TWeakPtr<FNiagaraEmitterViewModel>(InEmitterViewModel))
+		, OnItemAdded(InOnItemAdded)
+	{
+	}
+
+	virtual void AddItemDirectly() override { unimplemented(); }
+
+	virtual void GenerateAddActions(TArray<TSharedRef<INiagaraStackItemGroupAddAction>>& OutAddActions, const FNiagaraStackItemGroupAddOptions& AddProperties) const override
+	{
+		TArray<FString> EventsCategories = { "Events" };
+		TArray<FString> SimulationStageCategories = { "Simulation Stages" };
+
+		OutAddActions.Add(MakeShared<FAddEmitterStageAction>(
+			EEmitterAddMode::Event,
+			nullptr,
+			EventsCategories,
+			LOCTEXT("AddEventActionDisplayName", "Event Handler"),
+			LOCTEXT("AddEventActionDescription", "Add an event handler to this emitter."),
+			LOCTEXT("AddEventActionKeywords", "Event")));
+
+		TArray<UClass*> SimulationStageClasses;
+		GetDerivedClasses(UNiagaraSimulationStageBase::StaticClass(), SimulationStageClasses);
+		for (UClass* SimulationStageClass : SimulationStageClasses)
+		{
+			OutAddActions.Add(MakeShared<FAddEmitterStageAction>(
+				EEmitterAddMode::SimulationStage,
+				SimulationStageClass,
+				SimulationStageCategories,
+				SimulationStageClass->GetDisplayNameText(),
+				FText::FromString(SimulationStageClass->GetDescription()),
+				FText::FromString(SimulationStageClass->GetName())));
+		}
+	}
+
+	virtual void ExecuteAddAction(TSharedRef<INiagaraStackItemGroupAddAction> AddAction, int32 TargetIndex) override
+	{
+		TSharedPtr<FNiagaraEmitterViewModel> EmitterViewModel = EmitterViewModelWeak.Pin();
+		if (EmitterViewModel.IsValid() == false)
+		{
+			return;
+		}
+
+		UNiagaraEmitter* Emitter = EmitterViewModel->GetEmitter();
+		UNiagaraScriptSource* Source = EmitterViewModel->GetSharedScriptViewModel()->GetGraphViewModel()->GetScriptSource();
+		UNiagaraGraph* Graph = EmitterViewModel->GetSharedScriptViewModel()->GetGraphViewModel()->GetGraph();
+
+		// The stack should not have been created if any of these are null, so bail out if it happens somehow rather than try to handle all of these cases.
+		checkf(Emitter != nullptr && Source != nullptr && Graph != nullptr, TEXT("Stack created for invalid emitter or graph."));
+
+		TSharedRef<FAddEmitterStageAction> AddEmitterStageAction = StaticCastSharedRef<FAddEmitterStageAction>(AddAction);
+		FGuid AddedEventHandlerId;
+		UNiagaraSimulationStageBase* AddedSimulationStage = nullptr;
+		if (AddEmitterStageAction->AddMode == EEmitterAddMode::Event)
+		{
+			// since this is potentially modifying live data we need to kill off any existing instances that might be in flight before we make
+			// the change.
+			if (const UNiagaraSystem* EmitterSystem = Cast<UNiagaraSystem>(Emitter->GetOuter()))
+			{
+				FNiagaraEditorUtilities::KillSystemInstances(*EmitterSystem);
+			}
+
+			FScopedTransaction ScopedTransaction(LOCTEXT("AddNewEventHandlerTransaction", "Add new event handler"));
+
+			Emitter->Modify();
+			FNiagaraEventScriptProperties EventScriptProperties;
+			EventScriptProperties.Script = NewObject<UNiagaraScript>(Emitter, MakeUniqueObjectName(Emitter, UNiagaraScript::StaticClass(), "EventScript"), EObjectFlags::RF_Transactional);
+			EventScriptProperties.Script->SetUsage(ENiagaraScriptUsage::ParticleEventScript);
+			EventScriptProperties.Script->SetUsageId(FGuid::NewGuid());
+			EventScriptProperties.Script->SetLatestSource(Source);
+			Emitter->AddEventHandler(EventScriptProperties);
+			FNiagaraStackGraphUtilities::ResetGraphForOutput(*Graph, ENiagaraScriptUsage::ParticleEventScript, EventScriptProperties.Script->GetUsageId());
+			AddedEventHandlerId = EventScriptProperties.Script->GetUsageId();
+		}
+		else if (AddEmitterStageAction->AddMode == EEmitterAddMode::SimulationStage && AddEmitterStageAction->SimulationStageClass != nullptr)
+		{
+			FScopedTransaction ScopedTransaction(LOCTEXT("AddNewSimulationStagesTransaction", "Add new shader stage"));
+
+			Emitter->Modify();
+			AddedSimulationStage = NewObject<UNiagaraSimulationStageBase>(Emitter, AddEmitterStageAction->SimulationStageClass, NAME_None, RF_Transactional);
+			AddedSimulationStage->Script = NewObject<UNiagaraScript>(AddedSimulationStage, MakeUniqueObjectName(AddedSimulationStage, UNiagaraScript::StaticClass(), "SimulationStage"), EObjectFlags::RF_Transactional);
+			AddedSimulationStage->Script->SetUsage(ENiagaraScriptUsage::ParticleSimulationStageScript);
+			AddedSimulationStage->Script->SetUsageId(AddedSimulationStage->GetMergeId());
+			AddedSimulationStage->Script->SetLatestSource(Source);
+			Emitter->AddSimulationStage(AddedSimulationStage);
+			FNiagaraStackGraphUtilities::ResetGraphForOutput(*Graph, ENiagaraScriptUsage::ParticleSimulationStageScript, AddedSimulationStage->Script->GetUsageId());
+		}
+
+		// Set the emitter here so that the internal state of the view model is updated.
+		// TODO: Move the logic for managing additional scripts into the emitter view model or script view model.
+		TWeakPtr<FNiagaraEmitterInstance, ESPMode::ThreadSafe> Simulation = EmitterViewModel->GetSimulation();
+		EmitterViewModel->Reset();
+		EmitterViewModel->Initialize(Emitter, Simulation);
+
+		OnItemAdded.ExecuteIfBound(AddedEventHandlerId, AddedSimulationStage);
+	}
+
+private:
+	TWeakPtr<FNiagaraEmitterViewModel> EmitterViewModelWeak;
+	FOnItemAdded OnItemAdded;
+};
+
+UNiagaraStackEmitterPropertiesGroup::UNiagaraStackEmitterPropertiesGroup()
+	: PropertiesItem(nullptr)
+{
+}
+
+void UNiagaraStackEmitterPropertiesGroup::Initialize(FRequiredEntryData InRequiredEntryData)
+{
+	FText DisplayName = LOCTEXT("EmitterPropertiesGroupName", "Properties");
+	FText Tooltip = LOCTEXT("EmitterPropertiesTooltip", "Properties which are set per Emitter.");
+	AddUtilities = MakeShared<FNiagaraStackEmitterPropertiesGroupAddUtilities>(InRequiredEntryData.EmitterViewModel.ToSharedRef(), 
+		FNiagaraStackEmitterPropertiesGroupAddUtilities::FOnItemAdded::CreateUObject(this, &UNiagaraStackEmitterPropertiesGroup::ItemAddedFromUtilities));
+	Super::Initialize(InRequiredEntryData, DisplayName, Tooltip, AddUtilities.Get());
+}
+
+const FSlateBrush* UNiagaraStackEmitterPropertiesGroup::GetIconBrush() const
+{
+	return FAppStyle::Get().GetBrush("Icons.Details");
+}
+
+const FSlateBrush* UNiagaraStackEmitterPropertiesGroup::GetSecondaryIconBrush() const
+{
+	UNiagaraEmitter* Emitter = GetEmitterViewModel()->GetEmitter();
+	if (Emitter->SimTarget == ENiagaraSimTarget::CPUSim)
+	{
+		return FNiagaraEditorStyle::Get().GetBrush("NiagaraEditor.Stack.CPUIcon");
+	}
+	if (Emitter->SimTarget == ENiagaraSimTarget::GPUComputeSim)
+	{
+		return FNiagaraEditorStyle::Get().GetBrush("NiagaraEditor.Stack.GPUIcon");
+	}
+	return FEditorStyle::GetBrush("NoBrush");
+}
+
+void UNiagaraStackEmitterPropertiesGroup::RefreshChildrenInternal(const TArray<UNiagaraStackEntry*>& CurrentChildren, TArray<UNiagaraStackEntry*>& NewChildren, TArray<FStackIssue>& NewIssues)
+{
+	if (PropertiesItem == nullptr)
+	{
+		PropertiesItem = NewObject<UNiagaraStackEmitterPropertiesItem>(this);
+		PropertiesItem->Initialize(CreateDefaultChildRequiredData());
+	}
+	NewChildren.Add(PropertiesItem);
+
+	Super::RefreshChildrenInternal(CurrentChildren, NewChildren, NewIssues);
+}
+
+void UNiagaraStackEmitterPropertiesGroup::ItemAddedFromUtilities(FGuid AddedEventHandlerId, UNiagaraSimulationStageBase* AddedSimulationStage)
+{
+	OnRequestFullRefresh().Broadcast();
+
+	GetSystemViewModel()->GetSelectionViewModel()->EmptySelection();
+	if (AddedEventHandlerId.IsValid())
+	{
+		GetSystemViewModel()->GetSelectionViewModel()->AddEntryToSelectionBySelectionIdDeferred(AddedEventHandlerId);
+	}
+	else if (AddedSimulationStage != nullptr)
+	{
+		GetSystemViewModel()->GetSelectionViewModel()->AddEntryToSelectionByDisplayedObjectKeyDeferred(FObjectKey(AddedSimulationStage));
+	}
+}
+
+#undef LOCTEXT_NAMESPACE
\ No newline at end of file
diff --git a/Engine/Plugins/FX/Niagara/Source/NiagaraEditor/Private/ViewModels/Stack/NiagaraStackEmitterSettingsGroup.cpp b/Engine/Plugins/FX/Niagara/Source/NiagaraEditor/Private/ViewModels/Stack/NiagaraStackEmitterSettingsGroup.cpp
--- a/Engine/Plugins/FX/Niagara/Source/NiagaraEditor/Private/ViewModels/Stack/NiagaraStackEmitterSettingsGroup.cpp
+++ b/Engine/Plugins/FX/Niagara/Source/NiagaraEditor/Private/ViewModels/Stack/NiagaraStackEmitterSettingsGroup.cpp
@@ -13,6 +13,7 @@
 #include "NiagaraEditorStyle.h"
 #include "NiagaraSystem.h"
 #include "NiagaraEditorStyle.h"
+#include "Styling/AppStyle.h"
 
 #define LOCTEXT_NAMESPACE "UNiagaraStackEmitterItemGroup"
 
@@ -202,26 +203,6 @@ void UNiagaraStackEmitterSummaryItem::RefreshChildrenInternal(const TArray<UNiag
 	Super::RefreshChildrenInternal(CurrentChildren, NewChildren, NewIssues);
 }
 
-
-
-
-UNiagaraStackEmitterSettingsGroup::UNiagaraStackEmitterSettingsGroup()
-	: PropertiesItem(nullptr)
-{
-}
-
-void UNiagaraStackEmitterSettingsGroup::RefreshChildrenInternal(const TArray<UNiagaraStackEntry*>& CurrentChildren, TArray<UNiagaraStackEntry*>& NewChildren, TArray<FStackIssue>& NewIssues)
-{
-	if (PropertiesItem == nullptr)
-	{
-		PropertiesItem = NewObject<UNiagaraStackEmitterPropertiesItem>(this);
-		PropertiesItem->Initialize(CreateDefaultChildRequiredData());
-	}
-	NewChildren.Add(PropertiesItem);
-
-	Super::RefreshChildrenInternal(CurrentChildren, NewChildren, NewIssues);
-}
-
 UNiagaraStackEmitterSummaryGroup::UNiagaraStackEmitterSummaryGroup()
 	: SummaryItem(nullptr)
 {
@@ -239,6 +220,11 @@ void UNiagaraStackEmitterSummaryGroup::RefreshChildrenInternal(const TArray<UNia
 	Super::RefreshChildrenInternal(CurrentChildren, NewChildren, NewIssues);
 }
 
+const FSlateBrush* UNiagaraStackEmitterSummaryGroup::GetIconBrush() const
+{
+	return FAppStyle::Get().GetBrush("Icons.Edit");
+}
+
 
 
 
diff --git a/Engine/Plugins/FX/Niagara/Source/NiagaraEditor/Private/ViewModels/Stack/NiagaraStackEntry.cpp b/Engine/Plugins/FX/Niagara/Source/NiagaraEditor/Private/ViewModels/Stack/NiagaraStackEntry.cpp
--- a/Engine/Plugins/FX/Niagara/Source/NiagaraEditor/Private/ViewModels/Stack/NiagaraStackEntry.cpp
+++ b/Engine/Plugins/FX/Niagara/Source/NiagaraEditor/Private/ViewModels/Stack/NiagaraStackEntry.cpp
@@ -220,6 +220,11 @@ UObject* UNiagaraStackEntry::GetDisplayedObject() const
 	return nullptr;
 }
 
+FGuid UNiagaraStackEntry::GetSelectionId() const
+{
+	return FGuid();
+}
+
 UNiagaraStackEditorData& UNiagaraStackEntry::GetStackEditorData() const
 {
 	return *StackEditorData;
@@ -360,6 +365,11 @@ bool UNiagaraStackEntry::GetShouldShowInStack() const
 	return true;
 }
 
+bool UNiagaraStackEntry::GetShouldShowInOverview() const
+{
+	return true;
+}
+
 void UNiagaraStackEntry::GetFilteredChildren(TArray<UNiagaraStackEntry*>& OutFilteredChildren) const
 {
 	if (bFilterChildrenPending)
diff --git a/Engine/Plugins/FX/Niagara/Source/NiagaraEditor/Private/ViewModels/Stack/NiagaraStackEventHandlerGroup.cpp b/Engine/Plugins/FX/Niagara/Source/NiagaraEditor/Private/ViewModels/Stack/NiagaraStackEventHandlerGroup.cpp
deleted file mode 100644
diff --git a/Engine/Plugins/FX/Niagara/Source/NiagaraEditor/Private/ViewModels/Stack/NiagaraStackModuleItem.cpp b/Engine/Plugins/FX/Niagara/Source/NiagaraEditor/Private/ViewModels/Stack/NiagaraStackModuleItem.cpp
--- a/Engine/Plugins/FX/Niagara/Source/NiagaraEditor/Private/ViewModels/Stack/NiagaraStackModuleItem.cpp
+++ b/Engine/Plugins/FX/Niagara/Source/NiagaraEditor/Private/ViewModels/Stack/NiagaraStackModuleItem.cpp
@@ -1087,11 +1087,6 @@ void UNiagaraStackModuleItem::Refresh()
 	}
 }
 
-TOptional<float> UNiagaraStackModuleItem::GetCustomOverviewIndent()
-{
-	return 8.f;
-}
-
 bool UNiagaraStackModuleItem::GetIsEnabled() const
 {
 	return bIsEnabled;
diff --git a/Engine/Plugins/FX/Niagara/Source/NiagaraEditor/Private/ViewModels/Stack/NiagaraStackRenderItemGroup.cpp b/Engine/Plugins/FX/Niagara/Source/NiagaraEditor/Private/ViewModels/Stack/NiagaraStackRenderItemGroup.cpp
--- a/Engine/Plugins/FX/Niagara/Source/NiagaraEditor/Private/ViewModels/Stack/NiagaraStackRenderItemGroup.cpp
+++ b/Engine/Plugins/FX/Niagara/Source/NiagaraEditor/Private/ViewModels/Stack/NiagaraStackRenderItemGroup.cpp
@@ -58,7 +58,7 @@ class FRenderItemGroupAddUtilities : public TNiagaraStackItemGroupAddUtilities<U
 {
 public:
 	FRenderItemGroupAddUtilities(TSharedRef<FNiagaraEmitterViewModel> InEmitterViewModel, FRenderItemGroupAddUtilities::FOnItemAdded OnItemAdded = FRenderItemGroupAddUtilities::FOnItemAdded())
-		: TNiagaraStackItemGroupAddUtilities(LOCTEXT("RenderGroupAddItemName", "Renderer"), EAddMode::AddFromAction, true, OnItemAdded)
+		: TNiagaraStackItemGroupAddUtilities(LOCTEXT("RenderGroupAddItemName", "Renderer"), EAddMode::AddFromAction, true, false, OnItemAdded)
 		, EmitterViewModel(InEmitterViewModel)
 	{
 	}
diff --git a/Engine/Plugins/FX/Niagara/Source/NiagaraEditor/Private/ViewModels/Stack/NiagaraStackRoot.cpp b/Engine/Plugins/FX/Niagara/Source/NiagaraEditor/Private/ViewModels/Stack/NiagaraStackRoot.cpp
--- a/Engine/Plugins/FX/Niagara/Source/NiagaraEditor/Private/ViewModels/Stack/NiagaraStackRoot.cpp
+++ b/Engine/Plugins/FX/Niagara/Source/NiagaraEditor/Private/ViewModels/Stack/NiagaraStackRoot.cpp
@@ -2,11 +2,10 @@
 
 #include "ViewModels/Stack/NiagaraStackRoot.h"
 #include "ViewModels/Stack/NiagaraStackScriptItemGroup.h"
+#include "ViewModels/Stack/NiagaraStackEmitterPropertiesGroup.h"
 #include "ViewModels/Stack/NiagaraStackEmitterSettingsGroup.h"
 #include "ViewModels/Stack/NiagaraStackRenderItemGroup.h"
-#include "ViewModels/Stack/NiagaraStackEventHandlerGroup.h"
 #include "ViewModels/Stack/NiagaraStackEventScriptItemGroup.h"
-#include "ViewModels/Stack/NiagaraStackSimulationStagesGroup.h"
 #include "ViewModels/Stack/NiagaraStackSimulationStageGroup.h"
 #include "ViewModels/NiagaraSystemViewModel.h"
 #include "ViewModels/NiagaraEmitterViewModel.h"
@@ -20,15 +19,14 @@
 #define LOCTEXT_NAMESPACE "NiagaraStackViewModel"
 
 UNiagaraStackRoot::UNiagaraStackRoot()
-	: SystemSettingsGroup(nullptr)
-	, EmitterSettingsGroup(nullptr)
+	: SystemPropertiesGroup(nullptr)
+	, SystemUserParametersGroup(nullptr)
+	, EmitterPropertiesGroup(nullptr)
 	, EmitterSummaryGroup(nullptr)
 	, EmitterSpawnGroup(nullptr)
 	, EmitterUpdateGroup(nullptr)
 	, ParticleSpawnGroup(nullptr)
 	, ParticleUpdateGroup(nullptr)
-	, AddEventHandlerGroup(nullptr)
-	, AddSimulationStageGroup(nullptr)
 	, RenderGroup(nullptr)
 {
 }
@@ -38,15 +36,14 @@ void UNiagaraStackRoot::Initialize(FRequiredEntryData InRequiredEntryData, bool
 	Super::Initialize(InRequiredEntryData, TEXT("Root"));
 	bIncludeSystemInformation = bInIncludeSystemInformation;
 	bIncludeEmitterInformation = bInIncludeEmitterInformation;
-	SystemSettingsGroup = nullptr;
-	EmitterSettingsGroup = nullptr;
+	SystemPropertiesGroup = nullptr;
+	SystemUserParametersGroup = nullptr;
+	EmitterPropertiesGroup = nullptr;
 	EmitterSummaryGroup = nullptr;
 	EmitterSpawnGroup = nullptr;
 	EmitterUpdateGroup = nullptr;
 	ParticleSpawnGroup = nullptr;
 	ParticleUpdateGroup = nullptr;
-	AddEventHandlerGroup = nullptr;
-	AddSimulationStageGroup = nullptr;
 	RenderGroup = nullptr;
 
 	if (bInIncludeEmitterInformation && GetEmitterViewModel())
@@ -77,13 +74,22 @@ bool UNiagaraStackRoot::GetShouldShowInStack() const
 void UNiagaraStackRoot::RefreshChildrenInternal(const TArray<UNiagaraStackEntry*>& CurrentChildren, TArray<UNiagaraStackEntry*>& NewChildren, TArray<FStackIssue>& NewIssues)
 {
 	// Create static entries as needed.
-	if (bIncludeSystemInformation && SystemSettingsGroup == nullptr)
+	if (bIncludeSystemInformation && SystemPropertiesGroup == nullptr)
 	{
-		SystemSettingsGroup = NewObject<UNiagaraStackSystemSettingsGroup>(this);
+		SystemPropertiesGroup = NewObject<UNiagaraStackSystemPropertiesGroup>(this);
 		FRequiredEntryData RequiredEntryData(GetSystemViewModel(), GetEmitterViewModel(),
 			FExecutionCategoryNames::System, FExecutionSubcategoryNames::Settings,
 			GetSystemViewModel()->GetEditorData().GetStackEditorData());
-		SystemSettingsGroup->Initialize(RequiredEntryData, &GetSystemViewModel()->GetSystem(), &GetSystemViewModel()->GetSystem().GetExposedParameters());
+		SystemPropertiesGroup->Initialize(RequiredEntryData);
+	}
+
+	if (bIncludeSystemInformation && SystemUserParametersGroup == nullptr)
+	{
+		SystemUserParametersGroup = NewObject<UNiagaraStackSystemUserParametersGroup>(this);
+		FRequiredEntryData RequiredEntryData(GetSystemViewModel(), GetEmitterViewModel(),
+			FExecutionCategoryNames::System, FExecutionSubcategoryNames::Settings,
+			GetSystemViewModel()->GetEditorData().GetStackEditorData());
+		SystemUserParametersGroup->Initialize(RequiredEntryData, &GetSystemViewModel()->GetSystem(), &GetSystemViewModel()->GetSystem().GetExposedParameters());
 	}
 
 	if (bIncludeSystemInformation && SystemSpawnGroup == nullptr)
@@ -118,19 +124,15 @@ void UNiagaraStackRoot::RefreshChildrenInternal(const TArray<UNiagaraStackEntry*
 		EmitterUpdateGroup = nullptr;
 		ParticleSpawnGroup = nullptr;
 		ParticleUpdateGroup = nullptr;
-		AddEventHandlerGroup = nullptr;
-		AddSimulationStageGroup = nullptr;
 	}
 	
-	if (bIncludeEmitterInformation && EmitterSettingsGroup == nullptr)
+	if (bIncludeEmitterInformation && EmitterPropertiesGroup == nullptr)
 	{
-		EmitterSettingsGroup = NewObject<UNiagaraStackEmitterSettingsGroup>(this);
+		EmitterPropertiesGroup = NewObject<UNiagaraStackEmitterPropertiesGroup>(this);
 		FRequiredEntryData RequiredEntryData(GetSystemViewModel(), GetEmitterViewModel(),
 			FExecutionCategoryNames::Emitter, FExecutionSubcategoryNames::Settings,
 			GetEmitterViewModel()->GetOrCreateEditorData().GetStackEditorData());
-		FText DisplayName = LOCTEXT("EmitterSettingsGroupName", "Emitter Settings");
-		FText Tooltip = LOCTEXT("EmitterSettingsTooltip", "Settings that are handled per Emitter.");
-		EmitterSettingsGroup->Initialize(RequiredEntryData, DisplayName, Tooltip, nullptr);
+		EmitterPropertiesGroup->Initialize(RequiredEntryData);
 	}
 
 	if (bIncludeEmitterInformation && EmitterSummaryGroup == nullptr)
@@ -188,26 +190,6 @@ void UNiagaraStackRoot::RefreshChildrenInternal(const TArray<UNiagaraStackEntry*
 		ParticleUpdateGroup->Initialize(RequiredEntryData, DisplayName, ToolTip, GetEmitterViewModel()->GetSharedScriptViewModel(), ENiagaraScriptUsage::ParticleUpdateScript);
 	}
 
-	if (bIncludeEmitterInformation && !bShouldShowSummaryView && AddEventHandlerGroup == nullptr)
-	{
-		AddEventHandlerGroup = NewObject<UNiagaraStackEventHandlerGroup>(this);
-		FRequiredEntryData RequiredEntryData(GetSystemViewModel(), GetEmitterViewModel(),
-			FExecutionCategoryNames::Particle, FExecutionSubcategoryNames::Event,
-			GetEmitterViewModel()->GetOrCreateEditorData().GetStackEditorData());
-		AddEventHandlerGroup->Initialize(RequiredEntryData);
-		AddEventHandlerGroup->SetOnItemAdded(UNiagaraStackEventHandlerGroup::FOnItemAdded::CreateUObject(this, &UNiagaraStackRoot::EmitterArraysChanged));
-	}
-
-	if (bIncludeEmitterInformation && !bShouldShowSummaryView && AddSimulationStageGroup == nullptr)
-	{
-		AddSimulationStageGroup = NewObject<UNiagaraStackSimulationStagesGroup>(this);
-		FRequiredEntryData RequiredEntryData(GetSystemViewModel(), GetEmitterViewModel(),
-			FExecutionCategoryNames::Particle, FExecutionSubcategoryNames::SimulationStage,
-			GetEmitterViewModel()->GetOrCreateEditorData().GetStackEditorData());
-		AddSimulationStageGroup->Initialize(RequiredEntryData);
-		AddSimulationStageGroup->SetOnItemAdded(UNiagaraStackSimulationStagesGroup::FOnItemAdded::CreateUObject(this, &UNiagaraStackRoot::EmitterArraysChanged));
-	}
-
 	if (bIncludeEmitterInformation && RenderGroup == nullptr)
 	{
 		RenderGroup = NewObject<UNiagaraStackRenderItemGroup>(this);
@@ -229,14 +211,15 @@ void UNiagaraStackRoot::RefreshChildrenInternal(const TArray<UNiagaraStackEntry*
 	// Populate new children
 	if (bIncludeSystemInformation)
 	{
-		NewChildren.Add(SystemSettingsGroup);
+		NewChildren.Add(SystemPropertiesGroup);
+		NewChildren.Add(SystemUserParametersGroup);
 		NewChildren.Add(SystemSpawnGroup);
 		NewChildren.Add(SystemUpdateGroup);
 	}
 
 	if (bIncludeEmitterInformation)
 	{
-		NewChildren.Add(EmitterSettingsGroup);	
+		NewChildren.Add(EmitterPropertiesGroup);	
 		NewChildren.Add(EmitterSummaryGroup);
 
 		if (!bShouldShowSummaryView)
@@ -265,8 +248,6 @@ void UNiagaraStackRoot::RefreshChildrenInternal(const TArray<UNiagaraStackEntry*
 				NewChildren.Add(EventHandlerGroup);
 			}
 
-			NewChildren.Add(AddEventHandlerGroup);
-
 			for (UNiagaraSimulationStageBase* SimulationStage : GetEmitterViewModel()->GetEmitter()->GetSimulationStages())
 			{
 				UNiagaraStackSimulationStageGroup* SimulationStageGroup = FindCurrentChildOfTypeByPredicate<UNiagaraStackSimulationStageGroup>(CurrentChildren,
@@ -284,7 +265,6 @@ void UNiagaraStackRoot::RefreshChildrenInternal(const TArray<UNiagaraStackEntry*
 
 				NewChildren.Add(SimulationStageGroup);
 			}
-			NewChildren.Add(AddSimulationStageGroup);
 		}
 
 		NewChildren.Add(RenderGroup);
diff --git a/Engine/Plugins/FX/Niagara/Source/NiagaraEditor/Private/ViewModels/Stack/NiagaraStackScriptItemGroup.cpp b/Engine/Plugins/FX/Niagara/Source/NiagaraEditor/Private/ViewModels/Stack/NiagaraStackScriptItemGroup.cpp
--- a/Engine/Plugins/FX/Niagara/Source/NiagaraEditor/Private/ViewModels/Stack/NiagaraStackScriptItemGroup.cpp
+++ b/Engine/Plugins/FX/Niagara/Source/NiagaraEditor/Private/ViewModels/Stack/NiagaraStackScriptItemGroup.cpp
@@ -228,7 +228,7 @@ class FScriptItemGroupAddUtilities : public TNiagaraStackItemGroupAddUtilities<U
 {
 public:
 	FScriptItemGroupAddUtilities(TSharedRef<FNiagaraSystemViewModel> InSystemViewModel, TSharedPtr<FNiagaraEmitterViewModel> InEmitterViewModel, UNiagaraStackEditorData& InStackEditorData, FOnItemAdded InOnItemAdded)
-		: TNiagaraStackItemGroupAddUtilities(LOCTEXT("ScriptGroupAddItemName", "Module"), EAddMode::AddFromAction, false, InOnItemAdded)
+		: TNiagaraStackItemGroupAddUtilities(LOCTEXT("ScriptGroupAddItemName", "Module"), EAddMode::AddFromAction, false, false, InOnItemAdded)
 		, OutputNode(nullptr)
 		, SystemViewModel(InSystemViewModel)
 		, EmitterViewModel(InEmitterViewModel)
diff --git a/Engine/Plugins/FX/Niagara/Source/NiagaraEditor/Private/ViewModels/Stack/NiagaraStackSimulationStagesGroup.cpp b/Engine/Plugins/FX/Niagara/Source/NiagaraEditor/Private/ViewModels/Stack/NiagaraStackSimulationStagesGroup.cpp
deleted file mode 100644
diff --git a/Engine/Plugins/FX/Niagara/Source/NiagaraEditor/Private/ViewModels/Stack/NiagaraStackSummaryViewInputCollection.cpp b/Engine/Plugins/FX/Niagara/Source/NiagaraEditor/Private/ViewModels/Stack/NiagaraStackSummaryViewInputCollection.cpp
--- a/Engine/Plugins/FX/Niagara/Source/NiagaraEditor/Private/ViewModels/Stack/NiagaraStackSummaryViewInputCollection.cpp
+++ b/Engine/Plugins/FX/Niagara/Source/NiagaraEditor/Private/ViewModels/Stack/NiagaraStackSummaryViewInputCollection.cpp
@@ -20,7 +20,6 @@
 #include "ViewModels/Stack/NiagaraStackInputCategory.h"
 #include "ViewModels/Stack/NiagaraStackModuleItem.h"
 #include "ViewModels/Stack/NiagaraStackFunctionInputCollection.h"
-#include "ViewModels/Stack/NiagaraStackSimulationStagesGroup.h"
 #include "ViewModels/Stack/NiagaraStackSimulationStageGroup.h"
 
 
diff --git a/Engine/Plugins/FX/Niagara/Source/NiagaraEditor/Private/ViewModels/Stack/NiagaraStackSystemSettingsGroup.cpp b/Engine/Plugins/FX/Niagara/Source/NiagaraEditor/Private/ViewModels/Stack/NiagaraStackSystemSettingsGroup.cpp
--- a/Engine/Plugins/FX/Niagara/Source/NiagaraEditor/Private/ViewModels/Stack/NiagaraStackSystemSettingsGroup.cpp
+++ b/Engine/Plugins/FX/Niagara/Source/NiagaraEditor/Private/ViewModels/Stack/NiagaraStackSystemSettingsGroup.cpp
@@ -58,7 +58,7 @@ class FParameterStoreGroupAddUtiliites : public TNiagaraStackItemGroupAddUtiliti
 {
 public:
 	FParameterStoreGroupAddUtiliites(UObject& InParameterStoreOwner, FNiagaraParameterStore& InParameterStore, UNiagaraStackEditorData& InStackEditorData, FOnItemAdded InOnItemAdded)
-		: TNiagaraStackItemGroupAddUtilities(LOCTEXT("ScriptGroupAddItemName", "Parameter"), EAddMode::AddFromAction, true, InOnItemAdded)
+		: TNiagaraStackItemGroupAddUtilities(LOCTEXT("ScriptGroupAddItemName", "Parameter"), EAddMode::AddFromAction, true, false, InOnItemAdded)
 		, ParameterStoreOwner(InParameterStoreOwner)
 		, ParameterStore(InParameterStore)
 		, StackEditorData(InStackEditorData)
@@ -96,22 +96,54 @@ private:
 	UNiagaraStackEditorData& StackEditorData;
 };
 
-void UNiagaraStackSystemSettingsGroup::Initialize(
+void UNiagaraStackSystemPropertiesGroup::Initialize(FRequiredEntryData InRequiredEntryData)
+{
+	FText DisplayName = LOCTEXT("SystemPropertiesGroupName", "Properties");
+	FText Tooltip = LOCTEXT("SystemSettingsTooltip", "Properties set for the entire system.");
+	Super::Initialize(InRequiredEntryData, DisplayName, Tooltip, nullptr);
+}
+
+const FSlateBrush* UNiagaraStackSystemPropertiesGroup::GetIconBrush() const
+{
+	return FAppStyle::Get().GetBrush("Icons.Details");
+}
+
+void UNiagaraStackSystemPropertiesGroup::RefreshChildrenInternal(const TArray<UNiagaraStackEntry*>& CurrentChildren, TArray<UNiagaraStackEntry*>& NewChildren, TArray<FStackIssue>& NewIssues)
+{
+	UNiagaraStackSystemPropertiesItem* SystemPropertiesItem = FindCurrentChildOfTypeByPredicate<UNiagaraStackSystemPropertiesItem>(CurrentChildren,
+		[=](UNiagaraStackSystemPropertiesItem* CurrentItem) { return true; });
+
+	if (SystemPropertiesItem == nullptr)
+	{
+		SystemPropertiesItem = NewObject<UNiagaraStackSystemPropertiesItem>(this);
+		SystemPropertiesItem->Initialize(CreateDefaultChildRequiredData());
+	}
+	NewChildren.Add(SystemPropertiesItem);
+
+	Super::RefreshChildrenInternal(CurrentChildren, NewChildren, NewIssues);
+}
+
+void UNiagaraStackSystemUserParametersGroup::Initialize(
 	FRequiredEntryData InRequiredEntryData,
 	UObject* InOwner,
 	FNiagaraParameterStore* InParameterStore)
 {
 	AddUtilities = MakeShared<FParameterStoreGroupAddUtiliites>(*InOwner, *InParameterStore, *InRequiredEntryData.StackEditorData,
-		FParameterStoreGroupAddUtiliites::FOnItemAdded::CreateUObject(this, &UNiagaraStackSystemSettingsGroup::ParameterAdded));
-	FText DisplayName = LOCTEXT("SystemSettingsGroupName", "System Settings");
-	FText Tooltip = LOCTEXT("SystemSettingsTooltip", "Settings of the System.");
+		FParameterStoreGroupAddUtiliites::FOnItemAdded::CreateUObject(this, &UNiagaraStackSystemUserParametersGroup::ParameterAdded));
+	FText DisplayName = LOCTEXT("SystemUserParametersGroupName", "User Parameters");
+	FText Tooltip = LOCTEXT("SystemUserParametersTooltip", "Parameters for the system which are exposed externally.");
 	Super::Initialize(InRequiredEntryData, DisplayName, Tooltip, AddUtilities.Get());
 
 	Owner = InOwner;
 	UserParameterStore = InParameterStore;
 }
 
-void UNiagaraStackSystemSettingsGroup::RefreshChildrenInternal(const TArray<UNiagaraStackEntry*>& CurrentChildren, TArray<UNiagaraStackEntry*>& NewChildren, TArray<FStackIssue>& NewIssues)
+const FSlateBrush* UNiagaraStackSystemUserParametersGroup::GetIconBrush() const
+{
+	return FAppStyle::Get().GetBrush("Icons.Edit");
+}
+
+void UNiagaraStackSystemUserParametersGroup::RefreshChildrenInternal(const TArray<UNiagaraStackEntry*>& CurrentChildren, TArray<UNiagaraStackEntry*>& NewChildren, TArray<FStackIssue>& NewIssues)
 {
 	if (UserParameterStore != nullptr)
 	{
@@ -121,26 +153,16 @@ void UNiagaraStackSystemSettingsGroup::RefreshChildrenInternal(const TArray<UNia
 		if (UserParameterStoreItem == nullptr)
 		{
 			UserParameterStoreItem = NewObject<UNiagaraStackParameterStoreItem>(this);
-			UserParameterStoreItem->Initialize(CreateDefaultChildRequiredData(), Owner.Get(), UserParameterStore);
+			UserParameterStoreItem->Initialize(CreateDefaultChildRequiredData(), Owner.Get(), UserParameterStore, AddUtilities.Get());
 		}
 
 		NewChildren.Add(UserParameterStoreItem);
 	}
 
-	UNiagaraStackSystemPropertiesItem* SystemPropertiesItem = FindCurrentChildOfTypeByPredicate<UNiagaraStackSystemPropertiesItem>(CurrentChildren,
-		[=](UNiagaraStackSystemPropertiesItem* CurrentItem) { return true; });
-
-	if (SystemPropertiesItem == nullptr)
-	{
-		SystemPropertiesItem = NewObject<UNiagaraStackSystemPropertiesItem>(this);
-		SystemPropertiesItem->Initialize(CreateDefaultChildRequiredData());
-	}
-	NewChildren.Add(SystemPropertiesItem);
-
 	Super::RefreshChildrenInternal(CurrentChildren, NewChildren, NewIssues);
 }
 
-void UNiagaraStackSystemSettingsGroup::ParameterAdded(FNiagaraVariable AddedParameter)
+void UNiagaraStackSystemUserParametersGroup::ParameterAdded(FNiagaraVariable AddedParameter)
 {
 	if (AddedParameter.GetType().IsDataInterface())
 	{
@@ -158,7 +180,8 @@ void UNiagaraStackSystemSettingsGroup::ParameterAdded(FNiagaraVariable AddedPara
 void UNiagaraStackParameterStoreItem::Initialize(
 	FRequiredEntryData InRequiredEntryData,
 	UObject* InOwner,
-	FNiagaraParameterStore* InParameterStore)
+	FNiagaraParameterStore* InParameterStore,
+	INiagaraStackItemGroupAddUtilities* InGroupAddUtilities)
 {
 	Super::Initialize(InRequiredEntryData, TEXT("ParameterStoreItem"));
 
@@ -166,6 +189,7 @@ void UNiagaraStackParameterStoreItem::Initialize(
 	ParameterStore = InParameterStore;
 	ParameterStoreChangedHandle = ParameterStore->AddOnChangedHandler(
 		FNiagaraParameterStore::FOnChanged::FDelegate::CreateUObject(this, &UNiagaraStackParameterStoreItem::ParameterStoreChanged));
+	GroupAddUtilities = InGroupAddUtilities;
 }
 
 FText UNiagaraStackParameterStoreItem::GetDisplayName() const
diff --git a/Engine/Plugins/FX/Niagara/Source/NiagaraEditor/Public/ViewModels/NiagaraSystemSelectionViewModel.h b/Engine/Plugins/FX/Niagara/Source/NiagaraEditor/Public/ViewModels/NiagaraSystemSelectionViewModel.h
--- a/Engine/Plugins/FX/Niagara/Source/NiagaraEditor/Public/ViewModels/NiagaraSystemSelectionViewModel.h
+++ b/Engine/Plugins/FX/Niagara/Source/NiagaraEditor/Public/ViewModels/NiagaraSystemSelectionViewModel.h
@@ -60,6 +60,8 @@ public:
 
 	void AddEntryToSelectionByDisplayedObjectKeyDeferred(const FObjectKey& InObjectKey);
 
+	void AddEntryToSelectionBySelectionIdDeferred(const FGuid& InSelectionId);
+
 	bool Refresh();
 
 	void RefreshDeferred();
@@ -102,5 +104,7 @@ private:
 
 	TArray<FObjectKey> DeferredDisplayedObjectKeysToAddToSelection;
 
+	TArray<FGuid> DeferredSelectionIdsToAddToSelection;
+
 	bool bRefreshIsPending;
 };
\ No newline at end of file
diff --git a/Engine/Plugins/FX/Niagara/Source/NiagaraEditor/Public/ViewModels/Stack/INiagaraStackItemGroupAddUtilities.h b/Engine/Plugins/FX/Niagara/Source/NiagaraEditor/Public/ViewModels/Stack/INiagaraStackItemGroupAddUtilities.h
--- a/Engine/Plugins/FX/Niagara/Source/NiagaraEditor/Public/ViewModels/Stack/INiagaraStackItemGroupAddUtilities.h
+++ b/Engine/Plugins/FX/Niagara/Source/NiagaraEditor/Public/ViewModels/Stack/INiagaraStackItemGroupAddUtilities.h
@@ -67,6 +67,8 @@ public:
 	/** Gets the generic name for the type of item to add e.g. "Module" */
 	virtual FText GetAddItemName() const = 0;
 
+	virtual bool GetShowLabel() const = 0;
+
 	/** Gets whether or not the add actions should be automatically expanded in the UI. */
 	virtual bool GetAutoExpandAddActions() const = 0;
 
@@ -94,14 +96,16 @@ public:
 class FNiagaraStackItemGroupAddUtilities : public INiagaraStackItemGroupAddUtilities
 {
 public:
-	FNiagaraStackItemGroupAddUtilities(FText InAddItemName, EAddMode InAddMode, bool bInAutoExpandAddActions)
+	FNiagaraStackItemGroupAddUtilities(FText InAddItemName, EAddMode InAddMode, bool bInAutoExpandAddActions, bool bInShowLabel)
 		: AddItemName(InAddItemName)
+		, bShowLabel(bInShowLabel)
 		, bAutoExpandAddActions(bInAutoExpandAddActions)
 		, AddMode(InAddMode)
 	{
 	}
 
 	virtual FText GetAddItemName() const override { return AddItemName; }
+	virtual bool GetShowLabel() const override { return bShowLabel; }
 	virtual bool GetAutoExpandAddActions() const override { return bAutoExpandAddActions; }
 	virtual EAddMode GetAddMode() const override { return AddMode; }
 
@@ -110,6 +114,7 @@ public:
 
 protected:
 	FText AddItemName;
+	bool bShowLabel;
 	bool bAutoExpandAddActions;
 	EAddMode AddMode;
 };
@@ -121,8 +126,8 @@ public:
 	DECLARE_DELEGATE_OneParam(FOnItemAdded, AddedItemType);
 
 public:
-	TNiagaraStackItemGroupAddUtilities(FText InAddItemName, EAddMode InAddMode, bool bInAutoExpandAddActions, FOnItemAdded InOnItemAdded)
-		: FNiagaraStackItemGroupAddUtilities(InAddItemName, InAddMode, bInAutoExpandAddActions)
+	TNiagaraStackItemGroupAddUtilities(FText InAddItemName, EAddMode InAddMode, bool bInAutoExpandAddActions, bool bInShowLabel, FOnItemAdded InOnItemAdded)
+		: FNiagaraStackItemGroupAddUtilities(InAddItemName, InAddMode, bInAutoExpandAddActions, bInShowLabel)
 		, OnItemAdded(InOnItemAdded)
 	{
 	}
diff --git a/Engine/Plugins/FX/Niagara/Source/NiagaraEditor/Public/ViewModels/Stack/NiagaraStackEmitterPropertiesGroup.h b/Engine/Plugins/FX/Niagara/Source/NiagaraEditor/Public/ViewModels/Stack/NiagaraStackEmitterPropertiesGroup.h
new file mode 100644
--- /dev/null
+++ b/Engine/Plugins/FX/Niagara/Source/NiagaraEditor/Public/ViewModels/Stack/NiagaraStackEmitterPropertiesGroup.h
@@ -0,0 +1,43 @@
+// Copyright Epic Games, Inc. All Rights Reserved.
+
+#pragma once
+
+#include "ViewModels/Stack/NiagaraStackScriptItemGroup.h"
+#include "NiagaraStackEmitterPropertiesGroup.generated.h"
+
+class FNiagaraEmitterViewModel;
+class FNiagaraScriptViewModel;
+class INiagaraStackItemGroupAddUtilities;
+class UNiagaraSimulationStageBase;
+class UNiagaraStackEmitterPropertiesItem;
+class UNiagaraStackObject;
+
+UCLASS()
+class NIAGARAEDITOR_API UNiagaraStackEmitterPropertiesGroup : public UNiagaraStackItemGroup
+{
+	GENERATED_BODY()
+
+public:
+	UNiagaraStackEmitterPropertiesGroup();
+
+	void Initialize(FRequiredEntryData InRequiredEntryData);
+
+	virtual bool SupportsIcon() const override { return true; }
+	virtual const FSlateBrush* GetIconBrush() const override;
+	virtual bool SupportsSecondaryIcon() const override { return true; }
+	virtual const FSlateBrush* GetSecondaryIconBrush() const override;
+	virtual bool GetCanExpandInOverview() const override { return false; }
+	virtual bool GetShouldShowInStack() const override { return false; }
+
+protected:
+	void RefreshChildrenInternal(const TArray<UNiagaraStackEntry*>& CurrentChildren, TArray<UNiagaraStackEntry*>& NewChildren, TArray<FStackIssue>& NewIssues) override;
+
+private:
+	void ItemAddedFromUtilities(FGuid AddedEventHandlerId, UNiagaraSimulationStageBase* AddedSimulationStage);
+
+private:
+	UPROPERTY()
+	TObjectPtr<UNiagaraStackEmitterPropertiesItem> PropertiesItem;
+
+	TSharedPtr<INiagaraStackItemGroupAddUtilities> AddUtilities;
+};
\ No newline at end of file
diff --git a/Engine/Plugins/FX/Niagara/Source/NiagaraEditor/Public/ViewModels/Stack/NiagaraStackEmitterSettingsGroup.h b/Engine/Plugins/FX/Niagara/Source/NiagaraEditor/Public/ViewModels/Stack/NiagaraStackEmitterSettingsGroup.h
--- a/Engine/Plugins/FX/Niagara/Source/NiagaraEditor/Public/ViewModels/Stack/NiagaraStackEmitterSettingsGroup.h
+++ b/Engine/Plugins/FX/Niagara/Source/NiagaraEditor/Public/ViewModels/Stack/NiagaraStackEmitterSettingsGroup.h
@@ -27,10 +27,13 @@ public:
 	virtual bool TestCanResetToBaseWithMessage(FText& OutCanResetToBaseMessage) const override;
 	virtual void ResetToBase() override;
 
-	virtual bool IsExpandedByDefault() const override;
 	virtual bool SupportsIcon() const override { return true; }
 	virtual const FSlateBrush* GetIconBrush() const override;
 
+	virtual bool IsExpandedByDefault() const override;
+
+	virtual bool GetShouldShowInOverview() const override { return false; }
+
 protected:
 	virtual void FinalizeInternal() override;
 
@@ -65,6 +68,7 @@ public:
 	virtual bool IsExpandedByDefault() const override { return false; }
 	virtual bool SupportsIcon() const override { return false; }
 	virtual const FSlateBrush* GetIconBrush() const override;
+	virtual bool GetShouldShowInOverview() const override { return false; }
 
 protected:
 
@@ -77,24 +81,6 @@ private:
 	TObjectPtr<UNiagaraStackSummaryViewObject> FilteredObject;
 };
 
-
-UCLASS()
-class NIAGARAEDITOR_API UNiagaraStackEmitterSettingsGroup : public UNiagaraStackItemGroup
-{
-	GENERATED_BODY()
-
-public:
-	UNiagaraStackEmitterSettingsGroup();
-
-protected:
-	void RefreshChildrenInternal(const TArray<UNiagaraStackEntry*>& CurrentChildren, TArray<UNiagaraStackEntry*>& NewChildren, TArray<FStackIssue>& NewIssues) override;
-
-private:
-	UPROPERTY()
-	TObjectPtr<UNiagaraStackEmitterPropertiesItem> PropertiesItem;
-};
-
-
 UCLASS()
 class NIAGARAEDITOR_API UNiagaraStackEmitterSummaryGroup : public UNiagaraStackItemGroup
 {
@@ -103,6 +89,11 @@ class NIAGARAEDITOR_API UNiagaraStackEmitterSummaryGroup : public UNiagaraStackI
 public:
 	UNiagaraStackEmitterSummaryGroup();
 
+	virtual bool SupportsIcon() const override { return true; }
+	virtual const FSlateBrush* GetIconBrush() const override;
+	virtual bool GetCanExpandInOverview() const override { return false; }
+	virtual bool GetShouldShowInStack() const override { return false; }
+
 protected:
 	void RefreshChildrenInternal(const TArray<UNiagaraStackEntry*>& CurrentChildren, TArray<UNiagaraStackEntry*>& NewChildren, TArray<FStackIssue>& NewIssues) override;
 
diff --git a/Engine/Plugins/FX/Niagara/Source/NiagaraEditor/Public/ViewModels/Stack/NiagaraStackEntry.h b/Engine/Plugins/FX/Niagara/Source/NiagaraEditor/Public/ViewModels/Stack/NiagaraStackEntry.h
--- a/Engine/Plugins/FX/Niagara/Source/NiagaraEditor/Public/ViewModels/Stack/NiagaraStackEntry.h
+++ b/Engine/Plugins/FX/Niagara/Source/NiagaraEditor/Public/ViewModels/Stack/NiagaraStackEntry.h
@@ -251,6 +251,8 @@ public:
 
 	virtual UObject* GetDisplayedObject() const;
 
+	virtual FGuid GetSelectionId() const;
+
 	UNiagaraStackEditorData& GetStackEditorData() const;
 
 	FString GetStackEditorDataKey() const;
@@ -290,6 +292,8 @@ public:
 
 	virtual bool GetShouldShowInStack() const;
 
+	virtual bool GetShouldShowInOverview() const;
+
 	void GetFilteredChildren(TArray<UNiagaraStackEntry*>& OutFilteredChildren) const;
 	
 	void GetUnfilteredChildren(TArray<UNiagaraStackEntry*>& OutUnfilteredChildren) const;
diff --git a/Engine/Plugins/FX/Niagara/Source/NiagaraEditor/Public/ViewModels/Stack/NiagaraStackEventHandlerGroup.h b/Engine/Plugins/FX/Niagara/Source/NiagaraEditor/Public/ViewModels/Stack/NiagaraStackEventHandlerGroup.h
deleted file mode 100644
diff --git a/Engine/Plugins/FX/Niagara/Source/NiagaraEditor/Public/ViewModels/Stack/NiagaraStackEventScriptItemGroup.h b/Engine/Plugins/FX/Niagara/Source/NiagaraEditor/Public/ViewModels/Stack/NiagaraStackEventScriptItemGroup.h
--- a/Engine/Plugins/FX/Niagara/Source/NiagaraEditor/Public/ViewModels/Stack/NiagaraStackEventScriptItemGroup.h
+++ b/Engine/Plugins/FX/Niagara/Source/NiagaraEditor/Public/ViewModels/Stack/NiagaraStackEventScriptItemGroup.h
@@ -27,6 +27,8 @@ public:
 	virtual bool TestCanResetToBaseWithMessage(FText& OutCanResetToBaseMessage) const override;
 	virtual void ResetToBase() override;
 
+	virtual FGuid GetSelectionId() const override { return EventScriptUsageId; }
+
 	FGuid GetEventScriptUsageId() const { return EventScriptUsageId; };
 
 protected:
diff --git a/Engine/Plugins/FX/Niagara/Source/NiagaraEditor/Public/ViewModels/Stack/NiagaraStackItem.h b/Engine/Plugins/FX/Niagara/Source/NiagaraEditor/Public/ViewModels/Stack/NiagaraStackItem.h
--- a/Engine/Plugins/FX/Niagara/Source/NiagaraEditor/Public/ViewModels/Stack/NiagaraStackItem.h
+++ b/Engine/Plugins/FX/Niagara/Source/NiagaraEditor/Public/ViewModels/Stack/NiagaraStackItem.h
@@ -33,14 +33,9 @@ public:
 
 	virtual bool SupportsChangeEnabled() const { return false; }
 	void SetIsEnabled(bool bInIsEnabled);
-
 	
 	virtual bool SupportsIcon() const { return false; }
 	virtual const FSlateBrush* GetIconBrush() const;
-	/** The custom overview indent can be used as a backup indent in the overview stack.
-	 *  If an icon exists, the icon will be used instead, and if no custom overview indent is specified and no icon is set, we use a default indent
-	 */
-	virtual TOptional<float> GetCustomOverviewIndent() { return TOptional<float>();}
 
 	virtual bool SupportsResetToBase() const { return false; }
 	virtual bool TestCanResetToBaseWithMessage(FText& OutCanResetToBaseMessage) const { return false; }
diff --git a/Engine/Plugins/FX/Niagara/Source/NiagaraEditor/Public/ViewModels/Stack/NiagaraStackItemGroup.h b/Engine/Plugins/FX/Niagara/Source/NiagaraEditor/Public/ViewModels/Stack/NiagaraStackItemGroup.h
--- a/Engine/Plugins/FX/Niagara/Source/NiagaraEditor/Public/ViewModels/Stack/NiagaraStackItemGroup.h
+++ b/Engine/Plugins/FX/Niagara/Source/NiagaraEditor/Public/ViewModels/Stack/NiagaraStackItemGroup.h
@@ -9,6 +9,7 @@
 
 class FNiagaraEmitterHandleViewModel;
 class UNiagaraStackItemGroupFooter;
+struct FSlateBrush;
 
 UCLASS()
 class NIAGARAEDITOR_API UNiagaraStackItemGroup : public UNiagaraStackEntry
@@ -28,6 +29,12 @@ public:
 	virtual void SetIsEnabled(bool bEnabled) {}
 	virtual bool SupportsChangeEnabled() const { return false; }
 
+	virtual bool SupportsIcon() const { return false; }
+	virtual const FSlateBrush* GetIconBrush() const { return nullptr; };
+
+	virtual bool SupportsSecondaryIcon() const { return false; }
+	virtual const FSlateBrush* GetSecondaryIconBrush() const { return nullptr; };
+
 	INiagaraStackItemGroupAddUtilities* GetAddUtilities() const;
 
 	uint32 GetRecursiveStackIssuesCount() const;
diff --git a/Engine/Plugins/FX/Niagara/Source/NiagaraEditor/Public/ViewModels/Stack/NiagaraStackModuleItem.h b/Engine/Plugins/FX/Niagara/Source/NiagaraEditor/Public/ViewModels/Stack/NiagaraStackModuleItem.h
--- a/Engine/Plugins/FX/Niagara/Source/NiagaraEditor/Public/ViewModels/Stack/NiagaraStackModuleItem.h
+++ b/Engine/Plugins/FX/Niagara/Source/NiagaraEditor/Public/ViewModels/Stack/NiagaraStackModuleItem.h
@@ -43,8 +43,6 @@ public:
 
 	virtual bool SupportsRename() const override { return true; }
 
-	virtual TOptional<float> GetCustomOverviewIndent() override;
-
 	virtual bool SupportsChangeEnabled() const override { return true; }
 	virtual bool GetIsEnabled() const override;
 
diff --git a/Engine/Plugins/FX/Niagara/Source/NiagaraEditor/Public/ViewModels/Stack/NiagaraStackRoot.h b/Engine/Plugins/FX/Niagara/Source/NiagaraEditor/Public/ViewModels/Stack/NiagaraStackRoot.h
--- a/Engine/Plugins/FX/Niagara/Source/NiagaraEditor/Public/ViewModels/Stack/NiagaraStackRoot.h
+++ b/Engine/Plugins/FX/Niagara/Source/NiagaraEditor/Public/ViewModels/Stack/NiagaraStackRoot.h
@@ -6,11 +6,10 @@
 #include "NiagaraStackRoot.generated.h"
 
 class FNiagaraEmitterViewModel;
-class UNiagaraStackSystemSettingsGroup;
-class UNiagaraStackEmitterSettingsGroup;
+class UNiagaraStackSystemPropertiesGroup;
+class UNiagaraStackSystemUserParametersGroup;
+class UNiagaraStackEmitterPropertiesGroup;
 class UNiagaraStackScriptItemGroup;
-class UNiagaraStackEventHandlerGroup;
-class UNiagaraStackSimulationStagesGroup;
 class UNiagaraStackRenderItemGroup;
 class UNiagaraStackEmitterSummaryGroup;
 class UNiagaraStackSummaryViewCollapseButton;
@@ -41,7 +40,10 @@ private:
 
 private:
 	UPROPERTY()
-	TObjectPtr<UNiagaraStackSystemSettingsGroup> SystemSettingsGroup;
+	TObjectPtr<UNiagaraStackSystemPropertiesGroup> SystemPropertiesGroup;
+
+	UPROPERTY()
+	TObjectPtr<UNiagaraStackSystemUserParametersGroup> SystemUserParametersGroup;
 
 	UPROPERTY()
 	TObjectPtr<UNiagaraStackScriptItemGroup> SystemSpawnGroup;
@@ -50,7 +52,7 @@ private:
 	TObjectPtr<UNiagaraStackScriptItemGroup> SystemUpdateGroup;
 
 	UPROPERTY()
-	TObjectPtr<UNiagaraStackEmitterSettingsGroup> EmitterSettingsGroup;
+	TObjectPtr<UNiagaraStackEmitterPropertiesGroup> EmitterPropertiesGroup;
 
 	UPROPERTY()
 	TObjectPtr<UNiagaraStackEmitterSummaryGroup> EmitterSummaryGroup;
@@ -67,12 +69,6 @@ private:
 	UPROPERTY()
 	TObjectPtr<UNiagaraStackScriptItemGroup> ParticleUpdateGroup;
 
-	UPROPERTY()
-	TObjectPtr<UNiagaraStackEventHandlerGroup> AddEventHandlerGroup;
-
-	UPROPERTY()
-	TObjectPtr<UNiagaraStackSimulationStagesGroup> AddSimulationStageGroup;
-
 	UPROPERTY()
 	TObjectPtr<UNiagaraStackRenderItemGroup> RenderGroup;
 
diff --git a/Engine/Plugins/FX/Niagara/Source/NiagaraEditor/Public/ViewModels/Stack/NiagaraStackSimulationStagesGroup.h b/Engine/Plugins/FX/Niagara/Source/NiagaraEditor/Public/ViewModels/Stack/NiagaraStackSimulationStagesGroup.h
deleted file mode 100644
diff --git a/Engine/Plugins/FX/Niagara/Source/NiagaraEditor/Public/ViewModels/Stack/NiagaraStackSystemPropertiesItem.h b/Engine/Plugins/FX/Niagara/Source/NiagaraEditor/Public/ViewModels/Stack/NiagaraStackSystemPropertiesItem.h
--- a/Engine/Plugins/FX/Niagara/Source/NiagaraEditor/Public/ViewModels/Stack/NiagaraStackSystemPropertiesItem.h
+++ b/Engine/Plugins/FX/Niagara/Source/NiagaraEditor/Public/ViewModels/Stack/NiagaraStackSystemPropertiesItem.h
@@ -22,6 +22,8 @@ public:
 
 	virtual bool IsExpandedByDefault() const override;
 
+	virtual bool GetShouldShowInOverview() const override { return false; }
+
 protected:
 	virtual void RefreshChildrenInternal(const TArray<UNiagaraStackEntry*>& CurrentChildren, TArray<UNiagaraStackEntry*>& NewChildren, TArray<FStackIssue>& NewIssues) override;
 
diff --git a/Engine/Plugins/FX/Niagara/Source/NiagaraEditor/Public/ViewModels/Stack/NiagaraStackSystemSettingsGroup.h b/Engine/Plugins/FX/Niagara/Source/NiagaraEditor/Public/ViewModels/Stack/NiagaraStackSystemSettingsGroup.h
--- a/Engine/Plugins/FX/Niagara/Source/NiagaraEditor/Public/ViewModels/Stack/NiagaraStackSystemSettingsGroup.h
+++ b/Engine/Plugins/FX/Niagara/Source/NiagaraEditor/Public/ViewModels/Stack/NiagaraStackSystemSettingsGroup.h
@@ -12,12 +12,34 @@
 class FNiagaraScriptViewModel;
 
 UCLASS()
-class NIAGARAEDITOR_API UNiagaraStackSystemSettingsGroup : public UNiagaraStackItemGroup
+class NIAGARAEDITOR_API UNiagaraStackSystemPropertiesGroup : public UNiagaraStackItemGroup
 {
 	GENERATED_BODY()
 		
 public:
-	void Initialize(FRequiredEntryData InRequiredEntryData,	UObject* InOwner, FNiagaraParameterStore* InParameterStore);
+	void Initialize(FRequiredEntryData InRequiredEntryData);
+
+	virtual bool SupportsIcon() const override { return true; }
+	virtual const FSlateBrush* GetIconBrush() const override;
+	virtual bool GetCanExpandInOverview() const override { return false; }
+	virtual bool GetShouldShowInStack() const override { return false; }
+
+protected:
+	virtual void RefreshChildrenInternal(const TArray<UNiagaraStackEntry*>& CurrentChildren, TArray<UNiagaraStackEntry*>& NewChildren, TArray<FStackIssue>& NewIssues) override;
+};
+
+UCLASS()
+class NIAGARAEDITOR_API UNiagaraStackSystemUserParametersGroup : public UNiagaraStackItemGroup
+{
+	GENERATED_BODY()
+
+public:
+	void Initialize(FRequiredEntryData InRequiredEntryData, UObject* InOwner, FNiagaraParameterStore* InParameterStore);
+
+	virtual bool SupportsIcon() const override { return true; }
+	virtual const FSlateBrush* GetIconBrush() const override;
+	virtual bool GetCanExpandInOverview() const override { return false; }
+	virtual bool GetShouldShowInStack() const override { return false; }
 
 protected:
 	virtual void RefreshChildrenInternal(const TArray<UNiagaraStackEntry*>& CurrentChildren, TArray<UNiagaraStackEntry*>& NewChildren, TArray<FStackIssue>& NewIssues) override;
@@ -33,17 +55,21 @@ private:
 };
 
 UCLASS()
-class UNiagaraStackParameterStoreItem : public UNiagaraStackItem
+class NIAGARAEDITOR_API UNiagaraStackParameterStoreItem : public UNiagaraStackItem
 {
 	GENERATED_BODY()
 
 public:
-	void Initialize(FRequiredEntryData InRequiredEntryData, UObject* InOwner, FNiagaraParameterStore* InParameterStore);
+	void Initialize(FRequiredEntryData InRequiredEntryData, UObject* InOwner, FNiagaraParameterStore* InParameterStore, INiagaraStackItemGroupAddUtilities* InGroupAddUtilities);
 
 	virtual FText GetDisplayName() const override;
 
 	virtual FText GetTooltipText() const override;
 
+	virtual bool GetShouldShowInOverview() const override { return false; }
+
+	INiagaraStackItemGroupAddUtilities* GetGroupAddUtilities() const { return GroupAddUtilities; }
+
 protected:
 	virtual void RefreshChildrenInternal(const TArray<UNiagaraStackEntry*>& CurrentChildren, TArray<UNiagaraStackEntry*>& NewChildren, TArray<FStackIssue>& NewIssues) override;
 	virtual void FinalizeInternal() override;
@@ -54,5 +80,6 @@ private:
 private:
 	TWeakObjectPtr<UObject> Owner;
 	FNiagaraParameterStore* ParameterStore;
+	INiagaraStackItemGroupAddUtilities* GroupAddUtilities;
 	FDelegateHandle ParameterStoreChangedHandle;
 };
\ No newline at end of file
diff --git a/Engine/Plugins/FX/Niagara/Source/NiagaraEditorWidgets/NiagaraEditorWidgets.Build.cs b/Engine/Plugins/FX/Niagara/Source/NiagaraEditorWidgets/NiagaraEditorWidgets.Build.cs
--- a/Engine/Plugins/FX/Niagara/Source/NiagaraEditorWidgets/NiagaraEditorWidgets.Build.cs
+++ b/Engine/Plugins/FX/Niagara/Source/NiagaraEditorWidgets/NiagaraEditorWidgets.Build.cs
@@ -31,6 +31,7 @@ public class NiagaraEditorWidgets : ModuleRules
             "EditorWidgets",
 			"ApplicationCore",
 			"CurveEditor",
+			"ToolWidgets",
 		});
 
 		PrivateIncludePathModuleNames.AddRange(new string[] {
diff --git a/Engine/Plugins/FX/Niagara/Source/NiagaraEditorWidgets/Private/NiagaraEditorWidgetsStyle.cpp b/Engine/Plugins/FX/Niagara/Source/NiagaraEditorWidgets/Private/NiagaraEditorWidgetsStyle.cpp
--- a/Engine/Plugins/FX/Niagara/Source/NiagaraEditorWidgets/Private/NiagaraEditorWidgetsStyle.cpp
+++ b/Engine/Plugins/FX/Niagara/Source/NiagaraEditorWidgets/Private/NiagaraEditorWidgetsStyle.cpp
@@ -8,6 +8,7 @@
 #include "Styling/SlateStyleRegistry.h"
 #include "Styling/SlateTypes.h"
 #include "Styling/CoreStyle.h"
+#include "Styling/StarshipCoreStyle.h"
 #include "Classes/EditorStyleSettings.h"
 #include "Styling/AppStyle.h"
 #include "Styling/StyleColors.h"
@@ -76,8 +77,7 @@ FNiagaraEditorWidgetsStyle::FNiagaraEditorWidgetsStyle() : FSlateStyleSet("Niaga
 	const FTextBlockStyle CategoryText = FAppStyle::Get().GetWidgetStyle<FTextBlockStyle>("DetailsView.CategoryTextStyle");
 	FSlateFontInfo CategoryFont = FAppStyle::Get().GetFontStyle(TEXT("PropertyWindow.BoldFont"));
 	FTextBlockStyle StackGroupText = FTextBlockStyle(CategoryText)
-		.SetFont(CategoryFont)
-		.SetTransformPolicy(ETextTransformPolicy::ToUpper);
+		.SetFont(CategoryFont);
 
 	Set("NiagaraEditor.Stack.GroupText", StackGroupText);
 
@@ -92,8 +92,7 @@ FNiagaraEditorWidgetsStyle::FNiagaraEditorWidgetsStyle() : FSlateStyleSet("Niaga
 
 	FSlateFontInfo NormalFont = FAppStyle::Get().GetFontStyle(TEXT("PropertyWindow.NormalFont"));
 	FTextBlockStyle StackItemText = FTextBlockStyle(CategoryText)
-		.SetFont(NormalFont)
-		.SetTransformPolicy(ETextTransformPolicy::ToUpper);
+		.SetFont(NormalFont);
 
 	Set("NiagaraEditor.Stack.ItemText", StackItemText);
 	
@@ -297,6 +296,14 @@ FNiagaraEditorWidgetsStyle::FNiagaraEditorWidgetsStyle() : FSlateStyleSet("Niaga
 
 	FSlateBoxBrush StackRowSelectionBrush = BOX_BRUSH("Icons/StackSelectionBorder", FMargin(2.0f / 8.0f), FStyleColors::Select);
 	FSlateBoxBrush StackRowSubduedSelectionBrush = BOX_BRUSH("Icons/StackSelectionBorder", FMargin(2.0f / 8.0f), FStyleColors::SelectInactive);
+	
+	Set("Niagara.CompactStackIssue.Error", new IMAGE_BRUSH("Icons/StackIssueError", Icon8x8));
+	Set("Niagara.CompactStackIssue.Warning", new IMAGE_BRUSH("Icons/StackIssueWarning", Icon8x8));
+	Set("Niagara.CompactStackIssue.Info", new IMAGE_BRUSH("Icons/StackIssueInfo", Icon8x8));
+	Set("Niagara.CompactStackIssue.Message", new IMAGE_BRUSH("Icons/StackIssueMessage", Icon8x8));
+
+	Set("Niagara.TableViewRowBorder", new BOX_BRUSH("Icons/Row", FMargin(3.0f / 8.0f), FStyleColors::Recessed));
+
 	Set("NiagaraEditor.Stack.TableViewRow", FTableRowStyle(NormalTableRowStyle)
 		.SetUseParentRowBrush(true)
  		.SetEvenRowBackgroundBrush(FSlateColorBrush(FStyleColors::Transparent))
@@ -309,11 +316,29 @@ FNiagaraEditorWidgetsStyle::FNiagaraEditorWidgetsStyle() : FSlateStyleSet("Niaga
 		.SetInactiveBrush(StackRowSubduedSelectionBrush)
 		.SetInactiveHoveredBrush(StackRowSelectionBrush));
 
-	Set("NiagaraEditor.SystemOverview.TableViewRow", FTableRowStyle(NormalTableRowStyle)
+	Set("NiagaraEditor.SystemOverview.TableViewRow.Item", FTableRowStyle(NormalTableRowStyle)
 		.SetEvenRowBackgroundBrush(FSlateColorBrush(FStyleColors::Panel))
 		.SetOddRowBackgroundBrush(FSlateColorBrush(FStyleColors::Panel))
-		.SetSelectedTextColor(FStyleColors::ForegroundInverted)
-		.SetInactiveBrush(CORE_IMAGE_BRUSH("Common/Selection", Icon8x8, GetDefault<UEditorStyleSettings>()->GetSubduedSelectionColor())));
+		.SetEvenRowBackgroundHoveredBrush(CORE_IMAGE_BRUSH("Common/Selection", Icon8x8, FLinearColor(1.0f, 1.0f, 1.0f, 0.1f)))
+		.SetOddRowBackgroundHoveredBrush(CORE_IMAGE_BRUSH("Common/Selection", Icon8x8, FLinearColor(1.0f, 1.0f, 1.0f, 0.1f)))
+		.SetSelectorFocusedBrush(CORE_BORDER_BRUSH("Common/Selector", FMargin(4.f / 16.f), FStarshipCoreStyle::GetCoreStyle().GetSlateColor("SelectorColor")))
+		.SetActiveBrush(FSlateColorBrush(FStyleColors::Select))
+		.SetActiveHoveredBrush(FSlateColorBrush(FStyleColors::Select))
+		.SetInactiveBrush(FSlateColorBrush(FStyleColors::SelectInactive))
+		.SetActiveHighlightedBrush(FSlateColorBrush(FStyleColors::PrimaryHover))
+		.SetTextColor(FStyleColors::Foreground));
+
+	Set("NiagaraEditor.SystemOverview.TableViewRow.Group", FTableRowStyle(NormalTableRowStyle)
+		.SetEvenRowBackgroundBrush(FSlateColorBrush(FStyleColors::Header))
+		.SetOddRowBackgroundBrush(FSlateColorBrush(FStyleColors::Header))
+		.SetEvenRowBackgroundHoveredBrush(CORE_IMAGE_BRUSH("Common/Selection", Icon8x8, FLinearColor(1.0f, 1.0f, 1.0f, 0.1f)))
+		.SetOddRowBackgroundHoveredBrush(CORE_IMAGE_BRUSH("Common/Selection", Icon8x8, FLinearColor(1.0f, 1.0f, 1.0f, 0.1f)))
+		.SetSelectorFocusedBrush(CORE_BORDER_BRUSH("Common/Selector", FMargin(4.f / 16.f), FStarshipCoreStyle::GetCoreStyle().GetSlateColor("SelectorColor")))
+		.SetActiveBrush(FSlateColorBrush(FStyleColors::Select))
+		.SetActiveHoveredBrush(FSlateColorBrush(FStyleColors::Select))
+		.SetInactiveBrush(FSlateColorBrush(FStyleColors::SelectInactive))
+		.SetActiveHighlightedBrush(FSlateColorBrush(FStyleColors::PrimaryHover))
+		.SetTextColor(FStyleColors::Foreground));
 
 	Set("NiagaraEditor.Stack.BackgroundColor", FLinearColor(FColor(96, 96, 96)));
 	Set("NiagaraEditor.Stack.Item.HeaderBackgroundColor", FLinearColor(FColor(48, 48, 48)));
@@ -425,5 +450,63 @@ FNiagaraEditorWidgetsStyle::FNiagaraEditorWidgetsStyle() : FSlateStyleSet("Niaga
 	Set("NiagaraEditor.Scope.Particles", FLinearColor(FColor(131, 218, 9)));
 	Set("NiagaraEditor.Scope.ScriptPersistent", FLinearColor(FColor(255, 247, 77)));
 	Set("NiagaraEditor.Scope.ScriptTransient", FLinearColor(FColor(255, 247, 77)));
+
+	FButtonStyle SimpleButtonStyle = FAppStyle::Get().GetWidgetStyle<FButtonStyle>("SimpleButton");
+	float NormalDesaturate = .25f;
+	float DisabledDesaturate = .5;
+
+	FLinearColor SystemIconColor = GetColor("NiagaraEditor.Stack.IconColor.System");
+	FButtonStyle SystemAddItemButtonStyle = SimpleButtonStyle;
+	SystemAddItemButtonStyle
+		.SetNormalForeground(SystemIconColor.Desaturate(NormalDesaturate))
+		.SetHoveredForeground(SystemIconColor)
+		.SetPressedForeground(SystemIconColor.Desaturate(NormalDesaturate))
+		.SetDisabledForeground(SystemIconColor.Desaturate(DisabledDesaturate));
+
+	FLinearColor EmitterIconColor = GetColor("NiagaraEditor.Stack.IconColor.Emitter");
+	FButtonStyle EmitterAddItemButtonStyle = SimpleButtonStyle;
+	EmitterAddItemButtonStyle
+		.SetNormalForeground(EmitterIconColor.Desaturate(NormalDesaturate))
+		.SetHoveredForeground(EmitterIconColor)
+		.SetPressedForeground(EmitterIconColor.Desaturate(NormalDesaturate))
+		.SetDisabledForeground(EmitterIconColor.Desaturate(DisabledDesaturate));
+
+	FLinearColor ParticleIconColor = GetColor("NiagaraEditor.Stack.IconColor.Particle");
+	FButtonStyle ParticleAddItemButtonStyle = SimpleButtonStyle;
+	ParticleAddItemButtonStyle
+		.SetNormalForeground(ParticleIconColor.Desaturate(NormalDesaturate))
+		.SetHoveredForeground(ParticleIconColor)
+		.SetPressedForeground(ParticleIconColor.Desaturate(NormalDesaturate))
+		.SetDisabledForeground(ParticleIconColor.Desaturate(DisabledDesaturate));
+
+	FLinearColor RenderIconColor = GetColor("NiagaraEditor.Stack.IconColor.Render");
+	FButtonStyle RenderAddItemButtonStyle = SimpleButtonStyle;
+	RenderAddItemButtonStyle
+		.SetNormalForeground(RenderIconColor.Desaturate(NormalDesaturate))
+		.SetHoveredForeground(RenderIconColor)
+		.SetPressedForeground(RenderIconColor.Desaturate(NormalDesaturate))
+		.SetDisabledForeground(RenderIconColor.Desaturate(DisabledDesaturate));
+
+	FButtonStyle LabeledAddItemButtonStyle = FAppStyle::Get().GetWidgetStyle<FButtonStyle>("Button");
+	LabeledAddItemButtonStyle
+		.SetNormal(FSlateRoundedBoxBrush(FStyleColors::Panel, 4))
+		.SetHovered(FSlateRoundedBoxBrush(FStyleColors::Hover, 4))
+		.SetPressed(FSlateRoundedBoxBrush(FStyleColors::Recessed, 4))
+		.SetNormalPadding(FMargin(2.0f, 0.0f))
+		.SetPressedPadding(FMargin(2.0f, 0.0f));
+
+	FButtonStyle StackSimpleButtonStyle = SimpleButtonStyle;
+	StackSimpleButtonStyle
+		.SetNormalForeground(FStyleColors::Foreground)
+		.SetHoveredForeground(FStyleColors::ForegroundHover)
+		.SetPressedForeground(FStyleColors::Foreground)
+		.SetDisabledForeground(FStyleColors::Foreground.GetSpecifiedColor().Desaturate(DisabledDesaturate));
+
+	Set("NiagaraEditor.Stack.AddItemButton.System", SystemAddItemButtonStyle);
+	Set("NiagaraEditor.Stack.AddItemButton.Emitter", EmitterAddItemButtonStyle);
+	Set("NiagaraEditor.Stack.AddItemButton.Particle", ParticleAddItemButtonStyle);
+	Set("NiagaraEditor.Stack.AddItemButton.Render", RenderAddItemButtonStyle);
+	Set("NiagaraEditor.Stack.LabeledAddItemButton", LabeledAddItemButtonStyle);
+	Set("NiagaraEditor.Stack.SimpleButton", StackSimpleButtonStyle);
 }
 
diff --git a/Engine/Plugins/FX/Niagara/Source/NiagaraEditorWidgets/Private/NiagaraEditorWidgetsUtilities.cpp b/Engine/Plugins/FX/Niagara/Source/NiagaraEditorWidgets/Private/NiagaraEditorWidgetsUtilities.cpp
--- a/Engine/Plugins/FX/Niagara/Source/NiagaraEditorWidgets/Private/NiagaraEditorWidgetsUtilities.cpp
+++ b/Engine/Plugins/FX/Niagara/Source/NiagaraEditorWidgets/Private/NiagaraEditorWidgetsUtilities.cpp
@@ -139,7 +139,29 @@ FName FNiagaraStackEditorWidgetsUtilities::GetIconColorNameForExecutionCategory(
 	}
 }
 
-
+FName FNiagaraStackEditorWidgetsUtilities::GetAddItemButtonStyleNameForExecutionCategory(FName ExecutionCategoryName)
+{
+	if (ExecutionCategoryName == UNiagaraStackEntry::FExecutionCategoryNames::System)
+	{
+		return "NiagaraEditor.Stack.AddItemButton.System";
+	}
+	else if (ExecutionCategoryName == UNiagaraStackEntry::FExecutionCategoryNames::Emitter)
+	{
+		return "NiagaraEditor.Stack.AddItemButton.Emitter";
+	}
+	else if (ExecutionCategoryName == UNiagaraStackEntry::FExecutionCategoryNames::Particle)
+	{
+		return "NiagaraEditor.Stack.AddItemButton.Particle";
+	}
+	else if (ExecutionCategoryName == UNiagaraStackEntry::FExecutionCategoryNames::Render)
+	{
+		return "NiagaraEditor.Stack.AddItemButton.Render";
+	}
+	else
+	{
+		return NAME_None;
+	}
+}
 
 FText FNiagaraStackEditorWidgetsUtilities::GetIconTextForInputMode(UNiagaraStackFunctionInput::EValueMode InputValueMode)
 {
diff --git a/Engine/Plugins/FX/Niagara/Source/NiagaraEditorWidgets/Private/NiagaraEditorWidgetsUtilities.h b/Engine/Plugins/FX/Niagara/Source/NiagaraEditorWidgets/Private/NiagaraEditorWidgetsUtilities.h
--- a/Engine/Plugins/FX/Niagara/Source/NiagaraEditorWidgets/Private/NiagaraEditorWidgetsUtilities.h
+++ b/Engine/Plugins/FX/Niagara/Source/NiagaraEditorWidgets/Private/NiagaraEditorWidgetsUtilities.h
@@ -22,6 +22,8 @@ namespace FNiagaraStackEditorWidgetsUtilities
 
 	FName GetIconColorNameForExecutionCategory(FName ExecutionCategoryName);
 
+	FName GetAddItemButtonStyleNameForExecutionCategory(FName ExecutionCategoryName);
+
 	FText GetIconTextForInputMode(UNiagaraStackFunctionInput::EValueMode InputValueMode);
 
 	FText GetIconToolTipForInputMode(UNiagaraStackFunctionInput::EValueMode InputValueMode);
diff --git a/Engine/Plugins/FX/Niagara/Source/NiagaraEditorWidgets/Private/SNiagaraOverviewStack.cpp b/Engine/Plugins/FX/Niagara/Source/NiagaraEditorWidgets/Private/SNiagaraOverviewStack.cpp
--- a/Engine/Plugins/FX/Niagara/Source/NiagaraEditorWidgets/Private/SNiagaraOverviewStack.cpp
+++ b/Engine/Plugins/FX/Niagara/Source/NiagaraEditorWidgets/Private/SNiagaraOverviewStack.cpp
@@ -1,56 +1,60 @@
 // Copyright Epic Games, Inc. All Rights Reserved.
 
 #include "SNiagaraOverviewStack.h"
+
+#include "EditorStyleSet.h"
+#include "Framework/Application/SlateApplication.h"
+#include "Framework/MultiBox/MultiBoxBuilder.h"
+#include "NiagaraEditorModule.h"
+#include "NiagaraEditorStyle.h"
+#include "NiagaraEditorWidgetsStyle.h"
+#include "NiagaraEditorWidgetsUtilities.h"
+#include "NiagaraNodeAssignment.h"
+#include "NiagaraNodeFunctionCall.h"
+#include "NiagaraStackCommandContext.h"
+#include "NiagaraStackEditorData.h"
 #include "NiagaraSystem.h"
 #include "NiagaraSystemEditorData.h"
-#include "NiagaraStackEditorData.h"
-#include "ViewModels/NiagaraEmitterViewModel.h"
-#include "ViewModels/NiagaraSystemViewModel.h"
+#include "ScopedTransaction.h"
+#include "SNiagaraOverviewInlineParameterBox.h"
+#include "SNiagaraStack.h"
+#include "Stack/SNiagaraStackIssueIcon.h"
+#include "Stack/SNiagaraStackItemGroupAddButton.h"
+#include "Stack/SNiagaraStackRowPerfWidget.h"
+#include "Styling/StyleColors.h"
+#include "Subsystems/AssetEditorSubsystem.h"
 #include "ViewModels/NiagaraEmitterHandleViewModel.h"
-#include "ViewModels/NiagaraSystemSelectionViewModel.h"
-#include "ViewModels/NiagaraScratchPadViewModel.h"
+#include "ViewModels/NiagaraEmitterViewModel.h"
 #include "ViewModels/NiagaraScratchPadScriptViewModel.h"
+#include "ViewModels/NiagaraScratchPadViewModel.h"
+#include "ViewModels/NiagaraSystemSelectionViewModel.h"
+#include "ViewModels/NiagaraSystemViewModel.h"
 #include "ViewModels/Stack/NiagaraStackEntry.h"
-#include "ViewModels/Stack/NiagaraStackItemGroup.h"
 #include "ViewModels/Stack/NiagaraStackItem.h"
+#include "ViewModels/Stack/NiagaraStackItemGroup.h"
 #include "ViewModels/Stack/NiagaraStackModuleItem.h"
 #include "ViewModels/Stack/NiagaraStackViewModel.h"
-#include "NiagaraEditorWidgetsStyle.h"
-#include "NiagaraEditorWidgetsUtilities.h"
-#include "SNiagaraStack.h"
-#include "Stack/SNiagaraStackItemGroupAddButton.h"
-#include "Stack/SNiagaraStackIssueIcon.h"
-#include "Stack/SNiagaraStackRowPerfWidget.h"
-#include "NiagaraStackCommandContext.h"
-#include "NiagaraNodeFunctionCall.h"
-#include "NiagaraEditorModule.h"
-#include "NiagaraNodeAssignment.h"
-#include "Widgets/SNiagaraParameterName.h"
-#include "SNiagaraOverviewInlineParameterBox.h"
-#include "Widgets/Layout/SScrollBox.h"
-#include "Widgets/Layout/SGridPanel.h"
-#include "Widgets/Layout/SScaleBox.h"
+#include "Widgets/Colors/SColorBlock.h"
 #include "Widgets/Images/SImage.h"
 #include "Widgets/Input/SButton.h"
 #include "Widgets/Input/SCheckBox.h"
-#include "Widgets/SToolTip.h"
-#include "EditorStyleSet.h"
-#include "Framework/Application/SlateApplication.h"
-#include "Framework/MultiBox/MultiBoxBuilder.h"
-#include "Subsystems/AssetEditorSubsystem.h"
-#include "Widgets/Colors/SColorBlock.h"
-#include "NiagaraEditorStyle.h"
-#include "ScopedTransaction.h"
-#include "Styling/StyleColors.h"
 #include "Widgets/Input/SCheckBox.h"
+#include "Widgets/Layout/SGridPanel.h"
+#include "Widgets/Layout/SScaleBox.h"
+#include "Widgets/Layout/SScrollBox.h"
+#include "Widgets/Layout/SSpacer.h"
 #include "Widgets/Layout/SWrapBox.h"
+#include "Widgets/SNiagaraParameterName.h"
+#include "Widgets/SToolTip.h"
 
 #define LOCTEXT_NAMESPACE "NiagaraOverviewStack"
 
 class SNiagaraSystemOverviewEntryListRow : public STableRow<UNiagaraStackEntry*>
 {
-	SLATE_BEGIN_ARGS(SNiagaraSystemOverviewEntryListRow) {}
-		SLATE_ATTRIBUTE(EVisibility, IssueIconVisibility)
+	SLATE_BEGIN_ARGS(SNiagaraSystemOverviewEntryListRow)
+		: _Style(&GetDefaultStyle())
+		{}
+		SLATE_STYLE_ARGUMENT(FTableRowStyle, Style)
 		SLATE_EVENT(FOnDragDetected, OnDragDetected)
 		SLATE_EVENT(FOnTableRowDragLeave, OnDragLeave)
 		SLATE_EVENT(FOnCanAcceptDrop, OnCanAcceptDrop)
@@ -68,26 +72,10 @@ class SNiagaraSystemOverviewEntryListRow : public STableRow<UNiagaraStackEntry*>
 		StackViewModel = InStackViewModel;
 		StackEntry = InStackEntry;
 		StackCommandContext = InStackCommandContext;
-		IssueIconVisibility = InArgs._IssueIconVisibility;
 		FSlateColor IconColor = FNiagaraEditorWidgetsStyle::Get().GetColor(FNiagaraStackEditorWidgetsUtilities::GetColorNameForExecutionCategory(StackEntry->GetExecutionCategoryName()));
 
 		ExpandedImage = FCoreStyle::Get().GetBrush("TreeArrow_Expanded");
 		CollapsedImage = FCoreStyle::Get().GetBrush("TreeArrow_Collapsed");
-		
-		FMargin ContentPadding;
-		if (StackEntry->IsA<UNiagaraStackItem>())
-		{
-			BackgroundColor = FStyleColors::Panel;
-			ContentPadding = FMargin(0, 1, 0, 1);
-		}
-		else
-		{
-			BackgroundColor = FStyleColors::Header;
-			ContentPadding = FMargin(0, 2, 0, 2);
-		}
-
-		DisabledBackgroundColor = FStyleColors::Recessed;
-		IsolatedBackgroundColor = FStyleColors::Select;
 
 		TSharedRef<FNiagaraSystemViewModel> SystemViewModel = StackEntry->GetSystemViewModel();
 		TSharedPtr<FNiagaraEmitterViewModel> EmitterViewModel = StackEntry->GetEmitterViewModel();
@@ -96,25 +84,23 @@ class SNiagaraSystemOverviewEntryListRow : public STableRow<UNiagaraStackEntry*>
 			EmitterHandleViewModel = SystemViewModel->GetEmitterHandleViewModelForEmitter(EmitterViewModel->GetEmitter());
 		}
 
-		StackEntry->OnStructureChanged().AddSP(this, &SNiagaraSystemOverviewEntryListRow::InvalidateChildrenCount);
+		StackEntry->OnStructureChanged().AddSP(this, &SNiagaraSystemOverviewEntryListRow::InvalidateCachedData);
 
+		TSharedPtr<SOverlay> RowOverlay;
 		STableRow<UNiagaraStackEntry*>::Construct(STableRow<UNiagaraStackEntry*>::FArguments()
-			.Style(FNiagaraEditorWidgetsStyle::Get(), "NiagaraEditor.SystemOverview.TableViewRow")
+			.Style(InArgs._Style)
 			.OnDragDetected(InArgs._OnDragDetected)
 			.OnDragLeave(InArgs._OnDragLeave)
 			.OnCanAcceptDrop(InArgs._OnCanAcceptDrop)
 			.OnAcceptDrop(InArgs._OnAcceptDrop)
 		[
-			SNew(SBorder)
-				.BorderImage(FEditorStyle::GetBrush("WhiteBrush"))
-				.BorderBackgroundColor(FStyleColors::Recessed)
-				.ToolTipText_UObject(StackEntry, &UNiagaraStackEntry::GetTooltipText)
-				.Padding(FMargin(0, 1, 0, 1))
-				[
+			SAssignNew(RowOverlay, SOverlay)
+			+ SOverlay::Slot()
+			[
 					SNew(SBorder)
-					.BorderImage(FEditorStyle::GetBrush("WhiteBrush"))
-					.BorderBackgroundColor(TAttribute<FSlateColor>(this, &SNiagaraSystemOverviewEntryListRow::GetBackgroundColor))
-					.Padding(ContentPadding)
+					.BorderImage(FNiagaraEditorWidgetsStyle::Get().GetBrush("Niagara.TableViewRowBorder"))
+					.ToolTipText_UObject(StackEntry, &UNiagaraStackEntry::GetTooltipText)
+					.Padding(FMargin(0, 0, 0, 1))
 					[
 						SNew(SHorizontalBox)
 						// Expand button
@@ -122,10 +108,10 @@ class SNiagaraSystemOverviewEntryListRow : public STableRow<UNiagaraStackEntry*>
 						.AutoWidth()
 						.VAlign(VAlign_Center)
 						.HAlign(HAlign_Center)
-						.Padding(3, 0, 1, 0)
+						.Padding(0, 0, 0, 0)
 						[					
 							SNew(SButton)
-							.ButtonStyle(FCoreStyle::Get(), "NoBorder")
+							.ButtonStyle(FNiagaraEditorWidgetsStyle::Get(), "NiagaraEditor.Stack.SimpleButton")
 							.Visibility_Lambda([=]()
 							{
 								if (StackEntry->GetCanExpandInOverview())
@@ -152,8 +138,7 @@ class SNiagaraSystemOverviewEntryListRow : public STableRow<UNiagaraStackEntry*>
 								StackEntry->SetIsExpandedInOverview(bWillBeExpanded);									
 								return FReply::Handled();
 							})
-							.ForegroundColor(FSlateColor::UseForeground())
-							.ContentPadding(2)
+							.ContentPadding(FMargin(0.0f, 1.0f))
 							[
 								SNew(SImage)
 								.Image(this, &SNiagaraSystemOverviewEntryListRow::GetExpandButtonImage)
@@ -165,18 +150,50 @@ class SNiagaraSystemOverviewEntryListRow : public STableRow<UNiagaraStackEntry*>
 						[
 							InArgs._Content.Widget
 						]
-						+ SHorizontalBox::Slot()
-						.AutoWidth()
-						.VAlign(VAlign_Center)
-						.Padding(2, 0, 2, 0)
-						[
-							SNew(SNiagaraStackIssueIcon, StackViewModel, StackEntry)
-							.Visibility(IssueIconVisibility)
-						]
 					]
+			]
+			// Stack Issue Highlight
+			+ SOverlay::Slot()
+			.HAlign(HAlign_Right)
+			.VAlign(VAlign_Center)
+			[
+				SNew(SBorder)
+				.BorderImage(FAppStyle::Get().GetBrush("WhiteBrush"))
+				.BorderBackgroundColor(this, &SNiagaraSystemOverviewEntryListRow::GetIsssueHighlightColor)
+				.Visibility_Static(&SNiagaraOverviewStack::GetIssueIconVisibility, StackEntry)
+				.Padding(FMargin(0, 0, 0, 1))
+				[
+					SNew(SBox)
+					.WidthOverride(2)
+					.HeightOverride(10)
 				]
+			]
 		],
 		InOwnerTableView);
+
+		if (StackEntry->IsA<UNiagaraStackItem>())
+		{
+			// Execution Category Item Highlight
+			RowOverlay->AddSlot()
+			.HAlign(HAlign_Left)
+			.VAlign(VAlign_Fill)
+			.Padding(FMargin(1, 0, 0, 0))
+			[
+				SNew(SBorder)
+				.BorderImage(FAppStyle::Get().GetBrush("WhiteBrush"))
+				.BorderBackgroundColor(IconColor)
+				.Padding(FMargin(0))
+				[
+					SNew(SBox)
+					.WidthOverride(2)
+				]
+			];
+		}
+	}
+
+	static const FTableRowStyle& GetDefaultStyle()
+	{
+		return FNiagaraEditorWidgetsStyle::Get().GetWidgetStyle<FTableRowStyle>("NiagaraEditor.SystemOverview.TableViewRow.Item");
 	}
 
 	virtual FReply OnMouseButtonDown(const FGeometry& MyGeometry, const FPointerEvent& MouseEvent) override
@@ -261,46 +278,18 @@ class SNiagaraSystemOverviewEntryListRow : public STableRow<UNiagaraStackEntry*>
 	}
 
 private:
-	FSlateColor GetBackgroundColor() const
-	{
-		if (IsSelected())
-		{
-			return FStyleColors::Select;
-		}
-		// robk - hack bandaid fix
-		if (StackEntry->IsFinalized())
-		{
-			return DisabledBackgroundColor;
-		}
-		TSharedRef<FNiagaraSystemViewModel> SystemViewModel = StackEntry->GetSystemViewModel();
-		if (SystemViewModel->GetSystem().GetIsolateEnabled())
-		{
-			TSharedPtr<FNiagaraEmitterHandleViewModel> EmitterHandleViewModelPtr = EmitterHandleViewModel.Pin();
-			if (EmitterHandleViewModelPtr.IsValid() && !EmitterHandleViewModelPtr->GetIsIsolated())
-			{
-				return IsolatedBackgroundColor;
-			}
-		}
-
-		return StackEntry->GetIsEnabledAndOwnerIsEnabled() ? BackgroundColor : DisabledBackgroundColor;
-	}
-
 	FMargin GetInnerContentPadding() const
 	{
-		if (IssueIconVisibility.Get() == EVisibility::Visible)
-		{
-			return FMargin(5, 0, 1, 0);
-		}
-		else
-		{
-			return FMargin(5, 0, 5, 0);
-		}
+		float InnerContentLeftPadding = StackEntry->IsA<UNiagaraStackItem>() ? 5 : 0;
+		float InnerContentRightPadding = 1;
+		return FMargin(InnerContentLeftPadding, 1, 1, 1);
 	}
 
 	const FSlateBrush* GetExpandButtonImage() const
 	{
 		return StackEntry->GetIsExpandedInOverview() ? ExpandedImage : CollapsedImage;
 	}
+	
 
 	uint32 GetOverviewChildrenCount() const
 	{
@@ -314,23 +303,50 @@ private:
         return ChildrenCount.GetValue();
 	}
 
-	void InvalidateChildrenCount(ENiagaraStructureChangedFlags Flags)
+	void InvalidateCachedData(ENiagaraStructureChangedFlags Flags)
 	{
 		ChildrenCount.Reset();
+		IssueHighlightColor.Reset();
+	}
+
+	FSlateColor GetIsssueHighlightColor() const
+	{
+		if (IssueHighlightColor.IsSet() == false)
+		{
+			if (StackEntry->GetTotalNumberOfErrorIssues() > 0)
+			{
+				IssueHighlightColor = FStyleColors::AccentRed.GetSpecifiedColor().Desaturate(.15f);
+			}
+			else if (StackEntry->GetTotalNumberOfWarningIssues() > 0)
+			{
+				IssueHighlightColor = FStyleColors::AccentYellow.GetSpecifiedColor().Desaturate(.15f);
+			}
+			else if (StackEntry->GetTotalNumberOfInfoIssues() > 0)
+			{
+				IssueHighlightColor = FStyleColors::AccentBlue.GetSpecifiedColor().Desaturate(.15f);
+			}
+			else if (StackEntry->GetTotalNumberOfCustomNotes() > 0)
+			{
+				IssueHighlightColor = FStyleColors::AccentGray.GetSpecifiedColor().Desaturate(.15f);
+			}
+			else
+			{
+				IssueHighlightColor = FStyleColors::Transparent;
+			}
+		}
+		return IssueHighlightColor.GetValue();
 	}
 
 private:
 	UNiagaraStackViewModel* StackViewModel;
 	UNiagaraStackEntry* StackEntry;
 	TSharedPtr<FNiagaraStackCommandContext> StackCommandContext;
-	FSlateColor BackgroundColor;
-	FSlateColor DisabledBackgroundColor;
-	FSlateColor IsolatedBackgroundColor;
 	TAttribute<EVisibility> IssueIconVisibility;
 	TWeakPtr<FNiagaraEmitterHandleViewModel> EmitterHandleViewModel;
 	const FSlateBrush* ExpandedImage = nullptr;
 	const FSlateBrush* CollapsedImage = nullptr;
 	mutable TOptional<uint32> ChildrenCount;
+	mutable TOptional<FSlateColor> IssueHighlightColor;
 };
 
 class SNiagaraSystemOverviewEnabledCheckBox : public SCompoundWidget
@@ -600,6 +616,7 @@ void SNiagaraOverviewStack::Construct(const FArguments& InArgs, UNiagaraStackVie
 	ChildSlot
 	[
 		SAssignNew(EntryListView, SListView<UNiagaraStackEntry*>)
+		.Clipping(EWidgetClipping::OnDemand)
 		.ListItemsSource(&FlattenedEntryList)
 		.OnGenerateRow(this, &SNiagaraOverviewStack::OnGenerateRowForEntry)
 		.OnSelectionChanged(this, &SNiagaraOverviewStack::OnSelectionChanged)
@@ -650,9 +667,24 @@ void SNiagaraOverviewStack::Tick(const FGeometry& AllottedGeometry, const double
 	RefreshEntryList();
 }
 
+EVisibility SNiagaraOverviewStack::GetIssueIconVisibility(UNiagaraStackEntry* StackEntry)
+{
+	bool bDisplayIcon;
+	if (StackEntry->IsA<UNiagaraStackItemGroup>() && StackEntry->GetCanExpandInOverview() && StackEntry->GetIsExpandedInOverview())
+	{
+		// If the entry is a group and it can expand and it is expanded, we only want to show the stack issue icon if the group itself has issues.
+		bDisplayIcon = StackEntry->GetIssues().Num() > 0;
+	}
+	else
+	{
+		bDisplayIcon = StackEntry->HasIssuesOrAnyChildHasIssues();
+	}
+	return bDisplayIcon ? EVisibility::Visible : EVisibility::Collapsed;
+}
+
 void SNiagaraOverviewStack::AddEntriesRecursive(UNiagaraStackEntry& EntryToAdd, TArray<UNiagaraStackEntry*>& EntryList, const TArray<UClass*>& AcceptableClasses, TArray<UNiagaraStackEntry*> ParentChain)
 {
-	if (AcceptableClasses.ContainsByPredicate([&] (UClass* Class) { return EntryToAdd.IsA(Class); }))
+	if (AcceptableClasses.ContainsByPredicate([&] (UClass* Class) { return EntryToAdd.IsA(Class); }) && EntryToAdd.GetShouldShowInOverview())
 	{
 		EntryList.Add(&EntryToAdd);
 		EntryObjectKeyToParentChain.Add(FObjectKey(&EntryToAdd), ParentChain);
@@ -718,15 +750,20 @@ TSharedRef<ITableRow> SNiagaraOverviewStack::OnGenerateRowForEntry(UNiagaraStack
 {
 	FVector2D IconSize = FNiagaraEditorWidgetsStyle::Get().GetVector("NiagaraEditor.Stack.IconSize");
 	TSharedPtr<SWidget> Content;
+	const FTableRowStyle* RowStyle = &SNiagaraSystemOverviewEntryListRow::GetDefaultStyle();
 	if (Item->IsA<UNiagaraStackItem>())
 	{
 		UNiagaraStackItem* StackItem = CastChecked<UNiagaraStackItem>(Item);
-		TSharedPtr<SWidget> IndentContent;
-
-		// we give icons priority for the indent content
+		TSharedRef<SHorizontalBox> ContentBox = SNew(SHorizontalBox);
 		if (StackItem->SupportsIcon())
 		{
-			IndentContent = SNew(SBox)
+			// Icon
+			ContentBox->AddSlot()
+			.Padding(0, 1, 0, 1)
+			.AutoWidth()
+			.VAlign(VAlign_Center)
+			[
+				SNew(SBox)
 				.IsEnabled_UObject(Item, &UNiagaraStackEntry::GetIsEnabledAndOwnerIsEnabled)
 				.WidthOverride(IconSize.X)
 				.HeightOverride(IconSize.Y)
@@ -735,46 +772,35 @@ TSharedRef<ITableRow> SNiagaraOverviewStack::OnGenerateRowForEntry(UNiagaraStack
 				[
 					SNew(SImage)
 					.Image_UObject(StackItem, &UNiagaraStackItem::GetIconBrush)
-				];
-		}
-		// if we didn't have an icon but we specified a custom indent, we use that instead
-		else if(StackItem->GetCustomOverviewIndent().IsSet())
-		{
-			float IndentValue = StackItem->GetCustomOverviewIndent().GetValue();
-			IndentContent = SNew(SBox)
-				.WidthOverride(IndentValue)
-				.HeightOverride(1.f);
-		}
-		// otherwise, use a default
-		else
-		{
-			IndentContent = SNew(SBox)
-				.WidthOverride(IconSize.X)
-				.HeightOverride(IconSize.Y);
+					.DesiredSizeOverride(FVector2D(14.0f, 14.0f))
+				]
+			];
 		}
-		
-		TSharedRef<SHorizontalBox> ContentBox = SNew(SHorizontalBox)
-			// Indent content
-			+ SHorizontalBox::Slot()
-			.Padding(0, 1, 2, 1)
-			.AutoWidth()
-			.VAlign(VAlign_Center)
-			[
-				IndentContent.ToSharedRef()
-			]
-			// Name
-			+ SHorizontalBox::Slot()
-			.VAlign(VAlign_Center)
-			.Padding(3, 2, 0, 2)
-			[
-				SNew(SNiagaraSystemOverviewItemName, StackItem)
-			]
-			+ SHorizontalBox::Slot()
-            .AutoWidth()
-            .Padding(2, 0, 2, 0)
-            [
-                SNew(SNiagaraStackRowPerfWidget, Item)
-            ];
+		// Name
+		ContentBox->AddSlot()
+		.VAlign(VAlign_Center)
+		.AutoWidth()
+		.Padding(4, 2, 0, 2)
+		[
+			SNew(SNiagaraSystemOverviewItemName, StackItem)
+		];
+		// Issue Icon
+		ContentBox->AddSlot()
+		.VAlign(VAlign_Center)
+		.AutoWidth()
+		.Padding(4.0f, 0.0f)
+		[
+			SNew(SNiagaraStackIssueIcon, StackViewModel, StackItem)
+			.IconMode(SNiagaraStackIssueIcon::EIconMode::Compact)
+			.Visibility_Static(&SNiagaraOverviewStack::GetIssueIconVisibility, Item)
+		];
+		// Perf widget
+		ContentBox->AddSlot()
+        .AutoWidth()
+        .Padding(2, 0, 2, 0)
+        [
+            SNew(SNiagaraStackRowPerfWidget, Item)
+        ];
 
 		TSharedRef<SHorizontalBox> OptionsBox = SNew(SHorizontalBox);
 		
@@ -893,47 +919,94 @@ TSharedRef<ITableRow> SNiagaraOverviewStack::OnGenerateRowForEntry(UNiagaraStack
 			];
 		}
 		
-		Content = ContentBox;		
+		Content = ContentBox;	
+		RowStyle = &FNiagaraEditorWidgetsStyle::Get().GetWidgetStyle<FTableRowStyle>("NiagaraEditor.SystemOverview.TableViewRow.Item");
 	}
 	else if (Item->IsA<UNiagaraStackItemGroup>())
 	{
 		UNiagaraStackItemGroup* StackItemGroup = CastChecked<UNiagaraStackItemGroup>(Item);
-		TSharedRef<SHorizontalBox> ContentBox = SNew(SHorizontalBox)
-			// Execution category icon
-			+ SHorizontalBox::Slot()
-			.Padding(0, 0, 2, 0)
-			.VAlign(VAlign_Center)
-			.AutoWidth()
-			[
-				SNew(SBox)
-				.WidthOverride(IconSize.X)
-				.HeightOverride(IconSize.Y)
-				.HAlign(HAlign_Center)
+		TSharedRef<SHorizontalBox> ContentBox = SNew(SHorizontalBox);
+		
+		// Icon
+		if (StackItemGroup->SupportsIcon())
+		{
+			ContentBox->AddSlot()
 				.VAlign(VAlign_Center)
+				.AutoWidth()
+				.Padding(2, 1, 2, 1)
 				[
-					SNew(SImage)
-					.Image(FNiagaraEditorWidgetsStyle::Get().GetBrush(FNiagaraStackEditorWidgetsUtilities::GetIconNameForExecutionSubcategory(Item->GetExecutionSubcategoryName(), true)))
-					.ColorAndOpacity(FNiagaraEditorWidgetsStyle::Get().GetColor(FNiagaraStackEditorWidgetsUtilities::GetIconColorNameForExecutionCategory(Item->GetExecutionCategoryName())))
+					SNew(SBox)
 					.IsEnabled_UObject(Item, &UNiagaraStackEntry::GetIsEnabledAndOwnerIsEnabled)
-				]
-			]
-			// Name
-			+ SHorizontalBox::Slot()
+					.WidthOverride(IconSize.X)
+					.HeightOverride(IconSize.Y)
+					.HAlign(HAlign_Center)
+					.VAlign(VAlign_Center)
+					[
+						SNew(SImage)
+						.Image_UObject(StackItemGroup, &UNiagaraStackItemGroup::GetIconBrush)
+						.DesiredSizeOverride(FVector2D(14.0f, 14.0f))
+					]
+				];
+		}
+
+		// Name
+		ContentBox->AddSlot()
 			.VAlign(VAlign_Center)
-			.Padding(3, 2, 0, 2)
+			.AutoWidth()
+			.Padding(1, 2, 0, 2)
 			[
 				SNew(STextBlock)
 				.TextStyle(FNiagaraEditorWidgetsStyle::Get(), "NiagaraEditor.Stack.GroupText")
 				.Text_UObject(Item, &UNiagaraStackEntry::GetDisplayName)
 				.IsEnabled_UObject(Item, &UNiagaraStackEntry::GetIsEnabledAndOwnerIsEnabled)
 				.ColorAndOpacity(FSlateColor::UseForeground())
-			]
-			+ SHorizontalBox::Slot()
-            .AutoWidth()
-            .Padding(6, 0, 5, 0)
-            [
-                SNew(SNiagaraStackRowPerfWidget, Item)
-            ];
+			];
+
+		// Issue Icon
+		ContentBox->AddSlot()
+			.VAlign(VAlign_Center)
+			.AutoWidth()
+			.Padding(4.0f, 0.0f)
+			[
+				SNew(SNiagaraStackIssueIcon, StackViewModel, StackItemGroup)
+				.IconMode(SNiagaraStackIssueIcon::EIconMode::Compact)
+				.Visibility_Static(&SNiagaraOverviewStack::GetIssueIconVisibility, Item)
+			];
+
+		// Secondary Icon
+		if (StackItemGroup->SupportsSecondaryIcon())
+		{
+			ContentBox->AddSlot()
+				.VAlign(VAlign_Center)
+				.AutoWidth()
+				.Padding(4, 1, 0, 1)
+				[
+					SNew(SBox)
+					.IsEnabled_UObject(Item, &UNiagaraStackEntry::GetIsEnabledAndOwnerIsEnabled)
+					.WidthOverride(IconSize.X)
+					.HeightOverride(IconSize.Y)
+					.HAlign(HAlign_Center)
+					.VAlign(VAlign_Center)
+					[
+						SNew(SImage)
+						.Image_UObject(StackItemGroup, &UNiagaraStackItemGroup::GetSecondaryIconBrush)
+					]
+				];
+		}
+
+		// Spacer
+		ContentBox->AddSlot()
+			[
+				SNew(SSpacer)
+			];
+
+		// Perf 
+		ContentBox->AddSlot()
+			.AutoWidth()
+			.Padding(6, 0, 5, 0)
+			[
+				SNew(SNiagaraStackRowPerfWidget, Item)
+			];
 
 		// Delete button
 		if (StackItemGroup->SupportsDelete())
@@ -946,16 +1019,16 @@ TSharedRef<ITableRow> SNiagaraOverviewStack::OnGenerateRowForEntry(UNiagaraStack
 					SNew(SButton)
 					.ButtonStyle(FEditorStyle::Get(), "HoverHintOnly")
 					.IsFocusable(false)
-					.ForegroundColor(FNiagaraEditorWidgetsStyle::Get().GetColor("NiagaraEditor.Stack.ForegroundColor"))
 					.ToolTipText(this, &SNiagaraOverviewStack::GetItemGroupDeleteButtonToolTip, StackItemGroup)
 					.OnClicked(this, &SNiagaraOverviewStack::OnItemGroupDeleteClicked, StackItemGroup)
 					.IsEnabled(this, &SNiagaraOverviewStack::GetItemGroupDeleteButtonIsEnabled, StackItemGroup)
 					.Visibility(this, &SNiagaraOverviewStack::GetItemGroupDeleteButtonVisibility, StackItemGroup)
 					.Content()
 					[
-						SNew(STextBlock)
-						.Font(FEditorStyle::Get().GetFontStyle("FontAwesome.10"))
-						.Text(FText::FromString(FString(TEXT("\xf1f8"))))
+						SNew(SImage)
+						.Image(FAppStyle::Get().GetBrush("Icons.Delete"))
+						.ColorAndOpacity(FSlateColor::UseForeground())
+						.DesiredSizeOverride(FVector2D(14.0f, 14.0f))
 					]
 				];
 		}
@@ -979,28 +1052,29 @@ TSharedRef<ITableRow> SNiagaraOverviewStack::OnGenerateRowForEntry(UNiagaraStack
 		if (StackItemGroup->GetAddUtilities() != nullptr)
 		{
 			ContentBox->AddSlot()
+			.VAlign(VAlign_Center)
 			.AutoWidth()
 			.Padding(0, 0, 1, 0)
 			[
-				SNew(SNiagaraStackItemGroupAddButton, *StackItemGroup)
+				SNew(SNiagaraStackItemGroupAddButton, StackItemGroup, StackItemGroup->GetAddUtilities())
 				.Width(22)
 			];
 		}
 
 		Content = ContentBox;
+		RowStyle = &FNiagaraEditorWidgetsStyle::Get().GetWidgetStyle<FTableRowStyle>("NiagaraEditor.SystemOverview.TableViewRow.Group");
 	}
 	else
 	{
 		Content = SNullWidget::NullWidget;
 	}
 
-
 	return SNew(SNiagaraSystemOverviewEntryListRow, StackViewModel, Item, StackCommandContext.ToSharedRef(), OwnerTable)
+		.Style(RowStyle)
 		.OnDragDetected(this, &SNiagaraOverviewStack::OnRowDragDetected, TWeakObjectPtr<UNiagaraStackEntry>(Item))
 		.OnDragLeave(this, &SNiagaraOverviewStack::OnRowDragLeave)
 		.OnCanAcceptDrop(this, &SNiagaraOverviewStack::OnRowCanAcceptDrop)
 		.OnAcceptDrop(this, &SNiagaraOverviewStack::OnRowAcceptDrop)
-		.IssueIconVisibility(this, &SNiagaraOverviewStack::GetIssueIconVisibility)
 	[
 		Content.ToSharedRef()
 	];
@@ -1151,11 +1225,6 @@ FReply SNiagaraOverviewStack::OnRowAcceptDrop(const FDragDropEvent& InDragDropEv
 	return bHandled ? FReply::Handled() : FReply::Unhandled();
 }
 
-EVisibility SNiagaraOverviewStack::GetIssueIconVisibility() const
-{
-	return StackViewModel->HasIssues() ? EVisibility::Visible : EVisibility::Collapsed;
-}
-
 void SNiagaraOverviewStack::OnItemGroupEnabledStateChanged(ECheckBoxState InCheckState, UNiagaraStackItemGroup* Group)
 {
 	Group->SetIsEnabled(InCheckState == ECheckBoxState::Checked);
diff --git a/Engine/Plugins/FX/Niagara/Source/NiagaraEditorWidgets/Private/SNiagaraOverviewStack.h b/Engine/Plugins/FX/Niagara/Source/NiagaraEditorWidgets/Private/SNiagaraOverviewStack.h
--- a/Engine/Plugins/FX/Niagara/Source/NiagaraEditorWidgets/Private/SNiagaraOverviewStack.h
+++ b/Engine/Plugins/FX/Niagara/Source/NiagaraEditorWidgets/Private/SNiagaraOverviewStack.h
@@ -33,6 +33,8 @@ public:
 
 	virtual void Tick(const FGeometry& AllottedGeometry, const double InCurrentTime, const float InDeltaTime) override;
 
+	static EVisibility GetIssueIconVisibility(UNiagaraStackEntry* StackEntry);
+
 private:
 	void AddEntriesRecursive(UNiagaraStackEntry& EntryToAdd, TArray<UNiagaraStackEntry*>& EntryList, const TArray<UClass*>& AcceptableClasses, TArray<UNiagaraStackEntry*> ParentChain);
 
@@ -68,8 +70,6 @@ private:
 
 	FReply OnRowAcceptDrop(const FDragDropEvent& InDragDropEvent, EItemDropZone InDropZone, UNiagaraStackEntry* InTargetEntry);
 
-	EVisibility GetIssueIconVisibility() const;
-
 	void OnItemGroupEnabledStateChanged(ECheckBoxState InCheckState, UNiagaraStackItemGroup* Group);
 	ECheckBoxState ItemGroupCheckEnabledStatus(UNiagaraStackItemGroup* Group) const;
 	bool GetItemGroupEnabledCheckboxEnabled(UNiagaraStackItemGroup* Group) const;
diff --git a/Engine/Plugins/FX/Niagara/Source/NiagaraEditorWidgets/Private/SNiagaraStack.cpp b/Engine/Plugins/FX/Niagara/Source/NiagaraEditorWidgets/Private/SNiagaraStack.cpp
--- a/Engine/Plugins/FX/Niagara/Source/NiagaraEditorWidgets/Private/SNiagaraStack.cpp
+++ b/Engine/Plugins/FX/Niagara/Source/NiagaraEditorWidgets/Private/SNiagaraStack.cpp
@@ -22,6 +22,7 @@
 #include "ViewModels/Stack/NiagaraStackItemFooter.h"
 #include "ViewModels/Stack/NiagaraStackFunctionInputCollection.h"
 #include "ViewModels/Stack/NiagaraStackModuleItem.h"
+#include "ViewModels/Stack/NiagaraStackSystemSettingsGroup.h"
 #include "ViewModels/Stack/NiagaraStackParameterStoreEntry.h"
 #include "Framework/Application/SlateApplication.h"
 #include "Framework/Commands/GenericCommands.h"
@@ -41,6 +42,7 @@
 #include "Stack/SNiagaraStackItemFooter.h"
 #include "Stack/SNiagaraStackItemGroup.h"
 #include "Stack/SNiagaraStackModuleItem.h"
+#include "Stack/SNiagaraStackParameterStoreItem.h"
 #include "Stack/SNiagaraStackParameterStoreEntryName.h"
 #include "Stack/SNiagaraStackParameterStoreEntryValue.h"
 #include "Stack/SNiagaraStackTableRow.h"
@@ -930,6 +932,11 @@ SNiagaraStack::FRowWidgets SNiagaraStack::ConstructNameAndValueWidgetsForItem(UN
 		UNiagaraStackItemFooter* ItemExpander = CastChecked<UNiagaraStackItemFooter>(Item);
 		return FRowWidgets(SNew(SNiagaraStackItemFooter, *ItemExpander));
 	}
+	else if (Item->IsA<UNiagaraStackParameterStoreItem>())
+	{
+		UNiagaraStackParameterStoreItem* StackEntry = CastChecked<UNiagaraStackParameterStoreItem>(Item);
+		return FRowWidgets(SNew(SNiagaraStackParameterStoreItem, *StackEntry, StackViewModel));
+	}
 	else if (Item->IsA<UNiagaraStackParameterStoreEntry>())
 	{
 		UNiagaraStackParameterStoreEntry* StackEntry = CastChecked<UNiagaraStackParameterStoreEntry>(Item);
diff --git a/Engine/Plugins/FX/Niagara/Source/NiagaraEditorWidgets/Private/Stack/SNiagaraStackIssueIcon.cpp b/Engine/Plugins/FX/Niagara/Source/NiagaraEditorWidgets/Private/Stack/SNiagaraStackIssueIcon.cpp
--- a/Engine/Plugins/FX/Niagara/Source/NiagaraEditorWidgets/Private/Stack/SNiagaraStackIssueIcon.cpp
+++ b/Engine/Plugins/FX/Niagara/Source/NiagaraEditorWidgets/Private/Stack/SNiagaraStackIssueIcon.cpp
@@ -11,6 +11,7 @@
 #include "NiagaraEditorStyle.h"
 #include "Internationalization/Text.h"
 #include "NiagaraEditorWidgetsStyle.h"
+#include "Styling/StyleColors.h"
 
 #define LOCTEXT_NAMESPACE "SNiagaraStackIssueIcon"
 
@@ -18,42 +19,67 @@ void SNiagaraStackIssueIcon::Construct(const FArguments& InArgs, UNiagaraStackVi
 {
 	StackViewModel = InStackViewModel;
 	StackEntry = InStackEntry;
+	IconMode = InArgs._IconMode;
 	if (StackEntry != nullptr)
 	{
 		StackEntry->OnStructureChanged().AddSP(this, &SNiagaraStackIssueIcon::UpdateFromEntry);
 	}
 
-	TSharedRef<SWidget> IconWidget =
-		SNew(SBox)
-		.IsEnabled(this, &SNiagaraStackIssueIcon::GetIconIsEnabled)
-		.ToolTipText(this, &SNiagaraStackIssueIcon::GetIconToolTip)
-		.HeightOverride(16)
-		.WidthOverride(16)
-		.HAlign(HAlign_Center)
-		.VAlign(VAlign_Center)
-		[
-			SNew(SImage)
-			.Image(this, &SNiagaraStackIssueIcon::GetIconBrush)
-		];
-
-	if (InArgs._OnClicked.IsBound())
+	TSharedPtr<SWidget> IconWidget;
+	if (IconMode == EIconMode::Normal)
 	{
-		TSharedRef<SButton> IconButton =
-			SNew(SButton)
-			.ContentPadding(0)
-			.ButtonStyle(FAppStyle::Get(), "SimpleButton")
-			.OnClicked(InArgs._OnClicked)
-			.Content()
+		IconWidget =
+			SNew(SBox)
+			.IsEnabled(this, &SNiagaraStackIssueIcon::GetIconIsEnabled)
+			.ToolTipText(this, &SNiagaraStackIssueIcon::GetIconToolTip)
+			.HeightOverride(16)
+			.WidthOverride(16)
+			.HAlign(HAlign_Center)
+			.VAlign(VAlign_Center)
 			[
-				IconWidget
+				SNew(SImage)
+				.Image(this, &SNiagaraStackIssueIcon::GetIconBrush)
 			];
 
-		IconWidget = IconButton;
+		if (InArgs._OnClicked.IsBound())
+		{
+			TSharedRef<SButton> IconButton =
+				SNew(SButton)
+				.ContentPadding(0)
+				.ButtonStyle(FAppStyle::Get(), "SimpleButton")
+				.OnClicked(InArgs._OnClicked)
+				.Content()
+				[
+					IconWidget.ToSharedRef()
+				];
+
+			IconWidget = IconButton;
+		}
+	}
+	else if (IconMode == EIconMode::Compact)
+	{
+		IconWidget = 
+			SNew(SBorder)
+			.BorderImage(FNiagaraEditorWidgetsStyle::Get().GetBrush("NiagaraEditor.Stack.ModuleHighlightLarge"))
+			.BorderBackgroundColor(FStyleColors::Panel)
+			.IsEnabled(this, &SNiagaraStackIssueIcon::GetIconIsEnabled)
+			.ToolTipText(this, &SNiagaraStackIssueIcon::GetIconToolTip)
+			[
+				SNew(SBox)
+				.HeightOverride(10)
+				.WidthOverride(10)
+				.HAlign(HAlign_Center)
+				.VAlign(VAlign_Center)
+				[
+					SNew(SImage)
+					.Image(this, &SNiagaraStackIssueIcon::GetIconBrush)
+				]
+			];
 	}
 
 	ChildSlot
 	[
-		IconWidget
+		IconWidget.ToSharedRef()
 	];
 	UpdateFromEntry(ENiagaraStructureChangedFlags::StructureChanged);
 }
@@ -228,19 +254,27 @@ void SNiagaraStackIssueIcon::UpdateFromEntry(ENiagaraStructureChangedFlags Flags
 
 	if (StackEntry->GetTotalNumberOfErrorIssues() > 0)
 	{
-		IconBrush = FAppStyle::Get().GetBrush("Icons.ErrorWithColor");
+		IconBrush = IconMode == EIconMode::Compact
+			? FNiagaraEditorWidgetsStyle::Get().GetBrush("Niagara.CompactStackIssue.Error")
+			: FAppStyle::Get().GetBrush("Icons.ErrorWithColor");
 	}
 	else if (StackEntry->GetTotalNumberOfWarningIssues() > 0)
 	{
-		IconBrush = FAppStyle::Get().GetBrush("Icons.WarningWithColor");
+		IconBrush = IconMode == EIconMode::Compact
+			? FNiagaraEditorWidgetsStyle::Get().GetBrush("Niagara.CompactStackIssue.Warning")
+			: FAppStyle::Get().GetBrush("Icons.WarningWithColor");
 	}
 	else if (StackEntry->GetTotalNumberOfInfoIssues() > 0)
 	{
-		IconBrush = FAppStyle::Get().GetBrush("Icons.InfoWithColor");
+		IconBrush = IconMode == EIconMode::Compact
+			? FNiagaraEditorWidgetsStyle::Get().GetBrush("Niagara.CompactStackIssue.Info")
+			: FAppStyle::Get().GetBrush("Icons.InfoWithColor");
 	}
 	else if (StackEntry->GetTotalNumberOfCustomNotes() > 0)
 	{
-		IconBrush = FNiagaraEditorStyle::Get().GetBrush("NiagaraEditor.Message.CustomNote");
+		IconBrush = IconMode == EIconMode::Compact
+			? FNiagaraEditorWidgetsStyle::Get().GetBrush("Niagara.CompactStackIssue.Message")
+			: FNiagaraEditorStyle::Get().GetBrush("NiagaraEditor.Message.CustomNote");
 	}
 
 	IconToolTipCache.Reset();
diff --git a/Engine/Plugins/FX/Niagara/Source/NiagaraEditorWidgets/Private/Stack/SNiagaraStackIssueIcon.h b/Engine/Plugins/FX/Niagara/Source/NiagaraEditorWidgets/Private/Stack/SNiagaraStackIssueIcon.h
--- a/Engine/Plugins/FX/Niagara/Source/NiagaraEditorWidgets/Private/Stack/SNiagaraStackIssueIcon.h
+++ b/Engine/Plugins/FX/Niagara/Source/NiagaraEditorWidgets/Private/Stack/SNiagaraStackIssueIcon.h
@@ -13,7 +13,17 @@ class UNiagaraStackEntry;
 class SNiagaraStackIssueIcon : public SCompoundWidget
 {
 public:
-	SLATE_BEGIN_ARGS(SNiagaraStackIssueIcon) {}
+	enum class EIconMode
+	{
+		Normal,
+		Compact
+	};
+
+public:
+	SLATE_BEGIN_ARGS(SNiagaraStackIssueIcon) 
+		: _IconMode(EIconMode::Normal)
+		{}
+		SLATE_ARGUMENT(EIconMode, IconMode)
 		SLATE_EVENT(FOnClicked, OnClicked)
 	SLATE_END_ARGS()
 
@@ -26,6 +36,7 @@ private:
 
 	const FSlateBrush* GetIconBrush() const;
 	FText GetIconToolTip() const;
+	EVisibility GetCompactIconBorderVisibility() const;
 
 	void UpdateFromEntry(ENiagaraStructureChangedFlags Flags);
 
@@ -36,4 +47,5 @@ private:
 
 	UNiagaraStackViewModel* StackViewModel;
 	TWeakObjectPtr<UNiagaraStackEntry> StackEntry;
+	EIconMode IconMode;
 };
\ No newline at end of file
diff --git a/Engine/Plugins/FX/Niagara/Source/NiagaraEditorWidgets/Private/Stack/SNiagaraStackItemGroup.cpp b/Engine/Plugins/FX/Niagara/Source/NiagaraEditorWidgets/Private/Stack/SNiagaraStackItemGroup.cpp
--- a/Engine/Plugins/FX/Niagara/Source/NiagaraEditorWidgets/Private/Stack/SNiagaraStackItemGroup.cpp
+++ b/Engine/Plugins/FX/Niagara/Source/NiagaraEditorWidgets/Private/Stack/SNiagaraStackItemGroup.cpp
@@ -88,7 +88,7 @@ TSharedRef<SWidget> SNiagaraStackItemGroup::ConstructAddButton()
 	INiagaraStackItemGroupAddUtilities* AddUtilities = Group->GetAddUtilities();
 	if (AddUtilities != nullptr)
 	{
-		return SNew(SNiagaraStackItemGroupAddButton, *Group);
+		return SNew(SNiagaraStackItemGroupAddButton, Group, AddUtilities);
 	}
 	return SNullWidget::NullWidget;
 }
diff --git a/Engine/Plugins/FX/Niagara/Source/NiagaraEditorWidgets/Private/Stack/SNiagaraStackItemGroupAddButton.cpp b/Engine/Plugins/FX/Niagara/Source/NiagaraEditorWidgets/Private/Stack/SNiagaraStackItemGroupAddButton.cpp
--- a/Engine/Plugins/FX/Niagara/Source/NiagaraEditorWidgets/Private/Stack/SNiagaraStackItemGroupAddButton.cpp
+++ b/Engine/Plugins/FX/Niagara/Source/NiagaraEditorWidgets/Private/Stack/SNiagaraStackItemGroupAddButton.cpp
@@ -1,16 +1,17 @@
 // Copyright Epic Games, Inc. All Rights Reserved.
 
 #include "Stack/SNiagaraStackItemGroupAddButton.h"
+
+#include "EditorFontGlyphs.h"
+#include "EditorStyleSet.h"
+#include "NiagaraEditorWidgetsStyle.h"
+#include "NiagaraEditorWidgetsUtilities.h"
 #include "Stack/SNiagaraStackItemGroupAddMenu.h"
 #include "ViewModels/Stack/INiagaraStackItemGroupAddUtilities.h"
 #include "ViewModels/Stack/NiagaraStackItemGroup.h"
-#include "NiagaraEditorWidgetsUtilities.h"
-#include "NiagaraEditorWidgetsStyle.h"
-#include "EditorStyleSet.h"
-
+#include "Widgets/Input/SButton.h"
 #include "Widgets/Input/SComboButton.h"
 #include "Widgets/Input/SEditableTextBox.h"
-#include "Widgets/Input/SButton.h"
 #include "Widgets/Layout/SBox.h"
 #include "Widgets/Text/STextBlock.h"
 
@@ -18,43 +19,76 @@
 
 const float SNiagaraStackItemGroupAddButton::TextIconSize = 16;
 
-void SNiagaraStackItemGroupAddButton::Construct(const FArguments& InArgs, UNiagaraStackItemGroup& InStackItemGroup)
+void SNiagaraStackItemGroupAddButton::Construct(const FArguments& InArgs, UNiagaraStackEntry* InSourceEntry, INiagaraStackItemGroupAddUtilities* InAddUtilities)
 {
+	SourceEntryWeak = InSourceEntry;
+	AddUtilities = InAddUtilities;
 	TSharedPtr<SWidget> Content;
-	if(InStackItemGroup.GetAddUtilities() != nullptr)
+	if(InSourceEntry != nullptr && InAddUtilities != nullptr)
 	{
-		INiagaraStackItemGroupAddUtilities* AddUtilities = InStackItemGroup.GetAddUtilities();
+		TSharedPtr<SWidget> ButtonContent;
+		const FButtonStyle* ButtonStyle = nullptr;
+		if (AddUtilities->GetShowLabel())
+		{
+			ButtonContent = SNew(SHorizontalBox)
+				+ SHorizontalBox::Slot()
+				.VAlign(VAlign_Center)
+				.AutoWidth()
+				.Padding(FMargin(1.0f))
+				[
+					SNew(STextBlock)
+					.TextStyle(FEditorStyle::Get(), "NormalText.Important")
+					.Font(FEditorStyle::Get().GetFontStyle("FontAwesome.9"))
+					.Text(FEditorFontGlyphs::Plus)
+					.ColorAndOpacity(FNiagaraEditorWidgetsStyle::Get().GetColor(
+						FNiagaraStackEditorWidgetsUtilities::GetIconColorNameForExecutionCategory(InSourceEntry->GetExecutionCategoryName())))
+				]
+				+ SHorizontalBox::Slot()
+				.VAlign(VAlign_Center)
+				.Padding(FMargin(2.0f))
+				[
+					SNew(STextBlock)
+					.Text(AddUtilities->GetAddItemName())
+					.TextStyle(FNiagaraEditorStyle::Get(), "NiagaraEditor.ParameterText")
+				];
+			ButtonStyle = &FNiagaraEditorWidgetsStyle::Get().GetWidgetStyle<FButtonStyle>("NiagaraEditor.Stack.LabeledAddItemButton");
+		}
+		else
+		{
+			ButtonContent = SNew(SImage)
+				.Image(FAppStyle::Get().GetBrush("Icons.PlusCircle"))
+				.ColorAndOpacity(FSlateColor::UseForeground());
+			ButtonStyle = &FNiagaraEditorWidgetsStyle::Get().GetWidgetStyle<FButtonStyle>(
+				FNiagaraStackEditorWidgetsUtilities::GetAddItemButtonStyleNameForExecutionCategory(InSourceEntry->GetExecutionCategoryName()));
+		}
+
 		FText AddToGroupFormat = LOCTEXT("AddToGroupFormat", "Add a new {0} to this group.");
 		if (AddUtilities->GetAddMode() == INiagaraStackItemGroupAddUtilities::AddFromAction)
 		{
 			Content = SAssignNew(AddActionButton, SComboButton)
-				.ButtonStyle(FAppStyle::Get(), "SimpleButton")
+				.ButtonStyle(ButtonStyle)
 				.ToolTipText(FText::Format(AddToGroupFormat, AddUtilities->GetAddItemName()))
 				.HasDownArrow(false)
 				.OnGetMenuContent(this, &SNiagaraStackItemGroupAddButton::GetAddMenu)
-				.IsEnabled_UObject(&InStackItemGroup, &UNiagaraStackEntry::GetIsEnabledAndOwnerIsEnabled)
+				.IsEnabled_UObject(InSourceEntry, &UNiagaraStackEntry::GetIsEnabledAndOwnerIsEnabled)
 				.ContentPadding(1.0f)
 				.MenuPlacement(MenuPlacement_BelowRightAnchor)
 				.ButtonContent()
 				[
-					SNew(SImage)
-					.Image(FAppStyle::Get().GetBrush("Icons.PlusCircle"))
-					.ColorAndOpacity(FNiagaraEditorWidgetsStyle::Get().GetColor(FNiagaraStackEditorWidgetsUtilities::GetIconColorNameForExecutionCategory(InStackItemGroup.GetExecutionCategoryName())))
+					ButtonContent.ToSharedRef()
 				];
 		}
 		else if (AddUtilities->GetAddMode() == INiagaraStackItemGroupAddUtilities::AddDirectly)
 		{
 			Content = SNew(SButton)
-				.ButtonStyle(FAppStyle::Get(), "SimpleButton")
+				.ButtonStyle(ButtonStyle)
 				.ToolTipText(FText::Format(AddToGroupFormat, AddUtilities->GetAddItemName()))
-				.IsEnabled_UObject(&InStackItemGroup, &UNiagaraStackEntry::GetIsEnabledAndOwnerIsEnabled)
+				.IsEnabled_UObject(InSourceEntry, &UNiagaraStackEntry::GetIsEnabledAndOwnerIsEnabled)
 				.ContentPadding(1.0f)
 				.OnClicked(this, &SNiagaraStackItemGroupAddButton::AddDirectlyButtonClicked)
 				.Content()
 				[
-					SNew(SImage)
-					.Image(FAppStyle::Get().GetBrush("Icons.PlusCircle"))
-					.ColorAndOpacity(FNiagaraEditorWidgetsStyle::Get().GetColor(FNiagaraStackEditorWidgetsUtilities::GetIconColorNameForExecutionCategory(InStackItemGroup.GetExecutionCategoryName())))
+					ButtonContent.ToSharedRef()
 				];
 		}
 	}
@@ -68,15 +102,13 @@ void SNiagaraStackItemGroupAddButton::Construct(const FArguments& InArgs, UNiaga
 	[
 		Content.ToSharedRef()
 	];
-	
-	StackItemGroupWeak = &InStackItemGroup;
 }
 
 TSharedRef<SWidget> SNiagaraStackItemGroupAddButton::GetAddMenu()
 {
-	if(StackItemGroupWeak.IsValid())
+	if(SourceEntryWeak.IsValid())
 	{
-		TSharedRef<SNiagaraStackItemGroupAddMenu> AddMenu = SNew(SNiagaraStackItemGroupAddMenu, StackItemGroupWeak, StackItemGroupWeak->GetAddUtilities(), INDEX_NONE);
+		TSharedRef<SNiagaraStackItemGroupAddMenu> AddMenu = SNew(SNiagaraStackItemGroupAddMenu, SourceEntryWeak.Get(), AddUtilities, INDEX_NONE);
 		AddActionButton->SetMenuContentWidgetToFocus(AddMenu->GetFilterTextBox()->AsShared());
 		return AddMenu;
 	}
@@ -85,10 +117,10 @@ TSharedRef<SWidget> SNiagaraStackItemGroupAddButton::GetAddMenu()
 
 FReply SNiagaraStackItemGroupAddButton::AddDirectlyButtonClicked()
 {
-	if(StackItemGroupWeak.IsValid())
+	if(SourceEntryWeak.IsValid())
 	{
-		StackItemGroupWeak->GetAddUtilities()->AddItemDirectly();
-		StackItemGroupWeak->SetIsExpandedInOverview(true);
+		AddUtilities->AddItemDirectly();
+		SourceEntryWeak->SetIsExpandedInOverview(true);
 	}
 	return FReply::Handled();
 }
diff --git a/Engine/Plugins/FX/Niagara/Source/NiagaraEditorWidgets/Private/Stack/SNiagaraStackItemGroupAddButton.h b/Engine/Plugins/FX/Niagara/Source/NiagaraEditorWidgets/Private/Stack/SNiagaraStackItemGroupAddButton.h
--- a/Engine/Plugins/FX/Niagara/Source/NiagaraEditorWidgets/Private/Stack/SNiagaraStackItemGroupAddButton.h
+++ b/Engine/Plugins/FX/Niagara/Source/NiagaraEditorWidgets/Private/Stack/SNiagaraStackItemGroupAddButton.h
@@ -5,7 +5,8 @@
 #include "Widgets/DeclarativeSyntaxSupport.h"
 #include "Widgets/SCompoundWidget.h"
 
-class UNiagaraStackItemGroup;
+class UNiagaraStackEntry;
+class INiagaraStackItemGroupAddUtilities;
 class SComboButton;
 
 class SNiagaraStackItemGroupAddButton : public SCompoundWidget
@@ -17,7 +18,7 @@ public:
 		SLATE_ARGUMENT(float, Width)
 	SLATE_END_ARGS();
 
-	void Construct(const FArguments& InArgs, UNiagaraStackItemGroup& InStackItemGroup);
+	void Construct(const FArguments& InArgs, UNiagaraStackEntry* InSourceEntry, INiagaraStackItemGroupAddUtilities* InAddUtilities);
 
 private:
 	TSharedRef<SWidget> GetAddMenu();
@@ -25,7 +26,8 @@ private:
 	FReply AddDirectlyButtonClicked();
 
 private:
-	TWeakObjectPtr<UNiagaraStackItemGroup> StackItemGroupWeak;
+	TWeakObjectPtr<UNiagaraStackEntry> SourceEntryWeak;
+	INiagaraStackItemGroupAddUtilities* AddUtilities;
 	TSharedPtr<SComboButton> AddActionButton;
 	static const float TextIconSize;
 };
\ No newline at end of file
diff --git a/Engine/Plugins/FX/Niagara/Source/NiagaraEditorWidgets/Private/Stack/SNiagaraStackItemGroupAddMenu.cpp b/Engine/Plugins/FX/Niagara/Source/NiagaraEditorWidgets/Private/Stack/SNiagaraStackItemGroupAddMenu.cpp
--- a/Engine/Plugins/FX/Niagara/Source/NiagaraEditorWidgets/Private/Stack/SNiagaraStackItemGroupAddMenu.cpp
+++ b/Engine/Plugins/FX/Niagara/Source/NiagaraEditorWidgets/Private/Stack/SNiagaraStackItemGroupAddMenu.cpp
@@ -22,9 +22,9 @@
 
 bool SNiagaraStackItemGroupAddMenu::bLibraryOnly = true;
 
-void SNiagaraStackItemGroupAddMenu::Construct(const FArguments& InArgs, TWeakObjectPtr<UNiagaraStackItemGroup> InItemGroup, INiagaraStackItemGroupAddUtilities* InAddUtilities, int32 InInsertIndex)
+void SNiagaraStackItemGroupAddMenu::Construct(const FArguments& InArgs, UNiagaraStackEntry* InSourceEntry, INiagaraStackItemGroupAddUtilities* InAddUtilities, int32 InInsertIndex)
 {
-	StackItemGroup = InItemGroup;
+	SourceEntry = InSourceEntry;
 	AddUtilities = InAddUtilities;
 	InsertIndex = InInsertIndex;
 	bSetFocusOnNextTick = true;
@@ -82,7 +82,7 @@ void SNiagaraStackItemGroupAddMenu::Construct(const FArguments& InArgs, TWeakObj
 	                .AllowMultiselect(false)
 	                .OnDoesItemPassCustomFilter(this, &SNiagaraStackItemGroupAddMenu::DoesItemPassCustomFilter)
 	                .ClickActivateMode(EItemSelectorClickActivateMode::SingleClick)
-	                .ExpandInitially(false)
+	                .ExpandInitially(AddUtilities->GetAutoExpandAddActions())
 	                .OnGetSectionData_Lambda([](const ENiagaraMenuSections& Section)
 	                {
 	                    if(Section == ENiagaraMenuSections::Suggested)
@@ -137,9 +137,9 @@ void SNiagaraStackItemGroupAddMenu::OnItemActivated(const TSharedPtr<FNiagaraMen
 		FSlateApplication::Get().DismissAllMenus();
 		CurrentAction->Execute();
 
-		if(StackItemGroup.IsValid())
+		if(SourceEntry.IsValid())
 		{
-			StackItemGroup->SetIsExpandedInOverview(true);
+			SourceEntry->SetIsExpandedInOverview(true);
 		}
 	}	
 }
diff --git a/Engine/Plugins/FX/Niagara/Source/NiagaraEditorWidgets/Private/Stack/SNiagaraStackItemGroupAddMenu.h b/Engine/Plugins/FX/Niagara/Source/NiagaraEditorWidgets/Private/Stack/SNiagaraStackItemGroupAddMenu.h
--- a/Engine/Plugins/FX/Niagara/Source/NiagaraEditorWidgets/Private/Stack/SNiagaraStackItemGroupAddMenu.h
+++ b/Engine/Plugins/FX/Niagara/Source/NiagaraEditorWidgets/Private/Stack/SNiagaraStackItemGroupAddMenu.h
@@ -21,7 +21,7 @@ public:
 	SLATE_BEGIN_ARGS(SNiagaraStackItemGroupAddMenu) { }
 	SLATE_END_ARGS();
 
-	void Construct(const FArguments& InArgs, TWeakObjectPtr<UNiagaraStackItemGroup> InItemGroup, INiagaraStackItemGroupAddUtilities* InAddUtilities, int32 InInsertIndex);
+	void Construct(const FArguments& InArgs, UNiagaraStackEntry* InSourceEntry, INiagaraStackItemGroupAddUtilities* InAddUtilities, int32 InInsertIndex);
 
 	TSharedPtr<SWidget> GetFilterTextBox();
 
@@ -32,7 +32,7 @@ private:
 
 private:
 	INiagaraStackItemGroupAddUtilities* AddUtilities;
-	TWeakObjectPtr<UNiagaraStackItemGroup> StackItemGroup;
+	TWeakObjectPtr<UNiagaraStackEntry> SourceEntry;
 	int32 InsertIndex;
 
 	TSharedPtr<SNiagaraStackAddSelector> ActionSelector;
diff --git a/Engine/Plugins/FX/Niagara/Source/NiagaraEditorWidgets/Private/Stack/SNiagaraStackParameterStoreItem.cpp b/Engine/Plugins/FX/Niagara/Source/NiagaraEditorWidgets/Private/Stack/SNiagaraStackParameterStoreItem.cpp
new file mode 100644
--- /dev/null
+++ b/Engine/Plugins/FX/Niagara/Source/NiagaraEditorWidgets/Private/Stack/SNiagaraStackParameterStoreItem.cpp
@@ -0,0 +1,26 @@
+// Copyright Epic Games, Inc. All Rights Reserved.
+
+#include "Stack/SNiagaraStackParameterStoreItem.h"
+
+#include "Stack/SNiagaraStackItemGroupAddButton.h"
+#include "ViewModels/Stack/INiagaraStackItemGroupAddUtilities.h"
+#include "ViewModels/Stack/NiagaraStackSystemSettingsGroup.h"
+
+void SNiagaraStackParameterStoreItem::Construct(const FArguments& InArgs, UNiagaraStackParameterStoreItem& InParameterStoreItem, UNiagaraStackViewModel* InStackViewModel)
+{
+	ParameterStoreItem = &InParameterStoreItem;
+	SNiagaraStackItem::Construct(SNiagaraStackItem::FArguments(), InParameterStoreItem, InStackViewModel);
+}
+
+void SNiagaraStackParameterStoreItem::AddCustomRowWidgets(TSharedRef<SHorizontalBox> HorizontalBox)
+{
+	INiagaraStackItemGroupAddUtilities* GroupAddUtilities = ParameterStoreItem->GetGroupAddUtilities();
+	if(GroupAddUtilities != nullptr)
+	{
+		HorizontalBox->AddSlot()
+			.AutoWidth()
+			[
+				SNew(SNiagaraStackItemGroupAddButton, ParameterStoreItem, GroupAddUtilities)
+			];
+	}
+}
\ No newline at end of file
diff --git a/Engine/Plugins/FX/Niagara/Source/NiagaraEditorWidgets/Private/Stack/SNiagaraStackParameterStoreItem.h b/Engine/Plugins/FX/Niagara/Source/NiagaraEditorWidgets/Private/Stack/SNiagaraStackParameterStoreItem.h
new file mode 100644
--- /dev/null
+++ b/Engine/Plugins/FX/Niagara/Source/NiagaraEditorWidgets/Private/Stack/SNiagaraStackParameterStoreItem.h
@@ -0,0 +1,23 @@
+// Copyright Epic Games, Inc. All Rights Reserved.
+
+#pragma once
+
+#include "SNiagaraStackItem.h"
+
+class UNiagaraStackParameterStoreItem;
+class UNiagaraStackViewModel;
+
+class SNiagaraStackParameterStoreItem : public SNiagaraStackItem
+{
+public:
+	SLATE_BEGIN_ARGS(SNiagaraStackParameterStoreItem) {}
+	SLATE_END_ARGS();
+
+	void Construct(const FArguments& InArgs, UNiagaraStackParameterStoreItem& InParameterStoreItem, UNiagaraStackViewModel* InStackViewModel);
+
+protected:
+	virtual void AddCustomRowWidgets(TSharedRef<SHorizontalBox> HorizontalBox) override;
+
+private:
+	UNiagaraStackParameterStoreItem* ParameterStoreItem;
+};
\ No newline at end of file
