Improved numerics resolution and made type checking stricter for op nodes

diff --git a/Engine/Plugins/FX/Niagara/Source/Niagara/Private/NiagaraConstants.cpp b/Engine/Plugins/FX/Niagara/Source/Niagara/Private/NiagaraConstants.cpp
--- a/Engine/Plugins/FX/Niagara/Source/Niagara/Private/NiagaraConstants.cpp
+++ b/Engine/Plugins/FX/Niagara/Source/Niagara/Private/NiagaraConstants.cpp
@@ -147,6 +147,7 @@ void FNiagaraConstants::Init()
 	{
 		OldPositionTypes.Add(FNiagaraVariable(FNiagaraTypeDefinition::GetVec3Def(),  SYS_PARAM_PARTICLES_POSITION.GetName()));
 		OldPositionTypes.Add(FNiagaraVariable(FNiagaraTypeDefinition::GetVec3Def(), FName("Particles.Initial.Position")));
+		OldPositionTypes.Add(FNiagaraVariable(FNiagaraTypeDefinition::GetVec3Def(), FName("Particles.Previous.Position")));
 		OldPositionTypes.Add(FNiagaraVariable(FNiagaraTypeDefinition::GetVec3Def(),  SYS_PARAM_ENGINE_POSITION.GetName()));
 		OldPositionTypes.Add(FNiagaraVariable(FNiagaraTypeDefinition::GetVec3Def(),  SYS_PARAM_ENGINE_EMITTER_SIMULATION_POSITION.GetName()));
 	}
diff --git a/Engine/Plugins/FX/Niagara/Source/Niagara/Private/NiagaraModule.cpp b/Engine/Plugins/FX/Niagara/Source/Niagara/Private/NiagaraModule.cpp
--- a/Engine/Plugins/FX/Niagara/Source/Niagara/Private/NiagaraModule.cpp
+++ b/Engine/Plugins/FX/Niagara/Source/Niagara/Private/NiagaraModule.cpp
@@ -1181,6 +1181,7 @@ bool FNiagaraTypeDefinition::IsLossyConversion(const FNiagaraTypeDefinition& Typ
 FNiagaraTypeDefinition FNiagaraTypeDefinition::GetNumericOutputType(const TArray<FNiagaraTypeDefinition> TypeDefinintions, ENiagaraNumericOutputTypeSelectionMode SelectionMode)
 {
 	checkf(SelectionMode != ENiagaraNumericOutputTypeSelectionMode::None, TEXT("Can not get numeric output type with selection mode none."));
+	checkf(SelectionMode != ENiagaraNumericOutputTypeSelectionMode::Custom, TEXT("Can not get numeric output type with selection mode custom."));
 
 	//This may need some work. Should work fine for now.
 	if (SelectionMode == ENiagaraNumericOutputTypeSelectionMode::Scalar)
diff --git a/Engine/Plugins/FX/Niagara/Source/Niagara/Public/NiagaraTypes.h b/Engine/Plugins/FX/Niagara/Source/Niagara/Public/NiagaraTypes.h
--- a/Engine/Plugins/FX/Niagara/Source/Niagara/Public/NiagaraTypes.h
+++ b/Engine/Plugins/FX/Niagara/Source/Niagara/Public/NiagaraTypes.h
@@ -400,6 +400,8 @@ enum class ENiagaraNumericOutputTypeSelectionMode : uint8
 	Smallest,
 	/** Selects the base scalar type for this numeric inputs. */
 	Scalar,
+	/** Selects the type based on custom logic from the node. */
+	Custom,
 };
 
 /** 
diff --git a/Engine/Plugins/FX/Niagara/Source/NiagaraEditor/Private/NiagaraEditorCommon.cpp b/Engine/Plugins/FX/Niagara/Source/NiagaraEditor/Private/NiagaraEditorCommon.cpp
--- a/Engine/Plugins/FX/Niagara/Source/NiagaraEditor/Private/NiagaraEditorCommon.cpp
+++ b/Engine/Plugins/FX/Niagara/Source/NiagaraEditor/Private/NiagaraEditorCommon.cpp
@@ -51,6 +51,19 @@ void FNiagaraOpInfo::BuildName(FString InName, FString InCategory)
 	Name = FName(*(InCategory + TEXT("::") + InName));
 }
 
+int32 CountType(const TArray<FNiagaraTypeDefinition>& Types, const FNiagaraTypeDefinition& SearchType)
+{
+	int32 Count = 0;
+	for (const FNiagaraTypeDefinition& Type : Types)
+	{
+		if (Type == SearchType)
+		{
+			Count++;
+		}
+	}
+	return Count;
+}
+
 bool FNiagaraOpInfo::CreateHlslForAddedInputs(int32 InputCount, FString & HlslResult) const
 {
 	if (!bSupportsAddedInputs || AddedInputFormatting.IsEmpty() || InputCount < 2)
@@ -79,7 +92,7 @@ BEGIN_FUNCTION_BUILD_OPTIMIZATION
 void FNiagaraOpInfo::Init()
 {
 	//Common input and output names.
- 	static FName Result(TEXT("Result"));
+	static FName Result(TEXT("Result"));
 	static FText ResultText = NSLOCTEXT("NiagaraOpInfo", "Operation Result", "Result");
 	static FText InText = NSLOCTEXT("NiagaraOpInfo", "Single Function Param", "In"); 
 
@@ -102,10 +115,21 @@ void FNiagaraOpInfo::Init()
 	static FText MinText = NSLOCTEXT("NiagaraOpInfo", "Min", "Min");
 	static FText MaxText = NSLOCTEXT("NiagaraOpInfo", "Max", "Max");
 
+	auto NoPositionTypeCheck = [=](const TArray<FNiagaraTypeDefinition>& InputTypes, FText& ErrorMessage)
+	{
+		int32 PositionTypes = CountType(InputTypes, FNiagaraTypeDefinition::GetPositionDef());
+		if (PositionTypes > 0)
+		{
+			ErrorMessage = NSLOCTEXT("NiagaraOpInfo", "Position Type Error", "This operation cannot be used to for position types, as it would break large world coordinates. To resolve this error, you can explicitly truncate the position to vectors first.");
+			return false;
+		}
+		return true;
+	};
+
 	//Add all numeric ops
 	FText NumericCategory = NSLOCTEXT("NiagaraOpInfo", "NumericOpCategory", "Numeric");
 	{
-		FNiagaraTypeDefinition Type = FNiagaraTypeDefinition::GetGenericNumericDef();
+		FNiagaraTypeDefinition NumericType = FNiagaraTypeDefinition::GetGenericNumericDef();
 		FString DefaultStr_Zero(TEXT("0.0"));
 		FString DefaultStr_One(TEXT("1.0"));
 		FString CategoryName(TEXT("Numeric"));
@@ -116,12 +140,12 @@ void FNiagaraOpInfo::Init()
 		Op->FriendlyName = NSLOCTEXT("NiagaraOpInfo", "Add Name", "Add");
 		Op->Description = NSLOCTEXT("NiagaraOpInfo", "Add Desc", "Result = A + B");
 		Op->Keywords = FText::FromString(TEXT("+"));
-		Op->Inputs.Add(FNiagaraOpInOutInfo(A, Type, AText, AText, DefaultStr_Zero));
-		Op->Inputs.Add(FNiagaraOpInOutInfo(B, Type, BText, BText, DefaultStr_Zero));
-		Op->Outputs.Add(FNiagaraOpInOutInfo(Result, Type, ResultText, ResultText, DefaultStr_Zero, TEXT("{0} + {1}")));
+		Op->Inputs.Add(FNiagaraOpInOutInfo(A, NumericType, AText, AText, DefaultStr_Zero));
+		Op->Inputs.Add(FNiagaraOpInOutInfo(B, NumericType, BText, BText, DefaultStr_Zero));
+		Op->Outputs.Add(FNiagaraOpInOutInfo(Result, NumericType, ResultText, ResultText, DefaultStr_Zero, TEXT("{0} + {1}")));
 		Op->BuildName(TEXT("Add"), CategoryName);
 		Op->bSupportsAddedInputs = true;
-		Op->AddedInputTypeRestrictions.Add(Type);
+		Op->AddedInputTypeRestrictions.Add(NumericType);
 		Op->AddedInputFormatting = TEXT("{A} + {B}");
 		Op->bSupportsStaticResolution = true;
 		Op->StaticVariableResolveFunction.BindLambda([=](const TArray<int32>& InPinValues)
@@ -139,6 +163,26 @@ void FNiagaraOpInfo::Init()
 
 				return Return;
 			});
+		Op->InputTypeValidationFunction.BindLambda([=](const TArray<FNiagaraTypeDefinition>& InputTypes, FText& ErrorMessage)
+			{
+				int32 PositionTypes = CountType(InputTypes, FNiagaraTypeDefinition::GetPositionDef());
+				if (PositionTypes > 1)
+				{
+					ErrorMessage = NSLOCTEXT("NiagaraOpInfo", "Add Type Error", "The add operation cannot be used to add two positions. If this is intended, convert one or all of them to vectors first.");
+					return false;
+				}
+				return true;
+			});
+		Op->NumericOuputTypeSelectionMode = ENiagaraNumericOutputTypeSelectionMode::Custom;
+		Op->CustomNumericResolveFunction.BindLambda([=](const TArray<FNiagaraTypeDefinition>& NonNumericInputs)
+			{
+				int32 PositionTypes = CountType(NonNumericInputs, FNiagaraTypeDefinition::GetPositionDef());
+				if (PositionTypes > 0)
+				{
+					return FNiagaraTypeDefinition::GetPositionDef();
+				}
+				return FNiagaraTypeDefinition::GetNumericOutputType(NonNumericInputs, ENiagaraNumericOutputTypeSelectionMode::Largest);
+			});
 		OpInfoMap.Add(Op->Name) = Idx;
 
 		Idx = OpInfos.AddDefaulted();
@@ -147,12 +191,12 @@ void FNiagaraOpInfo::Init()
 		Op->FriendlyName = NSLOCTEXT("NiagaraOpInfo", "Subtract Name", "Subtract");
 		Op->Description = NSLOCTEXT("NiagaraOpInfo", "Subtract Desc", "Result = A - B");
 		Op->Keywords = FText::FromString(TEXT("-"));
-		Op->Inputs.Add(FNiagaraOpInOutInfo(A, Type, AText, AText, DefaultStr_Zero));
-		Op->Inputs.Add(FNiagaraOpInOutInfo(B, Type, BText, BText, DefaultStr_Zero));
-		Op->Outputs.Add(FNiagaraOpInOutInfo(Result, Type, ResultText, ResultText, DefaultStr_Zero, TEXT("{0} - {1}")));		
+		Op->Inputs.Add(FNiagaraOpInOutInfo(A, NumericType, AText, AText, DefaultStr_Zero));
+		Op->Inputs.Add(FNiagaraOpInOutInfo(B, NumericType, BText, BText, DefaultStr_Zero));
+		Op->Outputs.Add(FNiagaraOpInOutInfo(Result, NumericType, ResultText, ResultText, DefaultStr_Zero, TEXT("{0} - {1}")));		
 		Op->BuildName(TEXT("Subtract"), CategoryName);
 		Op->bSupportsAddedInputs = true;
-		Op->AddedInputTypeRestrictions.Add(Type);
+		Op->AddedInputTypeRestrictions.Add(NumericType);
 		Op->AddedInputFormatting = TEXT("{A} - {B}");
 		Op->bSupportsStaticResolution = true;
 		Op->StaticVariableResolveFunction.BindLambda([=](const TArray<int32>& InPinValues)
@@ -170,6 +214,30 @@ void FNiagaraOpInfo::Init()
 
 				return Return;
 			});
+		Op->InputTypeValidationFunction.BindLambda([=](const TArray<FNiagaraTypeDefinition>& InputTypes, FText& ErrorMessage)
+			{
+				int32 PositionTypes = CountType(InputTypes, FNiagaraTypeDefinition::GetPositionDef());
+				if (PositionTypes > 2)
+				{
+					ErrorMessage = NSLOCTEXT("NiagaraOpInfo", "Subtract Type Error", "The subtract operation cannot be used for more than two positions.");
+					return false;
+				}
+				return true;
+			});
+		Op->NumericOuputTypeSelectionMode = ENiagaraNumericOutputTypeSelectionMode::Custom;
+		Op->CustomNumericResolveFunction.BindLambda([=](const TArray<FNiagaraTypeDefinition>& NonNumericInputs)
+			{
+				int32 PositionTypes = CountType(NonNumericInputs, FNiagaraTypeDefinition::GetPositionDef());
+				if (PositionTypes == 1)
+				{
+					return FNiagaraTypeDefinition::GetPositionDef();
+				}
+				if (PositionTypes > 1)
+				{
+					 return FNiagaraTypeDefinition::GetVec3Def();
+				 }
+				return FNiagaraTypeDefinition::GetNumericOutputType(NonNumericInputs, ENiagaraNumericOutputTypeSelectionMode::Largest);
+			});
 		OpInfoMap.Add(Op->Name) = Idx;
 
 		Idx = OpInfos.AddDefaulted();
@@ -178,14 +246,15 @@ void FNiagaraOpInfo::Init()
 		Op->FriendlyName = NSLOCTEXT("NiagaraOpInfo", "Multiply Name", "Multiply");
 		Op->Description = NSLOCTEXT("NiagaraOpInfo", "Multiply Desc", "Result = A * B");
 		Op->Keywords = FText::FromString(TEXT("*"));
-		Op->Inputs.Add(FNiagaraOpInOutInfo(A, Type, AText, AText, DefaultStr_One));
-		Op->Inputs.Add(FNiagaraOpInOutInfo(B, Type, BText, BText, DefaultStr_One));
-		Op->Outputs.Add(FNiagaraOpInOutInfo(Result, Type, ResultText, ResultText, DefaultStr_One, TEXT("{0} * {1}")));
+		Op->Inputs.Add(FNiagaraOpInOutInfo(A, NumericType, AText, AText, DefaultStr_One));
+		Op->Inputs.Add(FNiagaraOpInOutInfo(B, NumericType, BText, BText, DefaultStr_One));
+		Op->Outputs.Add(FNiagaraOpInOutInfo(Result, NumericType, ResultText, ResultText, DefaultStr_One, TEXT("{0} * {1}")));
 		Op->BuildName(TEXT("Mul"), CategoryName);
 		Op->bSupportsAddedInputs = true;
-		Op->AddedInputTypeRestrictions.Add(Type);
+		Op->AddedInputTypeRestrictions.Add(NumericType);
 		Op->AddedInputFormatting = TEXT("{A} * {B}");
 		Op->bSupportsStaticResolution = true;
+		Op->InputTypeValidationFunction.BindLambda(NoPositionTypeCheck);
 		Op->StaticVariableResolveFunction.BindLambda([=](const TArray<int32>& InPinValues)
 			{
 				if (InPinValues.Num() == 0)
@@ -209,10 +278,11 @@ void FNiagaraOpInfo::Init()
 		Op->FriendlyName = NSLOCTEXT("NiagaraOpInfo", "Divide Name", "Divide");
 		Op->Description = NSLOCTEXT("NiagaraOpInfo", "Divide Desc", "Result = A / B");
 		Op->Keywords = FText::FromString(TEXT("/"));
-		Op->Inputs.Add(FNiagaraOpInOutInfo(A, Type, AText, AText, DefaultStr_One));
-		Op->Inputs.Add(FNiagaraOpInOutInfo(B, Type, BText, BText, DefaultStr_One));
-		Op->Outputs.Add(FNiagaraOpInOutInfo(Result, Type, ResultText, ResultText, DefaultStr_One, TEXT("{0} / {1}")));
+		Op->Inputs.Add(FNiagaraOpInOutInfo(A, NumericType, AText, AText, DefaultStr_One));
+		Op->Inputs.Add(FNiagaraOpInOutInfo(B, NumericType, BText, BText, DefaultStr_One));
+		Op->Outputs.Add(FNiagaraOpInOutInfo(Result, NumericType, ResultText, ResultText, DefaultStr_One, TEXT("{0} / {1}")));
 		Op->BuildName(TEXT("Div"), CategoryName);
+		Op->InputTypeValidationFunction.BindLambda(NoPositionTypeCheck);
 		Op->bSupportsStaticResolution = true;
 		Op->StaticVariableResolveFunction.BindLambda([=](const TArray<int32>& InPinValues)
 			{
@@ -239,11 +309,35 @@ void FNiagaraOpInfo::Init()
 		Op->Category = NumericCategory;
 		Op->FriendlyName = NSLOCTEXT("NiagaraOpInfo", "MultiplyAdd Name", "MultiplyAdd");
 		Op->Description = NSLOCTEXT("NiagaraOpInfo", "MultiplyAdd Desc", "Result = (A * B) + C");
-		Op->Inputs.Add(FNiagaraOpInOutInfo(A, Type, AText, AText, DefaultStr_One));
-		Op->Inputs.Add(FNiagaraOpInOutInfo(B, Type, BText, BText, DefaultStr_One));
-		Op->Inputs.Add(FNiagaraOpInOutInfo(C, Type, CText, CText, DefaultStr_Zero));
-		Op->Outputs.Add(FNiagaraOpInOutInfo(Result, Type, ResultText, ResultText, DefaultStr_One, TEXT("{0} * {1} + {2}")));
+		Op->Inputs.Add(FNiagaraOpInOutInfo(A, NumericType, AText, AText, DefaultStr_One));
+		Op->Inputs.Add(FNiagaraOpInOutInfo(B, NumericType, BText, BText, DefaultStr_One));
+		Op->Inputs.Add(FNiagaraOpInOutInfo(C, NumericType, CText, CText, DefaultStr_Zero));
+		Op->Outputs.Add(FNiagaraOpInOutInfo(Result, NumericType, ResultText, ResultText, DefaultStr_One, TEXT("{0} * {1} + {2}")));
 		Op->BuildName(TEXT("Madd"), CategoryName);
+		Op->InputTypeValidationFunction.BindLambda([=](const TArray<FNiagaraTypeDefinition>& InputTypes, FText& ErrorMessage)
+			{
+				int32 PositionTypes = CountType(InputTypes, FNiagaraTypeDefinition::GetPositionDef());
+				if (PositionTypes == 0)
+				{
+					return true;
+				}
+				if (PositionTypes > 1 || (InputTypes[2] != FNiagaraTypeDefinition::GetPositionDef()))
+				{
+					ErrorMessage = NSLOCTEXT("NiagaraOpInfo", "MultiplyAdd Type Error", "The MultiplyAdd operation can only use positions for the addition, not the multiplication part.");
+					return false;
+				}
+				return true;
+			});
+		Op->NumericOuputTypeSelectionMode = ENiagaraNumericOutputTypeSelectionMode::Custom;
+		Op->CustomNumericResolveFunction.BindLambda([=](const TArray<FNiagaraTypeDefinition>& NonNumericInputs)
+			{
+				int32 PositionTypes = CountType(NonNumericInputs, FNiagaraTypeDefinition::GetPositionDef());
+				if (PositionTypes == 1)
+				{
+					return FNiagaraTypeDefinition::GetPositionDef();
+				}
+				return FNiagaraTypeDefinition::GetNumericOutputType(NonNumericInputs, ENiagaraNumericOutputTypeSelectionMode::Largest);
+			});
 		OpInfoMap.Add(Op->Name) = Idx;
 
 		Idx = OpInfos.AddDefaulted();
@@ -252,11 +346,35 @@ void FNiagaraOpInfo::Init()
 		Op->FriendlyName = NSLOCTEXT("NiagaraOpInfo", "Lerp Name", "Lerp");
 		Op->Description = NSLOCTEXT("NiagaraOpInfo", "Lerp Desc", "Result = (A * (1 - C)) + (B * C)");
 		Op->Keywords = FText::FromString(TEXT("lerp"));
-		Op->Inputs.Add(FNiagaraOpInOutInfo(A, Type, AText, AText, DefaultStr_Zero));
-		Op->Inputs.Add(FNiagaraOpInOutInfo(B, Type, BText, BText, DefaultStr_One));
-		Op->Inputs.Add(FNiagaraOpInOutInfo(C, Type, CText, CText, DefaultStr_Zero));
-		Op->Outputs.Add(FNiagaraOpInOutInfo(Result, Type, ResultText, ResultText, DefaultStr_Zero, TEXT("lerp({0},{1},{2})")));
+		Op->Inputs.Add(FNiagaraOpInOutInfo(A, NumericType, AText, AText, DefaultStr_Zero));
+		Op->Inputs.Add(FNiagaraOpInOutInfo(B, NumericType, BText, BText, DefaultStr_One));
+		Op->Inputs.Add(FNiagaraOpInOutInfo(C, NumericType, CText, CText, DefaultStr_Zero));
+		Op->Outputs.Add(FNiagaraOpInOutInfo(Result, NumericType, ResultText, ResultText, DefaultStr_Zero, TEXT("lerp({0},{1},{2})")));
 		Op->BuildName(TEXT("Lerp"), CategoryName);
+		Op->InputTypeValidationFunction.BindLambda([=](const TArray<FNiagaraTypeDefinition>& InputTypes, FText& ErrorMessage)
+			{
+				int32 PositionTypes = CountType(InputTypes, FNiagaraTypeDefinition::GetPositionDef());
+				 if (PositionTypes == 0)
+				 {
+					 return true;
+				 }
+				if (PositionTypes != 2 || InputTypes[2] == FNiagaraTypeDefinition::GetPositionDef())
+				{
+					ErrorMessage = NSLOCTEXT("NiagaraOpInfo", "Lerp Type Error", "The Lerp operation can only lerp between two positions, not between a position and another type.");
+					return false;
+				}
+				return true;
+			});
+		Op->NumericOuputTypeSelectionMode = ENiagaraNumericOutputTypeSelectionMode::Custom;
+		Op->CustomNumericResolveFunction.BindLambda([=](const TArray<FNiagaraTypeDefinition>& NonNumericInputs)
+			{
+				 int32 PositionTypes = CountType(NonNumericInputs, FNiagaraTypeDefinition::GetPositionDef());
+				if (PositionTypes == 2)
+				{
+					return FNiagaraTypeDefinition::GetPositionDef();
+				}
+				return FNiagaraTypeDefinition::GetNumericOutputType(NonNumericInputs, ENiagaraNumericOutputTypeSelectionMode::Largest);
+			});
 		OpInfoMap.Add(Op->Name) = Idx;
 
 		Idx = OpInfos.AddDefaulted();
@@ -264,9 +382,10 @@ void FNiagaraOpInfo::Init()
 		Op->Category = NumericCategory;
 		Op->FriendlyName = NSLOCTEXT("NiagaraOpInfo", "Reciprocal Fast Name", "Reciprocal Fast");
 		Op->Description = NSLOCTEXT("NiagaraOpInfo", "Reciprocal Fast Desc", "12-bits of accuracy, but faster. Result = 1 / A using Newton/Raphson approximation.");
-		Op->Inputs.Add(FNiagaraOpInOutInfo(A, Type, AText, AText, DefaultStr_One));
-		Op->Outputs.Add(FNiagaraOpInOutInfo(Result, Type, ResultText, ResultText, DefaultStr_One, TEXT("rcp({0})")));
+		Op->Inputs.Add(FNiagaraOpInOutInfo(A, NumericType, AText, AText, DefaultStr_One));
+		Op->Outputs.Add(FNiagaraOpInOutInfo(Result, NumericType, ResultText, ResultText, DefaultStr_One, TEXT("rcp({0})")));
 		Op->BuildName(TEXT("RcpFast"), CategoryName);
+		Op->InputTypeValidationFunction.BindLambda(NoPositionTypeCheck);
 		OpInfoMap.Add(Op->Name) = Idx;
 
 		Idx = OpInfos.AddDefaulted();
@@ -274,20 +393,21 @@ void FNiagaraOpInfo::Init()
 		Op->Category = NumericCategory;
 		Op->FriendlyName = NSLOCTEXT("NiagaraOpInfo", "Reciprocal Name", "Reciprocal");
 		Op->Description = NSLOCTEXT("NiagaraOpInfo", "Reciprocal Desc", "More accurate than Reciprocal Fast. Result = 1 / A");
-		Op->Inputs.Add(FNiagaraOpInOutInfo(A, Type, AText, AText, DefaultStr_One));
-		Op->Outputs.Add(FNiagaraOpInOutInfo(Result, Type, ResultText, ResultText, DefaultStr_One, TEXT("Reciprocal({0})")));
+		Op->Inputs.Add(FNiagaraOpInOutInfo(A, NumericType, AText, AText, DefaultStr_One));
+		Op->Outputs.Add(FNiagaraOpInOutInfo(Result, NumericType, ResultText, ResultText, DefaultStr_One, TEXT("Reciprocal({0})")));
 		Op->BuildName(TEXT("Rcp"), CategoryName);
+		Op->InputTypeValidationFunction.BindLambda(NoPositionTypeCheck);
 		OpInfoMap.Add(Op->Name) = Idx;
 
-
 		Idx = OpInfos.AddDefaulted();
 		Op = &OpInfos[Idx];
 		Op->Category = NumericCategory;
 		Op->FriendlyName = NSLOCTEXT("NiagaraOpInfo", "Reciprocal Sqrt Name", "Reciprocal Sqrt");
 		Op->Description = NSLOCTEXT("NiagaraOpInfo", "Reciprocal Sqrt Desc", "Result = 1 / sqrt(A)");
-		Op->Inputs.Add(FNiagaraOpInOutInfo(A, Type, AText, AText, DefaultStr_One));
-		Op->Outputs.Add(FNiagaraOpInOutInfo(Result, Type, ResultText, ResultText, DefaultStr_One, TEXT("rsqrt({0})")));
+		Op->Inputs.Add(FNiagaraOpInOutInfo(A, NumericType, AText, AText, DefaultStr_One));
+		Op->Outputs.Add(FNiagaraOpInOutInfo(Result, NumericType, ResultText, ResultText, DefaultStr_One, TEXT("rsqrt({0})")));
 		Op->BuildName(TEXT("RSqrt"), CategoryName);
+		Op->InputTypeValidationFunction.BindLambda(NoPositionTypeCheck);
 		OpInfoMap.Add(Op->Name) = Idx;
 
 		Idx = OpInfos.AddDefaulted();
@@ -296,9 +416,10 @@ void FNiagaraOpInfo::Init()
 		Op->FriendlyName = NSLOCTEXT("NiagaraOpInfo", "Sqrt Name", "Sqrt");
 		Op->Description = NSLOCTEXT("NiagaraOpInfo", "Sqrt Desc", "Result = sqrt(A)");
 		Op->Keywords = FText::FromString(TEXT("sqrt"));
-		Op->Inputs.Add(FNiagaraOpInOutInfo(A, Type, AText, AText, DefaultStr_One));
-		Op->Outputs.Add(FNiagaraOpInOutInfo(Result, Type, ResultText, ResultText, DefaultStr_One, TEXT("sqrt({0})")));
+		Op->Inputs.Add(FNiagaraOpInOutInfo(A, NumericType, AText, AText, DefaultStr_One));
+		Op->Outputs.Add(FNiagaraOpInOutInfo(Result, NumericType, ResultText, ResultText, DefaultStr_One, TEXT("sqrt({0})")));
 		Op->BuildName(TEXT("Sqrt"), CategoryName);
+		Op->InputTypeValidationFunction.BindLambda(NoPositionTypeCheck);
 		OpInfoMap.Add(Op->Name) = Idx;
 
 		Idx = OpInfos.AddDefaulted();
@@ -307,9 +428,10 @@ void FNiagaraOpInfo::Init()
 		Op->FriendlyName = NSLOCTEXT("NiagaraOpInfo", "One Minus Name", "One Minus");
 		Op->Description = NSLOCTEXT("NiagaraOpInfo", "One Minus Desc", "Result = 1 - A");
 		Op->Keywords = FText::FromString(TEXT("1-x"));
-		Op->Inputs.Add(FNiagaraOpInOutInfo(A, Type, AText, AText, DefaultStr_One));
-		Op->Outputs.Add(FNiagaraOpInOutInfo(Result, Type, ResultText, ResultText, DefaultStr_One, TEXT("1 - {0}")));
+		Op->Inputs.Add(FNiagaraOpInOutInfo(A, NumericType, AText, AText, DefaultStr_One));
+		Op->Outputs.Add(FNiagaraOpInOutInfo(Result, NumericType, ResultText, ResultText, DefaultStr_One, TEXT("1 - {0}")));
 		Op->BuildName(TEXT("OneMinus"), CategoryName);
+		Op->InputTypeValidationFunction.BindLambda(NoPositionTypeCheck);
 		Op->bSupportsStaticResolution = true;
 		Op->StaticVariableResolveFunction.BindLambda([=](const TArray<int32>& InPinValues)
 			{
@@ -325,9 +447,10 @@ void FNiagaraOpInfo::Init()
 		Op->Category = NumericCategory;
 		Op->FriendlyName = NSLOCTEXT("NiagaraOpInfo", "Negate Name", "Negate");
 		Op->Description = NSLOCTEXT("NiagaraOpInfo", "Negate Desc", "Result = -A");
-		Op->Inputs.Add(FNiagaraOpInOutInfo(A, Type, AText, AText, DefaultStr_One));
-		Op->Outputs.Add(FNiagaraOpInOutInfo(Result, Type, ResultText, ResultText, DefaultStr_One, TEXT("-({0})")));
+		Op->Inputs.Add(FNiagaraOpInOutInfo(A, NumericType, AText, AText, DefaultStr_One));
+		Op->Outputs.Add(FNiagaraOpInOutInfo(Result, NumericType, ResultText, ResultText, DefaultStr_One, TEXT("-({0})")));
 		Op->BuildName(TEXT("Negate"), CategoryName);
+		Op->InputTypeValidationFunction.BindLambda(NoPositionTypeCheck);
 		Op->bSupportsStaticResolution = true;
 		Op->StaticVariableResolveFunction.BindLambda([=](const TArray<int32>& InPinValues)
 			{
@@ -343,9 +466,10 @@ void FNiagaraOpInfo::Init()
 		Op->Category = NumericCategory;
 		Op->FriendlyName = NSLOCTEXT("NiagaraOpInfo", "Abs Name", "Abs");
 		Op->Description = NSLOCTEXT("NiagaraOpInfo", "Abs Desc", "Result = abs(A)");
-		Op->Inputs.Add(FNiagaraOpInOutInfo(A, Type, AText, AText, DefaultStr_One));
-		Op->Outputs.Add(FNiagaraOpInOutInfo(Result, Type, ResultText, ResultText, DefaultStr_One, TEXT("abs({0})")));
+		Op->Inputs.Add(FNiagaraOpInOutInfo(A, NumericType, AText, AText, DefaultStr_One));
+		Op->Outputs.Add(FNiagaraOpInOutInfo(Result, NumericType, ResultText, ResultText, DefaultStr_One, TEXT("abs({0})")));
 		Op->BuildName(TEXT("Abs"), CategoryName);
+		Op->InputTypeValidationFunction.BindLambda(NoPositionTypeCheck);
 		OpInfoMap.Add(Op->Name) = Idx;
 
 		Idx = OpInfos.AddDefaulted();
@@ -353,9 +477,10 @@ void FNiagaraOpInfo::Init()
 		Op->Category = NumericCategory;
 		Op->FriendlyName = NSLOCTEXT("NiagaraOpInfo", "Exp Name", "Exp");
 		Op->Description = NSLOCTEXT("NiagaraOpInfo", "Exp Desc", "Result = exp(A)");
-		Op->Inputs.Add(FNiagaraOpInOutInfo(A, Type, AText, AText, DefaultStr_One));
-		Op->Outputs.Add(FNiagaraOpInOutInfo(Result, Type, ResultText, ResultText, DefaultStr_One, TEXT("exp({0})")));
+		Op->Inputs.Add(FNiagaraOpInOutInfo(A, NumericType, AText, AText, DefaultStr_One));
+		Op->Outputs.Add(FNiagaraOpInOutInfo(Result, NumericType, ResultText, ResultText, DefaultStr_One, TEXT("exp({0})")));
 		Op->BuildName(TEXT("Exp"), CategoryName);
+		Op->InputTypeValidationFunction.BindLambda(NoPositionTypeCheck);
 		OpInfoMap.Add(Op->Name) = Idx;
 
 		Idx = OpInfos.AddDefaulted();
@@ -363,9 +488,10 @@ void FNiagaraOpInfo::Init()
 		Op->Category = NumericCategory;
 		Op->FriendlyName = NSLOCTEXT("NiagaraOpInfo", "Exp2 Name", "Exp2");
 		Op->Description = NSLOCTEXT("NiagaraOpInfo", "Exp2 Desc", "Result = exp2(A)");
-		Op->Inputs.Add(FNiagaraOpInOutInfo(A, Type, AText, AText, DefaultStr_One));
-		Op->Outputs.Add(FNiagaraOpInOutInfo(Result, Type, ResultText, ResultText, DefaultStr_One, TEXT("exp2({0})")));
+		Op->Inputs.Add(FNiagaraOpInOutInfo(A, NumericType, AText, AText, DefaultStr_One));
+		Op->Outputs.Add(FNiagaraOpInOutInfo(Result, NumericType, ResultText, ResultText, DefaultStr_One, TEXT("exp2({0})")));
 		Op->BuildName(TEXT("Exp2"), CategoryName);
+		Op->InputTypeValidationFunction.BindLambda(NoPositionTypeCheck);
 		OpInfoMap.Add(Op->Name) = Idx;
 
 		Idx = OpInfos.AddDefaulted();
@@ -373,9 +499,10 @@ void FNiagaraOpInfo::Init()
 		Op->Category = NumericCategory;
 		Op->FriendlyName = NSLOCTEXT("NiagaraOpInfo", "Log Name", "Log");
 		Op->Description = NSLOCTEXT("NiagaraOpInfo", "Log Desc", "Result = log(A)");
-		Op->Inputs.Add(FNiagaraOpInOutInfo(A, Type, AText, AText, DefaultStr_One));
-		Op->Outputs.Add(FNiagaraOpInOutInfo(Result, Type, ResultText, ResultText, DefaultStr_One, TEXT("log({0})")));
+		Op->Inputs.Add(FNiagaraOpInOutInfo(A, NumericType, AText, AText, DefaultStr_One));
+		Op->Outputs.Add(FNiagaraOpInOutInfo(Result, NumericType, ResultText, ResultText, DefaultStr_One, TEXT("log({0})")));
 		Op->BuildName(TEXT("Log"), CategoryName);
+		Op->InputTypeValidationFunction.BindLambda(NoPositionTypeCheck);
 		OpInfoMap.Add(Op->Name) = Idx;
 
 		Idx = OpInfos.AddDefaulted();
@@ -384,9 +511,10 @@ void FNiagaraOpInfo::Init()
 		Op->FriendlyName = NSLOCTEXT("NiagaraOpInfo", "Log2 Name", "Log2");
 		Op->Description = NSLOCTEXT("NiagaraOpInfo", "Log2 Desc", "Result = log2(A)");
 		Op->Keywords = FText::FromString(TEXT("log2"));
-		Op->Inputs.Add(FNiagaraOpInOutInfo(A, Type, AText, AText, DefaultStr_One));
-		Op->Outputs.Add(FNiagaraOpInOutInfo(Result, Type, ResultText, ResultText, DefaultStr_One, TEXT("log2({0})")));
+		Op->Inputs.Add(FNiagaraOpInOutInfo(A, NumericType, AText, AText, DefaultStr_One));
+		Op->Outputs.Add(FNiagaraOpInOutInfo(Result, NumericType, ResultText, ResultText, DefaultStr_One, TEXT("log2({0})")));
 		Op->BuildName(TEXT("Log2"), CategoryName);
+		Op->InputTypeValidationFunction.BindLambda(NoPositionTypeCheck);
 		OpInfoMap.Add(Op->Name) = Idx;
 
 		FText AngleFriendlyNameText = NSLOCTEXT("NiagaraOpInfo", "Angle Name", "Angle");
@@ -404,7 +532,7 @@ void FNiagaraOpInfo::Init()
 		Op->Keywords = FText::FromString(TEXT("sine"));
 		Op->Inputs.Add(FNiagaraOpInOutInfo(TEXT("Angle"), AngleType, AngleFriendlyNameText, AngleDescText, DefaultStr_One));
 		Op->Inputs.Add(FNiagaraOpInOutInfo(TEXT("Period"), PeriodType, PeriodFriendlyNameText, PeriodDescText, DefaultStr_One));
-		Op->Outputs.Add(FNiagaraOpInOutInfo(Result, Type, ResultText, ResultText, DefaultStr_One, TEXT("sin({0}*(TWO_PI/{1}))")));
+		Op->Outputs.Add(FNiagaraOpInOutInfo(Result, NumericType, ResultText, ResultText, DefaultStr_One, TEXT("sin({0}*(TWO_PI/{1}))")));
 		Op->BuildName(TEXT("Sine"), CategoryName);
 		OpInfoMap.Add(Op->Name) = Idx;
 
@@ -414,9 +542,10 @@ void FNiagaraOpInfo::Init()
 		Op->FriendlyName = NSLOCTEXT("NiagaraOpInfo", "SinRad Name", "Sine(Radians)");
 		Op->Description = NSLOCTEXT("NiagaraOpInfo", "SinRad Desc", "Result = sin(AngleInRadians)");
 		Op->Keywords = FText::FromString(TEXT("sine"));
-		Op->Inputs.Add(FNiagaraOpInOutInfo(A, Type, AText, AText, DefaultStr_One));
-		Op->Outputs.Add(FNiagaraOpInOutInfo(Result, Type, ResultText, ResultText, DefaultStr_One, TEXT("sin({0})")));
+		Op->Inputs.Add(FNiagaraOpInOutInfo(A, NumericType, AText, AText, DefaultStr_One));
+		Op->Outputs.Add(FNiagaraOpInOutInfo(Result, NumericType, ResultText, ResultText, DefaultStr_One, TEXT("sin({0})")));
 		Op->BuildName(TEXT("Sine(Radians)"), CategoryName);
+		Op->InputTypeValidationFunction.BindLambda(NoPositionTypeCheck);
 		OpInfoMap.Add(Op->Name) = Idx;
 
 		Idx = OpInfos.AddDefaulted();
@@ -425,9 +554,10 @@ void FNiagaraOpInfo::Init()
 		Op->FriendlyName = NSLOCTEXT("NiagaraOpInfo", "SinDeg Name", "Sine(Degrees)");
 		Op->Description = NSLOCTEXT("NiagaraOpInfo", "SinDeg Desc", "Result = sin(AngleInDegrees*DegreesToRadians)");
 		Op->Keywords = FText::FromString(TEXT("sine"));
-		Op->Inputs.Add(FNiagaraOpInOutInfo(A, Type, AText, AText, DefaultStr_One));
-		Op->Outputs.Add(FNiagaraOpInOutInfo(Result, Type, ResultText, ResultText, DefaultStr_One, TEXT("sin((PI/180.0f)*{0})")));
+		Op->Inputs.Add(FNiagaraOpInOutInfo(A, NumericType, AText, AText, DefaultStr_One));
+		Op->Outputs.Add(FNiagaraOpInOutInfo(Result, NumericType, ResultText, ResultText, DefaultStr_One, TEXT("sin((PI/180.0f)*{0})")));
 		Op->BuildName(TEXT("Sine(Degrees)"), CategoryName);
+		Op->InputTypeValidationFunction.BindLambda(NoPositionTypeCheck);
 		OpInfoMap.Add(Op->Name) = Idx;
 
 		Idx = OpInfos.AddDefaulted();
@@ -438,7 +568,7 @@ void FNiagaraOpInfo::Init()
 		Op->Keywords = FText::FromString(TEXT("Cosine"));
 		Op->Inputs.Add(FNiagaraOpInOutInfo(TEXT("Angle"), AngleType, AngleFriendlyNameText, AngleDescText, DefaultStr_One));
 		Op->Inputs.Add(FNiagaraOpInOutInfo(TEXT("Period"), PeriodType, PeriodFriendlyNameText, PeriodDescText, DefaultStr_One));
-		Op->Outputs.Add(FNiagaraOpInOutInfo(Result, Type, ResultText, ResultText, DefaultStr_One, TEXT("cos({0}*(TWO_PI/{1}))")));
+		Op->Outputs.Add(FNiagaraOpInOutInfo(Result, NumericType, ResultText, ResultText, DefaultStr_One, TEXT("cos({0}*(TWO_PI/{1}))")));
 		Op->BuildName(TEXT("Cosine"), CategoryName);
 		OpInfoMap.Add(Op->Name) = Idx;
 
@@ -448,9 +578,10 @@ void FNiagaraOpInfo::Init()
 		Op->FriendlyName = NSLOCTEXT("NiagaraOpInfo", "CosRad Name", "Cosine(Radians)");
 		Op->Description = NSLOCTEXT("NiagaraOpInfo", "CosRad Desc", "Result = cos(AngleInRadians)");
 		Op->Keywords = FText::FromString(TEXT("Cosine"));
-		Op->Inputs.Add(FNiagaraOpInOutInfo(A, Type, AText, AText, DefaultStr_One));
-		Op->Outputs.Add(FNiagaraOpInOutInfo(Result, Type, ResultText, ResultText, DefaultStr_One, TEXT("cos({0})")));
+		Op->Inputs.Add(FNiagaraOpInOutInfo(A, NumericType, AText, AText, DefaultStr_One));
+		Op->Outputs.Add(FNiagaraOpInOutInfo(Result, NumericType, ResultText, ResultText, DefaultStr_One, TEXT("cos({0})")));
 		Op->BuildName(TEXT("Cosine(Radians)"), CategoryName);
+		Op->InputTypeValidationFunction.BindLambda(NoPositionTypeCheck);
 		OpInfoMap.Add(Op->Name) = Idx;
 
 		Idx = OpInfos.AddDefaulted();
@@ -459,9 +590,10 @@ void FNiagaraOpInfo::Init()
 		Op->FriendlyName = NSLOCTEXT("NiagaraOpInfo", "CosDeg Name", "Cosine(Degrees)");
 		Op->Description = NSLOCTEXT("NiagaraOpInfo", "CosDeg Desc", "Result = cos(AngleInDegrees*DegreesToRadians)");
 		Op->Keywords = FText::FromString(TEXT("Cosine"));
-		Op->Inputs.Add(FNiagaraOpInOutInfo(A, Type, AText, AText, DefaultStr_One));
-		Op->Outputs.Add(FNiagaraOpInOutInfo(Result, Type, ResultText, ResultText, DefaultStr_One, TEXT("cos((PI/180.0f)*{0})")));
+		Op->Inputs.Add(FNiagaraOpInOutInfo(A, NumericType, AText, AText, DefaultStr_One));
+		Op->Outputs.Add(FNiagaraOpInOutInfo(Result, NumericType, ResultText, ResultText, DefaultStr_One, TEXT("cos((PI/180.0f)*{0})")));
 		Op->BuildName(TEXT("Cosine(Degrees)"), CategoryName);
+		Op->InputTypeValidationFunction.BindLambda(NoPositionTypeCheck);
 		OpInfoMap.Add(Op->Name) = Idx;
 
 		Idx = OpInfos.AddDefaulted();
@@ -472,7 +604,7 @@ void FNiagaraOpInfo::Init()
 		Op->Keywords = FText::FromString(TEXT("Tangent"));
 		Op->Inputs.Add(FNiagaraOpInOutInfo(TEXT("Angle"), AngleType, AngleFriendlyNameText, AngleDescText, DefaultStr_One));
 		Op->Inputs.Add(FNiagaraOpInOutInfo(TEXT("Period"), PeriodType, PeriodFriendlyNameText, PeriodDescText, DefaultStr_One));
-		Op->Outputs.Add(FNiagaraOpInOutInfo(Result, Type, ResultText, ResultText, DefaultStr_One, TEXT("tan({0}*(TWO_PI/{1}))")));
+		Op->Outputs.Add(FNiagaraOpInOutInfo(Result, NumericType, ResultText, ResultText, DefaultStr_One, TEXT("tan({0}*(TWO_PI/{1}))")));
 		Op->BuildName(TEXT("Tangent"), CategoryName);
 		OpInfoMap.Add(Op->Name) = Idx;
 
@@ -482,9 +614,10 @@ void FNiagaraOpInfo::Init()
 		Op->FriendlyName = NSLOCTEXT("NiagaraOpInfo", "TanRad Name", "Tangent(Radians)");
 		Op->Description = NSLOCTEXT("NiagaraOpInfo", "TanRad Desc", "Result = tan(AngleInRadians)");
 		Op->Keywords = FText::FromString(TEXT("Tangent"));
-		Op->Inputs.Add(FNiagaraOpInOutInfo(A, Type, AText, AText, DefaultStr_One));
-		Op->Outputs.Add(FNiagaraOpInOutInfo(Result, Type, ResultText, ResultText, DefaultStr_One, TEXT("tan({0})")));
+		Op->Inputs.Add(FNiagaraOpInOutInfo(A, NumericType, AText, AText, DefaultStr_One));
+		Op->Outputs.Add(FNiagaraOpInOutInfo(Result, NumericType, ResultText, ResultText, DefaultStr_One, TEXT("tan({0})")));
 		Op->BuildName(TEXT("Tangent(Radians)"), CategoryName);
+		Op->InputTypeValidationFunction.BindLambda(NoPositionTypeCheck);
 		OpInfoMap.Add(Op->Name) = Idx;
 
 		Idx = OpInfos.AddDefaulted();
@@ -493,9 +626,10 @@ void FNiagaraOpInfo::Init()
 		Op->FriendlyName = NSLOCTEXT("NiagaraOpInfo", "TanDeg Name", "Tangent(Degrees)");
 		Op->Description = NSLOCTEXT("NiagaraOpInfo", "TanDeg Desc", "Result = tan(AngleInDegrees*DegreesToRadians)");
 		Op->Keywords = FText::FromString(TEXT("Tangent"));
-		Op->Inputs.Add(FNiagaraOpInOutInfo(A, Type, AText, AText, DefaultStr_One));
-		Op->Outputs.Add(FNiagaraOpInOutInfo(Result, Type, ResultText, ResultText, DefaultStr_One, TEXT("tan((PI/180.0f)*{0})")));
+		Op->Inputs.Add(FNiagaraOpInOutInfo(A, NumericType, AText, AText, DefaultStr_One));
+		Op->Outputs.Add(FNiagaraOpInOutInfo(Result, NumericType, ResultText, ResultText, DefaultStr_One, TEXT("tan((PI/180.0f)*{0})")));
 		Op->BuildName(TEXT("Tangent(Degrees)"), CategoryName);
+		Op->InputTypeValidationFunction.BindLambda(NoPositionTypeCheck);
 		OpInfoMap.Add(Op->Name) = Idx;
 
 
@@ -505,10 +639,11 @@ void FNiagaraOpInfo::Init()
 		Op->FriendlyName = NSLOCTEXT("NiagaraOpInfo", "ArcSine Name", "ArcSine");
 		Op->Description = NSLOCTEXT("NiagaraOpInfo", "ArcSine Desc", "Result = asin(A)*(Period/TWO_PI)");
 		Op->Keywords = FText::FromString(TEXT("ArcSine"));
-		Op->Inputs.Add(FNiagaraOpInOutInfo(A, Type, AText, AText, DefaultStr_One));
+		Op->Inputs.Add(FNiagaraOpInOutInfo(A, NumericType, AText, AText, DefaultStr_One));
 		Op->Inputs.Add(FNiagaraOpInOutInfo(TEXT("Period"), PeriodType, PeriodFriendlyNameText, PeriodDescText, DefaultStr_One));
-		Op->Outputs.Add(FNiagaraOpInOutInfo(Result, Type, ResultText, ResultText, DefaultStr_One, TEXT("asin({0})*({1}/TWO_PI)")));
+		Op->Outputs.Add(FNiagaraOpInOutInfo(Result, NumericType, ResultText, ResultText, DefaultStr_One, TEXT("asin({0})*({1}/TWO_PI)")));
 		Op->BuildName(TEXT("ArcSine"), CategoryName);
+		Op->InputTypeValidationFunction.BindLambda(NoPositionTypeCheck);
 		OpInfoMap.Add(Op->Name) = Idx;
 
 		Idx = OpInfos.AddDefaulted();
@@ -517,9 +652,10 @@ void FNiagaraOpInfo::Init()
 		Op->FriendlyName = NSLOCTEXT("NiagaraOpInfo", "ArcSineRad Name", "ArcSine(Radians)");
 		Op->Description = NSLOCTEXT("NiagaraOpInfo", "ArcSineRad Desc", "Result = asin(A)");
 		Op->Keywords = FText::FromString(TEXT("ArcSine"));
-		Op->Inputs.Add(FNiagaraOpInOutInfo(A, Type, AText, AText, DefaultStr_One));
-		Op->Outputs.Add(FNiagaraOpInOutInfo(Result, Type, ResultText, ResultText, DefaultStr_One, TEXT("asin({0})")));
+		Op->Inputs.Add(FNiagaraOpInOutInfo(A, NumericType, AText, AText, DefaultStr_One));
+		Op->Outputs.Add(FNiagaraOpInOutInfo(Result, NumericType, ResultText, ResultText, DefaultStr_One, TEXT("asin({0})")));
 		Op->BuildName(TEXT("ArcSine(Radians)"), CategoryName);
+		Op->InputTypeValidationFunction.BindLambda(NoPositionTypeCheck);
 		OpInfoMap.Add(Op->Name) = Idx;
 
 		Idx = OpInfos.AddDefaulted();
@@ -528,9 +664,10 @@ void FNiagaraOpInfo::Init()
 		Op->FriendlyName = NSLOCTEXT("NiagaraOpInfo", "ArcSineDeg Name", "ArcSine(Degrees)");
 		Op->Description = NSLOCTEXT("NiagaraOpInfo", "ArcSineDeg Desc", "Result = asin(A)*RadiansToDegrees");
 		Op->Keywords = FText::FromString(TEXT("ArcSine"));
-		Op->Inputs.Add(FNiagaraOpInOutInfo(A, Type, AText, AText, DefaultStr_One));
-		Op->Outputs.Add(FNiagaraOpInOutInfo(Result, Type, ResultText, ResultText, DefaultStr_One, TEXT("(180.0f/PI)*asin({0})")));
+		Op->Inputs.Add(FNiagaraOpInOutInfo(A, NumericType, AText, AText, DefaultStr_One));
+		Op->Outputs.Add(FNiagaraOpInOutInfo(Result, NumericType, ResultText, ResultText, DefaultStr_One, TEXT("(180.0f/PI)*asin({0})")));
 		Op->BuildName(TEXT("ArcSine(Degrees)"), CategoryName);
+		Op->InputTypeValidationFunction.BindLambda(NoPositionTypeCheck);
 		OpInfoMap.Add(Op->Name) = Idx;
 
 		Idx = OpInfos.AddDefaulted();
@@ -560,10 +697,11 @@ void FNiagaraOpInfo::Init()
 		Op->FriendlyName = NSLOCTEXT("NiagaraOpInfo", "ArcCosine Name", "ArcCosine");
 		Op->Description = NSLOCTEXT("NiagaraOpInfo", "ArcCosine Desc", "Result = acos(A)*(Period/TWO_PI)");
 		Op->Keywords = FText::FromString(TEXT("ArcCosine"));
-		Op->Inputs.Add(FNiagaraOpInOutInfo(A, Type, AText, AText, DefaultStr_One));
+		Op->Inputs.Add(FNiagaraOpInOutInfo(A, NumericType, AText, AText, DefaultStr_One));
 		Op->Inputs.Add(FNiagaraOpInOutInfo(TEXT("Period"), PeriodType, PeriodFriendlyNameText, PeriodDescText, DefaultStr_One));
-		Op->Outputs.Add(FNiagaraOpInOutInfo(Result, Type, ResultText, ResultText, DefaultStr_One, TEXT("acos({0})*({1}/TWO_PI)")));
+		Op->Outputs.Add(FNiagaraOpInOutInfo(Result, NumericType, ResultText, ResultText, DefaultStr_One, TEXT("acos({0})*({1}/TWO_PI)")));
 		Op->BuildName(TEXT("ArcCosine"), CategoryName);
+		Op->InputTypeValidationFunction.BindLambda(NoPositionTypeCheck);
 		OpInfoMap.Add(Op->Name) = Idx;
 
 		Idx = OpInfos.AddDefaulted();
@@ -572,9 +710,10 @@ void FNiagaraOpInfo::Init()
 		Op->FriendlyName = NSLOCTEXT("NiagaraOpInfo", "ArcCosineRad Name", "ArcCosine(Radians)");
 		Op->Description = NSLOCTEXT("NiagaraOpInfo", "ArcCosineRad Desc", "Result = acos(A)");
 		Op->Keywords = FText::FromString(TEXT("ArcCosine"));
-		Op->Inputs.Add(FNiagaraOpInOutInfo(A, Type, AText, AText, DefaultStr_One));
-		Op->Outputs.Add(FNiagaraOpInOutInfo(Result, Type, ResultText, ResultText, DefaultStr_One, TEXT("acos({0})")));
+		Op->Inputs.Add(FNiagaraOpInOutInfo(A, NumericType, AText, AText, DefaultStr_One));
+		Op->Outputs.Add(FNiagaraOpInOutInfo(Result, NumericType, ResultText, ResultText, DefaultStr_One, TEXT("acos({0})")));
 		Op->BuildName(TEXT("ArcCosine(Radians)"), CategoryName);
+		Op->InputTypeValidationFunction.BindLambda(NoPositionTypeCheck);
 		OpInfoMap.Add(Op->Name) = Idx;
 
 		Idx = OpInfos.AddDefaulted();
@@ -583,9 +722,10 @@ void FNiagaraOpInfo::Init()
 		Op->FriendlyName = NSLOCTEXT("NiagaraOpInfo", "ArcCosineDeg Name", "ArcCosine(Degrees)");
 		Op->Description = NSLOCTEXT("NiagaraOpInfo", "ArcCosineDeg Desc", "Result = acos(A)*RadiansToDegrees");
 		Op->Keywords = FText::FromString(TEXT("ArcCosine"));
-		Op->Inputs.Add(FNiagaraOpInOutInfo(A, Type, AText, AText, DefaultStr_One));
-		Op->Outputs.Add(FNiagaraOpInOutInfo(Result, Type, ResultText, ResultText, DefaultStr_One, TEXT("(180.0f/PI)*acos({0})")));
+		Op->Inputs.Add(FNiagaraOpInOutInfo(A, NumericType, AText, AText, DefaultStr_One));
+		Op->Outputs.Add(FNiagaraOpInOutInfo(Result, NumericType, ResultText, ResultText, DefaultStr_One, TEXT("(180.0f/PI)*acos({0})")));
 		Op->BuildName(TEXT("ArcCosine(Degrees)"), CategoryName);
+		Op->InputTypeValidationFunction.BindLambda(NoPositionTypeCheck);
 		OpInfoMap.Add(Op->Name) = Idx;
 
 
@@ -595,10 +735,11 @@ void FNiagaraOpInfo::Init()
 		Op->FriendlyName = NSLOCTEXT("NiagaraOpInfo", "ArcTangent Name", "ArcTangent");
 		Op->Description = NSLOCTEXT("NiagaraOpInfo", "ArcTangent Desc", "Result = atan(A)*(Period/TWO_PI)");
 		Op->Keywords = FText::FromString(TEXT("ArcTangent"));
-		Op->Inputs.Add(FNiagaraOpInOutInfo(A, Type, AText, AText, DefaultStr_One));
+		Op->Inputs.Add(FNiagaraOpInOutInfo(A, NumericType, AText, AText, DefaultStr_One));
 		Op->Inputs.Add(FNiagaraOpInOutInfo(TEXT("Period"), PeriodType, PeriodFriendlyNameText, PeriodDescText, DefaultStr_One));
-		Op->Outputs.Add(FNiagaraOpInOutInfo(Result, Type, ResultText, ResultText, DefaultStr_One, TEXT("atan({0})*({1}/TWO_PI)")));
+		Op->Outputs.Add(FNiagaraOpInOutInfo(Result, NumericType, ResultText, ResultText, DefaultStr_One, TEXT("atan({0})*({1}/TWO_PI)")));
 		Op->BuildName(TEXT("ArcTangent"), CategoryName);
+		Op->InputTypeValidationFunction.BindLambda(NoPositionTypeCheck);
 		OpInfoMap.Add(Op->Name) = Idx;
 
 		Idx = OpInfos.AddDefaulted();
@@ -607,9 +748,10 @@ void FNiagaraOpInfo::Init()
 		Op->FriendlyName = NSLOCTEXT("NiagaraOpInfo", "ArcTangentRad Name", "ArcTangent(Radians)");
 		Op->Description = NSLOCTEXT("NiagaraOpInfo", "ArcTangentRad Desc", "Result = atan(A)");
 		Op->Keywords = FText::FromString(TEXT("ArcTangent"));
-		Op->Inputs.Add(FNiagaraOpInOutInfo(A, Type, AText, AText, DefaultStr_One));
-		Op->Outputs.Add(FNiagaraOpInOutInfo(Result, Type, ResultText, ResultText, DefaultStr_One, TEXT("atan({0})")));
+		Op->Inputs.Add(FNiagaraOpInOutInfo(A, NumericType, AText, AText, DefaultStr_One));
+		Op->Outputs.Add(FNiagaraOpInOutInfo(Result, NumericType, ResultText, ResultText, DefaultStr_One, TEXT("atan({0})")));
 		Op->BuildName(TEXT("ArcTangent(Radians)"), CategoryName);
+		Op->InputTypeValidationFunction.BindLambda(NoPositionTypeCheck);
 		OpInfoMap.Add(Op->Name) = Idx;
 
 		Idx = OpInfos.AddDefaulted();
@@ -618,9 +760,10 @@ void FNiagaraOpInfo::Init()
 		Op->FriendlyName = NSLOCTEXT("NiagaraOpInfo", "ArcTangentDeg Name", "ArcTangent(Degrees)");
 		Op->Description = NSLOCTEXT("NiagaraOpInfo", "ArcTangentDeg Desc", "Result = atan(A)*RadiansToDegrees");
 		Op->Keywords = FText::FromString(TEXT("ArcTangent"));
-		Op->Inputs.Add(FNiagaraOpInOutInfo(A, Type, AText, AText, DefaultStr_One));
-		Op->Outputs.Add(FNiagaraOpInOutInfo(Result, Type, ResultText, ResultText, DefaultStr_One, TEXT("(180.0f/PI)*atan({0})")));
+		Op->Inputs.Add(FNiagaraOpInOutInfo(A, NumericType, AText, AText, DefaultStr_One));
+		Op->Outputs.Add(FNiagaraOpInOutInfo(Result, NumericType, ResultText, ResultText, DefaultStr_One, TEXT("(180.0f/PI)*atan({0})")));
 		Op->BuildName(TEXT("ArcTangent(Degrees)"), CategoryName);
+		Op->InputTypeValidationFunction.BindLambda(NoPositionTypeCheck);
 		OpInfoMap.Add(Op->Name) = Idx;
 
 		Idx = OpInfos.AddDefaulted();
@@ -628,11 +771,12 @@ void FNiagaraOpInfo::Init()
 		Op->Category = NumericCategory;
 		Op->FriendlyName = NSLOCTEXT("NiagaraOpInfo", "ATan2 Name", "ArcTangent2");
 		Op->Description = NSLOCTEXT("NiagaraOpInfo", "ATan2 Desc", "ResultInPeriod = Period * atan2(A, B) / 2PI");
-		Op->Inputs.Add(FNiagaraOpInOutInfo(A, Type, AText, AText, DefaultStr_One));
-		Op->Inputs.Add(FNiagaraOpInOutInfo(B, Type, BText, BText, DefaultStr_One));
+		Op->Inputs.Add(FNiagaraOpInOutInfo(A, NumericType, AText, AText, DefaultStr_One));
+		Op->Inputs.Add(FNiagaraOpInOutInfo(B, NumericType, BText, BText, DefaultStr_One));
 		Op->Inputs.Add(FNiagaraOpInOutInfo(TEXT("Period"), PeriodType, PeriodFriendlyNameText, PeriodDescText, DefaultStr_One));
-		Op->Outputs.Add(FNiagaraOpInOutInfo(Result, Type, ResultText, ResultText, DefaultStr_One, TEXT("atan2({0},{1})*({2}/TWO_PI)")));
+		Op->Outputs.Add(FNiagaraOpInOutInfo(Result, NumericType, ResultText, ResultText, DefaultStr_One, TEXT("atan2({0},{1})*({2}/TWO_PI)")));
 		Op->BuildName(TEXT("ArcTangent2"), CategoryName);
+		Op->InputTypeValidationFunction.BindLambda(NoPositionTypeCheck);
 		OpInfoMap.Add(Op->Name) = Idx;
 
 		Idx = OpInfos.AddDefaulted();
@@ -640,10 +784,11 @@ void FNiagaraOpInfo::Init()
 		Op->Category = NumericCategory;
 		Op->FriendlyName = NSLOCTEXT("NiagaraOpInfo", "ATan2Rad Name", "ArcTangent2(Radians)");
 		Op->Description = NSLOCTEXT("NiagaraOpInfo", "ATan2Rad Desc", "ResultInRadians = atan2(A, B)");
-		Op->Inputs.Add(FNiagaraOpInOutInfo(A, Type, AText, AText, DefaultStr_One));
-		Op->Inputs.Add(FNiagaraOpInOutInfo(B, Type, BText, BText, DefaultStr_One));
-		Op->Outputs.Add(FNiagaraOpInOutInfo(Result, Type, ResultText, ResultText, DefaultStr_One, TEXT("atan2({0},{1})")));
+		Op->Inputs.Add(FNiagaraOpInOutInfo(A, NumericType, AText, AText, DefaultStr_One));
+		Op->Inputs.Add(FNiagaraOpInOutInfo(B, NumericType, BText, BText, DefaultStr_One));
+		Op->Outputs.Add(FNiagaraOpInOutInfo(Result, NumericType, ResultText, ResultText, DefaultStr_One, TEXT("atan2({0},{1})")));
 		Op->BuildName(TEXT("ArcTangent2(Radians)"), CategoryName);
+		Op->InputTypeValidationFunction.BindLambda(NoPositionTypeCheck);
 		OpInfoMap.Add(Op->Name) = Idx;
 
 		Idx = OpInfos.AddDefaulted();
@@ -651,10 +796,11 @@ void FNiagaraOpInfo::Init()
 		Op->Category = NumericCategory;
 		Op->FriendlyName = NSLOCTEXT("NiagaraOpInfo", "ATan2Deg Name", "ArcTangent2(Degrees)");
 		Op->Description = NSLOCTEXT("NiagaraOpInfo", "ATan2Deg Desc", "ResultInPeriod = 180 * atan2(A, B) / PI");
-		Op->Inputs.Add(FNiagaraOpInOutInfo(A, Type, AText, AText, DefaultStr_One));
-		Op->Inputs.Add(FNiagaraOpInOutInfo(B, Type, BText, BText, DefaultStr_One));
-		Op->Outputs.Add(FNiagaraOpInOutInfo(Result, Type, ResultText, ResultText, DefaultStr_One, TEXT("atan2({0},{1})*(180.0f/PI)")));
+		Op->Inputs.Add(FNiagaraOpInOutInfo(A, NumericType, AText, AText, DefaultStr_One));
+		Op->Inputs.Add(FNiagaraOpInOutInfo(B, NumericType, BText, BText, DefaultStr_One));
+		Op->Outputs.Add(FNiagaraOpInOutInfo(Result, NumericType, ResultText, ResultText, DefaultStr_One, TEXT("atan2({0},{1})*(180.0f/PI)")));
 		Op->BuildName(TEXT("ArcTangent2(Degrees)"), CategoryName);
+		Op->InputTypeValidationFunction.BindLambda(NoPositionTypeCheck);
 		OpInfoMap.Add(Op->Name) = Idx;
 
 		Idx = OpInfos.AddDefaulted();
@@ -682,8 +828,8 @@ void FNiagaraOpInfo::Init()
 		Op->Category = NumericCategory;
 		Op->FriendlyName = NSLOCTEXT("NiagaraOpInfo", "Ceil Name", "Ceil");
 		Op->Description = NSLOCTEXT("NiagaraOpInfo", "Ceil Desc", "Rounds A to the nearest integer higher than A.");
-		Op->Inputs.Add(FNiagaraOpInOutInfo(A, Type, AText, AText, DefaultStr_One));
-		Op->Outputs.Add(FNiagaraOpInOutInfo(Result, Type, ResultText, ResultText, DefaultStr_One, TEXT("ceil({0})")));
+		Op->Inputs.Add(FNiagaraOpInOutInfo(A, NumericType, AText, AText, DefaultStr_One));
+		Op->Outputs.Add(FNiagaraOpInOutInfo(Result, NumericType, ResultText, ResultText, DefaultStr_One, TEXT("ceil({0})")));
 		Op->BuildName(TEXT("Ceil"), CategoryName);
 		OpInfoMap.Add(Op->Name) = Idx;
 
@@ -692,8 +838,8 @@ void FNiagaraOpInfo::Init()
 		Op->Category = NumericCategory;
 		Op->FriendlyName = NSLOCTEXT("NiagaraOpInfo", "Floor Name", "Floor");
 		Op->Description = NSLOCTEXT("NiagaraOpInfo", "Floor Desc", "Rounds A to the nearest integer lower than A.");
-		Op->Inputs.Add(FNiagaraOpInOutInfo(A, Type, AText, AText, DefaultStr_One));
-		Op->Outputs.Add(FNiagaraOpInOutInfo(Result, Type, ResultText, ResultText, DefaultStr_One, TEXT("floor({0})")));
+		Op->Inputs.Add(FNiagaraOpInOutInfo(A, NumericType, AText, AText, DefaultStr_One));
+		Op->Outputs.Add(FNiagaraOpInOutInfo(Result, NumericType, ResultText, ResultText, DefaultStr_One, TEXT("floor({0})")));
 		Op->BuildName(TEXT("Floor"), CategoryName);
 		OpInfoMap.Add(Op->Name) = Idx;
 
@@ -702,8 +848,8 @@ void FNiagaraOpInfo::Init()
 		Op->Category = NumericCategory;
 		Op->FriendlyName = NSLOCTEXT("NiagaraOpInfo", "Round Name", "Round");
 		Op->Description = NSLOCTEXT("NiagaraOpInfo", "Round Desc", "Rounds A to the nearest integer.");
-		Op->Inputs.Add(FNiagaraOpInOutInfo(A, Type, AText, AText, DefaultStr_One));
-		Op->Outputs.Add(FNiagaraOpInOutInfo(Result, Type, ResultText, ResultText, DefaultStr_One, TEXT("round({0})")));
+		Op->Inputs.Add(FNiagaraOpInOutInfo(A, NumericType, AText, AText, DefaultStr_One));
+		Op->Outputs.Add(FNiagaraOpInOutInfo(Result, NumericType, ResultText, ResultText, DefaultStr_One, TEXT("round({0})")));
 		Op->BuildName(TEXT("Round"), CategoryName);
 		OpInfoMap.Add(Op->Name) = Idx;
 
@@ -713,10 +859,11 @@ void FNiagaraOpInfo::Init()
 		Op->FriendlyName = NSLOCTEXT("NiagaraOpInfo", "Fmod Name", "Modulo");
 		Op->Description = NSLOCTEXT("NiagaraOpInfo", "Fmod Desc", "Result = A % B");
 		Op->Keywords = FText::FromString(TEXT("%"));
-		Op->Inputs.Add(FNiagaraOpInOutInfo(A, Type, AText, AText, DefaultStr_One));
-		Op->Inputs.Add(FNiagaraOpInOutInfo(B, Type, BText, BText, DefaultStr_One));
-		Op->Outputs.Add(FNiagaraOpInOutInfo(Result, Type, ResultText, ResultText, DefaultStr_One, TEXT("ModuloPrecise({0}, {1})")));
+		Op->Inputs.Add(FNiagaraOpInOutInfo(A, NumericType, AText, AText, DefaultStr_One));
+		Op->Inputs.Add(FNiagaraOpInOutInfo(B, NumericType, BText, BText, DefaultStr_One));
+		Op->Outputs.Add(FNiagaraOpInOutInfo(Result, NumericType, ResultText, ResultText, DefaultStr_One, TEXT("ModuloPrecise({0}, {1})")));
 		Op->BuildName(TEXT("FMod"), CategoryName);
+		Op->InputTypeValidationFunction.BindLambda(NoPositionTypeCheck);
 		OpInfoMap.Add(Op->Name) = Idx;
 
 		Idx = OpInfos.AddDefaulted();
@@ -725,10 +872,11 @@ void FNiagaraOpInfo::Init()
 		Op->FriendlyName = NSLOCTEXT("NiagaraOpInfo", "Fmod Name Fast", "Modulo Fast");
 		Op->Description = NSLOCTEXT("NiagaraOpInfo", "Fmod Desc Fast", "Result = A % B. May be less precise than regular FMod.");
 		Op->Keywords = FText::FromString(TEXT("%"));
-		Op->Inputs.Add(FNiagaraOpInOutInfo(A, Type, AText, AText, DefaultStr_One));
-		Op->Inputs.Add(FNiagaraOpInOutInfo(B, Type, BText, BText, DefaultStr_One));
-		Op->Outputs.Add(FNiagaraOpInOutInfo(Result, Type, ResultText, ResultText, DefaultStr_One, TEXT("Modulo({0}, {1})")));
+		Op->Inputs.Add(FNiagaraOpInOutInfo(A, NumericType, AText, AText, DefaultStr_One));
+		Op->Inputs.Add(FNiagaraOpInOutInfo(B, NumericType, BText, BText, DefaultStr_One));
+		Op->Outputs.Add(FNiagaraOpInOutInfo(Result, NumericType, ResultText, ResultText, DefaultStr_One, TEXT("Modulo({0}, {1})")));
 		Op->BuildName(TEXT("FModFast"), CategoryName);
+		Op->InputTypeValidationFunction.BindLambda(NoPositionTypeCheck);
 		OpInfoMap.Add(Op->Name) = Idx;
 
 		Idx = OpInfos.AddDefaulted();
@@ -736,8 +884,8 @@ void FNiagaraOpInfo::Init()
 		Op->Category = NumericCategory;
 		Op->FriendlyName = NSLOCTEXT("NiagaraOpInfo", "Frac Name", "Frac");
 		Op->Description = NSLOCTEXT("NiagaraOpInfo", "Frac Desc", "Result = frac(A)");
-		Op->Inputs.Add(FNiagaraOpInOutInfo(A, Type, AText, AText, DefaultStr_One));
-		Op->Outputs.Add(FNiagaraOpInOutInfo(Result, Type, ResultText, ResultText, DefaultStr_One, TEXT("frac({0})")));
+		Op->Inputs.Add(FNiagaraOpInOutInfo(A, NumericType, AText, AText, DefaultStr_One));
+		Op->Outputs.Add(FNiagaraOpInOutInfo(Result, NumericType, ResultText, ResultText, DefaultStr_One, TEXT("frac({0})")));
 		Op->BuildName(TEXT("Frac"), CategoryName);
 		OpInfoMap.Add(Op->Name) = Idx;
 
@@ -746,8 +894,8 @@ void FNiagaraOpInfo::Init()
 		Op->Category = NumericCategory;
 		Op->FriendlyName = NSLOCTEXT("NiagaraOpInfo", "Trunc Name", "Trunc");
 		Op->Description = NSLOCTEXT("NiagaraOpInfo", "Trunc Desc", "Result = trunc(A)");
-		Op->Inputs.Add(FNiagaraOpInOutInfo(A, Type, AText, AText, DefaultStr_One));
-		Op->Outputs.Add(FNiagaraOpInOutInfo(Result, Type, ResultText, ResultText, DefaultStr_One, TEXT("trunc({0})")));
+		Op->Inputs.Add(FNiagaraOpInOutInfo(A, NumericType, AText, AText, DefaultStr_One));
+		Op->Outputs.Add(FNiagaraOpInOutInfo(Result, NumericType, ResultText, ResultText, DefaultStr_One, TEXT("trunc({0})")));
 		Op->BuildName(TEXT("Trunc"), CategoryName);
 		OpInfoMap.Add(Op->Name) = Idx;
 
@@ -756,11 +904,31 @@ void FNiagaraOpInfo::Init()
 		Op->Category = NumericCategory;
 		Op->FriendlyName = NSLOCTEXT("NiagaraOpInfo", "Clamp Name", "Clamp");
 		Op->Description = NSLOCTEXT("NiagaraOpInfo", "Clamp Desc", "Result = clamp(A, Min, Max)");
-		Op->Inputs.Add(FNiagaraOpInOutInfo(A, Type, AText, AText, DefaultStr_One));
-		Op->Inputs.Add(FNiagaraOpInOutInfo(Min, Type, MinText, MinText, DefaultStr_One));
-		Op->Inputs.Add(FNiagaraOpInOutInfo(Max, Type, MaxText, MaxText, DefaultStr_One));
-		Op->Outputs.Add(FNiagaraOpInOutInfo(Result, Type, ResultText, ResultText, DefaultStr_One, TEXT("clamp({0},{1},{2})")));
+		Op->Inputs.Add(FNiagaraOpInOutInfo(A, NumericType, AText, AText, DefaultStr_One));
+		Op->Inputs.Add(FNiagaraOpInOutInfo(Min, NumericType, MinText, MinText, DefaultStr_One));
+		Op->Inputs.Add(FNiagaraOpInOutInfo(Max, NumericType, MaxText, MaxText, DefaultStr_One));
+		Op->Outputs.Add(FNiagaraOpInOutInfo(Result, NumericType, ResultText, ResultText, DefaultStr_One, TEXT("clamp({0},{1},{2})")));
 		Op->BuildName(TEXT("Clamp"), CategoryName);
+		Op->InputTypeValidationFunction.BindLambda([=](const TArray<FNiagaraTypeDefinition>& InputTypes, FText& ErrorMessage)
+			{
+				int32 PositionTypes = CountType(InputTypes, FNiagaraTypeDefinition::GetPositionDef());
+				if (PositionTypes > 0 && PositionTypes != 3)
+				{
+					ErrorMessage = NSLOCTEXT("NiagaraOpInfo", "Clamp Type Error", "The clamp operation can only clamp between positions, not between a position and another type.");
+					return false;
+				}
+				return true;
+			});
+		Op->NumericOuputTypeSelectionMode = ENiagaraNumericOutputTypeSelectionMode::Custom;
+		Op->CustomNumericResolveFunction.BindLambda([=](const TArray<FNiagaraTypeDefinition>& NonNumericInputs)
+			{
+				 int32 PositionTypes = CountType(NonNumericInputs, FNiagaraTypeDefinition::GetPositionDef());
+				if (PositionTypes == 3)
+				{
+					return FNiagaraTypeDefinition::GetPositionDef();
+				}
+				return FNiagaraTypeDefinition::GetNumericOutputType(NonNumericInputs, ENiagaraNumericOutputTypeSelectionMode::Largest);
+			});
 		OpInfoMap.Add(Op->Name) = Idx;
 
 		Idx = OpInfos.AddDefaulted();
@@ -768,13 +936,33 @@ void FNiagaraOpInfo::Init()
 		Op->Category = NumericCategory;
 		Op->FriendlyName = NSLOCTEXT("NiagaraOpInfo", "Min Name", "Min");
 		Op->Description = NSLOCTEXT("NiagaraOpInfo", "Min Desc", "Result = min(A, B)");
-		Op->Inputs.Add(FNiagaraOpInOutInfo(A, Type, AText, AText, DefaultStr_One));
-		Op->Inputs.Add(FNiagaraOpInOutInfo(B, Type, BText, BText, DefaultStr_One));
-		Op->Outputs.Add(FNiagaraOpInOutInfo(Result, Type, ResultText, ResultText, DefaultStr_One, TEXT("min({0},{1})")));
+		Op->Inputs.Add(FNiagaraOpInOutInfo(A, NumericType, AText, AText, DefaultStr_One));
+		Op->Inputs.Add(FNiagaraOpInOutInfo(B, NumericType, BText, BText, DefaultStr_One));
+		Op->Outputs.Add(FNiagaraOpInOutInfo(Result, NumericType, ResultText, ResultText, DefaultStr_One, TEXT("min({0},{1})")));
 		Op->BuildName(TEXT("Min"), CategoryName);
 		Op->bSupportsAddedInputs = true;
-		Op->AddedInputTypeRestrictions.Add(Type);
+		Op->AddedInputTypeRestrictions.Add(NumericType);
 		Op->AddedInputFormatting = TEXT("min({A}, {B})");
+		Op->InputTypeValidationFunction.BindLambda([=](const TArray<FNiagaraTypeDefinition>& InputTypes, FText& ErrorMessage)
+			{
+				int32 PositionTypes = CountType(InputTypes, FNiagaraTypeDefinition::GetPositionDef());
+				if (PositionTypes > 0 && PositionTypes != 2)
+				{
+					ErrorMessage = NSLOCTEXT("NiagaraOpInfo", "Min Type Error", "The min operation cannot be used for a position and a non-position type.");
+					return false;
+				}
+				return true;
+			});
+		Op->NumericOuputTypeSelectionMode = ENiagaraNumericOutputTypeSelectionMode::Custom;
+		Op->CustomNumericResolveFunction.BindLambda([=](const TArray<FNiagaraTypeDefinition>& NonNumericInputs)
+			{
+				 int32 PositionTypes = CountType(NonNumericInputs, FNiagaraTypeDefinition::GetPositionDef());
+				if (PositionTypes == 2)
+				{
+					return FNiagaraTypeDefinition::GetPositionDef();
+				}
+				return FNiagaraTypeDefinition::GetNumericOutputType(NonNumericInputs, ENiagaraNumericOutputTypeSelectionMode::Largest);
+			});
 		OpInfoMap.Add(Op->Name) = Idx;
 
 		Idx = OpInfos.AddDefaulted();
@@ -782,13 +970,33 @@ void FNiagaraOpInfo::Init()
 		Op->Category = NumericCategory;
 		Op->FriendlyName = NSLOCTEXT("NiagaraOpInfo", "Max Name", "Max");
 		Op->Description = NSLOCTEXT("NiagaraOpInfo", "Max Desc", "Result = max(A, B)");
-		Op->Inputs.Add(FNiagaraOpInOutInfo(A, Type, AText, AText, DefaultStr_One));
-		Op->Inputs.Add(FNiagaraOpInOutInfo(B, Type, BText, BText, DefaultStr_One));
-		Op->Outputs.Add(FNiagaraOpInOutInfo(Result, Type, ResultText, ResultText, DefaultStr_One, TEXT("max({0},{1})")));
+		Op->Inputs.Add(FNiagaraOpInOutInfo(A, NumericType, AText, AText, DefaultStr_One));
+		Op->Inputs.Add(FNiagaraOpInOutInfo(B, NumericType, BText, BText, DefaultStr_One));
+		Op->Outputs.Add(FNiagaraOpInOutInfo(Result, NumericType, ResultText, ResultText, DefaultStr_One, TEXT("max({0},{1})")));
 		Op->BuildName(TEXT("Max"), CategoryName);
 		Op->bSupportsAddedInputs = true;
-		Op->AddedInputTypeRestrictions.Add(Type);
+		Op->AddedInputTypeRestrictions.Add(NumericType);
 		Op->AddedInputFormatting = TEXT("max({A}, {B})");
+		Op->InputTypeValidationFunction.BindLambda([=](const TArray<FNiagaraTypeDefinition>& InputTypes, FText& ErrorMessage)
+			{
+				int32 PositionTypes = CountType(InputTypes, FNiagaraTypeDefinition::GetPositionDef());
+				if (PositionTypes > 0 && PositionTypes != 2)
+				{
+					ErrorMessage = NSLOCTEXT("NiagaraOpInfo", "Max Type Error", "The max operation cannot be used for a position and a non-position type.");
+					return false;
+				}
+				return true;
+			});
+		Op->NumericOuputTypeSelectionMode = ENiagaraNumericOutputTypeSelectionMode::Custom;
+		Op->CustomNumericResolveFunction.BindLambda([=](const TArray<FNiagaraTypeDefinition>& NonNumericInputs)
+			{
+				 int32 PositionTypes = CountType(NonNumericInputs, FNiagaraTypeDefinition::GetPositionDef());
+				if (PositionTypes == 2)
+				{
+					return FNiagaraTypeDefinition::GetPositionDef();
+				}
+				return FNiagaraTypeDefinition::GetNumericOutputType(NonNumericInputs, ENiagaraNumericOutputTypeSelectionMode::Largest);
+			});
 		OpInfoMap.Add(Op->Name) = Idx;
 
 		Idx = OpInfos.AddDefaulted();
@@ -796,9 +1004,10 @@ void FNiagaraOpInfo::Init()
 		Op->Category = NumericCategory;
 		Op->FriendlyName = NSLOCTEXT("NiagaraOpInfo", "Pow Name", "Pow");
 		Op->Description = NSLOCTEXT("NiagaraOpInfo", "Pow Desc", "Result = pow(A, B). When A < 0, it will be clamped to 0.");
-		Op->Inputs.Add(FNiagaraOpInOutInfo(A, Type, AText, AText, DefaultStr_One));
-		Op->Inputs.Add(FNiagaraOpInOutInfo(B, Type, BText, BText, DefaultStr_One));
-		Op->Outputs.Add(FNiagaraOpInOutInfo(Result, Type, ResultText, ResultText, DefaultStr_One, TEXT("pow(max(0, {0}),{1})")));
+		Op->Inputs.Add(FNiagaraOpInOutInfo(A, NumericType, AText, AText, DefaultStr_One));
+		Op->Inputs.Add(FNiagaraOpInOutInfo(B, NumericType, BText, BText, DefaultStr_One));
+		Op->Outputs.Add(FNiagaraOpInOutInfo(Result, NumericType, ResultText, ResultText, DefaultStr_One, TEXT("pow(max(0, {0}),{1})")));
+		Op->InputTypeValidationFunction.BindLambda(NoPositionTypeCheck);
 		Op->BuildName(TEXT("Pow"), CategoryName);
 		OpInfoMap.Add(Op->Name) = Idx;
 
@@ -807,8 +1016,9 @@ void FNiagaraOpInfo::Init()
 		Op->Category = NumericCategory;
 		Op->FriendlyName = NSLOCTEXT("NiagaraOpInfo", "Sign Name", "Sign");
 		Op->Description = NSLOCTEXT("NiagaraOpInfo", "Sign Desc", "Result = sign(A)");
-		Op->Inputs.Add(FNiagaraOpInOutInfo(A, Type, AText, AText, DefaultStr_One));
-		Op->Outputs.Add(FNiagaraOpInOutInfo(Result, Type, ResultText, ResultText, DefaultStr_One, TEXT("sign({0})")));
+		Op->Inputs.Add(FNiagaraOpInOutInfo(A, NumericType, AText, AText, DefaultStr_One));
+		Op->Outputs.Add(FNiagaraOpInOutInfo(Result, NumericType, ResultText, ResultText, DefaultStr_One, TEXT("sign({0})")));
+		Op->InputTypeValidationFunction.BindLambda(NoPositionTypeCheck);
 		Op->BuildName(TEXT("Sign"), CategoryName);
 		OpInfoMap.Add(Op->Name) = Idx;
 
@@ -817,8 +1027,9 @@ void FNiagaraOpInfo::Init()
 		Op->Category = NumericCategory;
 		Op->FriendlyName = NSLOCTEXT("NiagaraOpInfo", "Step Name", "Step");
 		Op->Description = NSLOCTEXT("NiagaraOpInfo", "Step Desc", "Result = step(0, A)");
-		Op->Inputs.Add(FNiagaraOpInOutInfo(A, Type, AText, AText, DefaultStr_One));
-		Op->Outputs.Add(FNiagaraOpInOutInfo(Result, Type, ResultText, ResultText, DefaultStr_One, TEXT("step(0, {0})")));
+		Op->Inputs.Add(FNiagaraOpInOutInfo(A, NumericType, AText, AText, DefaultStr_One));
+		Op->Outputs.Add(FNiagaraOpInOutInfo(Result, NumericType, ResultText, ResultText, DefaultStr_One, TEXT("step(0, {0})")));
+		Op->InputTypeValidationFunction.BindLambda(NoPositionTypeCheck);
 		Op->BuildName(TEXT("Step"), CategoryName);
 		OpInfoMap.Add(Op->Name) = Idx;
 
@@ -827,8 +1038,8 @@ void FNiagaraOpInfo::Init()
 		Op->Category = NumericCategory;
 		Op->FriendlyName = NSLOCTEXT("NiagaraOpInfo", "Noise Name", "Noise");
 		Op->Description = NSLOCTEXT("NiagaraOpInfo", "Noise Desc", "A continuous pseudo random noise function.");
-		Op->Inputs.Add(FNiagaraOpInOutInfo(X, Type, XText, XText, DefaultStr_Zero));
-		Op->Outputs.Add(FNiagaraOpInOutInfo(Result, Type, ResultText, ResultText, DefaultStr_One, TEXT("noise({0})")));
+		Op->Inputs.Add(FNiagaraOpInOutInfo(X, NumericType, XText, XText, DefaultStr_Zero));
+		Op->Outputs.Add(FNiagaraOpInOutInfo(Result, NumericType, ResultText, ResultText, DefaultStr_One, TEXT("noise({0})")));
 		Op->NumericOuputTypeSelectionMode = ENiagaraNumericOutputTypeSelectionMode::Scalar;
 		Op->BuildName(TEXT("Noise"), CategoryName);
 		OpInfoMap.Add(Op->Name) = Idx;
@@ -838,9 +1049,10 @@ void FNiagaraOpInfo::Init()
 		Op->Category = NumericCategory;
 		Op->FriendlyName = NSLOCTEXT("NiagaraOpInfo", "Dot Name", "Dot");
 		Op->Description = NSLOCTEXT("NiagaraOpInfo", "Dot Desc", "Dot product of two vectors.");
-		Op->Inputs.Add(FNiagaraOpInOutInfo(A, Type, AText, AText, DefaultStr_One));
-		Op->Inputs.Add(FNiagaraOpInOutInfo(B, Type, BText, BText, DefaultStr_One));
-		Op->Outputs.Add(FNiagaraOpInOutInfo(Result, Type, ResultText, ResultText, DefaultStr_One, TEXT("dot({0},{1})")));
+		Op->Inputs.Add(FNiagaraOpInOutInfo(A, NumericType, AText, AText, DefaultStr_One));
+		Op->Inputs.Add(FNiagaraOpInOutInfo(B, NumericType, BText, BText, DefaultStr_One));
+		Op->Outputs.Add(FNiagaraOpInOutInfo(Result, NumericType, ResultText, ResultText, DefaultStr_One, TEXT("dot({0},{1})")));
+		Op->InputTypeValidationFunction.BindLambda(NoPositionTypeCheck);
 		Op->NumericOuputTypeSelectionMode = ENiagaraNumericOutputTypeSelectionMode::Scalar;
 		Op->BuildName(TEXT("Dot"), CategoryName);
 		OpInfoMap.Add(Op->Name) = Idx;
@@ -850,8 +1062,9 @@ void FNiagaraOpInfo::Init()
 		Op->Category = NumericCategory;
 		Op->FriendlyName = NSLOCTEXT("NiagaraOpInfo", "Normalize Name", "Normalize");
 		Op->Description = NSLOCTEXT("NiagaraOpInfo", "Normalize Desc", "Normalizes the passed value.");
-		Op->Inputs.Add(FNiagaraOpInOutInfo(A, Type, AText, AText, DefaultStr_One));
-		Op->Outputs.Add(FNiagaraOpInOutInfo(Result, Type, ResultText, ResultText, DefaultStr_One, TEXT("normalize({0})")));
+		Op->Inputs.Add(FNiagaraOpInOutInfo(A, NumericType, AText, AText, DefaultStr_One));
+		Op->Outputs.Add(FNiagaraOpInOutInfo(Result, NumericType, ResultText, ResultText, DefaultStr_One, TEXT("normalize({0})")));
+		Op->InputTypeValidationFunction.BindLambda(NoPositionTypeCheck);
 		Op->BuildName(TEXT("Normalize"), CategoryName);
 		OpInfoMap.Add(Op->Name) = Idx;
 
@@ -860,20 +1073,34 @@ void FNiagaraOpInfo::Init()
 		Op->Category = NumericCategory;
 		Op->FriendlyName = NSLOCTEXT("NiagaraOpInfo", "Length Name", "Length");
 		Op->Description = NSLOCTEXT("NiagaraOpInfo", "Length Desc", "Returns the length of the passed value.");
-		Op->Inputs.Add(FNiagaraOpInOutInfo(A, Type, AText, AText, DefaultStr_One));
-		Op->Outputs.Add(FNiagaraOpInOutInfo(Result, Type, ResultText, ResultText, DefaultStr_One, TEXT("length({0})")));
+		Op->Inputs.Add(FNiagaraOpInOutInfo(A, NumericType, AText, AText, DefaultStr_One));
+		Op->Outputs.Add(FNiagaraOpInOutInfo(Result, NumericType, ResultText, ResultText, DefaultStr_One, TEXT("length({0})")));
+		Op->InputTypeValidationFunction.BindLambda(NoPositionTypeCheck);
 		Op->NumericOuputTypeSelectionMode = ENiagaraNumericOutputTypeSelectionMode::Scalar;
 		Op->BuildName(TEXT("Length"), CategoryName);
 		OpInfoMap.Add(Op->Name) = Idx;
 
+		Idx = OpInfos.AddDefaulted();
+		Op = &OpInfos[Idx];
+		Op->Category = NumericCategory;
+		Op->FriendlyName = NSLOCTEXT("NiagaraOpInfo", "Distance Name", "Distance");
+		Op->Description = NSLOCTEXT("NiagaraOpInfo", "Distance Desc", "Returns the distance between two positions.");
+		Op->Inputs.Add(FNiagaraOpInOutInfo(A, FNiagaraTypeDefinition::GetPositionDef(), AText, AText, DefaultStr_One));
+		Op->Inputs.Add(FNiagaraOpInOutInfo(B, FNiagaraTypeDefinition::GetPositionDef(), BText, BText, DefaultStr_One));
+		Op->Outputs.Add(FNiagaraOpInOutInfo(Result, FNiagaraTypeDefinition::GetFloatDef(), ResultText, ResultText, DefaultStr_One, TEXT("length({0} - {1})")));
+		Op->NumericOuputTypeSelectionMode = ENiagaraNumericOutputTypeSelectionMode::None;
+		Op->BuildName(TEXT("DistancePos"), CategoryName);
+		OpInfoMap.Add(Op->Name) = Idx;
+
 		// Non-deterministic random number generation. Calls FRandomStream on the CPU. 
 		Idx = OpInfos.AddDefaulted();
 		Op = &OpInfos[Idx];
 		Op->Category = NumericCategory;
 		Op->FriendlyName = NSLOCTEXT("NiagaraOpInfo", "Rand Name", "Random");
 		Op->Description = NSLOCTEXT("NiagaraOpInfo", "Rand Desc", "Returns a non-deterministic random value between 0 and A.");
-		Op->Inputs.Add(FNiagaraOpInOutInfo(A, Type, AText, AText, DefaultStr_One));
-		Op->Outputs.Add(FNiagaraOpInOutInfo(Result, Type, ResultText, ResultText, DefaultStr_One, TEXT("rand({0})")));
+		Op->Inputs.Add(FNiagaraOpInOutInfo(A, NumericType, AText, AText, DefaultStr_One));
+		Op->Outputs.Add(FNiagaraOpInOutInfo(Result, NumericType, ResultText, ResultText, DefaultStr_One, TEXT("rand({0})")));
+		Op->InputTypeValidationFunction.BindLambda(NoPositionTypeCheck);
 		Op->BuildName(TEXT("Rand"), CategoryName);
 		OpInfoMap.Add(Op->Name) = Idx;
 
@@ -883,8 +1110,9 @@ void FNiagaraOpInfo::Init()
 		Op->Category = NumericCategory;
 		Op->FriendlyName = NSLOCTEXT("NiagaraOpInfo", "Rand Integer Name", "Random Integer");
 		Op->Description = NSLOCTEXT("NiagaraOpInfo", "Rand Integer Desc", "Returns a non-deterministic random integer value between 0 and Max-1");
-		Op->Inputs.Add(FNiagaraOpInOutInfo(Max, Type, MaxText, MaxText, DefaultStr_One));
-		Op->Outputs.Add(FNiagaraOpInOutInfo(Result, Type, ResultText, ResultText, DefaultStr_One, TEXT("rand_int({0})")));
+		Op->Inputs.Add(FNiagaraOpInOutInfo(Max, NumericType, MaxText, MaxText, DefaultStr_One));
+		Op->Outputs.Add(FNiagaraOpInOutInfo(Result, NumericType, ResultText, ResultText, DefaultStr_One, TEXT("rand_int({0})")));
+		Op->InputTypeValidationFunction.BindLambda(NoPositionTypeCheck);
 		Op->BuildName(TEXT("Rand Integer"), CategoryName);
 		Op->bNumericsCanBeIntegers = true;
 		Op->bNumericsCanBeFloats = false;
@@ -896,8 +1124,9 @@ void FNiagaraOpInfo::Init()
 		Op->Category = NumericCategory;
 		Op->FriendlyName = NSLOCTEXT("NiagaraOpInfo", "Rand Float Name", "Random Float");
 		Op->Description = NSLOCTEXT("NiagaraOpInfo", "Rand Float Desc", "Returns a non-deterministic random float value between 0 and Max");
-		Op->Inputs.Add(FNiagaraOpInOutInfo(Max, Type, MaxText, MaxText, DefaultStr_One));
-		Op->Outputs.Add(FNiagaraOpInOutInfo(Result, Type, ResultText, ResultText, DefaultStr_One, TEXT("rand_float({0})")));
+		Op->Inputs.Add(FNiagaraOpInOutInfo(Max, NumericType, MaxText, MaxText, DefaultStr_One));
+		Op->Outputs.Add(FNiagaraOpInOutInfo(Result, NumericType, ResultText, ResultText, DefaultStr_One, TEXT("rand_float({0})")));
+		Op->InputTypeValidationFunction.BindLambda(NoPositionTypeCheck);
 		Op->BuildName(TEXT("Rand Float"), CategoryName);
 		Op->bNumericsCanBeIntegers = false;
 		Op->bNumericsCanBeFloats = true;
@@ -918,11 +1147,12 @@ void FNiagaraOpInfo::Init()
 		Op->Category = NumericCategory;
 		Op->FriendlyName = NSLOCTEXT("NiagaraOpInfo", "Seeded Rand Name", "Seeded Random");
 		Op->Description = NSLOCTEXT("NiagaraOpInfo", "Seeded Rand Desc", "Returns a deterministic random value between 0 and A.");
-		Op->Inputs.Add(FNiagaraOpInOutInfo(A, Type, AText, AText, DefaultStr_One));
+		Op->Inputs.Add(FNiagaraOpInOutInfo(A, NumericType, AText, AText, DefaultStr_One));
 		Op->Inputs.Add(FNiagaraOpInOutInfo(SeedName1, SeedType, SeedText1, SeedText1, DefaultSeed_Zero));
 		Op->Inputs.Add(FNiagaraOpInOutInfo(SeedName2, SeedType, SeedText2, SeedText2, DefaultSeed_Zero));
 		Op->Inputs.Add(FNiagaraOpInOutInfo(SeedName3, SeedType, SeedText3, SeedText3, DefaultSeed_Zero));
-		Op->Outputs.Add(FNiagaraOpInOutInfo(Result, Type, ResultText, ResultText, DefaultStr_One, TEXT("rand({0}, {1}, {2}, {3})")));
+		Op->Outputs.Add(FNiagaraOpInOutInfo(Result, NumericType, ResultText, ResultText, DefaultStr_One, TEXT("rand({0}, {1}, {2}, {3})")));
+		Op->InputTypeValidationFunction.BindLambda(NoPositionTypeCheck);
 		Op->BuildName(TEXT("SeededRand"), CategoryName);
 		OpInfoMap.Add(Op->Name) = Idx;
 
@@ -931,11 +1161,12 @@ void FNiagaraOpInfo::Init()
 		Op->Category = NumericCategory;
 		Op->FriendlyName = NSLOCTEXT("NiagaraOpInfo", "Seeded Integer Rand Name", "Seeded Integer Random");
 		Op->Description = NSLOCTEXT("NiagaraOpInfo", "Seeded Integer Rand Desc", "Returns a deterministic random integer value between 0 and Max-1.");
-		Op->Inputs.Add(FNiagaraOpInOutInfo(Max, Type, MaxText, MaxText, DefaultStr_One));
+		Op->Inputs.Add(FNiagaraOpInOutInfo(Max, NumericType, MaxText, MaxText, DefaultStr_One));
 		Op->Inputs.Add(FNiagaraOpInOutInfo(SeedName1, SeedType, SeedText1, SeedText1, DefaultSeed_Zero));
 		Op->Inputs.Add(FNiagaraOpInOutInfo(SeedName2, SeedType, SeedText2, SeedText2, DefaultSeed_Zero));
 		Op->Inputs.Add(FNiagaraOpInOutInfo(SeedName3, SeedType, SeedText3, SeedText3, DefaultSeed_Zero));
-		Op->Outputs.Add(FNiagaraOpInOutInfo(Result, Type, ResultText, ResultText, DefaultStr_One, TEXT("rand_int({0}, {1}, {2}, {3})")));
+		Op->Outputs.Add(FNiagaraOpInOutInfo(Result, NumericType, ResultText, ResultText, DefaultStr_One, TEXT("rand_int({0}, {1}, {2}, {3})")));
+		Op->InputTypeValidationFunction.BindLambda(NoPositionTypeCheck);
 		Op->BuildName(TEXT("SeededRand Integer"), CategoryName);
 		Op->bNumericsCanBeIntegers = true;
 		Op->bNumericsCanBeFloats = false;
@@ -946,55 +1177,69 @@ void FNiagaraOpInfo::Init()
 		Op->Category = NumericCategory;
 		Op->FriendlyName = NSLOCTEXT("NiagaraOpInfo", "Seeded Float Rand Name", "Seeded Float Random");
 		Op->Description = NSLOCTEXT("NiagaraOpInfo", "Seeded Float Rand Desc", "Returns a deterministic random float value between 0 and Max.");
-		Op->Inputs.Add(FNiagaraOpInOutInfo(Max, Type, MaxText, MaxText, DefaultStr_One));
+		Op->Inputs.Add(FNiagaraOpInOutInfo(Max, NumericType, MaxText, MaxText, DefaultStr_One));
 		Op->Inputs.Add(FNiagaraOpInOutInfo(SeedName1, SeedType, SeedText1, SeedText1, DefaultSeed_Zero));
 		Op->Inputs.Add(FNiagaraOpInOutInfo(SeedName2, SeedType, SeedText2, SeedText2, DefaultSeed_Zero));
 		Op->Inputs.Add(FNiagaraOpInOutInfo(SeedName3, SeedType, SeedText3, SeedText3, DefaultSeed_Zero));
-		Op->Outputs.Add(FNiagaraOpInOutInfo(Result, Type, ResultText, ResultText, DefaultStr_One, TEXT("rand_float({0}, {1}, {2}, {3})")));
+		Op->Outputs.Add(FNiagaraOpInOutInfo(Result, NumericType, ResultText, ResultText, DefaultStr_One, TEXT("rand_float({0}, {1}, {2}, {3})")));
+		Op->InputTypeValidationFunction.BindLambda(NoPositionTypeCheck);
 		Op->BuildName(TEXT("SeededRand Float"), CategoryName);
 		Op->bNumericsCanBeIntegers = false;
 		Op->bNumericsCanBeFloats = true;
 		OpInfoMap.Add(Op->Name) = Idx;
 
- 		//Hash ops
- 		FNiagaraTypeDefinition IntHashType = FNiagaraTypeDefinition::GetIntDef();
- 		Idx = OpInfos.AddDefaulted();
- 		Op = &OpInfos[Idx];
- 		Op->Category = NumericCategory;
- 		Op->FriendlyName = NSLOCTEXT("NiagaraOpInfo", "HashInt Name", "Hash Integer");
- 		Op->Description = NSLOCTEXT("NiagaraOpInfo", "HashInt Desc", "Returns a hashed value for the given inputs. The output is in the range 0 to (2^31 - 1).");
- 		Op->Inputs.Add(FNiagaraOpInOutInfo(A, IntHashType, AText, AText, DefaultStr_Zero));
- 		Op->Inputs.Add(FNiagaraOpInOutInfo(B, IntHashType, BText, BText, DefaultStr_Zero));
- 		Op->Outputs.Add(FNiagaraOpInOutInfo(Result, IntHashType, ResultText, ResultText, DefaultStr_One, TEXT("hash({0}, {1})")));
- 		Op->BuildName(TEXT("Hash Integer"), CategoryName);
- 		Op->bSupportsAddedInputs = true;
- 		Op->AddedInputTypeRestrictions.Add(IntHashType);
- 		Op->AddedInputFormatting = TEXT("hash({A}, {B})");
- 		OpInfoMap.Add(Op->Name) = Idx;
-
- 		Idx = OpInfos.AddDefaulted();
- 		Op = &OpInfos[Idx];
- 		Op->Category = NumericCategory;
- 		Op->FriendlyName = NSLOCTEXT("NiagaraOpInfo", "HashFloat Name", "Hash Float");
- 		Op->Description = NSLOCTEXT("NiagaraOpInfo", "HashFloat Desc", "Returns a hashed value for the given inputs. The result is normalized to the range 0 (inclusive) to 1 (exclusive).");
- 		Op->Inputs.Add(FNiagaraOpInOutInfo(A, IntHashType, AText, AText, DefaultStr_Zero));
- 		Op->Inputs.Add(FNiagaraOpInOutInfo(B, IntHashType, BText, BText, DefaultStr_Zero));
- 		Op->Outputs.Add(FNiagaraOpInOutInfo(Result, FNiagaraTypeDefinition::GetFloatDef(), ResultText, ResultText, DefaultStr_One, TEXT("hash_float({0}, {1})")));
- 		Op->BuildName(TEXT("Hash Float"), CategoryName);
- 		Op->bSupportsAddedInputs = true;
- 		Op->AddedInputTypeRestrictions.Add(IntHashType);
- 		Op->AddedInputFormatting = TEXT("hash_float({A}, {B})");
- 		OpInfoMap.Add(Op->Name) = Idx;
+		//Hash ops
+		FNiagaraTypeDefinition IntHashType = FNiagaraTypeDefinition::GetIntDef();
+		Idx = OpInfos.AddDefaulted();
+		Op = &OpInfos[Idx];
+		Op->Category = NumericCategory;
+		Op->FriendlyName = NSLOCTEXT("NiagaraOpInfo", "HashInt Name", "Hash Integer");
+		Op->Description = NSLOCTEXT("NiagaraOpInfo", "HashInt Desc", "Returns a hashed value for the given inputs. The output is in the range 0 to (2^31 - 1).");
+		Op->Inputs.Add(FNiagaraOpInOutInfo(A, IntHashType, AText, AText, DefaultStr_Zero));
+		Op->Inputs.Add(FNiagaraOpInOutInfo(B, IntHashType, BText, BText, DefaultStr_Zero));
+		Op->Outputs.Add(FNiagaraOpInOutInfo(Result, IntHashType, ResultText, ResultText, DefaultStr_One, TEXT("hash({0}, {1})")));
+		Op->InputTypeValidationFunction.BindLambda(NoPositionTypeCheck);
+		Op->BuildName(TEXT("Hash Integer"), CategoryName);
+		Op->bSupportsAddedInputs = true;
+		Op->AddedInputTypeRestrictions.Add(IntHashType);
+		Op->AddedInputFormatting = TEXT("hash({A}, {B})");
+		OpInfoMap.Add(Op->Name) = Idx;
+
+		Idx = OpInfos.AddDefaulted();
+		Op = &OpInfos[Idx];
+		Op->Category = NumericCategory;
+		Op->FriendlyName = NSLOCTEXT("NiagaraOpInfo", "HashFloat Name", "Hash Float");
+		Op->Description = NSLOCTEXT("NiagaraOpInfo", "HashFloat Desc", "Returns a hashed value for the given inputs. The result is normalized to the range 0 (inclusive) to 1 (exclusive).");
+		Op->Inputs.Add(FNiagaraOpInOutInfo(A, IntHashType, AText, AText, DefaultStr_Zero));
+		Op->Inputs.Add(FNiagaraOpInOutInfo(B, IntHashType, BText, BText, DefaultStr_Zero));
+		Op->Outputs.Add(FNiagaraOpInOutInfo(Result, FNiagaraTypeDefinition::GetFloatDef(), ResultText, ResultText, DefaultStr_One, TEXT("hash_float({0}, {1})")));
+		Op->InputTypeValidationFunction.BindLambda(NoPositionTypeCheck);
+		Op->BuildName(TEXT("Hash Float"), CategoryName);
+		Op->bSupportsAddedInputs = true;
+		Op->AddedInputTypeRestrictions.Add(IntHashType);
+		Op->AddedInputFormatting = TEXT("hash_float({A}, {B})");
+		OpInfoMap.Add(Op->Name) = Idx;
 
 		//Comparison ops
+		auto ComparisonOpTypeValidation = [=](const TArray<FNiagaraTypeDefinition>& InputTypes, FText& ErrorMessage)
+		{
+			int32 PositionTypes = CountType(InputTypes, FNiagaraTypeDefinition::GetPositionDef());
+			if (PositionTypes > 0 && PositionTypes != 2)
+			{
+				ErrorMessage = NSLOCTEXT("NiagaraOpInfo", "Comparison Type Error", "The comparison operation cannot be used for a position and a non-position type.");
+				return false;
+			}
+			return true;
+		};
+		
 		Idx = OpInfos.AddDefaulted();
 		Op = &OpInfos[Idx];
 		Op->Category = NumericCategory;
 		Op->FriendlyName = NSLOCTEXT("NiagaraOpInfo", "CmpLT Name", "Less Than");
 		Op->Description = NSLOCTEXT("NiagaraOpInfo", "CmpLT Desc", "Result = A < B");
 		Op->Keywords = FText::FromString(TEXT("<"));
-		Op->Inputs.Add(FNiagaraOpInOutInfo(A, Type, AText, AText, DefaultStr_Zero));
-		Op->Inputs.Add(FNiagaraOpInOutInfo(B, Type, BText, BText, DefaultStr_One));
+		Op->Inputs.Add(FNiagaraOpInOutInfo(A, NumericType, AText, AText, DefaultStr_Zero));
+		Op->Inputs.Add(FNiagaraOpInOutInfo(B, NumericType, BText, BText, DefaultStr_One));
 		Op->Outputs.Add(FNiagaraOpInOutInfo(Result, FNiagaraTypeDefinition::GetBoolDef(), ResultText, ResultText, DefaultStr_One, TEXT("NiagaraAll({0} < {1})")));
 		Op->BuildName(TEXT("CmpLT"), CategoryName);
 		Op->bSupportsStaticResolution = true;
@@ -1006,6 +1251,7 @@ void FNiagaraOpInfo::Init()
 			int32 Return = InPinValues[0] < InPinValues[1];
 			return Return;
 		});
+		Op->InputTypeValidationFunction.BindLambda(ComparisonOpTypeValidation);
 		OpInfoMap.Add(Op->Name) = Idx;
 
 		
@@ -1016,8 +1262,8 @@ void FNiagaraOpInfo::Init()
 		Op->FriendlyName = NSLOCTEXT("NiagaraOpInfo", "CmpLE Name", "Less Than Or Equal");
 		Op->Description = NSLOCTEXT("NiagaraOpInfo", "CmpLE Desc", "Result = A <= B");
 		Op->Keywords = FText::FromString(TEXT("<="));
-		Op->Inputs.Add(FNiagaraOpInOutInfo(A, Type, AText, AText, DefaultStr_Zero));
-		Op->Inputs.Add(FNiagaraOpInOutInfo(B, Type, BText, BText, DefaultStr_One));
+		Op->Inputs.Add(FNiagaraOpInOutInfo(A, NumericType, AText, AText, DefaultStr_Zero));
+		Op->Inputs.Add(FNiagaraOpInOutInfo(B, NumericType, BText, BText, DefaultStr_One));
 		Op->Outputs.Add(FNiagaraOpInOutInfo(Result, FNiagaraTypeDefinition::GetBoolDef(), ResultText, ResultText, DefaultStr_One, TEXT("NiagaraAll({0} <= {1})")));
 		Op->BuildName(TEXT("CmpLE"), CategoryName);
 		Op->bSupportsStaticResolution = true;
@@ -1029,6 +1275,7 @@ void FNiagaraOpInfo::Init()
 			int32 Return = InPinValues[0] <= InPinValues[1];
 			return Return;
 		});
+		Op->InputTypeValidationFunction.BindLambda(ComparisonOpTypeValidation);
 		OpInfoMap.Add(Op->Name) = Idx;
 
 
@@ -1038,8 +1285,8 @@ void FNiagaraOpInfo::Init()
 		Op->FriendlyName = NSLOCTEXT("NiagaraOpInfo", "CmpGT Name", "Greater Than");
 		Op->Description = NSLOCTEXT("NiagaraOpInfo", "CmpGT Desc", "Result = A > B");
 		Op->Keywords = FText::FromString(TEXT(">"));
-		Op->Inputs.Add(FNiagaraOpInOutInfo(A, Type, AText, AText, DefaultStr_Zero));
-		Op->Inputs.Add(FNiagaraOpInOutInfo(B, Type, BText, BText, DefaultStr_One));
+		Op->Inputs.Add(FNiagaraOpInOutInfo(A, NumericType, AText, AText, DefaultStr_Zero));
+		Op->Inputs.Add(FNiagaraOpInOutInfo(B, NumericType, BText, BText, DefaultStr_One));
 		Op->Outputs.Add(FNiagaraOpInOutInfo(Result, FNiagaraTypeDefinition::GetBoolDef(), ResultText, ResultText, DefaultStr_One, TEXT("NiagaraAll({0} > {1})")));
 		Op->BuildName(TEXT("CmpGT"), CategoryName);
 		Op->bSupportsStaticResolution = true;
@@ -1051,6 +1298,7 @@ void FNiagaraOpInfo::Init()
 			int32 Return = InPinValues[0] > InPinValues[1];
 			return Return;
 		});
+		Op->InputTypeValidationFunction.BindLambda(ComparisonOpTypeValidation);
 		OpInfoMap.Add(Op->Name) = Idx;
 
 		Idx = OpInfos.AddDefaulted();
@@ -1059,8 +1307,8 @@ void FNiagaraOpInfo::Init()
 		Op->FriendlyName = NSLOCTEXT("NiagaraOpInfo", "CmpGE Name", "Greater Than Or Equal");
 		Op->Description = NSLOCTEXT("NiagaraOpInfo", "CmpGE Desc", "Result = A >= B");
 		Op->Keywords = FText::FromString(TEXT(">="));
-		Op->Inputs.Add(FNiagaraOpInOutInfo(A, Type, AText, AText, DefaultStr_Zero));
-		Op->Inputs.Add(FNiagaraOpInOutInfo(B, Type, BText, BText, DefaultStr_One));
+		Op->Inputs.Add(FNiagaraOpInOutInfo(A, NumericType, AText, AText, DefaultStr_Zero));
+		Op->Inputs.Add(FNiagaraOpInOutInfo(B, NumericType, BText, BText, DefaultStr_One));
 		Op->Outputs.Add(FNiagaraOpInOutInfo(Result, FNiagaraTypeDefinition::GetBoolDef(), ResultText, ResultText, DefaultStr_One, TEXT("NiagaraAll({0} >= {1})")));
 		Op->BuildName(TEXT("CmpGE"), CategoryName);
 		Op->bSupportsStaticResolution = true;
@@ -1071,7 +1319,7 @@ void FNiagaraOpInfo::Init()
 
 			int32 Return = InPinValues[0] >= InPinValues[1];
 			return Return;
-		});
+		});Op->InputTypeValidationFunction.BindLambda(ComparisonOpTypeValidation);
 		OpInfoMap.Add(Op->Name) = Idx;
 
 		Idx = OpInfos.AddDefaulted();
@@ -1080,8 +1328,8 @@ void FNiagaraOpInfo::Init()
 		Op->FriendlyName = NSLOCTEXT("NiagaraOpInfo", "CmpEQ Name", "Equal");
 		Op->Description = NSLOCTEXT("NiagaraOpInfo", "CmpEQ Desc", "Result = A == B");
 		Op->Keywords = FText::FromString(TEXT("=="));
-		Op->Inputs.Add(FNiagaraOpInOutInfo(A, Type, AText, AText, DefaultStr_Zero));
-		Op->Inputs.Add(FNiagaraOpInOutInfo(B, Type, BText, BText, DefaultStr_One));
+		Op->Inputs.Add(FNiagaraOpInOutInfo(A, NumericType, AText, AText, DefaultStr_Zero));
+		Op->Inputs.Add(FNiagaraOpInOutInfo(B, NumericType, BText, BText, DefaultStr_One));
 		Op->Outputs.Add(FNiagaraOpInOutInfo(Result, FNiagaraTypeDefinition::GetBoolDef(), ResultText, ResultText, DefaultStr_One, TEXT("NiagaraAll({0} == {1})")));
 		Op->BuildName(TEXT("CmpEQ"), CategoryName);
 		Op->bSupportsStaticResolution = true;
@@ -1093,6 +1341,7 @@ void FNiagaraOpInfo::Init()
 			int32 Return = InPinValues[0] == InPinValues[1];
 			return Return;
 		});
+		Op->InputTypeValidationFunction.BindLambda(ComparisonOpTypeValidation);
 		OpInfoMap.Add(Op->Name) = Idx;
 
 		OpInfoMap.Add(Op->Name) = Idx;
@@ -1103,8 +1352,8 @@ void FNiagaraOpInfo::Init()
 		Op->FriendlyName = NSLOCTEXT("NiagaraOpInfo", "CmpNEQ Name", "Not Equal");
 		Op->Description = NSLOCTEXT("NiagaraOpInfo", "CmpNEQ Desc", "Result = A != B");
 		Op->Keywords = FText::FromString(TEXT("!="));
-		Op->Inputs.Add(FNiagaraOpInOutInfo(A, Type, AText, AText, DefaultStr_Zero));
-		Op->Inputs.Add(FNiagaraOpInOutInfo(B, Type, BText, BText, DefaultStr_One));
+		Op->Inputs.Add(FNiagaraOpInOutInfo(A, NumericType, AText, AText, DefaultStr_Zero));
+		Op->Inputs.Add(FNiagaraOpInOutInfo(B, NumericType, BText, BText, DefaultStr_One));
 		Op->Outputs.Add(FNiagaraOpInOutInfo(Result, FNiagaraTypeDefinition::GetBoolDef(), ResultText, ResultText, DefaultStr_One, TEXT("NiagaraAll({0} != {1})")));
 		Op->BuildName(TEXT("CmpNEQ"), CategoryName);
 		Op->bSupportsStaticResolution = true;
@@ -1116,14 +1365,15 @@ void FNiagaraOpInfo::Init()
 			int32 Return = InPinValues[0] != InPinValues[1];
 			return Return;
 		});
+		Op->InputTypeValidationFunction.BindLambda(ComparisonOpTypeValidation);
 		OpInfoMap.Add(Op->Name) = Idx;
 
 		
 		OpInfoMap.Add(Op->Name) = Idx;
 	}
 
- 	//////////////////////////////////////////////////////////////////////////
- 	//Integer Only Ops
+	//////////////////////////////////////////////////////////////////////////
+	//Integer Only Ops
 
 	FString Default_IntZero(TEXT("0"));
 	FString Default_IntOne(TEXT("1"));
diff --git a/Engine/Plugins/FX/Niagara/Source/NiagaraEditor/Private/NiagaraHlslTranslator.cpp b/Engine/Plugins/FX/Niagara/Source/NiagaraEditor/Private/NiagaraHlslTranslator.cpp
--- a/Engine/Plugins/FX/Niagara/Source/NiagaraEditor/Private/NiagaraHlslTranslator.cpp
+++ b/Engine/Plugins/FX/Niagara/Source/NiagaraEditor/Private/NiagaraHlslTranslator.cpp
@@ -6600,6 +6600,7 @@ void FHlslNiagaraTranslator::Operation(class UNiagaraNodeOp* Operation, TArray<i
 		});
 
 	TArray<UEdGraphPin*> InputPins;
+	TArray<FNiagaraTypeDefinition> InputTypes;
 	Operation->GetInputPins(InputPins);
 
 	bool bAllPinsStatic = true;
@@ -6609,6 +6610,7 @@ void FHlslNiagaraTranslator::Operation(class UNiagaraNodeOp* Operation, TArray<i
 			if (Operation->IsAddPin(InputPins[InputIdx]))
 				continue;
 			FNiagaraTypeDefinition InputType = Schema->PinToTypeDefinition(InputPins[InputIdx]);
+			InputTypes.Add(InputType);
 			if (!InputType.IsStatic())
 				bAllPinsStatic = false;
 		}
@@ -6623,6 +6625,13 @@ void FHlslNiagaraTranslator::Operation(class UNiagaraNodeOp* Operation, TArray<i
 		}
 	}
 
+	FText ValidationError;
+	if (OpInfo && OpInfo->InputTypeValidationFunction.IsBound() && OpInfo->InputTypeValidationFunction.Execute(InputTypes, ValidationError) == false)
+	{
+		Error(ValidationError, Operation, OutputPins[0]);
+		Outputs.Add(INDEX_NONE);
+		return;
+	}
 
 	if (OpInfo && OpInfo->StaticVariableResolveFunction.IsBound() && bAllPinsStatic)
 	{
@@ -8887,7 +8896,7 @@ int32 FHlslNiagaraTranslator::CompilePin(const UEdGraphPin* Pin)
 		{
 			if (TypeDef == FNiagaraTypeDefinition::GetParameterMapDef())
 			{
-				Error(FText::FromString(TEXT("Parameter Maps must be created via an Input Node, not the default value of a pin! Please connect to a valid input Parameter Map.")), Cast<UNiagaraNode>(Pin->GetOwningNode()), nullptr);
+				Error(FText::FromString(TEXT("Parameter Maps must be created via an Input Node, not the default value of a pin! Please connect to a valid input Parameter Map.")), Cast<UNiagaraNode>(Pin->GetOwningNode()), Pin);
 				return INDEX_NONE;
 			}
 			else
diff --git a/Engine/Plugins/FX/Niagara/Source/NiagaraEditor/Private/NiagaraNode.cpp b/Engine/Plugins/FX/Niagara/Source/NiagaraEditor/Private/NiagaraNode.cpp
--- a/Engine/Plugins/FX/Niagara/Source/NiagaraEditor/Private/NiagaraNode.cpp
+++ b/Engine/Plugins/FX/Niagara/Source/NiagaraEditor/Private/NiagaraNode.cpp
@@ -780,9 +780,18 @@ void UNiagaraNode::NumericResolutionByPins(const UEdGraphSchema_Niagara* Schema,
 	for (UEdGraphPin* OutputPin : OutputPins)
 	{
 		FNiagaraTypeDefinition OutputPinType = Schema->PinToTypeDefinition(OutputPin);
-		if (NonNumericInputs.Num() > 0 && GetNumericOutputTypeSelectionMode() != ENiagaraNumericOutputTypeSelectionMode::None)
+		ENiagaraNumericOutputTypeSelectionMode NumericMode = GetNumericOutputTypeSelectionMode();
+		if (NonNumericInputs.Num() > 0 && NumericMode != ENiagaraNumericOutputTypeSelectionMode::None)
 		{
-			FNiagaraTypeDefinition OutputNumericType = FNiagaraTypeDefinition::GetNumericOutputType(NonNumericInputs, GetNumericOutputTypeSelectionMode());
+			FNiagaraTypeDefinition OutputNumericType;
+			if (NumericMode == ENiagaraNumericOutputTypeSelectionMode::Custom)
+			{
+				OutputNumericType = ResolveCustomNumericType(NonNumericInputs);
+			}
+			else
+			{
+				OutputNumericType = FNiagaraTypeDefinition::GetNumericOutputType(NonNumericInputs, NumericMode);
+			}
 			if (OutputNumericType != FNiagaraTypeDefinition::GetGenericNumericDef())
 			{
 				if (OutputPinType == FNiagaraTypeDefinition::GetGenericNumericDef())
@@ -814,15 +823,21 @@ void UNiagaraNode::NumericResolutionByPins(const UEdGraphSchema_Niagara* Schema,
 	}
 }
 
+FNiagaraTypeDefinition UNiagaraNode::ResolveCustomNumericType(const TArray<FNiagaraTypeDefinition>&) const
+{
+	checkf(false, TEXT("Not implemented for node type"));
+	return FNiagaraTypeDefinition::GetFloatDef();
+}
+
 
 void UNiagaraNode::ResolveNumerics(const UEdGraphSchema_Niagara* Schema, bool bSetInline, TMap<TPair<FGuid, UEdGraphNode*>, FNiagaraTypeDefinition>* PinCache)
 {
-// Fix up numeric input pins and keep track of numeric types to decide the output type.
-FPinCollectorArray InputPins;
-GetInputPins(InputPins);
-FPinCollectorArray OutputPins;
-GetOutputPins(OutputPins);
-NumericResolutionByPins(Schema, InputPins, OutputPins, bSetInline, PinCache);
+	// Fix up numeric input pins and keep track of numeric types to decide the output type.
+	FPinCollectorArray InputPins;
+	GetInputPins(InputPins);
+	FPinCollectorArray OutputPins;
+	GetOutputPins(OutputPins);
+	NumericResolutionByPins(Schema, InputPins, OutputPins, bSetInline, PinCache);
 }
 
 void UNiagaraNode::GetWildcardPinHoverConnectionTextAddition(const UEdGraphPin* WildcardPin, const UEdGraphPin* OtherPin, ECanCreateConnectionResponse ConnectionResponse, FString& OutString) const
diff --git a/Engine/Plugins/FX/Niagara/Source/NiagaraEditor/Private/NiagaraNodeOp.cpp b/Engine/Plugins/FX/Niagara/Source/NiagaraEditor/Private/NiagaraNodeOp.cpp
--- a/Engine/Plugins/FX/Niagara/Source/NiagaraEditor/Private/NiagaraNodeOp.cpp
+++ b/Engine/Plugins/FX/Niagara/Source/NiagaraEditor/Private/NiagaraNodeOp.cpp
@@ -131,6 +131,17 @@ void UNiagaraNodeOp::OnPostSynchronizationInReallocatePins()
 	HandleStaticOutputPinUpgrade();		
 }
 
+FNiagaraTypeDefinition UNiagaraNodeOp::ResolveCustomNumericType(const TArray<FNiagaraTypeDefinition>& NonNumericInputs) const
+{
+	const FNiagaraOpInfo* OpInfo = FNiagaraOpInfo::GetOpInfo(OpName);
+	if (OpInfo && OpInfo->CustomNumericResolveFunction.IsBound())
+	{
+		return OpInfo->CustomNumericResolveFunction.Execute(NonNumericInputs);
+	}
+	
+	return FNiagaraTypeDefinition::GetGenericNumericDef();
+}
+
 
 void UNiagaraNodeOp::Compile(class FHlslNiagaraTranslator* Translator, TArray<int32>& Outputs)
 {
diff --git a/Engine/Plugins/FX/Niagara/Source/NiagaraEditor/Public/NiagaraEditorCommon.h b/Engine/Plugins/FX/Niagara/Source/NiagaraEditor/Public/NiagaraEditorCommon.h
--- a/Engine/Plugins/FX/Niagara/Source/NiagaraEditor/Public/NiagaraEditorCommon.h
+++ b/Engine/Plugins/FX/Niagara/Source/NiagaraEditor/Public/NiagaraEditorCommon.h
@@ -73,8 +73,12 @@ public:
 	bool CreateHlslForAddedInputs(int32 InputCount, FString& HlslResult) const;
 
 	DECLARE_DELEGATE_RetVal_OneParam(int32, FStaticVariableResolve, const TArray<int32>& );
+	DECLARE_DELEGATE_RetVal_TwoParams(bool, FInputTypeValidation, const TArray<FNiagaraTypeDefinition>&, FText& );;
+	DECLARE_DELEGATE_RetVal_OneParam(FNiagaraTypeDefinition, FCustomNumericResolve, const TArray<FNiagaraTypeDefinition>& )
 
 	FStaticVariableResolve StaticVariableResolveFunction;
+	FInputTypeValidation InputTypeValidationFunction;
+	FCustomNumericResolve CustomNumericResolveFunction;
 
 	/** Whether or not you can upgrade type to static on connection.*/
 	bool bSupportsStaticResolution;
diff --git a/Engine/Plugins/FX/Niagara/Source/NiagaraEditor/Public/NiagaraNode.h b/Engine/Plugins/FX/Niagara/Source/NiagaraEditor/Public/NiagaraNode.h
--- a/Engine/Plugins/FX/Niagara/Source/NiagaraEditor/Public/NiagaraNode.h
+++ b/Engine/Plugins/FX/Niagara/Source/NiagaraEditor/Public/NiagaraNode.h
@@ -210,6 +210,8 @@ protected:
 
 	void NumericResolutionByPins(const UEdGraphSchema_Niagara* Schema, TArrayView<UEdGraphPin* const> InputPins, TArrayView<UEdGraphPin* const> OutputPins,
 		bool bFixInline, TMap<TPair<FGuid, UEdGraphNode*>, FNiagaraTypeDefinition>* PinCache);
+		
+	virtual FNiagaraTypeDefinition ResolveCustomNumericType(const TArray<FNiagaraTypeDefinition>& NonNumericInputs) const;
 
 	/** Route input parameter map to output parameter map if it exists. Note that before calling this function,
 		the input pins should have been visited already.*/
diff --git a/Engine/Plugins/FX/Niagara/Source/NiagaraEditor/Public/NiagaraNodeOp.h b/Engine/Plugins/FX/Niagara/Source/NiagaraEditor/Public/NiagaraNodeOp.h
--- a/Engine/Plugins/FX/Niagara/Source/NiagaraEditor/Public/NiagaraNodeOp.h
+++ b/Engine/Plugins/FX/Niagara/Source/NiagaraEditor/Public/NiagaraNodeOp.h
@@ -87,6 +87,7 @@ protected:
 	//~ End UNiagaraNodeWithDynamicPins Interface
 
 	virtual void OnPostSynchronizationInReallocatePins() override;
+	virtual FNiagaraTypeDefinition ResolveCustomNumericType(const TArray<FNiagaraTypeDefinition>& NonNumericInputs) const override;
 
 	void HandleStaticInputPinUpgrade(UEdGraphPin* InputPin);
 	void HandleStaticOutputPinUpgrade();
