Fix for ensure loading Niagara_Advanced.
Fixes mismatch in data size/contents between override params and instance/template overrides.
Converting instance/template overrides to doubles when doing conversion to Position type.

diff --git a/Engine/Plugins/FX/Niagara/Source/Niagara/Private/NiagaraComponent.cpp b/Engine/Plugins/FX/Niagara/Source/Niagara/Private/NiagaraComponent.cpp
--- a/Engine/Plugins/FX/Niagara/Source/Niagara/Private/NiagaraComponent.cpp
+++ b/Engine/Plugins/FX/Niagara/Source/Niagara/Private/NiagaraComponent.cpp
@@ -3037,15 +3037,25 @@ void UNiagaraComponent::FixInvalidUserParameterOverrideData()
 			if (SourceVars.Contains(PositionVar))
 			{
 				OverrideParameters.ChangeParameterType(ExistingVar, FNiagaraTypeDefinition::GetPositionDef());
-				FNiagaraVariant VariantData;
-				if (InstanceParameterOverrides.RemoveAndCopyValue(ExistingVar, VariantData))
-				{
-					InstanceParameterOverrides.Add(PositionVar, VariantData);
-				}
-				if (TemplateParameterOverrides.RemoveAndCopyValue(ExistingVar, VariantData))
+
+				auto ConvertOverrideToPosition = [&ExistingVar, &PositionVar](auto & Params)
 				{
-					TemplateParameterOverrides.Add(PositionVar, VariantData);
-				}
+					FNiagaraVariant VariantData;
+					if (Params.RemoveAndCopyValue(ExistingVar, VariantData))
+					{
+						//If needed, convert the old float data to doubles.
+	#if !UE_LARGE_WORLD_COORDINATES_DISABLED
+						check(VariantData.GetNumBytes() == sizeof(FVector3f));
+						FVector3f ExistingData = *((FVector3f*)(VariantData.GetBytes()));						
+						FVector NewData(ExistingData);
+						VariantData.SetBytes((uint8*)&NewData, sizeof(NewData));
+	#endif	
+						Params.Add(PositionVar, VariantData);
+					}
+				};
+
+				ConvertOverrideToPosition(InstanceParameterOverrides);
+				ConvertOverrideToPosition(TemplateParameterOverrides);
 				
 				ExistingVar = PositionVar;
 			}
