Niagara - Fix the Equals function on UNiagaraMergeable so that it ignores transient properties when checking for equality.  This fixes an issue where Renderers were always showing up as different from parent.

diff --git a/Engine/Plugins/FX/Niagara/Source/NiagaraCore/Private/NiagaraMergeable.cpp b/Engine/Plugins/FX/Niagara/Source/NiagaraCore/Private/NiagaraMergeable.cpp
--- a/Engine/Plugins/FX/Niagara/Source/NiagaraCore/Private/NiagaraMergeable.cpp
+++ b/Engine/Plugins/FX/Niagara/Source/NiagaraCore/Private/NiagaraMergeable.cpp
@@ -26,6 +26,11 @@ bool UNiagaraMergeable::Equals(const UNiagaraMergeable* Other)
 
 	for (TFieldIterator<FProperty> PropertyIterator(this->GetClass()); PropertyIterator; ++PropertyIterator)
 	{
+		if (PropertyIterator->HasAnyPropertyFlags(CPF_Transient))
+		{
+			// Transient properties should not be used to determine merge equality.
+			continue;
+		}
 		if (PropertyIterator->Identical(
 			PropertyIterator->ContainerPtrToValuePtr<void>(this),
 			PropertyIterator->ContainerPtrToValuePtr<void>(Other), PPF_DeepComparison) == false)
