Fix LWC bugs with Niagara meshes
Bring back the define to be able to disable Niagara meshes from using GPU Scene for testing

diff --git a/Engine/Plugins/FX/Niagara/Shaders/Private/NiagaraCommon.ush b/Engine/Plugins/FX/Niagara/Shaders/Private/NiagaraCommon.ush
--- a/Engine/Plugins/FX/Niagara/Shaders/Private/NiagaraCommon.ush
+++ b/Engine/Plugins/FX/Niagara/Shaders/Private/NiagaraCommon.ush
@@ -34,6 +34,11 @@ float4x4 NiagaraComposeTransformMatrix(float3 S, float3x3 R, float3 T)
 		T, 1.0f);
 }
 
+FLWCMatrix NiagaraComposeTransformMatrix(float3 S, float3x3 R, FLWCVector3 T)
+{
+	return MakeLWCMatrix(LWCGetTile(T), NiagaraComposeTransformMatrix(S, R, T.Offset));
+}
+
 float4x4 NiagaraComposeInvTransformMatrix(float3 S, float3x3 R, float3 T)
 {
 	float3x3 InvScaleMat = float3x3(
@@ -59,6 +64,11 @@ float4x4 NiagaraComposeInvTransformMatrix(float3 S, float3x3 R, float3 T)
 	return mul(InvTransMat, InvRotScale44);
 }
 
+FLWCInverseMatrix NiagaraComposeInvTransformMatrix(float3 S, float3x3 R, FLWCVector3 T)
+{
+	return MakeLWCInverseMatrix(LWCGetTile(T), NiagaraComposeInvTransformMatrix(S, R, T.Offset));
+}
+
 float3x3 NiagaraFindRotationBetweenNormals(float3 A, float3 B)
 {
 	float W = dot(A, B);
diff --git a/Engine/Plugins/FX/Niagara/Shaders/Private/NiagaraMeshParticleUtils.ush b/Engine/Plugins/FX/Niagara/Shaders/Private/NiagaraMeshParticleUtils.ush
--- a/Engine/Plugins/FX/Niagara/Shaders/Private/NiagaraMeshParticleUtils.ush
+++ b/Engine/Plugins/FX/Niagara/Shaders/Private/NiagaraMeshParticleUtils.ush
@@ -27,15 +27,15 @@ struct NiagaraMeshParticleSRT
 {
 	float3 Scale;
 	float3x3 Rotation;
-	float3 Translation;
+	FLWCVector3 Translation;
 };
 
 // Transform data for a single mesh particle, used for rendering and culling
 struct NiagaraMeshParticleTransforms
 {
-	float4x4 LocalToWorld;
-	float4x4 WorldToLocal;
-	float4x4 PrevLocalToWorld;
+	FLWCMatrix LocalToWorld;
+	FLWCInverseMatrix WorldToLocal;
+	FLWCMatrix PrevLocalToWorld;
 	float3x3 LocalToWorldNoScale;
 };
 
@@ -78,12 +78,12 @@ struct NiagaraMeshParticleSRTParams
 	float4 VelocityDirMag;
 
 	// Camera basis vectors (in world space)
-	float3 CameraOrigin;
+	FLWCVector3 CameraOrigin;
 	float3 CameraForwardDir;
 	float3 CameraUpDir;
 
 	// Primitive transform data
-	float4x4 PrimitiveLocalToWorld;
+	FLWCMatrix PrimitiveLocalToWorld;
 	float3 PrimitiveInvNonUniformScale;
 };
 
@@ -138,37 +138,33 @@ struct NiagaraMeshParticleTransformsParams
 	float4 VelocityDirMag;
 	float4 PrevVelocityDirMag;
 
-	// Camera basis vectors (in world space)
-	float3 CameraOrigin;
+	// Camera position and basis vectors (in world space)
+	FLWCVector3 CameraOrigin;
 	float3 CameraForwardDir;
 	float3 CameraUpDir;
-	float3 PrevCameraOrigin;
+	FLWCVector3 PrevCameraOrigin;
 	float3 PrevCameraForwardDir;
 	float3 PrevCameraUpDir;
 
 	// Primitive transform data
-	float4x4 PrimitiveLocalToWorld;
-	float4x4 PrimitiveWorldToLocal;
-	float4x4 PrimitivePrevLocalToWorld;
+	FLWCMatrix PrimitiveLocalToWorld;
+	FLWCInverseMatrix PrimitiveWorldToLocal;
+	FLWCMatrix PrimitivePrevLocalToWorld;
 	float3 PrimitiveInvNonUniformScale;
 };
 
 // Transforms a local-space position to world space
-float3 NiagaraLocalToWorldPos(float3 LocalSpacePos, NiagaraMeshParticleSRTParams Params)
+FLWCVector3 NiagaraLocalToWorldPos(float3 LocalSpacePos, NiagaraMeshParticleSRTParams Params)
 {
-	return
-		Params.PrimitiveLocalToWorld[0].xyz * LocalSpacePos.x +
-		Params.PrimitiveLocalToWorld[1].xyz * LocalSpacePos.y +
-		Params.PrimitiveLocalToWorld[2].xyz * LocalSpacePos.z +
-		Params.PrimitiveLocalToWorld[3].xyz;
+	return LWCMultiply(LocalSpacePos, Params.PrimitiveLocalToWorld);
 }
 
 // Transforms a simulation-space position to world space
-float3 NiagaraSimToWorldPos(float3 SimSpacePos, NiagaraMeshParticleSRTParams Params)
+FLWCVector3 NiagaraSimToWorldPos(FLWCVector3 SimSpacePos, NiagaraMeshParticleSRTParams Params)
 {
 	if (Params.bLocalSpace)
 	{
-		return NiagaraLocalToWorldPos(SimSpacePos, Params);
+		return NiagaraLocalToWorldPos(LWCToFloat(SimSpacePos), Params);
 	}
 	
 	return SimSpacePos;
@@ -178,9 +174,9 @@ float3 NiagaraSimToWorldPos(float3 SimSpacePos, NiagaraMeshParticleSRTParams Par
 float3 NiagaraLocalToWorldVec(float3 LocalSpaceVec, NiagaraMeshParticleSRTParams Params)
 {
 	return
-		Params.PrimitiveInvNonUniformScale.x * Params.PrimitiveLocalToWorld[0].xyz * LocalSpaceVec.xxx +
-		Params.PrimitiveInvNonUniformScale.y * Params.PrimitiveLocalToWorld[1].xyz * LocalSpaceVec.yyy +
-		Params.PrimitiveInvNonUniformScale.z * Params.PrimitiveLocalToWorld[2].xyz * LocalSpaceVec.zzz;
+		Params.PrimitiveInvNonUniformScale.x * Params.PrimitiveLocalToWorld.M[0].xyz * LocalSpaceVec.xxx +
+		Params.PrimitiveInvNonUniformScale.y * Params.PrimitiveLocalToWorld.M[1].xyz * LocalSpaceVec.yyy +
+		Params.PrimitiveInvNonUniformScale.z * Params.PrimitiveLocalToWorld.M[2].xyz * LocalSpaceVec.zzz;
 }
 
 // Transforms a world-space direction vector to local space
@@ -188,9 +184,9 @@ float3 NiagaraWorldToLocalVec(float3 WorldSpaceVec, NiagaraMeshParticleSRTParams
 {
 	float3x3 InvRot =
 	{
-		Params.PrimitiveInvNonUniformScale.x * Params.PrimitiveLocalToWorld[0].xyz,
-		Params.PrimitiveInvNonUniformScale.y * Params.PrimitiveLocalToWorld[1].xyz,
-		Params.PrimitiveInvNonUniformScale.z * Params.PrimitiveLocalToWorld[2].xyz
+		Params.PrimitiveInvNonUniformScale.x * Params.PrimitiveLocalToWorld.M[0].xyz,
+		Params.PrimitiveInvNonUniformScale.y * Params.PrimitiveLocalToWorld.M[1].xyz,
+		Params.PrimitiveInvNonUniformScale.z * Params.PrimitiveLocalToWorld.M[2].xyz
 	};
 	InvRot = transpose(InvRot);
 	return mul(WorldSpaceVec, InvRot);
@@ -231,7 +227,7 @@ float4 NiagaraGetVelocityDirMag(int VelocityDataOffset, float3 DefaultVelocity,
 }
 
 // Constructs a rotation matrix that satisfies the mesh particle's FacingMode
-float3x3 NiagaraGetMeshFacingMatrix(float3 ParticleSimPosition, NiagaraMeshParticleSRTParams Params)
+float3x3 NiagaraGetMeshFacingMatrix(FLWCVector3 ParticleSimPosition, NiagaraMeshParticleSRTParams Params)
 {
 	float3 WorldX = float3(1, 0, 0);
 	float3 WorldZ = float3(0, 0, 1);
@@ -255,9 +251,9 @@ float3x3 NiagaraGetMeshFacingMatrix(float3 ParticleSimPosition, NiagaraMeshParti
 
 		case MESH_FACING_CAMERA_POSITION:
 		{
-			float3 WorldPosition = NiagaraSimToWorldPos(ParticleSimPosition, Params);
-			float3 CameraVector = Params.CameraOrigin - WorldPosition;
-			FacingDir = NiagaraWorldToSimVec(normalize(CameraVector), Params);
+			FLWCVector3 WorldPosition = NiagaraSimToWorldPos(ParticleSimPosition, Params);
+			float3 CameraDir = LWCNormalize(LWCSubtract(Params.CameraOrigin, WorldPosition));
+			FacingDir = NiagaraWorldToSimVec(CameraDir, Params);
 			break;
 		}
 
@@ -329,11 +325,11 @@ float3x3 NiagaraGetMeshFacingMatrix(float3 ParticleSimPosition, NiagaraMeshParti
 }
 
 // Calculates the simulation-space CameraOffset of the mesh particle 
-float3 NiagaraCalculateCameraOffset(float3 ParticleSimPosition, NiagaraMeshParticleSRTParams Params)
+float3 NiagaraCalculateCameraOffset(FLWCVector3 ParticleSimPosition, NiagaraMeshParticleSRTParams Params)
 {
-	float3 WorldPosition = NiagaraSimToWorldPos(ParticleSimPosition, Params);
+	FLWCVector3 WorldPosition = NiagaraSimToWorldPos(ParticleSimPosition, Params);
 	float Offset = NiagaraSafeGetFloat(Params.CameraOffsetDataOffset, Params.ParticleIndex, Params.DefaultCameraOffset);
-	float3 CameraVector = NiagaraWorldToSimVec(normalize(Params.CameraOrigin - WorldPosition), Params);
+	float3 CameraVector = NiagaraWorldToSimVec(LWCNormalize(LWCSubtract(Params.CameraOrigin, WorldPosition)), Params);
 
 	return CameraVector * Offset;
 }
@@ -342,9 +338,18 @@ float3 NiagaraCalculateCameraOffset(float3 ParticleSimPosition, NiagaraMeshParti
 NiagaraMeshParticleSRT NiagaraCalculateMeshParticleSRT(NiagaraMeshParticleSRTParams Params)
 {
 	float3 ParticleSimPosition = NiagaraSafeGetVec3(Params.PositionDataOffset, Params.ParticleIndex, Params.DefaultPosition);
-	float3 ParticlePosition = Params.bLocalSpace ? ParticleSimPosition : LWCHackToFloat(MakeLWCVector3(Params.SystemLWCTile, ParticleSimPosition));
 	float4 ParticleRotation = NiagaraSafeGetVec4(Params.RotationDataOffset, Params.ParticleIndex, Params.DefaultRotation);
 	float3 ParticleScale = NiagaraSafeGetVec3(Params.ScaleDataOffset, Params.ParticleIndex, Params.DefaultScale);
+	
+	FLWCVector3 ParticlePosition;
+	if (Params.bLocalSpace)
+	{
+		ParticlePosition = MakeLWCVector3((float3)0, ParticleSimPosition);
+	}
+	else
+ 	{
+		ParticlePosition = MakeLWCVector3(Params.SystemLWCTile, ParticleSimPosition);
+	}
 
 	NiagaraMeshParticleSRT SRT;
 
@@ -367,17 +372,17 @@ NiagaraMeshParticleSRT NiagaraCalculateMeshParticleSRT(NiagaraMeshParticleSRTPar
 	SRT.Translation = ParticlePosition;
 
 	// Apply CameraOffset
-	SRT.Translation += NiagaraCalculateCameraOffset(SRT.Translation, Params);
+	SRT.Translation = LWCAdd(SRT.Translation, NiagaraCalculateCameraOffset(SRT.Translation, Params));
 	
 	// Apply MeshOffset
 	if (Params.bMeshOffsetIsWorldSpace)
 	{
-		SRT.Translation += NiagaraWorldToSimVec(Params.MeshOffset, Params);
+		SRT.Translation = LWCAdd(SRT.Translation, NiagaraWorldToSimVec(Params.MeshOffset, Params));
 	}
 	else
 	{
 		// NOTE: MeshOffset here is mesh-local, not Primitive-local
-		SRT.Translation += mul(Params.MeshOffset * ParticleScale, SRT.Rotation);
+		SRT.Translation = LWCAdd(SRT.Translation, mul(Params.MeshOffset * ParticleScale, SRT.Rotation));
 	}
 
 	return SRT;
@@ -442,7 +447,7 @@ NiagaraMeshParticleTransforms NiagaraCalculateMeshParticleTransforms(NiagaraMesh
 		PrevParticleSRT = ParticleSRT;
 		
 		const float3 Velocity = Params.VelocityDirMag.xyz * Params.VelocityDirMag.w;		
-		PrevParticleSRT.Translation -= Velocity * Params.DeltaSeconds;
+		PrevParticleSRT.Translation = LWCSubtract(PrevParticleSRT.Translation, Velocity * Params.DeltaSeconds);
 	}
 
 	NiagaraMeshParticleTransforms Output;
@@ -454,11 +459,11 @@ NiagaraMeshParticleTransforms NiagaraCalculateMeshParticleTransforms(NiagaraMesh
 	if (Params.bLocalSpace)
 	{
 		// Transform from primitive to world space
-		Output.LocalToWorld = mul(Output.LocalToWorld, Params.PrimitiveLocalToWorld);
-		Output.WorldToLocal = mul(Params.PrimitiveWorldToLocal, Output.WorldToLocal);
-		Output.PrevLocalToWorld = mul(Output.PrevLocalToWorld, Params.PrimitivePrevLocalToWorld);
+		Output.LocalToWorld = LWCMultiply(LWCToFloat(Output.LocalToWorld), Params.PrimitiveLocalToWorld);
+		Output.WorldToLocal = LWCMultiply(Params.PrimitiveWorldToLocal, LWCToFloat(Output.WorldToLocal));
+		Output.PrevLocalToWorld = LWCMultiply(LWCToFloat(Output.PrevLocalToWorld), Params.PrimitivePrevLocalToWorld);
 
-		float3x3 PrimLocalToWorldNoScale = (float3x3)Params.PrimitiveLocalToWorld;
+		float3x3 PrimLocalToWorldNoScale = LWCToFloat3x3(Params.PrimitiveLocalToWorld);
 		PrimLocalToWorldNoScale[0] *= Params.PrimitiveInvNonUniformScale.x;
 		PrimLocalToWorldNoScale[1] *= Params.PrimitiveInvNonUniformScale.y;
 		PrimLocalToWorldNoScale[2] *= Params.PrimitiveInvNonUniformScale.z;
diff --git a/Engine/Plugins/FX/Niagara/Shaders/Private/NiagaraMeshVertexFactory.ush b/Engine/Plugins/FX/Niagara/Shaders/Private/NiagaraMeshVertexFactory.ush
--- a/Engine/Plugins/FX/Niagara/Shaders/Private/NiagaraMeshVertexFactory.ush
+++ b/Engine/Plugins/FX/Niagara/Shaders/Private/NiagaraMeshVertexFactory.ush
@@ -542,22 +542,22 @@ FVertexFactoryIntermediates GetVertexFactoryIntermediates(FVertexFactoryInput In
 	Params.DefaultPrevCameraOffset 		= NiagaraMeshVF.DefaultPrevCameraOffset;
 	Params.VelocityDirMag				= Intermediates.ParticleVelocity;
 	Params.PrevVelocityDirMag			= NiagaraGetVelocityDirMag(NiagaraMeshVF.PrevVelocityDataOffset, NiagaraMeshVF.DefaultPrevVelocity, Intermediates.ParticleIndex);
-	Params.CameraOrigin 				= LWCHackToFloat(ResolvedView.WorldCameraOrigin);
+	Params.CameraOrigin 				= ResolvedView.WorldCameraOrigin;
 	Params.CameraForwardDir				= ResolvedView.ViewForward;
 	Params.CameraUpDir					= ResolvedView.ViewUp;
-	Params.PrevCameraOrigin 			= LWCHackToFloat(ResolvedView.PrevWorldCameraOrigin);
+	Params.PrevCameraOrigin 			= ResolvedView.PrevWorldCameraOrigin;
 	Params.PrevCameraForwardDir			= ResolvedView.PrevViewToTranslatedWorld[2].xyz;
 	Params.PrevCameraUpDir				= ResolvedView.PrevViewToTranslatedWorld[1].xyz;
-	Params.PrimitiveLocalToWorld		= LWCHackToFloat(Intermediates.SceneData.Primitive.LocalToWorld);
-	Params.PrimitiveWorldToLocal 		= LWCHackToFloat(Intermediates.SceneData.Primitive.WorldToLocal);
-	Params.PrimitivePrevLocalToWorld 	= LWCHackToFloat(Intermediates.SceneData.Primitive.PreviousLocalToWorld);
+	Params.PrimitiveLocalToWorld		= Intermediates.SceneData.Primitive.LocalToWorld;
+	Params.PrimitiveWorldToLocal 		= Intermediates.SceneData.Primitive.WorldToLocal;
+	Params.PrimitivePrevLocalToWorld 	= Intermediates.SceneData.Primitive.PreviousLocalToWorld;
 	Params.PrimitiveInvNonUniformScale	= Intermediates.SceneData.Primitive.InvNonUniformScale;
 
 	NiagaraMeshParticleTransforms Transforms = NiagaraCalculateMeshParticleTransforms(Params);
 
-	Intermediates.ParticleToWorld 			= Transforms.LocalToWorld;
-	Intermediates.WorldToParticle 			= Transforms.WorldToLocal;
-	Intermediates.PrevParticleToWorld		= Transforms.PrevLocalToWorld;
+	Intermediates.ParticleToWorld 			= LWCHackToFloat(Transforms.LocalToWorld);
+	Intermediates.WorldToParticle 			= LWCHackToFloat(Transforms.WorldToLocal);
+	Intermediates.PrevParticleToWorld		= LWCHackToFloat(Transforms.PrevLocalToWorld);
 	Intermediates.ParticleToWorldNoScale	= Transforms.LocalToWorldNoScale;
 #endif // (VF_USE_PRIMITIVE_SCENE_DATA == 1)
 
@@ -921,4 +921,3 @@ float4 VertexFactoryGetTranslatedPrimitiveVolumeBounds(FVertexFactoryInterpolant
 #endif // #if COMPUTESHADER || RAYHITGROUPSHADER
 
 #include "/Engine/Private/VertexFactoryDefaultInterface.ush"
-
diff --git a/Engine/Plugins/FX/Niagara/Shaders/Private/NiagaraUpdateMeshGPUSceneInstances.usf b/Engine/Plugins/FX/Niagara/Shaders/Private/NiagaraUpdateMeshGPUSceneInstances.usf
--- a/Engine/Plugins/FX/Niagara/Shaders/Private/NiagaraUpdateMeshGPUSceneInstances.usf
+++ b/Engine/Plugins/FX/Niagara/Shaders/Private/NiagaraUpdateMeshGPUSceneInstances.usf
@@ -154,23 +154,23 @@ void UpdateMeshInstancesCS(uint3 GroupId : SV_GroupID, uint GroupIndex : SV_Grou
 	Params.DefaultPrevCameraOffset 		= DefaultPrevCameraOffset;
 	Params.VelocityDirMag				= NiagaraGetVelocityDirMag(VelocityDataOffset, DefaultVelocity, ParticleIndex);
 	Params.PrevVelocityDirMag			= NiagaraGetVelocityDirMag(PrevVelocityDataOffset, DefaultPrevVelocity, ParticleIndex);
-	Params.CameraOrigin 				= LWCHackToFloat(PrimaryView.WorldCameraOrigin);
+	Params.CameraOrigin 				= PrimaryView.WorldCameraOrigin;
 	Params.CameraForwardDir				= View.ViewForward;
 	Params.CameraUpDir					= View.ViewUp;
-	Params.PrevCameraOrigin 			= LWCHackToFloat(PrimaryView.PrevWorldCameraOrigin);
+	Params.PrevCameraOrigin 			= PrimaryView.PrevWorldCameraOrigin;
 	Params.PrevCameraForwardDir			= View.PrevViewToTranslatedWorld[2].xyz;
 	Params.PrevCameraUpDir				= View.PrevViewToTranslatedWorld[1].xyz;
-	Params.PrimitiveLocalToWorld		= LWCHackToFloat(PrimitiveData.LocalToWorld);
-	Params.PrimitiveWorldToLocal 		= LWCHackToFloat(PrimitiveData.WorldToLocal);
-	Params.PrimitivePrevLocalToWorld 	= LWCHackToFloat(PrimitiveData.PreviousLocalToWorld);
+	Params.PrimitiveLocalToWorld		= PrimitiveData.LocalToWorld;
+	Params.PrimitiveWorldToLocal 		= PrimitiveData.WorldToLocal;
+	Params.PrimitivePrevLocalToWorld 	= PrimitiveData.PreviousLocalToWorld;
 	Params.PrimitiveInvNonUniformScale	= PrimitiveData.InvNonUniformScale;
 
 	NiagaraMeshParticleTransforms Transforms = NiagaraCalculateMeshParticleTransforms(Params);
 
 	// Transform the bounding center to world-space and distance cull
 	{
-		float3 BCenter = mul(float4(LocalBoundingCenter, 1.0f), Transforms.LocalToWorld).xyz;
-		float3 CamVec = BCenter - LWCHackToFloat(PrimaryView.WorldCameraOrigin);
+		FLWCVector3 BCenter = LWCMultiply(LocalBoundingCenter, Transforms.LocalToWorld);
+		float3 CamVec = LWCToFloat(LWCSubtract(BCenter, PrimaryView.WorldCameraOrigin));
 		float CamDistSq = dot(CamVec, CamVec);
 		if (CamDistSq < DistanceCullRangeSquared.x  || CamDistSq > DistanceCullRangeSquared.y)
 		{
@@ -180,6 +180,16 @@ void UpdateMeshInstancesCS(uint3 GroupId : SV_GroupID, uint GroupIndex : SV_Grou
 		}
 	}
 
+	// The transforms solved are currently relative to the system LWC tile. We need to re-base it to the 
+	// Primitive's current tile
+	FLWCMatrix LocalToWorld = Transforms.LocalToWorld;
+	FLWCMatrix PrevLocalToWorld = Transforms.PrevLocalToWorld;
+	FLWCVector3 NewTileOrigin = MakeLWCVector3(PrimitiveData.TilePosition, (float3)0);
+	FLWCVector3 PrimitiveOffset = LWCSubtract(LWCGetOrigin(Transforms.LocalToWorld), NewTileOrigin);
+	FLWCVector3 PrevPrimitiveOffset = LWCSubtract(LWCGetOrigin(Transforms.PrevLocalToWorld), NewTileOrigin);
+	LWCSetOrigin(LocalToWorld, MakeLWCVector3(PrimitiveData.TilePosition, LWCToFloat(PrimitiveOffset)));
+	LWCSetOrigin(PrevLocalToWorld, MakeLWCVector3(PrimitiveData.TilePosition, LWCToFloat(PrevPrimitiveOffset)));
+
 	const uint PayloadDataFlags = GetPayloadDataFlags();
 	InitializeInstanceSceneData(
 		InstanceId,
@@ -188,10 +198,10 @@ void UpdateMeshInstancesCS(uint3 GroupId : SV_GroupID, uint GroupIndex : SV_Grou
  		PayloadDataFlags,
  		1,
  		0.0f,
-		Transforms.LocalToWorld
+		LocalToWorld.M
 	);
 	
 	FInstancePayloadDataOffsets PayloadOffsets = GetInstancePayloadDataOffsets(PrimitiveId, PayloadDataFlags, LocalInstanceId);
-	WriteInstanceDynamicData(PayloadOffsets, Transforms.PrevLocalToWorld); 	
+	WriteInstanceDynamicData(PayloadOffsets, PrevLocalToWorld.M);
 	WriteInstanceCustomData(PayloadOffsets, 0, float4(asfloat(ParticleIndex), 0.0f, 0.0f, 0.0f));
 }
\ No newline at end of file
diff --git a/Engine/Plugins/FX/Niagara/Source/Niagara/Private/NiagaraRendererMeshes.cpp b/Engine/Plugins/FX/Niagara/Source/Niagara/Private/NiagaraRendererMeshes.cpp
--- a/Engine/Plugins/FX/Niagara/Source/Niagara/Private/NiagaraRendererMeshes.cpp
+++ b/Engine/Plugins/FX/Niagara/Source/Niagara/Private/NiagaraRendererMeshes.cpp
@@ -248,9 +248,13 @@ int32 FNiagaraRendererMeshes::GetLODIndex(int32 MeshIndex) const
 void FNiagaraRendererMeshes::PrepareParticleMeshRenderData(FParticleMeshRenderData& ParticleMeshRenderData, const FSceneViewFamily& ViewFamily, FMeshElementCollector& Collector, FNiagaraDynamicDataBase* InDynamicData, const FNiagaraSceneProxy* SceneProxy) const
 {
 	ParticleMeshRenderData.Collector = &Collector;
-	
+
+#if NIAGARA_ENABLE_GPU_SCENE_MESHES
 	ParticleMeshRenderData.bUseGPUScene = FeatureLevel > ERHIFeatureLevel::ES3_1	
 		&& UseGPUScene(SceneProxy->GetScene().GetShaderPlatform(), FeatureLevel);
+#else
+	ParticleMeshRenderData.bUseGPUScene = false;
+#endif
 
 	// Anything to render?
 	ParticleMeshRenderData.DynamicDataMesh = static_cast<FNiagaraDynamicDataMesh*>(InDynamicData);
diff --git a/Engine/Plugins/FX/Niagara/Source/NiagaraVertexFactories/Private/NiagaraMeshVertexFactory.cpp b/Engine/Plugins/FX/Niagara/Source/NiagaraVertexFactories/Private/NiagaraMeshVertexFactory.cpp
--- a/Engine/Plugins/FX/Niagara/Source/NiagaraVertexFactories/Private/NiagaraMeshVertexFactory.cpp
+++ b/Engine/Plugins/FX/Niagara/Source/NiagaraVertexFactories/Private/NiagaraMeshVertexFactory.cpp
@@ -125,8 +125,10 @@ void FNiagaraMeshVertexFactory::InitRHI()
 			}
 		}
 
+#if NIAGARA_ENABLE_GPU_SCENE_MESHES
 		// TODO: Support GPU Scene on mobile? Maybe only for CPU particles?
 		AddPrimitiveIdStreamElement(EVertexInputStreamType::Default, Elements, 13, 0xFF);
+#endif
 
 		//if (Streams.Num() > 0)
 		{
@@ -152,6 +154,7 @@ void FNiagaraMeshVertexFactory::ModifyCompilationEnvironment(const FVertexFactor
 	OutEnvironment.SetDefine(TEXT("NIAGARA_MESH_INSTANCED"), TEXT("1"));
 	OutEnvironment.SetDefine(TEXT("NiagaraVFLooseParameters"), TEXT("NiagaraMeshVF"));
 
+#if NIAGARA_ENABLE_GPU_SCENE_MESHES
 	const ERHIFeatureLevel::Type MaxSupportedFeatureLevel = GetMaxSupportedFeatureLevel(Parameters.Platform);
 
 	// TODO: Support GPU Scene on mobile?
@@ -160,6 +163,7 @@ void FNiagaraMeshVertexFactory::ModifyCompilationEnvironment(const FVertexFactor
 	
 	OutEnvironment.SetDefine(TEXT("VF_SUPPORTS_PRIMITIVE_SCENE_DATA"), bSupportsPrimitiveIdStream && bUseGPUScene);
 	OutEnvironment.SetDefine(TEXT("VF_REQUIRES_PER_INSTANCE_CUSTOM_DATA"), bSupportsPrimitiveIdStream && bUseGPUScene);
+#endif
 
 	const bool ContainsManualVertexFetch = OutEnvironment.GetDefinitions().Contains("MANUAL_VERTEX_FETCH");
 	if (!ContainsManualVertexFetch && RHISupportsManualVertexFetch(Parameters.Platform))
@@ -175,10 +179,16 @@ void FNiagaraMeshVertexFactory::SetData(const FStaticMeshDataType& InData)
 	UpdateRHI();
 }
 
-#define NIAGARA_MESH_VF_FLAGS (EVertexFactoryFlags::UsedWithMaterials \
-	| EVertexFactoryFlags::SupportsDynamicLighting \
-	| EVertexFactoryFlags::SupportsRayTracing \
-	| EVertexFactoryFlags::SupportsPrimitiveIdStream)
+#if NIAGARA_ENABLE_GPU_SCENE_MESHES
+	#define NIAGARA_MESH_VF_FLAGS (EVertexFactoryFlags::UsedWithMaterials \
+		| EVertexFactoryFlags::SupportsDynamicLighting \
+		| EVertexFactoryFlags::SupportsRayTracing \
+		| EVertexFactoryFlags::SupportsPrimitiveIdStream)
+#else
+	#define NIAGARA_MESH_VF_FLAGS (EVertexFactoryFlags::UsedWithMaterials \
+		| EVertexFactoryFlags::SupportsDynamicLighting \
+		| EVertexFactoryFlags::SupportsRayTracing)
+#endif
 #define NIAGARA_MESH_VF_FLAGS_EX (NIAGARA_MESH_VF_FLAGS | EVertexFactoryFlags::SupportsPrecisePrevWorldPos)
 
 IMPLEMENT_VERTEX_FACTORY_PARAMETER_TYPE(FNiagaraMeshVertexFactory, SF_Vertex, FNiagaraMeshVertexFactoryShaderParametersVS);
diff --git a/Engine/Plugins/FX/Niagara/Source/NiagaraVertexFactories/Public/NiagaraMeshVertexFactory.h b/Engine/Plugins/FX/Niagara/Source/NiagaraVertexFactories/Public/NiagaraMeshVertexFactory.h
--- a/Engine/Plugins/FX/Niagara/Source/NiagaraVertexFactories/Public/NiagaraMeshVertexFactory.h
+++ b/Engine/Plugins/FX/Niagara/Source/NiagaraVertexFactories/Public/NiagaraMeshVertexFactory.h
@@ -16,6 +16,10 @@ ParticleVertexFactory.h: Particle vertex factory definitions.
 #include "SceneManagement.h"
 #include "VertexFactory.h"
 
+// Disable this define to test disabling the use of GPU Scene with Niagara mesh renderer
+// NOTE: Changing this will also require you to make a trivial change to the mesh factory shader, or it may use cached shaders
+#define NIAGARA_ENABLE_GPU_SCENE_MESHES 1
+
 class FMaterial;
 class FVertexBuffer;
 struct FDynamicReadBuffer;
diff --git a/Engine/Shaders/Private/LargeWorldCoordinates.ush b/Engine/Shaders/Private/LargeWorldCoordinates.ush
--- a/Engine/Shaders/Private/LargeWorldCoordinates.ush
+++ b/Engine/Shaders/Private/LargeWorldCoordinates.ush
@@ -348,6 +348,12 @@ FLWCVector3 LWCGetOrigin(FLWCMatrix InMatrix)
 	return MakeLWCVector3(LWCGetTile(InMatrix), InMatrix.M[3].xyz);
 }
 
+void LWCSetOrigin(inout FLWCMatrix InOutMatrix, FLWCVector3 Origin)
+{
+	LWCSetTile(InOutMatrix, LWCGetTile(Origin));
+	InOutMatrix.M[3].xyz = Origin.Offset;
+}
+
 #define FLWCType FLWCScalar
 #define FFloatType float
 #define FBoolType bool
