Fixup default pin types in graph postload

diff --git a/Engine/Plugins/FX/Niagara/Source/NiagaraEditor/Private/NiagaraGraph.cpp b/Engine/Plugins/FX/Niagara/Source/NiagaraEditor/Private/NiagaraGraph.cpp
--- a/Engine/Plugins/FX/Niagara/Source/NiagaraEditor/Private/NiagaraGraph.cpp
+++ b/Engine/Plugins/FX/Niagara/Source/NiagaraEditor/Private/NiagaraGraph.cpp
@@ -360,6 +360,17 @@ void UNiagaraGraph::PostLoad()
 					break;
 				}
 			}
+
+			// fix default pin types if necessary
+			if (UNiagaraNodeParameterMapGet* GetNode = Cast<UNiagaraNodeParameterMapGet>(NiagaraNode))
+			{
+				UEdGraphPin* DefaultPin = GetNode->GetDefaultPin(Pin);
+				if (Pin->Direction == EGPD_Output && DefaultPin && DefaultPin->PinType != Pin->PinType)
+				{
+					DefaultPin->PinType = Pin->PinType;
+					NiagaraNode->MarkNodeRequiresSynchronization(__FUNCTION__, true);
+				}
+			}
 		}
 	}
 	for (FNiagaraVariable OldVarType : OldTypes)
