Fix for uninitialied variables for niagara's HWRT DI for projects that don't support it.  Also fixes up some ahader warnings from all niagara GPU shaders.

diff --git a/Engine/Plugins/FX/Niagara/Shaders/Private/NiagaraDataInterfaceCollisionQuery.ush b/Engine/Plugins/FX/Niagara/Shaders/Private/NiagaraDataInterfaceCollisionQuery.ush
--- a/Engine/Plugins/FX/Niagara/Shaders/Private/NiagaraDataInterfaceCollisionQuery.ush
+++ b/Engine/Plugins/FX/Niagara/Shaders/Private/NiagaraDataInterfaceCollisionQuery.ush
@@ -182,6 +182,7 @@ void NDICollisionQuery_IssueAsyncRayTrace(uint MaxRayTraceCount, RWStructuredBuf
 
 void NDICollisionQuery_ReserveRayTraceIndex(uint MaxRayTraceCount, RWBuffer<uint> RWRayTraceCounts, uint RayTraceCountsOffset, int In_TraceCount, out int Out_StartQueryID, out bool Out_IndicesValid)
 {
+	Out_StartQueryID = -1;
 	Out_IndicesValid = false;
 
 	if ((In_TraceCount + RWRayTraceCounts[RayTraceCountsOffset + 0]) < MaxRayTraceCount)
diff --git a/Engine/Plugins/FX/Niagara/Shaders/Private/NiagaraDataInterfaceCollisionQueryTemplate.ush b/Engine/Plugins/FX/Niagara/Shaders/Private/NiagaraDataInterfaceCollisionQueryTemplate.ush
--- a/Engine/Plugins/FX/Niagara/Shaders/Private/NiagaraDataInterfaceCollisionQueryTemplate.ush
+++ b/Engine/Plugins/FX/Niagara/Shaders/Private/NiagaraDataInterfaceCollisionQueryTemplate.ush
@@ -49,6 +49,7 @@ void ReserveRayTraceIndex_{ParameterName}_UEImpureCall(int In_TraceCount, out in
 		return;
 	}
 #endif
+	Out_StartQueryID = -1;
 	Out_IndicesValid = false;
 }
 
@@ -61,6 +62,7 @@ void CreateAsyncRayTraceGpu_{ParameterName}_UEImpureCall(float3 In_TraceStart, f
 		return;
 	}
 #endif
+	Out_QueryID = -1;
 	Out_IsQueryValid = false;
 }
 
@@ -73,5 +75,8 @@ void ReadAsyncRayTraceGpu_{ParameterName}(int In_PreviousFrameQueryID, out bool
 		return;
 	}
 #endif
+	Out_CollisionDistance = 0.0f;
+	Out_CollisionPosWorld = float3(0.0f, 0.0f, 0.0f);
+	Out_CollisionNormal = float3(0.0f, 0.0f, 0.0f);
 	Out_CollisionValid = false;
 }
\ No newline at end of file
diff --git a/Engine/Plugins/FX/Niagara/Source/NiagaraCore/Private/NiagaraCustomVersion.cpp b/Engine/Plugins/FX/Niagara/Source/NiagaraCore/Private/NiagaraCustomVersion.cpp
--- a/Engine/Plugins/FX/Niagara/Source/NiagaraCore/Private/NiagaraCustomVersion.cpp
+++ b/Engine/Plugins/FX/Niagara/Source/NiagaraCore/Private/NiagaraCustomVersion.cpp
@@ -8,4 +8,4 @@ const FGuid FNiagaraCustomVersion::GUID(0xFCF57AFA, 0x50764283, 0xB9A9E658, 0xFF
 // Register the custom version with core
 FCustomVersionRegistration GRegisterNiagaraCustomVersion(FNiagaraCustomVersion::GUID, FNiagaraCustomVersion::LatestVersion, TEXT("NiagaraVer"));
 
-const FGuid FNiagaraCustomVersion::LatestScriptCompileVersion(0x14D5F635, 0x89F04AC8, 0xBE4C1F5F, 0x91E1134D);
\ No newline at end of file
+const FGuid FNiagaraCustomVersion::LatestScriptCompileVersion(0x4197BB40, 0xB7B14CE3, 0x83D1C54C, 0x17C36A92);
\ No newline at end of file
diff --git a/Engine/Plugins/FX/Niagara/Source/NiagaraEditor/Private/NiagaraHlslTranslator.cpp b/Engine/Plugins/FX/Niagara/Source/NiagaraEditor/Private/NiagaraHlslTranslator.cpp
--- a/Engine/Plugins/FX/Niagara/Source/NiagaraEditor/Private/NiagaraHlslTranslator.cpp
+++ b/Engine/Plugins/FX/Niagara/Source/NiagaraEditor/Private/NiagaraHlslTranslator.cpp
@@ -3239,7 +3239,7 @@ void FHlslNiagaraTranslator::DefineMainGPUFunctions(
 					"	{\n"
 					"		GSpawnStartInstance = RWInstanceCounts[ReadInstanceCountOffset];\n"
 					"	}\n"
-					"	const int MaxInstances = GSpawnStartInstance + SpawnedInstances;\n"
+					"	const uint MaxInstances = GSpawnStartInstance + SpawnedInstances;\n"
 					"	const bool bRunUpdateLogic = bRunSpawnUpdateLogic && GLinearThreadId < GSpawnStartInstance && GLinearThreadId < MaxInstances;\n"
 					"	const bool bRunSpawnLogic = bRunSpawnUpdateLogic && GLinearThreadId >= GSpawnStartInstance && GLinearThreadId < MaxInstances;\n"
 				);
