- Ui changes to Niagara asset wizzard (removed borders + changed template tab visuals)
- Refactored the template box to be cleaner to code for/with
- Added better tooltips for tabs in a system context
- Fixes to the asset thumbnails. The thumbnail pool doesn't tick in the modal window, which is why thumbnails no longer get rendered. A few thumbnails still don't render correctly for some reason, and it's always the same assets, until the asset is actually loaded. Seems to be an issue in core as it's the same in the content browser

diff --git a/Engine/Plugins/FX/Niagara/Source/NiagaraEditor/Private/SNiagaraAssetPickerList.cpp b/Engine/Plugins/FX/Niagara/Source/NiagaraEditor/Private/SNiagaraAssetPickerList.cpp
--- a/Engine/Plugins/FX/Niagara/Source/NiagaraEditor/Private/SNiagaraAssetPickerList.cpp
+++ b/Engine/Plugins/FX/Niagara/Source/NiagaraEditor/Private/SNiagaraAssetPickerList.cpp
@@ -25,7 +25,7 @@ FText SNiagaraAssetPickerList::UncategorizedCategory = LOCTEXT("Uncategorized",
 
 void SNiagaraAssetPickerList::Construct(const FArguments& InArgs, UClass* AssetClass)
 {
-	AssetThumbnailPool = MakeShareable(new FAssetThumbnailPool(24));
+	AssetThumbnailPool = MakeShareable(new FAssetThumbnailPool(48));
 
 	OnTemplateAssetActivated = InArgs._OnTemplateAssetActivated;
 	ViewOptions = InArgs._ViewOptions;
@@ -43,7 +43,8 @@ void SNiagaraAssetPickerList::Construct(const FArguments& InArgs, UClass* AssetC
 	.bLibraryOnly(this, &SNiagaraAssetPickerList::GetLibraryCheckBoxState)
 	.OnLibraryOnlyChanged(this, &SNiagaraAssetPickerList::LibraryCheckBoxStateChanged)
 	.OnSourceFiltersChanged(this, &SNiagaraAssetPickerList::TriggerRefresh)
-	.OnTabActivated(this, &SNiagaraAssetPickerList::TriggerRefreshFromTabs);
+	.OnTabActivated(this, &SNiagaraAssetPickerList::TriggerRefreshFromTabs)
+	.Class(AssetClass);
 	
 	if(!ViewOptions.GetAddLibraryOnlyCheckbox())
 	{
@@ -84,6 +85,7 @@ void SNiagaraAssetPickerList::Construct(const FArguments& InArgs, UClass* AssetC
 
 void SNiagaraAssetPickerList::Tick(const FGeometry& AllottedGeometry, const double InCurrentTime, const float InDeltaTime)
 {
+	AssetThumbnailPool->Tick(InDeltaTime);
 }
 
 TArray<FAssetData> SNiagaraAssetPickerList::GetSelectedAssets() const
@@ -330,7 +332,8 @@ const int32 ThumbnailSize = 72;
 
 TSharedRef<SWidget> SNiagaraAssetPickerList::OnGenerateWidgetForItem(const FAssetData& Item)
 {
-	TSharedRef<FAssetThumbnail> AssetThumbnail = MakeShared<FAssetThumbnail>(Item, ThumbnailSize, ThumbnailSize, UThumbnailManager::Get().GetSharedThumbnailPool());
+	TSharedRef<FAssetThumbnail> AssetThumbnail = MakeShared<FAssetThumbnail>(Item, ThumbnailSize, ThumbnailSize, AssetThumbnailPool);
+	
 	FAssetThumbnailConfig ThumbnailConfig;
 	ThumbnailConfig.bAllowFadeIn = false;
 	
diff --git a/Engine/Plugins/FX/Niagara/Source/NiagaraEditor/Private/Widgets/SNiagaraFilterBox.cpp b/Engine/Plugins/FX/Niagara/Source/NiagaraEditor/Private/Widgets/SNiagaraFilterBox.cpp
--- a/Engine/Plugins/FX/Niagara/Source/NiagaraEditor/Private/Widgets/SNiagaraFilterBox.cpp
+++ b/Engine/Plugins/FX/Niagara/Source/NiagaraEditor/Private/Widgets/SNiagaraFilterBox.cpp
@@ -20,7 +20,7 @@ TMap<EScriptSource, bool> SNiagaraSourceFilterBox::SourceState
 	{EScriptSource::Developer, false},
 };
 
-ENiagaraScriptTemplateSpecification SNiagaraTemplateTabBox::CachedActiveTab = ENiagaraScriptTemplateSpecification::Template;
+TOptional<ENiagaraScriptTemplateSpecification> SNiagaraTemplateTabBox::CachedActiveTab;
 
 void SNiagaraSourceFilterCheckBox::Construct(const FArguments& Args, EScriptSource InSource)
 {
@@ -265,135 +265,63 @@ void SNiagaraTemplateTabBox::Construct(const FArguments& InArgs, FNiagaraTemplat
 	TabOptions = InTabOptions;
 
 	OnTabActivatedDelegate = InArgs._OnTabActivated;
+	Class = InArgs._Class;
 	
-	SAssignNew(TabContainer, SHorizontalBox);
-
-	bool bActiveTabInitialized = false;
+	TabContainer = SNew(SSegmentedControl<ENiagaraScriptTemplateSpecification>)
+	.Value(TAttribute<ENiagaraScriptTemplateSpecification>::CreateSP(this, &SNiagaraTemplateTabBox::GetActiveTab))
+	.OnValueChanged(this, &SNiagaraTemplateTabBox::OnTabActivated);
 
 	checkf(TabOptions.GetNumAvailableTabs() >= 1, TEXT("At least one tab option needs to be set to true."));
 	
 	// we restore the cached active tab, if it is available in the current context
-	if(TabOptions.IsTabAvailable(CachedActiveTab))
+	if(CachedActiveTab.IsSet() && TabOptions.IsTabAvailable(CachedActiveTab.GetValue()))
 	{
 		ActiveTab = CachedActiveTab;
-		bUseActiveTab = true;
-		bActiveTabInitialized = true;
 	}	
 	
 	if(TabOptions.IsTabAvailable(ENiagaraScriptTemplateSpecification::Template))
 	{
 		// we set the currently active tab using the first available tab in a set order
-		if(bActiveTabInitialized == false)
+		if(!ActiveTab.IsSet())
 		{
 			ActiveTab = ENiagaraScriptTemplateSpecification::Template;
-			bUseActiveTab = true;
-			bActiveTabInitialized = true;
 		}
-
-		TabContainer->AddSlot()
-		.Padding(5.f)
-		[
-			SNew(SBorder)
-			.ToolTipText(LOCTEXT("TemplateTabTooltip", "Templates are intended as starting points for building functional emitters of different types,\n"
-				"and are copied into a system as a unique emitter with no inheritance"))
-			.BorderImage(FEditorStyle::GetBrush("ToolPanel.GroupBorder"))
-			[
-				SNew(SCheckBox)
-				.Style(FNiagaraEditorStyle::Get(), "GraphActionMenu.FilterCheckBox")
-				.BorderBackgroundColor(this, &SNiagaraTemplateTabBox::GetBackgroundColor, ENiagaraScriptTemplateSpecification::Template)
-				.ForegroundColor(this, &SNiagaraTemplateTabBox::GetTabForegroundColor, ENiagaraScriptTemplateSpecification::Template)
-				.IsChecked_Lambda([&]()
-				{
-					return ActiveTab == ENiagaraScriptTemplateSpecification::Template ? ECheckBoxState::Checked : ECheckBoxState::Unchecked;
-				})
-				.OnCheckStateChanged(this, &SNiagaraTemplateTabBox::OnTabActivated, ENiagaraScriptTemplateSpecification::Template)
-			   [
-			       SNew(STextBlock)
-			       .TextStyle(&FNiagaraEditorStyle::Get().GetWidgetStyle<FTextBlockStyle>("GraphActionMenu.TemplateTabTextBlock"))
-			       .Justification(ETextJustify::Center)
-			       .Text(LOCTEXT("TemplateTabLabel", "Templates"))
-			   ]
-			]
-		];
+		
+		TabContainer->AddSlot(ENiagaraScriptTemplateSpecification::Template, false)
+		.Text(DetermineControlLabel(ENiagaraScriptTemplateSpecification::Template))
+		.ToolTip(DetermineControlTooltip(ENiagaraScriptTemplateSpecification::Template));		
 	}
 
 	if(TabOptions.IsTabAvailable(ENiagaraScriptTemplateSpecification::None))
 	{
-		if(bActiveTabInitialized == false)
+		if(!ActiveTab.IsSet())
 		{
 			ActiveTab = ENiagaraScriptTemplateSpecification::None;
-			bUseActiveTab = true;
-			bActiveTabInitialized = true;
-		}
+		}		
 		
-		TabContainer->AddSlot()
-		.Padding(5.f)
-        [
-            SNew(SBorder)
-            .ToolTipText(LOCTEXT("ParentTabTooltip", "Parent Emitters assets are inherited as children and will receive changes from the parent emitter,\n"
-	            "and are meant to serve as art directed initial behaviors which can be propagated throughout a project quickly and easily.\n"
-	            "Over time, a library of parent emitters can be used to speed up the construction of complex effects specific to your project."))
-            .BorderImage(FEditorStyle::GetBrush("ToolPanel.GroupBorder"))
-			[
-	            SNew(SCheckBox)
-				.Style(FNiagaraEditorStyle::Get(), "GraphActionMenu.FilterCheckBox")
-				.BorderBackgroundColor(this, &SNiagaraTemplateTabBox::GetBackgroundColor, ENiagaraScriptTemplateSpecification::None)
-				.ForegroundColor(this, &SNiagaraTemplateTabBox::GetTabForegroundColor, ENiagaraScriptTemplateSpecification::None)
-				.IsChecked_Lambda([&]()
-				{
-					return ActiveTab == ENiagaraScriptTemplateSpecification::None ? ECheckBoxState::Checked : ECheckBoxState::Unchecked;
-				})
-				.OnCheckStateChanged(this, &SNiagaraTemplateTabBox::OnTabActivated, ENiagaraScriptTemplateSpecification::None)
-			     [
-			         SNew(STextBlock)
-			         .TextStyle(&FNiagaraEditorStyle::Get().GetWidgetStyle<FTextBlockStyle>("GraphActionMenu.TemplateTabTextBlock"))
-			         .Justification(ETextJustify::Center)
-			         .Text(LOCTEXT("ParentTabLabel", "Parents"))
-			     ]
-			]
-        ];
+		TabContainer->AddSlot(ENiagaraScriptTemplateSpecification::None, false)
+		.Text(DetermineControlLabel(ENiagaraScriptTemplateSpecification::None))
+		.ToolTip(DetermineControlTooltip(ENiagaraScriptTemplateSpecification::None));
 	}
 
 	if(TabOptions.IsTabAvailable(ENiagaraScriptTemplateSpecification::Behavior))
 	{
-		if(bActiveTabInitialized == false)
+		if(!ActiveTab.IsSet())
 		{
 			ActiveTab = ENiagaraScriptTemplateSpecification::Behavior;
-			bUseActiveTab = true;
-			bActiveTabInitialized = true;
 		}
-		
-		TabContainer->AddSlot()
-		.Padding(5.f)
-        [
-            SNew(SBorder)
-            .ToolTipText(LOCTEXT("BehaviorTabTooltip", "Behavior Examples are intended to serve as a guide to how Niagara works at a feature level.\n"
-	            "Each example shows a simplified setup used to achieve specific outcomes and are intended as starting points, building blocks, or simply as reference.\n"
-	            "These are copied into a system as a unique emitter with no inheritance"))
-            .BorderImage(FEditorStyle::GetBrush("ToolPanel.GroupBorder"))
-			[
-				SNew(SCheckBox)
-				.Style(FNiagaraEditorStyle::Get(), "GraphActionMenu.FilterCheckBox")
-				.BorderBackgroundColor(this, &SNiagaraTemplateTabBox::GetBackgroundColor, ENiagaraScriptTemplateSpecification::Behavior)
-				.ForegroundColor(this, &SNiagaraTemplateTabBox::GetTabForegroundColor, ENiagaraScriptTemplateSpecification::Behavior)
-				.IsChecked_Lambda([&]()
-				{
-					return ActiveTab == ENiagaraScriptTemplateSpecification::Behavior ? ECheckBoxState::Checked : ECheckBoxState::Unchecked;
-				})
-			    .OnCheckStateChanged(this, &SNiagaraTemplateTabBox::OnTabActivated, ENiagaraScriptTemplateSpecification::Behavior)
-			    [
-			        SNew(STextBlock)
-			        .TextStyle(&FNiagaraEditorStyle::Get().GetWidgetStyle<FTextBlockStyle>("GraphActionMenu.TemplateTabTextBlock"))
-			        .Justification(ETextJustify::Center)
-			        .Text(LOCTEXT("BehaviorTabLabel", "Behavior Examples"))
-			    ]
-			]
-        ];
+
+		TabContainer->AddSlot(ENiagaraScriptTemplateSpecification::Behavior, false)
+		.Text(DetermineControlLabel(ENiagaraScriptTemplateSpecification::Behavior))
+		.ToolTip(DetermineControlTooltip(ENiagaraScriptTemplateSpecification::Behavior));
 	}
 
-	// we cache the active tab if we have more than 1 tab available so we can activate it for other instances
-	CachedActiveTab = ActiveTab;
-	
+	TabContainer->RebuildChildren();
+	// we cache the active tab if we have more than 1 tab available so we can activate it for other instances as it was intentionally navigated to
+	if(TabOptions.GetNumAvailableTabs() > 1 && ActiveTab.IsSet())
+	{
+		CachedActiveTab = ActiveTab.GetValue();	
+	}
 
 	ChildSlot
 	[
@@ -401,17 +329,37 @@ void SNiagaraTemplateTabBox::Construct(const FArguments& InArgs, FNiagaraTemplat
 	];
 }
 
+ENiagaraScriptTemplateSpecification SNiagaraTemplateTabBox::GetActiveTab() const
+{
+	if(ActiveTab.IsSet())
+	{
+		return ActiveTab.GetValue();
+	}
+	
+	return ENiagaraScriptTemplateSpecification::None;	
+}
+
 bool SNiagaraTemplateTabBox::GetActiveTab(ENiagaraScriptTemplateSpecification& OutTemplateSpecification) const
 {
-	if(bUseActiveTab)
+	if(ActiveTab.IsSet())
 	{
-		OutTemplateSpecification = ActiveTab;
+		OutTemplateSpecification = ActiveTab.GetValue();
 		return true;
 	}	
 
 	return false;
 }
 
+void SNiagaraTemplateTabBox::OnTabActivated(ENiagaraScriptTemplateSpecification AssetTab)
+{
+	if(ActiveTab != AssetTab)
+	{
+		ActiveTab = AssetTab;
+		CachedActiveTab = ActiveTab;
+		OnTabActivatedDelegate.ExecuteIfBound(AssetTab);
+	}
+}
+
 void SNiagaraTemplateTabBox::OnTabActivated(ECheckBoxState NewState, ENiagaraScriptTemplateSpecification AssetTab)
 {
 	if(ActiveTab != AssetTab)
@@ -442,6 +390,87 @@ FSlateColor SNiagaraTemplateTabBox::GetTabForegroundColor(ENiagaraScriptTemplate
 	return FLinearColor::White;
 }
 
+FText SNiagaraTemplateTabBox::DetermineControlLabel(ENiagaraScriptTemplateSpecification TemplateSpecification) const
+{
+	if(Class.IsValid())
+	{
+		if(Class == UNiagaraEmitter::StaticClass())
+		{
+			switch(TemplateSpecification)
+			{
+			case ENiagaraScriptTemplateSpecification::None:
+				return LOCTEXT("EmitterNoneTabLabel", "Parent Emitters");
+			case ENiagaraScriptTemplateSpecification::Template:
+				return LOCTEXT("EmitterTemplateTabLabel", "Templates");
+			case ENiagaraScriptTemplateSpecification::Behavior:
+				return LOCTEXT("EmitterBehaviorExampleTabLabel", "Behavior Examples");
+			default:
+				return FText::GetEmpty();
+			}
+		}
+		else if(Class == UNiagaraSystem::StaticClass())
+		{
+			switch(TemplateSpecification)
+			{
+			case ENiagaraScriptTemplateSpecification::None:
+				return LOCTEXT("SystemNoneTabLabel", "Standard");
+			case ENiagaraScriptTemplateSpecification::Template:
+				return LOCTEXT("SystemTemplateTabLabel", "Templates");
+			case ENiagaraScriptTemplateSpecification::Behavior:
+				return LOCTEXT("SystemBehaviorExampleTabLabel", "Behavior Examples");
+			default:
+				return FText::GetEmpty();
+			}
+		}				
+	}
+
+	check(false);
+	return FText::GetEmpty();
+}
+
+FText SNiagaraTemplateTabBox::DetermineControlTooltip(ENiagaraScriptTemplateSpecification TemplateSpecification) const
+{
+	if(Class.IsValid())
+	{
+		if(Class == UNiagaraEmitter::StaticClass())
+		{
+			switch(TemplateSpecification)
+			{
+			case ENiagaraScriptTemplateSpecification::None:
+				return LOCTEXT("EmitterNoneTabTooltip", "Parent Emitters are inherited as children and will receive changes from the parent emitter,\n"
+					"and are meant to serve as art directed initial behaviors which can be propagated throughout a project quickly and easily.\n"
+					"Over time, a library of parent emitters can be used to speed up the construction of complex effects specific to your project.");
+			case ENiagaraScriptTemplateSpecification::Template:
+				return LOCTEXT("EmitterTemplateTabTooltip", "Templates are intended as starting points for building functional emitters of different types,\n"
+			"Emitter templates are copied into a system as a unique emitter with no inheritance");
+			case ENiagaraScriptTemplateSpecification::Behavior:
+				return LOCTEXT("EmitterBehaviorExampleTabTooltip", "Behavior Examples are intended to serve as a guide to how Niagara works at a feature level.\n"
+			"Each example shows a simplified setup used to achieve specific outcomes and are intended as starting points, building blocks, or simply as reference.\n"
+			"Behavior Example emitters are copied into a system as a unique emitter with no inheritance");
+			default:
+				return FText::GetEmpty();
+			}
+		}
+		else if(Class == UNiagaraSystem::StaticClass())
+		{
+			switch(TemplateSpecification)
+			{
+			case ENiagaraScriptTemplateSpecification::None:
+				return LOCTEXT("SystemNoneTabTooltip", "By default, any system will be categorized as a standard system.");
+			case ENiagaraScriptTemplateSpecification::Template:
+				return LOCTEXT("SystemTemplateTabTooltip", "Template systems ");
+			case ENiagaraScriptTemplateSpecification::Behavior:
+				return LOCTEXT("SystemBehaviorExampleTabTooltip", "Behavior Example systems serve to demonstrate different Niagara features that rely on multiple emitters interacting in some way.");
+			default:
+				return FText::GetEmpty();				
+			}
+		}	
+	}
+
+	check(false);
+	return FText::GetEmpty();
+}
+
 bool SNiagaraTemplateTabBox::FNiagaraTemplateTabOptions::IsTabAvailable(ENiagaraScriptTemplateSpecification AssetTab) const
 {
 	if(TabData.Contains(AssetTab))
@@ -573,11 +602,8 @@ void SNiagaraFilterBox::Construct(const FArguments& InArgs, FFilterOptions InFil
 			+ SVerticalBox::Slot()
 			.AutoHeight()
 			.Padding(3, 1)
-			[
-				SNew(SBorder)
-				[
-					SourceFilterBox.ToSharedRef()
-				]
+			[				
+				SourceFilterBox.ToSharedRef()
 			]
 		];
 	}
@@ -601,12 +627,10 @@ void SNiagaraFilterBox::Construct(const FArguments& InArgs, FFilterOptions InFil
 				.Text(LOCTEXT("LibraryVisiblityLabel", "Library Filtering"))
 			]
 			+ SVerticalBox::Slot()
+			.HAlign(HAlign_Left)
 			.Padding(5, 1)
 			[
-				SNew(SBorder)
-				[					
-					LibraryOnlyToggleHeader.ToSharedRef()					
-				]
+				LibraryOnlyToggleHeader.ToSharedRef()
 			]     
 		];
 	}
@@ -618,8 +642,11 @@ void SNiagaraFilterBox::Construct(const FArguments& InArgs, FFilterOptions InFil
 
 	if(InFilterOptions.GetAddTemplateFilter())
 	{
+		ensureMsgf(InArgs._Class != nullptr, TEXT("If the template filter is added, the class must be valid as depending on class the control labels and tooltips are generated."));
+		
 		TemplateTabBox = SNew(SNiagaraTemplateTabBox, InFilterOptions.GetTabOptions())
-		.OnTabActivated(InArgs._OnTabActivated);
+		.OnTabActivated(InArgs._OnTabActivated)
+		.Class(InArgs._Class);
 		
 		// the template filter is a whole-row filter
 		MainContainer->AddSlot()
@@ -638,10 +665,7 @@ void SNiagaraFilterBox::Construct(const FArguments& InArgs, FFilterOptions InFil
 			+ SVerticalBox::Slot()
 			.Padding(3, 1)
 			[
-				SNew(SBorder)
-				[
-					TemplateTabBox.ToSharedRef()	
-				]
+				TemplateTabBox.ToSharedRef()	
 			]     
 		];
 	}
diff --git a/Engine/Plugins/FX/Niagara/Source/NiagaraEditor/Public/Widgets/SNiagaraFilterBox.h b/Engine/Plugins/FX/Niagara/Source/NiagaraEditor/Public/Widgets/SNiagaraFilterBox.h
--- a/Engine/Plugins/FX/Niagara/Source/NiagaraEditor/Public/Widgets/SNiagaraFilterBox.h
+++ b/Engine/Plugins/FX/Niagara/Source/NiagaraEditor/Public/Widgets/SNiagaraFilterBox.h
@@ -8,6 +8,7 @@
 #include "Widgets/SNiagaraLibraryOnlyToggleHeader.h"
 #include "Widgets/Layout/SBorder.h"
 #include "Widgets/Input/SCheckBox.h"
+#include "Widgets/Input/SSegmentedControl.h"
 #include "Widgets/Views/SMultipleOptionTable.h"
 
 /**
@@ -104,29 +105,35 @@ public:
 	};
 	
 	SLATE_BEGIN_ARGS(SNiagaraTemplateTabBox)
+		: _Class(nullptr)
 	{}
 		SLATE_EVENT(FOnTabActivated, OnTabActivated)
+		SLATE_ARGUMENT(UClass*, Class)
 	SLATE_END_ARGS()
 
 	void Construct(const FArguments& InArgs, FNiagaraTemplateTabOptions TabOptions);
 
+	ENiagaraScriptTemplateSpecification GetActiveTab() const;
 	/** If true, will set the parameter to the currently active tab. If false, the template filter wasn't initialized with any tabs. */
 	bool GetActiveTab(ENiagaraScriptTemplateSpecification& OutTemplateSpecification) const;
 private:
 	FOnTabActivated OnTabActivatedDelegate;
+	void OnTabActivated(ENiagaraScriptTemplateSpecification AssetTab);
 	void OnTabActivated(ECheckBoxState NewState, ENiagaraScriptTemplateSpecification AssetTab);
 	FSlateColor GetBackgroundColor(ENiagaraScriptTemplateSpecification TemplateSpecification) const;
 	FSlateColor GetTabForegroundColor(ENiagaraScriptTemplateSpecification TemplateSpecification) const;
-
+	FText DetermineControlLabel(ENiagaraScriptTemplateSpecification TemplateSpecification) const;
+	FText DetermineControlTooltip(ENiagaraScriptTemplateSpecification TemplateSpecification) const;
 private:
-	TSharedPtr<SHorizontalBox> TabContainer;
-	
 	FNiagaraTemplateTabOptions TabOptions;
+	TSharedPtr<SSegmentedControl<ENiagaraScriptTemplateSpecification>> TabContainer;
 	
 	/** bUseActiveTab is required due to tab options not specifically needing to have any tab state set to true. */
 	bool bUseActiveTab = false;
-	ENiagaraScriptTemplateSpecification ActiveTab;
-	static ENiagaraScriptTemplateSpecification CachedActiveTab;
+	TOptional<ENiagaraScriptTemplateSpecification> ActiveTab;
+	static TOptional<ENiagaraScriptTemplateSpecification> CachedActiveTab;
+	
+	TWeakObjectPtr<UClass> Class = nullptr;
 };
 
 /**
@@ -166,16 +173,17 @@ public:
 	
 	SLATE_BEGIN_ARGS(SNiagaraFilterBox)
 		: _bLibraryOnly(true)
+		, _Class(nullptr)
 	{}
 		SLATE_ATTRIBUTE(bool, bLibraryOnly)
 		SLATE_EVENT(SNiagaraSourceFilterBox::FOnFiltersChanged, OnSourceFiltersChanged)
 		SLATE_EVENT(SNiagaraLibraryOnlyToggleHeader::FOnLibraryOnlyChanged, OnLibraryOnlyChanged)
 		SLATE_EVENT(SNiagaraTemplateTabBox::FOnTabActivated, OnTabActivated)
+		SLATE_ARGUMENT(UClass*, Class)
 	SLATE_END_ARGS()
 
 	void Construct(const FArguments& InArgs, FFilterOptions FilterOptions);
-
-	//void SetOnSourceFiltersChanged(SNiagaraSourceFilterBox::FOnFiltersChanged OnSourceFiltersChanged) { }
+	
 	bool IsSourceFilterActive(EScriptSource Source) const;
 	bool GetActiveTemplateTab(ENiagaraScriptTemplateSpecification& OutTemplateSpecification) const;
 	
