Make ope type validation emit warnings instead of error on compile

diff --git a/Engine/Plugins/FX/Niagara/Source/NiagaraEditor/Private/NiagaraHlslTranslator.cpp b/Engine/Plugins/FX/Niagara/Source/NiagaraEditor/Private/NiagaraHlslTranslator.cpp
--- a/Engine/Plugins/FX/Niagara/Source/NiagaraEditor/Private/NiagaraHlslTranslator.cpp
+++ b/Engine/Plugins/FX/Niagara/Source/NiagaraEditor/Private/NiagaraHlslTranslator.cpp
@@ -6628,7 +6628,7 @@ void FHlslNiagaraTranslator::Operation(class UNiagaraNodeOp* Operation, TArray<i
 	FText ValidationError;
 	if (OpInfo && OpInfo->InputTypeValidationFunction.IsBound() && OpInfo->InputTypeValidationFunction.Execute(InputTypes, ValidationError) == false)
 	{
-		Error(ValidationError, Operation, OutputPins[0]);
+		Warning(ValidationError, Operation, OutputPins[0], TEXT("Invalid op types"), true);
 		Outputs.Add(INDEX_NONE);
 		return;
 	}
