- Auto enable Gpu profiler for debug hud & csv with detailed stats

#rnx

diff --git a/Engine/Plugins/FX/Niagara/Source/Niagara/Classes/NiagaraGPUProfilerInterface.h b/Engine/Plugins/FX/Niagara/Source/Niagara/Classes/NiagaraGPUProfilerInterface.h
--- a/Engine/Plugins/FX/Niagara/Source/Niagara/Classes/NiagaraGPUProfilerInterface.h
+++ b/Engine/Plugins/FX/Niagara/Source/Niagara/Classes/NiagaraGPUProfilerInterface.h
@@ -58,10 +58,27 @@ struct FNiagaraGpuFrameResults : public TSharedFromThis<FNiagaraGpuFrameResults,
 
 using FNiagaraGpuFrameResultsPtr = TSharedPtr<FNiagaraGpuFrameResults, ESPMode::ThreadSafe>;
 
+//////////////////////////////////////////////////////////////////////////
+// Allows various systems to listen to profiler results
+struct NIAGARA_API FNiagaraGpuProfilerListener
+{
+	FNiagaraGpuProfilerListener();
+	~FNiagaraGpuProfilerListener();
+
+	void SetEnabled(bool bEnabled);
+	void SetHandler(TFunction<void(const FNiagaraGpuFrameResultsPtr&)> Function);
+	bool IsEnabled() const { return bEnabled; }
+
+private:
+	bool			bEnabled = false;
+	FDelegateHandle	GameThreadHandler;
+};
+
 //////////////////////////////////////////////////////////////////////////
 /** Public API to Niagara GPU Profiling. */
 class FNiagaraGPUProfilerInterface
 {
+	friend FNiagaraGpuProfilerListener;
 	DECLARE_MULTICAST_DELEGATE_OneParam(FOnFrameResults, const FNiagaraGpuFrameResultsPtr&);
 
 public:
@@ -72,6 +89,7 @@ protected:
 	void PostResults(const FNiagaraGpuFrameResultsPtr& FrameResults);
 
 protected:
+	static std::atomic<int>	NumReaders;
 	static FOnFrameResults	OnFrameResults_GameThread;
 	static FOnFrameResults	OnFrameResults_RenderThread;
 };
diff --git a/Engine/Plugins/FX/Niagara/Source/Niagara/Private/NiagaraDebugHud.cpp b/Engine/Plugins/FX/Niagara/Source/Niagara/Private/NiagaraDebugHud.cpp
--- a/Engine/Plugins/FX/Niagara/Source/Niagara/Private/NiagaraDebugHud.cpp
+++ b/Engine/Plugins/FX/Niagara/Source/Niagara/Private/NiagaraDebugHud.cpp
@@ -538,7 +538,7 @@ FNiagaraDebugHud::FNiagaraDebugHud(UWorld* World)
 	}
 	++GDebugDrawHandleUsers;
 #if WITH_NIAGARA_GPU_PROFILER
-	GpuProfilerDelegateHandle = FNiagaraGPUProfilerInterface::GetOnFrameResults_GameThread().AddLambda(
+	GpuProfilerListener.SetHandler(
 		[&](const FNiagaraGpuFrameResultsPtr& InGpuResults)
 		{
 			FNiagaraGpuComputeDispatchInterface* DispatchInterface = FNiagaraGpuComputeDispatchInterface::Get(WeakWorld.Get());
@@ -640,9 +640,6 @@ FNiagaraDebugHud::~FNiagaraDebugHud()
 		UDebugDrawService::Unregister(GDebugDrawHandle);
 		GDebugDrawHandle.Reset();
 	}
-#if WITH_NIAGARA_GPU_PROFILER
-	FNiagaraGPUProfilerInterface::GetOnFrameResults_GameThread().Remove(GpuProfilerDelegateHandle);
-#endif
 }
 
 void FNiagaraDebugHud::UpdateSettings(const FNiagaraDebugHUDSettingsData& NewSettings)
@@ -690,6 +687,11 @@ void FNiagaraDebugHud::GatherSystemInfo()
 		return;
 	}
 
+#if WITH_NIAGARA_GPU_PROFILER
+	// Update Gpu capture state
+	GpuProfilerListener.SetEnabled(Settings.bOverviewEnabled && (Settings.OverviewMode == ENiagaraDebugHUDOverviewMode::GpuComputePerformance));
+#endif
+
 	// When not enabled do nothing
 	if (!Settings.IsEnabled())
 	{
@@ -1699,12 +1701,16 @@ void FNiagaraDebugHud::DrawGpuComputeOverriew(class FNiagaraWorldManager* WorldM
 
 	if (GpuResults == nullptr)
 	{
-		static const FString EnableCVarWarning(TEXT("No GPU data ensure 'fx.NiagaraGpuProfilingEnabled' is enabled"));
+		static const auto ProfilingEnabledCVar = IConsoleManager::Get().FindConsoleVariable(TEXT("fx.Niagara.GpuProfiling.Enabled"));
+
+		static const FString EnableCVarWarning(TEXT("GPU Profiling is disabled, enable 'fx.Niagara.GpuProfiling.Enabled'"));
+		static const FString NoDataWarning(TEXT("No GPU data is ready"));
 
 		const FVector2D StringSize = GetStringSize(Font, *EnableCVarWarning);
 		DrawCanvas->DrawTile(TextLocation.X - 1.0f, TextLocation.Y - 1.0f, StringSize.X + 1.0f, 2.0f + StringSize.Y, 0.0f, 0.0f, 0.0f, 0.0f, BackgroundColor);
-		DrawCanvas->DrawShadowedString(TextLocation.X, TextLocation.Y, *EnableCVarWarning, Font, HeadingColor);
+		DrawCanvas->DrawShadowedString(TextLocation.X, TextLocation.Y, *(ProfilingEnabledCVar && ProfilingEnabledCVar->GetBool() ? NoDataWarning : EnableCVarWarning), Font, HeadingColor);
 		TextLocation.Y += StringSize.Y;
+		
 		return;
 	}
 
diff --git a/Engine/Plugins/FX/Niagara/Source/Niagara/Private/NiagaraDebugHud.h b/Engine/Plugins/FX/Niagara/Source/Niagara/Private/NiagaraDebugHud.h
--- a/Engine/Plugins/FX/Niagara/Source/Niagara/Private/NiagaraDebugHud.h
+++ b/Engine/Plugins/FX/Niagara/Source/Niagara/Private/NiagaraDebugHud.h
@@ -268,7 +268,7 @@ private:
 #endif
 
 #if WITH_NIAGARA_GPU_PROFILER
-	FDelegateHandle				GpuProfilerDelegateHandle;
+	FNiagaraGpuProfilerListener	GpuProfilerListener;
 	FNiagaraGpuFrameResultsPtr	GpuResults;
 	uint64						GpuResultsGameFrameCounter = 0;
 
diff --git a/Engine/Plugins/FX/Niagara/Source/Niagara/Private/NiagaraGPUProfiler.cpp b/Engine/Plugins/FX/Niagara/Source/Niagara/Private/NiagaraGPUProfiler.cpp
--- a/Engine/Plugins/FX/Niagara/Source/Niagara/Private/NiagaraGPUProfiler.cpp
+++ b/Engine/Plugins/FX/Niagara/Source/Niagara/Private/NiagaraGPUProfiler.cpp
@@ -6,11 +6,11 @@
 #include "HAL/IConsoleManager.h"
 #include "Misc/StringBuilder.h"
 
-int32 GNiagaraGpuProfilingEnabled = 0;
+bool GNiagaraGpuProfilingEnabled = true;
 static FAutoConsoleVariableRef CVarNiagaraGpuProfilingEnabled(
-	TEXT("fx.NiagaraGpuProfilingEnabled"),
+	TEXT("fx.Niagara.GpuProfiling.Enabled"),
 	GNiagaraGpuProfilingEnabled,
-	TEXT("Used by the profiling tool in the system overview to enable or disable gathering of gpu stats.\n"),
+	TEXT("Master control to allow Niagara to use GPU profiling or not.\n"),
 	ECVF_Default
 );
 
@@ -57,7 +57,9 @@ void FNiagaraGPUProfiler::BeginFrame(FRHICommandListImmediate& RHICmdList)
 	}
 
 	// If we are not enabled nothing to profile
-	if (!GNiagaraGpuProfilingEnabled)
+	static const auto CSVStatsEnabledCVar = IConsoleManager::Get().FindConsoleVariable(TEXT("fx.DetailedCSVStats"));
+	const bool bCsvCollecting = CSVStatsEnabledCVar && CSVStatsEnabledCVar->GetBool();
+	if ( !GNiagaraGpuProfilingEnabled || ((NumReaders == 0) && !bCsvCollecting) )
 	{
 		return;
 	}
diff --git a/Engine/Plugins/FX/Niagara/Source/Niagara/Private/NiagaraGPUProfilerInterface.cpp b/Engine/Plugins/FX/Niagara/Source/Niagara/Private/NiagaraGPUProfilerInterface.cpp
--- a/Engine/Plugins/FX/Niagara/Source/Niagara/Private/NiagaraGPUProfilerInterface.cpp
+++ b/Engine/Plugins/FX/Niagara/Source/Niagara/Private/NiagaraGPUProfilerInterface.cpp
@@ -11,6 +11,42 @@
 
 #if WITH_NIAGARA_GPU_PROFILER
 
+//////////////////////////////////////////////////////////////////////////
+
+FNiagaraGpuProfilerListener::FNiagaraGpuProfilerListener()
+{
+}
+
+FNiagaraGpuProfilerListener::~FNiagaraGpuProfilerListener()
+{
+	check(IsInGameThread());
+	SetEnabled(false);
+	SetHandler(nullptr);
+}
+
+void FNiagaraGpuProfilerListener::SetEnabled(bool bInEnabled)
+{
+	check(IsInGameThread());
+	if ( bEnabled != bInEnabled)
+	{
+		bEnabled = bInEnabled;
+		FNiagaraGPUProfilerInterface::NumReaders.fetch_add(bEnabled ? 1 : -1);
+	}
+}
+
+void FNiagaraGpuProfilerListener::SetHandler(TFunction<void(const FNiagaraGpuFrameResultsPtr&)> Function)
+{
+	check(IsInGameThread());
+	if (GameThreadHandler.IsValid())
+	{
+		FNiagaraGPUProfilerInterface::GetOnFrameResults_GameThread().Remove(GameThreadHandler);
+	}
+	if ( Function )
+	{
+		GameThreadHandler = FNiagaraGPUProfilerInterface::GetOnFrameResults_GameThread().AddLambda(Function);
+	}
+}
+
 //////////////////////////////////////////////////////////////////////////
 FNiagaraGpuProfileScope::FNiagaraGpuProfileScope(FRHICommandList& InRHICmdList, const FNiagaraGpuComputeDispatchInterface* ComputeDispatchInterface, FName StageName)
 	: RHICmdList(InRHICmdList)
@@ -138,6 +174,7 @@ void SetResultsForEditorStats(const FNiagaraGpuFrameResultsPtr& FrameResults)
 
 //////////////////////////////////////////////////////////////////////////
 
+std::atomic<int> FNiagaraGPUProfilerInterface::NumReaders(0);
 FNiagaraGPUProfilerInterface::FOnFrameResults FNiagaraGPUProfilerInterface::OnFrameResults_GameThread;
 FNiagaraGPUProfilerInterface::FOnFrameResults FNiagaraGPUProfilerInterface::OnFrameResults_RenderThread;
 
diff --git a/Engine/Plugins/FX/Niagara/Source/NiagaraEditor/Private/Toolkits/NiagaraSystemToolkit.cpp b/Engine/Plugins/FX/Niagara/Source/NiagaraEditor/Private/Toolkits/NiagaraSystemToolkit.cpp
--- a/Engine/Plugins/FX/Niagara/Source/NiagaraEditor/Private/Toolkits/NiagaraSystemToolkit.cpp
+++ b/Engine/Plugins/FX/Niagara/Source/NiagaraEditor/Private/Toolkits/NiagaraSystemToolkit.cpp
@@ -89,7 +89,6 @@ const FName FNiagaraSystemToolkit::ScratchPadTabID(TEXT("NiagaraSystemEditor_Scr
 const FName FNiagaraSystemToolkit::ScriptStatsTabID(TEXT("NiagaraSystemEditor_ScriptStats"));
 const FName FNiagaraSystemToolkit::BakerTabID(TEXT("NiagaraSystemEditor_Baker"));
 IConsoleVariable* FNiagaraSystemToolkit::VmStatEnabledVar = IConsoleManager::Get().FindConsoleVariable(TEXT("vm.DetailedVMScriptStats"));
-IConsoleVariable* FNiagaraSystemToolkit::GpuStatEnabledVar = IConsoleManager::Get().FindConsoleVariable(TEXT("fx.NiagaraGpuProfilingEnabled"));
 
 static int32 GbLogNiagaraSystemChanges = 0;
 static FAutoConsoleVariableRef CVarSuppressNiagaraSystems(
@@ -474,8 +473,11 @@ void FNiagaraSystemToolkit::InitializeInternal(const EToolkitMode::Type Mode, co
 	bChangesDiscarded = false;
 	bScratchPadChangesDiscarded = false;
 
-
 	GEditor->RegisterForUndo(this);
+
+#if WITH_NIAGARA_GPU_PROFILER
+	GpuProfilerListener.Reset(new FNiagaraGpuProfilerListener);
+#endif
 }
 
 FName FNiagaraSystemToolkit::GetToolkitFName() const
@@ -1579,15 +1581,23 @@ bool FNiagaraSystemToolkit::IsStatPerformanceChecked()
 
 void FNiagaraSystemToolkit::ToggleStatPerformanceGPU()
 {
-	if (GpuStatEnabledVar)
+#if WITH_NIAGARA_GPU_PROFILER
+	if (GpuProfilerListener)
 	{
-		GpuStatEnabledVar->Set(!IsStatPerformanceGPUChecked());
+		GpuProfilerListener->SetEnabled(!GpuProfilerListener->IsEnabled());
 	}
+#endif
 }
 
 bool FNiagaraSystemToolkit::IsStatPerformanceGPUChecked()
 {
-	return GpuStatEnabledVar ? GpuStatEnabledVar->GetBool() : false;
+#if WITH_NIAGARA_GPU_PROFILER
+	if (GpuProfilerListener)
+	{
+		return GpuProfilerListener->IsEnabled();
+	}
+#endif
+	return false;
 }
 
 void FNiagaraSystemToolkit::UpdateOriginalEmitter()
diff --git a/Engine/Plugins/FX/Niagara/Source/NiagaraEditor/Private/Toolkits/NiagaraSystemToolkit.h b/Engine/Plugins/FX/Niagara/Source/NiagaraEditor/Private/Toolkits/NiagaraSystemToolkit.h
--- a/Engine/Plugins/FX/Niagara/Source/NiagaraEditor/Private/Toolkits/NiagaraSystemToolkit.h
+++ b/Engine/Plugins/FX/Niagara/Source/NiagaraEditor/Private/Toolkits/NiagaraSystemToolkit.h
@@ -233,7 +233,10 @@ private:
 	bool bScratchPadChangesDiscarded;
 
 	static IConsoleVariable* VmStatEnabledVar;
-	static IConsoleVariable* GpuStatEnabledVar;
+
+#if WITH_NIAGARA_GPU_PROFILER
+	TUniquePtr<struct FNiagaraGpuProfilerListener> GpuProfilerListener;
+#endif
 
 public:
 	static const FName ViewportTabID;
