Fix for assignment nodes not showing up in summary view in some cases.
Fixed summary view sort order allowing to clear the value to revert to default.
Fixed some merging issues keeping parents from applying updates to children correctly.
Fixed sort order to work across modules correctly.
Fixed issue with NiagaraEmitterEditorData not being handled correctly during a merge and being lost on reload.

diff --git a/Engine/Plugins/FX/Niagara/Source/Niagara/Private/NiagaraEmitter.cpp b/Engine/Plugins/FX/Niagara/Source/Niagara/Private/NiagaraEmitter.cpp
--- a/Engine/Plugins/FX/Niagara/Source/Niagara/Private/NiagaraEmitter.cpp
+++ b/Engine/Plugins/FX/Niagara/Source/Niagara/Private/NiagaraEmitter.cpp
@@ -1525,11 +1525,11 @@ void UNiagaraEmitter::SetEditorData(UNiagaraEditorDataBase* InEditorData)
 	{
 		return;
 	}
+ 
 	if (EditorData != nullptr)
 	{
 		EditorData->OnPersistentDataChanged().RemoveAll(this);
 	}
-
 	EditorData = InEditorData;
 	
 	if (EditorData != nullptr)
@@ -1992,7 +1992,9 @@ void UNiagaraEmitter::UpdateFromMergedCopy(const INiagaraMergeManager& MergeMana
 		ScratchPadScripts.Add(MergedScratchPadScript);
 	}
 
-	SetEditorData(MergedEmitter->GetEditorData());
+	UNiagaraEditorDataBase* NewEditorData = MergedEmitter->GetEditorData();
+	ReouterMergedObject(this, NewEditorData);
+	SetEditorData(NewEditorData);
 
 	// Update the change id since we don't know what's changed.
 	UpdateChangeId(TEXT("Updated from merged copy"));
diff --git a/Engine/Plugins/FX/Niagara/Source/NiagaraEditor/Private/NiagaraScriptMergeManager.cpp b/Engine/Plugins/FX/Niagara/Source/NiagaraEditor/Private/NiagaraScriptMergeManager.cpp
--- a/Engine/Plugins/FX/Niagara/Source/NiagaraEditor/Private/NiagaraScriptMergeManager.cpp
+++ b/Engine/Plugins/FX/Niagara/Source/NiagaraEditor/Private/NiagaraScriptMergeManager.cpp
@@ -884,6 +884,7 @@ void FNiagaraEmitterMergeAdapter::Initialize(const UNiagaraEmitter& InEmitter, U
 	}
 
 	EditorData = Cast<const UNiagaraEmitterEditorData>(Emitter->GetEditorData());
+
 }
 
 UNiagaraEmitter* FNiagaraEmitterMergeAdapter::GetEditableEmitter() const
@@ -926,11 +927,6 @@ const TArray<TSharedRef<FNiagaraRendererMergeAdapter>> FNiagaraEmitterMergeAdapt
 	return Renderers;
 }
 
-const TArray<TSharedRef<FNiagaraInputSummaryMergeAdapter>> FNiagaraEmitterMergeAdapter::GetInputSummaryEntries() const
-{
-	return InputSummaryEntries;
-}
-
 const UNiagaraEmitterEditorData* FNiagaraEmitterMergeAdapter::GetEditorData() const
 {
 	return EditorData.Get();
@@ -1990,7 +1986,7 @@ FNiagaraEmitterDiffResults FNiagaraScriptMergeManager::DiffEmitters(UNiagaraEmit
 	DiffEventHandlers(BaseEmitterAdapter->GetEventHandlers(), OtherEmitterAdapter->GetEventHandlers(), EmitterDiffResults);
 	DiffSimulationStages(BaseEmitterAdapter->GetSimulationStages(), OtherEmitterAdapter->GetSimulationStages(), EmitterDiffResults);
 	DiffRenderers(BaseEmitterAdapter->GetRenderers(), OtherEmitterAdapter->GetRenderers(), EmitterDiffResults);
-	DiffInputSummaryEntries(BaseEmitterAdapter->GetInputSummaryEntries(), OtherEmitterAdapter->GetInputSummaryEntries(), EmitterDiffResults);
+	DiffInputSummaryEntries(BaseEmitterAdapter->GetEditorData(), OtherEmitterAdapter->GetEditorData(), EmitterDiffResults);
 	DiffEditableProperties(&BaseEmitter, &OtherEmitter, *UNiagaraEmitter::StaticClass(), EmitterDiffResults.DifferentEmitterProperties);
 	DiffStackEntryDisplayNames(BaseEmitterAdapter->GetEditorData(), OtherEmitterAdapter->GetEditorData(), EmitterDiffResults.ModifiedStackEntryDisplayNames);
 
@@ -2211,15 +2207,34 @@ void FNiagaraScriptMergeManager::DiffRenderers(const TArray<TSharedRef<FNiagaraR
 	}
 }
 
-void FNiagaraScriptMergeManager::DiffInputSummaryEntries(const TArray<TSharedRef<FNiagaraInputSummaryMergeAdapter>>& BaseInputEntries, const TArray<TSharedRef<FNiagaraInputSummaryMergeAdapter>>& OtherInputEntries, FNiagaraEmitterDiffResults& DiffResults) const
+void FNiagaraScriptMergeManager::DiffInputSummaryEntries(const UNiagaraEmitterEditorData* BaseEditorData, const UNiagaraEmitterEditorData* OtherEditorData, FNiagaraEmitterDiffResults& DiffResults) const
 {
+	TArray<TSharedRef<FNiagaraInputSummaryMergeAdapter>> BaseSummaryViewEntries;
+	TArray<TSharedRef<FNiagaraInputSummaryMergeAdapter>> OtherSummaryViewEntries;
+
+	if (BaseEditorData)
+	{
+		for (const auto& Entry : BaseEditorData->GetSummaryViewMetaDataMap())
+		{
+			BaseSummaryViewEntries.Add(MakeShared<FNiagaraInputSummaryMergeAdapter>(Entry.Key, Entry.Value));
+		}
+	}
+
+	if (OtherEditorData)
+	{
+		for (const auto& Entry : OtherEditorData->GetSummaryViewMetaDataMap())
+		{
+			OtherSummaryViewEntries.Add(MakeShared<FNiagaraInputSummaryMergeAdapter>(Entry.Key, Entry.Value));
+		}		
+	}
+	
 	FListDiffResults<TSharedRef<FNiagaraInputSummaryMergeAdapter>> InputListDiffResults = DiffLists<TSharedRef<FNiagaraInputSummaryMergeAdapter>, FFunctionInputSummaryViewKey>(
-		BaseInputEntries,
-		OtherInputEntries,
+		BaseSummaryViewEntries,
+		OtherSummaryViewEntries,
 		[](TSharedRef<FNiagaraInputSummaryMergeAdapter> Entry) { return Entry->GetKey(); });
 
 	DiffResults.RemovedInputSummaryEntries.Append(InputListDiffResults.RemovedBaseValues);
-	DiffResults.AddedInputSummaryEntries.Append(InputListDiffResults.AddedOtherValues);
+	DiffResults.AddedInputSummaryEntries.Append(InputListDiffResults.AddedOtherValues);	
 
 	for (const FCommonValuePair<TSharedRef<FNiagaraInputSummaryMergeAdapter>>& CommonValuePair : InputListDiffResults.CommonValuePairs)
 	{
@@ -3328,12 +3343,23 @@ FNiagaraScriptMergeManager::FApplyDiffResults FNiagaraScriptMergeManager::ApplyR
 	return Results;
 }
 
-FNiagaraScriptMergeManager::FApplyDiffResults FNiagaraScriptMergeManager::ApplyInputSummaryDiff(UNiagaraEmitter& BaseEmitter, const FNiagaraEmitterDiffResults& DiffResults, const bool bNoParentAtLastMerge) const
+FNiagaraScriptMergeManager::FApplyDiffResults FNiagaraScriptMergeManager::ApplyInputSummaryDiff(UNiagaraEmitter& Emitter, const FNiagaraEmitterDiffResults& DiffResults, const bool bNoParentAtLastMerge) const
 {
-	UNiagaraEmitterEditorData* EditorData = Cast<UNiagaraEmitterEditorData>(BaseEmitter.GetEditorData());
+	UNiagaraEmitterEditorData* EditorData = Cast<UNiagaraEmitterEditorData>(Emitter.GetEditorData());
 
+	if (EditorData == nullptr)
+	{
+		EditorData = NewObject<UNiagaraEmitterEditorData>(&Emitter, NAME_None, RF_Transactional);
+		Emitter.Modify();
+		Emitter.SetEditorData(EditorData);
+	}
+
+	check(EditorData != GetDefault<UNiagaraEmitterEditorData>());
+	check(EditorData->GetOuter() == &Emitter);
+	
 	if (EditorData)
 	{
+		EditorData->Modify();
 		for (TSharedRef<FNiagaraInputSummaryMergeAdapter> RemovedInput : DiffResults.RemovedInputSummaryEntries)
 		{
 			EditorData->SetSummaryViewMetaData(RemovedInput->GetKey(), FFunctionInputSummaryViewMetadata());
@@ -3364,6 +3390,7 @@ FNiagaraScriptMergeManager::FApplyDiffResults FNiagaraScriptMergeManager::ApplyS
 		if (EditorData == nullptr)
 		{
 			EditorData = NewObject<UNiagaraEmitterEditorData>(&Emitter, NAME_None, RF_Transactional);
+			EditorData->Modify();
 			Emitter.SetEditorData(EditorData);
 		}
 
diff --git a/Engine/Plugins/FX/Niagara/Source/NiagaraEditor/Private/NiagaraScriptMergeManager.h b/Engine/Plugins/FX/Niagara/Source/NiagaraEditor/Private/NiagaraScriptMergeManager.h
--- a/Engine/Plugins/FX/Niagara/Source/NiagaraEditor/Private/NiagaraScriptMergeManager.h
+++ b/Engine/Plugins/FX/Niagara/Source/NiagaraEditor/Private/NiagaraScriptMergeManager.h
@@ -255,7 +255,6 @@ public:
 	const TArray<TSharedRef<FNiagaraEventHandlerMergeAdapter>> GetEventHandlers() const;
 	const TArray<TSharedRef<FNiagaraSimulationStageMergeAdapter>> GetSimulationStages() const;
 	const TArray<TSharedRef<FNiagaraRendererMergeAdapter>> GetRenderers() const;
-	const TArray<TSharedRef<FNiagaraInputSummaryMergeAdapter>> GetInputSummaryEntries() const;
 	const UNiagaraEmitterEditorData* GetEditorData() const;
 
 	TSharedPtr<FNiagaraScriptStackMergeAdapter> GetScriptStack(ENiagaraScriptUsage Usage, FGuid UsageId);
@@ -281,7 +280,6 @@ private:
 	TArray<TSharedRef<FNiagaraEventHandlerMergeAdapter>> EventHandlers;
 	TArray<TSharedRef<FNiagaraSimulationStageMergeAdapter>> SimulationStages;
 	TArray<TSharedRef<FNiagaraRendererMergeAdapter>> Renderers;
-	TArray<TSharedRef<FNiagaraInputSummaryMergeAdapter>> InputSummaryEntries;
 	TWeakObjectPtr<const UNiagaraEmitterEditorData> EditorData;
 };
 
@@ -460,7 +458,7 @@ public:
 
 	void DiffRenderers(const TArray<TSharedRef<FNiagaraRendererMergeAdapter>>& BaseRenderers, const TArray<TSharedRef<FNiagaraRendererMergeAdapter>>& OtherRenderers, FNiagaraEmitterDiffResults& DiffResults) const;
 	
-	void DiffInputSummaryEntries(const TArray<TSharedRef<FNiagaraInputSummaryMergeAdapter>>& BaseInputEntries, const TArray<TSharedRef<FNiagaraInputSummaryMergeAdapter>>& OtherInputEntries, FNiagaraEmitterDiffResults& DiffResults) const;
+	void DiffInputSummaryEntries(const UNiagaraEmitterEditorData* BaseEditorData, const UNiagaraEmitterEditorData* OtherEditorData, FNiagaraEmitterDiffResults& DiffResults) const;
 
 	void DiffScriptStacks(TSharedRef<FNiagaraScriptStackMergeAdapter> BaseScriptStackAdapter, TSharedRef<FNiagaraScriptStackMergeAdapter> OtherScriptStackAdapter, FNiagaraScriptStackDiffResults& DiffResults) const;
 
@@ -509,7 +507,7 @@ private:
 
 	FApplyDiffResults ApplyRendererDiff(UNiagaraEmitter& BaseEmitter, const FNiagaraEmitterDiffResults& DiffResults, const bool bNoParentAtLastMerge) const;
 
-	FApplyDiffResults ApplyInputSummaryDiff(UNiagaraEmitter& BaseEmitter, const FNiagaraEmitterDiffResults& DiffResults, const bool bNoParentAtLastMerge) const;
+	FApplyDiffResults ApplyInputSummaryDiff(UNiagaraEmitter& Emitter, const FNiagaraEmitterDiffResults& DiffResults, const bool bNoParentAtLastMerge) const;
 	
 	FApplyDiffResults ApplyStackEntryDisplayNameDiffs(UNiagaraEmitter& Emitter, const FNiagaraEmitterDiffResults& DiffResults) const;
 
diff --git a/Engine/Plugins/FX/Niagara/Source/NiagaraEditor/Private/ViewModels/Stack/NiagaraStackFilteredObject.cpp b/Engine/Plugins/FX/Niagara/Source/NiagaraEditor/Private/ViewModels/Stack/NiagaraStackFilteredObject.cpp
--- a/Engine/Plugins/FX/Niagara/Source/NiagaraEditor/Private/ViewModels/Stack/NiagaraStackFilteredObject.cpp
+++ b/Engine/Plugins/FX/Niagara/Source/NiagaraEditor/Private/ViewModels/Stack/NiagaraStackFilteredObject.cpp
@@ -102,11 +102,6 @@ void UNiagaraStackFilteredObject::RefreshChildrenInternal(const TArray<UNiagaraS
 	AppendEmitterCategory(ScriptViewModelPinned, ENiagaraScriptUsage::ParticleSpawnScript, FGuid(), NewChildren, CurrentChildren, NewIssues);
 	AppendEmitterCategory(ScriptViewModelPinned, ENiagaraScriptUsage::ParticleUpdateScript, FGuid(), NewChildren, CurrentChildren, NewIssues);
 
-// 	AppendEmitterCategory(ScriptViewModelPinned, ENiagaraScriptUsage::EmitterUpdateScript, FGuid(), NewChildren, CurrentChildren, NewIssues);
-// 	AppendEmitterCategory(ScriptViewModelPinned, ENiagaraScriptUsage::EmitterUpdateScript, FGuid(), NewChildren, CurrentChildren, NewIssues);
-// 	AppendEmitterCategory(ScriptViewModelPinned, ENiagaraScriptUsage::EmitterUpdateScript, FGuid(), NewChildren, CurrentChildren, NewIssues);
-
-
 	Super::RefreshChildrenInternal(CurrentChildren, NewChildren, NewIssues);
 }
 
@@ -126,7 +121,7 @@ void UNiagaraStackFilteredObject::AppendEmitterCategory(TSharedPtr<FNiagaraScrip
 		{
 			if (ModuleNode && ModuleNode->ScriptIsValid())
 			{
-				RefreshChildrenForFunctionCall(ModuleNode, ModuleNode, CurrentChildren, NewChildren, NewIssues, true, FText::FromString(ModuleNode->GetFunctionName()));
+				RefreshChildrenForFunctionCall(ModuleNode, ModuleNode, CurrentChildren, NewChildren, NewIssues, true);
 			}
 		}
 	}
diff --git a/Engine/Plugins/FX/Niagara/Source/NiagaraEditor/Private/ViewModels/Stack/NiagaraStackFunctionInputCollection.cpp b/Engine/Plugins/FX/Niagara/Source/NiagaraEditor/Private/ViewModels/Stack/NiagaraStackFunctionInputCollection.cpp
--- a/Engine/Plugins/FX/Niagara/Source/NiagaraEditor/Private/ViewModels/Stack/NiagaraStackFunctionInputCollection.cpp
+++ b/Engine/Plugins/FX/Niagara/Source/NiagaraEditor/Private/ViewModels/Stack/NiagaraStackFunctionInputCollection.cpp
@@ -5,6 +5,7 @@
 #include "EdGraphSchema_Niagara.h"
 #include "NiagaraClipboard.h"
 #include "NiagaraDataInterface.h"
+#include "NiagaraEditorUtilities.h"
 #include "NiagaraGraph.h"
 #include "NiagaraNodeAssignment.h"
 #include "NiagaraNodeFunctionCall.h"
@@ -53,8 +54,8 @@ static TOptional<FFunctionInputSummaryViewKey> GetSummaryViewInputKeyForFunction
 
 static bool ShouldShowInSummaryView(UNiagaraEmitter* Emitter, UNiagaraNodeFunctionCall* FunctionCall, const FNiagaraVariable& InputVariable, const TOptional<FNiagaraVariableMetaData>& InputMetaData)
 {
-	UNiagaraEmitterEditorData* EditorData = Emitter? Cast<UNiagaraEmitterEditorData>(Emitter->GetEditorData()) : nullptr;
-	if (EditorData && InputMetaData.IsSet())
+	const UNiagaraEmitterEditorData* EditorData = Emitter? Cast<UNiagaraEmitterEditorData>(Emitter->GetEditorData()) : nullptr;
+	if (EditorData)
 	{
 		TOptional<FFunctionInputSummaryViewKey> SummaryViewKey = GetSummaryViewInputKeyForFunctionInput(FunctionCall, InputVariable, InputMetaData);
 		if (SummaryViewKey.IsSet())
@@ -195,7 +196,7 @@ TArray<UNiagaraStackFunctionInput*> UNiagaraStackFunctionInputCollection::GetInl
 
 void UNiagaraStackFunctionInputCollection::RefreshChildrenInternal(const TArray<UNiagaraStackEntry*>& CurrentChildren, TArray<UNiagaraStackEntry*>& NewChildren, TArray<FStackIssue>& NewIssues)
 {
-	RefreshChildrenForFunctionCall(ModuleNode, InputFunctionCallNode, CurrentChildren, NewChildren, NewIssues, false, UncategorizedName);
+	RefreshChildrenForFunctionCall(ModuleNode, InputFunctionCallNode, CurrentChildren, NewChildren, NewIssues, false);
 }
 
 
@@ -203,8 +204,16 @@ UNiagaraStackFunctionInputCollectionBase::UNiagaraStackFunctionInputCollectionBa
 {	
 }
 
+
 void UNiagaraStackFunctionInputCollectionBase::RefreshChildrenForFunctionCall(UNiagaraNodeFunctionCall* ModuleNode, UNiagaraNodeFunctionCall* InputFunctionCallNode, 
-	const TArray<UNiagaraStackEntry*>& CurrentChildren, TArray<UNiagaraStackEntry*>& NewChildren, TArray<FStackIssue>& NewIssues, bool bShouldApplySummaryFilter, const FText& BaseCategory)
+	const TArray<UNiagaraStackEntry*>& CurrentChildren, TArray<UNiagaraStackEntry*>& NewChildren, TArray<FStackIssue>& NewIssues, bool bShouldApplySummaryFilter)
+{	
+	FFunctionCallNodesState State;
+	AppendInputsForFunctionCall(State, ModuleNode, InputFunctionCallNode, NewIssues, bShouldApplySummaryFilter);
+	ApplyAllFunctionInputsToChildren(State, CurrentChildren, NewChildren, NewIssues, bShouldApplySummaryFilter);
+}
+
+void UNiagaraStackFunctionInputCollectionBase::AppendInputsForFunctionCall(FFunctionCallNodesState& State, UNiagaraNodeFunctionCall* ModuleNode, UNiagaraNodeFunctionCall* InputFunctionCallNode, TArray<FStackIssue>& NewIssues, bool bShouldApplySummaryFilter)
 {
 	UNiagaraEmitter* Emitter = GetEmitterViewModel().IsValid() ? GetEmitterViewModel()->GetEmitter() : nullptr;
 
@@ -232,9 +241,45 @@ void UNiagaraStackFunctionInputCollectionBase::RefreshChildrenForFunctionCall(UN
 	TArray<const UEdGraphPin*> PinsWithInvalidTypes;
 
 	UNiagaraGraph* InputFunctionGraph = InputFunctionCallNode->GetCalledGraph();
-	TArray<FInputData> InputDataCollection;
-	TMap<FName, FNiagaraParentData> ParentMapping;
+	
+	
+	const auto GetSummaryViewState = [&](UNiagaraEmitter* Emitter, UNiagaraNodeFunctionCall* FunctionCall, const FNiagaraVariable& InputVariable,
+		const TOptional<FNiagaraVariableMetaData>& InputMetaData, FText& InputCategory, TOptional<FText>& DisplayName, int32& EditorSortPriority ) -> bool
+	{
+		const UNiagaraEmitterEditorData* EditorData = Emitter? Cast<UNiagaraEmitterEditorData>(Emitter->GetEditorData()) : nullptr;
+
+		const TOptional<FFunctionInputSummaryViewKey> SummaryViewKey = GetSummaryViewInputKeyForFunctionInput(InputFunctionCallNode, InputVariable, InputMetaData);
+		if (EditorData && bShouldApplySummaryFilter && SummaryViewKey.IsSet())
+		{
+			const FFunctionInputSummaryViewMetadata SummaryViewData = EditorData->GetSummaryViewMetaData(SummaryViewKey.GetValue());
+			if (SummaryViewData.Category != NAME_None)
+			{
+				InputCategory = FText::FromName(SummaryViewData.Category);
+			}
+			else
+			{
+				if (!InputCategory.EqualTo(UncategorizedName))
+				{
+					InputCategory = FText::Format(LOCTEXT("EmitterSummaryNodeCategory", "{0} - {1}"), FText::FromString(*InputFunctionCallNode->GetFunctionName()), InputCategory);
+				}
+				else if (FunctionCall->IsA<UNiagaraNodeAssignment>())
+				{
+					InputCategory = LOCTEXT("EmitterSummaryNodeCategoryAssignment", "Set Parameters");									
+				}
+				else
+				{					
+					InputCategory = FText::FromString(*InputFunctionCallNode->GetFunctionName());
+				}				
+			}
+			EditorSortPriority = SummaryViewData.SortIndex != INDEX_NONE? SummaryViewData.SortIndex : EditorSortPriority;
+			DisplayName = (SummaryViewData.DisplayName != NAME_None) ? FText::FromName(SummaryViewData.DisplayName) : TOptional<FText>();		
+		}
+		
+		return ShouldShowInSummaryView(Emitter, InputFunctionCallNode, InputVariable, InputMetaData);
+	};
 
+
+	
 	// Gather input data
 	for (const UEdGraphPin* InputPin : InputPins)
 	{
@@ -267,44 +312,26 @@ void UNiagaraStackFunctionInputCollectionBase::RefreshChildrenForFunctionCall(UN
 		int32 EditorSortPriority = InputMetaData.IsSet() ? InputMetaData->EditorSortPriority : 0;
 		TOptional<FText> DisplayName;
 
-		UNiagaraEmitterEditorData* EditorData = Emitter? Cast<UNiagaraEmitterEditorData>(Emitter->GetEditorData()) : nullptr;
-
-		const auto& SummaryViewKey = GetSummaryViewInputKeyForFunctionInput(InputFunctionCallNode, InputVariable, InputMetaData);
-		if (EditorData && bShouldApplySummaryFilter && SummaryViewKey.IsSet())
-		{
-			FFunctionInputSummaryViewMetadata SummaryViewData = EditorData->GetSummaryViewMetaData(SummaryViewKey.GetValue());
-			if (SummaryViewData.Category != NAME_None)
-			{
-				InputCategory = FText::FromName(SummaryViewData.Category);
-			}
-			else
-			{
-				InputCategory = InputCategory.EqualTo(UncategorizedName)? FText::FromString(*InputFunctionCallNode->GetFunctionName()) : FText::FromString(*(InputFunctionCallNode->GetFunctionName() + TEXT(" - ") + InputCategory.ToString()));
-			}
-			EditorSortPriority = SummaryViewData.SortIndex != INDEX_NONE? SummaryViewData.SortIndex : EditorSortPriority;
-			DisplayName = (SummaryViewData.DisplayName != NAME_None) ? FText::FromName(SummaryViewData.DisplayName) : TOptional<FText>();		
-		}
-
-		bool bShouldShowInSummary = ShouldShowInSummaryView(Emitter, InputFunctionCallNode, InputVariable, InputMetaData);
+		bool bShouldShowInSummary = GetSummaryViewState(Emitter, InputFunctionCallNode, InputVariable, InputMetaData, InputCategory, DisplayName, EditorSortPriority);		
 		if (bShouldShowInSummary)
 		{
 			SummaryViewPins.Add(InputPin);
 		}		
 		bool bIsInputHidden = HiddenPins.Contains(InputPin);
-		FInputData InputData = { InputPin, InputVariable.GetType(), EditorSortPriority, DisplayName, InputCategory, false, bIsInputHidden, bShouldShowInSummary };
-		int32 Index = InputDataCollection.Add(InputData);
+		FInputData InputData = { InputPin, InputVariable.GetType(), EditorSortPriority, DisplayName, InputCategory, false, bIsInputHidden, bShouldShowInSummary, ModuleNode, InputFunctionCallNode };
+		int32 Index = State.InputDataCollection.Add(InputData);
 
 		// set up the data for the parent-child mapping
 		if (InputMetaData && !InputMetaData->ParentAttribute.IsNone())
 		{
 			if (InputMetaData->ParentAttribute.ToString().StartsWith(PARAM_MAP_MODULE_STR))
 			{
-				ParentMapping.FindOrAdd(InputMetaData->ParentAttribute).ChildIndices.Add(Index);
+				State.ParentMapping.FindOrAdd(InputMetaData->ParentAttribute).ChildIndices.Add(Index);
 			}
 			else
 			{
 				FString NamespacedParent = PARAM_MAP_MODULE_STR + InputMetaData->ParentAttribute.ToString();
-				ParentMapping.FindOrAdd(FName(*NamespacedParent)).ChildIndices.Add(Index);
+				State.ParentMapping.FindOrAdd(FName(*NamespacedParent)).ChildIndices.Add(Index);
 			}
 		}
 	}
@@ -350,60 +377,47 @@ void UNiagaraStackFunctionInputCollectionBase::RefreshChildrenForFunctionCall(UN
 
 		int32 EditorSortPriority = InputMetaData.IsSet() ? InputMetaData->EditorSortPriority : 0;
 		TOptional<FText> DisplayName;
-		
-		UNiagaraEmitterEditorData* EditorData = Emitter? Cast<UNiagaraEmitterEditorData>(Emitter->GetEditorData()) : nullptr;
-
-		const auto& SummaryViewKey = GetSummaryViewInputKeyForFunctionInput(InputFunctionCallNode, InputVariable, InputMetaData);
-		if (EditorData && bShouldApplySummaryFilter && SummaryViewKey.IsSet())
-		{
-			FFunctionInputSummaryViewMetadata SummaryViewData = EditorData->GetSummaryViewMetaData(SummaryViewKey.GetValue());
-			if (SummaryViewData.Category != NAME_None)
-			{
-				InputCategory = FText::FromName(SummaryViewData.Category);
-			}
-			else
-			{
-				InputCategory = InputCategory.EqualTo(UncategorizedName)? FText::FromString(*InputFunctionCallNode->GetFunctionName()) : FText::FromString(*(InputFunctionCallNode->GetFunctionName() + TEXT(" - ") + InputCategory.ToString()));
-			}
-			EditorSortPriority = SummaryViewData.SortIndex != INDEX_NONE? SummaryViewData.SortIndex : EditorSortPriority;
-			DisplayName = (SummaryViewData.DisplayName != NAME_None) ? FText::FromName(SummaryViewData.DisplayName) : TOptional<FText>();		
-		}
-		
-		bool bShouldShowInSummary = ShouldShowInSummaryView(Emitter, InputFunctionCallNode, InputVariable, InputMetaData);
+						
+		bool bShouldShowInSummary = GetSummaryViewState(Emitter, InputFunctionCallNode, InputVariable, InputMetaData, InputCategory, DisplayName, EditorSortPriority);		
 		if (bShouldShowInSummary)
 		{
 			SummaryViewPins.Add(InputPin);
 		}		
 		bool bIsInputHidden = HiddenSwitchPins.Contains(InputPin);
-		FInputData InputData = { InputPin, InputVariable.GetType(), EditorSortPriority, DisplayName, InputCategory, true, bIsInputHidden, bShouldShowInSummary };
-		int32 Index = InputDataCollection.Add(InputData);
+		FInputData InputData = { InputPin, InputVariable.GetType(), EditorSortPriority, DisplayName, InputCategory, true, bIsInputHidden, bShouldShowInSummary, ModuleNode, InputFunctionCallNode };
+		int32 Index = State.InputDataCollection.Add(InputData);
 
 		// set up the data for the parent-child mapping
 		if (InputMetaData)
 		{
-			FNiagaraParentData& ParentData = ParentMapping.FindOrAdd(SwitchPinName);
+			FNiagaraParentData& ParentData = State.ParentMapping.FindOrAdd(SwitchPinName);
 			ParentData.ParentPin = InputPin;
 			if (!InputMetaData->ParentAttribute.IsNone())
 			{
 				if (InputMetaData->ParentAttribute.ToString().StartsWith(PARAM_MAP_MODULE_STR))
 				{
-					ParentMapping.FindOrAdd(InputMetaData->ParentAttribute).ChildIndices.Add(Index);
+					State.ParentMapping.FindOrAdd(InputMetaData->ParentAttribute).ChildIndices.Add(Index);
 				}
 				else
 				{
 					FString NamespacedParent = PARAM_MAP_MODULE_STR + InputMetaData->ParentAttribute.ToString();
-					ParentMapping.FindOrAdd(FName(*NamespacedParent)).ChildIndices.Add(Index);
+					State.ParentMapping.FindOrAdd(FName(*NamespacedParent)).ChildIndices.Add(Index);
 				}
 			}
 		}
 	}
 
+	RefreshIssues(InputFunctionCallNode, DuplicateInputNames, ValidAliasedInputNames, PinsWithInvalidTypes, StaticSwitchInputs, NewIssues);
+}
+
+void UNiagaraStackFunctionInputCollectionBase::ApplyAllFunctionInputsToChildren(FFunctionCallNodesState& State, const TArray<UNiagaraStackEntry*>& CurrentChildren, TArray<UNiagaraStackEntry*>& NewChildren, TArray<FStackIssue>& NewIssues, bool bShouldApplySummaryFilter)
+{
 	// resolve the parent/child relationships
-	for (auto& Entry : ParentMapping)
+	for (auto& Entry : State.ParentMapping)
 	{
 		FNiagaraParentData& Data = Entry.Value;
 		if (Data.ChildIndices.Num() == 0) { continue; }
-		for (FInputData& InputData : InputDataCollection)
+		for (FInputData& InputData : State.InputDataCollection)
 		{
 			if (InputData.Pin != Data.ParentPin) { continue; }
 			if (InputData.bIsChild)
@@ -413,24 +427,23 @@ void UNiagaraStackFunctionInputCollectionBase::RefreshChildrenForFunctionCall(UN
 			}
 			for (int32 ChildIndex : Data.ChildIndices)
 			{
-				if (InputDataCollection[ChildIndex].Children.Num() > 0)
+				if (State.InputDataCollection[ChildIndex].Children.Num() > 0)
 				{
-					AddInvalidChildStackIssue(InputDataCollection[ChildIndex].Pin->PinName, NewIssues);
+					AddInvalidChildStackIssue(State.InputDataCollection[ChildIndex].Pin->PinName, NewIssues);
 					continue;
 				}
-				InputDataCollection[ChildIndex].bIsChild = true;
-				InputDataCollection[ChildIndex].Category = InputData.Category; // children get the parent category to prevent inconsistencies there
-				InputData.Children.Add(&InputDataCollection[ChildIndex]);
+				State.InputDataCollection[ChildIndex].bIsChild = true;
+				State.InputDataCollection[ChildIndex].Category = InputData.Category; // children get the parent category to prevent inconsistencies there
+				InputData.Children.Add(&State.InputDataCollection[ChildIndex]);
 			}
 		}
 	}
-
-
+	
 	if (bShouldApplySummaryFilter)
 	{
 		static TFunction<void(FInputData*)> PropagateSummaryViewShowToChildren = [](FInputData* Input)
 		{		
-			bool bShouldShowChildren = !Input->bIsHidden && Input->bShouldShowInSummary;
+			const bool bShouldShowChildren = !Input->bIsHidden && Input->bShouldShowInSummary;
 
 			Input->bIsHidden = !bShouldShowChildren;
 
@@ -442,7 +455,7 @@ void UNiagaraStackFunctionInputCollectionBase::RefreshChildrenForFunctionCall(UN
 		};
 
 		// Propagate summary visibility to children
-		for (FInputData& InputData : InputDataCollection)
+		for (FInputData& InputData : State.InputDataCollection)
 		{
 			if (!InputData.bIsChild)
 			{
@@ -450,7 +463,6 @@ void UNiagaraStackFunctionInputCollectionBase::RefreshChildrenForFunctionCall(UN
 			}
 		}
 	}
-
 	
 	auto SortPredicate = [](const FInputData& A, const FInputData& B)
 	{
@@ -471,35 +483,36 @@ void UNiagaraStackFunctionInputCollectionBase::RefreshChildrenForFunctionCall(UN
 	};
 
 	// Sort child and parent data separately
-	TArray<FInputData*> ParentDataCollection;
-	for (FInputData& InputData : InputDataCollection)
-	{
+	TArray<FInputData> ParentDataCollection;
+	for (FInputData InputData : State.InputDataCollection)
+	{		
 		if (!InputData.bIsChild)
-		{			
-			ParentDataCollection.Add(&InputData);
-			InputData.Children.Sort(SortPredicate);
+		{
+			InputData.Children.Sort(SortPredicate);			
+			ParentDataCollection.Add(MoveTemp(InputData));
 		}
 	}
 	ParentDataCollection.Sort(SortPredicate);
 
 	// Populate the categories
-	for (FInputData* ParentData : ParentDataCollection)
+	for (FInputData& ParentData : ParentDataCollection)
 	{
-		if (!ParentData->bIsHidden)
-		{
-			AddInputToCategory(ModuleNode, InputFunctionCallNode, *ParentData, CurrentChildren, NewChildren);
-			for (FInputData* ChildData : ParentData->Children)
+		if (!ParentData.bIsHidden)
+		{			
+			AddInputToCategory(ParentData, CurrentChildren, NewChildren);
+			for (FInputData* ChildData : ParentData.Children)
 			{
 				if (!ChildData->bIsHidden)
 				{
-					AddInputToCategory(ModuleNode, InputFunctionCallNode, *ChildData, CurrentChildren, NewChildren);
+					AddInputToCategory(*ChildData, CurrentChildren, NewChildren);
 				}
 			}
 		}
 	}
-	RefreshIssues(InputFunctionCallNode, DuplicateInputNames, ValidAliasedInputNames, PinsWithInvalidTypes, StaticSwitchInputs, NewIssues);
 }
 
+
+
 void UNiagaraStackFunctionInputCollectionBase::RefreshIssues(UNiagaraNodeFunctionCall* InputFunctionCallNode, const TArray<FName>& DuplicateInputNames, 
 	const TArray<FName>& ValidAliasedInputNames, const TArray<const UEdGraphPin*>& PinsWithInvalidTypes, const TMap<FName, UEdGraphPin*>& StaticSwitchInputs, TArray<FStackIssue>& NewIssues)
 {
@@ -666,8 +679,8 @@ void UNiagaraStackFunctionInputCollectionBase::AddInvalidChildStackIssue(FName P
 	OutIssues.Add(InvalidHierarchyWarning);
 }
 
-void UNiagaraStackFunctionInputCollectionBase::AddInputToCategory(UNiagaraNodeFunctionCall* ModuleNode, UNiagaraNodeFunctionCall* InputFunctionCallNode, const FInputData& InputData, const TArray<UNiagaraStackEntry*>& CurrentChildren, TArray<UNiagaraStackEntry*>& NewChildren)
-{
+void UNiagaraStackFunctionInputCollectionBase::AddInputToCategory(const FInputData& InputData, const TArray<UNiagaraStackEntry*>& CurrentChildren, TArray<UNiagaraStackEntry*>& NewChildren)
+{	
 	// Try to find an existing category in the already processed children.
 	UNiagaraStackInputCategory* InputCategory = FindCurrentChildOfTypeByPredicate<UNiagaraStackInputCategory>(NewChildren,
 		[&](UNiagaraStackInputCategory* CurrentCategory) { return CurrentCategory->GetCategoryName().CompareTo(InputData.Category) == 0; });
@@ -682,7 +695,7 @@ void UNiagaraStackFunctionInputCollectionBase::AddInputToCategory(UNiagaraNodeFu
 			// If we don't have a current child for this category make a new one.
 			InputCategory = NewObject<UNiagaraStackInputCategory>(this);
 
-			FString InputCategoryStackEditorDataKey = FString::Printf(TEXT("%s-InputCategory-%s"), *InputFunctionCallNode->NodeGuid.ToString(EGuidFormats::DigitsWithHyphens), *InputData.Category.ToString());
+			FString InputCategoryStackEditorDataKey = FString::Printf(TEXT("%s-InputCategory-%s"), *InputData.InputFunctionCallNode->NodeGuid.ToString(EGuidFormats::DigitsWithHyphens), *InputData.Category.ToString());
 			InputCategory->Initialize(CreateDefaultChildRequiredData(), InputCategoryStackEditorDataKey, InputData.Category, GetOwnerStackItemEditorDataKey());
 		}
 		else
@@ -697,7 +710,7 @@ void UNiagaraStackFunctionInputCollectionBase::AddInputToCategory(UNiagaraNodeFu
 		}
 		NewChildren.Add(InputCategory);
 	}
-	InputCategory->AddInput(ModuleNode, InputFunctionCallNode, InputData.Pin->PinName, InputData.Type, InputData.bIsStatic ? EStackParameterBehavior::Static : EStackParameterBehavior::Dynamic, InputData.DisplayName, InputData.bIsHidden, InputData.bIsChild);
+	InputCategory->AddInput(InputData.ModuleNode, InputData.InputFunctionCallNode, InputData.Pin->PinName, InputData.Type, InputData.bIsStatic ? EStackParameterBehavior::Static : EStackParameterBehavior::Dynamic, InputData.DisplayName, InputData.bIsHidden, InputData.bIsChild);
 
 	
 
diff --git a/Engine/Plugins/FX/Niagara/Source/NiagaraEditor/Private/ViewModels/Stack/NiagaraStackSummaryViewInputCollection.cpp b/Engine/Plugins/FX/Niagara/Source/NiagaraEditor/Private/ViewModels/Stack/NiagaraStackSummaryViewInputCollection.cpp
--- a/Engine/Plugins/FX/Niagara/Source/NiagaraEditor/Private/ViewModels/Stack/NiagaraStackSummaryViewInputCollection.cpp
+++ b/Engine/Plugins/FX/Niagara/Source/NiagaraEditor/Private/ViewModels/Stack/NiagaraStackSummaryViewInputCollection.cpp
@@ -69,21 +69,25 @@ void UNiagaraStackSummaryViewObject::RefreshChildrenInternal(const TArray<UNiaga
 	TSharedPtr<FNiagaraScriptViewModel> ScriptViewModelPinned = ViewModel->GetSharedScriptViewModel();
 	checkf(ScriptViewModelPinned.IsValid(), TEXT("Can not refresh children when the script view model has been deleted."));
 
-	AppendEmitterCategory(ScriptViewModelPinned, ENiagaraScriptUsage::EmitterSpawnScript, FGuid(), NewChildren, CurrentChildren, NewIssues);
-	AppendEmitterCategory(ScriptViewModelPinned, ENiagaraScriptUsage::EmitterUpdateScript, FGuid(), NewChildren, CurrentChildren, NewIssues);
+	FFunctionCallNodesState State;
 	
-	AppendEmitterCategory(ScriptViewModelPinned, ENiagaraScriptUsage::ParticleSpawnScript, FGuid(), NewChildren, CurrentChildren, NewIssues);
-	AppendEmitterCategory(ScriptViewModelPinned, ENiagaraScriptUsage::ParticleUpdateScript, FGuid(), NewChildren, CurrentChildren, NewIssues);
+	AppendEmitterCategory(State, ScriptViewModelPinned, ENiagaraScriptUsage::EmitterSpawnScript, FGuid(), NewIssues);
+	AppendEmitterCategory(State, ScriptViewModelPinned, ENiagaraScriptUsage::EmitterUpdateScript, FGuid(), NewIssues);
+	
+	AppendEmitterCategory(State, ScriptViewModelPinned, ENiagaraScriptUsage::ParticleSpawnScript, FGuid(), NewIssues);
+	AppendEmitterCategory(State, ScriptViewModelPinned, ENiagaraScriptUsage::ParticleUpdateScript, FGuid(), NewIssues);
 
 	for (const FNiagaraEventScriptProperties& EventScriptProperties : ViewModel->GetEmitter()->GetEventHandlers())
 	{
-		AppendEmitterCategory(ScriptViewModelPinned, ENiagaraScriptUsage::ParticleEventScript, EventScriptProperties.Script->GetUsageId(), NewChildren, CurrentChildren, NewIssues);
+		AppendEmitterCategory(State, ScriptViewModelPinned, ENiagaraScriptUsage::ParticleEventScript, EventScriptProperties.Script->GetUsageId(), NewIssues);
 	}
 	
 	for (UNiagaraSimulationStageBase* SimulationStage : GetEmitterViewModel()->GetEmitter()->GetSimulationStages())
 	{
-		AppendEmitterCategory(ScriptViewModelPinned, ENiagaraScriptUsage::ParticleSimulationStageScript, SimulationStage->Script->GetUsageId(), NewChildren, CurrentChildren, NewIssues);
-	}
+		AppendEmitterCategory(State, ScriptViewModelPinned, ENiagaraScriptUsage::ParticleSimulationStageScript, SimulationStage->Script->GetUsageId(), NewIssues);
+	}	
+	
+	ApplyAllFunctionInputsToChildren(State, CurrentChildren, NewChildren, NewIssues, true);
 	
 	Super::RefreshChildrenInternal(CurrentChildren, NewChildren, NewIssues);
 
@@ -102,7 +106,7 @@ void UNiagaraStackSummaryViewObject::RefreshChildrenInternal(const TArray<UNiaga
 	}
 }
 
-void UNiagaraStackSummaryViewObject::AppendEmitterCategory(TSharedPtr<FNiagaraScriptViewModel> ScriptViewModelPinned, ENiagaraScriptUsage ScriptUsage, FGuid ScriptUsageId, TArray<UNiagaraStackEntry*>& NewChildren, const TArray<UNiagaraStackEntry*>& CurrentChildren, TArray<FStackIssue>& NewIssues)
+void UNiagaraStackSummaryViewObject::AppendEmitterCategory(FFunctionCallNodesState& State, TSharedPtr<FNiagaraScriptViewModel> ScriptViewModelPinned, ENiagaraScriptUsage ScriptUsage, FGuid ScriptUsageId, TArray<FStackIssue>& NewIssues)
 {
 	UNiagaraGraph* Graph = ScriptViewModelPinned->GetGraphViewModel()->GetGraph();
 	FText ErrorMessage;
@@ -118,7 +122,7 @@ void UNiagaraStackSummaryViewObject::AppendEmitterCategory(TSharedPtr<FNiagaraSc
 		{
 			if (ModuleNode && ModuleNode->ScriptIsValid())
 			{
-				RefreshChildrenForFunctionCall(ModuleNode, ModuleNode, CurrentChildren, NewChildren, NewIssues, true, FText::FromString(ModuleNode->GetFunctionName()));
+				AppendInputsForFunctionCall(State, ModuleNode, ModuleNode, NewIssues, true);
 			}
 		}
 	}
diff --git a/Engine/Plugins/FX/Niagara/Source/NiagaraEditor/Public/NiagaraEmitterEditorData.h b/Engine/Plugins/FX/Niagara/Source/NiagaraEditor/Public/NiagaraEmitterEditorData.h
--- a/Engine/Plugins/FX/Niagara/Source/NiagaraEditor/Public/NiagaraEmitterEditorData.h
+++ b/Engine/Plugins/FX/Niagara/Source/NiagaraEditor/Public/NiagaraEmitterEditorData.h
@@ -48,7 +48,7 @@ private:
 	uint32 bShowSummaryView : 1;
 
 	/** Stores metadata for filtering function inputs when in Filtered/Simple view. */
-	UPROPERTY(meta=(NiagaraNoMerge))
+	UPROPERTY()
 	TMap<FFunctionInputSummaryViewKey, FFunctionInputSummaryViewMetadata> SummaryViewFunctionInputMetadata;
 
 	FSimpleMulticastDelegate OnSummaryViewStateChangedDelegate;
diff --git a/Engine/Plugins/FX/Niagara/Source/NiagaraEditor/Public/ViewModels/NiagaraEmitterViewModel.h b/Engine/Plugins/FX/Niagara/Source/NiagaraEditor/Public/ViewModels/NiagaraEmitterViewModel.h
--- a/Engine/Plugins/FX/Niagara/Source/NiagaraEditor/Public/ViewModels/NiagaraEmitterViewModel.h
+++ b/Engine/Plugins/FX/Niagara/Source/NiagaraEditor/Public/ViewModels/NiagaraEmitterViewModel.h
@@ -112,7 +112,7 @@ public:
 	FOnScriptParameterStoreChanged& OnScriptParameterStoreChanged();
 
 	/** Gets editor specific data which can be stored per emitter.  If this data hasn't been created the default version will be returned. */
-	const UNiagaraEmitterEditorData& GetEditorData() const;
+	NIAGARAEDITOR_API const UNiagaraEmitterEditorData& GetEditorData() const;
 
 	/** Gets editor specific data which is stored per emitter.  If this data hasn't been created then it will be created. */
 	NIAGARAEDITOR_API UNiagaraEmitterEditorData& GetOrCreateEditorData();
diff --git a/Engine/Plugins/FX/Niagara/Source/NiagaraEditor/Public/ViewModels/Stack/NiagaraStackFunctionInputCollection.h b/Engine/Plugins/FX/Niagara/Source/NiagaraEditor/Public/ViewModels/Stack/NiagaraStackFunctionInputCollection.h
--- a/Engine/Plugins/FX/Niagara/Source/NiagaraEditor/Public/ViewModels/Stack/NiagaraStackFunctionInputCollection.h
+++ b/Engine/Plugins/FX/Niagara/Source/NiagaraEditor/Public/ViewModels/Stack/NiagaraStackFunctionInputCollection.h
@@ -36,8 +36,12 @@ protected:
 		bool bIsHidden;
 		bool bShouldShowInSummary;
 
+		UNiagaraNodeFunctionCall* ModuleNode;
+		UNiagaraNodeFunctionCall* InputFunctionCallNode;
+		
 		TArray<FInputData*> Children;
 		bool bIsChild = false;
+
 	};
 
 	struct FNiagaraParentData
@@ -45,9 +49,19 @@ protected:
 		const UEdGraphPin* ParentPin;
 		TArray<int32> ChildIndices;
 	};
+
+	struct FFunctionCallNodesState
+	{		
+		TArray<FInputData> InputDataCollection;
+		TMap<FName, FNiagaraParentData> ParentMapping;
+	};
 	
 	void RefreshChildrenForFunctionCall(UNiagaraNodeFunctionCall* ModuleNode, UNiagaraNodeFunctionCall* InputFunctionCallNode, const TArray<UNiagaraStackEntry*>& CurrentChildren, 
-		TArray<UNiagaraStackEntry*>& NewChildren, TArray<FStackIssue>& NewIssues, bool bShouldApplySummaryFilter, const FText& BaseCategory);
+		TArray<UNiagaraStackEntry*>& NewChildren, TArray<FStackIssue>& NewIssues, bool bShouldApplySummaryFilter);
+	
+	void AppendInputsForFunctionCall(FFunctionCallNodesState& State, UNiagaraNodeFunctionCall* ModuleNode, UNiagaraNodeFunctionCall* InputFunctionCallNode, TArray<FStackIssue>& NewIssues, bool bShouldApplySummaryFilter);
+	
+	void ApplyAllFunctionInputsToChildren(FFunctionCallNodesState& State, const TArray<UNiagaraStackEntry*>& CurrentChildren, TArray<UNiagaraStackEntry*>& NewChildren, TArray<FStackIssue>& NewIssues, bool bShouldApplySummaryFilter);
 
 	void RefreshIssues(UNiagaraNodeFunctionCall* InputFunctionCallNode, const TArray<FName>& DuplicateInputNames, const TArray<FName>& ValidAliasedInputNames, const TArray<const UEdGraphPin*>& PinsWithInvalidTypes, const TMap<FName, UEdGraphPin*>& StaticSwitchInputs, TArray<FStackIssue>& NewIssues);
 
@@ -62,7 +76,7 @@ protected:
 	void AddInvalidChildStackIssue(FName PinName, TArray<FStackIssue>& OutIssues);
 
 
-	void AddInputToCategory(UNiagaraNodeFunctionCall* ModuleNode, UNiagaraNodeFunctionCall* InputFunctionCallNode, const FInputData& InputData, const TArray<UNiagaraStackEntry*>& CurrentChildren, TArray<UNiagaraStackEntry*>& NewChildren);
+	void AddInputToCategory(const FInputData& InputData, const TArray<UNiagaraStackEntry*>& CurrentChildren, TArray<UNiagaraStackEntry*>& NewChildren);
 
 };
 
diff --git a/Engine/Plugins/FX/Niagara/Source/NiagaraEditor/Public/ViewModels/Stack/NiagaraStackSummaryViewInputCollection.h b/Engine/Plugins/FX/Niagara/Source/NiagaraEditor/Public/ViewModels/Stack/NiagaraStackSummaryViewInputCollection.h
--- a/Engine/Plugins/FX/Niagara/Source/NiagaraEditor/Public/ViewModels/Stack/NiagaraStackSummaryViewInputCollection.h
+++ b/Engine/Plugins/FX/Niagara/Source/NiagaraEditor/Public/ViewModels/Stack/NiagaraStackSummaryViewInputCollection.h
@@ -36,7 +36,7 @@ protected:
 
 	virtual void RefreshChildrenInternal(const TArray<UNiagaraStackEntry*>& CurrentChildren, TArray<UNiagaraStackEntry*>& NewChildren, TArray<FStackIssue>& NewIssues) override;
 
-	void AppendEmitterCategory(TSharedPtr<FNiagaraScriptViewModel> ScriptViewModelPinned, ENiagaraScriptUsage ScriptUsage, FGuid ScriptUsageId, TArray<UNiagaraStackEntry*>& NewChildren, const TArray<UNiagaraStackEntry*>& CurrentChildren, TArray<FStackIssue>& NewIssues);
+	void AppendEmitterCategory(FFunctionCallNodesState& State, TSharedPtr<FNiagaraScriptViewModel> ScriptViewModelPinned, ENiagaraScriptUsage ScriptUsage, FGuid ScriptUsageId, TArray<FStackIssue>& NewIssues);
 
 	virtual void PostRefreshChildrenInternal() override;
 
diff --git a/Engine/Plugins/FX/Niagara/Source/NiagaraEditorWidgets/Private/NiagaraEditorWidgetsUtilities.cpp b/Engine/Plugins/FX/Niagara/Source/NiagaraEditorWidgets/Private/NiagaraEditorWidgetsUtilities.cpp
--- a/Engine/Plugins/FX/Niagara/Source/NiagaraEditorWidgets/Private/NiagaraEditorWidgetsUtilities.cpp
+++ b/Engine/Plugins/FX/Niagara/Source/NiagaraEditorWidgets/Private/NiagaraEditorWidgetsUtilities.cpp
@@ -559,8 +559,7 @@ TOptional<FFunctionInputSummaryViewKey> FNiagaraStackEditorWidgetsUtilities::Get
 
 UNiagaraStackFunctionInput* FNiagaraStackEditorWidgetsUtilities::GetParentInputForSummaryView(UNiagaraStackFunctionInput* FunctionInput)
 {
-	UNiagaraEmitter* Emitter = (FunctionInput && FunctionInput->GetEmitterViewModel())? FunctionInput->GetEmitterViewModel()->GetEmitter() : nullptr;
-	UNiagaraEmitterEditorData* EditorData = Emitter? Cast<UNiagaraEmitterEditorData>(Emitter->GetEditorData()) : nullptr;
+	const UNiagaraEmitterEditorData* EditorData = (FunctionInput && FunctionInput->GetEmitterViewModel())? &FunctionInput->GetEmitterViewModel()->GetEditorData() : nullptr;
 	
 	if (EditorData)
 	{		
diff --git a/Engine/Plugins/FX/Niagara/Source/NiagaraEditorWidgets/Private/SNiagaraOverviewStackNode.cpp b/Engine/Plugins/FX/Niagara/Source/NiagaraEditorWidgets/Private/SNiagaraOverviewStackNode.cpp
--- a/Engine/Plugins/FX/Niagara/Source/NiagaraEditorWidgets/Private/SNiagaraOverviewStackNode.cpp
+++ b/Engine/Plugins/FX/Niagara/Source/NiagaraEditorWidgets/Private/SNiagaraOverviewStackNode.cpp
@@ -610,7 +610,7 @@ EVisibility SNiagaraOverviewStackNode::GetOpenParentEmitterVisibility() const
 
 const FSlateBrush* SNiagaraOverviewStackNode::GetSummaryViewButtonBrush() const
 {
-	UNiagaraEmitterEditorData* EditorData = EmitterHandleViewModelWeak.IsValid()? &EmitterHandleViewModelWeak.Pin()->GetEmitterViewModel()->GetOrCreateEditorData() : nullptr;
+	const UNiagaraEmitterEditorData* EditorData = EmitterHandleViewModelWeak.IsValid()? &EmitterHandleViewModelWeak.Pin()->GetEmitterViewModel()->GetEditorData() : nullptr;
 	if (BottomSummaryExpander->IsHovered())
 	{
 		return EditorData && EditorData->ShouldShowSummaryView()
@@ -627,7 +627,7 @@ const FSlateBrush* SNiagaraOverviewStackNode::GetSummaryViewButtonBrush() const
 
 FText SNiagaraOverviewStackNode::GetSummaryViewCollapseTooltipText() const
 {
-	UNiagaraEmitterEditorData* EditorData = EmitterHandleViewModelWeak.IsValid()? &EmitterHandleViewModelWeak.Pin()->GetEmitterViewModel()->GetOrCreateEditorData() : nullptr;
+	const UNiagaraEmitterEditorData* EditorData = EmitterHandleViewModelWeak.IsValid()? &EmitterHandleViewModelWeak.Pin()->GetEmitterViewModel()->GetEditorData() : nullptr;
 	return EditorData && EditorData->ShouldShowSummaryView()? LOCTEXT("HideAdvancedToolTip", "Show Full Emitter") : LOCTEXT("ShowAdvancedToolTip", "Show Emitter Summary");
 }
 
diff --git a/Engine/Plugins/FX/Niagara/Source/NiagaraEditorWidgets/Private/Stack/SNiagaraStackFunctionInputValue.cpp b/Engine/Plugins/FX/Niagara/Source/NiagaraEditorWidgets/Private/Stack/SNiagaraStackFunctionInputValue.cpp
--- a/Engine/Plugins/FX/Niagara/Source/NiagaraEditorWidgets/Private/Stack/SNiagaraStackFunctionInputValue.cpp
+++ b/Engine/Plugins/FX/Niagara/Source/NiagaraEditorWidgets/Private/Stack/SNiagaraStackFunctionInputValue.cpp
@@ -1362,8 +1362,7 @@ const FSlateBrush* SNiagaraStackFunctionInputValue::GetFilteredViewIcon() const
 
 EVisibility SNiagaraStackFunctionInputValue::GetFilteredViewContextButtonVisibility() const
 {
-	UNiagaraEmitter* Emitter = (FunctionInput && FunctionInput->GetEmitterViewModel())? FunctionInput->GetEmitterViewModel()->GetEmitter() : nullptr;
-	UNiagaraEmitterEditorData* EditorData = Emitter? Cast<UNiagaraEmitterEditorData>(Emitter->GetEditorData()) : nullptr;
+	UNiagaraEmitterEditorData* EditorData = (FunctionInput && FunctionInput->GetEmitterViewModel())? &FunctionInput->GetEmitterViewModel()->GetOrCreateEditorData() : nullptr;
 	
 	if (!EditorData || EditorData->ShouldShowSummaryView() || FunctionInput->GetTypedOuter<UNiagaraStackSummaryViewObject>() == nullptr)
 	{
@@ -1439,15 +1438,11 @@ TSharedRef<SWidget> SNiagaraStackFunctionInputValue::GetFilteredViewPropertiesCo
 		SNew(SBox)
 		.WidthOverride(100)
 		.Padding(FMargin(5, 0, 0, 0))
-		[
-			SNew(SNumericEntryBox<int32>)
-			//.Style(FNiagaraEditorStyle::Get(), "NiagaraEditor.ParameterSpinBox")
-			//.Font(FNiagaraEditorStyle::Get().GetFontStyle("NiagaraEditor.ParameterFont"))
-			.MinValue(0)
-			.MaxValue(TOptional<int32>())
-			.Value(this, &SNiagaraStackFunctionInputValue::GetFilteredViewSortIndex)
-			.OnValueChanged(this, &SNiagaraStackFunctionInputValue::FilteredViewSortIndexChanged)
-			.OnValueCommitted(this, &SNiagaraStackFunctionInputValue::FilteredViewSortIndexCommitted)
+		[			
+			SNew(SEditableTextBox)
+			.Text(this, &SNiagaraStackFunctionInputValue::GetFilteredViewSortIndex)
+			.OnVerifyTextChanged(this, &SNiagaraStackFunctionInputValue::VerifyFilteredSortIndex)
+			.OnTextCommitted(this, &SNiagaraStackFunctionInputValue::FilteredSortIndexTextCommitted)
 		];
 	MenuBuilder.AddWidget(SortIndexWidget, LOCTEXT("FilteredViewSortIndex", "Sort Index"));
 
@@ -1458,7 +1453,7 @@ TSharedRef<SWidget> SNiagaraStackFunctionInputValue::GetFilteredViewPropertiesCo
 
 FText SNiagaraStackFunctionInputValue::GetFilteredViewDisplayName() const
 {
-	UNiagaraEmitterEditorData* EditorData = FunctionInput->GetEmitterViewModel()? &FunctionInput->GetEmitterViewModel()->GetOrCreateEditorData() : nullptr;	
+	const UNiagaraEmitterEditorData* EditorData = (FunctionInput && FunctionInput->GetEmitterViewModel())? &FunctionInput->GetEmitterViewModel()->GetEditorData() : nullptr;
 	UNiagaraStackFunctionInput* ParentInput = FNiagaraStackEditorWidgetsUtilities::GetParentInputForSummaryView(FunctionInput);
 	TOptional<FFunctionInputSummaryViewKey> Key = FNiagaraStackEditorWidgetsUtilities::GetSummaryViewInputKeyForFunctionInput(ParentInput);
 
@@ -1510,7 +1505,7 @@ FText GetFunctionInputCategory(UNiagaraStackFunctionInput* Input)
 
 FText SNiagaraStackFunctionInputValue::GetFilteredViewCategory() const
 {
-	UNiagaraEmitterEditorData* EditorData = FunctionInput->GetEmitterViewModel()? &FunctionInput->GetEmitterViewModel()->GetOrCreateEditorData() : nullptr;	
+	const UNiagaraEmitterEditorData* EditorData = FunctionInput->GetEmitterViewModel()? &FunctionInput->GetEmitterViewModel()->GetEditorData() : nullptr;	
 	UNiagaraStackFunctionInput* ParentInput = FNiagaraStackEditorWidgetsUtilities::GetParentInputForSummaryView(FunctionInput);
 	TOptional<FFunctionInputSummaryViewKey> Key = FNiagaraStackEditorWidgetsUtilities::GetSummaryViewInputKeyForFunctionInput(ParentInput);
 
@@ -1554,7 +1549,7 @@ void SNiagaraStackFunctionInputValue::FilteredViewCategoryTextCommitted(const FT
 }
 
 
-TOptional<int32> SNiagaraStackFunctionInputValue::GetFilteredViewSortIndex() const
+FText SNiagaraStackFunctionInputValue::GetFilteredViewSortIndex() const
 {
 	UNiagaraEmitterEditorData* EditorData = FunctionInput->GetEmitterViewModel()? &FunctionInput->GetEmitterViewModel()->GetOrCreateEditorData() : nullptr;	
 	UNiagaraStackFunctionInput* ParentInput = FNiagaraStackEditorWidgetsUtilities::GetParentInputForSummaryView(FunctionInput);
@@ -1563,39 +1558,43 @@ TOptional<int32> SNiagaraStackFunctionInputValue::GetFilteredViewSortIndex() con
 	if (EditorData && Key.IsSet())
 	{
 		int32 SortIndex = EditorData->GetSummaryViewMetaData(Key.GetValue()).SortIndex;
-		return SortIndex == INDEX_NONE? TOptional<int32>() : SortIndex;
+		return SortIndex != INDEX_NONE? FText::FromString(FString::FromInt(SortIndex)) : FText::GetEmpty();
 	}
 	
-	return TOptional<int32>();
+	return FText::GetEmpty();
 }
 
-void SNiagaraStackFunctionInputValue::FilteredViewSortIndexChanged(int32 Value)
+bool SNiagaraStackFunctionInputValue::VerifyFilteredSortIndex(const FText& InText, FText& OutErrorMessage) const
+{
+	if (!InText.IsEmptyOrWhitespace() && !InText.IsNumeric())
+	{
+		OutErrorMessage = LOCTEXT("FilteredSortIndexNotValid", "Sort Index must be empty or a valid whole number");
+		return false;		
+	}
+	return true;
+}
+
+void SNiagaraStackFunctionInputValue::FilteredSortIndexTextCommitted(const FText& Text, ETextCommit::Type CommitType)
 {
 	UNiagaraEmitterEditorData* EditorData = FunctionInput->GetEmitterViewModel()? &FunctionInput->GetEmitterViewModel()->GetOrCreateEditorData() : nullptr;	
 	UNiagaraStackFunctionInput* ParentInput = FNiagaraStackEditorWidgetsUtilities::GetParentInputForSummaryView(FunctionInput);
 	TOptional<FFunctionInputSummaryViewKey> Key = FNiagaraStackEditorWidgetsUtilities::GetSummaryViewInputKeyForFunctionInput(ParentInput);
+
+	int32 NewValue = !Text.IsEmptyOrWhitespace() && Text.IsNumeric()? FCString::Atoi(*Text.ToString()) : INDEX_NONE;
 	
 	if (EditorData && Key.IsSet())
-	{
+	{		
 		FFunctionInputSummaryViewMetadata SummaryViewMetaData = EditorData->GetSummaryViewMetaData(Key.GetValue());		
-		if (Value != SummaryViewMetaData.SortIndex)
+		if (NewValue != SummaryViewMetaData.SortIndex)
 		{
 			FScopedTransaction ScopedTransaction(FText::Format(LOCTEXT("SummaryViewChangedInputSortIndex", "Changed summary view sorty index for {0}"), FunctionInput->GetDisplayName()));
 			EditorData->Modify();
 			
-			SummaryViewMetaData.SortIndex = Value;
+			SummaryViewMetaData.SortIndex = NewValue;
 			EditorData->SetSummaryViewMetaData(Key.GetValue(), SummaryViewMetaData);
 		}		
 	}
 }
 
-void SNiagaraStackFunctionInputValue::FilteredViewSortIndexCommitted(int32 Value, ETextCommit::Type CommitInfo)
-{
-	if (CommitInfo == ETextCommit::OnEnter || CommitInfo == ETextCommit::OnUserMovedFocus)
-	{
-		FilteredViewSortIndexChanged(Value);
-	}
-}
-
 
 #undef LOCTEXT_NAMESPACE
diff --git a/Engine/Plugins/FX/Niagara/Source/NiagaraEditorWidgets/Private/Stack/SNiagaraStackFunctionInputValue.h b/Engine/Plugins/FX/Niagara/Source/NiagaraEditorWidgets/Private/Stack/SNiagaraStackFunctionInputValue.h
--- a/Engine/Plugins/FX/Niagara/Source/NiagaraEditorWidgets/Private/Stack/SNiagaraStackFunctionInputValue.h
+++ b/Engine/Plugins/FX/Niagara/Source/NiagaraEditorWidgets/Private/Stack/SNiagaraStackFunctionInputValue.h
@@ -172,9 +172,9 @@ private:
 	bool VerifyFilteredViewCategory(const FText& InText, FText& OutErrorMessage) const;
 	void FilteredViewCategoryTextCommitted(const FText& Text, ETextCommit::Type CommitType);
 
-	TOptional<int32> GetFilteredViewSortIndex() const;
-	void FilteredViewSortIndexChanged(int32 Value);
-	void FilteredViewSortIndexCommitted(int32 Value, ETextCommit::Type CommitInfo);
+	FText GetFilteredViewSortIndex() const;
+	bool VerifyFilteredSortIndex(const FText& InText, FText& OutErrorMessage) const;
+	void FilteredSortIndexTextCommitted(const FText& Text, ETextCommit::Type CommitType);
 
 private:
 	UNiagaraStackFunctionInput* FunctionInput;
diff --git a/Engine/Plugins/FX/Niagara/Source/NiagaraEditorWidgets/Private/Stack/SNiagaraStackTableRow.cpp b/Engine/Plugins/FX/Niagara/Source/NiagaraEditorWidgets/Private/Stack/SNiagaraStackTableRow.cpp
--- a/Engine/Plugins/FX/Niagara/Source/NiagaraEditorWidgets/Private/Stack/SNiagaraStackTableRow.cpp
+++ b/Engine/Plugins/FX/Niagara/Source/NiagaraEditorWidgets/Private/Stack/SNiagaraStackTableRow.cpp
@@ -620,9 +620,7 @@ void SNiagaraStackTableRow::NavigateTo(UNiagaraStackEntry* Item)
 bool SNiagaraStackTableRow::IsValidForSummaryView() const
 {
 	UNiagaraStackFunctionInput* FunctionInput = Cast<UNiagaraStackFunctionInput>(StackEntry);
-
-	UNiagaraEmitter* Emitter = (FunctionInput && FunctionInput->GetEmitterViewModel())? FunctionInput->GetEmitterViewModel()->GetEmitter() : nullptr;
-	UNiagaraEmitterEditorData* EditorData = Emitter? Cast<UNiagaraEmitterEditorData>(Emitter->GetEditorData()) : nullptr;
+	const UNiagaraEmitterEditorData* EditorData = (FunctionInput && FunctionInput->GetEmitterViewModel())? &FunctionInput->GetEmitterViewModel()->GetEditorData() : nullptr;
 
 	if (EditorData)
 	{
@@ -635,8 +633,7 @@ bool SNiagaraStackTableRow::IsValidForSummaryView() const
 void SNiagaraStackTableRow::ToggleShowInSummaryView()
 {
 	UNiagaraStackFunctionInput* FunctionInput = Cast<UNiagaraStackFunctionInput>(StackEntry);
-	UNiagaraEmitter* Emitter = (FunctionInput && FunctionInput->GetEmitterViewModel())? FunctionInput->GetEmitterViewModel()->GetEmitter() : nullptr;
-	UNiagaraEmitterEditorData* EditorData = Emitter? Cast<UNiagaraEmitterEditorData>(Emitter->GetEditorData()) : nullptr;
+	UNiagaraEmitterEditorData* EditorData = (FunctionInput && FunctionInput->GetEmitterViewModel())? &FunctionInput->GetEmitterViewModel()->GetOrCreateEditorData() : nullptr;
 	
 	if (EditorData)
 	{
@@ -659,8 +656,7 @@ void SNiagaraStackTableRow::ToggleShowInSummaryView()
 bool SNiagaraStackTableRow::ShouldShowInSummaryView() const
 {
 	UNiagaraStackFunctionInput* FunctionInput = Cast<UNiagaraStackFunctionInput>(StackEntry);
-	UNiagaraEmitter* Emitter = (FunctionInput && FunctionInput->GetEmitterViewModel())? FunctionInput->GetEmitterViewModel()->GetEmitter() : nullptr;
-	UNiagaraEmitterEditorData* EditorData = Emitter? Cast<UNiagaraEmitterEditorData>(Emitter->GetEditorData()) : nullptr;
+	const UNiagaraEmitterEditorData* EditorData = (FunctionInput && FunctionInput->GetEmitterViewModel())? &FunctionInput->GetEmitterViewModel()->GetEditorData() : nullptr;
 	
 	if (EditorData)
 	{
