- Use NameToDisplayString on data interface names to be consistent with node title
- Fix up various FName -> FString conversions that are not required

diff --git a/Engine/Plugins/FX/Niagara/Source/Niagara/Private/DataInterface/NiagaraDataInterfaceDebugDraw.cpp b/Engine/Plugins/FX/Niagara/Source/Niagara/Private/DataInterface/NiagaraDataInterfaceDebugDraw.cpp
--- a/Engine/Plugins/FX/Niagara/Source/Niagara/Private/DataInterface/NiagaraDataInterfaceDebugDraw.cpp
+++ b/Engine/Plugins/FX/Niagara/Source/Niagara/Private/DataInterface/NiagaraDataInterfaceDebugDraw.cpp
@@ -2803,7 +2803,7 @@ bool UNiagaraDataInterfaceDebugDraw::UpgradeFunctionCall(FNiagaraFunctionSignatu
 		GetFunctions(Sigs);
 		for (const FNiagaraFunctionSignature& Sig : Sigs)
 		{
-			if (Sig.GetName() == FunctionSignature.GetName())
+			if (Sig.Name == FunctionSignature.Name)
 			{
 				FunctionSignature.Inputs = Sig.Inputs;
 				bWasChanged = true;
diff --git a/Engine/Plugins/FX/Niagara/Source/Niagara/Private/NiagaraDataInterface.cpp b/Engine/Plugins/FX/Niagara/Source/Niagara/Private/NiagaraDataInterface.cpp
--- a/Engine/Plugins/FX/Niagara/Source/Niagara/Private/NiagaraDataInterface.cpp
+++ b/Engine/Plugins/FX/Niagara/Source/Niagara/Private/NiagaraDataInterface.cpp
@@ -184,14 +184,14 @@ void UNiagaraDataInterface::ValidateFunction(const FNiagaraFunctionSignature& Fu
 	{
 		//We couldn't find this signature in the list of available functions.
 		//Lets try to find one with the same name whose parameters may have changed.
-		int32 ExistingSigIdx = DIFuncs.IndexOfByPredicate([&](const FNiagaraFunctionSignature& Sig) { return Sig.GetName() == Function.GetName(); });;
+		int32 ExistingSigIdx = DIFuncs.IndexOfByPredicate([&](const FNiagaraFunctionSignature& Sig) { return Sig.Name == Function.Name; });;
 		if (ExistingSigIdx != INDEX_NONE)
 		{
-			OutValidationErrors.Add(FText::Format(LOCTEXT("DI Function Parameter Mismatch!", "Data Interface function called but it's parameters do not match any available function!\nThe API for this data interface function has likely changed and you need to update your graphs.\nInterface: {0}\nFunction: {1}\n"), FText::FromString(GetClass()->GetName()), FText::FromString(Function.GetName())));
+			OutValidationErrors.Add(FText::Format(LOCTEXT("DI Function Parameter Mismatch!", "Data Interface function called but it's parameters do not match any available function!\nThe API for this data interface function has likely changed and you need to update your graphs.\nInterface: {0}\nFunction: {1}\n"), FText::FromString(GetClass()->GetName()), FText::FromName(Function.Name)));
 		}
 		else
 		{
-			OutValidationErrors.Add(FText::Format(LOCTEXT("Unknown DI Function", "Unknown Data Interface function called!\nThe API for this data interface has likely changed and you need to update your graphs.\nInterface: {0}\nFunction: {1}\n"), FText::FromString(GetClass()->GetName()), FText::FromString(Function.GetName())));
+			OutValidationErrors.Add(FText::Format(LOCTEXT("Unknown DI Function", "Unknown Data Interface function called!\nThe API for this data interface has likely changed and you need to update your graphs.\nInterface: {0}\nFunction: {1}\n"), FText::FromString(GetClass()->GetName()), FText::FromName(Function.Name)));
 		}
 	}
 }
diff --git a/Engine/Plugins/FX/Niagara/Source/Niagara/Private/NiagaraDataInterfaceRW.cpp b/Engine/Plugins/FX/Niagara/Source/Niagara/Private/NiagaraDataInterfaceRW.cpp
--- a/Engine/Plugins/FX/Niagara/Source/Niagara/Private/NiagaraDataInterfaceRW.cpp
+++ b/Engine/Plugins/FX/Niagara/Source/Niagara/Private/NiagaraDataInterfaceRW.cpp
@@ -501,7 +501,7 @@ void UNiagaraDataInterfaceGrid2D::ValidateFunction(const FNiagaraFunctionSignatu
 	DeprecatedFunctionNames.Add(UNiagaraDataInterfaceRWBase::WorldBBoxSizeFunctionName);
 	DeprecatedFunctionNames.Add(UNiagaraDataInterfaceRWBase::CellSizeFunctionName);
 
-	if (DIFuncs.Contains(Function) && DeprecatedFunctionNames.Contains(FName(Function.GetName())))
+	if (DIFuncs.Contains(Function) && DeprecatedFunctionNames.Contains(Function.Name))
 	{
 		// #TODO(dmp): add validation warnings that aren't as strict as these errors
 		// OutValidationErrors.Add(FText::Format(LOCTEXT("Grid2DDeprecationMsgFmt", "Grid2D DI Function {0} has been deprecated. Specify grid size on your emitter.\n"), FText::FromString(Function.GetName())));	
diff --git a/Engine/Plugins/FX/Niagara/Source/Niagara/Private/NiagaraDataInterfaceSkeletalMesh.cpp b/Engine/Plugins/FX/Niagara/Source/Niagara/Private/NiagaraDataInterfaceSkeletalMesh.cpp
--- a/Engine/Plugins/FX/Niagara/Source/Niagara/Private/NiagaraDataInterfaceSkeletalMesh.cpp
+++ b/Engine/Plugins/FX/Niagara/Source/Niagara/Private/NiagaraDataInterfaceSkeletalMesh.cpp
@@ -2914,7 +2914,7 @@ void UNiagaraDataInterfaceSkeletalMesh::ValidateFunction(const FNiagaraFunctionS
 
 		if (SkinnedDataDeprecatedFunctions.Contains(Function))
 		{
-			OutValidationErrors.Add(FText::Format(LOCTEXT("SkinnedDataFunctionDeprecationMsgFmt", "Skeletal Mesh DI Function {0} has been deprecated. Use GetSinnedTriangleData or GetSkinnedTriangleDataWS instead.\n"), FText::FromString(Function.GetName())));
+			OutValidationErrors.Add(FText::Format(LOCTEXT("SkinnedDataFunctionDeprecationMsgFmt", "Skeletal Mesh DI Function {0} has been deprecated. Use GetSinnedTriangleData or GetSkinnedTriangleDataWS instead.\n"), FText::FromName(Function.Name)));
 		}
 		else
 		{
diff --git a/Engine/Plugins/FX/Niagara/Source/Niagara/Public/NiagaraCommon.h b/Engine/Plugins/FX/Niagara/Source/Niagara/Public/NiagaraCommon.h
--- a/Engine/Plugins/FX/Niagara/Source/Niagara/Public/NiagaraCommon.h
+++ b/Engine/Plugins/FX/Niagara/Source/Niagara/Public/NiagaraCommon.h
@@ -500,7 +500,7 @@ struct NIAGARA_API FNiagaraFunctionSignature
 		return bMatches;
 	}
 
-	FString GetName()const { return Name.ToString(); }
+	FString GetNameString() const { return Name.ToString(); }
 
 	void AddInput(FNiagaraVariable InputVar, FText Tooltip = FText())
 	{
diff --git a/Engine/Plugins/FX/Niagara/Source/NiagaraEditor/Private/EdGraphSchema_Niagara.cpp b/Engine/Plugins/FX/Niagara/Source/NiagaraEditor/Private/EdGraphSchema_Niagara.cpp
--- a/Engine/Plugins/FX/Niagara/Source/NiagaraEditor/Private/EdGraphSchema_Niagara.cpp
+++ b/Engine/Plugins/FX/Niagara/Source/NiagaraEditor/Private/EdGraphSchema_Niagara.cpp
@@ -697,7 +697,7 @@ TArray<TSharedPtr<FNiagaraAction_NewNode>> UEdGraphSchema_Niagara::GetGraphActio
 					continue;
 
 				UNiagaraNodeFunctionCall* FuncNode = NewObject<UNiagaraNodeFunctionCall>(OwnerOfTemporaries);
-				AddNewNodeMenuAction(NewActions, FuncNode, FText::FromString(Sig.GetName()), ENiagaraMenuSections::General, {MenuCat.ToString()}, FText::GetEmpty(), FText::GetEmpty());
+				AddNewNodeMenuAction(NewActions, FuncNode, FText::FromString(FName::NameToDisplayString(Sig.GetNameString(), false)), ENiagaraMenuSections::General, {MenuCat.ToString()}, FText::GetEmpty(), FText::GetEmpty());
 				FuncNode->Signature = Sig;
 			}
 		}
@@ -2003,7 +2003,7 @@ void UEdGraphSchema_Niagara::GenerateDataInterfacePinMenu(UToolMenu* ToolMenu, c
 		if ( Prototype.Len() > 0 )
 		{
 			Prototype.Append(TEXT("\r\n"));
-			FunctionPrototypes.Emplace(FunctionSignature.GetName(), Prototype);
+			FunctionPrototypes.Emplace(FunctionSignature.GetNameString(), Prototype);
 		}
 	}
 
diff --git a/Engine/Plugins/FX/Niagara/Source/NiagaraEditor/Private/NiagaraCompiler.cpp b/Engine/Plugins/FX/Niagara/Source/NiagaraEditor/Private/NiagaraCompiler.cpp
--- a/Engine/Plugins/FX/Niagara/Source/NiagaraEditor/Private/NiagaraCompiler.cpp
+++ b/Engine/Plugins/FX/Niagara/Source/NiagaraEditor/Private/NiagaraCompiler.cpp
@@ -1921,7 +1921,7 @@ TOptional<FNiagaraCompileResults> FHlslNiagaraCompiler::GetCompileResult(int32 J
 			if (Sig)
 			{
 				int32 NewBindingIdx = Results.Data->CalledVMExternalFunctions.AddDefaulted();
-				Results.Data->CalledVMExternalFunctions[NewBindingIdx].Name = *Sig->GetName();
+				Results.Data->CalledVMExternalFunctions[NewBindingIdx].Name = Sig->Name;
 				Results.Data->CalledVMExternalFunctions[NewBindingIdx].OwnerName = Sig->OwnerName;
 				Results.Data->CalledVMExternalFunctions[NewBindingIdx].InputParamLocations = FuncInfo.InputParamLocations;
 				Results.Data->CalledVMExternalFunctions[NewBindingIdx].NumOutputs = FuncInfo.NumOutputs;
diff --git a/Engine/Plugins/FX/Niagara/Source/NiagaraEditor/Private/NiagaraHlslTranslator.cpp b/Engine/Plugins/FX/Niagara/Source/NiagaraEditor/Private/NiagaraHlslTranslator.cpp
--- a/Engine/Plugins/FX/Niagara/Source/NiagaraEditor/Private/NiagaraHlslTranslator.cpp
+++ b/Engine/Plugins/FX/Niagara/Source/NiagaraEditor/Private/NiagaraHlslTranslator.cpp
@@ -7307,7 +7307,7 @@ void FHlslNiagaraTranslator::ProcessCustomHlsl(const FString& InCustomHlsl, ENia
 			SanitizedFunctionNames.Reserve(Funcs.Num());
 			for (const FNiagaraFunctionSignature& FunctionSignature : Funcs)
 			{
-				SanitizedFunctionNames.Emplace(GetSanitizedDIFunctionName(FunctionSignature.GetName()));
+				SanitizedFunctionNames.Emplace(GetSanitizedDIFunctionName(FunctionSignature.GetNameString()));
 			}
 
 			bool bPermuteSignatureByDataInterface = false;
@@ -7980,7 +7980,7 @@ void FHlslNiagaraTranslator::GenerateFunctionCall(ENiagaraScriptUsage ScriptUsag
 	if (ScriptUsage == ENiagaraScriptUsage::Module)
 	{
 		bEnteredStatScope = true;
-		EnterStatsScope(FNiagaraStatScope(*GetFunctionSignatureSymbol(FunctionSignature), *(FunctionSignature.GetName())));
+		EnterStatsScope(FNiagaraStatScope(*GetFunctionSignatureSymbol(FunctionSignature), *(FunctionSignature.GetNameString())));
 	}
 
 	TArray<FString> MissingParameters;
@@ -8138,7 +8138,7 @@ void FHlslNiagaraTranslator::GenerateFunctionCall(ENiagaraScriptUsage ScriptUsag
 
 FString FHlslNiagaraTranslator::GetFunctionSignatureSymbol(const FNiagaraFunctionSignature& Sig)
 {
-	FString SigStr = Sig.GetName();
+	FString SigStr = Sig.GetNameString();
 	if (!Sig.OwnerName.IsNone() && Sig.OwnerName.IsValid())
 	{
 		SigStr += TEXT("_") + Sig.OwnerName.ToString().Replace(TEXT("."), TEXT("_"));;
diff --git a/Engine/Plugins/FX/Niagara/Source/NiagaraEditor/Private/NiagaraNodeFunctionCall.cpp b/Engine/Plugins/FX/Niagara/Source/NiagaraEditor/Private/NiagaraNodeFunctionCall.cpp
--- a/Engine/Plugins/FX/Niagara/Source/NiagaraEditor/Private/NiagaraNodeFunctionCall.cpp
+++ b/Engine/Plugins/FX/Niagara/Source/NiagaraEditor/Private/NiagaraNodeFunctionCall.cpp
@@ -616,7 +616,7 @@ void UNiagaraNodeFunctionCall::UpdateInputNameBinding(const FGuid& BoundVariable
 
 FText UNiagaraNodeFunctionCall::GetNodeTitle(ENodeTitleType::Type TitleType) const
 {
-	FString DetectedName = FunctionScript ? FunctionScript->GetName() : Signature.GetName();
+	FString DetectedName = FunctionScript ? FunctionScript->GetName() : Signature.GetNameString();
 	if (DetectedName.IsEmpty())
 	{
 		return FText::FromString(TEXT("Missing ( Was\"") + FunctionDisplayName + TEXT("\")"));
@@ -1579,7 +1579,7 @@ void UNiagaraNodeFunctionCall::AutowireNewNode(UEdGraphPin* FromPin)
 
 void UNiagaraNodeFunctionCall::ComputeNodeName(FString SuggestedName, bool bForceSuggestion)
 {
-	FString FunctionName = FunctionScript ? FunctionScript->GetName() : Signature.GetName();
+	FString FunctionName = FunctionScript ? FunctionScript->GetName() : Signature.GetNameString();
 	FName ProposedName;
 	if (SuggestedName.IsEmpty() == false)
 	{ 
