Fixed deprecation properties not properly working using EditCondition.
This is because the versioned script data properties were added individually to the script details.

diff --git a/Engine/Plugins/FX/Niagara/Source/NiagaraEditor/Private/Customizations/NiagaraScriptDetails.cpp b/Engine/Plugins/FX/Niagara/Source/NiagaraEditor/Private/Customizations/NiagaraScriptDetails.cpp
--- a/Engine/Plugins/FX/Niagara/Source/NiagaraEditor/Private/Customizations/NiagaraScriptDetails.cpp
+++ b/Engine/Plugins/FX/Niagara/Source/NiagaraEditor/Private/Customizations/NiagaraScriptDetails.cpp
@@ -188,20 +188,17 @@ void FNiagaraScriptDetails::CustomizeDetails(IDetailLayoutBuilder& DetailBuilder
 
 		if (FVersionedNiagaraScriptData* Data = StandaloneScript.Script->GetScriptData(StandaloneScript.Version))
 		{
-			for (FProperty* ChildProperty : TFieldRange<FProperty>(FVersionedNiagaraScriptData::StaticStruct()))
+			TSharedPtr<FStructOnScope> StructData = MakeShareable(new FStructOnScope(FVersionedNiagaraScriptData::StaticStruct(), (uint8*)Data));
+			TArray<TSharedPtr<IPropertyHandle>> NewHandles = CategoryBuilder.AddAllExternalStructureProperties(StructData.ToSharedRef());
+
+			for(TSharedPtr<IPropertyHandle> NewHandle : NewHandles)
 			{
-				if (ChildProperty->HasAllPropertyFlags(CPF_Edit))
+				NewHandle->SetOnPropertyValueChanged((FSimpleDelegate::CreateLambda([StandaloneScript, NewHandle]()
 				{
-					TSharedPtr<FStructOnScope> StructData = MakeShareable(new FStructOnScope(FVersionedNiagaraScriptData::StaticStruct(), (uint8*)Data));
-					IDetailPropertyRow* PropertyRow = CategoryBuilder.AddExternalStructureProperty(StructData, ChildProperty->GetFName());
-					PropertyRow->GetPropertyHandle()->SetOnPropertyValueChanged(FSimpleDelegate::CreateLambda([StandaloneScript, ChildProperty]()
-					{
-						FPropertyChangedEvent ChangeEvent(ChildProperty);
-						StandaloneScript.Script->PostEditChangeVersionedProperty(ChangeEvent, StandaloneScript.Version);
-					}));
-				}
+					FPropertyChangedEvent ChangeEvent(NewHandle->GetProperty());
+					StandaloneScript.Script->PostEditChangeVersionedProperty(ChangeEvent, StandaloneScript.Version);
+				})));
 			}
-			
 		}		
 	}
 
