Niagara - Initialize the bNeedsUIResync flag on the system instance to false so that it doesn't rebuild the component UI every time a user parameter is changed.  This along with a core editor fix allows the user to edit a color user parameter without the color picker disappearing and also allows changes after the first change to actually work.

diff --git a/Engine/Plugins/FX/Niagara/Source/Niagara/Private/NiagaraSystemInstance.cpp b/Engine/Plugins/FX/Niagara/Source/Niagara/Private/NiagaraSystemInstance.cpp
--- a/Engine/Plugins/FX/Niagara/Source/Niagara/Private/NiagaraSystemInstance.cpp
+++ b/Engine/Plugins/FX/Niagara/Source/Niagara/Private/NiagaraSystemInstance.cpp
@@ -133,6 +133,9 @@ FNiagaraSystemInstance::FNiagaraSystemInstance(UWorld& InWorld, UNiagaraSystem&
 	, bLODDistanceIsValid(false)
 	, bLODDistanceIsOverridden(false)
 	, bPooled(bInPooled)
+#if WITH_EDITOR
+	, bNeedsUIResync(false)
+#endif
 	, CachedDeltaSeconds(0.0f)
 	, TimeSinceLastForceUpdateTransform(0.0f)
 	, FixedBounds_GT(EForceInit::ForceInit)
