Standardized scratch pad apply action
Also removed unneeded function

diff --git a/Engine/Plugins/FX/Niagara/Source/NiagaraEditor/Private/ViewModels/NiagaraScratchPadScriptViewModel.cpp b/Engine/Plugins/FX/Niagara/Source/NiagaraEditor/Private/ViewModels/NiagaraScratchPadScriptViewModel.cpp
--- a/Engine/Plugins/FX/Niagara/Source/NiagaraEditor/Private/ViewModels/NiagaraScratchPadScriptViewModel.cpp
+++ b/Engine/Plugins/FX/Niagara/Source/NiagaraEditor/Private/ViewModels/NiagaraScratchPadScriptViewModel.cpp
@@ -177,14 +177,10 @@ void FNiagaraScratchPadScriptViewModel::ApplyChanges()
 		OriginalScript->GetClass());
 	bHasPendingChanges = false;
 
-	TArray<UNiagaraNodeFunctionCall*> FunctionCallNodesToRefresh;
-	FNiagaraEditorUtilities::GetReferencingFunctionCallNodes(OriginalScript, FunctionCallNodesToRefresh);
-
-	for (UNiagaraNodeFunctionCall* FunctionCallNodeToRefresh : FunctionCallNodesToRefresh)
-	{
-		FunctionCallNodeToRefresh->RefreshFromExternalChanges();
-		FunctionCallNodeToRefresh->MarkNodeRequiresSynchronization(TEXT("ScratchPadChangesApplied"), true);
-	}
+	FRefreshAllScriptsFromExternalChangesArgs Args;
+	Args.OriginatingScript = OriginalScript;
+	Args.OriginatingGraph = CastChecked<UNiagaraScriptSource>(OriginalScript->GetSource(EditScript.Version))->NodeGraph;
+	FNiagaraEditorUtilities::RefreshAllScriptsFromExternalChanges(Args);
 
 	OnHasUnappliedChangesChangedDelegate.Broadcast();
 	OnChangesAppliedDelegate.Broadcast();
diff --git a/Engine/Plugins/FX/Niagara/Source/NiagaraEditorWidgets/Private/SNiagaraOverviewInlineParameterBox.h b/Engine/Plugins/FX/Niagara/Source/NiagaraEditorWidgets/Private/SNiagaraOverviewInlineParameterBox.h
--- a/Engine/Plugins/FX/Niagara/Source/NiagaraEditorWidgets/Private/SNiagaraOverviewInlineParameterBox.h
+++ b/Engine/Plugins/FX/Niagara/Source/NiagaraEditorWidgets/Private/SNiagaraOverviewInlineParameterBox.h
@@ -31,8 +31,6 @@ public:
 
 	virtual ~SNiagaraOverviewInlineParameterBox() override;
 
-	int32 GetNumChildWidgets();
-
 private:
 	/** Selects the module item in the stack and searches for the function input */
 	FReply NavigateToStack(TWeakObjectPtr<const UNiagaraStackFunctionInput> FunctionInput);
