Niagara - Fix a crash caused by the overview tweaks.

diff --git a/Engine/Plugins/FX/Niagara/Source/NiagaraEditor/Private/ViewModels/Stack/NiagaraStackEmitterPropertiesGroup.cpp b/Engine/Plugins/FX/Niagara/Source/NiagaraEditor/Private/ViewModels/Stack/NiagaraStackEmitterPropertiesGroup.cpp
--- a/Engine/Plugins/FX/Niagara/Source/NiagaraEditor/Private/ViewModels/Stack/NiagaraStackEmitterPropertiesGroup.cpp
+++ b/Engine/Plugins/FX/Niagara/Source/NiagaraEditor/Private/ViewModels/Stack/NiagaraStackEmitterPropertiesGroup.cpp
@@ -191,14 +191,17 @@ const FSlateBrush* UNiagaraStackEmitterPropertiesGroup::GetIconBrush() const
 
 const FSlateBrush* UNiagaraStackEmitterPropertiesGroup::GetSecondaryIconBrush() const
 {
-	UNiagaraEmitter* Emitter = GetEmitterViewModel()->GetEmitter();
-	if (Emitter->SimTarget == ENiagaraSimTarget::CPUSim)
+	if (IsFinalized() == false && GetEmitterViewModel().IsValid())
 	{
-		return FNiagaraEditorStyle::Get().GetBrush("NiagaraEditor.Stack.CPUIcon");
-	}
-	if (Emitter->SimTarget == ENiagaraSimTarget::GPUComputeSim)
-	{
-		return FNiagaraEditorStyle::Get().GetBrush("NiagaraEditor.Stack.GPUIcon");
+		UNiagaraEmitter* Emitter = GetEmitterViewModel()->GetEmitter();
+		if (Emitter->SimTarget == ENiagaraSimTarget::CPUSim)
+		{
+			return FNiagaraEditorStyle::Get().GetBrush("NiagaraEditor.Stack.CPUIcon");
+		}
+		if (Emitter->SimTarget == ENiagaraSimTarget::GPUComputeSim)
+		{
+			return FNiagaraEditorStyle::Get().GetBrush("NiagaraEditor.Stack.GPUIcon");
+		}
 	}
 	return FEditorStyle::GetBrush("NoBrush");
 }
