Refactor auto fixup script in stack dependencies to not apply when pasting modules manually.

diff --git a/Engine/Plugins/FX/CascadeToNiagaraConverter/Source/CascadeToNiagaraConverter/Private/NiagaraStackGraphUtilitiesAdapterLibrary.cpp b/Engine/Plugins/FX/CascadeToNiagaraConverter/Source/CascadeToNiagaraConverter/Private/NiagaraStackGraphUtilitiesAdapterLibrary.cpp
--- a/Engine/Plugins/FX/CascadeToNiagaraConverter/Source/CascadeToNiagaraConverter/Private/NiagaraStackGraphUtilitiesAdapterLibrary.cpp
+++ b/Engine/Plugins/FX/CascadeToNiagaraConverter/Source/CascadeToNiagaraConverter/Private/NiagaraStackGraphUtilitiesAdapterLibrary.cpp
@@ -1897,6 +1897,7 @@ void UNiagaraEmitterConversionContext::InternalFinalizeStackEntryAddActions()
 		}
 
 		UNiagaraClipboardContent* ClipboardContent = UNiagaraClipboardContent::Create();
+		ClipboardContent->bFixupPasteIndexForScriptDependenciesInStack = true;
 		UNiagaraScript* NiagaraScript = ScriptConversionContext->GetScript();
 
 		UNiagaraClipboardFunction* ClipboardFunction = UNiagaraClipboardFunction::CreateScriptFunction(ClipboardContent, "Function", NiagaraScript);
diff --git a/Engine/Plugins/FX/Niagara/Source/NiagaraEditor/Private/ViewModels/Stack/NiagaraStackClipboardUtilities.cpp b/Engine/Plugins/FX/Niagara/Source/NiagaraEditor/Private/ViewModels/Stack/NiagaraStackClipboardUtilities.cpp
--- a/Engine/Plugins/FX/Niagara/Source/NiagaraEditor/Private/ViewModels/Stack/NiagaraStackClipboardUtilities.cpp
+++ b/Engine/Plugins/FX/Niagara/Source/NiagaraEditor/Private/ViewModels/Stack/NiagaraStackClipboardUtilities.cpp
@@ -1,6 +1,7 @@
 // Copyright Epic Games, Inc. All Rights Reserved.
 
 #include "ViewModels/Stack/NiagaraStackClipboardUtilities.h"
+#include "ViewModels/Stack/NiagaraStackScriptItemGroup.h"
 #include "ViewModels/NiagaraSystemViewModel.h"
 #include "ViewModels/NiagaraSystemSelectionViewModel.h"
 #include "ViewModels/Stack/NiagaraStackEntry.h"
@@ -198,6 +199,12 @@ void FNiagaraStackClipboardUtilities::PasteSelection(const TArray<UNiagaraStackE
 		FScopedTransaction PasteTransaction(TransactionMessage);
 		if (ClipboardContent != nullptr)
 		{
+			// If the incoming paste is not targeting a specific location in the stack but the stack group header, auto fixup the paste indices.
+			if (SelectedEntries[0]->IsA<UNiagaraStackScriptItemGroup>())
+			{
+				ClipboardContent->bFixupPasteIndexForScriptDependenciesInStack = true;
+			}
+
 			SelectedEntries[0]->Paste(ClipboardContent, OutPasteWarning);
 		}
 	}
diff --git a/Engine/Plugins/FX/Niagara/Source/NiagaraEditor/Private/ViewModels/Stack/NiagaraStackScriptItemGroup.cpp b/Engine/Plugins/FX/Niagara/Source/NiagaraEditor/Private/ViewModels/Stack/NiagaraStackScriptItemGroup.cpp
--- a/Engine/Plugins/FX/Niagara/Source/NiagaraEditor/Private/ViewModels/Stack/NiagaraStackScriptItemGroup.cpp
+++ b/Engine/Plugins/FX/Niagara/Source/NiagaraEditor/Private/ViewModels/Stack/NiagaraStackScriptItemGroup.cpp
@@ -1314,9 +1314,10 @@ void UNiagaraStackScriptItemGroup::PasteModules(const UNiagaraClipboardContent*
 						// Otherwise it's a scratch pad script from another asset so we need to add a duplicate scratch pad script to this asset.
 						NewFunctionScript = GetSystemViewModel()->GetScriptScratchPadViewModel()->CreateNewScriptAsDuplicate(ClipboardFunctionScript)->GetOriginalScript();
 					}
+
 					FNiagaraStackGraphUtilities::FAddScriptModuleToStackArgs AddScriptModuleToStackArgs(NewFunctionScript, *OutputNode);
 					AddScriptModuleToStackArgs.TargetIndex = CurrentPasteIndex;
-					AddScriptModuleToStackArgs.bFixupTargetIndex = true;
+					AddScriptModuleToStackArgs.bFixupTargetIndex = ClipboardContent->bFixupPasteIndexForScriptDependenciesInStack;
 					AddScriptModuleToStackArgs.SuggestedName = ClipboardFunction->FunctionName;
 					AddScriptModuleToStackArgs.VersionGuid = ClipboardFunction->ScriptVersion;
 					NewFunctionCallNode = FNiagaraStackGraphUtilities::AddScriptModuleToStack(AddScriptModuleToStackArgs);
diff --git a/Engine/Plugins/FX/Niagara/Source/NiagaraEditor/Public/NiagaraClipboard.h b/Engine/Plugins/FX/Niagara/Source/NiagaraEditor/Public/NiagaraClipboard.h
--- a/Engine/Plugins/FX/Niagara/Source/NiagaraEditor/Public/NiagaraClipboard.h
+++ b/Engine/Plugins/FX/Niagara/Source/NiagaraEditor/Public/NiagaraClipboard.h
@@ -181,6 +181,10 @@ public:
 	/** We expect nodes to be exported into this string using FEdGraphUtilities::ExportNodesToText */
 	UPROPERTY()
 	FString ExportedNodes;
+
+	/** Markup MetaData to specify that if scripts are pasted from the clipboard to automatically fixup their order in the stack to satisfy dependencies. */
+	UPROPERTY()
+	mutable bool bFixupPasteIndexForScriptDependenciesInStack = false;
 };
 
 class NIAGARAEDITOR_API FNiagaraClipboard
