Fix RHI resource leak by calling release on uniform buffer before memstack wipes memory.

diff --git a/Engine/Plugins/FX/Niagara/Source/Niagara/Private/NiagaraGpuComputeDispatch.cpp b/Engine/Plugins/FX/Niagara/Source/Niagara/Private/NiagaraGpuComputeDispatch.cpp
--- a/Engine/Plugins/FX/Niagara/Source/Niagara/Private/NiagaraGpuComputeDispatch.cpp
+++ b/Engine/Plugins/FX/Niagara/Source/Niagara/Private/NiagaraGpuComputeDispatch.cpp
@@ -500,6 +500,10 @@ void FNiagaraGpuComputeDispatch::ProcessPendingTicksFlush(FRHICommandListImmedia
 				}
 			);
 			GraphBuilder.Execute();
+
+			// Properly clear the reference to ViewUniformBuffer before memstack wipes the memory
+			DummyView->~FViewInfo();
+
 			break;
 		}
 
