Changed op category to position

diff --git a/Engine/Plugins/FX/Niagara/Source/NiagaraEditor/Private/NiagaraEditorCommon.cpp b/Engine/Plugins/FX/Niagara/Source/NiagaraEditor/Private/NiagaraEditorCommon.cpp
--- a/Engine/Plugins/FX/Niagara/Source/NiagaraEditor/Private/NiagaraEditorCommon.cpp
+++ b/Engine/Plugins/FX/Niagara/Source/NiagaraEditor/Private/NiagaraEditorCommon.cpp
@@ -107,6 +107,8 @@ void FNiagaraOpInfo::Init()
 
 	FName X(TEXT("X"));	FText XText = NSLOCTEXT("NiagaraOpInfo", "First Vector Component", "X");
 
+	FText PositionCategory = NSLOCTEXT("NiagaraOpInfo", "PositionOpCategory", "Position");
+
 	static FText SelectionText = NSLOCTEXT("NiagaraOpInfo", "Selection", "Selection");
 	static FText MaskText = NSLOCTEXT("NiagaraOpInfo", "Mask", "Mask");
 	
@@ -1082,7 +1084,7 @@ void FNiagaraOpInfo::Init()
 
 		Idx = OpInfos.AddDefaulted();
 		Op = &OpInfos[Idx];
-		Op->Category = NumericCategory;
+		Op->Category = PositionCategory;
 		Op->FriendlyName = NSLOCTEXT("NiagaraOpInfo", "Distance Name", "Distance");
 		Op->Description = NSLOCTEXT("NiagaraOpInfo", "Distance Desc", "Returns the distance between two positions.");
 		Op->Inputs.Add(FNiagaraOpInOutInfo(A, FNiagaraTypeDefinition::GetPositionDef(), AText, AText, DefaultStr_One));
