- IWYU fix

#rnx

diff --git a/Engine/Plugins/FX/Niagara/Source/NiagaraEditor/Private/Toolkits/NiagaraSystemToolkit.h b/Engine/Plugins/FX/Niagara/Source/NiagaraEditor/Private/Toolkits/NiagaraSystemToolkit.h
--- a/Engine/Plugins/FX/Niagara/Source/NiagaraEditor/Private/Toolkits/NiagaraSystemToolkit.h
+++ b/Engine/Plugins/FX/Niagara/Source/NiagaraEditor/Private/Toolkits/NiagaraSystemToolkit.h
@@ -11,10 +11,15 @@
 #include "ISequencer.h"
 #include "ISequencerTrackEditor.h"
 
+#include "NiagaraCommon.h"
 #include "NiagaraScript.h"
 #include "Widgets/SItemSelector.h"
 #include "ViewModels/NiagaraSystemGraphSelectionViewModel.h"
 
+#if WITH_NIAGARA_GPU_PROFILER
+	#include "NiagaraGPUProfilerInterface.h"
+#endif
+
 class FNiagaraSystemInstance;
 class FNiagaraSystemViewModel;
 class FNiagaraObjectSelection;
@@ -235,7 +240,7 @@ private:
 	static IConsoleVariable* VmStatEnabledVar;
 
 #if WITH_NIAGARA_GPU_PROFILER
-	TUniquePtr<struct FNiagaraGpuProfilerListener> GpuProfilerListener;
+	TUniquePtr<FNiagaraGpuProfilerListener> GpuProfilerListener;
 #endif
 
 public:
