Fix for DIOcclusion_Rectangle shader failing to cross compile on the mac; avoided 'continue' statements in nested for loops avoids "Over 3 compilation loops detected. Must be a bug!" failure.

diff --git a/Engine/Plugins/FX/Niagara/Shaders/Private/NiagaraDataInterfaceOcclusion.ush b/Engine/Plugins/FX/Niagara/Shaders/Private/NiagaraDataInterfaceOcclusion.ush
--- a/Engine/Plugins/FX/Niagara/Shaders/Private/NiagaraDataInterfaceOcclusion.ush
+++ b/Engine/Plugins/FX/Niagara/Shaders/Private/NiagaraDataInterfaceOcclusion.ush
@@ -35,30 +35,29 @@ void DIOcclusion_Rectangle(in FNDIOcclusionContext DIContext, in float3 In_Sampl
 
 	if (Steps > 0) 
 	{
+		const float StepScale = 1.0f / (Steps - 1);
 		for (int ys = 0; ys < Steps; ys++)
 		{
-			float SampleY = ScreenUV.y - 0.5 * SampleHeight + ys * SampleHeight / (Steps - 1);
-			if (SampleY > 1 || SampleY < 0)
+			float SampleY = ScreenUV.y + SampleHeight * (ys * StepScale - 0.5f);
+			if (SampleY <= 1 && SampleY >= 0)
 			{
-				continue;
-			}
-			for (int xs = 0; xs < Steps; xs++)
-			{
-				float SampleX = ScreenUV.x - 0.5 * SampleWidth + xs * SampleWidth / (Steps - 1);
-				if (SampleX > 1 || SampleX < 0)
+				for (int xs = 0; xs < Steps; xs++)
 				{
-					continue;
+					float SampleX = ScreenUV.x + SampleWidth * (xs * StepScale - 0.5f);
+					if (SampleX <= 1 && SampleX >= 0)
+					{
+						float Depth = CalcSceneDepth(float2(SampleX, SampleY));
+						if (Depth < CameraDistance) 
+						{
+							OccludedSamples++;
+						}
+						TotalSamples++;
+					}
 				}
-				
-				float Depth = CalcSceneDepth(float2(SampleX, SampleY));
-				if (Depth < CameraDistance) 
-				{
-					OccludedSamples++;
-				}
-				TotalSamples++;
 			} 
 		}
 	}
+
 	Out_VisibilityFraction = TotalSamples > 0 ? 1 - OccludedSamples / TotalSamples : 0;
 	Out_SampleFraction = Steps == 0 ? 0 : (TotalSamples / (Steps * Steps));
 }
