Make the compilation status more obvious in Niagara

diff --git a/Engine/Plugins/FX/Niagara/Source/NiagaraEditor/Private/NiagaraEditorStyle.cpp b/Engine/Plugins/FX/Niagara/Source/NiagaraEditor/Private/NiagaraEditorStyle.cpp
--- a/Engine/Plugins/FX/Niagara/Source/NiagaraEditor/Private/NiagaraEditorStyle.cpp
+++ b/Engine/Plugins/FX/Niagara/Source/NiagaraEditor/Private/NiagaraEditorStyle.cpp
@@ -375,6 +375,14 @@ void FNiagaraEditorStyle::InitOverview()
 	Set("NiagaraEditor.Overview.DebugInactive", new IMAGE_BRUSH("Icons/OverviewDebugInactive", Icon16x16));
 }
 
+void FNiagaraEditorStyle::InitViewportStyle()
+{
+	FTextBlockStyle NormalTextStyle = FEditorStyle::GetWidgetStyle<FTextBlockStyle>("NormalText");
+	FTextBlockStyle CompileOverlayText = FTextBlockStyle(NormalTextStyle)
+		.SetFontSize(18);
+	Set("NiagaraEditor.Viewport.CompileOverlay", CompileOverlayText);
+}
+
 void FNiagaraEditorStyle::InitEmitterDetails()
 {
 	// Emitter details customization
@@ -681,6 +689,7 @@ FNiagaraEditorStyle::FNiagaraEditorStyle() : FSlateStyleSet("NiagaraEditorStyle"
 	InitOutlinerStyle();
 	InitScalabilityColors();
 	InitScalabilityIcons();
+	InitViewportStyle();
 }
 
 void FNiagaraEditorStyle::ReloadTextures()
diff --git a/Engine/Plugins/FX/Niagara/Source/NiagaraEditor/Private/Widgets/SNiagaraSystemViewport.cpp b/Engine/Plugins/FX/Niagara/Source/NiagaraEditor/Private/Widgets/SNiagaraSystemViewport.cpp
--- a/Engine/Plugins/FX/Niagara/Source/NiagaraEditor/Private/Widgets/SNiagaraSystemViewport.cpp
+++ b/Engine/Plugins/FX/Niagara/Source/NiagaraEditor/Private/Widgets/SNiagaraSystemViewport.cpp
@@ -11,6 +11,7 @@
 #include "NiagaraEditorCommands.h"
 #include "NiagaraEditorModule.h"
 #include "NiagaraEditorSettings.h"
+#include "NiagaraEditorStyle.h"
 #include "NiagaraEffectType.h"
 #include "NiagaraPerfBaseline.h"
 #include "NiagaraSettings.h"
@@ -677,16 +678,28 @@ EVisibility SNiagaraSystemViewport::OnGetViewportContentVisibility() const
 	return IsVisible() ? EVisibility::Visible : EVisibility::Collapsed;
 }
 
-EVisibility SNiagaraSystemViewport::OnGetViewportCompileTextVisibility() const
+FText SNiagaraSystemViewport::GetViewportCompileStatusText() const
 {
-	if (PreviewComponent && PreviewComponent->GetAsset() && PreviewComponent->GetAsset()->HasOutstandingCompilationRequests())
-	{
-		return EVisibility::Visible;
-	}
-	else
+	if (PreviewComponent && PreviewComponent->GetAsset())
 	{
-		return EVisibility::Collapsed;
+		if (PreviewComponent->GetAsset()->HasOutstandingCompilationRequests())
+		{
+			return LOCTEXT("Compiling", "Compiling...");
+		}
+		bool bCompilationFailed = false;
+		PreviewComponent->GetAsset()->ForEachScript([&bCompilationFailed](UNiagaraScript* Script)
+		{
+			if (Script && Script->GetLastCompileStatus() == ENiagaraScriptCompileStatus::NCS_Error)
+			{
+				bCompilationFailed = true;
+			}
+		});
+		if (bCompilationFailed)
+		{
+			return LOCTEXT("CompilingFailed", "Compilation Failed!");
+		}
 	}
+	return FText();
 }
 
 void SNiagaraSystemViewport::PopulateViewportOverlays(TSharedRef<class SOverlay> Overlay)
@@ -701,11 +714,12 @@ void SNiagaraSystemViewport::PopulateViewportOverlays(TSharedRef<class SOverlay>
 	.HAlign(HAlign_Center)
 	[
 		SAssignNew(CompileText, STextBlock)
-		.Visibility_Raw(this, &SNiagaraSystemViewport::OnGetViewportCompileTextVisibility)
+		.Text_Raw(this, &SNiagaraSystemViewport::GetViewportCompileStatusText)
+		.TextStyle(FNiagaraEditorStyle::Get(), "NiagaraEditor.Viewport.CompileOverlay")
+		.ColorAndOpacity(FLinearColor::White)
+		.ShadowOffset(FVector2D(1.5, 1.5))
+		.ShadowColorAndOpacity(FLinearColor(0, 0, 0, 0.9f))
 	];
-
-	CompileText->SetText(LOCTEXT("Compiling","Compiling"));
-
 }
 
 
diff --git a/Engine/Plugins/FX/Niagara/Source/NiagaraEditor/Private/Widgets/SNiagaraSystemViewport.h b/Engine/Plugins/FX/Niagara/Source/NiagaraEditor/Private/Widgets/SNiagaraSystemViewport.h
--- a/Engine/Plugins/FX/Niagara/Source/NiagaraEditor/Private/Widgets/SNiagaraSystemViewport.h
+++ b/Engine/Plugins/FX/Niagara/Source/NiagaraEditor/Private/Widgets/SNiagaraSystemViewport.h
@@ -102,7 +102,7 @@ protected:
 	virtual void BindCommands() override;
 	virtual void OnFocusViewportToSelection() override;
 	virtual void PopulateViewportOverlays(TSharedRef<class SOverlay> Overlay) override;
-	EVisibility OnGetViewportCompileTextVisibility() const;
+	FText GetViewportCompileStatusText() const;
 
 private:
 	bool IsVisible() const override;
diff --git a/Engine/Plugins/FX/Niagara/Source/NiagaraEditor/Public/NiagaraEditorStyle.h b/Engine/Plugins/FX/Niagara/Source/NiagaraEditor/Public/NiagaraEditorStyle.h
--- a/Engine/Plugins/FX/Niagara/Source/NiagaraEditor/Public/NiagaraEditorStyle.h
+++ b/Engine/Plugins/FX/Niagara/Source/NiagaraEditor/Public/NiagaraEditorStyle.h
@@ -34,6 +34,7 @@ private:
 	void InitTabIcons();
 	void InitIcons();
 	void InitOverview();
+	void InitViewportStyle();
 	void InitEmitterDetails();
 	void InitAssetColors();
 	void InitThumbnails();
