Fixed a problem with stackcontext variables not being updated to a position type

diff --git a/Engine/Plugins/FX/Niagara/Source/Niagara/Private/NiagaraConstants.cpp b/Engine/Plugins/FX/Niagara/Source/Niagara/Private/NiagaraConstants.cpp
--- a/Engine/Plugins/FX/Niagara/Source/Niagara/Private/NiagaraConstants.cpp
+++ b/Engine/Plugins/FX/Niagara/Source/Niagara/Private/NiagaraConstants.cpp
@@ -148,6 +148,9 @@ void FNiagaraConstants::Init()
 		OldPositionTypes.Add(FNiagaraVariable(FNiagaraTypeDefinition::GetVec3Def(),  SYS_PARAM_PARTICLES_POSITION.GetName()));
 		OldPositionTypes.Add(FNiagaraVariable(FNiagaraTypeDefinition::GetVec3Def(), FName("Particles.Initial.Position")));
 		OldPositionTypes.Add(FNiagaraVariable(FNiagaraTypeDefinition::GetVec3Def(), FName("Particles.Previous.Position")));
+		OldPositionTypes.Add(FNiagaraVariable(FNiagaraTypeDefinition::GetVec3Def(), FName(StackContextNamespaceString + ".Position")));
+		OldPositionTypes.Add(FNiagaraVariable(FNiagaraTypeDefinition::GetVec3Def(), FName(StackContextNamespaceString + ".Initial.Position")));
+		OldPositionTypes.Add(FNiagaraVariable(FNiagaraTypeDefinition::GetVec3Def(), FName(StackContextNamespaceString + ".Previous.Position")));
 		OldPositionTypes.Add(FNiagaraVariable(FNiagaraTypeDefinition::GetVec3Def(),  SYS_PARAM_ENGINE_POSITION.GetName()));
 		OldPositionTypes.Add(FNiagaraVariable(FNiagaraTypeDefinition::GetVec3Def(),  SYS_PARAM_ENGINE_EMITTER_SIMULATION_POSITION.GetName()));
 	}
diff --git a/Engine/Plugins/FX/Niagara/Source/NiagaraEditor/Private/NiagaraEditorData.cpp b/Engine/Plugins/FX/Niagara/Source/NiagaraEditor/Private/NiagaraEditorData.cpp
--- a/Engine/Plugins/FX/Niagara/Source/NiagaraEditor/Private/NiagaraEditorData.cpp
+++ b/Engine/Plugins/FX/Niagara/Source/NiagaraEditor/Private/NiagaraEditorData.cpp
@@ -1,6 +1,8 @@
 // Copyright Epic Games, Inc. All Rights Reserved.
 
 #include "NiagaraEditorData.h"
+
+#include "NiagaraConstants.h"
 #include "NiagaraEditorUtilities.h"
 #include "NiagaraParameterDefinitionsBase.h"
 #include "NiagaraParameterDefinitions.h"
@@ -19,6 +21,18 @@ void UNiagaraEditorParametersAdapter::PostLoad()
 			EditorOnlyScriptVars.RemoveAtSwap(Idx);
 		}
 	}
+
+	TArray<FNiagaraVariable> OldTypes = FNiagaraConstants::GetOldPositionTypeVariables();
+	for (FNiagaraVariable OldVarType : OldTypes)
+	{
+		for (TObjectPtr<UNiagaraScriptVariable> ScriptVariable : EditorOnlyScriptVars)
+		{
+			if (!ScriptVariable.IsNull() && ScriptVariable->Variable == OldVarType)
+			{
+				ScriptVariable->Variable.SetType(FNiagaraTypeDefinition::GetPositionDef());
+			}
+		}
+	}
 }
 
 TArray<TTuple<FName /*SyncedOldName*/, FName /*SyncedNewName*/>> UNiagaraEditorParametersAdapter::SynchronizeParametersWithParameterDefinitions(
