Removed module nodes from dynamic scripts and other module scripts

diff --git a/Engine/Plugins/FX/Niagara/Source/NiagaraEditor/Private/EdGraphSchema_Niagara.cpp b/Engine/Plugins/FX/Niagara/Source/NiagaraEditor/Private/EdGraphSchema_Niagara.cpp
--- a/Engine/Plugins/FX/Niagara/Source/NiagaraEditor/Private/EdGraphSchema_Niagara.cpp
+++ b/Engine/Plugins/FX/Niagara/Source/NiagaraEditor/Private/EdGraphSchema_Niagara.cpp
@@ -479,7 +479,7 @@ TArray<TSharedPtr<FNiagaraAction_NewNode>> UEdGraphSchema_Niagara::GetGraphActio
 	}
 
 	//Add modules
-	if (bAllowModuleNodes && !bFunctionGraph)
+	if (bAllowModuleNodes && !bFunctionGraph && !bModuleGraph && !bDynamicInputGraph)
 	{
 		TArray<FAssetData> ModuleScriptAssets;
 		FNiagaraEditorUtilities::FGetFilteredScriptAssetsOptions ModuleScriptFilterOptions;
