Fixed a bug that prevented cooked Niagara assets from working in the editor

diff --git a/Engine/Plugins/FX/Niagara/Source/Niagara/Private/NiagaraSystem.cpp b/Engine/Plugins/FX/Niagara/Source/Niagara/Private/NiagaraSystem.cpp
--- a/Engine/Plugins/FX/Niagara/Source/Niagara/Private/NiagaraSystem.cpp
+++ b/Engine/Plugins/FX/Niagara/Source/Niagara/Private/NiagaraSystem.cpp
@@ -704,7 +704,7 @@ void UNiagaraSystem::UpdateSystemAfterLoad()
 	bFullyLoaded = true;
 
 #if WITH_EDITORONLY_DATA
-	if (SystemSpawnScript)
+	if (SystemSpawnScript && !GetOutermost()->bIsCookedForEditor)
 	{
 		ensure(SystemSpawnScript->GetLatestSource() != nullptr);
 		ensure(SystemSpawnScript->GetLatestSource() == SystemUpdateScript->GetLatestSource());
@@ -3132,7 +3132,7 @@ bool UNiagaraSystem::RequestCompile(bool bForce, FNiagaraSystemUpdateContext* Op
 		return false;
 	}
 
-	if (!AllowShaderCompiling())
+	if (!AllowShaderCompiling() || GetOutermost()->bIsCookedForEditor)
 	{
 		return false;
 	}
