Force DXC for NiagaraDataInterface shaders

diff --git a/Engine/Plugins/FX/Niagara/Source/Niagara/Classes/NiagaraDataInterface.h b/Engine/Plugins/FX/Niagara/Source/Niagara/Classes/NiagaraDataInterface.h
--- a/Engine/Plugins/FX/Niagara/Source/Niagara/Classes/NiagaraDataInterface.h
+++ b/Engine/Plugins/FX/Niagara/Source/Niagara/Classes/NiagaraDataInterface.h
@@ -319,7 +319,7 @@ public:
 #endif
 
 	/** Allows data interfaces to influence the compilation of GPU shaders and is only called on the CDO object not the instance. */
-	virtual void ModifyCompilationEnvironment(EShaderPlatform ShaderPlatform, struct FShaderCompilerEnvironment& OutEnvironment) const {}
+	virtual void ModifyCompilationEnvironment(EShaderPlatform ShaderPlatform, struct FShaderCompilerEnvironment& OutEnvironment) const;
 
 	/** 
 		Subclasses that wish to work with GPU systems/emitters must implement this.
diff --git a/Engine/Plugins/FX/Niagara/Source/Niagara/Private/NiagaraDataInterface.cpp b/Engine/Plugins/FX/Niagara/Source/Niagara/Private/NiagaraDataInterface.cpp
--- a/Engine/Plugins/FX/Niagara/Source/Niagara/Private/NiagaraDataInterface.cpp
+++ b/Engine/Plugins/FX/Niagara/Source/Niagara/Private/NiagaraDataInterface.cpp
@@ -7,6 +7,7 @@
 #include "ShaderParameterUtils.h"
 #include "NiagaraShader.h"
 #include "NiagaraComponent.h"
+#include "ShaderCompilerCore.h"
 
 #define LOCTEXT_NAMESPACE "NiagaraDataInterface"
 
@@ -47,6 +48,12 @@ bool UNiagaraDataInterface::AppendCompileHash(FNiagaraCompileHashVisitor* InVisi
 }
 #endif
 
+void UNiagaraDataInterface::ModifyCompilationEnvironment(EShaderPlatform ShaderPlatform, struct FShaderCompilerEnvironment& OutEnvironment) const
+{
+	// Always enable DXC to avoid compile errors caused by RWBuffer/Buffer in structs. Example NiagaraDataInterfaceHairStrands.ush struct FDIHairStrandsContext
+	OutEnvironment.CompilerFlags.Add(CFLAG_ForceDXC);
+}
+
 void UNiagaraDataInterface::GetAssetTagsForContext(const UObject* InAsset, const TArray<const UNiagaraDataInterface*>& InProperties, TMap<FName, uint32>& NumericKeys, TMap<FName, FString>& StringKeys) const
 {
 	UClass* Class = GetClass();
