Added a new static switch constant for interpolated spawning

diff --git a/Engine/Plugins/FX/Niagara/Source/Niagara/Classes/NiagaraConstants.h b/Engine/Plugins/FX/Niagara/Source/Niagara/Classes/NiagaraConstants.h
--- a/Engine/Plugins/FX/Niagara/Source/Niagara/Classes/NiagaraConstants.h
+++ b/Engine/Plugins/FX/Niagara/Source/Niagara/Classes/NiagaraConstants.h
@@ -82,6 +82,7 @@
 #define SYS_PARAM_EMITTER_AGE                            INiagaraModule::GetVar_Emitter_Age()
 #define SYS_PARAM_EMITTER_LOCALSPACE                     INiagaraModule::GetVar_Emitter_LocalSpace()
 #define SYS_PARAM_EMITTER_DETERMINISM                    INiagaraModule::GetVar_Emitter_Determinism()
+#define SYS_PARAM_EMITTER_INTERP_SPAWN                   INiagaraModule::GetVar_Emitter_InterpolatedSpawn()
 #define SYS_PARAM_EMITTER_OVERRIDE_GLOBAL_SPAWN_COUNT_SCALE       INiagaraModule::GetVar_Emitter_OverrideGlobalSpawnCountScale()
 #define SYS_PARAM_EMITTER_RANDOM_SEED                    INiagaraModule::GetVar_Emitter_RandomSeed()
 #define SYS_PARAM_EMITTER_SPAWNRATE                      INiagaraModule::GetVar_Emitter_SpawnRate()
diff --git a/Engine/Plugins/FX/Niagara/Source/Niagara/Private/NiagaraConstants.cpp b/Engine/Plugins/FX/Niagara/Source/Niagara/Private/NiagaraConstants.cpp
--- a/Engine/Plugins/FX/Niagara/Source/Niagara/Private/NiagaraConstants.cpp
+++ b/Engine/Plugins/FX/Niagara/Source/Niagara/Private/NiagaraConstants.cpp
@@ -138,6 +138,7 @@ void FNiagaraConstants::Init()
 		SwitchParameters.Add(SYS_PARAM_EMITTER_DETERMINISM);
 		SwitchParameters.Add(SYS_PARAM_EMITTER_OVERRIDE_GLOBAL_SPAWN_COUNT_SCALE);
 		SwitchParameters.Add(SYS_PARAM_EMITTER_SIMULATION_TARGET);
+		SwitchParameters.Add(SYS_PARAM_EMITTER_INTERP_SPAWN);
 		SwitchParameters.Add(SYS_PARAM_SCRIPT_USAGE);
 		SwitchParameters.Add(SYS_PARAM_SCRIPT_CONTEXT);
 		SwitchParameters.Add(SYS_PARAM_FUNCTION_DEBUG_STATE);
@@ -185,6 +186,7 @@ void FNiagaraConstants::Init()
 		UpdatedSystemParameters.Add(FName(TEXT("Emitter Random Seed")), SYS_PARAM_EMITTER_RANDOM_SEED);
 		UpdatedSystemParameters.Add(FName(TEXT("Emitter Instance Seed")), SYS_PARAM_ENGINE_EMITTER_INSTANCE_SEED);
 		UpdatedSystemParameters.Add(FName(TEXT("Emitter Determinism")), SYS_PARAM_EMITTER_DETERMINISM);
+		UpdatedSystemParameters.Add(FName(TEXT("Emitter Interpolated Spawn")), SYS_PARAM_EMITTER_INTERP_SPAWN);
 		UpdatedSystemParameters.Add(FName(TEXT("Emitter Override Global Spawn Count Scale")), SYS_PARAM_EMITTER_OVERRIDE_GLOBAL_SPAWN_COUNT_SCALE);
 		UpdatedSystemParameters.Add(FName(TEXT("Effect Position")), SYS_PARAM_ENGINE_POSITION);
 		UpdatedSystemParameters.Add(FName(TEXT("Effect Velocity")), SYS_PARAM_ENGINE_VELOCITY);
diff --git a/Engine/Plugins/FX/Niagara/Source/Niagara/Private/NiagaraModule.cpp b/Engine/Plugins/FX/Niagara/Source/Niagara/Private/NiagaraModule.cpp
--- a/Engine/Plugins/FX/Niagara/Source/Niagara/Private/NiagaraModule.cpp
+++ b/Engine/Plugins/FX/Niagara/Source/Niagara/Private/NiagaraModule.cpp
@@ -135,6 +135,7 @@ FNiagaraVariable INiagaraModule::Engine_System_Age;
 FNiagaraVariable INiagaraModule::Emitter_Age;
 FNiagaraVariable INiagaraModule::Emitter_LocalSpace;
 FNiagaraVariable INiagaraModule::Emitter_Determinism;
+FNiagaraVariable INiagaraModule::Emitter_InterpolatedSpawn;
 FNiagaraVariable INiagaraModule::Emitter_OverrideGlobalSpawnCountScale;
 FNiagaraVariable INiagaraModule::Emitter_SimulationTarget;
 FNiagaraVariable INiagaraModule::Emitter_RandomSeed;
@@ -263,6 +264,7 @@ void INiagaraModule::StartupModule()
 	Emitter_LocalSpace = FNiagaraVariable(FNiagaraTypeDefinition::GetBoolDef(), TEXT("Emitter.LocalSpace"));
 	Emitter_RandomSeed = FNiagaraVariable(FNiagaraTypeDefinition::GetIntDef(), TEXT("Emitter.RandomSeed"));
 	Emitter_Determinism = FNiagaraVariable(FNiagaraTypeDefinition::GetBoolDef(), TEXT("Emitter.Determinism"));
+	Emitter_InterpolatedSpawn = FNiagaraVariable(FNiagaraTypeDefinition::GetBoolDef(), TEXT("Emitter.InterpolatedSpawn"));
 	Emitter_OverrideGlobalSpawnCountScale = FNiagaraVariable(FNiagaraTypeDefinition::GetBoolDef(), TEXT("Emitter.OverrideGlobalSpawnCountScale"));
 	Emitter_SimulationTarget = FNiagaraVariable(FNiagaraTypeDefinition::GetSimulationTargetEnum(), TEXT("Emitter.SimulationTarget"));
 	Emitter_SpawnRate = FNiagaraVariable(FNiagaraTypeDefinition::GetFloatDef(), TEXT("Emitter.SpawnRate"));
diff --git a/Engine/Plugins/FX/Niagara/Source/Niagara/Private/NiagaraScript.cpp b/Engine/Plugins/FX/Niagara/Source/Niagara/Private/NiagaraScript.cpp
--- a/Engine/Plugins/FX/Niagara/Source/Niagara/Private/NiagaraScript.cpp
+++ b/Engine/Plugins/FX/Niagara/Source/Niagara/Private/NiagaraScript.cpp
@@ -938,6 +938,10 @@ void UNiagaraScript::ComputeVMCompilationId(FNiagaraVMExecutableDataId& Id, FGui
 			Id.bInterpolatedSpawn = true;
 			Id.AdditionalDefines.Add(TEXT("InterpolatedSpawn"));
 		}
+		if (Emitter->bInterpolatedSpawning)
+		{
+			Id.AdditionalDefines.Add(TEXT("Emitter.InterpolatedSpawn"));
+		}
 		if (IsParticleScript(Usage))
 		{
 			SimTargetToBuild = Emitter->SimTarget;
@@ -1018,6 +1022,10 @@ void UNiagaraScript::ComputeVMCompilationId(FNiagaraVMExecutableDataId& Id, FGui
 				{
 					Id.AdditionalDefines.Add(Emitter->GetUniqueEmitterName() + TEXT(".Determinism"));
 				}
+				if (Emitter->bInterpolatedSpawning)
+				{
+					Id.AdditionalDefines.Add(Emitter->GetUniqueEmitterName() + TEXT(".InterpolatedSpawn"));
+				}
 
 				ComputeVMCompilationId_EmitterShared(Id, Emitter, System, ENiagaraRendererSourceDataMode::Emitter);
 			}
diff --git a/Engine/Plugins/FX/Niagara/Source/Niagara/Public/NiagaraModule.h b/Engine/Plugins/FX/Niagara/Source/Niagara/Public/NiagaraModule.h
--- a/Engine/Plugins/FX/Niagara/Source/Niagara/Public/NiagaraModule.h
+++ b/Engine/Plugins/FX/Niagara/Source/Niagara/Public/NiagaraModule.h
@@ -172,6 +172,7 @@ public:
 	FORCEINLINE static const FNiagaraVariable&  GetVar_Emitter_Age() { return Emitter_Age; }
 	FORCEINLINE static const FNiagaraVariable&  GetVar_Emitter_LocalSpace() { return Emitter_LocalSpace; }
 	FORCEINLINE static const FNiagaraVariable&  GetVar_Emitter_Determinism() { return Emitter_Determinism; }
+	FORCEINLINE static const FNiagaraVariable&  GetVar_Emitter_InterpolatedSpawn() { return Emitter_InterpolatedSpawn; }
 	FORCEINLINE static const FNiagaraVariable&  GetVar_Emitter_OverrideGlobalSpawnCountScale() { return Emitter_OverrideGlobalSpawnCountScale; }
 	FORCEINLINE static const FNiagaraVariable&  GetVar_Emitter_RandomSeed() { return Emitter_RandomSeed; }
 	FORCEINLINE static const FNiagaraVariable&  GetVar_Emitter_SpawnRate() { return Emitter_SpawnRate; }
@@ -293,6 +294,7 @@ private:
 	static FNiagaraVariable Emitter_Age;
 	static FNiagaraVariable Emitter_LocalSpace;
 	static FNiagaraVariable Emitter_Determinism;
+	static FNiagaraVariable Emitter_InterpolatedSpawn;
 	static FNiagaraVariable Emitter_OverrideGlobalSpawnCountScale;
 	static FNiagaraVariable Emitter_SimulationTarget;
 	static FNiagaraVariable Emitter_RandomSeed;
diff --git a/Engine/Plugins/FX/Niagara/Source/NiagaraEditor/Private/NiagaraHlslTranslator.cpp b/Engine/Plugins/FX/Niagara/Source/NiagaraEditor/Private/NiagaraHlslTranslator.cpp
--- a/Engine/Plugins/FX/Niagara/Source/NiagaraEditor/Private/NiagaraHlslTranslator.cpp
+++ b/Engine/Plugins/FX/Niagara/Source/NiagaraEditor/Private/NiagaraHlslTranslator.cpp
@@ -5253,6 +5253,12 @@ bool FHlslNiagaraTranslator::GetLiteralConstantVariable(FNiagaraVariable& OutVar
 			OutVar.SetValue(bEmitterDeterminism ? FNiagaraBool(true) : FNiagaraBool(false));
 			return true;
 		}
+		if (OutVar == FNiagaraVariable(FNiagaraTypeDefinition::GetBoolDef(), TEXT("Emitter.InterpolatedSpawn")))
+		{
+			bool bEmitterInterpolatedSpawn = CompileOptions.AdditionalDefines.Contains(ResolvedVar.GetName().ToString());
+			OutVar.SetValue(bEmitterInterpolatedSpawn ? FNiagaraBool(true) : FNiagaraBool(false));
+			return true;
+		}
 		if (OutVar == FNiagaraVariable(FNiagaraTypeDefinition::GetBoolDef(), TEXT("Emitter.OverrideGlobalSpawnCountScale")))
 		{
 			bool bOverrideGlobalSpawnCountScale = CompileOptions.AdditionalDefines.Contains(ResolvedVar.GetName().ToString());
diff --git a/Engine/Plugins/FX/Niagara/Source/NiagaraEditor/Private/NiagaraParameterMapHistory.cpp b/Engine/Plugins/FX/Niagara/Source/NiagaraEditor/Private/NiagaraParameterMapHistory.cpp
--- a/Engine/Plugins/FX/Niagara/Source/NiagaraEditor/Private/NiagaraParameterMapHistory.cpp
+++ b/Engine/Plugins/FX/Niagara/Source/NiagaraEditor/Private/NiagaraParameterMapHistory.cpp
@@ -2144,26 +2144,31 @@ bool FCompileConstantResolver::ResolveConstant(FNiagaraVariable& OutConstant) co
 		OutConstant.SetValue(Emitter->bLocalSpace ? FNiagaraBool(true) : FNiagaraBool(false));
 		return true;
 	}
-	else if (Emitter && OutConstant == FNiagaraVariable(FNiagaraTypeDefinition::GetBoolDef(), TEXT("Emitter.Determinism")))
+	if (Emitter && OutConstant == FNiagaraVariable(FNiagaraTypeDefinition::GetBoolDef(), TEXT("Emitter.Determinism")))
 	{
 		OutConstant.SetValue(Emitter->bDeterminism ? FNiagaraBool(true) : FNiagaraBool(false));
 		return true;
 	}
-	else if (Emitter && OutConstant == FNiagaraVariable(FNiagaraTypeDefinition::GetSimulationTargetEnum(), TEXT("Emitter.SimulationTarget")))
+	if (Emitter && OutConstant == FNiagaraVariable(FNiagaraTypeDefinition::GetBoolDef(), TEXT("Emitter.InterpolatedSpawn")))
+	{
+		OutConstant.SetValue(Emitter->bInterpolatedSpawning ? FNiagaraBool(true) : FNiagaraBool(false));
+		return true;
+	}
+	if (Emitter && OutConstant == FNiagaraVariable(FNiagaraTypeDefinition::GetSimulationTargetEnum(), TEXT("Emitter.SimulationTarget")))
 	{
 		FNiagaraInt32 EnumValue;
 		EnumValue.Value = (uint8)Emitter->SimTarget;
 		OutConstant.SetValue(EnumValue);
 		return true;
 	}
-	else if (Emitter && OutConstant == FNiagaraVariable(FNiagaraTypeDefinition::GetScriptUsageEnum(), TEXT("Script.Usage")))
+	if (Emitter && OutConstant == FNiagaraVariable(FNiagaraTypeDefinition::GetScriptUsageEnum(), TEXT("Script.Usage")))
 	{
 		FNiagaraInt32 EnumValue;
 		EnumValue.Value = (uint8)FNiagaraUtilities::ConvertScriptUsageToStaticSwitchUsage(Usage);
 		OutConstant.SetValue(EnumValue);
 		return true;
 	}
-	else if (Emitter && OutConstant == FNiagaraVariable(FNiagaraTypeDefinition::GetScriptContextEnum(), TEXT("Script.Context")))
+	if (Emitter && OutConstant == FNiagaraVariable(FNiagaraTypeDefinition::GetScriptContextEnum(), TEXT("Script.Context")))
 	{
 		FNiagaraInt32 EnumValue;
 		EnumValue.Value = (uint8)FNiagaraUtilities::ConvertScriptUsageToStaticSwitchContext(Usage);
@@ -2177,19 +2182,24 @@ bool FCompileConstantResolver::ResolveConstant(FNiagaraVariable& OutConstant) co
 		OutConstant.SetValue(FNiagaraBool(false));
 		return true;
 	}
-	else if (System && OutConstant == FNiagaraVariable(FNiagaraTypeDefinition::GetBoolDef(), TEXT("Emitter.Determinism")))
+	if (System && OutConstant == FNiagaraVariable(FNiagaraTypeDefinition::GetBoolDef(), TEXT("Emitter.Determinism")))
 	{
 		OutConstant.SetValue(FNiagaraBool(true));
 		return true;
 	}
-	else if (System && OutConstant == FNiagaraVariable(FNiagaraTypeDefinition::GetScriptUsageEnum(), TEXT("Script.Usage")))
+	if (System && OutConstant == FNiagaraVariable(FNiagaraTypeDefinition::GetBoolDef(), TEXT("Emitter.InterpolatedSpawn")))
+	{
+		OutConstant.SetValue(FNiagaraBool(false));
+		return true;
+	}
+	if (System && OutConstant == FNiagaraVariable(FNiagaraTypeDefinition::GetScriptUsageEnum(), TEXT("Script.Usage")))
 	{
 		FNiagaraInt32 EnumValue;
 		EnumValue.Value = (uint8)FNiagaraUtilities::ConvertScriptUsageToStaticSwitchUsage(Usage);
 		OutConstant.SetValue(EnumValue);
 		return true;
 	}
-	else if (System && OutConstant == FNiagaraVariable(FNiagaraTypeDefinition::GetScriptContextEnum(), TEXT("Script.Context")))
+	if (System && OutConstant == FNiagaraVariable(FNiagaraTypeDefinition::GetScriptContextEnum(), TEXT("Script.Context")))
 	{
 		FNiagaraInt32 EnumValue;
 		EnumValue.Value = (uint8)FNiagaraUtilities::ConvertScriptUsageToStaticSwitchContext(Usage);
