Avoid overwriting some bookkeeping data from niagara script compilation when compiling for GPU.

diff --git a/Engine/Plugins/FX/Niagara/Source/NiagaraEditor/Private/NiagaraCompiler.cpp b/Engine/Plugins/FX/Niagara/Source/NiagaraEditor/Private/NiagaraCompiler.cpp
--- a/Engine/Plugins/FX/Niagara/Source/NiagaraEditor/Private/NiagaraCompiler.cpp
+++ b/Engine/Plugins/FX/Niagara/Source/NiagaraEditor/Private/NiagaraCompiler.cpp
@@ -1594,15 +1594,10 @@ int32 FHlslNiagaraCompiler::CompileScript(const FNiagaraCompileRequestData* InCo
 
 	CompileResults.Data = MakeShared<FNiagaraVMExecutableData>();
 
-	CompileResults.AppendCompileEvents(MakeArrayView(InTranslateResults.CompileEvents));
-	CompileResults.Data->LastCompileEvents.Append(InTranslateResults.CompileEvents);
-	CompileResults.Data->ExternalDependencies = InTranslateResults.CompileDependencies;
-	CompileResults.Data->CompileTags = InTranslateResults.CompileTags;
-
 	//TODO: This should probably be done via the same route that other shaders take through the shader compiler etc.
 	//But that adds the complexity of a new shader type, new shader class and a new shader map to contain them etc.
 	//Can do things simply for now.
-	
+
 	CompileResults.Data->LastHlslTranslation = TEXT("");
 
 	FShaderCompilerInput Input;
@@ -1640,7 +1635,8 @@ int32 FHlslNiagaraCompiler::CompileScript(const FNiagaraCompileRequestData* InCo
 
 	// only issue jobs for VM compilation if we're going to be using the resulting byte code.  This excludes particle scripts when we're using
 	// a GPU simulation
-	if (InOptions.IsGpuScript() && UNiagaraScript::IsParticleScript(InOptions.TargetUsage))
+	const bool bCompilingGPUParticleScript = InOptions.IsGpuScript() && UNiagaraScript::IsParticleScript(InOptions.TargetUsage);
+	if (bCompilingGPUParticleScript)
 	{
 		CompileResults.bComputeSucceeded = false;
 		if (CompileResults.bVMSucceeded)
@@ -1653,8 +1649,19 @@ int32 FHlslNiagaraCompiler::CompileScript(const FNiagaraCompileRequestData* InCo
 		CompileResults.Data->CompileTags = InTranslateResults.CompileTags;
 		DumpDebugInfo(CompileResults, Input, true);
 		CompilationJob->CompileResults = CompileResults;
+	}
+
+	CompileResults.AppendCompileEvents(MakeArrayView(InTranslateResults.CompileEvents));
+	CompileResults.Data->LastCompileEvents.Append(InTranslateResults.CompileEvents);
+	CompileResults.Data->ExternalDependencies = InTranslateResults.CompileDependencies;
+	CompileResults.Data->CompileTags = InTranslateResults.CompileTags;
+
+	// Early out if compiling a GPU particle script as we do not need to submit a CPU compile request.
+	if (bCompilingGPUParticleScript)
+	{
 		return JobID;
 	}
+
 	CompilationJob->TranslatorOutput.ScriptData.LastHlslTranslation = TranslatedHLSL;
 	CompilationJob->TranslatorOutput.ScriptData.ExternalDependencies = InTranslateResults.CompileDependencies;
 	CompilationJob->TranslatorOutput.ScriptData.CompileTags = InTranslateResults.CompileTags;
