Slight type names changes to better match blueprint naming

diff --git a/Engine/Plugins/FX/Niagara/Source/Niagara/Public/NiagaraTypes.h b/Engine/Plugins/FX/Niagara/Source/Niagara/Public/NiagaraTypes.h
--- a/Engine/Plugins/FX/Niagara/Source/Niagara/Public/NiagaraTypes.h
+++ b/Engine/Plugins/FX/Niagara/Source/Niagara/Public/NiagaraTypes.h
@@ -873,11 +873,11 @@ public:
 			// fix names for common types after LWC changes
 			if (*this == GetVec2Def())
 			{
-				NameText = FText::FromString("Vector 2");
+				NameText = FText::FromString("Vector 2D");
 			}
 			else if (*this == GetVec3Def())
 			{
-				NameText = FText::FromString("Vector 3");
+				NameText = FText::FromString("Vector");
 			}
 			else if (*this == GetVec4Def())
 			{
