Niagara - Add support for filtering categories by section in the summary and in modules.  UI is similar to how the details panel handles this.  Also added support for drag/drop reordering array entries in the stack and other minor UI tweaks.

diff --git a/Engine/Plugins/FX/Niagara/Source/Niagara/Classes/NiagaraScript.h b/Engine/Plugins/FX/Niagara/Source/Niagara/Classes/NiagaraScript.h
--- a/Engine/Plugins/FX/Niagara/Source/Niagara/Classes/NiagaraScript.h
+++ b/Engine/Plugins/FX/Niagara/Source/Niagara/Classes/NiagaraScript.h
@@ -14,6 +14,7 @@
 #include "NiagaraDataSet.h"
 #include "NiagaraScriptExecutionParameterStore.h"
 #include "NiagaraScriptHighlight.h"
+#include "NiagaraStackSection.h"
 #include "NiagaraParameterDefinitionsSubscriber.h"
 #include "HAL/CriticalSection.h"
 #include "HAL/PlatformAtomics.h"
@@ -607,6 +608,9 @@ public:
 	UPROPERTY(EditAnywhere, Category = Script, DisplayName = "Script Metadata", meta = (ToolTip = "Script Metadata"))
 	TMap<FName, FString> ScriptMetaData;
 	
+	UPROPERTY(EditAnywhere, Category = Script)
+	TArray<FNiagaraStackSection> InputSections;
+	
 	/** Adjusted every time ComputeVMCompilationId is called.*/
 	UPROPERTY()
 	mutable FNiagaraVMExecutableDataId LastGeneratedVMId;
diff --git a/Engine/Plugins/FX/Niagara/Source/Niagara/Public/NiagaraStackSection.h b/Engine/Plugins/FX/Niagara/Source/Niagara/Public/NiagaraStackSection.h
new file mode 100644
--- /dev/null
+++ b/Engine/Plugins/FX/Niagara/Source/Niagara/Public/NiagaraStackSection.h
@@ -0,0 +1,21 @@
+// Copyright Epic Games, Inc. All Rights Reserved.
+
+#pragma once
+
+#include "Internationalization/Text.h"
+#include "UObject/ObjectMacros.h"
+
+#include "NiagaraStackSection.generated.h"
+
+/* Defines data for sections visible in the stack view. */
+USTRUCT()
+struct FNiagaraStackSection 
+{
+	GENERATED_BODY()
+
+	UPROPERTY(EditAnywhere, Category=Section)
+	FText SectionDisplayName;
+
+	UPROPERTY(EditAnywhere, Category=Section)
+	TArray<FText> Categories;
+};
\ No newline at end of file
diff --git a/Engine/Plugins/FX/Niagara/Source/NiagaraEditor/Private/NiagaraEmitterEditorData.cpp b/Engine/Plugins/FX/Niagara/Source/NiagaraEditor/Private/NiagaraEmitterEditorData.cpp
--- a/Engine/Plugins/FX/Niagara/Source/NiagaraEditor/Private/NiagaraEmitterEditorData.cpp
+++ b/Engine/Plugins/FX/Niagara/Source/NiagaraEditor/Private/NiagaraEmitterEditorData.cpp
@@ -5,6 +5,8 @@
 #include "NiagaraStackEditorData.h"
 #include "ScopedTransaction.h"
 
+const FName UNiagaraEmitterEditorData::PrivateMemberNames::SummarySections = GET_MEMBER_NAME_CHECKED(UNiagaraEmitterEditorData, SummarySections);
+
 UNiagaraEmitterEditorData::UNiagaraEmitterEditorData(const FObjectInitializer& ObjectInitializer)
 {
 	StackEditorData = ObjectInitializer.CreateDefaultSubobject<UNiagaraStackEditorData>(this, TEXT("StackEditorData"));
diff --git a/Engine/Plugins/FX/Niagara/Source/NiagaraEditor/Private/NiagaraStackEditorData.cpp b/Engine/Plugins/FX/Niagara/Source/NiagaraEditor/Private/NiagaraStackEditorData.cpp
--- a/Engine/Plugins/FX/Niagara/Source/NiagaraEditor/Private/NiagaraStackEditorData.cpp
+++ b/Engine/Plugins/FX/Niagara/Source/NiagaraEditor/Private/NiagaraStackEditorData.cpp
@@ -85,6 +85,20 @@ void UNiagaraStackEditorData::SetStackItemShowAdvanced(const FString& StackEntry
 	}
 }
 
+FText UNiagaraStackEditorData::GetStackEntryActiveSection(const FString& StackEntryKey, FText ActiveSectionDefault) const
+{
+	const FText* ActiveSectionPtr = StackEntryKeyToActiveSectionMap.Find(StackEntryKey);
+	return ActiveSectionPtr != nullptr ? *ActiveSectionPtr : ActiveSectionDefault;
+}
+
+void UNiagaraStackEditorData::SetStackEntryActiveSection(const FString& StackEntryKey, FText ActiveSection)
+{
+	if (ensureMsgf(StackEntryKey.IsEmpty() == false, TEXT("Can not set the active section with an empty key")))
+	{
+		StackEntryKeyToActiveSectionMap.FindOrAdd(StackEntryKey) = ActiveSection;
+	}
+}
+
 const FText* UNiagaraStackEditorData::GetStackEntryDisplayName(const FString& StackEntryKey) const
 {
 	return StackEntryKeyToDisplayName.Find(StackEntryKey);
diff --git a/Engine/Plugins/FX/Niagara/Source/NiagaraEditor/Private/ViewModels/NiagaraEmitterViewModel.cpp b/Engine/Plugins/FX/Niagara/Source/NiagaraEditor/Private/ViewModels/NiagaraEmitterViewModel.cpp
--- a/Engine/Plugins/FX/Niagara/Source/NiagaraEditor/Private/ViewModels/NiagaraEmitterViewModel.cpp
+++ b/Engine/Plugins/FX/Niagara/Source/NiagaraEditor/Private/ViewModels/NiagaraEmitterViewModel.cpp
@@ -71,7 +71,7 @@ bool FNiagaraEmitterViewModel::Initialize(UNiagaraEmitter* InEmitter, TWeakPtr<F
 {
 	SetEmitter(InEmitter);
 	SetSimulation(InSimulation);
-
+	bSummaryIsInEditMode = false;
 	return true;
 }
 
diff --git a/Engine/Plugins/FX/Niagara/Source/NiagaraEditor/Private/ViewModels/Stack/NiagaraStackEmitterSettingsGroup.cpp b/Engine/Plugins/FX/Niagara/Source/NiagaraEditor/Private/ViewModels/Stack/NiagaraStackEmitterSettingsGroup.cpp
--- a/Engine/Plugins/FX/Niagara/Source/NiagaraEditor/Private/ViewModels/Stack/NiagaraStackEmitterSettingsGroup.cpp
+++ b/Engine/Plugins/FX/Niagara/Source/NiagaraEditor/Private/ViewModels/Stack/NiagaraStackEmitterSettingsGroup.cpp
@@ -6,6 +6,7 @@
 #include "ViewModels/NiagaraSystemViewModel.h"
 #include "ViewModels/NiagaraEmitterViewModel.h"
 #include "NiagaraEmitter.h"
+#include "NiagaraEmitterEditorData.h"
 #include "NiagaraStackEditorData.h"
 #include "ViewModels/Stack/NiagaraStackGraphUtilities.h"
 #include "NiagaraScriptMergeManager.h"
@@ -14,6 +15,7 @@
 #include "NiagaraSystem.h"
 #include "NiagaraEditorStyle.h"
 #include "Styling/AppStyle.h"
+#include "IDetailTreeNode.h"
 
 #define LOCTEXT_NAMESPACE "UNiagaraStackEmitterItemGroup"
 
@@ -177,21 +179,29 @@ FText UNiagaraStackEmitterSummaryItem::GetTooltipText() const
 	return LOCTEXT("EmitterSummaryTooltip", "Subset of parameters from the stack, summarized here for easier access.");
 }
 
-const FSlateBrush* UNiagaraStackEmitterSummaryItem::GetIconBrush() const
+FText UNiagaraStackEmitterSummaryItem::GetIconText() const
 {
-	if (Emitter->SimTarget == ENiagaraSimTarget::CPUSim)
+	return FText::FromString(FString(TEXT("\xf0ca")/* fa-list-ul */));
+}
+
+void UNiagaraStackEmitterSummaryItem::RefreshChildrenInternal(const TArray<UNiagaraStackEntry*>& CurrentChildren, TArray<UNiagaraStackEntry*>& NewChildren, TArray<FStackIssue>& NewIssues)
+{
+	if (GetEmitterViewModel()->GetSummaryIsInEditMode())
 	{
-		return FNiagaraEditorStyle::Get().GetBrush("NiagaraEditor.Stack.CPUIcon");
+		if (SummaryEditorData == nullptr)
+		{
+			SummaryEditorData = NewObject<UNiagaraStackObject>(this);
+			SummaryEditorData->Initialize(CreateDefaultChildRequiredData(), &GetEmitterViewModel()->GetOrCreateEditorData(), GetStackEditorDataKey());
+			SummaryEditorData->SetOnSelectRootNodes(UNiagaraStackObject::FOnSelectRootNodes::CreateUObject(this,
+				&UNiagaraStackEmitterSummaryItem::SelectSummaryNodesFromEmitterEditorDataRootNodes));
+		}
+		NewChildren.Add(SummaryEditorData);
 	}
-	if (Emitter->SimTarget == ENiagaraSimTarget::GPUComputeSim)
+	else
 	{
-		return FNiagaraEditorStyle::Get().GetBrush("NiagaraEditor.Stack.GPUIcon");
+		SummaryEditorData = nullptr;
 	}
-	return FEditorStyle::GetBrush("NoBrush");
-}
 
-void UNiagaraStackEmitterSummaryItem::RefreshChildrenInternal(const TArray<UNiagaraStackEntry*>& CurrentChildren, TArray<UNiagaraStackEntry*>& NewChildren, TArray<FStackIssue>& NewIssues)
-{
 	if (FilteredObject == nullptr)
 	{
 		FilteredObject = NewObject<UNiagaraStackSummaryViewObject>(this);
@@ -203,6 +213,57 @@ void UNiagaraStackEmitterSummaryItem::RefreshChildrenInternal(const TArray<UNiag
 	Super::RefreshChildrenInternal(CurrentChildren, NewChildren, NewIssues);
 }
 
+TSharedPtr<IDetailTreeNode> GetSummarySectionsPropertyNode(const TArray<TSharedRef<IDetailTreeNode>>& Nodes)
+{
+	TArray<TSharedRef<IDetailTreeNode>> ChildrenToCheck;
+	for (TSharedRef<IDetailTreeNode> Node : Nodes)
+	{
+		if (Node->GetNodeType() == EDetailNodeType::Item)
+		{
+			TSharedPtr<IPropertyHandle> NodePropertyHandle = Node->CreatePropertyHandle();
+			if (NodePropertyHandle.IsValid() && NodePropertyHandle->GetProperty()->GetFName() == UNiagaraEmitterEditorData::PrivateMemberNames::SummarySections)
+			{
+				return Node;
+			}
+		}
+
+		TArray<TSharedRef<IDetailTreeNode>> Children;
+		Node->GetChildren(Children);
+		ChildrenToCheck.Append(Children);
+	}
+	if (ChildrenToCheck.Num() == 0)
+	{
+		return nullptr;
+	}
+	return GetSummarySectionsPropertyNode(ChildrenToCheck);
+}
+
+void UNiagaraStackEmitterSummaryItem::SelectSummaryNodesFromEmitterEditorDataRootNodes(TArray<TSharedRef<IDetailTreeNode>> Source, TArray<TSharedRef<IDetailTreeNode>>* Selected)
+{
+	TSharedPtr<IDetailTreeNode> SummarySectionsNode = GetSummarySectionsPropertyNode(Source);
+	if (SummarySectionsNode.IsValid())
+	{
+		Selected->Add(SummarySectionsNode.ToSharedRef());
+	}
+}
+
+bool UNiagaraStackEmitterSummaryItem::GetEditModeIsActive() const
+{ 
+	return GetEmitterViewModel().IsValid() && GetEmitterViewModel()->GetSummaryIsInEditMode();
+}
+
+void UNiagaraStackEmitterSummaryItem::SetEditModeIsActive(bool bInEditModeIsActive)
+{
+	if (GetEmitterViewModel().IsValid())
+	{
+		if (GetEmitterViewModel()->GetSummaryIsInEditMode() != bInEditModeIsActive)
+		{
+			GetEmitterViewModel()->SetSummaryIsInEditMode(bInEditModeIsActive);
+			RefreshChildren();
+		}
+	}
+}
+
 UNiagaraStackEmitterSummaryGroup::UNiagaraStackEmitterSummaryGroup()
 	: SummaryItem(nullptr)
 {
@@ -220,9 +281,9 @@ void UNiagaraStackEmitterSummaryGroup::RefreshChildrenInternal(const TArray<UNia
 	Super::RefreshChildrenInternal(CurrentChildren, NewChildren, NewIssues);
 }
 
-const FSlateBrush* UNiagaraStackEmitterSummaryGroup::GetIconBrush() const
+FText UNiagaraStackEmitterSummaryGroup::GetIconText() const
 {
-	return FAppStyle::Get().GetBrush("Icons.Edit");
+	return FText::FromString(FString(TEXT("\xf0ca")/* fa-list-ul */));
 }
 
 
diff --git a/Engine/Plugins/FX/Niagara/Source/NiagaraEditor/Private/ViewModels/Stack/NiagaraStackEntry.cpp b/Engine/Plugins/FX/Niagara/Source/NiagaraEditor/Private/ViewModels/Stack/NiagaraStackEntry.cpp
--- a/Engine/Plugins/FX/Niagara/Source/NiagaraEditor/Private/ViewModels/Stack/NiagaraStackEntry.cpp
+++ b/Engine/Plugins/FX/Niagara/Source/NiagaraEditor/Private/ViewModels/Stack/NiagaraStackEntry.cpp
@@ -17,7 +17,6 @@ const FName UNiagaraStackEntry::FExecutionCategoryNames::Particle = TEXT("Partic
 const FName UNiagaraStackEntry::FExecutionCategoryNames::Render = TEXT("Render");
 
 const FName UNiagaraStackEntry::FExecutionSubcategoryNames::Settings = TEXT("Settings");
-const FName UNiagaraStackEntry::FExecutionSubcategoryNames::Summary = TEXT("Summary");
 const FName UNiagaraStackEntry::FExecutionSubcategoryNames::Spawn = TEXT("Spawn");
 const FName UNiagaraStackEntry::FExecutionSubcategoryNames::Update = TEXT("Update");
 const FName UNiagaraStackEntry::FExecutionSubcategoryNames::Event = TEXT("Event");
diff --git a/Engine/Plugins/FX/Niagara/Source/NiagaraEditor/Private/ViewModels/Stack/NiagaraStackFilteredObject.cpp b/Engine/Plugins/FX/Niagara/Source/NiagaraEditor/Private/ViewModels/Stack/NiagaraStackFilteredObject.cpp
--- a/Engine/Plugins/FX/Niagara/Source/NiagaraEditor/Private/ViewModels/Stack/NiagaraStackFilteredObject.cpp
+++ b/Engine/Plugins/FX/Niagara/Source/NiagaraEditor/Private/ViewModels/Stack/NiagaraStackFilteredObject.cpp
@@ -30,13 +30,13 @@ UNiagaraStackFilteredObject::UNiagaraStackFilteredObject()
 {
 }
 
-void UNiagaraStackFilteredObject::Initialize(FRequiredEntryData InRequiredEntryData, UNiagaraEmitter* InEmitter, FString InOwningStackItemEditorDataKey)
+void UNiagaraStackFilteredObject::Initialize(FRequiredEntryData InRequiredEntryData, FString InOwningStackItemEditorDataKey)
 {
 	checkf(Emitter == nullptr, TEXT("Can only initialize once."));
 	FString ObjectStackEditorDataKey = FString::Printf(TEXT("%s-FilteredView"), *InOwningStackItemEditorDataKey);
 	Super::Initialize(InRequiredEntryData, InOwningStackItemEditorDataKey, ObjectStackEditorDataKey);
 
-	Emitter = InEmitter;
+	Emitter = GetEmitterViewModel()->GetEmitter();
 	if (Emitter->GetEditorData())
 	{
 		Cast<UNiagaraEmitterEditorData>(Emitter->GetEditorData())->OnSummaryViewStateChanged().AddUObject(this, &UNiagaraStackFilteredObject::OnViewStateChanged);
diff --git a/Engine/Plugins/FX/Niagara/Source/NiagaraEditor/Private/ViewModels/Stack/NiagaraStackFunctionInput.cpp b/Engine/Plugins/FX/Niagara/Source/NiagaraEditor/Private/ViewModels/Stack/NiagaraStackFunctionInput.cpp
--- a/Engine/Plugins/FX/Niagara/Source/NiagaraEditor/Private/ViewModels/Stack/NiagaraStackFunctionInput.cpp
+++ b/Engine/Plugins/FX/Niagara/Source/NiagaraEditor/Private/ViewModels/Stack/NiagaraStackFunctionInput.cpp
@@ -613,7 +613,7 @@ void UNiagaraStackFunctionInput::RefreshChildrenInternal(const TArray<UNiagaraSt
 			{
 				DynamicInputEntry = NewObject<UNiagaraStackFunctionInputCollection>(this);
 				DynamicInputEntry->Initialize(CreateDefaultChildRequiredData(), *OwningModuleNode, *InputValues.DynamicNode.Get(), GetOwnerStackItemEditorDataKey());
-				DynamicInputEntry->SetShouldShowInStack(false);
+				DynamicInputEntry->SetShouldDisplayLabel(false);
 			}
 
 			if (ScriptData != nullptr)
diff --git a/Engine/Plugins/FX/Niagara/Source/NiagaraEditor/Private/ViewModels/Stack/NiagaraStackFunctionInputCollection.cpp b/Engine/Plugins/FX/Niagara/Source/NiagaraEditor/Private/ViewModels/Stack/NiagaraStackFunctionInputCollection.cpp
--- a/Engine/Plugins/FX/Niagara/Source/NiagaraEditor/Private/ViewModels/Stack/NiagaraStackFunctionInputCollection.cpp
+++ b/Engine/Plugins/FX/Niagara/Source/NiagaraEditor/Private/ViewModels/Stack/NiagaraStackFunctionInputCollection.cpp
@@ -25,6 +25,8 @@
 
 FText UNiagaraStackFunctionInputCollectionBase::UncategorizedName = LOCTEXT("Uncategorized", "Uncategorized");
 
+FText UNiagaraStackFunctionInputCollectionBase::AllSectionName = LOCTEXT("All", "All");
+
 static FText GetUserFriendlyFunctionName(UNiagaraNodeFunctionCall* Node)
 {
 	if (Node->IsA<UNiagaraNodeAssignment>())
@@ -69,7 +71,6 @@ static bool ShouldShowInSummaryView(UNiagaraEmitter* Emitter, UNiagaraNodeFuncti
 UNiagaraStackFunctionInputCollection::UNiagaraStackFunctionInputCollection()
 	: ModuleNode(nullptr)
 	, InputFunctionCallNode(nullptr)
-	, bShouldShowInStack(true)
 {
 }
 
@@ -108,21 +109,11 @@ FText UNiagaraStackFunctionInputCollection::GetDisplayName() const
 	return LOCTEXT("InputCollectionDisplayName", "Inputs");
 }
 
-bool UNiagaraStackFunctionInputCollection::GetShouldShowInStack() const
-{
-	return bShouldShowInStack;
-}
-
 bool UNiagaraStackFunctionInputCollection::GetIsEnabled() const
 {
 	return InputFunctionCallNode->GetDesiredEnabledState() == ENodeEnabledState::Enabled;
 }
 
-void UNiagaraStackFunctionInputCollection::SetShouldShowInStack(bool bInShouldShowInStack)
-{
-	bShouldShowInStack = bInShouldShowInStack;
-}
-
 void UNiagaraStackFunctionInputCollection::ToClipboardFunctionInputs(UObject* InOuter, TArray<const UNiagaraClipboardFunctionInput*>& OutClipboardFunctionInputs) const
 {
 	TArray<UNiagaraStackInputCategory*> ChildCategories;
@@ -197,13 +188,72 @@ TArray<UNiagaraStackFunctionInput*> UNiagaraStackFunctionInputCollection::GetInl
 void UNiagaraStackFunctionInputCollection::RefreshChildrenInternal(const TArray<UNiagaraStackEntry*>& CurrentChildren, TArray<UNiagaraStackEntry*>& NewChildren, TArray<FStackIssue>& NewIssues)
 {
 	RefreshChildrenForFunctionCall(ModuleNode, InputFunctionCallNode, CurrentChildren, NewChildren, NewIssues, false);
+	Super::RefreshChildrenInternal(CurrentChildren, NewChildren, NewIssues);
 }
 
+void UNiagaraStackFunctionInputCollection::GetSectionsInternal(TArray<FNiagaraStackSection>& OutStackSections) const
+{
+	if (InputFunctionCallNode->GetCalledUsage() == ENiagaraScriptUsage::Module)
+	{
+		OutStackSections.Append(InputFunctionCallNode->GetScriptData()->InputSections);
+	}
+}
 
-UNiagaraStackFunctionInputCollectionBase::UNiagaraStackFunctionInputCollectionBase()
-{	
+void UNiagaraStackFunctionInputCollectionBase::Initialize(FRequiredEntryData InRequiredEntryData, FString InOwningStackItemEditorDataKey, FString InStackEditorDataKey)
+{
+	Super::Initialize(InRequiredEntryData, InOwningStackItemEditorDataKey, InStackEditorDataKey);
+	AddChildFilter(FOnFilterChild::CreateUObject(this, &UNiagaraStackFunctionInputCollectionBase::FilterByActiveSection));
+	ActiveSectionCache = GetStackEditorData().GetStackEntryActiveSection(GetStackEditorDataKey(), AllSectionName);
 }
 
+void UNiagaraStackFunctionInputCollectionBase::SetShouldDisplayLabel(bool bInShouldDisplayLabel)
+{
+	bShouldDisplayLabel = bInShouldDisplayLabel;
+}
+
+const TArray<FText>& UNiagaraStackFunctionInputCollectionBase::GetSections() const
+{
+	if (SectionsCache.IsSet() == false)
+	{
+		UpdateCachedSectionData();
+	}
+	return SectionsCache.GetValue();
+}
+
+FText UNiagaraStackFunctionInputCollectionBase::GetActiveSection() const
+{
+	if (ActiveSectionCache.IsSet() == false)
+	{
+		UpdateCachedSectionData();
+	}
+	return ActiveSectionCache.GetValue();
+}
+
+void UNiagaraStackFunctionInputCollectionBase::SetActiveSection(FText InActiveSection)
+{
+	ActiveSectionCache = InActiveSection;
+	GetStackEditorData().SetStackEntryActiveSection(GetStackEditorDataKey(), InActiveSection);
+	RefreshFilteredChildren();
+}
+
+bool UNiagaraStackFunctionInputCollectionBase::GetCanExpand() const
+{
+	return bShouldDisplayLabel;
+}
+
+bool UNiagaraStackFunctionInputCollectionBase::GetShouldShowInStack() const
+{
+	return GetSections().Num() > 0 || bShouldDisplayLabel;
+}
+
+void UNiagaraStackFunctionInputCollectionBase::RefreshChildrenInternal(const TArray<UNiagaraStackEntry*>& CurrentChildren, TArray<UNiagaraStackEntry*>& NewChildren, TArray<FStackIssue>& NewIssues)
+{
+	Super::RefreshChildrenInternal(CurrentChildren, NewChildren, NewIssues);
+	LastActiveSection = ActiveSectionCache.IsSet() ? ActiveSectionCache.GetValue() : AllSectionName;
+	SectionsCache.Reset();
+	SectionToCategoryMapCache.Reset();
+	ActiveSectionCache.Reset();
+}
 
 void UNiagaraStackFunctionInputCollectionBase::RefreshChildrenForFunctionCall(UNiagaraNodeFunctionCall* ModuleNode, UNiagaraNodeFunctionCall* InputFunctionCallNode, 
 	const TArray<UNiagaraStackEntry*>& CurrentChildren, TArray<UNiagaraStackEntry*>& NewChildren, TArray<FStackIssue>& NewIssues, bool bShouldApplySummaryFilter)
@@ -716,6 +766,85 @@ void UNiagaraStackFunctionInputCollectionBase::AddInputToCategory(const FInputDa
 
 }
 
+bool UNiagaraStackFunctionInputCollectionBase::FilterByActiveSection(const UNiagaraStackEntry& Child) const
+{
+	const TArray<FText>& Sections = GetSections();
+	FText ActiveSection = GetActiveSection();
+	if (Sections.Num() == 0 || ActiveSection.IdenticalTo(AllSectionName) || SectionToCategoryMapCache.IsSet() == false)
+	{
+		return true;
+	}
+
+	TArray<FText>* ActiveCategoryNames = SectionToCategoryMapCache->Find(ActiveSection.ToString());
+	const UNiagaraStackInputCategory* ChildCategory = Cast<UNiagaraStackInputCategory>(&Child);
+	return ChildCategory == nullptr || ActiveCategoryNames == nullptr || ActiveCategoryNames->ContainsByPredicate(
+		[ChildCategory](const FText& ActiveCategoryName) { return ActiveCategoryName.EqualTo(ChildCategory->GetDisplayName()); });
+}
+
+void UNiagaraStackFunctionInputCollectionBase::UpdateCachedSectionData() const
+{
+	TArray<FText> Sections;
+	TMap<FString, TArray<FText>> SectionToCategoryMap;
+	TArray<FNiagaraStackSection> StackSections;
+	GetSectionsInternal(StackSections);
+
+	if (StackSections.Num() > 0)
+	{
+		// Get the current list of categories.
+		TArray<FText> CategoryNames;
+		TArray<UNiagaraStackInputCategory*> ChildCategories;
+		GetUnfilteredChildrenOfType(ChildCategories);
+		
+		for (UNiagaraStackInputCategory* ChildCategory : ChildCategories)
+		{
+			if (ChildCategory->GetShouldShowInStack())
+			{
+				CategoryNames.Add(ChildCategory->GetDisplayName());
+			}
+		}
+
+		// Match sections to valid categories.
+		for (const FNiagaraStackSection& StackSection : StackSections)
+		{
+			TArray<FText> ContainedCategories;
+			for (FText SectionCategory : StackSection.Categories)
+			{
+				if (CategoryNames.ContainsByPredicate([SectionCategory](FText CategoryName) { return CategoryName.EqualTo(SectionCategory); }))
+				{
+					ContainedCategories.Add(SectionCategory);
+				}
+			}
+			if (ContainedCategories.Num() > 0)
+			{
+				Sections.Add(StackSection.SectionDisplayName);
+				SectionToCategoryMap.Add(StackSection.SectionDisplayName.ToString(), ContainedCategories);
+			}
+		}
+
+		Sections.Add(AllSectionName);
+		SectionToCategoryMap.Add(AllSectionName.ToString(), CategoryNames);
+
+		if (Sections.Num() == 1)
+		{
+			// If there is only one section, it's the "All" section which is not useful.
+			Sections.Empty();
+			SectionToCategoryMap.Empty();
+		}
+	}
+
+	SectionsCache = Sections;
+	SectionToCategoryMapCache = SectionToCategoryMap;
+	FText LastActiveSectionLocal = LastActiveSection;
+	if (Sections.ContainsByPredicate([LastActiveSectionLocal](FText Section) { return Section.EqualTo(LastActiveSectionLocal); }))
+	{
+		ActiveSectionCache = LastActiveSection;
+	}
+	else
+	{
+		ActiveSectionCache = AllSectionName;
+	}
+}
+
 UNiagaraStackEntry::FStackIssueFix UNiagaraStackFunctionInputCollectionBase::GetNodeRemovalFix(UEdGraphPin* PinToRemove, FText FixDescription)
 {
 	return FStackIssueFix(
diff --git a/Engine/Plugins/FX/Niagara/Source/NiagaraEditor/Private/ViewModels/Stack/NiagaraStackItem.cpp b/Engine/Plugins/FX/Niagara/Source/NiagaraEditor/Private/ViewModels/Stack/NiagaraStackItem.cpp
--- a/Engine/Plugins/FX/Niagara/Source/NiagaraEditor/Private/ViewModels/Stack/NiagaraStackItem.cpp
+++ b/Engine/Plugins/FX/Niagara/Source/NiagaraEditor/Private/ViewModels/Stack/NiagaraStackItem.cpp
@@ -47,11 +47,6 @@ void UNiagaraStackItem::SetIsEnabled(bool bInIsEnabled)
 	SetIsEnabledInternal(bInIsEnabled);
 }
 
-const FSlateBrush* UNiagaraStackItem::GetIconBrush() const
-{
-	return FEditorStyle::GetBrush("NoBrush");
-}
-
 void UNiagaraStackItem::RefreshChildrenInternal(const TArray<UNiagaraStackEntry*>& CurrentChildren, TArray<UNiagaraStackEntry*>& NewChildren, TArray<FStackIssue>& NewIssues)
 {
 	if (ItemFooter == nullptr)
diff --git a/Engine/Plugins/FX/Niagara/Source/NiagaraEditor/Private/ViewModels/Stack/NiagaraStackModuleItem.cpp b/Engine/Plugins/FX/Niagara/Source/NiagaraEditor/Private/ViewModels/Stack/NiagaraStackModuleItem.cpp
--- a/Engine/Plugins/FX/Niagara/Source/NiagaraEditor/Private/ViewModels/Stack/NiagaraStackModuleItem.cpp
+++ b/Engine/Plugins/FX/Niagara/Source/NiagaraEditor/Private/ViewModels/Stack/NiagaraStackModuleItem.cpp
@@ -179,7 +179,7 @@ void UNiagaraStackModuleItem::RefreshChildrenInternal(const TArray<UNiagaraStack
 				OutputCollection->AddChildFilter(FOnFilterChild::CreateUObject(this, &UNiagaraStackModuleItem::FilterOutputCollectionChild));
 			}
 
-			InputCollection->SetShouldShowInStack(GetStackEditorData().GetShowOutputs() || GetStackEditorData().GetShowLinkedInputs());
+			InputCollection->SetShouldDisplayLabel(GetStackEditorData().GetShowOutputs() || GetStackEditorData().GetShowLinkedInputs());
 
 			NewChildren.Add(InputCollection);
 			NewChildren.Add(LinkedInputCollection);
@@ -189,7 +189,7 @@ void UNiagaraStackModuleItem::RefreshChildrenInternal(const TArray<UNiagaraStack
 		else
 		{
 			// We do not show the expander arrow for InputCollections of NiagaraNodeAssignments as they only have this one collection
-			InputCollection->SetShouldShowInStack(false);
+			InputCollection->SetShouldDisplayLabel(false);
 
 			NewChildren.Add(InputCollection);
 
diff --git a/Engine/Plugins/FX/Niagara/Source/NiagaraEditor/Private/ViewModels/Stack/NiagaraStackPropertyRow.cpp b/Engine/Plugins/FX/Niagara/Source/NiagaraEditor/Private/ViewModels/Stack/NiagaraStackPropertyRow.cpp
--- a/Engine/Plugins/FX/Niagara/Source/NiagaraEditor/Private/ViewModels/Stack/NiagaraStackPropertyRow.cpp
+++ b/Engine/Plugins/FX/Niagara/Source/NiagaraEditor/Private/ViewModels/Stack/NiagaraStackPropertyRow.cpp
@@ -58,6 +58,12 @@ bool UNiagaraStackPropertyRow::IsExpandedByDefault() const
 	return DetailTreeNode->GetInitiallyCollapsed() == false;
 }
 
+bool UNiagaraStackPropertyRow::CanDrag() const
+{
+	TSharedPtr<IPropertyHandle> PropertyHandle = DetailTreeNode->CreatePropertyHandle();
+	return PropertyHandle.IsValid() && PropertyHandle->GetParentHandle().IsValid() && PropertyHandle->GetParentHandle()->AsArray().IsValid();
+}
+
 void UNiagaraStackPropertyRow::FinalizeInternal()
 {
 	Super::FinalizeInternal();
@@ -87,6 +93,7 @@ void UNiagaraStackPropertyRow::RefreshChildrenInternal(const TArray<UNiagaraStac
 		NewChildren.Add(ChildRow);
 	}
 }
+
 void UNiagaraStackPropertyRow::GetSearchItems(TArray<FStackSearchItem>& SearchItems) const
 {
 	SearchItems.Add({ FName("DisplayName"), GetDisplayName() });
@@ -110,3 +117,66 @@ void UNiagaraStackPropertyRow::GetSearchItems(TArray<FStackSearchItem>& SearchIt
 		}
 	}	
 }
+
+TOptional<UNiagaraStackEntry::FDropRequestResponse> UNiagaraStackPropertyRow::CanDropInternal(const FDropRequest& DropRequest)
+{
+	// Validate stack, drop zone, and drag type.
+	if (DropRequest.DropOptions == UNiagaraStackEntry::EDropOptions::Overview ||
+		DropRequest.DropZone == EItemDropZone::OntoItem ||
+		DropRequest.DragDropOperation->IsOfType<FNiagaraStackEntryDragDropOp>() == false)
+	{
+		return TOptional<FDropRequestResponse>();
+	}
+
+	// Validate stack entry count and type.
+	TSharedRef<const FNiagaraStackEntryDragDropOp> StackEntryDragDropOp = StaticCastSharedRef<const FNiagaraStackEntryDragDropOp>(DropRequest.DragDropOperation);
+	if (StackEntryDragDropOp->GetDraggedEntries().Num() != 1 || StackEntryDragDropOp->GetDraggedEntries()[0]->IsA<UNiagaraStackPropertyRow>() == false)
+	{
+		return TOptional<FDropRequestResponse>();
+	}
+
+	auto HaveSameParent = [](TSharedPtr<IPropertyHandle> HandleA, TSharedPtr<IPropertyHandle> HandleB)
+	{
+		TSharedPtr<IPropertyHandle> ParentA = HandleA->GetParentHandle();
+		TSharedPtr<IPropertyHandle> ParentB = HandleB->GetParentHandle();
+		if (ParentA.IsValid() && ParentB.IsValid() && ParentA->GetProperty() == ParentB->GetProperty())
+		{
+			TArray<UObject*> OuterObjectsA;
+			ParentA->GetOuterObjects(OuterObjectsA);
+			TArray<UObject*> OuterObjectsB;
+			ParentB->GetOuterObjects(OuterObjectsB);
+			return OuterObjectsA == OuterObjectsB;
+		}
+		return false;
+	};
+
+	// Validate property handle.
+	UNiagaraStackPropertyRow* DraggedPropertyRow = CastChecked<UNiagaraStackPropertyRow>(StackEntryDragDropOp->GetDraggedEntries()[0]);
+	TSharedPtr<IPropertyHandle> DraggedPropertyHandle = DraggedPropertyRow->GetDetailTreeNode()->CreatePropertyHandle();
+	TSharedPtr<IPropertyHandle> TargetPropertyHandle = DetailTreeNode->CreatePropertyHandle();
+	if (DraggedPropertyHandle.IsValid() == false ||
+		TargetPropertyHandle.IsValid() == false ||
+		DraggedPropertyHandle == TargetPropertyHandle ||
+		DraggedPropertyHandle->GetParentHandle().IsValid() == false ||
+		HaveSameParent(DraggedPropertyHandle, TargetPropertyHandle) == false ||
+		DraggedPropertyHandle->GetParentHandle()->AsArray().IsValid() == false)
+	{
+		return TOptional<FDropRequestResponse>();
+	}
+	return FDropRequestResponse(DropRequest.DropZone, NSLOCTEXT("NiagaraStackPropertyRow", "DropArrayItemMessage", "Move this array entry here."));
+}
+
+TOptional<UNiagaraStackEntry::FDropRequestResponse> UNiagaraStackPropertyRow::DropInternal(const FDropRequest& DropRequest)
+{
+	TOptional<FDropRequestResponse> CanDropResponse = CanDropInternal(DropRequest);
+	if (CanDropResponse.IsSet())
+	{
+		TSharedRef<const FNiagaraStackEntryDragDropOp> StackEntryDragDropOp = StaticCastSharedRef<const FNiagaraStackEntryDragDropOp>(DropRequest.DragDropOperation);
+		UNiagaraStackPropertyRow* DraggedPropertyRow = CastChecked<UNiagaraStackPropertyRow>(StackEntryDragDropOp->GetDraggedEntries()[0]);
+		TSharedPtr<IPropertyHandle> DraggedPropertyHandle = DraggedPropertyRow->GetDetailTreeNode()->CreatePropertyHandle();
+		TSharedPtr<IPropertyHandle> TargetPropertyHandle = DetailTreeNode->CreatePropertyHandle();
+		int32 IndexOffset = DropRequest.DropZone == EItemDropZone::AboveItem ? 0 : 1;
+		TargetPropertyHandle->GetParentHandle()->AsArray()->MoveElementTo(DraggedPropertyHandle->GetIndexInArray(), TargetPropertyHandle->GetIndexInArray() + IndexOffset);
+	}
+	return CanDropResponse;
+}
diff --git a/Engine/Plugins/FX/Niagara/Source/NiagaraEditor/Private/ViewModels/Stack/NiagaraStackRoot.cpp b/Engine/Plugins/FX/Niagara/Source/NiagaraEditor/Private/ViewModels/Stack/NiagaraStackRoot.cpp
--- a/Engine/Plugins/FX/Niagara/Source/NiagaraEditor/Private/ViewModels/Stack/NiagaraStackRoot.cpp
+++ b/Engine/Plugins/FX/Niagara/Source/NiagaraEditor/Private/ViewModels/Stack/NiagaraStackRoot.cpp
@@ -139,7 +139,7 @@ void UNiagaraStackRoot::RefreshChildrenInternal(const TArray<UNiagaraStackEntry*
 	{
 		EmitterSummaryGroup = NewObject<UNiagaraStackEmitterSummaryGroup>(this);
 		FRequiredEntryData RequiredEntryData(GetSystemViewModel(), GetEmitterViewModel(),
-			FExecutionCategoryNames::Emitter, FExecutionSubcategoryNames::Summary,
+			FExecutionCategoryNames::Emitter, FExecutionSubcategoryNames::Settings,
 			GetEmitterViewModel()->GetOrCreateEditorData().GetStackEditorData());
 		FText DisplayName = LOCTEXT("EmitterSummaryGroupName", "Emitter Summary");
 		FText Tooltip = LOCTEXT("EmitterSummaryTooltip", "Summary of parameters for this Emitter.");
@@ -203,7 +203,7 @@ void UNiagaraStackRoot::RefreshChildrenInternal(const TArray<UNiagaraStackEntry*
 	{
 		SummaryCollapseButton = NewObject<UNiagaraStackSummaryViewCollapseButton>(this);
 		FRequiredEntryData RequiredEntryData(GetSystemViewModel(), GetEmitterViewModel(),
-			FExecutionCategoryNames::Emitter, FExecutionSubcategoryNames::Summary,
+			FExecutionCategoryNames::Emitter, FExecutionSubcategoryNames::Settings,
 			GetEmitterViewModel()->GetOrCreateEditorData().GetStackEditorData());
 		SummaryCollapseButton->Initialize(RequiredEntryData);
 	}
diff --git a/Engine/Plugins/FX/Niagara/Source/NiagaraEditor/Private/ViewModels/Stack/NiagaraStackSummaryViewInputCollection.cpp b/Engine/Plugins/FX/Niagara/Source/NiagaraEditor/Private/ViewModels/Stack/NiagaraStackSummaryViewInputCollection.cpp
--- a/Engine/Plugins/FX/Niagara/Source/NiagaraEditor/Private/ViewModels/Stack/NiagaraStackSummaryViewInputCollection.cpp
+++ b/Engine/Plugins/FX/Niagara/Source/NiagaraEditor/Private/ViewModels/Stack/NiagaraStackSummaryViewInputCollection.cpp
@@ -52,11 +52,6 @@ FText UNiagaraStackSummaryViewObject::GetDisplayName() const
 	return LOCTEXT("FilteredInputCollectionDisplayName", "Filtered Inputs");
 }
 
-bool UNiagaraStackSummaryViewObject::GetShouldShowInStack() const
-{
-	return true;
-}
-
 bool UNiagaraStackSummaryViewObject::GetIsEnabled() const
 {
 	return true;
@@ -133,6 +128,14 @@ void UNiagaraStackSummaryViewObject::PostRefreshChildrenInternal()
 	Super::PostRefreshChildrenInternal();
 }
 
+void UNiagaraStackSummaryViewObject::GetSectionsInternal(TArray<FNiagaraStackSection>& OutStackSections) const
+{
+	if (GetEmitterViewModel().IsValid())
+	{
+		OutStackSections.Append(GetEmitterViewModel()->GetOrCreateEditorData().GetSummarySections());
+	}
+}
+
 void UNiagaraStackSummaryViewObject::OnViewStateChanged()
 {
 	if (!IsFinalized())
diff --git a/Engine/Plugins/FX/Niagara/Source/NiagaraEditor/Private/ViewModels/Stack/NiagaraStackSystemSettingsGroup.cpp b/Engine/Plugins/FX/Niagara/Source/NiagaraEditor/Private/ViewModels/Stack/NiagaraStackSystemSettingsGroup.cpp
--- a/Engine/Plugins/FX/Niagara/Source/NiagaraEditor/Private/ViewModels/Stack/NiagaraStackSystemSettingsGroup.cpp
+++ b/Engine/Plugins/FX/Niagara/Source/NiagaraEditor/Private/ViewModels/Stack/NiagaraStackSystemSettingsGroup.cpp
@@ -138,9 +138,9 @@ void UNiagaraStackSystemUserParametersGroup::Initialize(
 	UserParameterStore = InParameterStore;
 }
 
-const FSlateBrush* UNiagaraStackSystemUserParametersGroup::GetIconBrush() const
+FText UNiagaraStackSystemUserParametersGroup::GetIconText() const
 {
-	return FAppStyle::Get().GetBrush("Icons.Edit");
+	return FText::FromString(FString(TEXT("\xf007")/* fa-user */));
 }
 
 void UNiagaraStackSystemUserParametersGroup::RefreshChildrenInternal(const TArray<UNiagaraStackEntry*>& CurrentChildren, TArray<UNiagaraStackEntry*>& NewChildren, TArray<FStackIssue>& NewIssues)
diff --git a/Engine/Plugins/FX/Niagara/Source/NiagaraEditor/Public/NiagaraEmitterEditorData.h b/Engine/Plugins/FX/Niagara/Source/NiagaraEditor/Public/NiagaraEmitterEditorData.h
--- a/Engine/Plugins/FX/Niagara/Source/NiagaraEditor/Public/NiagaraEmitterEditorData.h
+++ b/Engine/Plugins/FX/Niagara/Source/NiagaraEditor/Public/NiagaraEmitterEditorData.h
@@ -4,6 +4,7 @@
 
 #include "NiagaraEditorDataBase.h"
 #include "NiagaraEditorCommon.h"
+#include "NiagaraStackSection.h"
 #include "NiagaraEmitterEditorData.generated.h"
 
 class UNiagaraStackEditorData;
@@ -15,6 +16,12 @@ class UNiagaraEmitterEditorData : public UNiagaraEditorDataBase
 {
 	GENERATED_BODY()
 
+public: 
+	struct PrivateMemberNames
+	{
+		static const FName SummarySections;
+	};
+
 public:
 	UNiagaraEmitterEditorData(const FObjectInitializer& ObjectInitializer);
 
@@ -34,6 +41,9 @@ public:
 	NIAGARAEDITOR_API TMap<FFunctionInputSummaryViewKey, FFunctionInputSummaryViewMetadata> GetSummaryViewMetaDataMap() const;
 	NIAGARAEDITOR_API FFunctionInputSummaryViewMetadata GetSummaryViewMetaData(const FFunctionInputSummaryViewKey& Key) const;
 	NIAGARAEDITOR_API void SetSummaryViewMetaData(const FFunctionInputSummaryViewKey& Key, const FFunctionInputSummaryViewMetadata& NewMetadata);
+
+	const TArray<FNiagaraStackSection> GetSummarySections() const { return SummarySections; }
+
 private:
 	UPROPERTY(Instanced)
 	TObjectPtr<UNiagaraStackEditorData> StackEditorData;
@@ -51,6 +61,9 @@ private:
 	UPROPERTY()
 	TMap<FFunctionInputSummaryViewKey, FFunctionInputSummaryViewMetadata> SummaryViewFunctionInputMetadata;
 
+	UPROPERTY(EditAnywhere, Category=Summary, meta=(NiagaraNoMerge))
+	TArray<FNiagaraStackSection> SummarySections;
+
 	FSimpleMulticastDelegate OnSummaryViewStateChangedDelegate;
 	
 	void StackEditorDataChanged();	
diff --git a/Engine/Plugins/FX/Niagara/Source/NiagaraEditor/Public/NiagaraStackEditorData.h b/Engine/Plugins/FX/Niagara/Source/NiagaraEditor/Public/NiagaraStackEditorData.h
--- a/Engine/Plugins/FX/Niagara/Source/NiagaraEditor/Public/NiagaraStackEditorData.h
+++ b/Engine/Plugins/FX/Niagara/Source/NiagaraEditor/Public/NiagaraStackEditorData.h
@@ -83,6 +83,20 @@ public:
 	*/
 	void SetStackItemShowAdvanced(const FString& StackEntryKey, bool bShowAdanced);
 
+	/*
+	* Gets the active section for a stack entry.
+	* @param StackItemKey A unique key for the entry.
+	* @param ActiveSectionDefault The default value to return if an active section has not been set before.
+	*/
+	FText GetStackEntryActiveSection(const FString& StackEntryKey, FText ActiveSectionDefault) const;
+
+	/*
+	* Sets the active section for a stack entry.
+	* @param StackItemKey A unique key for the entry.
+	* @param SelectedSection The selected section for the entry.
+	*/
+	void SetStackEntryActiveSection(const FString& StackEntryKey, FText ActiveSection);
+
 	/*
 	* Gets a stack entry's display name. Returns null if none is found.
 	* @param StackEntryKey A unique key for the stack entry.
@@ -155,6 +169,8 @@ private:
 	UPROPERTY()
 	TMap<FString, bool> StackEntryKeyToExpandedOverviewMap;
 
+	TMap<FString, FText> StackEntryKeyToActiveSectionMap;
+
 	/* Marking those FTexts explicitly as editoronly_data will make localization not pick these up.
 	 * This is a workaround. EditorDataBase in system & emitter is already flagged as editor only, but it doesn't propagate properly */
 #if WITH_EDITORONLY_DATA
diff --git a/Engine/Plugins/FX/Niagara/Source/NiagaraEditor/Public/ViewModels/NiagaraEmitterViewModel.h b/Engine/Plugins/FX/Niagara/Source/NiagaraEditor/Public/ViewModels/NiagaraEmitterViewModel.h
--- a/Engine/Plugins/FX/Niagara/Source/NiagaraEditor/Public/ViewModels/NiagaraEmitterViewModel.h
+++ b/Engine/Plugins/FX/Niagara/Source/NiagaraEditor/Public/ViewModels/NiagaraEmitterViewModel.h
@@ -120,6 +120,9 @@ public:
 	/** Add an event script to the owned emitter. Sets the Usage, UsageID and Source of the EventScriptProperties. */
 	NIAGARAEDITOR_API void AddEventHandler(FNiagaraEventScriptProperties& EventScriptProperties, bool bResetGraphForOutput = false);
 
+	bool GetSummaryIsInEditMode() const { return bSummaryIsInEditMode; }
+	void SetSummaryIsInEditMode(bool bInSummaryIsInEditMode) { bSummaryIsInEditMode = bInSummaryIsInEditMode; }
+
 	void Cleanup();
 
 private:
@@ -188,4 +191,6 @@ private:
 	TMap<FObjectKey, FDelegateHandle> ScriptToOnParameterStoreChangedHandleMap;
 
 	TSharedPtr<SWindow> NewParentWindow;
+
+	bool bSummaryIsInEditMode;
 };
\ No newline at end of file
diff --git a/Engine/Plugins/FX/Niagara/Source/NiagaraEditor/Public/ViewModels/Stack/NiagaraStackEmitterPropertiesGroup.h b/Engine/Plugins/FX/Niagara/Source/NiagaraEditor/Public/ViewModels/Stack/NiagaraStackEmitterPropertiesGroup.h
--- a/Engine/Plugins/FX/Niagara/Source/NiagaraEditor/Public/ViewModels/Stack/NiagaraStackEmitterPropertiesGroup.h
+++ b/Engine/Plugins/FX/Niagara/Source/NiagaraEditor/Public/ViewModels/Stack/NiagaraStackEmitterPropertiesGroup.h
@@ -22,7 +22,7 @@ public:
 
 	void Initialize(FRequiredEntryData InRequiredEntryData);
 
-	virtual bool SupportsIcon() const override { return true; }
+	virtual EIconMode GetSupportedIconMode() const override { return EIconMode::Brush; }
 	virtual const FSlateBrush* GetIconBrush() const override;
 	virtual bool SupportsSecondaryIcon() const override { return true; }
 	virtual const FSlateBrush* GetSecondaryIconBrush() const override;
diff --git a/Engine/Plugins/FX/Niagara/Source/NiagaraEditor/Public/ViewModels/Stack/NiagaraStackEmitterSettingsGroup.h b/Engine/Plugins/FX/Niagara/Source/NiagaraEditor/Public/ViewModels/Stack/NiagaraStackEmitterSettingsGroup.h
--- a/Engine/Plugins/FX/Niagara/Source/NiagaraEditor/Public/ViewModels/Stack/NiagaraStackEmitterSettingsGroup.h
+++ b/Engine/Plugins/FX/Niagara/Source/NiagaraEditor/Public/ViewModels/Stack/NiagaraStackEmitterSettingsGroup.h
@@ -10,6 +10,7 @@ class FNiagaraEmitterViewModel;
 class FNiagaraScriptViewModel;
 class UNiagaraStackObject;
 class UNiagaraStackSummaryViewObject;
+class IDetailTreeNode;
 
 UCLASS()
 class NIAGARAEDITOR_API UNiagaraStackEmitterPropertiesItem : public UNiagaraStackItem
@@ -27,7 +28,7 @@ public:
 	virtual bool TestCanResetToBaseWithMessage(FText& OutCanResetToBaseMessage) const override;
 	virtual void ResetToBase() override;
 
-	virtual bool SupportsIcon() const override { return true; }
+	virtual EIconMode GetSupportedIconMode() const { return EIconMode::Brush; }
 	virtual const FSlateBrush* GetIconBrush() const override;
 
 	virtual bool IsExpandedByDefault() const override;
@@ -66,19 +67,29 @@ public:
 
 	virtual bool SupportsResetToBase() const override { return false; }
 	virtual bool IsExpandedByDefault() const override { return false; }
-	virtual bool SupportsIcon() const override { return false; }
-	virtual const FSlateBrush* GetIconBrush() const override;
+	virtual EIconMode GetSupportedIconMode() const override { return EIconMode::Text; }
+	virtual FText GetIconText() const override;
 	virtual bool GetShouldShowInOverview() const override { return false; }
 
+	virtual bool SupportsEditMode() const override { return true; }
+	virtual bool GetEditModeIsActive() const override;
+	virtual void SetEditModeIsActive(bool bInEditModeIsActive) override;
+
 protected:
 
 	virtual void RefreshChildrenInternal(const TArray<UNiagaraStackEntry*>& CurrentChildren, TArray<UNiagaraStackEntry*>& NewChildren, TArray<FStackIssue>& NewIssues) override;
 
+private:
+	void SelectSummaryNodesFromEmitterEditorDataRootNodes(TArray<TSharedRef<IDetailTreeNode>> Source, TArray<TSharedRef<IDetailTreeNode>>* Selected);
+
 private:
 	TWeakObjectPtr<UNiagaraEmitter> Emitter;
 
 	UPROPERTY()
 	TObjectPtr<UNiagaraStackSummaryViewObject> FilteredObject;
+
+	UPROPERTY()
+	TObjectPtr<UNiagaraStackObject> SummaryEditorData;
 };
 
 UCLASS()
@@ -89,8 +100,8 @@ class NIAGARAEDITOR_API UNiagaraStackEmitterSummaryGroup : public UNiagaraStackI
 public:
 	UNiagaraStackEmitterSummaryGroup();
 
-	virtual bool SupportsIcon() const override { return true; }
-	virtual const FSlateBrush* GetIconBrush() const override;
+	virtual EIconMode GetSupportedIconMode() const override { return EIconMode::Text; }
+	virtual  FText GetIconText() const override;
 	virtual bool GetCanExpandInOverview() const override { return false; }
 	virtual bool GetShouldShowInStack() const override { return false; }
 
diff --git a/Engine/Plugins/FX/Niagara/Source/NiagaraEditor/Public/ViewModels/Stack/NiagaraStackEntry.h b/Engine/Plugins/FX/Niagara/Source/NiagaraEditor/Public/ViewModels/Stack/NiagaraStackEntry.h
--- a/Engine/Plugins/FX/Niagara/Source/NiagaraEditor/Public/ViewModels/Stack/NiagaraStackEntry.h
+++ b/Engine/Plugins/FX/Niagara/Source/NiagaraEditor/Public/ViewModels/Stack/NiagaraStackEntry.h
@@ -63,6 +63,13 @@ public:
 		None
 	};
 
+	enum class EIconMode
+	{
+		Brush,
+		Text,
+		None
+	};
+
 	struct FDropRequest
 	{
 		FDropRequest(TSharedRef<const FDragDropOperation> InDragDropOperation, EItemDropZone InDropZone, EDragOptions InDragOptions, EDropOptions InDropOptions)
@@ -113,7 +120,6 @@ public:
 	struct NIAGARAEDITOR_API FExecutionSubcategoryNames
 	{
 		static const FName Settings;
-		static const FName Summary;
 		static const FName Spawn;
 		static const FName Update;
 		static const FName Event;
@@ -501,6 +507,12 @@ public:
 	/** Handler for when a rename is committed for this stack entry. */
 	virtual void OnRenamed(FText NewName);
 
+	virtual EIconMode GetSupportedIconMode() const { return EIconMode::None; }
+
+	virtual const FSlateBrush* GetIconBrush() const { return nullptr; }
+
+	virtual FText GetIconText() const { return FText(); }
+
 protected:
 	virtual void BeginDestroy() override;
 
diff --git a/Engine/Plugins/FX/Niagara/Source/NiagaraEditor/Public/ViewModels/Stack/NiagaraStackFilteredObject.h b/Engine/Plugins/FX/Niagara/Source/NiagaraEditor/Public/ViewModels/Stack/NiagaraStackFilteredObject.h
--- a/Engine/Plugins/FX/Niagara/Source/NiagaraEditor/Public/ViewModels/Stack/NiagaraStackFilteredObject.h
+++ b/Engine/Plugins/FX/Niagara/Source/NiagaraEditor/Public/ViewModels/Stack/NiagaraStackFilteredObject.h
@@ -19,13 +19,10 @@ class NIAGARAEDITOR_API UNiagaraStackFilteredObject : public UNiagaraStackFuncti
 {
 	GENERATED_BODY()
 		
-public:
-	DECLARE_DELEGATE_TwoParams(FOnSelectRootNodes, TArray<TSharedRef<IDetailTreeNode>>, TArray<TSharedRef<IDetailTreeNode>>*);
-
 public:
 	UNiagaraStackFilteredObject();
 
-	void Initialize(FRequiredEntryData InRequiredEntryData, UNiagaraEmitter* InEmitter, FString InOwningStackItemEditorDataKey);
+	void Initialize(FRequiredEntryData InRequiredEntryData, FString InOwningStackItemEditorDataKey);
 	void FinalizeInternal() override;
 
 	virtual FText GetDisplayName() const override;
diff --git a/Engine/Plugins/FX/Niagara/Source/NiagaraEditor/Public/ViewModels/Stack/NiagaraStackFunctionInputCollection.h b/Engine/Plugins/FX/Niagara/Source/NiagaraEditor/Public/ViewModels/Stack/NiagaraStackFunctionInputCollection.h
--- a/Engine/Plugins/FX/Niagara/Source/NiagaraEditor/Public/ViewModels/Stack/NiagaraStackFunctionInputCollection.h
+++ b/Engine/Plugins/FX/Niagara/Source/NiagaraEditor/Public/ViewModels/Stack/NiagaraStackFunctionInputCollection.h
@@ -4,6 +4,7 @@
 
 #include "ViewModels/Stack/NiagaraStackItem.h"
 #include "NiagaraTypes.h"
+#include "NiagaraStackSection.h"
 #include "NiagaraStackFunctionInputCollection.generated.h"
 
 class UNiagaraNodeFunctionCall;
@@ -19,10 +20,25 @@ class NIAGARAEDITOR_API UNiagaraStackFunctionInputCollectionBase : public UNiaga
 	GENERATED_BODY()
 
 public:
-	UNiagaraStackFunctionInputCollectionBase();
+	void Initialize(FRequiredEntryData InRequiredEntryData, FString InOwningStackItemEditorDataKey, FString InStackEditorDataKey);
 
+	bool GetShouldDisplayLabel() const { return bShouldDisplayLabel; }
+	void SetShouldDisplayLabel(bool bInShouldShowLabel);
+
+	const TArray<FText>& GetSections() const;
+
+	FText GetActiveSection() const;
+
+	void SetActiveSection(FText InActiveSection);
+
+	virtual bool GetCanExpand() const override;
+	virtual bool GetShouldShowInStack() const override;
+
+public:
 	static FText UncategorizedName;
 
+	static FText AllSectionName;
+
 protected:
 
 	struct FInputData
@@ -56,6 +72,8 @@ protected:
 		TMap<FName, FNiagaraParentData> ParentMapping;
 	};
 	
+	virtual void RefreshChildrenInternal(const TArray<UNiagaraStackEntry*>& CurrentChildren, TArray<UNiagaraStackEntry*>& NewChildren, TArray<FStackIssue>& NewIssues) override;
+
 	void RefreshChildrenForFunctionCall(UNiagaraNodeFunctionCall* ModuleNode, UNiagaraNodeFunctionCall* InputFunctionCallNode, const TArray<UNiagaraStackEntry*>& CurrentChildren, 
 		TArray<UNiagaraStackEntry*>& NewChildren, TArray<FStackIssue>& NewIssues, bool bShouldApplySummaryFilter);
 	
@@ -78,9 +96,21 @@ protected:
 
 	void AddInputToCategory(const FInputData& InputData, const TArray<UNiagaraStackEntry*>& CurrentChildren, TArray<UNiagaraStackEntry*>& NewChildren);
 
-};
+	virtual void GetSectionsInternal(TArray<FNiagaraStackSection>& OutStackSections) const { }
 
+private:
+	bool FilterByActiveSection(const UNiagaraStackEntry& Child) const;
 
+	void UpdateCachedSectionData() const;
+
+private:
+	bool bShouldDisplayLabel;
+
+	mutable TOptional<TArray<FText>> SectionsCache;
+	mutable TOptional<TMap<FString, TArray<FText>>> SectionToCategoryMapCache;
+	mutable TOptional<FText> ActiveSectionCache;
+	FText LastActiveSection;
+};
 
 UCLASS()
 class NIAGARAEDITOR_API UNiagaraStackFunctionInputCollection : public UNiagaraStackFunctionInputCollectionBase
@@ -102,11 +132,8 @@ public:
 
 	//~ UNiagaraStackEntry interface
 	virtual FText GetDisplayName() const override;
-	virtual bool GetShouldShowInStack() const override;
 	virtual bool GetIsEnabled() const;
 
-	void SetShouldShowInStack(bool bInShouldShowInStack);
-
 	void ToClipboardFunctionInputs(UObject* InOuter, TArray<const UNiagaraClipboardFunctionInput*>& OutClipboardFunctionInputs) const;
 
 	void SetValuesFromClipboardFunctionInputs(const TArray<const UNiagaraClipboardFunctionInput*>& ClipboardFunctionInputs);
@@ -121,9 +148,10 @@ protected:
 	virtual void FinalizeInternal() override;
 
 	virtual void RefreshChildrenInternal(const TArray<UNiagaraStackEntry*>& CurrentChildren, TArray<UNiagaraStackEntry*>& NewChildren, TArray<FStackIssue>& NewIssues) override;
+
+	virtual void GetSectionsInternal(TArray<FNiagaraStackSection>& OutStackSections) const override;
 	
 private:
 	UNiagaraNodeFunctionCall* ModuleNode;
 	UNiagaraNodeFunctionCall* InputFunctionCallNode;
-	bool bShouldShowInStack;
 };
diff --git a/Engine/Plugins/FX/Niagara/Source/NiagaraEditor/Public/ViewModels/Stack/NiagaraStackItem.h b/Engine/Plugins/FX/Niagara/Source/NiagaraEditor/Public/ViewModels/Stack/NiagaraStackItem.h
--- a/Engine/Plugins/FX/Niagara/Source/NiagaraEditor/Public/ViewModels/Stack/NiagaraStackItem.h
+++ b/Engine/Plugins/FX/Niagara/Source/NiagaraEditor/Public/ViewModels/Stack/NiagaraStackItem.h
@@ -33,13 +33,14 @@ public:
 
 	virtual bool SupportsChangeEnabled() const { return false; }
 	void SetIsEnabled(bool bInIsEnabled);
-	
-	virtual bool SupportsIcon() const { return false; }
-	virtual const FSlateBrush* GetIconBrush() const;
 
 	virtual bool SupportsResetToBase() const { return false; }
 	virtual bool TestCanResetToBaseWithMessage(FText& OutCanResetToBaseMessage) const { return false; }
 	virtual void ResetToBase() { }
+
+	virtual bool SupportsEditMode() const { return false; }
+	virtual bool GetEditModeIsActive() const { return false; }
+	virtual void SetEditModeIsActive(bool bInEditModeIsActive) { }
 	
 protected:
 	virtual void RefreshChildrenInternal(const TArray<UNiagaraStackEntry*>& CurrentChildren, TArray<UNiagaraStackEntry*>& NewChildren, TArray<FStackIssue>& NewIssues) override;
diff --git a/Engine/Plugins/FX/Niagara/Source/NiagaraEditor/Public/ViewModels/Stack/NiagaraStackItemGroup.h b/Engine/Plugins/FX/Niagara/Source/NiagaraEditor/Public/ViewModels/Stack/NiagaraStackItemGroup.h
--- a/Engine/Plugins/FX/Niagara/Source/NiagaraEditor/Public/ViewModels/Stack/NiagaraStackItemGroup.h
+++ b/Engine/Plugins/FX/Niagara/Source/NiagaraEditor/Public/ViewModels/Stack/NiagaraStackItemGroup.h
@@ -29,9 +29,6 @@ public:
 	virtual void SetIsEnabled(bool bEnabled) {}
 	virtual bool SupportsChangeEnabled() const { return false; }
 
-	virtual bool SupportsIcon() const { return false; }
-	virtual const FSlateBrush* GetIconBrush() const { return nullptr; };
-
 	virtual bool SupportsSecondaryIcon() const { return false; }
 	virtual const FSlateBrush* GetSecondaryIconBrush() const { return nullptr; };
 
diff --git a/Engine/Plugins/FX/Niagara/Source/NiagaraEditor/Public/ViewModels/Stack/NiagaraStackPropertyRow.h b/Engine/Plugins/FX/Niagara/Source/NiagaraEditor/Public/ViewModels/Stack/NiagaraStackPropertyRow.h
--- a/Engine/Plugins/FX/Niagara/Source/NiagaraEditor/Public/ViewModels/Stack/NiagaraStackPropertyRow.h
+++ b/Engine/Plugins/FX/Niagara/Source/NiagaraEditor/Public/ViewModels/Stack/NiagaraStackPropertyRow.h
@@ -24,6 +24,8 @@ public:
 
 	virtual bool IsExpandedByDefault() const override;
 
+	virtual bool CanDrag() const override;
+
 protected:
 	virtual void FinalizeInternal() override;
 
@@ -31,6 +33,10 @@ protected:
 
 	virtual void GetSearchItems(TArray<FStackSearchItem>& SearchItems) const override;
 
+	virtual TOptional<FDropRequestResponse> CanDropInternal(const FDropRequest& DropRequest) override;
+
+	virtual TOptional<FDropRequestResponse> DropInternal(const FDropRequest& DropRequest) override;
+
 private:
 	TSharedPtr<IDetailTreeNode> DetailTreeNode;
 	UNiagaraNode* OwningNiagaraNode;
diff --git a/Engine/Plugins/FX/Niagara/Source/NiagaraEditor/Public/ViewModels/Stack/NiagaraStackRendererItem.h b/Engine/Plugins/FX/Niagara/Source/NiagaraEditor/Public/ViewModels/Stack/NiagaraStackRendererItem.h
--- a/Engine/Plugins/FX/Niagara/Source/NiagaraEditor/Public/ViewModels/Stack/NiagaraStackRendererItem.h
+++ b/Engine/Plugins/FX/Niagara/Source/NiagaraEditor/Public/ViewModels/Stack/NiagaraStackRendererItem.h
@@ -52,7 +52,7 @@ public:
 	virtual bool SupportsChangeEnabled() const override { return true; }
 	virtual bool GetIsEnabled() const override;
 
-	virtual bool SupportsIcon() const override { return true; }
+	virtual EIconMode GetSupportedIconMode() const override { return EIconMode::Brush; }
 	virtual const FSlateBrush* GetIconBrush() const override;
 
 	virtual bool SupportsResetToBase() const override { return true; }
diff --git a/Engine/Plugins/FX/Niagara/Source/NiagaraEditor/Public/ViewModels/Stack/NiagaraStackSummaryViewInputCollection.h b/Engine/Plugins/FX/Niagara/Source/NiagaraEditor/Public/ViewModels/Stack/NiagaraStackSummaryViewInputCollection.h
--- a/Engine/Plugins/FX/Niagara/Source/NiagaraEditor/Public/ViewModels/Stack/NiagaraStackSummaryViewInputCollection.h
+++ b/Engine/Plugins/FX/Niagara/Source/NiagaraEditor/Public/ViewModels/Stack/NiagaraStackSummaryViewInputCollection.h
@@ -29,9 +29,8 @@ public:
 	virtual void FinalizeInternal() override;
 
 	virtual FText GetDisplayName() const override;
-	virtual bool GetShouldShowInStack() const override;
 	virtual bool GetIsEnabled() const;
-	bool GetCanExpand() const override { return false; }
+
 protected:
 
 	virtual void RefreshChildrenInternal(const TArray<UNiagaraStackEntry*>& CurrentChildren, TArray<UNiagaraStackEntry*>& NewChildren, TArray<FStackIssue>& NewIssues) override;
@@ -40,6 +39,8 @@ protected:
 
 	virtual void PostRefreshChildrenInternal() override;
 
+	virtual void GetSectionsInternal(TArray<FNiagaraStackSection>& OutStackSections) const override;
+
 private:
 
 	void OnViewStateChanged();
diff --git a/Engine/Plugins/FX/Niagara/Source/NiagaraEditor/Public/ViewModels/Stack/NiagaraStackSystemSettingsGroup.h b/Engine/Plugins/FX/Niagara/Source/NiagaraEditor/Public/ViewModels/Stack/NiagaraStackSystemSettingsGroup.h
--- a/Engine/Plugins/FX/Niagara/Source/NiagaraEditor/Public/ViewModels/Stack/NiagaraStackSystemSettingsGroup.h
+++ b/Engine/Plugins/FX/Niagara/Source/NiagaraEditor/Public/ViewModels/Stack/NiagaraStackSystemSettingsGroup.h
@@ -19,7 +19,7 @@ class NIAGARAEDITOR_API UNiagaraStackSystemPropertiesGroup : public UNiagaraStac
 public:
 	void Initialize(FRequiredEntryData InRequiredEntryData);
 
-	virtual bool SupportsIcon() const override { return true; }
+	virtual EIconMode GetSupportedIconMode() const override { return UNiagaraStackEntry::EIconMode::Brush; }
 	virtual const FSlateBrush* GetIconBrush() const override;
 	virtual bool GetCanExpandInOverview() const override { return false; }
 	virtual bool GetShouldShowInStack() const override { return false; }
@@ -36,8 +36,8 @@ class NIAGARAEDITOR_API UNiagaraStackSystemUserParametersGroup : public UNiagara
 public:
 	void Initialize(FRequiredEntryData InRequiredEntryData, UObject* InOwner, FNiagaraParameterStore* InParameterStore);
 
-	virtual bool SupportsIcon() const override { return true; }
-	virtual const FSlateBrush* GetIconBrush() const override;
+	virtual EIconMode GetSupportedIconMode() const override { return UNiagaraStackEntry::EIconMode::Text; }
+	virtual FText GetIconText() const override;
 	virtual bool GetCanExpandInOverview() const override { return false; }
 	virtual bool GetShouldShowInStack() const override { return false; }
 
diff --git a/Engine/Plugins/FX/Niagara/Source/NiagaraEditorWidgets/Private/NiagaraEditorWidgetsStyle.cpp b/Engine/Plugins/FX/Niagara/Source/NiagaraEditorWidgets/Private/NiagaraEditorWidgetsStyle.cpp
--- a/Engine/Plugins/FX/Niagara/Source/NiagaraEditorWidgets/Private/NiagaraEditorWidgetsStyle.cpp
+++ b/Engine/Plugins/FX/Niagara/Source/NiagaraEditorWidgets/Private/NiagaraEditorWidgetsStyle.cpp
@@ -321,7 +321,7 @@ FNiagaraEditorWidgetsStyle::FNiagaraEditorWidgetsStyle() : FSlateStyleSet("Niaga
 		.SetOddRowBackgroundBrush(FSlateColorBrush(FStyleColors::Panel))
 		.SetEvenRowBackgroundHoveredBrush(CORE_IMAGE_BRUSH("Common/Selection", Icon8x8, FLinearColor(1.0f, 1.0f, 1.0f, 0.1f)))
 		.SetOddRowBackgroundHoveredBrush(CORE_IMAGE_BRUSH("Common/Selection", Icon8x8, FLinearColor(1.0f, 1.0f, 1.0f, 0.1f)))
-		.SetSelectorFocusedBrush(CORE_BORDER_BRUSH("Common/Selector", FMargin(4.f / 16.f), FStarshipCoreStyle::GetCoreStyle().GetSlateColor("SelectorColor")))
+		.SetSelectorFocusedBrush(FSlateNoResource())
 		.SetActiveBrush(FSlateColorBrush(FStyleColors::Select))
 		.SetActiveHoveredBrush(FSlateColorBrush(FStyleColors::Select))
 		.SetInactiveBrush(FSlateColorBrush(FStyleColors::SelectInactive))
@@ -333,7 +333,7 @@ FNiagaraEditorWidgetsStyle::FNiagaraEditorWidgetsStyle() : FSlateStyleSet("Niaga
 		.SetOddRowBackgroundBrush(FSlateColorBrush(FStyleColors::Header))
 		.SetEvenRowBackgroundHoveredBrush(CORE_IMAGE_BRUSH("Common/Selection", Icon8x8, FLinearColor(1.0f, 1.0f, 1.0f, 0.1f)))
 		.SetOddRowBackgroundHoveredBrush(CORE_IMAGE_BRUSH("Common/Selection", Icon8x8, FLinearColor(1.0f, 1.0f, 1.0f, 0.1f)))
-		.SetSelectorFocusedBrush(CORE_BORDER_BRUSH("Common/Selector", FMargin(4.f / 16.f), FStarshipCoreStyle::GetCoreStyle().GetSlateColor("SelectorColor")))
+		.SetSelectorFocusedBrush(FSlateNoResource())
 		.SetActiveBrush(FSlateColorBrush(FStyleColors::Select))
 		.SetActiveHoveredBrush(FSlateColorBrush(FStyleColors::Select))
 		.SetInactiveBrush(FSlateColorBrush(FStyleColors::SelectInactive))
diff --git a/Engine/Plugins/FX/Niagara/Source/NiagaraEditorWidgets/Private/SNiagaraOverviewStack.cpp b/Engine/Plugins/FX/Niagara/Source/NiagaraEditorWidgets/Private/SNiagaraOverviewStack.cpp
--- a/Engine/Plugins/FX/Niagara/Source/NiagaraEditorWidgets/Private/SNiagaraOverviewStack.cpp
+++ b/Engine/Plugins/FX/Niagara/Source/NiagaraEditorWidgets/Private/SNiagaraOverviewStack.cpp
@@ -177,7 +177,7 @@ class SNiagaraSystemOverviewEntryListRow : public STableRow<UNiagaraStackEntry*>
 			RowOverlay->AddSlot()
 			.HAlign(HAlign_Left)
 			.VAlign(VAlign_Fill)
-			.Padding(FMargin(1, 0, 0, 0))
+			.Padding(FMargin(0))
 			[
 				SNew(SBorder)
 				.BorderImage(FAppStyle::Get().GetBrush("WhiteBrush"))
@@ -313,21 +313,22 @@ private:
 	{
 		if (IssueHighlightColor.IsSet() == false)
 		{
+			float DesaturateAmount = .1f;
 			if (StackEntry->GetTotalNumberOfErrorIssues() > 0)
 			{
-				IssueHighlightColor = FStyleColors::AccentRed.GetSpecifiedColor().Desaturate(.15f);
+				IssueHighlightColor = FStyleColors::AccentRed.GetSpecifiedColor().Desaturate(DesaturateAmount);
 			}
 			else if (StackEntry->GetTotalNumberOfWarningIssues() > 0)
 			{
-				IssueHighlightColor = FStyleColors::AccentYellow.GetSpecifiedColor().Desaturate(.15f);
+				IssueHighlightColor = FStyleColors::AccentYellow.GetSpecifiedColor().Desaturate(DesaturateAmount);
 			}
 			else if (StackEntry->GetTotalNumberOfInfoIssues() > 0)
 			{
-				IssueHighlightColor = FStyleColors::AccentBlue.GetSpecifiedColor().Desaturate(.15f);
+				IssueHighlightColor = FStyleColors::AccentBlue.GetSpecifiedColor().Desaturate(DesaturateAmount);
 			}
 			else if (StackEntry->GetTotalNumberOfCustomNotes() > 0)
 			{
-				IssueHighlightColor = FStyleColors::AccentGray.GetSpecifiedColor().Desaturate(.15f);
+				IssueHighlightColor = FStyleColors::AccentWhite.GetSpecifiedColor().Desaturate(DesaturateAmount);
 			}
 			else
 			{
@@ -755,9 +756,9 @@ TSharedRef<ITableRow> SNiagaraOverviewStack::OnGenerateRowForEntry(UNiagaraStack
 	{
 		UNiagaraStackItem* StackItem = CastChecked<UNiagaraStackItem>(Item);
 		TSharedRef<SHorizontalBox> ContentBox = SNew(SHorizontalBox);
-		if (StackItem->SupportsIcon())
+		// Icon Brush
+		if (StackItem->GetSupportedIconMode() == UNiagaraStackEntry::EIconMode::Brush)
 		{
-			// Icon
 			ContentBox->AddSlot()
 			.Padding(0, 1, 0, 1)
 			.AutoWidth()
@@ -776,6 +777,28 @@ TSharedRef<ITableRow> SNiagaraOverviewStack::OnGenerateRowForEntry(UNiagaraStack
 				]
 			];
 		}
+		// Icon Text
+		if (StackItem->GetSupportedIconMode() == UNiagaraStackEntry::EIconMode::Text)
+		{
+			ContentBox->AddSlot()
+			.Padding(0, 1, 0, 1)
+			.AutoWidth()
+			.VAlign(VAlign_Center)
+			[
+				SNew(SBox)
+				.IsEnabled_UObject(Item, &UNiagaraStackEntry::GetIsEnabledAndOwnerIsEnabled)
+				.WidthOverride(IconSize.X)
+				.HeightOverride(IconSize.Y)
+				.HAlign(HAlign_Center)
+				.VAlign(VAlign_Center)
+				[
+					SNew(STextBlock)
+					.Font(FEditorStyle::Get().GetFontStyle("FontAwesome.10"))
+					.Text(StackItem->GetIconText())
+					.ColorAndOpacity(FSlateColor::UseSubduedForeground())
+				]
+			];
+		}
 		// Name
 		ContentBox->AddSlot()
 		.VAlign(VAlign_Center)
@@ -819,6 +842,8 @@ TSharedRef<ITableRow> SNiagaraOverviewStack::OnGenerateRowForEntry(UNiagaraStack
 						SNew(SImage)
 						.ToolTipText(LOCTEXT("ScratchPadOverviewTooltip", "This module is a scratch pad script."))
 						.Image(FNiagaraEditorStyle::Get().GetBrush("Tab.ScratchPad"))
+						.ColorAndOpacity(FSlateColor::UseSubduedForeground())
+						.DesiredSizeOverride(FVector2D(14.0f, 14.0f))
 					];
 			}
 
@@ -833,12 +858,11 @@ TSharedRef<ITableRow> SNiagaraOverviewStack::OnGenerateRowForEntry(UNiagaraStack
 						SNew(SButton)
 						.HAlign(HAlign_Center)
 						.VAlign(VAlign_Center)
-						.ButtonColorAndOpacity(FLinearColor::Transparent)
 						.ForegroundColor(FLinearColor::Transparent)
 						.ToolTipText(LOCTEXT("EnableDebugDrawCheckBoxToolTip", "Enable or disable debug drawing for this item."))
 						.OnClicked(this, &SNiagaraOverviewStack::ToggleModuleDebugDraw, StackItem)
-						.ContentPadding(FMargin(2.f))
-						.ButtonStyle(&FAppStyle::Get().GetWidgetStyle<FButtonStyle>("SimpleButton"))
+						.ContentPadding(FMargin(1.0f))
+						.ButtonStyle(FNiagaraEditorWidgetsStyle::Get(), "NiagaraEditor.Stack.SimpleButton")
 						[
 							SNew(SImage)
 							.Image(this, &SNiagaraOverviewStack::GetDebugIconBrush, StackItem)
@@ -926,14 +950,17 @@ TSharedRef<ITableRow> SNiagaraOverviewStack::OnGenerateRowForEntry(UNiagaraStack
 	{
 		UNiagaraStackItemGroup* StackItemGroup = CastChecked<UNiagaraStackItemGroup>(Item);
 		TSharedRef<SHorizontalBox> ContentBox = SNew(SHorizontalBox);
+
+		FLinearColor IconColor = FNiagaraEditorWidgetsStyle::Get().GetColor(
+			FNiagaraStackEditorWidgetsUtilities::GetIconColorNameForExecutionCategory(StackItemGroup->GetExecutionCategoryName())).Desaturate(.25f);
 		
-		// Icon
-		if (StackItemGroup->SupportsIcon())
+		// Icon Brush
+		if (StackItemGroup->GetSupportedIconMode() == UNiagaraStackEntry::EIconMode::Brush)
 		{
 			ContentBox->AddSlot()
 				.VAlign(VAlign_Center)
 				.AutoWidth()
-				.Padding(2, 1, 2, 1)
+				.Padding(0, 1, 1, 1)
 				[
 					SNew(SBox)
 					.IsEnabled_UObject(Item, &UNiagaraStackEntry::GetIsEnabledAndOwnerIsEnabled)
@@ -945,6 +972,30 @@ TSharedRef<ITableRow> SNiagaraOverviewStack::OnGenerateRowForEntry(UNiagaraStack
 						SNew(SImage)
 						.Image_UObject(StackItemGroup, &UNiagaraStackItemGroup::GetIconBrush)
 						.DesiredSizeOverride(FVector2D(14.0f, 14.0f))
+						.ColorAndOpacity(IconColor)
+					]
+				];
+		}
+
+		// Icon Text
+		if (StackItemGroup->GetSupportedIconMode() == UNiagaraStackEntry::EIconMode::Text)
+		{
+			ContentBox->AddSlot()
+				.VAlign(VAlign_Center)
+				.AutoWidth()
+				.Padding(0, 1, 1, 1)
+				[
+					SNew(SBox)
+					.IsEnabled_UObject(Item, &UNiagaraStackEntry::GetIsEnabledAndOwnerIsEnabled)
+					.WidthOverride(IconSize.X)
+					.HeightOverride(IconSize.Y)
+					.HAlign(HAlign_Center)
+					.VAlign(VAlign_Center)
+					[
+						SNew(STextBlock)
+						.Font(FEditorStyle::Get().GetFontStyle("FontAwesome.10"))
+						.Text(StackItemGroup->GetIconText())
+						.ColorAndOpacity(IconColor)
 					]
 				];
 		}
diff --git a/Engine/Plugins/FX/Niagara/Source/NiagaraEditorWidgets/Private/SNiagaraStack.cpp b/Engine/Plugins/FX/Niagara/Source/NiagaraEditorWidgets/Private/SNiagaraStack.cpp
--- a/Engine/Plugins/FX/Niagara/Source/NiagaraEditorWidgets/Private/SNiagaraStack.cpp
+++ b/Engine/Plugins/FX/Niagara/Source/NiagaraEditorWidgets/Private/SNiagaraStack.cpp
@@ -38,6 +38,7 @@
 #include "Stack/NiagaraStackPropertyRowUtilities.h"
 #include "Stack/SNiagaraStackFunctionInputName.h"
 #include "Stack/SNiagaraStackFunctionInputValue.h"
+#include "Stack/SNiagaraStackFunctionInputCollection.h"
 #include "Stack/SNiagaraStackItem.h"
 #include "Stack/SNiagaraStackItemFooter.h"
 #include "Stack/SNiagaraStackItemGroup.h"
@@ -969,8 +970,12 @@ SNiagaraStack::FRowWidgets SNiagaraStack::ConstructNameAndValueWidgetsForItem(UN
 				.IsReadOnly(true)
 			);
 	}
-	else if (Item->IsA<UNiagaraStackFunctionInputCollection>() ||
-		Item->IsA<UNiagaraStackModuleItemOutputCollection>() ||
+	else if (Item->IsA<UNiagaraStackFunctionInputCollectionBase>())
+	{
+		UNiagaraStackFunctionInputCollectionBase* InputCollection = CastChecked<UNiagaraStackFunctionInputCollectionBase>(Item);
+		return FRowWidgets(SNew(SNiagaraStackFunctionInputCollection, InputCollection));
+	}
+	else if (Item->IsA<UNiagaraStackModuleItemOutputCollection>() ||
 		Item->IsA<UNiagaraStackModuleItemLinkedInputCollection>())
 	{
 		return FRowWidgets(
diff --git a/Engine/Plugins/FX/Niagara/Source/NiagaraEditorWidgets/Private/Stack/SNiagaraStackFunctionInputCollection.cpp b/Engine/Plugins/FX/Niagara/Source/NiagaraEditorWidgets/Private/Stack/SNiagaraStackFunctionInputCollection.cpp
new file mode 100644
--- /dev/null
+++ b/Engine/Plugins/FX/Niagara/Source/NiagaraEditorWidgets/Private/Stack/SNiagaraStackFunctionInputCollection.cpp
@@ -0,0 +1,82 @@
+// Copyright Epic Games, Inc. All Rights Reserved.
+
+#include "Stack/SNiagaraStackFunctionInputCollection.h"
+
+#include "NiagaraEditorWidgetsStyle.h"
+#include "ViewModels/Stack/NiagaraStackFunctionInputCollection.h"
+#include "Widgets/Layout/SWrapBox.h"
+
+void SNiagaraStackFunctionInputCollection::Construct(const FArguments& InArgs, UNiagaraStackFunctionInputCollectionBase* InInputCollection)
+{
+	InputCollection = InInputCollection;
+	InputCollection->OnStructureChanged().AddSP(this, &SNiagaraStackFunctionInputCollection::InputCollectionStructureChanged);
+
+	ChildSlot
+	[
+		SNew(SVerticalBox)
+		+ SVerticalBox::Slot()
+		.AutoHeight()
+		.Padding(0, 2)
+		[
+			SNew(STextBlock)
+			.TextStyle(FNiagaraEditorWidgetsStyle::Get(), "NiagaraEditor.Stack.ItemText")
+			.ToolTipText_UObject(InputCollection, &UNiagaraStackEntry::GetTooltipText)
+			.Text_UObject(InputCollection, &UNiagaraStackEntry::GetDisplayName)
+			.IsEnabled_UObject(InputCollection, &UNiagaraStackEntry::GetOwnerIsEnabled)
+			.Visibility(this, &SNiagaraStackFunctionInputCollection::GetLabelVisibility)
+		]
+		+ SVerticalBox::Slot()
+		.AutoHeight()
+		.Padding(0, 2)
+		[
+			SAssignNew(SectionSelectorBox, SWrapBox)
+			.UseAllottedSize(true)
+			.InnerSlotPadding(FVector2D(4, 4))
+		]
+	];
+
+	ConstructSectionButtons();
+}
+
+EVisibility SNiagaraStackFunctionInputCollection::GetLabelVisibility() const
+{
+	return InputCollection->GetShouldDisplayLabel() ? EVisibility::Visible : EVisibility::Collapsed;
+}
+
+void SNiagaraStackFunctionInputCollection::ConstructSectionButtons()
+{
+	SectionSelectorBox->ClearChildren();
+	for (FText Section : InputCollection->GetSections())
+	{
+		SectionSelectorBox->AddSlot()
+		[
+			SNew(SCheckBox)
+			.Style(FAppStyle::Get(), "DetailsView.SectionButton")
+			.OnCheckStateChanged(this, &SNiagaraStackFunctionInputCollection::OnSectionChecked, Section)
+			.IsChecked(this, &SNiagaraStackFunctionInputCollection::GetSectionCheckState, Section)
+			[
+				SNew(STextBlock)
+				.TextStyle(FAppStyle::Get(), "SmallText")
+				.Text(Section)
+			]
+		];
+	}
+}
+
+void SNiagaraStackFunctionInputCollection::InputCollectionStructureChanged(ENiagaraStructureChangedFlags StructureChangedFlags)
+{
+	ConstructSectionButtons();
+}
+
+ECheckBoxState SNiagaraStackFunctionInputCollection::GetSectionCheckState(FText Section) const
+{
+	return Section.IdenticalTo(InputCollection->GetActiveSection()) ? ECheckBoxState::Checked : ECheckBoxState::Unchecked;
+}
+
+void SNiagaraStackFunctionInputCollection::OnSectionChecked(ECheckBoxState CheckState, FText Section)
+{
+	if (CheckState == ECheckBoxState::Checked)
+	{
+		InputCollection->SetActiveSection(Section);
+	}
+}
\ No newline at end of file
diff --git a/Engine/Plugins/FX/Niagara/Source/NiagaraEditorWidgets/Private/Stack/SNiagaraStackFunctionInputCollection.h b/Engine/Plugins/FX/Niagara/Source/NiagaraEditorWidgets/Private/Stack/SNiagaraStackFunctionInputCollection.h
new file mode 100644
--- /dev/null
+++ b/Engine/Plugins/FX/Niagara/Source/NiagaraEditorWidgets/Private/Stack/SNiagaraStackFunctionInputCollection.h
@@ -0,0 +1,36 @@
+// Copyright Epic Games, Inc. All Rights Reserved.
+
+#pragma once
+
+#include "NiagaraEditorCommon.h"
+#include "Widgets/SCompoundWidget.h"
+#include "Widgets/DeclarativeSyntaxSupport.h"
+#include "Widgets/Views/SExpanderArrow.h"
+
+class UNiagaraStackFunctionInputCollectionBase;
+class SWrapBox;
+
+class SNiagaraStackFunctionInputCollection : public SCompoundWidget
+{
+public:
+	SLATE_BEGIN_ARGS(SNiagaraStackFunctionInputCollection) {}
+	SLATE_END_ARGS();
+
+	void Construct(const FArguments& InArgs, UNiagaraStackFunctionInputCollectionBase* InInputCollection);
+
+private:
+	EVisibility GetLabelVisibility() const;
+
+	void ConstructSectionButtons();
+
+	void InputCollectionStructureChanged(ENiagaraStructureChangedFlags StructureChangedFlags);
+
+	ECheckBoxState GetSectionCheckState(FText Section) const;
+
+	void OnSectionChecked(ECheckBoxState CheckState, FText Section);
+
+private:
+	UNiagaraStackFunctionInputCollectionBase* InputCollection;
+
+	TSharedPtr<SWrapBox> SectionSelectorBox;
+};
diff --git a/Engine/Plugins/FX/Niagara/Source/NiagaraEditorWidgets/Private/Stack/SNiagaraStackFunctionInputValue.cpp b/Engine/Plugins/FX/Niagara/Source/NiagaraEditorWidgets/Private/Stack/SNiagaraStackFunctionInputValue.cpp
--- a/Engine/Plugins/FX/Niagara/Source/NiagaraEditorWidgets/Private/Stack/SNiagaraStackFunctionInputValue.cpp
+++ b/Engine/Plugins/FX/Niagara/Source/NiagaraEditorWidgets/Private/Stack/SNiagaraStackFunctionInputValue.cpp
@@ -188,25 +188,25 @@ void SNiagaraStackFunctionInputValue::Construct(const FArguments& InArgs, UNiaga
 		]
 	];
 
-	if (GetFilteredViewContextButtonVisibility() == EVisibility::Visible)
+	if (FunctionInput->GetTypedOuter<UNiagaraStackSummaryViewObject>() != nullptr)
 	{
 		ChildrenBox->AddSlot()
 		.AutoWidth()
 		.VAlign(VAlign_Center)
 		.Padding(3, 0, 0, 0)
 		[
-
 			SNew(SComboButton)
-			.ContentPadding(0)
-			.ForegroundColor(FLinearColor::Transparent)
-			.ButtonStyle(FEditorStyle::Get(), "NoBorder")
+			.ButtonStyle(FNiagaraEditorWidgetsStyle::Get(), "NiagaraEditor.Stack.SimpleButton")
 			.OnGetMenuContent(this, &SNiagaraStackFunctionInputValue::GetFilteredViewPropertiesContent)
 			.Visibility(this, &SNiagaraStackFunctionInputValue::GetFilteredViewContextButtonVisibility)
+			.ToolTipText(LOCTEXT("InputSummaryOptionsToolTip", "Input Summary Options"))
 			.IsEnabled(true)
 			.ButtonContent()
 			[
-				SNew(SImage)
-				.Image(this, &SNiagaraStackFunctionInputValue::GetFilteredViewIcon)
+				SNew(STextBlock)
+				.Font(FEditorStyle::Get().GetFontStyle("FontAwesome.10"))
+				.Text(FText::FromString(FString(TEXT("\xf0ca")/* fa-list-ul */)))
+				.ColorAndOpacity(FStyleColors::AccentYellow)
 			]
 		];
 	}
@@ -806,7 +806,7 @@ FSlateColor SNiagaraStackFunctionInputValue::GetInputIconColor() const
 
 EVisibility SNiagaraStackFunctionInputValue::GetTypeModifierIconVisibility() const
 {
-	return FunctionInput->GetInputType().IsStatic() ? EVisibility::Visible : EVisibility::Hidden;
+	return FunctionInput->GetInputType().IsStatic() ? EVisibility::Visible : EVisibility::Collapsed;
 }
 
 FText SNiagaraStackFunctionInputValue::GetTypeModifierIconText() const
@@ -1364,7 +1364,7 @@ EVisibility SNiagaraStackFunctionInputValue::GetFilteredViewContextButtonVisibil
 {
 	UNiagaraEmitterEditorData* EditorData = (FunctionInput && FunctionInput->GetEmitterViewModel())? &FunctionInput->GetEmitterViewModel()->GetOrCreateEditorData() : nullptr;
 	
-	if (!EditorData || EditorData->ShouldShowSummaryView() || FunctionInput->GetTypedOuter<UNiagaraStackSummaryViewObject>() == nullptr)
+	if (!EditorData || EditorData->ShouldShowSummaryView() || FunctionInput->GetTypedOuter<UNiagaraStackSummaryViewObject>() == nullptr || FunctionInput->GetEmitterViewModel()->GetSummaryIsInEditMode() == false)
 	{
 		return EVisibility::Hidden;
 	}
diff --git a/Engine/Plugins/FX/Niagara/Source/NiagaraEditorWidgets/Private/Stack/SNiagaraStackItem.cpp b/Engine/Plugins/FX/Niagara/Source/NiagaraEditorWidgets/Private/Stack/SNiagaraStackItem.cpp
--- a/Engine/Plugins/FX/Niagara/Source/NiagaraEditorWidgets/Private/Stack/SNiagaraStackItem.cpp
+++ b/Engine/Plugins/FX/Niagara/Source/NiagaraEditorWidgets/Private/Stack/SNiagaraStackItem.cpp
@@ -28,8 +28,8 @@ void SNiagaraStackItem::Construct(const FArguments& InArgs, UNiagaraStackItem& I
 
 	TSharedRef<SHorizontalBox> RowBox = SNew(SHorizontalBox);
 
-	// Icon
-	if (Item->SupportsIcon())
+	// Icon Brush
+	if (Item->GetSupportedIconMode() == UNiagaraStackEntry::EIconMode::Brush)
 	{
 		RowBox->AddSlot()
 		.Padding(2, 0, 3, 0)
@@ -41,6 +41,20 @@ void SNiagaraStackItem::Construct(const FArguments& InArgs, UNiagaraStackItem& I
 		];
 	}
 
+	// Icon Text
+	if (Item->GetSupportedIconMode() == UNiagaraStackEntry::EIconMode::Text)
+	{
+		RowBox->AddSlot()
+		.Padding(2, 0, 3, 0)
+		.AutoWidth()
+		.VAlign(VAlign_Center)
+		[
+			SNew(STextBlock)
+			.Font(FEditorStyle::Get().GetFontStyle("FontAwesome.10"))
+			.Text(Item->GetIconText())
+		];
+	}
+
 	// Display name
 	RowBox->AddSlot()
 		.Padding(2, 0, 2, 0)
@@ -54,6 +68,28 @@ void SNiagaraStackItem::Construct(const FArguments& InArgs, UNiagaraStackItem& I
 	// Allow derived classes to add additional widgets.
 	AddCustomRowWidgets(RowBox);
 
+	// Edit Mode Button
+	if (Item->SupportsEditMode())
+	{
+		RowBox->AddSlot()
+		.AutoWidth()
+		.VAlign(VAlign_Center)
+		.Padding(0, 0, 2, 0)
+		[
+			SNew(SButton)
+			.IsFocusable(false)
+			.ToolTipText(this, &SNiagaraStackItem::GetEditModeButtonToolTipText)
+			.ButtonStyle(FNiagaraEditorWidgetsStyle::Get(), "NiagaraEditor.Stack.SimpleButton")
+			.OnClicked(this, &SNiagaraStackItem::EditModeButtonClicked)
+			.ContentPadding(1)
+			[
+				SNew(SImage)
+				.Image(FAppStyle::Get().GetBrush("Icons.Edit"))
+				.ColorAndOpacity(this, &SNiagaraStackItem::GetEditModeButtonIconColor)
+			]
+		];
+	}
+
 	// Reset to base button
 	if (Item->SupportsResetToBase())
 	{
@@ -140,6 +176,29 @@ TSharedRef<SWidget> SNiagaraStackItem::AddContainerForRowWidgets(TSharedRef<SWid
 	return RowWidgets;
 }
 
+FText SNiagaraStackItem::GetEditModeButtonToolTipText() const
+{
+	return Item->GetEditModeIsActive() 
+		? LOCTEXT("DisableEditModeToolTip", "Disable Edit Mode")
+		: LOCTEXT("EnableEditModeToolTip", "Enable Edit Mode");
+}
+
+FReply SNiagaraStackItem::EditModeButtonClicked()
+{
+	if (Item->SupportsEditMode())
+	{
+		Item->SetEditModeIsActive(!Item->GetEditModeIsActive());
+	}
+	return FReply::Handled();
+}
+
+FSlateColor SNiagaraStackItem::GetEditModeButtonIconColor() const
+{
+	return Item->GetEditModeIsActive()
+		? FStyleColors::AccentYellow
+		: FSlateColor::UseSubduedForeground();
+}
+
 EVisibility SNiagaraStackItem::GetResetToBaseButtonVisibility() const
 {
 	FText Unused;
diff --git a/Engine/Plugins/FX/Niagara/Source/NiagaraEditorWidgets/Private/Stack/SNiagaraStackItem.h b/Engine/Plugins/FX/Niagara/Source/NiagaraEditorWidgets/Private/Stack/SNiagaraStackItem.h
--- a/Engine/Plugins/FX/Niagara/Source/NiagaraEditorWidgets/Private/Stack/SNiagaraStackItem.h
+++ b/Engine/Plugins/FX/Niagara/Source/NiagaraEditorWidgets/Private/Stack/SNiagaraStackItem.h
@@ -25,6 +25,12 @@ protected:
 	virtual TSharedRef<SWidget> AddContainerForRowWidgets(TSharedRef<SWidget> RowWidgets);
 
 private:
+	FText GetEditModeButtonToolTipText() const;
+
+	FReply EditModeButtonClicked();
+
+	FSlateColor GetEditModeButtonIconColor() const;
+
 	EVisibility GetResetToBaseButtonVisibility() const;
 
 	FText GetResetToBaseButtonToolTipText() const;
diff --git a/Engine/Plugins/FX/Niagara/Source/NiagaraEditorWidgets/Private/Stack/SNiagaraStackRowPerfWidget.cpp b/Engine/Plugins/FX/Niagara/Source/NiagaraEditorWidgets/Private/Stack/SNiagaraStackRowPerfWidget.cpp
--- a/Engine/Plugins/FX/Niagara/Source/NiagaraEditorWidgets/Private/Stack/SNiagaraStackRowPerfWidget.cpp
+++ b/Engine/Plugins/FX/Niagara/Source/NiagaraEditorWidgets/Private/Stack/SNiagaraStackRowPerfWidget.cpp
@@ -28,10 +28,10 @@ void SNiagaraStackRowPerfWidget::Construct(const FArguments& InArgs, UNiagaraSta
 {
 	StackEntry = InStackEntry;
 	SetToolTipText(CreateTooltipText());
+	SetVisibility(TAttribute<EVisibility>::CreateSP(this, &SNiagaraStackRowPerfWidget::IsVisible));
 	
 	TSharedRef<SWidget> PerfWidget =
 		SNew(SBox)
-		.Visibility(this, &SNiagaraStackRowPerfWidget::IsVisible)
         .HeightOverride(16)
 		.WidthOverride(70)
         .HAlign(HAlign_Right)
