Bind parameter collections associated with simulation stages to the relevant exec context

diff --git a/Engine/Plugins/FX/Niagara/Source/Niagara/Private/NiagaraEmitterInstance.cpp b/Engine/Plugins/FX/Niagara/Source/Niagara/Private/NiagaraEmitterInstance.cpp
--- a/Engine/Plugins/FX/Niagara/Source/Niagara/Private/NiagaraEmitterInstance.cpp
+++ b/Engine/Plugins/FX/Niagara/Source/Niagara/Private/NiagaraEmitterInstance.cpp
@@ -806,6 +806,20 @@ void FNiagaraEmitterInstance::BindParameters(bool bExternalOnly)
 		}
 	}
 
+	FNiagaraScriptInstanceParameterStore& TargetParamStore = CachedEmitter->SimTarget == ENiagaraSimTarget::GPUComputeSim ? GPUExecContext->CombinedParamStore : UpdateExecContext.Parameters;
+	for (const UNiagaraSimulationStageBase* SimStage : CachedEmitter->GetSimulationStages())
+	{
+		if (SimStage->bEnabled == false)
+		{
+			continue;
+		}
+
+		for (UNiagaraParameterCollection* Collection : SimStage->Script->GetCachedParameterCollectionReferences())
+		{
+			BindToParameterCollection(Collection, TargetParamStore);
+		}
+	}
+
 	if (!bExternalOnly)
 	{
 		//Now bind parameters from the component and system.
