Fixed a bug where user parameters linked to position types were only partly converted over to positions

diff --git a/Engine/Plugins/FX/Niagara/Source/Niagara/Classes/NiagaraScriptSourceBase.h b/Engine/Plugins/FX/Niagara/Source/Niagara/Classes/NiagaraScriptSourceBase.h
--- a/Engine/Plugins/FX/Niagara/Source/Niagara/Classes/NiagaraScriptSourceBase.h
+++ b/Engine/Plugins/FX/Niagara/Source/Niagara/Classes/NiagaraScriptSourceBase.h
@@ -171,6 +171,7 @@ class UNiagaraScriptSourceBase : public UObject
 	/** Checks if any of the provided variables are linked to function inputs of position type data
 	 */
 	virtual void GetLinkedPositionTypeInputs(const TArray<FNiagaraVariable>& ParametersToCheck, TSet<FNiagaraVariable>& OutLinkedParameters) {};
+	virtual void ChangedLinkedInputTypes(const FNiagaraVariable& ParametersToCheck, const FNiagaraTypeDefinition& NewType) {};
 
 protected:
 	FOnChanged OnChangedDelegate;
diff --git a/Engine/Plugins/FX/Niagara/Source/Niagara/Private/NiagaraSystem.cpp b/Engine/Plugins/FX/Niagara/Source/Niagara/Private/NiagaraSystem.cpp
--- a/Engine/Plugins/FX/Niagara/Source/Niagara/Private/NiagaraSystem.cpp
+++ b/Engine/Plugins/FX/Niagara/Source/Niagara/Private/NiagaraSystem.cpp
@@ -3807,10 +3807,20 @@ void UNiagaraSystem::FixupPositionUserParameters()
 
 	// looks like we have a few FVector3f user parameters that are linked to position inputs in the stack.
 	// Most likely this is because core modules were changed to use position data. To ease the transition of old assets, we auto-convert those inputs to position types as well.
-	for (FNiagaraVariable& LinkedParameter : LinkedPositionInputs)
+	for (const FNiagaraVariable& LinkedParameter : LinkedPositionInputs)
 	{
 		ExposedParameters.ConvertParameterType(LinkedParameter, FNiagaraTypeDefinition::GetPositionDef());
 		UE_LOG(LogNiagara, Log, TEXT("Converted parameter %s from vec3 to position type in asset %s"), *LinkedParameter.GetName().ToString(), *GetPathName());
+
+		// we need to go over the scripts again to fix existing usages. Just because it's linked to a position in one input doesn't mean that's the case for all of them and we don't want to end up with
+		// different types of the same user parameter linked throughout the system.
+		ForEachScript([&LinkedParameter](UNiagaraScript* Script)
+		{
+			if (Script)
+			{
+				Script->GetLatestSource()->ChangedLinkedInputTypes(LinkedParameter, FNiagaraTypeDefinition::GetPositionDef());
+			}
+		});
 	}
 }
 
diff --git a/Engine/Plugins/FX/Niagara/Source/NiagaraEditor/Private/NiagaraGraph.cpp b/Engine/Plugins/FX/Niagara/Source/NiagaraEditor/Private/NiagaraGraph.cpp
--- a/Engine/Plugins/FX/Niagara/Source/NiagaraEditor/Private/NiagaraGraph.cpp
+++ b/Engine/Plugins/FX/Niagara/Source/NiagaraEditor/Private/NiagaraGraph.cpp
@@ -404,6 +404,69 @@ void UNiagaraGraph::PostLoad()
 	InvalidateCachedParameterData();
 }
 
+void UNiagaraGraph::ChangeParameterType(const FNiagaraVariable& CurrentParameter, const FNiagaraTypeDefinition& NewType)
+{
+	const UEdGraphSchema_Niagara* NiagaraSchema = GetDefault<UEdGraphSchema_Niagara>();
+	TArray<UNiagaraNodeParameterMapBase*> NiagaraParameterNodes;
+	GetNodesOfClass<UNiagaraNodeParameterMapBase>(NiagaraParameterNodes);
+	for (UNiagaraNodeParameterMapBase* NiagaraNode : NiagaraParameterNodes)
+	{
+		for (UEdGraphPin* Pin : NiagaraNode->Pins)
+		{
+			if (NiagaraNode->IsAddPin(Pin))
+			{
+				continue;
+			}
+			FNiagaraVariable Variable = FNiagaraVariable(NiagaraSchema->PinToTypeDefinition(Pin), Pin->PinName);
+			if (Variable == CurrentParameter)
+			{
+				Pin->Modify();
+				Pin->PinType = NiagaraSchema->TypeDefinitionToPinType(NewType);
+				NiagaraNode->MarkNodeRequiresSynchronization(__FUNCTION__, true);
+				break;
+			}
+
+			// fix default pin types if necessary
+			if (UNiagaraNodeParameterMapGet* GetNode = Cast<UNiagaraNodeParameterMapGet>(NiagaraNode))
+			{
+				UEdGraphPin* DefaultPin = GetNode->GetDefaultPin(Pin);
+				if (Pin->Direction == EGPD_Output && DefaultPin && DefaultPin->PinType != Pin->PinType)
+				{
+					DefaultPin->PinType = Pin->PinType;
+					NiagaraNode->MarkNodeRequiresSynchronization(__FUNCTION__, true);
+				}
+			}
+		}
+	}
+	
+	TObjectPtr<UNiagaraScriptVariable>* ScriptVariablePtr = VariableToScriptVariable.Find(CurrentParameter);
+	if (ScriptVariablePtr != nullptr && !ScriptVariablePtr->IsNull())
+	{
+		UNiagaraScriptVariable& ScriptVar = *ScriptVariablePtr->Get();
+		ScriptVar.Modify();
+		ScriptVar.Variable.SetType(NewType);
+		VariableToScriptVariable.Remove(CurrentParameter);
+
+		FNiagaraVariable NewVarType(NewType, CurrentParameter.GetName());
+		VariableToScriptVariable.Add(NewVarType, TObjectPtr<UNiagaraScriptVariable>(&ScriptVar));
+		ScriptVariableChanged(NewVarType);
+	}
+
+	// Also fix stale parameter reference map entries. This should usually happen in RefreshParameterReferences, but since we just changed VariableToScriptVariable we need to fix that as well. 
+	FNiagaraGraphParameterReferenceCollection* ReferenceCollection = ParameterToReferencesMap.Find(CurrentParameter);
+	if (ReferenceCollection != nullptr)
+	{
+		if (ReferenceCollection->ParameterReferences.Num() > 0)
+		{
+			FNiagaraVariable NewVarType(NewType, CurrentParameter.GetName());
+			FNiagaraGraphParameterReferenceCollection NewReferenceCollection = *ReferenceCollection;
+			ParameterToReferencesMap.Add(NewVarType, NewReferenceCollection);
+		}
+		ParameterToReferencesMap.Remove(CurrentParameter);
+	}
+	InvalidateCachedParameterData();
+}
+
 TArray<UEdGraphPin*> UNiagaraGraph::FindParameterMapDefaultValuePins(const FName VariableName) const
 {
 	TArray<UEdGraphPin*> DefaultPins;
diff --git a/Engine/Plugins/FX/Niagara/Source/NiagaraEditor/Private/NiagaraScriptSource.cpp b/Engine/Plugins/FX/Niagara/Source/NiagaraEditor/Private/NiagaraScriptSource.cpp
--- a/Engine/Plugins/FX/Niagara/Source/NiagaraEditor/Private/NiagaraScriptSource.cpp
+++ b/Engine/Plugins/FX/Niagara/Source/NiagaraEditor/Private/NiagaraScriptSource.cpp
@@ -649,3 +649,12 @@ void UNiagaraScriptSource::GetLinkedPositionTypeInputs(const TArray<FNiagaraVari
 		}
 	}
 }
+
+void UNiagaraScriptSource::ChangedLinkedInputTypes(const FNiagaraVariable& ParametersToChange, const FNiagaraTypeDefinition& NewType)
+{
+	if (!NodeGraph)
+	{
+		return;
+	}
+	NodeGraph->ChangeParameterType(ParametersToChange, NewType);
+}
diff --git a/Engine/Plugins/FX/Niagara/Source/NiagaraEditor/Public/NiagaraGraph.h b/Engine/Plugins/FX/Niagara/Source/NiagaraEditor/Public/NiagaraGraph.h
--- a/Engine/Plugins/FX/Niagara/Source/NiagaraEditor/Public/NiagaraGraph.h
+++ b/Engine/Plugins/FX/Niagara/Source/NiagaraEditor/Public/NiagaraGraph.h
@@ -308,6 +308,9 @@ class UNiagaraGraph : public UEdGraph
 	/** Rename a pin inline in a graph. If this is the only instance used in the graph, then rename them all, otherwise make a duplicate. */
 	bool RenameParameterFromPin(const FNiagaraVariable& Parameter, FName NewName, UEdGraphPin* InPin);
 
+	/** Changes the type of an existing graph variable. This is an unsafe operation, usually guarded by a cvar and not directly exposed to users. */
+	void ChangeParameterType(const FNiagaraVariable& CurrentParameter, const FNiagaraTypeDefinition& NewType);
+
 	/** Gets a delegate which is called whenever a contained data interfaces changes. */
 	FOnDataInterfaceChanged& OnDataInterfaceChanged();
 
diff --git a/Engine/Plugins/FX/Niagara/Source/NiagaraEditor/Public/NiagaraScriptSource.h b/Engine/Plugins/FX/Niagara/Source/NiagaraEditor/Public/NiagaraScriptSource.h
--- a/Engine/Plugins/FX/Niagara/Source/NiagaraEditor/Public/NiagaraScriptSource.h
+++ b/Engine/Plugins/FX/Niagara/Source/NiagaraEditor/Public/NiagaraScriptSource.h
@@ -78,6 +78,7 @@ class UNiagaraScriptSource : public UNiagaraScriptSourceBase
 	virtual void RenameGraphAssignmentAndSetNodePins(const FName OldName, const FName NewName) override;
 
 	virtual void GetLinkedPositionTypeInputs(const TArray<FNiagaraVariable>& ParametersToCheck, TSet<FNiagaraVariable>& OutLinkedParameters) override;
+	virtual void ChangedLinkedInputTypes(const FNiagaraVariable& ParametersToChange, const FNiagaraTypeDefinition& NewType) override;
 private:
 	void OnGraphChanged(const FEdGraphEditAction &Action);
 	void OnGraphDataInterfaceChanged();
