Niagara - Fix missing include from last checkin.

diff --git a/Engine/Plugins/FX/Niagara/Source/NiagaraEditorWidgets/Private/Stack/SNiagaraStackFunctionInputCollection.cpp b/Engine/Plugins/FX/Niagara/Source/NiagaraEditorWidgets/Private/Stack/SNiagaraStackFunctionInputCollection.cpp
--- a/Engine/Plugins/FX/Niagara/Source/NiagaraEditorWidgets/Private/Stack/SNiagaraStackFunctionInputCollection.cpp
+++ b/Engine/Plugins/FX/Niagara/Source/NiagaraEditorWidgets/Private/Stack/SNiagaraStackFunctionInputCollection.cpp
@@ -4,6 +4,7 @@
 
 #include "NiagaraEditorWidgetsStyle.h"
 #include "ViewModels/Stack/NiagaraStackFunctionInputCollection.h"
+#include "Widgets/Input/SCheckBox.h"
 #include "Widgets/Layout/SWrapBox.h"
 
 void SNiagaraStackFunctionInputCollection::Construct(const FArguments& InArgs, UNiagaraStackFunctionInputCollectionBase* InInputCollection)
