- Ensure we tick instance parameters before building data interfaces, otherwise the GPU tick will have a reference to the wrong data interface

diff --git a/Engine/Plugins/FX/Niagara/Source/Niagara/Private/NiagaraSystemInstance.cpp b/Engine/Plugins/FX/Niagara/Source/Niagara/Private/NiagaraSystemInstance.cpp
--- a/Engine/Plugins/FX/Niagara/Source/Niagara/Private/NiagaraSystemInstance.cpp
+++ b/Engine/Plugins/FX/Niagara/Source/Niagara/Private/NiagaraSystemInstance.cpp
@@ -872,7 +872,6 @@ void FNiagaraSystemInstance::Reset(FNiagaraSystemInstance::EResetMode Mode)
 
 		if (bBindParams)
 		{
-			InstanceParameters.Tick();//Make sure the owner has flushed it's parameters by now. Especially it's DIs.
 			InitDataInterfaces();
 		}
 
@@ -942,8 +941,7 @@ void FNiagaraSystemInstance::ResetInternal(bool bResetSimulations)
 	check(World);
 	if (OverrideParameters && World->WorldType == EWorldType::Editor)
 	{
-		OverrideParameters->ResolvePositions(GetLWCConverter());
-		OverrideParameters->Tick();
+		InitDataInterfaces();
 	}
 #endif
 
@@ -1407,6 +1405,8 @@ void FNiagaraSystemInstance::InitDataInterfaces()
 		OverrideParameters->ResolvePositions(GetLWCConverter());
 		OverrideParameters->Tick();
 	}
+	// Make sure the owner has flushed it's parameters before we initialize data interfaces
+	InstanceParameters.Tick();
 
 	//-TODO: Validate that any queued ticks have been executed
 	DestroyDataInterfaceInstanceData();
