- Remove memory allocations from a few commonly used functions

#rnx

diff --git a/Engine/Plugins/FX/Niagara/Source/Niagara/Private/NiagaraUserRedirectionParameterStore.cpp b/Engine/Plugins/FX/Niagara/Source/Niagara/Private/NiagaraUserRedirectionParameterStore.cpp
--- a/Engine/Plugins/FX/Niagara/Source/Niagara/Private/NiagaraUserRedirectionParameterStore.cpp
+++ b/Engine/Plugins/FX/Niagara/Source/Niagara/Private/NiagaraUserRedirectionParameterStore.cpp
@@ -22,7 +22,8 @@ FNiagaraUserRedirectionParameterStore& FNiagaraUserRedirectionParameterStore::op
 
 bool FNiagaraUserRedirectionParameterStore::IsUserParameter(const FNiagaraVariableBase& InVar) 
 {
-	return InVar.GetName().ToString().StartsWith(TEXT("User."));
+	FNameBuilder VarName(InVar.GetName());
+	return VarName.ToView().StartsWith(TEXT("User."));
 }
 
 void FNiagaraUserRedirectionParameterStore::MakeUserVariable(FNiagaraVariableBase& InVar)
@@ -31,8 +32,10 @@ void FNiagaraUserRedirectionParameterStore::MakeUserVariable(FNiagaraVariableBas
 	{
 		return;
 	}
-	FName DisplayName(*(TEXT("User.") + InVar.GetName().ToString()));
-	InVar.SetName(DisplayName);
+	FNameBuilder DisplayName;
+	DisplayName.Append(TEXT("User."));
+	InVar.GetName().AppendString(DisplayName);
+	InVar.SetName(FName(DisplayName));
 	return;
 }
 
@@ -44,8 +47,9 @@ FNiagaraVariable FNiagaraUserRedirectionParameterStore::GetUserRedirection(const
 		return InVar;
 	}
 	FNiagaraVariable SimpleVar = InVar;
-	FName DisplayName(*InVar.GetName().ToString().RightChop(5));
-	SimpleVar.SetName(DisplayName);
+	FNameBuilder DisplayName;
+	InVar.GetName().ToString(DisplayName);
+	SimpleVar.SetName(FName(DisplayName.ToView().RightChop(5)));
 	return SimpleVar;
 }
 
