Add a bias to SubUV calculations to avoid errors when 1/ImageCount loses numerical precision (most significantly when it's half precision on mobile platforms)

diff --git a/Engine/Plugins/FX/Niagara/Shaders/Private/NiagaraSpriteVertexFactory.ush b/Engine/Plugins/FX/Niagara/Shaders/Private/NiagaraSpriteVertexFactory.ush
--- a/Engine/Plugins/FX/Niagara/Shaders/Private/NiagaraSpriteVertexFactory.ush
+++ b/Engine/Plugins/FX/Niagara/Shaders/Private/NiagaraSpriteVertexFactory.ush
@@ -909,10 +909,12 @@ FVertexFactoryIntermediates GetVertexFactoryIntermediates(FVertexFactoryInput In
 	
 	// SubUV.
 	float SubImageLerp = frac(SubImageIndex);
-	float SubImageA = SubImageIndex - SubImageLerp;
+
+	// add a bias to SubImageA to avoid numerical precision issues around SubImageSize.x vs SubImageSize.z
+	float SubImageA = SubImageIndex - SubImageLerp + 0.5f;
 	float SubImageB = SubImageA + 1;
-	float SubImageAH = fmod( SubImageA, NiagaraSpriteVF.SubImageSize.x );
-	float SubImageBH = fmod( SubImageB, NiagaraSpriteVF.SubImageSize.x );
+	float SubImageAH = floor( fmod( SubImageA, NiagaraSpriteVF.SubImageSize.x ) );
+	float SubImageBH = floor( fmod( SubImageB, NiagaraSpriteVF.SubImageSize.x ) );
 	float SubImageAV = floor( SubImageA * NiagaraSpriteVF.SubImageSize.z );
 	float SubImageBV = floor( SubImageB * NiagaraSpriteVF.SubImageSize.z );
 	Intermediates.TexCoord.xy = (float2( SubImageAH, SubImageAV ) + UVForTexturing) * NiagaraSpriteVF.SubImageSize.zw;
diff --git a/Engine/Shaders/Private/ParticleSpriteVertexFactory.ush b/Engine/Shaders/Private/ParticleSpriteVertexFactory.ush
--- a/Engine/Shaders/Private/ParticleSpriteVertexFactory.ush
+++ b/Engine/Shaders/Private/ParticleSpriteVertexFactory.ush
@@ -496,10 +496,12 @@ FVertexFactoryIntermediates GetVertexFactoryIntermediates(FVertexFactoryInput In
 	// SubUV.
 	float SubImageIndex = Input.SizeRotSubImage.w;
 	float SubImageLerp = frac(SubImageIndex);
-	float SubImageA = SubImageIndex - SubImageLerp;
+
+	// add a bias to SubImageA to avoid numerical precision issues around SubImageSize.x vs SubImageSize.z
+	float SubImageA = SubImageIndex - SubImageLerp + 0.5f;
 	float SubImageB = SubImageA + 1;
-	float SubImageAH = fmod( SubImageA, SpriteVF.SubImageSize.x );
-	float SubImageBH = fmod( SubImageB, SpriteVF.SubImageSize.x );
+	float SubImageAH = floor( fmod( SubImageA, SpriteVF.SubImageSize.x ) );
+	float SubImageBH = floor( fmod( SubImageB, SpriteVF.SubImageSize.x ) );
 	float SubImageAV = floor( SubImageA * SpriteVF.SubImageSize.z );
 	float SubImageBV = floor( SubImageB * SpriteVF.SubImageSize.z );
 	Intermediates.TexCoord.xy = (float2( SubImageAH, SubImageAV ) + UVForTexturing) * SpriteVF.SubImageSize.zw ;
