Changed the Niagara thumbnail renderer to use its parents draw function instead

diff --git a/Engine/Plugins/FX/Niagara/Source/NiagaraEditor/Private/NiagaraThumbnailRenderer.cpp b/Engine/Plugins/FX/Niagara/Source/NiagaraEditor/Private/NiagaraThumbnailRenderer.cpp
--- a/Engine/Plugins/FX/Niagara/Source/NiagaraEditor/Private/NiagaraThumbnailRenderer.cpp
+++ b/Engine/Plugins/FX/Niagara/Source/NiagaraEditor/Private/NiagaraThumbnailRenderer.cpp
@@ -25,12 +25,12 @@ void UNiagaraThumbnailRendererBase::GetThumbnailSize(UObject* Object, float Zoom
 	}
 }
 
-void UNiagaraThumbnailRendererBase::Draw(UObject* Object, int32 X, int32 Y, uint32 Width, uint32 Height, FRenderTarget*, FCanvas* Canvas, bool bAdditionalViewFamily)
+void UNiagaraThumbnailRendererBase::Draw(UObject* Object, int32 X, int32 Y, uint32 Width, uint32 Height, FRenderTarget* RenderTarget, FCanvas* Canvas, bool bAdditionalViewFamily)
 {
 	UTexture2D* ObjectTexture = GetThumbnailTextureFromObject(Object);
 	if (ObjectTexture != nullptr)
 	{
-		Canvas->DrawTile(X, Y, Width, Height, 0.0f, 0.0f, 1.0f, 1.0f, FLinearColor::White, ObjectTexture->GetResource(), false);
+		Super::Draw(ObjectTexture, X, Y, Width, Height, RenderTarget, Canvas, bAdditionalViewFamily);
 	}
 }
 
