LWCify the RVT WorldToUV calculation.
Note that Niagara RVT data provider will continue with non LWC behavior and still needs converting.
#prefligh 61f40fecf628752a568b8917

diff --git a/Engine/Plugins/FX/Niagara/Shaders/Private/NiagaraDataInterfaceLandscape.ush b/Engine/Plugins/FX/Niagara/Shaders/Private/NiagaraDataInterfaceLandscape.ush
--- a/Engine/Plugins/FX/Niagara/Shaders/Private/NiagaraDataInterfaceLandscape.ush
+++ b/Engine/Plugins/FX/Niagara/Shaders/Private/NiagaraDataInterfaceLandscape.ush
@@ -63,13 +63,14 @@ struct FNDILandscapeContext
 
 void NDILandscape_GetBaseColor(in FNDILandscapeContext DIContext, float3 InWorldPos, out float3 OutBaseColor, out bool OutIsValid)
 {
+	// todo[mg] convert world pos and transform origin to LWC pos
 	if (DIContext.BaseColorVirtualTextureEnabled)
 	{
 		float3 UnpackBaseColorOrigin = DIContext.BaseColorVirtualTextureWorldToUvTransform[0].xyz;
 		float3 UnpackBaseColorU = DIContext.BaseColorVirtualTextureWorldToUvTransform[1].xyz;
 		float3 UnpackBaseColorV = DIContext.BaseColorVirtualTextureWorldToUvTransform[2].xyz;
 
-		float2 SampleUv = VirtualTextureWorldToUV(InWorldPos, UnpackBaseColorOrigin, UnpackBaseColorU, UnpackBaseColorV);
+		float2 SampleUv = VirtualTextureWorldToUV(LWCPromote(InWorldPos), LWCPromote(UnpackBaseColorOrigin), UnpackBaseColorU, UnpackBaseColorV);
 
 		VTPageTableResult PageTable = TextureLoadVirtualPageTableLevel(
 			DIContext.BaseColorVirtualTexturePageTable,
@@ -98,7 +99,7 @@ void NDILandscape_GetBaseColor(in FNDILandscapeContext DIContext, float3 InWorld
 
 void NDILandscape_GetHeight(in FNDILandscapeContext DIContext, float3 InWorldPos, out float OutHeight, out bool OutIsValid)
 {
-	// todo[mg] convert world pos to LWC pos
+	// todo[mg] convert world pos and transform origin to LWC pos
 	if (DIContext.HeightVirtualTextureEnabled)
 	{
 		float3 UnpackHeightOrigin = DIContext.HeightVirtualTextureWorldToUvTransform[0].xyz;
@@ -106,7 +107,7 @@ void NDILandscape_GetHeight(in FNDILandscapeContext DIContext, float3 InWorldPos
 		float3 UnpackHeightV = DIContext.HeightVirtualTextureWorldToUvTransform[2].xyz;
 		float2 UnpackHeightScaleBias = DIContext.HeightVirtualTextureWorldToUvTransform[3].xy;
 
-		float2 SampleUv = VirtualTextureWorldToUV(InWorldPos, UnpackHeightOrigin, UnpackHeightU, UnpackHeightV);
+		float2 SampleUv = VirtualTextureWorldToUV(LWCPromote(InWorldPos), LWCPromote(UnpackHeightOrigin), UnpackHeightU, UnpackHeightV);
 
 		VTPageTableResult PageTable = TextureLoadVirtualPageTableLevel(
 			DIContext.HeightVirtualTexturePageTable,
@@ -142,14 +143,14 @@ void NDILandscape_GetHeight(in FNDILandscapeContext DIContext, float3 InWorldPos
 
 void NDILandscape_GetWorldNormal(in FNDILandscapeContext DIContext, float3 InWorldPos, out float3 OutNormal, out bool OutIsValid)
 {
-	// todo[mg] convert world pos to LWC pos
+	// todo[mg] convert world pos and transform origin to LWC pos
 	if (DIContext.NormalVirtualTextureEnabled)
 	{
 		float3 UnpackNormalOrigin = DIContext.NormalVirtualTextureWorldToUvTransform[0].xyz;
 		float3 UnpackNormalU = DIContext.NormalVirtualTextureWorldToUvTransform[1].xyz;
 		float3 UnpackNormalV = DIContext.NormalVirtualTextureWorldToUvTransform[2].xyz;
 
-		float2 SampleUv = VirtualTextureWorldToUV(InWorldPos, UnpackNormalOrigin, UnpackNormalU, UnpackNormalV);
+		float2 SampleUv = VirtualTextureWorldToUV(LWCPromote(InWorldPos), LWCPromote(UnpackNormalOrigin), UnpackNormalU, UnpackNormalV);
 
 		VTPageTableResult PageTable = TextureLoadVirtualPageTableLevel(
 			DIContext.NormalVirtualTexturePageTable,
diff --git a/Engine/Shaders/Private/VirtualTextureCommon.ush b/Engine/Shaders/Private/VirtualTextureCommon.ush
--- a/Engine/Shaders/Private/VirtualTextureCommon.ush
+++ b/Engine/Shaders/Private/VirtualTextureCommon.ush
@@ -717,13 +717,11 @@ MaterialFloat4 TextureVirtualSampleLevel(
 }
 
 
-/**
-	Helper function to convert world space to virtual texture UV space
-	todo[vt]: For large worlds we will probably want to use position relative to Camera to avoid precision issues
-*/
-float2 VirtualTextureWorldToUV(in float3 WorldPos, in float3 Origin, in float3 U, in float3 V)
+/**	Helper function to convert world space to virtual texture UV space. */
+float2 VirtualTextureWorldToUV(in FLWCVector3 WorldPos, in FLWCVector3 O, in float3 U, in float3 V)
 {
-	float3 P = WorldPos - Origin;
+	// After the subtract we can assume that the offset is small enough for float maths.
+	float3 P = LWCToFloat(LWCSubtract(WorldPos, O));
 	return float2(dot(P, U), dot(P, V));
 }
 
diff --git a/Engine/Source/Runtime/Engine/Classes/VT/RuntimeVirtualTexture.h b/Engine/Source/Runtime/Engine/Classes/VT/RuntimeVirtualTexture.h
--- a/Engine/Source/Runtime/Engine/Classes/VT/RuntimeVirtualTexture.h
+++ b/Engine/Source/Runtime/Engine/Classes/VT/RuntimeVirtualTexture.h
@@ -9,6 +9,8 @@
 #include "VT/RuntimeVirtualTextureEnum.h"
 #include "RuntimeVirtualTexture.generated.h"
 
+namespace UE { namespace Shader	{ enum class EValueType : uint8; } }
+
 /** Runtime virtual texture UObject */
 UCLASS(ClassGroup = Rendering, BlueprintType)
 class ENGINE_API URuntimeVirtualTexture : public UObject
@@ -159,6 +161,8 @@ public:
 
 	/** Getter for the shader uniform parameters. */
 	FVector4 GetUniformParameter(int32 Index) const;
+	/** Getter for the shader uniform parameter type. */
+	static UE::Shader::EValueType GetUniformParameterType(int32 Index);
 
 protected:
 	/** Initialize the render resources. This kicks off render thread work. */
diff --git a/Engine/Source/Runtime/Engine/Private/Materials/HLSLMaterialTranslator.cpp b/Engine/Source/Runtime/Engine/Private/Materials/HLSLMaterialTranslator.cpp
--- a/Engine/Source/Runtime/Engine/Private/Materials/HLSLMaterialTranslator.cpp
+++ b/Engine/Source/Runtime/Engine/Private/Materials/HLSLMaterialTranslator.cpp
@@ -6688,17 +6688,17 @@ int32 FHLSLMaterialTranslator::VirtualTextureParameter(FName ParameterName, URun
 	return AddUniformExpression(new FMaterialUniformExpressionTextureParameter(ParameterInfo, TextureReferenceIndex, TextureLayerIndex, PageTableLayerIndex, SamplerType), MCT_TextureVirtual, TEXT(""));
 }
 
-int32 FHLSLMaterialTranslator::VirtualTextureUniform(int32 TextureIndex, int32 VectorIndex)
+int32 FHLSLMaterialTranslator::VirtualTextureUniform(int32 TextureIndex, int32 VectorIndex, UE::Shader::EValueType Type)
 {
-	return AddUniformExpression(new FMaterialUniformExpressionRuntimeVirtualTextureUniform(TextureIndex, VectorIndex), MCT_Float3, TEXT(""));
+	return AddUniformExpression(new FMaterialUniformExpressionRuntimeVirtualTextureUniform(TextureIndex, VectorIndex), GetMaterialValueType(Type), TEXT(""));
 }
 
-int32 FHLSLMaterialTranslator::VirtualTextureUniform(FName ParameterName, int32 TextureIndex, int32 VectorIndex)
+int32 FHLSLMaterialTranslator::VirtualTextureUniform(FName ParameterName, int32 TextureIndex, int32 VectorIndex, UE::Shader::EValueType Type)
 {
 	FMaterialParameterInfo ParameterInfo = GetParameterAssociationInfo();
 	ParameterInfo.Name = ParameterName;
 
-	return AddUniformExpression(new FMaterialUniformExpressionRuntimeVirtualTextureUniform(ParameterInfo, TextureIndex, VectorIndex), MCT_Float3, TEXT(""));
+	return AddUniformExpression(new FMaterialUniformExpressionRuntimeVirtualTextureUniform(ParameterInfo, TextureIndex, VectorIndex), GetMaterialValueType(Type), TEXT(""));
 }
 
 int32 FHLSLMaterialTranslator::VirtualTextureWorldToUV(int32 WorldPositionIndex, int32 P0, int32 P1, int32 P2)
@@ -6708,8 +6708,7 @@ int32 FHLSLMaterialTranslator::VirtualTextureWorldToUV(int32 WorldPositionIndex,
 		return INDEX_NONE;
 	}
 	FString	SampleCode(TEXT("VirtualTextureWorldToUV(%s, %s, %s, %s)"));
-	// LWC_TODO: This should be made LWC aware, but need to update VirtualTexture uniforms to match, no support for LWC material parameters currently
-	return AddInlinedCodeChunk(MCT_Float2, *SampleCode, *CoerceParameter(WorldPositionIndex, MCT_Float3), *GetParameterCode(P0), *GetParameterCode(P1), *GetParameterCode(P2));
+	return AddInlinedCodeChunk(MCT_Float2, *SampleCode, *CoerceParameter(WorldPositionIndex, MCT_LWCVector3), *GetParameterCode(P0), *GetParameterCode(P1), *GetParameterCode(P2));
 }
 
 int32 FHLSLMaterialTranslator::VirtualTextureUnpack(int32 CodeIndex0, int32 CodeIndex1, int32 CodeIndex2, int32 P0, EVirtualTextureUnpackType UnpackType)
diff --git a/Engine/Source/Runtime/Engine/Private/Materials/HLSLMaterialTranslator.h b/Engine/Source/Runtime/Engine/Private/Materials/HLSLMaterialTranslator.h
--- a/Engine/Source/Runtime/Engine/Private/Materials/HLSLMaterialTranslator.h
+++ b/Engine/Source/Runtime/Engine/Private/Materials/HLSLMaterialTranslator.h
@@ -842,8 +842,8 @@ protected:
 
 	virtual int32 VirtualTexture(URuntimeVirtualTexture* InTexture, int32 TextureLayerIndex, int32 PageTableLayerIndex, int32& TextureReferenceIndex, EMaterialSamplerType SamplerType) override;
 	virtual int32 VirtualTextureParameter(FName ParameterName, URuntimeVirtualTexture* DefaultValue, int32 TextureLayerIndex, int32 PageTableLayerIndex, int32& TextureReferenceIndex, EMaterialSamplerType SamplerType) override;
-	virtual int32 VirtualTextureUniform(int32 TextureIndex, int32 VectorIndex) override;
-	virtual int32 VirtualTextureUniform(FName ParameterName, int32 TextureIndex, int32 VectorIndex) override;
+	virtual int32 VirtualTextureUniform(int32 TextureIndex, int32 VectorIndex, UE::Shader::EValueType Type) override;
+	virtual int32 VirtualTextureUniform(FName ParameterName, int32 TextureIndex, int32 VectorIndex, UE::Shader::EValueType Type) override;
 	virtual int32 VirtualTextureWorldToUV(int32 WorldPositionIndex, int32 P0, int32 P1, int32 P2) override;
 	virtual int32 VirtualTextureUnpack(int32 CodeIndex0, int32 CodeIndex1, int32 CodeIndex2, int32 P0, EVirtualTextureUnpackType UnpackType) override;
 
diff --git a/Engine/Source/Runtime/Engine/Private/Materials/MaterialExpressions.cpp b/Engine/Source/Runtime/Engine/Private/Materials/MaterialExpressions.cpp
--- a/Engine/Source/Runtime/Engine/Private/Materials/MaterialExpressions.cpp
+++ b/Engine/Source/Runtime/Engine/Private/Materials/MaterialExpressions.cpp
@@ -2562,13 +2562,14 @@ int32 UMaterialExpressionRuntimeVirtualTextureSample::Compile(class FMaterialCom
 	int32 Uniforms[ERuntimeVirtualTextureShaderUniform_Count];
 	for (int32 UniformIndex = 0; UniformIndex < ERuntimeVirtualTextureShaderUniform_Count; ++UniformIndex)
 	{
+		const UE::Shader::EValueType Type = URuntimeVirtualTexture::GetUniformParameterType(UniformIndex);
 		if (bIsParameter)
 		{
-			Uniforms[UniformIndex] = Compiler->VirtualTextureUniform(GetParameterName(), TextureReferenceIndex[0], UniformIndex);
+			Uniforms[UniformIndex] = Compiler->VirtualTextureUniform(GetParameterName(), TextureReferenceIndex[0], UniformIndex, Type);
 		}
 		else
 		{
-			Uniforms[UniformIndex] = Compiler->VirtualTextureUniform(TextureReferenceIndex[0], UniformIndex);
+			Uniforms[UniformIndex] = Compiler->VirtualTextureUniform(TextureReferenceIndex[0], UniformIndex, Type);
 		}
 	}
 
diff --git a/Engine/Source/Runtime/Engine/Private/VT/RuntimeVirtualTexture.cpp b/Engine/Source/Runtime/Engine/Private/VT/RuntimeVirtualTexture.cpp
--- a/Engine/Source/Runtime/Engine/Private/VT/RuntimeVirtualTexture.cpp
+++ b/Engine/Source/Runtime/Engine/Private/VT/RuntimeVirtualTexture.cpp
@@ -8,6 +8,7 @@
 #include "Engine/TextureLODSettings.h"
 #include "Interfaces/ITargetPlatform.h"
 #include "RendererInterface.h"
+#include "Shader/ShaderTypes.h"
 #include "VT/RuntimeVirtualTextureNotify.h"
 #include "VT/UploadingVirtualTexture.h"
 #include "VT/VirtualTexture.h"
@@ -515,10 +516,24 @@ FVector4 URuntimeVirtualTexture::GetUniformParameter(int32 Index) const
 	return FVector4(ForceInitToZero);
 }
 
+UE::Shader::EValueType URuntimeVirtualTexture::GetUniformParameterType(int32 Index)
+{
+	switch (Index)
+	{
+	case ERuntimeVirtualTextureShaderUniform_WorldToUVTransform0: return UE::Shader::EValueType::Double3;
+	case ERuntimeVirtualTextureShaderUniform_WorldToUVTransform1: return UE::Shader::EValueType::Float3;
+	case ERuntimeVirtualTextureShaderUniform_WorldToUVTransform2: return UE::Shader::EValueType::Float3;
+	case ERuntimeVirtualTextureShaderUniform_WorldHeightUnpack: return UE::Shader::EValueType::Float2;
+	default:
+		break;
+	}
+
+	check(0);
+	return UE::Shader::EValueType::Float4;
+}
+
 void URuntimeVirtualTexture::Initialize(IVirtualTexture* InProducer, FVTProducerDescription const& InProducerDesc, FTransform const& InVolumeToWorld, FBox const& InWorldBounds)
 {
-	//todo[vt]: possible issues with precision in large worlds here it might be better to calculate/upload camera space relative transform per frame?
-	// LWC_TODO: If we don't update this to use translated world space, will need to update this to store a FLWCVector3 for translation (and update VirtualTextureWorldToUV to match)
 	WorldToUVTransformParameters[0] = InVolumeToWorld.GetTranslation();
 	WorldToUVTransformParameters[1] = InVolumeToWorld.GetUnitAxis(EAxis::X) * 1.f / InVolumeToWorld.GetScale3D().X;
 	WorldToUVTransformParameters[2] = InVolumeToWorld.GetUnitAxis(EAxis::Y) * 1.f / InVolumeToWorld.GetScale3D().Y;
diff --git a/Engine/Source/Runtime/Engine/Public/MaterialCompiler.h b/Engine/Source/Runtime/Engine/Public/MaterialCompiler.h
--- a/Engine/Source/Runtime/Engine/Public/MaterialCompiler.h
+++ b/Engine/Source/Runtime/Engine/Public/MaterialCompiler.h
@@ -20,6 +20,7 @@ namespace UE
 namespace Shader
 {
 struct FValue;
+enum class EValueType : uint8;
 }
 }
 
@@ -239,8 +240,8 @@ public:
 
 	virtual int32 VirtualTexture(URuntimeVirtualTexture* InTexture, int32 TextureLayerIndex, int32 PageTableLayerIndex, int32& TextureReferenceIndex, EMaterialSamplerType SamplerType) = 0;
 	virtual int32 VirtualTextureParameter(FName ParameterName, URuntimeVirtualTexture* DefaultValue, int32 TextureLayerIndex, int32 PageTableLayerIndex, int32& TextureReferenceIndex, EMaterialSamplerType SamplerType) = 0;
-	virtual int32 VirtualTextureUniform(int32 TextureIndex, int32 VectorIndex) = 0;
-	virtual int32 VirtualTextureUniform(FName ParameterName, int32 TextureIndex, int32 VectorIndex) = 0;
+	virtual int32 VirtualTextureUniform(int32 TextureIndex, int32 VectorIndex, UE::Shader::EValueType Type) = 0;
+	virtual int32 VirtualTextureUniform(FName ParameterName, int32 TextureIndex, int32 VectorIndex, UE::Shader::EValueType Type) = 0;
 	virtual int32 VirtualTextureWorldToUV(int32 WorldPositionIndex, int32 P0, int32 P1, int32 P2) = 0;
 	virtual int32 VirtualTextureUnpack(int32 CodeIndex0, int32 CodeIndex1, int32 CodeIndex2, int32 P0, EVirtualTextureUnpackType UnpackType) = 0;
 
@@ -646,8 +647,8 @@ public:
 	{
 		return Compiler->VirtualTextureParameter(ParameterName, DefaultValue, TextureLayerIndex, PageTableLayerIndex, TextureReferenceIndex, SamplerType);
 	}
-	virtual int32 VirtualTextureUniform(int32 TextureIndex, int32 VectorIndex) override { return Compiler->VirtualTextureUniform(TextureIndex, VectorIndex); }
-	virtual int32 VirtualTextureUniform(FName ParameterName, int32 TextureIndex, int32 VectorIndex) override { return Compiler->VirtualTextureUniform(ParameterName, TextureIndex, VectorIndex); }
+	virtual int32 VirtualTextureUniform(int32 TextureIndex, int32 VectorIndex, UE::Shader::EValueType Type) override { return Compiler->VirtualTextureUniform(TextureIndex, VectorIndex, Type); }
+	virtual int32 VirtualTextureUniform(FName ParameterName, int32 TextureIndex, int32 VectorIndex, UE::Shader::EValueType Type) override { return Compiler->VirtualTextureUniform(ParameterName, TextureIndex, VectorIndex, Type); }
 	virtual int32 VirtualTextureWorldToUV(int32 WorldPositionIndex, int32 P0, int32 P1, int32 P2) override { return Compiler->VirtualTextureWorldToUV(WorldPositionIndex, P0, P1, P2); }
 	virtual int32 VirtualTextureUnpack(int32 CodeIndex0, int32 CodeIndex1, int32 CodeIndex2, int32 P0, EVirtualTextureUnpackType UnpackType) override { return Compiler->VirtualTextureUnpack(CodeIndex0, CodeIndex1, CodeIndex2, P0, UnpackType); }
 
