- Fix for Sprite distance culling not working

#rnx

diff --git a/Engine/Plugins/FX/Niagara/Source/Niagara/Private/NiagaraRendererSprites.cpp b/Engine/Plugins/FX/Niagara/Source/Niagara/Private/NiagaraRendererSprites.cpp
--- a/Engine/Plugins/FX/Niagara/Source/Niagara/Private/NiagaraRendererSprites.cpp
+++ b/Engine/Plugins/FX/Niagara/Source/Niagara/Private/NiagaraRendererSprites.cpp
@@ -261,7 +261,7 @@ void FNiagaraRendererSprites::PrepareParticleSpriteRenderData(FParticleSpriteRen
 		}
 
 		// Do we need culling?
-		ParticleSpriteRenderData.bNeedsCull = !bVisTagInParamStore && RendererVisTagOffset != INDEX_NONE;
+		ParticleSpriteRenderData.bNeedsCull = bEnableCulling;
 		ParticleSpriteRenderData.bSortCullOnGpu = (ParticleSpriteRenderData.bNeedsSort && FNiagaraUtilities::AllowGPUSorting(ShaderPlatform)) || (ParticleSpriteRenderData.bNeedsCull && FNiagaraUtilities::AllowGPUCulling(ShaderPlatform));
 
 		// Validate what we setup
@@ -382,6 +382,11 @@ void FNiagaraRendererSprites::InitializeSortInfo(FParticleSpriteRenderData& Part
 	OutSortInfo.DistanceCullRange = DistanceCullRange;
 	OutSortInfo.SystemLWCTile = UseLocalSpace(&SceneProxy) ? FVector3f::Zero() : SceneProxy.GetLWCRenderTile();
 
+	if ( bEnableDistanceCulling )
+	{
+		OutSortInfo.CullPositionAttributeOffset = ParticleSpriteRenderData.bSortCullOnGpu ? VFVariables[ENiagaraSpriteVFLayout::Position].GetGPUOffset() : VFVariables[ENiagaraSpriteVFLayout::Position].GetEncodedDatasetOffset();
+	}
+
 	auto GetViewMatrices =
 		[](const FSceneView& View) -> const FViewMatrices&
 		{
