- Add support for sampling mesh distance field inside static mesh data interface
- Project settiing is currently required to enable as all system using static mesh data interface are forced into PostRenderOpaque when enabled
- Partly constributed by [REDACTED_PERSON]

#rnx

diff --git a/Engine/Plugins/FX/Niagara/Shaders/Private/NiagaraDataInterfaceStaticMeshTemplate.ush b/Engine/Plugins/FX/Niagara/Shaders/Private/NiagaraDataInterfaceStaticMeshTemplate.ush
--- a/Engine/Plugins/FX/Niagara/Shaders/Private/NiagaraDataInterfaceStaticMeshTemplate.ush
+++ b/Engine/Plugins/FX/Niagara/Shaders/Private/NiagaraDataInterfaceStaticMeshTemplate.ush
@@ -557,3 +557,31 @@ void GetWorldVelocity_{ParameterName}(out float3 Velocity)
 {
 	Velocity = InstanceWorldVelocity_{ParameterName};
 }
+
+//////////////////////////////////////////////////////////////////////////
+// Experimental Distance Field Functions
+#if DISTATICMESH_ALLOWDISTANCEFIELD
+int		InstanceDistanceFieldIndex_{ParameterName};
+#endif
+
+void QueryDistanceField_{ParameterName}(bool bExecute, float3 WorldPosition, bool bUseMaxDistance, float MaxDistance, out bool OutIsValid, out float OutDistance)
+{
+	if ( !bUseMaxDistance )
+	{
+		MaxDistance = 1e38;
+	}
+
+#if DISTATICMESH_ALLOWDISTANCEFIELD
+	OutIsValid = bExecute && InstanceDistanceFieldIndex_{ParameterName} != int(-1);
+	OutDistance = MaxDistance;
+	if ( OutIsValid )
+	{
+		// Needs to be a smidge over 0.0 for DistanceToNearestSurfaceForObject to query the SDF
+		MaxDistance = max(0.001, MaxDistance);
+		OutDistance = DistanceToNearestSurfaceForObject(InstanceDistanceFieldIndex_{ParameterName}, WorldPosition, MaxDistance);
+	}
+#else
+	OutIsValid = false;
+	OutDistance = MaxDistance;
+#endif
+}
diff --git a/Engine/Plugins/FX/Niagara/Source/Niagara/Private/DataInterface/NiagaraDataInterfaceStaticMesh.cpp b/Engine/Plugins/FX/Niagara/Source/Niagara/Private/DataInterface/NiagaraDataInterfaceStaticMesh.cpp
--- a/Engine/Plugins/FX/Niagara/Source/Niagara/Private/DataInterface/NiagaraDataInterfaceStaticMesh.cpp
+++ b/Engine/Plugins/FX/Niagara/Source/Niagara/Private/DataInterface/NiagaraDataInterfaceStaticMesh.cpp
@@ -3,9 +3,12 @@
 #include "NiagaraDataInterfaceStaticMesh.h"
 #include "NiagaraEmitterInstance.h"
 #include "NiagaraComponent.h"
+#include "NiagaraStats.h"
 #include "NiagaraSystemInstance.h"
 #include "NiagaraRenderer.h"
+#include "NiagaraSettings.h"
 #include "NiagaraScript.h"
+#include "NiagaraDataInterfaceUtilities.h"
 #include "ShaderParameterUtils.h"
 #include "NiagaraStats.h"
 
@@ -20,6 +23,13 @@
 #include "GPUSkinCache.h"
 #include "ShaderCompilerCore.h"
 
+#include "Renderer/Private/ScenePrivate.h"
+#include "DistanceFieldAtlas.h"
+#include "Renderer/Private/DistanceFieldLightingShared.h"
+#include "NiagaraGpuComputeDispatchInterface.h"
+#include "NiagaraGpuComputeDispatch.h"
+#include "NiagaraDistanceFieldParameters.h"
+
 #if WITH_EDITOR
 #include "Subsystems/ImportSubsystem.h"
 #include "Editor.h"
@@ -136,6 +146,10 @@ namespace NDIStaticMeshLocal
 	static const FName	GetLocalToWorldInverseTransposedName("GetLocalToWorldInverseTransposed");
 	static const FName	GetWorldVelocityName("GetWorldVelocity");
 
+	//////////////////////////////////////////////////////////////////////////
+	// Distance Field Functions
+	static const FName	QueryDistanceFieldName("QueryDistanceField");
+
 	//////////////////////////////////////////////////////////////////////////
 	// Deprecated Functions
 	static const FName Deprecated_RandomSectionName("DeprecatedRandomSection");
@@ -187,6 +201,8 @@ namespace NDIStaticMeshLocal
 	static FString	InstancePreviousRotation_String("InstancePreviousRotation_");
 	static FString	InstanceWorldVelocity_String("InstanceWorldVelocity_");
 
+	static FString	InstanceDistanceFieldIndex_String("InstanceDistanceFieldIndex_");
+
 	//////////////////////////////////////////////////////////////////////////
 	struct FNDISectionInfo
 	{
@@ -292,6 +308,8 @@ namespace NDIStaticMeshLocal
 		FQuat4f		PrevRotation = FQuat4f::Identity;
 		float		DeltaSeconds = 0.0f;
 
+		FPrimitiveComponentId	DistanceFieldPrimitiveId;
+
 		FIntVector NumTriangles = FIntVector::ZeroValue;
 		int32 NumVertices = 0;
 		int32 NumUVs = 0;
@@ -412,11 +430,12 @@ namespace NDIStaticMeshLocal
 
 	struct FInstanceData_FromGameThread
 	{
-		FMatrix44f	Transform = FMatrix44f::Identity;
-		FMatrix44f	PrevTransform = FMatrix44f::Identity;
-		FQuat4f		Rotation = FQuat4f::Identity;
-		FQuat4f		PrevRotation = FQuat4f::Identity;
-		float		DeltaSeconds = 0.0f;
+		FMatrix44f				Transform = FMatrix44f::Identity;
+		FMatrix44f				PrevTransform = FMatrix44f::Identity;
+		FQuat4f					Rotation = FQuat4f::Identity;
+		FQuat4f					PrevRotation = FQuat4f::Identity;
+		float					DeltaSeconds = 0.0f;
+		FPrimitiveComponentId	DistanceFieldPrimitiveId;
 	};
 
 	struct FInstanceData_GameThread
@@ -860,11 +879,12 @@ namespace NDIStaticMeshLocal
 			FInstanceData_FromGameThread* FromGameThread = reinterpret_cast<FInstanceData_FromGameThread*>(FromGameThreadData);
 			FInstanceData_RenderThread* InstanceData = &PerInstanceData_RT.FindChecked(InstanceID);
 
-			InstanceData->Transform		= FromGameThread->Transform;
-			InstanceData->PrevTransform	= FromGameThread->PrevTransform;
-			InstanceData->Rotation		= FromGameThread->Rotation;
-			InstanceData->PrevRotation	= FromGameThread->PrevRotation;
-			InstanceData->DeltaSeconds	= FromGameThread->DeltaSeconds;
+			InstanceData->Transform					= FromGameThread->Transform;
+			InstanceData->PrevTransform				= FromGameThread->PrevTransform;
+			InstanceData->Rotation					= FromGameThread->Rotation;
+			InstanceData->PrevRotation				= FromGameThread->PrevRotation;
+			InstanceData->DeltaSeconds				= FromGameThread->DeltaSeconds;
+			InstanceData->DistanceFieldPrimitiveId	= FromGameThread->DistanceFieldPrimitiveId;
 
 			FromGameThread->~FInstanceData_FromGameThread();
 		}
@@ -1217,6 +1237,9 @@ public:
 		InstancePreviousTransformInverseTransposedParam.Bind(ParameterMap, *(InstancePreviousTransformInverseTransposed_String + ParameterInfo.DataInterfaceHLSLSymbol));
 		InstancePreviousRotationParam.Bind(ParameterMap, *(InstancePreviousRotation_String + ParameterInfo.DataInterfaceHLSLSymbol));
 		InstanceWorldVelocityParam.Bind(ParameterMap, *(InstanceWorldVelocity_String + ParameterInfo.DataInterfaceHLSLSymbol));
+
+		InstanceDistanceFieldIndexParam.Bind(ParameterMap, *(InstanceDistanceFieldIndex_String + ParameterInfo.DataInterfaceHLSLSymbol));
+		DistanceFieldParameters.Bind(ParameterMap);
 	}
 
 	void Set(FRHICommandList& RHICmdList, const FNiagaraDataInterfaceSetArgs& Context) const
@@ -1283,6 +1306,42 @@ public:
 		SetShaderValue(RHICmdList, ComputeShaderRHI, InstancePreviousTransformInverseTransposedParam, InstanceData.PrevTransform.Inverse().GetTransposed());
 		SetShaderValue(RHICmdList, ComputeShaderRHI, InstancePreviousRotationParam, InstanceData.PrevRotation);
 		SetShaderValue(RHICmdList, ComputeShaderRHI, InstanceWorldVelocityParam, DeltaPosition);
+
+		const FDistanceFieldSceneData* DistanceFieldSceneData = static_cast<const FNiagaraGpuComputeDispatch*>(Context.ComputeDispatchInterface)->GetMeshDistanceFieldParameters();	//-BATCHERTODO:
+
+		if ( InstanceDistanceFieldIndexParam.IsBound() )
+		{
+			int32 DistanceFieldIndex = -1;
+			if ( DistanceFieldSceneData != nullptr && InstanceData.DistanceFieldPrimitiveId.IsValid() )
+			{
+				if ( FScene* Scene = Context.ComputeDispatchInterface->GetScene() )
+				{
+					// Kind of gross, but currently no way to reference other primitive scene infos
+					const int32 PrimitiveSceneIndex = Scene->PrimitiveComponentIds.Find(InstanceData.DistanceFieldPrimitiveId);
+					if (PrimitiveSceneIndex != INDEX_NONE)
+					{
+						const TArray<int32, TInlineAllocator<1>>& DFIndices = Scene->Primitives[PrimitiveSceneIndex]->DistanceFieldInstanceIndices;
+						DistanceFieldIndex = DFIndices.Num() > 0 ? DFIndices[0] : -1;
+					}
+				}
+			}
+			SetShaderValue(RHICmdList, ComputeShaderRHI, InstanceDistanceFieldIndexParam, DistanceFieldIndex);
+		}
+
+		if (DistanceFieldParameters.IsBound())
+		{
+			if (ensure(DistanceFieldSceneData != nullptr))
+			{
+				DistanceFieldParameters.Set(RHICmdList, ComputeShaderRHI, DistanceFieldSceneData);
+			}
+			else 
+			{
+				//-TODO: We can't create dummy buffers here due to dll boundaries, we will skip the distance field calcs since DistanceFieldIndex will be set to -1
+				//       However some platforms may still complain due to lack of dummy data bound
+				//FDistanceFieldSceneData DummyDistanceFieldSceneData(Context.ComputeDispatchInterface->GetShaderPlatform());
+				//DistanceFieldParameters.Set(RHICmdList, ComputeShaderRHI, &DummyDistanceFieldSceneData);
+			}
+		}
 	}
 
 private:
@@ -1315,6 +1374,9 @@ private:
 	LAYOUT_FIELD(FShaderParameter, InstancePreviousTransformInverseTransposedParam);
 	LAYOUT_FIELD(FShaderParameter, InstancePreviousRotationParam);
 	LAYOUT_FIELD(FShaderParameter, InstanceWorldVelocityParam);
+
+	LAYOUT_FIELD(FShaderParameter, InstanceDistanceFieldIndexParam);
+	LAYOUT_FIELD(FDistanceFieldParameters, DistanceFieldParameters);	
 };
 
 IMPLEMENT_TYPE_LAYOUT(FNiagaraDataInterfaceParametersCS_StaticMesh);
@@ -1530,11 +1592,16 @@ void UNiagaraDataInterfaceStaticMesh::ProvidePerInstanceDataForRenderThread(void
 	NDIStaticMeshLocal::FInstanceData_GameThread* InstanceData = static_cast<NDIStaticMeshLocal::FInstanceData_GameThread*>(InInstanceData);
 	NDIStaticMeshLocal::FInstanceData_FromGameThread* DataFromGT = static_cast<NDIStaticMeshLocal::FInstanceData_FromGameThread*>(InDataFromGT);
 
-	DataFromGT->Transform		= FMatrix44f(InstanceData->Transform);			// LWC_TODO: Precision loss
-	DataFromGT->PrevTransform	= FMatrix44f(InstanceData->PrevTransform);
-	DataFromGT->Rotation		= InstanceData->Rotation;
-	DataFromGT->PrevRotation	= InstanceData->PrevRotation;
-	DataFromGT->DeltaSeconds	= InstanceData->DeltaSeconds;
+	DataFromGT->Transform					= FMatrix44f(InstanceData->Transform);			// LWC_TODO: Precision loss
+	DataFromGT->PrevTransform				= FMatrix44f(InstanceData->PrevTransform);
+	DataFromGT->Rotation					= InstanceData->Rotation;
+	DataFromGT->PrevRotation				= InstanceData->PrevRotation;
+	DataFromGT->DeltaSeconds				= InstanceData->DeltaSeconds;
+	DataFromGT->DistanceFieldPrimitiveId	= FPrimitiveComponentId();
+	if ( UPrimitiveComponent* PrimitiveComponent = Cast<UPrimitiveComponent>(InstanceData->SceneComponentWeakPtr) )
+	{
+		DataFromGT->DistanceFieldPrimitiveId = PrimitiveComponent->ComponentId;
+	}
 }
 
 void UNiagaraDataInterfaceStaticMesh::GetFunctions(TArray<FNiagaraFunctionSignature>& OutFunctions)
@@ -1840,6 +1907,38 @@ void UNiagaraDataInterfaceStaticMesh::GetFunctions(TArray<FNiagaraFunctionSignat
 		Sig.Outputs.Emplace(FNiagaraTypeDefinition::GetVec3Def(), TEXT("Velocity"));
 	}
 
+	//////////////////////////////////////////////////////////////////////////
+	// Distance Field Functions
+	{
+		FNiagaraFunctionSignature& Sig = OutFunctions.Add_GetRef(BaseSignature);
+		Sig.Name = QueryDistanceFieldName;
+		Sig.bSupportsCPU = false;
+		Sig.bExperimental = true;
+		Sig.Inputs.Emplace_GetRef(FNiagaraTypeDefinition::GetBoolDef(), TEXT("Execute")).SetValue(true);
+		Sig.Inputs.Emplace(FNiagaraTypeDefinition::GetPositionDef(), TEXT("World Position"));
+		const FNiagaraVariable& UseMaxDistanceVar = Sig.Inputs.Emplace_GetRef(FNiagaraTypeDefinition::GetBoolDef(), TEXT("Use Max Distance"));
+		const FNiagaraVariable& MaxDistanceVar = Sig.Inputs.Emplace_GetRef(FNiagaraTypeDefinition::GetFloatDef(), TEXT("Max Distance"));
+		Sig.Outputs.Emplace(FNiagaraTypeDefinition::GetBoolDef(), TEXT("IsValid"));
+		Sig.Outputs.Emplace(FNiagaraTypeDefinition::GetFloatDef(), TEXT("Distance"));
+		
+#if WITH_EDITORONLY_DATA
+		Sig.Description = LOCTEXT("QueryDistanceFieldDescription", "Given a world position, this returns the value of the parented static mesh`s signed distance field");
+		Sig.InputDescriptions.Add(
+			UseMaxDistanceVar,
+			LOCTEXT("UseMaxDistanceOptimizationDescription", 
+			"This enables an optimization that will skip reading the SDF texture if the world position exceeds a provided max distance value.\n"
+			"This is a useful optimization if you only need accurate distance information within a certain distance threshold. \n"
+			"One example of a good usecase is if you only want to spawn particles inside the mesh (i.e. distance < 0.0), you should enable this with a MaxDistance of 0.0.")
+		);
+		Sig.InputDescriptions.Add(
+			MaxDistanceVar,
+			LOCTEXT("MaxDistanceDescription",
+			"Only used if UseMaxDistanceOptimization is enabled. This is the max distance from the static meshes` BOUNDING BOX at which we should query it's SDF.\n"
+			"This defaults to 0.0, which means it will only read the SDF if the world position is inside the meshes' bounding box")
+		);
+#endif
+	}
+
 	//////////////////////////////////////////////////////////////////////////
 	// Deprecated Functions
 	{
@@ -1990,6 +2089,15 @@ bool UNiagaraDataInterfaceStaticMesh::UpgradeFunctionCall(FNiagaraFunctionSignat
 
 	return true;
 }
+
+void UNiagaraDataInterfaceStaticMesh::GetCommonHLSL(FString& OutHLSL)
+{
+	if (GetDefault<UNiagaraSettings>()->NDIStaticMesh_AllowDistanceFields)
+	{
+		OutHLSL += TEXT("#include \"/Engine/Private/DistanceFieldLightingShared.ush\"\n");
+		OutHLSL += TEXT("#include \"/Engine/Private/MeshDistanceFieldCommon.ush\"\n");
+	}
+}
 #endif //WITH_EDITORONLY_DATA
 
 void UNiagaraDataInterfaceStaticMesh::GetVMExternalFunction(const FVMExternalFunctionBindingInfo& BindingInfo, void* InstanceData, FVMExternalFunction& OutFunc)
@@ -2280,15 +2388,28 @@ void UNiagaraDataInterfaceStaticMesh::GetVMExternalFunction(const FVMExternalFun
 	}
 }
 
+bool UNiagaraDataInterfaceStaticMesh::RequiresDistanceFieldData() const
+{
+	return GetDefault<UNiagaraSettings>()->NDIStaticMesh_AllowDistanceFields;
+}
+
 #if WITH_EDITORONLY_DATA
 bool UNiagaraDataInterfaceStaticMesh::AppendCompileHash(FNiagaraCompileHashVisitor* InVisitor) const
 {
 	bool bSuccess = Super::AppendCompileHash(InVisitor);
 	FSHAHash Hash = GetShaderFileHash(NDIStaticMeshLocal::TemplateShaderFile, EShaderPlatform::SP_PCD3D_SM5);
 	InVisitor->UpdateString(TEXT("NiagaraDataInterfaceStaticMeshTemplateHLSLSource"), Hash.ToString());
+	InVisitor->UpdatePOD(TEXT("NDIStaticMesh_AllowDistanceField"), GetDefault<UNiagaraSettings>()->NDIStaticMesh_AllowDistanceFields ? 1 : 0);
 	return bSuccess;
 }
 
+void UNiagaraDataInterfaceStaticMesh::ModifyCompilationEnvironment(EShaderPlatform ShaderPlatform, struct FShaderCompilerEnvironment& OutEnvironment) const
+{
+	Super::ModifyCompilationEnvironment(ShaderPlatform, OutEnvironment);
+
+	OutEnvironment.SetDefine(TEXT("DISTATICMESH_ALLOWDISTANCEFIELD"), GetDefault<UNiagaraSettings>()->NDIStaticMesh_AllowDistanceFields ? 1 : 0);
+}
+
 void UNiagaraDataInterfaceStaticMesh::GetParameterDefinitionHLSL(const FNiagaraDataInterfaceGPUParamInfo& ParamInfo, FString& OutHLSL)
 {
 	TMap<FString, FStringFormatArg> TemplateArgs =
@@ -2482,6 +2603,27 @@ void UNiagaraDataInterfaceStaticMesh::GetFeedback(UNiagaraSystem* Asset, UNiagar
 
 		OutWarnings.Add(NoMeshAssignedError);
 	}
+
+	if (GetDefault<UNiagaraSettings>()->NDIStaticMesh_AllowDistanceFields == false)
+	{
+		FNiagaraDataInterfaceUtilities::ForEachGpuFunctionEquals(
+			this, Asset, Component,
+			[&](const FNiagaraDataInterfaceGeneratedFunction& FunctionBinding)
+			{
+				if (FunctionBinding.DefinitionName == NDIStaticMeshLocal::QueryDistanceFieldName )
+				{
+					OutWarnings.Emplace(
+						LOCTEXT("DistanceFieldsNotEnabled", "Distance fields functionality is disabled, this can be enabled in Niagara project settings."),
+						LOCTEXT("DistanceFieldsNotEnabledSummary", "Distance fields functionality is disabled."),
+						FNiagaraDataInterfaceFix()
+					);
+
+					return false;
+				}
+				return true;
+			}
+		);
+	}
 }
 #endif //WITH_EDITOR
 
diff --git a/Engine/Plugins/FX/Niagara/Source/Niagara/Private/DataInterface/NiagaraDataInterfaceStaticMesh.h b/Engine/Plugins/FX/Niagara/Source/Niagara/Private/DataInterface/NiagaraDataInterfaceStaticMesh.h
--- a/Engine/Plugins/FX/Niagara/Source/Niagara/Private/DataInterface/NiagaraDataInterfaceStaticMesh.h
+++ b/Engine/Plugins/FX/Niagara/Source/Niagara/Private/DataInterface/NiagaraDataInterfaceStaticMesh.h
@@ -113,14 +113,17 @@ public:
 
 	virtual void GetFunctions(TArray<FNiagaraFunctionSignature>& OutFunctions) override;
 #if WITH_EDITORONLY_DATA
+	virtual void GetCommonHLSL(FString& OutHLSL) override;
 	virtual bool UpgradeFunctionCall(FNiagaraFunctionSignature& FunctionSignature) override;
 #endif
 	virtual void GetVMExternalFunction(const FVMExternalFunctionBindingInfo& BindingInfo, void* InstanceData, FVMExternalFunction &OutFunc) override;
 	virtual bool CanExecuteOnTarget(ENiagaraSimTarget Target) const override { return true; }
+	virtual bool RequiresDistanceFieldData() const override;
 	virtual bool HasPreSimulateTick() const override { return true; }
 
 #if WITH_EDITORONLY_DATA
 	virtual bool AppendCompileHash(FNiagaraCompileHashVisitor* InVisitor) const override;
+	virtual void ModifyCompilationEnvironment(EShaderPlatform ShaderPlatform, struct FShaderCompilerEnvironment& OutEnvironment) const override;
 	virtual void GetParameterDefinitionHLSL(const FNiagaraDataInterfaceGPUParamInfo& ParamInfo, FString& OutHLSL) override;
 	virtual bool GetFunctionHLSL(const FNiagaraDataInterfaceGPUParamInfo& ParamInfo, const FNiagaraDataInterfaceGeneratedFunction& FunctionInfo, int FunctionInstanceIndex, FString& OutHLSL) override;
 #endif
diff --git a/Engine/Plugins/FX/Niagara/Source/Niagara/Private/DataInterface/NiagaraDataInterfaceUtilities.cpp b/Engine/Plugins/FX/Niagara/Source/Niagara/Private/DataInterface/NiagaraDataInterfaceUtilities.cpp
new file mode 100644
--- /dev/null
+++ b/Engine/Plugins/FX/Niagara/Source/Niagara/Private/DataInterface/NiagaraDataInterfaceUtilities.cpp
@@ -0,0 +1,296 @@
+// Copyright Epic Games, Inc. All Rights Reserved.
+
+#include "NiagaraDataInterfaceUtilities.h"
+#include "NiagaraComponent.h"
+#include "NiagaraComputeExecutionContext.h"
+#include "NiagaraDataInterface.h"
+#include "NiagaraSystem.h"
+
+void FNiagaraDataInterfaceUtilities::ForEachVMFunctionEquals(class UNiagaraDataInterface* DataInterface, class UNiagaraSystem* NiagaraSystem, class UNiagaraComponent* Component, TFunction<bool(const FVMExternalFunctionBindingInfo&)> Action)
+{
+	if ( DataInterface == nullptr || NiagaraSystem == nullptr )
+	{
+		return;
+	}
+
+	// Find all override parameter names
+	TArray<FName, TInlineAllocator<8>> OverrideParameterNames;
+	if ( Component )
+	{
+		for (const UNiagaraDataInterface* OverrideDI : Component->GetOverrideParameters().GetDataInterfaces())
+		{
+			if ( (OverrideDI != nullptr) && ((OverrideDI == DataInterface) || OverrideDI->Equals(DataInterface)) )
+			{
+				if ( const FNiagaraVariableBase* Variable = Component->GetOverrideParameters().FindVariable(OverrideDI) )
+				{
+					OverrideParameterNames.AddUnique(Variable->GetName());
+				}
+			}
+		}
+	}
+	else
+	{
+		for (const UNiagaraDataInterface* OverrideDI : NiagaraSystem->GetExposedParameters().GetDataInterfaces())
+		{
+			if ( (OverrideDI != nullptr) && ((OverrideDI == DataInterface) || OverrideDI->Equals(DataInterface)) )
+			{
+				if ( const FNiagaraVariableBase* Variable = NiagaraSystem->GetExposedParameters().FindVariable(OverrideDI) )
+				{
+					OverrideParameterNames.AddUnique(Variable->GetName());
+				}
+			}
+		}
+	}
+
+	// Loop through all scripts
+	bool bContinueSearching = true;
+	NiagaraSystem->ForEachScript(
+		[&](UNiagaraScript* Script)
+		{
+			if (bContinueSearching == false)
+			{
+				return;
+			}
+
+			const FNiagaraVMExecutableData& ScriptExecutableData = Script->GetVMExecutableData();
+			if (!ScriptExecutableData.IsValid())
+			{
+				return;
+			}
+
+			const TArray<FNiagaraScriptDataInterfaceInfo>& CachedDefaultDIs = Script->GetCachedDefaultDataInterfaces();
+			const TArray<FNiagaraScriptDataInterfaceCompileInfo>& DataInterfaceInfos = ScriptExecutableData.DataInterfaceInfo;
+			for ( const FVMExternalFunctionBindingInfo& FunctionBinding : ScriptExecutableData.CalledVMExternalFunctions )
+			{
+				const int NumDataInterface = FMath::Min(DataInterfaceInfos.Num(), CachedDefaultDIs.Num());	// Note: Should always be equal but lets be safe
+				for ( int iDataInterface=0; iDataInterface < NumDataInterface; ++iDataInterface )
+				{
+					const FNiagaraScriptDataInterfaceCompileInfo& DataInterfaceInfo = DataInterfaceInfos[iDataInterface];
+					if ( DataInterfaceInfo.Name != FunctionBinding.OwnerName )
+					{
+						continue;
+					}
+
+					const FNiagaraScriptDataInterfaceInfo& CachedDefaultDI = CachedDefaultDIs[iDataInterface];
+					if (CachedDefaultDI.DataInterface == nullptr || !DataInterfaceInfo.MatchesClass(DataInterface->GetClass()))
+					{
+						// Would be odd not to match here, but we are being safe
+						break;
+					}
+
+					// Do we have a match?
+					if (CachedDefaultDI.DataInterface->Equals(DataInterface) || OverrideParameterNames.Contains(CachedDefaultDI.Name))
+					{
+						if ( Action(FunctionBinding) == false )
+						{
+							bContinueSearching = false;
+							return;
+						}
+					}
+					break;
+				}
+			}
+		}
+	);
+}
+
+void FNiagaraDataInterfaceUtilities::ForEachGpuFunctionEquals(class UNiagaraDataInterface* DataInterface, class UNiagaraSystem* NiagaraSystem, class UNiagaraComponent* Component, TFunction<bool(const FNiagaraDataInterfaceGeneratedFunction&)> Action)
+{
+	if ( DataInterface == nullptr || NiagaraSystem == nullptr )
+	{
+		return;
+	}
+
+	// Find all override parameter names
+	TArray<FName, TInlineAllocator<8>> OverrideParameterNames;
+	if ( Component )
+	{
+		for (const UNiagaraDataInterface* OverrideDI : Component->GetOverrideParameters().GetDataInterfaces())
+		{
+			if ( (OverrideDI != nullptr) && ((OverrideDI == DataInterface) || OverrideDI->Equals(DataInterface)) )
+			{
+				if ( const FNiagaraVariableBase* Variable = Component->GetOverrideParameters().FindVariable(OverrideDI) )
+				{
+					OverrideParameterNames.AddUnique(Variable->GetName());
+				}
+			}
+		}
+	}
+	else
+	{
+		for (const UNiagaraDataInterface* OverrideDI : NiagaraSystem->GetExposedParameters().GetDataInterfaces())
+		{
+			if ( (OverrideDI != nullptr) && ((OverrideDI == DataInterface) || OverrideDI->Equals(DataInterface)) )
+			{
+				if ( const FNiagaraVariableBase* Variable = NiagaraSystem->GetExposedParameters().FindVariable(OverrideDI) )
+				{
+					OverrideParameterNames.AddUnique(Variable->GetName());
+				}
+			}
+		}
+	}
+
+	// Loop through all scripts
+	bool bContinueSearching = true;
+	NiagaraSystem->ForEachScript(
+		[&](UNiagaraScript* Script)
+		{
+			if (bContinueSearching == false)
+			{
+				return;
+			}
+
+			if (FNiagaraShaderScript* ShaderScript = Script->GetRenderThreadScript())
+			{
+				const TArray<FNiagaraScriptDataInterfaceInfo>& CachedDefaultDIs = Script->GetCachedDefaultDataInterfaces();
+				const TArray<FNiagaraDataInterfaceGPUParamInfo>& DataInterfaceParamInfos = ShaderScript->GetDataInterfaceParamInfo();
+
+				const int NumDataInterfaces = FMath::Min(CachedDefaultDIs.Num(), DataInterfaceParamInfos.Num());	// Note: Should always be equal but lets be safe
+				for (int iDataInterface = 0; iDataInterface < NumDataInterfaces; ++iDataInterface)
+				{
+					const FNiagaraScriptDataInterfaceInfo& CachedDefaultDI = CachedDefaultDIs[iDataInterface];
+					if (CachedDefaultDI.DataInterface == nullptr)
+					{
+						continue;
+					}
+					if (CachedDefaultDI.DataInterface->Equals(DataInterface) || OverrideParameterNames.Contains(CachedDefaultDI.Name))
+					{
+						for (const FNiagaraDataInterfaceGeneratedFunction& GeneratedFunction : DataInterfaceParamInfos[iDataInterface].GeneratedFunctions)
+						{
+							if (Action(GeneratedFunction) == false)
+							{
+								bContinueSearching = false;
+								return;
+							}
+						}
+					}
+				}
+			}
+		}
+	);
+}
+
+void FNiagaraDataInterfaceUtilities::ForEachVMFunction(class UNiagaraDataInterface* DataInterface, class FNiagaraSystemInstance* SystemInstance, TFunction<bool(const FVMExternalFunctionBindingInfo&)> Action)
+{
+	if ( !DataInterface || !SystemInstance )
+	{
+		return;
+	}
+
+	auto SearchVMContext =
+		[](const FNiagaraScriptExecutionContextBase* ExecContext, UNiagaraDataInterface* DataInterface, const TFunction<bool(const FVMExternalFunctionBindingInfo&)>& Action) -> bool
+		{
+			if (ExecContext == nullptr || ExecContext->Script == nullptr)
+			{
+				return true;
+			}
+
+			const FNiagaraVMExecutableData& ScriptExecutableData = ExecContext->Script->GetVMExecutableData();
+			if (!ScriptExecutableData.IsValid())
+			{
+				return true;
+			}
+
+			const TArray<UNiagaraDataInterface*>& DataInterfaces = ExecContext->GetDataInterfaces();
+			const int32 NumDataInterfaces = FMath::Min(ScriptExecutableData.DataInterfaceInfo.Num(), DataInterfaces.Num());		// Should be equal, but be safe
+			for (const FVMExternalFunctionBindingInfo& FunctionBinding : ScriptExecutableData.CalledVMExternalFunctions)
+			{
+				for (int32 iDataInterface = 0; iDataInterface < NumDataInterfaces; ++iDataInterface)
+				{
+					const FNiagaraScriptDataInterfaceCompileInfo& DataInterfaceInfo = ScriptExecutableData.DataInterfaceInfo[iDataInterface];
+					if (FunctionBinding.OwnerName != DataInterfaceInfo.Name)
+					{
+						continue;
+					}
+					if (DataInterfaces[iDataInterface] != DataInterface)
+					{
+						continue;
+					}
+
+					if (Action(FunctionBinding) == false)
+					{
+						return false;
+					}
+				}
+			}
+			return true;
+		};
+
+
+	// Search system scripts (always VM)
+	TSharedPtr<FNiagaraSystemSimulation, ESPMode::ThreadSafe> SystemSimulation = SystemInstance->GetSystemSimulation();
+	if ( SystemSimulation.IsValid() )
+	{
+		if (SearchVMContext(SystemSimulation->GetSpawnExecutionContext(), DataInterface, Action) == false)
+		{
+			return;
+		}
+		if (SearchVMContext(SystemSimulation->GetUpdateExecutionContext(), DataInterface, Action) == false)
+		{
+			return;
+		}
+	}
+
+	// Search emitter scripts
+	for (const TSharedRef<FNiagaraEmitterInstance, ESPMode::ThreadSafe>& EmitterInstance : SystemInstance->GetEmitters())
+	{
+		if (EmitterInstance->IsDisabled() || EmitterInstance->GetCachedEmitter() == nullptr || EmitterInstance->GetGPUContext() != nullptr)
+		{
+			continue;
+		}
+
+		if ( SearchVMContext(&EmitterInstance->GetSpawnExecutionContext(), DataInterface, Action) == false )
+		{
+			return;
+		}
+
+		if (SearchVMContext(&EmitterInstance->GetUpdateExecutionContext(), DataInterface, Action) == false)
+		{
+			return;
+		}
+
+		for ( const FNiagaraScriptExecutionContext& EventExecContext : EmitterInstance->GetEventExecutionContexts() )
+		{
+			if (SearchVMContext(&EventExecContext, DataInterface, Action) == false)
+			{
+				return;
+			}
+		}
+	}
+}
+
+void FNiagaraDataInterfaceUtilities::ForEachGpuFunction(class UNiagaraDataInterface* DataInterface, class FNiagaraSystemInstance* SystemInstance, TFunction<bool(const FNiagaraDataInterfaceGeneratedFunction&)> Action)
+{
+	if ( !DataInterface || !SystemInstance )
+	{
+		return;
+	}
+
+	for (const TSharedRef<FNiagaraEmitterInstance, ESPMode::ThreadSafe>& EmitterInstance : SystemInstance->GetEmitters())
+	{
+		if (EmitterInstance->IsDisabled() || EmitterInstance->GetCachedEmitter() == nullptr || EmitterInstance->GetGPUContext() == nullptr)
+		{
+			continue;
+		}
+
+		const FNiagaraScriptInstanceParameterStore& ParameterStore = EmitterInstance->GetGPUContext()->CombinedParamStore;
+		const TArray<UNiagaraDataInterface*>& DataInterfaces = ParameterStore.GetDataInterfaces();
+		const TArray<FNiagaraDataInterfaceGPUParamInfo>& DataInterfaceParamInfo = EmitterInstance->GetGPUContext()->GPUScript_RT->GetDataInterfaceParamInfo();
+		const int32 NumDataInterface = FMath::Min(DataInterfaces.Num(), DataInterfaceParamInfo.Num());		// Should be equal, but be safe
+		for ( int32 iDataInterface=0; iDataInterface < NumDataInterface; ++iDataInterface)
+		{
+			if ( DataInterfaces[iDataInterface] != DataInterface )
+			{
+				continue;
+			}
+
+			for (const FNiagaraDataInterfaceGeneratedFunction& GeneratedFunction : DataInterfaceParamInfo[iDataInterface].GeneratedFunctions)
+			{
+				if (Action(GeneratedFunction) == false)
+				{
+					return;
+				}
+			}
+		}
+	}
+}
diff --git a/Engine/Plugins/FX/Niagara/Source/Niagara/Private/DataInterface/NiagaraDistanceFieldParameters.h b/Engine/Plugins/FX/Niagara/Source/Niagara/Private/DataInterface/NiagaraDistanceFieldParameters.h
new file mode 100644
--- /dev/null
+++ b/Engine/Plugins/FX/Niagara/Source/Niagara/Private/DataInterface/NiagaraDistanceFieldParameters.h
@@ -0,0 +1,100 @@
+// Copyright Epic Games, Inc. All Rights Reserved.
+
+class FDistanceFieldParameters
+{
+	DECLARE_INLINE_TYPE_LAYOUT(FDistanceFieldParameters, NonVirtual);
+public:
+	void Bind(const FShaderParameterMap& ParameterMap)
+	{
+		SceneObjectBounds.Bind(ParameterMap, TEXT("SceneObjectBounds"));
+		SceneObjectData.Bind(ParameterMap, TEXT("SceneObjectData"));
+		NumSceneObjects.Bind(ParameterMap, TEXT("NumSceneObjects"));
+		SceneDistanceFieldAssetData.Bind(ParameterMap, TEXT("SceneDistanceFieldAssetData"));
+		DistanceFieldIndirectionTable.Bind(ParameterMap, TEXT("DistanceFieldIndirectionTable"));
+		DistanceFieldBrickTexture.Bind(ParameterMap, TEXT("DistanceFieldBrickTexture"));
+		DistanceFieldSampler.Bind(ParameterMap, TEXT("DistanceFieldSampler"));
+		DistanceFieldBrickSize.Bind(ParameterMap, TEXT("DistanceFieldBrickSize"));
+		DistanceFieldUniqueDataBrickSize.Bind(ParameterMap, TEXT("DistanceFieldUniqueDataBrickSize"));
+		DistanceFieldBrickAtlasSizeInBricks.Bind(ParameterMap, TEXT("DistanceFieldBrickAtlasSizeInBricks"));
+		DistanceFieldBrickAtlasMask.Bind(ParameterMap, TEXT("DistanceFieldBrickAtlasMask"));
+		DistanceFieldBrickAtlasSizeLog2.Bind(ParameterMap, TEXT("DistanceFieldBrickAtlasSizeLog2"));
+		DistanceFieldBrickAtlasTexelSize.Bind(ParameterMap, TEXT("DistanceFieldBrickAtlasTexelSize"));
+		DistanceFieldBrickAtlasHalfTexelSize.Bind(ParameterMap, TEXT("DistanceFieldBrickAtlasHalfTexelSize"));
+		DistanceFieldBrickOffsetToAtlasUVScale.Bind(ParameterMap, TEXT("DistanceFieldBrickOffsetToAtlasUVScale"));
+		DistanceFieldUniqueDataBrickSizeInAtlasTexels.Bind(ParameterMap, TEXT("DistanceFieldUniqueDataBrickSizeInAtlasTexels"));
+	}
+
+	bool IsBound() const
+	{
+		return SceneDistanceFieldAssetData.IsBound();
+	}
+
+	friend FArchive& operator<<(FArchive& Ar, FDistanceFieldParameters& Parameters)
+	{
+		Ar << Parameters.SceneObjectBounds;
+		Ar << Parameters.SceneObjectData;
+		Ar << Parameters.NumSceneObjects;
+		Ar << Parameters.SceneDistanceFieldAssetData;
+		Ar << Parameters.DistanceFieldIndirectionTable;
+		Ar << Parameters.DistanceFieldBrickTexture;
+		Ar << Parameters.DistanceFieldSampler;
+		Ar << Parameters.DistanceFieldBrickSize;
+		Ar << Parameters.DistanceFieldUniqueDataBrickSize;
+		Ar << Parameters.DistanceFieldBrickAtlasSizeInBricks;
+		Ar << Parameters.DistanceFieldBrickAtlasMask;
+		Ar << Parameters.DistanceFieldBrickAtlasSizeLog2;
+		Ar << Parameters.DistanceFieldBrickAtlasTexelSize;
+		Ar << Parameters.DistanceFieldBrickAtlasHalfTexelSize;
+		Ar << Parameters.DistanceFieldBrickOffsetToAtlasUVScale;
+		Ar << Parameters.DistanceFieldUniqueDataBrickSizeInAtlasTexels;
+
+		return Ar;
+	}
+
+	template<typename ShaderRHIParamRef>
+	FORCEINLINE_DEBUGGABLE void Set(FRHICommandList& RHICmdList, const ShaderRHIParamRef ShaderRHI, const FDistanceFieldSceneData* ParameterData) const
+	{
+		if (IsBound() && ParameterData->GetCurrentObjectBuffers() != nullptr)
+		{
+			SetSRVParameter(RHICmdList, ShaderRHI, SceneObjectBounds, ParameterData->GetCurrentObjectBuffers()->Bounds.SRV);
+			SetSRVParameter(RHICmdList, ShaderRHI, SceneObjectData, ParameterData->GetCurrentObjectBuffers()->Data.SRV);
+			SetShaderValue(RHICmdList, ShaderRHI, NumSceneObjects, ParameterData->NumObjectsInBuffer);
+			SetSRVParameter(RHICmdList, ShaderRHI, SceneDistanceFieldAssetData, ParameterData->AssetDataBuffer.SRV);
+			SetSRVParameter(RHICmdList, ShaderRHI, DistanceFieldIndirectionTable, ParameterData->IndirectionTable.SRV);
+			SetTextureParameter(RHICmdList, ShaderRHI, DistanceFieldBrickTexture, ParameterData->DistanceFieldBrickVolumeTexture->GetRenderTargetItem().ShaderResourceTexture);
+			SetSamplerParameter(RHICmdList, ShaderRHI, DistanceFieldSampler, TStaticSamplerState<SF_Bilinear, AM_Clamp, AM_Clamp, AM_Clamp>::GetRHI());
+			SetShaderValue(RHICmdList, ShaderRHI, DistanceFieldBrickSize, FVector3f(DistanceField::BrickSize));
+			SetShaderValue(RHICmdList, ShaderRHI, DistanceFieldUniqueDataBrickSize, FVector3f(DistanceField::UniqueDataBrickSize));
+			SetShaderValue(RHICmdList, ShaderRHI, DistanceFieldBrickAtlasSizeInBricks, ParameterData->BrickTextureDimensionsInBricks);
+			SetShaderValue(RHICmdList, ShaderRHI, DistanceFieldBrickAtlasMask, ParameterData->BrickTextureDimensionsInBricks - FIntVector(1));
+			SetShaderValue(RHICmdList, ShaderRHI, DistanceFieldBrickAtlasSizeLog2, FIntVector(
+				FMath::FloorLog2(ParameterData->BrickTextureDimensionsInBricks.X),
+				FMath::FloorLog2(ParameterData->BrickTextureDimensionsInBricks.Y),
+				FMath::FloorLog2(ParameterData->BrickTextureDimensionsInBricks.Z)));
+			FVector3f DistanceFieldBrickAtlasTexelSizeTmp = FVector3f(1.0f) / FVector3f(ParameterData->BrickTextureDimensionsInBricks * DistanceField::BrickSize);
+			SetShaderValue(RHICmdList, ShaderRHI, DistanceFieldBrickAtlasTexelSize, DistanceFieldBrickAtlasTexelSizeTmp);
+			SetShaderValue(RHICmdList, ShaderRHI, DistanceFieldBrickAtlasHalfTexelSize, 0.5f * DistanceFieldBrickAtlasTexelSizeTmp);
+			SetShaderValue(RHICmdList, ShaderRHI, DistanceFieldBrickOffsetToAtlasUVScale, FVector3f(DistanceField::BrickSize) * DistanceFieldBrickAtlasTexelSizeTmp);
+			SetShaderValue(RHICmdList, ShaderRHI, DistanceFieldUniqueDataBrickSizeInAtlasTexels, FVector3f(DistanceField::UniqueDataBrickSize) * DistanceFieldBrickAtlasTexelSizeTmp);
+		}
+	}
+
+private:
+
+	LAYOUT_FIELD(FShaderResourceParameter, SceneObjectBounds)
+		LAYOUT_FIELD(FShaderResourceParameter, SceneObjectData)
+		LAYOUT_FIELD(FShaderParameter, NumSceneObjects)
+		LAYOUT_FIELD(FShaderResourceParameter, SceneDistanceFieldAssetData)
+		LAYOUT_FIELD(FShaderResourceParameter, DistanceFieldIndirectionTable)
+		LAYOUT_FIELD(FShaderResourceParameter, DistanceFieldBrickTexture)
+		LAYOUT_FIELD(FShaderResourceParameter, DistanceFieldSampler)
+		LAYOUT_FIELD(FShaderParameter, DistanceFieldBrickSize)
+		LAYOUT_FIELD(FShaderParameter, DistanceFieldUniqueDataBrickSize)
+		LAYOUT_FIELD(FShaderParameter, DistanceFieldBrickAtlasSizeInBricks)
+		LAYOUT_FIELD(FShaderParameter, DistanceFieldBrickAtlasMask)
+		LAYOUT_FIELD(FShaderParameter, DistanceFieldBrickAtlasSizeLog2)
+		LAYOUT_FIELD(FShaderParameter, DistanceFieldBrickAtlasTexelSize)
+		LAYOUT_FIELD(FShaderParameter, DistanceFieldBrickAtlasHalfTexelSize)
+		LAYOUT_FIELD(FShaderParameter, DistanceFieldBrickOffsetToAtlasUVScale)
+		LAYOUT_FIELD(FShaderParameter, DistanceFieldUniqueDataBrickSizeInAtlasTexels)
+};
\ No newline at end of file
diff --git a/Engine/Plugins/FX/Niagara/Source/Niagara/Private/NiagaraDataInterfaceRigidMeshCollisionQuery.cpp b/Engine/Plugins/FX/Niagara/Source/Niagara/Private/NiagaraDataInterfaceRigidMeshCollisionQuery.cpp
--- a/Engine/Plugins/FX/Niagara/Source/Niagara/Private/NiagaraDataInterfaceRigidMeshCollisionQuery.cpp
+++ b/Engine/Plugins/FX/Niagara/Source/Niagara/Private/NiagaraDataInterfaceRigidMeshCollisionQuery.cpp
@@ -18,6 +18,7 @@
 #include "Renderer/Private/DistanceFieldLightingShared.h"
 #include "NiagaraGpuComputeDispatchInterface.h"
 #include "NiagaraGpuComputeDispatch.h"
+#include "DataInterface/NiagaraDistanceFieldParameters.h"
 
 #define LOCTEXT_NAMESPACE "NiagaraDataInterfaceRigidMeshCollisionQuery"
 DEFINE_LOG_CATEGORY_STATIC(LogRigidMeshCollision, Log, All);
@@ -546,107 +547,6 @@ ETickingGroup FNDIRigidMeshCollisionData::ComputeTickingGroup()
 
 //------------------------------------------------------------------------------------------------------------
 
-class FDistanceFieldParameters
-{
-	DECLARE_INLINE_TYPE_LAYOUT(FDistanceFieldParameters, NonVirtual);
-public:
-	void Bind(const FShaderParameterMap& ParameterMap)
-	{
-		SceneObjectBounds.Bind(ParameterMap, TEXT("SceneObjectBounds"));
-		SceneObjectData.Bind(ParameterMap, TEXT("SceneObjectData"));
-		NumSceneObjects.Bind(ParameterMap, TEXT("NumSceneObjects"));
-		SceneDistanceFieldAssetData.Bind(ParameterMap, TEXT("SceneDistanceFieldAssetData"));
-		DistanceFieldIndirectionTable.Bind(ParameterMap, TEXT("DistanceFieldIndirectionTable"));
-		DistanceFieldBrickTexture.Bind(ParameterMap, TEXT("DistanceFieldBrickTexture"));
-		DistanceFieldSampler.Bind(ParameterMap, TEXT("DistanceFieldSampler"));
-		DistanceFieldBrickSize.Bind(ParameterMap, TEXT("DistanceFieldBrickSize"));
-		DistanceFieldUniqueDataBrickSize.Bind(ParameterMap, TEXT("DistanceFieldUniqueDataBrickSize"));
-		DistanceFieldBrickAtlasSizeInBricks.Bind(ParameterMap, TEXT("DistanceFieldBrickAtlasSizeInBricks"));
-		DistanceFieldBrickAtlasMask.Bind(ParameterMap, TEXT("DistanceFieldBrickAtlasMask"));
-		DistanceFieldBrickAtlasSizeLog2.Bind(ParameterMap, TEXT("DistanceFieldBrickAtlasSizeLog2"));
-		DistanceFieldBrickAtlasTexelSize.Bind(ParameterMap, TEXT("DistanceFieldBrickAtlasTexelSize"));
-		DistanceFieldBrickAtlasHalfTexelSize.Bind(ParameterMap, TEXT("DistanceFieldBrickAtlasHalfTexelSize"));
-		DistanceFieldBrickOffsetToAtlasUVScale.Bind(ParameterMap, TEXT("DistanceFieldBrickOffsetToAtlasUVScale"));
-		DistanceFieldUniqueDataBrickSizeInAtlasTexels.Bind(ParameterMap, TEXT("DistanceFieldUniqueDataBrickSizeInAtlasTexels"));
-	}
-
-	bool IsBound() const
-	{
-		return SceneDistanceFieldAssetData.IsBound();
-	}
-
-	friend FArchive& operator<<(FArchive& Ar, FDistanceFieldParameters& Parameters)
-	{
-		Ar << Parameters.SceneObjectBounds;
-		Ar << Parameters.SceneObjectData;
-		Ar << Parameters.NumSceneObjects;
-		Ar << Parameters.SceneDistanceFieldAssetData;
-		Ar << Parameters.DistanceFieldIndirectionTable;
-		Ar << Parameters.DistanceFieldBrickTexture;
-		Ar << Parameters.DistanceFieldSampler;
-		Ar << Parameters.DistanceFieldBrickSize;
-		Ar << Parameters.DistanceFieldUniqueDataBrickSize;
-		Ar << Parameters.DistanceFieldBrickAtlasSizeInBricks;
-		Ar << Parameters.DistanceFieldBrickAtlasMask;
-		Ar << Parameters.DistanceFieldBrickAtlasSizeLog2;
-		Ar << Parameters.DistanceFieldBrickAtlasTexelSize;
-		Ar << Parameters.DistanceFieldBrickAtlasHalfTexelSize;
-		Ar << Parameters.DistanceFieldBrickOffsetToAtlasUVScale;
-		Ar << Parameters.DistanceFieldUniqueDataBrickSizeInAtlasTexels;
-
-		return Ar;
-	}
-
-	template<typename ShaderRHIParamRef>
-	FORCEINLINE_DEBUGGABLE void Set(FRHICommandList& RHICmdList, const ShaderRHIParamRef ShaderRHI, const FDistanceFieldSceneData* ParameterData) const
-	{
-		if (IsBound() && ParameterData->GetCurrentObjectBuffers() != nullptr)
-		{
-			SetSRVParameter(RHICmdList, ShaderRHI, SceneObjectBounds, ParameterData->GetCurrentObjectBuffers()->Bounds.SRV);
-			SetSRVParameter(RHICmdList, ShaderRHI, SceneObjectData, ParameterData->GetCurrentObjectBuffers()->Data.SRV);
-			SetShaderValue(RHICmdList, ShaderRHI, NumSceneObjects, ParameterData->NumObjectsInBuffer);
-			SetSRVParameter(RHICmdList, ShaderRHI, SceneDistanceFieldAssetData, ParameterData->AssetDataBuffer.SRV);
-			SetSRVParameter(RHICmdList, ShaderRHI, DistanceFieldIndirectionTable, ParameterData->IndirectionTable.SRV);
-			SetTextureParameter(RHICmdList, ShaderRHI, DistanceFieldBrickTexture, ParameterData->DistanceFieldBrickVolumeTexture->GetRenderTargetItem().ShaderResourceTexture);
-			SetSamplerParameter(RHICmdList, ShaderRHI, DistanceFieldSampler, TStaticSamplerState<SF_Bilinear, AM_Clamp, AM_Clamp, AM_Clamp>::GetRHI());
-			SetShaderValue(RHICmdList, ShaderRHI, DistanceFieldBrickSize, FVector3f(DistanceField::BrickSize));
-			SetShaderValue(RHICmdList, ShaderRHI, DistanceFieldUniqueDataBrickSize, FVector3f(DistanceField::UniqueDataBrickSize));
-			SetShaderValue(RHICmdList, ShaderRHI, DistanceFieldBrickAtlasSizeInBricks, ParameterData->BrickTextureDimensionsInBricks);
-			SetShaderValue(RHICmdList, ShaderRHI, DistanceFieldBrickAtlasMask, ParameterData->BrickTextureDimensionsInBricks - FIntVector(1));
-			SetShaderValue(RHICmdList, ShaderRHI, DistanceFieldBrickAtlasSizeLog2, FIntVector(
-				FMath::FloorLog2(ParameterData->BrickTextureDimensionsInBricks.X),
-				FMath::FloorLog2(ParameterData->BrickTextureDimensionsInBricks.Y),
-				FMath::FloorLog2(ParameterData->BrickTextureDimensionsInBricks.Z)));
-			FVector3f DistanceFieldBrickAtlasTexelSizeTmp = FVector3f(1.0f) / FVector3f(ParameterData->BrickTextureDimensionsInBricks * DistanceField::BrickSize);
-			SetShaderValue(RHICmdList, ShaderRHI, DistanceFieldBrickAtlasTexelSize, DistanceFieldBrickAtlasTexelSizeTmp);
-			SetShaderValue(RHICmdList, ShaderRHI, DistanceFieldBrickAtlasHalfTexelSize, 0.5f * DistanceFieldBrickAtlasTexelSizeTmp);
-			SetShaderValue(RHICmdList, ShaderRHI, DistanceFieldBrickOffsetToAtlasUVScale, FVector3f(DistanceField::BrickSize) * DistanceFieldBrickAtlasTexelSizeTmp);
-			SetShaderValue(RHICmdList, ShaderRHI, DistanceFieldUniqueDataBrickSizeInAtlasTexels, FVector3f(DistanceField::UniqueDataBrickSize) * DistanceFieldBrickAtlasTexelSizeTmp);
-		}
-	}
-
-private:
-
-	LAYOUT_FIELD(FShaderResourceParameter, SceneObjectBounds)
-		LAYOUT_FIELD(FShaderResourceParameter, SceneObjectData)
-		LAYOUT_FIELD(FShaderParameter, NumSceneObjects)
-		LAYOUT_FIELD(FShaderResourceParameter, SceneDistanceFieldAssetData)
-		LAYOUT_FIELD(FShaderResourceParameter, DistanceFieldIndirectionTable)
-		LAYOUT_FIELD(FShaderResourceParameter, DistanceFieldBrickTexture)
-		LAYOUT_FIELD(FShaderResourceParameter, DistanceFieldSampler)
-		LAYOUT_FIELD(FShaderParameter, DistanceFieldBrickSize)
-		LAYOUT_FIELD(FShaderParameter, DistanceFieldUniqueDataBrickSize)
-		LAYOUT_FIELD(FShaderParameter, DistanceFieldBrickAtlasSizeInBricks)
-		LAYOUT_FIELD(FShaderParameter, DistanceFieldBrickAtlasMask)
-		LAYOUT_FIELD(FShaderParameter, DistanceFieldBrickAtlasSizeLog2)
-		LAYOUT_FIELD(FShaderParameter, DistanceFieldBrickAtlasTexelSize)
-		LAYOUT_FIELD(FShaderParameter, DistanceFieldBrickAtlasHalfTexelSize)
-		LAYOUT_FIELD(FShaderParameter, DistanceFieldBrickOffsetToAtlasUVScale)
-		LAYOUT_FIELD(FShaderParameter, DistanceFieldUniqueDataBrickSizeInAtlasTexels)
-};
-
-//------------------------------------------------------------------------------------------------------------
-
 struct FNDIRigidMeshCollisionParametersCS : public FNiagaraDataInterfaceParametersCS
 {
 	DECLARE_TYPE_LAYOUT(FNDIRigidMeshCollisionParametersCS, NonVirtual);
diff --git a/Engine/Plugins/FX/Niagara/Source/Niagara/Private/NiagaraGpuComputeDispatch.cpp b/Engine/Plugins/FX/Niagara/Source/Niagara/Private/NiagaraGpuComputeDispatch.cpp
--- a/Engine/Plugins/FX/Niagara/Source/Niagara/Private/NiagaraGpuComputeDispatch.cpp
+++ b/Engine/Plugins/FX/Niagara/Source/Niagara/Private/NiagaraGpuComputeDispatch.cpp
@@ -1941,8 +1941,7 @@ const FGlobalDistanceFieldParameterData* FNiagaraGpuComputeDispatch::GetGlobalDi
 
 const FDistanceFieldSceneData* FNiagaraGpuComputeDispatch::GetMeshDistanceFieldParameters() const
 {
-	check(CurrentPassViews.Num() > 0);
-	if (CurrentPassViews[0].Family == nullptr || CurrentPassViews[0].Family->Scene == nullptr || CurrentPassViews[0].Family->Scene->GetRenderScene() == nullptr)
+	if (CurrentPassViews.Num() == 0 || CurrentPassViews[0].Family == nullptr || CurrentPassViews[0].Family->Scene == nullptr || CurrentPassViews[0].Family->Scene->GetRenderScene() == nullptr)
 	{
 		return nullptr;
 	}
diff --git a/Engine/Plugins/FX/Niagara/Source/Niagara/Public/NiagaraDataInterfaceUtilities.h b/Engine/Plugins/FX/Niagara/Source/Niagara/Public/NiagaraDataInterfaceUtilities.h
new file mode 100644
--- /dev/null
+++ b/Engine/Plugins/FX/Niagara/Source/Niagara/Public/NiagaraDataInterfaceUtilities.h
@@ -0,0 +1,26 @@
+// Copyright Epic Games, Inc. All Rights Reserved.
+#pragma once
+
+#include "NiagaraCommon.h"
+#include "NiagaraShared.h"
+
+namespace FNiagaraDataInterfaceUtilities
+{
+	// Finds all VM function calls made using the data interface that equals this one (i.e. A->Equals(B))
+	// The same function call made be made multiple times across different scripts so you may see the same function multiple times
+	// The action should return True to continue iteration or False to stop
+	NIAGARA_API void ForEachVMFunctionEquals(class UNiagaraDataInterface* DataInterface, class UNiagaraSystem* NiagaraSystem, class UNiagaraComponent* Component, TFunction<bool(const FVMExternalFunctionBindingInfo&)> Action);
+	// Finds all Gpu function calls made using the data interface that equals this one (i.e. A->Equals(B))
+	// The same function call made be made multiple times across different scripts so you may see the same function multiple times
+	// The action should return True to continue iteration or False to stop
+	NIAGARA_API void ForEachGpuFunctionEquals(class UNiagaraDataInterface* DataInterface, class UNiagaraSystem* NiagaraSystem, class UNiagaraComponent* Component, TFunction<bool(const FNiagaraDataInterfaceGeneratedFunction&)> Action);
+
+	// Finds all VM function calls made using the data interface (i.e. pointer comparison A == B)
+	// The same function call made be made multiple times across different scripts so you may see the same function multiple times
+	// The action should return True to continue iteration or False to stop
+	NIAGARA_API void ForEachVMFunction(class UNiagaraDataInterface* DataInterface, class FNiagaraSystemInstance* SystemInstance, TFunction<bool(const FVMExternalFunctionBindingInfo&)> Action);
+	// Finds all Gpu function calls made using the data interface (i.e. pointer comparison A == B)
+	// The same function call made be made multiple times across different scripts so you may see the same function multiple times
+	// The action should return True to continue iteration or False to stop
+	NIAGARA_API void ForEachGpuFunction(class UNiagaraDataInterface* DataInterface, class FNiagaraSystemInstance* SystemInstance, TFunction<bool(const FNiagaraDataInterfaceGeneratedFunction&)> Action);
+}
diff --git a/Engine/Plugins/FX/Niagara/Source/Niagara/Public/NiagaraSettings.h b/Engine/Plugins/FX/Niagara/Source/Niagara/Public/NiagaraSettings.h
--- a/Engine/Plugins/FX/Niagara/Source/Niagara/Public/NiagaraSettings.h
+++ b/Engine/Plugins/FX/Niagara/Source/Niagara/Public/NiagaraSettings.h
@@ -113,6 +113,14 @@ class NIAGARA_API UNiagaraSettings : public UDeveloperSettings
 	UPROPERTY(config, EditAnywhere, Category = SkeletalMeshDI, meta = (DisplayName = "Adjacency Triangle Index Format", ToolTip = "Controls the format used for specifying triangle indexes in adjacency buffers.  Changing this setting requires restarting the editor.", ConfigRestartRequired = true))
 	TEnumAsByte<ENDISkelMesh_AdjacencyTriangleIndexFormat::Type> NDISkelMesh_AdjacencyTriangleIndexFormat;
 
+	/**
+	When enabled the static mesh data interface is allowed to sample from the distance field data (if present) on the GPU.
+	Enabling this feature will move all systems that contain static mesh samples into PostRenderOpaque tick group regardless of the features used.
+	Changing this setting requires restarting the editor.
+	*/
+	UPROPERTY(config, EditAnywhere, Category = StaticMeshDI, meta = (DisplayName = "Allow Distance Fields (Experimental)", ConfigRestartRequired = true))
+	bool NDIStaticMesh_AllowDistanceFields = false;
+
 	// Begin UDeveloperSettings Interface
 	virtual FName GetCategoryName() const override;
 #if WITH_EDITOR
