Readded a part that causes a refresh in the Niagara stack UI
Was originally replaced with another function but it appears those were not semantically the same

diff --git a/Engine/Plugins/FX/Niagara/Source/NiagaraEditor/Private/ViewModels/NiagaraScratchPadScriptViewModel.cpp b/Engine/Plugins/FX/Niagara/Source/NiagaraEditor/Private/ViewModels/NiagaraScratchPadScriptViewModel.cpp
--- a/Engine/Plugins/FX/Niagara/Source/NiagaraEditor/Private/ViewModels/NiagaraScratchPadScriptViewModel.cpp
+++ b/Engine/Plugins/FX/Niagara/Source/NiagaraEditor/Private/ViewModels/NiagaraScratchPadScriptViewModel.cpp
@@ -177,6 +177,15 @@ void FNiagaraScratchPadScriptViewModel::ApplyChanges()
 		OriginalScript->GetClass());
 	bHasPendingChanges = false;
 
+	TArray<UNiagaraNodeFunctionCall*> FunctionCallNodesToRefresh;
+	FNiagaraEditorUtilities::GetReferencingFunctionCallNodes(OriginalScript, FunctionCallNodesToRefresh);
+
+	for (UNiagaraNodeFunctionCall* FunctionCallNodeToRefresh : FunctionCallNodesToRefresh)
+	{
+		FunctionCallNodeToRefresh->RefreshFromExternalChanges();
+		FunctionCallNodeToRefresh->MarkNodeRequiresSynchronization(TEXT("ScratchPadChangesApplied"), true);
+	}
+	
 	FRefreshAllScriptsFromExternalChangesArgs Args;
 	Args.OriginatingScript = OriginalScript;
 	Args.OriginatingGraph = CastChecked<UNiagaraScriptSource>(OriginalScript->GetSource(EditScript.Version))->NodeGraph;
