Fixing static variable driving static switches and hiding inputs/outputs

diff --git a/Engine/Plugins/FX/Niagara/Source/NiagaraCore/Private/NiagaraCustomVersion.cpp b/Engine/Plugins/FX/Niagara/Source/NiagaraCore/Private/NiagaraCustomVersion.cpp
--- a/Engine/Plugins/FX/Niagara/Source/NiagaraCore/Private/NiagaraCustomVersion.cpp
+++ b/Engine/Plugins/FX/Niagara/Source/NiagaraCore/Private/NiagaraCustomVersion.cpp
@@ -8,4 +8,4 @@ const FGuid FNiagaraCustomVersion::GUID(0xFCF57AFA, 0x50764283, 0xB9A9E658, 0xFF
 // Register the custom version with core
 FCustomVersionRegistration GRegisterNiagaraCustomVersion(FNiagaraCustomVersion::GUID, FNiagaraCustomVersion::LatestVersion, TEXT("NiagaraVer"));
 
-const FGuid FNiagaraCustomVersion::LatestScriptCompileVersion(0x2FCAFEF0, 0x0A33824E, 0x87058A67, 0x5C93C1B2);
\ No newline at end of file
+const FGuid FNiagaraCustomVersion::LatestScriptCompileVersion(0xF51BF7C2, 0x0DA14340, 0xB6DD6A83, 0xC44E2DCE);
\ No newline at end of file
diff --git a/Engine/Plugins/FX/Niagara/Source/NiagaraEditor/Private/NiagaraGraph.cpp b/Engine/Plugins/FX/Niagara/Source/NiagaraEditor/Private/NiagaraGraph.cpp
--- a/Engine/Plugins/FX/Niagara/Source/NiagaraEditor/Private/NiagaraGraph.cpp
+++ b/Engine/Plugins/FX/Niagara/Source/NiagaraEditor/Private/NiagaraGraph.cpp
@@ -1427,9 +1427,9 @@ void UNiagaraGraph::FindInputNodes(TArray<UNiagaraNodeInput*>& OutInputNodes, UN
 	}
 }
 
-TArray<FNiagaraVariable> UNiagaraGraph::FindStaticSwitchInputs(bool bReachableOnly) const
+TArray<FNiagaraVariable> UNiagaraGraph::FindStaticSwitchInputs(bool bReachableOnly, const TArray<FNiagaraVariable>& InStaticVars) const
 {
-	TArray<UEdGraphNode*> NodesToProcess = bReachableOnly ? FindReachableNodes() : Nodes;
+	TArray<UEdGraphNode*> NodesToProcess = bReachableOnly ? FindReachableNodes(InStaticVars) : Nodes;
 
 	TArray<FNiagaraVariable> Result;
 	for (UEdGraphNode* Node : NodesToProcess)
@@ -1453,55 +1453,33 @@ TArray<FNiagaraVariable> UNiagaraGraph::FindStaticSwitchInputs(bool bReachableOn
 	return Result;
 }
 
-TArray<UEdGraphNode*> UNiagaraGraph::FindReachableNodes() const
+TArray<UEdGraphNode*> UNiagaraGraph::FindReachableNodes(const TArray<FNiagaraVariable>& InStaticVars) const
 {
+	
 	TArray<UEdGraphNode*> ResultNodes;
+	FNiagaraParameterMapHistoryBuilder Builder;
+	Builder.RegisterExternalStaticVariables(InStaticVars);
+	Builder.SetIgnoreDisabled(false);
+
 	TArray<UNiagaraNodeOutput*> OutNodes;
 	FindOutputNodes(OutNodes);
 	ResultNodes.Append(OutNodes);
-
-	FPinCollectorArray OutPins;
-	for (int i = 0; i < ResultNodes.Num(); i++)
+	for (UNiagaraNodeOutput* OutNode : OutNodes)
 	{
-		UEdGraphNode* Node = ResultNodes[i];
-		if (Node == nullptr)
-		{
-			continue;
-		}
-		
-		UNiagaraNodeStaticSwitch* SwitchNode = Cast<UNiagaraNodeStaticSwitch>(Node);
-		if (SwitchNode)
-		{
-			OutPins.Reset();
-			SwitchNode->GetOutputPins(OutPins);
-			for (UEdGraphPin* Pin : OutPins)
-			{
-				UEdGraphPin* TracedPin = SwitchNode->GetTracedOutputPin(Pin, false, true);
-				if (TracedPin && TracedPin != Pin)
-				{
-					ResultNodes.AddUnique(TracedPin->GetOwningNode());
-				}
-			}
-		}
-		else
+		Builder.BuildParameterMaps(OutNode, true);
 		{
-			for (UEdGraphPin* Pin : Node->GetAllPins())
+			TArray<const UNiagaraNode*> VisitedNodes;
+			Builder.GetContextuallyVisitedNodes(VisitedNodes);
+			for (const UNiagaraNode* Node : VisitedNodes)
 			{
-				if (!Pin || Pin->Direction != EEdGraphPinDirection::EGPD_Input)
-				{
-					continue;
-				}
-				for (UEdGraphPin* LinkedPin : Pin->LinkedTo)
-				{
-					if (!LinkedPin)
-					{
-						continue;
-					}
-					ResultNodes.AddUnique(LinkedPin->GetOwningNode());
-				}
+				if (Node->GetOuter() == this)
+					ResultNodes.AddUnique((UNiagaraNode*)Node);
 			}
 		}
 	}
+
+	
+
 	return ResultNodes;
 }
 
diff --git a/Engine/Plugins/FX/Niagara/Source/NiagaraEditor/Private/ViewModels/Stack/NiagaraStackGraphUtilities.cpp b/Engine/Plugins/FX/Niagara/Source/NiagaraEditor/Private/ViewModels/Stack/NiagaraStackGraphUtilities.cpp
--- a/Engine/Plugins/FX/Niagara/Source/NiagaraEditor/Private/ViewModels/Stack/NiagaraStackGraphUtilities.cpp
+++ b/Engine/Plugins/FX/Niagara/Source/NiagaraEditor/Private/ViewModels/Stack/NiagaraStackGraphUtilities.cpp
@@ -173,6 +173,27 @@ void FNiagaraStackGraphUtilities::GetWrittenVariablesForGraph(UEdGraph& Graph, T
 	TArray<UNiagaraNodeOutput*> OutputNodes;
 	Graph.GetNodesOfClass<UNiagaraNodeOutput>(OutputNodes);
 	FPinCollectorArray InputPins;
+
+
+	UNiagaraEmitter* Emitter = Graph.GetTypedOuter<UNiagaraEmitter>();
+	UNiagaraSystem* System = Graph.GetTypedOuter<UNiagaraSystem>();
+
+	TArray<FNiagaraVariable> StaticVars;
+	TSharedPtr<FNiagaraGraphCachedDataBase, ESPMode::ThreadSafe> CachedTraversalData;
+	if (Emitter)
+	{
+		CachedTraversalData = Emitter->GetCachedTraversalData();
+	}
+	else if (System)
+	{
+		CachedTraversalData = System->GetCachedTraversalData();
+	}
+	if (CachedTraversalData.IsValid())
+	{
+		CachedTraversalData.Get()->GetStaticVariables(StaticVars);
+	}
+
+
 	for (UNiagaraNodeOutput* OutputNode : OutputNodes)
 	{
 		InputPins.Reset();
@@ -180,6 +201,7 @@ void FNiagaraStackGraphUtilities::GetWrittenVariablesForGraph(UEdGraph& Graph, T
 		if (InputPins.Num() == 1)
 		{
 			FNiagaraParameterMapHistoryBuilder Builder;
+			Builder.RegisterExternalStaticVariables(StaticVars);
 			Builder.BuildParameterMaps(OutputNode, true);
 			check(Builder.Histories.Num() == 1);
 			for (int32 i = 0; i < Builder.Histories[0].Variables.Num(); i++)
@@ -415,6 +437,20 @@ void FNiagaraStackGraphUtilities::BuildParameterMapHistoryWithStackContextResolu
 		bSetUsage = true;
 	}
 
+	TArray<FNiagaraVariable> StaticVars;
+	TSharedPtr<FNiagaraGraphCachedDataBase, ESPMode::ThreadSafe> CachedTraversalData;
+	if (OwningEmitter)
+	{
+		CachedTraversalData = OwningEmitter->GetCachedTraversalData();
+	}
+	
+	if (CachedTraversalData.IsValid())
+	{
+		CachedTraversalData.Get()->GetStaticVariables(StaticVars);
+	}
+
+	Builder.RegisterExternalStaticVariables(StaticVars);
+
 	NodeToVisit->BuildParameterMapHistory(Builder, bRecursive, bFilterForCompilation);
 
 	if (bSetUsage)
@@ -932,8 +968,46 @@ void FNiagaraStackGraphUtilities::GetStackFunctionStaticSwitchPins(UNiagaraNodeF
 	FunctionCallNode.GetInputPins(InputPins);
 	FNiagaraEditorUtilities::SetStaticSwitchConstants(FunctionCallGraph, InputPins, ConstantResolver);
 
-	TArray<FNiagaraVariable> SwitchInputs = FunctionCallGraph->FindStaticSwitchInputs(false);
-	TArray<FNiagaraVariable> ReachableInputs = FunctionCallGraph->FindStaticSwitchInputs(true);
+	UNiagaraEmitter* Emitter = FunctionCallNode.GetTypedOuter<UNiagaraEmitter>();
+	UNiagaraSystem* System = FunctionCallNode.GetTypedOuter<UNiagaraSystem>();
+
+	TArray<FNiagaraVariable> StaticVars;
+	TSharedPtr<FNiagaraGraphCachedDataBase, ESPMode::ThreadSafe> CachedTraversalData;
+	FString EmitterName;
+	FNiagaraAliasContext::ERapidIterationParameterMode Mode = FNiagaraAliasContext::ERapidIterationParameterMode::None;
+	if (Emitter)
+	{
+		CachedTraversalData = Emitter->GetCachedTraversalData();
+		EmitterName = Emitter->GetUniqueEmitterName();
+		Mode = FNiagaraAliasContext::ERapidIterationParameterMode::EmitterOrParticleScript;
+	}
+	else if (System)
+	{
+		CachedTraversalData = System->GetCachedTraversalData();
+		Mode = FNiagaraAliasContext::ERapidIterationParameterMode::SystemScript;
+	}
+	if (CachedTraversalData.IsValid())
+	{
+		TArray<FNiagaraVariable> CachedStatics;
+		CachedTraversalData.Get()->GetStaticVariables(CachedStatics);
+
+		// Because we've lost all context here about the specific function call we're in, we need to convert static vars back to 
+		// normal mode.
+		FString ModuleName = FunctionCallNode.GetFunctionName();
+		FNiagaraAliasContext AliasContext(Mode);
+		AliasContext.ChangeModuleNameToModule(ModuleName);
+		if (EmitterName.Len() != 0)	
+			AliasContext.ChangeEmitterNameToEmitter(EmitterName);
+
+		for (int32 i = 0; i < CachedStatics.Num(); i++)
+		{
+			FNiagaraVariable Var = FNiagaraUtilities::ResolveAliases(CachedStatics[i], AliasContext);
+			StaticVars.Add(Var);
+		}
+	}
+
+	TArray<FNiagaraVariable> SwitchInputs = FunctionCallGraph->FindStaticSwitchInputs();
+	TArray<FNiagaraVariable> ReachableInputs = FunctionCallGraph->FindStaticSwitchInputs(true, StaticVars);
 	for (FNiagaraVariable SwitchInput : SwitchInputs)
 	{
 		FEdGraphPinType PinType = Schema->TypeDefinitionToPinType(SwitchInput.GetType());
@@ -961,6 +1035,27 @@ void FNiagaraStackGraphUtilities::GetStackFunctionOutputVariables(UNiagaraNodeFu
 	FNiagaraParameterMapHistoryBuilder Builder;
 	Builder.SetIgnoreDisabled(false);
 	Builder.ConstantResolver = ConstantResolver;
+
+	UNiagaraEmitter* Emitter = FunctionCallNode.GetTypedOuter<UNiagaraEmitter>();
+	UNiagaraSystem* System = FunctionCallNode.GetTypedOuter<UNiagaraSystem>();
+
+	TArray<FNiagaraVariable> StaticVars;
+	TSharedPtr<FNiagaraGraphCachedDataBase, ESPMode::ThreadSafe> CachedTraversalData;
+	if (Emitter)
+	{
+		CachedTraversalData = Emitter->GetCachedTraversalData();
+	}
+	else if (System)
+	{
+		CachedTraversalData = System->GetCachedTraversalData();
+	}
+	if (CachedTraversalData.IsValid())
+	{
+		CachedTraversalData.Get()->GetStaticVariables(StaticVars);
+	}
+
+	Builder.RegisterExternalStaticVariables(StaticVars);
+
 	FunctionCallNode.BuildParameterMapHistory(Builder, false);
 
 	if (ensureMsgf(Builder.Histories.Num() == 1, TEXT("Invalid Stack Graph - Function call node has invalid history count!")))
@@ -993,6 +1088,27 @@ bool FNiagaraStackGraphUtilities::GetStackFunctionInputAndOutputVariables(UNiaga
 	FNiagaraParameterMapHistoryBuilder Builder;
 	Builder.SetIgnoreDisabled(false);
 	Builder.ConstantResolver = ConstantResolver;
+
+	UNiagaraEmitter* Emitter = FunctionCallNode.GetTypedOuter<UNiagaraEmitter>();
+	UNiagaraSystem* System = FunctionCallNode.GetTypedOuter<UNiagaraSystem>();
+
+	TArray<FNiagaraVariable> StaticVars;
+	TSharedPtr<FNiagaraGraphCachedDataBase, ESPMode::ThreadSafe> CachedTraversalData;
+	if (Emitter)
+	{
+		CachedTraversalData = Emitter->GetCachedTraversalData();
+	}
+	else if (System)
+	{
+		CachedTraversalData = System->GetCachedTraversalData();
+	}
+	if (CachedTraversalData.IsValid())
+	{
+		CachedTraversalData.Get()->GetStaticVariables(StaticVars);
+	}
+
+	Builder.RegisterExternalStaticVariables(StaticVars);
+
 	FunctionCallNode.BuildParameterMapHistory(Builder, false);
 
 	if (Builder.Histories.Num() == 0)
diff --git a/Engine/Plugins/FX/Niagara/Source/NiagaraEditor/Public/NiagaraGraph.h b/Engine/Plugins/FX/Niagara/Source/NiagaraEditor/Public/NiagaraGraph.h
--- a/Engine/Plugins/FX/Niagara/Source/NiagaraEditor/Public/NiagaraGraph.h
+++ b/Engine/Plugins/FX/Niagara/Source/NiagaraEditor/Public/NiagaraGraph.h
@@ -207,7 +207,7 @@ class UNiagaraGraph : public UEdGraph
 	void FindInputNodes(TArray<class UNiagaraNodeInput*>& OutInputNodes, FFindInputNodeOptions Options = FFindInputNodeOptions()) const;
 
 	/** Returns a list of variable inputs for all static switch nodes in the graph. */
-	TArray<FNiagaraVariable> NIAGARAEDITOR_API FindStaticSwitchInputs(bool bReachableOnly = false) const;
+	TArray<FNiagaraVariable> NIAGARAEDITOR_API FindStaticSwitchInputs(bool bReachableOnly = false, const TArray<FNiagaraVariable>& InStaticVars = TArray<FNiagaraVariable>()) const;
 
 	/** Get an in-order traversal of a graph by the specified target output script usage.*/
 	void BuildTraversal(TArray<class UNiagaraNode*>& OutNodesTraversed, ENiagaraScriptUsage TargetUsage, FGuid TargetUsageId, bool bEvaluateStaticSwitches = false) const;
@@ -404,7 +404,7 @@ private:
 	FOnGraphChanged OnGraphNeedsRecompile;
 
 	/** Find all nodes in the graph that can be reached during compilation. */
-	TArray<UEdGraphNode*> FindReachableNodes() const;
+	TArray<UEdGraphNode*> FindReachableNodes(const TArray<FNiagaraVariable>& InStaticVars) const;
 
 	/** Compares the values on the default pins with the metadata and syncs the two if necessary */
 	void ValidateDefaultPins();
diff --git a/Engine/Plugins/FX/Niagara/Source/NiagaraEditor/Public/NiagaraParameterMapHistory.h b/Engine/Plugins/FX/Niagara/Source/NiagaraEditor/Public/NiagaraParameterMapHistory.h
--- a/Engine/Plugins/FX/Niagara/Source/NiagaraEditor/Public/NiagaraParameterMapHistory.h
+++ b/Engine/Plugins/FX/Niagara/Source/NiagaraEditor/Public/NiagaraParameterMapHistory.h
@@ -508,8 +508,13 @@ public:
 	void SetConstantByStaticVariable(int32& OutValue, const FNiagaraVariable& Var);
 	int32 FindStaticVariable(const FNiagaraVariable& Var) const;
 
-
 	TArray<FNiagaraVariable> StaticVariables;
+
+	void GetContextuallyVisitedNodes(TArray<const class UNiagaraNode*>& OutVistedNodes)
+	{
+		if (ContextuallyVisitedNodes.Num() > 0)
+			OutVistedNodes.Append(ContextuallyVisitedNodes.Last());
+	}
 protected:
 	/**
 	* Generate the internal alias map from the current traversal state.
