- Fixes for LWC and structs containing LWC types, i.e. struct containing an FVector, etc.
- The editor still sees the original struct and under the hood we generate a replacement struct which is used inside the parameter stores, etc

#rnx

diff --git a/Engine/Plugins/FX/Niagara/Source/Niagara/Private/NiagaraCommon.cpp b/Engine/Plugins/FX/Niagara/Source/Niagara/Private/NiagaraCommon.cpp
--- a/Engine/Plugins/FX/Niagara/Source/Niagara/Private/NiagaraCommon.cpp
+++ b/Engine/Plugins/FX/Niagara/Source/Niagara/Private/NiagaraCommon.cpp
@@ -7,7 +7,6 @@
 #include "NiagaraConstants.h"
 #include "NiagaraCustomVersion.h"
 #include "NiagaraParameterCollection.h"
-#include "NiagaraParameterDefinitionsBase.h"
 #include "NiagaraScriptSourceBase.h"
 #include "NiagaraStats.h"
 #include "NiagaraSystemInstance.h"
@@ -67,121 +66,6 @@ static FAutoConsoleVariableRef CVarNiagaraLogVerboseWarnings(
 	ECVF_Default
 );
 
-//////////////////////////////////////////////////////////////////////////
-
-FString FNiagaraTypeHelper::ToString(const uint8* ValueData, const UObject* StructOrEnum)
-{
-	FString Ret;
-	if (const UEnum* Enum = Cast<const UEnum>(StructOrEnum))
-	{
-		Ret = Enum->GetNameStringByValue(*(int32*)ValueData);
-	}
-	else if (const UScriptStruct* Struct = Cast<const UScriptStruct>(StructOrEnum))
-	{
-		if (Struct == FNiagaraTypeDefinition::GetFloatStruct())
-		{
-			Ret += FString::Printf(TEXT("%g "), *(float*)ValueData);
-		}
-		else if (Struct == FNiagaraTypeDefinition::GetIntStruct())
-		{
-			Ret += FString::Printf(TEXT("%d "), *(int32*)ValueData);
-		}
-		else if (Struct == FNiagaraTypeDefinition::GetBoolStruct())
-		{
-			int32 Val = *(int32*)ValueData;
-			Ret += Val == 0xFFFFFFFF ? (TEXT("True")) : (Val == 0x0 ? TEXT("False") : TEXT("Invalid"));
-		}
-		else
-		{
-			for (TFieldIterator<FProperty> PropertyIt(Struct, EFieldIteratorFlags::IncludeSuper); PropertyIt; ++PropertyIt)
-			{
-				const FProperty* Property = *PropertyIt;
-				const uint8* PropPtr = ValueData + PropertyIt->GetOffset_ForInternal();
-				if (Property->IsA(FFloatProperty::StaticClass()))
-				{
-					Ret += FString::Printf(TEXT("%s: %g "), *Property->GetNameCPP(), *(float*)PropPtr);
-				}
-				else if (Property->IsA(FDoubleProperty::StaticClass()))
-				{
-					Ret += FString::Printf(TEXT("%s: %g "), *Property->GetNameCPP(), *(double*)PropPtr);
-				}
-				else if (Property->IsA(FUInt16Property::StaticClass()))
-				{
-					FFloat16 Val = *(FFloat16*)PropPtr;
-					Ret += FString::Printf(TEXT("%s: %f "), *Property->GetNameCPP(), Val.GetFloat());
-				}
-				else if (Property->IsA(FIntProperty::StaticClass()))
-				{
-					Ret += FString::Printf(TEXT("%s: %d "), *Property->GetNameCPP(), *(int32*)PropPtr);
-				}
-				else if (Property->IsA(FBoolProperty::StaticClass()))
-				{
-					int32 Val = *(int32*)ValueData;
-					FString BoolStr = Val == 0xFFFFFFFF ? (TEXT("True")) : (Val == 0x0 ? TEXT("False") : TEXT("Invalid"));
-					Ret += FString::Printf(TEXT("%s: %d "), *Property->GetNameCPP(), *BoolStr);
-				}
-				else if (const FStructProperty* StructProp = CastFieldChecked<const FStructProperty>(Property))
-				{
-					Ret += FString::Printf(TEXT("%s: (%s) "), *Property->GetNameCPP(), *FNiagaraTypeHelper::ToString(PropPtr, StructProp->Struct));
-				}
-				else
-				{
-					check(false);
-					Ret += TEXT("Unknown Type");
-				}
-			}
-		}
-	}
-	return Ret;
-}
-
-UScriptStruct* FNiagaraTypeHelper::FindNiagaraFriendlyTopLevelStruct(UScriptStruct* InStruct)
-{
-	if (!InStruct)
-		return nullptr;
-
-	// Note: UE core types are converted to the float variant as Niagara only works with float types.
-	if (InStruct->GetFName() == NAME_Vector2D || InStruct->GetFName() == NAME_Vector2d) // LWC support
-	{
-		return FNiagaraTypeDefinition::GetVec2Struct();
-	}
-
-	if (InStruct->GetFName() == NAME_Vector || InStruct->GetFName() == NAME_Vector3d) // LWC support
-	{
-		return FNiagaraTypeDefinition::GetVec3Struct();
-	}
-
-	if (InStruct->GetFName() == NAME_Vector4 || InStruct->GetFName() == NAME_Vector4d) // LWC support
-	{
-		return FNiagaraTypeDefinition::GetVec4Struct();
-	}
-
-	if (InStruct->GetFName() == NAME_Quat || InStruct->GetFName() == NAME_Quat4d) // LWC support
-	{
-		return FNiagaraTypeDefinition::GetQuatStruct();
-	}
-
-	if (InStruct->GetFName() == FName("NiagaraPosition"))
-	{
-		return FNiagaraTypeDefinition::GetPositionStruct();
-	}
-
-	return InStruct;
-}
-
-bool FNiagaraTypeHelper::IsNiagaraFriendlyTopLevelStruct(UScriptStruct* InStruct)
-{
-	if (FindNiagaraFriendlyTopLevelStruct(InStruct) == InStruct)
-	{
-		return true;
-	}
-	else
-	{
-		return false;
-	}
-}
-//////////////////////////////////////////////////////////////////////////
-
 FNiagaraSystemUpdateContext::~FNiagaraSystemUpdateContext()
 {
 	CommitUpdate();
diff --git a/Engine/Plugins/FX/Niagara/Source/Niagara/Private/NiagaraComponent.cpp b/Engine/Plugins/FX/Niagara/Source/Niagara/Private/NiagaraComponent.cpp
--- a/Engine/Plugins/FX/Niagara/Source/Niagara/Private/NiagaraComponent.cpp
+++ b/Engine/Plugins/FX/Niagara/Source/Niagara/Private/NiagaraComponent.cpp
@@ -2429,7 +2429,7 @@ void FixInvalidUserParameters(FNiagaraUserRedirectionParameterStore& ParameterSt
 				if (DataValuePtr != nullptr)
 				{
 					DataValue.AddUninitialized(IncorrectlyNamedParameter.GetSizeInBytes());
-					FMemory::Memcpy(DataValue.GetData(), DataValuePtr, IncorrectlyNamedParameter.GetSizeInBytes());
+					ParameterStore.CopyParameterData(IncorrectlyNamedParameter, DataValue.GetData());
 				}
 			}
 
@@ -2634,7 +2634,10 @@ static FNiagaraVariant GetParameterValueFromStore(const FNiagaraVariableBase& Va
 		return FNiagaraVariant();
 	}
 
-	return FNiagaraVariant(ParameterData, Var.GetSizeInBytes());
+	TArray<uint8> DataValue;
+	DataValue.AddUninitialized(Var.GetSizeInBytes());
+	Store.CopyParameterData(Var, DataValue.GetData());
+	return FNiagaraVariant(DataValue);
 }
 
 #if WITH_EDITOR
@@ -2836,10 +2839,11 @@ void UNiagaraComponent::EnsureOverrideParametersConsistent() const
 			}
 			else
 			{
-				const uint8* ActualData = OverrideParameters.GetParameterData(Key);
-				if (ActualData != nullptr)
+				TArray<uint8> DataValue;
+				DataValue.AddUninitialized(Key.GetSizeInBytes());
+				if (OverrideParameters.CopyParameterData(Key, DataValue.GetData()))
 				{
-					ensureAlways(FMemory::Memcmp(ActualData, OverrideValue.GetBytes(), Key.GetSizeInBytes()) == 0);
+					ensureAlways(FMemory::Memcmp(DataValue.GetData(), OverrideValue.GetBytes(), Key.GetSizeInBytes()) == 0);
 				}
 			}
 		}
@@ -3203,8 +3207,11 @@ FNiagaraVariant UNiagaraComponent::GetCurrentParameterValue(const FNiagaraVariab
 		const FVector* Value = OverrideParameters.GetPositionParameterValue(RedirectedVar.GetName());
 		return Value == nullptr ? FNiagaraVariant() : FNiagaraVariant(Value, sizeof(FVector));
 	}
-	
-	return FNiagaraVariant(OverrideParameters.GetParameterData(ParameterOffset), InKey.GetSizeInBytes());
+
+	TArray<uint8> DataValue;
+	DataValue.AddUninitialized(UserVariable.GetSizeInBytes());
+	OverrideParameters.CopyParameterData(UserVariable, DataValue.GetData());
+	return FNiagaraVariant(DataValue);
 }
 
 void UNiagaraComponent::SetOverrideParameterStoreValue(const FNiagaraVariableBase& InKey, const FNiagaraVariant& InValue)
diff --git a/Engine/Plugins/FX/Niagara/Source/Niagara/Private/NiagaraModule.cpp b/Engine/Plugins/FX/Niagara/Source/Niagara/Private/NiagaraModule.cpp
--- a/Engine/Plugins/FX/Niagara/Source/Niagara/Private/NiagaraModule.cpp
+++ b/Engine/Plugins/FX/Niagara/Source/Niagara/Private/NiagaraModule.cpp
@@ -1034,6 +1034,7 @@ void FNiagaraTypeDefinition::RecreateUserDefinedTypeRegistry()
 					{
 						Flags |= ENiagaraTypeRegistryFlags::AllowPayload;
 					}
+					ScriptStruct = FNiagaraTypeHelper::FindNiagaraFriendlyTopLevelStruct(ScriptStruct, ENiagaraStructConversion::UserFacing);
 					FNiagaraTypeRegistry::Register(ScriptStruct, Flags);
 				}
 
@@ -1263,7 +1264,9 @@ void FNiagaraTypeDefinition::PostSerialize(const FArchive& Ar)
 	if (Ar.IsLoading() && ClassStructOrEnum != nullptr)
 	{
 		if (ClassStructOrEnum.GetClass()->IsChildOf(UScriptStruct::StaticClass()))
-			ClassStructOrEnum = FNiagaraTypeHelper::FindNiagaraFriendlyTopLevelStruct((UScriptStruct*)ClassStructOrEnum);
+		{
+			ClassStructOrEnum = FNiagaraTypeHelper::FindNiagaraFriendlyTopLevelStruct(static_cast<UScriptStruct*>(ClassStructOrEnum), ENiagaraStructConversion::UserFacing);
+		}
 	}
 #endif
 }
diff --git a/Engine/Plugins/FX/Niagara/Source/Niagara/Private/NiagaraParameterStore.cpp b/Engine/Plugins/FX/Niagara/Source/Niagara/Private/NiagaraParameterStore.cpp
--- a/Engine/Plugins/FX/Niagara/Source/Niagara/Private/NiagaraParameterStore.cpp
+++ b/Engine/Plugins/FX/Niagara/Source/Niagara/Private/NiagaraParameterStore.cpp
@@ -23,14 +23,6 @@ static FAutoConsoleVariableRef CVarNiagaraDumpParticleParameterStores(
 );
 #endif
 
-int32 GNiagaraAllowQuickSortedParameterOffetsCopy = 1;
-static FAutoConsoleVariableRef CVarNiagaraAllowQuickSortedParameterOffetsCopy(
-	TEXT("Niagara.AllowQuickSortedParameterOffsetsCopy"),
-	GNiagaraAllowQuickSortedParameterOffetsCopy,
-	TEXT("Whether to use memcpy to copy sortedparameteroffset arrays. (default=1)\n"),
-	ECVF_Scalability
-);
-
 //////////////////////////////////////////////////////////////////////////
 
 struct FNiagaraVariableSearch
@@ -131,17 +123,7 @@ void FNiagaraVariableWithOffset::PostSerialize(const FArchive& Ar)
 
 void FNiagaraParameterStore::CopySortedParameterOffsets(TArrayView<const FNiagaraVariableWithOffset> Src)
 {
-	if (GNiagaraAllowQuickSortedParameterOffetsCopy)
-	{
-		const int32 VariableCount = Src.Num();
-
-		SortedParameterOffsets.SetNumUninitialized(VariableCount);
-		FMemory::Memcpy(SortedParameterOffsets.GetData(), Src.GetData(), SortedParameterOffsets.GetTypeSize() * VariableCount);
-	}
-	else
-	{
-		SortedParameterOffsets = TArray<FNiagaraVariableWithOffset>(Src.GetData(), Src.Num());
-	}
+	SortedParameterOffsets = TArray<FNiagaraVariableWithOffset>(Src.GetData(), Src.Num());
 }
 
 //////////////////////////////////////////////////////////////////////////
@@ -279,7 +261,16 @@ void FNiagaraParameterStoreBinding::MatchParameters(FNiagaraParameterStore* Dest
 			const FNiagaraVariableWithOffset& SrcParamWithOffset = SrcParamWithOffsets[SrcIndex];
 			const FNiagaraVariableWithOffset& DestParamWithOffset = DestParamWithOffsets[DestIndex];
 
-			const int32 CompValue = FNiagaraVariableSearch::Compare(SrcParamWithOffset, DestParamWithOffset);
+			int32 CompValue;
+			if (FNiagaraTypeHelper::IsLWCType(SrcParamWithOffset.GetType()) || FNiagaraTypeHelper::IsLWCType(DestParamWithOffset.GetType()))
+			{
+				CompValue = FNiagaraVariableSearch::CompareIgnoreType(SrcParamWithOffset, DestParamWithOffset);
+			}
+			else
+			{
+				CompValue = FNiagaraVariableSearch::Compare(SrcParamWithOffset, DestParamWithOffset);
+			}
+			
 			if (CompValue < 0)
 			{
 				++SrcIndex;
@@ -290,7 +281,7 @@ void FNiagaraParameterStoreBinding::MatchParameters(FNiagaraParameterStore* Dest
 			}
 			else // CompValue == 0
 			{
-				Visitor(SrcParamWithOffset, SrcParamWithOffset.Offset, DestParamWithOffset.Offset);
+				Visitor(DestParamWithOffset, SrcParamWithOffset.Offset, DestParamWithOffset.Offset);
 				++SrcIndex;
 				++DestIndex;
 			}
@@ -304,13 +295,20 @@ void FNiagaraParameterStoreBinding::MatchParameters(FNiagaraParameterStore* Dest
 		for (const FNiagaraVariableWithOffset& ParamWithOffset : DestParamWithOffsets)
 		{
 			Visitor(ParamWithOffset, SrcStore->IndexOf(ParamWithOffset), ParamWithOffset.Offset);
+			if (const FNiagaraVariableWithOffset* Var = SrcStore->FindParameterVariable(ParamWithOffset, true))
+			{
+				Visitor(ParamWithOffset, Var->Offset, ParamWithOffset.Offset);
+			}
 		}
 	}
 	else
 	{
 		for (const FNiagaraVariableWithOffset& ParamWithOffset : SrcParamWithOffsets)
 		{
-			Visitor(ParamWithOffset, ParamWithOffset.Offset, DestStore->IndexOf(ParamWithOffset));
+			if (const FNiagaraVariableWithOffset* Var = DestStore->FindParameterVariable(ParamWithOffset, true))
+			{
+				Visitor(*Var, ParamWithOffset.Offset, Var->Offset);
+			}
 		}
 	}
 }
@@ -618,15 +616,22 @@ bool FNiagaraParameterStore::AddParameter(const FNiagaraVariable& Param, bool bI
 		}
 	}
 
-	int32& Offset = SortedParameterOffsets.EmplaceAt_GetRef(InsertPos, Param, (int32)INDEX_NONE).Offset;
+	FNiagaraTypeDefinition ParamType = Param.GetType();
+	FNiagaraLwcStructConverter StructConverter;
+	if (FNiagaraTypeHelper::IsLWCType(ParamType))
+	{
+		StructConverter = FNiagaraTypeRegistry::GetStructConverter(ParamType);
+	}
+
+	int32& Offset = SortedParameterOffsets.EmplaceAt_GetRef(InsertPos, Param, (int32)INDEX_NONE, StructConverter).Offset;
 
-	if (Param.GetType().IsDataInterface())
+	if (ParamType.IsDataInterface())
 	{
 		Offset = DataInterfaces.AddZeroed();
-		DataInterfaces[Offset] = bInitInterfaces ? NewObject<UNiagaraDataInterface>(Owner, const_cast<UClass*>(Param.GetType().GetClass()), NAME_None, RF_Transactional | RF_Public) : nullptr;
+		DataInterfaces[Offset] = bInitInterfaces ? NewObject<UNiagaraDataInterface>(Owner, const_cast<UClass*>(ParamType.GetClass()), NAME_None, RF_Transactional | RF_Public) : nullptr;
 		bInterfacesDirty = true;
 	}
-	else if (Param.GetType().IsUObject())
+	else if (ParamType.IsUObject())
 	{
 		Offset = UObjects.AddDefaulted();
 		bUObjectsDirty = true;
@@ -635,8 +640,10 @@ bool FNiagaraParameterStore::AddParameter(const FNiagaraVariable& Param, bool bI
 	{
 		DEC_MEMORY_STAT_BY(STAT_NiagaraParamStoreMemory, ParameterData.GetAllocatedSize());
 
-		int32 ParamSize = Param.GetSizeInBytes();
-		int32 ParamAlignment = Param.GetAlignment();
+		// the size here is too large for custom structs with lwc types, as we only write data for the swc type. But since the api allows direct access to the memory,
+		// we allocate a bit more than necessary to prevent access violations due to wrong access. Note that the script execution store actually uses the correct size for the allocated data.
+		int32 ParamSize = ParamType.GetSize(); 
+		//int32 ParamAlignment = Param.GetAlignment();
 		//int32 Offset = AlignArbitrary(ParameterData.Num(), ParamAlignment);//TODO: We need to handle alignment better here. Need to both satisfy CPU and GPU alignment concerns. VM doesn't care but the VM complier needs to be aware. Probably best to have everything adhere to GPU alignment rules.
 		Offset = ParameterData.Num();
 				
@@ -644,7 +651,14 @@ bool FNiagaraParameterStore::AddParameter(const FNiagaraVariable& Param, bool bI
 		if (Param.IsDataAllocated())
 		{
 			ParameterData.AddUninitialized(ParamSize);
-			FMemory::Memcpy(GetParameterData_Internal(Offset), Param.GetData(), ParamSize);
+			if (StructConverter.IsValid())
+			{
+				StructConverter.ConvertDataToSimulation(GetParameterData_Internal(Offset), Param.GetData());
+			}
+			else
+			{
+				FMemory::Memcpy(GetParameterData_Internal(Offset), Param.GetData(), ParamSize);	
+			}			
 		}
 		else
 		{
@@ -655,7 +669,7 @@ bool FNiagaraParameterStore::AddParameter(const FNiagaraVariable& Param, bool bI
 
 		bParametersDirty = true;
 
-		if (Param.GetType() == FNiagaraTypeDefinition::GetPositionDef() && !HasPositionData(Param.GetName()))
+		if (ParamType == FNiagaraTypeDefinition::GetPositionDef() && !HasPositionData(Param.GetName()))
 		{
 			SetPositionData(Param.GetName(), FVector::ZeroVector);
 			bPositionDataDirty = true;
@@ -710,20 +724,20 @@ bool FNiagaraParameterStore::RemoveParameter(const FNiagaraVariableBase& ToRemov
 				if (ExistingVar.GetType().IsDataInterface())
 				{
 					int32 Offset = NewInterfaces.AddZeroed();
-					NewOffsets.Add(FNiagaraVariableWithOffset(ExistingVar, Offset));
+					NewOffsets.Add(FNiagaraVariableWithOffset(ExistingVar, Offset, FNiagaraLwcStructConverter()));
 					NewInterfaces[Offset] = DataInterfaces[ExistingOffset];
 				}
 				else if (ExistingVar.IsUObject())
 				{
 					int32 Offset = NewUObjects.AddDefaulted();
-					NewOffsets.Add(FNiagaraVariableWithOffset(ExistingVar, Offset));
+					NewOffsets.Add(FNiagaraVariableWithOffset(ExistingVar, Offset, FNiagaraLwcStructConverter()));
 					NewUObjects[Offset] = UObjects[ExistingOffset];
 				}
 				else
 				{
 					int32 Offset = NewData.Num();
 					int32 ParamSize = ExistingVar.GetSizeInBytes();
-					NewOffsets.Add(FNiagaraVariableWithOffset(ExistingVar, Offset));
+					NewOffsets.Add(FNiagaraVariableWithOffset(ExistingVar, Offset, Existing.StructConverter));
 					NewData.AddUninitialized(ParamSize);
 					FMemory::Memcpy(NewData.GetData() + Offset, ParameterData.GetData() + ExistingOffset, ParamSize);
 				}
@@ -1032,6 +1046,20 @@ void FNiagaraParameterStore::InitFromSource(const FNiagaraParameterStore* SrcSto
 	}
 }
 
+const FNiagaraVariableWithOffset* FNiagaraParameterStore::FindParameterVariable(const FNiagaraVariable& Parameter, bool IgnoreType) const
+{
+	auto ParameterVariables = ReadParameterVariables();
+	if (ParameterVariables.Num())
+	{
+		int32 MatchingIndex = 0;
+		if (FNiagaraVariableSearch::Find(ParameterVariables.GetData(), Parameter, 0, ParameterVariables.Num(), IgnoreType, MatchingIndex))
+		{
+			return &ParameterVariables[MatchingIndex];
+		}
+	}
+	return nullptr;
+}
+
 void FNiagaraParameterStore::RemoveParameters(FNiagaraParameterStore& DestStore)
 {
 	for (const FNiagaraVariableWithOffset& Parameter : ReadParameterVariables())
@@ -1125,6 +1153,64 @@ void FNiagaraParameterStore::ResolvePositions(FNiagaraLWCConverter LwcConverter)
 	OnParameterChange();
 }
 
+bool FNiagaraParameterStore::SetParameterData(const uint8* Data, FNiagaraVariable Param, bool bAdd)
+{
+	checkSlow(Data != nullptr);
+	FNiagaraTypeDefinition SourceType = Param.GetType();
+	if (SourceType == FNiagaraTypeDefinition::GetPositionDef())
+	{
+		FNiagaraPosition Value = *reinterpret_cast<const FNiagaraPosition*>(Data);
+		return SetPositionParameterValue(Value, Param.GetName(), bAdd);
+	}
+	
+	int32 Offset = IndexOf(Param);
+	if (Offset != INDEX_NONE)
+	{
+		checkSlow(!Param.IsDataInterface());
+		uint8* Dest = GetParameterData_Internal(Offset);
+		if (Dest != Data)
+		{
+			FNiagaraLwcStructConverter StructConverter = GetStructConverter(Param);
+			if (StructConverter.IsValid())
+			{
+				StructConverter.ConvertDataToSimulation(Dest, Data);
+			}
+			else
+			{
+				FMemory::Memcpy(Dest, Data, Param.GetSizeInBytes());
+			}
+		}
+		OnParameterChange();
+		return true;
+	}
+	else
+	{
+		if (bAdd)
+		{
+			bool bInitInterfaces = false;
+			bool bTriggerRebind = false;
+			AddParameter(Param, bInitInterfaces, bTriggerRebind, &Offset);
+			check(Offset != INDEX_NONE);
+			uint8* Dest = GetParameterData_Internal(Offset);
+			if (Dest != Data)
+			{
+				FNiagaraLwcStructConverter StructConverter = GetStructConverter(Param);
+				if (StructConverter.IsValid())
+				{
+					StructConverter.ConvertDataToSimulation(Dest, Data);
+				}
+				else
+				{
+					FMemory::Memcpy(Dest, Data, Param.GetSizeInBytes());
+				}
+			}
+			OnLayoutChange();
+			return true;
+		}
+	}
+	return false;
+}
+
 void FNiagaraParameterStore::OnLayoutChange()
 {
 	const int32 ExpectedSlack = PaddedParameterSize(ParameterData.Num());
@@ -1157,6 +1243,38 @@ const FNiagaraVariableBase* FNiagaraParameterStore::FindVariable(const UNiagaraD
 	return nullptr;
 }
 
+bool FNiagaraParameterStore::CopyParameterData(const FNiagaraVariable& Parameter, uint8* DestinationData) const
+{
+	if (const FNiagaraVariableWithOffset* NiagaraVariableWithOffset = FindParameterVariable(Parameter))
+	{
+		const uint8* SourceData = GetParameterData(NiagaraVariableWithOffset->Offset);
+		if (NiagaraVariableWithOffset->StructConverter.IsValid())
+		{
+			NiagaraVariableWithOffset->StructConverter.ConvertDataFromSimulation(DestinationData, SourceData);
+		}
+		else
+		{
+			FMemory::Memcpy(DestinationData, SourceData, Parameter.GetSizeInBytes());
+		}
+		return true;
+	}
+	return false;
+}
+
+FNiagaraLwcStructConverter FNiagaraParameterStore::GetStructConverter(const FNiagaraVariable& Parameter) const
+{
+	auto ParameterVariables = ReadParameterVariables();
+	if (ParameterVariables.Num())
+	{
+		int32 MatchingIndex = 0;
+		if (FNiagaraVariableSearch::Find(ParameterVariables.GetData(), Parameter, 0, ParameterVariables.Num(), false, MatchingIndex))
+		{
+			return ParameterVariables[MatchingIndex].StructConverter;
+		}
+	}
+	return FNiagaraLwcStructConverter();
+}
+
 const int32* FNiagaraParameterStore::FindParameterOffset(const FNiagaraVariableBase& Parameter, bool IgnoreType) const
 {
 #if WITH_EDITORONLY_DATA
@@ -1187,7 +1305,7 @@ void FNiagaraParameterStore::PostLoad()
 	{
 		for (const TPair<FNiagaraVariable, int32>& ParamOffsetPair : ParameterOffsets)
 		{
-			SortedParameterOffsets.Emplace(ParamOffsetPair.Key, ParamOffsetPair.Value);
+			SortedParameterOffsets.Emplace(ParamOffsetPair.Key, ParamOffsetPair.Value, FNiagaraLwcStructConverter());
 		}
 		ParameterOffsets.Empty();
 	}
@@ -1249,5 +1367,4 @@ void FNiagaraParameterStore::RemoveAllOnChangedHandlers(const void* InUserObject
 	OnChangedDelegate.RemoveAll(InUserObject);
 }
 #endif
-
-//////////////////////////////////////////////////////////////////////////
\ No newline at end of file
+//////////////////////////////////////////////////////////////////////////
diff --git a/Engine/Plugins/FX/Niagara/Source/Niagara/Private/NiagaraRendererComponents.cpp b/Engine/Plugins/FX/Niagara/Source/Niagara/Private/NiagaraRendererComponents.cpp
--- a/Engine/Plugins/FX/Niagara/Source/Niagara/Private/NiagaraRendererComponents.cpp
+++ b/Engine/Plugins/FX/Niagara/Source/Niagara/Private/NiagaraRendererComponents.cpp
@@ -681,7 +681,7 @@ void FNiagaraRendererComponents::TickPropertyBindings(
 #if UE_LARGE_WORLD_COORDINATES_DISABLED
 		bool bForceStructConversion = false;
 #else
-		bool bForceStructConversion = PropertyType.GetScriptStruct() && !FNiagaraTypeHelper::IsNiagaraFriendlyTopLevelStruct(PropertyType.GetScriptStruct());
+		bool bForceStructConversion = PropertyType.GetScriptStruct() && !FNiagaraTypeHelper::IsNiagaraFriendlyTopLevelStruct(PropertyType.GetScriptStruct(), ENiagaraStructConversion::UserFacing);
 #endif
 		
 		if (PropertyType.IsValid() && DataVariable.GetType() != PropertyType && (!PropertySetter->bIgnoreConversion || bForceStructConversion))
diff --git a/Engine/Plugins/FX/Niagara/Source/Niagara/Private/NiagaraScript.cpp b/Engine/Plugins/FX/Niagara/Source/Niagara/Private/NiagaraScript.cpp
--- a/Engine/Plugins/FX/Niagara/Source/Niagara/Private/NiagaraScript.cpp
+++ b/Engine/Plugins/FX/Niagara/Source/Niagara/Private/NiagaraScript.cpp
@@ -1114,11 +1114,12 @@ void UNiagaraScript::ComputeVMCompilationId(FNiagaraVMExecutableDataId& Id, FGui
 					FString VarTypeName = Vars[i].GetType().GetName();
 					HashState.UpdateWithString(*VarName, VarName.Len());
 					HashState.UpdateWithString(*VarTypeName, VarTypeName.Len());
-					const uint8* VarData = RapidIterationParameters.GetParameterData(Vars[i]);
-					if (VarData)
+					TArray<uint8> DataValue;
+					DataValue.AddUninitialized(Vars[i].GetSizeInBytes());
+					if (RapidIterationParameters.CopyParameterData(Vars[i], DataValue.GetData()))
 					{
-						//UE_LOG(LogNiagara, Display, TEXT("Param %s %s %d"), *VarTypeName, *VarName, (uint32)VarData[0]);
-						HashState.Update(VarData, Vars[i].GetType().GetSize());
+						//UE_LOG(LogNiagara, Display, TEXT("Param %s %s %s"), *VarTypeName, *VarName, *ByteStr);
+						HashState.Update(DataValue.GetData(), Vars[i].GetType().GetSize());
 					}
 				}
 			}
diff --git a/Engine/Plugins/FX/Niagara/Source/Niagara/Private/NiagaraScriptExecutionParameterStore.cpp b/Engine/Plugins/FX/Niagara/Source/Niagara/Private/NiagaraScriptExecutionParameterStore.cpp
--- a/Engine/Plugins/FX/Niagara/Source/Niagara/Private/NiagaraScriptExecutionParameterStore.cpp
+++ b/Engine/Plugins/FX/Niagara/Source/Niagara/Private/NiagaraScriptExecutionParameterStore.cpp
@@ -164,7 +164,7 @@ uint32 FNiagaraScriptExecutionParameterStore::GenerateLayoutInfoInternal(TArray<
 			}
 			else if (FStructProperty* StructProp = CastFieldChecked<FStructProperty>(Property))
 			{
-				Struct = FNiagaraTypeHelper::FindNiagaraFriendlyTopLevelStruct(StructProp->Struct);
+				Struct = FNiagaraTypeHelper::FindNiagaraFriendlyTopLevelStruct(StructProp->Struct, ENiagaraStructConversion::Simulation);
 			}
 			else
 			{
diff --git a/Engine/Plugins/FX/Niagara/Source/Niagara/Private/NiagaraTypes.cpp b/Engine/Plugins/FX/Niagara/Source/Niagara/Private/NiagaraTypes.cpp
--- a/Engine/Plugins/FX/Niagara/Source/Niagara/Private/NiagaraTypes.cpp
+++ b/Engine/Plugins/FX/Niagara/Source/Niagara/Private/NiagaraTypes.cpp
@@ -2,6 +2,11 @@
 
 #include "NiagaraTypes.h"
 
+#include "NiagaraDataInterface.h"
+#include "NiagaraStats.h"
+#include "Misc/StringBuilder.h"
+#include "UObject/Class.h"
+
 void FNiagaraVariableBase::SetNamespacedName(const FString& InNamespace, FName InVariableName)
 {
 	TStringBuilder<128> NameBuilder;
@@ -47,3 +52,469 @@ FVector FNiagaraLWCConverter::ConvertSimulationVectorToWorld(FVector3f Simulatio
 {
 	return FVector(SimulationPosition) + SystemWorldPos;
 }
+
+FNiagaraStructConversionStep::FNiagaraStructConversionStep()
+{
+}
+
+FNiagaraStructConversionStep::FNiagaraStructConversionStep(int32 InSourceBytes, int32 InSourceOffset, int32 InSimulationBytes, int32 InSimulationOffset, ENiagaraStructConversionType InConversionType)
+	: SourceBytes(InSourceBytes), SourceOffset(InSourceOffset), SimulationBytes(InSimulationBytes), SimulationOffset(InSimulationOffset), ConversionType(InConversionType)
+{
+}
+
+void FNiagaraStructConversionStep::CopyToSim(uint8* DestinationData, const uint8* SourceData) const
+{
+	void* Dest = DestinationData + SimulationOffset;
+	const void* Src = SourceData + SourceOffset;
+	if (ConversionType == ENiagaraStructConversionType::DoubleToFloat)
+	{
+		checkf(SourceBytes == 8 && SimulationBytes == 4, TEXT("Wrong bytesizes for double->float conversion: Source %i, Simulation %i"), SourceBytes, SimulationBytes);
+		*static_cast<float*>(Dest) = *static_cast<const double*>(Src);
+	}
+	else if (ConversionType == ENiagaraStructConversionType::Vector2)
+	{
+		checkf(SourceBytes == 16 && SimulationBytes == 8, TEXT("Wrong bytesizes for Vector2 d->f conversion: Source %i, Simulation %i"), SourceBytes, SimulationBytes);
+		FVector2f Vector2 = *static_cast<const FVector2d*>(Src);
+		FMemory::Memcpy(Dest, &Vector2, SimulationBytes);
+	}
+	else if (ConversionType == ENiagaraStructConversionType::Vector3)
+	{
+		checkf(SourceBytes == 24 && SimulationBytes == 12, TEXT("Wrong bytesizes for Vector3 d->f conversion: Source %i, Simulation %i"), SourceBytes, SimulationBytes);
+		FVector3f Vector = *static_cast<const FVector3d*>(Src);
+		FMemory::Memcpy(Dest, &Vector, SimulationBytes);
+	}
+	else if (ConversionType == ENiagaraStructConversionType::Vector4)
+	{
+		checkf(SourceBytes == 32 && SimulationBytes == 16, TEXT("Wrong bytesizes for Vector4 d->f conversion: Source %i, Simulation %i"), SourceBytes, SimulationBytes);
+		FVector4f Vector4 = *static_cast<const FVector4d*>(Src);
+		FMemory::Memcpy(Dest, &Vector4, SimulationBytes);
+	}
+	else if (ConversionType == ENiagaraStructConversionType::Quat)
+	{
+		checkf(SourceBytes == 32 && SimulationBytes == 16, TEXT("Wrong bytesizes for Quat d->f conversion: Source %i, Simulation %i"), SourceBytes, SimulationBytes);
+		FQuat4d SrcQuat = *static_cast<const FQuat4d*>(Src);
+		FQuat4f TargetQuat(SrcQuat.X, SrcQuat.Y, SrcQuat.Z, SrcQuat.W);
+		FMemory::Memcpy(Dest, &TargetQuat, SimulationBytes);
+	}
+	else
+	{
+		checkf(SourceBytes == SimulationBytes, TEXT("SourceBytes != TargetBytes: %i != %i"), SourceBytes, SimulationBytes);
+		FMemory::Memcpy(Dest, Src, SourceBytes);
+	}
+}
+
+void FNiagaraStructConversionStep::CopyFromSim(uint8* DestinationData, const uint8* SourceData) const
+{
+	void* Dest = DestinationData + SourceOffset;
+	const void* Src = SourceData + SimulationOffset;
+	if (ConversionType == ENiagaraStructConversionType::DoubleToFloat)
+	{
+		checkf(SourceBytes == 8 && SimulationBytes == 4, TEXT("Wrong bytesizes for float->double conversion: Source %i, Simulation %i"), SourceBytes, SimulationBytes);
+		*static_cast<double*>(Dest) = *static_cast<const float*>(Src);
+	}
+	else if (ConversionType == ENiagaraStructConversionType::Vector2)
+	{
+		checkf(SourceBytes == 16 && SimulationBytes == 8, TEXT("Wrong bytesizes for Vector2 f->d conversion: Source %i, Simulation %i"), SourceBytes, SimulationBytes);
+		FVector2d Vector2 = *static_cast<const FVector2f*>(Src);
+		FMemory::Memcpy(Dest, &Vector2, SourceBytes);
+	}
+	else if (ConversionType == ENiagaraStructConversionType::Vector3)
+	{
+		checkf(SourceBytes == 24 && SimulationBytes == 12, TEXT("Wrong bytesizes for Vector3 f->d conversion: Source %i, Simulation %i"), SourceBytes, SimulationBytes);
+		FVector3d Vector3 = *static_cast<const FVector3f*>(Src);
+        FMemory::Memcpy(Dest, &Vector3, SourceBytes);
+	}
+	else if (ConversionType == ENiagaraStructConversionType::Vector4)
+	{
+		checkf(SourceBytes == 32 && SimulationBytes == 16, TEXT("Wrong bytesizes for Vector4 f->d conversion: Source %i, Simulation %i"), SourceBytes, SimulationBytes);
+		FVector4d Vector4 = *static_cast<const FVector4f*>(Src);
+		FMemory::Memcpy(Dest, &Vector4, SourceBytes);
+	}
+	else if (ConversionType == ENiagaraStructConversionType::Quat)
+	{
+		checkf(SourceBytes == 32 && SimulationBytes == 16, TEXT("Wrong bytesizes for Quat f->d conversion: Source %i, Simulation %i"), SourceBytes, SimulationBytes);
+		FQuat4f SrcQuat = *static_cast<const FQuat4f*>(Src);
+		FQuat4d TargetQuat(SrcQuat.X, SrcQuat.Y, SrcQuat.Z, SrcQuat.W);
+		FMemory::Memcpy(Dest, &TargetQuat, SourceBytes);
+	}
+	else
+	{
+		checkf(SourceBytes == SimulationBytes, TEXT("SourceBytes != TargetBytes: %i != %i"), SourceBytes, SimulationBytes);
+		FMemory::Memcpy(DestinationData, SourceData, SourceBytes);
+	}
+}
+
+void FNiagaraLwcStructConverter::ConvertDataToSimulation(uint8* DestinationData, const uint8* SourceData) const
+{
+	for (const FNiagaraStructConversionStep& ConversionStep : ConversionSteps)
+	{
+		ConversionStep.CopyToSim(DestinationData, SourceData);
+	}
+}
+
+void FNiagaraLwcStructConverter::ConvertDataFromSimulation(uint8* DestinationData, const uint8* SourceData) const
+{
+	for (const FNiagaraStructConversionStep& ConversionStep : ConversionSteps)
+	{
+		ConversionStep.CopyFromSim(DestinationData, SourceData);
+	}
+}
+
+void FNiagaraLwcStructConverter::AddConversionStep(int32 InSourceBytes, int32 InSourceOffset, int32 InSimulationBytes, int32 InSimulationOffset, ENiagaraStructConversionType ConversionType)
+{
+	check(InSourceBytes >= 0);
+	check(InSourceOffset >= 0);
+	check(InSimulationBytes >= 0);
+	check(InSimulationOffset >= 0);
+	// TODO (mg) as an optimization, if the last step was a copy step and this one is as well, they can be merged into one
+	ConversionSteps.Emplace(InSourceBytes, InSourceOffset, InSimulationBytes, InSimulationOffset, ConversionType);
+}
+
+//////////////////////////////////////////////////////////////////////////
+
+FString FNiagaraTypeHelper::ToString(const uint8* ValueData, const UObject* StructOrEnum)
+{
+	FString Ret;
+	if (const UEnum* Enum = Cast<const UEnum>(StructOrEnum))
+	{
+		Ret = Enum->GetNameStringByValue(*(int32*)ValueData);
+	}
+	else if (const UScriptStruct* Struct = Cast<const UScriptStruct>(StructOrEnum))
+	{
+		if (Struct == FNiagaraTypeDefinition::GetFloatStruct())
+		{
+			Ret += FString::Printf(TEXT("%g "), *(float*)ValueData);
+		}
+		else if (Struct == FNiagaraTypeDefinition::GetIntStruct())
+		{
+			Ret += FString::Printf(TEXT("%d "), *(int32*)ValueData);
+		}
+		else if (Struct == FNiagaraTypeDefinition::GetBoolStruct())
+		{
+			int32 Val = *(int32*)ValueData;
+			Ret += Val == 0xFFFFFFFF ? (TEXT("True")) : (Val == 0x0 ? TEXT("False") : TEXT("Invalid"));
+		}
+		else
+		{
+			for (TFieldIterator<FProperty> PropertyIt(Struct, EFieldIteratorFlags::IncludeSuper); PropertyIt; ++PropertyIt)
+			{
+				const FProperty* Property = *PropertyIt;
+				const uint8* PropPtr = ValueData + PropertyIt->GetOffset_ForInternal();
+				if (Property->IsA(FFloatProperty::StaticClass()))
+				{
+					Ret += FString::Printf(TEXT("%s: %g "), *Property->GetNameCPP(), *(float*)PropPtr);
+				}
+				else if (Property->IsA(FDoubleProperty::StaticClass()))
+				{
+					Ret += FString::Printf(TEXT("%s: %g "), *Property->GetNameCPP(), *(double*)PropPtr);
+				}
+				else if (Property->IsA(FUInt16Property::StaticClass()))
+				{
+					FFloat16 Val = *(FFloat16*)PropPtr;
+					Ret += FString::Printf(TEXT("%s: %f "), *Property->GetNameCPP(), Val.GetFloat());
+				}
+				else if (Property->IsA(FIntProperty::StaticClass()))
+				{
+					Ret += FString::Printf(TEXT("%s: %d "), *Property->GetNameCPP(), *(int32*)PropPtr);
+				}
+				else if (Property->IsA(FBoolProperty::StaticClass()))
+				{
+					int32 Val = *(int32*)ValueData;
+					FString BoolStr = Val == 0xFFFFFFFF ? (TEXT("True")) : (Val == 0x0 ? TEXT("False") : TEXT("Invalid"));
+					Ret += FString::Printf(TEXT("%s: %d "), *Property->GetNameCPP(), *BoolStr);
+				}
+				else if (const FStructProperty* StructProp = CastFieldChecked<const FStructProperty>(Property))
+				{
+					Ret += FString::Printf(TEXT("%s: (%s) "), *Property->GetNameCPP(), *FNiagaraTypeHelper::ToString(PropPtr, StructProp->Struct));
+				}
+				else
+				{
+					check(false);
+					Ret += TEXT("Unknown Type");
+				}
+			}
+		}
+	}
+	return Ret;
+}
+
+FNiagaraLwcStructConverter BuildSWCStructure(UScriptStruct* NewStruct, UScriptStruct* InStruct)
+{
+	static UPackage* CoreUObjectPkg = FindObjectChecked<UPackage>(nullptr, TEXT("/Script/CoreUObject"));
+	static UScriptStruct* Vector2fStruct = FindObjectChecked<UScriptStruct>(CoreUObjectPkg, TEXT("Vector2f"));
+	static UScriptStruct* Vector3fStruct = FindObjectChecked<UScriptStruct>(CoreUObjectPkg, TEXT("Vector3f"));
+	static UScriptStruct* Vector4fStruct = FindObjectChecked<UScriptStruct>(CoreUObjectPkg, TEXT("Vector4f"));
+	static UScriptStruct* Quat4fStruct = FindObjectChecked<UScriptStruct>(CoreUObjectPkg, TEXT("Quat4f"));
+
+	FNiagaraLwcStructConverter StructConverter;
+	int32 AlignedOffset = 0;
+	FProperty* LastProperty = nullptr;
+	for (const FField* ChildProperty = InStruct->ChildProperties; ChildProperty; ChildProperty = ChildProperty->Next)
+	{
+		FProperty* NewProperty = nullptr;
+		if (const FDoubleProperty* ChildAsDouble = CastField<const FDoubleProperty>(ChildProperty))
+		{
+			AlignedOffset = Align(AlignedOffset, NewProperty->GetMinAlignment());
+			NewProperty = new FFloatProperty(NewStruct, ChildProperty->GetFName(), RF_Public);
+			StructConverter.AddConversionStep(ChildAsDouble->GetSize(), ChildAsDouble->GetOffset_ForInternal(), NewProperty->GetSize(), AlignedOffset, ENiagaraStructConversionType::DoubleToFloat);
+		}
+		else if (const FStructProperty* ChildAsStruct = CastField<const FStructProperty>(ChildProperty))
+		{
+			FStructProperty* NewStructProperty = new FStructProperty(NewStruct, ChildProperty->GetFName(), RF_Public);
+			NewProperty = NewStructProperty;
+
+			FName StructName = ChildAsStruct->Struct->GetFName();
+			int32 OldPropertySize = ChildAsStruct->GetSize();
+			int32 OldPropertyOffset = ChildAsStruct->GetOffset_ForInternal();
+			if ((StructName == NAME_Vector2d) || (StructName == NAME_Vector2D))
+			{
+				NewStructProperty->Struct = Vector2fStruct;
+				NewStructProperty->ElementSize = Vector2fStruct->GetStructureSize();
+				AlignedOffset = Align(AlignedOffset, Vector2fStruct->GetMinAlignment());
+				StructConverter.AddConversionStep(OldPropertySize, OldPropertyOffset, NewStructProperty->GetSize(), AlignedOffset, ENiagaraStructConversionType::Vector2);
+			}
+			else if ((StructName == NAME_Vector3d) || (StructName == NAME_Vector))
+			{
+				NewStructProperty->Struct = Vector3fStruct;
+				NewStructProperty->ElementSize = Vector3fStruct->GetStructureSize();
+				AlignedOffset = Align(AlignedOffset, Vector3fStruct->GetMinAlignment());
+				StructConverter.AddConversionStep(OldPropertySize, OldPropertyOffset, NewStructProperty->GetSize(), AlignedOffset, ENiagaraStructConversionType::Vector3);
+			}
+			else if ((StructName == NAME_Vector4d) || (StructName == NAME_Vector4))
+			{
+				NewStructProperty->Struct = Vector4fStruct;
+				NewStructProperty->ElementSize = Vector4fStruct->GetStructureSize();
+				AlignedOffset = Align(AlignedOffset, Vector4fStruct->GetMinAlignment());
+				StructConverter.AddConversionStep(OldPropertySize, OldPropertyOffset, NewStructProperty->GetSize(), AlignedOffset, ENiagaraStructConversionType::Vector4);
+			}
+			else if ((StructName == NAME_Quat4d) || (StructName == NAME_Quat))
+			{
+				NewStructProperty->Struct = Quat4fStruct;
+				NewStructProperty->ElementSize = Quat4fStruct->GetStructureSize();
+				AlignedOffset = Align(AlignedOffset, Quat4fStruct->GetMinAlignment());
+				StructConverter.AddConversionStep(OldPropertySize, OldPropertyOffset, NewStructProperty->GetSize(), AlignedOffset, ENiagaraStructConversionType::Quat);
+			}
+			else
+			{
+				NewStructProperty->Struct = ChildAsStruct->Struct;
+				NewStructProperty->ElementSize = ChildAsStruct->Struct->GetStructureSize();
+				AlignedOffset = Align(AlignedOffset, ChildAsStruct->Struct->GetMinAlignment());
+				StructConverter.AddConversionStep(OldPropertySize, OldPropertyOffset, NewStructProperty->GetSize(), AlignedOffset, ENiagaraStructConversionType::CopyOnly);
+			}
+		}
+		else if (
+			ChildProperty->IsA(FFloatProperty::StaticClass()) ||
+			ChildProperty->IsA(FIntProperty::StaticClass()) ||
+			ChildProperty->IsA(FUInt32Property::StaticClass()) ||
+			ChildProperty->IsA(FEnumProperty::StaticClass()) ||
+			ChildProperty->IsA(FByteProperty::StaticClass()) ||
+			ChildProperty->IsA(FBoolProperty::StaticClass()))
+		{
+			const FProperty* OldProperty = CastField<const FProperty>(ChildProperty);
+			NewProperty = CastField<FProperty>(FField::Duplicate(ChildProperty, NewStruct, ChildProperty->GetFName()));
+			AlignedOffset = Align(AlignedOffset, OldProperty->GetMinAlignment());
+			StructConverter.AddConversionStep(NewProperty->GetSize(), OldProperty->GetOffset_ForInternal(), NewProperty->GetSize(), AlignedOffset, ENiagaraStructConversionType::CopyOnly);
+		}
+
+		if ( NewProperty == nullptr )
+		{
+			UE_LOG(LogNiagara, Fatal, TEXT("Unsupported property '%s' in custom struct '%s'"), *InStruct->GetName(), *ChildProperty->GetName());
+		}
+
+#if WITH_EDITORONLY_DATA
+		FField::CopyMetaData(ChildProperty, NewProperty);
+#endif
+		if (LastProperty)
+		{
+			LastProperty->Next = NewProperty;
+			LastProperty->PropertyLinkNext = NewProperty;
+		}
+		else
+		{
+			NewStruct->PropertyLink = NewProperty;
+			NewStruct->ChildProperties = NewProperty;
+		}
+		LastProperty = NewProperty;
+		AlignedOffset += NewProperty->GetSize();
+	}
+
+	return StructConverter;
+}
+
+bool FNiagaraTypeHelper::IsLWCStructure(UStruct* InStruct)
+{
+	for (const FField* ChildProperty = InStruct->ChildProperties; ChildProperty; ChildProperty = ChildProperty->Next)
+	{
+		if (ChildProperty->IsA<FDoubleProperty>())
+		{
+			return true;
+		}
+		if ( const FStructProperty* ChildAsStruct = CastField<const FStructProperty>(ChildProperty) )
+		{
+			FName StructName = ChildAsStruct->Struct->GetFName();
+			if ( (StructName == NAME_Vector2d) || (StructName == NAME_Vector2D) ||
+				 (StructName == NAME_Vector3d) || (StructName == NAME_Vector) ||
+				 (StructName == NAME_Vector4d) || (StructName == NAME_Vector4) ||
+				 (StructName == NAME_Quat4d) || (StructName == NAME_Quat) )
+			{
+				return true;
+			}
+		}
+	}
+	return false;
+}
+
+bool FNiagaraTypeHelper::IsLWCType(const FNiagaraTypeDefinition& InType)
+{
+	return InType.IsValid() && IsLWCStructure(InType.GetStruct());
+}
+
+UScriptStruct* FNiagaraTypeHelper::GetSWCStruct(UScriptStruct* LWCStruct)
+{
+	struct FRemapEntry
+	{
+		UScriptStruct* Get(UScriptStruct* InStruct) const
+		{
+		#if WITH_EDITORONLY_DATA
+			return SerialNumber == InStruct->FieldPathSerialNumber ? Struct.Get() : nullptr;
+		#else
+			return Struct.Get();
+		#endif
+		}
+
+		TWeakObjectPtr<UScriptStruct>	Struct;
+	#if WITH_EDITORONLY_DATA
+		int32 SerialNumber = 0;
+	#endif
+	};
+
+	static FRWLock RemapTableLock;
+	static TMap<TWeakObjectPtr<UScriptStruct>, FRemapEntry> RemapTable;
+
+	bool bRemoveDeadRemaps = false;
+
+	// Attempt to find existing struct
+	UScriptStruct* SWCStruct = nullptr;
+	{
+		FReadScopeLock ReadLock(RemapTableLock);
+		if ( FRemapEntry* RemapEntry = RemapTable.Find(LWCStruct) )
+		{
+			SWCStruct = RemapEntry->Get(LWCStruct);
+			bRemoveDeadRemaps = true;
+		}
+	}
+
+	// No SWCStruct or it become invalid we need to create one
+	if (SWCStruct == nullptr)
+	{
+		// Enable write lock and search again as between the two locks things may have changed
+		FWriteScopeLock WriteLock(RemapTableLock);
+		if (FRemapEntry* RemapEntry = RemapTable.Find(LWCStruct))
+		{
+			SWCStruct = RemapEntry->Get(LWCStruct);
+			bRemoveDeadRemaps = true;
+		}
+
+		if (SWCStruct == nullptr)
+		{
+			QUICK_SCOPE_CYCLE_COUNTER(STAT_Niagara_GetSWCStruct);
+
+			// Do we need to remove dead remaps from the list?  
+			//-OPT: We should hook into post GC and prune the map at that point
+			bRemoveDeadRemaps = true;
+			if ( bRemoveDeadRemaps )
+			{
+				for (auto it = RemapTable.CreateIterator(); it; ++it)
+				{
+					UScriptStruct* SrcScript = it->Key.Get();
+					if (SrcScript == nullptr || it->Value.Get(SrcScript) == nullptr)
+					{
+						UScriptStruct* DstStruct = it->Value.Struct.Get();
+						if (DstStruct && SrcScript != DstStruct)
+						{
+							DstStruct->RemoveFromRoot();
+						}
+						it.RemoveCurrent();
+					}
+				}
+			}
+
+			// If this is a LWC structure we need to built a new structure now
+			if (IsLWCStructure(LWCStruct))
+			{
+				SWCStruct = NewObject<UScriptStruct>(GetTransientPackage(), FName(LWCStruct->GetName() + TEXT("_SWC")), RF_NoFlags);
+				FNiagaraLwcStructConverter StructConverter = BuildSWCStructure(SWCStruct, LWCStruct);
+				FNiagaraTypeRegistry::RegisterStructConverter(FNiagaraTypeDefinition(LWCStruct, FNiagaraTypeDefinition::EAllowUnfriendlyStruct::Allow), StructConverter);
+				SWCStruct->Bind();
+				SWCStruct->PrepareCppStructOps();
+				SWCStruct->StaticLink(true);
+				SWCStruct->AddToRoot();
+			}
+			else
+			{
+				SWCStruct = LWCStruct;
+			}
+			FRemapEntry& RemapEntry = RemapTable.Add(LWCStruct);
+			RemapEntry.Struct = SWCStruct;
+		#if WITH_EDITORONLY_DATA
+			RemapEntry.SerialNumber = LWCStruct->FieldPathSerialNumber;
+		#endif
+		}
+	}
+	return SWCStruct;
+}
+
+UScriptStruct* FNiagaraTypeHelper::FindNiagaraFriendlyTopLevelStruct(UScriptStruct* InStruct, ENiagaraStructConversion StructConversion)
+{
+	if (!InStruct)
+	{
+		return nullptr;
+	}
+
+	// Note: UE core types are converted to the float variant as Niagara only works with float types.
+	if (InStruct->GetFName() == NAME_Vector2D || InStruct->GetFName() == NAME_Vector2d) // LWC support
+	{
+		return FNiagaraTypeDefinition::GetVec2Struct();
+	}
+
+	if (InStruct->GetFName() == NAME_Vector || InStruct->GetFName() == NAME_Vector3d) // LWC support
+	{
+		return FNiagaraTypeDefinition::GetVec3Struct();
+	}
+
+	if (InStruct->GetFName() == NAME_Vector4 || InStruct->GetFName() == NAME_Vector4d) // LWC support
+	{
+		return FNiagaraTypeDefinition::GetVec4Struct();
+	}
+
+	if (InStruct->GetFName() == NAME_Quat || InStruct->GetFName() == NAME_Quat4d) // LWC support
+	{
+		return FNiagaraTypeDefinition::GetQuatStruct();
+	}
+
+	if (InStruct->GetFName() == FName("NiagaraPosition"))
+	{
+		return FNiagaraTypeDefinition::GetPositionStruct();
+	}
+
+	if (InStruct->IsA<UNiagaraDataInterface>())
+	{
+		return InStruct;
+	}
+
+	if (StructConversion == ENiagaraStructConversion::UserFacing)
+	{
+		return InStruct;
+	}
+	check(StructConversion == ENiagaraStructConversion::Simulation);
+
+	//-OPT: We could potentially add a list of known types here, it may or may not be faster
+	return GetSWCStruct(InStruct);
+}
+
+bool FNiagaraTypeHelper::IsNiagaraFriendlyTopLevelStruct(UScriptStruct* InStruct, ENiagaraStructConversion StructConversion)
+{
+	return FindNiagaraFriendlyTopLevelStruct(InStruct, StructConversion) == InStruct;
+}
+
+/////////////////////////////////////////////////////////////////////////////
diff --git a/Engine/Plugins/FX/Niagara/Source/Niagara/Public/NiagaraParameterStore.h b/Engine/Plugins/FX/Niagara/Source/Niagara/Public/NiagaraParameterStore.h
--- a/Engine/Plugins/FX/Niagara/Source/Niagara/Public/NiagaraParameterStore.h
+++ b/Engine/Plugins/FX/Niagara/Source/Niagara/Public/NiagaraParameterStore.h
@@ -134,8 +134,8 @@ struct FNiagaraVariableWithOffset : public FNiagaraVariableBase
 
 	// Those constructor enforce that there are no data allocated.
 	FORCEINLINE FNiagaraVariableWithOffset() : Offset(INDEX_NONE) {}
-	FORCEINLINE FNiagaraVariableWithOffset(const FNiagaraVariableWithOffset& InRef) : FNiagaraVariableBase(InRef.GetType(), InRef.GetName()), Offset(InRef.Offset) {}
-	FORCEINLINE FNiagaraVariableWithOffset(const FNiagaraVariableBase& InVariable, int32 InOffset) : FNiagaraVariableBase(InVariable.GetType(), InVariable.GetName()), Offset(InOffset) {}
+	FORCEINLINE FNiagaraVariableWithOffset(const FNiagaraVariableWithOffset& InRef) : FNiagaraVariableBase(InRef.GetType(), InRef.GetName()), Offset(InRef.Offset), StructConverter(InRef.StructConverter) {}
+	FORCEINLINE FNiagaraVariableWithOffset(const FNiagaraVariableBase& InVariable, int32 InOffset, const FNiagaraLwcStructConverter& InStructConverter) : FNiagaraVariableBase(InVariable.GetType(), InVariable.GetName()), Offset(InOffset), StructConverter(InStructConverter) {}
 
 	bool Serialize(FArchive& Ar);
 #if WITH_EDITORONLY_DATA
@@ -144,6 +144,9 @@ struct FNiagaraVariableWithOffset : public FNiagaraVariableBase
 
 	UPROPERTY()
 	int32 Offset;
+
+	UPROPERTY()
+	FNiagaraLwcStructConverter StructConverter;
 };
 
 template<>
@@ -352,6 +355,8 @@ public:
 		const int32* Off = FindParameterOffset(Parameter);
 		return Off ? *Off : (int32)INDEX_NONE;
 	}
+	
+	const FNiagaraVariableWithOffset* FindParameterVariable(const FNiagaraVariable& Parameter, bool IgnoreType = false) const;
 
 	/** Gets the typed parameter data. */
 	template<typename T>
@@ -393,6 +398,12 @@ public:
 		return nullptr;
 	}
 
+	/** Copies the data stored for the given variable to the target pointer. This method automatically converts custom struct data back to lwc types.
+	 * Just using the raw parameter store pointer via IndexOf() will be wrong if the target struct contains a lwc type like FVector.
+	 * Returns true if the data was copied, false if the parameter could not be found in the store.
+	 */
+	bool CopyParameterData(const FNiagaraVariable& Parameter, uint8* DestinationData) const;
+
 	/** Returns the data interface at the passed offset. */
 	FORCEINLINE UNiagaraDataInterface* GetDataInterface(int32 Offset)const
 	{
@@ -413,6 +424,9 @@ public:
 		return Interface;
 	}
 
+	/** Returns a struct converter for the given variable, if the store contains the variable and it's a LWC type. */
+	FNiagaraLwcStructConverter GetStructConverter(const FNiagaraVariable& Parameter) const;
+
 	/** Returns the associated FNiagaraVariable for the passed data interface if it exists in the store. Null if not.*/
 	const FNiagaraVariableBase* FindVariable(const UNiagaraDataInterface* Interface) const;
 
@@ -511,6 +525,7 @@ public:
 	{
 		check(Param.GetSizeInBytes() == sizeof(T));
 #if WITH_EDITOR
+		ensure(FNiagaraTypeHelper::IsLWCType(Param.GetType()) == false);
 		if (Param.GetType() == FNiagaraTypeDefinition::GetPositionDef())
 		{
 			check(HasPositionData(Param.GetName()));
@@ -557,46 +572,7 @@ public:
 		OnParameterChange();
 	}
 
-	FORCEINLINE_DEBUGGABLE bool SetParameterData(const uint8* Data, const FNiagaraVariable& Param, bool bAdd = false)
-	{
-		checkSlow(Data != nullptr);
-		if (Param.GetType() == FNiagaraTypeDefinition::GetPositionDef())
-		{
-			FNiagaraPosition Value = *reinterpret_cast<const FNiagaraPosition*>(Data);
-			return SetPositionParameterValue(Value, Param.GetName(), bAdd);
-		}
-		
-		int32 Offset = IndexOf(Param);
-		if (Offset != INDEX_NONE)
-		{
-			checkSlow(!Param.IsDataInterface());
-			uint8* Dest = GetParameterData_Internal(Offset);
-			if (Dest != Data)
-			{
-				FMemory::Memcpy(Dest, Data, Param.GetSizeInBytes());
-			}
-			OnParameterChange();
-			return true;
-		}
-		else
-		{
-			if (bAdd)
-			{
-				bool bInitInterfaces = false;
-				bool bTriggerRebind = false;
-				AddParameter(Param, bInitInterfaces, bTriggerRebind, &Offset);
-				check(Offset != INDEX_NONE);
-				uint8* Dest = GetParameterData_Internal(Offset);
-				if (Dest != Data)
-				{
-					FMemory::Memcpy(Dest, Data, Param.GetSizeInBytes());
-				}
-				OnLayoutChange();
-				return true;
-			}
-		}
-		return false;
-	}
+	bool SetParameterData(const uint8* Data, FNiagaraVariable Param, bool bAdd = false);
 
 	FORCEINLINE_DEBUGGABLE void SetDataInterface(UNiagaraDataInterface* InInterface, int32 Offset)
 	{
diff --git a/Engine/Plugins/FX/Niagara/Source/Niagara/Public/NiagaraScriptExecutionParameterStore.h b/Engine/Plugins/FX/Niagara/Source/Niagara/Public/NiagaraScriptExecutionParameterStore.h
--- a/Engine/Plugins/FX/Niagara/Source/Niagara/Public/NiagaraScriptExecutionParameterStore.h
+++ b/Engine/Plugins/FX/Niagara/Source/Niagara/Public/NiagaraScriptExecutionParameterStore.h
@@ -3,7 +3,6 @@
 #pragma once
 
 #include "NiagaraCommon.h"
-#include "NiagaraDataInterface.h"
 #include "NiagaraParameterStore.h"
 #include "NiagaraScriptExecutionParameterStore.generated.h"
 
@@ -53,7 +52,23 @@ public:
 	{
 #if WITH_EDITORONLY_DATA
 		int32 NewParamOffset = INDEX_NONE;
-		const bool bAdded = FNiagaraParameterStore::AddParameter(Param, bInitInterfaces, bTriggerRebind, &NewParamOffset);
+		bool bAdded;
+		if (FNiagaraTypeHelper::IsLWCType(Param.GetType()))
+		{
+			// Custom structs containing LWC type data are converted into SWC structs for the runtime 
+			UScriptStruct* Struct = FNiagaraTypeHelper::FindNiagaraFriendlyTopLevelStruct(Param.GetType().GetScriptStruct(), ENiagaraStructConversion::Simulation);
+			FNiagaraVariable SimParam(FNiagaraTypeDefinition(Struct), Param.GetName());
+			if (Param.IsDataAllocated())
+			{
+				SimParam.AllocateData();
+				FNiagaraTypeRegistry::GetStructConverter(Param.GetType()).ConvertDataToSimulation(SimParam.GetData(), Param.GetData());
+			}
+			bAdded = FNiagaraParameterStore::AddParameter(SimParam, bInitInterfaces, bTriggerRebind, &NewParamOffset);
+		}
+		else
+		{
+			bAdded = FNiagaraParameterStore::AddParameter(Param, bInitInterfaces, bTriggerRebind, &NewParamOffset);
+		}
 		if (bAdded)
 		{
 			AddPaddedParamSize(Param.GetType(), NewParamOffset);
diff --git a/Engine/Plugins/FX/Niagara/Source/Niagara/Public/NiagaraTypes.h b/Engine/Plugins/FX/Niagara/Source/Niagara/Public/NiagaraTypes.h
--- a/Engine/Plugins/FX/Niagara/Source/Niagara/Public/NiagaraTypes.h
+++ b/Engine/Plugins/FX/Niagara/Source/Niagara/Public/NiagaraTypes.h
@@ -11,6 +11,7 @@
 
 #include "NiagaraTypes.generated.h"
 
+struct FNiagaraTypeDefinition;
 class UNiagaraDataInterfaceBase;
 
 DECLARE_LOG_CATEGORY_EXTERN(LogNiagara, Log, Verbose);
@@ -70,7 +71,7 @@ struct FNiagaraBool
 
 private:
 	UPROPERTY(EditAnywhere, Category = Parameters)// Must be either FNiagaraBool::True or FNiagaraBool::False.
-	int32 Value = FNiagaraBool::False;
+	int32 Value = False;
 };
 
 USTRUCT(meta = (DisplayName = "Position"))
@@ -165,16 +166,16 @@ struct FNiagaraMatrix
 	GENERATED_USTRUCT_BODY()
 
 	UPROPERTY(EditAnywhere, Category=NiagaraMatrix)
-	FVector4f Row0 = FVector4f(EForceInit::ForceInitToZero);
+	FVector4f Row0 = FVector4f(ForceInitToZero);
 
 	UPROPERTY(EditAnywhere, Category = NiagaraMatrix)
-	FVector4f Row1 = FVector4f(EForceInit::ForceInitToZero);
+	FVector4f Row1 = FVector4f(ForceInitToZero);
 
 	UPROPERTY(EditAnywhere, Category = NiagaraMatrix)
-	FVector4f Row2 = FVector4f(EForceInit::ForceInitToZero);
+	FVector4f Row2 = FVector4f(ForceInitToZero);
 
 	UPROPERTY(EditAnywhere, Category = NiagaraMatrix)
-	FVector4f Row3 = FVector4f(EForceInit::ForceInitToZero);
+	FVector4f Row3 = FVector4f(ForceInitToZero);
 };
 
 USTRUCT()
@@ -218,6 +219,62 @@ private:
 	FVector SystemWorldPos;
 };
 
+UENUM()
+enum class ENiagaraStructConversionType : uint8
+{
+	CopyOnly, // no conversion, just copy the data
+	
+	DoubleToFloat,
+	Vector2,
+	Vector3,
+	Vector4,
+	Quat,
+};
+
+USTRUCT()
+struct FNiagaraStructConversionStep
+{
+	GENERATED_USTRUCT_BODY();
+	
+	UPROPERTY()
+	int32 SourceBytes = 0;
+
+	UPROPERTY()
+	int32 SourceOffset = 0;
+	
+	UPROPERTY()
+	int32 SimulationBytes = 0;
+
+	UPROPERTY()
+	int32 SimulationOffset = 0;
+	
+	UPROPERTY()
+	ENiagaraStructConversionType ConversionType = ENiagaraStructConversionType::CopyOnly;
+
+	FNiagaraStructConversionStep();
+	FNiagaraStructConversionStep(int32 InSourceBytes, int32 InSourceOffset, int32 InSimulationBytes, int32 InSimulationOffset, ENiagaraStructConversionType InConversionType);
+
+	void CopyToSim(uint8* DestinationData, const uint8* SourceData) const;
+	void CopyFromSim(uint8* DestinationData, const uint8* SourceData) const;
+};
+
+/** Can convert struct data from custom structs containing LWC data such as FVector3d into struct data suitable for Niagara simulations and vice versa. */
+USTRUCT()
+struct NIAGARA_API FNiagaraLwcStructConverter
+{
+	GENERATED_USTRUCT_BODY();
+	
+	void ConvertDataToSimulation(uint8* DestinationData, const uint8* SourceData) const;
+	void ConvertDataFromSimulation(uint8* DestinationData, const uint8* SourceData) const;
+	
+	void AddConversionStep(int32 InSourceBytes, int32 InSourceOffset, int32 InSimulationBytes, int32 InSimulationOffset, ENiagaraStructConversionType ConversionType);
+	bool IsValid() const { return ConversionSteps.Num() > 0; }
+	
+private:
+	UPROPERTY()
+	TArray<FNiagaraStructConversionStep> ConversionSteps;
+};
+
 FORCEINLINE uint32 GetTypeHash(const FNiagaraAssetVersion& Version)
 {
 	return HashCombine(GetTypeHash(Version.MajorVersion), GetTypeHash(Version.MinorVersion));
@@ -296,6 +353,15 @@ FORCEINLINE uint32 GetTypeHash(const FNiagaraID& ID)
 	return HashCombine(GetTypeHash(ID.Index), GetTypeHash(ID.AcquireTag));
 }
 
+enum class ENiagaraStructConversion : uint8
+{
+	/** Do not modify struct data members, even if they are not compatible with the Niagara VM, as the struct is user facing (or from an external api). */
+	UserFacing,
+	
+	/** Convert struct members that are not compatible with the simulation (e.g. lwc types) into compatible types and return a simulation-friendly struct for the VM. */
+	Simulation,
+};
+
 /*
 *  Can convert a UStruct with fields of base types only (float, int... - will likely add native vector types here as well)
 *	to an FNiagaraTypeDefinition (internal representation)
@@ -304,8 +370,11 @@ class NIAGARA_API FNiagaraTypeHelper
 {
 public:
 	static FString ToString(const uint8* ValueData, const UObject* StructOrEnum);
-	static UScriptStruct* FindNiagaraFriendlyTopLevelStruct(UScriptStruct* InStruct);
-	static bool IsNiagaraFriendlyTopLevelStruct(UScriptStruct* InStruct);
+	static bool IsLWCStructure(UStruct* InStruct);
+	static bool IsLWCType(const FNiagaraTypeDefinition& InType);
+	static UScriptStruct* FindNiagaraFriendlyTopLevelStruct(UScriptStruct* InStruct, ENiagaraStructConversion StructConversion);
+	static bool IsNiagaraFriendlyTopLevelStruct(UScriptStruct* InStruct, ENiagaraStructConversion StructConversion);
+	static UScriptStruct* GetSWCStruct(UScriptStruct* LWCStruct);
 };
 
 /** Information about how this type should be laid out in an FNiagaraDataSet */
@@ -377,7 +446,7 @@ private:
 			//Should be able to support double easily enough
 			else if (FStructProperty* StructProp = CastFieldChecked<FStructProperty>(Property))
 			{
-				GenerateLayoutInfoInternal(Layout, FNiagaraTypeHelper::FindNiagaraFriendlyTopLevelStruct(StructProp->Struct), PropOffset);
+				GenerateLayoutInfoInternal(Layout, FNiagaraTypeHelper::FindNiagaraFriendlyTopLevelStruct(StructProp->Struct, ENiagaraStructConversion::Simulation), PropOffset);
 			}
 			else
 			{
@@ -682,7 +751,7 @@ struct NIAGARA_API FNiagaraVariableMetaData
 		, bDisplayInOverviewStack(false)
 		, InlineParameterSortPriority(0)
 		, bOverrideColor(false)
-		, InlineParameterColorOverride(FLinearColor(EForceInit::ForceInit))
+		, InlineParameterColorOverride(FLinearColor(ForceInit))
 		, bEnableBoolOverride(false)
 		, EditorSortPriority(0)
 		, bInlineEditConditionToggle(false)
@@ -824,7 +893,7 @@ public:
 #endif
 	{
 		checkSlow(ClassStructOrEnum != nullptr);
-		ensure(AllowUnfriendlyStruct == EAllowUnfriendlyStruct::Allow || FNiagaraTypeHelper::IsNiagaraFriendlyTopLevelStruct(StructDef));
+		ensureAlwaysMsgf(AllowUnfriendlyStruct == EAllowUnfriendlyStruct::Allow || FNiagaraTypeHelper::IsNiagaraFriendlyTopLevelStruct(StructDef, ENiagaraStructConversion::UserFacing), TEXT("Struct(%s) is not supported."), *StructDef->GetName());
 	}
 	FORCEINLINE FNiagaraTypeDefinition(UScriptStruct* StructDef) : FNiagaraTypeDefinition(StructDef, EAllowUnfriendlyStruct::Deny) {}
 
@@ -1014,13 +1083,15 @@ public:
 
 	bool IsFloatPrimitive() const
 	{
-		return ClassStructOrEnum == FNiagaraTypeDefinition::GetFloatStruct() || ClassStructOrEnum == FNiagaraTypeDefinition::GetVec2Struct() || ClassStructOrEnum == FNiagaraTypeDefinition::GetVec3Struct() || ClassStructOrEnum == FNiagaraTypeDefinition::GetVec4Struct() ||
-			ClassStructOrEnum == FNiagaraTypeDefinition::GetMatrix4Struct() || ClassStructOrEnum == FNiagaraTypeDefinition::GetColorStruct() || ClassStructOrEnum == FNiagaraTypeDefinition::GetQuatStruct() || ClassStructOrEnum == FNiagaraTypeDefinition::GetPositionStruct();
+		return ClassStructOrEnum == GetFloatStruct() || ClassStructOrEnum == GetVec2Struct() || ClassStructOrEnum ==
+			GetVec3Struct() || ClassStructOrEnum == GetVec4Struct() ||
+			ClassStructOrEnum == GetMatrix4Struct() || ClassStructOrEnum == GetColorStruct() || ClassStructOrEnum ==
+			GetQuatStruct() || ClassStructOrEnum == GetPositionStruct();
  	}
 
 	bool IsIndexType() const
 	{
-		return ClassStructOrEnum == FNiagaraTypeDefinition::GetIntStruct() || ClassStructOrEnum == FNiagaraTypeDefinition::GetBoolStruct() || IsEnum();
+		return ClassStructOrEnum == GetIntStruct() || ClassStructOrEnum == GetBoolStruct() || IsEnum();
 	}
 	bool IsValid() const 
 	{ 
@@ -1247,18 +1318,18 @@ struct TStructOpsTypeTraits<FNiagaraTypeDefinition> : public TStructOpsTypeTrait
 template<typename T>
 const FNiagaraTypeDefinition& FNiagaraTypeDefinition::Get()
 {
-	if (TIsSame<T, float>::Value) { return FNiagaraTypeDefinition::GetFloatDef(); }
-	if (TIsSame<T, FVector2f>::Value) { return FNiagaraTypeDefinition::GetVec2Def(); }
-	if (TIsSame<T, FVector3f>::Value) { return FNiagaraTypeDefinition::GetVec3Def(); }	
-	if (TIsSame<T, FNiagaraPosition>::Value) { return FNiagaraTypeDefinition::GetPositionDef(); }	
-	if (TIsSame<T, FVector4f>::Value) { return FNiagaraTypeDefinition::GetVec4Def(); }
-	if (TIsSame<T, int32>::Value) { return FNiagaraTypeDefinition::GetIntDef(); }
-	if (TIsSame<T, FNiagaraBool>::Value) { return FNiagaraTypeDefinition::GetBoolDef(); }
-	if (TIsSame<T, FQuat4f>::Value) { return FNiagaraTypeDefinition::GetQuatDef(); }
-	if (TIsSame<T, FMatrix44f>::Value) { return FNiagaraTypeDefinition::GetMatrix4Def(); }
-	if (TIsSame<T, FLinearColor>::Value) { return FNiagaraTypeDefinition::GetColorDef(); }
-	if (TIsSame<T, FNiagaraID>::Value) { return FNiagaraTypeDefinition::GetIDDef(); }
-	if (TIsSame<T, FNiagaraRandInfo>::Value) { return FNiagaraTypeDefinition::GetRandInfoDef(); }
+	if (TIsSame<T, float>::Value) { return GetFloatDef(); }
+	if (TIsSame<T, FVector2f>::Value) { return GetVec2Def(); }
+	if (TIsSame<T, FVector3f>::Value) { return GetVec3Def(); }	
+	if (TIsSame<T, FNiagaraPosition>::Value) { return GetPositionDef(); }	
+	if (TIsSame<T, FVector4f>::Value) { return GetVec4Def(); }
+	if (TIsSame<T, int32>::Value) { return GetIntDef(); }
+	if (TIsSame<T, FNiagaraBool>::Value) { return GetBoolDef(); }
+	if (TIsSame<T, FQuat4f>::Value) { return GetQuatDef(); }
+	if (TIsSame<T, FMatrix44f>::Value) { return GetMatrix4Def(); }
+	if (TIsSame<T, FLinearColor>::Value) { return GetColorDef(); }
+	if (TIsSame<T, FNiagaraID>::Value) { return GetIDDef(); }
+	if (TIsSame<T, FNiagaraRandInfo>::Value) { return GetRandInfoDef(); }
 }
 
 FORCEINLINE uint32 GetTypeHash(const FNiagaraTypeDefinition& Type)
@@ -1359,7 +1430,7 @@ public:
 	{
 		FNiagaraTypeRegistry& Registry = Get();
 
-		FRWScopeLock Lock(Registry.RegisteredTypesLock, FRWScopeLockType::SLT_Write);
+		FRWScopeLock Lock(Registry.RegisteredTypesLock, SLT_Write);
 
 		for (const FNiagaraTypeDefinition& Def : Registry.RegisteredUserDefinedTypes)
 		{
@@ -1403,7 +1474,7 @@ public:
 	{
 		FNiagaraTypeRegistry& Registry = Get();
 
-		FRWScopeLock Lock(Registry.RegisteredTypesLock, FRWScopeLockType::SLT_Write);
+		FRWScopeLock Lock(Registry.RegisteredTypesLock, SLT_Write);
 
 		//TODO: Make this a map of type to a more verbose set of metadata? Such as the hlsl defs, offset table for conversions etc.
 		Registry.RegisteredTypeIndexMap.Add(GetTypeHash(NewType), Registry.RegisteredTypes.AddUnique(NewType));
@@ -1481,12 +1552,38 @@ public:
 			}
 		}
 
-		FRWScopeLock Lock(Registry.RegisteredTypesLock, FRWScopeLockType::SLT_Write);
+		FRWScopeLock Lock(Registry.RegisteredTypesLock, SLT_Write);
 		const int32 Index = Registry.RegisteredTypes.AddUnique(NewType);
 		Registry.RegisteredTypeIndexMap.Add(GetTypeHash(NewType), Index);
 		return Index;
 	}
 
+	static void RegisterStructConverter(const FNiagaraTypeDefinition& SourceType, const FNiagaraLwcStructConverter& StructConverter)
+	{
+		int32 TypeIndex = RegisterIndexed(SourceType);
+		FNiagaraTypeRegistry& Registry = Get();
+		
+		FRWScopeLock Lock(Registry.RegisteredTypesLock, SLT_Write);
+		Registry.RegisteredStructConversionMap.Add(TypeIndex, StructConverter);
+	}
+
+	static FNiagaraLwcStructConverter GetStructConverter(const FNiagaraTypeDefinition& SourceType)
+	{
+		FNiagaraTypeRegistry& Registry = Get();
+		
+		FReadScopeLock Lock(Registry.RegisteredTypesLock);
+		const uint32 TypeHash = GetTypeHash(SourceType);
+		if (const int32* TypeIndex = Registry.RegisteredTypeIndexMap.Find(TypeHash))
+		{
+			if (FNiagaraLwcStructConverter* Converter = Registry.RegisteredStructConversionMap.Find(*TypeIndex))
+			{
+				return *Converter;
+			}
+		}
+		
+		return FNiagaraLwcStructConverter();
+	}
+
 	/** LazySingleton interface */
 	static FNiagaraTypeRegistry& Get();
 	static void TearDown();
@@ -1510,6 +1607,8 @@ private:
 
 
 	TMap<uint32, int32> RegisteredTypeIndexMap;
+	TMap<TWeakObjectPtr<UScriptStruct>, TWeakObjectPtr<UScriptStruct>> LWCRegisteredStructRemapping;
+	TMap<uint32, FNiagaraLwcStructConverter> RegisteredStructConversionMap;
 	FRWLock RegisteredTypesLock;
 };
 
@@ -1980,13 +2079,13 @@ struct alignas(16) FNiagaraOwnerParameters
 	FMatrix44f EngineLocalToWorldNoScale = FMatrix44f::Identity;
 	FMatrix44f EngineWorldToLocalNoScale = FMatrix44f::Identity;
 	FQuat4f EngineRotation = FQuat4f::Identity;
-	FVector4f EnginePosition = FVector4f(EForceInit::ForceInitToZero);
-	FVector4f EngineVelocity = FVector4f(EForceInit::ForceInitToZero);
+	FVector4f EnginePosition = FVector4f(ForceInitToZero);
+	FVector4f EngineVelocity = FVector4f(ForceInitToZero);
 	FVector4f EngineXAxis = FVector4f(1.0f, 0.0f, 0.0f, 0.0f);
 	FVector4f EngineYAxis = FVector4f(0.0f, 1.0f, 0.0f, 0.0f);
 	FVector4f EngineZAxis = FVector4f(0.0f, 0.0f, 1.0f, 0.0f);
 	FVector4f EngineScale = FVector4f(1.0f, 1.0f, 1.0f, 0.0f);
-	FVector4f EngineLWCTile = FVector4f(EForceInit::ForceInitToZero);
+	FVector4f EngineLWCTile = FVector4f(ForceInitToZero);
 };
 
 // Any change to this structure, or it's GetVariables implementation will require a bump in the CustomNiagaraVersion so that we
diff --git a/Engine/Plugins/FX/Niagara/Source/NiagaraCore/Private/NiagaraCustomVersion.cpp b/Engine/Plugins/FX/Niagara/Source/NiagaraCore/Private/NiagaraCustomVersion.cpp
--- a/Engine/Plugins/FX/Niagara/Source/NiagaraCore/Private/NiagaraCustomVersion.cpp
+++ b/Engine/Plugins/FX/Niagara/Source/NiagaraCore/Private/NiagaraCustomVersion.cpp
@@ -8,4 +8,5 @@ const FGuid FNiagaraCustomVersion::GUID(0xFCF57AFA, 0x50764283, 0xB9A9E658, 0xFF
 // Register the custom version with core
 FCustomVersionRegistration GRegisterNiagaraCustomVersion(FNiagaraCustomVersion::GUID, FNiagaraCustomVersion::LatestVersion, TEXT("NiagaraVer"));
 
-const FGuid FNiagaraCustomVersion::LatestScriptCompileVersion(0xA023C89C, 0xCE558F49, 0xAB5BA7A4, 0x4B687EE5);
\ No newline at end of file
+// Note: When encountering a conflict on this file please generate a new GUID
+const FGuid FNiagaraCustomVersion::LatestScriptCompileVersion(0x0C09B728, 0x5EE44F4A, 0x9E135A16, 0xAA8A0160);
\ No newline at end of file
diff --git a/Engine/Plugins/FX/Niagara/Source/NiagaraEditor/Private/Customizations/NiagaraComponentDetails.cpp b/Engine/Plugins/FX/Niagara/Source/NiagaraEditor/Private/Customizations/NiagaraComponentDetails.cpp
--- a/Engine/Plugins/FX/Niagara/Source/NiagaraEditor/Private/Customizations/NiagaraComponentDetails.cpp
+++ b/Engine/Plugins/FX/Niagara/Source/NiagaraEditor/Private/Customizations/NiagaraComponentDetails.cpp
@@ -78,13 +78,13 @@ static FNiagaraVariant GetParameterValueFromAsset(const FNiagaraVariableBase& Pa
 		return Value == nullptr ? FNiagaraVariant() : FNiagaraVariant(Value, sizeof(FVector));
 	}
 
-	const uint8* ParameterData = UserParameterStore.GetParameterData(Parameter);
-	if (ParameterData == nullptr)
+	TArray<uint8> DataValue;
+	DataValue.AddUninitialized(Parameter.GetSizeInBytes());
+	if (UserParameterStore.CopyParameterData(Parameter, DataValue.GetData()))
 	{
-		return FNiagaraVariant();
+		return FNiagaraVariant(DataValue);
 	}
-	
-	return FNiagaraVariant(ParameterData, Parameter.GetSizeInBytes());
+	return FNiagaraVariant();
 }
 
 static FNiagaraVariant GetCurrentParameterValue(const FNiagaraVariableBase& Parameter, const UNiagaraComponent* Component)
@@ -508,7 +508,7 @@ private:
 						}
 						else
 						{
-							FMemory::Memcpy(DisplayStruct->GetStructMemory(), OverrideParameters.GetParameterData(UserParameter), UserParameter.GetSizeInBytes());
+							OverrideParameters.CopyParameterData(UserParameter, DisplayStruct->GetStructMemory());
 						}
 					}
 				}
diff --git a/Engine/Plugins/FX/Niagara/Source/NiagaraEditor/Private/EdGraphSchema_Niagara.cpp b/Engine/Plugins/FX/Niagara/Source/NiagaraEditor/Private/EdGraphSchema_Niagara.cpp
--- a/Engine/Plugins/FX/Niagara/Source/NiagaraEditor/Private/EdGraphSchema_Niagara.cpp
+++ b/Engine/Plugins/FX/Niagara/Source/NiagaraEditor/Private/EdGraphSchema_Niagara.cpp
@@ -1413,9 +1413,9 @@ bool UEdGraphSchema_Niagara::ShouldHidePinDefaultValue(UEdGraphPin* Pin) const
 	return false;
 }
 
-FNiagaraVariable UEdGraphSchema_Niagara::PinToNiagaraVariable(const UEdGraphPin* Pin, bool bNeedsValue)const
+FNiagaraVariable UEdGraphSchema_Niagara::PinToNiagaraVariable(const UEdGraphPin* Pin, bool bNeedsValue, ENiagaraStructConversion StructConversion) const
 {
-	FNiagaraVariable Var = FNiagaraVariable(PinToTypeDefinition(Pin), Pin->PinName);
+	FNiagaraVariable Var = FNiagaraVariable(PinToTypeDefinition(Pin, StructConversion), Pin->PinName);
 	bool bHasValue = false;
 	if (Pin->bDefaultValueIsIgnored == false && Pin->DefaultValue.IsEmpty() == false)
 	{
@@ -1481,7 +1481,7 @@ void UEdGraphSchema_Niagara::ConvertIllegalPinsInPlace(UEdGraphPin* Pin)
 	if (Pin->PinType.PinCategory == PinCategoryType && Pin->PinType.PinSubCategoryObject.IsValid())
 	{
 		UScriptStruct* Struct = Cast<UScriptStruct>(Pin->PinType.PinSubCategoryObject.Get());
-		UScriptStruct* ConvertedStruct = FNiagaraTypeHelper::FindNiagaraFriendlyTopLevelStruct(Struct);
+		UScriptStruct* ConvertedStruct = FNiagaraTypeHelper::FindNiagaraFriendlyTopLevelStruct(Struct, ENiagaraStructConversion::UserFacing);
 		if (Struct != ConvertedStruct)
 		{
 			Pin->PinType.PinSubCategoryObject = ConvertedStruct;
@@ -1489,7 +1489,7 @@ void UEdGraphSchema_Niagara::ConvertIllegalPinsInPlace(UEdGraphPin* Pin)
 	}
 }
 
-FNiagaraTypeDefinition UEdGraphSchema_Niagara::PinToTypeDefinition(const UEdGraphPin* Pin)
+FNiagaraTypeDefinition UEdGraphSchema_Niagara::PinToTypeDefinition(const UEdGraphPin* Pin, ENiagaraStructConversion StructConversion)
 {
 	if (Pin == nullptr)
 	{
@@ -1505,7 +1505,8 @@ FNiagaraTypeDefinition UEdGraphSchema_Niagara::PinToTypeDefinition(const UEdGrap
 				*Pin->PinName.ToString(), OwningNode ? *OwningNode->GetName() : TEXT("Invalid"));
 			return FNiagaraTypeDefinition();
 		}
-		else if (Struct && !FNiagaraTypeHelper::IsNiagaraFriendlyTopLevelStruct(Struct)) // LWC swapover support
+
+		if (Struct && !FNiagaraTypeHelper::IsNiagaraFriendlyTopLevelStruct(Struct, StructConversion)) // LWC swapover support
 		{
 			if (OwningNode && OwningNode->HasAnyFlags(EObjectFlags::RF_NeedPostLoad))
 			{
@@ -1517,7 +1518,7 @@ FNiagaraTypeDefinition UEdGraphSchema_Niagara::PinToTypeDefinition(const UEdGrap
 				UE_LOG(LogNiagaraEditor, Verbose, TEXT("Pin states that it is not a niagara friendly struct, but didn't get converted by PostLoad yet, it does not have RF_NeedPostLoad though. Pin Name '%s' Owning Node '%s'."),
 					*Pin->PinName.ToString(), OwningNode ? *OwningNode->GetPathName() : TEXT("Invalid"));
 			}
-			return FNiagaraTypeDefinition(FNiagaraTypeHelper::FindNiagaraFriendlyTopLevelStruct(Struct));
+			return FNiagaraTypeDefinition(FNiagaraTypeHelper::FindNiagaraFriendlyTopLevelStruct(Struct, StructConversion));
 		}
 
 		if (Pin->PinType.PinCategory == PinCategoryType)
@@ -1796,24 +1797,24 @@ FNiagaraTypeDefinition UEdGraphSchema_Niagara::GetTypeDefForProperty(const FProp
 	{
 		return FNiagaraTypeDefinition::GetFloatDef();
 	}
-	else if (Property->IsA(FIntProperty::StaticClass()))
+	if (Property->IsA(FIntProperty::StaticClass()))
 	{
 		return FNiagaraTypeDefinition::GetIntDef();
 	}
-	else if (Property->IsA(FBoolProperty::StaticClass()))
+	if (Property->IsA(FBoolProperty::StaticClass()))
 	{
 		return FNiagaraTypeDefinition::GetBoolDef();
 	}	
-	else if (Property->IsA(FEnumProperty::StaticClass()))
+	if (Property->IsA(FEnumProperty::StaticClass()))
 	{
 		const FEnumProperty* EnumProp = CastField<FEnumProperty>(Property);
 		return FNiagaraTypeDefinition(EnumProp->GetEnum());
 	}
-	else if (const FStructProperty* StructProp = CastField<const FStructProperty>(Property))
+	if (const FStructProperty* StructProp = CastField<const FStructProperty>(Property))
 	{
-		return FNiagaraTypeDefinition(FNiagaraTypeHelper::FindNiagaraFriendlyTopLevelStruct(StructProp->Struct));
+		return FNiagaraTypeDefinition(FNiagaraTypeHelper::FindNiagaraFriendlyTopLevelStruct(StructProp->Struct, ENiagaraStructConversion::UserFacing));
 	}
-	else if(Property->IsA(FUInt16Property::StaticClass()))
+	if(Property->IsA(FUInt16Property::StaticClass()))
 	{
 		return FNiagaraTypeDefinition::GetHalfDef();
 	}
diff --git a/Engine/Plugins/FX/Niagara/Source/NiagaraEditor/Private/NiagaraHlslTranslator.cpp b/Engine/Plugins/FX/Niagara/Source/NiagaraEditor/Private/NiagaraHlslTranslator.cpp
--- a/Engine/Plugins/FX/Niagara/Source/NiagaraEditor/Private/NiagaraHlslTranslator.cpp
+++ b/Engine/Plugins/FX/Niagara/Source/NiagaraEditor/Private/NiagaraHlslTranslator.cpp
@@ -171,6 +171,22 @@ ENiagaraScriptCompileStatus FNiagaraTranslateResults::TranslateResultsToSummary(
 	return SummaryStatus;
 }
 
+FNiagaraVariable ConvertToSimulationVariable(const FNiagaraVariable& Param)
+{
+	if (FNiagaraTypeHelper::IsLWCType(Param.GetType()))
+	{
+		UScriptStruct* Struct = FNiagaraTypeHelper::FindNiagaraFriendlyTopLevelStruct(Param.GetType().GetScriptStruct(), ENiagaraStructConversion::Simulation);
+		FNiagaraVariable SimParam(FNiagaraTypeDefinition(Struct), Param.GetName());
+		if (Param.IsDataAllocated())
+		{
+			SimParam.AllocateData();
+			FNiagaraTypeRegistry::GetStructConverter(Param.GetType()).ConvertDataToSimulation(SimParam.GetData(), Param.GetData());
+		}
+		return SimParam;
+	}
+	return Param;
+}
+
 // helper struct to provide an RAII interface for handling permutations scoping.  We either implement preprocessor directives
 // for creating different permutations, or if the Translator doesn't support it, then we fall back to static branches where
 // possible (this is not viable for declarations in the code)
@@ -384,7 +400,7 @@ bool FHlslNiagaraTranslator::ValidateTypePins(const UNiagaraNode* NodeToValidate
 		}
 		else if (Pin->PinType.PinCategory == UEdGraphSchema_Niagara::PinCategoryType || Pin->PinType.PinCategory == UEdGraphSchema_Niagara::PinCategoryStaticType)
 		{
-			FNiagaraTypeDefinition Type = Schema->PinToTypeDefinition(Pin);
+			FNiagaraTypeDefinition Type = Schema->PinToTypeDefinition(Pin, ENiagaraStructConversion::Simulation);
 			if (Type.IsValid() == false)
 			{
 				Error(LOCTEXT("InvalidPinTypeError", "Node pin has an undefined type."), NodeToValidate, Pin);
@@ -808,7 +824,7 @@ FString FHlslNiagaraTranslator::BuildParameterMapHlslDefinitions(TArray<FNiagara
 	for (int32 UniqueVarIdx = 0; UniqueVarIdx < UniqueVariables.Num(); UniqueVarIdx++)
 	{
 		int UniqueParamMapIdx = 0;
-		FNiagaraVariable Variable = UniqueVariables[UniqueVarIdx];
+		FNiagaraVariable Variable = ConvertToSimulationVariable(UniqueVariables[UniqueVarIdx]);
 
 		if (!AddStructToDefinitionSet(Variable.GetType()))
 		{
@@ -1666,7 +1682,7 @@ const FNiagaraTranslateResults &FHlslNiagaraTranslator::Translate(const FNiagara
 		for (const FNiagaraTypeDefinition& Type : StructsToDefine)
 		{
 			FText ErrorMessage;
-			HlslOutput += BuildHLSLStructDecl(Type, ErrorMessage);
+			HlslOutput += BuildHLSLStructDecl(Type, ErrorMessage, CompileOptions.TargetUsage == ENiagaraScriptUsage::ParticleGPUComputeScript);
 			if (ErrorMessage.IsEmpty() == false)
 			{
 				Error(ErrorMessage, nullptr, nullptr);
@@ -2338,7 +2354,7 @@ void FHlslNiagaraTranslator::GatherComponentsForDataSetAccess(UScriptStruct* Str
 
 		if (const FStructProperty* StructProp = CastField<const FStructProperty>(Property))
 		{
-			UScriptStruct* NiagaraStruct = FNiagaraTypeHelper::FindNiagaraFriendlyTopLevelStruct(StructProp->Struct);
+			UScriptStruct* NiagaraStruct = FNiagaraTypeHelper::FindNiagaraFriendlyTopLevelStruct(StructProp->Struct, ENiagaraStructConversion::Simulation);
 			if (bMatrixRoot && FNiagaraTypeDefinition(NiagaraStruct) == FNiagaraTypeDefinition::GetFloatDef())
 			{
 				GatherComponentsForDataSetAccess(NiagaraStruct, VariableSymbol + ComputeMatrixColumnAccess(Property->GetName()), bMatrixRoot, Components, Types);
@@ -3773,7 +3789,7 @@ void FHlslNiagaraTranslator::DecomposeVariableAccess(UStruct* Struct, bool bRead
 
 		if (FStructProperty* StructProp = CastFieldChecked<FStructProperty>(Property))
 		{
-			UScriptStruct* NiagaraStruct = FNiagaraTypeHelper::FindNiagaraFriendlyTopLevelStruct(StructProp->Struct);
+			UScriptStruct* NiagaraStruct = FNiagaraTypeHelper::FindNiagaraFriendlyTopLevelStruct(StructProp->Struct, ENiagaraStructConversion::Simulation);
 			FNiagaraTypeDefinition PropDef(NiagaraStruct);
 			if (!IsHlslBuiltinVector(PropDef))
 			{
@@ -4088,6 +4104,7 @@ int32 FHlslNiagaraTranslator::AddUniformChunk(FString SymbolName, const FNiagara
 
 	if (Ret == INDEX_NONE)
 	{
+		check(FNiagaraTypeHelper::IsLWCType(Type) == false);
 		Ret = CodeChunks.AddDefaulted();
 		FNiagaraCodeChunk& Chunk = CodeChunks[Ret];
 		Chunk.SymbolName = GetSanitizedSymbolName(SymbolName);
@@ -4131,6 +4148,7 @@ int32 FHlslNiagaraTranslator::AddSourceChunk(FString SymbolName, const FNiagaraT
 
 	if (Ret == INDEX_NONE)
 	{
+		check(FNiagaraTypeHelper::IsLWCType(Type) == false);
 		Ret = CodeChunks.AddDefaulted();
 		FNiagaraCodeChunk& Chunk = CodeChunks[Ret];
 		Chunk.SymbolName = bSanitize ? GetSanitizedSymbolName(SymbolName) : SymbolName;
@@ -4157,7 +4175,7 @@ int32 FHlslNiagaraTranslator::AddBodyChunk(FString SymbolName, FString Definitio
 {
 	check(CurrentBodyChunkMode == ENiagaraCodeChunkMode::Body || CurrentBodyChunkMode == ENiagaraCodeChunkMode::SpawnBody || CurrentBodyChunkMode == ENiagaraCodeChunkMode::UpdateBody ||
 		(CurrentBodyChunkMode >= ENiagaraCodeChunkMode::SimulationStageBody && CurrentBodyChunkMode < ENiagaraCodeChunkMode::SimulationStageBodyMax));
-
+	check(FNiagaraTypeHelper::IsLWCType(Type) == false);
 	int32 Ret = CodeChunks.AddDefaulted();
 	FNiagaraCodeChunk& Chunk = CodeChunks[Ret];
 	Chunk.SymbolName = GetSanitizedSymbolName(SymbolName);
@@ -4178,7 +4196,7 @@ int32 FHlslNiagaraTranslator::AddBodyChunk(FString SymbolName, FString Definitio
 {
 	check(CurrentBodyChunkMode == ENiagaraCodeChunkMode::Body || CurrentBodyChunkMode == ENiagaraCodeChunkMode::SpawnBody || CurrentBodyChunkMode == ENiagaraCodeChunkMode::UpdateBody ||
 		(CurrentBodyChunkMode >= ENiagaraCodeChunkMode::SimulationStageBody && CurrentBodyChunkMode < ENiagaraCodeChunkMode::SimulationStageBodyMax));
-
+	check(FNiagaraTypeHelper::IsLWCType(Type) == false);
 	int32 Ret = CodeChunks.AddDefaulted();
 	FNiagaraCodeChunk& Chunk = CodeChunks[Ret];
 	Chunk.SymbolName = GetSanitizedSymbolName(SymbolName);
@@ -4197,7 +4215,7 @@ int32 FHlslNiagaraTranslator::AddBodyChunk(FString SymbolName, FString Definitio
 {
 	check(CurrentBodyChunkMode == ENiagaraCodeChunkMode::Body || CurrentBodyChunkMode == ENiagaraCodeChunkMode::SpawnBody || CurrentBodyChunkMode == ENiagaraCodeChunkMode::UpdateBody ||
 			(CurrentBodyChunkMode >= ENiagaraCodeChunkMode::SimulationStageBody && CurrentBodyChunkMode < ENiagaraCodeChunkMode::SimulationStageBodyMax));
-
+	check(FNiagaraTypeHelper::IsLWCType(Type) == false);
 	int32 Ret = CodeChunks.AddDefaulted();
 	FNiagaraCodeChunk& Chunk = CodeChunks[Ret];
 	Chunk.SymbolName = GetSanitizedSymbolName(SymbolName);
@@ -4388,7 +4406,7 @@ int32 FHlslNiagaraTranslator::GetParameter(const FNiagaraVariable& Parameter)
 			}
 			//If this is a valid function parameter, use that.
 			FString SymbolName = TEXT("In_") + GetSanitizedSymbolName(Parameter.GetName().ToString());
-			return AddSourceChunk(SymbolName, Parameter.GetType());
+			return AddSourceChunk(SymbolName, ConvertToSimulationVariable(Parameter).GetType());
 		}
 	}
 
@@ -4534,8 +4552,27 @@ int32 FHlslNiagaraTranslator::GetConstantDirect(int InConstantValue)
 	return GetConstant(Constant);
 }
 
-bool FHlslNiagaraTranslator::GenerateStructInitializer(TStringBuilder<128>& InitializerString, UStruct* UserDefinedStruct, const void* StructData)
+bool FHlslNiagaraTranslator::GenerateStructInitializer(TStringBuilder<128>& InitializerString, UStruct* UserDefinedStruct, const void* StructData, int32 ByteOffset)
 {
+	//-TODO: Alignment Issues
+	//const bool bHasAlignmentDummy = CompileOptions.TargetUsage != ENiagaraScriptUsage::ParticleGPUComputeScript;
+
+	//auto CheckAlignment =
+	//	[&](FProperty* Property)
+	//	{
+	//		if ( bHasAlignmentDummy && (ByteOffset != Property->GetOffset_ReplaceWith_ContainerPtrToValuePtr()))
+	//		{
+	//			const int32 ByteDelta = Property->GetOffset_ReplaceWith_ContainerPtrToValuePtr() - ByteOffset;
+	//			const int32 NumFloats = ByteDelta / sizeof(float);
+	//			check(NumFloats > 0 && ByteDelta > 0);
+	//			for (int32 i = 0; i < NumFloats; ++i)
+	//			{
+	//				InitializerString.Append(TEXT("0, "));
+	//			}
+	//			ByteOffset += ByteDelta;
+	//		}
+	//	};
+
 	InitializerString.AppendChar(TEXT('{'));
 	for (FField* ChildProperty = UserDefinedStruct->ChildProperties; ChildProperty; ChildProperty = ChildProperty->Next)
 	{
@@ -4546,22 +4583,29 @@ bool FHlslNiagaraTranslator::GenerateStructInitializer(TStringBuilder<128>& Init
 
 		if (FFloatProperty* FloatProperty = CastField<FFloatProperty>(ChildProperty))
 		{
+			//CheckAlignment(FloatProperty);
 			const float Value = FloatProperty->GetPropertyValue_InContainer(StructData);
 			InitializerString.Appendf(TEXT("%g"), Value);
+			ByteOffset += 4;
 		}
 		else if (FIntProperty* IntProperty = CastField<FIntProperty>(ChildProperty))
 		{
+			//CheckAlignment(IntProperty);
 			const int32 Value = IntProperty->GetPropertyValue_InContainer(StructData);
 			InitializerString.Appendf(TEXT("%d"), Value);
+			ByteOffset += 4;
 		}
 		else if (FBoolProperty* BoolProperty = CastField<FBoolProperty>(ChildProperty))
 		{
+			//CheckAlignment(BoolProperty);
 			const bool bValue = BoolProperty->GetPropertyValue_InContainer(StructData);
 			InitializerString.Append(bValue ? TEXT("true") : TEXT("false"));
+			ByteOffset += 4;
 		}
 		else if (FStructProperty* StructProperty = CastField<FStructProperty>(ChildProperty))
 		{
-			if ( !GenerateStructInitializer(InitializerString, StructProperty->Struct, StructProperty->ContainerPtrToValuePtr<const void>(StructData)) )
+			//CheckAlignment(StructProperty);
+			if ( !GenerateStructInitializer(InitializerString, StructProperty->Struct, StructProperty->ContainerPtrToValuePtr<const void>(StructData), ByteOffset) )
 			{
 				return false;
 			}
@@ -4828,7 +4872,7 @@ void FHlslNiagaraTranslator::Output(UNiagaraNodeOutput* OutputNode, const TArray
 	check(NumberOfValidComputedInputs == Outputs.Num());
 	for (int32 PinIdx = 0; PinIdx < Outputs.Num(); PinIdx++)
 	{
-		Attributes.Add(Outputs[PinIdx]);
+		Attributes.Add(ConvertToSimulationVariable(Outputs[PinIdx]));
 		Inputs.Add(ComputedInputs[PinIdx]);
 	}
 
@@ -5031,7 +5075,7 @@ void FHlslNiagaraTranslator::ParameterMapSet(UNiagaraNodeParameterMapSet* SetNod
 		}
 		else // These are the pins that we are setting on the parameter map.
 		{
-			FNiagaraVariable Var = Schema->PinToNiagaraVariable(Inputs[i].Pin, false);
+			FNiagaraVariable Var = Schema->PinToNiagaraVariable(Inputs[i].Pin, false, ENiagaraStructConversion::Simulation);
 
 			if (!AddStructToDefinitionSet(Var.GetType()))
 			{
@@ -5451,7 +5495,7 @@ bool FHlslNiagaraTranslator::ParameterMapRegisterExternalConstantNamespaceVariab
 
 void FHlslNiagaraTranslator::FillVariableWithDefaultValue(FNiagaraVariable& InVariable, const UEdGraphPin* InDefaultPin)
 {
-	FNiagaraVariable Var = Schema->PinToNiagaraVariable(InDefaultPin, true);
+	FNiagaraVariable Var = Schema->PinToNiagaraVariable(InDefaultPin, true, ENiagaraStructConversion::Simulation);
 	FNiagaraEditorUtilities::ResetVariableToDefaultValue(InVariable);
 	if (Var.IsDataAllocated() && Var.GetData() != nullptr)
 	{
@@ -5461,7 +5505,7 @@ void FHlslNiagaraTranslator::FillVariableWithDefaultValue(FNiagaraVariable& InVa
 
 void FHlslNiagaraTranslator::FillVariableWithDefaultValue(int32& OutValue, const UEdGraphPin* InDefaultPin)
 {
-	FNiagaraVariable Var = Schema->PinToNiagaraVariable(InDefaultPin, true);
+	FNiagaraVariable Var = Schema->PinToNiagaraVariable(InDefaultPin, true, ENiagaraStructConversion::Simulation);
 	FNiagaraVariable VarFinal = Var;
 	FNiagaraEditorUtilities::ResetVariableToDefaultValue(VarFinal); // Do this to handle non-zero defaults
 	if (Var.IsDataAllocated() && Var.GetData() != nullptr)
@@ -5638,7 +5682,6 @@ void FHlslNiagaraTranslator::SetConstantByStaticVariable(int32& OutValue, const
 }
 
 
-
 bool FHlslNiagaraTranslator::ParameterMapRegisterUniformAttributeVariable(const FNiagaraVariable& InVariable, UNiagaraNode* InNode, int32 InParamMapHistoryIdx, int32& Output)
 {
 	FNiagaraVariable NewVar = FNiagaraParameterMapHistory::BasicAttributeToNamespacedAttribute(InVariable);
@@ -5881,7 +5924,7 @@ void FHlslNiagaraTranslator::ParameterMapGet(UNiagaraNodeParameterMapGet* GetNod
 			bool bNeedsValue =
 				OutputTypeDefinition != FNiagaraTypeDefinition::GetParameterMapDef() &&
 				OutputTypeDefinition.IsDataInterface() == false;
-			FNiagaraVariable Var = Schema->PinToNiagaraVariable(OutputPins[i], bNeedsValue);
+			FNiagaraVariable Var = Schema->PinToNiagaraVariable(OutputPins[i], bNeedsValue, ENiagaraStructConversion::Simulation);
 
 			UNiagaraScriptVariable* Variable = GetNode->GetNiagaraGraph()->GetScriptVariable(Var);
 			if (Var.GetType().IsStatic())
@@ -5917,7 +5960,7 @@ int32 FHlslNiagaraTranslator::MakeStaticVariableDirect(const UEdGraphPin* InDefa
 void FHlslNiagaraTranslator::HandleParameterRead(int32 ParamMapHistoryIdx, const FNiagaraVariable& InVar, const UEdGraphPin* DefaultPin, UNiagaraNode* ErrorNode, int32& OutputChunkId, UNiagaraScriptVariable* Variable, bool bTreatAsUnknownParameterMap, bool bIgnoreDefaultSetFirst)
 {
 	FString ParameterMapInstanceName = GetParameterMapInstanceName(ParamMapHistoryIdx);
-	FNiagaraVariable Var = InVar;
+	FNiagaraVariable Var = ConvertToSimulationVariable(InVar);
 	if (!AddStructToDefinitionSet(Var.GetType()))
 	{
 		Error(FText::Format(LOCTEXT("ParameterMapGetTypeError", "Cannot handle type {0}! Variable: {1}"), Var.GetType().GetNameText(), FText::FromName(Var.GetName())), nullptr, nullptr);
@@ -6141,7 +6184,7 @@ void FHlslNiagaraTranslator::HandleParameterRead(int32 ParamMapHistoryIdx, const
 				}
 				else if (InputPin != nullptr && !InputPin->bDefaultValueIsIgnored) // Use the default from the input pin because this variable was previously never encountered.
 				{
-					FNiagaraVariable PinVar = Schema->PinToNiagaraVariable(InputPin, true);
+					FNiagaraVariable PinVar = Schema->PinToNiagaraVariable(InputPin, true, ENiagaraStructConversion::Simulation);
 					FString DebugConstantStr;
 					OutputChunkId = GetConstant(PinVar, &DebugConstantStr);
 					UE_LOG(LogNiagaraEditor, VeryVerbose, TEXT("Converted default value of parameter %s to constant %s for script %s. Likely added since this system was last compiled."), *Var.GetName().ToString(), *DebugConstantStr, *CompileOptions.FullName);
@@ -6286,7 +6329,7 @@ void FHlslNiagaraTranslator::HandleParameterRead(int32 ParamMapHistoryIdx, const
 					}
 					else if (InputPin != nullptr && !InputPin->bDefaultValueIsIgnored) // Use the default from the input pin because this variable was previously never encountered.
 					{
-						FNiagaraVariable PinVar = Schema->PinToNiagaraVariable(InputPin, true);
+						FNiagaraVariable PinVar = Schema->PinToNiagaraVariable(InputPin, true, ENiagaraStructConversion::Simulation);
 						FString DebugConstantStr;
 						OutputChunkId = GetConstant(PinVar, &DebugConstantStr);
 						UE_LOG(LogNiagaraEditor, Display, TEXT("Converted default value of parameter %s to constant %s for script %s. Likely added since this system was last compiled."), *Var.GetName().ToString(), *DebugConstantStr, *CompileOptions.FullName);
@@ -6689,7 +6732,7 @@ void FHlslNiagaraTranslator::Operation(class UNiagaraNodeOp* Operation, TArray<i
 	for (int32 OutputIndex = 0; OutputIndex < OutputPins.Num(); OutputIndex++)
 	{
 		UEdGraphPin* OutputPin = OutputPins[OutputIndex];
-		FNiagaraTypeDefinition OutputType = Schema->PinToTypeDefinition(OutputPin);
+		FNiagaraTypeDefinition OutputType = Schema->PinToTypeDefinition(OutputPin, ENiagaraStructConversion::Simulation);
 
 		if (!AddStructToDefinitionSet(OutputType))
 		{
@@ -8052,7 +8095,7 @@ void FHlslNiagaraTranslator::GenerateFunctionCall(ENiagaraScriptUsage ScriptUsag
 	for (int32 i = 0; i < FunctionSignature.Outputs.Num(); ++i)
 	{
 		FNiagaraVariable& OutVar = FunctionSignature.Outputs[i];
-		FNiagaraTypeDefinition Type = OutVar.GetType();
+		FNiagaraTypeDefinition Type = ConvertToSimulationVariable(OutVar).GetType();
 
 		//We don't write class types as real params in the hlsl
 		if (!Type.GetClass())
@@ -8085,7 +8128,7 @@ void FHlslNiagaraTranslator::GenerateFunctionCall(ENiagaraScriptUsage ScriptUsag
 			else
 			{
 				FString OutputStr = FString::Printf(TEXT("%sOutput_%s"), *GetFunctionSignatureSymbol(FunctionSignature), *OutVar.GetName().ToString());
-				Output = AddBodyChunk(GetUniqueSymbolName(*OutputStr), TEXT(""), OutVar.GetType());
+				Output = AddBodyChunk(GetUniqueSymbolName(*OutputStr), TEXT(""), Type);
 				ParamOutput = Output;
 			}
 
@@ -8283,7 +8326,8 @@ FString FHlslNiagaraTranslator::GetFunctionSignature(const FNiagaraFunctionSigna
 					SigStr += TEXT(", ");
 				}
 
-				SigStr += FHlslNiagaraTranslator::GetStructHlslTypeName(Input.GetType()) + TEXT(" In_") + FHlslNiagaraTranslator::GetSanitizedSymbolName(Input.GetName().ToString(), true);
+				FNiagaraVariable SimInput = ConvertToSimulationVariable(Input);
+				SigStr += FHlslNiagaraTranslator::GetStructHlslTypeName(SimInput.GetType()) + TEXT(" In_") + FHlslNiagaraTranslator::GetSanitizedSymbolName(Input.GetName().ToString(), true);
 				++ParamIdx;
 			}
 		}
@@ -8306,7 +8350,8 @@ FString FHlslNiagaraTranslator::GetFunctionSignature(const FNiagaraFunctionSigna
 					SigStr += TEXT(", ");
 				}
 
-				SigStr += TEXT("out ") + FHlslNiagaraTranslator::GetStructHlslTypeName(Output.GetType()) + TEXT(" ") + FHlslNiagaraTranslator::GetSanitizedSymbolName(TEXT("Out_") + Output.GetName().ToString());
+				FNiagaraVariable SimOutput = ConvertToSimulationVariable(Output);
+				SigStr += TEXT("out ") + FHlslNiagaraTranslator::GetStructHlslTypeName(SimOutput.GetType()) + TEXT(" ") + FHlslNiagaraTranslator::GetSanitizedSymbolName(TEXT("Out_") + Output.GetName().ToString());
 				++ParamIdx;
 			}
 		}
@@ -8369,7 +8414,7 @@ FNiagaraTypeDefinition FHlslNiagaraTranslator::GetChildType(const FNiagaraTypeDe
 				}
 				else if (const FStructProperty* StructProp = CastFieldChecked<const FStructProperty>(Property))
 				{
-					return FNiagaraTypeDefinition(FNiagaraTypeHelper::FindNiagaraFriendlyTopLevelStruct(StructProp->Struct));
+					return FNiagaraTypeDefinition(FNiagaraTypeHelper::FindNiagaraFriendlyTopLevelStruct(StructProp->Struct, ENiagaraStructConversion::Simulation));
 				}
 			}
 		}
@@ -8501,7 +8546,7 @@ void FHlslNiagaraTranslator::Convert(class UNiagaraNodeConvert* Convert, TArrayV
 			InputPin->PinType.PinCategory == UEdGraphSchema_Niagara::PinCategoryEnum || 
 			InputPin->PinType.PinCategory == UEdGraphSchema_Niagara::PinCategoryStaticEnum)
 		{
-			FNiagaraTypeDefinition Type = Schema->PinToTypeDefinition(InputPin);
+			FNiagaraTypeDefinition Type = Schema->PinToTypeDefinition(InputPin, ENiagaraStructConversion::Simulation);
 			if (!AddStructToDefinitionSet(Type))
 			{
 				Error(FText::Format(LOCTEXT("ConvertTypeError_InvalidInput", "Cannot handle input pin type {0}! Pin: {1}"), Type.GetNameText(), InputPin->GetDisplayName()), nullptr, nullptr);
@@ -8518,7 +8563,7 @@ void FHlslNiagaraTranslator::Convert(class UNiagaraNodeConvert* Convert, TArrayV
 			OutputPin->PinType.PinCategory == UEdGraphSchema_Niagara::PinCategoryEnum ||
 			OutputPin->PinType.PinCategory == UEdGraphSchema_Niagara::PinCategoryStaticEnum)
 		{
-			FNiagaraTypeDefinition Type = Schema->PinToTypeDefinition(OutputPin);
+			FNiagaraTypeDefinition Type = Schema->PinToTypeDefinition(OutputPin, ENiagaraStructConversion::Simulation);
 			if (!AddStructToDefinitionSet(Type))
 			{
 				Error(FText::Format(LOCTEXT("ConvertTypeError_InvalidOutput", "Cannot handle output pin type {0}! Pin: {1}"), Type.GetNameText(), OutputPin->GetDisplayName()), nullptr, nullptr);
@@ -8543,14 +8588,14 @@ void FHlslNiagaraTranslator::Convert(class UNiagaraNodeConvert* Convert, TArrayV
 		int32 DestinationIndex = GetPinIndexById(OutputPins, Connection.DestinationPinId);
 		if (SourceIndex != INDEX_NONE && SourceIndex < Inputs.Num() && DestinationIndex != INDEX_NONE && DestinationIndex < Outputs.Num())
 		{
-			FNiagaraTypeDefinition SrcPinType = Schema->PinToTypeDefinition(InputPins[SourceIndex]);
+			FNiagaraTypeDefinition SrcPinType = Schema->PinToTypeDefinition(InputPins[SourceIndex], ENiagaraStructConversion::Simulation);
 			if (!AddStructToDefinitionSet(SrcPinType))
 			{
 				Error(FText::Format(LOCTEXT("ConvertTypeError_InvalidSubpinInput", "Cannot handle input subpin type {0}! Subpin: {1}"), SrcPinType.GetNameText(), InputPins[SourceIndex]->GetDisplayName()), nullptr, nullptr);
 			}
 			TArray<FName> ConditionedSourcePath = ConditionPropertyPath(SrcPinType, Connection.SourcePath);
 
-			FNiagaraTypeDefinition DestPinType = Schema->PinToTypeDefinition(OutputPins[DestinationIndex]);
+			FNiagaraTypeDefinition DestPinType = Schema->PinToTypeDefinition(OutputPins[DestinationIndex], ENiagaraStructConversion::Simulation);
 			if (!AddStructToDefinitionSet(DestPinType))
 			{
 				Error(FText::Format(LOCTEXT("ConvertTypeError_InvalidSubpinOutput", "Cannot handle output subpin type type {0}! Subpin: {1}"), DestPinType.GetNameText(), OutputPins[SourceIndex]->GetDisplayName()), nullptr, nullptr);
@@ -8580,7 +8625,7 @@ void FHlslNiagaraTranslator::If(UNiagaraNodeIf* IfNode, TArray<FNiagaraVariable>
 	int32 PinIdx = 1;
 	for (FNiagaraVariable& Var : Vars)
 	{
-		FNiagaraTypeDefinition Type = Schema->PinToTypeDefinition(IfNode->GetInputPin(PinIdx++));
+		FNiagaraTypeDefinition Type = Schema->PinToTypeDefinition(IfNode->GetInputPin(PinIdx++), ENiagaraStructConversion::Simulation);
 		if (!AddStructToDefinitionSet(Type))
 		{
 			FText OutErrorMessage = FText::Format(LOCTEXT("UnknownNumeric", "Variable in If node uses invalid type. Var: {0} Type: {1}"),
@@ -8878,7 +8923,7 @@ int32 FHlslNiagaraTranslator::CompilePin(const UEdGraphPin* Pin)
 
 	check(Pin);
 	int32 Ret = INDEX_NONE;
-	FNiagaraTypeDefinition TypeDef = Schema->PinToTypeDefinition(Pin);
+	
 	if (Pin->Direction == EGPD_Input)
 	{
 		if (Pin->LinkedTo.Num() > 0)
@@ -8898,6 +8943,7 @@ int32 FHlslNiagaraTranslator::CompilePin(const UEdGraphPin* Pin)
 		}
 		else if (!Pin->bDefaultValueIsIgnored && (Pin->PinType.PinCategory == UEdGraphSchema_Niagara::PinCategoryType || Pin->PinType.PinCategory == UEdGraphSchema_Niagara::PinCategoryStaticType))
 		{
+			FNiagaraTypeDefinition TypeDef = Schema->PinToTypeDefinition(Pin);
 			if (TypeDef == FNiagaraTypeDefinition::GetParameterMapDef())
 			{
 				Error(FText::FromString(TEXT("Parameter Maps must be created via an Input Node, not the default value of a pin! Please connect to a valid input Parameter Map.")), Cast<UNiagaraNode>(Pin->GetOwningNode()), Pin);
@@ -8906,14 +8952,14 @@ int32 FHlslNiagaraTranslator::CompilePin(const UEdGraphPin* Pin)
 			else
 			{
 				//No connections to this input so add the default as a const expression.			
-				FNiagaraVariable PinVar = Schema->PinToNiagaraVariable(Pin, true);
+				FNiagaraVariable PinVar = Schema->PinToNiagaraVariable(Pin, true, ENiagaraStructConversion::Simulation);
 				return GetConstant(PinVar);
 			}
 		}
 		else if (!Pin->bDefaultValueIsIgnored && (Pin->PinType.PinCategory == UEdGraphSchema_Niagara::PinCategoryEnum || Pin->PinType.PinCategory == UEdGraphSchema_Niagara::PinCategoryStaticEnum))
 		{
 			//No connections to this input so add the default as a const expression.			
-			FNiagaraVariable PinVar = Schema->PinToNiagaraVariable(Pin, true);
+			FNiagaraVariable PinVar = Schema->PinToNiagaraVariable(Pin, true, ENiagaraStructConversion::Simulation);
 			return GetConstant(PinVar);
 		}
 	}
@@ -9254,6 +9300,8 @@ bool FHlslNiagaraTranslator::IsBuiltInHlslType(const FNiagaraTypeDefinition& Typ
 
 FString FHlslNiagaraTranslator::GetStructHlslTypeName(const FNiagaraTypeDefinition& Type)
 {
+	check(FNiagaraTypeHelper::IsLWCType(Type) == false);
+
 	if (Type.IsValid() == false)
 	{
 		return "undefined";
@@ -9313,7 +9361,7 @@ FString FHlslNiagaraTranslator::GetPropertyHlslTypeName(const FProperty* Propert
 	else if (Property->IsA(FStructProperty::StaticClass()))
 	{
 		const FStructProperty* StructProp = CastField<const FStructProperty>(Property);
-		return GetStructHlslTypeName(FNiagaraTypeHelper::FindNiagaraFriendlyTopLevelStruct(StructProp->Struct));
+		return GetStructHlslTypeName(FNiagaraTypeHelper::FindNiagaraFriendlyTopLevelStruct(StructProp->Struct, ENiagaraStructConversion::Simulation));
 	}
 	else if (Property->IsA(FEnumProperty::StaticClass()))
 	{
@@ -9333,26 +9381,83 @@ FString FHlslNiagaraTranslator::GetPropertyHlslTypeName(const FProperty* Propert
 	}
 }
 
-FString FHlslNiagaraTranslator::BuildHLSLStructDecl(const FNiagaraTypeDefinition& Type, FText& OutErrorMessage)
+FString FHlslNiagaraTranslator::BuildHLSLStructDecl(const FNiagaraTypeDefinition& Type, FText& OutErrorMessage, bool bGpuScript)
 {
 	if (!IsBuiltInHlslType(Type))
 	{
-		FString StructName = GetStructHlslTypeName(Type);
+		check(FNiagaraTypeHelper::IsLWCType(Type) == false);
+		FString Decl = "struct " + GetStructHlslTypeName(Type) + "\n{\n";
 
-		FString Decl = "struct " + StructName + "\n{\n";
+		int32 StructSize = 0;
+		int32 DummyIndex = 0;
 		for (TFieldIterator<FProperty> PropertyIt(Type.GetStruct(), EFieldIteratorFlags::IncludeSuper); PropertyIt; ++PropertyIt)
 		{
 			FProperty* Property = *PropertyIt;
-			FString PropertyTypeName = GetPropertyHlslTypeName(Property);
-			if (PropertyTypeName.IsEmpty())
+
+			FString PropertyTypeName;
+			int32 PropertyTypeSize;
+			if (Property->IsA(FFloatProperty::StaticClass()))
+			{
+				PropertyTypeName = TEXT("float");
+				PropertyTypeSize = 4;
+			}
+			else if (
+				Property->IsA(FIntProperty::StaticClass()) ||
+				Property->IsA(FUInt32Property::StaticClass()) ||
+				Property->IsA(FEnumProperty::StaticClass()) ||
+				Property->IsA(FByteProperty::StaticClass()) ||
+				Property->IsA(FBoolProperty::StaticClass())
+			)
+			{
+				PropertyTypeName = TEXT("int");
+				PropertyTypeSize = 4;
+			}
+			else if (Property->IsA(FStructProperty::StaticClass()))
 			{
-				OutErrorMessage = FText::Format(LOCTEXT("UnknownPropertyTypeErrorFormat", "Failed to build hlsl struct declaration for type {0}.  Property {1} has an unsuported type {2}."),
-					FText::FromString(Type.GetName()), Property->GetDisplayNameText(), FText::FromString(Property->GetClass()->GetName()));
+				const FStructProperty* StructProp = CastField<const FStructProperty>(Property);
+				const FNiagaraTypeDefinition NiagaraType(StructProp->Struct);
+				PropertyTypeSize = NiagaraType.GetSize();
+				PropertyTypeName = GetStructHlslTypeName(NiagaraType);
+			}
+			else
+			{
+				OutErrorMessage = FText::Format(
+					LOCTEXT("UnknownPropertyTypeErrorFormat", "Failed to build hlsl struct declaration for type {0}.  Property {1} has an unsuported type {2}."),
+					FText::FromString(Type.GetName()), Property->GetDisplayNameText(), FText::FromString(Property->GetClass()->GetName())
+				);
 				return TEXT("");
 			}
-			Decl += TEXT("\t") + PropertyTypeName + TEXT(" ") + Property->GetName() + TEXT(";\n");
-		}
-		Decl += "};\n\n";
+
+			//-TODO: Alignment issues...
+			//if (!bGpuScript && (StructSize != Property->GetOffset_ReplaceWith_ContainerPtrToValuePtr()))
+			//{
+			//	const int32 ByteDelta = Property->GetOffset_ReplaceWith_ContainerPtrToValuePtr() - StructSize;
+			//	const int32 NumFloats = ByteDelta / sizeof(float);
+			//	check(NumFloats > 0 && ByteDelta > 0);
+			//	for ( int32 i=0; i < NumFloats; ++i )
+			//	{
+			//		Decl.Appendf(TEXT("\tfloat Dummy%d;\n"), DummyIndex++);
+			//	}
+			//	StructSize += ByteDelta;
+			//}
+
+			Decl.Appendf(TEXT("\t%s %s;\n"), *PropertyTypeName, *Property->GetName());
+			StructSize += PropertyTypeSize;
+		}
+
+		//-TODO: Alignment issues...
+		//if (!bGpuScript && (StructSize != Type.GetStruct()->GetStructureSize()))
+		//{
+		//	const int32 ByteDelta = Type.GetStruct()->GetStructureSize() - StructSize;
+		//	const int32 NumFloats = ByteDelta / sizeof(float);
+		//	check(NumFloats > 0 && ByteDelta > 0);
+		//	for (int32 i = 0; i < NumFloats; ++i)
+		//	{
+		//		Decl.Appendf(TEXT("\tfloat Dummy%d;\n"), DummyIndex++);
+		//	}
+		//}
+
+		Decl.Append(TEXT("};\n\n"));
 		return Decl;
 	}
 
@@ -9395,7 +9500,7 @@ bool FHlslNiagaraTranslator::AddStructToDefinitionSet(const FNiagaraTypeDefiniti
 
 	// Now make sure that we don't have any other struct types within our struct. Add them prior to the struct in question to make sure
 	// that the syntax works out properly.
-	const UScriptStruct* Struct = TypeDef.GetScriptStruct();
+	UScriptStruct* Struct = FNiagaraTypeHelper::FindNiagaraFriendlyTopLevelStruct(TypeDef.GetScriptStruct(), ENiagaraStructConversion::Simulation);
 	if (Struct != nullptr)
 	{
 		// We need to recursively dig through the struct to get at the lowest level of the input struct, which
@@ -9406,7 +9511,7 @@ bool FHlslNiagaraTranslator::AddStructToDefinitionSet(const FNiagaraTypeDefiniti
 			const FStructProperty* StructProp = CastField<const FStructProperty>(Property);
 			if (StructProp)
 			{
-				if (!AddStructToDefinitionSet(FNiagaraTypeHelper::FindNiagaraFriendlyTopLevelStruct(StructProp->Struct)))
+				if (!AddStructToDefinitionSet(StructProp->Struct))
 				{
 					return false;
 				}
@@ -9414,12 +9519,14 @@ bool FHlslNiagaraTranslator::AddStructToDefinitionSet(const FNiagaraTypeDefiniti
 		}
 
 		// Add the new type def
-		if (!StructsToDefine.Contains(TypeDef))
+		FNiagaraTypeDefinition NewTypeDef(Struct);
+		if (!StructsToDefine.Contains(NewTypeDef))
 		{
-			StructsToDefine.Add(TypeDef);
+			check(FNiagaraTypeHelper::IsLWCType(NewTypeDef) == false);
+			StructsToDefine.Add(NewTypeDef);
 
 			// Add the struct name to the unique symbol names to make it so that we don't declare a variable the same name as the struct type.
-			GetUniqueSymbolName(*TypeDef.GetName());
+			GetUniqueSymbolName(*NewTypeDef.GetName());
 		}
 	}
 
@@ -9466,7 +9573,7 @@ TArray<FName> FHlslNiagaraTranslator::ConditionPropertyPath(const FNiagaraTypeDe
 						ReturnPath.Add(InPath[0]);
 						TArray<FName> Subset = InPath;
 						Subset.RemoveAt(0);
-						TArray<FName> Children = ConditionPropertyPath(FNiagaraTypeHelper::FindNiagaraFriendlyTopLevelStruct(StructProp->Struct), Subset);
+						TArray<FName> Children = ConditionPropertyPath(FNiagaraTypeHelper::FindNiagaraFriendlyTopLevelStruct(StructProp->Struct, ENiagaraStructConversion::Simulation), Subset);
 						for (const FName& Child : Children)
 						{
 							ReturnPath.Add(Child);
diff --git a/Engine/Plugins/FX/Niagara/Source/NiagaraEditor/Private/NiagaraHlslTranslator.h b/Engine/Plugins/FX/Niagara/Source/NiagaraEditor/Private/NiagaraHlslTranslator.h
--- a/Engine/Plugins/FX/Niagara/Source/NiagaraEditor/Private/NiagaraHlslTranslator.h
+++ b/Engine/Plugins/FX/Niagara/Source/NiagaraEditor/Private/NiagaraHlslTranslator.h
@@ -634,7 +634,7 @@ public:
 	static bool IsBuiltInHlslType(const FNiagaraTypeDefinition& Type);
 	static FString GetStructHlslTypeName(const FNiagaraTypeDefinition& Type);
 	static FString GetPropertyHlslTypeName(const FProperty* Property);
-	static FString BuildHLSLStructDecl(const FNiagaraTypeDefinition& Type, FText& OutErrorMessage);
+	static FString BuildHLSLStructDecl(const FNiagaraTypeDefinition& Type, FText& OutErrorMessage, bool bGpuScript);
 	static FString GetHlslDefaultForType(const FNiagaraTypeDefinition& Type);
 	static bool IsHlslBuiltinVector(const FNiagaraTypeDefinition& Type);
 	static TArray<FName> ConditionPropertyPath(const FNiagaraTypeDefinition& Type, const TArray<FName>& InPath);
@@ -717,7 +717,7 @@ private:
 
 	// Generate a structure initializer string
 	// Returns true if we generated the structure successfully or false if we encounter something we could not handle
-	bool GenerateStructInitializer(TStringBuilder<128>& InitializerString, UStruct* UserDefinedStruct, const void* StructData);
+	bool GenerateStructInitializer(TStringBuilder<128>& InitializerString, UStruct* UserDefinedStruct, const void* StructData, int32 ByteOffset = 0);
 	// Convert a variable with actual data into a constant string
 	FString GenerateConstantString(const FNiagaraVariable& Constant);
 
diff --git a/Engine/Plugins/FX/Niagara/Source/NiagaraEditor/Private/NiagaraNodeConvert.cpp b/Engine/Plugins/FX/Niagara/Source/NiagaraEditor/Private/NiagaraNodeConvert.cpp
--- a/Engine/Plugins/FX/Niagara/Source/NiagaraEditor/Private/NiagaraNodeConvert.cpp
+++ b/Engine/Plugins/FX/Niagara/Source/NiagaraEditor/Private/NiagaraNodeConvert.cpp
@@ -110,7 +110,7 @@ struct FNiagaraConvertEntry
 
 			 if (StructProperty != nullptr)
 			{
-				CreateEntries(Schema, InPinId, InPin, FNiagaraTypeHelper::FindNiagaraFriendlyTopLevelStruct(StructProperty->Struct), OutEntries[Index].Children);
+				CreateEntries(Schema, InPinId, InPin, FNiagaraTypeHelper::FindNiagaraFriendlyTopLevelStruct(StructProperty->Struct, ENiagaraStructConversion::UserFacing), OutEntries[Index].Children);
 			}
 		}
 	}
diff --git a/Engine/Plugins/FX/Niagara/Source/NiagaraEditor/Private/NiagaraNodeDataSetBase.cpp b/Engine/Plugins/FX/Niagara/Source/NiagaraEditor/Private/NiagaraNodeDataSetBase.cpp
--- a/Engine/Plugins/FX/Niagara/Source/NiagaraEditor/Private/NiagaraNodeDataSetBase.cpp
+++ b/Engine/Plugins/FX/Niagara/Source/NiagaraEditor/Private/NiagaraNodeDataSetBase.cpp
@@ -230,6 +230,11 @@ bool UNiagaraNodeDataSetBase::GetSupportedNiagaraTypeDef(const FProperty* Proper
 		TypeDef = FNiagaraTypeDefinition::GetFloatDef();
 		return true;
 	}
+	else if (Property->IsA(FDoubleProperty::StaticClass()))
+	{
+		TypeDef = FNiagaraTypeDefinition::GetFloatDef();
+		return true;
+	}
 	else if (Property->IsA(FBoolProperty::StaticClass()))
 	{
 		TypeDef = FNiagaraTypeDefinition::GetBoolDef();
@@ -272,7 +277,7 @@ bool UNiagaraNodeDataSetBase::GetSupportedNiagaraTypeDef(const FProperty* Proper
 	}
 	else if (StructProp && StructProp->Struct)
 	{
-		TypeDef = FNiagaraTypeDefinition(FNiagaraTypeHelper::FindNiagaraFriendlyTopLevelStruct(StructProp->Struct));
+		TypeDef = FNiagaraTypeDefinition(FNiagaraTypeHelper::FindNiagaraFriendlyTopLevelStruct(StructProp->Struct, ENiagaraStructConversion::UserFacing));
 		return FNiagaraTypeRegistry::GetRegisteredPayloadTypes().Contains(TypeDef);
 	}
 	return false;
diff --git a/Engine/Plugins/FX/Niagara/Source/NiagaraEditor/Private/NiagaraNodeParameterMapBase.h b/Engine/Plugins/FX/Niagara/Source/NiagaraEditor/Private/NiagaraNodeParameterMapBase.h
--- a/Engine/Plugins/FX/Niagara/Source/NiagaraEditor/Private/NiagaraNodeParameterMapBase.h
+++ b/Engine/Plugins/FX/Niagara/Source/NiagaraEditor/Private/NiagaraNodeParameterMapBase.h
@@ -10,7 +10,7 @@
 
 class UEdGraphPin;
 
-/** A node which allows the user to build a set of arbitrary output types from an arbitrary set of input types by connecting their inner components. */
+/** Base node for parameter maps. */
 UCLASS()
 class UNiagaraNodeParameterMapBase : public UNiagaraNodeWithDynamicPins
 {
diff --git a/Engine/Plugins/FX/Niagara/Source/NiagaraEditor/Private/NiagaraStackFunctionInputBinder.cpp b/Engine/Plugins/FX/Niagara/Source/NiagaraEditor/Private/NiagaraStackFunctionInputBinder.cpp
--- a/Engine/Plugins/FX/Niagara/Source/NiagaraEditor/Private/NiagaraStackFunctionInputBinder.cpp
+++ b/Engine/Plugins/FX/Niagara/Source/NiagaraEditor/Private/NiagaraStackFunctionInputBinder.cpp
@@ -243,7 +243,7 @@ TArray<uint8> FNiagaraStackFunctionInputBinder::GetData() const
 	Data.AddUninitialized(InputType.GetSize());
 	if (RapidIterationParameter.IsValid() && Script->RapidIterationParameters.IndexOf(RapidIterationParameter) != -1)
 	{
-		FMemory::Memcpy(Data.GetData(), Script->RapidIterationParameters.GetParameterData(RapidIterationParameter), InputType.GetSize());
+		Script->RapidIterationParameters.CopyParameterData(RapidIterationParameter, Data.GetData());
 	}
 	else
 	{
diff --git a/Engine/Plugins/FX/Niagara/Source/NiagaraEditor/Private/ViewModels/Stack/NiagaraStackFunctionInput.cpp b/Engine/Plugins/FX/Niagara/Source/NiagaraEditor/Private/ViewModels/Stack/NiagaraStackFunctionInput.cpp
--- a/Engine/Plugins/FX/Niagara/Source/NiagaraEditor/Private/ViewModels/Stack/NiagaraStackFunctionInput.cpp
+++ b/Engine/Plugins/FX/Niagara/Source/NiagaraEditor/Private/ViewModels/Stack/NiagaraStackFunctionInput.cpp
@@ -1044,21 +1044,20 @@ void UNiagaraStackFunctionInput::RefreshValues()
 			// otherwise copy the value from the default.
 			InputValues.Mode = EValueMode::Local;
 			InputValues.LocalStruct = MakeShared<FStructOnScope>(InputType.GetStruct());
-			const uint8* RapidIterationParameterData = SourceScript->RapidIterationParameters.GetParameterData(RapidIterationParameter);
-			if (RapidIterationParameterData == nullptr && DefaultInputValues.LocalStruct.IsValid())
-			{
-				RapidIterationParameterData = DefaultInputValues.LocalStruct->GetStructMemory();
-			}
 
-			if (InputType.GetSize() > 0 && InputValues.LocalStruct->GetStructMemory() && RapidIterationParameterData)
-			{
-				FMemory::Memcpy(InputValues.LocalStruct->GetStructMemory(), RapidIterationParameterData, InputType.GetSize());
-			}
-			else
+			uint8* DestinationData = InputValues.LocalStruct->GetStructMemory();
+			if (!SourceScript->RapidIterationParameters.CopyParameterData(RapidIterationParameter, DestinationData))
 			{
-				UE_LOG(LogNiagaraEditor, Warning, TEXT("Type %s has no data! Cannot refresh values."), *InputType.GetName())
+				if (InputType.GetSize() > 0 && DestinationData && DefaultInputValues.LocalStruct.IsValid())
+				{
+					FMemory::Memcpy(DestinationData, DefaultInputValues.LocalStruct->GetStructMemory(), InputType.GetSize());
+				}
+				else
+				{
+					UE_LOG(LogNiagaraEditor, Warning, TEXT("Type %s has no data! Cannot refresh values."), *InputType.GetName())
+				}
 			}
-
+			
 			// we check if variable guid is already available in the parameter store and update it if that's not the case
 			if (InputMetaData.IsSet())
 			{
diff --git a/Engine/Plugins/FX/Niagara/Source/NiagaraEditor/Private/ViewModels/Stack/NiagaraStackParameterStoreEntry.cpp b/Engine/Plugins/FX/Niagara/Source/NiagaraEditor/Private/ViewModels/Stack/NiagaraStackParameterStoreEntry.cpp
--- a/Engine/Plugins/FX/Niagara/Source/NiagaraEditor/Private/ViewModels/Stack/NiagaraStackParameterStoreEntry.cpp
+++ b/Engine/Plugins/FX/Niagara/Source/NiagaraEditor/Private/ViewModels/Stack/NiagaraStackParameterStoreEntry.cpp
@@ -485,8 +485,8 @@ TSharedPtr<FNiagaraVariable> UNiagaraStackParameterStoreEntry::GetCurrentValueVa
 	if (InputType.GetClass() == nullptr)
 	{
 		FNiagaraVariable DefaultVariable(InputType, ParameterName);
-		const uint8* Data = ParameterStore->GetParameterData(DefaultVariable);
-		DefaultVariable.SetData(Data);
+		DefaultVariable.AllocateData();
+		ParameterStore->CopyParameterData(DefaultVariable, DefaultVariable.GetData());
 		return MakeShared<FNiagaraVariable>(DefaultVariable);
 	}
 	return TSharedPtr<FNiagaraVariable>();
diff --git a/Engine/Plugins/FX/Niagara/Source/NiagaraEditor/Private/Widgets/SNiagaraSpreadsheetView.cpp b/Engine/Plugins/FX/Niagara/Source/NiagaraEditor/Private/Widgets/SNiagaraSpreadsheetView.cpp
--- a/Engine/Plugins/FX/Niagara/Source/NiagaraEditor/Private/Widgets/SNiagaraSpreadsheetView.cpp
+++ b/Engine/Plugins/FX/Niagara/Source/NiagaraEditor/Private/Widgets/SNiagaraSpreadsheetView.cpp
@@ -1244,7 +1244,7 @@ void SNiagaraSpreadsheetView::GenerateLayoutInfo(FNiagaraTypeLayoutInfo& Layout,
 		else if (Property->IsA(FStructProperty::StaticClass()))
 		{
 			FStructProperty* StructProp = CastFieldChecked<FStructProperty>(Property);
-			GenerateLayoutInfo(Layout, FNiagaraTypeHelper::FindNiagaraFriendlyTopLevelStruct(StructProp->Struct), nullptr, PropertyName, PropertyNames, FieldInfo);
+			GenerateLayoutInfo(Layout, FNiagaraTypeHelper::FindNiagaraFriendlyTopLevelStruct(StructProp->Struct, ENiagaraStructConversion::Simulation), nullptr, PropertyName, PropertyNames, FieldInfo);
 		}
 		else
 		{
diff --git a/Engine/Plugins/FX/Niagara/Source/NiagaraEditor/Public/EdGraphSchema_Niagara.h b/Engine/Plugins/FX/Niagara/Source/NiagaraEditor/Public/EdGraphSchema_Niagara.h
--- a/Engine/Plugins/FX/Niagara/Source/NiagaraEditor/Public/EdGraphSchema_Niagara.h
+++ b/Engine/Plugins/FX/Niagara/Source/NiagaraEditor/Public/EdGraphSchema_Niagara.h
@@ -136,7 +136,7 @@ class NIAGARAEDITOR_API UEdGraphSchema_Niagara : public UEdGraphSchema
 	  * bNeedsValue Whether or not the returned variable must be allocated to a valid value. When true if the pin doesn't have a valid default value itself the variable will be reset to default before return.
 	  * returns The newly created variable.
 	  */
-	FNiagaraVariable PinToNiagaraVariable(const UEdGraphPin* Pin, bool bNeedsValue=false)const;
+	FNiagaraVariable PinToNiagaraVariable(const UEdGraphPin* Pin, bool bNeedsValue=false, ENiagaraStructConversion StructConversion = ENiagaraStructConversion::UserFacing) const;
 
 	/** 
 	  * Tries to get a default value string for a graph pin from a niagara variable.
@@ -147,8 +147,8 @@ class NIAGARAEDITOR_API UEdGraphSchema_Niagara : public UEdGraphSchema
 	bool TryGetPinDefaultValueFromNiagaraVariable(const FNiagaraVariable& Variable, FString& OutPinDefaultValue) const;
 
 	static void ConvertIllegalPinsInPlace(UEdGraphPin* Pin);
-	static FNiagaraTypeDefinition PinToTypeDefinition(const UEdGraphPin* Pin);
-	static FNiagaraTypeDefinition PinTypeToTypeDefinition(const FEdGraphPinType& PinType) ;
+	static FNiagaraTypeDefinition PinToTypeDefinition(const UEdGraphPin* Pin, ENiagaraStructConversion StructConversion = ENiagaraStructConversion::UserFacing);
+	static FNiagaraTypeDefinition PinTypeToTypeDefinition(const FEdGraphPinType& PinType);
 	static FEdGraphPinType TypeDefinitionToPinType(FNiagaraTypeDefinition TypeDef);
 	
 	static bool IsPinWildcard(const UEdGraphPin* Pin);
