- Fix for implicit cast as explicit is required now

#rnx

diff --git a/Engine/Plugins/FX/Niagara/Source/Niagara/Private/NiagaraTypes.cpp b/Engine/Plugins/FX/Niagara/Source/Niagara/Private/NiagaraTypes.cpp
--- a/Engine/Plugins/FX/Niagara/Source/Niagara/Private/NiagaraTypes.cpp
+++ b/Engine/Plugins/FX/Niagara/Source/Niagara/Private/NiagaraTypes.cpp
@@ -86,7 +86,7 @@ void FNiagaraStructConversionStep::CopyToSim(uint8* DestinationData, const uint8
 	else if (ConversionType == ENiagaraStructConversionType::Vector4)
 	{
 		checkf(SourceBytes == 32 && SimulationBytes == 16, TEXT("Wrong bytesizes for Vector4 d->f conversion: Source %i, Simulation %i"), SourceBytes, SimulationBytes);
-		FVector4f Vector4 = *static_cast<const FVector4d*>(Src);
+		FVector4f Vector4 = FVector4f(*static_cast<const FVector4d*>(Src));
 		FMemory::Memcpy(Dest, &Vector4, SimulationBytes);
 	}
 	else if (ConversionType == ENiagaraStructConversionType::Quat)
@@ -127,7 +127,7 @@ void FNiagaraStructConversionStep::CopyFromSim(uint8* DestinationData, const uin
 	else if (ConversionType == ENiagaraStructConversionType::Vector4)
 	{
 		checkf(SourceBytes == 32 && SimulationBytes == 16, TEXT("Wrong bytesizes for Vector4 f->d conversion: Source %i, Simulation %i"), SourceBytes, SimulationBytes);
-		FVector4d Vector4 = *static_cast<const FVector4f*>(Src);
+		FVector4d Vector4 = FVector4d(*static_cast<const FVector4f*>(Src));
 		FMemory::Memcpy(Dest, &Vector4, SourceBytes);
 	}
 	else if (ConversionType == ENiagaraStructConversionType::Quat)
