- Add check to catch destructing parameter stores on the render thread which have not been unbound

#rnx

diff --git a/Engine/Plugins/FX/Niagara/Source/Niagara/Private/NiagaraParameterStore.cpp b/Engine/Plugins/FX/Niagara/Source/Niagara/Private/NiagaraParameterStore.cpp
--- a/Engine/Plugins/FX/Niagara/Source/Niagara/Private/NiagaraParameterStore.cpp
+++ b/Engine/Plugins/FX/Niagara/Source/Niagara/Private/NiagaraParameterStore.cpp
@@ -225,6 +225,8 @@ FNiagaraParameterStore::~FNiagaraParameterStore()
 {
 	DEC_MEMORY_STAT_BY(STAT_NiagaraParamStoreMemory, ParameterData.GetAllocatedSize());
 	
+	checkf(IsInGameThread() || (IsInRenderingThread() && SourceStores.Num() == 0), TEXT("ParameterStore destruction must be on the GameThread(%d) or RenderThread(%d) and if on render must be unbound.  SourceStores(%d)"), IsInGameThread() ? 1 : 0, IsInRenderingThread() ? 1 : 0, SourceStores.Num());
+
 	UnbindAll();
 }
 
