Fixed a few bugs with the search bar in the generated code view

diff --git a/Engine/Plugins/FX/Niagara/Source/NiagaraEditor/Private/Widgets/SNiagaraGeneratedCodeView.cpp b/Engine/Plugins/FX/Niagara/Source/NiagaraEditor/Private/Widgets/SNiagaraGeneratedCodeView.cpp
--- a/Engine/Plugins/FX/Niagara/Source/NiagaraEditor/Private/Widgets/SNiagaraGeneratedCodeView.cpp
+++ b/Engine/Plugins/FX/Niagara/Source/NiagaraEditor/Private/Widgets/SNiagaraGeneratedCodeView.cpp
@@ -83,7 +83,7 @@ void SNiagaraGeneratedCodeView::Construct(const FArguments& InArgs, TSharedRef<F
 			SNew(SButton)
 			.ButtonStyle(FEditorStyle::Get(), "HoverHintOnly")
 			.IsFocusable(false)
-			.ToolTipText(LOCTEXT("UpToolTip", "Focus to next found search term"))
+			.ToolTipText(LOCTEXT("UpToolTip", "Focus to previous found search term"))
 			.OnClicked(this, &SNiagaraGeneratedCodeView::SearchUpClicked)
 			.Content()
 			[
@@ -232,8 +232,7 @@ FReply SNiagaraGeneratedCodeView::SearchDownClicked()
 			CurrentFoundTextEntry = 0;
 		}
 	}
-	
-	GeneratedCode[TabState].Text->AdvanceSearch(true);
+	GeneratedCode[TabState].Text->AdvanceSearch(false);
 
 	SetSearchMofN();
 
@@ -250,7 +249,7 @@ FReply SNiagaraGeneratedCodeView::SearchUpClicked()
 			CurrentFoundTextEntry = ActiveFoundTextEntries.Num() - 1;
 		}
 	}
-	GeneratedCode[TabState].Text->AdvanceSearch(false);
+	GeneratedCode[TabState].Text->AdvanceSearch(true);
 	
 	SetSearchMofN();
 
@@ -276,7 +275,6 @@ void SNiagaraGeneratedCodeView::DoSearch(const FText& InFilterText)
 	const FText OldText = GeneratedCode[TabState].Text->GetSearchText();
 	GeneratedCode[TabState].Text->SetSearchText(InFilterText);
 	GeneratedCode[TabState].Text->BeginSearch(InFilterText, ESearchCase::IgnoreCase, false);
-	InFilterText.ToString();
 
 	FString SearchString = InFilterText.ToString();
 	ActiveFoundTextEntries.Empty();
@@ -289,13 +287,20 @@ void SNiagaraGeneratedCodeView::DoSearch(const FText& InFilterText)
 	ActiveFoundTextEntries.Empty();
 	for (int32 i = 0; i < GeneratedCode[TabState].HlslByLines.Num(); i++)
 	{
-		int32 LastPos = INDEX_NONE;
-		int32 FoundPos = GeneratedCode[TabState].HlslByLines[i].Find(SearchString, ESearchCase::IgnoreCase, ESearchDir::FromStart, LastPos);
-		while (FoundPos != INDEX_NONE)
+		const FString& Line = GeneratedCode[TabState].HlslByLines[i];
+		int32 FoundPos = Line.Find(SearchString, ESearchCase::IgnoreCase);
+		while (FoundPos != INDEX_NONE && ActiveFoundTextEntries.Num() < 1000) // guard against a runaway loop
 		{
 			ActiveFoundTextEntries.Add(FTextLocation(i, FoundPos));
-			LastPos = FoundPos + 1;
-			FoundPos = GeneratedCode[TabState].HlslByLines[i].Find(SearchString, ESearchCase::IgnoreCase, ESearchDir::FromStart, LastPos);
+			int32 LastPos = FoundPos + SearchString.Len();
+			if (LastPos < Line.Len())
+			{
+				FoundPos = Line.Find(SearchString, ESearchCase::IgnoreCase, ESearchDir::FromStart, LastPos);
+			}
+			else
+			{
+				FoundPos = INDEX_NONE;
+			}
 		}
 	}
 
