Fix for Niagara DI containing double/LWC properties (eg Chaos DI) breaking while trying to fix up FNiagaraVariable

diff --git a/Engine/Plugins/FX/Niagara/Source/Niagara/Private/NiagaraTypes.cpp b/Engine/Plugins/FX/Niagara/Source/Niagara/Private/NiagaraTypes.cpp
--- a/Engine/Plugins/FX/Niagara/Source/Niagara/Private/NiagaraTypes.cpp
+++ b/Engine/Plugins/FX/Niagara/Source/Niagara/Private/NiagaraTypes.cpp
@@ -366,7 +366,7 @@ bool FNiagaraTypeHelper::IsLWCStructure(UStruct* InStruct)
 
 bool FNiagaraTypeHelper::IsLWCType(const FNiagaraTypeDefinition& InType)
 {
-	return InType.IsValid() && IsLWCStructure(InType.GetStruct());
+	return InType.IsValid() && !InType.IsUObject() && IsLWCStructure(InType.GetStruct());
 }
 
 UScriptStruct* FNiagaraTypeHelper::GetSWCStruct(UScriptStruct* LWCStruct)
diff --git a/Engine/Plugins/FX/Niagara/Source/NiagaraCore/Private/NiagaraCustomVersion.cpp b/Engine/Plugins/FX/Niagara/Source/NiagaraCore/Private/NiagaraCustomVersion.cpp
--- a/Engine/Plugins/FX/Niagara/Source/NiagaraCore/Private/NiagaraCustomVersion.cpp
+++ b/Engine/Plugins/FX/Niagara/Source/NiagaraCore/Private/NiagaraCustomVersion.cpp
@@ -9,4 +9,4 @@ const FGuid FNiagaraCustomVersion::GUID(0xFCF57AFA, 0x50764283, 0xB9A9E658, 0xFF
 FCustomVersionRegistration GRegisterNiagaraCustomVersion(FNiagaraCustomVersion::GUID, FNiagaraCustomVersion::LatestVersion, TEXT("NiagaraVer"));
 
 // Note: When encountering a conflict on this file please generate a new GUID
-const FGuid FNiagaraCustomVersion::LatestScriptCompileVersion(0x0C09B728, 0x5EE44F4A, 0x9E135A16, 0xAA8A0160);
\ No newline at end of file
+const FGuid FNiagaraCustomVersion::LatestScriptCompileVersion(0x2CC3E164, 0x5BCD4A0E, 0x9ACCF876, 0x825A0250);
\ No newline at end of file
