Fixed a bug where inputs with a changed type don't refresh after versioning up

diff --git a/Engine/Plugins/FX/Niagara/Source/NiagaraEditor/Private/NiagaraNodeFunctionCall.cpp b/Engine/Plugins/FX/Niagara/Source/NiagaraEditor/Private/NiagaraNodeFunctionCall.cpp
--- a/Engine/Plugins/FX/Niagara/Source/NiagaraEditor/Private/NiagaraNodeFunctionCall.cpp
+++ b/Engine/Plugins/FX/Niagara/Source/NiagaraEditor/Private/NiagaraNodeFunctionCall.cpp
@@ -1426,11 +1426,12 @@ void UNiagaraNodeFunctionCall::ChangeScriptVersion(FGuid NewScriptVersion, const
 	if (FunctionScript)
 	{
 		// Automatically remove old inputs so it does not show a bunch of warnings to the user
-		TArray<FName> FunctionInputNames;
+		TMap<FName, FEdGraphPinType> FunctionInputNames;
 		FPinCollectorArray OverridePins;
 		UNiagaraNodeParameterMapSet* OverrideNode = FNiagaraStackGraphUtilities::GetStackFunctionOverrideNode(*this);
 		if (OverrideNode != nullptr)
 		{
+			OverrideNode->Modify();
 			OverrideNode->GetInputPins(OverridePins);
 			
 			TSet<const UEdGraphPin*> HiddenModulePins;
@@ -1438,14 +1439,26 @@ void UNiagaraNodeFunctionCall::ChangeScriptVersion(FGuid NewScriptVersion, const
 			GetStackFunctionInputPins(*this, ModuleInputPins, HiddenModulePins, UpgradeContext.ConstantResolver, FNiagaraStackGraphUtilities::ENiagaraGetStackFunctionInputPinsOptions::ModuleInputsOnly);
 			for (const UEdGraphPin* InputPin : ModuleInputPins)
 			{
-				FunctionInputNames.Add(FNiagaraParameterHandle(InputPin->PinName).GetName());
+				FunctionInputNames.Add(FNiagaraParameterHandle(InputPin->PinName).GetName(), InputPin->PinType);
 			}
 		}
 		for (UEdGraphPin* OverridePin : OverridePins)
 		{
+			if (!FNiagaraStackGraphUtilities::IsOverridePinForFunction(*OverridePin, *this))
+			{
+				continue;
+			}
 			FName InputName = FNiagaraParameterHandle(OverridePin->PinName).GetName();
-			if (FNiagaraStackGraphUtilities::IsOverridePinForFunction(*OverridePin, *this) && FunctionInputNames.Contains(InputName) == false)
+			FEdGraphPinType* PinType = FunctionInputNames.Find(InputName);
+			if (PinType && *PinType != OverridePin->PinType)
+			{
+				// looks like the type of the module input changed - we'll change the override pin type as well
+				OverridePin->Modify();
+				OverridePin->PinType = *PinType;
+			}
+			else if (PinType == nullptr)
 			{
+				// the input doesn't exist any more, delete the override pin
 				TArray<TWeakObjectPtr<UNiagaraDataInterface>> RemovedDataObjects;
 				FNiagaraStackGraphUtilities::RemoveNodesForStackFunctionInputOverridePin(*OverridePin, RemovedDataObjects);
 				OverrideNode->RemovePin(OverridePin);
diff --git a/Engine/Plugins/FX/Niagara/Source/NiagaraEditor/Private/ViewModels/Stack/NiagaraStackInputCategory.cpp b/Engine/Plugins/FX/Niagara/Source/NiagaraEditor/Private/ViewModels/Stack/NiagaraStackInputCategory.cpp
--- a/Engine/Plugins/FX/Niagara/Source/NiagaraEditor/Private/ViewModels/Stack/NiagaraStackInputCategory.cpp
+++ b/Engine/Plugins/FX/Niagara/Source/NiagaraEditor/Private/ViewModels/Stack/NiagaraStackInputCategory.cpp
@@ -41,7 +41,7 @@ void UNiagaraStackInputCategory::RefreshChildrenInternal(const TArray<UNiagaraSt
 	{
 		UNiagaraStackFunctionInput* InputChild = FindCurrentChildOfTypeByPredicate<UNiagaraStackFunctionInput>(CurrentChildren, [&](UNiagaraStackFunctionInput* CurrentInput) 
 		{ 
-			return CurrentInput->GetInputParameterHandle() == Input.ParameterHandle && CurrentInput->GetInputFunctionCallInitialScript() == Input.InputFunctionCallNode->FunctionScript;
+			return CurrentInput->GetInputParameterHandle() == Input.ParameterHandle && CurrentInput->GetInputType() == Input.Type && CurrentInput->GetInputFunctionCallInitialScript() == Input.InputFunctionCallNode->FunctionScript;
 		});
 
 		if (InputChild == nullptr)
