Niagara - When allocating pins on function calls for static switches, assign their persisten guid from their script variable so that the pin default value is retained when reallocating pins after a rename.

diff --git a/Engine/Plugins/FX/Niagara/Source/NiagaraEditor/Private/NiagaraNodeFunctionCall.cpp b/Engine/Plugins/FX/Niagara/Source/NiagaraEditor/Private/NiagaraNodeFunctionCall.cpp
--- a/Engine/Plugins/FX/Niagara/Source/NiagaraEditor/Private/NiagaraNodeFunctionCall.cpp
+++ b/Engine/Plugins/FX/Niagara/Source/NiagaraEditor/Private/NiagaraNodeFunctionCall.cpp
@@ -303,6 +303,8 @@ void UNiagaraNodeFunctionCall::AllocateDefaultPins()
 				{
 					NewPin->DefaultValue = PinDefaultValue;
 				}
+
+				NewPin->PersistentGuid = ScriptVar->Metadata.GetVariableGuid();
 			}
 			else
 			{
