Fix for issues with undo/redo and Niagara compilation.
Undo/Redo transactions can directly modify the CachedCompileId which can confuse the compilation and regeneration of derived data which we use for execution.
Clearing the system's compile results on post undo to ensure that we always recompile/pull from DDC and regenerate our exec data.

diff --git a/Engine/Plugins/FX/Niagara/Source/NiagaraEditor/Private/ViewModels/NiagaraSystemViewModel.cpp b/Engine/Plugins/FX/Niagara/Source/NiagaraEditor/Private/ViewModels/NiagaraSystemViewModel.cpp
--- a/Engine/Plugins/FX/Niagara/Source/NiagaraEditor/Private/ViewModels/NiagaraSystemViewModel.cpp
+++ b/Engine/Plugins/FX/Niagara/Source/NiagaraEditor/Private/ViewModels/NiagaraSystemViewModel.cpp
@@ -646,10 +646,11 @@ void FNiagaraSystemViewModel::PostUndo(bool bSuccess)
 {
 	// Reset system stack and emitter handle view models to prevent accessing invalid data if they were in the undo operation.
 
-	//We have to force a reinit of all systems in the event of an undo.
-	FNiagaraSystemUpdateContext UpdateContext;
-	UpdateContext.SetDestroyOnAdd(true);
-	UpdateContext.Add(System, true);
+	//We have to invalidate all cached data that may have been directly modified by the transaction.
+	System->ForEachScript([](UNiagaraScript* Script)
+	{
+		Script->InvalidateCompileResults(TEXT("Post Undo"));
+	});
 
 	SystemStackViewModel->Reset();
 	ResetEmitterHandleViewModelsAndTracks();
