Undo //UE5/Release-5.0/Engine/Plugins/FX/Niagara/Source/NiagaraEditor/Private/... changelist 18795548

Undoing this change for now for additional testing.

diff --git a/Engine/Plugins/FX/Niagara/Source/NiagaraEditor/Private/NiagaraNodeFunctionCall.cpp b/Engine/Plugins/FX/Niagara/Source/NiagaraEditor/Private/NiagaraNodeFunctionCall.cpp
--- a/Engine/Plugins/FX/Niagara/Source/NiagaraEditor/Private/NiagaraNodeFunctionCall.cpp
+++ b/Engine/Plugins/FX/Niagara/Source/NiagaraEditor/Private/NiagaraNodeFunctionCall.cpp
@@ -303,8 +303,6 @@ void UNiagaraNodeFunctionCall::AllocateDefaultPins()
 				{
 					NewPin->DefaultValue = PinDefaultValue;
 				}
-
-				NewPin->PersistentGuid = ScriptVar->Metadata.GetVariableGuid();
 			}
 			else
 			{
