Strip data from NiagaraSystems when AV data is not required (server builds)

diff --git a/Engine/Plugins/FX/Niagara/Source/Niagara/Classes/NiagaraSystem.h b/Engine/Plugins/FX/Niagara/Source/Niagara/Classes/NiagaraSystem.h
--- a/Engine/Plugins/FX/Niagara/Source/Niagara/Classes/NiagaraSystem.h
+++ b/Engine/Plugins/FX/Niagara/Source/Niagara/Classes/NiagaraSystem.h
@@ -800,6 +800,8 @@ private:
 	void InitEmitterDataSetCompiledData(FNiagaraDataSetCompiledData& DataSetToInit, const UNiagaraEmitter* InAssociatedEmitter, const FNiagaraEmitterHandle& InAssociatedEmitterHandle);
 	void PrepareRapidIterationParametersForCompilation();
 
+	/** Resets internal data leaving it in a state which would have minimal cost to exist in headless builds (servers) */
+	void ResetToEmptySystem();
 #endif
 
 	void ResolveScalabilitySettings();
diff --git a/Engine/Plugins/FX/Niagara/Source/Niagara/Private/NiagaraEmitter.cpp b/Engine/Plugins/FX/Niagara/Source/Niagara/Private/NiagaraEmitter.cpp
--- a/Engine/Plugins/FX/Niagara/Source/Niagara/Private/NiagaraEmitter.cpp
+++ b/Engine/Plugins/FX/Niagara/Source/Niagara/Private/NiagaraEmitter.cpp
@@ -1407,6 +1407,14 @@ void UNiagaraEmitter::UpdateEmitterAfterLoad()
 #if WITH_EDITORONLY_DATA
 	check(IsInGameThread());
 
+	// We remove emitters and scripts on dedicated servers (and platforms which don't use AV data), so skip further work.
+	const bool bIsDedicatedServer = !GIsClient && GIsServer;
+	const bool bTargetRequiresAvData = WillNeedAudioVisualData();
+	if (bIsDedicatedServer || !bTargetRequiresAvData)
+	{
+		return;
+	}
+
 	// Synchronize with definitions before merging.
 	PostLoadDefinitionsSubscriptions();
 
diff --git a/Engine/Plugins/FX/Niagara/Source/Niagara/Private/NiagaraSystem.cpp b/Engine/Plugins/FX/Niagara/Source/Niagara/Private/NiagaraSystem.cpp
--- a/Engine/Plugins/FX/Niagara/Source/Niagara/Private/NiagaraSystem.cpp
+++ b/Engine/Plugins/FX/Niagara/Source/Niagara/Private/NiagaraSystem.cpp
@@ -721,11 +721,18 @@ void UNiagaraSystem::UpdateSystemAfterLoad()
 	}
 
 #if WITH_EDITORONLY_DATA
-	// We remove emitters and scripts on dedicated servers, so skip further work.
+
+	// We remove emitters and scripts on dedicated servers (and platforms which don't use AV data), so skip further work.
 	const bool bIsDedicatedServer = !GIsClient && GIsServer;
+	const bool bTargetRequiresAvData = WillNeedAudioVisualData();
+	if (bIsDedicatedServer || !bTargetRequiresAvData)
+	{
+		ResetToEmptySystem();
+		return;
+	}
 
 	TArray<UNiagaraScript*> AllSystemScripts;
-	if (!GetOutermost()->bIsCookedForEditor && !bIsDedicatedServer)
+	if (!GetOutermost()->bIsCookedForEditor)
 	{
 		UNiagaraScriptSourceBase* SystemScriptSource;
 		if (SystemSpawnScript == nullptr)
@@ -3464,6 +3471,24 @@ void UNiagaraSystem::InitSystemCompiledData()
 
 }
 
+void UNiagaraSystem::ResetToEmptySystem()
+{
+	EffectType = nullptr;
+	EmitterHandles.Empty();
+	ParameterCollectionOverrides.Empty();
+	SystemSpawnScript = nullptr;
+	SystemUpdateScript = nullptr;
+	EmitterCompiledData.Empty();
+	SystemCompiledData = FNiagaraSystemCompiledData();
+	UserDINamesReadInSystemScripts.Empty();
+	EmitterExecutionOrder.Empty();
+	RendererPostTickOrder.Empty();
+	RendererCompletionOrder.Empty();
+
+	// while we'd like to remove these as well, BP will generate warnings for missing parameters
+	//ExposedParameters = FNiagaraUserRedirectionParameterStore();
+}
+
 #endif
 
 TStatId UNiagaraSystem::GetStatID(bool bGameThread, bool bConcurrent)const
