- Fix CIS issues

#rnx

diff --git a/Engine/Plugins/FX/Niagara/Source/Niagara/Private/NiagaraTypes.cpp b/Engine/Plugins/FX/Niagara/Source/Niagara/Private/NiagaraTypes.cpp
--- a/Engine/Plugins/FX/Niagara/Source/Niagara/Private/NiagaraTypes.cpp
+++ b/Engine/Plugins/FX/Niagara/Source/Niagara/Private/NiagaraTypes.cpp
@@ -254,7 +254,7 @@ FNiagaraLwcStructConverter BuildSWCStructure(UScriptStruct* NewStruct, UScriptSt
 		FProperty* NewProperty = nullptr;
 		if (const FDoubleProperty* ChildAsDouble = CastField<const FDoubleProperty>(ChildProperty))
 		{
-			AlignedOffset = Align(AlignedOffset, NewProperty->GetMinAlignment());
+			AlignedOffset = Align(AlignedOffset, ChildAsDouble->GetMinAlignment());
 			NewProperty = new FFloatProperty(NewStruct, ChildProperty->GetFName(), RF_Public);
 			StructConverter.AddConversionStep(ChildAsDouble->GetSize(), ChildAsDouble->GetOffset_ForInternal(), NewProperty->GetSize(), AlignedOffset, ENiagaraStructConversionType::DoubleToFloat);
 		}
