Fixed nullptr checks on some summary view items

diff --git a/Engine/Plugins/FX/Niagara/Source/NiagaraEditor/Private/NiagaraScriptMergeManager.cpp b/Engine/Plugins/FX/Niagara/Source/NiagaraEditor/Private/NiagaraScriptMergeManager.cpp
--- a/Engine/Plugins/FX/Niagara/Source/NiagaraEditor/Private/NiagaraScriptMergeManager.cpp
+++ b/Engine/Plugins/FX/Niagara/Source/NiagaraEditor/Private/NiagaraScriptMergeManager.cpp
@@ -3357,23 +3357,20 @@ FNiagaraScriptMergeManager::FApplyDiffResults FNiagaraScriptMergeManager::ApplyI
 	check(EditorData != GetDefault<UNiagaraEmitterEditorData>());
 	check(EditorData->GetOuter() == &Emitter);
 	
-	if (EditorData)
+	EditorData->Modify();
+	for (TSharedRef<FNiagaraInputSummaryMergeAdapter> RemovedInput : DiffResults.RemovedInputSummaryEntries)
 	{
-		EditorData->Modify();
-		for (TSharedRef<FNiagaraInputSummaryMergeAdapter> RemovedInput : DiffResults.RemovedInputSummaryEntries)
-		{
-			EditorData->SetSummaryViewMetaData(RemovedInput->GetKey(), FFunctionInputSummaryViewMetadata());
-		}
-	
-		for (TSharedRef<FNiagaraInputSummaryMergeAdapter> AddedInput : DiffResults.AddedInputSummaryEntries)
-		{
-			EditorData->SetSummaryViewMetaData(AddedInput->GetKey(), AddedInput->GetValue());
-		}
-	
-		for (TSharedRef<FNiagaraInputSummaryMergeAdapter> ModifiedInput : DiffResults.ModifiedOtherInputSummaryEntries)
-		{
-			EditorData->SetSummaryViewMetaData(ModifiedInput->GetKey(), ModifiedInput->GetValue());
-		}
+		EditorData->SetSummaryViewMetaData(RemovedInput->GetKey(), FFunctionInputSummaryViewMetadata());
+	}
+
+	for (TSharedRef<FNiagaraInputSummaryMergeAdapter> AddedInput : DiffResults.AddedInputSummaryEntries)
+	{
+		EditorData->SetSummaryViewMetaData(AddedInput->GetKey(), AddedInput->GetValue());
+	}
+
+	for (TSharedRef<FNiagaraInputSummaryMergeAdapter> ModifiedInput : DiffResults.ModifiedOtherInputSummaryEntries)
+	{
+		EditorData->SetSummaryViewMetaData(ModifiedInput->GetKey(), ModifiedInput->GetValue());
 	}
 
 	FApplyDiffResults Results;
diff --git a/Engine/Plugins/FX/Niagara/Source/NiagaraEditorWidgets/Private/Stack/SNiagaraStackFunctionInputValue.cpp b/Engine/Plugins/FX/Niagara/Source/NiagaraEditorWidgets/Private/Stack/SNiagaraStackFunctionInputValue.cpp
--- a/Engine/Plugins/FX/Niagara/Source/NiagaraEditorWidgets/Private/Stack/SNiagaraStackFunctionInputValue.cpp
+++ b/Engine/Plugins/FX/Niagara/Source/NiagaraEditorWidgets/Private/Stack/SNiagaraStackFunctionInputValue.cpp
@@ -1362,7 +1362,7 @@ const FSlateBrush* SNiagaraStackFunctionInputValue::GetFilteredViewIcon() const
 
 EVisibility SNiagaraStackFunctionInputValue::GetFilteredViewContextButtonVisibility() const
 {
-	UNiagaraEmitterEditorData* EditorData = (FunctionInput && FunctionInput->GetEmitterViewModel())? &FunctionInput->GetEmitterViewModel()->GetOrCreateEditorData() : nullptr;
+	UNiagaraEmitterEditorData* EditorData = FunctionInput->GetEmitterViewModel()? &FunctionInput->GetEmitterViewModel()->GetOrCreateEditorData() : nullptr;
 	
 	if (!EditorData || EditorData->ShouldShowSummaryView() || FunctionInput->GetTypedOuter<UNiagaraStackSummaryViewObject>() == nullptr || FunctionInput->GetEmitterViewModel()->GetSummaryIsInEditMode() == false)
 	{
@@ -1453,7 +1453,7 @@ TSharedRef<SWidget> SNiagaraStackFunctionInputValue::GetFilteredViewPropertiesCo
 
 FText SNiagaraStackFunctionInputValue::GetFilteredViewDisplayName() const
 {
-	const UNiagaraEmitterEditorData* EditorData = (FunctionInput && FunctionInput->GetEmitterViewModel())? &FunctionInput->GetEmitterViewModel()->GetEditorData() : nullptr;
+	const UNiagaraEmitterEditorData* EditorData = FunctionInput->GetEmitterViewModel()? &FunctionInput->GetEmitterViewModel()->GetEditorData() : nullptr;
 	UNiagaraStackFunctionInput* ParentInput = FNiagaraStackEditorWidgetsUtilities::GetParentInputForSummaryView(FunctionInput);
 	TOptional<FFunctionInputSummaryViewKey> Key = FNiagaraStackEditorWidgetsUtilities::GetSummaryViewInputKeyForFunctionInput(ParentInput);
 
