Fixed summary view context menu from dynamic inputs

diff --git a/Engine/Plugins/FX/Niagara/Source/NiagaraEditorWidgets/Private/NiagaraEditorWidgetsUtilities.cpp b/Engine/Plugins/FX/Niagara/Source/NiagaraEditorWidgets/Private/NiagaraEditorWidgetsUtilities.cpp
--- a/Engine/Plugins/FX/Niagara/Source/NiagaraEditorWidgets/Private/NiagaraEditorWidgetsUtilities.cpp
+++ b/Engine/Plugins/FX/Niagara/Source/NiagaraEditorWidgets/Private/NiagaraEditorWidgetsUtilities.cpp
@@ -560,9 +560,15 @@ TOptional<FFunctionInputSummaryViewKey> FNiagaraStackEditorWidgetsUtilities::Get
 UNiagaraStackFunctionInput* FNiagaraStackEditorWidgetsUtilities::GetParentInputForSummaryView(UNiagaraStackFunctionInput* FunctionInput)
 {
 	const UNiagaraEmitterEditorData* EditorData = (FunctionInput && FunctionInput->GetEmitterViewModel())? &FunctionInput->GetEmitterViewModel()->GetEditorData() : nullptr;
+
+	// Find outermost function input first, as this will traverse the chain of dynamic inputs to find the base input which is what needs to go to summary view
+	while (FunctionInput->GetTypedOuter<UNiagaraStackFunctionInput>())
+	{
+		FunctionInput = FunctionInput->GetTypedOuter<UNiagaraStackFunctionInput>();
+	}
 	
 	if (EditorData)
-	{		
+	{
 		TOptional<FFunctionInputSummaryViewKey> Key = FNiagaraStackEditorWidgetsUtilities::GetSummaryViewInputKeyForFunctionInput(FunctionInput);
 
 		TOptional<FNiagaraVariableMetaData> Metadata = FunctionInput->GetMetadata();
