Fix change namespace menu for parameter definitions to only allow valid namespaces.

diff --git a/Engine/Plugins/FX/Niagara/Source/NiagaraEditor/Private/NiagaraEditorSettings.cpp b/Engine/Plugins/FX/Niagara/Source/NiagaraEditor/Private/NiagaraEditorSettings.cpp
--- a/Engine/Plugins/FX/Niagara/Source/NiagaraEditor/Private/NiagaraEditorSettings.cpp
+++ b/Engine/Plugins/FX/Niagara/Source/NiagaraEditor/Private/NiagaraEditorSettings.cpp
@@ -129,6 +129,7 @@ void UNiagaraEditorSettings::SetupNamespaceMetadata()
 			.SetBackgroundColor(FLinearColor(FColor(136, 66, 65)))
 			.SetSortId(40)
 			.AddOption(ENiagaraNamespaceMetadataOptions::HideInSystem)
+			.AddOption(ENiagaraNamespaceMetadataOptions::HideInDefinitions)
 			.SetGuid(FNiagaraEditorGuids::ModuleNamespaceMetaDataGuid),
 		FNiagaraNamespaceMetadata({FNiagaraConstants::OutputNamespace}, FNiagaraConstants::ModuleNamespace)
 			.SetDisplayName(LOCTEXT("ModuleOutputDisplayName", "Output"))
@@ -138,6 +139,7 @@ void UNiagaraEditorSettings::SetupNamespaceMetadata()
 			.SetSortId(60)
 			.AddOption(ENiagaraNamespaceMetadataOptions::AdvancedInScript)
 			.AddOption(ENiagaraNamespaceMetadataOptions::AdvancedInSystem)
+			.AddOption(ENiagaraNamespaceMetadataOptions::HideInDefinitions)
 			.AddOption(ENiagaraNamespaceMetadataOptions::PreventCreatingInSystemEditor)
 			.SetGuid(FNiagaraEditorGuids::ModuleOutputNamespaceMetaDataGuid),
 		FNiagaraNamespaceMetadata({FNiagaraConstants::LocalNamespace, FNiagaraConstants::ModuleNamespace})
@@ -148,6 +150,7 @@ void UNiagaraEditorSettings::SetupNamespaceMetadata()
 			.SetForegroundStyle("NiagaraEditor.ParameterName.NamespaceTextDark")
 			.SetSortId(50)
 			.AddOption(ENiagaraNamespaceMetadataOptions::HideInSystem)
+			.AddOption(ENiagaraNamespaceMetadataOptions::HideInDefinitions)
 			.AddOption(ENiagaraNamespaceMetadataOptions::PreventEditingNamespaceModifier)
 			.SetGuid(FNiagaraEditorGuids::ModuleLocalNamespaceMetaDataGuid),
 		FNiagaraNamespaceMetadata({FNiagaraConstants::TransientNamespace})
@@ -158,6 +161,7 @@ void UNiagaraEditorSettings::SetupNamespaceMetadata()
 			.SetSortId(80)
 			.AddOption(ENiagaraNamespaceMetadataOptions::AdvancedInScript)
 			.AddOption(ENiagaraNamespaceMetadataOptions::AdvancedInSystem)
+			.AddOption(ENiagaraNamespaceMetadataOptions::HideInDefinitions)
 			.AddOption(ENiagaraNamespaceMetadataOptions::PreventEditingNamespaceModifier)
 			.SetGuid(FNiagaraEditorGuids::TransientNamespaceMetaDataGuid),
 		FNiagaraNamespaceMetadata({FNiagaraConstants::StackContextNamespace})
@@ -168,6 +172,7 @@ void UNiagaraEditorSettings::SetupNamespaceMetadata()
 			.SetSortId(80)
 			.AddOption(ENiagaraNamespaceMetadataOptions::AdvancedInScript)
 			.AddOption(ENiagaraNamespaceMetadataOptions::AdvancedInSystem)
+			.AddOption(ENiagaraNamespaceMetadataOptions::HideInDefinitions)
 		//.AddOption(ENiagaraNamespaceMetadataOptions::PreventEditingNamespaceModifier),
 			.AddOptionalNamespaceModifier(FNiagaraConstants::ModuleNamespace)
 			.AddOptionalNamespaceModifier(FNiagaraConstants::InitialNamespace)
@@ -180,6 +185,7 @@ void UNiagaraEditorSettings::SetupNamespaceMetadata()
 			.SetBackgroundColor(FLinearColor(FColor(170, 170, 170)))
 			.SetForegroundStyle("NiagaraEditor.ParameterName.NamespaceTextDark")
 			.SetSortId(70)
+			.AddOption(ENiagaraNamespaceMetadataOptions::HideInDefinitions)
 			.AddOption(ENiagaraNamespaceMetadataOptions::PreventEditingNamespace)
 			.AddOption(ENiagaraNamespaceMetadataOptions::PreventEditingNamespaceModifier)
 			.AddOption(ENiagaraNamespaceMetadataOptions::PreventEditingName)
@@ -191,6 +197,7 @@ void UNiagaraEditorSettings::SetupNamespaceMetadata()
 			.SetBackgroundColor(FLinearColor(FColor(91, 161, 194)))
 			.SetSortId(0)
 			.AddOption(ENiagaraNamespaceMetadataOptions::HideInScript)
+			.AddOption(ENiagaraNamespaceMetadataOptions::HideInDefinitions)
 			.AddOption(ENiagaraNamespaceMetadataOptions::PreventEditingNamespaceModifier)
 			.AddOption(ENiagaraNamespaceMetadataOptions::PreventEditingNamespace)
 			.SetGuid(FNiagaraEditorGuids::UserNamespaceMetaDataGuid),
@@ -203,6 +210,7 @@ void UNiagaraEditorSettings::SetupNamespaceMetadata()
 			.SetSortId(90)
 			.AddOption(ENiagaraNamespaceMetadataOptions::AdvancedInScript)
 			.AddOption(ENiagaraNamespaceMetadataOptions::AdvancedInSystem)
+			.AddOption(ENiagaraNamespaceMetadataOptions::HideInDefinitions)
 			.AddOption(ENiagaraNamespaceMetadataOptions::PreventEditingNamespace)
 			.AddOption(ENiagaraNamespaceMetadataOptions::PreventEditingNamespaceModifier)
 			.AddOption(ENiagaraNamespaceMetadataOptions::PreventEditingName)
@@ -216,6 +224,7 @@ void UNiagaraEditorSettings::SetupNamespaceMetadata()
 			.AddOption(ENiagaraNamespaceMetadataOptions::HideInSystem)
 			.AddOption(ENiagaraNamespaceMetadataOptions::AdvancedInScript)
 			.AddOption(ENiagaraNamespaceMetadataOptions::AdvancedInSystem)
+			.AddOption(ENiagaraNamespaceMetadataOptions::HideInDefinitions)
 			.AddOption(ENiagaraNamespaceMetadataOptions::PreventEditingNamespace)
 			.AddOption(ENiagaraNamespaceMetadataOptions::PreventEditingNamespaceModifier)
 			.AddOption(ENiagaraNamespaceMetadataOptions::PreventEditingName)
@@ -225,6 +234,7 @@ void UNiagaraEditorSettings::SetupNamespaceMetadata()
 			.SetDescription(LOCTEXT("StaticSwitchDescription", "Values which can only be set at edit time."))
 			.SetSortId(45)
 			.AddOption(ENiagaraNamespaceMetadataOptions::HideInSystem)
+			.AddOption(ENiagaraNamespaceMetadataOptions::HideInDefinitions)
 			.AddOption(ENiagaraNamespaceMetadataOptions::PreventEditingNamespace)
 			.AddOption(ENiagaraNamespaceMetadataOptions::PreventEditingNamespaceModifier)
 			.AddOption(ENiagaraNamespaceMetadataOptions::PreventEditingName)
diff --git a/Engine/Plugins/FX/Niagara/Source/NiagaraEditor/Private/NiagaraEditorUtilities.cpp b/Engine/Plugins/FX/Niagara/Source/NiagaraEditor/Private/NiagaraEditorUtilities.cpp
--- a/Engine/Plugins/FX/Niagara/Source/NiagaraEditor/Private/NiagaraEditorUtilities.cpp
+++ b/Engine/Plugins/FX/Niagara/Source/NiagaraEditor/Private/NiagaraEditorUtilities.cpp
@@ -3221,7 +3221,8 @@ void FNiagaraParameterUtilities::GetChangeNamespaceMenuData(FName InParameterNam
 			Metadata == CurrentMetadata ||
 			Metadata.Options.Contains(ENiagaraNamespaceMetadataOptions::PreventEditingNamespace) ||
 			(InParameterContext == EParameterContext::Script && Metadata.Options.Contains(ENiagaraNamespaceMetadataOptions::HideInScript)) ||
-			(InParameterContext == EParameterContext::System && Metadata.Options.Contains(ENiagaraNamespaceMetadataOptions::HideInSystem)))
+			(InParameterContext == EParameterContext::System && Metadata.Options.Contains(ENiagaraNamespaceMetadataOptions::HideInSystem)) ||
+			(InParameterContext == EParameterContext::Definitions && Metadata.Options.Contains(ENiagaraNamespaceMetadataOptions::HideInDefinitions)))
 		{
 			continue;
 		}
diff --git a/Engine/Plugins/FX/Niagara/Source/NiagaraEditor/Private/ViewModels/NiagaraParameterPanelViewModel.cpp b/Engine/Plugins/FX/Niagara/Source/NiagaraEditor/Private/ViewModels/NiagaraParameterPanelViewModel.cpp
--- a/Engine/Plugins/FX/Niagara/Source/NiagaraEditor/Private/ViewModels/NiagaraParameterPanelViewModel.cpp
+++ b/Engine/Plugins/FX/Niagara/Source/NiagaraEditor/Private/ViewModels/NiagaraParameterPanelViewModel.cpp
@@ -180,6 +180,81 @@
 	return FReply::Handled().BeginDragDrop(DragOperation.ToSharedRef());
 }
 
+///////////////////////////////////////////////////////////////////////////////
+/// Generic Parameter Panel Utilities										///
+///////////////////////////////////////////////////////////////////////////////
+
+bool FNiagaraParameterPanelUtilities::GetCanSetParameterNamespaceAndToolTipForScriptOrSystem(const FNiagaraParameterPanelItem& ItemToModify, const FName NewNamespace, FText& OutCanSetParameterNamespaceToolTip)
+{
+	if (ItemToModify.ScriptVariable->GetIsSubscribedToParameterDefinitions())
+	{
+		OutCanSetParameterNamespaceToolTip = LOCTEXT("ParameterPanelViewModel_ChangeParameterNamespace_ParameterDefinitions", "Cannot change Parameter namespace: Parameters linked to Parameter Definitions may only have their namespace changed in the source Parameter Definitions asset.");
+		return false;
+	}
+	else if (ItemToModify.bExternallyReferenced)
+	{
+		OutCanSetParameterNamespaceToolTip = LOCTEXT("ParameterPanelViewModel_ChangeParameterNamespace_ExternallyReferenced", "Cannot change Parameter namespace: Parameter is from an externally referenced script and can't be directly edited.");
+		return false;
+	}
+	return true;
+}
+
+bool FNiagaraParameterPanelUtilities::GetCanSetParameterNamespaceModifierAndToolTipForScriptOrSystem(const TArray<FNiagaraParameterPanelItem>& CachedViewedItems, const FNiagaraParameterPanelItem& ItemToModify, const FName NamespaceModifier, bool bDuplicateParameter, FText& OutCanSetParameterNamespaceModifierToolTip)
+{
+	if (FNiagaraParameterUtilities::TestCanSetSpecificNamespaceModifierWithMessage(ItemToModify.GetVariable().GetName(), NamespaceModifier, OutCanSetParameterNamespaceModifierToolTip) == false)
+	{
+		return false;
+	}
+
+	if (bDuplicateParameter == false)
+	{
+		if (ItemToModify.ScriptVariable->GetIsSubscribedToParameterDefinitions())
+		{
+			OutCanSetParameterNamespaceModifierToolTip = LOCTEXT("ParameterPanelViewModel_ChangeParameterNamespaceModifier_ParameterDefinitions", "Cannot change Parameter namespace modifier: Parameters from Parameter Definitions may only have their namespace modifier changed in the source Parameter Definitions asset.");
+			return false;
+		}
+		else if (ItemToModify.bExternallyReferenced)
+		{
+			OutCanSetParameterNamespaceModifierToolTip = LOCTEXT("CantChangeNamespaceModifierExternallyReferenced", "Parameter is from an externally referenced script and can't be directly edited.");
+			return false;
+		}
+		else if (NamespaceModifier != NAME_None)
+		{
+			FName NewName = FNiagaraParameterUtilities::SetSpecificNamespaceModifier(ItemToModify.GetVariable().GetName(), NamespaceModifier);
+			if (CachedViewedItems.ContainsByPredicate([NewName](const FNiagaraParameterPanelItem& CachedViewedItem) {return CachedViewedItem.GetVariable().GetName() == NewName; }))
+			{
+				OutCanSetParameterNamespaceModifierToolTip = LOCTEXT("CantChangeNamespaceModifierAlreadyExits", "Can't set this namespace modifier because it would create a parameter that already exists.");
+				return false;
+			}
+		}
+	}
+
+	return true;
+}
+
+bool FNiagaraParameterPanelUtilities::GetCanSetParameterCustomNamespaceModifierAndToolTipForScriptOrSystem(const FNiagaraParameterPanelItem& ItemToModify, bool bDuplicateParameter, FText& OutCanSetParameterNamespaceModifierToolTip)
+{
+	if (FNiagaraParameterUtilities::TestCanSetCustomNamespaceModifierWithMessage(ItemToModify.GetVariable().GetName(), OutCanSetParameterNamespaceModifierToolTip) == false)
+	{
+		return false;
+	}
+
+	if (bDuplicateParameter == false) 
+	{
+		if (ItemToModify.bExternallyReferenced)
+		{
+			OutCanSetParameterNamespaceModifierToolTip = LOCTEXT("CantChangeNamespaceModifierExternallyReferenced", "Parameter is from an externally referenced script and can't be directly edited.");
+			return false;
+		}
+		else if (ItemToModify.ScriptVariable->GetIsSubscribedToParameterDefinitions())
+		{
+			OutCanSetParameterNamespaceModifierToolTip = LOCTEXT("ParameterPanelViewModel_ChangeParameterCustomNamespace_ParameterDefinitions", "Cannot set Parameter custom namespace: Parameters linked to Parameter Definitions may only have a custom namespace set in the source Parameter Definitions asset.");
+			return false;
+		}
+	}
+
+	return true;
+}
 
 ///////////////////////////////////////////////////////////////////////////////
 /// Immutable Parameter Panel View Model									///
@@ -451,21 +526,6 @@
 	}
 }
 
-bool INiagaraParameterPanelViewModel::GetCanSetParameterNamespaceAndToolTip(const FNiagaraParameterPanelItem& ItemToModify, const FName NewNamespace, FText& OutCanSetParameterNamespaceToolTip) const 
-{
-	if (ItemToModify.ScriptVariable->GetIsSubscribedToParameterDefinitions())
-	{
-		OutCanSetParameterNamespaceToolTip = LOCTEXT("ParameterPanelViewModel_ChangeParameterNamespace_ParameterDefinitions", "Cannot change Parameter namespace: Parameters linked to Parameter Definitions may only have their namespace changed in the source Parameter Definitions asset.");
-		return false;
-	}
-	else if (ItemToModify.bExternallyReferenced)
-	{
-		OutCanSetParameterNamespaceToolTip = LOCTEXT("ParameterPanelViewModel_ChangeParameterNamespace_ExternallyReferenced", "Cannot change Parameter namespace: Parameter is from an externally referenced script and can't be directly edited.");
-		return false;
-	}
-	return true;
-}
-
 void INiagaraParameterPanelViewModel::SetParameterNamespaceModifier(const FNiagaraParameterPanelItem ItemToModify, const FName NewNamespaceModifier, bool bDuplicateParameter) const
 {
 	FName NewName = FNiagaraParameterUtilities::SetSpecificNamespaceModifier(ItemToModify.GetVariable().GetName(), NewNamespaceModifier);
@@ -502,39 +562,6 @@
 	}
 }
 
-bool INiagaraParameterPanelViewModel::GetCanSetParameterNamespaceModifierAndToolTip(const FNiagaraParameterPanelItem& ItemToModify, const FName NamespaceModifier, bool bDuplicateParameter, FText& OutCanSetParameterNamespaceModifierToolTip) const
-{
-	if (FNiagaraParameterUtilities::TestCanSetSpecificNamespaceModifierWithMessage(ItemToModify.GetVariable().GetName(), NamespaceModifier, OutCanSetParameterNamespaceModifierToolTip) == false)
-	{
-		return false;
-	}
-
-	if (bDuplicateParameter == false)
-	{
-		if (ItemToModify.ScriptVariable->GetIsSubscribedToParameterDefinitions())
-		{
-			OutCanSetParameterNamespaceModifierToolTip = LOCTEXT("ParameterPanelViewModel_ChangeParameterNamespaceModifier_ParameterDefinitions", "Cannot change Parameter namespace modifier: Parameters from Parameter Definitions may only have their namespace modifier changed in the source Parameter Definitions asset.");
-			return false;
-		}
-		else if (ItemToModify.bExternallyReferenced)
-		{
-			OutCanSetParameterNamespaceModifierToolTip = LOCTEXT("CantChangeNamespaceModifierExternallyReferenced", "Parameter is from an externally referenced script and can't be directly edited.");
-			return false;
-		}
-		else if (NamespaceModifier != NAME_None)
-		{
-			FName NewName = FNiagaraParameterUtilities::SetSpecificNamespaceModifier(ItemToModify.GetVariable().GetName(), NamespaceModifier);
-			if (CachedViewedItems.ContainsByPredicate([NewName](const FNiagaraParameterPanelItem& CachedViewedItem) {return CachedViewedItem.GetVariable().GetName() == NewName; }))
-			{
-				OutCanSetParameterNamespaceModifierToolTip = LOCTEXT("CantChangeNamespaceModifierAlreadyExits", "Can't set this namespace modifier because it would create a parameter that already exists.");
-				return false;
-			}
-		}
-	}
-
-	return true;
-}
-
 void INiagaraParameterPanelViewModel::SetParameterCustomNamespaceModifier(const FNiagaraParameterPanelItem ItemToModify, bool bDuplicateParameter) const
 {
 	if (bDuplicateParameter == false && ItemToModify.bExternallyReferenced)
@@ -582,26 +609,6 @@
 	}
 }
 
-bool INiagaraParameterPanelViewModel::GetCanSetParameterCustomNamespaceModifierAndToolTip(const FNiagaraParameterPanelItem& ItemToModify, bool bDuplicateParameter, FText& OutCanSetParameterNamespaceModifierToolTip) const
-{
-	if (ItemToModify.ScriptVariable->GetIsSubscribedToParameterDefinitions())
-	{
-		OutCanSetParameterNamespaceModifierToolTip = LOCTEXT("ParameterPanelViewModel_ChangeParameterCustomNamespace_ParameterDefinitions", "Cannot set Parameter custom namespace: Parameters linked to Parameter Definitions may only have a custom namespace set in the source Parameter Definitions asset.");
-		return false;
-	}
-	else if (FNiagaraParameterUtilities::TestCanSetCustomNamespaceModifierWithMessage(ItemToModify.GetVariable().GetName(), OutCanSetParameterNamespaceModifierToolTip) == false)
-	{
-		return false;
-	}
 
-	if (bDuplicateParameter == false && ItemToModify.bExternallyReferenced)
-	{
-		OutCanSetParameterNamespaceModifierToolTip = LOCTEXT("CantChangeNamespaceModifierExternallyReferenced", "Parameter is from an externally referenced script and can't be directly edited.");
-		return false;
-	}
-
-	return true;
-}
-
 void INiagaraParameterPanelViewModel::GetChangeNamespaceSubMenu(FMenuBuilder& MenuBuilder, bool bDuplicateParameter, FNiagaraParameterPanelItem Item) const
 {
@@ -1313,6 +1320,21 @@
 	return true;
 }
 
+bool FNiagaraSystemToolkitParameterPanelViewModel::GetCanSetParameterNamespaceAndToolTip(const FNiagaraParameterPanelItem& ItemToModify, const FName NewNamespace, FText& OutCanSetParameterNamespaceToolTip) const
+{
+	return FNiagaraParameterPanelUtilities::GetCanSetParameterNamespaceAndToolTipForScriptOrSystem(ItemToModify, NewNamespace, OutCanSetParameterNamespaceToolTip);
+}
+
+bool FNiagaraSystemToolkitParameterPanelViewModel::GetCanSetParameterNamespaceModifierAndToolTip(const FNiagaraParameterPanelItem& ItemToModify, const FName NamespaceModifier, bool bDuplicateParameter, FText& OutCanSetParameterNamespaceModifierToolTip) const
+{
+	return FNiagaraParameterPanelUtilities::GetCanSetParameterNamespaceModifierAndToolTipForScriptOrSystem(CachedViewedItems, ItemToModify, NamespaceModifier, bDuplicateParameter, OutCanSetParameterNamespaceModifierToolTip);
+}
+
+bool FNiagaraSystemToolkitParameterPanelViewModel::GetCanSetParameterCustomNamespaceModifierAndToolTip(const FNiagaraParameterPanelItem& ItemToModify, bool bDuplicateParameter, FText& OutCanSetParameterNamespaceModifierToolTip) const
+{
+	return FNiagaraParameterPanelUtilities::GetCanSetParameterCustomNamespaceModifierAndToolTipForScriptOrSystem(ItemToModify, bDuplicateParameter, OutCanSetParameterNamespaceModifierToolTip);
+}
+
 TSharedRef<SWidget> FNiagaraSystemToolkitParameterPanelViewModel::CreateAddParameterMenuForAssignmentNode(UNiagaraNodeAssignment* AssignmentNode, const TSharedPtr<SComboButton>& AddButton) const
 {
 	auto AddParameterLambda = [this, AssignmentNode](FNiagaraVariable& NewParameter) {
@@ -2502,6 +2524,21 @@
 	return true;
 }
 
+bool FNiagaraScriptToolkitParameterPanelViewModel::GetCanSetParameterNamespaceAndToolTip(const FNiagaraParameterPanelItem& ItemToModify, const FName NewNamespace, FText& OutCanSetParameterNamespaceToolTip) const
+{
+	return FNiagaraParameterPanelUtilities::GetCanSetParameterNamespaceAndToolTipForScriptOrSystem(ItemToModify, NewNamespace, OutCanSetParameterNamespaceToolTip);
+}
+
+bool FNiagaraScriptToolkitParameterPanelViewModel::GetCanSetParameterNamespaceModifierAndToolTip(const FNiagaraParameterPanelItem& ItemToModify, const FName NamespaceModifier, bool bDuplicateParameter, FText& OutCanSetParameterNamespaceModifierToolTip) const
+{
+	return FNiagaraParameterPanelUtilities::GetCanSetParameterNamespaceModifierAndToolTipForScriptOrSystem(CachedViewedItems, ItemToModify, NamespaceModifier, bDuplicateParameter, OutCanSetParameterNamespaceModifierToolTip);
+}
+
+bool FNiagaraScriptToolkitParameterPanelViewModel::GetCanSetParameterCustomNamespaceModifierAndToolTip(const FNiagaraParameterPanelItem& ItemToModify, bool bDuplicateParameter, FText& OutCanSetParameterNamespaceModifierToolTip) const
+{
+	return FNiagaraParameterPanelUtilities::GetCanSetParameterCustomNamespaceModifierAndToolTipForScriptOrSystem(ItemToModify, bDuplicateParameter, OutCanSetParameterNamespaceModifierToolTip);
+}
+
 void FNiagaraScriptToolkitParameterPanelViewModel::SetParameterIsOverridingLibraryDefaultValue(const FNiagaraParameterPanelItem ItemToModify, const bool bOverriding) const
 {
 	if (ensureMsgf(ItemToModify.bExternallyReferenced == false, TEXT("Cannot modify an externally referenced parameter.")) == false)
@@ -2859,7 +2896,7 @@
 
 FNiagaraParameterUtilities::EParameterContext FNiagaraParameterDefinitionsToolkitParameterPanelViewModel::GetParameterContext() const
 {
-	return FNiagaraParameterUtilities::EParameterContext::System;
+	return FNiagaraParameterUtilities::EParameterContext::Definitions;
 }
 
 const TArray<UNiagaraScriptVariable*> FNiagaraParameterDefinitionsToolkitParameterPanelViewModel::GetEditableScriptVariablesWithName(const FName ParameterName) const
@@ -2933,6 +2970,44 @@
 	return true;
 }
 
+bool FNiagaraParameterDefinitionsToolkitParameterPanelViewModel::GetCanSetParameterNamespaceAndToolTip(const FNiagaraParameterPanelItem& ItemToModify, const FName NewNamespace, FText& OutCanSetParameterNamespaceToolTip) const
+{
+	return true;
+}
+
+bool FNiagaraParameterDefinitionsToolkitParameterPanelViewModel::GetCanSetParameterNamespaceModifierAndToolTip(const FNiagaraParameterPanelItem& ItemToModify, const FName NamespaceModifier, bool bDuplicateParameter, FText& OutCanSetParameterNamespaceModifierToolTip) const
+{
+	if (FNiagaraParameterUtilities::TestCanSetSpecificNamespaceModifierWithMessage(ItemToModify.GetVariable().GetName(), NamespaceModifier, OutCanSetParameterNamespaceModifierToolTip) == false)
+	{
+		return false;
+	}
+
+	if (bDuplicateParameter == false)
+	{
+		if (NamespaceModifier != NAME_None)
+		{
+			FName NewName = FNiagaraParameterUtilities::SetSpecificNamespaceModifier(ItemToModify.GetVariable().GetName(), NamespaceModifier);
+			if (CachedViewedItems.ContainsByPredicate([NewName](const FNiagaraParameterPanelItem& CachedViewedItem) {return CachedViewedItem.GetVariable().GetName() == NewName; }))
+			{
+				OutCanSetParameterNamespaceModifierToolTip = LOCTEXT("CantChangeNamespaceModifierAlreadyExits", "Can't set this namespace modifier because it would create a parameter that already exists.");
+				return false;
+			}
+		}
+	}
+
+	return true;
+}
+
+bool FNiagaraParameterDefinitionsToolkitParameterPanelViewModel::GetCanSetParameterCustomNamespaceModifierAndToolTip(const FNiagaraParameterPanelItem& ItemToModify, bool bDuplicateParameter, FText& OutCanSetParameterNamespaceModifierToolTip) const
+{
+	if (FNiagaraParameterUtilities::TestCanSetCustomNamespaceModifierWithMessage(ItemToModify.GetVariable().GetName(), OutCanSetParameterNamespaceModifierToolTip) == false)
+	{
+		return false;
+	}
+
+	return true;
+}
+
 const TArray<FNiagaraParameterPanelCategory>& FNiagaraParameterDefinitionsToolkitParameterPanelViewModel::GetDefaultCategories() const
 {
 	return FNiagaraParameterDefinitionsToolkitParameterPanelViewModel::DefaultCategories;
diff --git a/Engine/Plugins/FX/Niagara/Source/NiagaraEditor/Public/NiagaraEditorSettings.h b/Engine/Plugins/FX/Niagara/Source/NiagaraEditor/Public/NiagaraEditorSettings.h
--- a/Engine/Plugins/FX/Niagara/Source/NiagaraEditor/Public/NiagaraEditorSettings.h
+++ b/Engine/Plugins/FX/Niagara/Source/NiagaraEditor/Public/NiagaraEditorSettings.h
@@ -57,7 +57,8 @@ enum class ENiagaraNamespaceMetadataOptions
 	PreventEditingNamespace,
 	PreventEditingNamespaceModifier,
 	PreventEditingName,
-	PreventCreatingInSystemEditor
+	PreventCreatingInSystemEditor,
+	HideInDefinitions,
 };
 
 USTRUCT()
diff --git a/Engine/Plugins/FX/Niagara/Source/NiagaraEditor/Public/NiagaraEditorUtilities.h b/Engine/Plugins/FX/Niagara/Source/NiagaraEditor/Public/NiagaraEditorUtilities.h
--- a/Engine/Plugins/FX/Niagara/Source/NiagaraEditor/Public/NiagaraEditorUtilities.h
+++ b/Engine/Plugins/FX/Niagara/Source/NiagaraEditor/Public/NiagaraEditorUtilities.h
@@ -358,7 +358,8 @@ namespace FNiagaraParameterUtilities
 	enum class EParameterContext : uint8
 	{
 		Script,
-		System
+		System,
+		Definitions,
 	};
 
 	struct FChangeNamespaceMenuData
diff --git a/Engine/Plugins/FX/Niagara/Source/NiagaraEditor/Public/ViewModels/NiagaraParameterPanelViewModel.h b/Engine/Plugins/FX/Niagara/Source/NiagaraEditor/Public/ViewModels/NiagaraParameterPanelViewModel.h
--- a/Engine/Plugins/FX/Niagara/Source/NiagaraEditor/Public/ViewModels/NiagaraParameterPanelViewModel.h
+++ b/Engine/Plugins/FX/Niagara/Source/NiagaraEditor/Public/ViewModels/NiagaraParameterPanelViewModel.h
@@ -60,6 +60,15 @@ namespace FNiagaraScriptToolkitParameterPanelUtilities
 	);
 }
 
+namespace FNiagaraParameterPanelUtilities
+{
+	bool GetCanSetParameterNamespaceAndToolTipForScriptOrSystem(const FNiagaraParameterPanelItem& ItemToModify, const FName NewNamespace, FText& OutCanSetParameterNamespaceToolTip);
+
+	bool GetCanSetParameterNamespaceModifierAndToolTipForScriptOrSystem(const TArray<FNiagaraParameterPanelItem>& CachedViewedItems, const FNiagaraParameterPanelItem& ItemToModify, const FName NamespaceModifier, bool bDuplicateParameter, FText& OutCanSetParameterNamespaceModifierToolTip);
+
+	bool GetCanSetParameterCustomNamespaceModifierAndToolTipForScriptOrSystem(const FNiagaraParameterPanelItem& ItemToModify, bool bDuplicateParameter, FText& OutCanSetParameterNamespaceModifierToolTip);
+}
+
 struct FMenuAndSearchBoxWidgets
 {
 	TSharedPtr<SWidget> MenuWidget;
@@ -170,6 +179,12 @@ public:
 	virtual FReply HandleDragDropOperation(TSharedPtr<FDragDropOperation> DropOperation) const = 0;
 
 	virtual bool GetCanHandleDragDropOperation(TSharedPtr<FDragDropOperation> DragDropOperation) const = 0;
+
+	virtual bool GetCanSetParameterNamespaceAndToolTip(const FNiagaraParameterPanelItem& ItemToModify, const FName NewNamespace, FText& OutCanSetParameterNamespaceToolTip) const = 0;
+
+	virtual bool GetCanSetParameterNamespaceModifierAndToolTip(const FNiagaraParameterPanelItem& ItemToModify, const FName NamespaceModifier, bool bDuplicateParameter, FText& OutCanSetParameterNamespaceModifierToolTip) const = 0;
+
+	virtual bool GetCanSetParameterCustomNamespaceModifierAndToolTip(const FNiagaraParameterPanelItem& ItemToModify, bool bDuplicateParameter, FText& OutCanSetParameterNamespaceModifierToolTip) const = 0;
 	//~ End Pure Virtual Methods
 
 	virtual bool GetCanDeleteParameterAndToolTip(const FNiagaraParameterPanelItem& ItemToDelete, FText& OutCanDeleteParameterToolTip) const;
@@ -190,16 +205,10 @@ public:
 
 	virtual void SetParameterNamespace(const FNiagaraParameterPanelItem ItemToModify, FNiagaraNamespaceMetadata NewNamespaceMetaData, bool bDuplicateParameter) const;
 
-	virtual bool GetCanSetParameterNamespaceAndToolTip(const FNiagaraParameterPanelItem& ItemToModify, const FName NewNamespace, FText& OutCanSetParameterNamespaceToolTip) const;
-
 	virtual void SetParameterNamespaceModifier(const FNiagaraParameterPanelItem ItemToModify, const FName NewNamespaceModifier, bool bDuplicateParameter) const;
 
-	virtual bool GetCanSetParameterNamespaceModifierAndToolTip(const FNiagaraParameterPanelItem& ItemToModify, const FName NamespaceModifier, bool bDuplicateParameter, FText& OutCanSetParameterNamespaceModifierToolTip) const;
-
 	virtual void SetParameterCustomNamespaceModifier(const FNiagaraParameterPanelItem ItemToModify, bool bDuplicateParameter) const;
 
-	virtual bool GetCanSetParameterCustomNamespaceModifierAndToolTip(const FNiagaraParameterPanelItem& ItemToModify, bool bDuplicateParameter, FText& OutCanSetParameterNamespaceModifierToolTip) const;
-
 	virtual void GetChangeNamespaceSubMenu(FMenuBuilder& MenuBuilder, bool bDuplicateParameter, FNiagaraParameterPanelItem Item) const;
 
 	virtual void GetChangeNamespaceModifierSubMenu(FMenuBuilder& MenuBuilder, bool bDuplicateParameter, FNiagaraParameterPanelItem Item) const;
@@ -299,6 +308,12 @@ public:
 	virtual FReply HandleDragDropOperation(TSharedPtr<FDragDropOperation> DragDropOperation) const override;
 
 	virtual bool GetCanHandleDragDropOperation(TSharedPtr<FDragDropOperation> DragDropOperation) const override;
+
+	virtual bool GetCanSetParameterNamespaceAndToolTip(const FNiagaraParameterPanelItem& ItemToModify, const FName NewNamespace, FText& OutCanSetParameterNamespaceToolTip) const override;
+
+	virtual bool GetCanSetParameterNamespaceModifierAndToolTip(const FNiagaraParameterPanelItem& ItemToModify, const FName NamespaceModifier, bool bDuplicateParameter, FText& OutCanSetParameterNamespaceModifierToolTip) const override;
+
+	virtual bool GetCanSetParameterCustomNamespaceModifierAndToolTip(const FNiagaraParameterPanelItem& ItemToModify, bool bDuplicateParameter, FText& OutCanSetParameterNamespaceModifierToolTip) const override;
 	//~ End INiagaraParameterPanelViewModel interface
 
 	TSharedRef<SWidget> CreateAddParameterMenuForAssignmentNode(UNiagaraNodeAssignment* AssignmentNode, const TSharedPtr<SComboButton>& AddButton) const;
@@ -397,6 +412,12 @@ public:
 	virtual FReply HandleDragDropOperation(TSharedPtr<FDragDropOperation> DragDropOperation) const override;
 
 	virtual bool GetCanHandleDragDropOperation(TSharedPtr<FDragDropOperation> DragDropOperation) const override;
+
+	virtual bool GetCanSetParameterNamespaceAndToolTip(const FNiagaraParameterPanelItem& ItemToModify, const FName NewNamespace, FText& OutCanSetParameterNamespaceToolTip) const override;
+
+	virtual bool GetCanSetParameterNamespaceModifierAndToolTip(const FNiagaraParameterPanelItem& ItemToModify, const FName NamespaceModifier, bool bDuplicateParameter, FText& OutCanSetParameterNamespaceModifierToolTip) const override;
+
+	virtual bool GetCanSetParameterCustomNamespaceModifierAndToolTip(const FNiagaraParameterPanelItem& ItemToModify, bool bDuplicateParameter, FText& OutCanSetParameterNamespaceModifierToolTip) const override;
 	//~ End INiagaraParameterPanelViewModel interface
 
 	void RenameParameter(const UNiagaraScriptVariable* ScriptVarToRename, const FName NewName) const;
@@ -488,6 +509,11 @@ public:
 
 	virtual bool GetCanRenameParameterAndToolTip(const FNiagaraParameterPanelItem& ItemToRename, const FText& NewVariableNameText, bool bCheckEmptyNameText, FText& OutCanRenameParameterToolTip) const override;
 	
+	virtual bool GetCanSetParameterNamespaceAndToolTip(const FNiagaraParameterPanelItem& ItemToModify, const FName NewNamespace, FText& OutCanSetParameterNamespaceToolTip) const override;
+
+	virtual bool GetCanSetParameterNamespaceModifierAndToolTip(const FNiagaraParameterPanelItem& ItemToModify, const FName NamespaceModifier, bool bDuplicateParameter, FText& OutCanSetParameterNamespaceModifierToolTip) const override;
+
+	virtual bool GetCanSetParameterCustomNamespaceModifierAndToolTip(const FNiagaraParameterPanelItem& ItemToModify, bool bDuplicateParameter, FText& OutCanSetParameterNamespaceModifierToolTip) const override;
 	//~ End INiagaraParameterPanelViewModel interface
 
 private:
