Niagara - Implement merging for summary section data and refactor input summary merging to better support changes made in the child.

diff --git a/Engine/Plugins/FX/Niagara/Source/Niagara/Private/NiagaraScript.cpp b/Engine/Plugins/FX/Niagara/Source/Niagara/Private/NiagaraScript.cpp
--- a/Engine/Plugins/FX/Niagara/Source/Niagara/Private/NiagaraScript.cpp
+++ b/Engine/Plugins/FX/Niagara/Source/Niagara/Private/NiagaraScript.cpp
@@ -1893,6 +1893,17 @@ void UNiagaraScript::PostLoad()
         		ScriptData->LibraryVisibility = ENiagaraScriptLibraryVisibility::Library;
         	}
         }
+
+		if (Data.InputSections.Num() > 0)
+		{
+			for (FNiagaraStackSection& InputSection : Data.InputSections)
+			{
+				if (InputSection.SectionIdentifier == NAME_None)
+				{
+					InputSection.SectionIdentifier = *InputSection.SectionDisplayName.ToString();
+				}
+			}
+		}
 	}
 
 #endif
diff --git a/Engine/Plugins/FX/Niagara/Source/Niagara/Public/NiagaraStackSection.h b/Engine/Plugins/FX/Niagara/Source/Niagara/Public/NiagaraStackSection.h
--- a/Engine/Plugins/FX/Niagara/Source/Niagara/Public/NiagaraStackSection.h
+++ b/Engine/Plugins/FX/Niagara/Source/Niagara/Public/NiagaraStackSection.h
@@ -13,9 +13,20 @@ struct FNiagaraStackSection
 {
 	GENERATED_BODY()
 
+	FNiagaraStackSection()
+		: bEnabled(true)
+	{
+	}
+
+	UPROPERTY(EditAnywhere, Category=Section)
+	FName SectionIdentifier;
+
 	UPROPERTY(EditAnywhere, Category=Section)
 	FText SectionDisplayName;
 
 	UPROPERTY(EditAnywhere, Category=Section)
 	TArray<FText> Categories;
+
+	UPROPERTY(EditAnywhere, Category=Section)
+	bool bEnabled;
 };
\ No newline at end of file
diff --git a/Engine/Plugins/FX/Niagara/Source/NiagaraEditor/Private/NiagaraEmitterEditorData.cpp b/Engine/Plugins/FX/Niagara/Source/NiagaraEditor/Private/NiagaraEmitterEditorData.cpp
--- a/Engine/Plugins/FX/Niagara/Source/NiagaraEditor/Private/NiagaraEmitterEditorData.cpp
+++ b/Engine/Plugins/FX/Niagara/Source/NiagaraEditor/Private/NiagaraEmitterEditorData.cpp
@@ -56,26 +56,25 @@ void UNiagaraEmitterEditorData::StackEditorDataChanged()
 
 
 
-TMap<FFunctionInputSummaryViewKey, FFunctionInputSummaryViewMetadata> UNiagaraEmitterEditorData::GetSummaryViewMetaDataMap() const
+const TMap<FFunctionInputSummaryViewKey, FFunctionInputSummaryViewMetadata>& UNiagaraEmitterEditorData::GetSummaryViewMetaDataMap() const
 {
 	return SummaryViewFunctionInputMetadata;
 }
 
-FFunctionInputSummaryViewMetadata UNiagaraEmitterEditorData::GetSummaryViewMetaData(const FFunctionInputSummaryViewKey& Key) const
+TOptional<FFunctionInputSummaryViewMetadata> UNiagaraEmitterEditorData::GetSummaryViewMetaData(const FFunctionInputSummaryViewKey& Key) const
 {
-	return SummaryViewFunctionInputMetadata.Contains(Key)? SummaryViewFunctionInputMetadata[Key] : FFunctionInputSummaryViewMetadata();
+	return SummaryViewFunctionInputMetadata.Contains(Key)? SummaryViewFunctionInputMetadata[Key] : TOptional<FFunctionInputSummaryViewMetadata>();
 }
 
-void UNiagaraEmitterEditorData::SetSummaryViewMetaData(const FFunctionInputSummaryViewKey& Key, const FFunctionInputSummaryViewMetadata& NewMetadata)
+void UNiagaraEmitterEditorData::SetSummaryViewMetaData(const FFunctionInputSummaryViewKey& Key, TOptional<FFunctionInputSummaryViewMetadata> NewMetadata)
 {	
-	static const FFunctionInputSummaryViewMetadata Empty;
-	if (NewMetadata == Empty)
+	if (NewMetadata.IsSet() == false)
 	{
 		SummaryViewFunctionInputMetadata.Remove(Key);
 	}
 	else
 	{
-		SummaryViewFunctionInputMetadata.FindOrAdd(Key) = NewMetadata;
+		SummaryViewFunctionInputMetadata.FindOrAdd(Key) = NewMetadata.GetValue();
 	}
 	
 	OnPersistentDataChanged().Broadcast();
@@ -96,4 +95,14 @@ void UNiagaraEmitterEditorData::ToggleShowSummaryView()
 FSimpleMulticastDelegate& UNiagaraEmitterEditorData::OnSummaryViewStateChanged()
 {
 	return OnSummaryViewStateChangedDelegate;
+}
+
+const TArray<FNiagaraStackSection>& UNiagaraEmitterEditorData::GetSummarySections() const
+{
+	return SummarySections;
+}
+
+void UNiagaraEmitterEditorData::SetSummarySections(const TArray<FNiagaraStackSection>& InSummarySections)
+{
+	SummarySections = InSummarySections;
 }
\ No newline at end of file
diff --git a/Engine/Plugins/FX/Niagara/Source/NiagaraEditor/Private/NiagaraScriptMergeManager.cpp b/Engine/Plugins/FX/Niagara/Source/NiagaraEditor/Private/NiagaraScriptMergeManager.cpp
--- a/Engine/Plugins/FX/Niagara/Source/NiagaraEditor/Private/NiagaraScriptMergeManager.cpp
+++ b/Engine/Plugins/FX/Niagara/Source/NiagaraEditor/Private/NiagaraScriptMergeManager.cpp
@@ -1578,7 +1578,7 @@ INiagaraMergeManager::FMergeEmitterResults FNiagaraScriptMergeManager::MergeEmit
 		MergeResults.bModifiedGraph |= RendererResults.bModifiedGraph;
 		MergeResults.ErrorMessages.Append(RendererResults.ErrorMessages);
 
-		FApplyDiffResults InputSummaryResults = ApplyInputSummaryDiff(*MergedInstance, DiffResults, bNoParentAtLastMerge);
+		FApplyDiffResults InputSummaryResults = ApplyEmitterSummaryDiff(*MergedInstance, DiffResults, bNoParentAtLastMerge);
 		if (InputSummaryResults.bSucceeded == false)
 		{
 			MergeResults.MergeResult = INiagaraMergeManager::EMergeEmitterResult::FailedToMerge;
@@ -1986,7 +1986,7 @@ FNiagaraEmitterDiffResults FNiagaraScriptMergeManager::DiffEmitters(UNiagaraEmit
 	DiffEventHandlers(BaseEmitterAdapter->GetEventHandlers(), OtherEmitterAdapter->GetEventHandlers(), EmitterDiffResults);
 	DiffSimulationStages(BaseEmitterAdapter->GetSimulationStages(), OtherEmitterAdapter->GetSimulationStages(), EmitterDiffResults);
 	DiffRenderers(BaseEmitterAdapter->GetRenderers(), OtherEmitterAdapter->GetRenderers(), EmitterDiffResults);
-	DiffInputSummaryEntries(BaseEmitterAdapter->GetEditorData(), OtherEmitterAdapter->GetEditorData(), EmitterDiffResults);
+	DiffEmitterSummary(BaseEmitterAdapter->GetEditorData(), OtherEmitterAdapter->GetEditorData(), EmitterDiffResults);
 	DiffEditableProperties(&BaseEmitter, &OtherEmitter, *UNiagaraEmitter::StaticClass(), EmitterDiffResults.DifferentEmitterProperties);
 	DiffStackEntryDisplayNames(BaseEmitterAdapter->GetEditorData(), OtherEmitterAdapter->GetEditorData(), EmitterDiffResults.ModifiedStackEntryDisplayNames);
 
@@ -2207,7 +2207,7 @@ void FNiagaraScriptMergeManager::DiffRenderers(const TArray<TSharedRef<FNiagaraR
 	}
 }
 
-void FNiagaraScriptMergeManager::DiffInputSummaryEntries(const UNiagaraEmitterEditorData* BaseEditorData, const UNiagaraEmitterEditorData* OtherEditorData, FNiagaraEmitterDiffResults& DiffResults) const
+void FNiagaraScriptMergeManager::DiffEmitterSummary(const UNiagaraEmitterEditorData* BaseEditorData, const UNiagaraEmitterEditorData* OtherEditorData, FNiagaraEmitterDiffResults& DiffResults) const
 {
 	TArray<TSharedRef<FNiagaraInputSummaryMergeAdapter>> BaseSummaryViewEntries;
 	TArray<TSharedRef<FNiagaraInputSummaryMergeAdapter>> OtherSummaryViewEntries;
@@ -2244,6 +2244,43 @@ void FNiagaraScriptMergeManager::DiffInputSummaryEntries(const UNiagaraEmitterEd
 			DiffResults.ModifiedOtherInputSummaryEntries.Add(CommonValuePair.OtherValue);
 		}
 	}
+
+	FListDiffResults<FNiagaraStackSection> SectionDiffResults = DiffLists<FNiagaraStackSection, FName>(
+		BaseEditorData->GetSummarySections(), 
+		OtherEditorData->GetSummarySections(),
+		[](const FNiagaraStackSection& Section) { return Section.SectionIdentifier; });
+
+	DiffResults.RemovedBaseSummarySections.Append(SectionDiffResults.RemovedBaseValues);
+	DiffResults.AddedOtherSummarySections.Append(SectionDiffResults.AddedOtherValues);
+
+	for (const FCommonValuePair<FNiagaraStackSection>& CommonValuePair : SectionDiffResults.CommonValuePairs)
+	{
+		bool bSectionsMatch = true;
+		if (CommonValuePair.BaseValue.SectionDisplayName.CompareTo(CommonValuePair.OtherValue.SectionDisplayName) == 0 &&
+			CommonValuePair.BaseValue.bEnabled == CommonValuePair.OtherValue.bEnabled &&
+			CommonValuePair.BaseValue.Categories.Num() == CommonValuePair.OtherValue.Categories.Num())
+		{
+			for (int32 CategoryIndex = 0; CategoryIndex < CommonValuePair.BaseValue.Categories.Num(); ++CategoryIndex)
+			{
+				if (CommonValuePair.OtherValue.Categories.ContainsByPredicate([&CommonValuePair, CategoryIndex](const FText& Category)
+					{ return Category.CompareTo(CommonValuePair.BaseValue.Categories[CategoryIndex]) == 0; }) == false)
+				{
+					bSectionsMatch = false;
+					break;
+				}
+			}
+		}
+		else
+		{
+			bSectionsMatch = false;
+		}
+
+		if (bSectionsMatch == false)
+		{
+			DiffResults.ModifiedBaseSummarySections.Add(CommonValuePair.BaseValue);
+			DiffResults.ModifiedOtherSummarySections.Add(CommonValuePair.OtherValue);
+		}
+	}
 }
 
 void FNiagaraScriptMergeManager::DiffScriptStacks(TSharedRef<FNiagaraScriptStackMergeAdapter> BaseScriptStackAdapter, TSharedRef<FNiagaraScriptStackMergeAdapter> OtherScriptStackAdapter, FNiagaraScriptStackDiffResults& DiffResults) const
@@ -3343,7 +3380,7 @@ FNiagaraScriptMergeManager::FApplyDiffResults FNiagaraScriptMergeManager::ApplyR
 	return Results;
 }
 
-FNiagaraScriptMergeManager::FApplyDiffResults FNiagaraScriptMergeManager::ApplyInputSummaryDiff(UNiagaraEmitter& Emitter, const FNiagaraEmitterDiffResults& DiffResults, const bool bNoParentAtLastMerge) const
+FNiagaraScriptMergeManager::FApplyDiffResults FNiagaraScriptMergeManager::ApplyEmitterSummaryDiff(UNiagaraEmitter& Emitter, const FNiagaraEmitterDiffResults& DiffResults, const bool bNoParentAtLastMerge) const
 {
 	UNiagaraEmitterEditorData* EditorData = Cast<UNiagaraEmitterEditorData>(Emitter.GetEditorData());
 
@@ -3360,7 +3397,7 @@ FNiagaraScriptMergeManager::FApplyDiffResults FNiagaraScriptMergeManager::ApplyI
 	EditorData->Modify();
 	for (TSharedRef<FNiagaraInputSummaryMergeAdapter> RemovedInput : DiffResults.RemovedInputSummaryEntries)
 	{
-		EditorData->SetSummaryViewMetaData(RemovedInput->GetKey(), FFunctionInputSummaryViewMetadata());
+		// Don't support removed summary inputs for now.  The user can change the visible flag if needed.
 	}
 
 	for (TSharedRef<FNiagaraInputSummaryMergeAdapter> AddedInput : DiffResults.AddedInputSummaryEntries)
@@ -3373,6 +3410,41 @@ FNiagaraScriptMergeManager::FApplyDiffResults FNiagaraScriptMergeManager::ApplyI
 		EditorData->SetSummaryViewMetaData(ModifiedInput->GetKey(), ModifiedInput->GetValue());
 	}
 
+	TArray<FNiagaraStackSection> SummarySections = EditorData->GetSummarySections();
+	for (const FNiagaraStackSection& RemovedBaseSection : DiffResults.RemovedBaseSummarySections)
+	{
+		// Don't support removed summary sections for now.  The user can change the enabled flag if needed.
+	}
+	
+	for (const FNiagaraStackSection& AddedOtherSection : DiffResults.AddedOtherSummarySections)
+	{
+		SummarySections.Add(AddedOtherSection);
+	}
+
+	for (int32 ModifiedSectionIndex = 0; ModifiedSectionIndex < DiffResults.ModifiedBaseSummarySections.Num(); ++ModifiedSectionIndex)
+	{
+		const FNiagaraStackSection& BaseSummarySection = DiffResults.ModifiedBaseSummarySections[ModifiedSectionIndex];
+		const FNiagaraStackSection& OtherSummarySection = DiffResults.ModifiedBaseSummarySections[ModifiedSectionIndex];
+		FNiagaraStackSection* ModifiedSectionPtr = SummarySections.FindByPredicate([BaseSummarySection](const FNiagaraStackSection& SummarySection)
+			{ return SummarySection.SectionDisplayName.CompareTo(BaseSummarySection.SectionDisplayName) == 0; });
+		if (ModifiedSectionPtr != nullptr)
+		{
+			// Use the enabled state in the child.
+			ModifiedSectionPtr->bEnabled = OtherSummarySection.bEnabled;
+			for (FText OtherSummarySectionCategory : OtherSummarySection.Categories)
+			{
+				// Add any categories added by the child, while ignoring any removed by the child.
+				if (ModifiedSectionPtr->Categories.ContainsByPredicate([OtherSummarySectionCategory](FText& Category)
+					{ return Category.CompareTo(OtherSummarySectionCategory) == 0; }) == false)
+				{
+					ModifiedSectionPtr->Categories.Add(OtherSummarySectionCategory);
+				}
+			}
+		}
+	}
+
+	EditorData->SetSummarySections(SummarySections);
+
 	FApplyDiffResults Results;
 	Results.bSucceeded = true;
 	Results.bModifiedGraph = false;
diff --git a/Engine/Plugins/FX/Niagara/Source/NiagaraEditor/Private/NiagaraScriptMergeManager.h b/Engine/Plugins/FX/Niagara/Source/NiagaraEditor/Private/NiagaraScriptMergeManager.h
--- a/Engine/Plugins/FX/Niagara/Source/NiagaraEditor/Private/NiagaraScriptMergeManager.h
+++ b/Engine/Plugins/FX/Niagara/Source/NiagaraEditor/Private/NiagaraScriptMergeManager.h
@@ -386,8 +386,12 @@ struct FNiagaraEmitterDiffResults
 	TArray<TSharedRef<FNiagaraInputSummaryMergeAdapter>> AddedInputSummaryEntries;
 	TArray<TSharedRef<FNiagaraInputSummaryMergeAdapter>> ModifiedInputSummaryEntries;
 	TArray<TSharedRef<FNiagaraInputSummaryMergeAdapter>> ModifiedOtherInputSummaryEntries;
-	
 
+	TArray<FNiagaraStackSection> RemovedBaseSummarySections;
+	TArray<FNiagaraStackSection> AddedOtherSummarySections;
+	TArray<FNiagaraStackSection> ModifiedBaseSummarySections;
+	TArray<FNiagaraStackSection> ModifiedOtherSummarySections;
+	
 	bool bScratchPadModified;
 
 	TMap<FString, FText> ModifiedStackEntryDisplayNames;
@@ -458,7 +462,7 @@ public:
 
 	void DiffRenderers(const TArray<TSharedRef<FNiagaraRendererMergeAdapter>>& BaseRenderers, const TArray<TSharedRef<FNiagaraRendererMergeAdapter>>& OtherRenderers, FNiagaraEmitterDiffResults& DiffResults) const;
 	
-	void DiffInputSummaryEntries(const UNiagaraEmitterEditorData* BaseEditorData, const UNiagaraEmitterEditorData* OtherEditorData, FNiagaraEmitterDiffResults& DiffResults) const;
+	void DiffEmitterSummary(const UNiagaraEmitterEditorData* BaseEditorData, const UNiagaraEmitterEditorData* OtherEditorData, FNiagaraEmitterDiffResults& DiffResults) const;
 
 	void DiffScriptStacks(TSharedRef<FNiagaraScriptStackMergeAdapter> BaseScriptStackAdapter, TSharedRef<FNiagaraScriptStackMergeAdapter> OtherScriptStackAdapter, FNiagaraScriptStackDiffResults& DiffResults) const;
 
@@ -507,7 +511,7 @@ private:
 
 	FApplyDiffResults ApplyRendererDiff(UNiagaraEmitter& BaseEmitter, const FNiagaraEmitterDiffResults& DiffResults, const bool bNoParentAtLastMerge) const;
 
-	FApplyDiffResults ApplyInputSummaryDiff(UNiagaraEmitter& Emitter, const FNiagaraEmitterDiffResults& DiffResults, const bool bNoParentAtLastMerge) const;
+	FApplyDiffResults ApplyEmitterSummaryDiff(UNiagaraEmitter& Emitter, const FNiagaraEmitterDiffResults& DiffResults, const bool bNoParentAtLastMerge) const;
 	
 	FApplyDiffResults ApplyStackEntryDisplayNameDiffs(UNiagaraEmitter& Emitter, const FNiagaraEmitterDiffResults& DiffResults) const;
 
diff --git a/Engine/Plugins/FX/Niagara/Source/NiagaraEditor/Private/ViewModels/Stack/NiagaraStackFunctionInputCollection.cpp b/Engine/Plugins/FX/Niagara/Source/NiagaraEditor/Private/ViewModels/Stack/NiagaraStackFunctionInputCollection.cpp
--- a/Engine/Plugins/FX/Niagara/Source/NiagaraEditor/Private/ViewModels/Stack/NiagaraStackFunctionInputCollection.cpp
+++ b/Engine/Plugins/FX/Niagara/Source/NiagaraEditor/Private/ViewModels/Stack/NiagaraStackFunctionInputCollection.cpp
@@ -54,7 +54,7 @@ static TOptional<FFunctionInputSummaryViewKey> GetSummaryViewInputKeyForFunction
 	return TOptional<FFunctionInputSummaryViewKey>();
 }
 
-static bool ShouldShowInSummaryView(UNiagaraEmitter* Emitter, UNiagaraNodeFunctionCall* FunctionCall, const FNiagaraVariable& InputVariable, const TOptional<FNiagaraVariableMetaData>& InputMetaData)
+static bool ShouldShowInSummaryView(UNiagaraEmitter* Emitter, UNiagaraNodeFunctionCall* FunctionCall, const FNiagaraVariable& InputVariable, const TOptional<FNiagaraVariableMetaData>& InputMetaData, bool bIsSummaryInEditMode)
 {
 	const UNiagaraEmitterEditorData* EditorData = Emitter? Cast<UNiagaraEmitterEditorData>(Emitter->GetEditorData()) : nullptr;
 	if (EditorData)
@@ -62,7 +62,8 @@ static bool ShouldShowInSummaryView(UNiagaraEmitter* Emitter, UNiagaraNodeFuncti
 		TOptional<FFunctionInputSummaryViewKey> SummaryViewKey = GetSummaryViewInputKeyForFunctionInput(FunctionCall, InputVariable, InputMetaData);
 		if (SummaryViewKey.IsSet())
 		{
-			return EditorData->GetSummaryViewMetaData(SummaryViewKey.GetValue()).bVisible;
+			TOptional<FFunctionInputSummaryViewMetadata> SummaryViewMetadata = EditorData->GetSummaryViewMetaData(SummaryViewKey.GetValue());
+			return SummaryViewMetadata.IsSet() && (bIsSummaryInEditMode || SummaryViewMetadata->bVisible);
 		}
 	}
 	return false;
@@ -301,10 +302,10 @@ void UNiagaraStackFunctionInputCollectionBase::AppendInputsForFunctionCall(FFunc
 		const TOptional<FFunctionInputSummaryViewKey> SummaryViewKey = GetSummaryViewInputKeyForFunctionInput(InputFunctionCallNode, InputVariable, InputMetaData);
 		if (EditorData && bShouldApplySummaryFilter && SummaryViewKey.IsSet())
 		{
-			const FFunctionInputSummaryViewMetadata SummaryViewData = EditorData->GetSummaryViewMetaData(SummaryViewKey.GetValue());
-			if (SummaryViewData.Category != NAME_None)
+			TOptional<FFunctionInputSummaryViewMetadata> SummaryViewData = EditorData->GetSummaryViewMetaData(SummaryViewKey.GetValue());
+			if (SummaryViewData.IsSet() && SummaryViewData->Category != NAME_None)
 			{
-				InputCategory = FText::FromName(SummaryViewData.Category);
+				InputCategory = FText::FromName(SummaryViewData->Category);
 			}
 			else
 			{
@@ -321,11 +322,11 @@ void UNiagaraStackFunctionInputCollectionBase::AppendInputsForFunctionCall(FFunc
 					InputCategory = FText::FromString(*InputFunctionCallNode->GetFunctionName());
 				}				
 			}
-			EditorSortPriority = SummaryViewData.SortIndex != INDEX_NONE? SummaryViewData.SortIndex : EditorSortPriority;
-			DisplayName = (SummaryViewData.DisplayName != NAME_None) ? FText::FromName(SummaryViewData.DisplayName) : TOptional<FText>();		
+			EditorSortPriority = SummaryViewData.IsSet() && SummaryViewData->SortIndex != INDEX_NONE ? SummaryViewData->SortIndex : EditorSortPriority;
+			DisplayName = SummaryViewData.IsSet() && SummaryViewData->DisplayName != NAME_None ? FText::FromName(SummaryViewData->DisplayName) : TOptional<FText>();		
 		}
 		
-		return ShouldShowInSummaryView(Emitter, InputFunctionCallNode, InputVariable, InputMetaData);
+		return ShouldShowInSummaryView(Emitter, InputFunctionCallNode, InputVariable, InputMetaData, GetEmitterViewModel().IsValid() && GetEmitterViewModel()->GetSummaryIsInEditMode());
 	};
 
 
diff --git a/Engine/Plugins/FX/Niagara/Source/NiagaraEditor/Public/NiagaraEmitterEditorData.h b/Engine/Plugins/FX/Niagara/Source/NiagaraEditor/Public/NiagaraEmitterEditorData.h
--- a/Engine/Plugins/FX/Niagara/Source/NiagaraEditor/Public/NiagaraEmitterEditorData.h
+++ b/Engine/Plugins/FX/Niagara/Source/NiagaraEditor/Public/NiagaraEmitterEditorData.h
@@ -38,11 +38,12 @@ public:
 	NIAGARAEDITOR_API bool ShouldShowSummaryView() const { return bShowSummaryView; }
 	NIAGARAEDITOR_API void ToggleShowSummaryView();
 	
-	NIAGARAEDITOR_API TMap<FFunctionInputSummaryViewKey, FFunctionInputSummaryViewMetadata> GetSummaryViewMetaDataMap() const;
-	NIAGARAEDITOR_API FFunctionInputSummaryViewMetadata GetSummaryViewMetaData(const FFunctionInputSummaryViewKey& Key) const;
-	NIAGARAEDITOR_API void SetSummaryViewMetaData(const FFunctionInputSummaryViewKey& Key, const FFunctionInputSummaryViewMetadata& NewMetadata);
+	NIAGARAEDITOR_API const TMap<FFunctionInputSummaryViewKey, FFunctionInputSummaryViewMetadata>& GetSummaryViewMetaDataMap() const;
+	NIAGARAEDITOR_API TOptional<FFunctionInputSummaryViewMetadata> GetSummaryViewMetaData(const FFunctionInputSummaryViewKey& Key) const;
+	NIAGARAEDITOR_API void SetSummaryViewMetaData(const FFunctionInputSummaryViewKey& Key, TOptional<FFunctionInputSummaryViewMetadata> NewMetadata);
 
-	const TArray<FNiagaraStackSection> GetSummarySections() const { return SummarySections; }
+	const TArray<FNiagaraStackSection>& GetSummarySections() const;
+	void SetSummarySections(const TArray<FNiagaraStackSection>& InSummarySections);
 
 private:
 	UPROPERTY(Instanced)
@@ -61,7 +62,7 @@ private:
 	UPROPERTY()
 	TMap<FFunctionInputSummaryViewKey, FFunctionInputSummaryViewMetadata> SummaryViewFunctionInputMetadata;
 
-	UPROPERTY(EditAnywhere, Category=Summary, meta=(NiagaraNoMerge))
+	UPROPERTY(EditAnywhere, Category=Summary)
 	TArray<FNiagaraStackSection> SummarySections;
 
 	FSimpleMulticastDelegate OnSummaryViewStateChangedDelegate;
diff --git a/Engine/Plugins/FX/Niagara/Source/NiagaraEditorWidgets/Private/Stack/SNiagaraStackFunctionInputValue.cpp b/Engine/Plugins/FX/Niagara/Source/NiagaraEditorWidgets/Private/Stack/SNiagaraStackFunctionInputValue.cpp
--- a/Engine/Plugins/FX/Niagara/Source/NiagaraEditorWidgets/Private/Stack/SNiagaraStackFunctionInputValue.cpp
+++ b/Engine/Plugins/FX/Niagara/Source/NiagaraEditorWidgets/Private/Stack/SNiagaraStackFunctionInputValue.cpp
@@ -41,6 +41,7 @@
 #include "Widgets/Images/SImage.h"
 #include "Widgets/Input/SButton.h"
 #include "Widgets/Input/SComboButton.h"
+#include "Widgets/Input/SCheckBox.h"
 #include "Widgets/Input/SEditableTextBox.h"
 #include "Widgets/Layout/SBox.h"
 #include "Widgets/NiagaraHLSLSyntaxHighlighter.h"
@@ -1371,7 +1372,8 @@ EVisibility SNiagaraStackFunctionInputValue::GetFilteredViewContextButtonVisibil
 	
 	UNiagaraStackFunctionInput* ParentInput = FNiagaraStackEditorWidgetsUtilities::GetParentInputForSummaryView(FunctionInput);
 	TOptional<FFunctionInputSummaryViewKey> Key = FNiagaraStackEditorWidgetsUtilities::GetSummaryViewInputKeyForFunctionInput(ParentInput);
-	return (Key.IsSet() && EditorData->GetSummaryViewMetaData(Key.GetValue()).bVisible)? EVisibility::Visible : EVisibility::Hidden;
+	TOptional<FFunctionInputSummaryViewMetadata> Metadata = Key.IsSet() ? EditorData->GetSummaryViewMetaData(Key.GetValue()) : TOptional<FFunctionInputSummaryViewMetadata>();
+	return Metadata.IsSet() ? EVisibility::Visible : EVisibility::Collapsed;
 }
 
 bool SNiagaraStackFunctionInputValue::GetSummaryViewChangeEnabledStateAllowed() const
@@ -1446,7 +1448,17 @@ TSharedRef<SWidget> SNiagaraStackFunctionInputValue::GetFilteredViewPropertiesCo
 		];
 	MenuBuilder.AddWidget(SortIndexWidget, LOCTEXT("FilteredViewSortIndex", "Sort Index"));
 
-
+	// Visible
+	TSharedRef<SWidget> VisibleWidget =
+		SNew(SBox)
+		.WidthOverride(100)
+		.Padding(FMargin(5, 0, 0, 0))
+		[
+			SNew(SCheckBox)
+			.IsChecked(this, &SNiagaraStackFunctionInputValue::GetFilteredViewVisibleCheckState)
+			.OnCheckStateChanged(this, &SNiagaraStackFunctionInputValue::FilteredVisibleCheckStateChanged)
+		];
+	MenuBuilder.AddWidget(VisibleWidget, LOCTEXT("FilteredViewVisible", "Is Visible"));
 
 	return MenuBuilder.MakeWidget();
 }
@@ -1455,10 +1467,11 @@ FText SNiagaraStackFunctionInputValue::GetFilteredViewDisplayName() const
 {
 	const UNiagaraEmitterEditorData* EditorData = FunctionInput->GetEmitterViewModel()? &FunctionInput->GetEmitterViewModel()->GetEditorData() : nullptr;
 	UNiagaraStackFunctionInput* ParentInput = FNiagaraStackEditorWidgetsUtilities::GetParentInputForSummaryView(FunctionInput);
-	TOptional<FFunctionInputSummaryViewKey> Key = FNiagaraStackEditorWidgetsUtilities::GetSummaryViewInputKeyForFunctionInput(ParentInput);
 
-	const FName SummaryViewName = (EditorData && Key.IsSet())? EditorData->GetSummaryViewMetaData(Key.GetValue()).DisplayName : NAME_None;
+	TOptional<FFunctionInputSummaryViewKey> Key = FNiagaraStackEditorWidgetsUtilities::GetSummaryViewInputKeyForFunctionInput(ParentInput);
+	TOptional<FFunctionInputSummaryViewMetadata> Metadata = EditorData != nullptr && Key.IsSet() ? EditorData->GetSummaryViewMetaData(Key.GetValue()) : TOptional<FFunctionInputSummaryViewMetadata>();
 
+	const FName SummaryViewName = Metadata.IsSet() ? Metadata->DisplayName : NAME_None;
 	return (SummaryViewName != NAME_None) ? FText::FromName(SummaryViewName) : FunctionInput->GetDisplayName();
 }
 
@@ -1484,13 +1497,17 @@ void SNiagaraStackFunctionInputValue::FilteredViewDisplayNameTextCommitted(const
 	if (EditorData && Key.IsSet())
 	{
 		const FName NewName = Text.IsEmptyOrWhitespace()? NAME_None : FName(Text.ToString());
-		FFunctionInputSummaryViewMetadata SummaryViewMetaData = EditorData->GetSummaryViewMetaData(Key.GetValue());		
-		if (NewName != SummaryViewMetaData.DisplayName)
+		TOptional<FFunctionInputSummaryViewMetadata> SummaryViewMetaData = EditorData->GetSummaryViewMetaData(Key.GetValue());	
+		if (SummaryViewMetaData.IsSet() == false)
+		{
+			SummaryViewMetaData = FFunctionInputSummaryViewMetadata();
+		}
+		if (NewName != SummaryViewMetaData->DisplayName)
 		{
 			FScopedTransaction ScopedTransaction(FText::Format(LOCTEXT("SummaryViewChangedInputDisplayName", "Changed summary view display name for {0}"), FunctionInput->GetDisplayName()));
 			EditorData->Modify();
 			
-			SummaryViewMetaData.DisplayName = NewName;
+			SummaryViewMetaData->DisplayName = NewName;
 			EditorData->SetSummaryViewMetaData(Key.GetValue(), SummaryViewMetaData);
 		}		
 	}
@@ -1507,10 +1524,11 @@ FText SNiagaraStackFunctionInputValue::GetFilteredViewCategory() const
 {
 	const UNiagaraEmitterEditorData* EditorData = FunctionInput->GetEmitterViewModel()? &FunctionInput->GetEmitterViewModel()->GetEditorData() : nullptr;	
 	UNiagaraStackFunctionInput* ParentInput = FNiagaraStackEditorWidgetsUtilities::GetParentInputForSummaryView(FunctionInput);
+	
 	TOptional<FFunctionInputSummaryViewKey> Key = FNiagaraStackEditorWidgetsUtilities::GetSummaryViewInputKeyForFunctionInput(ParentInput);
+	TOptional<FFunctionInputSummaryViewMetadata> Metadata = EditorData != nullptr && Key.IsSet() ? EditorData->GetSummaryViewMetaData(Key.GetValue()) : TOptional<FFunctionInputSummaryViewMetadata>();
 
-	const FName SummaryViewCategory = (EditorData && Key.IsSet())? EditorData->GetSummaryViewMetaData(Key.GetValue()).Category : NAME_None;
-
+	const FName SummaryViewCategory = Metadata.IsSet() ? Metadata->Category : NAME_None;
 	return (SummaryViewCategory != NAME_None) ? FText::FromName(SummaryViewCategory) : GetFunctionInputCategory(FunctionInput);
 }
 
@@ -1536,13 +1554,17 @@ void SNiagaraStackFunctionInputValue::FilteredViewCategoryTextCommitted(const FT
 	if (EditorData && Key.IsSet())
 	{
 		const FName NewCategory = Text.IsEmptyOrWhitespace()? NAME_None : FName(Text.ToString());
-		FFunctionInputSummaryViewMetadata SummaryViewMetaData = EditorData->GetSummaryViewMetaData(Key.GetValue());		
-		if (NewCategory != SummaryViewMetaData.Category)
+		TOptional<FFunctionInputSummaryViewMetadata> SummaryViewMetaData = EditorData->GetSummaryViewMetaData(Key.GetValue());	
+		if (SummaryViewMetaData.IsSet() == false)
+		{
+			SummaryViewMetaData = FFunctionInputSummaryViewMetadata();
+		}
+		if (NewCategory != SummaryViewMetaData->Category)
 		{
 			FScopedTransaction ScopedTransaction(FText::Format(LOCTEXT("SummaryViewChangedInputCategory", "Changed summary view category for {0}"), FunctionInput->GetDisplayName()));
 			EditorData->Modify();
 			
-			SummaryViewMetaData.Category = NewCategory;
+			SummaryViewMetaData->Category = NewCategory;
 			EditorData->SetSummaryViewMetaData(Key.GetValue(), SummaryViewMetaData);
 		}		
 	}
@@ -1553,15 +1575,12 @@ FText SNiagaraStackFunctionInputValue::GetFilteredViewSortIndex() const
 {
 	UNiagaraEmitterEditorData* EditorData = FunctionInput->GetEmitterViewModel()? &FunctionInput->GetEmitterViewModel()->GetOrCreateEditorData() : nullptr;	
 	UNiagaraStackFunctionInput* ParentInput = FNiagaraStackEditorWidgetsUtilities::GetParentInputForSummaryView(FunctionInput);
+
 	TOptional<FFunctionInputSummaryViewKey> Key = FNiagaraStackEditorWidgetsUtilities::GetSummaryViewInputKeyForFunctionInput(ParentInput);
+	TOptional<FFunctionInputSummaryViewMetadata> Metadata = EditorData != nullptr && Key.IsSet() ? EditorData->GetSummaryViewMetaData(Key.GetValue()) : TOptional<FFunctionInputSummaryViewMetadata>();
 
-	if (EditorData && Key.IsSet())
-	{
-		int32 SortIndex = EditorData->GetSummaryViewMetaData(Key.GetValue()).SortIndex;
-		return SortIndex != INDEX_NONE? FText::FromString(FString::FromInt(SortIndex)) : FText::GetEmpty();
-	}
-	
-	return FText::GetEmpty();
+	int32 SortIndex = Metadata.IsSet() ? Metadata->SortIndex : INDEX_NONE;
+	return SortIndex != INDEX_NONE ? FText::FromString(FString::FromInt(SortIndex)) : FText::GetEmpty();
 }
 
 bool SNiagaraStackFunctionInputValue::VerifyFilteredSortIndex(const FText& InText, FText& OutErrorMessage) const
@@ -1584,17 +1603,57 @@ void SNiagaraStackFunctionInputValue::FilteredSortIndexTextCommitted(const FText
 	
 	if (EditorData && Key.IsSet())
 	{		
-		FFunctionInputSummaryViewMetadata SummaryViewMetaData = EditorData->GetSummaryViewMetaData(Key.GetValue());		
-		if (NewValue != SummaryViewMetaData.SortIndex)
+		TOptional<FFunctionInputSummaryViewMetadata> SummaryViewMetaData = EditorData->GetSummaryViewMetaData(Key.GetValue());	
+		if (SummaryViewMetaData.IsSet() == false)
+		{
+			SummaryViewMetaData = FFunctionInputSummaryViewMetadata();
+		}
+		if (NewValue != SummaryViewMetaData->SortIndex)
 		{
 			FScopedTransaction ScopedTransaction(FText::Format(LOCTEXT("SummaryViewChangedInputSortIndex", "Changed summary view sorty index for {0}"), FunctionInput->GetDisplayName()));
 			EditorData->Modify();
 			
-			SummaryViewMetaData.SortIndex = NewValue;
+			SummaryViewMetaData->SortIndex = NewValue;
 			EditorData->SetSummaryViewMetaData(Key.GetValue(), SummaryViewMetaData);
 		}		
 	}
 }
 
+ECheckBoxState SNiagaraStackFunctionInputValue::GetFilteredViewVisibleCheckState() const
+{
+	UNiagaraEmitterEditorData* EditorData = FunctionInput->GetEmitterViewModel() ? &FunctionInput->GetEmitterViewModel()->GetOrCreateEditorData() : nullptr;
+	UNiagaraStackFunctionInput* ParentInput = FNiagaraStackEditorWidgetsUtilities::GetParentInputForSummaryView(FunctionInput);
+
+	TOptional<FFunctionInputSummaryViewKey> Key = FNiagaraStackEditorWidgetsUtilities::GetSummaryViewInputKeyForFunctionInput(ParentInput);
+	TOptional<FFunctionInputSummaryViewMetadata> Metadata = EditorData != nullptr && Key.IsSet() ? EditorData->GetSummaryViewMetaData(Key.GetValue()) : TOptional<FFunctionInputSummaryViewMetadata>();
+
+	return Metadata.IsSet() && Metadata->bVisible ? ECheckBoxState::Checked : ECheckBoxState::Unchecked;
+}
+
+void SNiagaraStackFunctionInputValue::FilteredVisibleCheckStateChanged(ECheckBoxState CheckBoxState)
+{
+	UNiagaraEmitterEditorData* EditorData = FunctionInput->GetEmitterViewModel() ? &FunctionInput->GetEmitterViewModel()->GetOrCreateEditorData() : nullptr;
+	UNiagaraStackFunctionInput* ParentInput = FNiagaraStackEditorWidgetsUtilities::GetParentInputForSummaryView(FunctionInput);
+	TOptional<FFunctionInputSummaryViewKey> Key = FNiagaraStackEditorWidgetsUtilities::GetSummaryViewInputKeyForFunctionInput(ParentInput);
+
+	bool bNewValue = CheckBoxState == ECheckBoxState::Checked;
+	if (EditorData && Key.IsSet())
+	{
+		TOptional<FFunctionInputSummaryViewMetadata> SummaryViewMetaData = EditorData->GetSummaryViewMetaData(Key.GetValue());
+		if (SummaryViewMetaData.IsSet() == false)
+		{
+			SummaryViewMetaData = FFunctionInputSummaryViewMetadata();
+		}
+		if (bNewValue != SummaryViewMetaData->bVisible)
+		{
+			FScopedTransaction ScopedTransaction(FText::Format(LOCTEXT("SummaryViewChangedInputSortIndex", "Changed summary view visibility for {0}"), FunctionInput->GetDisplayName()));
+			EditorData->Modify();
+
+			SummaryViewMetaData->bVisible = bNewValue;
+			EditorData->SetSummaryViewMetaData(Key.GetValue(), SummaryViewMetaData);
+		}
+	}
+}
+
 
 #undef LOCTEXT_NAMESPACE
diff --git a/Engine/Plugins/FX/Niagara/Source/NiagaraEditorWidgets/Private/Stack/SNiagaraStackFunctionInputValue.h b/Engine/Plugins/FX/Niagara/Source/NiagaraEditorWidgets/Private/Stack/SNiagaraStackFunctionInputValue.h
--- a/Engine/Plugins/FX/Niagara/Source/NiagaraEditorWidgets/Private/Stack/SNiagaraStackFunctionInputValue.h
+++ b/Engine/Plugins/FX/Niagara/Source/NiagaraEditorWidgets/Private/Stack/SNiagaraStackFunctionInputValue.h
@@ -176,6 +176,9 @@ private:
 	bool VerifyFilteredSortIndex(const FText& InText, FText& OutErrorMessage) const;
 	void FilteredSortIndexTextCommitted(const FText& Text, ETextCommit::Type CommitType);
 
+	ECheckBoxState GetFilteredViewVisibleCheckState() const;
+	void FilteredVisibleCheckStateChanged(ECheckBoxState CheckBoxState);
+
 private:
 	UNiagaraStackFunctionInput* FunctionInput;
 
diff --git a/Engine/Plugins/FX/Niagara/Source/NiagaraEditorWidgets/Private/Stack/SNiagaraStackTableRow.cpp b/Engine/Plugins/FX/Niagara/Source/NiagaraEditorWidgets/Private/Stack/SNiagaraStackTableRow.cpp
--- a/Engine/Plugins/FX/Niagara/Source/NiagaraEditorWidgets/Private/Stack/SNiagaraStackTableRow.cpp
+++ b/Engine/Plugins/FX/Niagara/Source/NiagaraEditorWidgets/Private/Stack/SNiagaraStackTableRow.cpp
@@ -642,12 +642,34 @@ void SNiagaraStackTableRow::ToggleShowInSummaryView()
 	
 		if (Key.IsSet())
 		{
-			FFunctionInputSummaryViewMetadata SummaryViewMetaData = EditorData->GetSummaryViewMetaData(Key.GetValue());
-			
+			// TODO: Move this parent handling to the UNiagaraStackFunctionInput and merge manager.
+			bool bHasParentSummaryData = false;;
+			UNiagaraEmitter* ParentEmitter = FunctionInput->GetEmitterViewModel()->GetEmitter()->GetParent();
+			if (ParentEmitter != nullptr)
+			{
+				UNiagaraEmitterEditorData* ParentEmitterEditorData = Cast<UNiagaraEmitterEditorData>(ParentEmitter->GetEditorData());
+				if (ParentEmitterEditorData != nullptr)
+				{
+					bHasParentSummaryData = ParentEmitterEditorData->GetSummaryViewMetaData(Key.GetValue()).IsSet();
+				}
+			}
+
+			TOptional<FFunctionInputSummaryViewMetadata> SummaryViewMetaData = EditorData->GetSummaryViewMetaData(Key.GetValue());
+			if (SummaryViewMetaData.IsSet() == false)
+			{
+				SummaryViewMetaData = FFunctionInputSummaryViewMetadata();
+			}
+
 			FScopedTransaction ScopedTransaction(FText::Format(LOCTEXT("SummaryViewChangedInputVisibility", "Changed summary view visibility for {0}"), FunctionInput->GetDisplayName()));
 			EditorData->Modify();
 			
-			SummaryViewMetaData.bVisible = !SummaryViewMetaData.bVisible;
+			SummaryViewMetaData->bVisible = !SummaryViewMetaData->bVisible;
+			if (bHasParentSummaryData == false && SummaryViewMetaData->bVisible == false)
+			{
+				// If there is no parent summary data, and the input is no longer visible, reset the optional value
+				// to remove the input from the summary.
+				SummaryViewMetaData.Reset();
+			}
 			EditorData->SetSummaryViewMetaData(Key.GetValue(), SummaryViewMetaData);
 		}
 	}	
@@ -665,7 +687,8 @@ bool SNiagaraStackTableRow::ShouldShowInSummaryView() const
 	
 		if (Key.IsSet())
 		{
-			return EditorData->GetSummaryViewMetaData(Key.GetValue()).bVisible;
+			TOptional<FFunctionInputSummaryViewMetadata> Metadata = EditorData->GetSummaryViewMetaData(Key.GetValue());
+			return Metadata.IsSet() && Metadata->bVisible;
 		}			
 	}
 
