Force DXC for NiagaraDataInterface shaders for platforms that support DXC

diff --git a/Engine/Plugins/FX/Niagara/Shaders/Private/NiagaraShaderVersion.ush b/Engine/Plugins/FX/Niagara/Shaders/Private/NiagaraShaderVersion.ush
--- a/Engine/Plugins/FX/Niagara/Shaders/Private/NiagaraShaderVersion.ush
+++ b/Engine/Plugins/FX/Niagara/Shaders/Private/NiagaraShaderVersion.ush
@@ -3,4 +3,4 @@
 // in the DDC key, it will automatically invalidate.
 // 
 // If you are merging streams and there is a conflict with this GUID you should make a new GUID rather than taking one or the other.
-// GUID = 656D78A42FD3A348BE8B24B136771045
\ No newline at end of file
+// GUID = 32F26874A0DC4FADAA73ED5FE11D0FD7
diff --git a/Engine/Plugins/FX/Niagara/Source/Niagara/Classes/NiagaraDataInterface.h b/Engine/Plugins/FX/Niagara/Source/Niagara/Classes/NiagaraDataInterface.h
--- a/Engine/Plugins/FX/Niagara/Source/Niagara/Classes/NiagaraDataInterface.h
+++ b/Engine/Plugins/FX/Niagara/Source/Niagara/Classes/NiagaraDataInterface.h
@@ -319,7 +319,7 @@ public:
 #endif
 
 	/** Allows data interfaces to influence the compilation of GPU shaders and is only called on the CDO object not the instance. */
-	virtual void ModifyCompilationEnvironment(EShaderPlatform ShaderPlatform, struct FShaderCompilerEnvironment& OutEnvironment) const {}
+	virtual void ModifyCompilationEnvironment(EShaderPlatform ShaderPlatform, struct FShaderCompilerEnvironment& OutEnvironment) const;
 
 	/** 
 		Subclasses that wish to work with GPU systems/emitters must implement this.
diff --git a/Engine/Plugins/FX/Niagara/Source/Niagara/Private/NiagaraDataInterface.cpp b/Engine/Plugins/FX/Niagara/Source/Niagara/Private/NiagaraDataInterface.cpp
--- a/Engine/Plugins/FX/Niagara/Source/Niagara/Private/NiagaraDataInterface.cpp
+++ b/Engine/Plugins/FX/Niagara/Source/Niagara/Private/NiagaraDataInterface.cpp
@@ -7,6 +7,7 @@
 #include "ShaderParameterUtils.h"
 #include "NiagaraShader.h"
 #include "NiagaraComponent.h"
+#include "ShaderCompilerCore.h"
 
 #define LOCTEXT_NAMESPACE "NiagaraDataInterface"
 
@@ -47,6 +48,15 @@ bool UNiagaraDataInterface::AppendCompileHash(FNiagaraCompileHashVisitor* InVisi
 }
 #endif
 
+void UNiagaraDataInterface::ModifyCompilationEnvironment(EShaderPlatform ShaderPlatform, struct FShaderCompilerEnvironment& OutEnvironment) const
+{
+	if (FDataDrivenShaderPlatformInfo::GetSupportsDxc(ShaderPlatform))
+	{
+		// Always enable DXC to avoid compile errors caused by RWBuffer/Buffer in structs. Example NiagaraDataInterfaceHairStrands.ush struct FDIHairStrandsContext
+		OutEnvironment.CompilerFlags.Add(CFLAG_ForceDXC);
+	}
+}
+
 void UNiagaraDataInterface::GetAssetTagsForContext(const UObject* InAsset, const TArray<const UNiagaraDataInterface*>& InProperties, TMap<FName, uint32>& NumericKeys, TMap<FName, FString>& StringKeys) const
 {
 	UClass* Class = GetClass();
