-Replaced the "S" textblock to indicate static parameters in the stack with an "S" icon
-Removed the type coloring to strongly associate red with the static concept, same as static switches

diff --git a/Engine/Plugins/FX/Niagara/Source/NiagaraEditorWidgets/Private/NiagaraEditorWidgetsStyle.cpp b/Engine/Plugins/FX/Niagara/Source/NiagaraEditorWidgets/Private/NiagaraEditorWidgetsStyle.cpp
--- a/Engine/Plugins/FX/Niagara/Source/NiagaraEditorWidgets/Private/NiagaraEditorWidgetsStyle.cpp
+++ b/Engine/Plugins/FX/Niagara/Source/NiagaraEditorWidgets/Private/NiagaraEditorWidgetsStyle.cpp
@@ -388,6 +388,8 @@ FNiagaraEditorWidgetsStyle::FNiagaraEditorWidgetsStyle() : FSlateStyleSet("Niaga
 	Set("NiagaraEditor.Stack.InputValueIconColor.Dynamic", FLinearColor(FColor::Cyan));
 	Set("NiagaraEditor.Stack.InputValueIconColor.Expression", FLinearColor(FColor::Green));
 	Set("NiagaraEditor.Stack.InputValueIconColor.Default", FLinearColor(FColor::White));
+	
+	Set("NiagaraEditor.Stack.StaticInputValue", new IMAGE_BRUSH("Icons/StaticBoolStack", Icon16x16, FLinearColor::White));
 
  	Set("NiagaraEditor.Stack.DropTarget.BackgroundColor", FLinearColor(1.0f, 1.0f, 1.0f, 0.25f));
  	Set("NiagaraEditor.Stack.DropTarget.BackgroundColorHover", FLinearColor(1.0f, 1.0f, 1.0f, 0.1f));
diff --git a/Engine/Plugins/FX/Niagara/Source/NiagaraEditorWidgets/Private/Stack/SNiagaraStackFunctionInputValue.cpp b/Engine/Plugins/FX/Niagara/Source/NiagaraEditorWidgets/Private/Stack/SNiagaraStackFunctionInputValue.cpp
--- a/Engine/Plugins/FX/Niagara/Source/NiagaraEditorWidgets/Private/Stack/SNiagaraStackFunctionInputValue.cpp
+++ b/Engine/Plugins/FX/Niagara/Source/NiagaraEditorWidgets/Private/Stack/SNiagaraStackFunctionInputValue.cpp
@@ -110,10 +110,9 @@ void SNiagaraStackFunctionInputValue::Construct(const FArguments& InArgs, UNiaga
 					.VAlign(VAlign_Center)
 					.Visibility(this, &SNiagaraStackFunctionInputValue::GetTypeModifierIconVisibility)
 					[
-						SNew(STextBlock)
-						.Text(this, &SNiagaraStackFunctionInputValue::GetTypeModifierIconText)
-						.ToolTipText(this, &SNiagaraStackFunctionInputValue::GetTypeModifierIconToolTip)
-						.ColorAndOpacity(this, &SNiagaraStackFunctionInputValue::GetTypeModifierIconColor)
+						SNew(SImage)
+						.Image(GetTypeModifierIcon())
+						.ToolTipText(GetTypeModifierIconToolTip())
 					]
 				]
 				+ SHorizontalBox::Slot()
@@ -810,9 +809,9 @@ EVisibility SNiagaraStackFunctionInputValue::GetTypeModifierIconVisibility() con
 	return FunctionInput->GetInputType().IsStatic() ? EVisibility::Visible : EVisibility::Collapsed;
 }
 
-FText SNiagaraStackFunctionInputValue::GetTypeModifierIconText() const
+const FSlateBrush* SNiagaraStackFunctionInputValue::GetTypeModifierIcon() const
 {
-	return FText::FromString(FString(TEXT("S")));
+	return FNiagaraEditorWidgetsStyle::Get().GetBrush("NiagaraEditor.Stack.StaticInputValue");
 }
 
 FText SNiagaraStackFunctionInputValue::GetTypeModifierIconToolTip() const
diff --git a/Engine/Plugins/FX/Niagara/Source/NiagaraEditorWidgets/Private/Stack/SNiagaraStackFunctionInputValue.h b/Engine/Plugins/FX/Niagara/Source/NiagaraEditorWidgets/Private/Stack/SNiagaraStackFunctionInputValue.h
--- a/Engine/Plugins/FX/Niagara/Source/NiagaraEditorWidgets/Private/Stack/SNiagaraStackFunctionInputValue.h
+++ b/Engine/Plugins/FX/Niagara/Source/NiagaraEditorWidgets/Private/Stack/SNiagaraStackFunctionInputValue.h
@@ -137,7 +137,7 @@ private:
 
 	EVisibility GetTypeModifierIconVisibility() const;
 
-	FText GetTypeModifierIconText() const;
+	const FSlateBrush* GetTypeModifierIcon() const;
 
 	FText GetTypeModifierIconToolTip() const;
 
