Fixed crash when undoing an edit to a system when the Niagara editor has been closed.

diff --git a/Engine/Plugins/FX/Niagara/Source/Niagara/Classes/NiagaraScript.h b/Engine/Plugins/FX/Niagara/Source/Niagara/Classes/NiagaraScript.h
--- a/Engine/Plugins/FX/Niagara/Source/Niagara/Classes/NiagaraScript.h
+++ b/Engine/Plugins/FX/Niagara/Source/Niagara/Classes/NiagaraScript.h
@@ -927,6 +927,7 @@ public:
 #if WITH_EDITOR
 	virtual void PostEditChangeProperty(FPropertyChangedEvent& PropertyChangedEvent) override;
 	virtual void PostEditChangeVersionedProperty(FPropertyChangedEvent& PropertyChangedEvent, const FGuid& Version);
+	virtual void PostEditUndo();
 #endif
 	virtual void GetAssetRegistryTags(TArray<FAssetRegistryTag>& OutTags) const override;
 
diff --git a/Engine/Plugins/FX/Niagara/Source/Niagara/Private/NiagaraScript.cpp b/Engine/Plugins/FX/Niagara/Source/Niagara/Private/NiagaraScript.cpp
--- a/Engine/Plugins/FX/Niagara/Source/Niagara/Private/NiagaraScript.cpp
+++ b/Engine/Plugins/FX/Niagara/Source/Niagara/Private/NiagaraScript.cpp
@@ -1970,6 +1970,17 @@ void UNiagaraScript::GenerateStatIDs()
 
 #if WITH_EDITOR
 
+void UNiagaraScript::PostEditUndo()
+{
+	Super::PostEditUndo();
+
+	FNiagaraSystemUpdateContext UpdateContext;
+	UpdateContext.SetDestroyOnAdd(true);
+	UpdateContext.Add(this, true);
+
+	InvalidateCompileResults(TEXT("Post Undo"));
+}
+
 void UNiagaraScript::PostEditChangeProperty(FPropertyChangedEvent& PropertyChangedEvent)
 {
 	Super::PostEditChangeProperty(PropertyChangedEvent);
diff --git a/Engine/Plugins/FX/Niagara/Source/NiagaraEditor/Private/ViewModels/NiagaraSystemViewModel.cpp b/Engine/Plugins/FX/Niagara/Source/NiagaraEditor/Private/ViewModels/NiagaraSystemViewModel.cpp
--- a/Engine/Plugins/FX/Niagara/Source/NiagaraEditor/Private/ViewModels/NiagaraSystemViewModel.cpp
+++ b/Engine/Plugins/FX/Niagara/Source/NiagaraEditor/Private/ViewModels/NiagaraSystemViewModel.cpp
@@ -645,13 +645,6 @@ void FNiagaraSystemViewModel::AddReferencedObjects(FReferenceCollector& Collecto
 void FNiagaraSystemViewModel::PostUndo(bool bSuccess)
 {
 	// Reset system stack and emitter handle view models to prevent accessing invalid data if they were in the undo operation.
-
-	//We have to invalidate all cached data that may have been directly modified by the transaction.
-	System->ForEachScript([](UNiagaraScript* Script)
-	{
-		Script->InvalidateCompileResults(TEXT("Post Undo"));
-	});
-
 	SystemStackViewModel->Reset();
 	ResetEmitterHandleViewModelsAndTracks();
 	RefreshAll();
