- Fix undo crash when no version data

#rnx

diff --git a/Engine/Plugins/FX/Niagara/Source/Niagara/Private/NiagaraScript.cpp b/Engine/Plugins/FX/Niagara/Source/Niagara/Private/NiagaraScript.cpp
--- a/Engine/Plugins/FX/Niagara/Source/Niagara/Private/NiagaraScript.cpp
+++ b/Engine/Plugins/FX/Niagara/Source/Niagara/Private/NiagaraScript.cpp
@@ -3166,7 +3166,10 @@ void UNiagaraScript::InvalidateCompileResults(const FString& Reason)
 	CachedScriptVM.Reset();
 	ScriptResource->Invalidate();
 	CachedScriptVMId.Invalidate();
-	GetLastGeneratedVMId().Invalidate();
+	if (VersionData.Num() > 0)
+	{
+		GetLastGeneratedVMId().Invalidate();
+	}
 	CachedDefaultDataInterfaces.Reset();
 	CachedScriptVM.OptimizationTask.State = nullptr;
 }
