- FindParameterVariable is not consistent with FindParameterOffset in that it will allow user redirection

#rnx

diff --git a/Engine/Plugins/FX/Niagara/Source/Niagara/Private/NiagaraComponent.cpp b/Engine/Plugins/FX/Niagara/Source/Niagara/Private/NiagaraComponent.cpp
--- a/Engine/Plugins/FX/Niagara/Source/Niagara/Private/NiagaraComponent.cpp
+++ b/Engine/Plugins/FX/Niagara/Source/Niagara/Private/NiagaraComponent.cpp
@@ -2636,8 +2636,11 @@ static FNiagaraVariant GetParameterValueFromStore(const FNiagaraVariableBase& Va
 
 	TArray<uint8> DataValue;
 	DataValue.AddUninitialized(Var.GetSizeInBytes());
-	Store.CopyParameterData(Var, DataValue.GetData());
-	return FNiagaraVariant(DataValue);
+	if ( ensureMsgf(Store.CopyParameterData(Var, DataValue.GetData()), TEXT("Failed to copy parameter data for '%s'"), *Var.GetName().ToString()) )
+	{
+		return FNiagaraVariant(DataValue);
+	}
+	return FNiagaraVariant();
 }
 
 #if WITH_EDITOR
diff --git a/Engine/Plugins/FX/Niagara/Source/Niagara/Public/NiagaraParameterStore.h b/Engine/Plugins/FX/Niagara/Source/Niagara/Public/NiagaraParameterStore.h
--- a/Engine/Plugins/FX/Niagara/Source/Niagara/Public/NiagaraParameterStore.h
+++ b/Engine/Plugins/FX/Niagara/Source/Niagara/Public/NiagaraParameterStore.h
@@ -356,7 +356,7 @@ public:
 		return Off ? *Off : (int32)INDEX_NONE;
 	}
 	
-	const FNiagaraVariableWithOffset* FindParameterVariable(const FNiagaraVariable& Parameter, bool IgnoreType = false) const;
+	virtual const FNiagaraVariableWithOffset* FindParameterVariable(const FNiagaraVariable& Parameter, bool IgnoreType = false) const;
 
 	/** Gets the typed parameter data. */
 	template<typename T>
diff --git a/Engine/Plugins/FX/Niagara/Source/Niagara/Public/NiagaraUserRedirectionParameterStore.h b/Engine/Plugins/FX/Niagara/Source/Niagara/Public/NiagaraUserRedirectionParameterStore.h
--- a/Engine/Plugins/FX/Niagara/Source/Niagara/Public/NiagaraUserRedirectionParameterStore.h
+++ b/Engine/Plugins/FX/Niagara/Source/Niagara/Public/NiagaraUserRedirectionParameterStore.h
@@ -54,6 +54,26 @@ public:
 		}
 	}
 	
+	FORCEINLINE_DEBUGGABLE virtual const FNiagaraVariableWithOffset* FindParameterVariable(const FNiagaraVariable& Parameter, bool IgnoreType = false) const override
+	{
+		if (IgnoreType)
+		{
+			for (auto it = UserParameterRedirects.CreateConstIterator(); it; ++it)
+			{
+				if (it->Key.GetName() == Parameter.GetName())
+				{
+					return FNiagaraParameterStore::FindParameterVariable(it->Value);
+				}
+			}
+			return FNiagaraParameterStore::FindParameterVariable(Parameter, IgnoreType);
+		}
+		else
+		{
+			const FNiagaraVariable* Redirection = UserParameterRedirects.Find(Parameter);
+			return FNiagaraParameterStore::FindParameterVariable(Redirection ? *Redirection : Parameter);
+		}
+	}
+
 	virtual bool AddParameter(const FNiagaraVariable& Param, bool bInitialize = true, bool bTriggerRebind = true, int32* OutOffset = nullptr) override;
 	virtual bool RemoveParameter(const FNiagaraVariableBase& InVar) override;
 	virtual void InitFromSource(const FNiagaraParameterStore* SrcStore, bool bNotifyAsDirty) override;
