Fixing compile error.

diff --git a/Engine/Plugins/FX/Niagara/Source/Niagara/Private/NiagaraDataInterfaceSkeletalMeshConnectivity.cpp b/Engine/Plugins/FX/Niagara/Source/Niagara/Private/NiagaraDataInterfaceSkeletalMeshConnectivity.cpp
--- a/Engine/Plugins/FX/Niagara/Source/Niagara/Private/NiagaraDataInterfaceSkeletalMeshConnectivity.cpp
+++ b/Engine/Plugins/FX/Niagara/Source/Niagara/Private/NiagaraDataInterfaceSkeletalMeshConnectivity.cpp
@@ -293,9 +293,9 @@ static bool BuildAdjacencyBuffer(const FSkeletalMeshLODRenderData& LodRenderData
 
 		const FVector P[3] =
 		{
-			VertexBuffer.VertexPosition((FVector)V[0]),
-			VertexBuffer.VertexPosition((FVector)V[1]),
-			VertexBuffer.VertexPosition((FVector)V[2])
+			(FVector)VertexBuffer.VertexPosition(V[0]),
+			(FVector)VertexBuffer.VertexPosition(V[1]),
+			(FVector)VertexBuffer.VertexPosition(V[2])
 		};
 
 		if (SortBySize)
