Open the version selected in the graph when double clicking on a script node

diff --git a/Engine/Plugins/FX/Niagara/Source/NiagaraEditor/Private/NiagaraNodeFunctionCall.cpp b/Engine/Plugins/FX/Niagara/Source/NiagaraEditor/Private/NiagaraNodeFunctionCall.cpp
--- a/Engine/Plugins/FX/Niagara/Source/NiagaraEditor/Private/NiagaraNodeFunctionCall.cpp
+++ b/Engine/Plugins/FX/Niagara/Source/NiagaraEditor/Private/NiagaraNodeFunctionCall.cpp
@@ -7,7 +7,7 @@
 #include "NiagaraComponent.h"
 #include "NiagaraConstants.h"
 #include "NiagaraCustomVersion.h"
-#include "NiagaraDataInterfaceSkeletalMesh.h"
+#include "Subsystems/AssetEditorSubsystem.h"
 #include "NiagaraEditorUtilities.h"
 #include "NiagaraGraph.h"
 #include "NiagaraHlslTranslator.h"
@@ -955,6 +955,18 @@ UObject*  UNiagaraNodeFunctionCall::GetReferencedAsset() const
 	}
 }
 
+void UNiagaraNodeFunctionCall::OpenReferencedAsset() const
+{
+	if (FunctionScript && FunctionScript->GetOutermost() != GetOutermost())
+	{
+		if (FunctionScript->IsVersioningEnabled())
+		{
+			FunctionScript->VersionToOpenInEditor = SelectedScriptVersion;
+		}
+		GEditor->GetEditorSubsystem<UAssetEditorSubsystem>()->OpenEditorForAsset(FunctionScript);
+	}
+}
+
 void UNiagaraNodeFunctionCall::UpdateNodeErrorMessage()
 {
 	FVersionedNiagaraScriptData* ScriptData = GetScriptData();
diff --git a/Engine/Plugins/FX/Niagara/Source/NiagaraEditor/Private/Widgets/SNiagaraScriptGraph.cpp b/Engine/Plugins/FX/Niagara/Source/NiagaraEditor/Private/Widgets/SNiagaraScriptGraph.cpp
--- a/Engine/Plugins/FX/Niagara/Source/NiagaraEditor/Private/Widgets/SNiagaraScriptGraph.cpp
+++ b/Engine/Plugins/FX/Niagara/Source/NiagaraEditor/Private/Widgets/SNiagaraScriptGraph.cpp
@@ -20,17 +20,13 @@
 #include "NiagaraEditorSettings.h"
 #include "EdGraphSchema_Niagara.h"
 #include "ScopedTransaction.h"
-#include "NiagaraEditorUtilities.h"
 #include "NiagaraNodeFactory.h"
 #include "NiagaraEditorCommands.h"
 #include "NiagaraNodeOutput.h"
 #include "Widgets/Input/SButton.h"
-#include "Widgets/Images/SImage.h"
 #include "Layout/WidgetPath.h"
 #include "Framework/Application/SlateApplication.h"
 #include "GraphEditorActions.h"
-#include "Subsystems/AssetEditorSubsystem.h"
-#include "Editor.h"
 #include "EditorFontGlyphs.h"
 #include "Widgets/SNiagaraGraphActionMenu.h"
 
@@ -279,11 +275,7 @@ void SNiagaraScriptGraph::OnNodeDoubleClicked(UEdGraphNode* ClickedNode)
 	UNiagaraNode* NiagaraNode = Cast<UNiagaraNode>(ClickedNode);
 	if (NiagaraNode != nullptr)
 	{
-		UObject* ReferencedAsset = NiagaraNode->GetReferencedAsset();
-		if (ReferencedAsset != nullptr)
-		{
-			GEditor->GetEditorSubsystem<UAssetEditorSubsystem>()->OpenEditorForAsset(ReferencedAsset);
-		}
+		NiagaraNode->OpenReferencedAsset();
 	}
 }
 
diff --git a/Engine/Plugins/FX/Niagara/Source/NiagaraEditor/Public/NiagaraNode.h b/Engine/Plugins/FX/Niagara/Source/NiagaraEditor/Public/NiagaraNode.h
--- a/Engine/Plugins/FX/Niagara/Source/NiagaraEditor/Public/NiagaraNode.h
+++ b/Engine/Plugins/FX/Niagara/Source/NiagaraEditor/Public/NiagaraNode.h
@@ -70,6 +70,7 @@ public:
 
 	/** Gets the asset referenced by this node, or nullptr if there isn't one. */
 	virtual UObject* GetReferencedAsset() const { return nullptr; }
+	virtual void OpenReferencedAsset() const {};
 
 	/** Refreshes the node due to external changes, e.g. the underlying function changed for a function call node. Return true if the graph changed.*/
 	virtual bool RefreshFromExternalChanges() { return false; }
diff --git a/Engine/Plugins/FX/Niagara/Source/NiagaraEditor/Public/NiagaraNodeFunctionCall.h b/Engine/Plugins/FX/Niagara/Source/NiagaraEditor/Public/NiagaraNodeFunctionCall.h
--- a/Engine/Plugins/FX/Niagara/Source/NiagaraEditor/Public/NiagaraNodeFunctionCall.h
+++ b/Engine/Plugins/FX/Niagara/Source/NiagaraEditor/Public/NiagaraNodeFunctionCall.h
@@ -107,6 +107,7 @@ public:
 	//~ Begin UNiagaraNode Interface
 	virtual void Compile(class FHlslNiagaraTranslator* Translator, TArray<int32>& Outputs) override;
 	virtual UObject* GetReferencedAsset() const override;
+	virtual void OpenReferencedAsset() const override;
 	virtual bool RefreshFromExternalChanges() override;
 	virtual ENiagaraNumericOutputTypeSelectionMode GetNumericOutputTypeSelectionMode() const override;
 	virtual bool CanAddToGraph(UNiagaraGraph* TargetGraph, FString& OutErrorMsg) const override;
