Fix for Niagara meshes under raytracing with GPUScene

diff --git a/Engine/Plugins/FX/Niagara/Shaders/Private/NiagaraMeshVertexFactory.ush b/Engine/Plugins/FX/Niagara/Shaders/Private/NiagaraMeshVertexFactory.ush
--- a/Engine/Plugins/FX/Niagara/Shaders/Private/NiagaraMeshVertexFactory.ush
+++ b/Engine/Plugins/FX/Niagara/Shaders/Private/NiagaraMeshVertexFactory.ush
@@ -902,7 +902,7 @@ float4 VertexFactoryGetTranslatedPrimitiveVolumeBounds(FVertexFactoryInterpolant
 #if RAYHITGROUPSHADER
 	FVertexFactoryInput LoadVertexFactoryInputForHGS(uint TriangleIndex, int VertexIndex)
 	{
-		FVertexFactoryInput Input;
+		FVertexFactoryInput Input = (FVertexFactoryInput) 0;
 
 		FTriangleBaseAttributes Tri = LoadTriangleBaseAttributes(TriangleIndex);
 		uint VertexId = Tri.Indices[VertexIndex];
@@ -916,6 +916,12 @@ float4 VertexFactoryGetTranslatedPrimitiveVolumeBounds(FVertexFactoryInterpolant
 		Input.VertexId = VertexId;
 #endif // #if PASS_NEEDS_VERTEX_ID
 
+#if VF_USE_PRIMITIVE_SCENE_DATA
+		// Note: GetInstanceUserData() stores the GPU-Scene primitive ID
+		int PrimitiveId = GetInstanceUserData();
+		Input.InstanceIdOffset = PrimitiveId | VF_TREAT_INSTANCE_ID_OFFSET_AS_PRIMITIVE_ID_FLAG;
+#endif
+
 		return Input;
 	}
 #endif // #if COMPUTESHADER || RAYHITGROUPSHADER
