- Fix Position Array not being editable

#rnx

diff --git a/Engine/Plugins/FX/Niagara/Source/Niagara/Classes/NiagaraDataInterfaceArrayFloat.h b/Engine/Plugins/FX/Niagara/Source/Niagara/Classes/NiagaraDataInterfaceArrayFloat.h
--- a/Engine/Plugins/FX/Niagara/Source/Niagara/Classes/NiagaraDataInterfaceArrayFloat.h
+++ b/Engine/Plugins/FX/Niagara/Source/Niagara/Classes/NiagaraDataInterfaceArrayFloat.h
@@ -254,7 +254,7 @@ class NIAGARA_API UNiagaraDataInterfaceArrayPosition : public UNiagaraDataInterf
 {
 	GENERATED_UCLASS_BODY()
 public:
-	UPROPERTY(VisibleAnywhere, Category = "Array")
+	UPROPERTY(EditAnywhere, BlueprintReadWrite, Category = "Array")
 	TArray<FNiagaraPosition> PositionData;
 
 	TArray<FNiagaraPosition>& GetArrayReference() { return PositionData; }
