Scalability Toolkit Mode v0.1

diff --git a/Engine/Plugins/FX/Niagara/Source/Niagara/Classes/NiagaraEffectType.h b/Engine/Plugins/FX/Niagara/Source/Niagara/Classes/NiagaraEffectType.h
--- a/Engine/Plugins/FX/Niagara/Source/Niagara/Classes/NiagaraEffectType.h
+++ b/Engine/Plugins/FX/Niagara/Source/Niagara/Classes/NiagaraEffectType.h
@@ -170,7 +170,7 @@ struct FNiagaraSystemScalabilitySettings
 
 
 	/** Effects of this type are culled beyond this distance. */
-	UPROPERTY(EditAnywhere, Category = "Scalability", meta = (EditCondition = "bCullByDistance"))
+	UPROPERTY(EditAnywhere, Category = "Scalability", meta = (EditCondition = "bCullByDistance", DisplayInScalabilityValuesBar))
 	float MaxDistance;
 
 	/** 
@@ -178,7 +178,7 @@ struct FNiagaraSystemScalabilitySettings
 	If the effect type has a significance handler, instances are sorted by their significance and only the N most significant will be kept. The rest are culled.
 	If it does not have a significance handler, instance count culling will be applied at spawn time only. New FX that would exceed the counts are not spawned/activated.
 	*/
-	UPROPERTY(EditAnywhere, Category = "Scalability", DisplayName = "Max Effect Type Instances", meta = (EditCondition = "bCullMaxInstanceCount"))
+	UPROPERTY(EditAnywhere, Category = "Scalability", DisplayName = "Max Effect Type Instances", meta = (EditCondition = "bCullMaxInstanceCount", DisplayInScalabilityValuesBar))
 	int32 MaxInstances;
 
 	/**
@@ -186,7 +186,7 @@ struct FNiagaraSystemScalabilitySettings
 	If the effect type has a significance handler, instances are sorted by their significance and only the N most significant will be kept. The rest are culled.
 	If it does not have a significance handler, instance count culling will be applied at spawn time only. New FX that would exceed the counts are not spawned/activated.
 	*/
-	UPROPERTY(EditAnywhere, Category = "Scalability", meta = (EditCondition = "bCullPerSystemMaxInstanceCount"))
+	UPROPERTY(EditAnywhere, Category = "Scalability", meta = (EditCondition = "bCullPerSystemMaxInstanceCount", DisplayInScalabilityValuesBar))
 	int32 MaxSystemInstances;
 
 	//TODO:
@@ -194,11 +194,11 @@ struct FNiagaraSystemScalabilitySettings
 	//float ScreenFraction;
 
 	/** Effects will be culled if they go more than this length of time without being rendered. */
-	UPROPERTY(EditAnywhere, Category = "Scalability", meta = (EditCondition = "bCullByMaxTimeWithoutRender"))
+	UPROPERTY(EditAnywhere, Category = "Scalability", meta = (EditCondition = "bCullByMaxTimeWithoutRender", DisplayInScalabilityValuesBar))
 	float MaxTimeWithoutRender;
 		
 	/** Controls what, if any, proxy will be used in place of culled systems.  */
-	UPROPERTY(EditAnywhere, Category = "Scalability")
+	UPROPERTY(EditAnywhere, Category = "Scalability", meta = (DisplayInScalabilityValuesBar))
 	ENiagaraCullProxyMode CullProxyMode = ENiagaraCullProxyMode::None;
 
 	/** 
@@ -206,7 +206,7 @@ struct FNiagaraSystemScalabilitySettings
 	While much cheaper than full FX instances, proxies still incur some cost so must have a limit.
 	When significance information is available using a significance handler, the most significance proxies will be kept up to this limit.
 	*/
-	UPROPERTY(EditAnywhere, Category = "Scalability")
+	UPROPERTY(EditAnywhere, Category = "Scalability", meta = (EditCondition = "CullProxyMode != ENiagaraCullProxyMode::None", DisplayInScalabilityValuesBar))
 	int32 MaxSystemProxies = 32;
 
 	/** Settings related to scaling down FX based on the current budget usage. */
@@ -280,7 +280,7 @@ struct FNiagaraEmitterScalabilitySettings
 	uint32 bScaleSpawnCount : 1;
 
 	/** Scale factor applied to spawn counts for this emitter. */
-	UPROPERTY(EditAnywhere, Category = "Scalability", meta = (EditCondition = "bScaleSpawnCount"))
+	UPROPERTY(EditAnywhere, Category = "Scalability", meta = (EditCondition = "bScaleSpawnCount", DisplayInScalabilityValuesBar))
 	float SpawnCountScale;
 
 	FNiagaraEmitterScalabilitySettings();
@@ -379,19 +379,19 @@ class NIAGARA_API UNiagaraEffectType : public UObject
 	//UObject Interface END
 
 	/** Controls whether or not culling is allowed for FX that are owned by, attached to or instigated by a locally controlled pawn. */
-	UPROPERTY(EditAnywhere, Category = "Scalability")
+	UPROPERTY(EditAnywhere, Category = "Scalability", meta=(DisplayInSystemScalability, ScalabilityBarDisplayName="Local Player Culling"))
 	bool bAllowCullingForLocalPlayers = false;
 
 	/** How regularly effects of this type are checked for scalability. */
-	UPROPERTY(EditAnywhere, Category = "Scalability")
+	UPROPERTY(EditAnywhere, Category = "Scalability", meta=(DisplayInSystemScalability, ScalabilityBarDisplayName="Frequency"))
 	ENiagaraScalabilityUpdateFrequency UpdateFrequency;
 
 	/** How effects of this type react when they fail the cull checks. */
-	UPROPERTY(EditAnywhere, Category = "Scalability")
+	UPROPERTY(EditAnywhere, Category = "Scalability", meta=(DisplayInSystemScalability, ScalabilityBarDisplayName="Reaction"))
 	ENiagaraCullReaction CullReaction;
 
 	/** Used to determine the relative significance of FX in the scene which is used in other scalability systems such as instance count culling. */
-	UPROPERTY(EditAnywhere, Instanced, Category = "Scalability")
+	UPROPERTY(EditAnywhere, Instanced, Category = "Scalability", meta=(DisplayInSystemScalability, DisplayClassDisplayName, ScalabilityBarDisplayName="Significance"))
 	TObjectPtr<UNiagaraSignificanceHandler> SignificanceHandler;
 
 	/** Cull settings to use at each detail level. */
diff --git a/Engine/Plugins/FX/Niagara/Source/Niagara/Classes/NiagaraEmitter.h b/Engine/Plugins/FX/Niagara/Source/Niagara/Classes/NiagaraEmitter.h
--- a/Engine/Plugins/FX/Niagara/Source/Niagara/Classes/NiagaraEmitter.h
+++ b/Engine/Plugins/FX/Niagara/Source/Niagara/Classes/NiagaraEmitter.h
@@ -347,10 +347,10 @@ public:
 	UPROPERTY()
 	FNiagaraDetailsLevelScaleOverrides GlobalSpawnCountScaleOverrides_DEPRECATED;
 	
-	UPROPERTY(EditAnywhere, Category = "Scalability")
+	UPROPERTY(EditAnywhere, Category = "Scalability", meta=(DisplayInScalabilityContext))
 	FNiagaraPlatformSet Platforms;
 
-	UPROPERTY(EditAnywhere, Category = "Scalability")
+	UPROPERTY(EditAnywhere, Category = "Scalability", meta=(DisplayInScalabilityContext))
 	FNiagaraEmitterScalabilityOverrides ScalabilityOverrides;
 
 	/** When enabled, this will spawn using interpolated parameter values and perform a partial update at spawn time. This adds significant additional cost for spawning but will produce much smoother spawning for high spawn rates, erratic frame rates and fast moving emitters. */
@@ -489,6 +489,7 @@ public:
 #endif
 
 	FORCEINLINE const FNiagaraEmitterScalabilitySettings& GetScalabilitySettings()const { return CurrentScalabilitySettings; }
+	NIAGARA_API const FNiagaraEmitterScalabilityOverride& GetCurrentOverrideSettings() const;
 
 	/** Returns true if this emitter's platform filter allows it on this platform and quality level. */
 	NIAGARA_API bool IsAllowedByScalability()const;
@@ -516,6 +517,9 @@ public:
 
 	const TArray<UNiagaraRendererProperties*>& GetRenderers() const { return RendererProperties; }
 
+	template<typename TAction>
+	void ForEachRenderer(TAction Func) const;
+	
 	template<typename TAction>
 	void ForEachEnabledRenderer(TAction Func) const;
 
@@ -579,7 +583,7 @@ public:
 	NIAGARA_API void NotifyScratchPadScriptsChanged();
 #endif
 
-	void OnScalabilityCVarChanged();
+	void UpdateScalability();
 
 #if WITH_EDITORONLY_DATA
 	NIAGARA_API const TMap<FGuid, TObjectPtr<UNiagaraMessageDataBase>>& GetMessages() const { return MessageKeyToMessageMap; };
@@ -725,6 +729,17 @@ private:
 #endif
 };
 
+template<typename TAction>
+void UNiagaraEmitter::ForEachRenderer(TAction Func) const
+{
+	for (UNiagaraRendererProperties* Renderer : RendererProperties)
+	{
+		if (Renderer)
+		{
+			Func(Renderer);
+		}
+	}
+}
 
 template<typename TAction>
 void UNiagaraEmitter::ForEachEnabledRenderer(TAction Func) const
diff --git a/Engine/Plugins/FX/Niagara/Source/Niagara/Classes/NiagaraPlatformSet.h b/Engine/Plugins/FX/Niagara/Source/Niagara/Classes/NiagaraPlatformSet.h
--- a/Engine/Plugins/FX/Niagara/Source/Niagara/Classes/NiagaraPlatformSet.h
+++ b/Engine/Plugins/FX/Niagara/Source/Niagara/Classes/NiagaraPlatformSet.h
@@ -199,6 +199,9 @@ private:
 	static FCriticalSection ChangedDelegateHandlesCritSec;
 };
 
+void NIAGARA_API SetGNiagaraQualityLevel(int32 QualityLevel);
+void NIAGARA_API SetGNiagaraDeviceProfile(UDeviceProfile* Profile);
+
 USTRUCT()
 struct NIAGARA_API FNiagaraPlatformSet
 {
@@ -256,8 +259,7 @@ public:
 	In some cases this is not allowed.
 	*/
 	//bool Conflicts(const FNiagaraPlatformSet& Other, TArray<const UDeviceProfile*>& OutConflictingProfiles, bool bIncludeProfilesWithVariableQualityLevel=false)const;
-
-
+	
 	/** Direct state accessors for the UI. */
 	bool IsEffectQualityEnabled(int32 EffectQuality)const;
 	void SetEnabledForEffectQuality(int32 EffectQuality, bool bEnabled);
@@ -269,6 +271,11 @@ public:
 
 	/** Inspects the passed sets and generates an array of all conflicts between these sets. Used to keep arrays of platform sets orthogonal. */
 	static bool GatherConflicts(const TArray<const FNiagaraPlatformSet*>& PlatformSets, TArray<FNiagaraPlatformSetConflictInfo>& OutConflicts);
+	
+	DECLARE_DELEGATE_RetVal(int32, FOnOverrideQualityLevel);
+	FOnOverrideQualityLevel OnOverrideQualityLevelDelegate;
+	DECLARE_DELEGATE_RetVal(TOptional<TObjectPtr<UDeviceProfile>>, FOnOverrideActiveDeviceProfile);
+	FOnOverrideActiveDeviceProfile OnOverrideActiveDeviceProfileDelegate;
 #endif
 
 	/** Mask defining which effects qualities this set matches. */
diff --git a/Engine/Plugins/FX/Niagara/Source/Niagara/Classes/NiagaraSystem.h b/Engine/Plugins/FX/Niagara/Source/Niagara/Classes/NiagaraSystem.h
--- a/Engine/Plugins/FX/Niagara/Source/Niagara/Classes/NiagaraSystem.h
+++ b/Engine/Plugins/FX/Niagara/Source/Niagara/Classes/NiagaraSystem.h
@@ -295,6 +295,7 @@ public:
 #if WITH_EDITOR
 	DECLARE_MULTICAST_DELEGATE_OneParam(FOnSystemCompiled, UNiagaraSystem*);
 	DECLARE_MULTICAST_DELEGATE_OneParam(FOnSystemPostEditChange, UNiagaraSystem*);
+	DECLARE_MULTICAST_DELEGATE(FOnScalabilityChanged)
 #endif
 	//TestChange
 
@@ -472,6 +473,9 @@ public:
 	/** Delegate called on PostEditChange.*/
 	FOnSystemPostEditChange& OnSystemPostEditChange();
 
+	/** Delegate called on effect type or effect type value change */
+	FOnScalabilityChanged& OnScalabilityChanged();
+
 	/** Gets editor specific data stored with this system. */
 	UNiagaraEditorDataBase* GetEditorData();
 
@@ -721,6 +725,7 @@ public:
 #if WITH_EDITOR
 	void SetEffectType(UNiagaraEffectType* EffectType);
 
+	FNiagaraSystemScalabilityOverrides& GetSystemScalabilityOverrides() { return SystemScalabilityOverrides; }
 	FORCEINLINE bool GetOverrideScalabilitySettings()const { return bOverrideScalabilitySettings; }
 	FORCEINLINE void SetOverrideScalabilitySettings(bool bOverride) { bOverrideScalabilitySettings = bOverride; }
 
@@ -729,9 +734,10 @@ public:
 #endif
 	UNiagaraEffectType* GetEffectType()const;
 	FORCEINLINE const FNiagaraSystemScalabilitySettings& GetScalabilitySettings()const { return CurrentScalabilitySettings; }
+	const FNiagaraSystemScalabilityOverride& GetCurrentOverrideSettings() const;
 	FORCEINLINE bool NeedsSortedSignificanceCull()const{ return bNeedsSortedSignificanceCull; }
 	
-	void OnScalabilityCVarChanged();
+	void UpdateScalability();
 
 	FORCEINLINE ENiagaraCullProxyMode GetCullProxyMode()const { return GetScalabilitySettings().CullProxyMode; }
 
@@ -813,7 +819,7 @@ protected:
 	UPROPERTY(EditAnywhere, Category = "System")
 	TObjectPtr<UNiagaraEffectType> EffectType;
 
-	UPROPERTY(EditAnywhere, Category = "Scalability")
+	UPROPERTY(EditAnywhere, Category = "Scalability", meta=(DisplayInScalabilityContext))
 	bool bOverrideScalabilitySettings;
 
 	/** Controls whether we should override the Effect Type value for bAllowCullingForLocalPlayers. */
@@ -827,7 +833,7 @@ protected:
 	UPROPERTY()
 	TArray<FNiagaraSystemScalabilityOverride> ScalabilityOverrides_DEPRECATED;
 
-	UPROPERTY(EditAnywhere, Category = "Scalability", meta = (EditCondition="bOverrideScalabilitySettings"))
+	UPROPERTY(EditAnywhere, Category = "Scalability", meta = (EditCondition="bOverrideScalabilitySettings", DisplayInScalabilityContext))
 	FNiagaraSystemScalabilityOverrides SystemScalabilityOverrides;
 
 	/** Handles to the emitter this System will simulate. */
@@ -883,6 +889,9 @@ protected:
 
 	/** A multicast delegate which is called whenever this system's properties are changed. */
 	FOnSystemPostEditChange OnSystemPostEditChangeDelegate;
+
+	/** A multicast delegate that is called whenever the effect type or the effect type values are changed */
+	FOnScalabilityChanged OnScalabilityChangedDelegate;
 #endif
 
 	/** The fixed bounding box value. bFixedBounds is the condition whether the fixed bounds can be edited. */
diff --git a/Engine/Plugins/FX/Niagara/Source/Niagara/Private/NiagaraEffectType.cpp b/Engine/Plugins/FX/Niagara/Source/Niagara/Private/NiagaraEffectType.cpp
--- a/Engine/Plugins/FX/Niagara/Source/Niagara/Private/NiagaraEffectType.cpp
+++ b/Engine/Plugins/FX/Niagara/Source/Niagara/Private/NiagaraEffectType.cpp
@@ -109,7 +109,7 @@ void UNiagaraEffectType::PostEditChangeProperty(struct FPropertyChangedEvent& Pr
 		UNiagaraSystem* System = *It;
 		if (System->GetEffectType() == this)
 		{
-			System->OnScalabilityCVarChanged();
+			System->UpdateScalability();
 			UpdateContext.Add(System, true);
 		}
 	}
diff --git a/Engine/Plugins/FX/Niagara/Source/Niagara/Private/NiagaraEmitter.cpp b/Engine/Plugins/FX/Niagara/Source/Niagara/Private/NiagaraEmitter.cpp
--- a/Engine/Plugins/FX/Niagara/Source/Niagara/Private/NiagaraEmitter.cpp
+++ b/Engine/Plugins/FX/Niagara/Source/Niagara/Private/NiagaraEmitter.cpp
@@ -1500,6 +1500,20 @@ void UNiagaraEmitter::UpdateEmitterAfterLoad()
 #endif
 }
 
+const FNiagaraEmitterScalabilityOverride& UNiagaraEmitter::GetCurrentOverrideSettings() const
+{
+	for (const FNiagaraEmitterScalabilityOverride& Override : ScalabilityOverrides.Overrides)
+	{
+		if (Override.Platforms.IsActive())
+		{
+			return Override;
+		}
+	}	
+
+	static FNiagaraEmitterScalabilityOverride Dummy;
+	return Dummy;
+}
+
 bool UNiagaraEmitter::IsAllowedByScalability()const
 {
 	return Platforms.IsActive();
@@ -2457,7 +2471,7 @@ void UNiagaraEmitter::ResolveScalabilitySettings()
 	}
 }
 
-void UNiagaraEmitter::OnScalabilityCVarChanged()
+void UNiagaraEmitter::UpdateScalability()
 {
 	ResolveScalabilitySettings();
 }
diff --git a/Engine/Plugins/FX/Niagara/Source/Niagara/Private/NiagaraPlatformSet.cpp b/Engine/Plugins/FX/Niagara/Source/Niagara/Private/NiagaraPlatformSet.cpp
--- a/Engine/Plugins/FX/Niagara/Source/Niagara/Private/NiagaraPlatformSet.cpp
+++ b/Engine/Plugins/FX/Niagara/Source/Niagara/Private/NiagaraPlatformSet.cpp
@@ -43,13 +43,40 @@ static FAutoConsoleVariableRef CVarPruneEmittersOnCook(
 static const int32 DefaultQualityLevel = 3;
 const TCHAR* NiagaraQualityLevelName = TEXT("fx.Niagara.QualityLevel");
 int32 GNiagaraQualityLevel = DefaultQualityLevel;
-static FAutoConsoleVariableRef CVarNiagaraQualityLevel(
-	NiagaraQualityLevelName,
-	GNiagaraQualityLevel,
-	TEXT("The quality level for Niagara Effects. \n"),
-	FConsoleVariableDelegate::CreateStatic(&FNiagaraPlatformSet::OnQualityLevelChanged),
-	ECVF_Scalability | ECVF_Preview
-);
+
+static FAutoConsoleCommand GCmdSetNiagaraQualityLevelOverride(
+	TEXT("fx.Niagara.SetOverrideQualityLevel"),
+	TEXT("Sets which quality level we should override with, no args means reset to default (Epic). Valid levels are 0-4 (Low-Cinematic)"),
+	FConsoleCommandWithArgsDelegate::CreateLambda(
+		[](const TArray<FString>& Args)
+		{
+			if(Args.Num() == 0)
+			{
+				SetGNiagaraQualityLevel(DefaultQualityLevel);
+				UE_LOG(LogNiagara, Display, TEXT("Reset the global Niagara quality level to %s"), *FNiagaraPlatformSet::GetQualityLevelText(GNiagaraQualityLevel).ToString());
+			}
+			else if(Args.Num() == 1)
+			{
+				if(FCString::IsNumeric(*Args[0]))
+				{
+					int32 NewQualityLevel = FCString::Atoi(*Args[0]);
+					if(ensure(NewQualityLevel >= 0 && NewQualityLevel <= 4))
+					{
+						SetGNiagaraQualityLevel(NewQualityLevel);
+						UE_LOG(LogNiagara, Display, TEXT("Set the global Niagara quality level to %s"), *FNiagaraPlatformSet::GetQualityLevelText(GNiagaraQualityLevel).ToString());
+					}
+				}
+			}
+		}
+	)
+); 
+// static FAutoConsoleVariableRef CVarNiagaraQualityLevel(
+// 	NiagaraQualityLevelName,
+// 	GNiagaraQualityLevel,
+// 	TEXT("The quality level for Niagara Effects. \n"),
+// 	FConsoleVariableDelegate::CreateStatic(&FNiagaraPlatformSet::OnQualityLevelChanged),
+// 	ECVF_Scalability
+// );
 
 // Override platform device profile
 // In editor all profiles will be available
@@ -82,27 +109,14 @@ static FAutoConsoleCommand GCmdSetNiagaraPlatformOverride(
 					{
 						if (Profile->GetName() == Args[0])
 						{
-							GNiagaraPlatformOverride = Profile;
+							SetGNiagaraDeviceProfile(Profile);
 							break;
 						}
 					}
 				}
 
-
 				if (GNiagaraPlatformOverride.IsValid())
 				{
-					//Save the previous QL state the first time we enter a preview.
-					if (GNiagaraBackupQualityLevel == INDEX_NONE)
-					{
-						GNiagaraBackupQualityLevel = GNiagaraQualityLevel;
-					}
-
-					UDeviceProfile* OverrideDP = GNiagaraPlatformOverride.Get();
-					check(OverrideDP);
-					int32 DPQL = FNiagaraPlatformSet::QualityLevelFromMask(FNiagaraPlatformSet::GetEffectQualityMaskForDeviceProfile(OverrideDP));
-					
-					OnSetCVarFromIniEntry(*GDeviceProfilesIni, NiagaraQualityLevelName, *LexToString(DPQL), ECVF_SetByMask);
-
 					UE_LOG(LogNiagara, Warning, TEXT("Niagara Setting Override DeviceProfile '%s'"), *Args[0]);
 				}
 				else
@@ -148,15 +162,7 @@ void FNiagaraPlatformSet::OnQualityLevelChanged(IConsoleVariable* Variable)
 
 	if (CurrentLevel != NewQualityLevel)
 	{
-		CachedQualityLevel = NewQualityLevel;
-		InvalidateCachedData();
-
-		for (TObjectIterator<UNiagaraSystem> It; It; ++It)
-		{
-			UNiagaraSystem* System = *It;
-			check(System);
-			System->OnScalabilityCVarChanged();
-		}
+		SetGNiagaraQualityLevel(NewQualityLevel);
 	}
 }
 
@@ -221,7 +227,23 @@ bool FNiagaraPlatformSet::IsActive()const
 {
 	if (LastBuiltFrame <= LastDirtiedFrame || LastBuiltFrame == 0)
 	{
-		bEnabledForCurrentProfileAndEffectQuality = IsEnabled(NiagaraGetActiveDeviceProfile(), GetQualityLevel(), true);
+		UDeviceProfile* ActiveProfile = NiagaraGetActiveDeviceProfile();
+		int32 QualityLevel = GetQualityLevel();
+#if WITH_EDITOR
+		if(OnOverrideQualityLevelDelegate.IsBound())
+		{
+			QualityLevel = OnOverrideQualityLevelDelegate.Execute();
+		}
+		if(OnOverrideActiveDeviceProfileDelegate.IsBound())
+		{
+			TOptional<TObjectPtr<UDeviceProfile>> OverrideProfile = OnOverrideActiveDeviceProfileDelegate.Execute();
+			if(OverrideProfile.IsSet())
+			{
+				ActiveProfile = OverrideProfile.GetValue();
+			}
+		}
+#endif
+		bEnabledForCurrentProfileAndEffectQuality = IsEnabled(ActiveProfile, QualityLevel, true);
 		LastBuiltFrame = GFrameNumber;
 	}
 	return bEnabledForCurrentProfileAndEffectQuality;
@@ -403,6 +425,7 @@ void FNiagaraPlatformSet::InvalidateCachedData()
 #if WITH_EDITOR
 	CachedQLMasksPerDeviceProfile.Empty();
 	CachedPlatformIniSettings.Empty();
+	CachedQualityLevel = INDEX_NONE;
 	FDeviceProfileValueCache::Empty();
 #endif
 
@@ -453,6 +476,36 @@ bool FNiagaraPlatformSet::CanChangeScalabilityAtRuntime(const UDeviceProfile* De
 	return true;//Assuming true if we fail to find the platform seems safest.
 }
 
+void SetGNiagaraQualityLevel(int32 QualityLevel)
+{
+	GNiagaraQualityLevel = QualityLevel;
+	FNiagaraPlatformSet::InvalidateCachedData();
+
+	for (TObjectIterator<UNiagaraSystem> It; It; ++It)
+	{
+		UNiagaraSystem* System = *It;
+		check(System);
+		System->UpdateScalability();
+	}
+}
+
+void SetGNiagaraDeviceProfile(UDeviceProfile* Profile)
+{
+	GNiagaraPlatformOverride = Profile;
+
+	//Save the previous QL state the first time we enter a preview.
+	if (GNiagaraBackupQualityLevel == INDEX_NONE)
+	{
+		GNiagaraBackupQualityLevel = GNiagaraQualityLevel;
+	}
+
+	UDeviceProfile* OverrideDP = GNiagaraPlatformOverride.Get();
+	check(OverrideDP);
+	int32 DPQL = FNiagaraPlatformSet::QualityLevelFromMask(FNiagaraPlatformSet::GetEffectQualityMaskForDeviceProfile(OverrideDP));
+					
+	OnSetCVarFromIniEntry(*GDeviceProfilesIni, NiagaraQualityLevelName, *LexToString(DPQL), ECVF_SetByMask);
+}
+
 #if WITH_EDITOR
 
 // bool FNiagaraPlatformSet::Conflicts(const FNiagaraPlatformSet& Other, TArray<const UDeviceProfile*>& OutConflictingProfiles, bool bIncludeProfilesWithVariableQualityLevel)const
@@ -786,7 +839,7 @@ void FNiagaraPlatformSetCVarCondition::OnCVarChanged(IConsoleVariable* CVar)
 	{
 		UNiagaraSystem* System = *It;
 		check(System);
-		System->OnScalabilityCVarChanged();
+		System->UpdateScalability();
 	}
 }
 
diff --git a/Engine/Plugins/FX/Niagara/Source/Niagara/Private/NiagaraSystem.cpp b/Engine/Plugins/FX/Niagara/Source/Niagara/Private/NiagaraSystem.cpp
--- a/Engine/Plugins/FX/Niagara/Source/Niagara/Private/NiagaraSystem.cpp
+++ b/Engine/Plugins/FX/Niagara/Source/Niagara/Private/NiagaraSystem.cpp
@@ -1076,7 +1076,8 @@ void UNiagaraSystem::PostEditChangeProperty(struct FPropertyChangedEvent& Proper
 	UpdateDITickFlags();
 	UpdateHasGPUEmitters();
 	ResolveScalabilitySettings();
-
+	OnScalabilityChanged().Broadcast();
+	
 	UpdateContext.CommitUpdate();
 
 	static FName SkipReset = TEXT("SkipSystemResetOnChange");
@@ -2334,6 +2335,11 @@ UNiagaraSystem::FOnSystemPostEditChange& UNiagaraSystem::OnSystemPostEditChange(
 	return OnSystemPostEditChangeDelegate;
 }
 
+UNiagaraSystem::FOnScalabilityChanged& UNiagaraSystem::OnScalabilityChanged()
+{
+	return OnScalabilityChangedDelegate;
+}
+
 void UNiagaraSystem::ForceGraphToRecompileOnNextCheck()
 {
 	check(SystemSpawnScript->GetLatestSource() == SystemUpdateScript->GetLatestSource());
@@ -3608,6 +3614,23 @@ UNiagaraEffectType* UNiagaraSystem::GetEffectType()const
 	return EffectType;
 }
 
+const FNiagaraSystemScalabilityOverride& UNiagaraSystem::GetCurrentOverrideSettings() const
+{
+	if(bOverrideScalabilitySettings)
+	{
+		for (const FNiagaraSystemScalabilityOverride& Override : SystemScalabilityOverrides.Overrides)
+		{
+			if (Override.Platforms.IsActive())
+			{
+				return Override;
+			}
+		}
+	}
+
+	static FNiagaraSystemScalabilityOverride Dummy;
+	return Dummy;
+}
+
 #if WITH_EDITOR
 void UNiagaraSystem::SetEffectType(UNiagaraEffectType* InEffectType)
 {
@@ -3755,7 +3778,7 @@ void UNiagaraSystem::ResolveScalabilitySettings()
 	FNiagaraWorldManager::InvalidateCachedSystemScalabilityDataForAllWorlds();
 }
 
-void UNiagaraSystem::OnScalabilityCVarChanged()
+void UNiagaraSystem::UpdateScalability()
 {
 	CacheFromCompiledData();
 
@@ -3765,10 +3788,14 @@ void UNiagaraSystem::OnScalabilityCVarChanged()
 	{
 		if (Handle.GetInstance())
 		{
-			Handle.GetInstance()->OnScalabilityCVarChanged();
+			Handle.GetInstance()->UpdateScalability();
 		}
 	}
 
+#if WITH_EDITOR
+	OnScalabilityChangedDelegate.Broadcast();
+#endif
+
 	// Update components
 	{
 		FNiagaraSystemUpdateContext UpdateCtx;
diff --git a/Engine/Plugins/FX/Niagara/Source/Niagara/Public/NiagaraRendererProperties.h b/Engine/Plugins/FX/Niagara/Source/Niagara/Public/NiagaraRendererProperties.h
--- a/Engine/Plugins/FX/Niagara/Source/Niagara/Public/NiagaraRendererProperties.h
+++ b/Engine/Plugins/FX/Niagara/Source/Niagara/Public/NiagaraRendererProperties.h
@@ -162,7 +162,11 @@ UCLASS(ABSTRACT)
 class NIAGARA_API UNiagaraRendererProperties : public UNiagaraMergeable
 {
 	GENERATED_BODY()
-
+	
+public:
+#if WITH_EDITOR
+	DECLARE_MULTICAST_DELEGATE(FOnPropertiesChanged);
+#endif
 public:
 	UNiagaraRendererProperties()		
 		: bIsEnabled(true)
@@ -177,6 +181,7 @@ public:
 #if WITH_EDITORONLY_DATA
 	virtual void PostEditChangeProperty(struct FPropertyChangedEvent& PropertyChangedEvent) override;
 #endif
+
 	//UObject Interface End
 	
 	virtual FNiagaraRenderer* CreateEmitterRenderer(ERHIFeatureLevel::Type FeatureLevel, const FNiagaraEmitterInstance* Emitter, const FNiagaraSystemInstanceController& InController) PURE_VIRTUAL ( UNiagaraRendererProperties::CreateEmitterRenderer, return nullptr;);
@@ -230,9 +235,11 @@ public:
 
 	// The text to display in the niagara stack widget under the renderer section
 	virtual FText GetWidgetDisplayName() const;
-
 #endif // WITH_EDITORONLY_DATA
 
+#if WITH_EDITOR
+	FOnPropertiesChanged& OnPropertiesChanged();
+#endif
 	virtual ENiagaraRendererSourceDataMode GetCurrentSourceMode() const {	return ENiagaraRendererSourceDataMode::Particles;}
 
 	virtual bool GetIsActive() const;
@@ -251,7 +258,7 @@ public:
 	bool NeedsPreciseMotionVectors() const;
 
 	/** Platforms on which this renderer is enabled. */
-	UPROPERTY(EditAnywhere, Category = "Scalability")
+	UPROPERTY(EditAnywhere, Category = "Scalability", meta=(DisplayInScalabilityContext))
 	FNiagaraPlatformSet Platforms;
 
 	/** By default, emitters are drawn in the order that they are added to the system. This value will allow you to control the order in a more fine-grained manner.
@@ -286,4 +293,8 @@ protected:
 	/** returns the variable associated with the supplied binding if it should be bound given the current settings of the RendererProperties. */
 	virtual FNiagaraVariable GetBoundAttribute(const FNiagaraVariableAttributeBinding* Binding) const;
 #endif
+
+#if WITH_EDITOR
+	FOnPropertiesChanged OnPropertiesChangedDelegate;
+#endif
 };
diff --git a/Engine/Plugins/FX/Niagara/Source/NiagaraEditor/Private/NiagaraEditorModule.cpp b/Engine/Plugins/FX/Niagara/Source/NiagaraEditor/Private/NiagaraEditorModule.cpp
--- a/Engine/Plugins/FX/Niagara/Source/NiagaraEditor/Private/NiagaraEditorModule.cpp
+++ b/Engine/Plugins/FX/Niagara/Source/NiagaraEditor/Private/NiagaraEditorModule.cpp
@@ -1190,7 +1190,7 @@ void FNiagaraEditorModule::ShutdownModule()
 		PropertyModule.UnregisterCustomClassLayout("NiagaraComponent");
 	}
 
-	FNiagaraEditorStyle::Shutdown();
+	FNiagaraEditorStyle::Unregister();
 
 	UnregisterSettings();
 
@@ -1339,7 +1339,7 @@ void FNiagaraEditorModule::OnPreviewPlatformChanged()
 	{
 		UNiagaraSystem* System = *It;
 		check(System);
-		System->OnScalabilityCVarChanged();
+		System->UpdateScalability();
 	}
 }
 
diff --git a/Engine/Plugins/FX/Niagara/Source/NiagaraEditor/Private/NiagaraEditorUtilities.cpp b/Engine/Plugins/FX/Niagara/Source/NiagaraEditor/Private/NiagaraEditorUtilities.cpp
--- a/Engine/Plugins/FX/Niagara/Source/NiagaraEditor/Private/NiagaraEditorUtilities.cpp
+++ b/Engine/Plugins/FX/Niagara/Source/NiagaraEditor/Private/NiagaraEditorUtilities.cpp
@@ -3602,6 +3602,11 @@ TSharedRef<SToolTip> FNiagaraParameterUtilities::GetTooltipWidget(FNiagaraVariab
 	return TooltipWidget;
 }
 
+float FNiagaraEditorUtilities::GetScalabilityTintAlpha(FNiagaraEmitterHandle* EmitterHandle)
+{
+	return EmitterHandle->GetIsEnabled() ? 1 : 0.5f;
+}
+
 void FNiagaraParameterUtilities::FilterToRelevantStaticVariables(const TArray<FNiagaraVariable>& InVars, TArray<FNiagaraVariable>& OutVars, FName InOldEmitterAlias, FName InNewEmitterAlias, bool bFilterByEmitterAliasAndConvertToUnaliased)
 {
 	FNiagaraAliasContext RenameContext(ENiagaraScriptUsage::ParticleSpawnScript);
diff --git a/Engine/Plugins/FX/Niagara/Source/NiagaraEditor/Private/NiagaraEmitterDetailsCustomization.cpp b/Engine/Plugins/FX/Niagara/Source/NiagaraEditor/Private/NiagaraEmitterDetailsCustomization.cpp
--- a/Engine/Plugins/FX/Niagara/Source/NiagaraEditor/Private/NiagaraEmitterDetailsCustomization.cpp
+++ b/Engine/Plugins/FX/Niagara/Source/NiagaraEditor/Private/NiagaraEmitterDetailsCustomization.cpp
@@ -6,6 +6,7 @@
 #include "DetailCategoryBuilder.h"
 #include "NiagaraEmitter.h"
 #include "NiagaraEditorModule.h"
+#include "Toolkits/NiagaraSystemToolkit.h"
 
 TSharedRef<IDetailCustomization> FNiagaraEmitterDetails::MakeInstance()
 {
@@ -16,5 +17,9 @@ void FNiagaraEmitterDetails::CustomizeDetails(IDetailLayoutBuilder& InDetailLayo
 {
 	TSharedPtr<IPropertyHandle> EventHandlersPropertyHandle = InDetailLayout.GetProperty(UNiagaraEmitter::PrivateMemberNames::EventHandlerScriptProps);
 	EventHandlersPropertyHandle->MarkHiddenByCustomization();
+	if(GNiagaraScalabilityModeEnabled)
+	{
+		InDetailLayout.HideCategory("Scalability");
+	}
 }
 
diff --git a/Engine/Plugins/FX/Niagara/Source/NiagaraEditor/Private/NiagaraOverviewNode.cpp b/Engine/Plugins/FX/Niagara/Source/NiagaraEditor/Private/NiagaraOverviewNode.cpp
--- a/Engine/Plugins/FX/Niagara/Source/NiagaraEditor/Private/NiagaraOverviewNode.cpp
+++ b/Engine/Plugins/FX/Niagara/Source/NiagaraEditor/Private/NiagaraOverviewNode.cpp
@@ -65,6 +65,11 @@ static FNiagaraEmitterHandle* FindEmitterHandleByID(UNiagaraSystem* System, cons
 	return nullptr;
 }
 
+FNiagaraEmitterHandle* UNiagaraOverviewNode::TryGetEmitterHandle()
+{
+	return FindEmitterHandleByID(GetOwningSystem(), GetEmitterHandleGuid());
+}
+
 FText UNiagaraOverviewNode::GetNodeTitle(ENodeTitleType::Type TitleType) const
 {
 	if (OwningSystem == nullptr)
diff --git a/Engine/Plugins/FX/Niagara/Source/NiagaraEditor/Private/NiagaraSystemDetailsCustomization.cpp b/Engine/Plugins/FX/Niagara/Source/NiagaraEditor/Private/NiagaraSystemDetailsCustomization.cpp
new file mode 100644
--- /dev/null
+++ b/Engine/Plugins/FX/Niagara/Source/NiagaraEditor/Private/NiagaraSystemDetailsCustomization.cpp
@@ -0,0 +1,22 @@
+// Copyright Epic Games, Inc. All Rights Reserved.
+
+#include "NiagaraSystemDetailsCustomization.h"
+#include "PropertyHandle.h"
+#include "DetailLayoutBuilder.h"
+#include "DetailCategoryBuilder.h"
+#include "NiagaraSystem.h"
+#include "Toolkits/NiagaraSystemToolkit.h"
+
+TSharedRef<IDetailCustomization> FNiagaraSystemDetails::MakeInstance()
+{
+	return MakeShared<FNiagaraSystemDetails>();
+}
+
+void FNiagaraSystemDetails::CustomizeDetails(IDetailLayoutBuilder& InDetailLayout)
+{
+	if(GNiagaraScalabilityModeEnabled)
+	{
+		InDetailLayout.HideCategory("Scalability");
+		InDetailLayout.HideProperty("EffectType");
+	}
+}
diff --git a/Engine/Plugins/FX/Niagara/Source/NiagaraEditor/Private/NiagaraSystemDetailsCustomization.h b/Engine/Plugins/FX/Niagara/Source/NiagaraEditor/Private/NiagaraSystemDetailsCustomization.h
new file mode 100644
--- /dev/null
+++ b/Engine/Plugins/FX/Niagara/Source/NiagaraEditor/Private/NiagaraSystemDetailsCustomization.h
@@ -0,0 +1,22 @@
+// Copyright Epic Games, Inc. All Rights Reserved.
+
+#pragma once
+
+#include "IDetailCustomization.h"
+
+class IDetailLayoutBuilder;
+
+/*-----------------------------------------------------------------------------
+FNiagaraSystemDetails
+-----------------------------------------------------------------------------*/
+
+class FNiagaraSystemDetails : public IDetailCustomization
+{
+public:
+	/** Makes a new instance of this detail layout class for a specific detail view requesting it */
+	static TSharedRef<IDetailCustomization> MakeInstance();
+
+	/** IDetailCustomization interface */
+	virtual void CustomizeDetails(IDetailLayoutBuilder& InDetailLayout) override;
+};
+
diff --git a/Engine/Plugins/FX/Niagara/Source/NiagaraEditor/Private/Sequencer/NiagaraSequence/MovieSceneNiagaraEmitterTrack.cpp b/Engine/Plugins/FX/Niagara/Source/NiagaraEditor/Private/Sequencer/NiagaraSequence/MovieSceneNiagaraEmitterTrack.cpp
--- a/Engine/Plugins/FX/Niagara/Source/NiagaraEditor/Private/Sequencer/NiagaraSequence/MovieSceneNiagaraEmitterTrack.cpp
+++ b/Engine/Plugins/FX/Niagara/Source/NiagaraEditor/Private/Sequencer/NiagaraSequence/MovieSceneNiagaraEmitterTrack.cpp
@@ -16,6 +16,7 @@
 #include "Sections/MovieSceneNiagaraEmitterSection.h"
 
 #include "ISequencerSection.h"
+#include "NiagaraEditorUtilities.h"
 #include "SequencerSectionPainter.h"
 
 #define LOCTEXT_NAMESPACE "MovieSceneNiagaraEmitterTrack"
@@ -85,6 +86,14 @@ void UMovieSceneNiagaraEmitterTrack::Initialize(FNiagaraSystemViewModel& InSyste
 	SystemPath = SystemViewModel->GetSystem().GetPathName();
 	SetColorTint(FNiagaraEditorStyle::Get().GetColor("NiagaraEditor.NiagaraSequence.DefaultTrackColor").ToFColor(true));
 	CreateSections(InFrameResolution);
+
+	// since we are updating the color on tick when scalability mode is active, we reset the color when it is deactivated
+	SystemViewModel->GetScalabilityViewModel()->OnScalabilityModeChanged().AddUObject(this, &UMovieSceneNiagaraEmitterTrack::RestoreDefaultTrackColor);
+}
+
+UMovieSceneNiagaraEmitterTrack::~UMovieSceneNiagaraEmitterTrack()
+{
+	//SystemViewModel->GetScalabilityViewModel()->OnScalabilityModeChanged().RemoveAll(this);
 }
 
 bool UMovieSceneNiagaraEmitterTrack::CanRename() const
@@ -166,6 +175,33 @@ bool UMovieSceneNiagaraEmitterTrack::GetSectionsWereModified() const
 	return bSectionsWereModified;
 }
 
+void UMovieSceneNiagaraEmitterTrack::Tick(float DeltaTime)
+{
+	if(!HasAnyFlags(RF_ClassDefaultObject) && bScalabilityModeActive)
+	{
+		UNiagaraEmitter* Emitter = nullptr;
+		// the emitter instance might no longer be valid if we deleted the emitter but the track still lives until garbage collection
+		if(EmitterHandleViewModel.IsValid() && EmitterHandleViewModel.Pin()->GetEmitterHandle() != nullptr)
+		{
+			Emitter = EmitterHandleViewModel.Pin()->GetEmitterHandle()->GetInstance();
+		}
+		
+		if(Emitter && !Emitter->IsAllowedByScalability())
+		{
+			SetColorTint(FNiagaraEditorStyle::Get().GetColor("NiagaraEditor.SystemOverview.ExcludedFromScalability").ToFColor(true));
+		}
+		else
+		{
+			SetColorTint(FNiagaraEditorStyle::Get().GetColor("NiagaraEditor.NiagaraSequence.DefaultTrackColor").ToFColor(true));
+		}
+	}
+}
+
+TStatId UMovieSceneNiagaraEmitterTrack::GetStatId() const
+{
+	return TStatId();
+}
+
 bool UMovieSceneNiagaraEmitterTrack::HasSection(const UMovieSceneSection& Section) const
 {
 	return Sections.Contains(&Section);
@@ -297,5 +333,15 @@ void UMovieSceneNiagaraEmitterTrack::CreateSections(const FFrameRate& InFrameRes
 	bSectionsWereModified = false;
 }
 
+void UMovieSceneNiagaraEmitterTrack::RestoreDefaultTrackColor(bool bScalabilityModeActivated)
+{
+	bScalabilityModeActive = bScalabilityModeActivated;
+	
+	if(bScalabilityModeActive == false)
+	{
+		SetColorTint(FNiagaraEditorStyle::Get().GetColor("NiagaraEditor.NiagaraSequence.DefaultTrackColor").ToFColor(true));
+	}
+}
+
 #undef LOCTEXT_NAMESPACE
 
diff --git a/Engine/Plugins/FX/Niagara/Source/NiagaraEditor/Private/Sequencer/NiagaraSequence/MovieSceneNiagaraEmitterTrack.h b/Engine/Plugins/FX/Niagara/Source/NiagaraEditor/Private/Sequencer/NiagaraSequence/MovieSceneNiagaraEmitterTrack.h
--- a/Engine/Plugins/FX/Niagara/Source/NiagaraEditor/Private/Sequencer/NiagaraSequence/MovieSceneNiagaraEmitterTrack.h
+++ b/Engine/Plugins/FX/Niagara/Source/NiagaraEditor/Private/Sequencer/NiagaraSequence/MovieSceneNiagaraEmitterTrack.h
@@ -4,6 +4,7 @@
 
 #include "MovieSceneNameableTrack.h"
 #include "MovieSceneSection.h"
+#include "TickableEditorObject.h"
 #include "MovieSceneNiagaraEmitterTrack.generated.h"
 
 class UNiagaraSystem;
@@ -51,13 +52,14 @@ private:
 */
 UCLASS(MinimalAPI)
 class UMovieSceneNiagaraEmitterTrack
-	: public UMovieSceneNameableTrack
+	: public UMovieSceneNameableTrack, public FTickableEditorObject
 {
 	GENERATED_UCLASS_BODY()
 
 public:
 	void Initialize(FNiagaraSystemViewModel& SystemViewModel, TSharedRef<FNiagaraEmitterHandleViewModel> InEmitterHandleViewModel, const FFrameRate& InFrameResolution);
-
+	virtual ~UMovieSceneNiagaraEmitterTrack();
+	
 	virtual bool CanRename() const override;
 
 	FNiagaraSystemViewModel& GetSystemViewModel() const;
@@ -72,6 +74,10 @@ public:
 
 	bool GetSectionsWereModified() const;
 
+	// FTickableEditorObject interface
+	virtual void Tick(float DeltaTime) override;
+	virtual TStatId GetStatId() const override;
+
 	//~ UMovieSceneTrack interface
 	virtual void RemoveAllAnimationData() override { }
 	virtual bool HasSection(const UMovieSceneSection& Section) const override;
@@ -84,6 +90,7 @@ public:
 	virtual const TArray<UMovieSceneSection*>& GetAllSections() const override;
 	virtual bool SupportsMultipleRows() const override;
 	virtual bool ValidateDisplayName(const FText& NewDisplayName, FText& OutErrorMessage) const override;
+	
 
 	/** Gets the unique id for the emitter handle that was associated with this track; used for copy/paste detection */
 	FGuid GetEmitterHandleId() const;
@@ -95,6 +102,7 @@ public:
 
 private:
 	void CreateSections(const FFrameRate& InFrameResolution);
+	void RestoreDefaultTrackColor(bool bScalabilityModeActivated);
 
 private:
 	FNiagaraSystemViewModel* SystemViewModel;
@@ -116,4 +124,6 @@ private:
 	FString SystemPath;
 
 	TArray<FText> SectionInitializationErrors;
+
+	bool bScalabilityModeActive = false;
 };
\ No newline at end of file
diff --git a/Engine/Plugins/FX/Niagara/Source/NiagaraEditor/Private/Toolkits/NiagaraSystemToolkit.cpp b/Engine/Plugins/FX/Niagara/Source/NiagaraEditor/Private/Toolkits/NiagaraSystemToolkit.cpp
--- a/Engine/Plugins/FX/Niagara/Source/NiagaraEditor/Private/Toolkits/NiagaraSystemToolkit.cpp
+++ b/Engine/Plugins/FX/Niagara/Source/NiagaraEditor/Private/Toolkits/NiagaraSystemToolkit.cpp
@@ -61,6 +61,9 @@
 #include "Widgets/Input/SCheckBox.h"
 #include "ViewModels/NiagaraParameterPanelViewModel.h"
 #include "NiagaraEditor/Private/SNiagaraAssetPickerList.h"
+#include "SystemToolkitModes/NiagaraSystemToolkitModeBase.h"
+#include "SystemToolkitModes/NiagaraSystemToolkitMode_Default.h"
+#include "SystemToolkitModes/NiagaraSystemToolkitMode_Scalability.h"
 #include "Widgets/SNiagaraParameterDefinitionsPanel.h"
 #include "ViewModels/NiagaraParameterDefinitionsPanelViewModel.h"
 #include "ViewModels/Stack/NiagaraStackViewModel.h"
@@ -70,24 +73,20 @@
 
 DECLARE_CYCLE_STAT(TEXT("Niagara - SystemToolkit - OnApply"), STAT_NiagaraEditor_SystemToolkit_OnApply, STATGROUP_NiagaraEditor);
 
-const FName FNiagaraSystemToolkit::ViewportTabID(TEXT("NiagaraSystemEditor_Viewport"));
-const FName FNiagaraSystemToolkit::CurveEditorTabID(TEXT("NiagaraSystemEditor_CurveEditor"));
-const FName FNiagaraSystemToolkit::SequencerTabID(TEXT("NiagaraSystemEditor_Sequencer"));
-const FName FNiagaraSystemToolkit::SystemScriptTabID(TEXT("NiagaraSystemEditor_SystemScript"));
-const FName FNiagaraSystemToolkit::SystemDetailsTabID(TEXT("NiagaraSystemEditor_SystemDetails"));
-const FName FNiagaraSystemToolkit::SystemParametersTabID(TEXT("NiagaraSystemEditor_SystemParameters"));
-const FName FNiagaraSystemToolkit::SystemParametersTabID2(TEXT("NiagaraSystemEditor_SystemParameters2"));
-const FName FNiagaraSystemToolkit::SystemParameterDefinitionsTabID(TEXT("NiagaraSystemEditor_SystemParameterDefinitions"));
-const FName FNiagaraSystemToolkit::SelectedEmitterStackTabID(TEXT("NiagaraSystemEditor_SelectedEmitterStack"));
-const FName FNiagaraSystemToolkit::SelectedEmitterGraphTabID(TEXT("NiagaraSystemEditor_SelectedEmitterGraph"));
-const FName FNiagaraSystemToolkit::DebugSpreadsheetTabID(TEXT("NiagaraSystemEditor_DebugAttributeSpreadsheet"));
-const FName FNiagaraSystemToolkit::PreviewSettingsTabId(TEXT("NiagaraSystemEditor_PreviewSettings"));
-const FName FNiagaraSystemToolkit::GeneratedCodeTabID(TEXT("NiagaraSystemEditor_GeneratedCode"));
-const FName FNiagaraSystemToolkit::MessageLogTabID(TEXT("NiagaraSystemEditor_MessageLog"));
-const FName FNiagaraSystemToolkit::SystemOverviewTabID(TEXT("NiagaraSystemEditor_SystemOverview"));
-const FName FNiagaraSystemToolkit::ScratchPadTabID(TEXT("NiagaraSystemEditor_ScratchPad"));
-const FName FNiagaraSystemToolkit::ScriptStatsTabID(TEXT("NiagaraSystemEditor_ScriptStats"));
-const FName FNiagaraSystemToolkit::BakerTabID(TEXT("NiagaraSystemEditor_Baker"));
+const FName FNiagaraSystemToolkit::DefaultModeName(TEXT("Default"));
+const FName FNiagaraSystemToolkit::ScalabilityModeName(TEXT("Scalability"));
+
+int32 GNiagaraScalabilityModeEnabled = 0;
+
+// @TODO Remove cvar when scalability mode is going in properly
+static const TCHAR* NiagaraEnableScalabilityMode = TEXT("fx.Niagara.EnableScalabilityMode");
+static FAutoConsoleVariableRef CVarNiagaraScalabilityModeEnabled(
+	NiagaraEnableScalabilityMode,
+	GNiagaraScalabilityModeEnabled,
+	TEXT("Enable the scalability mode inside Niagara emitters and systems. \n"),
+	ECVF_Default
+);
+
 IConsoleVariable* FNiagaraSystemToolkit::VmStatEnabledVar = IConsoleManager::Get().FindConsoleVariable(TEXT("vm.DetailedVMScriptStats"));
 
 static int32 GbLogNiagaraSystemChanges = 0;
@@ -98,124 +97,6 @@ static FAutoConsoleVariableRef CVarSuppressNiagaraSystems(
 	ECVF_Default
 );
 
-void FNiagaraSystemToolkit::RegisterTabSpawners(const TSharedRef<class FTabManager>& InTabManager)
-{
-	WorkspaceMenuCategory = InTabManager->AddLocalWorkspaceMenuCategory(LOCTEXT("WorkspaceMenu_NiagaraSystemEditor", "Niagara System"));
-
-	FAssetEditorToolkit::RegisterTabSpawners(InTabManager);
-
-	InTabManager->RegisterTabSpawner(ViewportTabID, FOnSpawnTab::CreateSP(this, &FNiagaraSystemToolkit::SpawnTab_Viewport))
-		.SetDisplayName(LOCTEXT("Preview", "Preview"))
-		.SetGroup(WorkspaceMenuCategory.ToSharedRef())
-		.SetIcon(FSlateIcon(FNiagaraEditorStyle::Get().GetStyleSetName(), "Tab.Viewport"));
-
-	InTabManager->RegisterTabSpawner(CurveEditorTabID, FOnSpawnTab::CreateSP(this, &FNiagaraSystemToolkit::SpawnTab_CurveEd))
-		.SetDisplayName(LOCTEXT("Curves", "Curves"))
-		.SetGroup(WorkspaceMenuCategory.ToSharedRef())
-		.SetIcon(FSlateIcon(FNiagaraEditorStyle::Get().GetStyleSetName(), "Tab.Curves"));
-
-	InTabManager->RegisterTabSpawner(SequencerTabID, FOnSpawnTab::CreateSP(this, &FNiagaraSystemToolkit::SpawnTab_Sequencer))
-		.SetDisplayName(LOCTEXT("Timeline", "Timeline"))
-		.SetGroup(WorkspaceMenuCategory.ToSharedRef())
-		.SetIcon(FSlateIcon(FNiagaraEditorStyle::Get().GetStyleSetName(), "Tab.Timeline"));
-
-	InTabManager->RegisterTabSpawner(SystemScriptTabID, FOnSpawnTab::CreateSP(this, &FNiagaraSystemToolkit::SpawnTab_SystemScript))
-		.SetDisplayName(LOCTEXT("SystemScript", "System Script"))
-		.SetGroup(WorkspaceMenuCategory.ToSharedRef())
-		.SetAutoGenerateMenuEntry(GbShowNiagaraDeveloperWindows != 0);
-
-	InTabManager->RegisterTabSpawner(SystemParametersTabID, FOnSpawnTab::CreateSP(this, &FNiagaraSystemToolkit::SpawnTab_SystemParameters))
-		.SetDisplayName(LOCTEXT("SystemParameters", "Parameters"))
-		.SetGroup(WorkspaceMenuCategory.ToSharedRef())
-		.SetIcon(FSlateIcon(FNiagaraEditorStyle::Get().GetStyleSetName(), "Tab.Parameters"));
-
-	InTabManager->RegisterTabSpawner(SystemParametersTabID2, FOnSpawnTab::CreateSP(this, &FNiagaraSystemToolkit::SpawnTab_SystemParameters2))
-		.SetDisplayName(LOCTEXT("SystemParameters2", "Legacy Parameters"))
-		.SetGroup(WorkspaceMenuCategory.ToSharedRef())
-		.SetIcon(FSlateIcon(FNiagaraEditorStyle::Get().GetStyleSetName(), "Tab.Parameters"));
-
-//@todo(ng) disable parameter definitions panel pending bug fixes
-// 	InTabManager->RegisterTabSpawner(SystemParameterDefinitionsTabID, FOnSpawnTab::CreateSP(this, &FNiagaraSystemToolkit::SpawnTab_SystemParameterDefinitions)) 
-// 		.SetDisplayName(LOCTEXT("SystemParameterDefinitions", "Parameter Definitions"))
-// 		.SetGroup(WorkspaceMenuCategory.ToSharedRef());
-
-	InTabManager->RegisterTabSpawner(SelectedEmitterStackTabID, FOnSpawnTab::CreateSP(this, &FNiagaraSystemToolkit::SpawnTab_SelectedEmitterStack))
-		.SetDisplayName(LOCTEXT("SelectedEmitterStacks", "Selected Emitters"))
-		.SetGroup(WorkspaceMenuCategory.ToSharedRef())
-		.SetIcon(FSlateIcon(FNiagaraEditorStyle::Get().GetStyleSetName(), "Tab.VisualEffects"));
-
-	InTabManager->RegisterTabSpawner(SelectedEmitterGraphTabID, FOnSpawnTab::CreateSP(this, &FNiagaraSystemToolkit::SpawnTab_SelectedEmitterGraph))
-		.SetDisplayName(LOCTEXT("SelectedEmitterGraph", "Selected Emitter Graph"))
-		.SetGroup(WorkspaceMenuCategory.ToSharedRef())
-		.SetAutoGenerateMenuEntry(GbShowNiagaraDeveloperWindows != 0);
-
-	InTabManager->RegisterTabSpawner(DebugSpreadsheetTabID, FOnSpawnTab::CreateSP(this, &FNiagaraSystemToolkit::SpawnTab_DebugSpreadsheet))
-		.SetDisplayName(LOCTEXT("DebugSpreadsheet", "Attribute Spreadsheet"))
-		.SetGroup(WorkspaceMenuCategory.ToSharedRef())
-		.SetIcon(FSlateIcon(FNiagaraEditorStyle::Get().GetStyleSetName(), "Tab.Spreadsheet"));
-
-	InTabManager->RegisterTabSpawner(PreviewSettingsTabId, FOnSpawnTab::CreateSP(this, &FNiagaraSystemToolkit::SpawnTab_PreviewSettings))
-		.SetDisplayName(LOCTEXT("PreviewSceneSettingsTab", "Preview Scene Settings"))
-		.SetGroup(WorkspaceMenuCategory.ToSharedRef())
-		.SetIcon(FSlateIcon(FNiagaraEditorStyle::Get().GetStyleSetName(), "Tab.Settings"));
-
-	InTabManager->RegisterTabSpawner(GeneratedCodeTabID, FOnSpawnTab::CreateSP(this, &FNiagaraSystemToolkit::SpawnTab_GeneratedCode))
-		.SetDisplayName(LOCTEXT("GeneratedCode", "Generated Code"))
-		.SetGroup(WorkspaceMenuCategory.ToSharedRef())
-		.SetIcon(FSlateIcon(FNiagaraEditorStyle::Get().GetStyleSetName(), "Tab.GeneratedCode"));
-
-	InTabManager->RegisterTabSpawner(MessageLogTabID, FOnSpawnTab::CreateSP(this, &FNiagaraSystemToolkit::SpawnTab_MessageLog))
-		.SetDisplayName(LOCTEXT("NiagaraMessageLog", "Niagara Log"))
-		.SetGroup(WorkspaceMenuCategory.ToSharedRef())
-		.SetIcon(FSlateIcon(FNiagaraEditorStyle::Get().GetStyleSetName(), "Tab.Log"));
-
-	InTabManager->RegisterTabSpawner(SystemOverviewTabID, FOnSpawnTab::CreateSP(this, &FNiagaraSystemToolkit::SpawnTab_SystemOverview))
-		.SetDisplayName(LOCTEXT("SystemOverviewTabName", "System Overview"))
-		.SetGroup(WorkspaceMenuCategory.ToSharedRef())
-		.SetIcon(FSlateIcon(FNiagaraEditorStyle::Get().GetStyleSetName(), "Tab.SystemOverview"));
-
-	InTabManager->RegisterTabSpawner(ScratchPadTabID, FOnSpawnTab::CreateSP(this, &FNiagaraSystemToolkit::SpawnTab_ScratchPad))
-		.SetDisplayName(LOCTEXT("ScratchPadTabName", "Scratch Pad"))
-		.SetGroup(WorkspaceMenuCategory.ToSharedRef())
-		.SetIcon(FSlateIcon(FNiagaraEditorStyle::Get().GetStyleSetName(), "Tab.ScratchPad"));
-
-	InTabManager->RegisterTabSpawner(ScriptStatsTabID, FOnSpawnTab::CreateSP(this, &FNiagaraSystemToolkit::SpawnTab_ScriptStats))
-		.SetDisplayName(LOCTEXT("NiagaraScriptsStatsTab", "Script Stats"))
-		.SetGroup(WorkspaceMenuCategory.ToSharedRef())
-		.SetIcon(FSlateIcon(FNiagaraEditorStyle::Get().GetStyleSetName(), "Tab.ScriptStats"));
-
-	if (GetDefault<UNiagaraEditorSettings>()->bEnableBaker)
-	{
-		InTabManager->RegisterTabSpawner(BakerTabID, FOnSpawnTab::CreateSP(this, &FNiagaraSystemToolkit::SpawnTab_Baker))
-			.SetDisplayName(LOCTEXT("NiagaraBakerTab", "Baker"))
-			.SetGroup(WorkspaceMenuCategory.ToSharedRef());
-	}
-}
-
-void FNiagaraSystemToolkit::UnregisterTabSpawners(const TSharedRef<class FTabManager>& InTabManager)
-{
-	FAssetEditorToolkit::UnregisterTabSpawners(InTabManager);
-
-	InTabManager->UnregisterTabSpawner(ViewportTabID);
-	InTabManager->UnregisterTabSpawner(CurveEditorTabID);
-	InTabManager->UnregisterTabSpawner(SequencerTabID);
-	InTabManager->UnregisterTabSpawner(SystemScriptTabID);
-	InTabManager->UnregisterTabSpawner(SystemDetailsTabID);
-	InTabManager->UnregisterTabSpawner(SystemParametersTabID);
-	InTabManager->UnregisterTabSpawner(SystemParametersTabID2);
-	//@todo(ng) disable parameter definitions panel pending bug fixes
-	//InTabManager->UnregisterTabSpawner(SystemParameterDefinitionsTabID);
-	InTabManager->UnregisterTabSpawner(SelectedEmitterStackTabID);
-	InTabManager->UnregisterTabSpawner(SelectedEmitterGraphTabID);
-	InTabManager->UnregisterTabSpawner(DebugSpreadsheetTabID);
-	InTabManager->UnregisterTabSpawner(PreviewSettingsTabId);
-	InTabManager->UnregisterTabSpawner(GeneratedCodeTabID);
-	InTabManager->UnregisterTabSpawner(SystemOverviewTabID);
-	InTabManager->UnregisterTabSpawner(ScratchPadTabID);
-	InTabManager->UnregisterTabSpawner(ScriptStatsTabID);
-	InTabManager->UnregisterTabSpawner(BakerTabID);
-}
-
 FNiagaraSystemToolkit::~FNiagaraSystemToolkit()
 {
 	// Cleanup viewmodels that use the system viewmodel before cleaning up the system viewmodel itself.
@@ -245,6 +126,26 @@ void FNiagaraSystemToolkit::AddReferencedObjects(FReferenceCollector& Collector)
 	Collector.AddReferencedObject(System);
 }
 
+TSharedPtr<SWidget> FNiagaraSystemToolkit::GetSystemOverview() const
+{
+	return SystemOverview;
+}
+
+void FNiagaraSystemToolkit::SetSystemOverview(const TSharedPtr<SWidget>& InSystemOverview)
+{
+	this->SystemOverview = InSystemOverview;
+}
+
+TSharedPtr<SWidget> FNiagaraSystemToolkit::GetScriptScratchpad() const
+{
+	return ScriptScratchpad;
+}
+
+void FNiagaraSystemToolkit::SetScriptScratchpad(const TSharedPtr<SWidget>& InScriptScratchpad)
+{
+	this->ScriptScratchpad = InScriptScratchpad;
+}
+
 void FNiagaraSystemToolkit::InitializeWithSystem(const EToolkitMode::Type Mode, const TSharedPtr< class IToolkitHost >& InitToolkitHost, UNiagaraSystem& InSystem)
 {
 	System = &InSystem;
@@ -341,6 +242,9 @@ void FNiagaraSystemToolkit::InitializeWithEmitter(const EToolkitMode::Type Mode,
 	SystemViewModel->Initialize(*System, SystemOptions);
 	SystemViewModel->GetEditorData().SetOwningSystemIsPlaceholder(true, *System);
 	SystemViewModel->SetToolkitCommands(GetToolkitCommands());
+
+	SystemViewModel->OnGetWorkflowMode().BindSP(this, &FNiagaraSystemToolkit::GetCurrentMode);
+
 	SystemViewModel->AddEmitter(*Emitter);
 
 	ParameterPanelViewModel = MakeShared<FNiagaraSystemToolkitParameterPanelViewModel>(SystemViewModel);
@@ -387,94 +291,36 @@ void FNiagaraSystemToolkit::InitializeInternal(const EToolkitMode::Type Mode, co
 	SystemViewModel->GetOnPinnedEmittersChanged().AddSP(this, &FNiagaraSystemToolkit::RefreshParameters);
 	SystemViewModel->OnRequestFocusTab().AddSP(this, &FNiagaraSystemToolkit::OnViewModelRequestFocusTab);
 	
+	SystemViewModel->OnGetWorkflowMode().BindSP(this, &FNiagaraSystemToolkit::GetCurrentMode);
+	
 	const float InTime = -0.02f;
 	const float OutTime = 3.2f;
-
-	TSharedRef<FTabManager::FLayout> StandaloneDefaultLayout = FTabManager::NewLayout("Standalone_Niagara_System_Layout_v25")
-		->AddArea
-		(
-			FTabManager::NewPrimaryArea()->SetOrientation(Orient_Vertical)
-			->Split
-			(
-				FTabManager::NewSplitter()->SetOrientation(Orient_Horizontal)
-				->Split
-				(
-					// Top Level Left
-					FTabManager::NewSplitter()->SetOrientation(Orient_Vertical)
-					->SetSizeCoefficient(.75f)
-					->Split
-					(
-						// Inner Left Top
-						FTabManager::NewSplitter()->SetOrientation(Orient_Horizontal)
-						->SetSizeCoefficient(0.75f)
-						->Split
-						(
-							FTabManager::NewStack()
-							->SetSizeCoefficient(.25f)
-							->AddTab(ViewportTabID, ETabState::OpenedTab)
-						)
-						->Split
-						(
-							FTabManager::NewStack()
-							->SetSizeCoefficient(0.15f)
-							->AddTab(SystemParametersTabID, ETabState::OpenedTab)
-						)
-						->Split
-						(
-							FTabManager::NewStack()
-							->SetSizeCoefficient(0.6f)
-							->AddTab(SystemOverviewTabID, ETabState::OpenedTab)
-							->AddTab(ScratchPadTabID, ETabState::OpenedTab)
-							->AddTab(BakerTabID, ETabState::ClosedTab)
-							->SetForegroundTab(SystemOverviewTabID)
-						)
-					)
-					->Split
-					(
-						// Inner Left Bottom
-						FTabManager::NewStack()
-						->SetSizeCoefficient(0.25f)
-						->AddTab(CurveEditorTabID, ETabState::OpenedTab)
-						->AddTab(MessageLogTabID, ETabState::OpenedTab)
-						->AddTab(SequencerTabID, ETabState::OpenedTab)
-						->AddTab(ScriptStatsTabID, ETabState::ClosedTab)
-					)
-				)
-				->Split
-				(
-					// Top Level Right
-					FTabManager::NewStack()
-					->SetSizeCoefficient(0.25f)
-					->AddTab(SelectedEmitterStackTabID, ETabState::OpenedTab)
-					->AddTab(SelectedEmitterGraphTabID, ETabState::ClosedTab)
-					->AddTab(SystemScriptTabID, ETabState::ClosedTab)
-					->AddTab(SystemDetailsTabID, ETabState::ClosedTab)
-					->AddTab(DebugSpreadsheetTabID, ETabState::ClosedTab)
-					->AddTab(PreviewSettingsTabId, ETabState::ClosedTab)
-					->AddTab(GeneratedCodeTabID, ETabState::ClosedTab)
-				)
-			)
-		);
-
+	
 	const bool bCreateDefaultStandaloneMenu = true;
 	const bool bCreateDefaultToolbar = true;
 	UObject* ToolkitObject = SystemToolkitMode == ESystemToolkitMode::System ? (UObject*)System : (UObject*)Emitter;
 	// order of registering commands matters. SetupCommands before InitAssetEditor will make the toolkit prioritize niagara commands
 	SetupCommands();
+
+	const TSharedRef<FTabManager::FLayout> DummyLayout = FTabManager::NewLayout("NullLayout")->AddArea(FTabManager::NewPrimaryArea());
 	FAssetEditorToolkit::InitAssetEditor(Mode, InitToolkitHost, FNiagaraEditorModule::NiagaraEditorAppIdentifier,
-		StandaloneDefaultLayout, bCreateDefaultStandaloneMenu, bCreateDefaultToolbar, ToolkitObject);
+		DummyLayout, bCreateDefaultStandaloneMenu, bCreateDefaultToolbar, ToolkitObject);
+	
+	AddApplicationMode(DefaultModeName, MakeShared<FNiagaraSystemToolkitMode_Default>(SharedThis(this)));
+	AddApplicationMode(ScalabilityModeName, MakeShared<FNiagaraSystemToolkitMode_Scalability>(SharedThis(this)));
+	
+	SetCurrentMode(DefaultModeName);
+	//SystemViewModel->GetScalabilityViewModel()->Activate();
 	
 	FNiagaraEditorModule& NiagaraEditorModule = FModuleManager::LoadModuleChecked<FNiagaraEditorModule>("NiagaraEditor");
 	AddMenuExtender(NiagaraEditorModule.GetMenuExtensibilityManager()->GetAllExtenders(GetToolkitCommands(), GetEditingObjects()));
-
-	ExtendToolbar();
+	
 	RegenerateMenusAndToolbars();
 
 	bChangesDiscarded = false;
 	bScratchPadChangesDiscarded = false;
 
 	GEditor->RegisterForUndo(this);
-
 #if WITH_NIAGARA_GPU_PROFILER
 	GpuProfilerListener.Reset(new FNiagaraGpuProfilerListener);
 #endif
@@ -544,387 +390,6 @@ void FNiagaraSystemToolkit::PostUndo(bool bSuccess)
 	LastUndoGraphs.Empty();
 }
 
-TSharedRef<SDockTab> FNiagaraSystemToolkit::SpawnTab_Viewport(const FSpawnTabArgs& Args)
-{
-	check(Args.GetTabId().TabType == ViewportTabID);
-
-	Viewport = SNew(SNiagaraSystemViewport)
-		.OnThumbnailCaptured(this, &FNiagaraSystemToolkit::OnThumbnailCaptured)
-		.Sequencer(SystemViewModel->GetSequencer());
-
-	TSharedRef<SDockTab> SpawnedTab =
-		SNew(SDockTab)
-		[
-			Viewport.ToSharedRef()
-		];
-
-	Viewport->SetPreviewComponent(SystemViewModel->GetPreviewComponent());
-	Viewport->OnAddedToTab(SpawnedTab);
-
-	return SpawnedTab;
-}
-
-TSharedRef<SDockTab> FNiagaraSystemToolkit::SpawnTab_PreviewSettings(const FSpawnTabArgs& Args)
-{
-	check(Args.GetTabId() == PreviewSettingsTabId);
-
-	TSharedRef<SWidget> InWidget = SNullWidget::NullWidget;
-	if (Viewport.IsValid())
-	{
-		FAdvancedPreviewSceneModule& AdvancedPreviewSceneModule = FModuleManager::LoadModuleChecked<FAdvancedPreviewSceneModule>("AdvancedPreviewScene");
-		InWidget = AdvancedPreviewSceneModule.CreateAdvancedPreviewSceneSettingsWidget(Viewport->GetPreviewScene());
-	}
-
-	TSharedRef<SDockTab> SpawnedTab = SNew(SDockTab)
-		.Label(LOCTEXT("PreviewSceneSettingsTab", "Preview Scene Settings"))
-		[
-			InWidget
-		];
-
-	return SpawnedTab;
-}
-
-
-TSharedRef<SDockTab> FNiagaraSystemToolkit::SpawnTab_CurveEd(const FSpawnTabArgs& Args)
-{
-	check(Args.GetTabId().TabType == CurveEditorTabID);
-
-	TSharedRef<SDockTab> SpawnedTab =
-		SNew(SDockTab)
-		[
-			FNiagaraEditorModule::Get().GetWidgetProvider()->CreateCurveOverview(SystemViewModel.ToSharedRef())
-		];
-
-	return SpawnedTab;
-}
-
-
-TSharedRef<SDockTab> FNiagaraSystemToolkit::SpawnTab_Sequencer(const FSpawnTabArgs& Args)
-{
-	check(Args.GetTabId().TabType == SequencerTabID);
-
-	TSharedRef<SDockTab> SpawnedTab =
-		SNew(SDockTab)
-		[
-			SystemViewModel->GetSequencer()->GetSequencerWidget()
-		];
-
-	return SpawnedTab;
-}
-
-
-TSharedRef<SDockTab> FNiagaraSystemToolkit::SpawnTab_SystemScript(const FSpawnTabArgs& Args)
-{
-	check(Args.GetTabId().TabType == SystemScriptTabID);
-
-	TSharedRef<SDockTab> SpawnedTab =
-		SNew(SDockTab)
-		[
-			SNew(SNiagaraSystemScript, SystemViewModel.ToSharedRef())
-		];
-
-	return SpawnedTab;
-}
-
-TSharedRef<SDockTab> FNiagaraSystemToolkit::SpawnTab_SystemParameters(const FSpawnTabArgs& Args)
-{
-	check(Args.GetTabId().TabType == SystemParametersTabID);
-
-
-	TArray<TSharedRef<FNiagaraObjectSelection>> ObjectSelections;
-	ObjectSelections.Add(ObjectSelectionForParameterMapView.ToSharedRef());
-
-	TSharedRef<SDockTab> SpawnedTab =
-		SNew(SDockTab)
-		[
-			SAssignNew(ParameterPanel, SNiagaraParameterPanel, ParameterPanelViewModel, GetToolkitCommands())
-			.ShowParameterSynchronizingWithLibraryIconExternallyReferenced(false)
-		];
-	RefreshParameters();
-
-	return SpawnedTab;
-}
-
-//@todo(ng) cleanup
-TSharedRef<SDockTab> FNiagaraSystemToolkit::SpawnTab_SystemParameters2(const FSpawnTabArgs& Args)
-{
-	check(Args.GetTabId().TabType == SystemParametersTabID2);
-
-	TArray<TSharedRef<FNiagaraObjectSelection>> ObjectSelections;
-	ObjectSelections.Add(ObjectSelectionForParameterMapView.ToSharedRef());
-
-	TSharedRef<SDockTab> SpawnedTab =
-		SNew(SDockTab)
-		[
-			SAssignNew(ParameterMapView, SNiagaraParameterMapView, ObjectSelections, SNiagaraParameterMapView::EToolkitType::SYSTEM, GetToolkitCommands())
-		];
-	RefreshParameters();
-
-	return SpawnedTab;
-}
-
-TSharedRef<SDockTab> FNiagaraSystemToolkit::SpawnTab_SystemParameterDefinitions(const FSpawnTabArgs& Args)
-{
-	checkf(Args.GetTabId().TabType == SystemParameterDefinitionsTabID, TEXT("Wrong tab ID in NiagaraScriptToolkit"));
-
-	TSharedRef<SDockTab> SpawnedTab =
-		SNew(SDockTab)
-		[
-			SNew(SNiagaraParameterDefinitionsPanel, ParameterDefinitionsPanelViewModel, GetToolkitCommands())
-		];
-
-	return SpawnedTab;
-}
-
-TSharedRef<SDockTab> FNiagaraSystemToolkit::SpawnTab_SelectedEmitterStack(const FSpawnTabArgs& Args)
-{
-	check(Args.GetTabId().TabType == SelectedEmitterStackTabID);
-
-	FNiagaraEditorModule& NiagaraEditorModule = FModuleManager::LoadModuleChecked<FNiagaraEditorModule>("NiagaraEditor");
-	TSharedRef<SDockTab> SpawnedTab = SNew(SDockTab)
-		.Label(LOCTEXT("SystemOverviewSelection", "Selection"))
-		[
-			NiagaraEditorModule.GetWidgetProvider()->CreateStackView(*SystemViewModel->GetSelectionViewModel()->GetSelectionStackViewModel())
-		];
-
-	SDockTab::FOnTabClosedCallback TabClosedCallback = SDockTab::FOnTabClosedCallback::CreateLambda([=](TSharedRef<SDockTab> DockTab)
-	{
-		SystemViewModel->GetSelectionViewModel()->GetSelectionStackViewModel()->ResetSearchText();
-	});
-	
-	SpawnedTab->SetOnTabClosed(TabClosedCallback);
-	return SpawnedTab;
-}
-
-class SNiagaraSelectedEmitterGraph : public SCompoundWidget
-{
-public:
-	SLATE_BEGIN_ARGS(SNiagaraSelectedEmitterGraph)
-	{}
-	SLATE_END_ARGS();
-
-	void Construct(const FArguments& InArgs, TSharedRef<FNiagaraSystemViewModel> InSystemViewModel)
-	{
-		SystemViewModel = InSystemViewModel;
-		SystemViewModel->GetSelectionViewModel()->OnEmitterHandleIdSelectionChanged().AddSP(this, &SNiagaraSelectedEmitterGraph::SystemSelectionChanged);
-		ChildSlot
-		[
-			SAssignNew(GraphWidgetContainer, SBox)
-		];
-		UpdateGraphWidget();
-	}
-
-	~SNiagaraSelectedEmitterGraph()
-	{
-		if (SystemViewModel.IsValid() && SystemViewModel->GetSelectionViewModel())
-		{
-			SystemViewModel->GetSelectionViewModel()->OnEmitterHandleIdSelectionChanged().RemoveAll(this);
-		}
-	}
-
-private:
-	void SystemSelectionChanged()
-	{
-		UpdateGraphWidget();
-	}
-
-	void UpdateGraphWidget()
-	{
-		TArray<FGuid> SelectedEmitterHandleIds = SystemViewModel->GetSelectionViewModel()->GetSelectedEmitterHandleIds();
-		if (SelectedEmitterHandleIds.Num() == 1)
-		{
-			TSharedPtr<FNiagaraEmitterHandleViewModel> SelectedEmitterHandle = SystemViewModel->GetEmitterHandleViewModelById(SelectedEmitterHandleIds[0]);
-			TSharedRef<SWidget> EmitterWidget = 
-				SNew(SSplitter)
-				+ SSplitter::Slot()
-				.Value(.25f)
-				[
-					SNew(SNiagaraSelectedObjectsDetails, SelectedEmitterHandle->GetEmitterViewModel()->GetSharedScriptViewModel()->GetGraphViewModel()->GetNodeSelection())
-				]
-				+ SSplitter::Slot()
-				.Value(.75f)
-				[
-					SNew(SNiagaraScriptGraph, SelectedEmitterHandle->GetEmitterViewModel()->GetSharedScriptViewModel()->GetGraphViewModel())
-				];
-
-			UNiagaraEmitter* LastMergedEmitter = SelectedEmitterHandle->GetEmitterViewModel()->GetEmitter()->GetParentAtLastMerge();
-			if (LastMergedEmitter != nullptr)
-			{
-				UNiagaraScriptSource* LastMergedScriptSource = CastChecked<UNiagaraScriptSource>(LastMergedEmitter->GraphSource);
-				bool bIsForDataProcessingOnly = false;
-				TSharedRef<FNiagaraScriptGraphViewModel> LastMergedScriptGraphViewModel = MakeShared<FNiagaraScriptGraphViewModel>(FText(), bIsForDataProcessingOnly);
-				LastMergedScriptGraphViewModel->SetScriptSource(LastMergedScriptSource);
-				TSharedRef<SWidget> LastMergedEmitterWidget = 
-					SNew(SSplitter)
-					+ SSplitter::Slot()
-					.Value(.25f)
-					[
-						SNew(SNiagaraSelectedObjectsDetails, LastMergedScriptGraphViewModel->GetNodeSelection())
-					]
-					+ SSplitter::Slot()
-					.Value(.75f)
-					[
-						SNew(SNiagaraScriptGraph, LastMergedScriptGraphViewModel)
-					];
-
-				GraphWidgetContainer->SetContent
-				(
-					SNew(SSplitter)
-					.Orientation(Orient_Vertical)
-					+ SSplitter::Slot()
-					[
-						SNew(SVerticalBox)
-						+ SVerticalBox::Slot()
-						.AutoHeight()
-						[
-							SNew(STextBlock)
-							.Text(FText::FromString(TEXT("Emitter")))
-						]
-						+ SVerticalBox::Slot()
-						[
-							EmitterWidget
-						]
-					]
-					+ SSplitter::Slot()
-					[
-						SNew(SVerticalBox)
-						+ SVerticalBox::Slot()
-						.AutoHeight()
-						[
-							SNew(STextBlock)
-							.Text(FText::FromString(TEXT("Last Merged Emitter")))
-						]
-						+ SVerticalBox::Slot()
-						[
-							LastMergedEmitterWidget
-						]
-					]
-				);
-			}
-			else
-			{
-				GraphWidgetContainer->SetContent(EmitterWidget);
-			}
-		}
-		else
-		{
-			GraphWidgetContainer->SetContent(SNullWidget::NullWidget);
-		}
-	}
-
-private:
-	TSharedPtr<FNiagaraSystemViewModel> SystemViewModel;
-	TSharedPtr<SBox> GraphWidgetContainer;
-};
-
-TSharedRef<SDockTab> FNiagaraSystemToolkit::SpawnTab_SelectedEmitterGraph(const FSpawnTabArgs& Args)
-{
-	check(Args.GetTabId().TabType == SelectedEmitterGraphTabID);
-
-	TSharedRef<SDockTab> SpawnedTab =
-		SNew(SDockTab)
-		[
-			SNew(SNiagaraSelectedEmitterGraph, SystemViewModel.ToSharedRef())
-		];
-
-	return SpawnedTab;
-}
-
-TSharedRef<SDockTab> FNiagaraSystemToolkit::SpawnTab_DebugSpreadsheet(const FSpawnTabArgs& Args)
-{
-	check(Args.GetTabId().TabType == DebugSpreadsheetTabID);
-
-	TSharedRef<SDockTab> SpawnedTab =
-		SNew(SDockTab)
-		[
-			SNew(SNiagaraSpreadsheetView, SystemViewModel.ToSharedRef())
-		];
-
-	return SpawnedTab;
-}
-
-TSharedRef<SDockTab> FNiagaraSystemToolkit::SpawnTab_GeneratedCode(const FSpawnTabArgs& Args)
-{
-	check(Args.GetTabId().TabType == GeneratedCodeTabID);
-
-	TSharedRef<SDockTab> SpawnedTab = SNew(SDockTab);
-	SpawnedTab->SetContent(SNew(SNiagaraGeneratedCodeView, SystemViewModel.ToSharedRef(), SpawnedTab));
-	return SpawnedTab;
-}
-
-TSharedRef<SDockTab> FNiagaraSystemToolkit::SpawnTab_MessageLog(const FSpawnTabArgs& Args)
-{
-	check(Args.GetTabId().TabType == MessageLogTabID);
-
-	TSharedRef<SDockTab> SpawnedTab = SNew(SDockTab)
-		.Label(LOCTEXT("NiagaraMessageLogTitle", "Niagara Log"))
-		[
-			SNew(SBox)
-			.AddMetaData<FTagMetaData>(FTagMetaData(TEXT("NiagaraLog")))
-			[
-				NiagaraMessageLog.ToSharedRef()
-			]
-		];
-
-	return SpawnedTab;
-}
-
-TSharedRef<SDockTab> FNiagaraSystemToolkit::SpawnTab_SystemOverview(const FSpawnTabArgs& Args)
-{
-	TSharedRef<SDockTab> SpawnedTab = SNew(SDockTab)
-		.Label(LOCTEXT("SystemOverviewTabLabel", "System Overview"))
-		[
-			FNiagaraEditorModule::Get().GetWidgetProvider()->CreateSystemOverview(SystemViewModel.ToSharedRef())
-		];
-
-	return SpawnedTab;
-}
-
-TSharedRef<SDockTab> FNiagaraSystemToolkit::SpawnTab_ScratchPad(const FSpawnTabArgs& Args)
-{
-	TSharedRef<SDockTab> SpawnedTab = SNew(SDockTab)
-		.Label(LOCTEXT("ScratchPadTabLabel", "Scratch Pad"))
-		[
-			FNiagaraEditorModule::Get().GetWidgetProvider()->CreateScriptScratchPad(*SystemViewModel->GetScriptScratchPadViewModel())
-		];
-
-	return SpawnedTab;
-}
-
-TSharedRef<SDockTab> FNiagaraSystemToolkit::SpawnTab_ScriptStats(const FSpawnTabArgs& Args)
-{
-	check(Args.GetTabId().TabType == ScriptStatsTabID);
-
-	TSharedRef<SDockTab> SpawnedTab = SNew(SDockTab)
-		.Label(LOCTEXT("NiagaraScriptStatsTitle", "Script Stats"))
-		[
-			SNew(SBox)
-			.AddMetaData<FTagMetaData>(FTagMetaData(TEXT("ScriptStats")))
-			[
-				ScriptStats->GetWidget().ToSharedRef()
-			]
-		];
-
-	return SpawnedTab;
-}
-
-TSharedRef<SDockTab> FNiagaraSystemToolkit::SpawnTab_Baker(const FSpawnTabArgs& Args)
-{
-	check(Args.GetTabId().TabType == BakerTabID);
-
-	TSharedRef<SDockTab> SpawnedTab = SNew(SDockTab)
-		.Label(LOCTEXT("NiagaraBakerTitle", "Baker"))
-		[
-			SNew(SBox)
-			.AddMetaData<FTagMetaData>(FTagMetaData(TEXT("Baker")))
-			[
-				BakerViewModel->GetWidget().ToSharedRef()
-			]
-		];
-
-	return SpawnedTab;
-}
-
 void FNiagaraSystemToolkit::SetupCommands()
 {
 	FLevelEditorModule& Module = FModuleManager::Get().LoadModuleChecked<FLevelEditorModule>(TEXT("LevelEditor"));
@@ -1112,178 +577,6 @@ void FNiagaraSystemToolkit::ResetSimulation()
 	SystemViewModel->ResetSystem(FNiagaraSystemViewModel::ETimeResetMode::AllowResetTime, FNiagaraSystemViewModel::EMultiResetMode::AllowResetAllInstances, FNiagaraSystemViewModel::EReinitMode::ReinitializeSystem);
 }
 
-void FNiagaraSystemToolkit::ExtendToolbar()
-{
-	struct Local
-	{
-		static TSharedRef<SWidget> FillSimulationOptionsMenu(FNiagaraSystemToolkit* Toolkit)
-		{
-			FMenuBuilder MenuBuilder(true, Toolkit->GetToolkitCommands());
-			MenuBuilder.AddMenuEntry(FNiagaraEditorCommands::Get().ToggleAutoPlay);
-			MenuBuilder.AddMenuEntry(FNiagaraEditorCommands::Get().ToggleResetSimulationOnChange);
-			MenuBuilder.AddMenuEntry(FNiagaraEditorCommands::Get().ToggleResimulateOnChangeWhilePaused);
-			MenuBuilder.AddMenuEntry(FNiagaraEditorCommands::Get().ToggleResetDependentSystems);
-			return MenuBuilder.MakeWidget();
-		}
-
-		static TSharedRef<SWidget> GenerateBakerMenu(FNiagaraSystemToolkit* Toolkit)
-		{
-			FMenuBuilder MenuBuilder(true, Toolkit->GetToolkitCommands());
-
-			MenuBuilder.AddMenuEntry(
-				LOCTEXT("BakerTab", "Open Baker Tab"),
-				LOCTEXT("BakerTabTooltip", "Opens the flip book tab."),
-				FSlateIcon(),
-				FUIAction(FExecuteAction::CreateLambda([TabManager=Toolkit->TabManager]() { TabManager->TryInvokeTab(BakerTabID); }))
-			);
-
-			MenuBuilder.AddMenuEntry(
-				LOCTEXT("RenderBaker", "Render Baker"),
-				LOCTEXT("RenderBakerTooltip", "Renders the Baker using the current settings."),
-				FSlateIcon(),
-				FUIAction(FExecuteAction::CreateSP(Toolkit, &FNiagaraSystemToolkit::RenderBaker))
-			);
-
-			return MenuBuilder.MakeWidget();
-		}
-
-		static TSharedRef<SWidget> FillDebugOptionsMenu(FNiagaraSystemToolkit* Toolkit)
-		{
-			FMenuBuilder MenuBuilder(true, Toolkit->GetToolkitCommands());
-
-#if WITH_NIAGARA_DEBUGGER
-			MenuBuilder.AddMenuEntry(FNiagaraEditorCommands::Get().OpenDebugHUD);
-			MenuBuilder.AddMenuEntry(FNiagaraEditorCommands::Get().OpenDebugOutliner);
-#endif
-			MenuBuilder.AddMenuEntry(FNiagaraEditorCommands::Get().OpenAttributeSpreadsheet);
-			return MenuBuilder.MakeWidget();
-		}
-
-		static void FillToolbar(FToolBarBuilder& ToolbarBuilder, FNiagaraSystemToolkit* Toolkit)
-		{
-			ToolbarBuilder.BeginSection("Compile");
-			{
-				if (Toolkit->Emitter != nullptr)
-				{
-					{
-						ToolbarBuilder.AddToolBarButton(FNiagaraEditorCommands::Get().Apply,
-							NAME_None, TAttribute<FText>(), TAttribute<FText>(),
-							FSlateIcon(FAppStyle::Get().GetStyleSetName(), "AssetEditor.Apply"),
-							FName(TEXT("ApplyNiagaraEmitter")));
-					}
-				}
-				ToolbarBuilder.AddToolBarButton(FNiagaraEditorCommands::Get().ApplyScratchPadChanges,
-					NAME_None,
-					TAttribute<FText>(),
-					TAttribute<FText>(),
-					FSlateIcon(FNiagaraEditorStyle::Get().GetStyleSetName(), "NiagaraEditor.ApplyScratchPadChanges"),
-					FName(TEXT("ApplyScratchPadChanges")));
-				ToolbarBuilder.AddToolBarButton(FNiagaraEditorCommands::Get().Compile,
-					NAME_None,
-					TAttribute<FText>(),
-					TAttribute<FText>(Toolkit, &FNiagaraSystemToolkit::GetCompileStatusTooltip),
-					TAttribute<FSlateIcon>(Toolkit, &FNiagaraSystemToolkit::GetCompileStatusImage),
-					FName(TEXT("CompileNiagaraSystem")));
-				ToolbarBuilder.AddComboButton(
-					FUIAction(),
-					FOnGetContent::CreateRaw(Toolkit, &FNiagaraSystemToolkit::GenerateCompileMenuContent),
-					LOCTEXT("BuildCombo_Label", "Auto-Compile Options"),
-					LOCTEXT("BuildComboToolTip", "Auto-Compile options menu"),
-					FSlateIcon(FEditorStyle::GetStyleSetName(), "LevelEditor.Build"),
-					true);
-			}
-			ToolbarBuilder.EndSection();
-
-			ToolbarBuilder.BeginSection("NiagaraTools");
-			{
-				ToolbarBuilder.AddToolBarButton(FNiagaraEditorCommands::Get().SaveThumbnailImage, NAME_None,
-					LOCTEXT("GenerateThumbnail", "Thumbnail"),
-					LOCTEXT("GenerateThumbnailTooltip","Generate a thumbnail image."),
-					FSlateIcon(FAppStyle::Get().GetStyleSetName(), "AssetEditor.SaveThumbnail"));
-		
-				ToolbarBuilder.AddToolBarButton(FNiagaraEditorCommands::Get().ToggleBounds, NAME_None,
-					LOCTEXT("ShowBounds", "Bounds"),
-					LOCTEXT("ShowBoundsTooltip", "Show the bounds for the scene."),
-					FSlateIcon(FAppStyle::Get().GetStyleSetName(), "AssetEditor.ToggleShowBounds"));
-				ToolbarBuilder.AddComboButton(
-					FUIAction(),
-					FOnGetContent::CreateRaw(Toolkit, &FNiagaraSystemToolkit::GenerateBoundsMenuContent, Toolkit->GetToolkitCommands()),
-					LOCTEXT("BoundsMenuCombo_Label", "Bounds Options"),
-					LOCTEXT("BoundsMenuCombo_ToolTip", "Bounds options"),
-					FSlateIcon(FAppStyle::Get().GetStyleSetName(), "AssetEditor.ToggleShowBounds"),
-					true
-				);
-			}
-			ToolbarBuilder.EndSection();
-			
-#if STATS
-			ToolbarBuilder.BeginSection("NiagaraStatisticsOptions");
-			{
-				ToolbarBuilder.AddToolBarButton(FNiagaraEditorCommands::Get().ToggleStatPerformance, NAME_None,
-					LOCTEXT("NiagaraShowPerformance", "Performance"),
-					LOCTEXT("NiagaraShowPerformanceTooltip", "Show runtime performance for particle scripts."),
-					FSlateIcon(FAppStyle::Get().GetStyleSetName(), "AssetEditor.ToggleStats"));
-				ToolbarBuilder.AddComboButton(
-                    FUIAction(),
-                    FOnGetContent::CreateRaw(Toolkit, &FNiagaraSystemToolkit::GenerateStatConfigMenuContent, Toolkit->GetToolkitCommands()),
-                    FText(),
-                    LOCTEXT("NiagaraShowPerformanceCombo_ToolTip", "Runtime performance options"),
-					FSlateIcon(FAppStyle::Get().GetStyleSetName(), "AssetEditor.ToggleStats"),
-                    true);
-				ToolbarBuilder.AddComboButton(
-					FUIAction(),
-					FOnGetContent::CreateStatic(Local::FillDebugOptionsMenu, Toolkit),
-					LOCTEXT("DebugOptions", "Debug"),
-					LOCTEXT("DebugOptionsTooltip", "Debug options"),
-					FSlateIcon(FAppStyle::Get().GetStyleSetName(), "Debug"));
-			}
-			ToolbarBuilder.EndSection();
-#endif
-			
-			ToolbarBuilder.BeginSection("PlaybackOptions");
-			{
-				ToolbarBuilder.AddComboButton(
-					FUIAction(),
-					FOnGetContent::CreateStatic(Local::FillSimulationOptionsMenu, Toolkit),
-					LOCTEXT("SimulationOptions", "Simulation"),
-					LOCTEXT("SimulationOptionsTooltip", "Simulation options"),
-					FSlateIcon(FAppStyle::Get().GetStyleSetName(), "AssetEditor.Simulate")
-				);
-			}
-			ToolbarBuilder.EndSection();
-
-			if ( GetDefault<UNiagaraEditorSettings>()->bEnableBaker )
-			{
-				ToolbarBuilder.BeginSection("Baker");
-				{
-					ToolbarBuilder.AddComboButton(
-						FUIAction(),
-						FOnGetContent::CreateStatic(Local::GenerateBakerMenu, Toolkit),
-						LOCTEXT("Baker", "Baker"),
-						LOCTEXT("BakerTooltip", "Options for Baker rendering."),
-						FSlateIcon(FNiagaraEditorStyle::Get().GetStyleSetName(), "NiagaraEditor.Baker")
-					);
-				}
-				ToolbarBuilder.EndSection();
-			}
-		}
-	};
-
-	TSharedPtr<FExtender> ToolbarExtender = MakeShareable(new FExtender);
-
-	ToolbarExtender->AddToolBarExtension(
-		"Asset",
-		EExtensionHook::After,
-		GetToolkitCommands(),
-		FToolBarExtensionDelegate::CreateStatic(&Local::FillToolbar, this)
-		);
-
-	AddToolbarExtender(ToolbarExtender);
-
-	FNiagaraEditorModule& NiagaraEditorModule = FModuleManager::LoadModuleChecked<FNiagaraEditorModule>("NiagaraEditor");
-	AddToolbarExtender(NiagaraEditorModule.GetToolBarExtensibilityManager()->GetAllExtenders(GetToolkitCommands(), GetEditingObjects()));
-}
-
 TSharedRef<SWidget> FNiagaraSystemToolkit::GenerateBoundsMenuContent(TSharedRef<FUICommandList> InCommandList)
 {
 	const bool bShouldCloseWindowAfterMenuSelection = true;
@@ -1437,6 +730,16 @@ TSharedPtr<FNiagaraSystemViewModel> FNiagaraSystemToolkit::GetSystemViewModel()
 	return SystemViewModel;
 }
 
+void FNiagaraSystemToolkit::RegisterToolbarTab(const TSharedRef<FTabManager>& InTabManager)
+{
+	FAssetEditorToolkit::RegisterTabSpawners(InTabManager);
+}
+
+const TArray<UObject*>& FNiagaraSystemToolkit::GetObjectsBeingEdited() const
+{
+	return GetEditingObjects();
+}
+
 void FNiagaraSystemToolkit::OnToggleBounds()
 {
 	ToggleDrawOption(SNiagaraSystemViewport::Bounds);
@@ -1496,7 +799,7 @@ void FNiagaraSystemToolkit::OpenDebugOutliner()
 
 void FNiagaraSystemToolkit::OpenAttributeSpreadsheet()
 {
-	InvokeTab(DebugSpreadsheetTabID);
+	InvokeTab(FNiagaraSystemToolkitModeBase::DebugSpreadsheetTabID);
 }
 
 void FNiagaraSystemToolkit::OnToggleBoundsSetFixedBounds_Emitters()
@@ -1938,7 +1241,8 @@ bool FNiagaraSystemToolkit::OnApplyScratchPadChangesEnabled() const
 
 void FNiagaraSystemToolkit::OnPinnedCurvesChanged()
 {
-	TabManager->TryInvokeTab(CurveEditorTabID);
+	// does this work due to modes? @todo
+	TabManager->TryInvokeTab(FNiagaraSystemToolkitModeBase::CurveEditorTabID);
 }
 
 void FNiagaraSystemToolkit::RefreshParameters()
diff --git a/Engine/Plugins/FX/Niagara/Source/NiagaraEditor/Private/Toolkits/NiagaraSystemToolkit.h b/Engine/Plugins/FX/Niagara/Source/NiagaraEditor/Private/Toolkits/NiagaraSystemToolkit.h
--- a/Engine/Plugins/FX/Niagara/Source/NiagaraEditor/Private/Toolkits/NiagaraSystemToolkit.h
+++ b/Engine/Plugins/FX/Niagara/Source/NiagaraEditor/Private/Toolkits/NiagaraSystemToolkit.h
@@ -7,6 +7,7 @@
 #include "Toolkits/IToolkitHost.h"
 #include "Toolkits/AssetEditorToolkit.h"
 #include "EditorUndoClient.h"
+#include "NiagaraEditorModule.h"
 
 #include "ISequencer.h"
 #include "ISequencerTrackEditor.h"
@@ -15,6 +16,7 @@
 #include "NiagaraScript.h"
 #include "Widgets/SItemSelector.h"
 #include "ViewModels/NiagaraSystemGraphSelectionViewModel.h"
+#include "WorkflowOrientedApp/WorkflowCentricApplication.h"
 
 #if WITH_NIAGARA_GPU_PROFILER
 	#include "NiagaraGPUProfilerInterface.h"
@@ -39,20 +41,26 @@ class FNiagaraSystemToolkitParameterDefinitionsPanelViewModel;
 class FNiagaraScriptStatsViewModel;
 class FNiagaraBakerViewModel;
 
+extern int32 GNiagaraScalabilityModeEnabled;
+
 /** Viewer/editor for a NiagaraSystem
 */
-class FNiagaraSystemToolkit : public FAssetEditorToolkit, public FGCObject, public FEditorUndoClient
+class FNiagaraSystemToolkit : public FWorkflowCentricApplication, public FGCObject, public FEditorUndoClient
 {
+private:
 	enum class ESystemToolkitMode
 	{
 		System,
 		Emitter,
 	};
 
-public:
-	virtual void RegisterTabSpawners(const TSharedRef<class FTabManager>& TabManager) override;
-	virtual void UnregisterTabSpawners(const TSharedRef<class FTabManager>& TabManager) override;
+	enum class ESystemToolkitWorkflowMode
+	{
+		Default,
+		Scalability
+	};
 
+public:
 	/** Edits the specified Niagara System */
 	void InitializeWithSystem(const EToolkitMode::Type Mode, const TSharedPtr< class IToolkitHost >& InitToolkitHost, UNiagaraSystem& InSystem);
 
@@ -61,7 +69,7 @@ public:
 
 	/** Destructor */
 	virtual ~FNiagaraSystemToolkit();
-
+	
 	//~ Begin IToolkit Interface
 	virtual FName GetToolkitFName() const override;
 	virtual FText GetBaseToolkitName() const override;
@@ -89,7 +97,26 @@ public:
 	void CompileSystem(bool bFullRebuild);
 
 	TSharedPtr<FNiagaraSystemViewModel> GetSystemViewModel();
+	
+	// @todo This is a hack for now until we reconcile the default toolbar with application modes [duplicated from counterpart in Blueprint Editor]
+	void RegisterToolbarTab(const TSharedRef<class FTabManager>& TabManager);
 
+	TSharedPtr<SWidget> GetSystemOverview() const;
+	void SetSystemOverview(const TSharedPtr<SWidget>& SystemOverview);
+	TSharedPtr<SWidget> GetScriptScratchpad() const;
+	void SetScriptScratchpad(const TSharedPtr<SWidget>& ScriptScratchpad);
+	
+	/** Mode exposed functions */
+	void RenderBaker();
+	TSharedRef<SWidget> GenerateCompileMenuContent();
+	TSharedRef<SWidget> GenerateBoundsMenuContent(TSharedRef<FUICommandList> InCommandList);
+	TSharedRef<SWidget> GenerateStatConfigMenuContent(TSharedRef<FUICommandList> InCommandList);
+
+	bool HasEmitter() const { return SystemToolkitMode == ESystemToolkitMode::Emitter && Emitter != nullptr; }
+
+	bool HasSystem() const { return SystemToolkitMode == ESystemToolkitMode::System && System != nullptr; }
+
+	const TArray<UObject*>& GetObjectsBeingEdited() const;
 public:
 	FRefreshItemSelectorDelegate RefreshItemSelector;
 
@@ -132,33 +159,12 @@ private:
 	void UpdateOriginalEmitter();
 	void UpdateExistingEmitters();
 
-	TSharedRef<SDockTab> SpawnTab_Viewport(const FSpawnTabArgs& Args);
-	TSharedRef<SDockTab> SpawnTab_CurveEd(const FSpawnTabArgs& Args);
-	TSharedRef<SDockTab> SpawnTab_Sequencer(const FSpawnTabArgs& Args);
-	TSharedRef<SDockTab> SpawnTab_SystemScript(const FSpawnTabArgs& Args);
-	TSharedRef<SDockTab> SpawnTab_SystemParameters(const FSpawnTabArgs& Args);
-	TSharedRef<SDockTab> SpawnTab_SystemParameters2(const FSpawnTabArgs& Args); //@todo(ng) cleanup
-	TSharedRef<SDockTab> SpawnTab_SystemParameterDefinitions(const FSpawnTabArgs& Args);
-	TSharedRef<SDockTab> SpawnTab_SelectedEmitterStack(const FSpawnTabArgs& Args);
-	TSharedRef<SDockTab> SpawnTab_SelectedEmitterGraph(const FSpawnTabArgs& Args);
-	TSharedRef<SDockTab> SpawnTab_DebugSpreadsheet(const FSpawnTabArgs& Args);
-	TSharedRef<SDockTab> SpawnTab_PreviewSettings(const FSpawnTabArgs& Args);
-	TSharedRef<SDockTab> SpawnTab_GeneratedCode(const FSpawnTabArgs& Args);
-	TSharedRef<SDockTab> SpawnTab_MessageLog(const FSpawnTabArgs& Args);
-	TSharedRef<SDockTab> SpawnTab_SystemOverview(const FSpawnTabArgs& Args);
-	TSharedRef<SDockTab> SpawnTab_ScratchPad(const FSpawnTabArgs& Args);
-	TSharedRef<SDockTab> SpawnTab_ScriptStats(const FSpawnTabArgs& Args);
-	TSharedRef<SDockTab> SpawnTab_Baker(const FSpawnTabArgs& Args);
-
-	/** Builds the toolbar widget */
-	void ExtendToolbar();	
 	void SetupCommands();
 
 	void ResetSimulation();
 
 	void GetSequencerAddMenuContent(FMenuBuilder& MenuBuilder, TSharedRef<ISequencer> Sequencer);
 	TSharedRef<SWidget> CreateAddEmitterMenuContent();
-	TSharedRef<SWidget> GenerateCompileMenuContent();
 
 	void EmitterAssetSelected(const FAssetData& AssetData);
 
@@ -175,11 +181,7 @@ private:
 	void RefreshParameters();
 	void OnSystemSelectionChanged();
 	void OnViewModelRequestFocusTab(FName TabName);
-
-	void RenderBaker();
-
-	TSharedRef<SWidget> GenerateBoundsMenuContent(TSharedRef<FUICommandList> InCommandList);
-	TSharedRef<SWidget> GenerateStatConfigMenuContent(TSharedRef<FUICommandList> InCommandList);
+	
 	const FName GetNiagaraSystemMessageLogName(UNiagaraSystem* InSystem) const;
 	void OnSaveThumbnailImage();
 	void OnThumbnailCaptured(UTexture2D* Thumbnail);
@@ -204,11 +206,14 @@ private:
 
 	ESystemToolkitMode SystemToolkitMode;
 
+	/** The currently active workflow mode */
+	ESystemToolkitWorkflowMode ActiveWorkflowMode;
+	
 	TSharedPtr<SNiagaraSystemViewport> Viewport;
 
 	/* The view model for the System being edited */
 	TSharedPtr<FNiagaraSystemViewModel> SystemViewModel;
-
+	
 	/* The view model for the selected Emitter Script graphs of the System being edited. */
 	TSharedPtr<FNiagaraSystemGraphSelectionViewModel> SystemGraphSelectionViewModel;
 
@@ -216,15 +221,15 @@ private:
 	TSharedPtr<FNiagaraMessageLogViewModel> NiagaraMessageLogViewModel;
 	TSharedPtr<class SWidget> NiagaraMessageLog;
 
+	/** The command list for this editor */
+	TSharedPtr<FUICommandList> EditorCommands;
+
 	/** Display for script stats on selected platforms */
 	TSharedPtr<FNiagaraScriptStatsViewModel> ScriptStats;
 
 	/** Baker preview */
 	TSharedPtr<FNiagaraBakerViewModel> BakerViewModel;
 
-	/** The command list for this editor */
-	TSharedPtr<FUICommandList> EditorCommands;
-
 	TSharedPtr<class SNiagaraParameterMapView> ParameterMapView;
 
 	TSharedPtr<FNiagaraSystemToolkitParameterPanelViewModel> ParameterPanelViewModel;
@@ -237,29 +242,19 @@ private:
 
 	bool bScratchPadChangesDiscarded;
 
+	TSharedPtr<SWidget> SystemOverview;
+	TSharedPtr<SWidget> ScriptScratchpad;
+	
 	static IConsoleVariable* VmStatEnabledVar;
+	static IConsoleVariable* GpuStatEnabledVar;
+	
+	friend class FNiagaraSystemToolkitModeBase;
 
 #if WITH_NIAGARA_GPU_PROFILER
 	TUniquePtr<FNiagaraGpuProfilerListener> GpuProfilerListener;
 #endif
 
 public:
-	static const FName ViewportTabID;
-	static const FName CurveEditorTabID;
-	static const FName SequencerTabID;
-	static const FName SystemScriptTabID;
-	static const FName SystemDetailsTabID;
-	static const FName SystemParametersTabID;
-	static const FName SystemParametersTabID2;
-	static const FName SystemParameterDefinitionsTabID;
-	static const FName SelectedEmitterStackTabID;
-	static const FName SelectedEmitterGraphTabID;
-	static const FName DebugSpreadsheetTabID;
-	static const FName PreviewSettingsTabId;
-	static const FName GeneratedCodeTabID;
-	static const FName MessageLogTabID;
-	static const FName SystemOverviewTabID;
-	static const FName ScratchPadTabID;
-	static const FName ScriptStatsTabID;
-	static const FName BakerTabID;
+	static const FName DefaultModeName;
+	static const FName ScalabilityModeName;
 };
diff --git a/Engine/Plugins/FX/Niagara/Source/NiagaraEditor/Private/Toolkits/SystemToolkitModes/NiagaraSystemToolkitModeBase.cpp b/Engine/Plugins/FX/Niagara/Source/NiagaraEditor/Private/Toolkits/SystemToolkitModes/NiagaraSystemToolkitModeBase.cpp
new file mode 100644
--- /dev/null
+++ b/Engine/Plugins/FX/Niagara/Source/NiagaraEditor/Private/Toolkits/SystemToolkitModes/NiagaraSystemToolkitModeBase.cpp
@@ -0,0 +1,556 @@
+// Copyright Epic Games, Inc. All Rights Reserved.
+
+#pragma once
+
+#include "NiagaraSystemToolkitModeBase.h"
+
+#include "AdvancedPreviewSceneModule.h"
+#include "NiagaraEditorModule.h"
+#include "NiagaraEditorSettings.h"
+#include "NiagaraEditorStyle.h"
+#include "NiagaraScriptGraphViewModel.h"
+#include "NiagaraSystemToolkit.h"
+#include "NiagaraScriptSource.h"
+#include "ViewModels/NiagaraEmitterHandleViewModel.h"
+#include "ViewModels/NiagaraEmitterViewModel.h"
+#include "ViewModels/NiagaraScriptViewModel.h"
+#include "ViewModels/NiagaraSystemSelectionViewModel.h"
+#include "ViewModels/NiagaraSystemViewModel.h"
+#include "ViewModels/Stack/NiagaraStackViewModel.h"
+#include "ViewModels/NiagaraParameterDefinitionsPanelViewModel.h"
+#include "Widgets/SNiagaraSystemViewport.h"
+#include "Widgets/SNiagaraSystemScript.h"
+#include "Widgets/SNiagaraParameterPanel.h"
+#include "Widgets/SNiagaraParameterMapView.h"
+#include "Widgets/SNiagaraParameterDefinitionsPanel.h"
+#include "Widgets/SNiagaraScriptGraph.h"
+#include "Widgets/SNiagaraSelectedObjectsDetails.h"
+#include "Widgets/SNiagaraSpreadsheetView.h"
+#include "Widgets/SNiagaraGeneratedCodeView.h"
+
+#define LOCTEXT_NAMESPACE "NiagaraSystemToolkitModeBase"
+
+const FName FNiagaraSystemToolkitModeBase::ViewportTabID(TEXT("NiagaraSystemEditor_Viewport"));
+const FName FNiagaraSystemToolkitModeBase::CurveEditorTabID(TEXT("NiagaraSystemEditor_CurveEditor"));
+const FName FNiagaraSystemToolkitModeBase::SequencerTabID(TEXT("NiagaraSystemEditor_Sequencer"));
+const FName FNiagaraSystemToolkitModeBase::SystemScriptTabID(TEXT("NiagaraSystemEditor_SystemScript"));
+const FName FNiagaraSystemToolkitModeBase::SystemDetailsTabID(TEXT("NiagaraSystemEditor_SystemDetails"));
+const FName FNiagaraSystemToolkitModeBase::SystemParametersTabID(TEXT("NiagaraSystemEditor_SystemParameters"));
+const FName FNiagaraSystemToolkitModeBase::SystemParametersTabID2(TEXT("NiagaraSystemEditor_SystemParameters2"));
+const FName FNiagaraSystemToolkitModeBase::SystemParameterDefinitionsTabID(TEXT("NiagaraSystemEditor_SystemParameterDefinitions"));
+const FName FNiagaraSystemToolkitModeBase::SelectedEmitterStackTabID(TEXT("NiagaraSystemEditor_SelectedEmitterStack"));
+const FName FNiagaraSystemToolkitModeBase::SelectedEmitterGraphTabID(TEXT("NiagaraSystemEditor_SelectedEmitterGraph"));
+const FName FNiagaraSystemToolkitModeBase::DebugSpreadsheetTabID(TEXT("NiagaraSystemEditor_DebugAttributeSpreadsheet"));
+const FName FNiagaraSystemToolkitModeBase::PreviewSettingsTabId(TEXT("NiagaraSystemEditor_PreviewSettings"));
+const FName FNiagaraSystemToolkitModeBase::GeneratedCodeTabID(TEXT("NiagaraSystemEditor_GeneratedCode"));
+const FName FNiagaraSystemToolkitModeBase::MessageLogTabID(TEXT("NiagaraSystemEditor_MessageLog"));
+const FName FNiagaraSystemToolkitModeBase::SystemOverviewTabID(TEXT("NiagaraSystemEditor_SystemOverview"));
+const FName FNiagaraSystemToolkitModeBase::ScratchPadTabID(TEXT("NiagaraSystemEditor_ScratchPad"));
+const FName FNiagaraSystemToolkitModeBase::ScriptStatsTabID(TEXT("NiagaraSystemEditor_ScriptStats"));
+const FName FNiagaraSystemToolkitModeBase::BakerTabID(TEXT("NiagaraSystemEditor_Baker"));
+
+void FNiagaraSystemToolkitModeBase::RegisterTabFactories(TSharedPtr<FTabManager> InTabManager)
+{
+	WorkspaceMenuCategory = InTabManager->AddLocalWorkspaceMenuCategory(LOCTEXT("WorkspaceMenu_NiagaraSystemEditor", "Niagara System"));
+
+	SystemToolkit.Pin()->RegisterToolbarTab(InTabManager.ToSharedRef());
+	
+	InTabManager->RegisterTabSpawner(ViewportTabID, FOnSpawnTab::CreateSP(this, &FNiagaraSystemToolkitModeBase::SpawnTab_Viewport))
+		.SetDisplayName(LOCTEXT("Preview", "Preview"))
+		.SetGroup(WorkspaceMenuCategory.ToSharedRef())
+		.SetIcon(FSlateIcon(FNiagaraEditorStyle::Get().GetStyleSetName(), "Tab.Viewport"));
+
+	InTabManager->RegisterTabSpawner(CurveEditorTabID, FOnSpawnTab::CreateSP(this, &FNiagaraSystemToolkitModeBase::SpawnTab_CurveEd))
+		.SetDisplayName(LOCTEXT("Curves", "Curves"))
+		.SetGroup(WorkspaceMenuCategory.ToSharedRef())
+		.SetIcon(FSlateIcon(FNiagaraEditorStyle::Get().GetStyleSetName(), "Tab.Curves"));
+
+	InTabManager->RegisterTabSpawner(SequencerTabID, FOnSpawnTab::CreateSP(this, &FNiagaraSystemToolkitModeBase::SpawnTab_Sequencer))
+		.SetDisplayName(LOCTEXT("Timeline", "Timeline"))
+		.SetGroup(WorkspaceMenuCategory.ToSharedRef())
+		.SetIcon(FSlateIcon(FNiagaraEditorStyle::Get().GetStyleSetName(), "Tab.Timeline"));
+
+	InTabManager->RegisterTabSpawner(SystemScriptTabID, FOnSpawnTab::CreateSP(this, &FNiagaraSystemToolkitModeBase::SpawnTab_SystemScript))
+		.SetDisplayName(LOCTEXT("SystemScript", "System Script"))
+		.SetGroup(WorkspaceMenuCategory.ToSharedRef())
+		.SetAutoGenerateMenuEntry(GbShowNiagaraDeveloperWindows != 0);
+
+	InTabManager->RegisterTabSpawner(SystemParametersTabID, FOnSpawnTab::CreateSP(this, &FNiagaraSystemToolkitModeBase::SpawnTab_SystemParameters))
+		.SetDisplayName(LOCTEXT("SystemParameters", "Parameters"))
+		.SetGroup(WorkspaceMenuCategory.ToSharedRef())
+		.SetIcon(FSlateIcon(FNiagaraEditorStyle::Get().GetStyleSetName(), "Tab.Parameters"));
+
+	InTabManager->RegisterTabSpawner(SystemParametersTabID2, FOnSpawnTab::CreateSP(this, &FNiagaraSystemToolkitModeBase::SpawnTab_SystemParameters2))
+		.SetDisplayName(LOCTEXT("SystemParameters2", "Legacy Parameters"))
+		.SetGroup(WorkspaceMenuCategory.ToSharedRef())
+		.SetIcon(FSlateIcon(FNiagaraEditorStyle::Get().GetStyleSetName(), "Tab.Parameters"));
+
+//@todo(ng) disable parameter definitions panel pending bug fixes
+// 	InTabManager->RegisterTabSpawner(SystemParameterDefinitionsTabID, FOnSpawnTab::CreateSP(this, &FNiagaraSystemToolkit::SpawnTab_SystemParameterDefinitions)) 
+// 		.SetDisplayName(LOCTEXT("SystemParameterDefinitions", "Parameter Definitions"))
+// 		.SetGroup(WorkspaceMenuCategory.ToSharedRef());
+
+	InTabManager->RegisterTabSpawner(SelectedEmitterStackTabID, FOnSpawnTab::CreateSP(this, &FNiagaraSystemToolkitModeBase::SpawnTab_SelectedEmitterStack))
+		.SetDisplayName(LOCTEXT("SelectedEmitterStacks", "Selected Emitters"))
+		.SetGroup(WorkspaceMenuCategory.ToSharedRef())
+		.SetIcon(FSlateIcon(FNiagaraEditorStyle::Get().GetStyleSetName(), "Tab.VisualEffects"));
+
+	InTabManager->RegisterTabSpawner(SelectedEmitterGraphTabID, FOnSpawnTab::CreateSP(this, &FNiagaraSystemToolkitModeBase::SpawnTab_SelectedEmitterGraph))
+		.SetDisplayName(LOCTEXT("SelectedEmitterGraph", "Selected Emitter Graph"))
+		.SetGroup(WorkspaceMenuCategory.ToSharedRef())
+		.SetAutoGenerateMenuEntry(GbShowNiagaraDeveloperWindows != 0);
+
+	InTabManager->RegisterTabSpawner(DebugSpreadsheetTabID, FOnSpawnTab::CreateSP(this, &FNiagaraSystemToolkitModeBase::SpawnTab_DebugSpreadsheet))
+		.SetDisplayName(LOCTEXT("DebugSpreadsheet", "Attribute Spreadsheet"))
+		.SetGroup(WorkspaceMenuCategory.ToSharedRef())
+		.SetIcon(FSlateIcon(FNiagaraEditorStyle::Get().GetStyleSetName(), "Tab.Spreadsheet"));
+
+	InTabManager->RegisterTabSpawner(PreviewSettingsTabId, FOnSpawnTab::CreateSP(this, &FNiagaraSystemToolkitModeBase::SpawnTab_PreviewSettings))
+		.SetDisplayName(LOCTEXT("PreviewSceneSettingsTab", "Preview Scene Settings"))
+		.SetGroup(WorkspaceMenuCategory.ToSharedRef())
+		.SetIcon(FSlateIcon(FNiagaraEditorStyle::Get().GetStyleSetName(), "Tab.Settings"));
+
+	InTabManager->RegisterTabSpawner(GeneratedCodeTabID, FOnSpawnTab::CreateSP(this, &FNiagaraSystemToolkitModeBase::SpawnTab_GeneratedCode))
+		.SetDisplayName(LOCTEXT("GeneratedCode", "Generated Code"))
+		.SetGroup(WorkspaceMenuCategory.ToSharedRef())
+		.SetIcon(FSlateIcon(FNiagaraEditorStyle::Get().GetStyleSetName(), "Tab.GeneratedCode"));
+
+	InTabManager->RegisterTabSpawner(MessageLogTabID, FOnSpawnTab::CreateSP(this, &FNiagaraSystemToolkitModeBase::SpawnTab_MessageLog))
+		.SetDisplayName(LOCTEXT("NiagaraMessageLog", "Niagara Log"))
+		.SetGroup(WorkspaceMenuCategory.ToSharedRef())
+		.SetIcon(FSlateIcon(FNiagaraEditorStyle::Get().GetStyleSetName(), "Tab.Log"));
+
+	InTabManager->RegisterTabSpawner(SystemOverviewTabID, FOnSpawnTab::CreateSP(this, &FNiagaraSystemToolkitModeBase::SpawnTab_SystemOverview))
+		.SetDisplayName(LOCTEXT("SystemOverviewTabName", "System Overview"))
+		.SetGroup(WorkspaceMenuCategory.ToSharedRef())
+		.SetIcon(FSlateIcon(FNiagaraEditorStyle::Get().GetStyleSetName(), "Tab.SystemOverview"));
+
+	InTabManager->RegisterTabSpawner(ScratchPadTabID, FOnSpawnTab::CreateSP(this, &FNiagaraSystemToolkitModeBase::SpawnTab_ScratchPad))
+		.SetDisplayName(LOCTEXT("ScratchPadTabName", "Scratch Pad"))
+		.SetGroup(WorkspaceMenuCategory.ToSharedRef())
+		.SetIcon(FSlateIcon(FNiagaraEditorStyle::Get().GetStyleSetName(), "Tab.ScratchPad"));
+
+	InTabManager->RegisterTabSpawner(ScriptStatsTabID, FOnSpawnTab::CreateSP(this, &FNiagaraSystemToolkitModeBase::SpawnTab_ScriptStats))
+		.SetDisplayName(LOCTEXT("NiagaraScriptsStatsTab", "Script Stats"))
+		.SetGroup(WorkspaceMenuCategory.ToSharedRef())
+		.SetIcon(FSlateIcon(FNiagaraEditorStyle::Get().GetStyleSetName(), "Tab.ScriptStats"));
+
+	if (GetDefault<UNiagaraEditorSettings>()->bEnableBaker)
+	{
+		InTabManager->RegisterTabSpawner(BakerTabID, FOnSpawnTab::CreateSP(this, &FNiagaraSystemToolkitModeBase::SpawnTab_Baker))
+			.SetDisplayName(LOCTEXT("NiagaraBakerTab", "Baker"))
+			.SetGroup(WorkspaceMenuCategory.ToSharedRef());
+	}
+}
+
+class SNiagaraSelectedEmitterGraph : public SCompoundWidget
+{
+public:
+	SLATE_BEGIN_ARGS(SNiagaraSelectedEmitterGraph)
+	{}
+	SLATE_END_ARGS();
+
+	void Construct(const FArguments& InArgs, TSharedRef<FNiagaraSystemViewModel> InSystemViewModel)
+	{
+		SystemViewModel = InSystemViewModel;
+		SystemViewModel->GetSelectionViewModel()->OnEmitterHandleIdSelectionChanged().AddSP(this, &SNiagaraSelectedEmitterGraph::SystemSelectionChanged);
+		ChildSlot
+		[
+			SAssignNew(GraphWidgetContainer, SBox)
+		];
+		UpdateGraphWidget();
+	}
+
+	~SNiagaraSelectedEmitterGraph()
+	{
+		if (SystemViewModel.IsValid() && SystemViewModel->GetSelectionViewModel())
+		{
+			SystemViewModel->GetSelectionViewModel()->OnEmitterHandleIdSelectionChanged().RemoveAll(this);
+		}
+	}
+
+private:
+	void SystemSelectionChanged()
+	{
+		UpdateGraphWidget();
+	}
+
+	void UpdateGraphWidget()
+	{
+		TArray<FGuid> SelectedEmitterHandleIds = SystemViewModel->GetSelectionViewModel()->GetSelectedEmitterHandleIds();
+		if (SelectedEmitterHandleIds.Num() == 1)
+		{
+			TSharedPtr<FNiagaraEmitterHandleViewModel> SelectedEmitterHandle = SystemViewModel->GetEmitterHandleViewModelById(SelectedEmitterHandleIds[0]);
+			TSharedRef<SWidget> EmitterWidget = 
+				SNew(SSplitter)
+				+ SSplitter::Slot()
+				.Value(.25f)
+				[
+					SNew(SNiagaraSelectedObjectsDetails, SelectedEmitterHandle->GetEmitterViewModel()->GetSharedScriptViewModel()->GetGraphViewModel()->GetNodeSelection())
+				]
+				+ SSplitter::Slot()
+				.Value(.75f)
+				[
+					SNew(SNiagaraScriptGraph, SelectedEmitterHandle->GetEmitterViewModel()->GetSharedScriptViewModel()->GetGraphViewModel())
+				];
+
+			UNiagaraEmitter* LastMergedEmitter = SelectedEmitterHandle->GetEmitterViewModel()->GetEmitter()->GetParentAtLastMerge();
+			if (LastMergedEmitter != nullptr)
+			{
+				UNiagaraScriptSource* LastMergedScriptSource = CastChecked<UNiagaraScriptSource>(LastMergedEmitter->GraphSource);
+				bool bIsForDataProcessingOnly = false;
+				TSharedRef<FNiagaraScriptGraphViewModel> LastMergedScriptGraphViewModel = MakeShared<FNiagaraScriptGraphViewModel>(FText(), bIsForDataProcessingOnly);
+				LastMergedScriptGraphViewModel->SetScriptSource(LastMergedScriptSource);
+				TSharedRef<SWidget> LastMergedEmitterWidget = 
+					SNew(SSplitter)
+					+ SSplitter::Slot()
+					.Value(.25f)
+					[
+						SNew(SNiagaraSelectedObjectsDetails, LastMergedScriptGraphViewModel->GetNodeSelection())
+					]
+					+ SSplitter::Slot()
+					.Value(.75f)
+					[
+						SNew(SNiagaraScriptGraph, LastMergedScriptGraphViewModel)
+					];
+
+				GraphWidgetContainer->SetContent
+				(
+					SNew(SSplitter)
+					.Orientation(Orient_Vertical)
+					+ SSplitter::Slot()
+					[
+						SNew(SVerticalBox)
+						+ SVerticalBox::Slot()
+						.AutoHeight()
+						[
+							SNew(STextBlock)
+							.Text(FText::FromString(TEXT("Emitter")))
+						]
+						+ SVerticalBox::Slot()
+						[
+							EmitterWidget
+						]
+					]
+					+ SSplitter::Slot()
+					[
+						SNew(SVerticalBox)
+						+ SVerticalBox::Slot()
+						.AutoHeight()
+						[
+							SNew(STextBlock)
+							.Text(FText::FromString(TEXT("Last Merged Emitter")))
+						]
+						+ SVerticalBox::Slot()
+						[
+							LastMergedEmitterWidget
+						]
+					]
+				);
+			}
+			else
+			{
+				GraphWidgetContainer->SetContent(EmitterWidget);
+			}
+		}
+		else
+		{
+			GraphWidgetContainer->SetContent(SNullWidget::NullWidget);
+		}
+	}
+
+private:
+	TSharedPtr<FNiagaraSystemViewModel> SystemViewModel;
+	TSharedPtr<SBox> GraphWidgetContainer;
+};
+
+TSharedRef<SDockTab> FNiagaraSystemToolkitModeBase::SpawnTab_Viewport(const FSpawnTabArgs& Args)
+{
+	check(Args.GetTabId().TabType == ViewportTabID);
+
+	TSharedPtr<FNiagaraSystemToolkit> Toolkit = StaticCastSharedPtr<FNiagaraSystemToolkit>(SystemToolkit.Pin());
+
+	if(!Toolkit->Viewport.IsValid())
+	{
+		Toolkit->Viewport = SNew(SNiagaraSystemViewport)
+			.OnThumbnailCaptured(Toolkit.Get(), &FNiagaraSystemToolkit::OnThumbnailCaptured)
+			.Sequencer(Toolkit->GetSystemViewModel()->GetSequencer());
+	}
+
+	TSharedRef<SDockTab> SpawnedTab =
+		SNew(SDockTab)
+		[
+			Toolkit->Viewport.ToSharedRef()
+		];
+
+	Toolkit->Viewport->SetPreviewComponent(Toolkit->GetSystemViewModel()->GetPreviewComponent());
+	Toolkit->Viewport->OnAddedToTab(SpawnedTab);
+
+	return SpawnedTab;
+}
+
+TSharedRef<SDockTab> FNiagaraSystemToolkitModeBase::SpawnTab_PreviewSettings(const FSpawnTabArgs& Args)
+{
+	check(Args.GetTabId() == PreviewSettingsTabId);
+
+	TSharedRef<SWidget> InWidget = SNullWidget::NullWidget;
+	if (SystemToolkit.Pin()->Viewport.IsValid())
+	{
+		FAdvancedPreviewSceneModule& AdvancedPreviewSceneModule = FModuleManager::LoadModuleChecked<FAdvancedPreviewSceneModule>("AdvancedPreviewScene");
+		InWidget = AdvancedPreviewSceneModule.CreateAdvancedPreviewSceneSettingsWidget(SystemToolkit.Pin()->Viewport->GetPreviewScene());
+	}
+
+	TSharedRef<SDockTab> SpawnedTab = SNew(SDockTab)
+		.Label(LOCTEXT("PreviewSceneSettingsTab", "Preview Scene Settings"))
+		[
+			InWidget
+		];
+
+	return SpawnedTab;
+}
+
+
+TSharedRef<SDockTab> FNiagaraSystemToolkitModeBase::SpawnTab_CurveEd(const FSpawnTabArgs& Args)
+{
+	check(Args.GetTabId().TabType == CurveEditorTabID);
+
+	TSharedPtr<FNiagaraSystemToolkit> Toolkit = StaticCastSharedPtr<FNiagaraSystemToolkit>(SystemToolkit.Pin());
+
+	TSharedRef<SDockTab> SpawnedTab =
+		SNew(SDockTab)
+		[
+			FNiagaraEditorModule::Get().GetWidgetProvider()->CreateCurveOverview(SystemToolkit.Pin()->GetSystemViewModel().ToSharedRef())
+		];
+
+	return SpawnedTab;
+}
+
+
+TSharedRef<SDockTab> FNiagaraSystemToolkitModeBase::SpawnTab_Sequencer(const FSpawnTabArgs& Args)
+{
+	check(Args.GetTabId().TabType == SequencerTabID);
+
+	TSharedPtr<FNiagaraSystemToolkit> Toolkit = StaticCastSharedPtr<FNiagaraSystemToolkit>(SystemToolkit.Pin());
+
+	TSharedRef<SDockTab> SpawnedTab =
+		SNew(SDockTab)
+		[
+			Toolkit->GetSystemViewModel()->GetSequencer()->GetSequencerWidget()
+		];
+
+	return SpawnedTab;
+}
+
+
+TSharedRef<SDockTab> FNiagaraSystemToolkitModeBase::SpawnTab_SystemScript(const FSpawnTabArgs& Args)
+{
+	check(Args.GetTabId().TabType == SystemScriptTabID);
+
+	TSharedRef<SDockTab> SpawnedTab =
+		SNew(SDockTab)
+		[
+			SNew(SNiagaraSystemScript, SystemToolkit.Pin()->GetSystemViewModel().ToSharedRef())
+		];
+
+	return SpawnedTab;
+}
+
+TSharedRef<SDockTab> FNiagaraSystemToolkitModeBase::SpawnTab_SystemParameters(const FSpawnTabArgs& Args)
+{
+	check(Args.GetTabId().TabType == SystemParametersTabID);
+
+
+	TArray<TSharedRef<FNiagaraObjectSelection>> ObjectSelections;
+	ObjectSelections.Add(SystemToolkit.Pin()->ObjectSelectionForParameterMapView.ToSharedRef());
+
+	TSharedRef<SDockTab> SpawnedTab =
+		SNew(SDockTab)
+		[
+			SAssignNew(SystemToolkit.Pin()->ParameterPanel, SNiagaraParameterPanel, SystemToolkit.Pin()->ParameterPanelViewModel, SystemToolkit.Pin()->GetToolkitCommands())
+			.ShowParameterSynchronizingWithLibraryIconExternallyReferenced(false)
+		];
+	SystemToolkit.Pin()->RefreshParameters();
+
+	return SpawnedTab;
+}
+
+//@todo(ng) cleanup
+TSharedRef<SDockTab> FNiagaraSystemToolkitModeBase::SpawnTab_SystemParameters2(const FSpawnTabArgs& Args)
+{
+	check(Args.GetTabId().TabType == SystemParametersTabID2);
+
+	TArray<TSharedRef<FNiagaraObjectSelection>> ObjectSelections;
+	ObjectSelections.Add(SystemToolkit.Pin()->ObjectSelectionForParameterMapView.ToSharedRef());
+
+	TSharedRef<SDockTab> SpawnedTab =
+		SNew(SDockTab)
+		[
+			SAssignNew(SystemToolkit.Pin()->ParameterMapView, SNiagaraParameterMapView, ObjectSelections, SNiagaraParameterMapView::EToolkitType::SYSTEM, SystemToolkit.Pin()->GetToolkitCommands())
+		];
+	SystemToolkit.Pin()->RefreshParameters();
+
+	return SpawnedTab;
+}
+
+TSharedRef<SDockTab> FNiagaraSystemToolkitModeBase::SpawnTab_SystemParameterDefinitions(const FSpawnTabArgs& Args)
+{
+	checkf(Args.GetTabId().TabType == SystemParameterDefinitionsTabID, TEXT("Wrong tab ID in NiagaraScriptToolkit"));
+
+	TSharedRef<SDockTab> SpawnedTab =
+		SNew(SDockTab)
+		[
+			SNew(SNiagaraParameterDefinitionsPanel, SystemToolkit.Pin()->ParameterDefinitionsPanelViewModel, SystemToolkit.Pin()->GetToolkitCommands())
+		];
+
+	return SpawnedTab;
+}
+
+TSharedRef<SDockTab> FNiagaraSystemToolkitModeBase::SpawnTab_SelectedEmitterStack(const FSpawnTabArgs& Args)
+{
+	check(Args.GetTabId().TabType == SelectedEmitterStackTabID);
+
+	FNiagaraEditorModule& NiagaraEditorModule = FModuleManager::LoadModuleChecked<FNiagaraEditorModule>("NiagaraEditor");
+	TSharedRef<SDockTab> SpawnedTab = SNew(SDockTab)
+		.Label(LOCTEXT("SystemOverviewSelection", "Selection"))
+		[
+			NiagaraEditorModule.GetWidgetProvider()->CreateStackView(*SystemToolkit.Pin()->GetSystemViewModel()->GetSelectionViewModel()->GetSelectionStackViewModel())
+		];
+
+	SDockTab::FOnTabClosedCallback TabClosedCallback = SDockTab::FOnTabClosedCallback::CreateLambda([=](TSharedRef<SDockTab> DockTab)
+	{
+		SystemToolkit.Pin()->GetSystemViewModel()->GetSelectionViewModel()->GetSelectionStackViewModel()->ResetSearchText();
+	});
+	
+	SpawnedTab->SetOnTabClosed(TabClosedCallback);
+	return SpawnedTab;
+}
+
+TSharedRef<SDockTab> FNiagaraSystemToolkitModeBase::SpawnTab_SelectedEmitterGraph(const FSpawnTabArgs& Args)
+{
+	check(Args.GetTabId().TabType == SelectedEmitterGraphTabID);
+
+	TSharedRef<SDockTab> SpawnedTab =
+		SNew(SDockTab)
+		[
+			SNew(SNiagaraSelectedEmitterGraph, SystemToolkit.Pin()->SystemViewModel.ToSharedRef())
+		];
+
+	return SpawnedTab;
+}
+
+TSharedRef<SDockTab> FNiagaraSystemToolkitModeBase::SpawnTab_DebugSpreadsheet(const FSpawnTabArgs& Args)
+{
+	check(Args.GetTabId().TabType == DebugSpreadsheetTabID);
+
+	TSharedRef<SDockTab> SpawnedTab =
+		SNew(SDockTab)
+		[
+			SNew(SNiagaraSpreadsheetView, SystemToolkit.Pin()->GetSystemViewModel().ToSharedRef())
+		];
+
+	return SpawnedTab;
+}
+
+TSharedRef<SDockTab> FNiagaraSystemToolkitModeBase::SpawnTab_GeneratedCode(const FSpawnTabArgs& Args)
+{
+	check(Args.GetTabId().TabType == GeneratedCodeTabID);
+
+	TSharedRef<SDockTab> SpawnedTab = SNew(SDockTab);
+	SpawnedTab->SetContent(SNew(SNiagaraGeneratedCodeView, SystemToolkit.Pin()->GetSystemViewModel().ToSharedRef(), SpawnedTab));
+	return SpawnedTab;
+}
+
+TSharedRef<SDockTab> FNiagaraSystemToolkitModeBase::SpawnTab_MessageLog(const FSpawnTabArgs& Args)
+{
+	check(Args.GetTabId().TabType == MessageLogTabID);
+
+	TSharedRef<SDockTab> SpawnedTab = SNew(SDockTab)
+		.Label(LOCTEXT("NiagaraMessageLogTitle", "Niagara Log"))
+		[
+			SNew(SBox)
+			.AddMetaData<FTagMetaData>(FTagMetaData(TEXT("NiagaraLog")))
+			[
+				SystemToolkit.Pin()->NiagaraMessageLog.ToSharedRef()
+			]
+		];
+
+	return SpawnedTab;
+}
+
+TSharedRef<SDockTab> FNiagaraSystemToolkitModeBase::SpawnTab_SystemOverview(const FSpawnTabArgs& Args)
+{
+	if(!SystemToolkit.Pin()->GetSystemOverview().IsValid())
+	{
+		SystemToolkit.Pin()->SetSystemOverview(FNiagaraEditorModule::Get().GetWidgetProvider()->CreateSystemOverview(SystemToolkit.Pin()->GetSystemViewModel().ToSharedRef()));
+	}
+	
+	TSharedRef<SDockTab> SpawnedTab = SNew(SDockTab)
+		.Label(LOCTEXT("SystemOverviewTabLabel", "System Overview"))
+		[
+			SystemToolkit.Pin()->GetSystemOverview().ToSharedRef()
+		];
+
+	SpawnedTab->SetOnTabClosed(SDockTab::FOnTabClosedCallback::CreateLambda([this](TSharedRef<SDockTab>)
+	{
+		SystemToolkit.Pin()->SetSystemOverview(nullptr);
+	}));
+	
+	return SpawnedTab;
+}
+
+TSharedRef<SDockTab> FNiagaraSystemToolkitModeBase::SpawnTab_ScratchPad(const FSpawnTabArgs& Args)
+{
+	if(!SystemToolkit.Pin()->GetScriptScratchpad().IsValid())
+	{
+		SystemToolkit.Pin()->SetScriptScratchpad(FNiagaraEditorModule::Get().GetWidgetProvider()->CreateScriptScratchPad(*SystemToolkit.Pin()->GetSystemViewModel()->GetScriptScratchPadViewModel()));
+	}
+	
+	TSharedRef<SDockTab> SpawnedTab = SNew(SDockTab)
+		.Label(LOCTEXT("ScratchPadTabLabel", "Scratch Pad"))
+		[
+			SystemToolkit.Pin()->GetScriptScratchpad().ToSharedRef()
+		];
+
+	SpawnedTab->SetOnTabClosed(SDockTab::FOnTabClosedCallback::CreateLambda([this](TSharedRef<SDockTab>)
+	{
+		SystemToolkit.Pin()->SetScriptScratchpad(nullptr);
+	}));
+
+	return SpawnedTab;
+}
+
+TSharedRef<SDockTab> FNiagaraSystemToolkitModeBase::SpawnTab_ScriptStats(const FSpawnTabArgs& Args)
+{
+	check(Args.GetTabId().TabType == ScriptStatsTabID);
+
+	TSharedRef<SDockTab> SpawnedTab = SNew(SDockTab)
+		.Label(LOCTEXT("NiagaraScriptStatsTitle", "Script Stats"))
+		[
+			SNew(SBox)
+			.AddMetaData<FTagMetaData>(FTagMetaData(TEXT("ScriptStats")))
+			[
+				SystemToolkit.Pin()->ScriptStats->GetWidget().ToSharedRef()
+			]
+		];
+
+	return SpawnedTab;
+}
+
+TSharedRef<SDockTab> FNiagaraSystemToolkitModeBase::SpawnTab_Baker(const FSpawnTabArgs& Args)
+{
+	check(Args.GetTabId().TabType == BakerTabID);
+
+	TSharedRef<SDockTab> SpawnedTab = SNew(SDockTab)
+		.Label(LOCTEXT("NiagaraBakerTitle", "Baker"))
+		[
+			SNew(SBox)
+			.AddMetaData<FTagMetaData>(FTagMetaData(TEXT("Baker")))
+			[
+				SystemToolkit.Pin()->BakerViewModel->GetWidget().ToSharedRef()
+			]
+		];
+
+	return SpawnedTab;
+}
+
+#undef LOCTEXT_NAMESPACE
\ No newline at end of file
diff --git a/Engine/Plugins/FX/Niagara/Source/NiagaraEditor/Private/Toolkits/SystemToolkitModes/NiagaraSystemToolkitModeBase.h b/Engine/Plugins/FX/Niagara/Source/NiagaraEditor/Private/Toolkits/SystemToolkitModes/NiagaraSystemToolkitModeBase.h
new file mode 100644
--- /dev/null
+++ b/Engine/Plugins/FX/Niagara/Source/NiagaraEditor/Private/Toolkits/SystemToolkitModes/NiagaraSystemToolkitModeBase.h
@@ -0,0 +1,64 @@
+// Copyright Epic Games, Inc. All Rights Reserved.
+
+#pragma once
+
+#include "NiagaraBakerViewModel.h"
+#include "NiagaraScriptStatsViewModel.h"
+#include "WorkflowOrientedApp/ApplicationMode.h"
+#include "WorkflowOrientedApp/WorkflowCentricApplication.h"
+#include "Framework/MultiBox/MultiBoxBuilder.h"
+#include "Widgets/Docking/SDockTab.h"
+
+class FNiagaraSystemToolkit;
+
+class FNiagaraSystemToolkitModeBase : public FApplicationMode
+{
+public:
+	FNiagaraSystemToolkitModeBase(FName InModeName, TWeakPtr<FNiagaraSystemToolkit> InSystemToolkit) : FApplicationMode(InModeName), SystemToolkit(InSystemToolkit)
+	{ }
+	
+	virtual void RegisterTabFactories(TSharedPtr<FTabManager> InTabManager) override;
+
+private:
+	TSharedRef<SDockTab> SpawnTab_Viewport(const FSpawnTabArgs& Args);
+	TSharedRef<SDockTab> SpawnTab_CurveEd(const FSpawnTabArgs& Args);
+	TSharedRef<SDockTab> SpawnTab_Sequencer(const FSpawnTabArgs& Args);
+	TSharedRef<SDockTab> SpawnTab_SystemScript(const FSpawnTabArgs& Args);
+	TSharedRef<SDockTab> SpawnTab_SystemParameters(const FSpawnTabArgs& Args);
+	TSharedRef<SDockTab> SpawnTab_SystemParameters2(const FSpawnTabArgs& Args); //@todo(ng) cleanup
+	TSharedRef<SDockTab> SpawnTab_SystemParameterDefinitions(const FSpawnTabArgs& Args);
+	TSharedRef<SDockTab> SpawnTab_SelectedEmitterStack(const FSpawnTabArgs& Args);
+	TSharedRef<SDockTab> SpawnTab_SelectedEmitterGraph(const FSpawnTabArgs& Args);
+	TSharedRef<SDockTab> SpawnTab_DebugSpreadsheet(const FSpawnTabArgs& Args);
+	TSharedRef<SDockTab> SpawnTab_PreviewSettings(const FSpawnTabArgs& Args);
+	TSharedRef<SDockTab> SpawnTab_GeneratedCode(const FSpawnTabArgs& Args);
+	TSharedRef<SDockTab> SpawnTab_MessageLog(const FSpawnTabArgs& Args);
+	TSharedRef<SDockTab> SpawnTab_SystemOverview(const FSpawnTabArgs& Args);
+	TSharedRef<SDockTab> SpawnTab_ScratchPad(const FSpawnTabArgs& Args);
+	TSharedRef<SDockTab> SpawnTab_ScriptStats(const FSpawnTabArgs& Args);
+	TSharedRef<SDockTab> SpawnTab_Baker(const FSpawnTabArgs& Args);
+
+protected:
+	TWeakPtr<FNiagaraSystemToolkit> SystemToolkit;
+
+public:
+	static const FName ViewportTabID;
+	static const FName CurveEditorTabID;
+	static const FName SequencerTabID;
+	static const FName SystemScriptTabID;
+	static const FName SystemDetailsTabID;
+	static const FName SystemParametersTabID;
+	static const FName SystemParametersTabID2;
+	static const FName SystemParameterDefinitionsTabID;
+	static const FName SelectedEmitterStackTabID;
+	static const FName SelectedEmitterGraphTabID;
+	static const FName DebugSpreadsheetTabID;
+	static const FName PreviewSettingsTabId;
+	static const FName GeneratedCodeTabID;
+	static const FName MessageLogTabID;
+	static const FName SystemOverviewTabID;
+	static const FName ScratchPadTabID;
+	static const FName ScriptStatsTabID;
+	static const FName BakerTabID;
+};
+
diff --git a/Engine/Plugins/FX/Niagara/Source/NiagaraEditor/Private/Toolkits/SystemToolkitModes/NiagaraSystemToolkitMode_Default.cpp b/Engine/Plugins/FX/Niagara/Source/NiagaraEditor/Private/Toolkits/SystemToolkitModes/NiagaraSystemToolkitMode_Default.cpp
new file mode 100644
--- /dev/null
+++ b/Engine/Plugins/FX/Niagara/Source/NiagaraEditor/Private/Toolkits/SystemToolkitModes/NiagaraSystemToolkitMode_Default.cpp
@@ -0,0 +1,286 @@
+// Copyright Epic Games, Inc. All Rights Reserved.
+
+#include "NiagaraSystemToolkitMode_Default.h"
+
+#include "NiagaraEditorCommands.h"
+#include "NiagaraSystemToolkit.h"
+#include "NiagaraDebuggerCommon.h"
+#include "NiagaraEditorModule.h"
+#include "NiagaraEditorSettings.h"
+
+#define LOCTEXT_NAMESPACE "NiagaraSystemToolkitMode_Default"
+
+FNiagaraSystemToolkitMode_Default::FNiagaraSystemToolkitMode_Default(TWeakPtr<FNiagaraSystemToolkit> InSystemToolkit) : FNiagaraSystemToolkitModeBase(FNiagaraSystemToolkit::DefaultModeName, InSystemToolkit)
+{
+	TabLayout = FTabManager::NewLayout("Standalone_Niagara_System_Layout_v24")
+		->AddArea
+		(
+			FTabManager::NewPrimaryArea()->SetOrientation(Orient_Vertical)
+			->Split
+			(
+				FTabManager::NewSplitter()->SetOrientation(Orient_Horizontal)
+				->Split
+				(
+					// Top Level Left
+					FTabManager::NewSplitter()->SetOrientation(Orient_Vertical)
+					->SetSizeCoefficient(.75f)
+					->Split
+					(
+						// Inner Left Top
+						FTabManager::NewSplitter()->SetOrientation(Orient_Horizontal)
+						->SetSizeCoefficient(0.75f)
+						->Split
+						(
+							FTabManager::NewStack()
+							->SetSizeCoefficient(.25f)
+							->AddTab(ViewportTabID, ETabState::OpenedTab)
+						)
+						->Split
+						(
+							FTabManager::NewStack()
+							->SetSizeCoefficient(0.15f)
+							->AddTab(SystemParametersTabID, ETabState::OpenedTab)
+						)
+						->Split
+						(
+							FTabManager::NewStack()
+							->SetSizeCoefficient(0.6f)
+							->AddTab(SystemOverviewTabID, ETabState::OpenedTab)
+							->AddTab(ScratchPadTabID, ETabState::OpenedTab)
+							->AddTab(BakerTabID, ETabState::ClosedTab)
+							->SetForegroundTab(SystemOverviewTabID)
+						)
+					)
+					->Split
+					(
+						// Inner Left Bottom
+						FTabManager::NewStack()
+						->SetSizeCoefficient(0.25f)
+						->AddTab(CurveEditorTabID, ETabState::OpenedTab)
+						->AddTab(MessageLogTabID, ETabState::OpenedTab)
+						->AddTab(SequencerTabID, ETabState::OpenedTab)
+						->AddTab(ScriptStatsTabID, ETabState::ClosedTab)
+					)
+				)
+				->Split
+				(
+					// Top Level Right
+					FTabManager::NewStack()
+					->SetSizeCoefficient(0.25f)
+					->AddTab(SelectedEmitterStackTabID, ETabState::OpenedTab)
+					->AddTab(SelectedEmitterGraphTabID, ETabState::ClosedTab)
+					->AddTab(SystemScriptTabID, ETabState::ClosedTab)
+					->AddTab(SystemDetailsTabID, ETabState::ClosedTab)
+					->AddTab(DebugSpreadsheetTabID, ETabState::ClosedTab)
+					->AddTab(PreviewSettingsTabId, ETabState::ClosedTab)
+					->AddTab(GeneratedCodeTabID, ETabState::ClosedTab)
+				)
+			)
+		);
+
+	ExtendToolbar();
+}
+
+void FNiagaraSystemToolkitMode_Default::ExtendToolbar()
+{
+	struct Local
+	{
+		static TSharedRef<SWidget> FillSimulationOptionsMenu(FNiagaraSystemToolkit* Toolkit)
+		{
+			FMenuBuilder MenuBuilder(true, Toolkit->GetToolkitCommands());
+			MenuBuilder.AddMenuEntry(FNiagaraEditorCommands::Get().ToggleAutoPlay);
+			MenuBuilder.AddMenuEntry(FNiagaraEditorCommands::Get().ToggleResetSimulationOnChange);
+			MenuBuilder.AddMenuEntry(FNiagaraEditorCommands::Get().ToggleResimulateOnChangeWhilePaused);
+			MenuBuilder.AddMenuEntry(FNiagaraEditorCommands::Get().ToggleResetDependentSystems);
+			return MenuBuilder.MakeWidget();
+		}
+
+		static TSharedRef<SWidget> GenerateBakerMenu(FNiagaraSystemToolkit* Toolkit)
+		{
+			FMenuBuilder MenuBuilder(true, Toolkit->GetToolkitCommands());
+
+			MenuBuilder.AddMenuEntry(
+				LOCTEXT("BakerTab", "Open Baker Tab"),
+				LOCTEXT("BakerTabTooltip", "Opens the flip book tab."),
+				FSlateIcon(),
+				FUIAction(FExecuteAction::CreateLambda([TabManager=Toolkit->GetTabManager()]() { TabManager->TryInvokeTab(FNiagaraSystemToolkitModeBase::BakerTabID); }))
+			);
+
+			MenuBuilder.AddMenuEntry(
+				LOCTEXT("RenderBaker", "Render Baker"),
+				LOCTEXT("RenderBakerTooltip", "Renders the Baker using the current settings."),
+				FSlateIcon(),
+				FUIAction(FExecuteAction::CreateSP(Toolkit, &FNiagaraSystemToolkit::RenderBaker))
+			);
+
+			return MenuBuilder.MakeWidget();
+		}
+
+		static TSharedRef<SWidget> FillDebugOptionsMenu(FNiagaraSystemToolkit* Toolkit)
+		{
+			FMenuBuilder MenuBuilder(true, Toolkit->GetToolkitCommands());
+
+#if WITH_NIAGARA_DEBUGGER
+			MenuBuilder.AddMenuEntry(FNiagaraEditorCommands::Get().OpenDebugHUD);
+			MenuBuilder.AddMenuEntry(FNiagaraEditorCommands::Get().OpenDebugOutliner);
+#endif
+			MenuBuilder.AddMenuEntry(FNiagaraEditorCommands::Get().OpenAttributeSpreadsheet);
+			return MenuBuilder.MakeWidget();
+		}
+
+		static void FillToolbar(FToolBarBuilder& ToolbarBuilder, FNiagaraSystemToolkit* Toolkit)
+		{
+			ToolbarBuilder.BeginSection("Apply");
+			{
+				if (Toolkit->HasEmitter())
+				{
+					ToolbarBuilder.AddToolBarButton(FNiagaraEditorCommands::Get().Apply,
+						NAME_None, TAttribute<FText>(), TAttribute<FText>(),
+						FSlateIcon(FNiagaraEditorStyle::Get().GetStyleSetName(), "NiagaraEditor.Apply"),
+						FName(TEXT("ApplyNiagaraEmitter")));
+				}
+				ToolbarBuilder.AddToolBarButton(FNiagaraEditorCommands::Get().ApplyScratchPadChanges,
+					NAME_None, TAttribute<FText>(), TAttribute<FText>(),
+					FSlateIcon(FNiagaraEditorStyle::Get().GetStyleSetName(), "NiagaraEditor.ApplyScratchPadChanges"),
+					FName(TEXT("ApplyScratchPadChanges")));
+			}
+			ToolbarBuilder.EndSection();
+			ToolbarBuilder.BeginSection("Compile");
+			{
+				ToolbarBuilder.AddToolBarButton(FNiagaraEditorCommands::Get().Compile,
+					NAME_None,
+					TAttribute<FText>(),
+					TAttribute<FText>(Toolkit, &FNiagaraSystemToolkit::GetCompileStatusTooltip),
+					TAttribute<FSlateIcon>(Toolkit, &FNiagaraSystemToolkit::GetCompileStatusImage),
+					FName(TEXT("CompileNiagaraSystem")));
+				ToolbarBuilder.AddComboButton(
+					FUIAction(),
+					FOnGetContent::CreateRaw(Toolkit, &FNiagaraSystemToolkit::GenerateCompileMenuContent),
+					LOCTEXT("BuildCombo_Label", "Auto-Compile Options"),
+					LOCTEXT("BuildComboToolTip", "Auto-Compile options menu"),
+					FSlateIcon(FEditorStyle::GetStyleSetName(), "LevelEditor.Build"),
+					true);
+			}
+			ToolbarBuilder.EndSection();
+
+			ToolbarBuilder.BeginSection("NiagaraThumbnail");
+			{
+				ToolbarBuilder.AddToolBarButton(FNiagaraEditorCommands::Get().SaveThumbnailImage, NAME_None,
+					LOCTEXT("GenerateThumbnail", "Thumbnail"),
+					LOCTEXT("GenerateThumbnailTooltip","Generate a thumbnail image."),
+					FSlateIcon(FEditorStyle::GetStyleSetName(), "Cascade.SaveThumbnailImage"));
+			}
+			ToolbarBuilder.EndSection();
+
+			ToolbarBuilder.BeginSection("NiagaraPreviewOptions");
+			{
+				ToolbarBuilder.AddToolBarButton(FNiagaraEditorCommands::Get().ToggleBounds, NAME_None,
+					LOCTEXT("ShowBounds", "Bounds"),
+					LOCTEXT("ShowBoundsTooltip", "Show the bounds for the scene."),
+					FSlateIcon(FEditorStyle::GetStyleSetName(), "Cascade.ToggleBounds"));
+				ToolbarBuilder.AddComboButton(
+					FUIAction(),
+					FOnGetContent::CreateRaw(Toolkit, &FNiagaraSystemToolkit::GenerateBoundsMenuContent, Toolkit->GetToolkitCommands()),
+					LOCTEXT("BoundsMenuCombo_Label", "Bounds Options"),
+					LOCTEXT("BoundsMenuCombo_ToolTip", "Bounds options"),
+					FSlateIcon(FEditorStyle::GetStyleSetName(), "Cascade.ToggleBounds"),
+					true
+				);
+			}
+			ToolbarBuilder.EndSection();
+			
+#if STATS
+			ToolbarBuilder.BeginSection("NiagaraStatisticsOptions");
+			{
+				ToolbarBuilder.AddToolBarButton(FNiagaraEditorCommands::Get().ToggleStatPerformance, NAME_None,
+                    LOCTEXT("NiagaraShowPerformance", "Performance"),
+                    LOCTEXT("NiagaraShowPerformanceTooltip", "Show runtime performance for particle scripts."),
+                    FSlateIcon(FEditorStyle::GetStyleSetName(), "MaterialEditor.ToggleMaterialStats"));
+				ToolbarBuilder.AddComboButton(
+                    FUIAction(),
+                    FOnGetContent::CreateRaw(Toolkit, &FNiagaraSystemToolkit::GenerateStatConfigMenuContent, Toolkit->GetToolkitCommands()),
+                    FText(),
+                    LOCTEXT("NiagaraShowPerformanceCombo_ToolTip", "Runtime performance options"),
+                    FSlateIcon(FEditorStyle::GetStyleSetName(), "MaterialEditor.ToggleMaterialStats"),
+                    true);
+				ToolbarBuilder.AddComboButton(
+					FUIAction(),
+					FOnGetContent::CreateStatic(Local::FillDebugOptionsMenu, Toolkit),
+					LOCTEXT("DebugOptions", "Debug"),
+					LOCTEXT("DebugOptionsTooltip", "Debug options"),
+					FSlateIcon(FNiagaraEditorStyle::Get().GetStyleSetName(), "Tab.Debugger")
+				);
+			}
+			ToolbarBuilder.EndSection();
+#endif
+			
+			ToolbarBuilder.BeginSection("PlaybackOptions");
+			{
+				ToolbarBuilder.AddComboButton(
+					FUIAction(),
+					FOnGetContent::CreateStatic(Local::FillSimulationOptionsMenu, Toolkit),
+					LOCTEXT("SimulationOptions", "Simulation"),
+					LOCTEXT("SimulationOptionsTooltip", "Simulation options"),
+					FSlateIcon(FNiagaraEditorStyle::Get().GetStyleSetName(), "NiagaraEditor.SimulationOptions")
+				);
+			}
+			ToolbarBuilder.EndSection();
+
+			if ( GetDefault<UNiagaraEditorSettings>()->bEnableBaker )
+			{
+				ToolbarBuilder.BeginSection("Baker");
+				{
+					ToolbarBuilder.AddComboButton(
+						FUIAction(),
+						FOnGetContent::CreateStatic(Local::GenerateBakerMenu, Toolkit),
+						LOCTEXT("Baker", "Baker"),
+						LOCTEXT("BakerTooltip", "Options for Baker rendering."),
+						FSlateIcon(FNiagaraEditorStyle::Get().GetStyleSetName(), "NiagaraEditor.Baker")
+					);
+				}
+				ToolbarBuilder.EndSection();
+			}
+
+			if(GNiagaraScalabilityModeEnabled)
+			{
+				ToolbarBuilder.BeginSection("Scalability");
+				{
+					FUIAction ScalabilityToggleAction(FExecuteAction::CreateRaw(Toolkit, &FNiagaraSystemToolkit::SetCurrentMode, FNiagaraSystemToolkit::ScalabilityModeName));
+					ScalabilityToggleAction.GetActionCheckState = FGetActionCheckState::CreateLambda([Toolkit]()
+					{
+						return Toolkit->GetCurrentMode() == FNiagaraSystemToolkit::ScalabilityModeName ? ECheckBoxState::Checked : ECheckBoxState::Unchecked;
+					});
+					
+					ToolbarBuilder.AddToolBarButton(
+						ScalabilityToggleAction,
+						NAME_None, 
+						LOCTEXT("ScalabilityLabel", "Scalability"),
+						LOCTEXT("ScalabilityTooltip", "Turn on scalability mode to optimize your effects for various platforms and quality settings."),
+						FSlateIcon(FNiagaraEditorStyle::Get().GetStyleSetName(), "NiagaraEditor.SimulationOptions")
+					);
+				}
+				ToolbarBuilder.EndSection();
+			}
+		}
+	};
+
+	ToolbarExtender->AddToolBarExtension(
+		"Asset",
+		EExtensionHook::After,
+		SystemToolkit.Pin()->GetToolkitCommands(),
+		FToolBarExtensionDelegate::CreateStatic(&Local::FillToolbar, SystemToolkit.Pin().Get())
+		);
+
+	//SystemToolkit.Pin()->AddToolbarExtender(ToolbarExtender);
+	
+	// FNiagaraEditorModule& NiagaraEditorModule = FModuleManager::LoadModuleChecked<FNiagaraEditorModule>("NiagaraEditor");
+	// SystemToolkit.Pin()->AddToolbarExtender(NiagaraEditorModule.GetToolBarExtensibilityManager()->GetAllExtenders(SystemToolkit.Pin()->GetToolkitCommands(), SystemToolkit.Pin()->GetObjectsBeingEdited()));
+}
+
+void FNiagaraSystemToolkitMode_Default::RegisterTabFactories(TSharedPtr<FTabManager> InTabManager)
+{
+	// this will register the common tab factories shared across modes and will call ExtendToolbar
+	FNiagaraSystemToolkitModeBase::RegisterTabFactories(InTabManager);
+
+	// add additional tab factories here that are exclusive to this mode
+}
diff --git a/Engine/Plugins/FX/Niagara/Source/NiagaraEditor/Private/Toolkits/SystemToolkitModes/NiagaraSystemToolkitMode_Default.h b/Engine/Plugins/FX/Niagara/Source/NiagaraEditor/Private/Toolkits/SystemToolkitModes/NiagaraSystemToolkitMode_Default.h
new file mode 100644
--- /dev/null
+++ b/Engine/Plugins/FX/Niagara/Source/NiagaraEditor/Private/Toolkits/SystemToolkitModes/NiagaraSystemToolkitMode_Default.h
@@ -0,0 +1,18 @@
+// Copyright Epic Games, Inc. All Rights Reserved.
+
+#pragma once
+
+#include "NiagaraSystemToolkitModeBase.h"
+
+class FNiagaraSystemToolkit;
+
+class FNiagaraSystemToolkitMode_Default : public FNiagaraSystemToolkitModeBase
+{
+public:
+	FNiagaraSystemToolkitMode_Default(TWeakPtr<FNiagaraSystemToolkit> InSystemToolkit);
+
+	/** FApplicationMode interface */
+	virtual void RegisterTabFactories(TSharedPtr<FTabManager> InTabManager) override;
+	
+	void ExtendToolbar();
+};
diff --git a/Engine/Plugins/FX/Niagara/Source/NiagaraEditor/Private/Toolkits/SystemToolkitModes/NiagaraSystemToolkitMode_Scalability.cpp b/Engine/Plugins/FX/Niagara/Source/NiagaraEditor/Private/Toolkits/SystemToolkitModes/NiagaraSystemToolkitMode_Scalability.cpp
new file mode 100644
--- /dev/null
+++ b/Engine/Plugins/FX/Niagara/Source/NiagaraEditor/Private/Toolkits/SystemToolkitModes/NiagaraSystemToolkitMode_Scalability.cpp
@@ -0,0 +1,416 @@
+// Copyright Epic Games, Inc. All Rights Reserved.
+
+#include "NiagaraSystemToolkitMode_Scalability.h"
+
+#include "IDetailTreeNode.h"
+#include "NiagaraEditorModule.h"
+#include "NiagaraSettings.h"
+#include "NiagaraSystemToolkit.h"
+#include "ViewModels/NiagaraSystemViewModel.h"
+#include "ViewModels/Stack/NiagaraStackRendererItem.h"
+#include "Widgets/Scalability/SNiagaraSystemResolvedScalabilitySettings.h"
+#include "ViewModels/NiagaraEmitterHandleViewModel.h"
+#include "Widgets/Scalability/SNiagaraScalabilityContext.h"
+
+#define LOCTEXT_NAMESPACE "NiagaraSystemToolkitMode_Scalability"
+
+const FName FNiagaraSystemToolkitMode_Scalability::ScalabilityContextTabID(TEXT("NiagaraSystemEditor_ScalabilityMode_ScalabilityContext"));
+const FName FNiagaraSystemToolkitMode_Scalability::ScalabilityViewModeTabID(TEXT("NiagaraSystemEditor_ScalabilityMode_ScalabilityViewMode"));
+const FName FNiagaraSystemToolkitMode_Scalability::ResolvedScalabilityTabID(TEXT("NiagaraSystemEditor_ScalabilityMode_ResolvedScalability"));
+const FName FNiagaraSystemToolkitMode_Scalability::EffectTypeTabID(TEXT("NiagaraSystemEditor_ScalabilityMode_EffectType"));
+
+FNiagaraSystemToolkitMode_Scalability::FNiagaraSystemToolkitMode_Scalability(TWeakPtr<FNiagaraSystemToolkit> InSystemToolkit) : FNiagaraSystemToolkitModeBase(FNiagaraSystemToolkit::ScalabilityModeName, InSystemToolkit)
+{
+	TabLayout = ChooseTabLayout();
+
+	ExtendToolbar();
+}
+
+void FNiagaraSystemToolkitMode_Scalability::RegisterTabFactories(TSharedPtr<FTabManager> InTabManager)
+{
+	FNiagaraSystemToolkitModeBase::RegisterTabFactories(InTabManager);
+
+	InTabManager->RegisterTabSpawner(ScalabilityContextTabID, FOnSpawnTab::CreateSP(this, &FNiagaraSystemToolkitMode_Scalability::SpawnTab_ScalabilityContext)).SetDisplayName(LOCTEXT("ScalabilityContextTabName", "Scalability Context"));
+	
+	if(SystemToolkit.Pin()->GetSystemViewModel()->GetEditMode() == ENiagaraSystemViewModelEditMode::SystemAsset)
+	{
+		InTabManager->RegisterTabSpawner(ResolvedScalabilityTabID, FOnSpawnTab::CreateSP(this, &FNiagaraSystemToolkitMode_Scalability::SpawnTab_ResolvedScalability)).SetDisplayName(LOCTEXT("ResolvedScalabilityTabName", "Active Scalability"));
+	}
+}
+
+void FNiagaraSystemToolkitMode_Scalability::PreDeactivateMode()
+{
+	FNiagaraSystemToolkitModeBase::PreDeactivateMode();
+
+	SystemToolkit.Pin()->RemoveToolbarExtender(ToolbarExtender);
+
+	TSharedPtr<FNiagaraSystemViewModel> SystemViewModel = SystemToolkit.Pin()->GetSystemViewModel();
+	UNiagaraSystemScalabilityViewModel* ScalabilityViewModel = SystemViewModel->GetScalabilityViewModel();
+
+	SystemViewModel->OnEmitterHandleViewModelsChanged().RemoveAll(this);
+	SystemViewModel->GetSystem().OnSystemPostEditChange().RemoveAll(this);
+	
+	UnbindEmitterPreviewOverrides();
+	UnbindSystemPreviewOverrides();
+	
+	FNiagaraPlatformSet::InvalidateCachedData();
+	SystemViewModel->GetSystem().UpdateScalability();
+	
+	ScalabilityViewModel->OnScalabilityModeChanged().Broadcast(false);
+}
+
+void FNiagaraSystemToolkitMode_Scalability::PostActivateMode()
+{
+	FNiagaraSystemToolkitModeBase::PostActivateMode();
+
+	SystemToolkit.Pin()->AddToolbarExtender(ToolbarExtender);
+
+	TSharedPtr<FNiagaraSystemViewModel> SystemViewModel = SystemToolkit.Pin()->GetSystemViewModel();
+	UNiagaraSystemScalabilityViewModel* ScalabilityViewModel = SystemViewModel->GetScalabilityViewModel();
+
+	// if our emitter handles change, we rebind the
+	SystemViewModel->OnEmitterHandleViewModelsChanged().AddSP(this, &FNiagaraSystemToolkitMode_Scalability::BindEmitterUpdates);
+	SystemViewModel->GetSystem().OnSystemPostEditChange().AddSP(this, &FNiagaraSystemToolkitMode_Scalability::BindSystemPreviewOverrides);
+
+	BindEmitterUpdates();
+	BindSystemPreviewOverrides(&SystemViewModel->GetSystem());
+	
+	FNiagaraPlatformSet::InvalidateCachedData();
+	SystemViewModel->GetSystem().UpdateScalability();
+
+	ScalabilityViewModel->OnScalabilityModeChanged().Broadcast(true);
+}
+
+void FNiagaraSystemToolkitMode_Scalability::UnbindEmitterUpdates()
+{
+	TSharedPtr<FNiagaraSystemViewModel> SystemViewModel = SystemToolkit.Pin()->GetSystemViewModel();
+
+	for(const TSharedRef<FNiagaraEmitterHandleViewModel>& EmitterHandleViewModel : SystemViewModel->GetEmitterHandleViewModels())
+	{
+		EmitterHandleViewModel->GetEmitterHandle()->GetInstance()->OnRenderersChanged().RemoveAll(this);
+		EmitterHandleViewModel->GetEmitterHandle()->GetInstance()->OnPropertiesChanged().RemoveAll(this);
+	}
+}
+
+void FNiagaraSystemToolkitMode_Scalability::BindEmitterPreviewOverrides()
+{
+	TSharedPtr<FNiagaraSystemViewModel> SystemViewModel = SystemToolkit.Pin()->GetSystemViewModel();
+	UNiagaraSystemScalabilityViewModel* ScalabilityViewModel = SystemViewModel->GetScalabilityViewModel();
+	
+	for(const TSharedRef<FNiagaraEmitterHandleViewModel>& EmitterHandleViewModel : SystemViewModel->GetEmitterHandleViewModels())
+	{
+		// we bind the quality level and device profile of the emitter itself to use our scalability preview
+		EmitterHandleViewModel->GetEmitterHandle()->GetInstance()->Platforms.OnOverrideQualityLevelDelegate.BindUObject(ScalabilityViewModel, &UNiagaraSystemScalabilityViewModel::GetPreviewQualityLevel);
+		EmitterHandleViewModel->GetEmitterHandle()->GetInstance()->Platforms.OnOverrideActiveDeviceProfileDelegate.BindUObject(ScalabilityViewModel, &UNiagaraSystemScalabilityViewModel::GetPreviewDeviceProfile);
+
+		// we do the same for the overrides
+		for(FNiagaraEmitterScalabilityOverride& EmitterScalabilityOverride : EmitterHandleViewModel->GetEmitterHandle()->GetInstance()->ScalabilityOverrides.Overrides)
+		{
+			EmitterScalabilityOverride.Platforms.OnOverrideQualityLevelDelegate.BindUObject(ScalabilityViewModel, &UNiagaraSystemScalabilityViewModel::GetPreviewQualityLevel);
+			EmitterScalabilityOverride.Platforms.OnOverrideActiveDeviceProfileDelegate.BindUObject(ScalabilityViewModel, &UNiagaraSystemScalabilityViewModel::GetPreviewDeviceProfile);
+		}
+
+		// and also for the renderers
+		EmitterHandleViewModel->GetEmitterHandle()->GetInstance()->ForEachRenderer([=](UNiagaraRendererProperties* RendererProperties)
+		{
+			RendererProperties->Platforms.OnOverrideQualityLevelDelegate.BindUObject(ScalabilityViewModel, &UNiagaraSystemScalabilityViewModel::GetPreviewQualityLevel);
+			RendererProperties->Platforms.OnOverrideActiveDeviceProfileDelegate.BindUObject(ScalabilityViewModel, &UNiagaraSystemScalabilityViewModel::GetPreviewDeviceProfile);
+		});
+	}
+}
+
+void FNiagaraSystemToolkitMode_Scalability::BindSystemPreviewOverrides(UNiagaraSystem* InSystem)
+{
+	TSharedPtr<FNiagaraSystemViewModel> SystemViewModel = SystemToolkit.Pin()->GetSystemViewModel();
+	UNiagaraSystemScalabilityViewModel* ScalabilityViewModel = SystemViewModel->GetScalabilityViewModel();
+
+	for(FNiagaraSystemScalabilityOverride& SystemScalabilityOverride : SystemViewModel->GetSystem().GetSystemScalabilityOverrides().Overrides)
+	{
+		SystemScalabilityOverride.Platforms.OnOverrideQualityLevelDelegate.BindUObject(ScalabilityViewModel, &UNiagaraSystemScalabilityViewModel::GetPreviewQualityLevel);
+		SystemScalabilityOverride.Platforms.OnOverrideActiveDeviceProfileDelegate.BindUObject(ScalabilityViewModel, &UNiagaraSystemScalabilityViewModel::GetPreviewDeviceProfile);
+	}
+}
+
+void FNiagaraSystemToolkitMode_Scalability::UnbindEmitterPreviewOverrides()
+{
+	TSharedPtr<FNiagaraSystemViewModel> SystemViewModel = SystemToolkit.Pin()->GetSystemViewModel();
+
+	// we undo all the bindings we created in the Bind function here 
+	for(const TSharedRef<FNiagaraEmitterHandleViewModel>& EmitterHandleViewModel : SystemViewModel->GetEmitterHandleViewModels())
+	{
+		EmitterHandleViewModel->GetEmitterHandle()->GetInstance()->Platforms.OnOverrideQualityLevelDelegate.Unbind();
+		EmitterHandleViewModel->GetEmitterHandle()->GetInstance()->Platforms.OnOverrideActiveDeviceProfileDelegate.Unbind();
+		
+		for(FNiagaraEmitterScalabilityOverride& EmitterScalabilityOverride : EmitterHandleViewModel->GetEmitterHandle()->GetInstance()->ScalabilityOverrides.Overrides)
+		{
+			EmitterScalabilityOverride.Platforms.OnOverrideQualityLevelDelegate.Unbind();
+			EmitterScalabilityOverride.Platforms.OnOverrideActiveDeviceProfileDelegate.Unbind();
+		}
+
+		EmitterHandleViewModel->GetEmitterHandle()->GetInstance()->ForEachRenderer([](UNiagaraRendererProperties* RendererProperties)
+		{
+			RendererProperties->Platforms.OnOverrideQualityLevelDelegate.Unbind();
+			RendererProperties->Platforms.OnOverrideActiveDeviceProfileDelegate.Unbind();
+		});
+	}
+}
+
+void FNiagaraSystemToolkitMode_Scalability::UnbindSystemPreviewOverrides()
+{
+	TSharedPtr<FNiagaraSystemViewModel> SystemViewModel = SystemToolkit.Pin()->GetSystemViewModel();
+
+	for(FNiagaraSystemScalabilityOverride& SystemScalabilityOverride : SystemViewModel->GetSystem().GetSystemScalabilityOverrides().Overrides)
+	{
+		SystemScalabilityOverride.Platforms.OnOverrideQualityLevelDelegate.Unbind();
+		SystemScalabilityOverride.Platforms.OnOverrideActiveDeviceProfileDelegate.Unbind();
+	}
+}
+
+void FNiagaraSystemToolkitMode_Scalability::BindEmitterUpdates()
+{
+	// we make sure to remove all existing multicast bindings first
+	UnbindEmitterUpdates();
+
+	// if our emitter changes properties (= platforms) or a renderer changes (= platforms), we refresh the bindings to make sure all new values get properly overridden
+	TSharedPtr<FNiagaraSystemViewModel> SystemViewModel = SystemToolkit.Pin()->GetSystemViewModel();
+	for(const TSharedRef<FNiagaraEmitterHandleViewModel>& EmitterHandleViewModel : SystemViewModel->GetEmitterHandleViewModels())
+	{		
+		EmitterHandleViewModel->GetEmitterHandle()->GetInstance()->OnRenderersChanged().AddSP(this, &FNiagaraSystemToolkitMode_Scalability::BindEmitterPreviewOverrides);
+		EmitterHandleViewModel->GetEmitterHandle()->GetInstance()->OnPropertiesChanged().AddSP(this, &FNiagaraSystemToolkitMode_Scalability::BindEmitterPreviewOverrides);
+	}
+
+	BindEmitterPreviewOverrides();
+}
+
+void FNiagaraSystemToolkitMode_Scalability::ExtendToolbar()
+{
+	struct Local
+	{
+		static void FillToolbar(FToolBarBuilder& ToolbarBuilder, FNiagaraSystemToolkit* Toolkit)
+		{
+			ToolbarBuilder.BeginSection("Scalability");
+			{
+				FUIAction ScalabilityToggleAction(FExecuteAction::CreateRaw(Toolkit, &FNiagaraSystemToolkit::SetCurrentMode, FNiagaraSystemToolkit::DefaultModeName));
+				ScalabilityToggleAction.GetActionCheckState = FGetActionCheckState::CreateLambda([Toolkit]()
+				{
+					return Toolkit->GetCurrentMode() == FNiagaraSystemToolkit::ScalabilityModeName ? ECheckBoxState::Checked : ECheckBoxState::Unchecked;
+				});
+						
+				ToolbarBuilder.AddToolBarButton(
+					ScalabilityToggleAction,
+					NAME_None, 
+					LOCTEXT("ScalabilityLabel", "Scalability"),
+					LOCTEXT("ScalabilityTooltip", "Turn on scalability mode to optimize your effects for various platforms and quality settings."),
+					FSlateIcon(FNiagaraEditorStyle::Get().GetStyleSetName(), "NiagaraEditor.SimulationOptions")
+				);
+
+				// @todo remove this once preview settings are finalized (tab vs. overlay in graph etc.)
+				// TSharedRef<SHorizontalBox> ViewModeContainer = SNew(SHorizontalBox)
+				// + SHorizontalBox::Slot()
+				// .AutoWidth()
+				// [
+				// 	SNew(SNiagaraScalabilityPreviewSettings, *Toolkit->GetSystemViewModel()->GetScalabilityViewModel())
+				// ];
+				//
+				// ToolbarBuilder.AddWidget(ViewModeContainer);
+			}
+			ToolbarBuilder.EndSection();		
+		}		
+	};
+	
+	ToolbarExtender->AddToolBarExtension(
+    	"Asset",
+    	EExtensionHook::After,
+    	SystemToolkit.Pin()->GetToolkitCommands(),
+    	FToolBarExtensionDelegate::CreateStatic(&Local::FillToolbar, SystemToolkit.Pin().Get())
+    	);
+}
+
+TSharedRef<SDockTab> FNiagaraSystemToolkitMode_Scalability::SpawnTab_ScalabilityContext(const FSpawnTabArgs& Args)
+{
+	check(Args.GetTabId().TabType == ScalabilityContextTabID);
+
+	TSharedRef<SDockTab> SpawnedTab =
+		SNew(SDockTab)
+		[
+			SAssignNew(ScalabilityContext, SNiagaraScalabilityContext, *SystemToolkit.Pin()->GetSystemViewModel()->GetScalabilityViewModel())
+		];
+	
+	return SpawnedTab;
+}
+
+TSharedRef<SDockTab> FNiagaraSystemToolkitMode_Scalability::SpawnTab_ResolvedScalability(const FSpawnTabArgs& Args)
+{
+	check(Args.GetTabId().TabType == ResolvedScalabilityTabID);
+
+	TSharedRef<SDockTab> SpawnedTab =
+		SNew(SDockTab)
+		[			
+			SNew(SNiagaraSystemResolvedScalabilitySettings, SystemToolkit.Pin()->GetSystemViewModel()->GetSystem(), SystemToolkit.Pin()->GetSystemViewModel())
+		];
+
+	return SpawnedTab;
+}
+
+TSharedRef<FTabManager::FLayout> FNiagaraSystemToolkitMode_Scalability::ChooseTabLayout()
+{
+	if(SystemToolkit.Pin()->GetSystemViewModel()->GetEditMode() == ENiagaraSystemViewModelEditMode::SystemAsset)
+	{
+		return FTabManager::NewLayout("Standalone_Niagara_System_Layout_Scalability_v29")
+		->AddArea
+		(
+			FTabManager::NewPrimaryArea()->SetOrientation(Orient_Vertical)
+			->Split
+			(
+				// Main Content Area
+				FTabManager::NewSplitter()->SetOrientation(Orient_Horizontal)
+				//->SetSizeCoefficient(0.75f)
+				// Left area
+				->Split
+				(
+					// Main Content Area - Left
+					FTabManager::NewSplitter()->SetOrientation(Orient_Vertical)
+					->SetSizeCoefficient(.75f)
+					->Split
+					(
+						// Inner Left
+						FTabManager::NewSplitter()->SetOrientation(Orient_Horizontal)
+						->SetSizeCoefficient(.75f)
+						->Split
+						(
+							FTabManager::NewStack()
+							->SetSizeCoefficient(0.4f)
+							->AddTab(ViewportTabID, ETabState::OpenedTab)
+						)
+						// Inner Right
+						->Split
+						(
+							FTabManager::NewSplitter()->SetOrientation(Orient_Vertical)
+							->Split
+							(
+								FTabManager::NewStack()
+								->SetSizeCoefficient(0.8f)
+								->AddTab(SystemOverviewTabID, ETabState::OpenedTab)
+								->AddTab(ScratchPadTabID, ETabState::OpenedTab)
+								->SetForegroundTab(SystemOverviewTabID)
+							)
+						)
+					)
+					->Split
+					(
+						// Inner Left Bottom
+						FTabManager::NewStack()
+						->SetSizeCoefficient(0.25f)
+						->AddTab(CurveEditorTabID, ETabState::OpenedTab)
+						->AddTab(MessageLogTabID, ETabState::OpenedTab)
+						->AddTab(SequencerTabID, ETabState::OpenedTab)
+						->AddTab(ScriptStatsTabID, ETabState::ClosedTab)
+					)
+				)
+				// Right Sidebar Area
+				->Split
+				(
+					// Top Level Right
+					FTabManager::NewSplitter()->SetOrientation(Orient_Vertical)
+					->SetSizeCoefficient(0.25f)
+					->Split
+					(
+						FTabManager::NewStack()
+						->SetSizeCoefficient(0.35f)
+						->AddTab(ScalabilityContextTabID, ETabState::OpenedTab)
+					)
+					->Split
+					(
+						FTabManager::NewStack()
+						->SetSizeCoefficient(0.3f)
+						->AddTab(ResolvedScalabilityTabID, ETabState::OpenedTab)
+						->AddTab(SelectedEmitterStackTabID, ETabState::OpenedTab)
+						->AddTab(SelectedEmitterGraphTabID, ETabState::ClosedTab)
+						->AddTab(SystemScriptTabID, ETabState::ClosedTab)
+						->AddTab(SystemDetailsTabID, ETabState::ClosedTab)
+						->AddTab(DebugSpreadsheetTabID, ETabState::ClosedTab)
+						->AddTab(PreviewSettingsTabId, ETabState::ClosedTab)
+						->AddTab(GeneratedCodeTabID, ETabState::ClosedTab)
+						->SetForegroundTab(ResolvedScalabilityTabID)
+					)
+				)
+			)
+		);
+	}
+	else if(SystemToolkit.Pin()->GetSystemViewModel()->GetEditMode() == ENiagaraSystemViewModelEditMode::EmitterAsset)
+	{
+		return FTabManager::NewLayout("Standalone_Niagara_Emitter_Layout_Scalability_v1")
+		->AddArea
+		(
+			FTabManager::NewPrimaryArea()->SetOrientation(Orient_Vertical)
+			->Split
+			(
+				FTabManager::NewSplitter()->SetOrientation(Orient_Horizontal)
+				->SetSizeCoefficient(0.75f)
+				->Split
+				(
+					// Top Level Left
+					FTabManager::NewSplitter()->SetOrientation(Orient_Vertical)
+					->SetSizeCoefficient(.75f)
+					->Split
+					(
+						// Inner Left Top
+						FTabManager::NewSplitter()->SetOrientation(Orient_Horizontal)
+						->SetSizeCoefficient(0.75f)
+						->Split
+						(
+							FTabManager::NewStack()
+							->SetSizeCoefficient(0.4f)
+							->AddTab(ViewportTabID, ETabState::OpenedTab)
+						)
+						->Split
+						(							
+							FTabManager::NewStack()
+							->AddTab(SystemOverviewTabID, ETabState::OpenedTab)
+							->AddTab(ScratchPadTabID, ETabState::OpenedTab)
+							->SetForegroundTab(SystemOverviewTabID)							
+						)
+					)
+					->Split
+					(
+						// Inner Left Bottom
+						FTabManager::NewStack()
+						->SetSizeCoefficient(0.25f)
+						->AddTab(CurveEditorTabID, ETabState::OpenedTab)
+						->AddTab(MessageLogTabID, ETabState::OpenedTab)
+						->AddTab(SequencerTabID, ETabState::OpenedTab)
+						->AddTab(ScriptStatsTabID, ETabState::ClosedTab)
+					)
+				)
+				->Split
+				(
+					// Top Level Right
+					FTabManager::NewSplitter()->SetOrientation(Orient_Vertical)
+					->SetSizeCoefficient(0.25f)
+					->Split
+					(
+						FTabManager::NewStack()
+						->SetSizeCoefficient(0.35f)
+						->AddTab(ScalabilityContextTabID, ETabState::OpenedTab)
+					)
+					->Split
+					(
+						FTabManager::NewStack()
+						->SetSizeCoefficient(0.65f)
+						->AddTab(SelectedEmitterStackTabID, ETabState::OpenedTab)
+						->AddTab(SelectedEmitterGraphTabID, ETabState::ClosedTab)
+						->AddTab(SystemScriptTabID, ETabState::ClosedTab)
+						->AddTab(SystemDetailsTabID, ETabState::ClosedTab)
+						->AddTab(DebugSpreadsheetTabID, ETabState::ClosedTab)
+						->AddTab(PreviewSettingsTabId, ETabState::ClosedTab)
+						->AddTab(GeneratedCodeTabID, ETabState::ClosedTab)
+					)
+				)
+			)
+		);
+	}
+
+	return FTabManager::NewLayout("DummyLayout");
+}
+
+#undef LOCTEXT_NAMESPACE
diff --git a/Engine/Plugins/FX/Niagara/Source/NiagaraEditor/Private/Toolkits/SystemToolkitModes/NiagaraSystemToolkitMode_Scalability.h b/Engine/Plugins/FX/Niagara/Source/NiagaraEditor/Private/Toolkits/SystemToolkitModes/NiagaraSystemToolkitMode_Scalability.h
new file mode 100644
--- /dev/null
+++ b/Engine/Plugins/FX/Niagara/Source/NiagaraEditor/Private/Toolkits/SystemToolkitModes/NiagaraSystemToolkitMode_Scalability.h
@@ -0,0 +1,41 @@
+// Copyright Epic Games, Inc. All Rights Reserved.
+
+#pragma once
+
+#include "NiagaraSystemToolkitModeBase.h"
+#include "Customizations/NiagaraPlatformSetCustomization.h"
+
+class FNiagaraSystemToolkit;
+
+class FNiagaraSystemToolkitMode_Scalability : public FNiagaraSystemToolkitModeBase
+{
+public:
+	FNiagaraSystemToolkitMode_Scalability(TWeakPtr<FNiagaraSystemToolkit> InSystemToolkit);
+
+	/** FApplicationMode interface */
+	virtual void RegisterTabFactories(TSharedPtr<FTabManager> InTabManager) override;
+	virtual void PreDeactivateMode() override;
+	virtual void PostActivateMode() override;
+
+	void ExtendToolbar();
+protected:
+	TSharedRef<SDockTab> SpawnTab_ScalabilityContext(const FSpawnTabArgs& Args);
+	TSharedRef<SDockTab> SpawnTab_ResolvedScalability(const FSpawnTabArgs& Args);
+
+private:
+	TSharedRef<FTabManager::FLayout> ChooseTabLayout();
+	void BindEmitterUpdates();
+	void UnbindEmitterUpdates();
+
+	void BindEmitterPreviewOverrides();
+	void BindSystemPreviewOverrides(UNiagaraSystem* System);
+	void UnbindEmitterPreviewOverrides();
+	void UnbindSystemPreviewOverrides();
+
+	TSharedPtr<class SNiagaraScalabilityContext> ScalabilityContext; 
+public:
+	static const FName ScalabilityContextTabID;
+	static const FName ScalabilityViewModeTabID;
+	static const FName ResolvedScalabilityTabID;
+	static const FName EffectTypeTabID;
+};
diff --git a/Engine/Plugins/FX/Niagara/Source/NiagaraEditor/Private/ViewModels/NiagaraCurveSelectionViewModel.cpp b/Engine/Plugins/FX/Niagara/Source/NiagaraEditor/Private/ViewModels/NiagaraCurveSelectionViewModel.cpp
--- a/Engine/Plugins/FX/Niagara/Source/NiagaraEditor/Private/ViewModels/NiagaraCurveSelectionViewModel.cpp
+++ b/Engine/Plugins/FX/Niagara/Source/NiagaraEditor/Private/ViewModels/NiagaraCurveSelectionViewModel.cpp
@@ -27,6 +27,7 @@
 #include "ViewModels/NiagaraSystemViewModel.h"
 #include "ViewModels/Stack/NiagaraParameterHandle.h"
 #include "ViewModels/Stack/NiagaraStackGraphUtilities.h"
+#include "Toolkits/SystemToolkitModes/NiagaraSystemToolkitModeBase.h"
 
 bool FNiagaraCurveSelectionTreeNodeDataId::operator==(const FNiagaraCurveSelectionTreeNodeDataId& Other) const
 {
@@ -351,7 +352,7 @@ void UNiagaraCurveSelectionViewModel::FocusAndSelectCurveDataInterface(UNiagaraD
 
 		if(IdOfNodeToSelect.IsValid())
 		{
-			SystemViewModel->FocusTab(FNiagaraSystemToolkit::CurveEditorTabID);
+			SystemViewModel->FocusTab(FNiagaraSystemToolkitModeBase::CurveEditorTabID);
 			OnRequestSelectNodeDelegate.Broadcast(IdOfNodeToSelect);
 		}
 	}
diff --git a/Engine/Plugins/FX/Niagara/Source/NiagaraEditor/Private/ViewModels/NiagaraScratchPadViewModel.cpp b/Engine/Plugins/FX/Niagara/Source/NiagaraEditor/Private/ViewModels/NiagaraScratchPadViewModel.cpp
--- a/Engine/Plugins/FX/Niagara/Source/NiagaraEditor/Private/ViewModels/NiagaraScratchPadViewModel.cpp
+++ b/Engine/Plugins/FX/Niagara/Source/NiagaraEditor/Private/ViewModels/NiagaraScratchPadViewModel.cpp
@@ -8,6 +8,7 @@
 #include "ViewModels/NiagaraParameterPanelViewModel.h"
 #include "ViewModels/Stack/NiagaraStackModuleItem.h"
 #include "ViewModels/Stack/NiagaraStackFunctionInput.h"
+#include "Toolkits/SystemToolkitModes/NiagaraSystemToolkitModeBase.h"
 #include "NiagaraSystem.h"
 #include "NiagaraScriptSource.h"
 #include "NiagaraNodeOutput.h"
@@ -244,7 +245,7 @@ void UNiagaraScratchPadViewModel::FocusScratchPadScriptViewModel(TSharedRef<FNia
 	if (ensureMsgf(ScriptViewModels.Contains(InScriptViewModel), TEXT("Can only focus a view model from this scratch pad view model.")))
 	{
 		SetActiveScriptViewModel(InScriptViewModel);
-		GetSystemViewModel()->FocusTab(FNiagaraSystemToolkit::ScratchPadTabID);
+		GetSystemViewModel()->FocusTab(FNiagaraSystemToolkitModeBase::ScratchPadTabID);
 	}
 }
 
diff --git a/Engine/Plugins/FX/Niagara/Source/NiagaraEditor/Private/ViewModels/NiagaraSystemScalabilityViewModel.cpp b/Engine/Plugins/FX/Niagara/Source/NiagaraEditor/Private/ViewModels/NiagaraSystemScalabilityViewModel.cpp
new file mode 100644
--- /dev/null
+++ b/Engine/Plugins/FX/Niagara/Source/NiagaraEditor/Private/ViewModels/NiagaraSystemScalabilityViewModel.cpp
@@ -0,0 +1,86 @@
+// Copyright Epic Games, Inc. All Rights Reserved.
+
+#include "ViewModels/NiagaraSystemScalabilityViewModel.h"
+
+#include "ViewModels/NiagaraSystemViewModel.h"
+#include "NiagaraSystem.h"
+#include "ISequencerModule.h"
+#include "NiagaraSystemToolkit.h"
+#include "DeviceProfiles/DeviceProfile.h"
+#include "Modules/ModuleManager.h"
+#include "ViewModels/NiagaraEmitterHandleViewModel.h"
+
+#define LOCTEXT_NAMESPACE "NiagaraSystemScalabilityViewModel"
+
+UNiagaraSystemScalabilityViewModel::UNiagaraSystemScalabilityViewModel()
+	: SystemViewModel(nullptr)
+{
+}
+
+void UNiagaraSystemScalabilityViewModel::Initialize(TSharedRef<FNiagaraSystemViewModel> InSystemViewModel)
+{
+	SystemViewModel = InSystemViewModel;
+
+	PreviewPlatforms = MakeShared<FNiagaraPlatformSet>(FNiagaraPlatformSet::CreateQualityLevelMask(3));
+}
+
+bool UNiagaraSystemScalabilityViewModel::IsValid() const
+{
+	return SystemViewModel != nullptr;
+}
+
+bool UNiagaraSystemScalabilityViewModel::IsActive() const
+{
+	return GetSystemViewModel().Pin()->GetWorkflowMode() == FNiagaraSystemToolkit::ScalabilityModeName;
+}
+
+void UNiagaraSystemScalabilityViewModel::UpdatePreviewDeviceProfile(UDeviceProfile* DeviceProfile)
+{
+	if(DeviceProfile)
+	{
+		PreviewDeviceProfile = DeviceProfile;
+	}
+	else
+	{
+		PreviewDeviceProfile.Reset();
+	}
+
+	FNiagaraPlatformSet::InvalidateCachedData();
+	GetSystemViewModel().Pin()->GetSystem().UpdateScalability();
+}
+
+void UNiagaraSystemScalabilityViewModel::UpdatePreviewQualityLevel(int32 QualityLevel, bool bEnabled)
+{
+	PreviewPlatforms->QualityLevelMask = 0;
+	PreviewPlatforms->SetEnabledForEffectQuality(QualityLevel, bEnabled);
+
+	FNiagaraPlatformSet::InvalidateCachedData();
+	GetSystemViewModel().Pin()->GetSystem().UpdateScalability();
+}
+
+// void UNiagaraSystemScalabilityViewModel::NavigateToScalabilityProperty(UObject* Object, FName PropertyName)
+// {
+// 	TArray<UNiagaraOverviewNode*> OverviewNodes;
+// 	GetSystemViewModel().Pin()->GetOverviewGraphViewModel()->GetGraph()->GetNodesOfClass<UNiagaraOverviewNode>(OverviewNodes);
+// 	
+// 	for(UNiagaraOverviewNode* Node : OverviewNodes)
+// 	{
+// 		if(FNiagaraEmitterHandle* EmitterHandle = Node->TryGetEmitterHandle())
+// 		{
+// 			if(Object == EmitterHandle->GetInstance())
+// 			{
+// 				GetSystemViewModel().Pin()->GetOverviewGraphViewModel()->GetNodeSelection()->SetSelectedObject(Node);
+// 				break;
+// 			}
+// 		}
+// 		else if(Object == Node->GetOwningSystem())
+// 		{
+// 			GetSystemViewModel().Pin()->GetOverviewGraphViewModel()->GetNodeSelection()->SetSelectedObject(Node);
+// 			break;
+// 		}
+// 	}
+//
+// 	ScalabilityPropertySelectedDelegate.Broadcast(PropertyName);
+// }
+
+#undef LOCTEXT_NAMESPACE // NiagaraSystemScalabilityViewModel
diff --git a/Engine/Plugins/FX/Niagara/Source/NiagaraEditor/Private/ViewModels/NiagaraSystemViewModel.cpp b/Engine/Plugins/FX/Niagara/Source/NiagaraEditor/Private/ViewModels/NiagaraSystemViewModel.cpp
--- a/Engine/Plugins/FX/Niagara/Source/NiagaraEditor/Private/ViewModels/NiagaraSystemViewModel.cpp
+++ b/Engine/Plugins/FX/Niagara/Source/NiagaraEditor/Private/ViewModels/NiagaraSystemViewModel.cpp
@@ -135,8 +135,11 @@ void FNiagaraSystemViewModel::Initialize(UNiagaraSystem& InSystem, FNiagaraSyste
 	CurveSelectionViewModel = NewObject<UNiagaraCurveSelectionViewModel>(GetTransientPackage());
 	CurveSelectionViewModel->Initialize(this->AsShared());
 
+	ScalabilityViewModel = NewObject<UNiagaraSystemScalabilityViewModel>(GetTransientPackage());
+	ScalabilityViewModel->Initialize(this->AsShared());
+	
 	PlaceholderDataInterfaceManager = MakeShared<FNiagaraPlaceholderDataInterfaceManager>(this->AsShared());
-
+	
 	SetupPreviewComponentAndInstance();
 	SetupSequencer();
 	RefreshAll();
@@ -238,6 +241,11 @@ void FNiagaraSystemViewModel::Cleanup()
 		CurveSelectionViewModel = nullptr;
 	}
 
+	if(ScalabilityViewModel != nullptr)
+	{
+		ScalabilityViewModel = nullptr;
+	}
+
 	if (PlaceholderDataInterfaceManager.IsValid())
 	{
 		PlaceholderDataInterfaceManager.Reset();
@@ -355,6 +363,17 @@ ENiagaraSystemViewModelEditMode FNiagaraSystemViewModel::GetEditMode() const
 	return EditMode;
 }
 
+FName FNiagaraSystemViewModel::GetWorkflowMode() const
+{
+	check(OnGetWorkflowModeDelegate.IsBound());
+	return OnGetWorkflowModeDelegate.Execute();
+}
+
+FOnGetWorkflowMode& FNiagaraSystemViewModel::OnGetWorkflowMode()
+{
+	return OnGetWorkflowModeDelegate;
+}
+
 TSharedPtr<FNiagaraEmitterHandleViewModel> FNiagaraSystemViewModel::AddEmitterFromAssetData(const FAssetData& AssetData)
 {
 	UNiagaraEmitter* Emitter = Cast<UNiagaraEmitter>(AssetData.GetAsset());
@@ -640,6 +659,10 @@ void FNiagaraSystemViewModel::AddReferencedObjects(FReferenceCollector& Collecto
 	{
 		Collector.AddReferencedObject(CurveSelectionViewModel);
 	}
+	if (ScalabilityViewModel != nullptr)
+	{
+		Collector.AddReferencedObject(ScalabilityViewModel);
+	}
 }
 
 void FNiagaraSystemViewModel::PostUndo(bool bSuccess)
@@ -880,6 +903,11 @@ UNiagaraCurveSelectionViewModel* FNiagaraSystemViewModel::GetCurveSelectionViewM
 	return CurveSelectionViewModel;
 }
 
+UNiagaraSystemScalabilityViewModel* FNiagaraSystemViewModel::GetScalabilityViewModel()
+{
+	return ScalabilityViewModel;
+}
+
 TArray<float> FNiagaraSystemViewModel::OnGetPlaybackSpeeds() const
 {
 	return GetDefault<UNiagaraEditorSettings>()->GetPlaybackSpeeds();
diff --git a/Engine/Plugins/FX/Niagara/Source/NiagaraEditor/Private/ViewModels/Stack/NiagaraStackEntry.cpp b/Engine/Plugins/FX/Niagara/Source/NiagaraEditor/Private/ViewModels/Stack/NiagaraStackEntry.cpp
--- a/Engine/Plugins/FX/Niagara/Source/NiagaraEditor/Private/ViewModels/Stack/NiagaraStackEntry.cpp
+++ b/Engine/Plugins/FX/Niagara/Source/NiagaraEditor/Private/ViewModels/Stack/NiagaraStackEntry.cpp
@@ -1,6 +1,9 @@
 // Copyright Epic Games, Inc. All Rights Reserved.
 
 #include "ViewModels/Stack/NiagaraStackEntry.h"
+
+#include "NiagaraEditorUtilities.h"
+#include "NiagaraEmitter.h"
 #include "ViewModels/Stack/NiagaraStackErrorItem.h"
 #include "ViewModels/NiagaraSystemViewModel.h"
 #include "ViewModels/NiagaraEmitterViewModel.h"
diff --git a/Engine/Plugins/FX/Niagara/Source/NiagaraEditor/Private/ViewModels/Stack/NiagaraStackRendererItem.cpp b/Engine/Plugins/FX/Niagara/Source/NiagaraEditor/Private/ViewModels/Stack/NiagaraStackRendererItem.cpp
--- a/Engine/Plugins/FX/Niagara/Source/NiagaraEditor/Private/ViewModels/Stack/NiagaraStackRendererItem.cpp
+++ b/Engine/Plugins/FX/Niagara/Source/NiagaraEditor/Private/ViewModels/Stack/NiagaraStackRendererItem.cpp
@@ -131,6 +131,21 @@ bool UNiagaraStackRendererItem::AddMissingVariable(UNiagaraEmitter* Emitter, con
 	return true;
 }
 
+bool UNiagaraStackRendererItem::IsExcludedFromScalability() const
+{
+	if(RendererProperties.IsValid()) 
+	{
+		return !RendererProperties.Get()->Platforms.IsActive();
+	}
+
+	return false;
+}
+
+bool UNiagaraStackRendererItem::IsOwningEmitterExcludedFromScalability() const
+{
+	return GetEmitterViewModel().IsValid() ? !GetEmitterViewModel()->GetEmitter()->IsAllowedByScalability() : false;
+}
+
 UNiagaraRendererProperties* UNiagaraStackRendererItem::GetRendererProperties()
 {
 	return RendererProperties.Get();
diff --git a/Engine/Plugins/FX/Niagara/Source/NiagaraEditor/Private/ViewModels/Stack/NiagaraStackScriptItemGroup.cpp b/Engine/Plugins/FX/Niagara/Source/NiagaraEditor/Private/ViewModels/Stack/NiagaraStackScriptItemGroup.cpp
--- a/Engine/Plugins/FX/Niagara/Source/NiagaraEditor/Private/ViewModels/Stack/NiagaraStackScriptItemGroup.cpp
+++ b/Engine/Plugins/FX/Niagara/Source/NiagaraEditor/Private/ViewModels/Stack/NiagaraStackScriptItemGroup.cpp
@@ -30,6 +30,7 @@
 #include "ViewModels/Stack/NiagaraStackGraphUtilities.h"
 #include "ViewModels/Stack/NiagaraStackModuleItem.h"
 #include "Widgets/Notifications/SNotificationList.h"
+#include "Toolkits/SystemToolkitModes/NiagaraSystemToolkitModeBase.h"
 
 #define LOCTEXT_NAMESPACE "UNiagaraStackScriptItemGroup"
 
@@ -669,7 +670,7 @@ void UNiagaraStackScriptItemGroup::RefreshIssues(TArray<FStackIssue>& NewIssues)
 						LOCTEXT("ShowLogForErrorsFix", "Show the errors in the Niagara Log"), 
 						FStackIssueFixDelegate::CreateLambda([this]()
 						{
-							GetSystemViewModel()->FocusTab(FNiagaraSystemToolkit::MessageLogTabID);
+							GetSystemViewModel()->FocusTab(FNiagaraSystemToolkitModeBase::MessageLogTabID);
 						}),
 						EStackIssueFixStyle::Link));
 
@@ -687,7 +688,7 @@ void UNiagaraStackScriptItemGroup::RefreshIssues(TArray<FStackIssue>& NewIssues)
 						LOCTEXT("ShowLogForWarningsFix", "Show the warnings in the Niagara Log"),
 						FStackIssueFixDelegate::CreateLambda([this]()
 							{
-								GetSystemViewModel()->FocusTab(FNiagaraSystemToolkit::MessageLogTabID);
+								GetSystemViewModel()->FocusTab(FNiagaraSystemToolkitModeBase::MessageLogTabID);
 							}),
 						EStackIssueFixStyle::Link));
 				NewIssues.Add(CompileWarning);
diff --git a/Engine/Plugins/FX/Niagara/Source/NiagaraEditor/Private/ViewModels/Stack/NiagaraStackSystemPropertiesItem.cpp b/Engine/Plugins/FX/Niagara/Source/NiagaraEditor/Private/ViewModels/Stack/NiagaraStackSystemPropertiesItem.cpp
--- a/Engine/Plugins/FX/Niagara/Source/NiagaraEditor/Private/ViewModels/Stack/NiagaraStackSystemPropertiesItem.cpp
+++ b/Engine/Plugins/FX/Niagara/Source/NiagaraEditor/Private/ViewModels/Stack/NiagaraStackSystemPropertiesItem.cpp
@@ -3,6 +3,7 @@
 #include "ViewModels/Stack/NiagaraStackSystemPropertiesItem.h"
 #include "ViewModels/Stack/NiagaraStackObject.h"
 #include "NiagaraSystem.h"
+#include "NiagaraSystemDetailsCustomization.h"
 #include "ViewModels/NiagaraSystemViewModel.h"
 #include "ScopedTransaction.h"
 
@@ -36,6 +37,7 @@ void UNiagaraStackSystemPropertiesItem::RefreshChildrenInternal(const TArray<UNi
 		SystemObject = NewObject<UNiagaraStackObject>(this);
 		FRequiredEntryData RequiredEntryData(GetSystemViewModel(), GetEmitterViewModel(), FExecutionCategoryNames::System, NAME_None, GetStackEditorData());
 		SystemObject->Initialize(RequiredEntryData, System.Get(), GetStackEditorDataKey());
+		SystemObject->RegisterInstancedCustomPropertyLayout(UNiagaraSystem::StaticClass(), FOnGetDetailCustomizationInstance::CreateStatic(&FNiagaraSystemDetails::MakeInstance));
 	}
 	NewChildren.Add(SystemObject);
 	bCanResetToBase.Reset();
diff --git a/Engine/Plugins/FX/Niagara/Source/NiagaraEditor/Private/Widgets/Scalability/SNiagaraScalabilityContext.cpp b/Engine/Plugins/FX/Niagara/Source/NiagaraEditor/Private/Widgets/Scalability/SNiagaraScalabilityContext.cpp
new file mode 100644
--- /dev/null
+++ b/Engine/Plugins/FX/Niagara/Source/NiagaraEditor/Private/Widgets/Scalability/SNiagaraScalabilityContext.cpp
@@ -0,0 +1,137 @@
+// Copyright Epic Games, Inc. All Rights Reserved.
+
+#include "SNiagaraScalabilityContext.h"
+#include "IDetailChildrenBuilder.h"
+#include "IPropertyRowGenerator.h"
+#include "ISinglePropertyView.h"
+#include "NiagaraMeshRendererProperties.h"
+#include "NiagaraObjectSelection.h"
+#include "Modules/ModuleManager.h"
+#include "Widgets/Input/SButton.h"
+#include "ViewModels/NiagaraOverviewGraphViewModel.h"
+#include "ViewModels/NiagaraSystemSelectionViewModel.h"
+#include "ViewModels/NiagaraSystemViewModel.h"
+#include "ViewModels/Stack/NiagaraStackEntry.h"
+#include "ViewModels/Stack/NiagaraStackRendererItem.h"
+#include "PropertyEditorModule.h"
+
+#define LOCTEXT_NAMESPACE "NiagaraScalabilityContext"
+
+void SNiagaraScalabilityContext::Construct(const FArguments& InArgs, UNiagaraSystemScalabilityViewModel& InScalabilityViewModel)
+{
+	ScalabilityViewModel = &InScalabilityViewModel;
+	
+	ScalabilityViewModel->GetSystemViewModel().Pin()->GetSelectionViewModel()->OnEntrySelectionChanged().AddSP(this, &SNiagaraScalabilityContext::UpdateScalabilityContent);
+	ScalabilityViewModel->GetSystemViewModel().Pin()->GetOverviewGraphViewModel()->GetNodeSelection()->OnSelectedObjectsChanged().AddSP(this, &SNiagaraScalabilityContext::UpdateScalabilityContent);
+
+	FPropertyEditorModule& PropertyEditorModule = FModuleManager::Get().LoadModuleChecked<FPropertyEditorModule>("PropertyEditor");
+	
+	FDetailsViewArgs DetailsViewArgs;
+	DetailsViewArgs.ViewIdentifier = "ScalabilityContextDetailsView";
+	DetailsViewArgs.NameAreaSettings = FDetailsViewArgs::ObjectsUseNameArea;
+	DetailsView = PropertyEditorModule.CreateDetailView(DetailsViewArgs);
+
+	// important: we register a generic details customization to overwrite the global customization for this object. This is because using "EditCategory" in a customization
+	// puts a category from default into custom categories. This in turn will ignore the property visible delegate.
+	// without this, a few properties from the mesh renderer properties make it into the details panel even though the delegate returns false for them.
+	DetailsView->RegisterInstancedCustomPropertyLayout(UNiagaraMeshRendererProperties::StaticClass(), DetailsView->GetGenericLayoutDetailsDelegate());
+	// add more here...
+	
+	DetailsView->SetIsPropertyVisibleDelegate(FIsPropertyVisible::CreateSP(this, &SNiagaraScalabilityContext::FilterScalabilityProperties));
+	UpdateScalabilityContent();
+
+	ChildSlot
+	[
+		DetailsView.ToSharedRef()
+	];
+}
+
+void SNiagaraScalabilityContext::SetObject(UObject* Object)
+{	
+	if(UNiagaraEmitter* Emitter = Cast<UNiagaraEmitter>(Object))
+	{		
+		DetailsView->SetObject(Emitter, true);
+	}
+	else if(UNiagaraSystem* System = Cast<UNiagaraSystem>(Object))
+	{
+		DetailsView->SetObject(System, true);
+	}
+	else if(UNiagaraStackRendererItem* RendererItem = Cast<UNiagaraStackRendererItem>(Object))
+	{
+		DetailsView->SetObject(RendererItem->GetRendererProperties(), true);
+	}
+}
+
+void SNiagaraScalabilityContext::UpdateScalabilityContent()
+{
+	TArray<UNiagaraStackEntry*> StackEntries;
+	ScalabilityViewModel->GetSystemViewModel().Pin()->GetSelectionViewModel()->GetSelectedEntries(StackEntries);
+	const TSet<UObject*>& SelectedNodes = ScalabilityViewModel->GetSystemViewModel().Pin()->GetOverviewGraphViewModel()->GetNodeSelection()->GetSelectedObjects();
+
+	if(StackEntries.Num() == 1)
+	{
+		UNiagaraStackEntry* StackEntry = StackEntries[0];
+		bool bDoesSupportScalabilityContext =  StackEntry->IsA(UNiagaraStackRendererItem::StaticClass());
+
+		if(bDoesSupportScalabilityContext)
+		{
+			SetObject(StackEntry);
+			return;
+		}
+	}
+	
+	if(SelectedNodes.Num() > 0)
+	{
+		UObject* Object = SelectedNodes.Array()[0];
+		UNiagaraOverviewNode* OverviewNode = CastChecked<UNiagaraOverviewNode>(Object);
+		if(FNiagaraEmitterHandle* Handle = OverviewNode->TryGetEmitterHandle())
+		{
+			SetObject(Handle->GetInstance());
+			return;
+		}
+		else
+		{
+			SetObject(OverviewNode->GetOwningSystem());
+			return;
+		}
+	}
+}
+
+bool SNiagaraScalabilityContext::FilterScalabilityProperties(const FPropertyAndParent& InPropertyAndParent) const
+{
+	auto ShouldPropertyBeVisible = [](const FProperty& InProperty)
+	{
+		return InProperty.HasMetaData("DisplayInScalabilityContext");		
+	};
+
+	auto IsParentPropertyVisible = [&](const TArray<const FProperty*> ParentProperties)
+	{
+		for(const FProperty* Property : ParentProperties)
+		{
+			if(ShouldPropertyBeVisible(*Property))
+			{
+				return true;
+			}
+		}
+
+		return false;
+	};
+	
+	if (InPropertyAndParent.Property.IsA<FStructProperty>())
+	{
+		const FStructProperty* StructProperty = CastField<FStructProperty>(&InPropertyAndParent.Property);
+		for (TFieldIterator<FProperty> PropertyIt(StructProperty->Struct); PropertyIt; ++PropertyIt)
+		{
+			if (ShouldPropertyBeVisible(**PropertyIt))
+			{
+				return true;
+			}
+		}
+	}
+
+	bool bShowProperty = ShouldPropertyBeVisible(InPropertyAndParent.Property);
+	bShowProperty = bShowProperty || IsParentPropertyVisible(InPropertyAndParent.ParentProperties);
+	return  bShowProperty;
+}
+
+#undef LOCTEXT_NAMESPACE
diff --git a/Engine/Plugins/FX/Niagara/Source/NiagaraEditor/Private/Widgets/Scalability/SNiagaraScalabilityContext.h b/Engine/Plugins/FX/Niagara/Source/NiagaraEditor/Private/Widgets/Scalability/SNiagaraScalabilityContext.h
new file mode 100644
--- /dev/null
+++ b/Engine/Plugins/FX/Niagara/Source/NiagaraEditor/Private/Widgets/Scalability/SNiagaraScalabilityContext.h
@@ -0,0 +1,29 @@
+// Copyright Epic Games, Inc. All Rights Reserved.
+
+#pragma once
+
+#include "CoreMinimal.h"
+#include "IDetailCustomization.h"
+#include "NiagaraSystem.h"
+#include "ViewModels/NiagaraSystemScalabilityViewModel.h"
+#include "Widgets/Layout/SGridPanel.h"
+
+class SNiagaraScalabilityContext : public SCompoundWidget
+{
+public:
+	
+	SLATE_BEGIN_ARGS(SNiagaraScalabilityContext)
+	{}
+	SLATE_END_ARGS()
+
+	void Construct(const FArguments& InArgs, UNiagaraSystemScalabilityViewModel& InScalabilityViewModel);
+
+	void SetObject(UObject* Object);
+
+	void UpdateScalabilityContent();
+
+private:
+	bool FilterScalabilityProperties(const FPropertyAndParent& InPropertyAndParent) const;
+	TWeakObjectPtr<UNiagaraSystemScalabilityViewModel> ScalabilityViewModel;
+	TSharedPtr<IDetailsView> DetailsView;
+};
diff --git a/Engine/Plugins/FX/Niagara/Source/NiagaraEditor/Private/Widgets/Scalability/SNiagaraSystemResolvedScalabilitySettings.cpp b/Engine/Plugins/FX/Niagara/Source/NiagaraEditor/Private/Widgets/Scalability/SNiagaraSystemResolvedScalabilitySettings.cpp
new file mode 100644
--- /dev/null
+++ b/Engine/Plugins/FX/Niagara/Source/NiagaraEditor/Private/Widgets/Scalability/SNiagaraSystemResolvedScalabilitySettings.cpp
@@ -0,0 +1,462 @@
+// Copyright Epic Games, Inc. All Rights Reserved.
+
+#include "SNiagaraSystemResolvedScalabilitySettings.h"
+
+#include "IDetailTreeNode.h"
+#include "IPropertyRowGenerator.h"
+#include "ISinglePropertyView.h"
+#include "NiagaraEditorStyle.h"
+#include "NiagaraObjectSelection.h"
+#include "Modules/ModuleManager.h"
+#include "Styling/SlateIconFinder.h"
+#include "ViewModels/NiagaraEmitterHandleViewModel.h"
+#include "ViewModels/NiagaraOverviewGraphViewModel.h"
+#include "ViewModels/NiagaraSystemViewModel.h"
+#include "Widgets/Input/SButton.h"
+#include "Widgets/Input/SExpandableButton.h"
+#include "Widgets/Layout/SExpandableArea.h"
+#include "Widgets/Layout/SSeparator.h"
+#include "Widgets/Text/SInlineEditableTextBlock.h"
+#include "PropertyEditorModule.h"
+
+#define LOCTEXT_NAMESPACE "NiagaraScalabilitySystem"
+
+static const TCHAR* ColumnPropertyName = TEXT("Scalability Setting");
+static const TCHAR* ColumnDefaultValue = TEXT("Default");
+static const TCHAR* ColumnResolvedValue = TEXT("Override");
+static const TCHAR* ColumnLocate = TEXT("Locate");
+
+
+class FNiagaraEmptyPlatformCustomization : public IPropertyTypeCustomization
+{
+public:
+	static TSharedRef<IPropertyTypeCustomization> MakeInstance();
+	
+	virtual void CustomizeHeader(TSharedRef<IPropertyHandle> PropertyHandle, FDetailWidgetRow& HeaderRow,
+		IPropertyTypeCustomizationUtils& CustomizationUtils) override;
+	virtual void CustomizeChildren(TSharedRef<IPropertyHandle> PropertyHandle, IDetailChildrenBuilder& ChildBuilder,
+		IPropertyTypeCustomizationUtils& CustomizationUtils) override;
+};
+
+TSharedRef<IPropertyTypeCustomization> FNiagaraEmptyPlatformCustomization::MakeInstance()
+{
+	return MakeShared<FNiagaraEmptyPlatformCustomization>();
+}
+
+void FNiagaraEmptyPlatformCustomization::CustomizeHeader(TSharedRef<IPropertyHandle> PropertyHandle,
+                                                         FDetailWidgetRow& HeaderRow, IPropertyTypeCustomizationUtils& CustomizationUtils)
+{
+}
+
+void FNiagaraEmptyPlatformCustomization::CustomizeChildren(TSharedRef<IPropertyHandle> PropertyHandle,
+                                                           IDetailChildrenBuilder& ChildBuilder, IPropertyTypeCustomizationUtils& CustomizationUtils)
+{
+}
+
+bool IsEmitterScalabilityValueOverridden(FName PropertyName, const FNiagaraEmitterScalabilityOverride& ActiveOverrides)
+{
+	if(PropertyName == GET_MEMBER_NAME_CHECKED(FNiagaraEmitterScalabilityOverride, SpawnCountScale))
+	{
+		return ActiveOverrides.bOverrideSpawnCountScale;
+	}
+
+	// if new overrides get added, this function has to be extended
+	check(false);
+	return false;
+}
+
+bool IsSystemScalabilityValueOverridden(FName PropertyName, const FNiagaraSystemScalabilityOverride& ActiveOverrides)
+{
+	if(PropertyName == GET_MEMBER_NAME_CHECKED(FNiagaraSystemScalabilitySettings, MaxDistance))
+	{
+		return ActiveOverrides.bOverrideDistanceSettings;
+	}
+	else if(PropertyName == GET_MEMBER_NAME_CHECKED(FNiagaraSystemScalabilitySettings, MaxInstances))
+	{
+		return ActiveOverrides.bOverrideInstanceCountSettings;
+	}
+	else if(PropertyName == GET_MEMBER_NAME_CHECKED(FNiagaraSystemScalabilitySettings, MaxSystemInstances))
+	{
+		return ActiveOverrides.bOverridePerSystemInstanceCountSettings;
+	}
+	else if(PropertyName == GET_MEMBER_NAME_CHECKED(FNiagaraSystemScalabilitySettings, MaxTimeWithoutRender))
+	{
+		return ActiveOverrides.bOverrideTimeSinceRendererSettings;
+	}
+	// cull proxy mode and max system proxies are controlled via the same override bool
+	else if(PropertyName == GET_MEMBER_NAME_CHECKED(FNiagaraSystemScalabilitySettings, CullProxyMode))
+	{
+		return ActiveOverrides.bOverrideCullProxySettings;
+	}
+	else if(PropertyName == GET_MEMBER_NAME_CHECKED(FNiagaraSystemScalabilitySettings, MaxSystemProxies))
+	{
+		return ActiveOverrides.bOverrideCullProxySettings;
+	}
+	else if(PropertyName == GET_MEMBER_NAME_CHECKED(FNiagaraSystemScalabilitySettings, BudgetScaling))
+	{
+		return ActiveOverrides.bOverrideGlobalBudgetScalingSettings;
+	}
+
+	// if new overrides get added, this function has to be extended
+	check(false);
+	return false;
+}
+
+void SScalabilityResolvedRow::Construct(const FArguments& InArgs, TSharedRef<FScalabilityRowData> InScalabilityRowData, const TSharedRef<STableViewBase>& OwnerTableView, const TSharedRef<SNiagaraSystemResolvedScalabilitySettings>& InParentWidget)
+{
+	ScalabilityRowData = InScalabilityRowData;
+	ParentWidget = InParentWidget;
+
+	auto Args = FSuperRowType::FArguments()
+		.Style(&FEditorStyle::Get().GetWidgetStyle<FTableRowStyle>("SceneOutliner.TableViewRow"));
+	SMultiColumnTableRow<TSharedPtr<FScalabilityRowData>>::Construct(Args, OwnerTableView);
+}
+
+TSharedRef<SWidget> SScalabilityResolvedRow::GenerateWidgetForColumn(const FName& ColumnId)
+{
+	auto IconColorActive = [&](bool bActive) -> FLinearColor
+	{
+		return bActive ? FNiagaraEditorStyle::Get().GetColor("NiagaraEditor.Scalability.System.Feature.Active") : FNiagaraEditorStyle::Get().GetColor("NiagaraEditor.Scalability.System.Feature.Inactive");
+	};
+
+	TSharedPtr<SWidget> RowContent;
+	if(ColumnId == ColumnPropertyName)
+	{
+		RowContent = SNew(SHorizontalBox)
+		+ SHorizontalBox::Slot()
+		.VAlign(VAlign_Center)
+		[
+			SNew(STextBlock).Text(FText::FromName(ScalabilityRowData->ScalabilityAttribute)).ColorAndOpacity(IconColorActive(ScalabilityRowData->ResolvedTreeNode->CreatePropertyHandle()->IsEditable()))
+		];
+	}
+	else if(ColumnId == ColumnDefaultValue)
+	{		
+		FText DefaultValueText;
+		if(ScalabilityRowData->DefaultPropertyHandle.IsValid())
+		{
+			ScalabilityRowData->DefaultPropertyHandle->GetValueAsDisplayText(DefaultValueText);			
+			RowContent = SNew(SHorizontalBox)
+			+ SHorizontalBox::Slot()
+			.VAlign(VAlign_Center)
+			[
+				SNew(STextBlock).Text(DefaultValueText).ColorAndOpacity(IconColorActive(ScalabilityRowData->DefaultPropertyHandle->IsEditable()))
+			];
+		}
+		else
+		{
+			RowContent = SNullWidget::NullWidget;
+		}
+	}
+	// for the override value, we only add an entry if an override is active 
+	else if(ColumnId == ColumnResolvedValue)
+	{
+		bool bSkipWidget = false;
+		if(UNiagaraEmitter* Emitter = Cast<UNiagaraEmitter>(ScalabilityRowData->OwningObject))
+		{
+			const FNiagaraEmitterScalabilityOverride& CurrentOverrides = Emitter->GetCurrentOverrideSettings();
+			if(!IsEmitterScalabilityValueOverridden(ScalabilityRowData->ResolvedTreeNode->GetNodeName(), CurrentOverrides))
+			{
+				RowContent = SNullWidget::NullWidget;
+				bSkipWidget = true;
+			}				
+		}
+		else if(UNiagaraSystem* System = Cast<UNiagaraSystem>(ScalabilityRowData->OwningObject))
+		{
+			const FNiagaraSystemScalabilityOverride& CurrentOverrides = System->GetCurrentOverrideSettings();
+			if(!IsSystemScalabilityValueOverridden(ScalabilityRowData->ResolvedTreeNode->GetNodeName(), CurrentOverrides))
+			{
+				RowContent = SNullWidget::NullWidget;
+				bSkipWidget = true;
+			}
+		}
+
+		if(!bSkipWidget)
+		{
+			FText ResolvedValueText;
+			ScalabilityRowData->ResolvedPropertyHandle->GetValueAsDisplayText(ResolvedValueText);
+			
+			RowContent = SNew(SHorizontalBox)
+			+ SHorizontalBox::Slot()
+			.AutoWidth()
+			.Padding(2.f)
+			.VAlign(VAlign_Center)
+			[
+				SNew(SImage).Image(FNiagaraEditorStyle::Get().GetBrush("NiagaraEditor.Scalability.System.Feature.Override"))
+			]
+			+ SHorizontalBox::Slot()
+			.VAlign(VAlign_Center)
+			[
+				SNew(STextBlock).Text(ResolvedValueText).ColorAndOpacity(IconColorActive(ScalabilityRowData->ResolvedPropertyHandle->IsEditable()))
+			];			
+		}
+	}
+	// @Todo this won't be called currently as we removed the column from the header row
+	else if(ColumnId == ColumnLocate)
+	{
+		RowContent = SNew(SButton).OnClicked(this, &SScalabilityResolvedRow::NavigateToValue, ScalabilityRowData);
+	}
+
+	return SNew(SBox).MinDesiredHeight(24.f)
+		[
+			RowContent.ToSharedRef()
+		];
+}
+
+FReply SScalabilityResolvedRow::NavigateToValue(TSharedPtr<FScalabilityRowData> RowData)
+{
+	if(!RowData.IsValid())
+	{
+		return FReply::Handled();
+	}
+	
+	//ParentWidget.Pin()->GetSystemViewModel()->GetScalabilityViewModel()->NavigateToScalabilityProperty(RowData->OwningObject, FName(RowData->DefaultPropertyHandle->GeneratePathToProperty()));
+	return FReply::Unhandled();
+}
+
+void SNiagaraSystemResolvedScalabilitySettings::RebuildWidget()
+{
+	ResolvedScalabilityContainer->ClearChildren();
+	
+	const TSet<UObject*>& SelectedNodes = SystemViewModel->GetOverviewGraphViewModel()->GetNodeSelection()->GetSelectedObjects();
+
+	if(SelectedNodes.Num() > 0)
+	{
+		for(UObject* SelectedNode : SelectedNodes)
+		{
+			UNiagaraOverviewNode* OverviewNode = CastChecked<UNiagaraOverviewNode>(SelectedNode);
+			if(FNiagaraEmitterHandle* Handle = OverviewNode->TryGetEmitterHandle())
+			{
+				Handle->GetInstance()->OnPropertiesChanged().RemoveAll(this);
+				Handle->GetInstance()->OnPropertiesChanged().AddSP(this, &SNiagaraSystemResolvedScalabilitySettings::RebuildWidget);
+				ResolvedScalabilityContainer->AddSlot()
+				.AutoHeight()
+				[
+					GenerateResolvedScalabilityTable(Handle->GetInstance())
+				];
+			}
+			else
+			{
+				ResolvedScalabilityContainer->AddSlot()
+				.AutoHeight()
+				[
+					GenerateResolvedScalabilityTable(OverviewNode->GetOwningSystem())
+				];
+			}
+		}
+	}
+	else
+	{
+		ResolvedScalabilityContainer->AddSlot()
+		[
+			GenerateResolvedScalabilityTable(&SystemViewModel->GetSystem())
+		];
+		
+		for(const TSharedRef<FNiagaraEmitterHandleViewModel>& EmitterHandleViewModel : SystemViewModel->GetEmitterHandleViewModels())
+		{
+			EmitterHandleViewModel->GetEmitterHandle()->GetInstance()->OnPropertiesChanged().RemoveAll(this);
+			EmitterHandleViewModel->GetEmitterHandle()->GetInstance()->OnPropertiesChanged().AddSP(this, &SNiagaraSystemResolvedScalabilitySettings::RebuildWidget);
+			ResolvedScalabilityContainer->AddSlot()
+			[
+				GenerateResolvedScalabilityTable(EmitterHandleViewModel->GetEmitterHandle()->GetInstance())
+			];	
+		}
+	}
+	
+	ChildSlot
+	[
+		ResolvedScalabilityContainer.ToSharedRef()
+	];
+}
+
+TSharedRef<ITableRow> SNiagaraSystemResolvedScalabilitySettings::GenerateScalabilityValueRow(TSharedRef<FScalabilityRowData> ScalabilityRowData, const TSharedRef<STableViewBase>& TableViewBase)
+{
+	return SNew(SScalabilityResolvedRow, ScalabilityRowData, TableViewBase, SharedThis(this));
+}
+
+TSharedRef<SWidget> SNiagaraSystemResolvedScalabilitySettings::GenerateResolvedScalabilityTable(UObject* Object)
+{
+	FString DisplayName = Object->GetName();
+
+	FPropertyEditorModule& EditorModule = FModuleManager::Get().LoadModuleChecked<FPropertyEditorModule>("PropertyEditor");
+	FPropertyRowGeneratorArgs Args;
+	TSharedRef<IPropertyRowGenerator> ResolvedPropertyGenerator = EditorModule.CreatePropertyRowGenerator(Args);
+	ResolvedPropertyGenerator->RegisterInstancedCustomPropertyTypeLayout(FName(FNiagaraPlatformSet::StaticStruct()->GetName()), FOnGetPropertyTypeCustomizationInstance::CreateStatic(FNiagaraEmptyPlatformCustomization::MakeInstance));
+
+	TSharedRef<IPropertyRowGenerator> DefaultPropertyGenerator = EditorModule.CreatePropertyRowGenerator(Args);
+	DefaultPropertyGenerator->RegisterInstancedCustomPropertyTypeLayout(FName(FNiagaraPlatformSet::StaticStruct()->GetName()), FOnGetPropertyTypeCustomizationInstance::CreateStatic(FNiagaraEmptyPlatformCustomization::MakeInstance));
+
+	TSharedPtr<FStructOnScope> DefaultValueStruct = nullptr;
+	TSharedPtr<FStructOnScope> ResolvedValueStruct = nullptr;
+
+	if(UNiagaraEmitter* Emitter = Cast<UNiagaraEmitter>(Object))
+	{
+		if(System->GetEffectType())
+		{
+			const FNiagaraEmitterScalabilitySettings& DefaultEmitterScalabilitySettings = System->GetEffectType()->GetActiveEmitterScalabilitySettings();
+			DefaultValueStruct = MakeShared<FStructOnScope>(FNiagaraEmitterScalabilitySettings::StaticStruct(), (uint8*) &DefaultEmitterScalabilitySettings);			
+		}
+		const FNiagaraEmitterScalabilitySettings& EmitterScalabilitySettings = Emitter->GetScalabilitySettings();
+		ResolvedValueStruct = MakeShared<FStructOnScope>(FNiagaraEmitterScalabilitySettings::StaticStruct(), (uint8*) &EmitterScalabilitySettings);
+	}
+	else if(UNiagaraSystem* NiagaraSystem = Cast<UNiagaraSystem>(Object))
+	{
+		if(UNiagaraEffectType* EffectType = NiagaraSystem->GetEffectType())
+		{
+			const FNiagaraSystemScalabilitySettings& DefaultSystemScalabilitySettings = EffectType->GetActiveSystemScalabilitySettings();
+			DefaultValueStruct = MakeShared<FStructOnScope>(FNiagaraSystemScalabilitySettings::StaticStruct(), (uint8*) &DefaultSystemScalabilitySettings);			
+		}
+		const FNiagaraSystemScalabilitySettings& SystemScalabilitySettings = NiagaraSystem->GetScalabilitySettings();
+		ResolvedValueStruct = MakeShared<FStructOnScope>(FNiagaraSystemScalabilitySettings::StaticStruct(), (uint8*) &SystemScalabilitySettings);
+	}
+
+	ensure(ResolvedValueStruct.IsValid());
+	
+    ResolvedPropertyGenerator->SetStructure(ResolvedValueStruct);
+
+	// we are caching the default handles so we can later find the default handle corresponding to a resolved handle
+	TArray<TSharedRef<IDetailTreeNode>> DefaultTreeNodes;
+	TMap<FString, TSharedRef<IDetailTreeNode>> DefaultHandles;
+	if(DefaultValueStruct.IsValid())
+	{
+		DefaultPropertyGenerator->SetStructure(DefaultValueStruct);
+		DefaultTreeNodes = DefaultPropertyGenerator->GetRootTreeNodes();
+		
+		for(TSharedRef<IDetailTreeNode> RootNode : DefaultTreeNodes)
+		{
+			if(RootNode->GetNodeName().IsEqual("Scalability"))
+			{
+				TArray<TSharedRef<IDetailTreeNode>> Children;
+				RootNode->GetChildren(Children);
+
+				for(const TSharedRef<IDetailTreeNode>& Child : Children)
+				{
+					TSharedPtr<IPropertyHandle> ChildPropertyHandle = Child->CreatePropertyHandle();
+				
+					if(ChildPropertyHandle->HasMetaData("DisplayInScalabilityValuesBar"))
+					{
+						DefaultHandles.Add(ChildPropertyHandle->GetProperty()->GetName(), Child);
+					}
+				}
+			}
+		}
+	}	
+
+	// we now create the scalability row data necessary to create a list view
+	TArray<TSharedRef<FScalabilityRowData>> ScalabilityValues;
+	const TArray<TSharedRef<IDetailTreeNode>> RootTreeNodes = ResolvedPropertyGenerator->GetRootTreeNodes();
+	for(TSharedRef<IDetailTreeNode> RootNode : RootTreeNodes)
+	{
+		if(RootNode->GetNodeName().IsEqual("Scalability"))
+		{
+			TArray<TSharedRef<IDetailTreeNode>> Children;
+			RootNode->GetChildren(Children);
+
+			for(const TSharedRef<IDetailTreeNode>& Child : Children)
+			{
+				TSharedPtr<IPropertyHandle> ChildPropertyHandle = Child->CreatePropertyHandle();
+				
+				if(ChildPropertyHandle->HasMetaData("DisplayInScalabilityValuesBar"))
+				{
+					TSharedRef<FScalabilityRowData> ScalabilityValue = MakeShared<FScalabilityRowData>();
+					ScalabilityValue->ScalabilityAttribute = FName(ChildPropertyHandle->GetPropertyDisplayName().ToString());
+					ScalabilityValue->ResolvedPropertyHandle = ChildPropertyHandle.ToSharedRef();
+					ScalabilityValue->ResolvedTreeNode = Child;
+					ScalabilityValue->OwningObject = Object;
+
+					// if we found the same handle in the default handles, we have both handles needed to display default and override data
+					if(DefaultHandles.Contains(ChildPropertyHandle->GetProperty()->GetName()))
+					{
+						ScalabilityValue->DefaultTreeNode = DefaultHandles[ChildPropertyHandle->GetProperty()->GetName()];
+						ScalabilityValue->DefaultPropertyHandle = DefaultHandles[ChildPropertyHandle->GetProperty()->GetName()]->CreatePropertyHandle();
+					}
+					
+					// @todo
+					// ScalabilityValue->SourceProperties = ;
+					ScalabilityValues.Add(ScalabilityValue);
+				}
+			}
+		}
+	}
+
+	// we cache the objects to maintain their lifetimes while used
+	FRequiredInstanceInformation InstanceInfo;
+	InstanceInfo.DefaultPropertyRowGenerator = DefaultPropertyGenerator;
+	InstanceInfo.ResolvedPropertyRowGenerator = ResolvedPropertyGenerator;
+	InstanceInfo.ScalabilityValues = ScalabilityValues;
+	InstanceInformation.Add(Object, InstanceInfo);
+	const FSlateBrush* AssetIcon = FSlateIconFinder::FindIconBrushForClass(Object->GetClass());
+	
+	TSharedRef<SWidget> Widget = SNew(SExpandableArea)
+	.HeaderContent()
+	[
+		SNew(SHorizontalBox)
+		+ SHorizontalBox::Slot()
+		.AutoWidth()
+		.Padding(3.f)
+		[
+			SNew(SImage)
+			.Image(AssetIcon)
+		]
+		+ SHorizontalBox::Slot()
+		[
+			SNew(SInlineEditableTextBlock)
+			.Text(FText::FromString(DisplayName))
+			.Style(FAppStyle::Get(), "DetailsView.NameTextBlockStyle")
+			.IsReadOnly(true)
+		]
+	]
+	.BodyContent()
+	[	
+		SNew(SBorder)
+		.BorderImage(FEditorStyle::GetBrush(TEXT("Graph.TitleBackground")))
+		[
+			SNew(SListView<TSharedRef<FScalabilityRowData>>)
+			.ListItemsSource(&InstanceInformation[Object].ScalabilityValues)
+			.OnGenerateRow(this, &SNiagaraSystemResolvedScalabilitySettings::GenerateScalabilityValueRow)
+			.HeaderRow
+			(
+				SNew(SHeaderRow)
+                + SHeaderRow::Column(ColumnPropertyName)
+                	.DefaultLabel(LOCTEXT("ScalabilitySettingLabel", "Scalability Setting"))
+                	.DefaultTooltip(LOCTEXT("ScalabilitySettingLabelTooltip", "Scalability Settings can be turned on or off and generally can be overridden."))
+                	.FillWidth(0.5f)
+                + SHeaderRow::Column(ColumnDefaultValue)
+                	.DefaultLabel(LOCTEXT("ScalabilityDefaultValueLabel", "Default"))
+                	.DefaultTooltip(LOCTEXT("ScalabilityDefaultValueLabelTooltip", "Default values come from the effect type assigned to a system."))
+                	.FillWidth(0.2f)
+                + SHeaderRow::Column(ColumnResolvedValue)
+                	.DefaultLabel(LOCTEXT("ScalabilityOverrideValueLabel", "Override"))
+                	.DefaultTooltip(LOCTEXT("ScalabilityOverrideValueLabelTooltip", "Overrides can be specified in a system or emitter to override an effect type's default scalability settings."))
+                	.FillWidth(0.25f)
+                	// @Todo: Add locate functionality. This is involved as the active scalability tab has no functionality yet to select certain rows
+                	//+ SHeaderRow::Column(ColumnLocate)
+					// .DefaultLabel(LOCTEXT("ScalabilityLocateValueLabel", ""))
+					// .DefaultTooltip(LOCTEXT("BlueprintWarningBehaviorHeaderTooltip", "Determines what happens when the warning is raised - warnings can be treated more strictly or suppressed entirely"))
+					// .FixedWidth(20.f)
+            )
+			
+		]
+	];
+
+	return Widget;
+}
+
+SNiagaraSystemResolvedScalabilitySettings::~SNiagaraSystemResolvedScalabilitySettings()
+{
+	System->OnScalabilityChanged().RemoveAll(this);
+}
+
+void SNiagaraSystemResolvedScalabilitySettings::Construct(const FArguments& InArgs, UNiagaraSystem& InSystem, TSharedPtr<FNiagaraSystemViewModel> InSystemViewModel)
+{
+	System = &InSystem;
+	SystemViewModel = InSystemViewModel;
+
+	ResolvedScalabilityContainer = SNew(SVerticalBox);
+
+	SystemViewModel->GetOverviewGraphViewModel()->GetNodeSelection()->OnSelectedObjectsChanged().AddSP(this, &SNiagaraSystemResolvedScalabilitySettings::RebuildWidget);	
+	System->OnScalabilityChanged().AddSP(this, &SNiagaraSystemResolvedScalabilitySettings::RebuildWidget);
+	
+	RebuildWidget();
+}
+
+#undef LOCTEXT_NAMESPACE
diff --git a/Engine/Plugins/FX/Niagara/Source/NiagaraEditor/Private/Widgets/Scalability/SNiagaraSystemResolvedScalabilitySettings.h b/Engine/Plugins/FX/Niagara/Source/NiagaraEditor/Private/Widgets/Scalability/SNiagaraSystemResolvedScalabilitySettings.h
new file mode 100644
--- /dev/null
+++ b/Engine/Plugins/FX/Niagara/Source/NiagaraEditor/Private/Widgets/Scalability/SNiagaraSystemResolvedScalabilitySettings.h
@@ -0,0 +1,68 @@
+// Copyright Epic Games, Inc. All Rights Reserved.
+
+#pragma once
+
+#include "CoreMinimal.h"
+#include "NiagaraSystem.h"
+#include "ViewModels/NiagaraSystemViewModel.h"
+#include "Widgets/Layout/SGridPanel.h"
+
+struct FScalabilityRowData
+{
+	FName ScalabilityAttribute;
+	TSharedPtr<IDetailTreeNode> ResolvedTreeNode;
+	TSharedPtr<IDetailTreeNode> DefaultTreeNode;
+	TSharedPtr<IPropertyHandle> ResolvedPropertyHandle;
+	TSharedPtr<IPropertyHandle> DefaultPropertyHandle;
+	TArray<FProperty*> SourceProperties;
+	UObject* OwningObject;
+};
+
+class SScalabilityResolvedRow : public SMultiColumnTableRow<TSharedPtr<FScalabilityRowData>>
+{
+public:
+	SLATE_BEGIN_ARGS(SScalabilityResolvedRow) {}
+	SLATE_END_ARGS()
+	
+	void Construct(const FArguments& InArgs, TSharedRef<FScalabilityRowData> InScalabilityRowData, const TSharedRef<STableViewBase>& OwnerTableView, const TSharedRef<class SNiagaraSystemResolvedScalabilitySettings>& ParentWidget);
+	
+	/** SMultiColumnTableRow interface */
+	virtual TSharedRef<SWidget> GenerateWidgetForColumn(const FName& ColumnId) override;
+
+	FReply NavigateToValue(TSharedPtr<FScalabilityRowData> RowData);
+private:
+	TSharedPtr<FScalabilityRowData> ScalabilityRowData;
+	TWeakPtr<class SNiagaraSystemResolvedScalabilitySettings> ParentWidget;
+};
+
+class SNiagaraSystemResolvedScalabilitySettings : public SCompoundWidget
+{
+	// struct used to maintain lifetime
+	struct FRequiredInstanceInformation
+	{
+		TSharedPtr<IPropertyRowGenerator> DefaultPropertyRowGenerator;
+		TSharedPtr<IPropertyRowGenerator> ResolvedPropertyRowGenerator;
+		TArray<TSharedRef<FScalabilityRowData>> ScalabilityValues;
+	};
+public:
+	SLATE_BEGIN_ARGS(SNiagaraSystemResolvedScalabilitySettings)
+	{}
+	SLATE_END_ARGS()
+
+	virtual ~SNiagaraSystemResolvedScalabilitySettings();
+	
+	void Construct(const FArguments& InArgs, UNiagaraSystem& InSystem, TSharedPtr<FNiagaraSystemViewModel> SystemViewModel);
+	void RebuildWidget();
+	TSharedPtr<FNiagaraSystemViewModel> GetSystemViewModel() { return SystemViewModel; }
+private:
+	TWeakObjectPtr<UNiagaraSystem> System;
+	TSharedPtr<FNiagaraSystemViewModel> SystemViewModel;
+	TSharedPtr<SVerticalBox> ResolvedScalabilityContainer;
+private:
+	TMap<UObject*, FRequiredInstanceInformation> InstanceInformation;
+	
+	/** Generates an expandable widget for an emitter or system */
+	TSharedRef<SWidget> GenerateResolvedScalabilityTable(UObject* Object);
+	
+	TSharedRef<ITableRow> GenerateScalabilityValueRow(TSharedRef<FScalabilityRowData> ScalabilityRowData, const TSharedRef< class STableViewBase >& TableViewBase);
+};
diff --git a/Engine/Plugins/FX/Niagara/Source/NiagaraEditor/Public/NiagaraEditorUtilities.h b/Engine/Plugins/FX/Niagara/Source/NiagaraEditor/Public/NiagaraEditorUtilities.h
--- a/Engine/Plugins/FX/Niagara/Source/NiagaraEditor/Public/NiagaraEditorUtilities.h
+++ b/Engine/Plugins/FX/Niagara/Source/NiagaraEditor/Public/NiagaraEditorUtilities.h
@@ -12,6 +12,7 @@
 #include "NiagaraEditorSettings.h"
 #include "UpgradeNiagaraScriptResults.h"
 #include "EdGraph/EdGraphSchema.h"
+#include "ViewModels/NiagaraSystemScalabilityViewModel.h"
 
 class UNiagaraNodeInput;
 class UNiagaraNodeOutput;
@@ -335,6 +336,8 @@ namespace FNiagaraEditorUtilities
 
 	DECLARE_DELEGATE_OneParam(FNodeVisitor, UEdGraphNode* /*VisitedNode*/);
 	void VisitAllNodesConnectedToInputs(UEdGraphNode* StartNode, FNodeVisitor Visitor);
+	
+	NIAGARAEDITOR_API float GetScalabilityTintAlpha(FNiagaraEmitterHandle* EmitterHandle);
 };
 
 namespace FNiagaraParameterUtilities
diff --git a/Engine/Plugins/FX/Niagara/Source/NiagaraEditor/Public/NiagaraOverviewNode.h b/Engine/Plugins/FX/Niagara/Source/NiagaraEditor/Public/NiagaraOverviewNode.h
--- a/Engine/Plugins/FX/Niagara/Source/NiagaraEditor/Public/NiagaraOverviewNode.h
+++ b/Engine/Plugins/FX/Niagara/Source/NiagaraEditor/Public/NiagaraOverviewNode.h
@@ -16,7 +16,8 @@ public:
 	void Initialize(UNiagaraSystem* InOwningSystem);
 	void Initialize(UNiagaraSystem* InOwningSystem, FGuid InEmitterHandleGuid);
 	const FGuid GetEmitterHandleGuid() const;
-
+	struct FNiagaraEmitterHandle* TryGetEmitterHandle();
+	
 	//~ Begin UEdGraphNode Interface
 	/** Gets the name of this node, shown in title bar */
 	virtual FText GetNodeTitle(ENodeTitleType::Type TitleType) const override;
diff --git a/Engine/Plugins/FX/Niagara/Source/NiagaraEditor/Public/ViewModels/NiagaraSystemScalabilityViewModel.h b/Engine/Plugins/FX/Niagara/Source/NiagaraEditor/Public/ViewModels/NiagaraSystemScalabilityViewModel.h
new file mode 100644
--- /dev/null
+++ b/Engine/Plugins/FX/Niagara/Source/NiagaraEditor/Public/ViewModels/NiagaraSystemScalabilityViewModel.h
@@ -0,0 +1,65 @@
+// Copyright Epic Games, Inc. All Rights Reserved.
+
+#pragma once
+
+#include "IMovieScenePlayer.h"
+
+#include "NiagaraCommon.h"
+#include "NiagaraEditorCommon.h"
+#include "EditorUndoClient.h"
+#include "UObject/GCObject.h"
+#include "ViewModels/TNiagaraViewModelManager.h"
+#include "ISequencerModule.h"
+#include "NiagaraPlatformSet.h"
+#include "NiagaraSystemScalabilityViewModel.generated.h"
+
+class FNiagaraSystemViewModel;
+
+/** A view model for viewing and editing a UNiagaraSystem. */
+
+UCLASS()
+class NIAGARAEDITOR_API UNiagaraSystemScalabilityViewModel : public UObject
+{
+	GENERATED_BODY()
+public:
+	DECLARE_MULTICAST_DELEGATE_OneParam(FScalabilityModeChanged, bool /* Active */)
+	DECLARE_MULTICAST_DELEGATE_OneParam(FScalabilityPropertySelected, FName)
+
+public:
+
+	/** Creates a new view model with the supplied System and System instance. */
+	UNiagaraSystemScalabilityViewModel();
+	
+	/** Initializes this scalability view model with the supplied system view model. */
+	void Initialize(TSharedRef<FNiagaraSystemViewModel> InSystemViewModel);
+
+	/** Returns whether or not this view model is initialized and safe to use. */
+	bool IsValid() const;
+	
+	bool IsActive() const;
+
+	TWeakPtr<FNiagaraSystemViewModel> GetSystemViewModel() const { return SystemViewModel; }
+	FNiagaraPlatformSet* GetPreviewPlatformSet() { return PreviewPlatforms.Get(); }
+	const FNiagaraPlatformSet* GetPreviewPlatformSet() const { return PreviewPlatforms.Get(); }
+	
+	int32 GetPreviewQualityLevelMask() const { return PreviewPlatforms->QualityLevelMask; }
+	int32 GetPreviewQualityLevel() const { return FNiagaraPlatformSet::QualityLevelFromMask(PreviewPlatforms->QualityLevelMask); }
+	TOptional<TObjectPtr<UDeviceProfile>> GetPreviewDeviceProfile() const { return PreviewDeviceProfile; }
+	bool IsViewModeQualityEnabled(int32 QualityLevel) const { return PreviewPlatforms->IsEffectQualityEnabled(QualityLevel); }
+	
+	void UpdatePreviewDeviceProfile(UDeviceProfile* DeviceProfile);
+	void UpdatePreviewQualityLevel(int32 QualityLevel, bool bEnabled);
+
+	//void NavigateToScalabilityProperty(UObject* Object, FName PropertyName);
+	
+private:
+	TWeakPtr<FNiagaraSystemViewModel> SystemViewModel;
+	TSharedPtr<FNiagaraPlatformSet> PreviewPlatforms;
+	TOptional<TObjectPtr<UDeviceProfile>> PreviewDeviceProfile;
+public:
+	FScalabilityModeChanged& OnScalabilityModeChanged() { return ScalabilityModeChangedDelegate; }
+	FScalabilityPropertySelected& OnScalabilityPropertySelected() { return ScalabilityPropertySelectedDelegate; }
+private:
+	FScalabilityModeChanged ScalabilityModeChangedDelegate;
+	FScalabilityPropertySelected ScalabilityPropertySelectedDelegate;
+};
diff --git a/Engine/Plugins/FX/Niagara/Source/NiagaraEditor/Public/ViewModels/NiagaraSystemViewModel.h b/Engine/Plugins/FX/Niagara/Source/NiagaraEditor/Public/ViewModels/NiagaraSystemViewModel.h
--- a/Engine/Plugins/FX/Niagara/Source/NiagaraEditor/Public/ViewModels/NiagaraSystemViewModel.h
+++ b/Engine/Plugins/FX/Niagara/Source/NiagaraEditor/Public/ViewModels/NiagaraSystemViewModel.h
@@ -15,6 +15,7 @@
 
 #include "TickableEditorObject.h"
 #include "ISequencerModule.h"
+#include "NiagaraSystemScalabilityViewModel.h"
 #include "UObject/ObjectKey.h"
 
 class UNiagaraSystem;
@@ -59,6 +60,8 @@ enum class NIAGARAEDITOR_API ENiagaraSystemViewModelEditMode
 	EmitterDuringMerge,
 };
 
+DECLARE_DELEGATE_RetVal(FName, FOnGetWorkflowMode);
+
 /** Defines options for the niagara System view model */
 struct NIAGARAEDITOR_API FNiagaraSystemViewModelOptions
 {
@@ -205,6 +208,12 @@ public:
 	/** Gets the current editing mode for this system view model. */
 	NIAGARAEDITOR_API ENiagaraSystemViewModelEditMode GetEditMode() const;
 
+	/** Gets the current workflow mode for this system view model. */
+	NIAGARAEDITOR_API FName GetWorkflowMode() const;
+
+	/** Gets the delegate which is called any time the array of emitter handle view models changes. */
+	NIAGARAEDITOR_API FOnGetWorkflowMode& OnGetWorkflowMode();
+	
 	/** Adds a new emitter to the System from an emitter asset data. */
 	NIAGARAEDITOR_API TSharedPtr<FNiagaraEmitterHandleViewModel> AddEmitterFromAssetData(const FAssetData& AssetData);
 
@@ -350,6 +359,8 @@ public:
 
 	NIAGARAEDITOR_API UNiagaraCurveSelectionViewModel* GetCurveSelectionViewModel();
 
+	NIAGARAEDITOR_API UNiagaraSystemScalabilityViewModel* GetScalabilityViewModel();
+
 	TArray<float> OnGetPlaybackSpeeds() const;
 	
 	/** Duplicates a set of emitters and refreshes everything.*/
@@ -565,6 +576,9 @@ private:
 	/** The current editing mode for this view model. */
 	ENiagaraSystemViewModelEditMode EditMode;
 
+	/** A delegate to retrieve the current workflow mode from the toolkit. */
+	FOnGetWorkflowMode OnGetWorkflowModeDelegate;
+	
 	/** A delegate which is used to generate the content for the add menu in sequencer. */
 	FOnGetAddMenuContent OnGetSequencerAddMenuContent;
 
@@ -649,6 +663,8 @@ private:
 
 	UNiagaraCurveSelectionViewModel* CurveSelectionViewModel;
 
+	UNiagaraSystemScalabilityViewModel* ScalabilityViewModel;
+
 	TSharedPtr<FNiagaraPlaceholderDataInterfaceManager> PlaceholderDataInterfaceManager;
 
 	TArray<UNiagaraScript*> ScriptsToCheckForStatus;
diff --git a/Engine/Plugins/FX/Niagara/Source/NiagaraEditor/Public/ViewModels/Stack/NiagaraStackRendererItem.h b/Engine/Plugins/FX/Niagara/Source/NiagaraEditor/Public/ViewModels/Stack/NiagaraStackRendererItem.h
--- a/Engine/Plugins/FX/Niagara/Source/NiagaraEditor/Public/ViewModels/Stack/NiagaraStackRendererItem.h
+++ b/Engine/Plugins/FX/Niagara/Source/NiagaraEditor/Public/ViewModels/Stack/NiagaraStackRendererItem.h
@@ -62,6 +62,9 @@ public:
 	static TArray<FNiagaraVariable> GetMissingVariables(UNiagaraRendererProperties* RendererProperties, UNiagaraEmitter* Emitter);
 	static bool AddMissingVariable(UNiagaraEmitter* Emitter, const FNiagaraVariable& Variable);
 
+	bool IsExcludedFromScalability() const;
+	bool IsOwningEmitterExcludedFromScalability() const;
+
 protected:
 	virtual void FinalizeInternal() override;
 
diff --git a/Engine/Plugins/FX/Niagara/Source/NiagaraEditorWidgets/NiagaraEditorWidgets.Build.cs b/Engine/Plugins/FX/Niagara/Source/NiagaraEditorWidgets/NiagaraEditorWidgets.Build.cs
--- a/Engine/Plugins/FX/Niagara/Source/NiagaraEditorWidgets/NiagaraEditorWidgets.Build.cs
+++ b/Engine/Plugins/FX/Niagara/Source/NiagaraEditorWidgets/NiagaraEditorWidgets.Build.cs
@@ -31,7 +31,8 @@ public class NiagaraEditorWidgets : ModuleRules
             "EditorWidgets",
 			"ApplicationCore",
 			"CurveEditor",
-			"ToolWidgets",
+			"DesktopPlatform",
+			"ToolWidgets"
 		});
 
 		PrivateIncludePathModuleNames.AddRange(new string[] {
diff --git a/Engine/Plugins/FX/Niagara/Source/NiagaraEditorWidgets/Private/SNiagaraOverviewGraph.cpp b/Engine/Plugins/FX/Niagara/Source/NiagaraEditorWidgets/Private/SNiagaraOverviewGraph.cpp
--- a/Engine/Plugins/FX/Niagara/Source/NiagaraEditorWidgets/Private/SNiagaraOverviewGraph.cpp
+++ b/Engine/Plugins/FX/Niagara/Source/NiagaraEditorWidgets/Private/SNiagaraOverviewGraph.cpp
@@ -17,6 +17,7 @@
 #include "NiagaraEditorCommands.h"
 #include "NiagaraEditorModule.h"
 #include "GraphEditorActions.h"
+#include "NiagaraEditorWidgetsStyle.h"
 #include "Framework/MultiBox/MultiBoxBuilder.h"
 #include "ScopedTransaction.h"
 #include "Framework/Application/SlateApplication.h"
@@ -27,6 +28,7 @@
 #include "Widgets/Input/SCheckBox.h"
 #include "NiagaraEditor/Private/SNiagaraAssetPickerList.h"
 #include "Widgets/SItemSelector.h"
+#include "SNiagaraOverviewGraphTitleBar.h"
 
 #define LOCTEXT_NAMESPACE "NiagaraOverviewGraph"
 
@@ -62,15 +64,7 @@ void SNiagaraOverviewGraph::Construct(const FArguments& InArgs, TSharedRef<FNiag
 		AppearanceInfo.CornerText = LOCTEXT("NiagaraOverview_AppearanceCornerTextGeneric", "NIAGARA");
 	}
 	
-	TSharedRef<SWidget> TitleBarWidget = SNew(SBorder)
-	.BorderImage(FEditorStyle::GetBrush(TEXT("Graph.TitleBackground")))
-	.HAlign(HAlign_Fill)
-	[
-		SNew(STextBlock)
-		.Text(ViewModel.ToSharedRef(), &FNiagaraOverviewGraphViewModel::GetDisplayName)
-		.TextStyle(FEditorStyle::Get(), TEXT("GraphBreadcrumbButtonText"))
-		.Justification(ETextJustify::Center)
-	];
+	TSharedRef<SWidget> TitleBarWidget = SNew(SNiagaraOverviewGraphTitleBar, ViewModel->GetSystemViewModel()).Visibility(EVisibility::SelfHitTestInvisible);
 
 	TSharedRef<FUICommandList> Commands = ViewModel->GetCommands();
 	Commands->MapAction(
@@ -137,6 +131,8 @@ void SNiagaraOverviewGraph::Construct(const FArguments& InArgs, TSharedRef<FNiag
 		ZoomToFitFrameDelay = 2;
 	}
 
+	GraphEditor->GetCurrentGraph()->AddOnGraphChangedHandler(FOnGraphChanged::FDelegate::CreateSP(this, &SNiagaraOverviewGraph::OnNodesCreated));
+
 	ChildSlot
 	[
 		GraphEditor.ToSharedRef()
@@ -203,7 +199,7 @@ void SNiagaraOverviewGraph::PreClose()
 }
 
 void SNiagaraOverviewGraph::OpenAddEmitterMenu()
-{
+{	
 	FNiagaraAssetPickerListViewOptions ViewOptions;
 	ViewOptions.SetCategorizeUserDefinedCategory(true);
 	ViewOptions.SetCategorizeLibraryAssets(true);
@@ -213,17 +209,15 @@ void SNiagaraOverviewGraph::OpenAddEmitterMenu()
 	TabOptions.ChangeTabState(ENiagaraScriptTemplateSpecification::None, true);
 	TabOptions.ChangeTabState(ENiagaraScriptTemplateSpecification::Behavior, true);	
 
-	TSharedPtr<SNiagaraAssetPickerList> AssetPickerList;
-	
-	SAssignNew(AssetPickerList, SNiagaraAssetPickerList, UNiagaraEmitter::StaticClass())
+	TSharedPtr<SNiagaraAssetPickerList> AssetPickerList = SNew(SNiagaraAssetPickerList, UNiagaraEmitter::StaticClass())
 	.ClickActivateMode(EItemSelectorClickActivateMode::SingleClick)
 	.ViewOptions(ViewOptions)
 	.TabOptions(TabOptions)
 	.OnTemplateAssetActivated_Lambda([this](const FAssetData& AssetData) {
+		TSharedPtr<FNiagaraEmitterHandleViewModel> EmitterHandleViewModel = ViewModel->GetSystemViewModel()->AddEmitterFromAssetData(AssetData);
 		FSlateApplication::Get().DismissAllMenus();
-		ViewModel->GetSystemViewModel()->AddEmitterFromAssetData(AssetData);
 	});
-	
+    
 	TSharedRef<SWidget> EmitterAddSubMenu =
 		SNew(SBorder)
 		.BorderImage(FNiagaraEditorStyle::Get().GetBrush("GraphActionMenu.Background"))
@@ -236,9 +230,11 @@ void SNiagaraOverviewGraph::OpenAddEmitterMenu()
 				AssetPickerList.ToSharedRef()
 			]
 		];
-
+	
 	FSlateApplication::Get().PushMenu(SharedThis(this), FWidgetPath(), EmitterAddSubMenu, FSlateApplication::Get().GetCursorPos(), FPopupTransitionEffect::None);
 
+	// we need to set the keyboard focus for next tick instead of immediately
+	// if we open the menu via shortcut, we'd type the shortcut key into the search box
 	GEditor->GetTimerManager()->SetTimerForNextTick(FTimerDelegate::CreateLambda([AssetPickerList]()
 	{
 		FSlateApplication::Get().SetKeyboardFocus(AssetPickerList->GetSearchBox());
@@ -254,8 +250,8 @@ FActionMenuContent SNiagaraOverviewGraph::OnCreateGraphActionMenu(UEdGraph* InGr
 {
 	if (ViewModel->GetSystemViewModel()->GetEditMode() == ENiagaraSystemViewModelEditMode::SystemAsset)
 	{
-		FMenuBuilder MenuBuilder(true, ViewModel->GetCommands());
-
+		FMenuBuilder MenuBuilder(true, ViewModel->GetCommands(), TSharedPtr<FExtender>(), false, &FAppStyle::Get(), false);
+		
 		MenuBuilder.BeginSection(TEXT("NiagaraOverview_EditGraph"), LOCTEXT("EditGraph", "Edit Graph"));
 		{
 			MenuBuilder.AddMenuEntry(FNiagaraEditorCommands::Get().OpenAddEmitterMenu);
diff --git a/Engine/Plugins/FX/Niagara/Source/NiagaraEditorWidgets/Private/SNiagaraOverviewGraphTitleBar.cpp b/Engine/Plugins/FX/Niagara/Source/NiagaraEditorWidgets/Private/SNiagaraOverviewGraphTitleBar.cpp
new file mode 100644
--- /dev/null
+++ b/Engine/Plugins/FX/Niagara/Source/NiagaraEditorWidgets/Private/SNiagaraOverviewGraphTitleBar.cpp
@@ -0,0 +1,90 @@
+// Copyright Epic Games, Inc. All Rights Reserved.
+
+#include "SNiagaraOverviewGraphTitleBar.h"
+#include "NiagaraEditorStyle.h"
+#include "NiagaraSettings.h"
+#include "ViewModels/NiagaraOverviewGraphViewModel.h"
+#include "SNiagaraScalabilityPreviewSettings.h"
+#include "SNiagaraSystemEffectTypeBar.h"
+
+#define LOCTEXT_NAMESPACE "NiagaraOverviewGraphTitleBar"
+
+void SNiagaraOverviewGraphTitleBar::Construct(const FArguments& InArgs, TSharedRef<FNiagaraSystemViewModel> InSystemViewModel)
+{
+	SystemViewModel = InSystemViewModel;
+
+	bScalabilityModeActive = SystemViewModel.Pin()->GetScalabilityViewModel()->IsActive();
+	SystemViewModel.Pin()->GetScalabilityViewModel()->OnScalabilityModeChanged().AddSP(this, &SNiagaraOverviewGraphTitleBar::OnUpdateScalabilityMode);
+
+	RebuildWidget();
+}
+
+void SNiagaraOverviewGraphTitleBar::RebuildWidget()
+{
+	ContainerWidget = SNew(SVerticalBox);
+
+	// we only allow the Niagara system effect type bar on system assets
+	if(bScalabilityModeActive && SystemViewModel.IsValid() && SystemViewModel.Pin()->GetEditMode() == ENiagaraSystemViewModelEditMode::SystemAsset)
+	{
+		ContainerWidget->AddSlot()
+		.HAlign(HAlign_Fill)
+		.AutoHeight()
+		[
+			SNew(SBorder)
+			.BorderImage(FEditorStyle::GetBrush(TEXT("Graph.TitleBackground")))
+			.HAlign(HAlign_Fill)
+			[
+				SNew(SNiagaraSystemEffectTypeBar, SystemViewModel.Pin()->GetSystem())
+			]
+		];
+	}
+	
+	ContainerWidget->AddSlot()
+	.HAlign(HAlign_Fill)
+	.AutoHeight()
+	[
+		SNew(SBorder)
+		.BorderImage(FEditorStyle::GetBrush(TEXT("Graph.TitleBackground")))
+		.HAlign(HAlign_Fill)
+		[
+			SNew(STextBlock)
+			.Text(SystemViewModel.Pin()->GetOverviewGraphViewModel().Get(), &FNiagaraOverviewGraphViewModel::GetDisplayName)
+			.TextStyle(FEditorStyle::Get(), TEXT("GraphBreadcrumbButtonText"))
+			.Justification(ETextJustify::Center)
+		]
+	];
+	
+	if(bScalabilityModeActive)
+	{
+		ContainerWidget->AddSlot()
+		.HAlign(HAlign_Fill)
+		.Padding(5.f)
+		.AutoHeight()
+		[			
+			SNew(SNiagaraScalabilityPreviewSettings, *SystemViewModel.Pin()->GetScalabilityViewModel()).Visibility(EVisibility::SelfHitTestInvisible)			
+		];
+	}
+
+	ChildSlot
+	[
+		ContainerWidget.ToSharedRef()
+	];
+}
+
+// FReply SNiagaraOverviewGraphTitleBar::OnMouseButtonDown(const FGeometry& MyGeometry, const FPointerEvent& MouseEvent)
+// {
+// 	if(OverviewGraph.IsValid())
+// 	{
+// 		//OverviewGraph->g
+// 	}
+//
+// 	return FReply::Handled();
+// }
+
+void SNiagaraOverviewGraphTitleBar::OnUpdateScalabilityMode(bool bActive)
+{
+	bScalabilityModeActive = bActive;
+	RebuildWidget();
+}
+
+#undef LOCTEXT_NAMESPACE
diff --git a/Engine/Plugins/FX/Niagara/Source/NiagaraEditorWidgets/Private/SNiagaraOverviewGraphTitleBar.h b/Engine/Plugins/FX/Niagara/Source/NiagaraEditorWidgets/Private/SNiagaraOverviewGraphTitleBar.h
new file mode 100644
--- /dev/null
+++ b/Engine/Plugins/FX/Niagara/Source/NiagaraEditorWidgets/Private/SNiagaraOverviewGraphTitleBar.h
@@ -0,0 +1,33 @@
+// Copyright Epic Games, Inc. All Rights Reserved.
+
+#pragma once
+
+#include "CoreMinimal.h"
+#include "NiagaraSystem.h"
+#include "ViewModels/NiagaraSystemScalabilityViewModel.h"
+#include "ViewModels/NiagaraSystemViewModel.h"
+//#include "SNiagaraOverviewGraph.h"
+
+class SNiagaraOverviewGraph;
+
+class SNiagaraOverviewGraphTitleBar : public SCompoundWidget
+{
+	SLATE_BEGIN_ARGS(SNiagaraOverviewGraphTitleBar)
+		: _OverviewGraph(nullptr)
+	{}
+		SLATE_ARGUMENT(TSharedPtr<SNiagaraOverviewGraph>, OverviewGraph)
+	SLATE_END_ARGS()
+
+public:
+	void Construct(const FArguments& InArgs, TSharedRef<FNiagaraSystemViewModel> InSystemViewModel);
+
+	void RebuildWidget();
+	//virtual FReply OnMouseButtonDown(const FGeometry& MyGeometry, const FPointerEvent& MouseEvent) override;
+private:
+	TWeakPtr<FNiagaraSystemViewModel> SystemViewModel;
+	bool bScalabilityModeActive = false;
+	TSharedPtr<SVerticalBox> ContainerWidget;
+	TWeakPtr<SNiagaraOverviewGraph> OverviewGraph;
+private:
+	void OnUpdateScalabilityMode(bool bActive);
+};
diff --git a/Engine/Plugins/FX/Niagara/Source/NiagaraEditorWidgets/Private/SNiagaraOverviewStack.cpp b/Engine/Plugins/FX/Niagara/Source/NiagaraEditorWidgets/Private/SNiagaraOverviewStack.cpp
--- a/Engine/Plugins/FX/Niagara/Source/NiagaraEditorWidgets/Private/SNiagaraOverviewStack.cpp
+++ b/Engine/Plugins/FX/Niagara/Source/NiagaraEditorWidgets/Private/SNiagaraOverviewStack.cpp
@@ -1,51 +1,53 @@
 // Copyright Epic Games, Inc. All Rights Reserved.
 
 #include "SNiagaraOverviewStack.h"
-
-#include "EditorStyleSet.h"
-#include "Framework/Application/SlateApplication.h"
-#include "Framework/MultiBox/MultiBoxBuilder.h"
-#include "NiagaraEditorModule.h"
-#include "NiagaraEditorStyle.h"
-#include "NiagaraEditorWidgetsStyle.h"
-#include "NiagaraEditorWidgetsUtilities.h"
-#include "NiagaraNodeAssignment.h"
-#include "NiagaraNodeFunctionCall.h"
-#include "NiagaraStackCommandContext.h"
-#include "NiagaraStackEditorData.h"
 #include "NiagaraSystem.h"
 #include "NiagaraSystemEditorData.h"
-#include "ScopedTransaction.h"
-#include "SNiagaraOverviewInlineParameterBox.h"
-#include "SNiagaraStack.h"
-#include "Stack/SNiagaraStackIssueIcon.h"
-#include "Stack/SNiagaraStackItemGroupAddButton.h"
-#include "Stack/SNiagaraStackRowPerfWidget.h"
-#include "Styling/StyleColors.h"
-#include "Subsystems/AssetEditorSubsystem.h"
-#include "ViewModels/NiagaraEmitterHandleViewModel.h"
+#include "NiagaraStackEditorData.h"
 #include "ViewModels/NiagaraEmitterViewModel.h"
-#include "ViewModels/NiagaraScratchPadScriptViewModel.h"
-#include "ViewModels/NiagaraScratchPadViewModel.h"
-#include "ViewModels/NiagaraSystemSelectionViewModel.h"
 #include "ViewModels/NiagaraSystemViewModel.h"
+#include "ViewModels/NiagaraEmitterHandleViewModel.h"
+#include "ViewModels/NiagaraSystemSelectionViewModel.h"
+#include "ViewModels/NiagaraScratchPadViewModel.h"
+#include "ViewModels/NiagaraScratchPadScriptViewModel.h"
 #include "ViewModels/Stack/NiagaraStackEntry.h"
-#include "ViewModels/Stack/NiagaraStackItem.h"
 #include "ViewModels/Stack/NiagaraStackItemGroup.h"
+#include "ViewModels/Stack/NiagaraStackItem.h"
 #include "ViewModels/Stack/NiagaraStackModuleItem.h"
 #include "ViewModels/Stack/NiagaraStackViewModel.h"
-#include "Widgets/Colors/SColorBlock.h"
-#include "Widgets/Images/SImage.h"
-#include "Widgets/Input/SButton.h"
-#include "Widgets/Input/SCheckBox.h"
-#include "Widgets/Input/SCheckBox.h"
+#include "NiagaraEditorWidgetsStyle.h"
+#include "NiagaraEditorWidgetsUtilities.h"
+#include "SNiagaraStack.h"
+#include "Stack/SNiagaraStackItemGroupAddButton.h"
+#include "Stack/SNiagaraStackIssueIcon.h"
+#include "Stack/SNiagaraStackRowPerfWidget.h"
+#include "NiagaraStackCommandContext.h"
+#include "NiagaraNodeFunctionCall.h"
+#include "NiagaraEditorModule.h"
+#include "NiagaraNodeAssignment.h"
+#include "Widgets/SNiagaraParameterName.h"
+#include "SNiagaraOverviewInlineParameterBox.h"
+#include "Widgets/Layout/SScrollBox.h"
 #include "Widgets/Layout/SGridPanel.h"
 #include "Widgets/Layout/SScaleBox.h"
-#include "Widgets/Layout/SScrollBox.h"
 #include "Widgets/Layout/SSpacer.h"
-#include "Widgets/Layout/SWrapBox.h"
-#include "Widgets/SNiagaraParameterName.h"
+#include "Widgets/Images/SImage.h"
+#include "Widgets/Input/SButton.h"
+#include "Widgets/Input/SCheckBox.h"
 #include "Widgets/SToolTip.h"
+#include "EditorStyleSet.h"
+#include "Framework/Application/SlateApplication.h"
+#include "Framework/MultiBox/MultiBoxBuilder.h"
+#include "Subsystems/AssetEditorSubsystem.h"
+#include "Widgets/Colors/SColorBlock.h"
+#include "NiagaraEditorStyle.h"
+#include "NiagaraEditorUtilities.h"
+#include "ScopedTransaction.h"
+#include "Styling/StyleColors.h"
+#include "ViewModels/NiagaraOverviewGraphViewModel.h"
+#include "ViewModels/Stack/NiagaraStackRendererItem.h"
+#include "Widgets/Input/SCheckBox.h"
+#include "Widgets/Layout/SWrapBox.h"
 
 #define LOCTEXT_NAMESPACE "NiagaraOverviewStack"
 
@@ -64,7 +66,8 @@
 
 	virtual ~SNiagaraSystemOverviewEntryListRow() override
 	{
-		StackEntry->OnStructureChanged().RemoveAll(this);	
+		StackEntry->OnStructureChanged().RemoveAll(this);
+		ScalabilityViewModel->OnScalabilityModeChanged().RemoveAll(this);
 	}
 	
 	void Construct(const FArguments& InArgs, UNiagaraStackViewModel* InStackViewModel, UNiagaraStackEntry* InStackEntry, TSharedRef<FNiagaraStackCommandContext> InStackCommandContext, const TSharedRef<STableViewBase>& InOwnerTableView)
@@ -74,9 +77,13 @@
 		StackCommandContext = InStackCommandContext;
 		FSlateColor IconColor = FNiagaraEditorWidgetsStyle::Get().GetColor(FNiagaraStackEditorWidgetsUtilities::GetColorNameForExecutionCategory(StackEntry->GetExecutionCategoryName()));
 
+		ScalabilityViewModel = InStackEntry->GetSystemViewModel()->GetScalabilityViewModel();
+		bScalabilityModeActive = ScalabilityViewModel->IsActive();
+		ScalabilityViewModel->OnScalabilityModeChanged().AddSP(this, &SNiagaraSystemOverviewEntryListRow::UpdateScalabilityMode);
+
 		ExpandedImage = FCoreStyle::Get().GetBrush("TreeArrow_Expanded");
 		CollapsedImage = FCoreStyle::Get().GetBrush("TreeArrow_Collapsed");
-
+		
 		TSharedRef<FNiagaraSystemViewModel> SystemViewModel = StackEntry->GetSystemViewModel();
 		TSharedPtr<FNiagaraEmitterViewModel> EmitterViewModel = StackEntry->GetEmitterViewModel();
 		if (EmitterViewModel.IsValid())
@@ -97,6 +104,9 @@
 			SAssignNew(RowOverlay, SOverlay)
 			+ SOverlay::Slot()
 			[
+				SAssignNew(ScalabilityOverlay, SOverlay)
+				+ SOverlay::Slot()
+				[
 					SNew(SBorder)
 					.BorderImage(FNiagaraEditorWidgetsStyle::Get().GetBrush("Niagara.TableViewRowBorder"))
 					.ToolTipText_UObject(StackEntry, &UNiagaraStackEntry::GetTooltipText)
@@ -151,6 +161,7 @@
 							InArgs._Content.Widget
 						]
 					]
+				]	
 			]
 			// Stack Issue Highlight
 			+ SOverlay::Slot()
@@ -171,22 +182,21 @@
 		],
 		InOwnerTableView);
 
-		if (StackEntry->IsA<UNiagaraStackItem>())
-		{
-			// Execution Category Item Highlight
-			RowOverlay->AddSlot()
-			.HAlign(HAlign_Left)
-			.VAlign(VAlign_Fill)
-			.Padding(FMargin(0))
+		// we special case the red overlay for renderer items as they have their own scalability settings
+		if(UNiagaraStackRendererItem* RendererItem = Cast<UNiagaraStackRendererItem>(StackEntry))
+		{			
+			ScalabilityOverlay->AddSlot()
 			[
 				SNew(SBorder)
-				.BorderImage(FAppStyle::Get().GetBrush("WhiteBrush"))
-				.BorderBackgroundColor(IconColor)
-				.Padding(FMargin(0))
-				[
-					SNew(SBox)
-					.WidthOverride(2)
-				]
+				.Visibility_Lambda([=]()
+				{
+					return bScalabilityModeActive && RendererItem->IsExcludedFromScalability() && !RendererItem->IsOwningEmitterExcludedFromScalability() ? EVisibility::HitTestInvisible : EVisibility::Collapsed; 
+				})
+				.BorderImage(FNiagaraEditorStyle::Get().GetBrush("NiagaraEditor.SystemOverview.ExcludedFromScalability.RendererItem"))
+				.BorderBackgroundColor_Lambda([=]()
+				{
+					return FLinearColor(1, 1, 1, FNiagaraEditorUtilities::GetScalabilityTintAlpha(EmitterHandleViewModel.Pin()->GetEmitterHandle()));
+				})
 			];
 		}
 	}
@@ -309,6 +319,11 @@
 		IssueHighlightColor.Reset();
 	}
 
+	void UpdateScalabilityMode(bool bActivated)
+	{
+		bScalabilityModeActive = bActivated;
+	}
+
 	FSlateColor GetIsssueHighlightColor() const
 	{
 		if (IssueHighlightColor.IsSet() == false)
@@ -341,11 +356,15 @@
 private:
 	UNiagaraStackViewModel* StackViewModel;
 	UNiagaraStackEntry* StackEntry;
+	UNiagaraSystemScalabilityViewModel* ScalabilityViewModel;
+	TSharedPtr<SOverlay> ScalabilityOverlay;
 	TSharedPtr<FNiagaraStackCommandContext> StackCommandContext;
+	FSlateColor ExcludedFromScalabilityColor;
 	TAttribute<EVisibility> IssueIconVisibility;
 	TWeakPtr<FNiagaraEmitterHandleViewModel> EmitterHandleViewModel;
 	const FSlateBrush* ExpandedImage = nullptr;
 	const FSlateBrush* CollapsedImage = nullptr;
+	bool bScalabilityModeActive;
 	mutable TOptional<uint32> ChildrenCount;
 	mutable TOptional<FSlateColor> IssueHighlightColor;
 };
@@ -948,6 +967,9 @@
 	}
 	else if (Item->IsA<UNiagaraStackItemGroup>())
 	{
+		FName IconName = FNiagaraStackEditorWidgetsUtilities::GetIconNameForExecutionSubcategory(Item->GetExecutionSubcategoryName(), true);
+		const FSlateBrush* IconBrush = IconName != NAME_None ? FNiagaraEditorWidgetsStyle::Get().GetBrush(IconName) : FAppStyle::Get().GetDefaultBrush();
+		
 		UNiagaraStackItemGroup* StackItemGroup = CastChecked<UNiagaraStackItemGroup>(Item);
 		TSharedRef<SHorizontalBox> ContentBox = SNew(SHorizontalBox);
 
diff --git a/Engine/Plugins/FX/Niagara/Source/NiagaraEditorWidgets/Private/SNiagaraOverviewStackNode.cpp b/Engine/Plugins/FX/Niagara/Source/NiagaraEditorWidgets/Private/SNiagaraOverviewStackNode.cpp
--- a/Engine/Plugins/FX/Niagara/Source/NiagaraEditorWidgets/Private/SNiagaraOverviewStackNode.cpp
+++ b/Engine/Plugins/FX/Niagara/Source/NiagaraEditorWidgets/Private/SNiagaraOverviewStackNode.cpp
@@ -21,9 +21,11 @@
 #include "Widgets/Input/SCheckBox.h"
 #include "AssetThumbnail.h"
 #include "NiagaraEmitterInstance.h"
+#include "NiagaraEditorUtilities.h"
 #include "Widgets/SToolTip.h"
 #include "Widgets/SBoxPanel.h"
 #include "NiagaraRendererProperties.h"
+#include "SGraphPanel.h"
 #include "Widgets/SWidget.h"
 #include "Subsystems/AssetEditorSubsystem.h"
 #include "ViewModels/Stack/NiagaraStackRendererItem.h"
@@ -75,12 +77,27 @@ void SNiagaraOverviewStackNode::Construct(const FArguments& InArgs, UNiagaraOver
 			}
 			UMaterial::OnMaterialCompilationFinished().AddSP(this, &SNiagaraOverviewStackNode::OnMaterialCompiled);
 			OverviewSelectionViewModel = OwningSystemViewModel->GetSelectionViewModel();
+			ScalabilityViewModel = OwningSystemViewModel->GetScalabilityViewModel();
+
+			if(ScalabilityViewModel.IsValid())
+			{
+				bScalabilityModeActive = ScalabilityViewModel->IsActive();
+				ScalabilityViewModel->OnScalabilityModeChanged().AddSP(this, &SNiagaraOverviewStackNode::OnScalabilityModeChanged);
+			}
 		}
 	}
 
 	UpdateGraphNode();
 }
 
+SNiagaraOverviewStackNode::~SNiagaraOverviewStackNode()
+{
+	if(ScalabilityViewModel.IsValid())
+	{
+		ScalabilityViewModel->OnScalabilityModeChanged().RemoveAll(this);
+	}
+}
+
 TSharedRef<SWidget> SNiagaraOverviewStackNode::CreateTitleWidget(TSharedPtr<SNodeTitle> NodeTitle)
 {
 	TSharedRef<SWidget> DefaultTitle = SGraphNode::CreateTitleWidget(NodeTitle);
@@ -409,7 +426,21 @@ TSharedRef<SWidget> SNiagaraOverviewStackNode::CreateNodeContentArea()
 		];
 	}
 
-	return NodeWidget;
+	TSharedPtr<SOverlay> ScalabilityWrapper = SNew(SOverlay)
+	+ SOverlay::Slot()
+	[
+		NodeWidget
+	]
+	+ SOverlay::Slot()
+	.Padding(0, 1)
+	[
+		SNew(SBorder)
+		.Visibility(this, &SNiagaraOverviewStackNode::ShowExcludedOverlay)
+		.BorderImage(FNiagaraEditorStyle::Get().GetBrush("NiagaraEditor.SystemOverview.ExcludedFromScalability.NodeBody"))
+		.BorderBackgroundColor(TAttribute<FSlateColor>(this, &SNiagaraOverviewStackNode::GetScalabilityTintAlpha))
+	];
+	
+	return ScalabilityWrapper.ToSharedRef();
 
 }
 
@@ -437,8 +468,7 @@ void SNiagaraOverviewStackNode::FillTopContentBar()
 		TopContentBar->ClearChildren();
 	}
 	if (EmitterHandleViewModelWeak.IsValid())
-	{
-	
+	{		
 		// Isolate toggle button
 		TopContentBar->AddSlot()
 			.AutoWidth()
@@ -583,6 +613,45 @@ FSlateColor SNiagaraOverviewStackNode::GetToggleIsolateImageColor() const
 		: FLinearColor::Gray;
 }
 
+FSlateColor SNiagaraOverviewStackNode::GetScalabilityTintAlpha() const
+{
+	if(EmitterHandleViewModelWeak.IsValid())
+	{
+		float ScalabilityBaseAlpha = FNiagaraEditorUtilities::GetScalabilityTintAlpha(EmitterHandleViewModelWeak.Pin()->GetEmitterHandle());
+		return FLinearColor(1, 1, 1, ScalabilityBaseAlpha * GetGraphZoomDistanceAlphaMultiplier());
+	}
+
+	return FLinearColor(1, 1, 1, 1);
+}
+
+void SNiagaraOverviewStackNode::OnScalabilityModeChanged(bool bActive)
+{
+	bScalabilityModeActive = bActive;
+}
+
+EVisibility SNiagaraOverviewStackNode::GetScalabilityBarVisibility() const
+{
+	return bScalabilityModeActive ? EVisibility::Visible : EVisibility::Collapsed;
+}
+
+EVisibility SNiagaraOverviewStackNode::ShowExcludedOverlay() const
+{
+	// we only want actual results in scalability mode and for nodes representing emitters (not system nodes)
+	if(bScalabilityModeActive && EmitterHandleViewModelWeak.IsValid())
+	{		
+		return EmitterHandleViewModelWeak.Pin()->GetEmitterHandle()->GetInstance()->IsAllowedByScalability() ? EVisibility::Hidden : EVisibility::HitTestInvisible;
+	}
+	
+	return EVisibility::Hidden; 
+}
+
+float SNiagaraOverviewStackNode::GetGraphZoomDistanceAlphaMultiplier() const
+{
+	// we lower the alpha if the zoom amount is high
+	float ZoomAmount = OwnerGraphPanelPtr.Pin()->GetZoomAmount();
+	return FMath::Lerp(0.4f, 1.f, 1 - ZoomAmount);
+}
+
 FReply SNiagaraOverviewStackNode::OpenParentEmitter()
 {
 	TSharedPtr<FNiagaraEmitterHandleViewModel> EmitterHandleViewModel = EmitterHandleViewModelWeak.Pin();
diff --git a/Engine/Plugins/FX/Niagara/Source/NiagaraEditorWidgets/Private/SNiagaraOverviewStackNode.h b/Engine/Plugins/FX/Niagara/Source/NiagaraEditorWidgets/Private/SNiagaraOverviewStackNode.h
--- a/Engine/Plugins/FX/Niagara/Source/NiagaraEditorWidgets/Private/SNiagaraOverviewStackNode.h
+++ b/Engine/Plugins/FX/Niagara/Source/NiagaraEditorWidgets/Private/SNiagaraOverviewStackNode.h
@@ -5,6 +5,7 @@
 #include "SGraphNode.h"
 #include "Widgets/SBoxPanel.h"
 #include "SNiagaraOverviewStack.h"
+#include "ViewModels/NiagaraSystemScalabilityViewModel.h"
 
 class UNiagaraOverviewNode;
 class UNiagaraStackViewModel;
@@ -24,6 +25,7 @@ public:
 	virtual void Tick(const FGeometry& AllottedGeometry, const double InCurrentTime, const float InDeltaTime) override;
 	virtual bool ShouldAllowCulling() const override { return false; }
 
+	virtual ~SNiagaraOverviewStackNode();
 protected:
 	virtual TSharedRef<SWidget> CreateTitleWidget(TSharedPtr<SNodeTitle> NodeTitle) override;
 	virtual TSharedRef<SWidget> CreateTitleRightWidget() override;
@@ -45,6 +47,11 @@ private:
 	FReply OnToggleIsolateButtonClicked();
 	EVisibility GetToggleIsolateVisibility() const;
 	FSlateColor GetToggleIsolateImageColor() const;
+	FSlateColor GetScalabilityTintAlpha() const;
+	void OnScalabilityModeChanged(bool bActive);
+	EVisibility GetScalabilityBarVisibility() const;
+	EVisibility ShowExcludedOverlay() const;
+	float GetGraphZoomDistanceAlphaMultiplier() const;
 	void SetIsHoveringThumbnail(const FGeometry& InGeometry, const FPointerEvent& InEvent, const bool bInHoveringThumbnail)
 	{
 		SetIsHoveringThumbnail(InEvent, bInHoveringThumbnail);
@@ -69,6 +76,7 @@ private:
 	UNiagaraOverviewNode* OverviewStackNode;
 	UNiagaraStackViewModel* StackViewModel;
 	UNiagaraSystemSelectionViewModel* OverviewSelectionViewModel;
+	TWeakObjectPtr<UNiagaraSystemScalabilityViewModel> ScalabilityViewModel;
 	TWeakPtr<FNiagaraEmitterHandleViewModel> EmitterHandleViewModelWeak;
 	/** The top content bar houses the isolate button and the thumbnails */
 	TSharedPtr<SHorizontalBox> TopContentBar;
@@ -78,4 +86,5 @@ private:
 	bool bIsHoveringThumbnail;
 	bool bTopContentBarRefreshPending;
 	int32 CurrentIssueIndex;
+	bool bScalabilityModeActive;
 };
\ No newline at end of file
diff --git a/Engine/Plugins/FX/Niagara/Source/NiagaraEditorWidgets/Private/SNiagaraScalabilityPreviewSettings.cpp b/Engine/Plugins/FX/Niagara/Source/NiagaraEditorWidgets/Private/SNiagaraScalabilityPreviewSettings.cpp
new file mode 100644
--- /dev/null
+++ b/Engine/Plugins/FX/Niagara/Source/NiagaraEditorWidgets/Private/SNiagaraScalabilityPreviewSettings.cpp
@@ -0,0 +1,424 @@
+// Copyright Epic Games, Inc. All Rights Reserved.
+
+#include "SNiagaraScalabilityPreviewSettings.h"
+#include "IDetailTreeNode.h"
+#include "IPropertyRowGenerator.h"
+#include "ISinglePropertyView.h"
+#include "NiagaraEditorStyle.h"
+#include "NiagaraSettings.h"
+#include "PlatformInfo.h"
+#include "DeviceProfiles/DeviceProfile.h"
+#include "DeviceProfiles/DeviceProfileManager.h"
+#include "Widgets/Layout/SWrapBox.h"
+#include "Widgets/Layout/SSeparator.h"
+#include "NiagaraEditor/Private/Customizations/NiagaraPlatformSetCustomization.h"
+#include "Widgets/Images/SImage.h"
+#include "Widgets/Input/SButton.h"
+#include "Widgets/Input/SCheckBox.h"
+#include "Widgets/Layout/SSpacer.h"
+
+#define LOCTEXT_NAMESPACE "NiagaraScalabilityPreviewSettings"
+
+void SNiagaraScalabilityPreviewSettings::Construct(const FArguments& InArgs, UNiagaraSystemScalabilityViewModel& InScalabilityViewModel)
+{
+	ScalabilityViewModel = &InScalabilityViewModel;
+
+	ChildSlot
+	.VAlign(VAlign_Center)
+	[
+		SNew(SHorizontalBox)
+		+ SHorizontalBox::Slot()
+		.Padding(2.f)
+		.AutoWidth()
+		[
+			SNew(STextBlock)
+			.Text(LOCTEXT("ScalabilityPreviewSettingsLabel", "Preview"))
+		]
+		+ SHorizontalBox::Slot()
+		.HAlign(HAlign_Left)
+		.Padding(2.f)
+		.AutoWidth()
+		[
+			CreatePreviewQualityLevelWidgets()
+		]
+		+ SHorizontalBox::Slot()
+		.FillWidth(1.f)
+		[
+			SNew(SSpacer)
+			.Size(FVector2D(1, 1))
+		]
+		+ SHorizontalBox::Slot()
+		.HAlign(HAlign_Right)
+		.AutoWidth()
+		[
+			SAssignNew(PlatformMenuAnchor, SMenuAnchor)
+			.OnGetMenuContent(this, &SNiagaraScalabilityPreviewSettings::GenerateDeviceProfileTreeWidget)
+			[				
+				CreatePreviewPlatformWidgets()
+			]
+		]
+	];
+}
+
+TSharedRef<SWidget> SNiagaraScalabilityPreviewSettings::CreatePreviewQualityLevelWidgets()
+{
+	const UNiagaraSettings* Settings = GetDefault<UNiagaraSettings>();
+	check(Settings);
+
+	TSharedRef<SWrapBox> QualityLevelWidgetsBox = SNew(SWrapBox)
+	.PreferredSize(600.f)
+	.UseAllottedSize(true);
+	
+	int32 NumQualityLevels = Settings->QualityLevels.Num();
+
+	for (int32 QualityLevel = 0; QualityLevel < NumQualityLevels; ++QualityLevel)
+	{
+		bool First = QualityLevel == 0;
+		bool Last = QualityLevel == (NumQualityLevels - 1);
+		
+		QualityLevelWidgetsBox->AddSlot()
+		.Padding(0, 0, 1, 0)
+		[
+			SNew(SBox)
+			//.WidthOverride(80.f)
+			.VAlign(VAlign_Top)
+			[
+				SNew(SVerticalBox)
+				+ SVerticalBox::Slot()
+				.AutoHeight()
+				[
+					SNew(SCheckBox)
+					.ToolTipText(GetQualityButtonTooltip(QualityLevel))
+					.Style(FNiagaraEditorStyle::Get(), First ? "NiagaraEditor.PlatformSet.StartButton" : 
+						(Last ? "NiagaraEditor.PlatformSet.EndButton" : "NiagaraEditor.PlatformSet.MiddleButton"))
+					.IsChecked(this, &SNiagaraScalabilityPreviewSettings::IsQLChecked, QualityLevel)
+					.OnCheckStateChanged(this, &SNiagaraScalabilityPreviewSettings::QLCheckStateChanged, QualityLevel)
+					[
+						SNew(SHorizontalBox)
+						+ SHorizontalBox::Slot()
+						.HAlign(HAlign_Fill)
+						.Padding(6,2,6,4)
+						[
+							SNew(STextBlock)
+							.TextStyle(FNiagaraEditorStyle::Get(), "NiagaraEditor.PlatformSet.ButtonText")
+							.Text(FNiagaraPlatformSet::GetQualityLevelText(QualityLevel))
+							.ColorAndOpacity(this, &SNiagaraScalabilityPreviewSettings::GetQualityLevelButtonTextColor, QualityLevel)
+							.ShadowOffset(FVector2D(1, 1))
+						]
+					]
+				]
+			]
+		];
+	}
+
+	return QualityLevelWidgetsBox;
+}
+
+TSharedRef<SWidget> SNiagaraScalabilityPreviewSettings::CreatePreviewPlatformWidgets()
+{
+	return SNew(SHorizontalBox)
+		+ SHorizontalBox::Slot()
+		.AutoWidth()
+		[
+			SNew(SButton)
+			.ButtonStyle(FEditorStyle::Get(), "HoverHintOnly")
+			.OnClicked(this, &SNiagaraScalabilityPreviewSettings::OnResetPreviewPlatformClicked)
+			.ToolTipText(LOCTEXT("ResetPreviewPlatformToolTip", "Reset the platform to preview."))
+			.Visibility_Lambda([=]()
+			{
+				return ScalabilityViewModel->GetPreviewDeviceProfile().IsSet() ? EVisibility::Visible : EVisibility::Collapsed;
+			})
+			[
+				SNew(SImage)
+				.Image(FNiagaraEditorStyle::Get().GetBrush("NiagaraEditor.Scalability.Preview.ResetPlatform"))
+			]
+		]
+		+ SHorizontalBox::Slot()
+		.AutoWidth()
+		[
+			SNew(SButton)
+	       .ButtonStyle(FEditorStyle::Get(), "HoverHintOnly")
+	       .ForegroundColor(FSlateColor::UseForeground())
+	       .ToolTipText(LOCTEXT("ChoosePreviewPlatform", "Choose a platform to preview. This is cosmetic only."))
+	       .OnClicked(this, &SNiagaraScalabilityPreviewSettings::TogglePlatformMenuOpen)
+			[
+				SNew(SHorizontalBox)
+				+ SHorizontalBox::Slot()
+				.Padding(2.f)
+				.AutoWidth()
+				[
+					SNew(SBox)
+					.WidthOverride(16.f)
+					.HeightOverride(16.f)
+					.Visibility_Lambda([=]()
+				    {
+						return GetActivePreviewPlatformImage() != nullptr
+							? EVisibility::Visible
+					        : EVisibility::Collapsed;
+				    })
+					[
+						SNew(SImage)
+						.Image(this, &SNiagaraScalabilityPreviewSettings::GetActivePreviewPlatformImage)
+					]
+				]
+				+ SHorizontalBox::Slot()
+				.Padding(2.f)
+				.AutoWidth()
+				[
+					SNew(STextBlock)
+					.Text(this, &SNiagaraScalabilityPreviewSettings::GetActivePreviewPlatformName)
+				]
+				+ SHorizontalBox::Slot()
+				.Padding(2.f)
+				.HAlign(HAlign_Center)
+				.VAlign(VAlign_Center)
+				.AutoWidth()
+				[
+					SNew(SBox)
+					.WidthOverride(8.f)
+					.HeightOverride(8.f)
+					[
+						SNew(SImage)
+						.Image(FNiagaraEditorStyle::Get().GetBrush("NiagaraEditor.PlatformSet.DropdownButton"))
+					]
+				]
+			]
+		];
+}
+
+ECheckBoxState SNiagaraScalabilityPreviewSettings::IsQLChecked(int32 QualityLevel)const
+{	
+	return ScalabilityViewModel->IsViewModeQualityEnabled(QualityLevel) ? ECheckBoxState::Checked : ECheckBoxState::Unchecked;	
+}
+
+void SNiagaraScalabilityPreviewSettings::QLCheckStateChanged(ECheckBoxState CheckState, int32 QualityLevel)
+{
+	ScalabilityViewModel->UpdatePreviewQualityLevel(QualityLevel, CheckState == ECheckBoxState::Checked);
+}
+
+FSlateColor SNiagaraScalabilityPreviewSettings::GetQualityLevelButtonTextColor(int32 QualityLevel) const
+{
+	return ScalabilityViewModel->IsViewModeQualityEnabled(QualityLevel) ? 
+		FSlateColor(FLinearColor(0.95f, 0.95f, 0.95f)) :
+		FSlateColor::UseForeground();
+}
+
+const FSlateBrush* SNiagaraScalabilityPreviewSettings::GetActivePreviewPlatformImage() const
+{
+	UDeviceProfile* DeviceProfile = ScalabilityViewModel->GetPreviewDeviceProfile().Get(nullptr);
+
+	FName PlatformName;
+	if(DeviceProfile)
+	{
+		PlatformName = *DeviceProfile->DeviceType;
+	}
+	
+	if (const PlatformInfo::FTargetPlatformInfo* Info = PlatformInfo::FindPlatformInfo(PlatformName))
+	{
+		const FSlateBrush* DeviceProfileTypeIcon = FEditorStyle::GetBrush(Info->GetIconStyleName(EPlatformIconSize::Normal));
+		if (DeviceProfileTypeIcon != FEditorStyle::Get().GetDefaultBrush())
+		{
+			return DeviceProfileTypeIcon;
+		}
+	}
+
+	return nullptr;
+}
+
+FText SNiagaraScalabilityPreviewSettings::GetActivePreviewPlatformName() const
+{
+	if(ScalabilityViewModel->GetPreviewDeviceProfile().IsSet())
+	{
+		return FText::FromString(ScalabilityViewModel->GetPreviewDeviceProfile().GetValue()->GetName());
+	}
+	else
+	{
+		return LOCTEXT("Scalability_NoDeviceProfileSet", "No Preview Device Set");
+	}
+}
+
+FText SNiagaraScalabilityPreviewSettings::GetQualityButtonTooltip(int32 QualityLevel) const
+{
+	FText QualityLevelText = FNiagaraPlatformSet::GetQualityLevelText(QualityLevel);
+	return FText::FormatOrdered(LOCTEXT("QualityButtonTooltip", "Activates {0} preview. This is cosmetic only.\nEmitters and renderers excluded from preview will show a red overlay and will not be active in the viewport and PIE."), QualityLevelText);
+}
+
+void SNiagaraScalabilityPreviewSettings::CreateDeviceProfileTree()
+{
+	//Pull device profiles out by their hierarchy depth.
+	TArray<TArray<UDeviceProfile*>> DeviceProfilesByHierarchyLevel;
+	for (UObject* Profile : UDeviceProfileManager::Get().Profiles)
+	{
+		UDeviceProfile* DeviceProfile = CastChecked<UDeviceProfile>(Profile);
+
+		TFunction<void(int32&, UDeviceProfile*)> FindDepth;
+		FindDepth = [&](int32& Depth, UDeviceProfile* CurrProfile)
+		{
+			if (CurrProfile->Parent)
+			{
+				FindDepth(++Depth, Cast<UDeviceProfile>(CurrProfile->Parent));
+			}
+		};
+
+		int32 ProfileDepth = 0;
+		FindDepth(ProfileDepth, DeviceProfile);
+		DeviceProfilesByHierarchyLevel.SetNum(FMath::Max(ProfileDepth+1, DeviceProfilesByHierarchyLevel.Num()));
+		DeviceProfilesByHierarchyLevel[ProfileDepth].Add(DeviceProfile);
+	}
+	
+	FullDeviceProfileTree.Reset(DeviceProfilesByHierarchyLevel[0].Num());
+	for (int32 RootProfileIdx = 0; RootProfileIdx < DeviceProfilesByHierarchyLevel[0].Num(); ++RootProfileIdx)
+	{
+		TFunction<void(FNiagaraDeviceProfileViewModel*, int32)> BuildProfileTree;
+		BuildProfileTree = [&](FNiagaraDeviceProfileViewModel* CurrRoot, int32 CurrLevel)
+		{
+			int32 NextLevel = CurrLevel + 1;
+			if (NextLevel < DeviceProfilesByHierarchyLevel.Num())
+			{
+				for (UDeviceProfile* PossibleChild : DeviceProfilesByHierarchyLevel[NextLevel])
+				{
+					if (PossibleChild->Parent == CurrRoot->Profile)
+					{
+						TSharedPtr<FNiagaraDeviceProfileViewModel>& NewChild = CurrRoot->Children.Add_GetRef(MakeShared<FNiagaraDeviceProfileViewModel>());
+						NewChild->Profile = PossibleChild;
+						BuildProfileTree(NewChild.Get(), NextLevel);
+					}
+				}
+			}
+		};
+
+		//Add all root nodes and build their trees.
+		TSharedPtr<FNiagaraDeviceProfileViewModel> CurrRoot = MakeShared<FNiagaraDeviceProfileViewModel>();
+		CurrRoot->Profile = DeviceProfilesByHierarchyLevel[0][RootProfileIdx];
+		BuildProfileTree(CurrRoot.Get(), 0);
+		FullDeviceProfileTree.Add(CurrRoot);
+	}
+
+	const UNiagaraSettings* Settings = GetDefault<UNiagaraSettings>();
+	check(Settings);
+
+	int32 NumQualityLevels = Settings->QualityLevels.Num();
+	FilteredDeviceProfileTrees.SetNum(NumQualityLevels);
+	
+	for (TSharedPtr<FNiagaraDeviceProfileViewModel>& FullDeviceRoot : FullDeviceProfileTree)
+	{
+		for (int32 QualityLevel = 0; QualityLevel < NumQualityLevels; ++QualityLevel)
+		{			
+			TArray<TSharedPtr<FNiagaraDeviceProfileViewModel>>& FilteredRoots = FilteredDeviceProfileTrees[QualityLevel];
+
+			TSharedPtr<FNiagaraDeviceProfileViewModel>& FilteredRoot = FilteredRoots.Add_GetRef(MakeShared<FNiagaraDeviceProfileViewModel>());
+			FilteredRoot->Profile = FullDeviceRoot->Profile;
+		}
+	}
+}
+
+TSharedRef<SWidget> SNiagaraScalabilityPreviewSettings::GenerateDeviceProfileTreeWidget()
+{
+	if (FullDeviceProfileTree.Num() == 0)
+	{
+		CreateDeviceProfileTree();	
+	}
+
+	TArray<TSharedPtr<FNiagaraDeviceProfileViewModel>>* TreeToUse = &FullDeviceProfileTree;
+
+	TSharedPtr<FNiagaraDeviceProfileViewModel> ClearViewModel = nullptr;
+	return SNew(SBorder)
+		.BorderImage(FEditorStyle::Get().GetBrush("Menu.Background"))
+		[
+			SAssignNew(DeviceProfileTreeWidget, STreeView<TSharedPtr<FNiagaraDeviceProfileViewModel>>)
+			.TreeItemsSource(TreeToUse)
+			.OnGenerateRow(this, &SNiagaraScalabilityPreviewSettings::OnGenerateDeviceProfileTreeRow)
+			.OnGetChildren(this, &SNiagaraScalabilityPreviewSettings::OnGetDeviceProfileTreeChildren)
+			.SelectionMode(ESelectionMode::None)
+		];
+}
+
+TSharedRef<ITableRow> SNiagaraScalabilityPreviewSettings::OnGenerateDeviceProfileTreeRow(
+	TSharedPtr<FNiagaraDeviceProfileViewModel> InItem, const TSharedRef<STableViewBase>& OwnerTable)
+{
+	TSharedPtr<SHorizontalBox> RowContainer;
+	SAssignNew(RowContainer, SHorizontalBox);
+
+	int32 ProfileMask = ScalabilityViewModel->GetPreviewPlatformSet()->GetEffectQualityMaskForDeviceProfile(InItem->Profile);
+	FText NameTooltip = FText::Format(LOCTEXT("ProfileQLTooltipFmt", "Effects Quality: {0}"), FNiagaraPlatformSet::GetQualityLevelMaskText(ProfileMask));
+	
+	//Top level profile. Look for a platform icon.
+	if (InItem->Profile->Parent == nullptr)
+	{
+		if (const PlatformInfo::FTargetPlatformInfo* Info = PlatformInfo::FindPlatformInfo(*InItem->Profile->DeviceType))
+		{
+			const FSlateBrush* DeviceProfileTypeIcon = FEditorStyle::GetBrush(Info->GetIconStyleName(EPlatformIconSize::Normal));
+			if (DeviceProfileTypeIcon != FEditorStyle::Get().GetDefaultBrush())
+			{
+				RowContainer->AddSlot()
+					.AutoWidth()
+					.VAlign(VAlign_Center)
+					.HAlign(HAlign_Left)
+					.Padding(4, 0, 0, 0)
+					[
+						SNew(SBox)
+						.WidthOverride(16)
+						.HeightOverride(16)
+						[
+							SNew(SImage)
+							.Image(DeviceProfileTypeIcon)
+						]
+					];
+			}
+		}
+	}
+
+	FName TextStyleName("NormalText");
+	FSlateColor TextColor(FSlateColor::UseForeground());
+
+	RowContainer->AddSlot()
+		.Padding(4, 2, 0, 2)
+		.HAlign(HAlign_Fill)
+		.VAlign(VAlign_Center)
+		[
+			SNew(SButton)
+			.ButtonStyle(FEditorStyle::Get(), "NoBorder")
+			.OnClicked(this, &SNiagaraScalabilityPreviewSettings::OnProfileMenuButtonClicked, InItem)
+			.ForegroundColor(TextColor)
+			.ToolTipText(NameTooltip)
+			[
+				SNew(STextBlock)
+				.TextStyle(FEditorStyle::Get(), TextStyleName)
+				.Text(FText::FromString(InItem->Profile->GetName()))
+			]
+		];
+
+	return SNew(STableRow<TSharedPtr<FNiagaraDeviceProfileViewModel>>, OwnerTable)
+		.Style(FNiagaraEditorStyle::Get(), "NiagaraEditor.PlatformSet.TreeView")
+		[
+			RowContainer.ToSharedRef()
+		];
+}
+
+void SNiagaraScalabilityPreviewSettings::OnGetDeviceProfileTreeChildren(TSharedPtr<FNiagaraDeviceProfileViewModel> InItem,
+	TArray<TSharedPtr<FNiagaraDeviceProfileViewModel>>& OutChildren)
+{
+	OutChildren = InItem->Children;
+}
+
+FReply SNiagaraScalabilityPreviewSettings::TogglePlatformMenuOpen()
+{
+	PlatformMenuAnchor->SetIsOpen(!PlatformMenuAnchor->IsOpen());
+	return FReply::Handled();
+}
+
+FReply SNiagaraScalabilityPreviewSettings::OnProfileMenuButtonClicked(TSharedPtr<FNiagaraDeviceProfileViewModel> InItem)
+{
+	ScalabilityViewModel->UpdatePreviewDeviceProfile(InItem->Profile);
+	TogglePlatformMenuOpen();
+	
+	return FReply::Handled();
+}
+
+FReply SNiagaraScalabilityPreviewSettings::OnResetPreviewPlatformClicked()
+{
+	ScalabilityViewModel->UpdatePreviewDeviceProfile(nullptr);	
+	return FReply::Handled();
+}
+
+#undef LOCTEXT_NAMESPACE
diff --git a/Engine/Plugins/FX/Niagara/Source/NiagaraEditorWidgets/Private/SNiagaraScalabilityPreviewSettings.h b/Engine/Plugins/FX/Niagara/Source/NiagaraEditorWidgets/Private/SNiagaraScalabilityPreviewSettings.h
new file mode 100644
--- /dev/null
+++ b/Engine/Plugins/FX/Niagara/Source/NiagaraEditorWidgets/Private/SNiagaraScalabilityPreviewSettings.h
@@ -0,0 +1,51 @@
+// Copyright Epic Games, Inc. All Rights Reserved.
+
+#pragma once
+
+#include "CoreMinimal.h"
+#include "NiagaraSystem.h"
+#include "ViewModels/NiagaraSystemScalabilityViewModel.h"
+#include "Widgets/Layout/SGridPanel.h"
+#include "Widgets/Views/STreeView.h"
+#include "Widgets/Input/SMenuAnchor.h"
+#include "ViewModels/NiagaraSystemViewModel.h"
+
+struct FNiagaraDeviceProfileViewModel;
+class SNiagaraScalabilityPreviewSettings : public SCompoundWidget
+{
+public:
+	
+	SLATE_BEGIN_ARGS(SNiagaraScalabilityPreviewSettings)
+	{}
+	SLATE_END_ARGS()
+
+	void Construct(const FArguments& InArgs, UNiagaraSystemScalabilityViewModel& InScalabilityViewModel);
+
+private:
+	TSharedRef<SWidget> CreatePreviewQualityLevelWidgets();
+	TSharedRef<SWidget> CreatePreviewPlatformWidgets();
+
+	ECheckBoxState IsQLChecked(int32 QualityLevel) const;
+	void QLCheckStateChanged(ECheckBoxState CheckState, int32 QualityLevel);
+	FSlateColor GetQualityLevelButtonTextColor(int32 QualityLevel) const;
+	const FSlateBrush* GetActivePreviewPlatformImage() const;
+	FText GetActivePreviewPlatformName() const;
+	FText GetQualityButtonTooltip(int32 QualityLevel) const;
+
+	void CreateDeviceProfileTree();
+	TSharedRef<SWidget> GenerateDeviceProfileTreeWidget();
+	TSharedRef<ITableRow> OnGenerateDeviceProfileTreeRow(TSharedPtr<FNiagaraDeviceProfileViewModel> InItem, const TSharedRef<STableViewBase>& OwnerTable);
+	void OnGetDeviceProfileTreeChildren(TSharedPtr<FNiagaraDeviceProfileViewModel> InItem, TArray< TSharedPtr<FNiagaraDeviceProfileViewModel> >& OutChildren);
+	FReply TogglePlatformMenuOpen();
+
+	FReply OnProfileMenuButtonClicked(TSharedPtr<FNiagaraDeviceProfileViewModel> InItem);
+	FReply OnResetPreviewPlatformClicked();
+	
+	TArray<TSharedPtr<FNiagaraDeviceProfileViewModel>> FullDeviceProfileTree;
+	TArray<TArray<TSharedPtr<FNiagaraDeviceProfileViewModel>>> FilteredDeviceProfileTrees;
+	TSharedPtr<STreeView<TSharedPtr<FNiagaraDeviceProfileViewModel>>> DeviceProfileTreeWidget;
+
+private:
+	TWeakObjectPtr<UNiagaraSystemScalabilityViewModel> ScalabilityViewModel;
+	TSharedPtr<SMenuAnchor> PlatformMenuAnchor;
+};
diff --git a/Engine/Plugins/FX/Niagara/Source/NiagaraEditorWidgets/Private/SNiagaraSystemEffectTypeBar.cpp b/Engine/Plugins/FX/Niagara/Source/NiagaraEditorWidgets/Private/SNiagaraSystemEffectTypeBar.cpp
new file mode 100644
--- /dev/null
+++ b/Engine/Plugins/FX/Niagara/Source/NiagaraEditorWidgets/Private/SNiagaraSystemEffectTypeBar.cpp
@@ -0,0 +1,221 @@
+// Copyright Epic Games, Inc. All Rights Reserved.
+
+#include "SNiagaraSystemEffectTypeBar.h"
+#include "ISinglePropertyView.h"
+#include "NiagaraEditorStyle.h"
+#include "NiagaraEditorWidgetsStyle.h"
+#include "Delegates/DelegateInstanceInterface.h"
+#include "Modules/ModuleManager.h"
+#include "Styling/StyleColors.h"
+#include "Widgets/Images/SImage.h"
+#include "Widgets/Layout/SWidgetSwitcher.h"
+#include "Widgets/Layout/SSeparator.h"
+#include "PropertyEditorModule.h"
+
+#define LOCTEXT_NAMESPACE "NiagaraSystemEffectType"
+
+SNiagaraSystemEffectTypeBar::~SNiagaraSystemEffectTypeBar()
+{
+	System->OnScalabilityChanged().RemoveAll(this);
+}
+
+void SNiagaraSystemEffectTypeBar::Construct(const FArguments& InArgs, UNiagaraSystem& InSystem)
+{
+	System = &InSystem;
+	EffectType = System->GetEffectType();
+	
+	EffectTypeValuesBox = SNew(SWrapBox).UseAllottedSize(true);
+	
+	FPropertyEditorModule& PropertyEditorModule = FModuleManager::Get().LoadModuleChecked<FPropertyEditorModule>("PropertyEditor");
+	FSinglePropertyParams EffectTypePropertyArgs;
+	EffectTypePropertyArgs.bHideAssetThumbnail = true;
+	
+	TAttribute<int32> WidgetSwitcherIndexAttribute = TAttribute<int32>::CreateSP(this, &SNiagaraSystemEffectTypeBar::GetActiveDetailsWidgetIndex);
+	ChildSlot
+	.Padding(5.f)
+	[
+		SNew(SHorizontalBox)
+		+ SHorizontalBox::Slot()
+		.AutoWidth()
+		[
+			PropertyEditorModule.CreateSingleProperty(System.Get(), "EffectType", EffectTypePropertyArgs).ToSharedRef()
+		]
+		+ SHorizontalBox::Slot()
+		.AutoWidth()
+		[
+			SNew(SSeparator)
+			.Orientation(Orient_Vertical)
+		]
+		+ SHorizontalBox::Slot()
+		[
+			SNew(SWidgetSwitcher)
+			.WidgetIndex(WidgetSwitcherIndexAttribute)
+			+ SWidgetSwitcher::Slot()
+			[
+				SNew(STextBlock)
+				.Text(LOCTEXT("NoEffectTypeSelected", "No Effect Type assigned.\nEffect Types are used to provide default scalability settings that can partly be overridden."))
+				.AutoWrapText(true)
+			]
+			+ SWidgetSwitcher::Slot()
+			[
+				EffectTypeValuesBox.ToSharedRef()
+			]
+		]
+	];
+	
+	UpdateEffectTypeWidgets();
+	System->OnScalabilityChanged().AddSP(this, &SNiagaraSystemEffectTypeBar::UpdateEffectType);
+}
+
+void SNiagaraSystemEffectTypeBar::UpdateEffectTypeWidgets()
+{
+	EffectTypeValuesBox->ClearChildren();
+	
+	for(TFieldIterator<FProperty> PropIt(UNiagaraEffectType::StaticClass(), EFieldIterationFlags::None); PropIt; ++PropIt)
+	{
+		// we only allow display of properties of effect type that are explicitly marked
+		if(!PropIt->HasMetaData(TEXT("DisplayInSystemScalability")))
+		{
+			continue;
+		}
+		
+		FName PropertyName = PropIt->GetFName();
+		FText PropertyDisplayNameText = PropIt->GetDisplayNameText();
+		FSinglePropertyParams PropertyParams;
+		PropertyParams.NamePlacement = EPropertyNamePlacement::Hidden;
+
+		if(PropIt->HasMetaData("ScalabilityBarDisplayName"))
+		{
+			PropertyDisplayNameText = PropIt->GetMetaDataText("ScalabilityBarDisplayName");
+		}
+		
+		TSharedPtr<SWidget> PropertyValueWidget;
+		PropertyValueWidget = CreatePropertyValueWidget(*PropIt);
+		
+		if(PropertyValueWidget.IsValid())
+		{
+			FName ImageName = FName(FString("NiagaraEditor.Scalability.EffectType.") + PropertyName.ToString());
+			const FSlateBrush* ImageBrush = FNiagaraEditorWidgetsStyle::Get().GetBrush(ImageName);
+
+			TSharedPtr<SWidget> PropertyIconWidget = SNew(SImage)
+				.Image(ImageBrush)
+				.ToolTipText(PropIt->GetToolTipText());			
+			
+			EffectTypeValuesBox->AddSlot()
+			.VAlign(VAlign_Center)
+			.Padding(2.f, 2.f)
+			[
+				SNew(SHorizontalBox)
+				+ SHorizontalBox::Slot()
+				.AutoWidth()
+				.VAlign(VAlign_Center)
+				.Padding(3.f)
+				[
+					PropertyIconWidget.ToSharedRef()
+				]
+				+ SHorizontalBox::Slot()
+				[
+					SNew(SVerticalBox)
+					+ SVerticalBox::Slot()
+					.Padding(2.f, 1.f)
+					.HAlign(HAlign_Left)
+					.VAlign(VAlign_Center)
+					[
+						SNew(STextBlock)
+						.Text(PropertyDisplayNameText)
+						.ToolTipText(PropIt->GetToolTipText())
+						.TextStyle(&FNiagaraEditorWidgetsStyle::Get().GetWidgetStyle<FTextBlockStyle>("NiagaraEditor.Scalability.EffectType.Property"))
+					]
+					+ SVerticalBox::Slot()
+					.Padding(2.f, 1.f)
+					.HAlign(HAlign_Left)
+					.VAlign(VAlign_Center)
+					[
+						PropertyValueWidget.ToSharedRef()
+					]
+				]
+			];
+		}
+	}
+}
+
+TSharedPtr<SWidget> SNiagaraSystemEffectTypeBar::CreatePropertyValueWidget(FProperty* Property)
+{
+	TSharedRef<STextBlock> ValueTextBlock = SNew(STextBlock)
+		.TextStyle(&FNiagaraEditorWidgetsStyle::Get().GetWidgetStyle<FTextBlockStyle>("NiagaraEditor.Scalability.EffectType.Property"));
+
+	if(!EffectType.IsValid())
+	{
+		ValueTextBlock->SetText(FText::FromString("-"));
+		ValueTextBlock->SetColorAndOpacity(FStyleColors::AccentGray);
+	}
+	else
+	{
+		ValueTextBlock->SetColorAndOpacity(FStyleColors::AccentBlue);
+
+		if(FEnumProperty* EnumProperty = CastField<FEnumProperty>(Property))
+		{
+			int32 EnumValue = *Property->ContainerPtrToValuePtr<int32>(EffectType.Get());
+			ValueTextBlock->SetText(EnumProperty->GetEnum()->GetDisplayNameTextByValue(EnumValue));
+		}
+
+		if(FBoolProperty* BoolProperty = CastField<FBoolProperty>(Property))
+		{
+			bool BoolValue = BoolProperty->GetPropertyValue_InContainer(EffectType.Get());
+			ValueTextBlock->SetText(BoolValue ? FText::FromString("True") : FText::FromString("False"));
+		}
+				
+		if(FObjectPtrProperty* ObjectProperty = CastField<FObjectPtrProperty>(Property))
+		{
+			FObjectPtr& ObjectPtr = (FObjectPtr&) ObjectProperty->GetPropertyValue_InContainer(EffectType.Get());
+			
+			if(!ObjectPtr.IsNullNoResolve())
+			{				
+				if(ObjectProperty->HasMetaData("DisplayClassDisplayName"))
+				{
+					ValueTextBlock->SetText(GetClassDisplayName(ObjectPtr.Get()));
+				}
+				else
+				{
+					ValueTextBlock->SetText(GetObjectName(ObjectPtr.Get()));
+				}						
+			}
+			else
+			{
+				UObject* NullObject = nullptr;
+				ValueTextBlock->SetText(GetObjectName(NullObject));
+				ValueTextBlock->SetColorAndOpacity(FStyleColors::AccentGray);
+			}
+		}
+	}
+	
+	return ValueTextBlock;
+}
+
+FText SNiagaraSystemEffectTypeBar::GetObjectName(UObject* Object) const
+{
+	return IsValid(Object) ? FText::FromString(Object->GetName()) : LOCTEXT("NoObjectAssigned", "None");		
+}
+
+FText SNiagaraSystemEffectTypeBar::GetClassDisplayName(UObject* Object) const
+{
+	return IsValid(Object) ? Object->GetClass()->GetDisplayNameText() : LOCTEXT("NoClassAvailable", "None");
+}
+
+int32 SNiagaraSystemEffectTypeBar::GetActiveDetailsWidgetIndex() const
+{
+	if(System->GetEffectType())
+	{
+		return 1;
+	}
+
+	return 0;
+}
+
+void SNiagaraSystemEffectTypeBar::UpdateEffectType()
+{
+	EffectType = System->GetEffectType();
+	UpdateEffectTypeWidgets();
+}
+
+#undef LOCTEXT_NAMESPACE
diff --git a/Engine/Plugins/FX/Niagara/Source/NiagaraEditorWidgets/Private/SNiagaraSystemEffectTypeBar.h b/Engine/Plugins/FX/Niagara/Source/NiagaraEditorWidgets/Private/SNiagaraSystemEffectTypeBar.h
new file mode 100644
--- /dev/null
+++ b/Engine/Plugins/FX/Niagara/Source/NiagaraEditorWidgets/Private/SNiagaraSystemEffectTypeBar.h
@@ -0,0 +1,34 @@
+// Copyright Epic Games, Inc. All Rights Reserved.
+
+#pragma once
+
+#include "CoreMinimal.h"
+#include "NiagaraSystem.h"
+#include "Widgets/Layout/SGridPanel.h"
+#include "Widgets/Layout/SUniformWrapPanel.h"
+#include "Widgets/Layout/SWrapBox.h"
+
+class SNiagaraSystemEffectTypeBar : public SCompoundWidget
+{
+public:
+	SLATE_BEGIN_ARGS(SNiagaraSystemEffectTypeBar)
+	{}
+	SLATE_END_ARGS()
+
+	virtual ~SNiagaraSystemEffectTypeBar();
+	
+	void Construct(const FArguments& InArgs, UNiagaraSystem& InSystem);
+	void UpdateEffectTypeWidgets();
+private:
+	TSharedPtr<SWidget> CreatePropertyValueWidget(FProperty* Property);
+
+private:
+	FText GetObjectName(UObject* Object) const;
+	FText GetClassDisplayName(UObject* Object) const;
+	int32 GetActiveDetailsWidgetIndex() const;
+	void UpdateEffectType();
+private:
+	TWeakObjectPtr<UNiagaraEffectType> EffectType;
+	TWeakObjectPtr<UNiagaraSystem> System;
+	TSharedPtr<SWrapBox> EffectTypeValuesBox;
+};
