- IWYU fix

#rnx

diff --git a/Engine/Plugins/FX/Niagara/Source/NiagaraEditor/Private/NiagaraNodeFunctionCall.cpp b/Engine/Plugins/FX/Niagara/Source/NiagaraEditor/Private/NiagaraNodeFunctionCall.cpp
--- a/Engine/Plugins/FX/Niagara/Source/NiagaraEditor/Private/NiagaraNodeFunctionCall.cpp
+++ b/Engine/Plugins/FX/Niagara/Source/NiagaraEditor/Private/NiagaraNodeFunctionCall.cpp
@@ -2,6 +2,7 @@
 
 #include "NiagaraNodeFunctionCall.h"
 
+#include "Editor.h"
 #include "AssetRegistryModule.h"
 #include "EdGraphSchema_Niagara.h"
 #include "NiagaraComponent.h"
@@ -963,7 +964,12 @@ void UNiagaraNodeFunctionCall::OpenReferencedAsset() const
 		{
 			FunctionScript->VersionToOpenInEditor = SelectedScriptVersion;
 		}
-		GEditor->GetEditorSubsystem<UAssetEditorSubsystem>()->OpenEditorForAsset(FunctionScript);
+		#if WITH_EDITOR
+		if ( GEditor )
+		{
+			GEditor->GetEditorSubsystem<UAssetEditorSubsystem>()->OpenEditorForAsset(FunctionScript);
+		}
+		#endif
 	}
 }
 
