Fixed (some) tab icons in Niagara scripts

diff --git a/Engine/Plugins/FX/Niagara/Source/NiagaraEditor/Private/Toolkits/NiagaraScriptToolkit.cpp b/Engine/Plugins/FX/Niagara/Source/NiagaraEditor/Private/Toolkits/NiagaraScriptToolkit.cpp
--- a/Engine/Plugins/FX/Niagara/Source/NiagaraEditor/Private/Toolkits/NiagaraScriptToolkit.cpp
+++ b/Engine/Plugins/FX/Niagara/Source/NiagaraEditor/Private/Toolkits/NiagaraScriptToolkit.cpp
@@ -107,11 +107,13 @@ void FNiagaraScriptToolkit::RegisterTabSpawners(const TSharedRef<class FTabManag
 
 	InTabManager->RegisterTabSpawner(ParametersTabId, FOnSpawnTab::CreateSP(this, &FNiagaraScriptToolkit::SpawnTabScriptParameters))
 		.SetDisplayName(LOCTEXT("ParametersTab", "Parameters"))
-		.SetGroup(WorkspaceMenuCategoryRef);
+		.SetGroup(WorkspaceMenuCategoryRef)
+		.SetIcon(FSlateIcon(FNiagaraEditorStyle::Get().GetStyleSetName(), "Tab.Parameters"));
 
 	InTabManager->RegisterTabSpawner(ParametersTabId2, FOnSpawnTab::CreateSP(this, &FNiagaraScriptToolkit::SpawnTabScriptParameters2)) //@todo(ng) remove this
 		.SetDisplayName(LOCTEXT("ParametersTab2", "Legacy Parameters"))
-		.SetGroup(WorkspaceMenuCategoryRef);
+		.SetGroup(WorkspaceMenuCategoryRef)
+		.SetIcon(FSlateIcon(FNiagaraEditorStyle::Get().GetStyleSetName(), "Tab.Parameters"));
 		
 	//@todo(ng) disable parameter definitions panel pending bug fixes
 // 	InTabManager->RegisterTabSpawner(ParameterDefinitionsTabId, FOnSpawnTab::CreateSP(this, &FNiagaraScriptToolkit::SpawnTabParameterDefinitions))
@@ -125,11 +127,13 @@ void FNiagaraScriptToolkit::RegisterTabSpawners(const TSharedRef<class FTabManag
 
 	InTabManager->RegisterTabSpawner(MessageLogTabID, FOnSpawnTab::CreateSP(this, &FNiagaraScriptToolkit::SpawnTabMessageLog))
 		.SetDisplayName(LOCTEXT("NiagaraMessageLogTab", "Niagara Message Log"))
-		.SetGroup(WorkspaceMenuCategoryRef);
+		.SetGroup(WorkspaceMenuCategoryRef)
+		.SetIcon(FSlateIcon(FNiagaraEditorStyle::Get().GetStyleSetName(), "Tab.Log"));
 
 	InTabManager->RegisterTabSpawner(VersioningTabID, FOnSpawnTab::CreateSP(this, &FNiagaraScriptToolkit::SpawnTabVersioning))
         .SetDisplayName(LOCTEXT("VersioningTab", "Versioning"))
-        .SetGroup(WorkspaceMenuCategoryRef);
+        .SetGroup(WorkspaceMenuCategoryRef)
+		.SetIcon(FSlateIcon(FAppStyle::GetAppStyleSetName(), "Versions"));
 }
 
 void FNiagaraScriptToolkit::UnregisterTabSpawners(const TSharedRef<class FTabManager>& InTabManager)
