Niagara: Fixed nullptr access when loading certain FX in cooked editors

diff --git a/Engine/Plugins/FX/Niagara/Source/Niagara/Private/NiagaraSystem.cpp b/Engine/Plugins/FX/Niagara/Source/Niagara/Private/NiagaraSystem.cpp
--- a/Engine/Plugins/FX/Niagara/Source/Niagara/Private/NiagaraSystem.cpp
+++ b/Engine/Plugins/FX/Niagara/Source/Niagara/Private/NiagaraSystem.cpp
@@ -3853,7 +3853,10 @@ void UNiagaraSystem::FixupPositionUserParameters()
 		if (Script)
 		{
 			Script->ConditionalPostLoad();
-			Script->GetLatestSource()->GetLinkedPositionTypeInputs(UserParameters, LinkedPositionInputs);
+			if (UNiagaraScriptSourceBase* LatestSource = Script->GetLatestSource())
+			{
+				LatestSource->GetLinkedPositionTypeInputs(UserParameters, LinkedPositionInputs);
+			}
 		}
 	});
 
