Changed asset source to be considered under Niagara when the asset is under engineplugins/fx instead of engineplugins/fx/niagara

Also reenabled user categories for new emitter dialogs

diff --git a/Engine/Plugins/FX/Niagara/Source/NiagaraEditor/Private/NiagaraEditorUtilities.cpp b/Engine/Plugins/FX/Niagara/Source/NiagaraEditor/Private/NiagaraEditorUtilities.cpp
--- a/Engine/Plugins/FX/Niagara/Source/NiagaraEditor/Private/NiagaraEditorUtilities.cpp
+++ b/Engine/Plugins/FX/Niagara/Source/NiagaraEditor/Private/NiagaraEditorUtilities.cpp
@@ -1884,7 +1884,7 @@ TTuple<EScriptSource, FText> FNiagaraEditorUtilities::GetScriptSource(const FAss
 	FString PackagePathLocal ="";
 	FPackageName::TryConvertGameRelativePackagePathToLocalPath(ScriptAssetData.PackagePath.ToString(), PackagePathLocal);
 
-	if(FPaths::IsUnderDirectory(PackagePathLocal, FPaths::EnginePluginsDir() / TEXT("FX/Niagara")))
+	if(FPaths::IsUnderDirectory(PackagePathLocal, FPaths::EnginePluginsDir() / TEXT("FX")))
 	{
 		int32 ContentFoundIndex = PackagePathLocal.Find(TEXT("/Content"));
 
diff --git a/Engine/Plugins/FX/Niagara/Source/NiagaraEditor/Private/SNewEmitterDialog.cpp b/Engine/Plugins/FX/Niagara/Source/NiagaraEditor/Private/SNewEmitterDialog.cpp
--- a/Engine/Plugins/FX/Niagara/Source/NiagaraEditor/Private/SNewEmitterDialog.cpp
+++ b/Engine/Plugins/FX/Niagara/Source/NiagaraEditor/Private/SNewEmitterDialog.cpp
@@ -20,6 +20,7 @@ void SNewEmitterDialog::Construct(const FArguments& InArgs)
 	FNiagaraAssetPickerListViewOptions DisplayAllViewOptions;
 	DisplayAllViewOptions.SetExpandTemplateAndLibraryAssets(true);
 	DisplayAllViewOptions.SetCategorizeLibraryAssets(true);
+	DisplayAllViewOptions.SetCategorizeUserDefinedCategory(true);
 	DisplayAllViewOptions.SetAddLibraryOnlyCheckbox(true);
 
 	SNiagaraTemplateTabBox::FNiagaraTemplateTabOptions TabOptions;
