Fixed op node input pins being wired into an add pin resulting in integer pins being spawned regardless of input pin type

diff --git a/Engine/Plugins/FX/Niagara/Source/NiagaraEditor/Private/NiagaraNodeOp.cpp b/Engine/Plugins/FX/Niagara/Source/NiagaraEditor/Private/NiagaraNodeOp.cpp
--- a/Engine/Plugins/FX/Niagara/Source/NiagaraEditor/Private/NiagaraNodeOp.cpp
+++ b/Engine/Plugins/FX/Niagara/Source/NiagaraEditor/Private/NiagaraNodeOp.cpp
@@ -527,10 +527,15 @@ void UNiagaraNodeOp::HandleStaticInputPinUpgrade(UEdGraphPin* Pin)
 		{
 			// Handle linkage...
 			bool bHandled = false;
-			if (Pin->LinkedTo.Num() > 0)
+			// this shouldn't specifically check for add pins as there can be other special pins like the wildcard pin
+			// this will generally attempt to set the current input pin to the same type as the other's output pin. Ok for now, but might break in the future
+			if (Pin->LinkedTo.Num() > 0 && !IsAddPin(Pin->LinkedTo[0]))
 			{
 				FNiagaraTypeDefinition LinkedPinType = Schema->PinToTypeDefinition(Pin->LinkedTo[0]);
-				Pin->PinType = Schema->TypeDefinitionToPinType(LinkedPinType);
+				if(LinkedPinType.IsStatic())
+				{
+					Pin->PinType = Schema->TypeDefinitionToPinType(LinkedPinType);
+				}
 			}
 			else // Handle unlinkage...
 			{
diff --git a/Engine/Plugins/FX/Niagara/Source/NiagaraEditor/Private/NiagaraNodeStaticSwitch.cpp b/Engine/Plugins/FX/Niagara/Source/NiagaraEditor/Private/NiagaraNodeStaticSwitch.cpp
--- a/Engine/Plugins/FX/Niagara/Source/NiagaraEditor/Private/NiagaraNodeStaticSwitch.cpp
+++ b/Engine/Plugins/FX/Niagara/Source/NiagaraEditor/Private/NiagaraNodeStaticSwitch.cpp
@@ -302,7 +302,7 @@ void UNiagaraNodeStaticSwitch::AllocateDefaultPins()
 bool UNiagaraNodeStaticSwitch::AllowNiagaraTypeForAddPin(const FNiagaraTypeDefinition& InType) const
 {
 	// explicitly allow parameter maps and numeric types
-	return InType.GetScriptStruct() != nullptr;
+	return InType.GetScriptStruct() != nullptr && !InType.IsInternalType();
 }
 
 bool UNiagaraNodeStaticSwitch::GetVarIndex(FHlslNiagaraTranslator* Translator, int32 InputPinCount, int32& VarIndexOut) const
