Niagara - Update the post load for niagara script variables so that it generates uinique, but stable "guid" for the script variable based on a hash of the object path and variable data.  This will fix issues in summary view data loss.

diff --git a/Engine/Plugins/FX/Niagara/Source/NiagaraEditor/Private/NiagaraScriptVariable.cpp b/Engine/Plugins/FX/Niagara/Source/NiagaraEditor/Private/NiagaraScriptVariable.cpp
--- a/Engine/Plugins/FX/Niagara/Source/NiagaraEditor/Private/NiagaraScriptVariable.cpp
+++ b/Engine/Plugins/FX/Niagara/Source/NiagaraEditor/Private/NiagaraScriptVariable.cpp
@@ -85,7 +85,21 @@ void UNiagaraScriptVariable::PostLoad()
 
 	if (!Metadata.GetVariableGuid().IsValid())
 	{
-		Metadata.CreateNewGuid();
+		// If the variable doesn't have a valid guid we need to assign one, but we want it to be stable across sessions
+		// so we're going to generate a hash from an identifier string and then use that to generate the values for the GUID.
+		// This won't result in a value that's a true GUID but since every object path in a project must be unique, the
+		// md5 of these values should be more than unique enough for our purposes.
+		FString IdentifierString = GetPathName() + Variable.GetName().ToString() + Variable.GetType().GetName();
+
+		// Create an md5 hash for the string data.
+		uint32 HashBuffer[] { 0, 0, 0, 0 };
+		FMD5 IdentifierStringHash;
+		IdentifierStringHash.Update((uint8*)GetData(IdentifierString), IdentifierString.Len() * sizeof(TCHAR));
+		IdentifierStringHash.Final((uint8*)&HashBuffer);
+
+		// Assign the guid components from the hash.
+		FGuid GuidFromHash = FGuid(HashBuffer[0], HashBuffer[1], HashBuffer[2], HashBuffer[3]);
+		Metadata.SetVariableGuid(GuidFromHash);
 	}
 	if (ChangeId.IsValid() == false)
 	{
