Fixed a bug where duplicating a module would also duplicate all the script variable guids

diff --git a/Engine/Plugins/FX/Niagara/Source/NiagaraEditor/Private/NiagaraGraph.cpp b/Engine/Plugins/FX/Niagara/Source/NiagaraEditor/Private/NiagaraGraph.cpp
--- a/Engine/Plugins/FX/Niagara/Source/NiagaraEditor/Private/NiagaraGraph.cpp
+++ b/Engine/Plugins/FX/Niagara/Source/NiagaraEditor/Private/NiagaraGraph.cpp
@@ -907,6 +907,26 @@ void UNiagaraGraph::BeginDestroy()
 	ReleaseCompilationCopy();
 }
 
+void UNiagaraGraph::PostDuplicate(EDuplicateMode::Type DuplicateMode)
+{
+	Super::PostDuplicate(DuplicateMode);
+
+	// when duplicating a module script, we need to change the guids of the script variables in the graph
+	// the IsAsset() check makes sure that we don't change anything when opening a graph to be edited, as it's duplicated there as well
+	if (UNiagaraScript* Outer = GetTypedOuter<UNiagaraScript>(); DuplicateMode == EDuplicateMode::Normal && Outer && Outer->IsAsset())
+	{
+		ensure(bIsForCompilationOnly == false);
+		
+		for (auto It : VariableToScriptVariable)
+		{
+			if (TObjectPtr<UNiagaraScriptVariable> ScriptVariable = It.Value)
+			{
+				ScriptVariable->Metadata.CreateNewGuid();
+			}
+		}
+	}
+}
+
 class UNiagaraScriptSource* UNiagaraGraph::GetSource() const
 {
 	return CastChecked<UNiagaraScriptSource>(GetOuter());
diff --git a/Engine/Plugins/FX/Niagara/Source/NiagaraEditor/Public/NiagaraGraph.h b/Engine/Plugins/FX/Niagara/Source/NiagaraEditor/Public/NiagaraGraph.h
--- a/Engine/Plugins/FX/Niagara/Source/NiagaraEditor/Public/NiagaraGraph.h
+++ b/Engine/Plugins/FX/Niagara/Source/NiagaraEditor/Public/NiagaraGraph.h
@@ -147,6 +147,7 @@ class UNiagaraGraph : public UEdGraph
 	virtual void PostLoad() override;
 	virtual void PostEditChangeProperty(struct FPropertyChangedEvent& PropertyChangedEvent) override;
 	virtual void BeginDestroy() override;
+	virtual void PostDuplicate(EDuplicateMode::Type DuplicateMode) override;
 	//~ End UObjet Interface
 	
 	/** Get the source that owns this graph */
