Niagara - When allocating pins on function calls for static switches, assign their persistent guid from their script variable so that the pin default value is retained when reallocating pins after a rename.  Updated with a post load step which fixes up existing nodes so that it doesn't break everything.

diff --git a/Engine/Plugins/FX/Niagara/Source/NiagaraCore/Public/NiagaraCustomVersion.h b/Engine/Plugins/FX/Niagara/Source/NiagaraCore/Public/NiagaraCustomVersion.h
--- a/Engine/Plugins/FX/Niagara/Source/NiagaraCore/Public/NiagaraCustomVersion.h
+++ b/Engine/Plugins/FX/Niagara/Source/NiagaraCore/Public/NiagaraCustomVersion.h
@@ -168,6 +168,8 @@ struct FNiagaraCustomVersion
 
 		ChangeSystemDeterministicDefault,	// Changed the default mode from deterministic to non-deterministic which matches emitters
 
+		StaticSwitchFunctionPinsUsePersistentGuids, // Update static switch pins to use the PersistentId from their script variable so that when they're renamed their values aren't lost when reallocating pins. 
+
 		// DO NOT ADD A NEW VERSION UNLESS YOU HAVE TALKED TO THE NIAGARA LEAD. Mismanagement of these versions can lead to data loss if it is adjusted in multiple streams simultaneously.
 		// -----<new versions can be added above this line>  -------------------------------------------------
 		VersionPlusOne,
diff --git a/Engine/Plugins/FX/Niagara/Source/NiagaraEditor/Private/NiagaraNodeFunctionCall.cpp b/Engine/Plugins/FX/Niagara/Source/NiagaraEditor/Private/NiagaraNodeFunctionCall.cpp
--- a/Engine/Plugins/FX/Niagara/Source/NiagaraEditor/Private/NiagaraNodeFunctionCall.cpp
+++ b/Engine/Plugins/FX/Niagara/Source/NiagaraEditor/Private/NiagaraNodeFunctionCall.cpp
@@ -87,6 +87,11 @@ void UNiagaraNodeFunctionCall::PostLoad()
 				}
 			}
 		}
+
+		if (NiagaraCustomVersion < FNiagaraCustomVersion::StaticSwitchFunctionPinsUsePersistentGuids)
+		{
+			UpdateStaticSwitchPinsWithPersistentGuids();
+		}
 	}
 
 	// Allow data interfaces an opportunity to intercept changes
@@ -304,6 +309,8 @@ void UNiagaraNodeFunctionCall::AllocateDefaultPins()
 				{
 					NewPin->DefaultValue = PinDefaultValue;
 				}
+
+				NewPin->PersistentGuid = ScriptVar->Metadata.GetVariableGuid();
 			}
 			else
 			{
@@ -1128,6 +1135,29 @@ void UNiagaraNodeFunctionCall::UpdatePinTooltips()
 	}
 }
 
+void UNiagaraNodeFunctionCall::UpdateStaticSwitchPinsWithPersistentGuids()
+{
+	UNiagaraGraph* CalledGraph = GetCalledGraph();
+	if (CalledGraph != nullptr)
+	{
+		TArray<UEdGraphPin*> InputPins;
+		GetInputPins(InputPins);
+		TArray<FNiagaraVariable> StaticSwitchVariables = CalledGraph->FindStaticSwitchInputs();
+		for (UEdGraphPin* InputPin : InputPins)
+		{
+			FNiagaraVariable InputVariable = GetDefault<UEdGraphSchema_Niagara>()->PinToNiagaraVariable(InputPin);
+			if (StaticSwitchVariables.Contains(InputVariable))
+			{
+				UNiagaraScriptVariable* InputScriptVariable = CalledGraph->GetScriptVariable(InputPin->PinName);
+				if (InputScriptVariable != nullptr)
+				{
+					InputPin->PersistentGuid = InputScriptVariable->Metadata.GetVariableGuid();
+				}
+			}
+		}
+	}
+}
+
 void FixupLinkedOutputs(UNiagaraNode* StartNode)
 {
 	// This searches for function call nodes linked to this node via the graph and calls FixupPinNames on them so they can react to changes in this node's called graph
diff --git a/Engine/Plugins/FX/Niagara/Source/NiagaraEditor/Public/NiagaraNodeFunctionCall.h b/Engine/Plugins/FX/Niagara/Source/NiagaraEditor/Public/NiagaraNodeFunctionCall.h
--- a/Engine/Plugins/FX/Niagara/Source/NiagaraEditor/Public/NiagaraNodeFunctionCall.h
+++ b/Engine/Plugins/FX/Niagara/Source/NiagaraEditor/Public/NiagaraNodeFunctionCall.h
@@ -205,6 +205,8 @@ protected:
 
 	void UpdatePinTooltips();
 
+	void UpdateStaticSwitchPinsWithPersistentGuids();
+
 	/** Adjusted every time that we compile this script. Lets us know that we might differ from any cached versions.*/
 	UPROPERTY(meta = (SkipForCompileHash="true"))
 	FGuid CachedChangeId;
