Fixed the niagara qualitylevel cvar not forwarding changes to the platform sets since the callback wasn't registered

diff --git a/Engine/Plugins/FX/Niagara/Source/Niagara/Private/NiagaraPlatformSet.cpp b/Engine/Plugins/FX/Niagara/Source/Niagara/Private/NiagaraPlatformSet.cpp
--- a/Engine/Plugins/FX/Niagara/Source/Niagara/Private/NiagaraPlatformSet.cpp
+++ b/Engine/Plugins/FX/Niagara/Source/Niagara/Private/NiagaraPlatformSet.cpp
@@ -70,13 +70,13 @@ static FAutoConsoleCommand GCmdSetNiagaraQualityLevelOverride(
 		}
 	)
 ); 
-// static FAutoConsoleVariableRef CVarNiagaraQualityLevel(
-// 	NiagaraQualityLevelName,
-// 	GNiagaraQualityLevel,
-// 	TEXT("The quality level for Niagara Effects. \n"),
-// 	FConsoleVariableDelegate::CreateStatic(&FNiagaraPlatformSet::OnQualityLevelChanged),
-// 	ECVF_Scalability
-// );
+static FAutoConsoleVariableRef CVarNiagaraQualityLevel(
+	NiagaraQualityLevelName,
+	GNiagaraQualityLevel,
+	TEXT("The quality level for Niagara Effects. \n"),
+	FConsoleVariableDelegate::CreateStatic(&FNiagaraPlatformSet::OnQualityLevelChanged),
+	ECVF_Scalability
+);
 
 // Override platform device profile
 // In editor all profiles will be available
