- Add setting to allow sprites to maintain a one pixel coverage while applying a coverage 'fade'
- The renderer can opt into this automatically or manually as we only cover translucent & additive for automatic currently
- The project can also determine what automatic mode means to opt in or out of this feature when in automatic mode

diff --git a/Engine/Plugins/FX/Niagara/Shaders/Private/NiagaraSpriteVertexFactory.ush b/Engine/Plugins/FX/Niagara/Shaders/Private/NiagaraSpriteVertexFactory.ush
--- a/Engine/Plugins/FX/Niagara/Shaders/Private/NiagaraSpriteVertexFactory.ush
+++ b/Engine/Plugins/FX/Niagara/Shaders/Private/NiagaraSpriteVertexFactory.ush
@@ -858,11 +858,28 @@ FVertexFactoryIntermediates GetVertexFactoryIntermediates(FVertexFactoryInput In
 
 	UVForTexturing.xy = float2(0.5f, 0.5f) + (UVForTexturing.xy - float2(0.5f, 0.5f) ) *  UVScale.xy; 
 	const float2 Size = abs(Intermediates.SizeRotSubImage.xy);
-	
+
+	// When enabled try to main a pixel sized sprite and optionally fade particles based on coverage
+	// In reality this is more like 2x2 coverage to ensure we always hit 1 pixel, the alternative would be to snap
+	// both the size and the location to pixel centers.
+	// This also does not work correctly with sprites that are not camera facing, future work may be to
+	// calculate the size based on Tangents and the camera axis.
+	float PixelSizeRatio = 1.0f;
+	if ( NiagaraSpriteVF.PixelCoverageEnabled != 0 )
+	{
+		const float Depth = dot(ParticlePosition - LWCHackToFloat(ResolvedView.WorldCameraOrigin), ResolvedView.ViewForward);
+		const float HalfSize = max(Size.x, Size.y) * 0.5f;
+		const float PixelSize = max(View.ViewSizeAndInvSize.z * GetTanHalfFieldOfView().x * Depth * 2.0f, HalfSize);
+		PixelSizeRatio = PixelSize / HalfSize;
+
+		float Coverage = saturate(HalfSize * 2.0f / PixelSize);
+		Coverage = Coverage * Coverage;
+		Intermediates.Color = (Intermediates.Color * Coverage * NiagaraSpriteVF.PixelCoverageColorBlend) + (Intermediates.Color * (1 - NiagaraSpriteVF.PixelCoverageColorBlend));
+	}
 
 	// Vertex position
 	const float2x3 Tangents = float2x3(Intermediates.TangentRight, Intermediates.TangentUp);
-	const float3 VertexOffset = CameraOffset + mul(Size * (UVForPosition - PivotOffset), Tangents);
+	const float3 VertexOffset = CameraOffset + mul(Size * PixelSizeRatio * (UVForPosition - PivotOffset), Tangents);
 
 	Intermediates.VertexWorldPosition = ParticlePosition + VertexOffset;
 		
diff --git a/Engine/Plugins/FX/Niagara/Source/Niagara/Private/NiagaraRendererSprites.cpp b/Engine/Plugins/FX/Niagara/Source/Niagara/Private/NiagaraRendererSprites.cpp
--- a/Engine/Plugins/FX/Niagara/Source/Niagara/Private/NiagaraRendererSprites.cpp
+++ b/Engine/Plugins/FX/Niagara/Source/Niagara/Private/NiagaraRendererSprites.cpp
@@ -8,6 +8,7 @@
 #include "NiagaraComponent.h"
 #include "NiagaraCutoutVertexBuffer.h"
 #include "NiagaraGpuComputeDispatchInterface.h"
+#include "NiagaraSettings.h"
 #include "RayTracingDefinitions.h"
 #include "RayTracingDynamicGeometryCollection.h"
 #include "RayTracingInstance.h"
@@ -104,6 +105,16 @@ FNiagaraRendererSprites::FNiagaraRendererSprites(ERHIFeatureLevel::Type FeatureL
 	RendererVisibility = Properties->RendererVisibility;
 	bAccurateMotionVectors = Properties->NeedsPreciseMotionVectors();
 
+	PixelCoverageMode = Properties->PixelCoverageMode;
+	if (PixelCoverageMode == ENiagaraRendererPixelCoverageMode::Automatic)
+	{
+		if ( GetDefault<UNiagaraSettings>()->DefaultPixelCoverageMode != ENiagaraDefaultRendererPixelCoverageMode::Enabled )
+		{
+			PixelCoverageMode = ENiagaraRendererPixelCoverageMode::Disabled;
+		}
+	}
+	PixelCoverageBlend = FMath::Clamp(Properties->PixelCoverageBlend, 0.0f, 1.0f);
+
 	bEnableDistanceCulling = Properties->bEnableCameraDistanceCulling;
 	if (Properties->bEnableCameraDistanceCulling)
 	{
@@ -222,7 +233,8 @@ void FNiagaraRendererSprites::PrepareParticleSpriteRenderData(FParticleSpriteRen
 	check(MaterialRenderProxy);
 
 	// Do we have anything to render?
-	EBlendMode BlendMode = MaterialRenderProxy->GetIncompleteMaterialWithFallback(FeatureLevel).GetBlendMode();
+	const EBlendMode BlendMode = MaterialRenderProxy->GetIncompleteMaterialWithFallback(FeatureLevel).GetBlendMode();
+	ParticleSpriteRenderData.BlendMode = BlendMode;
 	ParticleSpriteRenderData.bHasTranslucentMaterials = IsTranslucentBlendMode(BlendMode);
 	ParticleSpriteRenderData.SourceParticleData = ParticleSpriteRenderData.DynamicDataSprites->GetParticleDataToRender(ParticleSpriteRenderData.bHasTranslucentMaterials && bGpuLowLatencyTranslucency && !SceneProxy->CastsVolumetricTranslucentShadow());
 	if ( !ParticleSpriteRenderData.SourceParticleData || (SourceMode == ENiagaraRendererSourceDataMode::Particles && ParticleSpriteRenderData.SourceParticleData->GetNumInstances() == 0) )
@@ -550,6 +562,49 @@ FNiagaraSpriteUniformBufferRef FNiagaraRendererSprites::CreateViewUniformBuffer(
 	PerViewUniformParameters.PrevCameraOffsetDataOffset = INDEX_NONE;
 	PerViewUniformParameters.PrevPivotOffsetDataOffset = INDEX_NONE;
 
+	// Determine pixel coverage settings
+	PerViewUniformParameters.PixelCoverageEnabled = PixelCoverageMode != ENiagaraRendererPixelCoverageMode::Disabled;
+	PerViewUniformParameters.PixelCoverageColorBlend = FVector4f::Zero();
+	if (PixelCoverageMode != ENiagaraRendererPixelCoverageMode::Disabled)
+	{
+		if ( PixelCoverageMode == ENiagaraRendererPixelCoverageMode::Automatic )
+		{
+			PerViewUniformParameters.PixelCoverageEnabled = ParticleSpriteRenderData.bHasTranslucentMaterials;
+			if (PerViewUniformParameters.PixelCoverageEnabled)
+			{
+				switch (ParticleSpriteRenderData.BlendMode)
+				{
+					case BLEND_Translucent:
+						ParticleSpriteRenderData.bHasTranslucentMaterials = true;
+						PerViewUniformParameters.PixelCoverageColorBlend = FVector4f(PixelCoverageBlend, PixelCoverageBlend, PixelCoverageBlend, 0.0f);
+						break;
+					case BLEND_Additive:
+						ParticleSpriteRenderData.bHasTranslucentMaterials = true;
+						PerViewUniformParameters.PixelCoverageColorBlend = FVector4f(PixelCoverageBlend, PixelCoverageBlend, PixelCoverageBlend, PixelCoverageBlend);
+						break;
+					//-TODO: Support these blend modes
+					//BLEND_Modulate
+					//BLEND_AlphaComposite
+					//BLEND_AlphaHoldout
+					default:
+						ParticleSpriteRenderData.bHasTranslucentMaterials = false;
+						break;
+				}
+			}
+		}
+		else
+		{
+			PerViewUniformParameters.PixelCoverageEnabled = true;
+			switch (PixelCoverageMode)
+			{
+				case ENiagaraRendererPixelCoverageMode::Enabled_RGBA:	PerViewUniformParameters.PixelCoverageColorBlend = FVector4f(PixelCoverageBlend, PixelCoverageBlend, PixelCoverageBlend, PixelCoverageBlend); break;
+				case ENiagaraRendererPixelCoverageMode::Enabled_RGB:	PerViewUniformParameters.PixelCoverageColorBlend = FVector4f(PixelCoverageBlend, PixelCoverageBlend, PixelCoverageBlend, 0.0f); break;
+				case ENiagaraRendererPixelCoverageMode::Enabled_A:		PerViewUniformParameters.PixelCoverageColorBlend = FVector4f(0.0f, 0.0f, 0.0f, PixelCoverageBlend); break;
+				default: break;
+			}
+		}
+	}
+
 	if (SourceMode == ENiagaraRendererSourceDataMode::Particles)
 	{
 		TConstArrayView<FNiagaraRendererVariableInfo> VFVariables = ParticleSpriteRenderData.RendererLayout->GetVFVariables_RenderThread();
diff --git a/Engine/Plugins/FX/Niagara/Source/Niagara/Public/NiagaraRendererSprites.h b/Engine/Plugins/FX/Niagara/Source/Niagara/Public/NiagaraRendererSprites.h
--- a/Engine/Plugins/FX/Niagara/Source/Niagara/Public/NiagaraRendererSprites.h
+++ b/Engine/Plugins/FX/Niagara/Source/Niagara/Public/NiagaraRendererSprites.h
@@ -6,6 +6,7 @@ NiagaraRendererSprites.h: Renderer for rendering Niagara particles as sprites.
 
 #pragma once
 
+#include "Engine/EngineTypes.h"
 #include "NiagaraRenderer.h"
 #include "NiagaraSpriteRendererProperties.h"
 #include "NiagaraSpriteVertexFactory.h"
@@ -41,6 +42,7 @@ private:
 		const FNiagaraDynamicDataSprites*	DynamicDataSprites = nullptr;
 		class FNiagaraDataBuffer*			SourceParticleData = nullptr;
 
+		EBlendMode							BlendMode = BLEND_Opaque;
 		bool								bHasTranslucentMaterials = false;
 		bool								bSortCullOnGpu = false;
 		bool								bNeedsSort = false;
@@ -120,6 +122,9 @@ private:
 	uint32 bSetAnyBoundVars : 1;
 	uint32 bVisTagInParamStore : 1;
 
+	ENiagaraRendererPixelCoverageMode PixelCoverageMode = ENiagaraRendererPixelCoverageMode::Automatic;
+	float PixelCoverageBlend = 0.0f;
+
 	float MinFacingCameraBlendDistance;
 	float MaxFacingCameraBlendDistance;
 	FVector2f DistanceCullRange;
diff --git a/Engine/Plugins/FX/Niagara/Source/Niagara/Public/NiagaraSettings.h b/Engine/Plugins/FX/Niagara/Source/Niagara/Public/NiagaraSettings.h
--- a/Engine/Plugins/FX/Niagara/Source/Niagara/Public/NiagaraSettings.h
+++ b/Engine/Plugins/FX/Niagara/Source/Niagara/Public/NiagaraSettings.h
@@ -52,6 +52,15 @@ namespace ENDISkelMesh_AdjacencyTriangleIndexFormat
 	};
 }
 
+UENUM()
+enum class ENiagaraDefaultRendererPixelCoverageMode : uint8
+{
+	/** When renderer is set to automatic mode pixel coverage is enabled. */
+	Enabled,
+	/** When renderer is set to automatic mode pixel coverage is disabled. */
+	Disabled,
+};
+
 UCLASS(config = Niagara, defaultconfig, meta=(DisplayName="Niagara"))
 class NIAGARA_API UNiagaraSettings : public UDeveloperSettings
 {
@@ -104,6 +113,10 @@ class NIAGARA_API UNiagaraSettings : public UDeveloperSettings
 	UPROPERTY(config, EditAnywhere, Category = Renderer)
 	ENiagaraDefaultRendererMotionVectorSetting DefaultRendererMotionVectorSetting = ENiagaraDefaultRendererMotionVectorSetting::Precise;
 
+	/** The default setting for pixel coverage mode when automatic is set on the Niagara Renderer. */
+	UPROPERTY(config, EditAnywhere, Category = Renderer)
+	ENiagaraDefaultRendererPixelCoverageMode DefaultPixelCoverageMode = ENiagaraDefaultRendererPixelCoverageMode::Enabled;
+
 	UPROPERTY(config, EditAnywhere, Category=SkeletalMeshDI, meta = ( DisplayName = "Gpu Max Bone Influences", ToolTip = "Controls the maximum number of influences we allow the Skeletal Mesh Data Interface to use on the GPU.  Changing this setting requires restarting the editor.", ConfigRestartRequired = true))
 	TEnumAsByte<ENDISkelMesh_GpuMaxInfluences::Type> NDISkelMesh_GpuMaxInfluences;
 
diff --git a/Engine/Plugins/FX/Niagara/Source/Niagara/Public/NiagaraSpriteRendererProperties.h b/Engine/Plugins/FX/Niagara/Source/Niagara/Public/NiagaraSpriteRendererProperties.h
--- a/Engine/Plugins/FX/Niagara/Source/Niagara/Public/NiagaraSpriteRendererProperties.h
+++ b/Engine/Plugins/FX/Niagara/Source/Niagara/Public/NiagaraSpriteRendererProperties.h
@@ -38,6 +38,23 @@ enum class ENiagaraSpriteFacingMode : uint8
 	FaceCameraDistanceBlend
 };
 
+UENUM()
+enum class ENiagaraRendererPixelCoverageMode : uint8
+{
+	/** Automatically determine if we want pixel coverage enabled or disabled, based on project setting and the material used on the renderer. */
+	Automatic,
+	/** Disable pixel coverage. */
+	Disabled,
+	/** Enable pixel coverage with no color adjustment based on coverage. */
+	Enabled UMETA(DisplayName = "Enabled (No Color Adjustment)"),
+	/** Enable pixel coverage and adjust the RGBA channels according to coverage. */
+	Enabled_RGBA UMETA(DisplayName = "Enabled (RGBA)"),
+	/** Enable pixel coverage and adjust the RGB channels according to coverage. */
+	Enabled_RGB UMETA(DisplayName = "Enabled (RGB)"),
+	/** Enable pixel coverage and adjust the Alpha channel only according to coverage. */
+	Enabled_A UMETA(DisplayName = "Enabled (A)"),
+};
+
 namespace ENiagaraSpriteVFLayout
 {
 	enum Type
@@ -184,6 +201,24 @@ public:
 	UPROPERTY(EditAnywhere, AdvancedDisplay, Category = "Rendering")
 	uint32 bGpuLowLatencyTranslucency : 1;
 
+	/**
+	This setting controls what happens when a sprite becomes less than a pixel in size.
+	Disabling will apply nothing and can result in flickering issues, especially with Temporal Super Resolution.
+	Automatic will enable the appropriate settings when the material blend mode is some form of translucency, project setting must also be enabled.
+	When coverage is less than a pixel, we also calculate a percentage of coverage, and then darken or reduce opacity
+	to visually compensate.	The different enabled settings allow you to control how the coverage amount is applied to
+	your particle color.  If particle color is not connected to your material the compensation will not be applied.
+	*/
+	UPROPERTY(EditAnywhere, Category = "Sprite Rendering")
+	ENiagaraRendererPixelCoverageMode PixelCoverageMode = ENiagaraRendererPixelCoverageMode::Automatic;
+
+	/**
+	When pixel coverage is enabled this allows you to control the blend of the pixel coverage color adjustment.
+	i.e. 1.0 = full, 0.5 = 50/50 blend, 0.0 = none
+	*/
+	UPROPERTY(EditAnywhere, Category = "Sprite Rendering", meta = (EditCondition = "PixelCoverageMode != ENiagaraRendererPixelCoverageMode::Disabled", UIMin=0.0f, UIMax=1.0f))
+	float PixelCoverageBlend = 1.0f;
+
 	/** When FacingMode is FacingCameraDistanceBlend, the distance at which the sprite is fully facing the camera plane. */
 	UPROPERTY(EditAnywhere, Category = "Sprite Rendering", meta = (UIMin = "0"))
 	float MinFacingCameraBlendDistance;
diff --git a/Engine/Plugins/FX/Niagara/Source/NiagaraVertexFactories/Public/NiagaraSpriteVertexFactory.h b/Engine/Plugins/FX/Niagara/Source/NiagaraVertexFactories/Public/NiagaraSpriteVertexFactory.h
--- a/Engine/Plugins/FX/Niagara/Source/NiagaraVertexFactories/Public/NiagaraSpriteVertexFactory.h
+++ b/Engine/Plugins/FX/Niagara/Source/NiagaraVertexFactories/Public/NiagaraSpriteVertexFactory.h
@@ -84,6 +84,8 @@ BEGIN_GLOBAL_SHADER_PARAMETER_STRUCT( FNiagaraSpriteUniformParameters, NIAGARAVE
 	SHADER_PARAMETER(FVector4f, DefaultDynamicMaterialParameter1)
 	SHADER_PARAMETER(FVector4f, DefaultDynamicMaterialParameter2)
 	SHADER_PARAMETER(FVector4f, DefaultDynamicMaterialParameter3)
+	SHADER_PARAMETER(int, PixelCoverageEnabled)
+	SHADER_PARAMETER(FVector4f, PixelCoverageColorBlend)
 END_GLOBAL_SHADER_PARAMETER_STRUCT()
 
 typedef TUniformBufferRef<FNiagaraSpriteUniformParameters> FNiagaraSpriteUniformBufferRef;
