- Enable Niagara rendering on SM6

#rnx

diff --git a/Engine/Plugins/FX/Niagara/Source/Niagara/Public/NiagaraCommon.h b/Engine/Plugins/FX/Niagara/Source/Niagara/Public/NiagaraCommon.h
--- a/Engine/Plugins/FX/Niagara/Source/Niagara/Public/NiagaraCommon.h
+++ b/Engine/Plugins/FX/Niagara/Source/Niagara/Public/NiagaraCommon.h
@@ -1187,7 +1187,7 @@ namespace FNiagaraUtilities
 
 	inline bool SupportsNiagaraRendering(ERHIFeatureLevel::Type FeatureLevel)
 	{
-		return FeatureLevel == ERHIFeatureLevel::SM5 || FeatureLevel == ERHIFeatureLevel::ES3_1;
+		return FeatureLevel >= ERHIFeatureLevel::SM5 || FeatureLevel == ERHIFeatureLevel::ES3_1;
 	}
 
 	inline bool SupportsNiagaraRendering(EShaderPlatform ShaderPlatform)
