- Fix duplicate keys for localized text

#rnx

diff --git a/Engine/Plugins/FX/Niagara/Source/NiagaraEditor/Private/Customizations/NiagaraBakerSettingsDetails.cpp b/Engine/Plugins/FX/Niagara/Source/NiagaraEditor/Private/Customizations/NiagaraBakerSettingsDetails.cpp
--- a/Engine/Plugins/FX/Niagara/Source/NiagaraEditor/Private/Customizations/NiagaraBakerSettingsDetails.cpp
+++ b/Engine/Plugins/FX/Niagara/Source/NiagaraEditor/Private/Customizations/NiagaraBakerSettingsDetails.cpp
@@ -90,16 +90,16 @@ void FNiagaraBakerTextureSourceDetails::CollectAllActions(FGraphActionListBuilde
 				switch ( RendererType )
 				{
 					case FNiagaraBakerRenderer::ERenderType::SceneCapture:
-						CategoryText = LOCTEXT("BufferVis", "Scene Capture");
+						CategoryText = LOCTEXT("SceneCapture", "Scene Capture");
 						break;
 					case FNiagaraBakerRenderer::ERenderType::BufferVisualization:
 						CategoryText = LOCTEXT("BufferVis", "Buffer Visualization");
 						break;
 					case FNiagaraBakerRenderer::ERenderType::DataInterface:
-						CategoryText = LOCTEXT("BufferVis", "Emitter DataInterface");
+						CategoryText = LOCTEXT("EmitterDataInterface", "Emitter DataInterface");
 						break;
 					case FNiagaraBakerRenderer::ERenderType::Particle:
-						CategoryText = LOCTEXT("BufferVis", "Particle Attribute");
+						CategoryText = LOCTEXT("ParticleAttribute", "Particle Attribute");
 						break;
 				}
 				FText MenuText = FText::FromString(*OptionName.ToString());
