Fixed a bug where position binding's cant see position varaibles

diff --git a/Engine/Plugins/FX/Niagara/Source/Niagara/Private/NiagaraLightRendererProperties.cpp b/Engine/Plugins/FX/Niagara/Source/Niagara/Private/NiagaraLightRendererProperties.cpp
--- a/Engine/Plugins/FX/Niagara/Source/Niagara/Private/NiagaraLightRendererProperties.cpp
+++ b/Engine/Plugins/FX/Niagara/Source/Niagara/Private/NiagaraLightRendererProperties.cpp
@@ -40,6 +40,10 @@ UNiagaraLightRendererProperties::UNiagaraLightRendererProperties()
 void UNiagaraLightRendererProperties::PostLoad()
 {
 	Super::PostLoad();
+	
+#if WITH_EDITORONLY_DATA
+	ChangeToPositionBinding(PositionBinding);
+#endif
 	PostLoadBindings(ENiagaraRendererSourceDataMode::Particles);
 }
 
diff --git a/Engine/Plugins/FX/Niagara/Source/Niagara/Private/NiagaraMeshRendererProperties.cpp b/Engine/Plugins/FX/Niagara/Source/Niagara/Private/NiagaraMeshRendererProperties.cpp
--- a/Engine/Plugins/FX/Niagara/Source/Niagara/Private/NiagaraMeshRendererProperties.cpp
+++ b/Engine/Plugins/FX/Niagara/Source/Niagara/Private/NiagaraMeshRendererProperties.cpp
@@ -528,7 +528,11 @@ void UNiagaraMeshRendererProperties::PostLoad()
 		}
 	}
 
-
+#if WITH_EDITORONLY_DATA
+	ChangeToPositionBinding(PositionBinding);
+	ChangeToPositionBinding(PrevPositionBinding);
+#endif
+	
 	PostLoadBindings(SourceMode);
 	
 	// Fix up these bindings from their loaded source bindings
diff --git a/Engine/Plugins/FX/Niagara/Source/Niagara/Private/NiagaraRendererProperties.cpp b/Engine/Plugins/FX/Niagara/Source/Niagara/Private/NiagaraRendererProperties.cpp
--- a/Engine/Plugins/FX/Niagara/Source/Niagara/Private/NiagaraRendererProperties.cpp
+++ b/Engine/Plugins/FX/Niagara/Source/Niagara/Private/NiagaraRendererProperties.cpp
@@ -133,6 +133,15 @@ TArray<FNiagaraVariable> UNiagaraRendererProperties::GetBoundAttributes() const
 	return BoundAttributes;
 }
 
+void UNiagaraRendererProperties::ChangeToPositionBinding(FNiagaraVariableAttributeBinding& Binding)
+{
+	if (Binding.GetType() == FNiagaraTypeDefinition::GetVec3Def())
+	{
+		FNiagaraVariable NewVarType(FNiagaraTypeDefinition::GetPositionDef(), Binding.GetParamMapBindableVariable().GetName());
+		Binding = FNiagaraConstants::GetAttributeDefaultBinding(NewVarType);
+	}
+}
+
 FNiagaraVariable UNiagaraRendererProperties::GetBoundAttribute(const FNiagaraVariableAttributeBinding* Binding) const
 {
 	if (Binding->GetParamMapBindableVariable().IsValid())
diff --git a/Engine/Plugins/FX/Niagara/Source/Niagara/Private/NiagaraRibbonRendererProperties.cpp b/Engine/Plugins/FX/Niagara/Source/Niagara/Private/NiagaraRibbonRendererProperties.cpp
--- a/Engine/Plugins/FX/Niagara/Source/Niagara/Private/NiagaraRibbonRendererProperties.cpp
+++ b/Engine/Plugins/FX/Niagara/Source/Niagara/Private/NiagaraRibbonRendererProperties.cpp
@@ -132,8 +132,9 @@ void UNiagaraRibbonRendererProperties::PostLoad()
 		UpgradeUVSettings(UV1Settings, UV1TilingDistance_DEPRECATED, UV1Offset_DEPRECATED, UV1Scale_DEPRECATED);
 	}
 
+	ChangeToPositionBinding(PositionBinding);
 #endif
-
+	
 	PostLoadBindings(ENiagaraRendererSourceDataMode::Particles);
 
 	if ( Material )
diff --git a/Engine/Plugins/FX/Niagara/Source/Niagara/Private/NiagaraSpriteRendererProperties.cpp b/Engine/Plugins/FX/Niagara/Source/Niagara/Private/NiagaraSpriteRendererProperties.cpp
--- a/Engine/Plugins/FX/Niagara/Source/Niagara/Private/NiagaraSpriteRendererProperties.cpp
+++ b/Engine/Plugins/FX/Niagara/Source/Niagara/Private/NiagaraSpriteRendererProperties.cpp
@@ -143,6 +143,8 @@ void UNiagaraSpriteRendererProperties::PostLoad()
 		}
 		CacheDerivedData();
 	}
+	ChangeToPositionBinding(PositionBinding);
+	ChangeToPositionBinding(PrevPositionBinding);
 	PostLoadBindings(SourceMode);
 
 	// Fix up these bindings from their loaded source bindings
diff --git a/Engine/Plugins/FX/Niagara/Source/Niagara/Public/NiagaraRendererProperties.h b/Engine/Plugins/FX/Niagara/Source/Niagara/Public/NiagaraRendererProperties.h
--- a/Engine/Plugins/FX/Niagara/Source/Niagara/Public/NiagaraRendererProperties.h
+++ b/Engine/Plugins/FX/Niagara/Source/Niagara/Public/NiagaraRendererProperties.h
@@ -292,6 +292,9 @@ protected:
 #if WITH_EDITORONLY_DATA
 	/** returns the variable associated with the supplied binding if it should be bound given the current settings of the RendererProperties. */
 	virtual FNiagaraVariable GetBoundAttribute(const FNiagaraVariableAttributeBinding* Binding) const;
+
+	/** utility function that can be used to fix up old vec3 bindings into position bindings. */
+	static void ChangeToPositionBinding(FNiagaraVariableAttributeBinding& Binding);
 #endif
 
 #if WITH_EDITOR
