Fixed duplicate loctext keys

diff --git a/Engine/Plugins/FX/Niagara/Source/Niagara/Private/NiagaraDataInterfaceCollisionQuery.cpp b/Engine/Plugins/FX/Niagara/Source/Niagara/Private/NiagaraDataInterfaceCollisionQuery.cpp
--- a/Engine/Plugins/FX/Niagara/Source/Niagara/Private/NiagaraDataInterfaceCollisionQuery.cpp
+++ b/Engine/Plugins/FX/Niagara/Source/Niagara/Private/NiagaraDataInterfaceCollisionQuery.cpp
@@ -338,13 +338,14 @@ void UNiagaraDataInterfaceCollisionQuery::GetFunctions(TArray<FNiagaraFunctionSi
 		OutFunctions.Add(SigMeshField);
 	}
 
+	
+	const FText TraceStartWorldDescription = LOCTEXT("TraceStartWorldDescription", "Ray starting point in world space");
+	const FText TraceEndWorldDescription = LOCTEXT("TraceEndWorldDescription", "Ray end point in world space");
 	{
+		const FText AsyncTraceChannelDescription = LOCTEXT("TraceChannelDescription", "Currently unused, will represent the trace channels for which geometry the trace should test against");
 		const FText QueryIDDescription = LOCTEXT("QueryIDDescription", "Unique (for this frame) index of the query being enqueued (used in subsequent frames to retrieve results).  Must be less than MaxRayTraceCount");
-		const FText TraceStartWorldDescription = LOCTEXT("TraceStartWorldDescription", "Ray starting point in world space");
-		const FText TraceEndWorldDescription = LOCTEXT("TraceEndWorldDescription", "Ray end point in world space");
-		const FText TraceChannelDescription = LOCTEXT("TraceChannelDescription", "Currently unused, will represent the trace channels for which geometry the trace should test against");
 		const FText IsQueryValidDescription = LOCTEXT("IsQueryValidDescription", "Returns true if the query was enqueued");
-
+		
 		FNiagaraFunctionSignature& IssueRayTrace = OutFunctions.AddDefaulted_GetRef();
 		IssueRayTrace.Name = NDICollisionQueryLocal::IssueAsyncRayTraceName;
 		IssueRayTrace.bRequiresExecPin = true;
@@ -352,45 +353,33 @@ void UNiagaraDataInterfaceCollisionQuery::GetFunctions(TArray<FNiagaraFunctionSi
 		IssueRayTrace.bSupportsCPU = false;
 #if WITH_EDITORONLY_DATA
 		IssueRayTrace.FunctionVersion = FNiagaraCollisionDIFunctionVersion::LatestVersion;
-		IssueRayTrace.Description = LOCTEXT("IssueAsyncRayTraceDescription", "Enqueues a GPU raytrace with the result being available the following frame");
+		IssueRayTrace.Description = LOCTEXT("IssueAsync_RayTraceDescription", "Enqueues a GPU raytrace with the result being available the following frame");
 #endif
 		IssueRayTrace.AddInput(FNiagaraVariable(FNiagaraTypeDefinition(GetClass()), TEXT("CollisionQuery")));
 		IssueRayTrace.AddInput(FNiagaraVariable(FNiagaraTypeDefinition::GetIntDef(), TEXT("QueryID")), QueryIDDescription);
 		IssueRayTrace.AddInput(FNiagaraVariable(FNiagaraTypeDefinition::GetPositionDef(), TEXT("TraceStartWorld")), TraceStartWorldDescription);
 		IssueRayTrace.AddInput(FNiagaraVariable(FNiagaraTypeDefinition::GetPositionDef(), TEXT("TraceEndWorld")), TraceEndWorldDescription);
-		IssueRayTrace.AddInput(FNiagaraVariable(FNiagaraTypeDefinition::GetIntDef(), TEXT("TraceChannel")), TraceChannelDescription);
+		IssueRayTrace.AddInput(FNiagaraVariable(FNiagaraTypeDefinition::GetIntDef(), TEXT("TraceChannel")), AsyncTraceChannelDescription);
 		IssueRayTrace.AddOutput(FNiagaraVariable(FNiagaraTypeDefinition::GetBoolDef(), TEXT("IsQueryValid")), IsQueryValidDescription);
-	}
-
-	{
-		const FText TraceStartWorldDescription = LOCTEXT("TraceStartWorldDescription", "Ray starting point in world space");
-		const FText TraceEndWorldDescription = LOCTEXT("TraceEndWorldDescription", "Ray end point in world space");
-		const FText TraceChannelDescription = LOCTEXT("TraceChannelDescription", "Currently unused, will represent the trace channels for which geometry the trace should test against");
-		const FText QueryIDDescription = LOCTEXT("QueryIDDescription", "Unique (for this frame) index of the query being enqueued (used in subsequent frames to retrieve results).");
-		const FText IsQueryValidDescription = LOCTEXT("IsQueryValidDescription", "Returns true if the query was enqueued");
-
-		FNiagaraFunctionSignature& IssueRayTrace = OutFunctions.AddDefaulted_GetRef();
-		IssueRayTrace.Name = NDICollisionQueryLocal::CreateAsyncRayTraceName;
-		IssueRayTrace.bRequiresExecPin = true;
-		IssueRayTrace.bMemberFunction = true;
-		IssueRayTrace.bSupportsCPU = false;
+		
+		FNiagaraFunctionSignature& CreateRayTrace = OutFunctions.AddDefaulted_GetRef();
+		CreateRayTrace.Name = NDICollisionQueryLocal::CreateAsyncRayTraceName;
+		CreateRayTrace.bRequiresExecPin = true;
+		CreateRayTrace.bMemberFunction = true;
+		CreateRayTrace.bSupportsCPU = false;
 #if WITH_EDITORONLY_DATA
-		IssueRayTrace.FunctionVersion = FNiagaraCollisionDIFunctionVersion::LatestVersion;
-		IssueRayTrace.Description = LOCTEXT("IssueAsyncRayTraceDescription", "Creates a GPU raytrace with the result being available the following frame (index is returned)");
+		CreateRayTrace.FunctionVersion = FNiagaraCollisionDIFunctionVersion::LatestVersion;
+		CreateRayTrace.Description = LOCTEXT("CreateAsync_RayTraceDescription", "Creates a GPU raytrace with the result being available the following frame (index is returned)");
 #endif
-		IssueRayTrace.AddInput(FNiagaraVariable(FNiagaraTypeDefinition(GetClass()), TEXT("CollisionQuery")));
-		IssueRayTrace.AddInput(FNiagaraVariable(FNiagaraTypeDefinition::GetPositionDef(), TEXT("TraceStartWorld")), TraceStartWorldDescription);
-		IssueRayTrace.AddInput(FNiagaraVariable(FNiagaraTypeDefinition::GetPositionDef(), TEXT("TraceEndWorld")), TraceEndWorldDescription);
-		IssueRayTrace.AddInput(FNiagaraVariable(FNiagaraTypeDefinition::GetIntDef(), TEXT("TraceChannel")), TraceChannelDescription);
-		IssueRayTrace.AddOutput(FNiagaraVariable(FNiagaraTypeDefinition::GetIntDef(), TEXT("QueryID")), QueryIDDescription);
-		IssueRayTrace.AddOutput(FNiagaraVariable(FNiagaraTypeDefinition::GetBoolDef(), TEXT("IsQueryValid")), IsQueryValidDescription);
+		CreateRayTrace.AddInput(FNiagaraVariable(FNiagaraTypeDefinition(GetClass()), TEXT("CollisionQuery")));
+		CreateRayTrace.AddInput(FNiagaraVariable(FNiagaraTypeDefinition::GetPositionDef(), TEXT("TraceStartWorld")), TraceStartWorldDescription);
+		CreateRayTrace.AddInput(FNiagaraVariable(FNiagaraTypeDefinition::GetPositionDef(), TEXT("TraceEndWorld")), TraceEndWorldDescription);
+		CreateRayTrace.AddInput(FNiagaraVariable(FNiagaraTypeDefinition::GetIntDef(), TEXT("TraceChannel")), AsyncTraceChannelDescription);
+		CreateRayTrace.AddOutput(FNiagaraVariable(FNiagaraTypeDefinition::GetIntDef(), TEXT("QueryID")), QueryIDDescription);
+		CreateRayTrace.AddOutput(FNiagaraVariable(FNiagaraTypeDefinition::GetBoolDef(), TEXT("IsQueryValid")), IsQueryValidDescription);
 	}
 
 	{
-		const FText TraceCountDescription = LOCTEXT("QueryIDDescription", "Number of async raytrace requests to be reserved");
-		const FText FirstQueryIDValidDescription = LOCTEXT("TraceChannelDescription", "The first index in the block reserved through this call");
-		const FText IsQueryValidDescription = LOCTEXT("IsQueryValidDescription", "Returns true if the requested indices were reserved");
-
 		FNiagaraFunctionSignature& IssueRayTrace = OutFunctions.AddDefaulted_GetRef();
 		IssueRayTrace.Name = NDICollisionQueryLocal::ReserveAsyncRayTraceName;
 		IssueRayTrace.bRequiresExecPin = true;
@@ -398,48 +387,32 @@ void UNiagaraDataInterfaceCollisionQuery::GetFunctions(TArray<FNiagaraFunctionSi
 		IssueRayTrace.bSupportsCPU = false;
 #if WITH_EDITORONLY_DATA
 		IssueRayTrace.FunctionVersion = FNiagaraCollisionDIFunctionVersion::LatestVersion;
-		IssueRayTrace.Description = LOCTEXT("IssueAsyncRayTraceDescription", "Reserves a number of ray trace request slots");
+		IssueRayTrace.Description = LOCTEXT("ReserveAsync_RayTraceDescription", "Reserves a number of ray trace request slots");
 #endif
 		IssueRayTrace.AddInput(FNiagaraVariable(FNiagaraTypeDefinition(GetClass()), TEXT("CollisionQuery")));
-		IssueRayTrace.AddInput(FNiagaraVariable(FNiagaraTypeDefinition::GetIntDef(), TEXT("TraceCount")), TraceCountDescription);
-		IssueRayTrace.AddOutput(FNiagaraVariable(FNiagaraTypeDefinition::GetIntDef(), TEXT("FirstQueryID")), FirstQueryIDValidDescription);
-		IssueRayTrace.AddOutput(FNiagaraVariable(FNiagaraTypeDefinition::GetBoolDef(), TEXT("IsQueryValid")), IsQueryValidDescription);
+		IssueRayTrace.AddInput(FNiagaraVariable(FNiagaraTypeDefinition::GetIntDef(), TEXT("TraceCount")), LOCTEXT("ReserveAsync_QueryIDDescription", "Number of async raytrace requests to be reserved"));
+		IssueRayTrace.AddOutput(FNiagaraVariable(FNiagaraTypeDefinition::GetIntDef(), TEXT("FirstQueryID")), LOCTEXT("ReserveAsync_TraceChannelDescription", "The first index in the block reserved through this call"));
+		IssueRayTrace.AddOutput(FNiagaraVariable(FNiagaraTypeDefinition::GetBoolDef(), TEXT("IsQueryValid")), LOCTEXT("ReserveAsync_IsQueryValidDescription", "Returns true if the requested indices were reserved"));
 	}
 
 	{
-		const FText PreviousFrameQueryIDDescription = LOCTEXT("PreviousFrameQueryIDDescription", "The index of the results being retrieved");
-		const FText CollisionValidDescription = LOCTEXT("CollisionValidDescription", "Returns true if the a Hit was encountered");
-		const FText CollisionDistanceDescription = LOCTEXT("CollisionDistanceDescription", "The distance in world space from the ray starting point to the intersection");
-		const FText CollisionPosWorldDescription = LOCTEXT("CollisionPosWorldDescription", "The point in world space where the intersection occured");
-		const FText CollisionNormalDescription = LOCTEXT("CollisionNormalDescription", "The surface normal of the world geometry at the point of intersection");
-
 		FNiagaraFunctionSignature& ReadRayTrace = OutFunctions.AddDefaulted_GetRef();
 		ReadRayTrace.Name = NDICollisionQueryLocal::ReadAsyncRayTraceName;
 		ReadRayTrace.bMemberFunction = true;
 		ReadRayTrace.bSupportsCPU = false;
 #if WITH_EDITORONLY_DATA
 		ReadRayTrace.FunctionVersion = FNiagaraCollisionDIFunctionVersion::LatestVersion;
-		ReadRayTrace.Description = LOCTEXT("ReadAsyncRayTraceDescription", "Reads the results of a previously enqueued GPU ray trace");
+		ReadRayTrace.Description = LOCTEXT("ReadAsync_RayTraceDescription", "Reads the results of a previously enqueued GPU ray trace");
 #endif
 		ReadRayTrace.AddInput(FNiagaraVariable(FNiagaraTypeDefinition(GetClass()), TEXT("CollisionQuery")));
-		ReadRayTrace.AddInput(FNiagaraVariable(FNiagaraTypeDefinition::GetIntDef(), TEXT("PreviousFrameQueryID")), PreviousFrameQueryIDDescription);
-		ReadRayTrace.AddOutput(FNiagaraVariable(FNiagaraTypeDefinition::GetBoolDef(), TEXT("CollisionValid")), CollisionValidDescription);
-		ReadRayTrace.AddOutput(FNiagaraVariable(FNiagaraTypeDefinition::GetFloatDef(), TEXT("CollisionDistance")), CollisionDistanceDescription);
-		ReadRayTrace.AddOutput(FNiagaraVariable(FNiagaraTypeDefinition::GetPositionDef(), TEXT("CollisionPosWorld")), CollisionPosWorldDescription);
-		ReadRayTrace.AddOutput(FNiagaraVariable(FNiagaraTypeDefinition::GetVec3Def(), TEXT("CollisionNormal")), CollisionNormalDescription);
+		ReadRayTrace.AddInput(FNiagaraVariable(FNiagaraTypeDefinition::GetIntDef(), TEXT("PreviousFrameQueryID")), LOCTEXT("ReadAsync_PreviousFrameQueryIDDescription", "The index of the results being retrieved"));
+		ReadRayTrace.AddOutput(FNiagaraVariable(FNiagaraTypeDefinition::GetBoolDef(), TEXT("CollisionValid")), LOCTEXT("ReadAsync_CollisionValidDescription", "Returns true if a Hit was encountered"));
+		ReadRayTrace.AddOutput(FNiagaraVariable(FNiagaraTypeDefinition::GetFloatDef(), TEXT("CollisionDistance")), LOCTEXT("ReadAsync_CollisionDistanceDescription", "The distance in world space from the ray starting point to the intersection"));
+		ReadRayTrace.AddOutput(FNiagaraVariable(FNiagaraTypeDefinition::GetPositionDef(), TEXT("CollisionPosWorld")), LOCTEXT("ReadAsync_CollisionPosWorldDescription", "The point in world space where the intersection occured"));
+		ReadRayTrace.AddOutput(FNiagaraVariable(FNiagaraTypeDefinition::GetVec3Def(), TEXT("CollisionNormal")), LOCTEXT("ReadAsync_CollisionNormalDescription", "The surface normal of the world geometry at the point of intersection"));
 	}
 
 	{
-		FNiagaraFunctionSignature SigCpuSync;
-		SigCpuSync.Name = NDICollisionQueryLocal::SyncTraceName;
-		SigCpuSync.bMemberFunction = true;
-		SigCpuSync.bSupportsGPU = false;
-#if WITH_EDITORONLY_DATA
-		SigCpuSync.FunctionVersion = FNiagaraCollisionDIFunctionVersion::LatestVersion;
-		SigCpuSync.Description = LOCTEXT("SigCpuSyncDescription", "Traces a ray against the world using a specific channel and return the first blocking hit.");
-#endif
-		const FText TraceStartWorldDescription = LOCTEXT("TraceStartWorldDescription", "The world position where the line trace should start.");
-		const FText TraceEndWorldDescription = LOCTEXT("TraceEndWorldDescription", "The world position where the line trace should end.");
 		const FText TraceChannelDescription = LOCTEXT("TraceChannelDescription", "The trace channel to collide against. Trace channels can be configured in the project settings.");
 		const FText SkipTraceDescription = LOCTEXT("SkipTraceDescription", "If true then the trace will be skipped completely.\nThis can be used as a performance optimization, as branch nodes in the graph still execute every path.");
 		const FText CollisionValidDescription = LOCTEXT("CollisionValidDescription", "Returns true if the trace was not skipped and the trace was blocked by some world geometry.");
@@ -449,6 +422,15 @@ void UNiagaraDataInterfaceCollisionQuery::GetFunctions(TArray<FNiagaraFunctionSi
 		const FText CollisionMaterialFrictionDescription = LOCTEXT("CollisionMaterialFrictionDescription", "Friction value of surface, controls how easily things can slide on this surface (0 is frictionless, higher values increase the amount of friction).");
 		const FText CollisionMaterialRestitutionDescription = LOCTEXT("CollisionMaterialRestitutionDescription", "Restitution or 'bounciness' of this surface, between 0 (no bounce) and 1 (outgoing velocity is same as incoming)");
 		const FText CollisionMaterialIndexDescription = LOCTEXT("CollisionMaterialIndexDescription", "Returns the index of the surface as defined in the ProjectSettings/Physics/PhysicalSurface section");
+		
+		FNiagaraFunctionSignature SigCpuSync;
+		SigCpuSync.Name = NDICollisionQueryLocal::SyncTraceName;
+		SigCpuSync.bMemberFunction = true;
+		SigCpuSync.bSupportsGPU = false;
+#if WITH_EDITORONLY_DATA
+		SigCpuSync.FunctionVersion = FNiagaraCollisionDIFunctionVersion::LatestVersion;
+		SigCpuSync.Description = LOCTEXT("SigCpuSyncDescription", "Traces a ray against the world using a specific channel and return the first blocking hit.");
+#endif
 	
 		SigCpuSync.Inputs.Add(FNiagaraVariable(FNiagaraTypeDefinition(GetClass()), TEXT("CollisionQuery")));
 		SigCpuSync.AddInput(FNiagaraVariable(FNiagaraTypeDefinition::GetPositionDef(), TEXT("TraceStartWorld")), TraceStartWorldDescription);
@@ -463,7 +445,7 @@ void UNiagaraDataInterfaceCollisionQuery::GetFunctions(TArray<FNiagaraFunctionSi
 		SigCpuSync.AddOutput(FNiagaraVariable(FNiagaraTypeDefinition::GetFloatDef(), TEXT("CollisionMaterialRestitution")), CollisionMaterialRestitutionDescription);
 		SigCpuSync.AddOutput(FNiagaraVariable(FNiagaraTypeDefinition::GetIntDef(), TEXT("CollisionMaterialIndex")), CollisionMaterialIndexDescription);
 		OutFunctions.Add(SigCpuSync);
-
+		
 		FNiagaraFunctionSignature SigCpuAsync;
 		SigCpuAsync.Name = NDICollisionQueryLocal::AsyncTraceName;
 		SigCpuAsync.bMemberFunction = true;
@@ -472,16 +454,14 @@ void UNiagaraDataInterfaceCollisionQuery::GetFunctions(TArray<FNiagaraFunctionSi
 		SigCpuAsync.FunctionVersion = FNiagaraCollisionDIFunctionVersion::LatestVersion;
 		SigCpuAsync.Description = LOCTEXT("SigCpuAsyncDescription", "Traces a ray against the world using a specific channel and return the first blocking hit the next frame.\nNote that this is the ASYNC version of the trace function, meaning it will not returns the result right away, but with one frame latency.");
 #endif
-		const FText PreviousFrameQueryIDDescription = LOCTEXT("PreviousFrameQueryIDDescription", "The query ID returned from the last frame's async trace call.\nRegardless if it is a valid ID or not this function call with issue a new async line trace, but it will only return results with a valid ID.");
-		const FText NextFrameQueryIDDescription = LOCTEXT("NextFrameQueryIDDescription", "The query ID to save and use as input to this function in the next frame.");
 	
 		SigCpuAsync.Inputs.Add(FNiagaraVariable(FNiagaraTypeDefinition(GetClass()), TEXT("CollisionQuery")));
-		SigCpuAsync.AddInput(FNiagaraVariable(FNiagaraTypeDefinition::GetIntDef(), TEXT("PreviousFrameQueryID")), PreviousFrameQueryIDDescription);
+		SigCpuAsync.AddInput(FNiagaraVariable(FNiagaraTypeDefinition::GetIntDef(), TEXT("PreviousFrameQueryID")), LOCTEXT("TraceAsync_PreviousFrameQueryIDDescription", "The query ID returned from the last frame's async trace call.\nRegardless if it is a valid ID or not this function call with issue a new async line trace, but it will only return results with a valid ID."));
 		SigCpuAsync.AddInput(FNiagaraVariable(FNiagaraTypeDefinition::GetPositionDef(), TEXT("TraceStartWorld")), TraceStartWorldDescription);
 		SigCpuAsync.AddInput(FNiagaraVariable(FNiagaraTypeDefinition::GetPositionDef(), TEXT("TraceEndWorld")), TraceEndWorldDescription);
 		SigCpuAsync.AddInput(FNiagaraVariable(FNiagaraTypeDefinition(TraceChannelEnum), TEXT("TraceChannel")), TraceChannelDescription);
 		SigCpuAsync.AddInput(FNiagaraVariable(FNiagaraTypeDefinition::GetBoolDef(), TEXT("SkipTrace")), SkipTraceDescription);
-		SigCpuAsync.AddOutput(FNiagaraVariable(FNiagaraTypeDefinition::GetIntDef(), TEXT("NextFrameQueryID")), NextFrameQueryIDDescription);
+		SigCpuAsync.AddOutput(FNiagaraVariable(FNiagaraTypeDefinition::GetIntDef(), TEXT("NextFrameQueryID")), LOCTEXT("TraceAsync_NextFrameQueryIDDescription", "The query ID to save and use as input to this function in the next frame."));
 		SigCpuAsync.AddOutput(FNiagaraVariable(FNiagaraTypeDefinition::GetBoolDef(), TEXT("CollisionValid")), CollisionValidDescription);
 		SigCpuAsync.AddOutput(FNiagaraVariable(FNiagaraTypeDefinition::GetBoolDef(), TEXT("IsTraceInsideMesh")), IsTraceInsideMeshDescription);
 		SigCpuAsync.AddOutput(FNiagaraVariable(FNiagaraTypeDefinition::GetPositionDef(), TEXT("CollisionPosWorld")), CollisionPosWorldDescription);
diff --git a/Engine/Plugins/FX/Niagara/Source/Niagara/Private/NiagaraMeshRendererProperties.cpp b/Engine/Plugins/FX/Niagara/Source/Niagara/Private/NiagaraMeshRendererProperties.cpp
--- a/Engine/Plugins/FX/Niagara/Source/Niagara/Private/NiagaraMeshRendererProperties.cpp
+++ b/Engine/Plugins/FX/Niagara/Source/Niagara/Private/NiagaraMeshRendererProperties.cpp
@@ -1015,7 +1015,7 @@ void UNiagaraMeshRendererProperties::RebuildMeshList()
 
 	if (bAnyError)
 	{
-		ShowFlipbookWarningToast(LOCTEXT("FlipbookSuffixWarningToastMessage", "Failed to load one or more meshes for Mesh Flipbook. See the Output Log for details."));
+		ShowFlipbookWarningToast(LOCTEXT("FlipbookWarningToastMessage", "Failed to load one or more meshes for Mesh Flipbook. See the Output Log for details."));
 	}
 }
 
diff --git a/Engine/Plugins/FX/Niagara/Source/NiagaraEditor/Private/Customizations/NiagaraOutlinerCustomization.cpp b/Engine/Plugins/FX/Niagara/Source/NiagaraEditor/Private/Customizations/NiagaraOutlinerCustomization.cpp
--- a/Engine/Plugins/FX/Niagara/Source/NiagaraEditor/Private/Customizations/NiagaraOutlinerCustomization.cpp
+++ b/Engine/Plugins/FX/Niagara/Source/NiagaraEditor/Private/Customizations/NiagaraOutlinerCustomization.cpp
@@ -1093,7 +1093,7 @@ TSharedRef<SWidget> FNiagaraOutlinerTreeSystemItem::GetHeaderWidget()
 				[
 					SNew(SNiagaraOutlinerTreeItemHeaderDataWidget<FNiagaraOutlinerTimingData>, Outliner->ViewSettings)
 					//.LabelText(LOCTEXT("SystemMaxPerFrameTime", "Frame Max"))
-					.ToolTipText(LOCTEXT("SystemMaxPerFrameTimeToolTip", "Max frame time for all instances of this System."))
+					.ToolTipText(LOCTEXT("AllInstancesMaxFrameTimeToolTip", "Max frame time for all instances of this System."))
 					.Data(Data->MaxPerFrameTime)
 					.MinDesiredWidth(50)
 				];
diff --git a/Engine/Plugins/FX/Niagara/Source/NiagaraEditor/Private/NiagaraEditorUtilities.cpp b/Engine/Plugins/FX/Niagara/Source/NiagaraEditor/Private/NiagaraEditorUtilities.cpp
--- a/Engine/Plugins/FX/Niagara/Source/NiagaraEditor/Private/NiagaraEditorUtilities.cpp
+++ b/Engine/Plugins/FX/Niagara/Source/NiagaraEditor/Private/NiagaraEditorUtilities.cpp
@@ -2097,7 +2097,7 @@ bool FNiagaraEditorUtilities::VerifyNameChangeForInputOrOutputNode(const UNiagar
 {
 	if (NewNameString.Len() >= FNiagaraConstants::MaxParameterLength)
 	{
-		OutErrorMessage = FText::FormatOrdered(LOCTEXT("EmptyNameError", "Name cannot exceed {0} characters."), FNiagaraConstants::MaxParameterLength);
+		OutErrorMessage = FText::FormatOrdered(LOCTEXT("InputOutputNodeNameLengthExceeded", "Name cannot exceed {0} characters."), FNiagaraConstants::MaxParameterLength);
 		return false;
 	}
 
@@ -2105,7 +2105,7 @@ bool FNiagaraEditorUtilities::VerifyNameChangeForInputOrOutputNode(const UNiagar
 
 	if (NewName == NAME_None)
 	{
-		OutErrorMessage = LOCTEXT("EmptyNameError", "Name can not be empty.");
+		OutErrorMessage = LOCTEXT("InputOutputNodeNameEmptyNameError", "Name can not be empty.");
 		return false;
 	}
 
diff --git a/Engine/Plugins/FX/Niagara/Source/NiagaraEditor/Private/NiagaraHlslTranslator.cpp b/Engine/Plugins/FX/Niagara/Source/NiagaraEditor/Private/NiagaraHlslTranslator.cpp
--- a/Engine/Plugins/FX/Niagara/Source/NiagaraEditor/Private/NiagaraHlslTranslator.cpp
+++ b/Engine/Plugins/FX/Niagara/Source/NiagaraEditor/Private/NiagaraHlslTranslator.cpp
@@ -6543,7 +6543,7 @@ void FHlslNiagaraTranslator::WriteDataSet(const FNiagaraDataSetID DataSet, const
 	if (DataSetWriteForInput)
 	{
 		//TODO: improve error report.
-		Error(LOCTEXT("WritingToSameDataSetError", "Writing to the same dataset with the same condition/index."), NULL, NULL);
+		Error(LOCTEXT("WritingToSameDataSetForInputError", "Writing to the same dataset with the same condition/index."), NULL, NULL);
 		return;
 	}
 
@@ -8628,7 +8628,7 @@ void FHlslNiagaraTranslator::If(UNiagaraNodeIf* IfNode, TArray<FNiagaraVariable>
 		FNiagaraTypeDefinition Type = Schema->PinToTypeDefinition(IfNode->GetInputPin(PinIdx++), ENiagaraStructConversion::Simulation);
 		if (!AddStructToDefinitionSet(Type))
 		{
-			FText OutErrorMessage = FText::Format(LOCTEXT("UnknownNumeric", "Variable in If node uses invalid type. Var: {0} Type: {1}"),
+			FText OutErrorMessage = FText::Format(LOCTEXT("If_UnknownNumeric", "Variable in If node uses invalid type. Var: {0} Type: {1}"),
 				FText::FromName(Var.GetName()), Type.GetNameText());
 
 			Error(OutErrorMessage, IfNode, nullptr);
@@ -8670,7 +8670,7 @@ void FHlslNiagaraTranslator::Select(UNiagaraNodeSelect* SelectNode, int32 Select
 	{
 		if (!AddStructToDefinitionSet(Variable.GetType()))
 		{
-			FText OutErrorMessage = FText::Format(LOCTEXT("UnknownNumeric", "Output type in Select node uses invalid type. Type: {0}"),
+			FText OutErrorMessage = FText::Format(LOCTEXT("Select_UnknownNumeric", "Output type in Select node uses invalid type. Type: {0}"),
 				Variable.GetType().GetNameText());
 
 			Error(OutErrorMessage, SelectNode, SelectNode->GetOutputPin(Variable));
diff --git a/Engine/Plugins/FX/Niagara/Source/NiagaraEditor/Private/NiagaraNodeParameterMapFor.cpp b/Engine/Plugins/FX/Niagara/Source/NiagaraEditor/Private/NiagaraNodeParameterMapFor.cpp
--- a/Engine/Plugins/FX/Niagara/Source/NiagaraEditor/Private/NiagaraNodeParameterMapFor.cpp
+++ b/Engine/Plugins/FX/Niagara/Source/NiagaraEditor/Private/NiagaraNodeParameterMapFor.cpp
@@ -108,7 +108,7 @@ void UNiagaraNodeParameterMapForWithContinue::AllocateDefaultPins()
 
 FText UNiagaraNodeParameterMapForWithContinue::GetNodeTitle(ENodeTitleType::Type TitleType) const
 {
-	return LOCTEXT("UNiagaraNodeParameterMapForName", "Map For With Continue");
+	return LOCTEXT("UNiagaraNodeParameterMapForWithContinueName", "Map For With Continue");
 }
 
 UNiagaraNodeParameterMapForIndex::UNiagaraNodeParameterMapForIndex() : UNiagaraNode()
diff --git a/Engine/Plugins/FX/Niagara/Source/NiagaraEditor/Private/NiagaraNodeStaticSwitch.cpp b/Engine/Plugins/FX/Niagara/Source/NiagaraEditor/Private/NiagaraNodeStaticSwitch.cpp
--- a/Engine/Plugins/FX/Niagara/Source/NiagaraEditor/Private/NiagaraNodeStaticSwitch.cpp
+++ b/Engine/Plugins/FX/Niagara/Source/NiagaraEditor/Private/NiagaraNodeStaticSwitch.cpp
@@ -446,7 +446,7 @@ bool UNiagaraNodeStaticSwitch::GetVarIndex(FHlslNiagaraTranslator* Translator, i
 			{
 				if (Translator)
 				{
-					Translator->Error(FText::Format(LOCTEXT("InvalidSwitchEnumDefinition", "Number of valid enum entries: {0}. Expected. {1}. Please refresh the node."), FText::FromString(FString::FromInt(OptionValues.Num())), FText::FromString(FString::FromInt(NumOptionsPerVariable))), this, nullptr);
+					Translator->Error(FText::Format(LOCTEXT("NumEnumEntriesDifferent", "Number of valid enum entries: {0}. Expected. {1}. Please refresh the node."), FText::FromString(FString::FromInt(OptionValues.Num())), FText::FromString(FString::FromInt(NumOptionsPerVariable))), this, nullptr);
 				}
 			}
 			// if we have at least one output pin, check if input pin count and output pin count match up
@@ -454,7 +454,7 @@ bool UNiagaraNodeStaticSwitch::GetVarIndex(FHlslNiagaraTranslator* Translator, i
 			{
 				if (Translator)
 				{
-					Translator->Error(FText::Format(LOCTEXT("InvalidSwitchEnumDefinition", "The number of pins on the static switch does not match the number of values defined in the enum."), FText::FromString(FString::FromInt(SwitchValue))), this, nullptr);
+					Translator->Error(FText::Format(LOCTEXT("NumPinsNumEnumEntriesDifferent", "The number of pins on the static switch does not match the number of values defined in the enum."), FText::FromString(FString::FromInt(SwitchValue))), this, nullptr);
 				}
 			}
 			if (Value <= MaxValue && Value >= 0)
diff --git a/Engine/Plugins/FX/Niagara/Source/NiagaraEditor/Private/Sequencer/NiagaraSequence/Sections/MovieSceneNiagaraEmitterSection.cpp b/Engine/Plugins/FX/Niagara/Source/NiagaraEditor/Private/Sequencer/NiagaraSequence/Sections/MovieSceneNiagaraEmitterSection.cpp
--- a/Engine/Plugins/FX/Niagara/Source/NiagaraEditor/Private/Sequencer/NiagaraSequence/Sections/MovieSceneNiagaraEmitterSection.cpp
+++ b/Engine/Plugins/FX/Niagara/Source/NiagaraEditor/Private/Sequencer/NiagaraSequence/Sections/MovieSceneNiagaraEmitterSection.cpp
@@ -170,7 +170,7 @@ bool UMovieSceneNiagaraEmitterSection::TryAddTimeRangeModule(UNiagaraNodeFunctio
 	if (UseStartTimeBinder.TryBind(EmitterUpdateScript, BinderDependentScripts, EmitterUpdateConstantResolver, GetEmitterHandleViewModel()->GetEmitterHandle()->GetUniqueInstanceName(), SectionTimingModule.Get(),
 		InputUsageKey, UseStartTimeInputUsageValue, FNiagaraTypeDefinition::GetBoolDef(), bUseStartTimeIsRequired, UseStartTimeErrorText) == false)
 	{
-		OutErrorMessage = FText::Format(LOCTEXT("StartTimeErrorFormat", "Failed to bind 'use start time' for module.  Message: {0}"), UseStartTimeErrorText);
+		OutErrorMessage = FText::Format(LOCTEXT("UseStartTimeErrorFormat", "Failed to bind 'use start time' for module.  Message: {0}"), UseStartTimeErrorText);
 		return false;
 	}
 	else
diff --git a/Engine/Plugins/FX/Niagara/Source/NiagaraEditor/Private/ViewModels/Stack/NiagaraStackEmitterSettingsGroup.cpp b/Engine/Plugins/FX/Niagara/Source/NiagaraEditor/Private/ViewModels/Stack/NiagaraStackEmitterSettingsGroup.cpp
--- a/Engine/Plugins/FX/Niagara/Source/NiagaraEditor/Private/ViewModels/Stack/NiagaraStackEmitterSettingsGroup.cpp
+++ b/Engine/Plugins/FX/Niagara/Source/NiagaraEditor/Private/ViewModels/Stack/NiagaraStackEmitterSettingsGroup.cpp
@@ -296,7 +296,7 @@ void UNiagaraStackSummaryViewCollapseButton::Initialize(FRequiredEntryData InReq
 
 FText UNiagaraStackSummaryViewCollapseButton::GetDisplayName() const
 {
-	return LOCTEXT("SummaryCollapseButtonTooltip", "Show Advanced");
+	return LOCTEXT("SummaryCollapseButtonDisplayName", "Show Advanced");
 }
 
 UNiagaraStackEntry::EStackRowStyle UNiagaraStackSummaryViewCollapseButton::GetStackRowStyle() const
diff --git a/Engine/Plugins/FX/Niagara/Source/NiagaraEditor/Private/ViewModels/Stack/NiagaraStackErrorItem.cpp b/Engine/Plugins/FX/Niagara/Source/NiagaraEditor/Private/ViewModels/Stack/NiagaraStackErrorItem.cpp
--- a/Engine/Plugins/FX/Niagara/Source/NiagaraEditor/Private/ViewModels/Stack/NiagaraStackErrorItem.cpp
+++ b/Engine/Plugins/FX/Niagara/Source/NiagaraEditor/Private/ViewModels/Stack/NiagaraStackErrorItem.cpp
@@ -223,7 +223,7 @@ UNiagaraStackEntry::EStackRowStyle UNiagaraStackErrorItemDismiss::GetStackRowSty
 
 FText UNiagaraStackErrorItemDismiss::GetFixButtonText() const
 {
-	return StackIssue.GetSeverity() == EStackIssueSeverity::Info || StackIssue.GetSeverity() == EStackIssueSeverity::CustomNote ? LOCTEXT("DismissNote", "Dismiss") : LOCTEXT("DismissIssue", "Dismiss issue");
+	return StackIssue.GetSeverity() == EStackIssueSeverity::Info || StackIssue.GetSeverity() == EStackIssueSeverity::CustomNote ? LOCTEXT("DismissNoteButtonText", "Dismiss") : LOCTEXT("DismissIssueButtonText", "Dismiss issue");
 }
 
 #undef LOCTEXT_NAMESPACE
\ No newline at end of file
diff --git a/Engine/Plugins/FX/Niagara/Source/NiagaraEditor/Private/ViewModels/Stack/NiagaraStackFunctionInput.cpp b/Engine/Plugins/FX/Niagara/Source/NiagaraEditor/Private/ViewModels/Stack/NiagaraStackFunctionInput.cpp
--- a/Engine/Plugins/FX/Niagara/Source/NiagaraEditor/Private/ViewModels/Stack/NiagaraStackFunctionInput.cpp
+++ b/Engine/Plugins/FX/Niagara/Source/NiagaraEditor/Private/ViewModels/Stack/NiagaraStackFunctionInput.cpp
@@ -52,6 +52,8 @@
 
 #define LOCTEXT_NAMESPACE "NiagaraStackViewModel"
 
+static FText TooManyConversionScripts = LOCTEXT("TooManyConversionScripts", "There is more than one dynamic input script available auto-convert the dragged parameter. Please fix this by disabling conversion for all but one of them:\n{0}");
+
 UNiagaraStackFunctionInput::UNiagaraStackFunctionInput()
 	: OwningModuleNode(nullptr)
 	, OwningFunctionCallNode(nullptr)
@@ -319,12 +321,12 @@ bool UNiagaraStackFunctionInput::TestCanPasteWithMessage(const UNiagaraClipboard
 			}
 			if (ClipboardFunctionInput->InputType == InputType)
 			{
-				OutMessage = LOCTEXT("PastMessage", "Paste the input from the clipboard here.");
+				OutMessage = LOCTEXT("PasteMessage", "Paste the input from the clipboard here.");
 				return true;
 			}
 			if (GetPossibleConversionScripts(ClipboardFunctionInput->InputType).Num() > 0)
 			{
-				OutMessage = LOCTEXT("PastMessage", "Paste the input from the clipboard here and auto-convert it with a dynamic input.");
+				OutMessage = LOCTEXT("PasteWithConversionMessage", "Paste the input from the clipboard here and auto-convert it with a dynamic input.");
 				return true;
 			}
 			
@@ -2336,7 +2338,7 @@ void UNiagaraStackFunctionInput::SetLinkedInputViaConversionScript(const FName&
 		{
 			ScriptNames += NiagaraScript->GetPathName() + "\n";
 		}
-		FNiagaraEditorUtilities::WarnWithToastAndLog(FText::Format(LOCTEXT("TooManyConversionScripts", "There is more than one dynamic input script available auto-convert the dragged parameter. Please fix this by disabling conversion for all but one of them:\n{0}"), FText::FromString(ScriptNames)));
+		FNiagaraEditorUtilities::WarnWithToastAndLog(FText::Format(TooManyConversionScripts, FText::FromString(ScriptNames)));
 	}
 	FScopedTransaction ScopedTransaction(LOCTEXT("SetConversionInput", "Make auto-convert dynamic input"));
 	SetDynamicInput(NiagaraScripts[0]);
@@ -2383,7 +2385,7 @@ void UNiagaraStackFunctionInput::SetClipboardContentViaConversionScript(const UN
 		{
 			ScriptNames += NiagaraScript->GetPathName() + "\n";
 		}
-		FNiagaraEditorUtilities::WarnWithToastAndLog(FText::Format(LOCTEXT("TooManyConversionScripts", "There is more than one dynamic input script available auto-convert the pasted parameter. Please fix this by disabling conversion for all but one of them:\n{0}"), FText::FromString(ScriptNames)));
+		FNiagaraEditorUtilities::WarnWithToastAndLog(FText::Format(TooManyConversionScripts, FText::FromString(ScriptNames)));
 	}
 	FScopedTransaction ScopedTransaction(LOCTEXT("SetConversionInput", "Make auto-convert dynamic input"));
 	SetDynamicInput(NiagaraScripts[0]);
diff --git a/Engine/Plugins/FX/Niagara/Source/NiagaraEditor/Private/ViewModels/Stack/NiagaraStackModuleItem.cpp b/Engine/Plugins/FX/Niagara/Source/NiagaraEditor/Private/ViewModels/Stack/NiagaraStackModuleItem.cpp
--- a/Engine/Plugins/FX/Niagara/Source/NiagaraEditor/Private/ViewModels/Stack/NiagaraStackModuleItem.cpp
+++ b/Engine/Plugins/FX/Niagara/Source/NiagaraEditor/Private/ViewModels/Stack/NiagaraStackModuleItem.cpp
@@ -1107,7 +1107,7 @@ bool UNiagaraStackModuleItem::IsDebugDrawEnabled() const
 
 void UNiagaraStackModuleItem::SetDebugDrawEnabled(bool bInEnabled)
 {
-	FScopedTransaction ScopedTransaction(LOCTEXT("EnableDisableModule", "Enable/Disable Debug for Module"));
+	FScopedTransaction ScopedTransaction(LOCTEXT("EnableDisableDebugModule", "Enable/Disable Debug for Module"));
 	FunctionCallNode->DebugState = bInEnabled ? ENiagaraFunctionDebugState::Basic : ENiagaraFunctionDebugState::NoDebug;
 	FunctionCallNode->MarkNodeRequiresSynchronization(__FUNCTION__, true);
 	OnRequestFullRefreshDeferred().Broadcast();
diff --git a/Engine/Plugins/FX/Niagara/Source/NiagaraEditor/Private/Widgets/SNiagaraDebugger.cpp b/Engine/Plugins/FX/Niagara/Source/NiagaraEditor/Private/Widgets/SNiagaraDebugger.cpp
--- a/Engine/Plugins/FX/Niagara/Source/NiagaraEditor/Private/Widgets/SNiagaraDebugger.cpp
+++ b/Engine/Plugins/FX/Niagara/Source/NiagaraEditor/Private/Widgets/SNiagaraDebugger.cpp
@@ -997,11 +997,11 @@ TSharedRef<SWidget> SNiagaraDebugger::MakePlaybackOptionsMenu()
 	{
 		static const TTuple<float, FText, FText> PlaybackSpeeds[] =
 		{
-			MakeTuple(1.0000f,	LOCTEXT("PlaybackSpeed_Normal", "Normal Speed"),		LOCTEXT("NormalSpeedTooltip", "Set playback speed to normal")),
-			MakeTuple(0.5000f,	LOCTEXT("PlaybackSpeed_Half", "Half Speed"),			LOCTEXT("NormalSpeedTooltip", "Set playback speed to half the normal speed")),
-			MakeTuple(0.2500f,	LOCTEXT("PlaybackSpeed_Quarter", "Quarter Speed "),		LOCTEXT("NormalSpeedTooltip", "Set playback speed to quarter the normal speed")),
-			MakeTuple(0.1250f,	LOCTEXT("PlaybackSpeed_Eighth", "Eighth Speed "),		LOCTEXT("NormalSpeedTooltip", "Set playback speed to eighth the normal speed")),
-			MakeTuple(0.0625f,	LOCTEXT("PlaybackSpeed_Sixteenth", "Sixteenth Speed "),	LOCTEXT("NormalSpeedTooltip", "Set playback speed to sixteenth the normal speed")),
+			MakeTuple(1.0000f,	LOCTEXT("PlaybackSpeed_Normal", "Normal Speed"),		LOCTEXT("NormalSpeedTooltip_1", "Set playback speed to normal")),
+			MakeTuple(0.5000f,	LOCTEXT("PlaybackSpeed_Half", "Half Speed"),			LOCTEXT("NormalSpeedTooltip_0.5", "Set playback speed to half the normal speed")),
+			MakeTuple(0.2500f,	LOCTEXT("PlaybackSpeed_Quarter", "Quarter Speed "),		LOCTEXT("NormalSpeedTooltip_0.25", "Set playback speed to quarter the normal speed")),
+			MakeTuple(0.1250f,	LOCTEXT("PlaybackSpeed_Eighth", "Eighth Speed "),		LOCTEXT("NormalSpeedTooltip_0.125", "Set playback speed to eighth the normal speed")),
+			MakeTuple(0.0625f,	LOCTEXT("PlaybackSpeed_Sixteenth", "Sixteenth Speed "),	LOCTEXT("NormalSpeedTooltip_0.0625", "Set playback speed to sixteenth the normal speed")),
 		};
 		UNiagaraDebugHUDSettings* Settings = GetMutableDefault<UNiagaraDebugHUDSettings>();
 
diff --git a/Engine/Plugins/FX/Niagara/Source/NiagaraEditor/Public/Widgets/SNiagaraParameterMenu.cpp b/Engine/Plugins/FX/Niagara/Source/NiagaraEditor/Public/Widgets/SNiagaraParameterMenu.cpp
--- a/Engine/Plugins/FX/Niagara/Source/NiagaraEditor/Public/Widgets/SNiagaraParameterMenu.cpp
+++ b/Engine/Plugins/FX/Niagara/Source/NiagaraEditor/Public/Widgets/SNiagaraParameterMenu.cpp
@@ -283,7 +283,7 @@ void SNiagaraAddParameterFromPanelMenu::CollectAllActions(FGraphActionListBuilde
 	auto CollectEmitterNamespaceParameterActions = [this, &Collector]() {
 		const FNiagaraNamespaceMetadata NamespaceMetaData = FNiagaraEditorUtilities::GetNamespaceMetaDataForId(FNiagaraEditorGuids::EmitterNamespaceMetaDataGuid);
 		TArray<FNiagaraVariable> Variables = FNiagaraConstants::GetEngineConstants().FilterByPredicate([](const FNiagaraVariable& Var) { return Var.IsInNameSpace(FNiagaraConstants::EmitterNamespaceString); });
-		const FText CategoryText = bShowNamespaceCategory ? GetNamespaceCategoryText(NamespaceMetaData) : LOCTEXT("EngineConstantNamespaceCategory", "Add Emitter Constant");
+		const FText CategoryText = bShowNamespaceCategory ? GetNamespaceCategoryText(NamespaceMetaData) : LOCTEXT("EmitterConstantNamespaceCategory", "Add Emitter Constant");
 		const FString RootCategoryStr = FString();
 		const bool bMakeNameUnique = false;
 		AddParameterGroup(
diff --git a/Engine/Plugins/FX/Niagara/Source/NiagaraEditorWidgets/Private/Stack/SNiagaraStackFunctionInputValue.cpp b/Engine/Plugins/FX/Niagara/Source/NiagaraEditorWidgets/Private/Stack/SNiagaraStackFunctionInputValue.cpp
--- a/Engine/Plugins/FX/Niagara/Source/NiagaraEditorWidgets/Private/Stack/SNiagaraStackFunctionInputValue.cpp
+++ b/Engine/Plugins/FX/Niagara/Source/NiagaraEditorWidgets/Private/Stack/SNiagaraStackFunctionInputValue.cpp
@@ -1645,7 +1645,7 @@ void SNiagaraStackFunctionInputValue::FilteredVisibleCheckStateChanged(ECheckBox
 		}
 		if (bNewValue != SummaryViewMetaData->bVisible)
 		{
-			FScopedTransaction ScopedTransaction(FText::Format(LOCTEXT("SummaryViewChangedInputSortIndex", "Changed summary view visibility for {0}"), FunctionInput->GetDisplayName()));
+			FScopedTransaction ScopedTransaction(FText::Format(LOCTEXT("SummaryViewChangedVisibilityForInput", "Changed summary view visibility for {0}"), FunctionInput->GetDisplayName()));
 			EditorData->Modify();
 
 			SummaryViewMetaData->bVisible = bNewValue;
