Fixed a bug where using undo after adding a renderer could cause a crash.

diff --git a/Engine/Plugins/FX/Niagara/Source/NiagaraEditor/Private/ViewModels/NiagaraSystemViewModel.cpp b/Engine/Plugins/FX/Niagara/Source/NiagaraEditor/Private/ViewModels/NiagaraSystemViewModel.cpp
--- a/Engine/Plugins/FX/Niagara/Source/NiagaraEditor/Private/ViewModels/NiagaraSystemViewModel.cpp
+++ b/Engine/Plugins/FX/Niagara/Source/NiagaraEditor/Private/ViewModels/NiagaraSystemViewModel.cpp
@@ -1156,7 +1156,7 @@ void FNiagaraSystemViewModel::RefreshAll()
 		// when used in conjunction with warm up.  This is a temporary fix for these crashes.
 		CompileSystem(false);
 	}
-	ResetSystem(ETimeResetMode::AllowResetTime, EMultiResetMode::ResetThisInstance, EReinitMode::ReinitializeSystem);
+	ResetSystem(ETimeResetMode::AllowResetTime, EMultiResetMode::AllowResetAllInstances, EReinitMode::ReinitializeSystem);
 	RefreshEmitterHandleViewModels();
 	RefreshSequencerTracks();
 	InvalidateCachedCompileStatus();
@@ -1754,7 +1754,7 @@ void FNiagaraSystemViewModel::EmitterHandleNameChanged()
 
 void FNiagaraSystemViewModel::EmitterPropertyChanged()
 {
-	ResetSystem(ETimeResetMode::AllowResetTime, EMultiResetMode::ResetThisInstance, EReinitMode::ReinitializeSystem);
+	ResetSystem(ETimeResetMode::AllowResetTime, EMultiResetMode::AllowResetAllInstances, EReinitMode::ReinitializeSystem);
 }
 
 void FNiagaraSystemViewModel::ScriptCompiled(UNiagaraScript*, const FGuid&)
