Fixed a bug where user parameters bound to renderer position bindings were not updated to position types

diff --git a/Engine/Plugins/FX/Niagara/Source/Niagara/Private/NiagaraSystem.cpp b/Engine/Plugins/FX/Niagara/Source/Niagara/Private/NiagaraSystem.cpp
--- a/Engine/Plugins/FX/Niagara/Source/Niagara/Private/NiagaraSystem.cpp
+++ b/Engine/Plugins/FX/Niagara/Source/Niagara/Private/NiagaraSystem.cpp
@@ -1165,7 +1165,7 @@ void UNiagaraSystem::PostLoad()
 	}
 
 #if WITH_EDITORONLY_DATA
-	FixupPositionUserParameters(); //TODO[mg]: hide this behind a new custom version?
+	FixupPositionUserParameters();
 	
 	if (EditorData == nullptr)
 	{
@@ -3860,6 +3860,35 @@ void UNiagaraSystem::FixupPositionUserParameters()
 		}
 	});
 
+	// check if any of the renderers have a user param bound to a position binding
+	for (FNiagaraEmitterHandle& EmitterHandle : EmitterHandles)
+	{
+		if (UNiagaraEmitter* NiagaraEmitter = EmitterHandle.GetInstance())
+		{
+			for (UNiagaraRendererProperties* Renderer : NiagaraEmitter->GetRenderers())
+			{
+				for (const FNiagaraVariableAttributeBinding* AttributeBinding : Renderer->GetAttributeBindings())
+				{
+					if (AttributeBinding == nullptr)
+					{
+						continue;
+					}
+					if (AttributeBinding->GetBindingSourceMode() == ExplicitUser && AttributeBinding->GetType() == FNiagaraTypeDefinition::GetPositionDef())
+					{
+						for (FNiagaraVariable UserParam : UserParameters)
+						{
+							if (UserParam.GetName() == AttributeBinding->GetParamMapBindableVariable().GetName() && !LinkedPositionInputs.Contains(UserParam))
+							{
+								LinkedPositionInputs.Add(UserParam);
+								break;
+							}
+						}
+					}
+				}
+			}
+		}
+	}
+
 	// looks like we have a few FVector3f user parameters that are linked to position inputs in the stack.
 	// Most likely this is because core modules were changed to use position data. To ease the transition of old assets, we auto-convert those inputs to position types as well.
 	for (const FNiagaraVariable& LinkedParameter : LinkedPositionInputs)
