Niagara - Make sure that script variable ids are unique within a graph on post load.

diff --git a/Engine/Plugins/FX/Niagara/Source/NiagaraEditor/Private/NiagaraGraph.cpp b/Engine/Plugins/FX/Niagara/Source/NiagaraEditor/Private/NiagaraGraph.cpp
--- a/Engine/Plugins/FX/Niagara/Source/NiagaraEditor/Private/NiagaraGraph.cpp
+++ b/Engine/Plugins/FX/Niagara/Source/NiagaraEditor/Private/NiagaraGraph.cpp
@@ -401,6 +401,32 @@ void UNiagaraGraph::PostLoad()
 		}
 	}
 
+	if (NiagaraVer < FNiagaraCustomVersion::StaticSwitchFunctionPinsUsePersistentGuids)
+	{
+		// Make sure that the variable persistent ids are unique, otherwise pin allocation for static switches will be inconsistent.
+		TArray<TObjectPtr<UNiagaraScriptVariable>> ScriptVariables;
+		VariableToScriptVariable.GenerateValueArray(ScriptVariables);
+		TSet<FGuid> SeenGuids;
+		for (TObjectPtr<UNiagaraScriptVariable>& ScriptVariable : ScriptVariables)
+		{
+			if (ScriptVariable != nullptr)
+			{
+				bool bAlreadySeen = false;
+				SeenGuids.Add(ScriptVariable->Metadata.GetVariableGuid(), &bAlreadySeen);
+				if (bAlreadySeen)
+				{
+					if (ensureMsgf(ScriptVariable->GetIsSubscribedToParameterDefinitions() == false, TEXT("Duplicate id found for script variable %s controlled by parameter definitions."), *ScriptVariable->GetPathName()))
+					{
+						ScriptVariable->Metadata.SetVariableGuid(UNiagaraScriptVariable::GenerateStableGuid(ScriptVariable));
+						bAlreadySeen = false;
+						SeenGuids.Add(ScriptVariable->Metadata.GetVariableGuid(), &bAlreadySeen);
+						ensureMsgf(bAlreadySeen == false, TEXT("Failed to generate a unique stable guid for script variable %s."), *ScriptVariable->GetPathName());
+					}
+				}
+			}
+		}
+	}
+
 	InvalidateCachedParameterData();
 }
 
diff --git a/Engine/Plugins/FX/Niagara/Source/NiagaraEditor/Private/NiagaraScriptVariable.cpp b/Engine/Plugins/FX/Niagara/Source/NiagaraEditor/Private/NiagaraScriptVariable.cpp
--- a/Engine/Plugins/FX/Niagara/Source/NiagaraEditor/Private/NiagaraScriptVariable.cpp
+++ b/Engine/Plugins/FX/Niagara/Source/NiagaraEditor/Private/NiagaraScriptVariable.cpp
@@ -85,21 +85,7 @@ void UNiagaraScriptVariable::PostLoad()
 
 	if (!Metadata.GetVariableGuid().IsValid())
 	{
-		// If the variable doesn't have a valid guid we need to assign one, but we want it to be stable across sessions
-		// so we're going to generate a hash from an identifier string and then use that to generate the values for the GUID.
-		// This won't result in a value that's a true GUID but since every object path in a project must be unique, the
-		// md5 of these values should be more than unique enough for our purposes.
-		FString IdentifierString = GetPathName() + Variable.GetName().ToString() + Variable.GetType().GetName();
-
-		// Create an md5 hash for the string data.
-		uint32 HashBuffer[] { 0, 0, 0, 0 };
-		FMD5 IdentifierStringHash;
-		IdentifierStringHash.Update((uint8*)GetData(IdentifierString), IdentifierString.Len() * sizeof(TCHAR));
-		IdentifierStringHash.Final((uint8*)&HashBuffer);
-
-		// Assign the guid components from the hash.
-		FGuid GuidFromHash = FGuid(HashBuffer[0], HashBuffer[1], HashBuffer[2], HashBuffer[3]);
-		Metadata.SetVariableGuid(GuidFromHash);
+		Metadata.SetVariableGuid(GenerateStableGuid(this));
 	}
 	if (ChangeId.IsValid() == false)
 	{
@@ -173,3 +159,21 @@ bool UNiagaraScriptVariable::DefaultsAreEquivalent(const UNiagaraScriptVariable*
 	ensureMsgf(false, TEXT("Encountered unknown ENiagaraDefaultMode when comparing UNiagaraScriptVariable defaults!"));
 	return false;
 }
+
+FGuid UNiagaraScriptVariable::GenerateStableGuid(const UNiagaraScriptVariable* ScriptVariable)
+{
+	// To create a guid that will be stable when generated across sessions we're going to generate a hash from 
+	// an identifier string and then use that to generate the values for the GUID.  This won't result in a value
+	// that's a true GUID but since every object path in a project must be unique, the md5 of these values should
+	// be more than unique enough for our purposes.
+	FString IdentifierString = ScriptVariable->GetPathName() + ScriptVariable->Variable.GetName().ToString() + ScriptVariable->Variable.GetType().GetName();
+
+	// Create an md5 hash for the string data.
+	uint32 HashBuffer[]{ 0, 0, 0, 0 };
+	FMD5 IdentifierStringHash;
+	IdentifierStringHash.Update((uint8*)GetData(IdentifierString), IdentifierString.Len() * sizeof(TCHAR));
+	IdentifierStringHash.Final((uint8*)&HashBuffer);
+
+	// Assign the guid components from the hash.
+	return FGuid(HashBuffer[0], HashBuffer[1], HashBuffer[2], HashBuffer[3]);
+}
diff --git a/Engine/Plugins/FX/Niagara/Source/NiagaraEditor/Public/NiagaraScriptVariable.h b/Engine/Plugins/FX/Niagara/Source/NiagaraEditor/Public/NiagaraScriptVariable.h
--- a/Engine/Plugins/FX/Niagara/Source/NiagaraEditor/Public/NiagaraScriptVariable.h
+++ b/Engine/Plugins/FX/Niagara/Source/NiagaraEditor/Public/NiagaraScriptVariable.h
@@ -95,6 +95,8 @@ public:
 
 	static bool DefaultsAreEquivalent(const UNiagaraScriptVariable* ScriptVarA, const UNiagaraScriptVariable* ScriptVarB);
 
+	static FGuid GenerateStableGuid(const UNiagaraScriptVariable* ScriptVariable);
+
 private:
 	void AllocateData()
 	{
