Merged CL 18707599 from 5.0: Fix for ensure loading Niagara_Advanced.
Fixes mismatch in data size/contents between override params and instance/template overrides.
Converting instance/template overrides to doubles when doing conversion to Position type.

diff --git a/Engine/Plugins/FX/Niagara/Source/Niagara/Private/NiagaraComponent.cpp b/Engine/Plugins/FX/Niagara/Source/Niagara/Private/NiagaraComponent.cpp
--- a/Engine/Plugins/FX/Niagara/Source/Niagara/Private/NiagaraComponent.cpp
+++ b/Engine/Plugins/FX/Niagara/Source/Niagara/Private/NiagaraComponent.cpp
@@ -3051,10 +3051,12 @@ void UNiagaraComponent::FixInvalidUserParameterOverrideData()
 					if (Params.RemoveAndCopyValue(ExistingVar, VariantData))
 					{
 						// Convert the old float data to doubles.
+#if !UE_LARGE_WORLD_COORDINATES_DISABLED
 						check(VariantData.GetNumBytes() == sizeof(FVector3f));
 						FVector3f ExistingData = *((FVector3f*)(VariantData.GetBytes()));						
 						FVector NewData(ExistingData);
 						VariantData.SetBytes((uint8*)&NewData, sizeof(NewData));
+#endif
 						Params.Add(PositionVar, VariantData);
 					}
 				};
