Build fix after robomerge

diff --git a/Engine/Plugins/FX/Niagara/Source/Niagara/Private/NiagaraComponent.cpp b/Engine/Plugins/FX/Niagara/Source/Niagara/Private/NiagaraComponent.cpp
--- a/Engine/Plugins/FX/Niagara/Source/Niagara/Private/NiagaraComponent.cpp
+++ b/Engine/Plugins/FX/Niagara/Source/Niagara/Private/NiagaraComponent.cpp
@@ -3051,12 +3051,10 @@ void UNiagaraComponent::FixInvalidUserParameterOverrideData()
 					if (Params.RemoveAndCopyValue(ExistingVar, VariantData))
 					{
 						// Convert the old float data to doubles.
-#if !UE_LARGE_WORLD_COORDINATES_DISABLED
 						check(VariantData.GetNumBytes() == sizeof(FVector3f));
 						FVector3f ExistingData = *((FVector3f*)(VariantData.GetBytes()));						
 						FVector NewData(ExistingData);
 						VariantData.SetBytes((uint8*)&NewData, sizeof(NewData));
-#endif
 						Params.Add(PositionVar, VariantData);
 					}
 				};
