named attributes in grid2d reader

diff --git a/Engine/Plugins/FX/Niagara/Source/Niagara/Classes/NiagaraDataInterfaceGrid2DCollection.h b/Engine/Plugins/FX/Niagara/Source/Niagara/Classes/NiagaraDataInterfaceGrid2DCollection.h
--- a/Engine/Plugins/FX/Niagara/Source/Niagara/Classes/NiagaraDataInterfaceGrid2DCollection.h
+++ b/Engine/Plugins/FX/Niagara/Source/Niagara/Classes/NiagaraDataInterfaceGrid2DCollection.h
@@ -60,6 +60,12 @@ struct FGrid2DCollectionRWInstanceData_GameThread
 
 	bool NeedsRealloc = false;
 
+	// We need to essentially make this a linked list to avoid more refactoring for now
+	// eventually we can clean this logic up, but this allows us to have a subclass that
+	// overrides the render thread data, which in this case is for a grid reader
+	UNiagaraDataInterface* OtherDI = nullptr;
+	FGrid2DCollectionRWInstanceData_GameThread* OtherInstanceData = nullptr;
+
 	int32 FindAttributeIndexByName(const FName& InName, int32 NumChannels);
 
 	bool UpdateTargetTexture(ENiagaraGpuBufferFormat BufferFormat);
@@ -88,6 +94,11 @@ struct FGrid2DCollectionRWInstanceData_RenderThread
 	FIntVector4 PreviewAttribute = FIntVector4(INDEX_NONE, INDEX_NONE, INDEX_NONE, INDEX_NONE);
 #endif
 
+	// We need to essentially make this a linked list to avoid more refactoring for now
+	// eventually we can clean this logic up, but this allows us to have a subclass that
+	// overrides the render thread data, which in this case is for a grid reader
+	FNiagaraDataInterfaceProxy* OtherProxy = nullptr;
+
 	void BeginSimulate(FRHICommandList& RHICmdList);
 	void EndSimulate(FRHICommandList& RHICmdList);
 };
@@ -110,6 +121,31 @@ struct FNiagaraDataInterfaceProxyGrid2DCollectionProxy : public FNiagaraDataInte
 
 };
 
+struct FNiagaraDataInterfaceParametersCS_Grid2DCollection : public FNiagaraDataInterfaceParametersCS
+{
+	DECLARE_TYPE_LAYOUT(FNiagaraDataInterfaceParametersCS_Grid2DCollection, NonVirtual);
+
+public:
+	void Bind(const FNiagaraDataInterfaceGPUParamInfo& ParameterInfo, const class FShaderParameterMap& ParameterMap);
+	void Set(FRHICommandList& RHICmdList, const FNiagaraDataInterfaceSetArgs& Context) const;
+	void Unset(FRHICommandList& RHICmdList, const FNiagaraDataInterfaceSetArgs& Context) const;
+
+private:
+	LAYOUT_FIELD(FShaderParameter, NumAttributesParam);
+	LAYOUT_FIELD(FShaderParameter, UnitToUVParam);
+	LAYOUT_FIELD(FShaderParameter, NumCellsParam);
+	LAYOUT_FIELD(FShaderParameter, CellSizeParam);
+	LAYOUT_FIELD(FShaderParameter, WorldBBoxSizeParam);
+
+	LAYOUT_FIELD(FShaderResourceParameter, GridParam);
+	LAYOUT_FIELD(FRWShaderParameter, OutputGridParam);
+	LAYOUT_FIELD(FShaderParameter, AttributeIndicesParam);
+
+	LAYOUT_FIELD(FShaderResourceParameter, SamplerParam);
+	LAYOUT_FIELD(TMemoryImageArray<FName>, AttributeNames);
+	LAYOUT_FIELD(TMemoryImageArray<uint32>, AttributeChannelCount);
+};
+
 UCLASS(EditInlineNew, Category = "Grid", meta = (DisplayName = "Grid2D Collection", Experimental), Blueprintable, BlueprintType)
 class NIAGARA_API UNiagaraDataInterfaceGrid2DCollection : public UNiagaraDataInterfaceGrid2D
 {
@@ -274,6 +310,8 @@ public:
 	static int32 GetComponentCountFromFuncName(const FName& FuncName);
 	static FNiagaraTypeDefinition GetValueTypeFromFuncName(const FName& FuncName);
 	static bool CanCreateVarFromFuncName(const FName& FuncName);
+
+	TMap<FNiagaraSystemInstanceID, FGrid2DCollectionRWInstanceData_GameThread*>& GetSystemInstancesToProxyData_GT() { return SystemInstancesToProxyData_GT; }
 protected:
 #if WITH_EDITORONLY_DATA
 	void WriteSetHLSL(const FNiagaraDataInterfaceGPUParamInfo& ParamInfo, const FNiagaraDataInterfaceGeneratedFunction& FunctionInfo, int FunctionInstanceIndex, int32 InNumChannels, FString& OutHLSL);
diff --git a/Engine/Plugins/FX/Niagara/Source/Niagara/Classes/NiagaraDataInterfaceGrid2DCollectionReader.h b/Engine/Plugins/FX/Niagara/Source/Niagara/Classes/NiagaraDataInterfaceGrid2DCollectionReader.h
--- a/Engine/Plugins/FX/Niagara/Source/Niagara/Classes/NiagaraDataInterfaceGrid2DCollectionReader.h
+++ b/Engine/Plugins/FX/Niagara/Source/Niagara/Classes/NiagaraDataInterfaceGrid2DCollectionReader.h
@@ -3,95 +3,41 @@
 
 #include "NiagaraDataInterface.h"
 #include "NiagaraDataInterfaceRW.h"
-#include "NiagaraDataInterfaceGrid2DCollection.h"
 #include "ClearQuad.h"
 #include "NiagaraComponent.h"
+#include "Niagara/Private/NiagaraStats.h"
+#include "NiagaraDataInterfaceGrid2DCollection.h"
 
 #include "NiagaraDataInterfaceGrid2DCollectionReader.generated.h"
 
 class FNiagaraSystemInstance;
-class UTextureRenderTarget2D;
+class UTextureRenderTarget;
+class UTextureRenderTargetVolume;
 
-
-struct FGrid2DCollectionReaderInstanceData_GameThread
-{
-	FNiagaraSystemInstance* SystemInstance = nullptr;
-	FNiagaraEmitterInstance* EmitterInstance = nullptr;	
-
-	FString EmitterName;
-	FString DIName;
-};
-
-struct FGrid2DCollectionReaderInstanceData_RenderThread
-{
-	FNiagaraDataInterfaceProxyGrid2DCollectionProxy* ProxyToUse = nullptr;
-};
-
-struct FNiagaraDataInterfaceProxyGrid2DCollectionReaderProxy : public FNiagaraDataInterfaceProxyRW
-{
-	FNiagaraDataInterfaceProxyGrid2DCollectionReaderProxy() {}
-
-	virtual FIntVector GetElementCount(FNiagaraSystemInstanceID SystemInstanceID) const;
-
-	/* List of proxy data for each system instances*/
-	// #todo(dmp): this should all be refactored to avoid duplicate code
-	TMap<FNiagaraSystemInstanceID, FGrid2DCollectionReaderInstanceData_RenderThread> SystemInstancesToProxyData_RT;
-};
-
-// #todo(dmp): base class has all the RW attributes, even though we only care about the attributes that query the grid.  Cleaning this up would be great
-UCLASS(EditInlineNew, Category = "Grid", meta = (DisplayName = "Grid2D Collection Reader", Experimental), Blueprintable, BlueprintType, hidecategories = (Grid,RW))
-class NIAGARA_API UNiagaraDataInterfaceGrid2DCollectionReader : public UNiagaraDataInterfaceGrid2D
+UCLASS(EditInlineNew, Category = "Grid", hidecategories = (Grid2DCollection, Grid, Deprecated), meta = (DisplayName = "Grid2D Collection Reader", Experimental), Blueprintable, BlueprintType)
+class NIAGARA_API UNiagaraDataInterfaceGrid2DCollectionReader : public UNiagaraDataInterfaceGrid2DCollection
 {
 	GENERATED_UCLASS_BODY()
 
 public:
-
 	DECLARE_NIAGARA_DI_PARAMETER();
-	
+
 	// Name of the emitter to read from
 	UPROPERTY(EditAnywhere, Category = "Reader")
-	FString EmitterName;
+		FString EmitterName;
 
 	// Name of the Grid2DCollection Data Interface on the emitter
 	UPROPERTY(EditAnywhere, Category = "Reader")
-	FString DIName;
-
-	virtual void PostInitProperties() override;
-	
-	//~ UNiagaraDataInterface interface
-	// VM functionality
-	virtual void GetFunctions(TArray<FNiagaraFunctionSignature>& OutFunctions) override;
-	virtual void GetVMExternalFunction(const FVMExternalFunctionBindingInfo& BindingInfo, void* InstanceData, FVMExternalFunction &OutFunc) override;
+		FString DIName;
 
 	virtual bool Equals(const UNiagaraDataInterface* Other) const override;
-
-	// GPU sim functionality
-#if WITH_EDITORONLY_DATA
-	virtual void GetParameterDefinitionHLSL(const FNiagaraDataInterfaceGPUParamInfo& ParamInfo, FString& OutHLSL) override;
-	virtual bool GetFunctionHLSL(const FNiagaraDataInterfaceGPUParamInfo& ParamInfo, const FNiagaraDataInterfaceGeneratedFunction& FunctionInfo, int FunctionInstanceIndex, FString& OutHLSL) override;
-#endif
-
-	virtual void ProvidePerInstanceDataForRenderThread(void* DataForRenderThread, void* PerInstanceData, const FNiagaraSystemInstanceID& SystemInstance) override {}
 	virtual bool InitPerInstanceData(void* PerInstanceData, FNiagaraSystemInstance* SystemInstance) override;
-	virtual void DestroyPerInstanceData(void* PerInstanceData, FNiagaraSystemInstance* SystemInstance) override;
-	virtual int32 PerInstanceDataSize() const override { return sizeof(FGrid2DCollectionReaderInstanceData_GameThread); }
-	
 	virtual void GetEmitterDependencies(UNiagaraSystem* Asset, TArray<UNiagaraEmitter*>& Dependencies) const override;
-	//~ UNiagaraDataInterface interface END
-
-	static const FString NumTilesName;
-
-	static const FString GridName;
-	static const FString OutputGridName;
-	static const FString SamplerName;
-
-	static const FName GetValueFunctionName;
-	static const FName SampleGridFunctionName;
+	virtual void GetFunctions(TArray<FNiagaraFunctionSignature>& OutFunctions) override;
+	virtual bool PerInstanceTickPostSimulate(void* PerInstanceData, FNiagaraSystemInstance* SystemInstance, float DeltaSeconds) override;
 
 protected:
 	//~ UNiagaraDataInterface interface
 	virtual bool CopyToInternal(UNiagaraDataInterface* Destination) const override;
 	//~ UNiagaraDataInterface interface END
-
-	TMap<FNiagaraSystemInstanceID, FGrid2DCollectionReaderInstanceData_GameThread*> SystemInstancesToProxyData_GT;	
-};
+};
\ No newline at end of file
diff --git a/Engine/Plugins/FX/Niagara/Source/Niagara/Private/NiagaraDataInterfaceGrid2DCollection.cpp b/Engine/Plugins/FX/Niagara/Source/Niagara/Private/NiagaraDataInterfaceGrid2DCollection.cpp
--- a/Engine/Plugins/FX/Niagara/Source/Niagara/Private/NiagaraDataInterfaceGrid2DCollection.cpp
+++ b/Engine/Plugins/FX/Niagara/Source/Niagara/Private/NiagaraDataInterfaceGrid2DCollection.cpp
@@ -203,146 +203,134 @@ struct FNiagaraGrid2DLegacyTiled2DInfo
 };
 
 /*--------------------------------------------------------------------------------------------------------------------------*/
-struct FNiagaraDataInterfaceParametersCS_Grid2DCollection : public FNiagaraDataInterfaceParametersCS
+
+void FNiagaraDataInterfaceParametersCS_Grid2DCollection::Bind(const FNiagaraDataInterfaceGPUParamInfo& ParameterInfo, const class FShaderParameterMap& ParameterMap)
 {
-	DECLARE_TYPE_LAYOUT(FNiagaraDataInterfaceParametersCS_Grid2DCollection, NonVirtual);
+	NumAttributesParam.Bind(ParameterMap, *(UNiagaraDataInterfaceRWBase::NumAttributesName + ParameterInfo.DataInterfaceHLSLSymbol));
+	NumCellsParam.Bind(ParameterMap, *(UNiagaraDataInterfaceRWBase::NumCellsName + ParameterInfo.DataInterfaceHLSLSymbol));
+	UnitToUVParam.Bind(ParameterMap, *(UNiagaraDataInterfaceRWBase::UnitToUVName + ParameterInfo.DataInterfaceHLSLSymbol));
+	CellSizeParam.Bind(ParameterMap, *(UNiagaraDataInterfaceRWBase::CellSizeName + ParameterInfo.DataInterfaceHLSLSymbol));
+	WorldBBoxSizeParam.Bind(ParameterMap, *(UNiagaraDataInterfaceRWBase::WorldBBoxSizeName + ParameterInfo.DataInterfaceHLSLSymbol));
 
-public:
-	void Bind(const FNiagaraDataInterfaceGPUParamInfo& ParameterInfo, const class FShaderParameterMap& ParameterMap)
-	{
-		NumAttributesParam.Bind(ParameterMap, *(UNiagaraDataInterfaceRWBase::NumAttributesName + ParameterInfo.DataInterfaceHLSLSymbol));
-		NumCellsParam.Bind(ParameterMap, *(UNiagaraDataInterfaceRWBase::NumCellsName + ParameterInfo.DataInterfaceHLSLSymbol));
-		UnitToUVParam.Bind(ParameterMap, *(UNiagaraDataInterfaceRWBase::UnitToUVName + ParameterInfo.DataInterfaceHLSLSymbol));
-		CellSizeParam.Bind(ParameterMap, *(UNiagaraDataInterfaceRWBase::CellSizeName + ParameterInfo.DataInterfaceHLSLSymbol));
-		WorldBBoxSizeParam.Bind(ParameterMap, *(UNiagaraDataInterfaceRWBase::WorldBBoxSizeName + ParameterInfo.DataInterfaceHLSLSymbol));
+	GridParam.Bind(ParameterMap, *(UNiagaraDataInterfaceGrid2DCollection::GridName + ParameterInfo.DataInterfaceHLSLSymbol));
+	OutputGridParam.Bind(ParameterMap, *(UNiagaraDataInterfaceGrid2DCollection::OutputGridName + ParameterInfo.DataInterfaceHLSLSymbol));
 
-		GridParam.Bind(ParameterMap, *(UNiagaraDataInterfaceGrid2DCollection::GridName + ParameterInfo.DataInterfaceHLSLSymbol));
-		OutputGridParam.Bind(ParameterMap, *(UNiagaraDataInterfaceGrid2DCollection::OutputGridName + ParameterInfo.DataInterfaceHLSLSymbol));
+	SamplerParam.Bind(ParameterMap, *(UNiagaraDataInterfaceGrid2DCollection::SamplerName + ParameterInfo.DataInterfaceHLSLSymbol));
+	AttributeIndicesParam.Bind(ParameterMap, *(UNiagaraDataInterfaceGrid2DCollection::AttributeIndicesBaseName + ParameterInfo.DataInterfaceHLSLSymbol));
 
-		SamplerParam.Bind(ParameterMap, *(UNiagaraDataInterfaceGrid2DCollection::SamplerName + ParameterInfo.DataInterfaceHLSLSymbol));
-		AttributeIndicesParam.Bind(ParameterMap, *(UNiagaraDataInterfaceGrid2DCollection::AttributeIndicesBaseName + ParameterInfo.DataInterfaceHLSLSymbol));
+	// Gather up all the attribute names referenced. Note that there may be multiple in the list of the same name,
+	// but we only deal with this by the number of bound methods.
+	{
+		int32 NumFuncs = ParameterInfo.GeneratedFunctions.Num();
 
-		// Gather up all the attribute names referenced. Note that there may be multiple in the list of the same name,
-		// but we only deal with this by the number of bound methods.
+		for (int32 FuncIdx = 0; FuncIdx < NumFuncs; ++FuncIdx)
 		{
-			int32 NumFuncs = ParameterInfo.GeneratedFunctions.Num();
-
-			for (int32 FuncIdx = 0; FuncIdx < NumFuncs; ++FuncIdx)
+			const FNiagaraDataInterfaceGeneratedFunction& Func = ParameterInfo.GeneratedFunctions[FuncIdx];
+			static const FName NAME_Attribute("Attribute");
+			const FName* AttributeName = Func.FindSpecifierValue(NAME_Attribute);
+			if (AttributeName != nullptr)
 			{
-				const FNiagaraDataInterfaceGeneratedFunction& Func = ParameterInfo.GeneratedFunctions[FuncIdx];
-				static const FName NAME_Attribute("Attribute");
-				const FName* AttributeName = Func.FindSpecifierValue(NAME_Attribute);
-				if (AttributeName != nullptr)
-				{
-					int32 ComponentCount = UNiagaraDataInterfaceGrid2DCollection::GetComponentCountFromFuncName(Func.DefinitionName);
-					AttributeNames.Add(*AttributeName);
-					AttributeChannelCount.Add(ComponentCount);
-				}
-				else
-				{
-					AttributeNames.Add(FName());
-					AttributeChannelCount.Add(INDEX_NONE);
-				}
+				int32 ComponentCount = UNiagaraDataInterfaceGrid2DCollection::GetComponentCountFromFuncName(Func.DefinitionName);
+				AttributeNames.Add(*AttributeName);
+				AttributeChannelCount.Add(ComponentCount);
+			}
+			else
+			{
+				AttributeNames.Add(FName());
+				AttributeChannelCount.Add(INDEX_NONE);
 			}
 		}
 	}
+}
 
-	void Set(FRHICommandList& RHICmdList, const FNiagaraDataInterfaceSetArgs& Context) const
-	{
-		check(IsInRenderingThread());
-
-		// Get shader and DI
-		FRHIComputeShader* ComputeShaderRHI = Context.Shader.GetComputeShader();
-		FNiagaraDataInterfaceProxyGrid2DCollectionProxy* VFDI = static_cast<FNiagaraDataInterfaceProxyGrid2DCollectionProxy*>(Context.DataInterface);
-
-		FGrid2DCollectionRWInstanceData_RenderThread* ProxyData = VFDI->SystemInstancesToProxyData_RT.Find(Context.SystemInstanceID);
-		check(ProxyData);
+void FNiagaraDataInterfaceParametersCS_Grid2DCollection::Set(FRHICommandList& RHICmdList, const FNiagaraDataInterfaceSetArgs& Context) const
+{
+	check(IsInRenderingThread());
 
-		if (ProxyData->AttributeIndices.Num() == 0 && AttributeNames.Num() > 0)
-		{
-			int NumAttrIndices = Align(AttributeNames.Num(), 4);
-			ProxyData->AttributeIndices.SetNumZeroed(NumAttrIndices);
+	// Get shader and DI
+	FRHIComputeShader* ComputeShaderRHI = Context.Shader.GetComputeShader();
+	FNiagaraDataInterfaceProxyGrid2DCollectionProxy* VFDI = static_cast<FNiagaraDataInterfaceProxyGrid2DCollectionProxy*>(Context.DataInterface);
 
-			// TODO handle mismatched types!
-			for (int32 i = 0; i < AttributeNames.Num(); i++)
-			{
-				int32 FoundIdx = ProxyData->Vars.Find(AttributeNames[i]);
-				check(AttributeNames.Num() == AttributeChannelCount.Num());
-				check(ProxyData->Offsets.Num() == ProxyData->VarComponents.Num());
-				check(ProxyData->Offsets.Num() == ProxyData->Vars.Num());
-				if (ProxyData->Offsets.IsValidIndex(FoundIdx) && AttributeChannelCount[i] == ProxyData->VarComponents[FoundIdx])
-				{
-					ProxyData->AttributeIndices[i] = ProxyData->Offsets[FoundIdx];
-				}
-				else
-				{
-					ProxyData->AttributeIndices[i] = -1; // We may need to protect against this in the hlsl as this might underflow an array lookup if used incorrectly.
-				}
-			}
-		}
+	FGrid2DCollectionRWInstanceData_RenderThread* ProxyData = VFDI->SystemInstancesToProxyData_RT.Find(Context.SystemInstanceID);
+	FGrid2DCollectionRWInstanceData_RenderThread* OriginalProxyData = ProxyData;
+	check(ProxyData);
 
-		SetShaderValue(RHICmdList, ComputeShaderRHI, NumAttributesParam, ProxyData->NumAttributes);
-		SetShaderValue(RHICmdList, ComputeShaderRHI, UnitToUVParam, FVector2f(1.0f) / FVector2f(ProxyData->NumCells));
-		SetShaderValue(RHICmdList, ComputeShaderRHI, NumCellsParam, ProxyData->NumCells);
-		SetShaderValue(RHICmdList, ComputeShaderRHI, CellSizeParam, FVector2f(ProxyData->CellSize));
-		SetShaderValue(RHICmdList, ComputeShaderRHI, WorldBBoxSizeParam, FVector2f(ProxyData->WorldBBoxSize));	// LWC_TODO: Precision loss?
+	if (ProxyData->OtherProxy != nullptr)
+	{
+		FNiagaraDataInterfaceProxyGrid2DCollectionProxy* OtherGrid3DProxy = static_cast<FNiagaraDataInterfaceProxyGrid2DCollectionProxy*>(ProxyData->OtherProxy);
+		ProxyData = OtherGrid3DProxy->SystemInstancesToProxyData_RT.Find(Context.SystemInstanceID);
+		check(ProxyData);
+	}
 
-		SetShaderValueArray(RHICmdList, ComputeShaderRHI, AttributeIndicesParam, ProxyData->AttributeIndices.GetData(), ProxyData->AttributeIndices.Num());
-		FRHISamplerState *SamplerState = TStaticSamplerState<SF_Bilinear, AM_Clamp, AM_Clamp, AM_Clamp>::GetRHI();
-		SetSamplerParameter(RHICmdList, ComputeShaderRHI, SamplerParam, SamplerState);
+	if (OriginalProxyData->AttributeIndices.Num() == 0 && AttributeNames.Num() > 0)
+	{
+		int NumAttrIndices = Align(AttributeNames.Num(), 4);
+		OriginalProxyData->AttributeIndices.SetNumZeroed(NumAttrIndices);
 
-		if (GridParam.IsBound())
+		// TODO handle mismatched types!
+		for (int32 i = 0; i < AttributeNames.Num(); i++)
 		{
-			FRHIShaderResourceView* InputGridBuffer;
-			if (ProxyData->CurrentData != nullptr)
+			int32 FoundIdx = ProxyData->Vars.Find(AttributeNames[i]);
+			check(AttributeNames.Num() == AttributeChannelCount.Num());
+			check(ProxyData->Offsets.Num() == ProxyData->VarComponents.Num());
+			check(ProxyData->Offsets.Num() == ProxyData->Vars.Num());
+			if (ProxyData->Offsets.IsValidIndex(FoundIdx) && AttributeChannelCount[i] == ProxyData->VarComponents[FoundIdx])
 			{
-				InputGridBuffer = ProxyData->CurrentData->GridSRV;
+				OriginalProxyData->AttributeIndices[i] = ProxyData->Offsets[FoundIdx];
 			}
 			else
 			{
-				InputGridBuffer = FNiagaraRenderer::GetDummyTextureReadBuffer2DArray();
+				OriginalProxyData->AttributeIndices[i] = -1; // We may need to protect against this in the hlsl as this might underflow an array lookup if used incorrectly.
 			}
-			SetSRVParameter(RHICmdList, Context.Shader.GetComputeShader(), GridParam, InputGridBuffer);
 		}
+	}
+
+	SetShaderValue(RHICmdList, ComputeShaderRHI, NumAttributesParam, ProxyData->NumAttributes);
+	SetShaderValue(RHICmdList, ComputeShaderRHI, UnitToUVParam, FVector2f(1.0f) / FVector2f(ProxyData->NumCells));
+	SetShaderValue(RHICmdList, ComputeShaderRHI, NumCellsParam, ProxyData->NumCells);
+	SetShaderValue(RHICmdList, ComputeShaderRHI, CellSizeParam, FVector2f(ProxyData->CellSize));
+	SetShaderValue(RHICmdList, ComputeShaderRHI, WorldBBoxSizeParam, FVector2f(ProxyData->WorldBBoxSize));	// LWC_TODO: Precision loss?
 
-		if ( OutputGridParam.IsUAVBound() )
+	SetShaderValueArray(RHICmdList, ComputeShaderRHI, AttributeIndicesParam, OriginalProxyData->AttributeIndices.GetData(), OriginalProxyData->AttributeIndices.Num());
+	FRHISamplerState *SamplerState = TStaticSamplerState<SF_Bilinear, AM_Clamp, AM_Clamp, AM_Clamp>::GetRHI();
+	SetSamplerParameter(RHICmdList, ComputeShaderRHI, SamplerParam, SamplerState);
+
+	if (GridParam.IsBound())
+	{
+		FRHIShaderResourceView* InputGridBuffer;
+		if (ProxyData->CurrentData != nullptr)
 		{
-			FRHIUnorderedAccessView* OutputGridUAV;
-			if (Context.IsOutputStage && ProxyData->DestinationData != nullptr)
-			{
-				OutputGridUAV = ProxyData->DestinationData->GridUAV;
-			}
-			else
-			{
-				OutputGridUAV = Context.ComputeDispatchInterface->GetEmptyUAVFromPool(RHICmdList, PF_R32_FLOAT, ENiagaraEmptyUAVType::Texture2DArray);
-			}
-			RHICmdList.SetUAVParameter(ComputeShaderRHI, OutputGridParam.GetUAVIndex(), OutputGridUAV);
+			InputGridBuffer = ProxyData->CurrentData->GridSRV;
+		}
+		else
+		{
+			InputGridBuffer = FNiagaraRenderer::GetDummyTextureReadBuffer2DArray();
 		}
+		SetSRVParameter(RHICmdList, Context.Shader.GetComputeShader(), GridParam, InputGridBuffer);
 	}
 
-	void Unset(FRHICommandList& RHICmdList, const FNiagaraDataInterfaceSetArgs& Context) const
+	if ( OutputGridParam.IsUAVBound() )
 	{
-		if (OutputGridParam.IsBound())
+		FRHIUnorderedAccessView* OutputGridUAV;
+		if (Context.IsOutputStage && ProxyData->DestinationData != nullptr)
+		{
+			OutputGridUAV = ProxyData->DestinationData->GridUAV;
+		}
+		else
 		{
-			OutputGridParam.UnsetUAV(RHICmdList, Context.Shader.GetComputeShader());
+			OutputGridUAV = Context.ComputeDispatchInterface->GetEmptyUAVFromPool(RHICmdList, PF_R32_FLOAT, ENiagaraEmptyUAVType::Texture2DArray);
 		}
+		RHICmdList.SetUAVParameter(ComputeShaderRHI, OutputGridParam.GetUAVIndex(), OutputGridUAV);
 	}
+}
 
-private:
-	LAYOUT_FIELD(FShaderParameter, NumAttributesParam);
-	LAYOUT_FIELD(FShaderParameter, UnitToUVParam);
-	LAYOUT_FIELD(FShaderParameter, NumCellsParam);
-	LAYOUT_FIELD(FShaderParameter, CellSizeParam);
-	LAYOUT_FIELD(FShaderParameter, WorldBBoxSizeParam);
-
-	LAYOUT_FIELD(FShaderResourceParameter, GridParam);
-	LAYOUT_FIELD(FRWShaderParameter, OutputGridParam);
-	LAYOUT_FIELD(FShaderParameter, AttributeIndicesParam);
-
-	LAYOUT_FIELD(FShaderResourceParameter, SamplerParam);
-	LAYOUT_FIELD(TMemoryImageArray<FName>, AttributeNames);
-	LAYOUT_FIELD(TMemoryImageArray<uint32>, AttributeChannelCount);
-};
+void FNiagaraDataInterfaceParametersCS_Grid2DCollection::Unset(FRHICommandList& RHICmdList, const FNiagaraDataInterfaceSetArgs& Context) const
+{
+	if (OutputGridParam.IsBound())
+	{
+		OutputGridParam.UnsetUAV(RHICmdList, Context.Shader.GetComputeShader());
+	}
+}
 
 IMPLEMENT_TYPE_LAYOUT(FNiagaraDataInterfaceParametersCS_Grid2DCollection);
 
@@ -2320,10 +2308,18 @@ void UNiagaraDataInterfaceGrid2DCollection::GetNumCells(FVectorVMExternalFunctio
 	FNDIOutputParam<int> OutNumCellsX(Context);
 	FNDIOutputParam<int> OutNumCellsY(Context);
 
+	int32 TmpNumCellsX = InstData->NumCells.X;
+	int32 TmpNumCellsY = InstData->NumCells.Y;	
+	if (InstData->OtherInstanceData != nullptr)
+	{
+		TmpNumCellsX = InstData->OtherInstanceData->NumCells.X;
+		TmpNumCellsY = InstData->OtherInstanceData->NumCells.Y;	
+	}
+
 	for (int32 InstanceIdx = 0; InstanceIdx < Context.GetNumInstances(); ++InstanceIdx)
 	{
-		OutNumCellsX.SetAndAdvance(InstData->NumCells.X);
-		OutNumCellsY.SetAndAdvance(InstData->NumCells.Y);
+		OutNumCellsX.SetAndAdvance(TmpNumCellsX);
+		OutNumCellsY.SetAndAdvance(TmpNumCellsY);
 	}
 }
 
@@ -2645,9 +2641,19 @@ void FNiagaraDataInterfaceProxyGrid2DCollectionProxy::ResetData(FRHICommandList&
 
 FIntVector FNiagaraDataInterfaceProxyGrid2DCollectionProxy::GetElementCount(FNiagaraSystemInstanceID SystemInstanceID) const
 {
-	if ( const FGrid2DCollectionRWInstanceData_RenderThread* TargetData = SystemInstancesToProxyData_RT.Find(SystemInstanceID) )
+	if ( const FGrid2DCollectionRWInstanceData_RenderThread* ProxyData = SystemInstancesToProxyData_RT.Find(SystemInstanceID) )
 	{
-		return FIntVector(TargetData->NumCells.X, TargetData->NumCells.Y, 1);
+		// support a grid reader acting as an iteration source
+		if (ProxyData->OtherProxy != nullptr)
+		{
+			FNiagaraDataInterfaceProxyGrid2DCollectionProxy* OtherGrid2DProxy = static_cast<FNiagaraDataInterfaceProxyGrid2DCollectionProxy*>(ProxyData->OtherProxy);
+			const FGrid2DCollectionRWInstanceData_RenderThread* OtherProxyData = OtherGrid2DProxy->SystemInstancesToProxyData_RT.Find(SystemInstanceID);
+			return  FIntVector(OtherProxyData->NumCells.X, OtherProxyData->NumCells.Y, 1);
+		}
+		else
+		{
+			return FIntVector(ProxyData->NumCells.X, ProxyData->NumCells.Y, 1);
+		}		
 	}
 	return FIntVector::ZeroValue;
 }
diff --git a/Engine/Plugins/FX/Niagara/Source/Niagara/Private/NiagaraDataInterfaceGrid2DCollectionReader.cpp b/Engine/Plugins/FX/Niagara/Source/Niagara/Private/NiagaraDataInterfaceGrid2DCollectionReader.cpp
--- a/Engine/Plugins/FX/Niagara/Source/Niagara/Private/NiagaraDataInterfaceGrid2DCollectionReader.cpp
+++ b/Engine/Plugins/FX/Niagara/Source/Niagara/Private/NiagaraDataInterfaceGrid2DCollectionReader.cpp
@@ -1,195 +1,31 @@
 // Copyright Epic Games, Inc. All Rights Reserved.
 #include "NiagaraDataInterfaceGrid2DCollectionReader.h"
-#include "NiagaraComputeExecutionContext.h"
 #include "NiagaraShader.h"
 #include "ShaderParameterUtils.h"
 #include "ClearQuad.h"
 #include "TextureResource.h"
-#include "Engine/Texture2DArray.h"
 #include "NiagaraSystemInstance.h"
 #include "NiagaraRenderer.h"
-#include "Engine/TextureRenderTarget2D.h"
+#include "Engine/VolumeTexture.h"
+#include "Engine/TextureRenderTargetVolume.h"
+#include "NiagaraSettings.h"
+#include "NiagaraConstants.h"
+#include "NiagaraComputeExecutionContext.h"
+#if WITH_EDITOR
+#include "NiagaraGpuComputeDebug.h"
+#endif
 
 #define LOCTEXT_NAMESPACE "NiagaraDataInterfaceGrid2DCollectionReader"
 
-const FString UNiagaraDataInterfaceGrid2DCollectionReader::GridName(TEXT("Grid_"));
-const FString UNiagaraDataInterfaceGrid2DCollectionReader::SamplerName(TEXT("Sampler_"));
-
-// Global VM function names, also used by the shaders code generation methods.
-const FName UNiagaraDataInterfaceGrid2DCollectionReader::GetValueFunctionName("GetGridValue");
-
-const FName UNiagaraDataInterfaceGrid2DCollectionReader::SampleGridFunctionName("SampleGrid");
-
-//#todo(dmp): it would be nice if this class didn't have duplciated code.  It is acting as a proxy for the grid it is reading from.  Refactoring
-// could be nice here
-/*--------------------------------------------------------------------------------------------------------------------------*/
-struct FNiagaraDataInterfaceParametersCS_Grid2DCollectionReader : public FNiagaraDataInterfaceParametersCS
-{
-	DECLARE_TYPE_LAYOUT(FNiagaraDataInterfaceParametersCS_Grid2DCollectionReader, NonVirtual);
-public:
-	void Bind(const FNiagaraDataInterfaceGPUParamInfo& ParameterInfo, const class FShaderParameterMap& ParameterMap)
-	{			
-		NumCellsParam.Bind(ParameterMap, *(UNiagaraDataInterfaceRWBase::NumCellsName + ParameterInfo.DataInterfaceHLSLSymbol));
-		UnitToUVParam.Bind(ParameterMap, *(UNiagaraDataInterfaceRWBase::UnitToUVName + ParameterInfo.DataInterfaceHLSLSymbol));
-		CellSizeParam.Bind(ParameterMap, *(UNiagaraDataInterfaceRWBase::CellSizeName + ParameterInfo.DataInterfaceHLSLSymbol));
-		WorldBBoxSizeParam.Bind(ParameterMap, *(UNiagaraDataInterfaceRWBase::WorldBBoxSizeName + ParameterInfo.DataInterfaceHLSLSymbol));
-		GridParam.Bind(ParameterMap, *(UNiagaraDataInterfaceGrid2DCollectionReader::GridName + ParameterInfo.DataInterfaceHLSLSymbol));		
-		SamplerParam.Bind(ParameterMap, *(UNiagaraDataInterfaceGrid2DCollectionReader::SamplerName + ParameterInfo.DataInterfaceHLSLSymbol));
-	}
-
-	void Set(FRHICommandList& RHICmdList, const FNiagaraDataInterfaceSetArgs& Context) const
-	{
-		check(IsInRenderingThread());
-
-		// Get shader and DI
-		FRHIComputeShader* ComputeShaderRHI = Context.Shader.GetComputeShader();
-				
-		// #todo(dmp): read this from instance data and correct proxy from the reader's proxy
-		FNiagaraDataInterfaceProxyGrid2DCollectionReaderProxy* ReaderDIProxy = static_cast<FNiagaraDataInterfaceProxyGrid2DCollectionReaderProxy*>(Context.DataInterface);
-		
-		FGrid2DCollectionReaderInstanceData_RenderThread* ReaderProxyData = ReaderDIProxy->SystemInstancesToProxyData_RT.Find(Context.SystemInstanceID);
-
-		FGrid2DCollectionRWInstanceData_RenderThread* Grid2DProxyData = nullptr;
-
-		if (ReaderProxyData && ReaderProxyData->ProxyToUse)
-		{
-			Grid2DProxyData = ReaderProxyData->ProxyToUse->SystemInstancesToProxyData_RT.Find(Context.SystemInstanceID);
-		}
-
-		// no proxy data so fill with dummy values
-		if (Grid2DProxyData == nullptr)
-		{			
-			SetShaderValue(RHICmdList, ComputeShaderRHI, NumCellsParam, FIntPoint(0, 0));
-			SetShaderValue(RHICmdList, ComputeShaderRHI, UnitToUVParam, FVector2f::ZeroVector);
-			SetShaderValue(RHICmdList, ComputeShaderRHI, CellSizeParam, FVector2f::ZeroVector);
-			SetShaderValue(RHICmdList, ComputeShaderRHI, WorldBBoxSizeParam, FVector2f::ZeroVector);
-
-			FRHISamplerState* SamplerState = TStaticSamplerState<SF_Bilinear, AM_Clamp, AM_Clamp, AM_Clamp>::GetRHI();
-			SetSamplerParameter(RHICmdList, ComputeShaderRHI, SamplerParam, SamplerState);
-			
-			FRHIShaderResourceView* InputGridBuffer = FNiagaraRenderer::GetDummyTextureReadBuffer2DArray();				
-			SetSRVParameter(RHICmdList, Context.Shader.GetComputeShader(), GridParam, InputGridBuffer);			
-
-			return;
-		}
-
-		check(Grid2DProxyData);
-
-		SetShaderValue(RHICmdList, ComputeShaderRHI, NumCellsParam, Grid2DProxyData->NumCells);
-		SetShaderValue(RHICmdList, ComputeShaderRHI, UnitToUVParam, FVector2f(1.0f) / FVector2f(Grid2DProxyData->NumCells));
-		SetShaderValue(RHICmdList, ComputeShaderRHI, CellSizeParam, FVector2f(Grid2DProxyData->CellSize));
-			
-		SetShaderValue(RHICmdList, ComputeShaderRHI, WorldBBoxSizeParam, FVector2f(Grid2DProxyData->WorldBBoxSize));
-
-		FRHISamplerState *SamplerState = TStaticSamplerState<SF_Bilinear, AM_Clamp, AM_Clamp, AM_Clamp>::GetRHI();
-		SetSamplerParameter(RHICmdList, ComputeShaderRHI, SamplerParam, SamplerState);
-
-		if (GridParam.IsBound())
-		{
-			FRHIShaderResourceView* InputGridBuffer;
-			if (Grid2DProxyData->CurrentData != nullptr)
-			{
-				// FIXME: this shouldn't be necessary, since the Grid2D DI leaves the buffer in the SRVMask state. Confirm and remove the commented out line.
-				//RHICmdList.Transition(FRHITransitionInfo(Grid2DProxyData->CurrentData->GridUAV, ERHIAccess::Unknown, ERHIAccess::SRVMask));
-				InputGridBuffer = Grid2DProxyData->CurrentData->GridSRV;
-			}
-			else
-			{
-				InputGridBuffer = FNiagaraRenderer::GetDummyTextureReadBuffer2DArray();
-			}
-			SetSRVParameter(RHICmdList, Context.Shader.GetComputeShader(), GridParam, InputGridBuffer);
-		}
-	}
-
-	void Unset(FRHICommandList& RHICmdList, const FNiagaraDataInterfaceSetArgs& Context) const 
-	{
-
-	}
-
-private:
-	LAYOUT_FIELD(FShaderParameter, NumCellsParam);
-	LAYOUT_FIELD(FShaderParameter, UnitToUVParam);
-	LAYOUT_FIELD(FShaderParameter, CellSizeParam);
-	LAYOUT_FIELD(FShaderParameter, WorldBBoxSizeParam);
-
-	LAYOUT_FIELD(FShaderResourceParameter, GridParam);
-	
-	LAYOUT_FIELD(FShaderResourceParameter, SamplerParam);
-};
-
-IMPLEMENT_TYPE_LAYOUT(FNiagaraDataInterfaceParametersCS_Grid2DCollectionReader);
-
-IMPLEMENT_NIAGARA_DI_PARAMETER(UNiagaraDataInterfaceGrid2DCollectionReader, FNiagaraDataInterfaceParametersCS_Grid2DCollectionReader);
-
+IMPLEMENT_NIAGARA_DI_PARAMETER(UNiagaraDataInterfaceGrid2DCollectionReader, FNiagaraDataInterfaceParametersCS_Grid2DCollection);
 
 UNiagaraDataInterfaceGrid2DCollectionReader::UNiagaraDataInterfaceGrid2DCollectionReader(FObjectInitializer const& ObjectInitializer)
 	: Super(ObjectInitializer)
 {
-	Proxy.Reset(new FNiagaraDataInterfaceProxyGrid2DCollectionReaderProxy());
-
-	FNiagaraTypeDefinition Def(UObject::StaticClass());	
-}
-
-
-void UNiagaraDataInterfaceGrid2DCollectionReader::PostInitProperties()
-{
-	Super::PostInitProperties();
+	Proxy.Reset(new FNiagaraDataInterfaceProxyGrid2DCollectionProxy());
 
-	//Can we register data interfaces as regular types and fold them into the FNiagaraVariable framework for UI and function calls etc?
-	if (HasAnyFlags(RF_ClassDefaultObject))
-	{
-		ENiagaraTypeRegistryFlags Flags = ENiagaraTypeRegistryFlags::AllowAnyVariable | ENiagaraTypeRegistryFlags::AllowParameter;
-		FNiagaraTypeRegistry::Register(FNiagaraTypeDefinition(GetClass()), Flags);
-	}
-}
-
-void UNiagaraDataInterfaceGrid2DCollectionReader::GetFunctions(TArray<FNiagaraFunctionSignature>& OutFunctions)
-{
-	Super::GetFunctions(OutFunctions);
-
-	{
-		FNiagaraFunctionSignature Sig;
-		Sig.Name = GetValueFunctionName;
-		Sig.Inputs.Add(FNiagaraVariable(FNiagaraTypeDefinition(GetClass()), TEXT("Grid")));
-		Sig.Inputs.Add(FNiagaraVariable(FNiagaraTypeDefinition::GetIntDef(), TEXT("IndexX")));
-		Sig.Inputs.Add(FNiagaraVariable(FNiagaraTypeDefinition::GetIntDef(), TEXT("IndexY")));
-		Sig.Inputs.Add(FNiagaraVariable(FNiagaraTypeDefinition::GetIntDef(), TEXT("AttributeIndex")));
-		Sig.Outputs.Add(FNiagaraVariable(FNiagaraTypeDefinition::GetFloatDef(), TEXT("Value")));
-
-		Sig.bExperimental = true;
-		Sig.bMemberFunction = true;
-		Sig.bRequiresContext = false;
-		OutFunctions.Add(Sig);
-	}
-
-	{
-		FNiagaraFunctionSignature Sig;
-		Sig.Name = SampleGridFunctionName;
-		Sig.Inputs.Add(FNiagaraVariable(FNiagaraTypeDefinition(GetClass()), TEXT("Grid")));
-		Sig.Inputs.Add(FNiagaraVariable(FNiagaraTypeDefinition::GetFloatDef(), TEXT("UnitX")));
-		Sig.Inputs.Add(FNiagaraVariable(FNiagaraTypeDefinition::GetFloatDef(), TEXT("UnitY")));		
-		Sig.Inputs.Add(FNiagaraVariable(FNiagaraTypeDefinition::GetIntDef(), TEXT("AttributeIndex")));
-		Sig.Outputs.Add(FNiagaraVariable(FNiagaraTypeDefinition::GetFloatDef(), TEXT("Value")));
-
-		Sig.bExperimental = true;
-		Sig.bMemberFunction = true;
-		Sig.bRequiresContext = false;
-		OutFunctions.Add(Sig);
-	}
-
-}
-
-// #todo(dmp): expose more CPU functionality
-// #todo(dmp): ideally these would be exposed on the parent class, but we can't bind functions of parent classes but need to work on the interface
-// for sharing an instance data object with the super class
-void UNiagaraDataInterfaceGrid2DCollectionReader::GetVMExternalFunction(const FVMExternalFunctionBindingInfo& BindingInfo, void* InstanceData, FVMExternalFunction &OutFunc)
-{
-	Super::GetVMExternalFunction(BindingInfo, InstanceData, OutFunc);
-	
-	//if (BindingInfo.Name == WorldBBoxSizeFunctionName) { OutFunc = FVMExternalFunction::CreateUObject(this, &UNiagaraDataInterfaceRWBase::EmptyVMFunction); }
-	//else if (BindingInfo.Name == CellSizeFunctionName) { OutFunc = FVMExternalFunction::CreateUObject(this, &UNiagaraDataInterfaceRWBase::EmptyVMFunction); }
-	//else if (BindingInfo.Name == GetValueFunctionName) { OutFunc = FVMExternalFunction::CreateUObject(this, &UNiagaraDataInterfaceRWBase::EmptyVMFunction); }
-	//else if (BindingInfo.Name == SampleGridFunctionName) { OutFunc = FVMExternalFunction::CreateUObject(this, &UNiagaraDataInterfaceRWBase::EmptyVMFunction); }
+	FNiagaraTypeDefinition Def(UTextureRenderTarget::StaticClass());
+	RenderTargetUserParameter.Parameter.SetType(Def);
 }
 
 bool UNiagaraDataInterfaceGrid2DCollectionReader::Equals(const UNiagaraDataInterface* Other) const
@@ -200,71 +36,8 @@ bool UNiagaraDataInterfaceGrid2DCollectionReader::Equals(const UNiagaraDataInter
 	}
 	const UNiagaraDataInterfaceGrid2DCollectionReader* OtherTyped = CastChecked<const UNiagaraDataInterfaceGrid2DCollectionReader>(Other);
 
-	return OtherTyped != nullptr && OtherTyped->EmitterName == EmitterName && OtherTyped->DIName == DIName;		
-}
-
-#if WITH_EDITORONLY_DATA
-void UNiagaraDataInterfaceGrid2DCollectionReader::GetParameterDefinitionHLSL(const FNiagaraDataInterfaceGPUParamInfo& ParamInfo, FString& OutHLSL)
-{
-	Super::GetParameterDefinitionHLSL(ParamInfo, OutHLSL);
-
-	static const TCHAR *FormatDeclarations = TEXT(R"(				
-		Texture2DArray<float> {GridName};
-		SamplerState {SamplerName};
-	
-	)");
-	TMap<FString, FStringFormatArg> ArgsDeclarations = {				
-		{ TEXT("GridName"),    GridName + ParamInfo.DataInterfaceHLSLSymbol },
-		{ TEXT("SamplerName"),    SamplerName + ParamInfo.DataInterfaceHLSLSymbol },		
-	};
-	OutHLSL += FString::Format(FormatDeclarations, ArgsDeclarations);
-}
-
-bool UNiagaraDataInterfaceGrid2DCollectionReader::GetFunctionHLSL(const FNiagaraDataInterfaceGPUParamInfo& ParamInfo, const FNiagaraDataInterfaceGeneratedFunction& FunctionInfo, int FunctionInstanceIndex, FString& OutHLSL)
-{
-	bool ParentRet = Super::GetFunctionHLSL(ParamInfo, FunctionInfo, FunctionInstanceIndex, OutHLSL);
-	if (ParentRet)
-	{
-		return true;
-	} 
-	else if (FunctionInfo.DefinitionName == GetValueFunctionName)
-	{
-		static const TCHAR *FormatBounds = TEXT(R"(
-			void {FunctionName}(int In_IndexX, int In_IndexY, int In_AttributeIndex, out float Out_Val)
-			{
-				Out_Val = {Grid}.Load(int4(In_IndexX, In_IndexY, In_AttributeIndex, 0));
-			}
-		)");
-		TMap<FString, FStringFormatArg> ArgsBounds = {
-			{TEXT("FunctionName"), FunctionInfo.InstanceName},
-			{TEXT("Grid"), GridName + ParamInfo.DataInterfaceHLSLSymbol},
-			{TEXT("NumCellsName"), UNiagaraDataInterfaceRWBase::NumCellsName + ParamInfo.DataInterfaceHLSLSymbol},
-			{TEXT("UnitToUVName"), UNiagaraDataInterfaceRWBase::UnitToUVName + ParamInfo.DataInterfaceHLSLSymbol},
-		};
-		OutHLSL += FString::Format(FormatBounds, ArgsBounds);
-		return true;
-	}
-	else if (FunctionInfo.DefinitionName == SampleGridFunctionName)
-	{
-		static const TCHAR *FormatBounds = TEXT(R"(
-			void {FunctionName}(float In_UnitX, float In_UnitY, int In_AttributeIndex, out float Out_Val)
-			{
-				Out_Val = {Grid}.SampleLevel({SamplerName}, float3(In_UnitX, In_UnitY, In_AttributeIndex), 0);
-			}
-		)");
-		TMap<FString, FStringFormatArg> ArgsBounds = {
-			{TEXT("FunctionName"), FunctionInfo.InstanceName},
-			{TEXT("Grid"), GridName + ParamInfo.DataInterfaceHLSLSymbol},
-			{TEXT("SamplerName"),    SamplerName + ParamInfo.DataInterfaceHLSLSymbol },
-			{TEXT("NumCellsName"), UNiagaraDataInterfaceRWBase::NumCellsName + ParamInfo.DataInterfaceHLSLSymbol},
-			{TEXT("UnitToUVName"), UNiagaraDataInterfaceRWBase::UnitToUVName + ParamInfo.DataInterfaceHLSLSymbol},
-		};
-		OutHLSL += FString::Format(FormatBounds, ArgsBounds);
-		return true;
-	}
-	return false;
+	return OtherTyped != nullptr && OtherTyped->EmitterName == EmitterName && OtherTyped->DIName == DIName;
 }
-#endif
 
 bool UNiagaraDataInterfaceGrid2DCollectionReader::CopyToInternal(UNiagaraDataInterface* Destination) const
 {
@@ -283,13 +56,13 @@ bool UNiagaraDataInterfaceGrid2DCollectionReader::InitPerInstanceData(void* PerI
 {
 	check(Proxy);
 
-	FGrid2DCollectionReaderInstanceData_GameThread* InstanceData = new (PerInstanceData) FGrid2DCollectionReaderInstanceData_GameThread();
+	FGrid2DCollectionRWInstanceData_GameThread* InstanceData = new (PerInstanceData) FGrid2DCollectionRWInstanceData_GameThread();
 	SystemInstancesToProxyData_GT.Emplace(SystemInstance->GetId(), InstanceData);
 
 	TSharedPtr<FNiagaraEmitterInstance, ESPMode::ThreadSafe> RT_EmitterInstance;
 
-	InstanceData->EmitterName = EmitterName;
-	InstanceData->DIName = DIName;
+
+	FNiagaraEmitterInstance* EmitterInstanceToUse = nullptr;
 	for (TSharedPtr<FNiagaraEmitterInstance, ESPMode::ThreadSafe> EmitterInstance : SystemInstance->GetEmitters())
 	{
 		UNiagaraEmitter* Emitter = EmitterInstance->GetCachedEmitter();
@@ -300,68 +73,64 @@ bool UNiagaraDataInterfaceGrid2DCollectionReader::InitPerInstanceData(void* PerI
 
 		if (EmitterName == Emitter->GetUniqueEmitterName())
 		{
-			InstanceData->EmitterInstance = EmitterInstance.Get();			
+			EmitterInstanceToUse = EmitterInstance.Get();
 			break;
 		}
 	}
 
 	// Look for proxy we are going to use
-	FNiagaraDataInterfaceProxyGrid2DCollectionProxy* ProxyToUse = nullptr;
+	FNiagaraDataInterfaceProxy* ProxyToUse = nullptr;
+	UNiagaraDataInterfaceGrid2DCollection* GridInterface = nullptr;
 
-	if (InstanceData->EmitterInstance != nullptr)
+	if (EmitterInstanceToUse != nullptr)
 	{
-		FNiagaraComputeExecutionContext* ExecContext = InstanceData->EmitterInstance->GetGPUContext();
-		if ( (ExecContext != nullptr) && (ExecContext->GPUScript != nullptr) )
+		FNiagaraComputeExecutionContext* ExecContext = EmitterInstanceToUse->GetGPUContext();
+		if ((ExecContext != nullptr) && (ExecContext->GPUScript != nullptr))
 		{
 			const TArray<FNiagaraScriptDataInterfaceCompileInfo>& DataInterfaceInfo = ExecContext->GPUScript->GetVMExecutableData().DataInterfaceInfo;
 			const TArray<UNiagaraDataInterface*>& DataInterfaces = ExecContext->CombinedParamStore.GetDataInterfaces();
 
-			FString FullName = FString("Emitter.") + InstanceData->DIName;
+			FString FullName = FString("Emitter.") + DIName;
 			int Index = 0;
 
 			// #todo(dmp): we are looking at the UObjects that define the DIs here 
+
 			for (UNiagaraDataInterface* Interface : DataInterfaces)
 			{
 				if (DataInterfaceInfo[Index].Name.GetPlainNameString() == FullName)
 				{
-					ProxyToUse = static_cast<FNiagaraDataInterfaceProxyGrid2DCollectionProxy*>(Interface->GetProxy());
+					GridInterface = static_cast<UNiagaraDataInterfaceGrid2DCollection*>(Interface);
+					ProxyToUse = Interface->GetProxy();
 					break;
 				}
 				++Index;
 			}
 		}
-	}
-
-	// Push Updates to Proxy.
-	FNiagaraDataInterfaceProxyGrid2DCollectionReaderProxy* RT_Proxy = GetProxyAs<FNiagaraDataInterfaceProxyGrid2DCollectionReaderProxy>();
-	ENQUEUE_RENDER_COMMAND(FUpdateData)(
-		[RT_Proxy, InstanceID = SystemInstance->GetId(), RT_ProxyToUse=ProxyToUse](FRHICommandListImmediate& RHICmdList)
-	{
-		check(!RT_Proxy->SystemInstancesToProxyData_RT.Contains(InstanceID));
-		FGrid2DCollectionReaderInstanceData_RenderThread* TargetData = &RT_Proxy->SystemInstancesToProxyData_RT.Add(InstanceID);			
-		TargetData->ProxyToUse = RT_ProxyToUse;
-	});
-	
-	return true;
-}
 
+		if (GridInterface && ProxyToUse)
+		{
 
-void UNiagaraDataInterfaceGrid2DCollectionReader::DestroyPerInstanceData(void* PerInstanceData, FNiagaraSystemInstance* SystemInstance)
-{
-	SystemInstancesToProxyData_GT.Remove(SystemInstance->GetId());
-
-	FGrid2DCollectionReaderInstanceData_GameThread* InstanceData = static_cast<FGrid2DCollectionReaderInstanceData_GameThread*>(PerInstanceData);
-
-	InstanceData->~FGrid2DCollectionReaderInstanceData_GameThread();
+			FGrid2DCollectionRWInstanceData_GameThread* Grid2DInstanceData = GridInterface->GetSystemInstancesToProxyData_GT().FindRef(SystemInstance->GetId());
+			InstanceData->OtherDI = GridInterface;
+			InstanceData->OtherInstanceData = Grid2DInstanceData;
 
-	FNiagaraDataInterfaceProxyGrid2DCollectionReaderProxy* RT_Proxy = GetProxyAs<FNiagaraDataInterfaceProxyGrid2DCollectionReaderProxy>();
-	ENQUEUE_RENDER_COMMAND(FNiagaraDIDestroyInstanceData) (
-		[RT_Proxy, InstanceID=SystemInstance->GetId()](FRHICommandListImmediate& CmdList)
-		{
-			//check(ThisProxy->SystemInstancesToProxyData.Contains(InstanceID));
-			RT_Proxy->SystemInstancesToProxyData_RT.Remove(InstanceID);
+			// Push Updates to Proxy.
+			FNiagaraDataInterfaceProxyGrid2DCollectionProxy* RT_Proxy = GetProxyAs<FNiagaraDataInterfaceProxyGrid2DCollectionProxy>();
+			ENQUEUE_RENDER_COMMAND(FUpdateData)(
+				[RT_Proxy, InstanceID = SystemInstance->GetId(), RT_ProxyToUse = ProxyToUse](FRHICommandListImmediate& RHICmdList)
+			{
+				check(!RT_Proxy->SystemInstancesToProxyData_RT.Contains(InstanceID));
+				//check(RT_ProxyToUse->SystemInstancesToProxyData_RT.Contains(InstanceID));
+				//FGrid2DCollectionRWInstanceData_RenderThread* DataToRead = RT_ProxyToUse->SystemInstancesToProxyData_RT.Find(InstanceID);
+
+				FGrid2DCollectionRWInstanceData_RenderThread* TargetData = &RT_Proxy->SystemInstancesToProxyData_RT.Add(InstanceID);
+				TargetData->OtherProxy = RT_ProxyToUse;
+			});
+			return true;
 		}
-	);
+	}
+
+	return false;
 }
 
 void UNiagaraDataInterfaceGrid2DCollectionReader::GetEmitterDependencies(UNiagaraSystem* Asset, TArray<UNiagaraEmitter*>& Dependencies) const
@@ -383,18 +152,57 @@ void UNiagaraDataInterfaceGrid2DCollectionReader::GetEmitterDependencies(UNiagar
 	}
 }
 
-FIntVector FNiagaraDataInterfaceProxyGrid2DCollectionReaderProxy::GetElementCount(FNiagaraSystemInstanceID SystemInstanceID) const
+
+
+void UNiagaraDataInterfaceGrid2DCollectionReader::GetFunctions(TArray<FNiagaraFunctionSignature>& OutFunctions)
 {
-	const auto* ReaderProxyData = SystemInstancesToProxyData_RT.Find(SystemInstanceID);
-	if (ReaderProxyData && ReaderProxyData->ProxyToUse)
+	int32 NumFunctionsBefore = OutFunctions.Num();
+
+	Super::GetFunctions(OutFunctions);
+	TSet<FName> FunctionsToRemove =
+	{
+		ClearCellFunctionName,
+		CopyPreviousToCurrentForCellFunctionName,		
+		SetValueFunctionName,		
+		SetVector4ValueFunctionName,
+		SetVector3ValueFunctionName,
+		SetVector2ValueFunctionName,
+		SetFloatValueFunctionName,
+		GetPreviousVector4ValueFunctionName,
+		SetVectorValueFunctionName,
+		SetVector2DValueFunctionName,
+		SetNumCellsFunctionName,
+		SetNumCellsFunctionName
+	};
+
+	for (int32 ItemIndex = NumFunctionsBefore; ItemIndex < OutFunctions.Num();)
 	{
-		if ( FGrid2DCollectionRWInstanceData_RenderThread* Grid2DProxyData = ReaderProxyData->ProxyToUse->SystemInstancesToProxyData_RT.Find(SystemInstanceID) )
+		if (FunctionsToRemove.Find(OutFunctions[ItemIndex].Name) != nullptr)
 		{
-			return FIntVector(Grid2DProxyData->NumCells.X, Grid2DProxyData->NumCells.Y, 1);
+			OutFunctions.RemoveAtSwap(ItemIndex, 1);
 		}
+		else
+		{
+			++ItemIndex;
+		}
+	}
+}
+
+bool UNiagaraDataInterfaceGrid2DCollectionReader::PerInstanceTickPostSimulate(void* PerInstanceData, FNiagaraSystemInstance* SystemInstance, float DeltaSeconds)
+{
+	FGrid2DCollectionRWInstanceData_GameThread* InstanceData = SystemInstancesToProxyData_GT.FindRef(SystemInstance->GetId());
+
+	// we can run into the case where depending on the ordering of DI initialization, we might have not been able to grab the other grid's InstanceData
+	// in InitPerInstanceData.  If this is the case, we ensure it is correct here.
+	if (InstanceData->OtherInstanceData == nullptr && InstanceData->OtherDI)
+	{
+		UNiagaraDataInterfaceGrid2DCollection* OtherGridDI = static_cast<UNiagaraDataInterfaceGrid2DCollection*>(InstanceData->OtherDI);
+		FGrid2DCollectionRWInstanceData_GameThread* Grid2DInstanceData = OtherGridDI->GetSystemInstancesToProxyData_GT().FindRef(SystemInstance->GetId());
+
+		InstanceData->OtherInstanceData = Grid2DInstanceData;
 	}
 
-	return FIntVector::ZeroValue;
+	return false;
 }
 
 #undef LOCTEXT_NAMESPACE
