Fixed a crash when using fx.SuppressNiagaraSystems 1 with NiagaraDataInterfaceArrayFunctionLibrary

diff --git a/Engine/Plugins/FX/Niagara/Source/Niagara/Private/NiagaraDataInterfaceArrayFunctionLibrary.cpp b/Engine/Plugins/FX/Niagara/Source/Niagara/Private/NiagaraDataInterfaceArrayFunctionLibrary.cpp
--- a/Engine/Plugins/FX/Niagara/Source/Niagara/Private/NiagaraDataInterfaceArrayFunctionLibrary.cpp
+++ b/Engine/Plugins/FX/Niagara/Source/Niagara/Private/NiagaraDataInterfaceArrayFunctionLibrary.cpp
@@ -56,7 +56,13 @@ TArrayType GetNiagaraArrayValue(UNiagaraComponent* NiagaraSystem, FName Override
 
 FNiagaraLWCConverter GetLWCConverter(UNiagaraComponent* NiagaraSystem)
 {
-	FNiagaraSystemInstance* SystemInstance = NiagaraSystem->GetSystemInstanceController()->GetSystemInstance_Unsafe();
+	FNiagaraSystemInstanceControllerPtr SystemInstanceController = NiagaraSystem->GetSystemInstanceController();
+	if (!SystemInstanceController.IsValid())
+	{
+		// the instance controller can be invalid if there is no simulation, for example when fx.SuppressNiagaraSystems is set 
+		return FNiagaraLWCConverter();
+	}
+	FNiagaraSystemInstance* SystemInstance = SystemInstanceController->GetSystemInstance_Unsafe();
 	return SystemInstance->GetLWCConverter();
 }
 
