Fixed a cooker crash due to dangling pointers in FNiagaraEditorModule::Get().GetCachedTypeConversionScripts().
Also prevent type conversion scripts to be retrieved during emitter merge.

diff --git a/Engine/Plugins/FX/Niagara/Source/NiagaraEditor/Private/NiagaraEditorModule.cpp b/Engine/Plugins/FX/Niagara/Source/NiagaraEditor/Private/NiagaraEditorModule.cpp
--- a/Engine/Plugins/FX/Niagara/Source/NiagaraEditor/Private/NiagaraEditorModule.cpp
+++ b/Engine/Plugins/FX/Niagara/Source/NiagaraEditor/Private/NiagaraEditorModule.cpp
@@ -775,6 +775,12 @@ void FNiagaraEditorModule::OnPreExit()
 	ClearObjectPool();
 }
 
+void FNiagaraEditorModule::PostGarbageCollect()
+{
+	// could be that some of the asset data was garbage collected, so we reset the cache
+	InvalidateCachedScriptAssetData();
+}
+
 void FNiagaraEditorModule::StartupModule()
 {
 	bThumbnailRenderersRegistered = false;
@@ -812,6 +818,7 @@ void FNiagaraEditorModule::StartupModule()
 
 	UNiagaraSettings::OnSettingsChanged().AddRaw(this, &FNiagaraEditorModule::OnNiagaraSettingsChangedEvent);
 	FCoreUObjectDelegates::GetPreGarbageCollectDelegate().AddRaw(this, &FNiagaraEditorModule::OnPreGarbageCollection);
+	FCoreUObjectDelegates::GetPostGarbageCollect().AddRaw(this, &FNiagaraEditorModule::PostGarbageCollect);
 	
 	// Any attempt to use GEditor right now will fail as it hasn't been initialized yet. Waiting for post engine init resolves that.
 	FCoreDelegates::OnPostEngineInit.AddRaw(this, &FNiagaraEditorModule::OnPostEngineInit);
@@ -1176,6 +1183,7 @@ void FNiagaraEditorModule::ShutdownModule()
 	UNiagaraSettings::OnSettingsChanged().RemoveAll(this);
 
 	FCoreUObjectDelegates::GetPreGarbageCollectDelegate().RemoveAll(this);
+	FCoreUObjectDelegates::GetPostGarbageCollect().RemoveAll(this);
 	FCoreDelegates::OnPostEngineInit.RemoveAll(this);
 	FCoreDelegates::OnEnginePreExit.RemoveAll(this);
 	
diff --git a/Engine/Plugins/FX/Niagara/Source/NiagaraEditor/Private/ViewModels/Stack/NiagaraStackFunctionInput.cpp b/Engine/Plugins/FX/Niagara/Source/NiagaraEditor/Private/ViewModels/Stack/NiagaraStackFunctionInput.cpp
--- a/Engine/Plugins/FX/Niagara/Source/NiagaraEditor/Private/ViewModels/Stack/NiagaraStackFunctionInput.cpp
+++ b/Engine/Plugins/FX/Niagara/Source/NiagaraEditor/Private/ViewModels/Stack/NiagaraStackFunctionInput.cpp
@@ -1329,8 +1329,12 @@ FName GetNamespaceForUsage(ENiagaraScriptUsage Usage)
 	}
 }
 
-UNiagaraScript* UNiagaraStackFunctionInput::FindConversionScript(const FNiagaraTypeDefinition& FromType, TMap<FNiagaraTypeDefinition, UNiagaraScript*>& ConversionScriptCache) const
+UNiagaraScript* UNiagaraStackFunctionInput::FindConversionScript(const FNiagaraTypeDefinition& FromType, TMap<FNiagaraTypeDefinition, UNiagaraScript*>& ConversionScriptCache, bool bIncludeConversionScripts) const
 {
+	if (bIncludeConversionScripts == false)
+	{
+		return nullptr;
+	}
 	if (UNiagaraScript** CacheEntry = ConversionScriptCache.Find(FromType))
 	{
 		return *CacheEntry;
@@ -1340,7 +1344,7 @@ UNiagaraScript* UNiagaraStackFunctionInput::FindConversionScript(const FNiagaraT
 	return ConversionScriptCache[FromType];
 }
 
-void UNiagaraStackFunctionInput::GetAvailableParameterHandles(TArray<FNiagaraParameterHandle>& AvailableParameterHandles, TMap<FNiagaraVariable, UNiagaraScript*>& AvailableConversionHandles) const
+void UNiagaraStackFunctionInput::GetAvailableParameterHandles(TArray<FNiagaraParameterHandle>& AvailableParameterHandles, TMap<FNiagaraVariable, UNiagaraScript*>& AvailableConversionHandles, bool bIncludeConversionScripts) const
 {
 	TMap<FNiagaraTypeDefinition, UNiagaraScript*> ConversionScriptCache;
 	
@@ -1351,7 +1355,7 @@ void UNiagaraStackFunctionInput::GetAvailableParameterHandles(TArray<FNiagaraPar
 		{
 			AvailableParameterHandles.Add(FNiagaraParameterHandle::CreateEngineParameterHandle(SystemVariable));
 		}
-		else if (UNiagaraScript* ConversionScript = FindConversionScript(SystemVariable.GetType(), ConversionScriptCache))
+		else if (UNiagaraScript* ConversionScript = FindConversionScript(SystemVariable.GetType(), ConversionScriptCache, bIncludeConversionScripts))
 		{
 			AvailableConversionHandles.Add(SystemVariable, ConversionScript);
 		}
@@ -1366,7 +1370,7 @@ void UNiagaraStackFunctionInput::GetAvailableParameterHandles(TArray<FNiagaraPar
 		{
 			AvailableParameterHandles.Add(FNiagaraParameterHandle::CreateEngineParameterHandle(ExposedVar));
 		}
-		else if (UNiagaraScript* ConversionScript = FindConversionScript(ExposedVar.GetType(), ConversionScriptCache))
+		else if (UNiagaraScript* ConversionScript = FindConversionScript(ExposedVar.GetType(), ConversionScriptCache, bIncludeConversionScripts))
 		{
 			AvailableConversionHandles.Add(ExposedVar, ConversionScript);
 		}
@@ -1460,7 +1464,7 @@ void UNiagaraStackFunctionInput::GetAvailableParameterHandles(TArray<FNiagaraPar
 									}
 								}
 							}
-							else if (UNiagaraScript* ConversionScript = FindConversionScript(HistoryVariable.GetType(), ConversionScriptCache))
+							else if (UNiagaraScript* ConversionScript = FindConversionScript(HistoryVariable.GetType(), ConversionScriptCache, bIncludeConversionScripts))
 							{
 								AvailableConversionHandles.Add(HistoryVariable, ConversionScript);
 							}
@@ -1501,7 +1505,7 @@ void UNiagaraStackFunctionInput::GetAvailableParameterHandles(TArray<FNiagaraPar
 				{
 					AvailableParameterHandles.AddUnique(FNiagaraParameterHandle(CollectionParam.GetName()));
 				}
-				else if (UNiagaraScript* ConversionScript = FindConversionScript(CollectionParam.GetType(), ConversionScriptCache))
+				else if (UNiagaraScript* ConversionScript = FindConversionScript(CollectionParam.GetType(), ConversionScriptCache, bIncludeConversionScripts))
 				{
 					AvailableConversionHandles.Add(CollectionParam, ConversionScript);
 				}
@@ -2800,7 +2804,7 @@ void UNiagaraStackFunctionInput::GetDefaultLinkedHandleOrLinkedFunctionFromDefau
 		// If there are a chain of linked values use the first one that's available, otherwise just use the last one.
 		TArray<FNiagaraParameterHandle> AvailableHandles;
 		TMap<FNiagaraVariable, UNiagaraScript*> ConversionHandles;
-		GetAvailableParameterHandles(AvailableHandles, ConversionHandles);
+		GetAvailableParameterHandles(AvailableHandles, ConversionHandles, false);
 		for (FLinkedHandleOrFunctionNode& LinkedValue : LinkedValues)
 		{
 			if (LinkedValue.LinkedFunctionCallNode.IsValid() ||
diff --git a/Engine/Plugins/FX/Niagara/Source/NiagaraEditor/Public/NiagaraEditorModule.h b/Engine/Plugins/FX/Niagara/Source/NiagaraEditor/Public/NiagaraEditorModule.h
--- a/Engine/Plugins/FX/Niagara/Source/NiagaraEditor/Public/NiagaraEditorModule.h
+++ b/Engine/Plugins/FX/Niagara/Source/NiagaraEditor/Public/NiagaraEditorModule.h
@@ -273,6 +273,7 @@ private:
 	void OnDeviceProfileManagerUpdated();
 	void OnPreviewPlatformChanged();
 	void OnPreExit();
+	void PostGarbageCollect();
 
 	/** FGCObject interface */
 	virtual void AddReferencedObjects( FReferenceCollector& Collector ) override;
diff --git a/Engine/Plugins/FX/Niagara/Source/NiagaraEditor/Public/ViewModels/Stack/NiagaraStackFunctionInput.h b/Engine/Plugins/FX/Niagara/Source/NiagaraEditor/Public/ViewModels/Stack/NiagaraStackFunctionInput.h
--- a/Engine/Plugins/FX/Niagara/Source/NiagaraEditor/Public/ViewModels/Stack/NiagaraStackFunctionInput.h
+++ b/Engine/Plugins/FX/Niagara/Source/NiagaraEditor/Public/ViewModels/Stack/NiagaraStackFunctionInput.h
@@ -130,8 +130,8 @@ public:
 	/** Sets the value of this input to a linked parameter handle. */
 	void SetLinkedValueHandle(const FNiagaraParameterHandle& InParameterHandle);
 
-	/** Gets the current set of available parameter handles which can be assigned to this input. */
-	void GetAvailableParameterHandles(TArray<FNiagaraParameterHandle>& AvailableParameterHandles, TMap<FNiagaraVariable, UNiagaraScript*>& AvailableConversionHandles) const;
+	/** Gets the current set of available parameter handles which can be assigned to this input. Optionally returns possible conversion scripts. */
+	void GetAvailableParameterHandles(TArray<FNiagaraParameterHandle>& AvailableParameterHandles, TMap<FNiagaraVariable, UNiagaraScript*>& AvailableConversionHandles, bool bIncludeConversionScripts = true) const;
 
 	/** Gets the function node form the script graph if the current value mode is DefaultFunction. */
 	UNiagaraNodeFunctionCall* GetDefaultFunctionNode() const;
@@ -368,7 +368,7 @@ private:
 
 	void GetCurrentChangeIds(FGuid& OutOwningGraphChangeId, FGuid& OutFunctionGraphChangeId) const;
 
-	UNiagaraScript* FindConversionScript(const FNiagaraTypeDefinition& FromType, TMap<FNiagaraTypeDefinition, UNiagaraScript*>& ConversionScriptCache) const;
+	UNiagaraScript* FindConversionScript(const FNiagaraTypeDefinition& FromType, TMap<FNiagaraTypeDefinition, UNiagaraScript*>& ConversionScriptCache, bool bIncludeConversionScripts) const;
 
 private:
 	/** The module function call which owns this input entry. NOTE: This input might not be an input to the module function
