- Clean up multi GPU Niagara code
- Add cross GPU transfer support to Niagara and Cascade
- Add method to Compute Dispatch Interface to allow data interfaces to send resources for transfer
- Fix missing persistent ID buffer transfers with AFR
- Fix PostRenderOpaque transfer location to ensure we capture count updates

diff --git a/Engine/Plugins/FX/Niagara/Source/Niagara/Classes/NiagaraDataInterfaceRenderTarget2D.h b/Engine/Plugins/FX/Niagara/Source/Niagara/Classes/NiagaraDataInterfaceRenderTarget2D.h
--- a/Engine/Plugins/FX/Niagara/Source/Niagara/Classes/NiagaraDataInterfaceRenderTarget2D.h
+++ b/Engine/Plugins/FX/Niagara/Source/Niagara/Classes/NiagaraDataInterfaceRenderTarget2D.h
@@ -45,7 +45,9 @@ struct FRenderTarget2DRWInstanceData_RenderThread
 	FIntPoint Size = FIntPoint(EForceInit::ForceInitToZero);
 	ENiagaraMipMapGeneration MipMapGeneration = ENiagaraMipMapGeneration::Disabled;
 	ENiagaraMipMapGenerationType MipMapGenerationType = ENiagaraMipMapGenerationType::Linear;
-	bool bWasWrittenTo = false;
+	bool bRebuildMips = false;
+	bool bReadThisFrame = false;
+	bool bWroteThisFrame = false;
 
 	FSamplerStateRHIRef SamplerStateRHI;
 	FTexture2DRHIRef TextureRHI;
diff --git a/Engine/Plugins/FX/Niagara/Source/Niagara/Classes/NiagaraDataInterfaceRenderTarget2DArray.h b/Engine/Plugins/FX/Niagara/Source/Niagara/Classes/NiagaraDataInterfaceRenderTarget2DArray.h
--- a/Engine/Plugins/FX/Niagara/Source/Niagara/Classes/NiagaraDataInterfaceRenderTarget2DArray.h
+++ b/Engine/Plugins/FX/Niagara/Source/Niagara/Classes/NiagaraDataInterfaceRenderTarget2DArray.h
@@ -40,7 +40,9 @@ struct FRenderTarget2DArrayRWInstanceData_RenderThread
 	}
 
 	FIntVector Size = FIntVector(EForceInit::ForceInitToZero);
-	
+	bool bWroteThisFrame = false;
+	bool bReadThisFrame = false;
+
 	FSamplerStateRHIRef SamplerStateRHI;
 	FTexture2DArrayRHIRef TextureRHI;
 	FUnorderedAccessViewRHIRef UnorderedAccessViewRHI;
diff --git a/Engine/Plugins/FX/Niagara/Source/Niagara/Classes/NiagaraDataInterfaceRenderTargetCube.h b/Engine/Plugins/FX/Niagara/Source/Niagara/Classes/NiagaraDataInterfaceRenderTargetCube.h
--- a/Engine/Plugins/FX/Niagara/Source/Niagara/Classes/NiagaraDataInterfaceRenderTargetCube.h
+++ b/Engine/Plugins/FX/Niagara/Source/Niagara/Classes/NiagaraDataInterfaceRenderTargetCube.h
@@ -39,7 +39,9 @@ struct FRenderTargetCubeRWInstanceData_RenderThread
 	}
 
 	int Size = 0;
-	
+	bool bWroteThisFrame = false;
+	bool bReadThisFrame = false;
+
 	FSamplerStateRHIRef SamplerStateRHI;
 	FTextureCubeRHIRef TextureRHI;
 	FUnorderedAccessViewRHIRef UnorderedAccessViewRHI;
diff --git a/Engine/Plugins/FX/Niagara/Source/Niagara/Classes/NiagaraDataInterfaceRenderTargetVolume.h b/Engine/Plugins/FX/Niagara/Source/Niagara/Classes/NiagaraDataInterfaceRenderTargetVolume.h
--- a/Engine/Plugins/FX/Niagara/Source/Niagara/Classes/NiagaraDataInterfaceRenderTargetVolume.h
+++ b/Engine/Plugins/FX/Niagara/Source/Niagara/Classes/NiagaraDataInterfaceRenderTargetVolume.h
@@ -39,7 +39,9 @@ struct FRenderTargetVolumeRWInstanceData_RenderThread
 	}
 
 	FIntVector Size = FIntVector(EForceInit::ForceInitToZero);
-	
+	bool bWroteThisFrame = false;
+	bool bReadThisFrame = false;
+
 	FSamplerStateRHIRef SamplerStateRHI;
 	FTexture3DRHIRef TextureRHI;
 	FUnorderedAccessViewRHIRef UnorderedAccessViewRHI;
diff --git a/Engine/Plugins/FX/Niagara/Source/Niagara/Private/NiagaraDataInterfaceIntRenderTarget2D.cpp b/Engine/Plugins/FX/Niagara/Source/Niagara/Private/NiagaraDataInterfaceIntRenderTarget2D.cpp
--- a/Engine/Plugins/FX/Niagara/Source/Niagara/Private/NiagaraDataInterfaceIntRenderTarget2D.cpp
+++ b/Engine/Plugins/FX/Niagara/Source/Niagara/Private/NiagaraDataInterfaceIntRenderTarget2D.cpp
@@ -74,6 +74,8 @@ struct FNDIIntRenderTarget2DInstanceData_RenderThread
 	bool bPreviewRenderTarget = false;
 	FVector2D PreviewDisplayRange = FVector2D(0.0f, 255.0f);
 #endif
+	bool bWroteThisFrame = false;
+	bool bReadThisFrame = false;
 
 	FSamplerStateRHIRef SamplerStateRHI;
 	FTextureRHIRef TextureRHI;
@@ -105,21 +107,35 @@ struct FNDIIntRenderTarget2DProxy : public FNiagaraDataInterfaceProxyRW
 
 	virtual void PostSimulate(FRHICommandList& RHICmdList, const FNiagaraDataInterfaceArgs& Context) override
 	{
-		if (FNDIIntRenderTarget2DInstanceData_RenderThread* InstanceData = SystemInstancesToProxyData_RT.Find(Context.SystemInstanceID))
+		FNDIIntRenderTarget2DInstanceData_RenderThread* InstanceData = SystemInstancesToProxyData_RT.Find(Context.SystemInstanceID);
+		if ( InstanceData == nullptr )
 		{
+			return;
+		}
+
+		// We only need to transfer this frame if it was written to.
+		// If also read then we need to notify that the texture is important for the simulation
+		// We also assume the texture is important for rendering, without discovering renderer bindings we don't really know
+		if (InstanceData->bWroteThisFrame)
+		{
+			Context.ComputeDispatchInterface->MultiGPUResourceModified(RHICmdList, InstanceData->TextureRHI, InstanceData->bReadThisFrame, true);
+		}
+
+		InstanceData->bWroteThisFrame = true;
+		InstanceData->bReadThisFrame = true;
+
 #if NIAGARA_COMPUTEDEBUG_ENABLED && WITH_EDITORONLY_DATA
-			if (InstanceData->bPreviewRenderTarget)
+		if (InstanceData->bPreviewRenderTarget)
+		{
+			if (FNiagaraGpuComputeDebug* GpuComputeDebug = Context.ComputeDispatchInterface->GetGpuComputeDebug())
 			{
-				if (FNiagaraGpuComputeDebug* GpuComputeDebug = Context.ComputeDispatchInterface->GetGpuComputeDebug())
+				if (FRHITexture* RHITexture = InstanceData->TextureRHI)
 				{
-					if (FRHITexture* RHITexture = InstanceData->TextureRHI)
-					{
-						GpuComputeDebug->AddTexture(RHICmdList, Context.SystemInstanceID, SourceDIName, RHITexture, InstanceData->PreviewDisplayRange);
-					}
+					GpuComputeDebug->AddTexture(RHICmdList, Context.SystemInstanceID, SourceDIName, RHITexture, InstanceData->PreviewDisplayRange);
 				}
 			}
-#endif
 		}
+#endif
 	}
 
 	virtual FIntVector GetElementCount(FNiagaraSystemInstanceID SystemInstanceID) const override
@@ -161,6 +177,9 @@ public:
 	
 		if (TextureUAVParam.IsUAVBound())
 		{
+			InstanceData->bWroteThisFrame = true;
+			InstanceData->bReadThisFrame = true;
+
 			FRHIUnorderedAccessView* OutputUAV = InstanceData->UnorderedAccessViewRHI;
 			if (OutputUAV != nullptr)
 			{
diff --git a/Engine/Plugins/FX/Niagara/Source/Niagara/Private/NiagaraDataInterfaceRenderTarget2D.cpp b/Engine/Plugins/FX/Niagara/Source/Niagara/Private/NiagaraDataInterfaceRenderTarget2D.cpp
--- a/Engine/Plugins/FX/Niagara/Source/Niagara/Private/NiagaraDataInterfaceRenderTarget2D.cpp
+++ b/Engine/Plugins/FX/Niagara/Source/Niagara/Private/NiagaraDataInterfaceRenderTarget2D.cpp
@@ -198,7 +198,8 @@ public:
 				// FIXME: this transition needs to happen in FNiagaraDataInterfaceProxyRenderTarget2DProxy::PreStage so it doesn't break up the overlap group,
 				// but for some reason it stops working if I move it in there.
 				RHICmdList.Transition(FRHITransitionInfo(OutputUAV, ERHIAccess::Unknown, ERHIAccess::UAVCompute));
-				ProxyData->bWasWrittenTo = true;
+				ProxyData->bRebuildMips = true;
+				ProxyData->bWroteThisFrame = true;
 			}
 			else
 			{
@@ -210,6 +211,8 @@ public:
 
 		if ( InputParam.IsBound() )
 		{
+			ProxyData->bReadThisFrame = true;
+
 			FRHITexture* TextureRHI = ProxyData->TextureRHI;
 			if (!ensureMsgf(!OutputParam.IsUAVBound(), TEXT("NiagaraDIRenderTarget2D(%s) is bound as both read & write, read will be ignored."), *Context.DataInterface->SourceDIName.ToString()))
 			{
@@ -908,9 +911,9 @@ void FNiagaraDataInterfaceProxyRenderTarget2DProxy::PostStage(FRHICommandList& R
 {
 	if (FRenderTarget2DRWInstanceData_RenderThread* ProxyData = SystemInstancesToProxyData_RT.Find(Context.SystemInstanceID))
 	{
-		if (ProxyData->bWasWrittenTo && (ProxyData->MipMapGeneration == ENiagaraMipMapGeneration::PostStage))
+		if (ProxyData->bRebuildMips && (ProxyData->MipMapGeneration == ENiagaraMipMapGeneration::PostStage))
 		{
-			ProxyData->bWasWrittenTo = false;
+			ProxyData->bRebuildMips = false;
 			FNiagaraGpuProfileScope GpuProfileScope(RHICmdList, Context, GNiagaraRenderTarget2DGenerateMipsName);
 			NiagaraGenerateMips::GenerateMips(RHICmdList, ProxyData->TextureRHI, ProxyData->MipMapGenerationType);
 		}
@@ -919,28 +922,42 @@ void FNiagaraDataInterfaceProxyRenderTarget2DProxy::PostStage(FRHICommandList& R
 
 void FNiagaraDataInterfaceProxyRenderTarget2DProxy::PostSimulate(FRHICommandList& RHICmdList, const FNiagaraDataInterfaceArgs& Context)
 {
-	if (FRenderTarget2DRWInstanceData_RenderThread* ProxyData = SystemInstancesToProxyData_RT.Find(Context.SystemInstanceID))
+	FRenderTarget2DRWInstanceData_RenderThread* ProxyData = SystemInstancesToProxyData_RT.Find(Context.SystemInstanceID);
+	if (ProxyData == nullptr)
 	{
-		if (ProxyData->bWasWrittenTo && (ProxyData->MipMapGeneration == ENiagaraMipMapGeneration::PostSimulate))
-		{
-			ProxyData->bWasWrittenTo = false;
-			FNiagaraGpuProfileScope GpuProfileScope(RHICmdList, Context, GNiagaraRenderTarget2DGenerateMipsName);
-			NiagaraGenerateMips::GenerateMips(RHICmdList, ProxyData->TextureRHI, ProxyData->MipMapGenerationType);
-		}
+		return;
+	}
+
+	if (ProxyData->bRebuildMips && (ProxyData->MipMapGeneration == ENiagaraMipMapGeneration::PostSimulate))
+	{
+		FNiagaraGpuProfileScope GpuProfileScope(RHICmdList, Context, GNiagaraRenderTarget2DGenerateMipsName);
+		NiagaraGenerateMips::GenerateMips(RHICmdList, ProxyData->TextureRHI, ProxyData->MipMapGenerationType);
+	}
+
+	// We only need to transfer this frame if it was written to.
+	// If also read then we need to notify that the texture is important for the simulation
+	// We also assume the texture is important for rendering, without discovering renderer bindings we don't really know
+	if (ProxyData->bWroteThisFrame)
+	{
+		Context.ComputeDispatchInterface->MultiGPUResourceModified(RHICmdList, ProxyData->TextureRHI, ProxyData->bReadThisFrame, true);
+	}
+
+	ProxyData->bRebuildMips = false;
+	ProxyData->bReadThisFrame = false;
+	ProxyData->bWroteThisFrame = false;
 
 #if NIAGARA_COMPUTEDEBUG_ENABLED && WITH_EDITORONLY_DATA
-		if (ProxyData->bPreviewTexture)
+	if (ProxyData->bPreviewTexture)
+	{
+		if (FNiagaraGpuComputeDebug* GpuComputeDebug = Context.ComputeDispatchInterface->GetGpuComputeDebug())
 		{
-			if (FNiagaraGpuComputeDebug* GpuComputeDebug = Context.ComputeDispatchInterface->GetGpuComputeDebug())
+			if (FRHITexture* RHITexture = ProxyData->TextureRHI)
 			{
-				if (FRHITexture* RHITexture = ProxyData->TextureRHI)
-				{
-					GpuComputeDebug->AddTexture(RHICmdList, Context.SystemInstanceID, SourceDIName, RHITexture);
-				}
+				GpuComputeDebug->AddTexture(RHICmdList, Context.SystemInstanceID, SourceDIName, RHITexture);
 			}
 		}
-#endif
 	}
+#endif
 }
 
 FIntVector FNiagaraDataInterfaceProxyRenderTarget2DProxy::GetElementCount(FNiagaraSystemInstanceID SystemInstanceID) const
diff --git a/Engine/Plugins/FX/Niagara/Source/Niagara/Private/NiagaraDataInterfaceRenderTarget2DArray.cpp b/Engine/Plugins/FX/Niagara/Source/Niagara/Private/NiagaraDataInterfaceRenderTarget2DArray.cpp
--- a/Engine/Plugins/FX/Niagara/Source/Niagara/Private/NiagaraDataInterfaceRenderTarget2DArray.cpp
+++ b/Engine/Plugins/FX/Niagara/Source/Niagara/Private/NiagaraDataInterfaceRenderTarget2DArray.cpp
@@ -79,6 +79,7 @@ public:
 			if (OutputUAV)
 			{
 				RHICmdList.Transition(FRHITransitionInfo(OutputUAV, ERHIAccess::Unknown, ERHIAccess::UAVCompute));
+				ProxyData->bWroteThisFrame = true;
 			}
 			else
 			{
@@ -90,6 +91,8 @@ public:
 
 		if (InputParam.IsBound())
 		{
+			ProxyData->bReadThisFrame = true;
+
 			FRHITexture* TextureRHI = ProxyData->TextureRHI;
 			if (!ensureMsgf(!OutputParam.IsUAVBound(), TEXT("NiagaraDIRenderTarget2DArray(%s) is bound as both read & write, read will be ignored."), *Context.DataInterface->SourceDIName.ToString()))
 			{
@@ -754,10 +757,25 @@ bool UNiagaraDataInterfaceRenderTarget2DArray::PerInstanceTickPostSimulate(void*
 
 void FNiagaraDataInterfaceProxyRenderTarget2DArrayProxy::PostSimulate(FRHICommandList& RHICmdList, const FNiagaraDataInterfaceArgs& Context)
 {
-#if NIAGARA_COMPUTEDEBUG_ENABLED && WITH_EDITORONLY_DATA
 	FRenderTarget2DArrayRWInstanceData_RenderThread* ProxyData = SystemInstancesToProxyData_RT.Find(Context.SystemInstanceID);
+	if (ProxyData == nullptr)
+	{
+		return;
+	}
 
-	if (ProxyData && ProxyData->bPreviewTexture)
+	// We only need to transfer this frame if it was written to.
+	// If also read then we need to notify that the texture is important for the simulation
+	// We also assume the texture is important for rendering, without discovering renderer bindings we don't really know
+	if (ProxyData->bWroteThisFrame)
+	{
+		Context.ComputeDispatchInterface->MultiGPUResourceModified(RHICmdList, ProxyData->TextureRHI, ProxyData->bReadThisFrame, true);
+	}
+
+	ProxyData->bReadThisFrame = false;
+	ProxyData->bWroteThisFrame = false;
+
+#if NIAGARA_COMPUTEDEBUG_ENABLED && WITH_EDITORONLY_DATA
+	if (ProxyData->bPreviewTexture)
 	{
 		if (FNiagaraGpuComputeDebug* GpuComputeDebug = Context.ComputeDispatchInterface->GetGpuComputeDebug())
 		{
diff --git a/Engine/Plugins/FX/Niagara/Source/Niagara/Private/NiagaraDataInterfaceRenderTargetCube.cpp b/Engine/Plugins/FX/Niagara/Source/Niagara/Private/NiagaraDataInterfaceRenderTargetCube.cpp
--- a/Engine/Plugins/FX/Niagara/Source/Niagara/Private/NiagaraDataInterfaceRenderTargetCube.cpp
+++ b/Engine/Plugins/FX/Niagara/Source/Niagara/Private/NiagaraDataInterfaceRenderTargetCube.cpp
@@ -77,6 +77,7 @@ public:
 			if (OutputUAV)
 			{
 				RHICmdList.Transition(FRHITransitionInfo(OutputUAV, ERHIAccess::Unknown, ERHIAccess::UAVCompute));
+				ProxyData->bWroteThisFrame = true;
 			}
 			else
 			{
@@ -87,6 +88,8 @@ public:
 
 		if (InputParam.IsBound())
 		{
+			ProxyData->bReadThisFrame = true;
+
 			FRHITexture* TextureRHI = ProxyData->TextureRHI;
 			if (!ensureMsgf(!OutputParam.IsUAVBound(), TEXT("NiagaraDIRenderTargetCube(%s) is bound as both read & write, read will be ignored."), *Context.DataInterface->SourceDIName.ToString()))
 			{
@@ -731,10 +734,25 @@ bool UNiagaraDataInterfaceRenderTargetCube::PerInstanceTickPostSimulate(void* Pe
 
 void FNiagaraDataInterfaceProxyRenderTargetCubeProxy::PostSimulate(FRHICommandList& RHICmdList, const FNiagaraDataInterfaceArgs& Context)
 {
-#if NIAGARA_COMPUTEDEBUG_ENABLED && WITH_EDITORONLY_DATA
 	FRenderTargetCubeRWInstanceData_RenderThread* ProxyData = SystemInstancesToProxyData_RT.Find(Context.SystemInstanceID);
+	if (ProxyData == nullptr)
+	{
+		return;
+	}
 
-	if (ProxyData && ProxyData->bPreviewTexture && ProxyData->TextureRHI.IsValid())
+	// We only need to transfer this frame if it was written to.
+	// If also read then we need to notify that the texture is important for the simulation
+	// We also assume the texture is important for rendering, without discovering renderer bindings we don't really know
+	if (ProxyData->bWroteThisFrame)
+	{
+		Context.ComputeDispatchInterface->MultiGPUResourceModified(RHICmdList, ProxyData->TextureRHI, ProxyData->bReadThisFrame, true);
+	}
+
+	ProxyData->bReadThisFrame = false;
+	ProxyData->bWroteThisFrame = false;
+
+#if NIAGARA_COMPUTEDEBUG_ENABLED && WITH_EDITORONLY_DATA
+	if (ProxyData->bPreviewTexture && ProxyData->TextureRHI.IsValid())
 	{
 		if (FNiagaraGpuComputeDebug* GpuComputeDebug = Context.ComputeDispatchInterface->GetGpuComputeDebug())
 		{
diff --git a/Engine/Plugins/FX/Niagara/Source/Niagara/Private/NiagaraDataInterfaceRenderTargetVolume.cpp b/Engine/Plugins/FX/Niagara/Source/Niagara/Private/NiagaraDataInterfaceRenderTargetVolume.cpp
--- a/Engine/Plugins/FX/Niagara/Source/Niagara/Private/NiagaraDataInterfaceRenderTargetVolume.cpp
+++ b/Engine/Plugins/FX/Niagara/Source/Niagara/Private/NiagaraDataInterfaceRenderTargetVolume.cpp
@@ -78,6 +78,7 @@ public:
 			if (OutputUAV)
 			{
 				RHICmdList.Transition(FRHITransitionInfo(OutputUAV, ERHIAccess::Unknown, ERHIAccess::UAVCompute));
+				ProxyData->bWroteThisFrame = true;
 			}
 			else
 			{
@@ -88,6 +89,8 @@ public:
 
 		if (InputParam.IsBound())
 		{
+			ProxyData->bReadThisFrame = true;
+
 			FRHITexture* TextureRHI = ProxyData->TextureRHI;
 			if (!ensureMsgf(!OutputParam.IsUAVBound(), TEXT("NiagaraDIRenderTargetVolume(%s) is bound as both read & write, read will be ignored."), *Context.DataInterface->SourceDIName.ToString()))
 			{
@@ -753,10 +756,25 @@ bool UNiagaraDataInterfaceRenderTargetVolume::PerInstanceTickPostSimulate(void*
 
 void FNiagaraDataInterfaceProxyRenderTargetVolumeProxy::PostSimulate(FRHICommandList& RHICmdList, const FNiagaraDataInterfaceArgs& Context)
 {
-#if NIAGARA_COMPUTEDEBUG_ENABLED && WITH_EDITORONLY_DATA
 	FRenderTargetVolumeRWInstanceData_RenderThread* ProxyData = SystemInstancesToProxyData_RT.Find(Context.SystemInstanceID);
+	if (ProxyData == nullptr)
+	{
+		return;
+	}
 
-	if (ProxyData && ProxyData->bPreviewTexture && ProxyData->TextureRHI.IsValid())
+	// We only need to transfer this frame if it was written to.
+	// If also read then we need to notify that the texture is important for the simulation
+	// We also assume the texture is important for rendering, without discovering renderer bindings we don't really know
+	if (ProxyData->bWroteThisFrame)
+	{
+		Context.ComputeDispatchInterface->MultiGPUResourceModified(RHICmdList, ProxyData->TextureRHI, ProxyData->bReadThisFrame, true);
+	}
+
+	ProxyData->bReadThisFrame = false;
+	ProxyData->bWroteThisFrame = false;
+
+#if NIAGARA_COMPUTEDEBUG_ENABLED && WITH_EDITORONLY_DATA
+	if (ProxyData->bPreviewTexture && ProxyData->TextureRHI.IsValid())
 	{
 		if (FNiagaraGpuComputeDebug* GpuComputeDebug = Context.ComputeDispatchInterface->GetGpuComputeDebug())
 		{
diff --git a/Engine/Plugins/FX/Niagara/Source/Niagara/Private/NiagaraGpuComputeDispatch.cpp b/Engine/Plugins/FX/Niagara/Source/Niagara/Private/NiagaraGpuComputeDispatch.cpp
--- a/Engine/Plugins/FX/Niagara/Source/Niagara/Private/NiagaraGpuComputeDispatch.cpp
+++ b/Engine/Plugins/FX/Niagara/Source/Niagara/Private/NiagaraGpuComputeDispatch.cpp
@@ -53,10 +53,6 @@ DECLARE_GPU_STAT_NAMED(NiagaraGPUSorting, TEXT("Niagara GPU sorting"));
 
 uint32 FNiagaraComputeExecutionContext::TickCounter = 0;
 
-#if WITH_MGPU
-const FName FNiagaraGpuComputeDispatch::TemporalEffectName("FNiagaraGpuComputeDispatch");
-#endif // WITH_MGPU
-
 int32 GNiagaraGpuSubmitCommandHint = 0;
 static FAutoConsoleVariableRef CVarNiagaraGpuSubmitCommandHint(
 	TEXT("fx.NiagaraGpuSubmitCommandHint"),
@@ -87,6 +83,11 @@ const FName FNiagaraGpuComputeDispatch::Name(TEXT("FNiagaraGpuComputeDispatch"))
 
 namespace FNiagaraGpuComputeDispatchLocal
 {
+#if WITH_MGPU
+	const FName TemporalEffectBuffersName("FNiagaraGpuComputeDispatch_Buffers");
+	const FName TemporalEffectTexturesName("FNiagaraGpuComputeDispatch_Textures");
+#endif // WITH_MGPU
+
 	int32 GTickFlushMaxQueuedFrames = 10;
 	static FAutoConsoleVariableRef CVarNiagaraTickFlushMaxQueuedFrames(
 		TEXT("fx.Niagara.Batcher.TickFlush.MaxQueuedFrames"),
@@ -182,14 +183,9 @@ FNiagaraGpuComputeDispatch::FNiagaraGpuComputeDispatch(ERHIFeatureLevel::Type In
 				[this](FRHICommandListImmediate& RHICmdList)
 				{
 					GPUInstanceCounterManager.UpdateDrawIndirectBuffers(this, RHICmdList, ENiagaraGPUCountUpdatePhase::PreOpaque);
-#if WITH_MGPU
-					// For PreInitViews we actually need to broadcast here and not in ExecuteAll since
-					// this is the last place the instance count buffer is written to.
-					if (StageToBroadcastTemporalEffect == ENiagaraGpuComputeTickStage::PreInitViews)
-					{
-						BroadcastTemporalEffect(RHICmdList);
-					}
-#endif // WITH_MGPU
+				#if WITH_MGPU
+					TransferMultiGPUBufers(RHICmdList, ENiagaraGpuComputeTickStage::PreInitViews);
+				#endif // WITH_MGPU
 				}
 			);
 
@@ -198,6 +194,9 @@ FNiagaraGpuComputeDispatch::FNiagaraGpuComputeDispatch(ERHIFeatureLevel::Type In
 				[this](FRHICommandListImmediate& RHICmdList)
 				{
 					GPUInstanceCounterManager.UpdateDrawIndirectBuffers(this, RHICmdList, ENiagaraGPUCountUpdatePhase::PostOpaque);
+				#if WITH_MGPU
+					TransferMultiGPUBufers(RHICmdList, ENiagaraGpuComputeTickStage::PostOpaqueRender);
+				#endif // WITH_MGPU
 				}
 			);
 		}
@@ -1167,9 +1166,6 @@ void FNiagaraGpuComputeDispatch::PrepareTicksForProxy(FRHICommandListImmediate&
 			FinalBuffer->SetNumInstances(ComputeContext->CurrentNumInstances_RT);
 			FinalBuffer->SetGPUDataReadyStage(ComputeProxy->GetComputeTickStage());
 			ComputeContext->SetDataToRender(FinalBuffer);
-#if WITH_MGPU
-			AddTemporalEffectBuffers(FinalBuffer);
-#endif // WITH_MGPU
 		}
 		// When low latency translucency is enabled we can setup the final buffer / final count here.
 		// This will allow our mesh processor commands to pickup the data for the same frame.
@@ -1183,29 +1179,6 @@ void FNiagaraGpuComputeDispatch::PrepareTicksForProxy(FRHICommandListImmediate&
 			ComputeContext->SetTranslucentDataToRender(FinalBuffer);
 		}
 	}
-
-#if WITH_MGPU
-	if (GNumAlternateFrameRenderingGroups > 1)
-	{
-		StageToBroadcastTemporalEffect = ENiagaraGpuComputeTickStage::Last;
-		while (StageToBroadcastTemporalEffect > ENiagaraGpuComputeTickStage::First && !DispatchListPerStage[static_cast<int32>(StageToBroadcastTemporalEffect)].HasWork())
-		{
-			StageToBroadcastTemporalEffect = static_cast<ENiagaraGpuComputeTickStage::Type>(static_cast<int32>(StageToBroadcastTemporalEffect) - 1);
-		}
-
-		StageToWaitForTemporalEffect = ENiagaraGpuComputeTickStage::First;
-		// If we're going to write to the instance count buffer after PreInitViews then
-		// that needs to be the wait stage, regardless of whether or not we're ticking
-		// anything in that stage.
-		if (!GPUInstanceCounterManager.HasEntriesPendingFree())
-		{
-			while (StageToWaitForTemporalEffect < StageToBroadcastTemporalEffect && !DispatchListPerStage[static_cast<int32>(StageToWaitForTemporalEffect)].HasWork())
-			{
-				StageToWaitForTemporalEffect = static_cast<ENiagaraGpuComputeTickStage::Type>(static_cast<int32>(StageToWaitForTemporalEffect) + 1);
-			}
-		}
-	}
-#endif // WITH_MGPU
 }
 
 void FNiagaraGpuComputeDispatch::PrepareAllTicks(FRHICommandListImmediate& RHICmdList)
@@ -1222,10 +1195,7 @@ void FNiagaraGpuComputeDispatch::PrepareAllTicks(FRHICommandListImmediate& RHICm
 void FNiagaraGpuComputeDispatch::ExecuteTicks(FRHICommandList& RHICmdList, FRHIUniformBuffer* ViewUniformBuffer, ENiagaraGpuComputeTickStage::Type TickStage)
 {
 #if WITH_MGPU
-	if (GNumAlternateFrameRenderingGroups > 1 && TickStage == StageToWaitForTemporalEffect)
-	{
-		RHICmdList.WaitForTemporalEffect(TemporalEffectName);
-	}
+	WaitForMultiGPUBuffers(RHICmdList, TickStage);
 #endif // WITH_MGPU
 
 	// Anything to execute for this stage
@@ -1390,10 +1360,25 @@ void FNiagaraGpuComputeDispatch::ExecuteTicks(FRHICommandList& RHICmdList, FRHIU
 						ComputeContext->SetDataToRender(CurrentData);
 
 #if WITH_MGPU
-						AddTemporalEffectBuffers(CurrentData);
+						if (bAFREnabled)
+						{
+							AddAFRBuffer(CurrentData->GetGPUBufferFloat().Buffer);
+							AddAFRBuffer(CurrentData->GetGPUBufferHalf().Buffer);
+							AddAFRBuffer(CurrentData->GetGPUBufferInt().Buffer);
+							if (ComputeContext->MainDataSet->RequiresPersistentIDs())
+							{
+								AddAFRBuffer(ComputeContext->MainDataSet->GetGPUFreeIDs().Buffer);
+							}
+						}
+						if (bCrossGPUTransferEnabled)
+						{
+							AddCrossGPUTransfer(RHICmdList, CurrentData->GetGPUBufferFloat().Buffer);
+							AddCrossGPUTransfer(RHICmdList, CurrentData->GetGPUBufferHalf().Buffer);
+							AddCrossGPUTransfer(RHICmdList, CurrentData->GetGPUBufferInt().Buffer);
+						}
 #endif // WITH_MGPU
 					}
-						// If this is not the final tick of the final stage we need set our temporary buffer for data interfaces, etc, that may snoop from CurrentData
+					// If this is not the final tick of the final stage we need set our temporary buffer for data interfaces, etc, that may snoop from CurrentData
 					else
 					{
 						ComputeContext->MainDataSet->CurrentData = FinalSimStageDataBuffer;
@@ -1411,15 +1396,6 @@ void FNiagaraGpuComputeDispatch::ExecuteTicks(FRHICommandList& RHICmdList, FRHIU
 		RHICmdList.Transition(TransitionsAfter);
 		TransitionsAfter.Reset();
 
-#if WITH_MGPU
-		// For PreInitViews we actually need to broadcast after UpdateDrawIndirectBuffer
-		// since this is the last place the instance count buffer is written to.
-		if (TickStage == StageToBroadcastTemporalEffect && StageToBroadcastTemporalEffect != ENiagaraGpuComputeTickStage::PreInitViews)
-		{
-			BroadcastTemporalEffect(RHICmdList);
-		}
-#endif // WITH_MGPU
-
 		// Update free IDs
 		if (DispatchGroup.FreeIDUpdates.Num() > 0)
 		{
@@ -1657,6 +1633,12 @@ void FNiagaraGpuComputeDispatch::PreInitViews(FRDGBuilder& GraphBuilder, bool bA
 #if WITH_EDITOR
 	bRaisedWarningThisFrame = false;
 #endif
+#if WITH_MGPU
+	bAFREnabled = GNumAlternateFrameRenderingGroups > 1;
+	bCrossGPUTransferEnabled = !bAFREnabled && (GNumExplicitGPUsForRendering > 1);
+	StageToTransferGPUBuffers = ENiagaraGpuComputeTickStage::Last;
+	StageToWaitForGPUTransfers = ENiagaraGpuComputeTickStage::First;
+#endif
 
 	GpuReadbackManagerPtr->Tick();
 #if NIAGARA_COMPUTEDEBUG_ENABLED
@@ -1691,6 +1673,10 @@ void FNiagaraGpuComputeDispatch::PreInitViews(FRDGBuilder& GraphBuilder, bool bA
 			UpdateInstanceCountManager(GraphBuilder.RHICmdList);
 			PrepareAllTicks(GraphBuilder.RHICmdList);
 		
+		#if WITH_MGPU
+			CalculateCrossGPUTransferLocation();
+		#endif
+
 #if RHI_RAYTRACING
 			RayTracingHelper->BeginFrame(GraphBuilder.RHICmdList, HasRayTracingScene());
 #endif
@@ -1739,13 +1725,17 @@ void FNiagaraGpuComputeDispatch::PostInitViews(FRDGBuilder& GraphBuilder, TArray
 			PassParameters,
 			ERDGPassFlags::Compute | ERDGPassFlags::NeverCull,
 			[this, PassParameters, Views](FRHICommandListImmediate& RHICmdList)
-		{
-			NiagaraSceneTextures = PassParameters;
-			ON_SCOPE_EXIT{ NiagaraSceneTextures = nullptr; };
+			{
+				NiagaraSceneTextures = PassParameters;
+				ON_SCOPE_EXIT{ NiagaraSceneTextures = nullptr; };
 
-			FRHIUniformBuffer* ViewUniformBuffer = GetReferenceViewUniformBuffer(Views);
-			ExecuteTicks(RHICmdList, ViewUniformBuffer, ENiagaraGpuComputeTickStage::PostInitViews);
-		});
+				FRHIUniformBuffer* ViewUniformBuffer = GetReferenceViewUniformBuffer(Views);
+				ExecuteTicks(RHICmdList, ViewUniformBuffer, ENiagaraGpuComputeTickStage::PostInitViews);
+			#if WITH_MGPU
+				TransferMultiGPUBufers(RHICmdList, ENiagaraGpuComputeTickStage::PostInitViews);
+			#endif // WITH_MGPU
+			}
+		);
 	}
 }
 
@@ -2245,33 +2235,127 @@ void FNiagaraGpuComputeDispatch::DrawSceneDebug_RenderThread(class FRDGBuilder&
 }
 
 #if WITH_MGPU
-void FNiagaraGpuComputeDispatch::AddTemporalEffectBuffers(FNiagaraDataBuffer* FinalData)
+void FNiagaraGpuComputeDispatch::MultiGPUResourceModified(FRHICommandList& RHICmdList, FRHIBuffer* Buffer, bool bRequiredForSimulation, bool bRequiredForRendering) const
+{
+	if (bAFREnabled && bRequiredForSimulation)
+	{
+		const_cast<FNiagaraGpuComputeDispatch*>(this)->AddAFRBuffer(Buffer);
+	}
+	if (bCrossGPUTransferEnabled && bRequiredForRendering)
+	{
+		const_cast<FNiagaraGpuComputeDispatch*>(this)->AddCrossGPUTransfer(RHICmdList, Buffer);
+	}
+}
+
+void FNiagaraGpuComputeDispatch::MultiGPUResourceModified(FRHICommandList& RHICmdList, FRHITexture* Texture, bool bRequiredForSimulation, bool bRequiredForRendering) const
 {
-	if (GNumAlternateFrameRenderingGroups > 1)
+	if (bAFREnabled && bRequiredForSimulation)
+	{
+		if (Texture)
+		{
+			const_cast<FNiagaraGpuComputeDispatch*>(this)->AFRTextures.Add(Texture);
+		}
+	}
+	if (bCrossGPUTransferEnabled && bRequiredForRendering)
 	{
-		for (FRHIBuffer* Buffer : { FinalData->GetGPUBufferFloat().Buffer, FinalData->GetGPUBufferInt().Buffer, FinalData->GetGPUBufferHalf().Buffer })
+		const bool bPullData = false;
+		const bool bLockStep = true;
+
+		const FRHIGPUMask GPUMask = RHICmdList.GetGPUMask();
+		for (uint32 GPUIndex : FRHIGPUMask::All())
 		{
-			if (Buffer)
+			if (!GPUMask.Contains(GPUIndex))
 			{
-				TemporalEffectBuffers.Add(Buffer);
+				const_cast<FNiagaraGpuComputeDispatch*>(this)->CrossGPUTransferBuffers.Emplace(Texture, GPUMask.GetFirstIndex(), GPUIndex, bPullData, bLockStep);
 			}
 		}
 	}
 }
 
-void FNiagaraGpuComputeDispatch::BroadcastTemporalEffect(FRHICommandList& RHICmdList)
+void FNiagaraGpuComputeDispatch::AddAFRBuffer(FRHIBuffer* Buffer)
+{
+	check(bAFREnabled);
+	if (Buffer)
+	{
+		AFRBuffers.Add(Buffer);
+	}
+}
+
+void FNiagaraGpuComputeDispatch::AddCrossGPUTransfer(FRHICommandList& RHICmdList, FRHIBuffer* Buffer)
 {
-	if (GNumAlternateFrameRenderingGroups > 1)
+	check(bCrossGPUTransferEnabled);
+	if (Buffer)
 	{
-		if (GPUInstanceCounterManager.GetInstanceCountBuffer().Buffer)
+		const bool bPullData = false;
+		const bool bLockStep = true;
+
+		const FRHIGPUMask GPUMask = RHICmdList.GetGPUMask();
+		for (uint32 GPUIndex : FRHIGPUMask::All())
 		{
-			TemporalEffectBuffers.Add(GPUInstanceCounterManager.GetInstanceCountBuffer().Buffer);
+			if (!GPUMask.Contains(GPUIndex))
+			{
+				CrossGPUTransferBuffers.Emplace(Buffer, GPUMask.GetFirstIndex(), GPUIndex, bPullData, bLockStep);
+			}
 		}
-		if (TemporalEffectBuffers.Num())
+	}
+}
+
+void FNiagaraGpuComputeDispatch::CalculateCrossGPUTransferLocation()
+{
+	StageToTransferGPUBuffers = ENiagaraGpuComputeTickStage::Last;
+	while (StageToTransferGPUBuffers > ENiagaraGpuComputeTickStage::First && !DispatchListPerStage[static_cast<int32>(StageToTransferGPUBuffers)].HasWork())
+	{
+		StageToTransferGPUBuffers = static_cast<ENiagaraGpuComputeTickStage::Type>(static_cast<int32>(StageToTransferGPUBuffers) - 1);
+	}
+
+	StageToWaitForGPUTransfers = ENiagaraGpuComputeTickStage::First;
+	// If we're going to write to the instance count buffer after PreInitViews then
+	// that needs to be the wait stage, regardless of whether or not we're ticking
+	// anything in that stage.
+	if (!GPUInstanceCounterManager.HasEntriesPendingFree())
+	{
+		while (StageToWaitForGPUTransfers < StageToTransferGPUBuffers && !DispatchListPerStage[static_cast<int32>(StageToWaitForGPUTransfers)].HasWork())
 		{
-			RHICmdList.BroadcastTemporalEffect(TemporalEffectName, TemporalEffectBuffers);
-			TemporalEffectBuffers.Reset();
+			StageToWaitForGPUTransfers = static_cast<ENiagaraGpuComputeTickStage::Type>(static_cast<int32>(StageToWaitForGPUTransfers) + 1);
 		}
 	}
 }
+
+void FNiagaraGpuComputeDispatch::TransferMultiGPUBufers(FRHICommandList& RHICmdList, ENiagaraGpuComputeTickStage::Type TickStage)
+{
+	if (StageToTransferGPUBuffers != TickStage)
+	{
+		return;
+	}
+
+	// Transfer buffers for AFR rendering
+	if (AFRBuffers.Num())
+	{
+		AddAFRBuffer(GPUInstanceCounterManager.GetInstanceCountBuffer().Buffer);
+		RHICmdList.BroadcastTemporalEffect(FNiagaraGpuComputeDispatchLocal::TemporalEffectBuffersName, AFRBuffers);
+		AFRBuffers.Reset();
+	}
+	if (AFRTextures.Num())
+	{
+		RHICmdList.BroadcastTemporalEffect(FNiagaraGpuComputeDispatchLocal::TemporalEffectTexturesName, AFRTextures);
+		AFRTextures.Reset();
+	}
+
+	// Transfer buffers for cross GPU rendering
+	if (CrossGPUTransferBuffers.Num())
+	{
+		AddCrossGPUTransfer(RHICmdList, GPUInstanceCounterManager.GetInstanceCountBuffer().Buffer);
+		RHICmdList.TransferResources(CrossGPUTransferBuffers);
+		CrossGPUTransferBuffers.Reset();
+	}
+}
+
+void FNiagaraGpuComputeDispatch::WaitForMultiGPUBuffers(FRHICommandList& RHICmdList, ENiagaraGpuComputeTickStage::Type TickStage)
+{
+	if (StageToWaitForGPUTransfers == TickStage)
+	{
+		RHICmdList.WaitForTemporalEffect(FNiagaraGpuComputeDispatchLocal::TemporalEffectBuffersName);
+		RHICmdList.WaitForTemporalEffect(FNiagaraGpuComputeDispatchLocal::TemporalEffectTexturesName);
+	}
+}
 #endif // WITH_MGPU
diff --git a/Engine/Plugins/FX/Niagara/Source/Niagara/Private/NiagaraGpuComputeDispatch.h b/Engine/Plugins/FX/Niagara/Source/Niagara/Private/NiagaraGpuComputeDispatch.h
--- a/Engine/Plugins/FX/Niagara/Source/Niagara/Private/NiagaraGpuComputeDispatch.h
+++ b/Engine/Plugins/FX/Niagara/Source/Niagara/Private/NiagaraGpuComputeDispatch.h
@@ -135,6 +135,11 @@ public:
 	virtual void ReleaseGPURayTracedCollisionGroup(int32 CollisionGroup) override;
 #endif
 
+#if WITH_MGPU
+	virtual void MultiGPUResourceModified(FRHICommandList& RHICmdList, FRHIBuffer* Buffer, bool bRequiredForSimulation, bool bRequiredForRendering) const override;
+	virtual void MultiGPUResourceModified(FRHICommandList& RHICmdList, FRHITexture* Texture, bool bRequiredForSimulation, bool bRequiredForRendering) const override;
+#endif
+
 	void OnPrimitiveGPUSceneInstancesAllocated();
 	void OnPrimitiveGPUSceneInstancesFreed();
 
@@ -231,13 +236,22 @@ private:
 	TArray<FDebugReadbackInfo> GpuDebugReadbackInfos;
 
 #if WITH_MGPU
-	static const FName TemporalEffectName;
-	TArray<FRHIBuffer*> TemporalEffectBuffers;
-	ENiagaraGpuComputeTickStage::Type StageToWaitForTemporalEffect = ENiagaraGpuComputeTickStage::First;
-	ENiagaraGpuComputeTickStage::Type StageToBroadcastTemporalEffect = ENiagaraGpuComputeTickStage::First;
+	ENiagaraGpuComputeTickStage::Type StageToTransferGPUBuffers = ENiagaraGpuComputeTickStage::First;
+	ENiagaraGpuComputeTickStage::Type StageToWaitForGPUTransfers = ENiagaraGpuComputeTickStage::First;
+
+	bool bAFREnabled = false;
+	TArray<FRHIBuffer*> AFRBuffers;
+	TArray<FRHITexture*> AFRTextures;
+
+	bool bCrossGPUTransferEnabled = false;
+	TArray<FTransferResourceParams> CrossGPUTransferBuffers;
+
+	void AddAFRBuffer(FRHIBuffer* Buffer);
+	void AddCrossGPUTransfer(FRHICommandList& RHICmdList, FRHIBuffer* Buffer);
 
-	void AddTemporalEffectBuffers(FNiagaraDataBuffer* FinalData);
-	void BroadcastTemporalEffect(FRHICommandList& RHICmdList);
+	void CalculateCrossGPUTransferLocation();
+	void TransferMultiGPUBufers(FRHICommandList& RHICmdList, ENiagaraGpuComputeTickStage::Type TickStage);
+	void WaitForMultiGPUBuffers(FRHICommandList& RHICmdList, ENiagaraGpuComputeTickStage::Type TickStage);
 #endif // WITH_MGPU
 
 	/** Pool of free GPU ray traced collision groups. */
diff --git a/Engine/Plugins/FX/Niagara/Source/Niagara/Public/NiagaraGpuComputeDispatchInterface.h b/Engine/Plugins/FX/Niagara/Source/Niagara/Public/NiagaraGpuComputeDispatchInterface.h
--- a/Engine/Plugins/FX/Niagara/Source/Niagara/Public/NiagaraGpuComputeDispatchInterface.h
+++ b/Engine/Plugins/FX/Niagara/Source/Niagara/Public/NiagaraGpuComputeDispatchInterface.h
@@ -92,6 +92,19 @@ public:
 	virtual void ReleaseGPURayTracedCollisionGroup(int32 CollisionGroup) = 0;
 #endif
 
+#if WITH_MGPU
+	/**
+	Notify that a GPU resource was modified that will impact MultiGPU rendering.
+	bRequiredForSimulation	- When true we require this resource for simulation passes
+	bRequiredForRendering	- When true we require this resource for rendering passes
+	*/
+	virtual void MultiGPUResourceModified(FRHICommandList& RHICmdList, FRHIBuffer* Buffer, bool bRequiredForSimulation, bool bRequiredForRendering) const = 0;
+	virtual void MultiGPUResourceModified(FRHICommandList& RHICmdList, FRHITexture* Texture, bool bRequiredForSimulation, bool bRequiredForRendering) const = 0;
+#else
+	FORCEINLINE void MultiGPUResourceModified(FRHICommandList& RHICmdList, FRHIBuffer* Buffer, bool bRequiredForSimulation, bool bRequiredForRendering) const {}
+	FORCEINLINE void MultiGPUResourceModified(FRHICommandList& RHICmdList, FRHITexture* Texture, bool bRequiredForSimulation, bool bRequiredForRendering) const {}
+#endif
+
 protected:
 	EShaderPlatform							ShaderPlatform;
 	ERHIFeatureLevel::Type					FeatureLevel;
diff --git a/Engine/Source/Runtime/Engine/Private/Particles/ParticleGpuSimulation.cpp b/Engine/Source/Runtime/Engine/Private/Particles/ParticleGpuSimulation.cpp
--- a/Engine/Source/Runtime/Engine/Private/Particles/ParticleGpuSimulation.cpp
+++ b/Engine/Source/Runtime/Engine/Private/Particles/ParticleGpuSimulation.cpp
@@ -4733,6 +4733,29 @@ void FFXSystem::SimulateGPUParticles(
 		TemporalEffectTextures.Add(CurrentStateRenderTargets[0]);
 		TemporalEffectTextures.Add(CurrentStateRenderTargets[1]);
 	}
+
+	TArray<FTransferResourceParams, TFixedAllocator<4>> CrossGPUTransferResources;
+	const bool bCrossTransferEnabled = GNumAlternateFrameRenderingGroups == 1 && GNumExplicitGPUsForRendering > 1;
+	auto AddCrossGPUTransferResource =
+		[&](FRHITexture* TextureToTransfer)
+		{
+			const bool bPullData = false;
+			const bool bLockStep = true;
+
+			const FRHIGPUMask GPUMask = RHICmdList.GetGPUMask();
+			for (uint32 GPUIndex : FRHIGPUMask::All())
+			{
+				if (!GPUMask.Contains(GPUIndex))
+				{
+					CrossGPUTransferResources.Emplace(TextureToTransfer, GPUMask.GetFirstIndex(), GPUIndex, bPullData, bLockStep);
+				}
+			}
+		};
+	if (bCrossTransferEnabled)
+	{
+		AddCrossGPUTransferResource(CurrentStateRenderTargets[0]);
+		AddCrossGPUTransferResource(CurrentStateRenderTargets[1]);
+	}
 #endif
 
 	{
@@ -5051,6 +5074,13 @@ void FFXSystem::SimulateGPUParticles(
 #endif
 			}
 		}
+#if WITH_MGPU
+		if (bCrossTransferEnabled)
+		{
+			AddCrossGPUTransferResource(ParticleSimulationResources->RenderAttributesTexture.TextureTargetRHI);
+			AddCrossGPUTransferResource(ParticleSimulationResources->SimulationAttributesTexture.TextureTargetRHI);
+		}
+#endif
 		RHICmdList.EndUpdateMultiFrameResource(ParticleSimulationResources->RenderAttributesTexture.TextureTargetRHI);
 		RHICmdList.EndUpdateMultiFrameResource(ParticleSimulationResources->SimulationAttributesTexture.TextureTargetRHI);
 	}
@@ -5121,9 +5151,16 @@ void FFXSystem::SimulateGPUParticles(
 	// using the textures until we're done, however for a potential AFR performance
 	// boost, we can put the broadcast back to where it was, and use a third buffer for
 	// temporary extra particle simulation work.
-	if (GNumAlternateFrameRenderingGroups > 1 && Phase == PhaseToBroadcastTemporalEffect)
+	if (Phase == PhaseToBroadcastTemporalEffect)
 	{
-		RHICmdList.BroadcastTemporalEffect(TemporalEffectName, TemporalEffectTextures);
+		if (GNumAlternateFrameRenderingGroups > 1)
+		{
+			RHICmdList.BroadcastTemporalEffect(TemporalEffectName, TemporalEffectTextures);
+		}
+		if (CrossGPUTransferResources.Num() > 0)
+		{
+			RHICmdList.TransferResources(CrossGPUTransferResources);
+		}
 	}
 #endif
 
