Niagara - Make the script variable id fixup smarter so that it better handles duplicates subscribed to parameter definitions, and also change ensures to warnings and include more information for debugging issues in content.

diff --git a/Engine/Plugins/FX/Niagara/Source/NiagaraEditor/Private/NiagaraGraph.cpp b/Engine/Plugins/FX/Niagara/Source/NiagaraEditor/Private/NiagaraGraph.cpp
--- a/Engine/Plugins/FX/Niagara/Source/NiagaraEditor/Private/NiagaraGraph.cpp
+++ b/Engine/Plugins/FX/Niagara/Source/NiagaraEditor/Private/NiagaraGraph.cpp
@@ -406,21 +406,57 @@ void UNiagaraGraph::PostLoad()
 		// Make sure that the variable persistent ids are unique, otherwise pin allocation for static switches will be inconsistent.
 		TArray<TObjectPtr<UNiagaraScriptVariable>> ScriptVariables;
 		VariableToScriptVariable.GenerateValueArray(ScriptVariables);
-		TSet<FGuid> SeenGuids;
+		TMap<FGuid, UNiagaraScriptVariable*> VariableGuidToScriptVariable;
 		for (TObjectPtr<UNiagaraScriptVariable>& ScriptVariable : ScriptVariables)
 		{
 			if (ScriptVariable != nullptr)
 			{
-				bool bAlreadySeen = false;
-				SeenGuids.Add(ScriptVariable->Metadata.GetVariableGuid(), &bAlreadySeen);
-				if (bAlreadySeen)
+				ScriptVariable->ConditionalPostLoad();
+				UNiagaraScriptVariable** MatchingScriptVariablePtr = VariableGuidToScriptVariable.Find(ScriptVariable->Metadata.GetVariableGuid());
+				if (MatchingScriptVariablePtr == nullptr)
 				{
-					if (ensureMsgf(ScriptVariable->GetIsSubscribedToParameterDefinitions() == false, TEXT("Duplicate id found for script variable %s controlled by parameter definitions."), *ScriptVariable->GetPathName()))
+					VariableGuidToScriptVariable.Add(ScriptVariable->Metadata.GetVariableGuid(), ScriptVariable);
+				}
+				else
+				{
+					UNiagaraScriptVariable* MatchingScriptVariable = *MatchingScriptVariablePtr;
+					if (ScriptVariable->GetIsSubscribedToParameterDefinitions() && MatchingScriptVariablePtr->GetIsSubscribedToParameterDefinitions())
 					{
-						ScriptVariable->Metadata.SetVariableGuid(UNiagaraScriptVariable::GenerateStableGuid(ScriptVariable));
-						bAlreadySeen = false;
-						SeenGuids.Add(ScriptVariable->Metadata.GetVariableGuid(), &bAlreadySeen);
-						ensureMsgf(bAlreadySeen == false, TEXT("Failed to generate a unique stable guid for script variable %s."), *ScriptVariable->GetPathName());
+						// Both of the script variables with duplicate ids are controlled by parameter definitions so issue a warning because neither will be updated.
+						UE_LOG(LogNiagaraEditor, Warning, TEXT("Duplicate ids found for script variables which are both subscribed to parameter definitions.\nScript Variable 1 Name: %s Type: %s Path: %s\nScript Variable 2 Name: %s Type: %s Path: %s"),
+							*MatchingScriptVariable->Variable.GetName().ToString(), *MatchingScriptVariable->Variable.GetType().GetName(), *MatchingScriptVariable->GetPathName(),
+							*ScriptVariable->Variable.GetName().ToString(), *ScriptVariable->Variable.GetType().GetName(), *ScriptVariable->GetPathName());
+					}
+					else
+					{
+						// Remove the duplicated entry and regenerate the ids for entries not subscribed to parameter definitions.
+						VariableGuidToScriptVariable.Remove(ScriptVariable->Metadata.GetVariableGuid());
+						if (ScriptVariable->GetIsSubscribedToParameterDefinitions() == false)
+						{
+							ScriptVariable->Metadata.SetVariableGuid(UNiagaraScriptVariable::GenerateStableGuid(ScriptVariable));
+							if (VariableGuidToScriptVariable.Contains(ScriptVariable->Metadata.GetVariableGuid()))
+							{
+								UE_LOG(LogNiagaraEditor, Warning, TEXT("Failed to generate a stable unique variable guid for script variable. Name: %s Type: %s Path: %s"),
+									*ScriptVariable->Variable.GetName().ToString(), *ScriptVariable->Variable.GetType().GetName(), *ScriptVariable->GetPathName());
+							}
+							else
+							{
+								VariableGuidToScriptVariable.Add(ScriptVariable->Metadata.GetVariableGuid(), ScriptVariable);
+							}
+						}
+						if (MatchingScriptVariable->GetIsSubscribedToParameterDefinitions() == false)
+						{
+							MatchingScriptVariable->Metadata.SetVariableGuid(UNiagaraScriptVariable::GenerateStableGuid(MatchingScriptVariable));
+							if (VariableGuidToScriptVariable.Contains(MatchingScriptVariable->Metadata.GetVariableGuid()))
+							{
+								UE_LOG(LogNiagaraEditor, Warning, TEXT("Failed to generate a stable unique variable guid for script variable. Name: %s Type: %s Path: %s"),
+									*MatchingScriptVariable->Variable.GetName().ToString(), *MatchingScriptVariable->Variable.GetType().GetName(), *MatchingScriptVariable->GetPathName());
+							}
+							else
+							{
+								VariableGuidToScriptVariable.Add(MatchingScriptVariable->Metadata.GetVariableGuid(), MatchingScriptVariable);
+							}
+						}
 					}
 				}
 			}
