Niagara - Make the script variable id fixup smarter so that it better handles duplicates subscribed to parameter definitions, and also change ensures to warnings and include more information for debugging issues in content. Updated with correct file version.

diff --git a/Engine/Plugins/FX/Niagara/Source/NiagaraEditor/Private/NiagaraGraph.cpp b/Engine/Plugins/FX/Niagara/Source/NiagaraEditor/Private/NiagaraGraph.cpp
--- a/Engine/Plugins/FX/Niagara/Source/NiagaraEditor/Private/NiagaraGraph.cpp
+++ b/Engine/Plugins/FX/Niagara/Source/NiagaraEditor/Private/NiagaraGraph.cpp
@@ -420,7 +420,7 @@ void UNiagaraGraph::PostLoad()
 				else
 				{
 					UNiagaraScriptVariable* MatchingScriptVariable = *MatchingScriptVariablePtr;
-					if (ScriptVariable->GetIsSubscribedToParameterDefinitions() && MatchingScriptVariablePtr->GetIsSubscribedToParameterDefinitions())
+					if (ScriptVariable->GetIsSubscribedToParameterDefinitions() && MatchingScriptVariable->GetIsSubscribedToParameterDefinitions())
 					{
 						// Both of the script variables with duplicate ids are controlled by parameter definitions so issue a warning because neither will be updated.
 						UE_LOG(LogNiagaraEditor, Warning, TEXT("Duplicate ids found for script variables which are both subscribed to parameter definitions.\nScript Variable 1 Name: %s Type: %s Path: %s\nScript Variable 2 Name: %s Type: %s Path: %s"),
