Fix bug with name aliasing datainterface sourced parametes in simulation stage scripts where dependency checking would fail.

diff --git a/Engine/Plugins/FX/Niagara/Source/Niagara/Private/NiagaraScript.cpp b/Engine/Plugins/FX/Niagara/Source/Niagara/Private/NiagaraScript.cpp
--- a/Engine/Plugins/FX/Niagara/Source/Niagara/Private/NiagaraScript.cpp
+++ b/Engine/Plugins/FX/Niagara/Source/Niagara/Private/NiagaraScript.cpp
@@ -1167,6 +1167,9 @@ void UNiagaraScript::ComputeVMCompilationId(FNiagaraVMExecutableDataId& Id, FGui
 		}
 	}
 
+	// Append the state of the "Fail If Not Set severity" cvar, as it may affect the LastCompileEvents which is a member of the script's VMExecutableData.
+	Id.AdditionalDefines.Add(FString::Printf(TEXT("FailIfNotSetMessageSeverity: %d"), GNiagaraTranslatorFailIfNotSetSeverity));
+
 	FNiagaraVMExecutableDataId& LastGeneratedVMId = GetLastGeneratedVMId(VersionGuid);
 	if (GNiagaraDumpKeyGen == 1 && Id != LastGeneratedVMId)
 	{
diff --git a/Engine/Plugins/FX/Niagara/Source/Niagara/Private/NiagaraSystem.cpp b/Engine/Plugins/FX/Niagara/Source/Niagara/Private/NiagaraSystem.cpp
--- a/Engine/Plugins/FX/Niagara/Source/Niagara/Private/NiagaraSystem.cpp
+++ b/Engine/Plugins/FX/Niagara/Source/Niagara/Private/NiagaraSystem.cpp
@@ -14,6 +14,7 @@
 #include "NiagaraRendererProperties.h"
 #include "NiagaraScriptSourceBase.h"
 #include "NiagaraSettings.h"
+#include "NiagaraShared.h"
 #include "NiagaraSimulationStageBase.h"
 #include "NiagaraStats.h"
 #include "NiagaraTrace.h"
@@ -2701,7 +2702,7 @@ void UNiagaraSystem::EvaluateCompileResultDependencies() const
 		{
 			FNiagaraVariableBase TestVar = Dependency.DependentVariable;
 			ensure(TestVar.GetName() != NAME_None);
-			if (ValidationInfo.Emitter)
+			if (ValidationInfo.Emitter && Dependency.bDependentVariableFromCustomIterationNamespace == false)
 			{
 				FName NewName = GetEmitterVariableAliasName(TestVar, ValidationInfo.Emitter);
 				TestVar.SetName(NewName);
@@ -2730,7 +2731,7 @@ void UNiagaraSystem::EvaluateCompileResultDependencies() const
 			if (!bDependencyMet)
 			{
 				FNiagaraCompileEvent LinkerErrorEvent(
-					FNiagaraCompileEventSeverity::Error, Dependency.LinkerErrorMessage, FString(), false, Dependency.NodeGuid, Dependency.PinGuid, Dependency.StackGuids, FNiagaraCompileEventSource::ScriptDependency);
+					FNiagaraCVarUtilities::GetCompileEventSeverityForFailIfNotSet(), Dependency.LinkerErrorMessage, FString(), false, Dependency.NodeGuid, Dependency.PinGuid, Dependency.StackGuids, FNiagaraCompileEventSource::ScriptDependency);
 				ValidationInfo.CompileResults->LastCompileEvents.Add(LinkerErrorEvent);
 				ValidationInfo.bCompileResultStatusDirty = true;
 			}
@@ -2907,7 +2908,10 @@ bool UNiagaraSystem::QueryCompileComplete(bool bWait, bool bDoPost, bool bDoNotA
 		}
 
 		// Once compile results have been set, check dependencies found during compile and evaluate whether dependency compile events should be added.
-		EvaluateCompileResultDependencies();
+		if (FNiagaraCVarUtilities::GetShouldEmitMessagesForFailIfNotSet())
+		{
+			EvaluateCompileResultDependencies();
+		}
 
 		if (bDoPost)
 		{
diff --git a/Engine/Plugins/FX/Niagara/Source/Niagara/Public/NiagaraCommon.h b/Engine/Plugins/FX/Niagara/Source/Niagara/Public/NiagaraCommon.h
--- a/Engine/Plugins/FX/Niagara/Source/Niagara/Public/NiagaraCommon.h
+++ b/Engine/Plugins/FX/Niagara/Source/Niagara/Public/NiagaraCommon.h
@@ -1387,10 +1387,24 @@ public:
 		NodeGuid = FGuid();
 		PinGuid = FGuid();
 		StackGuids.Empty();
-	}
+		bDependentVariableFromCustomIterationNamespace = false;
+	};
 
-	FNiagaraCompileDependency(const FNiagaraVariableBase& InVar, const FString& InLinkerErrorMessage, FGuid InNodeGuid = FGuid(), FGuid InPinGuid = FGuid(), const TArray<FGuid>& InCallstackGuids = TArray<FGuid>())
-		: LinkerErrorMessage(InLinkerErrorMessage), NodeGuid(InNodeGuid), PinGuid(InPinGuid), StackGuids(InCallstackGuids), DependentVariable(InVar) {}
+	FNiagaraCompileDependency(
+		const FNiagaraVariableBase& InVar
+		, const FString& InLinkerErrorMessage
+		, FGuid InNodeGuid = FGuid()
+		, FGuid InPinGuid = FGuid()
+		, const TArray<FGuid>& InCallstackGuids = TArray<FGuid>()
+		, bool bInDependentVariableFromCustomIterationNamespace = false
+	)
+		: LinkerErrorMessage(InLinkerErrorMessage)
+		, NodeGuid(InNodeGuid)
+		, PinGuid(InPinGuid)
+		, StackGuids(InCallstackGuids)
+		, DependentVariable(InVar)
+		, bDependentVariableFromCustomIterationNamespace(bInDependentVariableFromCustomIterationNamespace)
+	{};
 
 	/* The message itself*/
 	UPROPERTY()
@@ -1412,6 +1426,9 @@ public:
 	UPROPERTY()
 	FNiagaraVariableBase DependentVariable;
 
+	UPROPERTY()
+	bool bDependentVariableFromCustomIterationNamespace;
+
 	FORCEINLINE bool operator==(const FNiagaraCompileDependency& Other)const
 	{
 		return DependentVariable == Other.DependentVariable && NodeGuid == Other.NodeGuid && PinGuid == Other.PinGuid && StackGuids == Other.StackGuids;
diff --git a/Engine/Plugins/FX/Niagara/Source/NiagaraEditor/Private/NiagaraHlslTranslator.cpp b/Engine/Plugins/FX/Niagara/Source/NiagaraEditor/Private/NiagaraHlslTranslator.cpp
--- a/Engine/Plugins/FX/Niagara/Source/NiagaraEditor/Private/NiagaraHlslTranslator.cpp
+++ b/Engine/Plugins/FX/Niagara/Source/NiagaraEditor/Private/NiagaraHlslTranslator.cpp
@@ -1,11 +1,27 @@
 // Copyright Epic Games, Inc. All Rights Reserved.
 
 #include "NiagaraHlslTranslator.h"
+
+#include "EdGraphSchema_Niagara.h"
+#include "EdGraphUtilities.h"
+#include "INiagaraEditorTypeUtilities.h"
+#include "Modules/ModuleManager.h"
+#include "NiagaraCommon.h"
 #include "NiagaraComponent.h"
+#include "NiagaraConstants.h"
+#include "NiagaraDataInterface.h"
+#include "NiagaraDataInterfaceCurve.h"
+#include "NiagaraDataInterfaceVector2DCurve.h"
+#include "NiagaraEditorModule.h"
+#include "NiagaraEditorSettings.h"
+#include "NiagaraEditorTickables.h"
+#include "NiagaraEditorUtilities.h"
+#include "NiagaraFunctionLibrary.h"
+#include "NiagaraNodeEmitter.h"
 #include "NiagaraGraph.h"
-#include "NiagaraScriptSource.h"
-#include "EdGraphUtilities.h"
 #include "NiagaraNode.h"
+#include "NiagaraNodeConvert.h"
+#include "NiagaraNodeCustomHlsl.h"
 #include "NiagaraNodeFunctionCall.h"
 #include "NiagaraNodeIf.h"
 #include "NiagaraNodeInput.h"
@@ -13,35 +29,17 @@
 #include "NiagaraNodeParameterMapGet.h"
 #include "NiagaraNodeParameterMapSet.h"
 #include "NiagaraNodeParameterMapFor.h"
-#include "NiagaraNodeCustomHlsl.h"
+#include "NiagaraNodeSelect.h"
+#include "NiagaraNodeStaticSwitch.h"
 #include "NiagaraNodeOp.h"
-#include "NiagaraNodeConvert.h"
-#include "EdGraphSchema_Niagara.h"
-#include "NiagaraConstants.h"
-#include "NiagaraNodeEmitter.h"
-#include "INiagaraEditorTypeUtilities.h"
-#include "NiagaraEditorUtilities.h"
-#include "NiagaraEditorModule.h"
+#include "NiagaraParameterCollection.h"
+#include "NiagaraScriptSource.h"
+#include "NiagaraScriptVariable.h"
+#include "NiagaraShared.h"
 #include "NiagaraSimulationStageBase.h"
 #include "NiagaraTrace.h"
-#include "NiagaraCommon.h"
-#include "Modules/ModuleManager.h"
-
-#include "NiagaraFunctionLibrary.h"
-
-#include "NiagaraDataInterface.h"
-#include "NiagaraDataInterfaceCurve.h"
-#include "NiagaraDataInterfaceVector2DCurve.h"
-
-#include "NiagaraParameterCollection.h"
-#include "NiagaraEditorTickables.h"
 #include "ShaderCore.h"
 
-#include "NiagaraEditorSettings.h"
-#include "NiagaraNodeStaticSwitch.h"
-#include "NiagaraNodeSelect.h"
-#include "NiagaraScriptVariable.h"
-
 #define LOCTEXT_NAMESPACE "NiagaraCompiler"
 
 DECLARE_CYCLE_STAT(TEXT("Niagara - HlslTranslator - Translate"), STAT_NiagaraEditor_HlslTranslator_Translate, STATGROUP_NiagaraEditor);
@@ -5113,8 +5111,15 @@
 					RecordParamMapDefinedAttributeToNamespaceVar(Var, Inputs[i].Pin);
 					if (!Var.GetType().IsStatic())
 					{
-						if (ParamMapHistories[ParamMapHistoryIdx].IsPrimaryDataSetOutput(Var, GetTargetUsage())) // Note that data interfaces aren't ever in the primary data set even if the namespace matches.)
+						// Note that data interfaces aren't ever in the primary data set even if the namespace matches.)
+						if (ParamMapHistories[ParamMapHistoryIdx].IsPrimaryDataSetOutput(Var, GetTargetUsage()))
+						{
 							CompilationOutput.ScriptData.AttributesWritten.AddUnique(Var);
+						}
+						else if (ParamMapHistories[ParamMapHistoryIdx].IsVariableFromCustomIterationNamespaceOverride(Var))
+						{
+							CompilationOutput.ScriptData.AttributesWritten.AddUnique(Var);
+						}
 					}
 					else 
 					{
@@ -5933,7 +5938,7 @@
 				{
 					if (Variable && Variable->DefaultMode == ENiagaraDefaultMode::FailIfPreviouslyNotSet && !Variable->Variable.IsInNameSpace(FNiagaraConstants::UserNamespaceString))
 					{
-						RegisterCompileDependency(Var, FText::Format(LOCTEXT("UsedBeforeSet", "Variable {0} was read before being set. It's default mode is \"Fail If Previously Not Set\", so this isn't allowed."), FText::FromName(Var.GetName())), GetNode, OutputPins[i], true);
+						RegisterCompileDependency(Var, FText::Format(LOCTEXT("UsedBeforeSet", "Variable {0} was read before being set. It's default mode is \"Fail If Previously Not Set\", so this isn't allowed."), FText::FromName(Var.GetName())), GetNode, OutputPins[i], true, ParamMapHistoryIdx);
 					}
 				}
 
@@ -6002,9 +6007,9 @@
 	
 	if (FNiagaraParameterMapHistory::IsExternalConstantNamespace(Var, CompileOptions.TargetUsage, CompileOptions.GetTargetUsageBitmask()))
 	{
-		if (Variable && Variable->DefaultMode == ENiagaraDefaultMode::FailIfPreviouslyNotSet && !Variable->Variable.IsInNameSpace(FNiagaraConstants::UserNamespaceString) && !bIgnoreDefaultSetFirst)
+		if (Variable && Variable->DefaultMode == ENiagaraDefaultMode::FailIfPreviouslyNotSet && !bIgnoreDefaultSetFirst)
 		{
-			RegisterCompileDependency(Var, FText::Format(LOCTEXT("UsedBeforeSet", "Variable {0} was read before being set. It's default mode is \"Fail If Previously Not Set\", so this isn't allowed."), FText::FromName(Var.GetName())), ErrorNode, nullptr, true);
+			RegisterCompileDependency(Var, FText::Format(LOCTEXT("UsedBeforeSet", "Variable {0} was read before being set. It's default mode is \"Fail If Previously Not Set\", so this isn't allowed."), FText::FromName(Var.GetName())), ErrorNode, nullptr, true, ParamMapHistoryIdx);
 		}
 		if (Var.GetType().IsStatic())
 		{
@@ -6202,7 +6207,7 @@
 		// First check to see if this is defaulted to fail if not set previously. If so, then make sure we don't suck in defaults and error out.
 		if (bValidateFailIfPreviouslyNotSet && bFailIfPreviouslyNotSetSentinel && !bIgnoreDefaultSetFirst)
 		{		
-			RegisterCompileDependency(Var, FText::Format(LOCTEXT("UsedBeforeSet", "Variable {0} was read before being set. It's default mode is \"Fail If Previously Not Set\", so this isn't allowed."), FText::FromName(Var.GetName())), ErrorNode, nullptr, false);
+			RegisterCompileDependency(Var, FText::Format(LOCTEXT("UsedBeforeSet", "Variable {0} was read before being set. It's default mode is \"Fail If Previously Not Set\", so this isn't allowed."), FText::FromName(Var.GetName())), ErrorNode, nullptr, false, ParamMapHistoryIdx);
 		}
 
 		if (bIsPerInstanceAttribute)
@@ -9132,10 +9137,17 @@
 	Message(FNiagaraCompileEventSeverity::Warning, WarningText, InNode, Pin, ShortDescription, bDismissable);
 }
 
-void FHlslNiagaraTranslator::RegisterCompileDependency(const FNiagaraVariableBase& InVar, FText MessageText, const UNiagaraNode* Node, const UEdGraphPin* Pin, bool bEmitAsLinker)
+void FHlslNiagaraTranslator::RegisterCompileDependency(const FNiagaraVariableBase& InVar, FText MessageText, const UNiagaraNode* Node, const UEdGraphPin* Pin, bool bEmitAsLinker, int32 ParamMapHistoryIdx)
 {
+	if (FNiagaraCVarUtilities::GetShouldEmitMessagesForFailIfNotSet() == false)
+	{
+		return;
+	}
+
 	if (InVar.GetType().IsDataInterface() || InVar.GetType().IsUObject() || InVar.IsInNameSpace(FNiagaraConstants::UserNamespaceString) || InVar.IsInNameSpace(FNiagaraConstants::EngineNamespaceString) || InVar.IsInNameSpace(FNiagaraConstants::ParameterCollectionNamespaceString))
+	{
 		return;
+	}
 
 	if (FNiagaraConstants::IsNiagaraConstant(InVar) || InVar.GetName() == TEXT("Emitter.InterpSpawnStartDt") || InVar.GetName() == TEXT("Emitter.SpawnInterval"))
 	{
@@ -9143,16 +9155,19 @@
 	}
 
 	if (bEmitAsLinker)																					
-	{																									 
+	{			
+		bool bVarFromCustomIterationNamespaceOverride = ParamMapHistories[ParamMapHistoryIdx].IsVariableFromCustomIterationNamespaceOverride(InVar);
 		const UNiagaraNode* CurContextNode = ActiveHistoryForFunctionCalls.GetCallingContext();
 		const UNiagaraNode* TargetNode = Node ? Node : CurContextNode;
 
 		FString MessageString = NodePinToMessage(MessageText, TargetNode, Pin);
-		TranslateResults.CompileDependencies.AddUnique(FNiagaraCompileDependency(InVar, MessageString, TargetNode ? TargetNode->NodeGuid : FGuid(), Pin ? Pin->PersistentGuid : FGuid(), GetCallstackGuids()));
+		TranslateResults.CompileDependencies.AddUnique(FNiagaraCompileDependency(InVar, MessageString, TargetNode ? TargetNode->NodeGuid : FGuid(), Pin ? Pin->PersistentGuid : FGuid(), GetCallstackGuids(), bVarFromCustomIterationNamespaceOverride));
 	}
 	else
 	{
-		Error(MessageText, Node, Pin);
+		const bool bDismissable = false;
+		const FString ShortDescription = FString();
+		Message(FNiagaraCVarUtilities::GetCompileEventSeverityForFailIfNotSet(), MessageText, Node, Pin, ShortDescription, bDismissable);
 	}
 }
 
diff --git a/Engine/Plugins/FX/Niagara/Source/NiagaraEditor/Private/NiagaraHlslTranslator.h b/Engine/Plugins/FX/Niagara/Source/NiagaraEditor/Private/NiagaraHlslTranslator.h
--- a/Engine/Plugins/FX/Niagara/Source/NiagaraEditor/Private/NiagaraHlslTranslator.h
+++ b/Engine/Plugins/FX/Niagara/Source/NiagaraEditor/Private/NiagaraHlslTranslator.h
@@ -616,7 +616,7 @@ public:
 	void Message(FNiagaraCompileEventSeverity Severity, FText MessageText, const UNiagaraNode* Node, const UEdGraphPin* Pin, FString ShortDescription = FString(), bool bDismissable = false);
 	virtual void Error(FText ErrorText, const UNiagaraNode* Node, const UEdGraphPin* Pin, FString ShortDescription = FString(), bool bDismissable = false);
 	virtual void Warning(FText WarningText, const UNiagaraNode* Node, const UEdGraphPin* Pin, FString ShortDescription = FString(), bool bDismissable = false);
-	void RegisterCompileDependency(const FNiagaraVariableBase& InVar, FText ErrorText, const UNiagaraNode* Node, const UEdGraphPin* Pin, bool bEmitAsLinker);
+	void RegisterCompileDependency(const FNiagaraVariableBase& InVar, FText ErrorText, const UNiagaraNode* Node, const UEdGraphPin* Pin, bool bEmitAsLinker, int32 ParamMapHistoryIdx);
 	FString NodePinToMessage(FText MessageText, const UNiagaraNode* Node, const UEdGraphPin* Pin);
 
 	virtual bool GetFunctionParameter(const FNiagaraVariable& Parameter, int32& OutParam)const;
diff --git a/Engine/Plugins/FX/Niagara/Source/NiagaraShader/Private/NiagaraShared.cpp b/Engine/Plugins/FX/Niagara/Source/NiagaraShader/Private/NiagaraShared.cpp
--- a/Engine/Plugins/FX/Niagara/Source/NiagaraShader/Private/NiagaraShared.cpp
+++ b/Engine/Plugins/FX/Niagara/Source/NiagaraShader/Private/NiagaraShared.cpp
@@ -29,6 +29,15 @@ IMPLEMENT_TYPE_LAYOUT(FNiagaraShaderMapContent);
 IMPLEMENT_TYPE_LAYOUT(FNiagaraShaderMapId);
 IMPLEMENT_TYPE_LAYOUT(FNiagaraComputeShaderCompilationOutput);
 
+//* CVars */
+int32 GNiagaraTranslatorFailIfNotSetSeverity = 3;
+static FAutoConsoleVariableRef CVarNiagaraTranslatorSilenceFailIfNotSet(
+	TEXT("fx.Niagara.FailIfNotSetSeverity"),
+	GNiagaraTranslatorFailIfNotSetSeverity,
+	TEXT("The severity of messages emitted by Parameters with Default Mode \"Fail If Not Set\". 3 = Error, 2 = Warning, 1= Log, 0 = Disabled.\n"),
+	ECVF_Default
+);
+
 #if WITH_EDITOR
 	FNiagaraCompilationQueue* FNiagaraCompilationQueue::Singleton = nullptr;
 #endif
@@ -750,4 +759,23 @@ bool FNiagaraShaderScript::BeginCompileShaderMap(
 #endif
 }
 
-#endif
\ No newline at end of file
+FNiagaraCompileEventSeverity FNiagaraCVarUtilities::GetCompileEventSeverityForFailIfNotSet()
+{
+	switch (GNiagaraTranslatorFailIfNotSetSeverity) {
+	case 3:
+		return FNiagaraCompileEventSeverity::Error;
+	case 2:
+		return FNiagaraCompileEventSeverity::Warning;
+	case 1:
+		return FNiagaraCompileEventSeverity::Log;
+	default:
+		return FNiagaraCompileEventSeverity::Log;
+	};
+}
+
+bool FNiagaraCVarUtilities::GetShouldEmitMessagesForFailIfNotSet()
+{
+	return GNiagaraTranslatorFailIfNotSetSeverity != 0;
+}
+
+#endif
diff --git a/Engine/Plugins/FX/Niagara/Source/NiagaraShader/Public/NiagaraShared.h b/Engine/Plugins/FX/Niagara/Source/NiagaraShader/Public/NiagaraShared.h
--- a/Engine/Plugins/FX/Niagara/Source/NiagaraShader/Public/NiagaraShared.h
+++ b/Engine/Plugins/FX/Niagara/Source/NiagaraShader/Public/NiagaraShared.h
@@ -958,4 +958,11 @@ private:
 	friend class FShaderCompilingManager;
 };
 
+/** CVars */
+extern NIAGARASHADER_API int32 GNiagaraTranslatorFailIfNotSetSeverity;
 
+namespace FNiagaraCVarUtilities
+{
+	NIAGARASHADER_API FNiagaraCompileEventSeverity GetCompileEventSeverityForFailIfNotSet();
+	NIAGARASHADER_API bool GetShouldEmitMessagesForFailIfNotSet();
+};
