Fixing refresh of Convert node breaking connections when doing float->bool or vice versa. Also cleaned up the connections UI when making this connection so it is clear what will happen.

diff --git a/Engine/Plugins/FX/Niagara/Source/Niagara/Private/NiagaraModule.cpp b/Engine/Plugins/FX/Niagara/Source/Niagara/Private/NiagaraModule.cpp
--- a/Engine/Plugins/FX/Niagara/Source/Niagara/Private/NiagaraModule.cpp
+++ b/Engine/Plugins/FX/Niagara/Source/Niagara/Private/NiagaraModule.cpp
@@ -1178,7 +1178,10 @@ bool FNiagaraTypeDefinition::TypesAreAssignable(const FNiagaraTypeDefinition& Ty
 
 bool FNiagaraTypeDefinition::IsLossyConversion(const FNiagaraTypeDefinition& TypeA, const FNiagaraTypeDefinition& TypeB)
 {
-	return (TypeA == IntDef && TypeB == FloatDef) || (TypeB == IntDef && TypeA == FloatDef);
+	return (TypeA.IsSameBaseDefinition(IntDef) && TypeB.IsSameBaseDefinition(FloatDef))
+		|| (TypeB.IsSameBaseDefinition(IntDef) && TypeA.IsSameBaseDefinition(FloatDef))
+		|| (TypeA.IsSameBaseDefinition(BoolDef) && TypeB.IsSameBaseDefinition(FloatDef))
+		|| (TypeB.IsSameBaseDefinition(BoolDef) && TypeA.IsSameBaseDefinition(FloatDef));
 }
 
 FNiagaraTypeDefinition FNiagaraTypeDefinition::GetNumericOutputType(const TArray<FNiagaraTypeDefinition> TypeDefinintions, ENiagaraNumericOutputTypeSelectionMode SelectionMode)
diff --git a/Engine/Plugins/FX/Niagara/Source/NiagaraEditor/Private/EdGraphSchema_Niagara.cpp b/Engine/Plugins/FX/Niagara/Source/NiagaraEditor/Private/EdGraphSchema_Niagara.cpp
--- a/Engine/Plugins/FX/Niagara/Source/NiagaraEditor/Private/EdGraphSchema_Niagara.cpp
+++ b/Engine/Plugins/FX/Niagara/Source/NiagaraEditor/Private/EdGraphSchema_Niagara.cpp
@@ -1011,6 +1011,10 @@ const FPinConnectionResponse UEdGraphSchema_Niagara::CanCreateConnection(const U
 				{
 					return FPinConnectionResponse(CONNECT_RESPONSE_DISALLOW, TEXT("Types are not compatible"));
 				}
+				else if (FNiagaraTypeDefinition::IsLossyConversion(PinToTypeDefinition(InputPin), PinToTypeDefinition(OutputPin)))
+				{
+					return FPinConnectionResponse(CONNECT_RESPONSE_MAKE_WITH_CONVERSION_NODE, FString::Printf(TEXT("Convert %s to %s. Be aware that this can be a lossy conversion."), *(PinToTypeDefinition(InputPin).GetNameText().ToString()), *(PinToTypeDefinition(OutputPin).GetNameText().ToString())));
+				}
 				else
 				{
 					return FPinConnectionResponse(CONNECT_RESPONSE_MAKE_WITH_CONVERSION_NODE, FString::Printf(TEXT("Convert %s to %s"), *(PinToTypeDefinition(InputPin).GetNameText().ToString()), *(PinToTypeDefinition(OutputPin).GetNameText().ToString())));
diff --git a/Engine/Plugins/FX/Niagara/Source/NiagaraEditor/Private/NiagaraNodeConvert.cpp b/Engine/Plugins/FX/Niagara/Source/NiagaraEditor/Private/NiagaraNodeConvert.cpp
--- a/Engine/Plugins/FX/Niagara/Source/NiagaraEditor/Private/NiagaraNodeConvert.cpp
+++ b/Engine/Plugins/FX/Niagara/Source/NiagaraEditor/Private/NiagaraNodeConvert.cpp
@@ -430,7 +430,7 @@ bool UNiagaraNodeConvert::RefreshFromExternalChanges()
 			FNiagaraTypeDefinition SourcePropTypeDef = SourceProperty ? Schema->GetTypeDefForProperty(SourceProperty) : SourcePinTypeDef;
 			FNiagaraTypeDefinition DestinationPropTypeDef = DestinationProperty ? Schema->GetTypeDefForProperty(DestinationProperty) : DestinationPinTypeDef;
 
-			bTypesAreAssignable = FNiagaraTypeDefinition::TypesAreAssignable(SourcePropTypeDef, DestinationPropTypeDef);
+			bTypesAreAssignable = FNiagaraTypeDefinition::TypesAreAssignable(SourcePropTypeDef, DestinationPropTypeDef) || FNiagaraTypeDefinition::IsLossyConversion(SourcePropTypeDef, DestinationPropTypeDef);
 		}
 
 
@@ -792,7 +792,7 @@ bool UNiagaraNodeConvert::InitConversion(UEdGraphPin* FromPin, UEdGraphPin* ToPi
 		FNiagaraTypeDefinition ToPropType = Schema->GetTypeDefForProperty(ToProp);
 		SrcPath.Empty();
 		DestPath.Empty();
-		if (FromPropType == ToPropType)
+		if (FNiagaraTypeDefinition::TypesAreAssignable(FromPropType, ToPropType) || FNiagaraTypeDefinition::IsLossyConversion(FromPropType, ToPropType))
 		{
 			SrcPath.Add(*FromProp->GetName());
 			DestPath.Add(*ToProp->GetName());
