- Don't execute Store if the instance isn't valid on an iteration interface
- Change check to compilation error to avoid loss of work

#rnx

diff --git a/Engine/Plugins/FX/Niagara/Source/NiagaraCore/Private/NiagaraCustomVersion.cpp b/Engine/Plugins/FX/Niagara/Source/NiagaraCore/Private/NiagaraCustomVersion.cpp
--- a/Engine/Plugins/FX/Niagara/Source/NiagaraCore/Private/NiagaraCustomVersion.cpp
+++ b/Engine/Plugins/FX/Niagara/Source/NiagaraCore/Private/NiagaraCustomVersion.cpp
@@ -9,4 +9,4 @@ const FGuid FNiagaraCustomVersion::GUID(0xFCF57AFA, 0x50764283, 0xB9A9E658, 0xFF
 FCustomVersionRegistration GRegisterNiagaraCustomVersion(FNiagaraCustomVersion::GUID, FNiagaraCustomVersion::LatestVersion, TEXT("NiagaraVer"));
 
 // Note: When encountering a conflict on this file please generate a new GUID
-const FGuid FNiagaraCustomVersion::LatestScriptCompileVersion(0x5A07FD79, 0x1A668744, 0x9695460A, 0xCCBC530D);
\ No newline at end of file
+const FGuid FNiagaraCustomVersion::LatestScriptCompileVersion(0x07F493BD, 0xE6B2A54A, 0xBD23F309, 0x93A5BC76);
\ No newline at end of file
diff --git a/Engine/Plugins/FX/Niagara/Source/NiagaraEditor/Private/NiagaraHlslTranslator.cpp b/Engine/Plugins/FX/Niagara/Source/NiagaraEditor/Private/NiagaraHlslTranslator.cpp
--- a/Engine/Plugins/FX/Niagara/Source/NiagaraEditor/Private/NiagaraHlslTranslator.cpp
+++ b/Engine/Plugins/FX/Niagara/Source/NiagaraEditor/Private/NiagaraHlslTranslator.cpp
@@ -3065,7 +3065,10 @@ void FHlslNiagaraTranslator::DefineMainGPUFunctions(
 			{
 				if (!TranslationStages[i].CustomWriteFunction.IsEmpty())
 				{
-					HlslOutput += TranslationStages[i].CustomWriteFunction + TEXT("(Context);\n\n");
+					HlslOutput += TEXT("if ( bIsValidInstance )\n");
+					HlslOutput += TEXT("{\n");
+					HlslOutput += TEXT("\t") + TranslationStages[i].CustomWriteFunction + TEXT("(Context);\n\n");
+					HlslOutput += TEXT("}\n");
 					continue;
 				}
 			}
@@ -5425,7 +5428,11 @@ bool FHlslNiagaraTranslator::ParameterMapRegisterExternalConstantNamespaceVariab
 					else
 					{
 						DataInterface = CompileDuplicateData->GetDuplicatedDataInterfaceCDOForClass(const_cast<UClass*>(InVariable.GetType().GetClass()));
-						checkf(DataInterface != nullptr, TEXT("GetDuplicatedDataInterfaceCDOForClass failed for Variable(%s) Class(%s)"), *InVariable.GetName().ToString(), *InVariable.GetType().GetClass()->GetName());
+						if (DataInterface == nullptr)
+						{
+							Error(FText::Format(LOCTEXT("GetDuplicatedDataInterfaceCDOForClassFailed", "GetDuplicatedDataInterfaceCDOForClass failed for Variable({0}) Class({1})"), FText::FromName(InVariable.GetName()), InVariable.GetType().GetNameText()), InNodeForErrorReporting, InDefaultPin);
+							return false;
+						}
 					}
 					if (ensure(DataInterface))
 					{
