Niagara - Changed the log levels and wording for script variable generation messages so that the cook doesn't fail.  Will investigate further.

diff --git a/Engine/Plugins/FX/Niagara/Source/NiagaraEditor/Private/NiagaraGraph.cpp b/Engine/Plugins/FX/Niagara/Source/NiagaraEditor/Private/NiagaraGraph.cpp
--- a/Engine/Plugins/FX/Niagara/Source/NiagaraEditor/Private/NiagaraGraph.cpp
+++ b/Engine/Plugins/FX/Niagara/Source/NiagaraEditor/Private/NiagaraGraph.cpp
@@ -423,7 +423,7 @@ void UNiagaraGraph::PostLoad()
 					if (ScriptVariable->GetIsSubscribedToParameterDefinitions() && MatchingScriptVariable->GetIsSubscribedToParameterDefinitions())
 					{
 						// Both of the script variables with duplicate ids are controlled by parameter definitions so issue a warning because neither will be updated.
-						UE_LOG(LogNiagaraEditor, Warning, TEXT("Duplicate ids found for script variables which are both subscribed to parameter definitions.\nScript Variable 1 Name: %s Type: %s Path: %s\nScript Variable 2 Name: %s Type: %s Path: %s"),
+						UE_LOG(LogNiagaraEditor, Log, TEXT("Duplicate ids found for script variables which are both subscribed to parameter definitions.\nScript Variable 1 Name: %s Type: %s Path: %s\nScript Variable 2 Name: %s Type: %s Path: %s"),
 							*MatchingScriptVariable->Variable.GetName().ToString(), *MatchingScriptVariable->Variable.GetType().GetName(), *MatchingScriptVariable->GetPathName(),
 							*ScriptVariable->Variable.GetName().ToString(), *ScriptVariable->Variable.GetType().GetName(), *ScriptVariable->GetPathName());
 					}
@@ -436,7 +436,7 @@ void UNiagaraGraph::PostLoad()
 							ScriptVariable->Metadata.SetVariableGuid(UNiagaraScriptVariable::GenerateStableGuid(ScriptVariable));
 							if (VariableGuidToScriptVariable.Contains(ScriptVariable->Metadata.GetVariableGuid()))
 							{
-								UE_LOG(LogNiagaraEditor, Warning, TEXT("Failed to generate a stable unique variable guid for script variable. Name: %s Type: %s Path: %s"),
+								UE_LOG(LogNiagaraEditor, Log, TEXT("Could not generate a stable unique variable guid for script variable. Name: %s Type: %s Path: %s"),
 									*ScriptVariable->Variable.GetName().ToString(), *ScriptVariable->Variable.GetType().GetName(), *ScriptVariable->GetPathName());
 							}
 							else
@@ -449,7 +449,7 @@ void UNiagaraGraph::PostLoad()
 							MatchingScriptVariable->Metadata.SetVariableGuid(UNiagaraScriptVariable::GenerateStableGuid(MatchingScriptVariable));
 							if (VariableGuidToScriptVariable.Contains(MatchingScriptVariable->Metadata.GetVariableGuid()))
 							{
-								UE_LOG(LogNiagaraEditor, Warning, TEXT("Failed to generate a stable unique variable guid for script variable. Name: %s Type: %s Path: %s"),
+								UE_LOG(LogNiagaraEditor, Log, TEXT("Could not generate a stable unique variable guid for script variable. Name: %s Type: %s Path: %s"),
 									*MatchingScriptVariable->Variable.GetName().ToString(), *MatchingScriptVariable->Variable.GetType().GetName(), *MatchingScriptVariable->GetPathName());
 							}
 							else
