Fixed dynamic inputs always using the exposed version due to an invalid script version guid being passed in during merge
This fixes an issue where an emitter with inheritance fails to merge during post load (no changes on the parent from last time, so we'd expect a success with no differences, but it fails instead)
when a module is using a dynamic input script with a version that is NOT exposed.
The function call node in the emitter script would correctly show its static switch inputs as pins, but when looking for the pin matching the variable, the variables were from the exposed version instead, so the pin couldn't be found

diff --git a/Engine/Plugins/FX/Niagara/Source/NiagaraEditor/Private/NiagaraScriptMergeManager.cpp b/Engine/Plugins/FX/Niagara/Source/NiagaraEditor/Private/NiagaraScriptMergeManager.cpp
--- a/Engine/Plugins/FX/Niagara/Source/NiagaraEditor/Private/NiagaraScriptMergeManager.cpp
+++ b/Engine/Plugins/FX/Niagara/Source/NiagaraEditor/Private/NiagaraScriptMergeManager.cpp
@@ -2820,7 +2820,7 @@ FNiagaraScriptMergeManager::FApplyDiffResults FNiagaraScriptMergeManager::AddInp
 
 					UNiagaraNodeFunctionCall* DynamicInputFunctionCall;
 					FNiagaraStackGraphUtilities::SetDynamicInputForFunctionInput(InputOverridePin, FunctionScript, 
-						DynamicInputFunctionCall, OverrideToAdd->GetOverrideNodeId(), OverrideToAdd->GetDynamicValueFunction()->GetFunctionCallNode()->GetFunctionName());
+						DynamicInputFunctionCall, OverrideToAdd->GetOverrideNodeId(), OverrideToAdd->GetDynamicValueFunction()->GetFunctionCallNode()->GetFunctionName(), OverrideToAdd->GetDynamicValueFunction()->GetFunctionCallNode()->SelectedScriptVersion);
 					for (TSharedRef<FNiagaraStackFunctionInputOverrideMergeAdapter> DynamicInputInputOverride : OverrideToAdd->GetDynamicValueFunction()->GetInputOverrides())
 					{
 						FApplyDiffResults AddResults = AddInputOverride(BaseEmitterAdapter, ScratchPadAdapter, OwningScript, *DynamicInputFunctionCall, DynamicInputInputOverride);
