Clean up pass on HWRT within Niagara (though it remains experimental).  Cleans up the abstraction to allow us to have a fallback method for platforms/scenarios where HW ray tracing is suppported/enabled/used (currently only SDF is provided as a fallback).
-shuffles a lot of the related files around to clean up public/private and moving stuff out of NiagaraShader where possible
-migrated async trace functions into it's own DI, away from the CollisionQuery one
-updates collision module to use the new DI for it's experimental ray traced GPU version

diff --git a/Engine/Plugins/FX/Niagara/Shaders/Private/NiagaraAsyncGpuTraceCommon.ush b/Engine/Plugins/FX/Niagara/Shaders/Private/NiagaraAsyncGpuTraceCommon.ush
new file mode 100644
--- /dev/null
+++ b/Engine/Plugins/FX/Niagara/Shaders/Private/NiagaraAsyncGpuTraceCommon.ush
@@ -0,0 +1,25 @@
+// Copyright Epic Games, Inc. All Rights Reserved.
+
+#pragma once
+
+#include "/Engine/Private/Common.ush"
+
+struct FNiagaraAsyncGpuTrace
+{
+	float3 Origin;
+	float TFar;
+	float3 Direction;
+	uint CollisionGroup;
+};
+
+struct FNiagaraAsyncGpuTraceResult
+{
+	float3 WorldPosition;
+	float HitT;
+
+	float3 WorldNormal;
+
+	float _Pad0; // padding to force 16 byte alignment to meet requirements for VK
+
+	bool IsHit() { return HitT > 0; }
+};
diff --git a/Engine/Plugins/FX/Niagara/Shaders/Private/NiagaraAsyncGpuTraceUtils.usf b/Engine/Plugins/FX/Niagara/Shaders/Private/NiagaraAsyncGpuTraceUtils.usf
new file mode 100644
--- /dev/null
+++ b/Engine/Plugins/FX/Niagara/Shaders/Private/NiagaraAsyncGpuTraceUtils.usf
@@ -0,0 +1,29 @@
+// Copyright Epic Games, Inc. All Rights Reserved.
+
+#pragma once
+
+#include "/Engine/Private/Common.ush"
+#include "NiagaraAsyncGpuTraceCommon.ush"
+
+#if !defined(THREADGROUP_SIZE_X)
+	#define THREADGROUP_SIZE_X 32
+#endif
+
+RWStructuredBuffer<FNiagaraAsyncGpuTraceResult> Results;
+
+uint TraceCount;
+uint ResultsOffset;
+
+[numthreads(THREADGROUP_SIZE_X, 1, 1)]
+void NiagaraClearAsyncGpuTraceCS(uint GroupIndex : SV_GroupIndex)
+{
+	if (GroupIndex >= TraceCount)
+	{
+		return;
+	}
+
+	FNiagaraAsyncGpuTraceResult Result = (FNiagaraAsyncGpuTraceResult) 0;
+	Result.HitT = -1.0f;
+
+	Results[ResultsOffset + GroupIndex] = Result;
+}
\ No newline at end of file
diff --git a/Engine/Plugins/FX/Niagara/Shaders/Private/NiagaraDataInterfaceAsyncGpuTrace.ush b/Engine/Plugins/FX/Niagara/Shaders/Private/NiagaraDataInterfaceAsyncGpuTrace.ush
new file mode 100644
--- /dev/null
+++ b/Engine/Plugins/FX/Niagara/Shaders/Private/NiagaraDataInterfaceAsyncGpuTrace.ush
@@ -0,0 +1,114 @@
+// Copyright Epic Games, Inc. All Rights Reserved.
+
+#pragma once
+
+#include "NiagaraAsyncGpuTraceCommon.ush"
+
+void NDIAsyncGpuTrace_IssueAsyncRayTrace(
+	uint MaxRayTraceCount,
+	RWStructuredBuffer<FNiagaraAsyncGpuTrace> RWRayRequests,
+	uint RayRequestOffset,
+	RWBuffer<uint> RWRayTraceCounts,
+	uint RayTraceCountsOffset,
+	float3 In_LWCTile,
+	int In_QueryID,
+	float3 In_TraceStart,
+	float3 In_TraceEnd,
+	int In_CollisionGroup,
+	out bool Out_IsQueryValid)
+{
+	Out_IsQueryValid = false;
+
+	if (In_QueryID >= 0 && uint(In_QueryID) < MaxRayTraceCount)
+	{
+		FLWCVector3 LwcStartPos = MakeLWCVector3(In_LWCTile, In_TraceStart);
+		FLWCVector3 LwcEndPos = MakeLWCVector3(In_LWCTile, In_TraceEnd);
+		
+		float3 Trace = LWCToFloat(LWCSubtract(LwcEndPos, LwcStartPos));
+		float TraceLength = length(Trace);
+
+		//Offset query into the actual ray buffer.
+		uint RayIndex = RayRequestOffset + In_QueryID;
+		RWRayRequests[RayIndex].Origin = LWCToFloat(LwcStartPos); // TODO: remove LWCToFloat here when FNiagaraAsyncGpuTrace supports LWC
+		RWRayRequests[RayIndex].Direction = normalize(Trace);
+		RWRayRequests[RayIndex].TFar = TraceLength;
+		RWRayRequests[RayIndex].CollisionGroup = In_CollisionGroup;
+		Out_IsQueryValid = true;
+
+		// we want to keep track of the maximum index used
+		InterlockedMax(RWRayTraceCounts[RayTraceCountsOffset + 0], In_QueryID + 1);
+		RWRayTraceCounts[RayTraceCountsOffset + 1] = 1;
+		RWRayTraceCounts[RayTraceCountsOffset + 2] = 1;
+	}
+}
+
+void NDIAsyncGpuTrace_ReserveRayTraceIndex(uint MaxRayTraceCount, RWBuffer<uint> RWRayTraceCounts, uint RayTraceCountsOffset, int In_TraceCount, out int Out_StartQueryID, out bool Out_IndicesValid)
+{
+	Out_StartQueryID = -1;
+	Out_IndicesValid = false;
+
+	if ((In_TraceCount + RWRayTraceCounts[RayTraceCountsOffset + 0]) < MaxRayTraceCount)
+	{
+		uint RequestIndex = 0;
+		InterlockedAdd(RWRayTraceCounts[RayTraceCountsOffset + 0], In_TraceCount, RequestIndex);
+		RWRayTraceCounts[RayTraceCountsOffset + 1] = 1;
+		RWRayTraceCounts[RayTraceCountsOffset + 2] = 1;
+
+		if ((In_TraceCount + RequestIndex) <= MaxRayTraceCount)
+		{
+			Out_StartQueryID = int(RequestIndex);
+			Out_IndicesValid = true;
+		}
+		else
+		{
+			InterlockedMax(RWRayTraceCounts[RayTraceCountsOffset + 0], MaxRayTraceCount);
+		}
+	}
+}
+
+void NDIAsyncGpuTrace_CreateAsyncRayTrace(	uint MaxRayTraceCount, RWStructuredBuffer<FNiagaraAsyncGpuTrace> RWRayRequests, uint RayRequestsOffset, RWBuffer<uint> RWRayTraceCounts, uint RayTraceCountsOffset, float3 In_LWCTile,
+											float3 In_TraceStart, float3 In_TraceEnd, int In_CollisionGroup, out int Out_QueryID, out bool Out_IsQueryValid)
+{
+	Out_IsQueryValid = false;
+	Out_QueryID = -1;
+
+	NDIAsyncGpuTrace_ReserveRayTraceIndex(MaxRayTraceCount, RWRayTraceCounts, RayTraceCountsOffset, 1, Out_QueryID, Out_IsQueryValid);
+
+	if (Out_IsQueryValid)
+	{
+		FLWCVector3 LwcStartPos = MakeLWCVector3(In_LWCTile, In_TraceStart);
+		FLWCVector3 LwcEndPos = MakeLWCVector3(In_LWCTile, In_TraceEnd);
+		
+		float3 Trace = LWCToFloat(LWCSubtract(LwcEndPos, LwcStartPos));
+		float TraceLength = length(Trace);
+
+		uint RayIndex = Out_QueryID + RayRequestsOffset;
+		RWRayRequests[RayIndex].Origin = LWCToFloat(LwcStartPos); // TODO: remove LWCToFloat here when FNiagaraAsyncGpuTrace supports LWC
+		RWRayRequests[RayIndex].Direction = normalize(Trace);
+		RWRayRequests[RayIndex].TFar = TraceLength;
+
+		//TODO: We can pack collision group and trace channels into the same int and test both in the RG Shader.
+		RWRayRequests[RayIndex].CollisionGroup = In_CollisionGroup;
+	}
+}
+
+void NDIAsyncGpuTrace_ReadAsyncRayTrace(uint MaxRayTraceCount, StructuredBuffer<FNiagaraAsyncGpuTraceResult> IntersectionResults, uint IntersectionRestultsOffest, float3 In_LWCTile, int In_PreviousFrameQueryID, out bool Out_CollisionValid, out float Out_CollisionDistance, out float3 Out_CollisionPosWorld, out float3 Out_CollisionNormal)
+{
+	Out_CollisionValid = false;
+	Out_CollisionDistance = 0.0f;
+	Out_CollisionPosWorld = float3(0.0f, 0.0f, 0.0f);
+	Out_CollisionNormal = float3(0.0f, 0.0f, 0.0f);
+
+	if (In_PreviousFrameQueryID >= 0 && uint(In_PreviousFrameQueryID) < MaxRayTraceCount)
+	{
+		FNiagaraAsyncGpuTraceResult HitResult = IntersectionResults[In_PreviousFrameQueryID + IntersectionRestultsOffest];
+
+		Out_CollisionValid = HitResult.IsHit();
+		if (Out_CollisionValid)
+		{
+			Out_CollisionDistance = HitResult.HitT;
+			Out_CollisionPosWorld = LWCToFloat(MakeLWCVector3(-In_LWCTile, HitResult.WorldPosition)); // TODO: HitResult.WorldPosition should support LWC properly
+			Out_CollisionNormal = HitResult.WorldNormal;
+		}
+	}
+}
\ No newline at end of file
diff --git a/Engine/Plugins/FX/Niagara/Shaders/Private/NiagaraDataInterfaceAsyncGpuTraceTemplate.ush b/Engine/Plugins/FX/Niagara/Shaders/Private/NiagaraDataInterfaceAsyncGpuTraceTemplate.ush
new file mode 100644
--- /dev/null
+++ b/Engine/Plugins/FX/Niagara/Shaders/Private/NiagaraDataInterfaceAsyncGpuTraceTemplate.ush
@@ -0,0 +1,31 @@
+// Copyright Epic Games, Inc. All Rights Reserved.
+
+float3										SystemLWCTile_{ParameterName};
+/** Niagara ray tracing uses shared buffers between dispatches so each buffer has a buffer + and offset into that buffer for the current dispatch. */
+uint										MaxRayTraceCount_{ParameterName};
+RWStructuredBuffer<FNiagaraAsyncGpuTrace>	RWRayRequests_{ParameterName};
+uint										RayRequestsOffset_{ ParameterName };
+StructuredBuffer<FNiagaraAsyncGpuTraceResult>	IntersectionResults_{ParameterName};
+uint										IntersectionResultsOffset_{ ParameterName };
+RWBuffer<uint>								RWRayTraceCounts_{ParameterName};
+uint										RayTraceCountsOffset_{ ParameterName };
+
+void IssueAsyncRayTraceGpu_{ParameterName}_UEImpureCall(int In_QueryID, float3 In_TraceStart, float3 In_TraceEnd, int In_CollisionGroup, out bool Out_IsQueryValid)
+{
+	NDIAsyncGpuTrace_IssueAsyncRayTrace(MaxRayTraceCount_{ParameterName}, RWRayRequests_{ParameterName}, RayRequestsOffset_{ ParameterName }, RWRayTraceCounts_{ParameterName}, RayTraceCountsOffset_{ ParameterName }, SystemLWCTile_{ParameterName}, In_QueryID, In_TraceStart, In_TraceEnd, In_CollisionGroup, Out_IsQueryValid);
+}
+
+void ReserveRayTraceIndex_{ParameterName}_UEImpureCall(int In_TraceCount, out int Out_StartQueryID, out bool Out_IndicesValid)
+{
+	NDIAsyncGpuTrace_ReserveRayTraceIndex(MaxRayTraceCount_{ParameterName}, RWRayTraceCounts_{ParameterName}, RayTraceCountsOffset_{ ParameterName }, In_TraceCount, Out_StartQueryID, Out_IndicesValid);
+}
+
+void CreateAsyncRayTraceGpu_{ParameterName}_UEImpureCall(float3 In_TraceStart, float3 In_TraceEnd, int In_CollisionGroup, out int Out_QueryID, out bool Out_IsQueryValid)
+{
+	NDIAsyncGpuTrace_CreateAsyncRayTrace(MaxRayTraceCount_{ParameterName}, RWRayRequests_{ParameterName}, RayRequestsOffset_{ ParameterName }, RWRayTraceCounts_{ParameterName}, RayTraceCountsOffset_{ ParameterName }, SystemLWCTile_{ParameterName}, In_TraceStart, In_TraceEnd, In_CollisionGroup, Out_QueryID, Out_IsQueryValid);
+}
+
+void ReadAsyncRayTraceGpu_{ParameterName}(int In_PreviousFrameQueryID, out bool Out_CollisionValid, out float Out_CollisionDistance, out float3 Out_CollisionPosWorld, out float3 Out_CollisionNormal)
+{
+	NDIAsyncGpuTrace_ReadAsyncRayTrace(MaxRayTraceCount_{ParameterName}, IntersectionResults_{ParameterName}, IntersectionResultsOffset_{ParameterName}, SystemLWCTile_{ParameterName}, In_PreviousFrameQueryID, Out_CollisionValid, Out_CollisionDistance, Out_CollisionPosWorld, Out_CollisionNormal);
+}
\ No newline at end of file
diff --git a/Engine/Plugins/FX/Niagara/Shaders/Private/NiagaraDataInterfaceCollisionQuery.ush b/Engine/Plugins/FX/Niagara/Shaders/Private/NiagaraDataInterfaceCollisionQuery.ush
--- a/Engine/Plugins/FX/Niagara/Shaders/Private/NiagaraDataInterfaceCollisionQuery.ush
+++ b/Engine/Plugins/FX/Niagara/Shaders/Private/NiagaraDataInterfaceCollisionQuery.ush
@@ -2,7 +2,7 @@
 
 #pragma once
 
-#include "NiagaraRayTracingCommon.ush"
+#include "/Engine/Private/Common.ush"
 
 #define NDICOLLISIONQUERY_USE_GBUFFER_NORMAL    (FEATURE_LEVEL >= FEATURE_LEVEL_SM5)
 
@@ -144,111 +144,3 @@ void NDICollisionQuery_QueryMeshDistanceFieldGPU(in float3 In_SamplePos, in floa
 	Out_IsDistanceFieldValid = false;
 #endif
 }
-
-#if NIAGARA_USE_RAYTRACING
-
-uint BuildRayTraceMask(int TraceChannel)
-{
-	// for now we'll ignore the TraceChannel till we can mark up the environment
-	return 0xFF;
-}
-
-void NDICollisionQuery_IssueAsyncRayTrace(uint MaxRayTraceCount, RWStructuredBuffer<FNiagaraRayData> RWRayRequests, uint RayRequestOffset, RWBuffer<uint> RWRayTraceCounts, uint RayTraceCountsOffset, float3 In_LWCTile, int In_QueryID, float3 In_TraceStart, float3 In_TraceEnd, int In_TraceChannel, out bool Out_IsQueryValid)
-{
-	Out_IsQueryValid = false;
-
-	if (In_QueryID >= 0 && uint(In_QueryID) < MaxRayTraceCount)
-	{
-		FLWCVector3 LwcStartPos = MakeLWCVector3(In_LWCTile, In_TraceStart);
-		FLWCVector3 LwcEndPos = MakeLWCVector3(In_LWCTile, In_TraceEnd);
-		
-		float3 Trace = LWCToFloat(LWCSubtract(LwcEndPos, LwcStartPos));
-		float TraceLength = length(Trace);
-
-		//Offset query into the actual ray buffer.
-		uint RayIndex = RayRequestOffset + In_QueryID;
-		RWRayRequests[RayIndex].Origin = LWCToFloat(LwcStartPos); // TODO: remove LWCToFloat here when FNiagaraRayData supports LWC
-		RWRayRequests[RayIndex].Mask = BuildRayTraceMask(In_TraceChannel);
-		RWRayRequests[RayIndex].Direction = normalize(Trace);
-		RWRayRequests[RayIndex].TFar = TraceLength;
-		Out_IsQueryValid = true;
-
-		// we want to keep track of the maximum index used
-		InterlockedMax(RWRayTraceCounts[RayTraceCountsOffset + 0], In_QueryID + 1);
-		RWRayTraceCounts[RayTraceCountsOffset + 1] = 1;
-		RWRayTraceCounts[RayTraceCountsOffset + 2] = 1;
-	}
-}
-
-void NDICollisionQuery_ReserveRayTraceIndex(uint MaxRayTraceCount, RWBuffer<uint> RWRayTraceCounts, uint RayTraceCountsOffset, int In_TraceCount, out int Out_StartQueryID, out bool Out_IndicesValid)
-{
-	Out_StartQueryID = -1;
-	Out_IndicesValid = false;
-
-	if ((In_TraceCount + RWRayTraceCounts[RayTraceCountsOffset + 0]) < MaxRayTraceCount)
-	{
-		uint RequestIndex = 0;
-		InterlockedAdd(RWRayTraceCounts[RayTraceCountsOffset + 0], In_TraceCount, RequestIndex);
-		RWRayTraceCounts[RayTraceCountsOffset + 1] = 1;
-		RWRayTraceCounts[RayTraceCountsOffset + 2] = 1;
-
-		if ((In_TraceCount + RequestIndex) <= MaxRayTraceCount)
-		{
-			Out_StartQueryID = int(RequestIndex);
-			Out_IndicesValid = true;
-		}
-		else
-		{
-			InterlockedMax(RWRayTraceCounts[RayTraceCountsOffset + 0], MaxRayTraceCount);
-		}
-	}
-}
-
-void NDICollisionQuery_CreateAsyncRayTrace(	uint MaxRayTraceCount, RWStructuredBuffer<FNiagaraRayData> RWRayRequests, uint RayRequestsOffset, RWBuffer<uint> RWRayTraceCounts, uint RayTraceCountsOffset, float3 In_LWCTile,
-											float3 In_TraceStart, float3 In_TraceEnd, int In_TraceChannel, out int Out_QueryID, out bool Out_IsQueryValid)
-{
-	Out_IsQueryValid = false;
-
-	NDICollisionQuery_ReserveRayTraceIndex(MaxRayTraceCount, RWRayTraceCounts, RayTraceCountsOffset, 1, Out_QueryID, Out_IsQueryValid);
-
-	if (Out_IsQueryValid)
-	{
-		FLWCVector3 LwcStartPos = MakeLWCVector3(In_LWCTile, In_TraceStart);
-		FLWCVector3 LwcEndPos = MakeLWCVector3(In_LWCTile, In_TraceEnd);
-		
-		float3 Trace = LWCToFloat(LWCSubtract(LwcEndPos, LwcStartPos));
-		float TraceLength = length(Trace);
-
-		uint RayIndex = Out_QueryID + RayRequestsOffset;
-		RWRayRequests[RayIndex].Origin = LWCToFloat(LwcStartPos); // TODO: remove LWCToFloat here when FNiagaraRayData supports LWC
-		RWRayRequests[RayIndex].Mask = BuildRayTraceMask(In_TraceChannel);
-		RWRayRequests[RayIndex].Direction = normalize(Trace);
-		RWRayRequests[RayIndex].TFar = TraceLength;
-
-		//TODO: We can pack collision group and trace channels into the same int and test both in the RG Shader.
-		RWRayRequests[RayIndex].CollisionGroup = In_TraceChannel;
-	}
-}
-
-void NDICollisionQuery_ReadAsyncRayTrace(uint MaxRayTraceCount, StructuredBuffer<FNiagaraRayTracingResult> IntersectionResults, uint IntersectionRestultsOffest, float3 In_LWCTile, int In_PreviousFrameQueryID, out bool Out_CollisionValid, out float Out_CollisionDistance, out float3 Out_CollisionPosWorld, out float3 Out_CollisionNormal)
-{
-	Out_CollisionValid = false;
-	Out_CollisionDistance = 0.0f;
-	Out_CollisionPosWorld = float3(0.0f, 0.0f, 0.0f);
-	Out_CollisionNormal = float3(0.0f, 0.0f, 0.0f);
-
-	if (In_PreviousFrameQueryID >= 0 && uint(In_PreviousFrameQueryID) < MaxRayTraceCount)
-	{
-		FNiagaraRayTracingResult HitResult = IntersectionResults[In_PreviousFrameQueryID + IntersectionRestultsOffest];
-
-		Out_CollisionValid = HitResult.IsHit();
-		if (Out_CollisionValid)
-		{
-			Out_CollisionDistance = HitResult.HitT;
-			Out_CollisionPosWorld = LWCToFloat(MakeLWCVector3(-In_LWCTile, HitResult.WorldPosition)); // TODO: HitResult.WorldPosition should support LWC properly
-			Out_CollisionNormal = HitResult.WorldNormal;
-		}
-	}
-}
-
-#endif
\ No newline at end of file
diff --git a/Engine/Plugins/FX/Niagara/Shaders/Private/NiagaraDataInterfaceCollisionQueryTemplate.ush b/Engine/Plugins/FX/Niagara/Shaders/Private/NiagaraDataInterfaceCollisionQueryTemplate.ush
--- a/Engine/Plugins/FX/Niagara/Shaders/Private/NiagaraDataInterfaceCollisionQueryTemplate.ush
+++ b/Engine/Plugins/FX/Niagara/Shaders/Private/NiagaraDataInterfaceCollisionQueryTemplate.ush
@@ -1,17 +1,6 @@
 // Copyright Epic Games, Inc. All Rights Reserved.
 
 float3										SystemLWCTile_{ParameterName};
-#if NIAGARA_USE_RAYTRACING
-/** Niagara ray tracing uses shared buffers between dispatches so each buffer has a buffer + and offset into that buffer for the current dispatch. */
-bool										RayTracingEnabled_{ParameterName};
-uint										MaxRayTraceCount_{ParameterName};
-RWStructuredBuffer<FNiagaraRayData>			RWRayRequests_{ParameterName};
-uint										RayRequestsOffset_{ ParameterName };
-StructuredBuffer<FNiagaraRayTracingResult>	IntersectionResults_{ParameterName};
-uint										IntersectionResultsOffset_{ ParameterName };
-RWBuffer<uint>								RWRayTraceCounts_{ParameterName};
-uint										RayTraceCountsOffset_{ ParameterName };
-#endif
 
 void QuerySceneDepthGPU_{ParameterName}(in float3 In_SamplePos, out float Out_SceneDepth, out float3 Out_CameraPosWorld, out bool Out_IsInsideView, out float3 Out_WorldPos, out float3 Out_WorldNormal)
 {
@@ -28,55 +17,29 @@ void QueryMeshDistanceFieldGPU_{ParameterName}(in float3 In_SamplePos, out float
 	NDICollisionQuery_QueryMeshDistanceFieldGPU(In_SamplePos, SystemLWCTile_{ParameterName}, Out_DistanceToNearestSurface, Out_FieldGradient, Out_IsDistanceFieldValid);
 }
 
+// DEPRECATE_BEGIN
 void IssueAsyncRayTraceGpu_{ParameterName}_UEImpureCall(int In_QueryID, float3 In_TraceStart, float3 In_TraceEnd, int In_TraceChannel, out bool Out_IsQueryValid)
 {
-#if NIAGARA_USE_RAYTRACING
-	if (RayTracingEnabled_{ParameterName})
-	{
-		NDICollisionQuery_IssueAsyncRayTrace(MaxRayTraceCount_{ParameterName}, RWRayRequests_{ParameterName}, RayRequestsOffset_{ ParameterName }, RWRayTraceCounts_{ParameterName}, RayTraceCountsOffset_{ ParameterName }, SystemLWCTile_{ParameterName}, In_QueryID, In_TraceStart, In_TraceEnd, In_TraceChannel, Out_IsQueryValid);
-		return;
-	}
-#endif
 	Out_IsQueryValid = false;
 }
 
 void ReserveRayTraceIndex_{ParameterName}_UEImpureCall(int In_TraceCount, out int Out_StartQueryID, out bool Out_IndicesValid)
 {
-#if NIAGARA_USE_RAYTRACING
-	if (RayTracingEnabled_{ParameterName})
-	{
-		NDICollisionQuery_ReserveRayTraceIndex(MaxRayTraceCount_{ParameterName}, RWRayTraceCounts_{ParameterName}, RayTraceCountsOffset_{ ParameterName }, In_TraceCount, Out_StartQueryID, Out_IndicesValid);
-		return;
-	}
-#endif
 	Out_StartQueryID = -1;
 	Out_IndicesValid = false;
 }
 
 void CreateAsyncRayTraceGpu_{ParameterName}_UEImpureCall(float3 In_TraceStart, float3 In_TraceEnd, int In_TraceChannel, out int Out_QueryID, out bool Out_IsQueryValid)
 {
-#if NIAGARA_USE_RAYTRACING
-	if (RayTracingEnabled_{ParameterName})
-	{
-		NDICollisionQuery_CreateAsyncRayTrace(MaxRayTraceCount_{ParameterName}, RWRayRequests_{ParameterName}, RayRequestsOffset_{ ParameterName }, RWRayTraceCounts_{ParameterName}, RayTraceCountsOffset_{ ParameterName }, SystemLWCTile_{ParameterName}, In_TraceStart, In_TraceEnd, In_TraceChannel, Out_QueryID, Out_IsQueryValid);
-		return;
-	}
-#endif
 	Out_QueryID = -1;
 	Out_IsQueryValid = false;
 }
 
 void ReadAsyncRayTraceGpu_{ParameterName}(int In_PreviousFrameQueryID, out bool Out_CollisionValid, out float Out_CollisionDistance, out float3 Out_CollisionPosWorld, out float3 Out_CollisionNormal)
 {
-#if NIAGARA_USE_RAYTRACING
-	if (RayTracingEnabled_{ParameterName})
-	{
-		NDICollisionQuery_ReadAsyncRayTrace(MaxRayTraceCount_{ParameterName}, IntersectionResults_{ParameterName}, IntersectionResultsOffset_{ParameterName}, SystemLWCTile_{ParameterName}, In_PreviousFrameQueryID, Out_CollisionValid, Out_CollisionDistance, Out_CollisionPosWorld, Out_CollisionNormal);
-		return;
-	}
-#endif
 	Out_CollisionDistance = 0.0f;
 	Out_CollisionPosWorld = float3(0.0f, 0.0f, 0.0f);
 	Out_CollisionNormal = float3(0.0f, 0.0f, 0.0f);
 	Out_CollisionValid = false;
-}
\ No newline at end of file
+}
+// DEPRECATE_END
diff --git a/Engine/Plugins/FX/Niagara/Shaders/Private/NiagaraRayMarchingShaders.usf b/Engine/Plugins/FX/Niagara/Shaders/Private/NiagaraRayMarchingShaders.usf
new file mode 100644
--- /dev/null
+++ b/Engine/Plugins/FX/Niagara/Shaders/Private/NiagaraRayMarchingShaders.usf
@@ -0,0 +1,72 @@
+// Copyright Epic Games, Inc. All Rights Reserved.
+
+#include "/Engine/Private/Common.ush"
+#include "/Engine/Private/GlobalDistanceFieldShared.ush"
+#include "NiagaraAsyncGpuTraceCommon.ush"
+
+#if !defined(THREADGROUP_SIZE_X)
+	#define THREADGROUP_SIZE_X 32
+#endif
+
+#if !defined(MAX_MARCH_COUNT)
+	#define MAX_MARCH_COUNT 16
+#endif
+
+StructuredBuffer<FNiagaraAsyncGpuTrace> Traces;
+Buffer<uint> TraceCounts;
+RWStructuredBuffer<FNiagaraAsyncGpuTraceResult> Results;
+
+uint TracesOffset;
+uint TraceCountsOffset;
+uint ResultsOffset;
+float DistanceThreshold;
+
+[numthreads(THREADGROUP_SIZE_X, 1, 1)]
+void NiagaraRayMarchGlobalSdfCS(uint GroupIndex : SV_GroupIndex)
+{
+	if (GroupIndex >= TraceCounts[TraceCountsOffset + 0])
+	{
+		return;
+	}
+
+	// read the offset request
+	FNiagaraAsyncGpuTrace Trace = Traces[TracesOffset + GroupIndex];
+
+	FNiagaraAsyncGpuTraceResult Result = (FNiagaraAsyncGpuTraceResult) 0;
+	Result.WorldPosition = Trace.Origin;
+	Result.HitT = 0.0f;
+
+	int MarchCount = 0;
+	for (MarchCount = 0; MarchCount < MAX_MARCH_COUNT; ++MarchCount)
+	{
+		const float CurrentDistance = GetDistanceToNearestSurfaceGlobal(Result.WorldPosition);
+		Result.HitT += CurrentDistance;
+		Result.WorldPosition = Trace.Origin + Trace.Direction * min(Result.HitT, Trace.TFar);
+
+		if (CurrentDistance < DistanceThreshold)
+		{
+			break;
+		}
+		if (Result.HitT > Trace.TFar)
+		{
+			// we've marched passed the end of the ray and haven't found anything
+			Result.HitT = -1;
+			break;
+		}
+	}
+
+	// if we exceeded our iteration count
+	if (MarchCount == MAX_MARCH_COUNT)
+	{
+		Result.HitT = -1;
+	}
+
+	// if we have a hit then generate a normal
+	if (Result.HitT >= 0.0f)
+	{
+		Result.WorldNormal = normalize(GetDistanceFieldGradientGlobal(Result.WorldPosition));
+	}
+
+	// write the offset result
+	Results[ResultsOffset + GroupIndex] = Result;
+}
\ No newline at end of file
diff --git a/Engine/Plugins/FX/Niagara/Shaders/Private/NiagaraRayTracingCommon.ush b/Engine/Plugins/FX/Niagara/Shaders/Private/NiagaraRayTracingCommon.ush
deleted file mode 100644
diff --git a/Engine/Plugins/FX/Niagara/Shaders/Private/NiagaraRayTracingShaders.usf b/Engine/Plugins/FX/Niagara/Shaders/Private/NiagaraRayTracingShaders.usf
--- a/Engine/Plugins/FX/Niagara/Shaders/Private/NiagaraRayTracingShaders.usf
+++ b/Engine/Plugins/FX/Niagara/Shaders/Private/NiagaraRayTracingShaders.usf
@@ -5,14 +5,14 @@
 #include "/Engine/Private/ShadingCommon.ush"
 #include "/Engine/Private/RayTracing/RayTracingCommon.ush"
 #include "/Engine/Private/RayTracing/RayTracingHitGroupCommon.ush"
-#include "NiagaraRayTracingCommon.ush"
+#include "NiagaraAsyncGpuTraceCommon.ush"
 #include "/Engine/Private/HashTable.ush"
 #include "/Engine/Private/RayTracing/VFXTraceRay.ush"
 
 RaytracingAccelerationStructure					TLAS;
-StructuredBuffer<FNiagaraRayData>				Rays;
+StructuredBuffer<FNiagaraAsyncGpuTrace>			Rays;
 uint											RaysOffset;
-RWStructuredBuffer<FNiagaraRayTracingResult>	CollisionOutput;
+RWStructuredBuffer<FNiagaraAsyncGpuTraceResult>	CollisionOutput;
 uint											CollisionOutputOffset;
 
 #if NIAGARA_RAYTRACE_FAKE_INDIRECT
@@ -22,6 +22,7 @@ Buffer<uint>                                  RayTraceCounts;
 Buffer<uint> HashToCollisionGroups;
 uint MaxRetraces;
 
+#if NIAGARA_SUPPORTS_COLLISION_GROUPS
 uint GetGPUSceneInstanceIndex(uint PrimitiveIndex, uint InstanceIndex)
 {
 #if VF_USE_PRIMITIVE_SCENE_DATA
@@ -50,6 +51,7 @@ bool FindCollisionGroup(int GPUSceneIndexId, out int CollisionGroup)
 
 	return false;
 }
+#endif
 
 
 RAY_TRACING_ENTRY_RAYGEN(NiagaraCollisionRayTraceRG)
@@ -63,7 +65,7 @@ RAY_TRACING_ENTRY_RAYGEN(NiagaraCollisionRayTraceRG)
 	}
 #endif
 
-	FNiagaraRayData InputRay = Rays[RayIndex + RaysOffset];
+	FNiagaraAsyncGpuTrace InputRay = Rays[RayIndex + RaysOffset];
 
 	RayDesc Ray;
 	Ray.Origin = InputRay.Origin;
@@ -76,6 +78,7 @@ RAY_TRACING_ENTRY_RAYGEN(NiagaraCollisionRayTraceRG)
 	FVFXTracePayload Payload = (FVFXTracePayload) 0;
 	uint RayCollisionGroup = InputRay.CollisionGroup;
 
+#if NIAGARA_SUPPORTS_COLLISION_GROUPS
 	if (RayCollisionGroup != -1)
 	{
 		uint NumCasts = 0;
@@ -113,6 +116,7 @@ RAY_TRACING_ENTRY_RAYGEN(NiagaraCollisionRayTraceRG)
 		}
 	}
 	else
+#endif
 	{
 		VfxTraceRay(
 			TLAS,   // AccelerationStructure
@@ -127,6 +131,7 @@ RAY_TRACING_ENTRY_RAYGEN(NiagaraCollisionRayTraceRG)
 	}
 
 	uint OutputIndex = RayIndex + CollisionOutputOffset;
+	CollisionOutput[OutputIndex] = (FNiagaraAsyncGpuTraceResult) 0;
 	CollisionOutput[OutputIndex].HitT = Payload.HitT;
 	if (Payload.IsHit())
 	{
diff --git a/Engine/Plugins/FX/Niagara/Source/Niagara/Classes/NiagaraDataInterfaceAsyncGpuTrace.h b/Engine/Plugins/FX/Niagara/Source/Niagara/Classes/NiagaraDataInterfaceAsyncGpuTrace.h
new file mode 100644
--- /dev/null
+++ b/Engine/Plugins/FX/Niagara/Source/Niagara/Classes/NiagaraDataInterfaceAsyncGpuTrace.h
@@ -0,0 +1,68 @@
+// Copyright Epic Games, Inc. All Rights Reserved.
+#pragma once
+
+#include "NiagaraCommon.h"
+#include "NiagaraDataInterface.h"
+#include "NiagaraSettings.h"
+#include "NiagaraShared.h"
+
+#include "NiagaraDataInterfaceAsyncGpuTrace.generated.h"
+
+class FNiagaraSystemInstance;
+
+/** Data interface for handling latent (delayed 1 frame) traces against the scene for GPU simulations. */
+UCLASS(EditInlineNew, Category = "Collision", meta = (DisplayName = "Async Gpu Trace"))
+class NIAGARA_API UNiagaraDataInterfaceAsyncGpuTrace : public UNiagaraDataInterface
+{
+	GENERATED_UCLASS_BODY()
+public:
+
+	DECLARE_NIAGARA_DI_PARAMETER();
+
+	/** The maximum number of traces (per particle) that can be created per frame.  Defines the size of the preallocated 
+		buffer that is used to contain the traces. */
+	UPROPERTY(EditAnywhere, Category = "Async GPU Trace")
+	int32 MaxTracesPerParticle = 1;
+
+	/** If a collision is rejected, how many times do we attempt to retrace from that collision point forward to find a new, valid collision.*/
+	UPROPERTY(EditAnywhere, Category = "Async GPU Trace")
+	int32 MaxRetraces = 0;
+
+	/** Defines which technique is used to resolve the trace - see 'AsyncGpuTraceDI/Provider Priority' in the Niagara project settings. */
+	UPROPERTY(EditAnywhere, Category = "Async GPU Trace")
+	TEnumAsByte<ENDICollisionQuery_AsyncGpuTraceProvider::Type> TraceProvider = ENDICollisionQuery_AsyncGpuTraceProvider::Default;
+
+	//UObject Interface
+	virtual void PostInitProperties() override;
+	virtual void PostLoad() override;
+
+#if WITH_EDITOR
+	virtual void PostEditChangeProperty(struct FPropertyChangedEvent& PropertyChangedEvent) override;
+#endif
+	//UObject Interface End
+
+	virtual int32 PerInstanceDataSize() const override;
+	virtual bool InitPerInstanceData(void* PerInstanceData, FNiagaraSystemInstance* InSystemInstance) override;
+	virtual void DestroyPerInstanceData(void* PerInstanceData, FNiagaraSystemInstance* InSystemInstance) override;
+
+	virtual void GetFunctions(TArray<FNiagaraFunctionSignature>& OutFunctions) override;
+
+	virtual bool CanExecuteOnTarget(ENiagaraSimTarget Target) const override { return Target == ENiagaraSimTarget::GPUComputeSim; }
+	virtual bool RequiresDistanceFieldData() const override;
+	virtual bool RequiresRayTracingScene() const override;
+
+#if WITH_EDITORONLY_DATA
+	virtual bool AppendCompileHash(FNiagaraCompileHashVisitor* InVisitor) const override;
+	virtual void GetCommonHLSL(FString& OutHLSL) override;
+	virtual void GetParameterDefinitionHLSL(const FNiagaraDataInterfaceGPUParamInfo& ParamInfo, FString& OutHLSL) override;
+	virtual bool GetFunctionHLSL(const FNiagaraDataInterfaceGPUParamInfo& ParamInfo, const FNiagaraDataInterfaceGeneratedFunction& FunctionInfo, int FunctionInstanceIndex, FString& OutHLSL) override;
+	virtual bool UpgradeFunctionCall(FNiagaraFunctionSignature& FunctionSignature) override;
+#endif
+
+private:
+	virtual bool Equals(const UNiagaraDataInterface* Other) const override;
+
+protected:
+	virtual bool CopyToInternal(UNiagaraDataInterface* Destination) const override;
+	virtual void PushToRenderThreadImpl() override;
+};
diff --git a/Engine/Plugins/FX/Niagara/Source/Niagara/Classes/NiagaraDataInterfaceCollisionQuery.h b/Engine/Plugins/FX/Niagara/Source/Niagara/Classes/NiagaraDataInterfaceCollisionQuery.h
--- a/Engine/Plugins/FX/Niagara/Source/Niagara/Classes/NiagaraDataInterfaceCollisionQuery.h
+++ b/Engine/Plugins/FX/Niagara/Source/Niagara/Classes/NiagaraDataInterfaceCollisionQuery.h
@@ -29,21 +29,8 @@ public:
 
 	DECLARE_NIAGARA_DI_PARAMETER();
 
-	UPROPERTY(EditAnywhere, Category = "Ray Trace")
-	int32 MaxTracesPerParticle = 0;
-
-	/** If a collision is rejected, how many times do we attempt to retrace from that collision point forward to find a new, valid collision.*/
-	UPROPERTY(EditAnywhere, Category = "Ray Trace")
-	int32 MaxRetraces = 0;
-
-	FNiagaraSystemInstance *SystemInstance;
-
 	//UObject Interface
 	virtual void PostInitProperties() override;
-	virtual void PostLoad() override;
-#if WITH_EDITOR
-	virtual void PostEditChangeProperty(struct FPropertyChangedEvent& PropertyChangedEvent) override;
-#endif
 	//UObject Interface End
 
 	/** Initializes the per instance data for this interface. Returns false if there was some error and the simulation should be disabled. */
@@ -59,7 +46,6 @@ public:
 	virtual void GetFunctions(TArray<FNiagaraFunctionSignature>& OutFunctions) override;
 	virtual void GetVMExternalFunction(const FVMExternalFunctionBindingInfo& BindingInfo, void* InstanceData, FVMExternalFunction &OutFunc) override;
 	virtual void GetAssetTagsForContext(const UObject* InAsset, const TArray<const UNiagaraDataInterface*>& InProperties, TMap<FName, uint32>& NumericKeys, TMap<FName, FString>& StringKeys) const override;
-	virtual void PushToRenderThreadImpl() override;
 
 	// VM functions
 	void PerformQuerySyncCPU(FVectorVMExternalFunctionContext& Context);
@@ -68,11 +54,9 @@ public:
 	virtual bool CanExecuteOnTarget(ENiagaraSimTarget Target) const override { return true; }
 	virtual bool RequiresDistanceFieldData() const override { return true; }
 	virtual bool RequiresDepthBuffer() const override { return true; }
-	virtual bool RequiresRayTracingScene() const override;
 
 #if WITH_EDITORONLY_DATA
 	virtual bool AppendCompileHash(FNiagaraCompileHashVisitor* InVisitor) const override;
-	virtual void ModifyCompilationEnvironment(EShaderPlatform ShaderPlatform, struct FShaderCompilerEnvironment& OutEnvironment) const override;
 	virtual void GetCommonHLSL(FString& OutHLSL) override;
 	virtual void GetParameterDefinitionHLSL(const FNiagaraDataInterfaceGPUParamInfo& ParamInfo, FString& OutHLSL) override;
 	virtual bool GetFunctionHLSL(const FNiagaraDataInterfaceGPUParamInfo& ParamInfo, const FNiagaraDataInterfaceGeneratedFunction& FunctionInfo, int FunctionInstanceIndex, FString& OutHLSL) override;
@@ -88,14 +72,6 @@ public:
 	virtual bool PostSimulateCanOverlapFrames() const { return false; }
 
 private:
-
-	virtual bool Equals(const UNiagaraDataInterface* Other) const override;
-
-protected:
-	virtual bool CopyToInternal(UNiagaraDataInterface* Destination) const override;
-
-private:
-
 	static FCriticalSection CriticalSection;
 	UEnum* TraceChannelEnum;
 };
diff --git a/Engine/Plugins/FX/Niagara/Source/Niagara/Private/DataInterface/NiagaraDataInterfaceAsyncGpuTrace.cpp b/Engine/Plugins/FX/Niagara/Source/Niagara/Private/DataInterface/NiagaraDataInterfaceAsyncGpuTrace.cpp
new file mode 100644
--- /dev/null
+++ b/Engine/Plugins/FX/Niagara/Source/Niagara/Private/DataInterface/NiagaraDataInterfaceAsyncGpuTrace.cpp
@@ -0,0 +1,589 @@
+// Copyright Epic Games, Inc. All Rights Reserved.
+
+#include "NiagaraDataInterfaceAsyncGpuTrace.h"
+
+#include "GlobalDistanceFieldParameters.h"
+#include "NiagaraAsyncGpuTraceHelper.h"
+#include "NiagaraComponent.h"
+#include "NiagaraDataInterfaceUtilities.h"
+#include "NiagaraGpuComputeDispatchInterface.h"
+#include "NiagaraGpuComputeDispatch.h"
+#include "NiagaraSimStageData.h"
+#include "NiagaraStats.h"
+#include "NiagaraTypes.h"
+#include "NiagaraWorldManager.h"
+#include "RenderResource.h"
+#include "Shader.h"
+#include "ShaderCore.h"
+#include "ShaderCompilerCore.h"
+#include "ShaderParameterUtils.h"
+
+#define LOCTEXT_NAMESPACE "NiagaraDataInterfaceAsyncGpuTrace"
+
+namespace NDIAsyncGpuTraceLocal
+{
+	static const TCHAR* CommonShaderFile = TEXT("/Plugin/FX/Niagara/Private/NiagaraDataInterfaceAsyncGpuTrace.ush");
+	static const TCHAR* TemplateShaderFile = TEXT("/Plugin/FX/Niagara/Private/NiagaraDataInterfaceAsyncGpuTraceTemplate.ush");
+
+	static const FName IssueAsyncRayTraceName(TEXT("IssueAsyncRayTraceGpu"));
+	static const FName CreateAsyncRayTraceName(TEXT("CreateAsyncRayTraceGpu"));
+	static const FName ReserveAsyncRayTraceName(TEXT("ReserveAsyncRayTraceGpu"));
+	static const FName ReadAsyncRayTraceName(TEXT("ReadAsyncRayTraceGpu"));
+
+	static const FString MaxRayTraceCountParamName(TEXT("MaxRayTraceCount_"));
+	static const FString RayRequestsParamName(TEXT("RayRequests_"));
+	static const FString RayRequestsOffsetParamName(TEXT("RayRequestsOffset_"));
+	static const FString IntersectionResultsParamName(TEXT("IntersectionResults_"));
+	static const FString IntersectionResultsOffsetParamName(TEXT("IntersectionResultsOffset_"));
+	static const FString RayTraceCountsParamName(TEXT("RayTraceCounts_"));
+	static const FString RayTraceCountsOffsetParamName(TEXT("RayTraceCountsOffset_"));
+	static const FString SystemLWCTileName(TEXT("SystemLWCTile_"));
+
+	struct FPerInstanceData
+	{
+		FNiagaraSystemInstanceID InstanceID;
+		ENDICollisionQuery_AsyncGpuTraceProvider::Type ProviderType;
+		bool RequiresAsyncTraces;
+	};
+}
+
+struct FNiagaraAsyncGpuTraceDIFunctionVersion
+{
+	enum Type
+	{
+		InitialVersion = 0,
+
+		VersionPlusOne,
+		LatestVersion = VersionPlusOne - 1
+	};
+};
+
+struct FNiagaraDataIntefaceProxyAsyncGpuTrace : public FNiagaraDataInterfaceProxy
+{
+	struct FInstanceData
+	{
+		int32 MaxTracesPerParticle = 0;
+		uint32 MaxRetraces = 0;
+		TEnumAsByte<ENDICollisionQuery_AsyncGpuTraceProvider::Type> ProviderType = ENDICollisionQuery_AsyncGpuTraceProvider::Default;
+	};
+
+	using ProxyDataMap = TMap<FNiagaraSystemInstanceID, FInstanceData>;
+	ProxyDataMap SystemInstancesToProxyData_RT;
+
+	virtual ~FNiagaraDataIntefaceProxyAsyncGpuTrace()
+	{
+	}
+
+	virtual int32 PerInstanceDataPassedToRenderThreadSize() const override
+	{
+		return 0;
+	}
+
+	virtual void PreStage(FRHICommandList& RHICmdList, const FNiagaraDataInterfaceStageArgs& Context) override
+	{
+		FNiagaraDataInterfaceProxy::PreStage(RHICmdList, Context);
+
+		if (const FInstanceData* InstanceData_RT = SystemInstancesToProxyData_RT.Find(Context.SystemInstanceID))
+		{
+			if (InstanceData_RT->MaxTracesPerParticle > 0)
+			{
+				//Accumulate the total ray requests for this DI for all dispatches in the stage.
+				const int32 RayRequests = InstanceData_RT->MaxTracesPerParticle * Context.SimStageData->DestinationNumInstances;
+				Context.ComputeDispatchInterface->GetAsyncGpuTraceHelper().AddToDispatch(
+					this,
+					RayRequests,
+					InstanceData_RT->MaxRetraces,
+					InstanceData_RT->ProviderType);
+
+			}
+		}
+	}
+
+	virtual bool RequiresPreStageFinalize() const override
+	{
+		return true;
+	}
+
+	virtual void FinalizePreStage(FRHICommandList& RHICmdList, const FNiagaraGpuComputeDispatchInterface* ComputeDispatchInterface) override
+	{
+		FNiagaraAsyncGpuTraceHelper& TraceHelper = ComputeDispatchInterface->GetAsyncGpuTraceHelper();
+		if (SystemInstancesToProxyData_RT.Num())
+		{
+			TraceHelper.BuildDispatch(RHICmdList, this);
+		}
+		else
+		{
+			TraceHelper.BuildDummyDispatch(RHICmdList);
+		}
+	}
+};
+
+////////////////////////////////////////////////////////////////////////////////
+/// 
+////////////////////////////////////////////////////////////////////////////////
+
+UNiagaraDataInterfaceAsyncGpuTrace::UNiagaraDataInterfaceAsyncGpuTrace(FObjectInitializer const& ObjectInitializer)
+	: Super(ObjectInitializer)
+{
+    Proxy.Reset(new FNiagaraDataIntefaceProxyAsyncGpuTrace());
+}
+
+int32 UNiagaraDataInterfaceAsyncGpuTrace::PerInstanceDataSize() const
+{
+	return sizeof(NDIAsyncGpuTraceLocal::FPerInstanceData);
+}
+
+bool UNiagaraDataInterfaceAsyncGpuTrace::InitPerInstanceData(void* PerInstanceData, FNiagaraSystemInstance* InSystemInstance)
+{
+	NDIAsyncGpuTraceLocal::FPerInstanceData* PIData = new (PerInstanceData) NDIAsyncGpuTraceLocal::FPerInstanceData;
+	PIData->InstanceID = InSystemInstance->GetId();
+	PIData->RequiresAsyncTraces = false;
+
+	if (TraceProvider == ENDICollisionQuery_AsyncGpuTraceProvider::None)
+	{
+		return true;
+	}
+	
+	FNiagaraDataInterfaceUtilities::ForEachGpuFunction(this, InSystemInstance, [&](const FNiagaraDataInterfaceGeneratedFunction Function)
+	{
+		if (Function.DefinitionName == NDIAsyncGpuTraceLocal::IssueAsyncRayTraceName
+			|| Function.DefinitionName == NDIAsyncGpuTraceLocal::CreateAsyncRayTraceName
+			|| Function.DefinitionName == NDIAsyncGpuTraceLocal::ReserveAsyncRayTraceName
+			|| Function.DefinitionName == NDIAsyncGpuTraceLocal::ReadAsyncRayTraceName)
+		{
+			PIData->RequiresAsyncTraces = true;
+		}
+
+		return !PIData->RequiresAsyncTraces;
+	});
+
+	PIData->ProviderType = FNiagaraAsyncGpuTraceProvider::ResolveSupportedType(TraceProvider, GetDefault<UNiagaraSettings>()->NDICollisionQuery_AsyncGpuTraceProviderOrder);
+
+	// if nothing is supported then we shouldn't bother with initializing things
+	if (PIData->ProviderType == ENDICollisionQuery_AsyncGpuTraceProvider::None)
+	{
+		UE_LOG(LogNiagara, Warning, TEXT("Failed to find a supported provider when initializing %s with TraceProvider = %s.  Failed to initialize System %s."),
+			*GetPathNameSafe(this),
+			*StaticEnum<ENDICollisionQuery_AsyncGpuTraceProvider::Type>()->GetDisplayValueAsText(TraceProvider).ToString(),
+			*GetPathNameSafe(InSystemInstance->GetSystem()));
+		return false;
+	}
+
+	if (PIData->RequiresAsyncTraces && MaxTracesPerParticle)
+	{
+		// Push Updates to Proxy.
+		FNiagaraDataIntefaceProxyAsyncGpuTrace* RT_Proxy = GetProxyAs<FNiagaraDataIntefaceProxyAsyncGpuTrace>();
+		ENQUEUE_RENDER_COMMAND(FUpdateData)(
+			[RT_Proxy,
+			RT_MaxTraceCount = MaxTracesPerParticle,
+			RT_MaxRetraceCount = MaxRetraces,
+			RT_ProviderType = PIData->ProviderType,
+			RT_InstanceID = InSystemInstance->GetId()](FRHICommandListImmediate& RHICmdList)
+		{
+			check(!RT_Proxy->SystemInstancesToProxyData_RT.Contains(RT_InstanceID));
+			FNiagaraDataIntefaceProxyAsyncGpuTrace::FInstanceData& TargetData = RT_Proxy->SystemInstancesToProxyData_RT.Add(RT_InstanceID);
+
+			TargetData.MaxTracesPerParticle = RT_MaxTraceCount;
+			TargetData.MaxRetraces = RT_MaxRetraceCount;
+			TargetData.ProviderType = RT_ProviderType;
+		});
+	}
+
+	return true;
+}
+
+void UNiagaraDataInterfaceAsyncGpuTrace::DestroyPerInstanceData(void* PerInstanceData, FNiagaraSystemInstance* InSystemInstance)
+{
+	NDIAsyncGpuTraceLocal::FPerInstanceData* PIData = reinterpret_cast<NDIAsyncGpuTraceLocal::FPerInstanceData*>(PerInstanceData);
+	PIData->~FPerInstanceData();
+
+	FNiagaraDataIntefaceProxyAsyncGpuTrace* RT_Proxy = GetProxyAs<FNiagaraDataIntefaceProxyAsyncGpuTrace>();
+	ENQUEUE_RENDER_COMMAND(FNiagaraDIDestroyInstanceData) (
+		[RT_Proxy, RT_InstanceID = InSystemInstance->GetId()](FRHICommandListImmediate& CmdList)
+	{
+		RT_Proxy->SystemInstancesToProxyData_RT.Remove(RT_InstanceID);
+	});
+}
+
+void UNiagaraDataInterfaceAsyncGpuTrace::PostInitProperties()
+{
+	Super::PostInitProperties();
+
+	//Can we register data interfaces as regular types and fold them into the FNiagaraVariable framework for UI and function calls etc?
+	if (HasAnyFlags(RF_ClassDefaultObject))
+	{
+		ENiagaraTypeRegistryFlags Flags = ENiagaraTypeRegistryFlags::AllowAnyVariable | ENiagaraTypeRegistryFlags::AllowParameter;
+		FNiagaraTypeRegistry::Register(FNiagaraTypeDefinition(GetClass()), Flags);
+	}
+}
+
+void UNiagaraDataInterfaceAsyncGpuTrace::PostLoad()
+{
+	Super::PostLoad();
+
+	MarkRenderDataDirty();
+}
+
+void UNiagaraDataInterfaceAsyncGpuTrace::GetFunctions(TArray<FNiagaraFunctionSignature>& OutFunctions)
+{
+	const FText RayTraceStartWorldDescription = LOCTEXT("TraceStartWorldDescription", "Ray starting point in world space");
+	const FText RayTraceEndWorldDescription = LOCTEXT("TraceEndWorldDescription", "Ray end point in world space");
+	const FText CollisionGroupDescription = LOCTEXT("CollisionGroupDescription", "Collision group index of the primitives we wish to skip");
+	const FText QueryIDDescription = LOCTEXT("QueryIDDescription", "Unique (for this frame) index of the query being enqueued (used in subsequent frames to retrieve results).");
+	const FText CollisionPosWorldDescription = LOCTEXT("CollisionPosWorldDescription", "If the collision is valid, this returns the location of the blocking hit.");
+	const FText CollisionNormalDescription = LOCTEXT("CollisionNormalDescription", "The surface normal of the world geometry at the point of intersection");
+	const FText PreviousFrameQueryIDDescription = LOCTEXT("PreviousFrameQueryIDDescription", "The query ID returned from the last frame's async trace call.\nRegardless if it is a valid ID or not this function call with issue a new async line trace, but it will only return results with a valid ID.");
+	const FText ExperimentalMessage = LOCTEXT("AsyncRayTraceExperimental", "Feature remains in an experimental stage, with some features (eg. collision groups) limited to a subset of the providers.");
+
+	FNiagaraVariable CollisionGroupVariable(FNiagaraTypeDefinition::GetIntDef(), TEXT("CollisionGroup"));
+	CollisionGroupVariable.SetValue(INDEX_NONE);
+
+	{
+		const FText IssueValidDescription = LOCTEXT("IssueValidDescription", "Returns true if the query was issued");
+
+		FNiagaraFunctionSignature& IssueRayTrace = OutFunctions.AddDefaulted_GetRef();
+		IssueRayTrace.Name = NDIAsyncGpuTraceLocal::IssueAsyncRayTraceName;
+		IssueRayTrace.bRequiresExecPin = true;
+		IssueRayTrace.bMemberFunction = true;
+		IssueRayTrace.bSupportsCPU = false;
+		IssueRayTrace.bExperimental = true;
+#if WITH_EDITORONLY_DATA
+		IssueRayTrace.FunctionVersion = FNiagaraAsyncGpuTraceDIFunctionVersion::LatestVersion;
+		IssueRayTrace.Description = LOCTEXT("IssueAsyncRayTraceDescription", "Enqueues a GPU raytrace with the result being available the following frame");
+		IssueRayTrace.ExperimentalMessage = ExperimentalMessage;
+#endif
+		IssueRayTrace.AddInput(FNiagaraVariable(FNiagaraTypeDefinition(GetClass()), TEXT("AsyncGpuTrace")));
+		IssueRayTrace.AddInput(FNiagaraVariable(FNiagaraTypeDefinition::GetIntDef(), TEXT("QueryID")), QueryIDDescription);
+		IssueRayTrace.AddInput(FNiagaraVariable(FNiagaraTypeDefinition::GetPositionDef(), TEXT("TraceStartWorld")), RayTraceStartWorldDescription);
+		IssueRayTrace.AddInput(FNiagaraVariable(FNiagaraTypeDefinition::GetPositionDef(), TEXT("TraceEndWorld")), RayTraceEndWorldDescription);
+		IssueRayTrace.AddInput(CollisionGroupVariable, CollisionGroupDescription);
+		IssueRayTrace.AddOutput(FNiagaraVariable(FNiagaraTypeDefinition::GetBoolDef(), TEXT("IsQueryValid")), IssueValidDescription);
+	}
+
+	{
+		const FText CreateValidDescription = LOCTEXT("CreateValidDescription", "Returns true if the query was created");
+
+		FNiagaraFunctionSignature& CreateRayTrace = OutFunctions.AddDefaulted_GetRef();
+		CreateRayTrace.Name = NDIAsyncGpuTraceLocal::CreateAsyncRayTraceName;
+		CreateRayTrace.bRequiresExecPin = true;
+		CreateRayTrace.bMemberFunction = true;
+		CreateRayTrace.bSupportsCPU = false;
+		CreateRayTrace.bExperimental = true;
+#if WITH_EDITORONLY_DATA
+		CreateRayTrace.FunctionVersion = FNiagaraAsyncGpuTraceDIFunctionVersion::LatestVersion;
+		CreateRayTrace.Description = LOCTEXT("CreateAsyncRayTraceDescription", "Creates a GPU raytrace with the result being available the following frame (index is returned)");
+		CreateRayTrace.ExperimentalMessage = ExperimentalMessage;
+#endif
+		CreateRayTrace.AddInput(FNiagaraVariable(FNiagaraTypeDefinition(GetClass()), TEXT("AsyncGpuTrace")));
+		CreateRayTrace.AddInput(FNiagaraVariable(FNiagaraTypeDefinition::GetPositionDef(), TEXT("TraceStartWorld")), RayTraceStartWorldDescription);
+		CreateRayTrace.AddInput(FNiagaraVariable(FNiagaraTypeDefinition::GetPositionDef(), TEXT("TraceEndWorld")), RayTraceEndWorldDescription);
+		CreateRayTrace.AddInput(CollisionGroupVariable, CollisionGroupDescription);
+		CreateRayTrace.AddOutput(FNiagaraVariable(FNiagaraTypeDefinition::GetIntDef(), TEXT("QueryID")), QueryIDDescription);
+		CreateRayTrace.AddOutput(FNiagaraVariable(FNiagaraTypeDefinition::GetBoolDef(), TEXT("IsQueryValid")), CreateValidDescription);
+	}
+
+	{
+		const FText TraceCountDescription = LOCTEXT("QueryIDDescription", "Number of async raytrace requests to be reserved");
+		const FText FirstQueryIDValidDescription = LOCTEXT("FirstQueryIDValidDescription", "The first index in the block reserved through this call");
+		const FText ReserveValidDescription = LOCTEXT("ReserveValidDescription", "Returns true if the requested indices were reserved");
+
+		FNiagaraFunctionSignature& ReserveRayTrace = OutFunctions.AddDefaulted_GetRef();
+		ReserveRayTrace.Name = NDIAsyncGpuTraceLocal::ReserveAsyncRayTraceName;
+		ReserveRayTrace.bRequiresExecPin = true;
+		ReserveRayTrace.bMemberFunction = true;
+		ReserveRayTrace.bSupportsCPU = false;
+		ReserveRayTrace.bExperimental = true;
+#if WITH_EDITORONLY_DATA
+		ReserveRayTrace.FunctionVersion = FNiagaraAsyncGpuTraceDIFunctionVersion::LatestVersion;
+		ReserveRayTrace.Description = LOCTEXT("ReserveAsyncRayTraceDescription", "Reserves a number of ray trace request slots");
+		ReserveRayTrace.ExperimentalMessage = ExperimentalMessage;
+#endif
+		ReserveRayTrace.AddInput(FNiagaraVariable(FNiagaraTypeDefinition(GetClass()), TEXT("AsyncGpuTrace")));
+		ReserveRayTrace.AddInput(FNiagaraVariable(FNiagaraTypeDefinition::GetIntDef(), TEXT("TraceCount")), TraceCountDescription);
+		ReserveRayTrace.AddOutput(FNiagaraVariable(FNiagaraTypeDefinition::GetIntDef(), TEXT("FirstQueryID")), FirstQueryIDValidDescription);
+		ReserveRayTrace.AddOutput(FNiagaraVariable(FNiagaraTypeDefinition::GetBoolDef(), TEXT("IsQueryValid")), ReserveValidDescription);
+	}
+
+	{
+		const FText CollisionValidDescription = LOCTEXT("AsyncCollisionValidDescription", "Returns true if the a Hit was encountered");
+		const FText CollisionDistanceDescription = LOCTEXT("CollisionDistanceDescription", "The distance in world space from the ray starting point to the intersection");
+
+		FNiagaraFunctionSignature& ReadRayTrace = OutFunctions.AddDefaulted_GetRef();
+		ReadRayTrace.Name = NDIAsyncGpuTraceLocal::ReadAsyncRayTraceName;
+		ReadRayTrace.bMemberFunction = true;
+		ReadRayTrace.bSupportsCPU = false;
+		ReadRayTrace.bExperimental = true;
+#if WITH_EDITORONLY_DATA
+		ReadRayTrace.FunctionVersion = FNiagaraAsyncGpuTraceDIFunctionVersion::LatestVersion;
+		ReadRayTrace.Description = LOCTEXT("ReadAsyncRayTraceDescription", "Reads the results of a previously enqueued GPU ray trace");
+		ReadRayTrace.ExperimentalMessage = ExperimentalMessage;
+#endif
+		ReadRayTrace.AddInput(FNiagaraVariable(FNiagaraTypeDefinition(GetClass()), TEXT("AsyncGpuTrace")));
+		ReadRayTrace.AddInput(FNiagaraVariable(FNiagaraTypeDefinition::GetIntDef(), TEXT("PreviousFrameQueryID")), PreviousFrameQueryIDDescription);
+		ReadRayTrace.AddOutput(FNiagaraVariable(FNiagaraTypeDefinition::GetBoolDef(), TEXT("CollisionValid")), CollisionValidDescription);
+		ReadRayTrace.AddOutput(FNiagaraVariable(FNiagaraTypeDefinition::GetFloatDef(), TEXT("CollisionDistance")), CollisionDistanceDescription);
+		ReadRayTrace.AddOutput(FNiagaraVariable(FNiagaraTypeDefinition::GetPositionDef(), TEXT("CollisionPosWorld")), CollisionPosWorldDescription);
+		ReadRayTrace.AddOutput(FNiagaraVariable(FNiagaraTypeDefinition::GetVec3Def(), TEXT("CollisionNormal")), CollisionNormalDescription);
+	}
+}
+
+bool UNiagaraDataInterfaceAsyncGpuTrace::RequiresDistanceFieldData() const
+{
+	return FNiagaraAsyncGpuTraceHelper::RequiresDistanceFieldData(TraceProvider);
+}
+
+bool UNiagaraDataInterfaceAsyncGpuTrace::RequiresRayTracingScene() const
+{
+	return FNiagaraAsyncGpuTraceHelper::RequiresRayTracingScene(TraceProvider);
+}
+
+#if WITH_EDITORONLY_DATA
+bool UNiagaraDataInterfaceAsyncGpuTrace::GetFunctionHLSL(const FNiagaraDataInterfaceGPUParamInfo& ParamInfo, const FNiagaraDataInterfaceGeneratedFunction& FunctionInfo, int FunctionInstanceIndex, FString& OutHLSL)
+{
+	if ( (FunctionInfo.DefinitionName == NDIAsyncGpuTraceLocal::IssueAsyncRayTraceName) ||
+		 (FunctionInfo.DefinitionName == NDIAsyncGpuTraceLocal::CreateAsyncRayTraceName) ||
+		 (FunctionInfo.DefinitionName == NDIAsyncGpuTraceLocal::ReserveAsyncRayTraceName) ||
+		 (FunctionInfo.DefinitionName == NDIAsyncGpuTraceLocal::ReadAsyncRayTraceName) )
+	{
+		return true;
+	}
+
+	return false;
+}
+
+bool UNiagaraDataInterfaceAsyncGpuTrace::UpgradeFunctionCall(FNiagaraFunctionSignature& FunctionSignature)
+{
+	// always upgrade to the latest version
+	if (FunctionSignature.FunctionVersion < FNiagaraAsyncGpuTraceDIFunctionVersion::LatestVersion)
+	{
+		TArray<FNiagaraFunctionSignature> AllFunctions;
+		GetFunctions(AllFunctions);
+		for (const FNiagaraFunctionSignature& Sig : AllFunctions)
+		{
+			if (FunctionSignature.Name == Sig.Name)
+			{
+				FunctionSignature = Sig;
+				return true;
+			}
+		}
+	}
+
+	return false;
+}
+
+void UNiagaraDataInterfaceAsyncGpuTrace::GetParameterDefinitionHLSL(const FNiagaraDataInterfaceGPUParamInfo& ParamInfo, FString& OutHLSL)
+{
+	TMap<FString, FStringFormatArg> TemplateArgs =
+	{
+		{TEXT("ParameterName"),	ParamInfo.DataInterfaceHLSLSymbol},
+	};
+
+	FString TemplateFile;
+	LoadShaderSourceFile(NDIAsyncGpuTraceLocal::TemplateShaderFile, EShaderPlatform::SP_PCD3D_SM5, &TemplateFile, nullptr);
+	OutHLSL += FString::Format(*TemplateFile, TemplateArgs);
+}
+
+void UNiagaraDataInterfaceAsyncGpuTrace::GetCommonHLSL(FString& OutHlsl)
+{
+	OutHlsl.Appendf(TEXT("#include \"%s\"\n"), NDIAsyncGpuTraceLocal::CommonShaderFile);
+}
+
+bool UNiagaraDataInterfaceAsyncGpuTrace::AppendCompileHash(FNiagaraCompileHashVisitor* InVisitor) const
+{
+	if (!Super::AppendCompileHash(InVisitor))
+	{
+		return false;
+	}
+
+	InVisitor->UpdateString(TEXT("NDIAsyncGpuTraceCommonHLSLSource"), GetShaderFileHash(NDIAsyncGpuTraceLocal::CommonShaderFile, EShaderPlatform::SP_PCD3D_SM5).ToString());
+	InVisitor->UpdateString(TEXT("NDIAsyncGpuTraceTemplateHLSLSource"), GetShaderFileHash(NDIAsyncGpuTraceLocal::TemplateShaderFile, EShaderPlatform::SP_PCD3D_SM5).ToString());
+
+	return true;
+}
+
+#endif
+
+bool UNiagaraDataInterfaceAsyncGpuTrace::Equals(const UNiagaraDataInterface* Other) const
+{
+	if (!Super::Equals(Other))
+	{
+		return false;
+	}
+
+	const UNiagaraDataInterfaceAsyncGpuTrace* OtherTyped = CastChecked<const UNiagaraDataInterfaceAsyncGpuTrace>(Other);
+	return OtherTyped->MaxTracesPerParticle == MaxTracesPerParticle
+		&& OtherTyped->MaxRetraces == MaxRetraces
+		&& OtherTyped->TraceProvider == TraceProvider;
+}
+
+bool UNiagaraDataInterfaceAsyncGpuTrace::CopyToInternal(UNiagaraDataInterface* Destination) const
+{
+	if (!Super::CopyToInternal(Destination))
+	{
+		return false;
+	}
+
+	UNiagaraDataInterfaceAsyncGpuTrace* OtherTyped = CastChecked<UNiagaraDataInterfaceAsyncGpuTrace>(Destination);
+	OtherTyped->MaxTracesPerParticle = MaxTracesPerParticle;
+	OtherTyped->MaxRetraces = MaxRetraces;
+	OtherTyped->TraceProvider = TraceProvider;
+	OtherTyped->MarkRenderDataDirty();
+	return true;
+}
+
+void UNiagaraDataInterfaceAsyncGpuTrace::PushToRenderThreadImpl()
+{
+	FNiagaraDataIntefaceProxyAsyncGpuTrace* RT_Proxy = GetProxyAs<FNiagaraDataIntefaceProxyAsyncGpuTrace>();
+
+	ENDICollisionQuery_AsyncGpuTraceProvider::Type ResolvedType = FNiagaraAsyncGpuTraceHelper::ResolveSupportedType(TraceProvider);
+
+	ENQUEUE_RENDER_COMMAND(FUpdateData)(
+		[RT_Proxy,
+		RT_MaxTraceCount = MaxTracesPerParticle,
+		RT_MaxRetraceCount = MaxRetraces,
+		RT_ProviderType = ResolvedType](FRHICommandListImmediate& RHICmdList)
+	{
+		for (FNiagaraDataIntefaceProxyAsyncGpuTrace::ProxyDataMap::TIterator ProxyDataIt(RT_Proxy->SystemInstancesToProxyData_RT); ProxyDataIt; ++ProxyDataIt)
+		{
+			if (RT_MaxTraceCount)
+			{
+				FNiagaraDataIntefaceProxyAsyncGpuTrace::FInstanceData& InstanceData = ProxyDataIt.Value();
+				InstanceData.MaxTracesPerParticle = RT_MaxTraceCount;
+				InstanceData.MaxRetraces = RT_MaxRetraceCount;
+				InstanceData.ProviderType = RT_ProviderType;
+			}
+			else
+			{
+				ProxyDataIt.RemoveCurrent();
+			}
+		}
+	});
+}
+
+#if WITH_EDITOR
+void UNiagaraDataInterfaceAsyncGpuTrace::PostEditChangeProperty(struct FPropertyChangedEvent& PropertyChangedEvent)
+{
+	Super::PostEditChangeProperty(PropertyChangedEvent);
+
+	if (PropertyChangedEvent.Property)
+	{
+		const FName PropertyName = PropertyChangedEvent.Property->GetFName();
+
+		if ((PropertyName == GET_MEMBER_NAME_CHECKED(UNiagaraDataInterfaceAsyncGpuTrace, MaxTracesPerParticle))
+		|| (PropertyName == GET_MEMBER_NAME_CHECKED(UNiagaraDataInterfaceAsyncGpuTrace, MaxRetraces))
+		|| (PropertyName == GET_MEMBER_NAME_CHECKED(UNiagaraDataInterfaceAsyncGpuTrace, TraceProvider)))
+		{
+			MarkRenderDataDirty();
+		}
+	}
+}
+#endif
+
+//////////////////////////////////////////////////////////////////////////
+
+struct FNiagaraDataInterfaceParametersCS_AsyncGpuTrace : public FNiagaraDataInterfaceParametersCS
+{
+	DECLARE_TYPE_LAYOUT(FNiagaraDataInterfaceParametersCS_AsyncGpuTrace, NonVirtual);
+public:
+	void Bind(const FNiagaraDataInterfaceGPUParamInfo& ParameterInfo, const class FShaderParameterMap& ParameterMap)
+	{
+		SystemLWCTileParam.Bind(ParameterMap, *(NDIAsyncGpuTraceLocal::SystemLWCTileName + ParameterInfo.DataInterfaceHLSLSymbol));
+		MaxRayTraceCountParam.Bind(ParameterMap, *(NDIAsyncGpuTraceLocal::MaxRayTraceCountParamName + ParameterInfo.DataInterfaceHLSLSymbol));
+		RayRequestsParam.Bind(ParameterMap, *(NDIAsyncGpuTraceLocal::RayRequestsParamName + ParameterInfo.DataInterfaceHLSLSymbol));
+		RayRequestOffsetParam.Bind(ParameterMap, *(NDIAsyncGpuTraceLocal::RayRequestsOffsetParamName + ParameterInfo.DataInterfaceHLSLSymbol));
+		IntersectionResultsParam.Bind(ParameterMap, *(NDIAsyncGpuTraceLocal::IntersectionResultsParamName + ParameterInfo.DataInterfaceHLSLSymbol));
+		IntersectionResultOffsetParam.Bind(ParameterMap, *(NDIAsyncGpuTraceLocal::IntersectionResultsOffsetParamName + ParameterInfo.DataInterfaceHLSLSymbol));
+		RayTraceCountsParam.Bind(ParameterMap, *(NDIAsyncGpuTraceLocal::RayTraceCountsParamName + ParameterInfo.DataInterfaceHLSLSymbol));
+		RayTraceCountsOffsetParam.Bind(ParameterMap, *(NDIAsyncGpuTraceLocal::RayTraceCountsOffsetParamName + ParameterInfo.DataInterfaceHLSLSymbol));
+	}
+
+	void Set(FRHICommandList& RHICmdList, const FNiagaraDataInterfaceSetArgs& Context) const
+	{
+		check(IsInRenderingThread());
+
+		FNiagaraDataIntefaceProxyAsyncGpuTrace* QueryDI = (FNiagaraDataIntefaceProxyAsyncGpuTrace*)Context.DataInterface;
+
+		FRHIComputeShader* ComputeShaderRHI = Context.Shader.GetComputeShader();
+		SetShaderValue(RHICmdList, ComputeShaderRHI, SystemLWCTileParam, Context.SystemLWCTile);
+		
+		const bool HasRayTracingParametersBound = RayRequestsParam.IsUAVBound()
+			|| IntersectionResultsParam.IsBound()
+			|| RayTraceCountsParam.IsBound();
+
+		if (HasRayTracingParametersBound)
+		{
+			const FNiagaraDataIntefaceProxyAsyncGpuTrace::FInstanceData* InstanceData = QueryDI->SystemInstancesToProxyData_RT.Find(Context.SystemInstanceID);
+
+			FNiagaraAsyncGpuTraceHelper& TraceHelper = Context.ComputeDispatchInterface->GetAsyncGpuTraceHelper();
+			const FNiagaraAsyncGpuTraceDispatchInfo* DispatchInfo = nullptr;
+			if (InstanceData && InstanceData->MaxTracesPerParticle > 0)
+			{
+				DispatchInfo = &TraceHelper.GetDispatch(QueryDI);
+			}
+			else
+			{
+				DispatchInfo = &TraceHelper.GetDummyDispatch();
+			}
+
+			SetShaderValue(RHICmdList, ComputeShaderRHI, MaxRayTraceCountParam, DispatchInfo->MaxTraces);
+
+			if (RayRequestsParam.IsUAVBound())
+			{
+				check(DispatchInfo->TraceRequests.IsValid());
+				RHICmdList.SetUAVParameter(ComputeShaderRHI, RayRequestsParam.GetUAVIndex(), DispatchInfo->TraceRequests.Buffer->UAV);
+				SetShaderValue(RHICmdList, ComputeShaderRHI, RayRequestOffsetParam, DispatchInfo->TraceRequests.Offset);
+			}
+
+			if (IntersectionResultsParam.IsBound())
+			{
+				check(DispatchInfo->LastFrameTraceResults.IsValid());
+
+				SetSRVParameter(RHICmdList, ComputeShaderRHI, IntersectionResultsParam, DispatchInfo->LastFrameTraceResults.Buffer->SRV);
+				SetShaderValue(RHICmdList, ComputeShaderRHI, IntersectionResultOffsetParam, DispatchInfo->LastFrameTraceResults.Offset);
+			}
+
+			if (RayTraceCountsParam.IsUAVBound())
+			{
+				check(DispatchInfo->TraceCounts.IsValid());
+
+				RHICmdList.SetUAVParameter(ComputeShaderRHI, RayTraceCountsParam.GetUAVIndex(), DispatchInfo->TraceCounts.Buffer->UAV);
+				SetShaderValue(RHICmdList, ComputeShaderRHI, RayTraceCountsOffsetParam, DispatchInfo->TraceCounts.Offset);
+			}
+		}
+		else
+		{
+			SetShaderValue(RHICmdList, ComputeShaderRHI, MaxRayTraceCountParam, 0);
+		}
+	}
+
+	void Unset(FRHICommandList& RHICmdList, const FNiagaraDataInterfaceSetArgs& Context) const
+	{
+		FRHIComputeShader* ComputeShaderRHI = Context.Shader.GetComputeShader();
+		if (RayRequestsParam.IsUAVBound())
+		{
+			RayRequestsParam.UnsetUAV(RHICmdList, ComputeShaderRHI);
+		}
+
+		if (RayTraceCountsParam.IsUAVBound())
+		{
+			RayTraceCountsParam.UnsetUAV(RHICmdList, ComputeShaderRHI);
+		}
+	}
+
+private:
+	LAYOUT_FIELD(FShaderParameter, SystemLWCTileParam);
+
+	LAYOUT_FIELD(FShaderParameter, MaxRayTraceCountParam);
+	LAYOUT_FIELD(FRWShaderParameter, RayRequestsParam);
+	LAYOUT_FIELD(FShaderParameter, RayRequestOffsetParam);
+	LAYOUT_FIELD(FShaderResourceParameter, IntersectionResultsParam);
+	LAYOUT_FIELD(FShaderParameter, IntersectionResultOffsetParam);
+	LAYOUT_FIELD(FRWShaderParameter, RayTraceCountsParam);
+	LAYOUT_FIELD(FShaderParameter, RayTraceCountsOffsetParam);
+};
+
+IMPLEMENT_TYPE_LAYOUT(FNiagaraDataInterfaceParametersCS_AsyncGpuTrace);
+
+IMPLEMENT_NIAGARA_DI_PARAMETER(UNiagaraDataInterfaceAsyncGpuTrace, FNiagaraDataInterfaceParametersCS_AsyncGpuTrace);
+
+#undef LOCTEXT_NAMESPACE
diff --git a/Engine/Plugins/FX/Niagara/Source/Niagara/Private/NiagaraAsyncGpuTraceHelper.cpp b/Engine/Plugins/FX/Niagara/Source/Niagara/Private/NiagaraAsyncGpuTraceHelper.cpp
new file mode 100644
--- /dev/null
+++ b/Engine/Plugins/FX/Niagara/Source/Niagara/Private/NiagaraAsyncGpuTraceHelper.cpp
@@ -0,0 +1,343 @@
+// Copyright Epic Games, Inc. All Rights Reserved.
+
+#include "NiagaraAsyncGpuTraceHelper.h"
+
+#include "NiagaraGpuComputeDispatchInterface.h"
+#include "NiagaraStats.h"
+
+int32 GNiagaraAsyncTraceScratchBucketSize = 1024;
+static FAutoConsoleVariableRef CVarNiagaraAsyncTraceScratchBucketSize(
+	TEXT("fx.Niagara.AsyncTrace.ScratchPadBucketSize"),
+	GNiagaraAsyncTraceScratchBucketSize,
+	TEXT("Size (in elements) for async gpu traces scratch buffer buckets. \n"),
+	ECVF_Default
+);
+
+int32 GNiagaraAsyncTraceCountsScratchBucketSize = 3 * 256;
+static FAutoConsoleVariableRef CVarNiagaraAsyncTrceCountsScratchBucketSize(
+	TEXT("fx.Niagara.AsyncTrace.CountsScratchPadBucketSize"),
+	GNiagaraAsyncTraceCountsScratchBucketSize,
+	TEXT("Scratch bucket size for the async gpu trace counts buffer. This buffer requires 4. \n"),
+	ECVF_Default
+);
+
+FNiagaraAsyncGpuTraceHelper::FNiagaraAsyncGpuTraceHelper(EShaderPlatform InShaderPlatform, ERHIFeatureLevel::Type InFeatureLevel, FNiagaraGpuComputeDispatchInterface* Dispatcher)
+	: ShaderPlatform(InShaderPlatform)
+	, FeatureLevel(InFeatureLevel)
+	, TraceRequests(GNiagaraAsyncTraceScratchBucketSize, TEXT("NiagaraRayRequests"))
+	, TraceResults(GNiagaraAsyncTraceScratchBucketSize, TEXT("NiagaraRayTraceIntersections"))
+	, TraceCounts(GNiagaraAsyncTraceCountsScratchBucketSize, TEXT("NiagaraRayTraceCounts"), BUF_Static | BUF_DrawIndirect)
+{
+	InitProviders(Dispatcher);
+
+#if NIAGARA_ASYNC_GPU_TRACE_COLLISION_GROUPS
+	FPrimitiveSceneInfo::OnGPUSceneInstancesAllocated.AddRaw(this, &FNiagaraAsyncGpuTraceHelper::OnPrimitiveGPUSceneInstancesDirtied);
+	FPrimitiveSceneInfo::OnGPUSceneInstancesFreed.AddRaw(this, &FNiagaraAsyncGpuTraceHelper::OnPrimitiveGPUSceneInstancesDirtied);
+#endif
+}
+
+FNiagaraAsyncGpuTraceHelper::~FNiagaraAsyncGpuTraceHelper()
+{
+#if NIAGARA_ASYNC_GPU_TRACE_COLLISION_GROUPS
+	FPrimitiveSceneInfo::OnGPUSceneInstancesAllocated.RemoveAll(this);
+	FPrimitiveSceneInfo::OnGPUSceneInstancesFreed.RemoveAll(this);
+#endif
+
+	Reset();
+	Dispatches.Reset();
+
+	DEC_MEMORY_STAT_BY(STAT_NiagaraGPUDataInterfaceMemory, TraceRequests.AllocatedBytes());
+	DEC_MEMORY_STAT_BY(STAT_NiagaraGPUDataInterfaceMemory, TraceResults.AllocatedBytes());
+	DEC_MEMORY_STAT_BY(STAT_NiagaraGPUDataInterfaceMemory, TraceCounts.AllocatedBytes());
+
+	TraceRequests.Release();
+	TraceResults.Release();
+	TraceCounts.Release();
+}
+
+void FNiagaraAsyncGpuTraceHelper::Reset()
+{
+	for (const auto& TraceProvider : TraceProviders)
+	{
+		TraceProvider->Reset();
+	}
+}
+
+void FNiagaraAsyncGpuTraceHelper::InitProviders(FNiagaraGpuComputeDispatchInterface* Dispatcher)
+{
+	const auto& ProviderPriorities = GetDefault<UNiagaraSettings>()->NDICollisionQuery_AsyncGpuTraceProviderOrder;
+
+	TraceProviders = FNiagaraAsyncGpuTraceProvider::CreateSupportedProviders(ShaderPlatform, Dispatcher, ProviderPriorities);
+}
+
+void FNiagaraAsyncGpuTraceHelper::BeginFrame(FRHICommandList& RHICmdList, FNiagaraGpuComputeDispatchInterface* Dispatcher)
+{
+	//Store off last frames dispatches so we can still access the previous results buffers for reading in this frame's simulations.
+	PreviousFrameDispatches = MoveTemp(Dispatches);
+	Dispatches.Reset();
+
+	//Ensure we're each buffer is definitely in the right access state.
+	TraceRequests.Transition(RHICmdList, ERHIAccess::Unknown, ERHIAccess::UAVCompute);//Simulations will write new ray trace requests.
+	TraceResults.Transition(RHICmdList, ERHIAccess::Unknown, ERHIAccess::SRVCompute);//Simulations will read from the previous frames results.
+	TraceCounts.Transition(RHICmdList, ERHIAccess::Unknown, ERHIAccess::UAVCompute);//Simulations will accumulate the number of traces requests here.
+
+	//Reset the allocations.
+	TraceRequests.Reset();
+
+	//Note this doesn't change any buffers or data themselves and the ray intersections buffer are still going to be read as SRVs in the up coming simulations dispatches.
+	//We're just clearing existing allocations here.
+	TraceResults.Reset();
+
+	//We have to also clear the counts/indirect args buffer.
+	TraceCounts.Reset(RHICmdList, true);
+
+	// Clear the dummy buffer allocations.
+	DummyDispatch.Reset();
+
+	// figure out the provider type that should be used for the 'default' case based on the provided scene
+	ResolvedDefaultProviderType = ENDICollisionQuery_AsyncGpuTraceProvider::None;
+
+	for (const auto& TraceProvider : TraceProviders)
+	{
+		if (TraceProvider->IsAvailable())
+		{
+			ResolvedDefaultProviderType = TraceProvider->GetType();
+			break;
+		}
+	}
+}
+
+void FNiagaraAsyncGpuTraceHelper::PostRenderOpaque(FRHICommandList& RHICmdList, FNiagaraGpuComputeDispatchInterface* Dispatcher, TConstArrayView<FViewInfo> Views)
+{
+#if NIAGARA_ASYNC_GPU_TRACE_COLLISION_GROUPS
+	if (bCollisionGroupMapDirty)
+	{
+		FNiagaraAsyncGpuTraceProvider::BuildCollisionGroupHashMap(RHICmdList, FeatureLevel, Dispatcher->GetScene(), CollisionGroupMap, CollisionGroupHashMapBuffer);
+		bCollisionGroupMapDirty = false;
+	}
+#endif
+
+	FNiagaraAsyncGpuTraceProvider::FCollisionGroupHashMap* CollisionHashMap = GetCollisionGroupHashMap();
+
+	for (const auto& TraceProvider : TraceProviders)
+	{
+		if (TraceProvider->IsAvailable())
+		{
+			TraceProvider->PostRenderOpaque(RHICmdList, Views, CollisionHashMap);
+		}
+	}
+}
+
+FNiagaraAsyncGpuTraceProvider* FNiagaraAsyncGpuTraceHelper::GetTraceProvider(ENDICollisionQuery_AsyncGpuTraceProvider::Type ProviderType)
+{
+	auto* FoundProvider = TraceProviders.FindByPredicate([&](const TUniquePtr<FNiagaraAsyncGpuTraceProvider>& TraceProvider)
+	{
+		return TraceProvider->GetType() == ProviderType
+			|| ((ProviderType == ENDICollisionQuery_AsyncGpuTraceProvider::Default)
+				&& (TraceProvider->GetType() == ResolvedDefaultProviderType));
+	});
+
+	if (FoundProvider && (*FoundProvider)->IsAvailable())
+	{
+		return FoundProvider->Get();
+	}
+
+	return nullptr;
+}
+
+void FNiagaraAsyncGpuTraceHelper::EndFrame(FRHICommandList& RHICmdList, FNiagaraGpuComputeDispatchInterface* Dispatcher)
+{
+	if (!Dispatches.Num())
+	{
+		return;
+	}
+
+	TraceRequests.Transition(RHICmdList, ERHIAccess::UAVCompute, ERHIAccess::SRVCompute);//Ray trace dispatches will read the ray requests.
+	TraceResults.Transition(RHICmdList, ERHIAccess::SRVCompute, ERHIAccess::UAVCompute);//Ray trace dispatches will write new results.
+	TraceCounts.Transition(RHICmdList, ERHIAccess::UAVCompute, ERHIAccess::IndirectArgs | ERHIAccess::SRVCompute);//Ray trace dispatches read these counts and use them for indirect dispatches.
+
+	for (const auto& DispatchInfoPair : Dispatches)
+	{
+		const FNiagaraAsyncGpuTraceDispatchInfo& DispatchInfo = DispatchInfoPair.Value;
+		if (DispatchInfo.MaxTraces)
+		{
+			FNiagaraAsyncGpuTraceProvider::FDispatchRequest Request;
+			Request.TracesBuffer = DispatchInfo.TraceRequests.Buffer;
+			Request.ResultsBuffer = DispatchInfo.TraceResults.Buffer;
+			Request.TraceCountsBuffer = DispatchInfo.TraceCounts.Buffer;
+			Request.MaxTraceCount = DispatchInfo.MaxTraces;
+			Request.TracesOffset = DispatchInfo.TraceRequests.Offset;
+			Request.ResultsOffset = DispatchInfo.TraceResults.Offset;
+			Request.TraceCountsOffset = DispatchInfo.TraceCounts.Offset;
+			Request.MaxRetraceCount = DispatchInfo.MaxRetraces;
+
+			FNiagaraAsyncGpuTraceProvider::FCollisionGroupHashMap* CollisionHashMap = GetCollisionGroupHashMap();
+
+			if (FNiagaraAsyncGpuTraceProvider* TraceProvider = GetTraceProvider(DispatchInfo.ProviderType))
+			{
+				TraceProvider->IssueTraces(RHICmdList, Request, CollisionHashMap);
+			}
+			else
+			{
+				// Clear the results
+				FNiagaraAsyncGpuTraceProvider::ClearResults(RHICmdList, ShaderPlatform, Request);
+			}
+		}
+	}
+
+	TraceRequests.Transition(RHICmdList, ERHIAccess::SRVCompute, ERHIAccess::UAVCompute);//Next frame, simulations will write new ray trace requests.
+	TraceResults.Transition(RHICmdList, ERHIAccess::UAVCompute, ERHIAccess::SRVCompute);//Next frame these results will be read by simulations shaders.
+	TraceCounts.Transition(RHICmdList, ERHIAccess::IndirectArgs | ERHIAccess::SRVCompute, ERHIAccess::UAVCompute);//Next frame these counts will be written by simulation shaders.
+}
+
+void FNiagaraAsyncGpuTraceHelper::AddToDispatch(FNiagaraDataInterfaceProxy* DispatchKey, uint32 MaxRays, int32 MaxRetraces, ENDICollisionQuery_AsyncGpuTraceProvider::Type ProviderType)
+{
+	FNiagaraAsyncGpuTraceDispatchInfo& Dispatch = Dispatches.FindOrAdd(DispatchKey);
+	Dispatch.MaxTraces += MaxRays;
+	Dispatch.MaxRetraces = MaxRetraces;
+	Dispatch.ProviderType = ProviderType;
+}
+
+void FNiagaraAsyncGpuTraceHelper::BuildDispatch(FRHICommandList& RHICmdList, FNiagaraDataInterfaceProxy* DispatchKey)
+{
+	FNiagaraAsyncGpuTraceDispatchInfo& Dispatch = Dispatches.FindChecked(DispatchKey);
+
+	//Finalize allocations first time we access the dispatch.
+	if (Dispatch.TraceRequests.IsValid() == false)
+	{
+ 		DEC_MEMORY_STAT_BY(STAT_NiagaraGPUDataInterfaceMemory, TraceRequests.AllocatedBytes());
+ 		DEC_MEMORY_STAT_BY(STAT_NiagaraGPUDataInterfaceMemory, TraceResults.AllocatedBytes());
+ 		DEC_MEMORY_STAT_BY(STAT_NiagaraGPUDataInterfaceMemory, TraceCounts.AllocatedBytes());
+
+		Dispatch.TraceRequests = TraceRequests.Alloc(Dispatch.MaxTraces);
+		Dispatch.TraceResults = TraceResults.Alloc(Dispatch.MaxTraces);
+		Dispatch.TraceCounts = TraceCounts.Alloc<uint32>(3, RHICmdList, true);
+
+ 		INC_MEMORY_STAT_BY(STAT_NiagaraGPUDataInterfaceMemory, TraceRequests.AllocatedBytes());
+ 		INC_MEMORY_STAT_BY(STAT_NiagaraGPUDataInterfaceMemory, TraceResults.AllocatedBytes());
+ 		INC_MEMORY_STAT_BY(STAT_NiagaraGPUDataInterfaceMemory, TraceCounts.AllocatedBytes());
+
+		//Find last frame's results buffer if there was one. Ideally we can do this without the map but it will need a bit of wrangling.
+		FNiagaraAsyncGpuTraceDispatchInfo* PrevDispatch = PreviousFrameDispatches.Find(DispatchKey);
+		if (PrevDispatch && PrevDispatch->TraceResults.IsValid())
+		{
+			Dispatch.LastFrameTraceResults = PrevDispatch->TraceResults;
+		}
+		else
+		{
+			//Set the current frame results just so this is a valid buffer.
+			//If we didn't have a results buffer last frame then this won't actually be read.
+			Dispatch.LastFrameTraceResults = Dispatch.TraceResults;
+		}
+	}
+}
+
+void FNiagaraAsyncGpuTraceHelper::BuildDummyDispatch(FRHICommandList& RHICmdList)
+{
+	if (DummyDispatch.IsValid() == false)
+	{
+ 		DEC_MEMORY_STAT_BY(STAT_NiagaraGPUDataInterfaceMemory, TraceRequests.AllocatedBytes());
+ 		DEC_MEMORY_STAT_BY(STAT_NiagaraGPUDataInterfaceMemory, TraceResults.AllocatedBytes());
+ 		DEC_MEMORY_STAT_BY(STAT_NiagaraGPUDataInterfaceMemory, TraceCounts.AllocatedBytes());
+
+		DummyDispatch.TraceRequests = TraceRequests.Alloc(1);
+		DummyDispatch.TraceResults = TraceResults.Alloc(1);
+		DummyDispatch.LastFrameTraceResults = DummyDispatch.TraceResults;
+		DummyDispatch.TraceCounts = TraceCounts.Alloc<uint32>(3, RHICmdList, true);
+		
+ 		INC_MEMORY_STAT_BY(STAT_NiagaraGPUDataInterfaceMemory, TraceRequests.AllocatedBytes());
+ 		INC_MEMORY_STAT_BY(STAT_NiagaraGPUDataInterfaceMemory, TraceResults.AllocatedBytes());
+ 		INC_MEMORY_STAT_BY(STAT_NiagaraGPUDataInterfaceMemory, TraceCounts.AllocatedBytes());
+
+		DummyDispatch.MaxTraces = 0;
+		DummyDispatch.MaxRetraces = 0;
+	}
+}
+
+const FNiagaraAsyncGpuTraceDispatchInfo& FNiagaraAsyncGpuTraceHelper::GetDispatch(FNiagaraDataInterfaceProxy* DispatchKey) const
+{
+	const FNiagaraAsyncGpuTraceDispatchInfo& Dispatch = Dispatches.FindChecked(DispatchKey);
+
+	// make sure that the buffer is valid (that BuildDispatch has been called)
+	check(Dispatch.TraceRequests.IsValid());
+
+	return Dispatch;
+}
+
+const FNiagaraAsyncGpuTraceDispatchInfo& FNiagaraAsyncGpuTraceHelper::GetDummyDispatch() const
+{
+	// make sure that the buffers are valid (that BuildDummyDispatch has been called)
+	check(DummyDispatch.TraceRequests.IsValid());
+
+	return DummyDispatch;
+}
+
+ENDICollisionQuery_AsyncGpuTraceProvider::Type FNiagaraAsyncGpuTraceHelper::ResolveSupportedType(ENDICollisionQuery_AsyncGpuTraceProvider::Type InType)
+{
+	return FNiagaraAsyncGpuTraceProvider::ResolveSupportedType(InType, GetDefault<UNiagaraSettings>()->NDICollisionQuery_AsyncGpuTraceProviderOrder);
+}
+
+bool FNiagaraAsyncGpuTraceHelper::RequiresDistanceFieldData(ENDICollisionQuery_AsyncGpuTraceProvider::Type InType)
+{
+	return FNiagaraAsyncGpuTraceProvider::RequiresDistanceFieldData(InType, GetDefault<UNiagaraSettings>()->NDICollisionQuery_AsyncGpuTraceProviderOrder);
+}
+
+bool FNiagaraAsyncGpuTraceHelper::RequiresRayTracingScene(ENDICollisionQuery_AsyncGpuTraceProvider::Type InType)
+{
+	return FNiagaraAsyncGpuTraceProvider::RequiresRayTracingScene(InType, GetDefault<UNiagaraSettings>()->NDICollisionQuery_AsyncGpuTraceProviderOrder);
+}
+
+FNiagaraAsyncGpuTraceProvider::FCollisionGroupHashMap* FNiagaraAsyncGpuTraceHelper::GetCollisionGroupHashMap()
+{
+#if NIAGARA_ASYNC_GPU_TRACE_COLLISION_GROUPS
+	return CollisionGroupMap.Num() > 0 ? &CollisionGroupHashMapBuffer : nullptr;
+#else
+	return nullptr;
+#endif
+}
+
+#if NIAGARA_ASYNC_GPU_TRACE_COLLISION_GROUPS
+void FNiagaraAsyncGpuTraceHelper::SetPrimitiveCollisionGroup(FPrimitiveSceneInfo& Primitive, uint32 CollisionGroup)
+{
+	CollisionGroupMap.FindOrAdd(Primitive.PrimitiveComponentId) = CollisionGroup;
+	bCollisionGroupMapDirty = true;
+}
+
+void FNiagaraAsyncGpuTraceHelper::SetPrimitiveRayTracingCollisionGroup_GT(UPrimitiveComponent* Primitive, uint32 Group)
+{
+	check(IsInGameThread());
+	if (Primitive == nullptr || Primitive->SceneProxy == nullptr)
+	{
+		return;
+	}
+
+	ENQUEUE_RENDER_COMMAND(NiagaraSetPrimHWRTCollisionGroup)(
+		[RT_Helper = this, Proxy = Primitive->SceneProxy, CollisionGroup = Group](FRHICommandListImmediate& RHICmdList)
+	{
+		check(Proxy);
+		RT_Helper->SetPrimitiveCollisionGroup(*Proxy->GetPrimitiveSceneInfo(), CollisionGroup);
+	});
+}
+
+int32 FNiagaraAsyncGpuTraceHelper::AcquireGPURayTracedCollisionGroup_GT()
+{
+	check(IsInGameThread());//one of the few batcher functions that should be called from the GT
+	if (FreeGPURayTracedCollisionGroups.Num() > 0)
+	{
+		return FreeGPURayTracedCollisionGroups.Pop();
+	}
+
+	return NumGPURayTracedCollisionGroups++;
+}
+
+void FNiagaraAsyncGpuTraceHelper::ReleaseGPURayTracedCollisionGroup_GT(int32 CollisionGroup)
+{
+	check(IsInGameThread());//one of the few batcher functions that should be called from the GT
+	FreeGPURayTracedCollisionGroups.Add(CollisionGroup);
+}
+
+void FNiagaraAsyncGpuTraceHelper::OnPrimitiveGPUSceneInstancesDirtied()
+{
+	bCollisionGroupMapDirty = true;
+}
+#endif
\ No newline at end of file
diff --git a/Engine/Plugins/FX/Niagara/Source/Niagara/Private/NiagaraAsyncGpuTraceHelper.h b/Engine/Plugins/FX/Niagara/Source/Niagara/Private/NiagaraAsyncGpuTraceHelper.h
new file mode 100644
--- /dev/null
+++ b/Engine/Plugins/FX/Niagara/Source/Niagara/Private/NiagaraAsyncGpuTraceHelper.h
@@ -0,0 +1,159 @@
+// Copyright Epic Games, Inc. All Rights Reserved.
+#pragma once
+
+#include "CoreMinimal.h"
+
+#include "NiagaraAsyncGpuTraceProvider.h"
+#include "NiagaraGpuScratchPad.h"
+
+#include "RHIDefinitions.h"
+
+#include "PrimitiveSceneInfo.h"
+#include "RHI.h"
+#include "RHIUtilities.h"
+
+struct FNiagaraDataInterfaceProxy;
+class FNiagaraGpuComputeDispatchInterface;
+class FRayTracingPipelineState;
+class FRHICommandList;
+class FRHIRayTracingScene;
+class FRHIShaderResourceView;
+class FRHIUniformBuffer;
+class FRHIUnorderedAccessView;
+struct FRWBuffer;
+class FScene;
+class FViewInfo;
+
+// we currently support assigning, from blueprint, primitives into collision groups which can be excluded
+// from ray traces.  Currently only supported with HWRT
+#define NIAGARA_ASYNC_GPU_TRACE_COLLISION_GROUPS RHI_RAYTRACING
+
+/** Holds all information on a single RayTracing dispatch. */
+struct FNiagaraAsyncGpuTraceDispatchInfo
+{
+	/** Buffer allocation for ray requests to trace. */
+	FNiagaraGpuScratchPadStructured<FNiagaraAsyncGpuTrace>::FAllocation TraceRequests;
+	/** Buffer allocation for writing trace results into. */
+	FNiagaraGpuScratchPadStructured<FNiagaraAsyncGpuTraceResult>::FAllocation TraceResults;
+	/** Buffer allocation for last frames trace results to read from in this frame's simulation. */
+	FNiagaraGpuScratchPadStructured<FNiagaraAsyncGpuTraceResult>::FAllocation LastFrameTraceResults;
+	/** Buffer allocation for ray trace counts. Accumulated in simulations shaders and used as dispatch args for ray tracing shaders. */
+	FNiagaraGpuScratchPad::FAllocation TraceCounts;
+	/** Total possible rays to trace. This may be significantly higher than the actual rays request we accumulate into RayCounts. */
+	uint32 MaxTraces = 0;
+	/** Max number or times rays can be re-traced in a shader if they hit something that is invalid/filtered. */
+	uint32 MaxRetraces = 0;
+
+	ENDICollisionQuery_AsyncGpuTraceProvider::Type ProviderType = ENDICollisionQuery_AsyncGpuTraceProvider::Default;
+
+	FORCEINLINE bool IsValid() const { return MaxTraces > 0; }
+	FORCEINLINE void Reset()
+	{
+		TraceRequests = FNiagaraGpuScratchPadStructured<FNiagaraAsyncGpuTrace>::FAllocation();
+		TraceResults = FNiagaraGpuScratchPadStructured<FNiagaraAsyncGpuTraceResult>::FAllocation();
+		TraceCounts = FNiagaraGpuScratchPad::FAllocation();
+		MaxTraces = 0;
+		MaxRetraces = 0;
+		ProviderType = ENDICollisionQuery_AsyncGpuTraceProvider::Default;
+	}
+};
+
+class FNiagaraAsyncGpuTraceHelper
+{
+public:
+	FNiagaraAsyncGpuTraceHelper(EShaderPlatform InShaderPlatform, ERHIFeatureLevel::Type FeatureLevel, FNiagaraGpuComputeDispatchInterface* Dispatcher);
+	FNiagaraAsyncGpuTraceHelper() = delete;
+	~FNiagaraAsyncGpuTraceHelper();
+
+	void Reset();
+
+	void BeginFrame(FRHICommandList& RHICmdList, FNiagaraGpuComputeDispatchInterface* Dispatcher);
+	void PostRenderOpaque(FRHICommandList& RHICmdList, FNiagaraGpuComputeDispatchInterface* Dispatcher, TConstArrayView<FViewInfo> Views);
+	void EndFrame(FRHICommandList& RHICmdList, FNiagaraGpuComputeDispatchInterface* Dispatcher);
+
+	/** Accumulates ray requests from user DIs into single dispatches per DI. */
+	void AddToDispatch(FNiagaraDataInterfaceProxy* DispatchKey, uint32 MaxRays, int32 MaxRetraces, ENDICollisionQuery_AsyncGpuTraceProvider::Type ProviderType);
+
+	/** Ensures that the final buffers for the provided proxy have been allocated. */
+	void BuildDispatch(FRHICommandList& RHICmdList, FNiagaraDataInterfaceProxy* DispatchKey);
+
+	/** Ensures that buffers for a dummy buffer have been allocated. */
+	void BuildDummyDispatch(FRHICommandList& RHICmdList);
+
+	/** Returns the final buffers for each DI for use in simulations and RT dispatches. */
+	const FNiagaraAsyncGpuTraceDispatchInfo& GetDispatch(FNiagaraDataInterfaceProxy* DispatchKey) const;
+
+	/** Returns a dummy dispatch to allow shaders to still bind to valid buffers. */
+	const FNiagaraAsyncGpuTraceDispatchInfo& GetDummyDispatch() const;
+
+	FNiagaraGpuScratchPadStructured<FNiagaraAsyncGpuTrace>& GetRayRequests(){ return TraceRequests; }
+	FNiagaraGpuScratchPadStructured<FNiagaraAsyncGpuTraceResult> GetRayTraceIntersections() { return TraceResults; }
+	FNiagaraGpuScratchPad GetTraceCounts() { return TraceCounts; }
+
+	static ENDICollisionQuery_AsyncGpuTraceProvider::Type ResolveSupportedType(ENDICollisionQuery_AsyncGpuTraceProvider::Type InType);
+	static bool RequiresDistanceFieldData(ENDICollisionQuery_AsyncGpuTraceProvider::Type InType);
+	static bool RequiresRayTracingScene(ENDICollisionQuery_AsyncGpuTraceProvider::Type InType);
+
+#if NIAGARA_ASYNC_GPU_TRACE_COLLISION_GROUPS
+	/** Adds a primitive to a collision group. This data is sent to the GPU on the start of the next frame. */
+	void SetPrimitiveCollisionGroup(FPrimitiveSceneInfo& Primitive, uint32 CollisionGroup);
+
+	/** Pushes cpu side copy of the collision group map to the GPU. */
+	void UpdateCollisionGroupMap(FRHICommandList& RHICmdList, FScene* Scene, ERHIFeatureLevel::Type FeatureLevel);
+
+	// Set of functions for managing, from the game thread, the creation and assignment of collision groups; currently
+	// only supported with ray traced collisions
+	void SetPrimitiveRayTracingCollisionGroup_GT(UPrimitiveComponent* Primitive, uint32 Group);
+	int32 AcquireGPURayTracedCollisionGroup_GT();
+	void ReleaseGPURayTracedCollisionGroup_GT(int32 CollisionGroup);
+
+	void OnPrimitiveGPUSceneInstancesDirtied();
+#endif
+
+private:
+	const EShaderPlatform ShaderPlatform;
+	const ERHIFeatureLevel::Type FeatureLevel;
+
+	/** Scratch buffer holding all ray trace requests for all simulations in the scene. */
+	FNiagaraGpuScratchPadStructured<FNiagaraAsyncGpuTrace> TraceRequests;
+	/** Scratch buffer holding all ray trace results for all simulations in the scene. */
+	FNiagaraGpuScratchPadStructured<FNiagaraAsyncGpuTraceResult> TraceResults;
+	/** Scratch buffer holding all accumulated ray request counts for all simulations in the scene. */
+	FNiagaraGpuScratchPad TraceCounts;
+
+	/** Map of DI Proxy to Dispatch Info. Each DI will have a single dispatch for all it's instances. */
+	TMap<FNiagaraDataInterfaceProxy*, FNiagaraAsyncGpuTraceDispatchInfo> Dispatches;
+
+	/** Last frame's DI proxy map. Allows use to retrieve the buffer allocations to read last frames trace results. TODO: Improve. Ideally don't need a map copy here. */
+	TMap<FNiagaraDataInterfaceProxy*, FNiagaraAsyncGpuTraceDispatchInfo> PreviousFrameDispatches;
+
+	/** In places where ray tracing is disabled etc we use a dummy dispatch allocation set so the simulations can bind to valid buffers. */
+	FNiagaraAsyncGpuTraceDispatchInfo DummyDispatch;
+
+	TArray<TUniquePtr<FNiagaraAsyncGpuTraceProvider>> TraceProviders;
+
+	ENDICollisionQuery_AsyncGpuTraceProvider::Type ResolvedDefaultProviderType = ENDICollisionQuery_AsyncGpuTraceProvider::None;
+
+	FNiagaraAsyncGpuTraceProvider* GetTraceProvider(ENDICollisionQuery_AsyncGpuTraceProvider::Type ProviderType);
+	void InitProviders(FNiagaraGpuComputeDispatchInterface* Dispatcher);
+
+	FNiagaraAsyncGpuTraceProvider::FCollisionGroupHashMap* GetCollisionGroupHashMap();
+
+#if NIAGARA_ASYNC_GPU_TRACE_COLLISION_GROUPS
+	/**
+	CPU side copy of the PrimID to Collision Group Map.
+	This is uploaded to the GPU and used in Niagara RG shader to filter self collisions between objects of the same group.
+	*/
+	TMap<FPrimitiveComponentId, uint32> CollisionGroupMap;
+
+	/**
+	GPU side of the collision group hash map
+	*/
+	FNiagaraAsyncGpuTraceProvider::FCollisionGroupHashMap CollisionGroupHashMapBuffer;
+
+	/** Pool of free GPU ray traced collision groups. */
+	TArray<int32> FreeGPURayTracedCollisionGroups;
+	int32 NumGPURayTracedCollisionGroups = 0;
+	bool bCollisionGroupMapDirty = true;
+#endif
+};
\ No newline at end of file
diff --git a/Engine/Plugins/FX/Niagara/Source/Niagara/Private/NiagaraDataInterfaceCollisionQuery.cpp b/Engine/Plugins/FX/Niagara/Source/Niagara/Private/NiagaraDataInterfaceCollisionQuery.cpp
--- a/Engine/Plugins/FX/Niagara/Source/Niagara/Private/NiagaraDataInterfaceCollisionQuery.cpp
+++ b/Engine/Plugins/FX/Niagara/Source/Niagara/Private/NiagaraDataInterfaceCollisionQuery.cpp
@@ -3,10 +3,11 @@
 #include "NiagaraDataInterfaceCollisionQuery.h"
 
 #include "GlobalDistanceFieldParameters.h"
+#include "NiagaraAsyncGpuTraceHelper.h"
 #include "NiagaraComponent.h"
+#include "NiagaraDataInterfaceUtilities.h"
 #include "NiagaraGpuComputeDispatchInterface.h"
 #include "NiagaraGpuComputeDispatch.h"
-#include "NiagaraRayTracingHelper.h"
 #include "NiagaraSimStageData.h"
 #include "NiagaraStats.h"
 #include "NiagaraTypes.h"
@@ -19,22 +20,6 @@
 
 #define LOCTEXT_NAMESPACE "NiagaraDataInterfaceCollisionQuery"
 
-void OnHWRTCollisionsEnabledChanged(IConsoleVariable* CVar)
-{
-	//Force a reinit of everything just to be safe.
-	FNiagaraSystemUpdateContext Context;
-	Context.AddAll(true);
-}
-
-int32 GEnableGPUHWRTCollisions = 1;
-static FAutoConsoleVariableRef CVarEnableGPUHWRTCollisions(
-	TEXT("fx.Niagara.Collision.EnableGPURayTracedCollisions"),
-	GEnableGPUHWRTCollisions,
-	TEXT("If greater than zero, GPU hardware ray trace collisions are enabled."),
-	FConsoleVariableDelegate::CreateStatic(&OnHWRTCollisionsEnabledChanged),
-	ECVF_Default
-);
-
 namespace NDICollisionQueryLocal
 {
 	static const TCHAR* CommonShaderFile = TEXT("/Plugin/FX/Niagara/Private/NiagaraDataInterfaceCollisionQuery.ush");
@@ -45,19 +30,14 @@ namespace NDICollisionQueryLocal
 	static const FName DistanceFieldName(TEXT("QueryMeshDistanceFieldGPU"));
 	static const FName SyncTraceName(TEXT("PerformCollisionQuerySyncCPU"));
 	static const FName AsyncTraceName(TEXT("PerformCollisionQueryAsyncCPU"));
+
+	// DEPRECATE_BEGIN
 	static const FName IssueAsyncRayTraceName(TEXT("IssueAsyncRayTraceGpu"));
 	static const FName CreateAsyncRayTraceName(TEXT("CreateAsyncRayTraceGpu"));
 	static const FName ReserveAsyncRayTraceName(TEXT("ReserveAsyncRayTraceGpu"));
 	static const FName ReadAsyncRayTraceName(TEXT("ReadAsyncRayTraceGpu"));
+	// DEPRECATE_END
 
-	static const FString RayTracingEnabledParamName(TEXT("RayTracingEnabled_"));
-	static const FString MaxRayTraceCountParamName(TEXT("MaxRayTraceCount_"));
-	static const FString RayRequestsParamName(TEXT("RayRequests_"));
-	static const FString RayRequestsOffsetParamName(TEXT("RayRequestsOffset_"));
-	static const FString IntersectionResultsParamName(TEXT("IntersectionResults_"));
-	static const FString IntersectionResultsOffsetParamName(TEXT("IntersectionResultsOffset_"));
-	static const FString RayTraceCountsParamName(TEXT("RayTraceCounts_"));
-	static const FString RayTraceCountsOffsetParamName(TEXT("RayTraceCountsOffset_"));
 	static const FString SystemLWCTileName(TEXT("SystemLWCTile_"));
 }
 
@@ -72,79 +52,28 @@ struct FNiagaraCollisionDIFunctionVersion
 		AddedCustomDepthCollision = 2,
 		ReturnCollisionMaterialIdx = 3,
 		LargeWorldCoordinates = 4,
+		DeprecatedAsyncGpuTrace = 5,
 
 		VersionPlusOne,
 		LatestVersion = VersionPlusOne - 1
 	};
 };
+
 struct FNiagaraDataIntefaceProxyCollisionQuery : public FNiagaraDataInterfaceProxy
 {
-#if RHI_RAYTRACING
-	int32 MaxTracesPerParticle = 0;
-	uint32 MaxRetraces = 0;
-#endif
-
-	virtual ~FNiagaraDataIntefaceProxyCollisionQuery()
-	{
-	}
+	// There's nothing in this proxy. It just reads from scene textures.
 
 	virtual int32 PerInstanceDataPassedToRenderThreadSize() const override
 	{
 		return 0;
 	}
-
-	virtual void PreStage(FRHICommandList& RHICmdList, const FNiagaraDataInterfaceStageArgs& Context) override
-	{
-		FNiagaraDataInterfaceProxy::PreStage(RHICmdList, Context);
-
-#if RHI_RAYTRACING
-		if (IsRayTracingEnabled() && GEnableGPUHWRTCollisions && MaxTracesPerParticle > 0)
-		{
-			//Accumulate the total ray requests for this DI for all dispatches in the stage.
-			int32 RayRequests = MaxTracesPerParticle * Context.SimStageData->DestinationNumInstances;
-			FNiagaraRayTracingHelper& RTHelper = Context.ComputeDispatchInterface->GetRayTracingHelper();
-			RTHelper.AddToDispatch(this, RayRequests, MaxRetraces);
-		}
-#endif
-	}
-
-#if RHI_RAYTRACING
-	virtual bool RequiresPreStageFinalize() const override
-	{
-		return true;
-	}
-
-	virtual void FinalizePreStage(FRHICommandList& RHICmdList, const FNiagaraGpuComputeDispatchInterface* ComputeDispatchInterface) override
-	{
-		FNiagaraRayTracingHelper& RTHelper = ComputeDispatchInterface->GetRayTracingHelper();
-		if (IsRayTracingEnabled() && GEnableGPUHWRTCollisions && MaxTracesPerParticle > 0)
-		{
-			RTHelper.BuildDispatch(RHICmdList, this);
-		}
-		else
-		{
-			RTHelper.BuildDummyDispatch(RHICmdList);
-		}
-	}
-#endif
-
-	void RenderThreadInitialize(int32 InMaxTracesPerParticle, uint32 InMaxRetraces)
-	{
-#if RHI_RAYTRACING
-		MaxRetraces = InMaxRetraces;
-		MaxTracesPerParticle = InMaxTracesPerParticle;
-#endif
-	}
 };
 
-
 UNiagaraDataInterfaceCollisionQuery::UNiagaraDataInterfaceCollisionQuery(FObjectInitializer const& ObjectInitializer)
 	: Super(ObjectInitializer)
 {
 	TraceChannelEnum = StaticEnum<ECollisionChannel>();
-	SystemInstance = nullptr;
-
-    Proxy.Reset(new FNiagaraDataIntefaceProxyCollisionQuery());
+	Proxy.Reset(new FNiagaraDataIntefaceProxyCollisionQuery());
 }
 
 bool UNiagaraDataInterfaceCollisionQuery::InitPerInstanceData(void* PerInstanceData, FNiagaraSystemInstance* InSystemInstance)
@@ -155,6 +84,7 @@ bool UNiagaraDataInterfaceCollisionQuery::InitPerInstanceData(void* PerInstanceD
 	{
 		PIData->CollisionBatch.Init(InSystemInstance->GetId(), InSystemInstance->GetWorld());
 	}
+
 	return true;
 }
 
@@ -177,16 +107,6 @@ void UNiagaraDataInterfaceCollisionQuery::PostInitProperties()
 	}
 }
 
-void UNiagaraDataInterfaceCollisionQuery::PostLoad()
-{
-	Super::PostLoad();
-
-	if (MaxTracesPerParticle)
-	{
-		MarkRenderDataDirty();
-	}
-}
-
 void UNiagaraDataInterfaceCollisionQuery::GetAssetTagsForContext(const UObject* InAsset, const TArray<const UNiagaraDataInterface*>& InProperties, TMap<FName, uint32>& NumericKeys, TMap<FName, FString>& StringKeys) const
 {
 #if WITH_EDITOR
@@ -341,6 +261,8 @@ void UNiagaraDataInterfaceCollisionQuery::GetFunctions(TArray<FNiagaraFunctionSi
 	
 	const FText TraceStartWorldDescription = LOCTEXT("TraceStartWorldDescription", "Ray starting point in world space");
 	const FText TraceEndWorldDescription = LOCTEXT("TraceEndWorldDescription", "Ray end point in world space");
+
+	// DEPRECATE_BEGIN
 	{
 		const FText AsyncTraceChannelDescription = LOCTEXT("TraceChannelDescription", "Currently unused, will represent the trace channels for which geometry the trace should test against");
 		const FText QueryIDDescription = LOCTEXT("QueryIDDescription", "Unique (for this frame) index of the query being enqueued (used in subsequent frames to retrieve results).  Must be less than MaxRayTraceCount");
@@ -351,6 +273,7 @@ void UNiagaraDataInterfaceCollisionQuery::GetFunctions(TArray<FNiagaraFunctionSi
 		IssueRayTrace.bRequiresExecPin = true;
 		IssueRayTrace.bMemberFunction = true;
 		IssueRayTrace.bSupportsCPU = false;
+		IssueRayTrace.bSoftDeprecatedFunction = true;
 #if WITH_EDITORONLY_DATA
 		IssueRayTrace.FunctionVersion = FNiagaraCollisionDIFunctionVersion::LatestVersion;
 		IssueRayTrace.Description = LOCTEXT("IssueAsync_RayTraceDescription", "Enqueues a GPU raytrace with the result being available the following frame");
@@ -367,6 +290,7 @@ void UNiagaraDataInterfaceCollisionQuery::GetFunctions(TArray<FNiagaraFunctionSi
 		CreateRayTrace.bRequiresExecPin = true;
 		CreateRayTrace.bMemberFunction = true;
 		CreateRayTrace.bSupportsCPU = false;
+		CreateRayTrace.bSoftDeprecatedFunction = true;
 #if WITH_EDITORONLY_DATA
 		CreateRayTrace.FunctionVersion = FNiagaraCollisionDIFunctionVersion::LatestVersion;
 		CreateRayTrace.Description = LOCTEXT("CreateAsync_RayTraceDescription", "Creates a GPU raytrace with the result being available the following frame (index is returned)");
@@ -380,19 +304,20 @@ void UNiagaraDataInterfaceCollisionQuery::GetFunctions(TArray<FNiagaraFunctionSi
 	}
 
 	{
-		FNiagaraFunctionSignature& IssueRayTrace = OutFunctions.AddDefaulted_GetRef();
-		IssueRayTrace.Name = NDICollisionQueryLocal::ReserveAsyncRayTraceName;
-		IssueRayTrace.bRequiresExecPin = true;
-		IssueRayTrace.bMemberFunction = true;
-		IssueRayTrace.bSupportsCPU = false;
+		FNiagaraFunctionSignature& ReserveRayTrace = OutFunctions.AddDefaulted_GetRef();
+		ReserveRayTrace.Name = NDICollisionQueryLocal::ReserveAsyncRayTraceName;
+		ReserveRayTrace.bRequiresExecPin = true;
+		ReserveRayTrace.bMemberFunction = true;
+		ReserveRayTrace.bSupportsCPU = false;
+		ReserveRayTrace.bSoftDeprecatedFunction = true;
 #if WITH_EDITORONLY_DATA
-		IssueRayTrace.FunctionVersion = FNiagaraCollisionDIFunctionVersion::LatestVersion;
-		IssueRayTrace.Description = LOCTEXT("ReserveAsync_RayTraceDescription", "Reserves a number of ray trace request slots");
+		ReserveRayTrace.FunctionVersion = FNiagaraCollisionDIFunctionVersion::LatestVersion;
+		ReserveRayTrace.Description = LOCTEXT("ReserveAsync_RayTraceDescription", "Reserves a number of ray trace request slots");
 #endif
-		IssueRayTrace.AddInput(FNiagaraVariable(FNiagaraTypeDefinition(GetClass()), TEXT("CollisionQuery")));
-		IssueRayTrace.AddInput(FNiagaraVariable(FNiagaraTypeDefinition::GetIntDef(), TEXT("TraceCount")), LOCTEXT("ReserveAsync_QueryIDDescription", "Number of async raytrace requests to be reserved"));
-		IssueRayTrace.AddOutput(FNiagaraVariable(FNiagaraTypeDefinition::GetIntDef(), TEXT("FirstQueryID")), LOCTEXT("ReserveAsync_TraceChannelDescription", "The first index in the block reserved through this call"));
-		IssueRayTrace.AddOutput(FNiagaraVariable(FNiagaraTypeDefinition::GetBoolDef(), TEXT("IsQueryValid")), LOCTEXT("ReserveAsync_IsQueryValidDescription", "Returns true if the requested indices were reserved"));
+		ReserveRayTrace.AddInput(FNiagaraVariable(FNiagaraTypeDefinition(GetClass()), TEXT("CollisionQuery")));
+		ReserveRayTrace.AddInput(FNiagaraVariable(FNiagaraTypeDefinition::GetIntDef(), TEXT("TraceCount")), LOCTEXT("ReserveAsync_QueryIDDescription", "Number of async raytrace requests to be reserved"));
+		ReserveRayTrace.AddOutput(FNiagaraVariable(FNiagaraTypeDefinition::GetIntDef(), TEXT("FirstQueryID")), LOCTEXT("ReserveAsync_TraceChannelDescription", "The first index in the block reserved through this call"));
+		ReserveRayTrace.AddOutput(FNiagaraVariable(FNiagaraTypeDefinition::GetBoolDef(), TEXT("IsQueryValid")), LOCTEXT("ReserveAsync_IsQueryValidDescription", "Returns true if the requested indices were reserved"));
 	}
 
 	{
@@ -400,6 +325,7 @@ void UNiagaraDataInterfaceCollisionQuery::GetFunctions(TArray<FNiagaraFunctionSi
 		ReadRayTrace.Name = NDICollisionQueryLocal::ReadAsyncRayTraceName;
 		ReadRayTrace.bMemberFunction = true;
 		ReadRayTrace.bSupportsCPU = false;
+		ReadRayTrace.bSoftDeprecatedFunction = true;
 #if WITH_EDITORONLY_DATA
 		ReadRayTrace.FunctionVersion = FNiagaraCollisionDIFunctionVersion::LatestVersion;
 		ReadRayTrace.Description = LOCTEXT("ReadAsync_RayTraceDescription", "Reads the results of a previously enqueued GPU ray trace");
@@ -411,6 +337,7 @@ void UNiagaraDataInterfaceCollisionQuery::GetFunctions(TArray<FNiagaraFunctionSi
 		ReadRayTrace.AddOutput(FNiagaraVariable(FNiagaraTypeDefinition::GetPositionDef(), TEXT("CollisionPosWorld")), LOCTEXT("ReadAsync_CollisionPosWorldDescription", "The point in world space where the intersection occured"));
 		ReadRayTrace.AddOutput(FNiagaraVariable(FNiagaraTypeDefinition::GetVec3Def(), TEXT("CollisionNormal")), LOCTEXT("ReadAsync_CollisionNormalDescription", "The surface normal of the world geometry at the point of intersection"));
 	}
+	// DEPRECATE_END
 
 	{
 		const FText TraceChannelDescription = LOCTEXT("TraceChannelDescription", "The trace channel to collide against. Trace channels can be configured in the project settings.");
@@ -484,10 +411,12 @@ bool UNiagaraDataInterfaceCollisionQuery::GetFunctionHLSL(const FNiagaraDataInte
 	if ( (FunctionInfo.DefinitionName == NDICollisionQueryLocal::SceneDepthName) ||
 		 (FunctionInfo.DefinitionName == NDICollisionQueryLocal::CustomDepthName) ||
 		 (FunctionInfo.DefinitionName == NDICollisionQueryLocal::DistanceFieldName) ||
+		// DEPRECATE_BEGIN
 		 (FunctionInfo.DefinitionName == NDICollisionQueryLocal::IssueAsyncRayTraceName) ||
 		 (FunctionInfo.DefinitionName == NDICollisionQueryLocal::CreateAsyncRayTraceName) ||
 		 (FunctionInfo.DefinitionName == NDICollisionQueryLocal::ReserveAsyncRayTraceName) ||
 		 (FunctionInfo.DefinitionName == NDICollisionQueryLocal::ReadAsyncRayTraceName) )
+		// DEPRECATE_END
 	{
 		return true;
 	}
@@ -535,11 +464,6 @@ void UNiagaraDataInterfaceCollisionQuery::ValidateFunction(const FNiagaraFunctio
 }
 #endif
 
-bool UNiagaraDataInterfaceCollisionQuery::RequiresRayTracingScene() const
-{
-	return IsRayTracingEnabled() && GEnableGPUHWRTCollisions && MaxTracesPerParticle > 0;
-}
-
 #if WITH_EDITORONLY_DATA
 void UNiagaraDataInterfaceCollisionQuery::GetParameterDefinitionHLSL(const FNiagaraDataInterfaceGPUParamInfo& ParamInfo, FString& OutHLSL)
 {
@@ -594,13 +518,6 @@ bool UNiagaraDataInterfaceCollisionQuery::AppendCompileHash(FNiagaraCompileHashV
 	return true;
 }
 
-void UNiagaraDataInterfaceCollisionQuery::ModifyCompilationEnvironment(EShaderPlatform ShaderPlatform, FShaderCompilerEnvironment& OutEnvironment) const
-{
-	Super::ModifyCompilationEnvironment(ShaderPlatform, OutEnvironment);
-
-	OutEnvironment.SetDefine(TEXT("NIAGARA_SUPPORTS_RAY_TRACING"), ShouldCompileRayTracingShadersForProject(ShaderPlatform) ? 1 : 0);
-}
-
 #endif
 
 void UNiagaraDataInterfaceCollisionQuery::PerformQuerySyncCPU(FVectorVMExternalFunctionContext & Context)
@@ -732,55 +649,6 @@ bool UNiagaraDataInterfaceCollisionQuery::PerInstanceTickPostSimulate(void* PerI
 	return false;
 }
 
-bool UNiagaraDataInterfaceCollisionQuery::Equals(const UNiagaraDataInterface* Other) const
-{
-	if (!Super::Equals(Other))
-	{
-		return false;
-	}
-
-	const UNiagaraDataInterfaceCollisionQuery* OtherTyped = CastChecked<const UNiagaraDataInterfaceCollisionQuery>(Other);
-	return OtherTyped->MaxTracesPerParticle == MaxTracesPerParticle;
-}
-
-bool UNiagaraDataInterfaceCollisionQuery::CopyToInternal(UNiagaraDataInterface* Destination) const
-{
-	if (!Super::CopyToInternal(Destination))
-	{
-		return false;
-	}
-
-	UNiagaraDataInterfaceCollisionQuery* OtherTyped = CastChecked<UNiagaraDataInterfaceCollisionQuery>(Destination);
-	OtherTyped->MaxTracesPerParticle = MaxTracesPerParticle;
-	OtherTyped->MaxRetraces = MaxRetraces;
-	OtherTyped->MarkRenderDataDirty();
-	return true;
-}
-
-void UNiagaraDataInterfaceCollisionQuery::PushToRenderThreadImpl()
-{
-	FNiagaraDataIntefaceProxyCollisionQuery* RT_Proxy = GetProxyAs<FNiagaraDataIntefaceProxyCollisionQuery>();
-
-	// Push Updates to Proxy, first release any resources
-	ENQUEUE_RENDER_COMMAND(FUpdateDI)(
-		[RT_Proxy, RT_MaxTracesPerParticle = MaxTracesPerParticle, RT_MaxRetraces = MaxRetraces](FRHICommandListImmediate& RHICmdList)
-		{
-			RT_Proxy->RenderThreadInitialize(RT_MaxTracesPerParticle, RT_MaxRetraces);
-		});
-}
-
-#if WITH_EDITOR
-void UNiagaraDataInterfaceCollisionQuery::PostEditChangeProperty(struct FPropertyChangedEvent& PropertyChangedEvent)
-{
-	Super::PostEditChangeProperty(PropertyChangedEvent);
-
-	if (PropertyChangedEvent.Property && PropertyChangedEvent.Property->GetFName() == GET_MEMBER_NAME_CHECKED(UNiagaraDataInterfaceCollisionQuery, MaxTracesPerParticle))
-	{
-		MarkRenderDataDirty();
-	}
-}
-#endif
-
 //////////////////////////////////////////////////////////////////////////
 
 struct FNiagaraDataInterfaceParametersCS_CollisionQuery : public FNiagaraDataInterfaceParametersCS
@@ -791,23 +659,12 @@ public:
 	{
 		GlobalDistanceFieldParameters.Bind(ParameterMap);
 		SystemLWCTileParam.Bind(ParameterMap, *(NDICollisionQueryLocal::SystemLWCTileName + ParameterInfo.DataInterfaceHLSLSymbol));
-#if RHI_RAYTRACING
-		RayTracingEnabledParam.Bind(ParameterMap, *(NDICollisionQueryLocal::RayTracingEnabledParamName + ParameterInfo.DataInterfaceHLSLSymbol));
-		MaxRayTraceCountParam.Bind(ParameterMap, *(NDICollisionQueryLocal::MaxRayTraceCountParamName + ParameterInfo.DataInterfaceHLSLSymbol));
-		RayRequestsParam.Bind(ParameterMap, *(NDICollisionQueryLocal::RayRequestsParamName + ParameterInfo.DataInterfaceHLSLSymbol));
-		RayRequestOffsetParam.Bind(ParameterMap, *(NDICollisionQueryLocal::RayRequestsOffsetParamName + ParameterInfo.DataInterfaceHLSLSymbol));
-		IntersectionResultsParam.Bind(ParameterMap, *(NDICollisionQueryLocal::IntersectionResultsParamName + ParameterInfo.DataInterfaceHLSLSymbol));
-		IntersectionResultOffsetParam.Bind(ParameterMap, *(NDICollisionQueryLocal::IntersectionResultsOffsetParamName + ParameterInfo.DataInterfaceHLSLSymbol));
-		RayTraceCountsParam.Bind(ParameterMap, *(NDICollisionQueryLocal::RayTraceCountsParamName + ParameterInfo.DataInterfaceHLSLSymbol));
-		RayTraceCountsOffsetParam.Bind(ParameterMap, *(NDICollisionQueryLocal::RayTraceCountsOffsetParamName + ParameterInfo.DataInterfaceHLSLSymbol));
-#endif
 	}
 
 	void Set(FRHICommandList& RHICmdList, const FNiagaraDataInterfaceSetArgs& Context) const
 	{
 		check(IsInRenderingThread());
 
-		FNiagaraDataIntefaceProxyCollisionQuery* QueryDI = (FNiagaraDataIntefaceProxyCollisionQuery*)Context.DataInterface;
 		FRHIComputeShader* ComputeShaderRHI = Context.Shader.GetComputeShader();
 		SetShaderValue(RHICmdList, ComputeShaderRHI, SystemLWCTileParam, Context.SystemLWCTile);
 		
@@ -825,91 +682,12 @@ public:
 			{
 				GlobalDistanceFieldParameters.Set(RHICmdList, ComputeShaderRHI, FGlobalDistanceFieldParameterData());
 			}
-			
 		}
-
-#if RHI_RAYTRACING
-		const bool HasRayTracingParametersBound = RayRequestsParam.IsUAVBound()
-			|| IntersectionResultsParam.IsBound()
-			|| RayTraceCountsParam.IsBound();
-
-		if ((IsRayTracingEnabled() && GEnableGPUHWRTCollisions) || HasRayTracingParametersBound)
-		{
-			FNiagaraRayTracingHelper& RTHelper = Context.ComputeDispatchInterface->GetRayTracingHelper();
-			const FNiagaraRayTraceDispatchInfo* DispatchInfo = nullptr;
-			if (IsRayTracingEnabled() && GEnableGPUHWRTCollisions && QueryDI->MaxTracesPerParticle > 0)
-			{
-				DispatchInfo = &RTHelper.GetDispatch(QueryDI);
-			}
-			else
-			{
-				DispatchInfo = &RTHelper.GetDummyDispatch();
-			}
-
-			SetShaderValue(RHICmdList, ComputeShaderRHI, RayTracingEnabledParam, IsRayTracingEnabled() && GEnableGPUHWRTCollisions ? 1 : 0);
-			SetShaderValue(RHICmdList, ComputeShaderRHI, MaxRayTraceCountParam, DispatchInfo->MaxRays);
-
-			if (RayRequestsParam.IsUAVBound())
-			{
-				check(DispatchInfo->RayRequests.IsValid());
-				RHICmdList.SetUAVParameter(ComputeShaderRHI, RayRequestsParam.GetUAVIndex(), DispatchInfo->RayRequests.Buffer->UAV);
-				SetShaderValue(RHICmdList, ComputeShaderRHI, RayRequestOffsetParam, DispatchInfo->RayRequests.Offset);
-			}
-
-			if (IntersectionResultsParam.IsBound())
-			{
-				check(DispatchInfo->LastFrameRayTraceIntersections.IsValid());
-
-				SetSRVParameter(RHICmdList, ComputeShaderRHI, IntersectionResultsParam, DispatchInfo->LastFrameRayTraceIntersections.Buffer->SRV);
-				SetShaderValue(RHICmdList, ComputeShaderRHI, IntersectionResultOffsetParam, DispatchInfo->LastFrameRayTraceIntersections.Offset);
-			}
-
-			if (RayTraceCountsParam.IsUAVBound())
-			{
-				check(DispatchInfo->RayCounts.IsValid());
-
-				RHICmdList.SetUAVParameter(ComputeShaderRHI, RayTraceCountsParam.GetUAVIndex(), DispatchInfo->RayCounts.Buffer->UAV);
-				SetShaderValue(RHICmdList, ComputeShaderRHI, RayTraceCountsOffsetParam, DispatchInfo->RayCounts.Offset);
-			}
-		}
-		else
-		{
-			SetShaderValue(RHICmdList, ComputeShaderRHI, RayTracingEnabledParam, 0);
-			SetShaderValue(RHICmdList, ComputeShaderRHI, MaxRayTraceCountParam, 0);
-		}
-#endif
 	}
 
-#if RHI_RAYTRACING
-	void Unset(FRHICommandList& RHICmdList, const FNiagaraDataInterfaceSetArgs& Context) const
-	{
-		FRHIComputeShader* ComputeShaderRHI = Context.Shader.GetComputeShader();
-		if (RayRequestsParam.IsUAVBound())
-		{
-			RayRequestsParam.UnsetUAV(RHICmdList, ComputeShaderRHI);
-		}
-
-		if (RayTraceCountsParam.IsUAVBound())
-		{
-			RayTraceCountsParam.UnsetUAV(RHICmdList, ComputeShaderRHI);
-		}
-	}
-#endif
-
 private:
 	LAYOUT_FIELD(FGlobalDistanceFieldParameters, GlobalDistanceFieldParameters);
 	LAYOUT_FIELD(FShaderParameter, SystemLWCTileParam);
-
-#if RHI_RAYTRACING
-	LAYOUT_FIELD(FShaderParameter, RayTracingEnabledParam);
-	LAYOUT_FIELD(FShaderParameter, MaxRayTraceCountParam);
-	LAYOUT_FIELD(FRWShaderParameter, RayRequestsParam);
-	LAYOUT_FIELD(FShaderParameter, RayRequestOffsetParam);
-	LAYOUT_FIELD(FShaderResourceParameter, IntersectionResultsParam);
-	LAYOUT_FIELD(FShaderParameter, IntersectionResultOffsetParam);
-	LAYOUT_FIELD(FRWShaderParameter, RayTraceCountsParam);
-	LAYOUT_FIELD(FShaderParameter, RayTraceCountsOffsetParam);
-#endif
 };
 
 IMPLEMENT_TYPE_LAYOUT(FNiagaraDataInterfaceParametersCS_CollisionQuery);
diff --git a/Engine/Plugins/FX/Niagara/Source/Niagara/Private/NiagaraFunctionLibrary.cpp b/Engine/Plugins/FX/Niagara/Source/Niagara/Private/NiagaraFunctionLibrary.cpp
--- a/Engine/Plugins/FX/Niagara/Source/Niagara/Private/NiagaraFunctionLibrary.cpp
+++ b/Engine/Plugins/FX/Niagara/Source/Niagara/Private/NiagaraFunctionLibrary.cpp
@@ -10,6 +10,7 @@
 #include "ContentStreaming.h"
 #include "Internationalization/Internationalization.h"
 
+#include "NiagaraAsyncGpuTraceHelper.h"
 #include "NiagaraWorldManager.h"
 #include "DataInterface/NiagaraDataInterfaceStaticMesh.h"
 #include "NiagaraStats.h"
@@ -1844,12 +1845,12 @@ void UNiagaraFunctionLibrary::InitVectorVMFastPathOps()
 
 void UNiagaraFunctionLibrary::SetComponentNiagaraGPURayTracedCollisionGroup(UObject* WorldContextObject, UPrimitiveComponent* Primitive, int32 CollisionGroup)
 {
-#if RHI_RAYTRACING
+#if NIAGARA_ASYNC_GPU_TRACE_COLLISION_GROUPS
 	if ( UWorld* World = GEngine->GetWorldFromContextObject(WorldContextObject, EGetWorldErrorMode::LogAndReturnNull) )
 	{
 		if ( FNiagaraGpuComputeDispatchInterface* ComputeDispatchInterface = FNiagaraGpuComputeDispatchInterface::Get(World) )
 		{
-			ComputeDispatchInterface->SetPrimitiveRayTracingCollisionGroup(Primitive, CollisionGroup);
+			ComputeDispatchInterface->GetAsyncGpuTraceHelper().SetPrimitiveRayTracingCollisionGroup_GT(Primitive, CollisionGroup);
 		}
 	}
 #endif
@@ -1857,7 +1858,7 @@ void UNiagaraFunctionLibrary::SetComponentNiagaraGPURayTracedCollisionGroup(UObj
 
 void UNiagaraFunctionLibrary::SetActorNiagaraGPURayTracedCollisionGroup(UObject* WorldContextObject, AActor* Actor, int32 CollisionGroup)
 {
-#if RHI_RAYTRACING
+#if NIAGARA_ASYNC_GPU_TRACE_COLLISION_GROUPS
 	if (Actor == nullptr)
 	{
 		return;
@@ -1867,11 +1868,13 @@ void UNiagaraFunctionLibrary::SetActorNiagaraGPURayTracedCollisionGroup(UObject*
 	{
 		if (FNiagaraGpuComputeDispatchInterface* ComputeDispatchInterface = FNiagaraGpuComputeDispatchInterface::Get(World))
 		{
+			FNiagaraAsyncGpuTraceHelper& Helper = ComputeDispatchInterface->GetAsyncGpuTraceHelper();
+
 			Actor->ForEachComponent<UPrimitiveComponent>(
 				/*bIncludeFromChildActors*/true,
-				[&ComputeDispatchInterface, &CollisionGroup](UPrimitiveComponent* PrimitiveComponent)
+				[&Helper, &CollisionGroup](UPrimitiveComponent* PrimitiveComponent)
 				{
-					ComputeDispatchInterface->SetPrimitiveRayTracingCollisionGroup(PrimitiveComponent, CollisionGroup);
+					Helper.SetPrimitiveRayTracingCollisionGroup_GT(PrimitiveComponent, CollisionGroup);
 				});
 			}
 	}
@@ -1880,12 +1883,12 @@ void UNiagaraFunctionLibrary::SetActorNiagaraGPURayTracedCollisionGroup(UObject*
 
 int32 UNiagaraFunctionLibrary::AcquireNiagaraGPURayTracedCollisionGroup(UObject* WorldContextObject)
 {
-#if RHI_RAYTRACING
+#if NIAGARA_ASYNC_GPU_TRACE_COLLISION_GROUPS
 	if (UWorld* World = GEngine->GetWorldFromContextObject(WorldContextObject, EGetWorldErrorMode::LogAndReturnNull))
 	{
 		if (FNiagaraGpuComputeDispatchInterface* ComputeDispatchInterface = FNiagaraGpuComputeDispatchInterface::Get(World))
 		{
-			return ComputeDispatchInterface->AcquireGPURayTracedCollisionGroup();
+			return ComputeDispatchInterface->GetAsyncGpuTraceHelper().AcquireGPURayTracedCollisionGroup_GT();
 		}
 	}
 #endif
@@ -1894,12 +1897,12 @@ int32 UNiagaraFunctionLibrary::AcquireNiagaraGPURayTracedCollisionGroup(UObject*
 
 void UNiagaraFunctionLibrary::ReleaseNiagaraGPURayTracedCollisionGroup(UObject* WorldContextObject, int32 CollisionGroup)
 {
-#if RHI_RAYTRACING
+#if NIAGARA_ASYNC_GPU_TRACE_COLLISION_GROUPS
 	if (UWorld* World = GEngine->GetWorldFromContextObject(WorldContextObject, EGetWorldErrorMode::LogAndReturnNull))
 	{
 		if (FNiagaraGpuComputeDispatchInterface* ComputeDispatchInterface = FNiagaraGpuComputeDispatchInterface::Get(World))
 		{
-			ComputeDispatchInterface->ReleaseGPURayTracedCollisionGroup(CollisionGroup);
+			ComputeDispatchInterface->GetAsyncGpuTraceHelper().ReleaseGPURayTracedCollisionGroup_GT(CollisionGroup);
 		}
 	}
 #endif
diff --git a/Engine/Plugins/FX/Niagara/Source/Niagara/Private/NiagaraGpuComputeDispatch.cpp b/Engine/Plugins/FX/Niagara/Source/Niagara/Private/NiagaraGpuComputeDispatch.cpp
--- a/Engine/Plugins/FX/Niagara/Source/Niagara/Private/NiagaraGpuComputeDispatch.cpp
+++ b/Engine/Plugins/FX/Niagara/Source/Niagara/Private/NiagaraGpuComputeDispatch.cpp
@@ -10,13 +10,13 @@
 #include "Engine/World.h"
 #include "Misc/ScopeExit.h"
 
+#include "NiagaraAsyncGpuTraceHelper.h"
 #include "NiagaraDataInterfaceRW.h"
 #if NIAGARA_COMPUTEDEBUG_ENABLED
 #include "NiagaraGpuComputeDebug.h"
 #endif
 #include "NiagaraGPUProfilerInterface.h"
 #include "NiagaraGpuReadbackManager.h"
-#include "NiagaraRayTracingHelper.h"
 #include "NiagaraRenderer.h"
 #include "NiagaraScriptExecutionContext.h"
 #include "NiagaraSystemGpuComputeProxy.h"
@@ -207,9 +207,8 @@ FNiagaraGpuComputeDispatch::FNiagaraGpuComputeDispatch(ERHIFeatureLevel::Type In
 	OwnerCBufferLayout   = new FNiagaraRHIUniformBufferLayout(TEXT("Niagara GPU Owner CBuffer"),   sizeof(FNiagaraOwnerParameters));
 	EmitterCBufferLayout = new FNiagaraRHIUniformBufferLayout(TEXT("Niagara GPU Emitter CBuffer"), sizeof(FNiagaraEmitterParameters));
 
-#if RHI_RAYTRACING
-	RayTracingHelper.Reset(new FNiagaraRayTracingHelper(InShaderPlatform));
-#endif
+
+	AsyncGpuTraceHelper.Reset(new FNiagaraAsyncGpuTraceHelper(InShaderPlatform, InFeatureLevel, this));
 
 #if NIAGARA_COMPUTEDEBUG_ENABLED
 	GpuComputeDebugPtr.Reset(new FNiagaraGpuComputeDebug(FeatureLevel));
@@ -219,18 +218,13 @@ FNiagaraGpuComputeDispatch::FNiagaraGpuComputeDispatch(ERHIFeatureLevel::Type In
 #endif
 	GpuReadbackManagerPtr.Reset(new FNiagaraGpuReadbackManager());
 	EmptyUAVPoolPtr.Reset(new FNiagaraEmptyUAVPool());
-
-	FPrimitiveSceneInfo::OnGPUSceneInstancesAllocated.AddRaw(this, &FNiagaraGpuComputeDispatch::OnPrimitiveGPUSceneInstancesAllocated);
-	FPrimitiveSceneInfo::OnGPUSceneInstancesAllocated.AddRaw(this, &FNiagaraGpuComputeDispatch::OnPrimitiveGPUSceneInstancesFreed);
 }
 
 FNiagaraGpuComputeDispatch::~FNiagaraGpuComputeDispatch()
 {
 	FinishDispatches();
 
-#if RHI_RAYTRACING
-	RayTracingHelper->Reset();
-#endif
+	AsyncGpuTraceHelper->Reset();
 
 	GlobalCBufferLayout = nullptr;
 	SystemCBufferLayout = nullptr;
@@ -512,10 +506,8 @@ void FNiagaraGpuComputeDispatch::ProcessPendingTicksFlush(FRHICommandListImmedia
 			//UE_LOG(LogNiagara, Log, TEXT("FNiagaraGpuComputeDispatch: Queued ticks are being Destroyed due to not rendering.  This may result in undesirable simulation artifacts."));
 
 			FinishDispatches();
+			AsyncGpuTraceHelper->Reset();
 
-#if RHI_RAYTRACING
-			RayTracingHelper->Reset();
-#endif
 			break;
 		}
 	}
@@ -1677,9 +1669,7 @@ void FNiagaraGpuComputeDispatch::PreInitViews(FRDGBuilder& GraphBuilder, bool bA
 			CalculateCrossGPUTransferLocation();
 		#endif
 
-#if RHI_RAYTRACING
-			RayTracingHelper->BeginFrame(GraphBuilder.RHICmdList, HasRayTracingScene());
-#endif
+			AsyncGpuTraceHelper->BeginFrame(GraphBuilder.RHICmdList, this);
 
 			if ( FNiagaraGpuComputeDispatchLocal::GDebugLogging)
 			{
@@ -1766,9 +1756,7 @@ void FNiagaraGpuComputeDispatch::PostRenderOpaque(FRDGBuilder& GraphBuilder, TCo
 			NiagaraSceneTextures = PassParameters;
 			ON_SCOPE_EXIT{ NiagaraSceneTextures = nullptr; };
 
-	#if RHI_RAYTRACING
-			BuildRayTracingSceneInfo(RHICmdList, Views);
-	#endif
+			AsyncGpuTraceHelper->PostRenderOpaque(RHICmdList, this, Views);
 
 			CurrentPassViews = Views;
 
@@ -1779,9 +1767,7 @@ void FNiagaraGpuComputeDispatch::PostRenderOpaque(FRDGBuilder& GraphBuilder, TCo
 
 			FinishDispatches();
 
-	#if RHI_RAYTRACING
-			RayTracingHelper->EndFrame(RHICmdList, HasRayTracingScene(), GetScene());
-	#endif
+			AsyncGpuTraceHelper->EndFrame(RHICmdList, this);
 
 			// Clear CurrentPassViews
 			CurrentPassViews = TConstArrayView<FViewInfo>();
@@ -2037,86 +2023,10 @@ void FNiagaraGpuComputeDispatch::GenerateSortKeys(FRHICommandListImmediate& RHIC
 	RHICmdList.EndUAVOverlap(MakeArrayView(OverlapUAVs, NumOverlapUAVs));
 }
 
-#if RHI_RAYTRACING
-void FNiagaraGpuComputeDispatch::BuildRayTracingSceneInfo(FRHICommandList& RHICmdList, TConstArrayView<FViewInfo> Views)
-{
-	if (NumProxiesThatRequireRayTracingScene > 0)
-	{
-		RayTracingHelper->BuildRayTracingSceneInfo(RHICmdList, Views);
-
-		RayTracingHelper->UpdateCollisionGroupMap(RHICmdList, GetScene(), FeatureLevel);
-	}
-	else
-	{
-		RayTracingHelper->Reset();
-	}
-}
-
-void FNiagaraGpuComputeDispatch::ResetRayTracingSceneInfo()
-{
-	RayTracingHelper->Reset();
-}
-
-
-FNiagaraRayTracingHelper& FNiagaraGpuComputeDispatch::GetRayTracingHelper()const
-{
-	check(RayTracingHelper.IsValid());
-	return *RayTracingHelper.Get();
-}
-
-bool FNiagaraGpuComputeDispatch::HasRayTracingScene() const
-{
-	return NumProxiesThatRequireRayTracingScene && RayTracingHelper->IsValid();
-}
-
-void FNiagaraGpuComputeDispatch::SetPrimitiveRayTracingCollisionGroup(UPrimitiveComponent* Primitive, uint32 Group)
+FNiagaraAsyncGpuTraceHelper& FNiagaraGpuComputeDispatch::GetAsyncGpuTraceHelper() const
 {
-	check(IsInGameThread());
-	if (Primitive == nullptr || Primitive->SceneProxy == nullptr)
-	{
-		return;
-	}
-
-	ENQUEUE_RENDER_COMMAND(NiagaraSetPrimHWRTCollisionGroup)(
-		[RT_NiagaraBatcher = this, Proxy = Primitive->SceneProxy, CollisionGroup=Group](FRHICommandListImmediate& RHICmdList)
-		{
-			check(Proxy);
-			RT_NiagaraBatcher->RayTracingHelper->SetPrimitiveCollisionGroup(*Proxy->GetPrimitiveSceneInfo(), CollisionGroup);
-		}
-	);
-}
-
-int32 FNiagaraGpuComputeDispatch::AcquireGPURayTracedCollisionGroup()
-{
-	check(IsInGameThread());//one of the few batcher functions that should be called from the GT
-	if (FreeGPURayTracedCollisionGroups.Num() > 0)
-	{
-		return FreeGPURayTracedCollisionGroups.Pop();
-	}
-
-	return NumGPURayTracedCollisionGroups++;
-}
-
-void FNiagaraGpuComputeDispatch::ReleaseGPURayTracedCollisionGroup(int32 CollisionGroup)
-{
-	check(IsInGameThread());//one of the few batcher functions that should be called from the GT
-	FreeGPURayTracedCollisionGroups.Add(CollisionGroup);
-}
-
-#endif
-
-void FNiagaraGpuComputeDispatch::OnPrimitiveGPUSceneInstancesAllocated()
-{
-#if RHI_RAYTRACING
-	RayTracingHelper->RefreshPrimitiveInstanceData();
-#endif
-}
-
-void FNiagaraGpuComputeDispatch::OnPrimitiveGPUSceneInstancesFreed()
-{
-#if RHI_RAYTRACING
-	RayTracingHelper->RefreshPrimitiveInstanceData();
-#endif
+	check(AsyncGpuTraceHelper.IsValid());
+	return *AsyncGpuTraceHelper.Get();
 }
 
 /* Set shader parameters for data interfaces
diff --git a/Engine/Plugins/FX/Niagara/Source/Niagara/Private/NiagaraGpuComputeDispatch.h b/Engine/Plugins/FX/Niagara/Source/Niagara/Private/NiagaraGpuComputeDispatch.h
--- a/Engine/Plugins/FX/Niagara/Source/Niagara/Private/NiagaraGpuComputeDispatch.h
+++ b/Engine/Plugins/FX/Niagara/Source/Niagara/Private/NiagaraGpuComputeDispatch.h
@@ -27,8 +27,8 @@ the same VectorVM byte code / compute shader code
 #include "Tickable.h"
 
 class FGPUSortManager;
-class FNiagaraRayTracingHelper;
 class FDistanceFieldSceneData;
+class FNiagaraAsyncGpuTraceHelper;
 
 class FNiagaraGpuComputeDispatch : public FNiagaraGpuComputeDispatchInterface
 {
@@ -125,23 +125,12 @@ public:
 	/** Debug only function to readback data. */
 	virtual void AddDebugReadback(FNiagaraSystemInstanceID InstanceID, TSharedPtr<struct FNiagaraScriptDebuggerInfo, ESPMode::ThreadSafe> DebugInfo, FNiagaraComputeExecutionContext* Context) override;
 
-#if RHI_RAYTRACING
-	virtual FNiagaraRayTracingHelper& GetRayTracingHelper() const override;
-	virtual bool HasRayTracingScene() const override;
-
-	virtual void SetPrimitiveRayTracingCollisionGroup(UPrimitiveComponent* PrimitiveComponent, uint32 Group) override;
-
-	virtual int32 AcquireGPURayTracedCollisionGroup() override;
-	virtual void ReleaseGPURayTracedCollisionGroup(int32 CollisionGroup) override;
-#endif
-
 #if WITH_MGPU
 	virtual void MultiGPUResourceModified(FRHICommandList& RHICmdList, FRHIBuffer* Buffer, bool bRequiredForSimulation, bool bRequiredForRendering) const override;
 	virtual void MultiGPUResourceModified(FRHICommandList& RHICmdList, FRHITexture* Texture, bool bRequiredForSimulation, bool bRequiredForRendering) const override;
 #endif
 
-	void OnPrimitiveGPUSceneInstancesAllocated();
-	void OnPrimitiveGPUSceneInstancesFreed();
+	virtual FNiagaraAsyncGpuTraceHelper& GetAsyncGpuTraceHelper() const override;
 
 private:
 	void DumpDebugFrame();
@@ -172,11 +161,6 @@ private:
 	 */
 	void GenerateSortKeys(FRHICommandListImmediate& RHICmdList, int32 BatchId, int32 NumElementsInBatch, EGPUSortFlags Flags, FRHIUnorderedAccessView* KeysUAV, FRHIUnorderedAccessView* ValuesUAV);
 
-#if RHI_RAYTRACING
-	void BuildRayTracingSceneInfo(FRHICommandList& RHICmdList, TConstArrayView<FViewInfo> Views);
-	void ResetRayTracingSceneInfo();
-#endif
-
 	void FinishDispatches();
 
 	void UpdateFreeIDsListSizesBuffer(FRHICommandList& RHICmdList, uint32 NumInstances);
@@ -190,10 +174,7 @@ private:
 	/** All sort tasks registered in AddSortedGPUSimulation(). Holds all the data required in GenerateSortKeys(). */
 	TArray<FNiagaraGPUSortInfo> SimulationsToSort;
 
-#if RHI_RAYTRACING
-	// helper object which is valid during PostRenderOpaque if DataInterfaces are present that require a RayTrace scene
-	TUniquePtr<FNiagaraRayTracingHelper> RayTracingHelper;
-#endif
+	TUniquePtr<FNiagaraAsyncGpuTraceHelper> AsyncGpuTraceHelper;
 
 #if WITH_NIAGARA_GPU_PROFILER
 	TUniquePtr<FNiagaraGPUProfiler> GPUProfilerPtr;
@@ -254,10 +235,6 @@ private:
 	void WaitForMultiGPUBuffers(FRHICommandList& RHICmdList, ENiagaraGpuComputeTickStage::Type TickStage);
 #endif // WITH_MGPU
 
-	/** Pool of free GPU ray traced collision groups. */
-	TArray<int32> FreeGPURayTracedCollisionGroups;
-	int32 NumGPURayTracedCollisionGroups = 0;
-
 	// Cached information to build a dummy view info if necessary
 	FIntRect CachedViewRect;
 
diff --git a/Engine/Plugins/FX/Niagara/Source/Niagara/Private/NiagaraGpuScratchPad.h b/Engine/Plugins/FX/Niagara/Source/Niagara/Private/NiagaraGpuScratchPad.h
new file mode 100644
--- /dev/null
+++ b/Engine/Plugins/FX/Niagara/Source/Niagara/Private/NiagaraGpuScratchPad.h
@@ -0,0 +1,415 @@
+// Copyright Epic Games, Inc. All Rights Reserved.
+#pragma once
+
+#include "CoreMinimal.h"
+
+#include "RHIDefinitions.h"
+#include "PrimitiveSceneInfo.h"
+#include "RHI.h"
+#include "RHIUtilities.h"
+
+class FScene;
+class FViewInfo;
+class FRayTracingPipelineState;
+class FRHICommandList;
+class FRHIRayTracingScene;
+class FRHIShaderResourceView;
+class FRHIUniformBuffer;
+class FRHIUnorderedAccessView;
+struct FRWBuffer;
+struct FNiagaraDataInterfaceProxy;
+
+
+//////////////////////////////////////////////////////////////////////////
+//TODO: Move scratch pad buffers to core rendering code along side global dynamic read buffers etc.
+
+/** Allocates space from a pool of fixed size static RWBuffers.
+*	Allocations are valid between calls to Reset.
+*	Useful for transient working data generated in one dispatch and consumed in others.
+*	All buffers in the same scratch pad are transitioned as one.
+*/
+class NIAGARA_API FNiagaraGpuScratchPad
+{
+public:
+	struct FAllocation
+	{
+		FRWBuffer* Buffer = nullptr;
+		uint32 Offset = 0;
+
+		FORCEINLINE bool IsValid()const { return Buffer != nullptr; }
+	};
+
+	struct FScratchBuffer
+	{
+	public:
+		struct FBuffer
+		{
+			FRWBuffer Buffer;
+			uint32 Used = 0;
+		};
+
+		FScratchBuffer(uint32 InBucketSize, EPixelFormat InPixelFormat, uint32 InBytesPerElement, EBufferUsageFlags InBufferUsageFlags)
+			: MinBucketSize(InBucketSize), PixelFormat(InPixelFormat), BytesPerElement(InBytesPerElement), BufferUsageFlags(InBufferUsageFlags)
+		{}
+
+		/** Allocates NumElements from the buffers. Will try to find an existing buffer to fit the allocation but will create a new one if needed. */
+		bool Allocate(uint32 NumElements, FString& BufferDebugName, FRHICommandList& RHICmdList, FAllocation& OutAllocation)
+		{
+			//Look for a buffer to fit this allocation.
+			FBuffer* ToUse = nullptr;
+			for (FBuffer& Buffer : Buffers)
+			{
+				uint32 NumBufferElements = Buffer.Buffer.NumBytes / BytesPerElement;
+				if (NumElements + Buffer.Used <= NumBufferElements)
+				{
+					ToUse = &Buffer;
+				}
+			}
+
+			bool bNew = false;
+			if (ToUse == nullptr)
+			{
+				//Allocate a new buffer that will fit this allocation.
+				FBuffer* NewBuffer = new FBuffer();
+				Buffers.Add(NewBuffer);
+
+				//Round up the size to pow 2 to ensure a slowly growing user doesn't just keep allocating new buckets. Probably a better guess size here than power of 2.
+				uint32 NewBufferSize = FPlatformMath::RoundUpToPowerOfTwo(FMath::Max(MinBucketSize, NumElements));
+
+				NewBuffer->Buffer.Initialize(*BufferDebugName, BytesPerElement, NewBufferSize, PixelFormat, CurrentAccess, BufferUsageFlags);
+			
+				ToUse = NewBuffer;
+
+				NumAllocatedBytes += NewBuffer->Buffer.NumBytes;
+
+				bNew = true;
+			}
+
+			OutAllocation.Buffer = &ToUse->Buffer;
+			OutAllocation.Offset = ToUse->Used;
+			ToUse->Used += NumElements;
+			NumUsedBytes += NumElements * BytesPerElement;
+			return bNew;
+		}
+
+		/** Resets the buffer and releases it's internal RWBuffers too. */
+		void Release()
+		{
+			Reset();
+			CurrentAccess = ERHIAccess::UAVCompute;
+			for (FBuffer& Buffer : Buffers)
+			{
+				Buffer.Buffer.Release();
+				Buffer.Used = 0;
+			}
+			Buffers.Reset();
+
+			NumAllocatedBytes = 0;
+		}
+
+		/** Resets the usage for all buffers to zero. */
+		void Reset()
+		{
+			for (FBuffer& Buffer : Buffers)
+			{
+				Buffer.Used = 0;
+			}
+
+			NumUsedBytes = 0;
+		}
+
+		template<typename TFunc>
+		void ForEachBuffer(TFunc Func)
+		{
+			for (FBuffer& Buffer : Buffers)
+			{
+				Func(Buffer.Buffer);
+			}
+		}
+
+		/** Transitions all buffers in the scratch pad. */
+		void Transition(FRHICommandList& RHICmdList, ERHIAccess InPreviousState, ERHIAccess InNewState, EResourceTransitionFlags InFlags = EResourceTransitionFlags::None)
+		{
+			//Store off current access state so we can place any new buffers in the same state.
+			CurrentAccess = InNewState;
+
+ 			TArray<FRHITransitionInfo, TMemStackAllocator<>> Transitions;
+ 			Transitions.Reserve(Buffers.Num());
+ 			for (FBuffer& Buffer : Buffers)
+ 			{
+ 				Transitions.Emplace(Buffer.Buffer.UAV, InPreviousState, InNewState, InFlags);
+ 			}
+ 
+ 			RHICmdList.Transition(Transitions);
+		}
+
+		/** Returns the total bytes allocated for all buffers in the scratch pad. */
+		FORCEINLINE uint32 AllocatedBytes()const
+		{
+			return NumAllocatedBytes;
+		}
+
+		/** Returns the number of used bytes allocated for all buffers in the scratch pad. */
+		FORCEINLINE uint32 UsedBytes()const
+		{
+			return NumUsedBytes;
+		}
+
+		FORCEINLINE ERHIAccess GetExpectedCurrentAccess()const{return CurrentAccess;}
+	private:
+
+		uint32 MinBucketSize = 4096;
+		TIndirectArray<FBuffer> Buffers;
+
+		EPixelFormat PixelFormat;
+		uint32 BytesPerElement = 0;
+		EBufferUsageFlags BufferUsageFlags = EBufferUsageFlags::None;
+		ERHIAccess CurrentAccess = ERHIAccess::UAVCompute;
+
+		uint32 NumAllocatedBytes = 0;
+		uint32 NumUsedBytes = 0;
+	};
+
+	FNiagaraGpuScratchPad(uint32 InBucketSize, const FString& InDebugName, EBufferUsageFlags BufferUsageFlags = BUF_Static)
+		: BucketSize(InBucketSize)
+		, DebugName(InDebugName)
+		, ScratchBufferFloat(BucketSize, PF_R32_FLOAT, sizeof(float), BufferUsageFlags)
+		, ScratchBufferUInt(BucketSize, PF_R32_UINT, sizeof(uint32), BufferUsageFlags)
+	{
+	}
+
+	/** Allocate a number of elements from the scratch pad. This will attempt to allocate from an existing buffer but will create a new one if it can't. */
+	template<typename T>
+	FORCEINLINE FAllocation Alloc(uint32 Size, FRHICommandList& RHICmdList, bool bClearNew);
+
+	/** Resets the scratch pad and also releases all buffers it holds. */
+	void Release()
+	{
+		ScratchBufferFloat.Release();
+		ScratchBufferUInt.Release();
+	}
+
+	/** Resets the usage counters for buffers in this scratch pad, invalidating all previous allocations. Does not affect actual allocations. Optionally clears all buffers to zero. */
+	void Reset(FRHICommandList& RHICmdList, bool bClear)
+	{
+		ScratchBufferFloat.Reset();
+		ScratchBufferUInt.Reset();
+
+		if (bClear)
+		{
+			Transition(RHICmdList, ERHIAccess::Unknown, ERHIAccess::UAVCompute);
+
+			auto ClearFloatBuffer = [&](FRWBuffer& Buffer)
+			{
+				RHICmdList.ClearUAVFloat(Buffer.UAV, FVector4f(0.0f));
+			};
+			ScratchBufferFloat.ForEachBuffer(ClearFloatBuffer);
+
+			auto ClearUIntBuffer = [&](FRWBuffer& Buffer)
+			{
+				RHICmdList.ClearUAVUint(Buffer.UAV, FUintVector4(0, 0, 0, 0));
+			};
+			ScratchBufferUInt.ForEachBuffer(ClearUIntBuffer);
+
+			Transition(RHICmdList, ERHIAccess::UAVCompute, ERHIAccess::UAVCompute);
+		}
+	}
+
+	/** Transitions all buffers in the scratch pad. */
+	void Transition(FRHICommandList& RHICmdList, ERHIAccess InPreviousState, ERHIAccess InNewState, EResourceTransitionFlags InFlags = EResourceTransitionFlags::None)
+	{
+		ScratchBufferFloat.Transition(RHICmdList, InPreviousState, InNewState, InFlags);
+		ScratchBufferUInt.Transition(RHICmdList, InPreviousState, InNewState, InFlags);
+	}
+
+	/** Returns the total bytes allocated for all buffers in the scratch pad. */
+	FORCEINLINE uint32 AllocatedBytes()const
+	{
+		return ScratchBufferFloat.AllocatedBytes() + ScratchBufferUInt.AllocatedBytes();
+	}
+
+	/** Returns the number of used bytes allocated for all buffers in the scratch pad. */
+	FORCEINLINE uint32 UsedBytes()const
+	{
+		return ScratchBufferFloat.UsedBytes() + ScratchBufferUInt.UsedBytes();
+	}
+
+private:
+
+	/** Size of each FRWBuffer in the pool. */
+	uint32 BucketSize = 4096;
+	FString DebugName;
+
+	FScratchBuffer ScratchBufferFloat;
+	FScratchBuffer ScratchBufferUInt;
+};
+
+template<>
+FORCEINLINE FNiagaraGpuScratchPad::FAllocation FNiagaraGpuScratchPad::Alloc<float>(uint32 Size, FRHICommandList& RHICmdList, bool bClearNew)
+{
+	FAllocation Allocation;
+	bool bNewBuffer = ScratchBufferFloat.Allocate(Size, DebugName, RHICmdList, Allocation);
+	if (bNewBuffer && bClearNew)
+	{
+		RHICmdList.Transition(FRHITransitionInfo(Allocation.Buffer->UAV, ScratchBufferFloat.GetExpectedCurrentAccess(), ERHIAccess::UAVCompute));
+		RHICmdList.ClearUAVFloat(Allocation.Buffer->UAV, FVector4f(0.0f));
+		RHICmdList.Transition(FRHITransitionInfo(Allocation.Buffer->UAV, ERHIAccess::UAVCompute, ScratchBufferFloat.GetExpectedCurrentAccess()));
+	}
+	return Allocation;
+}
+
+template<>
+FORCEINLINE FNiagaraGpuScratchPad::FAllocation FNiagaraGpuScratchPad::Alloc<uint32>(uint32 Size, FRHICommandList& RHICmdList, bool bClearNew)
+{
+	FAllocation Allocation;
+	bool bNewBuffer = ScratchBufferUInt.Allocate(Size, DebugName, RHICmdList, Allocation);
+	if (bNewBuffer && bClearNew)
+	{
+		RHICmdList.Transition(FRHITransitionInfo(Allocation.Buffer->UAV, ScratchBufferUInt.GetExpectedCurrentAccess(), ERHIAccess::UAVCompute));
+		RHICmdList.ClearUAVUint(Allocation.Buffer->UAV, FUintVector4(0, 0, 0, 0));
+		RHICmdList.Transition(FRHITransitionInfo(Allocation.Buffer->UAV, ERHIAccess::UAVCompute, ScratchBufferUInt.GetExpectedCurrentAccess()));
+	}
+	return Allocation;
+}
+
+//////////////////////////////////////////////////////////////////////////
+
+/** Behaves similar to FNiagaraGpuScratchPad but for use with RWStructuredBuffers
+*/
+template<typename T>
+class NIAGARA_API FNiagaraGpuScratchPadStructured
+{
+public:
+	struct FAllocation
+	{
+		FRWBufferStructured* Buffer = nullptr;
+		uint32 Offset = 0;
+
+		FORCEINLINE bool IsValid()const { return Buffer != nullptr; }
+	};
+
+	struct FBuffer
+	{
+		FRWBufferStructured Buffer;
+		uint32 Used = 0;
+	};
+
+	FNiagaraGpuScratchPadStructured(uint32 InBucketSize, const FString& InDebugName, EBufferUsageFlags InBufferUsageFlags = EBufferUsageFlags::Static | EBufferUsageFlags::UnorderedAccess)
+		: MinBucketSize(InBucketSize)
+		, DebugName(InDebugName)
+		, BufferUsageFlags(InBufferUsageFlags)
+	{
+	}
+
+	/** Allocates NumElements from the scratch pad. Will try to allocate from an existing buffer but will create a new one if needed. */
+	FAllocation Alloc(uint32 NumElements)
+	{
+		//Look for a buffer to fit this allocation.
+		FBuffer* ToUse = nullptr;
+		for (FBuffer& Buffer : Buffers)
+		{
+			uint32 NumBufferElements = Buffer.Buffer.NumBytes / sizeof(T);
+			if (NumElements + Buffer.Used < NumBufferElements)
+			{
+				ToUse = &Buffer;
+			}
+		}
+
+		if (ToUse == nullptr)
+		{
+			//Allocate a new buffer that will fit this allocation.
+
+			FBuffer* NewBuffer = new FBuffer();
+			Buffers.Add(NewBuffer);
+
+			//Round up the size to pow 2 to ensure a slowly growing user doesn't just keep allocating new buckets. Probably a better guess size here than power of 2.
+			uint32 NewBufferSize = FPlatformMath::RoundUpToPowerOfTwo(FMath::Max(MinBucketSize, NumElements));
+
+			NewBuffer->Buffer.Initialize(
+				*DebugName,
+				sizeof(T),
+				NewBufferSize,
+				BufferUsageFlags,
+				false /*bUseUavCounter*/,
+				false /*bAppendBuffer*/,
+				CurrentAccess);
+
+			ToUse = NewBuffer;
+			NumAllocatedBytes += NewBuffer->Buffer.NumBytes;
+		}
+
+		FAllocation Allocation;
+		Allocation.Buffer = &ToUse->Buffer;
+		Allocation.Offset = ToUse->Used;
+		ToUse->Used += NumElements;
+		NumUsedBytes += NumElements * sizeof(T);
+		return Allocation;
+	}
+
+	/** Resets the scratch pad and releases all buffers it holds. */
+	void Release()
+	{
+		Reset();
+		CurrentAccess = ERHIAccess::UAVCompute;
+		for (FBuffer& Buffer : Buffers)
+		{
+			Buffer.Buffer.Release();
+			Buffer.Used = 0;
+		}
+		Buffers.Reset();
+
+		NumAllocatedBytes = 0;
+	}
+
+	/** Resets the scratch pad, invalidating all previous allocations. Does not modify actual RWBuffer allocations. */
+	void Reset()
+	{
+		for (FBuffer& Buffer : Buffers)
+		{
+			Buffer.Used = 0;
+		}
+
+		NumUsedBytes = 0;
+	}
+
+	/** Transitions all buffers in the scratch pad. */
+	void Transition(FRHICommandList& RHICmdList, ERHIAccess InPreviousState, ERHIAccess InNewState, EResourceTransitionFlags InFlags = EResourceTransitionFlags::None)
+	{
+		//Store off current access state so we can place any new buffers in the same state.
+		CurrentAccess = InNewState;
+
+ 		TArray<FRHITransitionInfo, TMemStackAllocator<>> Transitions;
+ 		Transitions.Reserve(Buffers.Num());
+ 		for (FBuffer& Buffer : Buffers)
+ 		{
+ 			Transitions.Emplace(Buffer.Buffer.UAV, InPreviousState, InNewState, InFlags);
+ 		}
+
+ 		RHICmdList.Transition(Transitions);
+	}
+
+	/** Returns the total bytes allocated for all buffers in the scratch pad. */
+	FORCEINLINE uint32 AllocatedBytes()const
+	{
+		return NumAllocatedBytes;
+	}
+
+	/** Returns the number of used bytes allocated for all buffers in the scratch pad. */
+	FORCEINLINE uint32 UsedBytes()const
+	{
+		return NumUsedBytes;
+	}
+
+private:
+
+	/** Size of each FRWBuffer in the pool. */
+	uint32 MinBucketSize = 4096;
+	TIndirectArray<FBuffer> Buffers;
+	FString DebugName;
+	EBufferUsageFlags BufferUsageFlags = EBufferUsageFlags::None;
+	ERHIAccess CurrentAccess = ERHIAccess::UAVCompute;
+
+	uint32 NumAllocatedBytes = 0;
+	uint32 NumUsedBytes = 0;
+};
+
diff --git a/Engine/Plugins/FX/Niagara/Source/Niagara/Private/NiagaraSettings.cpp b/Engine/Plugins/FX/Niagara/Source/Niagara/Private/NiagaraSettings.cpp
--- a/Engine/Plugins/FX/Niagara/Source/Niagara/Private/NiagaraSettings.cpp
+++ b/Engine/Plugins/FX/Niagara/Source/Niagara/Private/NiagaraSettings.cpp
@@ -9,6 +9,9 @@ UNiagaraSettings::UNiagaraSettings(const FObjectInitializer& ObjectInitlaizer)
 	, NDISkelMesh_AdjacencyTriangleIndexFormat(ENDISkelMesh_AdjacencyTriangleIndexFormat::Full)
 {
 	PositionPinTypeColor = FLinearColor(1.0f, 0.3f, 1.0f, 1.0f);
+
+	NDICollisionQuery_AsyncGpuTraceProviderOrder.Add(ENDICollisionQuery_AsyncGpuTraceProvider::Type::HWRT);
+	NDICollisionQuery_AsyncGpuTraceProviderOrder.Add(ENDICollisionQuery_AsyncGpuTraceProvider::Type::GSDF);
 }
 
 FName UNiagaraSettings::GetCategoryName() const
diff --git a/Engine/Plugins/FX/Niagara/Source/Niagara/Public/NiagaraGpuComputeDispatchInterface.h b/Engine/Plugins/FX/Niagara/Source/Niagara/Public/NiagaraGpuComputeDispatchInterface.h
--- a/Engine/Plugins/FX/Niagara/Source/Niagara/Public/NiagaraGpuComputeDispatchInterface.h
+++ b/Engine/Plugins/FX/Niagara/Source/Niagara/Public/NiagaraGpuComputeDispatchInterface.h
@@ -7,6 +7,7 @@
 #include "NiagaraGPUInstanceCountManager.h"
 #include "FXSystem.h"
 
+class FNiagaraAsyncGpuTraceHelper;
 struct FNiagaraComputeExecutionContext;
 class FNiagaraGpuComputeDebug;
 class FNiagaraGPUInstanceCountManager;
@@ -82,15 +83,7 @@ public:
 	/** Processes all pending debug readbacks */
 	virtual void ProcessDebugReadbacks(FRHICommandListImmediate& RHICmdList, bool bWaitCompletion) = 0;
 
-#if RHI_RAYTRACING
-	virtual FNiagaraRayTracingHelper& GetRayTracingHelper() const = 0;
-	virtual bool HasRayTracingScene() const = 0;
-
-	virtual void SetPrimitiveRayTracingCollisionGroup(UPrimitiveComponent* PrimitiveComponent, uint32 Group) = 0;
-
-	virtual int32 AcquireGPURayTracedCollisionGroup() = 0;
-	virtual void ReleaseGPURayTracedCollisionGroup(int32 CollisionGroup) = 0;
-#endif
+	virtual FNiagaraAsyncGpuTraceHelper& GetAsyncGpuTraceHelper() const = 0;
 
 #if WITH_MGPU
 	/**
diff --git a/Engine/Plugins/FX/Niagara/Source/Niagara/Public/NiagaraSettings.h b/Engine/Plugins/FX/Niagara/Source/Niagara/Public/NiagaraSettings.h
--- a/Engine/Plugins/FX/Niagara/Source/Niagara/Public/NiagaraSettings.h
+++ b/Engine/Plugins/FX/Niagara/Source/Niagara/Public/NiagaraSettings.h
@@ -61,6 +61,18 @@ enum class ENiagaraDefaultRendererPixelCoverageMode : uint8
 	Disabled,
 };
 
+UENUM()
+namespace ENDICollisionQuery_AsyncGpuTraceProvider
+{
+	enum Type
+	{
+		Default = 0 UMETA(DisplayName = "Project Default"),
+		HWRT = 1 UMETA(DisplayName = "HW Ray Tracing"),
+		GSDF = 2 UMETA(DisplayName = "Global Signed Distance Fields"),
+		None = 3 UMETA(DisplayName = "Disabled"),
+	};
+}
+
 UCLASS(config = Niagara, defaultconfig, meta=(DisplayName="Niagara"))
 class NIAGARA_API UNiagaraSettings : public UDeveloperSettings
 {
@@ -134,6 +146,13 @@ class NIAGARA_API UNiagaraSettings : public UDeveloperSettings
 	UPROPERTY(config, EditAnywhere, Category = StaticMeshDI, meta = (DisplayName = "Allow Distance Fields (Experimental)", ConfigRestartRequired = true))
 	bool NDIStaticMesh_AllowDistanceFields = false;
 
+	/** 
+	Defines how traces tagged as 'Project Default' will be interpreted when using the AsyncGpuTrace data interface.
+	The system will go through (starting at element 0) to find the first provider that is available.
+	*/
+	UPROPERTY(config, EditAnywhere, Category = AsyncGpuTraceDI, meta = (DisplayName = "Trace Provider Priorities (Experimental)", ConfigRestartRequired = true))
+	TArray<TEnumAsByte<ENDICollisionQuery_AsyncGpuTraceProvider::Type>> NDICollisionQuery_AsyncGpuTraceProviderOrder;
+
 	// Begin UDeveloperSettings Interface
 	virtual FName GetCategoryName() const override;
 #if WITH_EDITOR
diff --git a/Engine/Plugins/FX/Niagara/Source/Niagara/Public/NiagaraSystemRenderData.h b/Engine/Plugins/FX/Niagara/Source/Niagara/Public/NiagaraSystemRenderData.h
--- a/Engine/Plugins/FX/Niagara/Source/Niagara/Public/NiagaraSystemRenderData.h
+++ b/Engine/Plugins/FX/Niagara/Source/Niagara/Public/NiagaraSystemRenderData.h
@@ -7,6 +7,10 @@
 #include "PrimitiveUniformShaderParameters.h"
 #include "PrimitiveViewRelevance.h"
 
+#if RHI_RAYTRACING
+#include "RayTracingInstance.h"
+#endif
+
 class FNiagaraSystemInstance;
 class FNiagaraSystemRenderData;
 class FMeshElementCollector;
diff --git a/Engine/Plugins/FX/Niagara/Source/NiagaraShader/Private/NiagaraAsyncGpuTraceProvider.cpp b/Engine/Plugins/FX/Niagara/Source/NiagaraShader/Private/NiagaraAsyncGpuTraceProvider.cpp
new file mode 100644
--- /dev/null
+++ b/Engine/Plugins/FX/Niagara/Source/NiagaraShader/Private/NiagaraAsyncGpuTraceProvider.cpp
@@ -0,0 +1,288 @@
+// Copyright Epic Games, Inc. All Rights Reserved.
+
+#include "NiagaraAsyncGpuTraceProvider.h"
+
+#include "GlobalShader.h"
+#include "NiagaraAsyncGpuTraceProviderGsdf.h"
+#include "NiagaraAsyncGpuTraceProviderHwrt.h"
+#include "NiagaraSettings.h"
+#include "NiagaraShaderParticleID.h"
+#include "Renderer/Private/ScenePrivate.h"
+
+////////////////////////////////////////////////////////////////////////////////
+////////////////////////////////////////////////////////////////////////////////
+////////////////////////////////////////////////////////////////////////////////
+
+class FNiagaraClearAsyncGpuTraceCS : public FGlobalShader
+{
+	DECLARE_GLOBAL_SHADER(FNiagaraClearAsyncGpuTraceCS);
+	SHADER_USE_PARAMETER_STRUCT(FNiagaraClearAsyncGpuTraceCS, FGlobalShader);
+
+	static constexpr uint32 kThreadGroupSizeX = 32;
+
+	static void ModifyCompilationEnvironment(const FGlobalShaderPermutationParameters& Parameters, FShaderCompilerEnvironment& OutEnvironment)
+	{
+		FGlobalShader::ModifyCompilationEnvironment(Parameters, OutEnvironment);
+		OutEnvironment.SetDefine(TEXT("THREADGROUP_SIZE_X"), kThreadGroupSizeX);
+	}
+
+	BEGIN_SHADER_PARAMETER_STRUCT(FParameters, NIAGARASHADER_API)
+		SHADER_PARAMETER_UAV(StructuredBuffer<FNiagaraAsyncGpuTraceResult>, Results)
+		SHADER_PARAMETER(uint32, TraceCount)
+		SHADER_PARAMETER(uint32, ResultsOffset)
+	END_SHADER_PARAMETER_STRUCT()
+};
+
+IMPLEMENT_GLOBAL_SHADER(FNiagaraClearAsyncGpuTraceCS, "/Plugin/FX/Niagara/Private/NiagaraAsyncGpuTraceUtils.usf", "NiagaraClearAsyncGpuTraceCS", SF_Compute);
+
+////////////////////////////////////////////////////////////////////////////////
+////////////////////////////////////////////////////////////////////////////////
+////////////////////////////////////////////////////////////////////////////////
+
+class FNiagaraUpdateCollisionGroupMapCS : public FGlobalShader
+{
+	DECLARE_GLOBAL_SHADER(FNiagaraUpdateCollisionGroupMapCS);
+	SHADER_USE_PARAMETER_STRUCT(FNiagaraUpdateCollisionGroupMapCS, FGlobalShader);
+
+	BEGIN_SHADER_PARAMETER_STRUCT(FParameters, )
+		SHADER_PARAMETER_UAV(Buffer<UINT>, RWHashTable)
+		SHADER_PARAMETER(uint32, HashTableSize)
+		SHADER_PARAMETER_UAV(Buffer<UINT>, RWHashToCollisionGroups)
+		SHADER_PARAMETER_SRV(Buffer<UINT>, NewPrimIdCollisionGroupPairs)
+		SHADER_PARAMETER(uint32, NumNewPrims)
+	END_SHADER_PARAMETER_STRUCT()
+
+public:
+	static bool ShouldCompilePermutation(const FGlobalShaderPermutationParameters& Parameters)
+	{
+		return RHISupportsComputeShaders(Parameters.Platform);
+	}
+
+	static void ModifyCompilationEnvironment(const FGlobalShaderPermutationParameters& Parameters, FShaderCompilerEnvironment& OutEnvironment)
+	{
+		FGlobalShader::ModifyCompilationEnvironment(Parameters, OutEnvironment);
+		OutEnvironment.SetDefine(TEXT("THREAD_COUNT"), THREAD_COUNT);
+	}
+
+	static constexpr uint32 THREAD_COUNT = 64;
+};
+
+IMPLEMENT_GLOBAL_SHADER(FNiagaraUpdateCollisionGroupMapCS, "/Plugin/FX/Niagara/Private/NiagaraRayTraceCollisionGroupShaders.usf", "UpdatePrimIdToCollisionGroupMap", SF_Compute);
+
+////////////////////////////////////////////////////////////////////////////////
+////////////////////////////////////////////////////////////////////////////////
+////////////////////////////////////////////////////////////////////////////////
+
+FNiagaraAsyncGpuTraceProvider::EProviderType FNiagaraAsyncGpuTraceProvider::ResolveSupportedType(EProviderType InType, const FProviderPriorityArray& Priorities)
+{
+	switch (InType)
+	{
+#if RHI_RAYTRACING
+		case ENDICollisionQuery_AsyncGpuTraceProvider::HWRT:
+			return FNiagaraAsyncGpuTraceProviderHwrt::IsSupported() ? InType : ENDICollisionQuery_AsyncGpuTraceProvider::None;
+#endif
+
+		case ENDICollisionQuery_AsyncGpuTraceProvider::GSDF:
+			return FNiagaraAsyncGpuTraceProviderGsdf::IsSupported() ? InType : ENDICollisionQuery_AsyncGpuTraceProvider::None;
+
+		case ENDICollisionQuery_AsyncGpuTraceProvider::Default:
+		{
+			ENDICollisionQuery_AsyncGpuTraceProvider::Type CurrentType = ENDICollisionQuery_AsyncGpuTraceProvider::None;
+
+			for (auto ProviderType : Priorities)
+			{
+				if (ENDICollisionQuery_AsyncGpuTraceProvider::Default == ProviderType)
+				{
+					continue;
+				}
+
+				EProviderType ResolvedType = ResolveSupportedType(ProviderType, Priorities);
+				if (ResolvedType != ENDICollisionQuery_AsyncGpuTraceProvider::None)
+				{
+					if (CurrentType == ENDICollisionQuery_AsyncGpuTraceProvider::None)
+					{
+						CurrentType = ResolvedType;
+					}
+					else if (CurrentType != ResolvedType)
+					{
+						// we've got multiple valid types so just return 'Default'
+						return ENDICollisionQuery_AsyncGpuTraceProvider::Default;
+					}
+				}
+			}
+
+			return CurrentType;
+		}
+	}
+
+	return ENDICollisionQuery_AsyncGpuTraceProvider::None;
+}
+
+template<typename T>
+bool CheckProviderIsRequestedAndSupported(FNiagaraAsyncGpuTraceProvider::EProviderType InType, const FNiagaraAsyncGpuTraceProvider::FProviderPriorityArray& Priorities)
+{
+	if ((InType == T::Type) ||
+		(InType == ENDICollisionQuery_AsyncGpuTraceProvider::Default && Priorities.Contains(T::Type)))
+	{
+		return T::IsSupported();
+	}
+
+	return false;
+}
+
+bool FNiagaraAsyncGpuTraceProvider::RequiresDistanceFieldData(EProviderType InType, const FProviderPriorityArray& Priorities)
+{
+	return CheckProviderIsRequestedAndSupported<FNiagaraAsyncGpuTraceProviderGsdf>(InType, Priorities);
+}
+
+bool FNiagaraAsyncGpuTraceProvider::RequiresRayTracingScene(EProviderType InType, const FProviderPriorityArray& Priorities)
+{
+#if RHI_RAYTRACING
+	return CheckProviderIsRequestedAndSupported<FNiagaraAsyncGpuTraceProviderHwrt>(InType, Priorities);
+#else
+	return false;
+#endif
+}
+
+TArray<TUniquePtr<FNiagaraAsyncGpuTraceProvider>> FNiagaraAsyncGpuTraceProvider::CreateSupportedProviders(EShaderPlatform ShaderPlatform, FNiagaraGpuComputeDispatchInterface* Dispatcher, const FProviderPriorityArray& Priorities)
+{
+	TArray<TUniquePtr<FNiagaraAsyncGpuTraceProvider>> Providers;
+
+	for (auto ProviderType : Priorities)
+	{
+		switch (ProviderType)
+		{
+#if RHI_RAYTRACING
+		case ENDICollisionQuery_AsyncGpuTraceProvider::HWRT:
+				if (FNiagaraAsyncGpuTraceProviderHwrt::IsSupported())
+				{
+					Providers.Emplace(MakeUnique<FNiagaraAsyncGpuTraceProviderHwrt>(ShaderPlatform, Dispatcher));
+				}
+				break;
+#endif
+
+			case ENDICollisionQuery_AsyncGpuTraceProvider::GSDF:
+				if (FNiagaraAsyncGpuTraceProviderGsdf::IsSupported())
+				{
+					Providers.Emplace(MakeUnique<FNiagaraAsyncGpuTraceProviderGsdf>(ShaderPlatform, Dispatcher));
+				}
+				break;
+		}
+	}
+
+	return Providers;
+}
+
+void FNiagaraAsyncGpuTraceProvider::ClearResults(FRHICommandList& RHICmdList, EShaderPlatform ShaderPlatform, const FDispatchRequest& Request)
+{
+	SCOPED_DRAW_EVENT(RHICmdList, NiagaraClearAsyncGpuTraceResults);
+
+	if (Request.MaxTraceCount == 0)
+	{
+		return;
+	}
+
+	TShaderMapRef<FNiagaraClearAsyncGpuTraceCS> TraceShader(GetGlobalShaderMap(ShaderPlatform));
+
+	FNiagaraClearAsyncGpuTraceCS::FParameters Params;
+	Params.Results = Request.ResultsBuffer->UAV;
+	Params.ResultsOffset = Request.ResultsOffset;
+	Params.TraceCount = Request.MaxTraceCount;
+
+	SetComputePipelineState(RHICmdList, TraceShader.GetComputeShader());
+
+	SetShaderParameters(RHICmdList, TraceShader, TraceShader.GetComputeShader(), Params);
+
+	const FIntVector ThreadGroupCount = FComputeShaderUtils::GetGroupCount(Request.MaxTraceCount, FNiagaraClearAsyncGpuTraceCS::kThreadGroupSizeX);
+	DispatchComputeShader(RHICmdList, TraceShader.GetShader(), ThreadGroupCount.X, ThreadGroupCount.Y, ThreadGroupCount.Z);
+
+	UnsetShaderUAVs(RHICmdList, TraceShader, TraceShader.GetComputeShader());
+}
+
+void FNiagaraAsyncGpuTraceProvider::BuildCollisionGroupHashMap(FRHICommandList& RHICmdList, ERHIFeatureLevel::Type FeatureLevel, FScene* Scene, const TMap<FPrimitiveComponentId, uint32>& CollisionGroupMap, FCollisionGroupHashMap& Result)
+{
+	SCOPED_DRAW_EVENT(RHICmdList, NiagaraUpdateCollisionGroupsMap);
+
+	const int32 CollisionMapEntryCount = CollisionGroupMap.Num();
+
+	if (!CollisionMapEntryCount)
+	{
+		Result.HashTableSize = 0;
+		Result.PrimIdHashTable.Release();
+		Result.HashToCollisionGroups.Release();
+		return;
+	}
+
+	uint32 MinBufferInstances = FNiagaraUpdateCollisionGroupMapCS::THREAD_COUNT;
+	uint32 NeededInstances = FMath::Max((uint32)CollisionMapEntryCount, MinBufferInstances);
+	uint32 AllocInstances = Align(NeededInstances, FNiagaraUpdateCollisionGroupMapCS::THREAD_COUNT);
+
+	//We can probably be smarter here but for now just push the whole map over to the GPU each time it's dirty.
+	//Ideally this should be a small amount of data and updated infrequently.
+	FReadBuffer NewPrimIdCollisionGroupPairs;
+	NewPrimIdCollisionGroupPairs.Initialize(TEXT("NewPrimIdCollisionGroupPairs"), sizeof(uint32) * 2, AllocInstances, EPixelFormat::PF_R32G32_UINT, BUF_Volatile);
+
+	uint32* PrimIdCollisionGroupPairPtr = (uint32*)RHILockBuffer(NewPrimIdCollisionGroupPairs.Buffer, 0, AllocInstances * sizeof(uint32) * 2, RLM_WriteOnly);
+	FMemory::Memset(PrimIdCollisionGroupPairPtr, 0, AllocInstances * sizeof(uint32) * 2);
+
+	for (auto Entry : CollisionGroupMap)
+	{
+		FPrimitiveComponentId PrimId = Entry.Key;
+		uint32 CollisionGroup = Entry.Value;
+
+		uint32 GPUSceneInstanceIndex = INDEX_NONE;
+		//Ugh this is a bit pants. Maybe try to rework things so I can use the direct prim index.
+		int32 PrimIndex = Scene->PrimitiveComponentIds.Find(PrimId);
+		if (PrimIndex != INDEX_NONE)
+		{
+			GPUSceneInstanceIndex = Scene->Primitives[PrimIndex]->GetInstanceSceneDataOffset();
+		}
+
+		PrimIdCollisionGroupPairPtr[0] = GPUSceneInstanceIndex;
+		PrimIdCollisionGroupPairPtr[1] = CollisionGroup;
+		PrimIdCollisionGroupPairPtr += 2;
+	}
+	RHIUnlockBuffer(NewPrimIdCollisionGroupPairs.Buffer);
+
+	//Init the hash table if needed
+	if (Result.HashTableSize < AllocInstances)
+	{
+		Result.HashTableSize = AllocInstances;
+
+		Result.PrimIdHashTable.Release();
+		Result.PrimIdHashTable.Initialize(TEXT("NiagaraPrimIdHashTable"), sizeof(uint32), AllocInstances, EPixelFormat::PF_R32_UINT, BUF_Static);
+
+		Result.HashToCollisionGroups.Release();
+		Result.HashToCollisionGroups.Initialize(TEXT("NiagaraPrimIdHashToCollisionGroups"), sizeof(uint32), AllocInstances, EPixelFormat::PF_R32_UINT, BUF_Static);
+	}
+
+	RHICmdList.Transition(FRHITransitionInfo(Result.PrimIdHashTable.UAV, ERHIAccess::Unknown, ERHIAccess::UAVCompute));
+	RHICmdList.Transition(FRHITransitionInfo(Result.HashToCollisionGroups.UAV, ERHIAccess::Unknown, ERHIAccess::UAVCompute));
+
+	//First we have to clear the buffers. Can probably do this better.
+	NiagaraFillGPUIntBuffer(RHICmdList, FeatureLevel, Result.PrimIdHashTable, 0);
+	RHICmdList.Transition(FRHITransitionInfo(Result.PrimIdHashTable.UAV, ERHIAccess::UAVCompute, ERHIAccess::UAVCompute));
+	NiagaraFillGPUIntBuffer(RHICmdList, FeatureLevel, Result.HashToCollisionGroups, INDEX_NONE);
+	RHICmdList.Transition(FRHITransitionInfo(Result.HashToCollisionGroups.UAV, ERHIAccess::UAVCompute, ERHIAccess::UAVCompute));
+
+	TShaderMapRef<FNiagaraUpdateCollisionGroupMapCS> ComputeShader(GetGlobalShaderMap(GMaxRHIFeatureLevel));
+	FRHIComputeShader* ShaderRHI = ComputeShader.GetComputeShader();
+
+	FNiagaraUpdateCollisionGroupMapCS::FParameters Params;
+	Params.RWHashTable = Result.PrimIdHashTable.UAV;
+	Params.HashTableSize = Result.HashTableSize;
+	Params.RWHashToCollisionGroups = Result.HashToCollisionGroups.UAV;
+	Params.NewPrimIdCollisionGroupPairs = NewPrimIdCollisionGroupPairs.SRV;
+	Params.NumNewPrims = CollisionGroupMap.Num();
+
+	// To simplify the shader code, the size of the table must be a multiple of the thread count.
+	check(Result.HashTableSize % FNiagaraUpdateCollisionGroupMapCS::THREAD_COUNT == 0);
+
+	SetComputePipelineState(RHICmdList, ShaderRHI);
+	SetShaderParameters(RHICmdList, ComputeShader, ShaderRHI, Params);
+	RHICmdList.DispatchComputeShader(FMath::DivideAndRoundUp(CollisionGroupMap.Num(), (int32)FNiagaraUpdateCollisionGroupMapCS::THREAD_COUNT), 1, 1);
+	UnsetShaderUAVs(RHICmdList, ComputeShader, ShaderRHI);
+
+	RHICmdList.Transition(FRHITransitionInfo(Result.PrimIdHashTable.UAV, ERHIAccess::UAVCompute, ERHIAccess::SRVCompute));
+	RHICmdList.Transition(FRHITransitionInfo(Result.HashToCollisionGroups.UAV, ERHIAccess::UAVCompute, ERHIAccess::SRVCompute));
+}
diff --git a/Engine/Plugins/FX/Niagara/Source/NiagaraShader/Private/NiagaraAsyncGpuTraceProviderGsdf.cpp b/Engine/Plugins/FX/Niagara/Source/NiagaraShader/Private/NiagaraAsyncGpuTraceProviderGsdf.cpp
new file mode 100644
--- /dev/null
+++ b/Engine/Plugins/FX/Niagara/Source/NiagaraShader/Private/NiagaraAsyncGpuTraceProviderGsdf.cpp
@@ -0,0 +1,161 @@
+// Copyright Epic Games, Inc. All Rights Reserved.
+
+#include "NiagaraAsyncGpuTraceProviderGsdf.h"
+
+#include "GlobalShader.h"
+#include "NiagaraSettings.h"
+#include "Renderer/Private/ScenePrivate.h"
+#include "Renderer/Private/SceneRendering.h"
+#include "SceneManagement.h"
+
+static int GNiagaraAsyncGpuTraceGsdfEnabled = 1;
+static FAutoConsoleVariableRef CVarNiagaraAsyncGpuTraceGsdfEnabled(
+	TEXT("fx.Niagara.AsyncGpuTrace.GlobalSdfEnabled"),
+	GNiagaraAsyncGpuTraceGsdfEnabled,
+	TEXT("If disabled AsyncGpuTrace will not be supported against Global SDF."),
+	ECVF_Default
+);
+
+class FNiagaraRayMarchGlobalSdfCS : public FGlobalShader
+{
+	DECLARE_GLOBAL_SHADER(FNiagaraRayMarchGlobalSdfCS);
+	SHADER_USE_PARAMETER_STRUCT(FNiagaraRayMarchGlobalSdfCS, FGlobalShader);
+
+	static constexpr uint32 kMaxMarchCount = 16;
+	static constexpr uint32 kThreadGroupSizeX = 32;
+
+	static bool ShouldCompilePermutation(const FGlobalShaderPermutationParameters& Parameters)
+	{
+		if (!FGlobalShader::ShouldCompilePermutation(Parameters))
+		{
+			return false;
+		}
+
+		if (!DoesPlatformSupportDistanceFields(Parameters.Platform))
+		{
+			return false;
+		}
+
+		return true;
+	}
+
+	static void ModifyCompilationEnvironment(const FGlobalShaderPermutationParameters& Parameters, FShaderCompilerEnvironment& OutEnvironment)
+	{
+		FGlobalShader::ModifyCompilationEnvironment(Parameters, OutEnvironment);
+		OutEnvironment.SetDefine(TEXT("MAX_MARCH_COUNT"), kMaxMarchCount);
+		OutEnvironment.SetDefine(TEXT("THREADGROUP_SIZE_X"), kThreadGroupSizeX);
+	}
+
+	BEGIN_SHADER_PARAMETER_STRUCT(FParameters, NIAGARASHADER_API)
+		SHADER_PARAMETER_SRV(StructuredBuffer<FNiagaraAsyncGpuTrace>, Traces)
+		SHADER_PARAMETER_SRV(Buffer<uint>, TraceCounts)
+		SHADER_PARAMETER_UAV(StructuredBuffer<FNiagaraAsyncGpuTraceResult>, Results)
+
+		SHADER_PARAMETER_STRUCT_REF(FViewUniformShaderParameters, View)
+		SHADER_PARAMETER_STRUCT_INCLUDE(FGlobalDistanceFieldParameters2, GlobalDistanceFieldParameters)
+		SHADER_PARAMETER(uint32, TracesOffset)
+		SHADER_PARAMETER(uint32, TraceCountsOffset)
+		SHADER_PARAMETER(uint32, ResultsOffset)
+		SHADER_PARAMETER(float, DistanceThreshold)
+	END_SHADER_PARAMETER_STRUCT()
+};
+
+IMPLEMENT_GLOBAL_SHADER(FNiagaraRayMarchGlobalSdfCS, "/Plugin/FX/Niagara/Private/NiagaraRayMarchingShaders.usf", "NiagaraRayMarchGlobalSdfCS", SF_Compute);
+
+////////////////////////////////////////////////////////////////////////////////
+////////////////////////////////////////////////////////////////////////////////
+////////////////////////////////////////////////////////////////////////////////
+
+// todo - currently duplicated from SetupGlobalDistanceFieldParameters (GlobalDistanceField.cpp) because of
+// problems getting it properly exported from the dll
+static void AssignGlobalDistanceFieldParameters(const FGlobalDistanceFieldParameterData& ParameterData, FGlobalDistanceFieldParameters2& ShaderParameters)
+{
+	ShaderParameters.GlobalDistanceFieldPageAtlasTexture = OrBlack3DIfNull(ParameterData.PageAtlasTexture);
+	ShaderParameters.GlobalDistanceFieldPageTableTexture = OrBlack3DUintIfNull(ParameterData.PageTableTexture);
+	ShaderParameters.GlobalDistanceFieldMipTexture = OrBlack3DIfNull(ParameterData.MipTexture);
+
+	for (int32 Index = 0; Index < GMaxGlobalDistanceFieldClipmaps; Index++)
+	{
+		ShaderParameters.GlobalVolumeCenterAndExtent[Index] = ParameterData.CenterAndExtent[Index];
+		ShaderParameters.GlobalVolumeWorldToUVAddAndMul[Index] = ParameterData.WorldToUVAddAndMul[Index];
+		ShaderParameters.GlobalDistanceFieldMipWorldToUVScale[Index] = ParameterData.MipWorldToUVScale[Index];
+		ShaderParameters.GlobalDistanceFieldMipWorldToUVBias[Index] = ParameterData.MipWorldToUVBias[Index];
+	}
+
+	ShaderParameters.GlobalDistanceFieldMipFactor = ParameterData.MipFactor;
+	ShaderParameters.GlobalDistanceFieldMipTransition = ParameterData.MipTransition;
+	ShaderParameters.GlobalDistanceFieldClipmapSizeInPages = ParameterData.ClipmapSizeInPages;
+	ShaderParameters.GlobalDistanceFieldInvPageAtlasSize = (FVector3f)ParameterData.InvPageAtlasSize;
+	ShaderParameters.GlobalVolumeDimension = ParameterData.GlobalDFResolution;
+	ShaderParameters.GlobalVolumeTexelSize = 1.0f / ParameterData.GlobalDFResolution;
+	ShaderParameters.MaxGlobalDFAOConeDistance = ParameterData.MaxDFAOConeDistance;
+	ShaderParameters.NumGlobalSDFClipmaps = ParameterData.NumGlobalSDFClipmaps;
+}
+
+////////////////////////////////////////////////////////////////////////////////
+////////////////////////////////////////////////////////////////////////////////
+////////////////////////////////////////////////////////////////////////////////
+
+const FNiagaraAsyncGpuTraceProvider::EProviderType FNiagaraAsyncGpuTraceProviderGsdf::Type = ENDICollisionQuery_AsyncGpuTraceProvider::GSDF;
+
+FNiagaraAsyncGpuTraceProviderGsdf::FNiagaraAsyncGpuTraceProviderGsdf(EShaderPlatform InShaderPlatform, FNiagaraGpuComputeDispatchInterface* Dispatcher)
+	: FNiagaraAsyncGpuTraceProvider(InShaderPlatform, Dispatcher)
+{
+}
+
+bool FNiagaraAsyncGpuTraceProviderGsdf::IsSupported()
+{
+	return GNiagaraAsyncGpuTraceGsdfEnabled && DoesProjectSupportDistanceFields();
+}
+
+bool FNiagaraAsyncGpuTraceProviderGsdf::IsAvailable() const
+{
+	return !!GNiagaraAsyncGpuTraceGsdfEnabled;
+}
+
+void FNiagaraAsyncGpuTraceProviderGsdf::PostRenderOpaque(FRHICommandList& RHICmdList, TConstArrayView<FViewInfo> Views, FCollisionGroupHashMap* CollisionGroupHash)
+{
+	const FViewInfo& ReferenceView = Views[0];
+	m_DistanceFieldData = ReferenceView.GlobalDistanceFieldInfo.ParameterData;
+	m_ViewUniformBuffer = ReferenceView.ViewUniformBuffer;
+}
+
+void FNiagaraAsyncGpuTraceProviderGsdf::IssueTraces(FRHICommandList& RHICmdList, const FDispatchRequest& Request, FCollisionGroupHashMap* CollisionHashMap)
+{
+	check(IsAvailable());
+
+	if (Request.MaxTraceCount == 0)
+	{
+		return;
+	}
+
+	SCOPED_DRAW_EVENT(RHICmdList, NiagaraIssueTracesGsdf);
+
+	TShaderMapRef<FNiagaraRayMarchGlobalSdfCS> TraceShader(GetGlobalShaderMap(ShaderPlatform));
+
+	FNiagaraRayMarchGlobalSdfCS::FParameters Params;
+	Params.View = GetShaderBinding(m_ViewUniformBuffer);
+	Params.Traces = Request.TracesBuffer->SRV;
+	Params.TraceCounts = Request.TraceCountsBuffer->SRV;
+	Params.Results = Request.ResultsBuffer->UAV;
+	AssignGlobalDistanceFieldParameters(m_DistanceFieldData, Params.GlobalDistanceFieldParameters);
+	Params.TracesOffset = Request.TracesOffset;
+	Params.TraceCountsOffset = Request.TraceCountsOffset;
+	Params.ResultsOffset = Request.ResultsOffset;
+	Params.DistanceThreshold = 1.0f;
+
+	SetComputePipelineState(RHICmdList, TraceShader.GetComputeShader());
+
+	SetShaderParameters(RHICmdList, TraceShader, TraceShader.GetComputeShader(), Params);
+
+	const FIntVector ThreadGroupCount = FComputeShaderUtils::GetGroupCount(Request.MaxTraceCount, FNiagaraRayMarchGlobalSdfCS::kThreadGroupSizeX);
+	DispatchComputeShader(RHICmdList, TraceShader.GetShader(), ThreadGroupCount.X, ThreadGroupCount.Y, ThreadGroupCount.Z);
+
+	UnsetShaderUAVs(RHICmdList, TraceShader, TraceShader.GetComputeShader());
+}
+
+void FNiagaraAsyncGpuTraceProviderGsdf::Reset()
+{
+	m_DistanceFieldData = FGlobalDistanceFieldParameterData();
+	m_ViewUniformBuffer = TUniformBufferRef<FViewUniformShaderParameters>();
+}
\ No newline at end of file
diff --git a/Engine/Plugins/FX/Niagara/Source/NiagaraShader/Private/NiagaraAsyncGpuTraceProviderGsdf.h b/Engine/Plugins/FX/Niagara/Source/NiagaraShader/Private/NiagaraAsyncGpuTraceProviderGsdf.h
new file mode 100644
--- /dev/null
+++ b/Engine/Plugins/FX/Niagara/Source/NiagaraShader/Private/NiagaraAsyncGpuTraceProviderGsdf.h
@@ -0,0 +1,27 @@
+// Copyright Epic Games, Inc. All Rights Reserved.
+
+#pragma once
+
+#include "CoreMinimal.h"
+
+#include "GlobalDistanceFieldParameters.h"
+#include "NiagaraAsyncGpuTraceProvider.h"
+
+class FNiagaraAsyncGpuTraceProviderGsdf : public FNiagaraAsyncGpuTraceProvider
+{
+public:
+	static const EProviderType Type;
+
+	FNiagaraAsyncGpuTraceProviderGsdf(EShaderPlatform InShaderPlatform, FNiagaraGpuComputeDispatchInterface* Dispatcher);
+
+	static bool IsSupported();
+	virtual void PostRenderOpaque(FRHICommandList& RHICmdList, TConstArrayView<FViewInfo> Views, FCollisionGroupHashMap* CollisionGroupHash) override;
+	virtual bool IsAvailable() const override;
+	virtual void IssueTraces(FRHICommandList& RHICmdList, const FDispatchRequest& Request, FCollisionGroupHashMap* CollisionGroupHash) override;
+	virtual void Reset() override;
+	virtual EProviderType GetType() const override { return Type; }
+
+private:
+	FGlobalDistanceFieldParameterData m_DistanceFieldData;
+	TUniformBufferRef<FViewUniformShaderParameters> m_ViewUniformBuffer;
+};
\ No newline at end of file
diff --git a/Engine/Plugins/FX/Niagara/Source/NiagaraShader/Private/NiagaraAsyncGpuTraceProviderHwrt.cpp b/Engine/Plugins/FX/Niagara/Source/NiagaraShader/Private/NiagaraAsyncGpuTraceProviderHwrt.cpp
new file mode 100644
--- /dev/null
+++ b/Engine/Plugins/FX/Niagara/Source/NiagaraShader/Private/NiagaraAsyncGpuTraceProviderHwrt.cpp
@@ -0,0 +1,410 @@
+// Copyright Epic Games, Inc. All Rights Reserved.
+
+#include "NiagaraAsyncGpuTraceProviderHwrt.h"
+#include "NiagaraGpuComputeDispatchInterface.h"
+
+#if RHI_RAYTRACING
+
+#include "GlobalShader.h"
+#include "NiagaraRenderer.h"
+#include "NiagaraSettings.h"
+#include "Renderer/Private/ScenePrivate.h"
+#include "Renderer/Private/SceneRendering.h"
+
+static int GNiagaraAsyncGpuTraceHwrtEnabled = 1;
+static FAutoConsoleVariableRef CVarNiagaraAsyncGpuTraceHwrtEnabled(
+	TEXT("fx.Niagara.AsyncGpuTrace.HWRayTraceEnabled"),
+	GNiagaraAsyncGpuTraceHwrtEnabled,
+	TEXT("If disabled AsyncGpuTrace will not be supported against the HW ray tracing scene."),
+	ECVF_Default
+);
+
+/// TODO
+///  -get geometry masking working when an environmental mask is implemented
+
+// c++ mirror of the struct defined in RayTracingCommon.ush
+struct FVFXTracePayload
+{
+	float HitT;
+	uint32 PrimitiveIndex;
+	uint32 InstanceIndex;
+	float Barycentrics[2];
+	float WorldPosition[3];
+	float WorldNormal[3];
+};
+
+BEGIN_SHADER_PARAMETER_STRUCT(FGPUSceneParameters, )
+	SHADER_PARAMETER_SRV(StructuredBuffer<float4>, GPUSceneInstanceSceneData)
+	SHADER_PARAMETER_SRV(StructuredBuffer<float4>, GPUSceneInstancePayloadData)
+	SHADER_PARAMETER_SRV(StructuredBuffer<float4>, GPUScenePrimitiveSceneData)
+	SHADER_PARAMETER(uint32, GPUSceneFrameNumber)
+END_SHADER_PARAMETER_STRUCT()
+
+class FNiagaraCollisionRayTraceRG : public FGlobalShader
+{
+	DECLARE_GLOBAL_SHADER(FNiagaraCollisionRayTraceRG)
+	SHADER_USE_ROOT_PARAMETER_STRUCT(FNiagaraCollisionRayTraceRG, FGlobalShader)
+
+	BEGIN_SHADER_PARAMETER_STRUCT(FParameters, )
+		SHADER_PARAMETER_STRUCT_REF(FViewUniformShaderParameters, View)
+		SHADER_PARAMETER_STRUCT_INCLUDE(FGPUSceneParameters, GPUSceneParameters)
+
+		SHADER_PARAMETER_SRV(Buffer<UINT>, HashTable)
+		SHADER_PARAMETER_SRV(Buffer<UINT>, HashToCollisionGroups)
+		SHADER_PARAMETER(uint32, HashTableSize)
+
+		SHADER_PARAMETER_SRV(RaytracingAccelerationStructure, TLAS)
+		SHADER_PARAMETER_SRV(Buffer<FNiagaraAsyncGpuTrace>, Rays)
+		SHADER_PARAMETER(uint32, RaysOffset)
+		SHADER_PARAMETER_UAV(Buffer<FNiagaraAsyncGpuTraceResult>, CollisionOutput)
+		SHADER_PARAMETER(uint32, CollisionOutputOffset)
+		SHADER_PARAMETER_SRV(Buffer<UINT>, RayTraceCounts)
+		SHADER_PARAMETER(uint32, MaxRetraces)
+	END_SHADER_PARAMETER_STRUCT()
+
+	class FFakeIndirectDispatch : SHADER_PERMUTATION_BOOL("NIAGARA_RAYTRACE_FAKE_INDIRECT");
+	class FSupportsCollisionGroups : SHADER_PERMUTATION_BOOL("NIAGARA_SUPPORTS_COLLISION_GROUPS");
+	using FPermutationDomain = TShaderPermutationDomain<FFakeIndirectDispatch, FSupportsCollisionGroups>;
+
+	static bool ShouldCompilePermutation(const FGlobalShaderPermutationParameters& Parameters);
+	static void ModifyCompilationEnvironment(const FGlobalShaderPermutationParameters& Parameters, FShaderCompilerEnvironment& OutEnvironment);
+	static TShaderRef< FNiagaraCollisionRayTraceRG> GetShader(FGlobalShaderMap* ShaderMap, bool SupportsCollisionGroups);
+	static FRHIRayTracingShader* GetRayTracingShader(FGlobalShaderMap* ShaderMap, bool SupportCollisionGroups);
+	static bool SupportsIndirectDispatch();
+};
+
+class FNiagaraCollisionRayTraceCH : public FGlobalShader
+{
+	DECLARE_GLOBAL_SHADER(FNiagaraCollisionRayTraceCH);
+
+	static bool ShouldCompilePermutation(const FGlobalShaderPermutationParameters& Parameters);
+	static void ModifyCompilationEnvironment(const FGlobalShaderPermutationParameters& Parameters, FShaderCompilerEnvironment& OutEnvironment);
+
+	FNiagaraCollisionRayTraceCH() = default;
+	FNiagaraCollisionRayTraceCH(const ShaderMetaType::CompiledShaderInitializerType& Initializer);
+};
+
+class FNiagaraCollisionRayTraceMiss : public FGlobalShader
+{
+	DECLARE_GLOBAL_SHADER(FNiagaraCollisionRayTraceMiss);
+
+	static bool ShouldCompilePermutation(const FGlobalShaderPermutationParameters& Parameters);
+	static void ModifyCompilationEnvironment(const FGlobalShaderPermutationParameters& Parameters, FShaderCompilerEnvironment& OutEnvironment);
+
+	FNiagaraCollisionRayTraceMiss() = default;
+	FNiagaraCollisionRayTraceMiss(const ShaderMetaType::CompiledShaderInitializerType& Initializer);
+};
+
+bool FNiagaraCollisionRayTraceRG::ShouldCompilePermutation(const FGlobalShaderPermutationParameters& Parameters)
+{
+	return ShouldCompileRayTracingShadersForProject(Parameters.Platform);
+}
+
+void FNiagaraCollisionRayTraceRG::ModifyCompilationEnvironment(const FGlobalShaderPermutationParameters& Parameters, FShaderCompilerEnvironment& OutEnvironment)
+{
+	FGlobalShader::ModifyCompilationEnvironment(Parameters, OutEnvironment);
+	OutEnvironment.SetDefine(TEXT("NIAGARA_SUPPORTS_RAY_TRACING"), 1);
+	OutEnvironment.SetDefine(TEXT("VF_SUPPORTS_PRIMITIVE_SCENE_DATA"), 1);
+	OutEnvironment.SetDefine(TEXT("USE_GLOBAL_GPU_SCENE_DATA"), 1);
+}
+
+TShaderRef<FNiagaraCollisionRayTraceRG> FNiagaraCollisionRayTraceRG::GetShader(FGlobalShaderMap* ShaderMap, bool SupportsCollisionGroups)
+{
+	FPermutationDomain PermutationVector;
+	PermutationVector.Set<FFakeIndirectDispatch>(!SupportsIndirectDispatch());
+	PermutationVector.Set<FSupportsCollisionGroups>(SupportsCollisionGroups);
+
+	return ShaderMap->GetShader<FNiagaraCollisionRayTraceRG>(PermutationVector);
+}
+
+FRHIRayTracingShader* FNiagaraCollisionRayTraceRG::GetRayTracingShader(FGlobalShaderMap* ShaderMap, bool SupportCollisionGroups)
+{
+	return GetShader(ShaderMap, SupportCollisionGroups).GetRayTracingShader();
+}
+
+bool FNiagaraCollisionRayTraceRG::SupportsIndirectDispatch()
+{
+	return GRHISupportsRayTracingDispatchIndirect;
+}
+
+bool FNiagaraCollisionRayTraceCH::ShouldCompilePermutation(const FGlobalShaderPermutationParameters& Parameters)
+{
+	return ShouldCompileRayTracingShadersForProject(Parameters.Platform);
+}
+
+void FNiagaraCollisionRayTraceCH::ModifyCompilationEnvironment(const FGlobalShaderPermutationParameters& Parameters, FShaderCompilerEnvironment& OutEnvironment)
+{
+	FGlobalShader::ModifyCompilationEnvironment(Parameters, OutEnvironment);
+	OutEnvironment.SetDefine(TEXT("NIAGARA_SUPPORTS_RAY_TRACING"), 1);
+}
+
+FNiagaraCollisionRayTraceCH::FNiagaraCollisionRayTraceCH(const ShaderMetaType::CompiledShaderInitializerType& Initializer)
+	: FGlobalShader(Initializer)
+{}
+
+bool FNiagaraCollisionRayTraceMiss::ShouldCompilePermutation(const FGlobalShaderPermutationParameters& Parameters)
+{
+	return ShouldCompileRayTracingShadersForProject(Parameters.Platform);
+}
+
+void FNiagaraCollisionRayTraceMiss::ModifyCompilationEnvironment(const FGlobalShaderPermutationParameters& Parameters, FShaderCompilerEnvironment& OutEnvironment)
+{
+	FGlobalShader::ModifyCompilationEnvironment(Parameters, OutEnvironment);
+	OutEnvironment.SetDefine(TEXT("NIAGARA_SUPPORTS_RAY_TRACING"), 1);
+}
+
+FNiagaraCollisionRayTraceMiss::FNiagaraCollisionRayTraceMiss(const ShaderMetaType::CompiledShaderInitializerType& Initializer)
+	: FGlobalShader(Initializer)
+{}
+
+IMPLEMENT_GLOBAL_SHADER(FNiagaraCollisionRayTraceRG, "/Plugin/FX/Niagara/Private/NiagaraRayTracingShaders.usf", "NiagaraCollisionRayTraceRG", SF_RayGen);
+IMPLEMENT_GLOBAL_SHADER(FNiagaraCollisionRayTraceCH, "/Plugin/FX/Niagara/Private/NiagaraRayTracingShaders.usf", "NiagaraCollisionRayTraceCH", SF_RayHitGroup);
+IMPLEMENT_GLOBAL_SHADER(FNiagaraCollisionRayTraceMiss, "/Plugin/FX/Niagara/Private/NiagaraRayTracingShaders.usf", "NiagaraCollisionRayTraceMiss", SF_RayMiss);
+
+////////////////////////////////////////////////////////////////////////////////
+////////////////////////////////////////////////////////////////////////////////
+////////////////////////////////////////////////////////////////////////////////
+
+static FRayTracingPipelineState* CreateNiagaraRayTracingPipelineState(
+	EShaderPlatform Platform,
+	FRHICommandList& RHICmdList,
+	FRHIRayTracingShader* RayGenShader,
+	FRHIRayTracingShader* ClosestHitShader,
+	FRHIRayTracingShader* MissShader)
+{
+	FRayTracingPipelineStateInitializer Initializer;
+	Initializer.MaxPayloadSizeInBytes = sizeof(FVFXTracePayload);
+
+	FRHIRayTracingShader* RayGenShaderTable[] = { RayGenShader };
+	Initializer.SetRayGenShaderTable(RayGenShaderTable);
+
+	FRHIRayTracingShader* HitGroupTable[] = { ClosestHitShader };
+	Initializer.SetHitGroupTable(HitGroupTable);
+
+	FRHIRayTracingShader* MissTable[] = { MissShader };
+	Initializer.SetMissShaderTable(MissTable);
+
+	Initializer.bAllowHitGroupIndexing = true; // Use the same hit shader for all geometry in the scene by disabling SBT indexing.
+
+	return PipelineStateCache::GetAndOrCreateRayTracingPipelineState(RHICmdList, Initializer);
+}
+
+static void BindNiagaraRayTracingMeshCommands(
+	FRHICommandList& RHICmdList,
+	FRayTracingSceneRHIRef RayTracingScene,
+	FRHIUniformBuffer* ViewUniformBuffer,
+	TConstArrayView<FVisibleRayTracingMeshCommand> RayTracingMeshCommands,
+	FRayTracingPipelineState* Pipeline,
+	TFunctionRef<uint32(const FRayTracingMeshCommand&)> PackUserData)
+{
+	const int32 NumTotalBindings = RayTracingMeshCommands.Num();
+
+	const uint32 MergedBindingsSize = sizeof(FRayTracingLocalShaderBindings) * NumTotalBindings;
+
+	const uint32 NumUniformBuffers = 1;
+	const uint32 UniformBufferArraySize = sizeof(FRHIUniformBuffer*) * NumUniformBuffers;
+
+	FRayTracingLocalShaderBindings* Bindings = nullptr;
+	FRHIUniformBuffer** UniformBufferArray = nullptr;
+
+	if (RHICmdList.Bypass())
+	{
+		FMemStack& MemStack = FMemStack::Get();
+
+		Bindings = reinterpret_cast<FRayTracingLocalShaderBindings*>(MemStack.Alloc(MergedBindingsSize, alignof(FRayTracingLocalShaderBindings)));
+		UniformBufferArray = reinterpret_cast<FRHIUniformBuffer**>(MemStack.Alloc(UniformBufferArraySize, alignof(FRHIUniformBuffer*)));
+	}
+	else
+	{
+		Bindings = reinterpret_cast<FRayTracingLocalShaderBindings*>(RHICmdList.Alloc(MergedBindingsSize, alignof(FRayTracingLocalShaderBindings)));
+		UniformBufferArray = reinterpret_cast<FRHIUniformBuffer**>(RHICmdList.Alloc(UniformBufferArraySize, alignof(FRHIUniformBuffer*)));
+	}
+
+	UniformBufferArray[0] = ViewUniformBuffer;
+
+	uint32 BindingIndex = 0;
+	for (const FVisibleRayTracingMeshCommand VisibleMeshCommand : RayTracingMeshCommands)
+	{
+		const FRayTracingMeshCommand& MeshCommand = *VisibleMeshCommand.RayTracingMeshCommand;
+
+		FRayTracingLocalShaderBindings Binding = {};
+		Binding.InstanceIndex = VisibleMeshCommand.InstanceIndex;
+		Binding.SegmentIndex = MeshCommand.GeometrySegmentIndex;
+		Binding.UserData = PackUserData(MeshCommand);
+		Binding.UniformBuffers = UniformBufferArray;
+		Binding.NumUniformBuffers = NumUniformBuffers;
+
+		Bindings[BindingIndex] = Binding;
+		BindingIndex++;
+	}
+
+	const bool bCopyDataToInlineStorage = false; // Storage is already allocated from RHICmdList, no extra copy necessary
+	RHICmdList.SetRayTracingHitGroups(
+		RayTracingScene,
+		Pipeline,
+		NumTotalBindings,
+		Bindings,
+		bCopyDataToInlineStorage);
+}
+
+////////////////////////////////////////////////////////////////////////////////
+////////////////////////////////////////////////////////////////////////////////
+////////////////////////////////////////////////////////////////////////////////
+
+const FNiagaraAsyncGpuTraceProvider::EProviderType FNiagaraAsyncGpuTraceProviderHwrt::Type = ENDICollisionQuery_AsyncGpuTraceProvider::HWRT;
+
+FNiagaraAsyncGpuTraceProviderHwrt::FNiagaraAsyncGpuTraceProviderHwrt(EShaderPlatform InShaderPlatform, FNiagaraGpuComputeDispatchInterface* Dispatcher)
+	: FNiagaraAsyncGpuTraceProvider(InShaderPlatform, Dispatcher)
+{
+}
+
+bool FNiagaraAsyncGpuTraceProviderHwrt::IsSupported()
+{
+	return GNiagaraAsyncGpuTraceHwrtEnabled && IsRayTracingEnabled();
+}
+
+bool FNiagaraAsyncGpuTraceProviderHwrt::IsAvailable() const
+{
+	if (!GNiagaraAsyncGpuTraceHwrtEnabled)
+	{
+		return false;
+	}
+
+	if (!Dispatcher->RequiresRayTracingScene())
+	{
+		return false;
+	}
+
+	FScene* Scene = Dispatcher->GetScene();
+
+	return Scene ? Scene->RayTracingScene.IsCreated() : false;
+}
+
+void FNiagaraAsyncGpuTraceProviderHwrt::PostRenderOpaque(FRHICommandList& RHICmdList, TConstArrayView<FViewInfo> Views, FCollisionGroupHashMap* CollisionGroupHash)
+{
+	check(IsAvailable());
+	check(Views.Num() > 0);
+
+	const FViewInfo& ReferenceView = Views[0];
+	FScene* Scene = Dispatcher->GetScene();
+
+	if (Scene && Scene->RayTracingScene.IsCreated())
+	{
+		RayTracingScene = Scene->RayTracingScene.GetRHIRayTracingSceneChecked();
+		RayTracingSceneView = Scene->RayTracingScene.GetShaderResourceViewChecked();
+		ViewUniformBuffer = ReferenceView.ViewUniformBuffer;
+
+		FGlobalShaderMap* ShaderMap = GetGlobalShaderMap(ShaderPlatform);
+		auto RayGenShader = FNiagaraCollisionRayTraceRG::GetRayTracingShader(ShaderMap, CollisionGroupHash != nullptr);
+		auto ClosestHitShader = ShaderMap->GetShader<FNiagaraCollisionRayTraceCH>().GetRayTracingShader();
+		auto MissShader = ShaderMap->GetShader<FNiagaraCollisionRayTraceMiss>().GetRayTracingShader();
+
+		RayTracingPipelineState = CreateNiagaraRayTracingPipelineState(
+			ShaderPlatform,
+			RHICmdList,
+			RayGenShader,
+			ClosestHitShader,
+			MissShader);
+
+		// some options for what we want with our per MeshCommand user data.  For now we'll ignore it, but possibly
+		// something we'd want to incorporate.  Some examples could be if the material is translucent, or possibly the physical material?
+		auto BakeTranslucent = [&](const FRayTracingMeshCommand& MeshCommand) {	return (MeshCommand.bIsTranslucent != 0) & 0x1;	};
+		auto BakeDefault = [&](const FRayTracingMeshCommand& MeshCommand) { return 0; };
+
+		BindNiagaraRayTracingMeshCommands(
+			RHICmdList,
+			RayTracingScene,
+			ViewUniformBuffer,
+			ReferenceView.VisibleRayTracingMeshCommands,
+			RayTracingPipelineState,
+			BakeDefault);
+	}
+	else
+	{
+		Reset();
+	}
+}
+
+void FNiagaraAsyncGpuTraceProviderHwrt::IssueTraces(FRHICommandList& RHICmdList, const FDispatchRequest& Request, FCollisionGroupHashMap* CollisionGroupHash)
+{
+	check(IsAvailable());
+	check(RayTracingPipelineState);
+	check(RayTracingScene);
+	check(RayTracingSceneView);
+
+	if (Request.MaxTraceCount == 0)
+	{
+		return;
+	}
+
+	SCOPED_DRAW_EVENT(RHICmdList, NiagaraIssueTracesHwrt);
+
+	FScene* Scene = Dispatcher->GetScene();
+	FGlobalShaderMap* ShaderMap = GetGlobalShaderMap(ShaderPlatform);
+
+	TShaderRef<FNiagaraCollisionRayTraceRG> RGShader = FNiagaraCollisionRayTraceRG::GetShader(ShaderMap, CollisionGroupHash != nullptr);
+
+	FNiagaraCollisionRayTraceRG::FParameters Params;
+
+	Params.View = GetShaderBinding(ViewUniformBuffer);
+	Params.GPUSceneParameters.GPUSceneInstanceSceneData = Scene->GPUScene.InstanceSceneDataBuffer.SRV;
+	Params.GPUSceneParameters.GPUSceneInstancePayloadData = Scene->GPUScene.InstancePayloadDataBuffer.SRV;
+	Params.GPUSceneParameters.GPUScenePrimitiveSceneData = Scene->GPUScene.PrimitiveBuffer.SRV;
+	Params.GPUSceneParameters.GPUSceneFrameNumber = Scene->GPUScene.GetSceneFrameNumber();
+
+	if (CollisionGroupHash)
+	{
+		Params.HashTable = CollisionGroupHash->PrimIdHashTable.SRV;
+		Params.HashTableSize = CollisionGroupHash->HashTableSize;
+		Params.HashToCollisionGroups = CollisionGroupHash->HashToCollisionGroups.SRV;
+	}
+
+	Params.TLAS = RayTracingSceneView;
+	Params.Rays = Request.TracesBuffer->SRV;
+	Params.RaysOffset = Request.TracesOffset;
+	Params.CollisionOutput = Request.ResultsBuffer->UAV;
+	Params.CollisionOutputOffset = Request.ResultsOffset;
+	Params.MaxRetraces = Request.MaxRetraceCount;
+
+	if (FNiagaraCollisionRayTraceRG::SupportsIndirectDispatch())
+	{
+		FRayTracingShaderBindingsWriter GlobalResources;
+		SetShaderParameters(GlobalResources, RGShader, Params);
+
+		//Can we wrangle things so we can have one indirect dispatch with each internal dispatch pointing to potentially different Ray and Results buffers?
+		//For now have a each as a unique dispatch.
+		RHICmdList.RayTraceDispatchIndirect(
+			RayTracingPipelineState,
+			RGShader.GetRayTracingShader(),
+			RayTracingScene,
+			GlobalResources,
+			Request.TraceCountsBuffer->Buffer,
+			Request.TraceCountsOffset * sizeof(uint32));
+	}
+	else
+	{
+		Params.RayTraceCounts = Request.TraceCountsBuffer->SRV;
+
+		FRayTracingShaderBindingsWriter GlobalResources;
+		SetShaderParameters(GlobalResources, RGShader, Params);
+
+		RHICmdList.RayTraceDispatch(
+			RayTracingPipelineState,
+			RGShader.GetRayTracingShader(),
+			RayTracingScene,
+			GlobalResources,
+			Request.MaxTraceCount,
+			1
+		);
+	}
+}
+
+void FNiagaraAsyncGpuTraceProviderHwrt::Reset()
+{
+	RayTracingPipelineState = nullptr;
+	RayTracingScene = nullptr;
+	RayTracingSceneView = nullptr;
+	ViewUniformBuffer = TUniformBufferRef<FViewUniformShaderParameters>();
+}
+
+#endif // RHI_RAYTRACING
\ No newline at end of file
diff --git a/Engine/Plugins/FX/Niagara/Source/NiagaraShader/Private/NiagaraAsyncGpuTraceProviderHwrt.h b/Engine/Plugins/FX/Niagara/Source/NiagaraShader/Private/NiagaraAsyncGpuTraceProviderHwrt.h
new file mode 100644
--- /dev/null
+++ b/Engine/Plugins/FX/Niagara/Source/NiagaraShader/Private/NiagaraAsyncGpuTraceProviderHwrt.h
@@ -0,0 +1,32 @@
+// Copyright Epic Games, Inc. All Rights Reserved.
+
+#pragma once
+
+#include "CoreMinimal.h"
+
+#include "NiagaraAsyncGpuTraceProvider.h"
+
+#if RHI_RAYTRACING
+
+class FNiagaraAsyncGpuTraceProviderHwrt : public FNiagaraAsyncGpuTraceProvider
+{
+public:
+	static const EProviderType Type;
+
+	FNiagaraAsyncGpuTraceProviderHwrt(EShaderPlatform InShaderPlatform, FNiagaraGpuComputeDispatchInterface* Dispatcher);
+
+	static bool IsSupported();
+	virtual void PostRenderOpaque(FRHICommandList& RHICmdList, TConstArrayView<FViewInfo> Views, FCollisionGroupHashMap* CollisionGroupHash) override;
+	virtual bool IsAvailable() const override;
+	virtual void IssueTraces(FRHICommandList& RHICmdList, const FDispatchRequest& Request, FCollisionGroupHashMap* CollisionGroupHash) override;
+	virtual void Reset() override;
+	virtual EProviderType GetType() const override { return Type; }
+
+private:
+	FRayTracingPipelineState* RayTracingPipelineState = nullptr;
+	FRHIRayTracingScene* RayTracingScene = nullptr;
+	FRHIShaderResourceView* RayTracingSceneView = nullptr;
+	TUniformBufferRef<FViewUniformShaderParameters> ViewUniformBuffer;
+};
+
+#endif
\ No newline at end of file
diff --git a/Engine/Plugins/FX/Niagara/Source/NiagaraShader/Private/NiagaraRayTracingHelper.cpp b/Engine/Plugins/FX/Niagara/Source/NiagaraShader/Private/NiagaraRayTracingHelper.cpp
deleted file mode 100644
diff --git a/Engine/Plugins/FX/Niagara/Source/NiagaraShader/Public/NiagaraAsyncGpuTraceProvider.h b/Engine/Plugins/FX/Niagara/Source/NiagaraShader/Public/NiagaraAsyncGpuTraceProvider.h
new file mode 100644
--- /dev/null
+++ b/Engine/Plugins/FX/Niagara/Source/NiagaraShader/Public/NiagaraAsyncGpuTraceProvider.h
@@ -0,0 +1,96 @@
+// Copyright Epic Games, Inc. All Rights Reserved.
+#pragma once
+
+#include "CoreMinimal.h"
+
+#include "NiagaraSettings.h"
+#include "RHIDefinitions.h"
+
+class FNiagaraGpuComputeDispatchInterface;
+class FViewInfo;
+
+// mirrors structure in Engine\Plugins\FX\Niagara\Shaders\Private\NiagaraAsyncGpuTraceCommon.ush
+struct FNiagaraAsyncGpuTrace
+{
+	float Origin[3];
+	float TFar;
+	float Direction[3];
+	uint32 CollisionGroup;
+};
+
+// mirrors structure in Engine\Plugins\FX\Niagara\Shaders\Private\NiagaraAsyncGpuTraceCommon.ush
+struct FNiagaraAsyncGpuTraceResult
+{
+	float WorldPosition[3];
+	float HitT;
+	float WorldNormal[3];
+	float _Pad0; // padding to force 16 byte alignment to meet requirements for VK
+};
+
+class NIAGARASHADER_API FNiagaraAsyncGpuTraceProvider
+{
+public:
+	using EProviderType = ENDICollisionQuery_AsyncGpuTraceProvider::Type;
+	using FProviderPriorityArray = TArray<TEnumAsByte<EProviderType>>;
+
+	FNiagaraAsyncGpuTraceProvider(EShaderPlatform InShaderPlatform, FNiagaraGpuComputeDispatchInterface* InDispatcher)
+		: ShaderPlatform(InShaderPlatform)
+		, Dispatcher(InDispatcher)
+	{
+	}
+
+	FNiagaraAsyncGpuTraceProvider() = delete;
+	virtual ~FNiagaraAsyncGpuTraceProvider() = default;
+
+	struct FDispatchRequest
+	{
+		FRWBufferStructured* TracesBuffer = nullptr;
+		FRWBufferStructured* ResultsBuffer = nullptr;
+		FRWBuffer* TraceCountsBuffer = nullptr;
+		uint32 TracesOffset = 0;
+		uint32 ResultsOffset = 0;
+		uint32 TraceCountsOffset = 0;
+		uint32 MaxTraceCount = 0;
+		uint32 MaxRetraceCount = 0;
+	};
+
+	static EProviderType ResolveSupportedType(EProviderType InType, const FProviderPriorityArray& Priorities);
+	static bool RequiresDistanceFieldData(EProviderType InType, const FProviderPriorityArray& Priorities);
+	static bool RequiresRayTracingScene(EProviderType InType, const FProviderPriorityArray& Priorities);
+
+	/** Hash table.
+		PrimIdHashTable is the main hash table that maps GPUSceneInstanceIndex to and Index we can use to store Collision Groups inside HashToCollisionGroups.
+	*/
+	struct FCollisionGroupHashMap
+	{
+		FRWBuffer PrimIdHashTable;
+		FRWBuffer HashToCollisionGroups;
+		uint32 HashTableSize = 0;
+	};
+
+	static void BuildCollisionGroupHashMap(FRHICommandList& RHICmdList, ERHIFeatureLevel::Type FeatureLevel, FScene* Scene, const TMap<FPrimitiveComponentId, uint32>& CollisionGroupMap, FCollisionGroupHashMap& Result);
+
+	static TArray<TUniquePtr<FNiagaraAsyncGpuTraceProvider>> CreateSupportedProviders(EShaderPlatform ShaderPlatform, FNiagaraGpuComputeDispatchInterface* Dispatcher, const FProviderPriorityArray& Priorities);
+	static void ClearResults(FRHICommandList& RHICmdList, EShaderPlatform ShaderPlatform, const FDispatchRequest& Request);
+
+	virtual bool IsAvailable() const = 0;
+	virtual EProviderType GetType() const = 0;
+
+	virtual void PostRenderOpaque(FRHICommandList& RHICmdList, TConstArrayView<FViewInfo> Views, FCollisionGroupHashMap* CollisionGroupHash)
+	{
+	}
+
+	virtual void IssueTraces(FRHICommandList& RHICmdList, const FDispatchRequest& Request, FCollisionGroupHashMap* CollisionGroupHash)
+	{
+	}
+
+	virtual void Reset()
+	{
+	}
+
+	const EShaderPlatform ShaderPlatform;
+	FNiagaraGpuComputeDispatchInterface* Dispatcher;
+
+private:
+	UE_NONCOPYABLE(FNiagaraAsyncGpuTraceProvider);
+};
\ No newline at end of file
diff --git a/Engine/Plugins/FX/Niagara/Source/NiagaraShader/Public/NiagaraRayTracingHelper.h b/Engine/Plugins/FX/Niagara/Source/NiagaraShader/Public/NiagaraRayTracingHelper.h
deleted file mode 100644
