Fixed parameter name being unreadable if selected + inactive + hovered
(InactiveHoveredBrush)

diff --git a/Engine/Plugins/FX/Niagara/Source/NiagaraEditor/Private/NiagaraEditorStyle.cpp b/Engine/Plugins/FX/Niagara/Source/NiagaraEditor/Private/NiagaraEditorStyle.cpp
--- a/Engine/Plugins/FX/Niagara/Source/NiagaraEditor/Private/NiagaraEditorStyle.cpp
+++ b/Engine/Plugins/FX/Niagara/Source/NiagaraEditor/Private/NiagaraEditorStyle.cpp
@@ -157,7 +157,7 @@ void FNiagaraEditorStyle::InitActionMenu()
             .SetActiveBrush(FSlateColorBrush(FStyleColors::Select))
             .SetActiveHoveredBrush(FSlateColorBrush(FStyleColors::Select))
             .SetInactiveBrush(FSlateColorBrush(FStyleColors::SelectInactive))
-            .SetInactiveHoveredBrush(FSlateColorBrush(FStyleColors::SelectHover))
+            .SetInactiveHoveredBrush(FSlateColorBrush(FStyleColors::SelectInactive))
             .SetActiveHighlightedBrush(FSlateColorBrush(FStyleColors::PrimaryHover))
             .SetInactiveHighlightedBrush(FSlateColorBrush(FStyleColors::SelectParent))
 			.SetTextColor(FStyleColors::Foreground)
