Remove struct member function for FNIagaraAsyncGpuTraceResult

diff --git a/Engine/Plugins/FX/Niagara/Shaders/Private/NiagaraAsyncGpuTraceCommon.ush b/Engine/Plugins/FX/Niagara/Shaders/Private/NiagaraAsyncGpuTraceCommon.ush
--- a/Engine/Plugins/FX/Niagara/Shaders/Private/NiagaraAsyncGpuTraceCommon.ush
+++ b/Engine/Plugins/FX/Niagara/Shaders/Private/NiagaraAsyncGpuTraceCommon.ush
@@ -20,6 +20,9 @@ struct FNiagaraAsyncGpuTraceResult
 	float3 WorldNormal;
 
 	float _Pad0; // padding to force 16 byte alignment to meet requirements for VK
-
-	bool IsHit() { return HitT > 0; }
 };
+
+bool IsHit(FNiagaraAsyncGpuTraceResult Result)
+{
+	return Result.HitT > 0;
+}
\ No newline at end of file
diff --git a/Engine/Plugins/FX/Niagara/Shaders/Private/NiagaraDataInterfaceAsyncGpuTrace.ush b/Engine/Plugins/FX/Niagara/Shaders/Private/NiagaraDataInterfaceAsyncGpuTrace.ush
--- a/Engine/Plugins/FX/Niagara/Shaders/Private/NiagaraDataInterfaceAsyncGpuTrace.ush
+++ b/Engine/Plugins/FX/Niagara/Shaders/Private/NiagaraDataInterfaceAsyncGpuTrace.ush
@@ -103,7 +103,7 @@ void NDIAsyncGpuTrace_ReadAsyncRayTrace(uint MaxRayTraceCount, StructuredBuffer<
 	{
 		FNiagaraAsyncGpuTraceResult HitResult = IntersectionResults[In_PreviousFrameQueryID + IntersectionRestultsOffest];
 
-		Out_CollisionValid = HitResult.IsHit();
+		Out_CollisionValid = IsHit(HitResult);
 		if (Out_CollisionValid)
 		{
 			Out_CollisionDistance = HitResult.HitT;
