[Backout] - CL18953506
Original CL Desc
-----------------------------------------------------------------
Clean up pass on HWRT within Niagara (though it remains experimental).  Cleans up the abstraction to allow us to have a fallback method for platforms/scenarios where HW ray tracing is suppported/enabled/used (currently only SDF is provided as a fallback).
-shuffles a lot of the related files around to clean up public/private and moving stuff out of NiagaraShader where possible
-migrated async trace functions into it's own DI, away from the CollisionQuery one
-updates collision module to use the new DI for it's experimental ray traced GPU version

diff --git a/Engine/Plugins/FX/Niagara/Shaders/Private/NiagaraAsyncGpuTraceCommon.ush b/Engine/Plugins/FX/Niagara/Shaders/Private/NiagaraAsyncGpuTraceCommon.ush
deleted file mode 100644
diff --git a/Engine/Plugins/FX/Niagara/Shaders/Private/NiagaraAsyncGpuTraceUtils.usf b/Engine/Plugins/FX/Niagara/Shaders/Private/NiagaraAsyncGpuTraceUtils.usf
deleted file mode 100644
diff --git a/Engine/Plugins/FX/Niagara/Shaders/Private/NiagaraDataInterfaceAsyncGpuTrace.ush b/Engine/Plugins/FX/Niagara/Shaders/Private/NiagaraDataInterfaceAsyncGpuTrace.ush
deleted file mode 100644
diff --git a/Engine/Plugins/FX/Niagara/Shaders/Private/NiagaraDataInterfaceAsyncGpuTraceTemplate.ush b/Engine/Plugins/FX/Niagara/Shaders/Private/NiagaraDataInterfaceAsyncGpuTraceTemplate.ush
deleted file mode 100644
diff --git a/Engine/Plugins/FX/Niagara/Shaders/Private/NiagaraDataInterfaceCollisionQuery.ush b/Engine/Plugins/FX/Niagara/Shaders/Private/NiagaraDataInterfaceCollisionQuery.ush
--- a/Engine/Plugins/FX/Niagara/Shaders/Private/NiagaraDataInterfaceCollisionQuery.ush
+++ b/Engine/Plugins/FX/Niagara/Shaders/Private/NiagaraDataInterfaceCollisionQuery.ush
@@ -2,7 +2,7 @@
 
 #pragma once
 
-#include "/Engine/Private/Common.ush"
+#include "NiagaraRayTracingCommon.ush"
 
 #define NDICOLLISIONQUERY_USE_GBUFFER_NORMAL    (FEATURE_LEVEL >= FEATURE_LEVEL_SM5)
 
@@ -144,3 +144,111 @@ void NDICollisionQuery_QueryMeshDistanceFieldGPU(in float3 In_SamplePos, in floa
 	Out_IsDistanceFieldValid = false;
 #endif
 }
+
+#if NIAGARA_USE_RAYTRACING
+
+uint BuildRayTraceMask(int TraceChannel)
+{
+	// for now we'll ignore the TraceChannel till we can mark up the environment
+	return 0xFF;
+}
+
+void NDICollisionQuery_IssueAsyncRayTrace(uint MaxRayTraceCount, RWStructuredBuffer<FNiagaraRayData> RWRayRequests, uint RayRequestOffset, RWBuffer<uint> RWRayTraceCounts, uint RayTraceCountsOffset, float3 In_LWCTile, int In_QueryID, float3 In_TraceStart, float3 In_TraceEnd, int In_TraceChannel, out bool Out_IsQueryValid)
+{
+	Out_IsQueryValid = false;
+
+	if (In_QueryID >= 0 && uint(In_QueryID) < MaxRayTraceCount)
+	{
+		FLWCVector3 LwcStartPos = MakeLWCVector3(In_LWCTile, In_TraceStart);
+		FLWCVector3 LwcEndPos = MakeLWCVector3(In_LWCTile, In_TraceEnd);
+		
+		float3 Trace = LWCToFloat(LWCSubtract(LwcEndPos, LwcStartPos));
+		float TraceLength = length(Trace);
+
+		//Offset query into the actual ray buffer.
+		uint RayIndex = RayRequestOffset + In_QueryID;
+		RWRayRequests[RayIndex].Origin = LWCToFloat(LwcStartPos); // TODO: remove LWCToFloat here when FNiagaraRayData supports LWC
+		RWRayRequests[RayIndex].Mask = BuildRayTraceMask(In_TraceChannel);
+		RWRayRequests[RayIndex].Direction = normalize(Trace);
+		RWRayRequests[RayIndex].TFar = TraceLength;
+		Out_IsQueryValid = true;
+
+		// we want to keep track of the maximum index used
+		InterlockedMax(RWRayTraceCounts[RayTraceCountsOffset + 0], In_QueryID + 1);
+		RWRayTraceCounts[RayTraceCountsOffset + 1] = 1;
+		RWRayTraceCounts[RayTraceCountsOffset + 2] = 1;
+	}
+}
+
+void NDICollisionQuery_ReserveRayTraceIndex(uint MaxRayTraceCount, RWBuffer<uint> RWRayTraceCounts, uint RayTraceCountsOffset, int In_TraceCount, out int Out_StartQueryID, out bool Out_IndicesValid)
+{
+	Out_StartQueryID = -1;
+	Out_IndicesValid = false;
+
+	if ((In_TraceCount + RWRayTraceCounts[RayTraceCountsOffset + 0]) < MaxRayTraceCount)
+	{
+		uint RequestIndex = 0;
+		InterlockedAdd(RWRayTraceCounts[RayTraceCountsOffset + 0], In_TraceCount, RequestIndex);
+		RWRayTraceCounts[RayTraceCountsOffset + 1] = 1;
+		RWRayTraceCounts[RayTraceCountsOffset + 2] = 1;
+
+		if ((In_TraceCount + RequestIndex) <= MaxRayTraceCount)
+		{
+			Out_StartQueryID = int(RequestIndex);
+			Out_IndicesValid = true;
+		}
+		else
+		{
+			InterlockedMax(RWRayTraceCounts[RayTraceCountsOffset + 0], MaxRayTraceCount);
+		}
+	}
+}
+
+void NDICollisionQuery_CreateAsyncRayTrace(	uint MaxRayTraceCount, RWStructuredBuffer<FNiagaraRayData> RWRayRequests, uint RayRequestsOffset, RWBuffer<uint> RWRayTraceCounts, uint RayTraceCountsOffset, float3 In_LWCTile,
+											float3 In_TraceStart, float3 In_TraceEnd, int In_TraceChannel, out int Out_QueryID, out bool Out_IsQueryValid)
+{
+	Out_IsQueryValid = false;
+
+	NDICollisionQuery_ReserveRayTraceIndex(MaxRayTraceCount, RWRayTraceCounts, RayTraceCountsOffset, 1, Out_QueryID, Out_IsQueryValid);
+
+	if (Out_IsQueryValid)
+	{
+		FLWCVector3 LwcStartPos = MakeLWCVector3(In_LWCTile, In_TraceStart);
+		FLWCVector3 LwcEndPos = MakeLWCVector3(In_LWCTile, In_TraceEnd);
+		
+		float3 Trace = LWCToFloat(LWCSubtract(LwcEndPos, LwcStartPos));
+		float TraceLength = length(Trace);
+
+		uint RayIndex = Out_QueryID + RayRequestsOffset;
+		RWRayRequests[RayIndex].Origin = LWCToFloat(LwcStartPos); // TODO: remove LWCToFloat here when FNiagaraRayData supports LWC
+		RWRayRequests[RayIndex].Mask = BuildRayTraceMask(In_TraceChannel);
+		RWRayRequests[RayIndex].Direction = normalize(Trace);
+		RWRayRequests[RayIndex].TFar = TraceLength;
+
+		//TODO: We can pack collision group and trace channels into the same int and test both in the RG Shader.
+		RWRayRequests[RayIndex].CollisionGroup = In_TraceChannel;
+	}
+}
+
+void NDICollisionQuery_ReadAsyncRayTrace(uint MaxRayTraceCount, StructuredBuffer<FNiagaraRayTracingResult> IntersectionResults, uint IntersectionRestultsOffest, float3 In_LWCTile, int In_PreviousFrameQueryID, out bool Out_CollisionValid, out float Out_CollisionDistance, out float3 Out_CollisionPosWorld, out float3 Out_CollisionNormal)
+{
+	Out_CollisionValid = false;
+	Out_CollisionDistance = 0.0f;
+	Out_CollisionPosWorld = float3(0.0f, 0.0f, 0.0f);
+	Out_CollisionNormal = float3(0.0f, 0.0f, 0.0f);
+
+	if (In_PreviousFrameQueryID >= 0 && uint(In_PreviousFrameQueryID) < MaxRayTraceCount)
+	{
+		FNiagaraRayTracingResult HitResult = IntersectionResults[In_PreviousFrameQueryID + IntersectionRestultsOffest];
+
+		Out_CollisionValid = HitResult.IsHit();
+		if (Out_CollisionValid)
+		{
+			Out_CollisionDistance = HitResult.HitT;
+			Out_CollisionPosWorld = LWCToFloat(MakeLWCVector3(-In_LWCTile, HitResult.WorldPosition)); // TODO: HitResult.WorldPosition should support LWC properly
+			Out_CollisionNormal = HitResult.WorldNormal;
+		}
+	}
+}
+
+#endif
\ No newline at end of file
diff --git a/Engine/Plugins/FX/Niagara/Shaders/Private/NiagaraDataInterfaceCollisionQueryTemplate.ush b/Engine/Plugins/FX/Niagara/Shaders/Private/NiagaraDataInterfaceCollisionQueryTemplate.ush
--- a/Engine/Plugins/FX/Niagara/Shaders/Private/NiagaraDataInterfaceCollisionQueryTemplate.ush
+++ b/Engine/Plugins/FX/Niagara/Shaders/Private/NiagaraDataInterfaceCollisionQueryTemplate.ush
@@ -1,6 +1,17 @@
 // Copyright Epic Games, Inc. All Rights Reserved.
 
 float3										SystemLWCTile_{ParameterName};
+#if NIAGARA_USE_RAYTRACING
+/** Niagara ray tracing uses shared buffers between dispatches so each buffer has a buffer + and offset into that buffer for the current dispatch. */
+bool										RayTracingEnabled_{ParameterName};
+uint										MaxRayTraceCount_{ParameterName};
+RWStructuredBuffer<FNiagaraRayData>			RWRayRequests_{ParameterName};
+uint										RayRequestsOffset_{ ParameterName };
+StructuredBuffer<FNiagaraRayTracingResult>	IntersectionResults_{ParameterName};
+uint										IntersectionResultsOffset_{ ParameterName };
+RWBuffer<uint>								RWRayTraceCounts_{ParameterName};
+uint										RayTraceCountsOffset_{ ParameterName };
+#endif
 
 void QuerySceneDepthGPU_{ParameterName}(in float3 In_SamplePos, out float Out_SceneDepth, out float3 Out_CameraPosWorld, out bool Out_IsInsideView, out float3 Out_WorldPos, out float3 Out_WorldNormal)
 {
@@ -17,29 +28,55 @@ void QueryMeshDistanceFieldGPU_{ParameterName}(in float3 In_SamplePos, out float
 	NDICollisionQuery_QueryMeshDistanceFieldGPU(In_SamplePos, SystemLWCTile_{ParameterName}, Out_DistanceToNearestSurface, Out_FieldGradient, Out_IsDistanceFieldValid);
 }
 
-// DEPRECATE_BEGIN
 void IssueAsyncRayTraceGpu_{ParameterName}_UEImpureCall(int In_QueryID, float3 In_TraceStart, float3 In_TraceEnd, int In_TraceChannel, out bool Out_IsQueryValid)
 {
+#if NIAGARA_USE_RAYTRACING
+	if (RayTracingEnabled_{ParameterName})
+	{
+		NDICollisionQuery_IssueAsyncRayTrace(MaxRayTraceCount_{ParameterName}, RWRayRequests_{ParameterName}, RayRequestsOffset_{ ParameterName }, RWRayTraceCounts_{ParameterName}, RayTraceCountsOffset_{ ParameterName }, SystemLWCTile_{ParameterName}, In_QueryID, In_TraceStart, In_TraceEnd, In_TraceChannel, Out_IsQueryValid);
+		return;
+	}
+#endif
 	Out_IsQueryValid = false;
 }
 
 void ReserveRayTraceIndex_{ParameterName}_UEImpureCall(int In_TraceCount, out int Out_StartQueryID, out bool Out_IndicesValid)
 {
+#if NIAGARA_USE_RAYTRACING
+	if (RayTracingEnabled_{ParameterName})
+	{
+		NDICollisionQuery_ReserveRayTraceIndex(MaxRayTraceCount_{ParameterName}, RWRayTraceCounts_{ParameterName}, RayTraceCountsOffset_{ ParameterName }, In_TraceCount, Out_StartQueryID, Out_IndicesValid);
+		return;
+	}
+#endif
 	Out_StartQueryID = -1;
 	Out_IndicesValid = false;
 }
 
 void CreateAsyncRayTraceGpu_{ParameterName}_UEImpureCall(float3 In_TraceStart, float3 In_TraceEnd, int In_TraceChannel, out int Out_QueryID, out bool Out_IsQueryValid)
 {
+#if NIAGARA_USE_RAYTRACING
+	if (RayTracingEnabled_{ParameterName})
+	{
+		NDICollisionQuery_CreateAsyncRayTrace(MaxRayTraceCount_{ParameterName}, RWRayRequests_{ParameterName}, RayRequestsOffset_{ ParameterName }, RWRayTraceCounts_{ParameterName}, RayTraceCountsOffset_{ ParameterName }, SystemLWCTile_{ParameterName}, In_TraceStart, In_TraceEnd, In_TraceChannel, Out_QueryID, Out_IsQueryValid);
+		return;
+	}
+#endif
 	Out_QueryID = -1;
 	Out_IsQueryValid = false;
 }
 
 void ReadAsyncRayTraceGpu_{ParameterName}(int In_PreviousFrameQueryID, out bool Out_CollisionValid, out float Out_CollisionDistance, out float3 Out_CollisionPosWorld, out float3 Out_CollisionNormal)
 {
+#if NIAGARA_USE_RAYTRACING
+	if (RayTracingEnabled_{ParameterName})
+	{
+		NDICollisionQuery_ReadAsyncRayTrace(MaxRayTraceCount_{ParameterName}, IntersectionResults_{ParameterName}, IntersectionResultsOffset_{ParameterName}, SystemLWCTile_{ParameterName}, In_PreviousFrameQueryID, Out_CollisionValid, Out_CollisionDistance, Out_CollisionPosWorld, Out_CollisionNormal);
+		return;
+	}
+#endif
 	Out_CollisionDistance = 0.0f;
 	Out_CollisionPosWorld = float3(0.0f, 0.0f, 0.0f);
 	Out_CollisionNormal = float3(0.0f, 0.0f, 0.0f);
 	Out_CollisionValid = false;
-}
-// DEPRECATE_END
+}
\ No newline at end of file
diff --git a/Engine/Plugins/FX/Niagara/Shaders/Private/NiagaraRayMarchingShaders.usf b/Engine/Plugins/FX/Niagara/Shaders/Private/NiagaraRayMarchingShaders.usf
deleted file mode 100644
diff --git a/Engine/Plugins/FX/Niagara/Shaders/Private/NiagaraRayTracingCommon.ush b/Engine/Plugins/FX/Niagara/Shaders/Private/NiagaraRayTracingCommon.ush
new file mode 100644
--- /dev/null
+++ b/Engine/Plugins/FX/Niagara/Shaders/Private/NiagaraRayTracingCommon.ush
@@ -0,0 +1,30 @@
+// Copyright Epic Games, Inc. All Rights Reserved.
+
+#pragma once
+
+#define NIAGARA_USE_RAYTRACING        (NIAGARA_SUPPORTS_RAY_TRACING && (FEATURE_LEVEL >= FEATURE_LEVEL_SM5))
+
+#if NIAGARA_USE_RAYTRACING
+
+#include "/Engine/Private/RayTracing/RayTracingCommon.ush"
+
+// ray input structure for Niagara ray traces
+struct FNiagaraRayData : FBasicRayData
+{
+	uint CollisionGroup;
+};
+
+struct FNiagaraRayTracingPayload : FVFXTracePayload
+{
+
+};
+
+// payload structure for Niagara ray traces
+struct FNiagaraRayTracingResult : FMinimalPayload
+{
+	float3 WorldPosition;
+	float3 WorldNormal;
+};
+
+
+#endif // NIAGARA_USE_RAYTRACING
\ No newline at end of file
diff --git a/Engine/Plugins/FX/Niagara/Shaders/Private/NiagaraRayTracingShaders.usf b/Engine/Plugins/FX/Niagara/Shaders/Private/NiagaraRayTracingShaders.usf
--- a/Engine/Plugins/FX/Niagara/Shaders/Private/NiagaraRayTracingShaders.usf
+++ b/Engine/Plugins/FX/Niagara/Shaders/Private/NiagaraRayTracingShaders.usf
@@ -5,14 +5,14 @@
 #include "/Engine/Private/ShadingCommon.ush"
 #include "/Engine/Private/RayTracing/RayTracingCommon.ush"
 #include "/Engine/Private/RayTracing/RayTracingHitGroupCommon.ush"
-#include "NiagaraAsyncGpuTraceCommon.ush"
+#include "NiagaraRayTracingCommon.ush"
 #include "/Engine/Private/HashTable.ush"
 #include "/Engine/Private/RayTracing/VFXTraceRay.ush"
 
 RaytracingAccelerationStructure					TLAS;
-StructuredBuffer<FNiagaraAsyncGpuTrace>			Rays;
+StructuredBuffer<FNiagaraRayData>				Rays;
 uint											RaysOffset;
-RWStructuredBuffer<FNiagaraAsyncGpuTraceResult>	CollisionOutput;
+RWStructuredBuffer<FNiagaraRayTracingResult>	CollisionOutput;
 uint											CollisionOutputOffset;
 
 #if NIAGARA_RAYTRACE_FAKE_INDIRECT
@@ -22,7 +22,6 @@ Buffer<uint>                                  RayTraceCounts;
 Buffer<uint> HashToCollisionGroups;
 uint MaxRetraces;
 
-#if NIAGARA_SUPPORTS_COLLISION_GROUPS
 uint GetGPUSceneInstanceIndex(uint PrimitiveIndex, uint InstanceIndex)
 {
 #if VF_USE_PRIMITIVE_SCENE_DATA
@@ -51,7 +50,6 @@ bool FindCollisionGroup(int GPUSceneIndexId, out int CollisionGroup)
 
 	return false;
 }
-#endif
 
 
 RAY_TRACING_ENTRY_RAYGEN(NiagaraCollisionRayTraceRG)
@@ -65,7 +63,7 @@ RAY_TRACING_ENTRY_RAYGEN(NiagaraCollisionRayTraceRG)
 	}
 #endif
 
-	FNiagaraAsyncGpuTrace InputRay = Rays[RayIndex + RaysOffset];
+	FNiagaraRayData InputRay = Rays[RayIndex + RaysOffset];
 
 	RayDesc Ray;
 	Ray.Origin = InputRay.Origin;
@@ -78,7 +76,6 @@ RAY_TRACING_ENTRY_RAYGEN(NiagaraCollisionRayTraceRG)
 	FVFXTracePayload Payload = (FVFXTracePayload) 0;
 	uint RayCollisionGroup = InputRay.CollisionGroup;
 
-#if NIAGARA_SUPPORTS_COLLISION_GROUPS
 	if (RayCollisionGroup != -1)
 	{
 		uint NumCasts = 0;
@@ -116,7 +113,6 @@ RAY_TRACING_ENTRY_RAYGEN(NiagaraCollisionRayTraceRG)
 		}
 	}
 	else
-#endif
 	{
 		VfxTraceRay(
 			TLAS,   // AccelerationStructure
@@ -131,7 +127,6 @@ RAY_TRACING_ENTRY_RAYGEN(NiagaraCollisionRayTraceRG)
 	}
 
 	uint OutputIndex = RayIndex + CollisionOutputOffset;
-	CollisionOutput[OutputIndex] = (FNiagaraAsyncGpuTraceResult) 0;
 	CollisionOutput[OutputIndex].HitT = Payload.HitT;
 	if (Payload.IsHit())
 	{
diff --git a/Engine/Plugins/FX/Niagara/Source/Niagara/Classes/NiagaraDataInterfaceAsyncGpuTrace.h b/Engine/Plugins/FX/Niagara/Source/Niagara/Classes/NiagaraDataInterfaceAsyncGpuTrace.h
deleted file mode 100644
diff --git a/Engine/Plugins/FX/Niagara/Source/Niagara/Classes/NiagaraDataInterfaceCollisionQuery.h b/Engine/Plugins/FX/Niagara/Source/Niagara/Classes/NiagaraDataInterfaceCollisionQuery.h
--- a/Engine/Plugins/FX/Niagara/Source/Niagara/Classes/NiagaraDataInterfaceCollisionQuery.h
+++ b/Engine/Plugins/FX/Niagara/Source/Niagara/Classes/NiagaraDataInterfaceCollisionQuery.h
@@ -29,8 +29,21 @@ public:
 
 	DECLARE_NIAGARA_DI_PARAMETER();
 
+	UPROPERTY(EditAnywhere, Category = "Ray Trace")
+	int32 MaxTracesPerParticle = 0;
+
+	/** If a collision is rejected, how many times do we attempt to retrace from that collision point forward to find a new, valid collision.*/
+	UPROPERTY(EditAnywhere, Category = "Ray Trace")
+	int32 MaxRetraces = 0;
+
+	FNiagaraSystemInstance *SystemInstance;
+
 	//UObject Interface
 	virtual void PostInitProperties() override;
+	virtual void PostLoad() override;
+#if WITH_EDITOR
+	virtual void PostEditChangeProperty(struct FPropertyChangedEvent& PropertyChangedEvent) override;
+#endif
 	//UObject Interface End
 
 	/** Initializes the per instance data for this interface. Returns false if there was some error and the simulation should be disabled. */
@@ -46,6 +59,7 @@ public:
 	virtual void GetFunctions(TArray<FNiagaraFunctionSignature>& OutFunctions) override;
 	virtual void GetVMExternalFunction(const FVMExternalFunctionBindingInfo& BindingInfo, void* InstanceData, FVMExternalFunction &OutFunc) override;
 	virtual void GetAssetTagsForContext(const UObject* InAsset, const TArray<const UNiagaraDataInterface*>& InProperties, TMap<FName, uint32>& NumericKeys, TMap<FName, FString>& StringKeys) const override;
+	virtual void PushToRenderThreadImpl() override;
 
 	// VM functions
 	void PerformQuerySyncCPU(FVectorVMExternalFunctionContext& Context);
@@ -54,9 +68,11 @@ public:
 	virtual bool CanExecuteOnTarget(ENiagaraSimTarget Target) const override { return true; }
 	virtual bool RequiresDistanceFieldData() const override { return true; }
 	virtual bool RequiresDepthBuffer() const override { return true; }
+	virtual bool RequiresRayTracingScene() const override;
 
 #if WITH_EDITORONLY_DATA
 	virtual bool AppendCompileHash(FNiagaraCompileHashVisitor* InVisitor) const override;
+	virtual void ModifyCompilationEnvironment(EShaderPlatform ShaderPlatform, struct FShaderCompilerEnvironment& OutEnvironment) const override;
 	virtual void GetCommonHLSL(FString& OutHLSL) override;
 	virtual void GetParameterDefinitionHLSL(const FNiagaraDataInterfaceGPUParamInfo& ParamInfo, FString& OutHLSL) override;
 	virtual bool GetFunctionHLSL(const FNiagaraDataInterfaceGPUParamInfo& ParamInfo, const FNiagaraDataInterfaceGeneratedFunction& FunctionInfo, int FunctionInstanceIndex, FString& OutHLSL) override;
@@ -72,6 +88,14 @@ public:
 	virtual bool PostSimulateCanOverlapFrames() const { return false; }
 
 private:
+
+	virtual bool Equals(const UNiagaraDataInterface* Other) const override;
+
+protected:
+	virtual bool CopyToInternal(UNiagaraDataInterface* Destination) const override;
+
+private:
+
 	static FCriticalSection CriticalSection;
 	UEnum* TraceChannelEnum;
 };
diff --git a/Engine/Plugins/FX/Niagara/Source/Niagara/Private/DataInterface/NiagaraDataInterfaceAsyncGpuTrace.cpp b/Engine/Plugins/FX/Niagara/Source/Niagara/Private/DataInterface/NiagaraDataInterfaceAsyncGpuTrace.cpp
deleted file mode 100644
diff --git a/Engine/Plugins/FX/Niagara/Source/Niagara/Private/NiagaraAsyncGpuTraceHelper.cpp b/Engine/Plugins/FX/Niagara/Source/Niagara/Private/NiagaraAsyncGpuTraceHelper.cpp
deleted file mode 100644
diff --git a/Engine/Plugins/FX/Niagara/Source/Niagara/Private/NiagaraAsyncGpuTraceHelper.h b/Engine/Plugins/FX/Niagara/Source/Niagara/Private/NiagaraAsyncGpuTraceHelper.h
deleted file mode 100644
diff --git a/Engine/Plugins/FX/Niagara/Source/Niagara/Private/NiagaraDataInterfaceCollisionQuery.cpp b/Engine/Plugins/FX/Niagara/Source/Niagara/Private/NiagaraDataInterfaceCollisionQuery.cpp
--- a/Engine/Plugins/FX/Niagara/Source/Niagara/Private/NiagaraDataInterfaceCollisionQuery.cpp
+++ b/Engine/Plugins/FX/Niagara/Source/Niagara/Private/NiagaraDataInterfaceCollisionQuery.cpp
@@ -3,11 +3,10 @@
 #include "NiagaraDataInterfaceCollisionQuery.h"
 
 #include "GlobalDistanceFieldParameters.h"
-#include "NiagaraAsyncGpuTraceHelper.h"
 #include "NiagaraComponent.h"
-#include "NiagaraDataInterfaceUtilities.h"
 #include "NiagaraGpuComputeDispatchInterface.h"
 #include "NiagaraGpuComputeDispatch.h"
+#include "NiagaraRayTracingHelper.h"
 #include "NiagaraSimStageData.h"
 #include "NiagaraStats.h"
 #include "NiagaraTypes.h"
@@ -20,6 +19,22 @@
 
 #define LOCTEXT_NAMESPACE "NiagaraDataInterfaceCollisionQuery"
 
+void OnHWRTCollisionsEnabledChanged(IConsoleVariable* CVar)
+{
+	//Force a reinit of everything just to be safe.
+	FNiagaraSystemUpdateContext Context;
+	Context.AddAll(true);
+}
+
+int32 GEnableGPUHWRTCollisions = 1;
+static FAutoConsoleVariableRef CVarEnableGPUHWRTCollisions(
+	TEXT("fx.Niagara.Collision.EnableGPURayTracedCollisions"),
+	GEnableGPUHWRTCollisions,
+	TEXT("If greater than zero, GPU hardware ray trace collisions are enabled."),
+	FConsoleVariableDelegate::CreateStatic(&OnHWRTCollisionsEnabledChanged),
+	ECVF_Default
+);
+
 namespace NDICollisionQueryLocal
 {
 	static const TCHAR* CommonShaderFile = TEXT("/Plugin/FX/Niagara/Private/NiagaraDataInterfaceCollisionQuery.ush");
@@ -30,14 +45,19 @@ namespace NDICollisionQueryLocal
 	static const FName DistanceFieldName(TEXT("QueryMeshDistanceFieldGPU"));
 	static const FName SyncTraceName(TEXT("PerformCollisionQuerySyncCPU"));
 	static const FName AsyncTraceName(TEXT("PerformCollisionQueryAsyncCPU"));
-
-	// DEPRECATE_BEGIN
 	static const FName IssueAsyncRayTraceName(TEXT("IssueAsyncRayTraceGpu"));
 	static const FName CreateAsyncRayTraceName(TEXT("CreateAsyncRayTraceGpu"));
 	static const FName ReserveAsyncRayTraceName(TEXT("ReserveAsyncRayTraceGpu"));
 	static const FName ReadAsyncRayTraceName(TEXT("ReadAsyncRayTraceGpu"));
-	// DEPRECATE_END
 
+	static const FString RayTracingEnabledParamName(TEXT("RayTracingEnabled_"));
+	static const FString MaxRayTraceCountParamName(TEXT("MaxRayTraceCount_"));
+	static const FString RayRequestsParamName(TEXT("RayRequests_"));
+	static const FString RayRequestsOffsetParamName(TEXT("RayRequestsOffset_"));
+	static const FString IntersectionResultsParamName(TEXT("IntersectionResults_"));
+	static const FString IntersectionResultsOffsetParamName(TEXT("IntersectionResultsOffset_"));
+	static const FString RayTraceCountsParamName(TEXT("RayTraceCounts_"));
+	static const FString RayTraceCountsOffsetParamName(TEXT("RayTraceCountsOffset_"));
 	static const FString SystemLWCTileName(TEXT("SystemLWCTile_"));
 }
 
@@ -52,28 +72,79 @@ struct FNiagaraCollisionDIFunctionVersion
 		AddedCustomDepthCollision = 2,
 		ReturnCollisionMaterialIdx = 3,
 		LargeWorldCoordinates = 4,
-		DeprecatedAsyncGpuTrace = 5,
 
 		VersionPlusOne,
 		LatestVersion = VersionPlusOne - 1
 	};
 };
-
 struct FNiagaraDataIntefaceProxyCollisionQuery : public FNiagaraDataInterfaceProxy
 {
-	// There's nothing in this proxy. It just reads from scene textures.
+#if RHI_RAYTRACING
+	int32 MaxTracesPerParticle = 0;
+	uint32 MaxRetraces = 0;
+#endif
+
+	virtual ~FNiagaraDataIntefaceProxyCollisionQuery()
+	{
+	}
 
 	virtual int32 PerInstanceDataPassedToRenderThreadSize() const override
 	{
 		return 0;
 	}
+
+	virtual void PreStage(FRHICommandList& RHICmdList, const FNiagaraDataInterfaceStageArgs& Context) override
+	{
+		FNiagaraDataInterfaceProxy::PreStage(RHICmdList, Context);
+
+#if RHI_RAYTRACING
+		if (IsRayTracingEnabled() && GEnableGPUHWRTCollisions && MaxTracesPerParticle > 0)
+		{
+			//Accumulate the total ray requests for this DI for all dispatches in the stage.
+			int32 RayRequests = MaxTracesPerParticle * Context.SimStageData->DestinationNumInstances;
+			FNiagaraRayTracingHelper& RTHelper = Context.ComputeDispatchInterface->GetRayTracingHelper();
+			RTHelper.AddToDispatch(this, RayRequests, MaxRetraces);
+		}
+#endif
+	}
+
+#if RHI_RAYTRACING
+	virtual bool RequiresPreStageFinalize() const override
+	{
+		return true;
+	}
+
+	virtual void FinalizePreStage(FRHICommandList& RHICmdList, const FNiagaraGpuComputeDispatchInterface* ComputeDispatchInterface) override
+	{
+		FNiagaraRayTracingHelper& RTHelper = ComputeDispatchInterface->GetRayTracingHelper();
+		if (IsRayTracingEnabled() && GEnableGPUHWRTCollisions && MaxTracesPerParticle > 0)
+		{
+			RTHelper.BuildDispatch(RHICmdList, this);
+		}
+		else
+		{
+			RTHelper.BuildDummyDispatch(RHICmdList);
+		}
+	}
+#endif
+
+	void RenderThreadInitialize(int32 InMaxTracesPerParticle, uint32 InMaxRetraces)
+	{
+#if RHI_RAYTRACING
+		MaxRetraces = InMaxRetraces;
+		MaxTracesPerParticle = InMaxTracesPerParticle;
+#endif
+	}
 };
 
+
 UNiagaraDataInterfaceCollisionQuery::UNiagaraDataInterfaceCollisionQuery(FObjectInitializer const& ObjectInitializer)
 	: Super(ObjectInitializer)
 {
 	TraceChannelEnum = StaticEnum<ECollisionChannel>();
-	Proxy.Reset(new FNiagaraDataIntefaceProxyCollisionQuery());
+	SystemInstance = nullptr;
+
+    Proxy.Reset(new FNiagaraDataIntefaceProxyCollisionQuery());
 }
 
 bool UNiagaraDataInterfaceCollisionQuery::InitPerInstanceData(void* PerInstanceData, FNiagaraSystemInstance* InSystemInstance)
@@ -84,7 +155,6 @@ bool UNiagaraDataInterfaceCollisionQuery::InitPerInstanceData(void* PerInstanceD
 	{
 		PIData->CollisionBatch.Init(InSystemInstance->GetId(), InSystemInstance->GetWorld());
 	}
-
 	return true;
 }
 
@@ -107,6 +177,16 @@ void UNiagaraDataInterfaceCollisionQuery::PostInitProperties()
 	}
 }
 
+void UNiagaraDataInterfaceCollisionQuery::PostLoad()
+{
+	Super::PostLoad();
+
+	if (MaxTracesPerParticle)
+	{
+		MarkRenderDataDirty();
+	}
+}
+
 void UNiagaraDataInterfaceCollisionQuery::GetAssetTagsForContext(const UObject* InAsset, const TArray<const UNiagaraDataInterface*>& InProperties, TMap<FName, uint32>& NumericKeys, TMap<FName, FString>& StringKeys) const
 {
 #if WITH_EDITOR
@@ -261,8 +341,6 @@ void UNiagaraDataInterfaceCollisionQuery::GetFunctions(TArray<FNiagaraFunctionSi
 	
 	const FText TraceStartWorldDescription = LOCTEXT("TraceStartWorldDescription", "Ray starting point in world space");
 	const FText TraceEndWorldDescription = LOCTEXT("TraceEndWorldDescription", "Ray end point in world space");
-
-	// DEPRECATE_BEGIN
 	{
 		const FText AsyncTraceChannelDescription = LOCTEXT("TraceChannelDescription", "Currently unused, will represent the trace channels for which geometry the trace should test against");
 		const FText QueryIDDescription = LOCTEXT("QueryIDDescription", "Unique (for this frame) index of the query being enqueued (used in subsequent frames to retrieve results).  Must be less than MaxRayTraceCount");
@@ -273,7 +351,6 @@ void UNiagaraDataInterfaceCollisionQuery::GetFunctions(TArray<FNiagaraFunctionSi
 		IssueRayTrace.bRequiresExecPin = true;
 		IssueRayTrace.bMemberFunction = true;
 		IssueRayTrace.bSupportsCPU = false;
-		IssueRayTrace.bSoftDeprecatedFunction = true;
 #if WITH_EDITORONLY_DATA
 		IssueRayTrace.FunctionVersion = FNiagaraCollisionDIFunctionVersion::LatestVersion;
 		IssueRayTrace.Description = LOCTEXT("IssueAsync_RayTraceDescription", "Enqueues a GPU raytrace with the result being available the following frame");
@@ -290,7 +367,6 @@ void UNiagaraDataInterfaceCollisionQuery::GetFunctions(TArray<FNiagaraFunctionSi
 		CreateRayTrace.bRequiresExecPin = true;
 		CreateRayTrace.bMemberFunction = true;
 		CreateRayTrace.bSupportsCPU = false;
-		CreateRayTrace.bSoftDeprecatedFunction = true;
 #if WITH_EDITORONLY_DATA
 		CreateRayTrace.FunctionVersion = FNiagaraCollisionDIFunctionVersion::LatestVersion;
 		CreateRayTrace.Description = LOCTEXT("CreateAsync_RayTraceDescription", "Creates a GPU raytrace with the result being available the following frame (index is returned)");
@@ -304,20 +380,19 @@ void UNiagaraDataInterfaceCollisionQuery::GetFunctions(TArray<FNiagaraFunctionSi
 	}
 
 	{
-		FNiagaraFunctionSignature& ReserveRayTrace = OutFunctions.AddDefaulted_GetRef();
-		ReserveRayTrace.Name = NDICollisionQueryLocal::ReserveAsyncRayTraceName;
-		ReserveRayTrace.bRequiresExecPin = true;
-		ReserveRayTrace.bMemberFunction = true;
-		ReserveRayTrace.bSupportsCPU = false;
-		ReserveRayTrace.bSoftDeprecatedFunction = true;
+		FNiagaraFunctionSignature& IssueRayTrace = OutFunctions.AddDefaulted_GetRef();
+		IssueRayTrace.Name = NDICollisionQueryLocal::ReserveAsyncRayTraceName;
+		IssueRayTrace.bRequiresExecPin = true;
+		IssueRayTrace.bMemberFunction = true;
+		IssueRayTrace.bSupportsCPU = false;
 #if WITH_EDITORONLY_DATA
-		ReserveRayTrace.FunctionVersion = FNiagaraCollisionDIFunctionVersion::LatestVersion;
-		ReserveRayTrace.Description = LOCTEXT("ReserveAsync_RayTraceDescription", "Reserves a number of ray trace request slots");
+		IssueRayTrace.FunctionVersion = FNiagaraCollisionDIFunctionVersion::LatestVersion;
+		IssueRayTrace.Description = LOCTEXT("ReserveAsync_RayTraceDescription", "Reserves a number of ray trace request slots");
 #endif
-		ReserveRayTrace.AddInput(FNiagaraVariable(FNiagaraTypeDefinition(GetClass()), TEXT("CollisionQuery")));
-		ReserveRayTrace.AddInput(FNiagaraVariable(FNiagaraTypeDefinition::GetIntDef(), TEXT("TraceCount")), LOCTEXT("ReserveAsync_QueryIDDescription", "Number of async raytrace requests to be reserved"));
-		ReserveRayTrace.AddOutput(FNiagaraVariable(FNiagaraTypeDefinition::GetIntDef(), TEXT("FirstQueryID")), LOCTEXT("ReserveAsync_TraceChannelDescription", "The first index in the block reserved through this call"));
-		ReserveRayTrace.AddOutput(FNiagaraVariable(FNiagaraTypeDefinition::GetBoolDef(), TEXT("IsQueryValid")), LOCTEXT("ReserveAsync_IsQueryValidDescription", "Returns true if the requested indices were reserved"));
+		IssueRayTrace.AddInput(FNiagaraVariable(FNiagaraTypeDefinition(GetClass()), TEXT("CollisionQuery")));
+		IssueRayTrace.AddInput(FNiagaraVariable(FNiagaraTypeDefinition::GetIntDef(), TEXT("TraceCount")), LOCTEXT("ReserveAsync_QueryIDDescription", "Number of async raytrace requests to be reserved"));
+		IssueRayTrace.AddOutput(FNiagaraVariable(FNiagaraTypeDefinition::GetIntDef(), TEXT("FirstQueryID")), LOCTEXT("ReserveAsync_TraceChannelDescription", "The first index in the block reserved through this call"));
+		IssueRayTrace.AddOutput(FNiagaraVariable(FNiagaraTypeDefinition::GetBoolDef(), TEXT("IsQueryValid")), LOCTEXT("ReserveAsync_IsQueryValidDescription", "Returns true if the requested indices were reserved"));
 	}
 
 	{
@@ -325,7 +400,6 @@ void UNiagaraDataInterfaceCollisionQuery::GetFunctions(TArray<FNiagaraFunctionSi
 		ReadRayTrace.Name = NDICollisionQueryLocal::ReadAsyncRayTraceName;
 		ReadRayTrace.bMemberFunction = true;
 		ReadRayTrace.bSupportsCPU = false;
-		ReadRayTrace.bSoftDeprecatedFunction = true;
 #if WITH_EDITORONLY_DATA
 		ReadRayTrace.FunctionVersion = FNiagaraCollisionDIFunctionVersion::LatestVersion;
 		ReadRayTrace.Description = LOCTEXT("ReadAsync_RayTraceDescription", "Reads the results of a previously enqueued GPU ray trace");
@@ -337,7 +411,6 @@ void UNiagaraDataInterfaceCollisionQuery::GetFunctions(TArray<FNiagaraFunctionSi
 		ReadRayTrace.AddOutput(FNiagaraVariable(FNiagaraTypeDefinition::GetPositionDef(), TEXT("CollisionPosWorld")), LOCTEXT("ReadAsync_CollisionPosWorldDescription", "The point in world space where the intersection occured"));
 		ReadRayTrace.AddOutput(FNiagaraVariable(FNiagaraTypeDefinition::GetVec3Def(), TEXT("CollisionNormal")), LOCTEXT("ReadAsync_CollisionNormalDescription", "The surface normal of the world geometry at the point of intersection"));
 	}
-	// DEPRECATE_END
 
 	{
 		const FText TraceChannelDescription = LOCTEXT("TraceChannelDescription", "The trace channel to collide against. Trace channels can be configured in the project settings.");
@@ -411,12 +484,10 @@ bool UNiagaraDataInterfaceCollisionQuery::GetFunctionHLSL(const FNiagaraDataInte
 	if ( (FunctionInfo.DefinitionName == NDICollisionQueryLocal::SceneDepthName) ||
 		 (FunctionInfo.DefinitionName == NDICollisionQueryLocal::CustomDepthName) ||
 		 (FunctionInfo.DefinitionName == NDICollisionQueryLocal::DistanceFieldName) ||
-		// DEPRECATE_BEGIN
 		 (FunctionInfo.DefinitionName == NDICollisionQueryLocal::IssueAsyncRayTraceName) ||
 		 (FunctionInfo.DefinitionName == NDICollisionQueryLocal::CreateAsyncRayTraceName) ||
 		 (FunctionInfo.DefinitionName == NDICollisionQueryLocal::ReserveAsyncRayTraceName) ||
 		 (FunctionInfo.DefinitionName == NDICollisionQueryLocal::ReadAsyncRayTraceName) )
-		// DEPRECATE_END
 	{
 		return true;
 	}
@@ -464,6 +535,11 @@ void UNiagaraDataInterfaceCollisionQuery::ValidateFunction(const FNiagaraFunctio
 }
 #endif
 
+bool UNiagaraDataInterfaceCollisionQuery::RequiresRayTracingScene() const
+{
+	return IsRayTracingEnabled() && GEnableGPUHWRTCollisions && MaxTracesPerParticle > 0;
+}
+
 #if WITH_EDITORONLY_DATA
 void UNiagaraDataInterfaceCollisionQuery::GetParameterDefinitionHLSL(const FNiagaraDataInterfaceGPUParamInfo& ParamInfo, FString& OutHLSL)
 {
@@ -518,6 +594,13 @@ bool UNiagaraDataInterfaceCollisionQuery::AppendCompileHash(FNiagaraCompileHashV
 	return true;
 }
 
+void UNiagaraDataInterfaceCollisionQuery::ModifyCompilationEnvironment(EShaderPlatform ShaderPlatform, FShaderCompilerEnvironment& OutEnvironment) const
+{
+	Super::ModifyCompilationEnvironment(ShaderPlatform, OutEnvironment);
+
+	OutEnvironment.SetDefine(TEXT("NIAGARA_SUPPORTS_RAY_TRACING"), ShouldCompileRayTracingShadersForProject(ShaderPlatform) ? 1 : 0);
+}
+
 #endif
 
 void UNiagaraDataInterfaceCollisionQuery::PerformQuerySyncCPU(FVectorVMExternalFunctionContext & Context)
@@ -649,6 +732,55 @@ bool UNiagaraDataInterfaceCollisionQuery::PerInstanceTickPostSimulate(void* PerI
 	return false;
 }
 
+bool UNiagaraDataInterfaceCollisionQuery::Equals(const UNiagaraDataInterface* Other) const
+{
+	if (!Super::Equals(Other))
+	{
+		return false;
+	}
+
+	const UNiagaraDataInterfaceCollisionQuery* OtherTyped = CastChecked<const UNiagaraDataInterfaceCollisionQuery>(Other);
+	return OtherTyped->MaxTracesPerParticle == MaxTracesPerParticle;
+}
+
+bool UNiagaraDataInterfaceCollisionQuery::CopyToInternal(UNiagaraDataInterface* Destination) const
+{
+	if (!Super::CopyToInternal(Destination))
+	{
+		return false;
+	}
+
+	UNiagaraDataInterfaceCollisionQuery* OtherTyped = CastChecked<UNiagaraDataInterfaceCollisionQuery>(Destination);
+	OtherTyped->MaxTracesPerParticle = MaxTracesPerParticle;
+	OtherTyped->MaxRetraces = MaxRetraces;
+	OtherTyped->MarkRenderDataDirty();
+	return true;
+}
+
+void UNiagaraDataInterfaceCollisionQuery::PushToRenderThreadImpl()
+{
+	FNiagaraDataIntefaceProxyCollisionQuery* RT_Proxy = GetProxyAs<FNiagaraDataIntefaceProxyCollisionQuery>();
+
+	// Push Updates to Proxy, first release any resources
+	ENQUEUE_RENDER_COMMAND(FUpdateDI)(
+		[RT_Proxy, RT_MaxTracesPerParticle = MaxTracesPerParticle, RT_MaxRetraces = MaxRetraces](FRHICommandListImmediate& RHICmdList)
+		{
+			RT_Proxy->RenderThreadInitialize(RT_MaxTracesPerParticle, RT_MaxRetraces);
+		});
+}
+
+#if WITH_EDITOR
+void UNiagaraDataInterfaceCollisionQuery::PostEditChangeProperty(struct FPropertyChangedEvent& PropertyChangedEvent)
+{
+	Super::PostEditChangeProperty(PropertyChangedEvent);
+
+	if (PropertyChangedEvent.Property && PropertyChangedEvent.Property->GetFName() == GET_MEMBER_NAME_CHECKED(UNiagaraDataInterfaceCollisionQuery, MaxTracesPerParticle))
+	{
+		MarkRenderDataDirty();
+	}
+}
+#endif
+
 //////////////////////////////////////////////////////////////////////////
 
 struct FNiagaraDataInterfaceParametersCS_CollisionQuery : public FNiagaraDataInterfaceParametersCS
@@ -659,12 +791,23 @@ public:
 	{
 		GlobalDistanceFieldParameters.Bind(ParameterMap);
 		SystemLWCTileParam.Bind(ParameterMap, *(NDICollisionQueryLocal::SystemLWCTileName + ParameterInfo.DataInterfaceHLSLSymbol));
+#if RHI_RAYTRACING
+		RayTracingEnabledParam.Bind(ParameterMap, *(NDICollisionQueryLocal::RayTracingEnabledParamName + ParameterInfo.DataInterfaceHLSLSymbol));
+		MaxRayTraceCountParam.Bind(ParameterMap, *(NDICollisionQueryLocal::MaxRayTraceCountParamName + ParameterInfo.DataInterfaceHLSLSymbol));
+		RayRequestsParam.Bind(ParameterMap, *(NDICollisionQueryLocal::RayRequestsParamName + ParameterInfo.DataInterfaceHLSLSymbol));
+		RayRequestOffsetParam.Bind(ParameterMap, *(NDICollisionQueryLocal::RayRequestsOffsetParamName + ParameterInfo.DataInterfaceHLSLSymbol));
+		IntersectionResultsParam.Bind(ParameterMap, *(NDICollisionQueryLocal::IntersectionResultsParamName + ParameterInfo.DataInterfaceHLSLSymbol));
+		IntersectionResultOffsetParam.Bind(ParameterMap, *(NDICollisionQueryLocal::IntersectionResultsOffsetParamName + ParameterInfo.DataInterfaceHLSLSymbol));
+		RayTraceCountsParam.Bind(ParameterMap, *(NDICollisionQueryLocal::RayTraceCountsParamName + ParameterInfo.DataInterfaceHLSLSymbol));
+		RayTraceCountsOffsetParam.Bind(ParameterMap, *(NDICollisionQueryLocal::RayTraceCountsOffsetParamName + ParameterInfo.DataInterfaceHLSLSymbol));
+#endif
 	}
 
 	void Set(FRHICommandList& RHICmdList, const FNiagaraDataInterfaceSetArgs& Context) const
 	{
 		check(IsInRenderingThread());
 
+		FNiagaraDataIntefaceProxyCollisionQuery* QueryDI = (FNiagaraDataIntefaceProxyCollisionQuery*)Context.DataInterface;
 		FRHIComputeShader* ComputeShaderRHI = Context.Shader.GetComputeShader();
 		SetShaderValue(RHICmdList, ComputeShaderRHI, SystemLWCTileParam, Context.SystemLWCTile);
 		
@@ -682,12 +825,91 @@ public:
 			{
 				GlobalDistanceFieldParameters.Set(RHICmdList, ComputeShaderRHI, FGlobalDistanceFieldParameterData());
 			}
+			
 		}
+
+#if RHI_RAYTRACING
+		const bool HasRayTracingParametersBound = RayRequestsParam.IsUAVBound()
+			|| IntersectionResultsParam.IsBound()
+			|| RayTraceCountsParam.IsBound();
+
+		if ((IsRayTracingEnabled() && GEnableGPUHWRTCollisions) || HasRayTracingParametersBound)
+		{
+			FNiagaraRayTracingHelper& RTHelper = Context.ComputeDispatchInterface->GetRayTracingHelper();
+			const FNiagaraRayTraceDispatchInfo* DispatchInfo = nullptr;
+			if (IsRayTracingEnabled() && GEnableGPUHWRTCollisions && QueryDI->MaxTracesPerParticle > 0)
+			{
+				DispatchInfo = &RTHelper.GetDispatch(QueryDI);
+			}
+			else
+			{
+				DispatchInfo = &RTHelper.GetDummyDispatch();
+			}
+
+			SetShaderValue(RHICmdList, ComputeShaderRHI, RayTracingEnabledParam, IsRayTracingEnabled() && GEnableGPUHWRTCollisions ? 1 : 0);
+			SetShaderValue(RHICmdList, ComputeShaderRHI, MaxRayTraceCountParam, DispatchInfo->MaxRays);
+
+			if (RayRequestsParam.IsUAVBound())
+			{
+				check(DispatchInfo->RayRequests.IsValid());
+				RHICmdList.SetUAVParameter(ComputeShaderRHI, RayRequestsParam.GetUAVIndex(), DispatchInfo->RayRequests.Buffer->UAV);
+				SetShaderValue(RHICmdList, ComputeShaderRHI, RayRequestOffsetParam, DispatchInfo->RayRequests.Offset);
+			}
+
+			if (IntersectionResultsParam.IsBound())
+			{
+				check(DispatchInfo->LastFrameRayTraceIntersections.IsValid());
+
+				SetSRVParameter(RHICmdList, ComputeShaderRHI, IntersectionResultsParam, DispatchInfo->LastFrameRayTraceIntersections.Buffer->SRV);
+				SetShaderValue(RHICmdList, ComputeShaderRHI, IntersectionResultOffsetParam, DispatchInfo->LastFrameRayTraceIntersections.Offset);
+			}
+
+			if (RayTraceCountsParam.IsUAVBound())
+			{
+				check(DispatchInfo->RayCounts.IsValid());
+
+				RHICmdList.SetUAVParameter(ComputeShaderRHI, RayTraceCountsParam.GetUAVIndex(), DispatchInfo->RayCounts.Buffer->UAV);
+				SetShaderValue(RHICmdList, ComputeShaderRHI, RayTraceCountsOffsetParam, DispatchInfo->RayCounts.Offset);
+			}
+		}
+		else
+		{
+			SetShaderValue(RHICmdList, ComputeShaderRHI, RayTracingEnabledParam, 0);
+			SetShaderValue(RHICmdList, ComputeShaderRHI, MaxRayTraceCountParam, 0);
+		}
+#endif
 	}
 
+#if RHI_RAYTRACING
+	void Unset(FRHICommandList& RHICmdList, const FNiagaraDataInterfaceSetArgs& Context) const
+	{
+		FRHIComputeShader* ComputeShaderRHI = Context.Shader.GetComputeShader();
+		if (RayRequestsParam.IsUAVBound())
+		{
+			RayRequestsParam.UnsetUAV(RHICmdList, ComputeShaderRHI);
+		}
+
+		if (RayTraceCountsParam.IsUAVBound())
+		{
+			RayTraceCountsParam.UnsetUAV(RHICmdList, ComputeShaderRHI);
+		}
+	}
+#endif
+
 private:
 	LAYOUT_FIELD(FGlobalDistanceFieldParameters, GlobalDistanceFieldParameters);
 	LAYOUT_FIELD(FShaderParameter, SystemLWCTileParam);
+
+#if RHI_RAYTRACING
+	LAYOUT_FIELD(FShaderParameter, RayTracingEnabledParam);
+	LAYOUT_FIELD(FShaderParameter, MaxRayTraceCountParam);
+	LAYOUT_FIELD(FRWShaderParameter, RayRequestsParam);
+	LAYOUT_FIELD(FShaderParameter, RayRequestOffsetParam);
+	LAYOUT_FIELD(FShaderResourceParameter, IntersectionResultsParam);
+	LAYOUT_FIELD(FShaderParameter, IntersectionResultOffsetParam);
+	LAYOUT_FIELD(FRWShaderParameter, RayTraceCountsParam);
+	LAYOUT_FIELD(FShaderParameter, RayTraceCountsOffsetParam);
+#endif
 };
 
 IMPLEMENT_TYPE_LAYOUT(FNiagaraDataInterfaceParametersCS_CollisionQuery);
diff --git a/Engine/Plugins/FX/Niagara/Source/Niagara/Private/NiagaraFunctionLibrary.cpp b/Engine/Plugins/FX/Niagara/Source/Niagara/Private/NiagaraFunctionLibrary.cpp
--- a/Engine/Plugins/FX/Niagara/Source/Niagara/Private/NiagaraFunctionLibrary.cpp
+++ b/Engine/Plugins/FX/Niagara/Source/Niagara/Private/NiagaraFunctionLibrary.cpp
@@ -10,7 +10,6 @@
 #include "ContentStreaming.h"
 #include "Internationalization/Internationalization.h"
 
-#include "NiagaraAsyncGpuTraceHelper.h"
 #include "NiagaraWorldManager.h"
 #include "DataInterface/NiagaraDataInterfaceStaticMesh.h"
 #include "NiagaraStats.h"
@@ -1845,12 +1844,12 @@ void UNiagaraFunctionLibrary::InitVectorVMFastPathOps()
 
 void UNiagaraFunctionLibrary::SetComponentNiagaraGPURayTracedCollisionGroup(UObject* WorldContextObject, UPrimitiveComponent* Primitive, int32 CollisionGroup)
 {
-#if NIAGARA_ASYNC_GPU_TRACE_COLLISION_GROUPS
+#if RHI_RAYTRACING
 	if ( UWorld* World = GEngine->GetWorldFromContextObject(WorldContextObject, EGetWorldErrorMode::LogAndReturnNull) )
 	{
 		if ( FNiagaraGpuComputeDispatchInterface* ComputeDispatchInterface = FNiagaraGpuComputeDispatchInterface::Get(World) )
 		{
-			ComputeDispatchInterface->GetAsyncGpuTraceHelper().SetPrimitiveRayTracingCollisionGroup_GT(Primitive, CollisionGroup);
+			ComputeDispatchInterface->SetPrimitiveRayTracingCollisionGroup(Primitive, CollisionGroup);
 		}
 	}
 #endif
@@ -1858,7 +1857,7 @@ void UNiagaraFunctionLibrary::SetComponentNiagaraGPURayTracedCollisionGroup(UObj
 
 void UNiagaraFunctionLibrary::SetActorNiagaraGPURayTracedCollisionGroup(UObject* WorldContextObject, AActor* Actor, int32 CollisionGroup)
 {
-#if NIAGARA_ASYNC_GPU_TRACE_COLLISION_GROUPS
+#if RHI_RAYTRACING
 	if (Actor == nullptr)
 	{
 		return;
@@ -1868,13 +1867,11 @@ void UNiagaraFunctionLibrary::SetActorNiagaraGPURayTracedCollisionGroup(UObject*
 	{
 		if (FNiagaraGpuComputeDispatchInterface* ComputeDispatchInterface = FNiagaraGpuComputeDispatchInterface::Get(World))
 		{
-			FNiagaraAsyncGpuTraceHelper& Helper = ComputeDispatchInterface->GetAsyncGpuTraceHelper();
-
 			Actor->ForEachComponent<UPrimitiveComponent>(
 				/*bIncludeFromChildActors*/true,
-				[&Helper, &CollisionGroup](UPrimitiveComponent* PrimitiveComponent)
+				[&ComputeDispatchInterface, &CollisionGroup](UPrimitiveComponent* PrimitiveComponent)
 				{
-					Helper.SetPrimitiveRayTracingCollisionGroup_GT(PrimitiveComponent, CollisionGroup);
+					ComputeDispatchInterface->SetPrimitiveRayTracingCollisionGroup(PrimitiveComponent, CollisionGroup);
 				});
 			}
 	}
@@ -1883,12 +1880,12 @@ void UNiagaraFunctionLibrary::SetActorNiagaraGPURayTracedCollisionGroup(UObject*
 
 int32 UNiagaraFunctionLibrary::AcquireNiagaraGPURayTracedCollisionGroup(UObject* WorldContextObject)
 {
-#if NIAGARA_ASYNC_GPU_TRACE_COLLISION_GROUPS
+#if RHI_RAYTRACING
 	if (UWorld* World = GEngine->GetWorldFromContextObject(WorldContextObject, EGetWorldErrorMode::LogAndReturnNull))
 	{
 		if (FNiagaraGpuComputeDispatchInterface* ComputeDispatchInterface = FNiagaraGpuComputeDispatchInterface::Get(World))
 		{
-			return ComputeDispatchInterface->GetAsyncGpuTraceHelper().AcquireGPURayTracedCollisionGroup_GT();
+			return ComputeDispatchInterface->AcquireGPURayTracedCollisionGroup();
 		}
 	}
 #endif
@@ -1897,12 +1894,12 @@ int32 UNiagaraFunctionLibrary::AcquireNiagaraGPURayTracedCollisionGroup(UObject*
 
 void UNiagaraFunctionLibrary::ReleaseNiagaraGPURayTracedCollisionGroup(UObject* WorldContextObject, int32 CollisionGroup)
 {
-#if NIAGARA_ASYNC_GPU_TRACE_COLLISION_GROUPS
+#if RHI_RAYTRACING
 	if (UWorld* World = GEngine->GetWorldFromContextObject(WorldContextObject, EGetWorldErrorMode::LogAndReturnNull))
 	{
 		if (FNiagaraGpuComputeDispatchInterface* ComputeDispatchInterface = FNiagaraGpuComputeDispatchInterface::Get(World))
 		{
-			ComputeDispatchInterface->GetAsyncGpuTraceHelper().ReleaseGPURayTracedCollisionGroup_GT(CollisionGroup);
+			ComputeDispatchInterface->ReleaseGPURayTracedCollisionGroup(CollisionGroup);
 		}
 	}
 #endif
diff --git a/Engine/Plugins/FX/Niagara/Source/Niagara/Private/NiagaraGpuComputeDispatch.cpp b/Engine/Plugins/FX/Niagara/Source/Niagara/Private/NiagaraGpuComputeDispatch.cpp
--- a/Engine/Plugins/FX/Niagara/Source/Niagara/Private/NiagaraGpuComputeDispatch.cpp
+++ b/Engine/Plugins/FX/Niagara/Source/Niagara/Private/NiagaraGpuComputeDispatch.cpp
@@ -10,13 +10,13 @@
 #include "Engine/World.h"
 #include "Misc/ScopeExit.h"
 
-#include "NiagaraAsyncGpuTraceHelper.h"
 #include "NiagaraDataInterfaceRW.h"
 #if NIAGARA_COMPUTEDEBUG_ENABLED
 #include "NiagaraGpuComputeDebug.h"
 #endif
 #include "NiagaraGPUProfilerInterface.h"
 #include "NiagaraGpuReadbackManager.h"
+#include "NiagaraRayTracingHelper.h"
 #include "NiagaraRenderer.h"
 #include "NiagaraScriptExecutionContext.h"
 #include "NiagaraSystemGpuComputeProxy.h"
@@ -207,8 +207,9 @@ FNiagaraGpuComputeDispatch::FNiagaraGpuComputeDispatch(ERHIFeatureLevel::Type In
 	OwnerCBufferLayout   = new FNiagaraRHIUniformBufferLayout(TEXT("Niagara GPU Owner CBuffer"),   sizeof(FNiagaraOwnerParameters));
 	EmitterCBufferLayout = new FNiagaraRHIUniformBufferLayout(TEXT("Niagara GPU Emitter CBuffer"), sizeof(FNiagaraEmitterParameters));
 
-
-	AsyncGpuTraceHelper.Reset(new FNiagaraAsyncGpuTraceHelper(InShaderPlatform, InFeatureLevel, this));
+#if RHI_RAYTRACING
+	RayTracingHelper.Reset(new FNiagaraRayTracingHelper(InShaderPlatform));
+#endif
 
 #if NIAGARA_COMPUTEDEBUG_ENABLED
 	GpuComputeDebugPtr.Reset(new FNiagaraGpuComputeDebug(FeatureLevel));
@@ -218,13 +219,18 @@ FNiagaraGpuComputeDispatch::FNiagaraGpuComputeDispatch(ERHIFeatureLevel::Type In
 #endif
 	GpuReadbackManagerPtr.Reset(new FNiagaraGpuReadbackManager());
 	EmptyUAVPoolPtr.Reset(new FNiagaraEmptyUAVPool());
+
+	FPrimitiveSceneInfo::OnGPUSceneInstancesAllocated.AddRaw(this, &FNiagaraGpuComputeDispatch::OnPrimitiveGPUSceneInstancesAllocated);
+	FPrimitiveSceneInfo::OnGPUSceneInstancesAllocated.AddRaw(this, &FNiagaraGpuComputeDispatch::OnPrimitiveGPUSceneInstancesFreed);
 }
 
 FNiagaraGpuComputeDispatch::~FNiagaraGpuComputeDispatch()
 {
 	FinishDispatches();
 
-	AsyncGpuTraceHelper->Reset();
+#if RHI_RAYTRACING
+	RayTracingHelper->Reset();
+#endif
 
 	GlobalCBufferLayout = nullptr;
 	SystemCBufferLayout = nullptr;
@@ -506,8 +512,10 @@ void FNiagaraGpuComputeDispatch::ProcessPendingTicksFlush(FRHICommandListImmedia
 			//UE_LOG(LogNiagara, Log, TEXT("FNiagaraGpuComputeDispatch: Queued ticks are being Destroyed due to not rendering.  This may result in undesirable simulation artifacts."));
 
 			FinishDispatches();
-			AsyncGpuTraceHelper->Reset();
 
+#if RHI_RAYTRACING
+			RayTracingHelper->Reset();
+#endif
 			break;
 		}
 	}
@@ -1669,7 +1677,9 @@ void FNiagaraGpuComputeDispatch::PreInitViews(FRDGBuilder& GraphBuilder, bool bA
 			CalculateCrossGPUTransferLocation();
 		#endif
 
-			AsyncGpuTraceHelper->BeginFrame(GraphBuilder.RHICmdList, this);
+#if RHI_RAYTRACING
+			RayTracingHelper->BeginFrame(GraphBuilder.RHICmdList, HasRayTracingScene());
+#endif
 
 			if ( FNiagaraGpuComputeDispatchLocal::GDebugLogging)
 			{
@@ -1756,7 +1766,9 @@ void FNiagaraGpuComputeDispatch::PostRenderOpaque(FRDGBuilder& GraphBuilder, TCo
 			NiagaraSceneTextures = PassParameters;
 			ON_SCOPE_EXIT{ NiagaraSceneTextures = nullptr; };
 
-			AsyncGpuTraceHelper->PostRenderOpaque(RHICmdList, this, Views);
+	#if RHI_RAYTRACING
+			BuildRayTracingSceneInfo(RHICmdList, Views);
+	#endif
 
 			CurrentPassViews = Views;
 
@@ -1767,7 +1779,9 @@ void FNiagaraGpuComputeDispatch::PostRenderOpaque(FRDGBuilder& GraphBuilder, TCo
 
 			FinishDispatches();
 
-			AsyncGpuTraceHelper->EndFrame(RHICmdList, this);
+	#if RHI_RAYTRACING
+			RayTracingHelper->EndFrame(RHICmdList, HasRayTracingScene(), GetScene());
+	#endif
 
 			// Clear CurrentPassViews
 			CurrentPassViews = TConstArrayView<FViewInfo>();
@@ -2023,10 +2037,86 @@ void FNiagaraGpuComputeDispatch::GenerateSortKeys(FRHICommandListImmediate& RHIC
 	RHICmdList.EndUAVOverlap(MakeArrayView(OverlapUAVs, NumOverlapUAVs));
 }
 
-FNiagaraAsyncGpuTraceHelper& FNiagaraGpuComputeDispatch::GetAsyncGpuTraceHelper() const
+#if RHI_RAYTRACING
+void FNiagaraGpuComputeDispatch::BuildRayTracingSceneInfo(FRHICommandList& RHICmdList, TConstArrayView<FViewInfo> Views)
+{
+	if (NumProxiesThatRequireRayTracingScene > 0)
+	{
+		RayTracingHelper->BuildRayTracingSceneInfo(RHICmdList, Views);
+
+		RayTracingHelper->UpdateCollisionGroupMap(RHICmdList, GetScene(), FeatureLevel);
+	}
+	else
+	{
+		RayTracingHelper->Reset();
+	}
+}
+
+void FNiagaraGpuComputeDispatch::ResetRayTracingSceneInfo()
+{
+	RayTracingHelper->Reset();
+}
+
+
+FNiagaraRayTracingHelper& FNiagaraGpuComputeDispatch::GetRayTracingHelper()const
+{
+	check(RayTracingHelper.IsValid());
+	return *RayTracingHelper.Get();
+}
+
+bool FNiagaraGpuComputeDispatch::HasRayTracingScene() const
+{
+	return NumProxiesThatRequireRayTracingScene && RayTracingHelper->IsValid();
+}
+
+void FNiagaraGpuComputeDispatch::SetPrimitiveRayTracingCollisionGroup(UPrimitiveComponent* Primitive, uint32 Group)
 {
-	check(AsyncGpuTraceHelper.IsValid());
-	return *AsyncGpuTraceHelper.Get();
+	check(IsInGameThread());
+	if (Primitive == nullptr || Primitive->SceneProxy == nullptr)
+	{
+		return;
+	}
+
+	ENQUEUE_RENDER_COMMAND(NiagaraSetPrimHWRTCollisionGroup)(
+		[RT_NiagaraBatcher = this, Proxy = Primitive->SceneProxy, CollisionGroup=Group](FRHICommandListImmediate& RHICmdList)
+		{
+			check(Proxy);
+			RT_NiagaraBatcher->RayTracingHelper->SetPrimitiveCollisionGroup(*Proxy->GetPrimitiveSceneInfo(), CollisionGroup);
+		}
+	);
+}
+
+int32 FNiagaraGpuComputeDispatch::AcquireGPURayTracedCollisionGroup()
+{
+	check(IsInGameThread());//one of the few batcher functions that should be called from the GT
+	if (FreeGPURayTracedCollisionGroups.Num() > 0)
+	{
+		return FreeGPURayTracedCollisionGroups.Pop();
+	}
+
+	return NumGPURayTracedCollisionGroups++;
+}
+
+void FNiagaraGpuComputeDispatch::ReleaseGPURayTracedCollisionGroup(int32 CollisionGroup)
+{
+	check(IsInGameThread());//one of the few batcher functions that should be called from the GT
+	FreeGPURayTracedCollisionGroups.Add(CollisionGroup);
+}
+
+#endif
+
+void FNiagaraGpuComputeDispatch::OnPrimitiveGPUSceneInstancesAllocated()
+{
+#if RHI_RAYTRACING
+	RayTracingHelper->RefreshPrimitiveInstanceData();
+#endif
+}
+
+void FNiagaraGpuComputeDispatch::OnPrimitiveGPUSceneInstancesFreed()
+{
+#if RHI_RAYTRACING
+	RayTracingHelper->RefreshPrimitiveInstanceData();
+#endif
 }
 
 /* Set shader parameters for data interfaces
diff --git a/Engine/Plugins/FX/Niagara/Source/Niagara/Private/NiagaraGpuComputeDispatch.h b/Engine/Plugins/FX/Niagara/Source/Niagara/Private/NiagaraGpuComputeDispatch.h
--- a/Engine/Plugins/FX/Niagara/Source/Niagara/Private/NiagaraGpuComputeDispatch.h
+++ b/Engine/Plugins/FX/Niagara/Source/Niagara/Private/NiagaraGpuComputeDispatch.h
@@ -27,8 +27,8 @@ the same VectorVM byte code / compute shader code
 #include "Tickable.h"
 
 class FGPUSortManager;
+class FNiagaraRayTracingHelper;
 class FDistanceFieldSceneData;
-class FNiagaraAsyncGpuTraceHelper;
 
 class FNiagaraGpuComputeDispatch : public FNiagaraGpuComputeDispatchInterface
 {
@@ -125,12 +125,23 @@ public:
 	/** Debug only function to readback data. */
 	virtual void AddDebugReadback(FNiagaraSystemInstanceID InstanceID, TSharedPtr<struct FNiagaraScriptDebuggerInfo, ESPMode::ThreadSafe> DebugInfo, FNiagaraComputeExecutionContext* Context) override;
 
+#if RHI_RAYTRACING
+	virtual FNiagaraRayTracingHelper& GetRayTracingHelper() const override;
+	virtual bool HasRayTracingScene() const override;
+
+	virtual void SetPrimitiveRayTracingCollisionGroup(UPrimitiveComponent* PrimitiveComponent, uint32 Group) override;
+
+	virtual int32 AcquireGPURayTracedCollisionGroup() override;
+	virtual void ReleaseGPURayTracedCollisionGroup(int32 CollisionGroup) override;
+#endif
+
 #if WITH_MGPU
 	virtual void MultiGPUResourceModified(FRHICommandList& RHICmdList, FRHIBuffer* Buffer, bool bRequiredForSimulation, bool bRequiredForRendering) const override;
 	virtual void MultiGPUResourceModified(FRHICommandList& RHICmdList, FRHITexture* Texture, bool bRequiredForSimulation, bool bRequiredForRendering) const override;
 #endif
 
-	virtual FNiagaraAsyncGpuTraceHelper& GetAsyncGpuTraceHelper() const override;
+	void OnPrimitiveGPUSceneInstancesAllocated();
+	void OnPrimitiveGPUSceneInstancesFreed();
 
 private:
 	void DumpDebugFrame();
@@ -161,6 +172,11 @@ private:
 	 */
 	void GenerateSortKeys(FRHICommandListImmediate& RHICmdList, int32 BatchId, int32 NumElementsInBatch, EGPUSortFlags Flags, FRHIUnorderedAccessView* KeysUAV, FRHIUnorderedAccessView* ValuesUAV);
 
+#if RHI_RAYTRACING
+	void BuildRayTracingSceneInfo(FRHICommandList& RHICmdList, TConstArrayView<FViewInfo> Views);
+	void ResetRayTracingSceneInfo();
+#endif
+
 	void FinishDispatches();
 
 	void UpdateFreeIDsListSizesBuffer(FRHICommandList& RHICmdList, uint32 NumInstances);
@@ -174,7 +190,10 @@ private:
 	/** All sort tasks registered in AddSortedGPUSimulation(). Holds all the data required in GenerateSortKeys(). */
 	TArray<FNiagaraGPUSortInfo> SimulationsToSort;
 
-	TUniquePtr<FNiagaraAsyncGpuTraceHelper> AsyncGpuTraceHelper;
+#if RHI_RAYTRACING
+	// helper object which is valid during PostRenderOpaque if DataInterfaces are present that require a RayTrace scene
+	TUniquePtr<FNiagaraRayTracingHelper> RayTracingHelper;
+#endif
 
 #if WITH_NIAGARA_GPU_PROFILER
 	TUniquePtr<FNiagaraGPUProfiler> GPUProfilerPtr;
@@ -235,6 +254,10 @@ private:
 	void WaitForMultiGPUBuffers(FRHICommandList& RHICmdList, ENiagaraGpuComputeTickStage::Type TickStage);
 #endif // WITH_MGPU
 
+	/** Pool of free GPU ray traced collision groups. */
+	TArray<int32> FreeGPURayTracedCollisionGroups;
+	int32 NumGPURayTracedCollisionGroups = 0;
+
 	// Cached information to build a dummy view info if necessary
 	FIntRect CachedViewRect;
 
diff --git a/Engine/Plugins/FX/Niagara/Source/Niagara/Private/NiagaraGpuScratchPad.h b/Engine/Plugins/FX/Niagara/Source/Niagara/Private/NiagaraGpuScratchPad.h
deleted file mode 100644
diff --git a/Engine/Plugins/FX/Niagara/Source/Niagara/Private/NiagaraSettings.cpp b/Engine/Plugins/FX/Niagara/Source/Niagara/Private/NiagaraSettings.cpp
--- a/Engine/Plugins/FX/Niagara/Source/Niagara/Private/NiagaraSettings.cpp
+++ b/Engine/Plugins/FX/Niagara/Source/Niagara/Private/NiagaraSettings.cpp
@@ -9,9 +9,6 @@ UNiagaraSettings::UNiagaraSettings(const FObjectInitializer& ObjectInitlaizer)
 	, NDISkelMesh_AdjacencyTriangleIndexFormat(ENDISkelMesh_AdjacencyTriangleIndexFormat::Full)
 {
 	PositionPinTypeColor = FLinearColor(1.0f, 0.3f, 1.0f, 1.0f);
-
-	NDICollisionQuery_AsyncGpuTraceProviderOrder.Add(ENDICollisionQuery_AsyncGpuTraceProvider::Type::HWRT);
-	NDICollisionQuery_AsyncGpuTraceProviderOrder.Add(ENDICollisionQuery_AsyncGpuTraceProvider::Type::GSDF);
 }
 
 FName UNiagaraSettings::GetCategoryName() const
diff --git a/Engine/Plugins/FX/Niagara/Source/Niagara/Public/NiagaraGpuComputeDispatchInterface.h b/Engine/Plugins/FX/Niagara/Source/Niagara/Public/NiagaraGpuComputeDispatchInterface.h
--- a/Engine/Plugins/FX/Niagara/Source/Niagara/Public/NiagaraGpuComputeDispatchInterface.h
+++ b/Engine/Plugins/FX/Niagara/Source/Niagara/Public/NiagaraGpuComputeDispatchInterface.h
@@ -7,7 +7,6 @@
 #include "NiagaraGPUInstanceCountManager.h"
 #include "FXSystem.h"
 
-class FNiagaraAsyncGpuTraceHelper;
 struct FNiagaraComputeExecutionContext;
 class FNiagaraGpuComputeDebug;
 class FNiagaraGPUInstanceCountManager;
@@ -83,7 +82,15 @@ public:
 	/** Processes all pending debug readbacks */
 	virtual void ProcessDebugReadbacks(FRHICommandListImmediate& RHICmdList, bool bWaitCompletion) = 0;
 
-	virtual FNiagaraAsyncGpuTraceHelper& GetAsyncGpuTraceHelper() const = 0;
+#if RHI_RAYTRACING
+	virtual FNiagaraRayTracingHelper& GetRayTracingHelper() const = 0;
+	virtual bool HasRayTracingScene() const = 0;
+
+	virtual void SetPrimitiveRayTracingCollisionGroup(UPrimitiveComponent* PrimitiveComponent, uint32 Group) = 0;
+
+	virtual int32 AcquireGPURayTracedCollisionGroup() = 0;
+	virtual void ReleaseGPURayTracedCollisionGroup(int32 CollisionGroup) = 0;
+#endif
 
 #if WITH_MGPU
 	/**
diff --git a/Engine/Plugins/FX/Niagara/Source/Niagara/Public/NiagaraSettings.h b/Engine/Plugins/FX/Niagara/Source/Niagara/Public/NiagaraSettings.h
--- a/Engine/Plugins/FX/Niagara/Source/Niagara/Public/NiagaraSettings.h
+++ b/Engine/Plugins/FX/Niagara/Source/Niagara/Public/NiagaraSettings.h
@@ -61,18 +61,6 @@ enum class ENiagaraDefaultRendererPixelCoverageMode : uint8
 	Disabled,
 };
 
-UENUM()
-namespace ENDICollisionQuery_AsyncGpuTraceProvider
-{
-	enum Type
-	{
-		Default = 0 UMETA(DisplayName = "Project Default"),
-		HWRT = 1 UMETA(DisplayName = "HW Ray Tracing"),
-		GSDF = 2 UMETA(DisplayName = "Global Signed Distance Fields"),
-		None = 3 UMETA(DisplayName = "Disabled"),
-	};
-}
-
 UCLASS(config = Niagara, defaultconfig, meta=(DisplayName="Niagara"))
 class NIAGARA_API UNiagaraSettings : public UDeveloperSettings
 {
@@ -146,13 +134,6 @@ class NIAGARA_API UNiagaraSettings : public UDeveloperSettings
 	UPROPERTY(config, EditAnywhere, Category = StaticMeshDI, meta = (DisplayName = "Allow Distance Fields (Experimental)", ConfigRestartRequired = true))
 	bool NDIStaticMesh_AllowDistanceFields = false;
 
-	/** 
-	Defines how traces tagged as 'Project Default' will be interpreted when using the AsyncGpuTrace data interface.
-	The system will go through (starting at element 0) to find the first provider that is available.
-	*/
-	UPROPERTY(config, EditAnywhere, Category = AsyncGpuTraceDI, meta = (DisplayName = "Trace Provider Priorities (Experimental)", ConfigRestartRequired = true))
-	TArray<TEnumAsByte<ENDICollisionQuery_AsyncGpuTraceProvider::Type>> NDICollisionQuery_AsyncGpuTraceProviderOrder;
-
 	// Begin UDeveloperSettings Interface
 	virtual FName GetCategoryName() const override;
 #if WITH_EDITOR
diff --git a/Engine/Plugins/FX/Niagara/Source/Niagara/Public/NiagaraSystemRenderData.h b/Engine/Plugins/FX/Niagara/Source/Niagara/Public/NiagaraSystemRenderData.h
--- a/Engine/Plugins/FX/Niagara/Source/Niagara/Public/NiagaraSystemRenderData.h
+++ b/Engine/Plugins/FX/Niagara/Source/Niagara/Public/NiagaraSystemRenderData.h
@@ -7,10 +7,6 @@
 #include "PrimitiveUniformShaderParameters.h"
 #include "PrimitiveViewRelevance.h"
 
-#if RHI_RAYTRACING
-#include "RayTracingInstance.h"
-#endif
-
 class FNiagaraSystemInstance;
 class FNiagaraSystemRenderData;
 class FMeshElementCollector;
diff --git a/Engine/Plugins/FX/Niagara/Source/NiagaraShader/Private/NiagaraAsyncGpuTraceProvider.cpp b/Engine/Plugins/FX/Niagara/Source/NiagaraShader/Private/NiagaraAsyncGpuTraceProvider.cpp
deleted file mode 100644
diff --git a/Engine/Plugins/FX/Niagara/Source/NiagaraShader/Private/NiagaraAsyncGpuTraceProviderGsdf.cpp b/Engine/Plugins/FX/Niagara/Source/NiagaraShader/Private/NiagaraAsyncGpuTraceProviderGsdf.cpp
deleted file mode 100644
diff --git a/Engine/Plugins/FX/Niagara/Source/NiagaraShader/Private/NiagaraAsyncGpuTraceProviderGsdf.h b/Engine/Plugins/FX/Niagara/Source/NiagaraShader/Private/NiagaraAsyncGpuTraceProviderGsdf.h
deleted file mode 100644
diff --git a/Engine/Plugins/FX/Niagara/Source/NiagaraShader/Private/NiagaraAsyncGpuTraceProviderHwrt.cpp b/Engine/Plugins/FX/Niagara/Source/NiagaraShader/Private/NiagaraAsyncGpuTraceProviderHwrt.cpp
deleted file mode 100644
diff --git a/Engine/Plugins/FX/Niagara/Source/NiagaraShader/Private/NiagaraAsyncGpuTraceProviderHwrt.h b/Engine/Plugins/FX/Niagara/Source/NiagaraShader/Private/NiagaraAsyncGpuTraceProviderHwrt.h
deleted file mode 100644
diff --git a/Engine/Plugins/FX/Niagara/Source/NiagaraShader/Private/NiagaraRayTracingHelper.cpp b/Engine/Plugins/FX/Niagara/Source/NiagaraShader/Private/NiagaraRayTracingHelper.cpp
new file mode 100644
--- /dev/null
+++ b/Engine/Plugins/FX/Niagara/Source/NiagaraShader/Private/NiagaraRayTracingHelper.cpp
@@ -0,0 +1,638 @@
+// Copyright Epic Games, Inc. All Rights Reserved.
+
+#include "NiagaraRayTracingHelper.h"
+#include "NiagaraShader.h"
+
+#if RHI_RAYTRACING
+
+#include "GlobalShader.h"
+#include "MeshPassProcessor.h"
+#include "Renderer/Private/RayTracing/RayTracingScene.h"
+#include "Renderer/Private/ScenePrivate.h"
+#include "Renderer/Private/SceneRendering.h"
+#include "ShaderParameterStruct.h"
+#include "NiagaraShaderParticleID.h"
+
+//////////////////////////////////////////////////////////////////////////
+
+class FNiagaraUpdateCollisionGroupMapCS : public FGlobalShader
+{
+	DECLARE_GLOBAL_SHADER(FNiagaraUpdateCollisionGroupMapCS);
+	SHADER_USE_PARAMETER_STRUCT(FNiagaraUpdateCollisionGroupMapCS, FGlobalShader);
+
+	BEGIN_SHADER_PARAMETER_STRUCT(FParameters, )
+		SHADER_PARAMETER_UAV(Buffer<UINT>, RWHashTable)		
+		SHADER_PARAMETER(uint32, HashTableSize)
+		SHADER_PARAMETER_UAV(Buffer<UINT>, RWHashToCollisionGroups)
+		SHADER_PARAMETER_SRV(Buffer<UINT>, NewPrimIdCollisionGroupPairs)
+		SHADER_PARAMETER(uint32, NumNewPrims)
+	END_SHADER_PARAMETER_STRUCT()
+public:
+
+	static bool ShouldCompilePermutation(const FGlobalShaderPermutationParameters& Parameters)
+	{
+		return RHISupportsComputeShaders(Parameters.Platform);
+	}
+
+	static void ModifyCompilationEnvironment(const FGlobalShaderPermutationParameters& Parameters, FShaderCompilerEnvironment& OutEnvironment)
+	{
+		FGlobalShader::ModifyCompilationEnvironment(Parameters, OutEnvironment);
+		OutEnvironment.SetDefine(TEXT("THREAD_COUNT"), THREAD_COUNT);
+	}
+
+	static constexpr uint32 THREAD_COUNT = 64;
+};
+IMPLEMENT_GLOBAL_SHADER(FNiagaraUpdateCollisionGroupMapCS, "/Plugin/FX/Niagara/Private/NiagaraRayTraceCollisionGroupShaders.usf", "UpdatePrimIdToCollisionGroupMap", SF_Compute);
+
+//////////////////////////////////////////////////////////////////////////
+
+/// TODO
+///  -get geometry masking working when an environmental mask is implemented
+
+BEGIN_SHADER_PARAMETER_STRUCT(FGPUSceneParameters, )
+SHADER_PARAMETER_SRV(StructuredBuffer<float4>, GPUSceneInstanceSceneData)
+SHADER_PARAMETER_SRV(StructuredBuffer<float4>, GPUSceneInstancePayloadData)
+SHADER_PARAMETER_SRV(StructuredBuffer<float4>, GPUScenePrimitiveSceneData)
+SHADER_PARAMETER(uint32, GPUSceneFrameNumber)
+END_SHADER_PARAMETER_STRUCT()
+
+class FNiagaraCollisionRayTraceRG : public FGlobalShader
+{
+	DECLARE_GLOBAL_SHADER(FNiagaraCollisionRayTraceRG)
+	SHADER_USE_ROOT_PARAMETER_STRUCT(FNiagaraCollisionRayTraceRG, FGlobalShader)
+
+	BEGIN_SHADER_PARAMETER_STRUCT(FParameters, )
+		SHADER_PARAMETER_STRUCT_INCLUDE(FGPUSceneParameters, GPUSceneParameters)
+		SHADER_PARAMETER_SRV(Buffer<UINT>, HashTable)
+		SHADER_PARAMETER_UAV(Buffer<UINT>, RWHashTable)
+		SHADER_PARAMETER(uint32, HashTableSize)
+		SHADER_PARAMETER_SRV(RaytracingAccelerationStructure, TLAS)
+		SHADER_PARAMETER_SRV(Buffer<FNiagaraRayData>, Rays)
+		SHADER_PARAMETER(uint32, RaysOffset)
+		SHADER_PARAMETER_UAV(Buffer<FNiagaraRayTracingResult>, CollisionOutput)
+		SHADER_PARAMETER(uint32, CollisionOutputOffset)
+		SHADER_PARAMETER_SRV(Buffer<UINT>, RayTraceCounts)
+		SHADER_PARAMETER_SRV(Buffer<UINT>, HashToCollisionGroups)
+		SHADER_PARAMETER(uint32, MaxRetraces)
+	END_SHADER_PARAMETER_STRUCT()
+
+	class FFakeIndirectDispatch : SHADER_PERMUTATION_BOOL("NIAGARA_RAYTRACE_FAKE_INDIRECT");
+	using FPermutationDomain = TShaderPermutationDomain<FFakeIndirectDispatch>;
+
+	static bool ShouldCompilePermutation(const FGlobalShaderPermutationParameters& Parameters)
+	{
+		return ShouldCompileRayTracingShadersForProject(Parameters.Platform);
+	}
+
+	static void ModifyCompilationEnvironment(const FGlobalShaderPermutationParameters& Parameters, FShaderCompilerEnvironment& OutEnvironment)
+	{
+		FGlobalShader::ModifyCompilationEnvironment(Parameters, OutEnvironment);
+		OutEnvironment.SetDefine(TEXT("NIAGARA_SUPPORTS_RAY_TRACING"), 1);
+		OutEnvironment.SetDefine(TEXT("VF_SUPPORTS_PRIMITIVE_SCENE_DATA"), 1);
+		OutEnvironment.SetDefine(TEXT("USE_GLOBAL_GPU_SCENE_DATA"), 1);
+	}
+
+	static TShaderRef< FNiagaraCollisionRayTraceRG> GetShader(FGlobalShaderMap* ShaderMap)
+	{
+		FPermutationDomain PermutationVector;
+		PermutationVector.Set<FFakeIndirectDispatch>(!SupportsIndirectDispatch());
+
+		return ShaderMap->GetShader<FNiagaraCollisionRayTraceRG>(PermutationVector);
+	}
+
+	static FRHIRayTracingShader* GetRayTracingShader(FGlobalShaderMap* ShaderMap)
+	{
+		return GetShader(ShaderMap).GetRayTracingShader();
+	}
+
+	static bool SupportsIndirectDispatch()
+	{
+		return GRHISupportsRayTracingDispatchIndirect;
+	}
+};
+
+class FNiagaraCollisionRayTraceCH : public FGlobalShader
+{
+	DECLARE_GLOBAL_SHADER(FNiagaraCollisionRayTraceCH);
+
+	static bool ShouldCompilePermutation(const FGlobalShaderPermutationParameters& Parameters)
+	{
+		return ShouldCompileRayTracingShadersForProject(Parameters.Platform);
+	}
+
+	static void ModifyCompilationEnvironment(const FGlobalShaderPermutationParameters& Parameters, FShaderCompilerEnvironment& OutEnvironment)
+	{
+		FGlobalShader::ModifyCompilationEnvironment(Parameters, OutEnvironment);
+		OutEnvironment.SetDefine(TEXT("NIAGARA_SUPPORTS_RAY_TRACING"), 1);
+	}
+
+	FNiagaraCollisionRayTraceCH() = default;
+	FNiagaraCollisionRayTraceCH(const ShaderMetaType::CompiledShaderInitializerType& Initializer)
+		: FGlobalShader(Initializer)
+	{}
+};
+
+class FNiagaraCollisionRayTraceMiss : public FGlobalShader
+{
+	DECLARE_GLOBAL_SHADER(FNiagaraCollisionRayTraceMiss);
+
+	static bool ShouldCompilePermutation(const FGlobalShaderPermutationParameters& Parameters)
+	{
+		return ShouldCompileRayTracingShadersForProject(Parameters.Platform);
+	}
+
+	static void ModifyCompilationEnvironment(const FGlobalShaderPermutationParameters& Parameters, FShaderCompilerEnvironment& OutEnvironment)
+	{
+		FGlobalShader::ModifyCompilationEnvironment(Parameters, OutEnvironment);
+		OutEnvironment.SetDefine(TEXT("NIAGARA_SUPPORTS_RAY_TRACING"), 1);
+	}
+
+	FNiagaraCollisionRayTraceMiss() = default;
+	FNiagaraCollisionRayTraceMiss(const ShaderMetaType::CompiledShaderInitializerType& Initializer)
+		: FGlobalShader(Initializer)
+	{}
+};
+
+IMPLEMENT_GLOBAL_SHADER(FNiagaraCollisionRayTraceRG, "/Plugin/FX/Niagara/Private/NiagaraRayTracingShaders.usf", "NiagaraCollisionRayTraceRG", SF_RayGen);
+IMPLEMENT_GLOBAL_SHADER(FNiagaraCollisionRayTraceCH, "/Plugin/FX/Niagara/Private/NiagaraRayTracingShaders.usf", "NiagaraCollisionRayTraceCH", SF_RayHitGroup);
+IMPLEMENT_GLOBAL_SHADER(FNiagaraCollisionRayTraceMiss, "/Plugin/FX/Niagara/Private/NiagaraRayTracingShaders.usf", "NiagaraCollisionRayTraceMiss", SF_RayMiss);
+
+
+int32 GNiagaraHWRTScratchBucketSize = 1024;
+static FAutoConsoleVariableRef CVarNiagaraHWRTScratchBucketSize(
+	TEXT("fx.Niagara.RayTracing.ScratchPadBucketSize"),
+	GNiagaraHWRTScratchBucketSize,
+	TEXT("Size (in elements) for Niagara hardware ray tracing scratch buffer buckets. \n"),
+	ECVF_Default
+);
+
+int32 GNiagaraHWRTCountsScratchBucketSize = 3 * 256;
+static FAutoConsoleVariableRef CVarNiagaraHWRTCountsScratchBucketSize(
+	TEXT("fx.Niagara.RayTracing.CountsScratchPadBucketSize"),
+	GNiagaraHWRTCountsScratchBucketSize,
+	TEXT("Scratch bucket size for the HWRT counts buffer. This buffer requires 4. \n"),
+	ECVF_Default
+);
+
+FNiagaraRayTracingHelper::FNiagaraRayTracingHelper(EShaderPlatform InShaderPlatform)
+	: ShaderPlatform(InShaderPlatform)
+	, RayRequests(GNiagaraHWRTScratchBucketSize, TEXT("NiagaraRayRequests"))
+	, RayTraceIntersections(GNiagaraHWRTScratchBucketSize, TEXT("NiagaraRayTraceIntersections"))
+	, RayTraceCounts(GNiagaraHWRTCountsScratchBucketSize, TEXT("NiagaraRayTraceCounts"), BUF_Static | BUF_DrawIndirect)
+{}
+
+FNiagaraRayTracingHelper::~FNiagaraRayTracingHelper()
+{
+	Reset();
+	Dispatches.Reset();
+
+	//TODO: Move RT Helper out of NiagaraShader and into Niagara. Makes general sense plus it would allow use to hook into these memory stats.
+// 	DEC_MEMORY_STAT_BY(STAT_NiagaraGPUDataInterfaceMemory, RayRequests.AllocatedBytes());
+// 	DEC_MEMORY_STAT_BY(STAT_NiagaraGPUDataInterfaceMemory, RayTraceIntersections.AllocatedBytes());
+// 	DEC_MEMORY_STAT_BY(STAT_NiagaraGPUDataInterfaceMemory, RayTraceCounts.AllocatedBytes());
+
+	RayRequests.Release();
+	RayTraceIntersections.Release();
+	RayTraceCounts.Release();
+}
+
+void FNiagaraRayTracingHelper::Reset()
+{
+	RayTracingPipelineState = nullptr;
+	RayTracingScene = nullptr;
+	RayTracingSceneView = nullptr;
+	ViewUniformBuffer = nullptr;
+}
+
+void FNiagaraRayTracingHelper::BeginFrame(FRHICommandList& RHICmdList, bool HasRayTracingScene)
+{
+	//Store off last frames dispatches so we can still access the previous results buffers for reading in this frame's simulations.
+	PreviousFrameDispatches = MoveTemp(Dispatches);
+	Dispatches.Reset();
+
+	//Ensure we're each buffer is definitely in the right access state.
+	RayRequests.Transition(RHICmdList, ERHIAccess::Unknown, ERHIAccess::UAVCompute);//Simulations will write new ray trace requests.
+	RayTraceIntersections.Transition(RHICmdList, ERHIAccess::Unknown, ERHIAccess::SRVCompute);//Simulations will read from the previous frames results.
+	RayTraceCounts.Transition(RHICmdList, ERHIAccess::Unknown, ERHIAccess::UAVCompute);//Simulations will accumulate the number of traces requests here.
+
+	//Reset the allocations.
+	RayRequests.Reset();
+
+	//Note this doesn't change any buffers or data themselves and the ray intersections buffer are still going to be read as SRVs in the up coming simulations dispatches.
+	//We're just clearing existing allocations here.
+	RayTraceIntersections.Reset();
+
+	//We have to also clear the counts/indirect args buffer.
+	RayTraceCounts.Reset(RHICmdList, true);
+
+	// Clear the dummy buffer allocations.
+	DummyDispatch.Reset();
+
+	RayTracingPipelineState = nullptr;
+	RayTracingScene = nullptr;
+	RayTracingSceneView = nullptr;
+	ViewUniformBuffer = nullptr;
+}
+
+void FNiagaraRayTracingHelper::EndFrame(FRHICommandList& RHICmdList, bool HasRayTracingScene, FScene* Scene)
+{
+	// skip everything if we don't have a ray tracing scene
+	if (HasRayTracingScene)
+	{
+		RayRequests.Transition(RHICmdList, ERHIAccess::UAVCompute, ERHIAccess::SRVCompute);//Ray trace dispatches will read the ray requests.
+		RayTraceIntersections.Transition(RHICmdList, ERHIAccess::SRVCompute, ERHIAccess::UAVCompute);//Ray trace dispatches will write new results.
+		RayTraceCounts.Transition(RHICmdList, ERHIAccess::UAVCompute, ERHIAccess::IndirectArgs | ERHIAccess::SRVCompute);//Ray trace dispatches read these counts and use them for indirect dispatches.
+	
+		IssueRayTraces(RHICmdList, Scene);
+
+		RayRequests.Transition(RHICmdList, ERHIAccess::SRVCompute, ERHIAccess::UAVCompute);//Next frame, simulations will write new ray trace requests.
+		RayTraceIntersections.Transition(RHICmdList, ERHIAccess::UAVCompute, ERHIAccess::SRVCompute);//Next frame these results will be read by simulations shaders.
+		RayTraceCounts.Transition(RHICmdList, ERHIAccess::IndirectArgs | ERHIAccess::SRVCompute, ERHIAccess::UAVCompute);//Next frame these counts will be written by simulation shaders.
+	}
+}
+
+void FNiagaraRayTracingHelper::AddToDispatch(FNiagaraDataInterfaceProxy* DispatchKey, uint32 MaxRays, int32 MaxRetraces)
+{
+	FNiagaraRayTraceDispatchInfo& Dispatch = Dispatches.FindOrAdd(DispatchKey);
+	Dispatch.MaxRays += MaxRays;
+	Dispatch.MaxRetraces = MaxRetraces;
+}
+
+void FNiagaraRayTracingHelper::BuildDispatch(FRHICommandList& RHICmdList, FNiagaraDataInterfaceProxy* DispatchKey)
+{
+	FNiagaraRayTraceDispatchInfo& Dispatch = Dispatches.FindChecked(DispatchKey);
+
+	//Finalize allocations first time we access the dispatch.
+	if (Dispatch.RayRequests.IsValid() == false)
+	{
+		//TODO: Move RT Helper out of NiagaraShader and into Niagara. Makes general sense plus it would allow use to hook into these memory stats.
+// 		DEC_MEMORY_STAT_BY(STAT_NiagaraGPUDataInterfaceMemory, RayRequests.AllocatedBytes());
+// 		DEC_MEMORY_STAT_BY(STAT_NiagaraGPUDataInterfaceMemory, RayTraceIntersections.AllocatedBytes());
+// 		DEC_MEMORY_STAT_BY(STAT_NiagaraGPUDataInterfaceMemory, RayTraceCounts.AllocatedBytes());
+
+		Dispatch.RayRequests = RayRequests.Alloc(Dispatch.MaxRays);
+		Dispatch.RayTraceIntersections = RayTraceIntersections.Alloc(Dispatch.MaxRays);
+		Dispatch.RayCounts = RayTraceCounts.Alloc<uint32>(3, RHICmdList, true);
+
+		//TODO: Could add a delegate to call from the scratch when new buffers are allocated to reduce stat spam here?
+// 		INC_MEMORY_STAT_BY(STAT_NiagaraGPUDataInterfaceMemory, RayRequests.AllocatedBytes());
+// 		INC_MEMORY_STAT_BY(STAT_NiagaraGPUDataInterfaceMemory, RayTraceIntersections.AllocatedBytes());
+// 		INC_MEMORY_STAT_BY(STAT_NiagaraGPUDataInterfaceMemory, RayTraceCounts.AllocatedBytes());
+
+		//Find last frame's results buffer if there was one. Ideally we can do this without the map but it will need a bit of wrangling.
+		FNiagaraRayTraceDispatchInfo* PrevDispatch = PreviousFrameDispatches.Find(DispatchKey);
+		if (PrevDispatch && PrevDispatch->RayTraceIntersections.IsValid())
+		{
+			Dispatch.LastFrameRayTraceIntersections = PrevDispatch->RayTraceIntersections;
+		}
+		else
+		{
+			//Set the current frame results just so this is a valid buffer.
+			//If we didn't have a results buffer last frame then this won't actually be read.
+			Dispatch.LastFrameRayTraceIntersections = Dispatch.RayTraceIntersections;
+		}
+	}
+}
+
+void FNiagaraRayTracingHelper::BuildDummyDispatch(FRHICommandList& RHICmdList)
+{
+	if (DummyDispatch.IsValid() == false)
+	{
+		//TODO: Move RT Helper out of NiagaraShader and into Niagara. Makes general sense plus it would allow use to hook into these memory stats.
+// 		DEC_MEMORY_STAT_BY(STAT_NiagaraGPUDataInterfaceMemory, RayRequests.AllocatedBytes());
+// 		DEC_MEMORY_STAT_BY(STAT_NiagaraGPUDataInterfaceMemory, RayTraceIntersections.AllocatedBytes());
+// 		DEC_MEMORY_STAT_BY(STAT_NiagaraGPUDataInterfaceMemory, RayTraceCounts.AllocatedBytes());
+
+		DummyDispatch.RayRequests = RayRequests.Alloc(1);
+		DummyDispatch.RayTraceIntersections = RayTraceIntersections.Alloc(1);
+		DummyDispatch.LastFrameRayTraceIntersections = DummyDispatch.RayTraceIntersections;
+		DummyDispatch.RayCounts = RayTraceCounts.Alloc<uint32>(3, RHICmdList, true);
+		
+		//TODO: Move RT Helper out of NiagaraShader and into Niagara. Makes general sense plus it would allow use to hook into these memory stats.
+// 		INC_MEMORY_STAT_BY(STAT_NiagaraGPUDataInterfaceMemory, RayRequests.AllocatedBytes());
+// 		INC_MEMORY_STAT_BY(STAT_NiagaraGPUDataInterfaceMemory, RayTraceIntersections.AllocatedBytes());
+// 		INC_MEMORY_STAT_BY(STAT_NiagaraGPUDataInterfaceMemory, RayTraceCounts.AllocatedBytes());
+
+		DummyDispatch.MaxRays = 0;
+		DummyDispatch.MaxRetraces = 0;
+	}
+}
+
+const FNiagaraRayTraceDispatchInfo& FNiagaraRayTracingHelper::GetDispatch(FNiagaraDataInterfaceProxy* DispatchKey) const
+{
+	const FNiagaraRayTraceDispatchInfo& Dispatch = Dispatches.FindChecked(DispatchKey);
+
+	// make sure that the buffer is valid (that BuildDispatch has been called)
+	check(Dispatch.RayRequests.IsValid());
+
+	return Dispatch;
+}
+
+const FNiagaraRayTraceDispatchInfo& FNiagaraRayTracingHelper::GetDummyDispatch() const
+{
+	// make sure that the buffers are valid (that BuildDummyDispatch has been called)
+	check(DummyDispatch.RayRequests.IsValid());
+
+	return DummyDispatch;
+}
+
+static FRayTracingPipelineState* CreateNiagaraRayTracingPipelineState(
+	EShaderPlatform Platform,
+	FRHICommandList& RHICmdList,
+	FRHIRayTracingShader* RayGenShader,
+	FRHIRayTracingShader* ClosestHitShader,
+	FRHIRayTracingShader* MissShader)
+{
+	FRayTracingPipelineStateInitializer Initializer;
+	Initializer.MaxPayloadSizeInBytes = sizeof(FNiagaraRayTracingPayload);
+
+	FRHIRayTracingShader* RayGenShaderTable[] = { RayGenShader };
+	Initializer.SetRayGenShaderTable(RayGenShaderTable);
+
+	FRHIRayTracingShader* HitGroupTable[] = { ClosestHitShader };
+	Initializer.SetHitGroupTable(HitGroupTable);
+
+	FRHIRayTracingShader* MissTable[] = { MissShader };
+	Initializer.SetMissShaderTable(MissTable);
+
+	Initializer.bAllowHitGroupIndexing = true; // Use the same hit shader for all geometry in the scene by disabling SBT indexing.
+
+	return PipelineStateCache::GetAndOrCreateRayTracingPipelineState(RHICmdList, Initializer);
+}
+
+static void BindNiagaraRayTracingMeshCommands(
+	FRHICommandList& RHICmdList,
+	FRayTracingSceneRHIRef RayTracingScene,
+	FRHIUniformBuffer* ViewUniformBuffer,
+	TConstArrayView<FVisibleRayTracingMeshCommand> RayTracingMeshCommands,
+	FRayTracingPipelineState* Pipeline,
+	TFunctionRef<uint32(const FRayTracingMeshCommand&)> PackUserData)
+{
+	const int32 NumTotalBindings = RayTracingMeshCommands.Num();
+
+	const uint32 MergedBindingsSize = sizeof(FRayTracingLocalShaderBindings) * NumTotalBindings;
+
+	const uint32 NumUniformBuffers = 1;
+	const uint32 UniformBufferArraySize = sizeof(FRHIUniformBuffer*) * NumUniformBuffers;
+
+	FRayTracingLocalShaderBindings* Bindings = nullptr;
+	FRHIUniformBuffer** UniformBufferArray = nullptr;
+
+	if (RHICmdList.Bypass())
+	{
+		FMemStack& MemStack = FMemStack::Get();
+
+		Bindings = reinterpret_cast<FRayTracingLocalShaderBindings*>(MemStack.Alloc(MergedBindingsSize, alignof(FRayTracingLocalShaderBindings)));
+		UniformBufferArray = reinterpret_cast<FRHIUniformBuffer**>(MemStack.Alloc(UniformBufferArraySize, alignof(FRHIUniformBuffer*)));
+	}
+	else
+	{
+		Bindings = reinterpret_cast<FRayTracingLocalShaderBindings*>(RHICmdList.Alloc(MergedBindingsSize, alignof(FRayTracingLocalShaderBindings)));
+		UniformBufferArray = reinterpret_cast<FRHIUniformBuffer**>(RHICmdList.Alloc(UniformBufferArraySize, alignof(FRHIUniformBuffer*)));
+	}
+
+	UniformBufferArray[0] = ViewUniformBuffer;
+
+	uint32 BindingIndex = 0;
+	for (const FVisibleRayTracingMeshCommand VisibleMeshCommand : RayTracingMeshCommands)
+	{
+		const FRayTracingMeshCommand& MeshCommand = *VisibleMeshCommand.RayTracingMeshCommand;
+
+		FRayTracingLocalShaderBindings Binding = {};
+		Binding.InstanceIndex = VisibleMeshCommand.InstanceIndex;
+		Binding.SegmentIndex = MeshCommand.GeometrySegmentIndex;
+		Binding.UserData = PackUserData(MeshCommand);
+		Binding.UniformBuffers = UniformBufferArray;
+		Binding.NumUniformBuffers = NumUniformBuffers;
+
+		Bindings[BindingIndex] = Binding;
+		BindingIndex++;
+	}
+
+	const bool bCopyDataToInlineStorage = false; // Storage is already allocated from RHICmdList, no extra copy necessary
+	RHICmdList.SetRayTracingHitGroups(
+		RayTracingScene,
+		Pipeline,
+		NumTotalBindings,
+		Bindings,
+		bCopyDataToInlineStorage);
+}
+
+void FNiagaraRayTracingHelper::SetPrimitiveCollisionGroup(FPrimitiveSceneInfo& Primitive, uint32 CollisionGroup)
+{
+	CollisionGroupMap.FindOrAdd(Primitive.PrimitiveComponentId) = CollisionGroup;
+	bCollisionGroupMapDirty = true;
+}
+
+void FNiagaraRayTracingHelper::RefreshPrimitiveInstanceData()
+{
+	bCollisionGroupMapDirty = true;
+}
+
+void FNiagaraRayTracingHelper::UpdateCollisionGroupMap(FRHICommandList& RHICmdList, FScene* Scene, ERHIFeatureLevel::Type FeatureLevel)
+{
+	if (bCollisionGroupMapDirty)
+	{
+		SCOPED_DRAW_EVENT(RHICmdList, NiagaraUpdateCollisionGroupsMap);
+		bCollisionGroupMapDirty = false;
+
+		uint32 MinBufferInstances = FNiagaraUpdateCollisionGroupMapCS::THREAD_COUNT;
+		uint32 NeededInstances = FMath::Max((uint32)CollisionGroupMap.Num(), MinBufferInstances);
+		uint32 AllocInstances = Align(NeededInstances, FNiagaraUpdateCollisionGroupMapCS::THREAD_COUNT);
+
+		//We can probably be smarter here but for now just push the whole map over to the GPU each time it's dirty.
+		//Ideally this should be a small amount of data and updated infrequently.
+		FReadBuffer NewPrimIdCollisionGroupPairs;		
+		NewPrimIdCollisionGroupPairs.Initialize(TEXT("NewPrimIdCollisionGroupPairs"), sizeof(uint32) * 2, AllocInstances, EPixelFormat::PF_R32G32_UINT, BUF_Volatile);
+
+		uint32* PrimIdCollisionGroupPairPtr = (uint32*)RHILockBuffer(NewPrimIdCollisionGroupPairs.Buffer, 0, AllocInstances * sizeof(uint32) * 2, RLM_WriteOnly);
+		FMemory::Memset(PrimIdCollisionGroupPairPtr, 0, AllocInstances * sizeof(uint32) * 2);
+
+		for (auto Entry : CollisionGroupMap)
+		{
+			FPrimitiveComponentId PrimId = Entry.Key;
+			uint32 CollisionGroup = Entry.Value;
+			
+			uint32 GPUSceneInstanceIndex = INDEX_NONE;
+			//Ugh this is a bit pants. Maybe try to rework things so I can use the direct prim index.
+			int32 PrimIndex = Scene->PrimitiveComponentIds.Find(PrimId);
+			if (PrimIndex != INDEX_NONE)
+			{
+				GPUSceneInstanceIndex = Scene->Primitives[PrimIndex]->GetInstanceSceneDataOffset();
+			}
+
+			PrimIdCollisionGroupPairPtr[0] = GPUSceneInstanceIndex;
+			PrimIdCollisionGroupPairPtr[1] = CollisionGroup;
+			PrimIdCollisionGroupPairPtr += 2;
+		}
+		RHIUnlockBuffer(NewPrimIdCollisionGroupPairs.Buffer);
+
+		//Init the hash table if needed
+		if (HashTableSize < AllocInstances)
+		{
+			HashTableSize = AllocInstances;
+
+			PrimIdHashTable.Release();
+			PrimIdHashTable.Initialize(TEXT("NiagaraPrimIdHashTable"), sizeof(uint32), AllocInstances, EPixelFormat::PF_R32_UINT, BUF_Static);
+
+			HashToCollisionGroups.Release();
+			HashToCollisionGroups.Initialize(TEXT("NiagaraPrimIdHashToCollisionGroups"), sizeof(uint32), AllocInstances, EPixelFormat::PF_R32_UINT, BUF_Static);
+		}
+
+		RHICmdList.Transition(FRHITransitionInfo(PrimIdHashTable.UAV, ERHIAccess::Unknown, ERHIAccess::UAVCompute));
+		RHICmdList.Transition(FRHITransitionInfo(HashToCollisionGroups.UAV, ERHIAccess::Unknown, ERHIAccess::UAVCompute));
+
+ 		//First we have to clear the buffers. Can probably do this better.
+ 		NiagaraFillGPUIntBuffer(RHICmdList, FeatureLevel, PrimIdHashTable, 0);
+ 		RHICmdList.Transition(FRHITransitionInfo(PrimIdHashTable.UAV, ERHIAccess::UAVCompute, ERHIAccess::UAVCompute));
+ 		NiagaraFillGPUIntBuffer(RHICmdList, FeatureLevel, HashToCollisionGroups, INDEX_NONE);
+ 		RHICmdList.Transition(FRHITransitionInfo(HashToCollisionGroups.UAV, ERHIAccess::UAVCompute, ERHIAccess::UAVCompute));
+
+		TShaderMapRef<FNiagaraUpdateCollisionGroupMapCS> ComputeShader(GetGlobalShaderMap(GMaxRHIFeatureLevel));
+		FRHIComputeShader* ShaderRHI = ComputeShader.GetComputeShader();
+
+		FNiagaraUpdateCollisionGroupMapCS::FParameters Params;
+		Params.RWHashTable = PrimIdHashTable.UAV;
+		Params.HashTableSize = HashTableSize;
+		Params.RWHashToCollisionGroups = HashToCollisionGroups.UAV;
+		Params.NewPrimIdCollisionGroupPairs = NewPrimIdCollisionGroupPairs.SRV;
+		Params.NumNewPrims = CollisionGroupMap.Num();
+
+		// To simplify the shader code, the size of the table must be a multiple of the thread count.
+		check(HashTableSize % FNiagaraUpdateCollisionGroupMapCS::THREAD_COUNT == 0);
+
+		RHICmdList.SetComputeShader(ShaderRHI);
+		SetShaderParameters(RHICmdList, ComputeShader, ShaderRHI, Params);
+		RHICmdList.DispatchComputeShader(FMath::DivideAndRoundUp(CollisionGroupMap.Num(), (int32)FNiagaraUpdateCollisionGroupMapCS::THREAD_COUNT), 1, 1);
+		UnsetShaderUAVs(RHICmdList, ComputeShader, ShaderRHI);
+
+		RHICmdList.Transition(FRHITransitionInfo(PrimIdHashTable.UAV, ERHIAccess::UAVCompute, ERHIAccess::SRVCompute));
+		RHICmdList.Transition(FRHITransitionInfo(HashToCollisionGroups.UAV, ERHIAccess::UAVCompute, ERHIAccess::SRVCompute));
+	}
+}
+
+void FNiagaraRayTracingHelper::BuildRayTracingSceneInfo(FRHICommandList& RHICmdList, TConstArrayView<FViewInfo> Views)
+{
+	check(Views.Num() > 0);
+
+	const FViewInfo& ReferenceView = Views[0];
+	const FScene* RenderScene = (ReferenceView.Family && ReferenceView.Family->Scene) ? ReferenceView.Family->Scene->GetRenderScene() : nullptr;
+
+	if (RenderScene && RenderScene->RayTracingScene.IsCreated())
+	{
+		RayTracingScene = RenderScene->RayTracingScene.GetRHIRayTracingSceneChecked();
+		RayTracingSceneView = RenderScene->RayTracingScene.GetShaderResourceViewChecked();
+		ViewUniformBuffer = ReferenceView.ViewUniformBuffer;
+
+		FGlobalShaderMap* ShaderMap = GetGlobalShaderMap(ShaderPlatform);
+		auto RayGenShader = FNiagaraCollisionRayTraceRG::GetRayTracingShader(ShaderMap);
+		auto ClosestHitShader = ShaderMap->GetShader<FNiagaraCollisionRayTraceCH>().GetRayTracingShader();
+		auto MissShader = ShaderMap->GetShader<FNiagaraCollisionRayTraceMiss>().GetRayTracingShader();
+
+		RayTracingPipelineState = CreateNiagaraRayTracingPipelineState(
+			ShaderPlatform,
+			RHICmdList,
+			RayGenShader,
+			ClosestHitShader,
+			MissShader);
+
+		// some options for what we want with our per MeshCommand user data.  For now we'll ignore it, but possibly
+		// something we'd want to incorporate.  Some examples could be if the material is translucent, or possibly the physical material?
+		auto BakeTranslucent = [&](const FRayTracingMeshCommand& MeshCommand) {	return (MeshCommand.bIsTranslucent != 0) & 0x1;	};
+		auto BakeDefault = [&](const FRayTracingMeshCommand& MeshCommand) { return 0; };
+
+		BindNiagaraRayTracingMeshCommands(
+			RHICmdList,
+			RayTracingScene,
+			ViewUniformBuffer,
+			ReferenceView.VisibleRayTracingMeshCommands,
+			RayTracingPipelineState,
+			BakeDefault);
+	}
+	else
+	{
+		Reset();
+	}
+}
+
+void FNiagaraRayTracingHelper::IssueRayTraces(FRHICommandList& RHICmdList, FScene* Scene)
+{
+	SCOPED_DRAW_EVENT(RHICmdList, NiagaraIssueCollisionRayTraces);
+	check(RayTracingPipelineState);
+	check(RayTracingScene);
+	check(RayTracingSceneView);
+	check(Scene);
+
+	for (TPair<FNiagaraDataInterfaceProxy*, FNiagaraRayTraceDispatchInfo>& Pair : Dispatches)
+	{
+		FNiagaraRayTraceDispatchInfo& DispatchInfo = Pair.Value;
+		
+		if (DispatchInfo.MaxRays == 0)
+		{
+			continue;
+		}
+
+		FGlobalShaderMap* ShaderMap = GetGlobalShaderMap(ShaderPlatform);
+
+		TShaderRef<FNiagaraCollisionRayTraceRG> RGShader = FNiagaraCollisionRayTraceRG::GetShader(ShaderMap);
+
+		FNiagaraCollisionRayTraceRG::FParameters Params;
+
+		Params.GPUSceneParameters.GPUSceneInstanceSceneData = Scene->GPUScene.InstanceSceneDataBuffer.SRV;
+		Params.GPUSceneParameters.GPUSceneInstancePayloadData = Scene->GPUScene.InstancePayloadDataBuffer.SRV;
+		Params.GPUSceneParameters.GPUScenePrimitiveSceneData = Scene->GPUScene.PrimitiveBuffer.SRV;
+		Params.GPUSceneParameters.GPUSceneFrameNumber = Scene->GPUScene.GetSceneFrameNumber();
+
+		Params.HashTable = PrimIdHashTable.SRV;
+		Params.HashTableSize = HashTableSize;
+		Params.TLAS = RayTracingSceneView;
+		Params.Rays = DispatchInfo.RayRequests.Buffer->SRV;
+		Params.RaysOffset = DispatchInfo.RayRequests.Offset;
+		Params.CollisionOutput = DispatchInfo.RayTraceIntersections.Buffer->UAV;
+		Params.CollisionOutputOffset = DispatchInfo.RayTraceIntersections.Offset;
+		Params.HashToCollisionGroups = HashToCollisionGroups.SRV;
+		Params.MaxRetraces = DispatchInfo.MaxRetraces;
+
+		if (FNiagaraCollisionRayTraceRG::SupportsIndirectDispatch())
+		{
+			FRayTracingShaderBindingsWriter GlobalResources;
+			SetShaderParameters(GlobalResources, RGShader, Params);
+
+			//Can we wrangle things so we can have one indirect dispatch with each internal dispatch pointing to potentially different Ray and Results buffers?
+			//For now have a each as a unique dispatch.
+			RHICmdList.RayTraceDispatchIndirect(
+				RayTracingPipelineState,
+				RGShader.GetRayTracingShader(),
+				RayTracingScene,
+				GlobalResources,
+				DispatchInfo.RayCounts.Buffer->Buffer,
+				DispatchInfo.RayCounts.Offset * sizeof(uint32));
+		}
+		else
+		{
+			Params.RayTraceCounts = DispatchInfo.RayCounts.Buffer->SRV;
+
+			FRayTracingShaderBindingsWriter GlobalResources;
+			SetShaderParameters(GlobalResources, RGShader, Params);
+
+			RHICmdList.RayTraceDispatch(
+				RayTracingPipelineState,
+				RGShader.GetRayTracingShader(),
+				RayTracingScene,
+				GlobalResources,
+				DispatchInfo.MaxRays,
+				1
+			);
+		}
+	}
+}
+
+bool FNiagaraRayTracingHelper::IsValid() const
+{
+	return RayTracingPipelineState
+		&& RayTracingScene
+		&& RayTracingSceneView
+		&& ViewUniformBuffer;
+}
+
+#endif
\ No newline at end of file
diff --git a/Engine/Plugins/FX/Niagara/Source/NiagaraShader/Public/NiagaraAsyncGpuTraceProvider.h b/Engine/Plugins/FX/Niagara/Source/NiagaraShader/Public/NiagaraAsyncGpuTraceProvider.h
deleted file mode 100644
diff --git a/Engine/Plugins/FX/Niagara/Source/NiagaraShader/Public/NiagaraRayTracingHelper.h b/Engine/Plugins/FX/Niagara/Source/NiagaraShader/Public/NiagaraRayTracingHelper.h
new file mode 100644
--- /dev/null
+++ b/Engine/Plugins/FX/Niagara/Source/NiagaraShader/Public/NiagaraRayTracingHelper.h
@@ -0,0 +1,554 @@
+// Copyright Epic Games, Inc. All Rights Reserved.
+#pragma once
+
+#include "CoreMinimal.h"
+
+#include "RHIDefinitions.h"
+#if RHI_RAYTRACING
+
+#include "PrimitiveSceneInfo.h"
+#include "RHI.h"
+#include "RHIUtilities.h"
+
+class FScene;
+class FViewInfo;
+class FRayTracingPipelineState;
+class FRHICommandList;
+class FRHIRayTracingScene;
+class FRHIShaderResourceView;
+class FRHIUniformBuffer;
+class FRHIUnorderedAccessView;
+struct FRWBuffer;
+struct FNiagaraDataInterfaceProxy;
+
+
+//////////////////////////////////////////////////////////////////////////
+//TODO: Move scratch pad buffers to core rendering code along side global dynamic read buffers etc.
+
+/** Allocates space from a pool of fixed size static RWBuffers.
+*	Allocations are valid between calls to Reset.
+*	Useful for transient working data generated in one dispatch and consumed in others.
+*	All buffers in the same scratch pad are transitioned as one.
+*/
+class NIAGARASHADER_API FGPUScratchPad
+{
+public:
+	struct FAllocation
+	{
+		FRWBuffer* Buffer = nullptr;
+		uint32 Offset = 0;
+
+		FORCEINLINE bool IsValid()const { return Buffer != nullptr; }
+	};
+
+	struct FGPUScratchBuffer
+	{
+	public:
+		struct FBuffer
+		{
+			FRWBuffer Buffer;
+			uint32 Used = 0;
+		};
+
+		FGPUScratchBuffer(uint32 InBucketSize, EPixelFormat InPixelFormat, uint32 InBytesPerElement, EBufferUsageFlags InBufferUsageFlags)
+			: MinBucketSize(InBucketSize), PixelFormat(InPixelFormat), BytesPerElement(InBytesPerElement), BufferUsageFlags(InBufferUsageFlags)
+		{}
+
+		/** Allocates NumElements from the buffers. Will try to find an existing buffer to fit the allocation but will create a new one if needed. */
+		bool Allocate(uint32 NumElements, FString& BufferDebugName, FRHICommandList& RHICmdList, FGPUScratchPad::FAllocation& OutAllocation)
+		{
+			//Look for a buffer to fit this allocation.
+			FBuffer* ToUse = nullptr;
+			for (FBuffer& Buffer : Buffers)
+			{
+				uint32 NumBufferElements = Buffer.Buffer.NumBytes / BytesPerElement;
+				if (NumElements + Buffer.Used <= NumBufferElements)
+				{
+					ToUse = &Buffer;
+				}
+			}
+
+			bool bNew = false;
+			if (ToUse == nullptr)
+			{
+				//Allocate a new buffer that will fit this allocation.
+				FBuffer* NewBuffer = new FBuffer();
+				Buffers.Add(NewBuffer);
+
+				//Round up the size to pow 2 to ensure a slowly growing user doesn't just keep allocating new buckets. Probably a better guess size here than power of 2.
+				uint32 NewBufferSize = FPlatformMath::RoundUpToPowerOfTwo(FMath::Max(MinBucketSize, NumElements));
+
+				NewBuffer->Buffer.Initialize(*BufferDebugName, BytesPerElement, NewBufferSize, PixelFormat, CurrentAccess, BufferUsageFlags);
+			
+				ToUse = NewBuffer;
+
+				NumAllocatedBytes += NewBuffer->Buffer.NumBytes;
+
+				bNew = true;
+			}
+
+			OutAllocation.Buffer = &ToUse->Buffer;
+			OutAllocation.Offset = ToUse->Used;
+			ToUse->Used += NumElements;
+			NumUsedBytes += NumElements * BytesPerElement;
+			return bNew;
+		}
+
+		/** Resets the buffer and releases it's internal RWBuffers too. */
+		void Release()
+		{
+			Reset();
+			CurrentAccess = ERHIAccess::UAVCompute;
+			for (FBuffer& Buffer : Buffers)
+			{
+				Buffer.Buffer.Release();
+				Buffer.Used = 0;
+			}
+			Buffers.Reset();
+
+			NumAllocatedBytes = 0;
+		}
+
+		/** Resets the usage for all buffers to zero. */
+		void Reset()
+		{
+			for (FBuffer& Buffer : Buffers)
+			{
+				Buffer.Used = 0;
+			}
+
+			NumUsedBytes = 0;
+		}
+
+		template<typename TFunc>
+		void ForEachBuffer(TFunc Func)
+		{
+			for (FBuffer& Buffer : Buffers)
+			{
+				Func(Buffer.Buffer);
+			}
+		}
+
+		/** Transitions all buffers in the scratch pad. */
+		void Transition(FRHICommandList& RHICmdList, ERHIAccess InPreviousState, ERHIAccess InNewState, EResourceTransitionFlags InFlags = EResourceTransitionFlags::None)
+		{
+			//Store off current access state so we can place any new buffers in the same state.
+			CurrentAccess = InNewState;
+
+ 			TArray<FRHITransitionInfo, TMemStackAllocator<>> Transitions;
+ 			Transitions.Reserve(Buffers.Num());
+ 			for (FBuffer& Buffer : Buffers)
+ 			{
+ 				Transitions.Emplace(Buffer.Buffer.UAV, InPreviousState, InNewState, InFlags);
+ 			}
+ 
+ 			RHICmdList.Transition(Transitions);
+		}
+
+		/** Returns the total bytes allocated for all buffers in the scratch pad. */
+		FORCEINLINE uint32 AllocatedBytes()const
+		{
+			return NumAllocatedBytes;
+		}
+
+		/** Returns the number of used bytes allocated for all buffers in the scratch pad. */
+		FORCEINLINE uint32 UsedBytes()const
+		{
+			return NumUsedBytes;
+		}
+
+		FORCEINLINE ERHIAccess GetExpectedCurrentAccess()const{return CurrentAccess;}
+	private:
+
+		uint32 MinBucketSize = 4096;
+		TIndirectArray<FBuffer> Buffers;
+
+		EPixelFormat PixelFormat;
+		uint32 BytesPerElement = 0;
+		EBufferUsageFlags BufferUsageFlags = EBufferUsageFlags::None;
+		ERHIAccess CurrentAccess = ERHIAccess::UAVCompute;
+
+		uint32 NumAllocatedBytes = 0;
+		uint32 NumUsedBytes = 0;
+	};
+
+	FGPUScratchPad(uint32 InBucketSize, const FString& InDebugName, EBufferUsageFlags BufferUsageFlags = BUF_Static)
+		: BucketSize(InBucketSize)
+		, DebugName(InDebugName)
+		, ScratchBufferFloat(BucketSize, PF_R32_FLOAT, sizeof(float), BufferUsageFlags)
+		, ScratchBufferUInt(BucketSize, PF_R32_UINT, sizeof(uint32), BufferUsageFlags)
+	{
+	}
+
+	/** Allocate a number of elements from the scratch pad. This will attempt to allocate from an existing buffer but will create a new one if it can't. */
+	template<typename T>
+	FORCEINLINE FAllocation Alloc(uint32 Size, FRHICommandList& RHICmdList, bool bClearNew);
+
+	/** Resets the scratch pad and also releases all buffers it holds. */
+	void Release()
+	{
+		ScratchBufferFloat.Release();
+		ScratchBufferUInt.Release();
+	}
+
+	/** Resets the usage counters for buffers in this scratch pad, invalidating all previous allocations. Does not affect actual allocations. Optionally clears all buffers to zero. */
+	void Reset(FRHICommandList& RHICmdList, bool bClear)
+	{
+		ScratchBufferFloat.Reset();
+		ScratchBufferUInt.Reset();
+
+		if (bClear)
+		{
+			Transition(RHICmdList, ERHIAccess::Unknown, ERHIAccess::UAVCompute);
+
+			auto ClearFloatBuffer = [&](FRWBuffer& Buffer)
+			{
+				RHICmdList.ClearUAVFloat(Buffer.UAV, FVector4f(0.0f));
+			};
+			ScratchBufferFloat.ForEachBuffer(ClearFloatBuffer);
+
+			auto ClearUIntBuffer = [&](FRWBuffer& Buffer)
+			{
+				RHICmdList.ClearUAVUint(Buffer.UAV, FUintVector4(0, 0, 0, 0));
+			};
+			ScratchBufferUInt.ForEachBuffer(ClearUIntBuffer);
+
+			Transition(RHICmdList, ERHIAccess::UAVCompute, ERHIAccess::UAVCompute);
+		}
+	}
+
+	/** Transitions all buffers in the scratch pad. */
+	void Transition(FRHICommandList& RHICmdList, ERHIAccess InPreviousState, ERHIAccess InNewState, EResourceTransitionFlags InFlags = EResourceTransitionFlags::None)
+	{
+		ScratchBufferFloat.Transition(RHICmdList, InPreviousState, InNewState, InFlags);
+		ScratchBufferUInt.Transition(RHICmdList, InPreviousState, InNewState, InFlags);
+	}
+
+	/** Returns the total bytes allocated for all buffers in the scratch pad. */
+	FORCEINLINE uint32 AllocatedBytes()const
+	{
+		return ScratchBufferFloat.AllocatedBytes() + ScratchBufferUInt.AllocatedBytes();
+	}
+
+	/** Returns the number of used bytes allocated for all buffers in the scratch pad. */
+	FORCEINLINE uint32 UsedBytes()const
+	{
+		return ScratchBufferFloat.UsedBytes() + ScratchBufferUInt.UsedBytes();
+	}
+
+private:
+
+	/** Size of each FRWBuffer in the pool. */
+	uint32 BucketSize = 4096;
+	FString DebugName;
+
+	FGPUScratchBuffer ScratchBufferFloat;
+	FGPUScratchBuffer ScratchBufferUInt;
+};
+
+template<>
+FORCEINLINE FGPUScratchPad::FAllocation FGPUScratchPad::Alloc<float>(uint32 Size, FRHICommandList& RHICmdList, bool bClearNew)
+{
+	FGPUScratchPad::FAllocation Allocation;
+	bool bNewBuffer = ScratchBufferFloat.Allocate(Size, DebugName, RHICmdList, Allocation);
+	if (bNewBuffer && bClearNew)
+	{
+		RHICmdList.Transition(FRHITransitionInfo(Allocation.Buffer->UAV, ScratchBufferFloat.GetExpectedCurrentAccess(), ERHIAccess::UAVCompute));
+		RHICmdList.ClearUAVFloat(Allocation.Buffer->UAV, FVector4f(0.0f));
+		RHICmdList.Transition(FRHITransitionInfo(Allocation.Buffer->UAV, ERHIAccess::UAVCompute, ScratchBufferFloat.GetExpectedCurrentAccess()));
+	}
+	return Allocation;
+}
+
+template<>
+FORCEINLINE FGPUScratchPad::FAllocation FGPUScratchPad::Alloc<uint32>(uint32 Size, FRHICommandList& RHICmdList, bool bClearNew)
+{
+	FGPUScratchPad::FAllocation Allocation;
+	bool bNewBuffer = ScratchBufferUInt.Allocate(Size, DebugName, RHICmdList, Allocation);
+	if (bNewBuffer && bClearNew)
+	{
+		RHICmdList.Transition(FRHITransitionInfo(Allocation.Buffer->UAV, ScratchBufferUInt.GetExpectedCurrentAccess(), ERHIAccess::UAVCompute));
+		RHICmdList.ClearUAVUint(Allocation.Buffer->UAV, FUintVector4(0, 0, 0, 0));
+		RHICmdList.Transition(FRHITransitionInfo(Allocation.Buffer->UAV, ERHIAccess::UAVCompute, ScratchBufferUInt.GetExpectedCurrentAccess()));
+	}
+	return Allocation;
+}
+
+//////////////////////////////////////////////////////////////////////////
+
+/** Behaves similar to FGPUScratchPad but for use with RWStructuredBuffers
+*/
+template<typename T>
+class NIAGARASHADER_API FGPUScratchPadStructured
+{
+public:
+	struct NIAGARASHADER_API FAllocation
+	{
+		FRWBufferStructured* Buffer = nullptr;
+		uint32 Offset = 0;
+
+		FORCEINLINE bool IsValid()const { return Buffer != nullptr; }
+	};
+
+	struct NIAGARASHADER_API FBuffer
+	{
+		FRWBufferStructured Buffer;
+		uint32 Used = 0;
+	};
+
+	FGPUScratchPadStructured(uint32 InBucketSize, const FString& InDebugName, EBufferUsageFlags InBufferUsageFlags = EBufferUsageFlags::Static | EBufferUsageFlags::UnorderedAccess)
+		: MinBucketSize(InBucketSize)
+		, DebugName(InDebugName)
+		, BufferUsageFlags(InBufferUsageFlags)
+	{
+	}
+
+	/** Allocates NumElements from the scratch pad. Will try to allocate from an existing buffer but will create a new one if needed. */
+	FAllocation Alloc(uint32 NumElements)
+	{
+		//Look for a buffer to fit this allocation.
+		FBuffer* ToUse = nullptr;
+		for (FBuffer& Buffer : Buffers)
+		{
+			uint32 NumBufferElements = Buffer.Buffer.NumBytes / sizeof(T);
+			if (NumElements + Buffer.Used < NumBufferElements)
+			{
+				ToUse = &Buffer;
+			}
+		}
+
+		if (ToUse == nullptr)
+		{
+			//Allocate a new buffer that will fit this allocation.
+
+			FBuffer* NewBuffer = new FBuffer();
+			Buffers.Add(NewBuffer);
+
+			//Round up the size to pow 2 to ensure a slowly growing user doesn't just keep allocating new buckets. Probably a better guess size here than power of 2.
+			uint32 NewBufferSize = FPlatformMath::RoundUpToPowerOfTwo(FMath::Max(MinBucketSize, NumElements));
+
+			NewBuffer->Buffer.Initialize(
+				*DebugName,
+				sizeof(T),
+				NewBufferSize,
+				BufferUsageFlags,
+				false /*bUseUavCounter*/,
+				false /*bAppendBuffer*/,
+				CurrentAccess);
+
+			ToUse = NewBuffer;
+			NumAllocatedBytes += NewBuffer->Buffer.NumBytes;
+		}
+
+		FAllocation Allocation;
+		Allocation.Buffer = &ToUse->Buffer;
+		Allocation.Offset = ToUse->Used;
+		ToUse->Used += NumElements;
+		NumUsedBytes += NumElements * sizeof(T);
+		return Allocation;
+	}
+
+	/** Resets the scratch pad and releases all buffers it holds. */
+	void Release()
+	{
+		Reset();
+		CurrentAccess = ERHIAccess::UAVCompute;
+		for (FBuffer& Buffer : Buffers)
+		{
+			Buffer.Buffer.Release();
+			Buffer.Used = 0;
+		}
+		Buffers.Reset();
+
+		NumAllocatedBytes = 0;
+	}
+
+	/** Resets the scratch pad, invalidating all previous allocations. Does not modify actual RWBuffer allocations. */
+	void Reset()
+	{
+		for (FBuffer& Buffer : Buffers)
+		{
+			Buffer.Used = 0;
+		}
+
+		NumUsedBytes = 0;
+	}
+
+	/** Transitions all buffers in the scratch pad. */
+	void Transition(FRHICommandList& RHICmdList, ERHIAccess InPreviousState, ERHIAccess InNewState, EResourceTransitionFlags InFlags = EResourceTransitionFlags::None)
+	{
+		//Store off current access state so we can place any new buffers in the same state.
+		CurrentAccess = InNewState;
+
+ 		TArray<FRHITransitionInfo, TMemStackAllocator<>> Transitions;
+ 		Transitions.Reserve(Buffers.Num());
+ 		for (FBuffer& Buffer : Buffers)
+ 		{
+ 			Transitions.Emplace(Buffer.Buffer.UAV, InPreviousState, InNewState, InFlags);
+ 		}
+
+ 		RHICmdList.Transition(Transitions);
+	}
+
+	/** Returns the total bytes allocated for all buffers in the scratch pad. */
+	FORCEINLINE uint32 AllocatedBytes()const
+	{
+		return NumAllocatedBytes;
+	}
+
+	/** Returns the number of used bytes allocated for all buffers in the scratch pad. */
+	FORCEINLINE uint32 UsedBytes()const
+	{
+		return NumUsedBytes;
+	}
+
+private:
+
+	/** Size of each FRWBuffer in the pool. */
+	uint32 MinBucketSize = 4096;
+	TIndirectArray<FBuffer> Buffers;
+	FString DebugName;
+	EBufferUsageFlags BufferUsageFlags = EBufferUsageFlags::None;
+	ERHIAccess CurrentAccess = ERHIAccess::UAVCompute;
+
+	uint32 NumAllocatedBytes = 0;
+	uint32 NumUsedBytes = 0;
+};
+
+//c++ mirror of the ray struct defined in NiagaraRayTracingCommon.ush
+struct FNiagaraRayData : public FBasicRayData
+{
+	int32 CollisionGroup;
+};
+
+// c++ mirror of the struct defined in NiagaraRayTracingCommon.ush
+struct FNiagaraRayTracingPayload
+{
+	float HitT;
+	uint32 PrimitiveIndex;
+	uint32 InstanceIndex;
+	float Barycentrics[2];
+	float WorldPosition[3];
+	float WorldNormal[3];
+};
+
+struct FNiagaraRayTracingResult
+{
+	float HitT;
+
+	float WorldPosition[3];
+	float WorldNormal[3];
+};
+
+/** Holds all information on a single RayTracing dispatch. */
+struct FNiagaraRayTraceDispatchInfo
+{
+	/** Buffer allocation for ray requests to trace. */
+	FGPUScratchPadStructured<FNiagaraRayData>::FAllocation RayRequests;
+	/** Buffer allocation for writing trace results into. */
+	FGPUScratchPadStructured<FNiagaraRayTracingResult>::FAllocation RayTraceIntersections;
+	/** Buffer allocation for last frames trace results to read from in this frame's simulation. */
+	FGPUScratchPadStructured<FNiagaraRayTracingResult>::FAllocation LastFrameRayTraceIntersections;
+	/** Buffer allocation for ray trace counts. Accumulated in simulations shaders and used as dispatch args for ray tracing shaders. */
+	FGPUScratchPad::FAllocation RayCounts;
+	/** Total possible rays to trace. This may be significantly higher than the actual rays request we accumulate into RayCounts. */
+	uint32 MaxRays = 0;
+	/** Max number or times rays can be re-traced in a shader if they hit something that is invalid/filtered. */
+	uint32 MaxRetraces = 0;
+
+	FORCEINLINE bool IsValid() const { return MaxRays > 0; }
+	FORCEINLINE void Reset()
+	{
+		RayRequests = FGPUScratchPadStructured<FNiagaraRayData>::FAllocation();
+		RayTraceIntersections = FGPUScratchPadStructured<FNiagaraRayTracingResult>::FAllocation();
+		RayCounts = FGPUScratchPad::FAllocation();
+		MaxRays = 0;
+		MaxRetraces = 0;
+	}
+};
+
+class NIAGARASHADER_API FNiagaraRayTracingHelper
+{
+public:
+		FNiagaraRayTracingHelper(EShaderPlatform InShaderPlatform);
+		FNiagaraRayTracingHelper() = delete;
+		~FNiagaraRayTracingHelper();
+
+		void Reset();
+		void BuildRayTracingSceneInfo(FRHICommandList& RHICmdList, TConstArrayView<FViewInfo> Views);
+
+		void IssueRayTraces(FRHICommandList& RHICmdList, FScene* Scene);
+		void IssueRayTraces(FRHICommandList& RHICmdList, FScene* Scene, const FIntPoint& RayTraceCounts, uint32 MaxRetraces, FRHIShaderResourceView* RayTraceRequests, FRWBuffer* IndirectArgsBuffer, uint32 IndirectArgsOffset, FRHIUnorderedAccessView* RayTraceResults);
+		bool IsValid() const;
+		
+		void BeginFrame(FRHICommandList& RHICmdList, bool HasRayTracingScene);
+		void EndFrame(FRHICommandList& RHICmdList, bool HasRayTracingScene, FScene* Scene);
+
+		/** Accumulates ray requests from user DIs into single dispatches per DI. */
+		void AddToDispatch(FNiagaraDataInterfaceProxy* DispatchKey, uint32 MaxRays, int32 MaxRetraces);
+
+		/** Ensures that the final buffers for the provided proxy have been allocated. */
+		void BuildDispatch(FRHICommandList& RHICmdList, FNiagaraDataInterfaceProxy* DispatchKey);
+
+		/** Ensures that buffers for a dummy buffer have been allocated. */
+		void BuildDummyDispatch(FRHICommandList& RHICmdList);
+
+		/** Returns the final buffers for each DI for use in simulations and RT dispatches. */
+		const FNiagaraRayTraceDispatchInfo& GetDispatch(FNiagaraDataInterfaceProxy* DispatchKey) const;
+
+		/** Returns a dummy dispatch to allow shaders to still bind to valid buffers. */
+		const FNiagaraRayTraceDispatchInfo& GetDummyDispatch() const;
+
+		/** Adds a primitive to a collision group. This data is sent to the GPU on the start of the next frame. */
+		void SetPrimitiveCollisionGroup(FPrimitiveSceneInfo& Primitive, uint32 CollisionGroup);
+
+		/** Pushes cpu side copy of the collision group map to the GPU. */
+		void UpdateCollisionGroupMap(FRHICommandList& RHICmdList, FScene* Scene, ERHIFeatureLevel::Type FeatureLevel);
+
+		void RefreshPrimitiveInstanceData();
+
+		FGPUScratchPadStructured<FNiagaraRayData>& GetRayRequests(){ return RayRequests; }
+		FGPUScratchPadStructured<FNiagaraRayTracingResult> GetRayTraceIntersections() { return RayTraceIntersections; }
+		FGPUScratchPad GetRayTraceCounts() { return RayTraceCounts; }
+
+	private:
+		const EShaderPlatform ShaderPlatform;
+
+		mutable FRayTracingPipelineState* RayTracingPipelineState = nullptr;
+		mutable FRHIRayTracingScene* RayTracingScene = nullptr;
+		mutable FRHIShaderResourceView* RayTracingSceneView = nullptr;
+		mutable FRHIUniformBuffer* ViewUniformBuffer = nullptr;
+
+		/**
+		CPU side copy of the PrimID to Collision Group Map.
+		This is uploaded to the GPU and used in Niagara RG shader to filter self collisions between objects of the same group.
+		*/
+		TMap<FPrimitiveComponentId, uint32> CollisionGroupMap;
+		
+		/** Hash table. 
+		 PrimIdHashTable is the main hash table that maps GPUSceneInstanceIndex to and Index we can use to store Collision Groups inside HashToCollisionGroups.
+*/
+		FRWBuffer PrimIdHashTable;
+		uint32 HashTableSize = 0;
+		FRWBuffer HashToCollisionGroups;
+
+		bool bCollisionGroupMapDirty = true;
+
+		/** Scratch buffer holding all ray trace requests for all simulations in the scene. */
+		FGPUScratchPadStructured<FNiagaraRayData> RayRequests;
+		/** Scratch buffer holding all ray trace results for all simulations in the scene. */
+		FGPUScratchPadStructured<FNiagaraRayTracingResult> RayTraceIntersections;
+		/** Scratch buffer holding all accumulated ray request counts for all simulations in the scene. */
+		FGPUScratchPad RayTraceCounts;
+
+		/** Map of DI Proxy to Dispatch Info. Each DI will have a single dispatch for all it's instances. */
+		TMap<FNiagaraDataInterfaceProxy*, FNiagaraRayTraceDispatchInfo> Dispatches;
+
+		/** Last frame's DI proxy map. Allows use to retrieve the buffer allocations to read last frames trace results. TODO: Improve. Ideally don't need a map copy here. */
+		TMap<FNiagaraDataInterfaceProxy*, FNiagaraRayTraceDispatchInfo> PreviousFrameDispatches;
+
+		/** In places where ray tracing is disabled etc we use a dummy dispatch allocation set so the simulations can bind to valid buffers. */
+		FNiagaraRayTraceDispatchInfo DummyDispatch;
+};
+
+#endif
