improve rasterization grid
- bug with setting int values
- options for interlocked add with ints
- option to detect overflow

diff --git a/Engine/Plugins/FX/Niagara/Source/Niagara/Private/NiagaraDataInterfaceRasterizationGrid3D.cpp b/Engine/Plugins/FX/Niagara/Source/Niagara/Private/NiagaraDataInterfaceRasterizationGrid3D.cpp
--- a/Engine/Plugins/FX/Niagara/Source/Niagara/Private/NiagaraDataInterfaceRasterizationGrid3D.cpp
+++ b/Engine/Plugins/FX/Niagara/Source/Niagara/Private/NiagaraDataInterfaceRasterizationGrid3D.cpp
@@ -22,10 +22,18 @@ const FString UNiagaraDataInterfaceRasterizationGrid3D::PerAttributeDataName(TEX
 
 static const FName SetFloatValueFunctionName("SetFloatGridValue");
 static const FName GetFloatValueFunctionName("GetFloatGridValue");
+
 static const FName InterlockedAddFloatGridValueFunctionName("InterlockedAddFloatGridValue");
+static const FName InterlockedAddIntGridValueFunctionName("InterlockedAddIntGridValue");
+static const FName InterlockedAddFloatGridValueSafeFunctionName("InterlockedAddFloatGridValueSafe");
+
+
 static const FName InterlockedMinFloatGridValueFunctionName("InterlockedMinFloatGridValue");
 static const FName InterlockedMaxFloatGridValueFunctionName("InterlockedMaxFloatGridValue");
 
+
+
+
 static const FName IntToFloatFunctionName("IntToFloat");
 static const FName FloatToIntFunctionName("FloatToInt");
 
@@ -289,6 +297,56 @@ void UNiagaraDataInterfaceRasterizationGrid3D::GetFunctions(TArray<FNiagaraFunct
 		OutFunctions.Add(Sig);
 	}
 
+	{
+		FNiagaraFunctionSignature Sig;
+		Sig.Name = InterlockedAddIntGridValueFunctionName;
+		Sig.Inputs.Add(FNiagaraVariable(FNiagaraTypeDefinition(GetClass()), TEXT("Grid")));
+		Sig.Inputs.Add(FNiagaraVariable(FNiagaraTypeDefinition::GetIntDef(), TEXT("IndexX")));
+		Sig.Inputs.Add(FNiagaraVariable(FNiagaraTypeDefinition::GetIntDef(), TEXT("IndexY")));
+		Sig.Inputs.Add(FNiagaraVariable(FNiagaraTypeDefinition::GetIntDef(), TEXT("IndexZ")));
+		Sig.Inputs.Add(FNiagaraVariable(FNiagaraTypeDefinition::GetIntDef(), TEXT("AttributeIndex")));
+		Sig.Inputs.Add(FNiagaraVariable(FNiagaraTypeDefinition::GetIntDef(), TEXT("Value")));
+		Sig.Outputs.Add(FNiagaraVariable(FNiagaraTypeDefinition::GetIntDef(), TEXT("IGNORE")));
+
+		Sig.bExperimental = true;
+		Sig.bMemberFunction = true;
+		Sig.bRequiresContext = false;
+		Sig.bWriteFunction = true;
+		Sig.ModuleUsageBitmask = ENiagaraScriptUsageMask::Particle;
+
+		Sig.bSupportsCPU = false;
+		Sig.bSupportsGPU = true;
+#if WITH_EDITORONLY_DATA
+		Sig.Description = NSLOCTEXT("Niagara", "NiagaraDataInterfaceGridColl2D_SetValueFunction", "Set the value at a specific index. Note that this is an older way of working with Grids. Consider using the SetFloat or other typed, named functions or parameter map variables with StackContext namespace instead.");
+#endif
+		OutFunctions.Add(Sig);
+	}
+
+	{
+		FNiagaraFunctionSignature Sig;
+		Sig.Name = InterlockedAddFloatGridValueSafeFunctionName;
+		Sig.Inputs.Add(FNiagaraVariable(FNiagaraTypeDefinition(GetClass()), TEXT("Grid")));
+		Sig.Inputs.Add(FNiagaraVariable(FNiagaraTypeDefinition::GetIntDef(), TEXT("IndexX")));
+		Sig.Inputs.Add(FNiagaraVariable(FNiagaraTypeDefinition::GetIntDef(), TEXT("IndexY")));
+		Sig.Inputs.Add(FNiagaraVariable(FNiagaraTypeDefinition::GetIntDef(), TEXT("IndexZ")));
+		Sig.Inputs.Add(FNiagaraVariable(FNiagaraTypeDefinition::GetIntDef(), TEXT("AttributeIndex")));
+		Sig.Inputs.Add(FNiagaraVariable(FNiagaraTypeDefinition::GetFloatDef(), TEXT("Value")));
+		Sig.Outputs.Add(FNiagaraVariable(FNiagaraTypeDefinition::GetIntDef(), TEXT("IsSafe")));
+
+		Sig.bExperimental = true;
+		Sig.bMemberFunction = true;
+		Sig.bRequiresContext = false;
+		Sig.bWriteFunction = true;
+		Sig.ModuleUsageBitmask = ENiagaraScriptUsageMask::Particle;
+
+		Sig.bSupportsCPU = false;
+		Sig.bSupportsGPU = true;
+#if WITH_EDITORONLY_DATA
+		Sig.Description = NSLOCTEXT("Niagara", "NiagaraDataInterfaceGridColl2D_SetValueFunction", "Set the value at a specific index. Note that this is an older way of working with Grids. Consider using the SetFloat or other typed, named functions or parameter map variables with StackContext namespace instead.");
+#endif
+		OutFunctions.Add(Sig);
+	}
+
 	{
 		FNiagaraFunctionSignature Sig;
 		Sig.Name = InterlockedMinFloatGridValueFunctionName;
@@ -470,7 +528,7 @@ void UNiagaraDataInterfaceRasterizationGrid3D::GetParameterDefinitionHLSL(const
 		static const TCHAR* FormatHLSL = TEXT(R"(
 				float {IntToFloatFunction}(int IntValue)
 				{
-					return 1. * IntValue / {Precision};
+					return float(IntValue) / {Precision};
 				}
 			)");
 		TMap<FString, FStringFormatArg> FormatArgs =
@@ -542,7 +600,6 @@ bool UNiagaraDataInterfaceRasterizationGrid3D::GetFunctionHLSL(const FNiagaraDat
 			void {FunctionName}(int In_IndexX, int In_IndexY, int In_IndexZ, int In_AttributeIndex, float In_Value, out int val)
 			{			
 				val = 0;
-				Out_Val = 0;
 				if ( In_AttributeIndex < {NumAttributesName} )
 				{
 					int3 TileOffset = {PerAttributeDataName}[In_AttributeIndex].xyz;
@@ -586,6 +643,50 @@ bool UNiagaraDataInterfaceRasterizationGrid3D::GetFunctionHLSL(const FNiagaraDat
 		OutHLSL += FString::Format(FormatBounds, ArgsBounds);
 		return true;
 	}
+	else if (FunctionInfo.DefinitionName == InterlockedAddIntGridValueFunctionName)
+	{
+		static const TCHAR* FormatBounds = TEXT(R"(
+			void {FunctionName}(int In_IndexX, int In_IndexY, int In_IndexZ, int In_AttributeIndex, int In_Value, out int val)
+			{							
+				val = 0;					
+				if ( In_AttributeIndex < {NumAttributesName} )
+				{
+					int3 TileOffset = {PerAttributeDataName}[In_AttributeIndex].xyz;		
+					InterlockedAdd(RW{OutputIntGrid}[int3(In_IndexX + TileOffset.x, In_IndexY + TileOffset.y, In_IndexZ + TileOffset.z)], In_Value);
+				}
+			}
+		)");
+		OutHLSL += FString::Format(FormatBounds, ArgsBounds);
+		return true;
+	}
+	else if (FunctionInfo.DefinitionName == InterlockedAddFloatGridValueSafeFunctionName)
+	{
+		static const TCHAR* FormatBounds = TEXT(R"(
+			void {FunctionName}(int In_IndexX, int In_IndexY, int In_IndexZ, int In_AttributeIndex, float In_Value, out int val)
+			{							
+				val = -1;					
+				if ( In_AttributeIndex < {NumAttributesName} )
+				{
+					int OriginalValue;
+					int IntValue = {FloatToIntFunctionName}(In_Value);
+					int3 TileOffset = {PerAttributeDataName}[In_AttributeIndex].xyz;		
+					InterlockedAdd(RW{OutputIntGrid}[int3(In_IndexX + TileOffset.x, In_IndexY + TileOffset.y, In_IndexZ + TileOffset.z)], IntValue, OriginalValue);
+
+					if ((IntValue > 0 && IntValue + OriginalValue < OriginalValue) || 
+						(IntValue < 0 && IntValue + OriginalValue > OriginalValue))
+					{
+						val = 0;
+					}
+					else
+					{
+						val = 1;
+					}
+				}
+			}
+		)");
+		OutHLSL += FString::Format(FormatBounds, ArgsBounds);
+		return true;
+	}
 	else if (FunctionInfo.DefinitionName == InterlockedMinFloatGridValueFunctionName)
 	{
 		static const TCHAR* FormatBounds = TEXT(R"(
