Build fix
#rnx

diff --git a/Engine/Plugins/FX/Niagara/Source/NiagaraShader/Private/NiagaraRayTracingHelper.cpp b/Engine/Plugins/FX/Niagara/Source/NiagaraShader/Private/NiagaraRayTracingHelper.cpp
--- a/Engine/Plugins/FX/Niagara/Source/NiagaraShader/Private/NiagaraRayTracingHelper.cpp
+++ b/Engine/Plugins/FX/Niagara/Source/NiagaraShader/Private/NiagaraRayTracingHelper.cpp
@@ -501,7 +501,7 @@ void FNiagaraRayTracingHelper::UpdateCollisionGroupMap(FRHICommandList& RHICmdLi
 		// To simplify the shader code, the size of the table must be a multiple of the thread count.
 		check(HashTableSize % FNiagaraUpdateCollisionGroupMapCS::THREAD_COUNT == 0);
 
-		RHICmdList.SetComputeShader(ShaderRHI);
+		SetComputePipelineState(RHICmdList, ShaderRHI);
 		SetShaderParameters(RHICmdList, ComputeShader, ShaderRHI, Params);
 		RHICmdList.DispatchComputeShader(FMath::DivideAndRoundUp(CollisionGroupMap.Num(), (int32)FNiagaraUpdateCollisionGroupMapCS::THREAD_COUNT), 1, 1);
 		UnsetShaderUAVs(RHICmdList, ComputeShader, ShaderRHI);
