Move to SetComputePipelineState for NiagaraAsyncGpuTraceProvider

diff --git a/Engine/Plugins/FX/Niagara/Source/NiagaraShader/Private/NiagaraAsyncGpuTraceProvider.cpp b/Engine/Plugins/FX/Niagara/Source/NiagaraShader/Private/NiagaraAsyncGpuTraceProvider.cpp
--- a/Engine/Plugins/FX/Niagara/Source/NiagaraShader/Private/NiagaraAsyncGpuTraceProvider.cpp
+++ b/Engine/Plugins/FX/Niagara/Source/NiagaraShader/Private/NiagaraAsyncGpuTraceProvider.cpp
@@ -278,7 +278,7 @@ void FNiagaraAsyncGpuTraceProvider::BuildCollisionGroupHashMap(FRHICommandList&
 	// To simplify the shader code, the size of the table must be a multiple of the thread count.
 	check(Result.HashTableSize % FNiagaraUpdateCollisionGroupMapCS::THREAD_COUNT == 0);
 
-	RHICmdList.SetComputeShader(ShaderRHI);
+	SetComputePipelineState(RHICmdList, ShaderRHI);
 	SetShaderParameters(RHICmdList, ComputeShader, ShaderRHI, Params);
 	RHICmdList.DispatchComputeShader(FMath::DivideAndRoundUp(CollisionGroupMap.Num(), (int32)FNiagaraUpdateCollisionGroupMapCS::THREAD_COUNT), 1, 1);
 	UnsetShaderUAVs(RHICmdList, ComputeShader, ShaderRHI);
