Remove struct member function for FNiagaraAsyncGpuTraceResult

diff --git a/Engine/Plugins/FX/Niagara/Shaders/Private/NiagaraAsyncGpuTraceCommon.ush b/Engine/Plugins/FX/Niagara/Shaders/Private/NiagaraAsyncGpuTraceCommon.ush
--- a/Engine/Plugins/FX/Niagara/Shaders/Private/NiagaraAsyncGpuTraceCommon.ush
+++ b/Engine/Plugins/FX/Niagara/Shaders/Private/NiagaraAsyncGpuTraceCommon.ush
@@ -20,6 +20,9 @@ struct FNiagaraAsyncGpuTraceResult
 	float3 WorldNormal;
 
 	float _Pad0; // padding to force 16 byte alignment to meet requirements for VK
+};
 
-	bool IsHit() { return HitT > 0; }
-};
\ No newline at end of file
+bool IsHit(FNiagaraAsyncGpuTraceResult Result)
+{
+	return Result.HitT > 0;
+}
\ No newline at end of file
