Made auto-converting linked inputs in the stack optional (defaulting to off)

diff --git a/Engine/Plugins/FX/Niagara/Source/Niagara/Public/NiagaraSettings.h b/Engine/Plugins/FX/Niagara/Source/Niagara/Public/NiagaraSettings.h
--- a/Engine/Plugins/FX/Niagara/Source/Niagara/Public/NiagaraSettings.h
+++ b/Engine/Plugins/FX/Niagara/Source/Niagara/Public/NiagaraSettings.h
@@ -91,6 +91,10 @@ class NIAGARA_API UNiagaraSettings : public UDeveloperSettings
 	/** Sets the default navigation behavior for the system preview viewport. */
 	UPROPERTY(config, EditAnywhere, Category = Viewport)
 	bool bSystemViewportInOrbitMode = true;
+
+	/** If true then the "link input" menu will also show variables of different types, as long as there is a conversion script for them. */
+	UPROPERTY(config, EditAnywhere, Category = Niagara)
+	bool bShowConvertibleInputsInStack = false;
 #endif // WITH_EDITORONLY_DATA
 
 	/** If true then active effects rebase the simulation positions to not lose precision. Can be turned off if not needed to skip unnecessary rebasing calculations. */
diff --git a/Engine/Plugins/FX/Niagara/Source/NiagaraEditorWidgets/Private/Stack/SNiagaraStackFunctionInputValue.cpp b/Engine/Plugins/FX/Niagara/Source/NiagaraEditorWidgets/Private/Stack/SNiagaraStackFunctionInputValue.cpp
--- a/Engine/Plugins/FX/Niagara/Source/NiagaraEditorWidgets/Private/Stack/SNiagaraStackFunctionInputValue.cpp
+++ b/Engine/Plugins/FX/Niagara/Source/NiagaraEditorWidgets/Private/Stack/SNiagaraStackFunctionInputValue.cpp
@@ -50,6 +50,7 @@
 #include "Widgets/Input/SNumericEntryBox.h"
 #include "NiagaraNodeAssignment.h"
 #include "NiagaraCommon.h"
+#include "NiagaraSettings.h"
 
 #define LOCTEXT_NAMESPACE "NiagaraStackFunctionInputValue"
 
@@ -1109,6 +1110,11 @@ TArray<TSharedPtr<FNiagaraMenuAction_Generic>> SNiagaraStackFunctionInputValue::
 		}
 	}
 
+	const UNiagaraSettings* Settings = GetDefault<UNiagaraSettings>();
+	if (ensure(Settings) && Settings->bShowConvertibleInputsInStack == false)
+	{
+		AvailableConversionHandles.Empty();
+	}
 	const FText ConvertInputFormat = LOCTEXT("ConvertInputFormat", "Link this input to {0} via a conversion script");
 	for (const auto& Entry : AvailableConversionHandles)
 	{
