- Niagara vertex factory fixes for LWC

#rnx

diff --git a/Engine/Plugins/FX/Niagara/Shaders/Private/NiagaraMeshVertexFactory.ush b/Engine/Plugins/FX/Niagara/Shaders/Private/NiagaraMeshVertexFactory.ush
--- a/Engine/Plugins/FX/Niagara/Shaders/Private/NiagaraMeshVertexFactory.ush
+++ b/Engine/Plugins/FX/Niagara/Shaders/Private/NiagaraMeshVertexFactory.ush
@@ -204,9 +204,9 @@ struct FVertexFactoryIntermediates
 	float ParticleRandom;
 #endif
 
-	float4x4 ParticleToWorld;
-	float4x4 WorldToParticle;
-	float4x4 PrevParticleToWorld;	
+	FLWCMatrix ParticleToWorld;
+	FLWCInverseMatrix WorldToParticle;
+	FLWCMatrix PrevParticleToWorld;	
 	float3x3 ParticleToWorldNoScale;
 
 	float3x3 TangentToLocal;
@@ -221,24 +221,24 @@ FPrimitiveSceneData GetPrimitiveData(FVertexFactoryIntermediates Intermediates)
 	return Intermediates.SceneData.Primitive;
 }
 
-float4x4 GetParticleTransform(FVertexFactoryIntermediates Inters)
+FLWCMatrix GetParticleTransform(FVertexFactoryIntermediates Inters)
 {
 	return Inters.ParticleToWorld;
 }
 
-float3x3 GetParticleRotationNoScale(FVertexFactoryIntermediates Inters)
+FLWCInverseMatrix GetParticleInvTransform(FVertexFactoryIntermediates Inters)
 {
-	return Inters.ParticleToWorldNoScale;
+	return Inters.WorldToParticle;
 }
 
-float4x4 GetParticleInvTransform(FVertexFactoryIntermediates Inters)
+FLWCMatrix GetParticlePrevTransform(FVertexFactoryIntermediates Inters)
 {
-	return Inters.WorldToParticle;
+	return Inters.PrevParticleToWorld;
 }
 
-float4x4 GetParticlePrevTransform(FVertexFactoryIntermediates Inters)
+float3x3 GetParticleRotationNoScale(FVertexFactoryIntermediates Inters)
 {
-	return Inters.PrevParticleToWorld;
+	return Inters.ParticleToWorldNoScale;
 }
 
 /** Converts from vertex factory specific interpolants to a FMaterialPixelParameters, which is used by material inputs. */
@@ -301,7 +301,7 @@ FMaterialPixelParameters GetMaterialPixelParameters(FVertexFactoryInterpolantsVS
 #endif
 
 #if USE_PARTICLE_POSITION
-	Result.Particle.TranslatedWorldPositionAndSize.xyz = Interpolants.ParticleTranslatedWorldPosition.xyz;
+	Result.Particle.TranslatedWorldPositionAndSize.xyz = Interpolants.ParticleTranslatedWorldPosition;
 	Result.Particle.TranslatedWorldPositionAndSize.w = 1;
 #endif
 
@@ -320,6 +320,7 @@ FMaterialPixelParameters GetMaterialPixelParameters(FVertexFactoryInterpolantsVS
 #endif
 
 #if USE_PARTICLE_LOCAL_TO_WORLD
+	//-TODO:
 	FLWCMatrix ParticleToWorld = LWCPromote(transpose(float4x4(Interpolants.ParticleToWorld[0], Interpolants.ParticleToWorld[1], Interpolants.ParticleToWorld[2], float4(0.0f, 0.0f, 0.0f, 1.0f))));
 	#if NEEDS_PARTICLE_LOCAL_TO_WORLD
 		Result.Particle.ParticleToWorld = ParticleToWorld;
@@ -330,6 +331,7 @@ FMaterialPixelParameters GetMaterialPixelParameters(FVertexFactoryInterpolantsVS
 #endif
 
 #if USE_PARTICLE_WORLD_TO_LOCAL
+	//-TODO:
 	FLWCInverseMatrix WorldToParticle = LWCPromoteInverse(transpose(float4x4(Interpolants.WorldToParticle[0], Interpolants.WorldToParticle[1], Interpolants.WorldToParticle[2], float4(0.0f, 0.0f, 0.0f, 1.0f))));
 	#if NEEDS_PARTICLE_WORLD_TO_LOCAL
 		Result.Particle.WorldToParticle = WorldToParticle;
@@ -361,7 +363,7 @@ FMaterialVertexParameters GetMaterialVertexParameters(FVertexFactoryInput Input,
 	Result.PreSkinnedNormal = Intermediates.TangentToLocal[2];
 
 #if NEEDS_PARTICLE_POSITION
-	Result.Particle.TranslatedWorldPositionAndSize.xyz = Intermediates.ParticleTranslatedWorldPosition.xyz;
+	Result.Particle.TranslatedWorldPositionAndSize.xyz = Intermediates.ParticleTranslatedWorldPosition;
 	Result.Particle.TranslatedWorldPositionAndSize.w = 1;
 #endif
 
@@ -397,8 +399,8 @@ FMaterialVertexParameters GetMaterialVertexParameters(FVertexFactoryInput Input,
 	Result.Particle.DynamicParameter3 = Intermediates.DynamicParameter3;
 #endif
 
-	Result.InstanceLocalToWorld = LWCPromote(GetParticleTransform(Intermediates));
-	Result.InstanceWorldToLocal = LWCPromoteInverse(GetParticleInvTransform(Intermediates));
+	Result.InstanceLocalToWorld = GetParticleTransform(Intermediates);
+	Result.InstanceWorldToLocal = GetParticleInvTransform(Intermediates);
 
 	Result.Particle.ParticleToWorld = Result.InstanceLocalToWorld;
 	Result.Particle.WorldToParticle = Result.InstanceWorldToLocal;
@@ -406,7 +408,7 @@ FMaterialVertexParameters GetMaterialVertexParameters(FVertexFactoryInput Input,
 	// Use the particle's unscaled transform for rotating normals to match static mesh
 	Result.TangentToWorld = mul(TangentToLocal, GetParticleRotationNoScale(Intermediates));
 
-	Result.PrevFrameLocalToWorld = LWCPromote(GetParticlePrevTransform(Intermediates));
+	Result.PrevFrameLocalToWorld = GetParticlePrevTransform(Intermediates);
 
 #if NUM_MATERIAL_TEXCOORDS_VERTEX
 	for(int CoordinateIndex = 0; CoordinateIndex < NUM_MATERIAL_TEXCOORDS_VERTEX; CoordinateIndex++)
@@ -496,12 +498,12 @@ FVertexFactoryIntermediates GetVertexFactoryIntermediates(FVertexFactoryInput In
 
 	// Compute transforms
 #if (VF_USE_PRIMITIVE_SCENE_DATA == 1)
-	Intermediates.ParticleToWorld = LWCHackToFloat(Intermediates.SceneData.InstanceData.LocalToWorld);
-	Intermediates.WorldToParticle = LWCHackToFloat(Intermediates.SceneData.InstanceData.WorldToLocal);
-	Intermediates.PrevParticleToWorld = LWCHackToFloat(Intermediates.SceneData.InstanceData.PrevLocalToWorld);
+	Intermediates.ParticleToWorld = Intermediates.SceneData.InstanceData.LocalToWorld;
+	Intermediates.WorldToParticle = Intermediates.SceneData.InstanceData.WorldToLocal;
+	Intermediates.PrevParticleToWorld = Intermediates.SceneData.InstanceData.PrevLocalToWorld;
 
 	const float3 InstInvScale = Intermediates.SceneData.InstanceData.InvNonUniformScale;
-	Intermediates.ParticleToWorldNoScale = (float3x3)Intermediates.ParticleToWorld;
+	Intermediates.ParticleToWorldNoScale = LWCToFloat3x3(Intermediates.ParticleToWorld);
 	Intermediates.ParticleToWorldNoScale[0] *= InstInvScale.x;
 	Intermediates.ParticleToWorldNoScale[1] *= InstInvScale.y;
 	Intermediates.ParticleToWorldNoScale[2] *= InstInvScale.z;
@@ -555,14 +557,14 @@ FVertexFactoryIntermediates GetVertexFactoryIntermediates(FVertexFactoryInput In
 
 	NiagaraMeshParticleTransforms Transforms = NiagaraCalculateMeshParticleTransforms(Params);
 
-	Intermediates.ParticleToWorld 			= LWCHackToFloat(Transforms.LocalToWorld);
-	Intermediates.WorldToParticle 			= LWCHackToFloat(Transforms.WorldToLocal);
-	Intermediates.PrevParticleToWorld		= LWCHackToFloat(Transforms.PrevLocalToWorld);
+	Intermediates.ParticleToWorld 			= Transforms.LocalToWorld;
+	Intermediates.WorldToParticle 			= Transforms.WorldToLocal;
+	Intermediates.PrevParticleToWorld		= Transforms.PrevLocalToWorld;
 	Intermediates.ParticleToWorldNoScale	= Transforms.LocalToWorldNoScale;
 #endif // (VF_USE_PRIMITIVE_SCENE_DATA == 1)
 
 	// Particle translated world position.
-	Intermediates.ParticleTranslatedWorldPosition = Intermediates.ParticleToWorld[3].xyz + LWCHackToFloat(ResolvedView.PreViewTranslation);
+	Intermediates.ParticleTranslatedWorldPosition = LWCToFloat(LWCAdd(LWCGetOrigin(Intermediates.ParticleToWorld), ResolvedView.PreViewTranslation));
 
 #if NUM_MATERIAL_TEXCOORDS_VERTEX
 	#if MANUAL_VERTEX_FETCH
@@ -650,8 +652,8 @@ float3x3 VertexFactoryGetTangentToLocal( FVertexFactoryInput Input, FVertexFacto
 // @return translated world position
 float4 VertexFactoryGetWorldPosition(FVertexFactoryInput Input, FVertexFactoryIntermediates Intermediates)
 {	
-	float4x4 Transform = GetParticleTransform(Intermediates);
-	float3 WorldPosition = mul(Input.Position, Transform).xyz + LWCHackToFloat(ResolvedView.PreViewTranslation);
+	FLWCMatrix Transform = GetParticleTransform(Intermediates);
+	float3 WorldPosition = LWCToFloat(LWCAdd(LWCMultiply(Input.Position.xyz, Transform), ResolvedView.PreViewTranslation));
 	return float4(WorldPosition, Input.Position.w);
 }
 
@@ -662,7 +664,7 @@ float4 VertexFactoryGetRasterizedWorldPosition(FVertexFactoryInput Input, FVerte
 
 float3 VertexFactoryGetWorldNormal(FVertexFactoryInput Input, FVertexFactoryIntermediates Intermediates)
 {
-	return Intermediates.ParticleToWorld[2].xyz;
+	return LWCToFloat3x3(Intermediates.ParticleToWorld)[2].xyz;
 }
 
 float3 VertexFactoryGetPositionForVertexLighting(FVertexFactoryInput Input, FVertexFactoryIntermediates Intermediates, float3 TranslatedWorldPosition)
@@ -791,8 +793,8 @@ FVertexFactoryInterpolantsVSToPS VertexFactoryGetInterpolantsVSToPS(FVertexFacto
 // @return previous translated world position
 float4 VertexFactoryGetPreviousWorldPosition(FVertexFactoryInput Input, FVertexFactoryIntermediates Intermediates)
 {
-    float4x4 Transform = GetParticlePrevTransform(Intermediates);
-    float3 WorldPosition = mul(Input.Position, Transform).xyz + LWCHackToFloat(ResolvedView.PrevPreViewTranslation);
+    FLWCMatrix Transform = GetParticlePrevTransform(Intermediates);
+    float3 WorldPosition = LWCToFloat(LWCAdd(LWCMultiply(Input.Position.xyz, Transform), ResolvedView.PreViewTranslation));
     return float4(WorldPosition, Input.Position.w);
 }
 
diff --git a/Engine/Plugins/FX/Niagara/Shaders/Private/NiagaraRibbonVertexFactory.ush b/Engine/Plugins/FX/Niagara/Shaders/Private/NiagaraRibbonVertexFactory.ush
--- a/Engine/Plugins/FX/Niagara/Shaders/Private/NiagaraRibbonVertexFactory.ush
+++ b/Engine/Plugins/FX/Niagara/Shaders/Private/NiagaraRibbonVertexFactory.ush
@@ -59,15 +59,15 @@ float UnpackParticleVertexTextureV(uint InterpIdx)
 }
 
 
-float3 TransformPosition(float3 InPosition)
+FLWCVector3 TransformPosition(float3 InPosition)
 {
 	if (NiagaraRibbonVF.bLocalSpace)
 	{
-		return LWCHackToFloat(TransformLocalToWorld(InPosition));
+		return TransformLocalToWorld(InPosition);
 	}
 	else
 	{
-		return LWCHackToFloat(MakeLWCVector3(NiagaraRibbonVF.SystemLWCTile, InPosition));
+		return MakeLWCVector3(NiagaraRibbonVF.SystemLWCTile, InPosition);
 	}
 }
 
@@ -142,8 +142,7 @@ struct FSegmentData
 
 struct FVertexFactoryIntermediates
 {
-	float3 Position;
-	float3 OldPosition;
+	FLWCVector3 Position;
 	float3 Direction;
 	float DistanceOnSegment;
 	float4 Color;
@@ -413,7 +412,7 @@ FMaterialVertexParameters GetMaterialVertexParameters(FVertexFactoryInput Input,
 	Result.VertexColor = Intermediates.Color;
 	Result.TangentToWorld = mul(TangentToLocal, GetLocalToWorld3x3()); 
 	Result.Particle.Color = Intermediates.Color;
-	Result.PreSkinnedPosition = Intermediates.Position.xyz;
+	Result.PreSkinnedPosition = LWCToFloat(Intermediates.Position);
 	Result.PreSkinnedNormal = TangentToLocal[2].xyz;
 
 #if (DYNAMIC_PARAMETERS_MASK != 0)
@@ -473,15 +472,12 @@ float3 SafeNormalize(float3 V)
 
 float4 CalcWorldPosition(FVertexFactoryInput Input, FVertexFactoryIntermediates Intermediates)
 {
-	{
-		float2 SliceVertexPositions = UnpackParticleVertexPosition(Input.InterpVtxId);
-				
-		float3 Pos = Intermediates.Position.xyz +
-			((Intermediates.GeoRight * SliceVertexPositions.xxx) + (Intermediates.GeoUp * SliceVertexPositions.yyy)) * Intermediates.RibbonWidth.xxx;
+	const float2 SliceVertexPositions = UnpackParticleVertexPosition(Input.InterpVtxId);
+	const float3 Offset = ((Intermediates.GeoRight * SliceVertexPositions.xxx) + (Intermediates.GeoUp * SliceVertexPositions.yyy)) * Intermediates.RibbonWidth.xxx;
+	const FLWCVector3 Pos = LWCAdd(LWCAdd(Intermediates.Position, LWCPromote(Offset)), ResolvedView.PreViewTranslation);
 		
-		// Intermediates are already in world space, only need to account for the translation around view origin.
-		return float4(Pos + LWCHackToFloat(ResolvedView.PreViewTranslation), 1);
-	}
+	// Intermediates are already in world space, only need to account for the translation around view origin.
+	return float4(LWCToFloat(Pos), 1);
 }
 
 /** derive basis vectors */
@@ -663,7 +659,7 @@ FVertexFactoryIntermediates GetVertexFactoryIntermediates(FVertexFactoryInput In
 	}
 	else
 	{
-		Facing = normalize(LWCHackToFloat(ResolvedView.WorldCameraOrigin) - Intermediates.Position);
+		Facing = normalize(LWCToFloat(LWCSubtract(ResolvedView.WorldCameraOrigin, Intermediates.Position)));
 	}
 
 	float3 Right = normalize(cross(Facing, Intermediates.Direction));
@@ -712,13 +708,13 @@ FVertexFactoryIntermediates GetVertexFactoryIntermediates(FVertexFactoryInput In
 	
 	// Flip the normal only for multiplane shape if the normal is looking away from the camera.
 	// This lets us still render as two sided, but at least gets the normals corrected so the plane faces the camera.
-	float3 ViewNormal = normalize(LWCHackToFloat(ResolvedView.WorldCameraOrigin) - Intermediates.Position);	
+	float3 ViewNormal = normalize(LWCToFloat(LWCSubtract(ResolvedView.WorldCameraOrigin, Intermediates.Position)));
 	Intermediates.TangentZ *= (NiagaraRibbonVF.ShouldFlipNormalToView > 0 && dot(ViewNormal, Intermediates.TangentZ) < 0) ? -1 : 1;
 	
 	Intermediates.TangentX = cross(Intermediates.TangentZ, Intermediates.TangentY);
 
 #if USE_PARTICLE_POSITION
-	Intermediates.TranslatedWorldPositionAndSize = float4(Intermediates.Position + LWCHackToFloat(ResolvedView.PreViewTranslation),  .5 * Intermediates.RibbonWidth);
+	Intermediates.TranslatedWorldPositionAndSize = float4(LWCToFloat(LWCAdd(Intermediates.Position, ResolvedView.PreViewTranslation)),  .5 * Intermediates.RibbonWidth);
 #endif
 
 #if USE_PARTICLE_TIME
diff --git a/Engine/Plugins/FX/Niagara/Shaders/Private/NiagaraSpriteVertexFactory.ush b/Engine/Plugins/FX/Niagara/Shaders/Private/NiagaraSpriteVertexFactory.ush
--- a/Engine/Plugins/FX/Niagara/Shaders/Private/NiagaraSpriteVertexFactory.ush
+++ b/Engine/Plugins/FX/Niagara/Shaders/Private/NiagaraSpriteVertexFactory.ush
@@ -162,9 +162,9 @@ struct FVertexFactoryInterpolantsVSToPS
 struct FVertexFactoryIntermediates
 {
 	/** The position of the vertex in translated world space. */
-	float3 VertexWorldPosition;
+	FLWCVector3 VertexWorldPosition;
 	/** Last frame's position of the vertex in translated world space. */
-	float3 PreviousVertexWorldPosition;
+	FLWCVector3 PreviousVertexWorldPosition;
 	/** Particle translated world space position and size. */
 	float4 TranslatedWorldPositionAndSize;
 	
@@ -188,7 +188,7 @@ struct FVertexFactoryIntermediates
 	float4 ParticleVelocity;
 
 	/** The particle worldspace position */
-	float3 Position;
+	FLWCVector3 Position;
 	
 	float4 SizeRotSubImage;
 	
@@ -368,7 +368,7 @@ FMaterialVertexParameters GetMaterialVertexParameters(FVertexFactoryInput Input,
 	Result.Particle.MacroUV = NiagaraSpriteVF.MacroUVParameters;
 	Result.Particle.Color = Intermediates.Color;
 	Result.Particle.MotionBlurFade = 1.0f;
-	Result.PreSkinnedPosition = Intermediates.Position.xyz;
+	Result.PreSkinnedPosition = LWCToFloat(Intermediates.Position);
 	Result.PreSkinnedNormal = TangentToLocal[2].xyz;
 
 	// Previous frame not handled deliberately. Lacks necessary information and
@@ -450,7 +450,7 @@ float3 SafeNormalizeWithFallback(float3 V, float3 Fallback)
 	return L > 0.000001 ? V / L : Fallback; 
 }
 
-void GetTangents(FVertexFactoryInput Input, float3 ParticlePosition, float Rotation, float3 CustomFacing, float3 CustomAlignment, float3 ParticleDirection, out float3 OutRight, out float3 OutUp)
+void GetTangents(FVertexFactoryInput Input, FLWCVector3 ParticlePosition, float Rotation, float3 CustomFacing, float3 CustomAlignment, float3 ParticleDirection, out float3 OutRight, out float3 OutUp)
 {
 	// Select camera up/right vectors.
 	float3 ResolvedViewRight = lerp(ResolvedView.ViewRight, ResolvedView.HMDViewNoRollRight, NiagaraSpriteVF.RemoveHMDRoll);
@@ -462,7 +462,7 @@ void GetTangents(FVertexFactoryInput Input, float3 ParticlePosition, float Rotat
 	float3 CameraFacingVector = -ResolvedView.ViewForward;
 
 	// Determine the vector from the particle to the camera and the particle's movement direction.
-	const float3 CameraVector = LWCHackToFloat(ResolvedView.WorldCameraOrigin) - ParticlePosition;
+	const float3 CameraVector = LWCToFloat(LWCSubtract(ResolvedView.WorldCameraOrigin, ParticlePosition));
 	const float3 CameraDirection = SafeNormalize(CameraVector);
 
 	float3 CustomFacingVector = NiagaraSpriteVFLooseParameters.ParticleFacingMode == FACING_MODE_CUSTOMFACINGVECTOR ? CustomFacing : CameraDirection;
@@ -564,14 +564,14 @@ float3x3 CalcTangentBasis(FVertexFactoryIntermediates Intermediates)
 	// ENM_Spherical
 	else if (NiagaraSpriteVF.NormalsType < 1.5f)
 	{
-		float3 TangentZ = normalize(Intermediates.Position - NiagaraSpriteVF.NormalsSphereCenter.xyz);
+		float3 TangentZ = normalize(LWCToFloat(LWCSubtract(Intermediates.Position, NiagaraSpriteVF.NormalsSphereCenter.xyz)));
 		Result[2] = TangentZ;	
 	}
 	// ENM_Cylindrical
 	else
 	{
-		float3 ClosestPointOnCylinder = NiagaraSpriteVF.NormalsSphereCenter.xyz + dot(NiagaraSpriteVF.NormalsCylinderUnitDirection.xyz, Intermediates.Position - NiagaraSpriteVF.NormalsSphereCenter.xyz) * NiagaraSpriteVF.NormalsCylinderUnitDirection.xyz;
-		float3 TangentZ = normalize(Intermediates.Position - ClosestPointOnCylinder);
+		float3 ClosestPointOnCylinder = NiagaraSpriteVF.NormalsSphereCenter.xyz + dot(NiagaraSpriteVF.NormalsCylinderUnitDirection.xyz, LWCToFloat(LWCSubtract(Intermediates.Position, NiagaraSpriteVF.NormalsSphereCenter.xyz))) * NiagaraSpriteVF.NormalsCylinderUnitDirection.xyz;
+		float3 TangentZ = normalize(LWCToFloat(LWCSubtract(Intermediates.Position, ClosestPointOnCylinder)));
 		Result[2] = TangentZ;	
 	}
 
@@ -619,23 +619,14 @@ void ComputeBillboardUVs(FVertexFactoryInput Input, float2 ParticleSize, float S
 	}
 }
 
-float3 LocalToWorldPos(float3 LocalSpacePos, float4x4 LocalToWorld)
-{
-	return
-		LocalToWorld[0].xyz * LocalSpacePos.x +
-		LocalToWorld[1].xyz * LocalSpacePos.y +
-		LocalToWorld[2].xyz * LocalSpacePos.z +
-		LocalToWorld[3].xyz;
-}
-
-float3 SimToWorldPos(float3 SimSpacePos, FLWCMatrix LocalToWorld)
+FLWCVector3 SimToWorldPos(float3 SimSpacePos, FLWCMatrix LocalToWorld)
 {
 	if (NiagaraSpriteVF.bLocalSpace)
 	{
-		return LocalToWorldPos(SimSpacePos, LWCHackToFloat(LocalToWorld));
+		return LWCMultiply(SimSpacePos, LocalToWorld);
 	}
 	
-	return LWCHackToFloat(MakeLWCVector3(NiagaraSpriteVF.SystemLWCTile, SimSpacePos));
+	return MakeLWCVector3(NiagaraSpriteVF.SystemLWCTile, SimSpacePos);
 }
 
 float3 LocalToWorldVec(float3 LocalSpaceVec, float3x3 LocalToWorld)
@@ -661,13 +652,13 @@ float3 SimToWorldVec(float3 SimSpaceVec, float3x3 LocalToWorld)
  *							Particle specific
  *********************************************************************************/
 
-float3 GetNiagaraParticlePosition(uint InstanceID)
+FLWCVector3 GetNiagaraParticlePosition(uint InstanceID)
 {
 	float3 ParticlePos = SafeGetVec3(NiagaraSpriteVF.PositionDataOffset, InstanceID, NiagaraSpriteVF.DefaultPos.xyz);
 	return SimToWorldPos(ParticlePos, GetPrimitiveDataFromUniformBuffer().LocalToWorld);
 }
 
-float3 GetNiagaraPreviousParticlePosition(uint InstanceID)
+FLWCVector3 GetNiagaraPreviousParticlePosition(uint InstanceID)
 {
 	float3 PrevParticlePos = SafeGetVec3(NiagaraSpriteVF.PrevPositionDataOffset, InstanceID, NiagaraSpriteVF.DefaultPrevPos.xyz);
 	return SimToWorldPos(PrevParticlePos, GetPrimitiveDataFromUniformBuffer().PreviousLocalToWorld);
@@ -810,7 +801,7 @@ FVertexFactoryIntermediates GetVertexFactoryIntermediates(FVertexFactoryInput In
 		ParticleID = NiagaraSpriteVFLooseParameters.SortedIndices[NiagaraSpriteVFLooseParameters.SortedIndicesOffset + ParticleID];
 	}
 
-	const float3 ParticlePosition = GetNiagaraParticlePosition(ParticleID);
+	const FLWCVector3 ParticlePosition = GetNiagaraParticlePosition(ParticleID);
 	const float ParticleRotation = GetNiagaraParticleRotation(ParticleID);
 	const float2 ParticleSize = GetNiagaraParticleSize(ParticleID);
 	const float3 ParticleVelocity = GetNiagaraParticleVelocity(ParticleID);
@@ -819,7 +810,7 @@ FVertexFactoryIntermediates GetVertexFactoryIntermediates(FVertexFactoryInput In
 	const float3 CustomAlignment = SafeNormalize(GetNiagaraParticleAlignmentVector(ParticleID));	
 	const float2 PivotOffset = GetNiagaraPivotOffset(ParticleID);
 	const float2 UVScale = GetNiagaraUVScale(ParticleID);
-	const float3 CameraOffset = SafeNormalize(LWCHackToFloat(ResolvedView.WorldCameraOrigin) - ParticlePosition) * GetNiagaraCameraOffset(ParticleID);
+	const float3 CameraOffset = SafeNormalize(LWCToFloat(LWCSubtract(ResolvedView.WorldCameraOrigin, ParticlePosition))) * GetNiagaraCameraOffset(ParticleID);
 	
 	Intermediates.SceneData = VF_GPUSCENE_GET_INTERMEDIATES(Input);
 	Intermediates.Position = ParticlePosition;
@@ -867,7 +858,7 @@ FVertexFactoryIntermediates GetVertexFactoryIntermediates(FVertexFactoryInput In
 	float PixelSizeRatio = 1.0f;
 	if ( NiagaraSpriteVF.PixelCoverageEnabled != 0 )
 	{
-		const float Depth = dot(ParticlePosition - LWCHackToFloat(ResolvedView.WorldCameraOrigin), ResolvedView.ViewForward);
+		const float Depth = dot(LWCToFloat(LWCSubtract(ParticlePosition, ResolvedView.WorldCameraOrigin)), ResolvedView.ViewForward);
 		const float HalfSize = max(Size.x, Size.y) * 0.5f;
 		const float PixelSize = max(View.ViewSizeAndInvSize.z * GetTanHalfFieldOfView().x * Depth * 2.0f, HalfSize);
 		PixelSizeRatio = PixelSize / HalfSize;
@@ -881,19 +872,19 @@ FVertexFactoryIntermediates GetVertexFactoryIntermediates(FVertexFactoryInput In
 	const float2x3 Tangents = float2x3(Intermediates.TangentRight, Intermediates.TangentUp);
 	const float3 VertexOffset = CameraOffset + mul(Size * PixelSizeRatio * (UVForPosition - PivotOffset), Tangents);
 
-	Intermediates.VertexWorldPosition = ParticlePosition + VertexOffset;
+	Intermediates.VertexWorldPosition = LWCAdd(ParticlePosition, VertexOffset);
 		
 	// Prev vertex position (approximated from current velocity unless ENABLE_PRECISE_MOTION_VECTORS)
 #if ENABLE_PRECISE_MOTION_VECTORS
 	{
-		const float3 PrevParticlePosition = GetNiagaraPreviousParticlePosition(ParticleID);
+		const FLWCVector3 PrevParticlePosition = GetNiagaraPreviousParticlePosition(ParticleID);
 		const float PrevParticleRotation = GetNiagaraPreviousParticleRotation(ParticleID);
 		const float2 PrevParticleSize = GetNiagaraPreviousParticleSize(ParticleID);
 		const float3 PrevParticleVelocity = GetNiagaraPreviousParticleVelocity(ParticleID);		
 		const float3 PrevCustomFacing = SafeNormalize(GetNiagaraPreviousParticleFacingVector(ParticleID));
 		const float3 PrevCustomAlignment = SafeNormalize(GetNiagaraPreviousParticleAlignmentVector(ParticleID));
 		const float2 PrevPivotOffset = GetNiagaraPreviousPivotOffset(ParticleID);
-		const float3 PrevCameraOffset = SafeNormalize(LWCHackToFloat(ResolvedView.PrevWorldCameraOrigin) - PrevParticlePosition) * GetNiagaraPreviousCameraOffset(ParticleID);
+		const float3 PrevCameraOffset = SafeNormalize(LWCToFloat(LWCSubtract(ResolvedView.PrevWorldCameraOrigin, PrevParticlePosition))) * GetNiagaraPreviousCameraOffset(ParticleID);
 		const float3 PrevParticleDirection = SafeNormalizeWithFallback(PrevParticleVelocity, float3(0,0,1));
 			
 		float3 PrevRight, PrevUp;
@@ -908,15 +899,15 @@ FVertexFactoryIntermediates GetVertexFactoryIntermediates(FVertexFactoryInput In
 		const float2x3 PrevTangents = float2x3(PrevRight, PrevUp);
 		const float3 PrevVertexOffset = PrevCameraOffset + mul(abs(PrevParticleSize) * (PrevUVForPosition - PivotOffset), PrevTangents);
 		
-		Intermediates.PreviousVertexWorldPosition = PrevParticlePosition + PrevVertexOffset;
+		Intermediates.PreviousVertexWorldPosition = LWCAdd(PrevParticlePosition, PrevVertexOffset);
 	}
 #else
-	Intermediates.PreviousVertexWorldPosition = Intermediates.VertexWorldPosition - ParticleVelocity * NiagaraSpriteVF.DeltaSeconds;
+	Intermediates.PreviousVertexWorldPosition = LWCSubtract(Intermediates.VertexWorldPosition, LWCPromote(ParticleVelocity * NiagaraSpriteVF.DeltaSeconds));
 #endif
 
 #if USE_PARTICLE_POSITION
 	// Do we really want the camera offset here ? What about when GENERATE_SPHERICAL_PARTICLE_NORMALS
-	Intermediates.TranslatedWorldPositionAndSize.xyz = Intermediates.Position + CameraOffset + LWCHackToFloat(ResolvedView.PreViewTranslation);
+	Intermediates.TranslatedWorldPositionAndSize.xyz = LWCToFloat(LWCAdd(Intermediates.Position, ResolvedView.PreViewTranslation)) + CameraOffset;
 	Intermediates.TranslatedWorldPositionAndSize.w = .5f * min(Size.x, Size.y);
 #if GENERATE_SPHERICAL_PARTICLE_NORMALS
 	// Recenter the particle position for GENERATE_SPHERICAL_PARTICLE_NORMALS to use the right center.
@@ -975,7 +966,7 @@ FVertexFactoryIntermediates GetVertexFactoryIntermediates(FVertexFactoryInput In
 
 float4 VertexFactoryGetWorldPosition(FVertexFactoryInput Input, FVertexFactoryIntermediates Intermediates)
 {
-	return float4(Intermediates.VertexWorldPosition + LWCHackToFloat(ResolvedView.PreViewTranslation), 1);
+	return float4(LWCToFloat(LWCAdd(Intermediates.VertexWorldPosition, ResolvedView.PreViewTranslation)), 1.0f);
 }
 
 float3 VertexFactoryGetWorldNormal(FVertexFactoryInput Input, FVertexFactoryIntermediates Intermediates)
@@ -1001,7 +992,7 @@ float3 VertexFactoryGetPositionForVertexLighting(FVertexFactoryInput Input, FVer
 {
 #if (TRANSLUCENCY_LIGHTING_VOLUMETRIC_PERVERTEX_DIRECTIONAL || TRANSLUCENCY_LIGHTING_VOLUMETRIC_PERVERTEX_NONDIRECTIONAL)
 	// Do per-vertex lighting with particle position instead of vertex position when we're using SubUV cutouts, since the vertex positions vary per-frame which would cause popping
-	return NiagaraSpriteVFLooseParameters.NumCutoutVerticesPerFrame > 0 ? (Intermediates.Position + LWCHackToFloat(ResolvedView.PreViewTranslation)) : TranslatedWorldPosition;
+	return NiagaraSpriteVFLooseParameters.NumCutoutVerticesPerFrame > 0 ? LWCToFloat(LWCAdd(Intermediates.Position, ResolvedView.PreViewTranslation)) : TranslatedWorldPosition;
 #else
 	return TranslatedWorldPosition;
 #endif
@@ -1191,7 +1182,7 @@ FVertexFactoryInterpolantsVSToPS VertexFactoryAssignInterpolants(FVertexFactoryR
 
 float4 VertexFactoryGetPreviousWorldPosition(FVertexFactoryInput Input, FVertexFactoryIntermediates Intermediates)
 {
-	return float4(Intermediates.PreviousVertexWorldPosition + LWCHackToFloat(ResolvedView.PrevPreViewTranslation), 1.0f);
+	return float4(LWCToFloat(LWCAdd(Intermediates.PreviousVertexWorldPosition, ResolvedView.PrevPreViewTranslation)), 1.0f);
 }
 
 /**
