Fixed bulk applying of scratch pads resulting in crashes due to scratch pad view models being refreshed mid-bulk-apply
Individual scratch pad applies will now refresh the system view model, whereas the bulk apply action will wait until all scratch pad view models are done
The refresh might be very ungranular, but some level of refresh is necessary for stack overview UI rebuilding, as inline parameters need to update properly

diff --git a/Engine/Plugins/FX/Niagara/Source/NiagaraEditor/Private/Toolkits/NiagaraSystemToolkit.cpp b/Engine/Plugins/FX/Niagara/Source/NiagaraEditor/Private/Toolkits/NiagaraSystemToolkit.cpp
--- a/Engine/Plugins/FX/Niagara/Source/NiagaraEditor/Private/Toolkits/NiagaraSystemToolkit.cpp
+++ b/Engine/Plugins/FX/Niagara/Source/NiagaraEditor/Private/Toolkits/NiagaraSystemToolkit.cpp
@@ -1137,10 +1137,7 @@ bool FNiagaraSystemToolkit::OnRequestClose()
 		switch (YesNoCancelReply)
 		{
 		case EAppReturnType::Yes:
-			for (TSharedRef<FNiagaraScratchPadScriptViewModel> ScratchPadViewModel : SystemViewModel->GetScriptScratchPadViewModel()->GetScriptViewModels())
-			{
-				ScratchPadViewModel->ApplyChanges();
-			}
+			SystemViewModel->GetScriptScratchPadViewModel()->ApplyScratchPadChanges();
 			break;
 		case EAppReturnType::No:
 			bScratchPadChangesDiscarded = true;
@@ -1224,13 +1221,7 @@ void FNiagaraSystemToolkit::OnApplyScratchPadChanges()
 {
 	if (SystemViewModel.IsValid() && SystemViewModel->GetScriptScratchPadViewModel() != nullptr)
 	{
-		for (TSharedRef<FNiagaraScratchPadScriptViewModel> ScratchPadScriptViewModel : SystemViewModel->GetScriptScratchPadViewModel()->GetScriptViewModels())
-		{
-			if(ScratchPadScriptViewModel->HasUnappliedChanges())
-			{
-				ScratchPadScriptViewModel->ApplyChanges();
-			}
-		}
+		SystemViewModel->GetScriptScratchPadViewModel()->ApplyScratchPadChanges();
 	}
 }
 
diff --git a/Engine/Plugins/FX/Niagara/Source/NiagaraEditor/Private/ViewModels/NiagaraScratchPadScriptViewModel.cpp b/Engine/Plugins/FX/Niagara/Source/NiagaraEditor/Private/ViewModels/NiagaraScratchPadScriptViewModel.cpp
--- a/Engine/Plugins/FX/Niagara/Source/NiagaraEditor/Private/ViewModels/NiagaraScratchPadScriptViewModel.cpp
+++ b/Engine/Plugins/FX/Niagara/Source/NiagaraEditor/Private/ViewModels/NiagaraScratchPadScriptViewModel.cpp
@@ -185,11 +185,6 @@ void FNiagaraScratchPadScriptViewModel::ApplyChanges()
 		FunctionCallNodeToRefresh->RefreshFromExternalChanges();
 		FunctionCallNodeToRefresh->MarkNodeRequiresSynchronization(TEXT("ScratchPadChangesApplied"), true);
 	}
-	
-	FRefreshAllScriptsFromExternalChangesArgs Args;
-	Args.OriginatingScript = OriginalScript;
-	Args.OriginatingGraph = CastChecked<UNiagaraScriptSource>(OriginalScript->GetSource(EditScript.Version))->NodeGraph;
-	FNiagaraEditorUtilities::RefreshAllScriptsFromExternalChanges(Args);
 
 	OnHasUnappliedChangesChangedDelegate.Broadcast();
 	OnChangesAppliedDelegate.Broadcast();
diff --git a/Engine/Plugins/FX/Niagara/Source/NiagaraEditor/Private/ViewModels/NiagaraScratchPadViewModel.cpp b/Engine/Plugins/FX/Niagara/Source/NiagaraEditor/Private/ViewModels/NiagaraScratchPadViewModel.cpp
--- a/Engine/Plugins/FX/Niagara/Source/NiagaraEditor/Private/ViewModels/NiagaraScratchPadViewModel.cpp
+++ b/Engine/Plugins/FX/Niagara/Source/NiagaraEditor/Private/ViewModels/NiagaraScratchPadViewModel.cpp
@@ -168,6 +168,21 @@ void UNiagaraScratchPadViewModel::RefreshScriptViewModels()
 	}
 }
 
+void UNiagaraScratchPadViewModel::ApplyScratchPadChanges()
+{
+	bIsBulkApplying = true;
+	for(const TSharedRef<FNiagaraScratchPadScriptViewModel>& ScriptViewModel : ScriptViewModels)
+	{
+		if(ScriptViewModel->HasUnappliedChanges())
+		{
+			ScriptViewModel->ApplyChanges();
+		}
+	}
+
+	GetSystemViewModel()->RefreshAll();
+	bIsBulkApplying = false;
+}
+
 const TArray<TSharedRef<FNiagaraScratchPadScriptViewModel>>& UNiagaraScratchPadViewModel::GetScriptViewModels() const
 {
 	return ScriptViewModels;
@@ -714,6 +729,12 @@ void UNiagaraScratchPadViewModel::ScriptViewModelHasUnappliedChangesChanged()
 void UNiagaraScratchPadViewModel::ScriptViewModelChangesApplied()
 {
 	UpdateChangeId(GetSystemViewModel());
+
+	// if we are bulk applying, we don't want to force refreshes more often than necessary. The bulk apply will refresh explicitly after all scripts have been applied.
+	if(!bIsBulkApplying)
+	{
+		SystemViewModelWeak.Pin()->RefreshAll();
+	}
 }
 
 void UNiagaraScratchPadViewModel::ScriptViewModelRequestDiscardChanges(TWeakPtr<FNiagaraScratchPadScriptViewModel> ScriptViewModelWeak)
diff --git a/Engine/Plugins/FX/Niagara/Source/NiagaraEditor/Public/ViewModels/NiagaraScratchPadViewModel.h b/Engine/Plugins/FX/Niagara/Source/NiagaraEditor/Public/ViewModels/NiagaraScratchPadViewModel.h
--- a/Engine/Plugins/FX/Niagara/Source/NiagaraEditor/Public/ViewModels/NiagaraScratchPadViewModel.h
+++ b/Engine/Plugins/FX/Niagara/Source/NiagaraEditor/Public/ViewModels/NiagaraScratchPadViewModel.h
@@ -29,6 +29,8 @@ public:
 
 	void RefreshScriptViewModels();
 
+	void ApplyScratchPadChanges();
+	
 	const TArray<TSharedRef<FNiagaraScratchPadScriptViewModel>>& GetScriptViewModels() const;
 
 	const TArray<TSharedRef<FNiagaraScratchPadScriptViewModel>>& GetEditScriptViewModels() const;
@@ -132,4 +134,6 @@ private:
 	FOnScriptRenamed OnScriptRenamedDelegate;
 
 	FOnScriptDeleted OnScriptDeletedDelegate;
+
+	bool bIsBulkApplying = false;
 };
\ No newline at end of file
