Added a change type context menu for parameters in script toolkits
- Fixed the convert tooltips to read "Convert [output pin type] to [input pin type]" instead of the other way around
- Reordered some functions that were hard to parse visually
- Readded the "parameter pin type" subcategory for pins that get their types changed; this kept the parameter references from updating properly

diff --git a/Engine/Plugins/FX/Niagara/Source/NiagaraEditor/Private/EdGraphSchema_Niagara.cpp b/Engine/Plugins/FX/Niagara/Source/NiagaraEditor/Private/EdGraphSchema_Niagara.cpp
--- a/Engine/Plugins/FX/Niagara/Source/NiagaraEditor/Private/EdGraphSchema_Niagara.cpp
+++ b/Engine/Plugins/FX/Niagara/Source/NiagaraEditor/Private/EdGraphSchema_Niagara.cpp
@@ -1013,11 +1013,11 @@ const FPinConnectionResponse UEdGraphSchema_Niagara::CanCreateConnection(const U
 				}
 				else if (FNiagaraTypeDefinition::IsLossyConversion(PinToTypeDefinition(InputPin), PinToTypeDefinition(OutputPin)))
 				{
-					return FPinConnectionResponse(CONNECT_RESPONSE_MAKE_WITH_CONVERSION_NODE, FString::Printf(TEXT("Convert %s to %s. Be aware that this can be a lossy conversion."), *(PinToTypeDefinition(InputPin).GetNameText().ToString()), *(PinToTypeDefinition(OutputPin).GetNameText().ToString())));
+					return FPinConnectionResponse(CONNECT_RESPONSE_MAKE_WITH_CONVERSION_NODE, FString::Printf(TEXT("Convert %s to %s. Be aware that this can be a lossy conversion."), *(PinToTypeDefinition(OutputPin).GetNameText().ToString()), *(PinToTypeDefinition(InputPin).GetNameText().ToString())));
 				}
 				else
 				{
-					return FPinConnectionResponse(CONNECT_RESPONSE_MAKE_WITH_CONVERSION_NODE, FString::Printf(TEXT("Convert %s to %s"), *(PinToTypeDefinition(InputPin).GetNameText().ToString()), *(PinToTypeDefinition(OutputPin).GetNameText().ToString())));
+					return FPinConnectionResponse(CONNECT_RESPONSE_MAKE_WITH_CONVERSION_NODE, FString::Printf(TEXT("Convert %s to %s"), *(PinToTypeDefinition(OutputPin).GetNameText().ToString()), *(PinToTypeDefinition(InputPin).GetNameText().ToString())));
 				}
 			}
 		}
diff --git a/Engine/Plugins/FX/Niagara/Source/NiagaraEditor/Private/NiagaraGraph.cpp b/Engine/Plugins/FX/Niagara/Source/NiagaraEditor/Private/NiagaraGraph.cpp
--- a/Engine/Plugins/FX/Niagara/Source/NiagaraEditor/Private/NiagaraGraph.cpp
+++ b/Engine/Plugins/FX/Niagara/Source/NiagaraEditor/Private/NiagaraGraph.cpp
@@ -356,6 +356,7 @@ void UNiagaraGraph::PostLoad()
 				if (Variable == OldVarType)
 				{
 					Pin->PinType = NiagaraSchema->TypeDefinitionToPinType(FNiagaraTypeDefinition::GetPositionDef());
+					Pin->PinType.PinSubCategory = UNiagaraNodeParameterMapBase::ParameterPinSubCategory;
 					NiagaraNode->MarkNodeRequiresSynchronization(__FUNCTION__, true);
 					break;
 				}
@@ -484,6 +485,7 @@ void UNiagaraGraph::ChangeParameterType(const FNiagaraVariable& CurrentParameter
 			{
 				Pin->Modify();
 				Pin->PinType = NiagaraSchema->TypeDefinitionToPinType(NewType);
+				Pin->PinType.PinSubCategory = UNiagaraNodeParameterMapBase::ParameterPinSubCategory;
 				NiagaraNode->MarkNodeRequiresSynchronization(__FUNCTION__, true);
 				break;
 			}
diff --git a/Engine/Plugins/FX/Niagara/Source/NiagaraEditor/Private/ViewModels/NiagaraParameterPanelViewModel.cpp b/Engine/Plugins/FX/Niagara/Source/NiagaraEditor/Private/ViewModels/NiagaraParameterPanelViewModel.cpp
--- a/Engine/Plugins/FX/Niagara/Source/NiagaraEditor/Private/ViewModels/NiagaraParameterPanelViewModel.cpp
+++ b/Engine/Plugins/FX/Niagara/Source/NiagaraEditor/Private/ViewModels/NiagaraParameterPanelViewModel.cpp
@@ -50,7 +50,6 @@ TArray<FNiagaraParameterPanelCategory> FNiagaraScriptToolkitParameterPanelViewMo
 TArray<FNiagaraParameterPanelCategory> FNiagaraScriptToolkitParameterPanelViewModel::DefaultAdvancedCategories;
 TArray<FNiagaraParameterPanelCategory> FNiagaraParameterDefinitionsToolkitParameterPanelViewModel::DefaultCategories;
 
-
 ///////////////////////////////////////////////////////////////////////////////
 /// System Toolkit Parameter Panel Utilities								///
 ///////////////////////////////////////////////////////////////////////////////
@@ -334,18 +333,6 @@ bool INiagaraParameterPanelViewModel::GetCanDeleteParameterAndToolTip(const FNia
 	return true;
 }
 
-bool INiagaraParameterPanelViewModel::GetCanPasteParameterMetaDataAndToolTip(FText& OutCanPasteToolTip)
-{
-	const UNiagaraClipboardContent* ClipboardContent = FNiagaraEditorModule::Get().GetClipboard().GetClipboardContent();
-	if (ClipboardContent == nullptr || ClipboardContent->ScriptVariables.Num() != 1)
-	{
-		OutCanPasteToolTip = LOCTEXT("CantPasteMetaDataToolTip_Invalid", "Cannot Paste: There is not any parameter metadata to paste in the clipboard.");
-		return false;
-	}
-	OutCanPasteToolTip = LOCTEXT("PasteMetaDataToolTip", "Paste the parameter metadata from the system clipboard to the selected parameters.");
-	return true;
-}
-
 void INiagaraParameterPanelViewModel::PasteParameterMetaData(const TArray<FNiagaraParameterPanelItem> SelectedItems)
 {
 	const UNiagaraClipboardContent* ClipboardContent = FNiagaraEditorModule::Get().GetClipboard().GetClipboardContent();
@@ -373,6 +360,18 @@ void INiagaraParameterPanelViewModel::PasteParameterMetaData(const TArray<FNiaga
 	}
 }
 
+bool INiagaraParameterPanelViewModel::GetCanPasteParameterMetaDataAndToolTip(FText& OutCanPasteToolTip)
+{
+	const UNiagaraClipboardContent* ClipboardContent = FNiagaraEditorModule::Get().GetClipboard().GetClipboardContent();
+	if (ClipboardContent == nullptr || ClipboardContent->ScriptVariables.Num() != 1)
+	{
+		OutCanPasteToolTip = LOCTEXT("CantPasteMetaDataToolTip_Invalid", "Cannot Paste: There is not any parameter metadata to paste in the clipboard.");
+		return false;
+	}
+	OutCanPasteToolTip = LOCTEXT("PasteMetaDataToolTip", "Paste the parameter metadata from the system clipboard to the selected parameters.");
+	return true;
+}
+
 void INiagaraParameterPanelViewModel::DuplicateParameter(const FNiagaraParameterPanelItem ItemToDuplicate) const
 {
 	TSet<FName> ParameterNames;
@@ -562,6 +561,74 @@ void INiagaraParameterPanelViewModel::SetParameterNamespaceModifier(const FNiaga
 	}
 }
 
+void INiagaraParameterPanelViewModel::ChangeParameterType(const FNiagaraParameterPanelItem ItemToModify, const FNiagaraTypeDefinition NewType) const
+{
+	ensureMsgf(false, TEXT("type change not supported for this view"));
+}
+
+bool INiagaraParameterPanelViewModel::GetCanChangeParameterType(const FNiagaraParameterPanelItem ItemToChange, FText& OutTooltip) const
+{
+	if (ItemToChange.bExternallyReferenced)
+	{
+		OutTooltip = LOCTEXT("CantChangeSelectedType_External", "This parameter is referenced in an external script and its type cannot be changed.");
+		return false;
+	}
+	
+	for(const UNiagaraGraph* Graph : GetEditableGraphsConst())
+	{
+		if(UNiagaraScriptVariable* ScriptVariable = Graph->GetScriptVariable(ItemToChange.GetVariable()))
+		{
+			if(ScriptVariable->GetIsStaticSwitch())
+			{
+				OutTooltip = LOCTEXT("CantChangeSelectedType_StaticSwitch", "This parameter is a static switch parameter and its type cannot be changed here.");
+				return false;
+			}
+
+			if(ItemToChange.GetVariable().GetType().IsStatic())
+			{
+				OutTooltip = LOCTEXT("CantChangeSelectedType_StaticParameter", "This parameter is static and its type cannot be changed.");
+				return false;
+			}
+
+			if(ScriptVariable->GetIsSubscribedToParameterDefinitions())
+			{
+				OutTooltip = LOCTEXT("CantChangeSelectedType_SubscribedParameter", "This parameter is subscribed to a parameter definition. The parameter can only be changed in the parameter definition asset.");
+				return false;
+			}
+		}
+	}	
+
+	OutTooltip = LOCTEXT("CanChangeSelectedType", "Change this parameter's type");
+	return true;
+}
+
+void INiagaraParameterPanelViewModel::GetChangeTypeSubMenu(FMenuBuilder& MenuBuilder, FNiagaraParameterPanelItem Item) const
+{
+	TArray<FNiagaraTypeDefinition> FilteredTypes;
+	for (const FNiagaraTypeDefinition& RegisteredType : FNiagaraTypeRegistry::GetRegisteredTypes())
+	{
+		// only allow basic types for now
+		if (RegisteredType.IsDataInterface() || RegisteredType.IsEnum() || RegisteredType.IsUObject() ||RegisteredType.GetNameText().ToString().Contains("event"))
+		{
+			continue;
+		}
+		FilteredTypes.Add(RegisteredType);
+	}
+	FilteredTypes.Sort([](const FNiagaraTypeDefinition& Lhs, const FNiagaraTypeDefinition& Rhs)
+	{
+		return Lhs.GetNameText().CompareTo(Rhs.GetNameText()) < 0;
+	});
+	
+	for (const FNiagaraTypeDefinition& Type : FilteredTypes)
+	{
+		MenuBuilder.AddMenuEntry(
+			Type.GetNameText(),
+			LOCTEXT("ChangeTypeActionTooltip", "Change the type of this variable. Replaced all pins in the graph and metadata with the selected type."),
+			FSlateIcon(),
+			FUIAction(FExecuteAction::CreateSP(this, &INiagaraParameterPanelViewModel::ChangeParameterType, Item, Type))); 
+	}
+}
+
 void INiagaraParameterPanelViewModel::SetParameterCustomNamespaceModifier(const FNiagaraParameterPanelItem ItemToModify, bool bDuplicateParameter) const
 {
 	if (bDuplicateParameter == false && ItemToModify.bExternallyReferenced)
@@ -2191,6 +2258,23 @@ void FNiagaraScriptToolkitParameterPanelViewModel::RenameParameter(const FNiagar
 	RenameParameter(ItemToRename.ScriptVariable, NewName);
 }
 
+void FNiagaraScriptToolkitParameterPanelViewModel::ChangeParameterType(const FNiagaraParameterPanelItem ItemToModify, const FNiagaraTypeDefinition NewType) const
+{
+	FScopedTransaction Transaction(LOCTEXT("ChangeParameterTypeTransaction", "Change parameter type"));
+	
+	FNiagaraVariable CurrentVar = ItemToModify.GetVariable();
+	
+	ScriptViewModel->GetStandaloneScript().Script->Modify();
+	for (UNiagaraGraph* Graph : GetEditableGraphs())
+	{
+		Graph->Modify();
+		Graph->ChangeParameterType(CurrentVar, NewType);
+	}
+
+	Refresh();
+	SelectParameterItemByName(CurrentVar.GetName(), false);
+}
+
 void FNiagaraScriptToolkitParameterPanelViewModel::RenameParameter(const UNiagaraScriptVariable* ScriptVarToRename, const FName NewName) const
 {
 	FScopedTransaction RenameTransaction(LOCTEXT("RenameParameterTransaction", "Rename parameter"));
@@ -2362,6 +2446,20 @@ TSharedPtr<SWidget> FNiagaraScriptToolkitParameterPanelViewModel::CreateContextM
 
 
 				MenuBuilder.AddMenuSeparator();
+				
+				FText CanChangeTooltip;
+				bool bCanChangeType = GetCanChangeParameterType(SelectedItem, CanChangeTooltip);
+				FUIAction CanChangeTypeAction;
+				CanChangeTypeAction.CanExecuteAction = FCanExecuteAction::CreateLambda([bCanChangeType]()
+				{
+					return bCanChangeType;
+				});
+				
+				MenuBuilder.AddSubMenu(
+					LOCTEXT("ChangeType", "Change Type"),
+					CanChangeTooltip,
+					FNewMenuDelegate::CreateSP(this, &FNiagaraScriptToolkitParameterPanelViewModel::GetChangeTypeSubMenu, SelectedItem),
+					CanChangeTypeAction, NAME_None, EUserInterfaceActionType::Button);
 
 				MenuBuilder.AddSubMenu(
 					LOCTEXT("ChangeNamespace", "Change Namespace"),
@@ -3052,4 +3150,4 @@ void FNiagaraParameterDefinitionsToolkitParameterPanelViewModel::OnParameterItem
 	}
 }
 
-#undef LOCTEXT_NAMESPACE // "FNiagaraParameterPanelViewModel"
+#undef LOCTEXT_NAMESPACE // "FNiagaraParameterPanelViewModel"
\ No newline at end of file
diff --git a/Engine/Plugins/FX/Niagara/Source/NiagaraEditor/Public/ViewModels/NiagaraParameterPanelViewModel.h b/Engine/Plugins/FX/Niagara/Source/NiagaraEditor/Public/ViewModels/NiagaraParameterPanelViewModel.h
--- a/Engine/Plugins/FX/Niagara/Source/NiagaraEditor/Public/ViewModels/NiagaraParameterPanelViewModel.h
+++ b/Engine/Plugins/FX/Niagara/Source/NiagaraEditor/Public/ViewModels/NiagaraParameterPanelViewModel.h
@@ -189,12 +189,14 @@ public:
 
 	virtual bool GetCanDeleteParameterAndToolTip(const FNiagaraParameterPanelItem& ItemToDelete, FText& OutCanDeleteParameterToolTip) const;
 
-	virtual bool GetCanPasteParameterMetaDataAndToolTip(FText& OutCanPasteToolTip);
+	virtual void ChangeParameterType(const FNiagaraParameterPanelItem ItemToModify, const FNiagaraTypeDefinition NewType) const;
+	virtual bool GetCanChangeParameterType(const FNiagaraParameterPanelItem ItemToChange, FText& OutTooltip) const;
+	virtual void GetChangeTypeSubMenu(FMenuBuilder& MenuBuilder, FNiagaraParameterPanelItem Item) const;
 
 	virtual void PasteParameterMetaData(const TArray<FNiagaraParameterPanelItem> SelectedItems);
+	virtual bool GetCanPasteParameterMetaDataAndToolTip(FText& OutCanPasteToolTip);
 
 	virtual void DuplicateParameter(const FNiagaraParameterPanelItem ItemToDuplicate) const;
-
 	virtual bool GetCanDuplicateParameterAndToolTip(const FNiagaraParameterPanelItem& ItemToDuplicate, FText& OutCanDuplicateParameterToolTip) const;
 
 	virtual bool GetCanRenameParameterAndToolTip(const FNiagaraParameterPanelItem& ItemToRename, const FText& NewVariableNameText, bool bCheckEmptyNameText, FText& OutCanRenameParameterToolTip) const;
@@ -388,6 +390,8 @@ public:
 	virtual void DeleteParameter(const FNiagaraParameterPanelItem& ItemToDelete) const override;
 
 	virtual void RenameParameter(const FNiagaraParameterPanelItem& ItemToRename, const FName NewName) const override;
+	
+	virtual void ChangeParameterType(const FNiagaraParameterPanelItem ItemToModify, const FNiagaraTypeDefinition NewType) const override;
 
 	virtual void DuplicateParameter(const FNiagaraParameterPanelItem ItemToDuplicate) const override;
 
