- Fix for using local to world / world to local transforms with Niagara mesh VF

#rnx

diff --git a/Engine/Plugins/FX/Niagara/Shaders/Private/NiagaraMeshVertexFactory.ush b/Engine/Plugins/FX/Niagara/Shaders/Private/NiagaraMeshVertexFactory.ush
--- a/Engine/Plugins/FX/Niagara/Shaders/Private/NiagaraMeshVertexFactory.ush
+++ b/Engine/Plugins/FX/Niagara/Shaders/Private/NiagaraMeshVertexFactory.ush
@@ -320,7 +320,7 @@ FMaterialPixelParameters GetMaterialPixelParameters(FVertexFactoryInterpolantsVS
 #endif
 
 #if USE_PARTICLE_LOCAL_TO_WORLD
-	//-TODO:
+	//-TODO: LWC Precision Loss
 	FLWCMatrix ParticleToWorld = LWCPromote(transpose(float4x4(Interpolants.ParticleToWorld[0], Interpolants.ParticleToWorld[1], Interpolants.ParticleToWorld[2], float4(0.0f, 0.0f, 0.0f, 1.0f))));
 	#if NEEDS_PARTICLE_LOCAL_TO_WORLD
 		Result.Particle.ParticleToWorld = ParticleToWorld;
@@ -331,7 +331,7 @@ FMaterialPixelParameters GetMaterialPixelParameters(FVertexFactoryInterpolantsVS
 #endif
 
 #if USE_PARTICLE_WORLD_TO_LOCAL
-	//-TODO:
+	//-TODO: LWC Precision Loss
 	FLWCInverseMatrix WorldToParticle = LWCPromoteInverse(transpose(float4x4(Interpolants.WorldToParticle[0], Interpolants.WorldToParticle[1], Interpolants.WorldToParticle[2], float4(0.0f, 0.0f, 0.0f, 1.0f))));
 	#if NEEDS_PARTICLE_WORLD_TO_LOCAL
 		Result.Particle.WorldToParticle = WorldToParticle;
@@ -766,14 +766,16 @@ FVertexFactoryInterpolantsVSToPS VertexFactoryGetInterpolantsVSToPS(FVertexFacto
 #endif
 
 #if USE_PARTICLE_LOCAL_TO_WORLD
-	float4x4 ParticleToWorldT = transpose(Intermediates.ParticleToWorld);
+	//-TODO: LWC Precision Loss
+	float4x4 ParticleToWorldT = transpose(LWCHackToFloat(Intermediates.ParticleToWorld));
 	Interpolants.ParticleToWorld[0] = ParticleToWorldT[0];
 	Interpolants.ParticleToWorld[1] = ParticleToWorldT[1];
 	Interpolants.ParticleToWorld[2] = ParticleToWorldT[2];
 #endif
 
 #if USE_PARTICLE_WORLD_TO_LOCAL
-	float4x4 WorldToParticleT = transpose(Intermediates.WorldToParticle);
+	//-TODO: LWC Precision Loss
+	float4x4 WorldToParticleT = transpose(LWCHackToFloat(Intermediates.WorldToParticle));
 	Interpolants.WorldToParticle[0] = WorldToParticleT[0];
 	Interpolants.WorldToParticle[1] = WorldToParticleT[1];
 	Interpolants.WorldToParticle[2] = WorldToParticleT[2];
