- New Vector VM with optimizer and serialization.
- curve sampling DIs changed to 4-wide
- Off by default

diff --git a/Engine/Plugins/FX/Niagara/Source/Niagara/Classes/NiagaraDataSet.h b/Engine/Plugins/FX/Niagara/Source/Niagara/Classes/NiagaraDataSet.h
--- a/Engine/Plugins/FX/Niagara/Source/Niagara/Classes/NiagaraDataSet.h
+++ b/Engine/Plugins/FX/Niagara/Source/Niagara/Classes/NiagaraDataSet.h
@@ -403,6 +403,7 @@ public:
 
 	const FNiagaraDataSetCompiledData& GetCompiledData() const { check(CompiledData.Get() != nullptr); return *CompiledData.Get(); }
 
+	int NumSpawnedIDs;
 private:
 
 	void Reset();
diff --git a/Engine/Plugins/FX/Niagara/Source/Niagara/Classes/NiagaraScript.h b/Engine/Plugins/FX/Niagara/Source/Niagara/Classes/NiagaraScript.h
--- a/Engine/Plugins/FX/Niagara/Source/Niagara/Classes/NiagaraScript.h
+++ b/Engine/Plugins/FX/Niagara/Source/Niagara/Classes/NiagaraScript.h
@@ -18,6 +18,7 @@
 #include "NiagaraParameterDefinitionsSubscriber.h"
 #include "HAL/CriticalSection.h"
 #include "HAL/PlatformAtomics.h"
+#include "VectorVM.h"
 
 #include "NiagaraScript.generated.h"
 
@@ -649,6 +650,9 @@ class UNiagaraScript : public UNiagaraScriptBase
 public:
 	UNiagaraScript();
 	~UNiagaraScript();
+#	ifdef NIAGARA_EXP_VM
+	FVectorVMOptimizeContext OptimizeContext;
+#	endif //NIAGARA_EXP_VM
 
 #if WITH_EDITORONLY_DATA
 	/** If true then this script asset uses active version control to track changes. */
diff --git a/Engine/Plugins/FX/Niagara/Source/Niagara/Classes/NiagaraScriptExecutionContext.h b/Engine/Plugins/FX/Niagara/Source/Niagara/Classes/NiagaraScriptExecutionContext.h
--- a/Engine/Plugins/FX/Niagara/Source/Niagara/Classes/NiagaraScriptExecutionContext.h
+++ b/Engine/Plugins/FX/Niagara/Source/Niagara/Classes/NiagaraScriptExecutionContext.h
@@ -152,6 +152,8 @@ struct FNiagaraScriptExecutionContextBase
 {
 	UNiagaraScript* Script;
 
+	FVectorVMState *VectorVMState;
+	
 	/** Table of external function delegate handles called from the VM. */
 	TArray<const FVMExternalFunction*> FunctionTable;
 
@@ -169,6 +171,9 @@ struct FNiagaraScriptExecutionContextBase
 
 	static uint32 TickCounter;
 
+	/** The script type this context is for. Allows us to access the correct per instance function table on the system instance. */
+	ENiagaraSystemSimulationScript ScriptType;
+
 	int32 HasInterpolationParameters : 1;
 	int32 bAllowParallel : 1;
 #if STATS
@@ -241,14 +246,12 @@ protected:
 	*/
 	TArray<FNiagaraSystemInstance*>* SystemInstances;
 
-	/** The script type this context is for. Allows us to access the correct per instance function table on the system instance. */
-	ENiagaraSystemSimulationScript ScriptType;
 
 	/** Helper function that handles calling into per instance DI calls and massages the VM context appropriately. */
 	void PerInstanceFunctionHook(FVectorVMExternalFunctionContext& Context, int32 PerInstFunctionIndex, int32 UserPtrIndex);
 
 public:
-	FNiagaraSystemScriptExecutionContext(ENiagaraSystemSimulationScript InScriptType) : SystemInstances(nullptr), ScriptType(InScriptType){}
+	FNiagaraSystemScriptExecutionContext(ENiagaraSystemSimulationScript InScriptType) : SystemInstances(nullptr) { ScriptType = InScriptType; }
 	
 	virtual bool Init(UNiagaraScript* InScript, ENiagaraSimTarget InTarget)override;
 	virtual bool Tick(class FNiagaraSystemInstance* Instance, ENiagaraSimTarget SimTarget);
diff --git a/Engine/Plugins/FX/Niagara/Source/Niagara/Private/NiagaraDataInterfaceColorCurve.cpp b/Engine/Plugins/FX/Niagara/Source/Niagara/Private/NiagaraDataInterfaceColorCurve.cpp
--- a/Engine/Plugins/FX/Niagara/Source/Niagara/Private/NiagaraDataInterfaceColorCurve.cpp
+++ b/Engine/Plugins/FX/Niagara/Source/Niagara/Private/NiagaraDataInterfaceColorCurve.cpp
@@ -237,10 +237,87 @@ FORCEINLINE_DEBUGGABLE FLinearColor UNiagaraDataInterfaceColorCurve::SampleCurve
 	return FLinearColor(RedCurve.Eval(X), GreenCurve.Eval(X), BlueCurve.Eval(X), AlphaCurve.Eval(X));
 }
 
+
+#if defined(NIAGARA_EXP_VM) && PLATFORM_ENABLE_VECTORINTRINSICS && !PLATFORM_ENABLE_VECTORINTRINSICS_NEON
+template<>
+void UNiagaraDataInterfaceColorCurve::SampleCurve<TIntegralConstant<bool, true>>(FVectorVMExternalFunctionContext& Context)
+{
+	if (Context.NumInstances == 1) //could be a per-instance function call, in which can we can't write 4-wide so just use the old method
+	{ 
+		VectorVM::FExternalFuncInputHandler<float> XParam(Context);
+		VectorVM::FExternalFuncRegisterHandler<float> SamplePtrR(Context);
+		VectorVM::FExternalFuncRegisterHandler<float> SamplePtrG(Context);
+		VectorVM::FExternalFuncRegisterHandler<float> SamplePtrB(Context);
+		VectorVM::FExternalFuncRegisterHandler<float> SamplePtrA(Context);
+
+		for (int32 i = 0; i < Context.GetNumInstances(); ++i)
+		{
+			float X = XParam.GetAndAdvance();
+			FLinearColor C = SampleCurveInternal<TIntegralConstant<bool, true>>(X);
+			*SamplePtrR.GetDestAndAdvance() = C.R;
+			*SamplePtrG.GetDestAndAdvance() = C.G;
+			*SamplePtrB.GetDestAndAdvance() = C.B;
+			*SamplePtrA.GetDestAndAdvance() = C.A;
+		}
+	}
+	else
+	{
+		float *LUT = ShaderLUT.GetData();
+		VectorRegister4f LutNumSamplesMinusOne4 = VectorSetFloat1(LUTNumSamplesMinusOne);
+		VectorRegister4f LutMinTime4            = VectorSetFloat1(LUTMinTime);
+		VectorRegister4f LutInvTimeRange4       = VectorSetFloat1(LUTInvTimeRange);
+
+		VectorRegister4f *XParam = (VectorRegister4f *)Context.RegisterData[0];
+		VectorRegister4f *RReg   = (VectorRegister4f *)Context.RegisterData[1];
+		VectorRegister4f *GReg   = (VectorRegister4f *)Context.RegisterData[2];
+		VectorRegister4f *BReg   = (VectorRegister4f *)Context.RegisterData[3];
+		VectorRegister4f *AReg   = (VectorRegister4f *)Context.RegisterData[4];
+
+		int IdxA[4];
+		int IdxB[4];
+
+		for (int i = 0; i < Context.NumLoops; ++i)
+		{
+			VectorRegister4f NormalizedX = VectorMultiply(VectorSubtract(XParam[i & Context.RegInc[0]], LutMinTime4), LutInvTimeRange4);
+			VectorRegister4f RemappedX   = VectorClamp(VectorMultiply(NormalizedX, LutNumSamplesMinusOne4), VectorZeroFloat(), LutNumSamplesMinusOne4);
+			VectorRegister4f PrevEntry   = VectorTruncate(RemappedX);
+			VectorRegister4f NextEntry   = VectorAdd(PrevEntry, VectorBitwiseAnd(VectorOneFloat(), VectorCompareLT(PrevEntry, LutNumSamplesMinusOne4))); //this could be made faster by duplicating the last entry in the LUT so you can read one past it
+			VectorRegister4f Interp      = VectorSubtract(RemappedX, PrevEntry);
+			VectorRegister4i IdxA4       = VectorShiftLeftImm(VectorFloatToInt(PrevEntry), 2);
+			VectorRegister4i IdxB4       = VectorShiftLeftImm(VectorFloatToInt(NextEntry), 2);
+
+			VectorIntStore(IdxA4, IdxA);
+			VectorIntStore(IdxB4, IdxB);
+
+			VectorRegister4f A0 = VectorLoad(LUT + IdxA[0]);
+			VectorRegister4f A1 = VectorLoad(LUT + IdxA[1]);
+			VectorRegister4f A2 = VectorLoad(LUT + IdxA[2]);
+			VectorRegister4f A3 = VectorLoad(LUT + IdxA[3]);
+
+			VectorRegister4f B0 = VectorLoad(LUT + IdxB[0]);
+			VectorRegister4f B1 = VectorLoad(LUT + IdxB[1]);
+			VectorRegister4f B2 = VectorLoad(LUT + IdxB[2]);
+			VectorRegister4f B3 = VectorLoad(LUT + IdxB[3]);
+
+			VectorRegister4f I0 = VectorCastIntToFloat(VectorShuffleImmediate(VectorCastFloatToInt(Interp), 0, 0, 0, 0));
+			VectorRegister4f I1 = VectorCastIntToFloat(VectorShuffleImmediate(VectorCastFloatToInt(Interp), 1, 1, 1, 1));
+			VectorRegister4f I2 = VectorCastIntToFloat(VectorShuffleImmediate(VectorCastFloatToInt(Interp), 2, 2, 2, 2));
+			VectorRegister4f I3 = VectorCastIntToFloat(VectorShuffleImmediate(VectorCastFloatToInt(Interp), 3, 3, 3, 3));
+
+			RReg[i] = VectorMultiplyAdd(B0, I0, VectorMultiply(A0, VectorSubtract(VectorOneFloat(), I0)));
+			GReg[i] = VectorMultiplyAdd(B1, I1, VectorMultiply(A1, VectorSubtract(VectorOneFloat(), I1)));
+			BReg[i] = VectorMultiplyAdd(B2, I2, VectorMultiply(A2, VectorSubtract(VectorOneFloat(), I2)));
+			AReg[i] = VectorMultiplyAdd(B3, I3, VectorMultiply(A3, VectorSubtract(VectorOneFloat(), I3)));
+
+			_MM_TRANSPOSE4_PS(RReg[i], GReg[i], BReg[i], AReg[i]);
+		}
+	}
+}
+#endif //NIAGARA_EXP_VM
+
 template<typename UseLUT>
 void UNiagaraDataInterfaceColorCurve::SampleCurve(FVectorVMExternalFunctionContext& Context)
 {
-	//TODO: Create some SIMDable optimized representation of the curve to do this faster.
 	VectorVM::FExternalFuncInputHandler<float> XParam(Context);
 	VectorVM::FExternalFuncRegisterHandler<float> SamplePtrR(Context);
 	VectorVM::FExternalFuncRegisterHandler<float> SamplePtrG(Context);
diff --git a/Engine/Plugins/FX/Niagara/Source/Niagara/Private/NiagaraDataInterfaceCurve.cpp b/Engine/Plugins/FX/Niagara/Source/Niagara/Private/NiagaraDataInterfaceCurve.cpp
--- a/Engine/Plugins/FX/Niagara/Source/Niagara/Private/NiagaraDataInterfaceCurve.cpp
+++ b/Engine/Plugins/FX/Niagara/Source/Niagara/Private/NiagaraDataInterfaceCurve.cpp
@@ -195,6 +195,57 @@ FORCEINLINE_DEBUGGABLE float UNiagaraDataInterfaceCurve::SampleCurveInternal<TIn
 	return Curve.Eval(X);
 }
 
+#ifdef NIAGARA_EXP_VM
+template<>
+void UNiagaraDataInterfaceCurve::SampleCurve<TIntegralConstant<bool, true>>(FVectorVMExternalFunctionContext& Context)
+{
+	if (Context.NumInstances == 1)
+	{
+		VectorVM::FExternalFuncInputHandler<float> XParam(Context);
+		VectorVM::FExternalFuncRegisterHandler<float> OutSample(Context);
+
+		for (int32 i = 0; i < Context.GetNumInstances(); ++i)
+		{
+			*OutSample.GetDest() = SampleCurveInternal<TIntegralConstant<bool, true>>(XParam.Get());
+			XParam.Advance();
+			OutSample.Advance();
+		}
+	}
+	else
+	{
+		float *LUT = ShaderLUT.GetData();
+		VectorRegister4f LutNumSamplesMinusOne4 = VectorSetFloat1(LUTNumSamplesMinusOne);
+		VectorRegister4f LutMinTime4            = VectorSetFloat1(LUTMinTime);
+		VectorRegister4f LutInvTimeRange4       = VectorSetFloat1(LUTInvTimeRange);
+	
+		VectorRegister4f *XParam = (VectorRegister4f *)Context.RegisterData[0];
+		VectorRegister4f *OutSample = (VectorRegister4f *)Context.RegisterData[1];
+
+		int IdxA[4];
+		int IdxB[4];
+
+		for (int i = 0; i < Context.NumLoops; ++i)
+		{
+			VectorRegister4f NormalizedTime4 = VectorMultiply(VectorSubtract(XParam[i & Context.RegInc[0]], LutMinTime4), LutInvTimeRange4);
+			VectorRegister4f RemappedX4      = VectorMin(VectorMax(VectorMultiply(NormalizedTime4, LutNumSamplesMinusOne4), VectorZeroFloat()), LutNumSamplesMinusOne4);
+			VectorRegister4f PrevEntry4      = VectorTruncate(RemappedX4);
+			VectorRegister4f NextEntry4      = VectorAdd(PrevEntry4, VectorBitwiseAnd(VectorOneFloat(), VectorCompareLT(PrevEntry4, LutNumSamplesMinusOne4))); //this could be made faster by duplicating the last entry in the LUT so you can read one past it
+			VectorRegister4f Interp4         = VectorSubtract(RemappedX4, PrevEntry4);
+			VectorRegister4i IdxA4           = VectorFloatToInt(PrevEntry4);
+			VectorRegister4i IdxB4           = VectorFloatToInt(NextEntry4);
+
+			VectorIntStore(IdxA4, IdxA);
+			VectorIntStore(IdxB4, IdxB);
+
+			VectorRegister4f A4 = MakeVectorRegisterFloat(LUT[IdxA[0]], LUT[IdxA[1]], LUT[IdxA[2]], LUT[IdxA[3]]);
+			VectorRegister4f B4 = MakeVectorRegisterFloat(LUT[IdxB[0]], LUT[IdxB[1]], LUT[IdxB[2]], LUT[IdxB[3]]);
+
+			OutSample[i] = VectorMultiplyAdd(B4, Interp4, VectorMultiply(A4, VectorSubtract(VectorOneFloat(), Interp4)));
+		}
+	}
+}
+#endif
+
 template<typename UseLUT>
 void UNiagaraDataInterfaceCurve::SampleCurve(FVectorVMExternalFunctionContext& Context)
 {
diff --git a/Engine/Plugins/FX/Niagara/Source/Niagara/Private/NiagaraDataInterfaceParticleRead.cpp b/Engine/Plugins/FX/Niagara/Source/Niagara/Private/NiagaraDataInterfaceParticleRead.cpp
--- a/Engine/Plugins/FX/Niagara/Source/Niagara/Private/NiagaraDataInterfaceParticleRead.cpp
+++ b/Engine/Plugins/FX/Niagara/Source/Niagara/Private/NiagaraDataInterfaceParticleRead.cpp
@@ -1420,8 +1420,17 @@ void UNiagaraDataInterfaceParticleRead::GetSpawnedIDAtIndex(FVectorVMExternalFun
 	VectorVM::FExternalFuncRegisterHandler<int32> OutIDAcquireTag(Context);
 
 	FNiagaraEmitterInstance* EmitterInstance = InstData.Get()->EmitterInstance;
+#ifdef NIAGARA_EXP_VM
+	int32 *SpawnedIDsTable = EmitterInstance ? EmitterInstance->GetData().GetFreeIDTable().GetData() + EmitterInstance->GetData().GetFreeIDTable().Max() : nullptr;
+	int32 NumSpawned = EmitterInstance ? EmitterInstance->GetData().NumSpawnedIDs : 0;
+	if (SpawnedIDsTable && NumSpawned) {
+		SpawnedIDsTable -= NumSpawned; //spawned ids are at the end of the FreeID table
+	}
+#else
 	const TArray<int32>* SpawnedIDsTable = EmitterInstance ? &EmitterInstance->GetData().GetSpawnedIDsTable() : nullptr;
 	int32 NumSpawned = SpawnedIDsTable ? SpawnedIDsTable->Num() : 0;
+#endif
+
 	int32 IDAcquireTag = EmitterInstance ? EmitterInstance->GetData().GetIDAcquireTag() : 0;
 
 	for (int32 InstanceIdx = 0; InstanceIdx < Context.GetNumInstances(); ++InstanceIdx)
@@ -1433,7 +1442,11 @@ void UNiagaraDataInterfaceParticleRead::GetSpawnedIDAtIndex(FVectorVMExternalFun
 		if (SpawnIndex >= 0 && SpawnIndex < NumSpawned)
 		{
 			ValidValue.SetValue(true);
+#ifdef NIAGARA_EXP_VM
+			IDValue.Index = SpawnedIDsTable[SpawnIndex];
+#else
 			IDValue.Index = (*SpawnedIDsTable)[SpawnIndex];
+#endif
 			IDValue.AcquireTag = IDAcquireTag;
 		}
 		else
diff --git a/Engine/Plugins/FX/Niagara/Source/Niagara/Private/NiagaraDataSet.cpp b/Engine/Plugins/FX/Niagara/Source/Niagara/Private/NiagaraDataSet.cpp
--- a/Engine/Plugins/FX/Niagara/Source/Niagara/Private/NiagaraDataSet.cpp
+++ b/Engine/Plugins/FX/Niagara/Source/Niagara/Private/NiagaraDataSet.cpp
@@ -204,6 +204,7 @@ void FNiagaraDataSet::ResetBuffersInternal()
 	FreeIDsTable.Reset();
 	NumFreeIDs = 0;
 	MaxUsedID = INDEX_NONE;
+	NumSpawnedIDs = 0;
 	SpawnedIDsTable.Reset();
 
 	//Ensure we have a valid current buffer
@@ -356,8 +357,10 @@ void FNiagaraDataSet::Allocate(int32 NumInstances, bool bMaintainExisting)
 			RequiredIDs = ExistingNumIDs;
 		}
 
+#ifndef NIAGARA_EXP_VM //the new VM doesn't use the SpawnedIDsTable, so remove it.
 		// We know that we can't spawn more than NumFreeIDs particles, so we can pre-allocate SpawnedIDsTable here, to avoid allocations during execution.
 		SpawnedIDsTable.Reserve(NumFreeIDs);
+#endif
 
 		// We need to clear the ID to index table to -1 so we don't have stale entries for particles which died in the previous
 		// frame (when the results were written to another buffer). All the entries which are in use will be filled in by the script.
diff --git a/Engine/Plugins/FX/Niagara/Source/Niagara/Private/NiagaraEmitterInstance.cpp b/Engine/Plugins/FX/Niagara/Source/Niagara/Private/NiagaraEmitterInstance.cpp
--- a/Engine/Plugins/FX/Niagara/Source/Niagara/Private/NiagaraEmitterInstance.cpp
+++ b/Engine/Plugins/FX/Niagara/Source/Niagara/Private/NiagaraEmitterInstance.cpp
@@ -1671,8 +1671,10 @@ void FNiagaraEmitterInstance::Tick(float DeltaSeconds)
 	const int32 NumBeforeSpawn = Data.GetDestinationDataChecked().GetNumInstances();
 	uint32 TotalActualEventSpawns = 0;
 
+	Data.NumSpawnedIDs = 0;
+#ifndef NIAGARA_EXP_VM
 	Data.GetSpawnedIDsTable().SetNum(0, false);
-
+#endif
 	int32 SpawnCountRemaining = AllocationSize - OrigNumParticles;
 
 	//Init new particles with the spawn script.
diff --git a/Engine/Plugins/FX/Niagara/Source/Niagara/Private/NiagaraScript.cpp b/Engine/Plugins/FX/Niagara/Source/Niagara/Private/NiagaraScript.cpp
--- a/Engine/Plugins/FX/Niagara/Source/Niagara/Private/NiagaraScript.cpp
+++ b/Engine/Plugins/FX/Niagara/Source/Niagara/Private/NiagaraScript.cpp
@@ -289,6 +289,9 @@ FVersionedNiagaraScriptData::FVersionedNiagaraScriptData()
 
 UNiagaraScript::UNiagaraScript()
 {
+#	ifdef NIAGARA_EXP_VM
+	FMemory::Memzero(&OptimizeContext, sizeof(OptimizeContext));
+#	endif //NIAGARA_EXP_VM
 }
 
 #if WITH_EDITORONLY_DATA
@@ -735,6 +738,9 @@ UNiagaraScript::UNiagaraScript(const FObjectInitializer& ObjectInitializer)
 	, IsCooked(false)
 #endif
 {
+#	ifdef NIAGARA_EXP_VM
+	FMemory::Memzero(&OptimizeContext, sizeof(OptimizeContext));
+#	endif //NIAGARA_EXP_VM
 #if WITH_EDITORONLY_DATA
 	ScriptResource = MakeUnique<FNiagaraShaderScript>();
 	ScriptResource->OnCompilationComplete().AddUniqueDynamic(this, &UNiagaraScript::RaiseOnGPUCompilationComplete);
@@ -745,6 +751,9 @@ UNiagaraScript::UNiagaraScript(const FObjectInitializer& ObjectInitializer)
 
 UNiagaraScript::~UNiagaraScript()
 {
+#	ifdef NIAGARA_EXP_VM
+	FreeVectorVMOptimizeContext(&OptimizeContext);
+#	endif //NIAGARA_EXP_VM
 }
 
 #if WITH_EDITORONLY_DATA
@@ -1392,6 +1401,17 @@ bool UNiagaraScript::ShouldFreeUnoptimizedByteCode() const
 
 FGraphEventRef UNiagaraScript::HandleByteCodeOptimization(bool bShouldForceNow)
 {
+#	ifdef NIAGARA_EXP_VM
+	//this is just necessary because VectorVM doesn't know about FVMExternalFunctionBindingInfo
+	int NumExtFns = CachedScriptVM.CalledVMExternalFunctions.Num();
+	FVectorVMExtFunctionData *ExtFnTable = (FVectorVMExtFunctionData *)alloca(32 + sizeof(FVectorVMExtFunctionData) * NumExtFns); //32 is pointless, it just gets around static analysis triggering for unknown reasons
+	for (int i = 0; i < NumExtFns; ++i) {
+		ExtFnTable[i].NumInputs = CachedScriptVM.CalledVMExternalFunctions[i].GetNumInputs();
+		ExtFnTable[i].NumOutputs = CachedScriptVM.CalledVMExternalFunctions[i].GetNumOutputs();
+	}
+	OptimizeVectorVMScript(CachedScriptVM.ByteCode.GetDataPtr(), CachedScriptVM.ByteCode.GetLength(), ExtFnTable, NumExtFns, &OptimizeContext, VVMOptFlag_SaveIntermediateState);
+
+#	endif //NIAGARA_EXP_VM
 	check(IsInGameThread());
 
 	const bool bHasOptimizationTask = CachedScriptVM.OptimizationTask.State.IsValid();
diff --git a/Engine/Plugins/FX/Niagara/Source/Niagara/Private/NiagaraScriptExecutionContext.cpp b/Engine/Plugins/FX/Niagara/Source/Niagara/Private/NiagaraScriptExecutionContext.cpp
--- a/Engine/Plugins/FX/Niagara/Source/Niagara/Private/NiagaraScriptExecutionContext.cpp
+++ b/Engine/Plugins/FX/Niagara/Source/Niagara/Private/NiagaraScriptExecutionContext.cpp
@@ -7,6 +7,8 @@
 #include "NiagaraWorldManager.h"
 #include "NiagaraFunctionLibrary.h"
 #include "NiagaraGPUInstanceCountManager.h"
+#include "Interfaces/ITargetPlatform.h"
+#include "HAL/PlatformFileManager.h"
 
 DECLARE_CYCLE_STAT(TEXT("Register Setup"), STAT_NiagaraSimRegisterSetup, STATGROUP_Niagara);
 DECLARE_CYCLE_STAT(TEXT("Context Ticking"), STAT_NiagaraScriptExecContextTick, STATGROUP_Niagara);
@@ -27,7 +29,8 @@ static FAutoConsoleVariableRef CVarNiagaraExecVMScripts(
 
 FNiagaraScriptExecutionContextBase::FNiagaraScriptExecutionContextBase()
 	: Script(nullptr)
-	, HasInterpolationParameters(false)
+	, VectorVMState(nullptr)
+	, ScriptType(ENiagaraSystemSimulationScript::Update)
 	, bAllowParallel(true)
 {
 
@@ -35,6 +38,9 @@ FNiagaraScriptExecutionContextBase::FNiagaraScriptExecutionContextBase()
 
 FNiagaraScriptExecutionContextBase::~FNiagaraScriptExecutionContextBase()
 {
+#	ifdef NIAGARA_EXP_VM
+	FreeVectorVMState(VectorVMState);
+#	endif 
 }
 
 bool FNiagaraScriptExecutionContextBase::Init(UNiagaraScript* InScript, ENiagaraSimTarget InTarget)
@@ -62,7 +68,7 @@ void FNiagaraScriptExecutionContextBase::BindData(int32 Index, FNiagaraDataSet&
 
 	DataSetMetaTable.SetNum(FMath::Max(DataSetMetaTable.Num(), Index + 1));
 	DataSetMetaTable[Index].Init(InputRegisters, OutputRegisters, StartInstance,
-		Output ? &Output->GetIDTable() : nullptr, &DataSet.GetFreeIDTable(), &DataSet.GetNumFreeIDs(), &DataSet.GetMaxUsedID(), DataSet.GetIDAcquireTag(), &DataSet.GetSpawnedIDsTable());
+		Output ? &Output->GetIDTable() : nullptr, &DataSet.GetFreeIDTable(), &DataSet.GetNumFreeIDs(), &DataSet.NumSpawnedIDs, &DataSet.GetMaxUsedID(), DataSet.GetIDAcquireTag(), &DataSet.GetSpawnedIDsTable());
 
 	if (InputRegisters.Num() > 0)
 	{
@@ -87,7 +93,7 @@ void FNiagaraScriptExecutionContextBase::BindData(int32 Index, FNiagaraDataBuffe
 	TArrayView<uint8 const* RESTRICT const> InputRegisters = Input->GetRegisterTable();
 
 	DataSetMetaTable.SetNum(FMath::Max(DataSetMetaTable.Num(), Index + 1));
-	DataSetMetaTable[Index].Init(InputRegisters, TArrayView<uint8 const* RESTRICT const>(), StartInstance, nullptr, nullptr, &DataSet->GetNumFreeIDs(), &DataSet->GetMaxUsedID(), DataSet->GetIDAcquireTag(), &DataSet->GetSpawnedIDsTable());
+	DataSetMetaTable[Index].Init(InputRegisters, TArrayView<uint8 const* RESTRICT const>(), StartInstance, nullptr, nullptr, &DataSet->GetNumFreeIDs(), &DataSet->NumSpawnedIDs, &DataSet->GetMaxUsedID(), DataSet->GetIDAcquireTag(), &DataSet->GetSpawnedIDsTable());
 
 	if (InputRegisters.Num() > 0)
 	{
@@ -121,6 +127,15 @@ TMap<TStatIdData const*, float> FNiagaraScriptExecutionContextBase::ReportStats(
 }
 #endif
 
+static void *VVMRealloc(void* Ptr, size_t NumBytes, const char *Filename, int LineNum)
+{
+	return FMemory::Realloc(Ptr, NumBytes);
+}
+static void VVMFree(void* Ptr, const char *Filename, int LineNum)
+{
+	return FMemory::Free(Ptr);
+}
+
 bool FNiagaraScriptExecutionContextBase::Execute(uint32 NumInstances, const FScriptExecutionConstantBufferTable& ConstantBufferTable)
 {
 	if (NumInstances == 0)
@@ -151,6 +166,9 @@ bool FNiagaraScriptExecutionContextBase::Execute(uint32 NumInstances, const FScr
 			}
 		}
 
+		UObject *OuterObj0 = Script->GetTypedOuter<UNiagaraSystem>();
+		UObject *OuterObj1 = Script->GetTypedOuter<UNiagaraEmitter>();
+
 		check((ExecData.ByteCode.HasByteCode() && !ExecData.ByteCode.IsCompressed()) || (ExecData.OptimizedByteCode.HasByteCode() && !ExecData.OptimizedByteCode.IsCompressed()));
 
 		VectorVM::FVectorVMExecArgs ExecArgs;
@@ -171,8 +189,120 @@ bool FNiagaraScriptExecutionContextBase::Execute(uint32 NumInstances, const FScr
 #endif
 		
 		ExecArgs.bAllowParallel = bAllowParallel;
-		VectorVM::Exec(ExecArgs);
-	}
+		FVectorVMSerializeState ExpSerializeState = { };
+		FVectorVMSerializeState UESerializeState = { };
+
+#if defined(NIAGARA_EXP_VM) || defined(VVM_INCLUDE_SERIALIZATION)
+		FVectorVMConstData *VVMConstData = (FVectorVMConstData *)alloca(sizeof(FVectorVMConstData) * ConstantBufferTable.BufferSizes.Num());
+		int NumVVMConstData              = ConstantBufferTable.BufferSizes.Num();
+		for (int i = 0; i < NumVVMConstData; ++i)
+		{
+			int32 NumBytes = ConstantBufferTable.BufferSizes[i];
+			check((NumBytes & 3) == 0); //should only be float32 or int32
+			VVMConstData[i].NumDWords = NumBytes >> 2;
+			VVMConstData[i].RegisterData = (void *)ConstantBufferTable.Buffers[i];
+		}
+#endif //NIAGARA_EXP_VM
+
+#ifdef VVM_INCLUDE_SERIALIZATION
+		FString ScriptName;
+		Script->GetName(ScriptName);
+		
+		static const IConsoleVariable* CVarInstancesPerChunk = IConsoleManager::Get().FindConsoleVariable(TEXT("vm.InstancesPerChunk"));
+		int32 NumParallelInstancesPerChunk = CVarInstancesPerChunk ? CVarInstancesPerChunk->GetInt() : 128;
+		
+		UESerializeState.ReallocFn        = VVMRealloc;
+		UESerializeState.FreeFn           = VVMFree;
+		UESerializeState.NumInstances     = NumInstances;
+		UESerializeState.NumTempRegisters = ExecData.NumTempRegisters;
+		UESerializeState.NumTempRegFlags  = UESerializeState.NumTempRegisters;
+		
+		UESerializeState.TempRegFlags     = (unsigned char *)VVMRealloc(NULL, UESerializeState.NumTempRegisters, __FILE__, __LINE__);
+		FMemory::Memset(UESerializeState.TempRegFlags, 0, UESerializeState.NumTempRegisters);
+		UESerializeState.Bytecode         = (unsigned char *)ExecData.ByteCode.GetData().GetData();
+		UESerializeState.NumBytecodeBytes = ExecData.ByteCode.GetData().Num();
+		SerializeVectorVMInputDataSets(&UESerializeState, DataSetMetaTable, VVMConstData, NumVVMConstData);
+
+		UESerializeState.NumChunks        = (NumInstances + NumParallelInstancesPerChunk - 1) / NumParallelInstancesPerChunk;
+		UESerializeState.Chunks           = (FVectorVMSerializeChunk *)VVMRealloc(NULL, sizeof(FVectorVMSerializeChunk) * UESerializeState.NumChunks, __FILE__, __LINE__);
+
+		
+		int NumExternalFunctions = FunctionTable.Num();
+		if (NumExternalFunctions != 0)
+		{
+			UESerializeState.ExternalData = (FVectorVMSerializeExternalData *)VVMRealloc(NULL, sizeof(FVectorVMSerializeExternalData) * NumExternalFunctions, __FILE__, __LINE__);
+			if (UESerializeState.ExternalData)
+			{
+				UESerializeState.NumExternalData = NumExternalFunctions;
+				const FNiagaraVMExecutableData& ScriptExecutableData = Script->GetVMExecutableData();
+				for (int i = 0; i < NumExternalFunctions; ++i)
+				{
+					FVectorVMSerializeExternalData *ExtData = UESerializeState.ExternalData + i;
+					TCHAR buff[128];
+					uint32 NameLen = ScriptExecutableData.CalledVMExternalFunctions[i].Name.ToString(buff, sizeof(buff) >> 1);
+					if (NameLen > 0)
+					{
+						ExtData->Name       = (wchar_t *)UESerializeState.ReallocFn(NULL, sizeof(wchar_t) * NameLen, __FILE__, __LINE__);
+						FMemory::Memcpy(ExtData->Name, buff, sizeof(TCHAR) * NameLen);
+						ExtData->NameLen    = NameLen;
+					}
+					else
+					{
+						ExtData->Name       = NULL;
+						ExtData->NameLen    = 0;
+					}
+					ExtData->NumInputs      = ScriptExecutableData.CalledVMExternalFunctions[i].GetNumInputs();
+					ExtData->NumOutputs     = ScriptExecutableData.CalledVMExternalFunctions[i].GetNumOutputs();
+				}
+			}
+		}
+#endif //VVM_INCLUDE_SERIALIZATION
+
+#ifndef NIAGARA_EXP_VM
+		ExpSerializeState.ReallocFn = VVMRealloc;
+		ExpSerializeState.FreeFn = VVMFree;
+		VectorVM::Exec(ExecArgs, &UESerializeState);
+#ifdef VVM_INCLUDE_SERIALIZATION
+		SerializeVectorVMOutputDataSets(&UESerializeState, DataSetMetaTable, VVMConstData, NumVVMConstData);
+#endif //VVM_INCLUDE_SERIALIZATION
+#else
+		FVectorVMInitData InitData     = { };
+		InitData.OptimizeContext       = &Script->OptimizeContext;
+		InitData.NumInstances          = NumInstances;
+		InitData.DataSets              = DataSetMetaTable;
+		InitData.ConstData             = VVMConstData;
+		InitData.NumConstData          = NumVVMConstData;
+		InitData.ExtFunctionTable      = FunctionTable;
+		InitData.UserPtrTable          = UserPtrTable.GetData();
+		InitData.NumUserPtrTable       = UserPtrTable.Num();
+		InitData.ExistingVectorVMState = VectorVMState;
+		InitData.ReallocFn             = VVMRealloc;
+		InitData.FreeFn                = VVMFree;
+		FVectorVMExternalFnPerInstanceData *PerInstanceExternalData = nullptr;
+		VectorVMState = InitVectorVMState(&InitData, &PerInstanceExternalData, &ExpSerializeState);
+
+#ifdef VVM_INCLUDE_SERIALIZATION
+		if (ExpSerializeState.NumExternalData > 0)
+		{
+			const FNiagaraVMExecutableData& ScriptExecutableData = Script->GetVMExecutableData();
+			for (uint32 i = 0; i < ExpSerializeState.NumExternalData; ++i)
+			{
+				TCHAR buff[128];
+				uint32 NameLen = ScriptExecutableData.CalledVMExternalFunctions[i].Name.ToString(buff, sizeof(buff) >> 1);
+				if (NameLen > 0)
+				{
+					ExpSerializeState.ExternalData[i].Name = (wchar_t *)UESerializeState.ReallocFn(NULL, sizeof(wchar_t) * NameLen, __FILE__, __LINE__);
+					memcpy(ExpSerializeState.ExternalData[i].Name, buff, 2 * NameLen);
+					ExpSerializeState.ExternalData[i].NameLen = NameLen;
+				}
+			}
+		}
+#endif
+		if (VectorVMState && VectorVMState->Error.Flags == 0)
+		{
+			ExecVectorVMState(VectorVMState, &ExpSerializeState, &UESerializeState);
+		}
+#		endif
 
 	// Tell the datasets we wrote how many instances were actually written.
 	for (int Idx = 0; Idx < DataSetInfo.Num(); Idx++)
@@ -182,20 +312,75 @@ bool FNiagaraScriptExecutionContextBase::Execute(uint32 NumInstances, const FScr
 #if NIAGARA_NAN_CHECKING
 		Info.DataSet->CheckForNaNs();
 #endif
+			if (Info.bUpdateInstanceCount)
+			{
+				Info.Output->SetNumInstances(Info.StartInstance + DataSetMetaTable[Idx].DataSetAccessIndex + 1);
+			}
 
-		if (Info.bUpdateInstanceCount)
-		{
-			Info.Output->SetNumInstances(Info.StartInstance + DataSetMetaTable[Idx].DataSetAccessIndex + 1);
+#ifndef NIAGARA_EXP_VM //@TODO: this can go when we remove the old VM.  It's only here for serialization purposes as the VM doesn't use a spawned table
+			Info.DataSet->NumSpawnedIDs = Info.DataSet->GetSpawnedIDsTable().Num();
+#endif
 		}
-	}
 
-	//Can maybe do without resetting here. Just doing it for tidiness.
-	for (int32 DataSetIdx = 0; DataSetIdx < DataSetInfo.Num(); ++DataSetIdx)
-	{
-		DataSetInfo[DataSetIdx].Reset();
-		DataSetMetaTable[DataSetIdx].Reset();
-	}
+#		ifdef VVM_INCLUDE_SERIALIZATION
+		{
+#			ifdef NIAGARA_EXP_VM
+			FVectorVMSerializeState *SerializeState = &ExpSerializeState;
+			const char *state_ext = "exp";
+			uint8 which_state_written = 1;
+#			else //NIAGARA_EXP_VM
+			FVectorVMSerializeState *SerializeState = &UESerializeState;
+			const char *state_ext = "ue";
+			uint8 which_state_written = 2;
+#			endif //NIAGARA_EXP_VM
+			
+			//only write under certain circumstances
+			if (OuterObj0 &&
+				OuterObj0->GetName() == L"BoneTest_System" && 
+				OuterObj1 &&
+				OuterObj1->GetName() == L"BoneTest_1" &&
+				ScriptName == L"SpawnScript_0")
+			{
+				SerializeVectorVMOutputDataSets(SerializeState, DataSetMetaTable, VVMConstData, NumVVMConstData);
+				TStringBuilder<256> sb;
+				//sb.Append(Parameters.DebugName);
+				//sb.Append(L"_");
+				sb.Append(ScriptName);
+				sb.Append(L"_");
+
+				TAnsiStringBuilder<256> sb_ansi;
+				
+				sb_ansi.Appendf("0x%08X_%s_%d.vvm_dump\0", FPlatformTime::Cycles(), state_ext, NumInstances);
+				//sb.AppendAnsi(sb_ansi.GetData());
+
+				sb.Append(sb_ansi.GetData());
+
+				//char buff[64];
+				//sb.Appendf(
+				//int32 buff_len = (int32)snprintf(buff, sizeof(buff), "0x%08X_%s_%d.vvm_dump", FPlatform::Cycles(), state_ext, NumInstances);
+				//if (buff_len > 0 && buff_len < sizeof(buff)) {
+				//	buff[buff_len] = 0;
+				//	sb.AppendAnsi(buff, buff_len + 1);
+				//}
+				wchar_t *Name16 = (wchar_t *)sb.GetData();
+				Name16[sb.Len()] = 0;
+				SerializeVectorVMWriteToFile(SerializeState, which_state_written, Name16);
+			}
 
+			FreeVectorVMSerializeState(&ExpSerializeState);
+			//we don't own the bytecode memory so we can't free it
+			UESerializeState.Bytecode = NULL;
+			UESerializeState.NumBytecodeBytes = 0;
+			FreeVectorVMSerializeState(&UESerializeState);
+		}
+#		endif //VVM_INCLUDE_SERIALIZATION
+		//Can maybe do without resetting here. Just doing it for tidiness.
+		for (int32 DataSetIdx = 0; DataSetIdx < DataSetInfo.Num(); ++DataSetIdx)
+		{
+			DataSetInfo[DataSetIdx].Reset();
+			DataSetMetaTable[DataSetIdx].Reset();
+		}
+	}
 	return true;//TODO: Error cases?
 }
 
@@ -385,8 +570,38 @@ void FNiagaraScriptExecutionContextBase::PostTick()
 }
 
 //////////////////////////////////////////////////////////////////////////
-
-void FNiagaraSystemScriptExecutionContext::PerInstanceFunctionHook(FVectorVMExternalFunctionContext& Context, int32 PerInstFunctionIndex, int32 UserPtrIndex)
+#ifdef NIAGARA_EXP_VM
+static void PerInsFn(FVectorVMExternalFunctionContext &PerInsFnContext, TArray<FNiagaraSystemInstance *> **SystemInstances, ENiagaraSystemSimulationScript ScriptType, int32 PerInstFunctionIndex, int32 UserPtrIdx) {
+	check(PerInsFnContext.DataSets.Num() > 0);
+	check(SystemInstances);
+	check(*SystemInstances);
+	
+	if (PerInsFnContext.UserPtrTable && UserPtrIdx < PerInsFnContext.NumUserPtrs) {
+		void *SavedUserPtrData       = PerInsFnContext.UserPtrTable[UserPtrIdx];
+		int32 InstanceOffset         = PerInsFnContext.DataSets[0].InstanceOffset; //Apparently the function table is generated based off the first data set, therefore this is safe. (I don't like this - @shawn mcgrath)	
+		int NumInstances             = PerInsFnContext.NumInstances;
+		PerInsFnContext.NumInstances = 1;
+		for (int i = 0; i < NumInstances; ++i) {
+			PerInsFnContext.RegReadCount             = 0;
+			PerInsFnContext.PerInstanceFnInstanceIdx = i;
+
+			int32 InstanceIndex                           = InstanceOffset + PerInsFnContext.StartInstance + i;
+			FNiagaraSystemInstance *Instance              = (**SystemInstances)[InstanceIndex];
+			const FNiagaraPerInstanceDIFuncInfo &FuncInfo = Instance->GetPerInstanceDIFunction(ScriptType, PerInstFunctionIndex);
+		
+			if (UserPtrIdx != INDEX_NONE) {
+				PerInsFnContext.UserPtrTable[UserPtrIdx] = FuncInfo.InstData;
+			}
+			FuncInfo.Function.Execute(PerInsFnContext);
+		}
+	
+		if (SavedUserPtrData) {
+			PerInsFnContext.UserPtrTable[UserPtrIdx] = SavedUserPtrData;
+		}
+	}
+}
+#else
+void FNiagaraSystemScriptExecutionContext::PerInstanceFunctionHook(FVectorVMExternalFunctionContext &Context, int32 PerInstFunctionIndex, int32 UserPtrIndex)
 {
 	check(SystemInstances);
 	
@@ -433,6 +648,7 @@ void FNiagaraSystemScriptExecutionContext::PerInstanceFunctionHook(FVectorVMExte
 	Context.VectorVMContext->StartInstance = CachedContextStartInstance;
 	Context.VectorVMContext->NumInstances = CachedContextNumInstances;
 }
+#endif
 
 bool FNiagaraSystemScriptExecutionContext::Init(UNiagaraScript* InScript, ENiagaraSimTarget InTarget)
 {
@@ -490,12 +706,17 @@ bool FNiagaraSystemScriptExecutionContext::Tick(class FNiagaraSystemInstance* In
 						//Currently we must assume that any User DI is overridden but maybe we can be less conservative with this in future.
 						if (ScriptDIInfo.NeedsPerInstanceBinding())
 						{
-							//This DI needs a binding per instance so we just bind to the external function hook which will call the correct binding for each instance.
+#							ifdef NIAGARA_EXP_VM
+							auto PerInstFunctionHookLambda = [SystemInstances = &this->SystemInstances, ScriptType = this->ScriptType, NumPerInstanceFunctions, UserPtrIdx = ScriptDIInfo.UserPtrIdx](FVectorVMExternalFunctionContext &ExtFnContext) {
+								PerInsFn(ExtFnContext, SystemInstances, ScriptType, NumPerInstanceFunctions, UserPtrIdx);
+							};
+#							else
 							auto PerInstFunctionHookLambda = [ExecContext = this, NumPerInstanceFunctions, UserPtrIndex = ScriptDIInfo.UserPtrIdx](FVectorVMExternalFunctionContext& Context)
 							{
+								//This DI needs a binding per instance so we just bind to the external function hook which will call the correct binding for each instance.
 								ExecContext->PerInstanceFunctionHook(Context, NumPerInstanceFunctions, UserPtrIndex);
 							};
-
+#							endif
 							++NumPerInstanceFunctions;
 							FuncInfo.Function = FVMExternalFunction::CreateLambda(PerInstFunctionHookLambda);
 						}
diff --git a/Engine/Source/Runtime/VectorVM/Private/VectorVM.cpp b/Engine/Source/Runtime/VectorVM/Private/VectorVM.cpp
--- a/Engine/Source/Runtime/VectorVM/Private/VectorVM.cpp
+++ b/Engine/Source/Runtime/VectorVM/Private/VectorVM.cpp
@@ -7,6 +7,8 @@
 #include "VectorVMPrivate.h"
 #include "Stats/Stats.h"
 #include "HAL/ConsoleManager.h"
+#include "HAL/PlatformFileManager.h"
+#include "GenericPlatform/GenericPlatformFile.h"
 #include "Async/ParallelFor.h"
 
 IMPLEMENT_MODULE(FDefaultModuleImpl, VectorVM);
@@ -138,6 +140,15 @@ FORCEINLINE VectorRegister4Int VectorIntShuffle(const VectorRegister4Int& Vec, c
 }
 #endif
 
+#ifdef NIAGARA_EXP_VM
+//@TODO: remove these functions when we go to the new VM
+FORCEINLINE void FDataSetMeta::LockFreeTable() {
+
+}
+FORCEINLINE void FDataSetMeta::UnlockFreeTable() {
+
+}
+#else
 //Temporarily locking the free table until we can implement a lock free algorithm. UE-65856
 FORCEINLINE void FDataSetMeta::LockFreeTable()
 {
@@ -161,6 +172,7 @@ FORCEINLINE void FDataSetMeta::UnlockFreeTable()
 {
  	FreeTableLock.Unlock();
 }
+#endif //NIAGARA_EXP_VM
 
 static int32 GbParallelVVM = 1;
 static FAutoConsoleVariableRef CVarbParallelVVM(
@@ -195,6 +207,22 @@ static FAutoConsoleVariableRef CVarParallelVVMInstancesPerChunk(
 	ECVF_ReadOnly
 );
 
+static int32 GVVMChunkSizeInBytes = 32768;
+static FAutoConsoleVariableRef CVarVVMChunkSizeInBytes(
+	TEXT("vm.ChunkSizeInBytes"),
+	GVVMChunkSizeInBytes,
+	TEXT("Number of bytes per VM chunk  Ideally <= L1 size. (default=32768) \n"),
+	ECVF_Default
+);
+
+static int32 GVVMMaxThreadsPerScript = 8;
+static FAutoConsoleVariableRef CVarVVMMaxThreadsPerScript(
+	TEXT("vm.MaxThreadsPerScript"),
+	GVVMMaxThreadsPerScript,
+	TEXT("Maximum number of threads per script. Set 0 to mean 'as many as necessary'\n"),
+	ECVF_Default
+);
+
 static int32 GbOptimizeVMByteCode = 1;
 static FAutoConsoleVariableRef CVarbOptimizeVMByteCode(
 	TEXT("vm.OptimizeVMByteCode"),
@@ -251,6 +279,8 @@ static FAutoConsoleVariableRef CVarbBatchPackVMOutput(
 	ECVF_Default
 );
 
+#include "VectorVMExperimental.inl"
+
 //////////////////////////////////////////////////////////////////////////
 //  VM Code Optimizer Context
 
@@ -2528,7 +2558,7 @@ void VectorVM::Init()
 	}
 }
 
-void VectorVM::Exec(FVectorVMExecArgs& Args)
+void VectorVM::Exec(FVectorVMExecArgs& Args, FVectorVMSerializeState *SerializeState)
 {
 	//TRACE_CPUPROFILER_EVENT_SCOPE("VMExec");
 	SCOPE_CYCLE_COUNTER(STAT_VVMExec);
@@ -2546,9 +2576,14 @@ void VectorVM::Exec(FVectorVMExecArgs& Args)
 	const int32 ChunksPerBatch = (GbParallelVVM != 0 && FApp::ShouldUseThreadingForPerformance()) ? GParallelVVMChunksPerBatch : NumChunks;
 	const int32 NumBatches = FMath::DivideAndRoundUp(NumChunks, ChunksPerBatch);
 	const bool bParallel = NumBatches > 1 && Args.bAllowParallel;
+#	ifdef VVM_INCLUDE_SERIALIZATION
+	const bool bUseOptimizedByteCode = false; //serializes the bytecode from the instructions, cannot use jump table
+#	else //VVM_INCLUDE_SERIALIZATION
 	const bool bUseOptimizedByteCode = (Args.OptimizedByteCode != nullptr) && GbUseOptimizedVMByteCode;
-
+#	endif
 	const FVectorVMExecFunction* OptimizedJumpTable = bUseOptimizedByteCode ? FVectorVMCodeOptimizerContext::DecodeJumpTable(Args.OptimizedByteCode) : nullptr;
+	
+	uint64 StartTime = FPlatformTime::Cycles64();
 
 	auto ExecChunkBatch = [&](int32 BatchIdx)
 	{
@@ -2593,10 +2628,12 @@ void VectorVM::Exec(FVectorVMExecArgs& Args)
 				// Setup execution context.
 				Context.PrepareForChunk(Args.ByteCode, NumInstancesThisChunk, StartInstance);
 
+				VVMSer_chunkStart(Context, ChunkIdx, BatchIdx);
 				// Execute VM on all vectors in this chunk.
 				EVectorVMOp Op = EVectorVMOp::done;
 				do
 				{
+					VVMSer_insStart(Context);
 					Op = Context.DecodeOp();
 					switch (Op)
 					{
@@ -2712,7 +2749,9 @@ void VectorVM::Exec(FVectorVMExecArgs& Args)
 							UE_LOG(LogVectorVM, Fatal, TEXT("Unknown op code 0x%02x"), (uint32)Op);
 							return;//BAIL
 					}
+					VVMSer_insEnd(Context, (int)(VVMSerCtxStartInsCode - VVMSerStartCtxCode), (int)(Context.Code - VVMSerCtxStartInsCode));
 				} while (Op != EVectorVMOp::done);
+				VVMSer_chunkEnd(SerializeState)
 			}
 
 			InstancesLeft -= GParallelVVMInstancesPerChunk;
@@ -2736,6 +2775,13 @@ void VectorVM::Exec(FVectorVMExecArgs& Args)
 		FPlatformMisc::EndNamedEvent();
 	}
 #endif
+
+#ifdef VVM_INCLUDE_SERIALIZATION
+	uint64 EndTime = FPlatformTime::Cycles64();
+	if (SerializeState) {
+		SerializeState->ExecDt = EndTime - StartTime; //NOTE: doesn't work if ParallelFor splits the work into multiple threads
+	}
+#endif //VVM_INCLUDE_SERIALIZATION
 }
 
 uint8 VectorVM::GetNumOpCodes()
