Post merge fixes
#rnx

diff --git a/Engine/Plugins/FX/Niagara/Source/Niagara/Private/NiagaraComponent.cpp b/Engine/Plugins/FX/Niagara/Source/Niagara/Private/NiagaraComponent.cpp
--- a/Engine/Plugins/FX/Niagara/Source/Niagara/Private/NiagaraComponent.cpp
+++ b/Engine/Plugins/FX/Niagara/Source/Niagara/Private/NiagaraComponent.cpp
@@ -95,40 +95,40 @@ FAutoConsoleCommandWithWorldAndArgs DumpNiagaraComponentsCommand(
 		{
 			const bool bFullDump = Args.Contains(TEXT("full"));
 			const bool bAllWorlds = (World == nullptr) || Args.Contains(TEXT("allworlds"));
-			const FString SystemFilter = 
+			const FString SystemFilter =
 				[&]() -> FString
+			{
+				static const FString FilterPrefix(TEXT("filter="));
+				for (const FString& Arg : Args)
 				{
-					static const FString FilterPrefix(TEXT("filter="));
-					for ( const FString& Arg : Args )
+					if (Arg.StartsWith(FilterPrefix))
 					{
-						if ( Arg.StartsWith(FilterPrefix) )
-						{
-							return Arg.Mid(FilterPrefix.Len());
-						}
+						return Arg.Mid(FilterPrefix.Len());
 					}
-					return FString();
-				}();
+				}
+				return FString();
+			}();
 
 			UE_LOG(LogNiagara, Log, TEXT("=========================== Begin Niagara Dump ==========================="));
 			for (TObjectIterator<UNiagaraComponent> It; It; ++It)
 			{
 				UNiagaraComponent* Component = *It;
-				if ( !IsValid(Component) )
+				if (!IsValid(Component))
 				{
 					continue;
 				}
 
 				// Filter by world
-				if (!bAllWorlds && (Component->GetWorld() != World) )
+				if (!bAllWorlds && (Component->GetWorld() != World))
 				{
 					continue;
 				}
 
 				// Filter by asset name, we allow null assets to continue so we can identify components who are attached but have no valid asset associated with them
 				UNiagaraSystem* NiagaraSystem = Component->GetAsset();
-				if ( !SystemFilter.IsEmpty() )
+				if (!SystemFilter.IsEmpty())
 				{
-					if (!NiagaraSystem || !NiagaraSystem->GetName().Contains(SystemFilter) )
+					if (!NiagaraSystem || !NiagaraSystem->GetName().Contains(SystemFilter))
 					{
 						continue;
 					}
@@ -151,18 +151,18 @@ static FAutoConsoleVariableRef CVarNiagaraCanPreventScalabilityCullingOnPlayerFX
 	GNiagaraCanPreventScalabilityCullingOnPlayerFX,
 	TEXT("When enabled Niagara can optionally prevent scalability culling on FX linked to the player."),
 	FConsoleVariableDelegate::CreateLambda([](IConsoleVariable* CVar)
-	{
-		FNiagaraWorldManager::OnRefreshOwnerAllowsScalability();
-	}),
+		{
+			FNiagaraWorldManager::OnRefreshOwnerAllowsScalability();
+		}),
 	ECVF_Default
-);
+			);
 
 bool FitsIntoFloat(FVector InValue)
 {
 	// copied from PlayerCameraManager, 1048576 is HALF_WORLD_MAX for float precision vectors
 	return (InValue.X < 1048576.0f && InValue.X > -1048576.0f &&
-			InValue.Y < 1048576.0f && InValue.Y > -1048576.0f &&
-			InValue.Z < 1048576.0f && InValue.Z > -1048576.0f);
+		InValue.Y < 1048576.0f && InValue.Y > -1048576.0f &&
+		InValue.Z < 1048576.0f && InValue.Z > -1048576.0f);
 }
 
 FNiagaraSceneProxy::FNiagaraSceneProxy(UNiagaraComponent* InComponent)
@@ -199,7 +199,7 @@ FNiagaraSceneProxy::~FNiagaraSceneProxy()
 {
 	check(IsInRenderingThread());
 
-	for ( auto& CustomUB : CustomUniformBuffers)
+	for (auto& CustomUB : CustomUniformBuffers)
 	{
 		CustomUB.Value.ReleaseResource();
 	}
@@ -286,7 +286,7 @@ TUniformBuffer<FPrimitiveUniformShaderParameters>& FNiagaraSceneProxy::GetCustom
 	KeyHash = HashCombine(KeyHash, PreviousOrigin.Y);
 	KeyHash = HashCombine(KeyHash, PreviousOrigin.Z);
 
-	if ( InstanceBounds.IsValid )
+	if (InstanceBounds.IsValid)
 	{
 		KeyHash = HashCombine(KeyHash, InstanceBounds.Min.X);
 		KeyHash = HashCombine(KeyHash, InstanceBounds.Min.Y);
@@ -301,29 +301,29 @@ TUniformBuffer<FPrimitiveUniformShaderParameters>& FNiagaraSceneProxy::GetCustom
 	{
 		FPrimitiveUniformShaderParametersBuilder UBBuilder = FPrimitiveUniformShaderParametersBuilder()
 			.Defaults()
-				.LocalToWorld(GetLocalToWorld())
-				.PreviousLocalToWorld(PreviousLocalToWorld)
-				.ActorWorldPosition(GetActorPosition())
-				.WorldBounds(GetBounds())
-				.LocalBounds(GetLocalBounds())
-				.CustomPrimitiveData(GetCustomPrimitiveData())
-				.LightingChannelMask(GetLightingChannelMask())
-				.LightmapDataIndex(LocalPrimitiveSceneInfo ? LocalPrimitiveSceneInfo->GetLightmapDataOffset() : 0)
-				.LightmapUVIndex(GetLightMapCoordinateIndex())
-				.SingleCaptureIndex(SingleCaptureIndex)
-				.InstanceSceneDataOffset(LocalPrimitiveSceneInfo ? LocalPrimitiveSceneInfo->GetInstanceSceneDataOffset() : INDEX_NONE)
-				.NumInstanceSceneDataEntries(LocalPrimitiveSceneInfo ? LocalPrimitiveSceneInfo->GetNumInstanceSceneDataEntries() : 0)
-				.InstancePayloadDataOffset(LocalPrimitiveSceneInfo ? LocalPrimitiveSceneInfo->GetInstancePayloadDataOffset() : INDEX_NONE)
-				.InstancePayloadDataStride(LocalPrimitiveSceneInfo ? LocalPrimitiveSceneInfo->GetInstancePayloadDataStride() : 0)
-				.ReceivesDecals(ReceivesDecals())
-				.DrawsVelocity(DrawsVelocity())
-				.OutputVelocity(bOutputVelocity)
-				.CastContactShadow(CastsContactShadow())
-				.CastShadow(CastsDynamicShadow())
-				.HasCapsuleRepresentation(HasDynamicIndirectShadowCasterRepresentation())
-				.UseVolumetricLightmap(bHasPrecomputedVolumetricLightmap)
-				.UseSingleSampleShadowFromStationaryLights(UseSingleSampleShadowFromStationaryLights());
-		if ( InstanceBounds.IsValid )
+			.LocalToWorld(GetLocalToWorld())
+			.PreviousLocalToWorld(PreviousLocalToWorld)
+			.ActorWorldPosition(GetActorPosition())
+			.WorldBounds(GetBounds())
+			.LocalBounds(GetLocalBounds())
+			.CustomPrimitiveData(GetCustomPrimitiveData())
+			.LightingChannelMask(GetLightingChannelMask())
+			.LightmapDataIndex(LocalPrimitiveSceneInfo ? LocalPrimitiveSceneInfo->GetLightmapDataOffset() : 0)
+			.LightmapUVIndex(GetLightMapCoordinateIndex())
+			.SingleCaptureIndex(SingleCaptureIndex)
+			.InstanceSceneDataOffset(LocalPrimitiveSceneInfo ? LocalPrimitiveSceneInfo->GetInstanceSceneDataOffset() : INDEX_NONE)
+			.NumInstanceSceneDataEntries(LocalPrimitiveSceneInfo ? LocalPrimitiveSceneInfo->GetNumInstanceSceneDataEntries() : 0)
+			.InstancePayloadDataOffset(LocalPrimitiveSceneInfo ? LocalPrimitiveSceneInfo->GetInstancePayloadDataOffset() : INDEX_NONE)
+			.InstancePayloadDataStride(LocalPrimitiveSceneInfo ? LocalPrimitiveSceneInfo->GetInstancePayloadDataStride() : 0)
+			.ReceivesDecals(ReceivesDecals())
+			.DrawsVelocity(DrawsVelocity())
+			.OutputVelocity(bOutputVelocity)
+			.CastContactShadow(CastsContactShadow())
+			.CastShadow(CastsDynamicShadow())
+			.HasCapsuleRepresentation(HasDynamicIndirectShadowCasterRepresentation())
+			.UseVolumetricLightmap(bHasPrecomputedVolumetricLightmap)
+			.UseSingleSampleShadowFromStationaryLights(UseSingleSampleShadowFromStationaryLights());
+		if (InstanceBounds.IsValid)
 		{
 			UBBuilder.InstanceLocalBounds(InstanceBounds);
 
@@ -350,12 +350,12 @@ FRHIUniformBuffer* FNiagaraSceneProxy::GetCustomUniformBuffer(bool bHasVelocity,
 }
 
 uint32 FNiagaraSceneProxy::GetMemoryFootprint() const
-{ 
-	return (sizeof(*this) + GetAllocatedSize()); 
+{
+	return (sizeof(*this) + GetAllocatedSize());
 }
 
 uint32 FNiagaraSceneProxy::GetAllocatedSize() const
-{ 
+{
 	uint32 Size = FPrimitiveSceneProxy::GetAllocatedSize();
 	if (RenderData)
 	{
@@ -408,7 +408,7 @@ void FNiagaraSceneProxy::GetDynamicMeshElements(const TArray<const FSceneView*>&
 #if RHI_RAYTRACING
 void FNiagaraSceneProxy::GetDynamicRayTracingInstances(FRayTracingMaterialGatheringContext& Context, TArray<FRayTracingInstance>& OutRayTracingInstances)
 {
-if (RenderData)
+	if (RenderData)
 	{
 		RenderData->GetDynamicRayTracingInstances(Context, OutRayTracingInstances, *this);
 	}
@@ -533,7 +533,7 @@ void UNiagaraComponent::SetEmitterEnable(FName EmitterName, bool bNewEnableState
 void UNiagaraComponent::SetSystemFixedBounds(FBox LocalBounds)
 {
 	SystemFixedBounds = LocalBounds;
-	if ( SystemInstanceController.IsValid() )
+	if (SystemInstanceController.IsValid())
 	{
 		SystemInstanceController->SetSystemFixedBounds(LocalBounds);
 	}
@@ -564,7 +564,7 @@ void UNiagaraComponent::SetEmitterFixedBounds(FName EmitterName, FBox LocalBound
 
 FBox UNiagaraComponent::GetEmitterFixedBounds(FName EmitterName) const
 {
-	if ( const FBox* EmitterBounds = EmitterFixedBounds.Find(EmitterName) )
+	if (const FBox* EmitterBounds = EmitterFixedBounds.Find(EmitterName))
 	{
 		return *EmitterBounds;
 	}
@@ -584,16 +584,16 @@ void UNiagaraComponent::ReleaseToPool()
 {
 	// A component may be marked pending kill before the owner has it's reference set to null.
 	// In that case there's a window where it can be released back into the pool incorrectly, so we just skip releasing as we know it will be deleted shortly
-	if ( !IsValidChecked(this) || IsUnreachable() )
+	if (!IsValidChecked(this) || IsUnreachable())
 	{
 		return;
 	}
 
 	if (PoolingMethod != ENCPoolMethod::ManualRelease)
-	{		
+	{
 		// Only emit this warning if pooling is enabled. If it's not, all components will have PoolingMethod none.
 		if (UNiagaraComponentPool::Enabled() && FNiagaraUtilities::LogVerboseWarnings())
-		{	
+		{
 			UE_LOG(LogNiagara, Warning, TEXT("ReleaseToPool called but PoolingMethod is (%s) when it should be ENCPoolMethod::ManualRelease. Asset=%s Component=%s"), *StaticEnum<ENCPoolMethod>()->GetNameStringByValue(int64(PoolingMethod)), *GetPathNameSafe(Asset), *GetPathName());
 		}
 		return;
@@ -606,7 +606,7 @@ void UNiagaraComponent::ReleaseToPool()
 		//If we're already complete then release to the pool straight away.
 		UWorld* World = GetWorld();
 		check(World);
-		if( FNiagaraWorldManager* WorldMan = FNiagaraWorldManager::Get(World))
+		if (FNiagaraWorldManager* WorldMan = FNiagaraWorldManager::Get(World))
 		{
 			WorldMan->GetComponentPool()->ReclaimWorldParticleSystem(this);
 		}
@@ -697,7 +697,7 @@ void UNiagaraComponent::TickComponent(float DeltaSeconds, enum ELevelTick TickTy
 		{
 			SystemInstanceController->ManualTick(DeltaSeconds * CustomTimeDilation, (ThisTickFunction && ThisTickFunction->IsCompletionHandleValid()) ? ThisTickFunction->GetCompletionHandle() : nullptr);
 		}
-		else if(AgeUpdateMode == ENiagaraAgeUpdateMode::DesiredAge)
+		else if (AgeUpdateMode == ENiagaraAgeUpdateMode::DesiredAge)
 		{
 			float AgeDiff = FMath::Max(DesiredAge, 0.0f) - SystemInstanceController->GetAge();
 			int32 TicksToProcess = 0;
@@ -712,7 +712,7 @@ void UNiagaraComponent::TickComponent(float DeltaSeconds, enum ELevelTick TickTy
 					SystemInstanceController->Reset(FNiagaraSystemInstance::EResetMode::ResetAll);
 					AgeDiff = DesiredAge - SystemInstanceController->GetAge();
 				}
-				
+
 				if (AgeDiff > 0.0f)
 				{
 					FNiagaraSystemSimulation* SystemSim = SystemInstanceController->GetSoloSystemSimulation().Get();
@@ -726,7 +726,7 @@ void UNiagaraComponent::TickComponent(float DeltaSeconds, enum ELevelTick TickTy
 							const double EndMaxSimTime = bUseMaxSimTime ? FPlatformTime::Seconds() + MaxSimTime : 0.0f;
 
 							TicksToProcess = FMath::FloorToInt(AgeDiff / TickDelta);
-							while ( TicksToProcess > 0 )
+							while (TicksToProcess > 0)
 							{
 								// Cannot do multiple tick off the game thread here without additional work. So we pass in null for the completion event which will force GT execution.
 								// If this becomes a perf problem I can add a new path for the tick code to handle multiple ticks.
@@ -734,10 +734,10 @@ void UNiagaraComponent::TickComponent(float DeltaSeconds, enum ELevelTick TickTy
 								--TicksToProcess;
 
 								// This limits the amount of time we will consume seeking forward
-								if ( bUseMaxSimTime )
+								if (bUseMaxSimTime)
 								{
 									const double CurrentTime = FPlatformTime::Seconds();
-									if ( CurrentTime >= EndMaxSimTime )
+									if (CurrentTime >= EndMaxSimTime)
 									{
 										break;
 									}
@@ -868,12 +868,12 @@ bool UNiagaraComponent::IsWorldReadyToRun() const
 {
 	// The Niagara system instance assumes that a ComputeDispatchInterface exists when it is created. We need to wait until this has happened before successfully activating this system.
 	UWorld* World = GetWorld();
-	if ( World == nullptr)
+	if (World == nullptr)
 	{
 		return false;
 	}
 
-	if ( FNiagaraGpuComputeDispatchInterface::Get(World) == nullptr )
+	if (FNiagaraGpuComputeDispatchInterface::Get(World) == nullptr)
 	{
 		return false;
 	}
@@ -891,7 +891,7 @@ bool UNiagaraComponent::InitializeSystem()
 	{
 		LLM_SCOPE(ELLMTag::Niagara);
 		CSV_SCOPED_TIMING_STAT_EXCLUSIVE(Effects);
-		
+
 		UWorld* World = GetWorld();
 		check(World);
 		check(Asset);
@@ -923,7 +923,7 @@ bool UNiagaraComponent::InitializeSystem()
 		}
 		if (!EmitterFixedBounds.IsEmpty())
 		{
-			for ( auto it=EmitterFixedBounds.CreateConstIterator(); it; ++it )
+			for (auto it = EmitterFixedBounds.CreateConstIterator(); it; ++it)
 			{
 				SystemInstanceController->SetEmitterFixedBounds(it.Key(), it.Value());
 			}
@@ -993,7 +993,7 @@ void UNiagaraComponent::ActivateInternal(bool bReset /* = false */, bool bIsScal
 	// If the particle system can't ever render (ie on dedicated server or in a commandlet) than do not activate...
 	if (!FApp::CanEverRender() || !World || World->IsNetMode(NM_DedicatedServer) || World->bIsTearingDown)
 	{
-		if ( bAutoDestroy )
+		if (bAutoDestroy)
 		{
 			DestroyComponent();
 		}
@@ -1030,7 +1030,7 @@ void UNiagaraComponent::ActivateInternal(bool bReset /* = false */, bool bIsScal
 		Asset->PollForCompilationComplete();
 	}
 #endif
-	
+
 	if (!Asset->IsReadyToRun() || !IsWorldReadyToRun())
 	{
 		bAwaitingActivationDueToNotReady = true;
@@ -1052,10 +1052,10 @@ void UNiagaraComponent::ActivateInternal(bool bReset /* = false */, bool bIsScal
 	}
 
 	bIsCulledByScalability = false;
-	
+
 	//Update our owner scalability state but don't register with the manager.
 	ResolveOwnerAllowsScalability(false);
-	
+
 	//We reset last render time to the current time so that any visibility culling on a delay will function correctly.
 	//Leaving as the default of -1000 causes the visibility code to always assume this should be culled until it's first rendered and initialized by the RT.
 	//Set the component last render time *before* preculling otherwise there are cases where we can pre cull incorrectly.
@@ -1063,7 +1063,7 @@ void UNiagaraComponent::ActivateInternal(bool bReset /* = false */, bool bIsScal
 
 	if (ShouldPreCull())
 	{
-		#if WITH_PARTICLE_PERF_CSV_STATS
+#if WITH_PARTICLE_PERF_CSV_STATS
 
 		if (FCsvProfiler* CSVProfiler = FCsvProfiler::Get())
 		{
@@ -1073,7 +1073,7 @@ void UNiagaraComponent::ActivateInternal(bool bReset /* = false */, bool bIsScal
 				CSVProfiler->RecordCustomStat(PreCullStatName, CSV_CATEGORY_INDEX(Particles), 1, ECsvCustomStatOp::Accumulate);
 			}
 		}
-		#endif
+#endif
 
 		//We have decided to pre cull the system.
 		OnSystemComplete(true);
@@ -1101,7 +1101,7 @@ void UNiagaraComponent::ActivateInternal(bool bReset /* = false */, bool bIsScal
 		SetActiveFlag(true);
 		OnComponentActivated.Broadcast(this, bReset);
 	}
-	
+
 	// Auto attach if requested
 	const bool bWasAutoAttached = bDidAutoAttach;
 	bDidAutoAttach = false;
@@ -1131,11 +1131,11 @@ void UNiagaraComponent::ActivateInternal(bool bReset /* = false */, bool bIsScal
 			CancelAutoAttachment(true);
 		}
 	}
-	
+
 #if WITH_EDITOR
 	//TODO: Do this else where. I get needing to ensure params are correct from the component but these are stomping over runtime changes to the params in editor builds.
 	//For now we can bypass the worst of the impact by disallowing in game worlds.
-	if(!World->IsGameWorld())
+	if (!World->IsGameWorld())
 	{
 		ApplyOverridesToParameterStore();
 	}
@@ -1235,7 +1235,7 @@ void UNiagaraComponent::DeactivateInternal(bool bIsScalabilityCull /* = false */
 	{
 		Super::Deactivate();
 
-		if(bWasCulledByScalabiltiy && !bIsCulledByScalability)//We were culled by scalability but no longer, ensure we've handled completion correctly. E.g. returned to the pool etc.
+		if (bWasCulledByScalabiltiy && !bIsCulledByScalability)//We were culled by scalability but no longer, ensure we've handled completion correctly. E.g. returned to the pool etc.
 		{
 			OnSystemComplete(true);
 		}
@@ -1378,12 +1378,12 @@ void UNiagaraComponent::OnSystemComplete(bool bExternalCompletion)
 {
 	// Debug feature, if we have loop enabled all world FX will loop unless deactivate by scalability
 #if !UE_BUILD_SHIPPING
-	if ( !bExternalCompletion && !bIsCulledByScalability )
+	if (!bExternalCompletion && !bIsCulledByScalability)
 	{
-		if ( UWorld* World = GetWorld() )
+		if (UWorld* World = GetWorld())
 		{
 			FNiagaraWorldManager* WorldManager = FNiagaraWorldManager::Get(GetWorld());
-			if ( WorldManager->GetDebugPlaybackMode() == ENiagaraDebugPlaybackMode::Loop )
+			if (WorldManager->GetDebugPlaybackMode() == ENiagaraDebugPlaybackMode::Loop)
 			{
 				// If we have a loop time set the WorldManager will force a reset, so we will just ignore the completion event in that case
 				static const auto CVarGlobalLoopTime = IConsoleManager::Get().FindConsoleVariable(TEXT("fx.Niagara.Debug.GlobalLoopTime"));
@@ -1759,7 +1759,7 @@ void UNiagaraComponent::CreateCullProxy(bool bForce)
 	{
 		FNiagaraWorldManager* WorldMan = FNiagaraWorldManager::Get(GetWorld());
 		check(WorldMan);
-		
+
 		UNiagaraCullProxyComponent* SystemCullProxy = WorldMan->GetCullProxy(this);
 		if (SystemCullProxy)
 		{
@@ -1832,7 +1832,7 @@ void UNiagaraComponent::SendRenderDynamicData_Concurrent()
 		if (FNiagaraSystemRenderData* RenderData = NiagaraProxy->GetSystemRenderData())
 		{
 			RenderData->LWCRenderTile = SystemInstanceController->GetSystemInstance_Unsafe()->GetLWCTile();
-			
+
 			FNiagaraSceneProxy::FDynamicData NewProxyDynamicData;
 			NewProxyDynamicData.LODDistanceOverride = GetPreviewLODDistance();
 
@@ -1849,7 +1849,7 @@ void UNiagaraComponent::SendRenderDynamicData_Concurrent()
 
 				//Some renderers may need this transform on the GT so override here.
 				FNiagaraSystemInstanceControllerPtr CullProxyController = CullProxy->GetSystemInstanceController();
-				if ( ensureMsgf(CullProxyController.IsValid(), TEXT("CullProxy is missing the system instance controller for System '%s' Component '%s'"), *GetNameSafe(Asset), *GetFullNameSafe(this)) )
+				if (ensureMsgf(CullProxyController.IsValid(), TEXT("CullProxy is missing the system instance controller for System '%s' Component '%s'"), *GetNameSafe(Asset), *GetFullNameSafe(this)))
 				{
 					CullProxyController->GetSystemInstance_Unsafe()->SetWorldTransform(SystemInstanceController->GetSystemInstance_Unsafe()->GetWorldTransform());
 					CullProxyController->GenerateSetDynamicDataCommands(SetDataCommands, *RenderData, *NiagaraProxy);
@@ -1863,14 +1863,14 @@ void UNiagaraComponent::SendRenderDynamicData_Concurrent()
 
 			ENQUEUE_RENDER_COMMAND(NiagaraSetDynamicData)(
 				[NiagaraProxy, CommandsRT = MoveTemp(SetDataCommands), NewProxyDynamicData](FRHICommandListImmediate& RHICmdList)
-				{
-					SCOPE_CYCLE_COUNTER(STAT_NiagaraSetDynamicData);
-					PARTICLE_PERF_STAT_CYCLES_WITH_COUNT_RT(NewProxyDynamicData.PerfStatsContext, RenderUpdate, 1);
+			{
+				SCOPE_CYCLE_COUNTER(STAT_NiagaraSetDynamicData);
+				PARTICLE_PERF_STAT_CYCLES_WITH_COUNT_RT(NewProxyDynamicData.PerfStatsContext, RenderUpdate, 1);
 
-					FNiagaraSystemRenderData::ExecuteDynamicDataCommands_RenderThread(CommandsRT);
+				FNiagaraSystemRenderData::ExecuteDynamicDataCommands_RenderThread(CommandsRT);
 
-					NiagaraProxy->SetProxyDynamicData(NewProxyDynamicData);
-				}
+				NiagaraProxy->SetProxyDynamicData(NewProxyDynamicData);
+			}
 			);
 		}
 	}
@@ -1918,7 +1918,7 @@ FBoxSphereBounds UNiagaraComponent::CalcBounds(const FTransform& LocalToWorld) c
 			ensureMsgf(!SystemFixedBounds.Min.ContainsNaN() && !SystemFixedBounds.Max.ContainsNaN(), TEXT("SystemFixedBounds contains NaN for System(%d)"), *GetNameSafe(Asset));
 			SystemBounds = FBoxSphereBounds(SystemFixedBounds);
 		}
-		else if ( Asset && Asset->bFixedBounds )
+		else if (Asset && Asset->bFixedBounds)
 		{
 			const FBox AssetFixedBounds = Asset->GetFixedBounds();
 			ensureMsgf(!AssetFixedBounds.Min.ContainsNaN() && !AssetFixedBounds.Max.ContainsNaN(), TEXT("AssetFixedBounds contains NaN for System(%d)"), *GetNameSafe(Asset));
@@ -1955,18 +1955,18 @@ void UNiagaraComponent::GetUsedMaterials(TArray<UMaterialInterface*>& OutMateria
 bool UNiagaraComponent::ResolveOwnerAllowsScalability(bool bRegister)
 {
 	bOwnerAllowsScalabiltiy = true;
-	
-	if(GNiagaraCanPreventScalabilityCullingOnPlayerFX && Asset && Asset->AllowScalabilityForLocalPlayerFX() == false)
+
+	if (GNiagaraCanPreventScalabilityCullingOnPlayerFX && Asset && Asset->AllowScalabilityForLocalPlayerFX() == false)
 	{
 		bOwnerAllowsScalabiltiy = IsLocalPlayerEffect() == false;
-	}	
+	}
 
 	//Update manager registration if needed.
 	if (bRegister)
 	{
 		if (bOwnerAllowsScalabiltiy)
 		{
-			if ( IsActive() )
+			if (IsActive())
 			{
 				RegisterWithScalabilityManager();
 			}
@@ -2107,12 +2107,12 @@ void UNiagaraComponent::SetNiagaraVariableVec4(const FString& InVariableName, co
 void UNiagaraComponent::SetVariableVec3(FName InVariableName, FVector InValue)
 {
 	// If the user parameter is of position type we use the position function instead, so it gets converted correctly 
-	if (OverrideParameters.GetPositionParameterValue(InVariableName)) 
+	if (OverrideParameters.GetPositionParameterValue(InVariableName))
 	{
 		SetVariablePosition(InVariableName, InValue);
 		return;
 	}
-	
+
 #if !(UE_BUILD_SHIPPING || UE_BUILD_TEST)
 	if (!FitsIntoFloat(InValue))
 	{
@@ -2142,12 +2142,12 @@ void UNiagaraComponent::SetNiagaraVariableVec3(const FString& InVariableName, FV
 void UNiagaraComponent::SetVariablePosition(FName InVariableName, FVector InValue)
 {
 	// If the user parameter is not of position type we try to set it as vector instead, as users might confuse the two 
-	if (!OverrideParameters.GetPositionParameterValue(InVariableName)) 
+	if (!OverrideParameters.GetPositionParameterValue(InVariableName))
 	{
 		SetVariableVec3(InVariableName, InValue);
 		return;
 	}
-	
+
 	const FNiagaraVariable VariableDesc(FNiagaraTypeDefinition::GetPositionDef(), InVariableName);
 	if (SystemInstanceController.IsValid())
 	{
@@ -2403,7 +2403,7 @@ void FixInvalidUserParameters(FNiagaraUserRedirectionParameterStore& ParameterSt
 	}
 
 	bool bParameterRenamed = false;
-	for(FNiagaraVariable& IncorrectlyNamedParameter : IncorrectlyNamedParameters)
+	for (FNiagaraVariable& IncorrectlyNamedParameter : IncorrectlyNamedParameters)
 	{
 		FString FixedNameString = UserPrefix + IncorrectlyNamedParameter.GetName().ToString();
 		FName FixedName = *FixedNameString;
@@ -2456,7 +2456,7 @@ void FixInvalidUserParameters(FNiagaraUserRedirectionParameterStore& ParameterSt
 				}
 			}
 		}
-		else 
+		else
 		{
 			// If the correctly named parameter was not in the collection already we can just rename the incorrect parameter
 			// to the correct one.
@@ -2502,7 +2502,7 @@ void UNiagaraComponent::PostLoad()
 	{
 		Asset->ConditionalPostLoad();
 
-		FixInvalidUserParameters(OverrideParameters);		
+		FixInvalidUserParameters(OverrideParameters);
 		UpgradeDeprecatedParameterOverrides();
 
 #if WITH_EDITORONLY_DATA
@@ -2600,7 +2600,7 @@ void UNiagaraComponent::PostLoad()
 		FixDataInterfaceOuters();
 		OverrideParameters.SanityCheckData();
 
-		if(AssetExposedParametersChangedHandle.IsValid() == false)
+		if (AssetExposedParametersChangedHandle.IsValid() == false)
 		{
 			AssetExposedParametersChangedHandle = Asset->GetExposedParameters().AddOnChangedHandler(
 				FNiagaraParameterStore::FOnChanged::FDelegate::CreateUObject(this, &UNiagaraComponent::AssetExposedParametersChanged));
@@ -2636,7 +2636,7 @@ static FNiagaraVariant GetParameterValueFromStore(const FNiagaraVariableBase& Va
 
 	TArray<uint8> DataValue;
 	DataValue.AddUninitialized(Var.GetSizeInBytes());
-	if ( ensureMsgf(Store.CopyParameterData(Var, DataValue.GetData()), TEXT("Failed to copy parameter data for '%s'"), *Var.GetName().ToString()) )
+	if (ensureMsgf(Store.CopyParameterData(Var, DataValue.GetData()), TEXT("Failed to copy parameter data for '%s'"), *Var.GetName().ToString()))
 	{
 		return FNiagaraVariant(DataValue);
 	}
@@ -2647,8 +2647,8 @@ static FNiagaraVariant GetParameterValueFromStore(const FNiagaraVariableBase& Va
 
 void UNiagaraComponent::PreEditChange(FProperty* PropertyAboutToChange)
 {
-	if (PropertyAboutToChange != nullptr && 
-		PropertyAboutToChange->GetFName() == GET_MEMBER_NAME_CHECKED(UNiagaraComponent, Asset) && 
+	if (PropertyAboutToChange != nullptr &&
+		PropertyAboutToChange->GetFName() == GET_MEMBER_NAME_CHECKED(UNiagaraComponent, Asset) &&
 		Asset != nullptr)
 	{
 		Asset->GetExposedParameters().RemoveOnChangedHandler(AssetExposedParametersChangedHandle);
@@ -2674,7 +2674,7 @@ void UNiagaraComponent::PostEditChangeProperty(FPropertyChangedEvent& PropertyCh
 				FNiagaraParameterStore::FOnChanged::FDelegate::CreateUObject(this, &UNiagaraComponent::AssetExposedParametersChanged));
 		}
 	}
-	else if(PropertyName == GET_MEMBER_NAME_CHECKED(UNiagaraComponent, OverrideParameters))
+	else if (PropertyName == GET_MEMBER_NAME_CHECKED(UNiagaraComponent, OverrideParameters))
 	{
 		SynchronizeWithSourceSystem();
 	}
@@ -2698,7 +2698,7 @@ void UNiagaraComponent::SetUserParametersToDefaultValues()
 	TemplateParameterOverrides.Empty();
 	InstanceParameterOverrides.Empty();
 #endif
-	
+
 	if (Asset == nullptr)
 	{
 		OverrideParameters.Empty(false);
@@ -2713,7 +2713,7 @@ void UNiagaraComponent::SetUserParametersToDefaultValues()
 		TArray<FNiagaraVariableBase, TInlineAllocator<8>> ParametersToRemove;
 		bool bInterfacesChanged = false;
 
-		for (int32 i=0; i < DestParameters.Num(); ++i)
+		for (int32 i = 0; i < DestParameters.Num(); ++i)
 		{
 			const FNiagaraVariableWithOffset& DestParameter = DestParameters[i];
 			const int32 DestIndex = DestParameter.Offset;
@@ -2776,16 +2776,16 @@ void UNiagaraComponent::UpgradeDeprecatedParameterOverrides()
 
 	OverrideParameters.SanityCheckData();
 	PostLoadNormalizeOverrideNames();
-	
+
 	TArray<FNiagaraVariable> UserParameters;
 	OverrideParameters.GetUserParameters(UserParameters);
 
 	for (const TPair<FName, bool>& Pair : EditorOverridesValue_DEPRECATED)
 	{
 		FNiagaraVariable* Found = UserParameters.FindByPredicate([Pair](const FNiagaraVariable& Var)
-		{
-			return Var.GetName() == Pair.Key;
-		});
+			{
+				return Var.GetName() == Pair.Key;
+			});
 
 		if (Found != nullptr)
 		{
@@ -2807,10 +2807,10 @@ void UNiagaraComponent::EnsureOverrideParametersConsistent() const
 	{
 		return;
 	}
-	
+
 	TArray<FNiagaraVariable> UserParameters;
 	Asset->GetExposedParameters().GetUserParameters(UserParameters);
-	
+
 	for (const FNiagaraVariable& Key : UserParameters)
 	{
 		FNiagaraVariant OverrideValue = FindParameterOverride(Key);
@@ -2864,7 +2864,7 @@ void UNiagaraComponent::ApplyOverridesToParameterStore()
 	}
 
 	for (const auto& Pair : TemplateParameterOverrides)
-	{ 
+	{
 		if (!FNiagaraUserRedirectionParameterStore::IsUserParameter(Pair.Key))
 		{
 			continue;
@@ -2880,7 +2880,7 @@ void UNiagaraComponent::ApplyOverridesToParameterStore()
 	if (!IsTemplate())
 	{
 		for (const auto& Pair : InstanceParameterOverrides)
-		{ 
+		{
 			if (!FNiagaraUserRedirectionParameterStore::IsUserParameter(Pair.Key))
 			{
 				continue;
@@ -2989,7 +2989,7 @@ void UNiagaraComponent::SynchronizeWithSourceSystem()
 #if WITH_EDITORONLY_DATA
 		OverrideParameters.Empty();
 		EditorOverridesValue_DEPRECATED.Empty();
-		
+
 		OnSynchronizedWithAssetParametersDelegate.Broadcast();
 #endif
 		return;
@@ -3045,14 +3045,14 @@ void UNiagaraComponent::FixInvalidUserParameterOverrideData()
 			{
 				OverrideParameters.ChangeParameterType(ExistingVar, FNiagaraTypeDefinition::GetPositionDef());
 
-				auto ConvertOverrideToPosition = [&ExistingVar, &PositionVar](auto & Params)
+				auto ConvertOverrideToPosition = [&ExistingVar, &PositionVar](auto& Params)
 				{
 					FNiagaraVariant VariantData;
 					if (Params.RemoveAndCopyValue(ExistingVar, VariantData))
 					{
 						// Convert the old float data to doubles.
 						check(VariantData.GetNumBytes() == sizeof(FVector3f));
-						FVector3f ExistingData = *((FVector3f*)(VariantData.GetBytes()));						
+						FVector3f ExistingData = *((FVector3f*)(VariantData.GetBytes()));
 						FVector NewData(ExistingData);
 						VariantData.SetBytes((uint8*)&NewData, sizeof(NewData));
 						Params.Add(PositionVar, VariantData);
@@ -3061,7 +3061,7 @@ void UNiagaraComponent::FixInvalidUserParameterOverrideData()
 
 				ConvertOverrideToPosition(InstanceParameterOverrides);
 				ConvertOverrideToPosition(TemplateParameterOverrides);
-				
+
 				ExistingVar = PositionVar;
 			}
 		}
@@ -3081,7 +3081,7 @@ void UNiagaraComponent::AssetExposedParametersChanged()
 #if WITH_EDITOR
 	//-TODO: We can possibly remove the ReinitializeSystem here after more testing.
 	//       This fixes an issue with deleting a renderer and hitting undo in the editor.
-	if ( !GIsTransacting )
+	if (!GIsTransacting)
 #endif
 	{
 		ReinitializeSystem();
@@ -3090,7 +3090,7 @@ void UNiagaraComponent::AssetExposedParametersChanged()
 
 #if WITH_EDITOR
 
-bool UNiagaraComponent::HasParameterOverride(const FNiagaraVariableBase& InKey) const 
+bool UNiagaraComponent::HasParameterOverride(const FNiagaraVariableBase& InKey) const
 {
 	FNiagaraVariableBase UserVariable = InKey;
 
@@ -3398,7 +3398,7 @@ void UNiagaraComponent::SetPreviewLODDistance(bool bInEnablePreviewLODDistance,
 	bEnablePreviewLODDistance = bInEnablePreviewLODDistance;
 	PreviewLODDistance = InPreviewLODDistance;
 
-	if(SystemInstanceController)
+	if (SystemInstanceController)
 	{
 		if (bEnablePreviewLODDistance)
 		{
@@ -3413,10 +3413,10 @@ void UNiagaraComponent::SetPreviewLODDistance(bool bInEnablePreviewLODDistance,
 
 void UNiagaraComponent::SetAllowScalability(bool bAllow)
 {
-	bAllowScalability = bAllow; 
+	bAllowScalability = bAllow;
 	if (bAllow)
 	{
-		if ( IsActive() )
+		if (IsActive())
 		{
 			RegisterWithScalabilityManager();
 		}
@@ -3454,7 +3454,7 @@ void UNiagaraComponent::PostLoadNormalizeOverrideNames()
 		FString ValueNameString = Pair.Key.ToString();
 		if (ValueNameString.StartsWith(TEXT("User.")))
 		{
-			 ValueNameString = ValueNameString.RightChop(5);
+			ValueNameString = ValueNameString.RightChop(5);
 		}
 
 		ValueMap.Add(FName(*ValueNameString), Pair.Value);
@@ -3473,9 +3473,9 @@ void UNiagaraComponent::SetAsset(UNiagaraSystem* InAsset, bool bResetExistingOve
 
 	if (FNiagaraUtilities::LogVerboseWarnings())
 	{
-		if ( PoolingMethod != ENCPoolMethod::None )
+		if (PoolingMethod != ENCPoolMethod::None)
 		{
-			UE_LOG(LogNiagara, Warning, TEXT("SetAsset called on pooled component '%s' Before '%s' New '%s', pleased fix calling code to not do this."), *GetFullNameSafe(this), *GetFullNameSafe(Asset), * GetFullNameSafe(InAsset));
+			UE_LOG(LogNiagara, Warning, TEXT("SetAsset called on pooled component '%s' Before '%s' New '%s', pleased fix calling code to not do this."), *GetFullNameSafe(this), *GetFullNameSafe(Asset), *GetFullNameSafe(InAsset));
 		}
 	}
 
@@ -3512,7 +3512,7 @@ void UNiagaraComponent::SetAsset(UNiagaraSystem* InAsset, bool bResetExistingOve
 	}
 #endif
 
-	if ( Asset )
+	if (Asset)
 	{
 		// Apply system overrides to the component
 		CastShadow = Asset->bOverrideCastShadow ? Asset->bCastShadow : CastShadow;
@@ -3523,7 +3523,7 @@ void UNiagaraComponent::SetAsset(UNiagaraSystem* InAsset, bool bResetExistingOve
 		CustomDepthStencilValue = Asset->bOverrideCustomDepthStencilWriteMask ? Asset->CustomDepthStencilValue : CustomDepthStencilValue;
 
 		// Activate
-		if ( IsRegistered() )
+		if (IsRegistered())
 		{
 			if (bAutoActivate || bWasActive)
 			{
@@ -3533,8 +3533,8 @@ void UNiagaraComponent::SetAsset(UNiagaraSystem* InAsset, bool bResetExistingOve
 	}
 }
 
-void UNiagaraComponent::SetForceSolo(bool bInForceSolo) 
-{ 
+void UNiagaraComponent::SetForceSolo(bool bInForceSolo)
+{
 	if (bForceSolo != bInForceSolo)
 	{
 		bForceSolo = bInForceSolo;
@@ -3571,7 +3571,7 @@ void UNiagaraComponent::SetGpuComputeDebug(bool bEnableDebug)
 
 void UNiagaraComponent::SetCustomTimeDilation(float Dilation)
 {
-	if ( !FMath::IsNearlyEqual(CustomTimeDilation, Dilation) )
+	if (!FMath::IsNearlyEqual(CustomTimeDilation, Dilation))
 	{
 		CustomTimeDilation = Dilation;
 		UpdateInstanceSoloMode();
@@ -3608,4 +3608,4 @@ void UNiagaraComponent::CancelAutoAttachment(bool bDetachFromParent)
 			//bIsChangingAutoAttachment = false;
 		}
 	}
-}
+}
\ No newline at end of file
