protect against overflow for 3d flip particles to grid

diff --git a/Engine/Plugins/FX/Niagara/Source/Niagara/Private/NiagaraDataInterfaceRasterizationGrid3D.cpp b/Engine/Plugins/FX/Niagara/Source/Niagara/Private/NiagaraDataInterfaceRasterizationGrid3D.cpp
--- a/Engine/Plugins/FX/Niagara/Source/Niagara/Private/NiagaraDataInterfaceRasterizationGrid3D.cpp
+++ b/Engine/Plugins/FX/Niagara/Source/Niagara/Private/NiagaraDataInterfaceRasterizationGrid3D.cpp
@@ -664,22 +664,26 @@ bool UNiagaraDataInterfaceRasterizationGrid3D::GetFunctionHLSL(const FNiagaraDat
 		static const TCHAR* FormatBounds = TEXT(R"(
 			void {FunctionName}(int In_IndexX, int In_IndexY, int In_IndexZ, int In_AttributeIndex, float In_Value, out int val)
 			{							
-				val = -1;					
+				val = 1;					
 				if ( In_AttributeIndex < {NumAttributesName} )
 				{
 					int OriginalValue;
 					int IntValue = {FloatToIntFunctionName}(In_Value);
-					int3 TileOffset = {PerAttributeDataName}[In_AttributeIndex].xyz;		
-					InterlockedAdd(RW{OutputIntGrid}[int3(In_IndexX + TileOffset.x, In_IndexY + TileOffset.y, In_IndexZ + TileOffset.z)], IntValue, OriginalValue);
+					int3 TileOffset = {PerAttributeDataName}[In_AttributeIndex].xyz;	
+
+					int3 Index = int3(In_IndexX + TileOffset.x, In_IndexY + TileOffset.y, In_IndexZ + TileOffset.z);
+					InterlockedAdd(RW{OutputIntGrid}[Index], IntValue, OriginalValue);
+
+					int StoredValue = IntValue + OriginalValue;
 
-					if ((IntValue > 0 && IntValue + OriginalValue < OriginalValue) || 
-						(IntValue < 0 && IntValue + OriginalValue > OriginalValue))
+					// make sure to store max/min float value in the grid if we've over/underflowed
+					[branch]
+					if ((IntValue > 0 && StoredValue < OriginalValue) || (IntValue < 0 && StoredValue > OriginalValue))
 					{
 						val = 0;
-					}
-					else
-					{
-						val = 1;
+
+						int NewValue = IntValue > 0 ? 2147483647 : -2147483648;
+						RW{OutputIntGrid}[Index] = In_Value;
 					}
 				}
 			}
