Added custom validation rule support to Niagara

diff --git a/Engine/Plugins/FX/Niagara/Source/Niagara/Classes/NiagaraEffectType.h b/Engine/Plugins/FX/Niagara/Source/Niagara/Classes/NiagaraEffectType.h
--- a/Engine/Plugins/FX/Niagara/Source/Niagara/Classes/NiagaraEffectType.h
+++ b/Engine/Plugins/FX/Niagara/Source/Niagara/Classes/NiagaraEffectType.h
@@ -6,6 +6,7 @@
 #include "RenderCommandFence.h"
 #include "NiagaraPlatformSet.h"
 #include "NiagaraPerfBaseline.h"
+#include "NiagaraValidationRule.h"
 #include "NiagaraEffectType.generated.h"
 
 #define DEBUG_SCALABILITY_STATE (!UE_BUILD_SHIPPING)
@@ -404,6 +405,10 @@ class NIAGARA_API UNiagaraEffectType : public UObject
 	UPROPERTY(EditAnywhere, Category = "Scalability")
 	FNiagaraEmitterScalabilitySettingsArray EmitterScalabilitySettings;
 
+	/** A set of rules to apply when checking content. To create your own rules, write a custom class that extends UNiagaraValidationRule. */
+	UPROPERTY(EditAnywhere, Category = "Validation")
+	TArray<TSubclassOf<UNiagaraValidationRule>> ValidationRules;
+
 	FORCEINLINE const FNiagaraSystemScalabilitySettingsArray& GetSystemScalabilitySettings()const { return SystemScalabilitySettings; }
 	FORCEINLINE const FNiagaraEmitterScalabilitySettingsArray& GetEmitterScalabilitySettings()const { return EmitterScalabilitySettings; }
 
diff --git a/Engine/Plugins/FX/Niagara/Source/Niagara/Classes/NiagaraValidationRule.h b/Engine/Plugins/FX/Niagara/Source/Niagara/Classes/NiagaraValidationRule.h
new file mode 100644
--- /dev/null
+++ b/Engine/Plugins/FX/Niagara/Source/Niagara/Classes/NiagaraValidationRule.h
@@ -0,0 +1,58 @@
+﻿// Copyright Epic Games, Inc. All Rights Reserved.
+
+#pragma once
+
+#include "CoreMinimal.h"
+#include "UObject/Object.h"
+#include "NiagaraValidationRule.generated.h"
+
+class FNiagaraSystemViewModel;
+
+UENUM()
+enum class ENiagaraValidationSeverity
+{
+	/** Just an info message for the user. */
+	Info,
+	/** Could be a potential problem, for example bad performance. */
+	Warning,
+	/** A problem that must be fixed for the content to be valid. */
+	Error,
+};
+
+USTRUCT()
+struct FNiagaraValidationResult
+{
+	GENERATED_USTRUCT_BODY()
+
+	FNiagaraValidationResult() {};
+	FNiagaraValidationResult(ENiagaraValidationSeverity InSeverity, FText InSummaryText, FText InDescription, TWeakObjectPtr<UObject> InSourceObject)
+		: Severity(InSeverity), SummaryText(InSummaryText), Description(InDescription), SourceObject(InSourceObject) {}
+
+	UPROPERTY()
+	ENiagaraValidationSeverity Severity = ENiagaraValidationSeverity::Info;
+	
+	UPROPERTY()
+	FText SummaryText;
+
+	UPROPERTY()
+	FText Description;
+
+	UPROPERTY()
+	TWeakObjectPtr<UObject> SourceObject;
+};
+
+/**
+Base class for system validation logic. 
+These allow Niagara systems to be inspected for content validation either at save time or from a commandlet.
+*/
+UCLASS(abstract)
+class NIAGARA_API UNiagaraValidationRule : public UObject
+{
+	GENERATED_BODY()
+public:
+
+#if WITH_EDITOR
+	virtual TArray<FNiagaraValidationResult> CheckValidity(TSharedPtr<FNiagaraSystemViewModel> ViewModel) const;
+#endif
+};
+
diff --git a/Engine/Plugins/FX/Niagara/Source/Niagara/Private/NiagaraValidationRule.cpp b/Engine/Plugins/FX/Niagara/Source/Niagara/Private/NiagaraValidationRule.cpp
new file mode 100644
--- /dev/null
+++ b/Engine/Plugins/FX/Niagara/Source/Niagara/Private/NiagaraValidationRule.cpp
@@ -0,0 +1,10 @@
+﻿// Copyright Epic Games, Inc. All Rights Reserved.
+
+#include "NiagaraValidationRule.h"
+
+#if WITH_EDITOR
+TArray<FNiagaraValidationResult> UNiagaraValidationRule::CheckValidity(TSharedPtr<FNiagaraSystemViewModel> ViewModel) const
+{
+	return TArray<FNiagaraValidationResult>();
+}
+#endif
\ No newline at end of file
diff --git a/Engine/Plugins/FX/Niagara/Source/NiagaraEditor/Private/NiagaraValidationRules.cpp b/Engine/Plugins/FX/Niagara/Source/NiagaraEditor/Private/NiagaraValidationRules.cpp
new file mode 100644
--- /dev/null
+++ b/Engine/Plugins/FX/Niagara/Source/NiagaraEditor/Private/NiagaraValidationRules.cpp
@@ -0,0 +1,143 @@
+﻿// Copyright Epic Games, Inc. All Rights Reserved.
+
+#include "NiagaraValidationRules.h"
+
+#include "NiagaraComponentRendererProperties.h"
+#include "NiagaraSystem.h"
+#include "NiagaraSystemEditorData.h"
+#include "ViewModels/NiagaraEmitterHandleViewModel.h"
+#include "ViewModels/NiagaraSystemViewModel.h"
+#include "ViewModels/Stack/NiagaraStackEmitterSettingsGroup.h"
+#include "ViewModels/Stack/NiagaraStackRendererItem.h"
+#include "ViewModels/Stack/NiagaraStackSystemPropertiesItem.h"
+#include "ViewModels/Stack/NiagaraStackViewModel.h"
+
+#define LOCTEXT_NAMESPACE "NiagaraValidationRules"
+
+template<typename T>
+TArray<T*> GetStackEntries(UNiagaraStackViewModel* StackViewModel, bool bRefresh = false)
+{
+	TArray<T*> Results;
+	TArray<UNiagaraStackEntry*> EntriesToCheck;
+	UNiagaraStackEntry* RootEntry = StackViewModel->GetRootEntry();
+	if (bRefresh)
+	{
+		RootEntry->RefreshChildren();
+	}
+	RootEntry->GetUnfilteredChildren(EntriesToCheck);
+	while (EntriesToCheck.Num() > 0)
+	{
+		UNiagaraStackEntry* Entry = EntriesToCheck.Pop();
+		if (T* ItemToCheck = Cast<T>(Entry))
+		{
+			Results.Add(ItemToCheck);
+		}
+		Entry->GetUnfilteredChildren(EntriesToCheck);
+	}
+	return Results;
+}
+
+// helper function to retrieve a single stack entry from the system or emitter view model
+template<typename T>
+T* GetStackEntry(UNiagaraStackViewModel* StackViewModel, bool bRefresh = false)
+{
+	TArray<T*> StackEntries = GetStackEntries<T>(StackViewModel, bRefresh);
+	if (StackEntries.Num() > 0)
+	{
+		return StackEntries[0];
+	}
+	return nullptr;
+}
+
+// --------------------------------------------------------------------------------------------------------------------------------------------
+
+TArray<FNiagaraValidationResult> UNiagara_NoWarmupTime::CheckValidity(TSharedPtr<FNiagaraSystemViewModel> ViewModel) const
+{
+	TArray<FNiagaraValidationResult> Results;
+	UNiagaraSystem& System = ViewModel->GetSystem();
+	if (System.NeedsWarmup())
+	{
+		UNiagaraStackSystemPropertiesItem* SystemProperties = GetStackEntry<UNiagaraStackSystemPropertiesItem>(ViewModel->GetSystemStackViewModel());
+		FNiagaraValidationResult Result(ENiagaraValidationSeverity::Error, LOCTEXT("WarumupSummary", "Warmuptime > 0 is not allowed"), LOCTEXT("WarmupDescription", "Systems with the chosen effect type do not allow warmup time, as it costs too much performance.\nPlease set the warmup time to 0 in the system properties."), SystemProperties);
+		Results.Add(Result);
+	}
+	return Results;
+}
+
+TArray<FNiagaraValidationResult> UNiagara_FixedGPUBoundsSet::CheckValidity(TSharedPtr<FNiagaraSystemViewModel> ViewModel) const
+{
+	TArray<FNiagaraValidationResult> Results;
+
+	// if the system has fixed bounds set then it overrides the emitter settings
+	if (ViewModel->GetSystem().bFixedBounds)
+	{
+		return Results;
+	}
+
+	// check that all the gpu emitters have fixed bounds set
+	TArray<TSharedRef<FNiagaraEmitterHandleViewModel>> EmitterHandleViewModels = ViewModel->GetEmitterHandleViewModels();
+	for (TSharedRef<FNiagaraEmitterHandleViewModel> EmitterHandleModel : EmitterHandleViewModels)
+	{
+		UNiagaraEmitter* NiagaraEmitter = EmitterHandleModel.Get().GetEmitterHandle()->GetInstance();
+		if (NiagaraEmitter->SimTarget == ENiagaraSimTarget::GPUComputeSim && NiagaraEmitter->bFixedBounds == false)
+		{
+			UNiagaraStackEmitterPropertiesItem* EmitterProperties = GetStackEntry<UNiagaraStackEmitterPropertiesItem>(EmitterHandleModel.Get().GetEmitterStackViewModel());
+			FNiagaraValidationResult Result(ENiagaraValidationSeverity::Error, LOCTEXT("FixedBoundsSummary", "GPU emitters require fixed bounds"), LOCTEXT("FixedBoundsDescription", "Please set fixed bounds for gpu emitters or the system as a whole."), EmitterProperties);
+			Results.Add(Result);
+		}
+	}
+	return Results;
+}
+
+bool IsEnabledForMaxQualityLevel(FNiagaraPlatformSet Platforms, int32 MaxQualityLevel)
+{
+	for (int i = 0; i < MaxQualityLevel; i++)
+	{
+		if (Platforms.IsEnabledForQualityLevel(i))
+		{
+			return true;
+		}
+	}
+	return false;
+}
+
+TArray<FNiagaraValidationResult> UNiagara_NoComponentRendererOnLowSettings::CheckValidity(TSharedPtr<FNiagaraSystemViewModel> ViewModel) const
+{
+	TArray<FNiagaraValidationResult> Results;
+
+	// check that component renderers are only active for high and cinematic settings 
+	TArray<TSharedRef<FNiagaraEmitterHandleViewModel>> EmitterHandleViewModels = ViewModel->GetEmitterHandleViewModels();
+	for (TSharedRef<FNiagaraEmitterHandleViewModel> EmitterHandleModel : EmitterHandleViewModels)
+	{
+		int32 MaxQualityLevel = 3;
+		UNiagaraEmitter* NiagaraEmitter = EmitterHandleModel.Get().GetEmitterHandle()->GetInstance();
+		NiagaraEmitter->ForEachRenderer([&Results, EmitterHandleModel, MaxQualityLevel](UNiagaraRendererProperties* RendererProperties)
+		{
+			if (Cast<UNiagaraComponentRendererProperties>(RendererProperties) && RendererProperties->GetIsEnabled() && IsEnabledForMaxQualityLevel(RendererProperties->Platforms, MaxQualityLevel))
+			{
+				TArray<UNiagaraStackRendererItem*> RendererItems = GetStackEntries<UNiagaraStackRendererItem>(EmitterHandleModel.Get().GetEmitterStackViewModel());
+				for (UNiagaraStackRendererItem* Item : RendererItems)
+				{
+					if (Item->GetRendererProperties() != RendererProperties)
+					{
+						continue;
+					}
+					FNiagaraValidationResult Result(ENiagaraValidationSeverity::Warning, LOCTEXT("ComponentRenderSummary", "Component renderers should not be used for low-level platforms"), LOCTEXT("ComponentRenderDescription", "Please disable low level platforms in the renderer's scalability settings."), Item);
+					Results.Add(Result);
+				}
+			}
+		});
+	}
+	return Results;
+}
+
+TArray<FNiagaraValidationResult> UNiagara_InvalidEffectType::CheckValidity(TSharedPtr<FNiagaraSystemViewModel> ViewModel) const
+{
+	TArray<FNiagaraValidationResult> Results;
+	UNiagaraStackSystemPropertiesItem* SystemProperties = GetStackEntry<UNiagaraStackSystemPropertiesItem>(ViewModel->GetSystemStackViewModel());
+	FNiagaraValidationResult Result(ENiagaraValidationSeverity::Error, LOCTEXT("InvalidEffectSummary", "Invalid Effect Type"), LOCTEXT("InvalidEffectDescription", "The effect type on this system was marked as invalid for production content and should only be used as placeholder."), SystemProperties);
+	Results.Add(Result);
+	return Results;
+}
+
+#undef LOCTEXT_NAMESPACE
diff --git a/Engine/Plugins/FX/Niagara/Source/NiagaraEditor/Private/ViewModels/Stack/NiagaraStackEmitterSettingsGroup.cpp b/Engine/Plugins/FX/Niagara/Source/NiagaraEditor/Private/ViewModels/Stack/NiagaraStackEmitterSettingsGroup.cpp
--- a/Engine/Plugins/FX/Niagara/Source/NiagaraEditor/Private/ViewModels/Stack/NiagaraStackEmitterSettingsGroup.cpp
+++ b/Engine/Plugins/FX/Niagara/Source/NiagaraEditor/Private/ViewModels/Stack/NiagaraStackEmitterSettingsGroup.cpp
@@ -14,6 +14,8 @@
 #include "NiagaraEditorStyle.h"
 #include "NiagaraSystem.h"
 #include "NiagaraEditorStyle.h"
+#include "ViewModels/NiagaraEmitterHandleViewModel.h"
+#include "ViewModels/Stack/NiagaraStackViewModel.h"
 #include "Styling/AppStyle.h"
 #include "IDetailTreeNode.h"
 
@@ -154,13 +156,14 @@ void UNiagaraStackEmitterPropertiesItem::EmitterPropertiesChanged()
 		// so guard against receiving an event when finalized here.
 		bCanResetToBaseCache.Reset();
 		RefreshChildren();
+
+		if (Emitter.IsValid())
+		{
+			GetSystemViewModel()->GetEmitterHandleViewModelForEmitter(Emitter.Get()).Get()->GetEmitterStackViewModel()->RequestValidationUpdate();
+		}
 	}
 }
 
-
-
-
-
 void UNiagaraStackEmitterSummaryItem::Initialize(FRequiredEntryData InRequiredEntryData)
 {
 	Super::Initialize(InRequiredEntryData, TEXT("EmitterParameters"));
diff --git a/Engine/Plugins/FX/Niagara/Source/NiagaraEditor/Private/ViewModels/Stack/NiagaraStackEntry.cpp b/Engine/Plugins/FX/Niagara/Source/NiagaraEditor/Private/ViewModels/Stack/NiagaraStackEntry.cpp
--- a/Engine/Plugins/FX/Niagara/Source/NiagaraEditor/Private/ViewModels/Stack/NiagaraStackEntry.cpp
+++ b/Engine/Plugins/FX/Niagara/Source/NiagaraEditor/Private/ViewModels/Stack/NiagaraStackEntry.cpp
@@ -608,6 +608,29 @@ const TArray<UNiagaraStackEntry*>& UNiagaraStackEntry::GetAllChildrenWithIssues(
 	return GetCollectedIssueData().ChildrenWithIssues;
 }
 
+void UNiagaraStackEntry::AddExternalIssue(EStackIssueSeverity Severity, const FText& SummaryText, const FText& Description, bool bCanBeDismissed)
+{
+	FStackIssue NewIssue(Severity, SummaryText, Description, GetStackEditorDataKey(), bCanBeDismissed);
+	for (const FStackIssue& ExistingIssue : ExternalStackIssues)
+	{
+		if (ExistingIssue.GetUniqueIdentifier().Equals(NewIssue.GetUniqueIdentifier()))
+		{
+			return;
+		}
+	}
+	ExternalStackIssues.Add(NewIssue);
+	RefreshChildren();
+}
+
+void UNiagaraStackEntry::ClearExternalIssues()
+{
+	if (ExternalStackIssues.Num() > 0)
+	{
+		ExternalStackIssues.Empty();
+		Cast<UNiagaraStackEntry>(GetOuter())->RefreshChildren();
+	}
+}
+
 TOptional<UNiagaraStackEntry::FDropRequestResponse> UNiagaraStackEntry::CanDropInternal(const FDropRequest& DropRequest)
 {
 	return TOptional<FDropRequestResponse>();
@@ -713,6 +736,7 @@ void UNiagaraStackEntry::RefreshChildren()
 
 	StackIssues.Empty();
 	StackIssues.Append(NewStackIssues);
+	StackIssues.Append(ExternalStackIssues);
 	RefreshStackErrorChildren();
 	for (UNiagaraStackErrorItem* ErrorChild : ErrorChildren)
 	{
diff --git a/Engine/Plugins/FX/Niagara/Source/NiagaraEditor/Private/ViewModels/Stack/NiagaraStackRendererItem.cpp b/Engine/Plugins/FX/Niagara/Source/NiagaraEditor/Private/ViewModels/Stack/NiagaraStackRendererItem.cpp
--- a/Engine/Plugins/FX/Niagara/Source/NiagaraEditor/Private/ViewModels/Stack/NiagaraStackRendererItem.cpp
+++ b/Engine/Plugins/FX/Niagara/Source/NiagaraEditor/Private/ViewModels/Stack/NiagaraStackRendererItem.cpp
@@ -28,6 +28,8 @@
 #include "NiagaraEmitterEditorData.h"
 
 #include "Styling/SlateIconFinder.h"
+#include "ViewModels/NiagaraEmitterHandleViewModel.h"
+#include "ViewModels/Stack/NiagaraStackViewModel.h"
 
 #define LOCTEXT_NAMESPACE "UNiagaraStackRendererItem"
 
@@ -467,6 +469,12 @@ void UNiagaraStackRendererItem::RendererChanged()
 		// so guard against receiving an event when finalized here.
 		bCanResetToBaseCache.Reset();
 		RefreshChildren();
+
+		
+		if (GetEmitterViewModel().IsValid())
+		{
+			GetSystemViewModel()->GetEmitterHandleViewModelForEmitter(GetEmitterViewModel()->GetEmitter()).Get()->GetEmitterStackViewModel()->RequestValidationUpdate();
+		}
 	}
 }
 
diff --git a/Engine/Plugins/FX/Niagara/Source/NiagaraEditor/Private/ViewModels/Stack/NiagaraStackSystemPropertiesItem.cpp b/Engine/Plugins/FX/Niagara/Source/NiagaraEditor/Private/ViewModels/Stack/NiagaraStackSystemPropertiesItem.cpp
--- a/Engine/Plugins/FX/Niagara/Source/NiagaraEditor/Private/ViewModels/Stack/NiagaraStackSystemPropertiesItem.cpp
+++ b/Engine/Plugins/FX/Niagara/Source/NiagaraEditor/Private/ViewModels/Stack/NiagaraStackSystemPropertiesItem.cpp
@@ -6,6 +6,8 @@
 #include "NiagaraSystemDetailsCustomization.h"
 #include "ViewModels/NiagaraSystemViewModel.h"
 #include "ScopedTransaction.h"
+#include "ViewModels/NiagaraEmitterHandleViewModel.h"
+#include "ViewModels/Stack/NiagaraStackViewModel.h"
 
 #define LOCTEXT_NAMESPACE "UNiagaraStackSystemItemGroup"
 
@@ -78,6 +80,12 @@ void UNiagaraStackSystemPropertiesItem::RefreshChildrenInternal(const TArray<UNi
 void UNiagaraStackSystemPropertiesItem::SystemPropertiesChanged()
 {
 	bCanResetToBase.Reset();
+
+	GetSystemViewModel()->GetSystemStackViewModel()->RequestValidationUpdate();
+	for (TSharedRef<FNiagaraEmitterHandleViewModel> EmitterHandleViewModel : GetSystemViewModel()->GetEmitterHandleViewModels())
+	{
+		EmitterHandleViewModel->GetEmitterStackViewModel()->RequestValidationUpdate();
+	}
 }
 
 #undef LOCTEXT_NAMESPACE
diff --git a/Engine/Plugins/FX/Niagara/Source/NiagaraEditor/Private/ViewModels/Stack/NiagaraStackViewModel.cpp b/Engine/Plugins/FX/Niagara/Source/NiagaraEditor/Private/ViewModels/Stack/NiagaraStackViewModel.cpp
--- a/Engine/Plugins/FX/Niagara/Source/NiagaraEditor/Private/ViewModels/Stack/NiagaraStackViewModel.cpp
+++ b/Engine/Plugins/FX/Niagara/Source/NiagaraEditor/Private/ViewModels/Stack/NiagaraStackViewModel.cpp
@@ -74,6 +74,43 @@ bool UNiagaraStackViewModel::FTopLevelViewModel::operator==(const FTopLevelViewM
 	return Other.SystemViewModel == SystemViewModel && Other.EmitterHandleViewModel == EmitterHandleViewModel && Other.RootEntry == RootEntry;
 }
 
+void UNiagaraStackViewModel::UpdateStackWithValidationResults()
+{
+	bValidatorUpdatePending = false;
+
+	// clear out old stack entries
+	TArray<UNiagaraStackEntry*> EntriesToClear;
+	RootEntry->GetUnfilteredChildren(EntriesToClear);
+	while (EntriesToClear.Num() > 0)
+	{
+		UNiagaraStackEntry* Entry = EntriesToClear.Pop();
+		Entry->ClearExternalIssues();
+		Entry->GetUnfilteredChildren(EntriesToClear);
+	}
+
+	UNiagaraSystem& NiagaraSystem = SystemViewModel.Pin()->GetSystem();
+	if (NiagaraSystem.GetEffectType())
+	{
+		// go over the validation rules
+		for (TSubclassOf<UNiagaraValidationRule> ValidationClass : NiagaraSystem.GetEffectType()->ValidationRules)
+		{
+			if (UNiagaraValidationRule* ValidationRule = Cast<UNiagaraValidationRule>(ValidationClass->GetDefaultObject()))
+			{
+				TArray<FNiagaraValidationResult> NiagaraValidationResults = ValidationRule->CheckValidity(SystemViewModel.Pin());
+				for (const FNiagaraValidationResult& Result : NiagaraValidationResults)
+				{
+					EStackIssueSeverity Severity = (Result.Severity == ENiagaraValidationSeverity::Error) ? EStackIssueSeverity::Error : (Result.Severity == ENiagaraValidationSeverity::Warning ? EStackIssueSeverity::Warning : EStackIssueSeverity::Info);
+					if (UNiagaraStackEntry* SourceEntry = Cast<UNiagaraStackEntry>(Result.SourceObject.Get()))
+					{
+						SourceEntry->AddExternalIssue(Severity, Result.SummaryText, Result.Description, Result.Severity == ENiagaraValidationSeverity::Info);
+					}
+				}
+			}
+			
+		}
+	}
+}
+
 void UNiagaraStackViewModel::InitializeWithViewModels(TSharedPtr<FNiagaraSystemViewModel> InSystemViewModel, TSharedPtr<FNiagaraEmitterHandleViewModel> InEmitterHandleViewModel, FNiagaraStackViewModelOptions InOptions)
 {
 	Reset();
@@ -114,6 +151,8 @@ void UNiagaraStackViewModel::InitializeWithViewModels(TSharedPtr<FNiagaraSystemV
 		RootEntries.Add(RootEntry);
 
 		bExternalRootEntry = false;
+
+		UpdateStackWithValidationResults();
 	}
 
 	StructureChangedDelegate.Broadcast(ENiagaraStructureChangedFlags::StructureChanged);
@@ -178,6 +217,7 @@ void UNiagaraStackViewModel::Reset()
 	CurrentFocusedSearchMatchIndex = -1;
 	bRestartSearch = false;
 	bRefreshPending = false;
+	bValidatorUpdatePending = false;
 	bUsesTopLevelViewModels = false;
 }
 
@@ -208,6 +248,11 @@ void UNiagaraStackViewModel::Tick()
 			InvalidateSearchResults();
 		}
 
+		if (bValidatorUpdatePending)
+		{
+			UpdateStackWithValidationResults();
+		}
+
 		SearchTick();
 	}
 }
diff --git a/Engine/Plugins/FX/Niagara/Source/NiagaraEditor/Public/NiagaraValidationRules.h b/Engine/Plugins/FX/Niagara/Source/NiagaraEditor/Public/NiagaraValidationRules.h
new file mode 100644
--- /dev/null
+++ b/Engine/Plugins/FX/Niagara/Source/NiagaraEditor/Public/NiagaraValidationRules.h
@@ -0,0 +1,42 @@
+﻿// Copyright Epic Games, Inc. All Rights Reserved.
+
+#pragma once
+
+#include "CoreMinimal.h"
+#include "NiagaraValidationRule.h"
+#include "NiagaraValidationRules.generated.h"
+
+
+UCLASS(Category = "Validation")
+class UNiagara_NoWarmupTime : public UNiagaraValidationRule
+{
+	GENERATED_BODY()
+public:
+	virtual TArray<FNiagaraValidationResult> CheckValidity(TSharedPtr<FNiagaraSystemViewModel> ViewModel) const override;
+};
+
+
+UCLASS(Category = "Validation")
+class UNiagara_FixedGPUBoundsSet : public UNiagaraValidationRule
+{
+	GENERATED_BODY()
+public:
+	virtual TArray<FNiagaraValidationResult> CheckValidity(TSharedPtr<FNiagaraSystemViewModel> ViewModel) const override;
+};
+
+
+UCLASS(Category = "Validation")
+class UNiagara_NoComponentRendererOnLowSettings : public UNiagaraValidationRule
+{
+	GENERATED_BODY()
+public:
+	virtual TArray<FNiagaraValidationResult> CheckValidity(TSharedPtr<FNiagaraSystemViewModel> ViewModel) const override;
+};
+
+UCLASS(Category = "Validation")
+class UNiagara_InvalidEffectType : public UNiagaraValidationRule
+{
+	GENERATED_BODY()
+public:
+	virtual TArray<FNiagaraValidationResult> CheckValidity(TSharedPtr<FNiagaraSystemViewModel> ViewModel) const override;
+};
diff --git a/Engine/Plugins/FX/Niagara/Source/NiagaraEditor/Public/ViewModels/Stack/NiagaraStackEntry.h b/Engine/Plugins/FX/Niagara/Source/NiagaraEditor/Public/ViewModels/Stack/NiagaraStackEntry.h
--- a/Engine/Plugins/FX/Niagara/Source/NiagaraEditor/Public/ViewModels/Stack/NiagaraStackEntry.h
+++ b/Engine/Plugins/FX/Niagara/Source/NiagaraEditor/Public/ViewModels/Stack/NiagaraStackEntry.h
@@ -4,6 +4,7 @@
 
 #include "DragAndDrop/DecoratedDragDropOp.h"
 #include "NiagaraEditorCommon.h"
+#include "NiagaraValidationRule.h"
 #include "Widgets/Views/STableRow.h"
 #include "NiagaraStackEntry.generated.h"
 
@@ -464,6 +465,9 @@ public:
 
 	const TArray<UNiagaraStackEntry*>& GetAllChildrenWithIssues() const;
 
+	void AddExternalIssue(EStackIssueSeverity Severity, const FText& SummaryText, const FText& Description, bool bCanBeDismissed);
+	void ClearExternalIssues();
+
 	virtual bool SupportsCut() const { return false; }
 
 	virtual bool TestCanCutWithMessage(FText& OutMessage) const { return false; }
@@ -626,6 +630,7 @@ private:
 	FOnRequestDrop OnRequestDropDelegate;
 	
 	TArray<FStackIssue> StackIssues;
+	TArray<FStackIssue> ExternalStackIssues;
 
 	bool bIsFinalized;
 
diff --git a/Engine/Plugins/FX/Niagara/Source/NiagaraEditor/Public/ViewModels/Stack/NiagaraStackViewModel.h b/Engine/Plugins/FX/Niagara/Source/NiagaraEditor/Public/ViewModels/Stack/NiagaraStackViewModel.h
--- a/Engine/Plugins/FX/Niagara/Source/NiagaraEditor/Public/ViewModels/Stack/NiagaraStackViewModel.h
+++ b/Engine/Plugins/FX/Niagara/Source/NiagaraEditor/Public/ViewModels/Stack/NiagaraStackViewModel.h
@@ -145,6 +145,9 @@ public:
 	bool HasIssues() const;
 
 	void Refresh() { bRefreshPending = true; }
+
+	void RequestValidationUpdate() { bValidatorUpdatePending = true; }
+
 private:
 	/** Recursively Expands all groups and collapses all items in the stack. */
 	void CollapseToHeadersRecursive(TArray<UNiagaraStackEntry*> Entries);
@@ -177,6 +180,7 @@ private:
 	void GeneratePathForEntry(UNiagaraStackEntry* Root, UNiagaraStackEntry* Entry, TArray<UNiagaraStackEntry*> CurrentPath, TArray<UNiagaraStackEntry*>& EntryPath) const;
 
 	void InvalidateSearchResults();
+	void UpdateStackWithValidationResults();
 
 private:
 	TWeakPtr<FNiagaraEmitterHandleViewModel> EmitterHandleViewModel;
@@ -206,6 +210,7 @@ private:
 	static const double MaxSearchTime;
 	bool bRestartSearch;
 	bool bRefreshPending;
+	bool bValidatorUpdatePending;
 	bool bHasIssues;
 	int32 CurrentIssueCycleIndex;
 	TWeakPtr<FTopLevelViewModel> CyclingIssuesForTopLevel;
