Added content validation rule to check obvious LWC problems

diff --git a/Engine/Plugins/FX/Niagara/Source/NiagaraEditor/Private/NiagaraValidationRules.cpp b/Engine/Plugins/FX/Niagara/Source/NiagaraEditor/Private/NiagaraValidationRules.cpp
--- a/Engine/Plugins/FX/Niagara/Source/NiagaraEditor/Private/NiagaraValidationRules.cpp
+++ b/Engine/Plugins/FX/Niagara/Source/NiagaraEditor/Private/NiagaraValidationRules.cpp
@@ -2,12 +2,17 @@
 
 #include "NiagaraValidationRules.h"
 
+#include "NiagaraClipboard.h"
 #include "NiagaraComponentRendererProperties.h"
+#include "NiagaraNodeFunctionCall.h"
+#include "NiagaraScriptSource.h"
 #include "NiagaraSystem.h"
 #include "NiagaraSystemEditorData.h"
 #include "ViewModels/NiagaraEmitterHandleViewModel.h"
 #include "ViewModels/NiagaraSystemViewModel.h"
 #include "ViewModels/Stack/NiagaraStackEmitterSettingsGroup.h"
+#include "ViewModels/Stack/NiagaraStackFunctionInput.h"
+#include "ViewModels/Stack/NiagaraStackModuleItem.h"
 #include "ViewModels/Stack/NiagaraStackRendererItem.h"
 #include "ViewModels/Stack/NiagaraStackSystemPropertiesItem.h"
 #include "ViewModels/Stack/NiagaraStackViewModel.h"
@@ -19,12 +24,14 @@ TArray<T*> GetStackEntries(UNiagaraStackViewModel* StackViewModel, bool bRefresh
 {
 	TArray<T*> Results;
 	TArray<UNiagaraStackEntry*> EntriesToCheck;
-	UNiagaraStackEntry* RootEntry = StackViewModel->GetRootEntry();
-	if (bRefresh)
+	if (UNiagaraStackEntry* RootEntry = StackViewModel->GetRootEntry())
 	{
-		RootEntry->RefreshChildren();
+		if (bRefresh)
+		{
+			RootEntry->RefreshChildren();
+		}
+		RootEntry->GetUnfilteredChildren(EntriesToCheck);
 	}
-	RootEntry->GetUnfilteredChildren(EntriesToCheck);
 	while (EntriesToCheck.Num() > 0)
 	{
 		UNiagaraStackEntry* Entry = EntriesToCheck.Pop();
@@ -37,6 +44,19 @@ TArray<T*> GetStackEntries(UNiagaraStackViewModel* StackViewModel, bool bRefresh
 	return Results;
 }
 
+template<typename T>
+TArray<T*> GetAllStackEntriesInSystem(TSharedPtr<FNiagaraSystemViewModel> ViewModel, bool bRefresh = false)
+{
+	TArray<T*> Results;
+	Results.Append(GetStackEntries<T>(ViewModel->GetSystemStackViewModel(), bRefresh));
+	TArray<TSharedRef<FNiagaraEmitterHandleViewModel>> EmitterHandleViewModels = ViewModel->GetEmitterHandleViewModels();
+	for (TSharedRef<FNiagaraEmitterHandleViewModel> EmitterHandleModel : EmitterHandleViewModels)
+	{
+		Results.Append(GetStackEntries<T>(EmitterHandleModel.Get().GetEmitterStackViewModel(), bRefresh));
+	}
+	return Results;
+}
+
 // helper function to retrieve a single stack entry from the system or emitter view model
 template<typename T>
 T* GetStackEntry(UNiagaraStackViewModel* StackViewModel, bool bRefresh = false)
@@ -140,4 +160,80 @@ TArray<FNiagaraValidationResult> UNiagara_InvalidEffectType::CheckValidity(TShar
 	return Results;
 }
 
+TArray<FNiagaraValidationResult> UNiagara_LWC_Validator::CheckValidity(TSharedPtr<FNiagaraSystemViewModel> ViewModel) const
+{
+	TArray<FNiagaraValidationResult> Results;
+	UNiagaraSystem& System = ViewModel->GetSystem();
+	if (!System.SupportsLargeWorldCoordinates())
+	{
+		return Results;
+	}
+
+	// gather all the modules in the system, excluding localspace emitters
+	TArray<UNiagaraStackModuleItem*> AllModules;
+	AllModules.Append(GetStackEntries<UNiagaraStackModuleItem>(ViewModel->GetSystemStackViewModel()));
+	TArray<TSharedRef<FNiagaraEmitterHandleViewModel>> EmitterHandleViewModels = ViewModel->GetEmitterHandleViewModels();
+	for (TSharedRef<FNiagaraEmitterHandleViewModel> EmitterHandleModel : EmitterHandleViewModels)
+	{
+		if (EmitterHandleModel->GetEmitterHandle()->GetInstance()->bLocalSpace == false)
+		{
+			AllModules.Append(GetStackEntries<UNiagaraStackModuleItem>(EmitterHandleModel.Get().GetEmitterStackViewModel()));
+		}
+	}
+
+	for (UNiagaraStackModuleItem* Module : AllModules)
+	{
+		TArray<UNiagaraStackFunctionInput*> StackInputs;
+		Module->GetParameterInputs(StackInputs);
+		
+		for (UNiagaraStackFunctionInput* Input : StackInputs)
+		{
+			if (Input->GetInputType() == FNiagaraTypeDefinition::GetPositionDef())
+			{
+				// check if any position inputs are set locally to absolute values
+				if (Input->GetValueMode() == UNiagaraStackFunctionInput::EValueMode::Local)
+				{
+					FNiagaraValidationResult Result(ENiagaraValidationSeverity::Warning, FText::Format(LOCTEXT("LocalPosInputSummary", "Input '{0}' set to absolute value"), Input->GetDisplayName()), LOCTEXT("LocalPosInputDescription", "Position attributes should never be set to an absolute values, because they will be offset when using large world coordinates.\nInstead, set them relative to a known position like Engine.Owner.Position."), Input);
+					Results.Add(Result);
+				}
+
+				// check if the linked dynamic input script outputs a vector
+				if (Input->GetValueMode() == UNiagaraStackFunctionInput::EValueMode::Dynamic && Input->GetDynamicInputNode())
+				{
+					if (UNiagaraScriptSource* DynamicInputSource = Cast<UNiagaraScriptSource>(Input->GetDynamicInputNode()->GetFunctionScriptSource()))
+					{
+						TArray<FNiagaraVariable> OutNodes;
+						DynamicInputSource->NodeGraph->GetOutputNodeVariables(OutNodes);
+						for (const FNiagaraVariable& OutVariable : OutNodes)
+						{
+							if (OutVariable.GetType() == FNiagaraTypeDefinition::GetVec3Def())
+							{
+								FTextFormat DescriptionFormat = LOCTEXT("VecDILinkedToPosInputDescription", "The position input {0} is linked to a dynamic input that outputs a vector.\nPlease use a dynamic input that outputs a position instead or explicitly convert the vector to a position type.");
+								FNiagaraValidationResult Result(ENiagaraValidationSeverity::Warning, LOCTEXT("VecDILinkedToPosInputSummary", "Position input is linked to a vector output"), FText::Format(DescriptionFormat, Input->GetDisplayName()), Input);
+								Results.Add(Result);
+							}
+						}
+					}
+				}
+
+				// check if the linked input variable is a vector
+				if (Input->GetValueMode() == UNiagaraStackFunctionInput::EValueMode::Linked)
+				{
+					FNiagaraVariable VectorVar(FNiagaraTypeDefinition::GetVec3Def(), Input->GetLinkedValueHandle().GetParameterHandleString());
+					const UNiagaraGraph* NiagaraGraph = Input->GetInputFunctionCallNode().GetNiagaraGraph();
+
+					// we check if metadata for a vector attribute with the linked name exists in the emitter/system script graph. Not 100% correct, but it needs to be fast and a few false negatives are acceptable.
+					if (NiagaraGraph && NiagaraGraph->GetMetaData(VectorVar).IsSet())
+					{
+						FNiagaraValidationResult Result(ENiagaraValidationSeverity::Warning, FText::Format(LOCTEXT("PositionLinkedVectorSummary", "Input '{0}' is linked to a vector attribute"), Input->GetDisplayName()), LOCTEXT("PositionLinkedVectorDescription", "Position types should only be set linked to position attributes. In this case, it is linked to a vector attribute and the implicit conversion can cause problems with large world coordinates."), Input);
+						Results.Add(Result);
+					}
+				}
+			}
+		}
+	}
+	
+	return Results;
+}
+
 #undef LOCTEXT_NAMESPACE
diff --git a/Engine/Plugins/FX/Niagara/Source/NiagaraEditor/Private/ViewModels/NiagaraSystemViewModel.cpp b/Engine/Plugins/FX/Niagara/Source/NiagaraEditor/Private/ViewModels/NiagaraSystemViewModel.cpp
--- a/Engine/Plugins/FX/Niagara/Source/NiagaraEditor/Private/ViewModels/NiagaraSystemViewModel.cpp
+++ b/Engine/Plugins/FX/Niagara/Source/NiagaraEditor/Private/ViewModels/NiagaraSystemViewModel.cpp
@@ -322,6 +322,15 @@ void FNiagaraSystemViewModel::CompileSystem(bool bForce)
 	SystemScriptViewModel->CompileSystem(bForce);
 	bCompilePendingCompletion = true;
 	InvalidateCachedCompileStatus();
+
+	if (SystemStackViewModel)
+	{
+		SystemStackViewModel->RequestValidationUpdate();
+		for (TSharedRef<FNiagaraEmitterHandleViewModel> EmitterHandleViewModel :GetEmitterHandleViewModels())
+		{
+			EmitterHandleViewModel->GetEmitterStackViewModel()->RequestValidationUpdate();
+		}
+	}
 }
 
 ENiagaraScriptCompileStatus FNiagaraSystemViewModel::GetLatestCompileStatus() const
diff --git a/Engine/Plugins/FX/Niagara/Source/NiagaraEditor/Public/NiagaraValidationRules.h b/Engine/Plugins/FX/Niagara/Source/NiagaraEditor/Public/NiagaraValidationRules.h
--- a/Engine/Plugins/FX/Niagara/Source/NiagaraEditor/Public/NiagaraValidationRules.h
+++ b/Engine/Plugins/FX/Niagara/Source/NiagaraEditor/Public/NiagaraValidationRules.h
@@ -40,3 +40,11 @@ class UNiagara_InvalidEffectType : public UNiagaraValidationRule
 public:
 	virtual TArray<FNiagaraValidationResult> CheckValidity(TSharedPtr<FNiagaraSystemViewModel> ViewModel) const override;
 };
+
+UCLASS(Category = "Validation")
+class UNiagara_LWC_Validator : public UNiagaraValidationRule
+{
+	GENERATED_BODY()
+public:
+	virtual TArray<FNiagaraValidationResult> CheckValidity(TSharedPtr<FNiagaraSystemViewModel> ViewModel) const override;
+};
\ No newline at end of file
