- Add tooltip for particle state iteration

#rnx

diff --git a/Engine/Plugins/FX/Niagara/Source/Niagara/Public/NiagaraSimulationStageBase.h b/Engine/Plugins/FX/Niagara/Source/Niagara/Public/NiagaraSimulationStageBase.h
--- a/Engine/Plugins/FX/Niagara/Source/Niagara/Public/NiagaraSimulationStageBase.h
+++ b/Engine/Plugins/FX/Niagara/Source/Niagara/Public/NiagaraSimulationStageBase.h
@@ -80,6 +80,11 @@ public:
 	UPROPERTY(EditAnywhere, Category = "Simulation Stage", meta = (editcondition = "IterationSource == ENiagaraIterationSource::DataInterface"))
 	FNiagaraVariableDataInterfaceBinding DataInterface;
 
+	/**
+	When enabled we iterate over the subset of particles that satisfy the binding & range selected.
+	For example, if the binding was set to Particles.MyStateIndex and the range was 1-3 only the particles
+	with values of 1, 2 or 3 for Particles.MyStateIndex would run the script, all others would be ignored.
+	*/
 	UPROPERTY(EditAnywhere, Category = "Simulation Stage", meta = (editcondition = "IterationSource == ENiagaraIterationSource::Particles"))
 	uint32 bParticleIterationStateEnabled : 1;
 
