- Fix for INF with distance culling (FLT_MAX * FLT_MAX == inf)

#rnx

diff --git a/Engine/Plugins/FX/Niagara/Source/Niagara/Private/NiagaraRendererMeshes.cpp b/Engine/Plugins/FX/Niagara/Source/Niagara/Private/NiagaraRendererMeshes.cpp
--- a/Engine/Plugins/FX/Niagara/Source/Niagara/Private/NiagaraRendererMeshes.cpp
+++ b/Engine/Plugins/FX/Niagara/Source/Niagara/Private/NiagaraRendererMeshes.cpp
@@ -87,6 +87,7 @@ FNiagaraRendererMeshes::FNiagaraRendererMeshes(ERHIFeatureLevel::Type FeatureLev
 	bEnableFrustumCulling = Properties->bEnableFrustumCulling;
 	bEnableCulling = bEnableFrustumCulling;
 	DistanceCullRange = FVector2f(0, FLT_MAX);
+	DistanceCullRangeSquared = FVector2f(0, FLT_MAX);
 	RendererVisibility = Properties->RendererVisibility;
 	bAccurateMotionVectors = Properties->NeedsPreciseMotionVectors();
 	MaxSectionCount = 0;
@@ -94,6 +95,7 @@ FNiagaraRendererMeshes::FNiagaraRendererMeshes(ERHIFeatureLevel::Type FeatureLev
 	if (Properties->bEnableCameraDistanceCulling)
 	{
 		DistanceCullRange = FVector2f(Properties->MinCameraDistance, Properties->MaxCameraDistance);
+		DistanceCullRangeSquared = DistanceCullRange* DistanceCullRange;
 		bEnableCulling = true;
 	}
 
@@ -944,7 +946,7 @@ void FNiagaraRendererMeshes::SetupElementForGPUScene(
 		GPUSceneRes.GPUWriteParams.RendererVisibility 		= RendererVisibility;
 		GPUSceneRes.GPUWriteParams.VisibilityTagDataOffset 	= ParticleMeshRenderData.RendererVisTagOffset;
 		GPUSceneRes.GPUWriteParams.LocalBoundingCenter		= (FVector3f)MeshData.LocalBounds.GetCenter();
-		GPUSceneRes.GPUWriteParams.DistanceCullRangeSquared = DistanceCullRange * DistanceCullRange;
+		GPUSceneRes.GPUWriteParams.DistanceCullRangeSquared = DistanceCullRangeSquared;
 		GPUSceneRes.GPUWriteParams.bNeedsPrevTransform		= bNeedsPrevTransform ? 1 : 0;
 
 		// We need to set this flag to force the system to always cull individual instances, because we may need to discard instances that are:
diff --git a/Engine/Plugins/FX/Niagara/Source/Niagara/Public/NiagaraRendererMeshes.h b/Engine/Plugins/FX/Niagara/Source/Niagara/Public/NiagaraRendererMeshes.h
--- a/Engine/Plugins/FX/Niagara/Source/Niagara/Public/NiagaraRendererMeshes.h
+++ b/Engine/Plugins/FX/Niagara/Source/Niagara/Public/NiagaraRendererMeshes.h
@@ -191,6 +191,7 @@ private:
 	ENiagaraMeshLockedAxisSpace LockedAxisSpace;
 
 	FVector2f DistanceCullRange;
+	FVector2f DistanceCullRangeSquared;
 	int32 RendererVisTagOffset;
 	int32 RendererVisibility;
 	int32 MeshIndexOffset;
