- Add option to only show filtered results on overview
- Add LWC tile to outliner information

#rnx

diff --git a/Engine/Plugins/FX/Niagara/Source/Niagara/Private/NiagaraDebugHud.cpp b/Engine/Plugins/FX/Niagara/Source/Niagara/Private/NiagaraDebugHud.cpp
--- a/Engine/Plugins/FX/Niagara/Source/Niagara/Private/NiagaraDebugHud.cpp
+++ b/Engine/Plugins/FX/Niagara/Source/Niagara/Private/NiagaraDebugHud.cpp
@@ -1572,13 +1572,16 @@ void FNiagaraDebugHud::DrawOverview(class FNiagaraWorldManager* WorldManager, FC
 	{
 		TextLocation.Y += fAdvanceHeight;		
 		uint32 NumLines = 1;
-		for (auto& Pair : PerSystemDebugInfo)
+		for (const auto& Pair : PerSystemDebugInfo)
 		{
-			const FSystemDebugInfo& SysInfo = Pair.Value;
-			if (SysInfo.FramesSinceVisible < Settings.PerfHistoryFrames)
+			const FSystemDebugInfo& SystemInfo = Pair.Value;
+			if ((SystemInfo.FramesSinceVisible >= Settings.PerfHistoryFrames) ||
+				(Settings.bOverviewShowFilteredSystemOnly && !SystemInfo.bPassesSystemFilter))
 			{
-				++NumLines;
+				continue;
 			}
+
+			++NumLines;
 		}
 		DrawCanvas->DrawTile(TextLocation.X - 1.0f, TextLocation.Y - 1.0f, TotalOverviewWidth + 1.0f, 2.0f + (float(NumLines) * fAdvanceHeight), 0.0f, 0.0f, 0.0f, 0.0f, BackgroundColor);
 
@@ -1595,10 +1598,12 @@ void FNiagaraDebugHud::DrawOverview(class FNiagaraWorldManager* WorldManager, FC
 		for (auto it = PerSystemDebugInfo.CreateConstIterator(); it; ++it)
 		{
 			const auto& SystemInfo = it->Value;
-			if (SystemInfo.FramesSinceVisible >= Settings.PerfHistoryFrames)
+			if ((SystemInfo.FramesSinceVisible >= Settings.PerfHistoryFrames) ||
+				(Settings.bOverviewShowFilteredSystemOnly && !SystemInfo.bPassesSystemFilter))
 			{
 				continue;
 			}
+
 			for (FOverviewColumn& Column : OverviewColumns)
 			{
 				Column.DrawSystemData(DrawCanvas, Font, TextLocation.X, TextLocation.Y, SystemInfo);
diff --git a/Engine/Plugins/FX/Niagara/Source/Niagara/Private/NiagaraDebuggerClient.cpp b/Engine/Plugins/FX/Niagara/Source/Niagara/Private/NiagaraDebuggerClient.cpp
--- a/Engine/Plugins/FX/Niagara/Source/Niagara/Private/NiagaraDebuggerClient.cpp
+++ b/Engine/Plugins/FX/Niagara/Source/Niagara/Private/NiagaraDebuggerClient.cpp
@@ -365,7 +365,7 @@ bool FNiagaraDebuggerClient::UpdateOutliner(float DeltaSeconds)
 					if (FNiagaraSystemInstance* Inst = InstController->GetSystemInstance_Unsafe())
 					{
 						InstData.TickGroup = Inst->CalculateTickGroup();
-
+						InstData.LWCTile = Inst->GetLWCTile();
 						InstData.bIsSolo = Inst->IsSolo();
 						InstData.bRequiresDistanceFieldData = Inst->RequiresDistanceFieldData();
 						InstData.bRequiresDepthBuffer = Inst->RequiresDepthBuffer();
diff --git a/Engine/Plugins/FX/Niagara/Source/Niagara/Public/NiagaraDebuggerCommon.h b/Engine/Plugins/FX/Niagara/Source/Niagara/Public/NiagaraDebuggerCommon.h
--- a/Engine/Plugins/FX/Niagara/Source/Niagara/Public/NiagaraDebuggerCommon.h
+++ b/Engine/Plugins/FX/Niagara/Source/Niagara/Public/NiagaraDebuggerCommon.h
@@ -70,6 +70,9 @@ struct FNiagaraOutlinerSystemInstanceData
 	UPROPERTY(VisibleAnywhere, Category = "System")
 	FString ComponentName;
 
+	UPROPERTY(VisibleAnywhere, Category = "System")
+	FVector3f LWCTile = FVector3f::Zero();
+
 	UPROPERTY(VisibleAnywhere, Category = "System")
 	TArray<FNiagaraOutlinerEmitterInstanceData> Emitters;
 	
@@ -461,6 +464,10 @@ struct NIAGARA_API FNiagaraDebugHUDSettingsData
 	UPROPERTY(EditAnywhere, Category = "Debug Overview", meta = (DisplayName = "Debug Overview Text Location", EditCondition = "bOverviewEnabled"))
 	FVector2D OverviewLocation = FIntPoint(30.0f, 150.0f);
 
+	/** When enabled the overview will only show the filter system information. */
+	UPROPERTY(EditAnywhere, Category = "Debug Overview", meta = (EditCondition = "bOverviewEnabled"))
+	bool bOverviewShowFilteredSystemOnly = false;
+
 	/**
 	Wildcard filter which is compared against the Components Actor name to narrow down the detailed information.
 	For example, "*Water*" would match all actors that contain the string "water".
