- Emitter Properties data interface to allow access to data that is not part of the data buffers
- Add method to particle reader to understand local or world space

diff --git a/Engine/Plugins/FX/Niagara/Shaders/Private/NiagaraDataInterfaceEmitterPropertiesTemplate.ush b/Engine/Plugins/FX/Niagara/Shaders/Private/NiagaraDataInterfaceEmitterPropertiesTemplate.ush
new file mode 100644
--- /dev/null
+++ b/Engine/Plugins/FX/Niagara/Shaders/Private/NiagaraDataInterfaceEmitterPropertiesTemplate.ush
@@ -0,0 +1,26 @@
+// Copyright Epic Games, Inc. All Rights Reserved.
+
+bool			LocalSpace_{ParameterName};
+bool			FixedBoundsValid_{ParameterName};
+float3			FixedBoundsMin_{ParameterName};
+float3			FixedBoundsMax_{ParameterName};
+
+void GetLocalSpace_{ParameterName}(out bool bLocalSpace)
+{
+	bLocalSpace = LocalSpace_{ParameterName};
+}
+
+void GetBounds_{ParameterName}(out bool bValid, out float3 Min, out float3 Max)
+{
+	// GPU is always fixed no dynamic bounds currently
+	bValid = FixedBoundsValid_{ParameterName};
+	Min = FixedBoundsMin_{ParameterName};
+	Max = FixedBoundsMax_{ParameterName};
+}
+
+void GetFixedBounds_{ParameterName}(out bool bValid, out float3 Min, out float3 Max)
+{
+	bValid = FixedBoundsValid_{ParameterName};
+	Min = FixedBoundsMin_{ParameterName};
+	Max = FixedBoundsMax_{ParameterName};
+}
diff --git a/Engine/Plugins/FX/Niagara/Source/Niagara/Classes/NiagaraDataInterfaceParticleRead.h b/Engine/Plugins/FX/Niagara/Source/Niagara/Classes/NiagaraDataInterfaceParticleRead.h
--- a/Engine/Plugins/FX/Niagara/Source/Niagara/Classes/NiagaraDataInterfaceParticleRead.h
+++ b/Engine/Plugins/FX/Niagara/Source/Niagara/Classes/NiagaraDataInterfaceParticleRead.h
@@ -49,6 +49,7 @@ public:
 	virtual bool HasInternalAttributeReads(const UNiagaraEmitter* OwnerEmitter, const UNiagaraEmitter* Provider) const override;
 	//UNiagaraDataInterface Interface End
 
+	void VMGetLocalSpace(FVectorVMExternalFunctionContext& Context);
 	void GetNumSpawnedParticles(FVectorVMExternalFunctionContext& Context);
 	void GetSpawnedIDAtIndex(FVectorVMExternalFunctionContext& Context);
 	void GetNumParticles(FVectorVMExternalFunctionContext& Context);
diff --git a/Engine/Plugins/FX/Niagara/Source/Niagara/Classes/NiagaraEmitterInstance.h b/Engine/Plugins/FX/Niagara/Source/Niagara/Classes/NiagaraEmitterInstance.h
--- a/Engine/Plugins/FX/Niagara/Source/Niagara/Classes/NiagaraEmitterInstance.h
+++ b/Engine/Plugins/FX/Niagara/Source/Niagara/Classes/NiagaraEmitterInstance.h
@@ -124,7 +124,7 @@ public:
 	ENiagaraExecutionState NIAGARA_API GetExecutionState() { return ExecutionState; }
 	void NIAGARA_API SetExecutionState(ENiagaraExecutionState InState);
 
-	FBox GetBounds();
+	FBox GetBounds() const;
 
 	FNiagaraScriptExecutionContext& GetSpawnExecutionContext() { return SpawnExecContext; }
 	FNiagaraScriptExecutionContext& GetUpdateExecutionContext() { return UpdateExecContext; }
@@ -147,7 +147,11 @@ public:
 	}
 
 	void SetSystemFixedBoundsOverride(FBox SystemFixedBounds);
-	FORCEINLINE void SetFixedBounds(const FBox& InLocalBounds) { FixedBounds_GT = InLocalBounds; }
+	FORCEINLINE void SetFixedBounds(const FBox& InLocalBounds)
+	{
+		FRWScopeLock ScopeLock(FixedBoundsGuard, SLT_Write);
+		FixedBounds = InLocalBounds;
+	}
 	FBox GetFixedBounds() const;
 
 	UObject* FindBinding(const FNiagaraVariable& InVariable) const;
@@ -200,9 +204,9 @@ private:
 	/* Emitter bounds */
 	FBox CachedBounds;
 
-	/** Optional user specified bounds. */
-	FBox FixedBounds_GT;
-	FBox FixedBounds_CNC;
+	/** Optional user or VM specified bounds. */
+	mutable FRWLock FixedBoundsGuard;
+	FBox FixedBounds;
 
 	/** Cached fixed bounds of the parent system which override this Emitter Instances bounds if set. Whenever we initialize the owning SystemInstance we will reconstruct this
 	 ** EmitterInstance and the cached bounds will be unset. */
diff --git a/Engine/Plugins/FX/Niagara/Source/Niagara/Private/DataInterface/NiagaraDataInterfaceEmitterProperties.cpp b/Engine/Plugins/FX/Niagara/Source/Niagara/Private/DataInterface/NiagaraDataInterfaceEmitterProperties.cpp
new file mode 100644
--- /dev/null
+++ b/Engine/Plugins/FX/Niagara/Source/Niagara/Private/DataInterface/NiagaraDataInterfaceEmitterProperties.cpp
@@ -0,0 +1,450 @@
+// Copyright Epic Games, Inc. All Rights Reserved.
+
+#include "NiagaraDataInterfaceEmitterProperties.h"
+#include "NiagaraEmitterInstance.h"
+#include "NiagaraComputeExecutionContext.h"
+#include "NiagaraSystemInstance.h"
+
+#include "Internationalization/Internationalization.h"
+#include "ShaderCompilerCore.h"
+#include "ShaderParameterUtils.h"
+
+#define LOCTEXT_NAMESPACE "NiagaraDataInterfaceEmitterProperties"
+
+//////////////////////////////////////////////////////////////////////////
+
+namespace NDIEmitterPropertiesLocal
+{
+	static const TCHAR* TemplateShaderFile = TEXT("/Plugin/FX/Niagara/Private/NiagaraDataInterfaceEmitterPropertiesTemplate.ush");
+
+	static const FName GetLocalSpaceName("GetLocalSpace");
+	static const FName GetBoundsName("GetBounds");
+	static const FName GetFixedBoundsName("GetFixedBounds");
+	static const FName SetFixedBoundsName("SetFixedBounds");
+
+	struct FInstanceData_GameToRender
+	{
+		FBox	FixedBounds;
+	};
+
+	struct FInstanceData_GameThread
+	{
+		FNiagaraEmitterInstance*	EmitterInstance = nullptr;
+		bool						bLocalSpace = false;
+	};
+
+	struct FInstanceData_RenderThread
+	{
+		bool	bLocalSpace = false;
+		FBox	FixedBounds = FBox(EForceInit::ForceInit);
+	};
+
+	struct FNDIProxy : public FNiagaraDataInterfaceProxy
+	{
+		void ProvidePerInstanceDataForRenderThread(void* DataForRenderThread, void* PerInstanceData, const FNiagaraSystemInstanceID& InstanceID)
+		{
+			FInstanceData_GameToRender* InstanceData_ForRT = new(DataForRenderThread) FInstanceData_GameToRender();
+			FInstanceData_GameThread* InstanceData_GT = reinterpret_cast<FInstanceData_GameThread*>(PerInstanceData);
+
+			InstanceData_ForRT->FixedBounds = InstanceData_GT->EmitterInstance ? InstanceData_GT->EmitterInstance->GetFixedBounds() : FBox(EForceInit::ForceInit);
+		}
+
+		virtual void ConsumePerInstanceDataFromGameThread(void* PerInstanceData, const FNiagaraSystemInstanceID& InstanceID) override
+		{
+			FInstanceData_GameToRender* InstanceData_FromGT = static_cast<FInstanceData_GameToRender*>(PerInstanceData);
+			if ( FInstanceData_RenderThread* InstanceData_RT = PerInstanceData_RenderThread.Find(InstanceID) )
+			{
+				InstanceData_RT->FixedBounds = InstanceData_FromGT->FixedBounds;
+			}
+			InstanceData_FromGT->~FInstanceData_GameToRender();
+		}
+
+		virtual int32 PerInstanceDataPassedToRenderThreadSize() const override
+		{
+			return sizeof(FInstanceData_GameToRender);
+		}
+	
+		virtual void PreStage(FRHICommandList& RHICmdList, const FNiagaraDataInterfaceStageArgs& Context) override {}
+		virtual void PostSimulate(FRHICommandList& RHICmdList, const FNiagaraDataInterfaceArgs& Context) override {}
+
+		TMap<FNiagaraSystemInstanceID, FInstanceData_RenderThread>	PerInstanceData_RenderThread;
+	};
+}
+
+//////////////////////////////////////////////////////////////////////////
+
+struct FNDIEmitterPropertiesCS : public FNiagaraDataInterfaceParametersCS
+{
+	DECLARE_TYPE_LAYOUT(FNDIEmitterPropertiesCS, NonVirtual);
+
+public:
+	void Bind(const FNiagaraDataInterfaceGPUParamInfo& ParameterInfo, const class FShaderParameterMap& ParameterMap)
+	{
+		using namespace NDIEmitterPropertiesLocal;
+
+		LocalSpaceParam.Bind(ParameterMap, *(TEXT("LocalSpace_") + ParameterInfo.DataInterfaceHLSLSymbol));
+		FixedBoundsValidParam.Bind(ParameterMap, *(TEXT("FixedBoundsValid_") + ParameterInfo.DataInterfaceHLSLSymbol));
+		FixedBoundsMinParam.Bind(ParameterMap, *(TEXT("FixedBoundsMin_") + ParameterInfo.DataInterfaceHLSLSymbol));
+		FixedBoundsMaxParam.Bind(ParameterMap, *(TEXT("FixedBoundsMax_") + ParameterInfo.DataInterfaceHLSLSymbol));
+	}
+
+	void Set(FRHICommandList& RHICmdList, const FNiagaraDataInterfaceSetArgs& Context) const
+	{
+		using namespace NDIEmitterPropertiesLocal;
+
+		FRHIComputeShader* ComputeShaderRHI = Context.Shader.GetComputeShader();
+		const FNDIProxy* DIProxy = static_cast<FNDIProxy*>(Context.DataInterface);
+		const FInstanceData_RenderThread* InstanceData_RT = &DIProxy->PerInstanceData_RenderThread.FindChecked(Context.SystemInstanceID);
+
+		SetShaderValue(RHICmdList, ComputeShaderRHI, LocalSpaceParam, InstanceData_RT->bLocalSpace ? 1 : 0);
+		SetShaderValue(RHICmdList, ComputeShaderRHI, FixedBoundsValidParam, InstanceData_RT->FixedBounds.IsValid ? 1 : 0);
+		SetShaderValue(RHICmdList, ComputeShaderRHI, FixedBoundsMinParam, FVector3f(InstanceData_RT->FixedBounds.Min));
+		SetShaderValue(RHICmdList, ComputeShaderRHI, FixedBoundsMaxParam, FVector3f(InstanceData_RT->FixedBounds.Max));
+	}
+
+private:
+	LAYOUT_FIELD(FShaderParameter, LocalSpaceParam);
+	LAYOUT_FIELD(FShaderParameter, FixedBoundsValidParam);
+	LAYOUT_FIELD(FShaderParameter, FixedBoundsMinParam);
+	LAYOUT_FIELD(FShaderParameter, FixedBoundsMaxParam);
+};
+
+IMPLEMENT_TYPE_LAYOUT(FNDIEmitterPropertiesCS);
+IMPLEMENT_NIAGARA_DI_PARAMETER(UNiagaraDataInterfaceEmitterProperties, FNDIEmitterPropertiesCS);
+
+//////////////////////////////////////////////////////////////////////////
+
+UNiagaraDataInterfaceEmitterProperties::UNiagaraDataInterfaceEmitterProperties(FObjectInitializer const& ObjectInitializer)
+	: Super(ObjectInitializer)
+{
+	using namespace NDIEmitterPropertiesLocal;
+
+	Proxy.Reset(new FNDIProxy());
+}
+
+void UNiagaraDataInterfaceEmitterProperties::PostInitProperties()
+{
+	Super::PostInitProperties();
+
+	//Can we register data interfaces as regular types and fold them into the FNiagaraVariable framework for UI and function calls etc?
+	if (HasAnyFlags(RF_ClassDefaultObject))
+	{
+		ENiagaraTypeRegistryFlags Flags = ENiagaraTypeRegistryFlags::AllowEmitterVariable | ENiagaraTypeRegistryFlags::AllowParticleVariable | ENiagaraTypeRegistryFlags::AllowParameter;
+		FNiagaraTypeRegistry::Register(FNiagaraTypeDefinition(GetClass()), Flags);
+	}
+}
+
+bool UNiagaraDataInterfaceEmitterProperties::Equals(const UNiagaraDataInterface* Other) const
+{
+	if (!Super::Equals(Other))
+	{
+		return false;
+	}
+
+	auto* OtherTyped = CastChecked<const UNiagaraDataInterfaceEmitterProperties>(Other);
+	return
+		OtherTyped->EmitterBinding == EmitterBinding;
+}
+
+bool UNiagaraDataInterfaceEmitterProperties::CopyToInternal(UNiagaraDataInterface* Destination) const
+{
+	if (!Super::CopyToInternal(Destination))
+	{
+		return false;
+	}
+
+	auto* DestinationTyped = CastChecked<UNiagaraDataInterfaceEmitterProperties>(Destination);
+	DestinationTyped->EmitterBinding = EmitterBinding;
+
+	return true;
+}
+
+bool UNiagaraDataInterfaceEmitterProperties::InitPerInstanceData(void* PerInstanceData, FNiagaraSystemInstance* SystemInstance)
+{
+	using namespace NDIEmitterPropertiesLocal;
+
+	FInstanceData_GameThread* InstanceData_GT = new(PerInstanceData) FInstanceData_GameThread();
+	InstanceData_GT->EmitterInstance = EmitterBinding.Resolve(SystemInstance, this);
+	if (InstanceData_GT->EmitterInstance)
+	{
+		InstanceData_GT->bLocalSpace = InstanceData_GT->EmitterInstance->GetCachedEmitter()->bLocalSpace;
+	}
+
+	if ( IsUsedByGPUEmitter() )
+	{
+		// Initialize render side instance data
+		ENQUEUE_RENDER_COMMAND(NDIEmitter_InitRT)
+		(
+			[Proxy_RT=GetProxyAs<FNDIProxy>(), InstanceID=SystemInstance->GetId(), bLocalSpace_RT=InstanceData_GT->bLocalSpace](FRHICommandList& CmdList)
+			{
+				FInstanceData_RenderThread* InstanceData_RT = &Proxy_RT->PerInstanceData_RenderThread.Add(InstanceID);
+				InstanceData_RT->bLocalSpace = bLocalSpace_RT;
+			}
+		);
+	}
+
+	return true;
+}
+
+void UNiagaraDataInterfaceEmitterProperties::DestroyPerInstanceData(void* PerInstanceData, FNiagaraSystemInstance* SystemInstance)
+{
+	using namespace NDIEmitterPropertiesLocal;
+
+	FInstanceData_GameThread* InstanceData_GT = reinterpret_cast<FInstanceData_GameThread*>(PerInstanceData);
+	InstanceData_GT->~FInstanceData_GameThread();
+
+	if ( IsUsedByGPUEmitter() )
+	{
+		ENQUEUE_RENDER_COMMAND(NDIEmitter_InitRT)
+		(
+			[Proxy_RT=GetProxyAs<FNDIProxy>(), InstanceID=SystemInstance->GetId()](FRHICommandList& CmdList)
+			{
+				Proxy_RT->PerInstanceData_RenderThread.Remove(InstanceID);
+			}
+		);
+	}
+}
+
+int32 UNiagaraDataInterfaceEmitterProperties::PerInstanceDataSize() const
+{
+	using namespace NDIEmitterPropertiesLocal;
+	return sizeof(FInstanceData_GameThread);
+}
+
+void UNiagaraDataInterfaceEmitterProperties::ProvidePerInstanceDataForRenderThread(void* DataForRenderThread, void* PerInstanceData, const FNiagaraSystemInstanceID& InstanceID)
+{
+	using namespace NDIEmitterPropertiesLocal;
+	FNDIProxy* DIProxy = GetProxyAs<FNDIProxy>();
+	DIProxy->ProvidePerInstanceDataForRenderThread(DataForRenderThread, PerInstanceData, InstanceID);
+}
+
+bool UNiagaraDataInterfaceEmitterProperties::PerInstanceTick(void* PerInstanceData, FNiagaraSystemInstance* SystemInstance, float DeltaSeconds)
+{
+	using namespace NDIEmitterPropertiesLocal;
+
+	check(PerInstanceData && SystemInstance);
+
+	FInstanceData_GameThread* InstanceData_GT = reinterpret_cast<FInstanceData_GameThread*>(PerInstanceData);
+	//-TODO:
+
+	return false;
+}
+
+void UNiagaraDataInterfaceEmitterProperties::GetFunctions(TArray<FNiagaraFunctionSignature>& OutFunctions)
+{
+	using namespace NDIEmitterPropertiesLocal;
+
+	// Reserve space
+	OutFunctions.Reserve(OutFunctions.Num() + 4);
+
+	// Build default signature
+	FNiagaraFunctionSignature DefaultSignature;
+	DefaultSignature.bMemberFunction = true;
+	DefaultSignature.bRequiresContext = false;
+	DefaultSignature.bSupportsCPU = true;
+	DefaultSignature.bSupportsGPU = true;
+	DefaultSignature.Inputs.Emplace(FNiagaraTypeDefinition(GetClass()), TEXT("Emitter"));
+
+	{
+		FNiagaraFunctionSignature& Sig = OutFunctions.Add_GetRef(DefaultSignature);
+		Sig.Name = GetLocalSpaceName;
+		Sig.Outputs.Emplace(FNiagaraTypeDefinition::GetBoolDef(), TEXT("IsLocalSpace"));
+#if WITH_EDITORONLY_DATA
+		Sig.SetDescription(LOCTEXT("GetEmitterLocalSpace", "Returns if the emitter is using local space or world space."));
+#endif
+	}
+	{
+		FNiagaraFunctionSignature & Sig = OutFunctions.Add_GetRef(DefaultSignature);
+		Sig.Name = GetBoundsName;
+		Sig.Outputs.Emplace(FNiagaraTypeDefinition::GetBoolDef(), TEXT("Valid"));
+		Sig.Outputs.Emplace(FNiagaraTypeDefinition::GetVec3Def(), TEXT("Local Min"));
+		Sig.Outputs.Emplace(FNiagaraTypeDefinition::GetVec3Def(), TEXT("Local Max"));
+#if WITH_EDITORONLY_DATA
+		Sig.SetDescription(LOCTEXT("GetEmitterBounds", "Get local space bounds for the emitter, this is the bounds calculated from the last frame, either dynamic or fixed bounds depending on what is set."));
+#endif
+	}
+	{
+		FNiagaraFunctionSignature& Sig = OutFunctions.Add_GetRef(DefaultSignature);
+		Sig.Name = GetFixedBoundsName;
+		Sig.Outputs.Emplace(FNiagaraTypeDefinition::GetBoolDef(), TEXT("Valid"));
+		Sig.Outputs.Emplace(FNiagaraTypeDefinition::GetVec3Def(), TEXT("Local Min"));
+		Sig.Outputs.Emplace(FNiagaraTypeDefinition::GetVec3Def(), TEXT("Local Max"));
+#if WITH_EDITORONLY_DATA
+		Sig.SetDescription(LOCTEXT("GetEmitterFixedBounds", "Get local space fixed bounds for the emitter.  If Valid is false then the emitter has no fixed bounds and is using dynamically generated bounds."));
+#endif
+	}
+	{
+		FNiagaraFunctionSignature& Sig = OutFunctions.Add_GetRef(DefaultSignature);
+		Sig.Name = SetFixedBoundsName;
+		Sig.bRequiresExecPin = true;
+		Sig.ModuleUsageBitmask = ENiagaraScriptUsageMask::Emitter | ENiagaraScriptUsageMask::System;
+		Sig.Inputs.Emplace_GetRef(FNiagaraTypeDefinition::GetBoolDef(), TEXT("Execute")).SetValue(true);
+		Sig.Inputs.Emplace_GetRef(FNiagaraTypeDefinition::GetBoolDef(), TEXT("Valid")).SetValue(true);
+		Sig.Inputs.Emplace(FNiagaraTypeDefinition::GetVec3Def(), TEXT("Local Min"));
+		Sig.Inputs.Emplace(FNiagaraTypeDefinition::GetVec3Def(), TEXT("Local Max"));
+#if WITH_EDITORONLY_DATA
+		Sig.SetDescription(LOCTEXT("SetEmitterFixedBounds", "Set local space fixed bounds for the emitter.  If Valid is set to false the emitter will fallback to the asset set fixed bounds (if valid) or dynamically generating bounds."));
+#endif
+	}
+}
+
+void UNiagaraDataInterfaceEmitterProperties::GetVMExternalFunction(const FVMExternalFunctionBindingInfo& BindingInfo, void* PerInstanceData, FVMExternalFunction& OutFunc)
+{
+	using namespace NDIEmitterPropertiesLocal;
+
+	if ( BindingInfo.Name == GetLocalSpaceName )
+	{
+		OutFunc = FVMExternalFunction::CreateLambda([this](FVectorVMExternalFunctionContext& Context) { VMGetLocalSpace(Context); });
+	}
+	else if (BindingInfo.Name == GetBoundsName)
+	{
+		OutFunc = FVMExternalFunction::CreateLambda([this](FVectorVMExternalFunctionContext& Context) { VMGetBounds(Context); });
+	}
+	else if (BindingInfo.Name == GetFixedBoundsName)
+	{
+		OutFunc = FVMExternalFunction::CreateLambda([this](FVectorVMExternalFunctionContext& Context) { VMGetFixedBounds(Context); });
+	}
+	else if (BindingInfo.Name == SetFixedBoundsName)
+	{
+		OutFunc = FVMExternalFunction::CreateLambda([this](FVectorVMExternalFunctionContext& Context) { VMSetFixedBounds(Context); });
+	}
+}
+
+#if WITH_EDITORONLY_DATA
+bool UNiagaraDataInterfaceEmitterProperties::AppendCompileHash(FNiagaraCompileHashVisitor* InVisitor) const
+{
+	bool bSuccess = Super::AppendCompileHash(InVisitor);
+	FSHAHash Hash = GetShaderFileHash(NDIEmitterPropertiesLocal::TemplateShaderFile, EShaderPlatform::SP_PCD3D_SM5);
+	InVisitor->UpdateString(TEXT("NiagaraDataInterfaceEmitterPropertiesTemplateHLSLSource"), Hash.ToString());
+	return bSuccess;
+}
+
+void UNiagaraDataInterfaceEmitterProperties::GetParameterDefinitionHLSL(const FNiagaraDataInterfaceGPUParamInfo& ParameterInfo, FString& OutHLSL)
+{
+	TMap<FString, FStringFormatArg> TemplateArgs =
+	{
+		{TEXT("ParameterName"),	ParameterInfo.DataInterfaceHLSLSymbol},
+	};
+
+	FString TemplateFile;
+	LoadShaderSourceFile(NDIEmitterPropertiesLocal::TemplateShaderFile, EShaderPlatform::SP_PCD3D_SM5, &TemplateFile, nullptr);
+	OutHLSL += FString::Format(*TemplateFile, TemplateArgs);
+}
+
+bool UNiagaraDataInterfaceEmitterProperties::GetFunctionHLSL(const FNiagaraDataInterfaceGPUParamInfo& ParamInfo, const FNiagaraDataInterfaceGeneratedFunction& FunctionInfo, int FunctionInstanceIndex, FString& OutHLSL)
+{
+	using namespace NDIEmitterPropertiesLocal;
+
+	if ((FunctionInfo.DefinitionName == GetLocalSpaceName) ||
+		(FunctionInfo.DefinitionName == GetBoundsName) ||
+		(FunctionInfo.DefinitionName == GetFixedBoundsName) )
+	{
+		return true;
+	}
+
+	return false;
+}
+#endif
+#if WITH_EDITOR	
+void UNiagaraDataInterfaceEmitterProperties::GetFeedback(UNiagaraSystem* Asset, UNiagaraComponent* Component, TArray<FNiagaraDataInterfaceError>& OutErrors, TArray<FNiagaraDataInterfaceFeedback>& OutWarnings, TArray<FNiagaraDataInterfaceFeedback>& OutInfo)
+{
+	if (Asset == nullptr)
+	{
+		return;
+	}
+
+	// See if we are resolve the source emitter
+	if ( EmitterBinding.BindingMode == ENiagaraDataInterfaceEmitterBindingMode::Other )
+	{
+		UNiagaraEmitter* NiagaraEmitter = EmitterBinding.Resolve(Asset);
+		if (NiagaraEmitter == nullptr)
+		{
+			OutWarnings.Emplace(
+				LOCTEXT("SourceEmitterNotFound", "Source emitter was not found."),
+				FText::Format(LOCTEXT("SourceEmitterNotFoundSummary", "Source emitter '{0}' was not could not be found"), FText::FromName(EmitterBinding.EmitterName)),
+				FNiagaraDataInterfaceFix()
+			);
+		}
+	}
+}
+#endif
+
+void UNiagaraDataInterfaceEmitterProperties::VMGetLocalSpace(FVectorVMExternalFunctionContext& Context)
+{
+	VectorVM::FUserPtrHandler<NDIEmitterPropertiesLocal::FInstanceData_GameThread> InstanceData(Context);
+	FNDIOutputParam<bool> OutIsLocalSpace(Context);
+
+	for (int32 InstanceIdx = 0; InstanceIdx < Context.GetNumInstances(); ++InstanceIdx)
+	{
+		OutIsLocalSpace.SetAndAdvance(InstanceData->bLocalSpace);
+	}
+}
+
+void UNiagaraDataInterfaceEmitterProperties::VMGetBounds(FVectorVMExternalFunctionContext& Context)
+{
+	VectorVM::FUserPtrHandler<NDIEmitterPropertiesLocal::FInstanceData_GameThread> InstanceData(Context);
+	FNDIOutputParam<FNiagaraBool> OutIsValid(Context);
+	FNDIOutputParam<FVector3f> OutMin(Context);
+	FNDIOutputParam<FVector3f> OutMax(Context);
+
+	const FNiagaraEmitterInstance* EmitterInstance = InstanceData->EmitterInstance;
+	const FBox Bounds = EmitterInstance ? EmitterInstance->GetBounds() : FBox(EForceInit::ForceInit);
+	const FNiagaraBool BoundsValid(Bounds.IsValid ? true : false);
+	const FVector3f BoundsMin = Bounds.IsValid ? FVector3f(Bounds.Min) : FVector3f::Zero();
+	const FVector3f BoundsMax = Bounds.IsValid ? FVector3f(Bounds.Max) : FVector3f::Zero();
+
+	for (int32 InstanceIdx = 0; InstanceIdx < Context.GetNumInstances(); ++InstanceIdx)
+	{
+		OutIsValid.SetAndAdvance(BoundsValid);
+		OutMin.SetAndAdvance(BoundsMin);
+		OutMax.SetAndAdvance(BoundsMax);
+	}
+}
+
+void UNiagaraDataInterfaceEmitterProperties::VMGetFixedBounds(FVectorVMExternalFunctionContext& Context)
+{
+	VectorVM::FUserPtrHandler<NDIEmitterPropertiesLocal::FInstanceData_GameThread> InstanceData(Context);
+	FNDIOutputParam<FNiagaraBool> OutIsValid(Context);
+	FNDIOutputParam<FVector3f> OutMin(Context);
+	FNDIOutputParam<FVector3f> OutMax(Context);
+
+	const FNiagaraEmitterInstance* EmitterInstance = InstanceData->EmitterInstance;
+	const FBox FixedBounds = EmitterInstance ? EmitterInstance->GetFixedBounds() : FBox(EForceInit::ForceInit);
+	const FNiagaraBool BoundsValid(FixedBounds.IsValid ? true : false);
+	const FVector3f BoundsMin = FixedBounds.IsValid ? FVector3f(FixedBounds.Min) : FVector3f::Zero();
+	const FVector3f BoundsMax = FixedBounds.IsValid ? FVector3f(FixedBounds.Max) : FVector3f::Zero();
+
+	for (int32 InstanceIdx = 0; InstanceIdx < Context.GetNumInstances(); ++InstanceIdx)
+	{
+		OutIsValid.SetAndAdvance(BoundsValid);
+		OutMin.SetAndAdvance(BoundsMin);
+		OutMax.SetAndAdvance(BoundsMax);
+	}
+}
+
+void UNiagaraDataInterfaceEmitterProperties::VMSetFixedBounds(FVectorVMExternalFunctionContext& Context)
+{
+	VectorVM::FUserPtrHandler<NDIEmitterPropertiesLocal::FInstanceData_GameThread> InstanceData(Context);
+	FNDIInputParam<bool> InExecute(Context);
+	FNDIInputParam<FNiagaraBool> InIsValid(Context);
+	FNDIInputParam<FVector3f> InMin(Context);
+	FNDIInputParam<FVector3f> InMax(Context);
+
+	if (FNiagaraEmitterInstance* EmitterInstance = InstanceData->EmitterInstance)
+	{
+		for (int32 InstanceIdx = 0; InstanceIdx < Context.GetNumInstances(); ++InstanceIdx)
+		{
+			FBox FixedBounds;
+			FixedBounds.IsValid = InIsValid.GetAndAdvance();
+			FixedBounds.Min = InMin.GetAndAdvance();
+			FixedBounds.Max = InMax.GetAndAdvance();
+			const bool bExecute = InExecute.GetAndAdvance();
+			if (bExecute)
+			{
+				EmitterInstance->SetFixedBounds(FixedBounds);
+			}
+		}
+	}
+}
+
+#undef LOCTEXT_NAMESPACE
diff --git a/Engine/Plugins/FX/Niagara/Source/Niagara/Private/DataInterface/NiagaraDataInterfaceEmitterProperties.h b/Engine/Plugins/FX/Niagara/Source/Niagara/Private/DataInterface/NiagaraDataInterfaceEmitterProperties.h
new file mode 100644
--- /dev/null
+++ b/Engine/Plugins/FX/Niagara/Source/Niagara/Private/DataInterface/NiagaraDataInterfaceEmitterProperties.h
@@ -0,0 +1,62 @@
+// Copyright Epic Games, Inc. All Rights Reserved.
+#pragma once
+
+#include "NiagaraCommon.h"
+#include "NiagaraShared.h"
+#include "NiagaraDataInterface.h"
+#include "NiagaraDataInterfaceEmitterBinding.h"
+#include "NiagaraParameterStore.h"
+#include "NiagaraDataInterfaceEmitterProperties.generated.h"
+
+/**
+Allows access to various emitter properties that are not part of the simulation data.
+*/
+UCLASS(EditInlineNew, Category = "DataInterface", meta=(DisplayName="Emitter Properties", Experimental))
+class UNiagaraDataInterfaceEmitterProperties : public UNiagaraDataInterface
+{
+	GENERATED_UCLASS_BODY()
+
+public:
+	DECLARE_NIAGARA_DI_PARAMETER();
+
+	/** Selects which emitter the data inteface will bind to, i.e the emitter we are contained within or a named emitter. */
+	UPROPERTY(EditAnywhere, Category = "Emitter", meta = (EditCondition = "SourceMode == NDIEmitterSourceMode::Other"))
+	FNiagaraDataInterfaceEmitterBinding EmitterBinding;
+
+	//UNiagaraDataInterface Interface
+	virtual void PostInitProperties() override;
+	// UObject Interface End
+
+	// UNiagaraDataInterface Interface Begin
+	virtual bool CanExecuteOnTarget(ENiagaraSimTarget Target) const override { return true; }
+
+	virtual bool Equals(const UNiagaraDataInterface* Other) const override;
+	virtual bool CopyToInternal(UNiagaraDataInterface* Destination) const override;
+
+	virtual bool InitPerInstanceData(void* PerInstanceData, FNiagaraSystemInstance* SystemInstance) override;
+	virtual void DestroyPerInstanceData(void* PerInstanceData, FNiagaraSystemInstance* SystemInstance) override;
+	virtual int32 PerInstanceDataSize() const override;
+
+	virtual void ProvidePerInstanceDataForRenderThread(void* DataForRenderThread, void* PerInstanceData, const FNiagaraSystemInstanceID& InstanceID) override;
+
+	virtual bool HasPreSimulateTick() const override { return true; }
+	virtual bool PerInstanceTick(void* PerInstanceData, FNiagaraSystemInstance* SystemInstance, float DeltaSeconds) override;
+
+	virtual void GetFunctions(TArray<FNiagaraFunctionSignature>& OutFunctions) override;
+	virtual void GetVMExternalFunction(const FVMExternalFunctionBindingInfo& BindingInfo, void* InstanceData, FVMExternalFunction& OutFunc) override;
+#if WITH_EDITORONLY_DATA
+	virtual bool AppendCompileHash(FNiagaraCompileHashVisitor* InVisitor) const override;
+	virtual void GetParameterDefinitionHLSL(const FNiagaraDataInterfaceGPUParamInfo& ParamInfo, FString& OutHLSL) override;
+	virtual bool GetFunctionHLSL(const FNiagaraDataInterfaceGPUParamInfo& ParamInfo, const FNiagaraDataInterfaceGeneratedFunction& FunctionInfo, int FunctionInstanceIndex, FString& OutHLSL) override;
+#endif
+#if WITH_EDITOR	
+	virtual void GetFeedback(UNiagaraSystem* Asset, UNiagaraComponent* Component, TArray<FNiagaraDataInterfaceError>& OutErrors, TArray<FNiagaraDataInterfaceFeedback>& Warnings, TArray<FNiagaraDataInterfaceFeedback>& Info) override;
+#endif
+	//UNiagaraDataInterface Interface
+
+private:
+	void VMGetLocalSpace(FVectorVMExternalFunctionContext& Context);
+	void VMGetBounds(FVectorVMExternalFunctionContext& Context);
+	void VMGetFixedBounds(FVectorVMExternalFunctionContext& Context);
+	void VMSetFixedBounds(FVectorVMExternalFunctionContext& Context);
+};
diff --git a/Engine/Plugins/FX/Niagara/Source/Niagara/Private/NiagaraDataInterfaceEmitterBinding.cpp b/Engine/Plugins/FX/Niagara/Source/Niagara/Private/NiagaraDataInterfaceEmitterBinding.cpp
new file mode 100644
--- /dev/null
+++ b/Engine/Plugins/FX/Niagara/Source/Niagara/Private/NiagaraDataInterfaceEmitterBinding.cpp
@@ -0,0 +1,133 @@
+// Copyright Epic Games, Inc. All Rights Reserved.
+
+#include "NiagaraDataInterfaceEmitterBinding.h"
+#include "NiagaraEmitterInstance.h"
+#include "NiagaraSystemInstance.h"
+
+FNiagaraEmitterInstance* FNiagaraDataInterfaceEmitterBinding::Resolve(FNiagaraSystemInstance* SystemInstance, UNiagaraDataInterface* DataInterface)
+{
+	if (BindingMode == ENiagaraDataInterfaceEmitterBindingMode::Self)
+	{
+		// If we came from a particle script our outer will be a UNiagaraEmitter
+		if (UNiagaraEmitter* OwnerEmitter = DataInterface->GetTypedOuter<UNiagaraEmitter>())
+		{
+			for (TSharedPtr<FNiagaraEmitterInstance, ESPMode::ThreadSafe> EmitterInstance : SystemInstance->GetEmitters())
+			{
+				if (EmitterInstance->GetCachedEmitter() == OwnerEmitter)
+				{
+					return EmitterInstance.Get();
+				}
+			}
+		}
+		// If we came from an emitter script we need to look over the default data interfaces to find ourself
+		else
+		{
+			auto FindDataInterfaceName =
+				[DataInterface](UNiagaraSystem* NiagaraSystem) -> FName
+				{
+					for (UNiagaraScript* NiagaraScript : { NiagaraSystem->GetSystemUpdateScript(), NiagaraSystem->GetSystemSpawnScript() })
+					{
+						if (NiagaraScript != nullptr)
+						{
+							for (const FNiagaraScriptDataInterfaceInfo& DataInterfaceInfo : NiagaraScript->GetCachedDefaultDataInterfaces())
+							{
+								if (DataInterfaceInfo.DataInterface == DataInterface)
+								{
+									return DataInterfaceInfo.RegisteredParameterMapWrite.IsNone() ? DataInterfaceInfo.Name : DataInterfaceInfo.RegisteredParameterMapWrite;
+								}
+							}
+						}
+					}
+					return NAME_None;
+				};
+
+			const FName DataInterfaceName = FindDataInterfaceName(SystemInstance->GetSystem());
+			if (!DataInterfaceName.IsNone())
+			{
+				FNameBuilder SelfEmitterString;
+				DataInterfaceName.ToString(SelfEmitterString);
+				FStringView SelfEmitterView = SelfEmitterString.ToView();
+				int32 DotIndex;
+				if (SelfEmitterView.FindChar('.', DotIndex))
+				{
+					SelfEmitterView = SelfEmitterView.Mid(0, DotIndex);
+					for (TSharedPtr<FNiagaraEmitterInstance, ESPMode::ThreadSafe> EmitterInstance : SystemInstance->GetEmitters())
+					{
+						if (const UNiagaraEmitter* CachedEmitter = EmitterInstance->GetCachedEmitter())
+						{
+							if (SelfEmitterView.Equals(CachedEmitter->GetUniqueEmitterName()))
+							{
+								return EmitterInstance.Get();
+							}
+						}
+					}
+				}
+			}
+		}
+		if ( FNiagaraUtilities::LogVerboseWarnings() )
+		{
+			UE_LOG(LogNiagara, Warning, TEXT("EmitterBinding failed to find self emitter"));
+		}
+	}
+	else
+	{
+		if ( EmitterName.IsNone() )
+		{
+			if ( FNiagaraUtilities::LogVerboseWarnings() )
+			{
+				UE_LOG(LogNiagara, Warning, TEXT("EmitterName has not been set but we are in Other mode"));
+			}
+			return nullptr;
+		}
+
+		FNameBuilder EmitterNameString;
+		EmitterName.ToString(EmitterNameString);
+		FStringView EmitterNameStringView = EmitterNameString.ToView();
+
+		for (TSharedPtr<FNiagaraEmitterInstance, ESPMode::ThreadSafe> EmitterInstance : SystemInstance->GetEmitters())
+		{
+			if (const UNiagaraEmitter* CachedEmitter = EmitterInstance->GetCachedEmitter())
+			{
+				//-TODO: UniqueEmitterName should probably be a FName?
+				if (EmitterNameStringView.Equals(CachedEmitter->GetUniqueEmitterName()) )
+				{
+					return EmitterInstance.Get();
+				}
+			}
+		}
+
+		if (FNiagaraUtilities::LogVerboseWarnings())
+		{
+			UE_LOG(LogNiagara, Warning, TEXT("EmitterBinding failed to find emitter '%s' it might not exist or has been cooked out"), EmitterNameStringView.GetData());
+		}
+	}
+	return nullptr;
+}
+
+UNiagaraEmitter* FNiagaraDataInterfaceEmitterBinding::Resolve(UNiagaraSystem* NiagaraSystem)
+{
+	if ( BindingMode == ENiagaraDataInterfaceEmitterBindingMode::Other )
+	{
+		if (EmitterName.IsNone())
+		{
+			return nullptr;
+		}
+
+		FNameBuilder EmitterNameString;
+		EmitterName.ToString(EmitterNameString);
+		FStringView EmitterNameStringView = EmitterNameString.ToView();
+
+		for ( const FNiagaraEmitterHandle& EmitterHandle : NiagaraSystem->GetEmitterHandles() )
+		{
+			if ( UNiagaraEmitter* NiagaraEmitter = EmitterHandle.GetInstance() )
+			{
+				//-TODO: UniqueEmitterName should probably be a FName?
+				if (EmitterNameStringView.Equals(NiagaraEmitter->GetUniqueEmitterName()))
+				{
+					return NiagaraEmitter;
+				}
+			}
+		}
+	}
+	return nullptr;
+}
diff --git a/Engine/Plugins/FX/Niagara/Source/Niagara/Private/NiagaraDataInterfaceParticleRead.cpp b/Engine/Plugins/FX/Niagara/Source/Niagara/Private/NiagaraDataInterfaceParticleRead.cpp
--- a/Engine/Plugins/FX/Niagara/Source/Niagara/Private/NiagaraDataInterfaceParticleRead.cpp
+++ b/Engine/Plugins/FX/Niagara/Source/Niagara/Private/NiagaraDataInterfaceParticleRead.cpp
@@ -23,6 +23,8 @@ struct FNiagaraParticleReadDIFunctionVersion
 	};
 };
 
+static const FName GetLocalSpaceName("GetLocalSpace");
+
 static const FName GetNumSpawnedParticlesFunctionName("Get Num Spawned Particles");
 static const FName GetIDAtSpawnIndexFunctionName("Get ID At Spawn Index");
 static const FName GetNumParticlesFunctionName("Get Num Particles");
@@ -718,6 +720,20 @@ int32 UNiagaraDataInterfaceParticleRead::PerInstanceDataSize() const
 
 void UNiagaraDataInterfaceParticleRead::GetFunctions(TArray<FNiagaraFunctionSignature>& OutFunctions)
 {
+	// Misc functionality
+	{
+		FNiagaraFunctionSignature& Sig = OutFunctions.AddDefaulted_GetRef();
+		Sig.Name = GetLocalSpaceName;
+		Sig.Inputs.Emplace(FNiagaraTypeDefinition(GetClass()), TEXT("Particle Reader"));
+		Sig.bMemberFunction = true;
+		Sig.bRequiresContext = false;
+		Sig.bExperimental = true;
+		Sig.Outputs.Emplace(FNiagaraTypeDefinition::GetBoolDef(), TEXT("IsLocalSpace"));
+	#if WITH_EDITORONLY_DATA
+		Sig.SetDescription(LOCTEXT("GetEmitterLocalSpace", "Returns if the emitter is using local space or world space."));
+	#endif
+	}
+
 	//
 	// Spawn info and particle count
 	//
@@ -1196,6 +1212,13 @@ static bool HasMatchingVariable(TArrayView<const FNiagaraVariable> Variables, FN
 
 void UNiagaraDataInterfaceParticleRead::GetVMExternalFunction(const FVMExternalFunctionBindingInfo& BindingInfo, void* InstanceData, FVMExternalFunction& OutFunc)
 {
+	// Misc functionality
+	if ( BindingInfo.Name == GetLocalSpaceName )
+	{	
+		OutFunc = FVMExternalFunction::CreateLambda([this](FVectorVMExternalFunctionContext& Context) { VMGetLocalSpace(Context); });
+		return;
+	}
+
 	//
 	// Spawn info and particle count
 	//
@@ -1395,6 +1418,21 @@ void UNiagaraDataInterfaceParticleRead::GetVMExternalFunction(const FVMExternalF
 	}
 }
 
+void UNiagaraDataInterfaceParticleRead::VMGetLocalSpace(FVectorVMExternalFunctionContext& Context)
+{
+	VectorVM::FUserPtrHandler<FNDIParticleRead_InstanceData> InstData(Context);
+	FNDIOutputParam<bool> OutIsLocalSpace(Context);
+
+	const FNiagaraEmitterInstance* EmitterInstance = InstData.Get()->EmitterInstance;
+	UNiagaraEmitter* NiagaraEmitter = EmitterInstance ? EmitterInstance->GetCachedEmitter() : nullptr;
+	const bool bIsLocalSpace = NiagaraEmitter ? NiagaraEmitter->bLocalSpace : false;
+
+	for (int32 InstanceIdx = 0; InstanceIdx < Context.GetNumInstances(); ++InstanceIdx)
+	{
+		OutIsLocalSpace.SetAndAdvance(bIsLocalSpace);
+	}
+}
+
 void UNiagaraDataInterfaceParticleRead::GetNumSpawnedParticles(FVectorVMExternalFunctionContext& Context)
 {
 	VectorVM::FUserPtrHandler<FNDIParticleRead_InstanceData> InstData(Context);
diff --git a/Engine/Plugins/FX/Niagara/Source/Niagara/Private/NiagaraEmitterInstance.cpp b/Engine/Plugins/FX/Niagara/Source/Niagara/Private/NiagaraEmitterInstance.cpp
--- a/Engine/Plugins/FX/Niagara/Source/Niagara/Private/NiagaraEmitterInstance.cpp
+++ b/Engine/Plugins/FX/Niagara/Source/Niagara/Private/NiagaraEmitterInstance.cpp
@@ -122,8 +122,7 @@ struct FNiagaraEditorOnlyCycleTimer
 
 FNiagaraEmitterInstance::FNiagaraEmitterInstance(FNiagaraSystemInstance* InParentSystemInstance)
 	: CachedBounds(ForceInit)
-	, FixedBounds_GT(ForceInit)
-	, FixedBounds_CNC(ForceInit)
+	, FixedBounds(ForceInit)
 	, CachedSystemFixedBounds(ForceInit)
 	, ParentSystemInstance(InParentSystemInstance)
 {
@@ -175,7 +174,7 @@ FNiagaraEmitterInstance::~FNiagaraEmitterInstance()
 	}
 }
 
-FBox FNiagaraEmitterInstance::GetBounds()
+FBox FNiagaraEmitterInstance::GetBounds() const
 {
 	return CachedBounds;
 }
@@ -565,8 +564,7 @@ void FNiagaraEmitterInstance::OnPooledReuse()
 	bResetPending = true;
 	TotalSpawnedParticles = 0;
 
-	FixedBounds_GT.Init();
-	FixedBounds_CNC.Init();
+	FixedBounds.Init();
 }
 
 void FNiagaraEmitterInstance::SetParticleComponentActive(FObjectKey ComponentKey, int32 ParticleID) const
@@ -1027,36 +1025,40 @@ void FNiagaraEmitterInstance::PostTick()
 	}
 
 	CachedBounds.Init();
-	if (FixedBounds_CNC.IsValid)
 	{
-		CachedBounds = FixedBounds_CNC;
-	}
-	else if (CachedSystemFixedBounds.IsValid)
-	{
-		CachedBounds = CachedSystemFixedBounds;
-	}
-	else if (CachedEmitter->bFixedBounds || CachedEmitter->SimTarget == ENiagaraSimTarget::GPUComputeSim)
-	{
-		CachedBounds = CachedEmitter->FixedBounds;
-	}
-	else
-	{
-		FBox DynamicBounds = InternalCalculateDynamicBounds(ParticleDataSet->GetCurrentDataChecked().GetNumInstances());
-		if (DynamicBounds.IsValid)
+		// Read lock can be smaller in scope, but probably not necessary
+		FRWScopeLock ScopeLock(FixedBoundsGuard, SLT_ReadOnly);
+		if (FixedBounds.IsValid)
+		{
+			CachedBounds = FixedBounds;
+		}
+		else if (CachedSystemFixedBounds.IsValid)
+		{
+			CachedBounds = CachedSystemFixedBounds;
+		}
+		else if (CachedEmitter->bFixedBounds || CachedEmitter->SimTarget == ENiagaraSimTarget::GPUComputeSim)
+		{
+			CachedBounds = CachedEmitter->FixedBounds;
+		}
+		else
 		{
-			if (CachedEmitter->bLocalSpace)
+			FBox DynamicBounds = InternalCalculateDynamicBounds(ParticleDataSet->GetCurrentDataChecked().GetNumInstances());
+			if (DynamicBounds.IsValid)
 			{
-				CachedBounds = DynamicBounds;
+				if (CachedEmitter->bLocalSpace)
+				{
+					CachedBounds = DynamicBounds;
+				}
+				else
+				{
+					CachedBounds = DynamicBounds.TransformBy(FMatrix(ParentSystemInstance->GetOwnerParameters().EngineWorldToLocal));
+				}
 			}
 			else
 			{
-				CachedBounds = DynamicBounds.TransformBy(FMatrix(ParentSystemInstance->GetOwnerParameters().EngineWorldToLocal));
+				CachedBounds = CachedEmitter->FixedBounds;
 			}
 		}
-		else
-		{
-			CachedBounds = CachedEmitter->FixedBounds;
-		}
 	}
 
 #if STATS
@@ -1146,8 +1148,6 @@ void FNiagaraEmitterInstance::PreTick()
 	FScopeCycleCounter SystemStatCounter(CachedEmitter->GetStatID(true, true));
 #endif
 
-	FixedBounds_CNC = FixedBounds_GT;
-
 	checkSlow(ParticleDataSet);
 	FNiagaraDataSet& Data = *ParticleDataSet;
 
@@ -1242,11 +1242,15 @@ void FNiagaraEmitterInstance::SetSystemFixedBoundsOverride(FBox SystemFixedBound
 
 FBox FNiagaraEmitterInstance::GetFixedBounds() const
 {
-	if ( FixedBounds_GT.IsValid )
 	{
-		return FixedBounds_GT;
+		FRWScopeLock ScopeLock(FixedBoundsGuard, SLT_ReadOnly);
+		if (FixedBounds.IsValid)
+		{
+			return FixedBounds;
+		}
 	}
-	else if ( CachedEmitter && CachedEmitter->bFixedBounds)
+
+	if ( CachedEmitter && CachedEmitter->bFixedBounds)
 	{
 		return CachedEmitter->FixedBounds;
 	}
diff --git a/Engine/Plugins/FX/Niagara/Source/Niagara/Public/NiagaraDataInterfaceEmitterBinding.h b/Engine/Plugins/FX/Niagara/Source/Niagara/Public/NiagaraDataInterfaceEmitterBinding.h
new file mode 100644
--- /dev/null
+++ b/Engine/Plugins/FX/Niagara/Source/Niagara/Public/NiagaraDataInterfaceEmitterBinding.h
@@ -0,0 +1,40 @@
+// Copyright Epic Games, Inc. All Rights Reserved.
+#pragma once
+
+#include "NiagaraCommon.h"
+#include "NiagaraDataInterface.h"
+#include "NiagaraDataInterfaceEmitterBinding.generated.h"
+
+class UNiagaraEmitter;
+class UNiagaraSystem;
+
+UENUM()
+enum class ENiagaraDataInterfaceEmitterBindingMode
+{
+	/** Attempt to bind to the emitter the data interface is used with, this may not be possible in all situations (i.e. user parameter). */
+	Self,
+	/** Attempt to bind to an emitter using it's name. */
+	Other
+};
+
+USTRUCT()
+struct FNiagaraDataInterfaceEmitterBinding
+{
+	GENERATED_BODY()
+
+	UPROPERTY(EditAnywhere, Category = "Binding")
+	ENiagaraDataInterfaceEmitterBindingMode	BindingMode = ENiagaraDataInterfaceEmitterBindingMode::Other;
+	
+	UPROPERTY(EditAnywhere, Category = "Binding")
+	FName EmitterName;
+
+	/** Resolves the emitter binding, this can return nullptr if we failed to resolve */
+	FNiagaraEmitterInstance* Resolve(FNiagaraSystemInstance* SystemInstance, UNiagaraDataInterface* DataInterface);
+	/** Resolves the emitter binding, only returns a valid result when in Named mode as we can not determine the source otherwise. */
+	UNiagaraEmitter* Resolve(UNiagaraSystem* NiagaraSystem);
+
+	bool operator==(const FNiagaraDataInterfaceEmitterBinding& Other) const
+	{
+		return BindingMode == Other.BindingMode && EmitterName == Other.EmitterName;
+	}
+};
diff --git a/Engine/Plugins/FX/Niagara/Source/NiagaraEditor/Private/Customizations/NiagaraDataInterfaceEmitterBindingCustomization.cpp b/Engine/Plugins/FX/Niagara/Source/NiagaraEditor/Private/Customizations/NiagaraDataInterfaceEmitterBindingCustomization.cpp
new file mode 100644
--- /dev/null
+++ b/Engine/Plugins/FX/Niagara/Source/NiagaraEditor/Private/Customizations/NiagaraDataInterfaceEmitterBindingCustomization.cpp
@@ -0,0 +1,108 @@
+// Copyright Epic Games, Inc. All Rights Reserved.
+
+#include "NiagaraDataInterfaceEmitterBindingCustomization.h"
+#include "NiagaraDataInterface.h"
+#include "NiagaraEditorModule.h"
+
+#include "DetailCategoryBuilder.h"
+#include "DetailLayoutBuilder.h"
+#include "DetailWidgetRow.h"
+#include "IDetailChildrenBuilder.h"
+#include "IPropertyUtilities.h"
+#include "Modules/ModuleManager.h"
+#include "Widgets/Input/SCheckBox.h"
+#include "Widgets/Input/SEditableTextBox.h"
+#include "Widgets/SBoxPanel.h"
+#include "Widgets/SNiagaraDebugger.h"
+#include "Widgets/Views/STreeView.h"
+#include "Widgets/Input/SSearchBox.h"
+#include "Widgets/Input/SSuggestionTextBox.h"
+#include "Widgets/Layout/SScrollBorder.h"
+#include "Widgets/Input/SMultiLineEditableTextBox.h"
+
+#define LOCTEXT_NAMESPACE "NiagaraDataInterfaceEmitterBindingCustomization"
+
+void FNiagaraDataInterfaceEmitterBindingCustomization::CustomizeHeader(TSharedRef<IPropertyHandle> StructPropertyHandle, FDetailWidgetRow& HeaderRow, IPropertyTypeCustomizationUtils& CustomizationUtils)
+{
+	// Get owner system (if any)
+	if ( TSharedPtr<IPropertyUtilities> PropertyUtilities = CustomizationUtils.GetPropertyUtilities() )
+	{
+		const TArray<TWeakObjectPtr<UObject>>& SelectedObjects = PropertyUtilities->GetSelectedObjects();
+		if (SelectedObjects.Num() == 1)
+		{
+			if ( UObject* OwnerObject = SelectedObjects[0].Get() )
+			{
+				NiagaraSystemWeak = OwnerObject->GetTypedOuter<UNiagaraSystem>();
+			}
+		}
+	}
+
+	// Get properties
+	BindingModeHandle = StructPropertyHandle->GetChildHandle(GET_MEMBER_NAME_CHECKED(FNiagaraDataInterfaceEmitterBinding, BindingMode));
+	EmitterNameHandle = StructPropertyHandle->GetChildHandle(GET_MEMBER_NAME_CHECKED(FNiagaraDataInterfaceEmitterBinding, EmitterName));
+	check(BindingModeHandle.IsValid() && EmitterNameHandle.IsValid())
+
+	// Build Widget
+	HeaderRow
+		.NameContent()
+		[
+			StructPropertyHandle->CreatePropertyNameWidget()
+		]
+		.ValueContent()
+		.MinDesiredWidth(200.0f)
+		[
+			SNew(SHorizontalBox)
+			+ SHorizontalBox::Slot()
+			.AutoWidth()
+			[
+				BindingModeHandle->CreatePropertyValueWidget()
+			]
+			+ SHorizontalBox::Slot()
+			[
+				SNew(SSuggestionTextBox)
+				.ForegroundColor(FSlateColor::UseForeground())
+				.Visibility(this, &FNiagaraDataInterfaceEmitterBindingCustomization::IsEmitterNameVisible)
+				.Font(IDetailLayoutBuilder::GetDetailFont())
+				.Text(EmitterName_GetText())
+				.OnTextCommitted(this, &FNiagaraDataInterfaceEmitterBindingCustomization::EmitterName_SetText)
+				.HintText(LOCTEXT("EmitterNameHint", "Enter Emitter Name..."))
+				.OnShowingSuggestions(this, &FNiagaraDataInterfaceEmitterBindingCustomization::EmitterName_GetSuggestions)
+			]
+		];
+}
+
+EVisibility FNiagaraDataInterfaceEmitterBindingCustomization::IsEmitterNameVisible() const
+{
+	uint8 BindingModeValue;
+	BindingModeHandle->GetValue(BindingModeValue);
+	return ENiagaraDataInterfaceEmitterBindingMode(BindingModeValue) == ENiagaraDataInterfaceEmitterBindingMode::Other ? EVisibility::Visible : EVisibility::Hidden;
+}
+
+FText FNiagaraDataInterfaceEmitterBindingCustomization::EmitterName_GetText() const
+{
+	FName EmitterName;
+	EmitterNameHandle->GetValue(EmitterName);
+	return EmitterName.IsNone() ? FText() : FText::FromName(EmitterName);
+}
+
+void FNiagaraDataInterfaceEmitterBindingCustomization::EmitterName_SetText(const FText& NewText, ETextCommit::Type CommitInfo)
+{
+	FName EmitterName(NewText.ToString());
+	EmitterNameHandle->SetValue(EmitterName);
+}
+
+void FNiagaraDataInterfaceEmitterBindingCustomization::EmitterName_GetSuggestions(const FString& CurrText, TArray<FString>& OutSuggestions)
+{
+	if ( UNiagaraSystem* NiagaraSystem = NiagaraSystemWeak.Get() )
+	{
+		for ( const FNiagaraEmitterHandle& EmitterHandle : NiagaraSystem->GetEmitterHandles() )
+		{
+			if ( EmitterHandle.GetIsEnabled() )
+			{
+				OutSuggestions.Add(EmitterHandle.GetUniqueInstanceName());
+			}
+		}
+	}
+}
+
+#undef LOCTEXT_NAMESPACE
diff --git a/Engine/Plugins/FX/Niagara/Source/NiagaraEditor/Private/Customizations/NiagaraDataInterfaceEmitterBindingCustomization.h b/Engine/Plugins/FX/Niagara/Source/NiagaraEditor/Private/Customizations/NiagaraDataInterfaceEmitterBindingCustomization.h
new file mode 100644
--- /dev/null
+++ b/Engine/Plugins/FX/Niagara/Source/NiagaraEditor/Private/Customizations/NiagaraDataInterfaceEmitterBindingCustomization.h
@@ -0,0 +1,36 @@
+// Copyright Epic Games, Inc. All Rights Reserved.
+
+#pragma once
+
+#include "CoreMinimal.h"
+#include "IPropertyTypeCustomization.h"
+#include "NiagaraDataInterfaceEmitterBinding.h"
+
+class FDetailWidgetRow;
+class IPropertyHandle;
+class UNiagaraSystem;
+
+class FNiagaraDataInterfaceEmitterBindingCustomization : public IPropertyTypeCustomization
+{
+public:
+	static TSharedRef<IPropertyTypeCustomization> MakeInstance()
+	{
+		return MakeShared<FNiagaraDataInterfaceEmitterBindingCustomization>();
+	}
+
+	virtual void CustomizeHeader(TSharedRef<IPropertyHandle> PropertyHandle, FDetailWidgetRow& HeaderRow, IPropertyTypeCustomizationUtils& CustomizationUtils);
+
+	virtual void CustomizeChildren( TSharedRef<IPropertyHandle> PropertyHandle, IDetailChildrenBuilder& ChildBuilder, IPropertyTypeCustomizationUtils& CustomizationUtils )
+	{
+	}
+
+	EVisibility IsEmitterNameVisible() const;
+	FText EmitterName_GetText() const;
+	void EmitterName_SetText(const FText& NewText, ETextCommit::Type CommitInfo);
+	void EmitterName_GetSuggestions(const FString& CurrText, TArray<FString>& OutSuggestions);
+
+	TWeakObjectPtr<UNiagaraSystem> NiagaraSystemWeak;
+	TSharedPtr<IPropertyHandle> BindingModeHandle;
+	TSharedPtr<IPropertyHandle> EmitterNameHandle;
+};
+
diff --git a/Engine/Plugins/FX/Niagara/Source/NiagaraEditor/Private/NiagaraEditorModule.cpp b/Engine/Plugins/FX/Niagara/Source/NiagaraEditor/Private/NiagaraEditorModule.cpp
--- a/Engine/Plugins/FX/Niagara/Source/NiagaraEditor/Private/NiagaraEditorModule.cpp
+++ b/Engine/Plugins/FX/Niagara/Source/NiagaraEditor/Private/NiagaraEditorModule.cpp
@@ -106,6 +106,7 @@
 #include "Customizations/NiagaraStaticSwitchNodeDetails.h"
 #include "Customizations/NiagaraTypeCustomizations.h"
 #include "Customizations/NiagaraComponentRendererPropertiesDetails.h"
+#include "Customizations/NiagaraDataInterfaceEmitterBindingCustomization.h"
 #include "Customizations/NiagaraDebugHUDCustomization.h"
 #include "Customizations/NiagaraBakerSettingsDetails.h"
 #include "Customizations/NiagaraOutlinerCustomization.h"
@@ -917,6 +918,10 @@ void FNiagaraEditorModule::StartupModule()
 		UNiagaraSimulationStageGeneric::StaticClass()->GetFName(),
 		FOnGetDetailCustomizationInstance::CreateStatic(&FNiagaraSimulationStageGenericCustomization::MakeInstance));
 
+	PropertyModule.RegisterCustomPropertyTypeLayout(
+		FNiagaraDataInterfaceEmitterBinding::StaticStruct()->GetFName(),
+		FOnGetPropertyTypeCustomizationInstance::CreateStatic(&FNiagaraDataInterfaceEmitterBindingCustomization::MakeInstance));
+
 #if WITH_NIAGARA_DEBUGGER
 	PropertyModule.RegisterCustomPropertyTypeLayout(
 		FNiagaraDebugHUDVariable::StaticStruct()->GetFName(),
