- Disable GPU on set fixed bounds

#rnx

diff --git a/Engine/Plugins/FX/Niagara/Source/Niagara/Private/DataInterface/NiagaraDataInterfaceEmitterProperties.cpp b/Engine/Plugins/FX/Niagara/Source/Niagara/Private/DataInterface/NiagaraDataInterfaceEmitterProperties.cpp
--- a/Engine/Plugins/FX/Niagara/Source/Niagara/Private/DataInterface/NiagaraDataInterfaceEmitterProperties.cpp
+++ b/Engine/Plugins/FX/Niagara/Source/Niagara/Private/DataInterface/NiagaraDataInterfaceEmitterProperties.cpp
@@ -277,13 +277,14 @@ void UNiagaraDataInterfaceEmitterProperties::GetFunctions(TArray<FNiagaraFunctio
 		FNiagaraFunctionSignature& Sig = OutFunctions.Add_GetRef(DefaultSignature);
 		Sig.Name = SetFixedBoundsName;
 		Sig.bRequiresExecPin = true;
+		Sig.bSupportsGPU = false;
 		Sig.ModuleUsageBitmask = ENiagaraScriptUsageMask::Emitter | ENiagaraScriptUsageMask::System;
 		Sig.Inputs.Emplace_GetRef(FNiagaraTypeDefinition::GetBoolDef(), TEXT("Execute")).SetValue(true);
 		Sig.Inputs.Emplace_GetRef(FNiagaraTypeDefinition::GetBoolDef(), TEXT("Valid")).SetValue(true);
 		Sig.Inputs.Emplace(FNiagaraTypeDefinition::GetVec3Def(), TEXT("Local Min"));
 		Sig.Inputs.Emplace(FNiagaraTypeDefinition::GetVec3Def(), TEXT("Local Max"));
 #if WITH_EDITORONLY_DATA
-		Sig.SetDescription(LOCTEXT("SetEmitterFixedBounds", "Set local space fixed bounds for the emitter.  If Valid is set to false the emitter will fallback to the asset set fixed bounds (if valid) or dynamically generating bounds."));
+		Sig.SetDescription(LOCTEXT("SetEmitterFixedBounds", "Set local space fixed bounds for the emitter, this stomps any fixed bounds set from Blueprint.  If Valid is set to false the emitter will fallback to the asset set fixed bounds (if valid) or dynamically generating bounds."));
 #endif
 	}
 }
diff --git a/Engine/Plugins/FX/Niagara/Source/Niagara/Private/DataInterface/NiagaraDataInterfaceEmitterProperties.h b/Engine/Plugins/FX/Niagara/Source/Niagara/Private/DataInterface/NiagaraDataInterfaceEmitterProperties.h
--- a/Engine/Plugins/FX/Niagara/Source/Niagara/Private/DataInterface/NiagaraDataInterfaceEmitterProperties.h
+++ b/Engine/Plugins/FX/Niagara/Source/Niagara/Private/DataInterface/NiagaraDataInterfaceEmitterProperties.h
@@ -19,8 +19,8 @@ class UNiagaraDataInterfaceEmitterProperties : public UNiagaraDataInterface
 public:
 	DECLARE_NIAGARA_DI_PARAMETER();
 
-	/** Selects which emitter the data inteface will bind to, i.e the emitter we are contained within or a named emitter. */
-	UPROPERTY(EditAnywhere, Category = "Emitter", meta = (EditCondition = "SourceMode == NDIEmitterSourceMode::Other"))
+	/** Selects which emitter the data interface will bind to, i.e the emitter we are contained within or a named emitter. */
+	UPROPERTY(EditAnywhere, Category = "Emitter")
 	FNiagaraDataInterfaceEmitterBinding EmitterBinding;
 
 	//UNiagaraDataInterface Interface
