Fix for implicit bound renderer attributes not getting included in the generated script because of a mismatch in the types (Position vs Vec3) during translation

diff --git a/Engine/Plugins/FX/Niagara/Source/NiagaraCore/Private/NiagaraCustomVersion.cpp b/Engine/Plugins/FX/Niagara/Source/NiagaraCore/Private/NiagaraCustomVersion.cpp
--- a/Engine/Plugins/FX/Niagara/Source/NiagaraCore/Private/NiagaraCustomVersion.cpp
+++ b/Engine/Plugins/FX/Niagara/Source/NiagaraCore/Private/NiagaraCustomVersion.cpp
@@ -9,4 +9,4 @@ const FGuid FNiagaraCustomVersion::GUID(0xFCF57AFA, 0x50764283, 0xB9A9E658, 0xFF
 FCustomVersionRegistration GRegisterNiagaraCustomVersion(FNiagaraCustomVersion::GUID, FNiagaraCustomVersion::LatestVersion, TEXT("NiagaraVer"));
 
 // Note: When encountering a conflict on this file please generate a new GUID
-const FGuid FNiagaraCustomVersion::LatestScriptCompileVersion(0x07F493BD, 0xE6B2A54A, 0xBD23F309, 0x93A5BC76);
\ No newline at end of file
+const FGuid FNiagaraCustomVersion::LatestScriptCompileVersion(0xDED3AAF7, 0xE9A50E4C, 0x9E7C429C, 0x94EB9117);
\ No newline at end of file
diff --git a/Engine/Plugins/FX/Niagara/Source/NiagaraEditor/Private/NiagaraHlslTranslator.cpp b/Engine/Plugins/FX/Niagara/Source/NiagaraEditor/Private/NiagaraHlslTranslator.cpp
--- a/Engine/Plugins/FX/Niagara/Source/NiagaraEditor/Private/NiagaraHlslTranslator.cpp
+++ b/Engine/Plugins/FX/Niagara/Source/NiagaraEditor/Private/NiagaraHlslTranslator.cpp
@@ -755,23 +755,40 @@ FString FHlslNiagaraTranslator::BuildParameterMapHlslDefinitions(TArray<FNiagara
 		TArray<FNiagaraVariable> Vars = OtherOutputParamMapHistories[ParamMapIdx].Variables;
 		for (const FNiagaraVariableBase& Var : CompileOptions.AdditionalVariables)
 		{
-			uint32 Uses = 0;
+			bool bFoundSource = false;
 			if (FNiagaraParameterMapHistory::IsPreviousValue(Var))
 			{
 				FNiagaraVariable Source = FNiagaraParameterMapHistory::GetSourceForPreviousValue(FNiagaraVariable(Var));
-				for (int32 ParamMapIdxTest = 0; ParamMapIdxTest < OtherOutputParamMapHistories.Num(); ParamMapIdxTest++)
+				const FNiagaraTypeDefinition& SourceType = Source.GetType();
+
+				for (int32 ParamMapIdxTest = 0; !bFoundSource && ParamMapIdxTest < OtherOutputParamMapHistories.Num(); ParamMapIdxTest++)
 				{				
+					const TArray<FNiagaraVariable>& OtherVariables = OtherOutputParamMapHistories[ParamMapIdxTest].Variables;
+					const int32 OtherVarIndex = OtherVariables.IndexOfByPredicate([&](const FNiagaraVariable& OtherVar)
+					{
+						return OtherVar.GetName() == Source.GetName();
+					});
 
-					if (OtherOutputParamMapHistories[ParamMapIdxTest].Variables.Contains(Source))
+					if (OtherVarIndex != INDEX_NONE)
 					{
-						// Disallow the addition previous values if its source is not used
-						Uses++;
-						break;
+						const FNiagaraTypeDefinition& OtherType = OtherVariables[OtherVarIndex].GetType();
+
+						if (OtherType == SourceType)
+						{
+							bFoundSource = true;
+						}
+						// special case handling for when we have an implicit Position type for the source
+						// and we've found a Vector
+						else if (SourceType == FNiagaraTypeDefinition::GetPositionDef()
+							&& OtherType == FNiagaraTypeDefinition::GetVec3Def())
+						{
+							bFoundSource = true;
+						}
 					}
 				}
 			}
 
-			if (Uses > 0)
+			if (bFoundSource)
 			{
 				Vars.AddUnique(FNiagaraVariable(Var.GetType(), Var.GetName()));
 			}
