- Customization for Particle Read so we have suggestions for emitter names
- Demote missing emitter from error to warning

#rnx

diff --git a/Engine/Plugins/FX/Niagara/Source/Niagara/Private/DataInterface/NiagaraDataInterfaceEmitterProperties.cpp b/Engine/Plugins/FX/Niagara/Source/Niagara/Private/DataInterface/NiagaraDataInterfaceEmitterProperties.cpp
--- a/Engine/Plugins/FX/Niagara/Source/Niagara/Private/DataInterface/NiagaraDataInterfaceEmitterProperties.cpp
+++ b/Engine/Plugins/FX/Niagara/Source/Niagara/Private/DataInterface/NiagaraDataInterfaceEmitterProperties.cpp
@@ -362,7 +362,7 @@ void UNiagaraDataInterfaceEmitterProperties::GetFeedback(UNiagaraSystem* Asset,
 		{
 			OutWarnings.Emplace(
 				LOCTEXT("SourceEmitterNotFound", "Source emitter was not found."),
-				FText::Format(LOCTEXT("SourceEmitterNotFoundSummary", "Source emitter '{0}' was not could not be found"), FText::FromName(EmitterBinding.EmitterName)),
+				FText::Format(LOCTEXT("SourceEmitterNotFoundSummary", "Source emitter '{0}' could not be found"), FText::FromName(EmitterBinding.EmitterName)),
 				FNiagaraDataInterfaceFix()
 			);
 		}
diff --git a/Engine/Plugins/FX/Niagara/Source/Niagara/Private/NiagaraDataInterfaceParticleRead.cpp b/Engine/Plugins/FX/Niagara/Source/Niagara/Private/NiagaraDataInterfaceParticleRead.cpp
--- a/Engine/Plugins/FX/Niagara/Source/Niagara/Private/NiagaraDataInterfaceParticleRead.cpp
+++ b/Engine/Plugins/FX/Niagara/Source/Niagara/Private/NiagaraDataInterfaceParticleRead.cpp
@@ -2392,10 +2392,11 @@ void UNiagaraDataInterfaceParticleRead::GetFeedback(UNiagaraSystem* Asset, UNiag
 
 	if (!FoundSourceEmitter)
 	{
-		FNiagaraDataInterfaceError SourceEmitterNotFoundError(LOCTEXT("SourceEmitterNotFoundError", "Source emitter not found."),
-			LOCTEXT("SourceEmitterNotFoundErrorSummary", "Source emitter not found"),
-			FNiagaraDataInterfaceFix());
-		OutErrors.Add(SourceEmitterNotFoundError);
+		Warnings.Emplace(
+			LOCTEXT("SourceEmitterNotFound", "Source emitter was not found."),
+			FText::Format(LOCTEXT("SourceEmitterNotFoundSummary", "Source emitter '{0}' could not be found"), FText::FromString(EmitterName)),
+			FNiagaraDataInterfaceFix()
+		);
 	}
 
 	// Filter through all the relevant CPU scripts
diff --git a/Engine/Plugins/FX/Niagara/Source/NiagaraEditorWidgets/Private/DetailCustomizations/NiagaraDataInterfaceParticleReadDetails.cpp b/Engine/Plugins/FX/Niagara/Source/NiagaraEditorWidgets/Private/DetailCustomizations/NiagaraDataInterfaceParticleReadDetails.cpp
new file mode 100644
--- /dev/null
+++ b/Engine/Plugins/FX/Niagara/Source/NiagaraEditorWidgets/Private/DetailCustomizations/NiagaraDataInterfaceParticleReadDetails.cpp
@@ -0,0 +1,92 @@
+// Copyright Epic Games, Inc. All Rights Reserved.
+
+#include "NiagaraDataInterfaceParticleReadDetails.h"
+#include "Niagara/Classes/NiagaraDataInterfaceParticleRead.h" 
+#include "Niagara/Classes/NiagaraSystem.h" 
+
+#include "Widgets/Input/SSuggestionTextBox.h"
+#include "DetailCategoryBuilder.h"
+#include "DetailLayoutBuilder.h"
+#include "DetailWidgetRow.h"
+
+#define LOCTEXT_NAMESPACE "NiagaraDataInterfaceParticleReadDetails"
+
+void FNiagaraDataInterfaceParticleReadDetails::CustomizeDetails(IDetailLayoutBuilder& DetailBuilder)
+{
+	FNiagaraDataInterfaceDetailsBase::CustomizeDetails(DetailBuilder);
+	 
+	TArray<TWeakObjectPtr<UObject>> SelectedObjects;
+	DetailBuilder.GetObjectsBeingCustomized(SelectedObjects);
+	if(SelectedObjects.Num() != 1 || SelectedObjects[0]->IsA<UNiagaraDataInterfaceParticleRead>() == false)
+	{
+		return;
+	}
+
+	UNiagaraDataInterfaceParticleRead* ReadDataInterface = CastChecked<UNiagaraDataInterfaceParticleRead>(SelectedObjects[0].Get());
+	ReadDataInterfaceWeak = ReadDataInterface;
+	NiagaraSystemWeak = ReadDataInterface->GetTypedOuter<UNiagaraSystem>();
+
+	static const FName ReadCategoryName = TEXT("ParticleRead");
+	IDetailCategoryBuilder& CategoryBuilder = DetailBuilder.EditCategory(ReadCategoryName);
+	
+	TArray<TSharedRef<IPropertyHandle>> Properties;
+	CategoryBuilder.GetDefaultProperties(Properties, true, true);
+	for (const TSharedRef<IPropertyHandle>& PropertyHandle : Properties)
+	{
+		FProperty* Property = PropertyHandle->GetProperty();
+		if (Property->GetFName() == GET_MEMBER_NAME_CHECKED(UNiagaraDataInterfaceParticleRead, EmitterName))
+		{
+			EmitterNameHandle = PropertyHandle;
+			DetailBuilder.HideProperty(PropertyHandle);
+
+			CategoryBuilder.AddCustomRow(FText::GetEmpty())
+				.NameContent()
+				[
+					PropertyHandle->CreatePropertyNameWidget()
+				]
+				.ValueContent()
+				[
+					SNew(SSuggestionTextBox)
+						.ForegroundColor(FSlateColor::UseForeground())
+						.Font(IDetailLayoutBuilder::GetDetailFont())
+						.Text(FText::FromString(ReadDataInterface->EmitterName))
+						.OnTextCommitted(this, &FNiagaraDataInterfaceParticleReadDetails::EmitterName_SetText)
+						.HintText(LOCTEXT("EmitterNameHint", "Enter Emitter Name..."))
+						.OnShowingSuggestions(this, &FNiagaraDataInterfaceParticleReadDetails::EmitterName_GetSuggestions)
+				];
+		}
+		else
+		{
+			CategoryBuilder.AddProperty(PropertyHandle);
+		}
+	}
+}
+
+TSharedRef<IDetailCustomization> FNiagaraDataInterfaceParticleReadDetails::MakeInstance()
+{
+	return MakeShared<FNiagaraDataInterfaceParticleReadDetails>();
+}
+
+void FNiagaraDataInterfaceParticleReadDetails::EmitterName_SetText(const FText& NewText, ETextCommit::Type CommitInfo)
+{
+	if (EmitterNameHandle)
+	{
+		EmitterNameHandle->SetValue(NewText.ToString());
+	}
+}
+
+void FNiagaraDataInterfaceParticleReadDetails::EmitterName_GetSuggestions(const FString& CurrText, TArray<FString>& OutSuggestions)
+{
+	if (UNiagaraSystem* NiagaraSystem = NiagaraSystemWeak.Get())
+	{
+		for (const FNiagaraEmitterHandle& EmitterHandle : NiagaraSystem->GetEmitterHandles())
+		{
+			if (EmitterHandle.GetIsEnabled())
+			{
+				OutSuggestions.Add(EmitterHandle.GetUniqueInstanceName());
+			}
+		}
+	}
+}
+
+#undef LOCTEXT_NAMESPACE
diff --git a/Engine/Plugins/FX/Niagara/Source/NiagaraEditorWidgets/Private/DetailCustomizations/NiagaraDataInterfaceParticleReadDetails.h b/Engine/Plugins/FX/Niagara/Source/NiagaraEditorWidgets/Private/DetailCustomizations/NiagaraDataInterfaceParticleReadDetails.h
new file mode 100644
--- /dev/null
+++ b/Engine/Plugins/FX/Niagara/Source/NiagaraEditorWidgets/Private/DetailCustomizations/NiagaraDataInterfaceParticleReadDetails.h
@@ -0,0 +1,32 @@
+// Copyright Epic Games, Inc. All Rights Reserved.
+
+#pragma once
+
+#include "NiagaraDataInterfaceDetails.h"
+#include "Types/SlateEnums.h"
+#include "Widgets/Input/SComboBox.h"
+
+class IDetailLayoutBuilder;
+class IDetailChildrenBuilder;
+class IPropertyHandle;
+class SWidget;
+class FName;
+class UNiagaraDataInterfaceParticleRead;
+class UNiagaraSystem;
+class SNiagaraNamePropertySelector;
+
+/** Details customization for Niagara particle read data interface. */
+class FNiagaraDataInterfaceParticleReadDetails : public FNiagaraDataInterfaceDetailsBase
+{
+public:
+	virtual void CustomizeDetails(IDetailLayoutBuilder& DetailBuilder) override;
+	static TSharedRef<IDetailCustomization> MakeInstance();
+
+private:
+	void EmitterName_SetText(const FText& NewText, ETextCommit::Type CommitInfo);
+	void EmitterName_GetSuggestions(const FString& CurrText, TArray<FString>& OutSuggestions);
+
+	TWeakObjectPtr<UNiagaraDataInterfaceParticleRead> ReadDataInterfaceWeak;
+	TWeakObjectPtr<UNiagaraSystem> NiagaraSystemWeak;
+	TSharedPtr<IPropertyHandle> EmitterNameHandle;
+};
diff --git a/Engine/Plugins/FX/Niagara/Source/NiagaraEditorWidgets/Private/NiagaraEditorWidgetsModule.cpp b/Engine/Plugins/FX/Niagara/Source/NiagaraEditorWidgets/Private/NiagaraEditorWidgetsModule.cpp
--- a/Engine/Plugins/FX/Niagara/Source/NiagaraEditorWidgets/Private/NiagaraEditorWidgetsModule.cpp
+++ b/Engine/Plugins/FX/Niagara/Source/NiagaraEditorWidgets/Private/NiagaraEditorWidgetsModule.cpp
@@ -7,6 +7,7 @@
 #include "DetailCustomizations/NiagaraDataInterfaceDetails.h"
 #include "DetailCustomizations/NiagaraDataInterfaceGrid2DCollectionDetails.h"
 #include "DetailCustomizations/NiagaraDataInterfaceGrid3DCollectionDetails.h"
+#include "DetailCustomizations/NiagaraDataInterfaceParticleReadDetails.h"
 #include "DetailCustomizations/NiagaraDataInterfaceSkeletalMeshDetails.h"
 #include "DetailCustomizations/NiagaraDataInterfaceStaticMeshDetails.h"
 #include "DetailCustomizations/NiagaraDataInterfaceMeshRendererInfoDetails.h"
@@ -122,6 +123,7 @@ void FNiagaraEditorWidgetsModule::StartupModule()
 	PropertyModule.RegisterCustomClassLayout("NiagaraDataInterfaceVectorCurve", FOnGetDetailCustomizationInstance::CreateStatic(&FNiagaraDataInterfaceVectorCurveDetails::MakeInstance));
 	PropertyModule.RegisterCustomClassLayout("NiagaraDataInterfaceVector4Curve", FOnGetDetailCustomizationInstance::CreateStatic(&FNiagaraDataInterfaceVector4CurveDetails::MakeInstance));
 	PropertyModule.RegisterCustomClassLayout("NiagaraDataInterfaceColorCurve", FOnGetDetailCustomizationInstance::CreateStatic(&FNiagaraDataInterfaceColorCurveDetails::MakeInstance));
+	PropertyModule.RegisterCustomClassLayout("NiagaraDataInterfaceParticleRead", FOnGetDetailCustomizationInstance::CreateStatic(&FNiagaraDataInterfaceParticleReadDetails::MakeInstance));
 	PropertyModule.RegisterCustomClassLayout("NiagaraDataInterfaceSkeletalMesh", FOnGetDetailCustomizationInstance::CreateStatic(&FNiagaraDataInterfaceSkeletalMeshDetails::MakeInstance));
 	PropertyModule.RegisterCustomClassLayout("NiagaraDataInterfaceStaticMesh", FOnGetDetailCustomizationInstance::CreateStatic(&FNiagaraDataInterfaceStaticMeshDetails::MakeInstance));
 	PropertyModule.RegisterCustomClassLayout("NiagaraDataInterfaceGrid2DCollection", FOnGetDetailCustomizationInstance::CreateStatic(&FNiagaraDataInterfaceGrid2DCollectionDetails::MakeInstance));
