Fixed default pin values on map get nodes not correctly updating when changing type

diff --git a/Engine/Plugins/FX/Niagara/Source/NiagaraEditor/Private/NiagaraGraph.cpp b/Engine/Plugins/FX/Niagara/Source/NiagaraEditor/Private/NiagaraGraph.cpp
--- a/Engine/Plugins/FX/Niagara/Source/NiagaraEditor/Private/NiagaraGraph.cpp
+++ b/Engine/Plugins/FX/Niagara/Source/NiagaraEditor/Private/NiagaraGraph.cpp
@@ -486,19 +486,19 @@ void UNiagaraGraph::ChangeParameterType(const FNiagaraVariable& CurrentParameter
 				Pin->Modify();
 				Pin->PinType = NiagaraSchema->TypeDefinitionToPinType(NewType);
 				Pin->PinType.PinSubCategory = UNiagaraNodeParameterMapBase::ParameterPinSubCategory;
-				NiagaraNode->MarkNodeRequiresSynchronization(__FUNCTION__, true);
-				break;
-			}
 
-			// fix default pin types if necessary
-			if (UNiagaraNodeParameterMapGet* GetNode = Cast<UNiagaraNodeParameterMapGet>(NiagaraNode))
-			{
-				UEdGraphPin* DefaultPin = GetNode->GetDefaultPin(Pin);
-				if (Pin->Direction == EGPD_Output && DefaultPin && DefaultPin->PinType != Pin->PinType)
+				// fix default pin types if necessary
+				if (UNiagaraNodeParameterMapGet* GetNode = Cast<UNiagaraNodeParameterMapGet>(NiagaraNode))
 				{
-					DefaultPin->PinType = Pin->PinType;
-					NiagaraNode->MarkNodeRequiresSynchronization(__FUNCTION__, true);
+					UEdGraphPin* DefaultPin = GetNode->GetDefaultPin(Pin);
+					if (Pin->Direction == EGPD_Output && DefaultPin && DefaultPin->PinType != Pin->PinType)
+					{
+						DefaultPin->PinType = Pin->PinType;
+					}
 				}
+
+				NiagaraNode->MarkNodeRequiresSynchronization(__FUNCTION__, true);
+				break;
 			}
 		}
 	}
