- Fix for not using ResolvedTextureRHI

#rnx

diff --git a/Engine/Plugins/FX/Niagara/Source/Niagara/Private/NiagaraDataInterface2DArrayTexture.cpp b/Engine/Plugins/FX/Niagara/Source/Niagara/Private/NiagaraDataInterface2DArrayTexture.cpp
--- a/Engine/Plugins/FX/Niagara/Source/Niagara/Private/NiagaraDataInterface2DArrayTexture.cpp
+++ b/Engine/Plugins/FX/Niagara/Source/Niagara/Private/NiagaraDataInterface2DArrayTexture.cpp
@@ -53,10 +53,10 @@ struct FNiagaraDataInterfaceProxyTexture2DArray : public FNiagaraDataInterfacePr
 					InstanceData->ResolvedTextureRHI = nullptr;
 				}
 			}
-			if (InstanceData->TextureReferenceRHI.IsValid())
+			if (InstanceData->ResolvedTextureRHI)
 			{
 				// Make sure the texture is readable, we don't know where it's coming from.
-				RHICmdList.Transition(FRHITransitionInfo(InstanceData->TextureReferenceRHI->GetReferencedTexture(), ERHIAccess::Unknown, ERHIAccess::SRVMask));
+				RHICmdList.Transition(FRHITransitionInfo(InstanceData->ResolvedTextureRHI, ERHIAccess::Unknown, ERHIAccess::SRVMask));
 			}
 		}
 	}
diff --git a/Engine/Plugins/FX/Niagara/Source/Niagara/Private/NiagaraDataInterfaceCubeTexture.cpp b/Engine/Plugins/FX/Niagara/Source/Niagara/Private/NiagaraDataInterfaceCubeTexture.cpp
--- a/Engine/Plugins/FX/Niagara/Source/Niagara/Private/NiagaraDataInterfaceCubeTexture.cpp
+++ b/Engine/Plugins/FX/Niagara/Source/Niagara/Private/NiagaraDataInterfaceCubeTexture.cpp
@@ -53,10 +53,10 @@ struct FNiagaraDataInterfaceProxyCubeTexture : public FNiagaraDataInterfaceProxy
 					InstanceData->ResolvedTextureRHI = nullptr;
 				}
 			}
-			if (InstanceData->TextureReferenceRHI.IsValid())
+			if (InstanceData->ResolvedTextureRHI)
 			{
 				// Make sure the texture is readable, we don't know where it's coming from.
-				RHICmdList.Transition(FRHITransitionInfo(InstanceData->TextureReferenceRHI->GetReferencedTexture(), ERHIAccess::Unknown, ERHIAccess::SRVMask));
+				RHICmdList.Transition(FRHITransitionInfo(InstanceData->ResolvedTextureRHI, ERHIAccess::Unknown, ERHIAccess::SRVMask));
 			}
 		}
 	}
diff --git a/Engine/Plugins/FX/Niagara/Source/Niagara/Private/NiagaraDataInterfaceTexture.cpp b/Engine/Plugins/FX/Niagara/Source/Niagara/Private/NiagaraDataInterfaceTexture.cpp
--- a/Engine/Plugins/FX/Niagara/Source/Niagara/Private/NiagaraDataInterfaceTexture.cpp
+++ b/Engine/Plugins/FX/Niagara/Source/Niagara/Private/NiagaraDataInterfaceTexture.cpp
@@ -60,10 +60,10 @@ struct FNiagaraDataInterfaceProxyTexture : public FNiagaraDataInterfaceProxy
 					InstanceData->ResolvedTextureRHI = nullptr;
 				}
 			}
-			if (InstanceData->TextureReferenceRHI.IsValid())
+			if (InstanceData->ResolvedTextureRHI.IsValid())
 			{
 				// Make sure the texture is readable, we don't know where it's coming from.
-				RHICmdList.Transition(FRHITransitionInfo(InstanceData->TextureReferenceRHI->GetReferencedTexture(), ERHIAccess::Unknown, ERHIAccess::SRVMask));
+				RHICmdList.Transition(FRHITransitionInfo(InstanceData->ResolvedTextureRHI, ERHIAccess::Unknown, ERHIAccess::SRVMask));
 			}
 		}
 	}
diff --git a/Engine/Plugins/FX/Niagara/Source/Niagara/Private/NiagaraDataInterfaceVolumeTexture.cpp b/Engine/Plugins/FX/Niagara/Source/Niagara/Private/NiagaraDataInterfaceVolumeTexture.cpp
--- a/Engine/Plugins/FX/Niagara/Source/Niagara/Private/NiagaraDataInterfaceVolumeTexture.cpp
+++ b/Engine/Plugins/FX/Niagara/Source/Niagara/Private/NiagaraDataInterfaceVolumeTexture.cpp
@@ -53,10 +53,10 @@ struct FNiagaraDataInterfaceProxyVolumeTexture : public FNiagaraDataInterfacePro
 					InstanceData->ResolvedTextureRHI = nullptr;
 				}
 			}
-			if (InstanceData->TextureReferenceRHI.IsValid())
+			if (InstanceData->ResolvedTextureRHI)
 			{
 				// Make sure the texture is readable, we don't know where it's coming from.
-				RHICmdList.Transition(FRHITransitionInfo(InstanceData->TextureReferenceRHI->GetReferencedTexture(), ERHIAccess::Unknown, ERHIAccess::SRVMask));
+				RHICmdList.Transition(FRHITransitionInfo(InstanceData->ResolvedTextureRHI, ERHIAccess::Unknown, ERHIAccess::SRVMask));
 			}
 		}
 	}
