Unavailable mesh distance field data causing issues with disappearing particles

diff --git a/Engine/Plugins/FX/Niagara/Source/Niagara/Private/NiagaraDataInterfaceRigidMeshCollisionQuery.cpp b/Engine/Plugins/FX/Niagara/Source/Niagara/Private/NiagaraDataInterfaceRigidMeshCollisionQuery.cpp
--- a/Engine/Plugins/FX/Niagara/Source/Niagara/Private/NiagaraDataInterfaceRigidMeshCollisionQuery.cpp
+++ b/Engine/Plugins/FX/Niagara/Source/Niagara/Private/NiagaraDataInterfaceRigidMeshCollisionQuery.cpp
@@ -613,11 +613,24 @@ public:
 
 				if (DistanceFieldSceneData == nullptr)
 				{
-					// UE_LOG(LogRigidMeshCollision, Error, TEXT("Distance fields are not available for use"));
-					// #todo(dmp): should we set something here in the case where distance field data is bound but we don't have it?
-					// there is no trivial constructor
+					UE_LOG(LogRigidMeshCollision, Error, TEXT("Distance fields are not available for use"));
+					// #todo(dmp): for now, we'll disable collisions when distance field data is not available
+					// There is no Dummy distance field data we can use.
+									
 					//FDistanceFieldSceneData DummyDistanceFieldSceneData(Context.Shader->GetShaderPlatform());
 					//DistanceFieldParameters.SetEmpty(RHICmdList, ComputeShaderRHI, DummyDistanceFieldSceneData);
+
+					SetSRVParameter(RHICmdList, ComputeShaderRHI, WorldTransformBuffer, FNiagaraRenderer::GetDummyFloatBuffer());
+					SetSRVParameter(RHICmdList, ComputeShaderRHI, InverseTransformBuffer, FNiagaraRenderer::GetDummyFloatBuffer());
+					SetSRVParameter(RHICmdList, ComputeShaderRHI, ElementExtentBuffer, FNiagaraRenderer::GetDummyFloatBuffer());
+					SetSRVParameter(RHICmdList, ComputeShaderRHI, PhysicsTypeBuffer, FNiagaraRenderer::GetDummyIntBuffer());
+					SetSRVParameter(RHICmdList, ComputeShaderRHI, DFIndexBuffer, FNiagaraRenderer::GetDummyIntBuffer());
+
+					static const FElementOffset DummyOffsets(0, 0, 0, 0);
+					SetShaderValue(RHICmdList, ComputeShaderRHI, MaxTransforms, 0);
+					SetShaderValue(RHICmdList, ComputeShaderRHI, CurrentOffset, 0);
+					SetShaderValue(RHICmdList, ComputeShaderRHI, PreviousOffset, 0);
+					SetShaderValue(RHICmdList, ComputeShaderRHI, ElementOffsets, DummyOffsets);
 				}
 				else 
 				{
