Disable GPU Scene for Niagara mesh particles in raytracing shaders.
Fix LWC bugs with Niagara mesh raytracing transforms.

diff --git a/Engine/Plugins/FX/Niagara/Shaders/Private/NiagaraGPURayTracingTransforms.usf b/Engine/Plugins/FX/Niagara/Shaders/Private/NiagaraGPURayTracingTransforms.usf
--- a/Engine/Plugins/FX/Niagara/Shaders/Private/NiagaraGPURayTracingTransforms.usf
+++ b/Engine/Plugins/FX/Niagara/Shaders/Private/NiagaraGPURayTracingTransforms.usf
@@ -16,6 +16,8 @@ uint	ParticleDataIntStride;
 uint	CPUNumInstances;
 uint	InstanceCountOffset;
 
+float3	SystemLWCTile;
+
 uint	PositionDataOffset;
 uint	RotationDataOffset;
 uint	ScaleDataOffset;
@@ -73,6 +75,10 @@ float4 SafeGetVec4(int RegisterIndex, uint InstanceID, float4 DefaultValue)
 #define THREADGROUP_SIZE 1
 #endif
 
+// TODO: This doesn't take a lot of features into account that are needed to solve for a Niagara mesh's transform.
+// See NiagaraCalculateMeshParticleTransforms in NiagaraMeshParticleUtils.ush for exactly what all goes into it.
+// In an upcoming update, this shader should be deprecated and instead, RayTracing should be getting transforms
+// from GPU Scene so this work doesn't have to be duplicated.
 [numthreads(THREADGROUP_SIZE,1,1)]
 void NiagaraGPURayTracingTransformsCS(uint3 DispatchThreadId : SV_DispatchThreadID)
 {
@@ -87,7 +93,9 @@ void NiagaraGPURayTracingTransformsCS(uint3 DispatchThreadId : SV_DispatchThread
 	float4x4 InstanceTransformTransposed = 0;
 	if (InstanceIndex < GPUNumInstances)
 	{
-		float3 InstancePosition = SafeGetVec3(PositionDataOffset, InstanceIndex, DefaultPosition);
+		float3 InstanceSimPosition = SafeGetVec3(PositionDataOffset, InstanceIndex, DefaultPosition);
+		// LWC_TODO: Precision loss
+		float3 InstancePosition = bLocalSpace ? InstanceSimPosition : LWCHackToFloat(MakeLWCVector3(SystemLWCTile, InstanceSimPosition));
 		float4 InstanceRotation = NiagaraQuatMul(SafeGetVec4(RotationDataOffset, InstanceIndex, DefaultRotation), MeshRotation);
 		float3 InstanceScale    = SafeGetVec3(ScaleDataOffset, InstanceIndex, DefaultScale) * MeshScale;
 
@@ -108,6 +116,10 @@ void NiagaraGPURayTracingTransformsCS(uint3 DispatchThreadId : SV_DispatchThread
 		{
 			bInstanceValid &= GetInt(MeshIndexOffset, InstanceIndex) == MeshIndexValue;
 		}
+		else
+		{
+			bInstanceValid &= MeshIndexValue == 0;
+		}
 
 		if ( bInstanceValid )
 		{
diff --git a/Engine/Plugins/FX/Niagara/Source/Niagara/Private/NiagaraRendererMeshes.cpp b/Engine/Plugins/FX/Niagara/Source/Niagara/Private/NiagaraRendererMeshes.cpp
--- a/Engine/Plugins/FX/Niagara/Source/Niagara/Private/NiagaraRendererMeshes.cpp
+++ b/Engine/Plugins/FX/Niagara/Source/Niagara/Private/NiagaraRendererMeshes.cpp
@@ -247,12 +247,13 @@ int32 FNiagaraRendererMeshes::GetLODIndex(int32 MeshIndex) const
 	return MeshData.RenderData->LODResources.IsValidIndex(LODIndex) ? LODIndex : INDEX_NONE;
 }
 
-void FNiagaraRendererMeshes::PrepareParticleMeshRenderData(FParticleMeshRenderData& ParticleMeshRenderData, const FSceneViewFamily& ViewFamily, FMeshElementCollector& Collector, FNiagaraDynamicDataBase* InDynamicData, const FNiagaraSceneProxy* SceneProxy) const
+void FNiagaraRendererMeshes::PrepareParticleMeshRenderData(FParticleMeshRenderData& ParticleMeshRenderData, const FSceneViewFamily& ViewFamily, FMeshElementCollector& Collector, FNiagaraDynamicDataBase* InDynamicData, const FNiagaraSceneProxy* SceneProxy, bool bRayTracing) const
 {
 	ParticleMeshRenderData.Collector = &Collector;
 
 #if NIAGARA_ENABLE_GPU_SCENE_MESHES
-	ParticleMeshRenderData.bUseGPUScene = FeatureLevel > ERHIFeatureLevel::ES3_1	
+	ParticleMeshRenderData.bUseGPUScene = !bRayTracing
+		&& FeatureLevel > ERHIFeatureLevel::ES3_1	
 		&& UseGPUScene(SceneProxy->GetScene().GetShaderPlatform(), FeatureLevel);
 #else
 	ParticleMeshRenderData.bUseGPUScene = false;
@@ -1154,7 +1155,7 @@ void FNiagaraRendererMeshes::GetDynamicMeshElements(const TArray<const FSceneVie
 	// Prepare our particle render data
 	// This will also determine if we have anything to render
 	FParticleMeshRenderData ParticleMeshRenderData;
-	PrepareParticleMeshRenderData(ParticleMeshRenderData, ViewFamily, Collector, DynamicDataRender, SceneProxy);
+	PrepareParticleMeshRenderData(ParticleMeshRenderData, ViewFamily, Collector, DynamicDataRender, SceneProxy, false);
 
 	if (ParticleMeshRenderData.SourceParticleData == nullptr )
 	{
@@ -1241,6 +1242,7 @@ void FNiagaraRendererMeshes::GetDynamicMeshElements(const TArray<const FSceneVie
 				VertexFactory.SetParticleFactoryType(NVFT_Mesh);
 				VertexFactory.SetMeshIndex(MeshIndex);
 				VertexFactory.SetLODIndex(LODIndex);
+				VertexFactory.EnablePrimitiveIDElement(ParticleMeshRenderData.bUseGPUScene);
 				VertexFactory.InitResource();
 				SetupVertexFactory(VertexFactory, LODModel);
 
@@ -1343,7 +1345,7 @@ void FNiagaraRendererMeshes::GetDynamicRayTracingInstances(FRayTracingMaterialGa
 	// Prepare our particle render data
 	// This will also determine if we have anything to render
 	FParticleMeshRenderData ParticleMeshRenderData;
-	PrepareParticleMeshRenderData(ParticleMeshRenderData, *View->Family, Context.RayTracingMeshResourceCollector, DynamicDataRender, SceneProxy);
+	PrepareParticleMeshRenderData(ParticleMeshRenderData, *View->Family, Context.RayTracingMeshResourceCollector, DynamicDataRender, SceneProxy, true);
 
 	if (ParticleMeshRenderData.SourceParticleData == nullptr)
 	{
@@ -1404,6 +1406,7 @@ void FNiagaraRendererMeshes::GetDynamicRayTracingInstances(FRayTracingMaterialGa
 		VertexFactory.SetParticleFactoryType(NVFT_Mesh);
 		VertexFactory.SetMeshIndex(MeshIndex);
 		VertexFactory.SetLODIndex(LODIndex);
+		VertexFactory.EnablePrimitiveIDElement(ParticleMeshRenderData.bUseGPUScene);
 		VertexFactory.InitResource();
 		SetupVertexFactory(VertexFactory, LODModel);
 
@@ -1489,7 +1492,7 @@ void FNiagaraRendererMeshes::GetDynamicRayTracingInstances(FRayTracingMaterialGa
 		const bool bUseLocalSpace = UseLocalSpace(SceneProxy);
 		if (SourceMode == ENiagaraRendererSourceDataMode::Emitter)
 		{
-			FVector3f Pos = bUseLocalSpace ? FVector3f() : (FVector3f)LocalTransform.GetOrigin();
+			FVector Pos = bUseLocalSpace ? FVector::ZeroVector : LocalTransform.GetOrigin();
 			FVector3f Scale{ 1.0f, 1.0f, 1.0f };
 			FQuat Rot = FQuat::Identity;
 
@@ -1499,7 +1502,14 @@ void FNiagaraRendererMeshes::GetDynamicRayTracingInstances(FRayTracingMaterialGa
 				if (VFBoundOffsetsInParamStore[ENiagaraMeshVFLayout::Type::Position] != INDEX_NONE
 					&& ParticleMeshRenderData.DynamicDataMesh->ParameterDataBound.IsValidIndex(VFBoundOffsetsInParamStore[ENiagaraMeshVFLayout::Type::Position]))
 				{
-					FMemory::Memcpy(&Pos, ParameterBoundData + VFBoundOffsetsInParamStore[ENiagaraMeshVFLayout::Type::Position], sizeof(FVector3f));
+					FVector3f TempPos;
+					FMemory::Memcpy(&TempPos, ParameterBoundData + VFBoundOffsetsInParamStore[ENiagaraMeshVFLayout::Type::Position], sizeof(FVector3f));
+					Pos = (FVector)TempPos;
+					if (!bUseLocalSpace)
+					{
+						// Handle LWC
+						Pos += FVector(SceneProxy->GetLWCRenderTile()) * FLargeWorldRenderScalar::GetTileSize();
+					}
 				}
 				if (VFBoundOffsetsInParamStore[ENiagaraMeshVFLayout::Type::Scale] != INDEX_NONE
 					&& ParticleMeshRenderData.DynamicDataMesh->ParameterDataBound.IsValidIndex(VFBoundOffsetsInParamStore[ENiagaraMeshVFLayout::Type::Scale]))
@@ -1618,7 +1628,13 @@ void FNiagaraRendererMeshes::GetDynamicRayTracingInstances(FRayTracingMaterialGa
 						continue;
 					}
 
-					const FVector InstancePosition = bHasPosition ? GetInstancePosition(InstanceIndex) : FVector::ZeroVector;
+					FVector InstancePosition = bHasPosition ? GetInstancePosition(InstanceIndex) : FVector::ZeroVector;
+					if (!bLocalSpace)
+					{
+						// Handle LWC
+						InstancePosition += FVector(SceneProxy->GetLWCRenderTile()) * FLargeWorldRenderScalar::GetTileSize();
+					}
+
 					const FQuat InstanceRotation = bHasRotation ? GetInstanceQuat(InstanceIndex).GetNormalized() * MeshRotation : MeshRotation;
 					FMatrix InstanceTransform = FQuatRotationTranslationMatrix::Make(InstanceRotation, InstancePosition);
 
@@ -1655,6 +1671,7 @@ void FNiagaraRendererMeshes::GetDynamicRayTracingInstances(FRayTracingMaterialGa
 					PassParameters->ParticleDataIntStride		= ParticleMeshRenderData.ParticleIntDataStride;
 					PassParameters->CPUNumInstances				= NumInstances;
 					PassParameters->InstanceCountOffset			= ParticleMeshRenderData.SourceParticleData->GetGPUInstanceCountBufferOffset();
+					PassParameters->SystemLWCTile				= SceneProxy->GetLWCRenderTile();
 					PassParameters->PositionDataOffset			= VFVariables[ENiagaraMeshVFLayout::Position].GetGPUOffset();
 					PassParameters->RotationDataOffset			= VFVariables[ENiagaraMeshVFLayout::Rotation].GetGPUOffset();
 					PassParameters->ScaleDataOffset				= VFVariables[ENiagaraMeshVFLayout::Scale].GetGPUOffset();
@@ -1664,7 +1681,7 @@ void FNiagaraRendererMeshes::GetDynamicRayTracingInstances(FRayTracingMaterialGa
 					PassParameters->RenderVisibilityValue		= RendererVisibility;
 					PassParameters->MeshIndexValue				= MeshData.SourceMeshIndex;
 					PassParameters->LocalTransform				= FMatrix44f(LocalTransform);		// LWC_TODO: Precision loss
-					PassParameters->DefaultPosition				= FVector3f::ZeroVector;
+					PassParameters->DefaultPosition				= CommonParameters.DefaultPosition;
 					PassParameters->DefaultRotation				= FVector4f(0.0f, 0.0f, 0.0f, 1.0f);
 					PassParameters->DefaultScale				= FVector3f(1.0f, 1.0f, 1.0f);
 					PassParameters->MeshScale					= MeshData.Scale;
diff --git a/Engine/Plugins/FX/Niagara/Source/Niagara/Public/NiagaraRendererMeshes.h b/Engine/Plugins/FX/Niagara/Source/Niagara/Public/NiagaraRendererMeshes.h
--- a/Engine/Plugins/FX/Niagara/Source/Niagara/Public/NiagaraRendererMeshes.h
+++ b/Engine/Plugins/FX/Niagara/Source/Niagara/Public/NiagaraRendererMeshes.h
@@ -132,7 +132,7 @@ protected:
 
 	int32 GetLODIndex(int32 MeshIndex) const;
 
-	void PrepareParticleMeshRenderData(FParticleMeshRenderData& ParticleMeshRenderData, const FSceneViewFamily& ViewFamily, FMeshElementCollector& Collector, FNiagaraDynamicDataBase* InDynamicData, const FNiagaraSceneProxy* SceneProxy) const;
+	void PrepareParticleMeshRenderData(FParticleMeshRenderData& ParticleMeshRenderData, const FSceneViewFamily& ViewFamily, FMeshElementCollector& Collector, FNiagaraDynamicDataBase* InDynamicData, const FNiagaraSceneProxy* SceneProxy, bool bRayTracing) const;
 	void PrepareParticleRenderBuffers(FParticleMeshRenderData& ParticleMeshRenderData, FGlobalDynamicReadBuffer& DynamicReadBuffer) const;
 	void InitializeSortInfo(const FParticleMeshRenderData& ParticleMeshRenderData, const FNiagaraSceneProxy& SceneProxy, const FSceneView& View, int32 ViewIndex, bool bIsInstancedStereo, FNiagaraGPUSortInfo& OutSortInfo) const;
 	void PreparePerMeshData(FParticleMeshRenderData& ParticleMeshRenderData, const FNiagaraMeshVertexFactory& VertexFactory, const FNiagaraSceneProxy& SceneProxy, const FMeshData& MeshData) const;
diff --git a/Engine/Plugins/FX/Niagara/Source/NiagaraVertexFactories/Private/NiagaraMeshVertexFactory.cpp b/Engine/Plugins/FX/Niagara/Source/NiagaraVertexFactories/Private/NiagaraMeshVertexFactory.cpp
--- a/Engine/Plugins/FX/Niagara/Source/NiagaraVertexFactories/Private/NiagaraMeshVertexFactory.cpp
+++ b/Engine/Plugins/FX/Niagara/Source/NiagaraVertexFactories/Private/NiagaraMeshVertexFactory.cpp
@@ -126,8 +126,11 @@ void FNiagaraMeshVertexFactory::InitRHI()
 		}
 
 #if NIAGARA_ENABLE_GPU_SCENE_MESHES
-		// TODO: Support GPU Scene on mobile? Maybe only for CPU particles?
-		AddPrimitiveIdStreamElement(EVertexInputStreamType::Default, Elements, 13, 0xFF);
+		if (bAddPrimitiveIDElement)
+		{
+			// TODO: Support GPU Scene on mobile? Maybe only for CPU particles?
+			AddPrimitiveIdStreamElement(EVertexInputStreamType::Default, Elements, 13, 0xFF);
+		}
 #endif
 
 		//if (Streams.Num() > 0)
@@ -161,8 +164,17 @@ void FNiagaraMeshVertexFactory::ModifyCompilationEnvironment(const FVertexFactor
 	const bool bUseGPUScene = UseGPUScene(Parameters.Platform, MaxSupportedFeatureLevel) && MaxSupportedFeatureLevel > ERHIFeatureLevel::ES3_1;
 	const bool bSupportsPrimitiveIdStream = Parameters.VertexFactoryType->SupportsPrimitiveIdStream();
 	
-	OutEnvironment.SetDefine(TEXT("VF_SUPPORTS_PRIMITIVE_SCENE_DATA"), bSupportsPrimitiveIdStream && bUseGPUScene);
-	OutEnvironment.SetDefine(TEXT("VF_REQUIRES_PER_INSTANCE_CUSTOM_DATA"), bSupportsPrimitiveIdStream && bUseGPUScene);
+	// TODO: Support GPU Scene for raytracing
+	if (bSupportsPrimitiveIdStream && bUseGPUScene)
+	{
+		OutEnvironment.SetDefine(TEXT("VF_SUPPORTS_PRIMITIVE_SCENE_DATA"), TEXT("!(RAYHITGROUPSHADER)"));
+		OutEnvironment.SetDefine(TEXT("VF_REQUIRES_PER_INSTANCE_CUSTOM_DATA"), TEXT("!(RAYHITGROUPSHADER)"));
+	}
+	else
+	{
+		OutEnvironment.SetDefine(TEXT("VF_SUPPORTS_PRIMITIVE_SCENE_DATA"), 0);
+		OutEnvironment.SetDefine(TEXT("VF_REQUIRES_PER_INSTANCE_CUSTOM_DATA"), 0);
+	}
 #endif
 
 	const bool ContainsManualVertexFetch = OutEnvironment.GetDefinitions().Contains("MANUAL_VERTEX_FETCH");
diff --git a/Engine/Plugins/FX/Niagara/Source/NiagaraVertexFactories/Public/NiagaraGPURayTracingTransformsShader.h b/Engine/Plugins/FX/Niagara/Source/NiagaraVertexFactories/Public/NiagaraGPURayTracingTransformsShader.h
--- a/Engine/Plugins/FX/Niagara/Source/NiagaraVertexFactories/Public/NiagaraGPURayTracingTransformsShader.h
+++ b/Engine/Plugins/FX/Niagara/Source/NiagaraVertexFactories/Public/NiagaraGPURayTracingTransformsShader.h
@@ -33,6 +33,8 @@ class NIAGARAVERTEXFACTORIES_API FNiagaraGPURayTracingTransformsCS : public FGlo
 		SHADER_PARAMETER(uint32,	CPUNumInstances)
 		SHADER_PARAMETER(uint32,	InstanceCountOffset)
 
+		SHADER_PARAMETER(FVector3f,	SystemLWCTile)
+
 		SHADER_PARAMETER(uint32,	PositionDataOffset)
 		SHADER_PARAMETER(uint32,	RotationDataOffset)
 		SHADER_PARAMETER(uint32,	ScaleDataOffset)
diff --git a/Engine/Plugins/FX/Niagara/Source/NiagaraVertexFactories/Public/NiagaraMeshVertexFactory.h b/Engine/Plugins/FX/Niagara/Source/NiagaraVertexFactories/Public/NiagaraMeshVertexFactory.h
--- a/Engine/Plugins/FX/Niagara/Source/NiagaraVertexFactories/Public/NiagaraMeshVertexFactory.h
+++ b/Engine/Plugins/FX/Niagara/Source/NiagaraVertexFactories/Public/NiagaraMeshVertexFactory.h
@@ -72,7 +72,7 @@ BEGIN_SHADER_PARAMETER_STRUCT(FNiagaraMeshCommonParameters, NIAGARAVERTEXFACTORI
 	SHADER_PARAMETER(FVector4f, DefaultPrevRotation)
 	SHADER_PARAMETER(FVector3f, DefaultPrevPosition)
 	SHADER_PARAMETER(FVector3f, DefaultPrevVelocity)
-	SHADER_PARAMETER(float, DefaultPrevCameraOffset)	
+	SHADER_PARAMETER(float, DefaultPrevCameraOffset)
 END_SHADER_PARAMETER_STRUCT()
 
 /**
@@ -130,6 +130,7 @@ public:
 		: FNiagaraVertexFactoryBase(InType, InFeatureLevel)
 		, MeshIndex(-1)
 		, LODIndex(-1)
+		, bAddPrimitiveIDElement(true)
 		, InstanceVerticesCPU(nullptr)
 		, FloatDataStride(0)
 	{}
@@ -138,6 +139,7 @@ public:
 		: FNiagaraVertexFactoryBase(NVFT_MAX, ERHIFeatureLevel::Num)
 		, MeshIndex(-1)
 		, LODIndex(-1)
+		, bAddPrimitiveIDElement(true)
 		, InstanceVerticesCPU(nullptr)
 		, FloatDataStride(0)
 	{}
@@ -200,12 +202,6 @@ public:
 	{
 		return Data.NumTexCoords;
 	}
-	
-	/**
-	* Copy the data from another vertex factory
-	* @param Other - factory to copy from
-	*/
-	void Copy(const FNiagaraMeshVertexFactory& Other);
 
 	// FRenderResource interface.
 	virtual void InitRHI() override;
@@ -215,11 +211,15 @@ public:
 
 	int32 GetLODIndex() const { return LODIndex; }
 	void SetLODIndex(int32 InLODIndex) { LODIndex = InLODIndex; }
+
+	bool IsPrimitiveIDElementEnabled() const { return bAddPrimitiveIDElement; }
+	void EnablePrimitiveIDElement(bool bEnable) { bAddPrimitiveIDElement = bEnable; }
 	
 protected:
 	FStaticMeshDataType Data;
 	int32 MeshIndex;	
-	int32 LODIndex;	
+	int32 LODIndex;
+	bool bAddPrimitiveIDElement;
 
 	/** Uniform buffer with mesh particle parameters. */
 	FRHIUniformBuffer* MeshParticleUniformBuffer;
