Fixed thumbnails not copying over for systems that had their thumbnail recorded via content browser context menu, as those thumbnails get saved in the UPackage which isn't copied over in our asset wizard.

diff --git a/Engine/Plugins/FX/Niagara/Source/NiagaraEditor/Private/NiagaraSystemFactoryNew.cpp b/Engine/Plugins/FX/Niagara/Source/NiagaraEditor/Private/NiagaraSystemFactoryNew.cpp
--- a/Engine/Plugins/FX/Niagara/Source/NiagaraEditor/Private/NiagaraSystemFactoryNew.cpp
+++ b/Engine/Plugins/FX/Niagara/Source/NiagaraEditor/Private/NiagaraSystemFactoryNew.cpp
@@ -8,6 +8,8 @@
 #include "NiagaraGraph.h"
 #include "NiagaraEditorSettings.h"
 #include "AssetData.h"
+#include "ContentBrowserModule.h"
+#include "ImageUtils.h"
 #include "ViewModels/Stack/NiagaraStackGraphUtilities.h"
 #include "SNewSystemDialog.h"
 #include "Misc/MessageDialog.h"
@@ -17,6 +19,7 @@
 #include "Framework/Application/SlateApplication.h"
 #include "NiagaraEditorUtilities.h"
 #include "NiagaraSettings.h"
+#include "ObjectTools.h"
 
 #define LOCTEXT_NAMESPACE "NiagaraSystemFactory"
 
@@ -126,6 +129,28 @@ UObject* UNiagaraSystemFactoryNew::FactoryCreateNew(UClass* Class, UObject* InPa
 		NewSystem->TemplateSpecification = ENiagaraScriptTemplateSpecification::None;
 		NewSystem->TemplateAssetDescription = FText();
 		NewSystem->Category = FText();
+
+		// if the new system doesn't have a thumbnail image, check the thumbnail map of the original asset's upackage
+		if(NewSystem->ThumbnailImage == nullptr)
+		{
+			FString ObjectFullName = SystemToCopy->GetFullName();
+			FName ObjectName = FName(ObjectFullName);
+			FString PackageFullName;
+			ThumbnailTools::QueryPackageFileNameForObject(ObjectFullName, PackageFullName);
+			FThumbnailMap ThumbnailMap;
+			ThumbnailTools::ConditionallyLoadThumbnailsFromPackage(PackageFullName, {FName(ObjectFullName)}, ThumbnailMap);
+
+			// there should always be a dummy thumbnail in here
+			if(ThumbnailMap.Contains(ObjectName))
+			{
+				FObjectThumbnail Thumbnail = ThumbnailMap[ObjectName];
+				// we only want to copy the thumbnail over if it's not a dummy
+				if(Thumbnail.GetImageWidth() != 0 && Thumbnail.GetImageHeight() != 0)
+				{
+					ThumbnailTools::CacheThumbnail(NewSystem->GetFullName(), &Thumbnail, NewSystem->GetOutermost());
+				}
+			}
+		}
 	}
 	else if (EmittersToAddToNewSystem.Num() > 0)
 	{
