Adding a "loose type" mode to Niagara to allow users to treat the position type like a vector

diff --git a/Engine/Plugins/FX/Niagara/Source/Niagara/Public/NiagaraSettings.h b/Engine/Plugins/FX/Niagara/Source/Niagara/Public/NiagaraSettings.h
--- a/Engine/Plugins/FX/Niagara/Source/Niagara/Public/NiagaraSettings.h
+++ b/Engine/Plugins/FX/Niagara/Source/Niagara/Public/NiagaraSettings.h
@@ -101,6 +101,13 @@ class NIAGARA_API UNiagaraSettings : public UDeveloperSettings
 	UPROPERTY(config, EditAnywhere, Category = Niagara, meta = ( ConfigRestartRequired = true ))
 	bool bSystemsSupportLargeWorldCoordinates = true;
 
+	/**
+	 If set to true, types like positions and vectors cannot be assigned to each other without an explicit conversion step.
+	 If false, type checks are loosened and some types can be implicitly converted into each other.
+	 It is recommended to not disable this when working with large world coordinates. */
+	UPROPERTY(config, EditAnywhere, Category = Niagara, meta = ( DisplayName = "Enforce strict type checks in the graph" ))
+	bool bEnforceStrictStackTypes = true;
+
 	/** Default effect type to use for effects that don't define their own. Can be null. */
 	UPROPERTY(config, EditAnywhere, Category = Niagara, meta = (AllowedClasses = "NiagaraEffectType"))
 	FSoftObjectPath DefaultEffectType;
diff --git a/Engine/Plugins/FX/Niagara/Source/NiagaraEditor/Private/EdGraphSchema_Niagara.cpp b/Engine/Plugins/FX/Niagara/Source/NiagaraEditor/Private/EdGraphSchema_Niagara.cpp
--- a/Engine/Plugins/FX/Niagara/Source/NiagaraEditor/Private/EdGraphSchema_Niagara.cpp
+++ b/Engine/Plugins/FX/Niagara/Source/NiagaraEditor/Private/EdGraphSchema_Niagara.cpp
@@ -1011,13 +1011,18 @@ const FPinConnectionResponse UEdGraphSchema_Niagara::CanCreateConnection(const U
 				{
 					return FPinConnectionResponse(CONNECT_RESPONSE_DISALLOW, TEXT("Types are not compatible"));
 				}
-				else if (FNiagaraTypeDefinition::IsLossyConversion(PinToTypeDefinition(InputPin), PinToTypeDefinition(OutputPin)))
+				else if (FNiagaraTypeDefinition::IsLossyConversion(PinTypeInput, PinTypeOutput))
 				{
-					return FPinConnectionResponse(CONNECT_RESPONSE_MAKE_WITH_CONVERSION_NODE, FString::Printf(TEXT("Convert %s to %s. Be aware that this can be a lossy conversion."), *(PinToTypeDefinition(OutputPin).GetNameText().ToString()), *(PinToTypeDefinition(InputPin).GetNameText().ToString())));
+					return FPinConnectionResponse(CONNECT_RESPONSE_MAKE_WITH_CONVERSION_NODE, FString::Printf(TEXT("Convert %s to %s. Be aware that this can be a lossy conversion."), *(PinTypeOutput.GetNameText().ToString()), *(PinTypeInput.GetNameText().ToString())));
 				}
 				else
 				{
-					return FPinConnectionResponse(CONNECT_RESPONSE_MAKE_WITH_CONVERSION_NODE, FString::Printf(TEXT("Convert %s to %s"), *(PinToTypeDefinition(OutputPin).GetNameText().ToString()), *(PinToTypeDefinition(InputPin).GetNameText().ToString())));
+					const UNiagaraSettings* Settings = GetDefault<UNiagaraSettings>();
+					if (Settings->bEnforceStrictStackTypes == false && ((PinTypeInput == FNiagaraTypeDefinition::GetPositionDef() && PinTypeOutput == FNiagaraTypeDefinition::GetVec3Def()) || (PinTypeOutput == FNiagaraTypeDefinition::GetPositionDef() && PinTypeInput == FNiagaraTypeDefinition::GetVec3Def())))
+					{
+						return FPinConnectionResponse(CONNECT_RESPONSE_MAKE, FString());
+					}
+					return FPinConnectionResponse(CONNECT_RESPONSE_MAKE_WITH_CONVERSION_NODE, FString::Printf(TEXT("Convert %s to %s"), *(PinTypeOutput.GetNameText().ToString()), *(PinTypeInput.GetNameText().ToString())));
 				}
 			}
 		}
diff --git a/Engine/Plugins/FX/Niagara/Source/NiagaraEditor/Private/NiagaraEditorUtilities.cpp b/Engine/Plugins/FX/Niagara/Source/NiagaraEditor/Private/NiagaraEditorUtilities.cpp
--- a/Engine/Plugins/FX/Niagara/Source/NiagaraEditor/Private/NiagaraEditorUtilities.cpp
+++ b/Engine/Plugins/FX/Niagara/Source/NiagaraEditor/Private/NiagaraEditorUtilities.cpp
@@ -29,6 +29,7 @@
 #include "NiagaraParameterDefinitions.h"
 #include "NiagaraScript.h"
 #include "NiagaraScriptSource.h"
+#include "NiagaraSettings.h"
 #include "NiagaraSimulationStageBase.h"
 #include "NiagaraStackEditorData.h"
 #include "ViewModels/NiagaraOverviewGraphViewModel.h"
@@ -993,6 +994,18 @@ FText FNiagaraEditorUtilities::GetTypeDefinitionCategory(const FNiagaraTypeDefin
 	return Category;
 }
 
+bool FNiagaraEditorUtilities::AreTypesAssignable(const FNiagaraTypeDefinition& TypeA, const FNiagaraTypeDefinition& TypeB)
+{
+	const UNiagaraSettings* Settings = GetDefault<UNiagaraSettings>();
+	if (Settings->bEnforceStrictStackTypes)
+	{
+		return TypeA == TypeB;
+	}
+	return (TypeA == TypeB)
+		|| (TypeA == FNiagaraTypeDefinition::GetPositionDef() && TypeB == FNiagaraTypeDefinition::GetVec3Def())
+		|| (TypeB == FNiagaraTypeDefinition::GetPositionDef() && TypeA == FNiagaraTypeDefinition::GetVec3Def());
+}
+
 void FNiagaraEditorUtilities::MarkDependentCompilableAssetsDirty(TArray<UObject*> InObjects)
 {
 	const FText LoadAndMarkDirtyDisplayName = NSLOCTEXT("NiagaraEditor", "MarkDependentAssetsDirtySlowTask", "Loading and marking dependent assets dirty.");
diff --git a/Engine/Plugins/FX/Niagara/Source/NiagaraEditor/Private/NiagaraHlslTranslator.cpp b/Engine/Plugins/FX/Niagara/Source/NiagaraEditor/Private/NiagaraHlslTranslator.cpp
--- a/Engine/Plugins/FX/Niagara/Source/NiagaraEditor/Private/NiagaraHlslTranslator.cpp
+++ b/Engine/Plugins/FX/Niagara/Source/NiagaraEditor/Private/NiagaraHlslTranslator.cpp
@@ -10,7 +10,6 @@
 #include "NiagaraComponent.h"
 #include "NiagaraConstants.h"
 #include "NiagaraDataInterface.h"
-#include "NiagaraDataInterfaceCurve.h"
 #include "NiagaraDataInterfaceVector2DCurve.h"
 #include "NiagaraEditorModule.h"
 #include "NiagaraEditorSettings.h"
@@ -35,6 +34,7 @@
 #include "NiagaraParameterCollection.h"
 #include "NiagaraScriptSource.h"
 #include "NiagaraScriptVariable.h"
+#include "NiagaraSettings.h"
 #include "NiagaraShared.h"
 #include "NiagaraSimulationStageBase.h"
 #include "NiagaraTrace.h"
@@ -570,6 +570,8 @@ FHlslNiagaraTranslator::FHlslNiagaraTranslator()
 	, ActiveStageIdx(-1)
 	, bInitializedDefaults(false)
 {
+	const UNiagaraSettings* Settings = GetDefault<UNiagaraSettings>();
+	bEnforceStrictTypesValidations = Settings->bEnforceStrictStackTypes;
 }
 
 
@@ -6245,7 +6247,10 @@ void FHlslNiagaraTranslator::HandleParameterRead(int32 ParamMapHistoryIdx, const
 				if ((ExistingVar->GetType() == FNiagaraTypeDefinition::GetVec3Def() && Var.GetType() == FNiagaraTypeDefinition::GetPositionDef())
 					|| (ExistingVar->GetType() == FNiagaraTypeDefinition::GetPositionDef() && Var.GetType() == FNiagaraTypeDefinition::GetVec3Def()))
 				{
-					Warning(FText::Format(LOCTEXT("MismatchedPositionTypes", "Variable {0} was defined both as position and vector, please check your modules and linked values for compatibility."), FText::FromName(Var.GetName())), ErrorNode, nullptr);
+					if (bEnforceStrictTypesValidations)
+					{
+						Warning(FText::Format(LOCTEXT("MismatchedPositionTypes", "Variable {0} was defined both as position and vector, please check your modules and linked values for compatibility."), FText::FromName(Var.GetName())), ErrorNode, nullptr);
+					}
 				}
 				else
 				{
@@ -6704,7 +6709,7 @@ void FHlslNiagaraTranslator::Operation(class UNiagaraNodeOp* Operation, TArray<i
 	}
 
 	FText ValidationError;
-	if (OpInfo && OpInfo->InputTypeValidationFunction.IsBound() && OpInfo->InputTypeValidationFunction.Execute(InputTypes, ValidationError) == false)
+	if (bEnforceStrictTypesValidations && OpInfo && OpInfo->InputTypeValidationFunction.IsBound() && OpInfo->InputTypeValidationFunction.Execute(InputTypes, ValidationError) == false)
 	{
 		Warning(ValidationError, Operation, OutputPins[0], TEXT("Invalid op types"), true);
 	}
diff --git a/Engine/Plugins/FX/Niagara/Source/NiagaraEditor/Private/NiagaraHlslTranslator.h b/Engine/Plugins/FX/Niagara/Source/NiagaraEditor/Private/NiagaraHlslTranslator.h
--- a/Engine/Plugins/FX/Niagara/Source/NiagaraEditor/Private/NiagaraHlslTranslator.h
+++ b/Engine/Plugins/FX/Niagara/Source/NiagaraEditor/Private/NiagaraHlslTranslator.h
@@ -854,6 +854,7 @@ private:
 	TArray<FHlslNiagaraTranslationStage> TranslationStages;
 	int32 ActiveStageIdx;
 	bool bInitializedDefaults;
+	bool bEnforceStrictTypesValidations;
 
 	TArray<const UEdGraphPin*> CurrentDefaultPinTraversal;
 	// Variables that need to be initialized based on some other variable's value at the end of spawn.
diff --git a/Engine/Plugins/FX/Niagara/Source/NiagaraEditor/Private/NiagaraNodeFunctionCall.cpp b/Engine/Plugins/FX/Niagara/Source/NiagaraEditor/Private/NiagaraNodeFunctionCall.cpp
--- a/Engine/Plugins/FX/Niagara/Source/NiagaraEditor/Private/NiagaraNodeFunctionCall.cpp
+++ b/Engine/Plugins/FX/Niagara/Source/NiagaraEditor/Private/NiagaraNodeFunctionCall.cpp
@@ -27,6 +27,7 @@
 #include "ViewModels/Stack/NiagaraParameterHandle.h"
 #include "ViewModels/Stack/NiagaraStackGraphUtilities.h"
 #include "NiagaraNodeStaticSwitch.h"
+#include "NiagaraSettings.h"
 #include "ScopedTransaction.h"
 
 
@@ -406,6 +407,7 @@ bool UNiagaraNodeFunctionCall::FixupPinNames()
 			return false;
 		}
 		const UEdGraphSchema_Niagara* NiagaraSchema = Graph->GetNiagaraSchema();
+		const UNiagaraSettings* Settings = GetDefault<UNiagaraSettings>();
 		
 		// find out which inputs the module offers
 		TSet<const UEdGraphPin*> HiddenModulePins;
@@ -463,7 +465,7 @@ bool UNiagaraNodeFunctionCall::FixupPinNames()
 				}
 
 				// when changing the module input type to position, changes the linked value pin as well
-				if (InputVar.GetType() == FNiagaraTypeDefinition::GetPositionDef() && OverridePin->PinType != LinkedValuePin->PinType)
+				if (InputVar.GetType() == FNiagaraTypeDefinition::GetPositionDef() && OverridePin->PinType != LinkedValuePin->PinType && Settings->bEnforceStrictStackTypes)
 				{
 					LinkedValuePin->PinType = OverridePin->PinType;
 					bNodeChanged = true;
diff --git a/Engine/Plugins/FX/Niagara/Source/NiagaraEditor/Private/ViewModels/Stack/NiagaraStackFunctionInput.cpp b/Engine/Plugins/FX/Niagara/Source/NiagaraEditor/Private/ViewModels/Stack/NiagaraStackFunctionInput.cpp
--- a/Engine/Plugins/FX/Niagara/Source/NiagaraEditor/Private/ViewModels/Stack/NiagaraStackFunctionInput.cpp
+++ b/Engine/Plugins/FX/Niagara/Source/NiagaraEditor/Private/ViewModels/Stack/NiagaraStackFunctionInput.cpp
@@ -30,6 +30,7 @@
 #include "NiagaraScriptMergeManager.h"
 #include "NiagaraScriptSource.h"
 #include "NiagaraScriptVariable.h"
+#include "NiagaraSettings.h"
 #include "NiagaraSystemScriptViewModel.h"
 #include "ScopedTransaction.h"
 #include "EdGraph/EdGraphPin.h"
@@ -1351,7 +1352,7 @@ void UNiagaraStackFunctionInput::GetAvailableParameterHandles(TArray<FNiagaraPar
 	// Engine Handles.
 	for (const FNiagaraVariable& SystemVariable : FNiagaraConstants::GetEngineConstants())
 	{
-		if (SystemVariable.GetType() == InputType)
+		if (FNiagaraEditorUtilities::AreTypesAssignable(SystemVariable.GetType(), InputType))
 		{
 			AvailableParameterHandles.Add(FNiagaraParameterHandle::CreateEngineParameterHandle(SystemVariable));
 		}
@@ -1366,7 +1367,7 @@ void UNiagaraStackFunctionInput::GetAvailableParameterHandles(TArray<FNiagaraPar
 	GetSystemViewModel()->GetSystem().GetExposedParameters().GetParameters(ExposedVars);
 	for (const FNiagaraVariable& ExposedVar : ExposedVars)
 	{
-		if (ExposedVar.GetType() == InputType)
+		if (FNiagaraEditorUtilities::AreTypesAssignable(ExposedVar.GetType(), InputType))
 		{
 			AvailableParameterHandles.Add(FNiagaraParameterHandle::CreateEngineParameterHandle(ExposedVar));
 		}
@@ -1445,7 +1446,7 @@ void UNiagaraStackFunctionInput::GetAvailableParameterHandles(TArray<FNiagaraPar
 
 						if (bWritten)
 						{
-							if (HistoryVariable.GetType() == InputType)
+							if (FNiagaraEditorUtilities::AreTypesAssignable(HistoryVariable.GetType(), InputType))
 							{
 								AvailableParameterHandles.AddUnique(AvailableHandle);
 								AvailableParameterHandlesForThisOutput.AddUnique(AvailableHandle);
@@ -1501,7 +1502,7 @@ void UNiagaraStackFunctionInput::GetAvailableParameterHandles(TArray<FNiagaraPar
 		{
 			for (const FNiagaraVariable& CollectionParam : Collection->GetParameters())
 			{
-				if (CollectionParam.GetType() == InputType)
+				if (FNiagaraEditorUtilities::AreTypesAssignable(CollectionParam.GetType(), InputType))
 				{
 					AvailableParameterHandles.AddUnique(FNiagaraParameterHandle(CollectionParam.GetName()));
 				}
diff --git a/Engine/Plugins/FX/Niagara/Source/NiagaraEditor/Private/ViewModels/Stack/NiagaraStackGraphUtilities.cpp b/Engine/Plugins/FX/Niagara/Source/NiagaraEditor/Private/ViewModels/Stack/NiagaraStackGraphUtilities.cpp
--- a/Engine/Plugins/FX/Niagara/Source/NiagaraEditor/Private/ViewModels/Stack/NiagaraStackGraphUtilities.cpp
+++ b/Engine/Plugins/FX/Niagara/Source/NiagaraEditor/Private/ViewModels/Stack/NiagaraStackGraphUtilities.cpp
@@ -42,6 +42,7 @@
 #include "INiagaraEditorTypeUtilities.h"
 #include "NiagaraEditorData.h"
 #include "NiagaraGraphDataCache.h"
+#include "NiagaraSettings.h"
 
 
 DECLARE_CYCLE_STAT(TEXT("Niagara - StackGraphUtilities - RelayoutGraph"), STAT_NiagaraEditor_StackGraphUtilities_RelayoutGraph, STATGROUP_NiagaraEditor);
@@ -1369,6 +1370,7 @@ UEdGraphPin* FNiagaraStackGraphUtilities::GetLinkedValueHandleForFunctionInput(c
 void FNiagaraStackGraphUtilities::SetLinkedValueHandleForFunctionInput(UEdGraphPin& OverridePin, FNiagaraParameterHandle LinkedParameterHandle, ENiagaraDefaultMode DesiredDefaultMode, const FGuid& NewNodePersistentId)
 {
 	checkf(OverridePin.LinkedTo.Num() == 0, TEXT("Can't set a linked value handle when the override pin already has a value."));
+	const UNiagaraSettings* Settings = GetDefault<UNiagaraSettings>();
 
 	UNiagaraNodeParameterMapSet* OverrideNode = CastChecked<UNiagaraNodeParameterMapSet>(OverridePin.GetOwningNode());
 	UNiagaraGraph* Graph = OverrideNode->GetNiagaraGraph();
@@ -1387,11 +1389,25 @@ void FNiagaraStackGraphUtilities::SetLinkedValueHandleForFunctionInput(UEdGraphP
 
 	const UEdGraphSchema_Niagara* NiagaraSchema = GetDefault<UEdGraphSchema_Niagara>();
 	FNiagaraTypeDefinition InputType = NiagaraSchema->PinToTypeDefinition(&OverridePin);
-	UEdGraphPin* GetOutputPin = GetNode->RequestNewTypedPin(EGPD_Output, InputType, LinkedParameterHandle.GetParameterHandleString());
+
+	FNiagaraVariable ParameterToRead(InputType, LinkedParameterHandle.GetParameterHandleString());
+	if (Settings->bEnforceStrictStackTypes == false && (InputType == FNiagaraTypeDefinition::GetVec3Def() || InputType == FNiagaraTypeDefinition::GetPositionDef()))
+	{
+		// see if the requested input exists or if we can use one of the equivalent loose types
+		FNiagaraTypeDefinition AlternateType = InputType == FNiagaraTypeDefinition::GetVec3Def() ? FNiagaraTypeDefinition::GetPositionDef() : FNiagaraTypeDefinition::GetVec3Def();
+		FNiagaraVariable ParameterAlternative(AlternateType, LinkedParameterHandle.GetParameterHandleString());
+		const TMap<FNiagaraVariable, FNiagaraGraphParameterReferenceCollection>& ReferenceMap = Graph->GetParameterReferenceMap();
+		if (ReferenceMap.Contains(ParameterToRead) == false && ReferenceMap.Contains(ParameterAlternative))
+		{
+			ParameterToRead = ParameterAlternative;
+		}
+	}
+	
+	UEdGraphPin* GetOutputPin = GetNode->RequestNewTypedPin(EGPD_Output, ParameterToRead.GetType(), LinkedParameterHandle.GetParameterHandleString());
 	MakeLinkTo(GetInputPin, PreviousStackNodeOutputPin);
 	MakeLinkTo(GetOutputPin, &OverridePin);
 
-	UNiagaraScriptVariable* ScriptVar = Graph->AddParameter(FNiagaraVariable(InputType, LinkedParameterHandle.GetParameterHandleString()), true);
+	UNiagaraScriptVariable* ScriptVar = Graph->AddParameter(ParameterToRead, false);
 	if (ScriptVar)
 	{
 		ScriptVar->DefaultMode = DesiredDefaultMode; 
diff --git a/Engine/Plugins/FX/Niagara/Source/NiagaraEditor/Public/NiagaraEditorUtilities.h b/Engine/Plugins/FX/Niagara/Source/NiagaraEditor/Public/NiagaraEditorUtilities.h
--- a/Engine/Plugins/FX/Niagara/Source/NiagaraEditor/Public/NiagaraEditorUtilities.h
+++ b/Engine/Plugins/FX/Niagara/Source/NiagaraEditor/Public/NiagaraEditorUtilities.h
@@ -163,6 +163,8 @@ namespace FNiagaraEditorUtilities
 
 	FText GetTypeDefinitionCategory(const FNiagaraTypeDefinition& TypeDefinition);
 
+	NIAGARAEDITOR_API bool AreTypesAssignable(const FNiagaraTypeDefinition& TypeA, const FNiagaraTypeDefinition& TypeB);
+
 	void MarkDependentCompilableAssetsDirty(TArray<UObject*> InObjects);
 
 	void ResolveNumerics(UNiagaraGraph* SourceGraph, bool bForceParametersToResolveNumerics, TArray<FNiagaraVariable>& ChangedNumericParams);
diff --git a/Engine/Plugins/FX/Niagara/Source/NiagaraEditorWidgets/Private/Stack/SNiagaraStackFunctionInputValue.cpp b/Engine/Plugins/FX/Niagara/Source/NiagaraEditorWidgets/Private/Stack/SNiagaraStackFunctionInputValue.cpp
--- a/Engine/Plugins/FX/Niagara/Source/NiagaraEditorWidgets/Private/Stack/SNiagaraStackFunctionInputValue.cpp
+++ b/Engine/Plugins/FX/Niagara/Source/NiagaraEditorWidgets/Private/Stack/SNiagaraStackFunctionInputValue.cpp
@@ -834,7 +834,8 @@ FReply SNiagaraStackFunctionInputValue::OnFunctionInputDrop(const FGeometry& InG
 		if (Action.IsValid())
 		{
 			FNiagaraTypeDefinition FromType = Action->GetParameter().GetType();
-			if (FromType == FunctionInput->GetInputType())
+			FNiagaraTypeDefinition ToType = FunctionInput->GetInputType();
+			if (FNiagaraEditorUtilities::AreTypesAssignable(FromType, ToType))
 			{
 				// the types are the same, so we can just link the value directly
 				FunctionInput->SetLinkedValueHandle(FNiagaraParameterHandle(Action->GetParameter().GetName()));
@@ -863,15 +864,16 @@ bool SNiagaraStackFunctionInputValue::OnFunctionInputAllowDrop(TSharedPtr<FDragD
 		TSharedPtr<FNiagaraParameterDragOperation> InputDragDropOperation = StaticCastSharedPtr<FNiagaraParameterDragOperation>(DragDropOperation);
 		TSharedPtr<FNiagaraParameterAction> Action = StaticCastSharedPtr<FNiagaraParameterAction>(InputDragDropOperation->GetSourceAction());
 		bool bAllowedInExecutionCategory = FNiagaraStackGraphUtilities::ParameterAllowedInExecutionCategory(Action->GetParameter().GetName(), FunctionInput->GetExecutionCategoryName());
-
+		FNiagaraTypeDefinition DropType = Action->GetParameter().GetType();
+		
 		// check if we can simply link the input directly
-		if (bAllowedInExecutionCategory && Action->GetParameter().GetType() == FunctionInput->GetInputType())
+		if (bAllowedInExecutionCategory && FNiagaraEditorUtilities::AreTypesAssignable(DropType, FunctionInput->GetInputType()))
 		{
 			return true;
 		}
 
 		// check if we can use a conversion script
-		if (bAllowedInExecutionCategory && FunctionInput->GetPossibleConversionScripts(Action->GetParameter().GetType()).Num() > 0)
+		if (bAllowedInExecutionCategory && FunctionInput->GetPossibleConversionScripts(DropType).Num() > 0)
 		{
 			return true;
 		}
