Removed all payload types from primitive types outside of preexisting event structs

diff --git a/Engine/Plugins/Experimental/ChaosNiagara/Source/ChaosNiagara/Private/NiagaraDataInterfacePhysicsField.cpp b/Engine/Plugins/Experimental/ChaosNiagara/Source/ChaosNiagara/Private/NiagaraDataInterfacePhysicsField.cpp
--- a/Engine/Plugins/Experimental/ChaosNiagara/Source/ChaosNiagara/Private/NiagaraDataInterfacePhysicsField.cpp
+++ b/Engine/Plugins/Experimental/ChaosNiagara/Source/ChaosNiagara/Private/NiagaraDataInterfacePhysicsField.cpp
@@ -298,8 +298,8 @@ void UNiagaraDataInterfacePhysicsField::PostInitProperties()
 
 		ENiagaraTypeRegistryFlags FieldFlags =
 			ENiagaraTypeRegistryFlags::AllowAnyVariable |
-			ENiagaraTypeRegistryFlags::AllowParameter |
-			ENiagaraTypeRegistryFlags::AllowPayload;
+			ENiagaraTypeRegistryFlags::AllowParameter;
+		
 		FNiagaraTypeRegistry::Register(FNiagaraTypeDefinition(StaticEnum<EFieldVectorType>()), FieldFlags);
 		FNiagaraTypeRegistry::Register(FNiagaraTypeDefinition(StaticEnum<EFieldScalarType>()), FieldFlags);
 		FNiagaraTypeRegistry::Register(FNiagaraTypeDefinition(StaticEnum<EFieldIntegerType>()), FieldFlags);
diff --git a/Engine/Plugins/FX/Niagara/Source/Niagara/Private/DataInterface/NiagaraDataInterfaceStaticMesh.cpp b/Engine/Plugins/FX/Niagara/Source/Niagara/Private/DataInterface/NiagaraDataInterfaceStaticMesh.cpp
--- a/Engine/Plugins/FX/Niagara/Source/Niagara/Private/DataInterface/NiagaraDataInterfaceStaticMesh.cpp
+++ b/Engine/Plugins/FX/Niagara/Source/Niagara/Private/DataInterface/NiagaraDataInterfaceStaticMesh.cpp
@@ -1405,8 +1405,7 @@ void UNiagaraDataInterfaceStaticMesh::PostInitProperties()
 		// Still some issues with using custom structs. Convert node for example throws a wobbler. TODO after GDC.
 		ENiagaraTypeRegistryFlags CoordFlags =
 			ENiagaraTypeRegistryFlags::AllowAnyVariable |
-			ENiagaraTypeRegistryFlags::AllowParameter |
-			ENiagaraTypeRegistryFlags::AllowPayload;
+			ENiagaraTypeRegistryFlags::AllowParameter;
 		FNiagaraTypeRegistry::Register(FMeshTriCoordinate::StaticStruct(), CoordFlags);
 	}
 }
diff --git a/Engine/Plugins/FX/Niagara/Source/Niagara/Private/NiagaraDataInterfaceSkeletalMesh.cpp b/Engine/Plugins/FX/Niagara/Source/Niagara/Private/NiagaraDataInterfaceSkeletalMesh.cpp
--- a/Engine/Plugins/FX/Niagara/Source/Niagara/Private/NiagaraDataInterfaceSkeletalMesh.cpp
+++ b/Engine/Plugins/FX/Niagara/Source/Niagara/Private/NiagaraDataInterfaceSkeletalMesh.cpp
@@ -2449,8 +2449,7 @@ void UNiagaraDataInterfaceSkeletalMesh::PostInitProperties()
 		//Still some issues with using custom structs. Convert node for example throws a wobbler. TODO after GDC.
 		ENiagaraTypeRegistryFlags CoordFlags =
 			ENiagaraTypeRegistryFlags::AllowAnyVariable |
-			ENiagaraTypeRegistryFlags::AllowParameter |
-			ENiagaraTypeRegistryFlags::AllowPayload;
+			ENiagaraTypeRegistryFlags::AllowParameter;
 		FNiagaraTypeRegistry::Register(FMeshTriCoordinate::StaticStruct(), CoordFlags);
 	}
 }
diff --git a/Engine/Plugins/FX/Niagara/Source/Niagara/Private/NiagaraModule.cpp b/Engine/Plugins/FX/Niagara/Source/Niagara/Private/NiagaraModule.cpp
--- a/Engine/Plugins/FX/Niagara/Source/Niagara/Private/NiagaraModule.cpp
+++ b/Engine/Plugins/FX/Niagara/Source/Niagara/Private/NiagaraModule.cpp
@@ -942,22 +942,23 @@ void FNiagaraTypeDefinition::RecreateUserDefinedTypeRegistry()
 	ENiagaraTypeRegistryFlags VarFlags = ENiagaraTypeRegistryFlags::AllowAnyVariable;
 	ENiagaraTypeRegistryFlags ParamFlags = VarFlags | ENiagaraTypeRegistryFlags::AllowParameter;
 	ENiagaraTypeRegistryFlags PayloadFlags = VarFlags | ENiagaraTypeRegistryFlags::AllowPayload;
+
 	FNiagaraTypeRegistry::Register(CollisionEventDef, PayloadFlags);
 
 	FNiagaraTypeRegistry::Register(ParameterMapDef, ParamFlags);
-	FNiagaraTypeRegistry::Register(IDDef, ParamFlags | PayloadFlags);
-	FNiagaraTypeRegistry::Register(RandInfoDef, ParamFlags | PayloadFlags);
+	FNiagaraTypeRegistry::Register(IDDef, ParamFlags);
+	FNiagaraTypeRegistry::Register(RandInfoDef, ParamFlags);
 	FNiagaraTypeRegistry::Register(NumericDef, ParamFlags);
-	FNiagaraTypeRegistry::Register(FloatDef, ParamFlags | PayloadFlags);
-	FNiagaraTypeRegistry::Register(HalfDef, ParamFlags | PayloadFlags);
-	FNiagaraTypeRegistry::Register(IntDef, ParamFlags | PayloadFlags);
-	FNiagaraTypeRegistry::Register(BoolDef, ParamFlags | PayloadFlags);
-	FNiagaraTypeRegistry::Register(Vec2Def, ParamFlags | PayloadFlags);
-	FNiagaraTypeRegistry::Register(Vec3Def, ParamFlags | PayloadFlags);
-	FNiagaraTypeRegistry::Register(Vec4Def, ParamFlags | PayloadFlags);
-	FNiagaraTypeRegistry::Register(ColorDef, ParamFlags | PayloadFlags);
-	FNiagaraTypeRegistry::Register(PositionDef, ParamFlags | PayloadFlags);
-	FNiagaraTypeRegistry::Register(QuatDef, ParamFlags | PayloadFlags);
+	FNiagaraTypeRegistry::Register(FloatDef, ParamFlags);
+	FNiagaraTypeRegistry::Register(HalfDef, ParamFlags);
+	FNiagaraTypeRegistry::Register(IntDef, ParamFlags);
+	FNiagaraTypeRegistry::Register(BoolDef, ParamFlags);
+	FNiagaraTypeRegistry::Register(Vec2Def, ParamFlags);
+	FNiagaraTypeRegistry::Register(Vec3Def, ParamFlags);
+	FNiagaraTypeRegistry::Register(Vec4Def, ParamFlags);
+	FNiagaraTypeRegistry::Register(ColorDef, ParamFlags);
+	FNiagaraTypeRegistry::Register(PositionDef, ParamFlags);
+	FNiagaraTypeRegistry::Register(QuatDef, ParamFlags);
 	FNiagaraTypeRegistry::Register(Matrix4Def, ParamFlags);
 
 	// Register static versions
@@ -968,10 +969,10 @@ void FNiagaraTypeDefinition::RecreateUserDefinedTypeRegistry()
 	// @todo wildcard shouldn't be available for parameters etc., so just don't register here?
 	//FNiagaraTypeRegistry::Register(WildcardDef, VarFlags);
 
-	FNiagaraTypeRegistry::Register(FNiagaraTypeDefinition(ExecutionStateEnum), ParamFlags | PayloadFlags);
-	FNiagaraTypeRegistry::Register(FNiagaraTypeDefinition(CoordinateSpaceEnum), ParamFlags | PayloadFlags);
-	FNiagaraTypeRegistry::Register(FNiagaraTypeDefinition(OrientationAxisEnum), ParamFlags | PayloadFlags);
-	FNiagaraTypeRegistry::Register(FNiagaraTypeDefinition(ExecutionStateSourceEnum), ParamFlags | PayloadFlags);
+	FNiagaraTypeRegistry::Register(FNiagaraTypeDefinition(ExecutionStateEnum), ParamFlags);
+	FNiagaraTypeRegistry::Register(FNiagaraTypeDefinition(CoordinateSpaceEnum), ParamFlags);
+	FNiagaraTypeRegistry::Register(FNiagaraTypeDefinition(OrientationAxisEnum), ParamFlags);
+	FNiagaraTypeRegistry::Register(FNiagaraTypeDefinition(ExecutionStateSourceEnum), ParamFlags);
 
 	UScriptStruct* SpawnInfoStruct = FindObjectChecked<UScriptStruct>(NiagaraPkg, TEXT("NiagaraSpawnInfo"));
 	FNiagaraTypeRegistry::Register(FNiagaraTypeDefinition(SpawnInfoStruct), ParamFlags);
@@ -981,7 +982,7 @@ void FNiagaraTypeDefinition::RecreateUserDefinedTypeRegistry()
 	FNiagaraTypeRegistry::Register(UTextureDef, VarFlags);
 	FNiagaraTypeRegistry::Register(UTextureRenderTargetDef, VarFlags);
 	FNiagaraTypeRegistry::Register(UStaticMeshDef, VarFlags);
-	FNiagaraTypeRegistry::Register(StaticEnum<ENiagaraLegacyTrailWidthMode>(), ParamFlags | PayloadFlags);
+	FNiagaraTypeRegistry::Register(StaticEnum<ENiagaraLegacyTrailWidthMode>(), ParamFlags);
 
 	
 	if (!IsRunningCommandlet())
@@ -1051,8 +1052,7 @@ void FNiagaraTypeDefinition::RecreateUserDefinedTypeRegistry()
 				UE_LOG(LogNiagara, Warning, TEXT("Could not find additional parameter/payload type: %s"), *AssetRef.ToString());
 			}
 		}
-
-
+		
 		for (FSoftObjectPath AssetRef : Settings->AdditionalParameterEnums)
 		{
 			FName AssetRefPathNamePreResolve = AssetRef.GetAssetPathName();
@@ -1070,20 +1070,11 @@ void FNiagaraTypeDefinition::RecreateUserDefinedTypeRegistry()
 
 			if (Obj != nullptr)
 			{
-				const FSoftObjectPath* ParamRefFound = Settings->AdditionalParameterEnums.FindByPredicate([&](const FSoftObjectPath& Ref) { return Ref.ToString() == AssetRef.ToString(); });
-				const FSoftObjectPath* PayloadRefFound = nullptr;
 				UEnum* Enum = Cast<UEnum>(Obj);
 				if (Enum != nullptr)
 				{
-					ENiagaraTypeRegistryFlags Flags = ENiagaraTypeRegistryFlags::AllowAnyVariable;
-					if (ParamRefFound)
-					{
-						Flags |= ENiagaraTypeRegistryFlags::AllowParameter;
-					}
-					if (PayloadRefFound)
-					{
-						Flags |= ENiagaraTypeRegistryFlags::AllowPayload;
-					}
+					ENiagaraTypeRegistryFlags Flags = ENiagaraTypeRegistryFlags::AllowAnyVariable | ENiagaraTypeRegistryFlags::AllowParameter;
+					
 					FNiagaraTypeDefinition Def(Enum);
 					FNiagaraTypeRegistry::Register(Def, Flags);
 					FNiagaraTypeRegistry::Register(Def.ToStaticDef(), StaticParamFlags);
diff --git a/Engine/Plugins/FX/Niagara/Source/NiagaraEditor/Private/EdGraphSchema_Niagara.cpp b/Engine/Plugins/FX/Niagara/Source/NiagaraEditor/Private/EdGraphSchema_Niagara.cpp
--- a/Engine/Plugins/FX/Niagara/Source/NiagaraEditor/Private/EdGraphSchema_Niagara.cpp
+++ b/Engine/Plugins/FX/Niagara/Source/NiagaraEditor/Private/EdGraphSchema_Niagara.cpp
@@ -498,33 +498,40 @@ TArray<TSharedPtr<FNiagaraAction_NewNode>> UEdGraphSchema_Niagara::GetGraphActio
 	if (bAllowEventNodes && bModuleGraph)
 	{
 		const FText MenuCat = LOCTEXT("NiagaraEventMenuCat", "Events");
-		const TArray<FNiagaraTypeDefinition>& RegisteredTypes = FNiagaraTypeRegistry::GetRegisteredPayloadTypes();
-		for (const FNiagaraTypeDefinition& Type : RegisteredTypes)
+		TArray<FString> ReadCategories = {MenuCat.ToString(), LOCTEXT("NiagaraEventCategory_Read", "Read").ToString()};
+		TArray<FString> WriteCategories = {MenuCat.ToString(), LOCTEXT("NiagaraEventCategory_Write", "Write").ToString()};
+
+		const TArray<FNiagaraTypeDefinition>& PayloadTypes = FNiagaraTypeRegistry::GetRegisteredPayloadTypes();
+		for (const FNiagaraTypeDefinition& PayloadType : PayloadTypes)
 		{
-			if (Type.IsInternalType())
+			if (PayloadType.IsInternalType())
 			{
 				continue;
 			}
 
-			if (Type.GetStruct() && !Type.GetStruct()->IsA(UNiagaraDataInterface::StaticClass()))
+			if (PayloadType.GetStruct() && !PayloadType.GetStruct()->IsChildOf(UNiagaraDataInterface::StaticClass()))
 			{
 				{
-					const FText MenuDescFmt = LOCTEXT("AddEventReadFmt", "Add {0} Event Read");
-					const FText DisplayName = FText::Format(MenuDescFmt, Type.GetNameText());
-				
+					const FText DisplayNameFormat = LOCTEXT("AddEventReadFmt", "Add {0} Event Read");
+					FString NameString = PayloadType.GetNameText().ToString();
+					NameString = NameString.Replace(TEXT(" Event"), TEXT(""), ESearchCase::IgnoreCase);
+					const FText DisplayName = FText::Format(DisplayNameFormat, FText::FromString(NameString));
+					
 					UNiagaraNodeReadDataSet* EventReadNode = NewObject<UNiagaraNodeReadDataSet>(OwnerOfTemporaries);
-					EventReadNode->InitializeFromStruct(Type.GetStruct());
+					EventReadNode->InitializeFromStruct(PayloadType.GetStruct());
 
-					AddNewNodeMenuAction(NewActions, EventReadNode, DisplayName, ENiagaraMenuSections::General, {MenuCat.ToString()}, FText::GetEmpty(), FText::GetEmpty());
+					AddNewNodeMenuAction(NewActions, EventReadNode, DisplayName, ENiagaraMenuSections::General, ReadCategories, FText::GetEmpty(), FText::GetEmpty());
 				}
 				{
-					const FText MenuDescFmt = LOCTEXT("AddEventWriteFmt", "Add {0} Event Write");
-					const FText DisplayName = FText::Format(MenuDescFmt, Type.GetNameText());
-
+					const FText DisplayNameFormat = LOCTEXT("AddEventWriteFmt", "Add {0} Event Write");
+					FString NameString = PayloadType.GetNameText().ToString();
+					NameString = NameString.Replace(TEXT("Event"), TEXT(""), ESearchCase::IgnoreCase);
+					const FText DisplayName = FText::Format(DisplayNameFormat, FText::FromString(NameString));
+					
 					UNiagaraNodeWriteDataSet* EventWriteNode = NewObject<UNiagaraNodeWriteDataSet>(OwnerOfTemporaries);
-					EventWriteNode->InitializeFromStruct(Type.GetStruct());
+					EventWriteNode->InitializeFromStruct(PayloadType.GetStruct());
 					
-					AddNewNodeMenuAction(NewActions, EventWriteNode, DisplayName, ENiagaraMenuSections::General, {MenuCat.ToString()}, FText::GetEmpty(), FText::GetEmpty());
+					AddNewNodeMenuAction(NewActions, EventWriteNode, DisplayName, ENiagaraMenuSections::General, WriteCategories, FText::GetEmpty(), FText::GetEmpty());
 				}
 			}
 		}
diff --git a/Engine/Plugins/FX/Niagara/Source/NiagaraEditor/Private/NiagaraNodeDataSetBase.cpp b/Engine/Plugins/FX/Niagara/Source/NiagaraEditor/Private/NiagaraNodeDataSetBase.cpp
--- a/Engine/Plugins/FX/Niagara/Source/NiagaraEditor/Private/NiagaraNodeDataSetBase.cpp
+++ b/Engine/Plugins/FX/Niagara/Source/NiagaraEditor/Private/NiagaraNodeDataSetBase.cpp
@@ -278,7 +278,7 @@ bool UNiagaraNodeDataSetBase::GetSupportedNiagaraTypeDef(const FProperty* Proper
 	else if (StructProp && StructProp->Struct)
 	{
 		TypeDef = FNiagaraTypeDefinition(FNiagaraTypeHelper::FindNiagaraFriendlyTopLevelStruct(StructProp->Struct, ENiagaraStructConversion::UserFacing));
-		return FNiagaraTypeRegistry::GetRegisteredPayloadTypes().Contains(TypeDef);
+		return FNiagaraTypeRegistry::GetRegisteredParameterTypes().Contains(TypeDef);
 	}
 	return false;
 }
