Fixed a crash when compiling dangling reroute nodes

diff --git a/Engine/Plugins/FX/Niagara/Source/NiagaraEditor/Private/NiagaraGraph.cpp b/Engine/Plugins/FX/Niagara/Source/NiagaraEditor/Private/NiagaraGraph.cpp
--- a/Engine/Plugins/FX/Niagara/Source/NiagaraEditor/Private/NiagaraGraph.cpp
+++ b/Engine/Plugins/FX/Niagara/Source/NiagaraEditor/Private/NiagaraGraph.cpp
@@ -1260,9 +1260,13 @@ UNiagaraGraph* UNiagaraGraph::CreateCompilationCopy(const TArray<ENiagaraScriptU
 					{
 						check(LinkedPin->Direction == EGPD_Output);
 						LinkedPin = CastChecked<UNiagaraNodeReroute>(OwningNode)->GetTracedOutputPin(LinkedPin, false);
-						OwningNode = LinkedPin->GetOwningNode();
-						NewNodePtr = DuplicationMapping.Find(OwningNode);
-						bAddReciprocalLink = true;
+						// reroute nodes can return nullptr for the traced output if they are not connected to anything
+						if (LinkedPin != nullptr)
+						{
+							OwningNode = LinkedPin->GetOwningNode();
+							NewNodePtr = DuplicationMapping.Find(OwningNode);
+							bAddReciprocalLink = true;
+						}
 					}
 				}
 
