Fixed a duplicate loctext key

diff --git a/Engine/Plugins/FX/Niagara/Source/Niagara/Private/NiagaraDataInterfaceCollisionQuery.cpp b/Engine/Plugins/FX/Niagara/Source/Niagara/Private/NiagaraDataInterfaceCollisionQuery.cpp
--- a/Engine/Plugins/FX/Niagara/Source/Niagara/Private/NiagaraDataInterfaceCollisionQuery.cpp
+++ b/Engine/Plugins/FX/Niagara/Source/Niagara/Private/NiagaraDataInterfaceCollisionQuery.cpp
@@ -264,7 +264,7 @@ void UNiagaraDataInterfaceCollisionQuery::GetFunctions(TArray<FNiagaraFunctionSi
 
 	// DEPRECATE_BEGIN
 	{
-		const FText AsyncTraceChannelDescription = LOCTEXT("TraceChannelDescription", "Currently unused, will represent the trace channels for which geometry the trace should test against");
+		const FText AsyncTraceChannelDescription = LOCTEXT("AsyncTraceChannelDescription", "Currently unused, will represent the trace channels for which geometry the trace should test against");
 		const FText QueryIDDescription = LOCTEXT("QueryIDDescription", "Unique (for this frame) index of the query being enqueued (used in subsequent frames to retrieve results).  Must be less than MaxRayTraceCount");
 		const FText IsQueryValidDescription = LOCTEXT("IsQueryValidDescription", "Returns true if the query was enqueued");
 		
