- Manually tap depth inside PS as it's mismatched when screen percent is enabled

#rnx

diff --git a/Engine/Plugins/FX/Niagara/Shaders/Private/NiagaraDebugDraw.usf b/Engine/Plugins/FX/Niagara/Shaders/Private/NiagaraDebugDraw.usf
--- a/Engine/Plugins/FX/Niagara/Shaders/Private/NiagaraDebugDraw.usf
+++ b/Engine/Plugins/FX/Niagara/Shaders/Private/NiagaraDebugDraw.usf
@@ -50,12 +50,27 @@ void MainVS(
 #endif
 
 #if NIAGARA_DEBUGDRAW_DRAWLINE_PS
+
+Texture2D<float> DepthTexture;
+SamplerState DepthSampler;
+float2 OutputInvResolution; 
+float2 OriginalViewRectMin;
+float2 OriginalViewSize;
+float2 OriginalBufferInvSize;
+
 void MainPS (
 	in float4 Position	: SV_POSITION,
 	in float4 Color		: TEXCOORD0,
 	out float4 OutColor	: SV_Target0
 )
 {
-	OutColor = Color;
+	uint2 PixelCoord		= Position.xy;
+	const float2 ViewportUV = float2(PixelCoord) * OutputInvResolution;
+	const float2 DepthUV	= (ViewportUV * OriginalViewSize.xy + OriginalViewRectMin) * OriginalBufferInvSize;
+	const float Depth		= DepthTexture.SampleLevel(DepthSampler, DepthUV, 0);
+	const bool bIsHidden	= Position.z < Depth; // Reverse-Z
+
+	const float ColorScale = bIsHidden ? 0.1f : 1;
+	OutColor = Color * ColorScale;
 }
 #endif
diff --git a/Engine/Plugins/FX/Niagara/Source/NiagaraShader/Private/NiagaraDebugShaders.cpp b/Engine/Plugins/FX/Niagara/Source/NiagaraShader/Private/NiagaraDebugShaders.cpp
--- a/Engine/Plugins/FX/Niagara/Source/NiagaraShader/Private/NiagaraDebugShaders.cpp
+++ b/Engine/Plugins/FX/Niagara/Source/NiagaraShader/Private/NiagaraDebugShaders.cpp
@@ -146,6 +146,13 @@ class NIAGARASHADER_API FNiagaraDebugDrawLinePS : public FGlobalShader
 	}
 
 	BEGIN_SHADER_PARAMETER_STRUCT(FParameters, )
+		SHADER_PARAMETER(FVector2f, OutputInvResolution)
+		SHADER_PARAMETER(FVector2f, OriginalViewRectMin)
+		SHADER_PARAMETER(FVector2f, OriginalViewSize)
+		SHADER_PARAMETER(FVector2f, OriginalBufferInvSize)
+		SHADER_PARAMETER_RDG_TEXTURE(Texture2D, DepthTexture)
+		SHADER_PARAMETER_SAMPLER(SamplerState, DepthSampler)
+
 		RENDER_TARGET_BINDING_SLOTS()
 	END_SHADER_PARAMETER_STRUCT()
 };
@@ -194,8 +201,13 @@ void NiagaraDebugShaders::DrawDebugLines(
 	PassParameters->VSParameters.View = View.ViewUniformBuffer;
 	PassParameters->VSParameters.GpuLineBuffer = LineBuffer;
 
-	PassParameters->PSParameters.RenderTargets[0] = FRenderTargetBinding(SceneColor, ERenderTargetLoadAction::ELoad);
-	PassParameters->PSParameters.RenderTargets.DepthStencil = FDepthStencilBinding(SceneDepth, ERenderTargetLoadAction::ELoad, FExclusiveDepthStencil::DepthRead_StencilNop);
+	PassParameters->PSParameters.OutputInvResolution	= FVector2f(1.0f / View.UnscaledViewRect.Width(), 1.0f / View.UnconstrainedViewRect.Height());
+	PassParameters->PSParameters.OriginalViewRectMin	= FVector2f(View.ViewRect.Min);
+	PassParameters->PSParameters.OriginalViewSize		= FVector2f(View.ViewRect.Width(), View.ViewRect.Height());
+	PassParameters->PSParameters.OriginalBufferInvSize	= FVector2f(1.f / SceneDepth->Desc.Extent.X, 1.f / SceneDepth->Desc.Extent.Y);
+	PassParameters->PSParameters.DepthTexture			= SceneDepth;
+	PassParameters->PSParameters.DepthSampler			= TStaticSamplerState<SF_Point, AM_Clamp, AM_Clamp, AM_Clamp>::GetRHI();
+	PassParameters->PSParameters.RenderTargets[0]		= FRenderTargetBinding(SceneColor, ERenderTargetLoadAction::ELoad);
 
 	GraphBuilder.AddPass(
 		RDG_EVENT_NAME("NiagaraDrawDebugLines"),
@@ -234,8 +246,13 @@ void NiagaraDebugShaders::DrawDebugLines(
 	PassParameters->VSParameters.View = View.ViewUniformBuffer;
 	PassParameters->VSParameters.GpuLineBuffer = LineBuffer;
 
+	PassParameters->PSParameters.OutputInvResolution = FVector2f(1.0f / View.UnscaledViewRect.Width(), 1.0f / View.UnconstrainedViewRect.Height());
+	PassParameters->PSParameters.OriginalViewRectMin = FVector2f(View.ViewRect.Min);
+	PassParameters->PSParameters.OriginalViewSize = FVector2f(View.ViewRect.Width(), View.ViewRect.Height());
+	PassParameters->PSParameters.OriginalBufferInvSize = FVector2f(1.f / SceneDepth->Desc.Extent.X, 1.f / SceneDepth->Desc.Extent.Y);
+	PassParameters->PSParameters.DepthTexture = SceneDepth;
+	PassParameters->PSParameters.DepthSampler = TStaticSamplerState<SF_Point, AM_Clamp, AM_Clamp, AM_Clamp>::GetRHI();
 	PassParameters->PSParameters.RenderTargets[0] = FRenderTargetBinding(SceneColor, ERenderTargetLoadAction::ELoad);
-	PassParameters->PSParameters.RenderTargets.DepthStencil = FDepthStencilBinding(SceneDepth, ERenderTargetLoadAction::ELoad, FExclusiveDepthStencil::DepthRead_StencilNop);
 
 	GraphBuilder.AddPass(
 		RDG_EVENT_NAME("NiagaraDrawDebugLines"),
