Fixed data interface functions for LWC types

diff --git a/Engine/Plugins/Experimental/ChaosNiagara/Shaders/NiagaraDataInterfacePhysicsField.ush b/Engine/Plugins/Experimental/ChaosNiagara/Shaders/NiagaraDataInterfacePhysicsField.ush
--- a/Engine/Plugins/Experimental/ChaosNiagara/Shaders/NiagaraDataInterfacePhysicsField.ush
+++ b/Engine/Plugins/Experimental/ChaosNiagara/Shaders/NiagaraDataInterfacePhysicsField.ush
@@ -23,6 +23,7 @@ struct FDIPhysicsFieldContext
 	int ClipmapCount;
 	int TargetCount;
 	int4 FieldTargets[MAX_PHYSICS_FIELD_TARGETS];
+	float3 SystemLWCTile;
 };
 
 #define DIPhysicsField_DECLARE_CONSTANTS(NAME)\
@@ -34,6 +35,7 @@ int ClipmapExponent_##NAME;\
 int ClipmapCount_##NAME;\
 int TargetCount_##NAME;\
 int4 FieldTargets_##NAME[MAX_PHYSICS_FIELD_TARGETS];\
+float3 SystemLWCTile_##NAME;\
 
 #define DIPhysicsField_MAKE_CONTEXT(NAME)\
 FDIPhysicsFieldContext DIContext; \
@@ -45,21 +47,31 @@ DIContext.ClipmapExponent = ClipmapExponent_##NAME;\
 DIContext.ClipmapCount = ClipmapCount_##NAME;\
 DIContext.TargetCount = TargetCount_##NAME;\
 DIContext.FieldTargets = FieldTargets_##NAME;\
+DIContext.SystemLWCTile = SystemLWCTile##NAME;\
 
 float3 DIPhysicsField_SamplePhysicsVectorField(in FDIPhysicsFieldContext DIContext, in float3 WorldPosition, in int VectorTarget)
 {
-	return PhysicsField_SamplePhysicsVectorField(WorldPosition, VectorTarget, DIContext.FieldTargets, DIContext.TargetCount, DIContext.ClipmapCenter, DIContext.ClipmapDistance, 
+	// LWC_TODO: fix this when the physics field functions are lwc ready
+	FLWCVector3 SamplePos = MakeLWCVector3(DIContext.SystemLWCTile, WorldPosition);
+	
+	return PhysicsField_SamplePhysicsVectorField(LWCHackToFloat(SamplePos), VectorTarget, DIContext.FieldTargets, DIContext.TargetCount, DIContext.ClipmapCenter, DIContext.ClipmapDistance, 
 				DIContext.ClipmapExponent, DIContext.ClipmapCount, DIContext.ClipmapResolution, DIContext.ClipmapBuffer);
 }
 
 float DIPhysicsField_SamplePhysicsScalarField(in FDIPhysicsFieldContext DIContext, in float3 WorldPosition, in int ScalarTarget)
 {
-	return PhysicsField_SamplePhysicsScalarField(WorldPosition, ScalarTarget, DIContext.FieldTargets, DIContext.TargetCount, DIContext.ClipmapCenter, DIContext.ClipmapDistance,
+	// LWC_TODO: fix this when the physics field functions are lwc ready
+	FLWCVector3 SamplePos = MakeLWCVector3(DIContext.SystemLWCTile, WorldPosition);
+	
+	return PhysicsField_SamplePhysicsScalarField(LWCHackToFloat(SamplePos), ScalarTarget, DIContext.FieldTargets, DIContext.TargetCount, DIContext.ClipmapCenter, DIContext.ClipmapDistance,
 				DIContext.ClipmapExponent, DIContext.ClipmapCount, DIContext.ClipmapResolution, DIContext.ClipmapBuffer);
 }
 
 int DIPhysicsField_SamplePhysicsIntegerField(in FDIPhysicsFieldContext DIContext, in float3 WorldPosition, in int IntegerTarget)
 {
-	return PhysicsField_SamplePhysicsIntegerField(WorldPosition, IntegerTarget, DIContext.FieldTargets, DIContext.TargetCount, DIContext.ClipmapCenter, DIContext.ClipmapDistance,
+	// LWC_TODO: fix this when the physics field functions are lwc ready
+	FLWCVector3 SamplePos = MakeLWCVector3(DIContext.SystemLWCTile, WorldPosition);
+	
+	return PhysicsField_SamplePhysicsIntegerField(LWCHackToFloat(SamplePos), IntegerTarget, DIContext.FieldTargets, DIContext.TargetCount, DIContext.ClipmapCenter, DIContext.ClipmapDistance,
 				DIContext.ClipmapExponent, DIContext.ClipmapCount, DIContext.ClipmapResolution, DIContext.ClipmapBuffer);
 }
\ No newline at end of file
diff --git a/Engine/Plugins/Experimental/ChaosNiagara/Source/ChaosNiagara/Classes/NiagaraDataInterfacePhysicsField.h b/Engine/Plugins/Experimental/ChaosNiagara/Source/ChaosNiagara/Classes/NiagaraDataInterfacePhysicsField.h
--- a/Engine/Plugins/Experimental/ChaosNiagara/Source/ChaosNiagara/Classes/NiagaraDataInterfacePhysicsField.h
+++ b/Engine/Plugins/Experimental/ChaosNiagara/Source/ChaosNiagara/Classes/NiagaraDataInterfacePhysicsField.h
@@ -36,6 +36,8 @@ struct FNDIPhysicsFieldData
 
 	/** Time in seconds*/
 	float TimeSeconds = 0.0;
+	
+	FNiagaraLWCConverter LWCConverter;
 };
 
 /** Data Interface for the strand base */
@@ -69,6 +71,7 @@ public:
 	virtual void GetCommonHLSL(FString& OutHLSL) override;
 	virtual void GetParameterDefinitionHLSL(const FNiagaraDataInterfaceGPUParamInfo& ParamInfo, FString& OutHLSL) override;
 	virtual bool GetFunctionHLSL(const FNiagaraDataInterfaceGPUParamInfo& ParamInfo, const FNiagaraDataInterfaceGeneratedFunction& FunctionInfo, int FunctionInstanceIndex, FString& OutHLSL) override;
+	virtual bool UpgradeFunctionCall(FNiagaraFunctionSignature& FunctionSignature) override;
 #endif
 	virtual void ProvidePerInstanceDataForRenderThread(void* DataForRenderThread, void* PerInstanceData, const FNiagaraSystemInstanceID& SystemInstance) override;
 
@@ -96,6 +99,7 @@ public:
 	static const FString ClipmapCountName;
 	static const FString TargetCountName;
 	static const FString FieldTargetsName;
+	static const FString SystemLWCTileName;
 
 protected:
 	/** Copy one niagara DI to this */
diff --git a/Engine/Plugins/Experimental/ChaosNiagara/Source/ChaosNiagara/Private/NiagaraDataInterfacePhysicsField.cpp b/Engine/Plugins/Experimental/ChaosNiagara/Source/ChaosNiagara/Private/NiagaraDataInterfacePhysicsField.cpp
--- a/Engine/Plugins/Experimental/ChaosNiagara/Source/ChaosNiagara/Private/NiagaraDataInterfacePhysicsField.cpp
+++ b/Engine/Plugins/Experimental/ChaosNiagara/Source/ChaosNiagara/Private/NiagaraDataInterfacePhysicsField.cpp
@@ -13,6 +13,18 @@
 #define LOCTEXT_NAMESPACE "NiagaraDataInterfacePhysicsField"
 DEFINE_LOG_CATEGORY_STATIC(LogPhysicsField, Log, All);
 
+struct FNiagaraPhysicsFieldDIFunctionVersion
+{
+	enum Type
+	{
+		InitialVersion = 0,
+		LargeWorldCoordinates = 1,
+
+		VersionPlusOne,
+		LatestVersion = VersionPlusOne - 1
+	};
+};
+
 //------------------------------------------------------------------------------------------------------------
 
 static const FName SamplePhysicsVectorFieldName(TEXT("SamplePhysicsVectorField"));
@@ -32,6 +44,7 @@ const FString UNiagaraDataInterfacePhysicsField::ClipmapExponentName(TEXT("Clipm
 const FString UNiagaraDataInterfacePhysicsField::ClipmapCountName(TEXT("ClipmapCount_"));
 const FString UNiagaraDataInterfacePhysicsField::TargetCountName(TEXT("TargetCount_"));
 const FString UNiagaraDataInterfacePhysicsField::FieldTargetsName(TEXT("FieldTargets_"));
+const FString UNiagaraDataInterfacePhysicsField::SystemLWCTileName(TEXT("SystemLWCTile_"));
 
 //------------------------------------------------------------------------------------------------------------
 
@@ -47,6 +60,7 @@ struct FNDIPhysicsFieldParametersName
 		ClipmapCountName = UNiagaraDataInterfacePhysicsField::ClipmapCountName + Suffix;
 		TargetCountName = UNiagaraDataInterfacePhysicsField::TargetCountName + Suffix;
 		FieldTargetsName = UNiagaraDataInterfacePhysicsField::FieldTargetsName + Suffix;
+		SystemLWCTileName = UNiagaraDataInterfacePhysicsField::SystemLWCTileName + Suffix;
 	}
 
 	FString ClipmapBufferName;
@@ -57,6 +71,7 @@ struct FNDIPhysicsFieldParametersName
 	FString ClipmapCountName;
 	FString TargetCountName;
 	FString FieldTargetsName;
+	FString SystemLWCTileName;
 };
 
 //------------------------------------------------------------------------------------------------------------
@@ -83,6 +98,7 @@ void FNDIPhysicsFieldData::Init(FNiagaraSystemInstance* SystemInstance)
 				FieldResource = FieldComponent->FieldInstance->FieldResource;
 			}
 		}
+		LWCConverter = SystemInstance->GetLWCConverter();
 	}
 }
 
@@ -122,6 +138,7 @@ public:
 		ClipmapCount.Bind(ParameterMap, *ParamNames.ClipmapCountName);
 		TargetCount.Bind(ParameterMap, *ParamNames.TargetCountName);
 		FieldTargets.Bind(ParameterMap, *ParamNames.FieldTargetsName);
+		SystemLWCTile.Bind(ParameterMap, *ParamNames.SystemLWCTileName);
 	}
 
 	void Set(FRHICommandList& RHICmdList, const FNiagaraDataInterfaceSetArgs& Context) const
@@ -168,6 +185,7 @@ public:
 			SetShaderValue(RHICmdList, ComputeShaderRHI, TargetCount, 0);
 			SetShaderValue(RHICmdList, ComputeShaderRHI, FieldTargets, LocalTargets);
 		}
+		SetShaderValue(RHICmdList, ComputeShaderRHI, SystemLWCTile, Context.SystemLWCTile);
 	}
 
 	void Unset(FRHICommandList& RHICmdList, const FNiagaraDataInterfaceSetArgs& Context) const
@@ -184,6 +202,7 @@ private:
 	LAYOUT_FIELD(FShaderParameter, ClipmapCount);
 	LAYOUT_FIELD(FShaderParameter, TargetCount);
 	LAYOUT_FIELD(FShaderParameter, FieldTargets);
+	LAYOUT_FIELD(FShaderParameter, SystemLWCTile);
 };
 
 IMPLEMENT_TYPE_LAYOUT(FNDIPhysicsFieldParametersCS);
@@ -209,9 +228,7 @@ void FNDIPhysicsFieldProxy::ConsumePerInstanceDataFromGameThread(void* PerInstan
 void FNDIPhysicsFieldProxy::InitializePerInstanceData(const FNiagaraSystemInstanceID& SystemInstance)
 {
 	check(IsInRenderingThread());
-
-	FNDIFieldRenderData* TargetData = SystemInstancesToProxyData.Find(SystemInstance);
-	TargetData = &SystemInstancesToProxyData.Add(SystemInstance);
+	SystemInstancesToProxyData.Add(SystemInstance);
 }
 
 void FNDIPhysicsFieldProxy::DestroyPerInstanceData(const FNiagaraSystemInstanceID& SystemInstance)
@@ -314,7 +331,7 @@ void UNiagaraDataInterfacePhysicsField::GetFunctions(TArray<FNiagaraFunctionSign
 		Sig.bMemberFunction = true;
 		Sig.bRequiresContext = false;
 		Sig.Inputs.Add(FNiagaraVariable(FNiagaraTypeDefinition(GetClass()), TEXT("Physics Field")));
-		Sig.Inputs.Add(FNiagaraVariable(FNiagaraTypeDefinition::GetVec3Def(), TEXT("World Position")));
+		Sig.Inputs.Add(FNiagaraVariable(FNiagaraTypeDefinition::GetPositionDef(), TEXT("World Position")));
 		Sig.Inputs.Add(FNiagaraVariable(StaticEnum<EFieldVectorType>(), TEXT("Target Type")));
 		Sig.Outputs.Add(FNiagaraVariable(FNiagaraTypeDefinition::GetVec3Def(), TEXT("Vector Value")));
 
@@ -326,7 +343,7 @@ void UNiagaraDataInterfacePhysicsField::GetFunctions(TArray<FNiagaraFunctionSign
 		Sig.bMemberFunction = true;
 		Sig.bRequiresContext = false;
 		Sig.Inputs.Add(FNiagaraVariable(FNiagaraTypeDefinition(GetClass()), TEXT("Physics Field")));
-		Sig.Inputs.Add(FNiagaraVariable(FNiagaraTypeDefinition::GetVec3Def(), TEXT("World Position")));
+		Sig.Inputs.Add(FNiagaraVariable(FNiagaraTypeDefinition::GetPositionDef(), TEXT("World Position")));
 		Sig.Inputs.Add(FNiagaraVariable(StaticEnum<EFieldScalarType>(), TEXT("Target Type")));
 		Sig.Outputs.Add(FNiagaraVariable(FNiagaraTypeDefinition::GetFloatDef(), TEXT("Scalar Value")));
 
@@ -338,7 +355,7 @@ void UNiagaraDataInterfacePhysicsField::GetFunctions(TArray<FNiagaraFunctionSign
 		Sig.bMemberFunction = true;
 		Sig.bRequiresContext = false;
 		Sig.Inputs.Add(FNiagaraVariable(FNiagaraTypeDefinition(GetClass()), TEXT("Physics Field")));
-		Sig.Inputs.Add(FNiagaraVariable(FNiagaraTypeDefinition::GetVec3Def(), TEXT("World Position")));
+		Sig.Inputs.Add(FNiagaraVariable(FNiagaraTypeDefinition::GetPositionDef(), TEXT("World Position")));
 		Sig.Inputs.Add(FNiagaraVariable(StaticEnum<EFieldIntegerType>(), TEXT("Target Type")));
 		Sig.Outputs.Add(FNiagaraVariable(FNiagaraTypeDefinition::GetIntDef(), TEXT("Integer Value")));
 
@@ -360,7 +377,7 @@ void UNiagaraDataInterfacePhysicsField::GetFunctions(TArray<FNiagaraFunctionSign
 		Sig.bMemberFunction = true;
 		Sig.bRequiresContext = false;
 		Sig.Inputs.Add(FNiagaraVariable(FNiagaraTypeDefinition(GetClass()), TEXT("Physics Field")));
-		Sig.Outputs.Add(FNiagaraVariable(FNiagaraTypeDefinition::GetVec3Def(), TEXT("Min Bounds")));
+		Sig.Outputs.Add(FNiagaraVariable(FNiagaraTypeDefinition::GetVec3Def(), TEXT("Min Bounds"))); //TODO (LWC) not sure what to do with these bounds, should they be converted as well?
 		Sig.Outputs.Add(FNiagaraVariable(FNiagaraTypeDefinition::GetVec3Def(), TEXT("Max Bounds")));
 
 		OutFunctions.Add(Sig);
@@ -527,7 +544,7 @@ void UNiagaraDataInterfacePhysicsField::SamplePhysicsVectorField(FVectorVMExtern
 	VectorVM::FUserPtrHandler<FNDIPhysicsFieldData> InstData(Context);
 
 	// Inputs 
-	FNDIInputParam<FVector3f> SamplePositionParam(Context);
+	FNDIInputParam<FNiagaraPosition> SamplePositionParam(Context);
 	FNDIInputParam<EFieldVectorType> VectorTargetParam(Context);
 
 	// Outputs...
@@ -538,11 +555,11 @@ void UNiagaraDataInterfacePhysicsField::SamplePhysicsVectorField(FVectorVMExtern
 		FFieldExecutionDatas ExecutionDatas;
 		ExecutionDatas.SamplePositions.Init(FVector(0, 0, 0), Context.GetNumInstances());
 
-		EFieldVectorType VectorTarget = EFieldVectorType::Vector_TargetMax;
+		EFieldVectorType VectorTarget = Vector_TargetMax;
 
 		for (int32 InstanceIdx = 0; InstanceIdx < Context.GetNumInstances(); ++InstanceIdx)
 		{
-			ExecutionDatas.SamplePositions[InstanceIdx] = (FVector)SamplePositionParam.GetAndAdvance();
+			ExecutionDatas.SamplePositions[InstanceIdx] = InstData->LWCConverter.ConvertSimulationPositionToWorld(SamplePositionParam.GetAndAdvance());
 			VectorTarget = VectorTargetParam.GetAndAdvance();
 		}
 		FFieldContextIndex::ContiguousIndices(ExecutionDatas.SampleIndices, Context.GetNumInstances());
@@ -559,12 +576,12 @@ void UNiagaraDataInterfacePhysicsField::SamplePhysicsVectorField(FVectorVMExtern
 			InstData->TimeSeconds
 		};
 
-		const EFieldPhysicsType PhysicsType = GetFieldTargetTypes(EFieldOutputType::Field_Output_Vector)[VectorTarget];
+		const EFieldPhysicsType PhysicsType = GetFieldTargetTypes(Field_Output_Vector)[VectorTarget];
 		EvaluateFieldNodes<FVector, FVectorFieldOperator>(InstData->FieldCommands, PhysicsType, FieldContext, SampleResults, SampleMax);
 
 		for (int32 InstanceIdx = 0; InstanceIdx < Context.GetNumInstances(); ++InstanceIdx)
 		{
-			OutVectorFieldParam.SetAndAdvance((FVector3f)SampleMax[InstanceIdx]);
+			OutVectorFieldParam.SetAndAdvance(SampleMax[InstanceIdx]);
 		}
 	}
 	else
@@ -581,7 +598,7 @@ void UNiagaraDataInterfacePhysicsField::SamplePhysicsIntegerField(FVectorVMExter
 	VectorVM::FUserPtrHandler<FNDIPhysicsFieldData> InstData(Context);
 
 	// Inputs 
-	FNDIInputParam<FVector3f> SamplePositionParam(Context);
+	FNDIInputParam<FNiagaraPosition> SamplePositionParam(Context);
 	FNDIInputParam<EFieldIntegerType> IntegerTargetParam(Context);
 
 	// Outputs...
@@ -592,11 +609,11 @@ void UNiagaraDataInterfacePhysicsField::SamplePhysicsIntegerField(FVectorVMExter
 		FFieldExecutionDatas ExecutionDatas;
 		ExecutionDatas.SamplePositions.Init(FVector(0, 0, 0), Context.GetNumInstances());
 
-		EFieldIntegerType IntegerTarget = EFieldIntegerType::Integer_TargetMax;
+		EFieldIntegerType IntegerTarget = Integer_TargetMax;
 
 		for (int32 InstanceIdx = 0; InstanceIdx < Context.GetNumInstances(); ++InstanceIdx)
 		{
-			ExecutionDatas.SamplePositions[InstanceIdx] = (FVector)SamplePositionParam.GetAndAdvance();
+			ExecutionDatas.SamplePositions[InstanceIdx] = InstData->LWCConverter.ConvertSimulationPositionToWorld(SamplePositionParam.GetAndAdvance());
 			IntegerTarget = IntegerTargetParam.GetAndAdvance();
 		}
 		FFieldContextIndex::ContiguousIndices(ExecutionDatas.SampleIndices, Context.GetNumInstances());
@@ -613,7 +630,7 @@ void UNiagaraDataInterfacePhysicsField::SamplePhysicsIntegerField(FVectorVMExter
 			InstData->TimeSeconds
 		};
 
-		const EFieldPhysicsType PhysicsType = GetFieldTargetTypes(EFieldOutputType::Field_Output_Integer)[IntegerTarget];
+		const EFieldPhysicsType PhysicsType = GetFieldTargetTypes(Field_Output_Integer)[IntegerTarget];
 		EvaluateFieldNodes<int32, FIntegerFieldOperator>(InstData->FieldCommands, PhysicsType, FieldContext, SampleResults, SampleMax);
 
 		for (int32 InstanceIdx = 0; InstanceIdx < Context.GetNumInstances(); ++InstanceIdx)
@@ -636,7 +653,7 @@ void UNiagaraDataInterfacePhysicsField::SamplePhysicsScalarField(FVectorVMExtern
 	VectorVM::FUserPtrHandler<FNDIPhysicsFieldData> InstData(Context);
 
 	// Inputs 
-	FNDIInputParam<FVector3f> SamplePositionParam(Context);
+	FNDIInputParam<FNiagaraPosition> SamplePositionParam(Context);
 	FNDIInputParam<EFieldScalarType> ScalarTargetParam(Context);
 
 	// Outputs...
@@ -647,11 +664,11 @@ void UNiagaraDataInterfacePhysicsField::SamplePhysicsScalarField(FVectorVMExtern
 		FFieldExecutionDatas ExecutionDatas;
 		ExecutionDatas.SamplePositions.Init(FVector(0, 0, 0), Context.GetNumInstances());
 
-		EFieldScalarType ScalarTarget = EFieldScalarType::Scalar_TargetMax;
+		EFieldScalarType ScalarTarget = Scalar_TargetMax;
 
 		for (int32 InstanceIdx = 0; InstanceIdx < Context.GetNumInstances(); ++InstanceIdx)
 		{
-			ExecutionDatas.SamplePositions[InstanceIdx] = (FVector)SamplePositionParam.GetAndAdvance();
+			ExecutionDatas.SamplePositions[InstanceIdx] = InstData->LWCConverter.ConvertSimulationPositionToWorld(SamplePositionParam.GetAndAdvance());
 			ScalarTarget = ScalarTargetParam.GetAndAdvance();
 		}
 		FFieldContextIndex::ContiguousIndices(ExecutionDatas.SampleIndices, Context.GetNumInstances());
@@ -668,7 +685,7 @@ void UNiagaraDataInterfacePhysicsField::SamplePhysicsScalarField(FVectorVMExtern
 			InstData->TimeSeconds
 		};
 
-		const EFieldPhysicsType PhysicsType = GetFieldTargetTypes(EFieldOutputType::Field_Output_Scalar)[ScalarTarget];
+		const EFieldPhysicsType PhysicsType = GetFieldTargetTypes(Field_Output_Scalar)[ScalarTarget];
 		EvaluateFieldNodes<float, FScalarFieldOperator>(InstData->FieldCommands, PhysicsType, FieldContext, SampleResults, SampleMax);
 
 		for (int32 InstanceIdx = 0; InstanceIdx < Context.GetNumInstances(); ++InstanceIdx)
@@ -734,7 +751,7 @@ bool UNiagaraDataInterfacePhysicsField::GetFunctionHLSL(const FNiagaraDataInterf
 	else if (FunctionInfo.DefinitionName == GetPhysicsFieldResolutionName)
 	{
 		static const TCHAR* FormatSample = TEXT(R"(
-		void {InstanceFunctionName}(in float3 SamplePosition, out float3 OutTextureSize)
+		void {InstanceFunctionName}(out float3 OutTextureSize)
 		{
 			{PhysicsFieldContextName}
 			OutTextureSize = DIContext.ClipmapResolution;
@@ -746,7 +763,7 @@ bool UNiagaraDataInterfacePhysicsField::GetFunctionHLSL(const FNiagaraDataInterf
 	else if (FunctionInfo.DefinitionName == GetPhysicsFieldBoundsName)
 	{
 		static const TCHAR* FormatSample = TEXT(R"(
-		void {InstanceFunctionName}(in float3 SamplePosition, out float3 OutMinBounds, out float3 OutMaxBounds)
+		void {InstanceFunctionName}(out float3 OutMinBounds, out float3 OutMaxBounds)
 		{
 			{PhysicsFieldContextName}
 			OutMinBounds = DIContext.ClipmapCenter - DIContext.ClipmapDistance;
@@ -761,6 +778,33 @@ bool UNiagaraDataInterfacePhysicsField::GetFunctionHLSL(const FNiagaraDataInterf
 	return false;
 }
 
+bool UNiagaraDataInterfacePhysicsField::UpgradeFunctionCall(FNiagaraFunctionSignature& FunctionSignature)
+{
+	bool bChanged = false;
+	
+	// upgrade from lwc changes, only parameter types changed there
+	if (FunctionSignature.FunctionVersion < FNiagaraPhysicsFieldDIFunctionVersion::LargeWorldCoordinates)
+	{
+		if (FunctionSignature.Name == SamplePhysicsVectorFieldName && ensure(FunctionSignature.Inputs.Num() == 3))
+		{
+			FunctionSignature.Inputs[1].SetType(FNiagaraTypeDefinition::GetPositionDef());
+			bChanged = true;
+		}
+		if (FunctionSignature.Name == SamplePhysicsScalarFieldName && ensure(FunctionSignature.Inputs.Num() == 3))
+		{
+			FunctionSignature.Inputs[1].SetType(FNiagaraTypeDefinition::GetPositionDef());
+			bChanged = true;
+		}
+		if (FunctionSignature.Name == SamplePhysicsIntegerFieldName && ensure(FunctionSignature.Inputs.Num() == 3))
+		{
+			FunctionSignature.Inputs[1].SetType(FNiagaraTypeDefinition::GetPositionDef());
+			bChanged = true;
+		}
+	}
+
+	return bChanged;
+}
+
 void UNiagaraDataInterfacePhysicsField::GetCommonHLSL(FString& OutHLSL)
 {
 	OutHLSL += TEXT("#include \"/Plugin/Experimental/ChaosNiagara/NiagaraDataInterfacePhysicsField.ush\"\n");
diff --git a/Engine/Plugins/Experimental/Water/Source/Runtime/Private/NiagaraDataInterfaceWater.cpp b/Engine/Plugins/Experimental/Water/Source/Runtime/Private/NiagaraDataInterfaceWater.cpp
--- a/Engine/Plugins/Experimental/Water/Source/Runtime/Private/NiagaraDataInterfaceWater.cpp
+++ b/Engine/Plugins/Experimental/Water/Source/Runtime/Private/NiagaraDataInterfaceWater.cpp
@@ -1,6 +1,8 @@
 // Copyright Epic Games, Inc. All Rights Reserved.
 
 #include "NiagaraDataInterfaceWater.h"
+
+#include "NiagaraSystemInstance.h"
 #include "WaterBodyComponent.h"
 #include "WaterSplineComponent.h"
 #include "WaterModule.h"
@@ -15,6 +17,18 @@ namespace WaterFunctionNames
 	const FName GetWaveParamLookupTableName(TEXT("GetWaveParamLookupTableOffset"));
 }
 
+struct FNiagaraWaterDIFunctionVersion
+{
+	enum Type
+	{
+		InitialVersion = 0,
+		LargeWorldCoordinates = 1,
+
+		VersionPlusOne,
+		LatestVersion = VersionPlusOne - 1
+	};
+};
+
 void UNiagaraDataInterfaceWater::PostInitProperties()
 {
 	Super::PostInitProperties();
@@ -33,18 +47,17 @@ void UNiagaraDataInterfaceWater::GetFunctions(TArray<FNiagaraFunctionSignature>&
 		Sig.Name = WaterFunctionNames::GetWaterDataAtPointName;
 
 		Sig.Inputs.Add(FNiagaraVariable(FNiagaraTypeDefinition(GetClass()), TEXT("Water")));
-		Sig.Inputs.Add(FNiagaraVariable(FNiagaraTypeDefinition::GetVec3Def(), TEXT("WorldPosition")));
+		Sig.Inputs.Add(FNiagaraVariable(FNiagaraTypeDefinition::GetPositionDef(), TEXT("WorldPosition")));
 		Sig.Inputs.Add(FNiagaraVariable(FNiagaraTypeDefinition::GetFloatDef(), TEXT("Time")));
 
 		Sig.Outputs.Add(FNiagaraVariable(FNiagaraTypeDefinition::GetFloatDef(), TEXT("WaveHeight")));
 		Sig.Outputs.Add(FNiagaraVariable(FNiagaraTypeDefinition::GetFloatDef(), TEXT("Depth")));
 		Sig.Outputs.Add(FNiagaraVariable(FNiagaraTypeDefinition::GetVec3Def(), TEXT("Velocity")));
-		Sig.Outputs.Add(FNiagaraVariable(FNiagaraTypeDefinition::GetVec3Def(), TEXT("SurfacePosition")));
+		Sig.Outputs.Add(FNiagaraVariable(FNiagaraTypeDefinition::GetPositionDef(), TEXT("SurfacePosition")));
 		Sig.Outputs.Add(FNiagaraVariable(FNiagaraTypeDefinition::GetVec3Def(), TEXT("SurfaceNormal")));
 
 		Sig.bMemberFunction = true;
 		Sig.bRequiresContext = false;
-		//Sig.Owner = *GetName();
 		Sig.bExperimental = true;
 		Sig.SetDescription(LOCTEXT("DataInterfaceWater_GetWaterDataAtPoint", "Get the water data at the provided world position and time"));
 		OutFunctions.Add(Sig);
@@ -60,7 +73,6 @@ void UNiagaraDataInterfaceWater::GetFunctions(TArray<FNiagaraFunctionSignature>&
 
 		Sig.bMemberFunction = true;
 		Sig.bRequiresContext = false;
-		//Sig.Owner = *GetName();
 		Sig.bExperimental = true;
 		Sig.SetDescription(LOCTEXT("DataInterfaceWater_GetWaveParamLookupTableOffset", "Get the lookup table offset into the wave data texture for the data interface's water body"));
 		OutFunctions.Add(Sig);
@@ -74,11 +86,11 @@ struct FNDIWater_InstanceData
 {
 	//Cached ptr to component we sample from. 
 	TWeakObjectPtr<UWaterBodyComponent> WaterBodyComponent;
+	FNiagaraLWCConverter LWCConverter;
 };
 
 void UNiagaraDataInterfaceWater::GetVMExternalFunction(const FVMExternalFunctionBindingInfo& BindingInfo, void* InstanceData, FVMExternalFunction &OutFunc)
 {
-	FNDIWater_InstanceData* InstData = (FNDIWater_InstanceData*)InstanceData;
 	if (BindingInfo.Name == WaterFunctionNames::GetWaterDataAtPointName)
 	{
 		if(BindingInfo.GetNumInputs() == 5 && BindingInfo.GetNumOutputs() == 11)
@@ -93,7 +105,6 @@ void UNiagaraDataInterfaceWater::GetVMExternalFunction(const FVMExternalFunction
 			NDI_FUNC_BINDER(UNiagaraDataInterfaceWater, GetWaveParamLookupTableOffset)::Bind(this, OutFunc);
 		}
 	}
-
 }
 
 bool UNiagaraDataInterfaceWater::Equals(const UNiagaraDataInterface* Other) const
@@ -106,9 +117,9 @@ bool UNiagaraDataInterfaceWater::Equals(const UNiagaraDataInterface* Other) cons
 	return OtherTyped->SourceBodyComponent == SourceBodyComponent;
 }
 
-bool UNiagaraDataInterfaceWater::CopyTo(UNiagaraDataInterface* Destination) const
+bool UNiagaraDataInterfaceWater::CopyToInternal(UNiagaraDataInterface* Destination) const
 {
-	if (!Super::CopyTo(Destination))
+	if (!Super::CopyToInternal(Destination))
 	{
 		return false;
 	}
@@ -129,20 +140,21 @@ bool UNiagaraDataInterfaceWater::InitPerInstanceData(void* PerInstanceData, FNia
 	FNDIWater_InstanceData* InstData = new (PerInstanceData) FNDIWater_InstanceData();
 
 	InstData->WaterBodyComponent = SourceBodyComponent;
+	InstData->LWCConverter = SystemInstance->GetLWCConverter();
 
 	return true;
 }
 
 void UNiagaraDataInterfaceWater::DestroyPerInstanceData(void* PerInstanceData, FNiagaraSystemInstance* SystemInstance)
 {
-	FNDIWater_InstanceData* InstData = (FNDIWater_InstanceData*)PerInstanceData;
+	FNDIWater_InstanceData* InstData = static_cast<FNDIWater_InstanceData*>(PerInstanceData);
 	InstData->~FNDIWater_InstanceData();
 }
 
 bool UNiagaraDataInterfaceWater::PerInstanceTick(void* PerInstanceData, FNiagaraSystemInstance* SystemInstance, float DeltaSeconds)
 {
 	check(SystemInstance);
-	FNDIWater_InstanceData* InstData = (FNDIWater_InstanceData*)PerInstanceData;
+	FNDIWater_InstanceData* InstData = static_cast<FNDIWater_InstanceData*>(PerInstanceData);
 
 	if (InstData->WaterBodyComponent != SourceBodyComponent)
 	{
@@ -152,34 +164,42 @@ bool UNiagaraDataInterfaceWater::PerInstanceTick(void* PerInstanceData, FNiagara
 	return false;
 }
 
+#if WITH_EDITORONLY_DATA
+bool UNiagaraDataInterfaceWater::UpgradeFunctionCall(FNiagaraFunctionSignature& FunctionSignature)
+{
+	bool bChanged = false;
+	
+	// upgrade from lwc changes, only parameter types changed there
+	if (FunctionSignature.FunctionVersion < FNiagaraWaterDIFunctionVersion::LargeWorldCoordinates)
+	{
+		if (FunctionSignature.Name == WaterFunctionNames::GetWaterDataAtPointName && ensure(FunctionSignature.Inputs.Num() == 3) && ensure(FunctionSignature.Outputs.Num() == 5))
+		{
+			FunctionSignature.Inputs[1].SetType(FNiagaraTypeDefinition::GetPositionDef());
+			FunctionSignature.Outputs[3].SetType(FNiagaraTypeDefinition::GetPositionDef());
+			bChanged = true;
+		}
+	}
+
+	return bChanged;
+}
+#endif
+
 void UNiagaraDataInterfaceWater::GetWaterDataAtPoint(FVectorVMExternalFunctionContext& Context)
 {
 	QUICK_SCOPE_CYCLE_COUNTER(NiagaraDataInterfaceWater_GetWaterDataAtPoint);
 
 	VectorVM::FUserPtrHandler<FNDIWater_InstanceData> InstData(Context);
-
+	
 	// Inputs
-	VectorVM::FExternalFuncInputHandler<float> WorldX(Context);
-	VectorVM::FExternalFuncInputHandler<float> WorldY(Context);
-	VectorVM::FExternalFuncInputHandler<float> WorldZ(Context);
-	VectorVM::FExternalFuncInputHandler<float> Time(Context);
+	FNDIInputParam<FNiagaraPosition> WorldPos(Context);
+	FNDIInputParam<float> Time(Context);
 
 	// Outputs
-	VectorVM::FExternalFuncRegisterHandler<float> OutHeight(Context);
-
-	VectorVM::FExternalFuncRegisterHandler<float> OutDepth(Context);
-
-	VectorVM::FExternalFuncRegisterHandler<float> OutVelocityX(Context);
-	VectorVM::FExternalFuncRegisterHandler<float> OutVelocityY(Context);
-	VectorVM::FExternalFuncRegisterHandler<float> OutVelocityZ(Context);
-
-	VectorVM::FExternalFuncRegisterHandler<float> OutSurfaceX(Context);
-	VectorVM::FExternalFuncRegisterHandler<float> OutSurfaceY(Context);
-	VectorVM::FExternalFuncRegisterHandler<float> OutSurfaceZ(Context);
-
-	VectorVM::FExternalFuncRegisterHandler<float> OutSurfaceNormalX(Context);
-	VectorVM::FExternalFuncRegisterHandler<float> OutSurfaceNormalY(Context);
-	VectorVM::FExternalFuncRegisterHandler<float> OutSurfaceNormalZ(Context);
+	FNDIOutputParam<float> OutHeight(Context);
+	FNDIOutputParam<float> OutDepth(Context);
+	FNDIOutputParam<FVector3f> OutVelocity(Context);
+	FNDIOutputParam<FNiagaraPosition> OutSurfacePos(Context);
+	FNDIOutputParam<FVector3f> OutSurfaceNormal(Context);
 
 	UWaterBodyComponent* Component = InstData->WaterBodyComponent.Get();
 	if (Component == nullptr)
@@ -194,8 +214,8 @@ void UNiagaraDataInterfaceWater::GetWaterDataAtPoint(FVectorVMExternalFunctionCo
 		bool bIsValid = false;
 		if (Component != nullptr)
 		{
-			FVector WorldPos(WorldX.Get(), WorldY.Get(), WorldZ.Get());
-			QueryResult = Component->QueryWaterInfoClosestToWorldLocation(WorldPos,
+			FVector QueryPos = InstData->LWCConverter.ConvertSimulationPositionToWorld(WorldPos.GetAndAdvance());
+			QueryResult = Component->QueryWaterInfoClosestToWorldLocation(QueryPos,
 				EWaterBodyQueryFlags::ComputeLocation
 				| EWaterBodyQueryFlags::ComputeVelocity
 				| EWaterBodyQueryFlags::ComputeNormal
@@ -204,29 +224,17 @@ void UNiagaraDataInterfaceWater::GetWaterDataAtPoint(FVectorVMExternalFunctionCo
 			bIsValid = !QueryResult.IsInExclusionVolume();
 		}
 
-		*OutHeight.GetDestAndAdvance() = bIsValid ? QueryResult.GetWaveInfo().Height : 0.0f;
-		*OutDepth.GetDestAndAdvance() = bIsValid ? QueryResult.GetWaterSurfaceDepth() : 0.0f;
-
-		const FVector& Velocity = bIsValid ? QueryResult.GetVelocity() : FVector::ZeroVector;
-		*OutVelocityX.GetDestAndAdvance() = Velocity.X;
-		*OutVelocityY.GetDestAndAdvance() = Velocity.Y;
-		*OutVelocityZ.GetDestAndAdvance() = Velocity.Z;
+		OutHeight.SetAndAdvance(bIsValid ? QueryResult.GetWaveInfo().Height : 0.0f);
+		OutDepth.SetAndAdvance(bIsValid ? QueryResult.GetWaterSurfaceDepth() : 0.0f);
+		OutVelocity.SetAndAdvance( bIsValid ? QueryResult.GetVelocity() : FVector::ZeroVector);
 
 		// Note we assume X and Y are in water by the time this is queried
 		const FVector& AdjustedSurfaceLoc = bIsValid ? QueryResult.GetWaterSurfaceLocation() : FVector::ZeroVector;
-		*OutSurfaceX.GetDestAndAdvance() =  AdjustedSurfaceLoc.X;
-		*OutSurfaceY.GetDestAndAdvance() =  AdjustedSurfaceLoc.Y;
-		*OutSurfaceZ.GetDestAndAdvance() =  AdjustedSurfaceLoc.Z;
-
-		const FVector& Normal = bIsValid ? QueryResult.GetWaterSurfaceNormal() : FVector::UpVector;
-		*OutSurfaceNormalX.GetDestAndAdvance() = Normal.X;
-		*OutSurfaceNormalY.GetDestAndAdvance() = Normal.Y;
-		*OutSurfaceNormalZ.GetDestAndAdvance() = Normal.Z;
-
-		WorldX.Advance();
-		WorldY.Advance();
-		WorldZ.Advance();
-		Time.Advance();
+		OutSurfacePos.SetAndAdvance( InstData->LWCConverter.ConvertWorldToSimulationPosition(AdjustedSurfaceLoc));
+
+		OutSurfaceNormal.SetAndAdvance(bIsValid ? QueryResult.GetWaterSurfaceNormal() : FVector::UpVector);
+
+		Time.GetAndAdvance();
 	}
 }
 
diff --git a/Engine/Plugins/Experimental/Water/Source/Runtime/Public/NiagaraDataInterfaceWater.h b/Engine/Plugins/Experimental/Water/Source/Runtime/Public/NiagaraDataInterfaceWater.h
--- a/Engine/Plugins/Experimental/Water/Source/Runtime/Public/NiagaraDataInterfaceWater.h
+++ b/Engine/Plugins/Experimental/Water/Source/Runtime/Public/NiagaraDataInterfaceWater.h
@@ -19,14 +19,19 @@ public:
 	/** UNiagaraDataInterface interface */
 	virtual void GetFunctions(TArray<FNiagaraFunctionSignature>& OutFunctions) override;
 	virtual void GetVMExternalFunction(const FVMExternalFunctionBindingInfo& BindingInfo, void* InstanceData, FVMExternalFunction &OutFunc) override;
-	virtual bool Equals(const UNiagaraDataInterface* Other) const;
-	virtual bool CopyTo(UNiagaraDataInterface* Destination) const;
+	virtual bool Equals(const UNiagaraDataInterface* Other) const override;
+	virtual bool CopyToInternal(UNiagaraDataInterface* Destination) const override;
 
 	virtual int32 PerInstanceDataSize() const override;
-	virtual bool InitPerInstanceData(void* PerInstanceData, FNiagaraSystemInstance* SystemInstance);
-	virtual void DestroyPerInstanceData(void* PerInstanceData, FNiagaraSystemInstance* SystemInstance);
+	virtual bool InitPerInstanceData(void* PerInstanceData, FNiagaraSystemInstance* SystemInstance) override;
+	virtual void DestroyPerInstanceData(void* PerInstanceData, FNiagaraSystemInstance* SystemInstance) override;
 	virtual bool PerInstanceTick(void* PerInstanceData, FNiagaraSystemInstance* SystemInstance, float DeltaSeconds) override;
 	virtual bool HasPreSimulateTick() const override { return true; }
+	virtual bool CanExecuteOnTarget(ENiagaraSimTarget Target) const override { return Target == ENiagaraSimTarget::CPUSim; }
+
+#if WITH_EDITORONLY_DATA
+	virtual bool UpgradeFunctionCall(FNiagaraFunctionSignature& FunctionSignature) override;
+#endif
 
 	void GetWaterDataAtPoint(FVectorVMExternalFunctionContext& Context);
 
diff --git a/Engine/Plugins/FX/Niagara/Source/Niagara/Classes/NiagaraDataInterface.h b/Engine/Plugins/FX/Niagara/Source/Niagara/Classes/NiagaraDataInterface.h
--- a/Engine/Plugins/FX/Niagara/Source/Niagara/Classes/NiagaraDataInterface.h
+++ b/Engine/Plugins/FX/Niagara/Source/Niagara/Classes/NiagaraDataInterface.h
@@ -270,7 +270,7 @@ class NIAGARA_API UNiagaraDataInterface : public UNiagaraDataInterfaceBase
 	GENERATED_UCLASS_BODY()
 
 public:
-	virtual ~UNiagaraDataInterface();
+	virtual ~UNiagaraDataInterface() override;
 
 	// UObject Interface
 	virtual void PostLoad() override;
diff --git a/Engine/Plugins/FX/Niagara/Source/Niagara/Classes/NiagaraDataInterfaceCollisionQuery.h b/Engine/Plugins/FX/Niagara/Source/Niagara/Classes/NiagaraDataInterfaceCollisionQuery.h
--- a/Engine/Plugins/FX/Niagara/Source/Niagara/Classes/NiagaraDataInterfaceCollisionQuery.h
+++ b/Engine/Plugins/FX/Niagara/Source/Niagara/Classes/NiagaraDataInterfaceCollisionQuery.h
@@ -20,7 +20,7 @@ struct CQDIPerInstanceData
 };
 
 
-/** Data Interface allowing sampling of color curves. */
+/** Data Interface that can be used to query collision related data, like geometry traces or sampling the depth buffer. */
 UCLASS(EditInlineNew, Category = "Collision", meta = (DisplayName = "Collision Query"))
 class NIAGARA_API UNiagaraDataInterfaceCollisionQuery : public UNiagaraDataInterface
 {
diff --git a/Engine/Plugins/FX/Niagara/Source/Niagara/Private/DataInterface/NiagaraDistanceFieldParameters.h b/Engine/Plugins/FX/Niagara/Source/Niagara/Private/DataInterface/NiagaraDistanceFieldParameters.h
--- a/Engine/Plugins/FX/Niagara/Source/Niagara/Private/DataInterface/NiagaraDistanceFieldParameters.h
+++ b/Engine/Plugins/FX/Niagara/Source/Niagara/Private/DataInterface/NiagaraDistanceFieldParameters.h
@@ -1,6 +1,7 @@
 // Copyright Epic Games, Inc. All Rights Reserved.
 
 #pragma once
+#include "Renderer/Private/DistanceFieldLightingShared.h"
 
 class FDistanceFieldParameters
 {
@@ -99,4 +100,4 @@ private:
 		LAYOUT_FIELD(FShaderParameter, DistanceFieldBrickAtlasHalfTexelSize)
 		LAYOUT_FIELD(FShaderParameter, DistanceFieldBrickOffsetToAtlasUVScale)
 		LAYOUT_FIELD(FShaderParameter, DistanceFieldUniqueDataBrickSizeInAtlasTexels)
-};
\ No newline at end of file
+};
diff --git a/Engine/Plugins/FX/Niagara/Source/Niagara/Private/NiagaraDataInterfaceRigidMeshCollisionQuery.cpp b/Engine/Plugins/FX/Niagara/Source/Niagara/Private/NiagaraDataInterfaceRigidMeshCollisionQuery.cpp
--- a/Engine/Plugins/FX/Niagara/Source/Niagara/Private/NiagaraDataInterfaceRigidMeshCollisionQuery.cpp
+++ b/Engine/Plugins/FX/Niagara/Source/Niagara/Private/NiagaraDataInterfaceRigidMeshCollisionQuery.cpp
@@ -4,7 +4,6 @@
 #include "Animation/SkeletalMeshActor.h"
 #include "Components/SkeletalMeshComponent.h"
 #include "SkeletalRenderPublic.h"
-#include "SkeletalMeshTypes.h"
 #include "AnimationRuntime.h"
 #include "NiagaraShader.h"
 #include "NiagaraComponent.h"
@@ -15,7 +14,6 @@
 #include "EngineUtils.h"
 #include "Renderer/Private/ScenePrivate.h"
 #include "DistanceFieldAtlas.h"
-#include "Renderer/Private/DistanceFieldLightingShared.h"
 #include "NiagaraGpuComputeDispatchInterface.h"
 #include "NiagaraGpuComputeDispatch.h"
 #include "DataInterface/NiagaraDistanceFieldParameters.h"
@@ -23,6 +21,18 @@
 #define LOCTEXT_NAMESPACE "NiagaraDataInterfaceRigidMeshCollisionQuery"
 DEFINE_LOG_CATEGORY_STATIC(LogRigidMeshCollision, Log, All);
 
+struct FNiagaraRigidMeshCollisionDIFunctionVersion
+{
+	enum Type
+	{
+		InitialVersion = 0,
+		LargeWorldCoordinates = 1,
+
+		VersionPlusOne,
+		LatestVersion = VersionPlusOne - 1
+	};
+};
+
 //------------------------------------------------------------------------------------------------------------
 
 static const FName GetNumBoxesName(TEXT("GetNumBoxes"));
@@ -118,8 +128,7 @@ void UpdateInternalBuffer(const TArray<BufferType>& InputData, FRWBuffer& Output
 	}
 }
 
-void FillCurrentTransforms(const FTransform& ElementTransform, uint32& ElementCount,
-	TArray<FVector4f>& OutCurrentTransform, TArray<FVector4f>& OutCurrentInverse)
+void FillCurrentTransforms(const FTransform& ElementTransform, uint32& ElementCount, TArray<FVector4f>& OutCurrentTransform, TArray<FVector4f>& OutCurrentInverse)
 {
 	// LWC_TODO: precision loss
 	const uint32 ElementOffset = 3 * ElementCount;
@@ -196,8 +205,7 @@ void CompactInternalArrays(FNDIRigidMeshCollisionArrays* OutAssetArrays)
 	}
 }
 
-void CreateInternalArrays(TArray<AActor*> Actors,
-	FNDIRigidMeshCollisionArrays* OutAssetArrays)
+void CreateInternalArrays(TArray<AActor*> Actors, FNDIRigidMeshCollisionArrays* OutAssetArrays, FVector LWCTile)
 {
 	if (OutAssetArrays != nullptr)
 	{
@@ -219,9 +227,6 @@ void CreateInternalArrays(TArray<AActor*> Actors,
 			OutAssetArrays->ElementOffsets.CapsuleOffset = OutAssetArrays->ElementOffsets.SphereOffset + NumSpheres;
 			OutAssetArrays->ElementOffsets.NumElements = OutAssetArrays->ElementOffsets.CapsuleOffset + NumCapsules;
 
-			const uint32 NumTransforms = OutAssetArrays->ElementOffsets.NumElements * 2;
-			const uint32 NumExtents = OutAssetArrays->ElementOffsets.NumElements;
-
 			uint32 BoxCount = OutAssetArrays->ElementOffsets.BoxOffset;
 			uint32 SphereCount = OutAssetArrays->ElementOffsets.SphereOffset;
 			uint32 CapsuleCount = OutAssetArrays->ElementOffsets.CapsuleOffset;
@@ -237,7 +242,8 @@ void CreateInternalArrays(TArray<AActor*> Actors,
 					bool FoundCollisionShapes = false;
 					if (BodySetup != nullptr)
 					{
-						const FTransform MeshTransform = Actor->GetTransform();
+						FTransform MeshTransform = Actor->GetTransform();
+						MeshTransform.AddToTranslation(LWCTile * -FLargeWorldRenderScalar::GetTileSize());
 
 						for (const FKConvexElem& ConvexElem : BodySetup->AggGeom.ConvexElems)
 						{							
@@ -252,7 +258,7 @@ void CreateInternalArrays(TArray<AActor*> Actors,
 								OutAssetArrays->PhysicsType[BoxCount] = (ConvexElem.GetCollisionEnabled() == ECollisionEnabled::QueryAndPhysics);
 								OutAssetArrays->SourceSceneProxy[BoxCount] = StaticMeshComponent->SceneProxy;
 								
-								const FTransform ElementTransform = FTransform((FVector)Center) * MeshTransform;
+								const FTransform ElementTransform = FTransform(Center) * MeshTransform;
 								OutAssetArrays->ElementExtent[BoxCount] = FVector4f(Extent.X, Extent.Y, Extent.Z, 0);
 								FillCurrentTransforms(ElementTransform, BoxCount, OutAssetArrays->CurrentTransform, OutAssetArrays->CurrentInverse);
 
@@ -324,7 +330,7 @@ void CreateInternalArrays(TArray<AActor*> Actors,
 	}
 }
 
-void UpdateInternalArrays(const TArray<AActor*> &Actors, FNDIRigidMeshCollisionArrays* OutAssetArrays)
+void UpdateInternalArrays(const TArray<AActor*> &Actors, FNDIRigidMeshCollisionArrays* OutAssetArrays, FVector LWCTile)
 {
 	if (OutAssetArrays != nullptr && OutAssetArrays->ElementOffsets.NumElements < OutAssetArrays->MaxPrimitives)
 	{
@@ -338,7 +344,7 @@ void UpdateInternalArrays(const TArray<AActor*> &Actors, FNDIRigidMeshCollisionA
 			((OutAssetArrays->ElementOffsets.CapsuleOffset - OutAssetArrays->ElementOffsets.SphereOffset) != NumSpheres) ||
 			((OutAssetArrays->ElementOffsets.NumElements - OutAssetArrays->ElementOffsets.CapsuleOffset) != NumCapsules))
 		{
-			CreateInternalArrays(Actors, OutAssetArrays);
+			CreateInternalArrays(Actors, OutAssetArrays, LWCTile);
 		}
 
 		OutAssetArrays->PreviousTransform = OutAssetArrays->CurrentTransform;
@@ -359,7 +365,8 @@ void UpdateInternalArrays(const TArray<AActor*> &Actors, FNDIRigidMeshCollisionA
 			UBodySetup* BodySetup = StaticMeshComponent != nullptr ? StaticMeshComponent->GetBodySetup() : nullptr;
 			if (BodySetup != nullptr)
 			{
-				const FTransform MeshTransform = Actor->GetTransform();
+				FTransform MeshTransform = Actor->GetTransform();
+				MeshTransform.AddToTranslation(LWCTile * -FLargeWorldRenderScalar::GetTileSize());
 
 				for (const FKConvexElem& ConvexElem : BodySetup->AggGeom.ConvexElems)
 				{
@@ -369,7 +376,7 @@ void UpdateInternalArrays(const TArray<AActor*> &Actors, FNDIRigidMeshCollisionA
 						FVector3f Center = FVector3f(BBox.Max + BBox.Min) * .5;
 
 						OutAssetArrays->SourceSceneProxy[BoxCount] = StaticMeshComponent->SceneProxy;
-						const FTransform ElementTransform = FTransform((FVector)Center) * MeshTransform;
+						const FTransform ElementTransform = FTransform(Center) * MeshTransform;
 						FillCurrentTransforms(ElementTransform, BoxCount, OutAssetArrays->CurrentTransform, OutAssetArrays->CurrentInverse);
 					}
 				}
@@ -516,7 +523,7 @@ void FNDIRigidMeshCollisionData::Update(UNiagaraDataInterfaceRigidMeshCollisionQ
 
 		if (0 < Actors.Num() && Actors[0] != nullptr)
 		{
-			UpdateInternalArrays(Actors, AssetArrays);
+			UpdateInternalArrays(Actors, AssetArrays, SystemInstance->GetLWCTile());
 		}
 	}
 }
@@ -537,7 +544,7 @@ ETickingGroup FNDIRigidMeshCollisionData::ComputeTickingGroup()
 			}
 				const ETickingGroup ComponentTickGroup = FMath::Max(Component->PrimaryComponentTick.TickGroup, Component->PrimaryComponentTick.EndTickGroup);
 				const ETickingGroup PhysicsTickGroup = ComponentTickGroup;
-				const ETickingGroup ClampedTickGroup = FMath::Clamp(ETickingGroup(PhysicsTickGroup + 1), NiagaraFirstTickGroup, NiagaraLastTickGroup);		
+				const ETickingGroup ClampedTickGroup = FMath::Clamp(static_cast<ETickingGroup>(PhysicsTickGroup + 1), NiagaraFirstTickGroup, NiagaraLastTickGroup);		
 
 			TickingGroup = FMath::Max(TickingGroup, ClampedTickGroup);
 		}
@@ -726,7 +733,7 @@ void FNDIRigidMeshCollisionProxy::InitializePerInstanceData(const FNiagaraSystem
 	check(IsInRenderingThread());
 
 	check(!SystemInstancesToProxyData.Contains(SystemInstance));
-	FNDIRigidMeshCollisionData* TargetData = &SystemInstancesToProxyData.Add(SystemInstance);
+	SystemInstancesToProxyData.Add(SystemInstance);
 }
 
 void FNDIRigidMeshCollisionProxy::DestroyPerInstanceData(const FNiagaraSystemInstanceID& SystemInstance)
@@ -788,9 +795,7 @@ bool UNiagaraDataInterfaceRigidMeshCollisionQuery::InitPerInstanceData(void* Per
 
 ETickingGroup UNiagaraDataInterfaceRigidMeshCollisionQuery::CalculateTickGroup(const void* PerInstanceData) const
 {
-	const FNDIRigidMeshCollisionData* InstanceData = static_cast<const FNDIRigidMeshCollisionData*>(PerInstanceData);
-
-	if (InstanceData)
+	if (const FNDIRigidMeshCollisionData* InstanceData = static_cast<const FNDIRigidMeshCollisionData*>(PerInstanceData))
 	{
 		return InstanceData->TickingGroup;
 	}
@@ -872,10 +877,11 @@ void UNiagaraDataInterfaceRigidMeshCollisionQuery::GetFunctions(TArray<FNiagaraF
 	{
 		FNiagaraFunctionSignature Sig;
 		Sig.Name = GetNumBoxesName;
+		Sig.SetDescription(LOCTEXT("GetNumBoxesNameDescription", "Returns the number of box primitives for the collection of static meshes the DI represents."));
+		Sig.SetFunctionVersion(FNiagaraRigidMeshCollisionDIFunctionVersion::LatestVersion);
 		Sig.bSupportsGPU = true;
 		Sig.bSupportsCPU = false;
 		Sig.bMemberFunction = true;
-		Sig.bRequiresContext = false;
 		Sig.Inputs.Add(FNiagaraVariable(FNiagaraTypeDefinition(GetClass()), TEXT("Collision DI")));
 		Sig.Outputs.Add(FNiagaraVariable(FNiagaraTypeDefinition::GetIntDef(), TEXT("Num Boxes")));
 
@@ -884,10 +890,11 @@ void UNiagaraDataInterfaceRigidMeshCollisionQuery::GetFunctions(TArray<FNiagaraF
 	{
 		FNiagaraFunctionSignature Sig;
 		Sig.Name = GetNumSpheresName;
+		Sig.SetDescription(LOCTEXT("GetNumSpheresNameDescription", "Returns the number of sphere primitives for the collection of static meshes the DI represents."));
+		Sig.SetFunctionVersion(FNiagaraRigidMeshCollisionDIFunctionVersion::LatestVersion);
 		Sig.bSupportsGPU = true;
 		Sig.bSupportsCPU = false;
 		Sig.bMemberFunction = true;
-		Sig.bRequiresContext = false;
 		Sig.Inputs.Add(FNiagaraVariable(FNiagaraTypeDefinition(GetClass()), TEXT("Collision DI")));
 		Sig.Outputs.Add(FNiagaraVariable(FNiagaraTypeDefinition::GetIntDef(), TEXT("Num Spheres")));
 
@@ -896,10 +903,11 @@ void UNiagaraDataInterfaceRigidMeshCollisionQuery::GetFunctions(TArray<FNiagaraF
 	{
 		FNiagaraFunctionSignature Sig;
 		Sig.Name = GetNumCapsulesName;
+		Sig.SetDescription(LOCTEXT("GetNumCapsulesNameDescription", "Returns the number of capsule primitives for the collection of static meshes the DI represents."));
+		Sig.SetFunctionVersion(FNiagaraRigidMeshCollisionDIFunctionVersion::LatestVersion);
 		Sig.bSupportsGPU = true;
 		Sig.bSupportsCPU = false;
 		Sig.bMemberFunction = true;
-		Sig.bRequiresContext = false;
 		Sig.Inputs.Add(FNiagaraVariable(FNiagaraTypeDefinition(GetClass()), TEXT("Collision DI")));
 		Sig.Outputs.Add(FNiagaraVariable(FNiagaraTypeDefinition::GetIntDef(), TEXT("Num Capsules")));
 
@@ -908,16 +916,17 @@ void UNiagaraDataInterfaceRigidMeshCollisionQuery::GetFunctions(TArray<FNiagaraF
 	{
 		FNiagaraFunctionSignature Sig;
 		Sig.Name = GetClosestPointName;
+		Sig.SetDescription(LOCTEXT("GetClosestPointDescription", "Given a world space position, computes the static mesh's closest point. Also returns normal and velocity for that point."));
+		Sig.SetFunctionVersion(FNiagaraRigidMeshCollisionDIFunctionVersion::LatestVersion);
 		Sig.bSupportsGPU = true;
 		Sig.bSupportsCPU = false;
 		Sig.bMemberFunction = true;
-		Sig.bRequiresContext = false;
 		Sig.Inputs.Add(FNiagaraVariable(FNiagaraTypeDefinition(GetClass()), TEXT("Collision DI")));
-		Sig.Inputs.Add(FNiagaraVariable(FNiagaraTypeDefinition::GetVec3Def(), TEXT("World Position")));
+		Sig.Inputs.Add(FNiagaraVariable(FNiagaraTypeDefinition::GetPositionDef(), TEXT("World Position")));
 		Sig.Inputs.Add(FNiagaraVariable(FNiagaraTypeDefinition::GetFloatDef(), TEXT("Delta Time")));
 		Sig.Inputs.Add(FNiagaraVariable(FNiagaraTypeDefinition::GetFloatDef(), TEXT("Time Fraction")));
 		Sig.Outputs.Add(FNiagaraVariable(FNiagaraTypeDefinition::GetFloatDef(), TEXT("Closest Distance")));
-		Sig.Outputs.Add(FNiagaraVariable(FNiagaraTypeDefinition::GetVec3Def(), TEXT("Closest Position")));
+		Sig.Outputs.Add(FNiagaraVariable(FNiagaraTypeDefinition::GetPositionDef(), TEXT("Closest Position")));
 		Sig.Outputs.Add(FNiagaraVariable(FNiagaraTypeDefinition::GetVec3Def(), TEXT("Closest Normal")));
 		Sig.Outputs.Add(FNiagaraVariable(FNiagaraTypeDefinition::GetVec3Def(), TEXT("Closest Velocity")));
 
@@ -926,12 +935,13 @@ void UNiagaraDataInterfaceRigidMeshCollisionQuery::GetFunctions(TArray<FNiagaraF
 	{
 		FNiagaraFunctionSignature Sig;
 		Sig.Name = GetClosestElementName;
+		Sig.SetDescription(LOCTEXT("GetClosestElementDescription", "Given a world space position, computes the static mesh's closest element."));
+		Sig.SetFunctionVersion(FNiagaraRigidMeshCollisionDIFunctionVersion::LatestVersion);
 		Sig.bSupportsGPU = true;
 		Sig.bSupportsCPU = false;
 		Sig.bMemberFunction = true;
-		Sig.bRequiresContext = false;
 		Sig.Inputs.Add(FNiagaraVariable(FNiagaraTypeDefinition(GetClass()), TEXT("Collision DI")));
-		Sig.Inputs.Add(FNiagaraVariable(FNiagaraTypeDefinition::GetVec3Def(), TEXT("World Position")));
+		Sig.Inputs.Add(FNiagaraVariable(FNiagaraTypeDefinition::GetPositionDef(), TEXT("World Position")));
 		Sig.Inputs.Add(FNiagaraVariable(FNiagaraTypeDefinition::GetFloatDef(), TEXT("Time Fraction")));
 		Sig.Outputs.Add(FNiagaraVariable(FNiagaraTypeDefinition::GetIntDef(), TEXT("Closest Element")));
 
@@ -940,16 +950,17 @@ void UNiagaraDataInterfaceRigidMeshCollisionQuery::GetFunctions(TArray<FNiagaraF
 	{
 		FNiagaraFunctionSignature Sig;
 		Sig.Name = GetElementPointName;
+		Sig.SetDescription(LOCTEXT("GetClosestElementPointDescription", "Given a world space position and an element index, computes the static mesh's closest point. Also returns normal and velocity for that point."));
+		Sig.SetFunctionVersion(FNiagaraRigidMeshCollisionDIFunctionVersion::LatestVersion);
 		Sig.bSupportsGPU = true;
 		Sig.bSupportsCPU = false;
 		Sig.bMemberFunction = true;
-		Sig.bRequiresContext = false;
 		Sig.Inputs.Add(FNiagaraVariable(FNiagaraTypeDefinition(GetClass()), TEXT("Collision DI")));
-		Sig.Inputs.Add(FNiagaraVariable(FNiagaraTypeDefinition::GetVec3Def(), TEXT("World Position")));
+		Sig.Inputs.Add(FNiagaraVariable(FNiagaraTypeDefinition::GetPositionDef(), TEXT("World Position")));
 		Sig.Inputs.Add(FNiagaraVariable(FNiagaraTypeDefinition::GetFloatDef(), TEXT("Delta Time")));
 		Sig.Inputs.Add(FNiagaraVariable(FNiagaraTypeDefinition::GetFloatDef(), TEXT("Time Fraction")));
 		Sig.Inputs.Add(FNiagaraVariable(FNiagaraTypeDefinition::GetIntDef(), TEXT("Element Index")));
-		Sig.Outputs.Add(FNiagaraVariable(FNiagaraTypeDefinition::GetVec3Def(), TEXT("Closest Position")));
+		Sig.Outputs.Add(FNiagaraVariable(FNiagaraTypeDefinition::GetPositionDef(), TEXT("Closest Position")));
 		Sig.Outputs.Add(FNiagaraVariable(FNiagaraTypeDefinition::GetVec3Def(), TEXT("Closest Normal")));
 		Sig.Outputs.Add(FNiagaraVariable(FNiagaraTypeDefinition::GetVec3Def(), TEXT("Closest Velocity")));
 
@@ -958,12 +969,13 @@ void UNiagaraDataInterfaceRigidMeshCollisionQuery::GetFunctions(TArray<FNiagaraF
 	{
 		FNiagaraFunctionSignature Sig;
 		Sig.Name = GetElementDistanceName;
+		Sig.SetDescription(LOCTEXT("GetElementDistanceDescription", "Given a world space position and element index, computes the distance to the closest point for the static mesh."));
+		Sig.SetFunctionVersion(FNiagaraRigidMeshCollisionDIFunctionVersion::LatestVersion);
 		Sig.bSupportsGPU = true;
 		Sig.bSupportsCPU = false;
 		Sig.bMemberFunction = true;
-		Sig.bRequiresContext = false;
 		Sig.Inputs.Add(FNiagaraVariable(FNiagaraTypeDefinition(GetClass()), TEXT("Collision DI")));
-		Sig.Inputs.Add(FNiagaraVariable(FNiagaraTypeDefinition::GetVec3Def(), TEXT("World Position")));
+		Sig.Inputs.Add(FNiagaraVariable(FNiagaraTypeDefinition::GetPositionDef(), TEXT("World Position")));
 		Sig.Inputs.Add(FNiagaraVariable(FNiagaraTypeDefinition::GetFloatDef(), TEXT("Time Fraction")));
 		Sig.Inputs.Add(FNiagaraVariable(FNiagaraTypeDefinition::GetIntDef(), TEXT("Element Index")));
 		Sig.Outputs.Add(FNiagaraVariable(FNiagaraTypeDefinition::GetFloatDef(), TEXT("Closest Distance")));
@@ -973,12 +985,13 @@ void UNiagaraDataInterfaceRigidMeshCollisionQuery::GetFunctions(TArray<FNiagaraF
 	{
 		FNiagaraFunctionSignature Sig;
 		Sig.Name = GetClosestDistanceName;
+		Sig.SetDescription(LOCTEXT("GetClosestDistanceDescription", "Given a world space position, computes the distance to the closest point for the static mesh."));
+		Sig.SetFunctionVersion(FNiagaraRigidMeshCollisionDIFunctionVersion::LatestVersion);
 		Sig.bSupportsGPU = true;
 		Sig.bSupportsCPU = false;
 		Sig.bMemberFunction = true;
-		Sig.bRequiresContext = false;
 		Sig.Inputs.Add(FNiagaraVariable(FNiagaraTypeDefinition(GetClass()), TEXT("Collision DI")));
-		Sig.Inputs.Add(FNiagaraVariable(FNiagaraTypeDefinition::GetVec3Def(), TEXT("World Position")));
+		Sig.Inputs.Add(FNiagaraVariable(FNiagaraTypeDefinition::GetPositionDef(), TEXT("World Position")));
 		Sig.Inputs.Add(FNiagaraVariable(FNiagaraTypeDefinition::GetFloatDef(), TEXT("Time Fraction")));
 		Sig.Outputs.Add(FNiagaraVariable(FNiagaraTypeDefinition::GetFloatDef(), TEXT("Closest Distance")));
 
@@ -988,16 +1001,17 @@ void UNiagaraDataInterfaceRigidMeshCollisionQuery::GetFunctions(TArray<FNiagaraF
 	{
 		FNiagaraFunctionSignature Sig;
 		Sig.Name = GetClosestPointMeshDistanceFieldName;
+		Sig.SetDescription(LOCTEXT("GetClosestPointMeshDistanceFieldDescription", "Given a world space position, computes the distance to the closest point for the static mesh, using the mesh's distance field."));
+		Sig.SetFunctionVersion(FNiagaraRigidMeshCollisionDIFunctionVersion::LatestVersion);
 		Sig.bSupportsGPU = true;
 		Sig.bSupportsCPU = false;
 		Sig.bMemberFunction = true;
-		Sig.bRequiresContext = false;
 		Sig.Inputs.Add(FNiagaraVariable(FNiagaraTypeDefinition(GetClass()), TEXT("Collision DI")));
-		Sig.Inputs.Add(FNiagaraVariable(FNiagaraTypeDefinition::GetVec3Def(), TEXT("World Position")));
+		Sig.Inputs.Add(FNiagaraVariable(FNiagaraTypeDefinition::GetPositionDef(), TEXT("World Position")));
 		Sig.Inputs.Add(FNiagaraVariable(FNiagaraTypeDefinition::GetFloatDef(), TEXT("Delta Time")));
 		Sig.Inputs.Add(FNiagaraVariable(FNiagaraTypeDefinition::GetFloatDef(), TEXT("Time Fraction")));
 		Sig.Outputs.Add(FNiagaraVariable(FNiagaraTypeDefinition::GetFloatDef(), TEXT("Closest Distance")));
-		Sig.Outputs.Add(FNiagaraVariable(FNiagaraTypeDefinition::GetVec3Def(), TEXT("Closest Position")));
+		Sig.Outputs.Add(FNiagaraVariable(FNiagaraTypeDefinition::GetPositionDef(), TEXT("Closest Position")));
 		Sig.Outputs.Add(FNiagaraVariable(FNiagaraTypeDefinition::GetVec3Def(), TEXT("Closest Normal")));
 		Sig.Outputs.Add(FNiagaraVariable(FNiagaraTypeDefinition::GetVec3Def(), TEXT("Closest Velocity")));
 
@@ -1007,16 +1021,17 @@ void UNiagaraDataInterfaceRigidMeshCollisionQuery::GetFunctions(TArray<FNiagaraF
 	{
 		FNiagaraFunctionSignature Sig;
 		Sig.Name = GetClosestPointMeshDistanceFieldNoNormalName;
+		Sig.SetDescription(LOCTEXT("GetClosestPointMeshDistanceFieldNNDescription", "Given a world space position, computes the distance to the closest point for the static mesh, using the mesh's distance field.\nSkips the normal calculation and is more performant than it's counterpart with normal."));
+		Sig.SetFunctionVersion(FNiagaraRigidMeshCollisionDIFunctionVersion::LatestVersion);
 		Sig.bSupportsGPU = true;
 		Sig.bSupportsCPU = false;
 		Sig.bMemberFunction = true;
-		Sig.bRequiresContext = false;
 		Sig.Inputs.Add(FNiagaraVariable(FNiagaraTypeDefinition(GetClass()), TEXT("Collision DI")));
-		Sig.Inputs.Add(FNiagaraVariable(FNiagaraTypeDefinition::GetVec3Def(), TEXT("World Position")));
+		Sig.Inputs.Add(FNiagaraVariable(FNiagaraTypeDefinition::GetPositionDef(), TEXT("World Position")));
 		Sig.Inputs.Add(FNiagaraVariable(FNiagaraTypeDefinition::GetFloatDef(), TEXT("Delta Time")));
 		Sig.Inputs.Add(FNiagaraVariable(FNiagaraTypeDefinition::GetFloatDef(), TEXT("Time Fraction")));
 		Sig.Outputs.Add(FNiagaraVariable(FNiagaraTypeDefinition::GetFloatDef(), TEXT("Closest Distance")));
-		Sig.Outputs.Add(FNiagaraVariable(FNiagaraTypeDefinition::GetVec3Def(), TEXT("Closest Position")));		
+		Sig.Outputs.Add(FNiagaraVariable(FNiagaraTypeDefinition::GetPositionDef(), TEXT("Closest Position")));		
 		Sig.Outputs.Add(FNiagaraVariable(FNiagaraTypeDefinition::GetVec3Def(), TEXT("Closest Velocity")));
 
 		OutFunctions.Add(Sig);
@@ -1171,6 +1186,57 @@ bool UNiagaraDataInterfaceRigidMeshCollisionQuery::GetFunctionHLSL(const FNiagar
 	return false;
 }
 
+bool UNiagaraDataInterfaceRigidMeshCollisionQuery::UpgradeFunctionCall(FNiagaraFunctionSignature& FunctionSignature)
+{
+	bool bChanged = false;
+	
+	// upgrade from lwc changes, only parameter types changed there
+	if (FunctionSignature.FunctionVersion < FNiagaraRigidMeshCollisionDIFunctionVersion::LargeWorldCoordinates)
+	{
+		if (FunctionSignature.Name == GetClosestPointName && ensure(FunctionSignature.Inputs.Num() == 4) && ensure(FunctionSignature.Outputs.Num() == 4))
+		{
+			FunctionSignature.Inputs[1].SetType(FNiagaraTypeDefinition::GetPositionDef());
+			FunctionSignature.Outputs[1].SetType(FNiagaraTypeDefinition::GetPositionDef());
+			bChanged = true;
+		}
+		if (FunctionSignature.Name == GetClosestElementName && ensure(FunctionSignature.Inputs.Num() == 3))
+		{
+			FunctionSignature.Inputs[1].SetType(FNiagaraTypeDefinition::GetPositionDef());
+			bChanged = true;
+		}
+		if (FunctionSignature.Name == GetElementPointName && ensure(FunctionSignature.Inputs.Num() == 5) && ensure(FunctionSignature.Outputs.Num() == 3))
+		{
+			FunctionSignature.Inputs[1].SetType(FNiagaraTypeDefinition::GetPositionDef());
+			FunctionSignature.Outputs[0].SetType(FNiagaraTypeDefinition::GetPositionDef());
+			bChanged = true;
+		}
+		if (FunctionSignature.Name == GetElementDistanceName && ensure(FunctionSignature.Inputs.Num() == 4))
+		{
+			FunctionSignature.Inputs[1].SetType(FNiagaraTypeDefinition::GetPositionDef());
+			bChanged = true;
+		}
+		if (FunctionSignature.Name == GetClosestDistanceName && ensure(FunctionSignature.Inputs.Num() == 3))
+		{
+			FunctionSignature.Inputs[1].SetType(FNiagaraTypeDefinition::GetPositionDef());
+			bChanged = true;
+		}
+		if (FunctionSignature.Name == GetClosestPointMeshDistanceFieldName && ensure(FunctionSignature.Inputs.Num() == 4) && ensure(FunctionSignature.Outputs.Num() == 4))
+		{
+			FunctionSignature.Inputs[1].SetType(FNiagaraTypeDefinition::GetPositionDef());
+			FunctionSignature.Outputs[1].SetType(FNiagaraTypeDefinition::GetPositionDef());
+			bChanged = true;
+		}
+		if (FunctionSignature.Name == GetClosestPointMeshDistanceFieldNoNormalName && ensure(FunctionSignature.Inputs.Num() == 4) && ensure(FunctionSignature.Outputs.Num() == 3))
+		{
+			FunctionSignature.Inputs[1].SetType(FNiagaraTypeDefinition::GetPositionDef());
+			FunctionSignature.Outputs[1].SetType(FNiagaraTypeDefinition::GetPositionDef());
+			bChanged = true;
+		}
+	}
+
+	return bChanged;
+}
+
 void UNiagaraDataInterfaceRigidMeshCollisionQuery::GetCommonHLSL(FString& OutHLSL)
 {
 	OutHLSL += TEXT("#include \"/Engine/Private/DistanceFieldLightingShared.ush\"\n");
diff --git a/Engine/Plugins/FX/Niagara/Source/Niagara/Public/NiagaraDataInterfaceRigidMeshCollisionQuery.h b/Engine/Plugins/FX/Niagara/Source/Niagara/Public/NiagaraDataInterfaceRigidMeshCollisionQuery.h
--- a/Engine/Plugins/FX/Niagara/Source/Niagara/Public/NiagaraDataInterfaceRigidMeshCollisionQuery.h
+++ b/Engine/Plugins/FX/Niagara/Source/Niagara/Public/NiagaraDataInterfaceRigidMeshCollisionQuery.h
@@ -156,7 +156,7 @@ struct FNDIRigidMeshCollisionData
 	TArray<AActor*> Actors;
 };
 
-/** Data Interface for the Collisions */
+/** Data Interface used to collide against static meshes - whether it is the mesh distance field or a physics asset's collision primitive */
 UCLASS(EditInlineNew, Category = "Collision", meta = (DisplayName = "Rigid Mesh Collision Query"))
 class UNiagaraDataInterfaceRigidMeshCollisionQuery : public UNiagaraDataInterface
 {
@@ -198,6 +198,7 @@ public:
 	virtual void GetCommonHLSL(FString& OutHLSL) override;
 	virtual void GetParameterDefinitionHLSL(const FNiagaraDataInterfaceGPUParamInfo& ParamInfo, FString& OutHLSL) override;
 	virtual bool GetFunctionHLSL(const FNiagaraDataInterfaceGPUParamInfo& ParamInfo, const FNiagaraDataInterfaceGeneratedFunction& FunctionInfo, int FunctionInstanceIndex, FString& OutHLSL) override;
+	virtual bool UpgradeFunctionCall(FNiagaraFunctionSignature& FunctionSignature) override;
 
 	virtual void ValidateFunction(const FNiagaraFunctionSignature& Function, TArray<FText>& OutValidationErrors) override;
 #endif
