Fixed script variables not being created if a new parameter is created via map node

diff --git a/Engine/Plugins/FX/Niagara/Source/NiagaraEditor/Private/NiagaraNodeWithDynamicPins.cpp b/Engine/Plugins/FX/Niagara/Source/NiagaraEditor/Private/NiagaraNodeWithDynamicPins.cpp
--- a/Engine/Plugins/FX/Niagara/Source/NiagaraEditor/Private/NiagaraNodeWithDynamicPins.cpp
+++ b/Engine/Plugins/FX/Niagara/Source/NiagaraEditor/Private/NiagaraNodeWithDynamicPins.cpp
@@ -332,14 +332,11 @@ void UNiagaraNodeWithDynamicPins::AddParameter(FNiagaraVariable Parameter, const
 		
 		UNiagaraGraph* Graph = GetNiagaraGraph();
 		checkf(Graph != nullptr, TEXT("Failed to get niagara graph when adding pin!"));
-
+		
 		// Resolve the unique parameter name before adding to the graph if the current parameter name is not reserved.
 		if (FNiagaraConstants::FindEngineConstant(Parameter) == nullptr)
 		{
-			if(Graph->GetAllMetaData().Contains(Parameter) == false)
-			{
-				Parameter.SetName(Graph->MakeUniqueParameterName(Parameter.GetName()));
-			}
+			Parameter.SetName(Graph->MakeUniqueParameterName(Parameter.GetName()));
 		}
 
 		Graph->Modify();
