Niagara - Merge fixes.
+ Merge failure messages are now added to the emitter's messages.
+ Stack objects will now show messages from the message manager if they reference the object they're displaying.
+ Missing static switch pins encountered during the merge will be preserved as orpahned pins on the function call.
+ Static switches which fail to merge due to other reasons now provide more context in their error message.
+ Fixed an issue which was causing orphaned pins which were different from the default value to be discarded if they weren't connected.

diff --git a/Engine/Plugins/FX/Niagara/Source/Niagara/Classes/NiagaraEmitter.h b/Engine/Plugins/FX/Niagara/Source/Niagara/Classes/NiagaraEmitter.h
--- a/Engine/Plugins/FX/Niagara/Source/Niagara/Classes/NiagaraEmitter.h
+++ b/Engine/Plugins/FX/Niagara/Source/Niagara/Classes/NiagaraEmitter.h
@@ -331,6 +331,7 @@ public:
 	UPROPERTY(EditAnywhere, AdvancedDisplay, Category = "Emitter")
 	TArray<FString> AttributesToPreserve;
 
+	const static FGuid EmitterMergeMessageId;
 #endif
 
 	UPROPERTY(EditAnywhere, Category = "Emitter")
diff --git a/Engine/Plugins/FX/Niagara/Source/Niagara/Private/NiagaraEmitter.cpp b/Engine/Plugins/FX/Niagara/Source/Niagara/Private/NiagaraEmitter.cpp
--- a/Engine/Plugins/FX/Niagara/Source/Niagara/Private/NiagaraEmitter.cpp
+++ b/Engine/Plugins/FX/Niagara/Source/Niagara/Private/NiagaraEmitter.cpp
@@ -25,6 +25,8 @@
 const FName UNiagaraEmitter::PrivateMemberNames::EventHandlerScriptProps = GET_MEMBER_NAME_CHECKED(UNiagaraEmitter, EventHandlerScriptProps);
 
 const FString InitialNotSynchronizedReason("Emitter created");
+
+const FGuid UNiagaraEmitter::EmitterMergeMessageId(0xDAF26E73, 0x4D1B416B, 0x815FA6C2, 0x6D5D0A75);
 #endif
 
 static int32 GbForceNiagaraCompileOnLoad = 0;
@@ -286,6 +288,12 @@ INiagaraMergeManager::FMergeEmitterResults UNiagaraEmitter::MergeChangesFromPare
 			*GetPathName(), Parent != nullptr ? *Parent->GetPathName() : TEXT("(null)"), *MergeResults.GetErrorMessagesString());
 	}
 
+	RemoveMessage(EmitterMergeMessageId);
+	if (MergeResults.MergeNiagaraMessage != nullptr)
+	{
+		AddMessage(EmitterMergeMessageId, CastChecked<UNiagaraMessageDataBase>(StaticDuplicateObject(MergeResults.MergeNiagaraMessage, this)));
+	}
+
 	return MergeResults;
 }
 
@@ -2018,6 +2026,17 @@ void UNiagaraEmitter::UpdateFromMergedCopy(const INiagaraMergeManager& MergeMana
 	ReouterMergedObject(this, NewEditorData);
 	SetEditorData(NewEditorData);
 
+	// Update messages
+	for (const TPair<FGuid, TObjectPtr<UNiagaraMessageDataBase>>& Message : MergedEmitter->GetMessages())
+	{
+		TObjectPtr<UNiagaraMessageDataBase>& MessageData = MessageKeyToMessageMap.FindOrAdd(Message.Key);
+		if (Message.Value != nullptr)
+		{
+			ReouterMergedObject(this, Message.Value);
+			MessageData = Message.Value;
+		}
+	}
+
 	// Update the change id since we don't know what's changed.
 	UpdateChangeId(TEXT("Updated from merged copy"));
 }
diff --git a/Engine/Plugins/FX/Niagara/Source/Niagara/Public/INiagaraMergeManager.h b/Engine/Plugins/FX/Niagara/Source/Niagara/Public/INiagaraMergeManager.h
--- a/Engine/Plugins/FX/Niagara/Source/Niagara/Public/INiagaraMergeManager.h
+++ b/Engine/Plugins/FX/Niagara/Source/Niagara/Public/INiagaraMergeManager.h
@@ -2,6 +2,8 @@
 
 #pragma once
 
+class UNiagaraMessageDataBase;
+
 class INiagaraMergeManager
 {
 public:
@@ -18,6 +20,7 @@ public:
 	{
 		FMergeEmitterResults()
 			: MergeResult(EMergeEmitterResult::None)
+			, MergeNiagaraMessage(nullptr)
 			, bModifiedGraph(false)
 			, MergedInstance(nullptr)
 		{
@@ -25,6 +28,7 @@ public:
 
 		EMergeEmitterResult MergeResult;
 		TArray<FText> ErrorMessages;
+		UNiagaraMessageDataBase* MergeNiagaraMessage;
 		bool bModifiedGraph;
 		UNiagaraEmitter* MergedInstance;
 
diff --git a/Engine/Plugins/FX/Niagara/Source/NiagaraEditor/Private/NiagaraNode.cpp b/Engine/Plugins/FX/Niagara/Source/NiagaraEditor/Private/NiagaraNode.cpp
--- a/Engine/Plugins/FX/Niagara/Source/NiagaraEditor/Private/NiagaraNode.cpp
+++ b/Engine/Plugins/FX/Niagara/Source/NiagaraEditor/Private/NiagaraNode.cpp
@@ -300,7 +300,7 @@ bool UNiagaraNode::ReallocatePins(bool bMarkNeedsResynchronizeOnChange)
 	for (UEdGraphPin* OldPin : OldPins)
 	{
 		// we exclude pins from being added as orphaned pins if they were set to pending kill already
-		if (OldPin->bOrphanedPin && OldPin->LinkedTo.Num() > 0 && !OldPin->IsPendingKill())
+		if (OldPin->bOrphanedPin && !OldPin->IsPendingKill())
 		{
 			// Add orphaned pins back in at the end so that the user can fix them.
 			Pins.Add(OldPin);
diff --git a/Engine/Plugins/FX/Niagara/Source/NiagaraEditor/Private/NiagaraNodeFunctionCall.cpp b/Engine/Plugins/FX/Niagara/Source/NiagaraEditor/Private/NiagaraNodeFunctionCall.cpp
--- a/Engine/Plugins/FX/Niagara/Source/NiagaraEditor/Private/NiagaraNodeFunctionCall.cpp
+++ b/Engine/Plugins/FX/Niagara/Source/NiagaraEditor/Private/NiagaraNodeFunctionCall.cpp
@@ -209,6 +209,49 @@ void UNiagaraNodeFunctionCall::RemoveCustomNoteViaDelegate(const FGuid MessageKe
 	}
 }
 
+void UNiagaraNodeFunctionCall::AddOrphanedStaticSwitchPinForDataRetention(FNiagaraVariableBase StaticSwitchVariable, const FString& StaticSwitchPinDefault)
+{
+	UEdGraphPin* NewPin = AddStaticSwitchInputPin(StaticSwitchVariable);
+	NewPin->DefaultValue = StaticSwitchPinDefault;
+	NewPin->bOrphanedPin = true;
+}
+
+UEdGraphPin* UNiagaraNodeFunctionCall::AddStaticSwitchInputPin(FNiagaraVariable Input)
+{
+	UNiagaraGraph* Graph = GetCalledGraph();
+	const UEdGraphSchema_Niagara* Schema = Graph->GetNiagaraSchema();
+
+	UEdGraphPin* NewPin = CreatePin(EGPD_Input, Schema->TypeDefinitionToPinType(Input.GetType()), Input.GetName());
+	NewPin->bNotConnectable = true;
+	NewPin->bDefaultValueIsIgnored = FindPropagatedVariable(Input) != nullptr;
+	NewPin->SetSavePinIfOrphaned(true);
+
+	FGuid PinPersistentGuid;
+	UNiagaraScriptVariable* ScriptVar = Graph->GetScriptVariable(Input);
+	if (ScriptVar)
+	{
+		int32 DefaultValue = ScriptVar->GetStaticSwitchDefaultValue();
+		Input.AllocateData();
+		Input.SetValue<FNiagaraInt32>({ DefaultValue });
+		PinPersistentGuid = ScriptVar->Metadata.GetVariableGuid();
+	}
+	else
+	{
+		FNiagaraEditorUtilities::ResetVariableToDefaultValue(Input);
+		PinPersistentGuid = FGuid();
+	}
+	NewPin->PersistentGuid = PinPersistentGuid;
+
+	FString PinDefaultValue;
+	if (Schema->TryGetPinDefaultValueFromNiagaraVariable(Input, PinDefaultValue))
+	{
+		NewPin->AutogeneratedDefaultValue = PinDefaultValue;
+		NewPin->DefaultValue = PinDefaultValue;
+	}
+
+	return NewPin;
+}
+
 void UNiagaraNodeFunctionCall::PostEditChangeProperty(struct FPropertyChangedEvent& PropertyChangedEvent)
 {
 	if (PropertyChangedEvent.Property != nullptr)
@@ -294,32 +337,7 @@ void UNiagaraNodeFunctionCall::AllocateDefaultPins()
 		TArray<FNiagaraVariable> SwitchNodeInputs = Graph->FindStaticSwitchInputs();
 		for (FNiagaraVariable& Input : SwitchNodeInputs)
 		{
-			UEdGraphPin* NewPin = CreatePin(EGPD_Input, Schema->TypeDefinitionToPinType(Input.GetType()), Input.GetName());
-			NewPin->bNotConnectable = true;
-			NewPin->bDefaultValueIsIgnored = FindPropagatedVariable(Input) != nullptr;
-			
-			FString PinDefaultValue;
-			UNiagaraScriptVariable* ScriptVar = Graph->GetScriptVariable(Input);
-			if (ScriptVar)
-			{
-				int32 DefaultValue = ScriptVar->GetStaticSwitchDefaultValue();
-				Input.AllocateData();
-				Input.SetValue<FNiagaraInt32>({ DefaultValue });
-				
-				if (Schema->TryGetPinDefaultValueFromNiagaraVariable(Input, PinDefaultValue))
-				{
-					NewPin->DefaultValue = PinDefaultValue;
-				}
-
-				NewPin->PersistentGuid = ScriptVar->Metadata.GetVariableGuid();
-			}
-			else
-			{
-				if (Schema->TryGetPinDefaultValueFromNiagaraVariable(Input, PinDefaultValue))
-				{
-					NewPin->DefaultValue = PinDefaultValue;
-				}
-			}
+			AddStaticSwitchInputPin(Input);
 		}
 
 		for (FNiagaraVariable& Output : Outputs)
diff --git a/Engine/Plugins/FX/Niagara/Source/NiagaraEditor/Private/NiagaraScriptMergeManager.cpp b/Engine/Plugins/FX/Niagara/Source/NiagaraEditor/Private/NiagaraScriptMergeManager.cpp
--- a/Engine/Plugins/FX/Niagara/Source/NiagaraEditor/Private/NiagaraScriptMergeManager.cpp
+++ b/Engine/Plugins/FX/Niagara/Source/NiagaraEditor/Private/NiagaraScriptMergeManager.cpp
@@ -9,6 +9,7 @@
 #include "EdGraphSchema_Niagara.h"
 #include "NiagaraGraph.h"
 #include "NiagaraDataInterface.h"
+#include "NiagaraMessages.h"
 #include "NiagaraNodeInput.h"
 #include "NiagaraNodeOutput.h"
 #include "NiagaraNodeFunctionCall.h"
@@ -323,9 +324,10 @@ FNiagaraStackFunctionMergeAdapter::FNiagaraStackFunctionMergeAdapter(const UNiag
 
 	for (UEdGraphPin* StaticSwitchPin : StaticSwitchPins)
 	{
-		// TODO: Only add static switch overrides when the current value is different from the default.  This requires a 
-		// refactor of the static switch default storage to use the same data format as FNiagaraVariables.
-		InputOverrides.Add(MakeShared<FNiagaraStackFunctionInputOverrideMergeAdapter>(StaticSwitchPin));
+		if (StaticSwitchPin->AutogeneratedDefaultValue.IsEmpty() || StaticSwitchPin->DoesDefaultValueMatchAutogenerated() == false)
+		{
+			InputOverrides.Add(MakeShared<FNiagaraStackFunctionInputOverrideMergeAdapter>(StaticSwitchPin));
+		}
 	}
 }
 
@@ -1670,6 +1672,21 @@ INiagaraMergeManager::FMergeEmitterResults FNiagaraScriptMergeManager::MergeEmit
 		FNiagaraEditorUtilities::GatherChangeIds(*MergedInstance, FinalChangeIds, TEXT("Final"));
 	}
 
+	if (MergeResults.MergeResult != INiagaraMergeManager::EMergeEmitterResult::SucceededNoDifferences && MergeResults.MergeResult != INiagaraMergeManager::EMergeEmitterResult::SucceededDifferencesApplied)
+	{
+		TArray<FText> Errors;
+		Errors.Add(FText::Format(LOCTEXT("MergeFailedFormat", "Failed to merge changes from parent emitter {0}"), FText::FromString(Parent.GetName())));
+		Errors.Append(MergeResults.ErrorMessages);
+
+		UNiagaraMessageDataText* Message = NewObject<UNiagaraMessageDataText>(GetTransientPackage());
+		Message->Init(
+			FText::Join(FText::FromString("\n"), Errors),
+			LOCTEXT("MergeFailedFormatShort", "Failed to merge changes from parent emitter"),
+			ENiagaraMessageSeverity::Error, "Emitter Merge");
+		Message->SetAllowDismissal(false);
+		MergeResults.MergeNiagaraMessage = Message;
+	}
+
 	return MergeResults;
 }
 
@@ -2833,7 +2850,8 @@ FNiagaraScriptMergeManager::FApplyDiffResults FNiagaraScriptMergeManager::AddInp
 			else
 			{
 				Results.bSucceeded = false;
-				Results.ErrorMessages.Add(LOCTEXT("AddPinBasedInputOverrideFailed", "Failed to add input override because it was invalid."));
+				Results.ErrorMessages.Add(FText::Format(LOCTEXT("AddPinBasedInputOverrideFailed", "Failed to add input override {0} to function {1} because it was invalid."),
+					FText::FromString(OverrideToAdd->GetInputName()), FText::FromString(TargetFunctionCall.GetFunctionName())));
 			}
 		}
 		Results.bModifiedGraph = true;
@@ -2870,19 +2888,22 @@ FNiagaraScriptMergeManager::FApplyDiffResults FNiagaraScriptMergeManager::AddInp
 				else
 				{
 					Results.bSucceeded = false;
-					Results.ErrorMessages.Add(LOCTEXT("AddStaticInputOverrideFailedWrongType", "Failed to add static switch input override because a the type of the pin matched by name did not match."));
+					Results.ErrorMessages.Add(FText::Format(LOCTEXT("AddStaticInputOverrideFailedWrongTypeFormat", "Failed to add static switch input override {0} to function {1} because a the type of the pin matched by name did not match."),
+						FText::FromString(OverrideToAdd->GetInputName()), FText::FromString(TargetFunctionCall.GetFunctionName())));
 				}
 			}
 			else
 			{
-				Results.bSucceeded = false;
-				Results.ErrorMessages.Add(LOCTEXT("AddStaticInputOverrideFailedNotFound", "Failed to add static switch input override because a matching pin could not be found."));
+				FNiagaraVariableBase StaticSwitchVariable(OverrideToAdd->GetType(), *OverrideToAdd->GetInputName());
+				TargetFunctionCall.AddOrphanedStaticSwitchPinForDataRetention(StaticSwitchVariable, OverrideToAdd->GetStaticSwitchValue().GetValue());
+				Results.bSucceeded = true;
 			}
 		}
 		else
 		{
 			Results.bSucceeded = false;
-			Results.ErrorMessages.Add(LOCTEXT("AddParameterBasedInputOverrideFailed", "Failed to add input override because it was invalid."));
+			Results.ErrorMessages.Add(FText::Format(LOCTEXT("AddParameterBasedInputOverrideFailedFormat", "Failed to add input override {0} to function {1} because it was invalid."),
+				FText::FromString(OverrideToAdd->GetInputName()), FText::FromString(TargetFunctionCall.GetFunctionName())));
 		}
 		Results.bModifiedGraph = false;
 	}
diff --git a/Engine/Plugins/FX/Niagara/Source/NiagaraEditor/Private/ViewModels/NiagaraSystemViewModel.cpp b/Engine/Plugins/FX/Niagara/Source/NiagaraEditor/Private/ViewModels/NiagaraSystemViewModel.cpp
--- a/Engine/Plugins/FX/Niagara/Source/NiagaraEditor/Private/ViewModels/NiagaraSystemViewModel.cpp
+++ b/Engine/Plugins/FX/Niagara/Source/NiagaraEditor/Private/ViewModels/NiagaraSystemViewModel.cpp
@@ -1173,6 +1173,7 @@ void FNiagaraSystemViewModel::RefreshAll()
 	CurveSelectionViewModel->Refresh();
 	SystemStackViewModel->InitializeWithViewModels(this->AsShared(), TSharedPtr<FNiagaraEmitterHandleViewModel>(), FNiagaraStackViewModelOptions(true, false));
 	SelectionViewModel->Refresh();
+	bPendingAssetMessagesChanged = true;
 }
 
 void FNiagaraSystemViewModel::NotifyDataObjectChanged(TArray<UObject*> ChangedObjects, ENiagaraDataObjectChange ChangeType)
diff --git a/Engine/Plugins/FX/Niagara/Source/NiagaraEditor/Private/ViewModels/Stack/NiagaraStackObject.cpp b/Engine/Plugins/FX/Niagara/Source/NiagaraEditor/Private/ViewModels/Stack/NiagaraStackObject.cpp
--- a/Engine/Plugins/FX/Niagara/Source/NiagaraEditor/Private/ViewModels/Stack/NiagaraStackObject.cpp
+++ b/Engine/Plugins/FX/Niagara/Source/NiagaraEditor/Private/ViewModels/Stack/NiagaraStackObject.cpp
@@ -5,6 +5,8 @@
 #include "ViewModels/NiagaraSystemViewModel.h"
 #include "NiagaraNode.h"
 #include "NiagaraEditorModule.h"
+#include "NiagaraMessageManager.h"
+#include "NiagaraMessageUtilities.h"
 
 #include "Modules/ModuleManager.h"
 #include "IPropertyRowGenerator.h"
@@ -26,6 +28,15 @@ void UNiagaraStackObject::Initialize(FRequiredEntryData InRequiredEntryData, UOb
 	WeakObject = InObject;
 	OwningNiagaraNode = InOwningNiagaraNode;
 	bIsRefresingDataInterfaceErrors = false;
+
+	MessageLogGuid = GetSystemViewModel()->GetMessageLogGuid();
+
+	FNiagaraMessageManager::Get()->SubscribeToAssetMessagesByObject(
+		FText::FromString("StackObject")
+		, MessageLogGuid
+		, FObjectKey(InObject)
+		, MessageManagerRegistrationKey
+	).BindUObject(this, &UNiagaraStackObject::OnMessageManagerRefresh);
 }
 
 void UNiagaraStackObject::SetOnSelectRootNodes(FOnSelectRootNodes OnSelectRootNodes)
@@ -84,6 +95,12 @@ void UNiagaraStackObject::FinalizeInternal()
 		FNiagaraEditorModule::Get().EnqueueObjectForDeferredDestruction(PropertyRowGenerator.ToSharedRef());
 		PropertyRowGenerator.Reset();
 	}
+
+	if (MessageManagerRegistrationKey.IsValid())
+	{
+		FNiagaraMessageManager::Get()->Unsubscribe(FText::FromString("StackObject"), MessageLogGuid, MessageManagerRegistrationKey);
+	}
+
 	Super::FinalizeInternal();
 }
 
@@ -269,6 +286,8 @@ void UNiagaraStackObject::RefreshChildrenInternal(const TArray<UNiagaraStackEntr
 			NewChildren.Add(ChildRow);
 		}
 	}
+
+	NewIssues.Append(MessageManagerIssues);
 }
 
 void UNiagaraStackObject::PostRefreshChildrenInternal()
@@ -283,3 +302,22 @@ void UNiagaraStackObject::PropertyRowsRefreshed()
 		RefreshChildren();
 	}
 }
+
+void UNiagaraStackObject::OnMessageManagerRefresh(const TArray<TSharedRef<const INiagaraMessage>>& NewMessages)
+{
+	if (MessageManagerIssues.Num() != 0 || NewMessages.Num() != 0)
+	{
+		MessageManagerIssues.Reset();
+		for (TSharedRef<const INiagaraMessage> Message : NewMessages)
+		{
+			FStackIssue Issue = FNiagaraMessageUtilities::MessageToStackIssue(Message, GetStackEditorDataKey());
+			if (MessageManagerIssues.ContainsByPredicate([&Issue](const FStackIssue& NewIssue)
+				{ return NewIssue.GetUniqueIdentifier() == Issue.GetUniqueIdentifier(); }) == false)
+			{
+				MessageManagerIssues.Add(Issue);
+			}
+		}
+
+		RefreshChildren();
+	}
+}
diff --git a/Engine/Plugins/FX/Niagara/Source/NiagaraEditor/Public/NiagaraNodeFunctionCall.h b/Engine/Plugins/FX/Niagara/Source/NiagaraEditor/Public/NiagaraNodeFunctionCall.h
--- a/Engine/Plugins/FX/Niagara/Source/NiagaraEditor/Public/NiagaraNodeFunctionCall.h
+++ b/Engine/Plugins/FX/Niagara/Source/NiagaraEditor/Public/NiagaraNodeFunctionCall.h
@@ -181,7 +181,12 @@ public:
 	NIAGARAEDITOR_API FSimpleDelegate& OnCustomNotesChanged() { return OnCustomNotesChangedDelegate; }
 	void RemoveCustomNoteViaDelegate(const FGuid MessageKey);
 
+	/** Adds a static switch pin to this function call node by variable id and sets it's default value using the supplied data and marks it as
+		orphaned. This allows a previously available static switch value to be retained on the node even if the the switch is no longer exposed. */
+	void AddOrphanedStaticSwitchPinForDataRetention(FNiagaraVariableBase StaticSwitchVariable, const FString& StaticSwitchPinDefault);
+
 protected:
+	UEdGraphPin* AddStaticSwitchInputPin(FNiagaraVariable Input);
 
 	virtual bool GetValidateDataInterfaces() const { return true; };
 
diff --git a/Engine/Plugins/FX/Niagara/Source/NiagaraEditor/Public/ViewModels/Stack/NiagaraStackObject.h b/Engine/Plugins/FX/Niagara/Source/NiagaraEditor/Public/ViewModels/Stack/NiagaraStackObject.h
--- a/Engine/Plugins/FX/Niagara/Source/NiagaraEditor/Public/ViewModels/Stack/NiagaraStackObject.h
+++ b/Engine/Plugins/FX/Niagara/Source/NiagaraEditor/Public/ViewModels/Stack/NiagaraStackObject.h
@@ -10,6 +10,7 @@
 class IPropertyRowGenerator;
 class UNiagaraNode;
 class IDetailTreeNode;
+class INiagaraMessage;
 
 UCLASS()
 class NIAGARAEDITOR_API UNiagaraStackObject : public UNiagaraStackItemContent, public FNotifyHook
@@ -47,6 +48,7 @@ protected:
 
 private:
 	void PropertyRowsRefreshed();
+	void OnMessageManagerRefresh(const TArray<TSharedRef<const INiagaraMessage>>& NewMessages);
 
 private:
 	struct FRegisteredClassCustomization
@@ -69,4 +71,8 @@ private:
 	TArray<FRegisteredPropertyCustomization> RegisteredPropertyCustomizations;
 	TSharedPtr<IPropertyRowGenerator> PropertyRowGenerator;
 	bool bIsRefresingDataInterfaceErrors;
+
+	FGuid MessageManagerRegistrationKey;
+	TArray<FStackIssue> MessageManagerIssues;
+	FGuid MessageLogGuid;
 };
