Fix for CVar condition text entry box.

diff --git a/Engine/Plugins/FX/Niagara/Source/NiagaraEditor/Private/Customizations/NiagaraPlatformSetCustomization.cpp b/Engine/Plugins/FX/Niagara/Source/NiagaraEditor/Private/Customizations/NiagaraPlatformSetCustomization.cpp
--- a/Engine/Plugins/FX/Niagara/Source/NiagaraEditor/Private/Customizations/NiagaraPlatformSetCustomization.cpp
+++ b/Engine/Plugins/FX/Niagara/Source/NiagaraEditor/Private/Customizations/NiagaraPlatformSetCustomization.cpp
@@ -1379,13 +1379,10 @@ void SNiagaraConsoleInputBox::OnTextChanged(const FText& InText)
 
 void SNiagaraConsoleInputBox::OnTextCommitted(const FText& InText, ETextCommit::Type CommitInfo)
 {
-	if (!WorkingText.EqualTo(InText))
-	{
-		WorkingText = InText;
-		IConsoleManager::Get().AddConsoleHistoryEntry(NiagaraCVarHistoryKey, *InText.ToString());
-		OnTextCommittedEvent.ExecuteIfBound(InText);
-		SuggestionBox->SetIsOpen(false);
-	}
+	WorkingText = InText;
+	IConsoleManager::Get().AddConsoleHistoryEntry(NiagaraCVarHistoryKey, *InText.ToString());
+	OnTextCommittedEvent.ExecuteIfBound(InText);
+	SuggestionBox->SetIsOpen(false);
 }
 
 FReply SNiagaraConsoleInputBox::OnPreviewKeyDown(const FGeometry& MyGeometry, const FKeyEvent& KeyEvent)
