- Add missing call to PostTick on GPU emitters

#rnx

diff --git a/Engine/Plugins/FX/Niagara/Source/Niagara/Private/NiagaraEmitterInstance.cpp b/Engine/Plugins/FX/Niagara/Source/Niagara/Private/NiagaraEmitterInstance.cpp
--- a/Engine/Plugins/FX/Niagara/Source/Niagara/Private/NiagaraEmitterInstance.cpp
+++ b/Engine/Plugins/FX/Niagara/Source/Niagara/Private/NiagaraEmitterInstance.cpp
@@ -1541,6 +1541,9 @@ void FNiagaraEmitterInstance::Tick(float DeltaSeconds)
 			);
 		}
 
+		// Call PostTick to adjust bounds, etc
+		PostTick();
+
 		// Need to call post-tick, which calls the copy to previous for interpolated spawning
 		SpawnExecContext.PostTick();
 		UpdateExecContext.PostTick();
@@ -1551,13 +1554,6 @@ void FNiagaraEmitterInstance::Tick(float DeltaSeconds)
 			EventContext.PostTick();
 		}
 
-		CachedBounds = CachedEmitter->FixedBounds;
-
-		/*if (CachedEmitter->SpawnScriptProps.Script->GetComputedVMCompilationId().HasInterpolatedParameters())
-		{
-			GPUExecContext.CombinedParamStore.CopyCurrToPrev();
-		}*/
-
 		return;
 	}
 
