Adds notion of precise motion vectors for niagara ribbons
-approximates previous frame's position using previous values for a number of attributes of the particles.  Not perfect but a reasonable approximation for now

diff --git a/Engine/Plugins/FX/Niagara/Shaders/Private/NiagaraRibbonVertexFactory.ush b/Engine/Plugins/FX/Niagara/Shaders/Private/NiagaraRibbonVertexFactory.ush
--- a/Engine/Plugins/FX/Niagara/Shaders/Private/NiagaraRibbonVertexFactory.ush
+++ b/Engine/Plugins/FX/Niagara/Shaders/Private/NiagaraRibbonVertexFactory.ush
@@ -58,12 +58,11 @@ float UnpackParticleVertexTextureV(uint InterpIdx)
 	return NiagaraRibbonVFLooseParameters.SliceVertexData[UnpackParticleVertexId(InterpIdx) * SLICEVERTEXDATA_STRIDE + SLICEVERTEXDATA_TEXTUREV_OFFSET];
 }
 
-
-FLWCVector3 TransformPosition(float3 InPosition)
+FLWCVector3 TransformPosition(FLWCMatrix LocalToWorld, float3 InPosition)
 {
 	if (NiagaraRibbonVF.bLocalSpace)
 	{
-		return TransformLocalToWorld(InPosition);
+		return TransformLocalToWorld(InPosition, LocalToWorld);
 	}
 	else
 	{
@@ -71,11 +70,11 @@ FLWCVector3 TransformPosition(float3 InPosition)
 	}
 }
 
-float3 TransformVector(float3 InVector)
+float3 TransformVector(FLWCMatrix LocalToWorld, float3 InvScale, float3 InVector)
 {
 	if (NiagaraRibbonVF.bLocalSpace)
 	{
-		return RotateLocalToWorld(InVector);
+		return RotateLocalToWorld(InVector, LocalToWorld, InvScale);
 	}
 	else
 	{
@@ -140,13 +139,12 @@ struct FSegmentData
 	uint StartParticleId;
 };
 
-struct FVertexFactoryIntermediates
+struct FRibbonPositionIntermediates
 {
 	FLWCVector3 Position;
 	float3 Direction;
 	float DistanceOnSegment;
-	float4 Color;
-	float RibbonWidth;	
+	float RibbonWidth;
 	
 	float3 GeoUp;
 	float3 GeoRight;
@@ -154,7 +152,14 @@ struct FVertexFactoryIntermediates
 	float3 TangentX;
 	float3 TangentY;
 	float3 TangentZ;
-	
+};
+
+struct FVertexFactoryIntermediates
+{
+	FRibbonPositionIntermediates Current;
+	FRibbonPositionIntermediates Previous;
+	float DistanceOnSegment;
+	float4 Color;
 	
 	float4 TexCoord;
 	
@@ -412,7 +417,7 @@ FMaterialVertexParameters GetMaterialVertexParameters(FVertexFactoryInput Input,
 	Result.VertexColor = Intermediates.Color;
 	Result.TangentToWorld = mul(TangentToLocal, GetLocalToWorld3x3()); 
 	Result.Particle.Color = Intermediates.Color;
-	Result.PreSkinnedPosition = LWCToFloat(Intermediates.Position);
+	Result.PreSkinnedPosition = LWCToFloat(Intermediates.Current.Position);
 	Result.PreSkinnedNormal = TangentToLocal[2].xyz;
 
 #if (DYNAMIC_PARAMETERS_MASK != 0)
@@ -450,9 +455,7 @@ FMaterialVertexParameters GetMaterialVertexParameters(FVertexFactoryInput Input,
 #if USE_PARTICLE_RANDOM
 	Result.Particle.Random = Intermediates.ParticleRandom;
 #endif
-	// Previous frame not handled deliberately. Lacks necessary information and
-	// primitives using this VF are usually transparent and hence don't output velocity
-	Result.PrevFrameLocalToWorld = GetPrimitiveDataFromUniformBuffer().LocalToWorld;
+	Result.PrevFrameLocalToWorld = GetPrimitiveDataFromUniformBuffer().PreviousLocalToWorld;
 
 #if NUM_MATERIAL_TEXCOORDS_VERTEX
 	Result.TexCoords[0] = Intermediates.TexCoord.xy;
@@ -470,20 +473,121 @@ float3 SafeNormalize(float3 V)
 }
 #endif
 
-float4 CalcWorldPosition(FVertexFactoryInput Input, FVertexFactoryIntermediates Intermediates)
+float4 GetDirectionAndMagnitude(float3 V)
 {
-	const float2 SliceVertexPositions = UnpackParticleVertexPosition(Input.InterpVtxId);
-	const float3 Offset = ((Intermediates.GeoRight * SliceVertexPositions.xxx) + (Intermediates.GeoUp * SliceVertexPositions.yyy)) * Intermediates.RibbonWidth.xxx;
-	const FLWCVector3 Pos = LWCAdd(LWCAdd(Intermediates.Position, LWCPromote(Offset)), ResolvedView.PreViewTranslation);
-		
-	// Intermediates are already in world space, only need to account for the translation around view origin.
-	return float4(LWCToFloat(Pos), 1);
+	const float MagSquared = dot(V, V);
+	if (MagSquared > 0)
+	{
+		const float MagScale = rsqrt(MagSquared);
+		return float4(MagScale * V, rcp(MagScale));
+	}
+
+	return float4(0.0f, 0.0f, 0.0f, 0.0f);
 }
 
 /** derive basis vectors */
 float3x3 CalcTangentBasis(FVertexFactoryInput Input, FVertexFactoryIntermediates Intermediates)
 {	
-	return float3x3(Intermediates.TangentX, Intermediates.TangentY, Intermediates.TangentZ);
+	return float3x3(Intermediates.Current.TangentX, Intermediates.Current.TangentY, Intermediates.Current.TangentZ);
+}
+
+// known limitations:
+// -TangentsAndDistances is from the current frame and if we're generating intermediates for the previous frame for velocity
+// rendering we won't get the exact same results.  Issues around trying to supply a previous version of the TangentsAndDistances
+// buffer include the fact that the active particles can differ between frames and so we'd also need to supply particle indices
+// for the previous frames as well
+FRibbonPositionIntermediates GetRibbonPositionIntermediates(
+	FInterpVtxData InterpVtxData,
+	int PositionDataOffset,
+	int FacingDataOffset,
+	int TwistDataOffset,
+	int WidthDataOffset,
+	float2 SliceVertexNormals,
+	FLWCMatrix LocalToWorld,
+	float3 InvScale)
+{
+	FRibbonPositionIntermediates Intermediates = (FRibbonPositionIntermediates) 0;
+
+	BRANCH // Also acts as ISOLATE
+	if (PositionDataOffset != -1)
+	{
+		const float3 Pos0 = GetVec3(PositionDataOffset, InterpVtxData.ParticleId0);
+		const float3 Pos1 = GetVec3(PositionDataOffset, InterpVtxData.ParticleId1);
+		const float4 TangentAndDist0 = NiagaraRibbonVFLooseParameters.TangentsAndDistances[InterpVtxData.RawParticleId0];
+		const float4 TangentAndDist1 = NiagaraRibbonVFLooseParameters.TangentsAndDistances[InterpVtxData.RawParticleId1];
+		
+		Intermediates.Position = TransformPosition(LocalToWorld, CubicInterp(Pos0, TangentAndDist0.xyz, Pos1, TangentAndDist1.xyz, InterpVtxData.Alpha, TangentAndDist1.w - TangentAndDist0.w));
+		Intermediates.Direction = normalize(TransformVector(LocalToWorld, InvScale, lerp(TangentAndDist0.xyz, TangentAndDist1.xyz, InterpVtxData.Alpha)));
+		Intermediates.DistanceOnSegment = lerp(TangentAndDist0.w, TangentAndDist1.w, InterpVtxData.Alpha);
+	}
+
+	float3 Facing;
+	
+	BRANCH
+	if(FacingDataOffset != -1)
+	{
+		const float3 Facing0 = GetVec3(FacingDataOffset, InterpVtxData.ParticleId0);
+		const float3 Facing1 = GetVec3(FacingDataOffset, InterpVtxData.ParticleId1);
+		Facing = normalize(TransformVector(LocalToWorld, InvScale, lerp(Facing0, Facing1, InterpVtxData.Alpha)));
+	}
+	else
+	{
+		Facing = normalize(LWCToFloat(LWCSubtract(ResolvedView.WorldCameraOrigin, Intermediates.Position)));
+	}
+
+	float3 Right = normalize(cross(Facing, Intermediates.Direction));
+	
+	if (NiagaraRibbonVFLooseParameters.FacingMode == 2 /* CustomSide */)
+	{
+		// Swap axis
+		const float3 Temp = Right;
+		Right = -Facing;
+		Facing = Temp;
+	}
+
+	BRANCH
+	if (TwistDataOffset != -1)
+	{
+		const float Twist0 = GetFloat(TwistDataOffset, InterpVtxData.ParticleId0);
+		const float Twist1 = GetFloat(TwistDataOffset, InterpVtxData.ParticleId1);
+		float SinTwist, CosTwist;
+		sincos(lerp(Twist0, Twist1, InterpVtxData.Alpha) /* + ResolvedView.RealTime * .5 */, SinTwist, CosTwist);
+
+		// Apply a rotation matrix to up and right.
+		float3 OriginalRight = Right; // Back it up for the transform.
+		Right = CosTwist * OriginalRight + SinTwist * Facing;
+		Facing = -SinTwist * OriginalRight + CosTwist * Facing;
+	}
+
+	float3 NormalizedFacing = cross(Intermediates.Direction, Right);
+	
+	Intermediates.TangentY = Intermediates.Direction;
+	Intermediates.TangentX = Right;
+	
+	Intermediates.GeoUp = cross(Intermediates.TangentY, Intermediates.TangentX);
+	Intermediates.GeoRight = Intermediates.TangentX;
+	
+	Intermediates.TangentZ = NiagaraRibbonVFLooseParameters.Shape >= 1 ? Intermediates.GeoUp : Facing; // Intermediates.GeoUp;
+	
+	Intermediates.TangentZ = normalize((Intermediates.TangentZ * SliceVertexNormals.y + Intermediates.TangentX * SliceVertexNormals.x));
+	
+	// Flip the normal only for multiplane shape if the normal is looking away from the camera.
+	// This lets us still render as two sided, but at least gets the normals corrected so the plane faces the camera.
+	float3 ViewNormal = normalize(LWCToFloat(LWCSubtract(ResolvedView.WorldCameraOrigin, Intermediates.Position)));
+	Intermediates.TangentZ *= (NiagaraRibbonVF.ShouldFlipNormalToView > 0 && dot(ViewNormal, Intermediates.TangentZ) < 0) ? -1 : 1;
+	
+	Intermediates.TangentX = cross(Intermediates.TangentZ, Intermediates.TangentY);
+
+	// Overrides (FVertexFactoryIntermediates)0; 
+	Intermediates.RibbonWidth = 1.0f;
+
+	FLATTEN
+	if (WidthDataOffset != -1)
+	{
+		Intermediates.RibbonWidth = lerp(GetFloat(WidthDataOffset, InterpVtxData.ParticleId0), GetFloat(WidthDataOffset, InterpVtxData.ParticleId1), InterpVtxData.Alpha);
+	}
+
+	return Intermediates;
 }
 
 FVertexFactoryIntermediates GetVertexFactoryIntermediates(FVertexFactoryInput Input)
@@ -496,28 +600,65 @@ FVertexFactoryIntermediates GetVertexFactoryIntermediates(FVertexFactoryInput In
 	uint RibbonIndex = NiagaraRibbonVFLooseParameters.MultiRibbonIndices[InterpVtxData.RawParticleId0];
 	FSegmentData SegmentData = UnpackPerRibbonDataByIndex(RibbonIndex);
 
-	BRANCH // Also acts as ISOLATE
-	if (NiagaraRibbonVF.PositionDataOffset != -1)
-	{
-		const float3 Pos0 = GetVec3(NiagaraRibbonVF.PositionDataOffset, InterpVtxData.ParticleId0);
-		const float3 Pos1 = GetVec3(NiagaraRibbonVF.PositionDataOffset, InterpVtxData.ParticleId1);
-		const float4 TangentAndDist0 = NiagaraRibbonVFLooseParameters.TangentsAndDistances[InterpVtxData.RawParticleId0];
-		const float4 TangentAndDist1 = NiagaraRibbonVFLooseParameters.TangentsAndDistances[InterpVtxData.RawParticleId1];
-		
-		Intermediates.Position = TransformPosition(CubicInterp(Pos0, TangentAndDist0.xyz, Pos1, TangentAndDist1.xyz, InterpVtxData.Alpha, TangentAndDist1.w - TangentAndDist0.w));
-		Intermediates.Direction = normalize(TransformVector(lerp(TangentAndDist0.xyz, TangentAndDist1.xyz, InterpVtxData.Alpha)));
+	const float2 SliceVertexNormals = UnpackParticleVertexNormal(Input.InterpVtxId);
 
-		Intermediates.DistanceOnSegment = lerp(TangentAndDist0.w, TangentAndDist1.w, InterpVtxData.Alpha);
-	}
+	FPrimitiveSceneData PrimitiveData = GetPrimitiveDataFromUniformBuffer();
+	FLWCMatrix LocalToWorld = GetPrimitiveDataFromUniformBuffer().LocalToWorld;
+	float3 InvScale = PrimitiveData.InvNonUniformScale;
+
+	Intermediates.Current = GetRibbonPositionIntermediates(
+		InterpVtxData,
+		NiagaraRibbonVF.PositionDataOffset,
+		NiagaraRibbonVF.FacingDataOffset,
+		NiagaraRibbonVF.TwistDataOffset,
+		NiagaraRibbonVF.WidthDataOffset,
+		SliceVertexNormals,
+		LocalToWorld,
+		InvScale);
 
 	BRANCH
 	if(NiagaraRibbonVF.VelocityDataOffset != -1)
 	{
 		const float3 Velocity0 = GetVec3(NiagaraRibbonVF.VelocityDataOffset, InterpVtxData.ParticleId0);
 		const float3 Velocity1 = GetVec3(NiagaraRibbonVF.VelocityDataOffset, InterpVtxData.ParticleId1);
-		const float3 Velocity = TransformVector(lerp(Velocity0, Velocity1, InterpVtxData.Alpha));
+		const float3 Velocity = TransformVector(LocalToWorld, InvScale, lerp(Velocity0, Velocity1, InterpVtxData.Alpha));
+
+		Intermediates.ParticleVelocity = GetDirectionAndMagnitude(Velocity);
+	}
+
+	BRANCH
+	if (NiagaraRibbonVFLooseParameters.NeedsPreciseMotionVectors)
+	{
+		FLWCMatrix PreviousLocalToWorld = PrimitiveData.PreviousLocalToWorld;
+
+		// we don't cache the InvScale so we generate it here
+		float3x3 PreviousTransform = LWCToFloat3x3(PreviousLocalToWorld);
+		const float3 PreviousXAxis = float3(PreviousTransform[0][0], PreviousTransform[0][1], PreviousTransform[0][2]);
+		const float3 PreviousYAxis = float3(PreviousTransform[1][0], PreviousTransform[1][1], PreviousTransform[1][2]);
+		const float3 PreviousZAxis = float3(PreviousTransform[2][0], PreviousTransform[2][1], PreviousTransform[2][2]);
+
+		float3 PreviousInvScale = float3(
+			rsqrt(max(0.0001f, dot(PreviousXAxis, PreviousXAxis))),
+			rsqrt(max(0.0001f, dot(PreviousYAxis, PreviousYAxis))),
+			rsqrt(max(0.0001f, dot(PreviousZAxis, PreviousZAxis)))
+		);
+
+		Intermediates.Previous = GetRibbonPositionIntermediates(
+			InterpVtxData,
+			NiagaraRibbonVF.PrevPositionDataOffset,
+			NiagaraRibbonVF.PrevFacingDataOffset,
+			NiagaraRibbonVF.PrevTwistDataOffset,
+			NiagaraRibbonVF.PrevWidthDataOffset,
+			SliceVertexNormals,
+			PreviousLocalToWorld,
+			PreviousInvScale);
+	}
+	else
+	{
+		const float3 VelocityOffset = TransformVector(LocalToWorld, InvScale, -Intermediates.ParticleVelocity.xyz * (Intermediates.ParticleVelocity.w * NiagaraRibbonVF.DeltaSeconds));
 
-		Intermediates.ParticleVelocity = float4(normalize(Velocity), length(Velocity));
+		Intermediates.Previous = Intermediates.Current;
+		Intermediates.Previous.Position = LWCAdd(Intermediates.Current.Position, LWCPromote(VelocityOffset));
 	}
 
 	// Overrides (FVertexFactoryIntermediates)0; 
@@ -639,82 +780,8 @@ FVertexFactoryIntermediates GetVertexFactoryIntermediates(FVertexFactoryInput In
 	Intermediates.TexCoord.xy = UV0ForSegment * SegmentData.UV0Scale + SegmentData.UV0Offset;
 	Intermediates.TexCoord.zw = UV1ForSegment * SegmentData.UV1Scale + SegmentData.UV1Offset;
 
-	// Overrides (FVertexFactoryIntermediates)0; 
-	Intermediates.RibbonWidth = 1;
-
-	FLATTEN
-	if (NiagaraRibbonVF.WidthDataOffset != -1)
-	{
-		Intermediates.RibbonWidth = lerp(GetFloat(NiagaraRibbonVF.WidthDataOffset, InterpVtxData.ParticleId0), GetFloat(NiagaraRibbonVF.WidthDataOffset, InterpVtxData.ParticleId1), InterpVtxData.Alpha);
-	}
-
-	float3 Facing;
-	
-	BRANCH
-	if(NiagaraRibbonVF.FacingDataOffset != -1)
-	{
-		const float3 Facing0 = GetVec3(NiagaraRibbonVF.FacingDataOffset, InterpVtxData.ParticleId0);
-		const float3 Facing1 = GetVec3(NiagaraRibbonVF.FacingDataOffset, InterpVtxData.ParticleId1);
-		Facing = normalize(TransformVector(lerp(Facing0, Facing1, InterpVtxData.Alpha)));
-	}
-	else
-	{
-		Facing = normalize(LWCToFloat(LWCSubtract(ResolvedView.WorldCameraOrigin, Intermediates.Position)));
-	}
-
-	float3 Right = normalize(cross(Facing, Intermediates.Direction));
-	
-	
-	if (NiagaraRibbonVFLooseParameters.FacingMode == 2 /* CustomSide */)
-	{
-		// Swap axis
-		const float3 Temp = Right;
-		Right = -Facing;
-		Facing = Temp;
-	}
-
-	BRANCH
-	if (NiagaraRibbonVF.TwistDataOffset != -1)
-	{
-		const float Twist0 = GetFloat(NiagaraRibbonVF.TwistDataOffset, InterpVtxData.ParticleId0);
-		const float Twist1 = GetFloat(NiagaraRibbonVF.TwistDataOffset, InterpVtxData.ParticleId1);
-		float SinTwist, CosTwist;
-		sincos(lerp(Twist0, Twist1, InterpVtxData.Alpha) /* + ResolvedView.RealTime * .5 */, SinTwist, CosTwist);
-
-		// Apply a rotation matrix to up and right.
-		float3 OriginalRight = Right; // Back it up for the transform.
-		Right = CosTwist * OriginalRight + SinTwist * Facing;
-		Facing = -SinTwist * OriginalRight + CosTwist * Facing;
-	}
-	float3 NormalizedFacing = cross(Intermediates.Direction, Right);
-	
-	
-	
-	Intermediates.TangentY = Intermediates.Direction;
-	Intermediates.TangentX = Right;
-	
-	Intermediates.GeoUp = cross(Intermediates.TangentY, Intermediates.TangentX);
-	Intermediates.GeoRight = Intermediates.TangentX;
-	
-	Intermediates.TangentZ = NiagaraRibbonVFLooseParameters.Shape >= 1 ? Intermediates.GeoUp : Facing; // Intermediates.GeoUp;
-	
-	
-	
-	float2 SliceVertexNormals = UnpackParticleVertexNormal(Input.InterpVtxId);
-	
-	
-	
-	Intermediates.TangentZ = normalize((Intermediates.TangentZ * SliceVertexNormals.y + Intermediates.TangentX * SliceVertexNormals.x));
-	
-	// Flip the normal only for multiplane shape if the normal is looking away from the camera.
-	// This lets us still render as two sided, but at least gets the normals corrected so the plane faces the camera.
-	float3 ViewNormal = normalize(LWCToFloat(LWCSubtract(ResolvedView.WorldCameraOrigin, Intermediates.Position)));
-	Intermediates.TangentZ *= (NiagaraRibbonVF.ShouldFlipNormalToView > 0 && dot(ViewNormal, Intermediates.TangentZ) < 0) ? -1 : 1;
-	
-	Intermediates.TangentX = cross(Intermediates.TangentZ, Intermediates.TangentY);
-
 #if USE_PARTICLE_POSITION
-	Intermediates.TranslatedWorldPositionAndSize = float4(LWCToFloat(LWCAdd(Intermediates.Position, ResolvedView.PreViewTranslation)),  .5 * Intermediates.RibbonWidth);
+	Intermediates.TranslatedWorldPositionAndSize = float4(LWCToFloat(LWCAdd(Intermediates.Current.Position, ResolvedView.PreViewTranslation)),  .5 * Intermediates.Current.RibbonWidth);
 #endif
 
 #if USE_PARTICLE_TIME
@@ -726,7 +793,7 @@ FVertexFactoryIntermediates GetVertexFactoryIntermediates(FVertexFactoryInput In
 #endif
 
 #if USE_PARTICLE_SIZE
-	Intermediates.ParticleSize = Intermediates.RibbonWidth.xx;
+	Intermediates.ParticleSize = Intermediates.Current.RibbonWidth.xx;
 #endif
 
 #if USE_PARTICLE_RANDOM
@@ -740,16 +807,29 @@ FVertexFactoryIntermediates GetVertexFactoryIntermediates(FVertexFactoryInput In
 	return Intermediates;
 }
 
+float3 ConstructWorldPosition(FVertexFactoryInput Input, FRibbonPositionIntermediates Intermediates, FLWCVector3 PreViewTranslation)
+{
+	const float2 SliceVertexPositions = UnpackParticleVertexPosition(Input.InterpVtxId);
+	const float3 Offset = ((Intermediates.GeoRight * SliceVertexPositions.xxx) + (Intermediates.GeoUp * SliceVertexPositions.yyy)) * Intermediates.RibbonWidth.xxx;
+	const FLWCVector3 Pos = LWCAdd(LWCAdd(Intermediates.Position, LWCPromote(Offset)), ResolvedView.PreViewTranslation);
+		
+	// Intermediates are already in world space, only need to account for the translation around view origin.
+	return LWCToFloat(Pos);
+}
+
 float4 VertexFactoryGetWorldPosition(FVertexFactoryInput Input, FVertexFactoryIntermediates Intermediates)
 {
-	return CalcWorldPosition(Input, Intermediates);
+	return float4(ConstructWorldPosition(Input, Intermediates.Current, ResolvedView.PreViewTranslation), 1);
+}
+
+float4 VertexFactoryGetPreviousWorldPosition(FVertexFactoryInput Input, FVertexFactoryIntermediates Intermediates)
+{
+	return float4(ConstructWorldPosition(Input, Intermediates.Previous, ResolvedView.PrevPreViewTranslation), 1);
 }
 
 float3 VertexFactoryGetWorldNormal(FVertexFactoryInput Input, FVertexFactoryIntermediates Intermediates)
 {
-	
-	
-	return Intermediates.TangentZ;
+	return Intermediates.Current.TangentZ;
 }
 
 
@@ -830,13 +910,6 @@ FVertexFactoryInterpolantsVSToPS VertexFactoryGetInterpolantsVSToPS(FVertexFacto
 	return Interpolants;
 }
 
-
-float4 VertexFactoryGetPreviousWorldPosition(FVertexFactoryInput Input, FVertexFactoryIntermediates Intermediates)
-{
-	const float3 WorldPosition = CalcWorldPosition(Input, Intermediates).xyz;
-	return float4(WorldPosition - Intermediates.ParticleVelocity.xyz * (Intermediates.ParticleVelocity.w * NiagaraRibbonVF.DeltaSeconds), 1);
-}
-
 /**
 * Get the 3x3 tangent basis vectors for this vertex factory
 *
diff --git a/Engine/Plugins/FX/Niagara/Source/Niagara/Private/NiagaraRendererRibbons.cpp b/Engine/Plugins/FX/Niagara/Source/Niagara/Private/NiagaraRendererRibbons.cpp
--- a/Engine/Plugins/FX/Niagara/Source/Niagara/Private/NiagaraRendererRibbons.cpp
+++ b/Engine/Plugins/FX/Niagara/Source/Niagara/Private/NiagaraRendererRibbons.cpp
@@ -203,6 +203,7 @@ FNiagaraRendererRibbons::FNiagaraRendererRibbons(ERHIFeatureLevel::Type FeatureL
 	, bCustomUseConstantFactor(false)
 	, CustomTessellationMinAngle(15.f * PI / 180.f)
 	, bCustomUseScreenSpace(true)
+	, bNeedsPreciseMotionVectors(false)
 {
 	const UNiagaraRibbonRendererProperties* Properties = CastChecked<const UNiagaraRibbonRendererProperties>(InProps);
 	FacingMode = Properties->FacingMode;
@@ -224,6 +225,7 @@ FNiagaraRendererRibbons::FNiagaraRendererRibbons(ERHIFeatureLevel::Type FeatureL
 	bCustomUseScreenSpace = Properties->bScreenSpaceTessellation;
 	MaterialParamValidMask = Properties->MaterialParamValidMask;
 	RendererLayout = &Properties->RendererLayout;
+	bNeedsPreciseMotionVectors = Properties->NeedsPreciseMotionVectors();
 }
 
 FNiagaraRendererRibbons::~FNiagaraRendererRibbons()
@@ -1182,6 +1184,7 @@ void FNiagaraRendererRibbons::SetupMeshBatchAndCollectorResourceForView(
 	VFLooseParams.SortedIndicesOffset = CollectorResources.VertexFactory.GetSortedIndicesOffset();
 	VFLooseParams.FacingMode = static_cast<uint32>(FacingMode);
 	VFLooseParams.Shape = static_cast<uint32>(Shape);
+	VFLooseParams.NeedsPreciseMotionVectors = bNeedsPreciseMotionVectors;
 
 	// Collector.AllocateOneFrameResource uses default ctor, initialize the vertex factory
 	CollectorResources.VertexFactory.SetParticleFactoryType(NVFT_Ribbon);
@@ -1401,10 +1404,13 @@ void FNiagaraRendererRibbons::CreatePerViewResources(
 
 	TConstArrayView<FNiagaraRendererVariableInfo> VFVariables = RendererLayout->GetVFVariables_RenderThread();
 	PerViewUniformParameters.PositionDataOffset = VFVariables[ENiagaraRibbonVFLayout::Position].GetGPUOffset();
+	PerViewUniformParameters.PrevPositionDataOffset = VFVariables[ENiagaraRibbonVFLayout::PrevPosition].GetGPUOffset();
 	PerViewUniformParameters.VelocityDataOffset = VFVariables[ENiagaraRibbonVFLayout::Velocity].GetGPUOffset();
 	PerViewUniformParameters.ColorDataOffset = VFVariables[ENiagaraRibbonVFLayout::Color].GetGPUOffset();
 	PerViewUniformParameters.WidthDataOffset = VFVariables[ENiagaraRibbonVFLayout::Width].GetGPUOffset();
+	PerViewUniformParameters.PrevWidthDataOffset = VFVariables[ENiagaraRibbonVFLayout::PrevRibbonWidth].GetGPUOffset();
 	PerViewUniformParameters.TwistDataOffset = VFVariables[ENiagaraRibbonVFLayout::Twist].GetGPUOffset();
+	PerViewUniformParameters.PrevTwistDataOffset = VFVariables[ENiagaraRibbonVFLayout::PrevRibbonTwist].GetGPUOffset();
 	PerViewUniformParameters.NormalizedAgeDataOffset = VFVariables[ENiagaraRibbonVFLayout::NormalizedAge].GetGPUOffset();
 	PerViewUniformParameters.MaterialRandomDataOffset = VFVariables[ENiagaraRibbonVFLayout::MaterialRandom].GetGPUOffset();
 	PerViewUniformParameters.MaterialParamDataOffset = VFVariables[ENiagaraRibbonVFLayout::MaterialParam0].GetGPUOffset();
@@ -1424,6 +1430,7 @@ void FNiagaraRendererRibbons::CreatePerViewResources(
 
 	bool bShouldDoFacing = FacingMode == ENiagaraRibbonFacingMode::Custom || FacingMode == ENiagaraRibbonFacingMode::CustomSideVector;
 	PerViewUniformParameters.FacingDataOffset = bShouldDoFacing ? VFVariables[ENiagaraRibbonVFLayout::Facing].GetGPUOffset() : -1;
+	PerViewUniformParameters.PrevFacingDataOffset = bShouldDoFacing ? VFVariables[ENiagaraRibbonVFLayout::PrevRibbonFacing].GetGPUOffset() : -1;
 
 	PerViewUniformParameters.U0DistributionMode = (int32)UV0Settings.DistributionMode;
 	PerViewUniformParameters.U1DistributionMode = (int32)UV1Settings.DistributionMode;
diff --git a/Engine/Plugins/FX/Niagara/Source/Niagara/Private/NiagaraRibbonRendererProperties.cpp b/Engine/Plugins/FX/Niagara/Source/Niagara/Private/NiagaraRibbonRendererProperties.cpp
--- a/Engine/Plugins/FX/Niagara/Source/Niagara/Private/NiagaraRibbonRendererProperties.cpp
+++ b/Engine/Plugins/FX/Niagara/Source/Niagara/Private/NiagaraRibbonRendererProperties.cpp
@@ -83,6 +83,11 @@ UNiagaraRibbonRendererProperties::UNiagaraRibbonRendererProperties()
 	AttributeBindings.Add(&V0RangeOverrideBinding);
 	AttributeBindings.Add(&U1OverrideBinding);
 	AttributeBindings.Add(&V1RangeOverrideBinding);
+
+	AttributeBindings.Add(&PrevPositionBinding);
+	AttributeBindings.Add(&PrevRibbonWidthBinding);
+	AttributeBindings.Add(&PrevRibbonFacingBinding);
+	AttributeBindings.Add(&PrevRibbonTwistBinding);
 }
 
 FNiagaraRenderer* UNiagaraRibbonRendererProperties::CreateEmitterRenderer(ERHIFeatureLevel::Type FeatureLevel, const FNiagaraEmitterInstance* Emitter, const FNiagaraSystemInstanceController& InController)
@@ -250,6 +255,16 @@ void UNiagaraRibbonRendererProperties::InitBindings()
 		U1OverrideBinding = FNiagaraConstants::GetAttributeDefaultBinding(SYS_PARAM_PARTICLES_RIBBONU1OVERRIDE);
 		V1RangeOverrideBinding = FNiagaraConstants::GetAttributeDefaultBinding(SYS_PARAM_PARTICLES_RIBBONV1RANGEOVERRIDE);
 	}
+
+	SetPreviousBindings(nullptr);
+}
+
+void UNiagaraRibbonRendererProperties::SetPreviousBindings(const UNiagaraEmitter* SrcEmitter)
+{
+	PrevPositionBinding.SetAsPreviousValue(PositionBinding, SrcEmitter, ENiagaraRendererSourceDataMode::Particles);
+	PrevRibbonWidthBinding.SetAsPreviousValue(RibbonWidthBinding, SrcEmitter, ENiagaraRendererSourceDataMode::Particles);
+	PrevRibbonFacingBinding.SetAsPreviousValue(RibbonFacingBinding, SrcEmitter, ENiagaraRendererSourceDataMode::Particles);
+	PrevRibbonTwistBinding.SetAsPreviousValue(RibbonTwistBinding, SrcEmitter, ENiagaraRendererSourceDataMode::Particles);
 }
 
 void UNiagaraRibbonRendererProperties::CacheFromCompiledData(const FNiagaraDataSetCompiledData* CompiledData)
@@ -319,6 +334,14 @@ void UNiagaraRibbonRendererProperties::CacheFromCompiledData(const FNiagaraDataS
 	MaterialParamValidMask |= RendererLayout.SetVariableFromBinding(CompiledData, DynamicMaterial1Binding, ENiagaraRibbonVFLayout::MaterialParam1) ? 2 : 0;
 	MaterialParamValidMask |= RendererLayout.SetVariableFromBinding(CompiledData, DynamicMaterial2Binding, ENiagaraRibbonVFLayout::MaterialParam2) ? 4 : 0;
 	MaterialParamValidMask |= RendererLayout.SetVariableFromBinding(CompiledData, DynamicMaterial3Binding, ENiagaraRibbonVFLayout::MaterialParam3) ? 8 : 0;
+
+	if (NeedsPreciseMotionVectors())
+	{
+		RendererLayout.SetVariableFromBinding(CompiledData, PrevPositionBinding, ENiagaraRibbonVFLayout::PrevPosition);
+		RendererLayout.SetVariableFromBinding(CompiledData, PrevRibbonWidthBinding, ENiagaraRibbonVFLayout::PrevRibbonWidth);
+		RendererLayout.SetVariableFromBinding(CompiledData, PrevRibbonFacingBinding, ENiagaraRibbonVFLayout::PrevRibbonFacing);
+		RendererLayout.SetVariableFromBinding(CompiledData, PrevRibbonTwistBinding, ENiagaraRibbonVFLayout::PrevRibbonTwist);
+	}
 	RendererLayout.Finalize();
 }
 
@@ -374,6 +397,34 @@ const TArray<FNiagaraVariable>& UNiagaraRibbonRendererProperties::GetOptionalAtt
 	return Attrs;
 }
 
+void UNiagaraRibbonRendererProperties::GetAdditionalVariables(TArray<FNiagaraVariableBase>& OutArray) const
+{
+	if (NeedsPreciseMotionVectors())
+	{
+		OutArray.Append({
+				PrevPositionBinding.GetParamMapBindableVariable(),
+				PrevRibbonWidthBinding.GetParamMapBindableVariable(),
+				PrevRibbonFacingBinding.GetParamMapBindableVariable(),
+				PrevRibbonTwistBinding.GetParamMapBindableVariable()
+		});
+	}
+}
+
+FNiagaraVariable UNiagaraRibbonRendererProperties::GetBoundAttribute(const FNiagaraVariableAttributeBinding* Binding) const
+{
+	if (!NeedsPreciseMotionVectors())
+	{
+		if ((Binding == &PrevPositionBinding)
+			|| (Binding == &PrevRibbonWidthBinding)
+			|| (Binding == &PrevRibbonFacingBinding)
+			|| (Binding == &PrevRibbonTwistBinding))
+		{
+			return FNiagaraVariable();
+		}
+	}
+
+	return Super::GetBoundAttribute(Binding);
+}
 
 void UNiagaraRibbonRendererProperties::GetRendererWidgets(const FNiagaraEmitterInstance* InEmitter, TArray<TSharedPtr<SWidget>>& OutWidgets, TSharedPtr<FAssetThumbnailPool> InThumbnailPool) const
 {
diff --git a/Engine/Plugins/FX/Niagara/Source/Niagara/Public/NiagaraRendererRibbons.h b/Engine/Plugins/FX/Niagara/Source/Niagara/Public/NiagaraRendererRibbons.h
--- a/Engine/Plugins/FX/Niagara/Source/Niagara/Public/NiagaraRendererRibbons.h
+++ b/Engine/Plugins/FX/Niagara/Source/Niagara/Public/NiagaraRendererRibbons.h
@@ -118,6 +118,7 @@ private:
 	bool bCustomUseConstantFactor;
 	float CustomTessellationMinAngle;
 	bool bCustomUseScreenSpace;
+	bool bNeedsPreciseMotionVectors;
 
 
 	uint32 MaterialParamValidMask;
diff --git a/Engine/Plugins/FX/Niagara/Source/Niagara/Public/NiagaraRibbonRendererProperties.h b/Engine/Plugins/FX/Niagara/Source/Niagara/Public/NiagaraRibbonRendererProperties.h
--- a/Engine/Plugins/FX/Niagara/Source/Niagara/Public/NiagaraRibbonRendererProperties.h
+++ b/Engine/Plugins/FX/Niagara/Source/Niagara/Public/NiagaraRibbonRendererProperties.h
@@ -177,6 +177,10 @@ namespace ENiagaraRibbonVFLayout
 		V0RangeOverride,
 		U1Override,
 		V1RangeOverride,
+		PrevPosition,
+		PrevRibbonWidth,
+		PrevRibbonFacing,
+		PrevRibbonTwist,
 		Num,
 	};
 };
@@ -213,6 +217,8 @@ public:
 	virtual bool IsMaterialValidForRenderer(UMaterial* Material, FText& InvalidMessage) override;
 	virtual void FixMaterial(UMaterial* Material);
 	virtual const TArray<FNiagaraVariable>& GetOptionalAttributes() override;
+	virtual void GetAdditionalVariables(TArray<FNiagaraVariableBase>& OutArray) const override;
+	virtual FNiagaraVariable GetBoundAttribute(const FNiagaraVariableAttributeBinding* Binding) const override;
 	virtual void GetRendererWidgets(const FNiagaraEmitterInstance* InEmitter, TArray<TSharedPtr<SWidget>>& OutWidgets, TSharedPtr<FAssetThumbnailPool> InThumbnailPool) const override;
 	virtual void GetRendererTooltipWidgets(const FNiagaraEmitterInstance* InEmitter, TArray<TSharedPtr<SWidget>>& OutWidgets, TSharedPtr<FAssetThumbnailPool> InThumbnailPool) const override;
 	virtual void GetRendererFeedback(const UNiagaraEmitter* InEmitter, TArray<FText>& OutErrors, TArray<FText>& OutWarnings, TArray<FText>& OutInfo) const override;
@@ -414,6 +420,23 @@ public:
 	UPROPERTY(EditAnywhere, Category = "Bindings")
 	TArray<FNiagaraMaterialAttributeBinding> MaterialParameterBindings;
 
+	/** Implicit binding for previous position */
+	UPROPERTY(Transient)
+	FNiagaraVariableAttributeBinding PrevPositionBinding;
+
+	/** Implicit binding for previous ribbon width */
+	UPROPERTY(Transient)
+	FNiagaraVariableAttributeBinding PrevRibbonWidthBinding;
+
+	/** Implicit binding for previous ribbon facing */
+	UPROPERTY(Transient)
+	FNiagaraVariableAttributeBinding PrevRibbonFacingBinding;
+
+	/** Implicit binding for previous ribbon twist */
+	UPROPERTY(Transient)
+	FNiagaraVariableAttributeBinding PrevRibbonTwistBinding;
+
+
 	bool								bSortKeyDataSetAccessorIsAge = false;
 	FNiagaraDataSetAccessor<float>		SortKeyDataSetAccessor;
 	FNiagaraDataSetAccessor<FNiagaraPosition>	PositionDataSetAccessor;
@@ -437,6 +460,7 @@ public:
 
 protected:
 	void InitBindings();
+	void SetPreviousBindings(const UNiagaraEmitter* SrcEmitter);
 
 	void UpdateSourceModeDerivates(ENiagaraRendererSourceDataMode InSourceMode, bool bFromPropertyEdit);
 
diff --git a/Engine/Plugins/FX/Niagara/Source/NiagaraVertexFactories/Public/NiagaraRibbonVertexFactory.h b/Engine/Plugins/FX/Niagara/Source/NiagaraVertexFactories/Public/NiagaraRibbonVertexFactory.h
--- a/Engine/Plugins/FX/Niagara/Source/NiagaraVertexFactories/Public/NiagaraRibbonVertexFactory.h
+++ b/Engine/Plugins/FX/Niagara/Source/NiagaraVertexFactories/Public/NiagaraRibbonVertexFactory.h
@@ -30,11 +30,15 @@ BEGIN_GLOBAL_SHADER_PARAMETER_STRUCT(FNiagaraRibbonUniformParameters, NIAGARAVER
 	SHADER_PARAMETER(FVector4f, CameraUp)
 	SHADER_PARAMETER(FVector4f, ScreenAlignment)
 	SHADER_PARAMETER(int, PositionDataOffset)
+	SHADER_PARAMETER(int, PrevPositionDataOffset)
 	SHADER_PARAMETER(int, VelocityDataOffset)
 	SHADER_PARAMETER(int, WidthDataOffset)
+	SHADER_PARAMETER(int, PrevWidthDataOffset)
 	SHADER_PARAMETER(int, TwistDataOffset)
+	SHADER_PARAMETER(int, PrevTwistDataOffset)
 	SHADER_PARAMETER(int, ColorDataOffset)
 	SHADER_PARAMETER(int, FacingDataOffset)
+	SHADER_PARAMETER(int, PrevFacingDataOffset)
 	SHADER_PARAMETER(int, NormalizedAgeDataOffset)
 	SHADER_PARAMETER(int, MaterialRandomDataOffset)
 	SHADER_PARAMETER(uint32, MaterialParamValidMask)
@@ -70,6 +74,7 @@ BEGIN_GLOBAL_SHADER_PARAMETER_STRUCT(FNiagaraRibbonVFLooseParameters, NIAGARAVER
 	SHADER_PARAMETER(uint32, SortedIndicesOffset)
 	SHADER_PARAMETER(uint32, FacingMode)
 	SHADER_PARAMETER(uint32, Shape)
+	SHADER_PARAMETER(uint32, NeedsPreciseMotionVectors)
 	SHADER_PARAMETER_SRV(Buffer<int>, SortedIndices)
 	SHADER_PARAMETER_SRV(Buffer<float4>, TangentsAndDistances)
 	SHADER_PARAMETER_SRV(Buffer<uint>, MultiRibbonIndices)
