Type change menu can now be used on multiple parameters at once, which enables maintaining pin connections between different parameters if they are changed at the same time.
For any other pin type change, orphan pins will get created and these outdated pin connections will be transferred to the new orphan pin.

- Some parameter actions were changed to accept arrays instead of single items, creating one transaction per action instead of item
- Other actions now require a single item to be selected instead in order to not change behavior too much
- Only the script asset parameters should be affected
- A few fixes in regards to orphaned pins. Map get & set nodes would markup even orphaned pins as "parameter pins", which re-added the outdated parameters they represented when reopening a script

diff --git a/Engine/Plugins/FX/Niagara/Source/NiagaraEditor/Private/NiagaraGraph.cpp b/Engine/Plugins/FX/Niagara/Source/NiagaraEditor/Private/NiagaraGraph.cpp
--- a/Engine/Plugins/FX/Niagara/Source/NiagaraEditor/Private/NiagaraGraph.cpp
+++ b/Engine/Plugins/FX/Niagara/Source/NiagaraEditor/Private/NiagaraGraph.cpp
@@ -467,68 +467,168 @@ void UNiagaraGraph::PostLoad()
 	InvalidateCachedParameterData();
 }
 
-void UNiagaraGraph::ChangeParameterType(const FNiagaraVariable& CurrentParameter, const FNiagaraTypeDefinition& NewType)
+void UNiagaraGraph::ChangeParameterType(const TArray<FNiagaraVariable>& ParametersToChange, const FNiagaraTypeDefinition& NewType, bool bAllowOrphanedPins)
 {
 	const UEdGraphSchema_Niagara* NiagaraSchema = GetDefault<UEdGraphSchema_Niagara>();
 	TArray<UNiagaraNodeParameterMapBase*> NiagaraParameterNodes;
 	GetNodesOfClass<UNiagaraNodeParameterMapBase>(NiagaraParameterNodes);
-	for (UNiagaraNodeParameterMapBase* NiagaraNode : NiagaraParameterNodes)
+
+	TMap<UEdGraphPin*, FNiagaraTypeDefinition> OldTypesCache;
+	
+	auto TryCreateOrphanedPin = [&](UNiagaraNode* Node, FName OrphanedPinName, UEdGraphPin* OriginalPin)
 	{
-		for (UEdGraphPin* Pin : NiagaraNode->Pins)
+		// We check pin connections for type to determine if we want to keep a connection or discard it
+		TSet<UEdGraphPin*> PinConnectionsToKeep;
+		for (UEdGraphPin* ConnectedPin : OriginalPin->LinkedTo)
 		{
-			if (NiagaraNode->IsAddPin(Pin))
+			// we might have already visited and changed the connected pin, so we test against its old type cache if available
+			FNiagaraTypeDefinition CurrentPinType = UEdGraphSchema_Niagara::PinToTypeDefinition(ConnectedPin);
+			FNiagaraVariable ConnectedVariable(OldTypesCache.Contains(ConnectedPin) ? OldTypesCache[ConnectedPin] : CurrentPinType, ConnectedPin->PinName);
+			
+			// keep it if the connected pin has a matching type already or will have a matching type after type conversion is done
+			if((ConnectedPin->PinType.PinSubCategory == UNiagaraNodeParameterMapBase::ParameterPinSubCategory && ParametersToChange.Contains(ConnectedVariable))
+				|| CurrentPinType == NewType)
 			{
-				continue;
+				PinConnectionsToKeep.Add(ConnectedPin);
+			}						
+		}
+
+		// in case we are keeping all connections, all previously existing connections are parameter pins of the same type or will be of the same type after this operation is done
+		if(PinConnectionsToKeep.Num() != OriginalPin->LinkedTo.Num())
+		{
+			ensure(OldTypesCache.Contains(OriginalPin));
+			
+			UEdGraphPin* NewOrphanedPin = Node->CreatePin(OriginalPin->Direction, OriginalPin->PinType.PinCategory, OrphanedPinName);
+			NewOrphanedPin->PinType = UEdGraphSchema_Niagara::TypeDefinitionToPinType(OldTypesCache[OriginalPin]);
+			NewOrphanedPin->PersistentGuid = FGuid::NewGuid();
+			NewOrphanedPin->bNotConnectable = true;
+			NewOrphanedPin->bOrphanedPin = true;
+			
+			TSet<UEdGraphPin*> PinsToDisconnect;
+			for (UEdGraphPin* PinConnection : OriginalPin->LinkedTo)
+			{
+				// wire up the new orphaned pin with previous connections if we aren't keeping the connection
+				if(!PinConnectionsToKeep.Contains(PinConnection))
+				{
+					NewOrphanedPin->LinkedTo.Add(PinConnection);
+					PinConnection->LinkedTo.Add(NewOrphanedPin);
+
+					PinsToDisconnect.Add(PinConnection);
+				}
 			}
-			FNiagaraVariable Variable = FNiagaraVariable(NiagaraSchema->PinToTypeDefinition(Pin), Pin->PinName);
-			if (Variable == CurrentParameter)
+
+			// disconnect the previous connections from the original pin
+			for (UEdGraphPin* PinToDisconnect : PinsToDisconnect)
 			{
-				Pin->Modify();
-				Pin->PinType = NiagaraSchema->TypeDefinitionToPinType(NewType);
-				Pin->PinType.PinSubCategory = UNiagaraNodeParameterMapBase::ParameterPinSubCategory;
+				PinToDisconnect->LinkedTo.Remove(OriginalPin);
+				OriginalPin->LinkedTo.Remove(PinToDisconnect);
+			}							
+		}
+	};
 
-				// fix default pin types if necessary
-				if (UNiagaraNodeParameterMapGet* GetNode = Cast<UNiagaraNodeParameterMapGet>(NiagaraNode))
+	// for every parameter, we loop over every node and every pin until we encounter a pin representing the parameter
+	for(const FNiagaraVariable& CurrentParameter : ParametersToChange)
+	{
+		// a safe guard so we only look for pins that actually represent parameters and not just any pin name with the same name as a parameter
+		if(!VariableToScriptVariable.Contains(CurrentParameter))
+		{
+			continue;
+		}
+		
+		for (UNiagaraNodeParameterMapBase* NiagaraNode : NiagaraParameterNodes)
+		{
+			TSet<UEdGraphPin*> NewOrphanedPins;
+			for (UEdGraphPin* Pin : NiagaraNode->Pins)
+			{
+				if (NiagaraNode->IsAddPin(Pin) || Pin->bOrphanedPin)
 				{
-					UEdGraphPin* DefaultPin = GetNode->GetDefaultPin(Pin);
-					if (Pin->Direction == EGPD_Output && DefaultPin && DefaultPin->PinType != Pin->PinType)
+					continue;
+				}
+
+				FNiagaraVariable PinVariable = FNiagaraVariable(UEdGraphSchema_Niagara::PinToTypeDefinition(Pin), Pin->PinName);
+				if (PinVariable == CurrentParameter)
+				{
+					// we cache off the old type before changing it as we might encounter a pin we already changed later on in connections
+					if(!OldTypesCache.Contains(Pin))
 					{
-						DefaultPin->PinType = Pin->PinType;
+						OldTypesCache.Add(Pin, PinVariable.GetType());					
 					}
-				}
+					
+					Pin->Modify();
+					Pin->PinType = NiagaraSchema->TypeDefinitionToPinType(NewType);
+					Pin->PinType.PinSubCategory = UNiagaraNodeParameterMapBase::ParameterPinSubCategory;
 
-				NiagaraNode->MarkNodeRequiresSynchronization(__FUNCTION__, true);
-				break;
+					// we only want to create an additional orphaned pin if there is at least one pin connection or if the pin value differs from the 'default' pin value
+					if(bAllowOrphanedPins)
+					{
+						bool bCreateOrphanedPin = (Pin->bDefaultValueIsIgnored == false && Pin->DefaultValue != Pin->AutogeneratedDefaultValue) || Pin->LinkedTo.Num() > 0;
+						if(bCreateOrphanedPin)
+						{
+							TryCreateOrphanedPin(NiagaraNode, Pin->PinName, Pin);
+						}
+					}
+
+					// fix default pin types if necessary. We don't support orphaned pins for default types so we just disconnect the pin
+					if (UNiagaraNodeParameterMapGet* GetNode = Cast<UNiagaraNodeParameterMapGet>(NiagaraNode))
+					{
+						UEdGraphPin* DefaultPin = GetNode->GetDefaultPin(Pin);
+						if (Pin->Direction == EGPD_Output && DefaultPin && Pin->PinType != DefaultPin->PinType)
+						{
+							ensure(OldTypesCache.Contains(Pin));
+							
+							if(!OldTypesCache.Contains(DefaultPin))
+							{
+								OldTypesCache.Add(DefaultPin, UEdGraphSchema_Niagara::PinToTypeDefinition(DefaultPin));					
+							}
+							
+							DefaultPin->Modify();
+							DefaultPin->PinType = Pin->PinType;
+
+							if(bAllowOrphanedPins)
+							{
+								bool bCreateOrphanedPin = DefaultPin->LinkedTo.Num() > 0;
+
+								if(bCreateOrphanedPin)
+								{
+									TryCreateOrphanedPin(NiagaraNode, Pin->PinName, DefaultPin);
+								}
+							}
+						}
+					}
+
+					NiagaraNode->MarkNodeRequiresSynchronization(__FUNCTION__, true);
+					break;
+				}			
 			}
 		}
-	}
-	
-	TObjectPtr<UNiagaraScriptVariable>* ScriptVariablePtr = VariableToScriptVariable.Find(CurrentParameter);
-	if (ScriptVariablePtr != nullptr && !ScriptVariablePtr->IsNull())
-	{
-		UNiagaraScriptVariable& ScriptVar = *ScriptVariablePtr->Get();
-		ScriptVar.Modify();
-		ScriptVar.Variable.SetType(NewType);
-		VariableToScriptVariable.Remove(CurrentParameter);
+		
+		TObjectPtr<UNiagaraScriptVariable>* ScriptVariablePtr = VariableToScriptVariable.Find(CurrentParameter);
+		if (ScriptVariablePtr != nullptr && !ScriptVariablePtr->IsNull())
+		{
+			UNiagaraScriptVariable& ScriptVar = *ScriptVariablePtr->Get();
+			ScriptVar.Modify();
+			ScriptVar.Variable.SetType(NewType);
+			VariableToScriptVariable.Remove(CurrentParameter);
 
-		FNiagaraVariable NewVarType(NewType, CurrentParameter.GetName());
-		VariableToScriptVariable.Add(NewVarType, TObjectPtr<UNiagaraScriptVariable>(&ScriptVar));
-		ScriptVariableChanged(NewVarType);
-	}
+			FNiagaraVariable NewVarType(NewType, CurrentParameter.GetName());
+			VariableToScriptVariable.Add(NewVarType, TObjectPtr<UNiagaraScriptVariable>(&ScriptVar));
+			ScriptVariableChanged(NewVarType);
+		}
 
-	// Also fix stale parameter reference map entries. This should usually happen in RefreshParameterReferences, but since we just changed VariableToScriptVariable we need to fix that as well. 
-	FNiagaraGraphParameterReferenceCollection* ReferenceCollection = ParameterToReferencesMap.Find(CurrentParameter);
-	if (ReferenceCollection != nullptr)
-	{
-		if (ReferenceCollection->ParameterReferences.Num() > 0)
+		// Also fix stale parameter reference map entries. This should usually happen in RefreshParameterReferences, but since we just changed VariableToScriptVariable we need to fix that as well. 
+		FNiagaraGraphParameterReferenceCollection* ReferenceCollection = ParameterToReferencesMap.Find(CurrentParameter);
+		if (ReferenceCollection != nullptr)
 		{
-			FNiagaraVariable NewVarType(NewType, CurrentParameter.GetName());
-			FNiagaraGraphParameterReferenceCollection NewReferenceCollection = *ReferenceCollection;
-			ParameterToReferencesMap.Add(NewVarType, NewReferenceCollection);
+			if (ReferenceCollection->ParameterReferences.Num() > 0)
+			{
+				FNiagaraVariable NewVarType(NewType, CurrentParameter.GetName());
+				FNiagaraGraphParameterReferenceCollection NewReferenceCollection = *ReferenceCollection;
+				ParameterToReferencesMap.Add(NewVarType, NewReferenceCollection);
+			}
+			ParameterToReferencesMap.Remove(CurrentParameter);
 		}
-		ParameterToReferencesMap.Remove(CurrentParameter);
-	}
-	InvalidateCachedParameterData();
+		InvalidateCachedParameterData();
+	}	
 }
 
 TArray<UEdGraphPin*> UNiagaraGraph::FindParameterMapDefaultValuePins(const FName VariableName) const
diff --git a/Engine/Plugins/FX/Niagara/Source/NiagaraEditor/Private/NiagaraNodeParameterMapGet.cpp b/Engine/Plugins/FX/Niagara/Source/NiagaraEditor/Private/NiagaraNodeParameterMapGet.cpp
--- a/Engine/Plugins/FX/Niagara/Source/NiagaraEditor/Private/NiagaraNodeParameterMapGet.cpp
+++ b/Engine/Plugins/FX/Niagara/Source/NiagaraEditor/Private/NiagaraNodeParameterMapGet.cpp
@@ -91,7 +91,7 @@ UEdGraphPin* UNiagaraNodeParameterMapGet::CreateDefaultPin(UEdGraphPin* OutputPi
 	{
 		return nullptr;
 	}
-
+	
 	UEdGraphPin* DefaultPin = CreatePin(EEdGraphPinDirection::EGPD_Input, OutputPin->PinType, TEXT(""));
 
 	// make sure the new pin name is legal
@@ -294,7 +294,7 @@ void UNiagaraNodeParameterMapGet::PostLoad()
 	for (int32 i = 0; i < OutputPins.Num(); i++)
 	{
 		UEdGraphPin* OutputPin = OutputPins[i];
-		if (IsAddPin(OutputPin))
+		if (IsAddPin(OutputPin) || OutputPin->bOrphanedPin)
 		{
 			continue;
 		}
diff --git a/Engine/Plugins/FX/Niagara/Source/NiagaraEditor/Private/NiagaraNodeParameterMapSet.cpp b/Engine/Plugins/FX/Niagara/Source/NiagaraEditor/Private/NiagaraNodeParameterMapSet.cpp
--- a/Engine/Plugins/FX/Niagara/Source/NiagaraEditor/Private/NiagaraNodeParameterMapSet.cpp
+++ b/Engine/Plugins/FX/Niagara/Source/NiagaraEditor/Private/NiagaraNodeParameterMapSet.cpp
@@ -343,7 +343,7 @@ void UNiagaraNodeParameterMapSet::PostLoad()
 			Pin->PersistentGuid = FGuid::NewGuid();
 		}
 
-		if (Pin->Direction == EEdGraphPinDirection::EGPD_Input && Pin->GetFName() != UNiagaraNodeParameterMapBase::SourcePinName && !IsAddPin(Pin))
+		if (Pin->Direction == EEdGraphPinDirection::EGPD_Input && Pin->GetFName() != UNiagaraNodeParameterMapBase::SourcePinName && !IsAddPin(Pin) && !Pin->bOrphanedPin)
 		{
 			Pin->PinType.PinSubCategory = UNiagaraNodeParameterMapBase::ParameterPinSubCategory;
 		}
diff --git a/Engine/Plugins/FX/Niagara/Source/NiagaraEditor/Private/NiagaraScriptSource.cpp b/Engine/Plugins/FX/Niagara/Source/NiagaraEditor/Private/NiagaraScriptSource.cpp
--- a/Engine/Plugins/FX/Niagara/Source/NiagaraEditor/Private/NiagaraScriptSource.cpp
+++ b/Engine/Plugins/FX/Niagara/Source/NiagaraEditor/Private/NiagaraScriptSource.cpp
@@ -656,5 +656,5 @@ void UNiagaraScriptSource::ChangedLinkedInputTypes(const FNiagaraVariable& Param
 	{
 		return;
 	}
-	NodeGraph->ChangeParameterType(ParametersToChange, NewType);
+	NodeGraph->ChangeParameterType({ParametersToChange}, NewType);
 }
diff --git a/Engine/Plugins/FX/Niagara/Source/NiagaraEditor/Private/ViewModels/NiagaraParameterPanelViewModel.cpp b/Engine/Plugins/FX/Niagara/Source/NiagaraEditor/Private/ViewModels/NiagaraParameterPanelViewModel.cpp
--- a/Engine/Plugins/FX/Niagara/Source/NiagaraEditor/Private/ViewModels/NiagaraParameterPanelViewModel.cpp
+++ b/Engine/Plugins/FX/Niagara/Source/NiagaraEditor/Private/ViewModels/NiagaraParameterPanelViewModel.cpp
@@ -2,6 +2,7 @@
 
 #include "ViewModels/NiagaraParameterPanelViewModel.h"
 
+#include "EdGraphSchema_Niagara.h"
 #include "NiagaraActions.h"
 #include "NiagaraClipboard.h"
 #include "NiagaraConstants.h"
@@ -372,28 +373,45 @@ bool INiagaraParameterPanelViewModel::GetCanPasteParameterMetaDataAndToolTip(FTe
 	return true;
 }
 
-void INiagaraParameterPanelViewModel::DuplicateParameter(const FNiagaraParameterPanelItem ItemToDuplicate) const
+void INiagaraParameterPanelViewModel::DuplicateParameters(const TArray<FNiagaraParameterPanelItem> ItemsToDuplicate) const
 {
-	TSet<FName> ParameterNames;
-	for (const FNiagaraParameterPanelItem& CachedViewedItem : CachedViewedItems)
+	bool bAnyChange = false;
+	FScopedTransaction Transaction(LOCTEXT("DuplicateParameterTransaction", "Duplicate parameter(s)"));
+
+	for(const FNiagaraParameterPanelItem& ItemToDuplicate : ItemsToDuplicate)
 	{
-		ParameterNames.Add(CachedViewedItem.GetVariable().GetName());
+		TSet<FName> ParameterNames;
+		for (const FNiagaraParameterPanelItem& CachedViewedItem : CachedViewedItems)
+		{
+			ParameterNames.Add(CachedViewedItem.GetVariable().GetName());
+		}
+		
+		const FName NewUniqueName = FNiagaraUtilities::GetUniqueName(ItemToDuplicate.GetVariable().GetName(), ParameterNames);
+		const bool bRequestRename = true;
+		const bool bMakeUniqueName = true;
+		AddParameter(FNiagaraVariable(ItemToDuplicate.GetVariable().GetType(), NewUniqueName), FNiagaraParameterPanelCategory(ItemToDuplicate.NamespaceMetaData), bRequestRename, bMakeUniqueName);
+		bAnyChange = true;
 	}
-	const FName NewUniqueName = FNiagaraUtilities::GetUniqueName(ItemToDuplicate.GetVariable().GetName(), ParameterNames);
-	FScopedTransaction Transaction(LOCTEXT("DuplicateParameterTransaction", "Duplicate parameter"));
-	const bool bRequestRename = true;
-	const bool bMakeUniqueName = true;
-	AddParameter(FNiagaraVariable(ItemToDuplicate.GetVariable().GetType(), NewUniqueName), FNiagaraParameterPanelCategory(ItemToDuplicate.NamespaceMetaData), bRequestRename, bMakeUniqueName);
+
+	if(!bAnyChange)
+	{
+		Transaction.Cancel();
+	}
+		
 }
 
-bool INiagaraParameterPanelViewModel::GetCanDuplicateParameterAndToolTip(const FNiagaraParameterPanelItem& ItemToDuplicate, FText& OutCanDuplicateParameterToolTip) const
+bool INiagaraParameterPanelViewModel::GetCanDuplicateParameterAndToolTip(const TArray<FNiagaraParameterPanelItem>& ItemsToDuplicate, FText& OutCanDuplicateParameterToolTip) const
 {
-	if (ItemToDuplicate.NamespaceMetaData.Options.Contains(ENiagaraNamespaceMetadataOptions::PreventEditingName))
+	for(const FNiagaraParameterPanelItem& ItemToDuplicate : ItemsToDuplicate)
 	{
-		OutCanDuplicateParameterToolTip = LOCTEXT("ParameterPanelViewModel_DuplicateParameter_PreventEditingName", "This parameter can not be duplicated because it does not support editing its name.");
-		return false;
+		if (ItemToDuplicate.NamespaceMetaData.Options.Contains(ENiagaraNamespaceMetadataOptions::PreventEditingName))
+		{
+			OutCanDuplicateParameterToolTip = LOCTEXT("ParameterPanelViewModel_DuplicateParameter_PreventEditingName", "A parameter can not be duplicated because it does not support editing its name.");
+			return false;
+		}
 	}
-	OutCanDuplicateParameterToolTip = LOCTEXT("ParameterPanellViewModel_DuplicateParameter", "Create a new parameter with the same type as this parameter.");
+	
+	OutCanDuplicateParameterToolTip = LOCTEXT("ParameterPanellViewModel_DuplicateParameter", "Create new parameters with the same types as the selected parameters.");
 	return true;
 }
 
@@ -561,54 +579,57 @@ void INiagaraParameterPanelViewModel::SetParameterNamespaceModifier(const FNiaga
 	}
 }
 
-void INiagaraParameterPanelViewModel::ChangeParameterType(const FNiagaraParameterPanelItem ItemToModify, const FNiagaraTypeDefinition NewType) const
+void INiagaraParameterPanelViewModel::ChangeParameterType(const TArray<FNiagaraParameterPanelItem> ItemsToModify, const FNiagaraTypeDefinition NewType) const
 {
 	ensureMsgf(false, TEXT("type change not supported for this view"));
 }
 
-bool INiagaraParameterPanelViewModel::GetCanChangeParameterType(const FNiagaraParameterPanelItem ItemToChange, FText& OutTooltip) const
+bool INiagaraParameterPanelViewModel::GetCanChangeParameterType(const TArray<FNiagaraParameterPanelItem>& ItemsToChange, FText& OutTooltip) const
 {
-	if (ItemToChange.bExternallyReferenced)
+	for (const FNiagaraParameterPanelItem& ItemToChange : ItemsToChange)
 	{
-		OutTooltip = LOCTEXT("CantChangeSelectedType_External", "This parameter is referenced in an external script and its type cannot be changed.");
-		return false;
-	}
-	
-	for(const UNiagaraGraph* Graph : GetEditableGraphsConst())
-	{
-		if(UNiagaraScriptVariable* ScriptVariable = Graph->GetScriptVariable(ItemToChange.GetVariable()))
+		if (ItemToChange.bExternallyReferenced)
 		{
-			if(ScriptVariable->GetIsStaticSwitch())
+			OutTooltip = LOCTEXT("CantChangeSelectedType_External", "A parameter is referenced in an external script and its type cannot be changed.");
+			return false;
+		}
+		
+		for(const UNiagaraGraph* Graph : GetEditableGraphsConst())
+		{
+			if(UNiagaraScriptVariable* ScriptVariable = Graph->GetScriptVariable(ItemToChange.GetVariable()))
 			{
-				OutTooltip = LOCTEXT("CantChangeSelectedType_StaticSwitch", "This parameter is a static switch parameter and its type cannot be changed here.");
-				return false;
-			}
+				if(ScriptVariable->GetIsStaticSwitch())
+				{
+					OutTooltip = LOCTEXT("CantChangeSelectedType_StaticSwitch", "A parameter is a static switch parameter and its type cannot be changed here.");
+					return false;
+				}
 
-			if(ItemToChange.GetVariable().GetType().IsStatic())
-			{
-				OutTooltip = LOCTEXT("CantChangeSelectedType_StaticParameter", "This parameter is static and its type cannot be changed.");
-				return false;
-			}
+				if(ItemToChange.GetVariable().GetType().IsStatic())
+				{
+					OutTooltip = LOCTEXT("CantChangeSelectedType_StaticParameter", "A parameter is static and its type cannot be changed.");
+					return false;
+				}
 
-			if(ScriptVariable->GetIsSubscribedToParameterDefinitions())
-			{
-				OutTooltip = LOCTEXT("CantChangeSelectedType_SubscribedParameter", "This parameter is subscribed to a parameter definition. The parameter can only be changed in the parameter definition asset.");
-				return false;
+				if(ScriptVariable->GetIsSubscribedToParameterDefinitions())
+				{
+					OutTooltip = LOCTEXT("CantChangeSelectedType_SubscribedParameter", "A parameter is subscribed to a parameter definition. The parameter can only be changed in the parameter definition asset.");
+					return false;
+				}
 			}
-		}
+		}	
 	}	
 
-	OutTooltip = LOCTEXT("CanChangeSelectedType", "Change this parameter's type");
+	OutTooltip = LOCTEXT("CanChangeSelectedType", "Change the type of selected parameters.\nPin connections will get fixed up as much as possible while leaving orphaned pins for connections that can't be maintained.");
 	return true;
 }
 
-void INiagaraParameterPanelViewModel::GetChangeTypeSubMenu(FMenuBuilder& MenuBuilder, FNiagaraParameterPanelItem Item) const
+void INiagaraParameterPanelViewModel::GetChangeTypeSubMenu(FMenuBuilder& MenuBuilder, TArray<FNiagaraParameterPanelItem> Items) const
 {
 	TArray<FNiagaraTypeDefinition> FilteredTypes;
 	for (const FNiagaraTypeDefinition& RegisteredType : FNiagaraTypeRegistry::GetRegisteredTypes())
 	{
 		// only allow basic types for now
-		if (RegisteredType.IsDataInterface() || RegisteredType.IsEnum() || RegisteredType.IsUObject() ||RegisteredType.GetNameText().ToString().Contains("event"))
+		if (RegisteredType.IsDataInterface() || RegisteredType.IsEnum() || RegisteredType.IsUObject() || FNiagaraTypeRegistry::GetRegisteredPayloadTypes().Contains(RegisteredType) || RegisteredType.IsInternalType())
 		{
 			continue;
 		}
@@ -625,7 +646,7 @@ void INiagaraParameterPanelViewModel::GetChangeTypeSubMenu(FMenuBuilder& MenuBui
 			Type.GetNameText(),
 			LOCTEXT("ChangeTypeActionTooltip", "Change the type of this variable. Replaced all pins in the graph and metadata with the selected type."),
 			FSlateIcon(),
-			FUIAction(FExecuteAction::CreateSP(this, &INiagaraParameterPanelViewModel::ChangeParameterType, Item, Type))); 
+			FUIAction(FExecuteAction::CreateSP(this, &INiagaraParameterPanelViewModel::ChangeParameterType, Items, Type))); 
 	}
 }
 
@@ -1047,25 +1068,39 @@ bool FNiagaraSystemToolkitParameterPanelViewModel::GetCanAddParametersToCategory
 	return GetEditableGraphsConst().Num() > 0 && Category.NamespaceMetaData.GetGuid() != FNiagaraEditorGuids::StaticSwitchNamespaceMetaDataGuid;
 }
 
-void FNiagaraSystemToolkitParameterPanelViewModel::DeleteParameter(const FNiagaraParameterPanelItem& ItemToDelete) const
+void FNiagaraSystemToolkitParameterPanelViewModel::DeleteParameters(const TArray<FNiagaraParameterPanelItem>& ItemsToDelete) const
 {
-	if (ItemToDelete.bExternallyReferenced)
+	bool bAnyChange = false;
+	FScopedTransaction RemoveParameterTransaction(LOCTEXT("RemoveParameter", "Removed Parameter(s)"));
+
+	for(const FNiagaraParameterPanelItem& ItemToDelete : ItemsToDelete)
 	{
-		return;
-	}
+		if (ItemToDelete.bExternallyReferenced)
+		{
+			continue;
+		}
+		
+		UNiagaraSystem& System = SystemViewModel->GetSystem();
+		const FGuid& ScriptVarId = ItemToDelete.GetVariableMetaData().GetVariableGuid();
+		System.Modify();
+		System.GetExposedParameters().RemoveParameter(ItemToDelete.GetVariable());
+		UNiagaraEditorParametersAdapter* EditorParametersAdapter = SystemViewModel->GetEditorOnlyParametersAdapter();
+		EditorParametersAdapter->Modify();
+		EditorParametersAdapter->GetParameters().RemoveAll([&ScriptVarId](const UNiagaraScriptVariable* ScriptVar) { return ScriptVar->Metadata.GetVariableGuid() == ScriptVarId; });
 
-	FScopedTransaction RemoveParameter(LOCTEXT("RemoveParameter", "Remove Parameter"));
-	UNiagaraSystem& System = SystemViewModel->GetSystem();
-	const FGuid& ScriptVarId = ItemToDelete.GetVariableMetaData().GetVariableGuid();
-	System.Modify();
-	System.GetExposedParameters().RemoveParameter(ItemToDelete.GetVariable());
-	UNiagaraEditorParametersAdapter* EditorParametersAdapter = SystemViewModel->GetEditorOnlyParametersAdapter();
-	EditorParametersAdapter->Modify();
-	EditorParametersAdapter->GetParameters().RemoveAll([&ScriptVarId](const UNiagaraScriptVariable* ScriptVar) { return ScriptVar->Metadata.GetVariableGuid() == ScriptVarId; });
+		// Update anything that was referencing that parameter
+		System.HandleVariableRemoved(ItemToDelete.GetVariable(), true);
+		bAnyChange = true;
+	}
 
-	// Update anything that was referencing that parameter
-	System.HandleVariableRemoved(ItemToDelete.GetVariable(), true);
-	Refresh();
+	if(bAnyChange)
+	{
+		Refresh();
+	}
+	else
+	{
+		RemoveParameterTransaction.Cancel();
+	}
 }
 
 void FNiagaraSystemToolkitParameterPanelViewModel::RenameParameter(const FNiagaraParameterPanelItem& ItemToRename, const FName NewName) const
@@ -1248,13 +1283,13 @@ TSharedPtr<SWidget> FNiagaraSystemToolkitParameterPanelViewModel::CreateContextM
 			MenuBuilder.AddMenuSeparator();
 
 			FText DuplicateToolTip;
-			bool bCanDuplicateParameter = GetCanDuplicateParameterAndToolTip(SelectedItem, DuplicateToolTip);
+			bool bCanDuplicateParameter = GetCanDuplicateParameterAndToolTip(Items, DuplicateToolTip);
 			MenuBuilder.AddMenuEntry(
 				LOCTEXT("DuplicateParameter", "Duplicate"),
 				DuplicateToolTip,
 				FSlateIcon(),
 				FUIAction(
-					FExecuteAction::CreateSP(this, &FNiagaraSystemToolkitParameterPanelViewModel::DuplicateParameter, SelectedItem),
+					FExecuteAction::CreateSP(this, &FNiagaraSystemToolkitParameterPanelViewModel::DuplicateParameters, Items),
 					FCanExecuteAction::CreateLambda([bCanDuplicateParameter](){return bCanDuplicateParameter;})));
 
 			MenuBuilder.AddSubMenu(
@@ -2222,26 +2257,42 @@ bool FNiagaraScriptToolkitParameterPanelViewModel::GetCanAddParametersToCategory
 	return GetEditableGraphsConst().Num() > 0 && Category.NamespaceMetaData.GetGuid() != FNiagaraEditorGuids::StaticSwitchNamespaceMetaDataGuid;
 }
 
-void FNiagaraScriptToolkitParameterPanelViewModel::DeleteParameter(const FNiagaraParameterPanelItem& ItemToDelete) const
+void FNiagaraScriptToolkitParameterPanelViewModel::DeleteParameters(const TArray<FNiagaraParameterPanelItem>& ItemsToDelete) const
 {
-	if (ItemToDelete.bExternallyReferenced)
+	bool bAnyChange = false;
+	FScopedTransaction RemoveParametersWithPins(LOCTEXT("RemoveParametersWithPins", "Removed parameter(s) and referenced pins"));
+	for(const FNiagaraParameterPanelItem& ItemToDelete : ItemsToDelete)
 	{
-		return;
-	}
+		if (ItemToDelete.bExternallyReferenced)
+		{
+			continue;
+		}			
 
-	FScopedTransaction RemoveParametersWithPins(LOCTEXT("RemoveParametersWithPins", "Remove parameter and referenced pins"));
-	for (UNiagaraGraph* Graph : GetEditableGraphs())
-	{
-		UNiagaraScriptVariable* SelectedScriptVariable = Graph->GetScriptVariable(ItemToDelete.GetVariable());
-		if (SelectedScriptVariable && VariableObjectSelection->GetSelectedObjects().Contains(SelectedScriptVariable))
+		for (UNiagaraGraph* Graph : GetEditableGraphs())
 		{
-			VariableObjectSelection->ClearSelectedObjects();
+			UNiagaraScriptVariable* SelectedScriptVariable = Graph->GetScriptVariable(ItemToDelete.GetVariable());
+			if (SelectedScriptVariable && VariableObjectSelection->GetSelectedObjects().Contains(SelectedScriptVariable))
+			{
+				VariableObjectSelection->ClearSelectedObjects();
+			}
+			Graph->RemoveParameter(ItemToDelete.GetVariable());
+			
+			if(!bAnyChange)
+			{
+				bAnyChange = true;
+			}
 		}
-		Graph->RemoveParameter(ItemToDelete.GetVariable());
 	}
 
-	Refresh();
-	UIContext.RefreshSelectionDetailsViewPanel();
+	if(bAnyChange)
+	{
+		Refresh();
+		UIContext.RefreshSelectionDetailsViewPanel();
+	}
+	else
+	{
+		RemoveParametersWithPins.Cancel();
+	}
 }
 
 void FNiagaraScriptToolkitParameterPanelViewModel::RenameParameter(const FNiagaraParameterPanelItem& ItemToRename, const FName NewName) const
@@ -2258,21 +2309,27 @@ void FNiagaraScriptToolkitParameterPanelViewModel::RenameParameter(const FNiagar
 	RenameParameter(ItemToRename.ScriptVariable, NewName);
 }
 
-void FNiagaraScriptToolkitParameterPanelViewModel::ChangeParameterType(const FNiagaraParameterPanelItem ItemToModify, const FNiagaraTypeDefinition NewType) const
+void FNiagaraScriptToolkitParameterPanelViewModel::ChangeParameterType(const TArray<FNiagaraParameterPanelItem> ItemsToModify, const FNiagaraTypeDefinition NewType) const
 {
 	FScopedTransaction Transaction(LOCTEXT("ChangeParameterTypeTransaction", "Change parameter type"));
-	
-	FNiagaraVariable CurrentVar = ItemToModify.GetVariable();
+
+	TArray<FNiagaraVariable> VariablesToModify;
+
+	for (const FNiagaraParameterPanelItem& PanelItem : ItemsToModify)
+	{
+		VariablesToModify.Add(PanelItem.GetVariable());
+	}
 	
 	ScriptViewModel->GetStandaloneScript().Script->Modify();
+	
 	for (UNiagaraGraph* Graph : GetEditableGraphs())
 	{
 		Graph->Modify();
-		Graph->ChangeParameterType(CurrentVar, NewType);
+		Graph->ChangeParameterType(VariablesToModify, NewType, true);
 	}
 
 	Refresh();
-	SelectParameterItemByName(CurrentVar.GetName(), false);
+	SelectParameterItemByName(VariablesToModify[0].GetName(), false);
 }
 
 void FNiagaraScriptToolkitParameterPanelViewModel::RenameParameter(const UNiagaraScriptVariable* ScriptVarToRename, const FName NewName) const
@@ -2317,36 +2374,54 @@ void FNiagaraScriptToolkitParameterPanelViewModel::RenameParameter(const FNiagar
 	}
 }
 
-void FNiagaraScriptToolkitParameterPanelViewModel::DuplicateParameter(const FNiagaraParameterPanelItem ItemToDuplicate) const
+void FNiagaraScriptToolkitParameterPanelViewModel::DuplicateParameters(const TArray<FNiagaraParameterPanelItem> ItemsToDuplicate) const
 {
 	FScopedTransaction Transaction(LOCTEXT("DuplicateParameterTransaction", "Duplicate parameter"));
 	TGuardValue<bool> AddParameterRefreshGuard(bIsAddingParameter, true);
 	bool bSuccess = false;
 
-	TSet<FName> Names;
-	for (const UNiagaraGraph* Graph : GetEditableGraphsConst())
+	FName SingleVariableName = NAME_None;
+	for(const FNiagaraParameterPanelItem& ItemToDuplicate : ItemsToDuplicate)
 	{
-		for (auto It = Graph->GetParameterReferenceMap().CreateConstIterator(); It; ++It)
+		TSet<FName> Names;
+		for (const UNiagaraGraph* Graph : GetEditableGraphsConst())
 		{
-			Names.Add(It.Key().GetName());
+			for (auto It = Graph->GetParameterReferenceMap().CreateConstIterator(); It; ++It)
+			{
+				Names.Add(It.Key().GetName());
+			}
+		}
+		const FName NewUniqueName = FNiagaraUtilities::GetUniqueName(ItemToDuplicate.GetVariable().GetName(), Names);
+
+		if(ItemsToDuplicate.Num() == 1 && SingleVariableName == NAME_None)
+		{
+			SingleVariableName = NewUniqueName;
+		}
+		
+		FNiagaraVariable NewVariable(ItemToDuplicate.GetVariable().GetType(), NewUniqueName);
+		FNiagaraVariableMetaData ParameterMetadata = ItemToDuplicate.ScriptVariable ? ItemToDuplicate.ScriptVariable->Metadata : FNiagaraVariableMetaData();
+		ParameterMetadata.CreateNewGuid();
+		
+		for (UNiagaraGraph* Graph : GetEditableGraphs())
+		{
+			Graph->Modify();
+			Graph->AddParameter(NewVariable, ParameterMetadata, false, false);
+			bSuccess = true;
 		}
-	}
-	const FName NewUniqueName = FNiagaraUtilities::GetUniqueName(ItemToDuplicate.GetVariable().GetName(), Names);
-	FNiagaraVariable NewVariable(ItemToDuplicate.GetVariable().GetType(), NewUniqueName);
-	FNiagaraVariableMetaData ParameterMetadata = ItemToDuplicate.ScriptVariable ? ItemToDuplicate.ScriptVariable->Metadata : FNiagaraVariableMetaData();
-	ParameterMetadata.CreateNewGuid();
-	
-	for (UNiagaraGraph* Graph : GetEditableGraphs())
-	{
-		Graph->Modify();
-		Graph->AddParameter(NewVariable, ParameterMetadata, false, false);
-		bSuccess = true;
 	}
 
 	if (bSuccess)
 	{
 		Refresh();
-		SelectParameterItemByName(NewUniqueName, true);
+
+		if(SingleVariableName != NAME_None)
+		{
+			SelectParameterItemByName(SingleVariableName, true);
+		}
+	}
+	else
+	{
+		Transaction.Cancel();
 	}
 }
 
@@ -2388,119 +2463,123 @@ TSharedPtr<SWidget> FNiagaraScriptToolkitParameterPanelViewModel::CreateContextM
 	// Create a menu with all relevant operations.
 	const bool bShouldCloseWindowAfterMenuSelection = true;
 	FMenuBuilder MenuBuilder(bShouldCloseWindowAfterMenuSelection, ToolkitCommands);
+	
+	const FNiagaraParameterPanelItem& SelectedItem = Items[0];
+	const FNiagaraParameterPanelItemBase& SelectedItemBase = Items[0];
 
-	// Only create context menus when a single item is selected.
-	if (Items.Num() == 1)
+	FUIAction SingleItemAction;
+	bool bSingleItemSelected = Items.Num() == 1;
+	SingleItemAction.CanExecuteAction = FCanExecuteAction::CreateLambda([=]()
 	{
-		const FNiagaraParameterPanelItem& SelectedItem = Items[0];
-		const FNiagaraParameterPanelItemBase& SelectedItemBase = Items[0];
-
-		// helper lambda to add copy/paste metadata actions.
-		auto AddMetaDataContextMenuEntries = [this, &MenuBuilder, &Items, &SelectedItem, &SelectedItemBase]() {
-			FText CopyParameterMetaDataToolTip;
-			const bool bCanCopyParameterMetaData = GetCanCopyParameterMetaDataAndToolTip(SelectedItem, CopyParameterMetaDataToolTip);
-			MenuBuilder.AddMenuEntry(
-				LOCTEXT("CopyParameterMetadata", "Copy Metadata"),
-				CopyParameterMetaDataToolTip,
-				FSlateIcon(),
-				FUIAction(
-					FExecuteAction::CreateSP(this, &INiagaraImmutableParameterPanelViewModel::CopyParameterMetaData, SelectedItemBase),
-					FCanExecuteAction::CreateLambda([bCanCopyParameterMetaData]() { return bCanCopyParameterMetaData; })));
+		return bSingleItemSelected;
+	});
+	
+	// helper lambda to add copy/paste metadata actions.
+	auto AddMetaDataContextMenuEntries = [this, &MenuBuilder, &Items, &SelectedItem, &SelectedItemBase, SingleItemAction]() {
+		FText CopyParameterMetaDataToolTip;
+		const bool bCanCopyParameterMetaData = GetCanCopyParameterMetaDataAndToolTip(SelectedItem, CopyParameterMetaDataToolTip) && SingleItemAction.CanExecute();
+		MenuBuilder.AddMenuEntry(
+			LOCTEXT("CopyParameterMetadata", "Copy Metadata"),
+			CopyParameterMetaDataToolTip,
+			FSlateIcon(),
+			FUIAction(
+				FExecuteAction::CreateSP(this, &INiagaraImmutableParameterPanelViewModel::CopyParameterMetaData, SelectedItemBase),
+				FCanExecuteAction::CreateLambda([bCanCopyParameterMetaData]() { return bCanCopyParameterMetaData; })));
 
-			FText PasteParameterMetaDataToolTip;
-			const bool bCanPasteParameterMetaData = GetCanPasteParameterMetaDataAndToolTip(PasteParameterMetaDataToolTip);
-			MenuBuilder.AddMenuEntry(
-				LOCTEXT("PasteParameterMetadata", "Paste Metadata"),
-				PasteParameterMetaDataToolTip,
-				FSlateIcon(),
-				FUIAction(
-					FExecuteAction::CreateSP(this, &INiagaraParameterPanelViewModel::PasteParameterMetaData, Items),
-					FCanExecuteAction::CreateLambda([bCanPasteParameterMetaData]() { return bCanPasteParameterMetaData; })));
-		};
+		FText PasteParameterMetaDataToolTip;
+		const bool bCanPasteParameterMetaData = GetCanPasteParameterMetaDataAndToolTip(PasteParameterMetaDataToolTip) && SingleItemAction.CanExecute();
+		MenuBuilder.AddMenuEntry(
+			LOCTEXT("PasteParameterMetadata", "Paste Metadata"),
+			PasteParameterMetaDataToolTip,
+			FSlateIcon(),
+			FUIAction(
+				FExecuteAction::CreateSP(this, &INiagaraParameterPanelViewModel::PasteParameterMetaData, Items),
+				FCanExecuteAction::CreateLambda([bCanPasteParameterMetaData]() { return bCanPasteParameterMetaData; })));
+	};
 
-		MenuBuilder.BeginSection("Edit", LOCTEXT("EditMenuHeader", "Edit"));
+	MenuBuilder.BeginSection("Edit", LOCTEXT("EditMenuHeader", "Edit"));
+	{
+		if (SelectedItem.ScriptVariable->GetIsStaticSwitch())
 		{
-			if (SelectedItem.ScriptVariable->GetIsStaticSwitch())
-			{
-				// Only allow modifying metadata for static switches.
-				AddMetaDataContextMenuEntries();
-			}
-			else
-			{ 
-				FText CopyReferenceToolTip;
-				GetCanCopyParameterReferenceAndToolTip(SelectedItem, CopyReferenceToolTip);
-				MenuBuilder.AddMenuEntry(FGenericCommands::Get().Copy, NAME_None, LOCTEXT("CopyReference", "Copy Reference"), CopyReferenceToolTip);
+			// Only allow modifying metadata for static switches.
+			AddMetaDataContextMenuEntries();
+		}
+		else
+		{ 
+			FText CopyReferenceToolTip;
+			GetCanCopyParameterReferenceAndToolTip(SelectedItem, CopyReferenceToolTip);
+			MenuBuilder.AddMenuEntry(FGenericCommands::Get().Copy, NAME_None, LOCTEXT("CopyReference", "Copy Reference"), CopyReferenceToolTip);
 
-				FText DeleteToolTip;
-				GetCanDeleteParameterAndToolTip(SelectedItem, DeleteToolTip);
-				MenuBuilder.AddMenuEntry(FNiagaraParameterPanelCommands::Get().DeleteItem, NAME_None, TAttribute<FText>(), DeleteToolTip);
+			FText DeleteToolTip;
+			GetCanDeleteParameterAndToolTip(SelectedItem, DeleteToolTip);
+			MenuBuilder.AddMenuEntry(FNiagaraParameterPanelCommands::Get().DeleteItem, NAME_None, TAttribute<FText>(), DeleteToolTip);
 
-				FText RenameToolTip;
-				GetCanRenameParameterAndToolTip(SelectedItem, FText(), false, RenameToolTip);
-				MenuBuilder.AddMenuEntry(FGenericCommands::Get().Rename, NAME_None, LOCTEXT("Rename", "Rename"), RenameToolTip);
+			FText RenameToolTip;
+			GetCanRenameParameterAndToolTip(SelectedItem, FText(), false, RenameToolTip);
+			MenuBuilder.AddMenuEntry(FGenericCommands::Get().Rename, NAME_None, LOCTEXT("Rename", "Rename"), RenameToolTip);
 
 
-				MenuBuilder.AddMenuSeparator();
+			MenuBuilder.AddMenuSeparator();
 
-				AddMetaDataContextMenuEntries();
+			AddMetaDataContextMenuEntries();
 
 
-				MenuBuilder.AddMenuSeparator();
-				
-				FText CanChangeTooltip;
-				bool bCanChangeType = GetCanChangeParameterType(SelectedItem, CanChangeTooltip);
-				FUIAction CanChangeTypeAction;
-				CanChangeTypeAction.CanExecuteAction = FCanExecuteAction::CreateLambda([bCanChangeType]()
-				{
-					return bCanChangeType;
-				});
-				
-				MenuBuilder.AddSubMenu(
-					LOCTEXT("ChangeType", "Change Type"),
-					CanChangeTooltip,
-					FNewMenuDelegate::CreateSP(this, &FNiagaraScriptToolkitParameterPanelViewModel::GetChangeTypeSubMenu, SelectedItem),
-					CanChangeTypeAction, NAME_None, EUserInterfaceActionType::Button);
-
-				MenuBuilder.AddSubMenu(
-					LOCTEXT("ChangeNamespace", "Change Namespace"),
-					LOCTEXT("ChangeNamespaceToolTip", "Select a new namespace for the selected parameter."),
-					FNewMenuDelegate::CreateSP(this, &FNiagaraScriptToolkitParameterPanelViewModel::GetChangeNamespaceSubMenu, false, SelectedItem));
-
-				MenuBuilder.AddSubMenu(
-					LOCTEXT("ChangeNamespaceModifier", "Change Namespace Modifier"),
-					LOCTEXT("ChangeNamespaceModifierToolTip", "Edit the namespace modifier for the selected parameter."),
-					FNewMenuDelegate::CreateSP(this, &FNiagaraScriptToolkitParameterPanelViewModel::GetChangeNamespaceModifierSubMenu, false, SelectedItem));
-
-
-				MenuBuilder.AddMenuSeparator();
-
-				FText DuplicateToolTip;
-				bool bCanDuplicateParameter = GetCanDuplicateParameterAndToolTip(SelectedItem, DuplicateToolTip);
-				MenuBuilder.AddMenuEntry(
-					LOCTEXT("DuplicateParameter", "Duplicate"),
-					DuplicateToolTip,
-					FSlateIcon(),
-					FUIAction(
-						FExecuteAction::CreateSP(this, &FNiagaraScriptToolkitParameterPanelViewModel::DuplicateParameter, SelectedItem),
-						FCanExecuteAction::CreateLambda([bCanDuplicateParameter]() {return bCanDuplicateParameter; })));
-
-				MenuBuilder.AddSubMenu(
-					LOCTEXT("DuplicateToNewNamespace", "Duplicate to Namespace"),
-					LOCTEXT("DuplicateToNewNamespaceToolTip", "Duplicate this parameter to a new namespace."),
-					FNewMenuDelegate::CreateSP(this, &FNiagaraScriptToolkitParameterPanelViewModel::GetChangeNamespaceSubMenu, true, SelectedItem));
-
-				MenuBuilder.AddSubMenu(
-					LOCTEXT("DuplicateWithNewNamespaceModifier", "Duplicate with Namespace Modifier"),
-					LOCTEXT("DupilcateWithNewNamespaceModifierToolTip", "Duplicate this parameter with a different namespace modifier."),
-					FNewMenuDelegate::CreateSP(this, &FNiagaraScriptToolkitParameterPanelViewModel::GetChangeNamespaceModifierSubMenu, true, SelectedItem));
+			MenuBuilder.AddMenuSeparator();
+			
+			FText CanChangeTooltip;
+			bool bCanChangeType = GetCanChangeParameterType(Items, CanChangeTooltip);
+			FUIAction CanChangeTypeAction;
+			CanChangeTypeAction.CanExecuteAction = FCanExecuteAction::CreateLambda([bCanChangeType]()
+			{
+				return bCanChangeType;
+			});
+			
+			MenuBuilder.AddSubMenu(
+				LOCTEXT("ChangeType", "Change Type"),
+				CanChangeTooltip,
+				FNewMenuDelegate::CreateSP(this, &FNiagaraScriptToolkitParameterPanelViewModel::GetChangeTypeSubMenu, Items),
+				CanChangeTypeAction, NAME_None, EUserInterfaceActionType::Button);
 
-			}
+			MenuBuilder.AddSubMenu(
+				LOCTEXT("ChangeNamespace", "Change Namespace"),
+				LOCTEXT("ChangeNamespaceToolTip", "Select a new namespace for the selected parameter."),
+				FNewMenuDelegate::CreateSP(this, &FNiagaraScriptToolkitParameterPanelViewModel::GetChangeNamespaceSubMenu, false, SelectedItem),
+				SingleItemAction, NAME_None, EUserInterfaceActionType::None);
+
+			MenuBuilder.AddSubMenu(
+				LOCTEXT("ChangeNamespaceModifier", "Change Namespace Modifier"),
+				LOCTEXT("ChangeNamespaceModifierToolTip", "Edit the namespace modifier for the selected parameter."),
+				FNewMenuDelegate::CreateSP(this, &FNiagaraScriptToolkitParameterPanelViewModel::GetChangeNamespaceModifierSubMenu, false, SelectedItem),
+				SingleItemAction, NAME_None, EUserInterfaceActionType::None);
+
+			MenuBuilder.AddMenuSeparator();
+
+			FText DuplicateToolTip;
+			bool bCanDuplicateParameter = GetCanDuplicateParameterAndToolTip(Items, DuplicateToolTip);
+			MenuBuilder.AddMenuEntry(
+				LOCTEXT("DuplicateParameter", "Duplicate"),
+				DuplicateToolTip,
+				FSlateIcon(),
+				FUIAction(
+					FExecuteAction::CreateSP(this, &FNiagaraScriptToolkitParameterPanelViewModel::DuplicateParameters, Items),
+					FCanExecuteAction::CreateLambda([bCanDuplicateParameter]() {return bCanDuplicateParameter; })));
+
+			MenuBuilder.AddSubMenu(
+				LOCTEXT("DuplicateToNewNamespace", "Duplicate to Namespace"),
+				LOCTEXT("DuplicateToNewNamespaceToolTip", "Duplicate this parameter to a new namespace."),
+				FNewMenuDelegate::CreateSP(this, &FNiagaraScriptToolkitParameterPanelViewModel::GetChangeNamespaceSubMenu, true, SelectedItem),
+				SingleItemAction, NAME_None, EUserInterfaceActionType::None);
+
+			MenuBuilder.AddSubMenu(
+				LOCTEXT("DuplicateWithNewNamespaceModifier", "Duplicate with Namespace Modifier"),
+				LOCTEXT("DupilcateWithNewNamespaceModifierToolTip", "Duplicate this parameter with a different namespace modifier."),
+				FNewMenuDelegate::CreateSP(this, &FNiagaraScriptToolkitParameterPanelViewModel::GetChangeNamespaceModifierSubMenu, true, SelectedItem),
+				SingleItemAction, NAME_None, EUserInterfaceActionType::None);
 		}
-		MenuBuilder.EndSection();
-		return MenuBuilder.MakeWidget();
 	}
-	// More than one item selected, do not return a context menu.
-	return SNullWidget::NullWidget;
+	
+	MenuBuilder.EndSection();
+	return MenuBuilder.MakeWidget();
 }
 
 FNiagaraParameterUtilities::EParameterContext FNiagaraScriptToolkitParameterPanelViewModel::GetParameterContext() const
@@ -2905,10 +2984,14 @@ bool FNiagaraParameterDefinitionsToolkitParameterPanelViewModel::GetCanAddParame
 	return true;
 }
 
-void FNiagaraParameterDefinitionsToolkitParameterPanelViewModel::DeleteParameter(const FNiagaraParameterPanelItem& ItemToDelete) const
+void FNiagaraParameterDefinitionsToolkitParameterPanelViewModel::DeleteParameters(const TArray<FNiagaraParameterPanelItem>& ItemsToDelete) const
 {
-	FScopedTransaction RemoveParameter(LOCTEXT("ParameterDefinitionsToolkitParameterPanelViewModel_RemoveParameter", "Remove Parameter"));
-	ParameterDefinitionsWeak.Get()->RemoveParameter(ItemToDelete.GetVariable());
+	FScopedTransaction RemoveParameter(LOCTEXT("ParameterDefinitionsToolkitParameterPanelViewModel_RemoveParameter", "Removed Parameter(s)"));
+
+	for(const FNiagaraParameterPanelItem& ItemToDelete : ItemsToDelete)
+	{
+		ParameterDefinitionsWeak.Get()->RemoveParameter(ItemToDelete.GetVariable());
+	}
 	VariableObjectSelection->ClearSelectedObjects();
 	UIContext.RefreshSelectionDetailsViewPanel();
 }
@@ -2966,16 +3049,16 @@ TSharedPtr<SWidget> FNiagaraParameterDefinitionsToolkitParameterPanelViewModel::
 			MenuBuilder.AddMenuSeparator();
 
 			FText DuplicateToolTip;
-			bool bCanDuplicateParameter = GetCanDuplicateParameterAndToolTip(SelectedItem, DuplicateToolTip);
+			bool bCanDuplicateParameter = GetCanDuplicateParameterAndToolTip(Items, DuplicateToolTip);
 			MenuBuilder.AddMenuEntry(
 				LOCTEXT("DuplicateParameter", "Duplicate"),
 				DuplicateToolTip,
 				FSlateIcon(),
 				FUIAction(
-					FExecuteAction::CreateSP(this, &FNiagaraParameterDefinitionsToolkitParameterPanelViewModel::DuplicateParameter, SelectedItem),
+					FExecuteAction::CreateSP(this, &FNiagaraParameterDefinitionsToolkitParameterPanelViewModel::DuplicateParameters, Items),
 					FCanExecuteAction::CreateLambda([bCanDuplicateParameter]() {return bCanDuplicateParameter; })));
-
-			MenuBuilder.AddSubMenu(
+			
+			MenuBuilder.AddSubMenu( 
 				LOCTEXT("DuplicateToNewNamespace", "Duplicate to Namespace"),
 				LOCTEXT("DuplicateToNewNamespaceToolTip", "Duplicate this parameter to a new namespace."),
 				FNewMenuDelegate::CreateSP(this, &FNiagaraParameterDefinitionsToolkitParameterPanelViewModel::GetChangeNamespaceSubMenu, true, SelectedItem));
diff --git a/Engine/Plugins/FX/Niagara/Source/NiagaraEditor/Private/Widgets/SNiagaraGraphParameterMapGetNode.cpp b/Engine/Plugins/FX/Niagara/Source/NiagaraEditor/Private/Widgets/SNiagaraGraphParameterMapGetNode.cpp
--- a/Engine/Plugins/FX/Niagara/Source/NiagaraEditor/Private/Widgets/SNiagaraGraphParameterMapGetNode.cpp
+++ b/Engine/Plugins/FX/Niagara/Source/NiagaraEditor/Private/Widgets/SNiagaraGraphParameterMapGetNode.cpp
@@ -79,85 +79,142 @@ void SNiagaraGraphParameterMapGetNode::CreatePinWidgets()
 	SGraphNode::CreatePinWidgets();
 	
 	UNiagaraNodeParameterMapGet* GetNode = Cast<UNiagaraNodeParameterMapGet>(GraphNode);
+	
+	TSet<TSharedRef<SGraphPin>> AddedPins;
 
+	auto AddRowWidgetToNode = [this](TSharedPtr<SWidget> RowWidget)
+	{
+		TSharedRef<SBorder> Border = SNew(SBorder)
+			.HAlign(HAlign_Fill)
+			.VAlign(VAlign_Fill)
+			.Padding(FMargin(0, 3))
+			//.OnMouseButtonDown(this, &SNiagaraGraphParameterMapGetNode::OnBorderMouseButtonDown, i)
+			[
+				RowWidget.ToSharedRef()
+			];
+		
+		Border->SetBorderImage(TAttribute<const FSlateBrush*>::Create(TAttribute<const FSlateBrush*>::FGetter::CreateRaw(this, &SNiagaraGraphParameterMapGetNode::GetBackgroundBrush, RowWidget)));
 
+		PinContainerRoot->AddSlot()
+		.AutoHeight()
+		[
+			Border
+		];
+	};
+	
+	// add the source pin on the first row
+	TSharedPtr<SGraphPin> SourcePin = nullptr;
+	if(InputPins.Num() > 0)
+	{
+		SourcePin = InputPins[0];
+	}
+		
+	TSharedPtr<SWidget> SourceRowWidget = SNew(SHorizontalBox)
+		.Visibility(EVisibility::Visible)
+		+ SHorizontalBox::Slot()
+		.HAlign(HAlign_Left)
+		.FillWidth(1.0f)
+		.Padding(Settings->GetInputPinPadding())
+		[
+			(SourcePin.IsValid() ? SourcePin.ToSharedRef() : SNullWidget::NullWidget)
+		]
+	+ SHorizontalBox::Slot()
+		.AutoWidth()
+		.HAlign(HAlign_Right)
+		.Padding(Settings->GetOutputPinPadding())
+		[
+			SNullWidget::NullWidget
+		];
+
+	if(SourcePin.IsValid())
+	{
+		AddedPins.Add(InputPins[0]);
+	}
+
+	AddRowWidgetToNode(SourceRowWidget);
+	
 	// Deferred pin adding to line up input/output pins by name.
-	for (int32 i = 0; i < OutputPins.Num() + 1; i++)
+	for (int32 i = 0; i < OutputPins.Num(); i++)
 	{
-		// Get nodes have an unequal number of pins. 
-		TSharedPtr<SWidget> Widget;
-		if (i == 0)
+		TSharedRef<SGraphPin> OutputPin = OutputPins[i];
+		UEdGraphPin* SrcOutputPin = OutputPin->GetPinObj(); 
+
+		UEdGraphPin* MatchingInputPin = GetNode->GetDefaultPin(SrcOutputPin);
+
+		TSharedPtr<SGraphPin> InputPin = nullptr;
+		for (TSharedRef<SGraphPin> Pin : InputPins)
 		{
+			UEdGraphPin* SrcInputPin = Pin->GetPinObj();
+			if (SrcInputPin == MatchingInputPin)
+			{
+				InputPin = Pin;
+				Pin->SetShowLabel(false);
+			}
+		}
 
-			SAssignNew(Widget, SHorizontalBox)
-				.Visibility(EVisibility::Visible)
-				+ SHorizontalBox::Slot()
-				.HAlign(HAlign_Left)
-				.FillWidth(1.0f)
-				.Padding(Settings->GetInputPinPadding())
-				[
-					(InputPins.Num() > 0 ? InputPins[0] : SNullWidget::NullWidget)
-				]
+		TSharedPtr<SWidget> ParameterRow = SNew(SHorizontalBox)
+			.Visibility(EVisibility::Visible)
 			+ SHorizontalBox::Slot()
-				.AutoWidth()
-				.HAlign(HAlign_Right)
-				.Padding(Settings->GetOutputPinPadding())
-				[
-					SNullWidget::NullWidget
-				];
-		}
-		else
+			.HAlign(HAlign_Left)
+			.FillWidth(1.0f)
+			.Padding(Settings->GetInputPinPadding())
+			[
+				(InputPin.IsValid() ? InputPin.ToSharedRef() : SNullWidget::NullWidget)
+			]
+			+ SHorizontalBox::Slot()
+			.AutoWidth()
+			.HAlign(HAlign_Right)
+			.Padding(Settings->GetOutputPinPadding())
+			[
+				OutputPin
+			];
+
+		if(InputPin.IsValid())
 		{
-			TSharedRef<SGraphPin> OutputPin = OutputPins[i - 1];
-			UEdGraphPin* SrcOutputPin = OutputPin->GetPinObj();
+			AddedPins.Add(InputPin.ToSharedRef());
+		}
+		AddedPins.Add(OutputPin);
 
-			UEdGraphPin* MatchingInputPin = GetNode->GetDefaultPin(SrcOutputPin);
+		AddRowWidgetToNode(ParameterRow);
+	}
 
-			TSharedPtr<SWidget> InputPin = SNullWidget::NullWidget;
-			for (TSharedRef<SGraphPin> Pin : InputPins)
+	// for each unadded input pin (for example orphaned pins), we add them back at the bottom
+	for (int32 i = 0; i < InputPins.Num(); i++)
+	{
+		if(!AddedPins.Contains(InputPins[i]))
+		{
+			TSharedRef<SGraphPin> UnaddedInputPin = InputPins[i];
+			// since we don't want to waste space, we look for some output pin we also haven't added yet
+			TSharedRef<SGraphPin>* SomeUnaddedOutputPin = OutputPins.FindByPredicate([&](TSharedRef<SGraphPin> Pin)
 			{
-				UEdGraphPin* SrcInputPin = Pin->GetPinObj();
-				if (SrcInputPin == MatchingInputPin)
-				{
-					InputPin = Pin;
-					Pin->SetShowLabel(false);
-				}
-			}
+				return !AddedPins.Contains(Pin);
+			});
 
-			SAssignNew(Widget, SHorizontalBox)
+			TSharedPtr<SWidget> LeftoverPinContainer = SNew(SHorizontalBox)
 				.Visibility(EVisibility::Visible)
 				+ SHorizontalBox::Slot()
 				.HAlign(HAlign_Left)
 				.FillWidth(1.0f)
 				.Padding(Settings->GetInputPinPadding())
 				[
-					InputPin.ToSharedRef()
+					UnaddedInputPin
 				]
 				+ SHorizontalBox::Slot()
 				.AutoWidth()
 				.HAlign(HAlign_Right)
 				.Padding(Settings->GetOutputPinPadding())
 				[
-					OutputPin
+					(SomeUnaddedOutputPin != nullptr ? (*SomeUnaddedOutputPin) : SNullWidget::NullWidget)
 				];
-		}
 
-		TSharedRef<SBorder> Border = SNew(SBorder)
-			.HAlign(HAlign_Fill)
-			.VAlign(VAlign_Fill)
-			.Padding(FMargin(0, 3))
-			//.OnMouseButtonDown(this, &SNiagaraGraphParameterMapGetNode::OnBorderMouseButtonDown, i)
-			[
-				Widget.ToSharedRef()
-			];
-		Border->SetBorderImage(TAttribute<const FSlateBrush*>::Create(TAttribute<const FSlateBrush*>::FGetter::CreateRaw(this, &SNiagaraGraphParameterMapGetNode::GetBackgroundBrush, Widget)));
-
-
-		PinContainerRoot->AddSlot()
-			.AutoHeight()
-			[
-				Border
-			];
+			AddedPins.Add(UnaddedInputPin);
+			if(SomeUnaddedOutputPin != nullptr)
+			{
+				AddedPins.Add(*SomeUnaddedOutputPin);
+			}
+			
+			AddRowWidgetToNode(LeftoverPinContainer);
+		}
 	}
 }
 
diff --git a/Engine/Plugins/FX/Niagara/Source/NiagaraEditor/Private/Widgets/SNiagaraParameterPanel.cpp b/Engine/Plugins/FX/Niagara/Source/NiagaraEditor/Private/Widgets/SNiagaraParameterPanel.cpp
--- a/Engine/Plugins/FX/Niagara/Source/NiagaraEditor/Private/Widgets/SNiagaraParameterPanel.cpp
+++ b/Engine/Plugins/FX/Niagara/Source/NiagaraEditor/Private/Widgets/SNiagaraParameterPanel.cpp
@@ -85,7 +85,7 @@ void SNiagaraParameterPanel::Construct(const FArguments& InArgs, const TSharedPt
 	.OnItemSelected(this, &SNiagaraParameterPanel::OnParameterItemSelected)
 	.OnItemsDragged(this, &SNiagaraParameterPanel::OnParameterItemsDragged)
 	.OnItemActivated(this, &SNiagaraParameterPanel::OnParameterItemActived)
-	.AllowMultiselect(false)
+	.AllowMultiselect(true)
 	.ClearSelectionOnClick(true)
 	.CategoryRowStyle(FNiagaraEditorStyle::Get(), "NiagaraEditor.Parameters.TableRow")
 	.OnGetCategoryBackgroundImage(this, &SNiagaraParameterPanel::GetCategoryBackgroundImage)
@@ -508,10 +508,7 @@ bool SNiagaraParameterPanel::GetCanAddParametersToCategory(FNiagaraParameterPane
 
 void SNiagaraParameterPanel::DeleteSelectedItems() const
 {
-	for (const FNiagaraParameterPanelItem& Item : ItemSelector->GetSelectedItems())
-	{
-		ParameterPanelViewModel->DeleteParameter(Item);
-	}
+	ParameterPanelViewModel->DeleteParameters(ItemSelector->GetSelectedItems());
 }
 
 bool SNiagaraParameterPanel::CanDeleteSelectedItems() const
diff --git a/Engine/Plugins/FX/Niagara/Source/NiagaraEditor/Public/NiagaraGraph.h b/Engine/Plugins/FX/Niagara/Source/NiagaraEditor/Public/NiagaraGraph.h
--- a/Engine/Plugins/FX/Niagara/Source/NiagaraEditor/Public/NiagaraGraph.h
+++ b/Engine/Plugins/FX/Niagara/Source/NiagaraEditor/Public/NiagaraGraph.h
@@ -309,8 +309,12 @@ class UNiagaraGraph : public UEdGraph
 	/** Rename a pin inline in a graph. If this is the only instance used in the graph, then rename them all, otherwise make a duplicate. */
 	bool RenameParameterFromPin(const FNiagaraVariable& Parameter, FName NewName, UEdGraphPin* InPin);
 
-	/** Changes the type of an existing graph variable. This is an unsafe operation, usually guarded by a cvar and not directly exposed to users. */
-	void ChangeParameterType(const FNiagaraVariable& CurrentParameter, const FNiagaraTypeDefinition& NewType);
+	/** Changes the type of existing graph parameters.
+	 *  Optionally creates orphaned pins for any connection that won't be kept, but tries to keep connections as long as types are matching.
+	 *  Changing multiple parameters at once helps with maintaining connections.
+	 *  CAUTION: Do not allow orphaned pins in the stack graphs, as they aren't user facing.
+	 */
+	void ChangeParameterType(const TArray<FNiagaraVariable>& ParametersToChange, const FNiagaraTypeDefinition& NewType, bool bAllowOrphanedPins = false);
 
 	/** Gets a delegate which is called whenever a contained data interfaces changes. */
 	FOnDataInterfaceChanged& OnDataInterfaceChanged();
diff --git a/Engine/Plugins/FX/Niagara/Source/NiagaraEditor/Public/ViewModels/NiagaraParameterPanelViewModel.h b/Engine/Plugins/FX/Niagara/Source/NiagaraEditor/Public/ViewModels/NiagaraParameterPanelViewModel.h
--- a/Engine/Plugins/FX/Niagara/Source/NiagaraEditor/Public/ViewModels/NiagaraParameterPanelViewModel.h
+++ b/Engine/Plugins/FX/Niagara/Source/NiagaraEditor/Public/ViewModels/NiagaraParameterPanelViewModel.h
@@ -158,7 +158,7 @@ public:
 
 	virtual bool GetCanAddParametersToCategory(FNiagaraParameterPanelCategory Category) const = 0;
 
-	virtual void DeleteParameter(const FNiagaraParameterPanelItem& ItemToDelete) const = 0;
+	virtual void DeleteParameters(const TArray<FNiagaraParameterPanelItem>& ItemsToDelete) const = 0;
 
 	virtual void RenameParameter(const FNiagaraParameterPanelItem& ItemToRename, const FName NewName) const = 0;
 
@@ -189,15 +189,15 @@ public:
 
 	virtual bool GetCanDeleteParameterAndToolTip(const FNiagaraParameterPanelItem& ItemToDelete, FText& OutCanDeleteParameterToolTip) const;
 
-	virtual void ChangeParameterType(const FNiagaraParameterPanelItem ItemToModify, const FNiagaraTypeDefinition NewType) const;
-	virtual bool GetCanChangeParameterType(const FNiagaraParameterPanelItem ItemToChange, FText& OutTooltip) const;
-	virtual void GetChangeTypeSubMenu(FMenuBuilder& MenuBuilder, FNiagaraParameterPanelItem Item) const;
+	virtual void ChangeParameterType(const TArray<FNiagaraParameterPanelItem> ItemsToModify, const FNiagaraTypeDefinition NewType) const;
+	virtual bool GetCanChangeParameterType(const TArray<FNiagaraParameterPanelItem>& ItemToChange, FText& OutTooltip) const;
+	virtual void GetChangeTypeSubMenu(FMenuBuilder& MenuBuilder, TArray<FNiagaraParameterPanelItem> Item) const;
 
 	virtual void PasteParameterMetaData(const TArray<FNiagaraParameterPanelItem> SelectedItems);
 	virtual bool GetCanPasteParameterMetaDataAndToolTip(FText& OutCanPasteToolTip);
 
-	virtual void DuplicateParameter(const FNiagaraParameterPanelItem ItemToDuplicate) const;
-	virtual bool GetCanDuplicateParameterAndToolTip(const FNiagaraParameterPanelItem& ItemToDuplicate, FText& OutCanDuplicateParameterToolTip) const;
+	virtual void DuplicateParameters(const TArray<FNiagaraParameterPanelItem> ItemToDuplicate) const;
+	virtual bool GetCanDuplicateParameterAndToolTip(const TArray<FNiagaraParameterPanelItem>& ItemsToDuplicate, FText& OutCanDuplicateParameterToolTip) const;
 
 	virtual bool GetCanRenameParameterAndToolTip(const FNiagaraParameterPanelItem& ItemToRename, const FText& NewVariableNameText, bool bCheckEmptyNameText, FText& OutCanRenameParameterToolTip) const;
 
@@ -287,7 +287,7 @@ public:
 
 	virtual bool GetCanAddParametersToCategory(FNiagaraParameterPanelCategory Category) const override;
 
-	virtual void DeleteParameter(const FNiagaraParameterPanelItem& ItemToDelete) const override;
+	virtual void DeleteParameters(const TArray<FNiagaraParameterPanelItem>& ItemsToDelete) const override;
 
 	virtual void RenameParameter(const FNiagaraParameterPanelItem& ItemToRename, const FName NewName) const override;
 
@@ -387,13 +387,13 @@ public:
 
 	virtual bool GetCanAddParametersToCategory(FNiagaraParameterPanelCategory Category) const override;
 
-	virtual void DeleteParameter(const FNiagaraParameterPanelItem& ItemToDelete) const override;
+	virtual void DeleteParameters(const TArray<FNiagaraParameterPanelItem>& ItemToDelete) const override;
 
 	virtual void RenameParameter(const FNiagaraParameterPanelItem& ItemToRename, const FName NewName) const override;
 	
-	virtual void ChangeParameterType(const FNiagaraParameterPanelItem ItemToModify, const FNiagaraTypeDefinition NewType) const override;
+	virtual void ChangeParameterType(const TArray<FNiagaraParameterPanelItem> ItemsToModify, const FNiagaraTypeDefinition NewType) const override;
 
-	virtual void DuplicateParameter(const FNiagaraParameterPanelItem ItemToDuplicate) const override;
+	virtual void DuplicateParameters(const TArray<FNiagaraParameterPanelItem> ItemsToDuplicate) const override;
 
 	virtual void SetParameterIsSubscribedToLibrary(const UNiagaraScriptVariable* ScriptVarToModify, const bool bSubscribed) const override;
 
@@ -487,7 +487,7 @@ public:
 
 	virtual bool GetCanAddParametersToCategory(FNiagaraParameterPanelCategory Category) const override;
 
-	virtual void DeleteParameter(const FNiagaraParameterPanelItem& ItemToDelete) const override;
+	virtual void DeleteParameters(const TArray<FNiagaraParameterPanelItem>& ItemToDelete) const override;
 
 	virtual void RenameParameter(const FNiagaraParameterPanelItem& ItemToRename, const FName NewName) const override;
 
