Fix for ribbon DistanceOnSegment

diff --git a/Engine/Plugins/FX/Niagara/Shaders/Private/NiagaraRibbonVertexFactory.ush b/Engine/Plugins/FX/Niagara/Shaders/Private/NiagaraRibbonVertexFactory.ush
--- a/Engine/Plugins/FX/Niagara/Shaders/Private/NiagaraRibbonVertexFactory.ush
+++ b/Engine/Plugins/FX/Niagara/Shaders/Private/NiagaraRibbonVertexFactory.ush
@@ -158,7 +158,6 @@ struct FVertexFactoryIntermediates
 {
 	FRibbonPositionIntermediates Current;
 	FRibbonPositionIntermediates Previous;
-	float DistanceOnSegment;
 	float4 Color;
 	
 	float4 TexCoord;
@@ -724,7 +723,7 @@ FVertexFactoryIntermediates GetVertexFactoryIntermediates(FVertexFactoryInput In
 	}
 	else if (NiagaraRibbonVF.U0DistributionMode == U_DISTRIBUTION_SCALED_USING_RIBBON_LENGTH || NiagaraRibbonVF.U0DistributionMode == U_DISTRIBUTION_TILED_OVER_RIBBON_LENGTH)
 	{
-		U0ForSegment = Intermediates.DistanceOnSegment * SegmentData.U0DistributionScaler;
+		U0ForSegment = Intermediates.Current.DistanceOnSegment * SegmentData.U0DistributionScaler;
 	}
 	else if (NiagaraRibbonVF.DistanceFromStartOffset != -1 && NiagaraRibbonVF.U0DistributionMode == U_DISTRIBUTION_TILED_FROM_START_OVER_RIBBON_LENGTH)
 	{
@@ -743,7 +742,7 @@ FVertexFactoryIntermediates GetVertexFactoryIntermediates(FVertexFactoryInput In
 	}
 	else if (NiagaraRibbonVF.U1DistributionMode == U_DISTRIBUTION_SCALED_USING_RIBBON_LENGTH || NiagaraRibbonVF.U1DistributionMode == U_DISTRIBUTION_TILED_OVER_RIBBON_LENGTH)
 	{
-		U1ForSegment = Intermediates.DistanceOnSegment * SegmentData.U1DistributionScaler;
+		U1ForSegment = Intermediates.Current.DistanceOnSegment * SegmentData.U1DistributionScaler;
 	}
 	else if (NiagaraRibbonVF.DistanceFromStartOffset != -1 && NiagaraRibbonVF.U1DistributionMode == U_DISTRIBUTION_TILED_FROM_START_OVER_RIBBON_LENGTH)
 	{
