- Remove stats from renderer that are either not used or not useful

#rnx

diff --git a/Engine/Plugins/FX/Niagara/Source/Niagara/Private/NiagaraRendererMeshes.cpp b/Engine/Plugins/FX/Niagara/Source/Niagara/Private/NiagaraRendererMeshes.cpp
--- a/Engine/Plugins/FX/Niagara/Source/Niagara/Private/NiagaraRendererMeshes.cpp
+++ b/Engine/Plugins/FX/Niagara/Source/Niagara/Private/NiagaraRendererMeshes.cpp
@@ -22,7 +22,6 @@
 
 DECLARE_CYCLE_STAT(TEXT("Generate Mesh Vertex Data [GT]"), STAT_NiagaraGenMeshVertexData, STATGROUP_Niagara);
 DECLARE_CYCLE_STAT(TEXT("Render Meshes [RT]"), STAT_NiagaraRenderMeshes, STATGROUP_Niagara);
-DECLARE_CYCLE_STAT(TEXT("Render Meshes - Allocate GPU Data [RT]"), STAT_NiagaraRenderMeshes_AllocateGPUData, STATGROUP_Niagara);
 
 DECLARE_DWORD_COUNTER_STAT(TEXT("NumMeshesRenderer"), STAT_NiagaraNumMeshes, STATGROUP_Niagara);
 DECLARE_DWORD_COUNTER_STAT(TEXT("NumMeshVerts"), STAT_NiagaraNumMeshVerts, STATGROUP_Niagara);
@@ -380,8 +379,6 @@ void FNiagaraRendererMeshes::PrepareParticleRenderBuffers(FParticleMeshRenderDat
 	{
 		if ( SimTarget == ENiagaraSimTarget::CPUSim )
 		{
-			SCOPE_CYCLE_COUNTER(STAT_NiagaraRenderMeshes_AllocateGPUData);
-
 			// Determine what integer attributes we need to copy to the GPU for culling
 			TArray<uint32, TInlineAllocator<2>> IntParamsToCopy;
 			if (ParticleMeshRenderData.bNeedsCull || ParticleMeshRenderData.bUseGPUScene)
diff --git a/Engine/Plugins/FX/Niagara/Source/Niagara/Private/NiagaraRendererRibbons.cpp b/Engine/Plugins/FX/Niagara/Source/Niagara/Private/NiagaraRendererRibbons.cpp
--- a/Engine/Plugins/FX/Niagara/Source/Niagara/Private/NiagaraRendererRibbons.cpp
+++ b/Engine/Plugins/FX/Niagara/Source/Niagara/Private/NiagaraRendererRibbons.cpp
@@ -15,9 +15,6 @@
 
 DECLARE_CYCLE_STAT(TEXT("Generate Ribbon Vertex Data [GT]"), STAT_NiagaraGenRibbonVertexData, STATGROUP_Niagara);
 DECLARE_CYCLE_STAT(TEXT("Render Ribbons [RT]"), STAT_NiagaraRenderRibbons, STATGROUP_Niagara);
-DECLARE_CYCLE_STAT(TEXT("Render Ribbons - CPU Sim Copy[RT]"), STAT_NiagaraRenderRibbonsCPUSimCopy, STATGROUP_Niagara);
-DECLARE_CYCLE_STAT(TEXT("Render Ribbons - CPU Sim Memcopy[RT]"), STAT_NiagaraRenderRibbonsCPUSimMemCopy, STATGROUP_Niagara);
-DECLARE_CYCLE_STAT(TEXT("Genereate GPU Buffers"), STAT_NiagaraGenRibbonGpuBuffers, STATGROUP_Niagara);
 
 int32 GNiagaraRibbonTessellationEnabled = 1;
 static FAutoConsoleVariableRef CVarNiagaraRibbonTessellationEnabled(
@@ -1237,7 +1234,6 @@ FNiagaraRendererRibbons::FCPUSimParticleDataAllocation FNiagaraRendererRibbons::
 
 	if (SimTarget == ENiagaraSimTarget::CPUSim)
 	{
-		SCOPE_CYCLE_COUNTER(STAT_NiagaraRenderRibbonsCPUSimMemCopy);
 		CPUSimParticleDataAllocation.ParticleData = TransferDataToGPU(DynamicReadBuffer, RendererLayout, MakeArrayView<uint32>(nullptr, 0), SourceParticleData);
 	}
 
diff --git a/Engine/Plugins/FX/Niagara/Source/Niagara/Private/NiagaraRendererSprites.cpp b/Engine/Plugins/FX/Niagara/Source/Niagara/Private/NiagaraRendererSprites.cpp
--- a/Engine/Plugins/FX/Niagara/Source/Niagara/Private/NiagaraRendererSprites.cpp
+++ b/Engine/Plugins/FX/Niagara/Source/Niagara/Private/NiagaraRendererSprites.cpp
@@ -16,14 +16,7 @@
 #include "Renderer/Private/ScenePrivate.h"
 #include "NiagaraCullProxyComponent.h"
 
-DECLARE_CYCLE_STAT(TEXT("Generate Sprite Dynamic Data [GT]"), STAT_NiagaraGenSpriteDynamicData, STATGROUP_Niagara);
 DECLARE_CYCLE_STAT(TEXT("Render Sprites [RT]"), STAT_NiagaraRenderSprites, STATGROUP_Niagara);
-DECLARE_CYCLE_STAT(TEXT("Render Sprites - CPU Sim Copy[RT]"), STAT_NiagaraRenderSpritesCPUSimCopy, STATGROUP_Niagara);
-DECLARE_CYCLE_STAT(TEXT("Render Sprites - CPU Sim Memcopy[RT]"), STAT_NiagaraRenderSpritesCPUSimMemCopy, STATGROUP_Niagara);
-DECLARE_CYCLE_STAT(TEXT("Render Sprites - Sorting[RT]"), STAT_NiagaraRenderSpritesSorting, STATGROUP_Niagara);
-DECLARE_CYCLE_STAT(TEXT("Render Sprites - GlobalSortCPU[RT]"), STAT_NiagaraRenderSpritesGlobalSortCPU, STATGROUP_Niagara);
-
-DECLARE_CYCLE_STAT(TEXT("Genereate GPU Buffers"), STAT_NiagaraGenSpriteGpuBuffers, STATGROUP_Niagara);
 DECLARE_DWORD_COUNTER_STAT(TEXT("NumSprites"), STAT_NiagaraNumSprites, STATGROUP_Niagara);
 
 static int32 GbEnableNiagaraSpriteRendering = 1;
@@ -328,8 +321,6 @@ void FNiagaraRendererSprites::PrepareParticleRenderBuffers(FParticleSpriteRender
 	{
 		if ( SimTarget == ENiagaraSimTarget::CPUSim )
 		{
-			SCOPE_CYCLE_COUNTER(STAT_NiagaraRenderSpritesCPUSimCopy);
-
 			// For CPU simulations we do not gather int parameters inside TransferDataToGPU currently so we need to copy off
 			// integrate attributes if we are culling on the GPU.
 			TArray<uint32, TInlineAllocator<1>> IntParamsToCopy;
@@ -1186,8 +1177,6 @@ FNiagaraDynamicDataBase *FNiagaraRendererSprites::GenerateDynamicData(const FNia
 			}
 		}
 
-		SCOPE_CYCLE_COUNTER(STAT_NiagaraGenSpriteDynamicData);
-
 		FNiagaraDataBuffer* DataToRender = Emitter->GetData().GetCurrentData();
 		if(SimTarget == ENiagaraSimTarget::GPUComputeSim || (DataToRender != nullptr &&  (SourceMode == ENiagaraRendererSourceDataMode::Emitter || (SourceMode == ENiagaraRendererSourceDataMode::Particles && DataToRender->GetNumInstances() > 0))))
 		{
