- Niagara Component Dump will call out components free in pool

#rnx

diff --git a/Engine/Plugins/FX/Niagara/Source/Niagara/Private/NiagaraComponent.cpp b/Engine/Plugins/FX/Niagara/Source/Niagara/Private/NiagaraComponent.cpp
--- a/Engine/Plugins/FX/Niagara/Source/Niagara/Private/NiagaraComponent.cpp
+++ b/Engine/Plugins/FX/Niagara/Source/Niagara/Private/NiagaraComponent.cpp
@@ -134,10 +134,17 @@ FAutoConsoleCommandWithWorldAndArgs DumpNiagaraComponentsCommand(
 					}
 				}
 
-				UE_LOG(LogNiagara, Log, TEXT("Component '%s' Asset '%s' Actor '%s' is %s"), *GetNameSafe(Component), *GetNameSafe(NiagaraSystem), *GetNameSafe(Component->GetTypedOuter<AActor>()), Component->IsActive() ? TEXT("Active") : TEXT("Inactive"));
-				if (FNiagaraSystemInstanceControllerPtr SystemInstanceController = Component->GetSystemInstanceController())
+				if (Component->PoolingMethod == ENCPoolMethod::FreeInPool)
 				{
-					SystemInstanceController->DebugDump(bFullDump);
+					UE_LOG(LogNiagara, Log, TEXT("Component '%s' Asset '%s' Free In Compnent Pool"), *GetNameSafe(Component), *GetNameSafe(NiagaraSystem));
+				}
+				else
+				{
+					UE_LOG(LogNiagara, Log, TEXT("Component '%s' Asset '%s' Actor '%s' is %s"), *GetNameSafe(Component), *GetNameSafe(NiagaraSystem), *GetNameSafe(Component->GetTypedOuter<AActor>()), Component->IsActive() ? TEXT("Active") : TEXT("Inactive"));
+					if (FNiagaraSystemInstanceControllerPtr SystemInstanceController = Component->GetSystemInstanceController())
+					{
+						SystemInstanceController->DebugDump(bFullDump);
+					}
 				}
 			}
 			UE_LOG(LogNiagara, Log, TEXT("=========================== End Niagara Dump ==========================="));
