MGPU:  Fix major performance issue with particles, by removing unnecessary GPU lock steps on cross GPU transfers.  Saves 20+ ms on Virtual Production Academy project.

#rnx

diff --git a/Engine/Plugins/FX/Niagara/Source/Niagara/Private/NiagaraGpuComputeDispatch.cpp b/Engine/Plugins/FX/Niagara/Source/Niagara/Private/NiagaraGpuComputeDispatch.cpp
--- a/Engine/Plugins/FX/Niagara/Source/Niagara/Private/NiagaraGpuComputeDispatch.cpp
+++ b/Engine/Plugins/FX/Niagara/Source/Niagara/Private/NiagaraGpuComputeDispatch.cpp
@@ -2169,7 +2169,7 @@ void FNiagaraGpuComputeDispatch::MultiGPUResourceModified(FRHICommandList& RHICm
 	if (bCrossGPUTransferEnabled && bRequiredForRendering)
 	{
 		const bool bPullData = false;
-		const bool bLockStep = true;
+		const bool bLockStep = false;
 
 		const FRHIGPUMask GPUMask = RHICmdList.GetGPUMask();
 		for (uint32 GPUIndex : FRHIGPUMask::All())
@@ -2197,7 +2197,7 @@ void FNiagaraGpuComputeDispatch::AddCrossGPUTransfer(FRHICommandList& RHICmdList
 	if (Buffer)
 	{
 		const bool bPullData = false;
-		const bool bLockStep = true;
+		const bool bLockStep = false;
 
 		const FRHIGPUMask GPUMask = RHICmdList.GetGPUMask();
 		for (uint32 GPUIndex : FRHIGPUMask::All())
diff --git a/Engine/Source/Runtime/Engine/Private/Particles/ParticleGpuSimulation.cpp b/Engine/Source/Runtime/Engine/Private/Particles/ParticleGpuSimulation.cpp
--- a/Engine/Source/Runtime/Engine/Private/Particles/ParticleGpuSimulation.cpp
+++ b/Engine/Source/Runtime/Engine/Private/Particles/ParticleGpuSimulation.cpp
@@ -4740,7 +4740,7 @@ void FFXSystem::SimulateGPUParticles(
 		[&](FRHITexture* TextureToTransfer)
 		{
 			const bool bPullData = false;
-			const bool bLockStep = true;
+			const bool bLockStep = false;
 
 			const FRHIGPUMask GPUMask = RHICmdList.GetGPUMask();
 			for (uint32 GPUIndex : FRHIGPUMask::All())
