Niagara - Fixed an issue where static switch inputs which were at the module default would be reset to the parent value when merged.  This was due to changes I made which implemented more correct diffing of the static switch input values.

diff --git a/Engine/Plugins/FX/Niagara/Source/NiagaraEditor/Private/NiagaraScriptMergeManager.cpp b/Engine/Plugins/FX/Niagara/Source/NiagaraEditor/Private/NiagaraScriptMergeManager.cpp
--- a/Engine/Plugins/FX/Niagara/Source/NiagaraEditor/Private/NiagaraScriptMergeManager.cpp
+++ b/Engine/Plugins/FX/Niagara/Source/NiagaraEditor/Private/NiagaraScriptMergeManager.cpp
@@ -2709,8 +2709,23 @@ FNiagaraScriptMergeManager::FApplyDiffResults FNiagaraScriptMergeManager::Remove
 	}
 	else if (OverrideToRemove->GetStaticSwitchValue().IsSet())
 	{
-		// TODO: Static switches are always treated as overrides right now so removing them is a no-op.  This code should updated
-		// so that removing a static switch override sets the value back to the module default.
+		const UEdGraphSchema_Niagara* Schema = GetDefault<UEdGraphSchema_Niagara>();
+		UEdGraphPin** StaticSwitchPinPtr = OverrideToRemove->GetOwningFunctionCall()->Pins.FindByPredicate([&OverrideToRemove, &Schema](UEdGraphPin* Pin)
+		{
+			if (Pin->Direction != EGPD_Input)
+			{
+				return false;
+			}
+			FNiagaraVariable PinVariable = Schema->PinToNiagaraVariable(Pin);
+			return PinVariable.GetName() == *OverrideToRemove->GetInputName() &&
+				PinVariable.GetType() == OverrideToRemove->GetType();
+		});
+
+		if (StaticSwitchPinPtr != nullptr && (*StaticSwitchPinPtr)->AutogeneratedDefaultValue.IsEmpty() == false)
+		{
+			(*StaticSwitchPinPtr)->DefaultValue = (*StaticSwitchPinPtr)->AutogeneratedDefaultValue;
+			Results.bModifiedGraph = true;
+		}
 	}
 	else
 	{
