- Fixed an issue with "Add Parameter" not knowing whether to add a new parameter to a node or add an existing one, which resulted in a suffix "001" being added
- Now the add widget is explicit with whether we are adding a new pin representing a new parameter or a new pin representing an old parameter
- drag & drop of parameters was fixed so it no longer accepts a drop on parameter nodes that already have the parameter present. Previously, it would add a new parameter instead

diff --git a/Engine/Plugins/FX/Niagara/Source/NiagaraEditor/Private/EdGraphSchema_Niagara.cpp b/Engine/Plugins/FX/Niagara/Source/NiagaraEditor/Private/EdGraphSchema_Niagara.cpp
--- a/Engine/Plugins/FX/Niagara/Source/NiagaraEditor/Private/EdGraphSchema_Niagara.cpp
+++ b/Engine/Plugins/FX/Niagara/Source/NiagaraEditor/Private/EdGraphSchema_Niagara.cpp
@@ -1429,7 +1429,7 @@ bool UEdGraphSchema_Niagara::ShouldHidePinDefaultValue(UEdGraphPin* Pin) const
 	return false;
 }
 
-FNiagaraVariable UEdGraphSchema_Niagara::PinToNiagaraVariable(const UEdGraphPin* Pin, bool bNeedsValue, ENiagaraStructConversion StructConversion) const
+FNiagaraVariable UEdGraphSchema_Niagara::PinToNiagaraVariable(const UEdGraphPin* Pin, bool bNeedsValue, ENiagaraStructConversion StructConversion)
 {
 	FNiagaraVariable Var = FNiagaraVariable(PinToTypeDefinition(Pin, StructConversion), Pin->PinName);
 	bool bHasValue = false;
diff --git a/Engine/Plugins/FX/Niagara/Source/NiagaraEditor/Private/NiagaraActions.cpp b/Engine/Plugins/FX/Niagara/Source/NiagaraEditor/Private/NiagaraActions.cpp
--- a/Engine/Plugins/FX/Niagara/Source/NiagaraEditor/Private/NiagaraActions.cpp
+++ b/Engine/Plugins/FX/Niagara/Source/NiagaraEditor/Private/NiagaraActions.cpp
@@ -387,13 +387,19 @@ FReply FNiagaraParameterGraphDragOperation::DroppedOnNode(FVector2D ScreenPositi
 		{
 			if (UNiagaraNodeParameterMapGet* GetMapNode = Cast<UNiagaraNodeParameterMapGet>(GetHoveredNode()))
 			{
-				FScopedTransaction AddNewPinTransaction(LOCTEXT("Drop Onto Get Pin", "Drop parameter onto Get node"));
-				FNiagaraStackGraphUtilities::AddNewVariableToParameterMapNode(GetMapNode, false, ScriptVar);
+				if(!GetMapNode->DoesParameterExistOnNode(ScriptVar->Variable))
+				{
+					FScopedTransaction AddNewPinTransaction(LOCTEXT("Drop Onto Get Pin", "Drop parameter onto Get node"));
+					FNiagaraStackGraphUtilities::AddNewVariableToParameterMapNode(GetMapNode, false, ScriptVar);
+				}
 			}
 			else if (UNiagaraNodeParameterMapSet* SetMapNode = Cast<UNiagaraNodeParameterMapSet>(GetHoveredNode()))
 			{
-				FScopedTransaction AddNewPinTransaction(LOCTEXT("Drop Onto Set Pin", "Drop parameter onto Set node"));
-				FNiagaraStackGraphUtilities::AddNewVariableToParameterMapNode(SetMapNode, true, ScriptVar);
+				if(!SetMapNode->DoesParameterExistOnNode(ScriptVar->Variable))
+				{
+					FScopedTransaction AddNewPinTransaction(LOCTEXT("Drop Onto Set Pin", "Drop parameter onto Set node"));
+					FNiagaraStackGraphUtilities::AddNewVariableToParameterMapNode(SetMapNode, true, ScriptVar);
+				}
 			}
 		}
 
diff --git a/Engine/Plugins/FX/Niagara/Source/NiagaraEditor/Private/NiagaraNodeParameterMapBase.cpp b/Engine/Plugins/FX/Niagara/Source/NiagaraEditor/Private/NiagaraNodeParameterMapBase.cpp
--- a/Engine/Plugins/FX/Niagara/Source/NiagaraEditor/Private/NiagaraNodeParameterMapBase.cpp
+++ b/Engine/Plugins/FX/Niagara/Source/NiagaraEditor/Private/NiagaraNodeParameterMapBase.cpp
@@ -101,11 +101,6 @@ void UNiagaraNodeParameterMapBase::PinDescriptionTextCommitted(const FText& Text
 	Graph->SetMetaData(Var, NewMetaData);
 }
 
-void UNiagaraNodeParameterMapBase::CollectAddPinActions(FGraphActionListBuilderBase& OutActions, bool& bOutCreateRemainingActions, UEdGraphPin* Pin)
-{
-	bOutCreateRemainingActions = true;
-}
-
 void UNiagaraNodeParameterMapBase::GetPinHoverText(const UEdGraphPin& Pin, FString& HoverTextOut) const
 {
 	FText Text;
@@ -206,7 +201,21 @@ bool UNiagaraNodeParameterMapBase::CanHandleDropOperation(TSharedPtr<FDragDropOp
 	{
 		if (StaticCastSharedPtr<FNiagaraParameterGraphDragOperation>(DragDropOperation)->IsCurrentlyHoveringNode(this))
 		{
-			// The FNiagaraParameterGraphDragOperation handles the drop action itself so just return true here and let it handle it when it's executed.
+			TSharedPtr<FNiagaraParameterGraphDragOperation> ParameterDragDropOperation = StaticCastSharedPtr<FNiagaraParameterGraphDragOperation>(DragDropOperation);
+			
+			if(ParameterDragDropOperation->GetSourceAction().IsValid())
+			{
+				if(TSharedPtr<FNiagaraParameterAction> ParameterAction = StaticCastSharedPtr<FNiagaraParameterAction>(ParameterDragDropOperation->GetSourceAction()))
+				{
+					FNiagaraVariable Parameter = ParameterAction->GetParameter();
+					// if the parameter already exists on a node, we won't allow a drop
+					if(DoesParameterExistOnNode(Parameter))
+					{
+						return false;
+					}
+				}				
+			}
+			
 			return true;
 		}
 	}
@@ -376,6 +385,24 @@ void UNiagaraNodeParameterMapBase::SelectParameterFromPin(const UEdGraphPin* InP
 	}
 }
 
+bool UNiagaraNodeParameterMapBase::DoesParameterExistOnNode(FNiagaraVariable Parameter)
+{
+	for(UEdGraphPin* Pin : Pins)
+	{
+		if(Pin->PinType.PinSubCategory == ParameterPinSubCategory)
+		{
+			FNiagaraVariable CandidateVariable = UEdGraphSchema_Niagara::PinToNiagaraVariable(Pin);
+
+			if(CandidateVariable == Parameter)
+			{
+				return true;
+			}
+		}
+	}
+
+	return false;
+}
+
 void UNiagaraNodeParameterMapBase::GetChangeNamespaceModifierSubMenuForPin(UToolMenu* Menu, UEdGraphPin* InPin)
 {
 	FToolMenuSection& Section = Menu->AddSection("Section");
diff --git a/Engine/Plugins/FX/Niagara/Source/NiagaraEditor/Private/NiagaraNodeParameterMapBase.h b/Engine/Plugins/FX/Niagara/Source/NiagaraEditor/Private/NiagaraNodeParameterMapBase.h
--- a/Engine/Plugins/FX/Niagara/Source/NiagaraEditor/Private/NiagaraNodeParameterMapBase.h
+++ b/Engine/Plugins/FX/Niagara/Source/NiagaraEditor/Private/NiagaraNodeParameterMapBase.h
@@ -32,7 +32,6 @@ public:
 	void PinDescriptionTextCommitted(const FText& Text, ETextCommit::Type CommitType, UEdGraphPin* Pin);
 
 	virtual void GetNodeContextMenuActions(class UToolMenu* Menu, class UGraphNodeContextMenuContext* Context) const override;
-	virtual void CollectAddPinActions(FGraphActionListBuilderBase& OutActions, bool& bOutCreateRemainingActions, UEdGraphPin* Pin) override;
 	virtual void GetPinHoverText(const UEdGraphPin& Pin, FString& HoverTextOut) const override;
 
 	virtual bool IncludeParentNodeContextMenu() const { return true; }
@@ -53,12 +52,13 @@ public:
 
 	bool HandleDropOperation(TSharedPtr<FDragDropOperation> DropOperation);
 
+	bool DoesParameterExistOnNode(FNiagaraVariable Parameter);
 protected:
 	void GetChangeNamespaceSubMenuForPin(UToolMenu* Menu, UEdGraphPin* InPin);
 	void ChangeNamespaceForPin(UEdGraphPin* InPin, FNiagaraNamespaceMetadata NewNamespaceMetadata);
 
 	virtual void SelectParameterFromPin(const UEdGraphPin* InPin);
-
+	
 	void GetChangeNamespaceModifierSubMenuForPin(UToolMenu* Menu, UEdGraphPin* InPin);
 
 	FText GetSetNamespaceModifierForPinToolTip(const UEdGraphPin* InPin, FName InNamespaceModifier) const;
diff --git a/Engine/Plugins/FX/Niagara/Source/NiagaraEditor/Private/NiagaraNodeWithDynamicPins.cpp b/Engine/Plugins/FX/Niagara/Source/NiagaraEditor/Private/NiagaraNodeWithDynamicPins.cpp
--- a/Engine/Plugins/FX/Niagara/Source/NiagaraEditor/Private/NiagaraNodeWithDynamicPins.cpp
+++ b/Engine/Plugins/FX/Niagara/Source/NiagaraEditor/Private/NiagaraNodeWithDynamicPins.cpp
@@ -294,35 +294,6 @@ void UNiagaraNodeWithDynamicPins::GetNodeContextMenuActions(UToolMenu* Menu, UGr
 	}
 }
 
-void UNiagaraNodeWithDynamicPins::CollectAddPinActions(FGraphActionListBuilderBase& OutActions, bool& bOutCreateRemainingActions, UEdGraphPin* Pin)
-{
-	TArray<FNiagaraTypeDefinition> Types(FNiagaraTypeRegistry::GetRegisteredTypes());
-	Types.Sort([](const FNiagaraTypeDefinition& A, const FNiagaraTypeDefinition& B) { return (A.GetNameText().ToLower().ToString() < B.GetNameText().ToLower().ToString()); });
-
-	for (const FNiagaraTypeDefinition& RegisteredType : Types)
-	{
-		bool bAllowType = false;
-		bAllowType = AllowNiagaraTypeForAddPin(RegisteredType);
-
-		if (bAllowType)
-		{
-			FNiagaraVariable Var(RegisteredType, FName(*RegisteredType.GetName()));
-			FNiagaraEditorUtilities::ResetVariableToDefaultValue(Var);
-
-			FText Category = FNiagaraEditorUtilities::GetVariableTypeCategory(Var);
-			const FText DisplayName = RegisteredType.GetNameText();
-			const FText Tooltip = FText::Format(LOCTEXT("AddButtonTypeEntryToolTipFormat", "Add a new {0} pin"), RegisteredType.GetNameText());
-			TSharedPtr<FNiagaraMenuAction> Action(new FNiagaraMenuAction(
-				Category, DisplayName, Tooltip, 0, FText::GetEmpty(),
-				FNiagaraMenuAction::FOnExecuteStackAction::CreateUObject(this, &UNiagaraNodeWithDynamicPins::AddParameter, Var, (const UEdGraphPin*)Pin)));
-
-			OutActions.AddAction(Action);
-		}
-	}
-
-	bOutCreateRemainingActions = false;
-}
-
 void UNiagaraNodeWithDynamicPins::AddParameter(FNiagaraVariable Parameter, const UEdGraphPin* AddPin)
 {
 	if (this->IsA<UNiagaraNodeParameterMapBase>())
@@ -374,6 +345,25 @@ void UNiagaraNodeWithDynamicPins::AddParameter(const UNiagaraScriptVariable* Scr
 	}
 }
 
+void UNiagaraNodeWithDynamicPins::AddExistingParameter(FNiagaraVariable Parameter, const UEdGraphPin* AddPin)
+{
+	if (this->IsA<UNiagaraNodeParameterMapBase>())
+	{
+		// Parameter map type nodes create new parameters when adding pins.
+		FScopedTransaction AddNewPinTransaction(LOCTEXT("AddNewPinTransaction", "Add pin to node"));
+		
+		UNiagaraGraph* Graph = GetNiagaraGraph();
+		checkf(Graph != nullptr, TEXT("Failed to get niagara graph when adding pin!"));
+		
+		Modify();
+		UEdGraphPin* Pin = this->RequestNewTypedPin(AddPin->Direction, Parameter.GetType(), Parameter.GetName());
+	}
+	else
+	{
+		RequestNewTypedPin(AddPin->Direction, Parameter.GetType(), Parameter.GetName());
+	}
+}
+
 void UNiagaraNodeWithDynamicPins::RemoveDynamicPin(UEdGraphPin* Pin)
 {
 	RemovePin(Pin);
diff --git a/Engine/Plugins/FX/Niagara/Source/NiagaraEditor/Private/NiagaraParameterDefinitions.cpp b/Engine/Plugins/FX/Niagara/Source/NiagaraEditor/Private/NiagaraParameterDefinitions.cpp
--- a/Engine/Plugins/FX/Niagara/Source/NiagaraEditor/Private/NiagaraParameterDefinitions.cpp
+++ b/Engine/Plugins/FX/Niagara/Source/NiagaraEditor/Private/NiagaraParameterDefinitions.cpp
@@ -6,6 +6,8 @@
 #include "NiagaraEditorUtilities.h"
 #include "NiagaraScriptVariable.h"
 
+#define LOCTEXT_NAMESPACE "NiagaraParameterDefinitions"
+
 UNiagaraParameterDefinitions::UNiagaraParameterDefinitions(const FObjectInitializer& ObjectInitializer)
 	: Super(ObjectInitializer)
 {
@@ -43,6 +45,7 @@ void UNiagaraParameterDefinitions::PostEditChangeProperty(struct FPropertyChange
 
 void UNiagaraParameterDefinitions::AddParameter(const FNiagaraVariable& NewVariable)
 {
+	FScopedTransaction AddTransaction(LOCTEXT("AddParameter", "Add Parameter"));
 	Modify();
 	TObjectPtr<UNiagaraScriptVariable>& NewScriptVariable = ScriptVariables.Add_GetRef(NewObject<UNiagaraScriptVariable>(this, FName(), RF_Transactional));
 	NewScriptVariable->Init(NewVariable, FNiagaraVariableMetaData());
@@ -51,6 +54,29 @@ void UNiagaraParameterDefinitions::AddParameter(const FNiagaraVariable& NewVaria
 	NotifyParameterDefinitionsChanged();
 }
 
+bool UNiagaraParameterDefinitions::HasParameter(const FNiagaraVariable& Variable)
+{
+	TObjectPtr<UNiagaraScriptVariable>* NewScriptVariable = ScriptVariables.FindByPredicate([Variable](TObjectPtr<UNiagaraScriptVariable> ScriptVariable)
+	{
+		return ScriptVariable->Variable == Variable;
+	});
+
+	return NewScriptVariable != nullptr;
+}
+
+void UNiagaraParameterDefinitions::FindOrAddParameter(const FNiagaraVariable& Variable)
+{
+	TObjectPtr<UNiagaraScriptVariable>* NewScriptVariable = ScriptVariables.FindByPredicate([Variable](TObjectPtr<UNiagaraScriptVariable> ScriptVariable)
+	{
+		return ScriptVariable->Variable == Variable;
+	});
+
+	if(NewScriptVariable == nullptr)
+	{
+		AddParameter(Variable);
+	}
+}
+
 void UNiagaraParameterDefinitions::RemoveParameter(const FNiagaraVariable& VariableToRemove)
 {
 	const int32 Idx = ScriptVariables.IndexOfByPredicate([&VariableToRemove](const UNiagaraScriptVariable* ScriptVariable){ return ScriptVariable->Variable == VariableToRemove; });
@@ -251,3 +277,5 @@ void UNiagaraParameterDefinitions::NotifyParameterDefinitionsChanged()
 {
 	OnParameterDefinitionsChangedDelegate.Broadcast();
 }
+
+#undef LOCTEXT_NAMESPACE
diff --git a/Engine/Plugins/FX/Niagara/Source/NiagaraEditor/Private/ViewModels/NiagaraParameterPanelViewModel.cpp b/Engine/Plugins/FX/Niagara/Source/NiagaraEditor/Private/ViewModels/NiagaraParameterPanelViewModel.cpp
--- a/Engine/Plugins/FX/Niagara/Source/NiagaraEditor/Private/ViewModels/NiagaraParameterPanelViewModel.cpp
+++ b/Engine/Plugins/FX/Niagara/Source/NiagaraEditor/Private/ViewModels/NiagaraParameterPanelViewModel.cpp
@@ -1063,6 +1063,45 @@ void FNiagaraSystemToolkitParameterPanelViewModel::AddParameter(FNiagaraVariable
 	}
 }
 
+void FNiagaraSystemToolkitParameterPanelViewModel::FindOrAddParameter(FNiagaraVariable Variable, const FNiagaraParameterPanelCategory Category) const
+{
+	TGuardValue<bool> AddParameterRefreshGuard(bIsAddingParameter, true);
+	bool bSuccess = false;
+	UNiagaraSystem& System = SystemViewModel->GetSystem();
+
+	FScopedTransaction AddTransaction(LOCTEXT("FindOrAddAddSystemParameterTransaction", "Add parameter to script."));
+	const FGuid NamespaceId = FNiagaraEditorUtilities::GetNamespaceMetaDataForVariableName(Variable.GetName()).GetGuid();
+	if (NamespaceId == FNiagaraEditorGuids::UserNamespaceMetaDataGuid)
+	{
+		if(System.GetExposedParameters().FindParameterVariable(Variable) == nullptr)
+		{
+			AddParameter(Variable, Category, false, false);
+		}
+	}
+	else
+	{
+		UNiagaraEditorParametersAdapter* EditorParametersAdapter = SystemViewModel->GetEditorOnlyParametersAdapter();
+		TArray<UNiagaraScriptVariable*>& EditorOnlyScriptVars = EditorParametersAdapter->GetParameters();
+		bool bNewScriptVarAlreadyExists = EditorOnlyScriptVars.ContainsByPredicate([&Variable](const UNiagaraScriptVariable* ScriptVar) { return ScriptVar->Variable == Variable; });
+		
+		if (bNewScriptVarAlreadyExists == false)
+		{
+			AddParameter(Variable, Category, false, false);
+		}
+	}
+
+	if (bSuccess)
+	{
+		Refresh();
+	}
+	else
+	{
+		AddTransaction.Cancel();
+	}
+	
+	SelectParameterItemByName(Variable.GetName(), false);
+}
+
 bool FNiagaraSystemToolkitParameterPanelViewModel::GetCanAddParametersToCategory(FNiagaraParameterPanelCategory Category) const
 {
 	return GetEditableGraphsConst().Num() > 0 && Category.NamespaceMetaData.GetGuid() != FNiagaraEditorGuids::StaticSwitchNamespaceMetaDataGuid;
@@ -1345,7 +1384,8 @@ FMenuAndSearchBoxWidgets FNiagaraSystemToolkitParameterPanelViewModel::GetParame
 		.Graphs(GetEditableGraphsConst())
 		.AvailableParameterDefinitions(SystemViewModel->GetAvailableParameterDefinitions(bSkipSubscribedLibraries))
 		.SubscribedParameterDefinitions(SystemViewModel->GetSubscribedParameterDefinitions())
-		.OnAddParameter(this, &FNiagaraSystemToolkitParameterPanelViewModel::AddParameter, Category, bRequestRename, bMakeUniqueName)
+		.OnNewParameterRequested(this, &FNiagaraSystemToolkitParameterPanelViewModel::AddParameter, Category, bRequestRename, bMakeUniqueName)
+		.OnSpecificParameterRequested(this, &FNiagaraSystemToolkitParameterPanelViewModel::FindOrAddParameter, Category)
 		.OnAddScriptVar(this, &FNiagaraSystemToolkitParameterPanelViewModel::AddScriptVariable)
 		.OnAddParameterDefinitions(this, &FNiagaraSystemToolkitParameterPanelViewModel::AddParameterDefinitions)
 		.OnAllowMakeType_Static(&INiagaraParameterPanelViewModel::CanMakeNewParameterOfType)
@@ -1487,7 +1527,7 @@ TSharedRef<SWidget> FNiagaraSystemToolkitParameterPanelViewModel::CreateAddParam
 		.Graphs(InGraphs)
 		.AvailableParameterDefinitions(SystemViewModel->GetAvailableParameterDefinitions(bSkipSubscribedLibraries))
 		.SubscribedParameterDefinitions(SystemViewModel->GetSubscribedParameterDefinitions())
-		.OnAddParameter_Lambda(AddParameterLambda)
+		.OnNewParameterRequested_Lambda(AddParameterLambda)
 		.OnAddScriptVar_Lambda(AddScriptVarLambda)
 		.OnAddParameterDefinitions(this, &FNiagaraSystemToolkitParameterPanelViewModel::AddParameterDefinitions)
 		.OnAllowMakeType_Static(&INiagaraParameterPanelViewModel::CanMakeNewParameterOfType)
@@ -2252,6 +2292,33 @@ void FNiagaraScriptToolkitParameterPanelViewModel::AddParameter(FNiagaraVariable
 	}
 }
 
+void FNiagaraScriptToolkitParameterPanelViewModel::FindOrAddParameter(FNiagaraVariable Variable, const FNiagaraParameterPanelCategory Category) const
+{
+	TGuardValue<bool> AddParameterRefreshGuard(bIsAddingParameter, true);
+	bool bSuccess = false;
+
+	FScopedTransaction AddTransaction(LOCTEXT("FindOrAddScriptParameterTransaction", "Add parameter to script."));
+	for (UNiagaraGraph* Graph : GetEditableGraphs())
+	{
+		Graph->Modify();
+		const UNiagaraScriptVariable* NewScriptVar = Graph->AddParameter(Variable);
+		bSuccess = true;
+
+		// Check if the new parameter has the same name and type as an existing parameter definition, and if so, link to the definition automatically.
+		SubscribeParameterToLibraryIfMatchingDefinition(NewScriptVar, NewScriptVar->Variable.GetName());
+	}
+
+	if (bSuccess)
+	{
+		Refresh();
+		SelectParameterItemByName(Variable.GetName(), false);
+	}
+	else
+	{
+		AddTransaction.Cancel();
+	}
+}
+
 bool FNiagaraScriptToolkitParameterPanelViewModel::GetCanAddParametersToCategory(FNiagaraParameterPanelCategory Category) const
 {
 	return GetEditableGraphsConst().Num() > 0 && Category.NamespaceMetaData.GetGuid() != FNiagaraEditorGuids::StaticSwitchNamespaceMetaDataGuid;
@@ -2623,7 +2690,8 @@ FMenuAndSearchBoxWidgets FNiagaraScriptToolkitParameterPanelViewModel::GetParame
 		.Graphs(GetEditableGraphsConst())
 		.AvailableParameterDefinitions(ScriptViewModel->GetAvailableParameterDefinitions(bSkipSubscribedLibraries))
  		.SubscribedParameterDefinitions(ScriptViewModel->GetSubscribedParameterDefinitions())
-		.OnAddParameter(this, &FNiagaraScriptToolkitParameterPanelViewModel::AddParameter, Category, bRequestRename, bMakeUniqueName)
+		.OnNewParameterRequested(this, &FNiagaraScriptToolkitParameterPanelViewModel::AddParameter, Category, bRequestRename, bMakeUniqueName)
+		.OnSpecificParameterRequested(this, &FNiagaraScriptToolkitParameterPanelViewModel::FindOrAddParameter, Category)
 		.OnAddScriptVar(this, &FNiagaraScriptToolkitParameterPanelViewModel::AddScriptVariable)
 		.OnAddParameterDefinitions(this, &FNiagaraScriptToolkitParameterPanelViewModel::AddParameterDefinitions)
 		.OnAllowMakeType_Static(&INiagaraParameterPanelViewModel::CanMakeNewParameterOfType)
@@ -2973,12 +3041,23 @@ void FNiagaraParameterDefinitionsToolkitParameterPanelViewModel::AddParameter(FN
 		const FName NewUniqueName = FNiagaraUtilities::GetUniqueName(NewVariable.GetName(), Names);
 		NewVariable.SetName(NewUniqueName);
 	}
-
-	FScopedTransaction AddTransaction(LOCTEXT("ParameterDefinitionsToolkitParameterPanelViewModel_AddParameter", "Add Parameter"));
+	
 	ParameterDefinitionsWeak.Get()->AddParameter(NewVariable);
 	SelectParameterItemByName(NewVariable.GetName(), bRequestRename);
 }
 
+void FNiagaraParameterDefinitionsToolkitParameterPanelViewModel::FindOrAddParameter(FNiagaraVariable Variable, const FNiagaraParameterPanelCategory Category) const
+{
+	TGuardValue<bool> AddParameterRefreshGuard(bIsAddingParameter, true);
+
+	if(!ParameterDefinitionsWeak.Get()->HasParameter(Variable))
+	{
+		ParameterDefinitionsWeak.Get()->FindOrAddParameter(Variable);
+	}
+	
+	SelectParameterItemByName(Variable.GetName(), false);
+}
+
 bool FNiagaraParameterDefinitionsToolkitParameterPanelViewModel::GetCanAddParametersToCategory(FNiagaraParameterPanelCategory Category) const
 {
 	return true;
@@ -3200,7 +3279,8 @@ FMenuAndSearchBoxWidgets FNiagaraParameterDefinitionsToolkitParameterPanelViewMo
 	const bool bMakeUniqueName = true;
 
 	TSharedPtr<SNiagaraAddParameterFromPanelMenu> MenuWidget = SAssignNew(ParameterMenuWidget, SNiagaraAddParameterFromPanelMenu)
-	.OnAddParameter(this, &FNiagaraParameterDefinitionsToolkitParameterPanelViewModel::AddParameter, Category, bRequestRename, bMakeUniqueName)
+	.OnNewParameterRequested(this, &FNiagaraParameterDefinitionsToolkitParameterPanelViewModel::AddParameter, Category, bRequestRename, bMakeUniqueName)
+	.OnSpecificParameterRequested(this, &FNiagaraParameterDefinitionsToolkitParameterPanelViewModel::FindOrAddParameter, Category)
 	.OnAllowMakeType_Static(&INiagaraParameterPanelViewModel::CanMakeNewParameterOfType)
 	.NamespaceId(Category.NamespaceMetaData.GetGuid())
 	.ShowNamespaceCategory(false)
diff --git a/Engine/Plugins/FX/Niagara/Source/NiagaraEditor/Private/Widgets/SNiagaraGraphPinAdd.cpp b/Engine/Plugins/FX/Niagara/Source/NiagaraEditor/Private/Widgets/SNiagaraGraphPinAdd.cpp
--- a/Engine/Plugins/FX/Niagara/Source/NiagaraEditor/Private/Widgets/SNiagaraGraphPinAdd.cpp
+++ b/Engine/Plugins/FX/Niagara/Source/NiagaraEditor/Private/Widgets/SNiagaraGraphPinAdd.cpp
@@ -126,7 +126,8 @@ TSharedRef<SWidget> SNiagaraGraphPinAdd::OnGetAddButtonMenuContent()
 		.Graphs(InGraphs)
 		.AvailableParameterDefinitions(ScriptViewModel->GetAvailableParameterDefinitions(bSkipSubscribedLibraries))
 		.SubscribedParameterDefinitions(ScriptViewModel->GetSubscribedParameterDefinitions())
-		.OnAddParameter_UObject(OwningNode, &UNiagaraNodeWithDynamicPins::AddParameter, ConstAddPin)
+		.OnNewParameterRequested_UObject(OwningNode, &UNiagaraNodeWithDynamicPins::AddParameter, ConstAddPin)
+		.OnSpecificParameterRequested_UObject(OwningNode, &UNiagaraNodeWithDynamicPins::AddExistingParameter, ConstAddPin)
 		.OnAddScriptVar_UObject(OwningNode, &UNiagaraNodeWithDynamicPins::AddParameter, ConstAddPin)
 		.OnAddParameterDefinitions(ScriptViewModel.ToSharedRef(), &FNiagaraScriptViewModel::SubscribeToParameterDefinitions)
 		.OnAllowMakeType_UObject(OwningNode, &UNiagaraNodeWithDynamicPins::AllowNiagaraTypeForAddPin) 
diff --git a/Engine/Plugins/FX/Niagara/Source/NiagaraEditor/Public/EdGraphSchema_Niagara.h b/Engine/Plugins/FX/Niagara/Source/NiagaraEditor/Public/EdGraphSchema_Niagara.h
--- a/Engine/Plugins/FX/Niagara/Source/NiagaraEditor/Public/EdGraphSchema_Niagara.h
+++ b/Engine/Plugins/FX/Niagara/Source/NiagaraEditor/Public/EdGraphSchema_Niagara.h
@@ -136,7 +136,7 @@ class NIAGARAEDITOR_API UEdGraphSchema_Niagara : public UEdGraphSchema
 	  * bNeedsValue Whether or not the returned variable must be allocated to a valid value. When true if the pin doesn't have a valid default value itself the variable will be reset to default before return.
 	  * returns The newly created variable.
 	  */
-	FNiagaraVariable PinToNiagaraVariable(const UEdGraphPin* Pin, bool bNeedsValue=false, ENiagaraStructConversion StructConversion = ENiagaraStructConversion::UserFacing) const;
+	static FNiagaraVariable PinToNiagaraVariable(const UEdGraphPin* Pin, bool bNeedsValue=false, ENiagaraStructConversion StructConversion = ENiagaraStructConversion::UserFacing);
 
 	/** 
 	  * Tries to get a default value string for a graph pin from a niagara variable.
diff --git a/Engine/Plugins/FX/Niagara/Source/NiagaraEditor/Public/NiagaraNodeWithDynamicPins.h b/Engine/Plugins/FX/Niagara/Source/NiagaraEditor/Public/NiagaraNodeWithDynamicPins.h
--- a/Engine/Plugins/FX/Niagara/Source/NiagaraEditor/Public/NiagaraNodeWithDynamicPins.h
+++ b/Engine/Plugins/FX/Niagara/Source/NiagaraEditor/Public/NiagaraNodeWithDynamicPins.h
@@ -38,12 +38,10 @@ public:
 	virtual bool AllowNiagaraTypeForAddPin(const FNiagaraTypeDefinition& InType) const;
 	virtual bool AllowNiagaraTypeForAddPin(const FNiagaraTypeDefinition& InType, EEdGraphPinDirection InDirection) const { return AllowNiagaraTypeForAddPin(InType); };
 
-	/** Used in to gather the actions for selecting the pin to add. */
-	virtual void CollectAddPinActions(FGraphActionListBuilderBase& OutActions, bool& bOutCreateRemainingActions, UEdGraphPin* Pin);
-
 	/** Request a new pin. */
 	void AddParameter(FNiagaraVariable Parameter, const UEdGraphPin* AddPin);
 	void AddParameter(const UNiagaraScriptVariable* ScriptVar, const UEdGraphPin* AddPin);
+	void AddExistingParameter(FNiagaraVariable Parameter, const UEdGraphPin* AddPin);
 
 	/** Convenience method to determine whether this Node is a Map Get or Map Set when adding a parameter through the parameter panel. */
 	virtual EEdGraphPinDirection GetPinDirectionForNewParameters() { return EEdGraphPinDirection::EGPD_MAX; };
diff --git a/Engine/Plugins/FX/Niagara/Source/NiagaraEditor/Public/NiagaraParameterDefinitions.h b/Engine/Plugins/FX/Niagara/Source/NiagaraEditor/Public/NiagaraParameterDefinitions.h
--- a/Engine/Plugins/FX/Niagara/Source/NiagaraEditor/Public/NiagaraParameterDefinitions.h
+++ b/Engine/Plugins/FX/Niagara/Source/NiagaraEditor/Public/NiagaraParameterDefinitions.h
@@ -60,6 +60,8 @@ public:
 	//~ End UObject Interface
 
 	void AddParameter(const FNiagaraVariable& NewVariable);
+	bool HasParameter(const FNiagaraVariable& Variable);
+	void FindOrAddParameter(const FNiagaraVariable& Variable);
 	void RemoveParameter(const FNiagaraVariable& VariableToRemove);
 	void RenameParameter(const FNiagaraVariable& VariableToRename, const FName NewName);
 	const TArray<UNiagaraScriptVariable*>& GetParametersConst() const;
diff --git a/Engine/Plugins/FX/Niagara/Source/NiagaraEditor/Public/ViewModels/NiagaraParameterPanelViewModel.h b/Engine/Plugins/FX/Niagara/Source/NiagaraEditor/Public/ViewModels/NiagaraParameterPanelViewModel.h
--- a/Engine/Plugins/FX/Niagara/Source/NiagaraEditor/Public/ViewModels/NiagaraParameterPanelViewModel.h
+++ b/Engine/Plugins/FX/Niagara/Source/NiagaraEditor/Public/ViewModels/NiagaraParameterPanelViewModel.h
@@ -156,6 +156,8 @@ public:
 
 	virtual void AddParameter(FNiagaraVariable NewVariable, const FNiagaraParameterPanelCategory Category, const bool bRequestRename, const bool bMakeUniqueName) const = 0;
 
+	virtual void FindOrAddParameter(FNiagaraVariable Variable, const FNiagaraParameterPanelCategory Category) const = 0;
+	
 	virtual bool GetCanAddParametersToCategory(FNiagaraParameterPanelCategory Category) const = 0;
 
 	virtual void DeleteParameters(const TArray<FNiagaraParameterPanelItem>& ItemsToDelete) const = 0;
@@ -285,6 +287,8 @@ public:
 	//~ Begin INiagaraParameterPanelViewModel interface
 	virtual void AddParameter(FNiagaraVariable NewVariable, const FNiagaraParameterPanelCategory Category, const bool bRequestRename, const bool bMakeUniqueName) const override;
 
+	virtual void FindOrAddParameter(FNiagaraVariable Variable, const FNiagaraParameterPanelCategory Category) const override;
+
 	virtual bool GetCanAddParametersToCategory(FNiagaraParameterPanelCategory Category) const override;
 
 	virtual void DeleteParameters(const TArray<FNiagaraParameterPanelItem>& ItemsToDelete) const override;
@@ -385,6 +389,8 @@ public:
 	//~ Begin INiagaraParameterPanelViewModel interface
 	virtual void AddParameter(FNiagaraVariable NewVariable, const FNiagaraParameterPanelCategory Category, const bool bRequestRename, const bool bMakeUniqueName) const override;
 
+	virtual void FindOrAddParameter(FNiagaraVariable Variable, const FNiagaraParameterPanelCategory Category) const override;
+
 	virtual bool GetCanAddParametersToCategory(FNiagaraParameterPanelCategory Category) const override;
 
 	virtual void DeleteParameters(const TArray<FNiagaraParameterPanelItem>& ItemToDelete) const override;
@@ -485,6 +491,8 @@ public:
 	//~ Begin INiagaraParameterPanelViewModel interface
 	virtual void AddParameter(FNiagaraVariable NewVariable, const FNiagaraParameterPanelCategory Category, const bool bRequestRename, const bool bMakeUniqueName) const override;
 
+	virtual void FindOrAddParameter(FNiagaraVariable Variable, const FNiagaraParameterPanelCategory Category) const override;
+
 	virtual bool GetCanAddParametersToCategory(FNiagaraParameterPanelCategory Category) const override;
 
 	virtual void DeleteParameters(const TArray<FNiagaraParameterPanelItem>& ItemToDelete) const override;
diff --git a/Engine/Plugins/FX/Niagara/Source/NiagaraEditor/Public/Widgets/SNiagaraParameterMenu.cpp b/Engine/Plugins/FX/Niagara/Source/NiagaraEditor/Public/Widgets/SNiagaraParameterMenu.cpp
--- a/Engine/Plugins/FX/Niagara/Source/NiagaraEditor/Public/Widgets/SNiagaraParameterMenu.cpp
+++ b/Engine/Plugins/FX/Niagara/Source/NiagaraEditor/Public/Widgets/SNiagaraParameterMenu.cpp
@@ -122,7 +122,8 @@ FText SNiagaraParameterMenu::GetNamespaceCategoryText(const FGuid& NamespaceId)
 
 void SNiagaraAddParameterFromPanelMenu::Construct(const FArguments& InArgs)
 {
-	this->OnAddParameter = InArgs._OnAddParameter;
+	this->OnSpecificParameterRequested = InArgs._OnSpecificParameterRequested;
+	this->OnNewParameterRequested = InArgs._OnNewParameterRequested;
 	this->OnAddScriptVar = InArgs._OnAddScriptVar;
 	this->OnAllowMakeType = InArgs._OnAllowMakeType;
 	this->OnAddParameterDefinitions = InArgs._OnAddParameterDefinitions;
@@ -153,8 +154,7 @@ void SNiagaraAddParameterFromPanelMenu::CollectMakeNew(FNiagaraMenuActionCollect
 	{
 		return;
 	}
-
-	TArray<FNiagaraVariable> Variables;
+	
 	TConstArrayView<FNiagaraTypeDefinition> SectionTypes;
 	if(InNamespaceId == FNiagaraEditorGuids::UserNamespaceMetaDataGuid) 
 	{
@@ -177,6 +177,7 @@ void SNiagaraAddParameterFromPanelMenu::CollectMakeNew(FNiagaraMenuActionCollect
 		SectionTypes = MakeArrayView(FNiagaraTypeRegistry::GetRegisteredTypes());
 	}
 
+	TArray<FNiagaraTypeDefinition> TypeDefinitions;
 	for (const FNiagaraTypeDefinition& RegisteredType : SectionTypes)
 	{
 		bool bAllowType = true;
@@ -187,13 +188,11 @@ void SNiagaraAddParameterFromPanelMenu::CollectMakeNew(FNiagaraMenuActionCollect
 
 		if (bAllowType)
 		{
-			FNiagaraVariable Var(RegisteredType, FName(*RegisteredType.GetNameText().ToString()));
-			FNiagaraEditorUtilities::ResetVariableToDefaultValue(Var);
-			Variables.Add(Var);
+			TypeDefinitions.Add(RegisteredType);
 		}
 	}
 
-	AddParameterGroup(Collector, Variables, InNamespaceId,
+	AddMakeNewGroup(Collector, TypeDefinitions,  InNamespaceId,
 		LOCTEXT("MakeNewCat", "Make New"), 1,
 		bShowNamespaceCategory ? GetNamespaceCategoryText(InNamespaceId).ToString() : FString());
 }
@@ -204,8 +203,7 @@ void SNiagaraAddParameterFromPanelMenu::AddParameterGroup(
 	const FGuid& InNamespaceId /*= FGuid()*/,
 	const FText& Category /*= FText::GetEmpty()*/,
 	int32 SortOrder /*= 0*/,
-	const FString& RootCategory /*= FString()*/,
-	const bool bCreateUniqueName /*= true*/)
+	const FString& RootCategory /*= FString()*/)
 {
 	for (const FNiagaraVariable& Variable : Variables)
 	{
@@ -228,7 +226,7 @@ void SNiagaraAddParameterFromPanelMenu::AddParameterGroup(
 		FText SubCategory = FNiagaraEditorUtilities::GetVariableTypeCategory(Variable);
 		FText FullCategory = SubCategory.IsEmpty() ? Category : FText::Format(FText::FromString("{0}|{1}"), Category, SubCategory);
 		TSharedPtr<FNiagaraMenuAction> Action(new FNiagaraMenuAction(FullCategory, DisplayName, Tooltip, 0, FText(),
-			FNiagaraMenuAction::FOnExecuteStackAction::CreateSP(this, &SNiagaraAddParameterFromPanelMenu::ParameterSelected, Variable, bCreateUniqueName, InNamespaceId)));
+			FNiagaraMenuAction::FOnExecuteStackAction::CreateSP(this, &SNiagaraAddParameterFromPanelMenu::ParameterSelected, Variable, InNamespaceId)));
 		Action->SetParameterVariable(Variable);
 
 		if (Variable.IsDataInterface())
@@ -243,6 +241,41 @@ void SNiagaraAddParameterFromPanelMenu::AddParameterGroup(
 	}
 }
 
+void SNiagaraAddParameterFromPanelMenu::AddMakeNewGroup(FNiagaraMenuActionCollector& Collector,
+	TArray<FNiagaraTypeDefinition>& TypeDefinitions, const FGuid& InNamespaceId, const FText& Category, int32 SortOrder,
+	const FString& RootCategory)
+{
+	for (const FNiagaraTypeDefinition& TypeDefinition : TypeDefinitions)
+	{
+		const FText DisplayName = TypeDefinition.GetNameText();
+		FText Tooltip = FText::GetEmpty();
+
+
+		if (const UStruct* VariableStruct = TypeDefinition.GetStruct())
+		{
+			Tooltip = VariableStruct->GetToolTipText(true);
+		}
+
+		FText SubCategory = FNiagaraEditorUtilities::GetTypeDefinitionCategory(TypeDefinition);
+		FText FullCategory = SubCategory.IsEmpty() ? Category : FText::Format(FText::FromString("{0}|{1}"), Category, SubCategory);
+		TSharedPtr<FNiagaraMenuAction> Action(new FNiagaraMenuAction(FullCategory, DisplayName, Tooltip, 0, FText(),
+			FNiagaraMenuAction::FOnExecuteStackAction::CreateSP(this, &SNiagaraAddParameterFromPanelMenu::NewParameterSelected, TypeDefinition, InNamespaceId)));
+		FNiagaraVariable PrototypeVariable(TypeDefinition, TypeDefinition.GetFName());
+
+		Action->SetParameterVariable(PrototypeVariable);
+
+		if (TypeDefinition.IsDataInterface())
+		{
+			if (const UClass* DataInterfaceClass = TypeDefinition.GetClass())
+			{
+				Action->IsExperimental = DataInterfaceClass->GetMetaData("DevelopmentStatus") == TEXT("Experimental");
+			}
+		}
+
+		Collector.AddAction(Action, SortOrder, RootCategory);
+	}
+}
+
 void SNiagaraAddParameterFromPanelMenu::CollectParameterCollectionsActions(FNiagaraMenuActionCollector& Collector)
 {
 	//Create sub menus for parameter collections.
@@ -256,7 +289,7 @@ void SNiagaraAddParameterFromPanelMenu::CollectParameterCollectionsActions(FNiag
 		UNiagaraParameterCollection* Collection = CastChecked<UNiagaraParameterCollection>(CollectionAsset.GetAsset());
 		if (Collection)
 		{
-			AddParameterGroup(Collector, Collection->GetParameters(), FNiagaraEditorGuids::ParameterCollectionNamespaceMetaDataGuid, Category, 10, FString(), false); //TODO
+			AddParameterGroup(Collector, Collection->GetParameters(), FNiagaraEditorGuids::ParameterCollectionNamespaceMetaDataGuid, Category, 10, FString()); //TODO
 		}
 	}
 }
@@ -269,15 +302,13 @@ void SNiagaraAddParameterFromPanelMenu::CollectAllActions(FGraphActionListBuilde
 		TArray<FNiagaraVariable> Variables = FNiagaraConstants::GetEngineConstants();
 		const FText CategoryText = bShowNamespaceCategory ? GetNamespaceCategoryText(NamespaceMetaData) : LOCTEXT("EngineConstantNamespaceCategory", "Add Engine Constant");
 		const FString RootCategoryStr = FString();
-		const bool bMakeNameUnique = false;
 		AddParameterGroup(
 			Collector,
 			Variables,
 			FNiagaraEditorGuids::EngineNamespaceMetaDataGuid,
 			CategoryText,
 			4,
-			RootCategoryStr,
-			bMakeNameUnique
+			RootCategoryStr
 		);
 	};
 	auto CollectEmitterNamespaceParameterActions = [this, &Collector]() {
@@ -285,15 +316,13 @@ void SNiagaraAddParameterFromPanelMenu::CollectAllActions(FGraphActionListBuilde
 		TArray<FNiagaraVariable> Variables = FNiagaraConstants::GetEngineConstants().FilterByPredicate([](const FNiagaraVariable& Var) { return Var.IsInNameSpace(FNiagaraConstants::EmitterNamespaceString); });
 		const FText CategoryText = bShowNamespaceCategory ? GetNamespaceCategoryText(NamespaceMetaData) : LOCTEXT("EmitterConstantNamespaceCategory", "Add Emitter Constant");
 		const FString RootCategoryStr = FString();
-		const bool bMakeNameUnique = false;
 		AddParameterGroup(
 			Collector,
 			Variables,
 			FNiagaraEditorGuids::EngineNamespaceMetaDataGuid,
 			CategoryText,
 			4,
-			RootCategoryStr,
-			bMakeNameUnique
+			RootCategoryStr
 		);
 	};
 
@@ -381,7 +410,7 @@ void SNiagaraAddParameterFromPanelMenu::CollectAllActions(FGraphActionListBuilde
 	{
 		TArray<FNiagaraVariable> Variables;
 		Variables.Add(SYS_PARAM_INSTANCE_ALIVE);
-		AddParameterGroup(Collector, Variables, FNiagaraEditorGuids::DataInstanceNamespaceMetaDataGuid, FText(), 3, FString(), false);
+		AddParameterGroup(Collector, Variables, FNiagaraEditorGuids::DataInstanceNamespaceMetaDataGuid, FText(), 3, FString());
 	}
 	// No NamespaceId set but still collecting engine namespace parameters (e.g. map get/set node menu.)
 	else if (NamespaceId.IsValid() == false && bForceCollectEngineNamespaceParameterActions)
@@ -392,7 +421,7 @@ void SNiagaraAddParameterFromPanelMenu::CollectAllActions(FGraphActionListBuilde
 		// Special case; collect DataInstance.Alive so that it is an option if we are selecting a parameter from a map node in a script.
 		TArray<FNiagaraVariable> Variables;
 		Variables.Add(SYS_PARAM_INSTANCE_ALIVE);
-		AddParameterGroup(Collector, Variables, FNiagaraEditorGuids::DataInstanceNamespaceMetaDataGuid, FText(), 3, FString(), false);
+		AddParameterGroup(Collector, Variables, FNiagaraEditorGuids::DataInstanceNamespaceMetaDataGuid, FText(), 3, FString());
 	}
 
 	// Any other "unreserved" namespace
@@ -677,9 +706,15 @@ void SNiagaraAddParameterFromPanelMenu::CollectAllActions(FGraphActionListBuilde
 	Collector.AddAllActionsTo(OutAllActions);
 }
 
-void SNiagaraAddParameterFromPanelMenu::ParameterSelected(FNiagaraVariable NewVariable, const bool bCreateUniqueName /*= false*/, const FGuid InNamespaceId /*= FGuid()*/)
+void SNiagaraAddParameterFromPanelMenu::ParameterSelected(FNiagaraVariable NewVariable, const FGuid InNamespaceId /*= FGuid()*/)
+{	
+	OnSpecificParameterRequested.ExecuteIfBound(NewVariable);
+}
+
+void SNiagaraAddParameterFromPanelMenu::NewParameterSelected(FNiagaraTypeDefinition NewParameterType, const FGuid InNamespaceId)
 {
-	auto GetNamespaceMetaData = [this, &InNamespaceId]()->const FNiagaraNamespaceMetadata {
+	auto GetNamespaceMetaData = [this, &InNamespaceId]()->const FNiagaraNamespaceMetadata
+	{
 		if (InNamespaceId.IsValid() == false)
 		{
 			if (bIsParameterReadNode) // Map Get
@@ -693,51 +728,49 @@ void SNiagaraAddParameterFromPanelMenu::ParameterSelected(FNiagaraVariable NewVa
 		}
 		return FNiagaraEditorUtilities::GetNamespaceMetaDataForId(InNamespaceId);
 	};
-
-	if (bCreateUniqueName)
+	
+	FString TypeDisplayName;
+	if (NewParameterType.GetEnum() != nullptr)
 	{
-		FString TypeDisplayName;
-		if (NewVariable.GetType().GetEnum() != nullptr)
-		{
-			TypeDisplayName = ((UField*)NewVariable.GetType().GetEnum())->GetDisplayNameText().ToString();
-		}
-		else if (NewVariable.GetType().GetStruct() != nullptr)
-		{
-			TypeDisplayName = NewVariable.GetType().GetNameText().ToString();
-		}
-		else if (NewVariable.GetType().GetClass() != nullptr)
-		{
-			TypeDisplayName = NewVariable.GetType().GetClass()->GetDisplayNameText().ToString();
-		}
-		const FString NewVariableDefaultName = TypeDisplayName.IsEmpty()
-			? TEXT("New Variable")
-			: TEXT("New ") + TypeDisplayName;
+		TypeDisplayName = ((UField*)NewParameterType.GetEnum())->GetDisplayNameText().ToString();
+	}
+	else if (NewParameterType.GetStruct() != nullptr)
+	{
+		TypeDisplayName = NewParameterType.GetNameText().ToString();
+	}
+	else if (NewParameterType.GetClass() != nullptr)
+	{
+		TypeDisplayName = NewParameterType.GetClass()->GetDisplayNameText().ToString();
+	}
 
-		TArray<FString> NameParts;
+	const FString NewVariableDefaultName = TypeDisplayName.IsEmpty()
+		? TEXT("New Variable")
+		: TEXT("New ") + TypeDisplayName;
 
-		const FNiagaraNamespaceMetadata NamespaceMetaData = GetNamespaceMetaData();
-		checkf(NamespaceMetaData.IsValid(), TEXT("Failed to get valid namespace metadata when creating unique name for parameter menu add parameter action!"));
-		for (const FName Namespace : NamespaceMetaData.Namespaces)
-		{
-			NameParts.Add(Namespace.ToString());
-		}
+	TArray<FString> NameParts;
 
-		if (NamespaceMetaData.RequiredNamespaceModifier != NAME_None)
-		{
-			NameParts.Add(NamespaceMetaData.RequiredNamespaceModifier.ToString());
-		}
+	const FNiagaraNamespaceMetadata NamespaceMetaData = GetNamespaceMetaData();
+	checkf(NamespaceMetaData.IsValid(), TEXT("Failed to get valid namespace metadata when creating unique name for parameter menu add parameter action!"));
+	for (const FName Namespace : NamespaceMetaData.Namespaces)
+	{
+		NameParts.Add(Namespace.ToString());
+	}
 
-		NameParts.Add(NewVariableDefaultName);
-		const FString ResultName = FString::Join(NameParts, TEXT("."));
-		NewVariable.SetName(FName(*ResultName));
+	if (NamespaceMetaData.RequiredNamespaceModifier != NAME_None)
+	{
+		NameParts.Add(NamespaceMetaData.RequiredNamespaceModifier.ToString());
 	}
 
-	OnAddParameter.ExecuteIfBound(NewVariable);
+	NameParts.Add(NewVariableDefaultName);
+	const FString ResultName = FString::Join(NameParts, TEXT("."));
+
+	FNiagaraVariable NewVariable(NewParameterType, FName(ResultName));
+	OnNewParameterRequested.ExecuteIfBound(NewVariable);
 }
 
 void SNiagaraAddParameterFromPanelMenu::ParameterSelected(FNiagaraVariable NewVariable)
 {
-	ParameterSelected(NewVariable, false, FGuid());
+	ParameterSelected(NewVariable, FGuid());
 }
 
 void SNiagaraAddParameterFromPanelMenu::ScriptVarFromParameterDefinitionsSelected(const UNiagaraScriptVariable* NewScriptVar, UNiagaraParameterDefinitions* SourceParameterDefinitions)
diff --git a/Engine/Plugins/FX/Niagara/Source/NiagaraEditor/Public/Widgets/SNiagaraParameterMenu.h b/Engine/Plugins/FX/Niagara/Source/NiagaraEditor/Public/Widgets/SNiagaraParameterMenu.h
--- a/Engine/Plugins/FX/Niagara/Source/NiagaraEditor/Public/Widgets/SNiagaraParameterMenu.h
+++ b/Engine/Plugins/FX/Niagara/Source/NiagaraEditor/Public/Widgets/SNiagaraParameterMenu.h
@@ -58,7 +58,7 @@ private:
 class SNiagaraAddParameterFromPanelMenu : public SNiagaraParameterMenu
 {
 public:
-	DECLARE_DELEGATE_OneParam(FOnAddParameter, FNiagaraVariable);
+	DECLARE_DELEGATE_OneParam(FOnParameterRequested, FNiagaraVariable);
 	DECLARE_DELEGATE_OneParam(FOnAddScriptVar, const UNiagaraScriptVariable*);
 	DECLARE_DELEGATE_OneParam(FOnAddParameterDefinitions, UNiagaraParameterDefinitions*);
 	DECLARE_DELEGATE_RetVal_OneParam(bool, FOnAllowMakeType, const FNiagaraTypeDefinition&);
@@ -76,7 +76,8 @@ public:
 		, _AssignmentNode(nullptr)
 	{}
 		//~ Begin Required Events
-		SLATE_EVENT(FOnAddParameter, OnAddParameter)
+		SLATE_EVENT(FOnParameterRequested, OnSpecificParameterRequested)
+		SLATE_EVENT(FOnParameterRequested, OnNewParameterRequested)
 		SLATE_EVENT(FOnAddScriptVar, OnAddScriptVar)
 		SLATE_EVENT(FOnAddParameterDefinitions, OnAddParameterDefinitions)
 		//~ End Required Events
@@ -106,8 +107,15 @@ public:
 		const FGuid& InNamespaceId = FGuid(),
 		const FText& Category = FText::GetEmpty(),
 		int32 SortOrder = 0,
-		const FString& RootCategory = FString(),
-		const bool bCreateUniqueName = true);
+		const FString& RootCategory = FString());
+
+	void AddMakeNewGroup(
+		FNiagaraMenuActionCollector& Collector,
+		TArray<FNiagaraTypeDefinition>& TypeDefinitions,
+		const FGuid& InNamespaceId = FGuid(),
+		const FText& Category = FText::GetEmpty(),
+		int32 SortOrder = 0,
+		const FString& RootCategory = FString());
 
 	void CollectParameterCollectionsActions(FNiagaraMenuActionCollector& Collector);
 	void CollectMakeNew(FNiagaraMenuActionCollector& Collector, const FGuid& InNamespaceId);
@@ -117,15 +125,20 @@ protected:
 
 private:
 	void ParameterSelected(FNiagaraVariable NewVariable);
-	void ParameterSelected(FNiagaraVariable NewVariable, const bool bCreateUniqueName, const FGuid InNamespaceId);
+	void ParameterSelected(FNiagaraVariable NewVariable, const FGuid InNamespaceId);
+	void NewParameterSelected(FNiagaraTypeDefinition NewParameterType, const FGuid InNamespaceId);
+	
 	void ScriptVarFromParameterDefinitionsSelected(const UNiagaraScriptVariable* NewScriptVar, UNiagaraParameterDefinitions* SourceParameterDefinitions);
 	TSet<FName> GetAllGraphParameterNames() const;
 
 	static FText GetSectionTitle(int32 SectionId);
 
 private:
+	/** Delegate to handle adding an existing parameter as an FNiagaraVariable. */
+	FOnParameterRequested OnSpecificParameterRequested;
+
 	/** Delegate to handle adding a parameter as an FNiagaraVariable. */
-	FOnAddParameter OnAddParameter;
+	FOnParameterRequested OnNewParameterRequested;
 
 	/** Delegate to handle adding a parameter as a UNiagaraScriptVariable. */
 	FOnAddScriptVar OnAddScriptVar;
