Fixing memstomp in HLSL translator for custom hlsl node

diff --git a/Engine/Plugins/FX/Niagara/Source/NiagaraCore/Private/NiagaraCustomVersion.cpp b/Engine/Plugins/FX/Niagara/Source/NiagaraCore/Private/NiagaraCustomVersion.cpp
--- a/Engine/Plugins/FX/Niagara/Source/NiagaraCore/Private/NiagaraCustomVersion.cpp
+++ b/Engine/Plugins/FX/Niagara/Source/NiagaraCore/Private/NiagaraCustomVersion.cpp
@@ -9,4 +9,4 @@ const FGuid FNiagaraCustomVersion::GUID(0xFCF57AFA, 0x50764283, 0xB9A9E658, 0xFF
 FCustomVersionRegistration GRegisterNiagaraCustomVersion(FNiagaraCustomVersion::GUID, FNiagaraCustomVersion::LatestVersion, TEXT("NiagaraVer"));
 
 // Note: When encountering a conflict on this file please generate a new GUID
-const FGuid FNiagaraCustomVersion::LatestScriptCompileVersion(0xDED3AAF7, 0xE9A50E4C, 0x9E7C429C, 0x94EB9117);
\ No newline at end of file
+const FGuid FNiagaraCustomVersion::LatestScriptCompileVersion(0x98862397, 0xB732490C, 0x9D9C8B1A, 0x6A4DA05F);
diff --git a/Engine/Plugins/FX/Niagara/Source/NiagaraEditor/Private/NiagaraHlslTranslator.cpp b/Engine/Plugins/FX/Niagara/Source/NiagaraEditor/Private/NiagaraHlslTranslator.cpp
--- a/Engine/Plugins/FX/Niagara/Source/NiagaraEditor/Private/NiagaraHlslTranslator.cpp
+++ b/Engine/Plugins/FX/Niagara/Source/NiagaraEditor/Private/NiagaraHlslTranslator.cpp
@@ -7277,8 +7277,8 @@ void FHlslNiagaraTranslator::ProcessCustomHlsl(const FString& InCustomHlsl, ENia
 		{
 			if (Token.StartsWith(ViewNamespace, ESearchCase::CaseSensitive))
 			{
-				FStringView TokenMemberName = FStringView(Token).Mid(ViewNamespace.Len());
-				FStringView TokenPostfix;
+				FString TokenMemberName = FString(Token).Mid(ViewNamespace.Len());
+				FString TokenPostfix;
 
 				int32 MemberEnd = INDEX_NONE;
 				if (TokenMemberName.FindChar(TCHAR('.'), MemberEnd))
