Fixing memstomp in HLSL translator for custom hlsl node

diff --git a/Engine/Plugins/FX/Niagara/Source/NiagaraCore/Private/NiagaraCustomVersion.cpp b/Engine/Plugins/FX/Niagara/Source/NiagaraCore/Private/NiagaraCustomVersion.cpp
--- a/Engine/Plugins/FX/Niagara/Source/NiagaraCore/Private/NiagaraCustomVersion.cpp
+++ b/Engine/Plugins/FX/Niagara/Source/NiagaraCore/Private/NiagaraCustomVersion.cpp
@@ -9,4 +9,4 @@ const FGuid FNiagaraCustomVersion::GUID(0xFCF57AFA, 0x50764283, 0xB9A9E658, 0xFF
 FCustomVersionRegistration GRegisterNiagaraCustomVersion(FNiagaraCustomVersion::GUID, FNiagaraCustomVersion::LatestVersion, TEXT("NiagaraVer"));
 
 // Note: When encountering a conflict on this file please generate a new GUID
-const FGuid FNiagaraCustomVersion::LatestScriptCompileVersion(0x98862397, 0xB732490C, 0x9D9C8B1A, 0x6A4DA05F);
+const FGuid FNiagaraCustomVersion::LatestScriptCompileVersion(0xBF690949, 0xE9C446E2, 0xB34FEB14, 0x3115170D);
