- Overhauled pin connection drag & drop information. Potentially lossy graph connections, if assignable, will give the user feedback.
- Also reworked the drag & drop for parameter actions in systems. The drag & drop widget looks cleaner and gives you more info when hovering over a function input value
- Introduced an SNiagaraParameterDropTarget widget to replace the SDropTarget used in the StachFunctionInputValue to support this.

diff --git a/Engine/Plugins/FX/Niagara/Source/Niagara/Private/NiagaraModule.cpp b/Engine/Plugins/FX/Niagara/Source/Niagara/Private/NiagaraModule.cpp
--- a/Engine/Plugins/FX/Niagara/Source/Niagara/Private/NiagaraModule.cpp
+++ b/Engine/Plugins/FX/Niagara/Source/Niagara/Private/NiagaraModule.cpp
@@ -1101,7 +1101,7 @@ bool FNiagaraTypeDefinition::IsScalarDefinition(const FNiagaraTypeDefinition& Ty
 	return ScalarStructs.Contains(Type.GetScriptStruct()) || (Type.GetScriptStruct() == IntStruct && Type.GetEnum() != nullptr);
 }
 
-bool FNiagaraTypeDefinition::TypesAreAssignable(const FNiagaraTypeDefinition& TypeInput, const FNiagaraTypeDefinition& TypeOutput)
+bool FNiagaraTypeDefinition::TypesAreAssignable(const FNiagaraTypeDefinition& TypeInput, const FNiagaraTypeDefinition& TypeOutput, bool bAllowLossyLWCConversions)
 {
 	if (const UClass* AClass = TypeInput.GetClass())
 	{
@@ -1152,7 +1152,9 @@ bool FNiagaraTypeDefinition::TypesAreAssignable(const FNiagaraTypeDefinition& Ty
 	}
 	else
 	{
-		bIsSupportedConversion = (TypeInput == ColorDef && TypeOutput == Vec4Def) || (TypeOutput == ColorDef && TypeInput == Vec4Def);
+		bIsSupportedConversion = (TypeInput == ColorDef && TypeOutput == Vec4Def) || (TypeOutput == ColorDef && TypeInput == Vec4Def)
+			// we only allow direct assignments between vec <-> position if we are allowing lossy conversions
+			|| (bAllowLossyLWCConversions && ((TypeOutput == PositionDef && TypeInput == Vec3Def) || (TypeOutput == Vec3Def && TypeInput == PositionDef)));
 	}
 
 	if (bIsSupportedConversion)
@@ -1167,12 +1169,15 @@ bool FNiagaraTypeDefinition::TypesAreAssignable(const FNiagaraTypeDefinition& Ty
 		(TypeOutput == NumericDef && (TypeInput.GetStruct() == GetIntStruct()) && TypeInput.GetEnum() != nullptr);
 }
 
-bool FNiagaraTypeDefinition::IsLossyConversion(const FNiagaraTypeDefinition& TypeA, const FNiagaraTypeDefinition& TypeB)
+bool FNiagaraTypeDefinition::IsLossyConversion(const FNiagaraTypeDefinition& FromType, const FNiagaraTypeDefinition& ToType)
 {
-	return (TypeA.IsSameBaseDefinition(IntDef) && TypeB.IsSameBaseDefinition(FloatDef))
-		|| (TypeB.IsSameBaseDefinition(IntDef) && TypeA.IsSameBaseDefinition(FloatDef))
-		|| (TypeA.IsSameBaseDefinition(BoolDef) && TypeB.IsSameBaseDefinition(FloatDef))
-		|| (TypeB.IsSameBaseDefinition(BoolDef) && TypeA.IsSameBaseDefinition(FloatDef));
+	// @todo this should be unidirectional, same with TypesAreAssignable
+	return (FromType.IsSameBaseDefinition(IntDef) && ToType.IsSameBaseDefinition(FloatDef))
+		|| (ToType.IsSameBaseDefinition(IntDef) && FromType.IsSameBaseDefinition(FloatDef))
+		|| (FromType.IsSameBaseDefinition(BoolDef) && ToType.IsSameBaseDefinition(FloatDef))
+		|| (ToType.IsSameBaseDefinition(BoolDef) && FromType.IsSameBaseDefinition(FloatDef))
+		|| (FromType == PositionDef && ToType == Vec3Def)
+		|| (FromType == Vec3Def && ToType == PositionDef);
 }
 
 FNiagaraTypeDefinition FNiagaraTypeDefinition::GetNumericOutputType(const TArray<FNiagaraTypeDefinition> TypeDefinintions, ENiagaraNumericOutputTypeSelectionMode SelectionMode)
diff --git a/Engine/Plugins/FX/Niagara/Source/Niagara/Public/NiagaraTypes.h b/Engine/Plugins/FX/Niagara/Source/Niagara/Public/NiagaraTypes.h
--- a/Engine/Plugins/FX/Niagara/Source/Niagara/Public/NiagaraTypes.h
+++ b/Engine/Plugins/FX/Niagara/Source/Niagara/Public/NiagaraTypes.h
@@ -1211,9 +1211,9 @@ public:
 		}
 		return FNiagaraTypeHelper::ToString(ValueData, ClassStructOrEnum);
 	}
-
-	static bool TypesAreAssignable(const FNiagaraTypeDefinition& TypeA, const FNiagaraTypeDefinition& TypeB);
-	static bool IsLossyConversion(const FNiagaraTypeDefinition& TypeA, const FNiagaraTypeDefinition& TypeB);
+	
+	static bool TypesAreAssignable(const FNiagaraTypeDefinition& TypeA, const FNiagaraTypeDefinition& TypeB, bool bAllowLossyLWCConversions = false);
+	static bool IsLossyConversion(const FNiagaraTypeDefinition& FromType, const FNiagaraTypeDefinition& ToType);
 	static FNiagaraTypeDefinition GetNumericOutputType(const TArray<FNiagaraTypeDefinition> TypeDefinintions, ENiagaraNumericOutputTypeSelectionMode SelectionMode);
 
 	static const TArray<FNiagaraTypeDefinition>& GetNumericTypes() { return OrderedNumericTypes; }
diff --git a/Engine/Plugins/FX/Niagara/Source/NiagaraEditor/Private/EdGraphSchema_Niagara.cpp b/Engine/Plugins/FX/Niagara/Source/NiagaraEditor/Private/EdGraphSchema_Niagara.cpp
--- a/Engine/Plugins/FX/Niagara/Source/NiagaraEditor/Private/EdGraphSchema_Niagara.cpp
+++ b/Engine/Plugins/FX/Niagara/Source/NiagaraEditor/Private/EdGraphSchema_Niagara.cpp
@@ -1010,25 +1010,29 @@ const FPinConnectionResponse UEdGraphSchema_Niagara::CanCreateConnection(const U
 			{
 				return FPinConnectionResponse(CONNECT_RESPONSE_DISALLOW, TEXT("Types are not compatible"));
 			}
-
-			else if (FNiagaraTypeDefinition::TypesAreAssignable(PinTypeInput, PinTypeOutput) == false)
+			else if (FNiagaraTypeDefinition::TypesAreAssignable(PinTypeInput, PinTypeOutput, GetDefault<UNiagaraSettings>()->bEnforceStrictStackTypes == false))
+			{
+				// if we have a lossy conversion but are assignable, let the user connect but inform him
+				if(FNiagaraTypeDefinition::IsLossyConversion(PinTypeOutput, PinTypeInput))
+				{
+					return FPinConnectionResponse(CONNECT_RESPONSE_MAKE, FString::Printf(TEXT("Convert %s to %s implicitly. Be aware that this can be a lossy conversion."), *(PinTypeOutput.GetNameText().ToString()), *(PinTypeInput.GetNameText().ToString())));
+				}
+				// if we are assignable and not lossy, we simply allow the connection
+				else
+				{
+					return FPinConnectionResponse(CONNECT_RESPONSE_MAKE, FString());
+				}
+			}
+			// if the types are not directly assignable at all, give the user a chance to use a conversion node
+			else
 			{
 				//Do some limiting on auto conversions here?
 				if (PinTypeInput.GetClass() || PinTypeInput.IsStatic() || PinTypeOutput.IsStatic())
 				{
 					return FPinConnectionResponse(CONNECT_RESPONSE_DISALLOW, TEXT("Types are not compatible"));
 				}
-				else if (FNiagaraTypeDefinition::IsLossyConversion(PinTypeInput, PinTypeOutput))
-				{
-					return FPinConnectionResponse(CONNECT_RESPONSE_MAKE_WITH_CONVERSION_NODE, FString::Printf(TEXT("Convert %s to %s. Be aware that this can be a lossy conversion."), *(PinTypeOutput.GetNameText().ToString()), *(PinTypeInput.GetNameText().ToString())));
-				}
 				else
 				{
-					const UNiagaraSettings* Settings = GetDefault<UNiagaraSettings>();
-					if (Settings->bEnforceStrictStackTypes == false && ((PinTypeInput == FNiagaraTypeDefinition::GetPositionDef() && PinTypeOutput == FNiagaraTypeDefinition::GetVec3Def()) || (PinTypeOutput == FNiagaraTypeDefinition::GetPositionDef() && PinTypeInput == FNiagaraTypeDefinition::GetVec3Def())))
-					{
-						return FPinConnectionResponse(CONNECT_RESPONSE_MAKE, FString());
-					}
 					return FPinConnectionResponse(CONNECT_RESPONSE_MAKE_WITH_CONVERSION_NODE, FString::Printf(TEXT("Convert %s to %s"), *(PinTypeOutput.GetNameText().ToString()), *(PinTypeInput.GetNameText().ToString())));
 				}
 			}
diff --git a/Engine/Plugins/FX/Niagara/Source/NiagaraEditor/Private/NiagaraActions.cpp b/Engine/Plugins/FX/Niagara/Source/NiagaraEditor/Private/NiagaraActions.cpp
--- a/Engine/Plugins/FX/Niagara/Source/NiagaraEditor/Private/NiagaraActions.cpp
+++ b/Engine/Plugins/FX/Niagara/Source/NiagaraEditor/Private/NiagaraActions.cpp
@@ -5,9 +5,11 @@
 #include "NiagaraNodeParameterMapGet.h"
 #include "NiagaraNodeParameterMapSet.h"
 #include "EdGraphSchema_Niagara.h"
+#include "NiagaraEditorUtilities.h"
 #include "NiagaraNodeStaticSwitch.h"
 #include "NiagaraScriptVariable.h"
 #include "Widgets/SWidget.h"
+#include "Widgets/Layout/SBox.h"
 #include "Framework/MultiBox/MultiBoxBuilder.h"
 #include "Framework/Application/MenuStack.h"
 #include "Framework/Application/SlateApplication.h"
@@ -16,6 +18,7 @@
 #include "Classes/EditorStyleSettings.h"
 #include "ViewModels/NiagaraParameterPanelViewModel.h"
 #include "Misc/MessageDialog.h"
+#include "Widgets/SNiagaraParameterName.h"
 
 #define LOCTEXT_NAMESPACE "NiagaraActions"
 
@@ -617,4 +620,39 @@ EVisibility FNiagaraParameterGraphDragOperation::GetErrorIconVisible() const
 	return EVisibility::Collapsed;
 }
 
+TSharedPtr<SWidget> FNiagaraParameterDragOperation::GetDefaultDecorator() const
+{
+	TSharedPtr<SWidget> Decorator = SNew(SToolTip)
+	.Content()
+	[
+		SNew(SVerticalBox)
+		+ SVerticalBox::Slot()
+		.AutoHeight()
+		[
+			FNiagaraParameterUtilities::GetParameterWidget(StaticCastSharedPtr<FNiagaraParameterAction>(SourceAction)->GetParameter(), false)
+		]
+		+ SVerticalBox::Slot()
+		.AutoHeight()
+		[
+			SNew(SBox)
+			.MaxDesiredWidth(250.f)
+			.Padding(5.f)
+			[
+				SNew(STextBlock)
+				.ColorAndOpacity(FLinearColor::White)
+				.Text(this, &FNiagaraParameterDragOperation::GetHoverText)
+				.Visibility(this, &FNiagaraParameterDragOperation::IsTextVisible)
+				.AutoWrapText(true)
+			]
+		]
+	];
+	
+	return Decorator;
+}
+
+EVisibility FNiagaraParameterDragOperation::IsTextVisible() const
+{
+	return CurrentHoverText.IsEmpty() ? EVisibility::Collapsed : EVisibility::Visible;
+}
+
 #undef LOCTEXT_NAMESPACE
diff --git a/Engine/Plugins/FX/Niagara/Source/NiagaraEditor/Private/NiagaraEditorUtilities.cpp b/Engine/Plugins/FX/Niagara/Source/NiagaraEditor/Private/NiagaraEditorUtilities.cpp
--- a/Engine/Plugins/FX/Niagara/Source/NiagaraEditor/Private/NiagaraEditorUtilities.cpp
+++ b/Engine/Plugins/FX/Niagara/Source/NiagaraEditor/Private/NiagaraEditorUtilities.cpp
@@ -3530,7 +3530,7 @@ bool FNiagaraParameterUtilities::TestCanRenameWithMessage(FName ParameterName, F
 	return false;
 }
 
-TSharedRef<SToolTip> FNiagaraParameterUtilities::GetTooltipWidget(FNiagaraVariable Variable, bool bShowValue)
+TSharedRef<SWidget> FNiagaraParameterUtilities::GetParameterWidget(FNiagaraVariable Variable, bool bShowValue)
 {
 	FNiagaraTypeDefinition Type = Variable.GetType();
 	const FLinearColor TypeColor = UEdGraphSchema_Niagara::GetTypeColor(Type);
@@ -3543,12 +3543,12 @@ TSharedRef<SToolTip> FNiagaraParameterUtilities::GetTooltipWidget(FNiagaraVariab
 	FSlateBrush const* SecondaryIconBrush = FEditorStyle::GetBrush(TEXT("NoBrush"));
 	FSlateColor        SecondaryIconColor = IconColor;
 
-
-	TSharedRef<SVerticalBox> TooltipContent = SNew(SVerticalBox)
+	TSharedPtr<SHorizontalBox> ParameterContainer;
+	TSharedRef<SVerticalBox> ParameterWidget = SNew(SVerticalBox)
 		+ SVerticalBox::Slot()
 		.AutoHeight()
 		[
-			SNew(SHorizontalBox)
+			SAssignNew(ParameterContainer, SHorizontalBox)
 			+ SHorizontalBox::Slot()
 			.AutoWidth()
 			.HAlign(HAlign_Center)
@@ -3579,7 +3579,7 @@ TSharedRef<SToolTip> FNiagaraParameterUtilities::GetTooltipWidget(FNiagaraVariab
 	{
 		const FText ValueText = TypeUtilities->GetStackDisplayText(Variable);
 
-		TooltipContent->AddSlot()
+		ParameterWidget->AddSlot()
 		.AutoHeight()
 		[
 			SNew(SHorizontalBox)
@@ -3603,15 +3603,58 @@ TSharedRef<SToolTip> FNiagaraParameterUtilities::GetTooltipWidget(FNiagaraVariab
 			]
 		];
 	}
-	
-	
+
+	return ParameterWidget;
+}
+
+TSharedRef<SToolTip> FNiagaraParameterUtilities::GetTooltipWidget(FNiagaraVariable Variable, bool bShowValue, TSharedPtr<SWidget> AdditionalVerticalWidget,  TSharedPtr<SWidget> AdditionalHorizontalWidget)
+{
+	FNiagaraTypeDefinition Type = Variable.GetType();
+	const FLinearColor TypeColor = UEdGraphSchema_Niagara::GetTypeColor(Type);
+	TSharedPtr<INiagaraEditorTypeUtilities> TypeUtilities = FNiagaraEditorModule::Get().GetTypeUtilities(Type);
+
+	FText			   IconToolTip = FText::GetEmpty();
+	FSlateBrush const* IconBrush = FAppStyle::Get().GetBrush(TEXT("Kismet.AllClasses.VariableIcon"));
+	FSlateColor        IconColor = FSlateColor(TypeColor);
+	FString			   IconDocLink, IconDocExcerpt;
+	FSlateBrush const* SecondaryIconBrush = FEditorStyle::GetBrush(TEXT("NoBrush"));
+	FSlateColor        SecondaryIconColor = IconColor;
+
+	TSharedPtr<SHorizontalBox> ParameterContainer;
+	TSharedRef<SWidget> TooltipContent = GetParameterWidget(Variable, bShowValue);
+
+	TSharedPtr<SVerticalBox> InnerContainer;
 	// we construct a tooltip widget that shows the parameter the value is associated with
 	TSharedRef<SToolTip> TooltipWidget = SNew(SToolTip)
 	.Content()
 	[
-		TooltipContent
+		SAssignNew(InnerContainer, SVerticalBox)
+		+ SVerticalBox::Slot()
+		.AutoHeight()
+		[
+			TooltipContent
+		]
 	];
 
+	// if we specified an additional widget, add it here
+	if(AdditionalVerticalWidget.IsValid())
+	{
+		InnerContainer->AddSlot()
+		.AutoHeight()
+		[
+			AdditionalVerticalWidget.ToSharedRef()
+		];
+	}
+
+	if(AdditionalHorizontalWidget.IsValid())
+	{
+		ParameterContainer->AddSlot()
+		.AutoWidth()
+		[
+			AdditionalHorizontalWidget.ToSharedRef()	
+		];
+	}
+
 	return TooltipWidget;
 }
 
diff --git a/Engine/Plugins/FX/Niagara/Source/NiagaraEditor/Private/ViewModels/NiagaraParameterPanelViewModel.cpp b/Engine/Plugins/FX/Niagara/Source/NiagaraEditor/Private/ViewModels/NiagaraParameterPanelViewModel.cpp
--- a/Engine/Plugins/FX/Niagara/Source/NiagaraEditor/Private/ViewModels/NiagaraParameterPanelViewModel.cpp
+++ b/Engine/Plugins/FX/Niagara/Source/NiagaraEditor/Private/ViewModels/NiagaraParameterPanelViewModel.cpp
@@ -144,7 +144,6 @@ FReply FNiagaraSystemToolkitParameterPanelUtilities::CreateDragEventForParameter
 	const FText ToolTip = FText::Format(TooltipFormat, Name, DraggedItem.GetVariable().GetType().GetNameText());
 	TSharedPtr<FEdGraphSchemaAction> ItemDragAction = MakeShared<FNiagaraParameterAction>(DraggedItem.ScriptVariable, FText::GetEmpty(), Name, ToolTip, 0, FText(), 0/*SectionID*/);
 	TSharedPtr<FNiagaraParameterDragOperation> DragOperation = MakeShared<FNiagaraParameterDragOperation>(ItemDragAction);
-	DragOperation->CurrentHoverText = FNiagaraParameterUtilities::FormatParameterNameForTextDisplay(DraggedItem.GetVariable().GetName());
 	DragOperation->SetupDefaults();
 	DragOperation->Construct();
 	return FReply::Handled().BeginDragDrop(DragOperation.ToSharedRef());
diff --git a/Engine/Plugins/FX/Niagara/Source/NiagaraEditor/Private/ViewModels/Stack/NiagaraStackFunctionInput.cpp b/Engine/Plugins/FX/Niagara/Source/NiagaraEditor/Private/ViewModels/Stack/NiagaraStackFunctionInput.cpp
--- a/Engine/Plugins/FX/Niagara/Source/NiagaraEditor/Private/ViewModels/Stack/NiagaraStackFunctionInput.cpp
+++ b/Engine/Plugins/FX/Niagara/Source/NiagaraEditor/Private/ViewModels/Stack/NiagaraStackFunctionInput.cpp
@@ -2278,10 +2278,15 @@ bool UNiagaraStackFunctionInput::GetShouldPassFilterForVisibleCondition() const
 }
 
 TArray<UNiagaraScript*> UNiagaraStackFunctionInput::GetPossibleConversionScripts(const FNiagaraTypeDefinition& FromType) const
+{
+	return GetPossibleConversionScripts(FromType, InputType);
+}
+
+TArray<UNiagaraScript*> UNiagaraStackFunctionInput::GetPossibleConversionScripts(const FNiagaraTypeDefinition& FromType, const FNiagaraTypeDefinition& ToType)
 {
     FPinCollectorArray InputPins;
     TArray<UNiagaraNodeOutput*> OutputNodes;
-    auto MatchesTypeConversion = [this, &InputPins, &OutputNodes, &FromType](UNiagaraScript* Script)
+    auto MatchesTypeConversion = [&InputPins, &OutputNodes, &FromType, &ToType](UNiagaraScript* Script)
     {
 	    OutputNodes.Reset();
 	    UNiagaraGraph* NodeGraph = Cast<UNiagaraScriptSource>(Script->GetLatestSource())->NodeGraph;
@@ -2307,7 +2312,7 @@ TArray<UNiagaraScript*> UNiagaraStackFunctionInput::GetPossibleConversionScripts
     		{
     			const UEdGraphSchema_Niagara* NiagaraSchema = GetDefault<UEdGraphSchema_Niagara>();
     			FNiagaraTypeDefinition PinTypeIn = NiagaraSchema->PinToTypeDefinition(InputPins[0]);
-    			if (PinTypeIn == InputType)
+    			if (PinTypeIn == ToType)
     			{
     				return true;
     			}
diff --git a/Engine/Plugins/FX/Niagara/Source/NiagaraEditor/Private/ViewModels/Stack/NiagaraStackModuleItem.cpp b/Engine/Plugins/FX/Niagara/Source/NiagaraEditor/Private/ViewModels/Stack/NiagaraStackModuleItem.cpp
--- a/Engine/Plugins/FX/Niagara/Source/NiagaraEditor/Private/ViewModels/Stack/NiagaraStackModuleItem.cpp
+++ b/Engine/Plugins/FX/Niagara/Source/NiagaraEditor/Private/ViewModels/Stack/NiagaraStackModuleItem.cpp
@@ -228,11 +228,11 @@ TOptional<UNiagaraStackEntry::FDropRequestResponse> UNiagaraStackModuleItem::Can
 		{
 			if (AssignmentNode->GetAssignmentTargets().Contains(ParameterAction->GetParameter()))
 			{
-				return FDropRequestResponse(TOptional<EItemDropZone>(), LOCTEXT("CantDropDuplicateParameter", "Can not drop this parameter here because\nit's already set by this module."));
+				return FDropRequestResponse(TOptional<EItemDropZone>(), LOCTEXT("CantDropDuplicateParameter", "Can not drop this parameter here because it's already set by this module."));
 			}
 			else if (FNiagaraStackGraphUtilities::CanWriteParameterFromUsageViaOutput(ParameterAction->GetParameter(), OutputNode) == false)
 			{
-				return FDropRequestResponse(TOptional<EItemDropZone>(), LOCTEXT("CantDropParameterByUsage", "Can not drop this parameter here because\nit can't be written in this usage context."));
+				return FDropRequestResponse(TOptional<EItemDropZone>(), LOCTEXT("CantDropParameterByUsage", "Can not drop this parameter here because it can't be written in this usage context."));
 			}
 			else
 			{
diff --git a/Engine/Plugins/FX/Niagara/Source/NiagaraEditor/Public/NiagaraActions.h b/Engine/Plugins/FX/Niagara/Source/NiagaraEditor/Public/NiagaraActions.h
--- a/Engine/Plugins/FX/Niagara/Source/NiagaraEditor/Public/NiagaraActions.h
+++ b/Engine/Plugins/FX/Niagara/Source/NiagaraEditor/Public/NiagaraActions.h
@@ -355,6 +355,11 @@ public:
 
 	const TSharedPtr<FEdGraphSchemaAction>& GetSourceAction() const { return SourceAction; }
 
+	virtual TSharedPtr<SWidget> GetDefaultDecorator() const override;
+
+	void SetAdditionalText(FText InAdditionalText) { CurrentHoverText = InAdditionalText; }
+	EVisibility IsTextVisible() const;
 private:
 	TSharedPtr<FEdGraphSchemaAction> SourceAction;
+	TOptional<FNiagaraVariable> TargetParameter;
 };
\ No newline at end of file
diff --git a/Engine/Plugins/FX/Niagara/Source/NiagaraEditor/Public/NiagaraEditorUtilities.h b/Engine/Plugins/FX/Niagara/Source/NiagaraEditor/Public/NiagaraEditorUtilities.h
--- a/Engine/Plugins/FX/Niagara/Source/NiagaraEditor/Public/NiagaraEditorUtilities.h
+++ b/Engine/Plugins/FX/Niagara/Source/NiagaraEditor/Public/NiagaraEditorUtilities.h
@@ -401,8 +401,10 @@ namespace FNiagaraParameterUtilities
 
 	NIAGARAEDITOR_API bool TestCanRenameWithMessage(FName ParameterName, FText& OutMessage);
 
+	NIAGARAEDITOR_API TSharedRef<SWidget> GetParameterWidget(FNiagaraVariable Variable, bool bShowValue);
+	
 	/** Creates a tooltip based on a parameter. Also shows the value, if allocated and enabled. */
-	NIAGARAEDITOR_API TSharedRef<SToolTip> GetTooltipWidget(FNiagaraVariable Variable, bool bShowValue = true);
+	NIAGARAEDITOR_API TSharedRef<SToolTip> GetTooltipWidget(FNiagaraVariable Variable, bool bShowValue = true, TSharedPtr<SWidget> AdditionalVerticalWidget = nullptr,  TSharedPtr<SWidget> AdditionalHorizontalWidget = nullptr);
 
 	NIAGARAEDITOR_API void FilterToRelevantStaticVariables(const TArray<FNiagaraVariable>& InVars, TArray<FNiagaraVariable>& OutVars, FName InOldEmitterAlias, FName InNewEmitterAlias, bool bFilterByEmitterAliasAndConvertToUnaliased);
 };
diff --git a/Engine/Plugins/FX/Niagara/Source/NiagaraEditor/Public/ViewModels/Stack/NiagaraStackFunctionInput.h b/Engine/Plugins/FX/Niagara/Source/NiagaraEditor/Public/ViewModels/Stack/NiagaraStackFunctionInput.h
--- a/Engine/Plugins/FX/Niagara/Source/NiagaraEditor/Public/ViewModels/Stack/NiagaraStackFunctionInput.h
+++ b/Engine/Plugins/FX/Niagara/Source/NiagaraEditor/Public/ViewModels/Stack/NiagaraStackFunctionInput.h
@@ -251,6 +251,7 @@ public:
 	bool GetShouldPassFilterForVisibleCondition() const;
 	
 	TArray<UNiagaraScript*> GetPossibleConversionScripts(const FNiagaraTypeDefinition& FromType) const;
+	static TArray<UNiagaraScript*> GetPossibleConversionScripts(const FNiagaraTypeDefinition& FromType, const FNiagaraTypeDefinition& ToType);
 
 	void SetLinkedInputViaConversionScript(const FName& LinkedInputName, const FNiagaraTypeDefinition& FromType);
 	void SetLinkedInputViaConversionScript(const FNiagaraVariable& LinkedInput, UNiagaraScript* ConversionScript);
diff --git a/Engine/Plugins/FX/Niagara/Source/NiagaraEditorWidgets/Private/NiagaraEditorWidgetsUtilities.cpp b/Engine/Plugins/FX/Niagara/Source/NiagaraEditorWidgets/Private/NiagaraEditorWidgetsUtilities.cpp
--- a/Engine/Plugins/FX/Niagara/Source/NiagaraEditorWidgets/Private/NiagaraEditorWidgetsUtilities.cpp
+++ b/Engine/Plugins/FX/Niagara/Source/NiagaraEditorWidgets/Private/NiagaraEditorWidgetsUtilities.cpp
@@ -492,7 +492,7 @@ TOptional<EItemDropZone> FNiagaraStackEditorWidgetsUtilities::RequestDropForStac
 		{
 			if (DecoratedDragDropOp.IsValid() && Response.GetValue().DropMessage.IsEmptyOrWhitespace() == false)
 			{
-				DecoratedDragDropOp->CurrentHoverText = FText::Format(LOCTEXT("DropFormat", "{0} - {1}"), DecoratedDragDropOp->GetDefaultHoverText(), Response.GetValue().DropMessage);
+				DecoratedDragDropOp->CurrentHoverText = Response.GetValue().DropMessage;
 			}
 
 			if (Response.GetValue().DropZone.IsSet())
diff --git a/Engine/Plugins/FX/Niagara/Source/NiagaraEditorWidgets/Private/SNiagaraParameterDropTarget.cpp b/Engine/Plugins/FX/Niagara/Source/NiagaraEditorWidgets/Private/SNiagaraParameterDropTarget.cpp
new file mode 100644
--- /dev/null
+++ b/Engine/Plugins/FX/Niagara/Source/NiagaraEditorWidgets/Private/SNiagaraParameterDropTarget.cpp
@@ -0,0 +1,95 @@
+// Copyright Epic Games, Inc. All Rights Reserved.
+
+#include "SNiagaraParameterDropTarget.h"
+
+#include "NiagaraActions.h"
+#include "NiagaraEditorUtilities.h"
+#include "ViewModels/Stack/NiagaraStackGraphUtilities.h"
+
+#define LOCTEXT_NAMESPACE "NiagaraParameterDropTarget"
+
+void SNiagaraParameterDropTarget::Construct(const FArguments& InArgs)
+{
+	SDropTarget::Construct(InArgs._DropTargetArgs);
+
+	TargetParameter = InArgs._TargetParameter;
+	TypeToTestAgainst = InArgs._TypeToTestAgainst;
+	ExecutionCategory = InArgs._ExecutionCategory;
+}
+
+FReply SNiagaraParameterDropTarget::OnDragOver(const FGeometry& MyGeometry, const FDragDropEvent& DragDropEvent)
+{
+	TSharedPtr<FNiagaraParameterDragOperation> ParameterDragOperation = DragDropEvent.GetOperationAs<FNiagaraParameterDragOperation>();
+
+	if(ParameterDragOperation.IsValid())
+	{
+		ensure(TargetParameter.IsSet() || TypeToTestAgainst.IsSet());
+		
+		FText TargetText = TargetParameter.IsSet() ? FText::FromName(TargetParameter->GetName()) : TypeToTestAgainst.IsSet() ? TypeToTestAgainst->GetNameText() : FText::GetEmpty();
+		
+		if(TargetParameter.IsSet())
+		{
+			TypeToTestAgainst = TargetParameter->GetType();
+		}
+		
+		if(!TargetText.IsEmpty() && ExecutionCategory->IsValid())
+		{
+			TSharedPtr<FNiagaraParameterAction> Action = StaticCastSharedPtr<FNiagaraParameterAction>(ParameterDragOperation->GetSourceAction());
+			bool bAllowedInExecutionCategory = FNiagaraStackGraphUtilities::ParameterAllowedInExecutionCategory(Action->GetParameter().GetName(), ExecutionCategory.GetValue());
+			FNiagaraTypeDefinition DropType = Action->GetParameter().GetType();
+
+			if (bAllowedInExecutionCategory)
+			{
+				// check if we can simply link the input directly
+				if(FNiagaraEditorUtilities::AreTypesAssignable(DropType, TypeToTestAgainst.GetValue()))
+				{
+					// if we can directly connect, we don't need any additional text
+					ParameterDragOperation->SetAdditionalText(FText::GetEmpty()); 
+				}
+				// check if we can use a conversion script
+				else if(UNiagaraStackFunctionInput::GetPossibleConversionScripts(DropType, TypeToTestAgainst.GetValue()).Num() > 0)
+				{
+					UNiagaraScript* ConversionScriptToUse = UNiagaraStackFunctionInput::GetPossibleConversionScripts(DropType, TypeToTestAgainst.GetValue())[0];
+					FText ConversionScriptDescription = ConversionScriptToUse->GetLatestScriptData()->Description;
+					FText ModifiedConversionScriptDescription = ConversionScriptDescription.IsEmpty() ? FText::GetEmpty() : FText::FromString("\n" + ConversionScriptDescription.ToString()); 
+					if(FNiagaraTypeDefinition::IsLossyConversion(Action->GetParameter().GetType(), TypeToTestAgainst.GetValue()))
+					{
+						FText Format = LOCTEXT("RequiresLossyAdapterDragDropInfo", "A conversion script will be inserted to adapt {0} to {1}.\nThis is a potentially lossy conversion.{2}");
+						FText HoverTooltipText = FText::FormatOrdered(Format, Action->GetParameter().GetType().GetNameText(), TypeToTestAgainst.GetValue().GetNameText(), ModifiedConversionScriptDescription);
+						ParameterDragOperation->SetAdditionalText(HoverTooltipText); 
+					}
+					else
+					{
+						FText Format = LOCTEXT("RequiresAdapterDragDropInfo", "A conversion script will be inserted to adapt {0} to {1}.{2}");
+						FText HoverTooltipText = FText::FormatOrdered(Format, Action->GetParameter().GetType().GetNameText(), TypeToTestAgainst.GetValue().GetNameText(), ModifiedConversionScriptDescription);
+						ParameterDragOperation->SetAdditionalText(HoverTooltipText); 
+					}
+				}
+			}
+			else
+			{
+				FText Format = LOCTEXT("ExecutionCategoryNotAllowed", "This parameter is not allowed in execution category {0}.");
+				FText HoverTooltipText = FText::FormatOrdered(Format, FText::FromName(ExecutionCategory.GetValue()));
+				ParameterDragOperation->SetAdditionalText(HoverTooltipText); 
+			}
+		}
+	}
+	
+	AllowDrop(DragDropEvent.GetOperation());
+	// we always handle the drag because we don't want another widget to interfere with the feedback we are giving
+	return FReply::Handled();
+}
+
+void SNiagaraParameterDropTarget::OnDragLeave(const FDragDropEvent& DragDropEvent)
+{
+	SDropTarget::OnDragLeave(DragDropEvent);
+
+	TSharedPtr<FNiagaraParameterDragOperation> ParameterDragOperation = DragDropEvent.GetOperationAs<FNiagaraParameterDragOperation>();
+
+	if(ParameterDragOperation.IsValid())
+	{
+		ParameterDragOperation->ResetToDefaultToolTip();
+	}
+}
+
+#undef LOCTEXT_NAMESPACE
diff --git a/Engine/Plugins/FX/Niagara/Source/NiagaraEditorWidgets/Private/SNiagaraParameterDropTarget.h b/Engine/Plugins/FX/Niagara/Source/NiagaraEditorWidgets/Private/SNiagaraParameterDropTarget.h
new file mode 100644
--- /dev/null
+++ b/Engine/Plugins/FX/Niagara/Source/NiagaraEditorWidgets/Private/SNiagaraParameterDropTarget.h
@@ -0,0 +1,41 @@
+// Copyright Epic Games, Inc. All Rights Reserved.
+
+#pragma once
+
+#include "SDropTarget.h"
+#include "CoreMinimal.h"
+#include "Styling/SlateColor.h"
+#include "Input/Reply.h"
+#include "Layout/Visibility.h"
+#include "Framework/SlateDelegates.h"
+#include "Widgets/DeclarativeSyntaxSupport.h"
+#include "Widgets/SCompoundWidget.h"
+#include "EditorStyleSet.h"
+#include "Styling/StyleColors.h"
+#include "ViewModels/Stack/NiagaraStackFunctionInput.h"
+
+class NIAGARAEDITORWIDGETS_API SNiagaraParameterDropTarget : public SDropTarget
+{
+public:
+    DECLARE_DELEGATE_OneParam(FOnDragOperationDragged, TSharedPtr<FDragDropOperation>)
+    SLATE_BEGIN_ARGS(SNiagaraParameterDropTarget)
+	    : _DropTargetArgs(SDropTarget::FArguments())
+		, _TargetParameter(TOptional<FNiagaraVariable>())
+		, _TypeToTestAgainst(TOptional<FNiagaraTypeDefinition>())
+		, _ExecutionCategory(TOptional<FName>())
+	{ }
+		SLATE_ARGUMENT(SDropTarget::FArguments, DropTargetArgs)
+		SLATE_ARGUMENT(TOptional<FNiagaraVariable>, TargetParameter)
+		SLATE_ARGUMENT(TOptional<FNiagaraTypeDefinition>, TypeToTestAgainst)
+		SLATE_ARGUMENT(TOptional<FName>, ExecutionCategory)
+    SLATE_END_ARGS()
+
+	void Construct(const FArguments& InArgs);
+protected:
+    virtual FReply OnDragOver(const FGeometry& MyGeometry, const FDragDropEvent& DragDropEvent) override;
+	virtual void OnDragLeave(const FDragDropEvent& DragDropEvent) override;
+private:
+	TOptional<FName> ExecutionCategory;
+	TOptional<FNiagaraVariable> TargetParameter;
+	TOptional<FNiagaraTypeDefinition> TypeToTestAgainst;
+};
diff --git a/Engine/Plugins/FX/Niagara/Source/NiagaraEditorWidgets/Private/Stack/SNiagaraStackFunctionInputValue.cpp b/Engine/Plugins/FX/Niagara/Source/NiagaraEditorWidgets/Private/Stack/SNiagaraStackFunctionInputValue.cpp
--- a/Engine/Plugins/FX/Niagara/Source/NiagaraEditorWidgets/Private/Stack/SNiagaraStackFunctionInputValue.cpp
+++ b/Engine/Plugins/FX/Niagara/Source/NiagaraEditorWidgets/Private/Stack/SNiagaraStackFunctionInputValue.cpp
@@ -51,6 +51,7 @@
 #include "NiagaraNodeAssignment.h"
 #include "NiagaraCommon.h"
 #include "NiagaraSettings.h"
+#include "SNiagaraParameterDropTarget.h"
 
 #define LOCTEXT_NAMESPACE "NiagaraStackFunctionInputValue"
 
@@ -71,7 +72,11 @@ void SNiagaraStackFunctionInputValue::Construct(const FArguments& InArgs, UNiaga
 		SAssignNew(ChildrenBox, SHorizontalBox)
 		+ SHorizontalBox::Slot()
 		[
-			SNew(SDropTarget)
+			SNew(SNiagaraParameterDropTarget)
+			.TypeToTestAgainst(FunctionInput->GetInputType())
+			.ExecutionCategory(FunctionInput->GetExecutionCategoryName())
+			.TargetParameter(FNiagaraVariable(FunctionInput->GetInputType(), FunctionInput->GetInputParameterHandle().GetParameterHandleString()))
+			.DropTargetArgs(SDropTarget::FArguments()
 			.OnAllowDrop(this, &SNiagaraStackFunctionInputValue::OnFunctionInputAllowDrop)
 			.OnDropped(this, &SNiagaraStackFunctionInputValue::OnFunctionInputDrop)
 			.HorizontalImage(FNiagaraEditorWidgetsStyle::Get().GetBrush("NiagaraEditor.Stack.DropTarget.BorderHorizontal"))
@@ -185,7 +190,7 @@ void SNiagaraStackFunctionInputValue::Construct(const FArguments& InArgs, UNiaga
 						.ColorAndOpacity(FSlateColor(FLinearColor::Green))
 					]
 				]			
-			]
+			])
 		]
 	];
 
