Fixed crash when selecting a comment node in scalability mode

diff --git a/Engine/Plugins/FX/Niagara/Source/NiagaraEditor/Private/Widgets/Scalability/SNiagaraScalabilityContext.cpp b/Engine/Plugins/FX/Niagara/Source/NiagaraEditor/Private/Widgets/Scalability/SNiagaraScalabilityContext.cpp
--- a/Engine/Plugins/FX/Niagara/Source/NiagaraEditor/Private/Widgets/Scalability/SNiagaraScalabilityContext.cpp
+++ b/Engine/Plugins/FX/Niagara/Source/NiagaraEditor/Private/Widgets/Scalability/SNiagaraScalabilityContext.cpp
@@ -83,16 +83,19 @@ void SNiagaraScalabilityContext::UpdateScalabilityContent()
 	if(SelectedNodes.Num() > 0)
 	{
 		UObject* Object = SelectedNodes.Array()[0];
-		UNiagaraOverviewNode* OverviewNode = CastChecked<UNiagaraOverviewNode>(Object);
-		if(FNiagaraEmitterHandle* Handle = OverviewNode->TryGetEmitterHandle())
+		UNiagaraOverviewNode* OverviewNode = Cast<UNiagaraOverviewNode>(Object);
+		if(OverviewNode)
 		{
-			SetObject(Handle->GetInstance());
-			return;
-		}
-		else
-		{
-			SetObject(OverviewNode->GetOwningSystem());
-			return;
+			if(FNiagaraEmitterHandle* Handle = OverviewNode->TryGetEmitterHandle())
+			{
+				SetObject(Handle->GetInstance());
+				return;
+			}
+			else
+			{
+				SetObject(OverviewNode->GetOwningSystem());
+				return;
+			}
 		}
 	}
 }
diff --git a/Engine/Plugins/FX/Niagara/Source/NiagaraEditor/Private/Widgets/Scalability/SNiagaraSystemResolvedScalabilitySettings.cpp b/Engine/Plugins/FX/Niagara/Source/NiagaraEditor/Private/Widgets/Scalability/SNiagaraSystemResolvedScalabilitySettings.cpp
--- a/Engine/Plugins/FX/Niagara/Source/NiagaraEditor/Private/Widgets/Scalability/SNiagaraSystemResolvedScalabilitySettings.cpp
+++ b/Engine/Plugins/FX/Niagara/Source/NiagaraEditor/Private/Widgets/Scalability/SNiagaraSystemResolvedScalabilitySettings.cpp
@@ -217,30 +217,42 @@ void SNiagaraSystemResolvedScalabilitySettings::RebuildWidget()
 {
 	ResolvedScalabilityContainer->ClearChildren();
 	
-	const TSet<UObject*>& SelectedNodes = SystemViewModel->GetOverviewGraphViewModel()->GetNodeSelection()->GetSelectedObjects();
+	TSet<UNiagaraOverviewNode*> SelectedOverviewNodes;
 
-	if(SelectedNodes.Num() > 0)
+	for(UObject* SelectedNode : SystemViewModel->GetOverviewGraphViewModel()->GetNodeSelection()->GetSelectedObjects())
 	{
-		for(UObject* SelectedNode : SelectedNodes)
+		UNiagaraOverviewNode* OverviewNode = Cast<UNiagaraOverviewNode>(SelectedNode);
+		if(OverviewNode)
 		{
-			UNiagaraOverviewNode* OverviewNode = CastChecked<UNiagaraOverviewNode>(SelectedNode);
-			if(FNiagaraEmitterHandle* Handle = OverviewNode->TryGetEmitterHandle())
-			{
-				Handle->GetInstance()->OnPropertiesChanged().RemoveAll(this);
-				Handle->GetInstance()->OnPropertiesChanged().AddSP(this, &SNiagaraSystemResolvedScalabilitySettings::RebuildWidget);
-				ResolvedScalabilityContainer->AddSlot()
-				.AutoHeight()
-				[
-					GenerateResolvedScalabilityTable(Handle->GetInstance())
-				];
-			}
-			else
+			SelectedOverviewNodes.Add(OverviewNode);
+		}
+	}
+	
+	if(SelectedOverviewNodes.Num() > 0)
+	{
+		for(UObject* SelectedNode : SelectedOverviewNodes)
+		{
+			UNiagaraOverviewNode* OverviewNode = Cast<UNiagaraOverviewNode>(SelectedNode);
+			if(OverviewNode)
 			{
-				ResolvedScalabilityContainer->AddSlot()
-				.AutoHeight()
-				[
-					GenerateResolvedScalabilityTable(OverviewNode->GetOwningSystem())
-				];
+				if(FNiagaraEmitterHandle* Handle = OverviewNode->TryGetEmitterHandle())
+				{
+					Handle->GetInstance()->OnPropertiesChanged().RemoveAll(this);
+					Handle->GetInstance()->OnPropertiesChanged().AddSP(this, &SNiagaraSystemResolvedScalabilitySettings::RebuildWidget);
+					ResolvedScalabilityContainer->AddSlot()
+					.AutoHeight()
+					[
+						GenerateResolvedScalabilityTable(Handle->GetInstance())
+					];
+				}
+				else
+				{
+					ResolvedScalabilityContainer->AddSlot()
+					.AutoHeight()
+					[
+						GenerateResolvedScalabilityTable(OverviewNode->GetOwningSystem())
+					];
+				}
 			}
 		}
 	}
