Fixed a bug where all user parameters were considered for position fixup, not just vec3 types

diff --git a/Engine/Plugins/FX/Niagara/Source/Niagara/Private/NiagaraSystem.cpp b/Engine/Plugins/FX/Niagara/Source/Niagara/Private/NiagaraSystem.cpp
--- a/Engine/Plugins/FX/Niagara/Source/Niagara/Private/NiagaraSystem.cpp
+++ b/Engine/Plugins/FX/Niagara/Source/Niagara/Private/NiagaraSystem.cpp
@@ -3840,7 +3840,7 @@ void UNiagaraSystem::FixupPositionUserParameters()
 {
 	TArray<FNiagaraVariable> UserParameters;
 	ExposedParameters.GetUserParameters(UserParameters);
-	UserParameters.FilterByPredicate([](const FNiagaraVariable& Var) { return Var.GetType() == FNiagaraTypeDefinition::GetVec3Def(); });
+	UserParameters = UserParameters.FilterByPredicate([](const FNiagaraVariable& Var) { return Var.GetType() == FNiagaraTypeDefinition::GetVec3Def(); });
 
 	if (UserParameters.Num() == 0)
 	{
