Fixed missing include for non-unity build

diff --git a/Engine/Plugins/FX/Niagara/Source/NiagaraEditor/Private/NiagaraActions.cpp b/Engine/Plugins/FX/Niagara/Source/NiagaraEditor/Private/NiagaraActions.cpp
--- a/Engine/Plugins/FX/Niagara/Source/NiagaraEditor/Private/NiagaraActions.cpp
+++ b/Engine/Plugins/FX/Niagara/Source/NiagaraEditor/Private/NiagaraActions.cpp
@@ -19,6 +19,7 @@
 #include "ViewModels/NiagaraParameterPanelViewModel.h"
 #include "Misc/MessageDialog.h"
 #include "Widgets/SNiagaraParameterName.h"
+#include "Widgets/SToolTip.h"
 
 #define LOCTEXT_NAMESPACE "NiagaraActions"
 
diff --git a/Engine/Plugins/FX/Niagara/Source/NiagaraEditor/Private/NiagaraParameterDefinitions.cpp b/Engine/Plugins/FX/Niagara/Source/NiagaraEditor/Private/NiagaraParameterDefinitions.cpp
--- a/Engine/Plugins/FX/Niagara/Source/NiagaraEditor/Private/NiagaraParameterDefinitions.cpp
+++ b/Engine/Plugins/FX/Niagara/Source/NiagaraEditor/Private/NiagaraParameterDefinitions.cpp
@@ -2,6 +2,7 @@
 
 #include "NiagaraParameterDefinitions.h"
 
+#include "ScopedTransaction.h"
 #include "NiagaraEditorModule.h"
 #include "NiagaraEditorUtilities.h"
 #include "NiagaraScriptVariable.h"
