- PR #8926: Change NiagaraDebugDraw GPULineBuffer format to uint (Contributed by [REDACTED_PERSON])

#rnx

diff --git a/Engine/Plugins/FX/Niagara/Shaders/Private/NiagaraDataInterfaceDebugDraw.ush b/Engine/Plugins/FX/Niagara/Shaders/Private/NiagaraDataInterfaceDebugDraw.ush
--- a/Engine/Plugins/FX/Niagara/Shaders/Private/NiagaraDataInterfaceDebugDraw.ush
+++ b/Engine/Plugins/FX/Niagara/Shaders/Private/NiagaraDataInterfaceDebugDraw.ush
@@ -6,7 +6,7 @@
 #include "NiagaraDebugDraw.ush"
 
 RWBuffer<uint>		RWNDIDebugDrawArgs;
-RWBuffer<float>		RWNDIDebugDrawLineVertex;
+RWBuffer<uint>		RWNDIDebugDrawLineVertex;
 uint				NDIDebugDrawLineMaxInstances;
 
 void NDIDebugDraw_DrawBox(bool bDraw, float3 Location, float4 Rotation, float3 Extents, float4 Color)
diff --git a/Engine/Plugins/FX/Niagara/Shaders/Private/NiagaraDebugDraw.usf b/Engine/Plugins/FX/Niagara/Shaders/Private/NiagaraDebugDraw.usf
--- a/Engine/Plugins/FX/Niagara/Shaders/Private/NiagaraDebugDraw.usf
+++ b/Engine/Plugins/FX/Niagara/Shaders/Private/NiagaraDebugDraw.usf
@@ -28,7 +28,7 @@ void MainCS(in uint3 ThreadID : SV_DispatchThreadID)
 ////////////////////////////////////////////////////////////////////////////////////
 
 #if NIAGARA_DEBUGDRAW_DRAWLINE_VS
-Buffer<float>	GpuLineBuffer;
+Buffer<uint>	GpuLineBuffer;
 
 void MainVS(
 	in uint InstanceID		: SV_InstanceID,
diff --git a/Engine/Plugins/FX/Niagara/Shaders/Private/NiagaraDebugDraw.ush b/Engine/Plugins/FX/Niagara/Shaders/Private/NiagaraDebugDraw.ush
--- a/Engine/Plugins/FX/Niagara/Shaders/Private/NiagaraDebugDraw.ush
+++ b/Engine/Plugins/FX/Niagara/Shaders/Private/NiagaraDebugDraw.ush
@@ -7,22 +7,22 @@
 	Includes for debug drawing
 =============================================================================*/
 
-void PackLine(inout RWBuffer<float> OutBuffer, uint InstanceIndex, float3 InPosition0, float3 InPosition1, uint InPackedColor)
+void PackLine(inout RWBuffer<uint> OutBuffer, uint InstanceIndex, float3 InPosition0, float3 InPosition1, uint InPackedColor)
 {
-	uint FloatOffset = InstanceIndex * 7;
+	uint BufferOffset = InstanceIndex * 7;
 
-	OutBuffer[FloatOffset + 0] = InPosition0.x;
-	OutBuffer[FloatOffset + 1] = InPosition0.y;
-	OutBuffer[FloatOffset + 2] = InPosition0.z;
+	OutBuffer[BufferOffset + 0] = asuint(InPosition0.x);
+	OutBuffer[BufferOffset + 1] = asuint(InPosition0.y);
+	OutBuffer[BufferOffset + 2] = asuint(InPosition0.z);
 
-	OutBuffer[FloatOffset + 3] = InPosition1.x;
-	OutBuffer[FloatOffset + 4] = InPosition1.y;
-	OutBuffer[FloatOffset + 5] = InPosition1.z;
+	OutBuffer[BufferOffset + 3] = asuint(InPosition1.x);
+	OutBuffer[BufferOffset + 4] = asuint(InPosition1.y);
+	OutBuffer[BufferOffset + 5] = asuint(InPosition1.z);
 
-	OutBuffer[FloatOffset + 6] = asfloat(InPackedColor);
+	OutBuffer[BufferOffset + 6] = InPackedColor;
 }
 
-void PackLine(inout RWBuffer<float> OutBuffer, uint InstanceIndex, float3 InPosition0, float3 InPosition1, float4 InColor)
+void PackLine(inout RWBuffer<uint> OutBuffer, uint InstanceIndex, float3 InPosition0, float3 InPosition1, float4 InColor)
 {
 	uint PackedColor = 0;
 	PackedColor  = uint(saturate(InColor.x) * 255.0f) << 24;
@@ -32,19 +32,19 @@ void PackLine(inout RWBuffer<float> OutBuffer, uint InstanceIndex, float3 InPosi
 	PackLine(OutBuffer, InstanceIndex, InPosition0, InPosition1, PackedColor);
 }
 
-void UnpackLine(in Buffer<float> InBuffer, uint InstanceIndex, out float3 OutPosition0, out float3 OutPosition1, out float4 OutColor)
+void UnpackLine(in Buffer<uint> InBuffer, uint InstanceIndex, out float3 OutPosition0, out float3 OutPosition1, out float4 OutColor)
 {
-	uint FloatOffset = InstanceIndex * 7;
+	uint BufferOffset = InstanceIndex * 7;
 
-	OutPosition0.x = InBuffer[FloatOffset + 0];
-	OutPosition0.y = InBuffer[FloatOffset + 1];
-	OutPosition0.z = InBuffer[FloatOffset + 2];
+	OutPosition0.x = asfloat(InBuffer[BufferOffset + 0]);
+	OutPosition0.y = asfloat(InBuffer[BufferOffset + 1]);
+	OutPosition0.z = asfloat(InBuffer[BufferOffset + 2]);
 
-	OutPosition1.x = InBuffer[FloatOffset + 3];
-	OutPosition1.y = InBuffer[FloatOffset + 4];
-	OutPosition1.z = InBuffer[FloatOffset + 5];
+	OutPosition1.x = asfloat(InBuffer[BufferOffset + 3]);
+	OutPosition1.y = asfloat(InBuffer[BufferOffset + 4]);
+	OutPosition1.z = asfloat(InBuffer[BufferOffset + 5]);
 
-	uint PackedColor = asuint(InBuffer[FloatOffset + 6]);
+	uint PackedColor = InBuffer[BufferOffset + 6];
 	OutColor.x = float((PackedColor >> 24) & 0xff) / 255.0f;
 	OutColor.y = float((PackedColor >> 16) & 0xff) / 255.0f;
 	OutColor.z = float((PackedColor >>  8) & 0xff) / 255.0f;
diff --git a/Engine/Plugins/FX/Niagara/Source/Niagara/Private/NiagaraGpuComputeDebug.cpp b/Engine/Plugins/FX/Niagara/Source/Niagara/Private/NiagaraGpuComputeDebug.cpp
--- a/Engine/Plugins/FX/Niagara/Source/Niagara/Private/NiagaraGpuComputeDebug.cpp
+++ b/Engine/Plugins/FX/Niagara/Source/Niagara/Private/NiagaraGpuComputeDebug.cpp
@@ -74,15 +74,15 @@ void FNiagaraGpuComputeDebug::Tick(FRHICommandListImmediate& RHICmdList)
 		DebugDrawData->StaticLineCount = DebugDrawData->StaticLines.Num();
 		if (DebugDrawData->StaticLineCount > 0 )
 		{
-			constexpr uint32 NumFloatsPerLine = 7;
-			static_assert(sizeof(FNiagaraSimulationDebugDrawData::FGpuLine) == (NumFloatsPerLine * sizeof(float)), "Line size does not match expected GPU size");
+			constexpr uint32 NumUintsPerLine = 7;
+			static_assert(sizeof(FNiagaraSimulationDebugDrawData::FGpuLine) == (NumUintsPerLine * sizeof(uint32)), "Line size does not match expected GPU size");
 
-			const uint32 NumElements = FMath::DivideAndRoundUp(DebugDrawData->StaticLineCount, 64u) * 64u * NumFloatsPerLine;
-			const uint32 RequiredBytes = NumElements * sizeof(float);
+			const uint32 NumElements = FMath::DivideAndRoundUp(DebugDrawData->StaticLineCount, 64u) * 64u * NumUintsPerLine;
+			const uint32 RequiredBytes = NumElements * sizeof(uint32);
 			if ( DebugDrawData->StaticLineBuffer.NumBytes < RequiredBytes )
 			{
 				DebugDrawData->StaticLineBuffer.Release();
-				DebugDrawData->StaticLineBuffer.Initialize(TEXT("NiagaraGpuComputeDebug::StaticLineBuffer"), sizeof(float), NumElements, EPixelFormat::PF_R32_FLOAT);
+				DebugDrawData->StaticLineBuffer.Initialize(TEXT("NiagaraGpuComputeDebug::StaticLineBuffer"), sizeof(uint32), NumElements, EPixelFormat::PF_R32_UINT);
 			}
 			void* VertexData = RHILockBuffer(DebugDrawData->StaticLineBuffer.Buffer, 0, RequiredBytes, RLM_WriteOnly);
 			FMemory::Memcpy(VertexData, DebugDrawData->StaticLines.GetData(), DebugDrawData->StaticLineCount * DebugDrawData->StaticLines.GetTypeSize());
@@ -241,7 +241,7 @@ FNiagaraSimulationDebugDrawData* FNiagaraGpuComputeDebug::GetSimulationDebugDraw
 		if (DebugDrawDataPtr->GpuLineMaxInstances > 0)
 		{
 			DebugDrawDataPtr->GpuLineBufferArgs.Initialize(TEXT("NiagaraGpuComputeDebug::DrawLineBufferArgs"), sizeof(uint32), 4, EPixelFormat::PF_R32_UINT, BUF_Static | BUF_DrawIndirect);
-			DebugDrawDataPtr->GpuLineVertexBuffer.Initialize(TEXT("NiagaraGpuComputeDebug::DrawLineVertexBuffer"), sizeof(float), 7 * DebugDrawDataPtr->GpuLineMaxInstances, EPixelFormat::PF_R32_FLOAT, BUF_Static);
+			DebugDrawDataPtr->GpuLineVertexBuffer.Initialize(TEXT("NiagaraGpuComputeDebug::DrawLineVertexBuffer"), sizeof(uint32), 7 * DebugDrawDataPtr->GpuLineMaxInstances, EPixelFormat::PF_R32_UINT, BUF_Static);
 
 			auto& RHICmdList = FRHICommandListExecutor::GetImmediateCommandList();
 			{
diff --git a/Engine/Plugins/FX/Niagara/Source/NiagaraShader/Private/NiagaraDebugShaders.cpp b/Engine/Plugins/FX/Niagara/Source/NiagaraShader/Private/NiagaraDebugShaders.cpp
--- a/Engine/Plugins/FX/Niagara/Source/NiagaraShader/Private/NiagaraDebugShaders.cpp
+++ b/Engine/Plugins/FX/Niagara/Source/NiagaraShader/Private/NiagaraDebugShaders.cpp
@@ -125,7 +125,7 @@ class NIAGARASHADER_API FNiagaraDebugDrawLineVS : public FGlobalShader
 
 	BEGIN_SHADER_PARAMETER_STRUCT(FParameters, )
 		SHADER_PARAMETER_STRUCT_REF(FViewUniformShaderParameters, View)
-		SHADER_PARAMETER_SRV(Buffer<float>, GpuLineBuffer)
+		SHADER_PARAMETER_SRV(Buffer<uint>, GpuLineBuffer)
 	END_SHADER_PARAMETER_STRUCT()
 };
 
