- Separate expanding bounds between dynamic & static
- Add option to snap bounds to units to reduce update transform frequency as they grow
- When no static bounds or dynamic bounds are present do not return the static bounds

diff --git a/Engine/Plugins/FX/Niagara/Source/Niagara/Classes/NiagaraEmitterInstance.h b/Engine/Plugins/FX/Niagara/Source/Niagara/Classes/NiagaraEmitterInstance.h
--- a/Engine/Plugins/FX/Niagara/Source/Niagara/Classes/NiagaraEmitterInstance.h
+++ b/Engine/Plugins/FX/Niagara/Source/Niagara/Classes/NiagaraEmitterInstance.h
@@ -124,7 +124,8 @@ public:
 	ENiagaraExecutionState NIAGARA_API GetExecutionState() { return ExecutionState; }
 	void NIAGARA_API SetExecutionState(ENiagaraExecutionState InState);
 
-	FBox GetBounds() const;
+	bool AreBoundsDynamic() const { return bCachedBoundsDynamic; }
+	FBox GetBounds() const { return CachedBounds; }
 
 	FNiagaraScriptExecutionContext& GetSpawnExecutionContext() { return SpawnExecContext; }
 	FNiagaraScriptExecutionContext& GetUpdateExecutionContext() { return UpdateExecContext; }
@@ -201,6 +202,9 @@ private:
 	/** A parameter store which contains the data interfaces parameters which were defined by the scripts. */
 	FNiagaraParameterStore ScriptDefinedDataInterfaceParameters;
 
+	/* Are the cached emitter bounds dynamic */
+	bool bCachedBoundsDynamic = false;
+
 	/* Emitter bounds */
 	FBox CachedBounds;
 
diff --git a/Engine/Plugins/FX/Niagara/Source/Niagara/Private/NiagaraEmitterInstance.cpp b/Engine/Plugins/FX/Niagara/Source/Niagara/Private/NiagaraEmitterInstance.cpp
--- a/Engine/Plugins/FX/Niagara/Source/Niagara/Private/NiagaraEmitterInstance.cpp
+++ b/Engine/Plugins/FX/Niagara/Source/Niagara/Private/NiagaraEmitterInstance.cpp
@@ -174,11 +174,6 @@ FNiagaraEmitterInstance::~FNiagaraEmitterInstance()
 	}
 }
 
-FBox FNiagaraEmitterInstance::GetBounds() const
-{
-	return CachedBounds;
-}
-
 TArrayView<FNiagaraScriptExecutionContext> FNiagaraEmitterInstance::GetEventExecutionContexts()
 {
 	if (EventInstanceData.IsValid())
@@ -531,6 +526,7 @@ void FNiagaraEmitterInstance::ResetSimulation(bool bKillExisting /*= true*/)
 	EmitterAge = 0;
 	TickCount = 0;
 	InstanceSeed = FGenericPlatformMath::Rand();
+	bCachedBoundsDynamic = false;
 	CachedBounds.Init();
 	ParticlesWithComponents.Empty();
 
@@ -1002,6 +998,7 @@ void FNiagaraEmitterInstance::CalculateFixedBounds(const FTransform& ToWorldSpac
 		CachedEmitter->FixedBounds = Bounds.TransformBy(ToWorldSpace);
 	}
 
+	bCachedBoundsDynamic = false;
 	CachedBounds = Bounds;
 }
 #endif
@@ -1024,6 +1021,7 @@ void FNiagaraEmitterInstance::PostTick()
 		}
 	}
 
+	bCachedBoundsDynamic = false;
 	CachedBounds.Init();
 	{
 		// Read lock can be smaller in scope, but probably not necessary
@@ -1045,6 +1043,7 @@ void FNiagaraEmitterInstance::PostTick()
 			FBox DynamicBounds = InternalCalculateDynamicBounds(ParticleDataSet->GetCurrentDataChecked().GetNumInstances());
 			if (DynamicBounds.IsValid)
 			{
+				bCachedBoundsDynamic = true;
 				if (CachedEmitter->bLocalSpace)
 				{
 					CachedBounds = DynamicBounds;
@@ -1056,7 +1055,7 @@ void FNiagaraEmitterInstance::PostTick()
 			}
 			else
 			{
-				CachedBounds = CachedEmitter->FixedBounds;
+				CachedBounds.IsValid = false;
 			}
 		}
 	}
diff --git a/Engine/Plugins/FX/Niagara/Source/Niagara/Private/NiagaraSystemInstance.cpp b/Engine/Plugins/FX/Niagara/Source/Niagara/Private/NiagaraSystemInstance.cpp
--- a/Engine/Plugins/FX/Niagara/Source/Niagara/Private/NiagaraSystemInstance.cpp
+++ b/Engine/Plugins/FX/Niagara/Source/Niagara/Private/NiagaraSystemInstance.cpp
@@ -65,11 +65,29 @@ static FAutoConsoleVariableRef CVarNiagaraForceLastTickGroup(
 	ECVF_Default
 );
 
-static float GNiagaraBoundsExpandByPercent = 0.1f;
-static FAutoConsoleVariableRef CVarNiagaraBoundsExpandByPercent(
-	TEXT("fx.Niagara.BoundsExpandByPercent"),
-	GNiagaraBoundsExpandByPercent,
-	TEXT("The percentage we expand the bounds to avoid updating every frame."),
+static float GNiagaraEmitterBoundsDynamicSnapValue = 0.0f;
+static FAutoConsoleVariableRef CVarNiagaraEmitterBoundsDynamicSnapValue(
+	TEXT("fx.Niagara.EmitterBounds.DynamicSnapValue"),
+	GNiagaraEmitterBoundsDynamicSnapValue,
+	TEXT("The value used to snap (round up) dynamic bounds calculations to.")
+	TEXT("For example, a snap of 128 and a value of 1 would result in 128"),
+	ECVF_Default
+);
+
+static float GNiagaraEmitterBoundsDynamicExpandMultiplier = 1.1f;
+static FAutoConsoleVariableRef CVarNiagaraEmitterBoundsDynamicExpandMultiplier(
+	TEXT("fx.Niagara.EmitterBounds.DynamicExpandMultiplier"),
+	GNiagaraEmitterBoundsDynamicExpandMultiplier,
+	TEXT("Multiplier used on dynamic bounds gathering, i.e. 1 means no change, 1.1 means increase by 10%.\n")
+	TEXT("This value is applied after we calculate any dynamic bounds snapping."),
+	ECVF_Default
+);
+
+static float GNiagaraEmitterBoundsFixedExpandMultiplier = 1.0f;
+static FAutoConsoleVariableRef CVarNiagaraEmitterBoundsFixedExpandMultiplier(
+	TEXT("fx.Niagara.EmitterBounds.FixedExpandMultiplier"),
+	GNiagaraEmitterBoundsFixedExpandMultiplier,
+	TEXT("Multiplier used on fixed bounds gathering, i.e. 1 means no change, 1.1 means increase by 10%."),
 	ECVF_Default
 );
 
@@ -2403,19 +2421,39 @@ void FNiagaraSystemInstance::Tick_Concurrent(bool bEnqueueGPUTickIfNeeded)
 	}
 	else
 	{
-		FBox NewLocalBounds(EForceInit::ForceInit);
+		FBox NewDynamicBounds(EForceInit::ForceInit);
+		FBox NewFixedBounds(EForceInit::ForceInit);
 		for (const auto& Emitter : Emitters)
 		{
-			NewLocalBounds += Emitter->GetBounds();
+			if (Emitter->AreBoundsDynamic())
+			{
+				NewDynamicBounds += Emitter->GetBounds();
+			}
+			else
+			{
+				NewFixedBounds += Emitter->GetBounds();
+			}
 		}
 
-		if (NewLocalBounds.IsValid)
+		LocalBounds = FBox(FVector::ZeroVector, FVector::ZeroVector);
+		if (NewDynamicBounds.IsValid)
 		{
-			LocalBounds = NewLocalBounds.ExpandBy(NewLocalBounds.GetExtent() * GNiagaraBoundsExpandByPercent);
+			FVector Center = NewDynamicBounds.GetCenter();
+			FVector Extent = NewDynamicBounds.GetExtent();
+			if (GNiagaraEmitterBoundsDynamicSnapValue > 1.0f)
+			{
+				Extent.X = FMath::CeilToDouble(Extent.X / GNiagaraEmitterBoundsDynamicSnapValue) * GNiagaraEmitterBoundsDynamicSnapValue;
+				Extent.Y = FMath::CeilToDouble(Extent.Y / GNiagaraEmitterBoundsDynamicSnapValue) * GNiagaraEmitterBoundsDynamicSnapValue;
+				Extent.Z = FMath::CeilToDouble(Extent.Z / GNiagaraEmitterBoundsDynamicSnapValue) * GNiagaraEmitterBoundsDynamicSnapValue;
+			}
+			Extent = Extent * GNiagaraEmitterBoundsDynamicExpandMultiplier;
+			LocalBounds += FBox(Center - Extent, Center + Extent);
 		}
-		else
+		if (NewFixedBounds.IsValid)
 		{
-			LocalBounds = FBox(FVector::ZeroVector, FVector::ZeroVector);
+			const FVector Center = NewFixedBounds.GetCenter();
+			const FVector Extent = NewFixedBounds.GetExtent() * GNiagaraEmitterBoundsFixedExpandMultiplier;
+			LocalBounds += FBox(Center - Extent, Center + Extent);
 		}
 	}
 
