Avoid issues where some RHIs may treat transitions as redundant, leading to problems when a GPU particle simulation is immediately followed by UpdateDrawIndirectBuffers() and a barrier isn't inserted for the instance count buffer (UAV -> SRV|CopySrc followed by SRV|CopySrc -> UAV).

diff --git a/Engine/Plugins/FX/Niagara/Source/Niagara/Private/NiagaraGPUInstanceCountManager.cpp b/Engine/Plugins/FX/Niagara/Source/Niagara/Private/NiagaraGPUInstanceCountManager.cpp
--- a/Engine/Plugins/FX/Niagara/Source/Niagara/Private/NiagaraGPUInstanceCountManager.cpp
+++ b/Engine/Plugins/FX/Niagara/Source/Niagara/Private/NiagaraGPUInstanceCountManager.cpp
@@ -414,7 +414,9 @@ void FNiagaraGPUInstanceCountManager::UpdateDrawIndirectBuffers(FNiagaraGpuCompu
 		const bool bCountBufferIsValid = CountBuffer.UAV.IsValid();
 		if (bCountBufferIsValid)
 		{
-			Transitions.Emplace(CountBuffer.UAV, kCountBufferDefaultState, ERHIAccess::UAVCompute);
+			// treat the incoming UAV as being unknown to be sure a barrier is inserted in the case
+			// where the preceding dispatch wrote to the counts buffer
+			Transitions.Emplace(CountBuffer.UAV, ERHIAccess::Unknown, ERHIAccess::UAVCompute);
 			CountsUAV = CountBuffer.UAV;
 		}
 		else
