- Add display for GPU tick & usage information to HUD / Outliner / System Editor

#rnx

diff --git a/Engine/Plugins/FX/Niagara/Source/Niagara/Private/NiagaraDebugHud.cpp b/Engine/Plugins/FX/Niagara/Source/Niagara/Private/NiagaraDebugHud.cpp
--- a/Engine/Plugins/FX/Niagara/Source/Niagara/Private/NiagaraDebugHud.cpp
+++ b/Engine/Plugins/FX/Niagara/Source/Niagara/Private/NiagaraDebugHud.cpp
@@ -2423,6 +2423,11 @@ void FNiagaraDebugHud::DrawComponents(FNiagaraWorldManager* WorldManager, UCanva
 						{
 							StringBuilder.Appendf(TEXT("TickGroup - %s\n"), *TickingGroupEnum->GetNameStringByValue(SystemInstance->CalculateTickGroup()));
 						}
+						static UEnum* GpuComputeTickStageEnum = StaticEnum<ENiagaraGpuComputeTickStage::Type>();
+						if (const FNiagaraSystemGpuComputeProxy* SystemInstanceComputeProxy = SystemInstance->GetSystemGpuComputeProxy())
+						{
+							StringBuilder.Appendf(TEXT("GpuTickStage - %s\n"), *GpuComputeTickStageEnum->GetNameStringByValue(SystemInstanceComputeProxy->GetComputeTickStage()));
+						}
 					}
 
 					int64 TotalBytes = 0;
diff --git a/Engine/Plugins/FX/Niagara/Source/Niagara/Private/NiagaraDebuggerClient.cpp b/Engine/Plugins/FX/Niagara/Source/Niagara/Private/NiagaraDebuggerClient.cpp
--- a/Engine/Plugins/FX/Niagara/Source/Niagara/Private/NiagaraDebuggerClient.cpp
+++ b/Engine/Plugins/FX/Niagara/Source/Niagara/Private/NiagaraDebuggerClient.cpp
@@ -365,6 +365,10 @@ bool FNiagaraDebuggerClient::UpdateOutliner(float DeltaSeconds)
 					if (FNiagaraSystemInstance* Inst = InstController->GetSystemInstance_Unsafe())
 					{
 						InstData.TickGroup = Inst->CalculateTickGroup();
+						if (const FNiagaraSystemGpuComputeProxy* SystemInstanceComputeProxy = Inst->GetSystemGpuComputeProxy())
+						{
+							InstData.GpuTickStage = SystemInstanceComputeProxy->GetComputeTickStage();
+						}
 						InstData.LWCTile = Inst->GetLWCTile();
 						InstData.bIsSolo = Inst->IsSolo();
 						InstData.bRequiresDistanceFieldData = Inst->RequiresDistanceFieldData();
diff --git a/Engine/Plugins/FX/Niagara/Source/Niagara/Public/NiagaraCommon.h b/Engine/Plugins/FX/Niagara/Source/Niagara/Public/NiagaraCommon.h
--- a/Engine/Plugins/FX/Niagara/Source/Niagara/Public/NiagaraCommon.h
+++ b/Engine/Plugins/FX/Niagara/Source/Niagara/Public/NiagaraCommon.h
@@ -1590,6 +1590,7 @@ struct NIAGARA_API FSynchronizeWithParameterDefinitionsArgs
 };
 
 /** Enum used to track stage that GPU compute proxies will execute in. */
+UENUM()
 namespace ENiagaraGpuComputeTickStage
 {
 	enum Type
@@ -1597,8 +1598,8 @@ namespace ENiagaraGpuComputeTickStage
 		PreInitViews,
 		PostInitViews,
 		PostOpaqueRender,
-		Max,
-		First = PreInitViews,
-		Last = PostOpaqueRender,
+		Max UMETA(hidden),
+		First = PreInitViews UMETA(hidden),
+		Last = PostOpaqueRender UMETA(hidden),
 	};
-};
+}
diff --git a/Engine/Plugins/FX/Niagara/Source/Niagara/Public/NiagaraDebuggerCommon.h b/Engine/Plugins/FX/Niagara/Source/Niagara/Public/NiagaraDebuggerCommon.h
--- a/Engine/Plugins/FX/Niagara/Source/Niagara/Public/NiagaraDebuggerCommon.h
+++ b/Engine/Plugins/FX/Niagara/Source/Niagara/Public/NiagaraDebuggerCommon.h
@@ -103,6 +103,9 @@ struct FNiagaraOutlinerSystemInstanceData
 	UPROPERTY(VisibleAnywhere, Category = "Ticking")
 	TEnumAsByte<ETickingGroup> TickGroup;
 
+	UPROPERTY(VisibleAnywhere, Category = "Ticking")
+	TEnumAsByte<ENiagaraGpuComputeTickStage::Type> GpuTickStage;
+
 	UPROPERTY(VisibleAnywhere, Category = "Gpu")
 	uint32 bIsSolo : 1;
 
@@ -125,6 +128,7 @@ struct FNiagaraOutlinerSystemInstanceData
 		: bPendingKill(false)
 		, bUsingCullProxy(false)
 		, TickGroup(0)
+		, GpuTickStage(ENiagaraGpuComputeTickStage::First)
 		, bIsSolo(false)
 		, bRequiresDistanceFieldData(false)
 		, bRequiresDepthBuffer(false)
@@ -132,10 +136,6 @@ struct FNiagaraOutlinerSystemInstanceData
 		, bRequiresViewUniformBuffer(false)
 		, bRequiresRayTracingScene(false)
 	{}
-
-	//TODO:
-	//Tick info, solo, tick group etc.
-	//Mem usage?
 };
 
 /** Wrapper for array of system instance outliner data so that it can be placed in a map. */
diff --git a/Engine/Plugins/FX/Niagara/Source/NiagaraEditor/Private/NiagaraEditorCommands.cpp b/Engine/Plugins/FX/Niagara/Source/NiagaraEditor/Private/NiagaraEditorCommands.cpp
--- a/Engine/Plugins/FX/Niagara/Source/NiagaraEditor/Private/NiagaraEditorCommands.cpp
+++ b/Engine/Plugins/FX/Niagara/Source/NiagaraEditor/Private/NiagaraEditorCommands.cpp
@@ -23,7 +23,8 @@ void FNiagaraEditorCommands::RegisterCommands()
 	UI_COMMAND(TogglePreviewGrid, "Grid", "Toggles the preview pane's grid.", EUserInterfaceActionType::ToggleButton, FInputChord() );
 	UI_COMMAND(ToggleInstructionCounts, "InstructionCounts", "Display Instruction Counts", EUserInterfaceActionType::ToggleButton, FInputChord());
 	UI_COMMAND(ToggleParticleCounts, "ParticleCounts", "Display Particle Counts", EUserInterfaceActionType::ToggleButton, FInputChord());
-	UI_COMMAND(ToggleEmitterExecutionOrder, "EmitterExecutionOrder", "Display Emitter Execution Order", EUserInterfaceActionType::ToggleButton, FInputChord());
+	UI_COMMAND(ToggleEmitterExecutionOrder, "Emitter Execution Order", "Display Emitter Execution Order", EUserInterfaceActionType::ToggleButton, FInputChord());
+	UI_COMMAND(ToggleGpuTickInformation, "Gpu Tick Information", "Display Gpu Tick Information", EUserInterfaceActionType::ToggleButton, FInputChord());
 	UI_COMMAND(TogglePreviewBackground, "Background", "Toggles the preview pane's background.", EUserInterfaceActionType::ToggleButton, FInputChord());
 	UI_COMMAND(ToggleUnlockToChanges, "Lock/Unlock To Changes", "Toggles whether or not changes in the source asset get pulled into this asset automatically.", EUserInterfaceActionType::Button, FInputChord());
 	UI_COMMAND(ToggleStatPerformance, "Performance", "Show runtime performance for particle scripts.", EUserInterfaceActionType::ToggleButton, FInputChord());
diff --git a/Engine/Plugins/FX/Niagara/Source/NiagaraEditor/Private/NiagaraEditorSettings.cpp b/Engine/Plugins/FX/Niagara/Source/NiagaraEditor/Private/NiagaraEditorSettings.cpp
--- a/Engine/Plugins/FX/Niagara/Source/NiagaraEditor/Private/NiagaraEditorSettings.cpp
+++ b/Engine/Plugins/FX/Niagara/Source/NiagaraEditor/Private/NiagaraEditorSettings.cpp
@@ -349,6 +349,20 @@ void UNiagaraEditorSettings::SetShowEmitterExecutionOrder(bool bInShowEmitterExe
 	}
 }
 
+bool UNiagaraEditorSettings::IsShowGpuTickInformation() const
+{
+	return bShowGpuTickInformation;
+}
+
+void UNiagaraEditorSettings::SetShowGpuTickInformation(bool bInShowGpuTickInformation)
+{
+	if (bShowGpuTickInformation != bInShowGpuTickInformation)
+	{
+		bShowGpuTickInformation = bInShowGpuTickInformation;
+		SaveConfig();
+	}
+}
+
 TArray<float> UNiagaraEditorSettings::GetPlaybackSpeeds() const
 {
 	if(!CachedPlaybackSpeeds.IsSet())
diff --git a/Engine/Plugins/FX/Niagara/Source/NiagaraEditor/Private/Widgets/SNiagaraSystemViewport.cpp b/Engine/Plugins/FX/Niagara/Source/NiagaraEditor/Private/Widgets/SNiagaraSystemViewport.cpp
--- a/Engine/Plugins/FX/Niagara/Source/NiagaraEditor/Private/Widgets/SNiagaraSystemViewport.cpp
+++ b/Engine/Plugins/FX/Niagara/Source/NiagaraEditor/Private/Widgets/SNiagaraSystemViewport.cpp
@@ -16,6 +16,8 @@
 #include "NiagaraPerfBaseline.h"
 #include "NiagaraSettings.h"
 #include "NiagaraSystem.h"
+#include "NiagaraSystemInstance.h"
+#include "NiagaraSystemGpuComputeProxy.h"
 #include "NiagaraSystemEditorData.h"
 #include "SNiagaraSystemViewportToolBar.h"
 #include "UnrealEdGlobals.h"
@@ -57,6 +59,7 @@ public:
 	void DrawInstructionCounts(UNiagaraSystem* ParticleSystem, FCanvas* Canvas, float& CurrentX, float& CurrentY, UFont* Font, const float FontHeight);
 	void DrawParticleCounts(UNiagaraComponent* Component, FCanvas* Canvas, float& CurrentX, float& CurrentY, UFont* Font, const float FontHeight);
 	void DrawEmitterExecutionOrder(UNiagaraComponent* Component, FCanvas* Canvas, float& CurrentX, float& CurrentY, UFont* Font, const float FontHeight);
+	void DrawGpuTickInformation(UNiagaraComponent* Component, FCanvas* Canvas, float& CurrentX, float& CurrentY, UFont* Font, const float FontHeight);
 
 	TWeakPtr<SNiagaraSystemViewport> NiagaraViewportPtr;
 	bool bCaptureScreenShot;
@@ -143,10 +146,17 @@ void FNiagaraSystemViewportClient::Draw(FViewport* InViewport, FCanvas* Canvas)
 		if (NiagaraViewport->GetDrawElement(SNiagaraSystemViewport::EDrawElements::ParticleCounts) && Component)
 		{
 			DrawParticleCounts(Component, Canvas, CurrentX, CurrentY, Font, FontHeight);
+			CurrentY += FontHeight;
 		}
 		if (NiagaraViewport->GetDrawElement(SNiagaraSystemViewport::EDrawElements::EmitterExecutionOrder) && Component)
 		{
 			DrawEmitterExecutionOrder(Component, Canvas, CurrentX, CurrentY, Font, FontHeight);
+			CurrentY += FontHeight;
+		}
+		if (NiagaraViewport->GetDrawElement(SNiagaraSystemViewport::EDrawElements::GpuTickInformation) && Component)
+		{
+			DrawGpuTickInformation(Component, Canvas, CurrentX, CurrentY, Font, FontHeight);
+			CurrentY += FontHeight;
 		}
 	}
 
@@ -317,6 +327,66 @@ void FNiagaraSystemViewportClient::DrawEmitterExecutionOrder(UNiagaraComponent*
 	}
 }
 
+void FNiagaraSystemViewportClient::DrawGpuTickInformation(UNiagaraComponent* Component, FCanvas* Canvas, float& CurrentX, float& CurrentY, UFont* Font, const float FontHeight)
+{
+	FNiagaraSystemInstanceControllerConstPtr SystemInstanceController = Component->GetSystemInstanceController();
+	if (SystemInstanceController.IsValid() == false)
+	{
+		return;
+	}
+
+	FNiagaraSystemInstance* SystemInstance = SystemInstanceController->GetSystemInstance_Unsafe();
+	if (SystemInstance == nullptr)
+	{
+		return;
+	}
+
+	Canvas->DrawShadowedString(CurrentX, CurrentY, TEXT("Gpu Tick Information"), Font, FLinearColor::White);
+	CurrentY += FontHeight;
+
+	if (SystemInstance->IsReadyToRun())
+	{
+		if (const FNiagaraSystemGpuComputeProxy* ComputeProxy = SystemInstance->GetSystemGpuComputeProxy())
+		{
+			static UEnum* GpuComputeTickStageEnum = StaticEnum<ENiagaraGpuComputeTickStage::Type>();
+			const ENiagaraGpuComputeTickStage::Type GpuTickStage = ComputeProxy->GetComputeTickStage();
+			const FLinearColor TextColor = GpuTickStage == ENiagaraGpuComputeTickStage::First ? FLinearColor::White : FLinearColor::Yellow;
+			Canvas->DrawShadowedString(CurrentX + 5.0f, CurrentY, *FString::Printf(TEXT("GpuTickStage = %s"), *GpuComputeTickStageEnum->GetNameStringByValue(GpuTickStage)), Font, TextColor);
+			CurrentY += FontHeight;
+		}
+		else
+		{
+			Canvas->DrawShadowedString(CurrentX + 5.0f, CurrentY, TEXT("No GPU Emitters"), Font, FLinearColor::White);
+			CurrentY += FontHeight;
+		}
+		if (SystemInstance->RequiresDistanceFieldData())
+		{
+			Canvas->DrawShadowedString(CurrentX + 5.0f, CurrentY, TEXT("RequiresDistanceFieldData"), Font, FLinearColor::White);
+			CurrentY += FontHeight;
+		}
+		if (SystemInstance->RequiresDepthBuffer())
+		{
+			Canvas->DrawShadowedString(CurrentX + 5.0f, CurrentY, TEXT("RequiresDepthBuffer"), Font, FLinearColor::White);
+			CurrentY += FontHeight;
+		}
+		if (SystemInstance->RequiresEarlyViewData())
+		{
+			Canvas->DrawShadowedString(CurrentX + 5.0f, CurrentY, TEXT("RequiresEarlyViewData"), Font, FLinearColor::White);
+			CurrentY += FontHeight;
+		}
+		if (SystemInstance->RequiresViewUniformBuffer())
+		{
+			Canvas->DrawShadowedString(CurrentX + 5.0f, CurrentY, TEXT("RequiresViewUniformBuffer"), Font, FLinearColor::White);
+			CurrentY += FontHeight;
+		}
+		if (SystemInstance->RequiresRayTracingScene())
+		{
+			Canvas->DrawShadowedString(CurrentX + 5.0f, CurrentY, TEXT("RequiresRayTracingScene"), Font, FLinearColor::White);
+			CurrentY += FontHeight;
+		}
+	}
+}
+
 void FNiagaraSystemViewportClient::SetOrbitModeFromSettings()
 {
 	const UNiagaraSettings* Settings = GetDefault<UNiagaraSettings>();
@@ -394,6 +464,7 @@ void SNiagaraSystemViewport::Construct(const FArguments& InArgs)
 	DrawFlags |= Settings->IsShowParticleCountsInViewport() ? EDrawElements::ParticleCounts : 0;
 	DrawFlags |= Settings->IsShowInstructionsCount() ? EDrawElements::InstructionCounts : 0;
 	DrawFlags |= Settings->IsShowEmitterExecutionOrder() ? EDrawElements::EmitterExecutionOrder : 0;
+	DrawFlags |= Settings->IsShowGpuTickInformation() ? EDrawElements::GpuTickInformation : 0;
 
 	bShowBackground = false;
 	PreviewComponent = nullptr;
@@ -591,6 +662,18 @@ void SNiagaraSystemViewport::BindCommands()
 		FIsActionChecked::CreateLambda([Viewport = this]() -> bool { return Viewport->GetDrawElement(EDrawElements::EmitterExecutionOrder); })
 	);
 
+	CommandList->MapAction(
+		Commands.ToggleGpuTickInformation,
+		FExecuteAction::CreateLambda([Viewport = this]()
+		{
+			Viewport->ToggleDrawElement(EDrawElements::GpuTickInformation);
+			GetMutableDefault<UNiagaraEditorSettings>()->SetShowGpuTickInformation(Viewport->GetDrawElement(EDrawElements::GpuTickInformation));
+			Viewport->RefreshViewport();
+		}),
+		FCanExecuteAction(),
+		FIsActionChecked::CreateLambda([Viewport = this]() -> bool { return Viewport->GetDrawElement(EDrawElements::GpuTickInformation); })
+	);
+
 	CommandList->MapAction(
 		Commands.TogglePreviewBackground,
 		FExecuteAction::CreateSP( this, &SNiagaraSystemViewport::TogglePreviewBackground ),
diff --git a/Engine/Plugins/FX/Niagara/Source/NiagaraEditor/Private/Widgets/SNiagaraSystemViewport.h b/Engine/Plugins/FX/Niagara/Source/NiagaraEditor/Private/Widgets/SNiagaraSystemViewport.h
--- a/Engine/Plugins/FX/Niagara/Source/NiagaraEditor/Private/Widgets/SNiagaraSystemViewport.h
+++ b/Engine/Plugins/FX/Niagara/Source/NiagaraEditor/Private/Widgets/SNiagaraSystemViewport.h
@@ -84,6 +84,7 @@ public:
 		InstructionCounts = 0x040,
 		ParticleCounts = 0x080,
 		EmitterExecutionOrder = 0x100,
+		GpuTickInformation = 0x200,
 	};
 
 	bool GetDrawElement(EDrawElements Element) const;
diff --git a/Engine/Plugins/FX/Niagara/Source/NiagaraEditor/Private/Widgets/SNiagaraSystemViewportToolBar.cpp b/Engine/Plugins/FX/Niagara/Source/NiagaraEditor/Private/Widgets/SNiagaraSystemViewportToolBar.cpp
--- a/Engine/Plugins/FX/Niagara/Source/NiagaraEditor/Private/Widgets/SNiagaraSystemViewportToolBar.cpp
+++ b/Engine/Plugins/FX/Niagara/Source/NiagaraEditor/Private/Widgets/SNiagaraSystemViewportToolBar.cpp
@@ -42,6 +42,7 @@ TSharedRef<SWidget> SNiagaraSystemViewportToolBar::GenerateShowMenu() const
 		ShowMenuBuilder.AddMenuEntry(Commands.ToggleInstructionCounts);
 		ShowMenuBuilder.AddMenuEntry(Commands.ToggleParticleCounts);
 		ShowMenuBuilder.AddMenuEntry(Commands.ToggleEmitterExecutionOrder);
+		ShowMenuBuilder.AddMenuEntry(Commands.ToggleGpuTickInformation);
 		//ShowMenuBuilder.AddMenuEntry(Commands.TogglePreviewBackground);
 	}
 
diff --git a/Engine/Plugins/FX/Niagara/Source/NiagaraEditor/Public/NiagaraEditorCommands.h b/Engine/Plugins/FX/Niagara/Source/NiagaraEditor/Public/NiagaraEditorCommands.h
--- a/Engine/Plugins/FX/Niagara/Source/NiagaraEditor/Public/NiagaraEditorCommands.h
+++ b/Engine/Plugins/FX/Niagara/Source/NiagaraEditor/Public/NiagaraEditorCommands.h
@@ -39,6 +39,7 @@ public:
 	TSharedPtr<FUICommandInfo> ToggleInstructionCounts;
 	TSharedPtr<FUICommandInfo> ToggleParticleCounts;
 	TSharedPtr<FUICommandInfo> ToggleEmitterExecutionOrder;
+	TSharedPtr<FUICommandInfo> ToggleGpuTickInformation;
 
 	/** Toggles the preview pane's background */
 	TSharedPtr< FUICommandInfo > TogglePreviewBackground;
diff --git a/Engine/Plugins/FX/Niagara/Source/NiagaraEditor/Public/NiagaraEditorSettings.h b/Engine/Plugins/FX/Niagara/Source/NiagaraEditor/Public/NiagaraEditorSettings.h
--- a/Engine/Plugins/FX/Niagara/Source/NiagaraEditor/Public/NiagaraEditorSettings.h
+++ b/Engine/Plugins/FX/Niagara/Source/NiagaraEditor/Public/NiagaraEditorSettings.h
@@ -398,9 +398,15 @@ private:
 
 	UPROPERTY(config)
 	bool bShowEmitterExecutionOrder;
+
+	UPROPERTY(config)
+	bool bShowGpuTickInformation;
+
 public:
 	bool IsShowInstructionsCount() const;
 	void SetShowInstructionsCount(bool bShowInstructionsCount);
 	bool IsShowEmitterExecutionOrder() const;
 	void SetShowEmitterExecutionOrder(bool bShowEmitterExecutionOrder);
+	bool IsShowGpuTickInformation() const;
+	void SetShowGpuTickInformation(bool bShowGpuTickInformation);
 };
