- Track registered components on debug HUD and optionally show information

diff --git a/Engine/Plugins/FX/Niagara/Source/Niagara/Private/NiagaraDebugHud.cpp b/Engine/Plugins/FX/Niagara/Source/Niagara/Private/NiagaraDebugHud.cpp
--- a/Engine/Plugins/FX/Niagara/Source/Niagara/Private/NiagaraDebugHud.cpp
+++ b/Engine/Plugins/FX/Niagara/Source/Niagara/Private/NiagaraDebugHud.cpp
@@ -664,7 +664,8 @@ void FNiagaraDebugHud::GatherSystemInfo()
 {
 	using namespace NiagaraDebugLocal;
 
-	GlobalTotalSystems = 0;
+	GlobalTotalRegistered = 0;
+	GlobalTotalActive = 0;
 	GlobalTotalScalability = 0;
 	GlobalTotalEmitters = 0;
 	GlobalTotalParticles = 0;
@@ -799,6 +800,12 @@ void FNiagaraDebugHud::GatherSystemInfo()
 			}
 		}
 
+		if (NiagaraComponent->IsRegistered())
+		{
+			++GlobalTotalRegistered;
+			++SystemDebugInfo.TotalRegistered;
+		}
+
 		if ( bHasScalability )
 		{
 			++GlobalTotalScalability;
@@ -884,11 +891,11 @@ void FNiagaraDebugHud::GatherSystemInfo()
 				ActiveParticles += EmitterInstance->GetNumParticles();
 			}
 
-			++SystemDebugInfo.TotalSystems;
+			++SystemDebugInfo.TotalActive;
 			SystemDebugInfo.TotalEmitters += ActiveEmitters;
 			SystemDebugInfo.TotalParticles += ActiveParticles;
 
-			++GlobalTotalSystems;
+			++GlobalTotalActive;
 			GlobalTotalEmitters += ActiveEmitters;
 			GlobalTotalParticles += ActiveParticles;
 
@@ -912,11 +919,15 @@ void FNiagaraDebugHud::GatherSystemInfo()
 		
 		if(Settings.OverviewMode == ENiagaraDebugHUDOverviewMode::Performance)
 		{
-			bWillBeVisible = SystemDebugInfo.TotalSystems > 0;
+			bWillBeVisible = SystemDebugInfo.TotalActive > 0;
 		}
 		else
 		{
-			bWillBeVisible = (SystemDebugInfo.TotalSystems > 0 || SystemDebugInfo.TotalScalability > 0);
+			bWillBeVisible = (SystemDebugInfo.TotalActive > 0 || SystemDebugInfo.TotalScalability > 0);
+			if ( Settings.OverviewMode == ENiagaraDebugHUDOverviewMode::Overview )
+			{
+				bWillBeVisible |= Settings.bShowRegisteredComponents && SystemDebugInfo.TotalRegistered > 0;
+			}
 		}
 		if (bWillBeVisible)
 		{
@@ -1239,14 +1250,27 @@ void FNiagaraDebugHud::DrawOverview(class FNiagaraWorldManager* WorldManager, FC
 
 		if (Settings.OverviewMode == ENiagaraDebugHUDOverviewMode::Overview)
 		{
-			OverviewColumns.Emplace(TEXT("TotalSystems:"), FString::FromInt(GlobalTotalSystems), TEXT("# Active"), ColumnOffset, Font, TEXT("0000"),
+			if (Settings.bShowRegisteredComponents)
+			{
+				OverviewColumns.Emplace(TEXT("TotalRegistered:"), FString::FromInt(GlobalTotalRegistered), TEXT("# Registered"), ColumnOffset, Font, TEXT("0000"),
+					[&DetailColor, &DetailHighlightColor, &fAdvanceHeight](FCanvas* Canvas, UFont* Font, float X, float Y, FOverviewColumn& Col, const FSystemDebugInfo& SystemInfo)
+				{
+					FLinearColor RowBGColor = SystemInfo.UniqueColor;
+					RowBGColor.A = Settings.SystemColorTableOpacity;
+					Canvas->DrawTile(X, Y, Col.MaxWidth, fAdvanceHeight, 0, 0, 0, 0, RowBGColor);
+					const FLinearColor RowColor = SystemInfo.bShowInWorld ? DetailHighlightColor : DetailColor;
+					Canvas->DrawShadowedString(X, Y, *FString::FromInt(SystemInfo.TotalRegistered), Font, RowColor);
+				});
+			}
+
+			OverviewColumns.Emplace(TEXT("TotalActive:"), FString::FromInt(GlobalTotalActive), TEXT("# Active"), ColumnOffset, Font, TEXT("0000"),
 				[&DetailColor, &DetailHighlightColor, &fAdvanceHeight](FCanvas* Canvas, UFont* Font, float X, float Y, FOverviewColumn& Col, const FSystemDebugInfo& SystemInfo)
 				{
 					FLinearColor RowBGColor = SystemInfo.UniqueColor;
 					RowBGColor.A = Settings.SystemColorTableOpacity;
 					Canvas->DrawTile(X, Y, Col.MaxWidth, fAdvanceHeight, 0, 0, 0, 0, RowBGColor);
 					const FLinearColor RowColor = SystemInfo.bShowInWorld ? DetailHighlightColor : DetailColor;
-					Canvas->DrawShadowedString(X, Y, *FString::FromInt(SystemInfo.TotalSystems), Font, RowColor);
+					Canvas->DrawShadowedString(X, Y, *FString::FromInt(SystemInfo.TotalActive), Font, RowColor);
 				});
 
 			OverviewColumns.Emplace(TEXT("TotalScalability:"), FString::FromInt(GlobalTotalScalability), TEXT("# Scalability"), ColumnOffset, Font, TEXT("0000"),
@@ -1291,14 +1315,14 @@ void FNiagaraDebugHud::DrawOverview(class FNiagaraWorldManager* WorldManager, FC
 		}
 		else if (Settings.OverviewMode == ENiagaraDebugHUDOverviewMode::Scalability)
 		{
-			OverviewColumns.Emplace(TEXT("TotalSystems:"), FString::FromInt(GlobalTotalSystems), TEXT("# Active"), ColumnOffset, Font, TEXT("0000"),
+			OverviewColumns.Emplace(TEXT("TotalActive:"), FString::FromInt(GlobalTotalActive), TEXT("# Active"), ColumnOffset, Font, TEXT("0000"),
 				[&DetailColor, &DetailHighlightColor, &fAdvanceHeight](FCanvas* Canvas, UFont* Font, float X, float Y, FOverviewColumn& Col, const FSystemDebugInfo& SystemInfo)
 				{
 					FLinearColor RowBGColor = SystemInfo.UniqueColor;
 					RowBGColor.A = Settings.SystemColorTableOpacity;
 					Canvas->DrawTile(X, Y, Col.MaxWidth, fAdvanceHeight, 0, 0, 0, 0, RowBGColor);
 					const FLinearColor RowColor = SystemInfo.bShowInWorld ? DetailHighlightColor : DetailColor;
-					Canvas->DrawShadowedString(X, Y, *FString::FromInt(SystemInfo.TotalSystems), Font, RowColor);
+					Canvas->DrawShadowedString(X, Y, *FString::FromInt(SystemInfo.TotalActive), Font, RowColor);
 				});
 
 			OverviewColumns.Emplace(TEXT("TotalScalability:"), FString::FromInt(GlobalTotalScalability), TEXT("# Scalability"), ColumnOffset, Font, TEXT("0000"),
@@ -1444,7 +1468,7 @@ void FNiagaraDebugHud::DrawOverview(class FNiagaraWorldManager* WorldManager, FC
 					[&DetailColor, &DetailHighlightColor, &fAdvanceHeight](FCanvas* Canvas, UFont* Font, float X, float Y, FOverviewColumn& Col, const FSystemDebugInfo& SystemInfo)
 					{
 						Canvas->DrawTile(X, Y, Col.MaxWidth - 6.0f, fAdvanceHeight - 2.0f, 0.0f, 0.0f, 0.0f, 0.0f, SystemInfo.UniqueColor);
-						FString SysCountString = LexToSanitizedString(SystemInfo.TotalSystems);
+						FString SysCountString = LexToSanitizedString(SystemInfo.TotalActive);
 						float StringWidth = Font->GetStringSize(*SysCountString);
 						Canvas->DrawShadowedString((X + (Col.MaxWidth * 0.5f)) - StringWidth * 0.5f, Y, *SysCountString, Font, FLinearColor::Black);
 					});
diff --git a/Engine/Plugins/FX/Niagara/Source/Niagara/Private/NiagaraDebugHud.h b/Engine/Plugins/FX/Niagara/Source/Niagara/Private/NiagaraDebugHud.h
--- a/Engine/Plugins/FX/Niagara/Source/Niagara/Private/NiagaraDebugHud.h
+++ b/Engine/Plugins/FX/Niagara/Source/Niagara/Private/NiagaraDebugHud.h
@@ -99,7 +99,8 @@ class FNiagaraDebugHud
 
 		bool		bShowInWorld = false;
 		bool		bPassesSystemFilter = true;
-		int32		TotalSystems = 0;
+		int32		TotalRegistered = 0;
+		int32		TotalActive = 0;
 		int32		TotalScalability = 0;
 		int32		TotalEmitters = 0;
 		int32		TotalParticles = 0;
@@ -117,7 +118,8 @@ class FNiagaraDebugHud
 		{
 			bShowInWorld = false;
 			bPassesSystemFilter = true;
-			TotalSystems = 0;
+			TotalRegistered = 0;
+			TotalActive = 0;
 			TotalScalability = 0;
 			TotalEmitters = 0;
 			TotalParticles = 0;
@@ -241,7 +243,8 @@ private:
 private:
 	TWeakObjectPtr<class UWorld>	WeakWorld;
 
-	int32 GlobalTotalSystems = 0;
+	int32 GlobalTotalRegistered = 0;
+	int32 GlobalTotalActive = 0;
 	int32 GlobalTotalScalability = 0;
 	int32 GlobalTotalEmitters = 0;
 	int32 GlobalTotalParticles = 0;
diff --git a/Engine/Plugins/FX/Niagara/Source/Niagara/Public/NiagaraDebuggerCommon.h b/Engine/Plugins/FX/Niagara/Source/Niagara/Public/NiagaraDebuggerCommon.h
--- a/Engine/Plugins/FX/Niagara/Source/Niagara/Public/NiagaraDebuggerCommon.h
+++ b/Engine/Plugins/FX/Niagara/Source/Niagara/Public/NiagaraDebuggerCommon.h
@@ -464,6 +464,10 @@ struct NIAGARA_API FNiagaraDebugHUDSettingsData
 	UPROPERTY(EditAnywhere, Category = "Debug Overview", meta = (DisplayName = "Debug Overview Text Location", EditCondition = "bOverviewEnabled"))
 	FVector2D OverviewLocation = FIntPoint(30.0f, 150.0f);
 
+	/** Overview display font to use. */
+	UPROPERTY(EditAnywhere, Category = "Debug Overview", meta = (EditCondition = "bOverviewEnabled && OverviewMode == ENiagaraDebugHUDOverviewMode::Overview"))
+	bool bShowRegisteredComponents = false;
+	
 	/** When enabled the overview will only show the filter system information. */
 	UPROPERTY(EditAnywhere, Category = "Debug Overview", meta = (EditCondition = "bOverviewEnabled"))
 	bool bOverviewShowFilteredSystemOnly = false;
