Fixes for summary view including emitter inheritance and "Show Summary", as well as not being able to edit parameters in edit mode if summary view is enabled

diff --git a/Engine/Plugins/FX/Niagara/Source/Niagara/Public/NiagaraTypes.h b/Engine/Plugins/FX/Niagara/Source/Niagara/Public/NiagaraTypes.h
--- a/Engine/Plugins/FX/Niagara/Source/Niagara/Public/NiagaraTypes.h
+++ b/Engine/Plugins/FX/Niagara/Source/Niagara/Public/NiagaraTypes.h
@@ -1332,9 +1332,7 @@ const FNiagaraTypeDefinition& FNiagaraTypeDefinition::Get()
 
 FORCEINLINE uint32 GetTypeHash(const FNiagaraTypeDefinition& Type)
 {
-	uint32 Result= HashCombine(GetTypeHash(Type.GetStruct()), GetTypeHash(Type.GetEnum()));
-	Result = HashCombine(Result, GetTypeHash(Type.GetFlags()));
-	return Result;
+	return HashCombine(HashCombine(GetTypeHash(Type.ClassStructOrEnum), GetTypeHash(Type.UnderlyingType)), GetTypeHash(Type.GetFlags()));
 }
 
 //////////////////////////////////////////////////////////////////////////
@@ -1652,7 +1650,6 @@ private:
 };
 
 //////////////////////////////////////////////////////////////////////////
-
 USTRUCT()
 struct FNiagaraVariableBase
 {
@@ -1977,6 +1974,7 @@ private:
 	TArray<uint8> VarData;
 };
 
+
 template<>
 struct TStructOpsTypeTraits<FNiagaraVariable> : public TStructOpsTypeTraitsBase2<FNiagaraVariable>
 {
diff --git a/Engine/Plugins/FX/Niagara/Source/NiagaraEditor/Private/NiagaraEmitterEditorData.cpp b/Engine/Plugins/FX/Niagara/Source/NiagaraEditor/Private/NiagaraEmitterEditorData.cpp
--- a/Engine/Plugins/FX/Niagara/Source/NiagaraEditor/Private/NiagaraEmitterEditorData.cpp
+++ b/Engine/Plugins/FX/Niagara/Source/NiagaraEditor/Private/NiagaraEmitterEditorData.cpp
@@ -81,17 +81,22 @@ void UNiagaraEmitterEditorData::SetSummaryViewMetaData(const FFunctionInputSumma
 	OnSummaryViewStateChangedDelegate.Broadcast();
 }
 
-void UNiagaraEmitterEditorData::ToggleShowSummaryView()
+void UNiagaraEmitterEditorData::SetShowSummaryView(bool bInShouldShowSummaryView)
 {
-	FScopedTransaction ScopedTransaction(NSLOCTEXT("NiagaraEmitter", "EmitterModuleShowSummaryChanged", "Emitter summary view enabled/disabled."));
-	Modify();
-	
-	bShowSummaryView = !bShowSummaryView;
+	bShowSummaryView = bInShouldShowSummaryView;
 	
 	OnPersistentDataChanged().Broadcast();
 	OnSummaryViewStateChangedDelegate.Broadcast();
 }
 
+void UNiagaraEmitterEditorData::ToggleShowSummaryView()
+{
+	FScopedTransaction ScopedTransaction(NSLOCTEXT("NiagaraEmitter", "EmitterModuleShowSummaryChanged", "Emitter summary view enabled/disabled."));
+	Modify();
+
+	SetShowSummaryView(!bShowSummaryView);
+}
+
 FSimpleMulticastDelegate& UNiagaraEmitterEditorData::OnSummaryViewStateChanged()
 {
 	return OnSummaryViewStateChangedDelegate;
diff --git a/Engine/Plugins/FX/Niagara/Source/NiagaraEditor/Private/NiagaraScriptMergeManager.cpp b/Engine/Plugins/FX/Niagara/Source/NiagaraEditor/Private/NiagaraScriptMergeManager.cpp
--- a/Engine/Plugins/FX/Niagara/Source/NiagaraEditor/Private/NiagaraScriptMergeManager.cpp
+++ b/Engine/Plugins/FX/Niagara/Source/NiagaraEditor/Private/NiagaraScriptMergeManager.cpp
@@ -1121,7 +1121,8 @@ bool FNiagaraEmitterDiffResults::IsEmpty() const
 		ModifiedInputSummaryEntries.Num() == 0 &&
 		ModifiedOtherInputSummaryEntries.Num() == 0 &&
 		ModifiedStackEntryDisplayNames.Num() == 0 &&
-		bScratchPadModified == false;
+		bScratchPadModified == false &&
+		NewShouldShowSummaryViewValue.IsSet() == false;
 }
 
 void FNiagaraEmitterDiffResults::AddError(FText ErrorMessage)
@@ -2300,6 +2301,11 @@ void FNiagaraScriptMergeManager::DiffEmitterSummary(const UNiagaraEmitterEditorD
 			DiffResults.ModifiedOtherSummarySections.Add(CommonValuePair.OtherValue);
 		}
 	}
+
+	if (BaseEditorData != nullptr && OtherEditorData != nullptr && BaseEditorData->ShouldShowSummaryView() != OtherEditorData->ShouldShowSummaryView())
+	{
+		DiffResults.NewShouldShowSummaryViewValue = OtherEditorData->ShouldShowSummaryView();
+	}	
 }
 
 void FNiagaraScriptMergeManager::DiffScriptStacks(TSharedRef<FNiagaraScriptStackMergeAdapter> BaseScriptStackAdapter, TSharedRef<FNiagaraScriptStackMergeAdapter> OtherScriptStackAdapter, FNiagaraScriptStackDiffResults& DiffResults) const
@@ -3483,6 +3489,11 @@ FNiagaraScriptMergeManager::FApplyDiffResults FNiagaraScriptMergeManager::ApplyE
 
 	EditorData->SetSummarySections(SummarySections);
 
+	if (DiffResults.NewShouldShowSummaryViewValue.IsSet())
+	{
+		EditorData->SetShowSummaryView(DiffResults.NewShouldShowSummaryViewValue.GetValue());
+	}
+
 	FApplyDiffResults Results;
 	Results.bSucceeded = true;
 	Results.bModifiedGraph = false;
diff --git a/Engine/Plugins/FX/Niagara/Source/NiagaraEditor/Private/NiagaraScriptMergeManager.h b/Engine/Plugins/FX/Niagara/Source/NiagaraEditor/Private/NiagaraScriptMergeManager.h
--- a/Engine/Plugins/FX/Niagara/Source/NiagaraEditor/Private/NiagaraScriptMergeManager.h
+++ b/Engine/Plugins/FX/Niagara/Source/NiagaraEditor/Private/NiagaraScriptMergeManager.h
@@ -391,6 +391,8 @@ struct FNiagaraEmitterDiffResults
 	TArray<FNiagaraStackSection> AddedOtherSummarySections;
 	TArray<FNiagaraStackSection> ModifiedBaseSummarySections;
 	TArray<FNiagaraStackSection> ModifiedOtherSummarySections;
+
+	TOptional<bool> NewShouldShowSummaryViewValue;
 	
 	bool bScratchPadModified;
 
diff --git a/Engine/Plugins/FX/Niagara/Source/NiagaraEditor/Public/NiagaraEmitterEditorData.h b/Engine/Plugins/FX/Niagara/Source/NiagaraEditor/Public/NiagaraEmitterEditorData.h
--- a/Engine/Plugins/FX/Niagara/Source/NiagaraEditor/Public/NiagaraEmitterEditorData.h
+++ b/Engine/Plugins/FX/Niagara/Source/NiagaraEditor/Public/NiagaraEmitterEditorData.h
@@ -36,6 +36,7 @@ public:
 	NIAGARAEDITOR_API FSimpleMulticastDelegate& OnSummaryViewStateChanged();
 	
 	NIAGARAEDITOR_API bool ShouldShowSummaryView() const { return bShowSummaryView; }
+	NIAGARAEDITOR_API void SetShowSummaryView(bool bInShouldShowSummaryView);
 	NIAGARAEDITOR_API void ToggleShowSummaryView();
 	
 	NIAGARAEDITOR_API const TMap<FFunctionInputSummaryViewKey, FFunctionInputSummaryViewMetadata>& GetSummaryViewMetaDataMap() const;
