Fixed numeric input pins accepting multiple connections and improved tooltips

diff --git a/Engine/Plugins/FX/Niagara/Source/NiagaraEditor/Private/EdGraphSchema_Niagara.cpp b/Engine/Plugins/FX/Niagara/Source/NiagaraEditor/Private/EdGraphSchema_Niagara.cpp
--- a/Engine/Plugins/FX/Niagara/Source/NiagaraEditor/Private/EdGraphSchema_Niagara.cpp
+++ b/Engine/Plugins/FX/Niagara/Source/NiagaraEditor/Private/EdGraphSchema_Niagara.cpp
@@ -57,12 +57,21 @@ const FLinearColor UEdGraphSchema_Niagara::NodeTitleColor_Event = FLinearColor::
 const FLinearColor UEdGraphSchema_Niagara::NodeTitleColor_TranslatorConstant = FLinearColor::Gray;
 const FLinearColor UEdGraphSchema_Niagara::NodeTitleColor_RapidIteration = FLinearColor::Black;
 
+const FText UEdGraphSchema_Niagara::ReplaceExistingInputConnectionsText = LOCTEXT("BreakExistingConnectionsText", "Replace existing input connections.");
+const FText UEdGraphSchema_Niagara::TypesAreNotCompatibleText = LOCTEXT("TypesNotCompatible", "Types are not compatible");
+const FText UEdGraphSchema_Niagara::ConvertText = LOCTEXT("ConvertTypesText", "Convert {0} to {1}");
+const FText UEdGraphSchema_Niagara::ConvertLossyText = LOCTEXT("ConvertTypesLossyText", "Convert {0} to {1} implicitly. Be aware that this can be a lossy conversion.");
+const FText UEdGraphSchema_Niagara::PinNotConnectableText = LOCTEXT("PinNotConnectable", "Pin is not connectable");
+const FText UEdGraphSchema_Niagara::SameNodeConnectionForbiddenText = LOCTEXT("SameNodeNotAllowed", "Both pins are on the same node");
+const FText UEdGraphSchema_Niagara::DirectionsNotCompatibleText = LOCTEXT("DirectionsNotCompatible", "Pin directions are incompatible");
+const FText UEdGraphSchema_Niagara::AddPinIncompatibleTypeText = LOCTEXT("AddPinCompatibleConnection", "Cannot make connections to or from add pins for non-parameter types");
+const FText UEdGraphSchema_Niagara::CircularConnectionFoundText = LOCTEXT("CircularConnectionFound", "Circular connection found");
+
 const FName UEdGraphSchema_Niagara::PinCategoryType("Type");
 const FName UEdGraphSchema_Niagara::PinCategoryMisc("Misc");
 const FName UEdGraphSchema_Niagara::PinCategoryClass("Class");
 const FName UEdGraphSchema_Niagara::PinCategoryEnum("Enum");
 
-
 const FName UEdGraphSchema_Niagara::PinCategoryStaticType("StaticType");
 const FName UEdGraphSchema_Niagara::PinCategoryStaticClass("StaticClass");
 const FName UEdGraphSchema_Niagara::PinCategoryStaticEnum("StaticEnum");
@@ -950,22 +959,29 @@ const FPinConnectionResponse UEdGraphSchema_Niagara::CanCreateConnection(const U
 	// Make sure the pins are not on the same node
 	if (PinA->GetOwningNode() == PinB->GetOwningNode())
 	{
-		return FPinConnectionResponse(CONNECT_RESPONSE_DISALLOW, TEXT("Both are on the same node"));
+		return FPinConnectionResponse(CONNECT_RESPONSE_DISALLOW, SameNodeConnectionForbiddenText);
 	}
 
 	// Check both pins support connections
 	if(PinA->bNotConnectable || PinB->bNotConnectable)
 	{
-		return FPinConnectionResponse(CONNECT_RESPONSE_DISALLOW, TEXT("Pin doesn't support connections."));
+		return FPinConnectionResponse(CONNECT_RESPONSE_DISALLOW, PinNotConnectableText);
 	}
-
+	
 	// Compare the directions
 	const UEdGraphPin* InputPin = NULL;
 	const UEdGraphPin* OutputPin = NULL;
 
 	if (!CategorizePinsByDirection(PinA, PinB, /*out*/ InputPin, /*out*/ OutputPin))
 	{
-		return FPinConnectionResponse(CONNECT_RESPONSE_DISALLOW, TEXT("Directions are not compatible"));
+		return FPinConnectionResponse(CONNECT_RESPONSE_DISALLOW, DirectionsNotCompatibleText);
+	}
+	
+	ECanCreateConnectionResponse AllowConnectionResponse = CONNECT_RESPONSE_MAKE;
+	bool bInputPinHasConnection = InputPin->LinkedTo.Num() > 0;
+	if(bInputPinHasConnection)
+	{
+		AllowConnectionResponse = (InputPin == PinA) ? CONNECT_RESPONSE_BREAK_OTHERS_A : CONNECT_RESPONSE_BREAK_OTHERS_B; 
 	}
 
 	// Do not allow making connections off of dynamic add pins to non parameter map associated pins 
@@ -986,14 +1002,14 @@ const FPinConnectionResponse UEdGraphSchema_Niagara::CanCreateConnection(const U
 
 	if (GetPinsAreInvalidAddPinCombination(PinA, PinB) || GetPinsAreInvalidAddPinCombination(PinB, PinA))
 	{
-		return FPinConnectionResponse(CONNECT_RESPONSE_DISALLOW, TEXT("Cannot make connections to or from add pins for non-parameter types"));
+		return FPinConnectionResponse(CONNECT_RESPONSE_DISALLOW, AddPinIncompatibleTypeText);
 	}
 
 	// Check for a circular connection before checking any type compatibility
 	TSet<const UEdGraphNode*> VisitedNodes;
 	if (UEdGraphSchema_Niagara::CheckCircularConnection(VisitedNodes, OutputPin->GetOwningNode(), InputPin->GetOwningNode()))
 	{
-		return FPinConnectionResponse(CONNECT_RESPONSE_DISALLOW, TEXT("Circular connection found"));
+		return FPinConnectionResponse(CONNECT_RESPONSE_DISALLOW, CircularConnectionFoundText);
 	}
 
 	if (!IsPinWildcard(PinA) && !IsPinWildcard(PinB))
@@ -1005,22 +1021,27 @@ const FPinConnectionResponse UEdGraphSchema_Niagara::CanCreateConnection(const U
 		{
 			FNiagaraTypeDefinition PinTypeInput = PinToTypeDefinition(InputPin);
 			FNiagaraTypeDefinition PinTypeOutput = PinToTypeDefinition(OutputPin);
-
+			
 			if (PinTypeInput == FNiagaraTypeDefinition::GetParameterMapDef() || PinTypeOutput == FNiagaraTypeDefinition::GetParameterMapDef())
 			{
-				return FPinConnectionResponse(CONNECT_RESPONSE_DISALLOW, TEXT("Types are not compatible"));
+				return FPinConnectionResponse(CONNECT_RESPONSE_DISALLOW, TypesAreNotCompatibleText);
 			}
 			else if (FNiagaraTypeDefinition::TypesAreAssignable(PinTypeInput, PinTypeOutput, GetDefault<UNiagaraSettings>()->bEnforceStrictStackTypes == false))
 			{
 				// if we have a lossy conversion but are assignable, let the user connect but inform him
 				if(FNiagaraTypeDefinition::IsLossyConversion(PinTypeOutput, PinTypeInput))
 				{
-					return FPinConnectionResponse(CONNECT_RESPONSE_MAKE, FString::Printf(TEXT("Convert %s to %s implicitly. Be aware that this can be a lossy conversion."), *(PinTypeOutput.GetNameText().ToString()), *(PinTypeInput.GetNameText().ToString())));
+					FText ConvertMessage = FText::Format(ConvertLossyText, PinTypeOutput.GetNameText(), PinTypeInput.GetNameText());
+					if(bInputPinHasConnection)
+					{
+						ConvertMessage = FText::FromString(ReplaceExistingInputConnectionsText.ToString() + TEXT("\n") + ConvertMessage.ToString()); 
+					}
+					return FPinConnectionResponse(AllowConnectionResponse, ConvertMessage);
 				}
 				// if we are assignable and not lossy, we simply allow the connection
 				else
 				{
-					return FPinConnectionResponse(CONNECT_RESPONSE_MAKE, FString());
+					return FPinConnectionResponse(AllowConnectionResponse, bInputPinHasConnection ? ReplaceExistingInputConnectionsText.ToString() : FString());
 				}
 			}
 			// if the types are not directly assignable at all, give the user a chance to use a conversion node
@@ -1029,11 +1050,16 @@ const FPinConnectionResponse UEdGraphSchema_Niagara::CanCreateConnection(const U
 				//Do some limiting on auto conversions here?
 				if (PinTypeInput.GetClass() || PinTypeInput.IsStatic() || PinTypeOutput.IsStatic())
 				{
-					return FPinConnectionResponse(CONNECT_RESPONSE_DISALLOW, TEXT("Types are not compatible"));
+					return FPinConnectionResponse(CONNECT_RESPONSE_DISALLOW, TypesAreNotCompatibleText);
 				}
 				else
 				{
-					return FPinConnectionResponse(CONNECT_RESPONSE_MAKE_WITH_CONVERSION_NODE, FString::Printf(TEXT("Convert %s to %s"), *(PinTypeOutput.GetNameText().ToString()), *(PinTypeInput.GetNameText().ToString())));
+					FText ConvertMessage = FText::Format(ConvertText, PinTypeOutput.GetNameText(), PinTypeInput.GetNameText());
+					if(bInputPinHasConnection)
+					{
+						ConvertMessage = FText::FromString(ReplaceExistingInputConnectionsText.ToString() + TEXT("\n") + ConvertMessage.ToString()); 
+					}
+					return FPinConnectionResponse(CONNECT_RESPONSE_MAKE_WITH_CONVERSION_NODE, ConvertMessage);
 				}
 			}
 		}
@@ -1066,7 +1092,7 @@ const FPinConnectionResponse UEdGraphSchema_Niagara::CanCreateConnection(const U
 
 			if (bPinAIsAddAndAcceptsPinB == false && bPinBIsAddAndAcceptsPinA == false)
 			{
-				return FPinConnectionResponse(CONNECT_RESPONSE_DISALLOW, TEXT("Types are not compatible"));
+				return FPinConnectionResponse(CONNECT_RESPONSE_DISALLOW, TypesAreNotCompatibleText);
 			}
 		}
 		else
@@ -1080,7 +1106,7 @@ const FPinConnectionResponse UEdGraphSchema_Niagara::CanCreateConnection(const U
 
 				if (AType != BType)
 				{
-					return FPinConnectionResponse(CONNECT_RESPONSE_DISALLOW, TEXT("Types are not compatible"));
+					return FPinConnectionResponse(CONNECT_RESPONSE_DISALLOW, TypesAreNotCompatibleText);
 				}
 			}
 
@@ -1091,7 +1117,7 @@ const FPinConnectionResponse UEdGraphSchema_Niagara::CanCreateConnection(const U
 				FNiagaraTypeDefinition PinTypeOutput = PinToTypeDefinition(OutputPin);
 				if (FNiagaraTypeDefinition::TypesAreAssignable(PinTypeInput, PinTypeOutput) == false)
 				{
-					return FPinConnectionResponse(CONNECT_RESPONSE_DISALLOW, TEXT("Types are not compatible"));
+					return FPinConnectionResponse(CONNECT_RESPONSE_DISALLOW, TypesAreNotCompatibleText);
 				}
 			}
 		}
@@ -1107,8 +1133,7 @@ const FPinConnectionResponse UEdGraphSchema_Niagara::CanCreateConnection(const U
 	const bool bBreakExistingDueToDataInput = (InputPin->LinkedTo.Num() > 0);
 	if (bBreakExistingDueToDataInput)
 	{
-		const ECanCreateConnectionResponse ReplyBreakInputs = (PinA == InputPin) ? CONNECT_RESPONSE_BREAK_OTHERS_A : CONNECT_RESPONSE_BREAK_OTHERS_B;
-		return FPinConnectionResponse(ReplyBreakInputs, TEXT("Replace existing input connections"));
+		return FPinConnectionResponse(AllowConnectionResponse, ReplaceExistingInputConnectionsText.ToString());
 	}
 	else
 	{
@@ -1697,17 +1722,28 @@ FPinConnectionResponse UEdGraphSchema_Niagara::GetWildcardConnectionResponse(con
 	{
 		if (NodeA->AllowNiagaraTypeForPinTypeChange(PinBType, const_cast<UEdGraphPin*>(PinA)))
 		{
-			if(PinA->Direction == EGPD_Input)
+			if(PinA->Direction == EGPD_Input && PinA->LinkedTo.Num() > 0)
 			{
 				if(!bPinsSwapped)
 				{
-					Response = CONNECT_RESPONSE_BREAK_OTHERS_A;
+					Response = CONNECT_RESPONSE_BREAK_OTHERS_A;			
 				}
 				else
 				{
 					Response = CONNECT_RESPONSE_BREAK_OTHERS_B;
 				}
 			}
+			else if(PinB->Direction == EGPD_Input && PinB->LinkedTo.Num() > 0)
+			{
+				if(!bPinsSwapped)
+				{
+					Response = CONNECT_RESPONSE_BREAK_OTHERS_B;			
+				}
+				else
+				{
+					Response = CONNECT_RESPONSE_BREAK_OTHERS_A;
+				}
+			}
 			else
 			{
 				Response = CONNECT_RESPONSE_MAKE;
diff --git a/Engine/Plugins/FX/Niagara/Source/NiagaraEditor/Public/EdGraphSchema_Niagara.h b/Engine/Plugins/FX/Niagara/Source/NiagaraEditor/Public/EdGraphSchema_Niagara.h
--- a/Engine/Plugins/FX/Niagara/Source/NiagaraEditor/Public/EdGraphSchema_Niagara.h
+++ b/Engine/Plugins/FX/Niagara/Source/NiagaraEditor/Public/EdGraphSchema_Niagara.h
@@ -170,7 +170,17 @@ class NIAGARAEDITOR_API UEdGraphSchema_Niagara : public UEdGraphSchema
 	static const FLinearColor NodeTitleColor_TranslatorConstant;
 	static const FLinearColor NodeTitleColor_RapidIteration;
 
-	bool PinTypesValidForNumericConversion(FEdGraphPinType AType, FEdGraphPinType BType) const;
+	static const FText ReplaceExistingInputConnectionsText;
+	static const FText TypesAreNotCompatibleText;
+	static const FText ConvertText;
+	static const FText ConvertLossyText;
+	static const FText PinNotConnectableText;
+	static const FText SameNodeConnectionForbiddenText;
+	static const FText DirectionsNotCompatibleText;
+	static const FText AddPinIncompatibleTypeText;
+	static const FText CircularConnectionFoundText;
+	
+	bool PinTypesValidForNumericConversion(FEdGraphPinType AType, FEdGraphPinType BType) const;	
 	
 private:
 	void GetNumericConversionToSubMenuActions(class UToolMenu* Menu, const FName SectionName, UEdGraphPin* InGraphPin);
