Fix crash when adding emitter renderer due to event registration not being removed.

diff --git a/Engine/Plugins/FX/Niagara/Source/NiagaraEditor/Private/Sequencer/NiagaraSequence/NiagaraEmitterTrackEditor.cpp b/Engine/Plugins/FX/Niagara/Source/NiagaraEditor/Private/Sequencer/NiagaraSequence/NiagaraEmitterTrackEditor.cpp
--- a/Engine/Plugins/FX/Niagara/Source/NiagaraEditor/Private/Sequencer/NiagaraSequence/NiagaraEmitterTrackEditor.cpp
+++ b/Engine/Plugins/FX/Niagara/Source/NiagaraEditor/Private/Sequencer/NiagaraSequence/NiagaraEmitterTrackEditor.cpp
@@ -92,12 +92,8 @@ public:
 		// Due to delegate bind order, the item creation happens after the refresh. To solve this, we just wait a frame. 
 		UNiagaraEmitter* NiagaraEmitter = EmitterTrack->GetEmitterHandleViewModel()->GetEmitterViewModel()->GetEmitter();
 		WeakNiagaraEmitter = NiagaraEmitter;
-		NiagaraEmitter->OnRenderersChanged().AddLambda(
-			[this]()
-			{
-				bShouldRefreshRenderers = true;
-			}
-		);
+		NiagaraEmitter->OnRenderersChanged().AddSP(this, &SEmitterTrackWidget::RefreshRenderers);
+
 
 		// Enabled checkbox.
 		TrackBox->AddSlot()
@@ -257,6 +253,11 @@ private:
 		return EmitterTrack.IsValid() && EmitterTrack->GetSystemViewModel().GetEditMode() == ENiagaraSystemViewModelEditMode::SystemAsset ? EVisibility::Visible : EVisibility::Collapsed;
 	}
 
+	void RefreshRenderers()
+	{
+		bShouldRefreshRenderers = true;
+	}
+
 private:
 	TWeakObjectPtr<UMovieSceneNiagaraEmitterTrack> EmitterTrack;
 	TWeakObjectPtr<UNiagaraEmitter> WeakNiagaraEmitter;
