Possible fix for crash within NiagaraSystemToolkitParameterPanelViewModel.   Non reproducible crash appears due to stale delegate

diff --git a/Engine/Plugins/FX/Niagara/Source/NiagaraEditor/Private/ViewModels/NiagaraParameterPanelViewModel.cpp b/Engine/Plugins/FX/Niagara/Source/NiagaraEditor/Private/ViewModels/NiagaraParameterPanelViewModel.cpp
--- a/Engine/Plugins/FX/Niagara/Source/NiagaraEditor/Private/ViewModels/NiagaraParameterPanelViewModel.cpp
+++ b/Engine/Plugins/FX/Niagara/Source/NiagaraEditor/Private/ViewModels/NiagaraParameterPanelViewModel.cpp
@@ -943,7 +943,7 @@ void FNiagaraSystemToolkitParameterPanelViewModel::Init(const FSystemToolkitUICo
 	for (UNiagaraGraph* Graph : GraphsToAddCallbacks)
 	{
 		FDelegateHandle OnGraphChangedHandle = Graph->AddOnGraphChangedHandler(
-			FOnGraphChanged::FDelegate::CreateRaw(this, &FNiagaraSystemToolkitParameterPanelViewModel::OnGraphChanged));
+			FOnGraphChanged::FDelegate::CreateSP(this, &FNiagaraSystemToolkitParameterPanelViewModel::OnGraphChanged));
 		GraphIdToOnGraphChangedHandleMap.Add(Graph->GetUniqueID(), OnGraphChangedHandle);
 	}
 
@@ -2183,7 +2183,7 @@ void FNiagaraSystemToolkitParameterPanelViewModel::ReconcileOnGraphChangedBindin
 		if (GraphIdToOnGraphChangedHandleMap.Find(Graph->GetUniqueID()) == nullptr)
 		{
 			FDelegateHandle OnGraphChangedHandle = Graph->AddOnGraphChangedHandler(
-				FOnGraphChanged::FDelegate::CreateRaw(this, &FNiagaraSystemToolkitParameterPanelViewModel::OnGraphChanged));
+				FOnGraphChanged::FDelegate::CreateSP(this, &FNiagaraSystemToolkitParameterPanelViewModel::OnGraphChanged));
 			GraphIdToOnGraphChangedHandleMap.Add(Graph->GetUniqueID(), OnGraphChangedHandle);
 		}
 	}
@@ -2225,9 +2225,9 @@ void FNiagaraScriptToolkitParameterPanelViewModel::Init(const FScriptToolkitUICo
 	// Init bindings
 	UNiagaraGraph* NiagaraGraph = static_cast<UNiagaraGraph*>(ScriptViewModel->GetGraphViewModel()->GetGraph());
 	OnGraphChangedHandle = NiagaraGraph->AddOnGraphChangedHandler(
-		FOnGraphChanged::FDelegate::CreateRaw(this, &FNiagaraScriptToolkitParameterPanelViewModel::OnGraphChanged));
+		FOnGraphChanged::FDelegate::CreateSP(this, &FNiagaraScriptToolkitParameterPanelViewModel::OnGraphChanged));
 	OnGraphNeedsRecompileHandle = NiagaraGraph->AddOnGraphNeedsRecompileHandler(
-		FOnGraphChanged::FDelegate::CreateRaw(this, &FNiagaraScriptToolkitParameterPanelViewModel::OnGraphChanged));
+		FOnGraphChanged::FDelegate::CreateSP(this, &FNiagaraScriptToolkitParameterPanelViewModel::OnGraphChanged));
 	OnSubObjectSelectionHandle = NiagaraGraph->OnSubObjectSelectionChanged().AddSP(this, &FNiagaraScriptToolkitParameterPanelViewModel::OnGraphSubObjectSelectionChanged);
 
 	ScriptViewModel->GetOnSubscribedParameterDefinitionsChangedDelegate().AddSP(this, &FNiagaraScriptToolkitParameterPanelViewModel::Refresh);
