Niagara - Fix 5.0 style consistency issues in the selection stack.
+ Update the stack row widget so that it's colors and padding for item headers and item content match the new details panel row style.
+ Refactor the niagara color parameter control to be more inline with the details panel color editor.
+ Fix parameter editor font size inconsistency.
+ Add a new lock icon to differentiate inherited modules and renderers.
+ Fixed inline object categories having all caps font.

diff --git a/Engine/Plugins/FX/Niagara/Source/NiagaraEditor/Private/TypeEditorUtilities/NiagaraColorTypeEditorUtilities.cpp b/Engine/Plugins/FX/Niagara/Source/NiagaraEditor/Private/TypeEditorUtilities/NiagaraColorTypeEditorUtilities.cpp
--- a/Engine/Plugins/FX/Niagara/Source/NiagaraEditor/Private/TypeEditorUtilities/NiagaraColorTypeEditorUtilities.cpp
+++ b/Engine/Plugins/FX/Niagara/Source/NiagaraEditor/Private/TypeEditorUtilities/NiagaraColorTypeEditorUtilities.cpp
@@ -9,7 +9,7 @@
 #include "Widgets/Colors/SColorPicker.h"
 #include "Widgets/Input/SNumericEntryBox.h"
 #include "Widgets/Colors/SColorBlock.h"
-#include "Widgets/Layout/SGridPanel.h"
+#include "Widgets/Layout/SWrapBox.h"
 
 class SNiagaraColorParameterEditor : public SNiagaraParameterEditor
 {
@@ -20,41 +20,39 @@ public:
 	void Construct(const FArguments& InArgs)
 	{
 		SNiagaraParameterEditor::Construct(SNiagaraParameterEditor::FArguments()
-			.MinimumDesiredWidth(3 * DefaultInputSize)
-			.MaximumDesiredWidth(3 * DefaultInputSize));
+			.HAlign(HAlign_Fill));
 
 		ChildSlot
 		[
-			SNew(SGridPanel)
-			.FillColumn(1, 1)
-			.FillColumn(2, 1)
-			.FillColumn(3, 1)
-			+ SGridPanel::Slot(0, 0)
-			.RowSpan(2)
+			SNew(SWrapBox)
+			.UseAllottedSize(true)
+			+ SWrapBox::Slot()
+			.VAlign(VAlign_Center)
+			.Padding(0, 0, 6, 0)
 			[
 				SAssignNew(ColorBlock, SColorBlock)
 				.Color(this, &SNiagaraColorParameterEditor::GetColor)
 				.ShowBackgroundForAlpha(true)
 				.OnMouseButtonDown(this, &SNiagaraColorParameterEditor::OnMouseButtonDownColorBlock)
-				.Size(FVector2D(35.0f, 12.0f))
+				.Size(FVector2D(40.0f, 16.0f))
+				.CornerRadius(FVector4(3.0f, 3.0f, 3.0f, 3.0f))
 			]
-			+ SGridPanel::Slot(1, 0)
-			.Padding(3, 0, 0, 0)
+			+ SWrapBox::Slot()
+			.Padding(0, 0, 4, 0)
 			[
 				ConstructComponentWidget(0, NSLOCTEXT("ColorParameterEditor", "RLabel", "R"))
 			]
-			+ SGridPanel::Slot(2, 0)
-			.Padding(3, 0, 0, 0)
+			+ SWrapBox::Slot()
+			.Padding(0, 0, 4, 0)
 			[
 				ConstructComponentWidget(1, NSLOCTEXT("ColorParameterEditor", "GLabel", "G"))
 			]
-			+ SGridPanel::Slot(3, 0)
-			.Padding(3, 0, 0, 0)
+			+ SWrapBox::Slot()
+			.Padding(0, 0, 4, 0)
 			[
 				ConstructComponentWidget(2, NSLOCTEXT("ColorParameterEditor", "BLabel", "B"))
 			]
-			+ SGridPanel::Slot(1, 1)
-			.Padding(3, 2, 0, 0)
+			+ SWrapBox::Slot()
 			[
 				ConstructComponentWidget(3, NSLOCTEXT("ColorParameterEditor", "ALabel", "A"))
 			]
@@ -79,7 +77,7 @@ private:
 	TSharedRef<SWidget> ConstructComponentWidget(int32 Index, FText ComponentLabel)
 	{
 		return SNew(SNumericEntryBox<float>)
-		.Font(FNiagaraEditorStyle::Get().GetFontStyle("NiagaraEditor.ParameterFont"))
+		.Font(FAppStyle::Get().GetFontStyle("PropertyWindow.NormalFont"))
 		.OverrideTextMargin(2)
 		.MinValue(TOptional<float>())
 		.MaxValue(TOptional<float>())
@@ -93,6 +91,7 @@ private:
 		.OnEndSliderMovement(this, &SNiagaraColorParameterEditor::EndSliderMovement)
 		.AllowSpin(true)
 		.LabelVAlign(EVerticalAlignment::VAlign_Center)
+		.MinDesiredValueWidth(30)
 		.Label()
 		[
 			SNew(STextBlock)
diff --git a/Engine/Plugins/FX/Niagara/Source/NiagaraEditor/Private/TypeEditorUtilities/NiagaraEnumTypeEditorUtilities.cpp b/Engine/Plugins/FX/Niagara/Source/NiagaraEditor/Private/TypeEditorUtilities/NiagaraEnumTypeEditorUtilities.cpp
--- a/Engine/Plugins/FX/Niagara/Source/NiagaraEditor/Private/TypeEditorUtilities/NiagaraEnumTypeEditorUtilities.cpp
+++ b/Engine/Plugins/FX/Niagara/Source/NiagaraEditor/Private/TypeEditorUtilities/NiagaraEnumTypeEditorUtilities.cpp
@@ -27,7 +27,7 @@ public:
 			SNew(SEnumComboBox, Enum)
 			.CurrentValue(this, &SNiagaraEnumParameterEditor::GetValue)
 			.ContentPadding(FMargin(2, 0))
-			.Font(FEditorStyle::GetFontStyle("Sequencer.AnimationOutliner.RegularFont"))
+			.Font(FAppStyle::Get().GetFontStyle("PropertyWindow.NormalFont"))
 			.OnEnumSelectionChanged(SEnumComboBox::FOnEnumSelectionChanged::CreateSP(this, &SNiagaraEnumParameterEditor::ValueChanged))
 		];
 	}
diff --git a/Engine/Plugins/FX/Niagara/Source/NiagaraEditor/Private/ViewModels/NiagaraOverviewGraphViewModel.cpp b/Engine/Plugins/FX/Niagara/Source/NiagaraEditor/Private/ViewModels/NiagaraOverviewGraphViewModel.cpp
--- a/Engine/Plugins/FX/Niagara/Source/NiagaraEditor/Private/ViewModels/NiagaraOverviewGraphViewModel.cpp
+++ b/Engine/Plugins/FX/Niagara/Source/NiagaraEditor/Private/ViewModels/NiagaraOverviewGraphViewModel.cpp
@@ -495,7 +495,8 @@ void FNiagaraOverviewGraphViewModel::GraphSelectionChanged()
 					StackObject = NewObject<UNiagaraStackObject>(GetTransientPackage());
 					UNiagaraStackEditorData* EditorData = NewObject< UNiagaraStackEditorData >(GetTransientPackage());
 					UNiagaraStackEntry::FRequiredEntryData RequiredEntryData(GetSystemViewModel(), nullptr, TEXT("Custom"), NAME_None, *EditorData);
-					StackObject->Initialize(RequiredEntryData, SelectedObject, TEXT(""));
+					bool bIsTopLevelObject = true;
+					StackObject->Initialize(RequiredEntryData, SelectedObject, bIsTopLevelObject, TEXT(""));
 					NewTempEntries.AddUnique(StackObject);
 				}
 
diff --git a/Engine/Plugins/FX/Niagara/Source/NiagaraEditor/Private/ViewModels/Stack/NiagaraStackEmitterSettingsGroup.cpp b/Engine/Plugins/FX/Niagara/Source/NiagaraEditor/Private/ViewModels/Stack/NiagaraStackEmitterSettingsGroup.cpp
--- a/Engine/Plugins/FX/Niagara/Source/NiagaraEditor/Private/ViewModels/Stack/NiagaraStackEmitterSettingsGroup.cpp
+++ b/Engine/Plugins/FX/Niagara/Source/NiagaraEditor/Private/ViewModels/Stack/NiagaraStackEmitterSettingsGroup.cpp
@@ -85,11 +85,6 @@ void UNiagaraStackEmitterPropertiesItem::ResetToBase()
 	}
 }
 
-bool UNiagaraStackEmitterPropertiesItem::IsExpandedByDefault() const
-{
-	return false;
-}
-
 const FSlateBrush* UNiagaraStackEmitterPropertiesItem::GetIconBrush() const
 {
 	if (Emitter->SimTarget == ENiagaraSimTarget::CPUSim)
@@ -109,7 +104,8 @@ void UNiagaraStackEmitterPropertiesItem::RefreshChildrenInternal(const TArray<UN
 	{
 		EmitterObject = NewObject<UNiagaraStackObject>(this);
 		FRequiredEntryData RequiredEntryData(GetSystemViewModel(), GetEmitterViewModel(), FExecutionCategoryNames::Emitter, NAME_None, GetStackEditorData());
-		EmitterObject->Initialize(RequiredEntryData, Emitter.Get(), GetStackEditorDataKey());
+		bool bIsTopLevelObject = true;
+		EmitterObject->Initialize(RequiredEntryData, Emitter.Get(), bIsTopLevelObject, GetStackEditorDataKey());
 		EmitterObject->RegisterInstancedCustomPropertyLayout(UNiagaraEmitter::StaticClass(), FOnGetDetailCustomizationInstance::CreateStatic(&FNiagaraEmitterDetails::MakeInstance));
 	}
 
@@ -194,7 +190,8 @@ void UNiagaraStackEmitterSummaryItem::RefreshChildrenInternal(const TArray<UNiag
 		if (SummaryEditorData == nullptr)
 		{
 			SummaryEditorData = NewObject<UNiagaraStackObject>(this);
-			SummaryEditorData->Initialize(CreateDefaultChildRequiredData(), &GetEmitterViewModel()->GetOrCreateEditorData(), GetStackEditorDataKey());
+			bool bIsTopLevelObject = true;
+			SummaryEditorData->Initialize(CreateDefaultChildRequiredData(), &GetEmitterViewModel()->GetOrCreateEditorData(), bIsTopLevelObject, GetStackEditorDataKey());
 			SummaryEditorData->SetOnSelectRootNodes(UNiagaraStackObject::FOnSelectRootNodes::CreateUObject(this,
 				&UNiagaraStackEmitterSummaryItem::SelectSummaryNodesFromEmitterEditorDataRootNodes));
 		}
diff --git a/Engine/Plugins/FX/Niagara/Source/NiagaraEditor/Private/ViewModels/Stack/NiagaraStackEntry.cpp b/Engine/Plugins/FX/Niagara/Source/NiagaraEditor/Private/ViewModels/Stack/NiagaraStackEntry.cpp
--- a/Engine/Plugins/FX/Niagara/Source/NiagaraEditor/Private/ViewModels/Stack/NiagaraStackEntry.cpp
+++ b/Engine/Plugins/FX/Niagara/Source/NiagaraEditor/Private/ViewModels/Stack/NiagaraStackEntry.cpp
@@ -352,16 +352,6 @@ UNiagaraStackEntry::EStackRowStyle UNiagaraStackEntry::GetStackRowStyle() const
 	return EStackRowStyle::None;
 }
 
-bool UNiagaraStackEntry::HasFrontDivider() const
-{
-	UNiagaraStackEntry* Outer = Cast<UNiagaraStackEntry>(GetOuter());
-	if (Outer == nullptr)
-	{
-		return false;
-	}
-	return Outer->HasFrontDivider();
-}
-
 bool UNiagaraStackEntry::GetShouldShowInStack() const
 {
 	return true;
@@ -715,7 +705,7 @@ void UNiagaraStackEntry::RefreshChildren()
 	for (UNiagaraStackEntry* Child : Children)
 	{
 		UNiagaraStackEntry* OuterOwner = Cast<UNiagaraStackEntry>(Child->GetOuter());
-		Child->IndentLevel = GetChildIndentLevel();
+		Child->IndentLevel = GetChildIndentLevel() + (Child->IsSemanticChild() ? 1 : 0);
 		Child->bOwnerIsEnabled = OuterOwner == nullptr || (OuterOwner->GetIsEnabled() && OuterOwner->GetOwnerIsEnabled());
 		Child->RefreshChildren();
 		Child->OnStructureChanged().AddUObject(this, &UNiagaraStackEntry::ChildStructureChanged);
@@ -1019,3 +1009,10 @@ void UNiagaraStackEntry::OnRenamed(FText NewName)
 		}
 	}
 }
+
+void UNiagaraStackSpacer::Initialize(FRequiredEntryData InRequiredEntryData, float InSpacerHeight, TAttribute<bool> InShouldShowInStack, FString InOwningStackItemEditorDataKey)
+{
+	Super::Initialize(InRequiredEntryData, InOwningStackItemEditorDataKey + TEXT("Spacer"));
+	SpacerHeight = InSpacerHeight;
+	ShouldShowInStack = InShouldShowInStack;
+}
diff --git a/Engine/Plugins/FX/Niagara/Source/NiagaraEditor/Private/ViewModels/Stack/NiagaraStackErrorItem.cpp b/Engine/Plugins/FX/Niagara/Source/NiagaraEditor/Private/ViewModels/Stack/NiagaraStackErrorItem.cpp
--- a/Engine/Plugins/FX/Niagara/Source/NiagaraEditor/Private/ViewModels/Stack/NiagaraStackErrorItem.cpp
+++ b/Engine/Plugins/FX/Niagara/Source/NiagaraEditor/Private/ViewModels/Stack/NiagaraStackErrorItem.cpp
@@ -41,7 +41,7 @@ void UNiagaraStackErrorItem::SetIsExpandedByDefault(bool bIsExpanded)
 
 UNiagaraStackEntry::EStackRowStyle UNiagaraStackErrorItem::GetStackRowStyle() const
 {
-	return EStackRowStyle::StackIssue;
+	return StackIssue.GetSeverity() == EStackIssueSeverity::CustomNote ? EStackRowStyle::ItemContentNote : EStackRowStyle::StackIssue;
 }
 
 UNiagaraStackErrorItem::FOnIssueNotify& UNiagaraStackErrorItem::OnIssueModified()
@@ -144,7 +144,7 @@ EStackIssueSeverity UNiagaraStackErrorItemLongDescription::GetIssueSeverity() co
 
 UNiagaraStackEntry::EStackRowStyle UNiagaraStackErrorItemLongDescription::GetStackRowStyle() const
 {
-	return EStackRowStyle::StackIssue;
+	return StackIssue.GetSeverity() == EStackIssueSeverity::CustomNote ? EStackRowStyle::ItemContentNote : EStackRowStyle::StackIssue;
 }
 
 //UNiagaraStackErrorItemFix
@@ -178,7 +178,7 @@ EStackIssueSeverity UNiagaraStackErrorItemFix::GetIssueSeverity() const
 
 UNiagaraStackEntry::EStackRowStyle UNiagaraStackErrorItemFix::GetStackRowStyle() const
 {
-	return EStackRowStyle::StackIssue;
+	return StackIssue.GetSeverity() == EStackIssueSeverity::CustomNote ? EStackRowStyle::ItemContentNote : EStackRowStyle::StackIssue;
 }
 
 FText UNiagaraStackErrorItemFix::GetFixButtonText() const
@@ -218,7 +218,7 @@ void UNiagaraStackErrorItemDismiss::DismissIssue()
 
 UNiagaraStackEntry::EStackRowStyle UNiagaraStackErrorItemDismiss::GetStackRowStyle() const
 {
-	return EStackRowStyle::StackIssue;
+	return StackIssue.GetSeverity() == EStackIssueSeverity::CustomNote ? EStackRowStyle::ItemContentNote : EStackRowStyle::StackIssue;
 }
 
 FText UNiagaraStackErrorItemDismiss::GetFixButtonText() const
diff --git a/Engine/Plugins/FX/Niagara/Source/NiagaraEditor/Private/ViewModels/Stack/NiagaraStackEventScriptItemGroup.cpp b/Engine/Plugins/FX/Niagara/Source/NiagaraEditor/Private/ViewModels/Stack/NiagaraStackEventScriptItemGroup.cpp
--- a/Engine/Plugins/FX/Niagara/Source/NiagaraEditor/Private/ViewModels/Stack/NiagaraStackEventScriptItemGroup.cpp
+++ b/Engine/Plugins/FX/Niagara/Source/NiagaraEditor/Private/ViewModels/Stack/NiagaraStackEventScriptItemGroup.cpp
@@ -101,7 +101,8 @@ void UNiagaraStackEventHandlerPropertiesItem::RefreshChildrenInternal(const TArr
 	if (EmitterObject == nullptr)
 	{
 		EmitterObject = NewObject<UNiagaraStackObject>(this);
-		EmitterObject->Initialize(CreateDefaultChildRequiredData(), Emitter.Get(), GetStackEditorDataKey());
+		bool bIsTopLevelObject = true;
+		EmitterObject->Initialize(CreateDefaultChildRequiredData(), Emitter.Get(), bIsTopLevelObject, GetStackEditorDataKey());
 		EmitterObject->RegisterInstancedCustomPropertyTypeLayout(FNiagaraEventScriptProperties::StaticStruct()->GetFName(), 
 			FOnGetPropertyTypeCustomizationInstance::CreateStatic(&FNiagaraEventScriptPropertiesCustomization::MakeInstance, 
 			TWeakObjectPtr<UNiagaraSystem>(&GetSystemViewModel()->GetSystem()), TWeakObjectPtr<UNiagaraEmitter>(GetEmitterViewModel()->GetEmitter())));
@@ -295,6 +296,16 @@ void UNiagaraStackEventScriptItemGroup::Delete()
 	OnModifiedEventHandlersDelegate.ExecuteIfBound();
 }
 
+bool UNiagaraStackEventScriptItemGroup::GetIsInherited() const
+{
+	return HasBaseEventHandler();
+}
+
+FText UNiagaraStackEventScriptItemGroup::GetInheritanceMessage() const
+{
+	return LOCTEXT("EventGroupInheritanceMessage", "This event handler is inherited from a parent emitter.  Inherited\nevent handlers can only be deleted while editing the parent emitter.");
+}
+
 bool UNiagaraStackEventScriptItemGroup::HasBaseEventHandler() const
 {
 	if (bHasBaseEventHandlerCache.IsSet() == false)
diff --git a/Engine/Plugins/FX/Niagara/Source/NiagaraEditor/Private/ViewModels/Stack/NiagaraStackFunctionInput.cpp b/Engine/Plugins/FX/Niagara/Source/NiagaraEditor/Private/ViewModels/Stack/NiagaraStackFunctionInput.cpp
--- a/Engine/Plugins/FX/Niagara/Source/NiagaraEditor/Private/ViewModels/Stack/NiagaraStackFunctionInput.cpp
+++ b/Engine/Plugins/FX/Niagara/Source/NiagaraEditor/Private/ViewModels/Stack/NiagaraStackFunctionInput.cpp
@@ -750,7 +750,8 @@ void UNiagaraStackFunctionInput::RefreshChildrenInternal(const TArray<UNiagaraSt
 		if(ValueObjectEntry == nullptr)
 		{
 			ValueObjectEntry = NewObject<UNiagaraStackObject>(this);
-			ValueObjectEntry->Initialize(CreateDefaultChildRequiredData(), InputValues.DataObject.Get(), GetOwnerStackItemEditorDataKey(), OwningFunctionCallNode.Get());
+			bool bIsTopLevelObject = false;
+			ValueObjectEntry->Initialize(CreateDefaultChildRequiredData(), InputValues.DataObject.Get(), bIsTopLevelObject, GetOwnerStackItemEditorDataKey(), OwningFunctionCallNode.Get());
 		}
 		NewChildren.Add(ValueObjectEntry);
 	}
@@ -2419,6 +2420,7 @@ void UNiagaraStackFunctionInput::ChangeScriptVersion(FGuid NewScriptVersion)
       FCompileConstantResolver(GetEmitterViewModel()->GetEmitter(), FNiagaraStackGraphUtilities::GetOutputNodeUsage(*GetDynamicInputNode())) :
       FCompileConstantResolver(&GetSystemViewModel()->GetSystem(), FNiagaraStackGraphUtilities::GetOutputNodeUsage(*GetDynamicInputNode()));
 	GetDynamicInputNode()->ChangeScriptVersion(NewScriptVersion, UpgradeContext, true);
+	RefreshChildren();
 }
 
 const UNiagaraClipboardFunctionInput* UNiagaraStackFunctionInput::ToClipboardFunctionInput(UObject* InOuter) const
@@ -2623,11 +2625,6 @@ void UNiagaraStackFunctionInput::GetSearchItems(TArray<FStackSearchItem>& Search
 	}
 }
 
-bool UNiagaraStackFunctionInput::HasFrontDivider() const
-{
-	return IsSemanticChild() || Super::HasFrontDivider();
-}
-
 TOptional<FNiagaraVariableMetaData> UNiagaraStackFunctionInput::GetMetadata() const
 {
 	return InputMetaData;
diff --git a/Engine/Plugins/FX/Niagara/Source/NiagaraEditor/Private/ViewModels/Stack/NiagaraStackFunctionInputCollection.cpp b/Engine/Plugins/FX/Niagara/Source/NiagaraEditor/Private/ViewModels/Stack/NiagaraStackFunctionInputCollection.cpp
--- a/Engine/Plugins/FX/Niagara/Source/NiagaraEditor/Private/ViewModels/Stack/NiagaraStackFunctionInputCollection.cpp
+++ b/Engine/Plugins/FX/Niagara/Source/NiagaraEditor/Private/ViewModels/Stack/NiagaraStackFunctionInputCollection.cpp
@@ -256,6 +256,12 @@ void UNiagaraStackFunctionInputCollectionBase::RefreshChildrenInternal(const TAr
 	ActiveSectionCache.Reset();
 }
 
+int32 UNiagaraStackFunctionInputCollectionBase::GetChildIndentLevel() const
+{
+	// We don't want the child categories to be indented.
+	return GetIndentLevel();
+}
+
 void UNiagaraStackFunctionInputCollectionBase::RefreshChildrenForFunctionCall(UNiagaraNodeFunctionCall* ModuleNode, UNiagaraNodeFunctionCall* InputFunctionCallNode, 
 	const TArray<UNiagaraStackEntry*>& CurrentChildren, TArray<UNiagaraStackEntry*>& NewChildren, TArray<FStackIssue>& NewIssues, bool bShouldApplySummaryFilter)
 {	
@@ -747,7 +753,8 @@ void UNiagaraStackFunctionInputCollectionBase::AddInputToCategory(const FInputDa
 			InputCategory = NewObject<UNiagaraStackInputCategory>(this);
 
 			FString InputCategoryStackEditorDataKey = FString::Printf(TEXT("%s-InputCategory-%s"), *InputData.InputFunctionCallNode->NodeGuid.ToString(EGuidFormats::DigitsWithHyphens), *InputData.Category.ToString());
-			InputCategory->Initialize(CreateDefaultChildRequiredData(), InputCategoryStackEditorDataKey, InputData.Category, GetOwnerStackItemEditorDataKey());
+			bool bIsTopLevelInput = InputData.InputFunctionCallNode == InputData.ModuleNode;
+			InputCategory->Initialize(CreateDefaultChildRequiredData(), InputCategoryStackEditorDataKey, InputData.Category, bIsTopLevelInput, GetOwnerStackItemEditorDataKey());
 		}
 		else
 		{
@@ -762,9 +769,6 @@ void UNiagaraStackFunctionInputCollectionBase::AddInputToCategory(const FInputDa
 		NewChildren.Add(InputCategory);
 	}
 	InputCategory->AddInput(InputData.ModuleNode, InputData.InputFunctionCallNode, InputData.Pin->PinName, InputData.Type, InputData.bIsStatic ? EStackParameterBehavior::Static : EStackParameterBehavior::Dynamic, InputData.DisplayName, InputData.bIsHidden, InputData.bIsChild);
-
-	
-
 }
 
 bool UNiagaraStackFunctionInputCollectionBase::FilterByActiveSection(const UNiagaraStackEntry& Child) const
diff --git a/Engine/Plugins/FX/Niagara/Source/NiagaraEditor/Private/ViewModels/Stack/NiagaraStackInputCategory.cpp b/Engine/Plugins/FX/Niagara/Source/NiagaraEditor/Private/ViewModels/Stack/NiagaraStackInputCategory.cpp
--- a/Engine/Plugins/FX/Niagara/Source/NiagaraEditor/Private/ViewModels/Stack/NiagaraStackInputCategory.cpp
+++ b/Engine/Plugins/FX/Niagara/Source/NiagaraEditor/Private/ViewModels/Stack/NiagaraStackInputCategory.cpp
@@ -9,12 +9,15 @@ void UNiagaraStackInputCategory::Initialize(
 	FRequiredEntryData InRequiredEntryData,
 	FString InputCategoryStackEditorDataKey,
 	FText InCategoryName,
+	bool bInIsTopLevelCategory,
 	FString InOwnerStackItemEditorDataKey)
 {
 	bool bCategoryIsAdvanced = false;
 	Super::Initialize(InRequiredEntryData, InOwnerStackItemEditorDataKey, InputCategoryStackEditorDataKey);
 	CategoryName = InCategoryName;
 	bShouldShowInStack = true;
+	bIsTopLevelCategory = bInIsTopLevelCategory;
+	CategorySpacer = nullptr;
 	
 	AddChildFilter(FOnFilterChild::CreateUObject(this, &UNiagaraStackInputCategory::FilterForVisibleCondition));
 	AddChildFilter(FOnFilterChild::CreateUObject(this, &UNiagaraStackInputCategory::FilterForIsInlineEditConditionToggle));
@@ -55,6 +58,24 @@ void UNiagaraStackInputCategory::RefreshChildrenInternal(const TArray<UNiagaraSt
 		InputChild->SetSummaryViewDisiplayName(Input.DisplayName);
 		NewChildren.Add(InputChild);
 	}
+
+	if (bIsTopLevelCategory)
+	{
+		if (CategorySpacer == nullptr)
+		{
+			CategorySpacer = NewObject<UNiagaraStackSpacer>(this);
+			TAttribute<bool> ShouldShowSpacerInStack;
+			ShouldShowSpacerInStack.BindUObject(this, &UNiagaraStackInputCategory::GetShouldShowInStack);
+			CategorySpacer->Initialize(CreateDefaultChildRequiredData(), 6, ShouldShowSpacerInStack, GetStackEditorDataKey());
+		}
+		NewChildren.Add(CategorySpacer);
+	}
+}
+
+int32 UNiagaraStackInputCategory::GetChildIndentLevel() const
+{
+	// We want to keep inputs under a top level category at the same indent level as the category.
+	return bIsTopLevelCategory ? GetIndentLevel() : Super::GetChildIndentLevel();
 }
 
 FText UNiagaraStackInputCategory::GetDisplayName() const
@@ -68,12 +89,13 @@ bool UNiagaraStackInputCategory::GetShouldShowInStack() const
 	// in the case where all children have been hidden, don't show the category in the stack.
 	TArray<UNiagaraStackEntry*> CurrentFilteredChildren;
 	GetFilteredChildren(CurrentFilteredChildren);
-	return bShouldShowInStack && CurrentFilteredChildren.Num() > 0;
+	int32 EmptyCount = CategorySpacer == nullptr ? 0 : 1;
+	return bShouldShowInStack && CurrentFilteredChildren.Num() > EmptyCount;
 }
 
 UNiagaraStackEntry::EStackRowStyle UNiagaraStackInputCategory::GetStackRowStyle() const
 {
-	return EStackRowStyle::ItemCategory;
+	return bIsTopLevelCategory ? EStackRowStyle::ItemCategory : EStackRowStyle::ItemSubCategory;
 }
 
 bool UNiagaraStackInputCategory::GetIsEnabled() const
diff --git a/Engine/Plugins/FX/Niagara/Source/NiagaraEditor/Private/ViewModels/Stack/NiagaraStackModuleItem.cpp b/Engine/Plugins/FX/Niagara/Source/NiagaraEditor/Private/ViewModels/Stack/NiagaraStackModuleItem.cpp
--- a/Engine/Plugins/FX/Niagara/Source/NiagaraEditor/Private/ViewModels/Stack/NiagaraStackModuleItem.cpp
+++ b/Engine/Plugins/FX/Niagara/Source/NiagaraEditor/Private/ViewModels/Stack/NiagaraStackModuleItem.cpp
@@ -1431,6 +1431,16 @@ void UNiagaraStackModuleItem::Delete()
 	}
 }
 
+bool UNiagaraStackModuleItem::GetIsInherited() const
+{
+	return CanMoveAndDelete() == false;
+}
+
+FText UNiagaraStackModuleItem::GetInheritanceMessage() const
+{
+	return LOCTEXT("ModuleItemInheritanceMessage", "This module is inherited from a parent emitter.  Inherited modules\ncan only be moved, deleted, and versioned while editing the parent emitter.");
+}
+
 bool UNiagaraStackModuleItem::IsScratchModule() const
 {
 	if (bIsScratchModuleCache.IsSet() == false)
diff --git a/Engine/Plugins/FX/Niagara/Source/NiagaraEditor/Private/ViewModels/Stack/NiagaraStackObject.cpp b/Engine/Plugins/FX/Niagara/Source/NiagaraEditor/Private/ViewModels/Stack/NiagaraStackObject.cpp
--- a/Engine/Plugins/FX/Niagara/Source/NiagaraEditor/Private/ViewModels/Stack/NiagaraStackObject.cpp
+++ b/Engine/Plugins/FX/Niagara/Source/NiagaraEditor/Private/ViewModels/Stack/NiagaraStackObject.cpp
@@ -20,12 +20,13 @@ UNiagaraStackObject::UNiagaraStackObject()
 {
 }
 
-void UNiagaraStackObject::Initialize(FRequiredEntryData InRequiredEntryData, UObject* InObject, FString InOwnerStackItemEditorDataKey, UNiagaraNode* InOwningNiagaraNode)
+void UNiagaraStackObject::Initialize(FRequiredEntryData InRequiredEntryData, UObject* InObject, bool bInIsTopLevelObject, FString InOwnerStackItemEditorDataKey, UNiagaraNode* InOwningNiagaraNode)
 {
 	checkf(WeakObject.IsValid() == false, TEXT("Can only initialize once."));
 	FString ObjectStackEditorDataKey = FString::Printf(TEXT("%s-%s"), *InOwnerStackItemEditorDataKey, *InObject->GetName());
 	Super::Initialize(InRequiredEntryData, InOwnerStackItemEditorDataKey, ObjectStackEditorDataKey);
 	WeakObject = InObject;
+	bIsTopLevelObject = bInIsTopLevelObject;
 	OwningNiagaraNode = InOwningNiagaraNode;
 	bIsRefresingDataInterfaceErrors = false;
 
@@ -280,7 +281,7 @@ void UNiagaraStackObject::RefreshChildrenInternal(const TArray<UNiagaraStackEntr
 			if (ChildRow == nullptr)
 			{
 				ChildRow = NewObject<UNiagaraStackPropertyRow>(this);
-				ChildRow->Initialize(CreateDefaultChildRequiredData(), RootTreeNode, GetOwnerStackItemEditorDataKey(), GetOwnerStackItemEditorDataKey(), OwningNiagaraNode);
+				ChildRow->Initialize(CreateDefaultChildRequiredData(), RootTreeNode, bIsTopLevelObject, GetOwnerStackItemEditorDataKey(), GetOwnerStackItemEditorDataKey(), OwningNiagaraNode);
 			}
 
 			NewChildren.Add(ChildRow);
diff --git a/Engine/Plugins/FX/Niagara/Source/NiagaraEditor/Private/ViewModels/Stack/NiagaraStackParameterStoreEntry.cpp b/Engine/Plugins/FX/Niagara/Source/NiagaraEditor/Private/ViewModels/Stack/NiagaraStackParameterStoreEntry.cpp
--- a/Engine/Plugins/FX/Niagara/Source/NiagaraEditor/Private/ViewModels/Stack/NiagaraStackParameterStoreEntry.cpp
+++ b/Engine/Plugins/FX/Niagara/Source/NiagaraEditor/Private/ViewModels/Stack/NiagaraStackParameterStoreEntry.cpp
@@ -63,7 +63,8 @@ void UNiagaraStackParameterStoreEntry::RefreshChildrenInternal(const TArray<UNia
 		if(ValueObjectEntry == nullptr || ValueObjectEntry->GetObject() != ValueObject)
 		{
 			ValueObjectEntry = NewObject<UNiagaraStackObject>(this);
-			ValueObjectEntry->Initialize(CreateDefaultChildRequiredData(), ValueObject.Get(), GetOwnerStackItemEditorDataKey());
+			bool bIsTopLevelObject = false;
+			ValueObjectEntry->Initialize(CreateDefaultChildRequiredData(), ValueObject.Get(), bIsTopLevelObject, GetOwnerStackItemEditorDataKey());
 		}
 		NewChildren.Add(ValueObjectEntry);
 	}
diff --git a/Engine/Plugins/FX/Niagara/Source/NiagaraEditor/Private/ViewModels/Stack/NiagaraStackPropertyRow.cpp b/Engine/Plugins/FX/Niagara/Source/NiagaraEditor/Private/ViewModels/Stack/NiagaraStackPropertyRow.cpp
--- a/Engine/Plugins/FX/Niagara/Source/NiagaraEditor/Private/ViewModels/Stack/NiagaraStackPropertyRow.cpp
+++ b/Engine/Plugins/FX/Niagara/Source/NiagaraEditor/Private/ViewModels/Stack/NiagaraStackPropertyRow.cpp
@@ -7,7 +7,7 @@
 #include "PropertyHandle.h"
 #include "IDetailPropertyRow.h"
 
-void UNiagaraStackPropertyRow::Initialize(FRequiredEntryData InRequiredEntryData, TSharedRef<IDetailTreeNode> InDetailTreeNode, FString InOwnerStackItemEditorDataKey, FString InOwnerStackEditorDataKey, UNiagaraNode* InOwningNiagaraNode)
+void UNiagaraStackPropertyRow::Initialize(FRequiredEntryData InRequiredEntryData, TSharedRef<IDetailTreeNode> InDetailTreeNode, bool bInIsTopLevelProperty, FString InOwnerStackItemEditorDataKey, FString InOwnerStackEditorDataKey, UNiagaraNode* InOwningNiagaraNode)
 {
 	TSharedPtr<IPropertyHandle> PropertyHandle = InDetailTreeNode->CreatePropertyHandle();
 	FString RowStackEditorDataKey = FString::Printf(TEXT("%s-%s"), *InOwnerStackEditorDataKey, *InDetailTreeNode->GetNodeName().ToString());
@@ -15,10 +15,17 @@ void UNiagaraStackPropertyRow::Initialize(FRequiredEntryData InRequiredEntryData
 	bool bRowIsAdvanced = PropertyHandle.IsValid() && PropertyHandle->GetProperty() && PropertyHandle->GetProperty()->HasAnyPropertyFlags(CPF_AdvancedDisplay);
 	SetIsAdvanced(bRowIsAdvanced);
 	DetailTreeNode = InDetailTreeNode;
+	bIsTopLevelProperty = bInIsTopLevelProperty;
 	OwningNiagaraNode = InOwningNiagaraNode;
-	RowStyle = DetailTreeNode->GetNodeType() == EDetailNodeType::Category
-		? EStackRowStyle::ItemCategory
-		: EStackRowStyle::ItemContent;
+	CategorySpacer = nullptr;
+	if (DetailTreeNode->GetNodeType() == EDetailNodeType::Category)
+	{
+		RowStyle = bInIsTopLevelProperty ? EStackRowStyle::ItemCategory : EStackRowStyle::ItemSubCategory;
+	}
+	else
+	{
+		RowStyle = EStackRowStyle::ItemContent;
+	}
 	bCannotEditInThisContext = false;
 	if (PropertyHandle.IsValid() && PropertyHandle.Get() && PropertyHandle->GetProperty())
 	{
@@ -43,6 +50,11 @@ bool UNiagaraStackPropertyRow::GetIsEnabled() const
 	return OwningNiagaraNode == nullptr || OwningNiagaraNode->GetDesiredEnabledState() == ENodeEnabledState::Enabled;
 }
 
+UNiagaraStackEntry::EStackRowStyle UNiagaraStackPropertyRow::GetStackRowStyle() const
+{
+	return RowStyle;
+}
+
 bool UNiagaraStackPropertyRow::HasOverridenContent() const
 {
 	TSharedPtr<IPropertyHandle> PropertyHandle = DetailTreeNode->CreatePropertyHandle();
@@ -86,12 +98,31 @@ void UNiagaraStackPropertyRow::RefreshChildrenInternal(const TArray<UNiagaraStac
 
 		if (ChildRow == nullptr)
 		{
+			bool bChildIsTopLevelProperty = false;
 			ChildRow = NewObject<UNiagaraStackPropertyRow>(this);
-			ChildRow->Initialize(CreateDefaultChildRequiredData(), NodeChild, GetOwnerStackItemEditorDataKey(), GetStackEditorDataKey(), OwningNiagaraNode);
+			ChildRow->Initialize(CreateDefaultChildRequiredData(), NodeChild, bChildIsTopLevelProperty, GetOwnerStackItemEditorDataKey(), GetStackEditorDataKey(), OwningNiagaraNode);
 		}
 
 		NewChildren.Add(ChildRow);
 	}
+
+	if (bIsTopLevelProperty && DetailTreeNode->GetNodeType() == EDetailNodeType::Category)
+	{
+		if (CategorySpacer == nullptr)
+		{
+			CategorySpacer = NewObject<UNiagaraStackSpacer>(this);
+			TAttribute<bool> ShouldShowSpacerInStack;
+			ShouldShowSpacerInStack.BindUObject(this, &UNiagaraStackPropertyRow::GetShouldShowInStack);
+			CategorySpacer->Initialize(CreateDefaultChildRequiredData(), 6, ShouldShowSpacerInStack, GetStackEditorDataKey());
+		}
+		NewChildren.Add(CategorySpacer);
+	}
+}
+
+int32 UNiagaraStackPropertyRow::GetChildIndentLevel() const
+{
+	// We want to keep inputs under a top level category at the same indent level as the category.
+	return bIsTopLevelProperty && DetailTreeNode->GetNodeType() == EDetailNodeType::Category ? GetIndentLevel() : Super::GetChildIndentLevel();
 }
 
 void UNiagaraStackPropertyRow::GetSearchItems(TArray<FStackSearchItem>& SearchItems) const
diff --git a/Engine/Plugins/FX/Niagara/Source/NiagaraEditor/Private/ViewModels/Stack/NiagaraStackRendererItem.cpp b/Engine/Plugins/FX/Niagara/Source/NiagaraEditor/Private/ViewModels/Stack/NiagaraStackRendererItem.cpp
--- a/Engine/Plugins/FX/Niagara/Source/NiagaraEditor/Private/ViewModels/Stack/NiagaraStackRendererItem.cpp
+++ b/Engine/Plugins/FX/Niagara/Source/NiagaraEditor/Private/ViewModels/Stack/NiagaraStackRendererItem.cpp
@@ -277,6 +277,16 @@ void UNiagaraStackRendererItem::Delete()
 	OnDataObjectModified().Broadcast(ChangedObjects, ENiagaraDataObjectChange::Removed);
 }
 
+bool UNiagaraStackRendererItem::GetIsInherited() const
+{
+	return HasBaseRenderer() == true;
+}
+
+FText UNiagaraStackRendererItem::GetInheritanceMessage() const
+{
+	return LOCTEXT("RendererItemInheritanceMessage", "This renderer is inherited from a parent emitter.  Inherited\nrenderers can only be deleted while editing the parent emitter.");
+}
+
 bool UNiagaraStackRendererItem::HasBaseRenderer() const
 {
 	if (HasBaseEmitter())
@@ -371,7 +381,8 @@ void UNiagaraStackRendererItem::RefreshChildrenInternal(const TArray<UNiagaraSta
 	if (RendererObject == nullptr)
 	{
 		RendererObject = NewObject<UNiagaraStackObject>(this);
-		RendererObject->Initialize(CreateDefaultChildRequiredData(), RendererProperties.Get(), GetStackEditorDataKey());
+		bool bIsTopLevelObject = true;
+		RendererObject->Initialize(CreateDefaultChildRequiredData(), RendererProperties.Get(), bIsTopLevelObject, GetStackEditorDataKey());
 	}
 
 	NewChildren.Add(RendererObject);
diff --git a/Engine/Plugins/FX/Niagara/Source/NiagaraEditor/Private/ViewModels/Stack/NiagaraStackSimulationStageGroup.cpp b/Engine/Plugins/FX/Niagara/Source/NiagaraEditor/Private/ViewModels/Stack/NiagaraStackSimulationStageGroup.cpp
--- a/Engine/Plugins/FX/Niagara/Source/NiagaraEditor/Private/ViewModels/Stack/NiagaraStackSimulationStageGroup.cpp
+++ b/Engine/Plugins/FX/Niagara/Source/NiagaraEditor/Private/ViewModels/Stack/NiagaraStackSimulationStageGroup.cpp
@@ -102,7 +102,8 @@ void UNiagaraStackSimulationStagePropertiesItem::RefreshChildrenInternal(const T
 	if (SimulationStageObject == nullptr)
 	{
 		SimulationStageObject = NewObject<UNiagaraStackObject>(this);
-		SimulationStageObject->Initialize(CreateDefaultChildRequiredData(), SimulationStage.Get(), GetStackEditorDataKey());
+		bool bIsTopLevelObject = true;
+		SimulationStageObject->Initialize(CreateDefaultChildRequiredData(), SimulationStage.Get(), bIsTopLevelObject, GetStackEditorDataKey());
 	}
 
 	if ( SimulationStage.IsValid() )
@@ -399,6 +400,16 @@ void UNiagaraStackSimulationStageGroup::Delete()
 	OnModifiedSimulationStagesDelegate.ExecuteIfBound();
 }
 
+bool UNiagaraStackSimulationStageGroup::GetIsInherited() const
+{
+	return HasBaseSimulationStage();
+}
+
+FText UNiagaraStackSimulationStageGroup::GetInheritanceMessage() const
+{
+	return LOCTEXT("SimulationStageGroupInheritanceMessage", "This simulation stage is inherited from a parent emitter.  Inherited\nsimulation stages can only be deleted while editing the parent emitter.");
+}
+
 void UNiagaraStackSimulationStageGroup::SimulationStagePropertiesChanged()
 {
 	SetDisplayName(FText::FromName(SimulationStage->SimulationStageName));
diff --git a/Engine/Plugins/FX/Niagara/Source/NiagaraEditor/Private/ViewModels/Stack/NiagaraStackSystemPropertiesItem.cpp b/Engine/Plugins/FX/Niagara/Source/NiagaraEditor/Private/ViewModels/Stack/NiagaraStackSystemPropertiesItem.cpp
--- a/Engine/Plugins/FX/Niagara/Source/NiagaraEditor/Private/ViewModels/Stack/NiagaraStackSystemPropertiesItem.cpp
+++ b/Engine/Plugins/FX/Niagara/Source/NiagaraEditor/Private/ViewModels/Stack/NiagaraStackSystemPropertiesItem.cpp
@@ -27,18 +27,14 @@ FText UNiagaraStackSystemPropertiesItem::GetTooltipText() const
 	return LOCTEXT("SystemPropertiesTooltip", "Properties of the System. These cannot change at runtime.");
 }
 
-bool UNiagaraStackSystemPropertiesItem::IsExpandedByDefault() const
-{
-	return false;
-}
-
 void UNiagaraStackSystemPropertiesItem::RefreshChildrenInternal(const TArray<UNiagaraStackEntry*>& CurrentChildren, TArray<UNiagaraStackEntry*>& NewChildren, TArray<FStackIssue>& NewIssues)
 {
 	if (SystemObject == nullptr)
 	{
 		SystemObject = NewObject<UNiagaraStackObject>(this);
 		FRequiredEntryData RequiredEntryData(GetSystemViewModel(), GetEmitterViewModel(), FExecutionCategoryNames::System, NAME_None, GetStackEditorData());
-		SystemObject->Initialize(RequiredEntryData, System.Get(), GetStackEditorDataKey());
+		bool bIsTopLevelObject = true;
+		SystemObject->Initialize(RequiredEntryData, System.Get(), bIsTopLevelObject, GetStackEditorDataKey());
 		SystemObject->RegisterInstancedCustomPropertyLayout(UNiagaraSystem::StaticClass(), FOnGetDetailCustomizationInstance::CreateStatic(&FNiagaraSystemDetails::MakeInstance));
 	}
 	NewChildren.Add(SystemObject);
diff --git a/Engine/Plugins/FX/Niagara/Source/NiagaraEditor/Public/ViewModels/Stack/NiagaraStackEmitterSettingsGroup.h b/Engine/Plugins/FX/Niagara/Source/NiagaraEditor/Public/ViewModels/Stack/NiagaraStackEmitterSettingsGroup.h
--- a/Engine/Plugins/FX/Niagara/Source/NiagaraEditor/Public/ViewModels/Stack/NiagaraStackEmitterSettingsGroup.h
+++ b/Engine/Plugins/FX/Niagara/Source/NiagaraEditor/Public/ViewModels/Stack/NiagaraStackEmitterSettingsGroup.h
@@ -31,8 +31,6 @@ public:
 	virtual EIconMode GetSupportedIconMode() const { return EIconMode::Brush; }
 	virtual const FSlateBrush* GetIconBrush() const override;
 
-	virtual bool IsExpandedByDefault() const override;
-
 	virtual bool GetShouldShowInOverview() const override { return false; }
 
 protected:
@@ -66,7 +64,6 @@ public:
 	virtual FText GetTooltipText() const override;
 
 	virtual bool SupportsResetToBase() const override { return false; }
-	virtual bool IsExpandedByDefault() const override { return false; }
 	virtual EIconMode GetSupportedIconMode() const override { return EIconMode::Text; }
 	virtual FText GetIconText() const override;
 	virtual bool GetShouldShowInOverview() const override { return false; }
diff --git a/Engine/Plugins/FX/Niagara/Source/NiagaraEditor/Public/ViewModels/Stack/NiagaraStackEntry.h b/Engine/Plugins/FX/Niagara/Source/NiagaraEditor/Public/ViewModels/Stack/NiagaraStackEntry.h
--- a/Engine/Plugins/FX/Niagara/Source/NiagaraEditor/Public/ViewModels/Stack/NiagaraStackEntry.h
+++ b/Engine/Plugins/FX/Niagara/Source/NiagaraEditor/Public/ViewModels/Stack/NiagaraStackEntry.h
@@ -136,9 +136,12 @@ public:
 		ItemHeader,
 		ItemContent,
 		ItemContentAdvanced,
+		ItemContentNote,
 		ItemFooter,
 		ItemCategory,
-		StackIssue
+		ItemSubCategory,
+		StackIssue,
+		Spacer
 	};
 
 	struct FRequiredEntryData
@@ -293,10 +296,12 @@ public:
 	FName GetExecutionSubcategoryName() const;
 
 	virtual EStackRowStyle GetStackRowStyle() const;
-	virtual bool HasFrontDivider() const;
 
 	int32 GetIndentLevel() const;
 
+	/** Returns whether or not this entry should be treated as a child of a previous sibling for layout purposes. */
+	virtual bool IsSemanticChild() const { return false; }
+
 	virtual bool GetShouldShowInStack() const;
 
 	virtual bool GetShouldShowInOverview() const;
@@ -517,6 +522,12 @@ public:
 
 	virtual FText GetIconText() const { return FText(); }
 
+	virtual bool SupportsInheritance() const { return false; }
+
+	virtual bool GetIsInherited() const { return false; }
+
+	virtual FText GetInheritanceMessage() const { return FText(); }
+
 protected:
 	virtual void BeginDestroy() override;
 
@@ -643,4 +654,22 @@ private:
 	TOptional<FText> AlternateDisplayName;
 
 	mutable TOptional<FCollectedIssueData> CachedCollectedIssueData;
+};
+
+
+UCLASS()
+class NIAGARAEDITOR_API UNiagaraStackSpacer : public UNiagaraStackEntry
+{
+	GENERATED_BODY()
+public:
+	void Initialize(FRequiredEntryData InRequiredEntryData, float InSpacerHeight, TAttribute<bool> InShouldShowInStack, FString InOwningStackItemEditorDataKey);
+	virtual EStackRowStyle GetStackRowStyle() const override { return UNiagaraStackEntry::EStackRowStyle::Spacer; }
+	virtual bool GetCanExpand() const override { return false; }
+	virtual bool GetShouldShowInStack() const override { return ShouldShowInStack.Get(); }
+
+	float GetSpacerHeight() const { return SpacerHeight; }
+
+private:
+	float SpacerHeight;
+	TAttribute<bool> ShouldShowInStack;
 };
\ No newline at end of file
diff --git a/Engine/Plugins/FX/Niagara/Source/NiagaraEditor/Public/ViewModels/Stack/NiagaraStackEventScriptItemGroup.h b/Engine/Plugins/FX/Niagara/Source/NiagaraEditor/Public/ViewModels/Stack/NiagaraStackEventScriptItemGroup.h
--- a/Engine/Plugins/FX/Niagara/Source/NiagaraEditor/Public/ViewModels/Stack/NiagaraStackEventScriptItemGroup.h
+++ b/Engine/Plugins/FX/Niagara/Source/NiagaraEditor/Public/ViewModels/Stack/NiagaraStackEventScriptItemGroup.h
@@ -79,6 +79,10 @@ public:
 	virtual bool TestCanDeleteWithMessage(FText& OutCanDeleteMessage) const override;
 	virtual void Delete() override;
 
+	virtual bool SupportsInheritance() const override { return true; }
+	virtual bool GetIsInherited() const override;
+	virtual FText GetInheritanceMessage() const override;
+
 	const TObjectPtr<UNiagaraStackEventHandlerPropertiesItem>& GetEventHandlerPropertiesItem() const { return EventHandlerProperties; }
 	FGuid GetEventSourceEmitterId() const {return EventSourceEmitterId; };
 
diff --git a/Engine/Plugins/FX/Niagara/Source/NiagaraEditor/Public/ViewModels/Stack/NiagaraStackFunctionInput.h b/Engine/Plugins/FX/Niagara/Source/NiagaraEditor/Public/ViewModels/Stack/NiagaraStackFunctionInput.h
--- a/Engine/Plugins/FX/Niagara/Source/NiagaraEditor/Public/ViewModels/Stack/NiagaraStackFunctionInput.h
+++ b/Engine/Plugins/FX/Niagara/Source/NiagaraEditor/Public/ViewModels/Stack/NiagaraStackFunctionInput.h
@@ -271,12 +271,11 @@ public:
 
 	TOptional<FNiagaraVariableMetaData> GetInputMetaData() const;
 	
-	virtual bool IsSemanticChild() const;
+	virtual bool IsSemanticChild() const override;
 	void SetSemanticChild(bool IsSemanticChild);
 
 	//~ UNiagaraStackEntry interface
 	virtual void GetSearchItems(TArray<FStackSearchItem>& SearchItems) const override;
-	virtual bool HasFrontDivider() const override;
 
 	TOptional<FNiagaraVariableMetaData> GetMetadata() const;
 	TOptional<FGuid> GetMetadataGuid() const;
diff --git a/Engine/Plugins/FX/Niagara/Source/NiagaraEditor/Public/ViewModels/Stack/NiagaraStackFunctionInputCollection.h b/Engine/Plugins/FX/Niagara/Source/NiagaraEditor/Public/ViewModels/Stack/NiagaraStackFunctionInputCollection.h
--- a/Engine/Plugins/FX/Niagara/Source/NiagaraEditor/Public/ViewModels/Stack/NiagaraStackFunctionInputCollection.h
+++ b/Engine/Plugins/FX/Niagara/Source/NiagaraEditor/Public/ViewModels/Stack/NiagaraStackFunctionInputCollection.h
@@ -74,6 +74,8 @@ protected:
 	
 	virtual void RefreshChildrenInternal(const TArray<UNiagaraStackEntry*>& CurrentChildren, TArray<UNiagaraStackEntry*>& NewChildren, TArray<FStackIssue>& NewIssues) override;
 
+	virtual int32 GetChildIndentLevel() const override;
+
 	void RefreshChildrenForFunctionCall(UNiagaraNodeFunctionCall* ModuleNode, UNiagaraNodeFunctionCall* InputFunctionCallNode, const TArray<UNiagaraStackEntry*>& CurrentChildren, 
 		TArray<UNiagaraStackEntry*>& NewChildren, TArray<FStackIssue>& NewIssues, bool bShouldApplySummaryFilter);
 	
diff --git a/Engine/Plugins/FX/Niagara/Source/NiagaraEditor/Public/ViewModels/Stack/NiagaraStackInputCategory.h b/Engine/Plugins/FX/Niagara/Source/NiagaraEditor/Public/ViewModels/Stack/NiagaraStackInputCategory.h
--- a/Engine/Plugins/FX/Niagara/Source/NiagaraEditor/Public/ViewModels/Stack/NiagaraStackInputCategory.h
+++ b/Engine/Plugins/FX/Niagara/Source/NiagaraEditor/Public/ViewModels/Stack/NiagaraStackInputCategory.h
@@ -26,6 +26,7 @@ public:
 		FRequiredEntryData InRequiredEntryData,
 		FString InputCategoryStackEditorDataKey,
 		FText InCategoryName,
+		bool bInIsTopLevelCategory,
 		FString InOwnerStackItemEditorDataKey);
 	
 	const FText& GetCategoryName() const;
@@ -48,11 +49,11 @@ public:
 	void SetStaticSwitchValuesFromClipboardFunctionInputs(const TArray<const UNiagaraClipboardFunctionInput*>& ClipboardFunctionInputs);
 
 	void SetStandardValuesFromClipboardFunctionInputs(const TArray<const UNiagaraClipboardFunctionInput*>& ClipboardFunctionInputs);
-protected:
-
 
+protected:
 	//~ UNiagaraStackEntry interface
 	virtual void RefreshChildrenInternal(const TArray<UNiagaraStackEntry*>& CurrentChildren, TArray<UNiagaraStackEntry*>& NewChildren, TArray<FStackIssue>& NewIssues) override;
+	virtual int32 GetChildIndentLevel() const override;
 
 private:
 	bool FilterForVisibleCondition(const UNiagaraStackEntry& Child) const;
@@ -74,5 +75,10 @@ private:
 	TOptional<FText> DisplayName;
 	FText CategoryName;
 	TArray<FInputParameterHandleAndType> Inputs;
+
+	UPROPERTY()
+	UNiagaraStackSpacer* CategorySpacer;
+
 	bool bShouldShowInStack;
+	bool bIsTopLevelCategory;
 };
\ No newline at end of file
diff --git a/Engine/Plugins/FX/Niagara/Source/NiagaraEditor/Public/ViewModels/Stack/NiagaraStackItem.h b/Engine/Plugins/FX/Niagara/Source/NiagaraEditor/Public/ViewModels/Stack/NiagaraStackItem.h
--- a/Engine/Plugins/FX/Niagara/Source/NiagaraEditor/Public/ViewModels/Stack/NiagaraStackItem.h
+++ b/Engine/Plugins/FX/Niagara/Source/NiagaraEditor/Public/ViewModels/Stack/NiagaraStackItem.h
@@ -41,6 +41,8 @@ public:
 	virtual bool SupportsEditMode() const { return false; }
 	virtual bool GetEditModeIsActive() const { return false; }
 	virtual void SetEditModeIsActive(bool bInEditModeIsActive) { }
+
+	virtual bool GetIsInherited() const { return false; }
 	
 protected:
 	virtual void RefreshChildrenInternal(const TArray<UNiagaraStackEntry*>& CurrentChildren, TArray<UNiagaraStackEntry*>& NewChildren, TArray<FStackIssue>& NewIssues) override;
diff --git a/Engine/Plugins/FX/Niagara/Source/NiagaraEditor/Public/ViewModels/Stack/NiagaraStackModuleItem.h b/Engine/Plugins/FX/Niagara/Source/NiagaraEditor/Public/ViewModels/Stack/NiagaraStackModuleItem.h
--- a/Engine/Plugins/FX/Niagara/Source/NiagaraEditor/Public/ViewModels/Stack/NiagaraStackModuleItem.h
+++ b/Engine/Plugins/FX/Niagara/Source/NiagaraEditor/Public/ViewModels/Stack/NiagaraStackModuleItem.h
@@ -95,6 +95,10 @@ public:
 	virtual FText GetDeleteTransactionText() const override;
 	virtual void Delete() override;
 
+	virtual bool SupportsInheritance() const override { return true; }
+	virtual bool GetIsInherited() const override;
+	virtual FText GetInheritanceMessage() const override;
+
 	bool IsScratchModule() const;
 
 	void SetOnRequestDeprecationRecommended(FOnRequestDeprecationRecommended InOnRequest)
diff --git a/Engine/Plugins/FX/Niagara/Source/NiagaraEditor/Public/ViewModels/Stack/NiagaraStackObject.h b/Engine/Plugins/FX/Niagara/Source/NiagaraEditor/Public/ViewModels/Stack/NiagaraStackObject.h
--- a/Engine/Plugins/FX/Niagara/Source/NiagaraEditor/Public/ViewModels/Stack/NiagaraStackObject.h
+++ b/Engine/Plugins/FX/Niagara/Source/NiagaraEditor/Public/ViewModels/Stack/NiagaraStackObject.h
@@ -23,7 +23,13 @@ public:
 public:
 	UNiagaraStackObject();
 
-	void Initialize(FRequiredEntryData InRequiredEntryData, UObject* InObject, FString InOwnerStackItemEditorDataKey, UNiagaraNode* InOwningNiagaraNode = nullptr);
+	/** Initialized an entry for displaying a UObject in the niagara stack. 
+		@param InRequiredEntryData Struct with required data for all stack entries.
+		@param InObject The object to be displayed in the stack.
+		@param bInIsTopLevelObject Whether or not the object being displayed should be treated as a top level object shown in the root of the stack.  This can be used to make layout decisions such as category styling.
+		@param InOwnerStackItemEditorDataKey The stack editor data key of the owning stack entry. 
+		@param InOwningNiagaraNode An optional niagara node which owns this object. */
+	void Initialize(FRequiredEntryData InRequiredEntryData, UObject* InObject, bool bInIsTopLevelObject, FString InOwnerStackItemEditorDataKey, UNiagaraNode* InOwningNiagaraNode = nullptr);
 
 	void SetOnSelectRootNodes(FOnSelectRootNodes OnSelectRootNodes);
 
@@ -65,6 +71,10 @@ private:
 	};
 
 	TWeakObjectPtr<UObject> WeakObject;
+
+	// Whether or not the object being displayed should be treated as a top level object shown in the root of the stack.  This can be used to make layout decisions such as category styling.
+	bool bIsTopLevelObject;
+
 	UNiagaraNode* OwningNiagaraNode;
 	FOnSelectRootNodes OnSelectRootNodesDelegate;
 	TArray<FRegisteredClassCustomization> RegisteredClassCustomizations;
diff --git a/Engine/Plugins/FX/Niagara/Source/NiagaraEditor/Public/ViewModels/Stack/NiagaraStackPropertyRow.h b/Engine/Plugins/FX/Niagara/Source/NiagaraEditor/Public/ViewModels/Stack/NiagaraStackPropertyRow.h
--- a/Engine/Plugins/FX/Niagara/Source/NiagaraEditor/Public/ViewModels/Stack/NiagaraStackPropertyRow.h
+++ b/Engine/Plugins/FX/Niagara/Source/NiagaraEditor/Public/ViewModels/Stack/NiagaraStackPropertyRow.h
@@ -14,7 +14,9 @@ class NIAGARAEDITOR_API UNiagaraStackPropertyRow : public UNiagaraStackItemConte
 	GENERATED_BODY()
 		
 public:
-	void Initialize(FRequiredEntryData InRequiredEntryData, TSharedRef<IDetailTreeNode> InDetailTreeNode, FString InOwnerStackItemEditorDataKey, FString InOwnerStackEditorDataKey, UNiagaraNode* InOwningNiagaraNode);
+	void Initialize(FRequiredEntryData InRequiredEntryData, TSharedRef<IDetailTreeNode> InDetailTreeNode, bool bInIsTopLevelProperty, FString InOwnerStackItemEditorDataKey, FString InOwnerStackEditorDataKey, UNiagaraNode* InOwningNiagaraNode);
+
+	virtual EStackRowStyle GetStackRowStyle() const override;
 
 	TSharedRef<IDetailTreeNode> GetDetailTreeNode() const;
 
@@ -31,6 +33,8 @@ protected:
 
 	virtual void RefreshChildrenInternal(const TArray<UNiagaraStackEntry*>& CurrentChildren, TArray<UNiagaraStackEntry*>& NewChildren, TArray<FStackIssue>& NewIssues) override;
 
+	virtual int32 GetChildIndentLevel() const override;
+
 	virtual void GetSearchItems(TArray<FStackSearchItem>& SearchItems) const override;
 
 	virtual TOptional<FDropRequestResponse> CanDropInternal(const FDropRequest& DropRequest) override;
@@ -41,5 +45,10 @@ private:
 	TSharedPtr<IDetailTreeNode> DetailTreeNode;
 	UNiagaraNode* OwningNiagaraNode;
 	EStackRowStyle RowStyle;
+
+	UPROPERTY()
+	UNiagaraStackSpacer* CategorySpacer;
+
 	bool bCannotEditInThisContext;
+	bool bIsTopLevelProperty;
 };
\ No newline at end of file
diff --git a/Engine/Plugins/FX/Niagara/Source/NiagaraEditor/Public/ViewModels/Stack/NiagaraStackRendererItem.h b/Engine/Plugins/FX/Niagara/Source/NiagaraEditor/Public/ViewModels/Stack/NiagaraStackRendererItem.h
--- a/Engine/Plugins/FX/Niagara/Source/NiagaraEditor/Public/ViewModels/Stack/NiagaraStackRendererItem.h
+++ b/Engine/Plugins/FX/Niagara/Source/NiagaraEditor/Public/ViewModels/Stack/NiagaraStackRendererItem.h
@@ -45,6 +45,10 @@ public:
 	virtual FText GetDeleteTransactionText() const override;
 	virtual void Delete();
 
+	virtual bool SupportsInheritance() const override { return true; }
+	virtual bool GetIsInherited() const override;
+	virtual FText GetInheritanceMessage() const override;
+	
 	virtual bool SupportsRename() const override { return true; }
 
 	bool HasBaseRenderer() const;
diff --git a/Engine/Plugins/FX/Niagara/Source/NiagaraEditor/Public/ViewModels/Stack/NiagaraStackSimulationStageGroup.h b/Engine/Plugins/FX/Niagara/Source/NiagaraEditor/Public/ViewModels/Stack/NiagaraStackSimulationStageGroup.h
--- a/Engine/Plugins/FX/Niagara/Source/NiagaraEditor/Public/ViewModels/Stack/NiagaraStackSimulationStageGroup.h
+++ b/Engine/Plugins/FX/Niagara/Source/NiagaraEditor/Public/ViewModels/Stack/NiagaraStackSimulationStageGroup.h
@@ -73,6 +73,10 @@ public:
 	virtual bool TestCanDeleteWithMessage(FText& OutCanDeleteMessage) const override;
 	virtual void Delete() override;
 
+	virtual bool SupportsInheritance() const override { return true; }
+	virtual bool GetIsInherited() const override;
+	virtual FText GetInheritanceMessage() const override;
+
 	virtual bool CanDrag() const override { return true; }
 
 	virtual bool GetIsEnabled() const override;
diff --git a/Engine/Plugins/FX/Niagara/Source/NiagaraEditor/Public/ViewModels/Stack/NiagaraStackSystemPropertiesItem.h b/Engine/Plugins/FX/Niagara/Source/NiagaraEditor/Public/ViewModels/Stack/NiagaraStackSystemPropertiesItem.h
--- a/Engine/Plugins/FX/Niagara/Source/NiagaraEditor/Public/ViewModels/Stack/NiagaraStackSystemPropertiesItem.h
+++ b/Engine/Plugins/FX/Niagara/Source/NiagaraEditor/Public/ViewModels/Stack/NiagaraStackSystemPropertiesItem.h
@@ -20,8 +20,6 @@ public:
 
 	virtual FText GetTooltipText() const override;
 
-	virtual bool IsExpandedByDefault() const override;
-
 	virtual bool GetShouldShowInOverview() const override { return false; }
 
 protected:
diff --git a/Engine/Plugins/FX/Niagara/Source/NiagaraEditorWidgets/Private/NiagaraEditorWidgetsStyle.cpp b/Engine/Plugins/FX/Niagara/Source/NiagaraEditorWidgets/Private/NiagaraEditorWidgetsStyle.cpp
--- a/Engine/Plugins/FX/Niagara/Source/NiagaraEditorWidgets/Private/NiagaraEditorWidgetsStyle.cpp
+++ b/Engine/Plugins/FX/Niagara/Source/NiagaraEditorWidgets/Private/NiagaraEditorWidgetsStyle.cpp
@@ -92,7 +92,7 @@ FNiagaraEditorWidgetsStyle::FNiagaraEditorWidgetsStyle() : FSlateStyleSet("Niaga
 
 	FSlateFontInfo NormalFont = FAppStyle::Get().GetFontStyle(TEXT("PropertyWindow.NormalFont"));
 	FTextBlockStyle StackItemText = FTextBlockStyle(CategoryText)
-		.SetFont(NormalFont);
+		.SetFont(CategoryFont);
 
 	Set("NiagaraEditor.Stack.ItemText", StackItemText);
 	
@@ -304,51 +304,70 @@ FNiagaraEditorWidgetsStyle::FNiagaraEditorWidgetsStyle() : FSlateStyleSet("Niaga
 
 	Set("Niagara.TableViewRowBorder", new BOX_BRUSH("Icons/Row", FMargin(3.0f / 8.0f), FStyleColors::Recessed));
 
-	Set("NiagaraEditor.Stack.TableViewRow", FTableRowStyle(NormalTableRowStyle)
-		.SetUseParentRowBrush(true)
- 		.SetEvenRowBackgroundBrush(FSlateColorBrush(FStyleColors::Transparent))
-		.SetOddRowBackgroundBrush(FSlateColorBrush(FStyleColors::Transparent))
-		.SetActiveBrush(StackRowSelectionBrush)
-		.SetActiveHoveredBrush(StackRowSelectionBrush)
-		.SetParentRowBackgroundHoveredBrush(FSlateColorBrush(FStyleColors::Header))
-		.SetOddRowBackgroundHoveredBrush(FSlateColorBrush(FStyleColors::Header))
-		.SetEvenRowBackgroundHoveredBrush(FSlateColorBrush(FStyleColors::Header))
-		.SetInactiveBrush(StackRowSubduedSelectionBrush)
-		.SetInactiveHoveredBrush(StackRowSelectionBrush));
+	// TODO - This color is not customizable using the editor style.
+	FLinearColor CustomNoteBackgroundColor = FLinearColor(FColor(56, 111, 75));
+	Set("NiagaraEditor.Stack.Item.CustomNoteBackgroundColor", CustomNoteBackgroundColor);
 
-	Set("NiagaraEditor.SystemOverview.TableViewRow.Item", FTableRowStyle(NormalTableRowStyle)
-		.SetEvenRowBackgroundBrush(FSlateColorBrush(FStyleColors::Panel))
-		.SetOddRowBackgroundBrush(FSlateColorBrush(FStyleColors::Panel))
-		.SetEvenRowBackgroundHoveredBrush(CORE_IMAGE_BRUSH("Common/Selection", Icon8x8, FLinearColor(1.0f, 1.0f, 1.0f, 0.1f)))
-		.SetOddRowBackgroundHoveredBrush(CORE_IMAGE_BRUSH("Common/Selection", Icon8x8, FLinearColor(1.0f, 1.0f, 1.0f, 0.1f)))
+	FLinearColor HoverAdd = FLinearColor(0.03f, 0.03f, 0.03f, 0.0f);
+	FLinearColor PanelHoverColor = FStyleColors::Panel.GetSpecifiedColor() + HoverAdd;
+	FLinearColor RecessedHoverColor = FStyleColors::Recessed.GetSpecifiedColor() + HoverAdd;
+	FLinearColor HeaderHoverColor = FStyleColors::Header.GetSpecifiedColor() + HoverAdd;
+	FLinearColor NoteHoverColor = CustomNoteBackgroundColor + HoverAdd;
+
+	FTableRowStyle BaseStackTableRowStyle = FTableRowStyle(NormalTableRowStyle)
 		.SetSelectorFocusedBrush(FSlateNoResource())
 		.SetActiveBrush(FSlateColorBrush(FStyleColors::Select))
 		.SetActiveHoveredBrush(FSlateColorBrush(FStyleColors::Select))
 		.SetInactiveBrush(FSlateColorBrush(FStyleColors::SelectInactive))
 		.SetActiveHighlightedBrush(FSlateColorBrush(FStyleColors::PrimaryHover))
-		.SetTextColor(FStyleColors::Foreground));
+		.SetTextColor(FStyleColors::Foreground);
 
-	Set("NiagaraEditor.SystemOverview.TableViewRow.Group", FTableRowStyle(NormalTableRowStyle)
+	Set("NiagaraEditor.Stack.TableViewRow.ItemContent", FTableRowStyle(BaseStackTableRowStyle)
+		.SetEvenRowBackgroundHoveredBrush(FSlateColorBrush(PanelHoverColor))
+		.SetOddRowBackgroundHoveredBrush(FSlateColorBrush(PanelHoverColor))
+		.SetEvenRowBackgroundBrush(FSlateColorBrush(FStyleColors::Panel))
+		.SetOddRowBackgroundBrush(FSlateColorBrush(FStyleColors::Panel)));
+
+	Set("NiagaraEditor.Stack.TableViewRow.ItemContentAdvanced", FTableRowStyle(BaseStackTableRowStyle)
+		.SetEvenRowBackgroundHoveredBrush(FSlateColorBrush(RecessedHoverColor))
+		.SetOddRowBackgroundHoveredBrush(FSlateColorBrush(RecessedHoverColor))
+		.SetEvenRowBackgroundBrush(FSlateColorBrush(FStyleColors::Recessed))
+		.SetOddRowBackgroundBrush(FSlateColorBrush(FStyleColors::Recessed)));
+
+	Set("NiagaraEditor.Stack.TableViewRow.ItemContentNote", FTableRowStyle(BaseStackTableRowStyle)
+		.SetEvenRowBackgroundHoveredBrush(FSlateColorBrush(NoteHoverColor))
+		.SetOddRowBackgroundHoveredBrush(FSlateColorBrush(NoteHoverColor))
+		.SetEvenRowBackgroundBrush(FSlateColorBrush(CustomNoteBackgroundColor))
+		.SetOddRowBackgroundBrush(FSlateColorBrush(CustomNoteBackgroundColor)));
+
+	Set("NiagaraEditor.Stack.TableViewRow.ItemHeader", FTableRowStyle(NormalTableRowStyle)
+		.SetEvenRowBackgroundHoveredBrush(FSlateColorBrush(HeaderHoverColor))
+		.SetOddRowBackgroundHoveredBrush(FSlateColorBrush(HeaderHoverColor))
 		.SetEvenRowBackgroundBrush(FSlateColorBrush(FStyleColors::Header))
-		.SetOddRowBackgroundBrush(FSlateColorBrush(FStyleColors::Header))
-		.SetEvenRowBackgroundHoveredBrush(CORE_IMAGE_BRUSH("Common/Selection", Icon8x8, FLinearColor(1.0f, 1.0f, 1.0f, 0.1f)))
-		.SetOddRowBackgroundHoveredBrush(CORE_IMAGE_BRUSH("Common/Selection", Icon8x8, FLinearColor(1.0f, 1.0f, 1.0f, 0.1f)))
+		.SetOddRowBackgroundBrush(FSlateColorBrush(FStyleColors::Header)));
+
+	Set("NiagaraEditor.Stack.TableViewRow.Spacer", FTableRowStyle(NormalTableRowStyle)
 		.SetSelectorFocusedBrush(FSlateNoResource())
-		.SetActiveBrush(FSlateColorBrush(FStyleColors::Select))
-		.SetActiveHoveredBrush(FSlateColorBrush(FStyleColors::Select))
-		.SetInactiveBrush(FSlateColorBrush(FStyleColors::SelectInactive))
-		.SetActiveHighlightedBrush(FSlateColorBrush(FStyleColors::PrimaryHover))
-		.SetTextColor(FStyleColors::Foreground));
-
-	Set("NiagaraEditor.Stack.BackgroundColor", FLinearColor(FColor(96, 96, 96)));
-	Set("NiagaraEditor.Stack.Item.HeaderBackgroundColor", FLinearColor(FColor(48, 48, 48)));
-	Set("NiagaraEditor.Stack.Item.ContentBackgroundColor", FLinearColor(FColor(62, 62, 62)));
-	Set("NiagaraEditor.Stack.Item.ContentAdvancedBackgroundColor", FLinearColor(FColor(53, 53, 53)));
-	Set("NiagaraEditor.Stack.Item.FooterBackgroundColor", FLinearColor(FColor(75, 75, 75)));
-	Set("NiagaraEditor.Stack.Item.CustomNoteBackgroundColor", FLinearColor(FColor(56, 111, 75)));
-	Set("NiagaraEditor.Stack.Item.InfoBackgroundColor", FLinearColor(FColor(68, 100, 106)));
-	Set("NiagaraEditor.Stack.Item.WarningBackgroundColor", FLinearColor(FColor(97, 97, 68)));
-	Set("NiagaraEditor.Stack.Item.ErrorBackgroundColor", FLinearColor(FColor(126, 78, 68)));
+		.SetActiveBrush(FSlateNoResource())
+		.SetActiveHoveredBrush(FSlateNoResource())
+		.SetInactiveBrush(FSlateNoResource())
+		.SetActiveHighlightedBrush(FSlateNoResource())
+		.SetEvenRowBackgroundBrush(FSlateNoResource())
+		.SetOddRowBackgroundBrush(FSlateNoResource())
+		.SetEvenRowBackgroundHoveredBrush(FSlateNoResource())
+		.SetOddRowBackgroundHoveredBrush(FSlateNoResource()));
+
+	Set("NiagaraEditor.SystemOverview.TableViewRow.Item", FTableRowStyle(NormalTableRowStyle)
+		.SetEvenRowBackgroundHoveredBrush(FSlateColorBrush(PanelHoverColor))
+		.SetOddRowBackgroundHoveredBrush(FSlateColorBrush(PanelHoverColor))
+		.SetEvenRowBackgroundBrush(FSlateColorBrush(FStyleColors::Panel))
+		.SetOddRowBackgroundBrush(FSlateColorBrush(FStyleColors::Panel)));
+
+	Set("NiagaraEditor.SystemOverview.TableViewRow.Group", FTableRowStyle(NormalTableRowStyle)
+		.SetEvenRowBackgroundHoveredBrush(FSlateColorBrush(HeaderHoverColor))
+		.SetOddRowBackgroundHoveredBrush(FSlateColorBrush(HeaderHoverColor))
+		.SetEvenRowBackgroundBrush(FSlateColorBrush(FStyleColors::Header))
+		.SetOddRowBackgroundBrush(FSlateColorBrush(FStyleColors::Header)));
 
 	Set("NiagaraEditor.Stack.UnknownColor", FLinearColor(1, 0, 1));
 
diff --git a/Engine/Plugins/FX/Niagara/Source/NiagaraEditorWidgets/Private/SNiagaraOverviewStack.cpp b/Engine/Plugins/FX/Niagara/Source/NiagaraEditorWidgets/Private/SNiagaraOverviewStack.cpp
--- a/Engine/Plugins/FX/Niagara/Source/NiagaraEditorWidgets/Private/SNiagaraOverviewStack.cpp
+++ b/Engine/Plugins/FX/Niagara/Source/NiagaraEditorWidgets/Private/SNiagaraOverviewStack.cpp
@@ -3,6 +3,15 @@
 #include "SNiagaraOverviewStack.h"
 #include "NiagaraSystem.h"
 #include "NiagaraSystemEditorData.h"
+#include "ScopedTransaction.h"
+#include "SNiagaraOverviewInlineParameterBox.h"
+#include "SNiagaraStack.h"
+#include "Stack/SNiagaraStackInheritanceIcon.h"
+#include "Stack/SNiagaraStackIssueIcon.h"
+#include "Stack/SNiagaraStackItemGroupAddButton.h"
+#include "Stack/SNiagaraStackRowPerfWidget.h"
+#include "Styling/StyleColors.h"
+#include "Subsystems/AssetEditorSubsystem.h"
 #include "NiagaraStackEditorData.h"
 #include "ViewModels/NiagaraEmitterViewModel.h"
 #include "ViewModels/NiagaraSystemViewModel.h"
@@ -830,7 +839,7 @@ TSharedRef<ITableRow> SNiagaraOverviewStack::OnGenerateRowForEntry(UNiagaraStack
 		ContentBox->AddSlot()
 		.VAlign(VAlign_Center)
 		.AutoWidth()
-		.Padding(4.0f, 0.0f)
+		.Padding(4, 0, 0, 0)
 		[
 			SNew(SNiagaraStackIssueIcon, StackViewModel, StackItem)
 			.IconMode(SNiagaraStackIssueIcon::EIconMode::Compact)
@@ -839,7 +848,7 @@ TSharedRef<ITableRow> SNiagaraOverviewStack::OnGenerateRowForEntry(UNiagaraStack
 		// Perf widget
 		ContentBox->AddSlot()
         .AutoWidth()
-        .Padding(2, 0, 2, 0)
+        .Padding(3, 0, 0, 0)
         [
             SNew(SNiagaraStackRowPerfWidget, Item)
         ];
@@ -890,11 +899,23 @@ TSharedRef<ITableRow> SNiagaraOverviewStack::OnGenerateRowForEntry(UNiagaraStack
 			}
 		}
 
+		// Inheritance Icon
+		if (StackItem->SupportsInheritance())
+		{
+			OptionsBox->AddSlot()
+				.VAlign(VAlign_Center)
+				.AutoWidth()
+				.Padding(3, 0, 0, 0)
+				[
+					SNew(SNiagaraStackInheritanceIcon, StackItem)
+				];
+		}
+
 		// Enabled checkbox
 		OptionsBox->AddSlot()
 			.VAlign(VAlign_Center)
 			.AutoWidth()
-			.Padding(0)
+			.Padding(3, 0, 0, 0)
 			[
 				SNew(SNiagaraSystemOverviewEnabledCheckBox)
 				.Visibility(this, &SNiagaraOverviewStack::GetEnabledCheckBoxVisibility, StackItem)
@@ -926,6 +947,7 @@ TSharedRef<ITableRow> SNiagaraOverviewStack::OnGenerateRowForEntry(UNiagaraStack
 				]
 				+ SWrapBox::Slot()
 				.VAlign(VAlign_Center)
+				.Padding(3, 0, 0, 0)
 				[
 					SNew(SBox)
 					.VAlign(VAlign_Center)
@@ -1069,6 +1091,7 @@ TSharedRef<ITableRow> SNiagaraOverviewStack::OnGenerateRowForEntry(UNiagaraStack
 
 		// Spacer
 		ContentBox->AddSlot()
+			.Padding(0, 0, 2, 0)
 			[
 				SNew(SSpacer)
 			];
@@ -1081,6 +1104,18 @@ TSharedRef<ITableRow> SNiagaraOverviewStack::OnGenerateRowForEntry(UNiagaraStack
 				SNew(SNiagaraStackRowPerfWidget, Item)
 			];
 
+		// Inheritance Icon
+		if (StackItemGroup->SupportsInheritance())
+		{
+			ContentBox->AddSlot()
+				.VAlign(VAlign_Center)
+				.AutoWidth()
+				.Padding(0, 0, 0, 0)
+				[
+					SNew(SNiagaraStackInheritanceIcon, Item)
+				];
+		}
+
 		// Delete button
 		if (StackItemGroup->SupportsDelete())
 		{
@@ -1094,7 +1129,6 @@ TSharedRef<ITableRow> SNiagaraOverviewStack::OnGenerateRowForEntry(UNiagaraStack
 					.IsFocusable(false)
 					.ToolTipText(this, &SNiagaraOverviewStack::GetItemGroupDeleteButtonToolTip, StackItemGroup)
 					.OnClicked(this, &SNiagaraOverviewStack::OnItemGroupDeleteClicked, StackItemGroup)
-					.IsEnabled(this, &SNiagaraOverviewStack::GetItemGroupDeleteButtonIsEnabled, StackItemGroup)
 					.Visibility(this, &SNiagaraOverviewStack::GetItemGroupDeleteButtonVisibility, StackItemGroup)
 					.Content()
 					[
@@ -1110,8 +1144,9 @@ TSharedRef<ITableRow> SNiagaraOverviewStack::OnGenerateRowForEntry(UNiagaraStack
 		if (StackItemGroup->SupportsChangeEnabled())
 		{
 			ContentBox->AddSlot()
+				.VAlign(VAlign_Center)
 				.AutoWidth()
-				.Padding(3, 3, 3, 3)
+				.Padding(2, 1, 1, 1)
 				[
 					SNew(SCheckBox)
 					.ForegroundColor(FSlateColor::UseSubduedForeground())
@@ -1127,7 +1162,7 @@ TSharedRef<ITableRow> SNiagaraOverviewStack::OnGenerateRowForEntry(UNiagaraStack
 			ContentBox->AddSlot()
 			.VAlign(VAlign_Center)
 			.AutoWidth()
-			.Padding(0, 0, 1, 0)
+			.Padding(0)
 			[
 				SNew(SNiagaraStackItemGroupAddButton, StackItemGroup, StackItemGroup->GetAddUtilities())
 				.Width(22)
@@ -1320,15 +1355,11 @@ FText SNiagaraOverviewStack::GetItemGroupDeleteButtonToolTip(UNiagaraStackItemGr
 	return Message;
 }
 
-bool SNiagaraOverviewStack::GetItemGroupDeleteButtonIsEnabled(UNiagaraStackItemGroup* Group) const
-{
-	FText UnusedMessage;
-	return Group->TestCanDeleteWithMessage(UnusedMessage);
-}
-
 EVisibility SNiagaraOverviewStack::GetItemGroupDeleteButtonVisibility(UNiagaraStackItemGroup* Group) const
 {
-	return Group->SupportsDelete() ? EVisibility::Visible : EVisibility::Collapsed;
+	FText UnusedMessage;
+	return Group->SupportsDelete() && Group->TestCanDeleteWithMessage(UnusedMessage)
+		? EVisibility::Visible : EVisibility::Collapsed;
 }
 
 FReply SNiagaraOverviewStack::OnItemGroupDeleteClicked(UNiagaraStackItemGroup* Group)
diff --git a/Engine/Plugins/FX/Niagara/Source/NiagaraEditorWidgets/Private/SNiagaraStack.cpp b/Engine/Plugins/FX/Niagara/Source/NiagaraEditorWidgets/Private/SNiagaraStack.cpp
--- a/Engine/Plugins/FX/Niagara/Source/NiagaraEditorWidgets/Private/SNiagaraStack.cpp
+++ b/Engine/Plugins/FX/Niagara/Source/NiagaraEditorWidgets/Private/SNiagaraStack.cpp
@@ -643,7 +643,7 @@ TSharedRef<ITableRow> SNiagaraStack::OnGenerateRowForStackItem(UNiagaraStackEntr
 {
 	TSharedRef<SNiagaraStackTableRow> Container = ConstructContainerForItem(Item);
 	FRowWidgets RowWidgets = ConstructNameAndValueWidgetsForItem(Item, Container);
-	Container->SetNameAndValueContent(RowWidgets.NameWidget, RowWidgets.ValueWidget, RowWidgets.EditConditionWidget, RowWidgets.ResetWidget);
+	Container->SetNameAndValueContent(RowWidgets.NameWidget, RowWidgets.ValueWidget);
 	return Container;
 }
 
@@ -744,108 +744,74 @@ TSharedPtr<SWidget> SNiagaraStack::GenerateStackMenu(TWeakPtr<UNiagaraStackViewM
 
 TSharedRef<SNiagaraStackTableRow> SNiagaraStack::ConstructContainerForItem(UNiagaraStackEntry* Item)
 {
-	float LeftContentPadding = 4;
-	float RightContentPadding = 6;
-	TAttribute<FMargin> RowPadding = FMargin(0, 0, 0, 0);
+	float LeftContentPadding = 0;
+	float RightContentPadding = 0;
+	float ItemContentVerticalPadding = 1;
 	FMargin ContentPadding(LeftContentPadding, 0, RightContentPadding, 0);
-	FSlateColor ItemBackgroundColor = FStyleColors::Panel;
-	FSlateColor ItemForegroundColor = FStyleColors::Foreground;
 	FSlateColor IndicatorColor = FLinearColor::Transparent;
-	bool bIsCategoryIconHighlighted;
-	bool bShowExecutionCategoryIcon;
+	bool bIsCategoryIconHighlighted = false;
+	bool bShowExecutionCategoryIcon = false;
+	const FTableRowStyle* TableRowStyle = nullptr;
 	switch (Item->GetStackRowStyle())
 	{
-	case UNiagaraStackEntry::EStackRowStyle::None:
-		bIsCategoryIconHighlighted = false;
-		bShowExecutionCategoryIcon = false;
-		break;
 	case UNiagaraStackEntry::EStackRowStyle::GroupHeader:
 		ContentPadding = FMargin(LeftContentPadding, 3, 0, 3);
-		ItemBackgroundColor = FStyleColors::Header;
 		bIsCategoryIconHighlighted = true;
 		bShowExecutionCategoryIcon = true;
-		break;
-	case UNiagaraStackEntry::EStackRowStyle::GroupFooter:
-		RowPadding = FMargin(0, 0, 0, 6);
-		bIsCategoryIconHighlighted = false;
-		bShowExecutionCategoryIcon = false;
-		ItemBackgroundColor = FStyleColors::Recessed;
+		TableRowStyle = &FNiagaraEditorWidgetsStyle::Get().GetWidgetStyle<FTableRowStyle>("NiagaraEditor.Stack.TableViewRow.ItemHeader");
 		break;
 	case UNiagaraStackEntry::EStackRowStyle::ItemHeader:
-		RowPadding.Bind(TAttribute<FMargin>::FGetter::CreateLambda([Item]()
-		{
-			if (Item->GetIsExpanded())
-			{
-				return FMargin(0, 2, 0, 0);
-			}
-			else
-			{
-				return FMargin(0, 2, 0, 2);
-			}
-		}));
 		ContentPadding = FMargin(LeftContentPadding, 2, 2, 2);
-		bIsCategoryIconHighlighted = false;
 		bShowExecutionCategoryIcon = true;
+		TableRowStyle = &FNiagaraEditorWidgetsStyle::Get().GetWidgetStyle<FTableRowStyle>("NiagaraEditor.Stack.TableViewRow.ItemHeader");
 		break;
 	case UNiagaraStackEntry::EStackRowStyle::ItemContent:
-		ContentPadding = FMargin(LeftContentPadding, 3, RightContentPadding, 3);
-		bIsCategoryIconHighlighted = false;
-		bShowExecutionCategoryIcon = false;
+		ContentPadding = FMargin(LeftContentPadding, ItemContentVerticalPadding, RightContentPadding, ItemContentVerticalPadding);
 		break;
 	case UNiagaraStackEntry::EStackRowStyle::ItemContentAdvanced:
-		ContentPadding = FMargin(LeftContentPadding, 3, RightContentPadding, 3);
-		bIsCategoryIconHighlighted = false;
-		bShowExecutionCategoryIcon = false;
-		ItemBackgroundColor = FStyleColors::Recessed;
+		ContentPadding = FMargin(LeftContentPadding, ItemContentVerticalPadding, RightContentPadding, ItemContentVerticalPadding);
+		TableRowStyle = &FNiagaraEditorWidgetsStyle::Get().GetWidgetStyle<FTableRowStyle>("NiagaraEditor.Stack.TableViewRow.ItemContentAdvanced");
 		break;
-	case UNiagaraStackEntry::EStackRowStyle::ItemFooter:
-		RowPadding = FMargin(0, 0, 0, 2);
-		bIsCategoryIconHighlighted = false;
-		bShowExecutionCategoryIcon = false;
+	case UNiagaraStackEntry::EStackRowStyle::ItemContentNote:
+		ContentPadding = FMargin(LeftContentPadding, ItemContentVerticalPadding, RightContentPadding, ItemContentVerticalPadding);
+		TableRowStyle = &FNiagaraEditorWidgetsStyle::Get().GetWidgetStyle<FTableRowStyle>("NiagaraEditor.Stack.TableViewRow.ItemContentNote");
 		break;
 	case UNiagaraStackEntry::EStackRowStyle::ItemCategory:
-		ContentPadding = FMargin(LeftContentPadding, 3, RightContentPadding, 3);
-		bIsCategoryIconHighlighted = false;
-		bShowExecutionCategoryIcon = false;
+		ContentPadding = FMargin(LeftContentPadding, ItemContentVerticalPadding, RightContentPadding, ItemContentVerticalPadding);
+		TableRowStyle = &FNiagaraEditorWidgetsStyle::Get().GetWidgetStyle<FTableRowStyle>("NiagaraEditor.Stack.TableViewRow.ItemHeader");
 		break;
 	case UNiagaraStackEntry::EStackRowStyle::StackIssue:
-
 		switch (Item->GetIssueSeverity())
 		{
 			case EStackIssueSeverity::Error:
-				ItemBackgroundColor = FStyleColors::Header;
 				IndicatorColor = FStyleColors::Error;
 				break;
 			case EStackIssueSeverity::Warning:
-				ItemBackgroundColor = FStyleColors::Header;
 				IndicatorColor = FStyleColors::Warning;
 				break;
 			case EStackIssueSeverity::Info:
-				ItemBackgroundColor = FStyleColors::Header;
 				IndicatorColor = FStyleColors::Foreground;
 				break;
 			case EStackIssueSeverity::CustomNote:
-				ItemBackgroundColor = FNiagaraEditorWidgetsStyle::Get().GetColor("NiagaraEditor.Stack.Item.CustomNoteBackgroundColor");
+				IndicatorColor = FStyleColors::Transparent;
 				break;
 			default:
 				checkf(false, TEXT("Issue severity not set for stack issue."));
 		}
-
-		ContentPadding = FMargin(LeftContentPadding, 3, RightContentPadding, 3);
-		bIsCategoryIconHighlighted = false;
-		bShowExecutionCategoryIcon = false;
-		break;
-	default:
-		bIsCategoryIconHighlighted = false;
-		bShowExecutionCategoryIcon = false;
+		ContentPadding = FMargin(LeftContentPadding, ItemContentVerticalPadding, RightContentPadding, ItemContentVerticalPadding);
 		break;
+	case UNiagaraStackEntry::EStackRowStyle::Spacer:
+		TableRowStyle = &FNiagaraEditorWidgetsStyle::Get().GetWidgetStyle<FTableRowStyle>("NiagaraEditor.Stack.TableViewRow.Spacer");
+	}
+
+	if (TableRowStyle == nullptr)
+	{
+		TableRowStyle = &FNiagaraEditorWidgetsStyle::Get().GetWidgetStyle<FTableRowStyle>("NiagaraEditor.Stack.TableViewRow.ItemContent");
 	}
 
 	return SNew(SNiagaraStackTableRow, StackViewModel, Item, StackCommandContext.ToSharedRef(), StackTree.ToSharedRef())
-		.RowPadding(RowPadding)
+		.Style(TableRowStyle)
 		.ContentPadding(ContentPadding)
- 		.ItemBackgroundColor(ItemBackgroundColor)
-		.ItemForegroundColor(ItemForegroundColor)
 		.IndicatorColor(IndicatorColor)
 		.IsCategoryIconHighlighted(bIsCategoryIconHighlighted)
 		.ShowExecutionCategoryIcon(bShowExecutionCategoryIcon)
@@ -948,8 +914,9 @@ SNiagaraStack::FRowWidgets SNiagaraStack::ConstructNameAndValueWidgetsForItem(UN
 	}
 	else if (Item->IsA<UNiagaraStackInputCategory>())
 	{
+		Container->SetOverrideNameAlignment(HAlign_Left, VAlign_Center);
 		return FRowWidgets(SNew(STextBlock)
-			.TextStyle(FNiagaraEditorWidgetsStyle::Get(), "NiagaraEditor.Stack.ItemText")
+			.TextStyle(FNiagaraEditorWidgetsStyle::Get(), "NiagaraEditor.Stack.ParameterText")
 			.ToolTipText_UObject(Item, &UNiagaraStackEntry::GetTooltipText)
 			.Text_UObject(Item, &UNiagaraStackEntry::GetDisplayName)
 			.ColorAndOpacity(this, &SNiagaraStack::GetTextColorForItem, Item)
@@ -1003,10 +970,27 @@ SNiagaraStack::FRowWidgets SNiagaraStack::ConstructNameAndValueWidgetsForItem(UN
 			Container->SetOverrideNameWidth(PropertyRowWidgets.WholeRowWidgetLayoutData.MinWidth, PropertyRowWidgets.WholeRowWidgetLayoutData.MaxWidth);
 			Container->SetOverrideNameAlignment(PropertyRowWidgets.WholeRowWidgetLayoutData.HorizontalAlignment, PropertyRowWidgets.WholeRowWidgetLayoutData.VerticalAlignment);
 			PropertyRowWidgets.WholeRowWidget->SetEnabled(IsEnabled);
-
-			FRowWidgets RowWidgets(PropertyRowWidgets.WholeRowWidget.ToSharedRef()); 
-			RowWidgets.EditConditionWidget = PropertyRowWidgets.EditConditionWidget;
-			return RowWidgets;
+			TSharedRef<SHorizontalBox> RowBox = SNew(SHorizontalBox);
+			if (PropertyRowWidgets.EditConditionWidget.IsValid())
+			{
+				RowBox->AddSlot()
+				.AutoWidth()
+				.Padding(0, 0, 3, 0)
+				[
+					PropertyRowWidgets.EditConditionWidget.ToSharedRef()
+				];
+			}
+			RowBox->AddSlot()
+			[
+				PropertyRowWidgets.WholeRowWidget.ToSharedRef()
+			];
+			RowBox->AddSlot()
+			.AutoWidth()
+			.Padding(3, 0, 0, 0)
+			[
+				SNew(SResetToDefaultPropertyEditor, PropertyRow->GetDetailTreeNode()->CreatePropertyHandle())
+			];
+			return FRowWidgets(RowBox); 
 		}
 		else
 		{
@@ -1016,10 +1000,39 @@ SNiagaraStack::FRowWidgets SNiagaraStack::ConstructNameAndValueWidgetsForItem(UN
 			PropertyRowWidgets.NameWidget->SetEnabled(IsEnabled);
 			PropertyRowWidgets.ValueWidget->SetEnabled(IsEnabled);
 
-			FRowWidgets RowWidgets(PropertyRowWidgets.NameWidget.ToSharedRef(), PropertyRowWidgets.ValueWidget.ToSharedRef());
-			RowWidgets.EditConditionWidget = PropertyRowWidgets.EditConditionWidget;
-			RowWidgets.ResetWidget = SNew(SResetToDefaultPropertyEditor, PropertyRow->GetDetailTreeNode()->CreatePropertyHandle());
-			return RowWidgets;
+			TSharedPtr<SWidget> NameWidget;
+			if (PropertyRowWidgets.EditConditionWidget.IsValid())
+			{
+				NameWidget = SNew(SHorizontalBox)
+				+ SHorizontalBox::Slot()
+				.AutoWidth()
+				.Padding(0, 0, 3, 0)
+				[
+					PropertyRowWidgets.EditConditionWidget.ToSharedRef()
+				]
+				+ SHorizontalBox::Slot()
+				[
+					PropertyRowWidgets.NameWidget.ToSharedRef()
+				];
+			}
+			else
+			{
+				NameWidget = PropertyRowWidgets.NameWidget;
+			}
+
+			TSharedRef<SWidget> ValueWidget = SNew(SHorizontalBox)
+			+ SHorizontalBox::Slot()
+			[
+				PropertyRowWidgets.ValueWidget.ToSharedRef()
+			]
+			+ SHorizontalBox::Slot()
+			.AutoWidth()
+			.Padding(3, 0, 0, 0)
+			[
+				SNew(SResetToDefaultPropertyEditor, PropertyRow->GetDetailTreeNode()->CreatePropertyHandle())
+			];
+
+			return FRowWidgets(NameWidget.ToSharedRef(), ValueWidget);
 		}
 	}
 	else if (Item->IsA<UNiagaraStackItem>())
@@ -1037,6 +1050,12 @@ SNiagaraStack::FRowWidgets SNiagaraStack::ConstructNameAndValueWidgetsForItem(UN
 			.AutoWrapText(true)
 			.Justification(ETextJustify::Center));
 	}
+	else if (Item->IsA<UNiagaraStackSpacer>())
+	{
+		UNiagaraStackSpacer* Spacer = CastChecked<UNiagaraStackSpacer>(Item);
+		return FRowWidgets(SNew(SBox)
+			.HeightOverride(Spacer->GetSpacerHeight()));
+	}
 	else
 	{
 		return FRowWidgets(SNullWidget::NullWidget);
diff --git a/Engine/Plugins/FX/Niagara/Source/NiagaraEditorWidgets/Private/SNiagaraStack.h b/Engine/Plugins/FX/Niagara/Source/NiagaraEditorWidgets/Private/SNiagaraStack.h
--- a/Engine/Plugins/FX/Niagara/Source/NiagaraEditorWidgets/Private/SNiagaraStack.h
+++ b/Engine/Plugins/FX/Niagara/Source/NiagaraEditorWidgets/Private/SNiagaraStack.h
@@ -46,8 +46,6 @@ private:
 
 		TSharedRef<SWidget> NameWidget;
 		TSharedPtr<SWidget> ValueWidget;
-		TSharedPtr<SWidget> EditConditionWidget;
-		TSharedPtr<SWidget> ResetWidget;
 	};
 
 	void SynchronizeTreeExpansion();
diff --git a/Engine/Plugins/FX/Niagara/Source/NiagaraEditorWidgets/Private/Stack/SNiagaraStackErrorItem.cpp b/Engine/Plugins/FX/Niagara/Source/NiagaraEditorWidgets/Private/Stack/SNiagaraStackErrorItem.cpp
--- a/Engine/Plugins/FX/Niagara/Source/NiagaraEditorWidgets/Private/Stack/SNiagaraStackErrorItem.cpp
+++ b/Engine/Plugins/FX/Niagara/Source/NiagaraEditorWidgets/Private/Stack/SNiagaraStackErrorItem.cpp
@@ -75,7 +75,6 @@ void SNiagaraStackErrorItemFix::Construct(const FArguments& InArgs, UNiagaraStac
 	{
 		FixWidget = SNew(SHorizontalBox)
 		+ SHorizontalBox::Slot()
-		.Padding(0, 4, 0, 0)
 		.VAlign(VAlign_Center)
 		[
 			SNew(STextBlock)
diff --git a/Engine/Plugins/FX/Niagara/Source/NiagaraEditorWidgets/Private/Stack/SNiagaraStackFunctionInputCollection.cpp b/Engine/Plugins/FX/Niagara/Source/NiagaraEditorWidgets/Private/Stack/SNiagaraStackFunctionInputCollection.cpp
--- a/Engine/Plugins/FX/Niagara/Source/NiagaraEditorWidgets/Private/Stack/SNiagaraStackFunctionInputCollection.cpp
+++ b/Engine/Plugins/FX/Niagara/Source/NiagaraEditorWidgets/Private/Stack/SNiagaraStackFunctionInputCollection.cpp
@@ -28,7 +28,7 @@ void SNiagaraStackFunctionInputCollection::Construct(const FArguments& InArgs, U
 		]
 		+ SVerticalBox::Slot()
 		.AutoHeight()
-		.Padding(0, 2)
+		.Padding(4, 2)
 		[
 			SAssignNew(SectionSelectorBox, SWrapBox)
 			.UseAllottedSize(true)
diff --git a/Engine/Plugins/FX/Niagara/Source/NiagaraEditorWidgets/Private/Stack/SNiagaraStackFunctionInputValue.cpp b/Engine/Plugins/FX/Niagara/Source/NiagaraEditorWidgets/Private/Stack/SNiagaraStackFunctionInputValue.cpp
--- a/Engine/Plugins/FX/Niagara/Source/NiagaraEditorWidgets/Private/Stack/SNiagaraStackFunctionInputValue.cpp
+++ b/Engine/Plugins/FX/Niagara/Source/NiagaraEditorWidgets/Private/Stack/SNiagaraStackFunctionInputValue.cpp
@@ -52,6 +52,7 @@
 #include "NiagaraCommon.h"
 #include "NiagaraSettings.h"
 #include "SNiagaraParameterDropTarget.h"
+#include "Stack/SNiagaraStackIndent.h"
 
 #define LOCTEXT_NAMESPACE "NiagaraStackFunctionInputValue"
 
@@ -88,6 +89,12 @@ void SNiagaraStackFunctionInputValue::Construct(const FArguments& InArgs, UNiaga
 				SNew(SHorizontalBox)
 				.IsEnabled(this, &SNiagaraStackFunctionInputValue::GetInputEnabled)
 				+ SHorizontalBox::Slot()
+				.AutoWidth()
+				.Padding(0, 0, 3, 0)
+				[
+					SNew(SNiagaraStackIndent, FunctionInput, ENiagaraStackIndentMode::Value)
+				]
+				+ SHorizontalBox::Slot()
 				.VAlign(VAlign_Center)
 				.AutoWidth()
 				.Padding(0, 0, 3, 0)
diff --git a/Engine/Plugins/FX/Niagara/Source/NiagaraEditorWidgets/Private/Stack/SNiagaraStackIndent.cpp b/Engine/Plugins/FX/Niagara/Source/NiagaraEditorWidgets/Private/Stack/SNiagaraStackIndent.cpp
new file mode 100644
--- /dev/null
+++ b/Engine/Plugins/FX/Niagara/Source/NiagaraEditorWidgets/Private/Stack/SNiagaraStackIndent.cpp
@@ -0,0 +1,76 @@
+// Copyright Epic Games, Inc. All Rights Reserved.
+
+#include "Stack/SNiagaraStackIndent.h"
+
+#include "ViewModels/Stack/NiagaraStackEntry.h"
+
+void SNiagaraStackIndent::Construct(const FArguments& InArgs, UNiagaraStackEntry* InStackEntry, ENiagaraStackIndentMode InMode)
+{
+	StackEntry = InStackEntry;
+	Mode = InMode;
+
+	if (Mode == ENiagaraStackIndentMode::Name)
+	{
+		SingleIndentWidth = 16.0f;
+	}
+	else if (Mode == ENiagaraStackIndentMode::Value)
+	{
+		SingleIndentWidth = 14.0f;
+	}
+
+	ChildSlot
+	[
+		SNew(SBox)
+		.WidthOverride(this, &SNiagaraStackIndent::GetIndentWidth)
+	];
+}
+
+int32 SNiagaraStackIndent::OnPaint(const FPaintArgs& Args, const FGeometry& AllottedGeometry, const FSlateRect& MyCullingRect, FSlateWindowElementList& OutDrawElements, int32 LayerId, const FWidgetStyle& InWidgetStyle, bool bParentEnabled) const
+{
+	const FSlateBrush* BackgroundBrush = FAppStyle::Get().GetBrush("DetailsView.CategoryMiddle");
+	const FSlateBrush* DropShadowBrush = FAppStyle::Get().GetBrush("DetailsView.ArrayDropShadow");
+
+	int32 IndentLevel = StackEntry->GetIndentLevel();
+	if (Mode == ENiagaraStackIndentMode::Name)
+	{
+		float PaintOffsetX = 1.0f;
+		for (int32 i = 0; i < IndentLevel; ++i)
+		{
+			FSlateColor BackgroundColor = GetRowBackgroundColor(i);
+
+			FSlateDrawElement::MakeBox(
+				OutDrawElements,
+				LayerId,
+				AllottedGeometry.ToPaintGeometry(FVector2D((SingleIndentWidth * i) + PaintOffsetX, 0), FVector2D(SingleIndentWidth, AllottedGeometry.GetLocalSize().Y)),
+				BackgroundBrush,
+				ESlateDrawEffect::None,
+				BackgroundColor.GetColor(InWidgetStyle));
+
+			FSlateDrawElement::MakeBox(
+				OutDrawElements,
+				LayerId + 1,
+				AllottedGeometry.ToPaintGeometry(FVector2D((SingleIndentWidth * i) + PaintOffsetX, 0), FVector2D(SingleIndentWidth, AllottedGeometry.GetLocalSize().Y)),
+				DropShadowBrush);
+		}
+	}
+	else
+	{
+		FSlateDrawElement::MakeBox(
+			OutDrawElements,
+			LayerId,
+			AllottedGeometry.ToPaintGeometry(FVector2D((SingleIndentWidth * (IndentLevel - 1)), 0), FVector2D(SingleIndentWidth, AllottedGeometry.GetLocalSize().Y)),
+			DropShadowBrush);
+	}
+	return LayerId + 1;
+}
+
+FOptionalSize SNiagaraStackIndent::GetIndentWidth() const
+{
+	int32 IndentLevel = StackEntry->GetIndentLevel();
+	return IndentLevel * SingleIndentWidth;
+}
+
+FSlateColor SNiagaraStackIndent::GetRowBackgroundColor(int32 IndentLevel) const
+{
+	return FLinearColor(0, 0, 0, FMath::Min(IndentLevel + 1, 4) * 0.08f);
+}
\ No newline at end of file
diff --git a/Engine/Plugins/FX/Niagara/Source/NiagaraEditorWidgets/Private/Stack/SNiagaraStackIndent.h b/Engine/Plugins/FX/Niagara/Source/NiagaraEditorWidgets/Private/Stack/SNiagaraStackIndent.h
new file mode 100644
--- /dev/null
+++ b/Engine/Plugins/FX/Niagara/Source/NiagaraEditorWidgets/Private/Stack/SNiagaraStackIndent.h
@@ -0,0 +1,32 @@
+// Copyright Epic Games, Inc. All Rights Reserved.
+
+#pragma once
+
+#include "Widgets/DeclarativeSyntaxSupport.h"
+
+class UNiagaraStackEntry;
+
+enum class ENiagaraStackIndentMode
+{
+	Name,
+	Value
+};
+
+class SNiagaraStackIndent : public SCompoundWidget
+{
+public:
+	SLATE_BEGIN_ARGS(SNiagaraStackIndent) {}
+	SLATE_END_ARGS()
+
+	void Construct(const FArguments& InArgs, UNiagaraStackEntry* InStackEntry, ENiagaraStackIndentMode InMode);
+
+private:
+	virtual int32 OnPaint(const FPaintArgs& Args, const FGeometry& AllottedGeometry, const FSlateRect& MyCullingRect, FSlateWindowElementList& OutDrawElements, int32 LayerId, const FWidgetStyle& InWidgetStyle, bool bParentEnabled) const override;
+	FOptionalSize GetIndentWidth() const;
+	FSlateColor GetRowBackgroundColor(int32 IndentLevel) const;
+
+private:
+	UNiagaraStackEntry* StackEntry;
+	ENiagaraStackIndentMode Mode;
+	float SingleIndentWidth;
+};
\ No newline at end of file
diff --git a/Engine/Plugins/FX/Niagara/Source/NiagaraEditorWidgets/Private/Stack/SNiagaraStackInheritanceIcon.cpp b/Engine/Plugins/FX/Niagara/Source/NiagaraEditorWidgets/Private/Stack/SNiagaraStackInheritanceIcon.cpp
new file mode 100644
--- /dev/null
+++ b/Engine/Plugins/FX/Niagara/Source/NiagaraEditorWidgets/Private/Stack/SNiagaraStackInheritanceIcon.cpp
@@ -0,0 +1,32 @@
+// Copyright Epic Games, Inc. All Rights Reserved.
+
+#include "Stack/SNiagaraStackInheritanceIcon.h"
+
+#include "Styling/StyleColors.h"
+#include "ViewModels/Stack/NiagaraStackEntry.h"
+
+void SNiagaraStackInheritanceIcon::Construct(const FArguments& InArgs, UNiagaraStackEntry* InStackEntry)
+{
+	StackEntry = InStackEntry;
+	SetVisibility(TAttribute<EVisibility>::CreateSP(this, &SNiagaraStackInheritanceIcon::IsVisible));
+	FLinearColor IconColor = FStyleColors::Foreground.GetSpecifiedColor() * FLinearColor(1.0f, 1.0f, 1.0f, 0.4f);
+	ChildSlot
+	[
+		SNew(SBox)
+		.WidthOverride(14)
+		.HeightOverride(14)
+		.ToolTipText_UObject(StackEntry, &UNiagaraStackEntry::GetInheritanceMessage)
+		[
+			SNew(SImage)
+			.Image(FAppStyle::Get().GetBrush("Icons.Lock"))
+			.ColorAndOpacity(IconColor)
+			.DesiredSizeOverride(FVector2D(12.0f, 12.0f))
+		]
+	];
+}
+
+EVisibility SNiagaraStackInheritanceIcon::IsVisible() const
+{
+	return StackEntry->GetIsInherited() ? EVisibility::Visible : EVisibility::Collapsed;
+}
+
diff --git a/Engine/Plugins/FX/Niagara/Source/NiagaraEditorWidgets/Private/Stack/SNiagaraStackInheritanceIcon.h b/Engine/Plugins/FX/Niagara/Source/NiagaraEditorWidgets/Private/Stack/SNiagaraStackInheritanceIcon.h
new file mode 100644
--- /dev/null
+++ b/Engine/Plugins/FX/Niagara/Source/NiagaraEditorWidgets/Private/Stack/SNiagaraStackInheritanceIcon.h
@@ -0,0 +1,22 @@
+// Copyright Epic Games, Inc. All Rights Reserved.
+
+#pragma once
+
+#include "Widgets/DeclarativeSyntaxSupport.h"
+
+class UNiagaraStackEntry;
+
+class SNiagaraStackInheritanceIcon : public SCompoundWidget
+{
+public:
+	SLATE_BEGIN_ARGS(SNiagaraStackInheritanceIcon) {}
+	SLATE_END_ARGS()
+
+	void Construct(const FArguments& InArgs, UNiagaraStackEntry* InStackEntry);
+
+private:
+	EVisibility IsVisible() const;
+
+private:
+	UNiagaraStackEntry* StackEntry;
+};
\ No newline at end of file
diff --git a/Engine/Plugins/FX/Niagara/Source/NiagaraEditorWidgets/Private/Stack/SNiagaraStackIssueIcon.cpp b/Engine/Plugins/FX/Niagara/Source/NiagaraEditorWidgets/Private/Stack/SNiagaraStackIssueIcon.cpp
--- a/Engine/Plugins/FX/Niagara/Source/NiagaraEditorWidgets/Private/Stack/SNiagaraStackIssueIcon.cpp
+++ b/Engine/Plugins/FX/Niagara/Source/NiagaraEditorWidgets/Private/Stack/SNiagaraStackIssueIcon.cpp
@@ -15,6 +15,9 @@
 
 #define LOCTEXT_NAMESPACE "SNiagaraStackIssueIcon"
 
+const float SNiagaraStackIssueIcon::NormalSize = 16.0f;
+const float SNiagaraStackIssueIcon::CompactSize = 10.0f;
+
 void SNiagaraStackIssueIcon::Construct(const FArguments& InArgs, UNiagaraStackViewModel* InStackViewModel, UNiagaraStackEntry* InStackEntry)
 {
 	StackViewModel = InStackViewModel;
@@ -32,8 +35,8 @@ void SNiagaraStackIssueIcon::Construct(const FArguments& InArgs, UNiagaraStackVi
 			SNew(SBox)
 			.IsEnabled(this, &SNiagaraStackIssueIcon::GetIconIsEnabled)
 			.ToolTipText(this, &SNiagaraStackIssueIcon::GetIconToolTip)
-			.HeightOverride(16)
-			.WidthOverride(16)
+			.HeightOverride(this, &SNiagaraStackIssueIcon::GetIconHeight)
+			.WidthOverride(NormalSize)
 			.HAlign(HAlign_Center)
 			.VAlign(VAlign_Center)
 			[
@@ -66,8 +69,8 @@ void SNiagaraStackIssueIcon::Construct(const FArguments& InArgs, UNiagaraStackVi
 			.ToolTipText(this, &SNiagaraStackIssueIcon::GetIconToolTip)
 			[
 				SNew(SBox)
-				.HeightOverride(10)
-				.WidthOverride(10)
+				.HeightOverride(this, &SNiagaraStackIssueIcon::GetIconHeight)
+				.WidthOverride(CompactSize)
 				.HAlign(HAlign_Center)
 				.VAlign(VAlign_Center)
 				[
@@ -243,15 +246,22 @@ FText SNiagaraStackIssueIcon::GetIconToolTip() const
 	return IconToolTipCache.GetValue();
 }
 
+FOptionalSize SNiagaraStackIssueIcon::GetIconHeight() const
+{
+	return IconHeight;
+}
+
 void SNiagaraStackIssueIcon::UpdateFromEntry(ENiagaraStructureChangedFlags Flags)
 {
 	if (StackEntry.IsValid() == false || StackEntry->IsFinalized() || StackEntry->HasIssuesOrAnyChildHasIssues() == false)
 	{
 		IconBrush = FEditorStyle::GetBrush("NoBrush");
+		IconHeight = 0;
 		IconToolTipCache = FText();
 		return;
 	}
 
+	IconHeight = IconMode == EIconMode::Compact ? CompactSize : NormalSize;
 	if (StackEntry->GetTotalNumberOfErrorIssues() > 0)
 	{
 		IconBrush = IconMode == EIconMode::Compact
diff --git a/Engine/Plugins/FX/Niagara/Source/NiagaraEditorWidgets/Private/Stack/SNiagaraStackIssueIcon.h b/Engine/Plugins/FX/Niagara/Source/NiagaraEditorWidgets/Private/Stack/SNiagaraStackIssueIcon.h
--- a/Engine/Plugins/FX/Niagara/Source/NiagaraEditorWidgets/Private/Stack/SNiagaraStackIssueIcon.h
+++ b/Engine/Plugins/FX/Niagara/Source/NiagaraEditorWidgets/Private/Stack/SNiagaraStackIssueIcon.h
@@ -36,16 +36,21 @@ private:
 
 	const FSlateBrush* GetIconBrush() const;
 	FText GetIconToolTip() const;
+	FOptionalSize GetIconHeight() const;
 	EVisibility GetCompactIconBorderVisibility() const;
 
 	void UpdateFromEntry(ENiagaraStructureChangedFlags Flags);
 
 private:
 	const FSlateBrush* IconBrush;
+	float IconHeight;
 
 	mutable TOptional<FText> IconToolTipCache;
 
 	UNiagaraStackViewModel* StackViewModel;
 	TWeakObjectPtr<UNiagaraStackEntry> StackEntry;
 	EIconMode IconMode;
+
+	static const float NormalSize;
+	static const float CompactSize;
 };
\ No newline at end of file
diff --git a/Engine/Plugins/FX/Niagara/Source/NiagaraEditorWidgets/Private/Stack/SNiagaraStackItem.cpp b/Engine/Plugins/FX/Niagara/Source/NiagaraEditorWidgets/Private/Stack/SNiagaraStackItem.cpp
--- a/Engine/Plugins/FX/Niagara/Source/NiagaraEditorWidgets/Private/Stack/SNiagaraStackItem.cpp
+++ b/Engine/Plugins/FX/Niagara/Source/NiagaraEditorWidgets/Private/Stack/SNiagaraStackItem.cpp
@@ -17,6 +17,7 @@
 #include "Widgets/Layout/SBox.h"
 #include "ViewModels/Stack/NiagaraStackClipboardUtilities.h"
 #include "Styling/StyleColors.h"
+#include "Stack/SNiagaraStackInheritanceIcon.h"
 
 #define LOCTEXT_NAMESPACE "NiagaraStackItem"
 
@@ -114,6 +115,18 @@ void SNiagaraStackItem::Construct(const FArguments& InArgs, UNiagaraStackItem& I
 		];
 	}
 
+	// Inheritance Icon
+	if (Item->SupportsInheritance())
+	{
+		RowBox->AddSlot()
+		.VAlign(VAlign_Center)
+		.AutoWidth()
+		.Padding(0, 0, 2, 0)
+		[
+			SNew(SNiagaraStackInheritanceIcon, Item)
+		];
+	}
+
 	// Delete button
 	if (Item->SupportsDelete())
 	{
@@ -126,7 +139,7 @@ void SNiagaraStackItem::Construct(const FArguments& InArgs, UNiagaraStackItem& I
 			.ButtonStyle(FEditorStyle::Get(), "HoverHintOnly")
 			.IsFocusable(false)
 			.ToolTipText(this, &SNiagaraStackItem::GetDeleteButtonToolTipText)
-			.IsEnabled(this, &SNiagaraStackItem::GetDeleteButtonEnabled)
+			.Visibility(this, &SNiagaraStackItem::GetDeleteButtonVisibility)
 			.OnClicked(this, &SNiagaraStackItem::DeleteClicked)
 			.Content()
 			[
@@ -225,10 +238,10 @@ FText SNiagaraStackItem::GetDeleteButtonToolTipText() const
 	return CanDeleteMessage;
 }
 
-bool SNiagaraStackItem::GetDeleteButtonEnabled() const
+EVisibility SNiagaraStackItem::GetDeleteButtonVisibility() const
 {
 	FText Unused;
-	return Item->TestCanDeleteWithMessage(Unused);
+	return Item->TestCanDeleteWithMessage(Unused) ? EVisibility::Visible : EVisibility::Collapsed;
 }
 
 FReply SNiagaraStackItem::DeleteClicked()
diff --git a/Engine/Plugins/FX/Niagara/Source/NiagaraEditorWidgets/Private/Stack/SNiagaraStackItem.h b/Engine/Plugins/FX/Niagara/Source/NiagaraEditorWidgets/Private/Stack/SNiagaraStackItem.h
--- a/Engine/Plugins/FX/Niagara/Source/NiagaraEditorWidgets/Private/Stack/SNiagaraStackItem.h
+++ b/Engine/Plugins/FX/Niagara/Source/NiagaraEditorWidgets/Private/Stack/SNiagaraStackItem.h
@@ -37,11 +37,9 @@ private:
 
 	FReply ResetToBaseButtonClicked();
 
-	EVisibility GetDeleteButtonVisibility() const;
-
 	FText GetDeleteButtonToolTipText() const;
 
-	bool GetDeleteButtonEnabled() const;
+	EVisibility GetDeleteButtonVisibility() const;
 
 	FReply DeleteClicked();
 
diff --git a/Engine/Plugins/FX/Niagara/Source/NiagaraEditorWidgets/Private/Stack/SNiagaraStackItemGroup.cpp b/Engine/Plugins/FX/Niagara/Source/NiagaraEditorWidgets/Private/Stack/SNiagaraStackItemGroup.cpp
--- a/Engine/Plugins/FX/Niagara/Source/NiagaraEditorWidgets/Private/Stack/SNiagaraStackItemGroup.cpp
+++ b/Engine/Plugins/FX/Niagara/Source/NiagaraEditorWidgets/Private/Stack/SNiagaraStackItemGroup.cpp
@@ -4,6 +4,7 @@
 #include "NiagaraEditorWidgetsStyle.h"
 #include "NiagaraEditorStyle.h"
 #include "EditorStyleSet.h"
+#include "Stack/SNiagaraStackInheritanceIcon.h"
 #include "Stack/SNiagaraStackItemGroupAddButton.h"
 #include "ViewModels/Stack/NiagaraStackItemGroup.h"
 #include "ViewModels/Stack/NiagaraStackViewModel.h"
@@ -34,24 +35,39 @@ void SNiagaraStackItemGroup::Construct(const FArguments& InArgs, UNiagaraStackIt
 		.EditableNameStyle(FNiagaraEditorWidgetsStyle::Get(), "NiagaraEditor.Stack.EditableGroupText")
 	];
 
+	// Inheritance Icon
+	if (Group->SupportsInheritance())
+	{
+		RowBox->AddSlot()
+		.VAlign(VAlign_Center)
+		.AutoWidth()
+		.Padding(0, 0, 2, 0)
+		[
+			SNew(SNiagaraStackInheritanceIcon, Group)
+		];
+	}
+
 	// Delete group button
-	RowBox->AddSlot()
-	.AutoWidth()
-	[
-		SNew(SButton)
-		.ButtonStyle(FEditorStyle::Get(), "HoverHintOnly")
-		.IsFocusable(false)
-		.ToolTipText(this, &SNiagaraStackItemGroup::GetDeleteButtonToolTip)
-		.OnClicked(this, &SNiagaraStackItemGroup::DeleteClicked)
-		.IsEnabled(this, &SNiagaraStackItemGroup::GetDeleteButtonIsEnabled)
-		.Visibility(this, &SNiagaraStackItemGroup::GetDeleteButtonVisibility)
-		.Content()
+	if (Group->SupportsDelete())
+	{
+		RowBox->AddSlot()
+		.AutoWidth()
+		.Padding(0, 0, 2, 0)
 		[
-			SNew(SImage)
-			.Image(FAppStyle::Get().GetBrush("Icons.Delete"))
-			.ColorAndOpacity(FSlateColor::UseForeground())
-		]
-	];
+			SNew(SButton)
+			.ButtonStyle(FEditorStyle::Get(), "HoverHintOnly")
+			.IsFocusable(false)
+			.ToolTipText(this, &SNiagaraStackItemGroup::GetDeleteButtonToolTip)
+			.OnClicked(this, &SNiagaraStackItemGroup::DeleteClicked)
+			.Visibility(this, &SNiagaraStackItemGroup::GetDeleteButtonVisibility)
+			.Content()
+			[
+				SNew(SImage)
+				.Image(FAppStyle::Get().GetBrush("Icons.Delete"))
+				.ColorAndOpacity(FSlateColor::UseForeground())
+			]
+		];
+	}
 
 	// Enabled button
 	if (Group->SupportsChangeEnabled())
@@ -100,15 +116,11 @@ FText SNiagaraStackItemGroup::GetDeleteButtonToolTip() const
 	return Message;
 }
 
-bool SNiagaraStackItemGroup::GetDeleteButtonIsEnabled() const
-{
-	FText UnusedMessage;
-	return Group->TestCanDeleteWithMessage(UnusedMessage);
-}
-
 EVisibility SNiagaraStackItemGroup::GetDeleteButtonVisibility() const
 {
-	return Group->SupportsDelete() ? EVisibility::Visible : EVisibility::Collapsed;
+	FText UnusedMessage;
+	return Group->SupportsDelete() && Group->TestCanDeleteWithMessage(UnusedMessage)
+		? EVisibility::Visible : EVisibility::Collapsed;
 }
 
 FReply SNiagaraStackItemGroup::DeleteClicked()
diff --git a/Engine/Plugins/FX/Niagara/Source/NiagaraEditorWidgets/Private/Stack/SNiagaraStackItemGroup.h b/Engine/Plugins/FX/Niagara/Source/NiagaraEditorWidgets/Private/Stack/SNiagaraStackItemGroup.h
--- a/Engine/Plugins/FX/Niagara/Source/NiagaraEditorWidgets/Private/Stack/SNiagaraStackItemGroup.h
+++ b/Engine/Plugins/FX/Niagara/Source/NiagaraEditorWidgets/Private/Stack/SNiagaraStackItemGroup.h
@@ -21,8 +21,6 @@ private:
 
 	FText GetDeleteButtonToolTip() const;
 
-	bool GetDeleteButtonIsEnabled() const;
-
 	EVisibility GetDeleteButtonVisibility() const;
 
 	FReply DeleteClicked();
diff --git a/Engine/Plugins/FX/Niagara/Source/NiagaraEditorWidgets/Private/Stack/SNiagaraStackTableRow.cpp b/Engine/Plugins/FX/Niagara/Source/NiagaraEditorWidgets/Private/Stack/SNiagaraStackTableRow.cpp
--- a/Engine/Plugins/FX/Niagara/Source/NiagaraEditorWidgets/Private/Stack/SNiagaraStackTableRow.cpp
+++ b/Engine/Plugins/FX/Niagara/Source/NiagaraEditorWidgets/Private/Stack/SNiagaraStackTableRow.cpp
@@ -11,10 +11,12 @@
 #include "NiagaraEmitter.h"
 #include "NiagaraNodeFunctionCall.h"
 #include "NiagaraStackCommandContext.h"
+#include "Stack/SNiagaraStackIndent.h"
 #include "Stack/SNiagaraStackIssueIcon.h"
 #include "Widgets/Input/SButton.h"
 #include "Widgets/Images/SImage.h"
 #include "Widgets/Layout/SBox.h"
+#include "Widgets/SOverlay.h"
 #include "Framework/Application/SlateApplication.h"
 #include "ViewModels/NiagaraEmitterViewModel.h"
 #include "Styling/StyleColors.h"
@@ -23,8 +25,6 @@
 
 #define LOCTEXT_NAMESPACE "NiagaraStackTableRow"
 
-const float IndentSize = 12;
-
 void SNiagaraStackTableRow::Construct(const FArguments& InArgs, UNiagaraStackViewModel* InStackViewModel, UNiagaraStackEntry* InStackEntry, TSharedRef<FNiagaraStackCommandContext> InStackCommandContext, const TSharedRef<STreeView<UNiagaraStackEntry*>>& InOwnerTree)
 {
 	ContentPadding = InArgs._ContentPadding;
@@ -35,18 +35,17 @@ void SNiagaraStackTableRow::Construct(const FArguments& InArgs, UNiagaraStackVie
 	NameColumnWidthChanged = InArgs._OnNameColumnWidthChanged;
 	ValueColumnWidthChanged = InArgs._OnValueColumnWidthChanged;
 	IssueIconVisibility = InArgs._IssueIconVisibility;
-	RowPadding = InArgs._RowPadding;
 	StackViewModel = InStackViewModel;
 	StackEntry = InStackEntry;
 	StackCommandContext = InStackCommandContext;
 	OwnerTree = InOwnerTree;
 
+	NameVerticalAlignment = VAlign_Center;
+	NameHorizontalAlignment = HAlign_Fill;
+
 	ExpandedImage = FCoreStyle::Get().GetBrush("TreeArrow_Expanded");
 	CollapsedImage = FCoreStyle::Get().GetBrush("TreeArrow_Collapsed");
 
-	ItemBackgroundColor = InArgs._ItemBackgroundColor;
-	DisabledItemBackgroundColor = FStyleColors::Recessed;
-	ForegroundColor = InArgs._ItemForegroundColor;
 	IndicatorColor = InArgs._IndicatorColor;
 
 	ExecutionCategoryToolTipText = (InStackEntry->GetExecutionSubcategoryName() != NAME_None)
@@ -55,7 +54,7 @@ void SNiagaraStackTableRow::Construct(const FArguments& InArgs, UNiagaraStackVie
 
 	ConstructInternal(
 		STableRow<UNiagaraStackEntry*>::FArguments()
-			.Style(FNiagaraEditorWidgetsStyle::Get(), "NiagaraEditor.Stack.TableViewRow")
+			.Style(InArgs._Style)
 			.OnDragDetected(InArgs._OnDragDetected)
 			.OnDragLeave(InArgs._OnDragLeave)
 			.OnCanAcceptDrop(InArgs._OnCanAcceptDrop)
@@ -97,47 +96,7 @@ void SNiagaraStackTableRow::SetContentPadding(FMargin InContentPadding)
 	ContentPadding = InContentPadding;
 }
 
-// searches for the first parent stack entry without a front divider
-UNiagaraStackEntry* GetParentEntryNoDivider(UNiagaraStackEntry* Entry)
-{
-	UNiagaraStackEntry* Outer = Cast<UNiagaraStackEntry>(Entry->GetOuter());
-	if (Outer == nullptr || !Outer->HasFrontDivider())
-	{
-		return Outer;
-	}
-	return GetParentEntryNoDivider(Outer);
-}
-
-UNiagaraStackEntry* GetEntryAbove(UNiagaraStackEntry* Entry)
-{
-	UNiagaraStackEntry* Outer = Cast<UNiagaraStackEntry>(Entry->GetOuter());
-	if (Outer)
-	{
-		TArray<UNiagaraStackEntry*> FilteredChildren;
-		Outer->GetFilteredChildren(FilteredChildren);
-		for (int32 Index = 1; Index < FilteredChildren.Num(); Index++)
-		{
-			if (FilteredChildren[Index] == Entry)
-			{
-				return FilteredChildren[Index - 1];
-			}
-		}
-	}
-	return nullptr;
-}
-
-bool HasVisibleChildren(UNiagaraStackEntry* Entry)
-{
-	if (!Entry)
-	{
-		return false;
-	}
-	TArray<UNiagaraStackEntry*> Children;
-	Entry->GetFilteredChildren(Children);
-	return Children.Num() > 0;
-}
-
-void SNiagaraStackTableRow::SetNameAndValueContent(TSharedRef<SWidget> InNameWidget, TSharedPtr<SWidget> InValueWidget, TSharedPtr<SWidget> InEditConditionWidget, TSharedPtr<SWidget> InResetWidget)
+void SNiagaraStackTableRow::SetNameAndValueContent(TSharedRef<SWidget> InNameWidget, TSharedPtr<SWidget> InValueWidget)
 {
 	FSlateColor IconColor = FNiagaraEditorWidgetsStyle::Get().GetColor(FNiagaraStackEditorWidgetsUtilities::GetColorNameForExecutionCategory(StackEntry->GetExecutionCategoryName()));
 	if (bIsCategoryIconHighlighted)
@@ -158,21 +117,21 @@ void SNiagaraStackTableRow::SetNameAndValueContent(TSharedRef<SWidget> InNameWid
 	// Indent
 	+ SHorizontalBox::Slot()
 	.AutoWidth()
+	.Padding(0)
 	[
-		SNew(SBox)
-		.WidthOverride(this, &SNiagaraStackTableRow::GetIndentSize)
+		SNew(SNiagaraStackIndent, StackEntry, ENiagaraStackIndentMode::Name)
 	]
 	// Expand button
 	+ SHorizontalBox::Slot()
 	.AutoWidth()
 	.VAlign(VAlign_Center)
-	.Padding(0, 0, 1, 0)
+	.Padding(0)
 	[
 		SNew(SButton)
-		.ButtonStyle(FCoreStyle::Get(), "NoBorder")
+		.ButtonStyle(FNiagaraEditorWidgetsStyle::Get(), "NiagaraEditor.Stack.SimpleButton")
 		.Visibility(this, &SNiagaraStackTableRow::GetExpanderVisibility)
 		.OnClicked(this, &SNiagaraStackTableRow::ExpandButtonClicked)
-		.ContentPadding(2)
+		.ContentPadding(FMargin(1, 2, 1, 2))
 		.HAlign(HAlign_Center)
 		[
 			SNew(SImage)
@@ -201,34 +160,15 @@ void SNiagaraStackTableRow::SetNameAndValueContent(TSharedRef<SWidget> InNameWid
 			.ColorAndOpacity(IconColor)
 		]
 	]
-	// Edit condition
-	+ SHorizontalBox::Slot()
-	.AutoWidth()
-	.Padding(0, 0, 3, 0)
-	[
-		InEditConditionWidget.IsValid() ? InEditConditionWidget.ToSharedRef() : SNullWidget::NullWidget
-	]
 	// Name content
 	+ SHorizontalBox::Slot()
 	.HAlign(NameHorizontalAlignment)
 	.VAlign(NameVerticalAlignment)
+	.Padding(FMargin(0, ContentPadding.Top, 0, ContentPadding.Bottom))
 	[
 		InNameWidget
 	];
 
-	UNiagaraStackEntry* ParentEntry = GetParentEntryNoDivider(StackEntry);
-	bool bIsDisplayedInCategory = StackEntry->HasFrontDivider() && ParentEntry && ParentEntry->GetShouldShowInStack();
-	UNiagaraStackEntry* AboveEntry = GetEntryAbove(StackEntry);
-	if (!StackEntry->HasFrontDivider() && AboveEntry && AboveEntry->HasFrontDivider())
-	{
-		ContentPadding.Top += 6;
-	}
-	if (StackEntry->HasFrontDivider())
-	{
-		ContentPadding.Left += IndentSize * (bIsDisplayedInCategory ? 3 : 2) - 4;
-	}
-	bool bInsertDivAbove = StackEntry->GetStackRowStyle() == UNiagaraStackEntry::EStackRowStyle::ItemCategory && HasVisibleChildren(AboveEntry);
-
 	TSharedPtr<SWidget> ChildContent;
 	if (InValueWidget.IsValid())
 	{
@@ -241,23 +181,12 @@ void SNiagaraStackTableRow::SetNameAndValueContent(TSharedRef<SWidget> InNameWid
 		.Value(NameColumnWidth)
 		.OnSlotResized(SSplitter::FOnSlotResized::CreateSP(this, &SNiagaraStackTableRow::OnNameColumnWidthChanged))
 		[
-			SNew(SHorizontalBox)
-			+ SHorizontalBox::Slot()
-			.Padding(FMargin(ContentPadding.Left, 0, 0, 0))
-		    .AutoWidth()
-		    [
-				SNew(SBox)
-				.WidthOverride(StackEntry->HasFrontDivider() ? 1 : 0)
-		    ]
-		    + SHorizontalBox::Slot()
-            [
-				SNew(SBox)
-				.Padding(FMargin(0, ContentPadding.Top, 5, ContentPadding.Bottom))
-				.MinDesiredWidth(NameMinWidth.IsSet() ? NameMinWidth.GetValue() : FOptionalSize())
-				.MaxDesiredWidth(NameMaxWidth.IsSet() ? NameMaxWidth.GetValue() : FOptionalSize())
-				[
-                    NameContent
-                ]
+			SNew(SBox)
+			.Padding(FMargin(ContentPadding.Left, 0, 5, 0))
+			.MinDesiredWidth(NameMinWidth.IsSet() ? NameMinWidth.GetValue() : FOptionalSize())
+			.MaxDesiredWidth(NameMaxWidth.IsSet() ? NameMaxWidth.GetValue() : FOptionalSize())
+			[
+				NameContent
 			]
 		]
 
@@ -280,9 +209,7 @@ void SNiagaraStackTableRow::SetNameAndValueContent(TSharedRef<SWidget> InNameWid
 	else
 	{
 		ChildContent = SNew(SBox)
-		.Padding(ContentPadding)
-		.HAlign(NameHorizontalAlignment)
-		.VAlign(NameVerticalAlignment)
+		.Padding(FMargin(ContentPadding.Left, 0, ContentPadding.Right, 0))
 		.MinDesiredWidth(NameMinWidth.IsSet() ? NameMinWidth.GetValue() : FOptionalSize())
 		.MaxDesiredWidth(NameMaxWidth.IsSet() ? NameMaxWidth.GetValue() : FOptionalSize())
 		[
@@ -290,98 +217,79 @@ void SNiagaraStackTableRow::SetNameAndValueContent(TSharedRef<SWidget> InNameWid
 		];
 	}
 
-	FName AccentColorName = FNiagaraStackEditorWidgetsUtilities::GetIconColorNameForExecutionCategory(StackEntry->GetExecutionCategoryName());
-	const bool bDisplayingIndicator = IndicatorColor != FStyleColors::Transparent;
-	FSlateColor AccentColor = bDisplayingIndicator ? IndicatorColor : ((AccentColorName != NAME_None) ? FNiagaraEditorWidgetsStyle::Get().GetColor(AccentColorName) : FStyleColors::Transparent);
+	FName AccentColorName = FNiagaraStackEditorWidgetsUtilities::GetColorNameForExecutionCategory(StackEntry->GetExecutionCategoryName());
+	TOptional<FSlateColor> AccentColor = AccentColorName != NAME_None ? FNiagaraEditorWidgetsStyle::Get().GetColor(AccentColorName) : FStyleColors::Transparent;
 
-	ChildSlot
-	[
-		SNew(SHorizontalBox)
+	TSharedRef<SOverlay> RowOverlay = SNew(SOverlay)
 		.Visibility(this, &SNiagaraStackTableRow::GetRowVisibility)
-		// Accent color.
-		+ SHorizontalBox::Slot()
-		.AutoWidth()
-		.Padding(bDisplayingIndicator ? FMargin(0, 0, 5, 0) : FMargin(1, 0, 6, 0) )
+		+ SOverlay::Slot()
 		[
+			// Row content
 			SNew(SBorder)
- 			.BorderImage(FEditorStyle::GetBrush("WhiteBrush"))
-			.BorderBackgroundColor(AccentColor)
+			.BorderImage(FNiagaraEditorWidgetsStyle::Get().GetBrush("Niagara.TableViewRowBorder"))
+			.ToolTipText_UObject(StackEntry, &UNiagaraStackEntry::GetTooltipText)
+			.Padding(FMargin(5, 0, 0, 1))
+			[
+				SNew(SHorizontalBox)
+				+ SHorizontalBox::Slot()
+				.Padding(0.0f)
+				[
+					ChildContent.ToSharedRef()
+				]
+				// Issue Icon
+				+ SHorizontalBox::Slot()
+				.AutoWidth()
+				.Padding(3, 0)
+				[
+					SNew(SNiagaraStackIssueIcon, StackViewModel, StackEntry)
+					.Visibility(IssueIconVisibility)
+				]
+			]
+		];
+
+	if (AccentColor.IsSet())
+	{
+		RowOverlay->AddSlot()
+		.HAlign(HAlign_Left)
+		[
+			SNew(SBorder)
+			.BorderImage(FEditorStyle::GetBrush("WhiteBrush"))
+			.BorderBackgroundColor(AccentColor.GetValue())
 			.Padding(0)
 			[
 				SNew(SBox)
-				.WidthOverride(bDisplayingIndicator ? 6 : 4)
+				.WidthOverride(4)
 			]
-		]
-		// Content
-		+ SHorizontalBox::Slot()
-		.Padding(0.0f)
+		];
+	}
+
+	if (IndicatorColor.IsSet())
+	{
+		RowOverlay->AddSlot()
+		.HAlign(HAlign_Right)
 		[
-			// Row content
 			SNew(SBorder)
-			.BorderImage(bDisplayingIndicator ? FAppStyle::Get().GetBrush("Brushes.Header") : FAppStyle::Get().GetBrush("DetailsView.GridLine"))
-			.Padding(FMargin(0, 0, 0, 1))
+			.BorderImage(FEditorStyle::GetBrush("WhiteBrush"))
+			.BorderBackgroundColor(IndicatorColor.GetValue())
+			.Padding(0)
 			[
-					SNew(SBorder)
- 					.BorderImage(FEditorStyle::GetBrush("WhiteBrush"))
-					.BorderBackgroundColor(this, &SNiagaraStackTableRow::GetItemBackgroundColor)
- 					.ForegroundColor(ForegroundColor)
-					.Padding(0)
-					[
-						SNew(SHorizontalBox)
-						+ SHorizontalBox::Slot()
-						.Padding(0.0f)
-						[
-							SNew(SVerticalBox)
-							+SVerticalBox::Slot()
-							.Padding(bInsertDivAbove ? FMargin(8, 8, 8, 4) : 0)
-							.AutoHeight()
-							[
-								SNew(SBorder)
-								.BorderImage(FEditorStyle::GetBrush("WhiteBrush"))
- 								.BorderBackgroundColor(FStyleColors::Panel)
-								.Visibility(bInsertDivAbove ? EVisibility::Visible : EVisibility::Collapsed)
-								.Padding(0)
-								[
-									SNew(SBox)
-									.HeightOverride(1)
-								]
-							]
-							+SVerticalBox::Slot()
-							.Padding(0)
-							[
-								SNew(SBorder)
-								.BorderImage(this, &SNiagaraStackTableRow::GetSearchResultBorderBrush)
- 								.BorderBackgroundColor(FStyleColors::Select)
-								.Padding(0)
-								[
-									SNew(SBorder)
-									.BorderImage(this, &SNiagaraStackTableRow::GetSelectionBorderBrush)
-									.Padding(0)
-									[
-										ChildContent.ToSharedRef()
-									]
-								]
-							]
-						]
-						+ SHorizontalBox::Slot()
-						.AutoWidth()
-						.Padding(3, 0)
-						[
-							SNew(SNiagaraStackIssueIcon, StackViewModel, StackEntry)
-							.Visibility(IssueIconVisibility)
-						]
-						// Reset To Default
-						+ SHorizontalBox::Slot()
-						.AutoWidth()
-						.VAlign(VAlign_Center)
-						.Padding(0, 0, 3, 0)
-						[
-							InResetWidget.IsValid() ? InResetWidget.ToSharedRef() : SNullWidget::NullWidget
-						]
-					]
-				]
+				SNew(SBox)
+				.WidthOverride(2)
 			]
-		
+		];
+	}
+
+	// Search result indicator.
+	RowOverlay->AddSlot()
+	[
+		SNew(SBorder)
+		.BorderImage(FNiagaraEditorWidgetsStyle::Get().GetBrush("NiagaraEditor.Stack.SearchResult"))
+		.Visibility(this, &SNiagaraStackTableRow::GetSearchResultBorderVisibility)
+	];
+
+	ChildSlot
+	[
+		RowOverlay
 	];
 }
 
@@ -489,12 +397,6 @@ FReply SNiagaraStackTableRow::OnMouseButtonUp(const FGeometry& MyGeometry, const
 	return STableRow<UNiagaraStackEntry*>::OnMouseButtonUp(MyGeometry, MouseEvent);
 }
 
-const FSlateBrush* SNiagaraStackTableRow::GetBorder() const
-{
-	// Return no brush here so that the background doesn't change.  The border color changing will be handled by an internal border.
-	return FEditorStyle::GetBrush("NoBrush");
-}
-
 void SetExpansionStateRecursive(UNiagaraStackEntry* StackEntry, bool bIsExpanded)
 {
 	if (StackEntry->GetCanExpand())
@@ -536,16 +438,6 @@ EVisibility SNiagaraStackTableRow::GetExecutionCategoryIconVisibility() const
 		: EVisibility::Collapsed;
 }
 
-FOptionalSize SNiagaraStackTableRow::GetIndentSize() const
-{
-	if (StackEntry->HasFrontDivider())
-	{
-		UNiagaraStackEntry* ParentEntry = GetParentEntryNoDivider(StackEntry);
-		return (StackEntry->GetIndentLevel() - (ParentEntry && ParentEntry->GetShouldShowInStack() ? 1 : 0)) * IndentSize;
-	}
-	return StackEntry->GetIndentLevel() * IndentSize;
-}
-
 EVisibility SNiagaraStackTableRow::GetExpanderVisibility() const
 {
 	if (StackEntry->GetCanExpand())
@@ -593,23 +485,9 @@ void SNiagaraStackTableRow::OnValueColumnWidthChanged(float Width)
 	ValueColumnWidthChanged.ExecuteIfBound(Width);
 }
 
-FSlateColor SNiagaraStackTableRow::GetItemBackgroundColor() const
-{
-	return StackEntry->GetIsEnabledAndOwnerIsEnabled() 
-		? ItemBackgroundColor 
-		: DisabledItemBackgroundColor;
-}
-
-const FSlateBrush* SNiagaraStackTableRow::GetSelectionBorderBrush() const
-{
-	return STableRow<UNiagaraStackEntry*>::GetBorder();
-}
-
-const FSlateBrush* SNiagaraStackTableRow::GetSearchResultBorderBrush() const
+EVisibility SNiagaraStackTableRow::GetSearchResultBorderVisibility() const
 {
-	return StackViewModel->GetCurrentFocusedEntry() == StackEntry
-		? FNiagaraEditorWidgetsStyle::Get().GetBrush("NiagaraEditor.Stack.SearchResult")
-		: FEditorStyle::GetBrush("NoBrush");
+	return StackViewModel->GetCurrentFocusedEntry() == StackEntry ? EVisibility::HitTestInvisible : EVisibility::Collapsed;
 }
 
 void SNiagaraStackTableRow::NavigateTo(UNiagaraStackEntry* Item)
diff --git a/Engine/Plugins/FX/Niagara/Source/NiagaraEditorWidgets/Private/Stack/SNiagaraStackTableRow.h b/Engine/Plugins/FX/Niagara/Source/NiagaraEditorWidgets/Private/Stack/SNiagaraStackTableRow.h
--- a/Engine/Plugins/FX/Niagara/Source/NiagaraEditorWidgets/Private/Stack/SNiagaraStackTableRow.h
+++ b/Engine/Plugins/FX/Niagara/Source/NiagaraEditorWidgets/Private/Stack/SNiagaraStackTableRow.h
@@ -23,21 +23,17 @@ public:
 public:
 	SLATE_BEGIN_ARGS(SNiagaraStackTableRow)
 		: _ContentPadding(FMargin(2, 0, 2, 0))
-		, _ItemBackgroundColor(FLinearColor::Transparent)
 		, _IsCategoryIconHighlighted(false)
 		, _ShowExecutionCategoryIcon(false)
-		, _RowPadding(FMargin(0, 0, 0, 0))
 	{}
-	SLATE_ARGUMENT(FMargin, ContentPadding)
-		SLATE_ARGUMENT(FSlateColor, ItemBackgroundColor)
-		SLATE_ARGUMENT(FSlateColor, ItemForegroundColor)
-		SLATE_ARGUMENT(FSlateColor, IndicatorColor)
+		SLATE_STYLE_ARGUMENT(FTableRowStyle, Style)
+		SLATE_ARGUMENT(FMargin, ContentPadding)
+		SLATE_ARGUMENT(TOptional<FSlateColor>, IndicatorColor)
 		SLATE_ARGUMENT(bool, IsCategoryIconHighlighted)
 		SLATE_ARGUMENT(bool, ShowExecutionCategoryIcon)
 		SLATE_ATTRIBUTE(float, NameColumnWidth)
 		SLATE_ATTRIBUTE(float, ValueColumnWidth)
 		SLATE_ATTRIBUTE(EVisibility, IssueIconVisibility)
-		SLATE_ATTRIBUTE(FMargin, RowPadding)
 		SLATE_EVENT(FOnColumnWidthChanged, OnNameColumnWidthChanged)
 		SLATE_EVENT(FOnColumnWidthChanged, OnValueColumnWidthChanged)
 		SLATE_EVENT(FOnDragDetected, OnDragDetected)
@@ -60,7 +56,7 @@ public:
 
 	void SetContentPadding(FMargin InContentPadding);
 
-	void SetNameAndValueContent(TSharedRef<SWidget> InNameWidget, TSharedPtr<SWidget> InValueWidget, TSharedPtr<SWidget> InEditConditionWidget, TSharedPtr<SWidget> InResetWidget);
+	void SetNameAndValueContent(TSharedRef<SWidget> InNameWidget, TSharedPtr<SWidget> InValueWidget);
 
 	void AddFillRowContextMenuHandler(FOnFillRowContextMenu FillRowContextMenuHandler);
 
@@ -68,8 +64,6 @@ public:
 
 	FReply OnMouseButtonUp(const FGeometry& MyGeometry, const FPointerEvent& MouseEvent);
 
-	virtual const FSlateBrush* GetBorder() const override;
-
 private:
 	void CollapseChildren();
 
@@ -79,8 +73,6 @@ private:
 
 	EVisibility GetExecutionCategoryIconVisibility() const;
 
-	FOptionalSize GetIndentSize() const;
-
 	EVisibility GetExpanderVisibility() const;
 
 	FReply ExpandButtonClicked();
@@ -91,11 +83,7 @@ private:
 
 	void OnValueColumnWidthChanged(float Width);
 
-	FSlateColor GetItemBackgroundColor() const;
-
-	const FSlateBrush* GetSelectionBorderBrush() const;
-
-	const FSlateBrush* GetSearchResultBorderBrush() const;
+	EVisibility GetSearchResultBorderVisibility() const;
 
 	void NavigateTo(UNiagaraStackEntry* Item);
 
@@ -114,15 +102,11 @@ private:
 	FOnColumnWidthChanged ValueColumnWidthChanged;
 
 	TAttribute<EVisibility> IssueIconVisibility;
-	TAttribute<FMargin> RowPadding;
 
 	const FSlateBrush* ExpandedImage;
 	const FSlateBrush* CollapsedImage;
 
-	FSlateColor ItemBackgroundColor;
-	FSlateColor DisabledItemBackgroundColor;
-	FSlateColor ForegroundColor;
-	FSlateColor IndicatorColor;
+	TOptional<FSlateColor> IndicatorColor;
 
 	FText ExecutionCategoryToolTipText;
 
diff --git a/Engine/Source/Editor/PropertyEditor/Private/DetailCategoryBuilderImpl.cpp b/Engine/Source/Editor/PropertyEditor/Private/DetailCategoryBuilderImpl.cpp
--- a/Engine/Source/Editor/PropertyEditor/Private/DetailCategoryBuilderImpl.cpp
+++ b/Engine/Source/Editor/PropertyEditor/Private/DetailCategoryBuilderImpl.cpp
@@ -746,7 +746,7 @@ bool FDetailCategoryImpl::GenerateStandaloneWidget(FDetailWidgetRow& OutRow) con
 	}
 
 	const bool bIsInnerCategory = DetailLayoutBuilder.Pin()->IsLayoutForExternalRoot();
-	FTextBlockStyle NameStyle = bIsInnerCategory ? FAppStyle::Get().GetWidgetStyle<FTextBlockStyle>("NormalText") : FAppStyle::Get().GetWidgetStyle<FTextBlockStyle>("DetailsView.CategoryTextStyleUpdate");
+	FTextBlockStyle NameStyle = bIsInnerCategory ? FAppStyle::Get().GetWidgetStyle<FTextBlockStyle>("NormalText") : FAppStyle::Get().GetWidgetStyle<FTextBlockStyle>("DetailsView.CategoryTextStyle");
 	OutRow.NameContent()
 	[
 		SNew(STextBlock)
