- Enable baker by default
- Fix upside down capture for attributes
- Change render / capture to 'Bake' for consistency
- Change Bake VP button to be more visible

#rnx

diff --git a/Engine/Plugins/FX/Niagara/Source/Niagara/Private/NiagaraDataInterfaceGrid2DCollection.cpp b/Engine/Plugins/FX/Niagara/Source/Niagara/Private/NiagaraDataInterfaceGrid2DCollection.cpp
--- a/Engine/Plugins/FX/Niagara/Source/Niagara/Private/NiagaraDataInterfaceGrid2DCollection.cpp
+++ b/Engine/Plugins/FX/Niagara/Source/Niagara/Private/NiagaraDataInterfaceGrid2DCollection.cpp
@@ -2096,7 +2096,7 @@ bool UNiagaraDataInterfaceGrid2DCollection::RenderVariableToCanvas(FNiagaraSyste
 	AttributeSlices.Z = NumFloats > 2 ? FirstSlice + 2 : INDEX_NONE;
 	AttributeSlices.W = NumFloats > 3 ? FirstSlice + 3 : INDEX_NONE;
 
-	FCanvasTileItem TileItem(FVector2D(DrawRect.Min.X, DrawRect.Min.Y), FVector2D(DrawRect.Width(), DrawRect.Height()), FLinearColor::White);
+	FCanvasTileItem TileItem(FVector2D(DrawRect.Min.X, DrawRect.Min.Y), GWhiteTexture, FVector2D(DrawRect.Width(), DrawRect.Height()), FVector2D(0.0, 1.0f), FVector2D(1.0, 0.0f), FLinearColor::White);
 	TileItem.BlendMode = SE_BLEND_Opaque;
 	TileItem.BatchedElementParameters = new FBatchedElementNiagara2DArrayAttribute(
 		AttributeSlices,
diff --git a/Engine/Plugins/FX/Niagara/Source/Niagara/Private/NiagaraDataInterfaceRenderTarget2D.cpp b/Engine/Plugins/FX/Niagara/Source/Niagara/Private/NiagaraDataInterfaceRenderTarget2D.cpp
--- a/Engine/Plugins/FX/Niagara/Source/Niagara/Private/NiagaraDataInterfaceRenderTarget2D.cpp
+++ b/Engine/Plugins/FX/Niagara/Source/Niagara/Private/NiagaraDataInterfaceRenderTarget2D.cpp
@@ -729,7 +729,7 @@ bool UNiagaraDataInterfaceRenderTarget2D::RenderVariableToCanvas(FNiagaraSystemI
 
 	Canvas->DrawTile(
 		DrawRect.Min.X, DrawRect.Min.Y, DrawRect.Width(), DrawRect.Height(),
-		0.0f, 0.0f, 1.0f, 1.0f,
+		0.0f, 1.0f, 1.0f, 0.0f,
 		FLinearColor::White,
 		GT_InstanceData->TargetTexture->GetResource(),
 		false
diff --git a/Engine/Plugins/FX/Niagara/Source/NiagaraEditor/Private/Widgets/SNiagaraBakerViewportToolbar.cpp b/Engine/Plugins/FX/Niagara/Source/NiagaraEditor/Private/Widgets/SNiagaraBakerViewportToolbar.cpp
--- a/Engine/Plugins/FX/Niagara/Source/NiagaraEditor/Private/Widgets/SNiagaraBakerViewportToolbar.cpp
+++ b/Engine/Plugins/FX/Niagara/Source/NiagaraEditor/Private/Widgets/SNiagaraBakerViewportToolbar.cpp
@@ -92,13 +92,13 @@ void SNiagaraBakerViewportToolbar::Construct(const FArguments& InArgs)
 				}
 			)
 		)
-		.ToolTipText(LOCTEXT("CaptureToolTip", "Captures the Baker."))
+		.ToolTipText(LOCTEXT("BakeToolTip", "Runs the bake process."))
 		.Content()
 		[
 			SNew(STextBlock)
 			.Font(FEditorStyle::GetFontStyle("EditorViewportToolBar.Font"))
-			.Text(LOCTEXT("Capture", "Capture"))
-			.ColorAndOpacity(FLinearColor::Black)
+			.Text(LOCTEXT("Bake", "Bake"))
+			.ColorAndOpacity(FLinearColor::White)
 		]
 	];
 	
diff --git a/Engine/Plugins/FX/Niagara/Source/NiagaraEditor/Public/NiagaraEditorSettings.h b/Engine/Plugins/FX/Niagara/Source/NiagaraEditor/Public/NiagaraEditorSettings.h
--- a/Engine/Plugins/FX/Niagara/Source/NiagaraEditor/Public/NiagaraEditorSettings.h
+++ b/Engine/Plugins/FX/Niagara/Source/NiagaraEditor/Public/NiagaraEditorSettings.h
@@ -249,7 +249,7 @@ public:
 
 	/** Enables the Niagara Baker to be used within the system editor. */
 	UPROPERTY(config, EditAnywhere, Category = Experimental)
-	bool bEnableBaker = false;
+	bool bEnableBaker = true;
 
 	TArray<float> GetPlaybackSpeeds() const;
 
