- Additional change from 19211885 due to files being moved

#rnx

diff --git a/Engine/Plugins/FX/Niagara/Source/NiagaraEditor/Private/Toolkits/SystemToolkitModes/NiagaraSystemToolkitMode_Default.cpp b/Engine/Plugins/FX/Niagara/Source/NiagaraEditor/Private/Toolkits/SystemToolkitModes/NiagaraSystemToolkitMode_Default.cpp
--- a/Engine/Plugins/FX/Niagara/Source/NiagaraEditor/Private/Toolkits/SystemToolkitModes/NiagaraSystemToolkitMode_Default.cpp
+++ b/Engine/Plugins/FX/Niagara/Source/NiagaraEditor/Private/Toolkits/SystemToolkitModes/NiagaraSystemToolkitMode_Default.cpp
@@ -107,8 +107,8 @@ void FNiagaraSystemToolkitMode_Default::ExtendToolbar()
 			);
 
 			MenuBuilder.AddMenuEntry(
-				LOCTEXT("RenderBaker", "Render Baker"),
-				LOCTEXT("RenderBakerTooltip", "Renders the Baker using the current settings."),
+				LOCTEXT("RunBaker", "Bake"),
+				LOCTEXT("RunBakerTooltip", "Runs the bake process."),
 				FSlateIcon(),
 				FUIAction(FExecuteAction::CreateSP(Toolkit, &FNiagaraSystemToolkit::RenderBaker))
 			);
