Fixed a crash when performing a right click on a niagara parameter panel category

diff --git a/Engine/Plugins/FX/Niagara/Source/NiagaraEditor/Private/ViewModels/NiagaraParameterPanelViewModel.cpp b/Engine/Plugins/FX/Niagara/Source/NiagaraEditor/Private/ViewModels/NiagaraParameterPanelViewModel.cpp
--- a/Engine/Plugins/FX/Niagara/Source/NiagaraEditor/Private/ViewModels/NiagaraParameterPanelViewModel.cpp
+++ b/Engine/Plugins/FX/Niagara/Source/NiagaraEditor/Private/ViewModels/NiagaraParameterPanelViewModel.cpp
@@ -2526,6 +2526,11 @@ FReply FNiagaraScriptToolkitParameterPanelViewModel::OnParameterItemsDragged(con
 
 TSharedPtr<SWidget> FNiagaraScriptToolkitParameterPanelViewModel::CreateContextMenuForItems(const TArray<FNiagaraParameterPanelItem>& Items, const TSharedPtr<FUICommandList>& ToolkitCommands)
 {
+	if(Items.Num() == 0)
+	{
+		return nullptr;
+	}
+	
 	// Create a menu with all relevant operations.
 	const bool bShouldCloseWindowAfterMenuSelection = true;
 	FMenuBuilder MenuBuilder(bShouldCloseWindowAfterMenuSelection, ToolkitCommands);
