Fixed deleting parameters that were in use on a map get node not deleting the default pin

We always had input pins leftover in a map get node whenever we deleted a parameter via parameter panel
We simply didn't create a widget for it, implicitly, as we only iterated over output pins and their matching input pins to create whole-row-widgets
I recently added code to add leftover pins to display orphaned pins, and this surfaced that issue

diff --git a/Engine/Plugins/FX/Niagara/Source/NiagaraEditor/Private/NiagaraNodeParameterMapGet.cpp b/Engine/Plugins/FX/Niagara/Source/NiagaraEditor/Private/NiagaraNodeParameterMapGet.cpp
--- a/Engine/Plugins/FX/Niagara/Source/NiagaraEditor/Private/NiagaraNodeParameterMapGet.cpp
+++ b/Engine/Plugins/FX/Niagara/Source/NiagaraEditor/Private/NiagaraNodeParameterMapGet.cpp
@@ -168,6 +168,15 @@ void UNiagaraNodeParameterMapGet::OnPinRenamed(UEdGraphPin* RenamedPin, const FS
 	MarkNodeRequiresSynchronization(__FUNCTION__, true);
 }
 
+void UNiagaraNodeParameterMapGet::OnPinRemoved(UEdGraphPin* InRemovedPin)
+{
+	if(UEdGraphPin* DefaultPin = GetDefaultPin(InRemovedPin))
+	{
+		RemovePin(DefaultPin);
+	}
+	
+	Super::OnPinRemoved(InRemovedPin);
+}
 
 void UNiagaraNodeParameterMapGet::OnNewTypedPinAdded(UEdGraphPin*& NewPin)
 {
diff --git a/Engine/Plugins/FX/Niagara/Source/NiagaraEditor/Private/NiagaraNodeParameterMapGet.h b/Engine/Plugins/FX/Niagara/Source/NiagaraEditor/Private/NiagaraNodeParameterMapGet.h
--- a/Engine/Plugins/FX/Niagara/Source/NiagaraEditor/Private/NiagaraNodeParameterMapGet.h
+++ b/Engine/Plugins/FX/Niagara/Source/NiagaraEditor/Private/NiagaraNodeParameterMapGet.h
@@ -51,7 +51,8 @@ public:
 protected:
 	virtual void OnNewTypedPinAdded(UEdGraphPin*& NewPin) override;
 	virtual void OnPinRenamed(UEdGraphPin* RenamedPin, const FString& OldName) override;
-
+	virtual void OnPinRemoved(UEdGraphPin* InRemovedPin) override;
+	
 	/** Synchronize the removal of the output pin with its default.*/
 	virtual void RemoveDynamicPin(UEdGraphPin* Pin) override;
 
