-Fixed crash when undo/redoing after type change by removing outdated postundo code that would try to use the wrong type utilities for the new type
The script variables customization does not need to be an undo-listener, as default values aren't stored in the script variable's variable anymore and forwarding changes to the script var's variable was the only thing this was doing.
Tested undo/redo behavior before and after the changes, seems like all changes get properly undo/redone
- Fixed new script variables not getting their defaults initialized. This resulted in new color parameters via map get node being 0,0,0,0 instead of 1,1,1,1
- Fixed type change not creating an orphaned pin for default pins that had a value differing from its default

diff --git a/Engine/Plugins/FX/Niagara/Source/NiagaraEditor/Private/Customizations/NiagaraScriptVariableCustomization.cpp b/Engine/Plugins/FX/Niagara/Source/NiagaraEditor/Private/Customizations/NiagaraScriptVariableCustomization.cpp
--- a/Engine/Plugins/FX/Niagara/Source/NiagaraEditor/Private/Customizations/NiagaraScriptVariableCustomization.cpp
+++ b/Engine/Plugins/FX/Niagara/Source/NiagaraEditor/Private/Customizations/NiagaraScriptVariableCustomization.cpp
@@ -33,7 +33,6 @@ TSharedRef<IDetailCustomization> FNiagaraScriptVariableDetails::MakeInstance()
  
 FNiagaraScriptVariableDetails::FNiagaraScriptVariableDetails()
 {
-	GEditor->RegisterForUndo(this);
 	bParameterValueChangedDuringOnValueChanged = false;
 	LibraryDefaultModeValue = 0;
 
@@ -43,7 +42,6 @@ FNiagaraScriptVariableDetails::FNiagaraScriptVariableDetails()
 
 FNiagaraScriptVariableDetails::~FNiagaraScriptVariableDetails()
 {
-	GEditor->UnregisterForUndo(this);
 }
 
 UEdGraphPin* FNiagaraScriptVariableDetails::GetAnyDefaultPin()
@@ -64,50 +62,6 @@ TArray<UEdGraphPin*> FNiagaraScriptVariableDetails::GetDefaultPins()
 	}
 	return TArray<UEdGraphPin*>();
 }
-
-void FNiagaraScriptVariableDetails::PostUndo(bool bSuccess)
-{
-	if (Variable == nullptr)
-	{
-		return;
-	}
-
-	if (Variable->GetIsStaticSwitch())
-	{
-		if (TypeUtilityStaticSwitchValue && ParameterEditorStaticSwitchValue)
-		{
-			TSharedPtr<FStructOnScope> ParameterValue = MakeShareable(new FStructOnScope(Variable->Variable.GetType().GetStruct()));
-			if (Variable->GetDefaultValueData() != nullptr)
-			{
-				Variable->Variable.SetData(Variable->GetDefaultValueData());
-			}
-			Variable->Variable.CopyTo(ParameterValue->GetStructMemory());
-			ParameterEditorStaticSwitchValue->UpdateInternalValueFromStruct(ParameterValue.ToSharedRef());
-		}
-	}
-	else if (Variable->GetOuter()->IsA<UNiagaraParameterDefinitions>())
-	{
-		if (TypeUtilityLibraryValue && ParameterEditorLibraryValue)
-		{
-			TSharedPtr<FStructOnScope> ParameterValue = MakeShareable(new FStructOnScope(Variable->Variable.GetType().GetStruct()));
-			Variable->CopyDefaultValueDataTo(ParameterValue->GetStructMemory());
-			ParameterEditorLibraryValue->UpdateInternalValueFromStruct(ParameterValue.ToSharedRef());
-		}
-	}
-	else
-	{
-		if (UEdGraphPin* Pin = GetAnyDefaultPin())
-		{
-			if (TypeUtilityValue && ParameterEditorValue)
-			{
-				TypeUtilityValue->SetValueFromPinDefaultString(Pin->DefaultValue, Variable->Variable);
-				TSharedPtr<FStructOnScope> ParameterValue = MakeShareable(new FStructOnScope(Variable->Variable.GetType().GetStruct()));
-				Variable->Variable.CopyTo(ParameterValue->GetStructMemory());
-				ParameterEditorValue->UpdateInternalValueFromStruct(ParameterValue.ToSharedRef());
-			}
-		}
-	}
-}
  
 void FNiagaraScriptVariableDetails::Refresh()
 {
diff --git a/Engine/Plugins/FX/Niagara/Source/NiagaraEditor/Private/Customizations/NiagaraScriptVariableCustomization.h b/Engine/Plugins/FX/Niagara/Source/NiagaraEditor/Private/Customizations/NiagaraScriptVariableCustomization.h
--- a/Engine/Plugins/FX/Niagara/Source/NiagaraEditor/Private/Customizations/NiagaraScriptVariableCustomization.h
+++ b/Engine/Plugins/FX/Niagara/Source/NiagaraEditor/Private/Customizations/NiagaraScriptVariableCustomization.h
@@ -42,7 +42,7 @@ enum class ENiagaraLibrarySourceDefaultMode : uint8
 };
 
 /** This customization sets up a custom details panel for the static switch Variable in the niagara module graph. */
-class FNiagaraScriptVariableDetails : public IDetailCustomization, public FEditorUndoClient
+class FNiagaraScriptVariableDetails : public IDetailCustomization
 {
 public:
  
@@ -56,11 +56,6 @@ public:
 	virtual void CustomizeDetails( IDetailLayoutBuilder& DetailBuilder ) override;
 	virtual void CustomizeDetails( const TSharedPtr<IDetailLayoutBuilder>& DetailBuilder ) override;
 
-	//~ Begin FEditorUndoClient Interface
-	virtual void PostUndo(bool bSuccess) override;
-	virtual void PostRedo(bool bSuccess) override { PostUndo(bSuccess); }
-	// End of FEditorUndoClient
-
 	// Fully regenerates the details view.
 	void Refresh();
  
diff --git a/Engine/Plugins/FX/Niagara/Source/NiagaraEditor/Private/NiagaraGraph.cpp b/Engine/Plugins/FX/Niagara/Source/NiagaraEditor/Private/NiagaraGraph.cpp
--- a/Engine/Plugins/FX/Niagara/Source/NiagaraEditor/Private/NiagaraGraph.cpp
+++ b/Engine/Plugins/FX/Niagara/Source/NiagaraEditor/Private/NiagaraGraph.cpp
@@ -474,6 +474,7 @@ void UNiagaraGraph::ChangeParameterType(const TArray<FNiagaraVariable>& Paramete
 	GetNodesOfClass<UNiagaraNodeParameterMapBase>(NiagaraParameterNodes);
 
 	TMap<UEdGraphPin*, FNiagaraTypeDefinition> OldTypesCache;
+	TMap<UEdGraphPin*, FString> OldDefaultValuesCache;
 	
 	auto TryCreateOrphanedPin = [&](UNiagaraNode* Node, FName OrphanedPinName, UEdGraphPin* OriginalPin)
 	{
@@ -494,12 +495,13 @@ void UNiagaraGraph::ChangeParameterType(const TArray<FNiagaraVariable>& Paramete
 		}
 
 		// in case we are keeping all connections, all previously existing connections are parameter pins of the same type or will be of the same type after this operation is done
-		if(PinConnectionsToKeep.Num() != OriginalPin->LinkedTo.Num())
+		if((PinConnectionsToKeep.Num() != OriginalPin->LinkedTo.Num()) || (OriginalPin->DefaultValue != OldDefaultValuesCache[OriginalPin]))
 		{
-			ensure(OldTypesCache.Contains(OriginalPin));
+			ensure(OldTypesCache.Contains(OriginalPin) && OldDefaultValuesCache.Contains(OriginalPin));
 			
 			UEdGraphPin* NewOrphanedPin = Node->CreatePin(OriginalPin->Direction, OriginalPin->PinType.PinCategory, OrphanedPinName);
 			NewOrphanedPin->PinType = UEdGraphSchema_Niagara::TypeDefinitionToPinType(OldTypesCache[OriginalPin]);
+			NewOrphanedPin->DefaultValue = OldDefaultValuesCache[OriginalPin];
 			NewOrphanedPin->PersistentGuid = FGuid::NewGuid();
 			NewOrphanedPin->bNotConnectable = true;
 			NewOrphanedPin->bOrphanedPin = true;
@@ -553,6 +555,10 @@ void UNiagaraGraph::ChangeParameterType(const TArray<FNiagaraVariable>& Paramete
 					{
 						OldTypesCache.Add(Pin, PinVariable.GetType());					
 					}
+					if(!OldDefaultValuesCache.Contains(Pin))
+					{
+						OldDefaultValuesCache.Add(Pin, Pin->DefaultValue);
+					}
 					
 					Pin->Modify();
 					Pin->PinType = NiagaraSchema->TypeDefinitionToPinType(NewType);
@@ -580,13 +586,19 @@ void UNiagaraGraph::ChangeParameterType(const TArray<FNiagaraVariable>& Paramete
 							{
 								OldTypesCache.Add(DefaultPin, UEdGraphSchema_Niagara::PinToTypeDefinition(DefaultPin));					
 							}
+							if(!OldDefaultValuesCache.Contains(DefaultPin))
+							{
+								OldDefaultValuesCache.Add(DefaultPin, DefaultPin->DefaultValue);
+							}
 							
 							DefaultPin->Modify();
 							DefaultPin->PinType = Pin->PinType;
+							DefaultPin->ResetDefaultValue();
+							UNiagaraNode::SetPinDefaultToTypeDefaultIfUnset(DefaultPin);							
 
 							if(bAllowOrphanedPins)
 							{
-								bool bCreateOrphanedPin = DefaultPin->LinkedTo.Num() > 0;
+								bool bCreateOrphanedPin = (Pin->bDefaultValueIsIgnored == false && OldDefaultValuesCache[DefaultPin] != DefaultPin->AutogeneratedDefaultValue) || Pin->LinkedTo.Num() > 0;
 
 								if(bCreateOrphanedPin)
 								{
@@ -1781,6 +1793,8 @@ UNiagaraScriptVariable* UNiagaraGraph::AddParameter(const FNiagaraVariable& Para
 		NewScriptVariable->Init(Parameter, ParameterMetaData);
 		NewScriptVariable->SetIsStaticSwitch(bIsStaticSwitch);
 
+		FNiagaraEditorUtilities::ResetVariableToDefaultValue(NewScriptVariable->Variable);
+		
 		// Inputs in graphs (module namespace) are intrinsically written to, so set the default mode to value instead of fail if not set.
 		if (NewScriptVariable->Variable.IsInNameSpace(FNiagaraConstants::ModuleNamespace))
 		{
diff --git a/Engine/Plugins/FX/Niagara/Source/NiagaraEditor/Private/NiagaraNodeParameterMapGet.cpp b/Engine/Plugins/FX/Niagara/Source/NiagaraEditor/Private/NiagaraNodeParameterMapGet.cpp
--- a/Engine/Plugins/FX/Niagara/Source/NiagaraEditor/Private/NiagaraNodeParameterMapGet.cpp
+++ b/Engine/Plugins/FX/Niagara/Source/NiagaraEditor/Private/NiagaraNodeParameterMapGet.cpp
@@ -114,6 +114,8 @@ UEdGraphPin* UNiagaraNodeParameterMapGet::CreateDefaultPin(UEdGraphPin* OutputPi
 	if (Schema->TryGetPinDefaultValueFromNiagaraVariable(Var, PinDefaultValue))
 	{
 		DefaultPin->DefaultValue = PinDefaultValue;
+		// we set the auto generated default value to the same as our default value so we can later test if we had any changes
+		DefaultPin->AutogeneratedDefaultValue = PinDefaultValue;
 	}
 
 	// If the variable of the new default pin is already in use in the graph we use the configured default value
diff --git a/Engine/Plugins/FX/Niagara/Source/NiagaraEditor/Private/ViewModels/NiagaraParameterPanelViewModel.cpp b/Engine/Plugins/FX/Niagara/Source/NiagaraEditor/Private/ViewModels/NiagaraParameterPanelViewModel.cpp
--- a/Engine/Plugins/FX/Niagara/Source/NiagaraEditor/Private/ViewModels/NiagaraParameterPanelViewModel.cpp
+++ b/Engine/Plugins/FX/Niagara/Source/NiagaraEditor/Private/ViewModels/NiagaraParameterPanelViewModel.cpp
@@ -2394,8 +2394,8 @@ void FNiagaraScriptToolkitParameterPanelViewModel::ChangeParameterType(const TAr
 		Graph->ChangeParameterType(VariablesToModify, NewType, true);
 	}
 
+	UIContext.RefreshSelectionDetailsViewPanel();
 	Refresh();
-	SelectParameterItemByName(VariablesToModify[0].GetName(), false);
 }
 
 void FNiagaraScriptToolkitParameterPanelViewModel::RenameParameter(const UNiagaraScriptVariable* ScriptVarToRename, const FName NewName) const
