Niagara - Fix non-unity, static analysis, and tobjectptr issues from the restyle change.

diff --git a/Engine/Plugins/FX/Niagara/Source/NiagaraEditor/Public/ViewModels/Stack/NiagaraStackInputCategory.h b/Engine/Plugins/FX/Niagara/Source/NiagaraEditor/Public/ViewModels/Stack/NiagaraStackInputCategory.h
--- a/Engine/Plugins/FX/Niagara/Source/NiagaraEditor/Public/ViewModels/Stack/NiagaraStackInputCategory.h
+++ b/Engine/Plugins/FX/Niagara/Source/NiagaraEditor/Public/ViewModels/Stack/NiagaraStackInputCategory.h
@@ -77,7 +77,7 @@ private:
 	TArray<FInputParameterHandleAndType> Inputs;
 
 	UPROPERTY()
-	UNiagaraStackSpacer* CategorySpacer;
+	TObjectPtr<UNiagaraStackSpacer> CategorySpacer;
 
 	bool bShouldShowInStack;
 	bool bIsTopLevelCategory;
diff --git a/Engine/Plugins/FX/Niagara/Source/NiagaraEditor/Public/ViewModels/Stack/NiagaraStackPropertyRow.h b/Engine/Plugins/FX/Niagara/Source/NiagaraEditor/Public/ViewModels/Stack/NiagaraStackPropertyRow.h
--- a/Engine/Plugins/FX/Niagara/Source/NiagaraEditor/Public/ViewModels/Stack/NiagaraStackPropertyRow.h
+++ b/Engine/Plugins/FX/Niagara/Source/NiagaraEditor/Public/ViewModels/Stack/NiagaraStackPropertyRow.h
@@ -47,7 +47,7 @@ private:
 	EStackRowStyle RowStyle;
 
 	UPROPERTY()
-	UNiagaraStackSpacer* CategorySpacer;
+	TObjectPtr<UNiagaraStackSpacer> CategorySpacer;
 
 	bool bCannotEditInThisContext;
 	bool bIsTopLevelProperty;
diff --git a/Engine/Plugins/FX/Niagara/Source/NiagaraEditorWidgets/Private/Stack/SNiagaraStackIndent.h b/Engine/Plugins/FX/Niagara/Source/NiagaraEditorWidgets/Private/Stack/SNiagaraStackIndent.h
--- a/Engine/Plugins/FX/Niagara/Source/NiagaraEditorWidgets/Private/Stack/SNiagaraStackIndent.h
+++ b/Engine/Plugins/FX/Niagara/Source/NiagaraEditorWidgets/Private/Stack/SNiagaraStackIndent.h
@@ -2,6 +2,7 @@
 
 #pragma once
 
+#include "Widgets/SCompoundWidget.h"
 #include "Widgets/DeclarativeSyntaxSupport.h"
 
 class UNiagaraStackEntry;
diff --git a/Engine/Plugins/FX/Niagara/Source/NiagaraEditorWidgets/Private/Stack/SNiagaraStackInheritanceIcon.h b/Engine/Plugins/FX/Niagara/Source/NiagaraEditorWidgets/Private/Stack/SNiagaraStackInheritanceIcon.h
--- a/Engine/Plugins/FX/Niagara/Source/NiagaraEditorWidgets/Private/Stack/SNiagaraStackInheritanceIcon.h
+++ b/Engine/Plugins/FX/Niagara/Source/NiagaraEditorWidgets/Private/Stack/SNiagaraStackInheritanceIcon.h
@@ -2,6 +2,7 @@
 
 #pragma once
 
+#include "Widgets/SCompoundWidget.h"
 #include "Widgets/DeclarativeSyntaxSupport.h"
 
 class UNiagaraStackEntry;
diff --git a/Engine/Plugins/FX/Niagara/Source/NiagaraEditorWidgets/Private/Stack/SNiagaraStackTableRow.cpp b/Engine/Plugins/FX/Niagara/Source/NiagaraEditorWidgets/Private/Stack/SNiagaraStackTableRow.cpp
--- a/Engine/Plugins/FX/Niagara/Source/NiagaraEditorWidgets/Private/Stack/SNiagaraStackTableRow.cpp
+++ b/Engine/Plugins/FX/Niagara/Source/NiagaraEditorWidgets/Private/Stack/SNiagaraStackTableRow.cpp
@@ -218,7 +218,7 @@ void SNiagaraStackTableRow::SetNameAndValueContent(TSharedRef<SWidget> InNameWid
 	}
 
 	FName AccentColorName = FNiagaraStackEditorWidgetsUtilities::GetColorNameForExecutionCategory(StackEntry->GetExecutionCategoryName());
-	TOptional<FSlateColor> AccentColor = AccentColorName != NAME_None ? FNiagaraEditorWidgetsStyle::Get().GetColor(AccentColorName) : FStyleColors::Transparent;
+	TOptional<FSlateColor> AccentColor = (AccentColorName != NAME_None) ? FNiagaraEditorWidgetsStyle::Get().GetColor(AccentColorName) : FStyleColors::Transparent;
 
 	TSharedRef<SOverlay> RowOverlay = SNew(SOverlay)
 		.Visibility(this, &SNiagaraStackTableRow::GetRowVisibility)
