- Fix for Grid2D named attributes being broken post FName changes

#rnx

diff --git a/Engine/Plugins/FX/Niagara/Source/Niagara/Classes/NiagaraDataInterfaceGrid2DCollection.h b/Engine/Plugins/FX/Niagara/Source/Niagara/Classes/NiagaraDataInterfaceGrid2DCollection.h
--- a/Engine/Plugins/FX/Niagara/Source/Niagara/Classes/NiagaraDataInterfaceGrid2DCollection.h
+++ b/Engine/Plugins/FX/Niagara/Source/Niagara/Classes/NiagaraDataInterfaceGrid2DCollection.h
@@ -142,7 +142,7 @@ private:
 	LAYOUT_FIELD(FShaderParameter, AttributeIndicesParam);
 
 	LAYOUT_FIELD(FShaderResourceParameter, SamplerParam);
-	LAYOUT_FIELD(TMemoryImageArray<FName>, AttributeNames);
+	LAYOUT_FIELD(TMemoryImageArray<FMemoryImageName>, AttributeNames);
 	LAYOUT_FIELD(TMemoryImageArray<uint32>, AttributeChannelCount);
 };
 
diff --git a/Engine/Plugins/FX/Niagara/Source/NiagaraCore/Private/NiagaraCustomVersion.cpp b/Engine/Plugins/FX/Niagara/Source/NiagaraCore/Private/NiagaraCustomVersion.cpp
--- a/Engine/Plugins/FX/Niagara/Source/NiagaraCore/Private/NiagaraCustomVersion.cpp
+++ b/Engine/Plugins/FX/Niagara/Source/NiagaraCore/Private/NiagaraCustomVersion.cpp
@@ -9,4 +9,4 @@ const FGuid FNiagaraCustomVersion::GUID(0xFCF57AFA, 0x50764283, 0xB9A9E658, 0xFF
 FCustomVersionRegistration GRegisterNiagaraCustomVersion(FNiagaraCustomVersion::GUID, FNiagaraCustomVersion::LatestVersion, TEXT("NiagaraVer"));
 
 // Note: When encountering a conflict on this file please generate a new GUID
-const FGuid FNiagaraCustomVersion::LatestScriptCompileVersion(0xBF690949, 0xE9C446E2, 0xB34FEB14, 0x3115170D);
+const FGuid FNiagaraCustomVersion::LatestScriptCompileVersion(0xB3FBD6FC, 0x7A5F6F44, 0x895A6F05, 0xDB9F56AA);
