- Fix merge error where visitor was duplicated

#rnx

diff --git a/Engine/Plugins/FX/Niagara/Source/Niagara/Private/NiagaraParameterStore.cpp b/Engine/Plugins/FX/Niagara/Source/Niagara/Private/NiagaraParameterStore.cpp
--- a/Engine/Plugins/FX/Niagara/Source/Niagara/Private/NiagaraParameterStore.cpp
+++ b/Engine/Plugins/FX/Niagara/Source/Niagara/Private/NiagaraParameterStore.cpp
@@ -296,7 +296,6 @@ void FNiagaraParameterStoreBinding::MatchParameters(FNiagaraParameterStore* Dest
 	{
 		for (const FNiagaraVariableWithOffset& ParamWithOffset : DestParamWithOffsets)
 		{
-			Visitor(ParamWithOffset, SrcStore->IndexOf(ParamWithOffset), ParamWithOffset.Offset);
 			if (const FNiagaraVariableWithOffset* Var = SrcStore->FindParameterVariable(ParamWithOffset, true))
 			{
 				Visitor(ParamWithOffset, Var->Offset, ParamWithOffset.Offset);
