Removing unnecessary ensure. StackContext can't be resolved when compiling module scripts and rather than introduce a hack around  it, just leave it be as it is harmless.

diff --git a/Engine/Plugins/FX/Niagara/Source/NiagaraEditor/Private/NiagaraHlslTranslator.cpp b/Engine/Plugins/FX/Niagara/Source/NiagaraEditor/Private/NiagaraHlslTranslator.cpp
--- a/Engine/Plugins/FX/Niagara/Source/NiagaraEditor/Private/NiagaraHlslTranslator.cpp
+++ b/Engine/Plugins/FX/Niagara/Source/NiagaraEditor/Private/NiagaraHlslTranslator.cpp
@@ -5109,7 +5109,7 @@ void FHlslNiagaraTranslator::ParameterMapSet(UNiagaraNodeParameterMapSet* SetNod
 			}
 
 			Var = ActiveHistoryForFunctionCalls.ResolveAliases(Var);
-			ensure(!Var.IsInNameSpace(FNiagaraConstants::StackContextNamespaceString));
+			
 			const FNiagaraVariable* ConstantVar = FNiagaraConstants::GetKnownConstant(Var.GetName(), false);
 			if (ConstantVar != nullptr && ConstantVar->GetType() != Var.GetType())
 			{
