Fix dragging pins from input of function to parameter map get failing to create a script variable associated with the new pin.

diff --git a/Engine/Plugins/FX/Niagara/Source/NiagaraEditor/Private/NiagaraNodeParameterMapGet.cpp b/Engine/Plugins/FX/Niagara/Source/NiagaraEditor/Private/NiagaraNodeParameterMapGet.cpp
--- a/Engine/Plugins/FX/Niagara/Source/NiagaraEditor/Private/NiagaraNodeParameterMapGet.cpp
+++ b/Engine/Plugins/FX/Niagara/Source/NiagaraEditor/Private/NiagaraNodeParameterMapGet.cpp
@@ -225,6 +225,13 @@ void UNiagaraNodeParameterMapGet::OnNewTypedPinAdded(UEdGraphPin*& NewPin)
 		return;
 	}
 
+	const UEdGraphSchema_Niagara* Schema = GetDefault<UEdGraphSchema_Niagara>();
+	constexpr bool bNeedsValue = false;
+	FNiagaraVariable PinVariable = Schema->PinToNiagaraVariable(NewPin, bNeedsValue);
+	constexpr bool bIsStaticSwitch = false;
+	UNiagaraScriptVariable* ScriptVar = GetNiagaraGraph()->AddParameter(PinVariable, bIsStaticSwitch);
+	NewPin->PinName = ScriptVar->Variable.GetName();
+
 	if (NewPin->Direction == EEdGraphPinDirection::EGPD_Output)
 	{
 		PinPendingRename = NewPin;
diff --git a/Engine/Plugins/FX/Niagara/Source/NiagaraEditor/Private/NiagaraNodeParameterMapSet.cpp b/Engine/Plugins/FX/Niagara/Source/NiagaraEditor/Private/NiagaraNodeParameterMapSet.cpp
--- a/Engine/Plugins/FX/Niagara/Source/NiagaraEditor/Private/NiagaraNodeParameterMapSet.cpp
+++ b/Engine/Plugins/FX/Niagara/Source/NiagaraEditor/Private/NiagaraNodeParameterMapSet.cpp
@@ -129,10 +129,11 @@ void UNiagaraNodeParameterMapSet::OnNewTypedPinAdded(UEdGraphPin*& NewPin)
 		if (bCreatedNamespace && GetNiagaraGraph())
 		{
 			const UEdGraphSchema_Niagara* Schema = GetDefault<UEdGraphSchema_Niagara>();
-			FNiagaraVariable PinVariable = Schema->PinToNiagaraVariable(NewPin, false);
-			const bool bIsStaticSwitch = false;
-			UNiagaraScriptVariable* Var = GetNiagaraGraph()->AddParameter(PinVariable, bIsStaticSwitch);
-			NewPin->PinName = Var->Variable.GetName();
+			constexpr bool bNeedsValue = false;
+			FNiagaraVariable PinVariable = Schema->PinToNiagaraVariable(NewPin, bNeedsValue);
+			constexpr bool bIsStaticSwitch = false;
+			UNiagaraScriptVariable* ScriptVar = GetNiagaraGraph()->AddParameter(PinVariable, bIsStaticSwitch);
+			NewPin->PinName = ScriptVar->Variable.GetName();
 		}
 	}
 
