- Add asset name to CalcBounds so we can more easily find information on attached components

#rnx

diff --git a/Engine/Plugins/FX/Niagara/Source/Niagara/Private/NiagaraComponent.cpp b/Engine/Plugins/FX/Niagara/Source/Niagara/Private/NiagaraComponent.cpp
--- a/Engine/Plugins/FX/Niagara/Source/Niagara/Private/NiagaraComponent.cpp
+++ b/Engine/Plugins/FX/Niagara/Source/Niagara/Private/NiagaraComponent.cpp
@@ -1858,10 +1858,10 @@ void UNiagaraComponent::SendRenderDynamicData_Concurrent()
 				FNiagaraSystemInstanceControllerPtr CullProxyController = CullProxy->GetSystemInstanceController();
 				if ( ensureMsgf(CullProxyController.IsValid(), TEXT("CullProxy is missing the system instance controller for System '%s' Component '%s'"), *GetNameSafe(Asset), *GetFullNameSafe(this)) )
 				{
-				CullProxyController->GetSystemInstance_Unsafe()->SetWorldTransform(SystemInstanceController->GetSystemInstance_Unsafe()->GetWorldTransform());
-				CullProxyController->GenerateSetDynamicDataCommands(SetDataCommands, *RenderData, *NiagaraProxy);
-				CullProxyController->GetSystemInstance_Unsafe()->SetWorldTransform(FTransform::Identity);
-			}
+					CullProxyController->GetSystemInstance_Unsafe()->SetWorldTransform(SystemInstanceController->GetSystemInstance_Unsafe()->GetWorldTransform());
+					CullProxyController->GenerateSetDynamicDataCommands(SetDataCommands, *RenderData, *NiagaraProxy);
+					CullProxyController->GetSystemInstance_Unsafe()->SetWorldTransform(FTransform::Identity);
+				}
 			}
 			else
 			{
@@ -1896,6 +1896,8 @@ int32 UNiagaraComponent::GetNumMaterials() const
 
 FBoxSphereBounds UNiagaraComponent::CalcBounds(const FTransform& LocalToWorld) const
 {
+	FScopeCycleCounter SystemStatCounter(Asset ? Asset->GetStatID(true, false) : TStatId());
+
 	const USceneComponent* UseAutoParent = (bAutoManageAttachment && GetAttachParent() == nullptr) ? AutoAttachParent.Get() : nullptr;
 	if (UseAutoParent)
 	{
