Restore the NiagaraShader ID to using the 8-4-4-4-12 output format
#rnx

diff --git a/Engine/Plugins/FX/Niagara/Source/NiagaraShader/Private/NiagaraShader.cpp b/Engine/Plugins/FX/Niagara/Source/NiagaraShader/Private/NiagaraShader.cpp
--- a/Engine/Plugins/FX/Niagara/Source/NiagaraShader/Private/NiagaraShader.cpp
+++ b/Engine/Plugins/FX/Niagara/Source/NiagaraShader/Private/NiagaraShader.cpp
@@ -619,7 +619,7 @@ void NiagaraShaderMapAppendKeyString(EShaderPlatform Platform, FString& KeyStrin
 /** Creates a string key for the derived data cache given a shader map id. */
 static FString GetNiagaraShaderMapKeyString(const FNiagaraShaderMapId& ShaderMapId, EShaderPlatform Platform)
 {
-	static const FString NIAGARASHADERMAP_DERIVEDDATA_VER = FDevSystemGuids::GetSystemGuid(FDevSystemGuids::NIAGARASHADERMAP_DERIVEDDATA_VER).ToString();
+	static const FString NIAGARASHADERMAP_DERIVEDDATA_VER = FDevSystemGuids::GetSystemGuid(FDevSystemGuids::NIAGARASHADERMAP_DERIVEDDATA_VER).ToString(EGuidFormats::DigitsWithHyphens);
 
 	FName Format = LegacyShaderPlatformToShaderFormat(Platform);
 	FString ShaderMapKeyString = Format.ToString() + TEXT("_") + FString(FString::FromInt(GetTargetPlatformManagerRef().ShaderFormatVersion(Format))) + TEXT("_");
