Fixed parameter action widgets treating new variables like pre-existing parameters, displaying their full names + types to the right instead of just the types

diff --git a/Engine/Plugins/FX/Niagara/Source/NiagaraEditor/Public/Widgets/SNiagaraParameterMenu.cpp b/Engine/Plugins/FX/Niagara/Source/NiagaraEditor/Public/Widgets/SNiagaraParameterMenu.cpp
--- a/Engine/Plugins/FX/Niagara/Source/NiagaraEditor/Public/Widgets/SNiagaraParameterMenu.cpp
+++ b/Engine/Plugins/FX/Niagara/Source/NiagaraEditor/Public/Widgets/SNiagaraParameterMenu.cpp
@@ -260,9 +260,6 @@ void SNiagaraAddParameterFromPanelMenu::AddMakeNewGroup(FNiagaraMenuActionCollec
 		FText FullCategory = SubCategory.IsEmpty() ? Category : FText::Format(FText::FromString("{0}|{1}"), Category, SubCategory);
 		TSharedPtr<FNiagaraMenuAction> Action(new FNiagaraMenuAction(FullCategory, DisplayName, Tooltip, 0, FText(),
 			FNiagaraMenuAction::FOnExecuteStackAction::CreateSP(this, &SNiagaraAddParameterFromPanelMenu::NewParameterSelected, TypeDefinition, InNamespaceId)));
-		FNiagaraVariable PrototypeVariable(TypeDefinition, TypeDefinition.GetFName());
-
-		Action->SetParameterVariable(PrototypeVariable);
 
 		if (TypeDefinition.IsDataInterface())
 		{
