- Add project setting for sort precision
- Allow renders to override the project setting

diff --git a/Engine/Plugins/FX/Niagara/Source/Niagara/Private/NiagaraGpuComputeDispatch.cpp b/Engine/Plugins/FX/Niagara/Source/Niagara/Private/NiagaraGpuComputeDispatch.cpp
--- a/Engine/Plugins/FX/Niagara/Source/Niagara/Private/NiagaraGpuComputeDispatch.cpp
+++ b/Engine/Plugins/FX/Niagara/Source/Niagara/Private/NiagaraGpuComputeDispatch.cpp
@@ -1927,14 +1927,15 @@ const FDistanceFieldSceneData* FNiagaraGpuComputeDispatch::GetMeshDistanceFieldP
 
 void FNiagaraGpuComputeDispatch::GenerateSortKeys(FRHICommandListImmediate& RHICmdList, int32 BatchId, int32 NumElementsInBatch, EGPUSortFlags Flags, FRHIUnorderedAccessView* KeysUAV, FRHIUnorderedAccessView* ValuesUAV)
 {
-	const FGPUSortManager::FKeyGenInfo KeyGenInfo((uint32)NumElementsInBatch, EnumHasAnyFlags(Flags, EGPUSortFlags::HighPrecisionKeys));
+	const bool bHighPrecision = EnumHasAnyFlags(Flags, EGPUSortFlags::HighPrecisionKeys);
+	const FGPUSortManager::FKeyGenInfo KeyGenInfo((uint32)NumElementsInBatch, bHighPrecision);
 
 	FNiagaraSortKeyGenCS::FPermutationDomain SortPermutationVector;
-	SortPermutationVector.Set<FNiagaraSortKeyGenCS::FSortUsingMaxPrecision>(GNiagaraGPUSortingUseMaxPrecision != 0);
+	SortPermutationVector.Set<FNiagaraSortKeyGenCS::FSortUsingMaxPrecision>(bHighPrecision);
 	SortPermutationVector.Set<FNiagaraSortKeyGenCS::FEnableCulling>(false);
 
 	FNiagaraSortKeyGenCS::FPermutationDomain SortAndCullPermutationVector;
-	SortAndCullPermutationVector.Set<FNiagaraSortKeyGenCS::FSortUsingMaxPrecision>(GNiagaraGPUSortingUseMaxPrecision != 0);
+	SortAndCullPermutationVector.Set<FNiagaraSortKeyGenCS::FSortUsingMaxPrecision>(bHighPrecision);
 	SortAndCullPermutationVector.Set<FNiagaraSortKeyGenCS::FEnableCulling>(true);
 
 	TShaderMapRef<FNiagaraSortKeyGenCS> SortKeyGenCS(GetGlobalShaderMap(FeatureLevel), SortPermutationVector);
diff --git a/Engine/Plugins/FX/Niagara/Source/Niagara/Private/NiagaraRendererMeshes.cpp b/Engine/Plugins/FX/Niagara/Source/Niagara/Private/NiagaraRendererMeshes.cpp
--- a/Engine/Plugins/FX/Niagara/Source/Niagara/Private/NiagaraRendererMeshes.cpp
+++ b/Engine/Plugins/FX/Niagara/Source/Niagara/Private/NiagaraRendererMeshes.cpp
@@ -78,6 +78,7 @@ FNiagaraRendererMeshes::FNiagaraRendererMeshes(ERHIFeatureLevel::Type FeatureLev
 	LockedAxis = Properties->LockedAxis;
 	LockedAxisSpace = Properties->LockedAxisSpace;
 	SortMode = Properties->SortMode;
+	bSortHighPrecision = UNiagaraRendererProperties::IsSortHighPrecision(Properties->SortPrecision);
 	bSortOnlyWhenTranslucent = Properties->bSortOnlyWhenTranslucent;
 	bGpuLowLatencyTranslucency = Properties->bGpuLowLatencyTranslucency;
 	bOverrideMaterials = Properties->bOverrideMaterials;
@@ -441,7 +442,7 @@ void FNiagaraRendererMeshes::InitializeSortInfo(const FParticleMeshRenderData& P
 
 	OutSortInfo.ParticleCount = ParticleMeshRenderData.SourceParticleData->GetNumInstances();
 	OutSortInfo.SortMode = SortMode;
-	OutSortInfo.SetSortFlags(GNiagaraGPUSortingUseMaxPrecision != 0, ParticleMeshRenderData.SourceParticleData->GetGPUDataReadyStage());
+	OutSortInfo.SetSortFlags(bSortHighPrecision, ParticleMeshRenderData.SourceParticleData->GetGPUDataReadyStage());
 	OutSortInfo.bEnableCulling = ParticleMeshRenderData.bNeedsCull;
 	OutSortInfo.RendererVisTagAttributeOffset = ParticleMeshRenderData.RendererVisTagOffset;
 	OutSortInfo.MeshIndexAttributeOffset = ParticleMeshRenderData.MeshIndexOffset;
diff --git a/Engine/Plugins/FX/Niagara/Source/Niagara/Private/NiagaraRendererProperties.cpp b/Engine/Plugins/FX/Niagara/Source/Niagara/Private/NiagaraRendererProperties.cpp
--- a/Engine/Plugins/FX/Niagara/Source/Niagara/Private/NiagaraRendererProperties.cpp
+++ b/Engine/Plugins/FX/Niagara/Source/Niagara/Private/NiagaraRendererProperties.cpp
@@ -458,4 +458,13 @@ bool UNiagaraRendererProperties::NeedsPreciseMotionVectors() const
 	}
 	
 	return MotionVectorSetting == ENiagaraRendererMotionVectorSetting::Precise;
-}
\ No newline at end of file
+}
+
+bool UNiagaraRendererProperties::IsSortHighPrecision(ENiagaraRendererSortPrecision SortPrecision)
+{
+	if (SortPrecision == ENiagaraRendererSortPrecision::Default)
+	{
+		return GetDefault<UNiagaraSettings>()->DefaultSortPrecision == ENiagaraDefaultSortPrecision::High;
+	}
+	return SortPrecision == ENiagaraRendererSortPrecision::High;
+}
diff --git a/Engine/Plugins/FX/Niagara/Source/Niagara/Private/NiagaraRendererSprites.cpp b/Engine/Plugins/FX/Niagara/Source/Niagara/Private/NiagaraRendererSprites.cpp
--- a/Engine/Plugins/FX/Niagara/Source/Niagara/Private/NiagaraRendererSprites.cpp
+++ b/Engine/Plugins/FX/Niagara/Source/Niagara/Private/NiagaraRendererSprites.cpp
@@ -68,6 +68,7 @@ FNiagaraRendererSprites::FNiagaraRendererSprites(ERHIFeatureLevel::Type FeatureL
 	, NumIndicesPerInstance(0)
 	, bSubImageBlend(false)
 	, bRemoveHMDRollInVR(false)
+	, bSortHighPrecision(false)
 	, bSortOnlyWhenTranslucent(true)
 	, bGpuLowLatencyTranslucency(true)
 	, bEnableDistanceCulling(false)
@@ -91,6 +92,7 @@ FNiagaraRendererSprites::FNiagaraRendererSprites(ERHIFeatureLevel::Type FeatureL
 	NumIndicesPerInstance = Properties->GetNumIndicesPerInstance();
 	bSubImageBlend = Properties->bSubImageBlend;
 	bRemoveHMDRollInVR = Properties->bRemoveHMDRollInVR;
+	bSortHighPrecision = UNiagaraRendererProperties::IsSortHighPrecision(Properties->SortPrecision);
 	bSortOnlyWhenTranslucent = Properties->bSortOnlyWhenTranslucent;
 	bGpuLowLatencyTranslucency = Properties->bGpuLowLatencyTranslucency;
 	MinFacingCameraBlendDistance = Properties->MinFacingCameraBlendDistance;
@@ -378,7 +380,7 @@ void FNiagaraRendererSprites::InitializeSortInfo(FParticleSpriteRenderData& Part
 
 	OutSortInfo.ParticleCount = ParticleSpriteRenderData.SourceParticleData->GetNumInstances();
 	OutSortInfo.SortMode = SortMode;
-	OutSortInfo.SetSortFlags(GNiagaraGPUSortingUseMaxPrecision != 0, ParticleSpriteRenderData.SourceParticleData->GetGPUDataReadyStage());
+	OutSortInfo.SetSortFlags(bSortHighPrecision, ParticleSpriteRenderData.SourceParticleData->GetGPUDataReadyStage());
 	OutSortInfo.bEnableCulling = ParticleSpriteRenderData.bNeedsCull;
 	OutSortInfo.RendererVisTagAttributeOffset = ParticleSpriteRenderData.RendererVisTagOffset;
 	OutSortInfo.RendererVisibility = RendererVisibility;
diff --git a/Engine/Plugins/FX/Niagara/Source/Niagara/Public/NiagaraMeshRendererProperties.h b/Engine/Plugins/FX/Niagara/Source/Niagara/Public/NiagaraMeshRendererProperties.h
--- a/Engine/Plugins/FX/Niagara/Source/Niagara/Public/NiagaraMeshRendererProperties.h
+++ b/Engine/Plugins/FX/Niagara/Source/Niagara/Public/NiagaraMeshRendererProperties.h
@@ -229,6 +229,10 @@ public:
 	UPROPERTY(EditAnywhere, Category = "Sorting")
 	uint32 bSortOnlyWhenTranslucent : 1;
 
+	/** Sort precision to use when sorting is active. */
+	UPROPERTY(EditAnywhere, Category = "Sorting")
+	ENiagaraRendererSortPrecision SortPrecision = ENiagaraRendererSortPrecision::Default;
+
 	/**
 	If true and a GPU emitter, we will use the current frames data to render with regardless of where the batcher may execute the dispatches.
 	If you have other emitters that are not translucent and using data that forces it to be a frame latent (i.e. view uniform buffer) you may need to disable
diff --git a/Engine/Plugins/FX/Niagara/Source/Niagara/Public/NiagaraRendererMeshes.h b/Engine/Plugins/FX/Niagara/Source/Niagara/Public/NiagaraRendererMeshes.h
--- a/Engine/Plugins/FX/Niagara/Source/Niagara/Public/NiagaraRendererMeshes.h
+++ b/Engine/Plugins/FX/Niagara/Source/Niagara/Public/NiagaraRendererMeshes.h
@@ -177,6 +177,7 @@ private:
 	ENiagaraSortMode SortMode;
 	ENiagaraMeshFacingMode FacingMode;
 	uint32 bOverrideMaterials : 1;
+	uint32 bSortHighPrecision : 1;
 	uint32 bSortOnlyWhenTranslucent : 1;
 	uint32 bGpuLowLatencyTranslucency : 1;
 	uint32 bLockedAxisEnable : 1;
diff --git a/Engine/Plugins/FX/Niagara/Source/Niagara/Public/NiagaraRendererProperties.h b/Engine/Plugins/FX/Niagara/Source/Niagara/Public/NiagaraRendererProperties.h
--- a/Engine/Plugins/FX/Niagara/Source/Niagara/Public/NiagaraRendererProperties.h
+++ b/Engine/Plugins/FX/Niagara/Source/Niagara/Public/NiagaraRendererProperties.h
@@ -152,6 +152,17 @@ private:
 	int32 TotalHalfComponents_RT;
 };
 
+UENUM()
+enum class ENiagaraRendererSortPrecision : uint8
+{
+	/** Uses the project settings value. */
+	Default,
+	/** Low precision sorting, half float (fp16) precision, faster and adequate for most cases. */
+	Low,
+	/** High precision sorting, float (fp32) precision, slower but may fix sorting artifacts. */
+	High,
+};
+
 /**
 * Emitter properties base class
 * Each EmitterRenderer derives from this with its own class, and returns it in GetProperties; a copy
@@ -257,6 +268,8 @@ public:
 
 	bool NeedsPreciseMotionVectors() const;
 
+	static bool IsSortHighPrecision(ENiagaraRendererSortPrecision SortPrecision);
+
 	/** Platforms on which this renderer is enabled. */
 	UPROPERTY(EditAnywhere, Category = "Scalability", meta=(DisplayInScalabilityContext))
 	FNiagaraPlatformSet Platforms;
diff --git a/Engine/Plugins/FX/Niagara/Source/Niagara/Public/NiagaraRendererSprites.h b/Engine/Plugins/FX/Niagara/Source/Niagara/Public/NiagaraRendererSprites.h
--- a/Engine/Plugins/FX/Niagara/Source/Niagara/Public/NiagaraRendererSprites.h
+++ b/Engine/Plugins/FX/Niagara/Source/Niagara/Public/NiagaraRendererSprites.h
@@ -114,6 +114,7 @@ private:
 
 	uint32 bSubImageBlend : 1;
 	uint32 bRemoveHMDRollInVR : 1;
+	uint32 bSortHighPrecision : 1;
 	uint32 bSortOnlyWhenTranslucent : 1;
 	uint32 bGpuLowLatencyTranslucency : 1;
 	uint32 bEnableCulling : 1;
diff --git a/Engine/Plugins/FX/Niagara/Source/Niagara/Public/NiagaraSettings.h b/Engine/Plugins/FX/Niagara/Source/Niagara/Public/NiagaraSettings.h
--- a/Engine/Plugins/FX/Niagara/Source/Niagara/Public/NiagaraSettings.h
+++ b/Engine/Plugins/FX/Niagara/Source/Niagara/Public/NiagaraSettings.h
@@ -61,6 +61,15 @@ enum class ENiagaraDefaultRendererPixelCoverageMode : uint8
 	Disabled,
 };
 
+UENUM()
+enum class ENiagaraDefaultSortPrecision : uint8
+{
+	/** Low precision sorting, half float (fp16) precision, faster and adequate for most cases. */
+	Low,
+	/** High precision sorting, float (fp32) precision, slower but may fix sorting artifacts. */
+	High,
+};
+
 UENUM()
 namespace ENDICollisionQuery_AsyncGpuTraceProvider
 {
@@ -140,6 +149,10 @@ class NIAGARA_API UNiagaraSettings : public UDeveloperSettings
 	UPROPERTY(config, EditAnywhere, Category = Renderer)
 	ENiagaraDefaultRendererPixelCoverageMode DefaultPixelCoverageMode = ENiagaraDefaultRendererPixelCoverageMode::Enabled;
 
+	/** The default setting for sorting precision when automatic is set on the Niagara Renderer. */
+	UPROPERTY(config, EditAnywhere, Category = Renderer)
+	ENiagaraDefaultSortPrecision DefaultSortPrecision = ENiagaraDefaultSortPrecision::Low;
+
 	UPROPERTY(config, EditAnywhere, Category=SkeletalMeshDI, meta = ( DisplayName = "Gpu Max Bone Influences", ToolTip = "Controls the maximum number of influences we allow the Skeletal Mesh Data Interface to use on the GPU.  Changing this setting requires restarting the editor.", ConfigRestartRequired = true))
 	TEnumAsByte<ENDISkelMesh_GpuMaxInfluences::Type> NDISkelMesh_GpuMaxInfluences;
 
diff --git a/Engine/Plugins/FX/Niagara/Source/Niagara/Public/NiagaraSpriteRendererProperties.h b/Engine/Plugins/FX/Niagara/Source/Niagara/Public/NiagaraSpriteRendererProperties.h
--- a/Engine/Plugins/FX/Niagara/Source/Niagara/Public/NiagaraSpriteRendererProperties.h
+++ b/Engine/Plugins/FX/Niagara/Source/Niagara/Public/NiagaraSpriteRendererProperties.h
@@ -193,6 +193,10 @@ public:
 	UPROPERTY(EditAnywhere, Category = "Sorting")
 	uint32 bSortOnlyWhenTranslucent : 1;
 
+	/** Sort precision to use when sorting is active. */
+	UPROPERTY(EditAnywhere, Category = "Sorting")
+	ENiagaraRendererSortPrecision SortPrecision = ENiagaraRendererSortPrecision::Default;
+
 	/**
 	If true and a GPU emitter, we will use the current frames data to render with regardless of where the batcher may execute the dispatches.
 	If you have other emitters that are not translucent and using data that forces it to be a frame latent (i.e. view uniform buffer) you may need to disable
diff --git a/Engine/Plugins/FX/Niagara/Source/NiagaraVertexFactories/Private/NiagaraSortingGPU.cpp b/Engine/Plugins/FX/Niagara/Source/NiagaraVertexFactories/Private/NiagaraSortingGPU.cpp
--- a/Engine/Plugins/FX/Niagara/Source/NiagaraVertexFactories/Private/NiagaraSortingGPU.cpp
+++ b/Engine/Plugins/FX/Niagara/Source/NiagaraVertexFactories/Private/NiagaraSortingGPU.cpp
@@ -9,14 +9,6 @@ NiagaraSortingGPU.cpp: Niagara sorting shaders
 #include "ShaderParameterUtils.h"
 #include "ShaderCompilerCore.h"
 
-int32 GNiagaraGPUSortingUseMaxPrecision = 0;
-static FAutoConsoleVariableRef CVarNiagaraGPUSortinUseMaxPrecision(
-	TEXT("Niagara.GPUSorting.UseMaxPrecision"),
-	GNiagaraGPUSortingUseMaxPrecision,
-	TEXT("Wether sorting using fp32 instead of fp16. (default=0)"),
-	ECVF_Default
-);
-
 int32 GNiagaraGPUSortingCPUToGPUThreshold = -1;
 static FAutoConsoleVariableRef CVarNiagaraGPUSortingCPUToGPUThreshold(
 	TEXT("Niagara.GPUSorting.CPUToGPUThreshold"),
diff --git a/Engine/Plugins/FX/Niagara/Source/NiagaraVertexFactories/Public/NiagaraSortingGPU.h b/Engine/Plugins/FX/Niagara/Source/NiagaraVertexFactories/Public/NiagaraSortingGPU.h
--- a/Engine/Plugins/FX/Niagara/Source/NiagaraVertexFactories/Public/NiagaraSortingGPU.h
+++ b/Engine/Plugins/FX/Niagara/Source/NiagaraVertexFactories/Public/NiagaraSortingGPU.h
@@ -13,7 +13,6 @@ NiagaraSortingGPU.h: Niagara sorting shaders
 
 struct FNiagaraGPUSortInfo;
 
-extern NIAGARAVERTEXFACTORIES_API int32 GNiagaraGPUSortingUseMaxPrecision;
 extern NIAGARAVERTEXFACTORIES_API int32 GNiagaraGPUSortingCPUToGPUThreshold;
 extern NIAGARAVERTEXFACTORIES_API int32 GNiagaraGPUCullingCPUToGPUThreshold;
 
