Fix for shader hash from niagara DI which are using 3d dispatches, and add a static assert to catch future cases

diff --git a/Engine/Plugins/FX/Niagara/Source/Niagara/Private/NiagaraDataInterface.cpp b/Engine/Plugins/FX/Niagara/Source/Niagara/Private/NiagaraDataInterface.cpp
--- a/Engine/Plugins/FX/Niagara/Source/Niagara/Private/NiagaraDataInterface.cpp
+++ b/Engine/Plugins/FX/Niagara/Source/Niagara/Private/NiagaraDataInterface.cpp
@@ -41,7 +41,7 @@ bool UNiagaraDataInterface::AppendCompileHash(FNiagaraCompileHashVisitor* InVisi
 	{
 		const FString DataInterfaceName = GetClass()->GetName();
 		InVisitor->UpdatePOD(*FString::Printf(TEXT("%s_GpuDispatchType"), *DataInterfaceName), (int32)GetGpuDispatchType());
-		InVisitor->UpdatePOD(*FString::Printf(TEXT("%s_GpuDispatchNumThreads"), *DataInterfaceName), *FString::Printf(TEXT("%dx%dx%d"), GetGpuDispatchNumThreads().X, GetGpuDispatchNumThreads().Y, GetGpuDispatchNumThreads().Z));
+		InVisitor->UpdateString(*FString::Printf(TEXT("%s_GpuDispatchNumThreads"), *DataInterfaceName), *FString::Printf(TEXT("%dx%dx%d"), GetGpuDispatchNumThreads().X, GetGpuDispatchNumThreads().Y, GetGpuDispatchNumThreads().Z));
 	}
 
 	return true;
diff --git a/Engine/Plugins/FX/Niagara/Source/Niagara/Public/NiagaraTypes.h b/Engine/Plugins/FX/Niagara/Source/Niagara/Public/NiagaraTypes.h
--- a/Engine/Plugins/FX/Niagara/Source/Niagara/Public/NiagaraTypes.h
+++ b/Engine/Plugins/FX/Niagara/Source/Niagara/Public/NiagaraTypes.h
@@ -636,6 +636,7 @@ public:
 	bool UpdatePOD(const TCHAR* InDebugName, const T& InData)
 	{
 		static_assert(TIsPODType<T>::Value, "Update does not support a constructor / destructor on the element class.");
+		static_assert(!TIsPointer<T>::Value, "Update does not support pointer types.");
 		HashState.Update((const uint8 *)&InData, sizeof(T));
 #if WITH_EDITORONLY_DATA
 		if (LogCompileIdGeneration != 0)
