Fixed the crash on startup when shaders are enabled and cooked Niagara Systems are loaded
#p4v-preflight-copy 19231383

diff --git a/Engine/Plugins/FX/Niagara/Source/Niagara/Private/NiagaraSystem.cpp b/Engine/Plugins/FX/Niagara/Source/Niagara/Private/NiagaraSystem.cpp
--- a/Engine/Plugins/FX/Niagara/Source/Niagara/Private/NiagaraSystem.cpp
+++ b/Engine/Plugins/FX/Niagara/Source/Niagara/Private/NiagaraSystem.cpp
@@ -3159,6 +3159,12 @@ bool UNiagaraSystem::RequestCompile(bool bForce, FNiagaraSystemUpdateContext* Op
 		ForceGraphToRecompileOnNextCheck();
 	}
 
+	// we can't compile systems that have been cooked without editor data
+	if (RootPackageHasAnyFlags(PKG_FilterEditorOnly))
+	{
+		return false;
+	}
+
 	if (bCompilationReentrantGuard)
 	{
 		return false;
