Null check to fix crash in camera DI.

diff --git a/Engine/Plugins/FX/Niagara/Source/Niagara/Private/NiagaraDataInterfaceCamera.cpp b/Engine/Plugins/FX/Niagara/Source/Niagara/Private/NiagaraDataInterfaceCamera.cpp
--- a/Engine/Plugins/FX/Niagara/Source/Niagara/Private/NiagaraDataInterfaceCamera.cpp
+++ b/Engine/Plugins/FX/Niagara/Source/Niagara/Private/NiagaraDataInterfaceCamera.cpp
@@ -88,7 +88,7 @@ bool UNiagaraDataInterfaceCamera::PerInstanceTick(void* PerInstanceData, FNiagar
 		for (FConstPlayerControllerIterator Iterator = World->GetPlayerControllerIterator(); Iterator; ++Iterator)
 		{
 			APlayerController* PlayerController = Iterator->Get();
-			if (i == PlayerControllerIndex && PlayerController)
+			if (i == PlayerControllerIndex && PlayerController && PlayerController->PlayerCameraManager)
 			{
 				PIData->CameraLocation = LWCConverter.ConvertWorldToSimulationPosition(PlayerController->PlayerCameraManager->GetCameraLocation());
 				PIData->CameraRotation = PlayerController->PlayerCameraManager->GetCameraRotation();
