Fix up some shader compilation warnings and clean up some DI based on some LWC updates

diff --git a/Engine/Plugins/FX/Niagara/Shaders/Private/NiagaraDataInterfaceRigidMeshCollisionQuery.ush b/Engine/Plugins/FX/Niagara/Shaders/Private/NiagaraDataInterfaceRigidMeshCollisionQuery.ush
--- a/Engine/Plugins/FX/Niagara/Shaders/Private/NiagaraDataInterfaceRigidMeshCollisionQuery.ush
+++ b/Engine/Plugins/FX/Niagara/Shaders/Private/NiagaraDataInterfaceRigidMeshCollisionQuery.ush
@@ -267,79 +267,87 @@ int DIRigidMeshCollision_GetClosestElement(in FDIRigidMeshCollisionContext DICon
 void DIRigidMeshCollision_GetElementPoint(in FDIRigidMeshCollisionContext DIContext, in float3 WorldPosition, in float DeltaTime, in float TimeFraction, in int ElementIndex, out float3 OutClosestPosition, 
 										  out float3 OutClosestNormal, out float3 OutClosestVelocity)
 {	
-	const float3 LocalPosition = DIRigidMeshCollision_GetLocalPosition(DIContext,WorldPosition,ElementIndex,TimeFraction);
-
-	float MinDistance = MAX_DISTANCE;
-	int OutputIndex = -1;
-	float3 CollisionPosition = float3(0,0,0);
-	float3 CollisionNormal = float3(0,0,0);
-
-	if( ElementIndex >= DIContext.ElementOffsets[SPHERE_INDEX] && ElementIndex < DIContext.ElementOffsets[SPHERE_INDEX+1] )
-	{
-		DIRigidMeshCollision_GetSphereProjection(LocalPosition, float3(0,0,0), DIContext.ElementExtentBuffer[ElementIndex].x, ElementIndex,
-				CollisionPosition, CollisionNormal, OutputIndex, MinDistance);
-	}
-	else if ( ElementIndex >= DIContext.ElementOffsets[CAPSULE_INDEX] && ElementIndex < DIContext.ElementOffsets[CAPSULE_INDEX+1] )
-	{
-		DIRigidMeshCollision_GetCapsuleProjection(LocalPosition, DIContext.ElementExtentBuffer[ElementIndex].xy, ElementIndex,
-				CollisionPosition, CollisionNormal, OutputIndex, MinDistance);
-	}
-	else if ( ElementIndex >= DIContext.ElementOffsets[BOX_INDEX] && ElementIndex < DIContext.ElementOffsets[BOX_INDEX+1] )
-	{
-		DIRigidMeshCollision_GetBoxProjection(LocalPosition, DIContext.ElementExtentBuffer[ElementIndex].xyz, ElementIndex,
-				CollisionPosition, CollisionNormal, OutputIndex, MinDistance);
-	}
-
 	OutClosestNormal = float3(0,0,0);
 	OutClosestPosition = float3(0,0,0);
 	OutClosestVelocity = float3(0,0,0);
 
-	if (ElementIndex != -1)
+	if (ElementIndex >= 0)
 	{
-		const float3 PreviousPosition = mul(DIRigidMeshCollision_GetPreviousTransform(DIContext,ElementIndex), float4(CollisionPosition,1.0)).xyz;
-		const float3 CurrentPosition = mul(DIRigidMeshCollision_GetCurrentTransform(DIContext,ElementIndex), float4(CollisionPosition,1.0)).xyz;
+		const uint ElementIndexUint = uint(ElementIndex);
+		const float3 LocalPosition = DIRigidMeshCollision_GetLocalPosition(DIContext,WorldPosition,ElementIndex,TimeFraction);
 
-		const float3 PreviousNormal = mul(DIRigidMeshCollision_GetPreviousTransform(DIContext,ElementIndex), float4(CollisionNormal,0.0)).xyz;
-		const float3 CurrentNormal = mul(DIRigidMeshCollision_GetCurrentTransform(DIContext,ElementIndex), float4(CollisionNormal,0.0)).xyz;
+		float MinDistance = MAX_DISTANCE;
+		int OutputIndex = -1;
+		float3 CollisionPosition = float3(0,0,0);
+		float3 CollisionNormal = float3(0,0,0);
 
-		OutClosestPosition = PreviousPosition + TimeFraction * (CurrentPosition-PreviousPosition);
-		OutClosestNormal = normalize(PreviousNormal + TimeFraction * (CurrentNormal-PreviousNormal));
+		if ( ElementIndexUint >= DIContext.ElementOffsets[SPHERE_INDEX] && ElementIndexUint < DIContext.ElementOffsets[SPHERE_INDEX+1] )
+		{
+			DIRigidMeshCollision_GetSphereProjection(LocalPosition, float3(0,0,0), DIContext.ElementExtentBuffer[ElementIndex].x, ElementIndex,
+					CollisionPosition, CollisionNormal, OutputIndex, MinDistance);
+		}
+		else if ( ElementIndexUint >= DIContext.ElementOffsets[CAPSULE_INDEX] && ElementIndexUint < DIContext.ElementOffsets[CAPSULE_INDEX+1] )
+		{
+			DIRigidMeshCollision_GetCapsuleProjection(LocalPosition, DIContext.ElementExtentBuffer[ElementIndex].xy, ElementIndex,
+					CollisionPosition, CollisionNormal, OutputIndex, MinDistance);
+		}
+		else if ( ElementIndexUint >= DIContext.ElementOffsets[BOX_INDEX] && ElementIndexUint < DIContext.ElementOffsets[BOX_INDEX+1] )
+		{
+			DIRigidMeshCollision_GetBoxProjection(LocalPosition, DIContext.ElementExtentBuffer[ElementIndex].xyz, ElementIndex,
+					CollisionPosition, CollisionNormal, OutputIndex, MinDistance);
+		}
 
-		OutClosestVelocity = ( CurrentPosition - PreviousPosition ) / DeltaTime;
+		if (ElementIndex != -1)
+		{
+			const float3 PreviousPosition = mul(DIRigidMeshCollision_GetPreviousTransform(DIContext,ElementIndex), float4(CollisionPosition,1.0)).xyz;
+			const float3 CurrentPosition = mul(DIRigidMeshCollision_GetCurrentTransform(DIContext,ElementIndex), float4(CollisionPosition,1.0)).xyz;
+
+			const float3 PreviousNormal = mul(DIRigidMeshCollision_GetPreviousTransform(DIContext,ElementIndex), float4(CollisionNormal,0.0)).xyz;
+			const float3 CurrentNormal = mul(DIRigidMeshCollision_GetCurrentTransform(DIContext,ElementIndex), float4(CollisionNormal,0.0)).xyz;
+
+			OutClosestPosition = PreviousPosition + TimeFraction * (CurrentPosition-PreviousPosition);
+			OutClosestNormal = normalize(PreviousNormal + TimeFraction * (CurrentNormal-PreviousNormal));
+
+			OutClosestVelocity = ( CurrentPosition - PreviousPosition ) / DeltaTime;
+		}
 	}
 }
 
 // Given a world space position (WorldPosition) and an element index compute the static mesh closest distance
 void DIRigidMeshCollision_GetElementDistance(in FDIRigidMeshCollisionContext DIContext, in float3 WorldPosition, in float TimeFraction, in int ElementIndex, out float OutClosestDistance)
 {	
-	const float3 LocalPosition = DIRigidMeshCollision_GetLocalPosition(DIContext,WorldPosition,ElementIndex,TimeFraction);
-
-	float MinDistance = MAX_DISTANCE;
-	int OutputIndex = -1;
-	float3 CollisionPosition = float3(0,0,0);
-	float3 CollisionNormal = float3(0,0,0);
-
-	if( ElementIndex >= DIContext.ElementOffsets[SPHERE_INDEX] && ElementIndex < DIContext.ElementOffsets[SPHERE_INDEX+1] )
-	{
-		DIRigidMeshCollision_GetSphereProjection(LocalPosition, float3(0,0,0), DIContext.ElementExtentBuffer[ElementIndex].x, ElementIndex,
-				CollisionPosition, CollisionNormal, OutputIndex, MinDistance);
-	}
-	else if ( ElementIndex >= DIContext.ElementOffsets[CAPSULE_INDEX] && ElementIndex < DIContext.ElementOffsets[CAPSULE_INDEX+1] )
-	{
-		DIRigidMeshCollision_GetCapsuleProjection(LocalPosition, DIContext.ElementExtentBuffer[ElementIndex].xy, ElementIndex,
-				CollisionPosition, CollisionNormal, OutputIndex, MinDistance);
-	}
-	else if ( ElementIndex >= DIContext.ElementOffsets[BOX_INDEX] && ElementIndex < DIContext.ElementOffsets[BOX_INDEX+1] )
-	{
-		DIRigidMeshCollision_GetBoxProjection(LocalPosition, DIContext.ElementExtentBuffer[ElementIndex].xyz, ElementIndex,
-				CollisionPosition, CollisionNormal, OutputIndex, MinDistance);
-	}
-	
 	OutClosestDistance = 0.0;
 
-	if (ElementIndex != -1)
+	if (ElementIndex >= 0)
 	{
-		OutClosestDistance = (DIContext.PhysicsTypeBuffer[ElementIndex] == 1) ? MinDistance : 0;
+		const uint ElementIndexUint = uint(ElementIndex);
+		const float3 LocalPosition = DIRigidMeshCollision_GetLocalPosition(DIContext,WorldPosition,ElementIndex,TimeFraction);
+
+		float MinDistance = MAX_DISTANCE;
+		int OutputIndex = -1;
+		float3 CollisionPosition = float3(0,0,0);
+		float3 CollisionNormal = float3(0,0,0);
+
+		if ( ElementIndexUint >= DIContext.ElementOffsets[SPHERE_INDEX] && ElementIndexUint < DIContext.ElementOffsets[SPHERE_INDEX+1] )
+		{
+			DIRigidMeshCollision_GetSphereProjection(LocalPosition, float3(0,0,0), DIContext.ElementExtentBuffer[ElementIndex].x, ElementIndex,
+					CollisionPosition, CollisionNormal, OutputIndex, MinDistance);
+		}
+		else if ( ElementIndexUint >= DIContext.ElementOffsets[CAPSULE_INDEX] && ElementIndexUint < DIContext.ElementOffsets[CAPSULE_INDEX+1] )
+		{
+			DIRigidMeshCollision_GetCapsuleProjection(LocalPosition, DIContext.ElementExtentBuffer[ElementIndex].xy, ElementIndex,
+					CollisionPosition, CollisionNormal, OutputIndex, MinDistance);
+		}
+		else if ( ElementIndexUint >= DIContext.ElementOffsets[BOX_INDEX] && ElementIndexUint < DIContext.ElementOffsets[BOX_INDEX+1] )
+		{
+			DIRigidMeshCollision_GetBoxProjection(LocalPosition, DIContext.ElementExtentBuffer[ElementIndex].xyz, ElementIndex,
+					CollisionPosition, CollisionNormal, OutputIndex, MinDistance);
+		}
+	
+		if (ElementIndex != -1)
+		{
+			OutClosestDistance = (DIContext.PhysicsTypeBuffer[ElementIndex] == 1) ? MinDistance : 0;
+		}
 	}
 }
 
@@ -414,7 +422,7 @@ void DIRigidMeshCollision_GetClosestPointMeshDistanceField(in FDIRigidMeshCollis
 
 	if (ElementIndex != -1)
 	{
-		int DFIndex = DIContext.DFIndexBuffer[ElementIndex];
+		uint DFIndex = DIContext.DFIndexBuffer[ElementIndex];
 
 		if (DFIndex >= 0 && DFIndex < NumSceneObjects)
 		{
@@ -422,14 +430,15 @@ void DIRigidMeshCollision_GetClosestPointMeshDistanceField(in FDIRigidMeshCollis
 			OutClosestDistance = DistanceToNearestSurfaceForObject(DFIndex, WorldPosition, 1e-3);
 
 			FDFObjectData DFObjectData = LoadDFObjectData(DFIndex);
+			float4x4 WorldToVolume = LWCHackToFloat(DFObjectData.WorldToVolume);
 			uint NumMips = LoadDFAssetData(DFObjectData.AssetIndex, 0).NumMips;
 			FDFAssetData DFAssetData = LoadDFAssetData(DFObjectData.AssetIndex, NumMips - 1);
-			float3 VolumePosition = mul(float4(WorldPosition, 1), DFObjectData.WorldToVolume).xyz;		
+			float3 VolumePosition = mul(float4(WorldPosition, 1), WorldToVolume).xyz;		
 			float3 VolumeGradient = CalculateMeshSDFGradient(VolumePosition, DFAssetData);
 			float VolumeGradientLength = length(VolumeGradient);
 			float3 VolumeNormal = VolumeGradientLength > 0.0f ? VolumeGradient / VolumeGradientLength : 0;
 			// Transform by transposed inverse to handle non-uniform scaling
-			float3 WorldGradient = mul(VolumeNormal, transpose((float3x3)DFObjectData.WorldToVolume.M));
+			float3 WorldGradient = mul(VolumeNormal, transpose((float3x3)WorldToVolume));
 			float WorldGradientLength = length(WorldGradient);
 			OutClosestNormal = WorldGradientLength > 0.0f ? WorldGradient / WorldGradientLength : 0;
 
diff --git a/Engine/Plugins/Runtime/HairStrands/Shaders/Private/NiagaraDataInterfacePhysicsAsset.ush b/Engine/Plugins/Runtime/HairStrands/Shaders/Private/NiagaraDataInterfacePhysicsAsset.ush
--- a/Engine/Plugins/Runtime/HairStrands/Shaders/Private/NiagaraDataInterfacePhysicsAsset.ush
+++ b/Engine/Plugins/Runtime/HairStrands/Shaders/Private/NiagaraDataInterfacePhysicsAsset.ush
@@ -315,79 +315,87 @@ int DIPhysicsAsset_GetClosestElement(in FDIPhysicsAssetContext DIContext, in flo
 void DIPhysicsAsset_GetElementPoint(in FDIPhysicsAssetContext DIContext, in float3 NodePosition, in float DeltaTime, in float TimeFraction, in int ElementIndex, out float3 OutClosestPosition, 
 										  out float3 OutClosestNormal, out float3 OutClosestVelocity)
 {	
-	const float3 LocalPosition = DIPhysicsAsset_GetLocalPosition(DIContext,NodePosition,ElementIndex,TimeFraction);
-
-	float MinDistance = MAX_DISTANCE;
-	int OutputIndex = -1;
-	float3 CollisionPosition = float3(0,0,0);
-	float3 CollisionNormal = float3(0,0,0);
-
-	if( ElementIndex >= DIContext.ElementOffsets[SPHERE_INDEX] && ElementIndex < DIContext.ElementOffsets[SPHERE_INDEX+1] )
-	{
-		DIPhysicsAsset_GetSphereProjection(LocalPosition, float3(0,0,0), DIContext.ElementExtentBuffer[ElementIndex].x, ElementIndex,
-				CollisionPosition, CollisionNormal, OutputIndex, MinDistance);
-	}
-	else if ( ElementIndex >= DIContext.ElementOffsets[CAPSULE_INDEX] && ElementIndex < DIContext.ElementOffsets[CAPSULE_INDEX+1] )
-	{
-		DIPhysicsAsset_GetCapsuleProjection(LocalPosition, DIContext.ElementExtentBuffer[ElementIndex].xy, ElementIndex,
-				CollisionPosition, CollisionNormal, OutputIndex, MinDistance);
-	}
-	else if ( ElementIndex >= DIContext.ElementOffsets[BOX_INDEX] && ElementIndex < DIContext.ElementOffsets[BOX_INDEX+1] )
-	{
-		DIPhysicsAsset_GetBoxProjection(LocalPosition, DIContext.ElementExtentBuffer[ElementIndex].xyz, ElementIndex,
-				CollisionPosition, CollisionNormal, OutputIndex, MinDistance);
-	}
-
 	OutClosestNormal = float3(0,0,0);
 	OutClosestPosition = float3(0,0,0);
 	OutClosestVelocity = float3(0,0,0);
 
-	if (ElementIndex != -1)
+	if (ElementIndex >= 0)
 	{
-		const float3 PreviousPosition = mul(DIPhysicsAsset_GetPreviousTransform(DIContext,ElementIndex), float4(CollisionPosition,1.0)).xyz;
-		const float3 CurrentPosition = mul(DIPhysicsAsset_GetCurrentTransform(DIContext,ElementIndex), float4(CollisionPosition,1.0)).xyz;
+		const uint ElementIndexUint = uint(ElementIndex);
+		const float3 LocalPosition = DIPhysicsAsset_GetLocalPosition(DIContext,NodePosition,ElementIndex,TimeFraction);
 
-		const float3 PreviousNormal = mul(DIPhysicsAsset_GetPreviousTransform(DIContext,ElementIndex), float4(CollisionNormal,0.0)).xyz;
-		const float3 CurrentNormal = mul(DIPhysicsAsset_GetCurrentTransform(DIContext,ElementIndex), float4(CollisionNormal,0.0)).xyz;
+		float MinDistance = MAX_DISTANCE;
+		int OutputIndex = -1;
+		float3 CollisionPosition = float3(0,0,0);
+		float3 CollisionNormal = float3(0,0,0);
 
-		OutClosestPosition = PreviousPosition + TimeFraction * (CurrentPosition-PreviousPosition);
-		OutClosestNormal = normalize(PreviousNormal + TimeFraction * (CurrentNormal-PreviousNormal));
+		if ( ElementIndexUint >= DIContext.ElementOffsets[SPHERE_INDEX] && ElementIndexUint < DIContext.ElementOffsets[SPHERE_INDEX+1] )
+		{
+			DIPhysicsAsset_GetSphereProjection(LocalPosition, float3(0,0,0), DIContext.ElementExtentBuffer[ElementIndex].x, ElementIndex,
+					CollisionPosition, CollisionNormal, OutputIndex, MinDistance);
+		}
+		else if ( ElementIndexUint >= DIContext.ElementOffsets[CAPSULE_INDEX] && ElementIndexUint < DIContext.ElementOffsets[CAPSULE_INDEX+1] )
+		{
+			DIPhysicsAsset_GetCapsuleProjection(LocalPosition, DIContext.ElementExtentBuffer[ElementIndex].xy, ElementIndex,
+					CollisionPosition, CollisionNormal, OutputIndex, MinDistance);
+		}
+		else if ( ElementIndexUint >= DIContext.ElementOffsets[BOX_INDEX] && ElementIndexUint < DIContext.ElementOffsets[BOX_INDEX+1] )
+		{
+			DIPhysicsAsset_GetBoxProjection(LocalPosition, DIContext.ElementExtentBuffer[ElementIndex].xyz, ElementIndex,
+					CollisionPosition, CollisionNormal, OutputIndex, MinDistance);
+		}
 
-		OutClosestVelocity = ( CurrentPosition - PreviousPosition ) / DeltaTime;
+		if (ElementIndex != -1)
+		{
+			const float3 PreviousPosition = mul(DIPhysicsAsset_GetPreviousTransform(DIContext,ElementIndex), float4(CollisionPosition,1.0)).xyz;
+			const float3 CurrentPosition = mul(DIPhysicsAsset_GetCurrentTransform(DIContext,ElementIndex), float4(CollisionPosition,1.0)).xyz;
+
+			const float3 PreviousNormal = mul(DIPhysicsAsset_GetPreviousTransform(DIContext,ElementIndex), float4(CollisionNormal,0.0)).xyz;
+			const float3 CurrentNormal = mul(DIPhysicsAsset_GetCurrentTransform(DIContext,ElementIndex), float4(CollisionNormal,0.0)).xyz;
+
+			OutClosestPosition = PreviousPosition + TimeFraction * (CurrentPosition-PreviousPosition);
+			OutClosestNormal = normalize(PreviousNormal + TimeFraction * (CurrentNormal-PreviousNormal));
+
+			OutClosestVelocity = ( CurrentPosition - PreviousPosition ) / DeltaTime;
+		}
 	}
 }
 
 // Given a world space position (NodePosition) and an element index compute the physics asset closest distance
 void DIPhysicsAsset_GetElementDistance(in FDIPhysicsAssetContext DIContext, in float3 NodePosition, in float TimeFraction, in int ElementIndex, out float OutClosestDistance)
 {	
-	const float3 LocalPosition = DIPhysicsAsset_GetLocalPosition(DIContext,NodePosition,ElementIndex,TimeFraction);
-
-	float MinDistance = MAX_DISTANCE;
-	int OutputIndex = -1;
-	float3 CollisionPosition = float3(0,0,0);
-	float3 CollisionNormal = float3(0,0,0);
-
-	if( ElementIndex >= DIContext.ElementOffsets[SPHERE_INDEX] && ElementIndex < DIContext.ElementOffsets[SPHERE_INDEX+1] )
-	{
-		DIPhysicsAsset_GetSphereProjection(LocalPosition, float3(0,0,0), DIContext.ElementExtentBuffer[ElementIndex].x, ElementIndex,
-				CollisionPosition, CollisionNormal, OutputIndex, MinDistance);
-	}
-	else if ( ElementIndex >= DIContext.ElementOffsets[CAPSULE_INDEX] && ElementIndex < DIContext.ElementOffsets[CAPSULE_INDEX+1] )
-	{
-		DIPhysicsAsset_GetCapsuleProjection(LocalPosition, DIContext.ElementExtentBuffer[ElementIndex].xy, ElementIndex,
-				CollisionPosition, CollisionNormal, OutputIndex, MinDistance);
-	}
-	else if ( ElementIndex >= DIContext.ElementOffsets[BOX_INDEX] && ElementIndex < DIContext.ElementOffsets[BOX_INDEX+1] )
-	{
-		DIPhysicsAsset_GetBoxProjection(LocalPosition, DIContext.ElementExtentBuffer[ElementIndex].xyz, ElementIndex,
-				CollisionPosition, CollisionNormal, OutputIndex, MinDistance);
-	}
-	
 	OutClosestDistance = 0.0;
 
-	if (ElementIndex != -1)
+	if (ElementIndex >= 0)
 	{
-		OutClosestDistance = (DIContext.PhysicsTypeBuffer[ElementIndex] == 1) ? MinDistance : 0;
+		const uint ElementIndexUint = uint(ElementIndex);
+		const float3 LocalPosition = DIPhysicsAsset_GetLocalPosition(DIContext,NodePosition,ElementIndex,TimeFraction);
+
+		float MinDistance = MAX_DISTANCE;
+		int OutputIndex = -1;
+		float3 CollisionPosition = float3(0,0,0);
+		float3 CollisionNormal = float3(0,0,0);
+
+		if( ElementIndexUint >= DIContext.ElementOffsets[SPHERE_INDEX] && ElementIndexUint < DIContext.ElementOffsets[SPHERE_INDEX+1] )
+		{
+			DIPhysicsAsset_GetSphereProjection(LocalPosition, float3(0,0,0), DIContext.ElementExtentBuffer[ElementIndex].x, ElementIndex,
+					CollisionPosition, CollisionNormal, OutputIndex, MinDistance);
+		}
+		else if ( ElementIndexUint >= DIContext.ElementOffsets[CAPSULE_INDEX] && ElementIndexUint < DIContext.ElementOffsets[CAPSULE_INDEX+1] )
+		{
+			DIPhysicsAsset_GetCapsuleProjection(LocalPosition, DIContext.ElementExtentBuffer[ElementIndex].xy, ElementIndex,
+					CollisionPosition, CollisionNormal, OutputIndex, MinDistance);
+		}
+		else if ( ElementIndexUint >= DIContext.ElementOffsets[BOX_INDEX] && ElementIndexUint < DIContext.ElementOffsets[BOX_INDEX+1] )
+		{
+			DIPhysicsAsset_GetBoxProjection(LocalPosition, DIContext.ElementExtentBuffer[ElementIndex].xyz, ElementIndex,
+					CollisionPosition, CollisionNormal, OutputIndex, MinDistance);
+		}
+
+		if (ElementIndex != -1)
+		{
+			OutClosestDistance = (DIContext.PhysicsTypeBuffer[ElementIndex] == 1) ? MinDistance : 0;
+		}
 	}
 }
 
