- Remove baker enabled settings as it's on by default

#rnx

diff --git a/Engine/Plugins/FX/Niagara/Source/NiagaraEditor/Private/Toolkits/SystemToolkitModes/NiagaraSystemToolkitModeBase.cpp b/Engine/Plugins/FX/Niagara/Source/NiagaraEditor/Private/Toolkits/SystemToolkitModes/NiagaraSystemToolkitModeBase.cpp
--- a/Engine/Plugins/FX/Niagara/Source/NiagaraEditor/Private/Toolkits/SystemToolkitModes/NiagaraSystemToolkitModeBase.cpp
+++ b/Engine/Plugins/FX/Niagara/Source/NiagaraEditor/Private/Toolkits/SystemToolkitModes/NiagaraSystemToolkitModeBase.cpp
@@ -133,12 +133,9 @@ void FNiagaraSystemToolkitModeBase::RegisterTabFactories(TSharedPtr<FTabManager>
 		.SetGroup(WorkspaceMenuCategory.ToSharedRef())
 		.SetIcon(FSlateIcon(FNiagaraEditorStyle::Get().GetStyleSetName(), "Tab.ScriptStats"));
 
-	if (GetDefault<UNiagaraEditorSettings>()->bEnableBaker)
-	{
-		InTabManager->RegisterTabSpawner(BakerTabID, FOnSpawnTab::CreateSP(this, &FNiagaraSystemToolkitModeBase::SpawnTab_Baker))
-			.SetDisplayName(LOCTEXT("NiagaraBakerTab", "Baker"))
-			.SetGroup(WorkspaceMenuCategory.ToSharedRef());
-	}
+	InTabManager->RegisterTabSpawner(BakerTabID, FOnSpawnTab::CreateSP(this, &FNiagaraSystemToolkitModeBase::SpawnTab_Baker))
+		.SetDisplayName(LOCTEXT("NiagaraBakerTab", "Baker"))
+		.SetGroup(WorkspaceMenuCategory.ToSharedRef());
 }
 
 class SNiagaraSelectedEmitterGraph : public SCompoundWidget
diff --git a/Engine/Plugins/FX/Niagara/Source/NiagaraEditor/Private/Toolkits/SystemToolkitModes/NiagaraSystemToolkitMode_Default.cpp b/Engine/Plugins/FX/Niagara/Source/NiagaraEditor/Private/Toolkits/SystemToolkitModes/NiagaraSystemToolkitMode_Default.cpp
--- a/Engine/Plugins/FX/Niagara/Source/NiagaraEditor/Private/Toolkits/SystemToolkitModes/NiagaraSystemToolkitMode_Default.cpp
+++ b/Engine/Plugins/FX/Niagara/Source/NiagaraEditor/Private/Toolkits/SystemToolkitModes/NiagaraSystemToolkitMode_Default.cpp
@@ -226,20 +226,17 @@ void FNiagaraSystemToolkitMode_Default::ExtendToolbar()
 			}
 			ToolbarBuilder.EndSection();
 
-			if ( GetDefault<UNiagaraEditorSettings>()->bEnableBaker )
+			ToolbarBuilder.BeginSection("Baker");
 			{
-				ToolbarBuilder.BeginSection("Baker");
-				{
-					ToolbarBuilder.AddComboButton(
-						FUIAction(),
-						FOnGetContent::CreateStatic(Local::GenerateBakerMenu, Toolkit),
-						LOCTEXT("Baker", "Baker"),
-						LOCTEXT("BakerTooltip", "Options for Baker rendering."),
-						FSlateIcon(FNiagaraEditorStyle::Get().GetStyleSetName(), "NiagaraEditor.Baker")
-					);
-				}
-				ToolbarBuilder.EndSection();
+				ToolbarBuilder.AddComboButton(
+					FUIAction(),
+					FOnGetContent::CreateStatic(Local::GenerateBakerMenu, Toolkit),
+					LOCTEXT("Baker", "Baker"),
+					LOCTEXT("BakerTooltip", "Options for Baker rendering."),
+					FSlateIcon(FNiagaraEditorStyle::Get().GetStyleSetName(), "NiagaraEditor.Baker")
+				);
 			}
+			ToolbarBuilder.EndSection();
 
 			if(GNiagaraScalabilityModeEnabled)
 			{
diff --git a/Engine/Plugins/FX/Niagara/Source/NiagaraEditor/Public/NiagaraEditorSettings.h b/Engine/Plugins/FX/Niagara/Source/NiagaraEditor/Public/NiagaraEditorSettings.h
--- a/Engine/Plugins/FX/Niagara/Source/NiagaraEditor/Public/NiagaraEditorSettings.h
+++ b/Engine/Plugins/FX/Niagara/Source/NiagaraEditor/Public/NiagaraEditorSettings.h
@@ -247,10 +247,6 @@ public:
 	UPROPERTY(config, EditAnywhere, Category = Niagara)
 	TArray<FNiagaraSpawnShortcut> GraphCreationShortcuts;
 
-	/** Enables the Niagara Baker to be used within the system editor. */
-	UPROPERTY(config, EditAnywhere, Category = Experimental)
-	bool bEnableBaker = true;
-
 	TArray<float> GetPlaybackSpeeds() const;
 
 	/** Gets whether or not auto-compile is enabled in the editors. */
