- Disable pixel coverage in ortho mode

#rnx

diff --git a/Engine/Plugins/FX/Niagara/Source/Niagara/Private/NiagaraRendererSprites.cpp b/Engine/Plugins/FX/Niagara/Source/Niagara/Private/NiagaraRendererSprites.cpp
--- a/Engine/Plugins/FX/Niagara/Source/Niagara/Private/NiagaraRendererSprites.cpp
+++ b/Engine/Plugins/FX/Niagara/Source/Niagara/Private/NiagaraRendererSprites.cpp
@@ -556,9 +556,10 @@ FNiagaraSpriteUniformBufferRef FNiagaraRendererSprites::CreateViewUniformBuffer(
 	PerViewUniformParameters.PrevPivotOffsetDataOffset = INDEX_NONE;
 
 	// Determine pixel coverage settings
-	PerViewUniformParameters.PixelCoverageEnabled = PixelCoverageMode != ENiagaraRendererPixelCoverageMode::Disabled;
+	const bool PixelCoverageEnabled = View.IsPerspectiveProjection() && (PixelCoverageMode != ENiagaraRendererPixelCoverageMode::Disabled);
+	PerViewUniformParameters.PixelCoverageEnabled = PixelCoverageEnabled;
 	PerViewUniformParameters.PixelCoverageColorBlend = FVector4f::Zero();
-	if (PixelCoverageMode != ENiagaraRendererPixelCoverageMode::Disabled)
+	if (PixelCoverageEnabled)
 	{
 		if ( PixelCoverageMode == ENiagaraRendererPixelCoverageMode::Automatic )
 		{
