- Add different options for emitter bounds calculation mode (dynamic / static / script) to bring mode clarity and allow programmable bounds to avoid emitting a stack warning

#rnx

diff --git a/Engine/Plugins/FX/Niagara/Source/Niagara/Classes/NiagaraEmitter.h b/Engine/Plugins/FX/Niagara/Source/Niagara/Classes/NiagaraEmitter.h
--- a/Engine/Plugins/FX/Niagara/Source/Niagara/Classes/NiagaraEmitter.h
+++ b/Engine/Plugins/FX/Niagara/Source/Niagara/Classes/NiagaraEmitter.h
@@ -113,6 +113,17 @@ enum class EParticleAllocationMode : uint8
 	FixedCount
 };
 
+UENUM()
+enum class ENiagaraEmitterCalculateBoundMode : uint8
+{
+	/** Bounds are calculated per frame (Only available for CPU emitters). */
+	Dynamic,
+	/** Bounds are set from the emitter's fixed bounds. */
+	Fixed,
+	/** Bounds will be set from the script or blueprint.  If not set from either source the emitter has no bounds. */
+	Programmable
+};
+
 USTRUCT()
 struct FNiagaraEmitterScriptProperties
 {
@@ -337,8 +348,12 @@ public:
 	UPROPERTY(EditAnywhere, Category = "Emitter")
 	ENiagaraSimTarget SimTarget;
 	
+	/** How should we calculate bounds for the emitter. */
+	UPROPERTY(EditAnywhere, Category = "Emitter")
+	ENiagaraEmitterCalculateBoundMode CalculateBoundsMode = ENiagaraEmitterCalculateBoundMode::Dynamic;
+
 	/** The fixed bounding box value. bFixedBounds is the condition whether the fixed bounds can be edited. */
-	UPROPERTY(EditAnywhere, Category = "Emitter", meta = (EditCondition = "bFixedBounds"))
+	UPROPERTY(EditAnywhere, Category = "Emitter", meta = (EditConditionHides, EditCondition = "CalculateBoundsMode == ENiagaraEmitterCalculateBoundMode::Fixed"))
 	FBox FixedBounds;
 	
 	UPROPERTY()
@@ -358,9 +373,10 @@ public:
 	UPROPERTY(EditAnywhere, Category = "Emitter")
 	uint32 bInterpolatedSpawning : 1;
 
-	/** Whether or not fixed bounds are enabled. */
-	UPROPERTY(EditAnywhere, Category = "Emitter", meta = (InlineEditConditionToggle))
-	uint32 bFixedBounds : 1;
+#if WITH_EDITORONLY_DATA
+	UPROPERTY()
+	uint32 bFixedBounds_DEPRECATED : 1;
+#endif
 
 	/** Whether to use the min detail or not. */
 	UPROPERTY()
@@ -595,6 +611,7 @@ public:
 
 	bool RequiresViewUniformBuffer() const { return bRequiresViewUniformBuffer; }
 
+	FORCEINLINE static FBox GetDefaultFixedBounds() { return FBox(FVector(-100), FVector(100)); }
 	TConstArrayView<TUniquePtr<FNiagaraBoundsCalculator>> GetBoundsCalculators() const { return MakeArrayView(BoundsCalculators); }
 
 	UPROPERTY()
diff --git a/Engine/Plugins/FX/Niagara/Source/Niagara/Private/NiagaraEmitter.cpp b/Engine/Plugins/FX/Niagara/Source/Niagara/Private/NiagaraEmitter.cpp
--- a/Engine/Plugins/FX/Niagara/Source/Niagara/Private/NiagaraEmitter.cpp
+++ b/Engine/Plugins/FX/Niagara/Source/Niagara/Private/NiagaraEmitter.cpp
@@ -146,11 +146,10 @@ bool FNiagaraEmitterScriptProperties::DataSetAccessSynchronized() const
 UNiagaraEmitter::UNiagaraEmitter(const FObjectInitializer& Initializer)
 : Super(Initializer)
 , PreAllocationCount(0)
-, FixedBounds(FBox(FVector(-100), FVector(100)))
+, FixedBounds(GetDefaultFixedBounds())
 , MinDetailLevel_DEPRECATED(0)
 , MaxDetailLevel_DEPRECATED(4)
 , bInterpolatedSpawning(false)
-, bFixedBounds(false)
 , bUseMinDetailLevel_DEPRECATED(false)
 , bUseMaxDetailLevel_DEPRECATED(false)
 , bRequiresPersistentIDs(false)
@@ -453,6 +452,10 @@ void UNiagaraEmitter::PostLoad()
 		LibraryVisibility = ENiagaraScriptLibraryVisibility::Library;
 	}
 	
+	if (bFixedBounds_DEPRECATED)
+	{
+		CalculateBoundsMode = ENiagaraEmitterCalculateBoundMode::Fixed;
+	}
 #endif
 
 	for (int32 RendererIndex = RendererProperties.Num() - 1; RendererIndex >= 0; --RendererIndex)
@@ -635,7 +638,7 @@ void UNiagaraEmitter::GetAssetRegistryTags(TArray<FAssetRegistryTag>& OutTags) c
 	OutTags.Add(FAssetRegistryTag("HasGPUEmitter", (SimTarget == ENiagaraSimTarget::GPUComputeSim) ? TEXT("True") : TEXT("False"), FAssetRegistryTag::TT_Alphabetical));
 
 	const float BoundsSize = FixedBounds.GetSize().GetMax();
-	OutTags.Add(FAssetRegistryTag("FixedBoundsSize", bFixedBounds ? FString::Printf(TEXT("%.2f"), BoundsSize) : FString(TEXT("None")), FAssetRegistryTag::TT_Numerical));
+	OutTags.Add(FAssetRegistryTag("FixedBoundsSize", CalculateBoundsMode == ENiagaraEmitterCalculateBoundMode::Fixed ? FString::Printf(TEXT("%.2f"), BoundsSize) : FString(TEXT("None")), FAssetRegistryTag::TT_Numerical));
 
 
 	uint32 NumActiveRenderers = 0;
@@ -1207,7 +1210,7 @@ void UNiagaraEmitter::CacheFromCompiledData(const FNiagaraDataSetCompiledData* C
 	// Initialize bounds calculators - skip creating if we won't ever use it.  We leave the GPU sims in there with the editor so that we can
 	// generate the bounds from the readback in the tool.
 #if !WITH_EDITOR
-	bool bUseDynamicBounds = !bFixedBounds && SimTarget == ENiagaraSimTarget::CPUSim;
+	const bool bUseDynamicBounds = CalculateBoundsMode == ENiagaraEmitterCalculateBoundMode::Dynamic && SimTarget == ENiagaraSimTarget::CPUSim;
 	if (bUseDynamicBounds)
 #endif
 	{
diff --git a/Engine/Plugins/FX/Niagara/Source/Niagara/Private/NiagaraEmitterInstance.cpp b/Engine/Plugins/FX/Niagara/Source/Niagara/Private/NiagaraEmitterInstance.cpp
--- a/Engine/Plugins/FX/Niagara/Source/Niagara/Private/NiagaraEmitterInstance.cpp
+++ b/Engine/Plugins/FX/Niagara/Source/Niagara/Private/NiagaraEmitterInstance.cpp
@@ -988,7 +988,7 @@ void FNiagaraEmitterInstance::CalculateFixedBounds(const FTransform& ToWorldSpac
 		return;
 
 	CachedEmitter->Modify();
-	CachedEmitter->bFixedBounds = true;
+	CachedEmitter->CalculateBoundsMode = ENiagaraEmitterCalculateBoundMode::Fixed;
 	if (CachedEmitter->bLocalSpace)
 	{
 		CachedEmitter->FixedBounds = Bounds;
@@ -1034,28 +1034,33 @@ void FNiagaraEmitterInstance::PostTick()
 		{
 			CachedBounds = CachedSystemFixedBounds;
 		}
-		else if (CachedEmitter->bFixedBounds || CachedEmitter->SimTarget == ENiagaraSimTarget::GPUComputeSim)
+		else if ( CachedEmitter->CalculateBoundsMode == ENiagaraEmitterCalculateBoundMode::Fixed )
 		{
 			CachedBounds = CachedEmitter->FixedBounds;
 		}
-		else
+		else if ( CachedEmitter->CalculateBoundsMode == ENiagaraEmitterCalculateBoundMode::Dynamic )
 		{
-			FBox DynamicBounds = InternalCalculateDynamicBounds(ParticleDataSet->GetCurrentDataChecked().GetNumInstances());
-			if (DynamicBounds.IsValid)
+			// We do not support dynamic bounds for GPU yet
+			// Therefore we need to return the default bounds in this case to avoid flickering (some content relies on this path)
+			if (CachedEmitter->SimTarget == ENiagaraSimTarget::GPUComputeSim)
 			{
-				bCachedBoundsDynamic = true;
-				if (CachedEmitter->bLocalSpace)
-				{
-					CachedBounds = DynamicBounds;
-				}
-				else
-				{
-					CachedBounds = DynamicBounds.TransformBy(FMatrix(ParentSystemInstance->GetOwnerParameters().EngineWorldToLocal));
-				}
+				CachedBounds = UNiagaraEmitter::GetDefaultFixedBounds();
 			}
 			else
 			{
-				CachedBounds.IsValid = false;
+				FBox DynamicBounds = InternalCalculateDynamicBounds(ParticleDataSet->GetCurrentDataChecked().GetNumInstances());
+				if (DynamicBounds.IsValid)
+				{
+					bCachedBoundsDynamic = true;
+					if (CachedEmitter->bLocalSpace)
+					{
+						CachedBounds = DynamicBounds;
+					}
+					else
+					{
+						CachedBounds = DynamicBounds.TransformBy(FMatrix(ParentSystemInstance->GetOwnerParameters().EngineWorldToLocal));
+					}
+				}
 			}
 		}
 	}
@@ -1249,7 +1254,7 @@ FBox FNiagaraEmitterInstance::GetFixedBounds() const
 		}
 	}
 
-	if ( CachedEmitter && CachedEmitter->bFixedBounds)
+	if ( CachedEmitter && CachedEmitter->CalculateBoundsMode == ENiagaraEmitterCalculateBoundMode::Fixed )
 	{
 		return CachedEmitter->FixedBounds;
 	}
diff --git a/Engine/Plugins/FX/Niagara/Source/Niagara/Private/NiagaraSystem.cpp b/Engine/Plugins/FX/Niagara/Source/Niagara/Private/NiagaraSystem.cpp
--- a/Engine/Plugins/FX/Niagara/Source/Niagara/Private/NiagaraSystem.cpp
+++ b/Engine/Plugins/FX/Niagara/Source/Niagara/Private/NiagaraSystem.cpp
@@ -1399,7 +1399,7 @@ void UNiagaraSystem::GetAssetRegistryTags(TArray<FAssetRegistryTag>& OutTags) co
 			if (Emitter)
 			{
 				// Only register fixed bounds requirement for GPU if the system itself isn't fixed bounds.
-				if (bFixedBounds == false && Emitter->bFixedBounds == false && Emitter->SimTarget == ENiagaraSimTarget::GPUComputeSim)
+				if (bFixedBounds == false && Emitter->CalculateBoundsMode == ENiagaraEmitterCalculateBoundMode::Dynamic && Emitter->SimTarget == ENiagaraSimTarget::GPUComputeSim)
 				{
 					GPUSimsMissingFixedBounds++;
 				}
diff --git a/Engine/Plugins/FX/Niagara/Source/Niagara/Private/NiagaraSystemInstance.cpp b/Engine/Plugins/FX/Niagara/Source/Niagara/Private/NiagaraSystemInstance.cpp
--- a/Engine/Plugins/FX/Niagara/Source/Niagara/Private/NiagaraSystemInstance.cpp
+++ b/Engine/Plugins/FX/Niagara/Source/Niagara/Private/NiagaraSystemInstance.cpp
@@ -2118,7 +2118,7 @@ void FNiagaraSystemInstance::InitEmitters()
 			{
 				if (const UNiagaraEmitter* EmitterAsset = EmitterHandles[EmitterIdx].GetInstance())
 				{
-					if (EmitterAsset->bFixedBounds)
+					if (EmitterAsset->CalculateBoundsMode == ENiagaraEmitterCalculateBoundMode::Fixed)
 					{
 						LocalBounds += EmitterAsset->FixedBounds;
 					}
diff --git a/Engine/Plugins/FX/Niagara/Source/NiagaraEditor/Private/Commandlets/NiagaraSystemAuditCommandlet.cpp b/Engine/Plugins/FX/Niagara/Source/NiagaraEditor/Private/Commandlets/NiagaraSystemAuditCommandlet.cpp
--- a/Engine/Plugins/FX/Niagara/Source/NiagaraEditor/Private/Commandlets/NiagaraSystemAuditCommandlet.cpp
+++ b/Engine/Plugins/FX/Niagara/Source/NiagaraEditor/Private/Commandlets/NiagaraSystemAuditCommandlet.cpp
@@ -300,7 +300,7 @@ bool UNiagaraSystemAuditCommandlet::ProcessNiagaraSystems()
 				}
 			}
 
-			if ( !NiagaraEmitter->bFixedBounds && !NiagaraSystem->bFixedBounds )
+			if ( (NiagaraEmitter->CalculateBoundsMode == ENiagaraEmitterCalculateBoundMode::Dynamic) && !NiagaraSystem->bFixedBounds )
 			{
 				EmittersWithDynamicBounds.Append(NiagaraEmitter->GetDebugSimName());
 			}
diff --git a/Engine/Plugins/FX/Niagara/Source/NiagaraEditor/Private/NiagaraValidationRules.cpp b/Engine/Plugins/FX/Niagara/Source/NiagaraEditor/Private/NiagaraValidationRules.cpp
--- a/Engine/Plugins/FX/Niagara/Source/NiagaraEditor/Private/NiagaraValidationRules.cpp
+++ b/Engine/Plugins/FX/Niagara/Source/NiagaraEditor/Private/NiagaraValidationRules.cpp
@@ -99,10 +99,10 @@ TArray<FNiagaraValidationResult> UNiagara_FixedGPUBoundsSet::CheckValidity(TShar
 	for (TSharedRef<FNiagaraEmitterHandleViewModel> EmitterHandleModel : EmitterHandleViewModels)
 	{
 		UNiagaraEmitter* NiagaraEmitter = EmitterHandleModel.Get().GetEmitterHandle()->GetInstance();
-		if (NiagaraEmitter->SimTarget == ENiagaraSimTarget::GPUComputeSim && NiagaraEmitter->bFixedBounds == false)
+		if (NiagaraEmitter->SimTarget == ENiagaraSimTarget::GPUComputeSim && NiagaraEmitter->CalculateBoundsMode == ENiagaraEmitterCalculateBoundMode::Dynamic)
 		{
 			UNiagaraStackEmitterPropertiesItem* EmitterProperties = GetStackEntry<UNiagaraStackEmitterPropertiesItem>(EmitterHandleModel.Get().GetEmitterStackViewModel());
-			FNiagaraValidationResult Result(ENiagaraValidationSeverity::Error, LOCTEXT("FixedBoundsSummary", "GPU emitters require fixed bounds"), LOCTEXT("FixedBoundsDescription", "Please set fixed bounds for gpu emitters or the system as a whole."), EmitterProperties);
+			FNiagaraValidationResult Result(ENiagaraValidationSeverity::Error, LOCTEXT("GpuDynamicBoundsErrorSummary", "GPU emitters do not support dynamic bounds"), LOCTEXT("GpuDynamicBoundsErrorDescription", "Gpu emitter should either not be in dynamic mode or the system must have fixed bounds."), EmitterProperties);
 			Results.Add(Result);
 		}
 	}
diff --git a/Engine/Plugins/FX/Niagara/Source/NiagaraEditor/Private/ViewModels/Stack/NiagaraStackEmitterSettingsGroup.cpp b/Engine/Plugins/FX/Niagara/Source/NiagaraEditor/Private/ViewModels/Stack/NiagaraStackEmitterSettingsGroup.cpp
--- a/Engine/Plugins/FX/Niagara/Source/NiagaraEditor/Private/ViewModels/Stack/NiagaraStackEmitterSettingsGroup.cpp
+++ b/Engine/Plugins/FX/Niagara/Source/NiagaraEditor/Private/ViewModels/Stack/NiagaraStackEmitterSettingsGroup.cpp
@@ -119,7 +119,7 @@ void UNiagaraStackEmitterPropertiesItem::RefreshChildrenInternal(const TArray<UN
 void UNiagaraStackEmitterPropertiesItem::RefreshIssues(TArray<FStackIssue>& NewIssues)
 {
 	UNiagaraEmitter* ActualEmitter = GetEmitterViewModel()->GetEmitter();
-	if (ActualEmitter && ActualEmitter->SimTarget == ENiagaraSimTarget::GPUComputeSim && ActualEmitter->bFixedBounds == false)
+	if (ActualEmitter && ActualEmitter->SimTarget == ENiagaraSimTarget::GPUComputeSim && ActualEmitter->CalculateBoundsMode == ENiagaraEmitterCalculateBoundMode::Dynamic)
 	{
 		bool bAddError = true;
 		
@@ -134,7 +134,7 @@ void UNiagaraStackEmitterPropertiesItem::RefreshIssues(TArray<FStackIssue>& NewI
 		{
 			FStackIssue MissingRequiredFixedBoundsModuleError(
 				EStackIssueSeverity::Warning,
-				LOCTEXT("RequiredFixedBoundsWarningFormat", "The emitter is GPU but the fixed bounds checkbox is not set.\r\nPlease update the Emitter or System properties, otherwise existing value for fixed bounds will be used."),
+				LOCTEXT("RequiredFixedBoundsWarningFormat", "The emitter is GPU and is using dynamic bounds mode.\r\nPlease update the Emitter or System properties otherwise bounds may be incorrect."),
 				LOCTEXT("MissingFixedBounds", "Missing fixed bounds."),
 				GetStackEditorDataKey(),
 				false);
