Change CameraDI to only include player controllers which have a LocalPlayer when looking for the appropriate camera informaiton

diff --git a/Engine/Plugins/FX/Niagara/Source/Niagara/Private/NiagaraDataInterfaceCamera.cpp b/Engine/Plugins/FX/Niagara/Source/Niagara/Private/NiagaraDataInterfaceCamera.cpp
--- a/Engine/Plugins/FX/Niagara/Source/Niagara/Private/NiagaraDataInterfaceCamera.cpp
+++ b/Engine/Plugins/FX/Niagara/Source/Niagara/Private/NiagaraDataInterfaceCamera.cpp
@@ -1,6 +1,8 @@
 // Copyright Epic Games, Inc. All Rights Reserved.
 
 #include "NiagaraDataInterfaceCamera.h"
+
+#include "Engine/LocalPlayer.h"
 #include "NiagaraTypes.h"
 #include "NiagaraWorldManager.h"
 #include "Internationalization/Internationalization.h"
@@ -84,18 +86,35 @@ bool UNiagaraDataInterfaceCamera::PerInstanceTick(void* PerInstanceData, FNiagar
 	UWorld* World = SystemInstance->GetWorldManager()->GetWorld();
 	if (World && PlayerControllerIndex < World->GetNumPlayerControllers())
 	{
-		int32 i = 0;
+		int32 LocalPlayerIndex = 0;
 		for (FConstPlayerControllerIterator Iterator = World->GetPlayerControllerIterator(); Iterator; ++Iterator)
 		{
 			APlayerController* PlayerController = Iterator->Get();
-			if (i == PlayerControllerIndex && PlayerController && PlayerController->PlayerCameraManager)
+			const bool IsLocalPlayer = (PlayerController && PlayerController->Player)
+				? PlayerController->Player->IsA<ULocalPlayer>()
+				: false;
+
+			if (IsLocalPlayer)
 			{
-				PIData->CameraLocation = LWCConverter.ConvertWorldToSimulationPosition(PlayerController->PlayerCameraManager->GetCameraLocation());
-				PIData->CameraRotation = PlayerController->PlayerCameraManager->GetCameraRotation();
-				PIData->CameraFOV = PlayerController->PlayerCameraManager->GetFOVAngle();
-				return false;
+				if (LocalPlayerIndex == PlayerControllerIndex)
+				{
+					if (PlayerController->PlayerCameraManager)
+					{
+						PIData->CameraLocation = LWCConverter.ConvertWorldToSimulationPosition(PlayerController->PlayerCameraManager->GetCameraLocation());
+						PIData->CameraRotation = PlayerController->PlayerCameraManager->GetCameraRotation();
+						PIData->CameraFOV = PlayerController->PlayerCameraManager->GetFOVAngle();
+						return false;
+					}
+					else
+					{
+						// if the requested local player doesn't have a player camera manager, then just break out of our search and provide defaults
+						UE_LOG(LogNiagara, Warning, TEXT("%s failed to find PlayerCameraManager for LocalPlayer %d"), *GetPathNameSafe(this), PlayerControllerIndex);
+						break;
+					}
+				}
+
+				++LocalPlayerIndex;
 			}
-			i++;
 		}
 	}
 #if WITH_EDITORONLY_DATA
