Connects FNiagaraShaderScript::ShouldCache to a new ShouldCompile function available to DataInterfaces.  Currently implemented for platforms to prevent compilation of scripts which use atomics on volumes (RasterizationGrid and XXX).

diff --git a/Engine/Plugins/FX/Niagara/Source/Niagara/Classes/NiagaraDataInterface.h b/Engine/Plugins/FX/Niagara/Source/Niagara/Classes/NiagaraDataInterface.h
--- a/Engine/Plugins/FX/Niagara/Source/Niagara/Classes/NiagaraDataInterface.h
+++ b/Engine/Plugins/FX/Niagara/Source/Niagara/Classes/NiagaraDataInterface.h
@@ -318,9 +318,14 @@ public:
 	virtual bool AppendCompileHash(FNiagaraCompileHashVisitor* InVisitor) const;
 #endif
 
+#if WITH_EDITOR
 	/** Allows data interfaces to influence the compilation of GPU shaders and is only called on the CDO object not the instance. */
 	virtual void ModifyCompilationEnvironment(EShaderPlatform ShaderPlatform, struct FShaderCompilerEnvironment& OutEnvironment) const;
 
+	/** Allows data interfaces to prevent compilation of GPU shaders and is only called on the CDO object not the instance. */
+	virtual bool ShouldCompile(EShaderPlatform ShaderPlatform) const { return true; };
+#endif
+
 	/** 
 		Subclasses that wish to work with GPU systems/emitters must implement this.
 		Those interfaces must fill DataForRenderThread with the data needed to upload to the GPU. It will be the last thing called on this
diff --git a/Engine/Plugins/FX/Niagara/Source/Niagara/Classes/NiagaraDataInterfaceNeighborGrid3D.h b/Engine/Plugins/FX/Niagara/Source/Niagara/Classes/NiagaraDataInterfaceNeighborGrid3D.h
--- a/Engine/Plugins/FX/Niagara/Source/Niagara/Classes/NiagaraDataInterfaceNeighborGrid3D.h
+++ b/Engine/Plugins/FX/Niagara/Source/Niagara/Classes/NiagaraDataInterfaceNeighborGrid3D.h
@@ -95,6 +95,9 @@ public:
 	virtual void GetParameterDefinitionHLSL(const FNiagaraDataInterfaceGPUParamInfo& ParamInfo, FString& OutHLSL) override;
 	virtual bool GetFunctionHLSL(const FNiagaraDataInterfaceGPUParamInfo& ParamInfo, const FNiagaraDataInterfaceGeneratedFunction& FunctionInfo, int FunctionInstanceIndex, FString& OutHLSL) override;
 #endif
+#if WITH_EDITOR
+	virtual bool ShouldCompile(EShaderPlatform ShaderPlatform) const override;
+#endif
 
 	virtual void ProvidePerInstanceDataForRenderThread(void* DataForRenderThread, void* PerInstanceData, const FNiagaraSystemInstanceID& SystemInstance) override {}
 	virtual bool InitPerInstanceData(void* PerInstanceData, FNiagaraSystemInstance* SystemInstance) override;
diff --git a/Engine/Plugins/FX/Niagara/Source/Niagara/Classes/NiagaraDataInterfaceRasterizationGrid3D.h b/Engine/Plugins/FX/Niagara/Source/Niagara/Classes/NiagaraDataInterfaceRasterizationGrid3D.h
--- a/Engine/Plugins/FX/Niagara/Source/Niagara/Classes/NiagaraDataInterfaceRasterizationGrid3D.h
+++ b/Engine/Plugins/FX/Niagara/Source/Niagara/Classes/NiagaraDataInterfaceRasterizationGrid3D.h
@@ -96,6 +96,9 @@ public:
 	virtual void GetVMExternalFunction(const FVMExternalFunctionBindingInfo& BindingInfo, void* InstanceData, FVMExternalFunction &OutFunc) override;
 
 	virtual bool Equals(const UNiagaraDataInterface* Other) const override;
+#if WITH_EDITOR
+	virtual bool ShouldCompile(EShaderPlatform ShaderPlatform) const override;
+#endif
 
 	// GPU sim functionality
 #if WITH_EDITORONLY_DATA
diff --git a/Engine/Plugins/FX/Niagara/Source/Niagara/Classes/NiagaraDataInterfaceSkeletalMesh.h b/Engine/Plugins/FX/Niagara/Source/Niagara/Classes/NiagaraDataInterfaceSkeletalMesh.h
--- a/Engine/Plugins/FX/Niagara/Source/Niagara/Classes/NiagaraDataInterfaceSkeletalMesh.h
+++ b/Engine/Plugins/FX/Niagara/Source/Niagara/Classes/NiagaraDataInterfaceSkeletalMesh.h
@@ -806,7 +806,9 @@ public:
 	virtual bool HasTickGroupPrereqs() const override { return true; }
 	virtual ETickingGroup CalculateTickGroup(const void* PerInstanceData) const override;
 
+#if WITH_EDITOR
 	virtual void ModifyCompilationEnvironment(EShaderPlatform ShaderPlatform, struct FShaderCompilerEnvironment& OutEnvironment) const override;
+#endif
 #if WITH_EDITORONLY_DATA
 	virtual bool AppendCompileHash(FNiagaraCompileHashVisitor* InVisitor) const override;
 	virtual void GetCommonHLSL(FString& OutHLSL) override;
diff --git a/Engine/Plugins/FX/Niagara/Source/Niagara/Classes/NiagaraScript.h b/Engine/Plugins/FX/Niagara/Source/Niagara/Classes/NiagaraScript.h
--- a/Engine/Plugins/FX/Niagara/Source/Niagara/Classes/NiagaraScript.h
+++ b/Engine/Plugins/FX/Niagara/Source/Niagara/Classes/NiagaraScript.h
@@ -942,6 +942,7 @@ public:
 	//~ End UObject interface
 
 	//~ Begin UNiagaraScriptBase interface
+	virtual bool ShouldCompile(EShaderPlatform Platform) const override;
 	virtual void ModifyCompilationEnvironment(EShaderPlatform Platform, struct FShaderCompilerEnvironment& OutEnvironment) const override;
 	virtual TConstArrayView<FSimulationStageMetaData> GetSimulationStageMetaData() const override { return MakeArrayView(CachedScriptVM.SimulationStageMetaData); }
 	//~ End UNiagaraScriptBase interface
diff --git a/Engine/Plugins/FX/Niagara/Source/Niagara/Private/NiagaraDataInterface.cpp b/Engine/Plugins/FX/Niagara/Source/Niagara/Private/NiagaraDataInterface.cpp
--- a/Engine/Plugins/FX/Niagara/Source/Niagara/Private/NiagaraDataInterface.cpp
+++ b/Engine/Plugins/FX/Niagara/Source/Niagara/Private/NiagaraDataInterface.cpp
@@ -48,6 +48,7 @@ bool UNiagaraDataInterface::AppendCompileHash(FNiagaraCompileHashVisitor* InVisi
 }
 #endif
 
+#if WITH_EDITOR
 void UNiagaraDataInterface::ModifyCompilationEnvironment(EShaderPlatform ShaderPlatform, struct FShaderCompilerEnvironment& OutEnvironment) const
 {
 	if (FDataDrivenShaderPlatformInfo::GetSupportsDxc(ShaderPlatform))
@@ -56,6 +57,7 @@ void UNiagaraDataInterface::ModifyCompilationEnvironment(EShaderPlatform ShaderP
 		OutEnvironment.CompilerFlags.Add(CFLAG_ForceDXC);
 	}
 }
+#endif
 
 void UNiagaraDataInterface::GetAssetTagsForContext(const UObject* InAsset, const TArray<const UNiagaraDataInterface*>& InProperties, TMap<FName, uint32>& NumericKeys, TMap<FName, FString>& StringKeys) const
 {
diff --git a/Engine/Plugins/FX/Niagara/Source/Niagara/Private/NiagaraDataInterfaceNeighborGrid3D.cpp b/Engine/Plugins/FX/Niagara/Source/Niagara/Private/NiagaraDataInterfaceNeighborGrid3D.cpp
--- a/Engine/Plugins/FX/Niagara/Source/Niagara/Private/NiagaraDataInterfaceNeighborGrid3D.cpp
+++ b/Engine/Plugins/FX/Niagara/Source/Niagara/Private/NiagaraDataInterfaceNeighborGrid3D.cpp
@@ -399,6 +399,18 @@ bool UNiagaraDataInterfaceNeighborGrid3D::Equals(const UNiagaraDataInterface* Ot
 	return OtherTyped->MaxNeighborsPerCell == MaxNeighborsPerCell;
 }
 
+#if WITH_EDITOR
+bool UNiagaraDataInterfaceNeighborGrid3D::ShouldCompile(EShaderPlatform ShaderPlatform) const
+{
+	if (!RHISupportsVolumeTextureAtomics(ShaderPlatform))
+	{
+		return false;
+	}
+
+	return UNiagaraDataInterface::ShouldCompile(ShaderPlatform);
+}
+#endif
+
 #if WITH_EDITORONLY_DATA
 void UNiagaraDataInterfaceNeighborGrid3D::GetParameterDefinitionHLSL(const FNiagaraDataInterfaceGPUParamInfo& ParamInfo, FString& OutHLSL)
 {
diff --git a/Engine/Plugins/FX/Niagara/Source/Niagara/Private/NiagaraDataInterfaceRasterizationGrid3D.cpp b/Engine/Plugins/FX/Niagara/Source/Niagara/Private/NiagaraDataInterfaceRasterizationGrid3D.cpp
--- a/Engine/Plugins/FX/Niagara/Source/Niagara/Private/NiagaraDataInterfaceRasterizationGrid3D.cpp
+++ b/Engine/Plugins/FX/Niagara/Source/Niagara/Private/NiagaraDataInterfaceRasterizationGrid3D.cpp
@@ -582,6 +582,18 @@ bool UNiagaraDataInterfaceRasterizationGrid3D::Equals(const UNiagaraDataInterfac
 	return OtherTyped->NumAttributes == NumAttributes && OtherTyped->Precision == Precision && OtherTyped->ResetValue == ResetValue;
 }
 
+#if WITH_EDITOR
+bool UNiagaraDataInterfaceRasterizationGrid3D::ShouldCompile(EShaderPlatform ShaderPlatform) const
+{
+	if (!RHISupportsVolumeTextureAtomics(ShaderPlatform))
+	{
+		return false;
+	}
+
+	return UNiagaraDataInterface::ShouldCompile(ShaderPlatform);
+}
+#endif
+
 #if WITH_EDITORONLY_DATA
 void UNiagaraDataInterfaceRasterizationGrid3D::GetParameterDefinitionHLSL(const FNiagaraDataInterfaceGPUParamInfo& ParamInfo, FString& OutHLSL)
 {
diff --git a/Engine/Plugins/FX/Niagara/Source/Niagara/Private/NiagaraDataInterfaceSkeletalMesh.cpp b/Engine/Plugins/FX/Niagara/Source/Niagara/Private/NiagaraDataInterfaceSkeletalMesh.cpp
--- a/Engine/Plugins/FX/Niagara/Source/Niagara/Private/NiagaraDataInterfaceSkeletalMesh.cpp
+++ b/Engine/Plugins/FX/Niagara/Source/Niagara/Private/NiagaraDataInterfaceSkeletalMesh.cpp
@@ -2985,6 +2985,7 @@ bool UNiagaraDataInterfaceSkeletalMesh::AppendCompileHash(FNiagaraCompileHashVis
 }
 #endif
 
+#if WITH_EDITOR
 void UNiagaraDataInterfaceSkeletalMesh::ModifyCompilationEnvironment(EShaderPlatform ShaderPlatform, struct FShaderCompilerEnvironment& OutEnvironment) const
 {
 	Super::ModifyCompilationEnvironment(ShaderPlatform, OutEnvironment);
@@ -2993,6 +2994,7 @@ void UNiagaraDataInterfaceSkeletalMesh::ModifyCompilationEnvironment(EShaderPlat
 	OutEnvironment.SetDefine(TEXT("DISKELMESH_PROBALIAS_FORMAT"), int(GetDefault<UNiagaraSettings>()->NDISkelMesh_GpuUniformSamplingFormat));
 	OutEnvironment.SetDefine(TEXT("DISKELMESH_ADJ_INDEX_FORMAT"), int(GetDefault<UNiagaraSettings>()->NDISkelMesh_AdjacencyTriangleIndexFormat));
 }
+#endif
 
 #if WITH_EDITORONLY_DATA
 void UNiagaraDataInterfaceSkeletalMesh::GetCommonHLSL(FString& OutHLSL)
diff --git a/Engine/Plugins/FX/Niagara/Source/Niagara/Private/NiagaraScript.cpp b/Engine/Plugins/FX/Niagara/Source/Niagara/Private/NiagaraScript.cpp
--- a/Engine/Plugins/FX/Niagara/Source/Niagara/Private/NiagaraScript.cpp
+++ b/Engine/Plugins/FX/Niagara/Source/Niagara/Private/NiagaraScript.cpp
@@ -2805,8 +2805,39 @@ bool UNiagaraScript::IsEditorOnly() const
 	return Super::IsEditorOnly();
 }
 
+bool UNiagaraScript::ShouldCompile(EShaderPlatform Platform) const
+{
+#if WITH_EDITOR
+	// Add all data interfaces
+	TSet<UClass*> DIUniqueClasses;
+	for (const FNiagaraScriptDataInterfaceInfo& DataInterfaceInfo : CachedDefaultDataInterfaces)
+	{
+		if (UNiagaraDataInterface* DataInterface = DataInterfaceInfo.DataInterface)
+		{
+			DIUniqueClasses.Add(DataInterface->GetClass());
+		}
+	}
+
+	// For each data interface allow them to modify the compilation environment
+	for (UClass* DIClass : DIUniqueClasses)
+	{
+		if (UNiagaraDataInterface* DICDO = CastChecked<UNiagaraDataInterface>(DIClass->GetDefaultObject(true)))
+		{
+			if (!DICDO->ShouldCompile(Platform))
+			{
+				return false;
+			}
+		}
+	}
+	return true;
+#else
+	return false;
+#endif
+}
+
 void UNiagaraScript::ModifyCompilationEnvironment(EShaderPlatform Platform, struct FShaderCompilerEnvironment& OutEnvironment) const
 {
+#if WITH_EDITOR
 	// Add all data interfaces
 	TSet<UClass*> DIUniqueClasses;
 	for ( const FNiagaraScriptDataInterfaceInfo& DataInterfaceInfo : CachedDefaultDataInterfaces )
@@ -2825,6 +2856,7 @@ void UNiagaraScript::ModifyCompilationEnvironment(EShaderPlatform Platform, stru
 			DICDO->ModifyCompilationEnvironment(Platform, OutEnvironment);
 		}
 	}
+#endif
 }
 
 #if WITH_EDITOR
diff --git a/Engine/Plugins/FX/Niagara/Source/NiagaraShader/Private/NiagaraShared.cpp b/Engine/Plugins/FX/Niagara/Source/NiagaraShader/Private/NiagaraShared.cpp
--- a/Engine/Plugins/FX/Niagara/Source/NiagaraShader/Private/NiagaraShared.cpp
+++ b/Engine/Plugins/FX/Niagara/Source/NiagaraShader/Private/NiagaraShared.cpp
@@ -63,6 +63,13 @@ void FNiagaraShaderScript::SetupShaderCompilationEnvironment(
 bool FNiagaraShaderScript::ShouldCache(EShaderPlatform Platform, const FShaderType* ShaderType) const
 {
 	check(ShaderType->GetNiagaraShaderType() != nullptr);
+	if (BaseVMScript)
+	{
+		if (!BaseVMScript->ShouldCompile(Platform))
+		{
+			return false;
+		}
+	}
 	return true;
 }
 
diff --git a/Engine/Plugins/FX/Niagara/Source/NiagaraShader/Public/NiagaraScriptBase.h b/Engine/Plugins/FX/Niagara/Source/NiagaraShader/Public/NiagaraScriptBase.h
--- a/Engine/Plugins/FX/Niagara/Source/NiagaraShader/Public/NiagaraScriptBase.h
+++ b/Engine/Plugins/FX/Niagara/Source/NiagaraShader/Public/NiagaraScriptBase.h
@@ -109,5 +109,6 @@ class UNiagaraScriptBase : public UObject
 	GENERATED_UCLASS_BODY()
 
 	virtual void ModifyCompilationEnvironment(EShaderPlatform Platform, struct FShaderCompilerEnvironment& OutEnvironment) const PURE_VIRTUAL(UNiagaraScriptBase::ModifyCompilationEnvironment, );
+	virtual bool ShouldCompile(EShaderPlatform Platform) const PURE_VIRTUAL(UNiagaraScriptBase::ShouldCompile, return true; );
 	virtual TConstArrayView<FSimulationStageMetaData> GetSimulationStageMetaData() const PURE_VIRTUAL(UNiagaraScriptBase::GetSimulationStageMetaData(), return MakeArrayView<FSimulationStageMetaData>(nullptr, 0); )
 };
