Avoid re-initializing the same FName repeatedly in high frequency code in Niagara.  Saves 2.5mins from large project cook.

diff --git a/Engine/Plugins/FX/Niagara/Source/Niagara/Private/NiagaraTypes.cpp b/Engine/Plugins/FX/Niagara/Source/Niagara/Private/NiagaraTypes.cpp
--- a/Engine/Plugins/FX/Niagara/Source/Niagara/Private/NiagaraTypes.cpp
+++ b/Engine/Plugins/FX/Niagara/Source/Niagara/Private/NiagaraTypes.cpp
@@ -7,6 +7,8 @@
 #include "Misc/StringBuilder.h"
 #include "UObject/Class.h"
 
+static FName NAME_NiagaraPosition(TEXT("NiagaraPosition"));
+
 void FNiagaraVariableBase::SetNamespacedName(const FString& InNamespace, FName InVariableName)
 {
 	TStringBuilder<128> NameBuilder;
@@ -492,7 +494,7 @@ UScriptStruct* FNiagaraTypeHelper::FindNiagaraFriendlyTopLevelStruct(UScriptStru
 		return FNiagaraTypeDefinition::GetQuatStruct();
 	}
 
-	if (InStruct->GetFName() == FName("NiagaraPosition"))
+	if (InStruct->GetFName() == NAME_NiagaraPosition)
 	{
 		return FNiagaraTypeDefinition::GetPositionStruct();
 	}
