int param for interlocked ops returns original value

diff --git a/Engine/Plugins/FX/Niagara/Source/Niagara/Private/NiagaraDataInterfaceRasterizationGrid3D.cpp b/Engine/Plugins/FX/Niagara/Source/Niagara/Private/NiagaraDataInterfaceRasterizationGrid3D.cpp
--- a/Engine/Plugins/FX/Niagara/Source/Niagara/Private/NiagaraDataInterfaceRasterizationGrid3D.cpp
+++ b/Engine/Plugins/FX/Niagara/Source/Niagara/Private/NiagaraDataInterfaceRasterizationGrid3D.cpp
@@ -729,12 +729,13 @@ bool UNiagaraDataInterfaceRasterizationGrid3D::GetFunctionHLSL(const FNiagaraDat
 				if ( In_AttributeIndex < {NumAttributesName} )
 				{
 					int3 TileOffset = {PerAttributeDataName}[In_AttributeIndex].xyz;		
-					InterlockedAdd(RW{OutputIntGrid}[int3(In_IndexX + TileOffset.x, In_IndexY + TileOffset.y, In_IndexZ + TileOffset.z)], {FloatToIntFunctionName}(In_Value));
+					InterlockedAdd(RW{OutputIntGrid}[int3(In_IndexX + TileOffset.x, In_IndexY + TileOffset.y, In_IndexZ + TileOffset.z)], {FloatToIntFunctionName}(In_Value), val);					
 				}
 			}
 
 			void {FunctionName}(int In_IndexX, int In_IndexY, int In_IndexZ, int In_AttributeIndex, float In_Value, out float Original_Value)
-			{															
+			{			
+				Original_Value = 0.0;												
 				if ( In_AttributeIndex < {NumAttributesName} )
 				{
 					int3 TileOffset = {PerAttributeDataName}[In_AttributeIndex].xyz;		
@@ -805,12 +806,13 @@ bool UNiagaraDataInterfaceRasterizationGrid3D::GetFunctionHLSL(const FNiagaraDat
 				if ( In_AttributeIndex < {NumAttributesName} )
 				{
 					int3 TileOffset = {PerAttributeDataName}[In_AttributeIndex].xyz;			
-					InterlockedMin(RW{OutputIntGrid}[int3(In_IndexX + TileOffset.x, In_IndexY + TileOffset.y, In_IndexZ + TileOffset.z)], {FloatToIntFunctionName}(In_Value));
+					InterlockedMin(RW{OutputIntGrid}[int3(In_IndexX + TileOffset.x, In_IndexY + TileOffset.y, In_IndexZ + TileOffset.z)], {FloatToIntFunctionName}(In_Value), val);
 				}
 			}
 
 			void {FunctionName}(int In_IndexX, int In_IndexY, int In_IndexZ, int In_AttributeIndex, float In_Value, out float Original_Value)
-			{							
+			{			
+				Original_Value = 0.0;				
 				if ( In_AttributeIndex < {NumAttributesName} )
 				{
 					int3 TileOffset = {PerAttributeDataName}[In_AttributeIndex].xyz;			
@@ -833,12 +835,13 @@ bool UNiagaraDataInterfaceRasterizationGrid3D::GetFunctionHLSL(const FNiagaraDat
 				if ( In_AttributeIndex < {NumAttributesName} )
 				{
 					int3 TileOffset = {PerAttributeDataName}[In_AttributeIndex].xyz;			
-					InterlockedMax(RW{OutputIntGrid}[int3(In_IndexX + TileOffset.x, In_IndexY + TileOffset.y, In_IndexZ + TileOffset.z)], {FloatToIntFunctionName}(In_Value));
+					InterlockedMax(RW{OutputIntGrid}[int3(In_IndexX + TileOffset.x, In_IndexY + TileOffset.y, In_IndexZ + TileOffset.z)], {FloatToIntFunctionName}(In_Value), val);
 				}
 			}
 
 			void {FunctionName}(int In_IndexX, int In_IndexY, int In_IndexZ, int In_AttributeIndex, float In_Value, out float Original_Value)
-			{													
+			{		
+				Original_Value = 0.0;
 				if ( In_AttributeIndex < {NumAttributesName} )
 				{
 					int3 TileOffset = {PerAttributeDataName}[In_AttributeIndex].xyz;			
