Merging CL 19070938 - added a "loose type" mode to Niagara to allow users to treat the position type like a vector

diff --git a/Engine/Plugins/FX/Niagara/Source/NiagaraEditor/Private/NiagaraValidationRules.cpp b/Engine/Plugins/FX/Niagara/Source/NiagaraEditor/Private/NiagaraValidationRules.cpp
--- a/Engine/Plugins/FX/Niagara/Source/NiagaraEditor/Private/NiagaraValidationRules.cpp
+++ b/Engine/Plugins/FX/Niagara/Source/NiagaraEditor/Private/NiagaraValidationRules.cpp
@@ -6,6 +6,7 @@
 #include "NiagaraComponentRendererProperties.h"
 #include "NiagaraNodeFunctionCall.h"
 #include "NiagaraScriptSource.h"
+#include "NiagaraSettings.h"
 #include "NiagaraSystem.h"
 #include "NiagaraSystemEditorData.h"
 #include "ViewModels/NiagaraEmitterHandleViewModel.h"
@@ -163,6 +164,8 @@ TArray<FNiagaraValidationResult> UNiagara_InvalidEffectType::CheckValidity(TShar
 TArray<FNiagaraValidationResult> UNiagara_LWC_Validator::CheckValidity(TSharedPtr<FNiagaraSystemViewModel> ViewModel) const
 {
 	TArray<FNiagaraValidationResult> Results;
+
+	const UNiagaraSettings* Settings = GetDefault<UNiagaraSettings>();
 	UNiagaraSystem& System = ViewModel->GetSystem();
 	if (!System.SupportsLargeWorldCoordinates())
 	{
@@ -198,7 +201,7 @@ TArray<FNiagaraValidationResult> UNiagara_LWC_Validator::CheckValidity(TSharedPt
 				}
 
 				// check if the linked dynamic input script outputs a vector
-				if (Input->GetValueMode() == UNiagaraStackFunctionInput::EValueMode::Dynamic && Input->GetDynamicInputNode())
+				if (Input->GetValueMode() == UNiagaraStackFunctionInput::EValueMode::Dynamic && Input->GetDynamicInputNode() && Settings->bEnforceStrictStackTypes)
 				{
 					if (UNiagaraScriptSource* DynamicInputSource = Cast<UNiagaraScriptSource>(Input->GetDynamicInputNode()->GetFunctionScriptSource()))
 					{
@@ -217,7 +220,7 @@ TArray<FNiagaraValidationResult> UNiagara_LWC_Validator::CheckValidity(TSharedPt
 				}
 
 				// check if the linked input variable is a vector
-				if (Input->GetValueMode() == UNiagaraStackFunctionInput::EValueMode::Linked)
+				if (Input->GetValueMode() == UNiagaraStackFunctionInput::EValueMode::Linked && Settings->bEnforceStrictStackTypes)
 				{
 					FNiagaraVariable VectorVar(FNiagaraTypeDefinition::GetVec3Def(), Input->GetLinkedValueHandle().GetParameterHandleString());
 					const UNiagaraGraph* NiagaraGraph = Input->GetInputFunctionCallNode().GetNiagaraGraph();
@@ -225,7 +228,7 @@ TArray<FNiagaraValidationResult> UNiagara_LWC_Validator::CheckValidity(TSharedPt
 					// we check if metadata for a vector attribute with the linked name exists in the emitter/system script graph. Not 100% correct, but it needs to be fast and a few false negatives are acceptable.
 					if (NiagaraGraph && NiagaraGraph->GetMetaData(VectorVar).IsSet())
 					{
-						FNiagaraValidationResult Result(ENiagaraValidationSeverity::Warning, FText::Format(LOCTEXT("PositionLinkedVectorSummary", "Input '{0}' is linked to a vector attribute"), Input->GetDisplayName()), LOCTEXT("PositionLinkedVectorDescription", "Position types should only be set linked to position attributes. In this case, it is linked to a vector attribute and the implicit conversion can cause problems with large world coordinates."), Input);
+						FNiagaraValidationResult Result(ENiagaraValidationSeverity::Warning, FText::Format(LOCTEXT("PositionLinkedVectorSummary", "Input '{0}' is linked to a vector attribute"), Input->GetDisplayName()), LOCTEXT("PositionLinkedVectorDescription", "Position types should only be linked to position attributes. In this case, it is linked to a vector attribute and the implicit conversion can cause problems with large world coordinates."), Input);
 						Results.Add(Result);
 					}
 				}
