Allow pins to be moved multiple steps at once

diff --git a/Engine/Plugins/FX/Niagara/Source/NiagaraEditor/Private/NiagaraNodeWithDynamicPins.cpp b/Engine/Plugins/FX/Niagara/Source/NiagaraEditor/Private/NiagaraNodeWithDynamicPins.cpp
--- a/Engine/Plugins/FX/Niagara/Source/NiagaraEditor/Private/NiagaraNodeWithDynamicPins.cpp
+++ b/Engine/Plugins/FX/Niagara/Source/NiagaraEditor/Private/NiagaraNodeWithDynamicPins.cpp
@@ -131,7 +131,7 @@ void UNiagaraNodeWithDynamicPins::CreateAddPin(EEdGraphPinDirection Direction)
 bool UNiagaraNodeWithDynamicPins::IsAddPin(const UEdGraphPin* Pin) const
 {
 	return Pin->PinType.PinCategory == UEdGraphSchema_Niagara::PinCategoryMisc && 
-		Pin->PinType.PinSubCategory == UNiagaraNodeWithDynamicPins::AddPinSubCategory;
+		Pin->PinType.PinSubCategory == AddPinSubCategory;
 }
 
 bool UNiagaraNodeWithDynamicPins::CanRenamePin(const UEdGraphPin* Pin) const
@@ -169,49 +169,63 @@ bool UNiagaraNodeWithDynamicPins::CanMovePin(const UEdGraphPin* Pin, int32 Direc
 		{
 			return false;
 		}
-	}	
-	
-	return true;
+		return true;
+	}
+
+	return false;
 }
 
-void UNiagaraNodeWithDynamicPins::MoveDynamicPin(UEdGraphPin* Pin, int32 DirectionToMove)
+void UNiagaraNodeWithDynamicPins::MoveDynamicPin(UEdGraphPin* Pin, int32 MoveAmount)
 {
-	FPinCollectorArray SameDirectionPins;
-	if (Pin->Direction == EEdGraphPinDirection::EGPD_Input)
-	{
-		GetInputPins(SameDirectionPins);
-	}
-	else
-	{
-		GetOutputPins(SameDirectionPins);
-	}
+	Modify();
+	Pin->Modify();
 
-	for (int32 i = 0; i < SameDirectionPins.Num(); i++)
+	bool bModifiedPins = false;
+	int32 DirectionToMove = MoveAmount > 0 ? 1 : -1;
+	for (int Step = 0; Step < abs(MoveAmount); Step++)
 	{
-		if (SameDirectionPins[i] == Pin)
+		FPinCollectorArray SameDirectionPins;
+		if (Pin->Direction == EGPD_Input)
+		{
+			GetInputPins(SameDirectionPins);
+		}
+		else
+		{
+			GetOutputPins(SameDirectionPins);
+		}
+		
+		for (int32 i = 0; i < SameDirectionPins.Num(); i++)
 		{
-			if (i + DirectionToMove >= 0 && i + DirectionToMove < SameDirectionPins.Num())
+			if (SameDirectionPins[i] == Pin)
 			{
-				Modify();
-				UEdGraphPin* PinOld = SameDirectionPins[i + DirectionToMove];
-				if (PinOld)
-					PinOld->Modify();
-				Pin->Modify();
-
-				int32 RealPinIdx = INDEX_NONE;
-				int32 SwapRealPinIdx = INDEX_NONE;
-				Pins.Find(Pin, RealPinIdx);
-				Pins.Find(PinOld, SwapRealPinIdx);
+				if (i + DirectionToMove >= 0 && i + DirectionToMove < SameDirectionPins.Num())
+				{
+					UEdGraphPin* PinOld = SameDirectionPins[i + DirectionToMove];
+					if (PinOld)
+					{
+						PinOld->Modify();
+					}
+
+					int32 RealPinIdx = INDEX_NONE;
+					int32 SwapRealPinIdx = INDEX_NONE;
+					Pins.Find(Pin, RealPinIdx);
+					Pins.Find(PinOld, SwapRealPinIdx);
 				
-				Pins[SwapRealPinIdx] = Pin;
-				Pins[RealPinIdx] = PinOld;
-				//GetGraph()->NotifyGraphChanged();
+					Pins[SwapRealPinIdx] = Pin;
+					Pins[RealPinIdx] = PinOld;
+					//GetGraph()->NotifyGraphChanged();
 
-				MarkNodeRequiresSynchronization(__FUNCTION__, true);
-				break;
+					bModifiedPins = true;
+					break;
+				}
 			}
 		}
 	}
+
+	if (bModifiedPins)
+	{
+		MarkNodeRequiresSynchronization(__FUNCTION__, true);
+	}
 }
 
 bool UNiagaraNodeWithDynamicPins::IsValidPinToCompile(UEdGraphPin* Pin) const
@@ -261,7 +275,7 @@ void UNiagaraNodeWithDynamicPins::GetNodeContextMenuActions(UToolMenu* Menu, UGr
 		}
 		
 		FPinCollectorArray SameDirectionPins;
-		if (Context->Pin->Direction == EEdGraphPinDirection::EGPD_Input)
+		if (Context->Pin->Direction == EGPD_Input)
 		{
 			GetInputPins(SameDirectionPins);
 		}
@@ -272,23 +286,86 @@ void UNiagaraNodeWithDynamicPins::GetNodeContextMenuActions(UToolMenu* Menu, UGr
 		int32 PinIdx = INDEX_NONE;
 		SameDirectionPins.Find(const_cast<UEdGraphPin*>(Context->Pin), PinIdx);
 
+		FText MoveUpLabel = LOCTEXT("MoveDynamicPinUp", "Move pin up");
 		if (PinIdx != 0 && CanMovePin(Context->Pin, -1))
 		{
-			Section.AddMenuEntry(
+			FText MoveUpTooltip = LOCTEXT("MoveDynamicPinToolTipUp", "Move this pin and any connections up.");
+			if (CanMovePin(Context->Pin, -2))
+			{
+				Section.AddSubMenu(
+					"MoveDynamicPinUp",
+					MoveUpLabel, 
+					MoveUpTooltip,
+					FNewToolMenuDelegate::CreateLambda([=](UToolMenu* InSubMenuBuilder)
+				{
+					FToolMenuSection& SubSection = InSubMenuBuilder->FindOrAddSection("MovePin");
+					for (int i = PinIdx - 1; i >= 0; i--)
+					{
+						int32 MoveAmount = i - PinIdx;
+						if (!CanMovePin(Context->Pin, MoveAmount))
+						{
+							break;
+						}
+						FText PinName = FText::FromName(SameDirectionPins[i]->PinName);
+						SubSection.AddMenuEntry("MoveDynamicPinUp" + i,
+							FText::Format(LOCTEXT("MoveDynamicPinUpLabel", "Move pin above '{0}'"), PinName),
+							MoveUpTooltip,
+							FSlateIcon(),
+							FUIAction(FExecuteAction::CreateUObject(const_cast<UNiagaraNodeWithDynamicPins*>(this), &UNiagaraNodeWithDynamicPins::MoveDynamicPinFromMenu, const_cast<UEdGraphPin*>(Context->Pin), MoveAmount))
+						);
+					}
+				}));
+			}
+			else
+			{
+				Section.AddMenuEntry(
 				"MoveDynamicPinUp",
-				LOCTEXT("MoveDynamicPinUp", "Move pin up"),
+				MoveUpLabel,
 				LOCTEXT("MoveDynamicPinToolTipUp", "Move this pin and any connections one slot up."),
 				FSlateIcon(),
 				FUIAction(FExecuteAction::CreateUObject(const_cast<UNiagaraNodeWithDynamicPins*>(this), &UNiagaraNodeWithDynamicPins::MoveDynamicPinFromMenu, const_cast<UEdGraphPin*>(Context->Pin), -1)));
+			}
 		}
+
+		FText MoveDownLabel = LOCTEXT("MoveDynamicPinDown", "Move pin down");
 		if (PinIdx >= 0 && CanMovePin(Context->Pin, 1) && PinIdx < SameDirectionPins.Num() - 1)
 		{
-			Section.AddMenuEntry(
-				"MoveDynamicPinDown",
-				LOCTEXT("MoveDynamicPinDown", "Move pin down"),
-				LOCTEXT("MoveDynamicPinToolTipDown", "Move this pin and any connections one slot down."),
-				FSlateIcon(),
-				FUIAction(FExecuteAction::CreateUObject(const_cast<UNiagaraNodeWithDynamicPins*>(this), &UNiagaraNodeWithDynamicPins::MoveDynamicPinFromMenu, const_cast<UEdGraphPin*>(Context->Pin), 1)));
+			FText MoveDownTooltip = LOCTEXT("MoveDynamicPinToolTipDown", "Move this pin and any connections down.");
+			if (CanMovePin(Context->Pin, 2))
+			{
+				Section.AddSubMenu(
+					"MoveDynamicPinDown",
+					MoveDownLabel, 
+					MoveDownTooltip,
+					FNewToolMenuDelegate::CreateLambda([=](UToolMenu* InSubMenuBuilder)
+				{
+					FToolMenuSection& SubSection = InSubMenuBuilder->FindOrAddSection("MovePin");
+					for (int i = PinIdx + 1; i < SameDirectionPins.Num(); i++)
+					{
+						int32 MoveAmount = i - PinIdx;
+						if (!CanMovePin(Context->Pin, MoveAmount))
+						{
+							break;
+						}
+						FText PinName = FText::FromName(SameDirectionPins[i]->PinName);
+						SubSection.AddMenuEntry("MoveDynamicPinDown" + i,
+							FText::Format(LOCTEXT("MoveDynamicPinDownLabel", "Move pin below '{0}'"), PinName),
+							MoveDownTooltip,
+							FSlateIcon(),
+							FUIAction(FExecuteAction::CreateUObject(const_cast<UNiagaraNodeWithDynamicPins*>(this), &UNiagaraNodeWithDynamicPins::MoveDynamicPinFromMenu, const_cast<UEdGraphPin*>(Context->Pin), MoveAmount))
+						);
+					}
+				}));
+			}
+			else
+			{
+				Section.AddMenuEntry(
+					"MoveDynamicPinDown",
+					MoveDownLabel,
+					MoveDownTooltip,
+					FSlateIcon(),
+					FUIAction(FExecuteAction::CreateUObject(const_cast<UNiagaraNodeWithDynamicPins*>(this), &UNiagaraNodeWithDynamicPins::MoveDynamicPinFromMenu, const_cast<UEdGraphPin*>(Context->Pin), 1)));
+			}
 		}
 		
 	}
diff --git a/Engine/Plugins/FX/Niagara/Source/NiagaraEditor/Public/NiagaraNodeWithDynamicPins.h b/Engine/Plugins/FX/Niagara/Source/NiagaraEditor/Public/NiagaraNodeWithDynamicPins.h
--- a/Engine/Plugins/FX/Niagara/Source/NiagaraEditor/Public/NiagaraNodeWithDynamicPins.h
+++ b/Engine/Plugins/FX/Niagara/Source/NiagaraEditor/Public/NiagaraNodeWithDynamicPins.h
@@ -77,7 +77,7 @@ protected:
 	virtual void RemoveDynamicPin(UEdGraphPin* Pin);
 
 	/** Moves a pin among the pins of the same direction. Negative values for up, positive for down. */
-	virtual void MoveDynamicPin(UEdGraphPin* Pin, int32 DirectionToMove);
+	virtual void MoveDynamicPin(UEdGraphPin* Pin, int32 MoveAmount);
 
 	virtual bool OnVerifyTextChanged(const FText& NewText, FText& OutMessage) { return true; };
 
