Support more dynamic input scripts in loose type mode

diff --git a/Engine/Plugins/FX/Niagara/Source/NiagaraEditor/Private/ViewModels/Stack/NiagaraStackFunctionInput.cpp b/Engine/Plugins/FX/Niagara/Source/NiagaraEditor/Private/ViewModels/Stack/NiagaraStackFunctionInput.cpp
--- a/Engine/Plugins/FX/Niagara/Source/NiagaraEditor/Private/ViewModels/Stack/NiagaraStackFunctionInput.cpp
+++ b/Engine/Plugins/FX/Niagara/Source/NiagaraEditor/Private/ViewModels/Stack/NiagaraStackFunctionInput.cpp
@@ -354,7 +354,7 @@ void UNiagaraStackFunctionInput::Paste(const UNiagaraClipboardContent* Clipboard
 	{
 		if (const UNiagaraClipboardFunctionInput* ClipboardInput = ClipboardContent->FunctionInputs[0])
 		{
-			if (ClipboardInput->InputType == InputType)
+			if (FNiagaraEditorUtilities::AreTypesAssignable(ClipboardInput->InputType, InputType))
 			{
 				SetValueFromClipboardFunctionInput(*ClipboardInput);	
 			}
@@ -1553,10 +1553,7 @@ void UNiagaraStackFunctionInput::GetAvailableDynamicInputs(TArray<UNiagaraScript
 			{
 				const UEdGraphSchema_Niagara* NiagaraSchema = GetDefault<UEdGraphSchema_Niagara>();
 				FNiagaraTypeDefinition PinType = NiagaraSchema->PinToTypeDefinition(InputPins[0]);
-				if (PinType == InputType)
-				{
-					return true;
-				}
+				return FNiagaraEditorUtilities::AreTypesAssignable(PinType, InputType);
 			}
 		}
 		return false;
@@ -2475,7 +2472,7 @@ const UNiagaraClipboardFunctionInput* UNiagaraStackFunctionInput::ToClipboardFun
 
 void UNiagaraStackFunctionInput::SetValueFromClipboardFunctionInput(const UNiagaraClipboardFunctionInput& ClipboardFunctionInput)
 {
-	if (ensureMsgf(ClipboardFunctionInput.InputType == InputType, TEXT("Can not set input value from clipboard, input types don't match.")))
+	if (ensureMsgf(FNiagaraEditorUtilities::AreTypesAssignable(ClipboardFunctionInput.InputType, InputType), TEXT("Can not set input value from clipboard, input types don't match.")))
 	{
 		switch (ClipboardFunctionInput.ValueMode)
 		{
