* Fixed Linux build error
* Improve ui label

diff --git a/Engine/Plugins/FX/Niagara/Source/NiagaraEditor/Private/NiagaraNodeWithDynamicPins.cpp b/Engine/Plugins/FX/Niagara/Source/NiagaraEditor/Private/NiagaraNodeWithDynamicPins.cpp
--- a/Engine/Plugins/FX/Niagara/Source/NiagaraEditor/Private/NiagaraNodeWithDynamicPins.cpp
+++ b/Engine/Plugins/FX/Niagara/Source/NiagaraEditor/Private/NiagaraNodeWithDynamicPins.cpp
@@ -307,8 +307,8 @@ void UNiagaraNodeWithDynamicPins::GetNodeContextMenuActions(UToolMenu* Menu, UGr
 							break;
 						}
 						FText PinName = FText::FromName(SameDirectionPins[i]->PinName);
-						SubSection.AddMenuEntry("MoveDynamicPinUp" + i,
-							FText::Format(LOCTEXT("MoveDynamicPinUpLabel", "Move pin above '{0}'"), PinName),
+						SubSection.AddMenuEntry(FName("MoveDynamicPinUp" + FString::FromInt(i)),
+							FText::Format(LOCTEXT("MoveDynamicPinUpLabel", "Move up {0} - above '{1}'"), abs(MoveAmount), PinName),
 							MoveUpTooltip,
 							FSlateIcon(),
 							FUIAction(FExecuteAction::CreateUObject(const_cast<UNiagaraNodeWithDynamicPins*>(this), &UNiagaraNodeWithDynamicPins::MoveDynamicPinFromMenu, const_cast<UEdGraphPin*>(Context->Pin), MoveAmount))
@@ -348,8 +348,8 @@ void UNiagaraNodeWithDynamicPins::GetNodeContextMenuActions(UToolMenu* Menu, UGr
 							break;
 						}
 						FText PinName = FText::FromName(SameDirectionPins[i]->PinName);
-						SubSection.AddMenuEntry("MoveDynamicPinDown" + i,
-							FText::Format(LOCTEXT("MoveDynamicPinDownLabel", "Move pin below '{0}'"), PinName),
+						SubSection.AddMenuEntry(FName("MoveDynamicPinDown" + FString::FromInt(i)),
+							FText::Format(LOCTEXT("MoveDynamicPinDownLabel", "Move down {0} - below '{1}'"), abs(MoveAmount), PinName),
 							MoveDownTooltip,
 							FSlateIcon(),
 							FUIAction(FExecuteAction::CreateUObject(const_cast<UNiagaraNodeWithDynamicPins*>(this), &UNiagaraNodeWithDynamicPins::MoveDynamicPinFromMenu, const_cast<UEdGraphPin*>(Context->Pin), MoveAmount))
