- Move HasValidParameterBindings to after the RI parameters have been restored

#rnx

diff --git a/Engine/Plugins/FX/Niagara/Source/Niagara/Private/NiagaraScript.cpp b/Engine/Plugins/FX/Niagara/Source/Niagara/Private/NiagaraScript.cpp
--- a/Engine/Plugins/FX/Niagara/Source/Niagara/Private/NiagaraScript.cpp
+++ b/Engine/Plugins/FX/Niagara/Source/Niagara/Private/NiagaraScript.cpp
@@ -1604,11 +1604,6 @@ void UNiagaraScript::Serialize(FArchive& Ar)
 			ScriptExecutionParamStore.CoalescePaddingInfo();
 		}
 
-		if (!HasValidParameterBindings())
-		{
-			UE_LOG(LogNiagara, Warning, TEXT("Mismatch between binding between RapidIterationParamters and ScriptExecutionParameters for system %s"), *GetFullName());
-		}
-
 		if (GNiagaraCompressScriptByteCode)
 		{
 			ExecutableData.ByteCode.Compress();
@@ -1629,6 +1624,16 @@ void UNiagaraScript::Serialize(FArchive& Ar)
 		RapidIterationParameters = TemporaryStore;
 	}
 
+#if WITH_EDITORONLY_DATA
+	if (Ar.IsCooking() && Ar.IsSaving())
+	{
+		if (!HasValidParameterBindings())
+		{
+			UE_LOG(LogNiagara, Warning, TEXT("Mismatch between binding between RapidIterationParamters and ScriptExecutionParameters for system %s"), *GetFullName());
+		}
+	}
+#endif
+
 	bool IsValidShaderScript = false;
 	if (NiagaraVer < FNiagaraCustomVersion::DontCompileGPUWhenNotNeeded)
 	{
