- Fix potential divide by zero

diff --git a/Engine/Plugins/FX/Niagara/Source/Niagara/Private/NiagaraSystemInstance.cpp b/Engine/Plugins/FX/Niagara/Source/Niagara/Private/NiagaraSystemInstance.cpp
--- a/Engine/Plugins/FX/Niagara/Source/Niagara/Private/NiagaraSystemInstance.cpp
+++ b/Engine/Plugins/FX/Niagara/Source/Niagara/Private/NiagaraSystemInstance.cpp
@@ -1981,7 +1981,7 @@ void FNiagaraSystemInstance::TickInstanceParameters_Concurrent()
 		CurrentOwnerParameters.EngineWorldToLocalNoScale = FMatrix44f(LocalToWorldNoScale.Inverse());
 		CurrentOwnerParameters.EngineRotation = FQuat4f((float)LastRotation.X, (float)LastRotation.Y, (float)LastRotation.Z, (float)LastRotation.W);
 		CurrentOwnerParameters.EnginePosition = (FVector3f)Location; // LWC_TODO: precision loss
-		CurrentOwnerParameters.EngineVelocity = (FVector3f)((Location - LastLocation) / GatheredInstanceParameters.DeltaSeconds);
+		CurrentOwnerParameters.EngineVelocity = GatheredInstanceParameters.DeltaSeconds > 0.0f ? (FVector3f)((Location - LastLocation) / GatheredInstanceParameters.DeltaSeconds) : FVector3f::ZeroVector;
 		CurrentOwnerParameters.EngineXAxis = CurrentOwnerParameters.EngineRotation.GetAxisX();
 		CurrentOwnerParameters.EngineYAxis = CurrentOwnerParameters.EngineRotation.GetAxisY();
 		CurrentOwnerParameters.EngineZAxis = CurrentOwnerParameters.EngineRotation.GetAxisZ();
@@ -1997,7 +1997,7 @@ void FNiagaraSystemInstance::TickInstanceParameters_Concurrent()
 
 	FNiagaraGlobalParameters& CurrentGlobalParameter = GlobalParameters[ParameterIndex];
 	CurrentGlobalParameter.EngineDeltaTime = GatheredInstanceParameters.DeltaSeconds;
-	CurrentGlobalParameter.EngineInvDeltaTime = 1.0f / GatheredInstanceParameters.DeltaSeconds;
+	CurrentGlobalParameter.EngineInvDeltaTime = GatheredInstanceParameters.DeltaSeconds > 0.0f ? 1.0f / GatheredInstanceParameters.DeltaSeconds : 0.0f;
 	CurrentGlobalParameter.EngineRealTime = GatheredInstanceParameters.RealTimeSeconds;
 	CurrentGlobalParameter.EngineTime = GatheredInstanceParameters.TimeSeconds;
 	CurrentGlobalParameter.QualityLevel = FNiagaraPlatformSet::GetQualityLevel();
