- Fix bug where a mid frame EOF updates would clear out SimulationsWithPostActorWork.  This meant that the tick group promotion or spawn was potentially lost.

#rnx

diff --git a/Engine/Plugins/FX/Niagara/Source/Niagara/Private/NiagaraWorldManager.cpp b/Engine/Plugins/FX/Niagara/Source/Niagara/Private/NiagaraWorldManager.cpp
--- a/Engine/Plugins/FX/Niagara/Source/Niagara/Private/NiagaraWorldManager.cpp
+++ b/Engine/Plugins/FX/Niagara/Source/Niagara/Private/NiagaraWorldManager.cpp
@@ -914,7 +914,8 @@ void FNiagaraWorldManager::PreSendAllEndOfFrameUpdates()
 				Simulation->WaitForInstancesTickComplete();
 			}
 		}
-		SimulationsWithPostActorWork.Reset();
+		// Note: We do not clear the SimulationsWithPostActorWork array as we just want to safely wait for any async work,
+		//       they still require additional processing, i.e. Tick Group Changes / Spawning
 	}
 
 	for (const auto& Simulation : SimulationsWithEndOfFrameWait)
