Fix for CIS duplicate key warnings.
#rnx

diff --git a/Engine/Plugins/FX/Niagara/Source/NiagaraEditor/Private/ViewModels/NiagaraParameterPanelViewModel.cpp b/Engine/Plugins/FX/Niagara/Source/NiagaraEditor/Private/ViewModels/NiagaraParameterPanelViewModel.cpp
--- a/Engine/Plugins/FX/Niagara/Source/NiagaraEditor/Private/ViewModels/NiagaraParameterPanelViewModel.cpp
+++ b/Engine/Plugins/FX/Niagara/Source/NiagaraEditor/Private/ViewModels/NiagaraParameterPanelViewModel.cpp
@@ -375,7 +375,7 @@ bool INiagaraParameterPanelViewModel::GetCanPasteParameterMetaDataAndToolTip(FTe
 void INiagaraParameterPanelViewModel::DuplicateParameters(const TArray<FNiagaraParameterPanelItem> ItemsToDuplicate) const
 {
 	bool bAnyChange = false;
-	FScopedTransaction Transaction(LOCTEXT("DuplicateParameterTransaction", "Duplicate parameter(s)"));
+	FScopedTransaction Transaction(LOCTEXT("DuplicateParameter(s)Transaction", "Duplicate parameter(s)"));
 
 	for(const FNiagaraParameterPanelItem& ItemToDuplicate : ItemsToDuplicate)
 	{
diff --git a/Engine/Source/Editor/DataLayerEditor/Private/DataLayer/SDataLayerOutliner.cpp b/Engine/Source/Editor/DataLayerEditor/Private/DataLayer/SDataLayerOutliner.cpp
--- a/Engine/Source/Editor/DataLayerEditor/Private/DataLayer/SDataLayerOutliner.cpp
+++ b/Engine/Source/Editor/DataLayerEditor/Private/DataLayer/SDataLayerOutliner.cpp
@@ -24,7 +24,7 @@ void SDataLayerOutliner::CustomAddToToolbar(TSharedPtr<SHorizontalBox> Toolbar)
 			SNew(SButton)
 			.IsEnabled(this, &SDataLayerOutliner::CanAddSelectedActorsToSelectedDataLayersClicked)
 			.ButtonStyle(FAppStyle::Get(), "SimpleButton")
-			.ToolTipText(LOCTEXT("AddSelectedActorsToSelectedDataLayers", "Add selected actors to selected Data Layers"))
+			.ToolTipText(LOCTEXT("AddSelectedActorsToSelectedDataLayersTooltip", "Add selected actors to selected Data Layers"))
 			.OnClicked(this, &SDataLayerOutliner::OnAddSelectedActorsToSelectedDataLayersClicked)
 			[
 				SNew(SImage)
diff --git a/Engine/Source/Editor/PropertyEditor/Private/SDetailSingleItemRow.cpp b/Engine/Source/Editor/PropertyEditor/Private/SDetailSingleItemRow.cpp
--- a/Engine/Source/Editor/PropertyEditor/Private/SDetailSingleItemRow.cpp
+++ b/Engine/Source/Editor/PropertyEditor/Private/SDetailSingleItemRow.cpp
@@ -597,7 +597,7 @@ void SDetailSingleItemRow::Construct( const FArguments& InArgs, FDetailLayoutCus
 			ResetToDefault.ToolTip = TAttribute<FText>::Create([this]() 
 			{
 				return IsResetToDefaultVisible() ?
-					NSLOCTEXT("PropertyEditor", "ResetToDefaultToolTip", "Reset this property to its default value.") :
+					NSLOCTEXT("PropertyEditor", "ResetToDefaultPropertyValueToolTip", "Reset this property to its default value.") :
 					FText::GetEmpty();
 			});
 
diff --git a/Engine/Source/Editor/PropertyEditor/Private/UserInterface/PropertyEditor/SResetToDefaultPropertyEditor.cpp b/Engine/Source/Editor/PropertyEditor/Private/UserInterface/PropertyEditor/SResetToDefaultPropertyEditor.cpp
--- a/Engine/Source/Editor/PropertyEditor/Private/UserInterface/PropertyEditor/SResetToDefaultPropertyEditor.cpp
+++ b/Engine/Source/Editor/PropertyEditor/Private/UserInterface/PropertyEditor/SResetToDefaultPropertyEditor.cpp
@@ -54,7 +54,7 @@ void SResetToDefaultPropertyEditor::Tick(const FGeometry& AllottedGeometry, cons
 FText SResetToDefaultPropertyEditor::GetResetToolTip() const
 {
 	FString Tooltip;
-	Tooltip = NSLOCTEXT("PropertyEditor", "ResetToDefaultToolTip", "Reset to Default").ToString();
+	Tooltip = LOCTEXT("ResetToDefaultToolTip", "Reset to Default").ToString();
 
 	if( PropertyHandle.IsValid() && !PropertyHandle->IsEditConst() && PropertyHandle->DiffersFromDefault() )
 	{
diff --git a/Engine/Source/Editor/TextureEditor/Private/Widgets/STextureEditorViewport.cpp b/Engine/Source/Editor/TextureEditor/Private/Widgets/STextureEditorViewport.cpp
--- a/Engine/Source/Editor/TextureEditor/Private/Widgets/STextureEditorViewport.cpp
+++ b/Engine/Source/Editor/TextureEditor/Private/Widgets/STextureEditorViewport.cpp
@@ -41,7 +41,7 @@ void STextureEditorViewport::Construct( const FArguments& InArgs, const TSharedR
 
 	if (InToolkit->GetTexture() != nullptr)
 	{
-		FText FormattedText = InToolkit->HasValidTextureResource() ? FText::FromString(TEXT("{0}")) : LOCTEXT( "InvalidTexture", "{0} (Invalid Texture)");
+		FText FormattedText = InToolkit->HasValidTextureResource() ? FText::FromString(TEXT("{0}")) : LOCTEXT( "InvalidTextureWithParam", "{0} (Invalid Texture)");
 		TextureName = FText::Format(FormattedText, FText::FromName(InToolkit->GetTexture()->GetFName()));
 	}
 
