Fixed a bug where types for user parameters were not considered across emitters

diff --git a/Engine/Plugins/FX/Niagara/Source/NiagaraEditor/Private/NiagaraScriptMergeManager.cpp b/Engine/Plugins/FX/Niagara/Source/NiagaraEditor/Private/NiagaraScriptMergeManager.cpp
--- a/Engine/Plugins/FX/Niagara/Source/NiagaraEditor/Private/NiagaraScriptMergeManager.cpp
+++ b/Engine/Plugins/FX/Niagara/Source/NiagaraEditor/Private/NiagaraScriptMergeManager.cpp
@@ -2811,13 +2811,15 @@ FNiagaraScriptMergeManager::FApplyDiffResults FNiagaraScriptMergeManager::AddInp
 			else if (OverrideToAdd->GetLinkedValueHandle().IsSet())
 			{
 				check(OverrideToAdd->GetOverrideNode() && OverrideToAdd->GetOverrideNode()->GetNiagaraGraph());
-				UNiagaraScriptVariable* ScriptVar = OverrideToAdd->GetOverrideNode()->GetNiagaraGraph()->GetScriptVariable(FNiagaraVariable(InputType, OverrideToAdd->GetLinkedValueHandle().GetValue().GetParameterHandleString()));
+				FNiagaraVariable Parameter = FNiagaraVariable(InputType, OverrideToAdd->GetLinkedValueHandle().GetValue().GetParameterHandleString());
+				UNiagaraScriptVariable* ScriptVar = OverrideToAdd->GetOverrideNode()->GetNiagaraGraph()->GetScriptVariable(Parameter);
 				ENiagaraDefaultMode DesiredMode = ENiagaraDefaultMode::Value;
 				if (ScriptVar)
 					DesiredMode = ScriptVar->DefaultMode;
 				if (GNiagaraForceFailIfPreviouslyNotSetOnMerge > 0)
 					DesiredMode = ENiagaraDefaultMode::FailIfPreviouslyNotSet;
-				FNiagaraStackGraphUtilities::SetLinkedValueHandleForFunctionInput(InputOverridePin, OverrideToAdd->GetLinkedValueHandle().GetValue(), DesiredMode, OverrideToAdd->GetOverrideNodeId());
+				TSet KnownParameters = { Parameter };
+				FNiagaraStackGraphUtilities::SetLinkedValueHandleForFunctionInput(InputOverridePin, OverrideToAdd->GetLinkedValueHandle().GetValue(), KnownParameters, DesiredMode, OverrideToAdd->GetOverrideNodeId());
 				Results.bSucceeded = true;
 			}
 			else if (OverrideToAdd->GetDataValueInputName().IsSet() && OverrideToAdd->GetDataValueObject() != nullptr)
diff --git a/Engine/Plugins/FX/Niagara/Source/NiagaraEditor/Private/ViewModels/Stack/NiagaraStackFunctionInput.cpp b/Engine/Plugins/FX/Niagara/Source/NiagaraEditor/Private/ViewModels/Stack/NiagaraStackFunctionInput.cpp
--- a/Engine/Plugins/FX/Niagara/Source/NiagaraEditor/Private/ViewModels/Stack/NiagaraStackFunctionInput.cpp
+++ b/Engine/Plugins/FX/Niagara/Source/NiagaraEditor/Private/ViewModels/Stack/NiagaraStackFunctionInput.cpp
@@ -1248,7 +1248,9 @@ void UNiagaraStackFunctionInput::SetLinkedValueHandle(const FNiagaraParameterHan
 		}
 
 		// Only set the linked value if it's actually different from the default.
-		FNiagaraStackGraphUtilities::SetLinkedValueHandleForFunctionInput(GetOrCreateOverridePin(), InParameterHandle);
+		UEdGraphPin& OverridePin = GetOrCreateOverridePin();
+		TSet<FNiagaraVariable> KnownParameters = FNiagaraStackGraphUtilities::GetParametersForContext(OverridePin.GetOwningNode()->GetGraph(), GetSystemViewModel()->GetSystem());
+		FNiagaraStackGraphUtilities::SetLinkedValueHandleForFunctionInput(OverridePin, InParameterHandle, KnownParameters);
 		OwningFunctionCallNode->FixupPinNames(); // refresh the input guids and update the bound names
 	}
 
diff --git a/Engine/Plugins/FX/Niagara/Source/NiagaraEditor/Private/ViewModels/Stack/NiagaraStackGraphUtilities.cpp b/Engine/Plugins/FX/Niagara/Source/NiagaraEditor/Private/ViewModels/Stack/NiagaraStackGraphUtilities.cpp
--- a/Engine/Plugins/FX/Niagara/Source/NiagaraEditor/Private/ViewModels/Stack/NiagaraStackGraphUtilities.cpp
+++ b/Engine/Plugins/FX/Niagara/Source/NiagaraEditor/Private/ViewModels/Stack/NiagaraStackGraphUtilities.cpp
@@ -1367,7 +1367,25 @@ UEdGraphPin* FNiagaraStackGraphUtilities::GetLinkedValueHandleForFunctionInput(c
 	return GetOutputPin;
 }
 
-void FNiagaraStackGraphUtilities::SetLinkedValueHandleForFunctionInput(UEdGraphPin& OverridePin, FNiagaraParameterHandle LinkedParameterHandle, ENiagaraDefaultMode DesiredDefaultMode, const FGuid& NewNodePersistentId)
+TSet<FNiagaraVariable> FNiagaraStackGraphUtilities::GetParametersForContext(UEdGraph* InGraph, UNiagaraSystem& System)
+{
+	TSet<FNiagaraVariable> Result;
+	if (UNiagaraGraph* NiagaraGraph = Cast<UNiagaraGraph>(InGraph))
+	{
+		const TMap<FNiagaraVariable, FNiagaraGraphParameterReferenceCollection>& ReferenceMap = NiagaraGraph->GetParameterReferenceMap();
+		ReferenceMap.GetKeys(Result);
+	}
+	TArray<FNiagaraVariable> UserParams;
+	System.GetExposedParameters().GetUserParameters(UserParams);
+	for (FNiagaraVariable& Var : UserParams)
+	{
+		FNiagaraUserRedirectionParameterStore::MakeUserVariable(Var);
+	}
+	Result.Append(UserParams);
+	return Result;
+}
+
+void FNiagaraStackGraphUtilities::SetLinkedValueHandleForFunctionInput(UEdGraphPin& OverridePin, FNiagaraParameterHandle LinkedParameterHandle, const TSet<FNiagaraVariable>& KnownParameters, ENiagaraDefaultMode DesiredDefaultMode, const FGuid& NewNodePersistentId)
 {
 	checkf(OverridePin.LinkedTo.Num() == 0, TEXT("Can't set a linked value handle when the override pin already has a value."));
 	const UNiagaraSettings* Settings = GetDefault<UNiagaraSettings>();
@@ -1396,8 +1414,7 @@ void FNiagaraStackGraphUtilities::SetLinkedValueHandleForFunctionInput(UEdGraphP
 		// see if the requested input exists or if we can use one of the equivalent loose types
 		FNiagaraTypeDefinition AlternateType = InputType == FNiagaraTypeDefinition::GetVec3Def() ? FNiagaraTypeDefinition::GetPositionDef() : FNiagaraTypeDefinition::GetVec3Def();
 		FNiagaraVariable ParameterAlternative(AlternateType, LinkedParameterHandle.GetParameterHandleString());
-		const TMap<FNiagaraVariable, FNiagaraGraphParameterReferenceCollection>& ReferenceMap = Graph->GetParameterReferenceMap();
-		if (ReferenceMap.Contains(ParameterToRead) == false && ReferenceMap.Contains(ParameterAlternative))
+		if (KnownParameters.Contains(ParameterToRead) == false && KnownParameters.Contains(ParameterAlternative))
 		{
 			ParameterToRead = ParameterAlternative;
 		}
@@ -2471,7 +2488,8 @@ void SetInputValue(
 	else if (Value.LinkedValue.IsSet())
 	{
 		UEdGraphPin& OverridePin = FNiagaraStackGraphUtilities::GetOrCreateStackFunctionInputOverridePin(InputFunctionCallNode, AliasedFunctionHandle, Value.Type);
-		FNiagaraStackGraphUtilities::SetLinkedValueHandleForFunctionInput(OverridePin, FNiagaraParameterHandle(Value.LinkedValue.GetValue()));
+		TSet<FNiagaraVariable> KnownParameters = FNiagaraStackGraphUtilities::GetParametersForContext(OverridePin.GetOwningNode()->GetGraph(), SystemViewModel->GetSystem());
+		FNiagaraStackGraphUtilities::SetLinkedValueHandleForFunctionInput(OverridePin, FNiagaraParameterHandle(Value.LinkedValue.GetValue()), KnownParameters);
 	}
 	else if (Value.DataValue.IsSet())
 	{
diff --git a/Engine/Plugins/FX/Niagara/Source/NiagaraEditor/Private/ViewModels/Stack/NiagaraStackParameterStoreEntry.cpp b/Engine/Plugins/FX/Niagara/Source/NiagaraEditor/Private/ViewModels/Stack/NiagaraStackParameterStoreEntry.cpp
--- a/Engine/Plugins/FX/Niagara/Source/NiagaraEditor/Private/ViewModels/Stack/NiagaraStackParameterStoreEntry.cpp
+++ b/Engine/Plugins/FX/Niagara/Source/NiagaraEditor/Private/ViewModels/Stack/NiagaraStackParameterStoreEntry.cpp
@@ -255,7 +255,8 @@ void UNiagaraStackParameterStoreEntry::OnRenamed(FText NewName)
 			TArray<TWeakObjectPtr<UNiagaraDataInterface>> RemovedDataObjects;
 			FNiagaraStackGraphUtilities::RemoveNodesForStackFunctionInputOverridePin(*LinkedPin, RemovedDataObjects); // no need to broadcast data objects modified here, the graph will recompile
 			// generate current link
-			FNiagaraStackGraphUtilities::SetLinkedValueHandleForFunctionInput(*LinkedPin, ParameterHandle);
+			TSet KnownParameters = { FNiagaraVariable(InputType, VariableName) };
+			FNiagaraStackGraphUtilities::SetLinkedValueHandleForFunctionInput(*LinkedPin, ParameterHandle, KnownParameters);
 		}
 
 		UObject* OwnerObj = Owner.Get();
diff --git a/Engine/Plugins/FX/Niagara/Source/NiagaraEditor/Public/ViewModels/Stack/NiagaraStackGraphUtilities.h b/Engine/Plugins/FX/Niagara/Source/NiagaraEditor/Public/ViewModels/Stack/NiagaraStackGraphUtilities.h
--- a/Engine/Plugins/FX/Niagara/Source/NiagaraEditor/Public/ViewModels/Stack/NiagaraStackGraphUtilities.h
+++ b/Engine/Plugins/FX/Niagara/Source/NiagaraEditor/Public/ViewModels/Stack/NiagaraStackGraphUtilities.h
@@ -130,7 +130,8 @@ namespace FNiagaraStackGraphUtilities
 
 	UEdGraphPin* GetLinkedValueHandleForFunctionInput(const UEdGraphPin& OverridePin);
 
-	void SetLinkedValueHandleForFunctionInput(UEdGraphPin& OverridePin, FNiagaraParameterHandle LinkedParameterHandle, ENiagaraDefaultMode DesiredDefaultMode = ENiagaraDefaultMode::FailIfPreviouslyNotSet, const FGuid& NewNodePersistentId = FGuid());
+	TSet<FNiagaraVariable> GetParametersForContext(UEdGraph* Graph, UNiagaraSystem& System);
+	void SetLinkedValueHandleForFunctionInput(UEdGraphPin& OverridePin, FNiagaraParameterHandle LinkedParameterHandle, const TSet<FNiagaraVariable>& KnownParameters, ENiagaraDefaultMode DesiredDefaultMode = ENiagaraDefaultMode::FailIfPreviouslyNotSet, const FGuid& NewNodePersistentId = FGuid());
 
 	void SetDataValueObjectForFunctionInput(UEdGraphPin& OverridePin, UClass* DataObjectType, FString InputNodeInputName, UNiagaraDataInterface*& OutDataObject, const FGuid& NewNodePersistentId = FGuid());
 
