Platform set changes

Improved UX For Device Tree and CVar Conditions
- Clean up Niagara Device Tree
   - new bIsVisibleForAsssets value in device profiles to have them indicate which device profiles should be visible in other assets.
- Improved tooltips and colouring on device tree to better show what is enabled/disabled and why.
- Added tooltips for CVar conditions to text entry and suggestion boxes showing the CVar help text.

PlatformSet Redirect feature =
- Can redirect usage of device profiles to other device profiles or CVar conditions.
- Allows fixup of content referncing legacy device profiles.
- Adding CVar fx.Niagara.LegacyDeviceProfiles to facilitate some legacy DP fixup.

Platforms now specifiy a min & max Niagara Quality Level they support which reduces need for trawling ini files and allows fragments to set quality level.

CVar Conditions now have more control over what occurs when the condition is passed or failed. Previously they could only disable the set when the condition is failed which is now the default.

Misc Fixes
- Preview DP is cleared before restoring CVar states so CVar callbacks have the correct preview DP value.
- Data Driven CVars are now updated correctly when default is modified in project settings.
- Console Manager now has a callback when CVars are unregistered.
- Utility function for calling a functor on all platform sets in a Niagara system.
- Tweaked callbacks for CVar changes so updates are defered and also limited to only refresh things that references changed CVars.

diff --git a/Engine/Config/Android/AndroidEngine.ini b/Engine/Config/Android/AndroidEngine.ini
--- a/Engine/Config/Android/AndroidEngine.ini
+++ b/Engine/Config/Android/AndroidEngine.ini
@@ -125,3 +125,8 @@ TextureFormatPrefix=""
 [SystemSettings]
 ;Perplatform to PerQualityLevel conversion mapping for platform groups
 QualityLevelMapping="low"
+
+;Manual range settings for Niagara quality levels. Simplifies handling of device profiles/fragments and what quality levels they can set.
+;Please ensure device profiles or fragments for this platform do not set fx.Niagara.QualityLevel values outside of this range
+fx.Niagara.QualityLevel.Min=0
+fx.Niagara.QualityLevel.Max=1
\ No newline at end of file
diff --git a/Engine/Config/BaseDeviceProfiles.ini b/Engine/Config/BaseDeviceProfiles.ini
--- a/Engine/Config/BaseDeviceProfiles.ini
+++ b/Engine/Config/BaseDeviceProfiles.ini
@@ -145,6 +145,7 @@ TextureLODGroups=(Group=TEXTUREGROUP_World,MinLODSize=1,MaxLODSize=16384,LODBias
 [Windows DeviceProfile]
 DeviceType=Windows
 BaseProfileName=
+bIsVisibleForAssets=True
 
 [WindowsEditor DeviceProfile]
 DeviceType=Windows
@@ -231,6 +232,7 @@ AppleTV=AppleTV4K
 [IOS DeviceProfile]
 DeviceType=IOS
 BaseProfileName=
+bIsVisibleForAssets=True
 ; HZB will be slower with tiled, and not needed
 +CVars=r.HZBOcclusion=0
 ; Non need for depth prepass
@@ -612,6 +614,8 @@ BaseProfileName=AppleTV
 [Android DeviceProfile]
 DeviceType=Android
 BaseProfileName=
+bIsVisibleForAssets=True
+
 +CVars=r.MobileContentScaleFactor=1.0
 +CVars=slate.AbsoluteIndices=1
 +CVars=r.Vulkan.DelayAcquireBackBuffer=2
@@ -884,6 +888,7 @@ BaseProfileName=Oculus_Quest
 [Mac DeviceProfile]
 DeviceType=Mac
 BaseProfileName=
+bIsVisibleForAssets=True
 +CVars=r.Shaders.ZeroInitialise=1
 +CVars=r.Shaders.BoundsChecking=1
 
@@ -902,6 +907,7 @@ BaseProfileName=Mac
 [Linux DeviceProfile]
 DeviceType=Linux
 BaseProfileName=
+bIsVisibleForAssets=True
 MeshLODSettings=
 TextureLODSettings=
 
diff --git a/Engine/Config/HoloLens/HoloLensEngine.ini b/Engine/Config/HoloLens/HoloLensEngine.ini
--- a/Engine/Config/HoloLens/HoloLensEngine.ini
+++ b/Engine/Config/HoloLens/HoloLensEngine.ini
@@ -15,6 +15,11 @@ r.setres=1280x720
 ;Perplatform to PerQualityLevel conversion mapping for platform groups
 QualityLevelMapping="low"
 
+;Manual range settings for Niagara quality levels. Simplifies handling of device profiles/fragments and what quality levels they can set.
+;Please ensure device profiles or fragments for this platform do not set fx.Niagara.QualityLevel values outside of this range
+fx.Niagara.QualityLevel.Min=0
+fx.Niagara.QualityLevel.Max=0
+
 ; Enable IPv6 by default on Hololens
 net.DisableIPv6=0
 
diff --git a/Engine/Config/IOS/IOSEngine.ini b/Engine/Config/IOS/IOSEngine.ini
--- a/Engine/Config/IOS/IOSEngine.ini
+++ b/Engine/Config/IOS/IOSEngine.ini
@@ -81,6 +81,11 @@ net.DisableIPv6=0
 ;Perplatform to PerQualityLevel conversion mapping for platform groups
 QualityLevelMapping="low"
 
+;Manual range settings for Niagara quality levels. Simplifies handling of device profiles/fragments and what quality levels they can set.
+;Please ensure device profiles or fragments for this platform do not set fx.Niagara.QualityLevel values outside of this range
+fx.Niagara.QualityLevel.Min=0
+fx.Niagara.QualityLevel.Max=1
+
 [/Script/Engine.RendererSettings]
 r.DistanceFields=0
 
diff --git a/Engine/Config/LinuxArm64/LinuxArm64Engine.ini b/Engine/Config/LinuxArm64/LinuxArm64Engine.ini
--- a/Engine/Config/LinuxArm64/LinuxArm64Engine.ini
+++ b/Engine/Config/LinuxArm64/LinuxArm64Engine.ini
@@ -7,6 +7,11 @@ r.setres=1280x720
 ;Perplatform to PerQualityLevel conversion mapping for platform groups
 QualityLevelMapping="low"
 
+;Manual range settings for Niagara quality levels. Simplifies handling of device profiles/fragments and what quality levels they can set.
+;Please ensure device profiles or fragments for this platform do not set fx.Niagara.QualityLevel values outside of this range
+fx.Niagara.QualityLevel.Min=0
+fx.Niagara.QualityLevel.Max=0
+
 [SystemSettingsEditor]
 ;r.RHICmdBypass=1
 
diff --git a/Engine/Plugins/FX/Niagara/Source/Niagara/Classes/NiagaraEffectType.h b/Engine/Plugins/FX/Niagara/Source/Niagara/Classes/NiagaraEffectType.h
--- a/Engine/Plugins/FX/Niagara/Source/Niagara/Classes/NiagaraEffectType.h
+++ b/Engine/Plugins/FX/Niagara/Source/Niagara/Classes/NiagaraEffectType.h
@@ -433,6 +433,9 @@ class NIAGARA_API UNiagaraEffectType : public UObject
 
 	void SpawnBaselineActor(UWorld* World);
 #endif
+	/** Performs the passed action for all FNiagaraPlatformSets in this system. */
+	template<typename TAction>
+	void ForEachPlatformSet(TAction Func);
 
 private:
 
@@ -465,3 +468,18 @@ public:
 	static void GeneratePerfBaselines();
 #endif
 };
+
+/** Performs the passed action for all FNiagaraPlatformSets in this FXType. */
+template<typename TAction>
+void UNiagaraEffectType::ForEachPlatformSet(TAction Func)
+{
+	for (FNiagaraSystemScalabilitySettings& Setting : SystemScalabilitySettings.Settings)
+	{
+		Func(this, Setting.Platforms);
+	}
+
+	for (FNiagaraEmitterScalabilitySettings& Setting : EmitterScalabilitySettings.Settings)
+	{
+		Func(this, Setting.Platforms);
+	}
+}
diff --git a/Engine/Plugins/FX/Niagara/Source/Niagara/Classes/NiagaraEmitter.h b/Engine/Plugins/FX/Niagara/Source/Niagara/Classes/NiagaraEmitter.h
--- a/Engine/Plugins/FX/Niagara/Source/Niagara/Classes/NiagaraEmitter.h
+++ b/Engine/Plugins/FX/Niagara/Source/Niagara/Classes/NiagaraEmitter.h
@@ -15,6 +15,9 @@
 #include "NiagaraRendererProperties.h"
 #include "NiagaraParameterDefinitionsBase.h"
 #include "NiagaraParameterDefinitionsSubscriber.h"
+
+#include "NiagaraDataInterfacePlatformSet.h"
+
 #include "NiagaraEmitter.generated.h"
 
 class UMaterial;
@@ -543,6 +546,9 @@ public:
 	template<typename TAction>
 	void ForEachScript(TAction Func) const;
 
+	template<typename TAction>
+	void ForEachPlatformSet(TAction Func);
+
 	NIAGARA_API void AddRenderer(UNiagaraRendererProperties* Renderer);
 
 	NIAGARA_API void RemoveRenderer(UNiagaraRendererProperties* Renderer);
@@ -783,3 +789,37 @@ void UNiagaraEmitter::ForEachScript(TAction Func) const
 		Func(EventScriptProps.Script);
 	}
 }
+
+template<typename TAction>
+void UNiagaraEmitter::ForEachPlatformSet(TAction Func)
+{
+	Func(this, Platforms);
+
+	for (FNiagaraEmitterScalabilityOverride& Override : ScalabilityOverrides.Overrides)
+	{
+		Func(this, Override.Platforms);
+	}
+
+	for (UNiagaraRendererProperties* Renderer : RendererProperties)
+	{
+		if(Renderer)
+		{
+			Renderer->ForEachPlatformSet(Func);
+		}
+	}
+
+	auto HandleScript = [Func](UNiagaraScript* NiagaraScript)
+	{
+		if (NiagaraScript)
+		{
+			for (const FNiagaraScriptDataInterfaceInfo& DataInterfaceInfo : NiagaraScript->GetCachedDefaultDataInterfaces())
+			{
+				if (UNiagaraDataInterfacePlatformSet* PlatformSetDI = Cast<UNiagaraDataInterfacePlatformSet>(DataInterfaceInfo.DataInterface))
+				{
+					Func(PlatformSetDI, PlatformSetDI->Platforms);
+				}
+			}
+		}
+	};
+	ForEachScript(HandleScript);
+}
diff --git a/Engine/Plugins/FX/Niagara/Source/Niagara/Classes/NiagaraPlatformSet.h b/Engine/Plugins/FX/Niagara/Source/Niagara/Classes/NiagaraPlatformSet.h
--- a/Engine/Plugins/FX/Niagara/Source/Niagara/Classes/NiagaraPlatformSet.h
+++ b/Engine/Plugins/FX/Niagara/Source/Niagara/Classes/NiagaraPlatformSet.h
@@ -99,19 +99,74 @@ struct FNiagaraPlatformSetConflictInfo
 	TArray<FNiagaraPlatformSetConflictEntry> Conflicts;
 };
 
+struct FNiagaraPlatformSetEnabledState
+{
+	bool bIsActive = true;
+	bool bCanBeActive = true;
+};
+
+struct FNiagaraPlatformSetEnabledStateDetails
+{
 #if WITH_EDITOR
-//Helper class for accesssing and caching the value of CVars for device profiles.
+	int32 DefaultQualityMask = INDEX_NONE;
+	int32 AvailableQualityMask = INDEX_NONE;
+
+	/** Provide reasons why this profile is not currently active. */
+	TArray<FText> ReasonsForInActive;
+	/** Provide reasons why this profile could never be active. */
+	TArray<FText> ReasonsForDisabled;
+#endif
+};
+
+#if WITH_EDITOR
+
+struct FNiagaraCVarValue
+{
+	union 
+	{
+		bool bValue;
+		int32 IntValue = 0;
+		float FloatValue;
+	}Data;
+	
+	template<typename T> FORCEINLINE T Get()const;	
+	template<typename T> FORCEINLINE void SetTyped(const FString& Str);
+	FORCEINLINE void Set(const FString& Str, IConsoleVariable* CVar);
+};
+
+template<> FORCEINLINE bool FNiagaraCVarValue::Get<bool>()const{ return Data.bValue; }
+template<> FORCEINLINE void FNiagaraCVarValue::SetTyped<bool>(const FString& Str){ LexFromString(Data.bValue, *Str); }
+
+template<> FORCEINLINE int32 FNiagaraCVarValue::Get<int32>()const { return Data.IntValue; }
+template<> FORCEINLINE void FNiagaraCVarValue::SetTyped<int32>(const FString& Str) { LexFromString(Data.IntValue, *Str); }
+
+template<> FORCEINLINE float FNiagaraCVarValue::Get<float>()const { return Data.FloatValue; }
+template<> FORCEINLINE void FNiagaraCVarValue::SetTyped<float>(const FString& Str) { LexFromString(Data.FloatValue, *Str); }
+
+FORCEINLINE void FNiagaraCVarValue::Set(const FString& Str, IConsoleVariable* CVar)
+{
+	if (CVar->IsVariableBool()) { SetTyped<bool>(Str); }
+	else if (CVar->IsVariableInt()) { SetTyped<int32>(Str); }
+	else if (CVar->IsVariableFloat()) { SetTyped<float>(Str); }
+}
+
+struct FNiagaraCVarValues
+{
+	FNiagaraCVarValue Default;
+
+	TArray<FNiagaraCVarValue> PerQualityLevelValues;
+};
+
+//Helper class for accessing and caching the value of CVars for device profiles.
 struct FDeviceProfileValueCache
 {
 	static void Empty();
 
-	template<typename T>
-	static bool GetValue(const UDeviceProfile* DeviceProfile, FName CVarName, T& OutValue);
+	static const FNiagaraCVarValues& GetValues(const UDeviceProfile* DeviceProfile, IConsoleVariable* CVar, FName CVarName);
 
 private:
-	static bool GetValueInternal(const UDeviceProfile* DeviceProfile, FName CVarName, FString& OutValue);
 
-	typedef TMap<FName, FString> FCVarValueMap;
+	typedef TMap<IConsoleVariable*, FNiagaraCVarValues> FCVarValueMap;
 	/** Cached values for any polled CVar on particular device profiles. */
 	static TMap<const UDeviceProfile*, FCVarValueMap> CachedDeviceProfileValues;
 	/** Cached values for any polled CVar from the platform ini files. */
@@ -119,6 +174,14 @@ private:
 };
 #endif
 
+UENUM()
+enum class ENiagaraCVarConditionResponse : uint8
+{
+	None,
+	Enable,
+	Disable,
+};
+
 /** Imposes a condition that a CVar must contain a set value or range of values for a platform set to be enabled. */
 USTRUCT()
 struct NIAGARA_API FNiagaraPlatformSetCVarCondition
@@ -127,11 +190,8 @@ struct NIAGARA_API FNiagaraPlatformSetCVarCondition
 	
 	FNiagaraPlatformSetCVarCondition();
 
-	/** Returns true if this is ever met for any device profile for the given platform. */
-	bool IsEnabledForPlatform(const FString& PlatformName)const;
-
-	/** Return true if this is met by the given device profile. */
-	bool IsEnabledForDeviceProfile(const UDeviceProfile* DeviceProfile, bool bCheckCurrentStateOnly)const;
+	/** Return true if this condition is met for the given device profile. */
+	bool Evaluate(const UDeviceProfile* DeviceProfile, bool& bOutCanEverPass, bool& bOutCanEverFail)const;
 
 	/** Returns the CVar for this condition. Can be null if the name give is not a valid CVar or the CVar is removed. */
 	IConsoleVariable* GetCVar()const;
@@ -142,6 +202,14 @@ struct NIAGARA_API FNiagaraPlatformSetCVarCondition
 	UPROPERTY(EditAnywhere, Category = "CVar")
 	FName CVarName = NAME_None;
 
+	/** If this CVar condition passes, how should this affect the state of the Platform Set? */
+	UPROPERTY(EditAnywhere, Category = "CVar", AdvancedDisplay)
+	ENiagaraCVarConditionResponse PassResponse = ENiagaraCVarConditionResponse::None;
+	
+	/** If this CVar condition fails, how should this affect the state of the Platform Set? */
+	UPROPERTY(EditAnywhere, Category = "CVar", AdvancedDisplay)
+	ENiagaraCVarConditionResponse FailResponse = ENiagaraCVarConditionResponse::Disable;
+
 	/** The value this CVar must contain for this platform set to be enabled. */
 	UPROPERTY(EditAnywhere, Category = "CVar", meta=(DisplayName="Required Value"))
 	bool Value = true;
@@ -179,7 +247,7 @@ struct NIAGARA_API FNiagaraPlatformSetCVarCondition
 	uint32 bUseMaxFloat : 1;
 
 	template<typename T>
-	FORCEINLINE bool IsEnabledForDeviceProfile_Internal(const UDeviceProfile* DeviceProfile, bool bCheckCurrentStateOnly)const;
+	bool Evaluate_Internal(const UDeviceProfile* DeviceProfile, bool& bOutCanEverPass, bool& bOutCanEverFail)const;
 
 	FORCEINLINE bool CheckValue(bool CVarValue)const { return CVarValue == Value; }
 	FORCEINLINE bool CheckValue(int32 CVarValue)const { return (!bUseMinInt || CVarValue >= MinInt) && (!bUseMaxInt || CVarValue <= MaxInt); }
@@ -188,20 +256,52 @@ struct NIAGARA_API FNiagaraPlatformSetCVarCondition
 	template<typename T>
 	FORCEINLINE T GetCVarValue(IConsoleVariable* CVar)const;
 
-	static void OnCVarChanged(IConsoleVariable* CVar);
 private:
 	mutable IConsoleVariable* CachedCVar = nullptr;
-
-	/** 
-	Callbacks for any CVars Niagara has looked at during this run.
-	*/
-	static TMap<FName, FDelegateHandle> CVarChangedDelegateHandles;
-	static FCriticalSection ChangedDelegateHandlesCritSec;
+	mutable uint32 LastCachedFrame = INDEX_NONE;
 };
 
 void NIAGARA_API SetGNiagaraQualityLevel(int32 QualityLevel);
 void NIAGARA_API SetGNiagaraDeviceProfile(UDeviceProfile* Profile);
 
+UENUM()
+enum class ENiagaraDeviceProfileRedirectMode : uint8
+{
+	/** Replace Device Profile Reference with a CVar Condition. */
+	CVar,
+	/** Replace Device Profile Reference with a different Device Profile. */
+	DeviceProfile,
+};
+
+/**
+Allows us to replace a specific device profile usage condition in all NiagaraPlatformSets.
+Helpful when dealing with changes to device profile structure.
+*/
+USTRUCT()
+struct FNiagaraPlatformSetRedirect
+{
+	GENERATED_BODY()
+
+	/** The names of any device profile entry that will apply this redirect. */
+	UPROPERTY(EditAnywhere, Category = Redirect)
+	TArray<FName> ProfileNames;
+
+	UPROPERTY(EditAnywhere, Category = Redirect)
+	ENiagaraDeviceProfileRedirectMode Mode;
+
+	/** When in Device Profile mode, the name of the device profile to redirect to. */
+	UPROPERTY(EditAnywhere, Category = Redirect, meta = (EditCondition = "Mode == ENiagaraDeviceProfileRedirectMode::DeviceProfile"))
+	FName RedirectProfileName;
+
+	/** When in CVar mode, the CVar condition to replace this device profile entry with when the profile entry is in the Enabled state. */
+	UPROPERTY(EditAnywhere, Category = Redirect, meta = (EditCondition = "Mode == ENiagaraDeviceProfileRedirectMode::CVar"))
+	FNiagaraPlatformSetCVarCondition CVarConditionEnabled;
+	
+	/** When in CVar mode, the CVar condition to replace this device profile entry with when the profile entry is in the Disabled state. */
+	UPROPERTY(EditAnywhere, Category = Redirect, meta = (EditCondition = "Mode == ENiagaraDeviceProfileRedirectMode::CVar"))
+	FNiagaraPlatformSetCVarCondition CVarConditionDisabled;
+};
+
 USTRUCT()
 struct NIAGARA_API FNiagaraPlatformSet
 {
@@ -214,6 +314,9 @@ struct NIAGARA_API FNiagaraPlatformSet
 
 	static void OnQualityLevelChanged(IConsoleVariable* Variable);
 	static int32 GetQualityLevel();
+	static int32 GetMinQualityLevel();
+	static int32 GetMaxQualityLevel();
+	static int32 GetAvailableQualityLevelMask();
 public:
 
 	//Runtime public API
@@ -226,9 +329,12 @@ public:
 	/** Is this set active right now. i.e. enabled for the current device profile and quality level. */
 	bool IsActive()const;
 	
-	/** Is this platform set enabled on any quality level for the passed device profile. Returns the QualityLevelMask for all enabled effects qualities for this profile. */
+	/** Is this platform set enabled for the give device profile. */
 	int32 IsEnabledForDeviceProfile(const UDeviceProfile* DeviceProfile)const;
 
+	/** Is this platform set enabled on any quality level for the passed device profile. Returns the QualityLevelMask for all enabled effects qualities for this profile. */
+	int32 GetEnabledMaskForDeviceProfile(const UDeviceProfile* DeviceProfile)const;
+
 	/** Is this platform set enabled at this quality level on any device profile. */
 	bool IsEnabledForQualityLevel(int32 QualityLevel)const;
 
@@ -244,7 +350,11 @@ public:
 	/**Will force all platform sets to regenerate their cached data next time they are used.*/
 	static void InvalidateCachedData();
 
-	static int32 GetEffectQualityMaskForDeviceProfile(const UDeviceProfile* Profile);
+	/** Returns the currenlty(or default) active quality mask for a profile. */
+	static int32 GetActiveQualityMaskForDeviceProfile(const UDeviceProfile* Profile);
+
+	/** Returns the mask of all available quality levels for a device profile. */
+	static int32 GetAvailableQualityMaskForDeviceProfile(const UDeviceProfile* Profile);
 
 	/** Returns true if the passed platform should prune emitters on cook. */
 	static bool ShouldPruneEmittersOnCook(const FString& PlatformName);
@@ -252,6 +362,17 @@ public:
 	/** Returns true if the passed platform can modify it's niagara scalability settings at runtime. */
 	static bool CanChangeScalabilityAtRuntime(const UDeviceProfile* DeviceProfile);
 
+	FNiagaraPlatformSetEnabledState IsEnabled(const UDeviceProfile* Profile, int32 QualityLevel, FNiagaraPlatformSetEnabledStateDetails* OutDetails=nullptr)const;
+
+	bool CanConsiderDeviceProfile(const UDeviceProfile* Profile)const;
+
+	bool ApplyRedirects();
+
+	static void RefreshScalability();
+	static void OnCVarChanged(IConsoleVariable* CVar);
+	static void OnCVarUnregistered(IConsoleVariable* CVar);
+	static IConsoleVariable* GetCVar(FName CVarName);
+
 	//Editor only public API
 #if WITH_EDITOR
 	/**
@@ -269,6 +390,7 @@ public:
 	/** Invalidates any cached data on this platform set when something has changed. */
 	void OnChanged();
 
+
 	/** Inspects the passed sets and generates an array of all conflicts between these sets. Used to keep arrays of platform sets orthogonal. */
 	static bool GatherConflicts(const TArray<const FNiagaraPlatformSet*>& PlatformSets, TArray<FNiagaraPlatformSetConflictInfo>& OutConflicts);
 	
@@ -290,47 +412,61 @@ public:
 	UPROPERTY(EditAnywhere, Category = Platforms)
 	TArray<FNiagaraPlatformSetCVarCondition> CVarConditions;
 
-private:
-
-	static int32 CachedQualityLevel;
-
-	//Set from outside when we need to force all cached values to be regenerated. For example on CVar changes.
-	static uint32 LastDirtiedFrame;
-
-	//Last frame we built our cached data. 
-	mutable uint32 LastBuiltFrame;
-
-	mutable bool bEnabledForCurrentProfileAndEffectQuality;
-
-	bool IsEnabled(const UDeviceProfile* Profile, int32 QualityLevel, bool bConsiderCurrentStateOnly=false)const;
-
 #if WITH_EDITOR
 	//Data we pull from platform ini files.
 	struct FPlatformIniSettings
 	{
-		FPlatformIniSettings()
-			:bCanChangeScalabilitySettingsAtRuntime(0), bPruneEmittersOnCook(false), EffectsQuality(0)
-		{}
-		FPlatformIniSettings(int32 InbCanChangeScalabilitySettingsAtRuntime, int32 InbPruneEmittersOnCook, int32 InEffectsQuality)
-			:bCanChangeScalabilitySettingsAtRuntime(InbCanChangeScalabilitySettingsAtRuntime), bPruneEmittersOnCook(InbPruneEmittersOnCook), EffectsQuality(InEffectsQuality)
-		{}
+		FPlatformIniSettings();
+		FPlatformIniSettings(int32 InbCanChangeScalabilitySettingsAtRuntime, int32 InbPruneEmittersOnCook, int32 InEffectsQuality, int32 InMinQualityLevel, int32 InMaxQualityLevel);
 		int32 bCanChangeScalabilitySettingsAtRuntime;
 		int32 bPruneEmittersOnCook;
 		int32 EffectsQuality;
+		int32 MinQualityLevel;
+		int32 MaxQualityLevel;
+		int32 QualityLevelMask;
 
 		TArray<int32> QualityLevelsPerEffectsQuality;
 	};
+	static FPlatformIniSettings& GetPlatformIniSettings(const FString& PlatformName);
 
+	static int32 GetEffectQualityMaskForPlatform(const FString& PlatformName);
+#endif
+
+	static uint32 GetLastDirtiedFrame(){ return LastDirtiedFrame; }
+
+private:
+
+	mutable uint32 bEnabledForCurrentProfileAndEffectQuality : 1;
+
+	//Last frame we built our cached data. 
+	mutable uint32 LastBuiltFrame;
+
+	//Set from outside when we need to force all cached values to be regenerated. For example on CVar changes.
+	static uint32 LastDirtiedFrame;
+
+	static int32 CachedQualityLevel;
+	static int32 CachedAvailableQualityLevelMask;
+
+#if WITH_EDITOR
 	//Cached data read from platform ini files.
 	static TMap<FName, FPlatformIniSettings> CachedPlatformIniSettings;
 
 	//Cached final QualityLevel setting for each device profile.
 	static TMap<const UDeviceProfile*, int32> CachedQLMasksPerDeviceProfile;
+#endif
 
-	static FPlatformIniSettings& GetPlatformIniSettings(const FString& PlatformName);
+	/**
+	Callbacks for any CVars Niagara has looked at during this run.
+	*/
+	struct FCachedCVarInfo
+	{
+		IConsoleVariable* CVar = nullptr;
+		FDelegateHandle ChangedHandle;
+	};
+	static TMap<FName, FCachedCVarInfo> CachedCVarInfo;
+	static FCriticalSection CachedCVarInfoCritSec;
 
-	static int32 GetEffectQualityMaskForPlatform(const FString& PlatformName);
-#endif
+	static TArray<FName> ChangedCVars;
 };
 
 
diff --git a/Engine/Plugins/FX/Niagara/Source/Niagara/Classes/NiagaraSystem.h b/Engine/Plugins/FX/Niagara/Source/Niagara/Classes/NiagaraSystem.h
--- a/Engine/Plugins/FX/Niagara/Source/Niagara/Classes/NiagaraSystem.h
+++ b/Engine/Plugins/FX/Niagara/Source/Niagara/Classes/NiagaraSystem.h
@@ -20,6 +20,9 @@
 #include "UObject/Object.h"
 #include "UObject/ObjectMacros.h"
 
+#include "NiagaraDataInterfacePlatformSet.h"
+
+
 #include "NiagaraSystem.generated.h"
 
 #if WITH_EDITORONLY_DATA
@@ -347,6 +350,8 @@ public:
 	/** Gets an array of the emitter handles. */
 	const TArray<FNiagaraEmitterHandle>& GetEmitterHandles();
 	const TArray<FNiagaraEmitterHandle>& GetEmitterHandles()const;
+	
+	FNiagaraSystemScalabilityOverrides& GetScalabilityOverrides(){return SystemScalabilityOverrides; }
 
 private:
 	bool IsValidInternal() const;
@@ -416,6 +421,10 @@ public:
 	template<typename TAction>
 	void ForEachScript(TAction Func) const;
 
+	/** Performs the passed action for all FNiagaraPlatformSets used by this system. Some may not be owned by this system. */
+	template<typename TAction>
+	void ForEachPlatformSet(TAction Func);
+
 	bool AllowScalabilityForLocalPlayerFX()const;
 
 private:
@@ -823,11 +832,11 @@ protected:
 	bool bOverrideScalabilitySettings;
 
 	/** Controls whether we should override the Effect Type value for bAllowCullingForLocalPlayers. */
-	UPROPERTY(EditAnywhere, Category = "Override", meta = (InlineEditConditionToggle, EditCondition = bOverrideScalabilitySettings))
+	UPROPERTY(EditAnywhere, Category = "Scalability", meta = (InlineEditConditionToggle, EditCondition = bOverrideScalabilitySettings))
 	uint32 bOverrideAllowCullingForLocalPlayers : 1;
 	
 	/** The override value for bAllowCullingForLocalPlayers from the Effect Type. */
-	UPROPERTY(EditAnywhere, Category = "Override", meta = (EditCondition = bOverrideAllowCullingForLocalPlayers))
+	UPROPERTY(EditAnywhere, Category = "Scalability", meta = (EditCondition = bOverrideAllowCullingForLocalPlayers))
 	uint32 bAllowCullingForLocalPlayersOverride : 1;
 
 	UPROPERTY()
@@ -1052,10 +1061,10 @@ FORCEINLINE void UNiagaraSystem::UnregisterActiveInstance()
 
 template<typename TAction>
 void UNiagaraSystem::ForEachScript(TAction Func) const
-{	
+{
 	Func(SystemSpawnScript);
 	Func(SystemUpdateScript);
-			
+
 	for (const FNiagaraEmitterHandle& Handle : EmitterHandles)
 	{
 		if (UNiagaraEmitter* Emitter = Handle.GetInstance())
@@ -1064,3 +1073,49 @@ void UNiagaraSystem::ForEachScript(TAction Func) const
 		}
 	}
 }
+
+/** Performs the passed action for all FNiagaraPlatformSets in this system. */
+template<typename TAction>
+void UNiagaraSystem::ForEachPlatformSet(TAction Func)
+{
+	//Handle our scalability overrides
+	for (FNiagaraSystemScalabilityOverride& Override : SystemScalabilityOverrides.Overrides)
+	{
+		Func(this, Override.Platforms);
+	}
+
+	//Handle and platform set User DIs.
+	for (UNiagaraDataInterface* DI : GetExposedParameters().GetDataInterfaces())
+	{
+		if (UNiagaraDataInterfacePlatformSet* PlatformSetDI = Cast<UNiagaraDataInterfacePlatformSet>(DI))
+		{
+			Func(PlatformSetDI, PlatformSetDI->Platforms);
+		}
+	}
+
+	//Handle all platform set DIs held in scripts for this system.
+	auto HandleScript = [Func](UNiagaraScript* NiagaraScript)
+	{
+		if (NiagaraScript)
+		{
+			for (const FNiagaraScriptDataInterfaceInfo& DataInterfaceInfo : NiagaraScript->GetCachedDefaultDataInterfaces())
+			{
+				if (UNiagaraDataInterfacePlatformSet* PlatformSetDI = Cast<UNiagaraDataInterfacePlatformSet>(DataInterfaceInfo.DataInterface))
+				{
+					Func(PlatformSetDI, PlatformSetDI->Platforms);
+				}
+			}
+		}
+	};
+	HandleScript(SystemSpawnScript);
+	HandleScript(SystemUpdateScript);
+
+	//Finally handle all our emitters.
+	for (FNiagaraEmitterHandle& Handle : EmitterHandles)
+	{
+		if (UNiagaraEmitter* Emitter = Handle.GetInstance())
+		{
+			Emitter->ForEachPlatformSet(Func);
+		}
+	}
+}
\ No newline at end of file
diff --git a/Engine/Plugins/FX/Niagara/Source/Niagara/Private/NiagaraEffectType.cpp b/Engine/Plugins/FX/Niagara/Source/Niagara/Private/NiagaraEffectType.cpp
--- a/Engine/Plugins/FX/Niagara/Source/Niagara/Private/NiagaraEffectType.cpp
+++ b/Engine/Plugins/FX/Niagara/Source/Niagara/Private/NiagaraEffectType.cpp
@@ -58,6 +58,14 @@ void UNiagaraEffectType::PostLoad()
 		}
 	}
 
+	//Apply platform set redirectors
+	auto ApplyPlatformSetRedirects = [](UObject* Owner, FNiagaraPlatformSet& Platforms)
+	{
+		Platforms.ApplyRedirects();
+	};
+	ForEachPlatformSet(ApplyPlatformSetRedirects);
+
+
 #if !WITH_EDITOR && NIAGARA_PERF_BASELINES
 	//When not in the editor we clear out the baseline so that it's regenerated for play tests.
 	//We cannot use the saved editor/development config settings.
diff --git a/Engine/Plugins/FX/Niagara/Source/Niagara/Private/NiagaraModule.cpp b/Engine/Plugins/FX/Niagara/Source/Niagara/Private/NiagaraModule.cpp
--- a/Engine/Plugins/FX/Niagara/Source/Niagara/Private/NiagaraModule.cpp
+++ b/Engine/Plugins/FX/Niagara/Source/Niagara/Private/NiagaraModule.cpp
@@ -356,6 +356,8 @@ void INiagaraModule::StartupModule()
 	FCoreDelegates::OnPreExit.AddRaw(this, &INiagaraModule::OnPreExit);
 	FWorldDelegates::OnWorldBeginTearDown.AddRaw(this, &INiagaraModule::OnWorldBeginTearDown);
 	FWorldDelegates::OnWorldTickStart.AddRaw(this, &INiagaraModule::OnWorldTickStart);
+
+	FCoreDelegates::OnBeginFrame.AddRaw(this, &INiagaraModule::OnBeginFrame);
 }
 
 void INiagaraModule::OnPostEngineInit()
@@ -366,6 +368,14 @@ void INiagaraModule::OnPostEngineInit()
 		BaselineHandler = MakeUnique<FNiagaraPerfBaselineHandler>();
 	}
 #endif
+
+	IConsoleManager& CMan = IConsoleManager::Get();
+	OnCVarUnregisteredHandle = CMan.OnCVarUnregistered().AddLambda(
+		[](IConsoleVariable* CVar)
+		{
+			FNiagaraPlatformSet::OnCVarUnregistered(CVar);
+		}
+	);
 }
 
 void INiagaraModule::OnPreExit()
@@ -380,6 +390,9 @@ void INiagaraModule::OnPreExit()
 #if WITH_NIAGARA_DEBUGGER
 	DebuggerClient.Reset();
 #endif
+
+	IConsoleManager& CMan = IConsoleManager::Get();
+	CMan.OnCVarUnregistered().Remove(OnCVarUnregisteredHandle);
 }
 
 #if WITH_EDITOR
@@ -412,6 +425,11 @@ void INiagaraModule::OnWorldTickStart(UWorld* World, ELevelTick TickType, float
 #endif
 }
 
+void INiagaraModule::OnBeginFrame()
+{
+	FNiagaraPlatformSet::RefreshScalability();
+}
+
 void INiagaraModule::OnWorldBeginTearDown(UWorld* World)
 {
 #if NIAGARA_PERF_BASELINES
@@ -433,6 +451,7 @@ void INiagaraModule::ShutdownModule()
 {
 	FWorldDelegates::OnWorldBeginTearDown.RemoveAll(this);
 	FWorldDelegates::OnWorldTickStart.RemoveAll(this);
+	FCoreDelegates::OnBeginFrame.RemoveAll(this);
 
 	FNiagaraWorldManager::OnShutdown();
 
diff --git a/Engine/Plugins/FX/Niagara/Source/Niagara/Private/NiagaraPlatformSet.cpp b/Engine/Plugins/FX/Niagara/Source/Niagara/Private/NiagaraPlatformSet.cpp
--- a/Engine/Plugins/FX/Niagara/Source/Niagara/Private/NiagaraPlatformSet.cpp
+++ b/Engine/Plugins/FX/Niagara/Source/Niagara/Private/NiagaraPlatformSet.cpp
@@ -18,6 +18,13 @@
 
 #define LOCTEXT_NAMESPACE "NiagaraPlatformSet"
 
+//Array of CVars that have been changed this frame.
+//We do a deferred refresh of systems that use these CVars for scalability.
+TArray<FName> FNiagaraPlatformSet::ChangedCVars;
+
+TMap<FName, FNiagaraPlatformSet::FCachedCVarInfo> FNiagaraPlatformSet::CachedCVarInfo;
+FCriticalSection FNiagaraPlatformSet::CachedCVarInfoCritSec;
+
 /**
 Whether a platform can change it's scalability settings at runtime.
 Defaults to false for all platforms and is explicitly enabled for desktop platforms.
@@ -27,7 +34,7 @@ int32 GbNiagaraAllowRuntimeScalabilityChanges = 0;
 static FAutoConsoleVariableRef CVarNiagaraAllowRuntimeScalabilityChanges(
 	CanChangeEQCVarName,
 	GbNiagaraAllowRuntimeScalabilityChanges,
-	TEXT("If > 0 this platform allows niagara scalability level changes at runtime. \n"),
+	TEXT("If > 0 this platform allows niagara scalability settings changes at runtime. \n"),
 	ECVF_Scalability
 );
 
@@ -78,6 +85,49 @@ static FAutoConsoleVariableRef CVarNiagaraQualityLevel(
 	ECVF_Scalability
 );
 
+/** Minimum quality level available for a platform. Defaults to -1, indicating there is no minimum. */
+static const int32 DefaultMinQualityLevel = INDEX_NONE;
+const TCHAR* NiagaraMinQualityLevelName = TEXT("fx.Niagara.QualityLevel.Min");
+int32 GNiagaraMinQualityLevel = DefaultMinQualityLevel;
+static FAutoConsoleVariableRef CVarNiagaraMinQualityLevel(
+	NiagaraMinQualityLevelName,
+	GNiagaraMinQualityLevel,
+	TEXT("The minimum quality level for Niagara Effects. \n"),
+	ECVF_ReadOnly
+);
+
+/** Maximum quality level available for a platform. Defaults to -1, indicating there is no maximum. */
+static const int32 DefaultMaxQualityLevel = INDEX_NONE;
+const TCHAR* NiagaraMaxQualityLevelName = TEXT("fx.Niagara.QualityLevel.Max");
+int32 GNiagaraMaxQualityLevel = DefaultMaxQualityLevel;
+static FAutoConsoleVariableRef CVarNiagaraMaxQualityLevel(
+	NiagaraMaxQualityLevelName,
+	GNiagaraMaxQualityLevel,
+	TEXT("The Maximum quality level for Niagara Effects. \n"),
+	ECVF_ReadOnly
+);
+
+int32 GbAllowAllDeviceProfiles = 0;
+static FAutoConsoleVariableRef CVarAllowAllDeviceProfiles(
+	TEXT("fx.Niagara.AllowAllDeviceProfiles"),
+	GbAllowAllDeviceProfiles,
+	TEXT(" \n"),
+	ECVF_Default
+);
+
+/** 
+This is a special case CVar that allows us to use CVar conditions to maintain behavior with legacy device profiles.
+Do not use directly for new content.
+Legacy device profiles can be given a specific value for this CVar and then CVar conditions used to enable/disable as appropriate to match with legacy assets with enabled/disabled content based on them.
+*/
+int32 GNiagaraLegacyDeviceProfile = INDEX_NONE;
+static FAutoConsoleVariableRef CVarNiagaraLegacyDeviceProfile(
+	TEXT("fx.Niagara.LegacyDeviceProfile"),
+	GNiagaraLegacyDeviceProfile,
+	TEXT("This is a special case CVar that allows us to use CVar conditions to maintain behavior with legacy device profiles.\nDo not use directly for new content.\nLegacy device profiles can be given a specific value for this CVarand then CVar conditions used to enable / disable as appropriate to match with legacy assets with enabled / disabled content based on them. \n"),
+	ECVF_Scalability
+);
+
 // Override platform device profile
 // In editor all profiles will be available
 // On cooked builds only the profiles for that cooked platform will be available
@@ -145,6 +195,7 @@ static UDeviceProfile* NiagaraGetActiveDeviceProfile()
 }
 
 int32 FNiagaraPlatformSet::CachedQualityLevel = INDEX_NONE;
+int32 FNiagaraPlatformSet::CachedAvailableQualityLevelMask = INDEX_NONE;
 
 int32 FNiagaraPlatformSet::GetQualityLevel()
 {
@@ -155,6 +206,38 @@ int32 FNiagaraPlatformSet::GetQualityLevel()
 	return CachedQualityLevel;
 }
 
+int32 FNiagaraPlatformSet::GetMinQualityLevel()
+{
+	return GNiagaraMinQualityLevel;
+}
+
+int32 FNiagaraPlatformSet::GetMaxQualityLevel()
+{
+	return GNiagaraMaxQualityLevel;
+}
+
+int32 FNiagaraPlatformSet::GetAvailableQualityLevelMask()
+{
+	if (CachedAvailableQualityLevelMask == INDEX_NONE)
+	{
+		if (CanChangeScalabilityAtRuntime())
+		{
+			CachedAvailableQualityLevelMask = INDEX_NONE;
+		}
+		else
+		{
+			CachedAvailableQualityLevelMask = 0;
+			int32 Min = GetMinQualityLevel();
+			int32 Max = GetMaxQualityLevel();
+			for (int32 QL = Min; QL <= Max; ++QL)
+			{
+				CachedAvailableQualityLevelMask |= (1 << QL);
+			}
+		}
+	}
+	return CachedAvailableQualityLevelMask;
+}
+
 void FNiagaraPlatformSet::OnQualityLevelChanged(IConsoleVariable* Variable)
 {
 	int32 NewQualityLevel = Variable->GetInt();
@@ -205,16 +288,16 @@ FText FNiagaraPlatformSet::GetQualityLevelMaskText(int32 QualityLevelMask)
 
 FNiagaraPlatformSet::FNiagaraPlatformSet(int32 QLMask)
 	: QualityLevelMask(QLMask)
-	, LastBuiltFrame(0)
 	, bEnabledForCurrentProfileAndEffectQuality(false)
+	, LastBuiltFrame(0)
 {
 	IsActive();
 }
 
 FNiagaraPlatformSet::FNiagaraPlatformSet()
 	: QualityLevelMask(INDEX_NONE)
-	, LastBuiltFrame(0)
 	, bEnabledForCurrentProfileAndEffectQuality(false)
+	, LastBuiltFrame(0)
 {
 }
 
@@ -243,12 +326,31 @@ bool FNiagaraPlatformSet::IsActive()const
 			}
 		}
 #endif
-		bEnabledForCurrentProfileAndEffectQuality = IsEnabled(ActiveProfile, QualityLevel, true);
+		FNiagaraPlatformSetEnabledState EnabledState = IsEnabled(ActiveProfile, QualityLevel);
+		bEnabledForCurrentProfileAndEffectQuality = EnabledState.bIsActive;
 		LastBuiltFrame = GFrameNumber;
 	}
 	return bEnabledForCurrentProfileAndEffectQuality;
 }
 
+int32 FNiagaraPlatformSet::GetEnabledMaskForDeviceProfile(const UDeviceProfile* DeviceProfile)const
+{
+	const UNiagaraSettings* Settings = GetDefault<UNiagaraSettings>();
+	check(Settings);
+
+	int32 RetQLMask = 0;
+	for (int32 i = 0; i < Settings->QualityLevels.Num(); ++i)
+	{
+		FNiagaraPlatformSetEnabledState EnabledState = IsEnabled(DeviceProfile, i);
+		if (EnabledState.bCanBeActive)
+		{
+			RetQLMask |= CreateQualityLevelMask(i);
+		}
+	}
+
+	return RetQLMask;
+}
+
 int32 FNiagaraPlatformSet::IsEnabledForDeviceProfile(const UDeviceProfile* DeviceProfile)const
 {
 	const UNiagaraSettings* Settings = GetDefault<UNiagaraSettings>();
@@ -257,7 +359,8 @@ int32 FNiagaraPlatformSet::IsEnabledForDeviceProfile(const UDeviceProfile* Devic
 	int32 RetQLMask = 0;
 	for (int32 i = 0; i < Settings->QualityLevels.Num(); ++i)
 	{
-		if (IsEnabled(DeviceProfile, i))
+		FNiagaraPlatformSetEnabledState Enabled = IsEnabled(DeviceProfile, i);
+		if (Enabled.bCanBeActive)
 		{
 			RetQLMask |= CreateQualityLevelMask(i);
 		}
@@ -272,7 +375,8 @@ bool FNiagaraPlatformSet::IsEnabledForQualityLevel(int32 QualityLevel)const
 	{
 		if (UDeviceProfile* Profile = Cast<UDeviceProfile>(DeviceProfileObj))
 		{
-			if (IsEnabled(Profile, QualityLevel))
+			FNiagaraPlatformSetEnabledState Enabled = IsEnabled(Profile, QualityLevel);
+			if (Enabled.bCanBeActive)
 			{
 				return true;
 			}
@@ -290,19 +394,22 @@ void FNiagaraPlatformSet::GetOverridenDeviceProfiles(int32 QualityLevel, TArray<
 		if (TObjectPtr<UDeviceProfile>* DeviceProfile = UDeviceProfileManager::Get().Profiles.FindByPredicate([&](UObject* CheckProfile) {return CheckProfile->GetFName() == Entry.ProfileName;}))
 		{
 			UDeviceProfile* Profile = *DeviceProfile;
-			//If this platform cannot change at runtime then we store all EQs in the state so that the device is still overridden if someone changes it's EQ CVar.
-			//So here we must also check that this QualityLevel is the right one for the platforms current setting.
-			int32 ProfileQLMask = GetEffectQualityMaskForDeviceProfile(Profile);
-			if (ProfileQLMask == INDEX_NONE || (QLMask & ProfileQLMask) != 0)
+			if (CanConsiderDeviceProfile(Profile))
 			{
-				ENiagaraPlatformSelectionState State = Entry.GetState(QualityLevel);
-				if (State == ENiagaraPlatformSelectionState::Enabled)
-				{
-					OutEnabledProfiles.Add(Profile);
-				}
-				else if (State == ENiagaraPlatformSelectionState::Disabled)
+				//If this platform cannot change at runtime then we store all EQs in the state so that the device is still overridden if someone changes it's EQ CVar.
+				//So here we must also check that this QualityLevel is the right one for the platforms current setting.
+				int32 ProfileQLMask = GetAvailableQualityMaskForDeviceProfile(Profile);
+				if (ProfileQLMask == INDEX_NONE || (QLMask & ProfileQLMask) != 0)
 				{
-					OutDisabledProfiles.Add(Profile);
+					ENiagaraPlatformSelectionState State = Entry.GetState(QualityLevel);
+					if (State == ENiagaraPlatformSelectionState::Enabled)
+					{
+						OutEnabledProfiles.Add(Profile);
+					}
+					else if (State == ENiagaraPlatformSelectionState::Disabled)
+					{
+						OutDisabledProfiles.Add(Profile);
+					}
 				}
 			}
 		}
@@ -315,21 +422,19 @@ bool FNiagaraPlatformSet::CanChangeScalabilityAtRuntime()
 	return GbNiagaraAllowRuntimeScalabilityChanges != 0;
 }
 
-int32 FNiagaraPlatformSet::GetEffectQualityMaskForDeviceProfile(const UDeviceProfile* Profile)
+/** Returns the currenlty(or default) active quality mask for a profile. */
+int32 FNiagaraPlatformSet::GetActiveQualityMaskForDeviceProfile(const UDeviceProfile* Profile)
 {
 #if WITH_EDITOR
 
-	//When in the editor we may be asking for the EQ of a platform other than the current one.
-	//So we have to look in ini files and device profiles to find it.
-	if (int32* CachedQLMask = CachedQLMasksPerDeviceProfile.Find(Profile))
+	if (int32* CachedQLForProfile = CachedQLMasksPerDeviceProfile.Find(Profile))
 	{
-		//We've seen this Profile before so return the cached value;
-		return *CachedQLMask;
+		return *CachedQLForProfile;
 	}
-
-	int32 QLMask = INDEX_NONE;
-	FPlatformIniSettings& PlatformSettings = GetPlatformIniSettings(Profile->DeviceType);
 	
+	FPlatformIniSettings& PlatformSettings = GetPlatformIniSettings(Profile->DeviceType);
+	int32 QLMask = PlatformSettings.QualityLevelMask;
+
 	if (PlatformSettings.bCanChangeScalabilitySettingsAtRuntime)
 	{
 		QLMask = INDEX_NONE;
@@ -372,53 +477,509 @@ int32 FNiagaraPlatformSet::GetEffectQualityMaskForDeviceProfile(const UDevicePro
 #endif
 }
 
-bool FNiagaraPlatformSet::IsEnabled(const UDeviceProfile* Profile, int32 QualityLevel, bool bConsiderCurrentStateOnly)const
+/** Returns the mask of all available quality levels for a device profile. */
+int32 FNiagaraPlatformSet::GetAvailableQualityMaskForDeviceProfile(const UDeviceProfile* Profile)
+{
+#if WITH_EDITOR
+	
+	//We can no longer rely on trawling the device profiles and ini files for all the quality levels a platform can be at
+	//So we rely on manual min/max values from the inis
+	FPlatformIniSettings& PlatformSettings = GetPlatformIniSettings(Profile->DeviceType);
+	return PlatformSettings.QualityLevelMask;
+
+#else
+	
+	//In non editor builds we can assume we mean the current platform.
+	return GetAvailableQualityLevelMask();
+
+#endif
+}
+
+FNiagaraPlatformSetEnabledState FNiagaraPlatformSet::IsEnabled(const UDeviceProfile* Profile, int32 QualityLevel, FNiagaraPlatformSetEnabledStateDetails* OutDetails)const
 {
 	checkSlow(Profile);
 
-	//Check CVar conditions first.
-	//Only apply CVars if we're checking current state or we're not on a platform that can change scalability CVars at runtime.
-	if (bConsiderCurrentStateOnly || CanChangeScalabilityAtRuntime(Profile) == false)
+	FNiagaraPlatformSetEnabledState Ret;
+
+	bool bCanChangeScalabilityAtRuntime = CanChangeScalabilityAtRuntime(Profile);
+
+	//First get the base settings from the Quality Levels.
+	int32 ActiveProfileMask = GetActiveQualityMaskForDeviceProfile(Profile);
+	int32 ActiveProfileQuality = QualityLevelFromMask(ActiveProfileMask);
+	int32 AvailableProfileMask = GetAvailableQualityMaskForDeviceProfile(Profile);
+	int32 CurrentQLMask = CreateQualityLevelMask(QualityLevel);
+
+	/** Is the current QL enabled for the platform set? */
+	bool bQLIsEnabled = CurrentQLMask & QualityLevelMask;
+
+	/** Is this QL the active or default for this profile. */
+	bool bProfileActiveQL = (CurrentQLMask & ActiveProfileMask) != 0;
+
+	/** Is this QL available for this profile? */
+	bool bProfileAvailableQL = (CurrentQLMask & AvailableProfileMask) != 0;
+
+#if WITH_EDITOR
+	if (OutDetails)
 	{
-		for (const FNiagaraPlatformSetCVarCondition& CVarCondition : CVarConditions)
+		FPlatformIniSettings& PlatformSettings = GetPlatformIniSettings(Profile->DeviceType);
+
+		OutDetails->DefaultQualityMask = ActiveProfileMask;
+		OutDetails->AvailableQualityMask = AvailableProfileMask;
+	}
+
+	//Does this platform set match the passed in current quality level?
+	FText QLText = GetQualityLevelText(QualityLevel);
+	FText ActiveQLText = GetQualityLevelText(ActiveProfileQuality);
+#endif
+
+	if (bProfileAvailableQL)
+	{
+		if (bQLIsEnabled)//Is the whole quality level disabled?
 		{
-			//Bail if any cvar condition isn't met.
-			if (CVarCondition.IsEnabledForDeviceProfile(Profile, bConsiderCurrentStateOnly) == false)
+			if (bProfileActiveQL)//The quality level is enabled but this profile is not active for this QL
 			{
-				return false;
+				Ret.bIsActive = true;
+				Ret.bCanBeActive = true;
+			}
+			else
+			{
+				Ret.bIsActive = false;
+				Ret.bCanBeActive = true;
+
+#if WITH_EDITOR
+				if (OutDetails)
+				{
+					OutDetails->ReasonsForInActive.Emplace(FText::Format(LOCTEXT("QualityLevelNotActiveReasonFmt", "{0} is enabled at {1} but defaults to {2}."), FText::FromString(Profile->GetName()), QLText, ActiveQLText));
+				}
+#endif
 			}
 		}
-	}
+		else
+		{
+			Ret.bIsActive = false;
+			Ret.bCanBeActive = false;
 
-	//Does this platform set match the passed in current quality level?
-	int32 TestQLMask = CreateQualityLevelMask(QualityLevel);
-	bool bIsEnabledByEQ = (QualityLevelMask & TestQLMask) != 0;
+#if WITH_EDITOR
+			if (OutDetails)
+			{
+				OutDetails->ReasonsForDisabled.Emplace(FText::Format(LOCTEXT("QualityLevelDisabledReasonFmt", "{0} Quality is disabled."), QLText));
+			}
+#endif
+		}
+	}
+	else
+	{
+		Ret.bIsActive = false;
+		Ret.bCanBeActive = false;
 
-	//Does it match the device profile's quality level and do we care?
-	int32 ProfileQLMask = GetEffectQualityMaskForDeviceProfile(Profile);
-	bIsEnabledByEQ &= bConsiderCurrentStateOnly || ((ProfileQLMask & QualityLevelMask) != 0);
+#if WITH_EDITOR
+		if (OutDetails)
+		{
+			FPlatformIniSettings& PlatformSettings = GetPlatformIniSettings(Profile->DeviceType);
+			OutDetails->ReasonsForDisabled.Emplace(FText::Format(LOCTEXT("PlatformSetDisabledFmt", "{1} is outside the Min/Max quality for {0}."), QLText, FText::FromString(Profile->GetName())));
+		}
+#endif
+	}
 
+	//Next see if this device profile is specifically overridden.
 	if (DeviceProfileStates.Num() > 0)
 	{
 		//Walk up the parent hierarchy to see if we have an explicit state for this profile set.
 		const UDeviceProfile* CurrProfile = Profile;
 		while (CurrProfile)
 		{
-			if (const FNiagaraDeviceProfileStateEntry* StateEntry = DeviceProfileStates.FindByPredicate([&](const FNiagaraDeviceProfileStateEntry& ProfileState) {return ProfileState.ProfileName == CurrProfile->GetFName(); }))
+			//if (CanConsiderDeviceProfile(CurrProfile))
 			{
-				ENiagaraPlatformSelectionState SelectionState = StateEntry->GetState(QualityLevel);
-				if (SelectionState != ENiagaraPlatformSelectionState::Default)
+				if (const FNiagaraDeviceProfileStateEntry* StateEntry = DeviceProfileStates.FindByPredicate([&](const FNiagaraDeviceProfileStateEntry& ProfileState) {return ProfileState.ProfileName == CurrProfile->GetFName(); }))
 				{
-					return SelectionState == ENiagaraPlatformSelectionState::Enabled;
+					ENiagaraPlatformSelectionState SelectionState = StateEntry->GetState(QualityLevel);
+					if (SelectionState != ENiagaraPlatformSelectionState::Default)
+					{
+						bool bProfileEnabled = SelectionState == ENiagaraPlatformSelectionState::Enabled;
+						if (bProfileEnabled)
+						{
+							if (bProfileAvailableQL)
+							{
+								Ret.bIsActive = true;
+								Ret.bCanBeActive = true;
+
+#if WITH_EDITOR
+								if (OutDetails)
+								{
+									//Clear existing failure details
+									OutDetails->ReasonsForInActive.Reset();
+									OutDetails->ReasonsForDisabled.Reset();
+								}
+#endif
+							}
+						}
+						else
+						{
+							Ret.bIsActive = false;
+							Ret.bCanBeActive = false;
+
+							#if WITH_EDITOR
+							if (OutDetails)
+							{
+								if (Profile == CurrProfile)
+								{
+									if (AvailableProfileMask == INDEX_NONE)
+									{
+										OutDetails->ReasonsForDisabled.Emplace(FText::Format(LOCTEXT("PlatformSetDPDisabledFmt", "{0} is disabled at {1} Quality."), FText::FromString(CurrProfile->GetName()), QLText));
+									}
+									else
+									{
+										OutDetails->ReasonsForDisabled.Emplace(FText::Format(LOCTEXT("PlatformSetDPDisabledFmt", "{0} is disabled."), FText::FromString(CurrProfile->GetName())));
+									}
+								}
+								else
+								{
+									if (AvailableProfileMask == INDEX_NONE)
+									{
+										OutDetails->ReasonsForDisabled.Emplace(FText::Format(LOCTEXT("PlatformSetParentDPDisabledFmt", "{0} is Disabeld at {1} Quality."), FText::FromString(CurrProfile->GetName()), QLText));
+									}
+									else
+									{
+										OutDetails->ReasonsForDisabled.Emplace(FText::Format(LOCTEXT("PlatformSetParentDPDisabledFmt", "{0} is Disabled."), FText::FromString(CurrProfile->GetName())));
+									}
+								}
+							}
+							#endif
+						}
+
+						break;
+					}
 				}
 			}
 			CurrProfile = Cast<UDeviceProfile>(CurrProfile->Parent);
 		}
 	}
 
-	return bIsEnabledByEQ;
+	//Finally, check any CVar conditions placed on this platform set.
+	for (const FNiagaraPlatformSetCVarCondition& CVarCondition : CVarConditions)
+	{
+		//Bail if any cvar condition isn't met.
+		bool bCanEverPass = true;
+		bool bCanEverFail = false;
+		//Get all quality levels where this condition will pass.
+		int32 ConditionPassedMask = INDEX_NONE;
+		bool bConditionPassed = CVarCondition.Evaluate(Profile, bCanEverPass, bCanEverFail);
+
+		ensure(bCanEverPass || bCanEverFail);
+
+		bool bCanBeActiveDisabled = false;
+		bool bIsActiveDistabled = false;
+		if(CVarCondition.PassResponse == ENiagaraCVarConditionResponse::Enable && bProfileAvailableQL)
+		{
+			if (bConditionPassed)
+			{
+				Ret.bIsActive = true;				
+			}
+			if (bCanEverPass)
+			{
+				Ret.bCanBeActive = true;
+			}
+		}
+		else if (CVarCondition.PassResponse == ENiagaraCVarConditionResponse::Disable)
+		{
+			if (bConditionPassed)
+			{
+				Ret.bIsActive = false;
+				bIsActiveDistabled = true;
+			}
+			if (bCanEverFail == false)//If this condition will always pass then we'll always be disabled.
+			{
+				Ret.bCanBeActive = false;
+				bCanBeActiveDisabled = true;
+			}
+		}
+
+		if (CVarCondition.FailResponse == ENiagaraCVarConditionResponse::Enable && bProfileAvailableQL)
+		{
+			if (bConditionPassed == false)
+			{
+				Ret.bIsActive = true;
+			}
+			if (bCanEverFail)
+			{
+				Ret.bCanBeActive = true;
+			}
+		}
+		else if (CVarCondition.FailResponse == ENiagaraCVarConditionResponse::Disable)
+		{
+			if (bConditionPassed == false)
+			{
+				Ret.bIsActive = false;
+				bIsActiveDistabled = true;
+			}
+			if (bCanEverPass == false)//If we'll always fail than this platform set will never be active.
+			{
+				Ret.bCanBeActive = false;
+				bCanBeActiveDisabled = true;
+			}
+		}
+
+#if WITH_EDITOR
+		if(OutDetails)
+		{
+			if (bIsActiveDistabled)
+			{
+				FText ReasonTest = FText::Format(LOCTEXT("CVarConditionDisabledFmt","{0} is inactive due to a condition on CVar {1}"), FText::FromString(Profile->GetName()), FText::FromName(CVarCondition.CVarName));
+				OutDetails->ReasonsForInActive.Emplace(ReasonTest);
+			}
+			if (bCanBeActiveDisabled)
+			{
+				FText ReasonTest = FText::Format(LOCTEXT("CVarConditionDisabledFmt", "{0} is disabled due to a condition on CVar {1}"), FText::FromString(Profile->GetName()), FText::FromName(CVarCondition.CVarName));
+				OutDetails->ReasonsForDisabled.Emplace(ReasonTest);
+			}
+		}
+#endif
+	}
+
+	return Ret;
 }
 
+bool FNiagaraPlatformSet::ApplyRedirects()
+{
+	const UNiagaraSettings* Settings = GetDefault<UNiagaraSettings>();
+	check(Settings);
+
+	bool bSuccess = true;
+	for (const FNiagaraPlatformSetRedirect& Redirect : Settings->PlatformSetRedirects)
+	{
+		if (Redirect.Mode == ENiagaraDeviceProfileRedirectMode::CVar)
+		{
+#if WITH_EDITOR
+			//Do some additional validation in editor.
+			IConsoleVariable* CVar = FNiagaraPlatformSet::GetCVar(*Redirect.CVarConditionEnabled.CVarName.ToString());
+			if (!CVar)
+			{
+				UE_LOG(LogNiagara, Warning, TEXT("Could not find CVar %s when trying to apply PlatformSet Redirects."), *Redirect.CVarConditionEnabled.CVarName.ToString());
+				bSuccess  = false;
+				continue;
+			}
+			CVar = FNiagaraPlatformSet::GetCVar(*Redirect.CVarConditionDisabled.CVarName.ToString());
+			if (!CVar)
+			{
+				UE_LOG(LogNiagara, Warning, TEXT("Could not find CVar %s when trying to apply PlatformSet Redirects."), *Redirect.CVarConditionDisabled.CVarName.ToString());
+				bSuccess = false;
+				continue;
+			}
+#endif
+
+			bool bAddCVarEnabled = false;
+			bool bAddCVarDisabled = false;
+
+			for (auto DPEntryIt = DeviceProfileStates.CreateIterator(); DPEntryIt; ++DPEntryIt)
+			{
+				FNiagaraDeviceProfileStateEntry& DPEntry = *DPEntryIt;
+
+				//Do we want to redirect this profile?
+				if (Redirect.ProfileNames.Contains(DPEntry.ProfileName) == false)
+				{
+					continue;
+				}
+				
+// 				if (DPEntry.SetQualityLevelMask != INDEX_NONE)
+// 				{
+// 					UE_LOG(LogNiagara, Warning, TEXT("Cannot redirect entry for device profile %s with a CVar as it is not set for all quality levels."), *DPEntry.ProfileName.ToString());
+// 					continue;
+// 				}
+// 
+// 				if (DPEntry.QualityLevelMask != INDEX_NONE && DPEntry.QualityLevelMask != 0)
+// 				{
+// 					UE_LOG(LogNiagara, Warning, TEXT("Cannot redirect entry for device profile %s with a CVar as it is not set for all quality levels."), *DPEntry.ProfileName.ToString());
+// 					continue;
+// 				}
+
+				bool bEnabled = DPEntry.QualityLevelMask != 0;
+
+				if (bEnabled)
+				{
+					bAddCVarEnabled = true;
+				}
+				else
+				{
+					bAddCVarDisabled = true;
+				}			
+
+				DPEntryIt.RemoveCurrent();
+			}
+
+			if (bAddCVarEnabled)
+			{
+				CVarConditions.Add(Redirect.CVarConditionEnabled);
+			}
+
+			if (bAddCVarDisabled)
+			{
+				CVarConditions.Add(Redirect.CVarConditionDisabled);
+			}
+		}
+		else if (Redirect.Mode == ENiagaraDeviceProfileRedirectMode::DeviceProfile)
+		{
+			//Switch to new device profile.
+
+			bool bRedirect = false;
+			uint32 SetMask = 0xFFFFFFFF;
+			uint32 QLMask = 0xFFFFFFFF;
+			for (auto It = DeviceProfileStates.CreateIterator(); It; It++)
+			{
+				FNiagaraDeviceProfileStateEntry& Entry = *It;
+				if(Redirect.ProfileNames.Contains(Entry.ProfileName))
+				{
+					if (bRedirect)
+					{
+						//we've already applied this redirect for another check the state matches.
+						bool bMatch = QLMask == Entry.QualityLevelMask && SetMask == Entry.SetQualityLevelMask;
+						if (!bMatch)
+						{
+							UE_LOG(LogNiagara, Warning, TEXT("Found conflicting entries when trying to apply DP redirect from %s to %s"), *Entry.ProfileName.ToString(), *Redirect.RedirectProfileName.ToString());
+							bSuccess = false;
+							break;
+						}
+						It.RemoveCurrent();
+					}
+					else
+					{
+						//First time we've applied this redirect. Swap the DP name and note the masks.
+						Entry.ProfileName = Redirect.RedirectProfileName;
+						QLMask = Entry.QualityLevelMask;
+						SetMask = Entry.SetQualityLevelMask;
+						bRedirect = true;
+					}					
+				}
+			}
+		}
+		else
+		{
+			UE_LOG(LogNiagara, Warning, TEXT("Bad Mode for Platform Set Redirect!"));
+			bSuccess = false;
+		}
+	}
+
+	//Remove any dupes left in the DP array
+	for (int32 i = 0; i < DeviceProfileStates.Num(); ++i)
+	{
+		FNiagaraDeviceProfileStateEntry& Entry = DeviceProfileStates[i];
+
+		for (int32 j = i + 1; j < DeviceProfileStates.Num(); )
+		{
+			FNiagaraDeviceProfileStateEntry& Other = DeviceProfileStates[j];
+			
+			check(i!=j);
+
+			if (Entry == Other)
+			{
+				DeviceProfileStates.RemoveAtSwap(j);
+			}
+			else
+			{
+				++j;//If we remove an entry we want to recheck the current j.
+			}
+		}
+	}
+
+	return bSuccess;
+}
+
+void FNiagaraPlatformSet::OnCVarChanged(IConsoleVariable* CVar)
+{
+	InvalidateCachedData();
+
+	//Iterate on our cached cvars as there'll be a v small subset of CVars Niagara cares about.
+	for (auto& Pair : CachedCVarInfo)
+	{
+		FName CVarName  = Pair.Key;
+		FCachedCVarInfo& Info = Pair.Value;
+		if (Info.CVar == CVar)
+		{
+			ChangedCVars.AddUnique(Pair.Key);
+		}
+	}
+}
+
+void FNiagaraPlatformSet::OnCVarUnregistered(IConsoleVariable* CVar)
+{
+	//We also unregister the changed delegate when it's unregistered.
+	FScopeLock Lock(&CachedCVarInfoCritSec);
+
+	//Iterate on our cached cvars as there'll be a v small subset of CVars Niagara cares about.
+	for (auto& Pair : CachedCVarInfo)
+	{
+		FName CVarName = Pair.Key;
+		FCachedCVarInfo& Info = Pair.Value;
+		if (Info.CVar == CVar)
+		{
+			check(CVar == Info.CVar);
+			Info.CVar->OnChangedDelegate().Remove(Info.ChangedHandle);
+			OnCVarChanged(CVar);
+		}
+	}
+}
+
+IConsoleVariable* FNiagaraPlatformSet::GetCVar(FName CVarName)
+{
+	FScopeLock Lock(&CachedCVarInfoCritSec);
+	
+	FCachedCVarInfo* CachedInfo = CachedCVarInfo.Find(CVarName);
+	if (CachedInfo == nullptr)
+	{
+		IConsoleManager& CMan = IConsoleManager::Get();
+		CachedInfo = &CachedCVarInfo.Add(CVarName);
+		CachedInfo->CVar = CMan.FindConsoleVariable(*CVarName.ToString());
+		CachedInfo->ChangedHandle = CachedInfo->CVar->OnChangedDelegate().AddStatic(FNiagaraPlatformSet::OnCVarChanged);
+	}
+
+	return CachedInfo->CVar;
+}
+
+void FNiagaraPlatformSet::RefreshScalability()
+{
+	bool bNeedsRefresh = ChangedCVars.Num() > 0;//Other reasons to refresh deferred like this?
+
+	if (bNeedsRefresh)
+	{
+		InvalidateCachedData();
+
+		//Find any systems relying on a changed CVar and refresh it's scalability settings.
+		for (TObjectIterator<UNiagaraSystem> It; It; ++It)
+		{
+			UNiagaraSystem* System = *It;
+			check(System);
+
+			bool bNeedsUpdate = false;
+			auto NeedsScalabilityUpdate = [&](UObject* Owner, FNiagaraPlatformSet& PSet)
+			{
+				for (FNiagaraPlatformSetCVarCondition& CVarCondition : PSet.CVarConditions)
+				{
+					if (ChangedCVars.Contains(CVarCondition.CVarName))
+					{
+						bNeedsUpdate = true;
+					}
+				}
+			};
+
+			System->ForEachPlatformSet(NeedsScalabilityUpdate);
+			if (UNiagaraEffectType* FXType = System->GetEffectType())
+			{
+				FXType->ForEachPlatformSet(NeedsScalabilityUpdate);
+			}
+
+			if (bNeedsUpdate)
+			{
+				System->UpdateScalability();
+			}
+		}
+	}
+
+	ChangedCVars.Reset();
+}
+
+bool FNiagaraPlatformSet::CanConsiderDeviceProfile(const UDeviceProfile* Profile)const
+{
+	return GbAllowAllDeviceProfiles != 0 || Profile->IsVisibleForAssets();
+}
 
 void FNiagaraPlatformSet::InvalidateCachedData()
 {
@@ -430,6 +991,7 @@ void FNiagaraPlatformSet::InvalidateCachedData()
 
 	CachedQualityLevel = INDEX_NONE;
 	LastDirtiedFrame = GFrameNumber;
+	CachedQualityLevel = INDEX_NONE;
 }
 
 bool FNiagaraPlatformSet::IsEnabledForPlatform(const FString& PlatformName)const
@@ -438,9 +1000,9 @@ bool FNiagaraPlatformSet::IsEnabledForPlatform(const FString& PlatformName)const
 	{
 		if (const UDeviceProfile* Profile = Cast<const UDeviceProfile>(ProfileObj))
 		{
-			if (Profile->DeviceType == PlatformName)
+			if (CanConsiderDeviceProfile(Profile) && Profile->DeviceType == PlatformName)
 			{
-				if (IsEnabledForDeviceProfile(Profile) != 0)
+				if(IsEnabledForDeviceProfile(Profile))
 				{
 					return true;//At least one profile for this platform is enabled.
 				}
@@ -501,32 +1063,13 @@ void SetGNiagaraDeviceProfile(UDeviceProfile* Profile)
 
 	UDeviceProfile* OverrideDP = GNiagaraPlatformOverride.Get();
 	check(OverrideDP);
-	int32 DPQL = FNiagaraPlatformSet::QualityLevelFromMask(FNiagaraPlatformSet::GetEffectQualityMaskForDeviceProfile(OverrideDP));
+	int32 DPQL = FNiagaraPlatformSet::QualityLevelFromMask(FNiagaraPlatformSet::GetActiveQualityMaskForDeviceProfile(OverrideDP));
 					
 	OnSetCVarFromIniEntry(*GDeviceProfilesIni, NiagaraQualityLevelName, *LexToString(DPQL), ECVF_SetByMask);
 }
 
 #if WITH_EDITOR
 
-// bool FNiagaraPlatformSet::Conflicts(const FNiagaraPlatformSet& Other, TArray<const UDeviceProfile*>& OutConflictingProfiles, bool bIncludeProfilesWithVariableQualityLevel)const
-// {
-// 	for (const UObject* ProfileObj : UDeviceProfileManager::Get().Profiles)
-// 	{
-// 		if (const UDeviceProfile* Profile = Cast<const UDeviceProfile>(ProfileObj))
-// 		{
-// 			if (bIncludeProfilesWithVariableQualityLevel || CanChangeScalabilityAtRuntime(Profile) == false)
-// 			{
-// 				if (IsEnabledForDeviceProfile(Profile) && Other.IsEnabledForDeviceProfile(Profile))
-// 				{
-// 					OutConflictingProfiles.AddUnique(Profile);
-// 				}
-// 			}
-// 		}
-// 	}
-// 
-// 	return OutConflictingProfiles.Num() > 0;
-// }
-
 bool FNiagaraPlatformSet::IsEffectQualityEnabled(int32 EffectQuality)const
 {
 	return ((1 << EffectQuality) & QualityLevelMask) != 0;
@@ -551,7 +1094,7 @@ void FNiagaraPlatformSet::SetDeviceProfileState(UDeviceProfile* Profile, int32 Q
 	int32 Index = INDEX_NONE;
 	Index = DeviceProfileStates.IndexOfByPredicate([&](const FNiagaraDeviceProfileStateEntry& Entry) { return Entry.ProfileName == Profile->GetFName(); });
 
-	int32 ProfileQLMask = GetEffectQualityMaskForDeviceProfile(Profile);
+	int32 ProfileQLMask = GetAvailableQualityMaskForDeviceProfile(Profile);
 	if (ProfileQLMask != INDEX_NONE)
 	{
 		//For platforms that cannot change EQ at runtime we mark all state bits when setting state here so that if someone changes their EQ setting in the future, the state will be preserved.
@@ -583,7 +1126,7 @@ ENiagaraPlatformSelectionState FNiagaraPlatformSet::GetDeviceProfileState(UDevic
 {
 	if (const FNiagaraDeviceProfileStateEntry* ExistingEntry = DeviceProfileStates.FindByPredicate([&](const FNiagaraDeviceProfileStateEntry& Entry) { return Entry.ProfileName == Profile->GetFName(); }))
 	{
-		int32 ProfileQLMask = GetEffectQualityMaskForDeviceProfile(Profile);
+		int32 ProfileQLMask = GetAvailableQualityMaskForDeviceProfile(Profile);
 		if (ProfileQLMask == INDEX_NONE || ProfileQLMask & CreateQualityLevelMask(QualityLevel))
 		{
 			//For profiles that cannot change scalability at runtime we store all flags in their state so that if anyone ever changes their EQ Cvar, the state setting remains valid.
@@ -598,6 +1141,7 @@ ENiagaraPlatformSelectionState FNiagaraPlatformSet::GetDeviceProfileState(UDevic
 void FNiagaraPlatformSet::OnChanged()
 {
 	LastBuiltFrame = 0;
+	InvalidateCachedData();
 }
 
 bool FNiagaraPlatformSet::GatherConflicts(const TArray<const FNiagaraPlatformSet*>& PlatformSets, TArray<FNiagaraPlatformSetConflictInfo>& OutConflicts)
@@ -622,8 +1166,12 @@ bool FNiagaraPlatformSet::GatherConflicts(const TArray<const FNiagaraPlatformSet
 			for (UObject* DPObj : UDeviceProfileManager::Get().Profiles)
 			{
 				UDeviceProfile* Profile = CastChecked<UDeviceProfile>(DPObj);
-				int32 AEnabledMask = SetA.IsEnabledForDeviceProfile(Profile);
-				int32 BEnabledMask = SetB.IsEnabledForDeviceProfile(Profile);
+				if (SetA.CanConsiderDeviceProfile(Profile) == false && SetB.CanConsiderDeviceProfile(Profile) == false)
+				{
+					continue;
+				}
+				int32 AEnabledMask = SetA.GetEnabledMaskForDeviceProfile(Profile);
+				int32 BEnabledMask = SetB.GetEnabledMaskForDeviceProfile(Profile);
 				int32 ConflictMask = AEnabledMask & BEnabledMask;
 
 				if (ConflictMask != 0)
@@ -684,9 +1232,15 @@ FNiagaraPlatformSet::FPlatformIniSettings& FNiagaraPlatformSet::GetPlatformIniSe
 
 	//If this platform can change scalability settings at runtime then we return the full mask.
 	int32 CanChangeScalabiltiySettings = 0;
+	int32 MinQualityLevel = DefaultMinQualityLevel;
+	int32 MaxQualityLevel = DefaultMaxQualityLevel;
 
 	FindCVarValue(TEXT("SystemSettings"), CanChangeEQCVarName, CanChangeScalabiltiySettings);
 
+	FindCVarValue(TEXT("SystemSettings"), NiagaraMinQualityLevelName, MinQualityLevel);
+
+	FindCVarValue(TEXT("SystemSettings"), NiagaraMaxQualityLevelName, MaxQualityLevel);
+
 	if (!FindCVarValue(TEXT("ScalabilityGroups"), TEXT("sg.EffectsQuality"), EffectsQuality))
 	{
 		FindCVarValue(TEXT("SystemSettings"), TEXT("sg.EffectsQuality"), EffectsQuality);
@@ -699,7 +1253,7 @@ FNiagaraPlatformSet::FPlatformIniSettings& FNiagaraPlatformSet::GetPlatformIniSe
 	FindCVarValue(TEXT("SystemSettings"), PruneEmittersOnCookName, PruneEmittersOnCook);
 
 	FPlatformIniSettings& NewSetting = CachedPlatformIniSettings.Add(*PlatformName);
-	NewSetting = FPlatformIniSettings(CanChangeScalabiltiySettings, PruneEmittersOnCook, EffectsQuality);
+	NewSetting = FPlatformIniSettings(CanChangeScalabiltiySettings, PruneEmittersOnCook, EffectsQuality, MinQualityLevel, MaxQualityLevel);
 
 	//Find the Niagara Quality Levels set for each EffectsQuality Level for this platform.
 	int32 NumEffectsQualities = Scalability::GetQualityLevelCounts().EffectsQuality;
@@ -720,41 +1274,54 @@ FNiagaraPlatformSet::FPlatformIniSettings& FNiagaraPlatformSet::GetPlatformIniSe
 int32 FNiagaraPlatformSet::GetEffectQualityMaskForPlatform(const FString& PlatformName)
 {
 	FPlatformIniSettings& PlatformSettings = GetPlatformIniSettings(PlatformName);
-	return PlatformSettings.bCanChangeScalabilitySettingsAtRuntime;
+	return PlatformSettings.QualityLevelMask;
 }
-#endif
 
+FNiagaraPlatformSet::FPlatformIniSettings::FPlatformIniSettings()
+	:bCanChangeScalabilitySettingsAtRuntime(0), bPruneEmittersOnCook(false), EffectsQuality(0), MinQualityLevel(INDEX_NONE), MaxQualityLevel(INDEX_NONE), QualityLevelMask(INDEX_NONE)
+{}
+FNiagaraPlatformSet::FPlatformIniSettings::FPlatformIniSettings(int32 InbCanChangeScalabilitySettingsAtRuntime, int32 InbPruneEmittersOnCook, int32 InEffectsQuality, int32 InMinQualityLevel, int32 InMaxQualityLevel)
+	: bCanChangeScalabilitySettingsAtRuntime(InbCanChangeScalabilitySettingsAtRuntime), bPruneEmittersOnCook(InbPruneEmittersOnCook)
+	, EffectsQuality(InEffectsQuality), MinQualityLevel(InMinQualityLevel), MaxQualityLevel(InMaxQualityLevel)
+{
+	if (bCanChangeScalabilitySettingsAtRuntime)
+	{
+		QualityLevelMask = INDEX_NONE;
+	}
+	else
+	{
+		const UNiagaraSettings* Settings = GetDefault<UNiagaraSettings>();
+		check(Settings);
+		QualityLevelMask = 0;
+		int32 NumQualityLevels = Settings->QualityLevels.Num();
+		int32 PlatformQL = 0;
+		for (int32 QL = MinQualityLevel; QL <= MaxQualityLevel; ++QL)
+		{
+			QualityLevelMask |= (1 << QL);
+		}
+	}
+}
+
+#endif
 //////////////////////////////////////////////////////////////////////////
 
 #if WITH_EDITOR
 TMap<const UDeviceProfile*, FDeviceProfileValueCache::FCVarValueMap> FDeviceProfileValueCache::CachedDeviceProfileValues;
-TMap<FName, FDeviceProfileValueCache::FCVarValueMap> FDeviceProfileValueCache::CachedPlatformValues;
 
-bool FDeviceProfileValueCache::GetValueInternal(const UDeviceProfile* DeviceProfile, FName CVarName, FString& OutValue)
+const FNiagaraCVarValues& FDeviceProfileValueCache::GetValues(const UDeviceProfile* DeviceProfile, IConsoleVariable* CVar, FName CVarName)
 {
-	//First look if we've asked for this CVar for this device profile before.
-	if (FCVarValueMap* CVarMap = CachedDeviceProfileValues.Find(DeviceProfile))
-	{
-		if (FString* CVarValueString = CVarMap->Find(CVarName))
-		{
-			OutValue = *CVarValueString;
-			return true;
-		}
-	}
+	check(CVar);
 
-	//If not we'll need to look for it.
-	
-	FCVarValueMap& CVarMap = CachedDeviceProfileValues.Add(DeviceProfile);
+	//First look if we've asked for this CVar for this device profile before.
+	FCVarValueMap& CVarMap = CachedDeviceProfileValues.FindOrAdd(DeviceProfile);
 
-	//First see if the device profile has it explicitly set.
-	if (DeviceProfile->GetConsolidatedCVarValue(*CVarName.ToString(), OutValue, false))
+	if (FNiagaraCVarValues* CVarValues = CVarMap.Find(CVar))
 	{
-		CVarMap.Add(CVarName, OutValue);
-		return true;
+		return *CVarValues;
 	}
 
-	//Otherwise we need to check the ini files for the profiles platform.
-	//It'd be lovely if there were just a utility function to call that did this (or better) for us.
+	//If not we'll need to look for it.	
+	FNiagaraCVarValues& NewCVarValues = CVarMap.Add(CVar);
 	
 	const FString& PlatformName = DeviceProfile->DeviceType;
 
@@ -767,82 +1334,97 @@ bool FDeviceProfileValueCache::GetValueInternal(const UDeviceProfile* DeviceProf
 	FConfigFile ScalabilitySettings;
 	FConfigCacheIni::LoadLocalIniFile(ScalabilitySettings, TEXT("Scalability"), true, *PlatformName);
 
-	auto FindCVarValue = [&](const TCHAR* Section, const TCHAR* CVarName, FString& OutVal)
+	int32 NumEffectsQualities = Scalability::GetQualityLevelCounts().EffectsQuality;
+
+	FString DefaultStr;
+
+	//First see if the device profile has it explicitly set.
+	//If we found the cvar in the device profile then use that as a base, otherwise look through all the inis.
+	if (DeviceProfile->GetConsolidatedCVarValue(*CVarName.ToString(), DefaultStr, false) == false)
 	{
-		bool bFound = true;
-		if (!ScalabilitySettings.GetString(Section, CVarName, OutVal))
+		//Try to grab a default from our initial scalability group.
+		FNiagaraPlatformSet::FPlatformIniSettings& PlatformIniSettings = FNiagaraPlatformSet::GetPlatformIniSettings(PlatformName);
+		
+		int32 EffectsQuality = INDEX_NONE;
+		//See if this profile overrides effects quality.
+		if (!DeviceProfile->GetConsolidatedCVarValue(TEXT("sg.EffectsQuality"), EffectsQuality))
 		{
-			if (!GameSettings.GetString(Section, CVarName, OutVal))
+			EffectsQuality = PlatformIniSettings.EffectsQuality;
+		}
+		check(EffectsQuality != INDEX_NONE && EffectsQuality < NumEffectsQualities);
+		FString DefaultEffectsQualitySectionName = Scalability::GetScalabilitySectionString(TEXT("EffectsQuality"), EffectsQuality, NumEffectsQualities);
+
+		if (ScalabilitySettings.GetString(*DefaultEffectsQualitySectionName, *CVarName.ToString(), DefaultStr) == false)
+		{
+			//Otherwise try to grab from the main inis.
+			auto FindCVarValue = [&](const TCHAR* Section, const TCHAR* CVarName, FString& OutVal)
 			{
-				if (!EngineSettings.GetString(Section, CVarName, OutVal))
+				bool bFound = true;
+				if (!ScalabilitySettings.GetString(Section, CVarName, OutVal))
 				{
-					bFound = false;
+					if (!GameSettings.GetString(Section, CVarName, OutVal))
+					{
+						if (!EngineSettings.GetString(Section, CVarName, OutVal))
+						{
+							bFound = false;
+						}
+					}
+				}
+				return bFound;
+			};
+				
+			if (FindCVarValue(TEXT("SystemSettings"), *CVarName.ToString(), DefaultStr) == false)
+			{
+				if (IConsoleVariable* PlatformCVar = CVar->GetPlatformValueVariable(*PlatformName).Get())//Possible this could do all the above work for us.
+				{
+					DefaultStr = PlatformCVar->GetString();
+				}
+				else if (IConsoleVariable* DefaultCVar = CVar->GetDefaultValueVariable())
+				{
+					DefaultStr = DefaultCVar->GetString();
+				}
+				else
+				{
+					//Failing all that we just take the current value.
+					DefaultStr = CVar->GetString();
 				}
 			}
 		}
-		return bFound;
-	};
-
-	if (FindCVarValue(TEXT("SystemSettings"), *CVarName.ToString(), OutValue))
-	{
-		CVarMap.Add(CVarName, OutValue);
-		return true;
 	}
 
-	//Failing all that we just take the default value.
-	if (IConsoleVariable* CVar = IConsoleManager::Get().FindConsoleVariable(*CVarName.ToString()))
-	{
-		OutValue = CVar->GetString();
-		CVarMap.Add(CVarName, OutValue);
-		return true;
+	check(DefaultStr.IsEmpty() == false);
+
+	NewCVarValues.Default.Set(DefaultStr, CVar);
+	NewCVarValues.PerQualityLevelValues.SetNum(NumEffectsQualities);
+
+	//Also gather possible values from other scalability groups.
+	//Falling back to the overall default if that EQ doesn't set it.
+	FString EQStr;
+	for (int32 EQ = 0; EQ < NumEffectsQualities; ++EQ)
+	{		
+		FString SectionName = Scalability::GetScalabilitySectionString(TEXT("EffectsQuality"), EQ, NumEffectsQualities);
+		if (ScalabilitySettings.GetString(*SectionName, *CVarName.ToString(), EQStr))
+		{
+			NewCVarValues.PerQualityLevelValues[EQ].Set(EQStr, CVar);
+		}
+		else
+		{
+			NewCVarValues.PerQualityLevelValues[EQ].Set(DefaultStr, CVar);
+		}
 	}
 
-	//Only really possible if the CVar doesn't exist.
-	return false;
+	return NewCVarValues;
 }
 
 void FDeviceProfileValueCache::Empty()
 {
 	CachedDeviceProfileValues.Empty();
-	CachedPlatformValues.Empty();
-}
-
-template<typename T>
-bool FDeviceProfileValueCache::GetValue(const UDeviceProfile* DeviceProfile, FName CVarName, T& OutValue)
-{
-	FString ValString;
-	if (GetValueInternal(DeviceProfile, CVarName, ValString))
-	{
-		LexFromString(OutValue, *ValString);
-		return true;
-	}
-	return false;
 }
 
 #endif
 
 //////////////////////////////////////////////////////////////////////////
 
-TMap<FName, FDelegateHandle> FNiagaraPlatformSetCVarCondition::CVarChangedDelegateHandles;
-FCriticalSection FNiagaraPlatformSetCVarCondition::ChangedDelegateHandlesCritSec;
-
-void FNiagaraPlatformSetCVarCondition::OnCVarChanged(IConsoleVariable* CVar)
-{
-	//TODO: Iterate over all systems and recheck all their PlatformSet active states.
-	//Any states that have changes should trigger a reinit of the system in question.
-
-	//For right now just brute force reinit everything.
-	//At least this is limited to only the CVars that these conditions are reading.
-	FNiagaraPlatformSet::InvalidateCachedData();
-
-	for (TObjectIterator<UNiagaraSystem> It; It; ++It)
-	{
-		UNiagaraSystem* System = *It;
-		check(System);
-		System->UpdateScalability();
-	}
-}
-
 FNiagaraPlatformSetCVarCondition::FNiagaraPlatformSetCVarCondition()
 	: bUseMinInt(true)
 	, bUseMaxInt(false)
@@ -854,20 +1436,10 @@ FNiagaraPlatformSetCVarCondition::FNiagaraPlatformSetCVarCondition()
 
 IConsoleVariable* FNiagaraPlatformSetCVarCondition::GetCVar()const
 {
-	if (CachedCVar == nullptr)
+	if (CachedCVar == nullptr || FNiagaraPlatformSet::GetLastDirtiedFrame() >= LastCachedFrame || LastCachedFrame == INDEX_NONE)
 	{
-		IConsoleManager& CVarMan = IConsoleManager::Get();
-		CachedCVar = CVarMan.FindConsoleVariable(*CVarName.ToString());
-		if(CachedCVar)
-		{
-			FScopeLock Lock(&ChangedDelegateHandlesCritSec);
-			FDelegateHandle* Handle = CVarChangedDelegateHandles.Find(CVarName);
-			if(!Handle)
-			{
-				//We've not bound to this cvar's change callback yet. Do so now.
-				CVarChangedDelegateHandles.Add(CVarName) = CachedCVar->OnChangedDelegate().AddStatic(&FNiagaraPlatformSetCVarCondition::OnCVarChanged);
-			}
-		}
+		LastCachedFrame = GFrameNumber;
+		CachedCVar = FNiagaraPlatformSet::GetCVar(CVarName);
 	}
 	return CachedCVar;
 }
@@ -878,61 +1450,90 @@ void FNiagaraPlatformSetCVarCondition::SetCVar(FName InCVarName)
 	CachedCVar = nullptr;
 }
 
-bool FNiagaraPlatformSetCVarCondition::IsEnabledForPlatform(const FString& PlatformName)const
-{
-	for (const UObject* ProfileObj : UDeviceProfileManager::Get().Profiles)
+template<> bool FNiagaraPlatformSetCVarCondition::GetCVarValue<bool>(IConsoleVariable* CVar)const { return CVar->GetBool(); }
+template<> int32 FNiagaraPlatformSetCVarCondition::GetCVarValue<int32>(IConsoleVariable* CVar)const { return CVar->GetInt(); }
+template<> float FNiagaraPlatformSetCVarCondition::GetCVarValue<float>(IConsoleVariable* CVar)const { return CVar->GetFloat(); }
+
+template<typename T>
+bool FNiagaraPlatformSetCVarCondition::Evaluate_Internal(const UDeviceProfile* DeviceProfile, bool& bOutCanEverPass, bool& bOutCanEverFail)const
+{		
+	bool bConditionPassed = false;
+
+	IConsoleVariable* CVar = GetCVar();
+
+	if(!CVar)
 	{
-		if (const UDeviceProfile* Profile = Cast<const UDeviceProfile>(ProfileObj))
+		return false;
+	}
+
+	bool bCanChangeAtRuntime = FNiagaraPlatformSet::CanChangeScalabilityAtRuntime(DeviceProfile) && (CVar->GetFlags() & ECVF_ReadOnly) == 0;
+
+#if WITH_EDITOR
+	const FNiagaraCVarValues& PossibleValues = FDeviceProfileValueCache::GetValues(DeviceProfile, CVar, CVarName);
+
+	T Default = PossibleValues.Default.Get<T>();
+	bConditionPassed = CheckValue(Default);
+
+	bOutCanEverPass = bConditionPassed;
+	bOutCanEverFail = !bConditionPassed;
+
+ 	if (bCanChangeAtRuntime)
+ 	{
+ 		bOutCanEverFail = true;
+ 		bOutCanEverPass = true;
+	}
+ 	else
+	{	
+		for (int32 i = 0; i < PossibleValues.PerQualityLevelValues.Num(); ++i)
 		{
-			if (Profile->DeviceType == PlatformName)
+			T Val = PossibleValues.PerQualityLevelValues[i].Get<T>();
+			bool bPass = CheckValue(Val);
+			if (bPass)
 			{
-				if (IsEnabledForDeviceProfile(Profile, false) != 0)
-				{
-					return true;//At least one profile for this platform is enabled.
-				}
+				bOutCanEverPass = true;
+			}
+			else
+			{
+				bOutCanEverFail = true;
 			}
 		}
 	}
-	return false;
-}
 
-template<> bool FNiagaraPlatformSetCVarCondition::GetCVarValue<bool>(IConsoleVariable* CVar)const { return CVar->GetBool(); }
-template<> int32 FNiagaraPlatformSetCVarCondition::GetCVarValue<int32>(IConsoleVariable* CVar)const { return CVar->GetInt(); }
-template<> float FNiagaraPlatformSetCVarCondition::GetCVarValue<float>(IConsoleVariable* CVar)const { return CVar->GetFloat(); }
+#else
 
-template<typename T>
-bool FNiagaraPlatformSetCVarCondition::IsEnabledForDeviceProfile_Internal(const UDeviceProfile* DeviceProfile, bool bCheckCurrentStateOnly)const
-{
-#if WITH_EDITOR
-	T ProfileValue;
-	if (bCheckCurrentStateOnly == false && FDeviceProfileValueCache::GetValue(DeviceProfile, CVarName, ProfileValue))
+	bConditionPassed = CheckValue(GetCVarValue<T>(CVar));
+
+	if (bCanChangeAtRuntime)
 	{
-		return CheckValue(ProfileValue);
+		bOutCanEverPass = true;
+		bOutCanEverFail = true;
 	}
-#endif
-
-	if (IConsoleVariable* CVar = GetCVar())
+	else
 	{
-		return CheckValue(GetCVarValue<T>(CVar));
+		bOutCanEverPass = bConditionPassed;
+		bOutCanEverFail = !bConditionPassed;
 	}
-	return false;
+
+#endif
+
+	return bConditionPassed;
 }
 
-bool FNiagaraPlatformSetCVarCondition::IsEnabledForDeviceProfile(const UDeviceProfile* DeviceProfile, bool bCheckCurrentStateOnly)const
+bool FNiagaraPlatformSetCVarCondition::Evaluate(const UDeviceProfile* DeviceProfile, bool& bOutCanEverPass, bool& bOutCanEverFail)const
 {
 	if (IConsoleVariable* CVar = GetCVar())
 	{
 		if (CVar->IsVariableBool())
 		{
-			return IsEnabledForDeviceProfile_Internal<bool>(DeviceProfile, bCheckCurrentStateOnly);
+			return Evaluate_Internal<bool>(DeviceProfile, bOutCanEverPass, bOutCanEverFail);
 		}
 		else if (CVar->IsVariableInt())
 		{
-			return IsEnabledForDeviceProfile_Internal<int32>(DeviceProfile, bCheckCurrentStateOnly);
+			return Evaluate_Internal<int32>(DeviceProfile, bOutCanEverPass, bOutCanEverFail);
 		}
 		else if (CVar->IsVariableFloat())
 		{
-			return IsEnabledForDeviceProfile_Internal<float>(DeviceProfile, bCheckCurrentStateOnly);
+			return Evaluate_Internal<float>(DeviceProfile, bOutCanEverPass, bOutCanEverFail);
 		}
 		else
 		{
@@ -947,4 +1548,4 @@ bool FNiagaraPlatformSetCVarCondition::IsEnabledForDeviceProfile(const UDevicePr
 	return false;
 }
 
-#undef LOCTEXT_NAMESPACE
+#undef LOCTEXT_NAMESPACE
\ No newline at end of file
diff --git a/Engine/Plugins/FX/Niagara/Source/Niagara/Private/NiagaraSystem.cpp b/Engine/Plugins/FX/Niagara/Source/Niagara/Private/NiagaraSystem.cpp
--- a/Engine/Plugins/FX/Niagara/Source/Niagara/Private/NiagaraSystem.cpp
+++ b/Engine/Plugins/FX/Niagara/Source/Niagara/Private/NiagaraSystem.cpp
@@ -32,6 +32,8 @@
 #include "UObject/UObjectThreadContext.h"
 #include "ShaderCompiler.h"
 #include "NiagaraSimulationStageBase.h"
+#include "DeviceProfiles/DeviceProfile.h"
+#include "DeviceProfiles/DeviceProfileManager.h"
 
 #define LOCTEXT_NAMESPACE "NiagaraSystem"
 
@@ -1196,6 +1198,13 @@ void UNiagaraSystem::PostLoad()
 	}
 #endif // WITH_EDITORONLY_DATA
 
+	//Apply platform set redirectors
+	auto ApplyPlatformSetRedirects = [](UObject* Owner, FNiagaraPlatformSet& Platforms)
+	{
+		Platforms.ApplyRedirects();
+	};
+	ForEachPlatformSet(ApplyPlatformSetRedirects);
+
 #if !WITH_EDITOR
 	// When running without the editor in a cooked build we run the update immediately in post load since
 	// there will be no merging or compiling which makes it safe to do so.
diff --git a/Engine/Plugins/FX/Niagara/Source/Niagara/Public/NiagaraModule.h b/Engine/Plugins/FX/Niagara/Source/Niagara/Public/NiagaraModule.h
--- a/Engine/Plugins/FX/Niagara/Source/Niagara/Public/NiagaraModule.h
+++ b/Engine/Plugins/FX/Niagara/Source/Niagara/Public/NiagaraModule.h
@@ -56,6 +56,7 @@ public:
 	void OnPreExit();
 
 	void OnWorldTickStart(UWorld* World, ELevelTick TickType, float DeltaSeconds);
+	void OnBeginFrame();
 	void OnWorldBeginTearDown(UWorld* World);
 
 	FDelegateHandle SetOnProcessShaderCompilationQueue(FOnProcessQueue InOnProcessQueue);
@@ -358,5 +359,7 @@ private:
 #if WITH_NIAGARA_DEBUGGER
 	TPimplPtr<FNiagaraDebuggerClient> DebuggerClient;
 #endif
+
+	FDelegateHandle OnCVarUnregisteredHandle;
 };
 
diff --git a/Engine/Plugins/FX/Niagara/Source/Niagara/Public/NiagaraRendererProperties.h b/Engine/Plugins/FX/Niagara/Source/Niagara/Public/NiagaraRendererProperties.h
--- a/Engine/Plugins/FX/Niagara/Source/Niagara/Public/NiagaraRendererProperties.h
+++ b/Engine/Plugins/FX/Niagara/Source/Niagara/Public/NiagaraRendererProperties.h
@@ -270,6 +270,9 @@ public:
 
 	static bool IsSortHighPrecision(ENiagaraRendererSortPrecision SortPrecision);
 
+	template<typename TAction>
+	void ForEachPlatformSet(TAction Func);
+
 	/** Platforms on which this renderer is enabled. */
 	UPROPERTY(EditAnywhere, Category = "Scalability", meta=(DisplayInScalabilityContext))
 	FNiagaraPlatformSet Platforms;
@@ -314,3 +317,9 @@ protected:
 	FOnPropertiesChanged OnPropertiesChangedDelegate;
 #endif
 };
+
+template<typename TAction>
+void UNiagaraRendererProperties::ForEachPlatformSet(TAction Func)
+{
+	Func(this, Platforms);
+}
diff --git a/Engine/Plugins/FX/Niagara/Source/Niagara/Public/NiagaraSettings.h b/Engine/Plugins/FX/Niagara/Source/Niagara/Public/NiagaraSettings.h
--- a/Engine/Plugins/FX/Niagara/Source/Niagara/Public/NiagaraSettings.h
+++ b/Engine/Plugins/FX/Niagara/Source/Niagara/Public/NiagaraSettings.h
@@ -7,6 +7,7 @@
 #include "NiagaraScript.h"
 #include "Engine/DeveloperSettings.h"
 #include "InputCoreTypes.h"
+#include "NiagaraPlatformSet.h"
 #include "NiagaraSettings.generated.h"
 
 // This enum must match the order in NiagaraDataInterfaceSkeletalMesh.ush
@@ -177,6 +178,9 @@ class NIAGARA_API UNiagaraSettings : public UDeveloperSettings
 	UPROPERTY(config, EditAnywhere, Category = AsyncGpuTraceDI, meta = (DisplayName = "Trace Provider Priorities (Experimental)", ConfigRestartRequired = true))
 	TArray<TEnumAsByte<ENDICollisionQuery_AsyncGpuTraceProvider::Type>> NDICollisionQuery_AsyncGpuTraceProviderOrder;
 
+	UPROPERTY(config, EditAnywhere, Category = Scalability)
+	TArray<FNiagaraPlatformSetRedirect> PlatformSetRedirects;
+
 	// Begin UDeveloperSettings Interface
 	virtual FName GetCategoryName() const override;
 #if WITH_EDITOR
diff --git a/Engine/Plugins/FX/Niagara/Source/NiagaraEditor/Private/Customizations/NiagaraPlatformSetCustomization.cpp b/Engine/Plugins/FX/Niagara/Source/NiagaraEditor/Private/Customizations/NiagaraPlatformSetCustomization.cpp
--- a/Engine/Plugins/FX/Niagara/Source/NiagaraEditor/Private/Customizations/NiagaraPlatformSetCustomization.cpp
+++ b/Engine/Plugins/FX/Niagara/Source/NiagaraEditor/Private/Customizations/NiagaraPlatformSetCustomization.cpp
@@ -131,6 +131,8 @@ void FNiagaraPlatformSetCustomization::GenerateQualityLevelSelectionWidgets()
 
 	int32 NumQualityLevels = Settings->QualityLevels.Num();
 	
+	QualityLevelMenuAnchors.Reset();
+	QualityLevelMenuContents.Reset();
 	QualityLevelMenuAnchors.SetNum(NumQualityLevels);
 	QualityLevelMenuContents.SetNum(NumQualityLevels);
 
@@ -335,7 +337,7 @@ EVisibility FNiagaraPlatformSetCustomization::GetQualityLevelErrorVisibility(int
 		if (ConflictInfo.SetAIndex == PlatformSetArrayIndex || 
 			ConflictInfo.SetBIndex == PlatformSetArrayIndex)
 		{
-			// this conflict applies to this platform set, check if it applies to this quality leve button
+			// this conflict applies to this platform set, check if it applies to this quality level button
 			const int32 QLMask = FNiagaraPlatformSet::CreateQualityLevelMask(QualityLevel);
 
 			for (const FNiagaraPlatformSetConflictEntry& Conflict : ConflictInfo.Conflicts)
@@ -589,6 +591,9 @@ FText FNiagaraPlatformSetCustomization::GetDeviceProfileErrorToolTip(UDeviceProf
 FReply FNiagaraPlatformSetCustomization::ToggleMenuOpenForQualityLevel(int32 QualityLevel)
 {
 	check(QualityLevelMenuAnchors.IsValidIndex(QualityLevel));
+
+	QualityLevelMenuContents[QualityLevel].Reset();
+	GenerateDeviceProfileTreeWidget(QualityLevel);
 	
 	TSharedPtr<SMenuAnchor> MenuAnchor = QualityLevelMenuAnchors[QualityLevel];
 	MenuAnchor->SetIsOpen(!MenuAnchor->IsOpen());
@@ -599,21 +604,14 @@ FReply FNiagaraPlatformSetCustomization::ToggleMenuOpenForQualityLevel(int32 Qua
 // Is or does a viewmodel contain any children active at the given quality?
 bool FNiagaraPlatformSetCustomization::IsTreeActiveForQL(const TSharedPtr<FNiagaraDeviceProfileViewModel>& Tree, int32 QualityLevelMask) const
 {
-	int32 Mask = TargetPlatformSet->GetEffectQualityMaskForDeviceProfile(Tree->Profile);
-	if ((Mask & QualityLevelMask) != 0)
+	if (TargetPlatformSet->CanConsiderDeviceProfile(Tree->Profile) == false)
 	{
-		return true;
+		return false;
 	}
 
-	for (const TSharedPtr<FNiagaraDeviceProfileViewModel>& Child : Tree->Children)
-	{
-		if (IsTreeActiveForQL(Child, QualityLevelMask))
-		{
-			return true;
-		}
-	}
+	int32 AvailableMask = TargetPlatformSet->GetAvailableQualityMaskForDeviceProfile(Tree->Profile);
 
-	return false;
+	return (AvailableMask & QualityLevelMask) != 0;
 }
 
 void FNiagaraPlatformSetCustomization::FilterTreeForQL(const TSharedPtr<FNiagaraDeviceProfileViewModel>& SourceTree, TSharedPtr<FNiagaraDeviceProfileViewModel>& FilteredTree, int32 QualityLevelMask)
@@ -685,6 +683,7 @@ void FNiagaraPlatformSetCustomization::CreateDeviceProfileTree()
 	check(Settings);
 
 	int32 NumQualityLevels = Settings->QualityLevels.Num();
+	FilteredDeviceProfileTrees.Reset();
 	FilteredDeviceProfileTrees.SetNum(NumQualityLevels);
 	
 	for (TSharedPtr<FNiagaraDeviceProfileViewModel>& FullDeviceRoot : FullDeviceProfileTree)
@@ -708,15 +707,13 @@ void FNiagaraPlatformSetCustomization::CreateDeviceProfileTree()
 
 TSharedRef<SWidget> FNiagaraPlatformSetCustomization::GenerateDeviceProfileTreeWidget(int32 QualityLevel)
 {
+	FullDeviceProfileTree.Reset();
 	if (FullDeviceProfileTree.Num() == 0)
 	{
 		CreateDeviceProfileTree();	
 	}
-
-	if (QualityLevelMenuContents[QualityLevel].IsValid())
-	{
-		return QualityLevelMenuContents[QualityLevel].ToSharedRef();
-	}
+	
+	QualityLevelMenuContents[QualityLevel].Reset();
 
 	TArray<TSharedPtr<FNiagaraDeviceProfileViewModel>>* TreeToUse = &FullDeviceProfileTree;
 	if (QualityLevel != INDEX_NONE)
@@ -741,8 +738,16 @@ TSharedRef<ITableRow> FNiagaraPlatformSetCustomization::OnGenerateDeviceProfileT
 	TSharedPtr<SHorizontalBox> RowContainer;
 	SAssignNew(RowContainer, SHorizontalBox);
 
-	int32 ProfileMask = TargetPlatformSet->GetEffectQualityMaskForDeviceProfile(InItem->Profile);
-	FText NameTooltip = FText::Format(LOCTEXT("ProfileQLTooltipFmt", "Effects Quality: {0}"), FNiagaraPlatformSet::GetQualityLevelMaskText(ProfileMask));
+	FNiagaraPlatformSetEnabledStateDetails Details;
+	FNiagaraPlatformSetEnabledState EnabledState = TargetPlatformSet->IsEnabled(InItem->Profile, QualityLevel, &Details);
+
+	FText NameTooltip = FText::Format(LOCTEXT("ProfileQLTooltipFmt", "Effects Quality: {0}"), FNiagaraPlatformSet::GetQualityLevelMaskText(QualityLevel));
+
+	const UNiagaraSettings* Settings = GetDefault<UNiagaraSettings>();
+	check(Settings);
+
+	int32 NumQualityLevels = Settings->QualityLevels.Num();
+	
 	
 	//Top level profile. Look for a platform icon.
 	if (InItem->Profile->Parent == nullptr)
@@ -774,22 +779,114 @@ TSharedRef<ITableRow> FNiagaraPlatformSetCustomization::OnGenerateDeviceProfileT
 	FSlateColor TextColor(FSlateColor::UseForeground());
 	ENiagaraPlatformSelectionState CurrentState = TargetPlatformSet->GetDeviceProfileState(InItem->Profile, QualityLevel);
 	
-	if (CurrentState == ENiagaraPlatformSelectionState::Enabled)
+	FSlateColor ActiveColor = FSlateColor::UseForeground();
+	FSlateColor InactiveColor = FSlateColor::UseSubduedForeground();
+	FSlateColor DisabledColor = FSlateColor(FLinearColor(FVector4(1, 0, 0, 1)));
+	if (EnabledState.bIsActive)
 	{
-		TextStyleName = "RichTextBlock.Bold";
-		TextColor = FSlateColor(FLinearColor(FVector4(0,1,0,1)));
+		TextColor = ActiveColor;
 	}
-	else if (CurrentState == ENiagaraPlatformSelectionState::Disabled)
+	else
 	{
-		TextStyleName = "RichTextBlock.Italic";
-		TextColor = FSlateColor(FLinearColor(FVector4(1,0,0,1)));
+		if (EnabledState.bCanBeActive)
+		{
+			TextColor = InactiveColor;
+		}
+		else
+		{
+			TextColor = DisabledColor;
+		}
 	}
 
+	TSharedPtr<SVerticalBox> TooltipContentsBox = SNew(SVerticalBox);
+ 	if (EnabledState.bIsActive)
+ 	{
+		TooltipContentsBox->AddSlot()
+ 		[
+ 			SNew(STextBlock)
+ 			.TextStyle(FEditorStyle::Get(), TextStyleName)
+ 			.ColorAndOpacity(ActiveColor)
+ 			.Text(LOCTEXT("DPActiveTooltip", "Active"))
+ 		];
+ 	}
+ 	else
+ 	{
+ 		if(EnabledState.bCanBeActive)
+ 		{
+			TooltipContentsBox->AddSlot()
+			[
+				SNew(STextBlock)
+				.TextStyle(FEditorStyle::Get(), TextStyleName)
+				.ColorAndOpacity(InactiveColor)
+				.Text(LOCTEXT("DPInActiveTooltip", "Inactive Due To:"))
+			];
+ 			for (FText Reason : Details.ReasonsForInActive)
+ 			{			
+				TooltipContentsBox->AddSlot()
+ 				[
+ 					SNew(STextBlock)
+ 					.TextStyle(FEditorStyle::Get(), TextStyleName)
+ 					.ColorAndOpacity(InactiveColor)
+ 					.Text(Reason)
+ 				];
+ 			}
+ 		}
+ 		else
+ 		{
+			TooltipContentsBox->AddSlot()
+			[
+				SNew(STextBlock)
+				.TextStyle(FEditorStyle::Get(), TextStyleName)
+				.ColorAndOpacity(DisabledColor)
+				.Text(LOCTEXT("DPDisabledTooltip", "Disabled Due To:"))
+			];
+ 			for (FText Reason : Details.ReasonsForDisabled)
+ 			{			
+				TooltipContentsBox->AddSlot()
+ 				[
+ 					SNew(STextBlock)
+ 					.TextStyle(FEditorStyle::Get(), TextStyleName)
+ 					.ColorAndOpacity(DisabledColor)
+ 					.Text(Reason)
+ 				];
+ 			}
+ 		}
+	}
+	
+	TooltipContentsBox->AddSlot()
+		[
+			SNew(STextBlock)
+			.TextStyle(FEditorStyle::Get(), TextStyleName)
+			.ColorAndOpacity(ActiveColor)
+			.Text(LOCTEXT("QLTooltipAvailableListHeader", "Available:"))
+		];
 
-	int32 QualityLevelMask = FNiagaraPlatformSet::CreateQualityLevelMask(QualityLevel);
-	if ((ProfileMask & QualityLevelMask) == 0)
+	int32 DefaultQL = FNiagaraPlatformSet::QualityLevelFromMask(Details.DefaultQualityMask);
+	for (int32 QL = 0; QL < NumQualityLevels; ++QL)
 	{
-		TextColor = FSlateColor::UseSubduedForeground();
+		if(Details.AvailableQualityMask & (1<<QL))
+		{
+			if (DefaultQL == QL)
+			{
+				TooltipContentsBox->AddSlot()
+					[
+						SNew(STextBlock)
+						.TextStyle(FEditorStyle::Get(), TextStyleName)
+						.ColorAndOpacity(ActiveColor)
+						.Text(FText::Format(LOCTEXT("QLTooltipAvailableListDefaultFmt", "{0} (Default)"), FNiagaraPlatformSet::GetQualityLevelText(QL)))
+					];
+			}
+			else
+			{
+				TooltipContentsBox->AddSlot()
+					[
+						SNew(STextBlock)
+						.TextStyle(FEditorStyle::Get(), TextStyleName)
+						.ColorAndOpacity(ActiveColor)
+						.Text(FNiagaraPlatformSet::GetQualityLevelText(QL))
+					];
+			}
+		}
 	}
 
 	RowContainer->AddSlot()
@@ -802,7 +899,12 @@ TSharedRef<ITableRow> FNiagaraPlatformSetCustomization::OnGenerateDeviceProfileT
 			.OnClicked(this, &FNiagaraPlatformSetCustomization::OnProfileMenuButtonClicked, InItem, QualityLevel, false)
 			.IsEnabled(this, &FNiagaraPlatformSetCustomization::GetProfileMenuItemEnabled, InItem, QualityLevel)
 			.ForegroundColor(TextColor)
-			.ToolTipText(NameTooltip)
+			.ToolTip(
+				SNew(SToolTip)
+				[
+					TooltipContentsBox.ToSharedRef()
+				]
+			)
 			[
 				SNew(STextBlock)
 				.TextStyle(FEditorStyle::Get(), TextStyleName)
@@ -853,7 +955,7 @@ enum class EProfileButtonMode
 
 static EProfileButtonMode GetProfileMenuButtonMode(FNiagaraPlatformSet* PlatformSet, TSharedPtr<FNiagaraDeviceProfileViewModel> Item, int32 QualityLevel)
 {
-	int32 Mask = PlatformSet->GetEffectQualityMaskForDeviceProfile(Item->Profile);
+	int32 Mask = PlatformSet->GetAvailableQualityMaskForDeviceProfile(Item->Profile);
 	int32 QLMask = FNiagaraPlatformSet::CreateQualityLevelMask(QualityLevel);
 
 	if ((Mask & QLMask) == 0)
@@ -864,6 +966,7 @@ static EProfileButtonMode GetProfileMenuButtonMode(FNiagaraPlatformSet* Platform
 	ENiagaraPlatformSelectionState CurrentState = PlatformSet->GetDeviceProfileState(Item->Profile, QualityLevel);
 
 	bool bQualityEnabled = PlatformSet->IsEffectQualityEnabled(QualityLevel);
+	//FNiagaraPlatformSetEnabledState Enabled = PlatformSet->IsEnabled(Item->Profile, QualityLevel);
 	bool bIsDefault = CurrentState == ENiagaraPlatformSelectionState::Default;
 
 	if (bIsDefault && bQualityEnabled)
@@ -972,7 +1075,6 @@ FReply FNiagaraPlatformSetCustomization::OnProfileMenuButtonClicked(TSharedPtr<F
 
 void FNiagaraPlatformSetCustomization::OnGetDeviceProfileTreeChildren(TSharedPtr<FNiagaraDeviceProfileViewModel> InItem, TArray< TSharedPtr<FNiagaraDeviceProfileViewModel> >& OutChildren, int32 QualityLevel)
 {
-	if (TargetPlatformSet->GetDeviceProfileState(InItem->Profile, QualityLevel) == ENiagaraPlatformSelectionState::Default)
 	{
 		OutChildren = InItem->Children;
 	}
@@ -998,8 +1100,8 @@ void FNiagaraPlatformSetCustomization::QLCheckStateChanged(ECheckBoxState CheckS
 
 void FNiagaraPlatformSetCustomization::OnPropertyValueChanged()
 {
-	GenerateQualityLevelSelectionWidgets();
 	UpdateCachedConflicts();
+	GenerateQualityLevelSelectionWidgets();
 	TargetPlatformSet->OnChanged();
 }
 
@@ -1058,6 +1160,9 @@ protected:
 	// e.g. Tab or Key_Up
 	virtual FReply OnPreviewKeyDown(const FGeometry& MyGeometry, const FKeyEvent& KeyEvent) override;
 
+	/** Handles entering in a command */
+	FText GetTooltipText()const;
+
 	/** Handles entering in a command */
 	void OnTextCommitted(const FText& InText, ETextCommit::Type CommitInfo);
 
@@ -1195,6 +1300,7 @@ void SNiagaraConsoleInputBox::Construct(const FArguments& InArgs)
 			+ SHorizontalBox::Slot()
 			[
 				SAssignNew(InputText, SMultiLineEditableTextBox)
+				.ToolTipText_Raw(this, &SNiagaraConsoleInputBox::GetTooltipText)
 				.AllowMultiLine(false)
 				.ClearTextSelectionOnFocusLoss(true)
 				.SelectAllTextWhenFocused(true)
@@ -1298,11 +1404,14 @@ TSharedRef<ITableRow> SNiagaraConsoleInputBox::MakeSuggestionListItemWidget(TSha
 	SanitizedText.ReplaceInline(TEXT("\r"), TEXT(" "), ESearchCase::CaseSensitive);
 	SanitizedText.ReplaceInline(TEXT("\n"), TEXT(" "), ESearchCase::CaseSensitive);
 
+	IConsoleVariable* CVar = IConsoleManager::Get().FindConsoleVariable(*SanitizedText);
+	
 	return
 		SNew(STableRow< TSharedPtr<FString> >, OwnerTable)
 		[
 			SNew(STextBlock)
 			.Text(FText::FromString(SanitizedText))
+			.ToolTipText(CVar ? FText::FromString(CVar->GetHelp()) : FText::GetEmpty())
 			.TextStyle(FEditorStyle::Get(), "Log.Normal")
 			.HighlightText(Suggestions.SuggestionsHighlight)
 		];
@@ -1377,6 +1486,15 @@ void SNiagaraConsoleInputBox::OnTextChanged(const FText& InText)
 	}
 }
 
+FText SNiagaraConsoleInputBox::GetTooltipText()const
+{
+	if (IConsoleVariable* CVar = IConsoleManager::Get().FindConsoleVariable(*WorkingText.ToString()))
+	{
+		return FText::FromString(CVar->GetHelp());
+	}
+	return FText::GetEmpty();
+}
+
 void SNiagaraConsoleInputBox::OnTextCommitted(const FText& InText, ETextCommit::Type CommitInfo)
 {
 	WorkingText = InText;
@@ -1603,6 +1721,11 @@ void FNiagaraPlatformSetCVarConditionCustomization::CustomizeChildren(TSharedRef
 	ChildBuilder.AddProperty(MinFloatHandle.ToSharedRef()).Visibility(FloatVisAttr);
 	MaxFloatHandle = StructPropertyHandle->GetChildHandle(GET_MEMBER_NAME_CHECKED(FNiagaraPlatformSetCVarCondition, MaxFloat));
 	ChildBuilder.AddProperty(MaxFloatHandle.ToSharedRef()).Visibility(FloatVisAttr);
+
+	TSharedPtr<IPropertyHandle> PassHandle = StructPropertyHandle->GetChildHandle(GET_MEMBER_NAME_CHECKED(FNiagaraPlatformSetCVarCondition, PassResponse));
+	ChildBuilder.AddProperty(PassHandle.ToSharedRef());
+	TSharedPtr<IPropertyHandle> FailHandle = StructPropertyHandle->GetChildHandle(GET_MEMBER_NAME_CHECKED(FNiagaraPlatformSetCVarCondition, FailResponse));
+	ChildBuilder.AddProperty(FailHandle.ToSharedRef());
 }
 
 EVisibility FNiagaraPlatformSetCVarConditionCustomization::BoolPropertyVisibility() const
diff --git a/Engine/Plugins/FX/Niagara/Source/NiagaraEditorWidgets/Private/SNiagaraScalabilityPreviewSettings.cpp b/Engine/Plugins/FX/Niagara/Source/NiagaraEditorWidgets/Private/SNiagaraScalabilityPreviewSettings.cpp
--- a/Engine/Plugins/FX/Niagara/Source/NiagaraEditorWidgets/Private/SNiagaraScalabilityPreviewSettings.cpp
+++ b/Engine/Plugins/FX/Niagara/Source/NiagaraEditorWidgets/Private/SNiagaraScalabilityPreviewSettings.cpp
@@ -251,19 +251,22 @@ void SNiagaraScalabilityPreviewSettings::CreateDeviceProfileTree()
 	{
 		UDeviceProfile* DeviceProfile = CastChecked<UDeviceProfile>(Profile);
 
-		TFunction<void(int32&, UDeviceProfile*)> FindDepth;
-		FindDepth = [&](int32& Depth, UDeviceProfile* CurrProfile)
+		if (DeviceProfile && DeviceProfile->IsVisibleForAssets())
 		{
-			if (CurrProfile->Parent)
+			TFunction<void(int32&, UDeviceProfile*)> FindDepth;
+			FindDepth = [&](int32& Depth, UDeviceProfile* CurrProfile)
 			{
-				FindDepth(++Depth, Cast<UDeviceProfile>(CurrProfile->Parent));
-			}
-		};
+				if (CurrProfile->Parent)
+				{
+					FindDepth(++Depth, Cast<UDeviceProfile>(CurrProfile->Parent));
+				}
+			};
 
-		int32 ProfileDepth = 0;
-		FindDepth(ProfileDepth, DeviceProfile);
-		DeviceProfilesByHierarchyLevel.SetNum(FMath::Max(ProfileDepth+1, DeviceProfilesByHierarchyLevel.Num()));
-		DeviceProfilesByHierarchyLevel[ProfileDepth].Add(DeviceProfile);
+			int32 ProfileDepth = 0;
+			FindDepth(ProfileDepth, DeviceProfile);
+			DeviceProfilesByHierarchyLevel.SetNum(FMath::Max(ProfileDepth + 1, DeviceProfilesByHierarchyLevel.Num()));
+			DeviceProfilesByHierarchyLevel[ProfileDepth].Add(DeviceProfile);
+		}
 	}
 	
 	FullDeviceProfileTree.Reset(DeviceProfilesByHierarchyLevel[0].Num());
@@ -339,7 +342,7 @@ TSharedRef<ITableRow> SNiagaraScalabilityPreviewSettings::OnGenerateDeviceProfil
 	TSharedPtr<SHorizontalBox> RowContainer;
 	SAssignNew(RowContainer, SHorizontalBox);
 
-	int32 ProfileMask = ScalabilityViewModel->GetPreviewPlatformSet()->GetEffectQualityMaskForDeviceProfile(InItem->Profile);
+	int32 ProfileMask = ScalabilityViewModel->GetPreviewPlatformSet()->GetActiveQualityMaskForDeviceProfile(InItem->Profile);
 	FText NameTooltip = FText::Format(LOCTEXT("ProfileQLTooltipFmt", "Effects Quality: {0}"), FNiagaraPlatformSet::GetQualityLevelMaskText(ProfileMask));
 	
 	//Top level profile. Look for a platform icon.
diff --git a/Engine/Source/Runtime/Core/Private/HAL/ConsoleManager.cpp b/Engine/Source/Runtime/Core/Private/HAL/ConsoleManager.cpp
--- a/Engine/Source/Runtime/Core/Private/HAL/ConsoleManager.cpp
+++ b/Engine/Source/Runtime/Core/Private/HAL/ConsoleManager.cpp
@@ -1523,6 +1523,11 @@ void FConsoleManager::UnregisterConsoleObject(const TCHAR* Name, bool bKeepState
 	{
 		IConsoleVariable* CVar = Object->AsVariable();
 
+		if (CVar)
+		{
+			ConsoleVariableUnregisteredDelegate.Broadcast(CVar);
+		}
+
 		if(CVar && bKeepState)
 		{
 			// to be able to restore the value if we just recompile a module
@@ -2146,6 +2151,11 @@ void FConsoleManager::OnCVarChanged()
 	bCallAllConsoleVariableSinks = true;
 }
 
+FConsoleVariableMulticastDelegate& FConsoleManager::OnCVarUnregistered()
+{
+	return ConsoleVariableUnregisteredDelegate;
+}
+
 #if !(UE_BUILD_SHIPPING || UE_BUILD_TEST)
 uint32 GConsoleManagerSinkTestCounter = 0;
 void TestSinkCallback()
diff --git a/Engine/Source/Runtime/Core/Private/Misc/ConfigCacheIni.cpp b/Engine/Source/Runtime/Core/Private/Misc/ConfigCacheIni.cpp
--- a/Engine/Source/Runtime/Core/Private/Misc/ConfigCacheIni.cpp
+++ b/Engine/Source/Runtime/Core/Private/Misc/ConfigCacheIni.cpp
@@ -5561,7 +5561,14 @@ FConfigCacheIni* FConfigCacheIni::ForPlatform(FName PlatformName)
 #if WITH_EDITOR
 	// they are likely already loaded, but just block to make sure
 	{
-		GetPlatformConfigFutures()[PlatformName].Get();
+		if (auto* PlatformCache = GetPlatformConfigFutures().Find(PlatformName))
+		{
+			PlatformCache->Get();
+		}
+		else
+		{
+			return GConfig;
+		}
 	}
 #endif
 
diff --git a/Engine/Source/Runtime/Core/Public/HAL/ConsoleManager.h b/Engine/Source/Runtime/Core/Public/HAL/ConsoleManager.h
--- a/Engine/Source/Runtime/Core/Public/HAL/ConsoleManager.h
+++ b/Engine/Source/Runtime/Core/Public/HAL/ConsoleManager.h
@@ -43,6 +43,8 @@ public:
 
 	void OnCVarChanged();
 
+	virtual FConsoleVariableMulticastDelegate& OnCVarUnregistered()override;
+
 	// interface IConsoleManager -----------------------------------
 
 	virtual IConsoleVariable* RegisterConsoleVariable(const TCHAR* Name, bool DefaultValue, const TCHAR* Help, uint32 Flags) override;
@@ -90,6 +92,8 @@ private: // ----------------------------------------------------
 	TMap<FString, TArray<FString>>	HistoryEntriesMap;
 	TArray<FConsoleCommandDelegate>	ConsoleVariableChangeSinks;
 
+	FConsoleVariableMulticastDelegate ConsoleVariableUnregisteredDelegate;
+
 	IConsoleThreadPropagation* ThreadPropagationCallback;
 
 	// if true the next call to CallAllConsoleVariableSinks() we will call all registered sinks
diff --git a/Engine/Source/Runtime/Core/Public/HAL/IConsoleManager.h b/Engine/Source/Runtime/Core/Public/HAL/IConsoleManager.h
--- a/Engine/Source/Runtime/Core/Public/HAL/IConsoleManager.h
+++ b/Engine/Source/Runtime/Core/Public/HAL/IConsoleManager.h
@@ -979,6 +979,8 @@ struct CORE_API IConsoleManager
 	static bool VisitPlatformCVarsForEmulation(FName PlatformName, bool bVisitPlatformDeviceProfile, TFunctionRef<void(const FString& CVarName, const FString& CVarValue, EConsoleVariableFlags SetBy)> Visit);
 #endif
 
+	virtual FConsoleVariableMulticastDelegate& OnCVarUnregistered() = 0;
+
 protected:
 	virtual ~IConsoleManager() { }
 
diff --git a/Engine/Source/Runtime/Engine/Classes/DeviceProfiles/DeviceProfile.h b/Engine/Source/Runtime/Engine/Classes/DeviceProfiles/DeviceProfile.h
--- a/Engine/Source/Runtime/Engine/Classes/DeviceProfiles/DeviceProfile.h
+++ b/Engine/Source/Runtime/Engine/Classes/DeviceProfiles/DeviceProfile.h
@@ -30,6 +30,10 @@ class ENGINE_API UDeviceProfile : public UTextureLODSettings
 	UPROPERTY(EditAnywhere, config, Category=DeviceSettings)
 	FString BaseProfileName;
 
+	/** Some asset types can reference Device Profiles, is this profile visible to those assets. */
+	UPROPERTY(EditAnywhere, config, Category = DeviceSettings)
+	uint32 bIsVisibleForAssets : 1;
+
 	/** The parent object of this profile, it is the object matching this DeviceType with the BaseProfileName */
 	UPROPERTY()
 	TObjectPtr<UDeviceProfile> Parent;
@@ -88,6 +92,8 @@ public:
 	 * @return FSelectedFragmentProperties ptr or null if FragmentTag not found.
 	 */	
 	const FSelectedFragmentProperties* GetFragmentByTag(FName& FragmentTag) const;
+
+	bool IsVisibleForAssets()const { return bIsVisibleForAssets; }
 public:
 	/** 
 	 * Access to the device profiles Texture LOD Settings
diff --git a/Engine/Source/Runtime/Engine/Private/DataDrivenCVars/DataDrivenCVars.cpp b/Engine/Source/Runtime/Engine/Private/DataDrivenCVars/DataDrivenCVars.cpp
--- a/Engine/Source/Runtime/Engine/Private/DataDrivenCVars/DataDrivenCVars.cpp
+++ b/Engine/Source/Runtime/Engine/Private/DataDrivenCVars/DataDrivenCVars.cpp
@@ -73,6 +73,18 @@ void FDataDrivenConsoleVariable::Refresh()
 
 	// make sure the cvar is registered
 	Register();
+
+	//Ensure the default value is applied, assuming no other external changes to the CVar.
+	IConsoleVariable* CVarToRefresh = IConsoleManager::Get().FindConsoleVariable(*Name);
+	if (CVarToRefresh)
+	{
+		switch (Type)
+		{
+		case FDataDrivenCVarType::CVarBool: CVarToRefresh->Set(DefaultValueBool, ECVF_SetByConstructor); break;
+		case FDataDrivenCVarType::CVarInt: CVarToRefresh->Set(DefaultValueInt, ECVF_SetByConstructor); break;
+		case FDataDrivenCVarType::CVarFloat: CVarToRefresh->Set(DefaultValueFloat, ECVF_SetByConstructor); break;
+		}
+	}
 }
 #endif
 
diff --git a/Engine/Source/Runtime/Engine/Private/DeviceProfiles/DeviceProfile.cpp b/Engine/Source/Runtime/Engine/Private/DeviceProfiles/DeviceProfile.cpp
--- a/Engine/Source/Runtime/Engine/Private/DeviceProfiles/DeviceProfile.cpp
+++ b/Engine/Source/Runtime/Engine/Private/DeviceProfiles/DeviceProfile.cpp
@@ -24,6 +24,7 @@ UDeviceProfile::UDeviceProfile(const FObjectInitializer& ObjectInitializer)
 {
 	BaseProfileName = TEXT("");
 	DeviceType = TEXT("");
+	bIsVisibleForAssets = false;
 
 	bVisible = true;
 
