Non unity fix.

diff --git a/Engine/Plugins/FX/Niagara/Source/Niagara/Classes/NiagaraPlatformSet.h b/Engine/Plugins/FX/Niagara/Source/Niagara/Classes/NiagaraPlatformSet.h
--- a/Engine/Plugins/FX/Niagara/Source/Niagara/Classes/NiagaraPlatformSet.h
+++ b/Engine/Plugins/FX/Niagara/Source/Niagara/Classes/NiagaraPlatformSet.h
@@ -5,6 +5,7 @@
 #include "CoreMinimal.h"
 #include "NiagaraPlatformSet.generated.h"
 
+class IConsoleVariable;
 class UDeviceProfile;
 
 UENUM()
