- Fix no-unity build

#rnx

diff --git a/Engine/Plugins/FX/Niagara/Source/Niagara/Classes/NiagaraPlatformSet.h b/Engine/Plugins/FX/Niagara/Source/Niagara/Classes/NiagaraPlatformSet.h
--- a/Engine/Plugins/FX/Niagara/Source/Niagara/Classes/NiagaraPlatformSet.h
+++ b/Engine/Plugins/FX/Niagara/Source/Niagara/Classes/NiagaraPlatformSet.h
@@ -3,9 +3,9 @@
 #pragma once
 
 #include "CoreMinimal.h"
+#include "HAL/IConsoleManager.h"
 #include "NiagaraPlatformSet.generated.h"
 
-class IConsoleVariable;
 class UDeviceProfile;
 
 UENUM()
diff --git a/Engine/Plugins/FX/Niagara/Source/Niagara/Private/NiagaraPlatformSet.cpp b/Engine/Plugins/FX/Niagara/Source/Niagara/Private/NiagaraPlatformSet.cpp
--- a/Engine/Plugins/FX/Niagara/Source/Niagara/Private/NiagaraPlatformSet.cpp
+++ b/Engine/Plugins/FX/Niagara/Source/Niagara/Private/NiagaraPlatformSet.cpp
@@ -10,8 +10,6 @@
 #include "NiagaraSettings.h"
 #include "SystemSettings.h"
 
-#include "HAL/IConsoleManager.h"
-
 #if WITH_EDITOR
 #include "PlatformInfo.h"
 #endif
