- IWYU fix

#rnx

diff --git a/Engine/Plugins/FX/Niagara/Source/NiagaraEditor/Private/Customizations/NiagaraPlatformSetCustomization.cpp b/Engine/Plugins/FX/Niagara/Source/NiagaraEditor/Private/Customizations/NiagaraPlatformSetCustomization.cpp
--- a/Engine/Plugins/FX/Niagara/Source/NiagaraEditor/Private/Customizations/NiagaraPlatformSetCustomization.cpp
+++ b/Engine/Plugins/FX/Niagara/Source/NiagaraEditor/Private/Customizations/NiagaraPlatformSetCustomization.cpp
@@ -20,6 +20,7 @@
 #include "Widgets/Input/SMenuAnchor.h"
 #include "Widgets/Layout/SWrapBox.h"
 #include "Widgets/Text/STextBlock.h"
+#include "Widgets/SToolTip.h"
 #include "NiagaraSettings.h"
 #include "IDetailPropertyRow.h"
 
