Fixed crash when adding a CVar condition for a CVar that doesn't exist.

#rnx

diff --git a/Engine/Plugins/FX/Niagara/Source/Niagara/Private/NiagaraPlatformSet.cpp b/Engine/Plugins/FX/Niagara/Source/Niagara/Private/NiagaraPlatformSet.cpp
--- a/Engine/Plugins/FX/Niagara/Source/Niagara/Private/NiagaraPlatformSet.cpp
+++ b/Engine/Plugins/FX/Niagara/Source/Niagara/Private/NiagaraPlatformSet.cpp
@@ -920,16 +920,21 @@ IConsoleVariable* FNiagaraPlatformSet::GetCVar(FName CVarName)
 {
 	FScopeLock Lock(&CachedCVarInfoCritSec);
 	
-	FCachedCVarInfo* CachedInfo = CachedCVarInfo.Find(CVarName);
-	if (CachedInfo == nullptr)
+	if (FCachedCVarInfo* CachedInfo = CachedCVarInfo.Find(CVarName))
 	{
-		IConsoleManager& CMan = IConsoleManager::Get();
-		CachedInfo = &CachedCVarInfo.Add(CVarName);
-		CachedInfo->CVar = CMan.FindConsoleVariable(*CVarName.ToString());
+		return CachedInfo->CVar;
+	}
+
+	IConsoleManager& CMan = IConsoleManager::Get();
+	if (IConsoleVariable* CVar = CMan.FindConsoleVariable(*CVarName.ToString()))
+	{
+		FCachedCVarInfo* CachedInfo = &CachedCVarInfo.Add(CVarName);
+		CachedInfo->CVar = CVar;
 		CachedInfo->ChangedHandle = CachedInfo->CVar->OnChangedDelegate().AddStatic(FNiagaraPlatformSet::OnCVarChanged);
+		return CVar;
 	}
 
-	return CachedInfo->CVar;
+	return nullptr;
 }
 
 void FNiagaraPlatformSet::RefreshScalability()
