- Add component / asset to check to help track down issues

#rnx

diff --git a/Engine/Plugins/FX/Niagara/Source/Niagara/Private/NiagaraScalabilityManager.cpp b/Engine/Plugins/FX/Niagara/Source/Niagara/Private/NiagaraScalabilityManager.cpp
--- a/Engine/Plugins/FX/Niagara/Source/Niagara/Private/NiagaraScalabilityManager.cpp
+++ b/Engine/Plugins/FX/Niagara/Source/Niagara/Private/NiagaraScalabilityManager.cpp
@@ -146,7 +146,7 @@ void FNiagaraScalabilityManager::Register(UNiagaraComponent* Component)
 
 void FNiagaraScalabilityManager::Unregister(UNiagaraComponent* Component)
 {
-	check(Component->ScalabilityManagerHandle != INDEX_NONE);
+	checkf(Component->ScalabilityManagerHandle != INDEX_NONE, TEXT("Component(%s) is not in the scalability manager but is being unregistered.  Asset(%s)."), *GetFullNameSafe(Component), *GetFullNameSafe(Component->GetAsset()));
 
 	int32 IndexToRemove = Component->ScalabilityManagerHandle;
 	Component->ScalabilityManagerHandle = INDEX_NONE;
