- Fix crash if bCanExecute is not true

#rnx

diff --git a/Engine/Plugins/FX/Niagara/Source/Niagara/Private/NiagaraSystemSimulation.cpp b/Engine/Plugins/FX/Niagara/Source/Niagara/Private/NiagaraSystemSimulation.cpp
--- a/Engine/Plugins/FX/Niagara/Source/Niagara/Private/NiagaraSystemSimulation.cpp
+++ b/Engine/Plugins/FX/Niagara/Source/Niagara/Private/NiagaraSystemSimulation.cpp
@@ -786,8 +786,15 @@ void FNiagaraSystemSimulation::Destroy()
 
 	check(GetSystemInstances(ENiagaraSystemInstanceState::Spawning).Num() == 0);
 
-	SpawnExecContext->Parameters.UnbindFromSourceStores();
-	UpdateExecContext->Parameters.UnbindFromSourceStores();
+	// Can be nullptr if bCanExecute is false
+	if (SpawnExecContext)
+	{
+		SpawnExecContext->Parameters.UnbindFromSourceStores();
+	}
+	if (UpdateExecContext)
+	{
+		UpdateExecContext->Parameters.UnbindFromSourceStores();
+	}
 
 	World = nullptr;
 }
