Enable Niagara RayTracingTransforms CS for Vulkan RT: it was only compiled when ray tracing pipelines were available but the shader only prepares transforms for RayTracingScene and can be used for both RT pipeline/inline.

diff --git a/Engine/Config/VulkanPC/DataDrivenPlatformInfo.ini b/Engine/Config/VulkanPC/DataDrivenPlatformInfo.ini
--- a/Engine/Config/VulkanPC/DataDrivenPlatformInfo.ini
+++ b/Engine/Config/VulkanPC/DataDrivenPlatformInfo.ini
@@ -52,7 +52,7 @@ bSupportsRenderTargetWriteMask = false
 bSupportsRayTracing = true
 bSupportsRayTracingShaders = false
 bSupportsInlineRayTracing = true
-bSupportsRayTracingIndirectInstanceData = false
+bSupportsRayTracingIndirectInstanceData = true
 bSupportsGPUSkinCache = true
 bSupportsByteBufferComputeShaders = true
 
diff --git a/Engine/Plugins/FX/Niagara/Source/NiagaraVertexFactories/Private/NiagaraGPURayTracingTransformsShader.cpp b/Engine/Plugins/FX/Niagara/Source/NiagaraVertexFactories/Private/NiagaraGPURayTracingTransformsShader.cpp
--- a/Engine/Plugins/FX/Niagara/Source/NiagaraVertexFactories/Private/NiagaraGPURayTracingTransformsShader.cpp
+++ b/Engine/Plugins/FX/Niagara/Source/NiagaraVertexFactories/Private/NiagaraGPURayTracingTransformsShader.cpp
@@ -14,7 +14,7 @@ IMPLEMENT_GLOBAL_SHADER(FNiagaraGPURayTracingTransformsCS, "/Plugin/FX/Niagara/P
 
 bool FNiagaraGPURayTracingTransformsCS::ShouldCompilePermutation(const FGlobalShaderPermutationParameters& Parameters)
 {
-	return ShouldCompileRayTracingShadersForProject(Parameters.Platform) && RHISupportsComputeShaders(Parameters.Platform) && !IsMobilePlatform(Parameters.Platform);
+	return IsRayTracingEnabledForProject(Parameters.Platform) && RHISupportsRayTracing(Parameters.Platform) && RHISupportsComputeShaders(Parameters.Platform) && !IsMobilePlatform(Parameters.Platform);
 }
 
 void FNiagaraGPURayTracingTransformsCS::ModifyCompilationEnvironment(const FGlobalShaderPermutationParameters& Parameters, FShaderCompilerEnvironment& OutEnvironment)
