Fix for uninitialized FNiagaraPlatformSetRedirect::Mode causing cook failures

diff --git a/Engine/Plugins/FX/Niagara/Source/Niagara/Classes/NiagaraPlatformSet.h b/Engine/Plugins/FX/Niagara/Source/Niagara/Classes/NiagaraPlatformSet.h
--- a/Engine/Plugins/FX/Niagara/Source/Niagara/Classes/NiagaraPlatformSet.h
+++ b/Engine/Plugins/FX/Niagara/Source/Niagara/Classes/NiagaraPlatformSet.h
@@ -283,6 +283,8 @@ struct FNiagaraPlatformSetRedirect
 {
 	GENERATED_BODY()
 
+	FNiagaraPlatformSetRedirect();
+
 	/** The names of any device profile entry that will apply this redirect. */
 	UPROPERTY(EditAnywhere, Category = Redirect)
 	TArray<FName> ProfileNames;
diff --git a/Engine/Plugins/FX/Niagara/Source/Niagara/Private/NiagaraPlatformSet.cpp b/Engine/Plugins/FX/Niagara/Source/Niagara/Private/NiagaraPlatformSet.cpp
--- a/Engine/Plugins/FX/Niagara/Source/Niagara/Private/NiagaraPlatformSet.cpp
+++ b/Engine/Plugins/FX/Niagara/Source/Niagara/Private/NiagaraPlatformSet.cpp
@@ -1551,4 +1551,8 @@ bool FNiagaraPlatformSetCVarCondition::Evaluate(const UDeviceProfile* DeviceProf
 	return false;
 }
 
+FNiagaraPlatformSetRedirect::FNiagaraPlatformSetRedirect()
+	: Mode(ENiagaraDeviceProfileRedirectMode::CVar)
+{}
+
 #undef LOCTEXT_NAMESPACE
\ No newline at end of file
