Removed a few high-frequency/low-value Niagara stats to improve profiler readability and stability

diff --git a/Engine/Plugins/FX/Niagara/Source/Niagara/Private/NiagaraRendererMeshes.cpp b/Engine/Plugins/FX/Niagara/Source/Niagara/Private/NiagaraRendererMeshes.cpp
--- a/Engine/Plugins/FX/Niagara/Source/Niagara/Private/NiagaraRendererMeshes.cpp
+++ b/Engine/Plugins/FX/Niagara/Source/Niagara/Private/NiagaraRendererMeshes.cpp
@@ -21,7 +21,6 @@
 #include "FXRenderingUtils.h"
 
 DECLARE_CYCLE_STAT(TEXT("Generate Mesh Vertex Data [GT]"), STAT_NiagaraGenMeshVertexData, STATGROUP_Niagara);
-DECLARE_CYCLE_STAT(TEXT("Render Meshes [RT]"), STAT_NiagaraRenderMeshes, STATGROUP_Niagara);
 
 DECLARE_DWORD_COUNTER_STAT(TEXT("NumMeshesRenderer"), STAT_NiagaraNumMeshes, STATGROUP_Niagara);
 DECLARE_DWORD_COUNTER_STAT(TEXT("NumMeshVerts"), STAT_NiagaraNumMeshVerts, STATGROUP_Niagara);
@@ -1160,7 +1159,6 @@ void FNiagaraRendererMeshes::GetDynamicMeshElements(const TArray<const FSceneVie
 		return;
 	}
 
-	SCOPE_CYCLE_COUNTER(STAT_NiagaraRenderMeshes);
 #if STATS
 	FScopeCycleCounter EmitterStatsCounter(EmitterStatID);
 #endif
@@ -1355,8 +1353,6 @@ void FNiagaraRendererMeshes::GetDynamicRayTracingInstances(FRayTracingMaterialGa
 	ParticleMeshRenderData.bNeedsCull = false;
 	ParticleMeshRenderData.bSortCullOnGpu = false;
 
-	SCOPE_CYCLE_COUNTER(STAT_NiagaraRenderMeshes);
-
 	PrepareParticleRenderBuffers(ParticleMeshRenderData, Context.RayTracingMeshResourceCollector.GetDynamicReadBuffer());
 
 
diff --git a/Engine/Plugins/FX/Niagara/Source/Niagara/Private/NiagaraRendererRibbons.cpp b/Engine/Plugins/FX/Niagara/Source/Niagara/Private/NiagaraRendererRibbons.cpp
--- a/Engine/Plugins/FX/Niagara/Source/Niagara/Private/NiagaraRendererRibbons.cpp
+++ b/Engine/Plugins/FX/Niagara/Source/Niagara/Private/NiagaraRendererRibbons.cpp
@@ -14,7 +14,6 @@
 #include "NiagaraCullProxyComponent.h"
 
 DECLARE_CYCLE_STAT(TEXT("Generate Ribbon Vertex Data [GT]"), STAT_NiagaraGenRibbonVertexData, STATGROUP_Niagara);
-DECLARE_CYCLE_STAT(TEXT("Render Ribbons [RT]"), STAT_NiagaraRenderRibbons, STATGROUP_Niagara);
 
 int32 GNiagaraRibbonTessellationEnabled = 1;
 static FAutoConsoleVariableRef CVarNiagaraRibbonTessellationEnabled(
@@ -449,7 +448,6 @@ void FNiagaraRendererRibbons::GenerateIndexBuffer(
 
 void FNiagaraRendererRibbons::GetDynamicMeshElements(const TArray<const FSceneView*>& Views, const FSceneViewFamily& ViewFamily, uint32 VisibilityMap, FMeshElementCollector& Collector, const FNiagaraSceneProxy *SceneProxy) const
 {
-	SCOPE_CYCLE_COUNTER(STAT_NiagaraRenderRibbons);
 	PARTICLE_PERF_STAT_CYCLES_RT(SceneProxy->GetProxyDynamicData().PerfStatsContext, GetDynamicMeshElements);
 
 	FNiagaraDynamicDataRibbon *DynamicDataRibbon = static_cast<FNiagaraDynamicDataRibbon*>(DynamicDataRender);
@@ -1443,8 +1441,7 @@ void FNiagaraRendererRibbons::GetDynamicRayTracingInstances(FRayTracingMaterialG
 	{
 		return;
 	}
-
-	SCOPE_CYCLE_COUNTER(STAT_NiagaraRenderRibbons);
+	
 	check(SceneProxy);
 
 	FNiagaraDynamicDataRibbon *DynamicDataRibbon = static_cast<FNiagaraDynamicDataRibbon*>(DynamicDataRender);
diff --git a/Engine/Plugins/FX/Niagara/Source/Niagara/Private/NiagaraRendererSprites.cpp b/Engine/Plugins/FX/Niagara/Source/Niagara/Private/NiagaraRendererSprites.cpp
--- a/Engine/Plugins/FX/Niagara/Source/Niagara/Private/NiagaraRendererSprites.cpp
+++ b/Engine/Plugins/FX/Niagara/Source/Niagara/Private/NiagaraRendererSprites.cpp
@@ -16,7 +16,6 @@
 #include "Renderer/Private/ScenePrivate.h"
 #include "NiagaraCullProxyComponent.h"
 
-DECLARE_CYCLE_STAT(TEXT("Render Sprites [RT]"), STAT_NiagaraRenderSprites, STATGROUP_Niagara);
 DECLARE_DWORD_COUNTER_STAT(TEXT("NumSprites"), STAT_NiagaraNumSprites, STATGROUP_Niagara);
 
 static int32 GbEnableNiagaraSpriteRendering = 1;
@@ -928,7 +927,6 @@ void FNiagaraRendererSprites::GetDynamicMeshElements(const TArray<const FSceneVi
 		return;
 	}
 
-	SCOPE_CYCLE_COUNTER(STAT_NiagaraRenderSprites);
 #if STATS
 	FScopeCycleCounter EmitterStatsCounter(EmitterStatID);
 #endif
@@ -1049,8 +1047,7 @@ void FNiagaraRendererSprites::GetDynamicRayTracingInstances(FRayTracingMaterialG
 	{
 		return;
 	}
-
-	SCOPE_CYCLE_COUNTER(STAT_NiagaraRenderSprites);
+	
 #if STATS
 	FScopeCycleCounter EmitterStatsCounter(EmitterStatID);
 #endif
