- Convert to soft object to not loose the path if the BP is adjusted
- Fix a missing divide by sizeof(int32) for GPU buffer offset

#rnx

diff --git a/Engine/Plugins/FX/Niagara/Source/Niagara/Private/DataInterface/NiagaraDataInterfaceUObjectPropertyReader.cpp b/Engine/Plugins/FX/Niagara/Source/Niagara/Private/DataInterface/NiagaraDataInterfaceUObjectPropertyReader.cpp
--- a/Engine/Plugins/FX/Niagara/Source/Niagara/Private/DataInterface/NiagaraDataInterfaceUObjectPropertyReader.cpp
+++ b/Engine/Plugins/FX/Niagara/Source/Niagara/Private/DataInterface/NiagaraDataInterfaceUObjectPropertyReader.cpp
@@ -80,7 +80,7 @@ namespace NDIUObjectPropertyReaderLocal
 				for ( int32 i=0; i < InstanceData_GT->PropertyGetters.Num(); ++i )
 				{
 					const FNDIPropertyGetter& PropertyGetter = InstanceData_GT->PropertyGetters[i];
-					InstanceData_ForRT->PropertyOffsets[i] = PropertyGetter.PropertyCopyFunction ? PropertyGetter.DataOffset : INDEX_NONE;
+					InstanceData_ForRT->PropertyOffsets[i] = PropertyGetter.PropertyCopyFunction ? PropertyGetter.DataOffset / sizeof(int32) : INDEX_NONE;
 				}
 				InstanceData_ForRT->PropertyData = InstanceData_GT->PropertyData;
 			}
diff --git a/Engine/Plugins/FX/Niagara/Source/Niagara/Private/DataInterface/NiagaraDataInterfaceUObjectPropertyReader.h b/Engine/Plugins/FX/Niagara/Source/Niagara/Private/DataInterface/NiagaraDataInterfaceUObjectPropertyReader.h
--- a/Engine/Plugins/FX/Niagara/Source/Niagara/Private/DataInterface/NiagaraDataInterfaceUObjectPropertyReader.h
+++ b/Engine/Plugins/FX/Niagara/Source/Niagara/Private/DataInterface/NiagaraDataInterfaceUObjectPropertyReader.h
@@ -45,7 +45,7 @@ public:
 
 	/** Optional source actor to use, if the user parameter binding is valid this will be ignored. */
 	UPROPERTY(EditAnywhere, Category = "UObjectPropertyReader")
-	TLazyObjectPtr<AActor> SourceActor;
+	TSoftObjectPtr<AActor> SourceActor;
 
 	/**
 	When an actor is bound as the object we will also search for a component of this type to bind properties to.
