- Add function for getting the inverse transform

diff --git a/Engine/Plugins/FX/Niagara/Source/Niagara/Private/DataInterface/NiagaraDataInterfaceUObjectPropertyReader.cpp b/Engine/Plugins/FX/Niagara/Source/Niagara/Private/DataInterface/NiagaraDataInterfaceUObjectPropertyReader.cpp
--- a/Engine/Plugins/FX/Niagara/Source/Niagara/Private/DataInterface/NiagaraDataInterfaceUObjectPropertyReader.cpp
+++ b/Engine/Plugins/FX/Niagara/Source/Niagara/Private/DataInterface/NiagaraDataInterfaceUObjectPropertyReader.cpp
@@ -15,6 +15,7 @@
 namespace NDIUObjectPropertyReaderLocal
 {
 	static const FName GetComponentTransformName("GetComponentTransform");
+	static const FName GetComponentInvTransformName("GetComponentInverseTransform");
 
 	struct FNDIPropertyGetter
 	{
@@ -40,6 +41,7 @@ namespace NDIUObjectPropertyReaderLocal
 		uint32										ChangeId = 0;
 
 		TOptional<FTransform>						CachedTransform;
+		TOptional<FTransform>						CachedInvTransform;
 
 		uint32 AddProperty(FNiagaraVariableBase PropertyType)
 		{
@@ -62,6 +64,7 @@ namespace NDIUObjectPropertyReaderLocal
 	struct FInstanceData_RenderThread
 	{
 		TOptional<FTransform>	CachedTransform;
+		TOptional<FTransform>	CachedInvTransform;
 		FReadBuffer				PropertyData;
 		TArray<uint32>			PropertyOffsets;
 		TArray<uint32>			GpuFunctionToPropertyRemap;
@@ -92,6 +95,14 @@ namespace NDIUObjectPropertyReaderLocal
 			if ( FInstanceData_RenderThread* InstanceData_RT = PerInstanceData_RenderThread.Find(InstanceID) )
 			{
 				InstanceData_RT->CachedTransform = InstanceData_FromGT->CachedTransform;
+				if (InstanceData_RT->CachedTransform.IsSet())
+				{
+					InstanceData_RT->CachedInvTransform = InstanceData_RT->CachedTransform->Inverse();
+				}
+				else
+				{
+					InstanceData_RT->CachedInvTransform.Reset();
+				}
 
 				for ( int32 i=0; i < InstanceData_RT->GpuFunctionToPropertyRemap.Num(); ++i )
 				{
@@ -454,6 +465,7 @@ public:
 		}
 
 		TransformInfoParam.Bind(ParameterMap, *(TEXT("TransformInfo_") + ParameterInfo.DataInterfaceHLSLSymbol));
+		InvTransformInfoParam.Bind(ParameterMap, *(TEXT("InvTransformInfo_") + ParameterInfo.DataInterfaceHLSLSymbol));
 		PropertyDataParam.Bind(ParameterMap, *(TEXT("PropertyData_") + ParameterInfo.DataInterfaceHLSLSymbol));
 		PropertyOffsetsParam.Bind(ParameterMap, *(TEXT("PropertyOffsets_") + ParameterInfo.DataInterfaceHLSLSymbol));
 	}
@@ -470,25 +482,46 @@ public:
 		SetSRVParameter(RHICmdList, ComputeShaderRHI, PropertyDataParam, InstanceData_RT->PropertyData.SRV);
 		SetShaderValueArray(RHICmdList, ComputeShaderRHI, PropertyOffsetsParam, InstanceData_RT->PropertyOffsets.GetData(), InstanceData_RT->PropertyOffsets.Num());
 
-		FVector4f TransformData[3];
-		if ( InstanceData_RT->CachedTransform.IsSet() )
 		{
-			const FQuat4f TransformRotation = FQuat4f(InstanceData_RT->CachedTransform->GetRotation());
-			TransformData[0] = FVector4f(FVector3f(InstanceData_RT->CachedTransform->GetLocation()), 1.0f);
-			TransformData[1] = FVector4f(TransformRotation.X, TransformRotation.Y, TransformRotation.Z, TransformRotation.W);
-			TransformData[2] = FVector4f(FVector3f(InstanceData_RT->CachedTransform->GetScale3D()));
+			FVector4f TransformData[3];
+			if (InstanceData_RT->CachedTransform.IsSet())
+			{
+				const FQuat4f TransformRotation = FQuat4f(InstanceData_RT->CachedTransform->GetRotation());
+				TransformData[0] = FVector4f(FVector3f(InstanceData_RT->CachedTransform->GetLocation()), 1.0f);
+				TransformData[1] = FVector4f(TransformRotation.X, TransformRotation.Y, TransformRotation.Z, TransformRotation.W);
+				TransformData[2] = FVector4f(FVector3f(InstanceData_RT->CachedTransform->GetScale3D()));
+			}
+			else
+			{
+				TransformData[0] = FVector4f(0.0f, 0.0f, 0.0f, 0.0f);
+				TransformData[1] = FVector4f(0.0f, 0.0f, 0.0f, 1.0f);
+				TransformData[2] = FVector4f(1.0f, 1.0f, 1.0f, 0.0f);
+			}
+			SetShaderValueArray(RHICmdList, ComputeShaderRHI, TransformInfoParam, TransformData, UE_ARRAY_COUNT(TransformData));
 		}
-		else
+
 		{
-			TransformData[0] = FVector4f(0.0f, 0.0f, 0.0f, 0.0f);
-			TransformData[1] = FVector4f(0.0f, 0.0f, 0.0f, 1.0f);
-			TransformData[2] = FVector4f(1.0f, 1.0f, 1.0f, 0.0f);
+			FVector4f InvTransformData[3];
+			if (InstanceData_RT->CachedInvTransform.IsSet())
+			{
+				const FQuat4f InvTransformRotation = FQuat4f(InstanceData_RT->CachedInvTransform->GetRotation());
+				InvTransformData[0] = FVector4f(FVector3f(InstanceData_RT->CachedInvTransform->GetLocation()), 1.0f);
+				InvTransformData[1] = FVector4f(InvTransformRotation.X, InvTransformRotation.Y, InvTransformRotation.Z, InvTransformRotation.W);
+				InvTransformData[2] = FVector4f(FVector3f(InstanceData_RT->CachedInvTransform->GetScale3D()));
+			}
+			else
+			{
+				InvTransformData[0] = FVector4f(0.0f, 0.0f, 0.0f, 0.0f);
+				InvTransformData[1] = FVector4f(0.0f, 0.0f, 0.0f, 1.0f);
+				InvTransformData[2] = FVector4f(1.0f, 1.0f, 1.0f, 0.0f);
+			}
+			SetShaderValueArray(RHICmdList, ComputeShaderRHI, InvTransformInfoParam, InvTransformData, UE_ARRAY_COUNT(InvTransformData));
 		}
-		SetShaderValueArray(RHICmdList, ComputeShaderRHI, TransformInfoParam, TransformData, UE_ARRAY_COUNT(TransformData));
 	}
 
 private:
 	LAYOUT_FIELD(FShaderParameter,							TransformInfoParam);
+	LAYOUT_FIELD(FShaderParameter,							InvTransformInfoParam);
 	LAYOUT_FIELD(TMemoryImageArray<FNiagaraVariableBase>,	PropertiesToRead);
 	LAYOUT_FIELD(FShaderResourceParameter,					PropertyDataParam);
 	LAYOUT_FIELD(FShaderParameter,							PropertyOffsetsParam);
@@ -705,6 +738,7 @@ bool UNiagaraDataInterfaceUObjectPropertyReader::PerInstanceTick(void* PerInstan
 
 	// Update our data store as we can not read object's async it's unsafe
 	InstanceData_GT->CachedTransform.Reset();
+	InstanceData_GT->CachedInvTransform.Reset();
 	if (ObjectBinding != nullptr)
 	{
 		// Update transform
@@ -712,10 +746,12 @@ bool UNiagaraDataInterfaceUObjectPropertyReader::PerInstanceTick(void* PerInstan
 		{
 			USceneComponent* ActorComponent = Cast<USceneComponent>(SourceActorComponentClass ? ObjectActor->FindComponentByClass(SourceActorComponentClass) : ObjectActor->GetRootComponent());
 			InstanceData_GT->CachedTransform = ActorComponent ? ActorComponent->GetComponentToWorld() : ObjectActor->GetTransform();
+			InstanceData_GT->CachedInvTransform = InstanceData_GT->CachedTransform->Inverse();
 		}
 		else if ( USceneComponent* SceneComponent = Cast<USceneComponent>(ObjectBinding) )
 		{
 			InstanceData_GT->CachedTransform = SceneComponent->GetComponentToWorld();
+			InstanceData_GT->CachedInvTransform = InstanceData_GT->CachedTransform->Inverse();
 		}
 
 		// Update properties
@@ -765,6 +801,15 @@ void UNiagaraDataInterfaceUObjectPropertyReader::GetFunctions(TArray<FNiagaraFun
 		Sig.Outputs.Emplace(FNiagaraTypeDefinition::GetVec3Def(), TEXT("Scale"));
 		Sig.SetDescription(LOCTEXT("GetComponentTransformDesc", "If the object we are bound to is an actor it will return root component transform, or the component class we bound to"));
 	}
+	{
+		FNiagaraFunctionSignature& Sig = OutFunctions.Add_GetRef(DefaultSignature);
+		Sig.Name = GetComponentInvTransformName;
+		Sig.FunctionSpecifiers.Empty();
+		Sig.Outputs.Emplace(FNiagaraTypeDefinition::GetPositionDef(), TEXT("Position"));
+		Sig.Outputs.Emplace(FNiagaraTypeDefinition::GetQuatDef(), TEXT("Rotation"));
+		Sig.Outputs.Emplace(FNiagaraTypeDefinition::GetVec3Def(), TEXT("Scale"));
+		Sig.SetDescription(LOCTEXT("GetComponentInvTransformDesc", "If the object we are bound to is an actor it will return root component inverse transform, or the component class we bound to"));
+	}
 
 	// Build property function list
 	#define NDI_PROPERTY_TYPE(TYPE) \
@@ -791,6 +836,10 @@ void UNiagaraDataInterfaceUObjectPropertyReader::GetVMExternalFunction(const FVM
 	{
 		OutFunc = FVMExternalFunction::CreateLambda([this](FVectorVMExternalFunctionContext& Context) { VMGetComponentTransform(Context); });
 	}
+	else if (BindingInfo.Name == GetComponentInvTransformName)
+	{
+		OutFunc = FVMExternalFunction::CreateLambda([this](FVectorVMExternalFunctionContext& Context) { VMGetComponentInvTransform(Context); });
+	}
 	// Bind property functions
 	#define NDI_PROPERTY_TYPE(TYPE) \
 		else if (BindingInfo.Name == FTypeHelper<TYPE>::GetFunctionName()) \
@@ -807,6 +856,7 @@ void UNiagaraDataInterfaceUObjectPropertyReader::GetVMExternalFunction(const FVM
 void UNiagaraDataInterfaceUObjectPropertyReader::GetParameterDefinitionHLSL(const FNiagaraDataInterfaceGPUParamInfo& ParameterInfo, FString& OutHLSL)
 {
 	OutHLSL.Appendf(TEXT("float4 TransformInfo_%s[3];\n"), *ParameterInfo.DataInterfaceHLSLSymbol);
+	OutHLSL.Appendf(TEXT("float4 InvTransformInfo_%s[3];\n"), *ParameterInfo.DataInterfaceHLSLSymbol);
 	OutHLSL.Appendf(TEXT("Buffer<float> PropertyData_%s;\n"), *ParameterInfo.DataInterfaceHLSLSymbol);
 	OutHLSL.Appendf(TEXT("uint4 PropertyOffsets_%s[%d];\n"), *ParameterInfo.DataInterfaceHLSLSymbol, FMath::DivideAndRoundUp(ParameterInfo.GeneratedFunctions.Num(), 4));
 }
@@ -829,7 +879,18 @@ bool UNiagaraDataInterfaceUObjectPropertyReader::GetFunctionHLSL(const FNiagaraD
 		OutHLSL.Append(TEXT("}\n"));
 		return true;
 	}
-	#define NDI_PROPERTY_TYPE(TYPE) \
+	if (FunctionInfo.DefinitionName == GetComponentInvTransformName)
+	{
+		OutHLSL.Appendf(TEXT("void %s(out bool bSuccess, out float3 Position, out float4 Rotation, out float3 Scale)\n"), *FunctionInfo.InstanceName);
+		OutHLSL.Append(TEXT("{\n"));
+		OutHLSL.Appendf(TEXT("	bSuccess = InvTransformInfo_%s[0].w > 0.0f;\n"), *ParamInfo.DataInterfaceHLSLSymbol);
+		OutHLSL.Appendf(TEXT("	Position = InvTransformInfo_%s[0].xyz;\n"), *ParamInfo.DataInterfaceHLSLSymbol);
+		OutHLSL.Appendf(TEXT("	Rotation = InvTransformInfo_%s[1];\n"), *ParamInfo.DataInterfaceHLSLSymbol);
+		OutHLSL.Appendf(TEXT("	Scale = InvTransformInfo_%s[2].xyz;\n"), *ParamInfo.DataInterfaceHLSLSymbol);
+		OutHLSL.Append(TEXT("}\n"));
+		return true;
+	}
+#define NDI_PROPERTY_TYPE(TYPE) \
 		else if (FunctionInfo.DefinitionName == FTypeHelper<TYPE>::GetFunctionName()) \
 		{ \
 			HlslBufferType = FTypeHelper<TYPE>::HlslBufferType; \
@@ -902,5 +963,29 @@ void UNiagaraDataInterfaceUObjectPropertyReader::VMGetComponentTransform(FVector
 	}
 }
 
+void UNiagaraDataInterfaceUObjectPropertyReader::VMGetComponentInvTransform(FVectorVMExternalFunctionContext& Context)
+{
+	using namespace NDIUObjectPropertyReaderLocal;
+
+	VectorVM::FUserPtrHandler<FInstanceData_GameThread> InstanceData_GT(Context);
+	FNDIOutputParam<bool> OutValid(Context);
+	FNDIOutputParam<FNiagaraPosition> OutPosition(Context);
+	FNDIOutputParam<FQuat4f> OutRotation(Context);
+	FNDIOutputParam<FVector3f> OutScale(Context);
+
+	const bool bInvTransformValid = InstanceData_GT->CachedInvTransform.IsSet();
+	const FVector3f InvTransformPosition = bInvTransformValid ? FVector3f(InstanceData_GT->CachedInvTransform->GetLocation()) : FVector3f::ZeroVector;
+	const FQuat4f InvTransformRotation = bInvTransformValid ? FQuat4f(InstanceData_GT->CachedInvTransform->GetRotation()) : FQuat4f::Identity;
+	const FVector3f InvTransformScale = bInvTransformValid ? FVector3f(InstanceData_GT->CachedInvTransform->GetScale3D()) : FVector3f::OneVector;
+
+	for (int32 i = 0; i < Context.GetNumInstances(); ++i)
+	{
+		OutValid.SetAndAdvance(bInvTransformValid);
+		OutPosition.SetAndAdvance(InvTransformPosition);
+		OutRotation.SetAndAdvance(InvTransformRotation);
+		OutScale.SetAndAdvance(InvTransformScale);
+	}
+}
+
 #undef NDI_PROPERTY_TYPES
 #undef LOCTEXT_NAMESPACE
diff --git a/Engine/Plugins/FX/Niagara/Source/Niagara/Private/DataInterface/NiagaraDataInterfaceUObjectPropertyReader.h b/Engine/Plugins/FX/Niagara/Source/Niagara/Private/DataInterface/NiagaraDataInterfaceUObjectPropertyReader.h
--- a/Engine/Plugins/FX/Niagara/Source/Niagara/Private/DataInterface/NiagaraDataInterfaceUObjectPropertyReader.h
+++ b/Engine/Plugins/FX/Niagara/Source/Niagara/Private/DataInterface/NiagaraDataInterfaceUObjectPropertyReader.h
@@ -90,6 +90,7 @@ public:
 
 private:
 	void VMGetComponentTransform(FVectorVMExternalFunctionContext& Context);
+	void VMGetComponentInvTransform(FVectorVMExternalFunctionContext& Context);
 
 	FName GetRemappedPropertyName(FName InName) const
 	{
