Added type filtering for math ops in Niagara when dragging from a pin so that unapplicable ops get filtered out, i.e. when dragging off a float, the only multiply available will be the numerics multiply. When dragging off a matrix, it will only give you matrix * matrix and matrix * vec4

diff --git a/Engine/Plugins/FX/Niagara/Source/NiagaraEditor/Private/EdGraphSchema_Niagara.cpp b/Engine/Plugins/FX/Niagara/Source/NiagaraEditor/Private/EdGraphSchema_Niagara.cpp
--- a/Engine/Plugins/FX/Niagara/Source/NiagaraEditor/Private/EdGraphSchema_Niagara.cpp
+++ b/Engine/Plugins/FX/Niagara/Source/NiagaraEditor/Private/EdGraphSchema_Niagara.cpp
@@ -396,11 +396,15 @@ TArray<TSharedPtr<FNiagaraAction_NewNode>> UEdGraphSchema_Niagara::GetGraphActio
 	bool bAllowSelectNodes = true;
 	bool bAllowStaticSwitchNodes = true;
 
+	TOptional<FNiagaraTypeDefinition> FromType;
+	TOptional<EEdGraphPinDirection> FromDirection;
+
 	UNiagaraDataInterface* FromPinDataInterface = nullptr;
 	if ( FromPin )
 	{
-		FNiagaraTypeDefinition PinType = PinToTypeDefinition(FromPin);
-		if (const UClass* FromPinClass = PinType.GetClass() )
+		FromType = PinToTypeDefinition(FromPin);
+		FromDirection = FromPin->Direction;
+		if (const UClass* FromPinClass = FromType->GetClass() )
 		{
 			bAllowOpNodes = false;
 			bAllowFunctionNodes = false;
@@ -424,13 +428,50 @@ TArray<TSharedPtr<FNiagaraAction_NewNode>> UEdGraphSchema_Niagara::GetGraphActio
 	if (bAllowOpNodes && (GbAllowAllNiagaraNodesInEmitterGraphs || bModuleGraph || bFunctionGraph || bSystemGraph))
 	{
 		const TArray<FNiagaraOpInfo>& OpInfos = FNiagaraOpInfo::GetOpInfoArray();
+
+		bool bFilterTypes = false;
+		if(FromType.IsSet() && FromDirection.IsSet())
+		{
+			bFilterTypes = true;
+		}
 		
 		for (const FNiagaraOpInfo& OpInfo : OpInfos)
 		{
-			// todo suggestion info per op?
-			UNiagaraNodeOp* OpNode = NewObject<UNiagaraNodeOp>(OwnerOfTemporaries);
-			OpNode->OpName = OpInfo.Name;
-			AddNewNodeMenuAction(NewActions, OpNode, OpInfo.FriendlyName, ENiagaraMenuSections::General, {OpInfo.Category.ToString()}, OpInfo.Description,  OpInfo.Keywords);
+			bool bAddOp = true;
+			if(bFilterTypes)
+			{
+				bAddOp = false;
+				if(FromDirection == EGPD_Output)
+				{
+					for(const FNiagaraOpInOutInfo& InputInfo : OpInfo.Inputs)
+					{
+						if(FNiagaraTypeDefinition::TypesAreAssignable(InputInfo.DataType, FromType.GetValue()))
+						{
+							bAddOp = true;
+							break;
+						}
+					}
+				}
+				else
+				{
+					for(const FNiagaraOpInOutInfo& OutputInfo : OpInfo.Outputs)
+					{
+						if(FNiagaraTypeDefinition::TypesAreAssignable(FromType.GetValue(), OutputInfo.DataType))
+						{
+							bAddOp = true;
+							break;
+						}
+					}
+				}
+			}
+
+			if(bAddOp)
+			{
+				// todo suggestion info per op?
+				UNiagaraNodeOp* OpNode = NewObject<UNiagaraNodeOp>(OwnerOfTemporaries);
+				OpNode->OpName = OpInfo.Name;
+				AddNewNodeMenuAction(NewActions, OpNode, OpInfo.FriendlyName, ENiagaraMenuSections::General, {OpInfo.Category.ToString()}, OpInfo.Description,  OpInfo.Keywords);
+			}
 		}
 	}
 
