- Mark render data dirty if the LUTOffset for the curve changes

#rnx

diff --git a/Engine/Plugins/FX/Niagara/Source/Niagara/Private/NiagaraDataInterfaceCurveBase.cpp b/Engine/Plugins/FX/Niagara/Source/Niagara/Private/NiagaraDataInterfaceCurveBase.cpp
--- a/Engine/Plugins/FX/Niagara/Source/Niagara/Private/NiagaraDataInterfaceCurveBase.cpp
+++ b/Engine/Plugins/FX/Niagara/Source/Niagara/Private/NiagaraDataInterfaceCurveBase.cpp
@@ -329,11 +329,18 @@ bool UNiagaraDataInterfaceCurveBase::Equals(const UNiagaraDataInterface* Other)
 
 void UNiagaraDataInterfaceCurveBase::CacheStaticBuffers(struct FNiagaraSystemStaticBuffers& StaticBuffers, const struct FNiagaraScriptDataInterfaceInfo& DataInterfaceInfo, bool bUsedByCPU, bool bUsedByGPU)
 {
+	const uint32 PrevLUTOffset = LUTOffset;
+
 	LUTOffset = INDEX_NONE;
 	if (bUsedByGPU && DataInterfaceInfo.IsUserDataInterface() == false)
 	{
 		LUTOffset = StaticBuffers.AddGpuData(ShaderLUT);
 	}
+
+	if (PrevLUTOffset != LUTOffset)
+	{
+		MarkRenderDataDirty();
+	}
 }
 
 #if WITH_EDITORONLY_DATA
