Validator changes

Made validators instanced objects to allow for customization via properties rather than relying only on class.
Slight modification to API, class names and display names.
Validation rules can now offer auto fixes to the stack. May allow these to be done automatically on loading invalid content in future.
Now updating validation on renderer addition and modification.

Modifies NoComponentRendererOnLow rule to be a generalized "BannedRenderers" rule.
Added "Banned Modules" rule.

diff --git a/Engine/Plugins/FX/Niagara/Source/Niagara/Classes/NiagaraEffectType.h b/Engine/Plugins/FX/Niagara/Source/Niagara/Classes/NiagaraEffectType.h
--- a/Engine/Plugins/FX/Niagara/Source/Niagara/Classes/NiagaraEffectType.h
+++ b/Engine/Plugins/FX/Niagara/Source/Niagara/Classes/NiagaraEffectType.h
@@ -406,8 +406,8 @@ class NIAGARA_API UNiagaraEffectType : public UObject
 	FNiagaraEmitterScalabilitySettingsArray EmitterScalabilitySettings;
 
 	/** A set of rules to apply when checking content. To create your own rules, write a custom class that extends UNiagaraValidationRule. */
-	UPROPERTY(EditAnywhere, Category = "Validation")
-	TArray<TSubclassOf<UNiagaraValidationRule>> ValidationRules;
+	UPROPERTY(EditAnywhere, Category = "Validation", Instanced)
+	TArray<UNiagaraValidationRule*> ValidationRules;
 
 	FORCEINLINE const FNiagaraSystemScalabilitySettingsArray& GetSystemScalabilitySettings()const { return SystemScalabilitySettings; }
 	FORCEINLINE const FNiagaraEmitterScalabilitySettingsArray& GetEmitterScalabilitySettings()const { return EmitterScalabilitySettings; }
diff --git a/Engine/Plugins/FX/Niagara/Source/Niagara/Classes/NiagaraValidationRule.h b/Engine/Plugins/FX/Niagara/Source/Niagara/Classes/NiagaraValidationRule.h
--- a/Engine/Plugins/FX/Niagara/Source/Niagara/Classes/NiagaraValidationRule.h
+++ b/Engine/Plugins/FX/Niagara/Source/Niagara/Classes/NiagaraValidationRule.h
@@ -19,40 +19,49 @@ enum class ENiagaraValidationSeverity
 	Error,
 };
 
-USTRUCT()
-struct FNiagaraValidationResult
+DECLARE_DELEGATE(FNiagaraValidationFixDelegate);
+
+/** Delegate wrapper to automatically fix content that fails validation checks. */
+struct FNiagaraValidationFix
 {
-	GENERATED_USTRUCT_BODY()
+	FNiagaraValidationFix() {}
+	FNiagaraValidationFix(FText InDescription, FNiagaraValidationFixDelegate InFixDelegate)
+	: Description(InDescription)
+	, FixDelegate(InFixDelegate)
+	{
+	}
+	
+	FText Description;
+	FNiagaraValidationFixDelegate FixDelegate;
+};
 
-	FNiagaraValidationResult() {};
+struct FNiagaraValidationResult
+{
+	FNiagaraValidationResult() {}
 	FNiagaraValidationResult(ENiagaraValidationSeverity InSeverity, FText InSummaryText, FText InDescription, TWeakObjectPtr<UObject> InSourceObject)
 		: Severity(InSeverity), SummaryText(InSummaryText), Description(InDescription), SourceObject(InSourceObject) {}
 
-	UPROPERTY()
 	ENiagaraValidationSeverity Severity = ENiagaraValidationSeverity::Info;
-	
-	UPROPERTY()
-	FText SummaryText;
 
-	UPROPERTY()
+	FText SummaryText;
 	FText Description;
-
-	UPROPERTY()
 	TWeakObjectPtr<UObject> SourceObject;
+	TArray<FNiagaraValidationFix> Fixes;
+	TArray<FNiagaraValidationFix> Links;
 };
 
 /**
 Base class for system validation logic. 
 These allow Niagara systems to be inspected for content validation either at save time or from a commandlet.
 */
-UCLASS(abstract)
+UCLASS(abstract, EditInlineNew)
 class NIAGARA_API UNiagaraValidationRule : public UObject
 {
 	GENERATED_BODY()
 public:
 
 #if WITH_EDITOR
-	virtual TArray<FNiagaraValidationResult> CheckValidity(TSharedPtr<FNiagaraSystemViewModel> ViewModel) const;
+	virtual void CheckValidity(TSharedPtr<FNiagaraSystemViewModel> ViewModel, TArray<FNiagaraValidationResult>& OutResults) const;
 #endif
 };
 
diff --git a/Engine/Plugins/FX/Niagara/Source/Niagara/Private/NiagaraValidationRule.cpp b/Engine/Plugins/FX/Niagara/Source/Niagara/Private/NiagaraValidationRule.cpp
--- a/Engine/Plugins/FX/Niagara/Source/Niagara/Private/NiagaraValidationRule.cpp
+++ b/Engine/Plugins/FX/Niagara/Source/Niagara/Private/NiagaraValidationRule.cpp
@@ -3,8 +3,7 @@
 #include "NiagaraValidationRule.h"
 
 #if WITH_EDITOR
-TArray<FNiagaraValidationResult> UNiagaraValidationRule::CheckValidity(TSharedPtr<FNiagaraSystemViewModel> ViewModel) const
+void UNiagaraValidationRule::CheckValidity(TSharedPtr<FNiagaraSystemViewModel> ViewModel, TArray<FNiagaraValidationResult>& Results) const
 {
-	return TArray<FNiagaraValidationResult>();
 }
 #endif
\ No newline at end of file
diff --git a/Engine/Plugins/FX/Niagara/Source/NiagaraEditor/Private/NiagaraValidationRules.cpp b/Engine/Plugins/FX/Niagara/Source/NiagaraEditor/Private/NiagaraValidationRules.cpp
--- a/Engine/Plugins/FX/Niagara/Source/NiagaraEditor/Private/NiagaraValidationRules.cpp
+++ b/Engine/Plugins/FX/Niagara/Source/NiagaraEditor/Private/NiagaraValidationRules.cpp
@@ -17,6 +17,10 @@
 #include "ViewModels/Stack/NiagaraStackRendererItem.h"
 #include "ViewModels/Stack/NiagaraStackSystemPropertiesItem.h"
 #include "ViewModels/Stack/NiagaraStackViewModel.h"
+#include "ViewModels/NiagaraEmitterViewModel.h"
+
+#include "AssetToolsModule.h"
+#include "Toolkits/AssetEditorToolkit.h"
 
 #define LOCTEXT_NAMESPACE "NiagaraValidationRules"
 
@@ -70,11 +74,41 @@ T* GetStackEntry(UNiagaraStackViewModel* StackViewModel, bool bRefresh = false)
 	return nullptr;
 }
 
+// --------------------------------------------------------------------------------------------------------------------------------------------
+// Common fixes and links
+
+void AddGoToFXTypeLink(FNiagaraValidationResult& Result, UNiagaraEffectType* FXType)
+{
+	if (FXType == nullptr)
+	{
+		return;
+	}
+
+	FNiagaraValidationFix& GoToValidationRulesLink = Result.Links.AddDefaulted_GetRef();
+	GoToValidationRulesLink.Description = LOCTEXT("GoToValidationRulesFix", "Go To Validation Rules");
+	TWeakObjectPtr<UNiagaraEffectType> WeakFXType = FXType;
+	GoToValidationRulesLink.FixDelegate = FNiagaraValidationFixDelegate::CreateLambda([WeakFXType]
+		{
+			FAssetToolsModule& AssetToolsModule = FModuleManager::LoadModuleChecked<FAssetToolsModule>(TEXT("AssetTools"));
+			TWeakPtr<IAssetTypeActions> WeakAssetTypeActions = AssetToolsModule.Get().GetAssetTypeActionsForClass(UNiagaraEffectType::StaticClass());
+
+			if (UNiagaraEffectType* FXType = WeakFXType.Get())
+			{
+				if (TSharedPtr<IAssetTypeActions> AssetTypeActions = WeakAssetTypeActions.Pin())
+				{
+					TArray<UObject*> AssetsToEdit;
+					AssetsToEdit.Add(FXType);
+					AssetTypeActions->OpenAssetEditor(AssetsToEdit);
+					//TODO: Is there a way for us to auto navigate to and open up the validation rules inside FXType?
+				}
+			}
+		});
+}
+
 // --------------------------------------------------------------------------------------------------------------------------------------------
 
-TArray<FNiagaraValidationResult> UNiagara_NoWarmupTime::CheckValidity(TSharedPtr<FNiagaraSystemViewModel> ViewModel) const
+void UNiagaraValidationRule_NoWarmupTime::CheckValidity(TSharedPtr<FNiagaraSystemViewModel> ViewModel, TArray<FNiagaraValidationResult>& Results)  const
 {
-	TArray<FNiagaraValidationResult> Results;
 	UNiagaraSystem& System = ViewModel->GetSystem();
 	if (System.NeedsWarmup())
 	{
@@ -82,17 +116,14 @@ TArray<FNiagaraValidationResult> UNiagara_NoWarmupTime::CheckValidity(TSharedPtr
 		FNiagaraValidationResult Result(ENiagaraValidationSeverity::Error, LOCTEXT("WarumupSummary", "Warmuptime > 0 is not allowed"), LOCTEXT("WarmupDescription", "Systems with the chosen effect type do not allow warmup time, as it costs too much performance.\nPlease set the warmup time to 0 in the system properties."), SystemProperties);
 		Results.Add(Result);
 	}
-	return Results;
 }
 
-TArray<FNiagaraValidationResult> UNiagara_FixedGPUBoundsSet::CheckValidity(TSharedPtr<FNiagaraSystemViewModel> ViewModel) const
+void UNiagaraValidationRule_FixedGPUBoundsSet::CheckValidity(TSharedPtr<FNiagaraSystemViewModel> ViewModel, TArray<FNiagaraValidationResult>& Results)  const
 {
-	TArray<FNiagaraValidationResult> Results;
-
 	// if the system has fixed bounds set then it overrides the emitter settings
 	if (ViewModel->GetSystem().bFixedBounds)
 	{
-		return Results;
+		return;
 	}
 
 	// check that all the gpu emitters have fixed bounds set
@@ -107,7 +138,6 @@ TArray<FNiagaraValidationResult> UNiagara_FixedGPUBoundsSet::CheckValidity(TShar
 			Results.Add(Result);
 		}
 	}
-	return Results;
 }
 
 bool IsEnabledForMaxQualityLevel(FNiagaraPlatformSet Platforms, int32 MaxQualityLevel)
@@ -122,54 +152,141 @@ bool IsEnabledForMaxQualityLevel(FNiagaraPlatformSet Platforms, int32 MaxQuality
 	return false;
 }
 
-TArray<FNiagaraValidationResult> UNiagara_NoComponentRendererOnLowSettings::CheckValidity(TSharedPtr<FNiagaraSystemViewModel> ViewModel) const
+void UNiagaraValidationRule_BannedRenderers::CheckValidity(TSharedPtr<FNiagaraSystemViewModel> ViewModel, TArray<FNiagaraValidationResult>& Results)  const
 {
-	TArray<FNiagaraValidationResult> Results;
-
-	// check that component renderers are only active for high and cinematic settings 
+	UNiagaraSystem& System = ViewModel->GetSystem();
 	TArray<TSharedRef<FNiagaraEmitterHandleViewModel>> EmitterHandleViewModels = ViewModel->GetEmitterHandleViewModels();
 	for (TSharedRef<FNiagaraEmitterHandleViewModel> EmitterHandleModel : EmitterHandleViewModels)
 	{
 		int32 MaxQualityLevel = 3;
 		UNiagaraEmitter* NiagaraEmitter = EmitterHandleModel.Get().GetEmitterHandle()->GetInstance();
-		NiagaraEmitter->ForEachRenderer([&Results, EmitterHandleModel, MaxQualityLevel](UNiagaraRendererProperties* RendererProperties)
+		NiagaraEmitter->ForEachRenderer([&Results, EmitterHandleModel, MaxQualityLevel, this, &System](UNiagaraRendererProperties* RendererProperties)
 		{
-			if (Cast<UNiagaraComponentRendererProperties>(RendererProperties) && RendererProperties->GetIsEnabled() && IsEnabledForMaxQualityLevel(RendererProperties->Platforms, MaxQualityLevel))
+			if (RendererProperties->GetIsEnabled() && BannedRenderers.Contains(RendererProperties->GetClass()))
 			{
-				TArray<UNiagaraStackRendererItem*> RendererItems = GetStackEntries<UNiagaraStackRendererItem>(EmitterHandleModel.Get().GetEmitterStackViewModel());
-				for (UNiagaraStackRendererItem* Item : RendererItems)
+				TArray<const FNiagaraPlatformSet*> CheckSets;
+				CheckSets.Add(&Platforms);
+				CheckSets.Add(&RendererProperties->Platforms);
+
+				TArray<FNiagaraPlatformSetConflictInfo> Conflicts;
+				FNiagaraPlatformSet::GatherConflicts(CheckSets, Conflicts);
+				if (Conflicts.Num() > 0)
 				{
-					if (Item->GetRendererProperties() != RendererProperties)
+					TArray<UNiagaraStackRendererItem*> RendererItems = GetStackEntries<UNiagaraStackRendererItem>(EmitterHandleModel.Get().GetEmitterStackViewModel());
+					for (UNiagaraStackRendererItem* Item : RendererItems)
 					{
-						continue;
+						if (Item->GetRendererProperties() != RendererProperties)
+						{
+							continue;
+						}						
+						FNiagaraValidationResult Result = Results.AddDefaulted_GetRef();
+						
+						Result.Severity = ENiagaraValidationSeverity::Warning;
+						Result.SummaryText = LOCTEXT("BannedRenderSummary", "Banned renderers used.");
+						Result.Description = LOCTEXT("BannedRenderDescription", "Please ensure only allowed renderers are used for each platform according to the validation rules in the System's Effect Type.");
+						Result.SourceObject = Item;
+						
+						AddGoToFXTypeLink(Result, System.GetEffectType());
+
+						//Add autofix to disable the module
+						FNiagaraValidationFix& DisableRendererFix = Result.Fixes.AddDefaulted_GetRef();
+						DisableRendererFix.Description = LOCTEXT("DisableBannedRendererFix", "Disable Banned Renderer");
+						TWeakObjectPtr<UNiagaraStackRendererItem> WeakRendererItem = Item;
+						DisableRendererFix.FixDelegate = FNiagaraValidationFixDelegate::CreateLambda(
+							[WeakRendererItem]()
+							{
+								if (UNiagaraStackRendererItem* RendererItem = WeakRendererItem.Get())
+								{
+									RendererItem->SetIsEnabled(false);
+								}
+							});
+
+						Results.Add(Result);
 					}
-					FNiagaraValidationResult Result(ENiagaraValidationSeverity::Warning, LOCTEXT("ComponentRenderSummary", "Component renderers should not be used for low-level platforms"), LOCTEXT("ComponentRenderDescription", "Please disable low level platforms in the renderer's scalability settings."), Item);
-					Results.Add(Result);
 				}
 			}
 		});
 	}
-	return Results;
+} 
+
+void UNiagaraValidationRule_BannedModules::CheckValidity(TSharedPtr<FNiagaraSystemViewModel> ViewModel, TArray<FNiagaraValidationResult>& Results)  const
+{
+	UNiagaraSystem& System = ViewModel->GetSystem();
+
+	TArray<UNiagaraStackModuleItem*> StackModuleItems =	GetAllStackEntriesInSystem<UNiagaraStackModuleItem>(ViewModel);
+
+	for (UNiagaraStackModuleItem* Item : StackModuleItems)
+	{
+		if (Item && Item->GetIsEnabled())
+		{
+			UNiagaraNodeFunctionCall& FuncCall = Item->GetModuleNode();
+
+			for (UNiagaraScript* BannedModule : BannedModules)
+			{
+				if (BannedModule == FuncCall.FunctionScript)
+				{
+					UNiagaraEmitter* Emitter = Item->GetEmitterViewModel().IsValid() ? Item->GetEmitterViewModel()->GetEmitter() : nullptr;
+
+					bool bApplyBan = true;
+					if (Emitter)
+					{
+						//If we're on an emitter, this emitter may be culled on the platforms the rule applies to.
+						TArray<const FNiagaraPlatformSet*> CheckSets;
+						CheckSets.Add(&Platforms);
+						CheckSets.Add(&Emitter->Platforms);
+						TArray<FNiagaraPlatformSetConflictInfo> Conflicts;
+						FNiagaraPlatformSet::GatherConflicts(CheckSets, Conflicts);
+						bApplyBan = Conflicts.Num() > 0;
+					}
+
+					if (!bApplyBan)
+					{
+						continue;
+					}
+
+					const FTextFormat Format(LOCTEXT("BannedModuleFormat", "Module {0} is banned on some currently enabled platforms"));
+					const FText WarningMessage = FText::Format(Format, FText::FromString(FuncCall.FunctionScript->GetName()));
+
+					FNiagaraValidationResult& Result = Results.AddDefaulted_GetRef();
+					Result.Severity = ENiagaraValidationSeverity::Warning;
+					Result.SummaryText = WarningMessage;
+					Result.Description = LOCTEXT("BanndeModulesDescription", "Check this module against the Effect Type's Banned Modules validators");
+					Result.SourceObject = Item;
+
+					AddGoToFXTypeLink(Result, System.GetEffectType());
+
+					//Add autofix to disable the module
+					FNiagaraValidationFix& DisableModuleFix = Result.Fixes.AddDefaulted_GetRef();
+					DisableModuleFix.Description = LOCTEXT("DisableBannedModuleFix", "Disable Banned Module");					
+					TWeakObjectPtr<UNiagaraStackModuleItem> WeakModuleItem = Item;
+					DisableModuleFix.FixDelegate = FNiagaraValidationFixDelegate::CreateLambda([WeakModuleItem]()
+					{
+						if (UNiagaraStackModuleItem* ModuleItem = WeakModuleItem.Get())
+						{
+							ModuleItem->SetEnabled(false);
+						}
+					});
+				}
+			}
+		}
+	}
 }
 
-TArray<FNiagaraValidationResult> UNiagara_InvalidEffectType::CheckValidity(TSharedPtr<FNiagaraSystemViewModel> ViewModel) const
+void UNiagaraValidationRule_InvalidEffectType::CheckValidity(TSharedPtr<FNiagaraSystemViewModel> ViewModel, TArray<FNiagaraValidationResult>& Results)  const
 {
-	TArray<FNiagaraValidationResult> Results;
 	UNiagaraStackSystemPropertiesItem* SystemProperties = GetStackEntry<UNiagaraStackSystemPropertiesItem>(ViewModel->GetSystemStackViewModel());
 	FNiagaraValidationResult Result(ENiagaraValidationSeverity::Error, LOCTEXT("InvalidEffectSummary", "Invalid Effect Type"), LOCTEXT("InvalidEffectDescription", "The effect type on this system was marked as invalid for production content and should only be used as placeholder."), SystemProperties);
 	Results.Add(Result);
-	return Results;
 }
 
-TArray<FNiagaraValidationResult> UNiagara_LWC_Validator::CheckValidity(TSharedPtr<FNiagaraSystemViewModel> ViewModel) const
+void UNiagaraValidationRule_LWC::CheckValidity(TSharedPtr<FNiagaraSystemViewModel> ViewModel, TArray<FNiagaraValidationResult>& Results)  const
 {
-	TArray<FNiagaraValidationResult> Results;
 
 	const UNiagaraSettings* Settings = GetDefault<UNiagaraSettings>();
 	UNiagaraSystem& System = ViewModel->GetSystem();
 	if (!System.SupportsLargeWorldCoordinates())
 	{
-		return Results;
+		return;
 	}
 
 	// gather all the modules in the system, excluding localspace emitters
@@ -235,8 +352,6 @@ TArray<FNiagaraValidationResult> UNiagara_LWC_Validator::CheckValidity(TSharedPt
 			}
 		}
 	}
-	
-	return Results;
 }
 
-#undef LOCTEXT_NAMESPACE
+#undef LOCTEXT_NAMESPACE
\ No newline at end of file
diff --git a/Engine/Plugins/FX/Niagara/Source/NiagaraEditor/Private/ViewModels/Stack/NiagaraStackEntry.cpp b/Engine/Plugins/FX/Niagara/Source/NiagaraEditor/Private/ViewModels/Stack/NiagaraStackEntry.cpp
--- a/Engine/Plugins/FX/Niagara/Source/NiagaraEditor/Private/ViewModels/Stack/NiagaraStackEntry.cpp
+++ b/Engine/Plugins/FX/Niagara/Source/NiagaraEditor/Private/ViewModels/Stack/NiagaraStackEntry.cpp
@@ -598,6 +598,30 @@ const TArray<UNiagaraStackEntry*>& UNiagaraStackEntry::GetAllChildrenWithIssues(
 	return GetCollectedIssueData().ChildrenWithIssues;
 }
 
+void UNiagaraStackEntry::AddValidationIssue(EStackIssueSeverity Severity, const FText& SummaryText, const FText& Description, bool bCanBeDismissed, const TArray<FNiagaraValidationFix>& Fixes, const TArray<FNiagaraValidationFix>& Links)
+{
+	TArray<FStackIssueFix> StackFixes;
+	for (const FNiagaraValidationFix& Fix : Fixes)
+	{
+		StackFixes.Add(FStackIssueFix(Fix.Description, Fix.FixDelegate, EStackIssueFixStyle::Fix));
+	}
+	for (const FNiagaraValidationFix& Link : Links)
+	{
+		StackFixes.Add(FStackIssueFix(Link.Description, Link.FixDelegate, EStackIssueFixStyle::Link));
+	}
+
+	FStackIssue NewIssue(Severity, SummaryText, Description, GetStackEditorDataKey(), bCanBeDismissed, StackFixes);
+	for (const FStackIssue& ExistingIssue : ExternalStackIssues)
+	{
+		if (ExistingIssue.GetUniqueIdentifier().Equals(NewIssue.GetUniqueIdentifier()))
+		{
+			return;
+		}
+	}
+	ExternalStackIssues.Add(NewIssue);
+	RefreshChildren();
+}
+
 void UNiagaraStackEntry::AddExternalIssue(EStackIssueSeverity Severity, const FText& SummaryText, const FText& Description, bool bCanBeDismissed)
 {
 	FStackIssue NewIssue(Severity, SummaryText, Description, GetStackEditorDataKey(), bCanBeDismissed);
diff --git a/Engine/Plugins/FX/Niagara/Source/NiagaraEditor/Private/ViewModels/Stack/NiagaraStackRenderItemGroup.cpp b/Engine/Plugins/FX/Niagara/Source/NiagaraEditor/Private/ViewModels/Stack/NiagaraStackRenderItemGroup.cpp
--- a/Engine/Plugins/FX/Niagara/Source/NiagaraEditor/Private/ViewModels/Stack/NiagaraStackRenderItemGroup.cpp
+++ b/Engine/Plugins/FX/Niagara/Source/NiagaraEditor/Private/ViewModels/Stack/NiagaraStackRenderItemGroup.cpp
@@ -14,6 +14,7 @@
 #include "Framework/Notifications/NotificationManager.h"
 #include "Styling/CoreStyle.h"
 #include "ViewModels/NiagaraSystemSelectionViewModel.h"
+#include "ViewModels/Stack/NiagaraStackViewModel.h"
 
 #define LOCTEXT_NAMESPACE "UNiagaraStackRenderItemGroup"
 
@@ -209,6 +210,7 @@ void UNiagaraStackRenderItemGroup::OnRendererAdded(UNiagaraRendererProperties* R
 {
 	GetSystemViewModel()->GetSelectionViewModel()->EmptySelection();
 	GetSystemViewModel()->GetSelectionViewModel()->AddEntryToSelectionByDisplayedObjectKeyDeferred(FObjectKey(RendererProperties));
+	GetSystemViewModel()->GetSystemStackViewModel()->RequestValidationUpdate();
 }
 
 void UNiagaraStackRenderItemGroup::FinalizeInternal()
diff --git a/Engine/Plugins/FX/Niagara/Source/NiagaraEditor/Private/ViewModels/Stack/NiagaraStackRendererItem.cpp b/Engine/Plugins/FX/Niagara/Source/NiagaraEditor/Private/ViewModels/Stack/NiagaraStackRendererItem.cpp
--- a/Engine/Plugins/FX/Niagara/Source/NiagaraEditor/Private/ViewModels/Stack/NiagaraStackRendererItem.cpp
+++ b/Engine/Plugins/FX/Niagara/Source/NiagaraEditor/Private/ViewModels/Stack/NiagaraStackRendererItem.cpp
@@ -362,6 +362,11 @@ void UNiagaraStackRendererItem::SetIsEnabledInternal(bool bInIsEnabled)
 	ChangedObjects.Add(RendererProperties.Get());
 	OnDataObjectModified().Broadcast(ChangedObjects, ENiagaraDataObjectChange::Changed);
 	RefreshChildren();
+
+	if (GetEmitterViewModel().IsValid())
+	{
+		GetSystemViewModel()->GetEmitterHandleViewModelForEmitter(GetEmitterViewModel()->GetEmitter()).Get()->GetEmitterStackViewModel()->RequestValidationUpdate();
+	}
 }
 
 const FSlateBrush* UNiagaraStackRendererItem::GetIconBrush() const
diff --git a/Engine/Plugins/FX/Niagara/Source/NiagaraEditor/Private/ViewModels/Stack/NiagaraStackViewModel.cpp b/Engine/Plugins/FX/Niagara/Source/NiagaraEditor/Private/ViewModels/Stack/NiagaraStackViewModel.cpp
--- a/Engine/Plugins/FX/Niagara/Source/NiagaraEditor/Private/ViewModels/Stack/NiagaraStackViewModel.cpp
+++ b/Engine/Plugins/FX/Niagara/Source/NiagaraEditor/Private/ViewModels/Stack/NiagaraStackViewModel.cpp
@@ -87,26 +87,32 @@ void UNiagaraStackViewModel::UpdateStackWithValidationResults()
 		Entry->ClearExternalIssues();
 		Entry->GetUnfilteredChildren(EntriesToClear);
 	}
+	
+	TSharedPtr<FNiagaraSystemViewModel> SysViewModel = SystemViewModel.Pin();
+	if (SysViewModel == nullptr)
+	{
+		return;
+	}
 
-	UNiagaraSystem& NiagaraSystem = SystemViewModel.Pin()->GetSystem();
+	UNiagaraSystem& NiagaraSystem = SysViewModel->GetSystem();
 	if (NiagaraSystem.GetEffectType())
 	{
 		// go over the validation rules
-		for (TSubclassOf<UNiagaraValidationRule> ValidationClass : NiagaraSystem.GetEffectType()->ValidationRules)
+		for (UNiagaraValidationRule* ValidationRule : NiagaraSystem.GetEffectType()->ValidationRules)
 		{
-			if (UNiagaraValidationRule* ValidationRule = Cast<UNiagaraValidationRule>(ValidationClass->GetDefaultObject()))
+			if (ValidationRule)
 			{
-				TArray<FNiagaraValidationResult> NiagaraValidationResults = ValidationRule->CheckValidity(SystemViewModel.Pin());
+				TArray<FNiagaraValidationResult> NiagaraValidationResults;
+				ValidationRule->CheckValidity(SystemViewModel.Pin(), NiagaraValidationResults);
 				for (const FNiagaraValidationResult& Result : NiagaraValidationResults)
 				{
 					EStackIssueSeverity Severity = (Result.Severity == ENiagaraValidationSeverity::Error) ? EStackIssueSeverity::Error : (Result.Severity == ENiagaraValidationSeverity::Warning ? EStackIssueSeverity::Warning : EStackIssueSeverity::Info);
 					if (UNiagaraStackEntry* SourceEntry = Cast<UNiagaraStackEntry>(Result.SourceObject.Get()))
 					{
-						SourceEntry->AddExternalIssue(Severity, Result.SummaryText, Result.Description, Result.Severity == ENiagaraValidationSeverity::Info);
+						SourceEntry->AddValidationIssue(Severity, Result.SummaryText, Result.Description, Result.Severity == ENiagaraValidationSeverity::Info, Result.Fixes, Result.Links);
 					}
 				}
 			}
-			
 		}
 	}
 }
@@ -246,6 +252,7 @@ void UNiagaraStackViewModel::Tick()
 			RootEntry->RefreshChildren();
 			bRefreshPending = false;
 			InvalidateSearchResults();
+			bValidatorUpdatePending = true;
 		}
 
 		if (bValidatorUpdatePending)
diff --git a/Engine/Plugins/FX/Niagara/Source/NiagaraEditor/Public/NiagaraValidationRules.h b/Engine/Plugins/FX/Niagara/Source/NiagaraEditor/Public/NiagaraValidationRules.h
--- a/Engine/Plugins/FX/Niagara/Source/NiagaraEditor/Public/NiagaraValidationRules.h
+++ b/Engine/Plugins/FX/Niagara/Source/NiagaraEditor/Public/NiagaraValidationRules.h
@@ -4,47 +4,77 @@
 
 #include "CoreMinimal.h"
 #include "NiagaraValidationRule.h"
+#include "NiagaraPlatformSet.h"
 #include "NiagaraValidationRules.generated.h"
 
+class UNiagaraScript;
+class UNiagaraRendererProperties;
 
-UCLASS(Category = "Validation")
-class UNiagara_NoWarmupTime : public UNiagaraValidationRule
+
+UCLASS(Category = "Validation", DisplayName = "No Warmup Time")
+class UNiagaraValidationRule_NoWarmupTime : public UNiagaraValidationRule
 {
 	GENERATED_BODY()
 public:
-	virtual TArray<FNiagaraValidationResult> CheckValidity(TSharedPtr<FNiagaraSystemViewModel> ViewModel) const override;
+	virtual void CheckValidity(TSharedPtr<FNiagaraSystemViewModel> ViewModel, TArray<FNiagaraValidationResult>& OutResults) const override;
 };
 
 
-UCLASS(Category = "Validation")
-class UNiagara_FixedGPUBoundsSet : public UNiagaraValidationRule
+UCLASS(Category = "Validation", DisplayName = "Fixed GPU Bounds Set")
+class UNiagaraValidationRule_FixedGPUBoundsSet : public UNiagaraValidationRule
 {
 	GENERATED_BODY()
 public:
-	virtual TArray<FNiagaraValidationResult> CheckValidity(TSharedPtr<FNiagaraSystemViewModel> ViewModel) const override;
+	virtual void CheckValidity(TSharedPtr<FNiagaraSystemViewModel> ViewModel, TArray<FNiagaraValidationResult>& OutResults) const override;
 };
 
 
-UCLASS(Category = "Validation")
-class UNiagara_NoComponentRendererOnLowSettings : public UNiagaraValidationRule
+UCLASS(Category = "Validation", DisplayName = "Banned Renderers")
+class UNiagaraValidationRule_BannedRenderers : public UNiagaraValidationRule
+{
+	GENERATED_BODY()
+
+public:
+
+	//Platforms this validation rule will apply to.
+	UPROPERTY(EditAnywhere, Category = Validation)
+	FNiagaraPlatformSet Platforms;
+
+	UPROPERTY(EditAnywhere, Category = Validation)
+	TArray<TSubclassOf<UNiagaraRendererProperties>> BannedRenderers;
+
+	virtual void CheckValidity(TSharedPtr<FNiagaraSystemViewModel> ViewModel, TArray<FNiagaraValidationResult>& OutResults) const override;
+};
+
+UCLASS(Category = "Validation", DisplayName = "Banned Modules")
+class UNiagaraValidationRule_BannedModules : public UNiagaraValidationRule
 {
 	GENERATED_BODY()
+
 public:
-	virtual TArray<FNiagaraValidationResult> CheckValidity(TSharedPtr<FNiagaraSystemViewModel> ViewModel) const override;
+
+	//Platforms this validation rule will apply to.
+	UPROPERTY(EditAnywhere, Category=Validation)
+	FNiagaraPlatformSet Platforms;
+
+	UPROPERTY(EditAnywhere, Category = Validation)
+	TArray<UNiagaraScript*> BannedModules;
+
+	virtual void CheckValidity(TSharedPtr<FNiagaraSystemViewModel> ViewModel, TArray<FNiagaraValidationResult>& OutResults) const override;
 };
 
-UCLASS(Category = "Validation")
-class UNiagara_InvalidEffectType : public UNiagaraValidationRule
+UCLASS(Category = "Validation", DisplayName = "Invalid Effect Type")
+class UNiagaraValidationRule_InvalidEffectType : public UNiagaraValidationRule
 {
 	GENERATED_BODY()
 public:
-	virtual TArray<FNiagaraValidationResult> CheckValidity(TSharedPtr<FNiagaraSystemViewModel> ViewModel) const override;
+	virtual void CheckValidity(TSharedPtr<FNiagaraSystemViewModel> ViewModel, TArray<FNiagaraValidationResult>& OutResults) const override;
 };
 
-UCLASS(Category = "Validation")
-class UNiagara_LWC_Validator : public UNiagaraValidationRule
+UCLASS(Category = "Validation", DisplayName = "Large World Coordinates")
+class UNiagaraValidationRule_LWC : public UNiagaraValidationRule
 {
 	GENERATED_BODY()
 public:
-	virtual TArray<FNiagaraValidationResult> CheckValidity(TSharedPtr<FNiagaraSystemViewModel> ViewModel) const override;
+	virtual void CheckValidity(TSharedPtr<FNiagaraSystemViewModel> ViewModel, TArray<FNiagaraValidationResult>& OutResults) const override;
 };
\ No newline at end of file
diff --git a/Engine/Plugins/FX/Niagara/Source/NiagaraEditor/Public/ViewModels/Stack/NiagaraStackEntry.h b/Engine/Plugins/FX/Niagara/Source/NiagaraEditor/Public/ViewModels/Stack/NiagaraStackEntry.h
--- a/Engine/Plugins/FX/Niagara/Source/NiagaraEditor/Public/ViewModels/Stack/NiagaraStackEntry.h
+++ b/Engine/Plugins/FX/Niagara/Source/NiagaraEditor/Public/ViewModels/Stack/NiagaraStackEntry.h
@@ -469,7 +469,8 @@ public:
 	const TArray<FStackIssue>& GetIssues() const;
 
 	const TArray<UNiagaraStackEntry*>& GetAllChildrenWithIssues() const;
-
+	
+	void AddValidationIssue(EStackIssueSeverity Severity, const FText& SummaryText, const FText& Description, bool bCanBeDismissed, const TArray<FNiagaraValidationFix>& Fixes, const TArray<FNiagaraValidationFix>& Links);
 	void AddExternalIssue(EStackIssueSeverity Severity, const FText& SummaryText, const FText& Description, bool bCanBeDismissed);
 	void ClearExternalIssues();
 
