Non Unity fix

#rnx

diff --git a/Engine/Plugins/FX/Niagara/Source/NiagaraEditor/Public/NiagaraValidationRules.h b/Engine/Plugins/FX/Niagara/Source/NiagaraEditor/Public/NiagaraValidationRules.h
--- a/Engine/Plugins/FX/Niagara/Source/NiagaraEditor/Public/NiagaraValidationRules.h
+++ b/Engine/Plugins/FX/Niagara/Source/NiagaraEditor/Public/NiagaraValidationRules.h
@@ -5,11 +5,10 @@
 #include "CoreMinimal.h"
 #include "NiagaraValidationRule.h"
 #include "NiagaraPlatformSet.h"
+#include "NiagaraRendererProperties.h"
 #include "NiagaraValidationRules.generated.h"
 
 class UNiagaraScript;
-class UNiagaraRendererProperties;
-
 
 UCLASS(Category = "Validation", DisplayName = "No Warmup Time")
 class UNiagaraValidationRule_NoWarmupTime : public UNiagaraValidationRule
