Various CIS compiler fixes.

diff --git a/Engine/Plugins/FX/Niagara/Source/Niagara/Classes/NiagaraEffectType.h b/Engine/Plugins/FX/Niagara/Source/Niagara/Classes/NiagaraEffectType.h
--- a/Engine/Plugins/FX/Niagara/Source/Niagara/Classes/NiagaraEffectType.h
+++ b/Engine/Plugins/FX/Niagara/Source/Niagara/Classes/NiagaraEffectType.h
@@ -407,7 +407,7 @@ class NIAGARA_API UNiagaraEffectType : public UObject
 
 	/** A set of rules to apply when checking content. To create your own rules, write a custom class that extends UNiagaraValidationRule. */
 	UPROPERTY(EditAnywhere, Category = "Validation", Instanced)
-	TArray<UNiagaraValidationRule*> ValidationRules;
+	TArray<TObjectPtr<UNiagaraValidationRule>> ValidationRules;
 
 	FORCEINLINE const FNiagaraSystemScalabilitySettingsArray& GetSystemScalabilitySettings()const { return SystemScalabilitySettings; }
 	FORCEINLINE const FNiagaraEmitterScalabilitySettingsArray& GetEmitterScalabilitySettings()const { return EmitterScalabilitySettings; }
diff --git a/Engine/Plugins/FX/Niagara/Source/NiagaraEditor/Public/NiagaraValidationRules.h b/Engine/Plugins/FX/Niagara/Source/NiagaraEditor/Public/NiagaraValidationRules.h
--- a/Engine/Plugins/FX/Niagara/Source/NiagaraEditor/Public/NiagaraValidationRules.h
+++ b/Engine/Plugins/FX/Niagara/Source/NiagaraEditor/Public/NiagaraValidationRules.h
@@ -57,7 +57,7 @@ public:
 	FNiagaraPlatformSet Platforms;
 
 	UPROPERTY(EditAnywhere, Category = Validation)
-	TArray<UNiagaraScript*> BannedModules;
+	TArray<TObjectPtr<UNiagaraScript>> BannedModules;
 
 	virtual void CheckValidity(TSharedPtr<FNiagaraSystemViewModel> ViewModel, TArray<FNiagaraValidationResult>& OutResults) const override;
 };
