- Remove comparison of types as it's no longer required

#rnx

diff --git a/Engine/Plugins/FX/Niagara/Source/Niagara/Private/NiagaraDataSet.cpp b/Engine/Plugins/FX/Niagara/Source/Niagara/Private/NiagaraDataSet.cpp
--- a/Engine/Plugins/FX/Niagara/Source/Niagara/Private/NiagaraDataSet.cpp
+++ b/Engine/Plugins/FX/Niagara/Source/Niagara/Private/NiagaraDataSet.cpp
@@ -1089,14 +1089,6 @@ void FNiagaraDataBuffer::GPUCopyFrom(const float* GPUReadBackFloat, const int* G
 			float* Dst = GetInstancePtrFloat(CompIdx, 0);
 			size_t Count = SrcEnd - SrcStart;
 			FMemory::Memcpy(Dst, SrcStart, Count * sizeof(float));
-
-			if (Count > 0)
-			{
-				for (size_t i = 0; i < Count; i++)
-				{
-					check(SrcStart[i] == Dst[i]);
-				}
-			}
 		}
 	}
 	if (GPUReadBackInt)
@@ -1111,14 +1103,6 @@ void FNiagaraDataBuffer::GPUCopyFrom(const float* GPUReadBackFloat, const int* G
 			int32* Dst = GetInstancePtrInt32(CompIdx, 0);
 			size_t Count = SrcEnd - SrcStart;
 			FMemory::Memcpy(Dst, SrcStart, Count * sizeof(int32));
-
-			if (Count > 0)
-			{
-				for (size_t i = 0; i < Count; i++)
-				{
-					check(SrcStart[i] == Dst[i]);
-				}
-			}
 		}
 	}
 	if (GPUReadBackHalf)
@@ -1133,14 +1117,6 @@ void FNiagaraDataBuffer::GPUCopyFrom(const float* GPUReadBackFloat, const int* G
 			FFloat16* Dst = GetInstancePtrHalf(CompIdx, 0);
 			size_t Count = SrcEnd - SrcStart;
 			FMemory::Memcpy(Dst, SrcStart, Count * sizeof(FFloat16));
-
-			if (Count > 0)
-			{
-				for (size_t i = 0; i < Count; i++)
-				{
-					check(SrcStart[i] == Dst[i]);
-				}
-			}
 		}
 	}
 }
