- Allow GPU particle readback on non-editor builds

#rnx

diff --git a/Engine/Plugins/FX/Niagara/Source/Niagara/Private/NiagaraDebugHud.cpp b/Engine/Plugins/FX/Niagara/Source/Niagara/Private/NiagaraDebugHud.cpp
--- a/Engine/Plugins/FX/Niagara/Source/Niagara/Private/NiagaraDebugHud.cpp
+++ b/Engine/Plugins/FX/Niagara/Source/Niagara/Private/NiagaraDebugHud.cpp
@@ -947,7 +947,6 @@ FNiagaraDataSet* FNiagaraDebugHud::GetParticleDataSet(FNiagaraSystemInstance* Sy
 	// For GPU context we need to readback and cache the data
 	if (EmitterInstance->GetGPUContext())
 	{
-#if UE_ENABLE_DEBUG_DRAWING
 		if (!Settings.bEnableGpuParticleReadback)
 		{
 			return nullptr;
@@ -993,7 +992,6 @@ FNiagaraDataSet* FNiagaraDebugHud::GetParticleDataSet(FNiagaraSystemInstance* Sy
 		{
 			return &GpuCachedData->CurrentEmitterData[iEmitter]->Frame;
 		}
-#endif
 		return nullptr;
 	}
 
