Fixed issue with the CVar conditions in editor builds.

#rnx

diff --git a/Engine/Plugins/FX/Niagara/Source/Niagara/Classes/NiagaraPlatformSet.h b/Engine/Plugins/FX/Niagara/Source/Niagara/Classes/NiagaraPlatformSet.h
--- a/Engine/Plugins/FX/Niagara/Source/Niagara/Classes/NiagaraPlatformSet.h
+++ b/Engine/Plugins/FX/Niagara/Source/Niagara/Classes/NiagaraPlatformSet.h
@@ -192,7 +192,7 @@ struct NIAGARA_API FNiagaraPlatformSetCVarCondition
 	FNiagaraPlatformSetCVarCondition();
 
 	/** Return true if this condition is met for the given device profile. */
-	bool Evaluate(const UDeviceProfile* DeviceProfile, bool& bOutCanEverPass, bool& bOutCanEverFail)const;
+	bool Evaluate(const UDeviceProfile* DeviceProfile, int32 QualityLevel, bool bCheckCurrentStateOnly, bool& bOutCanEverPass, bool& bOutCanEverFail)const;
 
 	/** Returns the CVar for this condition. Can be null if the name give is not a valid CVar or the CVar is removed. */
 	IConsoleVariable* GetCVar()const;
@@ -248,7 +248,7 @@ struct NIAGARA_API FNiagaraPlatformSetCVarCondition
 	uint32 bUseMaxFloat : 1;
 
 	template<typename T>
-	bool Evaluate_Internal(const UDeviceProfile* DeviceProfile, bool& bOutCanEverPass, bool& bOutCanEverFail)const;
+	bool Evaluate_Internal(const UDeviceProfile* DeviceProfile, int32 QualityLevel, bool bCheckCurrentStateOnly, bool& bOutCanEverPass, bool& bOutCanEverFail)const;
 
 	FORCEINLINE bool CheckValue(bool CVarValue)const { return CVarValue == Value; }
 	FORCEINLINE bool CheckValue(int32 CVarValue)const { return (!bUseMinInt || CVarValue >= MinInt) && (!bUseMaxInt || CVarValue <= MaxInt); }
@@ -365,7 +365,7 @@ public:
 	/** Returns true if the passed platform can modify it's niagara scalability settings at runtime. */
 	static bool CanChangeScalabilityAtRuntime(const UDeviceProfile* DeviceProfile);
 
-	FNiagaraPlatformSetEnabledState IsEnabled(const UDeviceProfile* Profile, int32 QualityLevel, FNiagaraPlatformSetEnabledStateDetails* OutDetails=nullptr)const;
+	FNiagaraPlatformSetEnabledState IsEnabled(const UDeviceProfile* Profile, int32 QualityLevel, bool bCheckCurrentStateOnly, FNiagaraPlatformSetEnabledStateDetails* OutDetails=nullptr)const;
 
 	bool CanConsiderDeviceProfile(const UDeviceProfile* Profile)const;
 
diff --git a/Engine/Plugins/FX/Niagara/Source/Niagara/Private/NiagaraPlatformSet.cpp b/Engine/Plugins/FX/Niagara/Source/Niagara/Private/NiagaraPlatformSet.cpp
--- a/Engine/Plugins/FX/Niagara/Source/Niagara/Private/NiagaraPlatformSet.cpp
+++ b/Engine/Plugins/FX/Niagara/Source/Niagara/Private/NiagaraPlatformSet.cpp
@@ -324,7 +324,7 @@ bool FNiagaraPlatformSet::IsActive()const
 			}
 		}
 #endif
-		FNiagaraPlatformSetEnabledState EnabledState = IsEnabled(ActiveProfile, QualityLevel);
+		FNiagaraPlatformSetEnabledState EnabledState = IsEnabled(ActiveProfile, QualityLevel, true);
 		bEnabledForCurrentProfileAndEffectQuality = EnabledState.bIsActive;
 		LastBuiltFrame = GFrameNumber;
 	}
@@ -339,7 +339,7 @@ int32 FNiagaraPlatformSet::GetEnabledMaskForDeviceProfile(const UDeviceProfile*
 	int32 RetQLMask = 0;
 	for (int32 i = 0; i < Settings->QualityLevels.Num(); ++i)
 	{
-		FNiagaraPlatformSetEnabledState EnabledState = IsEnabled(DeviceProfile, i);
+		FNiagaraPlatformSetEnabledState EnabledState = IsEnabled(DeviceProfile, i, false);
 		if (EnabledState.bCanBeActive)
 		{
 			RetQLMask |= CreateQualityLevelMask(i);
@@ -357,7 +357,7 @@ int32 FNiagaraPlatformSet::IsEnabledForDeviceProfile(const UDeviceProfile* Devic
 	int32 RetQLMask = 0;
 	for (int32 i = 0; i < Settings->QualityLevels.Num(); ++i)
 	{
-		FNiagaraPlatformSetEnabledState Enabled = IsEnabled(DeviceProfile, i);
+		FNiagaraPlatformSetEnabledState Enabled = IsEnabled(DeviceProfile, i, false);
 		if (Enabled.bCanBeActive)
 		{
 			RetQLMask |= CreateQualityLevelMask(i);
@@ -373,7 +373,7 @@ bool FNiagaraPlatformSet::IsEnabledForQualityLevel(int32 QualityLevel)const
 	{
 		if (UDeviceProfile* Profile = Cast<UDeviceProfile>(DeviceProfileObj))
 		{
-			FNiagaraPlatformSetEnabledState Enabled = IsEnabled(Profile, QualityLevel);
+			FNiagaraPlatformSetEnabledState Enabled = IsEnabled(Profile, QualityLevel, false);
 			if (Enabled.bCanBeActive)
 			{
 				return true;
@@ -493,7 +493,7 @@ int32 FNiagaraPlatformSet::GetAvailableQualityMaskForDeviceProfile(const UDevice
 #endif
 }
 
-FNiagaraPlatformSetEnabledState FNiagaraPlatformSet::IsEnabled(const UDeviceProfile* Profile, int32 QualityLevel, FNiagaraPlatformSetEnabledStateDetails* OutDetails)const
+FNiagaraPlatformSetEnabledState FNiagaraPlatformSet::IsEnabled(const UDeviceProfile* Profile, int32 QualityLevel, bool bCheckCurrentStateOnly, FNiagaraPlatformSetEnabledStateDetails* OutDetails)const
 {
 	checkSlow(Profile);
 
@@ -661,7 +661,7 @@ FNiagaraPlatformSetEnabledState FNiagaraPlatformSet::IsEnabled(const UDeviceProf
 		bool bCanEverFail = false;
 		//Get all quality levels where this condition will pass.
 		int32 ConditionPassedMask = INDEX_NONE;
-		bool bConditionPassed = CVarCondition.Evaluate(Profile, bCanEverPass, bCanEverFail);
+		bool bConditionPassed = CVarCondition.Evaluate(Profile, QualityLevel, bCheckCurrentStateOnly, bCanEverPass, bCanEverFail);
 
 		ensure(bCanEverPass || bCanEverFail);
 
@@ -1458,7 +1458,7 @@ template<> int32 FNiagaraPlatformSetCVarCondition::GetCVarValue<int32>(IConsoleV
 template<> float FNiagaraPlatformSetCVarCondition::GetCVarValue<float>(IConsoleVariable* CVar)const { return CVar->GetFloat(); }
 
 template<typename T>
-bool FNiagaraPlatformSetCVarCondition::Evaluate_Internal(const UDeviceProfile* DeviceProfile, bool& bOutCanEverPass, bool& bOutCanEverFail)const
+bool FNiagaraPlatformSetCVarCondition::Evaluate_Internal(const UDeviceProfile* DeviceProfile, int32 QualityLevel, bool bCheckCurrentStateOnly, bool& bOutCanEverPass, bool& bOutCanEverFail)const
 {		
 	bool bConditionPassed = false;
 
@@ -1472,38 +1472,51 @@ bool FNiagaraPlatformSetCVarCondition::Evaluate_Internal(const UDeviceProfile* D
 	bool bCanChangeAtRuntime = FNiagaraPlatformSet::CanChangeScalabilityAtRuntime(DeviceProfile) && (CVar->GetFlags() & ECVF_ReadOnly) == 0;
 
 #if WITH_EDITOR
-	const FNiagaraCVarValues& PossibleValues = FDeviceProfileValueCache::GetValues(DeviceProfile, CVar, CVarName);
+	if (bCheckCurrentStateOnly == false)
+	{
+		const FNiagaraCVarValues& PossibleValues = FDeviceProfileValueCache::GetValues(DeviceProfile, CVar, CVarName);
 
-	T Default = PossibleValues.Default.Get<T>();
-	bConditionPassed = CheckValue(Default);
+		T CVarValue;
+		if (PossibleValues.PerQualityLevelValues.IsValidIndex(QualityLevel))
+		{
+			CVarValue = PossibleValues.PerQualityLevelValues[QualityLevel].Get<T>();
+		}
+		else
+		{
+			CVarValue = PossibleValues.Default.Get<T>();
+		}
 
-	bOutCanEverPass = bConditionPassed;
-	bOutCanEverFail = !bConditionPassed;
+		bConditionPassed = CheckValue(CVarValue);
 
- 	if (bCanChangeAtRuntime)
- 	{
- 		bOutCanEverFail = true;
- 		bOutCanEverPass = true;
-	}
- 	else
-	{	
-		for (int32 i = 0; i < PossibleValues.PerQualityLevelValues.Num(); ++i)
+		bOutCanEverPass = bConditionPassed;
+		bOutCanEverFail = !bConditionPassed;
+
+		if (bCanChangeAtRuntime)
 		{
-			T Val = PossibleValues.PerQualityLevelValues[i].Get<T>();
-			bool bPass = CheckValue(Val);
-			if (bPass)
-			{
-				bOutCanEverPass = true;
-			}
-			else
+			bOutCanEverFail = true;
+			bOutCanEverPass = true;
+		}
+		else
+		{
+			for (int32 i = 0; i < PossibleValues.PerQualityLevelValues.Num(); ++i)
 			{
-				bOutCanEverFail = true;
+				T Val = PossibleValues.PerQualityLevelValues[i].Get<T>();
+				bool bPass = CheckValue(Val);
+				if (bPass)
+				{
+					bOutCanEverPass = true;
+				}
+				else
+				{
+					bOutCanEverFail = true;
+				}
 			}
 		}
-	}
-
-#else
 
+		return bConditionPassed;
+	}
+#endif
+	
 	bConditionPassed = CheckValue(GetCVarValue<T>(CVar));
 
 	if (bCanChangeAtRuntime)
@@ -1517,26 +1530,24 @@ bool FNiagaraPlatformSetCVarCondition::Evaluate_Internal(const UDeviceProfile* D
 		bOutCanEverFail = !bConditionPassed;
 	}
 
-#endif
-
 	return bConditionPassed;
 }
 
-bool FNiagaraPlatformSetCVarCondition::Evaluate(const UDeviceProfile* DeviceProfile, bool& bOutCanEverPass, bool& bOutCanEverFail)const
+bool FNiagaraPlatformSetCVarCondition::Evaluate(const UDeviceProfile* DeviceProfile, int32 QualityLevel, bool bCheckCurrentStateOnly, bool& bOutCanEverPass, bool& bOutCanEverFail)const
 {
 	if (IConsoleVariable* CVar = GetCVar())
 	{
 		if (CVar->IsVariableBool())
 		{
-			return Evaluate_Internal<bool>(DeviceProfile, bOutCanEverPass, bOutCanEverFail);
+			return Evaluate_Internal<bool>(DeviceProfile, QualityLevel, bCheckCurrentStateOnly, bOutCanEverPass, bOutCanEverFail);
 		}
 		else if (CVar->IsVariableInt())
 		{
-			return Evaluate_Internal<int32>(DeviceProfile, bOutCanEverPass, bOutCanEverFail);
+			return Evaluate_Internal<int32>(DeviceProfile, QualityLevel, bCheckCurrentStateOnly, bOutCanEverPass, bOutCanEverFail);
 		}
 		else if (CVar->IsVariableFloat())
 		{
-			return Evaluate_Internal<float>(DeviceProfile, bOutCanEverPass, bOutCanEverFail);
+			return Evaluate_Internal<float>(DeviceProfile, QualityLevel, bCheckCurrentStateOnly, bOutCanEverPass, bOutCanEverFail);
 		}
 		else
 		{
diff --git a/Engine/Plugins/FX/Niagara/Source/NiagaraEditor/Private/Customizations/NiagaraPlatformSetCustomization.cpp b/Engine/Plugins/FX/Niagara/Source/NiagaraEditor/Private/Customizations/NiagaraPlatformSetCustomization.cpp
--- a/Engine/Plugins/FX/Niagara/Source/NiagaraEditor/Private/Customizations/NiagaraPlatformSetCustomization.cpp
+++ b/Engine/Plugins/FX/Niagara/Source/NiagaraEditor/Private/Customizations/NiagaraPlatformSetCustomization.cpp
@@ -24,6 +24,7 @@
 #include "NiagaraSettings.h"
 #include "IDetailPropertyRow.h"
 
+
 #define LOCTEXT_NAMESPACE "FNiagaraPlatformSetCustomization"
 
 void FNiagaraPlatformSetCustomization::CustomizeHeader(TSharedRef<IPropertyHandle> InPropertyHandle, FDetailWidgetRow& HeaderRow, IPropertyTypeCustomizationUtils& CustomizationUtils)
@@ -740,7 +741,7 @@ TSharedRef<ITableRow> FNiagaraPlatformSetCustomization::OnGenerateDeviceProfileT
 	SAssignNew(RowContainer, SHorizontalBox);
 
 	FNiagaraPlatformSetEnabledStateDetails Details;
-	FNiagaraPlatformSetEnabledState EnabledState = TargetPlatformSet->IsEnabled(InItem->Profile, QualityLevel, &Details);
+	FNiagaraPlatformSetEnabledState EnabledState = TargetPlatformSet->IsEnabled(InItem->Profile, QualityLevel, false, &Details);
 
 	FText NameTooltip = FText::Format(LOCTEXT("ProfileQLTooltipFmt", "Effects Quality: {0}"), FNiagaraPlatformSet::GetQualityLevelMaskText(QualityLevel));
 
