- Handle invalid data consistently between CPU & GPU

diff --git a/Engine/Plugins/FX/Niagara/Shaders/Private/NiagaraDataInterfaceSkeletalMesh.ush b/Engine/Plugins/FX/Niagara/Shaders/Private/NiagaraDataInterfaceSkeletalMesh.ush
--- a/Engine/Plugins/FX/Niagara/Shaders/Private/NiagaraDataInterfaceSkeletalMesh.ush
+++ b/Engine/Plugins/FX/Niagara/Shaders/Private/NiagaraDataInterfaceSkeletalMesh.ush
@@ -75,6 +75,7 @@ struct FDISkelMeshContext
 	bool			UniformTriangleSamplingEnable;
 	bool			UniformSamplingRegionEnabled;
 	bool			UnlimitedBoneInfluences;
+	bool			bHasMeshColors;
 };
 
 ////////////////////////////////////////////////////////////////////////////////////////////////////////////////////////////////////////////////////////////////
@@ -146,6 +147,18 @@ float4 DISkelMesh_UnpackWeights4(uint Packed)
 
 void DISkelMesh_GetVertex(in FDISkelMeshContext DIContext, uint VertexIndex, out float3 OutPosition, out float3 OutTangentX, out float3 OutTangentY, out float3 OutTangentZ)
 {
+	// Early out for bad data
+	if ( DIContext.MeshVertexCount == 0 )
+	{
+		OutPosition = float3(0.0f, 0.0f, 0.0f);
+		OutTangentX = float3(1.0f, 0.0f, 0.0f);
+		OutTangentY = float3(0.0f, 1.0f, 0.0f);
+		OutTangentZ = float3(0.0f, 0.0f, 1.0f);
+		return;
+	}
+
+	VertexIndex = clamp(VertexIndex, 0, DIContext.MeshVertexCount - 1);
+
 	//-TODO: Use a R32G32B32 format
 	float3 Position = float3(DIContext.MeshVertexBuffer[VertexIndex * 3], DIContext.MeshVertexBuffer[VertexIndex * 3 + 1], DIContext.MeshVertexBuffer[VertexIndex * 3 + 2]);
 	float3 TangentX = TangentBias(DIContext.MeshTangentBuffer[VertexIndex * 2    ].xyz);
@@ -160,6 +173,18 @@ void DISkelMesh_GetVertex(in FDISkelMeshContext DIContext, uint VertexIndex, out
 
 void DISkelMesh_GetPointOnTriangle(in FDISkelMeshContext DIContext, uint TriangleIndex, float3 BaryCoord, out float3 OutPosition, out float3 OutTangentX, out float3 OutTangentY, out float3 OutTangentZ)
 {
+	// Early out for bad data
+	if ( DIContext.MeshTriangleCount == 0 )
+	{
+		OutPosition = float3(0.0f, 0.0f, 0.0f);
+		OutTangentX = float3(1.0f, 0.0f, 0.0f);
+		OutTangentY = float3(0.0f, 1.0f, 0.0f);
+		OutTangentZ = float3(0.0f, 0.0f, 1.0f);
+		return;
+	}
+
+	TriangleIndex = clamp(TriangleIndex, 0, DIContext.MeshTriangleCount - 1);
+
 	uint IndexBufferOffset = TriangleIndex * 3;
 	uint VertexIndex0 = DIContext.MeshIndexBuffer[IndexBufferOffset];
 	uint VertexIndex1 = DIContext.MeshIndexBuffer[IndexBufferOffset + 1];
@@ -177,6 +202,20 @@ void DISkelMesh_GetPointOnTriangle(in FDISkelMeshContext DIContext, uint Triangl
 
 FDISkelMeshSkinnedVertex DISkelMesh_GetSkinnedVertex(in FDISkelMeshContext DIContext, uint VertexIndex)
 {
+	// Early out for bad data
+	if ( DIContext.MeshVertexCount == 0 )
+	{
+		FDISkelMeshSkinnedVertex SkinnedVertex;
+		SkinnedVertex.Position = float3(0.0f, 0.0f, 0.0f);
+		SkinnedVertex.PrevPosition = float3(0.0f, 0.0f, 0.0f);
+		SkinnedVertex.TangentX = float3(1.0f, 0.0f, 0.0f);
+		SkinnedVertex.TangentY = float3(0.0f, 1.0f, 0.0f);
+		SkinnedVertex.TangentZ = float3(0.0f, 0.0f, 1.0f);
+		return SkinnedVertex;
+	}
+
+	VertexIndex = clamp(VertexIndex, 0, DIContext.MeshVertexCount - 1);
+
 	//-TODO: Use a R32G32B32 format
 	float3 CurrPosition = float3(DIContext.MeshVertexBuffer[VertexIndex * 3], DIContext.MeshVertexBuffer[VertexIndex * 3 + 1], DIContext.MeshVertexBuffer[VertexIndex * 3 + 2]);
 	float3 PrevPosition = CurrPosition;
@@ -267,6 +306,20 @@ FDISkelMeshSkinnedVertex DISkelMesh_GetSkinnedVertex(in FDISkelMeshContext DICon
 
 FDISkelMeshSkinnedVertex DISkelMesh_GetSkinnedPointOnTriangle(in FDISkelMeshContext DIContext, uint TriangleIndex, float3 BaryCoord)
 {
+	// Early out for bad data
+	if ( DIContext.MeshTriangleCount == 0 )
+	{
+		FDISkelMeshSkinnedVertex SkinnedVertex;
+		SkinnedVertex.Position = float3(0.0f, 0.0f, 0.0f);
+		SkinnedVertex.PrevPosition = float3(0.0f, 0.0f, 0.0f);
+		SkinnedVertex.TangentX = float3(1.0f, 0.0f, 0.0f);
+		SkinnedVertex.TangentY = float3(0.0f, 1.0f, 0.0f);
+		SkinnedVertex.TangentZ = float3(0.0f, 0.0f, 1.0f);
+		return SkinnedVertex;
+	}
+
+	TriangleIndex = clamp(TriangleIndex, 0, DIContext.MeshTriangleCount - 1);
+
 	uint IndexBufferOffset = TriangleIndex * 3;
 	uint VertexIndex0 = DIContext.MeshIndexBuffer[IndexBufferOffset];
 	uint VertexIndex1 = DIContext.MeshIndexBuffer[IndexBufferOffset + 1];
@@ -300,6 +353,14 @@ FDISkelMeshSkinnedVertex DISkelMesh_GetSkinnedPointOnTriangleWS(in FDISkelMeshCo
 
 void DISKelMesh_RandomTriCoord(in FDISkelMeshContext DIContext, uint Seed1, uint Seed2, uint Seed3, out int OutTriangle, out float3 OutBaryCoord)
 {
+	// Early out for bad data
+	if (DIContext.MeshTriangleCount == 0)
+	{
+		OutTriangle = -1;
+		OutBaryCoord = 1.0f / 3.0f;
+		return;
+	}
+
 	float RandT0 = DISKelMesh_Random(Seed1, Seed2, Seed3);
 
 	[branch]
@@ -442,6 +503,14 @@ void DISKelMesh_GetTriangleCount(in FDISkelMeshContext DIContext, out int OutCou
 
 void DISKelMesh_RandomFilteredTriangle(in FDISkelMeshContext DIContext, uint Seed1, uint Seed2, uint Seed3, out int OutTriangle, out float3 OutBaryCoord)
 {
+	// Early out for bad data
+	if (DIContext.MeshNumSamplingRegionTriangles == 0)
+	{
+		OutTriangle = -1;
+		OutBaryCoord = 1.0f / 3.0f;
+		return;
+	}
+
 	float RandT0 = DISKelMesh_Random(Seed1, Seed2, Seed3);
 	uint RegionTriangle = min(uint(RandT0 * float(DIContext.MeshNumSamplingRegionTriangles)), DIContext.MeshNumSamplingRegionTriangles - 1); // avoid % by using mul/min to Tri = MeshTriangleCountName
 
@@ -745,6 +814,17 @@ float4 DISKelMesh_QuatSlerp(float4 Quat1, float4 Quat2, float Slerp)
 
 void DISkelMesh_GetSkinnedBoneCommon(in FDISkelMeshContext DIContext, in int Bone, in float Interp, in bool bWorldSpace, out float3 OutPosition, out float4 OutRotation, out float3 OutVelocity)
 {
+	// Early out for bad data
+	if ( DIContext.NumBones == 0 )
+	{
+		OutPosition = float3(0.0f, 0.0f, 0.0f);
+		OutRotation = float4(0.0f, 0.0f, 0.0f, 1.0f);
+		OutVelocity = float3(0.0f, 0.0f, 0.0f);
+		return;
+	}
+
+	Bone = clamp(Bone, 0, DIContext.NumBones - 1);
+
     float3 PrevPosition = DIContext.MeshPrevSamplingBonesBuffer[Bone * 3].xyz;
     float4 PrevRotation = DIContext.MeshPrevSamplingBonesBuffer[Bone * 3 + 1];
     float3 CurrPosition = DIContext.MeshCurrSamplingBonesBuffer[Bone * 3].xyz;
@@ -789,6 +869,16 @@ void DISkelMesh_GetSkinnedBoneInterpolatedWS(in FDISkelMeshContext DIContext, in
 
 void DISkelMesh_GetTriVertices(in FDISkelMeshContext DIContext, in uint TriangleIndex, out int VertexIndex0, out int VertexIndex1, out int VertexIndex2)
 {
+	// Early out for bad data
+	if ( DIContext.MeshTriangleCount == 0 )
+	{
+		VertexIndex0 = 0;
+		VertexIndex1 = 0;
+		VertexIndex2 = 0;
+		return;
+	}
+
+	TriangleIndex = clamp(TriangleIndex, 0, DIContext.MeshTriangleCount - 1);
     uint IndexBufferOffset = TriangleIndex * 3;
     VertexIndex0 = DIContext.MeshIndexBuffer[IndexBufferOffset];
     VertexIndex1 = DIContext.MeshIndexBuffer[IndexBufferOffset + 1];
@@ -797,6 +887,14 @@ void DISkelMesh_GetTriVertices(in FDISkelMeshContext DIContext, in uint Triangle
 
 void DISkelMesh_GetTriColor(in FDISkelMeshContext DIContext, in uint TriangleIndex, in float3 BaryCoord, out float4 OutColor)
 {
+	// Early out for bad data
+	if ( DIContext.bHasMeshColors == false )
+	{
+		OutColor = float4(1, 1, 1, 1);
+		return;
+	}
+	TriangleIndex = clamp(TriangleIndex, 0, DIContext.MeshTriangleCount - 1);
+
     uint IndexBufferOffset = TriangleIndex * 3;
     uint VertexIndex0 = DIContext.MeshIndexBuffer[IndexBufferOffset];
     uint VertexIndex1 = DIContext.MeshIndexBuffer[IndexBufferOffset + 1];
@@ -827,7 +925,7 @@ void DISkelMesh_GetSkinnedVertexWS(in FDISkelMeshContext DIContext, in int Verte
 
 void DISkelMesh_GetVertexColor(in FDISkelMeshContext DIContext, in int VertexIndex, out float4 OutColor)
 {
-    OutColor = DIContext.MeshColorBuffer[VertexIndex] FMANUALFETCH_COLOR_COMPONENT_SWIZZLE;
+    OutColor = DIContext.bHasMeshColors ? DIContext.MeshColorBuffer[VertexIndex] FMANUALFETCH_COLOR_COMPONENT_SWIZZLE : float4(1, 1, 1, 1);
 }
 
 void DISkelMesh_GetVertexUV(in FDISkelMeshContext DIContext, in int VertexIndex, in int UVSet, out float2 OutUV)
@@ -867,6 +965,13 @@ void DISkelMesh_IsValidFilteredVertex(in FDISkelMeshContext DIContext, int Filte
 
 void DISkelMesh_GetRandomFilteredVertex(in FDISkelMeshContext DIContext, uint Seed1, uint Seed2, uint Seed3, out int OutVertex)
 {
+	// Early out for bad data
+	if ( DIContext.MeshNumSamplingRegionVertices == 0 )
+	{
+		OutVertex = -1;
+		return;
+	}
+
 	float RandT0 = DISKelMesh_Random(Seed1, Seed2, Seed3);
 	int FilteredIndex = (int)min(uint(RandT0 * float(DIContext.MeshNumSamplingRegionVertices)), DIContext.MeshNumSamplingRegionVertices - 1); // avoid % by using mul/min to Tri = MeshVertexCount
 	OutVertex = (int)DIContext.MeshSampleRegionsVertices[FilteredIndex];
@@ -879,6 +984,13 @@ void DISkelMesh_GetFilteredVertexCount(in FDISkelMeshContext DIContext, out int
 
 void DISkelMesh_GetFilteredVertexAt(in FDISkelMeshContext DIContext, int FilteredIndex, out int VertexIndex)
 {
+	// Early out for bad data
+	if ( DIContext.MeshNumSamplingRegionVertices == 0 )
+	{
+		VertexIndex = -1;
+		return;
+	}
+
 	FilteredIndex = clamp(FilteredIndex, 0, (int)DIContext.MeshNumSamplingRegionVertices - 1);
 	VertexIndex = (int)DIContext.MeshSampleRegionsVertices[FilteredIndex];
 }
@@ -920,14 +1032,14 @@ void DISkelMesh_GetFilteredBoneCount(in FDISkelMeshContext DIContext, out int Co
 void DISkelMesh_GetFilteredBoneAt(in FDISkelMeshContext DIContext, in int InBoneIndex, out int Bone)
 {
     int BoneIndex = min(InBoneIndex, DIContext.NumFilteredBones - 1);
-    Bone = (int)DIContext.FilteredAndUnfilteredBones[BoneIndex];
+    Bone = DIContext.NumFilteredBones > 0 ? (int)DIContext.FilteredAndUnfilteredBones[BoneIndex] : -1;
 }
 
 void DISkelMesh_RandomFilteredBone(in FDISkelMeshContext DIContext, uint Seed1, uint Seed2, uint Seed3, out int Bone)
 {
 	float RandT0 = DISKelMesh_Random(Seed1, Seed2, Seed3);
     int FilteredBone = min(int(RandT0 * float(DIContext.NumFilteredBones)), DIContext.NumFilteredBones - 1); // avoid % by using mul/min to Tri = MeshVertexCount
-    Bone = (int)DIContext.FilteredAndUnfilteredBones[FilteredBone];
+    Bone = DIContext.NumFilteredBones > 0 ? (int)DIContext.FilteredAndUnfilteredBones[FilteredBone] : -1;
 }
 
 void DISkelMesh_GetUnfilteredBoneCount(in FDISkelMeshContext DIContext, out int Count)
@@ -1231,3 +1343,4 @@ void DISkelMesh_GetTriangleNeighbor(in FDISkelMeshContext DIContext, int Triangl
 	DIContext.UniformTriangleSamplingEnable = EnabledFeatures_ ## NAME & 0x1; \
 	DIContext.UniformSamplingRegionEnabled = EnabledFeatures_ ## NAME & 0x2; \
 	DIContext.UnlimitedBoneInfluences = EnabledFeatures_ ## NAME & 0x4; \
+	DIContext.bHasMeshColors = EnabledFeatures_ ## NAME & 0x8; \
diff --git a/Engine/Plugins/FX/Niagara/Source/Niagara/Classes/NiagaraDataInterfaceSkeletalMesh.h b/Engine/Plugins/FX/Niagara/Source/Niagara/Classes/NiagaraDataInterfaceSkeletalMesh.h
--- a/Engine/Plugins/FX/Niagara/Source/Niagara/Classes/NiagaraDataInterfaceSkeletalMesh.h
+++ b/Engine/Plugins/FX/Niagara/Source/Niagara/Classes/NiagaraDataInterfaceSkeletalMesh.h
@@ -455,6 +455,9 @@ public:
 	FRHIShaderResourceView* GetSampleRegionsTriangleIndicesSRV() const { return SampleRegionsTriangleIndicesSRV; }
 	FRHIShaderResourceView* GetSampleRegionsVerticesSRV() const { return SampleRegionsVerticesSRV; }
 
+	bool IsMeshValid() const { return bMeshValid; }
+	bool HasMeshColors() const { return bHasMeshColors; }
+
 	FRHIShaderResourceView* GetBufferPositionSRV() const { return MeshVertexBufferSRV; }
 	FRHIShaderResourceView* GetBufferIndexSRV() const { return MeshIndexBufferSRV; }
 	FRHIShaderResourceView* GetBufferTangentSRV() const { return MeshTangentBufferSRV; }
@@ -503,6 +506,8 @@ protected:
 	int32 FilteredSocketBoneOffset = 0;
 
 	/** Cached SRV to gpu buffers of the mesh we spawn from */
+	bool bMeshValid = false;
+	bool bHasMeshColors = false;
 	FShaderResourceViewRHIRef MeshVertexBufferSRV;
 	FShaderResourceViewRHIRef MeshIndexBufferSRV;
 	FShaderResourceViewRHIRef MeshTangentBufferSRV;
diff --git a/Engine/Plugins/FX/Niagara/Source/Niagara/Private/NiagaraDataInterfaceSkeletalMesh.cpp b/Engine/Plugins/FX/Niagara/Source/Niagara/Private/NiagaraDataInterfaceSkeletalMesh.cpp
--- a/Engine/Plugins/FX/Niagara/Source/Niagara/Private/NiagaraDataInterfaceSkeletalMesh.cpp
+++ b/Engine/Plugins/FX/Niagara/Source/Niagara/Private/NiagaraDataInterfaceSkeletalMesh.cpp
@@ -828,16 +828,48 @@ void FSkeletalMeshGpuSpawnStaticBuffers::InitRHI()
 	// Also see https://jira.it.epicgames.net/browse/UE-69376 : we would need to know if GPU sampling of the mesh surface is needed or not on the mesh to be able to do that.
 	// ALso today we do not know if an interface is create from a CPU or GPU emitter. So always allocate for now. Follow up in https://jira.it.epicgames.net/browse/UE-69375.
 
-	MeshIndexBufferSRV = FNiagaraRenderer::GetSrvOrDefaultUInt(LODRenderData->MultiSizeIndexContainer.GetIndexBuffer()->GetSRV());
-	MeshVertexBufferSRV = FNiagaraRenderer::GetSrvOrDefaultFloat(LODRenderData->StaticVertexBuffers.PositionVertexBuffer.GetSRV());
-	MeshTangentBufferSRV = FNiagaraRenderer::GetSrvOrDefaultFloat4(LODRenderData->StaticVertexBuffers.StaticMeshVertexBuffer.GetTangentsSRV());
-	MeshTexCoordBufferSRV = FNiagaraRenderer::GetSrvOrDefaultFloat2(LODRenderData->StaticVertexBuffers.StaticMeshVertexBuffer.GetTexCoordsSRV());
-	MeshColorBufferSRV = FNiagaraRenderer::GetSrvOrDefaultFloat4(LODRenderData->StaticVertexBuffers.ColorVertexBuffer.GetColorComponentsSRV());
-	NumTexCoord = LODRenderData->StaticVertexBuffers.StaticMeshVertexBuffer.GetNumTexCoords();
-
+	MeshIndexBufferSRV = LODRenderData->MultiSizeIndexContainer.GetIndexBuffer()->GetSRV();
+	MeshVertexBufferSRV = LODRenderData->StaticVertexBuffers.PositionVertexBuffer.GetSRV();
+	MeshTangentBufferSRV = LODRenderData->StaticVertexBuffers.StaticMeshVertexBuffer.GetTangentsSRV();
+	bMeshValid = MeshIndexBufferSRV.IsValid() && MeshVertexBufferSRV.IsValid() && MeshTangentBufferSRV.IsValid();
 	NumWeights = LODRenderData->SkinWeightVertexBuffer.GetMaxBoneInfluences();
 
-	uint32 SectionCount = LODRenderData->RenderSections.Num();
+	if ( bMeshValid == false )
+	{
+		MeshIndexBufferSRV = FNiagaraRenderer::GetDummyUIntBuffer();
+		MeshVertexBufferSRV = FNiagaraRenderer::GetDummyFloatBuffer();
+		MeshTangentBufferSRV = FNiagaraRenderer::GetDummyFloat4Buffer();
+		NumWeights = 0;
+		TriangleCount = 0;
+		VertexCount = 0;
+		bUseGpuUniformlyDistributedSampling = false;
+		NumSamplingRegionTriangles = 0;
+		bSamplingRegionsAllAreaWeighted = false;
+
+		SampleRegionsProbAlias.Empty();
+		SampleRegionsTriangleIndicies.Empty();
+		SampleRegionsVertices.Empty();
+	}
+
+	MeshColorBufferSRV = LODRenderData->StaticVertexBuffers.ColorVertexBuffer.GetColorComponentsSRV();
+	bHasMeshColors = bMeshValid && MeshColorBufferSRV.IsValid();
+	if (bHasMeshColors == false)
+	{
+		MeshColorBufferSRV = FNiagaraRenderer::GetDummyFloat4Buffer();
+	}
+
+	MeshTexCoordBufferSRV = LODRenderData->StaticVertexBuffers.StaticMeshVertexBuffer.GetTexCoordsSRV();
+	if (MeshTexCoordBufferSRV.IsValid())
+	{
+		NumTexCoord = LODRenderData->StaticVertexBuffers.StaticMeshVertexBuffer.GetNumTexCoords();
+	}
+	else
+	{
+		MeshTexCoordBufferSRV = FNiagaraRenderer::GetDummyFloat2Buffer();
+		NumTexCoord = 0;
+	}
+
+	const uint32 SectionCount = LODRenderData->RenderSections.Num();
 
 #if STATS
 	ensure(GPUMemoryUsage == 0);
@@ -880,6 +912,11 @@ void FSkeletalMeshGpuSpawnStaticBuffers::InitRHI()
 			GPUMemoryUsage += SampleRegionsProbAlias.Num() * SampleRegionsProbAlias.GetTypeSize();
 #endif
 		}
+		else
+		{
+			SampleRegionsProbAliasSRV = FNiagaraRenderer::GetDummyUIntBuffer();
+		}
+
 		CreateInfo.DebugName = (TEXT("SampleRegionsTriangleIndicesBuffer"));
 		CreateInfo.ResourceArray = &SampleRegionsTriangleIndicies;
 		SampleRegionsTriangleIndicesBuffer = RHICreateVertexBuffer(SampleRegionsTriangleIndicies.Num() * SampleRegionsTriangleIndicies.GetTypeSize(), BUF_Static | BUF_ShaderResource, CreateInfo);
@@ -896,10 +933,17 @@ void FSkeletalMeshGpuSpawnStaticBuffers::InitRHI()
 		GPUMemoryUsage += SampleRegionsVertices.Num() * SampleRegionsVertices.GetTypeSize();
 #endif
 	}
+	else
+	{
+		SampleRegionsProbAliasSRV = FNiagaraRenderer::GetDummyUIntBuffer();
+		SampleRegionsTriangleIndicesSRV = FNiagaraRenderer::GetDummyUIntBuffer();
+		SampleRegionsVerticesSRV = FNiagaraRenderer::GetDummyUIntBuffer();
+	}
 
 	// Prepare the vertex matrix lookup offset for each of the sections. This is needed because per vertex BlendIndicies are stored relatively to each Section used matrices.
 	// And these offset per section need to point to the correct matrix according to each section BoneMap.
 	// There is not section selection/culling in the interface so technically we could compute that array in the pipeline.
+	if ( bMeshValid )
 	{
 		FRHIResourceCreateInfo CreateInfo(TEXT("FSkeletalMeshGpuSpawnStaticBuffers"));
 		BufferTriangleMatricesOffsetRHI = RHICreateBuffer(VertexCount * sizeof(uint32), BUF_Static | BUF_VertexBuffer | BUF_ShaderResource, 0, ERHIAccess::VertexOrIndexBuffer | ERHIAccess::SRVMask, CreateInfo);
@@ -922,6 +966,10 @@ void FSkeletalMeshGpuSpawnStaticBuffers::InitRHI()
 		GPUMemoryUsage += VertexCount * sizeof(uint32);
 #endif
 	}
+	else
+	{
+		BufferTriangleMatricesOffsetSRV = FNiagaraRenderer::GetDummyUIntBuffer();
+	}
 
 	// Create arrays for filtered bones / sockets
 	if ( FilteredAndUnfilteredBonesArray.Num() > 0 )
@@ -1398,17 +1446,18 @@ public:
 		FRHIComputeShader* ComputeShaderRHI = Context.Shader.GetComputeShader();
 		FNiagaraDataInterfaceProxySkeletalMesh* InterfaceProxy = static_cast<FNiagaraDataInterfaceProxySkeletalMesh*>(Context.DataInterface);
 		FNiagaraDataInterfaceProxySkeletalMeshData* InstanceData = InterfaceProxy->SystemInstancesToData.Find(Context.SystemInstanceID);
-		if (InstanceData && InstanceData->StaticBuffers && InstanceData->StaticBuffers->GetBufferPositionSRV())
+		if (InstanceData && InstanceData->StaticBuffers)
 		{
 			FSkeletalMeshGpuSpawnStaticBuffers* StaticBuffers = InstanceData->StaticBuffers;
 
+			// Bind mesh buffers
 			SetSRVParameter(RHICmdList, ComputeShaderRHI, MeshVertexBuffer, StaticBuffers->GetBufferPositionSRV());
 			SetSRVParameter(RHICmdList, ComputeShaderRHI, MeshIndexBuffer, StaticBuffers->GetBufferIndexSRV());
 			SetSRVParameter(RHICmdList, ComputeShaderRHI, MeshTangentBuffer, StaticBuffers->GetBufferTangentSRV());
-
-			SetShaderValue(RHICmdList, ComputeShaderRHI, MeshNumTexCoord, StaticBuffers->GetNumTexCoord());
 			SetSRVParameter(RHICmdList, ComputeShaderRHI, MeshTexCoordBuffer, StaticBuffers->GetBufferTexCoordSRV());
 			SetSRVParameter(RHICmdList, ComputeShaderRHI, MeshColorBuffer, StaticBuffers->GetBufferColorSRV());
+
+			SetShaderValue(RHICmdList, ComputeShaderRHI, MeshNumTexCoord, StaticBuffers->GetNumTexCoord());
 			SetShaderValue(RHICmdList, ComputeShaderRHI, MeshTriangleCount, StaticBuffers->GetTriangleCount());
 			SetShaderValue(RHICmdList, ComputeShaderRHI, MeshVertexCount, StaticBuffers->GetVertexCount());
 
@@ -1433,7 +1482,8 @@ public:
 			uint32 EnabledFeaturesBits = 0;
 			EnabledFeaturesBits |= StaticBuffers->IsUseGpuUniformlyDistributedSampling() ? 1 : 0;
 			EnabledFeaturesBits |= StaticBuffers->IsSamplingRegionsAllAreaWeighted() ? 2 : 0;
-			EnabledFeaturesBits |= (InstanceData->bUnlimitedBoneInfluences ? 4 : 0);
+			EnabledFeaturesBits |= InstanceData->bUnlimitedBoneInfluences ? 4 : 0;
+			EnabledFeaturesBits |= StaticBuffers->HasMeshColors() ? 8 : 0;
 
 			FSkeletalMeshGpuDynamicBufferProxy* DynamicBuffers = InstanceData->DynamicBuffer;
 			check(DynamicBuffers);
@@ -1508,11 +1558,11 @@ public:
 			// Bind dummy buffers
 			SetSRVParameter(RHICmdList, ComputeShaderRHI, MeshVertexBuffer, FNiagaraRenderer::GetDummyFloatBuffer());
 			SetSRVParameter(RHICmdList, ComputeShaderRHI, MeshIndexBuffer, FNiagaraRenderer::GetDummyUIntBuffer());
-			SetSRVParameter(RHICmdList, ComputeShaderRHI, MeshTangentBuffer, FNiagaraRenderer::GetDummyFloatBuffer());
+			SetSRVParameter(RHICmdList, ComputeShaderRHI, MeshTangentBuffer, FNiagaraRenderer::GetDummyFloat4Buffer());
 
 			SetShaderValue(RHICmdList, ComputeShaderRHI, MeshNumTexCoord, 0);
-			SetSRVParameter(RHICmdList, ComputeShaderRHI, MeshTexCoordBuffer, FNiagaraRenderer::GetDummyFloatBuffer());
-			SetSRVParameter(RHICmdList, ComputeShaderRHI, MeshColorBuffer, FNiagaraRenderer::GetDummyFloatBuffer());
+			SetSRVParameter(RHICmdList, ComputeShaderRHI, MeshTexCoordBuffer, FNiagaraRenderer::GetDummyFloat2Buffer());
+			SetSRVParameter(RHICmdList, ComputeShaderRHI, MeshColorBuffer, FNiagaraRenderer::GetDummyFloat4Buffer());
 			SetShaderValue(RHICmdList, ComputeShaderRHI, MeshTriangleCount, 0);
 			SetShaderValue(RHICmdList, ComputeShaderRHI, MeshVertexCount, 0);
 			SetSRVParameter(RHICmdList, ComputeShaderRHI, MeshTriangleSamplerProbAliasBuffer, FNiagaraRenderer::GetDummyUIntBuffer());
