Changed some raw object pointers to TObjectPtr
Updated some tooltips.

#rnx

diff --git a/Engine/Plugins/FX/Niagara/Source/NiagaraEditor/Public/NiagaraValidationRules.h b/Engine/Plugins/FX/Niagara/Source/NiagaraEditor/Public/NiagaraValidationRules.h
--- a/Engine/Plugins/FX/Niagara/Source/NiagaraEditor/Public/NiagaraValidationRules.h
+++ b/Engine/Plugins/FX/Niagara/Source/NiagaraEditor/Public/NiagaraValidationRules.h
@@ -10,6 +10,7 @@
 
 class UNiagaraScript;
 
+/** This validation rule ensures that no Systems have a warm up time set. */
 UCLASS(Category = "Validation", DisplayName = "No Warmup Time")
 class UNiagaraValidationRule_NoWarmupTime : public UNiagaraValidationRule
 {
@@ -18,7 +19,7 @@ public:
 	virtual void CheckValidity(TSharedPtr<FNiagaraSystemViewModel> ViewModel, TArray<FNiagaraValidationResult>& OutResults) const override;
 };
 
-
+/** This validation rule ensures that GPU emitters have fixed bounds set. */
 UCLASS(Category = "Validation", DisplayName = "Fixed GPU Bounds Set")
 class UNiagaraValidationRule_FixedGPUBoundsSet : public UNiagaraValidationRule
 {
@@ -27,7 +28,7 @@ public:
 	virtual void CheckValidity(TSharedPtr<FNiagaraSystemViewModel> ViewModel, TArray<FNiagaraValidationResult>& OutResults) const override;
 };
 
-
+/** This validation rule can ban the use of certain renderers on all or a subset of platforms. */
 UCLASS(Category = "Validation", DisplayName = "Banned Renderers")
 class UNiagaraValidationRule_BannedRenderers : public UNiagaraValidationRule
 {
@@ -45,6 +46,7 @@ public:
 	virtual void CheckValidity(TSharedPtr<FNiagaraSystemViewModel> ViewModel, TArray<FNiagaraValidationResult>& OutResults) const override;
 };
 
+/** This validation rule can ban the use of certain modules on all or a subset of platforms. */
 UCLASS(Category = "Validation", DisplayName = "Banned Modules")
 class UNiagaraValidationRule_BannedModules : public UNiagaraValidationRule
 {
@@ -62,6 +64,7 @@ public:
 	virtual void CheckValidity(TSharedPtr<FNiagaraSystemViewModel> ViewModel, TArray<FNiagaraValidationResult>& OutResults) const override;
 };
 
+/** This validation rule can marks this effect type as invalid and so must be changed. Forces a choice of correct Effect Type for an System rather than. Leaving as the default. */
 UCLASS(Category = "Validation", DisplayName = "Invalid Effect Type")
 class UNiagaraValidationRule_InvalidEffectType : public UNiagaraValidationRule
 {
@@ -70,6 +73,7 @@ public:
 	virtual void CheckValidity(TSharedPtr<FNiagaraSystemViewModel> ViewModel, TArray<FNiagaraValidationResult>& OutResults) const override;
 };
 
+/** This validation rule checks for various issue with Large World Coordinates. */
 UCLASS(Category = "Validation", DisplayName = "Large World Coordinates")
 class UNiagaraValidationRule_LWC : public UNiagaraValidationRule
 {
