Fix type checking of visited parameters when evaluating Fail If Not Set errors as we may have intentional type mismatches for LWC types.

diff --git a/Engine/Plugins/FX/Niagara/Source/Niagara/Private/NiagaraCommon.cpp b/Engine/Plugins/FX/Niagara/Source/Niagara/Private/NiagaraCommon.cpp
--- a/Engine/Plugins/FX/Niagara/Source/Niagara/Private/NiagaraCommon.cpp
+++ b/Engine/Plugins/FX/Niagara/Source/Niagara/Private/NiagaraCommon.cpp
@@ -8,6 +8,7 @@
 #include "NiagaraCustomVersion.h"
 #include "NiagaraParameterCollection.h"
 #include "NiagaraScriptSourceBase.h"
+#include "NiagaraSettings.h"
 #include "NiagaraStats.h"
 #include "NiagaraSystemInstance.h"
 #include "NiagaraWorldManager.h"
@@ -1402,6 +1403,19 @@ void FNiagaraUtilities::PrepareRapidIterationParameters(const TArray<UNiagaraScr
 		}
 	}
 }
+
+bool FNiagaraUtilities::AreTypesAssignable(const FNiagaraTypeDefinition& TypeA, const FNiagaraTypeDefinition& TypeB)
+{
+	const UNiagaraSettings* Settings = GetDefault<UNiagaraSettings>();
+	if (Settings->bEnforceStrictStackTypes)
+	{
+		return TypeA == TypeB;
+	}
+	return (TypeA == TypeB)
+		|| (TypeA == FNiagaraTypeDefinition::GetPositionDef() && TypeB == FNiagaraTypeDefinition::GetVec3Def())
+		|| (TypeB == FNiagaraTypeDefinition::GetPositionDef() && TypeA == FNiagaraTypeDefinition::GetVec3Def());
+}
+
 #endif
 
 //////////////////////////////////////////////////////////////////////////
diff --git a/Engine/Plugins/FX/Niagara/Source/Niagara/Private/NiagaraSystem.cpp b/Engine/Plugins/FX/Niagara/Source/Niagara/Private/NiagaraSystem.cpp
--- a/Engine/Plugins/FX/Niagara/Source/Niagara/Private/NiagaraSystem.cpp
+++ b/Engine/Plugins/FX/Niagara/Source/Niagara/Private/NiagaraSystem.cpp
@@ -2717,6 +2717,10 @@ void UNiagaraSystem::EvaluateCompileResultDependencies() const
 				TestVar.SetName(NewName);
 			}
 
+			auto TestVarNameAndAssignableTypeMatchPred = [&TestVar](const FNiagaraVariableBase& ComparisonVar)->bool {
+				return TestVar.GetName() == ComparisonVar.GetName() && FNiagaraUtilities::AreTypesAssignable(TestVar.GetType(), ComparisonVar.GetType());
+			};
+
 			bool bDependencyMet = false;
 			for (const int32 TestIndex : ValidationInfo.ParentIndices)
 			{
@@ -2726,12 +2730,12 @@ void UNiagaraSystem::EvaluateCompileResultDependencies() const
 				}
 
 				const FScriptCompileResultValidationInfo& TestInfo = ScriptCompilesToValidate[TestIndex];
-				if (TestVar.GetType().IsStatic() && TestInfo.CompileResults->StaticVariablesWritten.Contains(TestVar))
+				if (TestVar.GetType().IsStatic() && TestInfo.CompileResults->StaticVariablesWritten.ContainsByPredicate(TestVarNameAndAssignableTypeMatchPred))
 				{
 					bDependencyMet = true;
 					break;
 				}
-				else if (TestInfo.CompileResults->AttributesWritten.Contains(TestVar))
+				else if (TestInfo.CompileResults->AttributesWritten.ContainsByPredicate(TestVarNameAndAssignableTypeMatchPred))
 				{
 					bDependencyMet = true;
 					break;
diff --git a/Engine/Plugins/FX/Niagara/Source/Niagara/Public/NiagaraCommon.h b/Engine/Plugins/FX/Niagara/Source/Niagara/Public/NiagaraCommon.h
--- a/Engine/Plugins/FX/Niagara/Source/Niagara/Public/NiagaraCommon.h
+++ b/Engine/Plugins/FX/Niagara/Source/Niagara/Public/NiagaraCommon.h
@@ -1257,6 +1257,8 @@ namespace FNiagaraUtilities
 	 * scripts array must have an entry in this map.
 	 */
 	void NIAGARA_API PrepareRapidIterationParameters(const TArray<UNiagaraScript*>& Scripts, const TMap<UNiagaraScript*, UNiagaraScript*>& ScriptDependencyMap, const TMap<UNiagaraScript*, const UNiagaraEmitter*>& ScriptToEmitterNameMap);
+
+	bool NIAGARA_API AreTypesAssignable(const FNiagaraTypeDefinition& TypeA, const FNiagaraTypeDefinition& TypeB);
 #endif
 
 	void NIAGARA_API DumpHLSLText(const FString& SourceCode, const FString& DebugName);
diff --git a/Engine/Plugins/FX/Niagara/Source/NiagaraEditor/Private/NiagaraEditorUtilities.cpp b/Engine/Plugins/FX/Niagara/Source/NiagaraEditor/Private/NiagaraEditorUtilities.cpp
--- a/Engine/Plugins/FX/Niagara/Source/NiagaraEditor/Private/NiagaraEditorUtilities.cpp
+++ b/Engine/Plugins/FX/Niagara/Source/NiagaraEditor/Private/NiagaraEditorUtilities.cpp
@@ -996,14 +996,7 @@ FText FNiagaraEditorUtilities::GetTypeDefinitionCategory(const FNiagaraTypeDefin
 
 bool FNiagaraEditorUtilities::AreTypesAssignable(const FNiagaraTypeDefinition& TypeA, const FNiagaraTypeDefinition& TypeB)
 {
-	const UNiagaraSettings* Settings = GetDefault<UNiagaraSettings>();
-	if (Settings->bEnforceStrictStackTypes)
-	{
-		return TypeA == TypeB;
-	}
-	return (TypeA == TypeB)
-		|| (TypeA == FNiagaraTypeDefinition::GetPositionDef() && TypeB == FNiagaraTypeDefinition::GetVec3Def())
-		|| (TypeB == FNiagaraTypeDefinition::GetPositionDef() && TypeA == FNiagaraTypeDefinition::GetVec3Def());
+	return FNiagaraUtilities::AreTypesAssignable(TypeA, TypeB);
 }
 
 void FNiagaraEditorUtilities::MarkDependentCompilableAssetsDirty(TArray<UObject*> InObjects)
