- Custom element count binding for simulation stages

diff --git a/Engine/Plugins/FX/Niagara/Source/Niagara/Classes/NiagaraGPUSystemTick.h b/Engine/Plugins/FX/Niagara/Source/Niagara/Classes/NiagaraGPUSystemTick.h
--- a/Engine/Plugins/FX/Niagara/Source/Niagara/Classes/NiagaraGPUSystemTick.h
+++ b/Engine/Plugins/FX/Niagara/Source/Niagara/Classes/NiagaraGPUSystemTick.h
@@ -15,6 +15,17 @@ struct FNiagaraComputeDataInterfaceInstanceData
 
 struct FNiagaraComputeInstanceData
 {
+	struct FPerStageInfo
+	{
+		FPerStageInfo() {}
+		FPerStageInfo(int32 InNumIterations, int32 InUserElementCount) : NumIterations(InNumIterations), UserElementCount(InUserElementCount) {}
+
+		int32	NumIterations = 0;
+		int32	UserElementCount = -1;
+
+		bool ShouldRunStage() const { return NumIterations > 0 && (UserElementCount != 0); }
+	};
+
 	FNiagaraComputeInstanceData()
 	{
 		bResetData = false;
@@ -28,7 +39,7 @@ struct FNiagaraComputeInstanceData
 	FNiagaraComputeExecutionContext* Context = nullptr;
 	TArray<FNiagaraDataInterfaceProxy*> DataInterfaceProxies;
 	TArray<FNiagaraDataInterfaceProxyRW*> IterationDataInterfaceProxies;
-	TArray<int32, TInlineAllocator<1>> NumIterationsPerStage;
+	TArray<FPerStageInfo, TInlineAllocator<1>> PerStageInfo;
 	uint32 ParticleCountFence = INDEX_NONE;
 	uint32 TotalDispatches = 0;
 	uint32 bResetData : 1;
diff --git a/Engine/Plugins/FX/Niagara/Source/Niagara/Private/NiagaraEmitterInstance.cpp b/Engine/Plugins/FX/Niagara/Source/Niagara/Private/NiagaraEmitterInstance.cpp
--- a/Engine/Plugins/FX/Niagara/Source/Niagara/Private/NiagaraEmitterInstance.cpp
+++ b/Engine/Plugins/FX/Niagara/Source/Niagara/Private/NiagaraEmitterInstance.cpp
@@ -489,6 +489,10 @@ void FNiagaraEmitterInstance::Init(int32 InEmitterIdx, FNiagaraSystemInstanceID
 				{
 					bAnyRendererBindingsAdded |= RendererBindings.AddParameter(FNiagaraVariable(FNiagaraTypeDefinition::GetBoolDef(), SimStageMetaData.EnabledBinding), false);
 				}
+				if (!SimStageMetaData.ElementCountBinding.IsNone())
+				{
+					bAnyRendererBindingsAdded |= RendererBindings.AddParameter(FNiagaraVariable(FNiagaraTypeDefinition::GetIntDef(), SimStageMetaData.ElementCountBinding), false);
+				}
 				if ( !SimStageMetaData.NumIterationsBinding.IsNone() )
 				{
 					bAnyRendererBindingsAdded |= RendererBindings.AddParameter(FNiagaraVariable(FNiagaraTypeDefinition::GetIntDef(), SimStageMetaData.NumIterationsBinding), false);
diff --git a/Engine/Plugins/FX/Niagara/Source/Niagara/Private/NiagaraGPUSystemTick.cpp b/Engine/Plugins/FX/Niagara/Source/Niagara/Private/NiagaraGPUSystemTick.cpp
--- a/Engine/Plugins/FX/Niagara/Source/Niagara/Private/NiagaraGPUSystemTick.cpp
+++ b/Engine/Plugins/FX/Niagara/Source/Niagara/Private/NiagaraGPUSystemTick.cpp
@@ -191,10 +191,11 @@ void FNiagaraGPUSystemTick::Init(FNiagaraSystemInstance* InSystemInstance)
 
 			// Gather number of iterations for each stage, and if the stage should run or not
 			InstanceData->bHasMultipleStages = false;
-			InstanceData->NumIterationsPerStage.Reserve(GPUContext->SimStageInfo.Num());
+			InstanceData->PerStageInfo.Reserve(GPUContext->SimStageInfo.Num());
 			for ( FSimulationStageMetaData& SimStageMetaData : GPUContext->SimStageInfo )
 			{
 				int32 NumIterations = SimStageMetaData.NumIterations;
+				int32 UserElementCount = -1;
 				if (SimStageMetaData.ShouldRunStage(InstanceData->bResetData))
 				{
 					InstanceData->bHasMultipleStages = true;
@@ -216,6 +217,14 @@ void FNiagaraGPUSystemTick::Init(FNiagaraSystemInstance* InSystemInstance)
 							NumIterations = StageEnabled.GetValue() ? NumIterations : 0;
 						}
 					}
+					if (!SimStageMetaData.ElementCountBinding.IsNone())
+					{
+						FNiagaraParameterStore& BoundParamStore = EmitterInstance->GetRendererBoundVariables();
+						if (const uint8* ParameterData = BoundParamStore.GetParameterData(FNiagaraVariable(FNiagaraTypeDefinition::GetIntDef(), SimStageMetaData.ElementCountBinding)))
+						{
+							UserElementCount = FMath::Max(*reinterpret_cast<const int32*>(ParameterData), 0);
+						}
+					}
 				}
 				else
 				{
@@ -223,7 +232,7 @@ void FNiagaraGPUSystemTick::Init(FNiagaraSystemInstance* InSystemInstance)
 				}
 
 				InstanceData->TotalDispatches += NumIterations;
-				InstanceData->NumIterationsPerStage.Add(NumIterations);
+				InstanceData->PerStageInfo.Emplace(NumIterations, UserElementCount);
 			}
 			TotalDispatches += InstanceData->TotalDispatches;
 		}
diff --git a/Engine/Plugins/FX/Niagara/Source/Niagara/Private/NiagaraGpuComputeDispatch.cpp b/Engine/Plugins/FX/Niagara/Source/Niagara/Private/NiagaraGpuComputeDispatch.cpp
--- a/Engine/Plugins/FX/Niagara/Source/Niagara/Private/NiagaraGpuComputeDispatch.cpp
+++ b/Engine/Plugins/FX/Niagara/Source/Niagara/Private/NiagaraGpuComputeDispatch.cpp
@@ -769,6 +769,10 @@ void FNiagaraGpuComputeDispatch::DumpDebugFrame()
 				Builder.Appendf(TEXT("Source(%p 0x%08x %d) "), SimStageData.Source, SimStageData.SourceCountOffset, SimStageData.SourceNumInstances);
 				Builder.Appendf(TEXT("Destination(%p 0x%08x %d) "), SimStageData.Destination, SimStageData.DestinationCountOffset, SimStageData.DestinationNumInstances);
 				Builder.Appendf(TEXT("Iteration(%d | %s) "), SimStageData.IterationIndex, SimStageData.AlternateIterationSource ? *SimStageData.AlternateIterationSource->SourceDIName.ToString() : TEXT("Particles"));
+				if (SimStageData.UserElementCount != -1)
+				{
+					Builder.Appendf(TEXT("UserElementCount(%d) "), SimStageData.UserElementCount);
+				}
 				UE_LOG(LogNiagara, Warning, TEXT("%s"), Builder.ToString());
 			}
 
@@ -1002,13 +1006,13 @@ void FNiagaraGpuComputeDispatch::PrepareTicksForProxy(FRHICommandListImmediate&
 			for (int32 SimStageIndex=0; SimStageIndex < ComputeContext->SimStageInfo.Num(); ++SimStageIndex)
 			{
 				const FSimulationStageMetaData& SimStageMetaData = ComputeContext->SimStageInfo[SimStageIndex];
-				if (InstanceData.NumIterationsPerStage[SimStageIndex] == 0)
+				if (InstanceData.PerStageInfo[SimStageIndex].ShouldRunStage() == false)
 				{
 					continue;
 				}
 
 				FNiagaraDataInterfaceProxyRW* IterationInterface = InstanceData.FindIterationInterface(SimStageIndex);
-				for ( int32 IterationIndex=0; IterationIndex < InstanceData.NumIterationsPerStage[SimStageIndex]; ++IterationIndex )
+				for ( int32 IterationIndex=0; IterationIndex < InstanceData.PerStageInfo[SimStageIndex].NumIterations; ++IterationIndex )
 				{
 					// Build SimStage data
 					FNiagaraGpuDispatchGroup& DispatchGroup = GpuDispatchList.DispatchGroups[iInstanceCurrDispatchGroup++];
@@ -1017,6 +1021,7 @@ void FNiagaraGpuComputeDispatch::PrepareTicksForProxy(FRHICommandListImmediate&
 					SimStageData.bFirstStage = bFirstStage;
 					SimStageData.StageIndex = SimStageIndex;
 					SimStageData.IterationIndex = IterationIndex;
+					SimStageData.UserElementCount = InstanceData.PerStageInfo[SimStageIndex].UserElementCount;
 					SimStageData.StageMetaData = &SimStageMetaData;
 					SimStageData.AlternateIterationSource = IterationInterface;
 
@@ -1468,6 +1473,12 @@ void FNiagaraGpuComputeDispatch::DispatchStage(FRHICommandList& RHICmdList, FRHI
 		DispatchNumThreads = FNiagaraShader::GetDefaultThreadGroupSize(ENiagaraGpuDispatchType::OneD);
 	}
 
+	// User override element count
+	if (SimStageData.UserElementCount != -1)
+	{
+		DispatchCount = FIntVector(SimStageData.UserElementCount, 1, 1);
+	}
+
 	const int32 TotalDispatchCount = DispatchCount.X * DispatchCount.Y * DispatchCount.Z;
 	if (TotalDispatchCount == 0)
 	{
@@ -1549,7 +1560,7 @@ void FNiagaraGpuComputeDispatch::DispatchStage(FRHICommandList& RHICmdList, FRHI
 			SimulationStageIterationInfo.Y = IterationInstanceCountOffset == INDEX_NONE ? TotalDispatchCount : 0;
 		}
 
-		const int32 NumIterations = InstanceData.NumIterationsPerStage[SimStageData.StageIndex];
+		const int32 NumIterations = InstanceData.PerStageInfo[SimStageData.StageIndex].NumIterations;
 		const int32 IterationIndex = SimStageData.IterationIndex;
 		SimulationStageIterationInfo.Z = IterationIndex;
 		SimulationStageIterationInfo.W = NumIterations;
diff --git a/Engine/Plugins/FX/Niagara/Source/Niagara/Private/NiagaraScript.cpp b/Engine/Plugins/FX/Niagara/Source/Niagara/Private/NiagaraScript.cpp
--- a/Engine/Plugins/FX/Niagara/Source/Niagara/Private/NiagaraScript.cpp
+++ b/Engine/Plugins/FX/Niagara/Source/Niagara/Private/NiagaraScript.cpp
@@ -924,6 +924,7 @@ void UNiagaraScript::ComputeVMCompilationId(FNiagaraVMExecutableDataId& Id, FGui
 					if (const UNiagaraSimulationStageGeneric* SimStageGeneric = Cast<const UNiagaraSimulationStageGeneric>(SimStageBase))
 					{
 						AddAttributeToPreserve(SimStageGeneric->EnabledBinding.GetParamMapBindableVariable().GetName());
+						AddAttributeToPreserve(SimStageGeneric->ElementCountBinding.GetParamMapBindableVariable().GetName());
 						AddAttributeToPreserve(SimStageGeneric->NumIterationsBinding.GetParamMapBindableVariable().GetName());
 					}
 				}
@@ -3144,6 +3145,12 @@ void UNiagaraScript::SyncAliases(const FNiagaraAliasContext& ResolveAliasesConte
 			SimStageMetaData.EnabledBinding = FNiagaraUtilities::ResolveAliases(Var, ResolveAliasesContext).GetName();
 		}
 
+		if (!SimStageMetaData.ElementCountBinding.IsNone())
+		{
+			FNiagaraVariable Var(FNiagaraTypeDefinition::GetIntDef(), SimStageMetaData.ElementCountBinding);
+			SimStageMetaData.ElementCountBinding = FNiagaraUtilities::ResolveAliases(Var, ResolveAliasesContext).GetName();
+		}
+
 		if (!SimStageMetaData.NumIterationsBinding.IsNone())
 		{
 			FNiagaraVariable Var(FNiagaraTypeDefinition::GetIntDef(), SimStageMetaData.NumIterationsBinding);
diff --git a/Engine/Plugins/FX/Niagara/Source/Niagara/Private/NiagaraSimulationStageBase.cpp b/Engine/Plugins/FX/Niagara/Source/Niagara/Private/NiagaraSimulationStageBase.cpp
--- a/Engine/Plugins/FX/Niagara/Source/Niagara/Private/NiagaraSimulationStageBase.cpp
+++ b/Engine/Plugins/FX/Niagara/Source/Niagara/Private/NiagaraSimulationStageBase.cpp
@@ -64,6 +64,7 @@ bool UNiagaraSimulationStageGeneric::AppendCompileHash(FNiagaraCompileHashVisito
 	Super::AppendCompileHash(InVisitor);
 
 	InVisitor->UpdateString(TEXT("EnabledBinding"), EnabledBinding.GetDataSetBindableVariable().GetName().ToString());
+	InVisitor->UpdateString(TEXT("ElementCountBinding"), ElementCountBinding.GetDataSetBindableVariable().GetName().ToString());
 	InVisitor->UpdatePOD(TEXT("Iterations"), Iterations);
 	InVisitor->UpdateString(TEXT("NumIterationsBinding"), NumIterationsBinding.GetDataSetBindableVariable().GetName().ToString());
 	InVisitor->UpdatePOD(TEXT("IterationSource"), (int32)IterationSource);
@@ -93,6 +94,12 @@ void UNiagaraSimulationStageGeneric::PostInitProperties()
 			ENiagaraRendererSourceDataMode::Emitter
 		);
 
+		ElementCountBinding.Setup(
+			FNiagaraVariableBase(FNiagaraTypeDefinition::GetIntDef(), NAME_None),
+			FNiagaraVariableBase(FNiagaraTypeDefinition::GetIntDef(), NAME_None),
+			ENiagaraRendererSourceDataMode::Emitter
+		);
+
 		NumIterationsBinding.Setup(
 			FNiagaraVariableBase(FNiagaraTypeDefinition::GetIntDef(), NAME_None),
 			FNiagaraVariableBase(FNiagaraTypeDefinition::GetIntDef(), NAME_None),
@@ -129,6 +136,10 @@ void UNiagaraSimulationStageGeneric::PostEditChangeProperty(struct FPropertyChan
 	{
 		bNeedsRecompile = true;
 	}
+	else if (PropertyName == GET_MEMBER_NAME_CHECKED(UNiagaraSimulationStageGeneric, ElementCountBinding))
+	{
+		bNeedsRecompile = true;
+	}
 	else if (PropertyName == GET_MEMBER_NAME_CHECKED(UNiagaraSimulationStageGeneric, Iterations))
 	{
 		bNeedsRecompile = true;
diff --git a/Engine/Plugins/FX/Niagara/Source/Niagara/Public/NiagaraSimStageData.h b/Engine/Plugins/FX/Niagara/Source/Niagara/Public/NiagaraSimStageData.h
--- a/Engine/Plugins/FX/Niagara/Source/Niagara/Public/NiagaraSimStageData.h
+++ b/Engine/Plugins/FX/Niagara/Source/Niagara/Public/NiagaraSimStageData.h
@@ -20,6 +20,7 @@ struct FNiagaraSimStageData
 
 	uint32 StageIndex = INDEX_NONE;
 	uint32 IterationIndex = 0;
+	int32 UserElementCount = -1;
 
 	class FNiagaraDataBuffer* Source = nullptr;
 	uint32 SourceCountOffset = INDEX_NONE;
diff --git a/Engine/Plugins/FX/Niagara/Source/Niagara/Public/NiagaraSimulationStageBase.h b/Engine/Plugins/FX/Niagara/Source/Niagara/Public/NiagaraSimulationStageBase.h
--- a/Engine/Plugins/FX/Niagara/Source/Niagara/Public/NiagaraSimulationStageBase.h
+++ b/Engine/Plugins/FX/Niagara/Source/Niagara/Public/NiagaraSimulationStageBase.h
@@ -52,28 +52,52 @@ class NIAGARA_API UNiagaraSimulationStageGeneric : public UNiagaraSimulationStag
 	GENERATED_BODY()
 
 public:
-	/** Binding to a bool parameter which dynamically controls if the simulation stage is enabled or not. */
+	/** Optional bool binding allowing scripts to control if the simulation stage is enabled or not. */
 	UPROPERTY(EditAnywhere, Category = "Simulation Stage")
 	FNiagaraVariableAttributeBinding EnabledBinding;
 
-	/** Determine which elements this script is iterating over. You are not allowed to */
+	/**
+	Optional integer binding to override the number of elements the stage will execute on per dispatch.
+	For example, if you want to iterate over a custom source such as triangles on a mesh you can
+	set an int to the triangle count in an emitter script and bind that as the element count.
+	*/
+	UPROPERTY(EditAnywhere, Category = "Simulation Stage")
+	FNiagaraVariableAttributeBinding ElementCountBinding;
+	
+	/**
+	Select what we should be iterating over, particles or data interfaces.
+	The source provides things such as element count (when not overriden) and stack context variables (i.e. attributes on grids)
+	*/
 	UPROPERTY(EditAnywhere, Category = "Simulation Stage")
 	ENiagaraIterationSource IterationSource;
 
-	UPROPERTY(EditAnywhere, Category = "Simulation Stage", meta = (NoSpinbox = "true", ClampMin = 1, DisplayName = "Num Iterations", Tooltip = "The number of times we run this simulation stage before moving to the next stage."))
+	/**
+	Number of times (or iterations) the simulation stage will execute in a row.
+	For example, setting this to 10 will mean this simulation stage runs 10 times in a row before the next stage.
+	*/
+	UPROPERTY(EditAnywhere, Category = "Simulation Stage", meta = (NoSpinbox = "true", ClampMin = 1, DisplayName = "Num Iterations"))
 	int32 Iterations = 1;
 
-	/** Binding to an int parameter which dynamically controls the number of times the simulation stage runs. */
+	/**
+	Optional integer binding allowing scripts to control the number of iterations.
+	*/
 	UPROPERTY(EditAnywhere, Category = "Simulation Stage")
 	FNiagaraVariableAttributeBinding NumIterationsBinding;
 
 	UPROPERTY()
 	uint32 bSpawnOnly_DEPRECATED : 1;
 
-	UPROPERTY(EditAnywhere, Category = "Simulation Stage", meta = (Tooltip = "Controls when the simulation stage should execute, only valid for data interface iteration stages", EditCondition = "IterationSource == ENiagaraIterationSource::DataInterface"))
+	/**
+	Controls when the simulation stage should execute, only valid for data interface iteration stages
+	*/
+	UPROPERTY(EditAnywhere, Category = "Simulation Stage", meta = (EditCondition = "IterationSource == ENiagaraIterationSource::DataInterface"))
 	ENiagaraSimStageExecuteBehavior ExecuteBehavior = ENiagaraSimStageExecuteBehavior::Always;
 
-	UPROPERTY(EditAnywhere, AdvancedDisplay, Category = "Simulation Stage", meta = (Tooltip = "Disables the ability to read / write from the same particle buffer, i.e. only update position and no other attributes.  By default this should not be changed and is a debugging tool.", EditCondition = "IterationSource == ENiagaraIterationSource::Particles"))
+	/**
+	Disables the ability to read / write from the same particle buffer, i.e. only update position and no other attributes.
+	By default this should not be changed, but can be used to debug issues.
+	*/
+	UPROPERTY(EditAnywhere, AdvancedDisplay, Category = "Simulation Stage", meta = (EditCondition = "IterationSource == ENiagaraIterationSource::Particles"))
 	uint32 bDisablePartialParticleUpdate : 1;
 
 	/** Source data interface to use for the simulation stage. The data interface needs to be a subclass of UNiagaraDataInterfaceRWBase, for example the Grid2D and Grid3D data interfaces. */
diff --git a/Engine/Plugins/FX/Niagara/Source/NiagaraEditor/Private/Customizations/NiagaraSimulationStageCustomization.cpp b/Engine/Plugins/FX/Niagara/Source/NiagaraEditor/Private/Customizations/NiagaraSimulationStageCustomization.cpp
--- a/Engine/Plugins/FX/Niagara/Source/NiagaraEditor/Private/Customizations/NiagaraSimulationStageCustomization.cpp
+++ b/Engine/Plugins/FX/Niagara/Source/NiagaraEditor/Private/Customizations/NiagaraSimulationStageCustomization.cpp
@@ -66,6 +66,7 @@ void FNiagaraSimulationStageGenericCustomization::CustomizeDetails(class IDetail
 
 	// Hide all
 	DetailBuilder.HideProperty(GET_MEMBER_NAME_CHECKED(UNiagaraSimulationStageGeneric, EnabledBinding));
+	DetailBuilder.HideProperty(GET_MEMBER_NAME_CHECKED(UNiagaraSimulationStageGeneric, ElementCountBinding));
 	DetailBuilder.HideProperty(GET_MEMBER_NAME_CHECKED(UNiagaraSimulationStageGeneric, IterationSource));
 	DetailBuilder.HideProperty(GET_MEMBER_NAME_CHECKED(UNiagaraSimulationStageGeneric, Iterations));
 	DetailBuilder.HideProperty(GET_MEMBER_NAME_CHECKED(UNiagaraSimulationStageGeneric, NumIterationsBinding));
@@ -81,6 +82,7 @@ void FNiagaraSimulationStageGenericCustomization::CustomizeDetails(class IDetail
 
 	// Show properties in the order we want them to appear
 	SimStageCategory.AddProperty(DetailBuilder.GetProperty(GET_MEMBER_NAME_CHECKED(UNiagaraSimulationStageGeneric, EnabledBinding)));
+	SimStageCategory.AddProperty(DetailBuilder.GetProperty(GET_MEMBER_NAME_CHECKED(UNiagaraSimulationStageGeneric, ElementCountBinding)));
 	SimStageCategory.AddProperty(DetailBuilder.GetProperty(GET_MEMBER_NAME_CHECKED(UNiagaraSimulationStageGeneric, Iterations)));
 	SimStageCategory.AddProperty(DetailBuilder.GetProperty(GET_MEMBER_NAME_CHECKED(UNiagaraSimulationStageGeneric, NumIterationsBinding)));
 	SimStageCategory.AddProperty(DetailBuilder.GetProperty(GET_MEMBER_NAME_CHECKED(UNiagaraSimulationStageGeneric, IterationSource)));
diff --git a/Engine/Plugins/FX/Niagara/Source/NiagaraEditor/Private/NiagaraCompiler.cpp b/Engine/Plugins/FX/Niagara/Source/NiagaraEditor/Private/NiagaraCompiler.cpp
--- a/Engine/Plugins/FX/Niagara/Source/NiagaraEditor/Private/NiagaraCompiler.cpp
+++ b/Engine/Plugins/FX/Niagara/Source/NiagaraEditor/Private/NiagaraCompiler.cpp
@@ -688,6 +688,7 @@ void FNiagaraCompileRequestData::FinishPrecompile(const TArray<FNiagaraVariable>
 					SimStageData.StageGuid				= GenericStage->Script->GetUsageId();
 					SimStageData.StageName				= GenericStage->SimulationStageName;
 					SimStageData.EnabledBinding			= GenericStage->EnabledBinding.GetName();
+					SimStageData.ElementCountBinding	= GenericStage->ElementCountBinding.GetName();
 					SimStageData.NumIterations			= GenericStage->Iterations;
 					SimStageData.NumIterationsBinding	= GenericStage->NumIterationsBinding.GetName();
 					SimStageData.IterationSource		= GenericStage->IterationSource == ENiagaraIterationSource::DataInterface ? GenericStage->DataInterface.BoundVariable.GetName() : FName();
diff --git a/Engine/Plugins/FX/Niagara/Source/NiagaraEditor/Private/NiagaraHlslTranslator.cpp b/Engine/Plugins/FX/Niagara/Source/NiagaraEditor/Private/NiagaraHlslTranslator.cpp
--- a/Engine/Plugins/FX/Niagara/Source/NiagaraEditor/Private/NiagaraHlslTranslator.cpp
+++ b/Engine/Plugins/FX/Niagara/Source/NiagaraEditor/Private/NiagaraHlslTranslator.cpp
@@ -1250,6 +1250,7 @@ const FNiagaraTranslateResults &FHlslNiagaraTranslator::Translate(const FNiagara
 					TranslationStages[Index].bCopyPreviousParams = false;
 					TranslationStages[Index].SimulationStageIndex = DestSimStageIndex;					
 					TranslationStages[Index].EnabledBinding = CompileSimStageData.EnabledBinding;
+					TranslationStages[Index].ElementCountBinding = CompileSimStageData.ElementCountBinding;
 					TranslationStages[Index].NumIterations = CompileSimStageData.NumIterations;
 					TranslationStages[Index].ExecuteBehavior = CompileSimStageData.ExecuteBehavior;
 					TranslationStages[Index].bPartialParticleUpdate = CompileSimStageData.PartialParticleUpdate;
@@ -1319,6 +1320,7 @@ const FNiagaraTranslateResults &FHlslNiagaraTranslator::Translate(const FNiagara
 					FSimulationStageMetaData& SimulationStageMetaData = CompilationOutput.ScriptData.SimulationStageMetaData.AddDefaulted_GetRef();
 					SimulationStageMetaData.SimulationStageName = InCompileData->CompileSimStageData[SourceSimStageIndex].StageName;
 					SimulationStageMetaData.EnabledBinding = InCompileData->CompileSimStageData[SourceSimStageIndex].EnabledBinding;
+					SimulationStageMetaData.ElementCountBinding = InCompileData->CompileSimStageData[SourceSimStageIndex].ElementCountBinding;
 					SimulationStageMetaData.ExecuteBehavior = TranslationStages[Index].ExecuteBehavior;
 					SimulationStageMetaData.IterationSource = InCompileData->CompileSimStageData[SourceSimStageIndex].IterationSource;
 					SimulationStageMetaData.NumIterationsBinding = InCompileData->CompileSimStageData[SourceSimStageIndex].NumIterationsBinding;
diff --git a/Engine/Plugins/FX/Niagara/Source/NiagaraEditor/Private/NiagaraHlslTranslator.h b/Engine/Plugins/FX/Niagara/Source/NiagaraEditor/Private/NiagaraHlslTranslator.h
--- a/Engine/Plugins/FX/Niagara/Source/NiagaraEditor/Private/NiagaraHlslTranslator.h
+++ b/Engine/Plugins/FX/Niagara/Source/NiagaraEditor/Private/NiagaraHlslTranslator.h
@@ -92,6 +92,7 @@ public:
 		FGuid							StageGuid;
 		FName							StageName;
 		FName							EnabledBinding;
+		FName							ElementCountBinding;
 		uint32							NumIterations = 1;
 		FName							NumIterationsBinding;
 		FName							IterationSource;
@@ -372,6 +373,7 @@ public:
 	FName IterationSource;
 	int32 SimulationStageIndex = -1;
 	FName EnabledBinding;
+	FName ElementCountBinding;
 	int32 NumIterations = 1;
 	FName NumIterationsBinding;
 	ENiagaraSimStageExecuteBehavior ExecuteBehavior = ENiagaraSimStageExecuteBehavior::Always;
diff --git a/Engine/Plugins/FX/Niagara/Source/NiagaraShader/Public/NiagaraScriptBase.h b/Engine/Plugins/FX/Niagara/Source/NiagaraShader/Public/NiagaraScriptBase.h
--- a/Engine/Plugins/FX/Niagara/Source/NiagaraShader/Public/NiagaraScriptBase.h
+++ b/Engine/Plugins/FX/Niagara/Source/NiagaraShader/Public/NiagaraScriptBase.h
@@ -45,6 +45,10 @@ public:
 	UPROPERTY()
 	FName EnabledBinding;
 
+	/** Optional binding to manually specify the element count. */
+	UPROPERTY()
+	FName ElementCountBinding;
+
 	/** The Data Interface that we iterate over for this stage. If None, then use particles.*/
 	UPROPERTY()
 	FName IterationSource;
