- Fix for DISkelMesh_GetSkinnedBoneCommon to take into account both bones & sockets

#rnx

diff --git a/Engine/Plugins/FX/Niagara/Shaders/Private/NiagaraDataInterfaceSkeletalMesh.ush b/Engine/Plugins/FX/Niagara/Shaders/Private/NiagaraDataInterfaceSkeletalMesh.ush
--- a/Engine/Plugins/FX/Niagara/Shaders/Private/NiagaraDataInterfaceSkeletalMesh.ush
+++ b/Engine/Plugins/FX/Niagara/Shaders/Private/NiagaraDataInterfaceSkeletalMesh.ush
@@ -815,7 +815,7 @@ float4 DISKelMesh_QuatSlerp(float4 Quat1, float4 Quat2, float Slerp)
 void DISkelMesh_GetSkinnedBoneCommon(in FDISkelMeshContext DIContext, in int Bone, in float Interp, in bool bWorldSpace, out float3 OutPosition, out float4 OutRotation, out float3 OutVelocity)
 {
 	// Early out for bad data
-	if ( DIContext.NumBones == 0 )
+	if ( DIContext.NumBones + DIContext.NumFilteredSockets == 0 )
 	{
 		OutPosition = float3(0.0f, 0.0f, 0.0f);
 		OutRotation = float4(0.0f, 0.0f, 0.0f, 1.0f);
@@ -823,7 +823,7 @@ void DISkelMesh_GetSkinnedBoneCommon(in FDISkelMeshContext DIContext, in int Bon
 		return;
 	}
 
-	Bone = clamp(Bone, 0, DIContext.NumBones - 1);
+	Bone = clamp(Bone, 0, DIContext.NumBones + DIContext.NumFilteredSockets - 1);
 
     float3 PrevPosition = DIContext.MeshPrevSamplingBonesBuffer[Bone * 3].xyz;
     float4 PrevRotation = DIContext.MeshPrevSamplingBonesBuffer[Bone * 3 + 1];
