- Track the effect type used by scalability
- This fixes a crash when a BP component's asset is overridden on a placed instance

#rnx

diff --git a/Engine/Plugins/FX/Niagara/Source/Niagara/Private/NiagaraComponent.cpp b/Engine/Plugins/FX/Niagara/Source/Niagara/Private/NiagaraComponent.cpp
--- a/Engine/Plugins/FX/Niagara/Source/Niagara/Private/NiagaraComponent.cpp
+++ b/Engine/Plugins/FX/Niagara/Source/Niagara/Private/NiagaraComponent.cpp
@@ -1322,7 +1322,8 @@ void UNiagaraComponent::RegisterWithScalabilityManager()
 			{
 				if (FNiagaraWorldManager* WorldMan = FNiagaraWorldManager::Get(GetWorld()))
 				{
-					WorldMan->RegisterWithScalabilityManager(this);
+					ScalabilityEffectType = EffectType;
+					WorldMan->RegisterWithScalabilityManager(this, ScalabilityEffectType);
 				}
 			}
 		}
@@ -1335,10 +1336,11 @@ void UNiagaraComponent::UnregisterWithScalabilityManager()
 	{
 		if (FNiagaraWorldManager* WorldMan = FNiagaraWorldManager::Get(GetWorld()))
 		{
-			WorldMan->UnregisterWithScalabilityManager(this);
+			WorldMan->UnregisterWithScalabilityManager(this, ScalabilityEffectType);
 		}
 	}
 	bIsCulledByScalability = false;
+	ScalabilityEffectType = nullptr;
 	ScalabilityManagerHandle = INDEX_NONE;//Just to be sure our state is unregistered.
 }
 
diff --git a/Engine/Plugins/FX/Niagara/Source/Niagara/Private/NiagaraWorldManager.cpp b/Engine/Plugins/FX/Niagara/Source/Niagara/Private/NiagaraWorldManager.cpp
--- a/Engine/Plugins/FX/Niagara/Source/Niagara/Private/NiagaraWorldManager.cpp
+++ b/Engine/Plugins/FX/Niagara/Source/Niagara/Private/NiagaraWorldManager.cpp
@@ -1171,7 +1171,7 @@ void FNiagaraWorldManager::UpdateScalabilityManagers(float DeltaSeconds, bool bN
 	}
 }
 
-void FNiagaraWorldManager::RegisterWithScalabilityManager(UNiagaraComponent* Component)
+void FNiagaraWorldManager::RegisterWithScalabilityManager(UNiagaraComponent* Component, UNiagaraEffectType* EffectType)
 {
 	if (GetScalabilityCullingMode() == ENiagaraScalabilityCullingMode::Disabled)
 	{
@@ -1179,7 +1179,7 @@ void FNiagaraWorldManager::RegisterWithScalabilityManager(UNiagaraComponent* Com
 	}
 
 	check(Component);
-	if (UNiagaraEffectType* EffectType = Component->GetAsset()->GetEffectType())
+	if ( EffectType )
 	{
 		FNiagaraScalabilityManager* ScalabilityManager = ScalabilityManagers.Find(EffectType);
 
@@ -1193,10 +1193,10 @@ void FNiagaraWorldManager::RegisterWithScalabilityManager(UNiagaraComponent* Com
 	}
 }
 
-void FNiagaraWorldManager::UnregisterWithScalabilityManager(UNiagaraComponent* Component)
+void FNiagaraWorldManager::UnregisterWithScalabilityManager(UNiagaraComponent* Component, UNiagaraEffectType* EffectType)
 {
 	check(Component);
-	if (UNiagaraEffectType* EffectType = Component->GetAsset()->GetEffectType())
+	if ( EffectType )
 	{
 		//Possibly the manager has been GCd.
 		if (FNiagaraScalabilityManager* ScalabilityManager = ScalabilityManagers.Find(EffectType))
diff --git a/Engine/Plugins/FX/Niagara/Source/Niagara/Public/NiagaraComponent.h b/Engine/Plugins/FX/Niagara/Source/Niagara/Public/NiagaraComponent.h
--- a/Engine/Plugins/FX/Niagara/Source/Niagara/Public/NiagaraComponent.h
+++ b/Engine/Plugins/FX/Niagara/Source/Niagara/Public/NiagaraComponent.h
@@ -764,6 +764,7 @@ private:
 
 	FDelegateHandle AssetExposedParametersChangedHandle;
 
+	UNiagaraEffectType* ScalabilityEffectType = nullptr;
 	int32 ScalabilityManagerHandle;
 
 	float ForceUpdateTransformTime;
diff --git a/Engine/Plugins/FX/Niagara/Source/Niagara/Public/NiagaraWorldManager.h b/Engine/Plugins/FX/Niagara/Source/Niagara/Public/NiagaraWorldManager.h
--- a/Engine/Plugins/FX/Niagara/Source/Niagara/Public/NiagaraWorldManager.h
+++ b/Engine/Plugins/FX/Niagara/Source/Niagara/Public/NiagaraWorldManager.h
@@ -155,8 +155,8 @@ public:
 
 	//Various helper functions for scalability culling.
 	
-	void RegisterWithScalabilityManager(UNiagaraComponent* Component);
-	void UnregisterWithScalabilityManager(UNiagaraComponent* Component);
+	void RegisterWithScalabilityManager(UNiagaraComponent* Component, UNiagaraEffectType* EffectType);
+	void UnregisterWithScalabilityManager(UNiagaraComponent* Component, UNiagaraEffectType* EffectType);
 
 	/** Should we cull an instance of this system at the passed location before it's even been spawned? */
 	NIAGARA_API bool ShouldPreCull(UNiagaraSystem* System, UNiagaraComponent* Component);
