fix bug with GC'd actors still being queried for static mesh collisions

diff --git a/Engine/Plugins/FX/Niagara/Source/Niagara/Private/NiagaraDataInterfaceRigidMeshCollisionQuery.cpp b/Engine/Plugins/FX/Niagara/Source/Niagara/Private/NiagaraDataInterfaceRigidMeshCollisionQuery.cpp
--- a/Engine/Plugins/FX/Niagara/Source/Niagara/Private/NiagaraDataInterfaceRigidMeshCollisionQuery.cpp
+++ b/Engine/Plugins/FX/Niagara/Source/Niagara/Private/NiagaraDataInterfaceRigidMeshCollisionQuery.cpp
@@ -140,19 +140,18 @@ void FillCurrentTransforms(const FTransform& ElementTransform, uint32& ElementCo
 	++ElementCount;
 }
 
-void GetNumPrimitives(TArray<AActor*> Actors, uint32& NumBoxes, uint32& NumSpheres, uint32& NumCapsules)
+void GetNumPrimitives(TArray<TWeakObjectPtr<AActor>> Actors, uint32& NumBoxes, uint32& NumSpheres, uint32& NumCapsules)
 {
 	NumBoxes = 0;
 	NumSpheres = 0;
 	NumCapsules = 0;
 
 	for (int32 ActorIndex = 0; ActorIndex < Actors.Num(); ++ActorIndex)
-	{
-		AActor* Actor = Actors[ActorIndex];
+	{		
 		UStaticMeshComponent* StaticMeshComponent = nullptr;
-		if (Actor != nullptr)
+		if (Actors[ActorIndex].IsValid())
 		{
-			StaticMeshComponent = Cast<UStaticMeshComponent>(Actor->GetComponentByClass(UStaticMeshComponent::StaticClass()));
+			StaticMeshComponent = Cast<UStaticMeshComponent>(Actors[ActorIndex]->GetComponentByClass(UStaticMeshComponent::StaticClass()));
 		}
 		UBodySetup* BodySetup = StaticMeshComponent != nullptr ? StaticMeshComponent->GetBodySetup() : nullptr;
 		if (BodySetup != nullptr)
@@ -205,7 +204,7 @@ void CompactInternalArrays(FNDIRigidMeshCollisionArrays* OutAssetArrays)
 	}
 }
 
-void CreateInternalArrays(TArray<AActor*> Actors, FNDIRigidMeshCollisionArrays* OutAssetArrays, FVector LWCTile)
+void CreateInternalArrays(TArray<TWeakObjectPtr<AActor>> Actors, FNDIRigidMeshCollisionArrays* OutAssetArrays, FVector LWCTile)
 {
 	if (OutAssetArrays != nullptr)
 	{
@@ -233,9 +232,9 @@ void CreateInternalArrays(TArray<AActor*> Actors, FNDIRigidMeshCollisionArrays*
 
 			for (int32 ActorIndex = 0; ActorIndex < Actors.Num(); ++ActorIndex)
 			{				
-				AActor* Actor = Actors[ActorIndex];
+				TWeakObjectPtr<AActor> Actor = Actors[ActorIndex].Get();
 				
-				if (Actor != nullptr && Actor->GetComponentByClass(UStaticMeshComponent::StaticClass()) != nullptr)
+				if (Actor.IsValid() && Actor->GetComponentByClass(UStaticMeshComponent::StaticClass()) != nullptr)
 				{
 					UStaticMeshComponent* StaticMeshComponent = Cast< UStaticMeshComponent>(Actor->GetComponentByClass(UStaticMeshComponent::StaticClass()));
 					UBodySetup* BodySetup = StaticMeshComponent->GetBodySetup();
@@ -330,7 +329,7 @@ void CreateInternalArrays(TArray<AActor*> Actors, FNDIRigidMeshCollisionArrays*
 	}
 }
 
-void UpdateInternalArrays(const TArray<AActor*> &Actors, FNDIRigidMeshCollisionArrays* OutAssetArrays, FVector LWCTile)
+void UpdateInternalArrays(const TArray<TWeakObjectPtr<AActor>> &Actors, FNDIRigidMeshCollisionArrays* OutAssetArrays, FVector LWCTile)
 {
 	if (OutAssetArrays != nullptr && OutAssetArrays->ElementOffsets.NumElements < OutAssetArrays->MaxPrimitives)
 	{
@@ -356,9 +355,9 @@ void UpdateInternalArrays(const TArray<AActor*> &Actors, FNDIRigidMeshCollisionA
 
 		for (int32 ActorIndex = 0; ActorIndex < Actors.Num(); ++ActorIndex)
 		{
-			AActor* Actor = Actors[ActorIndex];
+			TWeakObjectPtr<AActor> Actor = Actors[ActorIndex].Get();
 			UStaticMeshComponent* StaticMeshComponent = nullptr;
-			if (Actor != nullptr)
+			if (Actor.IsValid())
 			{
 				StaticMeshComponent = Cast<UStaticMeshComponent>(Actor->GetComponentByClass(UStaticMeshComponent::StaticClass()));
 			}
@@ -473,7 +472,7 @@ void FNDIRigidMeshCollisionData::Init(UNiagaraDataInterfaceRigidMeshCollisionQue
 			if ((!Interface->OnlyUseMoveable || (Interface->OnlyUseMoveable && Actor->IsRootComponentMovable())) &&
 				(Interface->Tag == FString("") || (Interface->Tag != FString("") && Actor->Tags.Contains(FName(Interface->Tag)))))
 			{
-				Actors.Add(Actor);
+				Actors.Add(TWeakObjectPtr<AActor>(Actor));
 			}
 		}	
 
@@ -495,7 +494,7 @@ void FNDIRigidMeshCollisionData::Update(UNiagaraDataInterfaceRigidMeshCollisionQ
 {
 	if (Interface != nullptr && SystemInstance != nullptr)
 	{		
-		TArray<AActor*> ActorsTmp;		
+		TArray<TWeakObjectPtr<AActor> > ActorsTmp;		
 
 
 		// #todo(dmp): loop over actors and count number of potential colliders.  This might change frame to frame and we might need a
@@ -508,7 +507,7 @@ void FNDIRigidMeshCollisionData::Update(UNiagaraDataInterfaceRigidMeshCollisionQ
 			if ((!Interface->OnlyUseMoveable || (Interface->OnlyUseMoveable && Actor->IsRootComponentMovable())) &&
 				(Interface->Tag == FString("") || (Interface->Tag != FString("") && Actor->Tags.Contains(FName(Interface->Tag)))))
 			{
-				ActorsTmp.Add(Actor);
+				ActorsTmp.Add(TWeakObjectPtr<AActor>(Actor));
 			}
 		}
 		
@@ -521,7 +520,7 @@ void FNDIRigidMeshCollisionData::Update(UNiagaraDataInterfaceRigidMeshCollisionQ
 
 		TickingGroup = ComputeTickingGroup();
 
-		if (0 < Actors.Num() && Actors[0] != nullptr)
+		if (0 < Actors.Num() && Actors[0].IsValid())
 		{
 			UpdateInternalArrays(Actors, AssetArrays, FVector(SystemInstance->GetLWCTile()));
 		}
@@ -533,10 +532,11 @@ ETickingGroup FNDIRigidMeshCollisionData::ComputeTickingGroup()
 	TickingGroup = NiagaraFirstTickGroup;
 	for (int32 ComponentIndex = 0; ComponentIndex < Actors.Num(); ++ComponentIndex)
 	{
-		if (Actors[ComponentIndex] != nullptr)
+		TWeakObjectPtr<AActor> Actor = Actors[ComponentIndex].Get();
+		if (Actor.IsValid())
 		{			
 		
-			UStaticMeshComponent* Component = Cast<UStaticMeshComponent>(Actors[ComponentIndex]->GetComponentByClass(UStaticMeshComponent::StaticClass()));
+			UStaticMeshComponent* Component = Cast<UStaticMeshComponent>(Actor->GetComponentByClass(UStaticMeshComponent::StaticClass()));
 
 			if (Component == nullptr)
 			{
diff --git a/Engine/Plugins/FX/Niagara/Source/Niagara/Public/NiagaraDataInterfaceRigidMeshCollisionQuery.h b/Engine/Plugins/FX/Niagara/Source/Niagara/Public/NiagaraDataInterfaceRigidMeshCollisionQuery.h
--- a/Engine/Plugins/FX/Niagara/Source/Niagara/Public/NiagaraDataInterfaceRigidMeshCollisionQuery.h
+++ b/Engine/Plugins/FX/Niagara/Source/Niagara/Public/NiagaraDataInterfaceRigidMeshCollisionQuery.h
@@ -153,7 +153,7 @@ struct FNDIRigidMeshCollisionData
 	FNDIRigidMeshCollisionArrays *AssetArrays = nullptr;
 
 	/** Static Mesh Components **/
-	TArray<AActor*> Actors;
+	TArray<TWeakObjectPtr<AActor>> Actors;
 };
 
 /** Data Interface used to collide against static meshes - whether it is the mesh distance field or a physics asset's collision primitive */
