- Initial pass of baker UI revamp

diff --git a/Engine/Plugins/FX/Niagara/Shaders/Private/NiagaraBatchedElements.usf b/Engine/Plugins/FX/Niagara/Shaders/Private/NiagaraBatchedElements.usf
--- a/Engine/Plugins/FX/Niagara/Shaders/Private/NiagaraBatchedElements.usf
+++ b/Engine/Plugins/FX/Niagara/Shaders/Private/NiagaraBatchedElements.usf
@@ -61,6 +61,11 @@ void MainPS(
 
 #if NIAGARA_SIMPLE_PS
 
+Texture2D		InTexture;
+SamplerState	InTextureSampler;
+float4x4		InColorWeights;
+float			InGamma;
+
 void MainPS(
 	in float2 InTextureCoordinate : TEXCOORD0,
 	in float4 InColor : TEXCOORD1,
@@ -72,7 +77,20 @@ void MainPS(
 #endif
 	)
 {
-	OutColor = RETURN_COLOR(1);
+	float4 Sample = Texture2DSample(InTexture, InTextureSampler, InTextureCoordinate);
+	float4 Color;
+	Color.r = dot(Sample, InColorWeights[0]);
+	Color.g = dot(Sample, InColorWeights[1]);
+	Color.b = dot(Sample, InColorWeights[2]);
+	Color.a = dot(Sample, InColorWeights[3]);
+	Color *= InColor;
+
+	if( InGamma != 1.0 )
+	{
+		Color.rgb = ApplyGammaCorrection(saturate(Color.rgb), 2.2 / (1.0 / InGamma));
+	}
+
+	OutColor = RETURN_COLOR(Color);
 
 #if WRITE_TO_GBUFFER
 	// Set the G buffer bits that indicate that deferred lighting and image reflections are not enabled
diff --git a/Engine/Plugins/FX/Niagara/Source/Niagara/Classes/NiagaraBakerSettings.h b/Engine/Plugins/FX/Niagara/Source/Niagara/Classes/NiagaraBakerSettings.h
--- a/Engine/Plugins/FX/Niagara/Source/Niagara/Classes/NiagaraBakerSettings.h
+++ b/Engine/Plugins/FX/Niagara/Source/Niagara/Classes/NiagaraBakerSettings.h
@@ -20,7 +20,7 @@ enum class ENiagaraBakerViewMode
 	OrthoRight,
 	OrthoTop,
 	OrthoBottom,
-	Num
+	Num UMETA(Hidden)
 };
 
 USTRUCT()
@@ -86,14 +86,14 @@ public:
 	UNiagaraBakerSettings(const FObjectInitializer& Init);
 
 	/**
-	This is the start time of the simultion where we being the capture.
+	This is the start time of the simulation where we begin the capture.
 	I.e. 2.0 would mean the simulation warms up by 2 seconds before we begin capturing.
 	*/
-	UPROPERTY(EditAnywhere, Category="Timeline")
+	UPROPERTY(EditAnywhere, Category="Settings")
 	float StartSeconds = 0.0f;
 
 	/** Duration in seconds to take the capture over. */
-	UPROPERTY(EditAnywhere, Category = "Timeline")
+	UPROPERTY(EditAnywhere, Category = "Settings")
 	float DurationSeconds = 4.0f;
 
 	/**
@@ -101,54 +101,54 @@ public:
 	This is only used for the preview view and calculating the number of ticks to execute
 	as we capture the generated texture.
 	*/
-	UPROPERTY(EditAnywhere, Category = "Timeline", meta = (ClampMin=1, ClampMax=480))
+	UPROPERTY(EditAnywhere, Category = "Settings", AdvancedDisplay, meta = (ClampMin=1, ClampMax=480))
 	int FramesPerSecond = 60;
 
 	/** Should the preview playback as looping or not. */
-	UPROPERTY(EditAnywhere, Category = "Preview")
+	UPROPERTY(EditAnywhere, Category = "Settings")
 	uint8 bPreviewLooping : 1;
 
 	/** Number of frames in each dimension. */
-	UPROPERTY(EditAnywhere, Category = "Texture")
+	UPROPERTY(EditAnywhere, Category = "Settings")
 	FIntPoint FramesPerDimension = FIntPoint(8, 8);
 
 	/** List of output textures we will generated. */
-	UPROPERTY(EditAnywhere, Category = "Texture")
+	UPROPERTY(EditAnywhere, Category = "Settings")
 	TArray<FNiagaraBakerTextureSettings> OutputTextures;
 
 	/** Current active viewport we will render from. */
-	UPROPERTY(EditAnywhere, Category = "Camera")
+	UPROPERTY(EditAnywhere, Category = "Settings")
 	ENiagaraBakerViewMode CameraViewportMode = ENiagaraBakerViewMode::Perspective;
 
 	/** Per viewport camera position.. */
-	UPROPERTY(EditAnywhere, Category = "Camera")
+	UPROPERTY(EditAnywhere, Category = "Settings")
 	FVector CameraViewportLocation[(int)ENiagaraBakerViewMode::Num];
 
 	/** Per viewport camera rotation.. */
-	UPROPERTY(EditAnywhere, Category = "Camera")
+	UPROPERTY(EditAnywhere, Category = "Settings")
 	FRotator CameraViewportRotation[(int)ENiagaraBakerViewMode::Num];
 
 	/** Perspective camera orbit distance. */
-	UPROPERTY(EditAnywhere, Category = "Camera", meta = (EditCondition = "CameraViewportMode == ENiagaraBakerViewMode::Perspective", ClampMin = "0.01"))
+	UPROPERTY(EditAnywhere, Category = "Settings", meta=(EditConditionHides, EditCondition="CameraViewportMode == ENiagaraBakerViewMode::Perspective", ClampMin = "0.01"))
 	float CameraOrbitDistance = 200.f;
 
 	/** Camera FOV to use when in perspective mode. */
-	UPROPERTY(EditAnywhere, Category = "Camera", meta=(EditCondition="CameraViewportMode == ENiagaraBakerViewMode::Perspective", ClampMin="1.0", ClampMax="180.0"))
+	UPROPERTY(EditAnywhere, Category = "Settings", meta=(EditConditionHides, EditCondition="CameraViewportMode == ENiagaraBakerViewMode::Perspective", ClampMin="1.0", ClampMax="179.0"))
 	float CameraFOV = 90.0f;
 
 	/** Camera Orthographic width to use with in orthographic mode. */
-	UPROPERTY(EditAnywhere, Category = "Camera", meta = (EditCondition = "CameraViewportMode != ENiagaraBakerViewMode::Perspective", ClampMin="1.0"))
+	UPROPERTY(EditAnywhere, Category = "Settings", meta=(EditConditionHides, EditCondition = "CameraViewportMode != ENiagaraBakerViewMode::Perspective", ClampMin="1.0"))
 	float CameraOrthoWidth = 512.0f;
 
-	UPROPERTY(EditAnywhere, Category = "Camera", meta = (PinHiddenByDefault, InlineEditConditionToggle))
+	UPROPERTY(EditAnywhere, Category = "Settings", meta=(PinHiddenByDefault, InlineEditConditionToggle))
 	uint8 bUseCameraAspectRatio : 1;
 
 	/** Custom aspect ratio to use rather than using the width & height to automatically calculate. */
-	UPROPERTY(EditAnywhere, Category = "Camera", meta = (EditCondition = "bUseCameraAspectRatio", ClampMin="0.01"))
+	UPROPERTY(EditAnywhere, Category = "Settings", meta=(EditCondition = "bUseCameraAspectRatio", ClampMin="0.01"))
 	float CameraAspectRatio = 1.0f;
 
 	/** Should we render just the component or the whole scene. */
-	UPROPERTY(EditAnywhere, Category = "Environment")
+	UPROPERTY(EditAnywhere, Category = "Settings", AdvancedDisplay)
 	uint8 bRenderComponentOnly : 1;
 
 	///** Type of level setup to use. */
diff --git a/Engine/Plugins/FX/Niagara/Source/NiagaraEditor/Private/NiagaraEditorStyle.cpp b/Engine/Plugins/FX/Niagara/Source/NiagaraEditor/Private/NiagaraEditorStyle.cpp
--- a/Engine/Plugins/FX/Niagara/Source/NiagaraEditor/Private/NiagaraEditorStyle.cpp
+++ b/Engine/Plugins/FX/Niagara/Source/NiagaraEditor/Private/NiagaraEditorStyle.cpp
@@ -547,7 +547,7 @@ void FNiagaraEditorStyle::InitDebuggerStyle()
 
 void FNiagaraEditorStyle::InitBakerStyle()
 {
-	Set("NiagaraEditor.Baker", new IMAGE_BRUSH("Icons/Baker/BakerIcon", Icon40x40));
+	Set("NiagaraEditor.BakerIcon", new IMAGE_BRUSH("Icons/Baker/BakerIcon", Icon40x40));
 }
 
 void FNiagaraEditorStyle::InitCommonColors()
diff --git a/Engine/Plugins/FX/Niagara/Source/NiagaraEditor/Private/Toolkits/NiagaraSystemToolkit.cpp b/Engine/Plugins/FX/Niagara/Source/NiagaraEditor/Private/Toolkits/NiagaraSystemToolkit.cpp
--- a/Engine/Plugins/FX/Niagara/Source/NiagaraEditor/Private/Toolkits/NiagaraSystemToolkit.cpp
+++ b/Engine/Plugins/FX/Niagara/Source/NiagaraEditor/Private/Toolkits/NiagaraSystemToolkit.cpp
@@ -1276,9 +1276,4 @@ void FNiagaraSystemToolkit::OnViewModelRequestFocusTab(FName TabName)
 	GetTabManager()->TryInvokeTab(TabName);
 }
 
-void FNiagaraSystemToolkit::RenderBaker()
-{
-	BakerViewModel->RenderBaker();
-}
-
 #undef LOCTEXT_NAMESPACE
diff --git a/Engine/Plugins/FX/Niagara/Source/NiagaraEditor/Private/Toolkits/NiagaraSystemToolkit.h b/Engine/Plugins/FX/Niagara/Source/NiagaraEditor/Private/Toolkits/NiagaraSystemToolkit.h
--- a/Engine/Plugins/FX/Niagara/Source/NiagaraEditor/Private/Toolkits/NiagaraSystemToolkit.h
+++ b/Engine/Plugins/FX/Niagara/Source/NiagaraEditor/Private/Toolkits/NiagaraSystemToolkit.h
@@ -107,7 +107,6 @@ public:
 	void SetScriptScratchpad(const TSharedPtr<SWidget>& ScriptScratchpad);
 	
 	/** Mode exposed functions */
-	void RenderBaker();
 	TSharedRef<SWidget> GenerateCompileMenuContent();
 	TSharedRef<SWidget> GenerateBoundsMenuContent(TSharedRef<FUICommandList> InCommandList);
 	TSharedRef<SWidget> GenerateStatConfigMenuContent(TSharedRef<FUICommandList> InCommandList);
diff --git a/Engine/Plugins/FX/Niagara/Source/NiagaraEditor/Private/Toolkits/SystemToolkitModes/NiagaraSystemToolkitMode_Default.cpp b/Engine/Plugins/FX/Niagara/Source/NiagaraEditor/Private/Toolkits/SystemToolkitModes/NiagaraSystemToolkitMode_Default.cpp
--- a/Engine/Plugins/FX/Niagara/Source/NiagaraEditor/Private/Toolkits/SystemToolkitModes/NiagaraSystemToolkitMode_Default.cpp
+++ b/Engine/Plugins/FX/Niagara/Source/NiagaraEditor/Private/Toolkits/SystemToolkitModes/NiagaraSystemToolkitMode_Default.cpp
@@ -95,27 +95,6 @@ void FNiagaraSystemToolkitMode_Default::ExtendToolbar()
 			return MenuBuilder.MakeWidget();
 		}
 
-		static TSharedRef<SWidget> GenerateBakerMenu(FNiagaraSystemToolkit* Toolkit)
-		{
-			FMenuBuilder MenuBuilder(true, Toolkit->GetToolkitCommands());
-
-			MenuBuilder.AddMenuEntry(
-				LOCTEXT("BakerTab", "Open Baker Tab"),
-				LOCTEXT("BakerTabTooltip", "Opens the flip book tab."),
-				FSlateIcon(),
-				FUIAction(FExecuteAction::CreateLambda([TabManager=Toolkit->GetTabManager()]() { TabManager->TryInvokeTab(FNiagaraSystemToolkitModeBase::BakerTabID); }))
-			);
-
-			MenuBuilder.AddMenuEntry(
-				LOCTEXT("RunBaker", "Bake"),
-				LOCTEXT("RunBakerTooltip", "Runs the bake process."),
-				FSlateIcon(),
-				FUIAction(FExecuteAction::CreateSP(Toolkit, &FNiagaraSystemToolkit::RenderBaker))
-			);
-
-			return MenuBuilder.MakeWidget();
-		}
-
 		static TSharedRef<SWidget> FillDebugOptionsMenu(FNiagaraSystemToolkit* Toolkit)
 		{
 			FMenuBuilder MenuBuilder(true, Toolkit->GetToolkitCommands());
@@ -228,12 +207,37 @@ void FNiagaraSystemToolkitMode_Default::ExtendToolbar()
 
 			ToolbarBuilder.BeginSection("Baker");
 			{
-				ToolbarBuilder.AddComboButton(
-					FUIAction(),
-					FOnGetContent::CreateStatic(Local::GenerateBakerMenu, Toolkit),
-					LOCTEXT("Baker", "Baker"),
-					LOCTEXT("BakerTooltip", "Options for Baker rendering."),
-					FSlateIcon(FNiagaraEditorStyle::Get().GetStyleSetName(), "NiagaraEditor.Baker")
+				FUIAction BakerToggleAction(
+					FExecuteAction::CreateLambda(
+						[TabManager = Toolkit->GetTabManager()]()
+						{
+							TSharedPtr<SDockTab> ExistingTab = TabManager->FindExistingLiveTab(FNiagaraSystemToolkitModeBase::BakerTabID);
+							if (ExistingTab)
+							{
+								ExistingTab->RequestCloseTab();
+							}
+							else
+							{
+								TabManager->TryInvokeTab(FNiagaraSystemToolkitModeBase::BakerTabID);
+							}
+						}
+					),
+					nullptr,
+					FIsActionChecked::CreateLambda(
+						[TabManager = Toolkit->GetTabManager()]()
+						{
+							return TabManager->FindExistingLiveTab(FNiagaraSystemToolkitModeBase::BakerTabID).IsValid();
+						}
+					)
+				);
+
+				ToolbarBuilder.AddToolBarButton(
+					BakerToggleAction,
+					NAME_None,
+					LOCTEXT("BakerLabel", "Baker"),
+					LOCTEXT("BakerTooltip", "Toggles the baker tab."),
+					FSlateIcon(FNiagaraEditorStyle::Get().GetStyleSetName(), "NiagaraEditor.BakerIcon"),
+					EUserInterfaceActionType::ToggleButton
 				);
 			}
 			ToolbarBuilder.EndSection();
diff --git a/Engine/Plugins/FX/Niagara/Source/NiagaraEditor/Private/ViewModels/NiagaraBakerViewModel.cpp b/Engine/Plugins/FX/Niagara/Source/NiagaraEditor/Private/ViewModels/NiagaraBakerViewModel.cpp
--- a/Engine/Plugins/FX/Niagara/Source/NiagaraEditor/Private/ViewModels/NiagaraBakerViewModel.cpp
+++ b/Engine/Plugins/FX/Niagara/Source/NiagaraEditor/Private/ViewModels/NiagaraBakerViewModel.cpp
@@ -16,11 +16,12 @@
 
 #include "Misc/ScopedSlowTask.h"
 #include "Modules/ModuleManager.h"
-#include "EngineModule.h"
+#include "AdvancedPreviewScene.h"
 #include "AssetToolsModule.h"
+#include "EngineModule.h"
 #include "LegacyScreenPercentageDriver.h"
 #include "PackageTools.h"
-#include "AdvancedPreviewScene.h"
+#include "ScopedTransaction.h"
 
 #include "Factories/Texture2dFactoryNew.h"
 
@@ -119,21 +120,266 @@ const UNiagaraBakerSettings* FNiagaraBakerViewModel::GetBakerGeneratedSettings()
 	return nullptr;
 }
 
-const FNiagaraBakerTextureSettings* FNiagaraBakerViewModel::GetPreviewTextureSettings() const
+TSharedPtr<class SWidget> FNiagaraBakerViewModel::GetWidget()
+{
+	return Widget;
+}
+
+void FNiagaraBakerViewModel::SetCameraViewMode(ENiagaraBakerViewMode ViewMode)
+{
+	if (UNiagaraBakerSettings* BakerSettings = GetBakerSettings())
+	{
+		const FScopedTransaction Transaction(LOCTEXT("SetCameraViewMode", "Set Camera View Mode"));
+		BakerSettings->CameraViewportMode = ViewMode;
+		BakerSettings->Modify();
+	}
+}
+
+bool FNiagaraBakerViewModel::IsCameraViewMode(ENiagaraBakerViewMode ViewMode)
+{
+	UNiagaraBakerSettings* BakerSettings = GetBakerSettings();
+	return BakerSettings ? BakerSettings->CameraViewportMode == ViewMode : false;
+}
+
+FText FNiagaraBakerViewModel::GetCurrentCameraModeText() const
+{
+	UNiagaraBakerSettings* BakerSettings = GetBakerSettings();
+	return BakerSettings ? GetCameraModeText(BakerSettings->CameraViewportMode) : FText::GetEmpty();
+}
+
+FName FNiagaraBakerViewModel::GetCurrentCameraModeIconName() const
+{
+	UNiagaraBakerSettings* BakerSettings = GetBakerSettings();
+	return GetCameraModeIconName(BakerSettings ? BakerSettings->CameraViewportMode : ENiagaraBakerViewMode::Perspective);
+}
+
+FSlateIcon FNiagaraBakerViewModel::GetCurrentCameraModeIcon() const
+{
+	UNiagaraBakerSettings* BakerSettings = GetBakerSettings();
+	return GetCameraModeIcon(BakerSettings ? BakerSettings->CameraViewportMode : ENiagaraBakerViewMode::Perspective);
+}
+
+FText FNiagaraBakerViewModel::GetCameraModeText(ENiagaraBakerViewMode Mode)
+{
+	switch (Mode)
+	{
+		case ENiagaraBakerViewMode::Perspective:	return LOCTEXT("Perspective", "Perspective");
+		case ENiagaraBakerViewMode::OrthoFront:		return LOCTEXT("OrthoFront", "Front");
+		case ENiagaraBakerViewMode::OrthoBack:		return LOCTEXT("OrthoBack", "Back");
+		case ENiagaraBakerViewMode::OrthoLeft:		return LOCTEXT("OrthoLeft", "Left");
+		case ENiagaraBakerViewMode::OrthoRight:		return LOCTEXT("OrthoRight", "Right");
+		case ENiagaraBakerViewMode::OrthoTop:		return LOCTEXT("OrthoTop", "Top");
+		case ENiagaraBakerViewMode::OrthoBottom:	return LOCTEXT("OrthoBottom", "Bottom");
+		default:									return LOCTEXT("Error", "Error");
+	}
+}
+
+FName FNiagaraBakerViewModel::GetCameraModeIconName(ENiagaraBakerViewMode Mode)
+{
+	static FName PerspectiveIcon("EditorViewport.Perspective");
+	static FName TopIcon("EditorViewport.Top");
+	static FName LeftIcon("EditorViewport.Left");
+	static FName FrontIcon("EditorViewport.Front");
+	static FName BottomIcon("EditorViewport.Bottom");
+	static FName RightIcon("EditorViewport.Right");
+	static FName BackIcon("EditorViewport.Back");
+
+	switch (Mode)
+	{
+		case ENiagaraBakerViewMode::Perspective:	return PerspectiveIcon;
+		case ENiagaraBakerViewMode::OrthoFront:		return FrontIcon;
+		case ENiagaraBakerViewMode::OrthoBack:		return BackIcon;
+		case ENiagaraBakerViewMode::OrthoLeft:		return LeftIcon;
+		case ENiagaraBakerViewMode::OrthoRight:		return RightIcon;
+		case ENiagaraBakerViewMode::OrthoTop:		return TopIcon;
+		case ENiagaraBakerViewMode::OrthoBottom:	return BottomIcon;
+		default:									return PerspectiveIcon;
+	}
+}
+
+FSlateIcon FNiagaraBakerViewModel::GetCameraModeIcon(ENiagaraBakerViewMode Mode)
+{
+	return FSlateIcon(FEditorStyle::GetStyleSetName(), GetCameraModeIconName(Mode));
+}
+
+FVector FNiagaraBakerViewModel::GetCurrentCameraLocation() const
+{
+	UNiagaraBakerSettings* BakerSettings = GetBakerSettings();
+	return BakerSettings ? BakerSettings->CameraViewportLocation[int(BakerSettings->CameraViewportMode)] : FVector::ZeroVector;
+}
+
+void FNiagaraBakerViewModel::SetCurrentCameraLocation(const FVector Value)
+{
+	if ( UNiagaraBakerSettings* BakerSettings = GetBakerSettings() )
+	{
+		BakerSettings->CameraViewportLocation[int(BakerSettings->CameraViewportMode)] = Value;
+	}
+}
+
+FRotator FNiagaraBakerViewModel::GetCurrentCameraRotation() const
+{
+	UNiagaraBakerSettings* BakerSettings = GetBakerSettings();
+	return BakerSettings ? BakerSettings->CameraViewportRotation[int(BakerSettings->CameraViewportMode)] : FRotator::ZeroRotator;
+}
+
+void FNiagaraBakerViewModel::SetCurrentCameraRotation(const FRotator Value) const
+{
+	if (UNiagaraBakerSettings* BakerSettings = GetBakerSettings())
+	{
+		BakerSettings->CameraViewportRotation[int(BakerSettings->CameraViewportMode)] = Value;
+	}
+}
+
+float FNiagaraBakerViewModel::GetCameraFOV() const
+{
+	UNiagaraBakerSettings* BakerSettings = GetBakerSettings();
+	return BakerSettings ? BakerSettings->CameraFOV : 0.0f;
+}
+
+void FNiagaraBakerViewModel::SetCameraFOV(float InFOV)
+{
+	if (UNiagaraBakerSettings* BakerSettings = GetBakerSettings())
+	{
+		const FScopedTransaction Transaction(LOCTEXT("SetCameraFOV", "Set Camera FOV"));
+		BakerSettings->CameraFOV = InFOV;
+		BakerSettings->Modify();
+	}
+}
+
+float FNiagaraBakerViewModel::GetCameraOrbitDistance() const
 {
-	if ( UNiagaraBakerSettings* Settings = GetBakerSettings() )
+	UNiagaraBakerSettings* BakerSettings = GetBakerSettings();
+	return BakerSettings ? BakerSettings->CameraOrbitDistance : 0.0f;
+}
+
+void FNiagaraBakerViewModel::SetCameraOrbitDistance(float InOrbitDistance)
+{
+	if (UNiagaraBakerSettings* BakerSettings = GetBakerSettings())
+	{
+		const FScopedTransaction Transaction(LOCTEXT("SetCameraOrbitDistance", "Set Camera Orbit Distance"));
+		BakerSettings->CameraOrbitDistance = InOrbitDistance;
+		BakerSettings->Modify();
+	}
+}
+
+float FNiagaraBakerViewModel::GetCameraOrthoWidth() const
+{
+	UNiagaraBakerSettings* BakerSettings = GetBakerSettings();
+	return BakerSettings ? BakerSettings->CameraOrthoWidth : 0.0f;
+}
+
+void FNiagaraBakerViewModel::SetCameraOrthoWidth(float InOrthoWidth)
+{
+	if (UNiagaraBakerSettings* BakerSettings = GetBakerSettings())
+	{
+		const FScopedTransaction Transaction(LOCTEXT("SetCameraOrthoWidth", "Set Camera Ortho Width"));
+		BakerSettings->CameraOrthoWidth = InOrthoWidth;
+		BakerSettings->Modify();
+	}
+}
+
+void FNiagaraBakerViewModel::ToggleCameraAspectRatioEnabled()
+{
+	if ( UNiagaraBakerSettings* BakerSettings = GetBakerSettings() )
+	{
+		const FScopedTransaction Transaction(LOCTEXT("ToggleUseCameraAspectRatio", "Toggle Use Camera Aspect Ratio"));
+		BakerSettings->bUseCameraAspectRatio = !BakerSettings->bUseCameraAspectRatio;
+		BakerSettings->Modify();
+	}
+}
+
+bool FNiagaraBakerViewModel::IsCameraAspectRatioEnabled() const
+{
+	UNiagaraBakerSettings* BakerSettings = GetBakerSettings();
+	return BakerSettings ? BakerSettings->bUseCameraAspectRatio : false;
+}
+
+float FNiagaraBakerViewModel::GetCameraAspectRatio() const
+{
+	UNiagaraBakerSettings* BakerSettings = GetBakerSettings();
+	return BakerSettings ? BakerSettings->CameraAspectRatio : 0.0f;
+}
+
+void FNiagaraBakerViewModel::SetCameraAspectRatio(float InAspectRatio)
+{
+	if (UNiagaraBakerSettings* BakerSettings = GetBakerSettings())
+	{
+		const FScopedTransaction Transaction(LOCTEXT("SetCameraAspectRatio", "Set Camera Aspect Ratio"));
+		BakerSettings->CameraAspectRatio = InAspectRatio;
+		BakerSettings->Modify();
+	}
+}
+
+void FNiagaraBakerViewModel::AddOutput()
+{
+	if (UNiagaraBakerSettings* BakerSettings = GetBakerSettings())
 	{
-		if ( Settings->OutputTextures.IsValidIndex(PreviewTextureIndex) )
+		const FScopedTransaction Transaction(LOCTEXT("AddOutput", "Add Output"));
+		BakerSettings->OutputTextures.AddDefaulted();
+		BakerSettings->Modify();
+
+		CurrentOutputIndex = BakerSettings->OutputTextures.Num() - 1;
+		OnCurrentOutputChanged.Broadcast();
+	}
+}
+
+void FNiagaraBakerViewModel::RemoveCurrentOutput()
+{
+	if (UNiagaraBakerSettings* BakerSettings = GetBakerSettings())
+	{
+		if (BakerSettings->OutputTextures.IsValidIndex(CurrentOutputIndex))
 		{
-			return &Settings->OutputTextures[PreviewTextureIndex];
+			const FScopedTransaction Transaction(LOCTEXT("RemoveOutput", "Remove Output"));
+			BakerSettings->OutputTextures.RemoveAt(CurrentOutputIndex);
+			BakerSettings->Modify();
+			CurrentOutputIndex = FMath::Max(0, CurrentOutputIndex - 1);
+			OnCurrentOutputChanged.Broadcast();
 		}
 	}
-	return nullptr;
 }
 
-TSharedPtr<class SWidget> FNiagaraBakerViewModel::GetWidget()
+bool FNiagaraBakerViewModel::CanRemoveCurrentOutput() const
 {
-	return Widget;
+	UNiagaraBakerSettings* BakerSettings = GetBakerSettings();
+	return BakerSettings && BakerSettings->OutputTextures.IsValidIndex(CurrentOutputIndex);
+}
+
+void FNiagaraBakerViewModel::SetCurrentOutputIndex(int32 OutputIndex)
+{
+	if (UNiagaraBakerSettings* BakerSettings = GetBakerSettings())
+	{
+		if (BakerSettings->OutputTextures.Num() > 0)
+		{
+			CurrentOutputIndex = FMath::Clamp(OutputIndex, 0, BakerSettings->OutputTextures.Num() - 1);
+		}
+		else
+		{
+			CurrentOutputIndex = 0;
+		}
+		OnCurrentOutputChanged.Broadcast();
+	}
+}
+
+FText FNiagaraBakerViewModel::GetOutputText(int32 OutputIndex) const
+{
+	UNiagaraBakerSettings* BakerSettings = GetBakerSettings();
+	if ( BakerSettings && BakerSettings->OutputTextures.IsValidIndex(OutputIndex) )
+	{
+		if ( BakerSettings->OutputTextures[OutputIndex].OutputName.IsNone() )
+		{
+			return FText::Format(LOCTEXT("OutputFormat", "Output {0}"), FText::AsNumber(OutputIndex));
+		}
+		else
+		{
+			return FText::FromName(BakerSettings->OutputTextures[OutputIndex].OutputName);
+		}
+	}
+	return FText::GetEmpty();
+}
+
+FText FNiagaraBakerViewModel::GetCurrentOutputText() const
+{
+	return GetOutputText(CurrentOutputIndex);
 }
 
 void FNiagaraBakerViewModel::AddReferencedObjects(FReferenceCollector& Collector)
diff --git a/Engine/Plugins/FX/Niagara/Source/NiagaraEditor/Private/ViewModels/NiagaraBakerViewModel.h b/Engine/Plugins/FX/Niagara/Source/NiagaraEditor/Private/ViewModels/NiagaraBakerViewModel.h
--- a/Engine/Plugins/FX/Niagara/Source/NiagaraEditor/Private/ViewModels/NiagaraBakerViewModel.h
+++ b/Engine/Plugins/FX/Niagara/Source/NiagaraEditor/Private/ViewModels/NiagaraBakerViewModel.h
@@ -10,9 +10,20 @@
 
 #include "NiagaraBakerSettings.h"
 
+enum class ENiagaraBakerColorChannel : uint8
+{
+	Red,
+	Green,
+	Blue,
+	Alpha,
+	Num
+};
+
 class FNiagaraBakerViewModel : public TSharedFromThis<FNiagaraBakerViewModel>, public FGCObject
 {
 public:
+	DECLARE_MULTICAST_DELEGATE(FOnCurrentOutputChanged);
+
 	struct FDisplayData
 	{
 		int32		NumFrames = 0;
@@ -34,19 +45,69 @@ public:
 
 	void RefreshView();
 
-	int32 GetPreviewTextureIndex() const { return PreviewTextureIndex; }
-	void SetPreviewTextureIndex(int32 InPreviewTextureIndex) { PreviewTextureIndex = InPreviewTextureIndex; }
-
 	void SetDisplayTimeFromNormalized(float NormalizeTime);
 
 	class UNiagaraComponent* GetPreviewComponent() const;
 	class UNiagaraBakerSettings* GetBakerSettings() const;
 	const class UNiagaraBakerSettings* GetBakerGeneratedSettings() const;
 
-	const struct FNiagaraBakerTextureSettings* GetPreviewTextureSettings() const;
-
 	bool RenderView(const FRenderTarget* RenderTarget, FCanvas* Canvas, float WorldTime, int32 iOutputTextureIndex, bool bFillCanvas = false) const;
 
+	bool IsChannelEnabled(ENiagaraBakerColorChannel Channel) const { return bColorChannelEnabled[int(Channel)]; }
+	void ToggleChannelEnabled(ENiagaraBakerColorChannel Channel) { bColorChannelEnabled[int(Channel)] = !bColorChannelEnabled[int(Channel)]; }
+	void SetChannelEnabled(ENiagaraBakerColorChannel Channel, bool bEnabled) { bColorChannelEnabled[int(Channel)] = bEnabled; }
+
+	bool ShowRealtimePreview() const { return bShowRealtimePreview; }
+	void ToggleRealtimePreview() { bShowRealtimePreview = !bShowRealtimePreview; }
+
+	bool ShowBakedView() const { return bShowBakedView; }
+	void ToggleBakedView() { bShowBakedView = !bShowBakedView; }
+
+	bool IsCheckerboardEnabled() const { return bCheckerboardEnabled; }
+	void ToggleCheckerboardEnabled() { bCheckerboardEnabled = !bCheckerboardEnabled; }
+
+	bool ShowInfoText() const { return bShowInfoText; }
+	void ToggleInfoText() { bShowInfoText = !bShowInfoText; }
+
+	void SetCameraViewMode(ENiagaraBakerViewMode ViewMode);
+	bool IsCameraViewMode(ENiagaraBakerViewMode ViewMode);
+
+	FText GetCurrentCameraModeText() const;
+	FName GetCurrentCameraModeIconName() const;
+	FSlateIcon GetCurrentCameraModeIcon() const;
+
+	static FText GetCameraModeText(ENiagaraBakerViewMode Mode);
+	static FName GetCameraModeIconName(ENiagaraBakerViewMode Mode);
+	static FSlateIcon GetCameraModeIcon(ENiagaraBakerViewMode Mode);
+
+	FVector GetCurrentCameraLocation() const;
+	void SetCurrentCameraLocation(const FVector Value);
+	FRotator GetCurrentCameraRotation() const;
+	void SetCurrentCameraRotation(const FRotator Value) const;
+
+	float GetCameraFOV() const;
+	void SetCameraFOV(float InFOV);
+
+	float GetCameraOrbitDistance() const;
+	void SetCameraOrbitDistance(float InOrbitDistance);
+
+	float GetCameraOrthoWidth() const;
+	void SetCameraOrthoWidth(float InOrthoWidth);
+
+	void ToggleCameraAspectRatioEnabled();
+	bool IsCameraAspectRatioEnabled() const;
+	float GetCameraAspectRatio() const;
+	void SetCameraAspectRatio(float InAspectRatio);
+
+	void AddOutput();
+	void RemoveCurrentOutput();
+	bool CanRemoveCurrentOutput() const;
+	bool IsCurrentOutputIndex(int32 OutputIndex) const { return OutputIndex == CurrentOutputIndex; }
+	int32 GetCurrentOutputIndex() const { return CurrentOutputIndex; }
+	void SetCurrentOutputIndex(int32 OutputIndex);
+	FText GetOutputText(int32 OutputIndex) const;
+	FText GetCurrentOutputText() const;
+
 	/** FGCObject interface */
 	virtual void AddReferencedObjects(FReferenceCollector& Collector) override;
 	virtual FString GetReferencerName() const override
@@ -54,6 +115,8 @@ public:
 		return TEXT("FNiagaraBakerViewModel");
 	}
 
+	FOnCurrentOutputChanged OnCurrentOutputChanged;
+
 private:
 	class UNiagaraComponent* PreviewComponent = nullptr;
 	TSharedPtr<class FAdvancedPreviewScene> AdvancedPreviewScene;
@@ -61,5 +124,11 @@ private:
 	TWeakPtr<FNiagaraSystemViewModel> WeakSystemViewModel;
 	TSharedPtr<class SNiagaraBakerWidget> Widget;
 
-	int32 PreviewTextureIndex = 0;
+	int32 CurrentOutputIndex = 0;
+
+	bool bShowRealtimePreview = true;
+	bool bShowBakedView = true;
+	bool bCheckerboardEnabled = true;		//-TODO: Move to Baker Settings?
+	bool bShowInfoText = true;				//-TODO: Remove later
+	bool bColorChannelEnabled[int(ENiagaraBakerColorChannel::Num)] = {true, true, true, false};
 };
diff --git a/Engine/Plugins/FX/Niagara/Source/NiagaraEditor/Private/Widgets/SNiagaraBakerViewport.cpp b/Engine/Plugins/FX/Niagara/Source/NiagaraEditor/Private/Widgets/SNiagaraBakerViewport.cpp
--- a/Engine/Plugins/FX/Niagara/Source/NiagaraEditor/Private/Widgets/SNiagaraBakerViewport.cpp
+++ b/Engine/Plugins/FX/Niagara/Source/NiagaraEditor/Private/Widgets/SNiagaraBakerViewport.cpp
@@ -4,11 +4,13 @@
 #include "ViewModels/NiagaraBakerViewModel.h"
 #include "SNiagaraBakerViewportToolbar.h"
 #include "NiagaraBakerRenderer.h"
+#include "NiagaraBatchedElements.h"
 #include "NiagaraComponent.h"
 
 #include "Engine/Font.h"
 #include "Framework/MultiBox/MultiBoxBuilder.h"
 #include "Modules/ModuleManager.h"
+#include "CanvasItem.h"
 #include "CanvasTypes.h"
 #include "EditorViewportCommands.h"
 #include "EngineModule.h"
@@ -34,7 +36,7 @@ public:
 		// Only allow movement when in preview mode
 		auto ViewModel = WeakViewModel.Pin();
 		UNiagaraBakerSettings* BakerSettings = ViewModel ? ViewModel->GetBakerSettings() : nullptr;
-		if ( bShowPreview && BakerSettings )
+		if ( ViewModel->ShowRealtimePreview() && BakerSettings )
 		{
 			bool bForwardKeyState = false;
 			bool bBackwardKeyState = false;
@@ -100,7 +102,7 @@ public:
 		// Viewport movement only enabled when preview is enabled otherwise there is no feedback
 		auto ViewModel = WeakViewModel.Pin();
 		UNiagaraBakerSettings* BakerSettings = ViewModel ? ViewModel->GetBakerSettings() : nullptr;
-		if ( bShowPreview && BakerSettings && (Key == EKeys::MouseX || Key == EKeys::MouseY) )
+		if ( ViewModel->ShowRealtimePreview() && BakerSettings && (Key == EKeys::MouseX || Key == EKeys::MouseY) )
 		{
 			if (BakerSettings->IsOrthographic())
 			{
@@ -246,7 +248,7 @@ public:
 
 		const UNiagaraBakerSettings* BakerSettings = ViewModel->GetBakerSettings();
 		const UNiagaraBakerSettings* BakerGeneratedSettings = ViewModel->GetBakerGeneratedSettings();
-		const int32 PreviewTextureIndex = ViewModel->GetPreviewTextureIndex();
+		const int32 CurrentOutputIndex = ViewModel->GetCurrentOutputIndex();
 		const float TimelineDurationSeconds = BakerSettings->DurationSeconds;
 
 		// Determine view rects
@@ -255,13 +257,13 @@ public:
 		{
 			const FIntRect ViewRect = Canvas->GetViewRect();
 			const int32 Border = 3;
-			const int32 ViewWidth = bShowPreview && bShowBaker ? (ViewRect.Width() >> 1) - Border : ViewRect.Width();
+			const int32 ViewWidth = ViewModel->ShowRealtimePreview() && ViewModel->ShowBakedView() ? (ViewRect.Width() >> 1) - Border : ViewRect.Width();
 
-			if (bShowPreview)
+			if (ViewModel->ShowRealtimePreview())
 			{
-				if (BakerSettings->OutputTextures.IsValidIndex(PreviewTextureIndex))
+				if (BakerSettings->OutputTextures.IsValidIndex(CurrentOutputIndex))
 				{
-					PreviewViewRect = ConstrainRect(FIntRect(ViewRect.Min.X, ViewRect.Min.Y, ViewRect.Min.X + ViewWidth, ViewRect.Max.Y), BakerSettings->OutputTextures[PreviewTextureIndex].FrameSize);
+					PreviewViewRect = ConstrainRect(FIntRect(ViewRect.Min.X, ViewRect.Min.Y, ViewRect.Min.X + ViewWidth, ViewRect.Max.Y), BakerSettings->OutputTextures[CurrentOutputIndex].FrameSize);
 				}
 				else
 				{
@@ -269,11 +271,11 @@ public:
 				}
 			}
 
-			if ( bShowBaker )
+			if (ViewModel->ShowBakedView())
 			{
-				if (BakerGeneratedSettings && BakerGeneratedSettings->OutputTextures.IsValidIndex(PreviewTextureIndex))
+				if (BakerGeneratedSettings && BakerGeneratedSettings->OutputTextures.IsValidIndex(CurrentOutputIndex))
 				{
-					BakerViewRect = ConstrainRect(FIntRect(ViewRect.Max.X - ViewWidth, ViewRect.Min.Y, ViewRect.Max.X, ViewRect.Max.Y), BakerGeneratedSettings->OutputTextures[PreviewTextureIndex].FrameSize);
+					BakerViewRect = ConstrainRect(FIntRect(ViewRect.Max.X - ViewWidth, ViewRect.Min.Y, ViewRect.Max.X, ViewRect.Max.Y), BakerGeneratedSettings->OutputTextures[CurrentOutputIndex].FrameSize);
 				}
 				else
 				{
@@ -282,16 +284,33 @@ public:
 			}
 		}
 
+		const bool bRGBEnabled = ViewModel->IsChannelEnabled(ENiagaraBakerColorChannel::Red) || ViewModel->IsChannelEnabled(ENiagaraBakerColorChannel::Green) || ViewModel->IsChannelEnabled(ENiagaraBakerColorChannel::Blue);
+		const bool bAEnabled = ViewModel->IsChannelEnabled(ENiagaraBakerColorChannel::Alpha);
+		const bool bAlphaBlend = bAEnabled && bRGBEnabled;
+		FMatrix ColorTransform;
+		{
+			FPlane RPlane = ViewModel->IsChannelEnabled(ENiagaraBakerColorChannel::Red)   ? FPlane(1.0f, 0.0f, 0.0f, 0.0f) : FPlane(0.0f, 0.0f, 0.0f, 0.0f);
+			FPlane GPlane = ViewModel->IsChannelEnabled(ENiagaraBakerColorChannel::Green) ? FPlane(0.0f, 1.0f, 0.0f, 0.0f) : FPlane(0.0f, 0.0f, 0.0f, 0.0f);
+			FPlane BPlane = ViewModel->IsChannelEnabled(ENiagaraBakerColorChannel::Blue)  ? FPlane(0.0f, 0.0f, 1.0f, 0.0f) : FPlane(0.0f, 0.0f, 0.0f, 0.0f);
+			if (bAEnabled && !bRGBEnabled)
+			{
+				RPlane = FPlane(0.0f, 0.0f, 0.0f, 1.0f);
+				GPlane = FPlane(0.0f, 0.0f, 0.0f, 1.0f);
+				BPlane = FPlane(0.0f, 0.0f, 0.0f, 1.0f);
+			}
+			ColorTransform = FMatrix(RPlane, GPlane, BPlane, FPlane(0.0f, 0.0f, 0.0f, 1.0f));
+		}
+
 		// Render World Preview
-		if ( bShowPreview )
+		if (ViewModel->ShowRealtimePreview())
 		{
 			ClearViewArea(Canvas, PreviewViewRect);
 
 			const float WorldTime = RelativeTime + BakerSettings->StartSeconds;
 			FVector2D TextPosition(PreviewViewRect.Min.X + TextStartOffset.X, PreviewViewRect.Min.Y + TextStartOffset.Y);
-			if (BakerSettings->OutputTextures.IsValidIndex(PreviewTextureIndex))
+			if (BakerSettings->OutputTextures.IsValidIndex(CurrentOutputIndex))
 			{
-				const FNiagaraBakerTextureSettings& OutputTexture = BakerSettings->OutputTextures[PreviewTextureIndex];
+				const FNiagaraBakerTextureSettings& OutputTexture = BakerSettings->OutputTextures[CurrentOutputIndex];
 				if ( OutputTexture.IsValidForBake() )
 				{
 					// Ensure render target is the correct size
@@ -312,24 +331,27 @@ public:
 					PreviewComponent->SeekToDesiredAge(WorldTime);
 					PreviewComponent->TickComponent(BakerSettings->GetSeekDelta(), ELevelTick::LEVELTICK_All, nullptr);
 
-					BakerRenderer.RenderView(PreviewComponent, BakerSettings, WorldTime, RealtimeRenderTarget, PreviewTextureIndex);
+					BakerRenderer.RenderView(PreviewComponent, BakerSettings, WorldTime, RealtimeRenderTarget, CurrentOutputIndex);
 
 					const FVector2D HalfPixel(0.5f / float(OutputTexture.FrameSize.X), 0.5f / float(OutputTexture.FrameSize.Y));
-					Canvas->DrawTile(
-						PreviewViewRect.Min.X, PreviewViewRect.Min.Y, PreviewViewRect.Width(), PreviewViewRect.Height(),
-						HalfPixel.X, HalfPixel.Y, 1.0f - HalfPixel.X, 1.0f - HalfPixel.Y,
-						FLinearColor::White,
+					FCanvasTileItem TileItem(
+						FVector2D(PreviewViewRect.Min.X, PreviewViewRect.Min.Y),
 						RealtimeRenderTarget->GetResource(),
-						bAlphaBlend ? SE_BLEND_AlphaComposite : SE_BLEND_Opaque
+						FVector2D(PreviewViewRect.Width(), PreviewViewRect.Height()),
+						FVector2D(HalfPixel.X, HalfPixel.Y),
+						FVector2D(1.0f - HalfPixel.X, 1.0f - HalfPixel.Y),
+						FLinearColor::White
 					);
+					TileItem.BatchedElementParameters = new FBatchedElementNiagaraSimple(ColorTransform, bAlphaBlend);
+					Canvas->DrawItem(TileItem);
 
 					// Display info text
-					if (bShowInfoText)
+					if (ViewModel->ShowInfoText())
 					{
 						Canvas->DrawShadowedString(TextPosition.X, TextPosition.Y, TEXT("Live Preview"), DisplayFont, FLinearColor::White);
 						TextPosition.Y += FontHeight;
 
-						Canvas->DrawShadowedString(TextPosition.X + 5.0f, TextPosition.Y, *FString::Printf(TEXT("Texture(%d) FrameSize(%d x %d)"), PreviewTextureIndex, OutputTexture.FrameSize.X, OutputTexture.FrameSize.Y), DisplayFont, FLinearColor::White);
+						Canvas->DrawShadowedString(TextPosition.X + 5.0f, TextPosition.Y, *FString::Printf(TEXT("Texture(%d) FrameSize(%d x %d)"), CurrentOutputIndex, OutputTexture.FrameSize.X, OutputTexture.FrameSize.Y), DisplayFont, FLinearColor::White);
 						TextPosition.Y += FontHeight;
 					}
 				}
@@ -338,7 +360,7 @@ public:
 					Canvas->DrawShadowedString(TextPosition.X, TextPosition.Y, TEXT("Live Preview"), DisplayFont, FLinearColor::White);
 					TextPosition.Y += FontHeight;
 
-					Canvas->DrawShadowedString(TextPosition.X + 5.0f, TextPosition.Y, *FString::Printf(TEXT("Texture(%d) Is Invalid For Baking FrameSize(%d x %d)"), PreviewTextureIndex, OutputTexture.FrameSize.X, OutputTexture.FrameSize.Y), DisplayFont, FLinearColor::White);
+					Canvas->DrawShadowedString(TextPosition.X + 5.0f, TextPosition.Y, *FString::Printf(TEXT("Texture(%d) Is Invalid For Baking FrameSize(%d x %d)"), CurrentOutputIndex, OutputTexture.FrameSize.X, OutputTexture.FrameSize.Y), DisplayFont, FLinearColor::White);
 				}
 			}
 			else
@@ -351,21 +373,21 @@ public:
 		}
 
 		// Render Baker
-		if ( bShowBaker)
+		if (ViewModel->ShowBakedView())
 		{
 			ClearViewArea(Canvas, BakerViewRect);
 
 			const bool bTimelineValid = BakerGeneratedSettings ? FMath::IsNearlyEqual(BakerGeneratedSettings->DurationSeconds, TimelineDurationSeconds) : false;
 			const bool bBakerValid =
 				(BakerGeneratedSettings != nullptr) &&
-				BakerGeneratedSettings->OutputTextures.IsValidIndex(PreviewTextureIndex) &&
-				(BakerGeneratedSettings->OutputTextures[PreviewTextureIndex].GeneratedTexture != nullptr);
+				BakerGeneratedSettings->OutputTextures.IsValidIndex(CurrentOutputIndex) &&
+				(BakerGeneratedSettings->OutputTextures[CurrentOutputIndex].GeneratedTexture != nullptr);
 			FVector2D TextPosition(BakerViewRect.Min.X + TextStartOffset.X, BakerViewRect.Min.Y + TextStartOffset.Y);
 
 			if (bBakerValid && bTimelineValid)
 			{
 				const auto DisplayData = BakerGeneratedSettings->GetDisplayInfo(RelativeTime, BakerGeneratedSettings->bPreviewLooping);
-				const FNiagaraBakerTextureSettings& OutputTexture = BakerGeneratedSettings->OutputTextures[PreviewTextureIndex];
+				const FNiagaraBakerTextureSettings& OutputTexture = BakerGeneratedSettings->OutputTextures[CurrentOutputIndex];
 
 				const FIntPoint TextureSize = FIntPoint(FMath::Max(OutputTexture.GeneratedTexture->GetSizeX(), 1), FMath::Max(OutputTexture.GeneratedTexture->GetSizeY(), 1));
 				const FIntPoint FramesPerDimension = BakerGeneratedSettings->FramesPerDimension;
@@ -374,21 +396,24 @@ public:
 				const FIntPoint FramePixelA = FIntPoint(FrameIndexA.X * OutputTexture.FrameSize.X, FrameIndexA.Y * OutputTexture.FrameSize.Y);
 				const FIntPoint FramePixelB = FIntPoint(FrameIndexB.X * OutputTexture.FrameSize.X, FrameIndexB.Y * OutputTexture.FrameSize.Y);
 
-				Canvas->DrawTile(
-					BakerViewRect.Min.X, BakerViewRect.Min.Y, BakerViewRect.Width(), BakerViewRect.Height(),
-					(float(FramePixelA.X) + 0.5f) / float(TextureSize.X), (float(FramePixelA.Y) + 0.5f) / float(TextureSize.Y), (float(FramePixelA.X + OutputTexture.FrameSize.X) - 0.5f) / float(TextureSize.X), (float(FramePixelA.Y + OutputTexture.FrameSize.Y) - 0.5f) / float(TextureSize.Y),
-					FLinearColor::White,
+				FCanvasTileItem TileItem(
+					FVector2D(BakerViewRect.Min.X, BakerViewRect.Min.Y),
 					OutputTexture.GeneratedTexture->GetResource(),
-					bAlphaBlend ? SE_BLEND_AlphaComposite : SE_BLEND_Opaque
+					FVector2D(BakerViewRect.Width(), BakerViewRect.Height()),
+					FVector2D((float(FramePixelA.X) + 0.5f) / float(TextureSize.X), (float(FramePixelA.Y) + 0.5f) / float(TextureSize.Y)),
+					FVector2D((float(FramePixelA.X + OutputTexture.FrameSize.X) - 0.5f) / float(TextureSize.X), (float(FramePixelA.Y + OutputTexture.FrameSize.Y) - 0.5f) / float(TextureSize.Y)),
+					FLinearColor::White
 				);
+				TileItem.BatchedElementParameters = new FBatchedElementNiagaraSimple(ColorTransform, bAlphaBlend);
+				Canvas->DrawItem(TileItem);
 
 				// Display info text
-				if (bShowInfoText)
+				if (ViewModel->ShowInfoText())
 				{
 					Canvas->DrawShadowedString(TextPosition.X, TextPosition.Y, TEXT("Baker"), DisplayFont, FLinearColor::White);
 					TextPosition.Y += FontHeight;
 
-					Canvas->DrawShadowedString(TextPosition.X + 5.0f, TextPosition.Y, *FString::Printf(TEXT("Texture(%d) FrameSize(%d x %d) Frame(%d/%d)"), PreviewTextureIndex, OutputTexture.FrameSize.X, OutputTexture.FrameSize.Y, DisplayData.FrameIndexA, BakerGeneratedSettings->GetNumFrames()), DisplayFont, FLinearColor::White);
+					Canvas->DrawShadowedString(TextPosition.X + 5.0f, TextPosition.Y, *FString::Printf(TEXT("Texture(%d) FrameSize(%d x %d) Frame(%d/%d)"), CurrentOutputIndex, OutputTexture.FrameSize.X, OutputTexture.FrameSize.Y, DisplayData.FrameIndexA, BakerGeneratedSettings->GetNumFrames()), DisplayFont, FLinearColor::White);
 					TextPosition.Y += FontHeight;
 
 					if (!BakerGeneratedSettings->Equals(*BakerSettings))
@@ -429,7 +454,7 @@ public:
 		FLinearColor Color = ClearColor;
 		FVector2D EndUV(1.0f, 1.0f);
 
-		if (bShowCheckerboard)
+		if (WeakViewModel.Pin()->IsCheckerboardEnabled())
 		{
 			Texture = GetCheckerboardTexture();
 			Color = FLinearColor::White;
@@ -474,7 +499,7 @@ public:
 		auto ViewModel = WeakViewModel.Pin();
 		UNiagaraBakerSettings* BakerSettings = ViewModel ? ViewModel->GetBakerSettings() : nullptr;
 		UNiagaraComponent* NiagaraComponent = ViewModel->GetPreviewComponent();
-		if ( bShowPreview && BakerSettings && NiagaraComponent )
+		if (ViewModel->ShowRealtimePreview() && BakerSettings && NiagaraComponent )
 		{
 			//-TODO: Should take aspect ratio into account here
 			const FBoxSphereBounds ComponentBounds = NiagaraComponent->CalcBounds(NiagaraComponent->GetComponentTransform());
@@ -503,7 +528,7 @@ public:
 		const UNiagaraBakerSettings* BakerSettings = ViewModel ? ViewModel->GetBakerSettings() : nullptr;
 		if (BakerSettings && (PreviewViewRect.Area() > 0))
 		{
-			OrthoUnits = BakerSettings->GetOrthoSize(ViewModel->GetPreviewTextureIndex());
+			OrthoUnits = BakerSettings->GetOrthoSize(ViewModel->GetCurrentOutputIndex());
 			OrthoUnits.X = OrthoUnits.X / float(PreviewViewRect.Width());
 			OrthoUnits.Y = OrthoUnits.Y / float(PreviewViewRect.Height());
 		}
@@ -532,9 +557,6 @@ public:
 	FIntRect									PreviewViewRect;
 	FIntRect									BakerViewRect;
 
-	bool										bAlphaBlend = false;
-
-	bool										bShowCheckerboard = true;
 	UTexture2D*									CheckerboardTexture = nullptr;
 	FColor										CheckerboardColorOne = FColor(128, 128, 128);
 	FColor										CheckerboardColorTwo = FColor(64, 64, 64);
@@ -542,10 +564,6 @@ public:
 
 	FLinearColor								ClearColor = FColor::Black;
 
-	bool										bShowInfoText = true;
-	bool										bShowPreview = true;
-	bool										bShowBaker = true;
-
 	TWeakPtr<FNiagaraBakerViewModel>			WeakViewModel;
 	float										RelativeTime = 0.0f;
 	float										DeltaTime = 0.0f;
@@ -610,134 +628,10 @@ void SNiagaraBakerViewport::OnFloatingButtonClicked()
 {
 }
 
-bool SNiagaraBakerViewport::IsAlphaBlendEnabled() const
-{
-	return ViewportClient->bAlphaBlend;
-}
-
-void SNiagaraBakerViewport::SetAlphaBlendEnabled(bool bEnabled)
-{
-	ViewportClient->bAlphaBlend = bEnabled;
-}
-
-bool SNiagaraBakerViewport::IsCheckerboardEnabled() const
-{
-	return ViewportClient->bShowCheckerboard;
-}
-
-void SNiagaraBakerViewport::SetCheckerboardEnabled(bool bEnabled)
-{
-	ViewportClient->bShowCheckerboard = bEnabled;
-}
-
-bool SNiagaraBakerViewport::IsInfoTextEnabled() const
-{
-	return ViewportClient->bShowInfoText;
-}
-
-void SNiagaraBakerViewport::SetInfoTextEnabled(bool bEnabled)
-{
-	ViewportClient->bShowInfoText = bEnabled;
-}
-
-bool SNiagaraBakerViewport::IsPreviewViewEnabled() const
-{
-	return ViewportClient->bShowPreview;
-}
-
-void SNiagaraBakerViewport::SetPreviewViewEnabled(bool bEnabled)
-{
-	ViewportClient->bShowPreview = bEnabled;
-}
-
-bool SNiagaraBakerViewport::IsBakerViewEnabled() const
-{
-	return ViewportClient->bShowBaker;
-}
-
-void SNiagaraBakerViewport::SetBakerViewEnabled(bool bEnabled)
-{
-	ViewportClient->bShowBaker = bEnabled;
-}
-
 void SNiagaraBakerViewport::RefreshView(const float RelativeTime, const float DeltaTime)
 {
 	ViewportClient->RelativeTime = RelativeTime;
 	ViewportClient->DeltaTime = DeltaTime;
 }
 
-ENiagaraBakerViewMode SNiagaraBakerViewport::GetCameraMode() const
-{
-	if ( auto ViewModel = WeakViewModel.Pin() )
-	{
-		const UNiagaraBakerSettings* BakerSettings = ViewModel->GetBakerSettings();
-		return BakerSettings->CameraViewportMode;
-	}
-	// Should never happen
-	checkf(false, TEXT("No ViewModel present"));
-	return ENiagaraBakerViewMode::Perspective;
-}
-
-void SNiagaraBakerViewport::SetCameraMode(ENiagaraBakerViewMode Mode)
-{
-	if (auto ViewModel = WeakViewModel.Pin())
-	{
-		UNiagaraBakerSettings* BakerSettings = ViewModel->GetBakerSettings();
-		BakerSettings->CameraViewportMode = Mode;
-	}
-}
-
-bool SNiagaraBakerViewport::IsCameraMode(ENiagaraBakerViewMode Mode) const
-{
-	if (auto ViewModel = WeakViewModel.Pin())
-	{
-		const UNiagaraBakerSettings* BakerSettings = ViewModel->GetBakerSettings();
-		return BakerSettings->CameraViewportMode == Mode;
-	}
-	return false;
-}
-
-FText SNiagaraBakerViewport::GetCameraModeText(ENiagaraBakerViewMode Mode) const
-{
-	switch (Mode)
-	{
-		case ENiagaraBakerViewMode::Perspective:		return LOCTEXT("Perspective", "Perspective");
-		case ENiagaraBakerViewMode::OrthoFront:		return LOCTEXT("OrthoFront", "Front");
-		case ENiagaraBakerViewMode::OrthoBack:		return LOCTEXT("OrthoBack", "Back");
-		case ENiagaraBakerViewMode::OrthoLeft:		return LOCTEXT("OrthoLeft", "Left");
-		case ENiagaraBakerViewMode::OrthoRight:		return LOCTEXT("OrthoRight", "Right");
-		case ENiagaraBakerViewMode::OrthoTop:		return LOCTEXT("OrthoTop", "Top");
-		case ENiagaraBakerViewMode::OrthoBottom:		return LOCTEXT("OrthoBottom", "Bottom");
-		default:										return LOCTEXT("Error", "Error");
-	}
-}
-
-FName SNiagaraBakerViewport::GetCameraModeIconName(ENiagaraBakerViewMode Mode) const
-{
-	static FName PerspectiveIcon("EditorViewport.Perspective");
-	static FName TopIcon("EditorViewport.Top");
-	static FName LeftIcon("EditorViewport.Left");
-	static FName FrontIcon("EditorViewport.Front");
-	static FName BottomIcon("EditorViewport.Bottom");
-	static FName RightIcon("EditorViewport.Right");
-	static FName BackIcon("EditorViewport.Back");
-
-	switch (Mode)
-	{
-		case ENiagaraBakerViewMode::Perspective:		return PerspectiveIcon;
-		case ENiagaraBakerViewMode::OrthoFront:		return FrontIcon;
-		case ENiagaraBakerViewMode::OrthoBack:		return BackIcon;
-		case ENiagaraBakerViewMode::OrthoLeft:		return LeftIcon;
-		case ENiagaraBakerViewMode::OrthoRight:		return RightIcon;
-		case ENiagaraBakerViewMode::OrthoTop:		return TopIcon;
-		case ENiagaraBakerViewMode::OrthoBottom:		return BottomIcon;
-		default:										return PerspectiveIcon;
-	}
-}
-
-const struct FSlateBrush* SNiagaraBakerViewport::GetActiveCameraModeIcon() const
-{
-	return FEditorStyle::GetBrush(GetCameraModeIconName(GetCameraMode()));
-}
-
 #undef LOCTEXT_NAMESPACE
diff --git a/Engine/Plugins/FX/Niagara/Source/NiagaraEditor/Private/Widgets/SNiagaraBakerViewport.h b/Engine/Plugins/FX/Niagara/Source/NiagaraEditor/Private/Widgets/SNiagaraBakerViewport.h
--- a/Engine/Plugins/FX/Niagara/Source/NiagaraEditor/Private/Widgets/SNiagaraBakerViewport.h
+++ b/Engine/Plugins/FX/Niagara/Source/NiagaraEditor/Private/Widgets/SNiagaraBakerViewport.h
@@ -34,31 +34,6 @@ public:
 
 	void RefreshView(const float RelativeTime, const float DeltaTime);
 
-	bool IsAlphaBlendEnabled() const;
-	void SetAlphaBlendEnabled(bool bEnabled);
-
-	bool IsCheckerboardEnabled() const;
-	void SetCheckerboardEnabled(bool bEnabled);
-
-	bool IsInfoTextEnabled() const;
-	void SetInfoTextEnabled(bool bEnabled);
-
-	bool IsPreviewViewEnabled() const;
-	void SetPreviewViewEnabled(bool bEnabled);
-
-	bool IsBakerViewEnabled() const;
-	void SetBakerViewEnabled(bool bEnabled);
-
-	ENiagaraBakerViewMode GetCameraMode() const;
-	void SetCameraMode(ENiagaraBakerViewMode Mode);
-	bool IsCameraMode(ENiagaraBakerViewMode Mode) const;
-
-	FText GetCameraModeText(ENiagaraBakerViewMode Mode) const;
-	FName GetCameraModeIconName(ENiagaraBakerViewMode Mode) const;
-
-	FText GetActiveCameraModeText() const { return GetCameraModeText(GetCameraMode()); }
-	const struct FSlateBrush* GetActiveCameraModeIcon() const;
-
 private:
 	TWeakPtr<FNiagaraBakerViewModel>		WeakViewModel;
 	TSharedPtr<FNiagaraBakerViewportClient>	ViewportClient;
diff --git a/Engine/Plugins/FX/Niagara/Source/NiagaraEditor/Private/Widgets/SNiagaraBakerViewportToolbar.cpp b/Engine/Plugins/FX/Niagara/Source/NiagaraEditor/Private/Widgets/SNiagaraBakerViewportToolbar.cpp
--- a/Engine/Plugins/FX/Niagara/Source/NiagaraEditor/Private/Widgets/SNiagaraBakerViewportToolbar.cpp
+++ b/Engine/Plugins/FX/Niagara/Source/NiagaraEditor/Private/Widgets/SNiagaraBakerViewportToolbar.cpp
@@ -31,6 +31,9 @@ void SNiagaraBakerViewportToolbar::Construct(const FArguments& InArgs)
 	SNiagaraBakerViewport* Viewport = WeakViewport.Pin().Get();
 	check(Viewport != nullptr);
 
+	FNiagaraBakerViewModel* ViewModel = WeakViewModel.Pin().Get();
+	check(ViewModel != nullptr);
+
 	ChildSlot
 	[
 		SNew( SBorder )
@@ -46,18 +49,6 @@ void SNiagaraBakerViewportToolbar::Construct(const FArguments& InArgs)
 		]
 	];
 
-	// Toolbar
-	MainBoxPtr->AddSlot()
-	.AutoWidth()
-	.Padding(ToolbarSlotPadding)
-	[
-		SNew(SEditorViewportToolbarMenu)
-		.ParentToolBar(SharedThis(this))
-		.Cursor(EMouseCursor::Default)
-		.Image("EditorViewportToolBar.MenuDropdown")
-		.OnGetMenuContent(this, &SNiagaraBakerViewportToolbar::GenerateOptionsMenu)
-	];
-
 	// Camera Selection
 	MainBoxPtr->AddSlot()
 	.AutoWidth()
@@ -66,240 +57,27 @@ void SNiagaraBakerViewportToolbar::Construct(const FArguments& InArgs)
 		SNew(SEditorViewportToolbarMenu)
 		.ParentToolBar(SharedThis(this))
 		.Cursor(EMouseCursor::Default)
-		.Label(Viewport, &SNiagaraBakerViewport::GetActiveCameraModeText)
-		.LabelIcon(Viewport, &SNiagaraBakerViewport::GetActiveCameraModeIcon)
+		.Label(ViewModel, &FNiagaraBakerViewModel::GetCurrentCameraModeText)
+		.LabelIcon(this, &SNiagaraBakerViewportToolbar::GetCurrentCameraModeBrush)
 		.OnGetMenuContent(this, &SNiagaraBakerViewportToolbar::GenerateCameraMenu)
 	];
 
-	// Capture button
-	MainBoxPtr->AddSlot()
-	.AutoWidth()
-	.Padding(ToolbarSlotPadding)
-	[
-		SNew(SEditorViewportToolBarButton)
-		.Cursor(EMouseCursor::Default)
-		.ButtonType(EUserInterfaceActionType::Button)
-		.ButtonStyle(&FEditorStyle::Get().GetWidgetStyle<FButtonStyle>("EditorViewportToolBar.WarningButton"))
-		.OnClicked(
-			FOnClicked::CreateLambda(
-				[WeakViewModel=WeakViewModel]()
-				{
-					if ( auto ViewModel = WeakViewModel.Pin() )
-					{
-						ViewModel->RenderBaker();
-					}
-					return FReply::Handled();
-				}
-			)
-		)
-		.ToolTipText(LOCTEXT("BakeToolTip", "Runs the bake process."))
-		.Content()
-		[
-			SNew(STextBlock)
-			.Font(FEditorStyle::GetFontStyle("EditorViewportToolBar.Font"))
-			.Text(LOCTEXT("Bake", "Bake"))
-			.ColorAndOpacity(FLinearColor::White)
-		]
-	];
-	
 	SViewportToolBar::Construct(SViewportToolBar::FArguments());
 }
 
-TSharedRef<SWidget> SNiagaraBakerViewportToolbar::GenerateOptionsMenu() const
+const FSlateBrush* SNiagaraBakerViewportToolbar::GetCurrentCameraModeBrush() const
 {
-	const bool bInShouldCloseWindowAfterMenuSelection = true;
-	FMenuBuilder MenuBuilder(bInShouldCloseWindowAfterMenuSelection, nullptr);
-
-	// Show options for preview & flip book
-	MenuBuilder.BeginSection(NAME_None, LOCTEXT("ShowOptions", "Show Options"));
+	if ( FNiagaraBakerViewModel* ViewModel = WeakViewModel.Pin().Get() )
 	{
-		MenuBuilder.AddMenuEntry(
-			FText(LOCTEXT("AlphaBlend", "Alpha Blend")),
-			FText(LOCTEXT("AlphaBlendTooltip", "If we should use alpha blend or opaque to render previews.")),
-			FSlateIcon(),
-			FUIAction(
-				FExecuteAction::CreateLambda(
-					[WeakViewport=WeakViewport]()
-					{
-						if ( auto Viewport = WeakViewport.Pin() )
-						{
-							Viewport->SetAlphaBlendEnabled(!Viewport->IsAlphaBlendEnabled());
-						}
-					}
-				),
-				FCanExecuteAction(),
-				FIsActionChecked::CreateLambda(
-					[WeakViewport = WeakViewport]()
-					{
-						auto Viewport = WeakViewport.Pin();
-						return Viewport && Viewport->IsAlphaBlendEnabled();
-					}
-				)
-			),
-			NAME_None,
-			EUserInterfaceActionType::ToggleButton
-		);
-
-		MenuBuilder.AddMenuEntry(
-			FText(LOCTEXT("Checkerboard", "Checkboard")),
-			FText(LOCTEXT("CheckerboardTooltip", "Should the background be a checkerboard or not.")),
-			FSlateIcon(),
-			FUIAction(
-				FExecuteAction::CreateLambda(
-					[WeakViewport=WeakViewport]()
-					{
-						if ( auto Viewport = WeakViewport.Pin() )
-						{
-							Viewport->SetCheckerboardEnabled(!Viewport->IsCheckerboardEnabled());
-						}
-					}
-				),
-				FCanExecuteAction(),
-				FIsActionChecked::CreateLambda(
-					[WeakViewport = WeakViewport]()
-					{
-						auto Viewport = WeakViewport.Pin();
-						return Viewport && Viewport->IsCheckerboardEnabled();
-					}
-				)
-			),
-			NAME_None,
-			EUserInterfaceActionType::ToggleButton
-		);
-
-		MenuBuilder.AddMenuEntry(
-			FText(LOCTEXT("ShowInfoText", "Info Text")),
-			FText(LOCTEXT("ShowInfoTextTooltip", "When enabled information will be overlaid on each display.")),
-			FSlateIcon(),
-			FUIAction(
-				FExecuteAction::CreateLambda(
-					[WeakViewport=WeakViewport]()
-					{
-						if ( auto Viewport = WeakViewport.Pin() )
-						{
-							Viewport->SetInfoTextEnabled(!Viewport->IsInfoTextEnabled());
-						}
-					}
-				),
-				FCanExecuteAction(),
-				FIsActionChecked::CreateLambda(
-					[WeakViewport = WeakViewport]()
-					{
-						auto Viewport = WeakViewport.Pin();
-						return Viewport && Viewport->IsInfoTextEnabled();
-					}
-				)
-			),
-			NAME_None,
-			EUserInterfaceActionType::ToggleButton
-		);
-
-		MenuBuilder.AddMenuEntry(
-			FText(LOCTEXT("ShowPreview", "Live Preview")),
-			FText(LOCTEXT("ShowPreviewTooltip", "When enabled shows a live preview of what will be rendered, this may not be accurate with all visualization modes.")),
-			FSlateIcon(),
-			FUIAction(
-				FExecuteAction::CreateLambda(
-					[WeakViewport=WeakViewport]()
-					{
-						if ( auto Viewport = WeakViewport.Pin() )
-						{
-							Viewport->SetPreviewViewEnabled(!Viewport->IsPreviewViewEnabled());
-						}
-					}
-				),
-				FCanExecuteAction(),
-				FIsActionChecked::CreateLambda(
-					[WeakViewport = WeakViewport]()
-					{
-						auto Viewport = WeakViewport.Pin();
-						return Viewport && Viewport->IsPreviewViewEnabled();
-					}
-				)
-			),
-			NAME_None,
-			EUserInterfaceActionType::ToggleButton
-		);
-		MenuBuilder.AddMenuEntry(
-			FText(LOCTEXT("ShowBaker", "Baker")),
-			FText(LOCTEXT("ShowBakerTooltip", "When enabled shows a the generated Baker texture.")),
-			FSlateIcon(),
-			FUIAction(
-				FExecuteAction::CreateLambda(
-					[WeakViewport=WeakViewport]()
-					{ 
-						if ( auto Viewport = WeakViewport.Pin())
-						{
-							Viewport->SetBakerViewEnabled(!Viewport->IsBakerViewEnabled());
-						}
-					}
-				),
-				FCanExecuteAction(),
-				FIsActionChecked::CreateLambda(
-					[WeakViewport = WeakViewport]()
-					{
-						auto Viewport = WeakViewport.Pin();
-						return Viewport && Viewport->IsBakerViewEnabled();
-					}
-				)
-			),
-			NAME_None,
-			EUserInterfaceActionType::ToggleButton
-		);
+		return FEditorStyle::GetBrush(ViewModel->GetCurrentCameraModeIconName());
 	}
-	MenuBuilder.EndSection();
-
-	// Texture Selection Control
-	MenuBuilder.BeginSection(NAME_None, LOCTEXT("TextureSelection", "Texture Selection"));
-	if (auto ViewModel = WeakViewModel.Pin())
-	{
-		const UNiagaraBakerSettings* BakerSettings = ViewModel->GetBakerSettings();
-		for ( int32 i=0; i < BakerSettings->OutputTextures.Num(); ++i )
-		{
-			TStringBuilder<128> TextureName;
-			TextureName.Appendf(TEXT("Texture(%d)"), i);
-			if (UTexture2D* GeneratedTexture = BakerSettings->OutputTextures[i].GeneratedTexture)
-			{
-				TextureName.Appendf(TEXT(" - %s"), *BakerSettings->OutputTextures[i].GeneratedTexture->GetName());
-			}
-			
-			MenuBuilder.AddMenuEntry(
-				FText::FromStringView(TextureName.ToView()),
-				FText(),
-				FSlateIcon(),
-				FUIAction(
-					FExecuteAction::CreateLambda(
-						[WeakViewModel=WeakViewModel, TextureIndex=i]()
-						{ 
-							if ( auto ViewModel = WeakViewModel.Pin() )
-							{
-								ViewModel->SetPreviewTextureIndex(TextureIndex);
-							}
-						}
-					),
-					FCanExecuteAction(),
-					FIsActionChecked::CreateLambda(
-						[WeakViewModel=WeakViewModel, TextureIndex=i]()
-						{
-							auto ViewModel = WeakViewModel.Pin();
-							return ViewModel && ViewModel->GetPreviewTextureIndex() == TextureIndex;
-						}
-					)
-				),
-				NAME_None,
-				EUserInterfaceActionType::RadioButton
-			);
-		}
-	}
-	MenuBuilder.EndSection();
-
-	return MenuBuilder.MakeWidget();
+	return nullptr;
 }
 
 TSharedRef<SWidget> SNiagaraBakerViewportToolbar::GenerateCameraMenu() const
 {
-	SNiagaraBakerViewport* Viewport = WeakViewport.Pin().Get();
-	check(Viewport != nullptr);
+	FNiagaraBakerViewModel* ViewModel = WeakViewModel.Pin().Get();
+	check(ViewModel != nullptr);
 
 	const bool bInShouldCloseWindowAfterMenuSelection = true;
 	FMenuBuilder MenuBuilder(bInShouldCloseWindowAfterMenuSelection, nullptr);
@@ -308,13 +86,13 @@ TSharedRef<SWidget> SNiagaraBakerViewportToolbar::GenerateCameraMenu() const
 	{
 		ENiagaraBakerViewMode ViewMode = ENiagaraBakerViewMode(i);
 		MenuBuilder.AddMenuEntry(
-			Viewport->GetCameraModeText(ViewMode),
+			FNiagaraBakerViewModel::GetCameraModeText(ViewMode),
 			FText(),
-			FSlateIcon(FEditorStyle::GetStyleSetName(), Viewport->GetCameraModeIconName(ViewMode)),
+			FNiagaraBakerViewModel::GetCameraModeIcon(ViewMode),
 			FUIAction(
-				FExecuteAction::CreateSP(Viewport, &SNiagaraBakerViewport::SetCameraMode, ViewMode),
+				FExecuteAction::CreateSP(ViewModel, &FNiagaraBakerViewModel::SetCameraViewMode, ViewMode),
 				FCanExecuteAction(),
-				FIsActionChecked::CreateSP(Viewport, &SNiagaraBakerViewport::IsCameraMode, ViewMode)
+				FIsActionChecked::CreateSP(ViewModel, &FNiagaraBakerViewModel::IsCameraViewMode, ViewMode)
 			),
 			NAME_None,
 			EUserInterfaceActionType::RadioButton
diff --git a/Engine/Plugins/FX/Niagara/Source/NiagaraEditor/Private/Widgets/SNiagaraBakerViewportToolbar.h b/Engine/Plugins/FX/Niagara/Source/NiagaraEditor/Private/Widgets/SNiagaraBakerViewportToolbar.h
--- a/Engine/Plugins/FX/Niagara/Source/NiagaraEditor/Private/Widgets/SNiagaraBakerViewportToolbar.h
+++ b/Engine/Plugins/FX/Niagara/Source/NiagaraEditor/Private/Widgets/SNiagaraBakerViewportToolbar.h
@@ -20,7 +20,7 @@ public:
 
 	void Construct(const FArguments& InArgs);
 
-	TSharedRef<SWidget> GenerateOptionsMenu() const;
+	const FSlateBrush* GetCurrentCameraModeBrush() const;
 	TSharedRef<SWidget> GenerateCameraMenu() const;
 
 private:
diff --git a/Engine/Plugins/FX/Niagara/Source/NiagaraEditor/Private/Widgets/SNiagaraBakerWidget.cpp b/Engine/Plugins/FX/Niagara/Source/NiagaraEditor/Private/Widgets/SNiagaraBakerWidget.cpp
--- a/Engine/Plugins/FX/Niagara/Source/NiagaraEditor/Private/Widgets/SNiagaraBakerWidget.cpp
+++ b/Engine/Plugins/FX/Niagara/Source/NiagaraEditor/Private/Widgets/SNiagaraBakerWidget.cpp
@@ -4,37 +4,418 @@
 #include "SNiagaraBakerTimelineWidget.h"
 #include "SNiagaraBakerViewport.h"
 #include "NiagaraEditorCommon.h"
+#include "NiagaraEditorStyle.h"
 #include "NiagaraComponent.h"
 #include "NiagaraSystem.h"
 #include "ViewModels/NiagaraBakerViewModel.h"
 
 #include "Modules/ModuleManager.h"
+#include "Widgets/Input/SVectorInputBox.h"
+#include "Widgets/Input/SCheckBox.h"
+#include "Widgets/Layout/SSplitter.h"
 #include "Widgets/SBoxPanel.h"
 #include "Widgets/SOverlay.h"
 #include "Widgets/SViewport.h"
-#include "ITransportControl.h"
+#include "DetailCategoryBuilder.h"
+#include "DetailLayoutBuilder.h"
+#include "DetailWidgetRow.h"
 #include "EditorWidgetsModule.h"
+#include "ITransportControl.h"
+#include "IDetailCustomization.h"
 #include "PropertyEditorModule.h"
 
 #define LOCTEXT_NAMESPACE "NiagaraBakerWidget"
 
+//-TODO: Remove details panel and replace with customization to make more compact
+//-TODO: Improve feedback & warning display
+//-TODO: Throttle sim frames on warmup
+//-TODO: Turn off playback when out of focus
+//-TODO: Looping vs non-looping play back
+//-TODO: Fluid systems stop simulation with to previous?
+//-TODO: First frame black even though I have non zero start seconds
+//-TODO: None should display HDR or something sensible
+
 //////////////////////////////////////////////////////////////////////////
 
+namespace NiagaraBakerWidgetLocal
+{
+	template<typename TType>
+	TSharedRef<SWidget> MakeSpinBox(TOptional<TType> MinValue, TOptional<TType> MaxValue, TOptional<TType> MinSliderValue, TOptional<TType> MaxSliderValue, TAttribute<TType> GetValue, typename SSpinBox<TType>::FOnValueChanged SetValue)
+	{
+		return SNew(SBox)
+		.HAlign(HAlign_Right)
+		[
+			SNew(SBox)
+			.Padding(FMargin(4.0f, 0.0f, 0.0f, 0.0f))
+			.WidthOverride(100.0f)
+			[
+				SNew(SSpinBox<float>)
+				.Font(FEditorStyle::GetFontStyle(TEXT("MenuItem.Font")))
+				.MinValue(MinValue)
+				.MaxValue(MaxValue)
+				.MinSliderValue(MinSliderValue)
+				.MaxSliderValue(MaxSliderValue)
+				.Value(GetValue)
+				.OnValueChanged(SetValue)
+			]
+		];
+	}
+
+	template<typename TType, int NumElements>
+	struct FMakeVectorBoxHelper
+	{
+		DECLARE_DELEGATE_RetVal(TType, FGetter);
+		DECLARE_DELEGATE_OneParam(FSetter, TType);
+
+		static TSharedRef<SWidget> Construct(FGetter GetValue, FSetter SetValue, TOptional<TType> MinValue = TOptional<TType>(), TOptional<TType> MaxValue = TOptional<TType>(), TOptional<TType> MinSliderValue = TOptional<TType>(), TOptional<TType> MaxSliderValue = TOptional<TType>())
+		{
+			TSharedRef<SHorizontalBox> HorizontalBox = SNew(SHorizontalBox);
+
+			for ( int i=0; i < NumElements; ++i )
+			{
+				HorizontalBox->AddSlot()
+				.FillWidth(1.0f / float(NumElements))
+				.MaxWidth(60.0f)
+				[
+					SNew(SSpinBox<float>)
+					.Font(FEditorStyle::GetFontStyle(TEXT("MenuItem.Font")))
+					.MinValue(MinValue.IsSet() ? MinValue.GetValue()[i] : TOptional<float>())
+					.MaxValue(MaxValue.IsSet() ? MaxValue.GetValue()[i] : TOptional<float>())
+					.MinSliderValue(MinSliderValue.IsSet() ? MinSliderValue.GetValue()[i] : TOptional<float>())
+					.MaxSliderValue(MaxSliderValue.IsSet() ? MaxSliderValue.GetValue()[i] : TOptional<float>())
+					.Value_Lambda([=]() { return GetValue.Execute()[i]; })
+					.OnValueChanged_Lambda([=](float InValue) { FVector VectorValue = GetValue.Execute(); VectorValue[i] = InValue; SetValue.Execute(VectorValue); })
+				];
+			}
+
+			return
+				SNew(SBox)
+				.WidthOverride(180.0f)
+				.HAlign(HAlign_Right)
+				[
+					HorizontalBox
+				];
+		}
+	};
+
+	using FMakeVectorBox = FMakeVectorBoxHelper<FVector, 3>;
+
+	struct FMakeRotatorBox
+	{
+		DECLARE_DELEGATE_RetVal(FRotator, FGetter);
+		DECLARE_DELEGATE_OneParam(FSetter, FRotator);
+
+		static TSharedRef<SWidget> Construct(FGetter GetValue, FSetter SetValue)
+		{
+			return
+				SNew(SBox)
+				.WidthOverride(180.0f)
+				.HAlign(HAlign_Right)
+				[
+					SNew(SHorizontalBox)
+					+ SHorizontalBox::Slot()
+					.FillWidth(1.0f / float(3.0f))
+					.MaxWidth(60.0f)
+					[
+						SNew(SSpinBox<float>)
+						.Font(FEditorStyle::GetFontStyle(TEXT("MenuItem.Font")))
+						.Value_Lambda([=]() { return GetValue.Execute().Pitch; })
+						.OnValueChanged_Lambda([=](float InValue) { FRotator VectorValue = GetValue.Execute(); VectorValue.Pitch = InValue; SetValue.Execute(VectorValue); })
+					]
+					+ SHorizontalBox::Slot()
+					.FillWidth(1.0f / float(3.0f))
+					.MaxWidth(60.0f)
+					[
+						SNew(SSpinBox<float>)
+						.Font(FEditorStyle::GetFontStyle(TEXT("MenuItem.Font")))
+						.Value_Lambda([=]() { return GetValue.Execute().Yaw; })
+						.OnValueChanged_Lambda([=](float InValue) { FRotator VectorValue = GetValue.Execute(); VectorValue.Yaw = InValue; SetValue.Execute(VectorValue); })
+					]
+					+ SHorizontalBox::Slot()
+					.FillWidth(1.0f / float(3.0f))
+					.MaxWidth(60.0f)
+					[
+						SNew(SSpinBox<float>)
+						.Font(FEditorStyle::GetFontStyle(TEXT("MenuItem.Font")))
+						.Value_Lambda([=]() { return GetValue.Execute().Roll; })
+						.OnValueChanged_Lambda([=](float InValue) { FRotator VectorValue = GetValue.Execute(); VectorValue.Roll = InValue; SetValue.Execute(VectorValue); })
+					]
+				];
+		}
+	};
+
+	TSharedRef<SWidget> MakeChannelWidget(TWeakPtr<FNiagaraBakerViewModel> WeakViewModel)
+	{
+		TSharedRef<SHorizontalBox> HorizontalBox = SNew(SHorizontalBox);
+
+		static const FText ChannelLetter[int(ENiagaraBakerColorChannel::Num)] = { FText::FromString("R"), FText::FromString("G"), FText::FromString("B"), FText::FromString("A") };
+
+		for ( int32 i=0; i < int(ENiagaraBakerColorChannel::Num); ++i )
+		{
+			const ENiagaraBakerColorChannel Channel = ENiagaraBakerColorChannel(i);
+
+			HorizontalBox->AddSlot()
+			.VAlign(VAlign_Center)
+			.Padding(2.0f)
+			.AutoWidth()
+			[
+				SNew(SCheckBox)
+				.Style(FAppStyle::Get(), "TextureEditor.ChannelButtonStyle")
+				//.BorderBackgroundColor_Lambda(
+				//	[WeakViewModel, Channel]() -> FLinearColor
+				//	{
+				//		FNiagaraBakerViewModel* ViewModel = WeakViewModel.Pin().Get();
+				//		return ViewModel && ViewModel->IsChannelEnabled(Channel) ? FLinearColor::Gray : FLinearColor::Gray;
+				//	}
+				//)
+				.ForegroundColor_Lambda(
+					[WeakViewModel, Channel]() -> FLinearColor
+					{
+						FNiagaraBakerViewModel* ViewModel = WeakViewModel.Pin().Get();
+						if ( ViewModel && ViewModel->IsChannelEnabled(Channel) )
+						{
+							switch (Channel)
+							{
+								case ENiagaraBakerColorChannel::Red:	return FLinearColor::Red;
+								case ENiagaraBakerColorChannel::Green:	return FLinearColor::Green;
+								case ENiagaraBakerColorChannel::Blue:	return FLinearColor::Blue;
+								case ENiagaraBakerColorChannel::Alpha:	return FLinearColor::White;
+							}
+						}
+						return FLinearColor::Black;
+					}
+				)
+				.OnCheckStateChanged_Lambda(
+					[WeakViewModel, Channel](ECheckBoxState CheckState)
+					{
+						if ( FNiagaraBakerViewModel* ViewModel = WeakViewModel.Pin().Get() )
+						{
+							ViewModel->SetChannelEnabled(Channel, CheckState == ECheckBoxState::Checked);
+						}
+					}
+				)
+				.IsChecked_Lambda(
+					[WeakViewModel, Channel]() -> ECheckBoxState
+					{
+						FNiagaraBakerViewModel* ViewModel = WeakViewModel.Pin().Get();
+						const bool bEnabled = ViewModel ? ViewModel->IsChannelEnabled(Channel) : false;
+						return bEnabled ? ECheckBoxState::Checked : ECheckBoxState::Unchecked;
+					}
+				)
+				[
+					SNew(STextBlock)
+					.Font(FAppStyle::Get().GetFontStyle("TextureEditor.ChannelButtonFont"))
+					.Text(ChannelLetter[i])
+				]
+			];
+		}
+
+		return HorizontalBox;
+	}
+}
+
+//////////////////////////////////////////////////////////////////////////
+
+class FBakerSettingsDetails : public IDetailCustomization
+{
+public:
+	FBakerSettingsDetails(TWeakPtr<FNiagaraBakerViewModel> InWeakViewModel)
+		: WeakViewModel(InWeakViewModel)
+	{
+	}
+
+	static TSharedRef<IDetailCustomization> MakeInstance(TWeakPtr<FNiagaraBakerViewModel> InWeakViewModel)
+	{
+		return MakeShared<FBakerSettingsDetails>(InWeakViewModel);
+	}
+
+	virtual void CustomizeDetails(IDetailLayoutBuilder& DetailBuilder) override
+	{
+		// We only support customization on 1 object
+		TArray<TWeakObjectPtr<UObject>> ObjectsCustomized;
+		DetailBuilder.GetObjectsBeingCustomized(ObjectsCustomized);
+		if (ObjectsCustomized.Num() != 1 || !ObjectsCustomized[0]->IsA<UNiagaraBakerSettings>())
+		{
+			return;
+		}
+
+		FNiagaraBakerViewModel* ViewModel = WeakViewModel.Pin().Get();
+
+		UNiagaraBakerSettings* BakerSettings = CastChecked<UNiagaraBakerSettings>(ObjectsCustomized[0]);
+
+		// Func for only showing information for current camera
+		auto PerCameraProperty =
+			[&DetailBuilder, BakerSettings](IDetailCategoryBuilder& DetailCategory, TSharedPtr<IPropertyHandle> PropertyHandle, FText DisplayName, FText ToolTip)
+			{
+				TSharedPtr<IPropertyHandleArray> ArrayPropertyHandle = PropertyHandle->AsArray();
+				if (ArrayPropertyHandle == nullptr)
+				{
+					return;
+				}
+				uint32 NumElements = 0;
+				if ( ArrayPropertyHandle->GetNumElements(NumElements) != FPropertyAccess::Success )
+				{
+					return;
+				}
+
+				DetailBuilder.HideProperty(PropertyHandle->AsShared());
+				for (uint32 i=0; i < NumElements; ++i)
+				{
+					TSharedRef<IPropertyHandle> ArrayValuePropertyHandle = ArrayPropertyHandle->GetElement(i);
+
+					DetailCategory.AddProperty(ArrayValuePropertyHandle)
+						.DisplayName(DisplayName)
+						.ToolTip(ToolTip)
+						.Visibility(TAttribute<EVisibility>::CreateLambda([=] { return uint32(BakerSettings->CameraViewportMode) == i ? EVisibility::Visible : EVisibility::Hidden; }));
+				}
+			};
+
+		// Settings
+		{
+			IDetailCategoryBuilder& DetailCategory = DetailBuilder.EditCategory(FName(TEXT("Settings")));
+
+			DetailCategory.AddProperty(GET_MEMBER_NAME_CHECKED(UNiagaraBakerSettings, StartSeconds));
+			DetailCategory.AddProperty(GET_MEMBER_NAME_CHECKED(UNiagaraBakerSettings, DurationSeconds));
+			DetailBuilder.HideProperty(GET_MEMBER_NAME_CHECKED(UNiagaraBakerSettings, FramesPerSecond));
+			DetailCategory.AddProperty(GET_MEMBER_NAME_CHECKED(UNiagaraBakerSettings, FramesPerDimension));
+
+			DetailBuilder.HideProperty(GET_MEMBER_NAME_CHECKED(UNiagaraBakerSettings, CameraViewportMode));
+			DetailBuilder.HideProperty(GET_MEMBER_NAME_CHECKED(UNiagaraBakerSettings, CameraViewportLocation));
+			DetailBuilder.HideProperty(GET_MEMBER_NAME_CHECKED(UNiagaraBakerSettings, CameraViewportRotation));
+			DetailBuilder.HideProperty(GET_MEMBER_NAME_CHECKED(UNiagaraBakerSettings, CameraOrbitDistance));
+			DetailBuilder.HideProperty(GET_MEMBER_NAME_CHECKED(UNiagaraBakerSettings, CameraFOV));
+			DetailBuilder.HideProperty(GET_MEMBER_NAME_CHECKED(UNiagaraBakerSettings, CameraOrthoWidth));
+			DetailBuilder.HideProperty(GET_MEMBER_NAME_CHECKED(UNiagaraBakerSettings, CameraAspectRatio));
+
+			DetailBuilder.HideProperty(GET_MEMBER_NAME_CHECKED(UNiagaraBakerSettings, bPreviewLooping));
+			DetailBuilder.HideProperty(GET_MEMBER_NAME_CHECKED(UNiagaraBakerSettings, bRenderComponentOnly));
+
+			DetailBuilder.HideProperty(GET_MEMBER_NAME_CHECKED(UNiagaraBakerSettings, OutputTextures));
+		}
+
+		// Output Settings
+		if ( BakerSettings->OutputTextures.IsValidIndex(ViewModel->GetCurrentOutputIndex()) )
+		{
+			const FNiagaraBakerTextureSettings& OutputTexture = BakerSettings->OutputTextures[ViewModel->GetCurrentOutputIndex()];
+
+			IDetailCategoryBuilder& OutputTextureCategory = DetailBuilder.EditCategory(FName("OutputSettings"));
+
+			TSharedPtr<IPropertyHandleArray> OutputArrayPropertyHandle = DetailBuilder.GetProperty(GET_MEMBER_NAME_CHECKED(UNiagaraBakerSettings, OutputTextures))->AsArray();
+			TSharedPtr<IPropertyHandle> OutputPropertyHandle = OutputArrayPropertyHandle->GetElement(ViewModel->GetCurrentOutputIndex());
+			OutputTextureCategory.AddProperty(OutputPropertyHandle->GetChildHandle(GET_MEMBER_NAME_CHECKED(FNiagaraBakerTextureSettings, OutputName)));
+			OutputTextureCategory.AddProperty(OutputPropertyHandle->GetChildHandle(GET_MEMBER_NAME_CHECKED(FNiagaraBakerTextureSettings, SourceBinding)));
+			OutputTextureCategory.AddProperty(OutputPropertyHandle->GetChildHandle(GET_MEMBER_NAME_CHECKED(FNiagaraBakerTextureSettings, TextureSize)));
+			OutputTextureCategory.AddProperty(OutputPropertyHandle->GetChildHandle(GET_MEMBER_NAME_CHECKED(FNiagaraBakerTextureSettings, FrameSize)));
+			OutputTextureCategory.AddProperty(OutputPropertyHandle->GetChildHandle(GET_MEMBER_NAME_CHECKED(FNiagaraBakerTextureSettings, GeneratedTexture)));
+		}
+	}
+
+	TWeakPtr<FNiagaraBakerViewModel> WeakViewModel;
+};
+
 //////////////////////////////////////////////////////////////////////////
 
 void SNiagaraBakerWidget::Construct(const FArguments& InArgs)
 {
 	WeakViewModel = InArgs._WeakViewModel;
+	FNiagaraBakerViewModel* ViewModel = WeakViewModel.Pin().Get();
+	check(ViewModel);
 
-	//-TODO: Add tool bar to set options for the pane
-
-	FDetailsViewArgs DetailsArgs;
-	DetailsArgs.bHideSelectionTip = true;
+	OnCurrentOutputIndexChangedHandle = ViewModel->OnCurrentOutputChanged.AddSP(this, &SNiagaraBakerWidget::RefreshWidget);
 
 	FPropertyEditorModule& PropertyModule = FModuleManager::LoadModuleChecked<FPropertyEditorModule>("PropertyEditor");
 	FEditorWidgetsModule& EditorWidgetsModule = FModuleManager::Get().LoadModuleChecked<FEditorWidgetsModule>("EditorWidgets");
 
+	// Baker Toolbar
+	FSlimHorizontalToolBarBuilder BakerToolbarBuilder(MakeShareable(new FUICommandList), FMultiBoxCustomization::None);
+	BakerToolbarBuilder.AddToolBarButton(
+		FUIAction(FExecuteAction::CreateLambda([=]() { OnCapture(); })),
+		NAME_None,
+		FText::GetEmpty(),	//LOCTEXT("Bake", "Bake"),
+		LOCTEXT("BakeTooltip", "Run the bake process"),
+		FSlateIcon(FNiagaraEditorStyle::Get().GetStyleSetName(), "NiagaraEditor.BakerIcon")
+	);
+
+	BakerToolbarBuilder.AddComboButton(
+		FUIAction(),
+		FOnGetContent::CreateSP(this, &SNiagaraBakerWidget::MakeCameraModeMenu),
+		TAttribute<FText>::CreateSP(ViewModel, &FNiagaraBakerViewModel::GetCurrentCameraModeText),
+		LOCTEXT("CameraModeToolTip", "Change the camera used to render from"),
+		TAttribute<FSlateIcon>::CreateSP(ViewModel, &FNiagaraBakerViewModel::GetCurrentCameraModeIcon)
+	);
+
+	BakerToolbarBuilder.AddComboButton(
+		FUIAction(),
+		FOnGetContent::CreateSP(this, &SNiagaraBakerWidget::MakeViewOptionsMenu),
+		FText::GetEmpty(),
+		LOCTEXT("ViewOptionsToolTip", "Modify view options"),
+		FSlateIcon(FAppStyle::Get().GetStyleSetName(), "Icons.Visibility")
+	);
+
+	// Outputs
+	BakerToolbarBuilder.BeginSection(NAME_None);
+	{
+		BakerToolbarBuilder.AddComboButton(
+			FUIAction(),
+			FOnGetContent::CreateSP(this, &SNiagaraBakerWidget::MakeOutputMenu),
+			TAttribute<FText>::CreateSP(ViewModel, &FNiagaraBakerViewModel::GetCurrentOutputText),
+			LOCTEXT("CurrentOutputTooltip", "Select which output is currently visible in the preview area"),
+			FSlateIcon(FAppStyle::GetAppStyleSetName(), "Icons.Edit")
+		);
+		BakerToolbarBuilder.AddToolBarButton(
+			FUIAction(FExecuteAction::CreateSP(ViewModel, &FNiagaraBakerViewModel::AddOutput)),
+			NAME_None,
+			FText::GetEmpty(),
+			LOCTEXT("AddOutputTooltip", "Add a new output"),
+			FSlateIcon(FAppStyle::GetAppStyleSetName(), "Icons.Plus")
+		);
+		BakerToolbarBuilder.AddToolBarButton(
+			FUIAction(
+				FExecuteAction::CreateSP(ViewModel, &FNiagaraBakerViewModel::RemoveCurrentOutput),
+				FCanExecuteAction::CreateSP(ViewModel, &FNiagaraBakerViewModel::CanRemoveCurrentOutput)
+			),
+			NAME_None,
+			FText::GetEmpty(),
+			LOCTEXT("RemoveOutputTooltip", "Delete the currently selected output"),
+			FSlateIcon(FAppStyle::GetAppStyleSetName(), "Icons.Delete")
+		);
+	}
+	BakerToolbarBuilder.EndSection();
+
+	// Color Channels
+	BakerToolbarBuilder.BeginSection(NAME_None);
+	{
+		BakerToolbarBuilder.AddWidget(NiagaraBakerWidgetLocal::MakeChannelWidget(WeakViewModel));
+	}
+	BakerToolbarBuilder.EndSection();
+
+	// Warnings
+	BakerToolbarBuilder.BeginSection(NAME_None);
+	{
+		BakerToolbarBuilder.AddComboButton(
+			FUIAction(
+				FExecuteAction(),
+				FCanExecuteAction(),
+				FIsActionChecked(),
+				FIsActionButtonVisible::CreateSP(this, &SNiagaraBakerWidget::HasWarnings)
+			),
+			FOnGetContent::CreateSP(this, &SNiagaraBakerWidget::MakeWarningsMenu),
+			LOCTEXT("Warnings", "Warnings"),
+			LOCTEXT("WarningsToolTip", "Contains any warnings that may need attention"),
+			FSlateIcon(FAppStyle::Get().GetStyleSetName(), "Icons.WarningWithColor")
+		);
+	}
+	BakerToolbarBuilder.EndSection();
+
+
+	// Baker Settings
+	FDetailsViewArgs DetailsArgs;
+	DetailsArgs.bHideSelectionTip = true;
+	DetailsArgs.bAllowSearch = false;
+
+	FOnGetDetailCustomizationInstance BakerSettingsCustomizeDetails = FOnGetDetailCustomizationInstance::CreateStatic(&FBakerSettingsDetails::MakeInstance, WeakViewModel);
 	BakerSettingsDetails = PropertyModule.CreateDetailView(DetailsArgs);
+	BakerSettingsDetails->RegisterInstancedCustomPropertyLayout(UNiagaraBakerSettings::StaticClass(), BakerSettingsCustomizeDetails);
 
 	// Transport control args
 	{
@@ -62,26 +443,40 @@ void SNiagaraBakerWidget::Construct(const FArguments& InArgs)
 	this->ChildSlot
 	[
 		SNew(SVerticalBox)
-		// Viewport
 		+SVerticalBox::Slot()
-		.FillHeight(0.60f)
+		.AutoHeight()
 		[
-			SNew(SOverlay)
-			+ SOverlay::Slot()
+			BakerToolbarBuilder.MakeWidget()
+		]
+		+SVerticalBox::Slot()
+		.FillHeight(1.0f)
+		[
+			SNew(SSplitter)
+			.Orientation(EOrientation::Orient_Vertical)
+			.PhysicalSplitterHandleSize(1.0f)
+			.HitDetectionSplitterHandleSize(5.0f)
+			.ResizeMode(ESplitterResizeMode::Fill)
+			+ SSplitter::Slot()
+			.Value(0.70f)
 			[
-				SAssignNew(ViewportWidget, SNiagaraBakerViewport)
-				.WeakViewModel(WeakViewModel)
+				SNew(SOverlay)
+				+ SOverlay::Slot()
+				[
+					SAssignNew(ViewportWidget, SNiagaraBakerViewport)
+					.WeakViewModel(WeakViewModel)
+				]
+			]
+			+SSplitter::Slot()
+			.Value(0.30f)
+			[
+				BakerSettingsDetails.ToSharedRef()
 			]
 		]
-		// Timeline
 		+SVerticalBox::Slot()
 		.AutoHeight()
-		.Padding(2)
-		.VAlign(VAlign_Fill)
 		[
 			SNew(SHorizontalBox)
-			+ SHorizontalBox::Slot()
-			.Padding(10.f, 0.f, 0.f, 0.f)
+			+SHorizontalBox::Slot()
 			.AutoWidth()
 			[
 				TransportControls.ToSharedRef()
@@ -92,73 +487,326 @@ void SNiagaraBakerWidget::Construct(const FArguments& InArgs)
 				.WeakViewModel(WeakViewModel)
 			]
 		]
-		// Settings Panel
-		+SVerticalBox::Slot()
-		.FillHeight(0.40f)
-		[
-			BakerSettingsDetails.ToSharedRef()
-		]
 	];
 
-	if ( auto ViewModel = WeakViewModel.Pin() )
+	RefreshWidget();
+}
+
+SNiagaraBakerWidget::~SNiagaraBakerWidget()
+{
+	if ( FNiagaraBakerViewModel* ViewModel = WeakViewModel.Pin().Get() )
 	{
-		if ( UNiagaraComponent* PreviewComponent = ViewModel->GetPreviewComponent() )
+		ViewModel->OnCurrentOutputChanged.Remove(OnCurrentOutputIndexChangedHandle);
+	}
+}
+
+TSharedRef<SWidget> SNiagaraBakerWidget::MakeCameraModeMenu()
+{
+	FNiagaraBakerViewModel* ViewModel = WeakViewModel.Pin().Get();
+	check(ViewModel != nullptr);
+
+	FMenuBuilder MenuBuilder(true, nullptr);
+
+	for (int i = 0; i < int(ENiagaraBakerViewMode::Num); ++i)
+	{
+		MenuBuilder.AddMenuEntry(
+			FNiagaraBakerViewModel::GetCameraModeText(ENiagaraBakerViewMode(i)),
+			FText::GetEmpty(),
+			FNiagaraBakerViewModel::GetCameraModeIcon(ENiagaraBakerViewMode(i)),
+			FUIAction(
+				FExecuteAction::CreateSP(ViewModel, &FNiagaraBakerViewModel::SetCameraViewMode, ENiagaraBakerViewMode(i)),
+				FCanExecuteAction(),
+				FIsActionChecked::CreateSP(ViewModel, &FNiagaraBakerViewModel::IsCameraViewMode, ENiagaraBakerViewMode(i))
+			),
+			NAME_None,
+			EUserInterfaceActionType::ToggleButton
+		);
+	}
+
+	return MenuBuilder.MakeWidget();
+}
+
+TSharedRef<SWidget> SNiagaraBakerWidget::MakeViewOptionsMenu()
+{
+	using namespace NiagaraBakerWidgetLocal;
+
+	FNiagaraBakerViewModel* ViewModel = WeakViewModel.Pin().Get();
+	check(ViewModel != nullptr);
+
+	FMenuBuilder MenuBuilder(true, nullptr);
+
+	MenuBuilder.BeginSection(NAME_None, LOCTEXT("PreviewSettings", "Preview Settings"));
+	{
+		MenuBuilder.AddMenuEntry(
+			LOCTEXT("ShowRealtimePreview", "Show Realtime Preview"),
+			LOCTEXT("ShowRealtimePreviewTooltip", "When enabled shows a live preview of what will be rendered, this may not be accurate with all visualization modes."),
+			FSlateIcon(),
+			FUIAction(
+				FExecuteAction::CreateSP(ViewModel, &FNiagaraBakerViewModel::ToggleRealtimePreview),
+				FCanExecuteAction(),
+				FIsActionChecked::CreateSP(ViewModel, &FNiagaraBakerViewModel::ShowRealtimePreview)
+			),
+			NAME_None,
+			EUserInterfaceActionType::ToggleButton
+		);
+
+		MenuBuilder.AddMenuEntry(
+			LOCTEXT("ShowBakedView", "Show Baked View"),
+			LOCTEXT("ShowBakedViewTooltip", "When enabled shows the baked texture."),
+			FSlateIcon(),
+			FUIAction(
+				FExecuteAction::CreateSP(ViewModel, &FNiagaraBakerViewModel::ToggleBakedView),
+				FCanExecuteAction(),
+				FIsActionChecked::CreateSP(ViewModel, &FNiagaraBakerViewModel::ShowBakedView)
+			),
+			NAME_None,
+			EUserInterfaceActionType::ToggleButton
+		);
+
+		MenuBuilder.AddMenuEntry(
+			LOCTEXT("ShowCheckerboard", "Show Checkerboard"),
+			LOCTEXT("ShowCheckerboardTooltip", "Show a checkerboard rather than a solid color to easily visualize alpha blending."),
+			FSlateIcon(),
+			FUIAction(
+				FExecuteAction::CreateSP(ViewModel, &FNiagaraBakerViewModel::ToggleCheckerboardEnabled),
+				FCanExecuteAction(),
+				FIsActionChecked::CreateSP(ViewModel, &FNiagaraBakerViewModel::IsCheckerboardEnabled)
+			),
+			NAME_None,
+			EUserInterfaceActionType::ToggleButton
+		);
+
+		MenuBuilder.AddMenuEntry(
+			LOCTEXT("ShowInfoText", "Show Info Text"),
+			LOCTEXT("ShowInfoTextTooltip", "Shows information about the preview and baked outputs."),
+			FSlateIcon(),
+			FUIAction(
+				FExecuteAction::CreateSP(ViewModel, &FNiagaraBakerViewModel::ToggleInfoText),
+				FCanExecuteAction(),
+				FIsActionChecked::CreateSP(ViewModel, &FNiagaraBakerViewModel::ShowInfoText)
+			),
+			NAME_None,
+			EUserInterfaceActionType::ToggleButton
+		);
+
+		//DetailCategory.HideProperty(GET_MEMBER_NAME_CHECKED(UNiagaraBakerSettings, bPreviewLooping));
+		//DetailCategory.HideProperty(GET_MEMBER_NAME_CHECKED(UNiagaraBakerSettings, bRenderComponentOnly));
+	}
+
+	MenuBuilder.BeginSection(NAME_None, LOCTEXT("Camera", "Camera"));
+	{
+		MenuBuilder.AddWidget(
+			FMakeVectorBox::Construct(
+				FMakeVectorBox::FGetter::CreateSP(ViewModel, &FNiagaraBakerViewModel::GetCurrentCameraLocation),
+				FMakeVectorBox::FSetter::CreateSP(ViewModel, &FNiagaraBakerViewModel::SetCurrentCameraLocation)
+			),
+			LOCTEXT("CameraLocation", "Camera Location")
+		);
+
+		MenuBuilder.AddWidget(
+			FMakeRotatorBox::Construct(
+				FMakeRotatorBox::FGetter::CreateSP(ViewModel, &FNiagaraBakerViewModel::GetCurrentCameraRotation),
+				FMakeRotatorBox::FSetter::CreateSP(ViewModel, &FNiagaraBakerViewModel::SetCurrentCameraRotation)
+			),
+			LOCTEXT("CameraLocation", "Camera Rotation")
+		);
+
+		if ( ViewModel->IsCameraViewMode(ENiagaraBakerViewMode::Perspective) )
 		{
-			if ( UNiagaraSystem* NiagaraSystem = PreviewComponent->GetAsset() )
-			{
-				BakerSettingsDetails->SetObject(NiagaraSystem->GetBakerSettings());
-			}
+			MenuBuilder.AddWidget(
+				MakeSpinBox<float>(1.0f, 170.0f, 1.0f, 170.0f, TAttribute<float>::CreateSP(ViewModel, &FNiagaraBakerViewModel::GetCameraFOV), SSpinBox<float>::FOnValueChanged::CreateSP(ViewModel, &FNiagaraBakerViewModel::SetCameraFOV)),
+				LOCTEXT("FOVAngle", "Field of View (H)")
+			);
+
+			MenuBuilder.AddWidget(
+				MakeSpinBox<float>(0.0f, TOptional<float>(), 0.0f, 1000.0f, TAttribute<float>::CreateSP(ViewModel, &FNiagaraBakerViewModel::GetCameraOrbitDistance), SSpinBox<float>::FOnValueChanged::CreateSP(ViewModel, &FNiagaraBakerViewModel::SetCameraOrbitDistance)),
+				LOCTEXT("OrbitDistance", "Orbit Distance")
+			);
+		}
+		else
+		{
+			MenuBuilder.AddWidget(
+				MakeSpinBox<float>(0.1f, TOptional<float>(), 0.1f, 1000.0f, TAttribute<float>::CreateSP(ViewModel, &FNiagaraBakerViewModel::GetCameraOrthoWidth), SSpinBox<float>::FOnValueChanged::CreateSP(ViewModel, &FNiagaraBakerViewModel::SetCameraOrthoWidth)),
+				LOCTEXT("CameraOrthoWidth", "Camera Orthographic Width")
+			);
+		}
+
+		//-TODO: Clean this up
+		MenuBuilder.AddMenuEntry(
+			LOCTEXT("CustomAspectRatio", "Custom Aspect Ratio"),
+			LOCTEXT("CustomAspectRatioTooltip", "Use a custom aspect ratio to render with."),
+			FSlateIcon(),
+			FUIAction(
+				FExecuteAction::CreateSP(ViewModel, &FNiagaraBakerViewModel::ToggleCameraAspectRatioEnabled),
+				FCanExecuteAction(),
+				FIsActionChecked::CreateSP(ViewModel, &FNiagaraBakerViewModel::IsCameraAspectRatioEnabled)
+			),
+			NAME_None,
+			EUserInterfaceActionType::ToggleButton
+		);
+
+		if ( ViewModel->IsCameraAspectRatioEnabled() )
+		{
+			MenuBuilder.AddWidget(
+				MakeSpinBox<float>(0.1f, 5.0f, 0.1f, 5.0f, TAttribute<float>::CreateSP(ViewModel, &FNiagaraBakerViewModel::GetCameraAspectRatio), SSpinBox<float>::FOnValueChanged::CreateSP(ViewModel, &FNiagaraBakerViewModel::SetCameraAspectRatio)),
+				LOCTEXT("CameraAspectRatio", "Camera Aspect Ratio")
+			);
 		}
 	}
+	MenuBuilder.EndSection();
+
+	return MenuBuilder.MakeWidget();
 }
 
-FText SNiagaraBakerWidget::GetSelectedTextureAsText() const
+TSharedRef<SWidget> SNiagaraBakerWidget::MakeOutputMenu()
 {
-	if (auto ViewModel = WeakViewModel.Pin())
+	FNiagaraBakerViewModel* ViewModel = WeakViewModel.Pin().Get();
+	check(ViewModel != nullptr);
+
+	FMenuBuilder MenuBuilder(true, nullptr);
+
+	if ( UNiagaraBakerSettings* BakerSettings = ViewModel->GetBakerSettings() )
+	{
+		for (int i = 0; i < BakerSettings->OutputTextures.Num(); ++i)
+		{
+			MenuBuilder.AddMenuEntry(
+				ViewModel->GetOutputText(i),
+				FText::GetEmpty(),
+				FSlateIcon(),
+				FUIAction(
+					FExecuteAction::CreateSP(ViewModel, &FNiagaraBakerViewModel::SetCurrentOutputIndex, i),
+					FCanExecuteAction(),
+					FIsActionChecked::CreateSP(ViewModel, &FNiagaraBakerViewModel::IsCurrentOutputIndex, i)
+				),
+				NAME_None,
+				EUserInterfaceActionType::ToggleButton
+			);
+		}
+	}
+
+	return MenuBuilder.MakeWidget();
+}
+
+bool SNiagaraBakerWidget::FindWarnings(TArray<FText>* OutWarnings) const
+{
+	if ( FNiagaraBakerViewModel* ViewModel = WeakViewModel.Pin().Get() )
 	{
-		if ( const UNiagaraBakerSettings* BakerSettings = ViewModel->GetBakerSettings() )
+		// Check determinism
+		if ( UNiagaraComponent* PreviewComponent = ViewModel->GetPreviewComponent() )
+		{
+			if ( UNiagaraSystem* NiagaraSystem = PreviewComponent->GetAsset() )
+			{
+				if ( NiagaraSystem->NeedsDeterminism() == false )
+				{
+					if ( OutWarnings == nullptr )
+					{
+						return true;
+					}
+					OutWarnings->Emplace(LOCTEXT("SystemDeterminism", "System is not set to deterministic, results will vary each bake"));
+				}
+				for (const FNiagaraEmitterHandle& EmitterHandle : NiagaraSystem->GetEmitterHandles() )
+				{
+					UNiagaraEmitter* NiagaraEmitter = EmitterHandle.GetInstance();
+					if (NiagaraEmitter == nullptr || NiagaraEmitter->bDeterminism)
+					{
+						continue;
+					}
+
+					if (OutWarnings == nullptr)
+					{
+						return true;
+					}
+					OutWarnings->Emplace(FText::Format(LOCTEXT("EmitterDeterminismFormat", "Emitter '{0}' is not set to deterministic, results will vary each bake"), FText::FromString(NiagaraEmitter->GetUniqueEmitterName())));
+				}
+			}
+		}
+
+		// Check outputs divide correctly by atlas frame count
+		if ( UNiagaraBakerSettings* BakerSettings = ViewModel->GetBakerSettings() )
 		{
-			const int32 PrevewTextureIndex = ViewModel->GetPreviewTextureIndex();
-			if ( BakerSettings->OutputTextures.IsValidIndex(PrevewTextureIndex) )
+			for ( int32 i=0; i < BakerSettings->OutputTextures.Num(); ++i )
 			{
-				TStringBuilder<128> TextureName;
-				TextureName.Appendf(TEXT("Texture(%d)"), PrevewTextureIndex);
-				if (UTexture2D* GeneratedTexture = BakerSettings->OutputTextures[PrevewTextureIndex].GeneratedTexture)
+				const FNiagaraBakerTextureSettings& TextureSettings = BakerSettings->OutputTextures[i];
+				const bool bXSizeIssue = (TextureSettings.TextureSize.X % BakerSettings->FramesPerDimension.X) == 0;
+				const bool bYSizeIssue = (TextureSettings.TextureSize.Y % BakerSettings->FramesPerDimension.Y) == 0;
+				if (bXSizeIssue && bYSizeIssue)
+				{
+					continue;
+				}
+
+				if (OutWarnings == nullptr)
 				{
-					TextureName.Appendf(TEXT(" - %s"), *BakerSettings->OutputTextures[PrevewTextureIndex].GeneratedTexture->GetName());
+					return true;
 				}
-				return FText::FromStringView(TextureName.ToView());
+				OutWarnings->Emplace(FText::Format(LOCTEXT("TextureSizeWarningText", "Output '{0}' frames per dimension {1}x{2} does not divide into texture size {3}x{4}, this can result in jitter on the flipbook if not taken into account"),
+					ViewModel->GetOutputText(i),
+					BakerSettings->FramesPerDimension.X, BakerSettings->FramesPerDimension.Y,
+					TextureSettings.TextureSize.X, TextureSettings.TextureSize.Y
+				));
 			}
 		}
 	}
-	return FText::FromString(FString(TEXT("No Texture")));
+
+	return OutWarnings ? OutWarnings->Num() > 0 : false;
 }
 
-void SNiagaraBakerWidget::Tick(const FGeometry& AllottedGeometry, const double CurrentTime, const float DeltaTime)
+TSharedRef<SWidget> SNiagaraBakerWidget::MakeWarningsMenu()
 {
-	if (auto ViewModel = WeakViewModel.Pin())
+	FNiagaraBakerViewModel* ViewModel = WeakViewModel.Pin().Get();
+	check(ViewModel != nullptr);
+
+	FMenuBuilder MenuBuilder(true, nullptr);
+
+	TArray<FText> FoundWarnings;
+	if ( FindWarnings(&FoundWarnings) )
 	{
-		float DurationSeconds = 0.0f;
-		if (UNiagaraBakerSettings* BakerSettings = GetBakerSettings())
+		for (const FText& Warning : FoundWarnings)
 		{
-			DurationSeconds = BakerSettings->DurationSeconds;
+			MenuBuilder.AddMenuEntry(
+				Warning,
+				FText::GetEmpty(),
+				FSlateIcon(),
+				FUIAction()
+			);
 		}
+	}
 
-		if ( DurationSeconds > 0.0f )
-		{
-			if (bIsPlaying)
-			{
-				PreviewRelativeTime += DeltaTime;
-				PreviewRelativeTime = FMath::Fmod(PreviewRelativeTime, DurationSeconds);
-			}
-			else
-			{
-				PreviewRelativeTime = FMath::Min(PreviewRelativeTime, DurationSeconds);
-			}
+	return MenuBuilder.MakeWidget();
+}
 
-			ViewportWidget->RefreshView(PreviewRelativeTime, DeltaTime);
-			TimelineWidget->SetRelativeTime(PreviewRelativeTime);
+void SNiagaraBakerWidget::RefreshWidget()
+{
+	BakerSettingsDetails->SetObject(GetBakerSettings(), true);
+}
+
+void SNiagaraBakerWidget::Tick(const FGeometry& AllottedGeometry, const double CurrentTime, const float DeltaTime)
+{
+	TSharedPtr<FNiagaraBakerViewModel> ViewModel = WeakViewModel.Pin();
+	if (ViewModel == nullptr)
+	{
+		return;
+	}
+
+	float DurationSeconds = 0.0f;
+	if (UNiagaraBakerSettings* BakerSettings = GetBakerSettings())
+	{
+		DurationSeconds = BakerSettings->DurationSeconds;
+	}
+
+	if ( DurationSeconds > 0.0f )
+	{
+		if (bIsPlaying)
+		{
+			PreviewRelativeTime += DeltaTime;
+			PreviewRelativeTime = FMath::Fmod(PreviewRelativeTime, DurationSeconds);
 		}
+		else
+		{
+			PreviewRelativeTime = FMath::Min(PreviewRelativeTime, DurationSeconds);
+		}
+
+		ViewportWidget->RefreshView(PreviewRelativeTime, DeltaTime);
+		TimelineWidget->SetRelativeTime(PreviewRelativeTime);
 	}
 }
 
diff --git a/Engine/Plugins/FX/Niagara/Source/NiagaraEditor/Private/Widgets/SNiagaraBakerWidget.h b/Engine/Plugins/FX/Niagara/Source/NiagaraEditor/Private/Widgets/SNiagaraBakerWidget.h
--- a/Engine/Plugins/FX/Niagara/Source/NiagaraEditor/Private/Widgets/SNiagaraBakerWidget.h
+++ b/Engine/Plugins/FX/Niagara/Source/NiagaraEditor/Private/Widgets/SNiagaraBakerWidget.h
@@ -18,13 +18,21 @@ public:
 	SLATE_END_ARGS()
 
 	void Construct(const FArguments& InArgs);
+	~SNiagaraBakerWidget();
 
 	virtual void Tick(const FGeometry& AllottedGeometry, const double InCurrentTime, const float InDeltaTime) override;
 
 	FReply OnCapture();
 
-	TSharedRef<SWidget> MakeTextureSelectionWidget() const;
-	FText GetSelectedTextureAsText() const;
+	TSharedRef<SWidget> MakeCameraModeMenu();
+	TSharedRef<SWidget> MakeViewOptionsMenu();
+	TSharedRef<SWidget> MakeOutputMenu();
+
+	bool FindWarnings(TArray<FText>* OutWarnings) const;
+	bool HasWarnings() const { return FindWarnings(nullptr); }
+	TSharedRef<SWidget> MakeWarningsMenu();
+
+	void RefreshWidget();
 
 	class UNiagaraBakerSettings* GetBakerSettings() const;
 	const class UNiagaraBakerSettings* GetBakerGeneratedSettings() const;
@@ -39,11 +47,13 @@ public:
 
 private:
 	TWeakPtr<FNiagaraBakerViewModel>			WeakViewModel;
-	TSharedPtr<SNiagaraBakerViewport>		ViewportWidget;
-	TSharedPtr<SNiagaraBakerTimelineWidget>	TimelineWidget;
+	TSharedPtr<SNiagaraBakerViewport>			ViewportWidget;
+	TSharedPtr<SNiagaraBakerTimelineWidget>		TimelineWidget;
 	TSharedPtr<IDetailsView>					BakerSettingsDetails;
 	TSharedPtr<SWidget>							TransportControls;
 
 	bool										bIsPlaying = true;
 	float										PreviewRelativeTime = 0.0f;
+
+	FDelegateHandle								OnCurrentOutputIndexChangedHandle;
 };
diff --git a/Engine/Plugins/FX/Niagara/Source/NiagaraShader/Private/NiagaraBatchedElements.cpp b/Engine/Plugins/FX/Niagara/Source/NiagaraShader/Private/NiagaraBatchedElements.cpp
--- a/Engine/Plugins/FX/Niagara/Source/NiagaraShader/Private/NiagaraBatchedElements.cpp
+++ b/Engine/Plugins/FX/Niagara/Source/NiagaraShader/Private/NiagaraBatchedElements.cpp
@@ -91,6 +91,10 @@ class FNiagaraSimpleElementPS : public FGlobalShader
 	}
 
 	BEGIN_SHADER_PARAMETER_STRUCT(FParameters, )
+		SHADER_PARAMETER(FMatrix44f, InColorWeights)
+		SHADER_PARAMETER(float, InGamma)
+		SHADER_PARAMETER_TEXTURE(Texture2D, InTexture)
+		SHADER_PARAMETER_SAMPLER(SamplerState, InTextureSampler)
 	END_SHADER_PARAMETER_STRUCT()
 };
 
@@ -113,5 +117,33 @@ void FBatchedElementNiagaraInvertColorChannel::BindShaders(FRHICommandList& RHIC
 	VertexShader->SetParameters(RHICmdList, InTransform);
 
 	FNiagaraSimpleElementPS::FParameters PassParameters;
+	PassParameters.InColorWeights	= FMatrix44f::Identity;
+	PassParameters.InGamma			= 1.0f;
+	PassParameters.InTexture		= GWhiteTexture->TextureRHI;
+	PassParameters.InTextureSampler = TStaticSamplerState<SF_Point>::GetRHI();
+	SetShaderParameters(RHICmdList, PixelShader, PixelShader.GetPixelShader(), PassParameters);
+}
+
+void FBatchedElementNiagaraSimple::BindShaders(FRHICommandList& RHICmdList, FGraphicsPipelineStateInitializer& GraphicsPSOInit, ERHIFeatureLevel::Type InFeatureLevel, const FMatrix& InTransform, const float InGamma, const FMatrix& ColorWeights, const FTexture* Texture)
+{
+	TShaderMapRef<FSimpleElementVS> VertexShader(GetGlobalShaderMap(InFeatureLevel));
+	TShaderMapRef<FNiagaraSimpleElementPS> PixelShader(GetGlobalShaderMap(InFeatureLevel));
+
+	GraphicsPSOInit.BoundShaderState.VertexDeclarationRHI = GSimpleElementVertexDeclaration.VertexDeclarationRHI;
+	GraphicsPSOInit.BoundShaderState.VertexShaderRHI = VertexShader.GetVertexShader();
+	GraphicsPSOInit.BoundShaderState.PixelShaderRHI = PixelShader.GetPixelShader();
+	GraphicsPSOInit.PrimitiveType = PT_TriangleList;
+	GraphicsPSOInit.BlendState = bAlphaBlend ? TStaticBlendState<CW_RGBA, BO_Add, BF_One, BF_InverseSourceAlpha, BO_Add, BF_One, BF_InverseSourceAlpha>::GetRHI() : TStaticBlendState<>::GetRHI();
+
+	RHICmdList.ApplyCachedRenderTargets(GraphicsPSOInit);
+	SetGraphicsPipelineState(RHICmdList, GraphicsPSOInit, 0);
+
+	VertexShader->SetParameters(RHICmdList, InTransform);
+
+	FNiagaraSimpleElementPS::FParameters PassParameters;
+	PassParameters.InColorWeights = ColorTransform;
+	PassParameters.InGamma = 1.0f;
+	PassParameters.InTexture = Texture ? Texture->TextureRHI : GWhiteTexture->TextureRHI;
+	PassParameters.InTextureSampler = TStaticSamplerState<SF_Bilinear>::GetRHI();
 	SetShaderParameters(RHICmdList, PixelShader, PixelShader.GetPixelShader(), PassParameters);
 }
diff --git a/Engine/Plugins/FX/Niagara/Source/NiagaraShader/Public/NiagaraBatchedElements.h b/Engine/Plugins/FX/Niagara/Source/NiagaraShader/Public/NiagaraBatchedElements.h
--- a/Engine/Plugins/FX/Niagara/Source/NiagaraShader/Public/NiagaraBatchedElements.h
+++ b/Engine/Plugins/FX/Niagara/Source/NiagaraShader/Public/NiagaraBatchedElements.h
@@ -27,7 +27,7 @@ private:
 };
 
 /**
- * Batched element parameters for gathering attributes from different slices into a single color
+ * Batched element parameters for inverting a color channel
  */
 class NIAGARASHADER_API FBatchedElementNiagaraInvertColorChannel : public FBatchedElementParameters
 {
@@ -43,3 +43,23 @@ public:
 private:
 	uint32	ChannelMask;
 };
+
+/**
+ * Simple batched element using a 2d texture
+ */
+class NIAGARASHADER_API FBatchedElementNiagaraSimple : public FBatchedElementParameters
+{
+public:
+	FBatchedElementNiagaraSimple(FMatrix InColorTransform = FMatrix::Identity, bool InAlphaBlend = false)
+		: ColorTransform(InColorTransform)
+		, bAlphaBlend(InAlphaBlend)
+	{
+	}
+
+	/** Binds vertex and pixel shaders for this element */
+	virtual void BindShaders(FRHICommandList& RHICmdList, FGraphicsPipelineStateInitializer& GraphicsPSOInit, ERHIFeatureLevel::Type InFeatureLevel, const FMatrix& InTransform, const float InGamma, const FMatrix& ColorWeights, const FTexture* Texture) override;
+
+private:
+	FMatrix	ColorTransform;
+	bool	bAlphaBlend;
+};
