- Fix compile error due to explicit cast change

#rnx

diff --git a/Engine/Plugins/FX/Niagara/Source/NiagaraShader/Private/NiagaraBatchedElements.cpp b/Engine/Plugins/FX/Niagara/Source/NiagaraShader/Private/NiagaraBatchedElements.cpp
--- a/Engine/Plugins/FX/Niagara/Source/NiagaraShader/Private/NiagaraBatchedElements.cpp
+++ b/Engine/Plugins/FX/Niagara/Source/NiagaraShader/Private/NiagaraBatchedElements.cpp
@@ -141,7 +141,7 @@ void FBatchedElementNiagaraSimple::BindShaders(FRHICommandList& RHICmdList, FGra
 	VertexShader->SetParameters(RHICmdList, InTransform);
 
 	FNiagaraSimpleElementPS::FParameters PassParameters;
-	PassParameters.InColorWeights = ColorTransform;
+	PassParameters.InColorWeights = FMatrix44f(ColorTransform);
 	PassParameters.InGamma = 1.0f;
 	PassParameters.InTexture = Texture ? Texture->TextureRHI : GWhiteTexture->TextureRHI;
 	PassParameters.InTextureSampler = TStaticSamplerState<SF_Bilinear>::GetRHI();
