- Add OutputName to texture settings

#rnx

diff --git a/Engine/Plugins/FX/Niagara/Source/Niagara/Classes/NiagaraBakerSettings.h b/Engine/Plugins/FX/Niagara/Source/Niagara/Classes/NiagaraBakerSettings.h
--- a/Engine/Plugins/FX/Niagara/Source/Niagara/Classes/NiagaraBakerSettings.h
+++ b/Engine/Plugins/FX/Niagara/Source/Niagara/Classes/NiagaraBakerSettings.h
@@ -37,6 +37,10 @@ struct FNiagaraBakerTextureSettings
 {
 	GENERATED_BODY()
 	
+	/** Optional output name, if left empty a name will be auto generated using the index of the texture/ */
+	UPROPERTY(EditAnywhere, Category = "Texture")
+	FName OutputName;
+
 	/** Source visualization we should capture, i.e. Scene Color, World Normal, etc */
 	UPROPERTY(EditAnywhere, Category = "Texture")
 	FNiagaraBakerTextureSource SourceBinding;
