- Handle null asset

#rnx

diff --git a/Engine/Plugins/FX/Niagara/Source/Niagara/Private/MovieScene/Parameters/MovieSceneNiagaraParameterSectionTemplate.cpp b/Engine/Plugins/FX/Niagara/Source/Niagara/Private/MovieScene/Parameters/MovieSceneNiagaraParameterSectionTemplate.cpp
--- a/Engine/Plugins/FX/Niagara/Source/Niagara/Private/MovieScene/Parameters/MovieSceneNiagaraParameterSectionTemplate.cpp
+++ b/Engine/Plugins/FX/Niagara/Source/Niagara/Private/MovieScene/Parameters/MovieSceneNiagaraParameterSectionTemplate.cpp
@@ -113,7 +113,10 @@ void FMovieSceneNiagaraParameterSectionTemplate::Initialize(const FMovieSceneEva
 			const uint8* ParameterData = NiagaraComponent->GetOverrideParameters().GetParameterData(Parameter);
 			if (ParameterData == nullptr)
 			{
-				ParameterData = NiagaraComponent->GetAsset()->GetExposedParameters().GetParameterData(Parameter);
+				if (UNiagaraSystem* NiagaraSystem = NiagaraComponent->GetAsset())
+				{
+					ParameterData = NiagaraSystem->GetExposedParameters().GetParameterData(Parameter);
+				}
 			}
 
 			TArray<uint8> CurrentValueData;
