Fix to scalability mode not being enabled by default

diff --git a/Engine/Plugins/FX/Niagara/Source/NiagaraEditor/Private/Toolkits/NiagaraSystemToolkit.cpp b/Engine/Plugins/FX/Niagara/Source/NiagaraEditor/Private/Toolkits/NiagaraSystemToolkit.cpp
--- a/Engine/Plugins/FX/Niagara/Source/NiagaraEditor/Private/Toolkits/NiagaraSystemToolkit.cpp
+++ b/Engine/Plugins/FX/Niagara/Source/NiagaraEditor/Private/Toolkits/NiagaraSystemToolkit.cpp
@@ -76,7 +76,7 @@ DECLARE_CYCLE_STAT(TEXT("Niagara - SystemToolkit - OnApply"), STAT_NiagaraEditor
 const FName FNiagaraSystemToolkit::DefaultModeName(TEXT("Default"));
 const FName FNiagaraSystemToolkit::ScalabilityModeName(TEXT("Scalability"));
 
-int32 GNiagaraScalabilityModeEnabled = 0;
+int32 GNiagaraScalabilityModeEnabled = 1;
 
 // @TODO Remove cvar when scalability mode is going in properly
 static const TCHAR* NiagaraEnableScalabilityMode = TEXT("fx.Niagara.EnableScalabilityMode");
