Additional support for new experimental VM
-moves optimization pass to the async compilation pipeline (including moving it to it's own file, out of NiagaraSystem)
-splits up the NIAGARA_VM_EXP preprocessor define into VECTORVM_SUPPORTS_LEGACY & VECTORVM_SUPPORTS_EXPERIMENTAL allowing support of one, the other or both (the default)

diff --git a/Engine/Plugins/FX/Niagara/Source/Niagara/Classes/NiagaraDataSet.h b/Engine/Plugins/FX/Niagara/Source/Niagara/Classes/NiagaraDataSet.h
--- a/Engine/Plugins/FX/Niagara/Source/Niagara/Classes/NiagaraDataSet.h
+++ b/Engine/Plugins/FX/Niagara/Source/Niagara/Classes/NiagaraDataSet.h
@@ -351,6 +351,8 @@ public:
 
 	FORCEINLINE TArray<int32>& GetFreeIDTable() { return FreeIDsTable; }
 	FORCEINLINE TArray<int32>& GetSpawnedIDsTable() { return SpawnedIDsTable; }
+	FORCEINLINE const TArray<int32>& GetFreeIDTable() const { return FreeIDsTable; }
+	FORCEINLINE const TArray<int32>& GetSpawnedIDsTable() const { return SpawnedIDsTable; }
 	FORCEINLINE int32& GetNumFreeIDs() { return NumFreeIDs; }
 	FORCEINLINE int32& GetMaxUsedID() { return MaxUsedID; }
 	FORCEINLINE int32& GetIDAcquireTag() { return IDAcquireTag; }
diff --git a/Engine/Plugins/FX/Niagara/Source/Niagara/Classes/NiagaraScript.h b/Engine/Plugins/FX/Niagara/Source/Niagara/Classes/NiagaraScript.h
--- a/Engine/Plugins/FX/Niagara/Source/Niagara/Classes/NiagaraScript.h
+++ b/Engine/Plugins/FX/Niagara/Source/Niagara/Classes/NiagaraScript.h
@@ -491,6 +491,9 @@ public:
 	TArray<FNiagaraCompileEvent> LastCompileEvents;
 #endif
 
+	UPROPERTY()
+	TArray<uint8> ExperimentalContextData;
+
 	UPROPERTY()
 	uint32 bReadsSignificanceIndex : 1;
 
@@ -508,6 +511,11 @@ public:
 #if WITH_EDITORONLY_DATA
 	void BakeScriptLiterals(TArray<uint8>& OutLiterals) const;
 #endif
+
+#if VECTORVM_SUPPORTS_EXPERIMENTAL
+	bool SupportsExperimentalVM() const;
+	FVectorVMOptimizeContext BuildExperimentalContext() const;
+#endif //VECTORVM_SUPPORTS_EXPERIMENTAL
 };
 
 struct NIAGARA_API FNiagaraGraphCachedDataBase
@@ -649,10 +657,6 @@ class UNiagaraScript : public UNiagaraScriptBase
 	GENERATED_UCLASS_BODY()
 public:
 	UNiagaraScript();
-	~UNiagaraScript();
-#	ifdef NIAGARA_EXP_VM
-	FVectorVMOptimizeContext OptimizeContext;
-#	endif //NIAGARA_EXP_VM
 
 #if WITH_EDITORONLY_DATA
 	/** If true then this script asset uses active version control to track changes. */
diff --git a/Engine/Plugins/FX/Niagara/Source/Niagara/Classes/NiagaraScriptExecutionContext.h b/Engine/Plugins/FX/Niagara/Source/Niagara/Classes/NiagaraScriptExecutionContext.h
--- a/Engine/Plugins/FX/Niagara/Source/Niagara/Classes/NiagaraScriptExecutionContext.h
+++ b/Engine/Plugins/FX/Niagara/Source/Niagara/Classes/NiagaraScriptExecutionContext.h
@@ -125,8 +125,10 @@ struct FNiagaraDataSetExecutionInfo
 
 struct FScriptExecutionConstantBufferTable
 {
-	TArray<const uint8*, TInlineAllocator<12>> Buffers;
-	TArray<int32, TInlineAllocator<12>> BufferSizes;
+	static constexpr uint32 MaxBufferCount = 12;
+
+	TArray<const uint8*, TInlineAllocator<MaxBufferCount>> Buffers;
+	TArray<int32, TInlineAllocator<MaxBufferCount>> BufferSizes;
 
 	void Reset(int32 ResetSize)
 	{
@@ -182,7 +184,11 @@ struct FNiagaraScriptExecutionContextBase
 	void CreateStatScopeData();
 	TMap<TStatIdData const*, float> ReportStats();
 #endif
-	
+
+#if VECTORVM_SUPPORTS_EXPERIMENTAL && VECTORVM_SUPPORTS_LEGACY
+	bool bUsingExperimentalVM;
+#endif
+
 	FNiagaraScriptExecutionContextBase();
 	virtual ~FNiagaraScriptExecutionContextBase();
 
@@ -203,6 +209,14 @@ struct FNiagaraScriptExecutionContextBase
 	//Unused. These are only useful in the new SystemScript context.
 	virtual void BindSystemInstances(TArray<FNiagaraSystemInstance*>& InSystemInstances) {}
 	virtual bool GeneratePerInstanceDIFunctionTable(FNiagaraSystemInstance* Inst, TArray<struct FNiagaraPerInstanceDIFuncInfo>& OutFunctions) {return true;}
+
+private:
+#if VECTORVM_SUPPORTS_EXPERIMENTAL
+	bool ExecuteInternal_Experimental(uint32 NumInstances, const FScriptExecutionConstantBufferTable& ConstantBufferTable);
+#endif // VECTORVM_SUPPORTS_EXPERIMENTAL
+#if VECTORVM_SUPPORTS_LEGACY
+	bool ExecuteInternal_Legacy(uint32 NumInstances, const FScriptExecutionConstantBufferTable& ConstantBufferTable);
+#endif // VECTORVM_SUPPORTS_LEGACY
 };
 
 struct FNiagaraScriptExecutionContext : public FNiagaraScriptExecutionContextBase
@@ -246,9 +260,10 @@ protected:
 	*/
 	TArray<FNiagaraSystemInstance*>* SystemInstances;
 
-
+#if VECTORVM_SUPPORTS_LEGACY
 	/** Helper function that handles calling into per instance DI calls and massages the VM context appropriately. */
-	void PerInstanceFunctionHook(FVectorVMExternalFunctionContext& Context, int32 PerInstFunctionIndex, int32 UserPtrIndex);
+	void PerInstanceFunctionHook(FVectorVMExternalFunctionContext& ParentContext, FVectorVMExternalFunctionContextLegacy& Context, int32 PerInstFunctionIndex, int32 UserPtrIndex);
+#endif
 
 public:
 	FNiagaraSystemScriptExecutionContext(ENiagaraSystemSimulationScript InScriptType) : SystemInstances(nullptr) { ScriptType = InScriptType; }
diff --git a/Engine/Plugins/FX/Niagara/Source/Niagara/Classes/NiagaraSystem.h b/Engine/Plugins/FX/Niagara/Source/Niagara/Classes/NiagaraSystem.h
--- a/Engine/Plugins/FX/Niagara/Source/Niagara/Classes/NiagaraSystem.h
+++ b/Engine/Plugins/FX/Niagara/Source/Niagara/Classes/NiagaraSystem.h
@@ -25,6 +25,8 @@
 
 #include "NiagaraSystem.generated.h"
 
+class FNiagaraAsyncCompileTask;
+
 #if WITH_EDITORONLY_DATA
 class UNiagaraEditorDataBase;
 #endif
@@ -150,105 +152,6 @@ struct FNiagaraSystemCompiledData
 	TArray<FNiagaraParameterDataSetBindingCollection> UpdateInstanceEmitterBindings;
 };
 
-USTRUCT()
-struct FEmitterCompiledScriptPair
-{
-	GENERATED_USTRUCT_BODY()
-	
-	bool bResultsReady = false;
-	UNiagaraEmitter* Emitter = nullptr;
-	UNiagaraScript* CompiledScript = nullptr;
-	uint32 PendingJobID = INDEX_NONE; // this is the ID for any active shader compiler worker job
-	FNiagaraVMExecutableDataId CompileId;
-	TSharedPtr<FNiagaraVMExecutableData> CompileResults;
-};
-
-UENUM()
-enum class ENiagaraCompilationState : uint8
-{
-	CheckDDC,
-	Precompile,
-	StartCompileJob,
-	AwaitResult,
-	ProcessResult,
-	PutToDDC,
-	Finished,
-	Aborted
-};
-
-struct FNiagaraLazyPrecompileReference
-{
-	TSharedPtr<FNiagaraCompileRequestDataBase, ESPMode::ThreadSafe> GetPrecompileData(UNiagaraScript* ForScript);
-	TSharedPtr<FNiagaraCompileRequestDuplicateDataBase, ESPMode::ThreadSafe> GetPrecompileDuplicateData(UNiagaraEmitter* OwningEmitter, UNiagaraScript* TargetScript);
-	
-	UNiagaraSystem* System = nullptr;
-	TArray<UNiagaraScript*> Scripts;
-	TMap<UNiagaraScript*, int32> EmitterScriptIndex;
-	TArray<TObjectPtr<UObject>> CompilationRootObjects;
-
-private:
-	TSharedPtr<FNiagaraCompileRequestDataBase, ESPMode::ThreadSafe> SystemPrecompiledData;
-	TMap<UNiagaraScript*, TSharedPtr<FNiagaraCompileRequestDataBase, ESPMode::ThreadSafe>> EmitterMapping;
-	TArray<TSharedPtr<FNiagaraCompileRequestDuplicateDataBase, ESPMode::ThreadSafe>> PrecompileDuplicateDatas;
-};
-
-struct FNiagaraAsyncTaskSharedData
-{
-	FNiagaraAsyncTaskSharedData();
-	~FNiagaraAsyncTaskSharedData();
-	
-	
-};
-
-class FNiagaraAsyncCompileTask
-{
-#if WITH_EDITORONLY_DATA
-public:
-	FString DDCKey;
-	FString AssetPath;
-	FString UniqueEmitterName;
-	uint32 TaskHandle = 0;
-
-	double StartTaskTime = 0;
-	double DDCFetchTime = 0;
-	double StartCompileTime = 0;
-	bool bWaitForCompileJob = false;
-	bool bUsedShaderCompilerWorker = false;
-	bool bFetchedGCObjects = false;
-	UNiagaraSystem* OwningSystem;
-	FEmitterCompiledScriptPair ScriptPair;
-	TArray<FNiagaraVariable> EncounteredExposedVars;
-
-	ENiagaraCompilationState CurrentState;
-
-	TSharedPtr<FNiagaraVMExecutableData> ExeData;
-	TArray<uint8> DDCOutData;
-
-	// this data is shared between the ddc thread and the game thread that starts the compilation
-	TSharedPtr<FNiagaraLazyPrecompileReference, ESPMode::ThreadSafe> PrecompileReference;
-	TSharedPtr<FNiagaraCompileRequestDataBase, ESPMode::ThreadSafe> ComputedPrecompileData;
-	TSharedPtr<FNiagaraCompileRequestDuplicateDataBase, ESPMode::ThreadSafe> ComputedPrecompileDuplicateData;
-
-	FNiagaraAsyncCompileTask(UNiagaraSystem* InOwningSystem, FString InAssetPath, const FEmitterCompiledScriptPair& InScriptPair);
-
-	const FEmitterCompiledScriptPair& GetScriptPair() const { return ScriptPair; };
-	void WaitAndResolveResult();
-	void AbortTask();
-
-	void CheckDDCResult();
-	void PutToDDC();
-	void ProcessCurrentState();
-	void MoveToState(ENiagaraCompilationState NewState);
-	bool IsDone() const;
-
-	void PrecompileData();
-	void StartCompileJob();
-	bool AwaitResult();
-	void ProcessResult();
-
-#endif
-};
-
 USTRUCT()
 struct FNiagaraSystemCompileRequest
 {
@@ -551,6 +454,8 @@ public:
 	bool UsesCollection(const UNiagaraParameterCollection* Collection)const;
 
 	bool SupportsLargeWorldCoordinates() const { return bSupportLargeWorldCoordinates && bLwcEnabledSettingCached; }
+	FORCEINLINE bool ShouldDisableExperimentalVM() const { return bDisableExperimentalVM; }
+
 #if WITH_EDITORONLY_DATA
 	bool UsesEmitter(const UNiagaraEmitter* Emitter) const;
 	bool UsesScript(const UNiagaraScript* Script)const; 
@@ -598,7 +503,6 @@ protected:
 	UPROPERTY(EditAnywhere, AdvancedDisplay, Category = "Performance")
 	uint32 bIgnoreParticleReadsForAttributeTrim : 1;
 
-
 	/** When enable debug switches are disabled while editing the system. */
 	UPROPERTY(EditAnywhere, AdvancedDisplay, Category = "Performance", meta = (DisplayName="Disable Debug Switches During Edit"))
 	uint32 bDisableDebugSwitches : 1;
@@ -654,6 +558,10 @@ public:
 	UPROPERTY(EditAnywhere, BlueprintReadOnly, Category="Rendering", meta=(DisplayName="Default Render CustomDepth Pass", EditCondition="bOverrideRenderCustomDepth"))
 	uint8 bRenderCustomDepth : 1;
 
+	/** If true, disables experimental VM, if available */
+	UPROPERTY(EditAnywhere, AdvancedDisplay, Category = "Performance", meta = (DisplayName = "Disable Experimental VM"))
+	uint8 bDisableExperimentalVM : 1;
+
 	/**
 	When enabled this is the default value set on the component.
 	Mask used for stencil buffer writes.
diff --git a/Engine/Plugins/FX/Niagara/Source/Niagara/Niagara.Build.cs b/Engine/Plugins/FX/Niagara/Source/Niagara/Niagara.Build.cs
--- a/Engine/Plugins/FX/Niagara/Source/Niagara/Niagara.Build.cs
+++ b/Engine/Plugins/FX/Niagara/Source/Niagara/Niagara.Build.cs
@@ -6,7 +6,7 @@ public class Niagara : ModuleRules
 {
     public Niagara(ReadOnlyTargetRules Target) : base(Target)
     {
-		PrivateIncludePaths.Add("../../../../Shaders/Shared");
+        PrivateIncludePaths.Add("../../../../Shaders/Shared");
 
         PrivateDependencyModuleNames.AddRange(
             new string[] {
@@ -45,16 +45,15 @@ public class Niagara : ModuleRules
             }
         );
 
-		PrivateIncludePathModuleNames.AddRange(
-			new string[] {
-				"DerivedDataCache",
-			});
+        PrivateIncludePathModuleNames.AddRange(
+                new string[] {
+                        "DerivedDataCache",
+                });
 
         PrivateIncludePaths.AddRange(
             new string[] {
                 "Niagara/Private",
-            })
-        ;
+            });
 
         if (Target.bBuildEditor == true)
         {
@@ -67,11 +66,12 @@ public class Niagara : ModuleRules
 				"Slate"
             });
         }
+
         PublicDefinitions.AddRange(
             new string[]
             {
-                "NIAGARA_EXP_VM",
-            }
-        );
+                "VECTORVM_SUPPORTS_EXPERIMENTAL=1",
+                "VECTORVM_SUPPORTS_LEGACY=1"
+            });
     }
 }
diff --git a/Engine/Plugins/FX/Niagara/Source/Niagara/Private/NiagaraAsyncCompile.cpp b/Engine/Plugins/FX/Niagara/Source/Niagara/Private/NiagaraAsyncCompile.cpp
new file mode 100644
--- /dev/null
+++ b/Engine/Plugins/FX/Niagara/Source/Niagara/Private/NiagaraAsyncCompile.cpp
@@ -0,0 +1,436 @@
+// Copyright Epic Games, Inc. All Rights Reserved.
+
+
+#include "NiagaraAsyncCompile.h"
+
+#include "Modules/ModuleManager.h"
+#include "NiagaraEmitterHandle.h"
+#include "NiagaraModule.h"
+#include "NiagaraPrecompileContainer.h"
+#include "NiagaraScriptSourceBase.h"
+#include "NiagaraStats.h"
+#include "VectorVM.h"
+
+#define LOCTEXT_NAMESPACE "NiagaraAsyncCompile"
+
+#if WITH_EDITOR
+#include "DerivedDataCacheInterface.h"
+#endif
+
+DECLARE_CYCLE_STAT(TEXT("Niagara - System - Precompile"), STAT_Niagara_System_Precompile, STATGROUP_Niagara);
+DECLARE_CYCLE_STAT(TEXT("Niagara - System - CompileScriptTaskGT"), STAT_Niagara_System_CompileScriptTaskGT, STATGROUP_Niagara);
+
+static float GNiagaraCompileDDCWaitTimeout = 10.0f;
+static FAutoConsoleVariableRef CVarNiagaraCompileDDCWaitTimeout(
+	TEXT("fx.Niagara.CompileDDCWaitTimeout"),
+	GNiagaraCompileDDCWaitTimeout,
+	TEXT("During script compilation, how long do we wait for the ddc to answer in seconds before starting shader compilation?"),
+	ECVF_Default
+);
+
+//////////////////////////////////////////////////////////////////////////
+
+#if WITH_EDITORONLY_DATA
+TSharedPtr<FNiagaraCompileRequestDataBase, ESPMode::ThreadSafe> FNiagaraLazyPrecompileReference::GetPrecompileData(UNiagaraScript* ForScript)
+{
+	if (SystemPrecompiledData == nullptr)
+	{
+		UNiagaraPrecompileContainer* Container = NewObject<UNiagaraPrecompileContainer>(GetTransientPackage());
+		Container->System = System;
+		Container->Scripts = Scripts;
+		INiagaraModule& NiagaraModule = FModuleManager::Get().LoadModuleChecked<INiagaraModule>(TEXT("Niagara"));
+		SystemPrecompiledData = NiagaraModule.Precompile(Container, FGuid());
+
+		const TArray<FNiagaraEmitterHandle>& EmitterHandles = System->GetEmitterHandles();
+		for (int32 i = 0; i < EmitterHandles.Num(); i++)
+		{
+			const FNiagaraEmitterHandle& Handle = EmitterHandles[i];
+			if (Handle.GetInstance() && Handle.GetIsEnabled())
+			{
+				TSharedPtr<FNiagaraCompileRequestDataBase, ESPMode::ThreadSafe> EmitterPrecompiledData = SystemPrecompiledData->GetDependentRequest(i);
+				TArray<UNiagaraScript*> EmitterScripts;
+				Handle.GetInstance()->GetScripts(EmitterScripts, false, true);
+				check(EmitterScripts.Num() > 0);
+				for (UNiagaraScript* EmitterScript : EmitterScripts)
+				{
+					EmitterMapping.Add(EmitterScript, EmitterPrecompiledData);
+				}
+			}
+		}
+	}
+	if (SystemPrecompiledData.IsValid() == false)
+	{
+		UE_LOG(LogNiagara, Error, TEXT("Failed to precompile %s.  This is due to unexpected invalid or broken data.  Additional details should be in the log."), *System->GetPathName());
+		return nullptr;
+	}
+
+	if (ForScript->IsSystemSpawnScript() || ForScript->IsSystemUpdateScript())
+	{
+		return SystemPrecompiledData;
+	}
+	TSharedPtr<FNiagaraCompileRequestDataBase, ESPMode::ThreadSafe>* EmitterPrecompileData = EmitterMapping.Find(ForScript);
+	if (EmitterPrecompileData == nullptr)
+	{
+		UE_LOG(LogNiagara, Log, TEXT("Failed to precompile %s. This is probably because the system was modified between the original compile request and now."), *System->GetPathName());
+		return nullptr;
+	}
+	return *EmitterPrecompileData;
+}
+
+TSharedPtr<FNiagaraCompileRequestDuplicateDataBase, ESPMode::ThreadSafe> FNiagaraLazyPrecompileReference::GetPrecompileDuplicateData(UNiagaraEmitter* OwningEmitter, UNiagaraScript* TargetScript)
+{
+	TSharedPtr<FNiagaraCompileRequestDuplicateDataBase, ESPMode::ThreadSafe> PrecompileDuplicateData;
+	for (TSharedPtr<FNiagaraCompileRequestDuplicateDataBase, ESPMode::ThreadSafe> ExistingPrecompileDuplicateData : PrecompileDuplicateDatas)
+	{
+		if (ExistingPrecompileDuplicateData->IsDuplicateDataFor(System, OwningEmitter, TargetScript))
+		{
+			PrecompileDuplicateData = ExistingPrecompileDuplicateData;
+		}
+	}
+
+	if (PrecompileDuplicateData.IsValid() == false)
+	{
+		INiagaraModule& NiagaraModule = FModuleManager::Get().LoadModuleChecked<INiagaraModule>(TEXT("Niagara"));
+
+		PrecompileDuplicateData = NiagaraModule.PrecompileDuplicate(SystemPrecompiledData.Get(), System, OwningEmitter, TargetScript, FGuid());
+		PrecompileDuplicateDatas.Add(PrecompileDuplicateData);
+
+		PrecompileDuplicateData->GetDuplicatedObjects(CompilationRootObjects);
+		for (int32 i = 0; i < PrecompileDuplicateData->GetDependentRequestCount(); i++)
+		{
+			PrecompileDuplicateData->GetDependentRequest(i)->GetDuplicatedObjects(CompilationRootObjects);
+		}
+	}
+
+	return PrecompileDuplicateData;
+}
+
+FNiagaraAsyncCompileTask::FNiagaraAsyncCompileTask(UNiagaraSystem* InOwningSystem, FString InAssetPath, const FEmitterCompiledScriptPair& InScriptPair)
+{
+	OwningSystem = InOwningSystem;
+	AssetPath = InAssetPath;
+	ScriptPair = InScriptPair;
+
+	CurrentState = ENiagaraCompilationState::CheckDDC;
+}
+
+void FNiagaraAsyncCompileTask::ProcessCurrentState()
+{
+	SCOPE_CYCLE_COUNTER(STAT_Niagara_System_CompileScriptTaskGT);
+	if (IsDone())
+	{
+		return;
+	}
+
+	if (CurrentState == ENiagaraCompilationState::CheckDDC)
+	{
+		// check if the DDC data is ready
+		CheckDDCResult();
+	}
+	else if (CurrentState == ENiagaraCompilationState::Precompile)
+	{
+		// gather precompile data
+		StartCompileTime = FPlatformTime::Seconds();
+		PrecompileData();
+		MoveToState(ENiagaraCompilationState::StartCompileJob);
+	}
+	else if (CurrentState == ENiagaraCompilationState::StartCompileJob)
+	{
+		// start the async compile job
+		StartCompileJob();
+		MoveToState(ENiagaraCompilationState::AwaitResult);
+	}
+	else if (CurrentState == ENiagaraCompilationState::AwaitResult)
+	{
+		if (AwaitResult())
+		{
+			MoveToState(ENiagaraCompilationState::OptimizeByteCode);
+		}
+	}
+	else if (CurrentState == ENiagaraCompilationState::OptimizeByteCode)
+	{
+		OptimizeByteCode();
+		MoveToState(ENiagaraCompilationState::ProcessResult);
+	}
+	else if (CurrentState == ENiagaraCompilationState::ProcessResult)
+	{
+		// save the result from the compile job
+		ProcessResult();
+		MoveToState(ENiagaraCompilationState::PutToDDC);
+	}
+	else if (CurrentState == ENiagaraCompilationState::PutToDDC)
+	{
+		// put the result from the compile job into the ddc
+		PutToDDC();
+		MoveToState(ENiagaraCompilationState::Finished);
+	}
+	else
+	{
+		check(false);
+	}
+}
+
+void FNiagaraAsyncCompileTask::MoveToState(ENiagaraCompilationState NewState)
+{
+	if (CurrentState == ENiagaraCompilationState::Finished || CurrentState == ENiagaraCompilationState::Aborted)
+	{
+		// we are already in an end state, so nothing to do
+		return;
+	}
+
+	// check state machine integrity
+	check(NewState != ENiagaraCompilationState::CheckDDC);
+	if (NewState == ENiagaraCompilationState::Precompile)
+	{
+		check(CurrentState == ENiagaraCompilationState::CheckDDC);
+	}
+	if (NewState == ENiagaraCompilationState::StartCompileJob)
+	{
+		check(CurrentState == ENiagaraCompilationState::Precompile);
+	}
+	if (NewState == ENiagaraCompilationState::AwaitResult)
+	{
+		check(CurrentState == ENiagaraCompilationState::StartCompileJob);
+	}
+	if (NewState == ENiagaraCompilationState::OptimizeByteCode)
+	{
+		check(CurrentState == ENiagaraCompilationState::AwaitResult);
+	}
+	if (NewState == ENiagaraCompilationState::ProcessResult)
+	{
+		check(CurrentState == ENiagaraCompilationState::OptimizeByteCode);
+	}
+	if (NewState == ENiagaraCompilationState::PutToDDC)
+	{
+		check(CurrentState == ENiagaraCompilationState::ProcessResult);
+	}
+	if (NewState == ENiagaraCompilationState::Finished)
+	{
+		check(CurrentState == ENiagaraCompilationState::CheckDDC || CurrentState == ENiagaraCompilationState::PutToDDC);
+	}
+
+	UE_LOG(LogNiagara, Verbose, TEXT("Changing state %i -> %i for for %s!"), CurrentState, NewState, *AssetPath);
+	CurrentState = NewState;
+}
+
+bool FNiagaraAsyncCompileTask::IsDone() const
+{
+	return CurrentState == ENiagaraCompilationState::Finished || CurrentState == ENiagaraCompilationState::Aborted;
+}
+
+void LogVMID(const FNiagaraVMExecutableDataId& ID, FString Name)
+{
+	// some debug logging to find out why the ddc key changed
+	UE_LOG(LogNiagara, Display, TEXT("FNiagaraVMExecutableDataId %s:"), *Name);
+	UE_LOG(LogNiagara, Display, TEXT("  ScriptUsageType: %i"), ID.ScriptUsageType);
+	UE_LOG(LogNiagara, Display, TEXT("  ScriptUsageTypeID: %s"), *ID.ScriptUsageTypeID.ToString());
+	UE_LOG(LogNiagara, Display, TEXT("  CompilerVersionID: %s"), *ID.CompilerVersionID.ToString());
+	UE_LOG(LogNiagara, Display, TEXT("  BaseScriptCompileHash: %s"), *ID.BaseScriptCompileHash.ToString());
+	UE_LOG(LogNiagara, Display, TEXT("  bUsesRapidIterationParams: %i"), ID.bUsesRapidIterationParams);
+	UE_LOG(LogNiagara, Display, TEXT("  AdditionalDefines:"));
+	for (int32 Idx = 0; Idx < ID.AdditionalDefines.Num(); Idx++)
+	{
+		UE_LOG(LogNiagara, Display, TEXT("     %s"), *ID.AdditionalDefines[Idx]);
+	}
+	UE_LOG(LogNiagara, Display, TEXT("  AdditionalVariables:"));
+	for (int32 Idx = 0; Idx < ID.AdditionalVariables.Num(); Idx++)
+	{
+		UE_LOG(LogNiagara, Display, TEXT("     %s / %s"), *ID.AdditionalVariables[Idx].GetName().ToString(), *ID.AdditionalVariables[Idx].GetType().GetName());
+	}
+	UE_LOG(LogNiagara, Display, TEXT("  ReferencedCompileHashes:"));
+	for (int32 HashIndex = 0; HashIndex < ID.ReferencedCompileHashes.Num(); HashIndex++)
+	{
+		UE_LOG(LogNiagara, Display, TEXT("     %s"), *ID.ReferencedCompileHashes[HashIndex].ToString());
+	}
+}
+
+void FNiagaraAsyncCompileTask::PrecompileData()
+{
+	SCOPE_CYCLE_COUNTER(STAT_Niagara_System_Precompile);
+	UE_LOG(LogNiagara, Verbose, TEXT("Getting precompile data for %s!"), *AssetPath);
+
+	// we compute the vmID here again to check if the script changed between the start of the task and now. If it was changed we need to update the ddc key,
+	// because otherwise we would save the compiled data under the wrong ddc key
+	FNiagaraVMExecutableDataId NewID;
+	ScriptPair.CompiledScript->ComputeVMCompilationId(NewID, FGuid());
+	FString CurrentDDCKey = UNiagaraScript::BuildNiagaraDDCKeyString(NewID);
+	if (DDCKey != CurrentDDCKey && LogNiagara.GetVerbosity() >= ELogVerbosity::Verbose)
+	{
+		UE_LOG(LogNiagara, Warning, TEXT("Compile ID for %s changed during the compilation task."), *AssetPath);
+		FNiagaraVMExecutableDataId OldID = ScriptPair.CompileId;
+		LogVMID(OldID, "OldID");
+		LogVMID(NewID, "NewID");
+	}
+	DDCKey = CurrentDDCKey;
+	
+	ComputedPrecompileData = PrecompileReference->GetPrecompileData(ScriptPair.CompiledScript);
+	TSharedPtr<FNiagaraCompileRequestDuplicateDataBase, ESPMode::ThreadSafe> SystemPrecompileDuplicateData = PrecompileReference->GetPrecompileDuplicateData(ScriptPair.Emitter, ScriptPair.CompiledScript);
+
+	if (ComputedPrecompileData == nullptr || SystemPrecompileDuplicateData == nullptr)
+	{
+		AbortTask();
+		return;
+	}
+
+	// Grab the list of user variables that were actually encountered so that we can add to them later.
+	ComputedPrecompileData->GatherPreCompiledVariables(TEXT("User"), EncounteredExposedVars);
+	for(int32 i = 0; i < ComputedPrecompileData->GetDependentRequestCount(); i++)
+	{
+		ComputedPrecompileData->GetDependentRequest(i)->GatherPreCompiledVariables(TEXT("User"), EncounteredExposedVars);
+	}
+
+	// Assign the correct duplicate data for this request.
+	if(ScriptPair.CompiledScript->GetUsage() == ENiagaraScriptUsage::SystemSpawnScript || ScriptPair.CompiledScript->GetUsage() == ENiagaraScriptUsage::SystemUpdateScript)
+	{
+		ComputedPrecompileDuplicateData = SystemPrecompileDuplicateData;
+	}
+	else
+	{
+		UNiagaraEmitter* OwningEmitter = ScriptPair.Emitter;
+		int32 EmitterIndex = OwningSystem->GetEmitterHandles().IndexOfByPredicate([OwningEmitter](const FNiagaraEmitterHandle& EmitterHandle) { return EmitterHandle.GetInstance() == OwningEmitter; });
+		if(EmitterIndex != INDEX_NONE)
+		{
+			ComputedPrecompileDuplicateData = SystemPrecompileDuplicateData->GetDependentRequest(EmitterIndex);
+		}
+	}
+}
+
+void FNiagaraAsyncCompileTask::StartCompileJob()
+{
+	UE_LOG(LogNiagara, Verbose, TEXT("Starting compilation for %s!"), *AssetPath);
+	
+	// Fire off the compile job
+	if (ScriptPair.CompiledScript->RequestExternallyManagedAsyncCompile(ComputedPrecompileData, ComputedPrecompileDuplicateData, ScriptPair.CompileId, ScriptPair.PendingJobID))
+	{
+		UE_LOG(LogNiagara, Verbose, TEXT("Designated compilation ID %i for compilation of %s!"), ScriptPair.PendingJobID, *AssetPath);
+		bUsedShaderCompilerWorker = true;
+	}
+	else
+	{
+		UE_LOG(LogNiagara, Warning, TEXT("For some reason we are reporting that %s is in sync even though AreScriptAndSourceSynchronized returned false!"), *ScriptPair.CompiledScript->GetPathName())
+		AbortTask();
+	}
+}
+
+bool FNiagaraAsyncCompileTask::AwaitResult()
+{
+	UE_LOG(LogNiagara, Verbose, TEXT("Waiting on compilation to finish for %s!"), *AssetPath);
+
+	INiagaraModule& NiagaraModule = FModuleManager::Get().LoadModuleChecked<INiagaraModule>(TEXT("Niagara"));
+	ExeData = NiagaraModule.GetCompileJobResult(ScriptPair.PendingJobID, bWaitForCompileJob);
+	return ExeData.IsValid();
+}
+
+void FNiagaraAsyncCompileTask::OptimizeByteCode()
+{
+	check(ExeData);
+
+#if VECTORVM_SUPPORTS_EXPERIMENTAL
+	if (const uint8* ByteCode = ExeData->ByteCode.GetDataPtr())
+	{
+		//this is just necessary because VectorVM doesn't know about FVMExternalFunctionBindingInfo
+		const int32 NumExtFns = ExeData->CalledVMExternalFunctions.Num();
+
+		TArray<FVectorVMExtFunctionData> ExtFnTable;
+		ExtFnTable.SetNumZeroed(NumExtFns);
+		for (int i = 0; i < NumExtFns; ++i)
+		{
+			ExtFnTable[i].NumInputs = ExeData->CalledVMExternalFunctions[i].GetNumInputs();
+			ExtFnTable[i].NumOutputs = ExeData->CalledVMExternalFunctions[i].GetNumOutputs();
+		}
+
+		FVectorVMOptimizeContext OptimizeContext;
+		FMemory::Memzero(&OptimizeContext, sizeof(OptimizeContext));
+		OptimizeVectorVMScript(ByteCode, ExeData->ByteCode.GetLength(), ExtFnTable.GetData(), ExtFnTable.Num(), &OptimizeContext, VVMOptFlag_OmitStats);
+
+		// freeze the OptimizeContext
+		FreezeVectorVMOptimizeContext(OptimizeContext, ExeData->ExperimentalContextData);
+
+		FreeVectorVMOptimizeContext(&OptimizeContext);
+	}
+#endif // VECTORVM_SUPPORTS_EXPERIMENTAL
+}
+
+void FNiagaraAsyncCompileTask::ProcessResult()
+{
+	check(ExeData);
+	UE_LOG(LogNiagara, Verbose, TEXT("Processing compilation results for %s!"), *AssetPath);
+
+	// convert results to be saved in the ddc
+	ScriptPair.CompileResults = ExeData;
+	ScriptPair.bResultsReady = true;
+	ExeData->CompileTime = FPlatformTime::Seconds() - StartCompileTime;
+	ScriptPair.CompiledScript->ExecToBinaryData(ScriptPair.CompiledScript, DDCOutData, *ExeData);
+	UE_LOG(LogNiagara, Verbose, TEXT("Got %i bytes in ddc data for %s"), DDCOutData.Num(), *AssetPath);
+}
+
+void FNiagaraAsyncCompileTask::WaitAndResolveResult()
+{
+	check(IsInGameThread());
+	
+	bWaitForCompileJob = true;
+	while (!IsDone())
+	{
+		ProcessCurrentState();
+	}
+}
+
+void FNiagaraAsyncCompileTask::AbortTask()
+{
+	MoveToState(ENiagaraCompilationState::Aborted);
+}
+
+void FNiagaraAsyncCompileTask::CheckDDCResult()
+{
+	FDerivedDataCacheInterface* DDC = GetDerivedDataCache();
+	if (bWaitForCompileJob)
+	{
+		DDC->WaitAsynchronousCompletion(TaskHandle);
+	}
+	if (DDC->PollAsynchronousCompletion(TaskHandle))
+	{
+		bool bHasResults = DDC->GetAsynchronousResults(TaskHandle, DDCOutData);
+		if (bHasResults && DDCOutData.Num() > 0)
+		{
+			ExeData = MakeShared<FNiagaraVMExecutableData>();
+			if (ScriptPair.CompiledScript->BinaryToExecData(ScriptPair.CompiledScript, DDCOutData, *ExeData))
+			{
+				ScriptPair.CompileResults = ExeData;
+				ScriptPair.bResultsReady = true;
+				MoveToState(ENiagaraCompilationState::Finished);
+				UE_LOG(LogNiagara, Verbose, TEXT("Compilation data for %s could be pulled from the ddc."), *AssetPath);
+			}
+			else
+			{
+				UE_LOG(LogNiagara, Warning, TEXT("Unable to create exec data from ddc data for script %s, going to recompile it from scratch. DDC might be corrupted or there is a problem with script serialization."), *AssetPath);
+				ExeData.Reset();
+				MoveToState(ENiagaraCompilationState::Precompile);
+			}
+		}
+		else
+		{
+			MoveToState(ENiagaraCompilationState::Precompile);
+			UE_LOG(LogNiagara, Verbose, TEXT("No compilation data for %s found in the ddc."), *AssetPath);
+		}
+		DDCFetchTime = FPlatformTime::Seconds() - StartTaskTime;
+		UE_LOG(LogNiagara, Verbose, TEXT("Compile task for %s took %f seconds to fetch data from ddc."), *AssetPath, DDCFetchTime);
+	}
+	else if ((FPlatformTime::Seconds() - StartTaskTime) > GNiagaraCompileDDCWaitTimeout)
+	{
+		MoveToState(ENiagaraCompilationState::Precompile);
+		UE_LOG(LogNiagara, Log, TEXT("DDC timed out after %i seconds, starting compilation for %s."), (int)GNiagaraCompileDDCWaitTimeout, *AssetPath);
+	}
+}
+
+void FNiagaraAsyncCompileTask::PutToDDC()
+{
+	if (DDCOutData.Num() > 0)
+	{
+		UE_LOG(LogNiagara, Verbose, TEXT("Writing data for %s to the ddc"), *AssetPath);
+		GetDerivedDataCache()->Put(*DDCKey, DDCOutData, *AssetPath, true);
+	}
+}
+
+#endif
+
+#undef LOCTEXT_NAMESPACE // NiagaraAsyncCompile
diff --git a/Engine/Plugins/FX/Niagara/Source/Niagara/Private/NiagaraAsyncCompile.h b/Engine/Plugins/FX/Niagara/Source/Niagara/Private/NiagaraAsyncCompile.h
new file mode 100644
--- /dev/null
+++ b/Engine/Plugins/FX/Niagara/Source/Niagara/Private/NiagaraAsyncCompile.h
@@ -0,0 +1,105 @@
+// Copyright Epic Games, Inc. All Rights Reserved.
+
+#pragma once
+
+#include "CoreMinimal.h"
+
+#include "NiagaraScript.h"
+
+#include "NiagaraAsyncCompile.generated.h"
+
+class UNiagaraEmitter;
+class UNiagaraSystem;
+
+USTRUCT()
+struct FEmitterCompiledScriptPair
+{
+	GENERATED_USTRUCT_BODY()
+	
+	bool bResultsReady = false;
+	UNiagaraEmitter* Emitter = nullptr;
+	UNiagaraScript* CompiledScript = nullptr;
+	uint32 PendingJobID = INDEX_NONE; // this is the ID for any active shader compiler worker job
+	FNiagaraVMExecutableDataId CompileId;
+	TSharedPtr<FNiagaraVMExecutableData> CompileResults;
+};
+
+UENUM()
+enum class ENiagaraCompilationState : uint8
+{
+	CheckDDC,
+	Precompile,
+	StartCompileJob,
+	AwaitResult,
+	OptimizeByteCode,
+	ProcessResult,
+	PutToDDC,
+	Finished,
+	Aborted
+};
+
+struct FNiagaraLazyPrecompileReference
+{
+	TSharedPtr<FNiagaraCompileRequestDataBase, ESPMode::ThreadSafe> GetPrecompileData(UNiagaraScript* ForScript);
+	TSharedPtr<FNiagaraCompileRequestDuplicateDataBase, ESPMode::ThreadSafe> GetPrecompileDuplicateData(UNiagaraEmitter* OwningEmitter, UNiagaraScript* TargetScript);
+	
+	UNiagaraSystem* System = nullptr;
+	TArray<UNiagaraScript*> Scripts;
+	TMap<UNiagaraScript*, int32> EmitterScriptIndex;
+	TArray<TObjectPtr<UObject>> CompilationRootObjects;
+
+private:
+	TSharedPtr<FNiagaraCompileRequestDataBase, ESPMode::ThreadSafe> SystemPrecompiledData;
+	TMap<UNiagaraScript*, TSharedPtr<FNiagaraCompileRequestDataBase, ESPMode::ThreadSafe>> EmitterMapping;
+	TArray<TSharedPtr<FNiagaraCompileRequestDuplicateDataBase, ESPMode::ThreadSafe>> PrecompileDuplicateDatas;
+};
+
+class FNiagaraAsyncCompileTask
+{
+#if WITH_EDITORONLY_DATA
+public:
+	FString DDCKey;
+	FString AssetPath;
+	FString UniqueEmitterName;
+	uint32 TaskHandle = 0;
+
+	double StartTaskTime = 0;
+	double DDCFetchTime = 0;
+	double StartCompileTime = 0;
+	bool bWaitForCompileJob = false;
+	bool bUsedShaderCompilerWorker = false;
+	bool bFetchedGCObjects = false;
+	UNiagaraSystem* OwningSystem;
+	FEmitterCompiledScriptPair ScriptPair;
+	TArray<FNiagaraVariable> EncounteredExposedVars;
+
+	ENiagaraCompilationState CurrentState;
+
+	TSharedPtr<FNiagaraVMExecutableData> ExeData;
+	TArray<uint8> DDCOutData;
+
+	// this data is shared between the ddc thread and the game thread that starts the compilation
+	TSharedPtr<FNiagaraLazyPrecompileReference, ESPMode::ThreadSafe> PrecompileReference;
+	TSharedPtr<FNiagaraCompileRequestDataBase, ESPMode::ThreadSafe> ComputedPrecompileData;
+	TSharedPtr<FNiagaraCompileRequestDuplicateDataBase, ESPMode::ThreadSafe> ComputedPrecompileDuplicateData;
+
+	FNiagaraAsyncCompileTask(UNiagaraSystem* InOwningSystem, FString InAssetPath, const FEmitterCompiledScriptPair& InScriptPair);
+
+	const FEmitterCompiledScriptPair& GetScriptPair() const { return ScriptPair; };
+	void WaitAndResolveResult();
+	void AbortTask();
+
+	void CheckDDCResult();
+	void PutToDDC();
+	void ProcessCurrentState();
+	void MoveToState(ENiagaraCompilationState NewState);
+	bool IsDone() const;
+
+	void PrecompileData();
+	void StartCompileJob();
+	bool AwaitResult();
+	void ProcessResult();
+	void OptimizeByteCode();
+
+#endif
+};
diff --git a/Engine/Plugins/FX/Niagara/Source/Niagara/Private/NiagaraDataInterfaceColorCurve.cpp b/Engine/Plugins/FX/Niagara/Source/Niagara/Private/NiagaraDataInterfaceColorCurve.cpp
--- a/Engine/Plugins/FX/Niagara/Source/Niagara/Private/NiagaraDataInterfaceColorCurve.cpp
+++ b/Engine/Plugins/FX/Niagara/Source/Niagara/Private/NiagaraDataInterfaceColorCurve.cpp
@@ -238,11 +238,13 @@ FORCEINLINE_DEBUGGABLE FLinearColor UNiagaraDataInterfaceColorCurve::SampleCurve
 }
 
 
-#if defined(NIAGARA_EXP_VM) && PLATFORM_ENABLE_VECTORINTRINSICS && !PLATFORM_ENABLE_VECTORINTRINSICS_NEON
+#if VECTORVM_SUPPORTS_EXPERIMENTAL && !VECTORVM_SUPPORTS_LEGACY && PLATFORM_ENABLE_VECTORINTRINSICS && !PLATFORM_ENABLE_VECTORINTRINSICS_NEON
 template<>
 void UNiagaraDataInterfaceColorCurve::SampleCurve<TIntegralConstant<bool, true>>(FVectorVMExternalFunctionContext& Context)
 {
-	if (Context.NumInstances == 1) //could be a per-instance function call, in which can we can't write 4-wide so just use the old method
+	const int32 NumInstances = Context.GetNumInstances();
+
+	if (NumInstances == 1) //could be a per-instance function call, in which can we can't write 4-wide so just use the old method
 	{ 
 		VectorVM::FExternalFuncInputHandler<float> XParam(Context);
 		VectorVM::FExternalFuncRegisterHandler<float> SamplePtrR(Context);
@@ -250,7 +252,7 @@ void UNiagaraDataInterfaceColorCurve::SampleCurve<TIntegralConstant<bool, true>>
 		VectorVM::FExternalFuncRegisterHandler<float> SamplePtrB(Context);
 		VectorVM::FExternalFuncRegisterHandler<float> SamplePtrA(Context);
 
-		for (int32 i = 0; i < Context.GetNumInstances(); ++i)
+		for (int32 i = 0; i < NumInstances; ++i)
 		{
 			float X = XParam.GetAndAdvance();
 			FLinearColor C = SampleCurveInternal<TIntegralConstant<bool, true>>(X);
@@ -262,6 +264,8 @@ void UNiagaraDataInterfaceColorCurve::SampleCurve<TIntegralConstant<bool, true>>
 	}
 	else
 	{
+		const int32 NumLoops = Context.GetNumLoops<4>();
+
 		float *LUT = ShaderLUT.GetData();
 		VectorRegister4f LutNumSamplesMinusOne4 = VectorSetFloat1(LUTNumSamplesMinusOne);
 		VectorRegister4f LutMinTime4            = VectorSetFloat1(LUTMinTime);
@@ -276,7 +280,7 @@ void UNiagaraDataInterfaceColorCurve::SampleCurve<TIntegralConstant<bool, true>>
 		int IdxA[4];
 		int IdxB[4];
 
-		for (int i = 0; i < Context.NumLoops; ++i)
+		for (int i = 0; i < NumLoops; ++i)
 		{
 			VectorRegister4f NormalizedX = VectorMultiply(VectorSubtract(XParam[i & Context.RegInc[0]], LutMinTime4), LutInvTimeRange4);
 			VectorRegister4f RemappedX   = VectorClamp(VectorMultiply(NormalizedX, LutNumSamplesMinusOne4), VectorZeroFloat(), LutNumSamplesMinusOne4);
@@ -313,7 +317,7 @@ void UNiagaraDataInterfaceColorCurve::SampleCurve<TIntegralConstant<bool, true>>
 		}
 	}
 }
-#endif //NIAGARA_EXP_VM
+#endif // VECTORVM_SUPPORTS_EXPERIMENTAL
 
 template<typename UseLUT>
 void UNiagaraDataInterfaceColorCurve::SampleCurve(FVectorVMExternalFunctionContext& Context)
diff --git a/Engine/Plugins/FX/Niagara/Source/Niagara/Private/NiagaraDataInterfaceCurve.cpp b/Engine/Plugins/FX/Niagara/Source/Niagara/Private/NiagaraDataInterfaceCurve.cpp
--- a/Engine/Plugins/FX/Niagara/Source/Niagara/Private/NiagaraDataInterfaceCurve.cpp
+++ b/Engine/Plugins/FX/Niagara/Source/Niagara/Private/NiagaraDataInterfaceCurve.cpp
@@ -195,16 +195,18 @@ FORCEINLINE_DEBUGGABLE float UNiagaraDataInterfaceCurve::SampleCurveInternal<TIn
 	return Curve.Eval(X);
 }
 
-#ifdef NIAGARA_EXP_VM
+#if VECTORVM_SUPPORTS_EXPERIMENTAL && !VECTORVM_SUPPORTS_LEGACY
 template<>
 void UNiagaraDataInterfaceCurve::SampleCurve<TIntegralConstant<bool, true>>(FVectorVMExternalFunctionContext& Context)
 {
-	if (Context.NumInstances == 1)
+	const int32 NumInstances = Context.GetNumInstances();
+
+	if (NumInstances == 1)
 	{
 		VectorVM::FExternalFuncInputHandler<float> XParam(Context);
 		VectorVM::FExternalFuncRegisterHandler<float> OutSample(Context);
 
-		for (int32 i = 0; i < Context.GetNumInstances(); ++i)
+		for (int32 i = 0; i < NumInstances; ++i)
 		{
 			*OutSample.GetDest() = SampleCurveInternal<TIntegralConstant<bool, true>>(XParam.Get());
 			XParam.Advance();
@@ -244,7 +246,7 @@ void UNiagaraDataInterfaceCurve::SampleCurve<TIntegralConstant<bool, true>>(FVec
 		}
 	}
 }
-#endif
+#endif // VECTORVM_SUPPORTS_EXPERIMENTAL
 
 template<typename UseLUT>
 void UNiagaraDataInterfaceCurve::SampleCurve(FVectorVMExternalFunctionContext& Context)
diff --git a/Engine/Plugins/FX/Niagara/Source/Niagara/Private/NiagaraDataInterfaceParticleRead.cpp b/Engine/Plugins/FX/Niagara/Source/Niagara/Private/NiagaraDataInterfaceParticleRead.cpp
--- a/Engine/Plugins/FX/Niagara/Source/Niagara/Private/NiagaraDataInterfaceParticleRead.cpp
+++ b/Engine/Plugins/FX/Niagara/Source/Niagara/Private/NiagaraDataInterfaceParticleRead.cpp
@@ -1451,6 +1451,54 @@ void UNiagaraDataInterfaceParticleRead::GetNumSpawnedParticles(FVectorVMExternal
 	}
 }
 
+static const int32* GetSpawnedIDs(FVectorVMExternalFunctionContext& Context, FNiagaraEmitterInstance* EmitterInstance, int32& NumSpawned)
+{
+	if (!EmitterInstance)
+	{
+		NumSpawned = 0;
+		return nullptr;
+	}
+
+	const FNiagaraDataSet& DataSet = EmitterInstance->GetData();
+
+#if VECTORVM_SUPPORTS_EXPERIMENTAL && VECTORVM_SUPPORTS_LEGACY
+	if (Context.UsingExperimentalVM)
+	{
+		const int32* SpawnedIDs = DataSet.GetFreeIDTable().GetData() + DataSet.GetFreeIDTable().Max();
+		NumSpawned = DataSet.NumSpawnedIDs;
+		if (SpawnedIDs)
+		{
+			SpawnedIDs -= NumSpawned;
+		}
+
+		return SpawnedIDs;
+	}
+	else
+	{
+		const TArray<int32>& SpawnedIDsTable = DataSet.GetSpawnedIDsTable();
+		NumSpawned = SpawnedIDsTable.Num();
+		return SpawnedIDsTable.GetData();
+	}
+#elif VECTORVM_SUPPORTS_EXPERIMENTAL
+
+	const int32* SpawnedIDs = DataSet.GetFreeIDTable().GetData() + DataSet.GetFreeIDTable().Max();
+	NumSpawned = DataSet.NumSpawnedIDs;
+	if (SpawnedIDs)
+	{
+		SpawnedIDs -= NumSpawned;
+	}
+
+	return SpawnedIDs;
+
+#elif VECTORVM_SUPPORTS_LEGACY
+
+	const TArray<int32>& SpawnedIDsTable = DataSet.GetSpawnedIDsTable();
+	NumSpawned = SpawnedIDsTable.Num();
+	return SpawnedIDsTable.GetData();
+
+#endif
+}
+
 void UNiagaraDataInterfaceParticleRead::GetSpawnedIDAtIndex(FVectorVMExternalFunctionContext& Context)
 {
 	VectorVM::FUserPtrHandler<FNDIParticleRead_InstanceData> InstData(Context);
@@ -1461,16 +1509,9 @@ void UNiagaraDataInterfaceParticleRead::GetSpawnedIDAtIndex(FVectorVMExternalFun
 	VectorVM::FExternalFuncRegisterHandler<int32> OutIDAcquireTag(Context);
 
 	FNiagaraEmitterInstance* EmitterInstance = InstData.Get()->EmitterInstance;
-#ifdef NIAGARA_EXP_VM
-	int32 *SpawnedIDsTable = EmitterInstance ? EmitterInstance->GetData().GetFreeIDTable().GetData() + EmitterInstance->GetData().GetFreeIDTable().Max() : nullptr;
-	int32 NumSpawned = EmitterInstance ? EmitterInstance->GetData().NumSpawnedIDs : 0;
-	if (SpawnedIDsTable && NumSpawned) {
-		SpawnedIDsTable -= NumSpawned; //spawned ids are at the end of the FreeID table
-	}
-#else
-	const TArray<int32>* SpawnedIDsTable = EmitterInstance ? &EmitterInstance->GetData().GetSpawnedIDsTable() : nullptr;
-	int32 NumSpawned = SpawnedIDsTable ? SpawnedIDsTable->Num() : 0;
-#endif
+
+	int32 NumSpawned = 0;
+	const int32* SpawnedIDs = GetSpawnedIDs(Context, EmitterInstance, NumSpawned);
 
 	int32 IDAcquireTag = EmitterInstance ? EmitterInstance->GetData().GetIDAcquireTag() : 0;
 
@@ -1483,11 +1524,7 @@ void UNiagaraDataInterfaceParticleRead::GetSpawnedIDAtIndex(FVectorVMExternalFun
 		if (SpawnIndex >= 0 && SpawnIndex < NumSpawned)
 		{
 			ValidValue.SetValue(true);
-#ifdef NIAGARA_EXP_VM
-			IDValue.Index = SpawnedIDsTable[SpawnIndex];
-#else
-			IDValue.Index = (*SpawnedIDsTable)[SpawnIndex];
-#endif
+			IDValue.Index = SpawnedIDs[SpawnIndex];
 			IDValue.AcquireTag = IDAcquireTag;
 		}
 		else
diff --git a/Engine/Plugins/FX/Niagara/Source/Niagara/Private/NiagaraDataSet.cpp b/Engine/Plugins/FX/Niagara/Source/Niagara/Private/NiagaraDataSet.cpp
--- a/Engine/Plugins/FX/Niagara/Source/Niagara/Private/NiagaraDataSet.cpp
+++ b/Engine/Plugins/FX/Niagara/Source/Niagara/Private/NiagaraDataSet.cpp
@@ -357,7 +357,7 @@ void FNiagaraDataSet::Allocate(int32 NumInstances, bool bMaintainExisting)
 			RequiredIDs = ExistingNumIDs;
 		}
 
-#ifndef NIAGARA_EXP_VM //the new VM doesn't use the SpawnedIDsTable, so remove it.
+#if VECTORVM_SUPPORTS_LEGACY //the new VM doesn't use the SpawnedIDsTable, so remove it.
 		// We know that we can't spawn more than NumFreeIDs particles, so we can pre-allocate SpawnedIDsTable here, to avoid allocations during execution.
 		SpawnedIDsTable.Reserve(NumFreeIDs);
 #endif
diff --git a/Engine/Plugins/FX/Niagara/Source/Niagara/Private/NiagaraEmitterInstance.cpp b/Engine/Plugins/FX/Niagara/Source/Niagara/Private/NiagaraEmitterInstance.cpp
--- a/Engine/Plugins/FX/Niagara/Source/Niagara/Private/NiagaraEmitterInstance.cpp
+++ b/Engine/Plugins/FX/Niagara/Source/Niagara/Private/NiagaraEmitterInstance.cpp
@@ -1680,7 +1680,7 @@ void FNiagaraEmitterInstance::Tick(float DeltaSeconds)
 	uint32 TotalActualEventSpawns = 0;
 
 	Data.NumSpawnedIDs = 0;
-#ifndef NIAGARA_EXP_VM
+#if VECTORVM_SUPPORTS_LEGACY
 	Data.GetSpawnedIDsTable().SetNum(0, false);
 #endif
 	int32 SpawnCountRemaining = AllocationSize - OrigNumParticles;
diff --git a/Engine/Plugins/FX/Niagara/Source/Niagara/Private/NiagaraScript.cpp b/Engine/Plugins/FX/Niagara/Source/Niagara/Private/NiagaraScript.cpp
--- a/Engine/Plugins/FX/Niagara/Source/Niagara/Private/NiagaraScript.cpp
+++ b/Engine/Plugins/FX/Niagara/Source/Niagara/Private/NiagaraScript.cpp
@@ -287,11 +287,31 @@ FVersionedNiagaraScriptData::FVersionedNiagaraScriptData()
 }
 #endif
 
+#if VECTORVM_SUPPORTS_EXPERIMENTAL
+bool FNiagaraVMExecutableData::SupportsExperimentalVM() const
+{
+	return ExperimentalContextData.Num() > 0;
+}
+
+FVectorVMOptimizeContext FNiagaraVMExecutableData::BuildExperimentalContext() const
+{
+	FVectorVMOptimizeContext Context;
+	if (SupportsExperimentalVM())
+	{
+		ReinterpretVectorVMOptimizeContextData(ExperimentalContextData, Context);
+	}
+	else
+	{
+		FMemory::Memzero(Context);
+	}
+
+	return Context;
+}
+#endif
+
+
 UNiagaraScript::UNiagaraScript()
 {
-#	ifdef NIAGARA_EXP_VM
-	FMemory::Memzero(&OptimizeContext, sizeof(OptimizeContext));
-#	endif //NIAGARA_EXP_VM
 }
 
 #if WITH_EDITORONLY_DATA
@@ -738,9 +758,6 @@ UNiagaraScript::UNiagaraScript(const FObjectInitializer& ObjectInitializer)
 	, IsCooked(false)
 #endif
 {
-#	ifdef NIAGARA_EXP_VM
-	FMemory::Memzero(&OptimizeContext, sizeof(OptimizeContext));
-#	endif //NIAGARA_EXP_VM
 #if WITH_EDITORONLY_DATA
 	ScriptResource = MakeUnique<FNiagaraShaderScript>();
 	ScriptResource->OnCompilationComplete().AddUniqueDynamic(this, &UNiagaraScript::RaiseOnGPUCompilationComplete);
@@ -749,13 +766,6 @@ UNiagaraScript::UNiagaraScript(const FObjectInitializer& ObjectInitializer)
 #endif
 }
 
-UNiagaraScript::~UNiagaraScript()
-{
-#	ifdef NIAGARA_EXP_VM
-	FreeVectorVMOptimizeContext(&OptimizeContext);
-#	endif //NIAGARA_EXP_VM
-}
-
 #if WITH_EDITORONLY_DATA
 class UNiagaraSystem* UNiagaraScript::FindRootSystem()
 {
@@ -1346,6 +1356,7 @@ void FNiagaraScriptAsyncOptimizeTaskState::OptimizeByteCode()
 
 	if (bShouldOptimizeByteCode && SourceByteCode.HasByteCode())
 	{
+#if VECTORVM_SUPPORTS_LEGACY
 		// Generate optimized byte code on any thread
 		OptimizedByteCode.Reserve(SourceByteCode.GetLength());
 		VectorVM::OptimizeByteCode(SourceByteCode.GetDataPtr(), OptimizedByteCode.GetData(), MakeArrayView(ExternalFunctionRegisterCounts));
@@ -1355,6 +1366,7 @@ void FNiagaraScriptAsyncOptimizeTaskState::OptimizeByteCode()
 		{
 			SourceByteCode.Reset();
 		}
+#endif
 	}
 	bOptimizationComplete = true;
 }
@@ -1380,7 +1392,11 @@ struct FNiagaraScriptAsyncOptimizeTask
 
 bool UNiagaraScript::ShouldDecompressByteCode() const
 {
+#if VECTORVM_SUPPORTS_LEGACY
 	return CachedScriptVM.IsValid() && CachedScriptVM.ByteCode.IsCompressed();
+#else
+	return false;
+#endif
 }
 
 bool UNiagaraScript::ShouldOptimizeByteCode() const
@@ -1402,18 +1418,6 @@ bool UNiagaraScript::ShouldFreeUnoptimizedByteCode() const
 
 FGraphEventRef UNiagaraScript::HandleByteCodeOptimization(bool bShouldForceNow)
 {
-#	ifdef NIAGARA_EXP_VM
-	//this is just necessary because VectorVM doesn't know about FVMExternalFunctionBindingInfo
-	int NumExtFns = CachedScriptVM.CalledVMExternalFunctions.Num();
-	FVectorVMExtFunctionData *ExtFnTable = (FVectorVMExtFunctionData *)FMemory_Alloca(sizeof(FVectorVMExtFunctionData) * NumExtFns); //32 is pointless, it just gets around static analysis triggering for unknown reasons
-	CA_ASSUME(ExtFnTable != nullptr);
-	for (int i = 0; i < NumExtFns; ++i) {
-		ExtFnTable[i].NumInputs = CachedScriptVM.CalledVMExternalFunctions[i].GetNumInputs();
-		ExtFnTable[i].NumOutputs = CachedScriptVM.CalledVMExternalFunctions[i].GetNumOutputs();
-	}
-	OptimizeVectorVMScript(CachedScriptVM.ByteCode.GetDataPtr(), CachedScriptVM.ByteCode.GetLength(), ExtFnTable, NumExtFns, &OptimizeContext, VVMOptFlag_SaveIntermediateState);
-
-#	endif //NIAGARA_EXP_VM
 	check(IsInGameThread());
 
 	const bool bHasOptimizationTask = CachedScriptVM.OptimizationTask.State.IsValid();
diff --git a/Engine/Plugins/FX/Niagara/Source/Niagara/Private/NiagaraScriptExecutionContext.cpp b/Engine/Plugins/FX/Niagara/Source/Niagara/Private/NiagaraScriptExecutionContext.cpp
--- a/Engine/Plugins/FX/Niagara/Source/Niagara/Private/NiagaraScriptExecutionContext.cpp
+++ b/Engine/Plugins/FX/Niagara/Source/Niagara/Private/NiagaraScriptExecutionContext.cpp
@@ -27,20 +27,31 @@ static FAutoConsoleVariableRef CVarNiagaraExecVMScripts(
 	ECVF_Default
 );
 
+static int32 GbForceExecVMPath = 0;
+static FAutoConsoleVariableRef CVarNiagaraForceExecVMPath(
+	TEXT("fx.ForceExecVMPath"),
+	GbForceExecVMPath,
+	TEXT("If < 0, the legacy VM path will be used, if > 0 the experimental version will be used, and the default if 0. \n"),
+	ECVF_Default
+);
+
 FNiagaraScriptExecutionContextBase::FNiagaraScriptExecutionContextBase()
 	: Script(nullptr)
 	, VectorVMState(nullptr)
 	, ScriptType(ENiagaraSystemSimulationScript::Update)
 	, bAllowParallel(true)
+#if VECTORVM_SUPPORTS_EXPERIMENTAL && VECTORVM_SUPPORTS_LEGACY
+	, bUsingExperimentalVM(true)
+#endif
 {
 
 }
 
 FNiagaraScriptExecutionContextBase::~FNiagaraScriptExecutionContextBase()
 {
-#	ifdef NIAGARA_EXP_VM
+#if VECTORVM_SUPPORTS_EXPERIMENTAL
 	FreeVectorVMState(VectorVMState);
-#	endif 
+#endif 
 }
 
 bool FNiagaraScriptExecutionContextBase::Init(UNiagaraScript* InScript, ENiagaraSimTarget InTarget)
@@ -51,6 +62,24 @@ bool FNiagaraScriptExecutionContextBase::Init(UNiagaraScript* InScript, ENiagara
 
 	HasInterpolationParameters = Script && Script->GetComputedVMCompilationId().HasInterpolatedParameters();
 
+#if VECTORVM_SUPPORTS_EXPERIMENTAL && VECTORVM_SUPPORTS_LEGACY
+	const bool ScriptSupportsExperimentalVM = Script ? Script->GetVMExecutableData().SupportsExperimentalVM() : false;
+	bool SystemDisabledExperimentalVM = false;
+
+	if (InScript)
+	{
+		if (const UNiagaraSystem* OwnerSystem = InScript->GetTypedOuter<UNiagaraSystem>())
+		{
+			if (OwnerSystem->ShouldDisableExperimentalVM())
+			{
+				SystemDisabledExperimentalVM = true;
+			}
+		}
+	}
+
+	bUsingExperimentalVM = ScriptSupportsExperimentalVM && ((GbForceExecVMPath > 0) || (!GbForceExecVMPath && !SystemDisabledExperimentalVM));
+#endif
+
 	return true;
 }
 
@@ -152,237 +181,39 @@ bool FNiagaraScriptExecutionContextBase::Execute(uint32 NumInstances, const FScr
 		CreateStatScopeData();
 #endif	
 
-		FNiagaraVMExecutableData& ExecData = Script->GetVMExecutableData();
+		bool bSuccess = true;
 
-		// If we have an optimization task it must be ready at this point
-		// However we will need to lock and test again as multiple threads may be coming in here
-		if (ExecData.OptimizationTask.State.IsValid())
+#if VECTORVM_SUPPORTS_EXPERIMENTAL && VECTORVM_SUPPORTS_LEGACY
+		if (bUsingExperimentalVM)
 		{
-			FScopeLock Lock(&ExecData.OptimizationTask.Lock);
-			if (ExecData.OptimizationTask.State.IsValid())
-			{
-				ExecData.ApplyFinishedOptimization(Script->GetVMExecutableDataCompilationId(), ExecData.OptimizationTask.State);
-				ExecData.OptimizationTask.State.Reset();
-			}
+			bSuccess = ExecuteInternal_Experimental(NumInstances, ConstantBufferTable);
 		}
-
-		UObject *OuterObj0 = Script->GetTypedOuter<UNiagaraSystem>();
-		UObject *OuterObj1 = Script->GetTypedOuter<UNiagaraEmitter>();
-
-		check((ExecData.ByteCode.HasByteCode() && !ExecData.ByteCode.IsCompressed()) || (ExecData.OptimizedByteCode.HasByteCode() && !ExecData.OptimizedByteCode.IsCompressed()));
-
-		VectorVM::FVectorVMExecArgs ExecArgs;
-		ExecArgs.ByteCode = ExecData.ByteCode.GetDataPtr();
-		ExecArgs.OptimizedByteCode = ExecData.OptimizedByteCode.HasByteCode() ? ExecData.OptimizedByteCode.GetDataPtr(): nullptr;
-		ExecArgs.NumTempRegisters = ExecData.NumTempRegisters;
-		ExecArgs.ConstantTableCount = ConstantBufferTable.Buffers.Num();
-		ExecArgs.ConstantTable = ConstantBufferTable.Buffers.GetData();
-		ExecArgs.ConstantTableSizes = ConstantBufferTable.BufferSizes.GetData();
-		ExecArgs.DataSetMetaTable = DataSetMetaTable;
-		ExecArgs.ExternalFunctionTable = FunctionTable.GetData();
-		ExecArgs.UserPtrTable = UserPtrTable.GetData();
-		ExecArgs.NumInstances = NumInstances;
-#if STATS
-		ExecArgs.StatScopes = MakeArrayView(StatScopeData);
-#elif ENABLE_STATNAMEDEVENTS
-		ExecArgs.StatNamedEventsScopes = Script->GetStatNamedEvents();
-#endif
-		
-		ExecArgs.bAllowParallel = bAllowParallel;
-		FVectorVMSerializeState ExpSerializeState = { };
-		FVectorVMSerializeState UESerializeState = { };
-
-#if defined(NIAGARA_EXP_VM) || defined(VVM_INCLUDE_SERIALIZATION)
-		FVectorVMConstData *VVMConstData = (FVectorVMConstData *)FMemory_Alloca(sizeof(FVectorVMConstData) * ConstantBufferTable.BufferSizes.Num());
-		int NumVVMConstData              = ConstantBufferTable.BufferSizes.Num();
-		CA_ASSUME(VVMConstData != nullptr);
-		for (int i = 0; i < NumVVMConstData; ++i)
+		else
 		{
-			int32 NumBytes = ConstantBufferTable.BufferSizes[i];
-			check((NumBytes & 3) == 0); //should only be float32 or int32
-			VVMConstData[i].NumDWords = NumBytes >> 2;
-			VVMConstData[i].RegisterData = (void *)ConstantBufferTable.Buffers[i];
+			bSuccess = ExecuteInternal_Legacy(NumInstances, ConstantBufferTable);
 		}
-#endif //NIAGARA_EXP_VM
-
-#ifdef VVM_INCLUDE_SERIALIZATION
-		FString ScriptName;
-		Script->GetName(ScriptName);
-		
-		static const IConsoleVariable* CVarInstancesPerChunk = IConsoleManager::Get().FindConsoleVariable(TEXT("vm.InstancesPerChunk"));
-		int32 NumParallelInstancesPerChunk = CVarInstancesPerChunk ? CVarInstancesPerChunk->GetInt() : 128;
-		
-		UESerializeState.ReallocFn        = VVMRealloc;
-		UESerializeState.FreeFn           = VVMFree;
-		UESerializeState.NumInstances     = NumInstances;
-		UESerializeState.NumTempRegisters = ExecData.NumTempRegisters;
-		UESerializeState.NumTempRegFlags  = UESerializeState.NumTempRegisters;
-		
-		UESerializeState.TempRegFlags     = (unsigned char *)VVMRealloc(NULL, UESerializeState.NumTempRegisters, __FILE__, __LINE__);
-		FMemory::Memset(UESerializeState.TempRegFlags, 0, UESerializeState.NumTempRegisters);
-		UESerializeState.Bytecode         = (unsigned char *)ExecData.ByteCode.GetData().GetData();
-		UESerializeState.NumBytecodeBytes = ExecData.ByteCode.GetData().Num();
-		SerializeVectorVMInputDataSets(&UESerializeState, DataSetMetaTable, VVMConstData, NumVVMConstData);
-
-		UESerializeState.NumChunks        = (NumInstances + NumParallelInstancesPerChunk - 1) / NumParallelInstancesPerChunk;
-		UESerializeState.Chunks           = (FVectorVMSerializeChunk *)VVMRealloc(NULL, sizeof(FVectorVMSerializeChunk) * UESerializeState.NumChunks, __FILE__, __LINE__);
-
-		
-		int NumExternalFunctions = FunctionTable.Num();
-		if (NumExternalFunctions != 0)
-		{
-			UESerializeState.ExternalData = (FVectorVMSerializeExternalData *)VVMRealloc(NULL, sizeof(FVectorVMSerializeExternalData) * NumExternalFunctions, __FILE__, __LINE__);
-			if (UESerializeState.ExternalData)
-			{
-				UESerializeState.NumExternalData = NumExternalFunctions;
-				const FNiagaraVMExecutableData& ScriptExecutableData = Script->GetVMExecutableData();
-				for (int i = 0; i < NumExternalFunctions; ++i)
-				{
-					FVectorVMSerializeExternalData *ExtData = UESerializeState.ExternalData + i;
-					TCHAR buff[128];
-					uint32 NameLen = ScriptExecutableData.CalledVMExternalFunctions[i].Name.ToString(buff, sizeof(buff) >> 1);
-					if (NameLen > 0)
-					{
-						ExtData->Name       = (wchar_t *)UESerializeState.ReallocFn(NULL, sizeof(wchar_t) * NameLen, __FILE__, __LINE__);
-						FMemory::Memcpy(ExtData->Name, buff, sizeof(TCHAR) * NameLen);
-						ExtData->NameLen    = NameLen;
-					}
-					else
-					{
-						ExtData->Name       = NULL;
-						ExtData->NameLen    = 0;
-					}
-					ExtData->NumInputs      = ScriptExecutableData.CalledVMExternalFunctions[i].GetNumInputs();
-					ExtData->NumOutputs     = ScriptExecutableData.CalledVMExternalFunctions[i].GetNumOutputs();
-				}
-			}
-		}
-#endif //VVM_INCLUDE_SERIALIZATION
-
-#ifndef NIAGARA_EXP_VM
-		ExpSerializeState.ReallocFn = VVMRealloc;
-		ExpSerializeState.FreeFn = VVMFree;
-		VectorVM::Exec(ExecArgs, &UESerializeState);
-#ifdef VVM_INCLUDE_SERIALIZATION
-		SerializeVectorVMOutputDataSets(&UESerializeState, DataSetMetaTable, VVMConstData, NumVVMConstData);
-#endif //VVM_INCLUDE_SERIALIZATION
+#elif VECTORVM_SUPPORTS_EXPERIMENTAL
+		bSuccess = ExecuteInternal_Experimental(NumInstances, ConstantBufferTable);
+#elif VECTORVM_SUPPORTS_LEGACY
+		bSuccess = ExecuteInternal_Legacy(NumInstances, ConstantBufferTable);
 #else
-		FVectorVMInitData InitData     = { };
-		InitData.OptimizeContext       = &Script->OptimizeContext;
-		InitData.NumInstances          = NumInstances;
-		InitData.DataSets              = DataSetMetaTable;
-		InitData.ConstData             = VVMConstData;
-		InitData.NumConstData          = NumVVMConstData;
-		InitData.ExtFunctionTable      = FunctionTable;
-		InitData.UserPtrTable          = UserPtrTable.GetData();
-		InitData.NumUserPtrTable       = UserPtrTable.Num();
-		InitData.ExistingVectorVMState = VectorVMState;
-		InitData.ReallocFn             = VVMRealloc;
-		InitData.FreeFn                = VVMFree;
-		FVectorVMExternalFnPerInstanceData *PerInstanceExternalData = nullptr;
-		VectorVMState = InitVectorVMState(&InitData, &PerInstanceExternalData, &ExpSerializeState);
-
-#ifdef VVM_INCLUDE_SERIALIZATION
-		if (ExpSerializeState.NumExternalData > 0)
-		{
-			const FNiagaraVMExecutableData& ScriptExecutableData = Script->GetVMExecutableData();
-			for (uint32 i = 0; i < ExpSerializeState.NumExternalData; ++i)
-			{
-				TCHAR buff[128];
-				uint32 NameLen = ScriptExecutableData.CalledVMExternalFunctions[i].Name.ToString(buff, sizeof(buff) >> 1);
-				if (NameLen > 0)
-				{
-					ExpSerializeState.ExternalData[i].Name = (wchar_t *)UESerializeState.ReallocFn(NULL, sizeof(wchar_t) * NameLen, __FILE__, __LINE__);
-					memcpy(ExpSerializeState.ExternalData[i].Name, buff, 2 * NameLen);
-					ExpSerializeState.ExternalData[i].NameLen = NameLen;
-				}
-			}
-		}
+	#error "Not implemented"
 #endif
-		if (VectorVMState && VectorVMState->Error.Flags == 0)
-		{
-			ExecVectorVMState(VectorVMState, &ExpSerializeState, &UESerializeState);
-		}
-#		endif
 
-	// Tell the datasets we wrote how many instances were actually written.
-	for (int Idx = 0; Idx < DataSetInfo.Num(); Idx++)
-	{
-		FNiagaraDataSetExecutionInfo& Info = DataSetInfo[Idx];
+		// Tell the datasets we wrote how many instances were actually written.
+		for (int Idx = 0; Idx < DataSetInfo.Num(); Idx++)
+		{
+			FNiagaraDataSetExecutionInfo& Info = DataSetInfo[Idx];
 
 #if NIAGARA_NAN_CHECKING
-		Info.DataSet->CheckForNaNs();
+			Info.DataSet->CheckForNaNs();
 #endif
 			if (Info.bUpdateInstanceCount)
 			{
 				Info.Output->SetNumInstances(Info.StartInstance + DataSetMetaTable[Idx].DataSetAccessIndex + 1);
 			}
-
-#ifndef NIAGARA_EXP_VM //@TODO: this can go when we remove the old VM.  It's only here for serialization purposes as the VM doesn't use a spawned table
-			Info.DataSet->NumSpawnedIDs = Info.DataSet->GetSpawnedIDsTable().Num();
-#endif
 		}
 
-#		ifdef VVM_INCLUDE_SERIALIZATION
-		{
-#			ifdef NIAGARA_EXP_VM
-			FVectorVMSerializeState *SerializeState = &ExpSerializeState;
-			const char *state_ext = "exp";
-			uint8 which_state_written = 1;
-#			else //NIAGARA_EXP_VM
-			FVectorVMSerializeState *SerializeState = &UESerializeState;
-			const char *state_ext = "ue";
-			uint8 which_state_written = 2;
-#			endif //NIAGARA_EXP_VM
-			
-			//only write under certain circumstances
-			//if (OuterObj0 &&
-			//	OuterObj0->GetName() == L"BoneTest_System" && 
-			//	OuterObj1 &&
-			//	OuterObj1->GetName() == L"BoneTest_1" &&
-			//	ScriptName == L"SpawnScript_0")
-			if (0) {
-				SerializeVectorVMOutputDataSets(SerializeState, DataSetMetaTable, VVMConstData, NumVVMConstData);
-				TStringBuilder<1024> sb;
-				
-				if (OuterObj0) {
-					sb.Append(OuterObj0->GetName());
-					sb.Append(L"_");
-				}
-				if (OuterObj1) {
-					sb.Append(OuterObj1->GetName());
-					sb.Append(L"_");
-				}
-
-				sb.Append(ScriptName);
-				sb.Append(L"_");
-
-				TAnsiStringBuilder<256> sb_ansi;
-				
-				sb_ansi.Appendf("0x%08X_%s_%d.vvm_dump\0", FPlatformTime::Cycles(), state_ext, NumInstances);
-				//sb.AppendAnsi(sb_ansi.GetData());
-
-				sb.Append(sb_ansi.GetData());
-
-				//char buff[64];
-				//sb.Appendf(
-				//int32 buff_len = (int32)snprintf(buff, sizeof(buff), "0x%08X_%s_%d.vvm_dump", FPlatform::Cycles(), state_ext, NumInstances);
-				//if (buff_len > 0 && buff_len < sizeof(buff)) {
-				//	buff[buff_len] = 0;
-				//	sb.AppendAnsi(buff, buff_len + 1);
-				//}
-				wchar_t *Name16 = (wchar_t *)sb.GetData();
-				Name16[sb.Len()] = 0;
-				SerializeVectorVMWriteToFile(SerializeState, which_state_written, Name16);
-			}
-
-			FreeVectorVMSerializeState(&ExpSerializeState);
-			//we don't own the bytecode memory so we can't free it
-			UESerializeState.Bytecode = NULL;
-			UESerializeState.NumBytecodeBytes = 0;
-			FreeVectorVMSerializeState(&UESerializeState);
-		}
-#		endif //VVM_INCLUDE_SERIALIZATION
 		//Can maybe do without resetting here. Just doing it for tidiness.
 		for (int32 DataSetIdx = 0; DataSetIdx < DataSetInfo.Num(); ++DataSetIdx)
 		{
@@ -393,6 +224,108 @@ bool FNiagaraScriptExecutionContextBase::Execute(uint32 NumInstances, const FScr
 	return true;//TODO: Error cases?
 }
 
+#if VECTORVM_SUPPORTS_EXPERIMENTAL
+
+bool FNiagaraScriptExecutionContextBase::ExecuteInternal_Experimental(uint32 NumInstances, const FScriptExecutionConstantBufferTable& ConstantBufferTable)
+{
+	FNiagaraVMExecutableData& ExecData = Script->GetVMExecutableData();
+
+	FVectorVMSerializeState UESerializeState = { };
+	FVectorVMSerializeState ExpSerializeState = { };
+
+	TArray<FVectorVMConstData, TInlineAllocator<FScriptExecutionConstantBufferTable::MaxBufferCount>> VVMConstData;
+	const int32 BufferTableCount = ConstantBufferTable.BufferSizes.Num();
+
+	for (int32 BufferTableIt = 0; BufferTableIt < BufferTableCount; ++BufferTableIt)
+	{
+		FVectorVMConstData& ConstData = VVMConstData.AddDefaulted_GetRef();
+
+		const int32 NumBytes = ConstantBufferTable.BufferSizes[BufferTableIt];
+		check((NumBytes & 3) == 0); //should only be float32 or int32
+		ConstData.NumDWords = NumBytes >> 2;
+		ConstData.RegisterData = (void*)ConstantBufferTable.Buffers[BufferTableIt];
+	}
+
+	FVectorVMOptimizeContext OptimizeContext = ExecData.BuildExperimentalContext();
+	FVectorVMInitData InitData = { };
+	InitData.OptimizeContext = &OptimizeContext;
+	InitData.NumInstances = NumInstances;
+	InitData.DataSets = DataSetMetaTable;
+	InitData.ConstData = VVMConstData.GetData();
+	InitData.NumConstData = VVMConstData.Num();
+	InitData.ExtFunctionTable = FunctionTable;
+	InitData.UserPtrTable = UserPtrTable.GetData();
+	InitData.NumUserPtrTable = UserPtrTable.Num();
+	InitData.ExistingVectorVMState = VectorVMState;
+	InitData.ReallocFn = VVMRealloc;
+	InitData.FreeFn = VVMFree;
+	FVectorVMExternalFnPerInstanceData* PerInstanceExternalData = nullptr;
+	VectorVMState = InitVectorVMState(&InitData, &PerInstanceExternalData, &ExpSerializeState);
+
+	if (VectorVMState && VectorVMState->Error.Flags == 0)
+	{
+		ExecVectorVMState(VectorVMState, &ExpSerializeState, &UESerializeState);
+	}
+
+	return true;
+}
+
+#endif
+
+#if VECTORVM_SUPPORTS_LEGACY
+
+bool FNiagaraScriptExecutionContextBase::ExecuteInternal_Legacy(uint32 NumInstances, const FScriptExecutionConstantBufferTable& ConstantBufferTable)
+{
+	FNiagaraVMExecutableData& ExecData = Script->GetVMExecutableData();
+
+	// If we have an optimization task it must be ready at this point
+	// However we will need to lock and test again as multiple threads may be coming in here
+	if (ExecData.OptimizationTask.State.IsValid())
+	{
+		FScopeLock Lock(&ExecData.OptimizationTask.Lock);
+		if (ExecData.OptimizationTask.State.IsValid())
+		{
+			ExecData.ApplyFinishedOptimization(Script->GetVMExecutableDataCompilationId(), ExecData.OptimizationTask.State);
+			ExecData.OptimizationTask.State.Reset();
+		}
+	}
+
+	check((ExecData.ByteCode.HasByteCode() && !ExecData.ByteCode.IsCompressed()) || (ExecData.OptimizedByteCode.HasByteCode() && !ExecData.OptimizedByteCode.IsCompressed()));
+
+	VectorVM::FVectorVMExecArgs ExecArgs;
+	ExecArgs.ByteCode = ExecData.ByteCode.GetDataPtr();
+	ExecArgs.OptimizedByteCode = ExecData.OptimizedByteCode.HasByteCode() ? ExecData.OptimizedByteCode.GetDataPtr() : nullptr;
+	ExecArgs.NumTempRegisters = ExecData.NumTempRegisters;
+	ExecArgs.ConstantTableCount = ConstantBufferTable.Buffers.Num();
+	ExecArgs.ConstantTable = ConstantBufferTable.Buffers.GetData();
+	ExecArgs.ConstantTableSizes = ConstantBufferTable.BufferSizes.GetData();
+	ExecArgs.DataSetMetaTable = DataSetMetaTable;
+	ExecArgs.ExternalFunctionTable = FunctionTable.GetData();
+	ExecArgs.UserPtrTable = UserPtrTable.GetData();
+	ExecArgs.NumInstances = NumInstances;
+#if STATS
+	ExecArgs.StatScopes = MakeArrayView(StatScopeData);
+#elif ENABLE_STATNAMEDEVENTS
+	ExecArgs.StatNamedEventsScopes = Script->GetStatNamedEvents();
+#endif
+
+	ExecArgs.bAllowParallel = bAllowParallel;
+	FVectorVMSerializeState UESerializeState = { };
+
+	VectorVM::Exec(ExecArgs, &UESerializeState);
+
+	for (int32 Idx = 0; Idx < DataSetInfo.Num(); Idx++)
+	{
+		FNiagaraDataSetExecutionInfo& Info = DataSetInfo[Idx];
+
+		Info.DataSet->NumSpawnedIDs = Info.DataSet->GetSpawnedIDsTable().Num();
+	}
+
+	return true;
+}
+
+#endif // VECTORVM_SUPPORTS_LEGACY
+
 TArrayView<const uint8> FNiagaraScriptExecutionContextBase::GetScriptLiterals() const
 {
 #if WITH_EDITORONLY_DATA
@@ -579,39 +512,45 @@ void FNiagaraScriptExecutionContextBase::PostTick()
 }
 
 //////////////////////////////////////////////////////////////////////////
-#ifdef NIAGARA_EXP_VM
-static void PerInsFn(FVectorVMExternalFunctionContext &PerInsFnContext, TArray<FNiagaraSystemInstance *> **SystemInstances, ENiagaraSystemSimulationScript ScriptType, int32 PerInstFunctionIndex, int32 UserPtrIdx) {
+#if VECTORVM_SUPPORTS_EXPERIMENTAL
+static void PerInsFn(FVectorVMExternalFunctionContext& ParentContext, FVectorVMExternalFunctionContextExperimental& PerInsFnContext, TArray<FNiagaraSystemInstance*>** SystemInstances, ENiagaraSystemSimulationScript ScriptType, int32 PerInstFunctionIndex, int32 UserPtrIdx)
+{
 	check(PerInsFnContext.DataSets.Num() > 0);
 	check(SystemInstances);
 	check(*SystemInstances);
-	
-	void *SavedUserPtrData       = UserPtrIdx != INDEX_NONE ? PerInsFnContext.UserPtrTable[UserPtrIdx] : NULL;
-	int32 InstanceOffset         = PerInsFnContext.DataSets[0].InstanceOffset; //Apparently the function table is generated based off the first data set, therefore this is safe. (I don't like this - @shawn mcgrath)	
-	int NumInstances             = PerInsFnContext.NumInstances;
+
+	void* SavedUserPtrData = UserPtrIdx != INDEX_NONE ? PerInsFnContext.UserPtrTable[UserPtrIdx] : NULL;
+	int32 InstanceOffset = PerInsFnContext.DataSets[0].InstanceOffset; //Apparently the function table is generated based off the first data set, therefore this is safe. (I don't like this - @shawn mcgrath)	
+	int NumInstances = PerInsFnContext.NumInstances;
 	PerInsFnContext.NumInstances = 1;
-	for (int i = 0; i < NumInstances; ++i) {
-		PerInsFnContext.RegReadCount             = 0;
+	for (int i = 0; i < NumInstances; ++i)
+	{
+		PerInsFnContext.RegReadCount = 0;
 		PerInsFnContext.PerInstanceFnInstanceIdx = i;
 
-		int32 InstanceIndex                           = InstanceOffset + PerInsFnContext.StartInstance + i;
-		FNiagaraSystemInstance *Instance              = (**SystemInstances)[InstanceIndex];
-		const FNiagaraPerInstanceDIFuncInfo &FuncInfo = Instance->GetPerInstanceDIFunction(ScriptType, PerInstFunctionIndex);
-		
-		if (UserPtrIdx != INDEX_NONE) {
+		int32 InstanceIndex = InstanceOffset + PerInsFnContext.StartInstance + i;
+		FNiagaraSystemInstance* Instance = (**SystemInstances)[InstanceIndex];
+		const FNiagaraPerInstanceDIFuncInfo& FuncInfo = Instance->GetPerInstanceDIFunction(ScriptType, PerInstFunctionIndex);
+
+		if (UserPtrIdx != INDEX_NONE)
+		{
 			PerInsFnContext.UserPtrTable[UserPtrIdx] = FuncInfo.InstData;
 		}
-		FuncInfo.Function.Execute(PerInsFnContext);
+		FuncInfo.Function.Execute(ParentContext);
 	}
-	
-	if (SavedUserPtrData) {
+
+	if (SavedUserPtrData)
+	{
 		PerInsFnContext.UserPtrTable[UserPtrIdx] = SavedUserPtrData;
 	}
 }
-#else
-void FNiagaraSystemScriptExecutionContext::PerInstanceFunctionHook(FVectorVMExternalFunctionContext &Context, int32 PerInstFunctionIndex, int32 UserPtrIndex)
+#endif
+
+#if VECTORVM_SUPPORTS_LEGACY
+void FNiagaraSystemScriptExecutionContext::PerInstanceFunctionHook(FVectorVMExternalFunctionContext& ParentContext, FVectorVMExternalFunctionContextLegacy& Context, int32 PerInstFunctionIndex, int32 UserPtrIndex)
 {
 	check(SystemInstances);
-	
+
 	//This is a bit of a hack. We grab the base offset into the instance data from the primary dataset.
 	//TODO: Find a cleaner way to do this.
 	int32 InstanceOffset = Context.VectorVMContext->GetDataSetMeta(0).InstanceOffset;
@@ -634,10 +573,10 @@ void FNiagaraSystemScriptExecutionContext::PerInstanceFunctionHook(FVectorVMExte
 		int32 InstanceIndex = InstanceOffset + CachedContextStartInstance + i;
 		FNiagaraSystemInstance* Instance = (*SystemInstances)[InstanceIndex];
 		const FNiagaraPerInstanceDIFuncInfo& FuncInfo = Instance->GetPerInstanceDIFunction(ScriptType, PerInstFunctionIndex);
-		
+
 		//TODO: We can embed the instance data inside the function lambda. No need for the user ptr table at all.
 		//Do this way for now to reduce overall complexity of the initial change. Doing this needs extensive boiler plate changes to most DI classes and a script recompile.
-		if(UserPtrIndex != INDEX_NONE)
+		if (UserPtrIndex != INDEX_NONE)
 		{
 			Context.VectorVMContext->UserPtrTable[UserPtrIndex] = FuncInfo.InstData;
 		}
@@ -647,7 +586,7 @@ void FNiagaraSystemScriptExecutionContext::PerInstanceFunctionHook(FVectorVMExte
 		//TODO: In future for DIs where more perf is needed here we could split the DI func into an args gen and a execution.
 		//The this path could gen args from the bytecode once and just run the execution func per instance.
 		//I wonder if we could auto generate the args gen in a template func and just pass them into the DI for perf and reduced end user/author complexity.
-		FuncInfo.Function.Execute(Context);
+		FuncInfo.Function.Execute(ParentContext);
 	}
 
 	//Restore the context state.
@@ -713,19 +652,43 @@ bool FNiagaraSystemScriptExecutionContext::Tick(class FNiagaraSystemInstance* In
 						//Currently we must assume that any User DI is overridden but maybe we can be less conservative with this in future.
 						if (ScriptDIInfo.NeedsPerInstanceBinding())
 						{
-#							ifdef NIAGARA_EXP_VM
-							auto PerInstFunctionHookLambda = [SystemInstances = &this->SystemInstances, ScriptType = this->ScriptType, NumPerInstanceFunctions, UserPtrIdx = ScriptDIInfo.UserPtrIdx](FVectorVMExternalFunctionContext &ExtFnContext) {
-								PerInsFn(ExtFnContext, SystemInstances, ScriptType, NumPerInstanceFunctions, UserPtrIdx);
-							};
-#							else
-							auto PerInstFunctionHookLambda = [ExecContext = this, NumPerInstanceFunctions, UserPtrIndex = ScriptDIInfo.UserPtrIdx](FVectorVMExternalFunctionContext& Context)
+#if VECTORVM_SUPPORTS_EXPERIMENTAL && VECTORVM_SUPPORTS_LEGACY
+							if (bUsingExperimentalVM)
+							{
+								FuncInfo.Function = FVMExternalFunction::CreateLambda(
+									[SystemInstances = &this->SystemInstances, ScriptType = this->ScriptType, NumPerInstanceFunctions, UserPtrIdx = ScriptDIInfo.UserPtrIdx](FVectorVMExternalFunctionContext& ExtFnContext)
+									{
+										PerInsFn(ExtFnContext, ExtFnContext.Experimental, SystemInstances, ScriptType, NumPerInstanceFunctions, UserPtrIdx);
+									}
+								);
+							}
+							else
+							{
+								FuncInfo.Function = FVMExternalFunction::CreateLambda(
+									[ExecContext = this, NumPerInstanceFunctions, UserPtrIndex = ScriptDIInfo.UserPtrIdx](FVectorVMExternalFunctionContext& Context)
+									{
+										//This DI needs a binding per instance so we just bind to the external function hook which will call the correct binding for each instance.
+										ExecContext->PerInstanceFunctionHook(Context, Context.Legacy, NumPerInstanceFunctions, UserPtrIndex);
+									}
+								);
+							}
+#elif VECTORVM_SUPPORTS_EXPERIMENTAL
+							FuncInfo.Function = FVMExternalFunction::CreateLambda(
+								[SystemInstances = &this->SystemInstances, ScriptType = this->ScriptType, NumPerInstanceFunctions, UserPtrIdx = ScriptDIInfo.UserPtrIdx](FVectorVMExternalFunctionContext& ExtFnContext)
+							{
+								PerInsFn(ExtFnContext, ExtFnContext, SystemInstances, ScriptType, NumPerInstanceFunctions, UserPtrIdx);
+							}
+							);
+#elif VECTORVM_SUPPORTS_LEGACY
+							FuncInfo.Function = FVMExternalFunction::CreateLambda(
+								[ExecContext = this, NumPerInstanceFunctions, UserPtrIndex = ScriptDIInfo.UserPtrIdx](FVectorVMExternalFunctionContext& Context)
 							{
 								//This DI needs a binding per instance so we just bind to the external function hook which will call the correct binding for each instance.
-								ExecContext->PerInstanceFunctionHook(Context, NumPerInstanceFunctions, UserPtrIndex);
-							};
-#							endif
+								ExecContext->PerInstanceFunctionHook(Context, Context, NumPerInstanceFunctions, UserPtrIndex);
+							}
+							);
+#endif
 							++NumPerInstanceFunctions;
-							FuncInfo.Function = FVMExternalFunction::CreateLambda(PerInstFunctionHookLambda);
 						}
 						else
 						{
diff --git a/Engine/Plugins/FX/Niagara/Source/Niagara/Private/NiagaraSystem.cpp b/Engine/Plugins/FX/Niagara/Source/Niagara/Private/NiagaraSystem.cpp
--- a/Engine/Plugins/FX/Niagara/Source/Niagara/Private/NiagaraSystem.cpp
+++ b/Engine/Plugins/FX/Niagara/Source/Niagara/Private/NiagaraSystem.cpp
@@ -4,6 +4,7 @@
 #include "NiagaraSystem.h"
 
 #include "INiagaraEditorOnlyDataUtlities.h"
+#include "NiagaraAsyncCompile.h"
 #include "NiagaraConstants.h"
 #include "NiagaraCustomVersion.h"
 #include "NiagaraEditorDataBase.h"
@@ -41,9 +42,7 @@
 #include "DerivedDataCacheInterface.h"
 #endif
 
-DECLARE_CYCLE_STAT(TEXT("Niagara - System - Precompile"), STAT_Niagara_System_Precompile, STATGROUP_Niagara);
 DECLARE_CYCLE_STAT(TEXT("Niagara - System - CompileScript"), STAT_Niagara_System_CompileScript, STATGROUP_Niagara);
-DECLARE_CYCLE_STAT(TEXT("Niagara - System - CompileScriptTaskGT"), STAT_Niagara_System_CompileScriptTaskGT, STATGROUP_Niagara);
 DECLARE_CYCLE_STAT(TEXT("Niagara - System - CompileScript_ResetAfter"), STAT_Niagara_System_CompileScriptResetAfter, STATGROUP_Niagara);
 
 #if ENABLE_COOK_STATS
@@ -89,14 +88,6 @@ static FAutoConsoleVariableRef CVarNiagaraCompileWaitLoggingTerminationCap(
 	ECVF_Default
 );
 
-static float GNiagaraCompileDDCWaitTimeout = 10.0f;
-static FAutoConsoleVariableRef CVarNiagaraCompileDDCWaitTimeout(
-	TEXT("fx.Niagara.CompileDDCWaitTimeout"),
-	GNiagaraCompileDDCWaitTimeout,
-	TEXT("During script compilation, how long do we wait for the ddc to answer in seconds before starting shader compilation?"),
-	ECVF_Default
-);
-
 //////////////////////////////////////////////////////////////////////////
 
 UNiagaraSystem::UNiagaraSystem(const FObjectInitializer& ObjectInitializer)
@@ -112,6 +103,7 @@ UNiagaraSystem::UNiagaraSystem(const FObjectInitializer& ObjectInitializer)
 , bDisableDebugSwitchesOnCook(true)
 #endif
 , bSupportLargeWorldCoordinates(true)
+, bDisableExperimentalVM(false)
 , bFixedBounds(false)
 #if WITH_EDITORONLY_DATA
 , bIsolateEnabled(false)
@@ -147,369 +139,6 @@ UNiagaraSystem::UNiagaraSystem(const FObjectInitializer& ObjectInitializer)
 	bLwcEnabledSettingCached = Settings ? Settings->bSystemsSupportLargeWorldCoordinates : true;
 }
 
-#if WITH_EDITORONLY_DATA
-TSharedPtr<FNiagaraCompileRequestDataBase, ESPMode::ThreadSafe> FNiagaraLazyPrecompileReference::GetPrecompileData(UNiagaraScript* ForScript)
-{
-	if (SystemPrecompiledData == nullptr)
-	{
-		UNiagaraPrecompileContainer* Container = NewObject<UNiagaraPrecompileContainer>(GetTransientPackage());
-		Container->System = System;
-		Container->Scripts = Scripts;
-		INiagaraModule& NiagaraModule = FModuleManager::Get().LoadModuleChecked<INiagaraModule>(TEXT("Niagara"));
-		SystemPrecompiledData = NiagaraModule.Precompile(Container, FGuid());
-
-		const TArray<FNiagaraEmitterHandle>& EmitterHandles = System->GetEmitterHandles();
-		for (int32 i = 0; i < EmitterHandles.Num(); i++)
-		{
-			const FNiagaraEmitterHandle& Handle = EmitterHandles[i];
-			if (Handle.GetInstance() && Handle.GetIsEnabled())
-			{
-				TSharedPtr<FNiagaraCompileRequestDataBase, ESPMode::ThreadSafe> EmitterPrecompiledData = SystemPrecompiledData->GetDependentRequest(i);
-				TArray<UNiagaraScript*> EmitterScripts;
-				Handle.GetInstance()->GetScripts(EmitterScripts, false, true);
-				check(EmitterScripts.Num() > 0);
-				for (UNiagaraScript* EmitterScript : EmitterScripts)
-				{
-					EmitterMapping.Add(EmitterScript, EmitterPrecompiledData);
-				}
-			}
-		}
-	}
-	if (SystemPrecompiledData.IsValid() == false)
-	{
-		UE_LOG(LogNiagara, Error, TEXT("Failed to precompile %s.  This is due to unexpected invalid or broken data.  Additional details should be in the log."), *System->GetPathName());
-		return nullptr;
-	}
-
-	if (ForScript->IsSystemSpawnScript() || ForScript->IsSystemUpdateScript())
-	{
-		return SystemPrecompiledData;
-	}
-	TSharedPtr<FNiagaraCompileRequestDataBase, ESPMode::ThreadSafe>* EmitterPrecompileData = EmitterMapping.Find(ForScript);
-	if (EmitterPrecompileData == nullptr)
-	{
-		UE_LOG(LogNiagara, Log, TEXT("Failed to precompile %s. This is probably because the system was modified between the original compile request and now."), *System->GetPathName());
-		return nullptr;
-	}
-	return *EmitterPrecompileData;
-}
-
-TSharedPtr<FNiagaraCompileRequestDuplicateDataBase, ESPMode::ThreadSafe> FNiagaraLazyPrecompileReference::GetPrecompileDuplicateData(UNiagaraEmitter* OwningEmitter, UNiagaraScript* TargetScript)
-{
-	TSharedPtr<FNiagaraCompileRequestDuplicateDataBase, ESPMode::ThreadSafe> PrecompileDuplicateData;
-	for (TSharedPtr<FNiagaraCompileRequestDuplicateDataBase, ESPMode::ThreadSafe> ExistingPrecompileDuplicateData : PrecompileDuplicateDatas)
-	{
-		if (ExistingPrecompileDuplicateData->IsDuplicateDataFor(System, OwningEmitter, TargetScript))
-		{
-			PrecompileDuplicateData = ExistingPrecompileDuplicateData;
-		}
-	}
-
-	if (PrecompileDuplicateData.IsValid() == false)
-	{
-		INiagaraModule& NiagaraModule = FModuleManager::Get().LoadModuleChecked<INiagaraModule>(TEXT("Niagara"));
-
-		PrecompileDuplicateData = NiagaraModule.PrecompileDuplicate(SystemPrecompiledData.Get(), System, OwningEmitter, TargetScript, FGuid());
-		PrecompileDuplicateDatas.Add(PrecompileDuplicateData);
-
-		PrecompileDuplicateData->GetDuplicatedObjects(CompilationRootObjects);
-		for (int32 i = 0; i < PrecompileDuplicateData->GetDependentRequestCount(); i++)
-		{
-			PrecompileDuplicateData->GetDependentRequest(i)->GetDuplicatedObjects(CompilationRootObjects);
-		}
-	}
-
-	return PrecompileDuplicateData;
-}
-
-FNiagaraAsyncCompileTask::FNiagaraAsyncCompileTask(UNiagaraSystem* InOwningSystem, FString InAssetPath, const FEmitterCompiledScriptPair& InScriptPair)
-{
-	OwningSystem = InOwningSystem;
-	AssetPath = InAssetPath;
-	ScriptPair = InScriptPair;
-
-	CurrentState = ENiagaraCompilationState::CheckDDC;
-}
-
-void FNiagaraAsyncCompileTask::ProcessCurrentState()
-{
-	SCOPE_CYCLE_COUNTER(STAT_Niagara_System_CompileScriptTaskGT);
-	if (IsDone())
-	{
-		return;
-	}
-
-	if (CurrentState == ENiagaraCompilationState::CheckDDC)
-	{
-		// check if the DDC data is ready
-		CheckDDCResult();
-	}
-	else if (CurrentState == ENiagaraCompilationState::Precompile)
-	{
-		// gather precompile data
-		StartCompileTime = FPlatformTime::Seconds();
-		PrecompileData();
-		MoveToState(ENiagaraCompilationState::StartCompileJob);
-	}
-	else if (CurrentState == ENiagaraCompilationState::StartCompileJob)
-	{
-		// start the async compile job
-		StartCompileJob();
-		MoveToState(ENiagaraCompilationState::AwaitResult);
-	}
-	else if (CurrentState == ENiagaraCompilationState::AwaitResult)
-	{
-		if (AwaitResult())
-		{
-			MoveToState(ENiagaraCompilationState::ProcessResult);
-		}
-	}
-	else if (CurrentState == ENiagaraCompilationState::ProcessResult)
-	{
-		// save the result from the compile job
-		ProcessResult();
-		MoveToState(ENiagaraCompilationState::PutToDDC);
-	}
-	else if (CurrentState == ENiagaraCompilationState::PutToDDC)
-	{
-		// put the result from the compile job into the ddc
-		PutToDDC();
-		MoveToState(ENiagaraCompilationState::Finished);
-	}
-	else
-	{
-		check(false);
-	}
-}
-
-void FNiagaraAsyncCompileTask::MoveToState(ENiagaraCompilationState NewState)
-{
-	if (CurrentState == ENiagaraCompilationState::Finished || CurrentState == ENiagaraCompilationState::Aborted)
-	{
-		// we are already in an end state, so nothing to do
-		return;
-	}
-
-	// check state machine integrity
-	check(NewState != ENiagaraCompilationState::CheckDDC);
-	if (NewState == ENiagaraCompilationState::Precompile)
-	{
-		check(CurrentState == ENiagaraCompilationState::CheckDDC);
-	}
-	if (NewState == ENiagaraCompilationState::StartCompileJob)
-	{
-		check(CurrentState == ENiagaraCompilationState::Precompile);
-	}
-	if (NewState == ENiagaraCompilationState::AwaitResult)
-	{
-		check(CurrentState == ENiagaraCompilationState::StartCompileJob);
-	}
-	if (NewState == ENiagaraCompilationState::ProcessResult)
-	{
-		check(CurrentState == ENiagaraCompilationState::AwaitResult);
-	}
-	if (NewState == ENiagaraCompilationState::PutToDDC)
-	{
-		check(CurrentState == ENiagaraCompilationState::ProcessResult);
-	}
-	if (NewState == ENiagaraCompilationState::Finished)
-	{
-		check(CurrentState == ENiagaraCompilationState::CheckDDC || CurrentState == ENiagaraCompilationState::PutToDDC);
-	}
-
-	UE_LOG(LogNiagara, Verbose, TEXT("Changing state %i -> %i for for %s!"), CurrentState, NewState, *AssetPath);
-	CurrentState = NewState;
-}
-
-bool FNiagaraAsyncCompileTask::IsDone() const
-{
-	return CurrentState == ENiagaraCompilationState::Finished || CurrentState == ENiagaraCompilationState::Aborted;
-}
-
-void LogVMID(const FNiagaraVMExecutableDataId& ID, FString Name)
-{
-	// some debug logging to find out why the ddc key changed
-	UE_LOG(LogNiagara, Display, TEXT("FNiagaraVMExecutableDataId %s:"), *Name);
-	UE_LOG(LogNiagara, Display, TEXT("  ScriptUsageType: %i"), ID.ScriptUsageType);
-	UE_LOG(LogNiagara, Display, TEXT("  ScriptUsageTypeID: %s"), *ID.ScriptUsageTypeID.ToString());
-	UE_LOG(LogNiagara, Display, TEXT("  CompilerVersionID: %s"), *ID.CompilerVersionID.ToString());
-	UE_LOG(LogNiagara, Display, TEXT("  BaseScriptCompileHash: %s"), *ID.BaseScriptCompileHash.ToString());
-	UE_LOG(LogNiagara, Display, TEXT("  bUsesRapidIterationParams: %i"), ID.bUsesRapidIterationParams);
-	UE_LOG(LogNiagara, Display, TEXT("  AdditionalDefines:"));
-	for (int32 Idx = 0; Idx < ID.AdditionalDefines.Num(); Idx++)
-	{
-		UE_LOG(LogNiagara, Display, TEXT("     %s"), *ID.AdditionalDefines[Idx]);
-	}
-	UE_LOG(LogNiagara, Display, TEXT("  AdditionalVariables:"));
-	for (int32 Idx = 0; Idx < ID.AdditionalVariables.Num(); Idx++)
-	{
-		UE_LOG(LogNiagara, Display, TEXT("     %s / %s"), *ID.AdditionalVariables[Idx].GetName().ToString(), *ID.AdditionalVariables[Idx].GetType().GetName());
-	}
-	UE_LOG(LogNiagara, Display, TEXT("  ReferencedCompileHashes:"));
-	for (int32 HashIndex = 0; HashIndex < ID.ReferencedCompileHashes.Num(); HashIndex++)
-	{
-		UE_LOG(LogNiagara, Display, TEXT("     %s"), *ID.ReferencedCompileHashes[HashIndex].ToString());
-	}
-}
-
-void FNiagaraAsyncCompileTask::PrecompileData()
-{
-	SCOPE_CYCLE_COUNTER(STAT_Niagara_System_Precompile);
-	UE_LOG(LogNiagara, Verbose, TEXT("Getting precompile data for %s!"), *AssetPath);
-	
-	// we compute the vmID here again to check if the script changed between the start of the task and now. If it was changed we need to update the ddc key,
-	// because otherwise we would save the compiled data under the wrong ddc key
-	FNiagaraVMExecutableDataId NewID;
-	ScriptPair.CompiledScript->ComputeVMCompilationId(NewID, FGuid());
-	FString CurrentDDCKey = UNiagaraScript::BuildNiagaraDDCKeyString(NewID);
-	if (DDCKey != CurrentDDCKey && LogNiagara.GetVerbosity() >= ELogVerbosity::Verbose)
-	{
-		UE_LOG(LogNiagara, Warning, TEXT("Compile ID for %s changed during the compilation task."), *AssetPath);
-		FNiagaraVMExecutableDataId OldID = ScriptPair.CompileId;
-		LogVMID(OldID, "OldID");
-		LogVMID(NewID, "NewID");
-	}
-	DDCKey = CurrentDDCKey;
-	
-	ComputedPrecompileData = PrecompileReference->GetPrecompileData(ScriptPair.CompiledScript);
-	TSharedPtr<FNiagaraCompileRequestDuplicateDataBase, ESPMode::ThreadSafe> SystemPrecompileDuplicateData = PrecompileReference->GetPrecompileDuplicateData(ScriptPair.Emitter, ScriptPair.CompiledScript);
-
-	if (ComputedPrecompileData == nullptr || SystemPrecompileDuplicateData == nullptr)
-	{
-		AbortTask();
-		return;
-	}
-
-	// Grab the list of user variables that were actually encountered so that we can add to them later.
-	ComputedPrecompileData->GatherPreCompiledVariables(TEXT("User"), EncounteredExposedVars);
-	for(int32 i = 0; i < ComputedPrecompileData->GetDependentRequestCount(); i++)
-	{
-		ComputedPrecompileData->GetDependentRequest(i)->GatherPreCompiledVariables(TEXT("User"), EncounteredExposedVars);
-	}
-
-	// Assign the correct duplicate data for this request.
-	if(ScriptPair.CompiledScript->GetUsage() == ENiagaraScriptUsage::SystemSpawnScript || ScriptPair.CompiledScript->GetUsage() == ENiagaraScriptUsage::SystemUpdateScript)
-	{
-		ComputedPrecompileDuplicateData = SystemPrecompileDuplicateData;
-	}
-	else
-	{
-		UNiagaraEmitter* OwningEmitter = ScriptPair.Emitter;
-		int32 EmitterIndex = OwningSystem->GetEmitterHandles().IndexOfByPredicate([OwningEmitter](const FNiagaraEmitterHandle& EmitterHandle) { return EmitterHandle.GetInstance() == OwningEmitter; });
-		if(EmitterIndex != INDEX_NONE)
-		{
-			ComputedPrecompileDuplicateData = SystemPrecompileDuplicateData->GetDependentRequest(EmitterIndex);
-		}
-	}
-}
-
-void FNiagaraAsyncCompileTask::StartCompileJob()
-{
-	UE_LOG(LogNiagara, Verbose, TEXT("Starting compilation for %s!"), *AssetPath);
-	
-	// Fire off the compile job
-	if (ScriptPair.CompiledScript->RequestExternallyManagedAsyncCompile(ComputedPrecompileData, ComputedPrecompileDuplicateData, ScriptPair.CompileId, ScriptPair.PendingJobID))
-	{
-		UE_LOG(LogNiagara, Verbose, TEXT("Designated compilation ID %i for compilation of %s!"), ScriptPair.PendingJobID, *AssetPath);
-		bUsedShaderCompilerWorker = true;
-	}
-	else
-	{
-		UE_LOG(LogNiagara, Warning, TEXT("For some reason we are reporting that %s is in sync even though AreScriptAndSourceSynchronized returned false!"), *ScriptPair.CompiledScript->GetPathName())
-		AbortTask();
-	}
-}
-
-bool FNiagaraAsyncCompileTask::AwaitResult()
-{
-	UE_LOG(LogNiagara, Verbose, TEXT("Waiting on compilation to finish for %s!"), *AssetPath);
-	INiagaraModule& NiagaraModule = FModuleManager::Get().LoadModuleChecked<INiagaraModule>(TEXT("Niagara"));
-	ExeData = NiagaraModule.GetCompileJobResult(ScriptPair.PendingJobID, bWaitForCompileJob);
-	return ExeData.IsValid();
-}
-
-void FNiagaraAsyncCompileTask::ProcessResult()
-{
-	check(ExeData);
-	UE_LOG(LogNiagara, Verbose, TEXT("Processing compilation results for %s!"), *AssetPath);
-	
-	// convert results to be saved in the ddc
-	ScriptPair.CompileResults = ExeData;
-	ScriptPair.bResultsReady = true;
-	ExeData->CompileTime = FPlatformTime::Seconds() - StartCompileTime;
-	ScriptPair.CompiledScript->ExecToBinaryData(ScriptPair.CompiledScript, DDCOutData, *ExeData);
-	UE_LOG(LogNiagara, Verbose, TEXT("Got %i bytes in ddc data for %s"), DDCOutData.Num(), *AssetPath);
-}
-
-void FNiagaraAsyncCompileTask::WaitAndResolveResult()
-{
-	check(IsInGameThread());
-	
-	bWaitForCompileJob = true;
-	while (!IsDone())
-	{
-		ProcessCurrentState();
-	}
-}
-
-void FNiagaraAsyncCompileTask::AbortTask()
-{
-	MoveToState(ENiagaraCompilationState::Aborted);
-}
-
-void FNiagaraAsyncCompileTask::CheckDDCResult()
-{
-	FDerivedDataCacheInterface* DDC = GetDerivedDataCache();
-	if (bWaitForCompileJob)
-	{
-		DDC->WaitAsynchronousCompletion(TaskHandle);
-	}
-	if (DDC->PollAsynchronousCompletion(TaskHandle))
-	{
-		bool bHasResults = DDC->GetAsynchronousResults(TaskHandle, DDCOutData);
-		if (bHasResults && DDCOutData.Num() > 0)
-		{
-			ExeData = MakeShared<FNiagaraVMExecutableData>();
-			if (ScriptPair.CompiledScript->BinaryToExecData(ScriptPair.CompiledScript, DDCOutData, *ExeData))
-			{
-				ScriptPair.CompileResults = ExeData;
-				ScriptPair.bResultsReady = true;
-				MoveToState(ENiagaraCompilationState::Finished);
-				UE_LOG(LogNiagara, Verbose, TEXT("Compilation data for %s could be pulled from the ddc."), *AssetPath);
-			}
-			else
-			{
-				UE_LOG(LogNiagara, Warning, TEXT("Unable to create exec data from ddc data for script %s, going to recompile it from scratch. DDC might be corrupted or there is a problem with script serialization."), *AssetPath);
-				ExeData.Reset();
-				MoveToState(ENiagaraCompilationState::Precompile);
-			}
-		}
-		else
-		{
-			MoveToState(ENiagaraCompilationState::Precompile);
-			UE_LOG(LogNiagara, Verbose, TEXT("No compilation data for %s found in the ddc."), *AssetPath);
-		}
-		DDCFetchTime = FPlatformTime::Seconds() - StartTaskTime;
-		UE_LOG(LogNiagara, Verbose, TEXT("Compile task for %s took %f seconds to fetch data from ddc."), *AssetPath, DDCFetchTime);
-	}
-	else if ((FPlatformTime::Seconds() - StartTaskTime) > GNiagaraCompileDDCWaitTimeout)
-	{
-		MoveToState(ENiagaraCompilationState::Precompile);
-		UE_LOG(LogNiagara, Log, TEXT("DDC timed out after %i seconds, starting compilation for %s."), (int)GNiagaraCompileDDCWaitTimeout, *AssetPath);
-	}
-}
-
-void FNiagaraAsyncCompileTask::PutToDDC()
-{
-	if (DDCOutData.Num() > 0)
-	{
-		UE_LOG(LogNiagara, Verbose, TEXT("Writing data for %s to the ddc"), *AssetPath);
-		GetDerivedDataCache()->Put(*DDCKey, DDCOutData, *AssetPath, true);
-	}
-}
-
-#endif
-
 UNiagaraSystem::UNiagaraSystem(FVTableHelper& Helper)
 	: Super(Helper)
 {
diff --git a/Engine/Source/Runtime/Core/Private/UObject/UE5MainStreamObjectVersion.cpp b/Engine/Source/Runtime/Core/Private/UObject/UE5MainStreamObjectVersion.cpp
--- a/Engine/Source/Runtime/Core/Private/UObject/UE5MainStreamObjectVersion.cpp
+++ b/Engine/Source/Runtime/Core/Private/UObject/UE5MainStreamObjectVersion.cpp
@@ -14,7 +14,7 @@ TMap<FGuid, FGuid> FUE5MainStreamObjectVersion::GetSystemGuids()
 	SystemGuids.Add(DevGuids.MATERIALSHADERMAP_DERIVEDDATA_VER, FGuid("65C01D817C9A4EEFAE9E988D41A1F3DD"));
 	SystemGuids.Add(DevGuids.NANITE_DERIVEDDATA_VER, FGuid("7E26CA5B0F5F03C40A3FD7173D88B8FA"));
 	SystemGuids.Add(DevGuids.NIAGARASHADERMAP_DERIVEDDATA_VER, FGuid("8A37C45D24F2423CBE5F8F371DE33575"));
-	SystemGuids.Add(DevGuids.Niagara_LatestScriptCompileVersion, FGuid("BFA6C45A89CE2B84BB2C444C78C46E76"));
+	SystemGuids.Add(DevGuids.Niagara_LatestScriptCompileVersion, FGuid("834E4334C1D7416CAF061A4FEFE4FE36"));
 	SystemGuids.Add(DevGuids.SkeletalMeshDerivedDataVersion, FGuid("FB7B2FB546DE49B2A2F88B7F69763CEE"));
 	SystemGuids.Add(DevGuids.STATICMESH_DERIVEDDATA_VER, FGuid("98F7D79A5811013825E75BBCC41ED3E9"));
 
diff --git a/Engine/Source/Runtime/VectorVM/Private/VectorVM.cpp b/Engine/Source/Runtime/VectorVM/Private/VectorVM.cpp
--- a/Engine/Source/Runtime/VectorVM/Private/VectorVM.cpp
+++ b/Engine/Source/Runtime/VectorVM/Private/VectorVM.cpp
@@ -141,15 +141,8 @@
 }
 #endif
 
-#ifdef NIAGARA_EXP_VM
-//@TODO: remove these functions when we go to the new VM
-FORCEINLINE void FDataSetMeta::LockFreeTable() {
+#if VECTORVM_SUPPORTS_LEGACY
 
-}
-FORCEINLINE void FDataSetMeta::UnlockFreeTable() {
-
-}
-#else
 //Temporarily locking the free table until we can implement a lock free algorithm. UE-65856
 FORCEINLINE void FDataSetMeta::LockFreeTable()
 {
@@ -173,7 +166,8 @@
 {
  	FreeTableLock.Unlock();
 }
-#endif //NIAGARA_EXP_VM
 
+#endif //VECTORVM_SUPPORTS_LEGACY
+
 static int32 GbParallelVVM = 1;
 static FAutoConsoleVariableRef CVarbParallelVVM(
@@ -282,6 +276,69 @@
 
 #include "VectorVMExperimental.inl"
 
+uint8 VectorVM::GetNumOpCodes()
+{
+	return (uint8)EVectorVMOp::NumOpcodes;
+}
+
+#if WITH_EDITOR
+UEnum* g_VectorVMEnumStateObj = nullptr;
+UEnum* g_VectorVMEnumOperandObj = nullptr;
+
+FString VectorVM::GetOpName(EVectorVMOp Op)
+{
+	check(g_VectorVMEnumStateObj);
+
+	FString OpStr = g_VectorVMEnumStateObj->GetNameByValue((uint8)Op).ToString();
+	int32 LastIdx = 0;
+	OpStr.FindLastChar(TEXT(':'), LastIdx);
+	return OpStr.RightChop(LastIdx + 1);
+}
+
+FString VectorVM::GetOperandLocationName(EVectorVMOperandLocation Location)
+{
+	check(g_VectorVMEnumOperandObj);
+
+	FString LocStr = g_VectorVMEnumOperandObj->GetNameByValue((uint8)Location).ToString();
+	int32 LastIdx = 0;
+	LocStr.FindLastChar(TEXT(':'), LastIdx);
+	return LocStr.RightChop(LastIdx + 1);
+}
+#endif
+
+uint8 VectorVM::CreateSrcOperandMask(EVectorVMOperandLocation Type0, EVectorVMOperandLocation Type1, EVectorVMOperandLocation Type2)
+{
+	return	(Type0 == EVectorVMOperandLocation::Constant ? OP0_CONST : OP_REGISTER) |
+		(Type1 == EVectorVMOperandLocation::Constant ? OP1_CONST : OP_REGISTER) |
+		(Type2 == EVectorVMOperandLocation::Constant ? OP2_CONST : OP_REGISTER);
+}
+
+#if VECTORVM_SUPPORTS_LEGACY
+namespace VectorKernelNoiseImpl
+{
+	static void BuildNoiseTable();
+}
+#endif
+
+void VectorVM::Init()
+{
+	static bool Inited = false;
+	if (Inited == false)
+	{
+#if WITH_EDITOR
+		g_VectorVMEnumStateObj = StaticEnum<EVectorVMOp>();
+		g_VectorVMEnumOperandObj = StaticEnum<EVectorVMOperandLocation>();
+#endif
+
+#if VECTORVM_SUPPORTS_LEGACY
+		VectorKernelNoiseImpl::BuildNoiseTable();
+#endif
+		Inited = true;
+	}
+}
+
+#if VECTORVM_SUPPORTS_LEGACY
+
 //////////////////////////////////////////////////////////////////////////
 //  VM Code Optimizer Context
 
@@ -646,13 +703,6 @@
 
 //////////////////////////////////////////////////////////////////////////
 
-uint8 VectorVM::CreateSrcOperandMask(EVectorVMOperandLocation Type0, EVectorVMOperandLocation Type1, EVectorVMOperandLocation Type2)
-{
-	return	(Type0 == EVectorVMOperandLocation::Constant ? OP0_CONST : OP_REGISTER) |
-		(Type1 == EVectorVMOperandLocation::Constant ? OP1_CONST : OP_REGISTER) |
-		(Type2 == EVectorVMOperandLocation::Constant ? OP2_CONST : OP_REGISTER);
-}
-
 //////////////////////////////////////////////////////////////////////////
 // Kernels
 template<typename Kernel, typename DstHandler, typename Arg0Handler, uint32 NumInstancesPerOp>
@@ -1560,6 +1610,87 @@
 	}
 };
 
+namespace VectorKernelNoiseImpl
+{
+	static void BuildNoiseTable()
+	{
+		// random noise
+		float TempTable[17][17][17];
+		for (int z = 0; z < 17; z++)
+		{
+			for (int y = 0; y < 17; y++)
+			{
+				for (int x = 0; x < 17; x++)
+				{
+					float f1 = (float)FMath::FRandRange(-1.0f, 1.0f);
+					TempTable[x][y][z] = f1;
+				}
+			}
+		}
+
+		// pad
+		for (int i = 0; i < 17; i++)
+		{
+			for (int j = 0; j < 17; j++)
+			{
+				TempTable[i][j][16] = TempTable[i][j][0];
+				TempTable[i][16][j] = TempTable[i][0][j];
+				TempTable[16][j][i] = TempTable[0][j][i];
+			}
+		}
+
+		// compute gradients
+		FVector3f TempTable2[17][17][17];
+		for (int z = 0; z < 16; z++)
+		{
+			for (int y = 0; y < 16; y++)
+			{
+				for (int x = 0; x < 16; x++)
+				{
+					FVector3f XGrad = FVector3f(1.0f, 0.0f, TempTable[x][y][z] - TempTable[x + 1][y][z]);
+					FVector3f YGrad = FVector3f(0.0f, 1.0f, TempTable[x][y][z] - TempTable[x][y + 1][z]);
+					FVector3f ZGrad = FVector3f(0.0f, 1.0f, TempTable[x][y][z] - TempTable[x][y][z + 1]);
+
+					FVector3f Grad = FVector3f(XGrad.Z, YGrad.Z, ZGrad.Z);
+					TempTable2[x][y][z] = Grad;
+				}
+			}
+		}
+
+		// pad
+		for (int i = 0; i < 17; i++)
+		{
+			for (int j = 0; j < 17; j++)
+			{
+				TempTable2[i][j][16] = TempTable2[i][j][0];
+				TempTable2[i][16][j] = TempTable2[i][0][j];
+				TempTable2[16][j][i] = TempTable2[0][j][i];
+			}
+		}
+
+
+		// compute curl of gradient field
+		for (int z = 0; z < 16; z++)
+		{
+			for (int y = 0; y < 16; y++)
+			{
+				for (int x = 0; x < 16; x++)
+				{
+					FVector3f Dy = TempTable2[x][y][z] - TempTable2[x][y + 1][z];
+					FVector3f Sy = TempTable2[x][y][z] + TempTable2[x][y + 1][z];
+					FVector3f Dx = TempTable2[x][y][z] - TempTable2[x + 1][y][z];
+					FVector3f Sx = TempTable2[x][y][z] + TempTable2[x + 1][y][z];
+					FVector3f Dz = TempTable2[x][y][z] - TempTable2[x][y][z + 1];
+					FVector3f Sz = TempTable2[x][y][z] + TempTable2[x][y][z + 1];
+					FVector3f Dir = FVector3f(Dy.Z - Sz.Y, Dz.X - Sx.Z, Dx.Y - Sy.X);
+
+					FVectorKernelNoise::RandomTable[x][y][z] = MakeVectorRegister(Dir.X, Dir.Y, Dir.Z, 0.f);
+				}
+			}
+		}
+	}
+}
+
 VectorRegister4Float FVectorKernelNoise::RandomTable[17][17][17];
 
 //////////////////////////////////////////////////////////////////////////
@@ -2104,16 +2235,18 @@
 
 	static void Exec(FVectorVMContext& Context)
 	{
-#ifdef NIAGARA_EXP_VM
-		check(false);
-#else
+#if VECTORVM_SUPPORTS_LEGACY
 		const uint32 ExternalFuncIdx = Context.DecodeU8();
 		const FVMExternalFunction* ExternalFunction = Context.ExternalFunctionTable[ExternalFuncIdx];
 		check(ExternalFunction);
 
 		if (ExternalFunction)
 		{
+#if VECTORVM_SUPPORTS_EXPERIMENTAL
 			FVectorVMExternalFunctionContext DataInterfaceFunctionContext(&Context);
+#else
+			FVectorVMExternalFunctionContextLegacy DataInterfaceFunctionContext(&Context);
+#endif
 			ExternalFunction->Execute(DataInterfaceFunctionContext);
 		}
 #endif
@@ -2464,101 +2597,6 @@
 	}
 };
 
-#if WITH_EDITOR
-UEnum* g_VectorVMEnumStateObj = nullptr;
-UEnum* g_VectorVMEnumOperandObj = nullptr;
-#endif
-
-
-void VectorVM::Init()
-{
-	static bool Inited = false;
-	if (Inited == false)
-	{
-#if WITH_EDITOR
-		g_VectorVMEnumStateObj = StaticEnum<EVectorVMOp>();
-		g_VectorVMEnumOperandObj = StaticEnum<EVectorVMOperandLocation>();
-#endif
-
-		// random noise
-		float TempTable[17][17][17];
-		for (int z = 0; z < 17; z++)
-		{
-			for (int y = 0; y < 17; y++)
-			{
-				for (int x = 0; x < 17; x++)
-				{
-					float f1 = (float)FMath::FRandRange(-1.0f, 1.0f);
-					TempTable[x][y][z] = f1;
-				}
-			}
-		}
-
-		// pad
-		for (int i = 0; i < 17; i++)
-		{
-			for (int j = 0; j < 17; j++)
-			{
-				TempTable[i][j][16] = TempTable[i][j][0];
-				TempTable[i][16][j] = TempTable[i][0][j];
-				TempTable[16][j][i] = TempTable[0][j][i];
-			}
-		}
-
-		// compute gradients
-		FVector3f TempTable2[17][17][17];
-		for (int z = 0; z < 16; z++)
-		{
-			for (int y = 0; y < 16; y++)
-			{
-				for (int x = 0; x < 16; x++)
-				{
-					FVector3f XGrad = FVector3f(1.0f, 0.0f, TempTable[x][y][z] - TempTable[x+1][y][z]);
-					FVector3f YGrad = FVector3f(0.0f, 1.0f, TempTable[x][y][z] - TempTable[x][y + 1][z]);
-					FVector3f ZGrad = FVector3f(0.0f, 1.0f, TempTable[x][y][z] - TempTable[x][y][z+1]);
-
-					FVector3f Grad = FVector3f(XGrad.Z, YGrad.Z, ZGrad.Z);
-					TempTable2[x][y][z] = Grad;
-				}
-			}
-		}
-
-		// pad
-		for (int i = 0; i < 17; i++)
-		{
-			for (int j = 0; j < 17; j++)
-			{
-				TempTable2[i][j][16] = TempTable2[i][j][0];
-				TempTable2[i][16][j] = TempTable2[i][0][j];
-				TempTable2[16][j][i] = TempTable2[0][j][i];
-			}
-		}
-
-
-		// compute curl of gradient field
-		for (int z = 0; z < 16; z++)
-		{
-			for (int y = 0; y < 16; y++)
-			{
-				for (int x = 0; x < 16; x++)
-				{
-					FVector3f Dy = TempTable2[x][y][z] - TempTable2[x][y + 1][z];
-					FVector3f Sy = TempTable2[x][y][z] + TempTable2[x][y + 1][z];
-					FVector3f Dx = TempTable2[x][y][z] - TempTable2[x + 1][y][z];
-					FVector3f Sx = TempTable2[x][y][z] + TempTable2[x + 1][y][z];
-					FVector3f Dz = TempTable2[x][y][z] - TempTable2[x][y][z + 1];
-					FVector3f Sz = TempTable2[x][y][z] + TempTable2[x][y][z + 1];
-					FVector3f Dir = FVector3f(Dy.Z - Sz.Y, Dz.X - Sx.Z, Dx.Y - Sy.X);
-
-					FVectorKernelNoise::RandomTable[x][y][z] = MakeVectorRegister(Dir.X, Dir.Y, Dir.Z, 0.f);
-				}
-			}
-		}
-
-		Inited = true;
-	}
-}
-
 void VectorVM::Exec(FVectorVMExecArgs& Args, FVectorVMSerializeState *SerializeState)
 {
 	//TRACE_CPUPROFILER_EVENT_SCOPE("VMExec");
@@ -2785,34 +2823,6 @@
 #endif //VVM_INCLUDE_SERIALIZATION
 }
 
-uint8 VectorVM::GetNumOpCodes()
-{
-	return (uint8)EVectorVMOp::NumOpcodes;
-}
-
-#if WITH_EDITOR
-FString VectorVM::GetOpName(EVectorVMOp Op)
-{
-	check(g_VectorVMEnumStateObj);
-
-	FString OpStr = g_VectorVMEnumStateObj->GetNameByValue((uint8)Op).ToString();
-	int32 LastIdx = 0;
-	OpStr.FindLastChar(TEXT(':'),LastIdx);
-	return OpStr.RightChop(LastIdx+1);
-}
-
-FString VectorVM::GetOperandLocationName(EVectorVMOperandLocation Location)
-{
-	check(g_VectorVMEnumOperandObj);
-
-	FString LocStr = g_VectorVMEnumOperandObj->GetNameByValue((uint8)Location).ToString();
-	int32 LastIdx = 0;
-	LocStr.FindLastChar(TEXT(':'), LastIdx);
-	return LocStr.RightChop(LastIdx+1);
-}
-#endif
-
-
 void ExecBatchedOutput(FVectorVMContext& Context)
 {
 	while (true)
@@ -3773,4 +3783,6 @@
 #endif //PLATFORM_SUPPORTS_UNALIGNED_LOADS && PLATFORM_LITTLE_ENDIAN
 }
 
+#endif // #if VECTORVM_SUPPORTS_LEGACY
+
 #undef VM_FORCEINLINE
diff --git a/Engine/Source/Runtime/VectorVM/Private/VectorVMExperimental.inl b/Engine/Source/Runtime/VectorVM/Private/VectorVMExperimental.inl
--- a/Engine/Source/Runtime/VectorVM/Private/VectorVMExperimental.inl
+++ b/Engine/Source/Runtime/VectorVM/Private/VectorVMExperimental.inl
@@ -285,7 +285,7 @@ static void VVMDefaultFree(void *Ptr, const char *Filename, int LineNumber)
 
 #include "./VectorVMExperimental_Serialization.inl"
 
-#ifdef NIAGARA_EXP_VM
+#if VECTORVM_SUPPORTS_EXPERIMENTAL
 
 #define VVM_CACHELINE_SIZE				64
 #define VVM_CHUNK_FIXED_OVERHEAD_SIZE	512
@@ -1333,7 +1333,7 @@ OpCodeSwitch: //I think computed gotos would be a huge win here... maybe write t
 						{
 							VVMDecodeInstructionRegisters(InsPtr + 2 + i * 2, BatchState->ChunkLocalData.ExtFnDecodedReg.RegData + i, BatchState->ChunkLocalData.ExtFnDecodedReg.RegInc + i);
 						}
-						FVectorVMExternalFunctionContext ExtFnCtx;
+						FVectorVMExternalFunctionContextExperimental ExtFnCtx;
 
 						ExtFnCtx.RegisterData             = (uint32 **)BatchState->ChunkLocalData.ExtFnDecodedReg.RegData;
 						ExtFnCtx.RegInc                   = BatchState->ChunkLocalData.ExtFnDecodedReg.RegInc;
@@ -1363,7 +1363,12 @@ OpCodeSwitch: //I think computed gotos would be a huge win here... maybe write t
 							}
 						}
 						check(DummyRegCount <= VVMState->NumDummyRegsReq);
+#if VECTORVM_SUPPORTS_LEGACY
+						FVectorVMExternalFunctionContext ProxyContext(ExtFnCtx);
+						ExtFnData->Function->Execute(ProxyContext);
+#else
 						ExtFnData->Function->Execute(ExtFnCtx);
+#endif
 					}
 					InsPtr += 2 + ((ExtFnData->NumInputs + ExtFnData->NumOutputs) << 1);
 				}
diff --git a/Engine/Source/Runtime/VectorVM/Private/VectorVMExperimental_Optimizer.inl b/Engine/Source/Runtime/VectorVM/Private/VectorVMExperimental_Optimizer.inl
--- a/Engine/Source/Runtime/VectorVM/Private/VectorVMExperimental_Optimizer.inl
+++ b/Engine/Source/Runtime/VectorVM/Private/VectorVMExperimental_Optimizer.inl
@@ -150,7 +150,70 @@ static void VectorVMFreeOptimizerIntermediateData(FVectorVMOptimizeContext *OptC
 	}
 }
 
-VECTORVM_API void FreeVectorVMOptimizeContext(FVectorVMOptimizeContext *OptContext)
+void FreezeVectorVMOptimizeContext(const FVectorVMOptimizeContext& Context, TArray<uint8>& ContextData)
+{
+	const uint32 BytecodeSize = Context.NumBytecodeBytes + 16;
+	const uint32 ConstRemapSize = Context.NumConstsAlloced * sizeof(uint16);
+	const uint32 ExtFnSize = Context.NumExtFns * sizeof(FVectorVMExtFunctionData);
+
+	const uint32 BytecodeOffset = Align(sizeof(FVectorVMOptimizeContext), 16);
+	const uint32 ConstRemap0Offset = Align(BytecodeOffset + BytecodeSize, 16);
+	const uint32 ConstRemap1Offset = Align(ConstRemap0Offset + ConstRemapSize, 16);
+	const uint32 ExtFnOffset = Align(ConstRemap1Offset + ConstRemapSize, 16);
+	const uint32 TotalSize = Align(ExtFnOffset + ExtFnSize, 16);
+
+	ContextData.SetNumZeroed(TotalSize);
+	uint8* DestData = ContextData.GetData();
+	FMemory::Memcpy(DestData, &Context, sizeof(Context));
+	FMemory::Memcpy(DestData + BytecodeOffset, Context.OutputBytecode, BytecodeSize);
+	FMemory::Memcpy(DestData + ConstRemap0Offset, Context.ConstRemap[0], ConstRemapSize);
+	FMemory::Memcpy(DestData + ConstRemap1Offset, Context.ConstRemap[1], ConstRemapSize);
+	FMemory::Memcpy(DestData + ExtFnOffset, Context.ExtFnTable, ExtFnSize);
+
+	// we want to clear out the pointers to any callbacks
+	FVectorVMOptimizeContext& DestContext = *reinterpret_cast<FVectorVMOptimizeContext*>(DestData);
+	DestContext.Init.ReallocFn = nullptr;
+	DestContext.Init.FreeFn = nullptr;
+	DestContext.Error.CallbackFn = nullptr;
+}
+
+static void* VectorVMFrozenRealloc(void* Ptr, size_t NumBytes, const char* Filename, int LineNum)
+{
+	check(false);
+	return nullptr;
+}
+
+static void VectorVMFrozenFree(void* Ptr, const char* Filename, int LineNum)
+{
+	check(false);
+}
+
+void ReinterpretVectorVMOptimizeContextData(TConstArrayView<uint8> ContextData, FVectorVMOptimizeContext& Context)
+{
+	const uint8* SrcData = ContextData.GetData();
+	const FVectorVMOptimizeContext& SrcContext = *reinterpret_cast<const FVectorVMOptimizeContext*>(SrcData);
+
+	const uint32 BytecodeSize = SrcContext.NumBytecodeBytes + 16;
+	const uint32 ConstRemapSize = SrcContext.NumConstsAlloced * sizeof(uint16);
+	const uint32 ExtFnSize = SrcContext.NumExtFns * sizeof(FVectorVMExtFunctionData);
+
+	const uint32 BytecodeOffset = Align(sizeof(FVectorVMOptimizeContext), 16);
+	const uint32 ConstRemap0Offset = Align(BytecodeOffset + BytecodeSize, 16);
+	const uint32 ConstRemap1Offset = Align(ConstRemap0Offset + ConstRemapSize, 16);
+	const uint32 ExtFnOffset = Align(ConstRemap1Offset + ConstRemapSize, 16);
+	const uint32 TotalSize = Align(ExtFnOffset + ExtFnSize, 16);
+
+	FMemory::Memcpy(&Context, SrcData, sizeof(Context));
+	Context.OutputBytecode = const_cast<uint8*>(reinterpret_cast<const uint8*>(SrcData + BytecodeOffset));
+	Context.ConstRemap[0] = const_cast<uint16*>(reinterpret_cast<const uint16*>(SrcData + ConstRemap0Offset));
+	Context.ConstRemap[1] = const_cast<uint16*>(reinterpret_cast<const uint16*>(SrcData + ConstRemap1Offset));
+	Context.ExtFnTable = const_cast<FVectorVMExtFunctionData*>(reinterpret_cast<const FVectorVMExtFunctionData*>(SrcData + ExtFnOffset));
+	Context.Init.ReallocFn = VectorVMFrozenRealloc;
+	Context.Init.FreeFn = VectorVMFrozenFree;
+	Context.Error.CallbackFn = nullptr;
+}
+
+void FreeVectorVMOptimizeContext(FVectorVMOptimizeContext *OptContext)
 {
 	//save init data
 	VectorVMReallocFn *ReallocFn = OptContext->Init.ReallocFn;
diff --git a/Engine/Source/Runtime/VectorVM/Private/VectorVMExperimental_Serialization.inl b/Engine/Source/Runtime/VectorVM/Private/VectorVMExperimental_Serialization.inl
--- a/Engine/Source/Runtime/VectorVM/Private/VectorVMExperimental_Serialization.inl
+++ b/Engine/Source/Runtime/VectorVM/Private/VectorVMExperimental_Serialization.inl
@@ -91,7 +91,7 @@ static uint32 VectorVMSerializeSetError_(FVectorVMSerializeState *SerializeState
 #define VectorVMSerializeSetError(SerializeState, Flags)     VectorVMSerializeSetError_(SerializeState, Flags, __LINE__)
 
 
-#ifdef NIAGARA_EXP_VM
+#if VECTORVM_SUPPORTS_EXPERIMENTAL
 
 /*********************************************************************************************************************************************************************************************************************************
 *** Serialization stuff for the Experimental VM
@@ -296,7 +296,7 @@ static uint32 VVMSer_initSerializationState_(FVectorVMState *VVMState, FVectorVM
 }
 #endif //VVM_SERIALIZE_NO_WRITE
 
-#else //NIAGARA_EXP_VM
+#else // VECTORVM_SUPPORTS_EXPERIMENTAL
 
 /*********************************************************************************************************************************************************************************************************************************
 *** Serialization stuff for the old VM
@@ -373,7 +373,7 @@ static void VVMSer_serializeInstruction(FVectorVMSerializeState *SerializeState,
 }
 #endif //VVM_SERIALIZE_NO_WRITE
 
-#endif //NIAGARA_EXP_VM
+#endif // VECTORVM_SUPPORTS_EXPERIMENTAL
 
 /*********************************************************************************************************************************************************************************************************************************
 *** Stuff shared between the new and old VM
diff --git a/Engine/Source/Runtime/VectorVM/Public/VectorVM.h b/Engine/Source/Runtime/VectorVM/Public/VectorVM.h
--- a/Engine/Source/Runtime/VectorVM/Public/VectorVM.h
+++ b/Engine/Source/Runtime/VectorVM/Public/VectorVM.h
@@ -10,276 +10,67 @@
 #include "Templates/AlignmentTemplates.h"
 #include "VectorVMCommon.h"
 #include "VectorVMExperimental.h"
+#include "VectorVMLegacy.h"
 
-//Data the VM will keep on each dataset locally per thread which is then thread safely pushed to it's destination at the end of execution.
-struct FDataSetThreadLocalTempData
-{
-	FDataSetThreadLocalTempData()
-	{
-		Reset();
-	}
-
-	FORCEINLINE void Reset()
-	{
-		IDsToFree.Reset();
-		MaxID = INDEX_NONE;
-	}
+#if VECTORVM_SUPPORTS_EXPERIMENTAL && VECTORVM_SUPPORTS_LEGACY
 
-	TArray<int32> IDsToFree;
-	int32 MaxID;
-
-	//TODO: Possibly store output data locally and memcpy to the real buffers. Could avoid false sharing in parallel execution and so improve perf.
-	//using _mm_stream_ps on platforms that support could also work for this?
-	//TArray<TArray<float>> OutputFloatData;
-	//TArray<TArray<int32>> OutputIntData;
-};
-
-/**
-* Context information passed around during VM execution.
-*/
-struct FVectorVMContext : TThreadSingleton<FVectorVMContext>
+class FVectorVMExternalFunctionContextProxy
 {
-private:
-
-	friend struct FVectorVMCodeOptimizerContext;
-
 public:
-
-	/** Pointer to the next element in the byte code. */
-	uint8 const* RESTRICT Code;
-
-	/** Pointer to the constant table. */
-	const uint8* const* RESTRICT ConstantTable;
-	const int32* ConstantTableSizes;
-	int32 ConstantTableCount;
-
-	/** Num temp registers required by this script. */
-	int32 NumTempRegisters;
-
-	/** Pointer to the shared data table. */
-	const FVMExternalFunction* const* RESTRICT ExternalFunctionTable;
-	/** Table of user pointers.*/
-	void** UserPtrTable;
-
-	/** Number of instances to process. */
-	int32 NumInstances;
-	/** Number of instances to process when doing batches of VECTOR_WIDTH_FLOATS. */
-	int32 NumInstancesVectorFloats;
-	/** Start instance of current chunk. */
-	int32 StartInstance;
-	
-	/** HACK: An additional instance offset to allow external functions direct access to specific instances in the buffers. */
-	int32 ExternalFunctionInstanceOffset;
-
-	/** Array of meta data on data sets. TODO: This struct should be removed and all features it contains be handled by more general vm ops and the compiler's knowledge of offsets etc. */
-	TArrayView<FDataSetMeta> DataSetMetaTable;
-
-	TArray<FDataSetThreadLocalTempData> ThreadLocalTempData;
-
-#if STATS
-	TArray<FStatStackEntry, TInlineAllocator<64>> StatCounterStack;
-	TArrayView<FStatScopeData> StatScopes;
-	TArray<uint64, TInlineAllocator<64>> ScopeExecCycles;
-#elif ENABLE_STATNAMEDEVENTS
-	TArrayView<const FString> StatNamedEventScopes;
-#endif
-
-	TArray<uint8, TAlignedHeapAllocator<VECTOR_WIDTH_BYTES>> TempRegTable;
-	uint32 TempRegisterSize;
-	uint32 TempBufferSize;
-
-	/** Thread local random stream for use in external functions needing non-deterministic randoms. */
-	FRandomStream RandStream;
-
-	/** Thread local per instance random counters for use in external functions needing deterministic randoms. */
-	TArray<int32> RandCounters;
-
-	bool bIsParallelExecution;
-
-	int32 ValidInstanceIndexStart = INDEX_NONE;
-	int32 ValidInstanceCount = 0;
-	bool ValidInstanceUniform = false;
-
-	FVectorVMContext();
-
-	void PrepareForExec(
-		int32 InNumTempRegisters,
-		int32 ConstantTableCount,
-		const uint8* const* InConstantTables,
-		const int32* InConstantTableSizes,
-		const FVMExternalFunction* const* InExternalFunctionTable,
-		void** InUserPtrTable,
-		TArrayView<FDataSetMeta> InDataSetMetaTable,
-		int32 MaxNumInstances,
-		bool bInParallelExecution);
-
-#if STATS
-	void SetStatScopes(TArrayView<FStatScopeData> InStatScopes);
-#elif ENABLE_STATNAMEDEVENTS
-	void SetStatNamedEventScopes(TArrayView<const FString> InStatNamedEventScopes);
-#endif
-
-	void FinishExec();
-
-	void PrepareForChunk(const uint8* InCode, int32 InNumInstances, int32 InStartInstance)
-	{
-		Code = InCode;
-		NumInstances = InNumInstances;
-		NumInstancesVectorFloats = (NumInstances + VECTOR_WIDTH_FLOATS - 1) / VECTOR_WIDTH_FLOATS;
-		StartInstance = InStartInstance;
-		
-		ExternalFunctionInstanceOffset = 0;
-
-		ValidInstanceCount = 0;
-		ValidInstanceIndexStart = INDEX_NONE;
-		ValidInstanceUniform = false;
-
-		RandCounters.Reset();
-		RandCounters.SetNumZeroed(InNumInstances);
-	}
-
-	FORCEINLINE FDataSetMeta& GetDataSetMeta(int32 DataSetIndex) { return DataSetMetaTable[DataSetIndex]; }
-	FORCEINLINE uint8 * RESTRICT GetTempRegister(int32 RegisterIndex) { return TempRegTable.GetData() + TempRegisterSize * RegisterIndex; }
-	template<typename T, int TypeOffset>
-	FORCEINLINE T* RESTRICT GetInputRegister(int32 DataSetIndex, int32 RegisterIndex) 
-	{
-		FDataSetMeta& Meta = GetDataSetMeta(DataSetIndex);
-		uint32 Offset = Meta.InputRegisterTypeOffsets[TypeOffset];
-		return ((T*)Meta.InputRegisters[Offset + RegisterIndex]) + Meta.InstanceOffset;
-	}
-	template<typename T, int TypeOffset>
-	FORCEINLINE T* RESTRICT GetOutputRegister(int32 DataSetIndex, int32 RegisterIndex) 
-	{ 
-		FDataSetMeta& Meta = GetDataSetMeta(DataSetIndex);
-		uint32 Offset = Meta.OutputRegisterTypeOffsets[TypeOffset];
-		return  ((T*)Meta.OutputRegisters[Offset + RegisterIndex]) + Meta.InstanceOffset;
-	}
-
-	int32 GetNumInstances() const { return NumInstances; }
-	int32 GetStartInstance() const { return StartInstance; }
-
-	template<uint32 InstancesPerOp>
-	int32 GetNumLoops() const { return (InstancesPerOp == VECTOR_WIDTH_FLOATS) ? NumInstancesVectorFloats : ((InstancesPerOp == 1) ? NumInstances : Align(NumInstances, InstancesPerOp));	}
-
-	FORCEINLINE uint8 DecodeU8() { return *Code++; }
-#if PLATFORM_SUPPORTS_UNALIGNED_LOADS
-	FORCEINLINE uint16 DecodeU16() { uint16 v = *reinterpret_cast<const uint16*>(Code); Code += sizeof(uint16); return INTEL_ORDER16(v); }
-	FORCEINLINE uint32 DecodeU32() { uint32 v = *reinterpret_cast<const uint32*>(Code); Code += sizeof(uint32); return INTEL_ORDER32(v); }
-	FORCEINLINE uint64 DecodeU64() { uint64 v = *reinterpret_cast<const uint64*>(Code); Code += sizeof(uint64); return INTEL_ORDER64(v); }
-#else
-	FORCEINLINE uint16 DecodeU16() { uint16 v = Code[1]; v = v << 8 | Code[0]; Code += 2; return INTEL_ORDER16(v); }
-	FORCEINLINE uint32 DecodeU32() { uint32 v = Code[3]; v = v << 8 | Code[2]; v = v << 8 | Code[1]; v = v << 8 | Code[0]; Code += 4; return INTEL_ORDER32(v); }
-	FORCEINLINE uint64 DecodeU64() { uint64 v = Code[7]; v = v << 8 | Code[6]; v = v << 8 | Code[5]; v = v << 8 | Code[4]; v = v << 8 | Code[3]; v = v << 8 | Code[2]; v = v << 8 | Code[1]; v = v << 8 | Code[0]; Code += 8; return INTEL_ORDER64(v); }
-#endif
-	FORCEINLINE uintptr_t DecodePtr() { return (sizeof(uintptr_t) == 4) ? DecodeU32() : DecodeU64(); }
-	FORCEINLINE void SkipCode(int64 Count) { Code += Count; }
-
-	/** Decode the next operation contained in the bytecode. */
-	FORCEINLINE EVectorVMOp DecodeOp()
+	const bool UsingExperimentalVM;
+	FVectorVMExternalFunctionContextExperimental Experimental;
+	FVectorVMExternalFunctionContextLegacy Legacy;
+
+	FVectorVMExternalFunctionContextProxy(const FVectorVMExternalFunctionContextExperimental& InExperimental)
+		: UsingExperimentalVM(true)
+		, Experimental(InExperimental)
+		, Legacy()
+	{}
+
+	FVectorVMExternalFunctionContextProxy(const FVectorVMExternalFunctionContextLegacy& InLegacy)
+		: UsingExperimentalVM(false)
+		, Experimental()
+		, Legacy(InLegacy)
+	{}
+
+	FORCEINLINE int32* GetRandCounters()
 	{
-		return static_cast<EVectorVMOp>(DecodeU8());
+		return UsingExperimentalVM ? Experimental.GetRandCounters() : Legacy.GetRandCounters().GetData();
 	}
 
-	FORCEINLINE uint8 DecodeSrcOperandTypes()
+	FORCEINLINE FRandomStream& GetRandStream()
 	{
-		return DecodeU8();
+		return UsingExperimentalVM ? Experimental.GetRandStream() : Legacy.GetRandStream();
 	}
 
-	FORCEINLINE bool IsParallelExecution()
+	FORCEINLINE int32 GetNumInstances() const
 	{
-		return bIsParallelExecution;
+		return UsingExperimentalVM ? Experimental.GetNumInstances() : Legacy.GetNumInstances();
 	}
 
-	template<typename T = uint8>
-	FORCEINLINE const T* GetConstant(int32 TableIndex, int32 TableOffset) const
+	template<uint32 InstancesPerOp>
+	FORCEINLINE int32 GetNumLoops() const
 	{
-		check(TableIndex < ConstantTableCount);
-		return reinterpret_cast<const T*>(ConstantTable[TableIndex] + TableOffset);
+		return UsingExperimentalVM ? Experimental.GetNumLoops<InstancesPerOp>() : Legacy.GetNumLoops<InstancesPerOp>();
 	}
 
-	template<typename T = uint8>
-	FORCEINLINE const T* GetConstant(int32 Offset) const
-	{
-		int32 TableIndex = 0;
-
-		while (Offset >= ConstantTableSizes[TableIndex])
-		{
-			Offset -= ConstantTableSizes[TableIndex];
-			++TableIndex;
-		}
-
-		check(TableIndex < ConstantTableCount);
-		check(Offset < ConstantTableSizes[TableIndex]);
-		return reinterpret_cast<const T*>(ConstantTable[TableIndex] + Offset);
-	}
+private:
+	FVectorVMExternalFunctionContextProxy() = delete;
 };
 
-class FVectorVMExternalFunctionContext {
-#ifdef NIAGARA_EXP_VM
-public:
-	uint32 **                RegisterData;
-	uint16 *                 RawVecIndices; //undecoded, for compatbility with the previous VM
-	uint32 *                 RegInc;
+using FVectorVMExternalFunctionContext = FVectorVMExternalFunctionContextProxy;
 
-	int                      RegReadCount;
-	int                      NumRegisters;
+#elif VECTORVM_SUPPORTS_EXPERIMENTAL
 
-	int                      StartInstance;
-	int                      NumInstances;
-	int                      NumLoops;
-	int                      PerInstanceFnInstanceIdx;
+using FVectorVMExternalFunctionContext = FVectorVMExternalFunctionContextExperimental;
 
-	void **                  UserPtrTable;
-	int                      NumUserPtrs;
-	
-	FRandomStream *          RandStream;
-	int32 *                  RandCounters;
-	TArrayView<FDataSetMeta> DataSets;
-
-	FORCEINLINE int32                                  GetStartInstance() const             { return StartInstance; }
-	FORCEINLINE int32                                  GetNumInstances() const              { return NumInstances; }
-	FORCEINLINE int32 *                                GetRandCounters()                    { return RandCounters; }
-	FORCEINLINE FRandomStream &                        GetRandStream()                      { return *RandStream; }
-	FORCEINLINE void *                                 GetUserPtrTable(int32 UserPtrIdx)    { check(UserPtrIdx < NumUserPtrs);  return UserPtrTable[UserPtrIdx]; }
-	template<uint32 InstancesPerOp> FORCEINLINE int32  GetNumLoops() const                  { static_assert(InstancesPerOp == 4); return NumLoops; };
-
-	FORCEINLINE float *GetNextRegister(int32 *OutAdvanceOffset, int32 *OutVecIndex) {
-		check(RegReadCount < NumRegisters);
-		*OutAdvanceOffset = RegInc[RegReadCount] & 1;
-		*OutVecIndex      = RawVecIndices[RegReadCount];
-		return (float *)RegisterData[RegReadCount++];
-	}
+#elif VECTORVM_SUPPORTS_LEGACY
 
-#else
-	friend struct FNiagaraSystemScriptExecutionContext; //@NOTE(smcgrath): required for the PerInstanceFunctionHook() in the non-experimental version of VectorVM
-public:
-	FVectorVMExternalFunctionContext(FVectorVMContext *InVectorVMContext) : VectorVMContext(InVectorVMContext) { }
-	FORCEINLINE int32 GetStartInstance() const { return VectorVMContext->GetStartInstance(); }
-	FORCEINLINE int32 GetNumInstances() const { return VectorVMContext->GetNumInstances(); }
-	FORCEINLINE TArray<int32> &GetRandCounters() { return VectorVMContext->RandCounters; }
-	FORCEINLINE FRandomStream &GetRandStream() { return VectorVMContext->RandStream; }
-
-#ifndef NIAGARA_EXP_VM
-	FORCEINLINE uint8 *RESTRICT GetTempRegister(int32 RegisterIndex) { return VectorVMContext->GetTempRegister(RegisterIndex); }
-	FORCEINLINE uint8 DecodeU8() { return VectorVMContext->DecodeU8(); }
-	FORCEINLINE uint16 DecodeU16() { return VectorVMContext->DecodeU16(); }
-	FORCEINLINE uint32 DecodeU32() { return VectorVMContext->DecodeU32(); }
-	FORCEINLINE uint64 DecodeU64() { return VectorVMContext->DecodeU64(); }
-	template<typename T = uint8>
-	FORCEINLINE const T* GetConstant(int32 TableIndex, int32 TableOffset) const { return VectorVMContext->GetConstant<T>(TableIndex, TableOffset); }
-	template<typename T = uint8>
-	FORCEINLINE const T* GetConstant(int32 Offset) const { return VectorVMContext->GetConstant<T>(Offset); }
-	FORCEINLINE void *GetUserPtrTable(int32 UserPtrIdx) { return VectorVMContext->UserPtrTable[UserPtrIdx]; }
-	FORCEINLINE int32 GetExternalFunctionInstanceOffset() const { return VectorVMContext->ExternalFunctionInstanceOffset; }
-	FORCEINLINE int32 GetNumTempRegisters() const { return VectorVMContext->NumTempRegisters; }
-#endif
+using FVectorVMExternalFunctionContext = FVectorVMExternalFunctionContextLegacy;
 
-	template<uint32 InstancesPerOp>
-	FORCEINLINE int32 GetNumLoops() const { return VectorVMContext->GetNumLoops<InstancesPerOp>(); }
-private:
-	FVectorVMContext *VectorVMContext;
+#else
+#error "At least one of VECTORVM_SUPPORTS_EXPERIMENTAL | VECTORVM_SUPPORTS_LEGACY must be defined"
 #endif
-};
 
 namespace VectorVM
 {
@@ -313,12 +104,13 @@ namespace VectorVM
 #endif
 	};
 
+#if VECTORVM_SUPPORTS_LEGACY
 	/**
 	 * Execute VectorVM bytecode.
 	 */
 	VECTORVM_API void Exec(FVectorVMExecArgs& Args, FVectorVMSerializeState *SerializeState);
-
 	VECTORVM_API void OptimizeByteCode(const uint8* ByteCode, TArray<uint8>& OptimizedCode, TArrayView<uint8> ExternalFunctionRegisterCounts);
+#endif
 
 	VECTORVM_API void Init();
 
@@ -332,23 +124,47 @@ namespace VectorVM
 		T* Ptr;
 		FUserPtrHandler(FVectorVMExternalFunctionContext& Context)
 		{
-#ifdef NIAGARA_EXP_VM
+#if VECTORVM_SUPPORTS_EXPERIMENTAL && VECTORVM_SUPPORTS_LEGACY
+			if (Context.UsingExperimentalVM)
+			{
+				int32 AdvanceOffset;
+				int32 RegIdx;
+				int32* ConstPtr = (int32*)Context.Experimental.GetNextRegister(&AdvanceOffset, &RegIdx);
+				check(AdvanceOffset == 0); //must be constant
+				UserPtrIdx = *ConstPtr;
+				check(UserPtrIdx != INDEX_NONE);
+				Ptr = (T*)Context.Experimental.GetUserPtrTable(UserPtrIdx);
+			}
+			else
+			{
+				const uint16 VariableOffset = Context.Legacy.DecodeU16();
+				check(!(VariableOffset& VVM_EXT_FUNC_INPUT_LOC_BIT));
+
+				const uint16 ConstantTableOffset = VariableOffset & VVM_EXT_FUNC_INPUT_LOC_MASK;
+				UserPtrIdx = *Context.Legacy.GetConstant<int32>(ConstantTableOffset);
+				check(UserPtrIdx != INDEX_NONE);
+
+				Ptr = static_cast<T*>(Context.Legacy.GetUserPtrTable(UserPtrIdx));
+			}
+#elif VECTORVM_SUPPORTS_EXPERIMENTAL
 			int32 AdvanceOffset;
 			int32 RegIdx;
-			int32 *ConstPtr = (int32 *)Context.GetNextRegister(&AdvanceOffset, &RegIdx);
+			int32* ConstPtr = (int32*)Context.GetNextRegister(&AdvanceOffset, &RegIdx);
 			check(AdvanceOffset == 0); //must be constant
 			UserPtrIdx = *ConstPtr;
 			check(UserPtrIdx != INDEX_NONE);
-			Ptr = (T *)Context.GetUserPtrTable(UserPtrIdx);
-#else
+			Ptr = (T*)Context.GetUserPtrTable(UserPtrIdx);
+#elif VECTORVM_SUPPORTS_LEGACY
 			const uint16 VariableOffset = Context.DecodeU16();
 			check(!(VariableOffset & VVM_EXT_FUNC_INPUT_LOC_BIT));
 
 			const uint16 ConstantTableOffset = VariableOffset & VVM_EXT_FUNC_INPUT_LOC_MASK;
 			UserPtrIdx = *Context.GetConstant<int32>(ConstantTableOffset);
 			check(UserPtrIdx != INDEX_NONE);
-			
+
 			Ptr = static_cast<T*>(Context.GetUserPtrTable(UserPtrIdx));
+#else 
+	#error "Not supported"
 #endif
 		}
 
@@ -370,12 +186,18 @@ namespace VectorVM
 		int32 InputOffset;
 		const T* RESTRICT InputPtr;
 		int32 AdvanceOffset;
+#if VECTORVM_SUPPORTS_EXPERIMENTAL && VECTORVM_SUPPORTS_LEGACY
+		bool bIsRegister;
+#endif
 
 	public:
 		FExternalFuncInputHandler()
 			: InputOffset(INDEX_NONE)
 			, InputPtr(nullptr)
 			, AdvanceOffset(0)
+#if VECTORVM_SUPPORTS_EXPERIMENTAL && VECTORVM_SUPPORTS_LEGACY
+			, bIsRegister(false)
+#endif
 		{}
 
 		FORCEINLINE FExternalFuncInputHandler(FVectorVMExternalFunctionContext& Context)
@@ -385,9 +207,27 @@ namespace VectorVM
 
 		void Init(FVectorVMExternalFunctionContext& Context)
 		{
-#ifdef NIAGARA_EXP_VM
+#if VECTORVM_SUPPORTS_EXPERIMENTAL && VECTORVM_SUPPORTS_LEGACY
+			if (Context.UsingExperimentalVM)
+			{
+				InputPtr = (T*)Context.Experimental.GetNextRegister(&AdvanceOffset, &InputOffset) + Context.Experimental.PerInstanceFnInstanceIdx;
+				bIsRegister = !!AdvanceOffset;
+			}
+			else
+			{
+				InputOffset = Context.Legacy.DecodeU16();
+
+				const int32 Offset = GetOffset();
+				InputPtr = IsConstant() ? Context.Legacy.GetConstant<T>(Offset) : reinterpret_cast<T*>(Context.Legacy.GetTempRegister(Offset));
+				AdvanceOffset = IsConstant() ? 0 : 1;
+
+				//Hack: Offset into the buffer by the instance offset.
+				InputPtr += Context.Legacy.GetExternalFunctionInstanceOffset() * AdvanceOffset;
+				bIsRegister = !!(InputOffset & VVM_EXT_FUNC_INPUT_LOC_BIT);
+			}
+#elif VECTORVM_SUPPORTS_EXPERIMENTAL
 			InputPtr = (T*)Context.GetNextRegister(&AdvanceOffset, &InputOffset) + Context.PerInstanceFnInstanceIdx;
-#else
+#elif VECTORVM_SUPPORTS_LEGACY
 			InputOffset = Context.DecodeU16();
 
 			const int32 Offset = GetOffset();
@@ -396,14 +236,20 @@ namespace VectorVM
 
 			//Hack: Offset into the buffer by the instance offset.
 			InputPtr += Context.GetExternalFunctionInstanceOffset() * AdvanceOffset;
+#else
+	#error "Not supported"
 #endif
 		}
 
 		FORCEINLINE bool IsConstant()const { return !IsRegister(); }
-#ifdef NIAGARA_EXP_VM
+#if VECTORVM_SUPPORTS_EXPERIMENTAL && VECTORVM_SUPPORTS_LEGACY
+		FORCEINLINE bool IsRegister() const { return bIsRegister; }
+#elif VECTORVM_SUPPORTS_EXPERIMENTAL
 		FORCEINLINE bool IsRegister() const { return (bool)AdvanceOffset; }
+#elif VECTORVM_SUPPORTS_LEGACY
+		FORCEINLINE bool IsRegister() const { return (InputOffset & VVM_EXT_FUNC_INPUT_LOC_BIT) != 0; }
 #else
-		FORCEINLINE bool IsRegister()const { return (InputOffset & VVM_EXT_FUNC_INPUT_LOC_BIT) != 0; }
+	#error "Not supported"
 #endif
 		FORCEINLINE int32 GetOffset()const { return InputOffset & VVM_EXT_FUNC_INPUT_LOC_MASK; }
 
@@ -433,11 +279,38 @@ namespace VectorVM
 		T Dummy;
 		T* RESTRICT Register;
 	public:
-#ifdef NIAGARA_EXP_VM
+#if VECTORVM_SUPPORTS_EXPERIMENTAL && VECTORVM_SUPPORTS_LEGACY
+		FORCEINLINE FExternalFuncRegisterHandler(FVectorVMExternalFunctionContext& Context)
+		{
+			if (Context.UsingExperimentalVM)
+			{
+				Register = (T*)Context.Experimental.GetNextRegister(&AdvanceOffset, &RegisterIndex) + Context.Experimental.PerInstanceFnInstanceIdx;
+			}
+			else
+			{
+				RegisterIndex = (Context.Legacy.DecodeU16() & VVM_EXT_FUNC_INPUT_LOC_MASK);
+				AdvanceOffset = (IsValid() ? 1 : 0);
+				{
+					if (IsValid())
+					{
+						checkSlow(RegisterIndex < Context.Legacy.GetNumTempRegisters());
+						Register = (T*)Context.Legacy.GetTempRegister(RegisterIndex);
+
+						//Hack: Offset into the buffer by the instance offset.
+						Register += Context.Legacy.GetExternalFunctionInstanceOffset() * AdvanceOffset;
+					}
+					else
+					{
+						Register = &Dummy;
+					}
+				}
+			}
+		}
+#elif VECTORVM_SUPPORTS_EXPERIMENTAL
 		FORCEINLINE FExternalFuncRegisterHandler(FVectorVMExternalFunctionContext& Context) {
 			Register = (T*)Context.GetNextRegister(&AdvanceOffset, &RegisterIndex) + Context.PerInstanceFnInstanceIdx;
 		}
-#else
+#elif VECTORVM_SUPPORTS_LEGACY
 		FORCEINLINE FExternalFuncRegisterHandler(FVectorVMExternalFunctionContext& Context)
 			: RegisterIndex(Context.DecodeU16() & VVM_EXT_FUNC_INPUT_LOC_MASK)
 			, AdvanceOffset(IsValid() ? 1 : 0)
@@ -455,7 +328,10 @@ namespace VectorVM
 				Register = &Dummy;
 			}
 		}
+#else
+	#error "Not supported"
 #endif
+
 		FORCEINLINE bool IsValid() const { return RegisterIndex != (uint16)VVM_EXT_FUNC_INPUT_LOC_MASK; }
 
 		FORCEINLINE const T Get() { return *Register; }
@@ -480,11 +356,26 @@ namespace VectorVM
 	{
 		uint16 ConstantIndex;
 		T Constant;
-#ifdef NIAGARA_EXP_VM
-		FExternalFuncConstHandler(FVectorVMExternalFunctionContext& Context) {
+
+#if VECTORVM_SUPPORTS_EXPERIMENTAL && VECTORVM_SUPPORTS_LEGACY
+		FExternalFuncConstHandler(FVectorVMExternalFunctionContext& Context)
+		{
+			if (Context.UsingExperimentalVM)
+			{
+				check(false);
+			}
+			else
+			{
+				ConstantIndex = (Context.Legacy.DecodeU16() & VVM_EXT_FUNC_INPUT_LOC_MASK);
+				Constant = (*Context.Legacy.GetConstant<T>(ConstantIndex));
+			}
+		}
+#elif VECTORVM_SUPPORTS_EXPERIMENTAL
+		FExternalFuncConstHandler(FVectorVMExternalFunctionContext& Context)
+		{
 			check(false);
 		}
-#else
+#elif VECTORVM_SUPPORTS_LEGACY
 		FExternalFuncConstHandler(FVectorVMExternalFunctionContext& Context)
 			: ConstantIndex(Context.DecodeU16() & VVM_EXT_FUNC_INPUT_LOC_MASK)
 			, Constant(*Context.GetConstant<T>(ConstantIndex))
diff --git a/Engine/Source/Runtime/VectorVM/Public/VectorVMCommon.h b/Engine/Source/Runtime/VectorVM/Public/VectorVMCommon.h
--- a/Engine/Source/Runtime/VectorVM/Public/VectorVMCommon.h
+++ b/Engine/Source/Runtime/VectorVM/Public/VectorVMCommon.h
@@ -10,8 +10,6 @@
 #define VECTOR_WIDTH_BYTES (16)
 #define VECTOR_WIDTH_FLOATS (4)
 
-DECLARE_DELEGATE_OneParam(FVMExternalFunction, class FVectorVMExternalFunctionContext& /*Context*/);
-
 UENUM()
 enum class EVectorVMBaseTypes : uint8
 {
@@ -251,11 +249,11 @@ struct FDataSetMeta
 	//UE-65856 for tracking this work.
 
 	//@NOTE (smcgrath): the lock/unlock functions can go with the ne VM, we don't need them
-#ifndef NIAGARA_EXP_VM
+#if VECTORVM_SUPPORTS_LEGACY
 	FCriticalSection FreeTableLock;
-#endif
 	FORCEINLINE void LockFreeTable();
 	FORCEINLINE void UnlockFreeTable();
+#endif
 
 	FDataSetMeta()
 		: InputRegisterTypeOffsets{}
@@ -308,3 +306,23 @@ private:
 	FDataSetMeta& operator=(FDataSetMeta&&) = delete;
 	FDataSetMeta& operator=(const FDataSetMeta&) = delete;
 };
+
+#if VECTORVM_SUPPORTS_EXPERIMENTAL && VECTORVM_SUPPORTS_LEGACY
+
+class FVectorVMExternalFunctionContextProxy;
+DECLARE_DELEGATE_OneParam(FVMExternalFunction, FVectorVMExternalFunctionContextProxy& /*Context*/);
+
+#elif VECTORVM_SUPPORTS_EXPERIMENTAL
+
+class FVectorVMExternalFunctionContextExperimental;
+DECLARE_DELEGATE_OneParam(FVMExternalFunction, FVectorVMExternalFunctionContextExperimental& /*Context*/);
+
+#elif VECTORVM_SUPPORTS_LEGACY
+
+class FVectorVMExternalFunctionContextLegacy;
+DECLARE_DELEGATE_OneParam(FVMExternalFunction, FVectorVMExternalFunctionContextLegacy& /*Context*/);
+
+#else
+#error "At least one of VECTORVM_SUPPORTS_EXPERIMENTAL | VECTORVM_SUPPORTS_LEGACY must be defined"
+#endif
+
diff --git a/Engine/Source/Runtime/VectorVM/Public/VectorVMExperimental.h b/Engine/Source/Runtime/VectorVM/Public/VectorVMExperimental.h
--- a/Engine/Source/Runtime/VectorVM/Public/VectorVMExperimental.h
+++ b/Engine/Source/Runtime/VectorVM/Public/VectorVMExperimental.h
@@ -8,7 +8,7 @@ typedef void   (VectorVMFreeFn)                 (void *Ptr, const char *Filename
 struct FDataSetMeta;
 
 //prototypes for serialization that are required whether or not serialization is enabled
-#if defined(NIAGARA_EXP_VM) || defined(VVM_INCLUDE_SERIALIZATION)
+#if VECTORVM_SUPPORTS_EXPERIMENTAL || defined(VVM_INCLUDE_SERIALIZATION)
 struct FVectorVMConstData
 {
 	void * RegisterData;
@@ -19,7 +19,7 @@ struct FVectorVMConstData
 typedef uint32 (VectorVMOptimizeErrorCallback)  (struct FVectorVMOptimizeContext *OptimizeContext, uint32 ErrorFlags);	//return new error flags
 typedef uint32 (VectorVMSerializeErrorCallback) (struct FVectorVMSerializeState *SerializeState, uint32 ErrorFlags);
 
-#endif //NIAGARA_EXP_VM || VVM_INCLUDE_SERIALIZATION
+#endif // VECTORVM_SUPPORTS_EXPERIMENTAL || VVM_INCLUDE_SERIALIZATION
 
 //Serialization
 enum EVectorVMSerializeFlags
@@ -176,7 +176,7 @@ struct FVectorVMSerializeState
 
 #endif // VVM_INCLUDE_SERIALIZATION
 
-#ifdef NIAGARA_EXP_VM
+#if VECTORVM_SUPPORTS_EXPERIMENTAL
 
 union FVecReg {
 	VectorRegister4f v;
@@ -337,7 +337,7 @@ struct FVectorVMExternalFnPerInstanceData
 struct FVectorVMInitData
 {
 	struct FVectorVMState *					ExistingVectorVMState;
-	FVectorVMOptimizeContext *              OptimizeContext;
+	const FVectorVMOptimizeContext *              OptimizeContext;
 	TArrayView<FDataSetMeta>                DataSets;
 	TArrayView<const FVMExternalFunction *> ExtFunctionTable;
 	
@@ -425,6 +425,44 @@ struct FVectorVMState
 	} Error;
 };
 
+class FVectorVMExternalFunctionContextExperimental
+{
+public:
+	uint32** RegisterData;
+	uint16* RawVecIndices; //undecoded, for compatbility with the previous VM
+	uint32* RegInc;
+
+	int                      RegReadCount;
+	int                      NumRegisters;
+
+	int                      StartInstance;
+	int                      NumInstances;
+	int                      NumLoops;
+	int                      PerInstanceFnInstanceIdx;
+
+	void** UserPtrTable;
+	int                      NumUserPtrs;
+
+	FRandomStream* RandStream;
+	int32* RandCounters;
+	TArrayView<FDataSetMeta> DataSets;
+
+	FORCEINLINE int32                                  GetStartInstance() const { return StartInstance; }
+	FORCEINLINE int32                                  GetNumInstances() const { return NumInstances; }
+	FORCEINLINE int32* GetRandCounters() { return RandCounters; }
+	FORCEINLINE FRandomStream& GetRandStream() { return *RandStream; }
+	FORCEINLINE void* GetUserPtrTable(int32 UserPtrIdx) { check(UserPtrIdx < NumUserPtrs);  return UserPtrTable[UserPtrIdx]; }
+	template<uint32 InstancesPerOp> FORCEINLINE int32  GetNumLoops() const { static_assert(InstancesPerOp == 4); return NumLoops; };
+
+	FORCEINLINE float* GetNextRegister(int32* OutAdvanceOffset, int32* OutVecIndex)
+	{
+		check(RegReadCount < NumRegisters);
+		*OutAdvanceOffset = RegInc[RegReadCount] & 1;
+		*OutVecIndex = RawVecIndices[RegReadCount];
+		return (float*)RegisterData[RegReadCount++];
+	}
+};
+
 //API FUNCTIONS
 
 //normal functions
@@ -434,8 +472,10 @@ VECTORVM_API void             ExecVectorVMState     (FVectorVMState *VectorVMSta
 VECTORVM_API int              GetNumOutputInstances (FVectorVMState *VectorVMState, int DataSetIdx);
 
 //optimize functions
-VECTORVM_API uint32  OptimizeVectorVMScript      (const uint8 *Bytecode, int BytecodeLen, FVectorVMExtFunctionData *ExtFnIOData, int NumExtFns, FVectorVMOptimizeContext *OptContext, uint32 Flags); //OutContext must be zeroed except the Init struct
-VECTORVM_API void    FreeVectorVMOptimizeContext (FVectorVMOptimizeContext *Context);
+VECTORVM_API uint32  OptimizeVectorVMScript                (const uint8 *Bytecode, int BytecodeLen, FVectorVMExtFunctionData *ExtFnIOData, int NumExtFns, FVectorVMOptimizeContext *OptContext, uint32 Flags); //OutContext must be zeroed except the Init struct
+VECTORVM_API void    FreeVectorVMOptimizeContext           (FVectorVMOptimizeContext *Context);
+VECTORVM_API void    FreezeVectorVMOptimizeContext         (const FVectorVMOptimizeContext& Context, TArray<uint8>& ContextData);
+VECTORVM_API void    ReinterpretVectorVMOptimizeContextData(TConstArrayView<uint8> ContextData, FVectorVMOptimizeContext& Context);
 
 //serialize functions
 VECTORVM_API uint32  SerializeVectorVMInputDataSets  (FVectorVMSerializeState *SerializeState, TArrayView<FDataSetMeta> DataSets, FVectorVMConstData *ConstData, int NumConstData); //only use when not calling InitVectorVMState()
@@ -443,7 +483,7 @@ VECTORVM_API uint32  SerializeVectorVMOutputDataSets (FVectorVMSerializeState *S
 VECTORVM_API void    SerializeVectorVMWriteToFile    (FVectorVMSerializeState *SerializeState, uint8 WhichStateWritten, const wchar_t *Filename);
 VECTORVM_API void    FreeVectorVMSerializeState      (FVectorVMSerializeState *SerializeState);
 
-#else //NIAGARA_EXP_VM
+#else // VECTORVM_SUPPORTS_EXPERIMENTAL
 
 struct FVectorVMState {
 
@@ -454,4 +494,4 @@ VECTORVM_API uint32 SerializeVectorVMOutputDataSets (FVectorVMSerializeState *Se
 VECTORVM_API void   SerializeVectorVMWriteToFile    (FVectorVMSerializeState *SerializeState, uint8 WhichStateWritten, const wchar_t *Filename);
 VECTORVM_API void   FreeVectorVMSerializeState      (FVectorVMSerializeState *SerializeState);
 
-#endif //NIAGARA_EXP_VM
+#endif // VECTORVM_SUPPORTS_EXPERIMENTAL
diff --git a/Engine/Source/Runtime/VectorVM/Public/VectorVMLegacy.h b/Engine/Source/Runtime/VectorVM/Public/VectorVMLegacy.h
new file mode 100644
--- /dev/null
+++ b/Engine/Source/Runtime/VectorVM/Public/VectorVMLegacy.h
@@ -0,0 +1,245 @@
+// Copyright Epic Games, Inc. All Rights Reserved.
+#pragma once
+
+#include "HAL/ThreadSingleton.h"
+#include "VectorVMCommon.h"
+
+#if VECTORVM_SUPPORTS_LEGACY
+
+//Data the VM will keep on each dataset locally per thread which is then thread safely pushed to it's destination at the end of execution.
+struct FDataSetThreadLocalTempData
+{
+	FDataSetThreadLocalTempData()
+	{
+		Reset();
+	}
+
+	FORCEINLINE void Reset()
+	{
+		IDsToFree.Reset();
+		MaxID = INDEX_NONE;
+	}
+
+	TArray<int32> IDsToFree;
+	int32 MaxID;
+
+	//TODO: Possibly store output data locally and memcpy to the real buffers. Could avoid false sharing in parallel execution and so improve perf.
+	//using _mm_stream_ps on platforms that support could also work for this?
+	//TArray<TArray<float>> OutputFloatData;
+	//TArray<TArray<int32>> OutputIntData;
+};
+
+
+/**
+* Context information passed around during VM execution.
+*/
+struct FVectorVMContext : TThreadSingleton<FVectorVMContext>
+{
+private:
+
+	friend struct FVectorVMCodeOptimizerContext;
+
+public:
+
+	/** Pointer to the next element in the byte code. */
+	uint8 const* RESTRICT Code;
+
+	/** Pointer to the constant table. */
+	const uint8* const* RESTRICT ConstantTable;
+	const int32* ConstantTableSizes;
+	int32 ConstantTableCount;
+
+	/** Num temp registers required by this script. */
+	int32 NumTempRegisters;
+
+	/** Pointer to the shared data table. */
+	const FVMExternalFunction* const* RESTRICT ExternalFunctionTable;
+	/** Table of user pointers.*/
+	void** UserPtrTable;
+
+	/** Number of instances to process. */
+	int32 NumInstances;
+	/** Number of instances to process when doing batches of VECTOR_WIDTH_FLOATS. */
+	int32 NumInstancesVectorFloats;
+	/** Start instance of current chunk. */
+	int32 StartInstance;
+
+	/** HACK: An additional instance offset to allow external functions direct access to specific instances in the buffers. */
+	int32 ExternalFunctionInstanceOffset;
+
+	/** Array of meta data on data sets. TODO: This struct should be removed and all features it contains be handled by more general vm ops and the compiler's knowledge of offsets etc. */
+	TArrayView<FDataSetMeta> DataSetMetaTable;
+
+	TArray<FDataSetThreadLocalTempData> ThreadLocalTempData;
+
+#if STATS
+	TArray<FStatStackEntry, TInlineAllocator<64>> StatCounterStack;
+	TArrayView<FStatScopeData> StatScopes;
+	TArray<uint64, TInlineAllocator<64>> ScopeExecCycles;
+#elif ENABLE_STATNAMEDEVENTS
+	TArrayView<const FString> StatNamedEventScopes;
+#endif
+
+	TArray<uint8, TAlignedHeapAllocator<VECTOR_WIDTH_BYTES>> TempRegTable;
+	uint32 TempRegisterSize;
+	uint32 TempBufferSize;
+
+	/** Thread local random stream for use in external functions needing non-deterministic randoms. */
+	FRandomStream RandStream;
+
+	/** Thread local per instance random counters for use in external functions needing deterministic randoms. */
+	TArray<int32> RandCounters;
+
+	bool bIsParallelExecution;
+
+	int32 ValidInstanceIndexStart = INDEX_NONE;
+	int32 ValidInstanceCount = 0;
+	bool ValidInstanceUniform = false;
+
+	FVectorVMContext();
+
+	void PrepareForExec(
+		int32 InNumTempRegisters,
+		int32 ConstantTableCount,
+		const uint8* const* InConstantTables,
+		const int32* InConstantTableSizes,
+		const FVMExternalFunction* const* InExternalFunctionTable,
+		void** InUserPtrTable,
+		TArrayView<FDataSetMeta> InDataSetMetaTable,
+		int32 MaxNumInstances,
+		bool bInParallelExecution);
+
+#if STATS
+	void SetStatScopes(TArrayView<FStatScopeData> InStatScopes);
+#elif ENABLE_STATNAMEDEVENTS
+	void SetStatNamedEventScopes(TArrayView<const FString> InStatNamedEventScopes);
+#endif
+
+	void FinishExec();
+
+	void PrepareForChunk(const uint8* InCode, int32 InNumInstances, int32 InStartInstance)
+	{
+		Code = InCode;
+		NumInstances = InNumInstances;
+		NumInstancesVectorFloats = (NumInstances + VECTOR_WIDTH_FLOATS - 1) / VECTOR_WIDTH_FLOATS;
+		StartInstance = InStartInstance;
+
+		ExternalFunctionInstanceOffset = 0;
+
+		ValidInstanceCount = 0;
+		ValidInstanceIndexStart = INDEX_NONE;
+		ValidInstanceUniform = false;
+
+		RandCounters.Reset();
+		RandCounters.SetNumZeroed(InNumInstances);
+	}
+
+	FORCEINLINE FDataSetMeta& GetDataSetMeta(int32 DataSetIndex) { return DataSetMetaTable[DataSetIndex]; }
+	FORCEINLINE uint8* RESTRICT GetTempRegister(int32 RegisterIndex) { return TempRegTable.GetData() + TempRegisterSize * RegisterIndex; }
+	template<typename T, int TypeOffset>
+	FORCEINLINE T* RESTRICT GetInputRegister(int32 DataSetIndex, int32 RegisterIndex)
+	{
+		FDataSetMeta& Meta = GetDataSetMeta(DataSetIndex);
+		uint32 Offset = Meta.InputRegisterTypeOffsets[TypeOffset];
+		return ((T*)Meta.InputRegisters[Offset + RegisterIndex]) + Meta.InstanceOffset;
+	}
+	template<typename T, int TypeOffset>
+	FORCEINLINE T* RESTRICT GetOutputRegister(int32 DataSetIndex, int32 RegisterIndex)
+	{
+		FDataSetMeta& Meta = GetDataSetMeta(DataSetIndex);
+		uint32 Offset = Meta.OutputRegisterTypeOffsets[TypeOffset];
+		return  ((T*)Meta.OutputRegisters[Offset + RegisterIndex]) + Meta.InstanceOffset;
+	}
+
+	int32 GetNumInstances() const { return NumInstances; }
+	int32 GetStartInstance() const { return StartInstance; }
+
+	template<uint32 InstancesPerOp>
+	int32 GetNumLoops() const { return (InstancesPerOp == VECTOR_WIDTH_FLOATS) ? NumInstancesVectorFloats : ((InstancesPerOp == 1) ? NumInstances : Align(NumInstances, InstancesPerOp)); }
+
+	FORCEINLINE uint8 DecodeU8() { return *Code++; }
+#if PLATFORM_SUPPORTS_UNALIGNED_LOADS
+	FORCEINLINE uint16 DecodeU16() { uint16 v = *reinterpret_cast<const uint16*>(Code); Code += sizeof(uint16); return INTEL_ORDER16(v); }
+	FORCEINLINE uint32 DecodeU32() { uint32 v = *reinterpret_cast<const uint32*>(Code); Code += sizeof(uint32); return INTEL_ORDER32(v); }
+	FORCEINLINE uint64 DecodeU64() { uint64 v = *reinterpret_cast<const uint64*>(Code); Code += sizeof(uint64); return INTEL_ORDER64(v); }
+#else
+	FORCEINLINE uint16 DecodeU16() { uint16 v = Code[1]; v = v << 8 | Code[0]; Code += 2; return INTEL_ORDER16(v); }
+	FORCEINLINE uint32 DecodeU32() { uint32 v = Code[3]; v = v << 8 | Code[2]; v = v << 8 | Code[1]; v = v << 8 | Code[0]; Code += 4; return INTEL_ORDER32(v); }
+	FORCEINLINE uint64 DecodeU64() { uint64 v = Code[7]; v = v << 8 | Code[6]; v = v << 8 | Code[5]; v = v << 8 | Code[4]; v = v << 8 | Code[3]; v = v << 8 | Code[2]; v = v << 8 | Code[1]; v = v << 8 | Code[0]; Code += 8; return INTEL_ORDER64(v); }
+#endif
+	FORCEINLINE uintptr_t DecodePtr() { return (sizeof(uintptr_t) == 4) ? DecodeU32() : DecodeU64(); }
+	FORCEINLINE void SkipCode(int64 Count) { Code += Count; }
+
+	/** Decode the next operation contained in the bytecode. */
+	FORCEINLINE EVectorVMOp DecodeOp()
+	{
+		return static_cast<EVectorVMOp>(DecodeU8());
+	}
+
+	FORCEINLINE uint8 DecodeSrcOperandTypes()
+	{
+		return DecodeU8();
+	}
+
+	FORCEINLINE bool IsParallelExecution()
+	{
+		return bIsParallelExecution;
+	}
+
+	template<typename T = uint8>
+	FORCEINLINE const T* GetConstant(int32 TableIndex, int32 TableOffset) const
+	{
+		check(TableIndex < ConstantTableCount);
+		return reinterpret_cast<const T*>(ConstantTable[TableIndex] + TableOffset);
+	}
+
+	template<typename T = uint8>
+	FORCEINLINE const T* GetConstant(int32 Offset) const
+	{
+		int32 TableIndex = 0;
+
+		while (Offset >= ConstantTableSizes[TableIndex])
+		{
+			Offset -= ConstantTableSizes[TableIndex];
+			++TableIndex;
+		}
+
+		check(TableIndex < ConstantTableCount);
+		check(Offset < ConstantTableSizes[TableIndex]);
+		return reinterpret_cast<const T*>(ConstantTable[TableIndex] + Offset);
+	}
+};
+
+class FVectorVMExternalFunctionContextLegacy
+{
+	friend struct FNiagaraSystemScriptExecutionContext; //@NOTE(smcgrath): required for the PerInstanceFunctionHook() in the non-experimental version of VectorVM
+public:
+	FVectorVMExternalFunctionContextLegacy() = default;
+	FVectorVMExternalFunctionContextLegacy(FVectorVMContext* InVectorVMContext) : VectorVMContext(InVectorVMContext) { }
+	FORCEINLINE int32 GetStartInstance() const { return VectorVMContext->GetStartInstance(); }
+	FORCEINLINE int32 GetNumInstances() const { return VectorVMContext->GetNumInstances(); }
+	FORCEINLINE TArray<int32>& GetRandCounters() { return VectorVMContext->RandCounters; }
+	FORCEINLINE FRandomStream& GetRandStream() { return VectorVMContext->RandStream; }
+
+	FORCEINLINE uint8* RESTRICT GetTempRegister(int32 RegisterIndex) { return VectorVMContext->GetTempRegister(RegisterIndex); }
+	FORCEINLINE uint8 DecodeU8() { return VectorVMContext->DecodeU8(); }
+	FORCEINLINE uint16 DecodeU16() { return VectorVMContext->DecodeU16(); }
+	FORCEINLINE uint32 DecodeU32() { return VectorVMContext->DecodeU32(); }
+	FORCEINLINE uint64 DecodeU64() { return VectorVMContext->DecodeU64(); }
+	template<typename T = uint8>
+	FORCEINLINE const T* GetConstant(int32 TableIndex, int32 TableOffset) const { return VectorVMContext->GetConstant<T>(TableIndex, TableOffset); }
+	template<typename T = uint8>
+	FORCEINLINE const T* GetConstant(int32 Offset) const { return VectorVMContext->GetConstant<T>(Offset); }
+	FORCEINLINE void* GetUserPtrTable(int32 UserPtrIdx) { return VectorVMContext->UserPtrTable[UserPtrIdx]; }
+	FORCEINLINE int32 GetExternalFunctionInstanceOffset() const { return VectorVMContext->ExternalFunctionInstanceOffset; }
+	FORCEINLINE int32 GetNumTempRegisters() const { return VectorVMContext->NumTempRegisters; }
+
+	template<uint32 InstancesPerOp>
+	FORCEINLINE int32 GetNumLoops() const { return VectorVMContext->GetNumLoops<InstancesPerOp>(); }
+
+private:
+	FVectorVMContext* VectorVMContext;
+};
+
+
+#endif // VECTORVM_SUPPORTS_LEGACY
diff --git a/Engine/Source/Runtime/VectorVM/VectorVM.Build.cs b/Engine/Source/Runtime/VectorVM/VectorVM.Build.cs
--- a/Engine/Source/Runtime/VectorVM/VectorVM.Build.cs
+++ b/Engine/Source/Runtime/VectorVM/VectorVM.Build.cs
@@ -4,8 +4,8 @@ using UnrealBuildTool;
 
 public class VectorVM : ModuleRules
 {
-	public VectorVM(ReadOnlyTargetRules Target) : base(Target)
-	{
+    public VectorVM(ReadOnlyTargetRules Target) : base(Target)
+    {
         PublicDependencyModuleNames.AddRange(
             new string[] {
                 "Core",
@@ -13,23 +13,22 @@ public class VectorVM : ModuleRules
             }
         );
 
-		PrivateDependencyModuleNames.AddRange(
+        PrivateDependencyModuleNames.AddRange(
             new string[] {
                 "Core",
-				"CoreUObject"
-            }
-        );
+                "CoreUObject"
+            });
 
         PrivateIncludePaths.AddRange(
             new string[] {
                 "Runtime/Engine/Classes/Curves"
-            }
-        );
+            });
+
         PublicDefinitions.AddRange(
             new string[]
             {
-                "NIAGARA_EXP_VM",
-            }
-        );
+                "VECTORVM_SUPPORTS_EXPERIMENTAL=1",
+                "VECTORVM_SUPPORTS_LEGACY=1"
+            });
     }
 }
