Readded effect type to system properties when scalability mode is active

diff --git a/Engine/Plugins/FX/Niagara/Source/NiagaraEditor/Private/NiagaraSystemDetailsCustomization.cpp b/Engine/Plugins/FX/Niagara/Source/NiagaraEditor/Private/NiagaraSystemDetailsCustomization.cpp
--- a/Engine/Plugins/FX/Niagara/Source/NiagaraEditor/Private/NiagaraSystemDetailsCustomization.cpp
+++ b/Engine/Plugins/FX/Niagara/Source/NiagaraEditor/Private/NiagaraSystemDetailsCustomization.cpp
@@ -17,6 +17,5 @@ void FNiagaraSystemDetails::CustomizeDetails(IDetailLayoutBuilder& InDetailLayou
 	if(GNiagaraScalabilityModeEnabled)
 	{
 		InDetailLayout.HideCategory("Scalability");
-		InDetailLayout.HideProperty("EffectType");
 	}
 }
