Disable experimental niagara VM while memstomps are causing troubles with Lyra effects

diff --git a/Engine/Plugins/FX/Niagara/Source/Niagara/Niagara.Build.cs b/Engine/Plugins/FX/Niagara/Source/Niagara/Niagara.Build.cs
--- a/Engine/Plugins/FX/Niagara/Source/Niagara/Niagara.Build.cs
+++ b/Engine/Plugins/FX/Niagara/Source/Niagara/Niagara.Build.cs
@@ -70,7 +70,7 @@ public class Niagara : ModuleRules
         PublicDefinitions.AddRange(
             new string[]
             {
-                "VECTORVM_SUPPORTS_EXPERIMENTAL=1",
+                "VECTORVM_SUPPORTS_EXPERIMENTAL=0",
                 "VECTORVM_SUPPORTS_LEGACY=1"
             });
     }
diff --git a/Engine/Source/Runtime/VectorVM/VectorVM.Build.cs b/Engine/Source/Runtime/VectorVM/VectorVM.Build.cs
--- a/Engine/Source/Runtime/VectorVM/VectorVM.Build.cs
+++ b/Engine/Source/Runtime/VectorVM/VectorVM.Build.cs
@@ -27,7 +27,7 @@ public class VectorVM : ModuleRules
         PublicDefinitions.AddRange(
             new string[]
             {
-                "VECTORVM_SUPPORTS_EXPERIMENTAL=1",
+                "VECTORVM_SUPPORTS_EXPERIMENTAL=0",
                 "VECTORVM_SUPPORTS_LEGACY=1"
             });
     }
