- Ensure bindings are all setup correctly in PostLoad.  This fixes an issue with some older content that are saved with incorrect bindings.

#rnx

diff --git a/Engine/Plugins/FX/Niagara/Source/Niagara/Private/NiagaraSimulationStageBase.cpp b/Engine/Plugins/FX/Niagara/Source/Niagara/Private/NiagaraSimulationStageBase.cpp
--- a/Engine/Plugins/FX/Niagara/Source/Niagara/Private/NiagaraSimulationStageBase.cpp
+++ b/Engine/Plugins/FX/Niagara/Source/Niagara/Private/NiagaraSimulationStageBase.cpp
@@ -7,6 +7,15 @@
 
 const FName UNiagaraSimulationStageBase::ParticleSpawnUpdateName("ParticleSpawnUpdate");
 
+namespace NiagaraSimulationStageLocal
+{
+	static FName NAME_Particles_StateIndex("Particles.StateIndex");
+	static FNiagaraVariableBase GetDefaultEnabledBinding() { return FNiagaraVariableBase(FNiagaraTypeDefinition::GetBoolDef(), NAME_None); }
+	static FNiagaraVariableBase GetDefaultElementCountBinding() { return FNiagaraVariableBase(FNiagaraTypeDefinition::GetIntDef(), NAME_None); }
+	static FNiagaraVariableBase GetDefaultNumIterationsBinding() { return FNiagaraVariableBase(FNiagaraTypeDefinition::GetIntDef(), NAME_None); }
+	static FNiagaraVariableBase GetDefaultParticleStateBinding() { return FNiagaraVariableBase(FNiagaraTypeDefinition::GetIntDef(), NAME_Particles_StateIndex); }
+}
+
 bool UNiagaraSimulationStageBase::AppendCompileHash(FNiagaraCompileHashVisitor* InVisitor) const
 {
 #if WITH_EDITORONLY_DATA
@@ -88,33 +97,40 @@ void UNiagaraSimulationStageGeneric::PostInitProperties()
 
 	if (HasAnyFlags(RF_ClassDefaultObject) == false)
 	{
-		EnabledBinding.Setup(
-			FNiagaraVariableBase(FNiagaraTypeDefinition::GetBoolDef(), NAME_None),
-			FNiagaraVariableBase(FNiagaraTypeDefinition::GetBoolDef(), NAME_None),
-			ENiagaraRendererSourceDataMode::Emitter
-		);
-
-		ElementCountBinding.Setup(
-			FNiagaraVariableBase(FNiagaraTypeDefinition::GetIntDef(), NAME_None),
-			FNiagaraVariableBase(FNiagaraTypeDefinition::GetIntDef(), NAME_None),
-			ENiagaraRendererSourceDataMode::Emitter
-		);
-
-		NumIterationsBinding.Setup(
-			FNiagaraVariableBase(FNiagaraTypeDefinition::GetIntDef(), NAME_None),
-			FNiagaraVariableBase(FNiagaraTypeDefinition::GetIntDef(), NAME_None),
-			ENiagaraRendererSourceDataMode::Emitter
-		);
-
-		static const FName ParticleStateIndex("Particles.StateIndex");
-		ParticleIterationStateBinding.Setup(
-			FNiagaraVariableBase(FNiagaraTypeDefinition::GetIntDef(), ParticleStateIndex),
-			FNiagaraVariableBase(FNiagaraTypeDefinition::GetIntDef(), ParticleStateIndex),
-			ENiagaraRendererSourceDataMode::Particles
-		);
+		using namespace NiagaraSimulationStageLocal;
+		EnabledBinding.Setup(GetDefaultEnabledBinding(), GetDefaultEnabledBinding(), ENiagaraRendererSourceDataMode::Emitter);
+		ElementCountBinding.Setup(GetDefaultElementCountBinding(), GetDefaultElementCountBinding(), ENiagaraRendererSourceDataMode::Emitter);
+		NumIterationsBinding.Setup(GetDefaultNumIterationsBinding(), GetDefaultNumIterationsBinding(), ENiagaraRendererSourceDataMode::Emitter);
+		ParticleIterationStateBinding.Setup(GetDefaultParticleStateBinding(), GetDefaultParticleStateBinding(), ENiagaraRendererSourceDataMode::Particles);
 	}
 }
 
+void UNiagaraSimulationStageGeneric::PostLoad()
+{
+	Super::PostLoad();
+
+#if WITH_EDITORONLY_DATA
+	// Ensure data wasn't somehow saved incorrectly
+	using namespace NiagaraSimulationStageLocal;
+	if (EnabledBinding.GetType() != GetDefaultEnabledBinding().GetType())
+	{
+		EnabledBinding.Setup(GetDefaultEnabledBinding(), GetDefaultEnabledBinding(), ENiagaraRendererSourceDataMode::Emitter);
+	}
+	if (ElementCountBinding.GetType() != GetDefaultElementCountBinding().GetType())
+	{
+		ElementCountBinding.Setup(GetDefaultElementCountBinding(), GetDefaultElementCountBinding(), ENiagaraRendererSourceDataMode::Emitter);
+	}
+	if (NumIterationsBinding.GetType() != GetDefaultNumIterationsBinding().GetType())
+	{
+		NumIterationsBinding.Setup(GetDefaultNumIterationsBinding(), GetDefaultNumIterationsBinding(), ENiagaraRendererSourceDataMode::Emitter);
+	}
+	if (ParticleIterationStateBinding.GetType() != GetDefaultParticleStateBinding().GetType())
+	{
+		ParticleIterationStateBinding.Setup(GetDefaultParticleStateBinding(), GetDefaultParticleStateBinding(), ENiagaraRendererSourceDataMode::Particles);
+	}
+#endif
+}
+
 #if WITH_EDITOR
 void UNiagaraSimulationStageGeneric::PostEditChangeProperty(struct FPropertyChangedEvent& PropertyChangedEvent) 
 {
diff --git a/Engine/Plugins/FX/Niagara/Source/Niagara/Public/NiagaraSimulationStageBase.h b/Engine/Plugins/FX/Niagara/Source/Niagara/Public/NiagaraSimulationStageBase.h
--- a/Engine/Plugins/FX/Niagara/Source/Niagara/Public/NiagaraSimulationStageBase.h
+++ b/Engine/Plugins/FX/Niagara/Source/Niagara/Public/NiagaraSimulationStageBase.h
@@ -132,6 +132,7 @@ public:
 	FIntVector OverrideGpuDispatchNumThreads = FIntVector(64, 1, 1);
 
 	virtual void PostInitProperties() override;
+	virtual void PostLoad() override;
 	virtual bool AppendCompileHash(FNiagaraCompileHashVisitor* InVisitor) const override;
 #if WITH_EDITOR
 	virtual FName GetStackContextReplacementName() const override; 
