Fixed SA issue where ParentNode was accessed before being nullptr checked later

diff --git a/Engine/Plugins/FX/Niagara/Source/NiagaraEditor/Private/TNiagaraGraphPinEditableName.h b/Engine/Plugins/FX/Niagara/Source/NiagaraEditor/Private/TNiagaraGraphPinEditableName.h
--- a/Engine/Plugins/FX/Niagara/Source/NiagaraEditor/Private/TNiagaraGraphPinEditableName.h
+++ b/Engine/Plugins/FX/Niagara/Source/NiagaraEditor/Private/TNiagaraGraphPinEditableName.h
@@ -99,18 +99,12 @@ protected:
 	virtual TSharedRef<SWidget> GetLabelWidget(const FName& InLabelStyle) override
 	{
 		UNiagaraNode* ParentNode = Cast<UNiagaraNode>(this->GraphPinObj->GetOwningNode());
-		const bool bAllowPinTypeChanges = ParentNode->AllowExternalPinTypeChanges(this->GraphPinObj) && this->GraphPinObj->bOrphanedPin == false;
-
-		auto CreateLabelTextBlock = [&]()->TSharedRef<SWidget> {
-			CreatedTextBlock = SNew(SInlineEditableTextBlock)
-				.Style(&FEditorStyle::Get().GetWidgetStyle<FInlineEditableTextBlockStyle>("Graph.Node.InlineEditablePinName"))
-				.Text(this, &TNiagaraGraphPinEditableName<BaseClass>::GetParentPinLabel)
-				.Visibility(this, &TNiagaraGraphPinEditableName<BaseClass>::GetParentPinVisibility)
-				.ColorAndOpacity(this, &TNiagaraGraphPinEditableName<BaseClass>::GetParentPinTextColor)
-				.IsReadOnly(true);
-			return CreatedTextBlock.ToSharedRef();
-		};
 
+		if(ParentNode == nullptr)
+		{
+			return BaseClass::GetLabelWidget(InLabelStyle);
+		}
+		
 		auto CreateRenamableLabelTextBlock = [&]()->TSharedRef<SWidget> {
 			if (ParentNode->IsPinNameEditableUponCreation(this->GraphPinObj))
 			{
@@ -126,9 +120,10 @@ protected:
 				.OnTextCommitted(this, &TNiagaraGraphPinEditableName<BaseClass>::OnTextCommitted);
 			return CreatedTextBlock.ToSharedRef();
 		};
-
-		bool bIsPinEditable = ParentNode && ParentNode->IsPinNameEditable(this->GraphPinObj);
-		if (ParentNode && ParentNode->IsA<UNiagaraNodeParameterMapBase>())
+		
+		const bool bAllowPinTypeChanges = ParentNode->AllowExternalPinTypeChanges(this->GraphPinObj) && this->GraphPinObj->bOrphanedPin == false;
+		const bool bIsPinEditable = ParentNode->IsPinNameEditable(this->GraphPinObj);
+		if (ParentNode->IsA<UNiagaraNodeParameterMapBase>())
 		{
 			return SNew(SBox)
 			.Padding(FMargin(0.0f, 1.0f, 0.0f, 1.0f))
@@ -142,29 +137,26 @@ protected:
 				.OnTextCommitted(this, &TNiagaraGraphPinEditableName<BaseClass>::OnTextCommitted)
 			];
 		}
-		else if (bIsPinEditable)
+		else if(bIsPinEditable)
 		{
 			return CreateRenamableLabelTextBlock();
 		}
-		else
-		{
-			// we want the pin type selector in the label widget only for output pins, otherwise we put it in the value widget
-			if (bAllowPinTypeChanges && this->GraphPinObj->Direction == EGPD_Output)
-			{
-				TSharedRef<SHorizontalBox> LabelWidgetContainer = SNew(SHorizontalBox);
-				LabelWidgetContainer->AddSlot().AutoWidth().Padding(3.f, 0.f)
-				[
-					SNew(SNiagaraPinTypeSelector, this->GraphPinObj)
-				];
-				LabelWidgetContainer->AddSlot()
-				[
-					BaseClass::GetLabelWidget(InLabelStyle)
-				];
-				return LabelWidgetContainer;
-			}
-			
-			return BaseClass::GetLabelWidget(InLabelStyle);
+		// we want the pin type selector in the label widget only for output pins, otherwise we put it in the value widget
+		else if (bAllowPinTypeChanges && this->GraphPinObj->Direction == EGPD_Output)
+		{				
+			TSharedRef<SHorizontalBox> LabelWidgetContainer = SNew(SHorizontalBox);
+			LabelWidgetContainer->AddSlot().AutoWidth().Padding(3.f, 0.f)
+			[
+				SNew(SNiagaraPinTypeSelector, this->GraphPinObj)
+			];
+			LabelWidgetContainer->AddSlot()
+			[
+				BaseClass::GetLabelWidget(InLabelStyle)
+			];
+			return LabelWidgetContainer;
 		}
+
+		return BaseClass::GetLabelWidget(InLabelStyle);
 	}
 
 	virtual TSharedRef<SWidget>	GetDefaultValueWidget() override
