Add logic to cook out Niagara scripts that won't compile for the shader platforms defined by the target

diff --git a/Engine/Plugins/FX/Niagara/Source/Niagara/Private/NiagaraScript.cpp b/Engine/Plugins/FX/Niagara/Source/Niagara/Private/NiagaraScript.cpp
--- a/Engine/Plugins/FX/Niagara/Source/Niagara/Private/NiagaraScript.cpp
+++ b/Engine/Plugins/FX/Niagara/Source/Niagara/Private/NiagaraScript.cpp
@@ -22,6 +22,7 @@
 #include "UObject/ReleaseObjectVersion.h"
 #include "NiagaraDataInterfaceSkeletalMesh.h"
 #include "DataInterface/NiagaraDataInterfaceStaticMesh.h"
+#include "Interfaces/ITargetPlatform.h"
 #if WITH_EDITOR
 	#include "DerivedDataCacheInterface.h"
 	#include "Interfaces/ITargetPlatform.h"
@@ -1974,6 +1975,29 @@ bool UNiagaraScript::ShouldCacheShadersForCooking(const ITargetPlatform* TargetP
 {
 	if (CanBeRunOnGpu())
 	{
+		{
+			bool ShaderPlatformsSupportScript = false;
+
+			TArray<FName> DesiredShaderFormats;
+			TargetPlatform->GetAllTargetedShaderFormats(DesiredShaderFormats);
+
+			// Cache for all the shader formats that the cooking target requires
+			for (const FName& DesiredShaderFormat : DesiredShaderFormats)
+			{
+				const EShaderPlatform LegacyShaderPlatform = ShaderFormatToLegacyShaderPlatform(DesiredShaderFormat);
+				if (ShouldCompile(LegacyShaderPlatform))
+				{
+					ShaderPlatformsSupportScript = true;
+					break;
+				}
+			}
+
+			if (!ShaderPlatformsSupportScript)
+			{
+				return false;
+			}
+		}
+
 		UNiagaraEmitter* OwningEmitter = GetTypedOuter<UNiagaraEmitter>();
 		if (OwningEmitter != nullptr && OwningEmitter->SimTarget == ENiagaraSimTarget::GPUComputeSim && OwningEmitter->NeedsLoadForTargetPlatform(TargetPlatform))
 		{
