- Add looping control
- Force mips to be resident on baked textures
- Fix Icon for baker menu entry
- Ensure deferred updates are flushed before executing compute dispatch flush to avoid picking up stale resources

#rnx

diff --git a/Engine/Plugins/FX/Niagara/Source/Niagara/Private/NiagaraBakerSettings.cpp b/Engine/Plugins/FX/Niagara/Source/Niagara/Private/NiagaraBakerSettings.cpp
--- a/Engine/Plugins/FX/Niagara/Source/Niagara/Private/NiagaraBakerSettings.cpp
+++ b/Engine/Plugins/FX/Niagara/Source/Niagara/Private/NiagaraBakerSettings.cpp
@@ -17,6 +17,7 @@ bool FNiagaraBakerTextureSettings::Equals(const FNiagaraBakerTextureSettings& Ot
 UNiagaraBakerSettings::UNiagaraBakerSettings(const FObjectInitializer& Init)
 	: Super(Init)
 {
+	bPreviewLooping = true;
 	bRenderComponentOnly = true;
 	OutputTextures.AddDefaulted();
 
diff --git a/Engine/Plugins/FX/Niagara/Source/Niagara/Private/NiagaraGpuComputeDispatch.cpp b/Engine/Plugins/FX/Niagara/Source/Niagara/Private/NiagaraGpuComputeDispatch.cpp
--- a/Engine/Plugins/FX/Niagara/Source/Niagara/Private/NiagaraGpuComputeDispatch.cpp
+++ b/Engine/Plugins/FX/Niagara/Source/Niagara/Private/NiagaraGpuComputeDispatch.cpp
@@ -407,18 +407,18 @@ void FNiagaraGpuComputeDispatch::FlushPendingTicks_GameThread()
 
 void FNiagaraGpuComputeDispatch::ProcessPendingTicksFlush(FRHICommandListImmediate& RHICmdList, bool bForceFlush)
 {
-	// No ticks are pending
-	bool bHasTicks = false;
+	// Test to see if we have any proxies, if not we have nothing to do
+	bool bHasProxies = false;
 	for ( int iTickStage=0; iTickStage < ENiagaraGpuComputeTickStage::Max; ++iTickStage)
 	{
 		if ( ProxiesPerStage[iTickStage].Num() > 0 )
 		{
-			bHasTicks = true;
+			bHasProxies = true;
 			break;
 		}
 	}
 
-	if ( !bHasTicks )
+	if ( !bHasProxies)
 	{
 		return;
 	}
@@ -446,6 +446,28 @@ void FNiagaraGpuComputeDispatch::ProcessPendingTicksFlush(FRHICommandListImmedia
 		{
 			//UE_LOG(LogNiagara, Log, TEXT("FNiagaraGpuComputeDispatch: Queued ticks are being Processed due to not rendering.  This may result in undesirable simulation artifacts."));
 
+			// Make a pass to see if we have any pending ticks, if not we can early out here
+			bool bHasPendingTicks = false;
+			for (int iTickStage=0; !bHasPendingTicks && (iTickStage < ENiagaraGpuComputeTickStage::Max); ++iTickStage)
+			{
+				for ( FNiagaraSystemGpuComputeProxy* Proxy : ProxiesPerStage[iTickStage] )
+				{
+					bHasPendingTicks = Proxy->PendingTicks.Num() > 0;
+					if (bHasPendingTicks)
+					{
+						break;
+					}
+				}
+			}
+			if (bHasPendingTicks == false)
+			{
+				return;
+			}
+
+			// Ensure any deferred updates are flushed out
+			FDeferredUpdateResource::UpdateResources(RHICmdList);
+			FMaterialRenderProxy::UpdateDeferredCachedUniformExpressions();
+
 			// Make a temporary ViewInfo
 			//-TODO: We could gather some more information here perhaps?
 			FMemMark Mark(FMemStack::Get());
diff --git a/Engine/Plugins/FX/Niagara/Source/NiagaraEditor/Private/NiagaraBakerRenderer.cpp b/Engine/Plugins/FX/Niagara/Source/NiagaraEditor/Private/NiagaraBakerRenderer.cpp
--- a/Engine/Plugins/FX/Niagara/Source/NiagaraEditor/Private/NiagaraBakerRenderer.cpp
+++ b/Engine/Plugins/FX/Niagara/Source/NiagaraEditor/Private/NiagaraBakerRenderer.cpp
@@ -92,6 +92,9 @@ bool FNiagaraBakerRenderer::RenderView(UNiagaraComponent* PreviewComponent, cons
 	UWorld* World = PreviewComponent->GetWorld();
 	check(World);
 
+	// Send EOF updates before we flush our pending ticks to ensure everything is ready for Niagara
+	World->SendAllEndOfFrameUpdates();
+
 	// We need any GPU sims to flush pending ticks
 	FNiagaraGpuComputeDispatchInterface* ComputeDispatchInterface = FNiagaraGpuComputeDispatchInterface::Get(World);
 	ensureMsgf(ComputeDispatchInterface, TEXT("The batcher was not valid on the world this may result in incorrect baking"));
@@ -291,7 +294,6 @@ bool FNiagaraBakerRenderer::RenderView(UNiagaraComponent* PreviewComponent, cons
 					continue;
 				}
 
-				//-TODO: We don't support GPU Context currently
 				if (EmitterInstance->GetGPUContext() != nullptr)
 				{
 					return false;
diff --git a/Engine/Plugins/FX/Niagara/Source/NiagaraEditor/Private/Toolkits/SystemToolkitModes/NiagaraSystemToolkitModeBase.cpp b/Engine/Plugins/FX/Niagara/Source/NiagaraEditor/Private/Toolkits/SystemToolkitModes/NiagaraSystemToolkitModeBase.cpp
--- a/Engine/Plugins/FX/Niagara/Source/NiagaraEditor/Private/Toolkits/SystemToolkitModes/NiagaraSystemToolkitModeBase.cpp
+++ b/Engine/Plugins/FX/Niagara/Source/NiagaraEditor/Private/Toolkits/SystemToolkitModes/NiagaraSystemToolkitModeBase.cpp
@@ -135,7 +135,8 @@ void FNiagaraSystemToolkitModeBase::RegisterTabFactories(TSharedPtr<FTabManager>
 
 	InTabManager->RegisterTabSpawner(BakerTabID, FOnSpawnTab::CreateSP(this, &FNiagaraSystemToolkitModeBase::SpawnTab_Baker))
 		.SetDisplayName(LOCTEXT("NiagaraBakerTab", "Baker"))
-		.SetGroup(WorkspaceMenuCategory.ToSharedRef());
+		.SetGroup(WorkspaceMenuCategory.ToSharedRef())
+		.SetIcon(FSlateIcon(FNiagaraEditorStyle::Get().GetStyleSetName(), "NiagaraEditor.BakerIcon"));
 }
 
 class SNiagaraSelectedEmitterGraph : public SCompoundWidget
diff --git a/Engine/Plugins/FX/Niagara/Source/NiagaraEditor/Private/ViewModels/NiagaraBakerViewModel.cpp b/Engine/Plugins/FX/Niagara/Source/NiagaraEditor/Private/ViewModels/NiagaraBakerViewModel.cpp
--- a/Engine/Plugins/FX/Niagara/Source/NiagaraEditor/Private/ViewModels/NiagaraBakerViewModel.cpp
+++ b/Engine/Plugins/FX/Niagara/Source/NiagaraEditor/Private/ViewModels/NiagaraBakerViewModel.cpp
@@ -1,10 +1,5 @@
 // Copyright Epic Games, Inc. All Rights Reserved.
 
-//-TODO: Remove post from capture
-//-TODO: Hook up loading a preview environment
-//-TODO: Add custom post processes for capturing data?
-//-TODO: Allow us to toggle preview on / off for baker / realtime
-
 #include "ViewModels/NiagaraBakerViewModel.h"
 #include "ViewModels/NiagaraSystemViewModel.h"
 #include "Widgets/SNiagaraBakerWidget.h"
@@ -125,6 +120,22 @@ TSharedPtr<class SWidget> FNiagaraBakerViewModel::GetWidget()
 	return Widget;
 }
 
+void FNiagaraBakerViewModel::TogglePlaybackLooping()
+{
+	if (UNiagaraBakerSettings* BakerSettings = GetBakerSettings())
+	{
+		const FScopedTransaction Transaction(LOCTEXT("ToggleLooping", "Toggle Looping"));
+		BakerSettings->bPreviewLooping = !BakerSettings->bPreviewLooping;
+		BakerSettings->Modify();
+	}
+}
+
+bool FNiagaraBakerViewModel::IsPlaybackLooping() const
+{
+	UNiagaraBakerSettings* BakerSettings = GetBakerSettings();
+	return BakerSettings ? BakerSettings->bPreviewLooping : false;
+}
+
 void FNiagaraBakerViewModel::SetCameraViewMode(ENiagaraBakerViewMode ViewMode)
 {
 	if (UNiagaraBakerSettings* BakerSettings = GetBakerSettings())
@@ -445,7 +456,6 @@ void FNiagaraBakerViewModel::RenderBaker()
 		PreviewComponent->SeekToDesiredAge(FrameTime);
 		PreviewComponent->TickComponent(BakerSettings->GetSeekDelta(), ELevelTick::LEVELTICK_All, nullptr);
 
-		const float WorldTime = FApp::GetCurrentTime() - BakerSettings->StartSeconds - BakerSettings->DurationSeconds + FrameTime;
 
 		// Render frame
 		for (int32 iOutputTexture=0; iOutputTexture < BakerSettings->OutputTextures.Num(); ++iOutputTexture)
@@ -458,7 +468,7 @@ void FNiagaraBakerViewModel::RenderBaker()
 			}
 			FTextureRenderTargetResource* RenderTargetResource = RenderTarget->GameThread_GetRenderTargetResource();
 
-			BakerRenderer.RenderView(PreviewComponent, BakerSettings, WorldTime, RenderTarget, iOutputTexture);
+			BakerRenderer.RenderView(PreviewComponent, BakerSettings, FrameTime, RenderTarget, iOutputTexture);
 
 			TArray<FFloat16Color> OutSamples;
 			RenderTargetResource->ReadFloat16Pixels(OutSamples);
@@ -514,6 +524,7 @@ void FNiagaraBakerViewModel::RenderBaker()
 		OutputTexture.GeneratedTexture->UpdateResource();
 		OutputTexture.GeneratedTexture->PostEditChange();
 		OutputTexture.GeneratedTexture->MarkPackageDirty();
+		OutputTexture.GeneratedTexture->TemporarilyDisableStreaming();
 	}
 
 	// Duplicate and set as generated data
diff --git a/Engine/Plugins/FX/Niagara/Source/NiagaraEditor/Private/ViewModels/NiagaraBakerViewModel.h b/Engine/Plugins/FX/Niagara/Source/NiagaraEditor/Private/ViewModels/NiagaraBakerViewModel.h
--- a/Engine/Plugins/FX/Niagara/Source/NiagaraEditor/Private/ViewModels/NiagaraBakerViewModel.h
+++ b/Engine/Plugins/FX/Niagara/Source/NiagaraEditor/Private/ViewModels/NiagaraBakerViewModel.h
@@ -57,6 +57,9 @@ public:
 	void ToggleChannelEnabled(ENiagaraBakerColorChannel Channel) { bColorChannelEnabled[int(Channel)] = !bColorChannelEnabled[int(Channel)]; }
 	void SetChannelEnabled(ENiagaraBakerColorChannel Channel, bool bEnabled) { bColorChannelEnabled[int(Channel)] = bEnabled; }
 
+	void TogglePlaybackLooping();
+	bool IsPlaybackLooping() const;
+
 	bool ShowRealtimePreview() const { return bShowRealtimePreview; }
 	void ToggleRealtimePreview() { bShowRealtimePreview = !bShowRealtimePreview; }
 
diff --git a/Engine/Plugins/FX/Niagara/Source/NiagaraEditor/Private/Widgets/SNiagaraBakerWidget.cpp b/Engine/Plugins/FX/Niagara/Source/NiagaraEditor/Private/Widgets/SNiagaraBakerWidget.cpp
--- a/Engine/Plugins/FX/Niagara/Source/NiagaraEditor/Private/Widgets/SNiagaraBakerWidget.cpp
+++ b/Engine/Plugins/FX/Niagara/Source/NiagaraEditor/Private/Widgets/SNiagaraBakerWidget.cpp
@@ -30,10 +30,12 @@
 //-TODO: Improve feedback & warning display
 //-TODO: Throttle sim frames on warmup
 //-TODO: Turn off playback when out of focus
-//-TODO: Looping vs non-looping play back
 //-TODO: Fluid systems stop simulation with to previous?
 //-TODO: First frame black even though I have non zero start seconds
 //-TODO: None should display HDR or something sensible
+//-TODO: Support particle attribute generation on GPU
+//-TODO: Hook up loading a preview environment / sequence
+//-TODO: Add custom post processes for capturing data?
 
 //////////////////////////////////////////////////////////////////////////
 
@@ -426,12 +428,15 @@ void SNiagaraBakerWidget::Construct(const FArguments& InArgs)
 		TransportControlArgs.OnForwardPlay.BindSP(this, &SNiagaraBakerWidget::OnTransportForwardPlay);
 		TransportControlArgs.OnForwardStep.BindSP(this, &SNiagaraBakerWidget::OnTransportForwardStep);
 		TransportControlArgs.OnForwardEnd.BindSP(this, &SNiagaraBakerWidget::OnTransportForwardEnd);
+		TransportControlArgs.OnToggleLooping.BindSP(this, &SNiagaraBakerWidget::OnTransportToggleLooping);
+		TransportControlArgs.OnGetLooping.BindSP(ViewModel, &FNiagaraBakerViewModel::IsPlaybackLooping);
 
 		TransportControlArgs.WidgetsToCreate.Add(FTransportControlWidget(ETransportControlWidgetType::BackwardEnd));
 		TransportControlArgs.WidgetsToCreate.Add(FTransportControlWidget(ETransportControlWidgetType::BackwardStep));
 		TransportControlArgs.WidgetsToCreate.Add(FTransportControlWidget(ETransportControlWidgetType::ForwardPlay));
 		TransportControlArgs.WidgetsToCreate.Add(FTransportControlWidget(ETransportControlWidgetType::ForwardStep));
 		TransportControlArgs.WidgetsToCreate.Add(FTransportControlWidget(ETransportControlWidgetType::ForwardEnd));
+		TransportControlArgs.WidgetsToCreate.Add(FTransportControlWidget(ETransportControlWidgetType::Loop));
 
 		TransportControlArgs.bAreButtonsFocusable = false;
 
@@ -781,24 +786,27 @@ void SNiagaraBakerWidget::RefreshWidget()
 
 void SNiagaraBakerWidget::Tick(const FGeometry& AllottedGeometry, const double CurrentTime, const float DeltaTime)
 {
-	TSharedPtr<FNiagaraBakerViewModel> ViewModel = WeakViewModel.Pin();
-	if (ViewModel == nullptr)
+	UNiagaraBakerSettings* BakerSettings = GetBakerSettings();
+	if (BakerSettings == nullptr)
 	{
 		return;
 	}
 
-	float DurationSeconds = 0.0f;
-	if (UNiagaraBakerSettings* BakerSettings = GetBakerSettings())
-	{
-		DurationSeconds = BakerSettings->DurationSeconds;
-	}
-
+	const float DurationSeconds = BakerSettings->DurationSeconds;
 	if ( DurationSeconds > 0.0f )
 	{
 		if (bIsPlaying)
 		{
 			PreviewRelativeTime += DeltaTime;
-			PreviewRelativeTime = FMath::Fmod(PreviewRelativeTime, DurationSeconds);
+			if ( BakerSettings->bPreviewLooping )
+			{
+				PreviewRelativeTime = FMath::Fmod(PreviewRelativeTime, DurationSeconds);
+			}
+			else if ( PreviewRelativeTime >= DurationSeconds )
+			{
+				PreviewRelativeTime = DurationSeconds;
+				bIsPlaying = false;
+			}
 		}
 		else
 		{
@@ -902,4 +910,13 @@ FReply SNiagaraBakerWidget::OnTransportForwardEnd()
 	return FReply::Handled();
 }
 
+FReply SNiagaraBakerWidget::OnTransportToggleLooping() const
+{
+	if (auto ViewModel = WeakViewModel.Pin())
+	{
+		ViewModel->TogglePlaybackLooping();
+	}
+	return FReply::Handled();
+}
+
 #undef LOCTEXT_NAMESPACE
diff --git a/Engine/Plugins/FX/Niagara/Source/NiagaraEditor/Private/Widgets/SNiagaraBakerWidget.h b/Engine/Plugins/FX/Niagara/Source/NiagaraEditor/Private/Widgets/SNiagaraBakerWidget.h
--- a/Engine/Plugins/FX/Niagara/Source/NiagaraEditor/Private/Widgets/SNiagaraBakerWidget.h
+++ b/Engine/Plugins/FX/Niagara/Source/NiagaraEditor/Private/Widgets/SNiagaraBakerWidget.h
@@ -44,6 +44,7 @@ public:
 	FReply OnTransportForwardPlay();
 	FReply OnTransportForwardStep();
 	FReply OnTransportForwardEnd();
+	FReply OnTransportToggleLooping() const;
 
 private:
 	TWeakPtr<FNiagaraBakerViewModel>			WeakViewModel;
