- Fix for fixed delta time and desired age mode

#rnx

diff --git a/Engine/Plugins/FX/Niagara/Source/Niagara/Private/NiagaraComponent.cpp b/Engine/Plugins/FX/Niagara/Source/Niagara/Private/NiagaraComponent.cpp
--- a/Engine/Plugins/FX/Niagara/Source/Niagara/Private/NiagaraComponent.cpp
+++ b/Engine/Plugins/FX/Niagara/Source/Niagara/Private/NiagaraComponent.cpp
@@ -725,7 +725,8 @@ void UNiagaraComponent::TickComponent(float DeltaSeconds, enum ELevelTick TickTy
 					FNiagaraSystemSimulation* SystemSim = SystemInstanceController->GetSoloSystemSimulation().Get();
 					if (SystemSim)
 					{
-						float TickDelta = Asset->HasFixedTickDelta() ? Asset->GetFixedTickDeltaTime() : SeekDelta;
+						const bool bFixedTickDelta = Asset->HasFixedTickDelta();
+						const float TickDelta = bFixedTickDelta ? Asset->GetFixedTickDeltaTime() : SeekDelta;
 						if (bLockDesiredAgeDeltaTimeToSeekDelta || AgeDiff > TickDelta)
 						{
 							// If we're locking the delta time to the seek delta, or we need to seek more than a frame, tick the simulation by the seek delta.
@@ -754,7 +755,10 @@ void UNiagaraComponent::TickComponent(float DeltaSeconds, enum ELevelTick TickTy
 						else
 						{
 							// Otherwise just tick by the age difference.
-							SystemInstanceController->ManualTick(AgeDiff, nullptr);
+							if (!bFixedTickDelta || AgeDiff >= TickDelta)
+							{
+								SystemInstanceController->ManualTick(AgeDiff, nullptr);
+							}
 						}
 					}
 				}
