Fix couple of typos and unused code.

diff --git a/Engine/Plugins/FX/Niagara/Shaders/Private/NiagaraRibbonVertexFactory.ush b/Engine/Plugins/FX/Niagara/Shaders/Private/NiagaraRibbonVertexFactory.ush
--- a/Engine/Plugins/FX/Niagara/Shaders/Private/NiagaraRibbonVertexFactory.ush
+++ b/Engine/Plugins/FX/Niagara/Shaders/Private/NiagaraRibbonVertexFactory.ush
@@ -939,7 +939,7 @@ FVertexFactoryInput LoadVertexFactoryInputForDynamicUpdate(uint TriangleIndex, i
 	Input.InterpVtxId = TriangleIndex * 3 + VertexIndex;
 
 #if FEATURE_LEVEL >= FEATURE_LEVEL_SM4 && ONEPASS_POINTLIGHT_SHADOW && USING_VERTEX_SHADER_LAYER
-	Input.InstanceInd = 0;
+	Input.InstanceId = 0;
 #endif
 	return Input;
 }
@@ -958,7 +958,7 @@ FVertexFactoryInput LoadVertexFactoryInputForHGS(uint TriangleIndex, int VertexI
 	uint VertexId = Tri.Indices[VertexIndex];
 	Input.InterpVtxId = VertexId;
 #if FEATURE_LEVEL >= FEATURE_LEVEL_SM4 && ONEPASS_POINTLIGHT_SHADOW && USING_VERTEX_SHADER_LAYER
-	Input.InstanceInd = 0;
+	Input.InstanceId = 0;
 #endif
 	return Input;
 }
diff --git a/Engine/Shaders/Private/HairStrands/HairStrandsVertexFactory.ush b/Engine/Shaders/Private/HairStrands/HairStrandsVertexFactory.ush
--- a/Engine/Shaders/Private/HairStrands/HairStrandsVertexFactory.ush
+++ b/Engine/Shaders/Private/HairStrands/HairStrandsVertexFactory.ush
@@ -811,9 +811,6 @@ FVertexFactoryInput LoadVertexFactoryInputForDynamicUpdate(uint TriangleIndex, i
 
 	Input.VertexId = TriangleIndex * 3 + VertexIndex;
 
-#if FEATURE_LEVEL >= FEATURE_LEVEL_SM4 && ONEPASS_POINTLIGHT_SHADOW && USING_VERTEX_SHADER_LAYER
-	Input.InstanceInd = 0;
-#endif
 	return Input;
 }
 
