Added the ability for an individual UObject class to change how its name is reported in stat IDs. This allows us to change stats like PlayerPawn.NiagaraComponent to PlayerPawn.FootstepFX.

diff --git a/Engine/Plugins/FX/Niagara/Source/Niagara/Private/NiagaraComponent.cpp b/Engine/Plugins/FX/Niagara/Source/Niagara/Private/NiagaraComponent.cpp
--- a/Engine/Plugins/FX/Niagara/Source/Niagara/Private/NiagaraComponent.cpp
+++ b/Engine/Plugins/FX/Niagara/Source/Niagara/Private/NiagaraComponent.cpp
@@ -656,6 +656,10 @@ void UNiagaraComponent::ActivateSystem(bool bFlagAsJustAttached)
 
 /********* UFXSystemComponent *********/
 
+FName UNiagaraComponent::GetFNameForStatID() const
+{
+	return Asset ? Asset->GetFNameForStatID() : Super::GetFNameForStatID();
+}
 
 void UNiagaraComponent::TickComponent(float DeltaSeconds, enum ELevelTick TickType, FActorComponentTickFunction* ThisTickFunction)
 {
diff --git a/Engine/Plugins/FX/Niagara/Source/Niagara/Public/NiagaraComponent.h b/Engine/Plugins/FX/Niagara/Source/Niagara/Public/NiagaraComponent.h
--- a/Engine/Plugins/FX/Niagara/Source/Niagara/Public/NiagaraComponent.h
+++ b/Engine/Plugins/FX/Niagara/Source/Niagara/Public/NiagaraComponent.h
@@ -60,6 +60,8 @@ public:
 	virtual void ActivateSystem(bool bFlagAsJustAttached = false) override;
 	/********* UFXSystemComponent *********/
 
+	virtual FName GetFNameForStatID() const override;
+
 private:
 	UPROPERTY(EditAnywhere, Category="Niagara", meta = (DisplayName = "Niagara System Asset"))
 	TObjectPtr<UNiagaraSystem> Asset;
diff --git a/Engine/Source/Runtime/CoreUObject/Private/UObject/UObjectArray.cpp b/Engine/Source/Runtime/CoreUObject/Private/UObject/UObjectArray.cpp
--- a/Engine/Source/Runtime/CoreUObject/Private/UObject/UObjectArray.cpp
+++ b/Engine/Source/Runtime/CoreUObject/Private/UObject/UObjectArray.cpp
@@ -43,7 +43,7 @@ void FUObjectItem::CreateStatID() const
 	for (int32 i = ClassChain.Num() - 1; i >= 0; i--)
 	{
 		Target = ClassChain[i];
-		const FNameEntry* NameEntry = Target->GetFName().GetDisplayNameEntry();
+		const FNameEntry* NameEntry = Target->GetFNameForStatID().GetDisplayNameEntry();
 		if (bFirstEntry)
 		{
 			NameEntry->AppendNameToPathString(LongName);
diff --git a/Engine/Source/Runtime/CoreUObject/Private/UObject/UObjectBase.cpp b/Engine/Source/Runtime/CoreUObject/Private/UObject/UObjectBase.cpp
--- a/Engine/Source/Runtime/CoreUObject/Private/UObject/UObjectBase.cpp
+++ b/Engine/Source/Runtime/CoreUObject/Private/UObject/UObjectBase.cpp
@@ -613,6 +613,11 @@ TMap<FName, FDynamicClassStaticData>& GetDynamicClassMap()
 }
 PRAGMA_ENABLE_DEPRECATION_WARNINGS
 
+FName UObjectBase::GetFNameForStatID() const
+{
+	return GetFName();
+}
+
 /** Removes prefix from the native class name */
 FString UObjectBase::RemoveClassPrefix(const TCHAR* ClassName)
 {
diff --git a/Engine/Source/Runtime/CoreUObject/Public/UObject/UObjectBase.h b/Engine/Source/Runtime/CoreUObject/Public/UObject/UObjectBase.h
--- a/Engine/Source/Runtime/CoreUObject/Public/UObject/UObjectBase.h
+++ b/Engine/Source/Runtime/CoreUObject/Public/UObject/UObjectBase.h
@@ -139,6 +139,9 @@ public:
 		return NamePrivate;
 	}
 
+	/** Overridable method to return a logical name for identification in stats. */
+	virtual FName GetFNameForStatID() const;
+
 	/** Removes the class prefix from the given string */
 	static FString RemoveClassPrefix(const TCHAR* ClassName);
 
