Update some vertex factories to use VF_GPUSCENE_SET_INPUT_FOR_RT

diff --git a/Engine/Plugins/Experimental/Water/Shaders/Private/WaterMeshVertexFactory.ush b/Engine/Plugins/Experimental/Water/Shaders/Private/WaterMeshVertexFactory.ush
--- a/Engine/Plugins/Experimental/Water/Shaders/Private/WaterMeshVertexFactory.ush
+++ b/Engine/Plugins/Experimental/Water/Shaders/Private/WaterMeshVertexFactory.ush
@@ -404,7 +404,7 @@ FVertexFactoryInput LoadVertexFactoryInputForDynamicUpdate(uint TriangleIndex, i
 	Input.InstanceData2 = 0.0f;
 #endif
 	
-	VF_GPUSCENE_SET_INPUT_FOR_RT(Input, GetInstanceUserData(), 0U);
+	VF_GPUSCENE_SET_INPUT_FOR_RT(Input, PrimitiveId, 0U);
 
 	return Input;
 }
diff --git a/Engine/Plugins/FX/Niagara/Shaders/Private/NiagaraMeshVertexFactory.ush b/Engine/Plugins/FX/Niagara/Shaders/Private/NiagaraMeshVertexFactory.ush
--- a/Engine/Plugins/FX/Niagara/Shaders/Private/NiagaraMeshVertexFactory.ush
+++ b/Engine/Plugins/FX/Niagara/Shaders/Private/NiagaraMeshVertexFactory.ush
@@ -75,10 +75,11 @@ struct FVertexFactoryInput
 #endif // !MANUAL_VERTEX_FETCH
 
 #if (VF_USE_PRIMITIVE_SCENE_DATA == 1)
-	uint InstanceIdOffset : ATTRIBUTE13;
+	VF_GPUSCENE_DECLARE_INPUT_BLOCK(13)
+#else
+	uint InstanceId	: SV_InstanceID;
 #endif
 
-	uint InstanceId	: SV_InstanceID;
 #if PASS_NEEDS_VERTEX_ID || MANUAL_VERTEX_FETCH
 	uint VertexId : SV_VertexID;
 #endif // #if PASS_NEEDS_VERTEX_ID
@@ -457,7 +458,7 @@ FVertexFactoryIntermediates GetVertexFactoryIntermediates(FVertexFactoryInput In
 	FVertexFactoryIntermediates Intermediates = (FVertexFactoryIntermediates)0;
 
 #if (VF_USE_PRIMITIVE_SCENE_DATA == 1)
-	Intermediates.SceneData = GetSceneDataIntermediates(Input.InstanceIdOffset, Input.InstanceId);
+	Intermediates.SceneData = VF_GPUSCENE_GET_INTERMEDIATES(Input);
 
 	// Pull the particle index from the custom data for the instance
 	Intermediates.ParticleIndex = asuint(LoadInstanceCustomDataFloat(Intermediates.SceneData.InstanceData, 0));
@@ -914,8 +915,6 @@ float4 VertexFactoryGetTranslatedPrimitiveVolumeBounds(FVertexFactoryInterpolant
 
 		FTriangleBaseAttributes Tri = LoadTriangleBaseAttributes(TriangleIndex);
 		uint VertexId = Tri.Indices[VertexIndex];
-		
-		Input.InstanceId = InstanceIndex() - GetBaseInstanceIndex();
 
 		// Fill FVertexFactoryInput with dummy values, the real ones will be fetched later from NiagaraMeshVF using InstanceId  
 		Input.Position = float4(Tri.LocalPositions[VertexIndex], 1.0f);
@@ -924,10 +923,11 @@ float4 VertexFactoryGetTranslatedPrimitiveVolumeBounds(FVertexFactoryInterpolant
 		Input.VertexId = VertexId;
 #endif // #if PASS_NEEDS_VERTEX_ID
 
-#if VF_USE_PRIMITIVE_SCENE_DATA
+#if (VF_USE_PRIMITIVE_SCENE_DATA == 1)
 		// Note: GetInstanceUserData() stores the GPU-Scene primitive ID
-		int PrimitiveId = GetInstanceUserData();
-		Input.InstanceIdOffset = PrimitiveId | VF_TREAT_INSTANCE_ID_OFFSET_AS_PRIMITIVE_ID_FLAG;
+		VF_GPUSCENE_SET_INPUT_FOR_RT(Input, GetInstanceUserData(), InstanceIndex() - GetBaseInstanceIndex());
+#else
+		Input.InstanceId = InstanceIndex() - GetBaseInstanceIndex(); // TIAGO_TODO
 #endif
 
 		return Input;
diff --git a/Engine/Shaders/Private/GeometryCollectionVertexFactory.ush b/Engine/Shaders/Private/GeometryCollectionVertexFactory.ush
--- a/Engine/Shaders/Private/GeometryCollectionVertexFactory.ush
+++ b/Engine/Shaders/Private/GeometryCollectionVertexFactory.ush
@@ -605,12 +605,8 @@ FVertexFactoryInput LoadVertexFactoryInputForHGS(uint TriangleIndex, int VertexI
 	Input.VertexId = Tri.Indices[VertexIndex];
 	Input.Position = float4(Tri.LocalPositions[VertexIndex], 1.0f);
 
-#if VF_USE_PRIMITIVE_SCENE_DATA
 	// Note: GetInstanceUserData() stores the GPU-Scene primitive ID
-	int PrimitiveId = GetInstanceUserData();
-	Input.InstanceIdOffset = PrimitiveId | VF_TREAT_INSTANCE_ID_OFFSET_AS_PRIMITIVE_ID_FLAG;
-	Input.DrawInstanceId = 0;
-#endif // VF_USE_PRIMITIVE_SCENE_DATA
+	VF_GPUSCENE_SET_INPUT_FOR_RT(Input, GetInstanceUserData(), 0U);
 
 	return Input;
 }
@@ -637,10 +633,8 @@ FVertexFactoryInput LoadVertexFactoryInputForDynamicUpdate(uint TriangleIndex, i
 	Input.Position.z = GeometryCollectionVF.VertexFetch_PositionBuffer[VertexOffset + 2];
 #endif
 
-#if FEATURE_LEVEL >= FEATURE_LEVEL_SM4 && VF_USE_PRIMITIVE_SCENE_DATA
-	Input.DrawInstanceId = 0;
-	Input.InstanceIdOffset = PrimitiveId | VF_TREAT_INSTANCE_ID_OFFSET_AS_PRIMITIVE_ID_FLAG;
-#endif
+	VF_GPUSCENE_SET_INPUT_FOR_RT(Input, PrimitiveId, 0U);
+
 	return Input;
 }
 
diff --git a/Engine/Shaders/Private/HairStrands/HairCardsVertexFactory.ush b/Engine/Shaders/Private/HairStrands/HairCardsVertexFactory.ush
--- a/Engine/Shaders/Private/HairStrands/HairCardsVertexFactory.ush
+++ b/Engine/Shaders/Private/HairStrands/HairCardsVertexFactory.ush
@@ -64,11 +64,9 @@ FVertexFactoryInput LoadVertexFactoryInputForHGS(uint TriangleIndex, int VertexI
 	Input.VertexId = TriangleAttributes.Indices[VertexIndex];
 #endif
 
-#if VF_USE_PRIMITIVE_SCENE_DATA
 	// Note: GetInstanceUserData() stores the GPU-Scene primitive ID
-	Input.InstanceIdOffset = GetInstanceUserData() | VF_TREAT_INSTANCE_ID_OFFSET_AS_PRIMITIVE_ID_FLAG;
-	Input.DrawInstanceId = 0;
-#endif // VF_USE_PRIMITIVE_SCENE_DATA
+	VF_GPUSCENE_SET_INPUT_FOR_RT(Input, GetInstanceUserData(), 0U);
+
 	return Input;
 }
 #endif
diff --git a/Engine/Shaders/Private/HairStrands/HairStrandsMaterialCommon.ush b/Engine/Shaders/Private/HairStrands/HairStrandsMaterialCommon.ush
--- a/Engine/Shaders/Private/HairStrands/HairStrandsMaterialCommon.ush
+++ b/Engine/Shaders/Private/HairStrands/HairStrandsMaterialCommon.ush
@@ -60,12 +60,8 @@ FVertexFactoryInterpolantsVSToPS GetInterpolants(uint PrimitiveId, uint HairPrim
 	FMaterialVertexParameters Empty = (FMaterialVertexParameters)0;
 
 	FVertexFactoryInput Input;
-#if VF_USE_PRIMITIVE_SCENE_DATA
-	// GPUCULL_TODO: Making use of dynamic draw workaround, assumes no actual instancing used for these draws.
-	//               Will load the first instance data.
-	Input.InstanceIdOffset = PrimitiveId | (1U << 31U);
-	Input.DrawInstanceId = 0;
-#endif // VF_USE_PRIMITIVE_SCENE_DATA 
+	// Assumes no actual instancing used for these draws. Will load the first instance data.
+	VF_GPUSCENE_SET_INPUT_FOR_RT(Input, PrimitiveId, 0U);
 #if USE_HAIR_TRIANGLE_STRIP
 	Input.VertexId = HairPrimitiveId * 2; // 6: the vertex factory expect triangles strip index. A quad is made of two triangles with indices: 0,1.
 #else
diff --git a/Engine/Shaders/Private/HairStrands/HairStrandsVertexFactory.ush b/Engine/Shaders/Private/HairStrands/HairStrandsVertexFactory.ush
--- a/Engine/Shaders/Private/HairStrands/HairStrandsVertexFactory.ush
+++ b/Engine/Shaders/Private/HairStrands/HairStrandsVertexFactory.ush
@@ -159,11 +159,9 @@ FVertexFactoryInput LoadVertexFactoryInputForHGS(uint TriangleIndex, int VertexI
 	Input.VertexId = TriangleBase.Indices[VertexIndex];
 	Input.TriangleId = TriangleIndex;
 	Input.TriangleVertexId = VertexIndex;
-#if VF_USE_PRIMITIVE_SCENE_DATA
 	// Note: GetInstanceUserData() stores the GPU-Scene primitive ID
-	Input.InstanceIdOffset = GetInstanceUserData() | VF_TREAT_INSTANCE_ID_OFFSET_AS_PRIMITIVE_ID_FLAG;
-	Input.DrawInstanceId = 0;
-#endif
+	VF_GPUSCENE_SET_INPUT_FOR_RT(Input, GetInstanceUserData(), 0U);
+
 	return Input;
 }
 #endif
@@ -466,8 +464,7 @@ FMaterialPixelParameters GetMaterialPixelParameters(float3 RayOrigin, float3 Ray
 #if VF_USE_PRIMITIVE_SCENE_DATA
 	FVertexFactoryInput Input;
 	// Note: GetInstanceUserData() stores the GPU-Scene primitive ID
-	Input.InstanceIdOffset = GetInstanceUserData() | VF_TREAT_INSTANCE_ID_OFFSET_AS_PRIMITIVE_ID_FLAG;
-	Input.DrawInstanceId = 0;
+	VF_GPUSCENE_SET_INPUT_FOR_RT(Input, GetInstanceUserData(), 0U);
 #endif
 
 	FSceneDataIntermediates SceneData = VF_GPUSCENE_GET_INTERMEDIATES(Input); // NOTE: Input is not used when VF_USE_PRIMITIVE_SCENE_DATA == 0
diff --git a/Engine/Shaders/Private/LandscapeVertexFactory.ush b/Engine/Shaders/Private/LandscapeVertexFactory.ush
--- a/Engine/Shaders/Private/LandscapeVertexFactory.ush
+++ b/Engine/Shaders/Private/LandscapeVertexFactory.ush
@@ -199,12 +199,9 @@ FVertexFactoryInput LoadVertexFactoryInputForHGS(uint TriangleIndex, int VertexI
 		uint Indices[6] = {i00, i11, i10, i00, i01, i10};
 		uint VertexId = Indices[TriangleIndex % 2 * 3 + VertexIndex];
 	#endif
-#if VF_USE_PRIMITIVE_SCENE_DATA
-	// Note: GetInstanceUserData() stores the GPU-Scene primitive ID
-	Input.InstanceIdOffset = GetInstanceUserData() | VF_TREAT_INSTANCE_ID_OFFSET_AS_PRIMITIVE_ID_FLAG;
-	Input.DrawInstanceId = 0;
-#endif // VF_USE_PRIMITIVE_SCENE_DATA
 
+	// Note: GetInstanceUserData() stores the GPU-Scene primitive ID
+	VF_GPUSCENE_SET_INPUT_FOR_RT(Input, GetInstanceUserData(), 0U);
 
 	Input.Position.x = VertexId % SubsectionSizeVertsCurrentLOD;
 	Input.Position.y = VertexId / SubsectionSizeVertsCurrentLOD;
@@ -231,11 +228,7 @@ FVertexFactoryInput LoadVertexFactoryInputForDynamicUpdate(uint TriangleIndex, i
 	Input.Position.z = LandscapeMVF.SubXY.x;
 	Input.Position.w = LandscapeMVF.SubXY.y;
 
-#if VF_USE_PRIMITIVE_SCENE_DATA
-	// GPUCULL_TODO: Makes use of DynamicDraw workaround, should instead refactor VF interfaces and provide the instance ID data directly.
-	Input.InstanceIdOffset = PrimitiveId | VF_TREAT_INSTANCE_ID_OFFSET_AS_PRIMITIVE_ID_FLAG;
-	Input.DrawInstanceId = 0;
-#endif // VF_USE_PRIMITIVE_SCENE_DATA
+	VF_GPUSCENE_SET_INPUT_FOR_RT(Input, PrimitiveId, 0U);
 
 	return Input;
 }
diff --git a/Engine/Shaders/Private/LocalVertexFactory.ush b/Engine/Shaders/Private/LocalVertexFactory.ush
--- a/Engine/Shaders/Private/LocalVertexFactory.ush
+++ b/Engine/Shaders/Private/LocalVertexFactory.ush
@@ -202,18 +202,16 @@ FVertexFactoryInput LoadVertexFactoryInputForHGS(uint TriangleIndex, int VertexI
 #if VF_USE_PRIMITIVE_SCENE_DATA
 	// Note: GetInstanceUserData() stores the GPU-Scene primitive ID
 	int PrimitiveId = GetInstanceUserData();
-	Input.InstanceIdOffset = PrimitiveId | VF_TREAT_INSTANCE_ID_OFFSET_AS_PRIMITIVE_ID_FLAG;
+	uint DrawInstanceId = 0;
 	// We have a unified path for both static meshes and instanced static meshes
 	// However due to static meshes being auto-instanced GetBaseInstanceIndex() may not match InstanceIndex() and a non-zero DrawInstanceId is generated
 	// Force DrawInstanceId to 0 when it is a static mesh by checking NumInstanceSceneDataEntries
-	if (GetPrimitiveData(PrimitiveId).NumInstanceSceneDataEntries == 1)
+	if (GetPrimitiveData(PrimitiveId).NumInstanceSceneDataEntries != 1)
 	{
-		Input.DrawInstanceId = 0;
-	}
-	else
-	{
-		Input.DrawInstanceId = InstanceIndex() - GetBaseInstanceIndex();
+		DrawInstanceId = InstanceIndex() - GetBaseInstanceIndex();
 	}
+
+	VF_GPUSCENE_SET_INPUT_FOR_RT(Input, PrimitiveId, DrawInstanceId);
 #endif // VF_USE_PRIMITIVE_SCENE_DATA
 
 #if USE_INSTANCING
@@ -245,11 +243,7 @@ FVertexFactoryInput LoadVertexFactoryInputForDynamicUpdate(uint TriangleIndex, i
 		Input.Position.z = LocalVF.VertexFetch_PositionBuffer[VertexOffset + 2];
 	#endif
 
-#if VF_USE_PRIMITIVE_SCENE_DATA
-		// GPUCULL_TODO: Makes use of DynamicDraw workaround, should instead refactor VF interfaces and provide the instance ID data directly.
-		Input.InstanceIdOffset = PrimitiveId | VF_TREAT_INSTANCE_ID_OFFSET_AS_PRIMITIVE_ID_FLAG;
-		Input.DrawInstanceId = 0;
-#endif // VF_USE_PRIMITIVE_SCENE_DATA
+	VF_GPUSCENE_SET_INPUT_FOR_RT(Input, PrimitiveId, 0U);
 
 	return Input;
 }
diff --git a/Engine/Shaders/Private/VertexFactoryCommon.ush b/Engine/Shaders/Private/VertexFactoryCommon.ush
--- a/Engine/Shaders/Private/VertexFactoryCommon.ush
+++ b/Engine/Shaders/Private/VertexFactoryCommon.ush
@@ -22,6 +22,7 @@
 	#define VF_GPUSCENE_GET_SHADOWMAP_UV_SCALE_BIAS(VFInput, LightmapDataIndex) \
 		GetLightmapData(LightmapDataIndex).ShadowMapCoordinateScaleBias
 
+	// GPUCULL_TODO: Makes use of DynamicDraw workaround, should instead refactor VF interfaces and provide the instance ID data directly.
 	#define VF_GPUSCENE_SET_INPUT_FOR_RT(_VFInputOut_, _PrimitiveId_, _DrawInstanceId_) \
 		_VFInputOut_.InstanceIdOffset = (_PrimitiveId_) | VF_TREAT_INSTANCE_ID_OFFSET_AS_PRIMITIVE_ID_FLAG;	\
 		_VFInputOut_.DrawInstanceId = _DrawInstanceId_;
