Added scoped traces for Niagara VM executions

diff --git a/Engine/Plugins/FX/Niagara/Source/Niagara/Private/NiagaraScriptExecutionContext.cpp b/Engine/Plugins/FX/Niagara/Source/Niagara/Private/NiagaraScriptExecutionContext.cpp
--- a/Engine/Plugins/FX/Niagara/Source/Niagara/Private/NiagaraScriptExecutionContext.cpp
+++ b/Engine/Plugins/FX/Niagara/Source/Niagara/Private/NiagaraScriptExecutionContext.cpp
@@ -228,6 +228,8 @@ bool FNiagaraScriptExecutionContextBase::Execute(uint32 NumInstances, const FScr
 
 bool FNiagaraScriptExecutionContextBase::ExecuteInternal_Experimental(uint32 NumInstances, const FScriptExecutionConstantBufferTable& ConstantBufferTable)
 {
+	TRACE_CPUPROFILER_EVENT_SCOPE(VectorVM_Experimental);
+
 	FNiagaraVMExecutableData& ExecData = Script->GetVMExecutableData();
 
 	FVectorVMSerializeState UESerializeState = { };
@@ -276,6 +278,8 @@ bool FNiagaraScriptExecutionContextBase::ExecuteInternal_Experimental(uint32 Num
 
 bool FNiagaraScriptExecutionContextBase::ExecuteInternal_Legacy(uint32 NumInstances, const FScriptExecutionConstantBufferTable& ConstantBufferTable)
 {
+	TRACE_CPUPROFILER_EVENT_SCOPE(VectorVM_Legacy);
+
 	FNiagaraVMExecutableData& ExecData = Script->GetVMExecutableData();
 
 	// If we have an optimization task it must be ready at this point
