Niagara: enabled legacy VM. #preflight none #rb none

diff --git a/Engine/Plugins/FX/Niagara/Source/Niagara/Niagara.Build.cs b/Engine/Plugins/FX/Niagara/Source/Niagara/Niagara.Build.cs
--- a/Engine/Plugins/FX/Niagara/Source/Niagara/Niagara.Build.cs
+++ b/Engine/Plugins/FX/Niagara/Source/Niagara/Niagara.Build.cs
@@ -71,7 +71,7 @@ public class Niagara : ModuleRules
             new string[]
             {
                 "VECTORVM_SUPPORTS_EXPERIMENTAL=1",
-                "VECTORVM_SUPPORTS_LEGACY=0"
+                "VECTORVM_SUPPORTS_LEGACY=1"
             });
     }
 }
diff --git a/Engine/Source/Runtime/VectorVM/VectorVM.Build.cs b/Engine/Source/Runtime/VectorVM/VectorVM.Build.cs
--- a/Engine/Source/Runtime/VectorVM/VectorVM.Build.cs
+++ b/Engine/Source/Runtime/VectorVM/VectorVM.Build.cs
@@ -28,7 +28,7 @@ public class VectorVM : ModuleRules
             new string[]
             {
                 "VECTORVM_SUPPORTS_EXPERIMENTAL=1",
-                "VECTORVM_SUPPORTS_LEGACY=0",
+                "VECTORVM_SUPPORTS_LEGACY=1"
             });
     }
 }
