Fixed math ops pins created by the user not being deletable after changing types

diff --git a/Engine/Plugins/FX/Niagara/Source/NiagaraEditor/Private/NiagaraNode.cpp b/Engine/Plugins/FX/Niagara/Source/NiagaraEditor/Private/NiagaraNode.cpp
--- a/Engine/Plugins/FX/Niagara/Source/NiagaraEditor/Private/NiagaraNode.cpp
+++ b/Engine/Plugins/FX/Niagara/Source/NiagaraEditor/Private/NiagaraNode.cpp
@@ -445,6 +445,7 @@ bool UNiagaraNode::ConvertNumericPinToType(UEdGraphPin* InGraphPin, FNiagaraType
 	InGraphPin->Modify();
 	InGraphPin->PinType = PinType;
 	InGraphPin->ResetDefaultValue();
+	PinTypeChanged(InGraphPin);
 	ReallocatePins();
 
 	return true;
diff --git a/Engine/Plugins/FX/Niagara/Source/NiagaraEditor/Private/NiagaraNodeOp.cpp b/Engine/Plugins/FX/Niagara/Source/NiagaraEditor/Private/NiagaraNodeOp.cpp
--- a/Engine/Plugins/FX/Niagara/Source/NiagaraEditor/Private/NiagaraNodeOp.cpp
+++ b/Engine/Plugins/FX/Niagara/Source/NiagaraEditor/Private/NiagaraNodeOp.cpp
@@ -108,6 +108,22 @@ FLinearColor UNiagaraNodeOp::GetNodeTitleColor() const
 	return GetDefault<UGraphEditorSettings>()->FunctionCallNodeTitleColor;
 }
 
+void UNiagaraNodeOp::PinTypeChanged(UEdGraphPin* Pin)
+{
+	FName PinName = Pin->PinName;
+
+	for(FAddedPinData& AddedPinData : AddedPins)
+	{
+		if(AddedPinData.PinName == PinName)
+		{
+			AddedPinData.PinType = Pin->PinType;
+			break;
+		}
+	}
+	
+	Super::PinTypeChanged(Pin);
+}
+
 void UNiagaraNodeOp::OnPostSynchronizationInReallocatePins() 
 {
 	FPinCollectorArray InputPins;
diff --git a/Engine/Plugins/FX/Niagara/Source/NiagaraEditor/Public/NiagaraNodeOp.h b/Engine/Plugins/FX/Niagara/Source/NiagaraEditor/Public/NiagaraNodeOp.h
--- a/Engine/Plugins/FX/Niagara/Source/NiagaraEditor/Public/NiagaraNodeOp.h
+++ b/Engine/Plugins/FX/Niagara/Source/NiagaraEditor/Public/NiagaraNodeOp.h
@@ -56,6 +56,7 @@ public:
 	virtual FText GetNodeTitle(ENodeTitleType::Type TitleType) const override;
 	virtual FText GetTooltipText() const override;
 	virtual FLinearColor GetNodeTitleColor() const override;
+	virtual void PinTypeChanged(UEdGraphPin* Pin) override;
 	//~ End EdGraphNode Interface
 
 	//~ Begin UNiagaraNode Interface
