- Fix race in IsRenderingEnabled using _GT vesrion of emitters rather than _RT

#rnx

diff --git a/Engine/Plugins/FX/Niagara/Source/Niagara/Public/NiagaraSystemRenderData.h b/Engine/Plugins/FX/Niagara/Source/Niagara/Public/NiagaraSystemRenderData.h
--- a/Engine/Plugins/FX/Niagara/Source/Niagara/Public/NiagaraSystemRenderData.h
+++ b/Engine/Plugins/FX/Niagara/Source/Niagara/Public/NiagaraSystemRenderData.h
@@ -76,7 +76,7 @@ public:
 	void OnSystemComplete(const FNiagaraSystemInstance& SystemInstance);
 	void RecacheRenderers(const FNiagaraSystemInstance& SystemInstance, const FNiagaraSystemInstanceController& Controller);
 
-	FORCEINLINE bool IsRenderingEnabled() const { return bRenderingEnabled && EmitterRenderers_GT.Num() > 0; }
+	FORCEINLINE bool IsRenderingEnabled() const { return bRenderingEnabled && (IsInRenderingThread() ? EmitterRenderers_RT.Num() > 0 : EmitterRenderers_GT.Num() > 0); }
 	FORCEINLINE void SetRenderingEnabled(bool bInEnabled) { bRenderingEnabled = bInEnabled; }
 	FORCEINLINE bool HasAnyMotionBlurEnabled() const { return bAnyMotionBlurEnabled; }
 
