Custom view frustum culling in Niagara.

- Reworked visibility settings into their own structure. Now has separete times for render based culling and view frustum culling.
- Improved Scalability Overrides UI.
- Now caching various view matrices on the world.
- Some 2D debug rendering helpers in Niagara Debug HUD.

diff --git a/Engine/Plugins/FX/Niagara/Source/Niagara/Classes/NiagaraEffectType.h b/Engine/Plugins/FX/Niagara/Source/Niagara/Classes/NiagaraEffectType.h
--- a/Engine/Plugins/FX/Niagara/Source/Niagara/Classes/NiagaraEffectType.h
+++ b/Engine/Plugins/FX/Niagara/Source/Niagara/Classes/NiagaraEffectType.h
@@ -128,22 +128,54 @@ struct FNiagaraGlobalBudgetScaling
 	uint32 bScaleSystemInstanceCountByGlobalBudgetUse : 1;
 
 	/** Effects will be culled if the global budget usage exceeds this fraction. A global budget usage of 1.0 means current global FX workload has reached it's max budget. Budgets are set by CVars under FX.Budget... */
-	UPROPERTY(EditAnywhere, Category = "Scalability", meta = (EditCondition = "bCullByGlobalBudget"))
+	UPROPERTY(EditAnywhere, Category = "Scalability", meta = (EditCondition = "bCullByGlobalBudget", DisplayInScalabilityValuesBar))
 	float MaxGlobalBudgetUsage;
 
 	/** When enabled, MaxDistance is scaled down by the global budget use based on this curve. Allows us to cull more aggressively when performance is poor.	*/
-	UPROPERTY(EditAnywhere, Category = "Budget Scaling", meta = (EditCondition = "bScaleMaxDistanceByGlobalBudgetUse"))
+	UPROPERTY(EditAnywhere, Category = "Budget Scaling", meta = (EditCondition = "bScaleMaxDistanceByGlobalBudgetUse", DisplayInScalabilityValuesBar))
 	FNiagaraLinearRamp MaxDistanceScaleByGlobalBudgetUse;
 	
 	/** When enabled, Max Effect Type Instances is scaled down by the global budget use based on this curve. Allows us to cull more aggressively when performance is poor.	*/
-	UPROPERTY(EditAnywhere, Category = "Budget Scaling", meta = (EditCondition = "bScaleMaxInstanceCountByGlobalBudgetUse"))
+	UPROPERTY(EditAnywhere, Category = "Budget Scaling", meta = (EditCondition = "bScaleMaxInstanceCountByGlobalBudgetUse", DisplayInScalabilityValuesBar))
 	FNiagaraLinearRamp MaxInstanceCountScaleByGlobalBudgetUse;
 	
 	/** When enabled, Max System Instances is scaled down by the global budget use based on this curve. Allows us to cull more aggressively when performance is poor.	*/
-	UPROPERTY(EditAnywhere, Category = "Budget Scaling", meta = (EditCondition = "bScaleSystemInstanceCountByGlobalBudgetUse"))
+	UPROPERTY(EditAnywhere, Category = "Budget Scaling", meta = (EditCondition = "bScaleSystemInstanceCountByGlobalBudgetUse", DisplayInScalabilityValuesBar))
 	FNiagaraLinearRamp MaxSystemInstanceCountScaleByGlobalBudgetUse;
 };
 
+/** Scalability settings for Niagara Systems for a particular platform set (unless overridden). */
+USTRUCT()
+struct FNiagaraSystemVisibilityCullingSettings
+{
+	GENERATED_USTRUCT_BODY()
+
+	FNiagaraSystemVisibilityCullingSettings();
+
+	/** Controls whether we cull when not rendered. This includes all reasons for being not rendered such as view frustum, occlusion and hidden flags etc. As this requires feedback form the rendering system, it cannot be used for pre-culling. */
+	UPROPERTY(EditAnywhere, Category = "Scalability", meta = (InlineEditConditionToggle))
+	uint32 bCullWhenNotRendered : 1;
+
+	/**
+	Controls whether view frustum culling is enabled. Niagara can do it's own separate view frustum culling to give coarse visibility culling that does not depend directly on rendering like CullWhenNotRendered.
+	This requires fixed bounds to be used in pre-culling.
+	*/
+	UPROPERTY(EditAnywhere, Category = "Scalability", meta = (InlineEditConditionToggle))
+	uint32 bCullByViewFrustum : 1;
+
+	/** If true and bCullByViewFrustum is enabled, we allow view frustum checks in pre-culling. Meaning we can cull off screen systems before they even spawn. */
+	UPROPERTY(EditAnywhere, Category = "Scalability", meta = (EditCondition = "bCullByViewFrustum"))
+	uint32 bAllowPreCullingByViewFrustum : 1;
+
+	/** Effects will be culled if they go longer than this time outside the view frustum. */
+	UPROPERTY(EditAnywhere, Category = "Scalability", meta = (EditCondition = "bCullByViewFrustum", DisplayInScalabilityValuesBar))
+	float MaxTimeOutsideViewFrustum;
+
+	/** Effects will be culled if they go longer than this time without being rendered. */
+	UPROPERTY(EditAnywhere, Category = "Scalability", meta = (EditCondition = "bCullWhenNotRendered", DisplayInScalabilityValuesBar))
+	float MaxTimeWithoutRender;
+};
+
 //////////////////////////////////////////////////////////////////////////
 
 /** Scalability settings for Niagara Systems for a particular platform set (unless overridden). */
@@ -165,15 +197,15 @@ struct FNiagaraSystemScalabilitySettings
 	/** Controls whether we should cull systems based on how many instances of the system are active. */
 	UPROPERTY(EditAnywhere, Category = "Scalability", DisplayName = "Cull By System Instance Count", meta = (InlineEditConditionToggle))
 	uint32 bCullPerSystemMaxInstanceCount : 1;
-	/** Controls whether visibility culling is enabled. */
-	UPROPERTY(EditAnywhere, Category = "Scalability", meta = (InlineEditConditionToggle))
-	uint32 bCullByMaxTimeWithoutRender : 1;
-
 
 	/** Effects of this type are culled beyond this distance. */
 	UPROPERTY(EditAnywhere, Category = "Scalability", meta = (EditCondition = "bCullByDistance", DisplayInScalabilityValuesBar))
 	float MaxDistance;
 
+	/** Deprecated visibility culling property. Moved into FNiagaraSystemVisibilityCullingSettings. */
+	UPROPERTY()
+	uint32 bCullByMaxTimeWithoutRender_DEPRECATED : 1;
+
 	/** 
 	Effects of this type will be culled when total active instances using this same EffectType exceeds this number. 
 	If the effect type has a significance handler, instances are sorted by their significance and only the N most significant will be kept. The rest are culled.
@@ -190,14 +222,10 @@ struct FNiagaraSystemScalabilitySettings
 	UPROPERTY(EditAnywhere, Category = "Scalability", meta = (EditCondition = "bCullPerSystemMaxInstanceCount", DisplayInScalabilityValuesBar))
 	int32 MaxSystemInstances;
 
-	//TODO:
-	/** The effect is culled when it's bounds take up less that this fraction of the total screen area. Only usable with fixed bounds. */
-	//float ScreenFraction;
+	/** Depreceated visibiltiy culling property. Moved into FNiagaraSystemVisibilityCullingSettings. */
+	UPROPERTY()
+	float MaxTimeWithoutRender_DEPRECATED;
 
-	/** Effects will be culled if they go more than this length of time without being rendered. */
-	UPROPERTY(EditAnywhere, Category = "Scalability", meta = (EditCondition = "bCullByMaxTimeWithoutRender", DisplayInScalabilityValuesBar))
-	float MaxTimeWithoutRender;
-		
 	/** Controls what, if any, proxy will be used in place of culled systems.  */
 	UPROPERTY(EditAnywhere, Category = "Scalability", meta = (DisplayInScalabilityValuesBar))
 	ENiagaraCullProxyMode CullProxyMode = ENiagaraCullProxyMode::None;
@@ -209,14 +237,19 @@ struct FNiagaraSystemScalabilitySettings
 	*/
 	UPROPERTY(EditAnywhere, Category = "Scalability", meta = (EditCondition = "CullProxyMode != ENiagaraCullProxyMode::None", DisplayInScalabilityValuesBar))
 	int32 MaxSystemProxies = 32;
+	
+	/** Settings controlling how systems are culled by visibility. */
+	UPROPERTY(EditAnywhere, Category = "Scalability", meta = (DisplayInScalabilityValuesBar))
+	FNiagaraSystemVisibilityCullingSettings VisibilityCulling;
 
 	/** Settings related to scaling down FX based on the current budget usage. */
-	UPROPERTY(EditAnywhere, Category = "Scalability")
+	UPROPERTY(EditAnywhere, Category = "Scalability", meta = (DisplayInScalabilityValuesBar))
 	FNiagaraGlobalBudgetScaling BudgetScaling;
 	
 	FNiagaraSystemScalabilitySettings();
 
 	void Clear();
+	void PostLoad(int32 Version);
 };
 
 /** Container struct for an array of system scalability settings. Enables details customization and data validation. */
@@ -247,7 +280,7 @@ struct FNiagaraSystemScalabilityOverride : public FNiagaraSystemScalabilitySetti
 	uint32 bOverridePerSystemInstanceCountSettings : 1;
 	/** Controls whether we override the visibility culling settings. */
 	UPROPERTY(EditAnywhere, Category = "Override")
-	uint32 bOverrideTimeSinceRendererSettings : 1;
+	uint32 bOverrideVisibilitySettings : 1;
 	/** Controls whether we override the global budget scaling settings. */
 	UPROPERTY(EditAnywhere, Category = "Override")
 	uint32 bOverrideGlobalBudgetScalingSettings : 1;
diff --git a/Engine/Plugins/FX/Niagara/Source/Niagara/Private/NiagaraDebugHud.cpp b/Engine/Plugins/FX/Niagara/Source/Niagara/Private/NiagaraDebugHud.cpp
--- a/Engine/Plugins/FX/Niagara/Source/Niagara/Private/NiagaraDebugHud.cpp
+++ b/Engine/Plugins/FX/Niagara/Source/Niagara/Private/NiagaraDebugHud.cpp
@@ -1046,6 +1046,8 @@ void FNiagaraDebugHud::Draw(FNiagaraWorldManager* WorldManager, UCanvas* Canvas,
 	// Draw in world components
 	DrawComponents(WorldManager, Canvas);
 
+	DrawDebugGeomerty(WorldManager, Canvas);
+
 	// Draw overview
 	DrawOverview(WorldManager, Canvas->Canvas);
 
@@ -1076,6 +1078,101 @@ void FNiagaraDebugHud::Draw(FNiagaraWorldManager* WorldManager, UCanvas* Canvas,
 	LastDrawTime = CurrTime;
 }
 
+void FNiagaraDebugHud::AddLine2D(FVector2D Start, FVector2D End, FLinearColor Color, float Thickness, float Lifetime)
+{
+	if (NiagaraDebugLocal::Settings.bEnabled)
+	{
+		FDebugLine2D NewLine;
+		NewLine.Start = Start;
+		NewLine.End = End;
+		NewLine.Color = Color;
+		NewLine.Thickness = Thickness;
+		NewLine.Lifetime = Lifetime;
+		Lines2D.Add(NewLine);
+	}
+}
+
+void FNiagaraDebugHud::AddCircle2D(FVector2D Pos, float Rad, float Segments, FLinearColor Color, float Thickness, float Lifetime)
+{
+	if (NiagaraDebugLocal::Settings.bEnabled)
+	{
+		FDebugCircle2D NewCircle;
+		NewCircle.Pos = Pos;
+		NewCircle.Rad = Rad;
+		NewCircle.Segments = Segments;
+		NewCircle.Color = Color;
+		NewCircle.Thickness = Thickness;
+		NewCircle.Lifetime = Lifetime;
+		Circles2D.Add(NewCircle);
+	}
+}
+
+void FNiagaraDebugHud::AddBox2D(FVector2D Pos, FVector2D Extents, FLinearColor Color, float Thickness, float Lifetime)
+{
+	if (NiagaraDebugLocal::Settings.bEnabled)
+	{
+		FDebugBox2D NewBox;
+		NewBox.Pos = Pos;
+		NewBox.Extents = Extents;
+		NewBox.Color = Color;
+		NewBox.Thickness = Thickness;
+		NewBox.Lifetime = Lifetime;
+		Boxes2D.Add(NewBox);
+	}
+}
+
+void FNiagaraDebugHud::DrawDebugGeomerty(class FNiagaraWorldManager* WorldManager, class UCanvas* Canvas)
+{
+	UWorld* World = WorldManager->GetWorld();
+
+	if (World == nullptr || Canvas == nullptr)
+	{
+		return;
+	}
+
+	FCanvas* DrawCanvas = Canvas->Canvas;
+
+	FVector2D Size(Canvas->ClipX, Canvas->ClipY);
+	
+	for (auto it = Lines2D.CreateIterator(); it; ++it)
+	{
+		FDebugLine2D& Line = *it;
+		DrawDebugCanvas2DLine(Canvas, Line.Start * Size, Line.End* Size, Line.Color, Line.Thickness);
+
+		Line.Lifetime -= DeltaSeconds;
+		if (Line.Lifetime <= 0.0f)
+		{
+			it.RemoveCurrent();
+		}
+	}
+	for (auto it = Circles2D.CreateIterator(); it; ++it)
+	{
+		FDebugCircle2D& Circle = *it;
+		DrawDebugCanvas2DCircle(Canvas, Circle.Pos * Size, Circle.Rad * Size.X, Circle.Segments, Circle.Color, Circle.Thickness);
+
+		Circle.Lifetime -= DeltaSeconds;
+		if (Circle.Lifetime <= 0.0f)
+		{
+			it.RemoveCurrent();
+		}
+	}
+	for (auto it = Boxes2D.CreateIterator(); it; ++it)
+	{
+		FDebugBox2D& Box = *it;
+
+		FBox2D Box2D;
+		Box2D.Min = (Box.Pos - Box.Extents) * Size;
+		Box2D.Max = (Box.Pos + Box.Extents) * Size;
+		DrawDebugCanvas2DBox(Canvas, Box2D, Box.Color, Box.Thickness);
+
+		Box.Lifetime -= DeltaSeconds;
+		if (Box.Lifetime <= 0.0f)
+		{
+			it.RemoveCurrent();
+		}
+	}
+}
+
 template<typename T>
 struct FGraph
 {
diff --git a/Engine/Plugins/FX/Niagara/Source/Niagara/Private/NiagaraDebugHud.h b/Engine/Plugins/FX/Niagara/Source/Niagara/Private/NiagaraDebugHud.h
--- a/Engine/Plugins/FX/Niagara/Source/Niagara/Private/NiagaraDebugHud.h
+++ b/Engine/Plugins/FX/Niagara/Source/Niagara/Private/NiagaraDebugHud.h
@@ -239,6 +239,7 @@ private:
 	void DrawValidation(class FNiagaraWorldManager* WorldManager, class FCanvas* DrawCanvas, FVector2D& TextLocation);
 	void DrawComponents(class FNiagaraWorldManager* WorldManager, class UCanvas* Canvas);
 	void DrawMessages(class FNiagaraWorldManager* WorldManager, class FCanvas* DrawCanvas, FVector2D& TextLocation);
+	void DrawDebugGeomerty(class FNiagaraWorldManager* WorldManager, class UCanvas* DrawCanvas);
 
 private:
 	TWeakObjectPtr<class UWorld>	WeakWorld;
@@ -315,4 +316,41 @@ private:
 
 	/** Generic messages that the debugger or other parts of Niagara can post to the HUD. */
 	TMap<FName, FNiagaraDebugMessage> Messages;
+
+	//Additional debug geometry helpers
+	struct FDebugLine2D
+	{
+		FVector2D Start;
+		FVector2D End;
+		FLinearColor Color;
+		float Thickness;
+		float Lifetime;
+	};
+	TArray<FDebugLine2D> Lines2D;
+	struct FDebugCircle2D
+	{
+		FVector2D Pos;
+		float Rad;
+		float Segments;
+		FLinearColor Color;
+		float Thickness;
+		float Lifetime;
+	};
+	TArray<FDebugCircle2D> Circles2D;
+	struct FDebugBox2D
+	{
+		FVector2D Pos;
+		FVector2D Extents;
+		FLinearColor Color;
+		float Thickness;
+		float Lifetime;
+	};
+	TArray<FDebugBox2D> Boxes2D;
+
+public:
+	/** Add a 2D line to the debug renering. Positions are in normalized screen space. (0,0) in top left, (1,1) bottom right.*/
+	void AddLine2D(FVector2D Start, FVector2D End, FLinearColor Color, float Thickness, float Lifetime);
+	void AddCircle2D(FVector2D Pos, float Rad, float Segments, FLinearColor Color, float Thickness, float Lifetime);
+	void AddBox2D(FVector2D Pos, FVector2D Extents, FLinearColor Color, float Thickness, float Lifetime);
+	//Additional debug geometry helpers END
 };
diff --git a/Engine/Plugins/FX/Niagara/Source/Niagara/Private/NiagaraEffectType.cpp b/Engine/Plugins/FX/Niagara/Source/Niagara/Private/NiagaraEffectType.cpp
--- a/Engine/Plugins/FX/Niagara/Source/Niagara/Private/NiagaraEffectType.cpp
+++ b/Engine/Plugins/FX/Niagara/Source/Niagara/Private/NiagaraEffectType.cpp
@@ -58,6 +58,11 @@ void UNiagaraEffectType::PostLoad()
 		}
 	}
 
+	for (FNiagaraSystemScalabilitySettings& SysScalabilitySetting : SystemScalabilitySettings.Settings)
+	{
+		SysScalabilitySetting.PostLoad(NiagaraVer);
+	}
+
 	//Apply platform set redirectors
 	auto ApplyPlatformSetRedirects = [](UObject* Owner, FNiagaraPlatformSet& Platforms)
 	{
@@ -167,13 +172,24 @@ FNiagaraGlobalBudgetScaling::FNiagaraGlobalBudgetScaling()
 	MaxSystemInstanceCountScaleByGlobalBudgetUse = DefaultBudgetScaleRamp;
 }
 
+//////////////////////////////////////////////////////////////////////////
+ 
+FNiagaraSystemVisibilityCullingSettings::FNiagaraSystemVisibilityCullingSettings()
+	: bCullWhenNotRendered(false)
+	, bCullByViewFrustum(false)
+	, bAllowPreCullingByViewFrustum(false)
+	, MaxTimeOutsideViewFrustum(1.0f)
+	, MaxTimeWithoutRender(1.0f)
+{
+}
+
 //////////////////////////////////////////////////////////////////////////
 
 FNiagaraSystemScalabilityOverride::FNiagaraSystemScalabilityOverride()
 	: bOverrideDistanceSettings(false)
 	, bOverrideInstanceCountSettings(false)
 	, bOverridePerSystemInstanceCountSettings(false)
-	, bOverrideTimeSinceRendererSettings(false)
+	, bOverrideVisibilitySettings(false)
 	, bOverrideGlobalBudgetScalingSettings(false)
 	, bOverrideCullProxySettings(false)
 {
@@ -188,19 +204,31 @@ void FNiagaraSystemScalabilitySettings::Clear()
 {
 	Platforms = FNiagaraPlatformSet();
 	bCullByDistance = false;
-	bCullByMaxTimeWithoutRender = false;
 	bCullMaxInstanceCount = false;
-	bCullPerSystemMaxInstanceCount = false;
+	bCullPerSystemMaxInstanceCount = false;	
 	MaxDistance = 0.0f;
 	MaxInstances = 0;
 	MaxSystemInstances = 0;
-	MaxTimeWithoutRender = 0.0f;
+
+
+	MaxTimeWithoutRender_DEPRECATED = 0.0f;
+	bCullByMaxTimeWithoutRender_DEPRECATED = false;
+	VisibilityCulling = FNiagaraSystemVisibilityCullingSettings();
 
 	BudgetScaling = FNiagaraGlobalBudgetScaling();
 	CullProxyMode = ENiagaraCullProxyMode::None;
 	MaxSystemProxies = 32;
 }
 
+void FNiagaraSystemScalabilitySettings::PostLoad(int32 Version)
+{
+	if (Version < FNiagaraCustomVersion::VisibilityCullingImprovements)
+	{
+		VisibilityCulling.bCullWhenNotRendered = bCullByMaxTimeWithoutRender_DEPRECATED;
+		VisibilityCulling.MaxTimeWithoutRender = MaxTimeWithoutRender_DEPRECATED;
+	}
+}
+
 FNiagaraEmitterScalabilitySettings::FNiagaraEmitterScalabilitySettings()
 {
 	Clear();
diff --git a/Engine/Plugins/FX/Niagara/Source/Niagara/Private/NiagaraSystem.cpp b/Engine/Plugins/FX/Niagara/Source/Niagara/Private/NiagaraSystem.cpp
--- a/Engine/Plugins/FX/Niagara/Source/Niagara/Private/NiagaraSystem.cpp
+++ b/Engine/Plugins/FX/Niagara/Source/Niagara/Private/NiagaraSystem.cpp
@@ -780,6 +780,11 @@ void UNiagaraSystem::PostLoad()
 		}
 	}
 
+	for (FNiagaraSystemScalabilityOverride& Override : SystemScalabilityOverrides.Overrides)
+	{
+		Override.PostLoad(NiagaraVer);
+	}
+
 #if UE_EDITOR
 	ExposedParameters.RecreateRedirections();
 #endif
@@ -3378,10 +3383,13 @@ void UNiagaraSystem::ResolveScalabilitySettings()
 					CurrentScalabilitySettings.MaxSystemInstances = Override.MaxSystemInstances;
 				}
 
-				if (Override.bOverrideTimeSinceRendererSettings)
+				if (Override.bOverrideVisibilitySettings)
 				{
-					CurrentScalabilitySettings.bCullByMaxTimeWithoutRender = Override.bCullByMaxTimeWithoutRender;
-					CurrentScalabilitySettings.MaxTimeWithoutRender = Override.MaxTimeWithoutRender;
+					CurrentScalabilitySettings.VisibilityCulling.bCullWhenNotRendered = Override.VisibilityCulling.bCullWhenNotRendered;
+					CurrentScalabilitySettings.VisibilityCulling.bCullByViewFrustum = Override.VisibilityCulling.bCullByViewFrustum;
+					CurrentScalabilitySettings.VisibilityCulling.bAllowPreCullingByViewFrustum = Override.VisibilityCulling.bAllowPreCullingByViewFrustum;
+					CurrentScalabilitySettings.VisibilityCulling.MaxTimeWithoutRender = Override.VisibilityCulling.MaxTimeWithoutRender;
+					CurrentScalabilitySettings.VisibilityCulling.MaxTimeOutsideViewFrustum = Override.VisibilityCulling.MaxTimeOutsideViewFrustum;
 				}
 
  				if (Override.bOverrideGlobalBudgetScalingSettings)
diff --git a/Engine/Plugins/FX/Niagara/Source/Niagara/Private/NiagaraSystemInstance.cpp b/Engine/Plugins/FX/Niagara/Source/Niagara/Private/NiagaraSystemInstance.cpp
--- a/Engine/Plugins/FX/Niagara/Source/Niagara/Private/NiagaraSystemInstance.cpp
+++ b/Engine/Plugins/FX/Niagara/Source/Niagara/Private/NiagaraSystemInstance.cpp
@@ -1739,24 +1739,16 @@ float FNiagaraSystemInstance::GetLODDistance()
 	LODDistance = DefaultLODDistance;
 
 	// If we are inside the WorldManager tick we will use the cache player view locations as we can be ticked on different threads
-	if (WorldManager->CachedPlayerViewLocationsValid())
+	if (WorldManager->GetCachedViewInfo().Num() > 0)
 	{
-		TArrayView<const FVector> PlayerViewLocations = WorldManager->GetCachedPlayerViewLocations();
-		if (PlayerViewLocations.Num() == 0)
+		// We are being ticked inside the WorldManager and can safely use the list of cached player view locations
+		float LODDistanceSqr = FMath::Square(WORLD_MAX);
+		for (const FNiagaraCachedViewInfo& ViewInfo : WorldManager->GetCachedViewInfo())
 		{
-			LODDistance = DefaultLODDistance;
-		}
-		else
-		{
-			// We are being ticked inside the WorldManager and can safely use the list of cached player view locations
-			float LODDistanceSqr = FMath::Square(WORLD_MAX);
-			for (const FVector& ViewLocation : PlayerViewLocations)
-			{
-				const float DistanceToEffectSqr = FVector(ViewLocation - EffectLocation).SizeSquared();
-				LODDistanceSqr = FMath::Min(LODDistanceSqr, DistanceToEffectSqr);
-			}
-			LODDistance = FMath::Sqrt(LODDistanceSqr);
+			const float DistanceToEffectSqr = FVector(ViewInfo.ViewToWorld.GetOrigin() - EffectLocation).SizeSquared();
+			LODDistanceSqr = FMath::Min(LODDistanceSqr, DistanceToEffectSqr);
 		}
+		LODDistance = FMath::Sqrt(LODDistanceSqr);
 	}
 	else
 	{
diff --git a/Engine/Plugins/FX/Niagara/Source/Niagara/Private/NiagaraWorldManager.cpp b/Engine/Plugins/FX/Niagara/Source/Niagara/Private/NiagaraWorldManager.cpp
--- a/Engine/Plugins/FX/Niagara/Source/Niagara/Private/NiagaraWorldManager.cpp
+++ b/Engine/Plugins/FX/Niagara/Source/Niagara/Private/NiagaraWorldManager.cpp
@@ -25,6 +25,7 @@
 #include "NiagaraFunctionLibrary.h"
 #include "NiagaraCullProxyComponent.h"
 #include "GameDelegates.h"
+#include "DrawDebugHelpers.h"
 
 #if WITH_EDITORONLY_DATA
 #include "EditorViewportClient.h"
@@ -828,8 +829,7 @@ void FNiagaraWorldManager::PostActorTick(float DeltaSeconds)
 	}
 
 	// Clear cached player view location list, it should never be used outside of the world tick
-	bCachedPlayerViewLocationsValid = false;
-	CachedPlayerViewLocations.Reset();
+	CachedViewInfo.Reset();
 
 	// Delete any instances that were pending deletion
 	//-TODO: This could be done after each system sim has run
@@ -1012,39 +1012,52 @@ void FNiagaraWorldManager::Tick(ETickingGroup TickGroup, float DeltaSeconds, ELe
 			return;
 		}
 
-		// Cache player view locations for all system instances to access
+		bool bUseWorldCachedViews = !World->GetPlayerControllerIterator();
+#if WITH_EDITOR
+		if (GCurrentLevelEditingViewportClient && (GCurrentLevelEditingViewportClient->GetWorld() == World))
+		{
+			bUseWorldCachedViews = true;
+		}
+#endif
+		// Cache player view info for all system instances to access
 		//-TODO: Do we need to do this per tick group?
-		bCachedPlayerViewLocationsValid = true;
-		if (World->GetPlayerControllerIterator())
+		if (bUseWorldCachedViews)
+		{
+			for (int32 i = 0; i < World->CachedViewInfoRenderedLastFrame.Num(); ++i)
+			{
+				FWorldCachedViewInfo& WorldViewInfo = World->CachedViewInfoRenderedLastFrame[i];
+
+				FNiagaraCachedViewInfo& ViewInfo = CachedViewInfo.AddDefaulted_GetRef();
+				ViewInfo.Init(WorldViewInfo);
+			}
+		}
+		else
 		{
 			for (FConstPlayerControllerIterator Iterator = World->GetPlayerControllerIterator(); Iterator; ++Iterator)
 			{
 				APlayerController* PlayerController = Iterator->Get();
 				if (PlayerController && PlayerController->IsLocalPlayerController())
 				{
-					FVector* POVLoc = new(CachedPlayerViewLocations) FVector;
+					FVector POVLoc;
 					FRotator POVRotation;
-					PlayerController->GetPlayerViewPoint(*POVLoc, POVRotation);
+					PlayerController->GetPlayerViewPoint(POVLoc, POVRotation);
+					FRotationTranslationMatrix ViewToWorld(POVRotation, POVLoc);
+					FMatrix ViewMat = ViewToWorld.InverseFast() * FMatrix(
+						FPlane(0, 0, 1, 0),
+						FPlane(1, 0, 0, 0),
+						FPlane(0, 1, 0, 0),
+						FPlane(0, 0, 0, 1));
+
+					FNiagaraCachedViewInfo& ViewInfo = CachedViewInfo.AddDefaulted_GetRef();
+					FWorldCachedViewInfo WorldViewInfo;
+					WorldViewInfo.ViewMatrix = ViewMat;
+					WorldViewInfo.ProjectionMatrix = PlayerController->PlayerCameraManager->GetCameraCacheView().CalculateProjectionMatrix();
+					WorldViewInfo.ViewProjectionMatrix = ViewMat * WorldViewInfo.ProjectionMatrix;
+					WorldViewInfo.ViewToWorld = ViewToWorld;
+					ViewInfo.Init(WorldViewInfo);
 				}
 			}
-		}
-		else
-		{
-			CachedPlayerViewLocations.Append(World->ViewLocationsRenderedLastFrame);
-		}
-		
-		if (CachedPlayerViewLocations.Num() == 0)
-		{
-			bCachedPlayerViewLocationsValid = false;
-		}
-
-#if WITH_EDITORONLY_DATA
-		if (GCurrentLevelEditingViewportClient && (GCurrentLevelEditingViewportClient->GetWorld() == World))
-		{
-			const FViewportCameraTransform& ViewTransform = GCurrentLevelEditingViewportClient->GetViewTransform();
-			CachedPlayerViewLocations.Add(ViewTransform.GetLocation());
-		}
-#endif
+		}		
 
 		UpdateScalabilityManagers(DeltaSeconds, false);
 
@@ -1259,6 +1272,17 @@ void FNiagaraWorldManager::CalculateScalabilityState(UNiagaraSystem* System, con
 
 	DistanceCull(EffectType, ScalabilitySettings, Location, OutState);
 
+	if (GEnableNiagaraVisCulling)
+	{
+		if (System->bFixedBounds)//We have no component in this path so we require fixed bounds for view based checks.
+		{
+			FBoxSphereBounds Bounds(System->GetFixedBounds());
+			Bounds.Origin = Location;
+			float TimeSinceRendered = 0.0f;
+			ViewBasedCulling(EffectType, ScalabilitySettings, Bounds.GetSphere(), TimeSinceRendered, bIsPreCull, OutState);
+		}
+	}
+
 	//If we have no significance handler there is no concept of relative significance for these systems so we can just pre cull if we go over the instance count.
 	if (GEnableNiagaraInstanceCountCulling && bIsPreCull && EffectType->GetSignificanceHandler() == nullptr)
 	{
@@ -1275,6 +1299,7 @@ void FNiagaraWorldManager::CalculateScalabilityState(UNiagaraSystem* System, con
 	//TODO: More progressive scalability options?
 }
 
+extern float GLastRenderTimeSafetyBias;
 void FNiagaraWorldManager::CalculateScalabilityState(UNiagaraSystem* System, const FNiagaraSystemScalabilitySettings& ScalabilitySettings, UNiagaraEffectType* EffectType, UNiagaraComponent* Component, bool bIsPreCull, float WorstGlobalBudgetUse, FNiagaraScalabilityState& OutState)
 {
 	if (GetScalabilityCullingMode() == ENiagaraScalabilityCullingMode::Disabled)
@@ -1286,9 +1311,26 @@ void FNiagaraWorldManager::CalculateScalabilityState(UNiagaraSystem* System, con
 
 	DistanceCull(EffectType, ScalabilitySettings, Component, OutState);
 	
-	if ((GbAllowVisibilityCullingForDynamicBounds || System->bFixedBounds) && bAppHasFocus && GEnableNiagaraVisCulling)
+	if (GEnableNiagaraVisCulling)
 	{
-		VisibilityCull(EffectType, ScalabilitySettings, Component, OutState);
+		FBoxSphereBounds Bounds = Component->CalcBounds(Component->GetComponentToWorld());		
+
+		float TimeSinceRendered = FMath::Max(0.0f, GetWorld()->LastRenderTime - Component->GetLastRenderTime() - World->GetDeltaSeconds() - GLastRenderTimeSafetyBias);		
+
+		if(bIsPreCull)
+		{
+			if (System->bFixedBounds)//We need valid bounds for view based culling so if we're preculling, restrict to systems with fixed bounds.
+			{
+				ViewBasedCulling(EffectType, ScalabilitySettings, Bounds.GetSphere(), TimeSinceRendered, bIsPreCull, OutState);
+			}
+		}
+		else
+		{
+			if (GbAllowVisibilityCullingForDynamicBounds || System->bFixedBounds)
+			{
+				ViewBasedCulling(EffectType, ScalabilitySettings, Bounds.GetSphere(), TimeSinceRendered, bIsPreCull, OutState);
+			}
+		}
 	}
 
 	//Only apply hard instance count cull limit for precull if we have no significance handler.
@@ -1370,15 +1412,68 @@ void FNiagaraWorldManager::SortedSignificanceCull(UNiagaraEffectType* EffectType
 #endif
 }
 
-void FNiagaraWorldManager::VisibilityCull(UNiagaraEffectType* EffectType, const FNiagaraSystemScalabilitySettings& ScalabilitySettings, UNiagaraComponent* Component, FNiagaraScalabilityState& OutState)
+void FNiagaraWorldManager::ViewBasedCulling(UNiagaraEffectType* EffectType, const FNiagaraSystemScalabilitySettings& ScalabilitySettings, FSphere BoundingSphere, float ComponentTimeSinceRendered, bool bIsPrecull, FNiagaraScalabilityState& OutState)
 {
 	if (GetScalabilityCullingMode() != ENiagaraScalabilityCullingMode::Enabled)
 	{
 		return;
 	}
 
-	float TimeSinceRendered = Component->GetSafeTimeSinceRendered(World->TimeSeconds);
-	bool bCull = Component->GetLastRenderTime() >= 0.0f && ScalabilitySettings.bCullByMaxTimeWithoutRender && TimeSinceRendered > ScalabilitySettings.MaxTimeWithoutRender;
+	bool bInsideAnyView = !ScalabilitySettings.VisibilityCulling.bCullByViewFrustum;
+
+	//Iterator over all views to calculate and check the bounds against the view frustum.
+	if (ScalabilitySettings.VisibilityCulling.bCullByViewFrustum)
+	{
+		for (FNiagaraCachedViewInfo& ViewInfo : CachedViewInfo)
+		{
+			//First check the view frustums to see if any part of the sphere is inside.
+			bool bInsideThisView = true;
+			if (bInsideAnyView == false)//If we already know we're in any view (or not doing frustum checks) then we can skip the rest.
+			{
+				for (FPlane& FrustumPlane : ViewInfo.FrutumPlanes)
+				{
+					if (FrustumPlane.IsValid())
+					{
+						//((dot(Plane, BSphere.xyz) - Plane.w) > BSphere.w);
+						bool bInside = FrustumPlane.PlaneDot(BoundingSphere.Center) <= BoundingSphere.W;
+						if (!bInside)
+						{
+							bInsideThisView = false;
+							break;
+						}
+					}
+				}
+			}
+
+			if (bInsideThisView)
+			{
+				bInsideAnyView = true;
+			}
+		}
+	}
+
+	float TimeSinceInsideView = 0.0f;
+	if (bInsideAnyView)
+	{
+		OutState.LastVisibleTime = World->GetTimeSeconds();
+	}
+	else
+	{
+		TimeSinceInsideView = World->GetTimeSeconds() - OutState.LastVisibleTime;
+	}
+
+	bool bCullByOutsideViewFrustum = ScalabilitySettings.VisibilityCulling.bCullByViewFrustum &&
+		(!bIsPrecull || ScalabilitySettings.VisibilityCulling.bAllowPreCullingByViewFrustum) &&
+		OutState.LastVisibleTime > ScalabilitySettings.VisibilityCulling.MaxTimeOutsideViewFrustum;
+
+	//Check for the component having been rendered recently. If the app doesn't have focus we skip this to avoid issues when alt-tabbing away from the game/editor.
+	float TimeSinceWorldRendered = World->GetTimeSeconds() - World->LastRenderTime;	
+	bool bCullByNotRendered =	bAppHasFocus && 
+								ScalabilitySettings.VisibilityCulling.bCullWhenNotRendered && 
+		ComponentTimeSinceRendered > ScalabilitySettings.VisibilityCulling.MaxTimeWithoutRender;
+
+	//TODO: Pull screen size out into it's own debug flag in the scalability state.
+	bool bCull = bCullByNotRendered || !bInsideAnyView;
 
 	OutState.bCulled |= bCull;
 #if DEBUG_SCALABILITY_STATE
@@ -1444,13 +1539,13 @@ void FNiagaraWorldManager::DistanceCull(UNiagaraEffectType* EffectType, const FN
 	{
 		LODDistance = Component->PreviewLODDistance;
 	}
-	else if(bCachedPlayerViewLocationsValid)
+	else if(GetCachedViewInfo().Num() > 0)
 	{
 		float ClosestDistSq = FLT_MAX;
 		FVector Location = Component->GetComponentLocation();
-		for (FVector ViewLocation : CachedPlayerViewLocations)
+		for (const FNiagaraCachedViewInfo& ViewInfo : GetCachedViewInfo())
 		{
-			ClosestDistSq = FMath::Min(ClosestDistSq, FVector::DistSquared(ViewLocation, Location));
+			ClosestDistSq = FMath::Min(ClosestDistSq, FVector::DistSquared(ViewInfo.ViewToWorld.GetOrigin(), Location));
 		}
 
 		LODDistance = FMath::Sqrt(ClosestDistSq);
@@ -1496,12 +1591,12 @@ void FNiagaraWorldManager::DistanceCull(UNiagaraEffectType* EffectType, const FN
 
 void FNiagaraWorldManager::DistanceCull(UNiagaraEffectType* EffectType, const FNiagaraSystemScalabilitySettings& ScalabilitySettings, FVector Location, FNiagaraScalabilityState& OutState)
 {
-	if (GetScalabilityCullingMode() == ENiagaraScalabilityCullingMode::Enabled && bCachedPlayerViewLocationsValid)
+	if (GetScalabilityCullingMode() == ENiagaraScalabilityCullingMode::Enabled && GetCachedViewInfo().Num() > 0)
 	{
 		float ClosestDistSq = FLT_MAX;
-		for (FVector ViewLocation : CachedPlayerViewLocations)
+		for (const FNiagaraCachedViewInfo& ViewInfo : GetCachedViewInfo())
 		{
-			ClosestDistSq = FMath::Min(ClosestDistSq, FVector::DistSquared(ViewLocation, Location));
+			ClosestDistSq = FMath::Min(ClosestDistSq, FVector::DistSquared(ViewInfo.ViewToWorld.GetOrigin(), Location));
 		}
 
 		if (GEnableNiagaraDistanceCulling && ScalabilitySettings.bCullByDistance)
@@ -1691,3 +1786,37 @@ FAutoConsoleCommandWithWorld GDumpNiagaraScalabilityData(
 }));
 
 #endif
+
+void FNiagaraCachedViewInfo::Init(const FWorldCachedViewInfo& WorldViewInfo)
+{
+	ViewMat = WorldViewInfo.ViewMatrix;
+	ProjectionMat = WorldViewInfo.ProjectionMatrix;
+	ViewToWorld = WorldViewInfo.ViewToWorld;
+	ViewProjMat = WorldViewInfo.ViewProjectionMatrix;
+
+	FrutumPlanes.SetNumUninitialized(6);
+	if (!ViewProjMat.GetFrustumNearPlane(FrutumPlanes[0]))
+	{
+		FrutumPlanes[0] = FPlane(0.0f, 0.0f, 0.0f, 0.0f);
+	}
+	if (!ViewProjMat.GetFrustumFarPlane(FrutumPlanes[1]))
+	{
+		FrutumPlanes[1] = FPlane(0.0f, 0.0f, 0.0f, 0.0f);
+	}
+	if (!ViewProjMat.GetFrustumTopPlane(FrutumPlanes[2]))
+	{
+		FrutumPlanes[2] = FPlane(0.0f, 0.0f, 0.0f, 0.0f);
+	}
+	if (!ViewProjMat.GetFrustumBottomPlane(FrutumPlanes[3]))
+	{
+		FrutumPlanes[3] = FPlane(0.0f, 0.0f, 0.0f, 0.0f);
+	}
+	if (!ViewProjMat.GetFrustumLeftPlane(FrutumPlanes[4]))
+	{
+		FrutumPlanes[4] = FPlane(0.0f, 0.0f, 0.0f, 0.0f);
+	}
+	if (!ViewProjMat.GetFrustumRightPlane(FrutumPlanes[5]))
+	{
+		FrutumPlanes[5] = FPlane(0.0f, 0.0f, 0.0f, 0.0f);
+	}
+}
diff --git a/Engine/Plugins/FX/Niagara/Source/Niagara/Public/NiagaraCommon.h b/Engine/Plugins/FX/Niagara/Source/Niagara/Public/NiagaraCommon.h
--- a/Engine/Plugins/FX/Niagara/Source/Niagara/Public/NiagaraCommon.h
+++ b/Engine/Plugins/FX/Niagara/Source/Niagara/Public/NiagaraCommon.h
@@ -1450,6 +1450,7 @@ struct FNiagaraScalabilityState
 
 	FNiagaraScalabilityState()
 		: Significance(1.0f)
+		, LastVisibleTime(0.0f)
 		, SystemDataIndex(INDEX_NONE)
 		, bCulled(0)
 		, bPreviousCulled(0)
@@ -1462,6 +1463,7 @@ struct FNiagaraScalabilityState
 
 	FNiagaraScalabilityState(float InSignificance, bool InCulled, bool InPreviousCulled)
 		: Significance(InSignificance)
+		, LastVisibleTime(0.0f)
 		, SystemDataIndex(INDEX_NONE)
 		, bCulled(InCulled)
 		, bPreviousCulled(InPreviousCulled)
@@ -1478,6 +1480,9 @@ struct FNiagaraScalabilityState
 	UPROPERTY(VisibleAnywhere, Category="Scalability")
 	float Significance;
 
+	UPROPERTY(VisibleAnywhere, Category = "Scalability")
+	float LastVisibleTime;
+
 	int16 SystemDataIndex;
 
 	UPROPERTY(VisibleAnywhere, Category = "Scalability")
diff --git a/Engine/Plugins/FX/Niagara/Source/Niagara/Public/NiagaraComponent.h b/Engine/Plugins/FX/Niagara/Source/Niagara/Public/NiagaraComponent.h
--- a/Engine/Plugins/FX/Niagara/Source/Niagara/Public/NiagaraComponent.h
+++ b/Engine/Plugins/FX/Niagara/Source/Niagara/Public/NiagaraComponent.h
@@ -186,7 +186,6 @@ public:
 
 	FORCEINLINE bool IsComplete() const { return SystemInstanceController ? SystemInstanceController->IsComplete() : true; }
 
-	FORCEINLINE float GetSafeTimeSinceRendered(float WorldTime)const;
 private:
 
 	//Internal versions that can be called from the scalability code.
@@ -885,12 +884,6 @@ private:
 	FDynamicData DynamicData;
 };
 
-extern float GLastRenderTimeSafetyBias;
-FORCEINLINE float UNiagaraComponent::GetSafeTimeSinceRendered(float WorldTime)const
-{
-	return FMath::Max(0.0f, WorldTime - GetLastRenderTime() - GLastRenderTimeSafetyBias);
-}
-
 FORCEINLINE void UNiagaraComponent::SetLODDistance(float InLODDistance, float InMaxLODDistance)
 {
 	if (!bEnablePreviewLODDistance && SystemInstanceController)
diff --git a/Engine/Plugins/FX/Niagara/Source/Niagara/Public/NiagaraWorldManager.h b/Engine/Plugins/FX/Niagara/Source/Niagara/Public/NiagaraWorldManager.h
--- a/Engine/Plugins/FX/Niagara/Source/Niagara/Public/NiagaraWorldManager.h
+++ b/Engine/Plugins/FX/Niagara/Source/Niagara/Public/NiagaraWorldManager.h
@@ -65,6 +65,16 @@ enum class NIAGARA_API ENiagaraScalabilityCullingMode : uint8
 	Disabled,
 };
 
+struct FNiagaraCachedViewInfo
+{
+	FMatrix ViewMat;
+	FMatrix ProjectionMat;
+	FMatrix ViewProjMat;
+	FMatrix ViewToWorld;
+	TArray<FPlane, TInlineAllocator<6>> FrutumPlanes;
+
+	void Init(const FWorldCachedViewInfo& WorldViewInfo);
+};
 /**
 * Manager class for any data relating to a particular world.
 */
@@ -140,8 +150,7 @@ public:
 		return static_cast<T&>(**ExistingValue);
 	}
 
-	NIAGARA_API bool CachedPlayerViewLocationsValid() const { return bCachedPlayerViewLocationsValid; }
-	NIAGARA_API TArrayView<const FVector> GetCachedPlayerViewLocations() const { check(bCachedPlayerViewLocationsValid); return MakeArrayView(CachedPlayerViewLocations); }
+	NIAGARA_API TArrayView<const FNiagaraCachedViewInfo> GetCachedViewInfo() const { return MakeArrayView(CachedViewInfo); }
 
 	UNiagaraComponentPool* GetComponentPool() { return ComponentPool; }
 
@@ -243,7 +252,7 @@ private:
 
 	void DistanceCull(UNiagaraEffectType* EffectType, const FNiagaraSystemScalabilitySettings& ScalabilitySettings, FVector Location, FNiagaraScalabilityState& OutState);
 	void DistanceCull(UNiagaraEffectType* EffectType, const FNiagaraSystemScalabilitySettings& ScalabilitySettings, UNiagaraComponent* Component, FNiagaraScalabilityState& OutState);
-	void VisibilityCull(UNiagaraEffectType* EffectType, const FNiagaraSystemScalabilitySettings& ScalabilitySettings, UNiagaraComponent* Component, FNiagaraScalabilityState& OutState);
+	void ViewBasedCulling(UNiagaraEffectType* EffectType, const FNiagaraSystemScalabilitySettings& ScalabilitySettings, FSphere BoundingSphere, float ComponentTimeSinceRendered, bool bIsPrecull, FNiagaraScalabilityState& OutState);
 	void InstanceCountCull(UNiagaraEffectType* EffectType, UNiagaraSystem* System, const FNiagaraSystemScalabilitySettings& ScalabilitySettings, FNiagaraScalabilityState& OutState);
 	void GlobalBudgetCull(const FNiagaraSystemScalabilitySettings& ScalabilitySettings, float WorstGlobalBudgetUse, FNiagaraScalabilityState& OutState);
 
@@ -281,8 +290,7 @@ private:
 
 	int32 CachedEffectsQuality;
 
-	bool bCachedPlayerViewLocationsValid = false;
-	TArray<FVector, TInlineAllocator<8> > CachedPlayerViewLocations;
+	TArray<FNiagaraCachedViewInfo, TInlineAllocator<8> > CachedViewInfo;
 
 	UNiagaraComponentPool* ComponentPool;
 	bool bPoolIsPrimed = false;
diff --git a/Engine/Plugins/FX/Niagara/Source/NiagaraCore/Public/NiagaraCustomVersion.h b/Engine/Plugins/FX/Niagara/Source/NiagaraCore/Public/NiagaraCustomVersion.h
--- a/Engine/Plugins/FX/Niagara/Source/NiagaraCore/Public/NiagaraCustomVersion.h
+++ b/Engine/Plugins/FX/Niagara/Source/NiagaraCore/Public/NiagaraCustomVersion.h
@@ -170,6 +170,8 @@ struct FNiagaraCustomVersion
 
 		StaticSwitchFunctionPinsUsePersistentGuids, // Update static switch pins to use the PersistentId from their script variable so that when they're renamed their values aren't lost when reallocating pins. 
 
+		VisibilityCullingImprovements, // Extended visibility culling options and moved properties into their own struct.
+
 		// DO NOT ADD A NEW VERSION UNLESS YOU HAVE TALKED TO THE NIAGARA LEAD. Mismanagement of these versions can lead to data loss if it is adjusted in multiple streams simultaneously.
 		// -----<new versions can be added above this line>  -------------------------------------------------
 		VersionPlusOne,
diff --git a/Engine/Plugins/FX/Niagara/Source/NiagaraEditor/Private/Customizations/NiagaraTypeCustomizations.cpp b/Engine/Plugins/FX/Niagara/Source/NiagaraEditor/Private/Customizations/NiagaraTypeCustomizations.cpp
--- a/Engine/Plugins/FX/Niagara/Source/NiagaraEditor/Private/Customizations/NiagaraTypeCustomizations.cpp
+++ b/Engine/Plugins/FX/Niagara/Source/NiagaraEditor/Private/Customizations/NiagaraTypeCustomizations.cpp
@@ -36,6 +36,7 @@
 #include "Widgets/Input/SComboButton.h"
 #include "Widgets/Input/SEditableTextBox.h"
 #include "Widgets/Text/STextBlock.h"
+#include "IDetailGroup.h"
 
 
 #define LOCTEXT_NAMESPACE "FNiagaraVariableAttributeBindingCustomization"
@@ -1922,4 +1923,55 @@ void FNiagaraVariableMetaDataCustomization::CustomizeChildren(TSharedRef<IProper
 	}
 }
 
+//////////////////////////////////////////////////////////////////////////
+
+void FNiagaraSystemScalabilityOverrideCustomization::CustomizeHeader(TSharedRef<IPropertyHandle> InPropertyHandle, FDetailWidgetRow& HeaderRow, IPropertyTypeCustomizationUtils& CustomizationUtils)
+{
+	HeaderRow.NameContent()
+	[
+		InPropertyHandle->CreatePropertyNameWidget()
+	];
+	HeaderRow.ValueContent()
+	[
+		InPropertyHandle->CreatePropertyValueWidget()
+	];
+}
+
+void FNiagaraSystemScalabilityOverrideCustomization::CustomizeChildren(TSharedRef<IPropertyHandle> PropertyHandle, IDetailChildrenBuilder& ChildBuilder, IPropertyTypeCustomizationUtils& CustomizationUtils)
+{
+	auto AddOverrideProperties = [&](FName MasterPropertyName, TArray<FName> OverrideProperties)
+	{
+		if (TSharedPtr<IPropertyHandle> MasterProperty = PropertyHandle->GetChildHandle(MasterPropertyName))
+		{
+			bool bMasterProp;
+			if (MasterProperty->GetValue(bMasterProp) == FPropertyAccess::Success && bMasterProp)
+			{
+				IDetailGroup& PropGroup = ChildBuilder.AddGroup(MasterPropertyName, MasterProperty->GetPropertyDisplayName());
+				PropGroup.HeaderProperty(MasterProperty.ToSharedRef());
+				for (FName OverridePropName : OverrideProperties)
+				{
+					if (TSharedPtr<IPropertyHandle> OverrideProp = PropertyHandle->GetChildHandle(OverridePropName))
+					{
+						PropGroup.AddPropertyRow(OverrideProp.ToSharedRef());
+					}
+				}
+			}
+			else
+			{
+				ChildBuilder.AddProperty(MasterProperty.ToSharedRef());
+			}
+		}
+	}; 
+	
+	TSharedPtr<IPropertyHandle> PlatformsProperty = PropertyHandle->GetChildHandle(TEXT("Platforms"));
+	ChildBuilder.AddProperty(PlatformsProperty.ToSharedRef());
+
+	AddOverrideProperties(TEXT("bOverrideDistanceSettings"), { TEXT("bCullByDistance"), TEXT("MaxDistance") });
+	AddOverrideProperties(TEXT("bOverrideInstanceCountSettings"), { TEXT("bCullMaxInstanceCount"), TEXT("MaxInstances") });
+	AddOverrideProperties(TEXT("bOverridePerSystemInstanceCountSettings"), { TEXT("bCullPerSystemMaxInstanceCount"), TEXT("MaxSystemInstances") });
+	AddOverrideProperties(TEXT("bOverrideVisibilitySettings"), { TEXT("VisibilityCulling") });
+	AddOverrideProperties(TEXT("bOverrideGlobalBudgetScalingSettings"), { TEXT("BudgetScaling") });
+	AddOverrideProperties(TEXT("bOverrideCullProxySettings"), { TEXT("CullProxyMode"), TEXT("MaxSystemProxies") });
+}
+
 #undef LOCTEXT_NAMESPACE
diff --git a/Engine/Plugins/FX/Niagara/Source/NiagaraEditor/Private/Customizations/NiagaraTypeCustomizations.h b/Engine/Plugins/FX/Niagara/Source/NiagaraEditor/Private/Customizations/NiagaraTypeCustomizations.h
--- a/Engine/Plugins/FX/Niagara/Source/NiagaraEditor/Private/Customizations/NiagaraTypeCustomizations.h
+++ b/Engine/Plugins/FX/Niagara/Source/NiagaraEditor/Private/Customizations/NiagaraTypeCustomizations.h
@@ -355,3 +355,17 @@ public:
 
 	virtual void CustomizeChildren(TSharedRef<IPropertyHandle> PropertyHandle, IDetailChildrenBuilder& ChildBuilder, IPropertyTypeCustomizationUtils& CustomizationUtils);
 };
+
+//** Properties customization for FNiagaraSystemScalabilityOverride to hide fields that are not relevant when editing parameter definitions. */
+class FNiagaraSystemScalabilityOverrideCustomization : public IPropertyTypeCustomization
+{
+public:
+	static TSharedRef<IPropertyTypeCustomization> MakeInstance()
+	{
+		return MakeShared<FNiagaraSystemScalabilityOverrideCustomization>();
+	}
+
+	virtual void CustomizeHeader(TSharedRef<IPropertyHandle> PropertyHandle, FDetailWidgetRow& HeaderRow, IPropertyTypeCustomizationUtils& CustomizationUtils);
+
+	virtual void CustomizeChildren(TSharedRef<IPropertyHandle> PropertyHandle, IDetailChildrenBuilder& ChildBuilder, IPropertyTypeCustomizationUtils& CustomizationUtils);
+};
diff --git a/Engine/Plugins/FX/Niagara/Source/NiagaraEditor/Private/NiagaraEditorModule.cpp b/Engine/Plugins/FX/Niagara/Source/NiagaraEditor/Private/NiagaraEditorModule.cpp
--- a/Engine/Plugins/FX/Niagara/Source/NiagaraEditor/Private/NiagaraEditorModule.cpp
+++ b/Engine/Plugins/FX/Niagara/Source/NiagaraEditor/Private/NiagaraEditorModule.cpp
@@ -914,6 +914,9 @@ void FNiagaraEditorModule::StartupModule()
 		FNiagaraVariableMetaData::StaticStruct()->GetFName(),
 		FOnGetPropertyTypeCustomizationInstance::CreateStatic(&FNiagaraVariableMetaDataCustomization::MakeInstance));
 
+	PropertyModule.RegisterCustomPropertyTypeLayout(
+		FNiagaraSystemScalabilityOverride::StaticStruct()->GetFName(),
+		FOnGetPropertyTypeCustomizationInstance::CreateStatic(&FNiagaraSystemScalabilityOverrideCustomization::MakeInstance));
 	PropertyModule.RegisterCustomClassLayout(
 		UNiagaraSimulationStageGeneric::StaticClass()->GetFName(),
 		FOnGetDetailCustomizationInstance::CreateStatic(&FNiagaraSimulationStageGenericCustomization::MakeInstance));
diff --git a/Engine/Plugins/FX/Niagara/Source/NiagaraEditor/Private/Widgets/Scalability/SNiagaraSystemResolvedScalabilitySettings.cpp b/Engine/Plugins/FX/Niagara/Source/NiagaraEditor/Private/Widgets/Scalability/SNiagaraSystemResolvedScalabilitySettings.cpp
--- a/Engine/Plugins/FX/Niagara/Source/NiagaraEditor/Private/Widgets/Scalability/SNiagaraSystemResolvedScalabilitySettings.cpp
+++ b/Engine/Plugins/FX/Niagara/Source/NiagaraEditor/Private/Widgets/Scalability/SNiagaraSystemResolvedScalabilitySettings.cpp
@@ -79,9 +79,9 @@ bool IsSystemScalabilityValueOverridden(FName PropertyName, const FNiagaraSystem
 	{
 		return ActiveOverrides.bOverridePerSystemInstanceCountSettings;
 	}
-	else if(PropertyName == GET_MEMBER_NAME_CHECKED(FNiagaraSystemScalabilitySettings, MaxTimeWithoutRender))
+	else if(PropertyName == GET_MEMBER_NAME_CHECKED(FNiagaraSystemScalabilitySettings, VisibilityCulling))
 	{
-		return ActiveOverrides.bOverrideTimeSinceRendererSettings;
+		return ActiveOverrides.bOverrideVisibilitySettings;
 	}
 	// cull proxy mode and max system proxies are controlled via the same override bool
 	else if(PropertyName == GET_MEMBER_NAME_CHECKED(FNiagaraSystemScalabilitySettings, CullProxyMode))
diff --git a/Engine/Plugins/Runtime/nDisplay/Source/DisplayCluster/Private/Game/EngineClasses/Basics/DisplayClusterViewportClient.cpp b/Engine/Plugins/Runtime/nDisplay/Source/DisplayCluster/Private/Game/EngineClasses/Basics/DisplayClusterViewportClient.cpp
--- a/Engine/Plugins/Runtime/nDisplay/Source/DisplayCluster/Private/Game/EngineClasses/Basics/DisplayClusterViewportClient.cpp
+++ b/Engine/Plugins/Runtime/nDisplay/Source/DisplayCluster/Private/Game/EngineClasses/Basics/DisplayClusterViewportClient.cpp
@@ -621,6 +621,13 @@ void UDisplayClusterViewportClient::Draw(FViewport* InViewport, FCanvas* SceneCa
 					const float StreamingScale = 1.f / FMath::Clamp<float>(View->LODDistanceFactor, .2f, 1.f);
 					IStreamingManager::Get().AddViewInformation(View->ViewMatrices.GetViewOrigin(), View->UnscaledViewRect.Width(), View->UnscaledViewRect.Width() * View->ViewMatrices.GetProjectionMatrix().M[0][0], StreamingScale);
 					MyWorld->ViewLocationsRenderedLastFrame.Add(View->ViewMatrices.GetViewOrigin());
+
+					FWorldCachedViewInfo& WorldViewInfo = World->CachedViewInfoRenderedLastFrame.AddDefaulted_GetRef();
+					WorldViewInfo.ViewMatrix = View->ViewMatrices.GetViewMatrix();
+					WorldViewInfo.ProjectionMatrix = View->ViewMatrices.GetProjectionMatrix();
+					WorldViewInfo.ViewProjectionMatrix = View->ViewMatrices.GetViewProjectionMatrix();
+					WorldViewInfo.ViewToWorld = View->ViewMatrices.GetInvViewMatrix();
+					World->LastRenderTime = World->GetTimeSeconds();
 				}
 			}
 
diff --git a/Engine/Source/Editor/UnrealEd/Private/EditorEngine.cpp b/Engine/Source/Editor/UnrealEd/Private/EditorEngine.cpp
--- a/Engine/Source/Editor/UnrealEd/Private/EditorEngine.cpp
+++ b/Engine/Source/Editor/UnrealEd/Private/EditorEngine.cpp
@@ -2261,6 +2261,19 @@ bool UEditorEngine::UpdateSingleViewportClient(FEditorViewportClient* InViewport
 			InViewportClient->bNeedsRedraw = false;
 			bUpdatedNonRealtimeViewport = true;
 		}
+		else if(UWorld* World = GetWorld())
+		{
+			// We're not rendering but calculate the view anyway so that we can cache the last "rendered" view info in the UWorld.
+			FSceneViewFamilyContext ViewFamily(FSceneViewFamily::ConstructionValues(InViewportClient->Viewport, InViewportClient->GetScene(), InViewportClient->EngineShowFlags));
+			FSceneView* View = InViewportClient->CalcSceneView(&ViewFamily);
+
+			FWorldCachedViewInfo& WorldViewInfo = World->CachedViewInfoRenderedLastFrame.AddDefaulted_GetRef();
+			WorldViewInfo.ViewMatrix = View->ViewMatrices.GetViewMatrix();
+			WorldViewInfo.ProjectionMatrix = View->ViewMatrices.GetProjectionMatrix();
+			WorldViewInfo.ViewProjectionMatrix = View->ViewMatrices.GetViewProjectionMatrix();
+			WorldViewInfo.ViewToWorld = View->ViewMatrices.GetInvViewMatrix();
+			World->LastRenderTime = World->GetTimeSeconds();
+		}
 
 		if (InViewportClient->bNeedsInvalidateHitProxy)
 		{
diff --git a/Engine/Source/Editor/UnrealEd/Private/EditorViewportClient.cpp b/Engine/Source/Editor/UnrealEd/Private/EditorViewportClient.cpp
--- a/Engine/Source/Editor/UnrealEd/Private/EditorViewportClient.cpp
+++ b/Engine/Source/Editor/UnrealEd/Private/EditorViewportClient.cpp
@@ -59,6 +59,8 @@
 #include "GPUSkinCacheVisualizationData.h"
 #include "UnrealWidget.h"
 #include "EdModeInteractiveToolsContext.h"
+#include "Engine/World.h"
+
 #include "CustomEditorStaticScreenPercentage.h"
 #include "IImageWrapperModule.h"
 
@@ -3972,6 +3974,16 @@ void FEditorViewportClient::Draw(FViewport* InViewport, FCanvas* Canvas)
 	    {
 		    View->CameraConstrainedViewRect = FIntRect(SafeFrame.Left, SafeFrame.Top, SafeFrame.Right, SafeFrame.Bottom);
 	    }
+
+		if (World)
+		{		
+			FWorldCachedViewInfo& WorldViewInfo = World->CachedViewInfoRenderedLastFrame.AddDefaulted_GetRef();
+			WorldViewInfo.ViewMatrix = View->ViewMatrices.GetViewMatrix();
+			WorldViewInfo.ProjectionMatrix = View->ViewMatrices.GetProjectionMatrix();
+			WorldViewInfo.ViewProjectionMatrix = View->ViewMatrices.GetViewProjectionMatrix();
+			WorldViewInfo.ViewToWorld = View->ViewMatrices.GetInvViewMatrix();
+			World->LastRenderTime = World->GetTimeSeconds();
+		}
  	}
 
 	{
diff --git a/Engine/Source/Runtime/Engine/Classes/Engine/World.h b/Engine/Source/Runtime/Engine/Classes/Engine/World.h
--- a/Engine/Source/Runtime/Engine/Classes/Engine/World.h
+++ b/Engine/Source/Runtime/Engine/Classes/Engine/World.h
@@ -993,6 +993,14 @@ public:
 	void Reevaluate(ULevelStreaming* StreamingLevel);
 };
 
+struct FWorldCachedViewInfo
+{
+	FMatrix ViewMatrix;
+	FMatrix ProjectionMatrix;
+	FMatrix ViewProjectionMatrix;
+	FMatrix ViewToWorld;
+};
+
 /**
  * Scope can be used to force load all external objects when loading a world. 
  */
@@ -1210,6 +1218,11 @@ public:
 	/** View locations rendered in the previous frame, if any. */
 	TArray<FVector>								ViewLocationsRenderedLastFrame;
 
+	/** Cached view information from the last rendered frame. */
+	TArray<FWorldCachedViewInfo>				CachedViewInfoRenderedLastFrame;
+	/** WorldTimeSeconds when this world was last rendered. */
+	float										LastRenderTime = 0.0f;
+
 	/** The current renderer feature level of this world */
 	TEnumAsByte<ERHIFeatureLevel::Type> FeatureLevel;
 
diff --git a/Engine/Source/Runtime/Engine/Private/GameViewportClient.cpp b/Engine/Source/Runtime/Engine/Private/GameViewportClient.cpp
--- a/Engine/Source/Runtime/Engine/Private/GameViewportClient.cpp
+++ b/Engine/Source/Runtime/Engine/Private/GameViewportClient.cpp
@@ -1568,6 +1568,12 @@ void UGameViewportClient::Draw(FViewport* InViewport, FCanvas* SceneCanvas)
 					const float StreamingScale = 1.f / FMath::Clamp<float>(View->LODDistanceFactor, .2f, 1.f);
 					IStreamingManager::Get().AddViewInformation(View->ViewMatrices.GetViewOrigin(), View->UnscaledViewRect.Width(), View->UnscaledViewRect.Width() * View->ViewMatrices.GetProjectionMatrix().M[0][0], StreamingScale);
 					MyWorld->ViewLocationsRenderedLastFrame.Add(View->ViewMatrices.GetViewOrigin());
+					FWorldCachedViewInfo& WorldViewInfo = World->CachedViewInfoRenderedLastFrame.AddDefaulted_GetRef();
+					WorldViewInfo.ViewMatrix = View->ViewMatrices.GetViewMatrix();
+					WorldViewInfo.ProjectionMatrix = View->ViewMatrices.GetProjectionMatrix();
+					WorldViewInfo.ViewProjectionMatrix = View->ViewMatrices.GetViewProjectionMatrix();
+					WorldViewInfo.ViewToWorld = View->ViewMatrices.GetInvViewMatrix();
+					World->LastRenderTime = World->GetTimeSeconds();
 				}
 			}
 		}
diff --git a/Engine/Source/Runtime/Engine/Private/LevelTick.cpp b/Engine/Source/Runtime/Engine/Private/LevelTick.cpp
--- a/Engine/Source/Runtime/Engine/Private/LevelTick.cpp
+++ b/Engine/Source/Runtime/Engine/Private/LevelTick.cpp
@@ -1791,6 +1791,7 @@ void UWorld::Tick( ELevelTick TickType, float DeltaSeconds )
 
 	// Dump the viewpoints with which we were rendered last frame. They will be updated when the world is next rendered.
 	ViewLocationsRenderedLastFrame.Reset();
+	CachedViewInfoRenderedLastFrame.Reset();
 
 	if (GEngine->XRSystem.IsValid())
 	{
