Compile fix

#rnx trivial

diff --git a/Engine/Plugins/FX/Niagara/Source/Niagara/Private/NiagaraDebugHud.cpp b/Engine/Plugins/FX/Niagara/Source/Niagara/Private/NiagaraDebugHud.cpp
--- a/Engine/Plugins/FX/Niagara/Source/Niagara/Private/NiagaraDebugHud.cpp
+++ b/Engine/Plugins/FX/Niagara/Source/Niagara/Private/NiagaraDebugHud.cpp
@@ -1080,7 +1080,7 @@ void FNiagaraDebugHud::Draw(FNiagaraWorldManager* WorldManager, UCanvas* Canvas,
 
 void FNiagaraDebugHud::AddLine2D(FVector2D Start, FVector2D End, FLinearColor Color, float Thickness, float Lifetime)
 {
-	if (NiagaraDebugLocal::Settings.bEnabled)
+	if (NiagaraDebugLocal::Settings.IsEnabled())
 	{
 		FDebugLine2D NewLine;
 		NewLine.Start = Start;
@@ -1094,7 +1094,7 @@ void FNiagaraDebugHud::AddLine2D(FVector2D Start, FVector2D End, FLinearColor Co
 
 void FNiagaraDebugHud::AddCircle2D(FVector2D Pos, float Rad, float Segments, FLinearColor Color, float Thickness, float Lifetime)
 {
-	if (NiagaraDebugLocal::Settings.bEnabled)
+	if (NiagaraDebugLocal::Settings.IsEnabled())
 	{
 		FDebugCircle2D NewCircle;
 		NewCircle.Pos = Pos;
@@ -1109,7 +1109,7 @@ void FNiagaraDebugHud::AddCircle2D(FVector2D Pos, float Rad, float Segments, FLi
 
 void FNiagaraDebugHud::AddBox2D(FVector2D Pos, FVector2D Extents, FLinearColor Color, float Thickness, float Lifetime)
 {
-	if (NiagaraDebugLocal::Settings.bEnabled)
+	if (NiagaraDebugLocal::Settings.IsEnabled())
 	{
 		FDebugBox2D NewBox;
 		NewBox.Pos = Pos;
