added custom viewinfo overrides to niagara manager

diff --git a/Engine/Plugins/FX/Niagara/Source/Niagara/Private/NiagaraWorldManager.cpp b/Engine/Plugins/FX/Niagara/Source/Niagara/Private/NiagaraWorldManager.cpp
--- a/Engine/Plugins/FX/Niagara/Source/Niagara/Private/NiagaraWorldManager.cpp
+++ b/Engine/Plugins/FX/Niagara/Source/Niagara/Private/NiagaraWorldManager.cpp
@@ -948,6 +948,35 @@ void FNiagaraWorldManager::MarkSimulationsForEndOfFrameWait(FNiagaraSystemSimula
 	}
 }
 
+#ifndef WITH_CUSTOM_NIAGARA_VIEWINFO
+bool FNiagaraWorldManager::PrepareCachedViewInfo(const APlayerController* PlayerController, FNiagaraCachedViewInfo& OutViewInfo)
+{
+	if (PlayerController && PlayerController->PlayerCameraManager)
+	{
+		FVector POVLoc;
+		FRotator POVRotation;
+		PlayerController->GetPlayerViewPoint(POVLoc, POVRotation);
+		FRotationTranslationMatrix ViewToWorld(POVRotation, POVLoc);
+		FMatrix ViewMat = ViewToWorld.InverseFast() * FMatrix(
+			FPlane(0, 0, 1, 0),
+			FPlane(1, 0, 0, 0),
+			FPlane(0, 1, 0, 0),
+			FPlane(0, 0, 0, 1));
+
+		FWorldCachedViewInfo WorldViewInfo;
+		WorldViewInfo.ViewMatrix = ViewMat;
+		WorldViewInfo.ProjectionMatrix = PlayerController->PlayerCameraManager->GetCameraCacheView().CalculateProjectionMatrix();
+		WorldViewInfo.ViewProjectionMatrix = ViewMat * WorldViewInfo.ProjectionMatrix;
+		WorldViewInfo.ViewToWorld = ViewToWorld;
+
+		OutViewInfo.Init(WorldViewInfo);
+		return true;
+	}
+
+	return false;
+}
+#endif // WITH_CUSTOM_NIAGARA_VIEWINFO
+
 void FNiagaraWorldManager::Tick(ETickingGroup TickGroup, float DeltaSeconds, ELevelTick TickType, ENamedThreads::Type CurrentThread, const FGraphEventRef& MyCompletionGraphEvent)
 {
 	check(TickGroup >= NiagaraFirstTickGroup && TickGroup <= NiagaraLastTickGroup);
@@ -1038,23 +1067,13 @@ void FNiagaraWorldManager::Tick(ETickingGroup TickGroup, float DeltaSeconds, ELe
 				APlayerController* PlayerController = Iterator->Get();
 				if (PlayerController && PlayerController->IsLocalPlayerController())
 				{
-					FVector POVLoc;
-					FRotator POVRotation;
-					PlayerController->GetPlayerViewPoint(POVLoc, POVRotation);
-					FRotationTranslationMatrix ViewToWorld(POVRotation, POVLoc);
-					FMatrix ViewMat = ViewToWorld.InverseFast() * FMatrix(
-						FPlane(0, 0, 1, 0),
-						FPlane(1, 0, 0, 0),
-						FPlane(0, 1, 0, 0),
-						FPlane(0, 0, 0, 1));
-
 					FNiagaraCachedViewInfo& ViewInfo = CachedViewInfo.AddDefaulted_GetRef();
-					FWorldCachedViewInfo WorldViewInfo;
-					WorldViewInfo.ViewMatrix = ViewMat;
-					WorldViewInfo.ProjectionMatrix = PlayerController->PlayerCameraManager->GetCameraCacheView().CalculateProjectionMatrix();
-					WorldViewInfo.ViewProjectionMatrix = ViewMat * WorldViewInfo.ProjectionMatrix;
-					WorldViewInfo.ViewToWorld = ViewToWorld;
-					ViewInfo.Init(WorldViewInfo);
+					
+					const bool bIsValid = PrepareCachedViewInfo(PlayerController, ViewInfo);
+					if (!bIsValid)
+					{
+						CachedViewInfo.RemoveAt(CachedViewInfo.Num() - 1);
+					}
 				}
 			}
 		}		
diff --git a/Engine/Plugins/FX/Niagara/Source/Niagara/Public/NiagaraWorldManager.h b/Engine/Plugins/FX/Niagara/Source/Niagara/Public/NiagaraWorldManager.h
--- a/Engine/Plugins/FX/Niagara/Source/Niagara/Public/NiagaraWorldManager.h
+++ b/Engine/Plugins/FX/Niagara/Source/Niagara/Public/NiagaraWorldManager.h
@@ -244,7 +244,10 @@ private:
 
 	// Callback to handle any pre GC processing needed.
 	static void OnPreGarbageCollectBeginDestroy();
-		
+	
+	// Calculates ViewInfo data for player
+	static bool PrepareCachedViewInfo(const APlayerController* PlayerController, FNiagaraCachedViewInfo& OutViewInfo);
+
 	// Gamethread callback to cleanup references to the given ComputeDispatchInterface before it gets deleted on the renderthread.
 	void OnComputeDispatchInterfaceDestroyed_Internal(class FNiagaraGpuComputeDispatchInterface* InComputeDispatchInterface);
 
