- Ensure Collection is post loaded for NPCI

#rnx

diff --git a/Engine/Plugins/FX/Niagara/Source/Niagara/Private/NiagaraParameterCollection.cpp b/Engine/Plugins/FX/Niagara/Source/Niagara/Private/NiagaraParameterCollection.cpp
--- a/Engine/Plugins/FX/Niagara/Source/Niagara/Private/NiagaraParameterCollection.cpp
+++ b/Engine/Plugins/FX/Niagara/Source/Niagara/Private/NiagaraParameterCollection.cpp
@@ -30,6 +30,14 @@ void UNiagaraParameterCollectionInstance::PostLoad()
 
 	ParameterStorage.PostLoad();
 
+	// Before calling SyncWithCollection we must ensure the collections parameter store is PostLoaded
+	// otherwise the parameters may not be sorted correctly.  If we are the default instance of the
+	// Collection then we don't need to do this.
+	if (Collection && (Collection->GetDefaultInstance() != this))
+	{
+		Collection->ConditionalPostLoad();
+	}
+
 	//Ensure we're synced up with our collection. TODO: Do conditionally via a version number/guid?
 	SyncWithCollection();
 }
