- Make check give more information about property name / type

#rnx

diff --git a/Engine/Plugins/FX/Niagara/Source/NiagaraEditor/Private/EdGraphSchema_Niagara.cpp b/Engine/Plugins/FX/Niagara/Source/NiagaraEditor/Private/EdGraphSchema_Niagara.cpp
--- a/Engine/Plugins/FX/Niagara/Source/NiagaraEditor/Private/EdGraphSchema_Niagara.cpp
+++ b/Engine/Plugins/FX/Niagara/Source/NiagaraEditor/Private/EdGraphSchema_Niagara.cpp
@@ -1916,7 +1916,7 @@ FNiagaraTypeDefinition UEdGraphSchema_Niagara::GetTypeDefForProperty(const FProp
 		return FNiagaraTypeDefinition::GetHalfDef();
 	}
 
-	check(0);
+	checkf(false, TEXT("Could not find type for Property(%s) Type(%s)"), *Property->GetName(), *Property->GetClass()->GetName());
 	return FNiagaraTypeDefinition::GetFloatDef();//Some invalid type?
 }
 
