Fixed properties not showing for CustomHlsl node when first input is a data interface (based on contribution from amuTBKT)

#9014

diff --git a/Engine/Plugins/FX/Niagara/Source/NiagaraEditor/Private/Customizations/NiagaraFunctionCallNodeDetails.cpp b/Engine/Plugins/FX/Niagara/Source/NiagaraEditor/Private/Customizations/NiagaraFunctionCallNodeDetails.cpp
--- a/Engine/Plugins/FX/Niagara/Source/NiagaraEditor/Private/Customizations/NiagaraFunctionCallNodeDetails.cpp
+++ b/Engine/Plugins/FX/Niagara/Source/NiagaraEditor/Private/Customizations/NiagaraFunctionCallNodeDetails.cpp
@@ -10,6 +10,7 @@
 #include "ScopedTransaction.h"
 #include "Framework/MultiBox/MultiBoxBuilder.h"
 #include "EdGraphSchema_Niagara.h"
+#include "NiagaraNodeCustomHlsl.h"
 #include "Layout/Margin.h"
 #include "Widgets/DeclarativeSyntaxSupport.h"
 #include "Widgets/Input/SCheckBox.h"
@@ -43,7 +44,8 @@ void FNiagaraFunctionCallNodeDetails::CustomizeDetails(IDetailLayoutBuilder& Det
 	FPinCollectorArray InputPins;
 	Node->GetInputPins(InputPins);
 	bool IsDIFunction = Node->FunctionScript == nullptr && InputPins.Num() > 0 && UNiagaraDataInterface::IsDataInterfaceType(NiagaraSchema->PinToTypeDefinition(InputPins[0]));
-	if (IsDIFunction)
+	bool bIsCustomHlslNode = Node->IsA(UNiagaraNodeCustomHlsl::StaticClass());
+	if (IsDIFunction || bIsCustomHlslNode)
 	{
 		DetailBuilder.EditCategory(FunctionCategoryName).SetCategoryVisibility(false);
 		DetailBuilder.EditCategory(VersionCategoryName).SetCategoryVisibility(false);
diff --git a/Engine/Plugins/FX/Niagara/Source/NiagaraEditor/Public/NiagaraNodeCustomHlsl.h b/Engine/Plugins/FX/Niagara/Source/NiagaraEditor/Public/NiagaraNodeCustomHlsl.h
--- a/Engine/Plugins/FX/Niagara/Source/NiagaraEditor/Public/NiagaraNodeCustomHlsl.h
+++ b/Engine/Plugins/FX/Niagara/Source/NiagaraEditor/Public/NiagaraNodeCustomHlsl.h
@@ -36,9 +36,8 @@ public:
 
 	// Replace items in the tokens array if they start with the src string or optionally src string and a namespace delimiter
 	static uint32 ReplaceExactMatchTokens(TArray<FString>& Tokens, const FString& SrcString, const FString& ReplaceString, bool bAllowNamespaceSeparation);
-	static FNiagaraVariable StripVariableToBaseType(const FNiagaraVariable& InVar);
-	virtual bool AllowNiagaraTypeForAddPin(const FNiagaraTypeDefinition& InType) const;
-	virtual bool AllowNiagaraTypeForAddPin(const FNiagaraTypeDefinition& InType, EEdGraphPinDirection InDirection) const;
+	virtual bool AllowNiagaraTypeForAddPin(const FNiagaraTypeDefinition& InType) const override;
+	virtual bool AllowNiagaraTypeForAddPin(const FNiagaraTypeDefinition& InType, EEdGraphPinDirection InDirection) const override;
 
 	virtual bool ReferencesVariable(const FNiagaraVariableBase& InVar) const;
 
@@ -82,6 +81,6 @@ protected:
 	UEdGraphPin* PinPendingRename;
 
 private:
-	UPROPERTY(EditAnywhere, Category = "Function", meta = (MultiLine = true))
+	UPROPERTY(EditAnywhere, Category = "HLSL", meta = (MultiLine = true))
 	FString CustomHlsl;
 };
