Niagara cook optimizations

Removes 3� min of FNiagaraCompileHash::ToString work from a 2h cook of an internal project

diff --git a/Engine/Plugins/FX/Niagara/Source/Niagara/Private/NiagaraScript.cpp b/Engine/Plugins/FX/Niagara/Source/Niagara/Private/NiagaraScript.cpp
--- a/Engine/Plugins/FX/Niagara/Source/Niagara/Private/NiagaraScript.cpp
+++ b/Engine/Plugins/FX/Niagara/Source/Niagara/Private/NiagaraScript.cpp
@@ -1208,7 +1208,7 @@ void UNiagaraScript::ComputeVMCompilationId(FNiagaraVMExecutableDataId& Id, FGui
 
 	// normalize the order
 	Id.AdditionalDefines.Sort();
-	Id.ReferencedCompileHashes.Sort([](const FNiagaraCompileHash& Lhs, const FNiagaraCompileHash& Rhs) { return Lhs.ToString() < Rhs.ToString(); });
+	Id.ReferencedCompileHashes.Sort();
 	Id.AdditionalVariables.Sort([](const FNiagaraVariableBase& Lhs, const FNiagaraVariableBase& Rhs) { return Lhs.GetName().LexicalLess(Rhs.GetName()); });
 
 	LastGeneratedVMId = Id;
diff --git a/Engine/Plugins/FX/Niagara/Source/NiagaraCore/Private/NiagaraCompileHash.cpp b/Engine/Plugins/FX/Niagara/Source/NiagaraCore/Private/NiagaraCompileHash.cpp
--- a/Engine/Plugins/FX/Niagara/Source/NiagaraCore/Private/NiagaraCompileHash.cpp
+++ b/Engine/Plugins/FX/Niagara/Source/NiagaraCore/Private/NiagaraCompileHash.cpp
@@ -3,17 +3,6 @@
 #include "NiagaraCompileHash.h"
 #include "Misc/SecureHash.h"
 
-const uint32 FNiagaraCompileHash::HashSize = 20;
-
-bool FNiagaraCompileHash::operator==(const FNiagaraCompileHash& Other) const
-{
-	return DataHash == Other.DataHash;
-}
-
-bool FNiagaraCompileHash::operator!=(const FNiagaraCompileHash& Other) const
-{
-	return DataHash != Other.DataHash;
-}
 
 bool FNiagaraCompileHash::operator==(const FSHAHash& Other) const
 {
@@ -52,21 +41,15 @@ const uint8* FNiagaraCompileHash::GetData() const
 	return DataHash.GetData();
 }
 
-FString FNiagaraCompileHash::ToString() const
+void FNiagaraCompileHash::AppendString(FString& Out) const
 {
 	if (DataHash.Num() == HashSize)
 	{
-		FString DataHashString;
-		DataHashString.Reserve(40);
-		for (int i = 0; i < HashSize; i++)
-		{
-			DataHashString.Append(FString::Printf(TEXT("%02x"), DataHash[i]));
-		}
-		return DataHashString;
+		BytesToHexLower(DataHash.GetData(), HashSize, Out);
 	}
 	else
 	{
-		return TEXT("Invalid");
+		Out.Append(TEXTVIEW("Invalid"));
 	}
 }
 
diff --git a/Engine/Plugins/FX/Niagara/Source/NiagaraCore/Public/NiagaraCompileHash.h b/Engine/Plugins/FX/Niagara/Source/NiagaraCore/Public/NiagaraCompileHash.h
--- a/Engine/Plugins/FX/Niagara/Source/NiagaraCore/Public/NiagaraCompileHash.h
+++ b/Engine/Plugins/FX/Niagara/Source/NiagaraCore/Public/NiagaraCompileHash.h
@@ -23,15 +23,10 @@ struct NIAGARACORE_API FNiagaraCompileHash
 	explicit FNiagaraCompileHash(const uint8* InDataHash, uint32 InCount) : DataHash(InDataHash, InCount)
 	{
 		checkf(InCount == HashSize, TEXT("Invalid hash data."));
-		/*.AddUninitialized(InCount);
-		for (uint32 i = 0; i < InCount; i++)
-		{
-			DataHash[i] = InDataHash[i];
-		}*/
 	}
 
-	bool operator==(const FNiagaraCompileHash& Other) const;
-	bool operator!=(const FNiagaraCompileHash& Other) const;
+	bool operator==(const FNiagaraCompileHash& Other) const { return DataHash == Other.DataHash; }
+	bool operator!=(const FNiagaraCompileHash& Other) const { return DataHash != Other.DataHash; }
 	bool operator==(const FSHAHash& Other) const;
 	inline bool operator!=(const FSHAHash& Other) const { return !operator==(Other); }
 
@@ -43,11 +38,28 @@ struct NIAGARACORE_API FNiagaraCompileHash
 
 	const uint8* GetData() const;
 
-	FString ToString() const;
+	void AppendString(FString& Out) const;
+	
+	FString ToString() const
+	{
+		FString Out;
+		AppendString(Out);
+		return Out;
+	}
 
 	NIAGARACORE_API friend FArchive& operator<<(FArchive& Ar, FNiagaraCompileHash& Id);
 
-	static const uint32 HashSize;
+	static constexpr uint32 HashSize = 20;
+
+	friend bool operator<(const FNiagaraCompileHash& Lhs, const FNiagaraCompileHash& Rhs)
+	{
+		if (int32 NumDiff = Lhs.DataHash.Num() - Rhs.DataHash.Num())
+		{
+			return NumDiff < 0;
+		}
+
+		return FMemory::Memcmp(Lhs.DataHash.GetData(), Rhs.DataHash.GetData(), Lhs.DataHash.Num()) < 0;
+	}
 
 private:
 	UPROPERTY()
