Make CullProxy on NiagaraComponent to ensure that we don't try to save any references to the dynamically created proxy component

diff --git a/Engine/Plugins/FX/Niagara/Source/Niagara/Public/NiagaraComponent.h b/Engine/Plugins/FX/Niagara/Source/Niagara/Public/NiagaraComponent.h
--- a/Engine/Plugins/FX/Niagara/Source/Niagara/Public/NiagaraComponent.h
+++ b/Engine/Plugins/FX/Niagara/Source/Niagara/Public/NiagaraComponent.h
@@ -777,7 +777,7 @@ private:
 
 	float CustomTimeDilation = 1.0f;
 
-	UPROPERTY()
+	UPROPERTY(transient)
 	TObjectPtr<class UNiagaraCullProxyComponent> CullProxy;
 
 	void CreateCullProxy(bool bForce = false);
