Fixed orphaned pin connections not appearing red

diff --git a/Engine/Plugins/FX/Niagara/Source/NiagaraEditor/Private/EdGraphSchema_Niagara.cpp b/Engine/Plugins/FX/Niagara/Source/NiagaraEditor/Private/EdGraphSchema_Niagara.cpp
--- a/Engine/Plugins/FX/Niagara/Source/NiagaraEditor/Private/EdGraphSchema_Niagara.cpp
+++ b/Engine/Plugins/FX/Niagara/Source/NiagaraEditor/Private/EdGraphSchema_Niagara.cpp
@@ -2339,6 +2339,14 @@ void FNiagaraConnectionDrawingPolicy::DetermineWiringStyle(UEdGraphPin* OutputPi
 				}
 			}
 		}
+
+		if(OutputPin && InputPin)
+		{
+			if(OutputPin->bOrphanedPin || InputPin->bOrphanedPin)
+			{
+				Params.WireColor = FLinearColor::Red;
+			}
+		}
 	}
 }
 
