- Fix for mesh renderer not showing the default material

#rnx

diff --git a/Engine/Plugins/FX/Niagara/Source/Niagara/Private/NiagaraRendererMeshes.cpp b/Engine/Plugins/FX/Niagara/Source/Niagara/Private/NiagaraRendererMeshes.cpp
--- a/Engine/Plugins/FX/Niagara/Source/Niagara/Private/NiagaraRendererMeshes.cpp
+++ b/Engine/Plugins/FX/Niagara/Source/Niagara/Private/NiagaraRendererMeshes.cpp
@@ -1262,14 +1262,8 @@ void FNiagaraRendererMeshes::GetDynamicMeshElements(const TArray<const FSceneVie
 					{
 						const FStaticMeshSection& Section = LODModel.Sections[SectionIndex];
 						const uint32 RemappedMaterialIndex = MeshData.MaterialRemapTable[Section.MaterialIndex];
-						if (!ParticleMeshRenderData.DynamicDataMesh->Materials.IsValidIndex(RemappedMaterialIndex))
-						{
-							// This should never occur. Otherwise, the section data changed since initialization
-							continue;
-						}
-
-						FMaterialRenderProxy* MaterialProxy = ParticleMeshRenderData.DynamicDataMesh->Materials[RemappedMaterialIndex];
-						if (Section.NumTriangles == 0 || MaterialProxy == NULL)
+						FMaterialRenderProxy* MaterialProxy = ParticleMeshRenderData.DynamicDataMesh->Materials.IsValidIndex(RemappedMaterialIndex) ? ParticleMeshRenderData.DynamicDataMesh->Materials[RemappedMaterialIndex] : UMaterial::GetDefaultMaterial(MD_Surface)->GetRenderProxy();
+						if (Section.NumTriangles == 0 || MaterialProxy == nullptr)
 						{
 							//@todo. This should never occur, but it does occasionally.
 							continue;
