- Fix for polluting the type registry with core LWC types

#rnx

diff --git a/Engine/Plugins/FX/Niagara/Source/Niagara/Private/NiagaraRendererComponents.cpp b/Engine/Plugins/FX/Niagara/Source/Niagara/Private/NiagaraRendererComponents.cpp
--- a/Engine/Plugins/FX/Niagara/Source/Niagara/Private/NiagaraRendererComponents.cpp
+++ b/Engine/Plugins/FX/Niagara/Source/Niagara/Private/NiagaraRendererComponents.cpp
@@ -53,73 +53,82 @@ void SetVariableByType(FNiagaraVariable& DataVariable, FNiagaraDataSet& Data, in
 	else if (VarType == FNiagaraTypeDefinition::GetPositionDef()) {	SetValueWithAccessor<FNiagaraPosition>(DataVariable, Data, ParticleIndex); }
 }
 
-void ConvertVariableToType(const FNiagaraVariable& SourceVariable, FNiagaraVariable& TargetVariable, const FNiagaraLWCConverter& LwcConverter)
+void ConvertVariableToType(const FNiagaraVariable& SourceVariable, const FNiagaraTypeDefinition& TargetType, TArray<uint8, TInlineAllocator<32>>& TargetData, const FNiagaraLWCConverter& LwcConverter)
 {
 	FNiagaraTypeDefinition SourceType = SourceVariable.GetType();
-	FNiagaraTypeDefinition TargetType = TargetVariable.GetType();
 
 	if (SourceType == FNiagaraTypeDefinition::GetVec3Def() && TargetType == UNiagaraComponentRendererProperties::GetFColorDef())
 	{
-		FVector3f Data = SourceVariable.GetValue<FVector3f>();
-		FColor ColorData((uint8)FMath::Clamp<int32>(FMath::TruncToInt(Data.X * 255.f), 0, 255),
+		const FVector3f Data = SourceVariable.GetValue<FVector3f>();
+		TargetData.AddUninitialized(sizeof(FColor));
+		*reinterpret_cast<FColor*>(TargetData.GetData()) = FColor(
+			(uint8)FMath::Clamp<int32>(FMath::TruncToInt(Data.X * 255.f), 0, 255),
 			(uint8)FMath::Clamp<int32>(FMath::TruncToInt(Data.Y * 255.f), 0, 255),
 			(uint8)FMath::Clamp<int32>(FMath::TruncToInt(Data.Z * 255.f), 0, 255));
-		TargetVariable.SetValue<FColor>(ColorData);
 	}
 	else if (SourceType == FNiagaraTypeDefinition::GetVec4Def() && TargetType == UNiagaraComponentRendererProperties::GetFColorDef())
 	{
-		FVector4f Data = SourceVariable.GetValue<FVector4f>();
-		FColor ColorData((uint8)FMath::Clamp<int32>(FMath::TruncToInt(Data.X * 255.f), 0, 255),
+		const FVector4f Data = SourceVariable.GetValue<FVector4f>();
+		TargetData.AddUninitialized(sizeof(FColor));
+		*reinterpret_cast<FColor*>(TargetData.GetData()) = FColor(
+			(uint8)FMath::Clamp<int32>(FMath::TruncToInt(Data.X * 255.f), 0, 255),
 			(uint8)FMath::Clamp<int32>(FMath::TruncToInt(Data.Y * 255.f), 0, 255),
 			(uint8)FMath::Clamp<int32>(FMath::TruncToInt(Data.Z * 255.f), 0, 255),
 			(uint8)FMath::Clamp<int32>(FMath::TruncToInt(Data.W * 255.f), 0, 255));
-		TargetVariable.SetValue<FColor>(ColorData);
 	}
 	else if (SourceType == FNiagaraTypeDefinition::GetColorDef() && TargetType == UNiagaraComponentRendererProperties::GetFColorDef())
 	{
-		FLinearColor Data = SourceVariable.GetValue<FLinearColor>();
-		TargetVariable.SetValue<FColor>(Data.QuantizeRound());
+		const FLinearColor Data = SourceVariable.GetValue<FLinearColor>();
+		TargetData.AddUninitialized(sizeof(FColor));
+		*reinterpret_cast<FColor*>(TargetData.GetData()) = Data.QuantizeRound();
 	}
 	else if (SourceType == FNiagaraTypeDefinition::GetVec3Def() && TargetType == UNiagaraComponentRendererProperties::GetFRotatorDef())
 	{
-		FVector3f Data = SourceVariable.GetValue<FVector3f>();
-		FRotator Rotator(Data.X, Data.Y, Data.Z);
-		TargetVariable.SetDoubleValue<FRotator>(Rotator);
+		const FVector3f Data = SourceVariable.GetValue<FVector3f>();
+		TargetData.AddUninitialized(sizeof(FRotator));
+		*reinterpret_cast<FRotator*>(TargetData.GetData()) = FRotator(Data.X, Data.Y, Data.Z);
 	}
 	else if (SourceType == FNiagaraTypeDefinition::GetVec2Def() && TargetType == UNiagaraComponentRendererProperties::GetFVector2DDef())
 	{
-		FVector2f Data = SourceVariable.GetValue<FVector2f>();
-		TargetVariable.SetDoubleValue<FVector2D>(FVector2D(Data));
+		const FVector2f Data = SourceVariable.GetValue<FVector2f>();
+		TargetData.AddUninitialized(sizeof(FVector2D));
+		*reinterpret_cast<FVector2D*>(TargetData.GetData()) = FVector2D(Data);
 	}	
 	else if (SourceType == FNiagaraTypeDefinition::GetVec3Def() && TargetType == UNiagaraComponentRendererProperties::GetFVectorDef())
 	{
-		FVector3f Data = SourceVariable.GetValue<FVector3f>();
-		TargetVariable.SetDoubleValue<FVector>(FVector(Data));
+		const FVector3f Data = SourceVariable.GetValue<FVector3f>();
+		TargetData.AddUninitialized(sizeof(FVector));
+		*reinterpret_cast<FVector*>(TargetData.GetData()) = FVector(Data);
 	}
 	else if (SourceType == FNiagaraTypeDefinition::GetVec4Def() && TargetType == UNiagaraComponentRendererProperties::GetFVector4Def())
 	{
-		FVector4f Data = SourceVariable.GetValue<FVector4f>();
-		TargetVariable.SetDoubleValue<FVector4>(FVector4(Data));
+		const FVector4f Data = SourceVariable.GetValue<FVector4f>();
+		TargetData.AddUninitialized(sizeof(FVector4));
+		*reinterpret_cast<FVector4*>(TargetData.GetData()) = FVector4(Data);
 	}
 	else if (SourceType == FNiagaraTypeDefinition::GetPositionDef() && TargetType == UNiagaraComponentRendererProperties::GetFVectorDef())
 	{
-		FVector Data = LwcConverter.ConvertSimulationPositionToWorld(SourceVariable.GetValue<FNiagaraPosition>());
-		TargetVariable.SetDoubleValue<FVector>(Data);
+		const FVector Data = LwcConverter.ConvertSimulationPositionToWorld(SourceVariable.GetValue<FNiagaraPosition>());
+		TargetData.AddUninitialized(sizeof(FVector));
+		*reinterpret_cast<FVector*>(TargetData.GetData()) = Data;
 	}
 	else if (SourceType == FNiagaraTypeDefinition::GetPositionDef() && TargetType == UNiagaraComponentRendererProperties::GetFVector3fDef())
 	{
-		FVector3f Data = (FVector3f)LwcConverter.ConvertSimulationPositionToWorld(SourceVariable.GetValue<FNiagaraPosition>());
-		TargetVariable.SetValue<FVector3f>(Data);
+		const FVector3f Data = (FVector3f)LwcConverter.ConvertSimulationPositionToWorld(SourceVariable.GetValue<FNiagaraPosition>());
+		TargetData.AddUninitialized(sizeof(FVector3f));
+		*reinterpret_cast<FVector3f*>(TargetData.GetData()) = Data;
 	}
 	else if (SourceType == FNiagaraTypeDefinition::GetQuatDef() && TargetType == UNiagaraComponentRendererProperties::GetFQuatDef())
 	{
-		FQuat4f Data = SourceVariable.GetValue<FQuat4f>();
-		TargetVariable.SetDoubleValue<FQuat>(FQuat(Data));
+		const FQuat4f Data = SourceVariable.GetValue<FQuat4f>();
+		TargetData.AddUninitialized(sizeof(FQuat));
+		*reinterpret_cast<FQuat*>(TargetData.GetData()) = FQuat(Data);
 	}
 	else if (SourceType == FNiagaraTypeDefinition::GetQuatDef() && TargetType == UNiagaraComponentRendererProperties::GetFRotatorDef())
 	{
-		FQuat4f Data = SourceVariable.GetValue<FQuat4f>();
-		TargetVariable.SetDoubleValue<FRotator>(FRotator(Data.Rotator()));
+		const FQuat4f Data = SourceVariable.GetValue<FQuat4f>();
+		TargetData.AddUninitialized(sizeof(FRotator));
+		*reinterpret_cast<FRotator*>(TargetData.GetData()) = FRotator(Data.Rotator());
 	}
 }
 
@@ -617,11 +626,7 @@ void FNiagaraRendererComponents::TickPropertyBindings(
 			continue;
 		}
 
-		FNiagaraVariable DataVariable = PropertyBinding.WritableValue;
-		const FNiagaraVariableBase& FoundVar = PropertyBinding.AttributeBinding.GetDataSetBindableVariable();
-		DataVariable.SetType(FoundVar.GetType());
-		DataVariable.SetName(FoundVar.GetName());
-		DataVariable.ClearData();
+		FNiagaraVariable DataVariable(PropertyBinding.AttributeBinding.GetDataSetBindableVariable());
 		if (!DataVariable.IsValid() || !Data.HasVariable(DataVariable))
 		{
 			continue;
@@ -679,12 +684,17 @@ void FNiagaraRendererComponents::TickPropertyBindings(
 		}
 
 		bool bForceStructConversion = PropertyType.GetScriptStruct() && !FNiagaraTypeHelper::IsNiagaraFriendlyTopLevelStruct(PropertyType.GetScriptStruct(), ENiagaraStructConversion::UserFacing);
-		
+
+		TArrayView<uint8> TargetData;
+		TArray<uint8, TInlineAllocator<32>> ConvertedData;
 		if (PropertyType.IsValid() && DataVariable.GetType() != PropertyType && (!PropertySetter->bIgnoreConversion || bForceStructConversion))
 		{
-			FNiagaraVariable TargetVariable(PropertyType, DataVariable.GetName());
-			ConvertVariableToType(DataVariable, TargetVariable, LwcConverter);
-			DataVariable = TargetVariable;
+			ConvertVariableToType(DataVariable, PropertyType, ConvertedData, LwcConverter);
+			TargetData = MakeArrayView(ConvertedData);
+		}
+		else
+		{
+			TargetData = MakeArrayView(DataVariable.GetData(), DataVariable.GetSizeInBytes());
 		}
 
 		if (!DataVariable.IsDataAllocated())
@@ -696,7 +706,7 @@ void FNiagaraRendererComponents::TickPropertyBindings(
 		if (PropertySetter->Function && PropertySetter->Function->NumParms >= 1)
 		{
 			// if we have a setter function we invoke it instead of setting the property directly, because then the object gets a chance to react to the new value
-			InvokeSetterFunction(Component, PropertySetter->Function, DataVariable.GetData(), DataVariable.GetSizeInBytes(), PropertyBinding.PropertySetterParameterDefaults);
+			InvokeSetterFunction(Component, PropertySetter->Function, TargetData.GetData(), TargetData.Num(), PropertyBinding.PropertySetterParameterDefaults);
 		}
 		else
 		{
@@ -711,7 +721,7 @@ void FNiagaraRendererComponents::TickPropertyBindings(
 				continue;
 			}
 			uint8* Dest = (uint8*)Component + PropertyAddress.GetProperty()->GetOffset_ForInternal();
-			DataVariable.CopyTo(Dest);
+			FMemory::Memcpy(Dest, TargetData.GetData(), TargetData.Num());
 		}
 	}
 }
diff --git a/Engine/Plugins/FX/Niagara/Source/Niagara/Public/NiagaraComponentRendererProperties.h b/Engine/Plugins/FX/Niagara/Source/Niagara/Public/NiagaraComponentRendererProperties.h
--- a/Engine/Plugins/FX/Niagara/Source/Niagara/Public/NiagaraComponentRendererProperties.h
+++ b/Engine/Plugins/FX/Niagara/Source/Niagara/Public/NiagaraComponentRendererProperties.h
@@ -35,9 +35,6 @@ struct FNiagaraComponentPropertyBinding
 	/** (Optional) If we have a setter with more than one parameter, this holds the default values of any optional function parameters */
 	UPROPERTY()
 	TMap<FString, FString> PropertySetterParameterDefaults;
-
-	UPROPERTY(Transient)
-	FNiagaraVariable WritableValue;
 };
 
 struct FNiagaraPropertySetter
diff --git a/Engine/Plugins/FX/Niagara/Source/Niagara/Public/NiagaraTypes.h b/Engine/Plugins/FX/Niagara/Source/Niagara/Public/NiagaraTypes.h
--- a/Engine/Plugins/FX/Niagara/Source/Niagara/Public/NiagaraTypes.h
+++ b/Engine/Plugins/FX/Niagara/Source/Niagara/Public/NiagaraTypes.h
@@ -1884,26 +1884,6 @@ struct FNiagaraVariable : public FNiagaraVariableBase
 		return Value;
 	}
 
-	template<typename T>
-	void SetDoubleValue(const T& Data)
-	{
-		static_assert(TIsUECoreVariant<T, double>::Value, "Float core variant, please use SetValue.");
-		check(sizeof(T) == TypeDefHandle->GetSize());
-		AllocateData();
-		FMemory::Memcpy(VarData.GetData(), &Data, VarData.Num());
-	}
-
-	template<typename T>
-	T GetDoubleValue() const
-	{
-		static_assert(TIsUECoreVariant<T, double>::Value, "Float core variant, please use GetValue.");
-		check(sizeof(T) == TypeDefHandle->GetSize());
-		check(IsDataAllocated());
-		T Value;
-		FMemory::Memcpy(&Value, GetData(), TypeDefHandle->GetSize());
-		return Value;
-	}
-
 	void SetData(const uint8* Data)
 	{
 		check(Data);
