Fix for crash in data driven CVars UI

Also fixed crash when removing CVar conditions.

#rnx

diff --git a/Engine/Plugins/FX/Niagara/Source/Niagara/Private/NiagaraPlatformSet.cpp b/Engine/Plugins/FX/Niagara/Source/Niagara/Private/NiagaraPlatformSet.cpp
--- a/Engine/Plugins/FX/Niagara/Source/Niagara/Private/NiagaraPlatformSet.cpp
+++ b/Engine/Plugins/FX/Niagara/Source/Niagara/Private/NiagaraPlatformSet.cpp
@@ -885,33 +885,39 @@ void FNiagaraPlatformSet::OnCVarChanged(IConsoleVariable* CVar)
 {
 	InvalidateCachedData();
 
+	FScopeLock Lock(&CachedCVarInfoCritSec);
+
 	//Iterate on our cached cvars as there'll be a v small subset of CVars Niagara cares about.
-	for (auto& Pair : CachedCVarInfo)
+	for (auto It = CachedCVarInfo.CreateIterator(); It; ++It)
 	{
-		FName CVarName  = Pair.Key;
-		FCachedCVarInfo& Info = Pair.Value;
+		FName CVarName = It->Key;
+		FCachedCVarInfo& Info = It->Value;
 		if (Info.CVar == CVar)
 		{
-			ChangedCVars.AddUnique(Pair.Key);
+			ChangedCVars.AddUnique(CVarName);
+			It.RemoveCurrent();
 		}
 	}
 }
 
 void FNiagaraPlatformSet::OnCVarUnregistered(IConsoleVariable* CVar)
 {
+	InvalidateCachedData();
+
 	//We also unregister the changed delegate when it's unregistered.
 	FScopeLock Lock(&CachedCVarInfoCritSec);
 
 	//Iterate on our cached cvars as there'll be a v small subset of CVars Niagara cares about.
-	for (auto& Pair : CachedCVarInfo)
+	for (auto It = CachedCVarInfo.CreateIterator(); It; ++It)
 	{
-		FName CVarName = Pair.Key;
-		FCachedCVarInfo& Info = Pair.Value;
+		FName CVarName = It->Key;
+		FCachedCVarInfo& Info = It->Value;
 		if (Info.CVar == CVar)
 		{
 			check(CVar == Info.CVar);
-			Info.CVar->OnChangedDelegate().Remove(Info.ChangedHandle);
-			OnCVarChanged(CVar);
+			CVar->OnChangedDelegate().Remove(Info.ChangedHandle);
+			ChangedCVars.AddUnique(CVarName);
+			It.RemoveCurrent();
 		}
 	}
 }
diff --git a/Engine/Plugins/FX/Niagara/Source/NiagaraEditor/Private/Customizations/NiagaraPlatformSetCustomization.cpp b/Engine/Plugins/FX/Niagara/Source/NiagaraEditor/Private/Customizations/NiagaraPlatformSetCustomization.cpp
--- a/Engine/Plugins/FX/Niagara/Source/NiagaraEditor/Private/Customizations/NiagaraPlatformSetCustomization.cpp
+++ b/Engine/Plugins/FX/Niagara/Source/NiagaraEditor/Private/Customizations/NiagaraPlatformSetCustomization.cpp
@@ -1790,7 +1790,11 @@ FNiagaraPlatformSetCVarCondition* FNiagaraPlatformSetCVarConditionCustomization:
 	TArray<UObject*> Objects;
 	PropertyHandle->GetOuterObjects(Objects);
 	int32 Index = 0;
-	return (FNiagaraPlatformSetCVarCondition*)PropertyHandle->GetValueBaseAddress((uint8*)Objects[Index]);
+	if (Objects.IsValidIndex(Index))
+	{
+		return (FNiagaraPlatformSetCVarCondition*)PropertyHandle->GetValueBaseAddress((uint8*)Objects[Index]);
+	}
+	return nullptr;
 }
 
 FNiagaraPlatformSet* FNiagaraPlatformSetCVarConditionCustomization::GetTargetPlatformSet()const
@@ -1799,8 +1803,13 @@ FNiagaraPlatformSet* FNiagaraPlatformSetCVarConditionCustomization::GetTargetPla
 	check(PropertyHandle && PropertyHandle->GetParentHandle());
 	TSharedPtr<IPropertyHandle> ParentHandle = PropertyHandle->GetParentHandle();
 	ParentHandle->GetOuterObjects(Objects);
+
 	int32 Index = 0;
-	return (FNiagaraPlatformSet*)ParentHandle->GetValueBaseAddress((uint8*)Objects[Index]);
+	if (Objects.IsValidIndex(Index))
+	{
+		return (FNiagaraPlatformSet*)ParentHandle->GetValueBaseAddress((uint8*)Objects[Index]);
+	}
+	return nullptr;
 }
 
-#undef LOCTEXT_NAMESPACE
\ No newline at end of file
+#undef LOCTEXT_NAMESPACE
