Static variables were depending on pin persistent guids for caching. This is insufficient as some pins are zeroed out for persistent guids. Picking the pin offset and parent guid in that case.

diff --git a/Engine/Plugins/FX/Niagara/Source/NiagaraEditor/Private/NiagaraCompiler.cpp b/Engine/Plugins/FX/Niagara/Source/NiagaraEditor/Private/NiagaraCompiler.cpp
--- a/Engine/Plugins/FX/Niagara/Source/NiagaraEditor/Private/NiagaraCompiler.cpp
+++ b/Engine/Plugins/FX/Niagara/Source/NiagaraEditor/Private/NiagaraCompiler.cpp
@@ -568,6 +568,7 @@ void FNiagaraCompileRequestData::FinishPrecompile(const TArray<FNiagaraVariable>
 
 				ensure(Builder.Histories.Num() <= 1);
 
+				int HistoryIdx = 0;
 				for (FNiagaraParameterMapHistory& History : Builder.Histories)
 				{
 					History.OriginatingScriptUsage = FoundOutputNode->GetUsage();
@@ -588,7 +589,16 @@ void FNiagaraCompileRequestData::FinishPrecompile(const TArray<FNiagaraVariable>
 							}
 						}
 					}
+
+					if (UNiagaraScript::LogCompileStaticVars > 0)
+					{
+						for (auto Iter : History.PinToConstantValues)
+						{
+							UE_LOG(LogNiagaraEditor, Log, TEXT("History [%d] Pin: %s Value: %s"), HistoryIdx , *Iter.Key.ToString(), *Iter.Value);
+						}
+					}
 					PinToConstantValues.Append(History.PinToConstantValues);
+					++HistoryIdx;
 				}
 
 				if (FoundOutputNode->GetUsage() == ENiagaraScriptUsage::ParticleSimulationStageScript)
@@ -611,6 +621,14 @@ void FNiagaraCompileRequestData::FinishPrecompile(const TArray<FNiagaraVariable>
 					}
 				}
 
+				if (UNiagaraScript::LogCompileStaticVars > 0)
+				{
+					for (auto Iter : PinToConstantValues)
+					{
+						UE_LOG(LogNiagaraEditor, Log, TEXT("Pin: %s Value: %s"), *Iter.Key.ToString(), *Iter.Value);
+					}
+				}
+
 				Builder.EndTranslation(TranslationName);
 
 				// Collect data interface information.
diff --git a/Engine/Plugins/FX/Niagara/Source/NiagaraEditor/Private/NiagaraHlslTranslator.cpp b/Engine/Plugins/FX/Niagara/Source/NiagaraEditor/Private/NiagaraHlslTranslator.cpp
--- a/Engine/Plugins/FX/Niagara/Source/NiagaraEditor/Private/NiagaraHlslTranslator.cpp
+++ b/Engine/Plugins/FX/Niagara/Source/NiagaraEditor/Private/NiagaraHlslTranslator.cpp
@@ -1998,6 +1998,16 @@ const FNiagaraTranslateResults &FHlslNiagaraTranslator::Translate(const FNiagara
 
 		}
 
+		if (CompileData->PinToConstantValues.Num() != 0)
+		{
+			Preamble.Appendf(TEXT("\n// Compile Data> PinToConstantValues Input: \n"));
+
+			for (auto Iter : CompileData->PinToConstantValues)
+			{
+				Preamble.Appendf(TEXT("//\tPin: %s Value: %s\n"), *Iter.Key.ToString(), *Iter.Value);
+			}
+		}
+
 		if (CompilationOutput.ScriptData.StaticVariablesWritten.Num() != 0)
 		{
 			Preamble.Appendf(TEXT("\n// Static Variables Written: \n"));
diff --git a/Engine/Plugins/FX/Niagara/Source/NiagaraEditor/Private/NiagaraParameterMapHistory.cpp b/Engine/Plugins/FX/Niagara/Source/NiagaraEditor/Private/NiagaraParameterMapHistory.cpp
--- a/Engine/Plugins/FX/Niagara/Source/NiagaraEditor/Private/NiagaraParameterMapHistory.cpp
+++ b/Engine/Plugins/FX/Niagara/Source/NiagaraEditor/Private/NiagaraParameterMapHistory.cpp
@@ -46,13 +46,26 @@ static FAutoConsoleVariableRef CVarNiagaraForcePrecompilerCullDataset(
 
 void FGraphTraversalHandle::Push(const UEdGraphNode* Node)
 {
+	ensure(Node->NodeGuid.IsValid());
 	Path.Push(Node->NodeGuid);
 	FriendlyPath.Push(Node->GetNodeTitle(ENodeTitleType::FullTitle).ToString());
 }
 
 void FGraphTraversalHandle::Push(const UEdGraphPin* Pin)
 {
-	Path.Push(Pin->PersistentGuid);
+	if (Pin->PersistentGuid.IsValid())
+	{
+		Path.Push(Pin->PersistentGuid);
+	}
+	else if (Pin->GetOwningNode())
+	{
+		FGuid NodeGuid = Pin->GetOwningNode()->NodeGuid;
+		int32 Index = Pin->GetOwningNode()->GetPinIndex((UEdGraphPin*)Pin);
+		ensure(NodeGuid.IsValid());
+		NodeGuid.A += Index; // This is a bit hacky, but the chance of collisions within the same graph is really low and we want a unique value for this graph.
+		Path.Push(NodeGuid);
+	}
+	
 	FriendlyPath.Push(Pin->GetOwningNode()->GetNodeTitle(ENodeTitleType::FullTitle).ToString() + TEXT("->") + Pin->PinName.ToString());
 }
 
@@ -173,11 +186,16 @@ void FNiagaraParameterMapHistory::RegisterConstantPin(const FGraphTraversalHandl
 		FGraphTraversalHandle Handle = InTraversalPath;
 		Handle.Push(InPin);
 		
+		FString* FoundValue = PinToConstantValues.Find(Handle);
+		if (FoundValue != nullptr)
+		{
+			ensure(*FoundValue == InValue);
+		}
 		PinToConstantValues.Add(Handle, InValue);
 
 		if (UNiagaraScript::LogCompileStaticVars > 0)
 		{
-			UE_LOG(LogNiagaraEditor, Log, TEXT("RegisterConstantPin \"%s\" Pin: \"%s\""), *InValue, *Handle.ToString());
+			UE_LOG(LogNiagaraEditor, Log, TEXT("RegisterConstantPin \"%s\" Pin: \"%s\""), *InValue, *Handle.ToString(true));
 		}
 	}
 }
diff --git a/Engine/Plugins/FX/Niagara/Source/NiagaraEditor/Public/NiagaraParameterMapHistory.h b/Engine/Plugins/FX/Niagara/Source/NiagaraEditor/Public/NiagaraParameterMapHistory.h
--- a/Engine/Plugins/FX/Niagara/Source/NiagaraEditor/Public/NiagaraParameterMapHistory.h
+++ b/Engine/Plugins/FX/Niagara/Source/NiagaraEditor/Public/NiagaraParameterMapHistory.h
@@ -87,11 +87,18 @@ public:
 		return Value;
 	}
 
-	FString ToString() const
+	FString ToString(bool bVerbose = false) const
 	{
 		FString Output;
 		for (const FString& Str : FriendlyPath)
 			Output += TEXT("/") + Str;
+
+		if (bVerbose)
+		{
+			Output += " GUID: ";
+			for (const FGuid& Guid : Path)
+				Output += TEXT("/") + Guid.ToString();
+		}
 		return Output;
 	}
 private:
diff --git a/Engine/Source/Runtime/Core/Private/UObject/UE5MainStreamObjectVersion.cpp b/Engine/Source/Runtime/Core/Private/UObject/UE5MainStreamObjectVersion.cpp
--- a/Engine/Source/Runtime/Core/Private/UObject/UE5MainStreamObjectVersion.cpp
+++ b/Engine/Source/Runtime/Core/Private/UObject/UE5MainStreamObjectVersion.cpp
@@ -14,7 +14,7 @@ TMap<FGuid, FGuid> FUE5MainStreamObjectVersion::GetSystemGuids()
 	SystemGuids.Add(DevGuids.MATERIALSHADERMAP_DERIVEDDATA_VER, FGuid("65C01D817C9A4EEFAE9E988D41A1F3DD"));
 	SystemGuids.Add(DevGuids.NANITE_DERIVEDDATA_VER, FGuid("7E26CA5B0F5F03C40A3FD7173D88B8FA"));
 	SystemGuids.Add(DevGuids.NIAGARASHADERMAP_DERIVEDDATA_VER, FGuid("8A37C45D24F2423CBE5F8F371DE33575"));
-	SystemGuids.Add(DevGuids.Niagara_LatestScriptCompileVersion, FGuid("2B4EE121451E44299EF2EFED80BDD0CD"));
+	SystemGuids.Add(DevGuids.Niagara_LatestScriptCompileVersion, FGuid("7AF02382DFED45A08DA36E13DEC058EC"));
 	SystemGuids.Add(DevGuids.SkeletalMeshDerivedDataVersion, FGuid("FB7B2FB546DE49B2A2F88B7F69763CEE"));
 	SystemGuids.Add(DevGuids.STATICMESH_DERIVEDDATA_VER, FGuid("98F7D79A5811013825E75BBCC41ED3E9"));
 	SystemGuids.Add(DevGuids.POSESEARCHDB_DERIVEDDATA_VER, FGuid("4E595C2AC5E947D6BA9ABC874353E5BC"));
