error for missing physics asset on selected static meshes

diff --git a/Engine/Plugins/FX/Niagara/Source/Niagara/Private/NiagaraDataInterfaceRigidMeshCollisionQuery.cpp b/Engine/Plugins/FX/Niagara/Source/Niagara/Private/NiagaraDataInterfaceRigidMeshCollisionQuery.cpp
--- a/Engine/Plugins/FX/Niagara/Source/Niagara/Private/NiagaraDataInterfaceRigidMeshCollisionQuery.cpp
+++ b/Engine/Plugins/FX/Niagara/Source/Niagara/Private/NiagaraDataInterfaceRigidMeshCollisionQuery.cpp
@@ -312,7 +312,7 @@ void CreateInternalArrays(TArray<TWeakObjectPtr<AActor>> Actors, FNDIRigidMeshCo
 
 					if (!FoundCollisionShapes)
 					{
-						UE_LOG(LogRigidMeshCollision, Warning, TEXT("No useable collision body setup found on mesh %s"), *Actor->GetName());
+						UE_LOG(LogRigidMeshCollision, Error, TEXT("No useable collision body setup found on mesh %s"), *Actor->GetName());
 
 					}
 				}
