- Refactor of the baker to allow different output types
- New output type added for volumes, can output atlas (XY), frames, or export frames
- Texture 2D output also supports atlas, frames and exporting frames

diff --git a/Engine/Plugins/FX/Niagara/Shaders/Private/NiagaraBatchedElements.usf b/Engine/Plugins/FX/Niagara/Shaders/Private/NiagaraBatchedElements.usf
--- a/Engine/Plugins/FX/Niagara/Shaders/Private/NiagaraBatchedElements.usf
+++ b/Engine/Plugins/FX/Niagara/Shaders/Private/NiagaraBatchedElements.usf
@@ -59,6 +59,59 @@ void MainPS(
 
 ////////////////////////////////////////////////////////////////////////////////////
 
+#if NIAGARA_VOLUME_ATTRIBUTE_PS
+
+Texture3D		InTexture;
+SamplerState	InTextureSampler;
+int				InNumAttributes;
+float4			InAttributeUVAndChannel[4];
+float2			InTileUV;
+float4x4		InColorWeights;
+float			InGamma;
+
+void MainPS(
+	in float2 InTextureCoordinate : TEXCOORD0,
+	in float4 InColor : TEXCOORD1,
+	in float4 InHitProxyId : TEXCOORD2,
+	in float4 InPosition : SV_POSITION,
+	out float4 OutColor : SV_Target0
+#if WRITE_TO_GBUFFER
+	,out float4 OutWorldNormal : SV_Target1 
+#endif
+	)
+{
+	float4 Sample = 0;
+	float3 TileUV = float3(InTextureCoordinate * InTileUV, 0);
+	for ( int i=0; i < InNumAttributes; ++i )
+	{
+		int Channel = asint(InAttributeUVAndChannel[i].w);
+		Sample[i] = Texture3DSample(InTexture, InTextureSampler, InAttributeUVAndChannel[i].xyz + TileUV)[Channel];
+	}
+
+	float4 Color;
+	Color.r = dot(Sample, InColorWeights[0]);
+	Color.g = dot(Sample, InColorWeights[1]);
+	Color.b = dot(Sample, InColorWeights[2]);
+	Color.a = dot(Sample, InColorWeights[3]);
+	Color *= InColor;
+
+	if( InGamma != 1.0 )
+	{
+		Color.rgb = ApplyGammaCorrection(saturate(Color.rgb), 2.2 / (1.0 / InGamma));
+	}
+
+	OutColor = RETURN_COLOR(Color);
+
+#if WRITE_TO_GBUFFER
+	// Set the G buffer bits that indicate that deferred lighting and image reflections are not enabled
+	OutWorldNormal = 0;
+#endif
+}
+
+#endif //NIAGARA_VOLUME_ATTRIBUTE_PS
+
+////////////////////////////////////////////////////////////////////////////////////
+
 #if NIAGARA_SIMPLE_PS
 
 Texture2D		InTexture;
diff --git a/Engine/Plugins/FX/Niagara/Source/Niagara/Classes/NiagaraBakerOutput.h b/Engine/Plugins/FX/Niagara/Source/Niagara/Classes/NiagaraBakerOutput.h
new file mode 100644
--- /dev/null
+++ b/Engine/Plugins/FX/Niagara/Source/Niagara/Classes/NiagaraBakerOutput.h
@@ -0,0 +1,62 @@
+// Copyright Epic Games, Inc. All Rights Reserved.
+
+#pragma once
+
+#include "CoreMinimal.h"
+#include "Misc/PathViews.h"
+#include "UObject/ObjectMacros.h"
+#include "UObject/Object.h"
+#include "UObject/SoftObjectPtr.h"
+
+#include "NiagaraBakerOutput.generated.h"
+
+USTRUCT()
+struct FNiagaraBakerTextureSource
+{
+	GENERATED_BODY()
+
+	UPROPERTY(EditAnywhere, Category = "Source")
+	FString DisplayString;
+
+	UPROPERTY(EditAnywhere, Category = "Source")
+	FName SourceName;
+};
+
+UCLASS(Abstract)
+class NIAGARA_API UNiagaraBakerOutput : public UObject
+{
+	GENERATED_BODY()
+public:
+	/** Optional output name, useful when you have multiple outputs */
+	UPROPERTY(EditAnywhere, Category = "Settings", meta = (DisplayPriority = "0"))
+	FString OutputName;
+
+	/** Used to determine if the bake is out of date or not. */
+	virtual bool Equals(const UNiagaraBakerOutput& Other) const PURE_VIRTUAL(UNiagaraBakerOutput::Equals, return OutputName == Other.OutputName;);
+
+	/** Make an OutputName used when creating a new output */
+	virtual FString MakeOutputName() const;
+
+	static FString SanitizeOutputName(FString Name);
+
+#if WITH_EDITOR
+	/** Find any warnings about the output. */
+	virtual void FindWarnings(TArray<FText>& OutWarnings) const;
+
+	FString GetAssetPath(FString PathFormat, int32 FrameIndex) const;
+	FString GetAssetFolder(FString PathFormat, int32 FrameIndex) const;
+	FString GetExportPath(FString PathFormat, int32 FrameIndex) const;
+	FString GetExportFolder(FString PathFormat, int32 FrameIndex) const;
+
+	template<class TObjectClass>
+	TObjectClass* GetAsset(FString PathFormat, int32 FrameIndex)
+	{
+		const FString AssetFullName = GetAssetPath(PathFormat, FrameIndex);
+		const FString AssetName = FString(FPathViews::GetCleanFilename(AssetFullName));
+		const FString AssetPath = AssetFullName + "." + AssetName;
+		return TSoftObjectPtr<TObjectClass>(AssetPath).LoadSynchronous();
+	}
+#endif
+
+	virtual void PostInitProperties() override;
+};
diff --git a/Engine/Plugins/FX/Niagara/Source/Niagara/Classes/NiagaraBakerOutputTexture2D.h b/Engine/Plugins/FX/Niagara/Source/Niagara/Classes/NiagaraBakerOutputTexture2D.h
new file mode 100644
--- /dev/null
+++ b/Engine/Plugins/FX/Niagara/Source/Niagara/Classes/NiagaraBakerOutputTexture2D.h
@@ -0,0 +1,73 @@
+// Copyright Epic Games, Inc. All Rights Reserved.
+
+#pragma once
+
+#include "NiagaraBakerOutput.h"
+#include "NiagaraBakerOutputTexture2D.generated.h"
+
+UCLASS(meta = (DisplayName = "Texture 2D Output"))
+class NIAGARA_API UNiagaraBakerOutputTexture2D : public UNiagaraBakerOutput
+{
+	GENERATED_BODY()
+
+public:
+	UNiagaraBakerOutputTexture2D(const FObjectInitializer& Init)
+		: bGenerateAtlas(true)
+	{
+	}
+
+	/** Source visualization we should capture, i.e. Scene Color, World Normal, etc */
+	UPROPERTY(EditAnywhere, Category = "Settings")
+	FNiagaraBakerTextureSource SourceBinding;
+
+	UPROPERTY(EditAnywhere, Category = "Settings", meta = (PinHiddenByDefault, InlineEditConditionToggle))
+	uint8 bGenerateAtlas : 1;	
+
+	UPROPERTY(EditAnywhere, Category = "Settings", meta = (PinHiddenByDefault, InlineEditConditionToggle))
+	uint8 bGenerateFrames : 1;	
+
+	UPROPERTY(EditAnywhere, Category = "Settings", meta = (PinHiddenByDefault, InlineEditConditionToggle))
+	uint8 bExportFrames : 1;
+
+	/** Size of each frame we generate. */
+	UPROPERTY(EditAnywhere, Category = "Settings")
+	FIntPoint FrameSize = FIntPoint(128, 128);
+
+	/** Size of the atlas texture when generating an atlas. */
+	UPROPERTY(EditAnywhere, Category = "Settings", meta = (EditCondition="bGenerateAtlas"))
+	FIntPoint AtlasTextureSize = FIntPoint(128 * 8, 128 * 8);
+
+	/** When either exporting or generating individual frames this is the size of the texture. */
+	FIntPoint SequenceFrameSize = FIntPoint(128, 128);
+
+	/**
+	When enabled a texture atlas is created
+	*/
+	UPROPERTY(EditAnywhere, Category = "Settings", meta = (EditCondition = "bGenerateAtlas"))
+	FString AtlasAssetPathFormat = TEXT("{AssetFolder}/{AssetName}_BakedAtlas_{OutputName}");
+
+	/**
+	When enabled each frame will create an asset.
+	*/
+	UPROPERTY(EditAnywhere, Category="Settings", meta=(EditCondition="bGenerateFrames"))
+	FString FramesAssetPathFormat = TEXT("{AssetFolder}/{AssetName}_BakedFrame_{OutputName}_{FrameIndex}");
+
+	/**
+	When enabled each frame will be exported to the output path using the format extension.
+	*/
+	UPROPERTY(EditAnywhere, Category="Settings", meta=(EditCondition="bExportFrames"))
+	FString FramesExportPathFormat = TEXT("{SavedDir}/NiagaraBaker/{AssetName}_{OutputName}/Frame_{FrameIndex}.png");
+
+	virtual bool Equals(const UNiagaraBakerOutput& Other) const override;
+
+	virtual void PostInitProperties() override;
+
+#if WITH_EDITOR
+	virtual FString MakeOutputName() const override;
+	virtual void FindWarnings(TArray<FText>& OutWarnings) const override;
+#endif
+
+#if WITH_EDITORONLY_DATA
+	virtual void PostEditChangeProperty(struct FPropertyChangedEvent& PropertyChangedEvent) override;
+#endif
+};
diff --git a/Engine/Plugins/FX/Niagara/Source/Niagara/Classes/NiagaraBakerOutputVolumeTexture.h b/Engine/Plugins/FX/Niagara/Source/Niagara/Classes/NiagaraBakerOutputVolumeTexture.h
new file mode 100644
--- /dev/null
+++ b/Engine/Plugins/FX/Niagara/Source/Niagara/Classes/NiagaraBakerOutputVolumeTexture.h
@@ -0,0 +1,58 @@
+// Copyright Epic Games, Inc. All Rights Reserved.
+
+#pragma once
+
+#include "NiagaraBakerOutput.h"
+#include "NiagaraBakerOutputVolumeTexture.generated.h"
+
+UCLASS(meta = (DisplayName = "Volume Texture Output"))
+class NIAGARA_API UNiagaraBakerOutputVolumeTexture : public UNiagaraBakerOutput
+{
+	GENERATED_BODY()
+
+public:
+	UNiagaraBakerOutputVolumeTexture(const FObjectInitializer& Init)
+		: bGenerateAtlas(true)
+	{
+	}
+
+	UPROPERTY(EditAnywhere, Category = "Settings")
+	FNiagaraBakerTextureSource SourceBinding;
+
+	UPROPERTY(EditAnywhere, Category = "Settings", meta = (PinHiddenByDefault, InlineEditConditionToggle))
+	uint8 bGenerateAtlas : 1;	
+
+	UPROPERTY(EditAnywhere, Category = "Settings", meta = (PinHiddenByDefault, InlineEditConditionToggle))
+	uint8 bGenerateFrames : 1;	
+
+	UPROPERTY(EditAnywhere, Category = "Settings", meta = (PinHiddenByDefault, InlineEditConditionToggle))
+	uint8 bExportFrames : 1;
+
+	/**
+	When enabled a volume atlas is created, the atlas is along X & Y not Z based on baker settings.
+	*/
+	UPROPERTY(EditAnywhere, Category = "Settings", meta = (EditCondition = "bGenerateAtlas"))
+	FString AtlasAssetPathFormat = TEXT("{AssetFolder}/{AssetName}_BakedAtlas_{OutputName}");
+
+	/**
+	When enabled each frame will create an asset.
+	*/
+	UPROPERTY(EditAnywhere, Category="Settings", meta=(EditCondition="bGenerateFrames"))
+	FString FramesAssetPathFormat = TEXT("{AssetFolder}/{AssetName}_BakedFrame_{OutputName}_{FrameIndex}");
+
+	/**
+	When enabled each frame will be exported to the output path using the format extension.
+	*/
+	UPROPERTY(EditAnywhere, Category="Settings", meta=(EditCondition="bExportFrames"))
+	FString FramesExportPathFormat = TEXT("{SavedDir}/NiagaraBaker/{AssetName}_{OutputName}/Frame_{FrameIndex}.png");
+
+	virtual bool Equals(const UNiagaraBakerOutput& Other) const override;
+
+#if WITH_EDITOR
+	FString MakeOutputName() const override;
+#endif
+
+#if WITH_EDITORONLY_DATA
+	virtual void PostEditChangeProperty(struct FPropertyChangedEvent& PropertyChangedEvent) override;
+#endif
+};
diff --git a/Engine/Plugins/FX/Niagara/Source/Niagara/Classes/NiagaraBakerSettings.h b/Engine/Plugins/FX/Niagara/Source/Niagara/Classes/NiagaraBakerSettings.h
--- a/Engine/Plugins/FX/Niagara/Source/Niagara/Classes/NiagaraBakerSettings.h
+++ b/Engine/Plugins/FX/Niagara/Source/Niagara/Classes/NiagaraBakerSettings.h
@@ -4,10 +4,13 @@
 
 #include "CoreMinimal.h"
 #include "Camera/CameraTypes.h"
+#include "Engine/TextureRenderTarget2D.h"
+#include "Misc/PathViews.h"
 #include "UObject/ObjectMacros.h"
 #include "UObject/Object.h"
-#include "Engine/TextureRenderTarget2D.h"
+#include "UObject/SoftObjectPtr.h"
 
+#include "NiagaraBakerOutput.h"
 #include "NiagaraBakerSettings.generated.h"
 
 UENUM()
@@ -24,23 +27,19 @@ enum class ENiagaraBakerViewMode
 };
 
 USTRUCT()
-struct FNiagaraBakerTextureSource
+struct FNiagaraBakerTextureSettings
 {
 	GENERATED_BODY()
 
-	UPROPERTY(EditAnywhere, Category = "Source")
-	FName SourceName;
-};
+	FNiagaraBakerTextureSettings()
+		: bUseFrameSize(false)
+	{
+	}
 
-USTRUCT()
-struct FNiagaraBakerTextureSettings
-{
-	GENERATED_BODY()
-	
 	/** Optional output name, if left empty a name will be auto generated using the index of the texture/ */
 	UPROPERTY(EditAnywhere, Category = "Texture")
 	FName OutputName;
-
+	
 	/** Source visualization we should capture, i.e. Scene Color, World Normal, etc */
 	UPROPERTY(EditAnywhere, Category = "Texture")
 	FNiagaraBakerTextureSource SourceBinding;
@@ -63,14 +62,15 @@ struct FNiagaraBakerTextureSettings
 	/** Final texture generated, an existing entry will be updated with new capture data. */
 	UPROPERTY(EditAnywhere, Category = "Texture")
 	TObjectPtr<UTexture2D> GeneratedTexture = nullptr;
+};
 
-	bool IsValidForBake() const { return FrameSize.X > 0 && FrameSize.Y > 0; }
-
-	bool Equals(const FNiagaraBakerTextureSettings& Other) const;
-
-	FNiagaraBakerTextureSettings()
-		: bUseFrameSize(false)
-	{}
+struct FNiagaraBakerOutputFrameIndices
+{
+	int		NumFrames = 1;
+	float	NormalizedTime = 0.0f;
+	int		FrameIndexA = 0;
+	int		FrameIndexB = 0;
+	float	Interp = 0.0f;
 };
 
 UCLASS()
@@ -79,14 +79,6 @@ class NIAGARA_API UNiagaraBakerSettings : public UObject
 	GENERATED_BODY()
 
 public:
-	struct FDisplayInfo
-	{
-		float NormalizedTime;
-		int FrameIndexA;
-		int FrameIndexB;
-		float Interp;
-	};
-
 	UNiagaraBakerSettings(const FObjectInitializer& Init);
 
 	/**
@@ -114,11 +106,11 @@ public:
 
 	/** Number of frames in each dimension. */
 	UPROPERTY(EditAnywhere, Category = "Settings")
-	FIntPoint FramesPerDimension = FIntPoint(8, 8);
+	FIntPoint FramesPerDimension = FIntPoint(8, 8);	//-TODO: Remove me...
 
-	/** List of output textures we will generated. */
+	/** Array of outputs for the baker to generate. */
 	UPROPERTY(EditAnywhere, Category = "Settings")
-	TArray<FNiagaraBakerTextureSettings> OutputTextures;
+	TArray<TObjectPtr<UNiagaraBakerOutput>> Outputs;
 
 	/** Current active viewport we will render from. */
 	UPROPERTY(EditAnywhere, Category = "Settings")
@@ -165,27 +157,29 @@ public:
 
 	bool Equals(const UNiagaraBakerSettings& Other) const;
 
-	int GetNumFrames() const { return FramesPerDimension.X * FramesPerDimension.Y; }
-
 	float GetSeekDelta() const { return 1.0f / float(FramesPerSecond); }
 
-	float GetAspectRatio(int32 iOutputTextureIndex) const;
-	FVector2D GetOrthoSize(int32 iOutputTextureIndex) const;
 	FVector GetCameraLocation() const;
 	FRotator GetCameraRotation() const;
 	FMatrix GetViewportMatrix() const;
 	FMatrix GetViewMatrix() const;
-	FMatrix GetProjectionMatrixForTexture(int32 iOutputTextureIndex) const;
+	FMatrix GetProjectionMatrix() const;
 
 	bool IsOrthographic() const { return CameraViewportMode != ENiagaraBakerViewMode::Perspective; }
 	bool IsPerspective() const { return CameraViewportMode == ENiagaraBakerViewMode::Perspective; }
 
-	// Get display info, the input time is expected to tbe relative, i.e. StartDuration is not taking into account
-	FDisplayInfo GetDisplayInfo(float Time, bool bLooping) const;
+	int GetOutputNumFrames(UNiagaraBakerOutput* BakerOutput) const;
+	FNiagaraBakerOutputFrameIndices GetOutputFrameIndices(UNiagaraBakerOutput* BakerOutput, float RelativeTime) const;
 
-	virtual void PostLoad() override;
+	int GetOutputNumFrames(int OutputIndex) const;
+	FNiagaraBakerOutputFrameIndices GetOutputFrameIndices(int OutputIndex, float RelativeTime) const;
 
+	virtual void PostLoad() override;
 #if WITH_EDITORONLY_DATA
 	virtual void PostEditChangeProperty(struct FPropertyChangedEvent& PropertyChangedEvent) override;
 #endif
+
+	// Deprecated properties
+	UPROPERTY()
+	TArray<FNiagaraBakerTextureSettings> OutputTextures_DEPRECATED;
 };
diff --git a/Engine/Plugins/FX/Niagara/Source/Niagara/Private/NiagaraBakerOutput.cpp b/Engine/Plugins/FX/Niagara/Source/Niagara/Private/NiagaraBakerOutput.cpp
new file mode 100644
--- /dev/null
+++ b/Engine/Plugins/FX/Niagara/Source/Niagara/Private/NiagaraBakerOutput.cpp
@@ -0,0 +1,92 @@
+// Copyright Epic Games, Inc. All Rights Reserved.
+
+#include "NiagaraBakerOutput.h"
+#include "NiagaraSystem.h"
+
+FString UNiagaraBakerOutput::MakeOutputName() const
+{
+	return SanitizeOutputName(GetFName().ToString());
+}
+
+FString UNiagaraBakerOutput::SanitizeOutputName(FString Name)
+{
+	FString OutName;
+	OutName.Reserve(Name.Len());
+	for ( TCHAR ch : Name )
+	{
+		switch (ch)
+		{
+			case TCHAR(' '):
+			case TCHAR(';'):
+			case TCHAR(':'):
+			case TCHAR(','):
+				ch = TCHAR('_');
+				break;
+			default:
+				break;
+		}
+		OutName.AppendChar(ch);
+	}
+
+	return OutName;
+}
+
+#if WITH_EDITOR
+void UNiagaraBakerOutput::FindWarnings(TArray<FText>& OutWarnings) const
+{
+}
+
+FString UNiagaraBakerOutput::GetAssetPath(FString PathFormat, int32 FrameIndex) const
+{
+	UNiagaraSystem* NiagaraSystem = GetTypedOuter<UNiagaraSystem>();
+	check(NiagaraSystem);
+
+	const TMap<FString, FStringFormatArg> PathFormatArgs =
+	{
+		{TEXT("AssetFolder"),	FString(FPathViews::GetPath(NiagaraSystem->GetPathName()))},
+		{TEXT("AssetName"),		NiagaraSystem->GetName()},
+		{TEXT("OutputName"),	SanitizeOutputName(OutputName)},
+		{TEXT("FrameIndex"),	FString::Printf(TEXT("%03d"), FrameIndex)},
+	};
+	FString AssetPath = FString::Format(*PathFormat, PathFormatArgs);
+	AssetPath.ReplaceInline(TEXT("//"), TEXT("/"));
+	return AssetPath;
+}
+
+FString UNiagaraBakerOutput::GetAssetFolder(FString PathFormat, int32 FrameIndex) const
+{
+	const FString AssetPath = GetAssetPath(PathFormat, FrameIndex);
+	return FString(FPathViews::GetPath(AssetPath));
+}
+
+FString UNiagaraBakerOutput::GetExportPath(FString PathFormat, int32 FrameIndex) const
+{
+	UNiagaraSystem* NiagaraSystem = GetTypedOuter<UNiagaraSystem>();
+	check(NiagaraSystem);
+
+	const TMap<FString, FStringFormatArg> PathFormatArgs =
+	{
+		{TEXT("SavedDir"),		FPaths::ProjectSavedDir()},
+		{TEXT("ProjectDir"),	FPaths::GetProjectFilePath()},
+		{TEXT("AssetName"),		NiagaraSystem->GetName()},
+		{TEXT("OutputName"),	SanitizeOutputName(OutputName)},
+		{TEXT("FrameIndex"),	FString::Printf(TEXT("%03d"), FrameIndex)},
+	};
+	FString ExportPath = FString::Format(*PathFormat, PathFormatArgs);
+	ExportPath.ReplaceInline(TEXT("//"), TEXT("/"));
+	return FPaths::ConvertRelativePathToFull(ExportPath);
+}
+
+FString UNiagaraBakerOutput::GetExportFolder(FString PathFormat, int32 FrameIndex) const
+{
+	const FString AssetPath = GetExportPath(PathFormat, FrameIndex);
+	return FString(FPathViews::GetPath(AssetPath));
+}
+#endif
+
+void UNiagaraBakerOutput::PostInitProperties()
+{
+	Super::PostInitProperties();
+
+	OutputName = MakeOutputName();
+}
diff --git a/Engine/Plugins/FX/Niagara/Source/Niagara/Private/NiagaraBakerOutputTexture2D.cpp b/Engine/Plugins/FX/Niagara/Source/Niagara/Private/NiagaraBakerOutputTexture2D.cpp
new file mode 100644
--- /dev/null
+++ b/Engine/Plugins/FX/Niagara/Source/Niagara/Private/NiagaraBakerOutputTexture2D.cpp
@@ -0,0 +1,95 @@
+// Copyright Epic Games, Inc. All Rights Reserved.
+
+#include "NiagaraBakerOutputTexture2D.h"
+#include "NiagaraBakerSettings.h"
+
+#define LOCTEXT_NAMESPACE "NiagaraBakerOutputTexture2D"
+
+bool UNiagaraBakerOutputTexture2D::Equals(const UNiagaraBakerOutput& OtherBase) const
+{
+	const UNiagaraBakerOutputTexture2D& Other = *CastChecked<UNiagaraBakerOutputTexture2D>(&OtherBase);
+	return
+		Super::Equals(Other) &&
+		SourceBinding.SourceName == Other.SourceBinding.SourceName &&
+		bGenerateAtlas == Other.bGenerateAtlas &&
+		bGenerateFrames == Other.bGenerateFrames &&
+		bExportFrames == Other.bExportFrames &&
+		FrameSize == Other.FrameSize &&
+		AtlasTextureSize == Other.AtlasTextureSize;
+}
+
+void UNiagaraBakerOutputTexture2D::PostInitProperties()
+{
+	Super::PostInitProperties();
+
+	if ( UNiagaraBakerSettings* BakerSettings = GetTypedOuter<UNiagaraBakerSettings>() )
+	{
+		AtlasTextureSize.X = FrameSize.X * BakerSettings->FramesPerDimension.X;
+		AtlasTextureSize.Y = FrameSize.Y * BakerSettings->FramesPerDimension.Y;
+	}
+}
+
+#if WITH_EDITOR
+FString UNiagaraBakerOutputTexture2D::MakeOutputName() const
+{
+	return FString::Printf(TEXT("Texture2D_%d"), GetFName().GetNumber());
+}
+
+void UNiagaraBakerOutputTexture2D::FindWarnings(TArray<FText>& OutWarnings) const
+{
+	if ( bGenerateAtlas )
+	{
+		if (const UNiagaraBakerSettings* BakerSettings = GetTypedOuter<UNiagaraBakerSettings>())
+		{
+			const FIntPoint ExpectedAtlasSize(FrameSize.X * BakerSettings->FramesPerDimension.X, FrameSize.Y * BakerSettings->FramesPerDimension.Y);
+			if ( (AtlasTextureSize.X != ExpectedAtlasSize.X) || (AtlasTextureSize.Y != ExpectedAtlasSize.Y) )
+			{
+				OutWarnings.Emplace(
+					FText::Format(LOCTEXT("TextureSizeWarningFormat", "Output '{0}' expect atlas size ({1} x {2}) does not match atlas size ({3} x {4}), this can result in jitter on the flipbook if not taken into account"),
+					FText::FromString(OutputName),
+					ExpectedAtlasSize.X, ExpectedAtlasSize.Y,
+					AtlasTextureSize.X, AtlasTextureSize.Y
+				));
+			}
+		}
+	}
+}
+#endif
+
+#if WITH_EDITORONLY_DATA
+void UNiagaraBakerOutputTexture2D::PostEditChangeProperty(struct FPropertyChangedEvent& PropertyChangedEvent)
+{
+	Super::PostEditChangeProperty(PropertyChangedEvent);
+
+	UNiagaraBakerSettings* Settings = GetTypedOuter<UNiagaraBakerSettings>();
+	if ( !ensure(Settings) )
+	{
+		return;
+	}
+
+	const FName PropertyName = PropertyChangedEvent.Property ? PropertyChangedEvent.Property->GetFName() : NAME_None;
+	const FName MemberPropertyName = PropertyChangedEvent.MemberProperty ? PropertyChangedEvent.MemberProperty->GetFName() : NAME_None;
+
+	if ( MemberPropertyName == GET_MEMBER_NAME_CHECKED(UNiagaraBakerOutputTexture2D, FrameSize) )
+	{
+		FrameSize.X = FMath::Max(FrameSize.X, 1);
+		FrameSize.Y = FMath::Max(FrameSize.Y, 1);
+		AtlasTextureSize.X = FrameSize.X * Settings->FramesPerDimension.X;
+		AtlasTextureSize.Y = FrameSize.Y * Settings->FramesPerDimension.Y;
+	}
+	else if (MemberPropertyName == GET_MEMBER_NAME_CHECKED(UNiagaraBakerOutputTexture2D, AtlasTextureSize))
+	{
+		AtlasTextureSize.X = FMath::Max(AtlasTextureSize.X, 1);
+		AtlasTextureSize.Y = FMath::Max(AtlasTextureSize.Y, 1);
+		FrameSize.X = FMath::DivideAndRoundDown(AtlasTextureSize.X, Settings->FramesPerDimension.X);
+		FrameSize.Y = FMath::DivideAndRoundDown(AtlasTextureSize.Y, Settings->FramesPerDimension.Y);
+	}
+	else if (PropertyName == GET_MEMBER_NAME_CHECKED(UNiagaraBakerSettings, FramesPerDimension))
+	{
+		FrameSize.X = FMath::DivideAndRoundDown(AtlasTextureSize.X, Settings->FramesPerDimension.X);
+		FrameSize.Y = FMath::DivideAndRoundDown(AtlasTextureSize.Y, Settings->FramesPerDimension.Y);
+	}
+}
+#endif
+
+#undef LOCTEXT_NAMESPACE
diff --git a/Engine/Plugins/FX/Niagara/Source/Niagara/Private/NiagaraBakerOutputVolumeTexture.cpp b/Engine/Plugins/FX/Niagara/Source/Niagara/Private/NiagaraBakerOutputVolumeTexture.cpp
new file mode 100644
--- /dev/null
+++ b/Engine/Plugins/FX/Niagara/Source/Niagara/Private/NiagaraBakerOutputVolumeTexture.cpp
@@ -0,0 +1,32 @@
+// Copyright Epic Games, Inc. All Rights Reserved.
+
+#include "NiagaraBakerOutputVolumeTexture.h"
+#include "NiagaraBakerSettings.h"
+
+bool UNiagaraBakerOutputVolumeTexture::Equals(const UNiagaraBakerOutput& OtherBase) const
+{
+	const UNiagaraBakerOutputVolumeTexture& Other = *CastChecked<UNiagaraBakerOutputVolumeTexture>(&OtherBase);
+	return
+		Super::Equals(Other) &&
+		SourceBinding.SourceName == Other.SourceBinding.SourceName;
+		// As these don't impact the generation process we don't consider for Equals
+		//bGenerateAssets == Other.bGenerateAssets &&
+		//AssetPathFormat == Other.AssetPathFormat &&
+		//bExportGenerated == Other.bExportGenerated &&
+		//ExportPathFormat == Other.ExportPathFormat;
+}
+
+#if WITH_EDITOR
+FString UNiagaraBakerOutputVolumeTexture::MakeOutputName() const
+{
+	return FString::Printf(TEXT("Volume_%d"), GetFName().GetNumber());
+}
+#endif
+
+#if WITH_EDITORONLY_DATA
+void UNiagaraBakerOutputVolumeTexture::PostEditChangeProperty(struct FPropertyChangedEvent& PropertyChangedEvent)
+{
+	Super::PostEditChangeProperty(PropertyChangedEvent);
+
+}
+#endif
diff --git a/Engine/Plugins/FX/Niagara/Source/Niagara/Private/NiagaraBakerSettings.cpp b/Engine/Plugins/FX/Niagara/Source/Niagara/Private/NiagaraBakerSettings.cpp
--- a/Engine/Plugins/FX/Niagara/Source/Niagara/Private/NiagaraBakerSettings.cpp
+++ b/Engine/Plugins/FX/Niagara/Source/Niagara/Private/NiagaraBakerSettings.cpp
@@ -1,25 +1,18 @@
 // Copyright Epic Games, Inc. All Rights Reserved.
 
 #include "NiagaraBakerSettings.h"
+#include "NiagaraBakerOutputTexture2D.h"
+#include "NiagaraSystem.h"
 
-#include "CoreMinimal.h"
 #include "Engine/Texture2D.h"
-
-bool FNiagaraBakerTextureSettings::Equals(const FNiagaraBakerTextureSettings& Other) const
-{
-	return
-		SourceBinding.SourceName == Other.SourceBinding.SourceName &&
-		bUseFrameSize == Other.bUseFrameSize &&
-		FrameSize == Other.FrameSize &&
-		TextureSize == Other.TextureSize;
-}
+#include "Misc/PathViews.h"
 
 UNiagaraBakerSettings::UNiagaraBakerSettings(const FObjectInitializer& Init)
 	: Super(Init)
 {
 	bPreviewLooping = true;
 	bRenderComponentOnly = true;
-	OutputTextures.AddDefaulted();
+	Outputs.Add(Init.CreateDefaultSubobject<UNiagaraBakerOutputTexture2D>(this, "DefaultOutput"));
 
 	for ( int i=0; i < (int)ENiagaraBakerViewMode::Num; ++i )
 	{
@@ -32,34 +25,47 @@ UNiagaraBakerSettings::UNiagaraBakerSettings(const FObjectInitializer& Init)
 
 bool UNiagaraBakerSettings::Equals(const UNiagaraBakerSettings& Other) const
 {
-	if (OutputTextures.Num() != Other.OutputTextures.Num())
-	{
-		return false;
-	}
-
-	for ( int i=0; i < OutputTextures.Num(); ++i )
-	{
-		if ( !OutputTextures[i].Equals(Other.OutputTextures[i]) )
+	auto OutputsEquals =
+		[&]() -> bool
 		{
-			return false;
-		}
-	}
+			if (Outputs.Num() != Other.Outputs.Num())
+			{
+				return false;
+			}
 
-	for ( int i=0; i < (int)ENiagaraBakerViewMode::Num; ++i )
-	{
-		if ( !CameraViewportLocation[i].Equals(Other.CameraViewportLocation[i]) ||
-			 !CameraViewportRotation[i].Equals(Other.CameraViewportRotation[i]) )
+			for (int i = 0; i < Outputs.Num(); ++i)
+			{
+				if ((Outputs[i]->GetClass() != Other.Outputs[i]->GetClass()) ||
+					(Outputs[i]->Equals(*Other.Outputs[i]) == false) )
+				{
+					return false;
+				}
+			}
+			return true;
+		};
+
+	auto CamerasEquals =
+		[&]() -> bool
 		{
-			return false;
-		}
-	}
+			for ( int i=0; i < (int)ENiagaraBakerViewMode::Num; ++i )
+			{
+				if ( !CameraViewportLocation[i].Equals(Other.CameraViewportLocation[i]) ||
+					 !CameraViewportRotation[i].Equals(Other.CameraViewportRotation[i]) )
+				{
+					return false;
+				}
+			}
+			return true;
+		};
 
 	return
+		OutputsEquals() &&
+		CamerasEquals() &&
 		FMath::IsNearlyEqual(StartSeconds, Other.StartSeconds) &&
 		FMath::IsNearlyEqual(DurationSeconds, Other.DurationSeconds) &&
 		FramesPerSecond == Other.FramesPerSecond &&
-		bPreviewLooping == Other.bPreviewLooping &&
 		FramesPerDimension == Other.FramesPerDimension &&
+		bPreviewLooping == Other.bPreviewLooping &&
 		CameraViewportMode == Other.CameraViewportMode &&
 		FMath::IsNearlyEqual(CameraOrbitDistance, Other.CameraOrbitDistance) &&
 		FMath::IsNearlyEqual(CameraFOV, Other.CameraFOV) &&
@@ -69,21 +75,6 @@ bool UNiagaraBakerSettings::Equals(const UNiagaraBakerSettings& Other) const
 		bRenderComponentOnly == Other.bRenderComponentOnly;
 }
 
-float UNiagaraBakerSettings::GetAspectRatio(int32 iOutputTextureIndex) const
-{
-	if (OutputTextures.IsValidIndex(iOutputTextureIndex))
-	{
-		const float TextureAspectRatio = float(OutputTextures[iOutputTextureIndex].FrameSize.Y) / float(OutputTextures[iOutputTextureIndex].FrameSize.X);
-		return bUseCameraAspectRatio ? CameraAspectRatio : TextureAspectRatio;
-	}
-	return 1.0f;
-}
-
-FVector2D UNiagaraBakerSettings::GetOrthoSize(int32 iOutputTextureIndex) const
-{
-	return FVector2D(CameraOrthoWidth, CameraOrthoWidth * GetAspectRatio(iOutputTextureIndex));
-}
-
 FVector UNiagaraBakerSettings::GetCameraLocation() const
 {
 	if (IsPerspective())
@@ -122,92 +113,103 @@ FMatrix UNiagaraBakerSettings::GetViewMatrix() const
 	return FInverseRotationMatrix(GetCameraRotation()) * GetViewportMatrix();
 }
 
-FMatrix UNiagaraBakerSettings::GetProjectionMatrixForTexture(int32 iOutputTextureIndex) const
+FMatrix UNiagaraBakerSettings::GetProjectionMatrix() const
 {
+	const float AspectRatioY = bUseCameraAspectRatio ? CameraAspectRatio : 1.0f;
 	if (CameraViewportMode == ENiagaraBakerViewMode::Perspective)
 	{
-		const float AspectRatio = GetAspectRatio(iOutputTextureIndex);
 		const float HalfXFOV = FMath::DegreesToRadians(CameraFOV) * 0.5f;
-		const float HalfYFOV = FMath::Atan(FMath::Tan(HalfXFOV) / AspectRatio);
+		const float HalfYFOV = FMath::Atan(FMath::Tan(HalfXFOV) / AspectRatioY);
 		return FReversedZPerspectiveMatrix(HalfXFOV, HalfYFOV, 1.0f, 1.0, GNearClippingPlane, GNearClippingPlane);
 	}
 	else
 	{
 		const float ZRange = WORLD_MAX;
-		const FVector2D OrthoSize = GetOrthoSize(iOutputTextureIndex);
-		return FReversedZOrthoMatrix(OrthoSize.X / 2.0f, OrthoSize.Y / 2.0f, 0.5f / ZRange, ZRange);
+		return FReversedZOrthoMatrix(CameraOrthoWidth / 2.0f, CameraOrthoWidth * AspectRatioY / 2.0f, 0.5f / ZRange, ZRange);
 	}
 }
 
-UNiagaraBakerSettings::FDisplayInfo UNiagaraBakerSettings::GetDisplayInfo(float Time, bool bLooping) const
+int UNiagaraBakerSettings::GetOutputNumFrames(UNiagaraBakerOutput* BakerOutput) const
 {
-	FDisplayInfo DisplayInfo;
-	DisplayInfo.NormalizedTime = FMath::Max(Time / DurationSeconds, 0.0f);
-	DisplayInfo.NormalizedTime = bLooping ? FMath::Fractional(DisplayInfo.NormalizedTime) : FMath::Min(DisplayInfo.NormalizedTime, 0.9999f);
+	return FramesPerDimension.X * FramesPerDimension.Y;
+}
+
+FNiagaraBakerOutputFrameIndices UNiagaraBakerSettings::GetOutputFrameIndices(UNiagaraBakerOutput* BakerOutput, float RelativeTime) const
+{
+	FNiagaraBakerOutputFrameIndices DisplayInfo;
+	DisplayInfo.NumFrames = GetOutputNumFrames(BakerOutput);
+	DisplayInfo.NormalizedTime = FMath::Max(RelativeTime / DurationSeconds, 0.0f);
+	DisplayInfo.NormalizedTime = bPreviewLooping ? FMath::Fractional(DisplayInfo.NormalizedTime) : FMath::Min(DisplayInfo.NormalizedTime, 0.9999f);
 
-	const int NumFrames = GetNumFrames();
-	const float FrameTime = DisplayInfo.NormalizedTime * float(NumFrames);
+	const float FrameTime = DisplayInfo.NormalizedTime * float(DisplayInfo.NumFrames);
 	DisplayInfo.FrameIndexA = FMath::FloorToInt(FrameTime);
-	DisplayInfo.FrameIndexB = bLooping ? (DisplayInfo.FrameIndexA + 1) % NumFrames : FMath::Min(DisplayInfo.FrameIndexA + 1, NumFrames - 1);
+	DisplayInfo.FrameIndexB = bPreviewLooping ? (DisplayInfo.FrameIndexA + 1) % DisplayInfo.NumFrames : FMath::Min(DisplayInfo.FrameIndexA + 1, DisplayInfo.NumFrames - 1);
 	DisplayInfo.Interp = FrameTime - float(DisplayInfo.FrameIndexA);
 
 	return DisplayInfo;
 }
 
-void UNiagaraBakerSettings::PostLoad()
+int UNiagaraBakerSettings::GetOutputNumFrames(int OutputIndex) const
 {
-	Super::PostLoad();
+	return Outputs.IsValidIndex(OutputIndex) ? GetOutputNumFrames(Outputs[OutputIndex]) : 0;
+}
 
-	for (FNiagaraBakerTextureSettings& Texture : OutputTextures)
-	{
-		if ( Texture.GeneratedTexture )
-		{
-			Texture.GeneratedTexture->PostLoad();
-		}
-	}
+FNiagaraBakerOutputFrameIndices UNiagaraBakerSettings::GetOutputFrameIndices(int OutputIndex, float RelativeTime) const
+{
+	return Outputs.IsValidIndex(OutputIndex) ? GetOutputFrameIndices(Outputs[OutputIndex], RelativeTime) : FNiagaraBakerOutputFrameIndices();
 }
 
-#if WITH_EDITORONLY_DATA
-void UNiagaraBakerSettings::PostEditChangeProperty(struct FPropertyChangedEvent& PropertyChangedEvent)
+void UNiagaraBakerSettings::PostLoad()
 {
-	Super::PostEditChangeProperty(PropertyChangedEvent);
+	Super::PostLoad();
 
-	bool bComputeOutputTextureSizes = false;
-	if ( PropertyChangedEvent.MemberProperty != nullptr )
+	if ( OutputTextures_DEPRECATED.Num() > 0 )
 	{
-		//if (PropertyChangedEvent.MemberProperty->GetFName() == GET_MEMBER_NAME_CHECKED(FNiagaraBakerTextureSettings))
-		if (PropertyChangedEvent.MemberProperty->GetFName() == GET_MEMBER_NAME_CHECKED(UNiagaraBakerSettings, OutputTextures))
-		{
-			bComputeOutputTextureSizes = true;
-		}
-		else if (PropertyChangedEvent.MemberProperty->GetFName() == GET_MEMBER_NAME_CHECKED(UNiagaraBakerSettings, FramesPerDimension))
+		Outputs.Empty();
+		for (FNiagaraBakerTextureSettings& Texture : OutputTextures_DEPRECATED)
 		{
-			FramesPerDimension.X = FMath::Max(FramesPerDimension.X, 1);
-			FramesPerDimension.Y = FMath::Max(FramesPerDimension.Y, 1);
-			bComputeOutputTextureSizes = true;
-		}
-	}
+			UNiagaraBakerOutputTexture2D* NewOutput = NewObject<UNiagaraBakerOutputTexture2D>(this);
+			NewOutput->SourceBinding		= Texture.SourceBinding;
 
-	// Recompute output texture sizes as something was modified which could impact it
-	if (bComputeOutputTextureSizes)
-	{
-		for (FNiagaraBakerTextureSettings& OutputTexture : OutputTextures)
-		{
-			if (OutputTexture.bUseFrameSize)
+			if ( Texture.bUseFrameSize )
 			{
-				OutputTexture.FrameSize.X = FMath::Max(OutputTexture.FrameSize.X, 1);
-				OutputTexture.FrameSize.Y = FMath::Max(OutputTexture.FrameSize.Y, 1);
-				OutputTexture.TextureSize.X = OutputTexture.FrameSize.X * FramesPerDimension.X;
-				OutputTexture.TextureSize.Y = OutputTexture.FrameSize.Y * FramesPerDimension.Y;
+				NewOutput->FrameSize		= Texture.FrameSize;
+				NewOutput->AtlasTextureSize	= FIntPoint(Texture.FrameSize.X * FramesPerDimension.X, Texture.FrameSize.Y * FramesPerDimension.Y);
 			}
 			else
 			{
-				OutputTexture.TextureSize.X = FMath::Max(OutputTexture.TextureSize.X, 1);
-				OutputTexture.TextureSize.Y = FMath::Max(OutputTexture.TextureSize.Y, 1);
-				OutputTexture.FrameSize.X = FMath::DivideAndRoundDown(OutputTexture.TextureSize.X, FramesPerDimension.X);
-				OutputTexture.FrameSize.Y = FMath::DivideAndRoundDown(OutputTexture.TextureSize.Y, FramesPerDimension.Y);
+				NewOutput->FrameSize		= FIntPoint(Texture.TextureSize.X / FramesPerDimension.X, Texture.TextureSize.Y / FramesPerDimension.Y);
+				NewOutput->AtlasTextureSize	= Texture.TextureSize;
 			}
+
+			if (!Texture.OutputName.IsNone())
+			{
+				const FString OutputName = Texture.OutputName.ToString();
+				if (OutputName.Len() > 0)
+				{
+					NewOutput->OutputName = OutputName;
+				}
+			}
+
+			if ( Texture.GeneratedTexture )
+			{
+				NewOutput->AtlasAssetPathFormat = Texture.GeneratedTexture->GetPackage()->GetPathName();
+			}
+
+			Outputs.Add(NewOutput);
 		}
+
+		OutputTextures_DEPRECATED.Empty();
+	}
+}
+
+#if WITH_EDITORONLY_DATA
+void UNiagaraBakerSettings::PostEditChangeProperty(struct FPropertyChangedEvent& PropertyChangedEvent)
+{
+	Super::PostEditChangeProperty(PropertyChangedEvent);
+	for ( UNiagaraBakerOutput* Output : Outputs )
+	{
+		Output->PostEditChangeProperty(PropertyChangedEvent);
 	}
 }
 #endif
diff --git a/Engine/Plugins/FX/Niagara/Source/NiagaraEditor/NiagaraEditor.Build.cs b/Engine/Plugins/FX/Niagara/Source/NiagaraEditor/NiagaraEditor.Build.cs
--- a/Engine/Plugins/FX/Niagara/Source/NiagaraEditor/NiagaraEditor.Build.cs
+++ b/Engine/Plugins/FX/Niagara/Source/NiagaraEditor/NiagaraEditor.Build.cs
@@ -54,9 +54,10 @@ public class NiagaraEditor : ModuleRules
 				"EditorWidgets",
 				"Renderer",
 				"DeveloperSettings",
-				"PythonScriptPlugin"
+				"PythonScriptPlugin",
+				"ImageWrapper",
 			}
-        );
+		);
 
 		if (Target.bBuildTargetDeveloperTools)
 		{
diff --git a/Engine/Plugins/FX/Niagara/Source/NiagaraEditor/Private/Customizations/NiagaraBakerOutputCustomization.cpp b/Engine/Plugins/FX/Niagara/Source/NiagaraEditor/Private/Customizations/NiagaraBakerOutputCustomization.cpp
new file mode 100644
--- /dev/null
+++ b/Engine/Plugins/FX/Niagara/Source/NiagaraEditor/Private/Customizations/NiagaraBakerOutputCustomization.cpp
@@ -0,0 +1,427 @@
+// Copyright Epic Games, Inc. All Rights Reserved.
+
+#include "NiagaraBakerOutputCustomization.h"
+#include "NiagaraBakerOutputTexture2D.h"
+#include "NiagaraBakerOutputVolumeTexture.h"
+
+#include "Editor/EditorEngine.h"
+#include "Widgets/Images/SImage.h"
+#include "Widgets/Input/SButton.h"
+#include "AssetRegistryModule.h"
+#include "DetailCategoryBuilder.h"
+#include "DetailLayoutBuilder.h"
+#include "DetailWidgetRow.h"
+#include "IDetailCustomization.h"
+#include "IDetailPropertyRow.h"
+#include "IDetailsView.h"
+
+#define LOCTEXT_NAMESPACE "NiagaraBakerOutputCustomization"
+
+namespace NiagaraBakerOutputCustomization
+{
+
+//////////////////////////////////////////////////////////////////////////
+
+struct FNiagaraAssetPathDetailBuilder : public IDetailCustomNodeBuilder, public TSharedFromThis<FNiagaraAssetPathDetailBuilder>
+{
+public:
+	FNiagaraAssetPathDetailBuilder(TSharedPtr<IPropertyHandle> InPropertyHandle)
+		: PropertyHandle(InPropertyHandle)
+	{
+	}
+
+	virtual void GenerateHeaderRowContent(FDetailWidgetRow& NodeRow) override
+	{
+		TSharedPtr<SWidget> NameWidget = PropertyHandle->CreatePropertyNameWidget();
+		TSharedPtr<SWidget> ValueWidget = PropertyHandle->CreatePropertyValueWidget();
+
+		//if (bCanResetToDefault)
+		//{
+		//	TSharedPtr<SHorizontalBox> Container;
+		//	SAssignNew(Container, SHorizontalBox)
+		//		+ SHorizontalBox::Slot()
+		//		.AutoWidth()
+		//		[
+		//			ValueWidget.ToSharedRef()
+		//		]
+		//	+ SHorizontalBox::Slot()
+		//		[
+		//			ParentProperty->CreateDefaultPropertyButtonWidgets()
+		//		];
+
+		//	ValueWidget = Container;
+		//}
+
+		FUIAction CopyAction;
+		FUIAction PasteAction;
+		PropertyHandle->CreateDefaultPropertyCopyPasteActions(CopyAction, PasteAction);
+			NodeRow
+			.NameContent()
+			[
+				NameWidget.ToSharedRef()
+			]
+			.ValueContent()
+			.MinDesiredWidth(300.0f)
+			[
+				SNew(SHorizontalBox)
+				+SHorizontalBox::Slot()
+				[
+					ValueWidget.ToSharedRef()
+				]
+				+SHorizontalBox::Slot()
+				[
+					SNew(STextBlock).Text(LOCTEXT("Ooooh", "Ooooh"))
+				]
+			]
+			.CopyAction(CopyAction)
+			.PasteAction(PasteAction);
+	}
+
+	virtual void GenerateChildContent(IDetailChildrenBuilder& ChildrenBuilder) override
+	{
+		//uint32 NumChildren = 0;
+		//if (ParentProperty->GetNumChildren(NumChildren) != FPropertyAccess::Fail)
+		//{
+		//	for (uint32 ChildIndex = 0; ChildIndex < NumChildren; ++ChildIndex)
+		//	{
+		//		TSharedPtr<IPropertyHandle> ChildProperty = ParentProperty->GetChildHandle(ChildIndex);
+		//		if (ChildProperty.IsValid() && ChildProperty->GetProperty()->GetFName() != MeshProperty->GetProperty()->GetFName())
+		//		{
+		//			ChildrenBuilder.AddProperty(ChildProperty.ToSharedRef())
+		//				.ShouldAutoExpand(false);
+		//		}
+		//	}
+		//}
+	}
+
+	virtual FName GetName() const override
+	{
+		static const FName NAME_NiagaraAssetPathDetailBuilder("NiagaraAssetPathDetailBuilder");
+		return NAME_NiagaraAssetPathDetailBuilder;
+	}
+
+	//virtual void SetOnRebuildChildren(FSimpleDelegate InOnRegenerateChildren) override
+	//{
+	//	OnRebuildChildren = InOnRegenerateChildren;
+	//}
+
+	virtual void Tick(float DeltaTime) override {}
+	virtual bool RequiresTick() const override { return false; }
+	virtual bool InitiallyCollapsed() const override { return true; }
+	virtual TSharedPtr<IPropertyHandle> GetPropertyHandle() const override { return PropertyHandle; }
+
+	//void Rebuild()
+	//{
+	//	OnRebuildChildren.ExecuteIfBound();
+	//}
+
+private:
+	TSharedPtr<IPropertyHandle> PropertyHandle;
+	//TSharedPtr<IPropertyHandle> MeshProperty;
+	//FSimpleDelegate OnRebuildChildren;
+	//bool bCanResetToDefault = true;
+};
+
+//////////////////////////////////////////////////////////////////////////
+
+struct FNiagaraBakerOutputDetails : public IDetailCustomization
+{
+	static void FocusContentBrowserToAsset(const FString& AssetPath)
+	{
+		FAssetRegistryModule& AssetRegistryModule = FModuleManager::LoadModuleChecked<FAssetRegistryModule>(AssetRegistryConstants::ModuleName);
+		IAssetRegistry& AssetRegistry = AssetRegistryModule.Get();
+		TArray<FAssetData> AssetDatas;
+		AssetRegistry.GetAssetsByPackageName(FName(*AssetPath), AssetDatas);
+
+		TArray<UObject*> AssetObjects;
+		AssetObjects.Reserve(AssetDatas.Num());
+		for (const FAssetData& AssetData : AssetDatas)
+		{
+			if (UObject* AssetObject = AssetData.GetAsset())
+			{
+				AssetObjects.Add(AssetObject);
+			}
+		}
+
+		if (AssetObjects.Num() > 0)
+		{
+			GEditor->SyncBrowserToObjects(AssetObjects);
+		}
+	}
+
+	static void ExploreFolder(const FString& Folder)
+	{
+		if ( FPaths::DirectoryExists(Folder) )
+		{
+			FPlatformProcess::ExploreFolder(*Folder);
+		}
+	}
+
+	static void BuildAssetPathWidget(IDetailCategoryBuilder& DetailCategory, TSharedRef<IPropertyHandle>& PropertyHandle, TAttribute<FText>::FGetter TooltipGetter, FOnClicked OnClicked)
+	{
+		DetailCategory.AddProperty(PropertyHandle)
+			.CustomWidget()
+			.NameContent()
+			[
+				PropertyHandle->CreatePropertyNameWidget()
+			]
+			.ValueContent()
+			[
+				SNew(SHorizontalBox)
+				+SHorizontalBox::Slot()
+				.FillWidth(1.0f)
+				[
+					PropertyHandle->CreatePropertyValueWidget()
+				]
+				+SHorizontalBox::Slot()
+				.AutoWidth()
+				.HAlign(HAlign_Right)
+				[
+					SNew(SButton)
+					.ToolTipText(TAttribute<FText>::Create(TooltipGetter))
+					.OnClicked(OnClicked)
+					[
+						SNew(SImage)
+						.Image(FEditorStyle::GetBrush(TEXT("Icons.Search")))
+					]
+				]
+			];
+	}
+};
+
+//////////////////////////////////////////////////////////////////////////
+
+struct FNiagaraBakerOutputTexture2DDetails : public FNiagaraBakerOutputDetails
+{
+	static TSharedRef<IDetailCustomization> MakeInstance()
+	{
+		return MakeShared<FNiagaraBakerOutputTexture2DDetails>();
+	}
+
+	static FText BrowseAtlasToolTipText(TWeakObjectPtr<UNiagaraBakerOutputTexture2D> WeakOutput)
+	{
+		UNiagaraBakerOutputTexture2D* Output = WeakOutput.Get();
+		return FText::Format(LOCTEXT("BrowseAtlasToolTipFormat", "Browse to atlas asset '{0}'"), FText::FromString(Output ? Output->GetAssetPath(Output->AtlasAssetPathFormat, 0) : FString()));
+	}
+
+	static FReply BrowseToAtlas(TWeakObjectPtr<UNiagaraBakerOutputTexture2D> WeakOutput)
+	{
+		if ( UNiagaraBakerOutputTexture2D* Output = WeakOutput.Get() )
+		{
+			const FString AssetPath = Output->GetAssetPath(Output->AtlasAssetPathFormat, 0);
+			FocusContentBrowserToAsset(AssetPath);
+		}
+		return FReply::Handled();
+	}
+
+	static FText BrowseFrameAssetsToolTipText(TWeakObjectPtr<UNiagaraBakerOutputTexture2D> WeakOutput)
+	{
+		UNiagaraBakerOutputTexture2D* Output = WeakOutput.Get();
+		return FText::Format(LOCTEXT("BrowseFrameAssetsToolTipFormat", "Browse to frames assets '{0}'"), FText::FromString(Output ? Output->GetAssetPath(Output->FramesAssetPathFormat, 0) : FString()));
+	}
+
+	static FReply BrowseToFrameAssets(TWeakObjectPtr<UNiagaraBakerOutputTexture2D> WeakOutput)
+	{
+		if (UNiagaraBakerOutputTexture2D* Output = WeakOutput.Get())
+		{
+			const FString AssetPath = Output->GetAssetPath(Output->FramesAssetPathFormat, 0);
+			FocusContentBrowserToAsset(AssetPath);
+		}
+		return FReply::Handled();
+	}
+
+	static FText BrowseExportAssetsToolTipText(TWeakObjectPtr<UNiagaraBakerOutputTexture2D> WeakOutput)
+	{
+		UNiagaraBakerOutputTexture2D* Output = WeakOutput.Get();
+		return FText::Format(LOCTEXT("BrowseExportAssetsToolTipFormat", "Browse to exported assets '{0}'"), FText::FromString(Output ? Output->GetExportFolder(Output->FramesExportPathFormat, 0) : FString()));
+	}
+
+	static FReply BrowseToExportAssets(TWeakObjectPtr<UNiagaraBakerOutputTexture2D> WeakOutput)
+	{
+		if (UNiagaraBakerOutputTexture2D* Output = WeakOutput.Get())
+		{
+			const FString ExportFolder = Output->GetExportFolder(Output->FramesExportPathFormat, 0);
+			ExploreFolder(ExportFolder);
+		}
+		return FReply::Handled();
+	}
+
+	virtual void CustomizeDetails(IDetailLayoutBuilder& DetailBuilder) override
+	{
+		// We only support customization on 1 object
+		TArray<TWeakObjectPtr<UObject>> ObjectsCustomized;
+		DetailBuilder.GetObjectsBeingCustomized(ObjectsCustomized);
+		if (ObjectsCustomized.Num() != 1 || !ObjectsCustomized[0]->IsA<UNiagaraBakerOutputTexture2D>())
+		{
+			return;
+		}
+		TWeakObjectPtr<UNiagaraBakerOutputTexture2D> WeakOutput = CastChecked<UNiagaraBakerOutputTexture2D>(ObjectsCustomized[0]);
+
+		static FName NAME_Settings("Settings");
+		IDetailCategoryBuilder& DetailCategory = DetailBuilder.EditCategory(NAME_Settings);
+
+		TArray<TSharedRef<IPropertyHandle>> CategoryProperties;
+		DetailCategory.GetDefaultProperties(CategoryProperties);
+
+		for (TSharedRef<IPropertyHandle>& PropertyHandle : CategoryProperties)
+		{
+			const FName PropertyName = PropertyHandle->GetProperty() ? PropertyHandle->GetProperty()->GetFName() : NAME_None;
+			if (PropertyName == GET_MEMBER_NAME_CHECKED(UNiagaraBakerOutputTexture2D, AtlasAssetPathFormat))
+			{
+				BuildAssetPathWidget(
+					DetailCategory, PropertyHandle,
+					TAttribute<FText>::FGetter::CreateStatic(&FNiagaraBakerOutputTexture2DDetails::BrowseAtlasToolTipText, WeakOutput),
+					FOnClicked::CreateStatic(&FNiagaraBakerOutputTexture2DDetails::BrowseToAtlas, WeakOutput)
+				);
+			}
+			else if (PropertyName == GET_MEMBER_NAME_CHECKED(UNiagaraBakerOutputTexture2D, FramesAssetPathFormat))
+			{
+				BuildAssetPathWidget(
+					DetailCategory, PropertyHandle,
+					TAttribute<FText>::FGetter::CreateStatic(&FNiagaraBakerOutputTexture2DDetails::BrowseFrameAssetsToolTipText, WeakOutput),
+					FOnClicked::CreateStatic(&FNiagaraBakerOutputTexture2DDetails::BrowseToFrameAssets, WeakOutput)
+				);
+			}
+			else if (PropertyName == GET_MEMBER_NAME_CHECKED(UNiagaraBakerOutputTexture2D, FramesExportPathFormat))
+			{
+				BuildAssetPathWidget(
+					DetailCategory, PropertyHandle,
+					TAttribute<FText>::FGetter::CreateStatic(&FNiagaraBakerOutputTexture2DDetails::BrowseExportAssetsToolTipText, WeakOutput),
+					FOnClicked::CreateStatic(&FNiagaraBakerOutputTexture2DDetails::BrowseToExportAssets, WeakOutput)
+				);
+			}
+			else
+			{
+				DetailCategory.AddProperty(PropertyHandle);
+			}
+		}
+	}
+};
+
+//////////////////////////////////////////////////////////////////////////
+
+struct NiagaraBakerOutputVolumeTexture : public FNiagaraBakerOutputDetails
+{
+	static TSharedRef<IDetailCustomization> MakeInstance()
+	{
+		return MakeShared<NiagaraBakerOutputVolumeTexture>();
+	}
+
+	static FText BrowseAtlasToolTipText(TWeakObjectPtr<UNiagaraBakerOutputVolumeTexture> WeakOutput)
+	{
+		UNiagaraBakerOutputVolumeTexture* Output = WeakOutput.Get();
+		return FText::Format(LOCTEXT("BrowseAtlasToolTipFormat", "Browse to atlas asset '{0}'"), FText::FromString(Output ? Output->GetAssetPath(Output->AtlasAssetPathFormat, 0) : FString()));
+	}
+
+	static FReply BrowseToAtlas(TWeakObjectPtr<UNiagaraBakerOutputVolumeTexture> WeakOutput)
+	{
+		if (UNiagaraBakerOutputVolumeTexture* Output = WeakOutput.Get() )
+		{
+			const FString AssetPath = Output->GetAssetPath(Output->AtlasAssetPathFormat, 0);
+			FocusContentBrowserToAsset(AssetPath);
+		}
+		return FReply::Handled();
+	}
+
+	static FText BrowseFrameAssetsToolTipText(TWeakObjectPtr<UNiagaraBakerOutputVolumeTexture> WeakOutput)
+	{
+		UNiagaraBakerOutputVolumeTexture* Output = WeakOutput.Get();
+		return FText::Format(LOCTEXT("BrowseFrameAssetsToolTipFormat", "Browse to frames assets '{0}'"), FText::FromString(Output ? Output->GetAssetPath(Output->FramesAssetPathFormat, 0) : FString()));
+	}
+
+	static FReply BrowseToFrameAssets(TWeakObjectPtr<UNiagaraBakerOutputVolumeTexture> WeakOutput)
+	{
+		if (UNiagaraBakerOutputVolumeTexture* Output = WeakOutput.Get())
+		{
+			const FString AssetPath = Output->GetAssetPath(Output->FramesAssetPathFormat, 0);
+			FocusContentBrowserToAsset(AssetPath);
+		}
+		return FReply::Handled();
+	}
+
+	static FText BrowseExportAssetsToolTipText(TWeakObjectPtr<UNiagaraBakerOutputVolumeTexture> WeakOutput)
+	{
+		UNiagaraBakerOutputVolumeTexture* Output = WeakOutput.Get();
+		return FText::Format(LOCTEXT("BrowseExportAssetsToolTipFormat", "Browse to exported assets '{0}'"), FText::FromString(Output ? Output->GetExportFolder(Output->FramesExportPathFormat, 0) : FString()));
+	}
+
+	static FReply BrowseToExportAssets(TWeakObjectPtr<UNiagaraBakerOutputVolumeTexture> WeakOutput)
+	{
+		if (UNiagaraBakerOutputVolumeTexture* Output = WeakOutput.Get())
+		{
+			const FString ExportFolder = Output->GetExportFolder(Output->FramesExportPathFormat, 0);
+			ExploreFolder(ExportFolder);
+		}
+		return FReply::Handled();
+	}
+
+	virtual void CustomizeDetails(IDetailLayoutBuilder& DetailBuilder) override
+	{
+		// We only support customization on 1 object
+		TArray<TWeakObjectPtr<UObject>> ObjectsCustomized;
+		DetailBuilder.GetObjectsBeingCustomized(ObjectsCustomized);
+		if (ObjectsCustomized.Num() != 1 || !ObjectsCustomized[0]->IsA<UNiagaraBakerOutputVolumeTexture>())
+		{
+			return;
+		}
+		TWeakObjectPtr<UNiagaraBakerOutputVolumeTexture> WeakOutput = CastChecked<UNiagaraBakerOutputVolumeTexture>(ObjectsCustomized[0]);
+
+		static FName NAME_Settings("Settings");
+		IDetailCategoryBuilder& DetailCategory = DetailBuilder.EditCategory(NAME_Settings);
+
+		TArray<TSharedRef<IPropertyHandle>> CategoryProperties;
+		DetailCategory.GetDefaultProperties(CategoryProperties);
+
+		for (TSharedRef<IPropertyHandle>& PropertyHandle : CategoryProperties)
+		{
+			const FName PropertyName = PropertyHandle->GetProperty() ? PropertyHandle->GetProperty()->GetFName() : NAME_None;
+			if (PropertyName == GET_MEMBER_NAME_CHECKED(UNiagaraBakerOutputVolumeTexture, AtlasAssetPathFormat))
+			{
+				BuildAssetPathWidget(
+					DetailCategory, PropertyHandle,
+					TAttribute<FText>::FGetter::CreateStatic(&NiagaraBakerOutputVolumeTexture::BrowseAtlasToolTipText, WeakOutput),
+					FOnClicked::CreateStatic(&NiagaraBakerOutputVolumeTexture::BrowseToAtlas, WeakOutput)
+				);
+			}
+			else if (PropertyName == GET_MEMBER_NAME_CHECKED(UNiagaraBakerOutputVolumeTexture, FramesAssetPathFormat))
+			{
+				BuildAssetPathWidget(
+					DetailCategory, PropertyHandle,
+					TAttribute<FText>::FGetter::CreateStatic(&NiagaraBakerOutputVolumeTexture::BrowseFrameAssetsToolTipText, WeakOutput),
+					FOnClicked::CreateStatic(&NiagaraBakerOutputVolumeTexture::BrowseToFrameAssets, WeakOutput)
+				);
+			}
+			else if (PropertyName == GET_MEMBER_NAME_CHECKED(UNiagaraBakerOutputVolumeTexture, FramesExportPathFormat))
+			{
+				BuildAssetPathWidget(
+					DetailCategory, PropertyHandle,
+					TAttribute<FText>::FGetter::CreateStatic(&NiagaraBakerOutputVolumeTexture::BrowseExportAssetsToolTipText, WeakOutput),
+					FOnClicked::CreateStatic(&NiagaraBakerOutputVolumeTexture::BrowseToExportAssets, WeakOutput)
+				);
+			}
+			else
+			{
+				DetailCategory.AddProperty(PropertyHandle);
+			}
+		}
+	}
+};
+
+//////////////////////////////////////////////////////////////////////////
+
+void RegisterCustomization(IDetailsView* DetailsView)
+{
+	DetailsView->RegisterInstancedCustomPropertyLayout(
+		UNiagaraBakerOutputTexture2D::StaticClass(),
+		FOnGetDetailCustomizationInstance::CreateStatic(&FNiagaraBakerOutputTexture2DDetails::MakeInstance)
+	);
+
+	DetailsView->RegisterInstancedCustomPropertyLayout(
+		UNiagaraBakerOutputVolumeTexture::StaticClass(),
+		FOnGetDetailCustomizationInstance::CreateStatic(&NiagaraBakerOutputVolumeTexture::MakeInstance)
+	);
+}
+
+} //NiagaraBakerOutputCustomization
+
+#undef LOCTEXT_NAMESPACE
diff --git a/Engine/Plugins/FX/Niagara/Source/NiagaraEditor/Private/Customizations/NiagaraBakerOutputCustomization.h b/Engine/Plugins/FX/Niagara/Source/NiagaraEditor/Private/Customizations/NiagaraBakerOutputCustomization.h
new file mode 100644
--- /dev/null
+++ b/Engine/Plugins/FX/Niagara/Source/NiagaraEditor/Private/Customizations/NiagaraBakerOutputCustomization.h
@@ -0,0 +1,8 @@
+// Copyright Epic Games, Inc. All Rights Reserved.
+
+class IDetailsView;
+
+namespace NiagaraBakerOutputCustomization
+{
+	void RegisterCustomization(IDetailsView* DetailsView);
+}
diff --git a/Engine/Plugins/FX/Niagara/Source/NiagaraEditor/Private/Customizations/NiagaraBakerSettingsDetails.cpp b/Engine/Plugins/FX/Niagara/Source/NiagaraEditor/Private/Customizations/NiagaraBakerSettingsDetails.cpp
--- a/Engine/Plugins/FX/Niagara/Source/NiagaraEditor/Private/Customizations/NiagaraBakerSettingsDetails.cpp
+++ b/Engine/Plugins/FX/Niagara/Source/NiagaraEditor/Private/Customizations/NiagaraBakerSettingsDetails.cpp
@@ -46,7 +46,14 @@ FText FNiagaraBakerTextureSourceDetails::GetText() const
 	if ( Objects.Num() == 1 )
 	{
 		FNiagaraBakerTextureSource* TargetVariable = (FNiagaraBakerTextureSource*)PropertyHandle->GetValueBaseAddress((uint8*)Objects[0]);
-		return FText::FromString(*TargetVariable->SourceName.ToString());
+		if ( TargetVariable->DisplayString.Len() > 0 )
+		{
+			return FText::FromString(TargetVariable->DisplayString);
+		}
+		else
+		{
+			return FText::FromName(TargetVariable->SourceName);
+		}
 	}
 
 	return LOCTEXT("Error", "Error");
@@ -78,35 +85,18 @@ void FNiagaraBakerTextureSourceDetails::CollectAllActions(FGraphActionListBuilde
 	PropertyHandle->GetOuterObjects(Objects);
 	if (Objects.Num() == 1)
 	{
-		if (UNiagaraSystem* NiagaraSystem = Objects[0]->GetTypedOuter<UNiagaraSystem>())
+		if ( UNiagaraBakerOutput* BakerOutput = Cast<UNiagaraBakerOutput>(Objects[0]) )
 		{
-			TArray<FName> RendererOptions = FNiagaraBakerRenderer::GatherAllRenderOptions(NiagaraSystem);
-			for ( FName OptionName : RendererOptions )
+			TUniquePtr<FNiagaraBakerOutputRenderer> OutputRenderer(FNiagaraBakerRenderer::GetOutputRenderer(BakerOutput->GetClass()));
+			if (OutputRenderer.IsValid())
 			{
-				FName SourceName;
-				auto RendererType = FNiagaraBakerRenderer::GetRenderType(OptionName, SourceName);
-				FText CategoryText;
-
-				switch ( RendererType )
+				for (const FNiagaraBakerOutputBinding& RenderBinding : OutputRenderer->GetRendererBindings(BakerOutput))
 				{
-					case FNiagaraBakerRenderer::ERenderType::SceneCapture:
-						CategoryText = LOCTEXT("SceneCapture", "Scene Capture");
-						break;
-					case FNiagaraBakerRenderer::ERenderType::BufferVisualization:
-						CategoryText = LOCTEXT("BufferVis", "Buffer Visualization");
-						break;
-					case FNiagaraBakerRenderer::ERenderType::DataInterface:
-						CategoryText = LOCTEXT("EmitterDataInterface", "Emitter DataInterface");
-						break;
-					case FNiagaraBakerRenderer::ERenderType::Particle:
-						CategoryText = LOCTEXT("ParticleAttribute", "Particle Attribute");
-						break;
+					TSharedPtr<FNiagaraBakerTextureSourceAction> NewNodeAction(
+						new FNiagaraBakerTextureSourceAction(RenderBinding.BindingName, RenderBinding.MenuCategory, RenderBinding.MenuEntry, FText(), 0, FText())
+					);
+					OutAllActions.AddAction(NewNodeAction);
 				}
-				FText MenuText = FText::FromString(*OptionName.ToString());
-				TSharedPtr<FNiagaraBakerTextureSourceAction> NewNodeAction(
-					new FNiagaraBakerTextureSourceAction(OptionName, CategoryText, MenuText, FText(), 0, FText())
-				);
-				OutAllActions.AddAction(NewNodeAction);
 			}
 		}
 	}
@@ -148,6 +138,7 @@ void FNiagaraBakerTextureSourceDetails::OnActionSelected(const TArray<TSharedPtr
 
 			PropertyHandle->NotifyPreChange();
 			FNiagaraBakerTextureSource* TargetVariable = (FNiagaraBakerTextureSource*)PropertyHandle->GetValueBaseAddress((uint8*)Objects[0]);
+			TargetVariable->DisplayString = EventSourceAction->GetCategory().ToString() + " - " + EventSourceAction->GetMenuDescription().ToString();
 			TargetVariable->SourceName = EventSourceAction->BindingName;
 			PropertyHandle->NotifyPostChange(EPropertyChangeType::ValueSet);
 			PropertyHandle->NotifyFinishedChangingProperties();
@@ -155,4 +146,4 @@ void FNiagaraBakerTextureSourceDetails::OnActionSelected(const TArray<TSharedPtr
 	}
 }
 
-#undef LOCTEXT_NAMESPACE
\ No newline at end of file
+#undef LOCTEXT_NAMESPACE
diff --git a/Engine/Plugins/FX/Niagara/Source/NiagaraEditor/Private/NiagaraBakerRenderer.cpp b/Engine/Plugins/FX/Niagara/Source/NiagaraEditor/Private/NiagaraBakerRenderer.cpp
--- a/Engine/Plugins/FX/Niagara/Source/NiagaraEditor/Private/NiagaraBakerRenderer.cpp
+++ b/Engine/Plugins/FX/Niagara/Source/NiagaraEditor/Private/NiagaraBakerRenderer.cpp
@@ -1,515 +1,694 @@
 // Copyright Epic Games, Inc. All Rights Reserved.
 
 #include "NiagaraBakerRenderer.h"
+#include "NiagaraBakerSettings.h"
+#include "NiagaraBakerOutputTexture2D.h"
+#include "NiagaraBakerOutputVolumeTexture.h"
 #include "NiagaraComponent.h"
 #include "NiagaraSystem.h"
 #include "NiagaraComputeExecutionContext.h"
 #include "NiagaraGpuComputeDispatchInterface.h"
 #include "NiagaraBatchedElements.h"
+#include "NiagaraBakerRendererOutputTexture2D.h"
+#include "NiagaraBakerRendererOutputVolumeTexture.h"
 
 #include "Niagara/Classes/NiagaraDataInterfaceRenderTarget2D.h"
 #include "Niagara/Classes/NiagaraDataInterfaceGrid2DCollection.h"
 
 #include "Components/SceneCaptureComponent2D.h"
+#include "Engine/Canvas.h"
 #include "Engine/TextureRenderTarget2D.h"
+#include "Misc/FileHelper.h"
+#include "Misc/Paths.h"
 #include "Modules/ModuleManager.h"
-#include "AssetToolsModule.h"
+#include "AdvancedPreviewScene.h"
 #include "BufferVisualizationData.h"
 #include "CanvasTypes.h"
 #include "CanvasItem.h"
 #include "EngineModule.h"
+#include "IImageWrapperModule.h"
+#include "IImageWrapper.h"
+#include "ImageWrapperHelper.h"
 #include "LegacyScreenPercentageDriver.h"
-#include "PackageTools.h"
 
-namespace NiagaraBakerRendererLocal
-{
-	const FString STRING_SceneCaptureSource("SceneCaptureSource");
-	const FString STRING_BufferVisualization("BufferVisualization");
-	const FString STRING_EmitterDI("EmitterDI");
-	const FString STRING_EmitterParticles("EmitterParticles");
-}
-
-FNiagaraBakerRenderer::FNiagaraBakerRenderer()
-{
-	SceneCaptureComponent = NewObject<USceneCaptureComponent2D>(GetTransientPackage(), NAME_None, RF_Transient);
-	SceneCaptureComponent->bTickInEditor = false;
-	SceneCaptureComponent->SetComponentTickEnabled(false);
-	SceneCaptureComponent->SetVisibility(true);
-	SceneCaptureComponent->bCaptureEveryFrame = false;
-	SceneCaptureComponent->bCaptureOnMovement = false;
-}
+//////////////////////////////////////////////////////////////////////////
 
-FNiagaraBakerRenderer::~FNiagaraBakerRenderer()
-{
-}
+const FString FNiagaraBakerOutputBindingHelper::STRING_SceneCaptureSource("SceneCaptureSource");
+const FString FNiagaraBakerOutputBindingHelper::STRING_BufferVisualization("BufferVisualization");
+const FString FNiagaraBakerOutputBindingHelper::STRING_EmitterDI("EmitterDI");
+const FString FNiagaraBakerOutputBindingHelper::STRING_EmitterParticles("EmitterParticles");
 
-bool FNiagaraBakerRenderer::RenderView(UNiagaraComponent* PreviewComponent, const UNiagaraBakerSettings* BakerSettings, float WorldTime, UTextureRenderTarget2D* RenderTarget, int32 iOutputTextureIndex) const
+FNiagaraBakerOutputBindingHelper::ERenderType FNiagaraBakerOutputBindingHelper::GetRenderType(FName BindingName, FName& OutName)
 {
-	if (!PreviewComponent || !BakerSettings || !RenderTarget)
+	OutName = FName();
+	if (BindingName.IsNone())
 	{
-		return false;
+		return ERenderType::SceneCapture;
 	}
 
-	UWorld* World = PreviewComponent->GetWorld();
-	if ( (World == nullptr) || (World->Scene == nullptr) )
+	FString SourceBindingString = BindingName.ToString();
+	TArray<FString> SplitNames;
+	SourceBindingString.ParseIntoArray(SplitNames, TEXT("."));
+
+	if (!ensure(SplitNames.Num() > 0))
 	{
-		return false;
+		return ERenderType::None;
 	}
 
-	FCanvas Canvas(RenderTarget->GameThread_GetRenderTargetResource(), nullptr, FGameTime::CreateUndilated(WorldTime, FApp::GetDeltaTime()), World->Scene->GetFeatureLevel());
-	Canvas.Clear(FLinearColor::Black);
-
-	bool bRendered = false;
-	if (BakerSettings->OutputTextures.IsValidIndex(iOutputTextureIndex))
+	// Scene Capture mode
+	if (SplitNames[0] == STRING_SceneCaptureSource)
 	{
-		const FNiagaraBakerTextureSettings& OutputTextureSettings = BakerSettings->OutputTextures[iOutputTextureIndex];
-		const FIntRect ViewRect = FIntRect(0, 0, OutputTextureSettings.FrameSize.X, OutputTextureSettings.FrameSize.Y);
-		bRendered = RenderView(PreviewComponent, BakerSettings, WorldTime, RenderTarget, &Canvas, iOutputTextureIndex, ViewRect);
-	}
-	Canvas.Flush_GameThread();
-	return bRendered;
-}
+		if (!ensure(SplitNames.Num() == 2))
+		{
+			return ERenderType::None;
+		}
 
-bool FNiagaraBakerRenderer::RenderView(UNiagaraComponent* PreviewComponent, const UNiagaraBakerSettings* BakerSettings, float WorldTime, UTextureRenderTarget2D* RenderTarget, FCanvas* Canvas, int32 iOutputTextureIndex, FIntRect ViewRect) const
-{
-	using namespace NiagaraBakerRendererLocal;
+		OutName = FName(SplitNames[1]);
+		return ERenderType::SceneCapture;
+	}
 
-	if (!PreviewComponent || !BakerSettings || !RenderTarget)
+	// Buffer Visualization Mode
+	if (SplitNames[0] == STRING_BufferVisualization)
 	{
-		return false;
+		if (!ensure(SplitNames.Num() == 2))
+		{
+			return ERenderType::None;
+		}
+		OutName = FName(SplitNames[1]);
+		return ERenderType::BufferVisualization;
 	}
 
-	if (!BakerSettings->OutputTextures.IsValidIndex(iOutputTextureIndex))
+	// Emitter Data Interface
+	if (SplitNames[0] == STRING_EmitterDI)
 	{
-		return false;
+		OutName = FName(*SourceBindingString.RightChop(SplitNames[0].Len() + 1));
+		return ERenderType::DataInterface;
 	}
 
-	if (Canvas->IsHitTesting())
+	// Emitter Data Interface
+	if (SplitNames[0] == STRING_EmitterParticles)
 	{
-		return false;
+		OutName = FName(*SourceBindingString.RightChop(SplitNames[0].Len() + 1));
+		return ERenderType::Particle;
 	}
+	return ERenderType::None;
+}
 
-	UWorld* World = PreviewComponent->GetWorld();
-	check(World);
-
-	// Send EOF updates before we flush our pending ticks to ensure everything is ready for Niagara
-	World->SendAllEndOfFrameUpdates();
-
-	// We need any GPU sims to flush pending ticks
-	FNiagaraGpuComputeDispatchInterface* ComputeDispatchInterface = FNiagaraGpuComputeDispatchInterface::Get(World);
-	ensureMsgf(ComputeDispatchInterface, TEXT("The batcher was not valid on the world this may result in incorrect baking"));
-	if (ComputeDispatchInterface)
+void FNiagaraBakerOutputBindingHelper::GetSceneCaptureBindings(TArray<FNiagaraBakerOutputBinding>& OutBindings)
+{
+	static UEnum* SceneCaptureOptions = StaticEnum<ESceneCaptureSource>();
+	for (int i = 0; i < SceneCaptureOptions->GetMaxEnumValue(); ++i)
 	{
-		ComputeDispatchInterface->FlushPendingTicks_GameThread();
+		FNiagaraBakerOutputBinding& NewBinding = OutBindings.AddDefaulted_GetRef();
+		NewBinding.BindingName = FName(STRING_SceneCaptureSource + TEXT(".") + SceneCaptureOptions->GetNameStringByIndex(i));
+		NewBinding.MenuCategory = FText::FromString(STRING_SceneCaptureSource);
+		NewBinding.MenuEntry = SceneCaptureOptions->GetDisplayNameTextByIndex(i);
 	}
-	
-	const FNiagaraBakerTextureSettings& OutputTextureSettings = BakerSettings->OutputTextures[iOutputTextureIndex];
+}
 
-	// Validate the rect
-	if ( ViewRect.Area() <= 0 )
+void FNiagaraBakerOutputBindingHelper::GetBufferVisualizationBindings(TArray<FNiagaraBakerOutputBinding>& OutBindings)
+{
+	// Gather all buffer visualization options
+	struct FIterator
 	{
-		return false;
-	}
+		TArray<FNiagaraBakerOutputBinding>& OutBindings;
+
+		FIterator(TArray<FNiagaraBakerOutputBinding>& InOutBindings)
+			: OutBindings(InOutBindings)
+		{}
+
+		void ProcessValue(const FString& MaterialName, UMaterialInterface* Material, const FText& DisplayName)
+		{
+			FNiagaraBakerOutputBinding& NewBinding = OutBindings.AddDefaulted_GetRef();
+			NewBinding.BindingName = FName(STRING_BufferVisualization + TEXT(".") + MaterialName);
+			NewBinding.MenuCategory = FText::FromString(STRING_BufferVisualization);
+			NewBinding.MenuEntry = DisplayName;
+		}
+	} Iterator(OutBindings);
+	GetBufferVisualizationData().IterateOverAvailableMaterials(Iterator);
+}
 
-	// Get visualization options
-	float GammaCorrection = 1.0f;
-	FName SourceName;
+void FNiagaraBakerOutputBindingHelper::ForEachEmitterDataInterface(UNiagaraSystem* NiagaraSystem, FEmitterDIFunction Function)
+{
+	check(NiagaraSystem);
 
-	auto RenderType = GetRenderType(OutputTextureSettings.SourceBinding.SourceName, SourceName);
-	switch ( RenderType )
+	for (int32 EmitterIndex=0; EmitterIndex < NiagaraSystem->GetEmitterHandles().Num(); ++EmitterIndex)
 	{
-		case ERenderType::SceneCapture:
+		const FNiagaraEmitterHandle& EmitterHandle = NiagaraSystem->GetEmitterHandle(EmitterIndex);
+		UNiagaraEmitter* EmitterInstance = EmitterHandle.GetInstance();
+		if (!EmitterHandle.IsValid() || !EmitterHandle.GetIsEnabled() || !EmitterInstance)
 		{
-			// Setup capture component
-			SceneCaptureComponent->RegisterComponentWithWorld(World);
-			SceneCaptureComponent->TextureTarget = RenderTarget;
+			continue;
+		}
 
-			// Translate source to enum
-			static UEnum* SceneCaptureOptions = StaticEnum<ESceneCaptureSource>();
-			const int32 CaptureSource = SceneCaptureOptions->GetIndexByName(SourceName);
-			SceneCaptureComponent->CaptureSource = (CaptureSource == INDEX_NONE) ? ESceneCaptureSource::SCS_SceneColorHDR : ESceneCaptureSource(CaptureSource);
+		const FString EmitterName = EmitterHandle.GetName().ToString();
+		const FString EmitterPrefix = EmitterName + TEXT(".");
 
-			// Set view location
-			if (BakerSettings->IsOrthographic())
-			{
-				SceneCaptureComponent->ProjectionType = ECameraProjectionMode::Orthographic;
-				SceneCaptureComponent->OrthoWidth = BakerSettings->CameraOrthoWidth;
-			}
-			else
+		EmitterInstance->ForEachScript(
+			[&](UNiagaraScript* NiagaraScript)
 			{
-				SceneCaptureComponent->ProjectionType = ECameraProjectionMode::Perspective;
-				SceneCaptureComponent->FOVAngle = BakerSettings->CameraFOV;
-			}
-
-			const FMatrix SceneCaptureMatrix = FMatrix(FPlane(0, 0, 1, 0), FPlane(1, 0, 0, 0), FPlane(0, 1, 0, 0), FPlane(0, 0, 0, 1));
-			FMatrix ViewMatrix = SceneCaptureMatrix * BakerSettings->GetViewportMatrix().Inverse() * FRotationTranslationMatrix(BakerSettings->GetCameraRotation(), BakerSettings->GetCameraLocation());
-			SceneCaptureComponent->SetWorldLocationAndRotation(ViewMatrix.GetOrigin(), ViewMatrix.Rotator());
+				if (const FNiagaraScriptExecutionParameterStore* SrcStore = NiagaraScript->GetExecutionReadyParameterStore(EmitterInstance->SimTarget))
+				{
+					for (const FNiagaraVariableWithOffset& Variable : SrcStore->ReadParameterVariables())
+					{
+						if (Variable.IsDataInterface() == false)
+						{
+							continue;
+						}
 
-			SceneCaptureComponent->bUseCustomProjectionMatrix = true;
-			SceneCaptureComponent->CustomProjectionMatrix = BakerSettings->GetProjectionMatrixForTexture(iOutputTextureIndex);
+						const FString VariableName = Variable.GetName().ToString();
+						if (!VariableName.StartsWith(EmitterPrefix))
+						{
+							continue;
+						}
 
-			SceneCaptureComponent->PrimitiveRenderMode = ESceneCapturePrimitiveRenderMode::PRM_UseShowOnlyList;
-			SceneCaptureComponent->ShowOnlyComponents.Empty(1);
-			SceneCaptureComponent->ShowOnlyComponents.Add(PreviewComponent);
+						UNiagaraDataInterface* DataInterface = SrcStore->GetDataInterface(Variable.Offset);
+						check(DataInterface);
 
-			SceneCaptureComponent->CaptureScene();
+						Function(EmitterName, VariableName.Mid(EmitterPrefix.Len()), DataInterface);
+					}
+				}
+			}
+		);
+	}
+}
 
-			SceneCaptureComponent->TextureTarget = nullptr;
-			SceneCaptureComponent->UnregisterComponent();
+UNiagaraDataInterface* FNiagaraBakerOutputBindingHelper::GetDataInterface(UNiagaraComponent* NiagaraComponent, FName DataInterfaceName)
+{
+	// Find data interface
+	FNiagaraSystemInstanceControllerPtr SystemInstanceController = NiagaraComponent->GetSystemInstanceController();
+	if ( SystemInstanceController.IsValid() == false )
+	{
+		return nullptr;
+	}
 
-			// Alpha from a scene capture is 1- so we need to invert
-			if ( SceneCaptureComponent->CaptureSource == ESceneCaptureSource::SCS_SceneColorHDR )
+	FNiagaraSystemInstance* SystemInstance = SystemInstanceController->GetSoloSystemInstance();
+	for (auto EmitterInstance : SystemInstance->GetEmitters())
+	{
+		if ( FNiagaraComputeExecutionContext* ComputeContext = EmitterInstance->GetGPUContext() )
+		{
+			for (const FNiagaraVariableWithOffset& Variable : ComputeContext->CombinedParamStore.ReadParameterVariables())
 			{
-				FCanvasTileItem TileItem(FVector2D(ViewRect.Min.X, ViewRect.Min.Y), FVector2D(ViewRect.Width(), ViewRect.Height()), FLinearColor::White);
-				TileItem.BlendMode = SE_BLEND_Opaque;
-				TileItem.BatchedElementParameters = new FBatchedElementNiagaraInvertColorChannel(0);
-				Canvas->DrawItem(TileItem);
+				if (Variable.IsDataInterface() && (Variable.GetName() == DataInterfaceName))
+				{
+					return ComputeContext->CombinedParamStore.GetDataInterface(Variable.Offset);
+				}
 			}
-
-			return true;
 		}
+	}
+	return nullptr;
+}
 
-		// Regular render
-		case ERenderType::BufferVisualization:
+void FNiagaraBakerOutputBindingHelper::GetDataInterfaceBindingsForCanvas(TArray<FNiagaraBakerOutputBinding>& OutBindings, UNiagaraSystem* NiagaraSystem)
+{
+	check(NiagaraSystem);
+
+	ForEachEmitterDataInterface(
+		NiagaraSystem,
+		[&](const FString& EmitterName, const FString& VariableName, UNiagaraDataInterface* DataInterface)
 		{
-			// Create View Family
-			FSceneViewFamilyContext ViewFamily(
-				FSceneViewFamily::ConstructionValues(RenderTarget->GameThread_GetRenderTargetResource(), World->Scene, FEngineShowFlags(ESFIM_Game))
-				.SetTime(FGameTime::CreateUndilated(WorldTime, FApp::GetDeltaTime()))
-				.SetGammaCorrection(GammaCorrection)
-			);
-
-			ViewFamily.EngineShowFlags.SetScreenPercentage(false);
-			//ViewFamily.EngineShowFlags.DisableAdvancedFeatures();
-			//ViewFamily.EngineShowFlags.MotionBlur = 0;
-			//ViewFamily.EngineShowFlags.SetDistanceCulledPrimitives(true); // show distance culled objects
-			//ViewFamily.EngineShowFlags.SetPostProcessing(false);
-
-			if (SourceName.IsValid())
+			if (DataInterface->CanRenderVariablesToCanvas())
 			{
-				ViewFamily.EngineShowFlags.SetPostProcessing(true);
-				ViewFamily.EngineShowFlags.SetVisualizeBuffer(true);
-				ViewFamily.EngineShowFlags.SetTonemapper(false);
-				ViewFamily.EngineShowFlags.SetScreenPercentage(false);
-			}
+				TArray<FNiagaraVariableBase> RendererableVariables;
+				DataInterface->GetCanvasVariables(RendererableVariables);
+				for (const FNiagaraVariableBase& RendererableVariable : RendererableVariables)
+				{
+					const FString VariableString = VariableName + TEXT(".") + RendererableVariable.GetName().ToString();
 
-			FSceneViewInitOptions ViewInitOptions;
-			ViewInitOptions.SetViewRectangle(ViewRect);
-			ViewInitOptions.ViewFamily = &ViewFamily;
-			ViewInitOptions.ViewOrigin = BakerSettings->GetCameraLocation();
-			ViewInitOptions.ViewRotationMatrix = BakerSettings->GetViewMatrix();
-			ViewInitOptions.ProjectionMatrix = BakerSettings->GetProjectionMatrixForTexture(iOutputTextureIndex);
-			ViewInitOptions.BackgroundColor = FLinearColor::Black;
-			if (BakerSettings->bRenderComponentOnly)
-			{
-				ViewInitOptions.ShowOnlyPrimitives.Emplace();
-				ViewInitOptions.ShowOnlyPrimitives->Add(PreviewComponent->ComponentId);
+					FNiagaraBakerOutputBinding& NewBinding = OutBindings.AddDefaulted_GetRef();
+					NewBinding.BindingName = FName(STRING_EmitterDI + TEXT(".") + EmitterName + TEXT(".") + VariableString);
+					NewBinding.MenuCategory = FText::FromString(TEXT("DataInterface") + EmitterName);
+					NewBinding.MenuEntry = FText::FromString(VariableString);
+				}
 			}
+		}
+	);
+}
 
-			FSceneView* NewView = new FSceneView(ViewInitOptions);
-			NewView->CurrentBufferVisualizationMode = SourceName;
-			ViewFamily.Views.Add(NewView);
-
-			ViewFamily.SetScreenPercentageInterface(new FLegacyScreenPercentageDriver(ViewFamily, 1.0f));
+void FNiagaraBakerOutputBindingHelper::GetParticleAttributeBindings(TArray<FNiagaraBakerOutputBinding>& OutBindings, UNiagaraSystem* NiagaraSystem)
+{
+	check(NiagaraSystem);
 
-			GetRendererModule().BeginRenderingViewFamily(Canvas, &ViewFamily);
+	const TArray<TSharedRef<const FNiagaraEmitterCompiledData>>& AllEmitterCompiledData = NiagaraSystem->GetEmitterCompiledData();
 
-			return true;
+	for (int32 EmitterIndex = 0; EmitterIndex < NiagaraSystem->GetEmitterHandles().Num(); ++EmitterIndex)
+	{
+		const FNiagaraEmitterHandle& EmitterHandle = NiagaraSystem->GetEmitterHandle(EmitterIndex);
+		UNiagaraEmitter* EmitterInstance = EmitterHandle.GetInstance();
+		if (!EmitterHandle.IsValid() || !EmitterHandle.GetIsEnabled() || !EmitterInstance)
+		{
+			continue;
 		}
 
-		// Data Interface
-		case ERenderType::DataInterface:
+		const FString EmitterName = EmitterHandle.GetName().ToString();
+
+		if (ensure(AllEmitterCompiledData.IsValidIndex(EmitterIndex)))
 		{
-			if (UNiagaraSystem* NiagaraSystem = BakerSettings->GetTypedOuter<UNiagaraSystem>())
+			const FNiagaraDataSetCompiledData& ParticleDataSet = AllEmitterCompiledData[EmitterIndex]->DataSetCompiledData;
+			for (int32 iVariable = 0; iVariable < ParticleDataSet.VariableLayouts.Num(); ++iVariable)
 			{
-				// Gather data interface / attribute name
-				FString SourceString = SourceName.ToString();
-				int32 DotIndex;
-				if (!SourceString.FindLastChar('.', DotIndex))
+				const FNiagaraVariable& Variable = ParticleDataSet.Variables[iVariable];
+				const FNiagaraVariableLayoutInfo& VariableLayout = ParticleDataSet.VariableLayouts[iVariable];
+				if (VariableLayout.GetNumFloatComponents() > 0)
 				{
-					return false;
-				}
-
-				const FName DataInterfaceName = FName(SourceString.LeftChop(SourceString.Len() - DotIndex));
-				const FName VariableName = FName(SourceString.RightChop(DotIndex + 1));
-
-				// Find data interface
-				FNiagaraSystemInstanceControllerPtr SystemInstanceController = PreviewComponent->GetSystemInstanceController();
-				check(SystemInstanceController.IsValid());
+					const FString VariableString = Variable.GetName().ToString();
 
-				FNiagaraSystemInstance* SystemInstance = SystemInstanceController->GetSoloSystemInstance();
-				const FNiagaraSystemInstanceID SystemInstanceID = SystemInstance->GetId();
-				for (auto EmitterInstance : SystemInstance->GetEmitters())
-				{
-					if (FNiagaraComputeExecutionContext* ExecContext = EmitterInstance->GetGPUContext())
-					{
-						for (const FNiagaraVariableWithOffset& Variable : ExecContext->CombinedParamStore.ReadParameterVariables())
-						{
-							if (Variable.IsDataInterface())
-							{
-								if (Variable.GetName() == DataInterfaceName)
-								{
-									UNiagaraDataInterface* DataInterface = ExecContext->CombinedParamStore.GetDataInterface(Variable.Offset);
-									check(DataInterface);
-									return DataInterface->RenderVariableToCanvas(SystemInstanceID, VariableName, Canvas, ViewRect);
-								}
-							}
-						}
-					}
+					FNiagaraBakerOutputBinding& NewBinding = OutBindings.AddDefaulted_GetRef();
+					NewBinding.BindingName = FName(STRING_EmitterParticles + TEXT(".") + EmitterName + TEXT(".") + VariableString);
+					NewBinding.MenuCategory = FText::FromString(TEXT("ParticleAttribute ") + EmitterName);
+					NewBinding.MenuEntry = FText::FromString(VariableString);
 				}
 			}
-			return false;
 		}
+	}
+}
 
-		// Particle attribute
-		case ERenderType::Particle:
-		{
-			FString SourceString = SourceName.ToString();
-			int32 DotIndex;
-			if ( !SourceString.FindChar('.', DotIndex) )
-			{
-				return false;
-			}
+//////////////////////////////////////////////////////////////////////////
 
-			const FString EmitterName = SourceString.LeftChop(SourceString.Len() - DotIndex);
-			const FName AttributeName = FName(SourceString.RightChop(DotIndex + 1));
+FNiagaraBakerRenderer::FNiagaraBakerRenderer(UNiagaraSystem* InNiagaraSystem)
+	: NiagaraSystem(InNiagaraSystem)
+{
+	SceneCaptureComponent = NewObject<USceneCaptureComponent2D>(GetTransientPackage(), NAME_None, RF_Transient);
+	SceneCaptureComponent->bTickInEditor = false;
+	SceneCaptureComponent->SetComponentTickEnabled(false);
+	SceneCaptureComponent->SetVisibility(true);
+	SceneCaptureComponent->bCaptureEveryFrame = false;
+	SceneCaptureComponent->bCaptureOnMovement = false;
 
-			FNiagaraSystemInstanceControllerPtr SystemInstanceController = PreviewComponent->GetSystemInstanceController();
-			if (!ensure(SystemInstanceController.IsValid()))
-			{
-				return false;
-			}
+	PreviewComponent = NewObject<UNiagaraComponent>(GetTransientPackage(), NAME_None, RF_Transient);
+	PreviewComponent->CastShadow = 1;
+	PreviewComponent->bCastDynamicShadow = 1;
+	PreviewComponent->SetAllowScalability(false);
+	PreviewComponent->SetAsset(NiagaraSystem);
+	PreviewComponent->SetForceSolo(true);
+	PreviewComponent->SetAgeUpdateMode(ENiagaraAgeUpdateMode::DesiredAge);
+	PreviewComponent->SetCanRenderWhileSeeking(true);
+	PreviewComponent->SetMaxSimTime(0.0f);
+	PreviewComponent->Activate(true);
+
+	AdvancedPreviewScene = MakeShareable(new FAdvancedPreviewScene(FPreviewScene::ConstructionValues()));
+	AdvancedPreviewScene->SetFloorVisibility(false);
+	AdvancedPreviewScene->AddComponent(PreviewComponent, PreviewComponent->GetRelativeTransform());
+}
 
-			FNiagaraSystemInstance* SystemInstance = SystemInstanceController->GetSoloSystemInstance();
-			if ( !ensure(SystemInstance) )
-			{
-				return false;
-			}
+FNiagaraBakerRenderer::~FNiagaraBakerRenderer()
+{
+	if (AdvancedPreviewScene && PreviewComponent)
+	{
+		AdvancedPreviewScene->RemoveComponent(PreviewComponent);
+	}
 
-			for ( const auto& EmitterInstance : SystemInstance->GetEmitters() )
-			{
-				UNiagaraEmitter* NiagaraEmitter = EmitterInstance->GetCachedEmitter();
-				if ( !NiagaraEmitter || (NiagaraEmitter->GetUniqueEmitterName() != EmitterName) )
-				{
-					continue;
-				}
+	if (PreviewComponent)
+	{
+		PreviewComponent->DestroyComponent();
+	}
+}
 
-				if (EmitterInstance->GetGPUContext() != nullptr)
-				{
-					return false;
-				}
+void FNiagaraBakerRenderer::SetAbsoluteTime(float AbsoluteTime, bool bShouldTickComponent)
+{
+	UNiagaraBakerSettings* BakerSettings = GetBakerSettings();
+	if ( !ensure(BakerSettings) )
+	{
+		return;
+	}
 
-				const FNiagaraDataSet& ParticleDataSet = EmitterInstance->GetData();
-				const FNiagaraDataBuffer* ParticleDataBuffer = ParticleDataSet.GetCurrentData();
-				FNiagaraDataSetReaderInt32<int32> UniqueIDAccessor = FNiagaraDataSetAccessor<int32>::CreateReader(ParticleDataSet, FName("UniqueID"));
-				if ( !ParticleDataBuffer || !UniqueIDAccessor.IsValid() )
-				{
-					return false;
-				}
+	PreviewComponent->SetSeekDelta(BakerSettings->GetSeekDelta());
+	PreviewComponent->SeekToDesiredAge(AbsoluteTime);
 
-				const int32 VariableIndex = ParticleDataSet.GetCompiledData().Variables.IndexOfByPredicate([&AttributeName](const FNiagaraVariable& Variable) { return Variable.GetName() == AttributeName; });
-				if (VariableIndex == INDEX_NONE)
-				{
-					return false;
-				}
-				const FNiagaraVariableLayoutInfo& VariableInfo = ParticleDataSet.GetCompiledData().VariableLayouts[VariableIndex];
+	if (bShouldTickComponent)
+	{
+		PreviewComponent->TickComponent(BakerSettings->GetSeekDelta(), ELevelTick::LEVELTICK_All, nullptr);
 
-				float* FloatChannels[4];
-				FloatChannels[0] = (float*)ParticleDataBuffer->GetComponentPtrFloat(VariableInfo.FloatComponentStart);
-				FloatChannels[1] = VariableInfo.GetNumFloatComponents() > 1 ? (float*)ParticleDataBuffer->GetComponentPtrFloat(VariableInfo.FloatComponentStart + 1) : nullptr;
-				FloatChannels[2] = VariableInfo.GetNumFloatComponents() > 2 ? (float*)ParticleDataBuffer->GetComponentPtrFloat(VariableInfo.FloatComponentStart + 2) : nullptr;
-				FloatChannels[3] = VariableInfo.GetNumFloatComponents() > 3 ? (float*)ParticleDataBuffer->GetComponentPtrFloat(VariableInfo.FloatComponentStart + 3) : nullptr;
+		// World should be guaranteed but let's be safe
+		UWorld* World = PreviewComponent->GetWorld();
+		if ( ensure(World) )
+		{
+			// Send EOF updates before we flush our pending ticks to ensure everything is ready for Niagara
+			World->SendAllEndOfFrameUpdates();
 
-				const int32 ParticleBufferStore = OutputTextureSettings.TextureSize.X * OutputTextureSettings.TextureSize.Y;
-				for ( uint32 i=0; i < ParticleDataBuffer->GetNumInstances(); ++i )
-				{
-					const int32 UniqueID = UniqueIDAccessor[i];
-					if (UniqueID >= ParticleBufferStore)
-					{
-						continue;
-					}
+			// Since captures, etc, don't flush GPU updates so we need to force flush them
+			FNiagaraGpuComputeDispatchInterface* ComputeDispatchInterface = FNiagaraGpuComputeDispatchInterface::Get(World);
+			if ( ensureMsgf(ComputeDispatchInterface, TEXT("The batcher was not valid on the world this may result in incorrect baking")) )
+			{
+				ComputeDispatchInterface->FlushPendingTicks_GameThread();
+			}
+		}
+	}
+}
 
-					FLinearColor OutputColor;
-					OutputColor.R = FloatChannels[0] ? FloatChannels[0][i] : 0.0f;
-					OutputColor.G = FloatChannels[1] ? FloatChannels[1][i] : 0.0f;
-					OutputColor.B = FloatChannels[2] ? FloatChannels[2][i] : 0.0f;
-					OutputColor.A = FloatChannels[3] ? FloatChannels[3][i] : 0.0f;
+void FNiagaraBakerRenderer::RenderSceneCapture(UTextureRenderTarget2D* RenderTarget, ESceneCaptureSource CaptureSource) const
+{
+	UNiagaraBakerSettings* BakerSettings = GetBakerSettings();
+	if ( !RenderTarget || !BakerSettings )
+	{
+		return;
+	}
 
-					const int32 TexelX = UniqueID % OutputTextureSettings.TextureSize.X;
-					const int32 TexelY = UniqueID / OutputTextureSettings.TextureSize.X;
-					Canvas->DrawTile(TexelX, TexelY, 1.0f, 1.0f, 0.0f, 0.0f, 1.0f, 1.0f, OutputColor);
-				}
+	const float WorldTime = GetWorldTime();
+	UWorld* World = GetWorld();
 
-				return true;
-			}
+	FCanvas Canvas(RenderTarget->GameThread_GetRenderTargetResource(), nullptr, WorldTime, FApp::GetDeltaTime(), WorldTime, GetFeatureLevel());
+	Canvas.Clear(FLinearColor::Black);
 
-			return false;
-		}
+	SceneCaptureComponent->RegisterComponentWithWorld(World);
+	SceneCaptureComponent->TextureTarget = RenderTarget;
+	SceneCaptureComponent->CaptureSource = CaptureSource;
+	
+	// Set view location
+	if (BakerSettings->IsOrthographic())
+	{
+		SceneCaptureComponent->ProjectionType = ECameraProjectionMode::Orthographic;
+		SceneCaptureComponent->OrthoWidth = BakerSettings->CameraOrthoWidth;
+	}
+	else
+	{
+		SceneCaptureComponent->ProjectionType = ECameraProjectionMode::Perspective;
+		SceneCaptureComponent->FOVAngle = BakerSettings->CameraFOV;
+	}
+	
+	const FMatrix SceneCaptureMatrix = FMatrix(FPlane(0, 0, 1, 0), FPlane(1, 0, 0, 0), FPlane(0, 1, 0, 0), FPlane(0, 0, 0, 1));
+	FMatrix ViewMatrix = SceneCaptureMatrix * BakerSettings->GetViewportMatrix().Inverse() * FRotationTranslationMatrix(BakerSettings->GetCameraRotation(), BakerSettings->GetCameraLocation());
+	SceneCaptureComponent->SetWorldLocationAndRotation(ViewMatrix.GetOrigin(), ViewMatrix.Rotator());
+	
+	SceneCaptureComponent->bUseCustomProjectionMatrix = true;
+	SceneCaptureComponent->CustomProjectionMatrix = BakerSettings->GetProjectionMatrix();
+	
+	SceneCaptureComponent->PrimitiveRenderMode = ESceneCapturePrimitiveRenderMode::PRM_UseShowOnlyList;
+	SceneCaptureComponent->ShowOnlyComponents.Empty(1);
+	SceneCaptureComponent->ShowOnlyComponents.Add(PreviewComponent);
+	
+	SceneCaptureComponent->CaptureScene();
+	
+	SceneCaptureComponent->TextureTarget = nullptr;
+	SceneCaptureComponent->UnregisterComponent();
+	
+	// Alpha from a scene capture is 1- so we need to invert
+	if ( SceneCaptureComponent->CaptureSource == ESceneCaptureSource::SCS_SceneColorHDR )
+	{
+		FCanvasTileItem TileItem(FVector2D(0, 0), FVector2D(RenderTarget->GetSurfaceWidth(), RenderTarget->GetSurfaceHeight()), FLinearColor::White);
+		TileItem.BlendMode = SE_BLEND_Opaque;
+		TileItem.BatchedElementParameters = new FBatchedElementNiagaraInvertColorChannel(0);
+		Canvas.DrawItem(TileItem);
 	}
-	return false;
+	Canvas.Flush_GameThread();
 }
 
-void FNiagaraBakerRenderer::AddReferencedObjects(FReferenceCollector& Collector)
+void FNiagaraBakerRenderer::RenderBufferVisualization(UTextureRenderTarget2D* RenderTarget, FName BufferVisualizationMode) const
 {
-	Collector.AddReferencedObject(SceneCaptureComponent);
+	UNiagaraBakerSettings* BakerSettings = GetBakerSettings();
+	if (!RenderTarget || !BakerSettings)
+	{
+		return;
+	}
+
+	const FIntRect ViewRect = FIntRect(0, 0, RenderTarget->GetSurfaceWidth(), RenderTarget->GetSurfaceHeight());
+	const float GammaCorrection = 1.0f;
+	const float WorldTime = GetWorldTime();
+	UWorld* World = GetWorld();
+
+	FCanvas Canvas(RenderTarget->GameThread_GetRenderTargetResource(), nullptr, WorldTime, FApp::GetDeltaTime(), WorldTime, GetFeatureLevel());
+	Canvas.Clear(FLinearColor::Black);
+
+	// Create View Family
+	FSceneViewFamilyContext ViewFamily(
+		FSceneViewFamily::ConstructionValues(RenderTarget->GameThread_GetRenderTargetResource(), World->Scene, FEngineShowFlags(ESFIM_Game))
+		.SetWorldTimes(WorldTime, FApp::GetDeltaTime(), WorldTime)
+		.SetGammaCorrection(GammaCorrection)
+	);
+	
+	ViewFamily.EngineShowFlags.SetScreenPercentage(false);
+	//ViewFamily.EngineShowFlags.DisableAdvancedFeatures();
+	//ViewFamily.EngineShowFlags.MotionBlur = 0;
+	//ViewFamily.EngineShowFlags.SetDistanceCulledPrimitives(true); // show distance culled objects
+	//ViewFamily.EngineShowFlags.SetPostProcessing(false);
+	
+	if (BufferVisualizationMode.IsValid())
+	{
+		ViewFamily.EngineShowFlags.SetPostProcessing(true);
+		ViewFamily.EngineShowFlags.SetVisualizeBuffer(true);
+		ViewFamily.EngineShowFlags.SetTonemapper(false);
+		ViewFamily.EngineShowFlags.SetScreenPercentage(false);
+	}
+	
+	FSceneViewInitOptions ViewInitOptions;
+	ViewInitOptions.SetViewRectangle(ViewRect);
+	ViewInitOptions.ViewFamily = &ViewFamily;
+	ViewInitOptions.ViewOrigin = BakerSettings->GetCameraLocation();
+	ViewInitOptions.ViewRotationMatrix = BakerSettings->GetViewMatrix();
+	ViewInitOptions.ProjectionMatrix = BakerSettings->GetProjectionMatrix();
+	ViewInitOptions.BackgroundColor = FLinearColor::Black;
+	if (BakerSettings->bRenderComponentOnly)
+	{
+		ViewInitOptions.ShowOnlyPrimitives.Emplace();
+		ViewInitOptions.ShowOnlyPrimitives->Add(PreviewComponent->ComponentId);
+	}
+	
+	FSceneView* NewView = new FSceneView(ViewInitOptions);
+	NewView->CurrentBufferVisualizationMode = BufferVisualizationMode;
+	ViewFamily.Views.Add(NewView);
+	
+	ViewFamily.SetScreenPercentageInterface(new FLegacyScreenPercentageDriver(ViewFamily, 1.0f));
+	
+	GetRendererModule().BeginRenderingViewFamily(&Canvas, &ViewFamily);
+
+	Canvas.Flush_GameThread();
 }
 
-FNiagaraBakerRenderer::ERenderType FNiagaraBakerRenderer::GetRenderType(FName SourceName, FName& OutName)
+void FNiagaraBakerRenderer::RenderDataInterface(UTextureRenderTarget2D* RenderTarget, FName BindingName) const
 {
-	using namespace NiagaraBakerRendererLocal;
+	const float WorldTime = GetWorldTime();
+	FCanvas Canvas(RenderTarget->GameThread_GetRenderTargetResource(), nullptr, WorldTime, FApp::GetDeltaTime(), WorldTime, GetFeatureLevel());
+	Canvas.Clear(FLinearColor::Black);
 
-	OutName = FName();
-	if (SourceName.IsNone())
+	ON_SCOPE_EXIT{ Canvas.Flush_GameThread(); };
+
+	// Gather data interface / attribute name
+	FString SourceString = BindingName.ToString();
+	int32 DotIndex;
+	if (!SourceString.FindLastChar('.', DotIndex))
 	{
-		return ERenderType::SceneCapture;
+		return;
 	}
 
-	FString SourceBindingString = SourceName.ToString();
-	TArray<FString> SplitNames;
-	SourceBindingString.ParseIntoArray(SplitNames, TEXT("."));
+	const FName DataInterfaceName = FName(SourceString.LeftChop(SourceString.Len() - DotIndex));
+	const FName VariableName = FName(SourceString.RightChop(DotIndex + 1));
 
-	if (!ensure(SplitNames.Num() > 0))
+	// Find data interface
+	FNiagaraSystemInstanceControllerPtr SystemInstanceController = PreviewComponent->GetSystemInstanceController();
+	if ( SystemInstanceController.IsValid() == false )
 	{
-		return ERenderType::None;
+		return;
 	}
 
-	// Scene Capture mode
-	if ( SplitNames[0] == STRING_SceneCaptureSource )
+	FNiagaraSystemInstance* SystemInstance = SystemInstanceController->GetSoloSystemInstance();
+	const FNiagaraSystemInstanceID SystemInstanceID = SystemInstance->GetId();
+	for (auto EmitterInstance : SystemInstance->GetEmitters())
 	{
-		if (!ensure(SplitNames.Num() == 2))
+		FNiagaraComputeExecutionContext* ExecContext = EmitterInstance->GetGPUContext();
+		if ( ExecContext == nullptr )
 		{
-			return ERenderType::None;
+			continue;
 		}
 
-		OutName = FName(SplitNames[1]);
-		return ERenderType::SceneCapture;
+		for (const FNiagaraVariableWithOffset& Variable : ExecContext->CombinedParamStore.ReadParameterVariables())
+		{
+			if (Variable.IsDataInterface())
+			{
+				if (Variable.GetName() == DataInterfaceName)
+				{
+					if ( UNiagaraDataInterface* DataInterface = ExecContext->CombinedParamStore.GetDataInterface(Variable.Offset) )
+					{
+						const FIntRect ViewRect(0, 0, RenderTarget->GetSurfaceWidth(), RenderTarget->GetSurfaceHeight());
+						DataInterface->RenderVariableToCanvas(SystemInstanceID, VariableName, &Canvas, ViewRect);
+						return;
+					}
+				}
+			}
+		}
 	}
+}
 
-	// Buffer Visualization Mode
-	if (SplitNames[0] == STRING_BufferVisualization)
+void FNiagaraBakerRenderer::RenderParticleAttribute(UTextureRenderTarget2D* RenderTarget, FName BindingName) const
+{
+	const float WorldTime = GetWorldTime();
+	FCanvas Canvas(RenderTarget->GameThread_GetRenderTargetResource(), nullptr, WorldTime, FApp::GetDeltaTime(), WorldTime, GetFeatureLevel());
+	Canvas.Clear(FLinearColor::Black);
+
+	ON_SCOPE_EXIT { Canvas.Flush_GameThread(); };
+
+	FString SourceString = BindingName.ToString();
+	int32 DotIndex;
+	if ( !SourceString.FindChar('.', DotIndex) )
 	{
-		if (!ensure(SplitNames.Num() == 2))
+		return;
+	}
+	
+	const FString EmitterName = SourceString.LeftChop(SourceString.Len() - DotIndex);
+	const FName AttributeName = FName(SourceString.RightChop(DotIndex + 1));
+	
+	FNiagaraSystemInstanceControllerPtr SystemInstanceController = PreviewComponent->GetSystemInstanceController();
+	if (!ensure(SystemInstanceController.IsValid()))
+	{
+		return;
+	}
+	
+	FNiagaraSystemInstance* SystemInstance = SystemInstanceController->GetSoloSystemInstance();
+	if ( !ensure(SystemInstance) )
+	{
+		return;
+	}
+	
+	for ( const auto& EmitterInstance : SystemInstance->GetEmitters() )
+	{
+		UNiagaraEmitter* NiagaraEmitter = EmitterInstance->GetCachedEmitter();
+		if ( !NiagaraEmitter || (NiagaraEmitter->GetUniqueEmitterName() != EmitterName) )
 		{
-			return ERenderType::None;
+			continue;
 		}
-		OutName = FName(SplitNames[1]);
-		return ERenderType::BufferVisualization;
+	
+		if (EmitterInstance->GetGPUContext() != nullptr)
+		{
+			return;
+		}
+	
+		const FNiagaraDataSet& ParticleDataSet = EmitterInstance->GetData();
+		const FNiagaraDataBuffer* ParticleDataBuffer = ParticleDataSet.GetCurrentData();
+		FNiagaraDataSetReaderInt32<int32> UniqueIDAccessor = FNiagaraDataSetAccessor<int32>::CreateReader(ParticleDataSet, FName("UniqueID"));
+		if ( !ParticleDataBuffer || !UniqueIDAccessor.IsValid() )
+		{
+			return;
+		}
+	
+		const int32 VariableIndex = ParticleDataSet.GetCompiledData().Variables.IndexOfByPredicate([&AttributeName](const FNiagaraVariable& Variable) { return Variable.GetName() == AttributeName; });
+		if (VariableIndex == INDEX_NONE)
+		{
+			return;
+		}
+		const FNiagaraVariableLayoutInfo& VariableInfo = ParticleDataSet.GetCompiledData().VariableLayouts[VariableIndex];
+	
+		float* FloatChannels[4];
+		FloatChannels[0] = (float*)ParticleDataBuffer->GetComponentPtrFloat(VariableInfo.FloatComponentStart);
+		FloatChannels[1] = VariableInfo.GetNumFloatComponents() > 1 ? (float*)ParticleDataBuffer->GetComponentPtrFloat(VariableInfo.FloatComponentStart + 1) : nullptr;
+		FloatChannels[2] = VariableInfo.GetNumFloatComponents() > 2 ? (float*)ParticleDataBuffer->GetComponentPtrFloat(VariableInfo.FloatComponentStart + 2) : nullptr;
+		FloatChannels[3] = VariableInfo.GetNumFloatComponents() > 3 ? (float*)ParticleDataBuffer->GetComponentPtrFloat(VariableInfo.FloatComponentStart + 3) : nullptr;
+	
+		const FIntPoint RenderTargetSize(RenderTarget->GetSurfaceWidth(), RenderTarget->GetSurfaceHeight());
+		const int32 ParticleBufferStore = RenderTargetSize.X * RenderTargetSize.Y;
+		for ( uint32 i=0; i < ParticleDataBuffer->GetNumInstances(); ++i )
+		{
+			const int32 UniqueID = UniqueIDAccessor[i];
+			if (UniqueID >= ParticleBufferStore)
+			{
+				continue;
+			}
+	
+			FLinearColor OutputColor;
+			OutputColor.R = FloatChannels[0] ? FloatChannels[0][i] : 0.0f;
+			OutputColor.G = FloatChannels[1] ? FloatChannels[1][i] : 0.0f;
+			OutputColor.B = FloatChannels[2] ? FloatChannels[2][i] : 0.0f;
+			OutputColor.A = FloatChannels[3] ? FloatChannels[3][i] : 0.0f;
+	
+			const int32 TexelX = UniqueID % RenderTargetSize.X;
+			const int32 TexelY = UniqueID / RenderTargetSize.X;
+			Canvas.DrawTile(TexelX, TexelY, 1.0f, 1.0f, 0.0f, 0.0f, 1.0f, 1.0f, OutputColor);
+		}
+
+		// We are done
+		break;
 	}
+}
 
-	// Emitter Data Interface
-	if (SplitNames[0] == STRING_EmitterDI)
+UWorld* FNiagaraBakerRenderer::GetWorld() const
+{
+	return PreviewComponent->GetWorld();
+}
+
+float FNiagaraBakerRenderer::GetWorldTime() const
+{
+	return PreviewComponent->GetDesiredAge();
+}
+
+ERHIFeatureLevel::Type FNiagaraBakerRenderer::GetFeatureLevel() const
+{
+	return PreviewComponent->GetWorld()->Scene->GetFeatureLevel();
+}
+
+UNiagaraSystem* FNiagaraBakerRenderer::GetNiagaraSystem() const
+{
+	return PreviewComponent->GetAsset();
+}
+
+void FNiagaraBakerRenderer::AddReferencedObjects(FReferenceCollector& Collector)
+{
+	Collector.AddReferencedObject(NiagaraSystem);
+	Collector.AddReferencedObject(PreviewComponent);
+	Collector.AddReferencedObject(SceneCaptureComponent);
+}
+
+FNiagaraBakerOutputRenderer* FNiagaraBakerRenderer::GetOutputRenderer(UClass* Class)
+{
+	if ( Class == UNiagaraBakerOutputTexture2D::StaticClass() )
 	{
-		OutName = FName(*SourceBindingString.RightChop(SplitNames[0].Len() + 1));
-		return ERenderType::DataInterface;
+		return new FNiagaraBakerRendererOutputTexture2D();
 	}
-
-	// Emitter Data Interface
-	if (SplitNames[0] == STRING_EmitterParticles)
+	else if (Class == UNiagaraBakerOutputVolumeTexture::StaticClass())
 	{
-		OutName = FName(*SourceBindingString.RightChop(SplitNames[0].Len() + 1));
-		return ERenderType::Particle;
+		return new FNiagaraBakerRendererOutputVolumeTexture();
 	}
-	return ERenderType::None;
+	return nullptr;
 }
 
-TArray<FName> FNiagaraBakerRenderer::GatherAllRenderOptions(UNiagaraSystem* NiagaraSystem)
+bool FNiagaraBakerRenderer::ExportImage(FStringView FilePath, FIntPoint ImageSize, TArrayView<FFloat16Color> ImageData)
 {
-	using namespace NiagaraBakerRendererLocal;
-
-	TArray<FName> RendererOptions;
-
-	// Put all scene capture options
-	static UEnum* SceneCaptureOptions = StaticEnum<ESceneCaptureSource>();
-	for ( int i=0; i < SceneCaptureOptions->GetMaxEnumValue(); ++i )
+	int32 LastDot = INDEX_NONE;
+	if ( FilePath.FindLastChar('.', LastDot) == false )
 	{
-		RendererOptions.Emplace(STRING_SceneCaptureSource + TEXT(".") + SceneCaptureOptions->GetNameStringByIndex(i));
+		return false;
 	}
 
-	// Gather all buffer visualization options
-	struct FIterator
+	const FStringView FileExtension = FilePath.Mid(LastDot);
+	const EImageFormat ImageFormat = ImageWrapperHelper::GetImageFormat(FileExtension);
+	if (ImageFormat == EImageFormat::Invalid)
 	{
-		TArray<FName>& OutAllOptions;
+		return false;
+	}
 
-		FIterator(TArray<FName>& InOutAllOptions)
-			: OutAllOptions(InOutAllOptions)
-		{}
+	IImageWrapperModule & ImageWrapperModule = FModuleManager::LoadModuleChecked<IImageWrapperModule>("ImageWrapper");
+	TSharedPtr<IImageWrapper> ImageWrapper = ImageWrapperModule.CreateImageWrapper(ImageFormat);
+	if ( ImageWrapper.IsValid() == false )
+	{
+		return false;
+	}
 
-		void ProcessValue(const FString& MaterialName, UMaterialInterface* Material, const FText& DisplayName)
+	if ( ImageFormat == EImageFormat::EXR || ImageFormat == EImageFormat::HDR )
+	{
+		TArray<FLinearColor> TempImageData;
+		TempImageData.Reserve(ImageData.Num());
+		for (const FFloat16Color& HalfColor : ImageData)
 		{
-			OutAllOptions.Emplace(STRING_BufferVisualization + TEXT(".") + MaterialName);
+			TempImageData.Emplace(HalfColor.GetFloats());
 		}
-	} Iterator(RendererOptions);
-	GetBufferVisualizationData().IterateOverAvailableMaterials(Iterator);
 
-	if (!NiagaraSystem)
-	{
-		return RendererOptions;
+		if (ImageWrapper->SetRaw(TempImageData.GetData(), TempImageData.Num() * TempImageData.GetTypeSize(), ImageSize.X, ImageSize.Y, ERGBFormat::RGBAF, 32) == false)
+		{
+			return false;
+		}
 	}
-
-	const TArray<TSharedRef<const FNiagaraEmitterCompiledData>>& AllEmitterCompiledData = NiagaraSystem->GetEmitterCompiledData();
-
-	// For each emitter find data interfaces / particle attributes we can render
-	for ( int32 EmitterIndex=0; EmitterIndex < NiagaraSystem->GetEmitterHandles().Num(); ++EmitterIndex )
+	else
 	{
-		const FNiagaraEmitterHandle& EmitterHandle = NiagaraSystem->GetEmitterHandle(EmitterIndex);
-		UNiagaraEmitter* EmitterInstance = EmitterHandle.GetInstance();
-		if (!EmitterHandle.IsValid() || !EmitterHandle.GetIsEnabled() || !EmitterInstance)
+		TArray<FColor> TempImageData;
+		TempImageData.Reserve(ImageData.Num());
+		for (const FFloat16Color& HalfColor : ImageData)
 		{
-			continue;
+			TempImageData.Add(HalfColor.GetFloats().ToFColor(true));
 		}
 
-		FString EmitterName = EmitterHandle.GetName().ToString();
-		FString EmitterPrefix = EmitterName + TEXT(".");
-
-		// Gather particle attributes
-		if ( ensure(AllEmitterCompiledData.IsValidIndex(EmitterIndex)) )
+		if (ImageWrapper->SetRaw(TempImageData.GetData(), TempImageData.Num() * TempImageData.GetTypeSize(), ImageSize.X, ImageSize.Y, ERGBFormat::BGRA, 8) == false)
 		{
-			const FNiagaraDataSetCompiledData& ParticleDataSet = AllEmitterCompiledData[EmitterIndex]->DataSetCompiledData;
-			for ( int32 iVariable=0; iVariable < ParticleDataSet.VariableLayouts.Num(); ++iVariable )
-			{
-				const FNiagaraVariable& Variable = ParticleDataSet.Variables[iVariable];
-				const FNiagaraVariableLayoutInfo& VariableLayout = ParticleDataSet.VariableLayouts[iVariable];
-				if ( VariableLayout.GetNumFloatComponents() > 0 )
-				{
-					RendererOptions.Emplace(STRING_EmitterParticles + TEXT(".") + EmitterName + TEXT(".") + *Variable.GetName().ToString());
-				}
-			}
+			return false;
 		}
-
-		// Gather data interfaces
-		EmitterInstance->ForEachScript(
-			[&](UNiagaraScript* NiagaraScript)
-			{
-				if (const FNiagaraScriptExecutionParameterStore* SrcStore = NiagaraScript->GetExecutionReadyParameterStore(EmitterInstance->SimTarget))
-				{
-					for (const FNiagaraVariableWithOffset& Variable : SrcStore->ReadParameterVariables())
-					{
-						if (Variable.IsDataInterface())
-						{
-							FString VariableName = Variable.GetName().ToString();
-							if (!VariableName.StartsWith(EmitterPrefix))
-							{
-								continue;
-							}
-
-							UNiagaraDataInterface* DataInterface = SrcStore->GetDataInterface(Variable.Offset);
-							check(DataInterface);
-
-							if ( DataInterface->CanRenderVariablesToCanvas() )
-							{
-								TArray<FNiagaraVariableBase> RendererableVariables;
-								DataInterface->GetCanvasVariables(RendererableVariables);
-								for (const FNiagaraVariableBase& RendererableVariable : RendererableVariables)
-								{
-									RendererOptions.Emplace(STRING_EmitterDI + TEXT(".") + VariableName + TEXT(".") + RendererableVariable.GetName().ToString());
-								}
-							}
-						}
-					}
-				}
-			}
-		);
 	}
 
-	return RendererOptions;
+	const TArray64<uint8> TempData = ImageWrapper->GetCompressed();
+	return FFileHelper::SaveArrayToFile(TempData, FilePath.GetData());
 }
diff --git a/Engine/Plugins/FX/Niagara/Source/NiagaraEditor/Private/NiagaraBakerRenderer.h b/Engine/Plugins/FX/Niagara/Source/NiagaraEditor/Private/NiagaraBakerRenderer.h
--- a/Engine/Plugins/FX/Niagara/Source/NiagaraEditor/Private/NiagaraBakerRenderer.h
+++ b/Engine/Plugins/FX/Niagara/Source/NiagaraEditor/Private/NiagaraBakerRenderer.h
@@ -4,16 +4,27 @@
 
 #include "CoreMinimal.h"
 #include "UObject/GCObject.h"
+#include "NiagaraSystem.h"
 
+class FNiagaraBakerRenderer;
 class UNiagaraComponent;
 class UNiagaraBakerSettings;
-class UNiagaraSystem;
 class UTextureRenderTarget2D;
+class USceneCaptureComponent2D;
+class FAdvancedPreviewScene;
 class FCanvas;
 
-struct FNiagaraBakerRenderer : FGCObject
+struct FNiagaraBakerOutputBinding
 {
-public:
+	FName	BindingName;
+	FText	MenuCategory;
+	FText	MenuEntry;
+};
+
+struct FNiagaraBakerOutputBindingHelper
+{
+	typedef TFunction<void(const FString& EmitterName, const FString& VariableName, UNiagaraDataInterface* DataInterface)> FEmitterDIFunction;
+
 	enum class ERenderType
 	{
 		None,
@@ -23,11 +34,73 @@ public:
 		Particle
 	};
 
+	static const FString STRING_SceneCaptureSource;
+	static const FString STRING_BufferVisualization;
+	static const FString STRING_EmitterDI;
+	static const FString STRING_EmitterParticles;
+
+	static ERenderType GetRenderType(FName BindingName, FName& OutName);
+
+	static void ForEachEmitterDataInterface(UNiagaraSystem* NiagaraSystem, FEmitterDIFunction Function);
+	static UNiagaraDataInterface* GetDataInterface(UNiagaraComponent* NiagaraComponent, FName DataInterfaceName);
+
+	static void GetSceneCaptureBindings(TArray<FNiagaraBakerOutputBinding>& OutBindings);
+	static void GetBufferVisualizationBindings(TArray<FNiagaraBakerOutputBinding>& OutBindings);
+	static void GetDataInterfaceBindingsForCanvas(TArray<FNiagaraBakerOutputBinding>& OutBindings, UNiagaraSystem* NiagaraSystem);
+	static void GetParticleAttributeBindings(TArray<FNiagaraBakerOutputBinding>& OutBindings, UNiagaraSystem* NiagaraSystem);
+};
+
+class FNiagaraBakerOutputRenderer
+{
+public:
+	FNiagaraBakerOutputRenderer() {}
+	virtual ~FNiagaraBakerOutputRenderer() {}
+
+	/**
+	Creates a list of all possible renderer bindings for the output, i.e. which sources you can pull from
+	*/
+	virtual TArray<FNiagaraBakerOutputBinding> GetRendererBindings(UNiagaraBakerOutput* InBakerOutput) const { return TArray<FNiagaraBakerOutputBinding>(); }
+
+	/**
+	Get the size we want to render the preview into.
+	Returning a negative or zero area will result in the preview being considered invalid
+	*/
+	virtual FIntPoint GetPreviewSize(UNiagaraBakerOutput* BakerOutput, FIntPoint InAvailableSize) const { return FIntPoint::ZeroValue; }
+
+	/**
+	Capture a preview of the baker output to the render target.
+	The render target will already be sized based on the result from GetPreviewSize
+	*/
+	virtual void RenderPreview(UNiagaraBakerOutput* BakerOutput, const FNiagaraBakerRenderer& BakerRenderer, UTextureRenderTarget2D* RenderTarget, TOptional<FString>& OutErrorString) const = 0;
+
+	virtual FIntPoint GetGeneratedSize(UNiagaraBakerOutput* BakerOutput, FIntPoint InAvailableSize) const { return FIntPoint::ZeroValue; }
+	virtual void RenderGenerated(UNiagaraBakerOutput* BakerOutput, const FNiagaraBakerRenderer& BakerRenderer, UTextureRenderTarget2D* RenderTarget, TOptional<FString>& OutErrorString) const = 0;
+
+	virtual bool BeginBake(UNiagaraBakerOutput* InBakerOutput) = 0;
+	virtual void BakeFrame(UNiagaraBakerOutput* InBakerOutput, int FrameIndex, const FNiagaraBakerRenderer& BakerRenderer) = 0;
+	virtual void EndBake(UNiagaraBakerOutput* InBakerOutput) = 0;
+};
+
+class FNiagaraBakerRenderer : FGCObject
+{
 public:
-	FNiagaraBakerRenderer();
+	FNiagaraBakerRenderer(UNiagaraSystem* NiagaraSystem);
 	virtual ~FNiagaraBakerRenderer();
-	bool RenderView(UNiagaraComponent* PreviewComponent, const UNiagaraBakerSettings* BakerSettings, float WorldTime, UTextureRenderTarget2D* RenderTarget, int32 iOutputTextureIndex) const;
-	bool RenderView(UNiagaraComponent* PreviewComponent, const UNiagaraBakerSettings* BakerSettings, float WorldTime, UTextureRenderTarget2D* RenderTarget, FCanvas* Canvas, int32 iOutputTextureIndex, FIntRect ViewRect) const;
+
+	void SetAbsoluteTime(float AbsoluteTime, bool bShouldTickComponent = true);
+
+	void RenderSceneCapture(UTextureRenderTarget2D* RenderTarget, ESceneCaptureSource CaptureSource) const;
+	void RenderBufferVisualization(UTextureRenderTarget2D* RenderTarget, FName BufferVisualizationMode = NAME_None) const;
+	void RenderDataInterface(UTextureRenderTarget2D* RenderTarget, FName BindingName) const;
+	void RenderParticleAttribute(UTextureRenderTarget2D* RenderTarget, FName BindingName) const;
+	
+	UWorld* GetWorld() const;
+	float GetWorldTime() const;
+	ERHIFeatureLevel::Type GetFeatureLevel() const;
+	UNiagaraComponent* GetPreviewComponent() const { return PreviewComponent; }
+	UNiagaraSystem* GetNiagaraSystem() const;
+	UNiagaraBakerSettings* GetBakerSettings() const { return NiagaraSystem->GetBakerSettings(); }
+	const UNiagaraBakerSettings* GetBakerGeneratedSettings() const { return NiagaraSystem->GetBakerGeneratedSettings(); }
 
 	// FGCObject Impl
 	virtual void AddReferencedObjects(FReferenceCollector& Collector) override;
@@ -37,10 +110,13 @@ public:
 	}
 	// FGCObject Impl
 
-	static ERenderType GetRenderType(FName SourceName, FName& OutName);
+	static FNiagaraBakerOutputRenderer* GetOutputRenderer(UClass* Class);
 
-	static TArray<FName> GatherAllRenderOptions(UNiagaraSystem* NiagaraSystem);
+	static bool ExportImage(FStringView FilePath, FIntPoint ImageSize, TArrayView<FFloat16Color> ImageData);
 
 private:
-	class USceneCaptureComponent2D* SceneCaptureComponent = nullptr;
+	UNiagaraSystem* NiagaraSystem = nullptr;
+	UNiagaraComponent* PreviewComponent = nullptr;
+	TSharedPtr<FAdvancedPreviewScene> AdvancedPreviewScene;
+	USceneCaptureComponent2D* SceneCaptureComponent = nullptr;
 };
diff --git a/Engine/Plugins/FX/Niagara/Source/NiagaraEditor/Private/NiagaraBakerRendererOutputTexture2D.cpp b/Engine/Plugins/FX/Niagara/Source/NiagaraEditor/Private/NiagaraBakerRendererOutputTexture2D.cpp
new file mode 100644
--- /dev/null
+++ b/Engine/Plugins/FX/Niagara/Source/NiagaraEditor/Private/NiagaraBakerRendererOutputTexture2D.cpp
@@ -0,0 +1,236 @@
+// Copyright Epic Games, Inc. All Rights Reserved.
+
+#include "NiagaraBakerRendererOutputTexture2D.h"
+#include "NiagaraBakerOutputTexture2D.h"
+
+#include "Engine/Texture2D.h"
+#include "Factories/Texture2dFactoryNew.h"
+#include "AssetToolsModule.h"
+#include "CanvasTypes.h"
+#include "CanvasItem.h"
+
+TArray<FNiagaraBakerOutputBinding> FNiagaraBakerRendererOutputTexture2D::GetRendererBindings(UNiagaraBakerOutput* InBakerOutput) const
+{
+	TArray<FNiagaraBakerOutputBinding> OutBindings;
+	FNiagaraBakerOutputBindingHelper::GetSceneCaptureBindings(OutBindings);
+	FNiagaraBakerOutputBindingHelper::GetBufferVisualizationBindings(OutBindings);
+	if ( UNiagaraSystem* NiagaraSystem = InBakerOutput->GetTypedOuter<UNiagaraSystem>() )
+	{
+		FNiagaraBakerOutputBindingHelper::GetDataInterfaceBindingsForCanvas(OutBindings, NiagaraSystem);
+		FNiagaraBakerOutputBindingHelper::GetParticleAttributeBindings(OutBindings, NiagaraSystem);
+	}
+	return OutBindings;
+}
+
+FIntPoint FNiagaraBakerRendererOutputTexture2D::GetPreviewSize(UNiagaraBakerOutput* InBakerOutput, FIntPoint InAvailableSize) const
+{
+	UNiagaraBakerOutputTexture2D* BakerOutput = CastChecked<UNiagaraBakerOutputTexture2D>(InBakerOutput);
+	return BakerOutput->FrameSize;
+}
+
+void FNiagaraBakerRendererOutputTexture2D::RenderPreview(UNiagaraBakerOutput* InBakerOutput, const FNiagaraBakerRenderer& BakerRenderer, UTextureRenderTarget2D* RenderTarget, TOptional<FString>& OutErrorString) const
+{
+	UNiagaraBakerOutputTexture2D* BakerOutput = Cast<UNiagaraBakerOutputTexture2D>(InBakerOutput);
+
+	FName SourceName;
+	const FNiagaraBakerOutputBindingHelper::ERenderType RenderType = FNiagaraBakerOutputBindingHelper::GetRenderType(BakerOutput->SourceBinding.SourceName, SourceName);
+	switch ( RenderType )
+	{
+		case FNiagaraBakerOutputBindingHelper::ERenderType::SceneCapture:
+		{
+			static UEnum* SceneCaptureSourceEnum = StaticEnum<ESceneCaptureSource>();
+			const int32 SceneCaptureSource = SceneCaptureSourceEnum->GetIndexByName(SourceName);
+			BakerRenderer.RenderSceneCapture(RenderTarget, SceneCaptureSource == INDEX_NONE ? ESceneCaptureSource::SCS_SceneColorHDR : ESceneCaptureSource(SceneCaptureSource));
+			break;
+		}
+		case FNiagaraBakerOutputBindingHelper::ERenderType::BufferVisualization:
+		{
+			BakerRenderer.RenderBufferVisualization(RenderTarget, SourceName);
+			break;
+		}
+		case FNiagaraBakerOutputBindingHelper::ERenderType::DataInterface:
+		{
+			BakerRenderer.RenderDataInterface(RenderTarget, SourceName);
+			break;
+		}
+		case FNiagaraBakerOutputBindingHelper::ERenderType::Particle:
+		{
+			BakerRenderer.RenderParticleAttribute(RenderTarget, SourceName);
+			break;
+		}
+		default:
+		{
+			break;
+		}
+	}
+}
+
+FIntPoint FNiagaraBakerRendererOutputTexture2D::GetGeneratedSize(UNiagaraBakerOutput* InBakerOutput, FIntPoint InAvailableSize) const
+{
+	UNiagaraBakerOutputTexture2D* BakerOutput = Cast<UNiagaraBakerOutputTexture2D>(InBakerOutput);
+	return BakerOutput ? BakerOutput->FrameSize : FIntPoint::ZeroValue;
+}
+
+void FNiagaraBakerRendererOutputTexture2D::RenderGenerated(UNiagaraBakerOutput* InBakerOutput, const FNiagaraBakerRenderer& BakerRenderer, UTextureRenderTarget2D* RenderTarget, TOptional<FString>& OutErrorString) const
+{
+	static FString AtlasNotFoundError(TEXT("Atlas texture not found.\nPlease bake an atlas to see the result."));
+
+	UNiagaraBakerOutputTexture2D* BakerOutput = Cast<UNiagaraBakerOutputTexture2D>(InBakerOutput);
+	UNiagaraBakerSettings* BakerGeneratedSettings = BakerOutput->GetTypedOuter<UNiagaraBakerSettings>();
+	if ( BakerOutput->bGenerateAtlas == false || BakerGeneratedSettings == nullptr)
+	{
+		OutErrorString = AtlasNotFoundError;
+		return;
+	}
+
+	UTexture2D* AtlasTexture = BakerOutput->GetAsset<UTexture2D>(BakerOutput->AtlasAssetPathFormat, 0);
+	if (AtlasTexture == nullptr)
+	{
+		OutErrorString = AtlasNotFoundError;
+		return;
+	}
+
+	const float WorldTime = BakerRenderer.GetWorldTime();
+	FCanvas Canvas(RenderTarget->GameThread_GetRenderTargetResource(), nullptr, WorldTime, FApp::GetDeltaTime(), WorldTime, BakerRenderer.GetFeatureLevel());
+
+	const FNiagaraBakerOutputFrameIndices FrameIndices = BakerGeneratedSettings->GetOutputFrameIndices(BakerOutput, WorldTime);
+
+	const FIntPoint TextureSize = FIntPoint(FMath::Max(AtlasTexture->GetSizeX(), 1), FMath::Max(AtlasTexture->GetSizeY(), 1));
+	const FIntPoint FrameSize = BakerOutput->FrameSize;
+	const FIntPoint FramesPerDimension = BakerGeneratedSettings->FramesPerDimension;
+	const FIntPoint FrameIndexA = FIntPoint(FrameIndices.FrameIndexA % FramesPerDimension.X, FrameIndices.FrameIndexA / FramesPerDimension.X);
+	const FIntPoint FrameIndexB = FIntPoint(FrameIndices.FrameIndexB % FramesPerDimension.X, FrameIndices.FrameIndexB / FramesPerDimension.X);
+	const FIntPoint FramePixelA = FIntPoint(FrameIndexA.X * FrameSize.X, FrameIndexA.Y * FrameSize.Y);
+	const FIntPoint FramePixelB = FIntPoint(FrameIndexB.X * FrameSize.X, FrameIndexB.Y * FrameSize.Y);
+
+	Canvas.DrawTile(
+		0.0f, 0.0f,
+		RenderTarget->GetSurfaceWidth(), RenderTarget->GetSurfaceHeight(),
+		(float(FramePixelA.X) + 0.5f) / float(TextureSize.X), (float(FramePixelA.Y) + 0.5f) / float(TextureSize.Y),
+		(float(FramePixelA.X + FrameSize.X) - 0.5f) / float(TextureSize.X), (float(FramePixelA.Y + FrameSize.Y) - 0.5f) / float(TextureSize.Y),
+		FLinearColor::White,
+		AtlasTexture->GetResource(),
+		SE_BLEND_Opaque
+	);
+
+	Canvas.Flush_GameThread();
+}
+
+bool FNiagaraBakerRendererOutputTexture2D::BeginBake(UNiagaraBakerOutput* InBakerOutput)
+{
+	UNiagaraBakerOutputTexture2D* BakerOutput = CastChecked<UNiagaraBakerOutputTexture2D>(InBakerOutput);
+	if ( !BakerOutput->bGenerateAtlas && !BakerOutput->bGenerateFrames && !BakerOutput->bExportFrames )
+	{
+		return false;
+	}
+
+	BakeRenderTarget = NewObject<UTextureRenderTarget2D>();
+	BakeRenderTarget->AddToRoot();
+	BakeRenderTarget->ClearColor = FLinearColor::Transparent;
+	BakeRenderTarget->TargetGamma = 1.0f;
+	BakeRenderTarget->InitCustomFormat(BakerOutput->FrameSize.X, BakerOutput->FrameSize.Y, PF_FloatRGBA, false);
+
+	if (BakerOutput->bGenerateAtlas)
+	{
+		BakeAtlasTextureData.Empty();
+		BakeAtlasTextureData.SetNumZeroed(BakerOutput->AtlasTextureSize.X * BakerOutput->AtlasTextureSize.Y);
+	}
+
+	return true;
+}
+
+void FNiagaraBakerRendererOutputTexture2D::BakeFrame(UNiagaraBakerOutput* InBakerOutput, int FrameIndex, const FNiagaraBakerRenderer& BakerRenderer)
+{
+	UNiagaraBakerOutputTexture2D* BakerOutput = CastChecked<UNiagaraBakerOutputTexture2D>(InBakerOutput);
+	UNiagaraBakerSettings* BakerSettings = BakerRenderer.GetBakerSettings();
+
+	TOptional<FString> ErrorString;
+	RenderPreview(InBakerOutput, BakerRenderer, BakeRenderTarget, ErrorString);
+
+	TArray<FFloat16Color> FrameResults;
+	BakeRenderTarget->GameThread_GetRenderTargetResource()->ReadFloat16Pixels(FrameResults);
+
+	const FIntPoint ViewportOffset = FIntPoint(
+		BakerOutput->FrameSize.X * (FrameIndex % BakerSettings->FramesPerDimension.X),
+		BakerOutput->FrameSize.Y * (FrameIndex / BakerSettings->FramesPerDimension.X)
+	);
+
+	// Are we generating an atlas
+	if ( BakerOutput->bGenerateAtlas )
+	{
+		for (int y = 0; y < BakerOutput->FrameSize.Y; ++y)
+		{
+			FFloat16Color* SrcPixel = FrameResults.GetData() + (y * BakerOutput->FrameSize.X);
+			FFloat16Color* DstPixel = BakeAtlasTextureData.GetData() + ViewportOffset.X + ((y + ViewportOffset.Y) * BakerOutput->AtlasTextureSize.X);
+			FMemory::Memcpy(DstPixel, SrcPixel, sizeof(FFloat16Color) * BakerOutput->FrameSize.X);
+		}
+	}
+
+	// Are generating frames?
+	if ( BakerOutput->bGenerateFrames )
+	{
+		const FString AssetFullName = BakerOutput->GetAssetPath(BakerOutput->FramesAssetPathFormat, FrameIndex);
+		const FString AssetName = FString(FPathViews::GetCleanFilename(AssetFullName));
+		const FString PackagePath = FString(FPathViews::GetPath(AssetFullName));
+
+		IAssetTools& AssetTools = FModuleManager::Get().LoadModuleChecked<FAssetToolsModule>("AssetTools").Get();
+		UTexture2DFactoryNew* TextureFactory = NewObject<UTexture2DFactoryNew>();
+		UTexture2D* OutputTexture = Cast<UTexture2D>(AssetTools.CreateAsset(AssetName, PackagePath, UTexture2D::StaticClass(), TextureFactory));
+		if (OutputTexture)
+		{
+			const bool bIsPoT = FMath::IsPowerOfTwo(BakerOutput->FrameSize.X) && FMath::IsPowerOfTwo(BakerOutput->FrameSize.Y);
+
+			OutputTexture->Source.Init(BakerOutput->FrameSize.X, BakerOutput->FrameSize.Y, 1, 1, TSF_RGBA16F, (const uint8*)(FrameResults.GetData()));
+			OutputTexture->PowerOfTwoMode = ETexturePowerOfTwoSetting::None;
+			OutputTexture->MipGenSettings = bIsPoT ? TextureMipGenSettings::TMGS_FromTextureGroup : TextureMipGenSettings::TMGS_NoMipmaps;
+			OutputTexture->AddressX = TextureAddress::TA_Clamp;
+			OutputTexture->AddressY = TextureAddress::TA_Clamp;
+			OutputTexture->UpdateResource();
+			OutputTexture->PostEditChange();
+			OutputTexture->MarkPackageDirty();
+		}
+	}
+
+	// Are we exporting frames
+	if ( BakerOutput->bExportFrames )
+	{
+		const FString FilePath = BakerOutput->GetExportPath(BakerOutput->FramesExportPathFormat, FrameIndex);
+		FNiagaraBakerRenderer::ExportImage(FilePath, BakerOutput->FrameSize, FrameResults);
+	}
+}
+
+void FNiagaraBakerRendererOutputTexture2D::EndBake(UNiagaraBakerOutput* InBakerOutput)
+{
+	UNiagaraBakerOutputTexture2D* BakerOutput = CastChecked<UNiagaraBakerOutputTexture2D>(InBakerOutput);
+
+	if ( BakerOutput->bGenerateAtlas )
+	{
+		const FString AssetFullName = BakerOutput->GetAssetPath(BakerOutput->AtlasAssetPathFormat, 0);
+		const FString AssetName = FString(FPathViews::GetCleanFilename(AssetFullName));
+		const FString PackagePath = FString(FPathViews::GetPath(AssetFullName));
+
+		IAssetTools& AssetTools = FModuleManager::Get().LoadModuleChecked<FAssetToolsModule>("AssetTools").Get();
+		UTexture2DFactoryNew* TextureFactory = NewObject<UTexture2DFactoryNew>();
+		UTexture2D* OutputTexture = Cast<UTexture2D>(AssetTools.CreateAsset(AssetName, PackagePath, UTexture2D::StaticClass(), TextureFactory));
+		if (OutputTexture)
+		{
+			const bool bIsPoT = FMath::IsPowerOfTwo(BakerOutput->AtlasTextureSize.X) && FMath::IsPowerOfTwo(BakerOutput->AtlasTextureSize.Y);
+
+			OutputTexture->Source.Init(BakerOutput->AtlasTextureSize.X, BakerOutput->AtlasTextureSize.Y, 1, 1, TSF_RGBA16F, (const uint8*)(BakeAtlasTextureData.GetData()));
+			OutputTexture->PowerOfTwoMode = ETexturePowerOfTwoSetting::None;
+			OutputTexture->MipGenSettings = bIsPoT ? TextureMipGenSettings::TMGS_FromTextureGroup : TextureMipGenSettings::TMGS_NoMipmaps;
+			OutputTexture->AddressX = TextureAddress::TA_Clamp;
+			OutputTexture->AddressY = TextureAddress::TA_Clamp;
+			OutputTexture->UpdateResource();
+			OutputTexture->PostEditChange();
+			OutputTexture->MarkPackageDirty();
+		}
+	}
+
+	// Clean up temporary storage
+	BakeRenderTarget->RemoveFromRoot();
+	BakeRenderTarget->MarkAsGarbage();
+	BakeRenderTarget = nullptr;
+
+	BakeAtlasTextureData.Empty();
+}
diff --git a/Engine/Plugins/FX/Niagara/Source/NiagaraEditor/Private/NiagaraBakerRendererOutputTexture2D.h b/Engine/Plugins/FX/Niagara/Source/NiagaraEditor/Private/NiagaraBakerRendererOutputTexture2D.h
new file mode 100644
--- /dev/null
+++ b/Engine/Plugins/FX/Niagara/Source/NiagaraEditor/Private/NiagaraBakerRendererOutputTexture2D.h
@@ -0,0 +1,27 @@
+// Copyright Epic Games, Inc. All Rights Reserved.
+
+#pragma once
+
+#include "NiagaraBakerRenderer.h"
+
+class UTextureRenderTarget2D;
+
+class FNiagaraBakerRendererOutputTexture2D : public FNiagaraBakerOutputRenderer
+{
+public:
+	virtual TArray<FNiagaraBakerOutputBinding> GetRendererBindings(UNiagaraBakerOutput* InBakerOutput) const override;
+
+	virtual FIntPoint GetPreviewSize(UNiagaraBakerOutput* InBakerOutput, FIntPoint InAvailableSize) const override;
+	virtual void RenderPreview(UNiagaraBakerOutput* InBakerOutput, const FNiagaraBakerRenderer& BakerRenderer, UTextureRenderTarget2D* RenderTarget, TOptional<FString>& OutErrorString) const override;
+
+	virtual FIntPoint GetGeneratedSize(UNiagaraBakerOutput* InBakerOutput, FIntPoint InAvailableSize) const override;
+	virtual void RenderGenerated(UNiagaraBakerOutput* InBakerOutput, const FNiagaraBakerRenderer& BakerRenderer, UTextureRenderTarget2D* RenderTarget, TOptional<FString>& OutErrorString) const override;
+
+	virtual bool BeginBake(UNiagaraBakerOutput* InBakerOutput) override;
+	virtual void BakeFrame(UNiagaraBakerOutput* InBakerOutput, int FrameIndex, const FNiagaraBakerRenderer& BakerRenderer) override;
+	virtual void EndBake(UNiagaraBakerOutput* InBakerOutput) override;
+
+protected:
+	UTextureRenderTarget2D* BakeRenderTarget = nullptr;
+	TArray<FFloat16Color> BakeAtlasTextureData;
+};
diff --git a/Engine/Plugins/FX/Niagara/Source/NiagaraEditor/Private/NiagaraBakerRendererOutputVolumeTexture.cpp b/Engine/Plugins/FX/Niagara/Source/NiagaraEditor/Private/NiagaraBakerRendererOutputVolumeTexture.cpp
new file mode 100644
--- /dev/null
+++ b/Engine/Plugins/FX/Niagara/Source/NiagaraEditor/Private/NiagaraBakerRendererOutputVolumeTexture.cpp
@@ -0,0 +1,608 @@
+// Copyright Epic Games, Inc. All Rights Reserved.
+
+#include "NiagaraBakerRendererOutputVolumeTexture.h"
+#include "NiagaraBakerOutputVolumeTexture.h"
+
+#include "NiagaraComponent.h"
+#include "NiagaraSystemInstance.h"
+#include "NiagaraComputeExecutionContext.h"
+#include "NiagaraDataInterfaceGrid3DCollection.h"
+#include "NiagaraDataInterfaceRenderTargetVolume.h"
+#include "NiagaraBatchedElements.h"
+
+#include "Engine/Canvas.h"
+#include "Engine/VolumeTexture.h"
+#include "Misc/PathViews.h"
+#include "Modules/ModuleManager.h"
+#include "Factories/VolumeTextureFactory.h"
+#include "IAssetTools.h"
+#include "AssetToolsModule.h"
+
+struct FVolumeDataInterfaceHelper
+{
+	UNiagaraComponent*										NiagaraComponent = nullptr;
+	FNiagaraSystemInstance*									SystemInstance = nullptr;
+	TArray<FString>											DataInterfacePath;
+
+	UNiagaraDataInterfaceGrid3DCollection*					Grid3DDataInterface = nullptr;
+	FNiagaraDataInterfaceProxyGrid3DCollectionProxy*		Grid3DProxy = nullptr;
+	FGrid3DCollectionRWInstanceData_GameThread*				Grid3DInstanceData_GameThread = nullptr;
+	FName													Grid3DAttributeName;
+	int32													Grid3DVariableIndex = INDEX_NONE;
+	int32													Grid3DAttributeStart = INDEX_NONE;
+	int32													Grid3DAttributeChannels = 0;
+	FIntVector												Grid3DTextureSize = FIntVector::ZeroValue;
+
+	UNiagaraDataInterfaceRenderTargetVolume*				VolumeRenderTargetDataInterface = nullptr;
+	FNiagaraDataInterfaceProxyRenderTargetVolumeProxy*		VolumeRenderTargetProxy = nullptr;
+	FRenderTargetVolumeRWInstanceData_GameThread*			VolumeRenderTargetInstanceData_GameThread = nullptr;
+
+	bool Initialize(UNiagaraBakerOutputVolumeTexture* OutputVolumeTexture, UNiagaraComponent* InNiagaraComponent)
+	{
+		NiagaraComponent = InNiagaraComponent;
+
+		OutputVolumeTexture->SourceBinding.SourceName.ToString().ParseIntoArray(DataInterfacePath, TEXT("."));
+		if (DataInterfacePath.Num() < 2)
+		{
+			return false;
+		}
+
+		const FName DataInterfaceName(DataInterfacePath[0] + "." + DataInterfacePath[1]);
+		UNiagaraDataInterface* DataInterface = FNiagaraBakerOutputBindingHelper::GetDataInterface(NiagaraComponent, DataInterfaceName);
+		if (DataInterface == nullptr)
+		{
+			return false;
+		}
+
+		// Guaranteed since we got a data interface
+		SystemInstance = NiagaraComponent->GetSystemInstanceController()->GetSoloSystemInstance();
+
+		// Render Target Volume
+		if (DataInterface->IsA<UNiagaraDataInterfaceRenderTargetVolume>())
+		{
+			VolumeRenderTargetDataInterface = CastChecked<UNiagaraDataInterfaceRenderTargetVolume>(DataInterface);
+			VolumeRenderTargetProxy = static_cast<FNiagaraDataInterfaceProxyRenderTargetVolumeProxy*>(VolumeRenderTargetDataInterface->GetProxy());
+			VolumeRenderTargetInstanceData_GameThread = reinterpret_cast<FRenderTargetVolumeRWInstanceData_GameThread*>(SystemInstance->FindDataInterfaceInstanceData(VolumeRenderTargetDataInterface));
+			if (VolumeRenderTargetInstanceData_GameThread == nullptr)
+			{
+				return false;
+			}
+		}
+		// Grid 3D
+		else if ( DataInterface->IsA<UNiagaraDataInterfaceGrid3DCollection>() )
+		{
+			Grid3DDataInterface = CastChecked<UNiagaraDataInterfaceGrid3DCollection>(DataInterface);
+			Grid3DProxy = static_cast<FNiagaraDataInterfaceProxyGrid3DCollectionProxy*>(Grid3DDataInterface->GetProxy());
+			Grid3DInstanceData_GameThread = reinterpret_cast<FGrid3DCollectionRWInstanceData_GameThread*>(SystemInstance->FindDataInterfaceInstanceData(Grid3DDataInterface));
+			if (Grid3DInstanceData_GameThread == nullptr)
+			{
+				return false;
+			}
+
+			if (DataInterfacePath.Num() != 3)
+			{
+				// Perhaps a path to pull all attributes, i.e. whole texture?
+				return false;
+			}
+
+			Grid3DAttributeName = FName(DataInterfacePath[2]);
+			Grid3DVariableIndex = Grid3DInstanceData_GameThread->Vars.IndexOfByPredicate([&](const FNiagaraVariableBase& VariableBase) { return VariableBase.GetName() == Grid3DAttributeName; });
+			if (Grid3DVariableIndex == INDEX_NONE)
+			{
+				return false;
+			}
+			Grid3DAttributeStart	= Grid3DInstanceData_GameThread->Offsets[Grid3DVariableIndex];
+			Grid3DAttributeChannels	= Grid3DInstanceData_GameThread->Vars[Grid3DVariableIndex].GetType().GetSize() / sizeof(float);
+			Grid3DTextureSize.X		= Grid3DInstanceData_GameThread->NumCells.X * Grid3DInstanceData_GameThread->NumTiles.X;
+			Grid3DTextureSize.Y		= Grid3DInstanceData_GameThread->NumCells.Y * Grid3DInstanceData_GameThread->NumTiles.Y;
+			Grid3DTextureSize.Z		= Grid3DInstanceData_GameThread->NumCells.Z * Grid3DInstanceData_GameThread->NumTiles.Z;
+		}
+		// Unsupported type
+		else
+		{
+			return false;
+		}
+
+		return true;
+	}
+};
+
+TArray<FNiagaraBakerOutputBinding> FNiagaraBakerRendererOutputVolumeTexture::GetRendererBindings(UNiagaraBakerOutput* InBakerOutput) const
+{
+	TArray<FNiagaraBakerOutputBinding> OutBindings;
+	if (UNiagaraSystem* NiagaraSystem = InBakerOutput->GetTypedOuter<UNiagaraSystem>())
+	{
+		FNiagaraBakerOutputBindingHelper::ForEachEmitterDataInterface(
+			NiagaraSystem,
+			[&](const FString& EmitterName, const FString& VariableName, UNiagaraDataInterface* DataInterface)
+			{
+				if ( UNiagaraDataInterfaceGrid3DCollection* Grid3D = Cast<UNiagaraDataInterfaceGrid3DCollection>(DataInterface) )
+				{
+					TArray<FNiagaraVariableBase> GridVariables;
+					TArray<uint32> GridVariableOffsets;
+					int32 NumAttribChannelsFound;
+				 	Grid3D->FindAttributes(GridVariables, GridVariableOffsets, NumAttribChannelsFound);
+
+					for ( const FNiagaraVariableBase& GridVariable : GridVariables)
+					{
+						const FString GridVariableString = GridVariable.GetName().ToString();
+
+						FNiagaraBakerOutputBinding& NewBinding = OutBindings.AddDefaulted_GetRef();
+						NewBinding.BindingName = FName(EmitterName + "." + VariableName + "." + GridVariableString);
+						NewBinding.MenuCategory = FText::FromString(EmitterName + " " + TEXT("Grid3DCollection"));
+						NewBinding.MenuEntry = FText::FromString(VariableName + "." + GridVariableString);
+					}
+				}
+				else if ( UNiagaraDataInterfaceRenderTargetVolume* VolumeRT = Cast<UNiagaraDataInterfaceRenderTargetVolume>(DataInterface) )
+				{
+					FNiagaraBakerOutputBinding& NewBinding = OutBindings.AddDefaulted_GetRef();
+					NewBinding.BindingName	= FName(EmitterName + "." + VariableName);
+					NewBinding.MenuCategory	= FText::FromString(EmitterName + " " + TEXT("VolumeRenderTarget"));
+					NewBinding.MenuEntry	= FText::FromString(VariableName);
+				}
+			}
+		);
+	}
+	return OutBindings;
+}
+
+FIntPoint FNiagaraBakerRendererOutputVolumeTexture::GetPreviewSize(UNiagaraBakerOutput* InBakerOutput, FIntPoint InAvailableSize) const
+{
+	return InAvailableSize;
+}
+
+void FNiagaraBakerRendererOutputVolumeTexture::RenderPreview(UNiagaraBakerOutput* InBakerOutput, const FNiagaraBakerRenderer& BakerRenderer, UTextureRenderTarget2D* RenderTarget, TOptional<FString>& OutErrorString) const
+{
+	UNiagaraBakerOutputVolumeTexture* OutputVolumeTexture = CastChecked<UNiagaraBakerOutputVolumeTexture>(InBakerOutput);
+
+	const float WorldTime = BakerRenderer.GetWorldTime();
+	FCanvas Canvas(RenderTarget->GameThread_GetRenderTargetResource(), nullptr, WorldTime, FApp::GetDeltaTime(), WorldTime, BakerRenderer.GetFeatureLevel());
+	Canvas.Clear(FLinearColor::Black);
+	ON_SCOPE_EXIT{ Canvas.Flush_GameThread(); };
+
+	FVolumeDataInterfaceHelper DataInterface;
+	if (DataInterface.Initialize(OutputVolumeTexture, BakerRenderer.GetPreviewComponent()) == false)
+	{
+		OutErrorString = TEXT("Could not find data to preview from.\nPlease ensure binding is set to a valid source.");
+		return;
+	}
+
+	// Volume Render Target
+	if ( DataInterface.VolumeRenderTargetDataInterface != nullptr )
+	{
+		const int32 NumSlices		= DataInterface.VolumeRenderTargetInstanceData_GameThread->Size.Z;
+		const int32 SlicePerAxis	= FMath::CeilToInt(FMath::Sqrt((float)NumSlices));
+		const int32 TileWidth		= RenderTarget->GetSurfaceWidth()  / SlicePerAxis;
+		const int32 TileHeight		= RenderTarget->GetSurfaceHeight() / SlicePerAxis;
+		if (TileWidth <= 0 && TileHeight <= 0)
+		{
+			return;
+		}
+
+		const int32 AttributeChannels[] = {0, 1, 2, 3};
+		const FIntVector TextureSize = DataInterface.VolumeRenderTargetInstanceData_GameThread->Size;
+		const FVector2f TileUVOffset = FVector2f(0.5f / float(TextureSize.X), 0.5f / float(TextureSize.Y));
+		const FVector2f TileUVSize = FVector2f(float(TextureSize.X - 1) / float(TextureSize.X), float(TextureSize.Y - 1) / float(TextureSize.Y));
+
+		for (int32 Slice=0; Slice < NumSlices; ++Slice)
+		{
+			const float SliceW = (float(Slice) + 0.5f) / float(NumSlices - 1);
+			const FVector3f AttributeUVs[] =
+			{
+				FVector3f(TileUVOffset.X, TileUVOffset.Y, SliceW),
+				FVector3f(TileUVOffset.X, TileUVOffset.Y, SliceW),
+				FVector3f(TileUVOffset.X, TileUVOffset.Y, SliceW),
+				FVector3f(TileUVOffset.X, TileUVOffset.Y, SliceW),
+			};
+			const int32 TileX = (Slice % SlicePerAxis) * TileWidth;
+			const int32 TileY = (Slice / SlicePerAxis) * TileHeight;
+
+			FCanvasTileItem TileItem(FVector2D(TileX, TileY), GWhiteTexture, FVector2D(TileWidth, TileHeight), FVector2D(0.0, 1.0f), FVector2D(1.0, 0.0f), FLinearColor::White);
+			TileItem.BlendMode = SE_BLEND_Opaque;
+			TileItem.BatchedElementParameters = new FBatchedElementNiagaraVolumeAttribute(
+				FVector2f(1.0f, 1.0f),
+				MakeArrayView(AttributeUVs),
+				MakeArrayView(AttributeChannels),
+				[Proxy_RT=DataInterface.VolumeRenderTargetProxy, SystemInstanceID_RT= DataInterface.SystemInstance->GetId()](FRHITexture*& OutTexture, FRHISamplerState*& OutSamplerState)
+				{
+					if ( const FRenderTargetVolumeRWInstanceData_RenderThread* InstanceData_RT = Proxy_RT->SystemInstancesToProxyData_RT.Find(SystemInstanceID_RT) )
+					{
+						OutTexture = InstanceData_RT->TextureRHI;
+						OutSamplerState = InstanceData_RT->SamplerStateRHI;
+			//				OutSamplerState = TStaticSamplerState<SF_Bilinear>::GetRHI();
+					}
+				}
+			);
+			Canvas.DrawItem(TileItem);
+		}
+	}
+	// Grid Collection
+	else if ( DataInterface.Grid3DDataInterface != nullptr )
+	{
+		const FVector2f TileUVSize = FVector2f(
+			float(DataInterface.Grid3DInstanceData_GameThread->NumCells.X - 1) / float(DataInterface.Grid3DTextureSize.X),
+			float(DataInterface.Grid3DInstanceData_GameThread->NumCells.Y - 1) / float(DataInterface.Grid3DTextureSize.Y)
+		);
+		TArray<FVector3f, TInlineAllocator<4>> AttributeUVs;
+		for ( int32 i=0; i < DataInterface.Grid3DAttributeChannels; ++i )
+		{
+			const int32 AttributeIndex = DataInterface.Grid3DAttributeStart + i;
+			const FIntVector AttributeTileIndex(
+				AttributeIndex % DataInterface.Grid3DInstanceData_GameThread->NumTiles.X,
+				(AttributeIndex / DataInterface.Grid3DInstanceData_GameThread->NumTiles.X) % DataInterface.Grid3DInstanceData_GameThread->NumTiles.Y,
+				AttributeIndex / (DataInterface.Grid3DInstanceData_GameThread->NumTiles.X * DataInterface.Grid3DInstanceData_GameThread->NumTiles.Y)
+			);
+			FVector3f& AttributeUV = AttributeUVs.AddDefaulted_GetRef();
+			AttributeUV.X = (float(AttributeTileIndex.X * DataInterface.Grid3DInstanceData_GameThread->NumCells.X) + 0.5f) / DataInterface.Grid3DTextureSize.X;
+			AttributeUV.Y = (float(AttributeTileIndex.Y * DataInterface.Grid3DInstanceData_GameThread->NumCells.Y) + 0.5f) / DataInterface.Grid3DTextureSize.Y;
+			AttributeUV.Z = (float(AttributeTileIndex.Z * DataInterface.Grid3DInstanceData_GameThread->NumCells.Z) + 0.5f) / DataInterface.Grid3DTextureSize.Z;
+		}
+
+		const int32 NumSlices		= DataInterface.Grid3DInstanceData_GameThread->NumCells.Z;
+		const int32 SlicePerAxis	= FMath::CeilToInt(FMath::Sqrt((float)NumSlices));
+		const int32 TileWidth		= RenderTarget->GetSurfaceWidth()  / SlicePerAxis;
+		const int32 TileHeight		= RenderTarget->GetSurfaceHeight() / SlicePerAxis;
+		if (TileWidth > 0 && TileHeight > 0)
+		{
+			for (int32 Slice=0; Slice < NumSlices; ++Slice)
+			{
+				const int32 TileX = (Slice % SlicePerAxis) * TileWidth;
+				const int32 TileY = (Slice / SlicePerAxis) * TileHeight;
+
+				TArray<FVector3f, TInlineAllocator<4>> SliceAttributeUVs;
+				for ( FVector3f SliceUV : AttributeUVs )
+				{
+					SliceUV.Z += float(Slice) / float(DataInterface.Grid3DTextureSize.Z);
+					SliceAttributeUVs.Add(SliceUV);
+				}
+
+				FCanvasTileItem TileItem(FVector2D(TileX, TileY), GWhiteTexture, FVector2D(TileWidth, TileHeight), FVector2D(0.0, 1.0f), FVector2D(1.0, 0.0f), FLinearColor::White);
+				TileItem.BlendMode = SE_BLEND_Opaque;
+				TileItem.BatchedElementParameters = new FBatchedElementNiagaraVolumeAttribute(
+					TileUVSize,
+					MakeArrayView(SliceAttributeUVs),
+					MakeArrayView<int32>(nullptr, 0),
+					[RT_Proxy=DataInterface.Grid3DProxy, RT_SystemInstanceID=DataInterface.SystemInstance->GetId()](FRHITexture*& OutTexture, FRHISamplerState*& OutSamplerState)
+					{
+						if ( const FGrid3DCollectionRWInstanceData_RenderThread* RT_InstanceData = RT_Proxy->SystemInstancesToProxyData_RT.Find(RT_SystemInstanceID) )
+						{
+							if ( RT_InstanceData->CurrentData != nullptr )
+							{
+								OutTexture = RT_InstanceData->CurrentData->GridBuffer.Buffer;
+								OutSamplerState = TStaticSamplerState<SF_Bilinear>::GetRHI();
+							}
+						}
+					}
+				);
+				Canvas.DrawItem(TileItem);
+			}
+		}
+	}
+}
+
+FIntPoint FNiagaraBakerRendererOutputVolumeTexture::GetGeneratedSize(UNiagaraBakerOutput* InBakerOutput, FIntPoint InAvailableSize) const
+{
+	UNiagaraBakerOutputVolumeTexture* OutputVolumeTexture = CastChecked<UNiagaraBakerOutputVolumeTexture>(InBakerOutput);
+	UNiagaraBakerSettings* BakerSettings = OutputVolumeTexture->GetTypedOuter<UNiagaraBakerSettings>();
+
+	FIntPoint RequiredSize = FIntPoint::ZeroValue;
+	if ( OutputVolumeTexture->bGenerateAtlas )
+	{
+		if ( UVolumeTexture* AtlasVolumeTexture = OutputVolumeTexture->GetAsset<UVolumeTexture>(OutputVolumeTexture->AtlasAssetPathFormat, 0) )
+		{
+			const FIntPoint FrameSize(
+				AtlasVolumeTexture->GetSizeX() / BakerSettings->FramesPerDimension.X,
+				AtlasVolumeTexture->GetSizeY() / BakerSettings->FramesPerDimension.Y
+			);
+			const float Ratio = FMath::Min(float(InAvailableSize.X) / float(FrameSize.X), float(InAvailableSize.Y) / float(FrameSize.Y));
+			RequiredSize.X = FMath::RoundToInt(FrameSize.X * Ratio);
+			RequiredSize.Y = FMath::RoundToInt(FrameSize.Y * Ratio);
+		}
+	}
+
+	return RequiredSize;
+}
+
+void FNiagaraBakerRendererOutputVolumeTexture::RenderGenerated(UNiagaraBakerOutput* InBakerOutput, const FNiagaraBakerRenderer& BakerRenderer, UTextureRenderTarget2D* RenderTarget, TOptional<FString>& OutErrorString) const
+{
+	static FString AtlasNotFoundError(TEXT("Atlas texture not found.\nPlease bake an atlas to see the result."));
+
+	UNiagaraBakerOutputVolumeTexture* OutputVolumeTexture = CastChecked<UNiagaraBakerOutputVolumeTexture>(InBakerOutput);
+	UNiagaraBakerSettings* BakerGeneratedSettings = OutputVolumeTexture->GetTypedOuter<UNiagaraBakerSettings>();
+
+	if (OutputVolumeTexture->bGenerateAtlas == false || !BakerGeneratedSettings)
+	{
+		OutErrorString = AtlasNotFoundError;
+		return;
+	}
+
+	UVolumeTexture* AtlasVolumeTexture = OutputVolumeTexture->GetAsset<UVolumeTexture>(OutputVolumeTexture->AtlasAssetPathFormat, 0);
+	if (AtlasVolumeTexture == nullptr)
+	{
+		OutErrorString = AtlasNotFoundError;
+		return;
+	}
+
+	const float WorldTime = BakerRenderer.GetWorldTime();
+	FCanvas Canvas(RenderTarget->GameThread_GetRenderTargetResource(), nullptr, WorldTime, FApp::GetDeltaTime(), WorldTime, BakerRenderer.GetFeatureLevel());
+	Canvas.Clear(FLinearColor::Black);
+
+	const FIntVector AtlasTextureSize(AtlasVolumeTexture->GetSizeX(), AtlasVolumeTexture->GetSizeY(), AtlasVolumeTexture->GetSizeZ());
+	if (AtlasTextureSize.X > 0 && AtlasTextureSize.Y > 0 && AtlasTextureSize.Z > 0)
+	{
+		const FNiagaraBakerOutputFrameIndices FrameIndices = BakerGeneratedSettings->GetOutputFrameIndices(OutputVolumeTexture, WorldTime);
+		const FIntPoint TileIndex(FrameIndices.FrameIndexA % BakerGeneratedSettings->FramesPerDimension.X, FrameIndices.FrameIndexA / BakerGeneratedSettings->FramesPerDimension.X);
+		const FIntPoint TileSize(AtlasVolumeTexture->GetSizeX() / BakerGeneratedSettings->FramesPerDimension.X, AtlasVolumeTexture->GetSizeY() / BakerGeneratedSettings->FramesPerDimension.Y);
+		const FVector2f TileUV((float(TileSize.X * TileIndex.X) + 0.5f) / float(AtlasTextureSize.X), (float(TileSize.Y * TileIndex.Y) + 0.5f) / float(AtlasTextureSize.Y));
+		const FVector2f TileUVSize(float(TileSize.X - 1) / float(AtlasTextureSize.X), float(TileSize.Y - 1) / float(AtlasTextureSize.Y));
+		const int32 AttributeChannels[] = { 0, 1, 2, 3 };
+
+		const int32 NumSlices = AtlasTextureSize.Z;
+		const int32 SlicePerAxis = FMath::CeilToInt(FMath::Sqrt((float)NumSlices));
+		const int32 TileWidth = RenderTarget->GetSurfaceWidth() / SlicePerAxis;
+		const int32 TileHeight = RenderTarget->GetSurfaceHeight() / SlicePerAxis;
+		if (TileWidth > 0 && TileHeight > 0)
+		{
+			for (int32 Slice=0; Slice < NumSlices; ++Slice)
+			{
+				const int32 TileX = (Slice % SlicePerAxis) * TileWidth;
+				const int32 TileY = (Slice / SlicePerAxis) * TileHeight;
+
+				FVector3f SliceAttributeUVs[4];
+				for (int i=0; i < UE_ARRAY_COUNT(SliceAttributeUVs); ++i)
+				{
+					SliceAttributeUVs[i].X = TileUV.X;
+					SliceAttributeUVs[i].Y = TileUV.Y;
+					SliceAttributeUVs[i].Z = (float(Slice) + 0.5f) / float(AtlasTextureSize.Z);
+				}
+
+				FCanvasTileItem TileItem(FVector2D(TileX, TileY), GWhiteTexture, FVector2D(TileWidth, TileHeight), FVector2D(0.0, 1.0f), FVector2D(1.0, 0.0f), FLinearColor::White);
+				TileItem.BlendMode = SE_BLEND_Opaque;
+				TileItem.BatchedElementParameters = new FBatchedElementNiagaraVolumeAttribute(
+					TileUVSize,
+					MakeArrayView(SliceAttributeUVs),
+					MakeArrayView(AttributeChannels),
+					[Resource_RT=AtlasVolumeTexture->GetResource()](FRHITexture*& OutTexture, FRHISamplerState*& OutSamplerState)
+					{
+						if (Resource_RT)
+						{
+							OutTexture = Resource_RT->GetTexture3DRHI();
+							OutSamplerState = Resource_RT->SamplerStateRHI;
+						}
+						//else
+						//{
+						//	OutTexture = &GBlackVolumeTexture;
+						//	OutSamplerState = TStaticSamplerState<SF_Bilinear>::GetRHI();
+						//}
+					}
+				);
+				Canvas.DrawItem(TileItem);
+			}
+		}
+	}
+	Canvas.Flush_GameThread();
+}
+
+bool FNiagaraBakerRendererOutputVolumeTexture::BeginBake(UNiagaraBakerOutput* InBakerOutput)
+{
+	UNiagaraBakerOutputVolumeTexture* OutputVolumeTexture = CastChecked<UNiagaraBakerOutputVolumeTexture>(InBakerOutput);
+	return OutputVolumeTexture->bGenerateAtlas || OutputVolumeTexture->bGenerateFrames || OutputVolumeTexture->bExportFrames;
+}
+
+void FNiagaraBakerRendererOutputVolumeTexture::BakeFrame(UNiagaraBakerOutput* InBakerOutput, int FrameIndex, const FNiagaraBakerRenderer& BakerRenderer)
+{
+	UNiagaraBakerOutputVolumeTexture* OutputVolumeTexture = CastChecked<UNiagaraBakerOutputVolumeTexture>(InBakerOutput);
+
+	FVolumeDataInterfaceHelper DataInterface;
+	if ( DataInterface.Initialize(OutputVolumeTexture, BakerRenderer.GetPreviewComponent()) == false )
+	{
+		return;
+	}
+
+	FIntVector TextureSize = FIntVector::ZeroValue;
+	TArray<FFloat16Color> TextureData;
+
+	// Are we a volume render target?
+	if ( DataInterface.VolumeRenderTargetDataInterface )
+	{
+		ENQUEUE_RENDER_COMMAND(CopyVolumeRenderTarget)
+		(
+			[Proxy_RT=DataInterface.VolumeRenderTargetProxy, SystemInstanceId_RT=DataInterface.SystemInstance->GetId(), OutTextureSize_RT=&TextureSize, OutTextureData_RT=&TextureData](FRHICommandListImmediate& RHICmdList)
+			{
+				if ( const FRenderTargetVolumeRWInstanceData_RenderThread* InstanceData_RT = Proxy_RT->SystemInstancesToProxyData_RT.Find(SystemInstanceId_RT) )
+				{
+					*OutTextureSize_RT = InstanceData_RT->Size;
+					RHICmdList.Read3DSurfaceFloatData(
+						InstanceData_RT->TextureRHI,
+						FIntRect(0, 0, InstanceData_RT->Size.X, InstanceData_RT->Size.Y),
+						FIntPoint(0, InstanceData_RT->Size.Z),
+						*OutTextureData_RT
+					);
+				}
+			}
+		);
+		FlushRenderingCommands();
+	}
+	// Are we a grid collection?
+	else if ( DataInterface.Grid3DDataInterface )
+	{
+		ENQUEUE_RENDER_COMMAND(CopyVolumeRenderTarget)
+		(
+			[Proxy_RT=DataInterface.Grid3DProxy, SystemInstanceId_RT=DataInterface.SystemInstance->GetId(), OutTextureSize_RT=&TextureSize, OutTextureData_RT=&TextureData](FRHICommandListImmediate& RHICmdList)
+			{
+				if ( const FGrid3DCollectionRWInstanceData_RenderThread* InstanceData_RT = Proxy_RT->SystemInstancesToProxyData_RT.Find(SystemInstanceId_RT) )
+				{
+					if (InstanceData_RT->CurrentData)
+					{
+						FIntVector TextureSize;
+						TextureSize.X = InstanceData_RT->NumCells.X * InstanceData_RT->NumTiles.X;
+						TextureSize.Y = InstanceData_RT->NumCells.Y * InstanceData_RT->NumTiles.Y;
+						TextureSize.Z = InstanceData_RT->NumCells.Z * InstanceData_RT->NumTiles.Z;
+						*OutTextureSize_RT = TextureSize;
+
+						RHICmdList.Read3DSurfaceFloatData(
+							InstanceData_RT->CurrentData->GridBuffer.Buffer,
+							FIntRect(0, 0, TextureSize.X, TextureSize.Y),
+							FIntPoint(0, TextureSize.Z),
+							*OutTextureData_RT
+						);
+					}
+				}
+			}
+		);
+		FlushRenderingCommands();
+
+		if ( TextureSize == DataInterface.Grid3DTextureSize )
+		{
+			TArray<FFloat16Color> TileTextureData;
+			const FIntVector TileTextureSize = DataInterface.Grid3DInstanceData_GameThread->NumCells;
+			TileTextureData.AddUninitialized(TileTextureSize.X * TileTextureSize.Y * TileTextureSize.Z);
+
+			FFloat16 Channels[4] = { 0.0f, 0.0f, 0.0f, 1.0f };
+			int32 AttributeOffset[4] = { 0, 0, 0, 0 };
+			for (int32 c=0; c < DataInterface.Grid3DAttributeChannels; ++c)
+			{
+				const int32 AttributeIndex = DataInterface.Grid3DAttributeStart + c;
+				const FIntVector AttributeTileIndex(
+					AttributeIndex % DataInterface.Grid3DInstanceData_GameThread->NumTiles.X,
+					(AttributeIndex / DataInterface.Grid3DInstanceData_GameThread->NumTiles.X) % DataInterface.Grid3DInstanceData_GameThread->NumTiles.Y,
+					AttributeIndex / (DataInterface.Grid3DInstanceData_GameThread->NumTiles.X * DataInterface.Grid3DInstanceData_GameThread->NumTiles.Y)
+				);
+
+				AttributeOffset[c] = AttributeTileIndex.X * TileTextureSize.X;
+				AttributeOffset[c] += AttributeTileIndex.Y * TextureSize.X;
+				AttributeOffset[c] += AttributeTileIndex.Z * TextureSize.X * TextureSize.Y;
+			}
+
+			int32 OutTexel = 0;
+			for ( int32 z=0; z < TileTextureSize.Z; ++z )
+			{
+				const int32 SliceOffset = z * TextureSize.X * TextureSize.Y;
+				for ( int32 y=0; y < TileTextureSize.Y; ++y )
+				{
+					const int32 SliceRowOffset = SliceOffset + (y * TextureSize.Y);
+					for ( int32 x=0; x < TileTextureSize.X; ++x )
+					{
+						const int32 Offset = SliceRowOffset + x;
+						for ( int32 c=0; c < DataInterface.Grid3DAttributeChannels; ++c )
+						{
+							Channels[c] = TextureData[AttributeOffset[c] + Offset].R;
+						}
+						TileTextureData[OutTexel].R = Channels[0];
+						TileTextureData[OutTexel].G = Channels[1];
+						TileTextureData[OutTexel].B = Channels[2];
+						TileTextureData[OutTexel].A = Channels[3];
+						++OutTexel;
+					}
+				}
+			}
+
+			TextureSize = TileTextureSize;
+			TextureData = TileTextureData;
+		}
+		else
+		{
+			TextureData.Empty();
+		}
+	}
+
+	// Process the volume data
+	if ( TextureData.Num() > 0 )
+	{
+		// Create atlas
+		if (OutputVolumeTexture->bGenerateAtlas)
+		{
+			UNiagaraBakerSettings* BakerSettings = OutputVolumeTexture->GetTypedOuter<UNiagaraBakerSettings>();
+			check(BakerSettings);
+			if (BakeAtlasTextureData.Num() == 0)
+			{
+				BakeAtlasFrameSize = TextureSize;
+				BakeAtlasTextureSize = FIntVector(
+					TextureSize.X * BakerSettings->FramesPerDimension.X,
+					TextureSize.Y * BakerSettings->FramesPerDimension.Y,
+					TextureSize.Z
+				);
+				BakeAtlasTextureData.AddDefaulted(BakeAtlasTextureSize.X * BakeAtlasTextureSize.Y * BakeAtlasTextureSize.Z);
+			}
+
+			for ( int32 z=0; z < BakeAtlasFrameSize.Z; ++z )
+			{
+				for ( int32 y=0; y < BakeAtlasFrameSize.Y; ++y )
+				{
+					const int32 TileX = FrameIndex % BakerSettings->FramesPerDimension.X;
+					const int32 TileY = FrameIndex / BakerSettings->FramesPerDimension.X;
+					const int32 DstOffset =
+						(BakeAtlasTextureSize.X * BakeAtlasTextureSize.Y * z) +
+						(BakeAtlasFrameSize.Y * BakeAtlasTextureSize.X * TileY) +
+						(BakeAtlasFrameSize.X * TileX) +
+						(BakeAtlasTextureSize.X * y);
+					const int32 SrcOffset =
+						(TextureSize.X * TextureSize.Y * z) +
+						(TextureSize.X * y);
+					FMemory::Memcpy(BakeAtlasTextureData.GetData() + DstOffset, TextureData.GetData() + SrcOffset, BakeAtlasFrameSize.X * sizeof(FFloat16Color));
+				}
+			}
+		}
+
+		// Create asset per frame
+		if (OutputVolumeTexture->bGenerateFrames)
+		{
+			const FString AssetFullName	= OutputVolumeTexture->GetAssetPath(OutputVolumeTexture->FramesAssetPathFormat, FrameIndex);
+			const FString AssetName		= FString(FPathViews::GetCleanFilename(AssetFullName));
+			const FString PackagePath	= FString(FPathViews::GetPath(AssetFullName));
+
+			//-TODO: Handle package existing
+
+			IAssetTools& AssetTools = FModuleManager::Get().LoadModuleChecked<FAssetToolsModule>("AssetTools").Get();
+			UVolumeTextureFactory* VolumeTextureFactory = NewObject<UVolumeTextureFactory>();
+			UVolumeTexture* OutputTexture = Cast<UVolumeTexture>(AssetTools.CreateAsset(AssetName, PackagePath, UVolumeTexture::StaticClass(), VolumeTextureFactory));
+			if (OutputTexture)
+			{
+				OutputTexture->Source.Init(TextureSize.X, TextureSize.Y, TextureSize.Z, 1, TSF_RGBA16F, (const uint8*)(TextureData.GetData()));
+				OutputTexture->PowerOfTwoMode = ETexturePowerOfTwoSetting::None;
+				OutputTexture->MipGenSettings = TextureMipGenSettings::TMGS_NoMipmaps;
+				OutputTexture->UpdateResource();
+				OutputTexture->PostEditChange();
+				OutputTexture->MarkPackageDirty();
+			}
+		}
+
+		// Export each frame
+		if (OutputVolumeTexture->bExportFrames)
+		{
+			const FString FilePath = OutputVolumeTexture->GetExportPath(OutputVolumeTexture->FramesExportPathFormat, FrameIndex);
+			FNiagaraBakerRenderer::ExportImage(FilePath, FIntPoint(TextureSize.X, TextureSize.Y * TextureSize.Z), TextureData);
+		}
+	}
+}
+
+void FNiagaraBakerRendererOutputVolumeTexture::EndBake(UNiagaraBakerOutput* InBakerOutput)
+{
+	UNiagaraBakerOutputVolumeTexture* OutputVolumeTexture = CastChecked<UNiagaraBakerOutputVolumeTexture>(InBakerOutput);
+	if (BakeAtlasTextureData.Num() > 0)
+	{
+		const FString AssetFullName = OutputVolumeTexture->GetAssetPath(OutputVolumeTexture->AtlasAssetPathFormat, 0);
+		const FString AssetName = FString(FPathViews::GetCleanFilename(AssetFullName));
+		const FString PackagePath = FString(FPathViews::GetPath(AssetFullName));
+
+		//-TODO: Handle package existing
+
+		IAssetTools& AssetTools = FModuleManager::Get().LoadModuleChecked<FAssetToolsModule>("AssetTools").Get();
+		UVolumeTextureFactory* VolumeTextureFactory = NewObject<UVolumeTextureFactory>();
+		UVolumeTexture* OutputTexture = Cast<UVolumeTexture>(AssetTools.CreateAsset(AssetName, PackagePath, UVolumeTexture::StaticClass(), VolumeTextureFactory));
+		if (OutputTexture)
+		{
+			OutputTexture->Source.Init(BakeAtlasTextureSize.X, BakeAtlasTextureSize.Y, BakeAtlasTextureSize.Z, 1, TSF_RGBA16F, (const uint8*)(BakeAtlasTextureData.GetData()));
+			OutputTexture->PowerOfTwoMode = ETexturePowerOfTwoSetting::None;
+			OutputTexture->MipGenSettings = TextureMipGenSettings::TMGS_NoMipmaps;
+			OutputTexture->UpdateResource();
+			OutputTexture->PostEditChange();
+			OutputTexture->MarkPackageDirty();
+		}
+
+		BakeAtlasFrameSize = FIntVector::ZeroValue;
+		BakeAtlasTextureSize = FIntVector::ZeroValue;
+		BakeAtlasTextureData.Empty();
+	}
+}
+
diff --git a/Engine/Plugins/FX/Niagara/Source/NiagaraEditor/Private/NiagaraBakerRendererOutputVolumeTexture.h b/Engine/Plugins/FX/Niagara/Source/NiagaraEditor/Private/NiagaraBakerRendererOutputVolumeTexture.h
new file mode 100644
--- /dev/null
+++ b/Engine/Plugins/FX/Niagara/Source/NiagaraEditor/Private/NiagaraBakerRendererOutputVolumeTexture.h
@@ -0,0 +1,28 @@
+// Copyright Epic Games, Inc. All Rights Reserved.
+
+#pragma once
+
+#include "NiagaraBakerRenderer.h"
+
+class UNiagaraBakerOutputVolumeTexture;
+
+class FNiagaraBakerRendererOutputVolumeTexture : public FNiagaraBakerOutputRenderer
+{
+public:
+	virtual TArray<FNiagaraBakerOutputBinding> GetRendererBindings(UNiagaraBakerOutput* InBakerOutput) const override;
+
+	virtual FIntPoint GetPreviewSize(UNiagaraBakerOutput* InBakerOutput, FIntPoint InAvailableSize) const override;
+	virtual void RenderPreview(UNiagaraBakerOutput* InBakerOutput, const FNiagaraBakerRenderer& BakerRenderer, UTextureRenderTarget2D* RenderTarget, TOptional<FString>& OutErrorString) const override;
+
+	virtual FIntPoint GetGeneratedSize(UNiagaraBakerOutput* InBakerOutput, FIntPoint InAvailableSize) const override;
+	virtual void RenderGenerated(UNiagaraBakerOutput* InBakerOutput, const FNiagaraBakerRenderer& BakerRenderer, UTextureRenderTarget2D* RenderTarget, TOptional<FString>& OutErrorString) const override;
+
+	virtual bool BeginBake(UNiagaraBakerOutput* InBakerOutput) override;
+	virtual void BakeFrame(UNiagaraBakerOutput* InBakerOutput, int FrameIndex, const FNiagaraBakerRenderer& BakerRenderer) override;
+	virtual void EndBake(UNiagaraBakerOutput* InBakerOutput) override;
+
+protected:
+	FIntVector				BakeAtlasFrameSize = FIntVector::ZeroValue;
+	FIntVector				BakeAtlasTextureSize = FIntVector::ZeroValue;
+	TArray<FFloat16Color>	BakeAtlasTextureData;
+};
diff --git a/Engine/Plugins/FX/Niagara/Source/NiagaraEditor/Private/ViewModels/NiagaraBakerViewModel.cpp b/Engine/Plugins/FX/Niagara/Source/NiagaraEditor/Private/ViewModels/NiagaraBakerViewModel.cpp
--- a/Engine/Plugins/FX/Niagara/Source/NiagaraEditor/Private/ViewModels/NiagaraBakerViewModel.cpp
+++ b/Engine/Plugins/FX/Niagara/Source/NiagaraEditor/Private/ViewModels/NiagaraBakerViewModel.cpp
@@ -11,7 +11,6 @@
 
 #include "Misc/ScopedSlowTask.h"
 #include "Modules/ModuleManager.h"
-#include "AdvancedPreviewScene.h"
 #include "AssetToolsModule.h"
 #include "EngineModule.h"
 #include "LegacyScreenPercentageDriver.h"
@@ -28,39 +27,14 @@ FNiagaraBakerViewModel::FNiagaraBakerViewModel()
 
 FNiagaraBakerViewModel::~FNiagaraBakerViewModel()
 {
-	if ( AdvancedPreviewScene && PreviewComponent )
-	{
-		AdvancedPreviewScene->RemoveComponent(PreviewComponent);
-	}
-
-	if ( PreviewComponent )
-	{
-		PreviewComponent->DestroyComponent();
-	}
 }
 
 void FNiagaraBakerViewModel::Initialize(TWeakPtr<FNiagaraSystemViewModel> InWeakSystemViewModel)
 {
 	WeakSystemViewModel = InWeakSystemViewModel;
 
-	{
-		UNiagaraSystem* System = WeakSystemViewModel.Pin()->GetPreviewComponent()->GetAsset();
-
-		PreviewComponent = NewObject<UNiagaraComponent>(GetTransientPackage(), NAME_None, RF_Transient);
-		PreviewComponent->CastShadow = 1;
-		PreviewComponent->bCastDynamicShadow = 1;
-		PreviewComponent->SetAllowScalability(false);
-		PreviewComponent->SetAsset(System);
-		PreviewComponent->SetForceSolo(true);
-		PreviewComponent->SetAgeUpdateMode(ENiagaraAgeUpdateMode::DesiredAge);
-		PreviewComponent->SetCanRenderWhileSeeking(true);
-		PreviewComponent->SetMaxSimTime(0.0f);
-		PreviewComponent->Activate(true);
-
-		AdvancedPreviewScene = MakeShareable(new FAdvancedPreviewScene(FPreviewScene::ConstructionValues()));
-		AdvancedPreviewScene->SetFloorVisibility(false);
-		AdvancedPreviewScene->AddComponent(PreviewComponent, PreviewComponent->GetRelativeTransform());
-	}
+	UNiagaraSystem* NiagaraSystem = WeakSystemViewModel.Pin()->GetPreviewComponent()->GetAsset();
+	BakerRenderer.Reset(new FNiagaraBakerRenderer(NiagaraSystem));
 
 	Widget = SNew(SNiagaraBakerWidget)
 		.WeakViewModel(this->AsShared());
@@ -88,33 +62,6 @@ void FNiagaraBakerViewModel::SetDisplayTimeFromNormalized(float NormalizeTime)
 	}
 }
 
-UNiagaraComponent* FNiagaraBakerViewModel::GetPreviewComponent() const
-{
-	return PreviewComponent;
-}
-
-UNiagaraBakerSettings* FNiagaraBakerViewModel::GetBakerSettings() const
-{
-	if (PreviewComponent)
-	{
-		UNiagaraSystem* Asset = PreviewComponent->GetAsset();
-		return Asset ? Asset->GetBakerSettings() : nullptr;
-	}
-
-	return nullptr;
-}
-
-const UNiagaraBakerSettings* FNiagaraBakerViewModel::GetBakerGeneratedSettings() const
-{
-	if (PreviewComponent)
-	{
-		UNiagaraSystem* Asset = PreviewComponent->GetAsset();
-		return Asset ? Asset->GetBakerGeneratedSettings() : nullptr;
-	}
-
-	return nullptr;
-}
-
 TSharedPtr<class SWidget> FNiagaraBakerViewModel::GetWidget()
 {
 	return Widget;
@@ -125,8 +72,8 @@ void FNiagaraBakerViewModel::TogglePlaybackLooping()
 	if (UNiagaraBakerSettings* BakerSettings = GetBakerSettings())
 	{
 		const FScopedTransaction Transaction(LOCTEXT("ToggleLooping", "Toggle Looping"));
-		BakerSettings->bPreviewLooping = !BakerSettings->bPreviewLooping;
 		BakerSettings->Modify();
+		BakerSettings->bPreviewLooping = !BakerSettings->bPreviewLooping;
 	}
 }
 
@@ -141,8 +88,8 @@ void FNiagaraBakerViewModel::SetCameraViewMode(ENiagaraBakerViewMode ViewMode)
 	if (UNiagaraBakerSettings* BakerSettings = GetBakerSettings())
 	{
 		const FScopedTransaction Transaction(LOCTEXT("SetCameraViewMode", "Set Camera View Mode"));
-		BakerSettings->CameraViewportMode = ViewMode;
 		BakerSettings->Modify();
+		BakerSettings->CameraViewportMode = ViewMode;
 	}
 }
 
@@ -252,8 +199,8 @@ void FNiagaraBakerViewModel::SetCameraFOV(float InFOV)
 	if (UNiagaraBakerSettings* BakerSettings = GetBakerSettings())
 	{
 		const FScopedTransaction Transaction(LOCTEXT("SetCameraFOV", "Set Camera FOV"));
-		BakerSettings->CameraFOV = InFOV;
 		BakerSettings->Modify();
+		BakerSettings->CameraFOV = InFOV;
 	}
 }
 
@@ -268,8 +215,8 @@ void FNiagaraBakerViewModel::SetCameraOrbitDistance(float InOrbitDistance)
 	if (UNiagaraBakerSettings* BakerSettings = GetBakerSettings())
 	{
 		const FScopedTransaction Transaction(LOCTEXT("SetCameraOrbitDistance", "Set Camera Orbit Distance"));
-		BakerSettings->CameraOrbitDistance = InOrbitDistance;
 		BakerSettings->Modify();
+		BakerSettings->CameraOrbitDistance = InOrbitDistance;
 	}
 }
 
@@ -284,8 +231,8 @@ void FNiagaraBakerViewModel::SetCameraOrthoWidth(float InOrthoWidth)
 	if (UNiagaraBakerSettings* BakerSettings = GetBakerSettings())
 	{
 		const FScopedTransaction Transaction(LOCTEXT("SetCameraOrthoWidth", "Set Camera Ortho Width"));
-		BakerSettings->CameraOrthoWidth = InOrthoWidth;
 		BakerSettings->Modify();
+		BakerSettings->CameraOrthoWidth = InOrthoWidth;
 	}
 }
 
@@ -294,8 +241,8 @@ void FNiagaraBakerViewModel::ToggleCameraAspectRatioEnabled()
 	if ( UNiagaraBakerSettings* BakerSettings = GetBakerSettings() )
 	{
 		const FScopedTransaction Transaction(LOCTEXT("ToggleUseCameraAspectRatio", "Toggle Use Camera Aspect Ratio"));
-		BakerSettings->bUseCameraAspectRatio = !BakerSettings->bUseCameraAspectRatio;
 		BakerSettings->Modify();
+		BakerSettings->bUseCameraAspectRatio = !BakerSettings->bUseCameraAspectRatio;
 	}
 }
 
@@ -316,20 +263,22 @@ void FNiagaraBakerViewModel::SetCameraAspectRatio(float InAspectRatio)
 	if (UNiagaraBakerSettings* BakerSettings = GetBakerSettings())
 	{
 		const FScopedTransaction Transaction(LOCTEXT("SetCameraAspectRatio", "Set Camera Aspect Ratio"));
-		BakerSettings->CameraAspectRatio = InAspectRatio;
 		BakerSettings->Modify();
+		BakerSettings->CameraAspectRatio = InAspectRatio;
 	}
 }
 
-void FNiagaraBakerViewModel::AddOutput()
+void FNiagaraBakerViewModel::AddOutput(UClass* Class)
 {
 	if (UNiagaraBakerSettings* BakerSettings = GetBakerSettings())
 	{
+		UNiagaraBakerOutput* NewOutput = NewObject<UNiagaraBakerOutput>(BakerSettings, Class);
+
 		const FScopedTransaction Transaction(LOCTEXT("AddOutput", "Add Output"));
-		BakerSettings->OutputTextures.AddDefaulted();
 		BakerSettings->Modify();
+		BakerSettings->Outputs.Add(NewOutput);
 
-		CurrentOutputIndex = BakerSettings->OutputTextures.Num() - 1;
+		CurrentOutputIndex = BakerSettings->Outputs.Num() - 1;
 		OnCurrentOutputChanged.Broadcast();
 	}
 }
@@ -338,11 +287,12 @@ void FNiagaraBakerViewModel::RemoveCurrentOutput()
 {
 	if (UNiagaraBakerSettings* BakerSettings = GetBakerSettings())
 	{
-		if (BakerSettings->OutputTextures.IsValidIndex(CurrentOutputIndex))
+		if (BakerSettings->Outputs.IsValidIndex(CurrentOutputIndex))
 		{
 			const FScopedTransaction Transaction(LOCTEXT("RemoveOutput", "Remove Output"));
-			BakerSettings->OutputTextures.RemoveAt(CurrentOutputIndex);
 			BakerSettings->Modify();
+			BakerSettings->Outputs[CurrentOutputIndex]->MarkAsGarbage();
+			BakerSettings->Outputs.RemoveAt(CurrentOutputIndex);
 			CurrentOutputIndex = FMath::Max(0, CurrentOutputIndex - 1);
 			OnCurrentOutputChanged.Broadcast();
 		}
@@ -352,16 +302,28 @@ void FNiagaraBakerViewModel::RemoveCurrentOutput()
 bool FNiagaraBakerViewModel::CanRemoveCurrentOutput() const
 {
 	UNiagaraBakerSettings* BakerSettings = GetBakerSettings();
-	return BakerSettings && BakerSettings->OutputTextures.IsValidIndex(CurrentOutputIndex);
+	return BakerSettings && BakerSettings->Outputs.IsValidIndex(CurrentOutputIndex);
+}
+
+UNiagaraBakerOutput* FNiagaraBakerViewModel::GetCurrentOutput() const
+{
+	if ( UNiagaraBakerSettings* BakerSettings = GetBakerSettings() )
+	{
+		if ( BakerSettings->Outputs.IsValidIndex(CurrentOutputIndex) )
+		{
+			return BakerSettings->Outputs[CurrentOutputIndex];
+		}
+	}
+	return nullptr;
 }
 
 void FNiagaraBakerViewModel::SetCurrentOutputIndex(int32 OutputIndex)
 {
 	if (UNiagaraBakerSettings* BakerSettings = GetBakerSettings())
 	{
-		if (BakerSettings->OutputTextures.Num() > 0)
+		if (BakerSettings->Outputs.Num() > 0)
 		{
-			CurrentOutputIndex = FMath::Clamp(OutputIndex, 0, BakerSettings->OutputTextures.Num() - 1);
+			CurrentOutputIndex = FMath::Clamp(OutputIndex, 0, BakerSettings->Outputs.Num() - 1);
 		}
 		else
 		{
@@ -374,16 +336,9 @@ void FNiagaraBakerViewModel::SetCurrentOutputIndex(int32 OutputIndex)
 FText FNiagaraBakerViewModel::GetOutputText(int32 OutputIndex) const
 {
 	UNiagaraBakerSettings* BakerSettings = GetBakerSettings();
-	if ( BakerSettings && BakerSettings->OutputTextures.IsValidIndex(OutputIndex) )
+	if ( BakerSettings && BakerSettings->Outputs.IsValidIndex(OutputIndex) )
 	{
-		if ( BakerSettings->OutputTextures[OutputIndex].OutputName.IsNone() )
-		{
-			return FText::Format(LOCTEXT("OutputFormat", "Output {0}"), FText::AsNumber(OutputIndex));
-		}
-		else
-		{
-			return FText::FromName(BakerSettings->OutputTextures[OutputIndex].OutputName);
-		}
+		return FText::FromString(BakerSettings->Outputs[OutputIndex]->OutputName);
 	}
 	return FText::GetEmpty();
 }
@@ -393,155 +348,194 @@ FText FNiagaraBakerViewModel::GetCurrentOutputText() const
 	return GetOutputText(CurrentOutputIndex);
 }
 
-void FNiagaraBakerViewModel::AddReferencedObjects(FReferenceCollector& Collector)
+int FNiagaraBakerViewModel::GetCurrentOutputNumFrames() const
 {
+	if ( UNiagaraBakerSettings* BakerSettings = GetBakerSettings() )
+	{
+		return BakerSettings->GetOutputNumFrames(CurrentOutputIndex);
+	}
+	return 1;
 }
 
-void FNiagaraBakerViewModel::RenderBaker()
+FNiagaraBakerOutputFrameIndices FNiagaraBakerViewModel::GetCurrentOutputFrameIndices(float RelativeTime) const
 {
-	UNiagaraBakerSettings* BakerSettings = GetBakerSettings();
-	if ( (PreviewComponent == nullptr) || (BakerSettings == nullptr) )
+	if (UNiagaraBakerSettings* BakerSettings = GetBakerSettings())
 	{
-		return;
+		return BakerSettings->GetOutputFrameIndices(CurrentOutputIndex, RelativeTime);
 	}
+	return FNiagaraBakerOutputFrameIndices();
+}
+
+float FNiagaraBakerViewModel::GetTimelineStart() const
+{
+	UNiagaraBakerSettings* BakerSettings = GetBakerSettings();
+	return BakerSettings ? BakerSettings->StartSeconds : 0.0f;
+}
 
-	if ( BakerSettings->OutputTextures.Num() == 0 )
+void FNiagaraBakerViewModel::SetTimelineStart(float Value)
+{
+	if ( UNiagaraBakerSettings* BakerSettings = GetBakerSettings() )
 	{
-		return;
+		Value = FMath::Max(Value, 0.0f);
+		if ( BakerSettings->StartSeconds != Value )
+		{
+			const float EndSeconds = BakerSettings->StartSeconds + BakerSettings->DurationSeconds;
+			const FScopedTransaction Transaction(LOCTEXT("StartSeconds", "StartSeconds"));
+			BakerSettings->Modify();
+			BakerSettings->StartSeconds = Value;
+			BakerSettings->DurationSeconds = FMath::Max(EndSeconds - Value, 0.0f);
+		}
 	}
+}
 
-	UWorld* World = PreviewComponent->GetWorld();
+float FNiagaraBakerViewModel::GetDurationSeconds() const
+{
+	UNiagaraBakerSettings* BakerSettings = GetBakerSettings();
+	return BakerSettings ? BakerSettings->DurationSeconds : 0.0f;
+}
 
-	// Create output render targets & output Baker data
-	TArray<UTextureRenderTarget2D*, TInlineAllocator<4>> RenderTargets;
-	TArray<TArray<FFloat16Color>> OutputBakers;
+float FNiagaraBakerViewModel::GetTimelineEnd() const
+{
+	UNiagaraBakerSettings* BakerSettings = GetBakerSettings();
+	return BakerSettings ? BakerSettings->StartSeconds + BakerSettings->DurationSeconds : 0.0f;
+}
 
-	for ( const auto& OutputTexture : BakerSettings->OutputTextures )
+void FNiagaraBakerViewModel::SetTimelineEnd(float Value)
+{
+	if (UNiagaraBakerSettings* BakerSettings = GetBakerSettings())
 	{
-		if (OutputTexture.IsValidForBake())
-		{
-			UTextureRenderTarget2D* RenderTarget = NewObject<UTextureRenderTarget2D>();
-			RenderTarget->AddToRoot();
-			RenderTarget->ClearColor = FLinearColor::Transparent;
-			RenderTarget->TargetGamma = 1.0f;
-			RenderTarget->InitCustomFormat(OutputTexture.FrameSize.X, OutputTexture.FrameSize.Y, PF_FloatRGBA, false);
-
-			RenderTargets.Add(RenderTarget);
-		}
-		else
+		Value = FMath::Max(Value - BakerSettings->StartSeconds, 0.0f);
+		if (BakerSettings->DurationSeconds != Value)
 		{
-			RenderTargets.Add(nullptr);
+			const FScopedTransaction Transaction(LOCTEXT("EndSeconds", "EndSeconds"));
+			BakerSettings->Modify();
+			BakerSettings->DurationSeconds = Value;
 		}
-		OutputBakers.AddDefaulted_GetRef().SetNumZeroed(OutputTexture.TextureSize.X * OutputTexture.TextureSize.Y);
 	}
+}
 
-	// Render frames
-	PreviewComponent->ReinitializeSystem();
+int32 FNiagaraBakerViewModel::GetFramesOnX() const
+{
+	UNiagaraBakerSettings* BakerSettings = GetBakerSettings();
+	return BakerSettings ? BakerSettings->FramesPerDimension.X : 1;
+}
 
-	const int32 TotalFrames = BakerSettings->FramesPerDimension.X * BakerSettings->FramesPerDimension.Y;
-	const float FrameDeltaSeconds = BakerSettings->DurationSeconds / float(TotalFrames);
+void FNiagaraBakerViewModel::SetFramesOnX(int32 Value)
+{
+	if (UNiagaraBakerSettings* BakerSettings = GetBakerSettings())
+	{
+		if (BakerSettings->FramesPerDimension.X != Value)
+		{
+			const FScopedTransaction Transaction(LOCTEXT("FramesPerDimension", "FramesPerDimension"));
+			BakerSettings->Modify();
+			BakerSettings->FramesPerDimension.X = Value;
 
-	FScopedSlowTask SlowTask(TotalFrames, LOCTEXT("RenderingBaker", "Rendering Baker..."));
-	SlowTask.MakeDialog();
+			FPropertyChangedEvent PropertyEvent(UNiagaraBakerSettings::StaticClass()->FindPropertyByName(GET_MEMBER_NAME_CHECKED(UNiagaraBakerSettings, FramesPerDimension)));
+			BakerSettings->PostEditChangeProperty(PropertyEvent);
+		}
+	}
+}
 
-	FNiagaraBakerRenderer BakerRenderer;
+int32 FNiagaraBakerViewModel::GetFramesOnY() const
+{
+	UNiagaraBakerSettings* BakerSettings = GetBakerSettings();
+	return BakerSettings ? BakerSettings->FramesPerDimension.Y : 1;
+}
 
-	for ( int32 iFrame=0; iFrame < TotalFrames; ++iFrame)
+void FNiagaraBakerViewModel::SetFramesOnY(int32 Value)
+{
+	if (UNiagaraBakerSettings* BakerSettings = GetBakerSettings())
 	{
-		SlowTask.EnterProgressFrame(1);
-
-		// Step component to time
-		const float FrameTime = BakerSettings->StartSeconds + (float(iFrame) * FrameDeltaSeconds);
-		PreviewComponent->SetSeekDelta(BakerSettings->GetSeekDelta());
-		PreviewComponent->SeekToDesiredAge(FrameTime);
-		PreviewComponent->TickComponent(BakerSettings->GetSeekDelta(), ELevelTick::LEVELTICK_All, nullptr);
+		if (BakerSettings->FramesPerDimension.Y != Value)
+		{
+			const FScopedTransaction Transaction(LOCTEXT("FramesPerDimension", "FramesPerDimension"));
+			BakerSettings->Modify();
+			BakerSettings->FramesPerDimension.Y = Value;
 
+			FPropertyChangedEvent PropertyEvent(UNiagaraBakerSettings::StaticClass()->FindPropertyByName(GET_MEMBER_NAME_CHECKED(UNiagaraBakerSettings, FramesPerDimension)));
+			BakerSettings->PostEditChangeProperty(PropertyEvent);
+		}
+	}
+}
 
-		// Render frame
-		for (int32 iOutputTexture=0; iOutputTexture < BakerSettings->OutputTextures.Num(); ++iOutputTexture)
-		{
-			const FNiagaraBakerTextureSettings& OutputTexture = BakerSettings->OutputTextures[iOutputTexture];
-			UTextureRenderTarget2D* RenderTarget = RenderTargets[iOutputTexture];
-			if (RenderTarget == nullptr)
-			{
-				continue;
-			}
-			FTextureRenderTargetResource* RenderTargetResource = RenderTarget->GameThread_GetRenderTargetResource();
+void FNiagaraBakerViewModel::RenderBaker()
+{
+	UNiagaraBakerSettings* BakerSettings = GetBakerSettings();
+	UNiagaraSystem* NiagaraSystem = BakerRenderer->GetNiagaraSystem();
+	if ( BakerSettings == nullptr || NiagaraSystem == nullptr )
+	{
+		return;
+	}
 
-			BakerRenderer.RenderView(PreviewComponent, BakerSettings, FrameTime, RenderTarget, iOutputTexture);
+	if ( BakerSettings->Outputs.Num() == 0 )
+	{
+		//-TODO: Message no outputs
+		NiagaraSystem->SetBakerGeneratedSettings(nullptr);
+		return;
+	}
 
-			TArray<FFloat16Color> OutSamples;
-			RenderTargetResource->ReadFloat16Pixels(OutSamples);
+	// Create output renderers
+	TArray<TUniquePtr<FNiagaraBakerOutputRenderer>> OutputRenderers;
+	OutputRenderers.Reserve(BakerSettings->Outputs.Num());
 
-			const FIntPoint ViewportOffset = FIntPoint(
-				OutputTexture.FrameSize.X * (iFrame % BakerSettings->FramesPerDimension.X),
-				OutputTexture.FrameSize.Y * (iFrame / BakerSettings->FramesPerDimension.X)
-			);
-			TArray<FFloat16Color>& OutputBaker = OutputBakers[iOutputTexture];
-			for ( int y=0; y < OutputTexture.FrameSize.Y; ++y )
-			{
-				FFloat16Color* SrcPixel = OutSamples.GetData() + (y * OutputTexture.FrameSize.X);
-				FFloat16Color* DstPixel = OutputBaker.GetData() + ViewportOffset.X + ((y + ViewportOffset.Y) * OutputTexture.TextureSize.X);
-				FMemory::Memcpy(DstPixel, SrcPixel, sizeof(FFloat16Color) * OutputTexture.FrameSize.X);
-			}
+	bool bHasValidRenderer = false;
+	for ( int iOutput=0; iOutput < BakerSettings->Outputs.Num(); ++iOutput)
+	{
+		UNiagaraBakerOutput* Output = BakerSettings->Outputs[iOutput];
+		OutputRenderers.Emplace(FNiagaraBakerRenderer::GetOutputRenderer(Output->GetClass()));
+		if ( FNiagaraBakerOutputRenderer* OutputRenderer = OutputRenderers.Last().Get() )
+		{
+			bHasValidRenderer |= OutputRenderer->BeginBake(Output);
 		}
 	}
 
-	// Remove existing generated settings to avoid holding references to textures
-	if (UNiagaraSystem* Asset = PreviewComponent->GetAsset())
+	// Early out if we have no valid renderers
+	if ( !bHasValidRenderer )
 	{
-		Asset->SetBakerGeneratedSettings(nullptr);
+		//-TODO: No valid renderers
+		NiagaraSystem->SetBakerGeneratedSettings(nullptr);
+		return;
 	}
 
-	// Send data to generated textures
-	for (int32 iOutputTexture = 0; iOutputTexture < BakerSettings->OutputTextures.Num(); ++iOutputTexture)
+	// Render frames
+	const int32 TotalFrames = BakerSettings->FramesPerDimension.X * BakerSettings->FramesPerDimension.Y;
+	const float FrameDeltaSeconds = BakerSettings->DurationSeconds / float(TotalFrames);
+
+	FScopedSlowTask SlowTask(TotalFrames);
+	SlowTask.MakeDialog();
+
+	for ( int32 iFrame=0; iFrame < TotalFrames; ++iFrame )
 	{
-		FNiagaraBakerTextureSettings& OutputTexture = BakerSettings->OutputTextures[iOutputTexture];
+		SlowTask.EnterProgressFrame(1, FText::Format(LOCTEXT("BakingFormat", "Baking Frame ({0} / {1})"), iFrame + 1, TotalFrames));
 
-		// If we don't already have a generated texture ask the user to create one
-		if (OutputTexture.GeneratedTexture == nullptr)
-		{
-			IAssetTools& AssetTools = FModuleManager::Get().LoadModuleChecked<FAssetToolsModule>("AssetTools").Get();
-			UTexture2DFactoryNew* Texture2DFactory = NewObject<UTexture2DFactoryNew>();
+		// Tick to right location
+		const float FrameTime = BakerSettings->StartSeconds + (float(iFrame) * FrameDeltaSeconds);
+		BakerRenderer->SetAbsoluteTime(FrameTime);
 
-			FString PackagePath = FPaths::GetPath(PreviewComponent->GetAsset()->GetPackage()->GetPathName());
-			FString AssetName = PreviewComponent->GetAsset()->GetName() + TEXT("_Baker") + FString::FromInt(iOutputTexture) + TEXT("Texture");
-			OutputTexture.GeneratedTexture = Cast<UTexture2D>(AssetTools.CreateAssetWithDialog(AssetName, PackagePath, UTexture2D::StaticClass(), Texture2DFactory));
-			if (OutputTexture.GeneratedTexture == nullptr)
+		// Capture frames
+		for ( int32 i=0; i < OutputRenderers.Num(); ++i )
+		{
+			if (OutputRenderers[i].IsValid())
 			{
-				// User skipped for some reason
-				continue;
+				OutputRenderers[i]->BakeFrame(BakerSettings->Outputs[i], iFrame, *BakerRenderer.Get());
 			}
 		}
-
-		const bool bIsPoT = FMath::IsPowerOfTwo(OutputTexture.TextureSize.X) && FMath::IsPowerOfTwo(OutputTexture.TextureSize.Y);
-
-		OutputTexture.GeneratedTexture->Source.Init(OutputTexture.TextureSize.X, OutputTexture.TextureSize.Y, 1, 1, TSF_RGBA16F, (const uint8*)(OutputBakers[iOutputTexture].GetData()));
-		OutputTexture.GeneratedTexture->PowerOfTwoMode = ETexturePowerOfTwoSetting::None;
-		OutputTexture.GeneratedTexture->MipGenSettings = bIsPoT ? TMGS_FromTextureGroup : TMGS_NoMipmaps;
-		OutputTexture.GeneratedTexture->AddressX = TextureAddress::TA_Clamp;
-		OutputTexture.GeneratedTexture->AddressY = TextureAddress::TA_Clamp;
-		OutputTexture.GeneratedTexture->UpdateResource();
-		OutputTexture.GeneratedTexture->PostEditChange();
-		OutputTexture.GeneratedTexture->MarkPackageDirty();
-		OutputTexture.GeneratedTexture->TemporarilyDisableStreaming();
 	}
 
-	// Duplicate and set as generated data
-	if (UNiagaraSystem* Asset = PreviewComponent->GetAsset())
+	// Complete bake
+	for (int32 i = 0; i < OutputRenderers.Num(); ++i)
 	{
-		Asset->SetBakerGeneratedSettings(DuplicateObject<UNiagaraBakerSettings>(BakerSettings, Asset));
-	}
-
-	// Clean up render targets
-	for ( UTextureRenderTarget2D* RenderTarget : RenderTargets )
-	{
-		if (RenderTarget)
+		if (OutputRenderers[i].IsValid())
 		{
-			RenderTarget->RemoveFromRoot();
-			RenderTarget = nullptr;
+			OutputRenderers[i]->EndBake(BakerSettings->Outputs[i]);
 		}
 	}
+
+	// Duplicate and set as generated data
+	NiagaraSystem->SetBakerGeneratedSettings(DuplicateObject<UNiagaraBakerSettings>(BakerSettings, NiagaraSystem));
+
+	// GC to cleanup any transient objects we created or ones we have ditched
+	CollectGarbage(GARBAGE_COLLECTION_KEEPFLAGS);
 }
 
 #undef LOCTEXT_NAMESPACE
diff --git a/Engine/Plugins/FX/Niagara/Source/NiagaraEditor/Private/ViewModels/NiagaraBakerViewModel.h b/Engine/Plugins/FX/Niagara/Source/NiagaraEditor/Private/ViewModels/NiagaraBakerViewModel.h
--- a/Engine/Plugins/FX/Niagara/Source/NiagaraEditor/Private/ViewModels/NiagaraBakerViewModel.h
+++ b/Engine/Plugins/FX/Niagara/Source/NiagaraEditor/Private/ViewModels/NiagaraBakerViewModel.h
@@ -4,12 +4,14 @@
 
 #include "CoreMinimal.h"
 #include "RHIDefinitions.h"
+#include "NiagaraBakerRenderer.h"
 #include "SceneView.h"
 #include "Templates/SharedPointer.h"
-#include "UObject/GCObject.h"
 
 #include "NiagaraBakerSettings.h"
 
+class UNiagaraBakerSettings;
+
 enum class ENiagaraBakerColorChannel : uint8
 {
 	Red,
@@ -19,21 +21,11 @@ enum class ENiagaraBakerColorChannel : uint8
 	Num
 };
 
-class FNiagaraBakerViewModel : public TSharedFromThis<FNiagaraBakerViewModel>, public FGCObject
+class FNiagaraBakerViewModel : public TSharedFromThis<FNiagaraBakerViewModel>
 {
 public:
 	DECLARE_MULTICAST_DELEGATE(FOnCurrentOutputChanged);
 
-	struct FDisplayData
-	{
-		int32		NumFrames = 0;
-		float		StartSeconds = 0.0f;
-		float		DurationSeconds = 0.0f;
-		float		NormalizedTime = 0.0f;
-		float		FrameTime = 0.0f;
-		int32		FrameIndex = 0;
-	};
-
 	FNiagaraBakerViewModel();
 	~FNiagaraBakerViewModel();
 
@@ -47,9 +39,9 @@ public:
 
 	void SetDisplayTimeFromNormalized(float NormalizeTime);
 
-	class UNiagaraComponent* GetPreviewComponent() const;
-	class UNiagaraBakerSettings* GetBakerSettings() const;
-	const class UNiagaraBakerSettings* GetBakerGeneratedSettings() const;
+	FNiagaraBakerRenderer& GetBakerRenderer() const { check(BakerRenderer.IsValid()); return *BakerRenderer.Get(); }
+	UNiagaraBakerSettings* GetBakerSettings() const { return BakerRenderer->GetBakerSettings(); }
+	const UNiagaraBakerSettings* GetBakerGeneratedSettings() { return BakerRenderer->GetBakerGeneratedSettings(); }
 
 	bool RenderView(const FRenderTarget* RenderTarget, FCanvas* Canvas, float WorldTime, int32 iOutputTextureIndex, bool bFillCanvas = false) const;
 
@@ -102,30 +94,37 @@ public:
 	float GetCameraAspectRatio() const;
 	void SetCameraAspectRatio(float InAspectRatio);
 
-	void AddOutput();
+	void AddOutput(UClass* Class);
 	void RemoveCurrentOutput();
 	bool CanRemoveCurrentOutput() const;
 	bool IsCurrentOutputIndex(int32 OutputIndex) const { return OutputIndex == CurrentOutputIndex; }
+	UNiagaraBakerOutput* GetCurrentOutput() const;
 	int32 GetCurrentOutputIndex() const { return CurrentOutputIndex; }
 	void SetCurrentOutputIndex(int32 OutputIndex);
 	FText GetOutputText(int32 OutputIndex) const;
 	FText GetCurrentOutputText() const;
+	int GetCurrentOutputNumFrames() const;
+	FNiagaraBakerOutputFrameIndices GetCurrentOutputFrameIndices(float RelativeTime) const;
+
+	float GetTimelineStart() const;
+	void SetTimelineStart(float Value);
 
-	/** FGCObject interface */
-	virtual void AddReferencedObjects(FReferenceCollector& Collector) override;
-	virtual FString GetReferencerName() const override
-	{
-		return TEXT("FNiagaraBakerViewModel");
-	}
+	float GetDurationSeconds() const;
+	
+	float GetTimelineEnd() const;
+	void SetTimelineEnd(float Value);
+
+	int32 GetFramesOnX() const;
+	void SetFramesOnX(int32 Value);
+	int32 GetFramesOnY() const;
+	void SetFramesOnY(int32 Value);
 
 	FOnCurrentOutputChanged OnCurrentOutputChanged;
 
 private:
-	class UNiagaraComponent* PreviewComponent = nullptr;
-	TSharedPtr<class FAdvancedPreviewScene> AdvancedPreviewScene;
-
 	TWeakPtr<FNiagaraSystemViewModel> WeakSystemViewModel;
 	TSharedPtr<class SNiagaraBakerWidget> Widget;
+	TUniquePtr<FNiagaraBakerRenderer> BakerRenderer;
 
 	int32 CurrentOutputIndex = 0;
 
diff --git a/Engine/Plugins/FX/Niagara/Source/NiagaraEditor/Private/Widgets/SNiagaraBakerTimelineWidget.cpp b/Engine/Plugins/FX/Niagara/Source/NiagaraEditor/Private/Widgets/SNiagaraBakerTimelineWidget.cpp
--- a/Engine/Plugins/FX/Niagara/Source/NiagaraEditor/Private/Widgets/SNiagaraBakerTimelineWidget.cpp
+++ b/Engine/Plugins/FX/Niagara/Source/NiagaraEditor/Private/Widgets/SNiagaraBakerTimelineWidget.cpp
@@ -23,18 +23,12 @@ void SNiagaraBakerTimelineWidget::Construct(const FArguments& InArgs)
 
 int32 SNiagaraBakerTimelineWidget::OnPaint(const FPaintArgs& Args, const FGeometry& AllottedGeometry, const FSlateRect& MyCullingRect, FSlateWindowElementList& OutDrawElements, int32 LayerId, const FWidgetStyle& InWidgetStyle, bool bParentEnabled) const
 {
-	auto ViewModel = WeakViewModel.Pin();
+	FNiagaraBakerViewModel* ViewModel = WeakViewModel.Pin().Get();
 	if ( !ViewModel )
 	{
 		return LayerId;
 	}
 
-	const UNiagaraBakerSettings* BakerSettings = ViewModel->GetBakerSettings();
-	if ( !BakerSettings )
-	{
-		return LayerId;
-	}
-
 	// Used to track the layer ID we will return.
 	int32 RetLayerId = LayerId;
 	const TSharedRef<FSlateFontMeasure> FontMeasureService = FSlateApplication::Get().GetRenderer()->GetFontMeasureService();
@@ -46,8 +40,8 @@ int32 SNiagaraBakerTimelineWidget::OnPaint(const FPaintArgs& Args, const FGeomet
 	static const FLinearColor SelectedBarColor(FLinearColor::White);
 
 	// Gather information to display
-	const int NumFrames = BakerSettings->GetNumFrames();
-	const auto BakerDisplayData = BakerSettings->GetDisplayInfo(RelativeTime, BakerSettings->bPreviewLooping);
+	const int NumFrames = ViewModel->GetCurrentOutputNumFrames();
+	const FNiagaraBakerOutputFrameIndices OutputFrameIndices = ViewModel->GetCurrentOutputFrameIndices(RelativeTime);
 
 	//////////////////////////////////////////////////////////////////////////
 	// Draw boxes to show frames
@@ -63,7 +57,7 @@ int32 SNiagaraBakerTimelineWidget::OnPaint(const FPaintArgs& Args, const FGeomet
 			{
 				const FVector2D BoxLocation(UStep * float(i), 0.0f);
 				const FVector2D BoxSize(UStep, AllottedGeometry.Size.Y);
-				const FLinearColor& Tint = (i == BakerDisplayData.FrameIndexA) ? CurrentTint : BoxTints[i & 1];
+				const FLinearColor& Tint = (i == OutputFrameIndices.FrameIndexA) ? CurrentTint : BoxTints[i & 1];
 				FSlateDrawElement::MakeBox(OutDrawElements, RetLayerId++, AllottedGeometry.ToPaintGeometry(BoxLocation, BoxSize), BoxBrush, ESlateDrawEffect::None, Tint);
 			}
 		}
@@ -79,7 +73,7 @@ int32 SNiagaraBakerTimelineWidget::OnPaint(const FPaintArgs& Args, const FGeomet
 		TArray<FVector2D> LinePoints;
 		LinePoints.AddUninitialized(2);
 
-		const float LineXPos = AllottedGeometry.Size.X * BakerDisplayData.NormalizedTime;
+		const float LineXPos = AllottedGeometry.Size.X * OutputFrameIndices.NormalizedTime;
 		LinePoints[0] = FVector2D(LineXPos, 0.0f);
 		LinePoints[1] = FVector2D(LineXPos, AllottedGeometry.Size.Y);
 
@@ -97,7 +91,7 @@ int32 SNiagaraBakerTimelineWidget::OnPaint(const FPaintArgs& Args, const FGeomet
 
 FReply SNiagaraBakerTimelineWidget::OnMouseButtonDown(const FGeometry& MyGeometry, const FPointerEvent& MouseEvent)
 {
-	if ( auto ViewModel = WeakViewModel.Pin() )
+	if ( FNiagaraBakerViewModel* ViewModel = WeakViewModel.Pin().Get() )
 	{
 		const FVector2D LocalLocation = MyGeometry.AbsoluteToLocal(MouseEvent.GetScreenSpacePosition());
 		const float NormalizeTime = LocalLocation.X / MyGeometry.Size.X;
diff --git a/Engine/Plugins/FX/Niagara/Source/NiagaraEditor/Private/Widgets/SNiagaraBakerViewport.cpp b/Engine/Plugins/FX/Niagara/Source/NiagaraEditor/Private/Widgets/SNiagaraBakerViewport.cpp
--- a/Engine/Plugins/FX/Niagara/Source/NiagaraEditor/Private/Widgets/SNiagaraBakerViewport.cpp
+++ b/Engine/Plugins/FX/Niagara/Source/NiagaraEditor/Private/Widgets/SNiagaraBakerViewport.cpp
@@ -7,6 +7,7 @@
 #include "NiagaraBatchedElements.h"
 #include "NiagaraComponent.h"
 
+#include "Engine/Canvas.h"
 #include "Engine/Font.h"
 #include "Framework/MultiBox/MultiBoxBuilder.h"
 #include "Modules/ModuleManager.h"
@@ -231,59 +232,110 @@ public:
 	{
 		Canvas->Clear(FLinearColor::Transparent);
 
-		auto ViewModel = WeakViewModel.Pin();
+		FNiagaraBakerViewModel* ViewModel = WeakViewModel.Pin().Get();
 		if (ViewModel == nullptr)
 		{
 			return;
 		}
 
-		//if (ViewRect.Width() <= 1 || ViewRect.Height() < = 1)
-		//{
-		//	return;
-		//}
-
 		UFont* DisplayFont = GetFont();
 		const float FontHeight = DisplayFont->GetMaxCharHeight() + 1.0f;
-		const FVector2D TextStartOffset(5.0f, 30.0f);
+		//const FVector2D TextStartOffset(5.0f, 30.0f);
 
 		const UNiagaraBakerSettings* BakerSettings = ViewModel->GetBakerSettings();
 		const UNiagaraBakerSettings* BakerGeneratedSettings = ViewModel->GetBakerGeneratedSettings();
-		const int32 CurrentOutputIndex = ViewModel->GetCurrentOutputIndex();
-		const float TimelineDurationSeconds = BakerSettings->DurationSeconds;
+		const float WorldTime = RelativeTime + BakerSettings->StartSeconds;
+
+		// Update Baker Renderer to the correct time
+		FNiagaraBakerRenderer& BakerRenderer = ViewModel->GetBakerRenderer();
+		BakerRenderer.SetAbsoluteTime(WorldTime, ViewModel->ShowRealtimePreview());
+
+		TUniquePtr<FNiagaraBakerOutputRenderer> BakerOutputRenderer;
+		UNiagaraBakerOutput* BakerPreviewOutput = ViewModel->GetCurrentOutput();
+		if (BakerPreviewOutput)
+		{
+			BakerOutputRenderer.Reset(FNiagaraBakerRenderer::GetOutputRenderer(BakerPreviewOutput->GetClass()));
+		}
 
-		// Determine view rects
+		UNiagaraBakerOutput* BakerGeneratedOutput = nullptr;
+		if ( BakerOutputRenderer.IsValid() && BakerGeneratedSettings && BakerGeneratedSettings->Outputs.IsValidIndex(ViewModel->GetCurrentOutputIndex()) )
+		{
+			BakerGeneratedOutput = BakerGeneratedSettings->Outputs[ViewModel->GetCurrentOutputIndex()];
+			if ((BakerGeneratedOutput->GetClass() != BakerPreviewOutput->GetClass()) ||
+				!FMath::IsNearlyEqual(BakerSettings->DurationSeconds, BakerGeneratedSettings->DurationSeconds) )
+			{
+				BakerGeneratedOutput = nullptr;
+			}
+		}
+
+		const FIntRect ViewRect(Canvas->GetViewRect().Min + FIntPoint(2, 2), Canvas->GetViewRect().Max - FIntPoint(2, 2));
+		if ( ViewRect.Width() <= 0 || ViewRect.Height() <= 0 )
+		{
+			return;
+		}
+
+		// Calculate view rects
 		PreviewViewRect = FIntRect();
-		BakerViewRect = FIntRect();
+		GeneratedViewRect = FIntRect();
+		bool bPreviewValid = false;
+		bool bGeneratedValid = false;
 		{
-			const FIntRect ViewRect = Canvas->GetViewRect();
-			const int32 Border = 3;
-			const int32 ViewWidth = ViewModel->ShowRealtimePreview() && ViewModel->ShowBakedView() ? (ViewRect.Width() >> 1) - Border : ViewRect.Width();
+			const int32 ViewWidth = ViewModel->ShowRealtimePreview() && ViewModel->ShowBakedView() ? (ViewRect.Width() >> 1) - 1 : ViewRect.Width();
 
 			if (ViewModel->ShowRealtimePreview())
 			{
-				if (BakerSettings->OutputTextures.IsValidIndex(CurrentOutputIndex))
-				{
-					PreviewViewRect = ConstrainRect(FIntRect(ViewRect.Min.X, ViewRect.Min.Y, ViewRect.Min.X + ViewWidth, ViewRect.Max.Y), BakerSettings->OutputTextures[CurrentOutputIndex].FrameSize);
-				}
-				else
+				PreviewViewRect = FIntRect(ViewRect.Min.X, ViewRect.Min.Y, ViewRect.Min.X + ViewWidth, ViewRect.Max.Y);
+				if (BakerOutputRenderer.IsValid())
 				{
-					PreviewViewRect = FIntRect(ViewRect.Min.X, ViewRect.Min.Y, ViewRect.Min.X + ViewWidth, ViewRect.Max.Y);
+					const FIntPoint PreviewSize = BakerOutputRenderer->GetPreviewSize(BakerPreviewOutput, PreviewViewRect.Size());
+					if (PreviewSize.SizeSquared() > 0 )
+					{
+						bPreviewValid = true;
+						PreviewViewRect = ConstrainRect(PreviewViewRect, PreviewSize);
+
+						// Resize our render target for the output
+						if (!PreviewRenderTarget || PreviewRenderTarget->SizeX != PreviewSize.X || PreviewRenderTarget->SizeY != PreviewSize.Y)
+						{
+							if (PreviewRenderTarget == nullptr)
+							{
+								PreviewRenderTarget = NewObject<UTextureRenderTarget2D>();
+							}
+							PreviewRenderTarget->ClearColor = FLinearColor::Transparent;
+							PreviewRenderTarget->TargetGamma = 1.0f;
+							PreviewRenderTarget->InitCustomFormat(PreviewSize.X, PreviewSize.Y, PF_FloatRGBA, false);
+						}
+					}
 				}
 			}
 
 			if (ViewModel->ShowBakedView())
 			{
-				if (BakerGeneratedSettings && BakerGeneratedSettings->OutputTextures.IsValidIndex(CurrentOutputIndex))
+				GeneratedViewRect = FIntRect(ViewRect.Max.X - ViewWidth, ViewRect.Min.Y, ViewRect.Max.X, ViewRect.Max.Y);
+				if ( BakerGeneratedOutput )
 				{
-					BakerViewRect = ConstrainRect(FIntRect(ViewRect.Max.X - ViewWidth, ViewRect.Min.Y, ViewRect.Max.X, ViewRect.Max.Y), BakerGeneratedSettings->OutputTextures[CurrentOutputIndex].FrameSize);
-				}
-				else
-				{
-					BakerViewRect = FIntRect(ViewRect.Max.X - ViewWidth, ViewRect.Min.Y, ViewRect.Max.X, ViewRect.Max.Y);
+					const FIntPoint GeneratedSize = BakerOutputRenderer->GetGeneratedSize(BakerGeneratedOutput, GeneratedViewRect.Size());
+					if (GeneratedSize.SizeSquared() > 0)
+					{
+						bGeneratedValid = true;
+						GeneratedViewRect = ConstrainRect(GeneratedViewRect, GeneratedSize);
+
+						// Resize our render target for the output
+						if (!GeneratedRenderTarget || GeneratedRenderTarget->SizeX != GeneratedSize.X || GeneratedRenderTarget->SizeY != GeneratedSize.Y)
+						{
+							if (GeneratedRenderTarget == nullptr)
+							{
+								GeneratedRenderTarget = NewObject<UTextureRenderTarget2D>();
+							}
+							GeneratedRenderTarget->ClearColor = FLinearColor::Transparent;
+							GeneratedRenderTarget->TargetGamma = 1.0f;
+							GeneratedRenderTarget->InitCustomFormat(GeneratedSize.X, GeneratedSize.Y, PF_FloatRGBA, false);
+						}
+					}
 				}
 			}
 		}
 
+		// Determine which color channels to show
 		const bool bRGBEnabled = ViewModel->IsChannelEnabled(ENiagaraBakerColorChannel::Red) || ViewModel->IsChannelEnabled(ENiagaraBakerColorChannel::Green) || ViewModel->IsChannelEnabled(ENiagaraBakerColorChannel::Blue);
 		const bool bAEnabled = ViewModel->IsChannelEnabled(ENiagaraBakerColorChannel::Alpha);
 		const bool bAlphaBlend = bAEnabled && bRGBEnabled;
@@ -301,42 +353,21 @@ public:
 			ColorTransform = FMatrix(RPlane, GPlane, BPlane, FPlane(0.0f, 0.0f, 0.0f, 1.0f));
 		}
 
-		// Render World Preview
+		// Render the realtime preview
 		if (ViewModel->ShowRealtimePreview())
 		{
 			ClearViewArea(Canvas, PreviewViewRect);
 
-			const float WorldTime = RelativeTime + BakerSettings->StartSeconds;
-			FVector2D TextPosition(PreviewViewRect.Min.X + TextStartOffset.X, PreviewViewRect.Min.Y + TextStartOffset.Y);
-			if (BakerSettings->OutputTextures.IsValidIndex(CurrentOutputIndex))
+			TOptional<FString> ErrorString;
+			if ( BakerOutputRenderer.IsValid() && bPreviewValid )
 			{
-				const FNiagaraBakerTextureSettings& OutputTexture = BakerSettings->OutputTextures[CurrentOutputIndex];
-				if ( OutputTexture.IsValidForBake() )
+				BakerOutputRenderer->RenderPreview(BakerPreviewOutput, BakerRenderer, PreviewRenderTarget, ErrorString);
+				if (ErrorString.IsSet() == false)
 				{
-					// Ensure render target is the correct size
-					if ( !RealtimeRenderTarget || RealtimeRenderTarget->SizeX != OutputTexture.FrameSize.X || RealtimeRenderTarget->SizeY != OutputTexture.FrameSize.Y )
-					{
-						if (RealtimeRenderTarget == nullptr)
-						{
-							RealtimeRenderTarget = NewObject<UTextureRenderTarget2D>();
-						}
-						RealtimeRenderTarget->ClearColor = FLinearColor::Transparent;
-						RealtimeRenderTarget->TargetGamma = 1.0f;
-						RealtimeRenderTarget->InitCustomFormat(OutputTexture.FrameSize.X, OutputTexture.FrameSize.Y, PF_FloatRGBA, false);
-					}
-
-					// Seek to correct time and render
-					UNiagaraComponent* PreviewComponent = ViewModel->GetPreviewComponent();
-					PreviewComponent->SetSeekDelta(BakerSettings->GetSeekDelta());
-					PreviewComponent->SeekToDesiredAge(WorldTime);
-					PreviewComponent->TickComponent(BakerSettings->GetSeekDelta(), ELevelTick::LEVELTICK_All, nullptr);
-
-					BakerRenderer.RenderView(PreviewComponent, BakerSettings, WorldTime, RealtimeRenderTarget, CurrentOutputIndex);
-
-					const FVector2D HalfPixel(0.5f / float(OutputTexture.FrameSize.X), 0.5f / float(OutputTexture.FrameSize.Y));
+					const FVector2D HalfPixel(0.5f / float(PreviewRenderTarget->GetSurfaceWidth()), 0.5f / float(PreviewRenderTarget->GetSurfaceHeight()));
 					FCanvasTileItem TileItem(
 						FVector2D(PreviewViewRect.Min.X, PreviewViewRect.Min.Y),
-						RealtimeRenderTarget->GetResource(),
+						PreviewRenderTarget->GetResource(),
 						FVector2D(PreviewViewRect.Width(), PreviewViewRect.Height()),
 						FVector2D(HalfPixel.X, HalfPixel.Y),
 						FVector2D(1.0f - HalfPixel.X, 1.0f - HalfPixel.Y),
@@ -344,95 +375,56 @@ public:
 					);
 					TileItem.BatchedElementParameters = new FBatchedElementNiagaraSimple(ColorTransform, bAlphaBlend);
 					Canvas->DrawItem(TileItem);
-
-					// Display info text
-					if (ViewModel->ShowInfoText())
-					{
-						Canvas->DrawShadowedString(TextPosition.X, TextPosition.Y, TEXT("Live Preview"), DisplayFont, FLinearColor::White);
-						TextPosition.Y += FontHeight;
-
-						Canvas->DrawShadowedString(TextPosition.X + 5.0f, TextPosition.Y, *FString::Printf(TEXT("Texture(%d) FrameSize(%d x %d)"), CurrentOutputIndex, OutputTexture.FrameSize.X, OutputTexture.FrameSize.Y), DisplayFont, FLinearColor::White);
-						TextPosition.Y += FontHeight;
-					}
-				}
-				else
-				{
-					Canvas->DrawShadowedString(TextPosition.X, TextPosition.Y, TEXT("Live Preview"), DisplayFont, FLinearColor::White);
-					TextPosition.Y += FontHeight;
-
-					Canvas->DrawShadowedString(TextPosition.X + 5.0f, TextPosition.Y, *FString::Printf(TEXT("Texture(%d) Is Invalid For Baking FrameSize(%d x %d)"), CurrentOutputIndex, OutputTexture.FrameSize.X, OutputTexture.FrameSize.Y), DisplayFont, FLinearColor::White);
 				}
 			}
 			else
 			{
-				Canvas->DrawShadowedString(TextPosition.X, TextPosition.Y, TEXT("Live Preview"), DisplayFont, FLinearColor::White);
-				TextPosition.Y += FontHeight;
+				ErrorString = TEXT("No output selected or the preview is invalid.");
+			}
 
-				Canvas->DrawShadowedString(TextPosition.X + 5.0f, TextPosition.Y, TEXT("No Output Texture"), DisplayFont, FLinearColor::White);
+			if ( ErrorString.IsSet() )
+			{
+				DrawErrorString(Canvas, PreviewViewRect, ErrorString.GetValue());
 			}
+
+			DrawRectBorder(Canvas, PreviewViewRect);
 		}
 
-		// Render Baker
+		// Render Generated View
 		if (ViewModel->ShowBakedView())
 		{
-			ClearViewArea(Canvas, BakerViewRect);
+			ClearViewArea(Canvas, GeneratedViewRect);
 
-			const bool bTimelineValid = BakerGeneratedSettings ? FMath::IsNearlyEqual(BakerGeneratedSettings->DurationSeconds, TimelineDurationSeconds) : false;
-			const bool bBakerValid =
-				(BakerGeneratedSettings != nullptr) &&
-				BakerGeneratedSettings->OutputTextures.IsValidIndex(CurrentOutputIndex) &&
-				(BakerGeneratedSettings->OutputTextures[CurrentOutputIndex].GeneratedTexture != nullptr);
-			FVector2D TextPosition(BakerViewRect.Min.X + TextStartOffset.X, BakerViewRect.Min.Y + TextStartOffset.Y);
-
-			if (bBakerValid && bTimelineValid)
+			TOptional<FString> ErrorString;
+			if ( BakerGeneratedOutput && bGeneratedValid )
 			{
-				const auto DisplayData = BakerGeneratedSettings->GetDisplayInfo(RelativeTime, BakerGeneratedSettings->bPreviewLooping);
-				const FNiagaraBakerTextureSettings& OutputTexture = BakerGeneratedSettings->OutputTextures[CurrentOutputIndex];
-
-				const FIntPoint TextureSize = FIntPoint(FMath::Max(OutputTexture.GeneratedTexture->GetSizeX(), 1), FMath::Max(OutputTexture.GeneratedTexture->GetSizeY(), 1));
-				const FIntPoint FramesPerDimension = BakerGeneratedSettings->FramesPerDimension;
-				const FIntPoint FrameIndexA = FIntPoint(DisplayData.FrameIndexA % FramesPerDimension.X, DisplayData.FrameIndexA / FramesPerDimension.X);
-				const FIntPoint FrameIndexB = FIntPoint(DisplayData.FrameIndexB % FramesPerDimension.X, DisplayData.FrameIndexB / FramesPerDimension.X);
-				const FIntPoint FramePixelA = FIntPoint(FrameIndexA.X * OutputTexture.FrameSize.X, FrameIndexA.Y * OutputTexture.FrameSize.Y);
-				const FIntPoint FramePixelB = FIntPoint(FrameIndexB.X * OutputTexture.FrameSize.X, FrameIndexB.Y * OutputTexture.FrameSize.Y);
-
-				FCanvasTileItem TileItem(
-					FVector2D(BakerViewRect.Min.X, BakerViewRect.Min.Y),
-					OutputTexture.GeneratedTexture->GetResource(),
-					FVector2D(BakerViewRect.Width(), BakerViewRect.Height()),
-					FVector2D((float(FramePixelA.X) + 0.5f) / float(TextureSize.X), (float(FramePixelA.Y) + 0.5f) / float(TextureSize.Y)),
-					FVector2D((float(FramePixelA.X + OutputTexture.FrameSize.X) - 0.5f) / float(TextureSize.X), (float(FramePixelA.Y + OutputTexture.FrameSize.Y) - 0.5f) / float(TextureSize.Y)),
-					FLinearColor::White
-				);
-				TileItem.BatchedElementParameters = new FBatchedElementNiagaraSimple(ColorTransform, bAlphaBlend);
-				Canvas->DrawItem(TileItem);
-
-				// Display info text
-				if (ViewModel->ShowInfoText())
+				BakerOutputRenderer->RenderGenerated(BakerGeneratedOutput, BakerRenderer, GeneratedRenderTarget, ErrorString);
+				if (ErrorString.IsSet() == false)
 				{
-					Canvas->DrawShadowedString(TextPosition.X, TextPosition.Y, TEXT("Baker"), DisplayFont, FLinearColor::White);
-					TextPosition.Y += FontHeight;
-
-					Canvas->DrawShadowedString(TextPosition.X + 5.0f, TextPosition.Y, *FString::Printf(TEXT("Texture(%d) FrameSize(%d x %d) Frame(%d/%d)"), CurrentOutputIndex, OutputTexture.FrameSize.X, OutputTexture.FrameSize.Y, DisplayData.FrameIndexA, BakerGeneratedSettings->GetNumFrames()), DisplayFont, FLinearColor::White);
-					TextPosition.Y += FontHeight;
-
-					if (!BakerGeneratedSettings->Equals(*BakerSettings))
-					{
-						Canvas->DrawShadowedString(TextPosition.X + 5.0f, TextPosition.Y, TEXT("Warning: Out Of Date"), DisplayFont, FLinearColor::White);
-						TextPosition.Y += FontHeight;
-					}
+					const FVector2D HalfPixel(0.5f / float(GeneratedRenderTarget->GetSurfaceWidth()), 0.5f / float(GeneratedRenderTarget->GetSurfaceHeight()));
+					FCanvasTileItem TileItem(
+						FVector2D(GeneratedViewRect.Min.X, GeneratedViewRect.Min.Y),
+						GeneratedRenderTarget->GetResource(),
+						FVector2D(GeneratedViewRect.Width(), GeneratedViewRect.Height()),
+						FVector2D(HalfPixel.X, HalfPixel.Y),
+						FVector2D(1.0f - HalfPixel.X, 1.0f - HalfPixel.Y),
+						FLinearColor::White
+					);
+					TileItem.BatchedElementParameters = new FBatchedElementNiagaraSimple(ColorTransform, bAlphaBlend);
+					Canvas->DrawItem(TileItem);
 				}
 			}
-			else if ( bBakerValid )
+			else
 			{
-				Canvas->DrawShadowedString(TextPosition.X, TextPosition.Y, TEXT("Baked duration does not match settings, can not map timeline."), DisplayFont, FLinearColor::White);
-				TextPosition.Y += FontHeight;
+				ErrorString = TEXT("No output or output not generated.\nPlease bake to generate the output");
 			}
-			else
+
+			if (ErrorString.IsSet())
 			{
-				Canvas->DrawShadowedString(TextPosition.X, TextPosition.Y, TEXT("Baker Not Generated"), DisplayFont, FLinearColor::White);
-				TextPosition.Y += FontHeight;
+				DrawErrorString(Canvas, GeneratedViewRect, ErrorString.GetValue());
 			}
+
+			DrawRectBorder(Canvas, GeneratedViewRect);
 		}
 	}
 
@@ -464,6 +456,65 @@ public:
 		Canvas->DrawTile(InRect.Min.X, InRect.Min.Y, InRect.Width(), InRect.Height(), 0.0f, 0.0f, EndUV.X, EndUV.Y, Color, Texture ? Texture->GetResource() : nullptr, false);
 	}
 
+	FVector2f GetStringSize(UFont* Font, const TCHAR* Text)
+	{
+		FVector2f MaxSize = FVector2f::ZeroVector;
+		FVector2f CurrSize = FVector2f::ZeroVector;
+
+		const float fAdvanceHeight = Font->GetMaxCharHeight();
+		const TCHAR* PrevChar = nullptr;
+		while (*Text)
+		{
+			if (*Text == '\n')
+			{
+				CurrSize.X = 0.0f;
+				CurrSize.Y = CurrSize.Y + fAdvanceHeight;
+				PrevChar = nullptr;
+				++Text;
+				continue;
+			}
+
+			float TmpWidth, TmpHeight;
+			Font->GetCharSize(*Text, TmpWidth, TmpHeight);
+
+			int8 CharKerning = 0;
+			if (PrevChar)
+			{
+				CharKerning = Font->GetCharKerning(*PrevChar, *Text);
+			}
+
+			CurrSize.X += TmpWidth + CharKerning;
+			MaxSize.X = FMath::Max(MaxSize.X, CurrSize.X);
+			MaxSize.Y = FMath::Max(MaxSize.Y, CurrSize.Y + TmpHeight);
+
+			PrevChar = Text++;
+		}
+
+		return MaxSize;
+	}
+
+	void DrawErrorString(FCanvas* Canvas, FIntRect ViewportRect, FStringView ErrorString)
+	{
+		Canvas->Flush_GameThread();
+
+		UFont* Font = GetFont();
+		const FVector2f StringSize = GetStringSize(Font, ErrorString.GetData());
+		const FIntPoint TextCenter(ViewportRect.Min.X + (ViewportRect.Width() >> 1), ViewportRect.Min.Y + (ViewportRect.Height() >> 1));
+		Canvas->DrawShadowedString(TextCenter.X - int32(StringSize.X * 0.5f), TextCenter.Y - int32(StringSize.Y * 0.5f), ErrorString.GetData(), Font, FLinearColor::White);
+
+		Canvas->Flush_GameThread();
+		Canvas->SetRenderTargetScissorRect(FIntRect(0, 0, 0, 0));
+	};
+
+	void DrawRectBorder(FCanvas* Canvas, const FIntRect& Rect)
+	{
+		Canvas->DrawTile(Rect.Min.X - 1, Rect.Min.Y - 1, 1, Rect.Size().Y + 2, 0.0f, 0.0f, 1.0f, 1.0f, FLinearColor::White);
+		Canvas->DrawTile(Rect.Max.X + 1, Rect.Min.Y - 1, 1, Rect.Size().Y + 2, 0.0f, 0.0f, 1.0f, 1.0f, FLinearColor::White);
+
+		Canvas->DrawTile(Rect.Min.X - 1, Rect.Min.Y - 1, Rect.Size().X + 2, 1, 0.0f, 0.0f, 1.0f, 1.0f, FLinearColor::White);
+		Canvas->DrawTile(Rect.Min.X - 1, Rect.Max.Y + 1, Rect.Size().X + 2, 1, 0.0f, 0.0f, 1.0f, 1.0f, FLinearColor::White);
+	}
+
 	UTexture2D* GetCheckerboardTexture()
 	{
 		if (CheckerboardTexture == nullptr)
@@ -491,14 +542,15 @@ public:
 		FEditorViewportClient::AddReferencedObjects(Collector);
 
 		Collector.AddReferencedObject(CheckerboardTexture);
-		Collector.AddReferencedObject(RealtimeRenderTarget);
+		Collector.AddReferencedObject(PreviewRenderTarget);
+		Collector.AddReferencedObject(GeneratedRenderTarget);
 	}
 
 	void FocusCamera()
 	{
 		auto ViewModel = WeakViewModel.Pin();
 		UNiagaraBakerSettings* BakerSettings = ViewModel ? ViewModel->GetBakerSettings() : nullptr;
-		UNiagaraComponent* NiagaraComponent = ViewModel->GetPreviewComponent();
+		UNiagaraComponent* NiagaraComponent = ViewModel->GetBakerRenderer().GetPreviewComponent();
 		if (ViewModel->ShowRealtimePreview() && BakerSettings && NiagaraComponent )
 		{
 			//-TODO: Should take aspect ratio into account here
@@ -528,9 +580,9 @@ public:
 		const UNiagaraBakerSettings* BakerSettings = ViewModel ? ViewModel->GetBakerSettings() : nullptr;
 		if (BakerSettings && (PreviewViewRect.Area() > 0))
 		{
-			OrthoUnits = BakerSettings->GetOrthoSize(ViewModel->GetCurrentOutputIndex());
-			OrthoUnits.X = OrthoUnits.X / float(PreviewViewRect.Width());
-			OrthoUnits.Y = OrthoUnits.Y / float(PreviewViewRect.Height());
+			const float AspectRatioY = BakerSettings->bUseCameraAspectRatio ? BakerSettings->CameraAspectRatio : 1.0f;
+			OrthoUnits.X = BakerSettings->CameraOrthoWidth / float(PreviewViewRect.Width());
+			OrthoUnits.Y = BakerSettings->CameraOrthoWidth * AspectRatioY / float(PreviewViewRect.Height());
 		}
 		return OrthoUnits;
 	}
@@ -539,13 +591,14 @@ public:
 	virtual UWorld* GetWorld() const override
 	{
 		auto ViewModel = WeakViewModel.Pin();
-		return ViewModel ? ViewModel->GetPreviewComponent()->GetWorld() : nullptr;
+		return ViewModel ? ViewModel->GetBakerRenderer().GetPreviewComponent()->GetWorld() : nullptr;
 	}
 
 	UFont* GetFont() const { return GetStatsFont(); }
 
 public:
-	UTextureRenderTarget2D*						RealtimeRenderTarget = nullptr;
+	UTextureRenderTarget2D*						PreviewRenderTarget = nullptr;
+	UTextureRenderTarget2D*						GeneratedRenderTarget = nullptr;
 
 	FVector										LocalMovement = FVector::ZeroVector;
 	float										LocalZoom = 0.0f;
@@ -555,7 +608,7 @@ public:
 	float										PerspectiveMoveSpeed = 2.0f;
 
 	FIntRect									PreviewViewRect;
-	FIntRect									BakerViewRect;
+	FIntRect									GeneratedViewRect;
 
 	UTexture2D*									CheckerboardTexture = nullptr;
 	FColor										CheckerboardColorOne = FColor(128, 128, 128);
@@ -568,7 +621,7 @@ public:
 	float										RelativeTime = 0.0f;
 	float										DeltaTime = 0.0f;
 
-	FNiagaraBakerRenderer						BakerRenderer;
+	//FNiagaraBakerRenderer						BakerRenderer;
 };
 
 //////////////////////////////////////////////////////////////////////////
@@ -578,10 +631,6 @@ void SNiagaraBakerViewport::Construct(const FArguments& InArgs)
 	WeakViewModel = InArgs._WeakViewModel;
 
 	SEditorViewport::FArguments ParentArgs;
-	//ParentArgs.EnableGammaCorrection(false);
-	//ParentArgs.IsEnabled(FSlateApplication::Get().GetNormalExecutionAttribute());
-	//ParentArgs.ShowEffectWhenDisabled(false);
-	//ParentArgs.EnableBlending(true);
 	SEditorViewport::Construct(ParentArgs);
 }
 
@@ -593,15 +642,6 @@ TSharedRef<FEditorViewportClient> SNiagaraBakerViewport::MakeEditorViewportClien
 	return ViewportClient.ToSharedRef();
 }
 
-TSharedPtr<SWidget> SNiagaraBakerViewport::MakeViewportToolbar()
-{
-	return
-		SNew(SNiagaraBakerViewportToolbar)
-		.WeakViewModel(WeakViewModel)
-		.WeakViewport(SharedThis(this));
-	;
-}
-
 void SNiagaraBakerViewport::Tick(const FGeometry& AllottedGeometry, const double InCurrentTime, const float InDeltaTime)
 {
 	ViewportClient->bNeedsRedraw = true;
diff --git a/Engine/Plugins/FX/Niagara/Source/NiagaraEditor/Private/Widgets/SNiagaraBakerViewport.h b/Engine/Plugins/FX/Niagara/Source/NiagaraEditor/Private/Widgets/SNiagaraBakerViewport.h
--- a/Engine/Plugins/FX/Niagara/Source/NiagaraEditor/Private/Widgets/SNiagaraBakerViewport.h
+++ b/Engine/Plugins/FX/Niagara/Source/NiagaraEditor/Private/Widgets/SNiagaraBakerViewport.h
@@ -23,7 +23,7 @@ public:
 	// SEditorViewport interface
 	void Tick(const FGeometry& AllottedGeometry, const double InCurrentTime, const float InDeltaTime) override;
 	TSharedRef<FEditorViewportClient> MakeEditorViewportClient() override;
-	TSharedPtr<SWidget> MakeViewportToolbar() override;
+	//TSharedPtr<SWidget> MakeViewportToolbar() override;
 	// SEditorViewport interface
 
 	// ICommonEditorViewportToolbarInfoProvider interface
diff --git a/Engine/Plugins/FX/Niagara/Source/NiagaraEditor/Private/Widgets/SNiagaraBakerWidget.cpp b/Engine/Plugins/FX/Niagara/Source/NiagaraEditor/Private/Widgets/SNiagaraBakerWidget.cpp
--- a/Engine/Plugins/FX/Niagara/Source/NiagaraEditor/Private/Widgets/SNiagaraBakerWidget.cpp
+++ b/Engine/Plugins/FX/Niagara/Source/NiagaraEditor/Private/Widgets/SNiagaraBakerWidget.cpp
@@ -7,6 +7,7 @@
 #include "NiagaraEditorStyle.h"
 #include "NiagaraComponent.h"
 #include "NiagaraSystem.h"
+#include "Customizations/NiagaraBakerOutputCustomization.h"
 #include "ViewModels/NiagaraBakerViewModel.h"
 
 #include "Modules/ModuleManager.h"
@@ -15,28 +16,19 @@
 #include "Widgets/Layout/SSplitter.h"
 #include "Widgets/SBoxPanel.h"
 #include "Widgets/SOverlay.h"
+#include "Widgets/SToolTip.h"
 #include "Widgets/SViewport.h"
 #include "DetailCategoryBuilder.h"
 #include "DetailLayoutBuilder.h"
 #include "DetailWidgetRow.h"
 #include "EditorWidgetsModule.h"
-#include "ITransportControl.h"
 #include "IDetailCustomization.h"
+#include "IDocumentation.h"
+#include "ITransportControl.h"
 #include "PropertyEditorModule.h"
 
 #define LOCTEXT_NAMESPACE "NiagaraBakerWidget"
 
-//-TODO: Remove details panel and replace with customization to make more compact
-//-TODO: Improve feedback & warning display
-//-TODO: Throttle sim frames on warmup
-//-TODO: Turn off playback when out of focus
-//-TODO: Fluid systems stop simulation with to previous?
-//-TODO: First frame black even though I have non zero start seconds
-//-TODO: None should display HDR or something sensible
-//-TODO: Support particle attribute generation on GPU
-//-TODO: Hook up loading a preview environment / sequence
-//-TODO: Add custom post processes for capturing data?
-
 //////////////////////////////////////////////////////////////////////////
 
 namespace NiagaraBakerWidgetLocal
@@ -51,7 +43,7 @@ namespace NiagaraBakerWidgetLocal
 			.Padding(FMargin(4.0f, 0.0f, 0.0f, 0.0f))
 			.WidthOverride(100.0f)
 			[
-				SNew(SSpinBox<float>)
+				SNew(SSpinBox<TType>)
 				.Font(FEditorStyle::GetFontStyle(TEXT("MenuItem.Font")))
 				.MinValue(MinValue)
 				.MaxValue(MaxValue)
@@ -274,7 +266,7 @@ public:
 
 		// Settings
 		{
-			IDetailCategoryBuilder& DetailCategory = DetailBuilder.EditCategory(FName(TEXT("Settings")));
+			IDetailCategoryBuilder& DetailCategory = DetailBuilder.EditCategory(FName("Settings"));
 
 			DetailCategory.AddProperty(GET_MEMBER_NAME_CHECKED(UNiagaraBakerSettings, StartSeconds));
 			DetailCategory.AddProperty(GET_MEMBER_NAME_CHECKED(UNiagaraBakerSettings, DurationSeconds));
@@ -292,23 +284,29 @@ public:
 			DetailBuilder.HideProperty(GET_MEMBER_NAME_CHECKED(UNiagaraBakerSettings, bPreviewLooping));
 			DetailBuilder.HideProperty(GET_MEMBER_NAME_CHECKED(UNiagaraBakerSettings, bRenderComponentOnly));
 
-			DetailBuilder.HideProperty(GET_MEMBER_NAME_CHECKED(UNiagaraBakerSettings, OutputTextures));
+			//DetailBuilder.HideProperty(GET_MEMBER_NAME_CHECKED(UNiagaraBakerSettings, Outputs));
 		}
 
 		// Output Settings
-		if ( BakerSettings->OutputTextures.IsValidIndex(ViewModel->GetCurrentOutputIndex()) )
+		if ( UNiagaraBakerOutput* CurrentOutput = ViewModel->GetCurrentOutput() )
 		{
-			const FNiagaraBakerTextureSettings& OutputTexture = BakerSettings->OutputTextures[ViewModel->GetCurrentOutputIndex()];
-
 			IDetailCategoryBuilder& OutputTextureCategory = DetailBuilder.EditCategory(FName("OutputSettings"));
 
-			TSharedPtr<IPropertyHandleArray> OutputArrayPropertyHandle = DetailBuilder.GetProperty(GET_MEMBER_NAME_CHECKED(UNiagaraBakerSettings, OutputTextures))->AsArray();
-			TSharedPtr<IPropertyHandle> OutputPropertyHandle = OutputArrayPropertyHandle->GetElement(ViewModel->GetCurrentOutputIndex());
-			OutputTextureCategory.AddProperty(OutputPropertyHandle->GetChildHandle(GET_MEMBER_NAME_CHECKED(FNiagaraBakerTextureSettings, OutputName)));
-			OutputTextureCategory.AddProperty(OutputPropertyHandle->GetChildHandle(GET_MEMBER_NAME_CHECKED(FNiagaraBakerTextureSettings, SourceBinding)));
-			OutputTextureCategory.AddProperty(OutputPropertyHandle->GetChildHandle(GET_MEMBER_NAME_CHECKED(FNiagaraBakerTextureSettings, TextureSize)));
-			OutputTextureCategory.AddProperty(OutputPropertyHandle->GetChildHandle(GET_MEMBER_NAME_CHECKED(FNiagaraBakerTextureSettings, FrameSize)));
-			OutputTextureCategory.AddProperty(OutputPropertyHandle->GetChildHandle(GET_MEMBER_NAME_CHECKED(FNiagaraBakerTextureSettings, GeneratedTexture)));
+			TSharedPtr<IPropertyHandleArray> OutputsArrayPropertyHandle = DetailBuilder.GetProperty(GET_MEMBER_NAME_CHECKED(UNiagaraBakerSettings, Outputs))->AsArray();
+			TSharedPtr<IPropertyHandle> OutputPropertyHandle = OutputsArrayPropertyHandle->GetElement(ViewModel->GetCurrentOutputIndex());
+
+			for (TFieldIterator<FProperty> PropertyIt(CurrentOutput->GetClass()); PropertyIt; ++PropertyIt)
+			{
+				if (PropertyIt->HasAnyPropertyFlags(CPF_Transient | CPF_Deprecated))
+				{
+					continue;
+				}
+				TSharedPtr<IPropertyHandle> PropertyHandle = OutputPropertyHandle->GetChildHandle(PropertyIt->GetFName());
+				if ( PropertyHandle )
+				{
+					OutputTextureCategory.AddProperty(PropertyHandle);
+				}
+			}
 		}
 	}
 
@@ -359,16 +357,16 @@ void SNiagaraBakerWidget::Construct(const FArguments& InArgs)
 	{
 		BakerToolbarBuilder.AddComboButton(
 			FUIAction(),
-			FOnGetContent::CreateSP(this, &SNiagaraBakerWidget::MakeOutputMenu),
+			FOnGetContent::CreateSP(this, &SNiagaraBakerWidget::MakeOutputSelectMenu),
 			TAttribute<FText>::CreateSP(ViewModel, &FNiagaraBakerViewModel::GetCurrentOutputText),
 			LOCTEXT("CurrentOutputTooltip", "Select which output is currently visible in the preview area"),
 			FSlateIcon(FAppStyle::GetAppStyleSetName(), "Icons.Edit")
 		);
-		BakerToolbarBuilder.AddToolBarButton(
-			FUIAction(FExecuteAction::CreateSP(ViewModel, &FNiagaraBakerViewModel::AddOutput)),
-			NAME_None,
+		BakerToolbarBuilder.AddComboButton(
+			FUIAction(),
+			FOnGetContent::CreateSP(this, &SNiagaraBakerWidget::MakeAddOutputMenu),
 			FText::GetEmpty(),
-			LOCTEXT("AddOutputTooltip", "Add a new output"),
+			LOCTEXT("AddOutputTooltip", "Add a new output to bake"),
 			FSlateIcon(FAppStyle::GetAppStyleSetName(), "Icons.Plus")
 		);
 		BakerToolbarBuilder.AddToolBarButton(
@@ -415,9 +413,8 @@ void SNiagaraBakerWidget::Construct(const FArguments& InArgs)
 	DetailsArgs.bHideSelectionTip = true;
 	DetailsArgs.bAllowSearch = false;
 
-	FOnGetDetailCustomizationInstance BakerSettingsCustomizeDetails = FOnGetDetailCustomizationInstance::CreateStatic(&FBakerSettingsDetails::MakeInstance, WeakViewModel);
 	BakerSettingsDetails = PropertyModule.CreateDetailView(DetailsArgs);
-	BakerSettingsDetails->RegisterInstancedCustomPropertyLayout(UNiagaraBakerSettings::StaticClass(), BakerSettingsCustomizeDetails);
+	NiagaraBakerOutputCustomization::RegisterCustomization(BakerSettingsDetails.Get());
 
 	// Transport control args
 	{
@@ -479,6 +476,11 @@ void SNiagaraBakerWidget::Construct(const FArguments& InArgs)
 		]
 		+SVerticalBox::Slot()
 		.AutoHeight()
+		[
+			MakeSettingsWidget()
+		]
+		+SVerticalBox::Slot()
+		.AutoHeight()
 		[
 			SNew(SHorizontalBox)
 			+SHorizontalBox::Slot()
@@ -505,6 +507,106 @@ SNiagaraBakerWidget::~SNiagaraBakerWidget()
 	}
 }
 
+TSharedRef<SWidget> SNiagaraBakerWidget::MakeSettingsWidget()
+{
+	FNiagaraBakerViewModel* ViewModel = WeakViewModel.Pin().Get();
+	check(ViewModel != nullptr);
+
+	const FString DocLink = TEXT("Shared/Editors/NiagaraBaker");
+
+	return
+		SNew(SHorizontalBox)
+		+SHorizontalBox::Slot()
+		.AutoWidth()
+		.VAlign(VAlign_Center)
+		.HAlign(HAlign_Right)
+		[
+			SNew(STextBlock)
+			.Text(LOCTEXT("TimeRange", "Time Range"))
+			.Font(IDetailLayoutBuilder::GetDetailFont())
+		]
+		+SHorizontalBox::Slot()
+		.AutoWidth()
+		[
+			SNew(SSpinBox<float>)
+			.MinDesiredWidth(50.0f)
+			.MinValue(0.0f)
+			.MinSliderValue(0.0f)
+			.Font(IDetailLayoutBuilder::GetDetailFont())
+			.Value(ViewModel, &FNiagaraBakerViewModel::GetTimelineStart)
+			.OnValueChanged(ViewModel, &FNiagaraBakerViewModel::SetTimelineStart)
+		]
+		+SHorizontalBox::Slot()
+		.AutoWidth()
+		[
+			SNew(SSpinBox<float>)
+			.MinDesiredWidth(50.0f)
+			.MinValue(0.0f)
+			.MinSliderValue(0.0f)
+			.Font(IDetailLayoutBuilder::GetDetailFont())
+			.Value(ViewModel, &FNiagaraBakerViewModel::GetTimelineEnd)
+			.OnValueChanged(ViewModel, &FNiagaraBakerViewModel::SetTimelineEnd)
+		]
+
+		+SHorizontalBox::Slot()
+		.AutoWidth()
+		.Padding(10.0f)
+		.VAlign(VAlign_Center)
+		.HAlign(HAlign_Right)
+		[
+			SNew(STextBlock)
+			.Text(LOCTEXT("FramesPerDimension", "Frames Per Dimension"))
+			.Font(IDetailLayoutBuilder::GetDetailFont())
+		]
+		+SHorizontalBox::Slot()
+		.AutoWidth()
+		[
+			SNew(SSpinBox<int32>)
+			.MinDesiredWidth(50.0f)
+			.MinValue(1)
+			.MinSliderValue(1)
+			.Font(IDetailLayoutBuilder::GetDetailFont())
+			.Value(ViewModel, &FNiagaraBakerViewModel::GetFramesOnX)
+			.OnValueChanged(ViewModel, &FNiagaraBakerViewModel::SetFramesOnX)
+		]
+		+SHorizontalBox::Slot()
+		.AutoWidth()
+		[
+			SNew(SSpinBox<int32>)
+			.MinDesiredWidth(50.0f)
+			.MinValue(1)
+			.MinSliderValue(1)
+			.Font(IDetailLayoutBuilder::GetDetailFont())
+			.Value(ViewModel, &FNiagaraBakerViewModel::GetFramesOnY)
+			.OnValueChanged(ViewModel, &FNiagaraBakerViewModel::SetFramesOnY)
+		]
+
+		+SHorizontalBox::Slot()
+		.AutoWidth()
+		.Padding(10.0f)
+		.VAlign(VAlign_Center)
+		.HAlign(HAlign_Right)
+		[
+			SNew(STextBlock)
+			.Text(this, &SNiagaraBakerWidget::GetCurrentTimeText)
+			.Font(IDetailLayoutBuilder::GetDetailFont())
+			.ToolTip(IDocumentation::Get()->CreateToolTip(LOCTEXT("CurrentTimeSecondsTooltip", "Current time in seconds"), nullptr, DocLink, "CurrentTimeText"))
+		]
+
+		+SHorizontalBox::Slot()
+		.AutoWidth()
+		.Padding(10.0f)
+		.VAlign(VAlign_Center)
+		.HAlign(HAlign_Right)
+		[
+			SNew(STextBlock)
+			.Text(this, &SNiagaraBakerWidget::GetCurrentFrameText)
+			.Font(IDetailLayoutBuilder::GetDetailFont())
+			.ToolTip(IDocumentation::Get()->CreateToolTip(LOCTEXT("CurrentFrameTooltip", "Current frame we are rendering, this is the baked frame index not the preview frame"), nullptr, DocLink, "CurrentFrameText"))
+		]
+	;
+}
+
 TSharedRef<SWidget> SNiagaraBakerWidget::MakeCameraModeMenu()
 {
 	FNiagaraBakerViewModel* ViewModel = WeakViewModel.Pin().Get();
@@ -663,7 +765,7 @@ TSharedRef<SWidget> SNiagaraBakerWidget::MakeViewOptionsMenu()
 	return MenuBuilder.MakeWidget();
 }
 
-TSharedRef<SWidget> SNiagaraBakerWidget::MakeOutputMenu()
+TSharedRef<SWidget> SNiagaraBakerWidget::MakeOutputSelectMenu()
 {
 	FNiagaraBakerViewModel* ViewModel = WeakViewModel.Pin().Get();
 	check(ViewModel != nullptr);
@@ -672,7 +774,7 @@ TSharedRef<SWidget> SNiagaraBakerWidget::MakeOutputMenu()
 
 	if ( UNiagaraBakerSettings* BakerSettings = ViewModel->GetBakerSettings() )
 	{
-		for (int i = 0; i < BakerSettings->OutputTextures.Num(); ++i)
+		for (int i = 0; i < BakerSettings->Outputs.Num(); ++i)
 		{
 			MenuBuilder.AddMenuEntry(
 				ViewModel->GetOutputText(i),
@@ -692,67 +794,69 @@ TSharedRef<SWidget> SNiagaraBakerWidget::MakeOutputMenu()
 	return MenuBuilder.MakeWidget();
 }
 
-bool SNiagaraBakerWidget::FindWarnings(TArray<FText>* OutWarnings) const
+TSharedRef<SWidget> SNiagaraBakerWidget::MakeAddOutputMenu()
 {
-	if ( FNiagaraBakerViewModel* ViewModel = WeakViewModel.Pin().Get() )
+	FNiagaraBakerViewModel* ViewModel = WeakViewModel.Pin().Get();
+	check(ViewModel != nullptr);
+
+	FMenuBuilder MenuBuilder(true, nullptr);
+
+	for (TObjectIterator<UClass> It; It; ++It)
 	{
-		// Check determinism
-		if ( UNiagaraComponent* PreviewComponent = ViewModel->GetPreviewComponent() )
+		if ( It->IsChildOf(UNiagaraBakerOutput::StaticClass()) && !It->HasAnyClassFlags(CLASS_Abstract) )
 		{
-			if ( UNiagaraSystem* NiagaraSystem = PreviewComponent->GetAsset() )
-			{
-				if ( NiagaraSystem->NeedsDeterminism() == false )
-				{
-					if ( OutWarnings == nullptr )
-					{
-						return true;
-					}
-					OutWarnings->Emplace(LOCTEXT("SystemDeterminism", "System is not set to deterministic, results will vary each bake"));
-				}
-				for (const FNiagaraEmitterHandle& EmitterHandle : NiagaraSystem->GetEmitterHandles() )
-				{
-					UNiagaraEmitter* NiagaraEmitter = EmitterHandle.GetInstance();
-					if (NiagaraEmitter == nullptr || NiagaraEmitter->bDeterminism)
-					{
-						continue;
-					}
-
-					if (OutWarnings == nullptr)
-					{
-						return true;
-					}
-					OutWarnings->Emplace(FText::Format(LOCTEXT("EmitterDeterminismFormat", "Emitter '{0}' is not set to deterministic, results will vary each bake"), FText::FromString(NiagaraEmitter->GetUniqueEmitterName())));
-				}
-			}
+			MenuBuilder.AddMenuEntry(
+				It->GetDisplayNameText(),
+				FText::GetEmpty(),
+				FSlateIcon(),
+				FUIAction(FExecuteAction::CreateSP(ViewModel, &FNiagaraBakerViewModel::AddOutput, *It))
+			);
 		}
+	}
+
+	return MenuBuilder.MakeWidget();
+}
+
+void SNiagaraBakerWidget::FindWarnings()
+{
+	FoundWarnings.Reset();
+
+	FNiagaraBakerViewModel* ViewModel = WeakViewModel.Pin().Get();
+	if ( ViewModel == nullptr )
+	{
+		return;
+	}
 
-		// Check outputs divide correctly by atlas frame count
-		if ( UNiagaraBakerSettings* BakerSettings = ViewModel->GetBakerSettings() )
+	// Check determinism
+	if ( UNiagaraComponent* PreviewComponent = ViewModel->GetBakerRenderer().GetPreviewComponent() )
+	{
+		if ( UNiagaraSystem* NiagaraSystem = PreviewComponent->GetAsset() )
 		{
-			for ( int32 i=0; i < BakerSettings->OutputTextures.Num(); ++i )
+			if ( NiagaraSystem->NeedsDeterminism() == false )
+			{
+				FoundWarnings.Emplace(LOCTEXT("SystemDeterminism", "System is not set to deterministic, results will vary each bake"));
+			}
+			for (const FNiagaraEmitterHandle& EmitterHandle : NiagaraSystem->GetEmitterHandles() )
 			{
-				const FNiagaraBakerTextureSettings& TextureSettings = BakerSettings->OutputTextures[i];
-				const bool bXSizeIssue = (TextureSettings.TextureSize.X % BakerSettings->FramesPerDimension.X) == 0;
-				const bool bYSizeIssue = (TextureSettings.TextureSize.Y % BakerSettings->FramesPerDimension.Y) == 0;
-				if (bXSizeIssue && bYSizeIssue)
+				UNiagaraEmitter* NiagaraEmitter = EmitterHandle.GetInstance();
+				if (NiagaraEmitter == nullptr || NiagaraEmitter->bDeterminism)
 				{
 					continue;
 				}
 
-				if (OutWarnings == nullptr)
-				{
-					return true;
-				}
-				OutWarnings->Emplace(FText::Format(LOCTEXT("TextureSizeWarningText", "Output '{0}' frames per dimension {1}x{2} does not divide into texture size {3}x{4}, this can result in jitter on the flipbook if not taken into account"),
-					ViewModel->GetOutputText(i),
-					BakerSettings->FramesPerDimension.X, BakerSettings->FramesPerDimension.Y,
-					TextureSettings.TextureSize.X, TextureSettings.TextureSize.Y
-				));
+				FoundWarnings.Emplace(FText::Format(LOCTEXT("EmitterDeterminismFormat", "Emitter '{0}' is not set to deterministic, results will vary each bake"), FText::FromString(NiagaraEmitter->GetUniqueEmitterName())));
 			}
 		}
 	}
 
-	return OutWarnings ? OutWarnings->Num() > 0 : false;
+	// Check outputs for any warnings
+	if ( UNiagaraBakerSettings* BakerSettings = ViewModel->GetBakerSettings() )
+	{
+		for ( UNiagaraBakerOutput* BakerOutput : BakerSettings->Outputs )
+		{
+			BakerOutput->FindWarnings(FoundWarnings);
+		}
+	}
 }
 
 TSharedRef<SWidget> SNiagaraBakerWidget::MakeWarningsMenu()
@@ -762,18 +866,14 @@ TSharedRef<SWidget> SNiagaraBakerWidget::MakeWarningsMenu()
 
 	FMenuBuilder MenuBuilder(true, nullptr);
 
-	TArray<FText> FoundWarnings;
-	if ( FindWarnings(&FoundWarnings) )
+	for (const FText& Warning : FoundWarnings)
 	{
-		for (const FText& Warning : FoundWarnings)
-		{
-			MenuBuilder.AddMenuEntry(
-				Warning,
-				FText::GetEmpty(),
-				FSlateIcon(),
-				FUIAction()
-			);
-		}
+		MenuBuilder.AddMenuEntry(
+			Warning,
+			FText::GetEmpty(),
+			FSlateIcon(),
+			FUIAction()
+		);
 	}
 
 	return MenuBuilder.MakeWidget();
@@ -781,7 +881,12 @@ TSharedRef<SWidget> SNiagaraBakerWidget::MakeWarningsMenu()
 
 void SNiagaraBakerWidget::RefreshWidget()
 {
-	BakerSettingsDetails->SetObject(GetBakerSettings(), true);
+	UObject* OutputObject = nullptr;
+	if ( FNiagaraBakerViewModel* ViewModel = WeakViewModel.Pin().Get() )
+	{
+		OutputObject = ViewModel->GetCurrentOutput();
+	}
+	BakerSettingsDetails->SetObject(OutputObject, true);
 }
 
 void SNiagaraBakerWidget::Tick(const FGeometry& AllottedGeometry, const double CurrentTime, const float DeltaTime)
@@ -816,6 +921,8 @@ void SNiagaraBakerWidget::Tick(const FGeometry& AllottedGeometry, const double C
 		ViewportWidget->RefreshView(PreviewRelativeTime, DeltaTime);
 		TimelineWidget->SetRelativeTime(PreviewRelativeTime);
 	}
+
+	FindWarnings();
 }
 
 void SNiagaraBakerWidget::SetPreviewRelativeTime(float RelativeTime)
@@ -862,15 +969,11 @@ FReply SNiagaraBakerWidget::OnTransportBackwardEnd()
 FReply SNiagaraBakerWidget::OnTransportBackwardStep()
 {
 	bIsPlaying = false;
-	if (auto ViewModel = WeakViewModel.Pin())
+	if (FNiagaraBakerViewModel* ViewModel = WeakViewModel.Pin().Get())
 	{
-		if(UNiagaraBakerSettings* Settings = ViewModel->GetBakerSettings())
-		{
-			const auto DisplayData = Settings->GetDisplayInfo(PreviewRelativeTime, Settings->bPreviewLooping);
-			const int NumFrames = Settings->GetNumFrames();
-			const int NewFrame = DisplayData.Interp > 0.25f ? DisplayData.FrameIndexA : FMath::Max(DisplayData.FrameIndexA - 1, 0);
-			PreviewRelativeTime = (Settings->DurationSeconds / float(NumFrames)) * float(NewFrame);
-		}
+		const FNiagaraBakerOutputFrameIndices OutputFrameIndices = ViewModel->GetCurrentOutputFrameIndices(PreviewRelativeTime);
+		const int NewFrame = OutputFrameIndices.Interp > 0.25f ? OutputFrameIndices.FrameIndexA : FMath::Max(OutputFrameIndices.FrameIndexA - 1, 0);
+		PreviewRelativeTime = (ViewModel->GetDurationSeconds() / float(OutputFrameIndices.NumFrames)) * float(NewFrame);
 	}
 
 	return FReply::Handled();
@@ -885,15 +988,11 @@ FReply SNiagaraBakerWidget::OnTransportForwardPlay()
 FReply SNiagaraBakerWidget::OnTransportForwardStep()
 {
 	bIsPlaying = false;
-	if (auto ViewModel = WeakViewModel.Pin())
+	if (FNiagaraBakerViewModel* ViewModel = WeakViewModel.Pin().Get())
 	{
-		if (UNiagaraBakerSettings* Settings = ViewModel->GetBakerSettings())
-		{
-			const auto DisplayData = Settings->GetDisplayInfo(PreviewRelativeTime, Settings->bPreviewLooping);
-			const int NumFrames = Settings->GetNumFrames();
-			const int NewFrame = DisplayData.Interp < 0.75f ? DisplayData.FrameIndexB : FMath::Min(DisplayData.FrameIndexB + 1, NumFrames - 1);
-			PreviewRelativeTime = (Settings->DurationSeconds / float(NumFrames)) * float(NewFrame);
-		}
+		const FNiagaraBakerOutputFrameIndices OutputFrameIndices = ViewModel->GetCurrentOutputFrameIndices(PreviewRelativeTime);
+		const int NewFrame = OutputFrameIndices.Interp < 0.75f ? OutputFrameIndices.FrameIndexB : FMath::Min(OutputFrameIndices.FrameIndexB + 1, OutputFrameIndices.NumFrames - 1);
+		PreviewRelativeTime = (ViewModel->GetDurationSeconds() / float(OutputFrameIndices.NumFrames)) * float(NewFrame);
 	}
 
 	return FReply::Handled();
@@ -904,7 +1003,7 @@ FReply SNiagaraBakerWidget::OnTransportForwardEnd()
 	bIsPlaying = false;
 	if (UNiagaraBakerSettings* BakerSettings = GetBakerSettings())
 	{
-		PreviewRelativeTime = BakerSettings->DurationSeconds;
+		PreviewRelativeTime = BakerSettings->DurationSeconds - KINDA_SMALL_NUMBER;
 	}
 
 	return FReply::Handled();
@@ -919,4 +1018,28 @@ FReply SNiagaraBakerWidget::OnTransportToggleLooping() const
 	return FReply::Handled();
 }
 
+FText SNiagaraBakerWidget::GetCurrentTimeText() const
+{
+	if (UNiagaraBakerSettings* BakerSettings = GetBakerSettings())
+	{
+		const float TimeSeconds = BakerSettings->StartSeconds + PreviewRelativeTime;
+		const FNumberFormattingOptions FormatOptions = FNumberFormattingOptions()
+			.SetMaximumFractionalDigits(2)
+			.SetMinimumFractionalDigits(2);
+		return FText::Format(LOCTEXT("CurrentTimeTextFormat", "{0} s"), FText::AsNumber(TimeSeconds, &FormatOptions));
+	}
+	return FText::GetEmpty();
+}
+
+FText SNiagaraBakerWidget::GetCurrentFrameText() const
+{
+	if ( FNiagaraBakerViewModel* ViewModel = WeakViewModel.Pin().Get() )
+	{
+		FNiagaraBakerOutputFrameIndices FrameIndices = ViewModel->GetCurrentOutputFrameIndices(PreviewRelativeTime);
+		return FText::Format(LOCTEXT("CurrentFrameTextFormat", "{0} / {1} f"), FText::AsNumber(FrameIndices.FrameIndexA), FText::AsNumber(FrameIndices.NumFrames));
+	}
+
+	return FText::GetEmpty();
+}
+
 #undef LOCTEXT_NAMESPACE
diff --git a/Engine/Plugins/FX/Niagara/Source/NiagaraEditor/Private/Widgets/SNiagaraBakerWidget.h b/Engine/Plugins/FX/Niagara/Source/NiagaraEditor/Private/Widgets/SNiagaraBakerWidget.h
--- a/Engine/Plugins/FX/Niagara/Source/NiagaraEditor/Private/Widgets/SNiagaraBakerWidget.h
+++ b/Engine/Plugins/FX/Niagara/Source/NiagaraEditor/Private/Widgets/SNiagaraBakerWidget.h
@@ -24,12 +24,15 @@ public:
 
 	FReply OnCapture();
 
+	TSharedRef<SWidget> MakeSettingsWidget();
+
 	TSharedRef<SWidget> MakeCameraModeMenu();
 	TSharedRef<SWidget> MakeViewOptionsMenu();
-	TSharedRef<SWidget> MakeOutputMenu();
+	TSharedRef<SWidget> MakeOutputSelectMenu();
+	TSharedRef<SWidget> MakeAddOutputMenu();
 
-	bool FindWarnings(TArray<FText>* OutWarnings) const;
-	bool HasWarnings() const { return FindWarnings(nullptr); }
+	void FindWarnings();
+	bool HasWarnings() const { return FoundWarnings.Num() > 0; }
 	TSharedRef<SWidget> MakeWarningsMenu();
 
 	void RefreshWidget();
@@ -46,6 +49,9 @@ public:
 	FReply OnTransportForwardEnd();
 	FReply OnTransportToggleLooping() const;
 
+	FText GetCurrentTimeText() const;
+	FText GetCurrentFrameText() const;
+
 private:
 	TWeakPtr<FNiagaraBakerViewModel>			WeakViewModel;
 	TSharedPtr<SNiagaraBakerViewport>			ViewportWidget;
@@ -57,4 +63,6 @@ private:
 	float										PreviewRelativeTime = 0.0f;
 
 	FDelegateHandle								OnCurrentOutputIndexChangedHandle;
+
+	TArray<FText>								FoundWarnings;
 };
diff --git a/Engine/Plugins/FX/Niagara/Source/NiagaraShader/Private/NiagaraBatchedElements.cpp b/Engine/Plugins/FX/Niagara/Source/NiagaraShader/Private/NiagaraBatchedElements.cpp
--- a/Engine/Plugins/FX/Niagara/Source/NiagaraShader/Private/NiagaraBatchedElements.cpp
+++ b/Engine/Plugins/FX/Niagara/Source/NiagaraShader/Private/NiagaraBatchedElements.cpp
@@ -74,6 +74,86 @@ void FBatchedElementNiagara2DArrayAttribute::BindShaders(FRHICommandList& RHICmd
 
 //////////////////////////////////////////////////////////////////////////
 
+class FNiagaraSimpleElementVolumeAttribute : public FGlobalShader
+{
+public:
+	DECLARE_SHADER_TYPE(FNiagaraSimpleElementVolumeAttribute, Global);
+	SHADER_USE_PARAMETER_STRUCT(FNiagaraSimpleElementVolumeAttribute, FGlobalShader);
+
+	static bool ShouldCompilePermutation(const FGlobalShaderPermutationParameters& Parameters)
+	{
+		return true;
+	}
+
+	static void ModifyCompilationEnvironment(const FGlobalShaderPermutationParameters& Parameters, FShaderCompilerEnvironment& OutEnvironment)
+	{
+		FGlobalShader::ModifyCompilationEnvironment(Parameters, OutEnvironment);
+		OutEnvironment.SetDefine(TEXT("NIAGARA_VOLUME_ATTRIBUTE_PS"), 1);
+	}
+
+
+	BEGIN_SHADER_PARAMETER_STRUCT(FParameters, )
+		SHADER_PARAMETER(int,					InNumAttributes)
+		SHADER_PARAMETER_ARRAY(FVector4f,		InAttributeUVAndChannel, [4])
+		SHADER_PARAMETER(FVector2f,				InTileUV)
+		SHADER_PARAMETER(FMatrix44f,			InColorWeights)
+		SHADER_PARAMETER(float,					InGamma)
+		SHADER_PARAMETER_TEXTURE(Texture2D,		InTexture)
+		SHADER_PARAMETER_SAMPLER(SamplerState,	InTextureSampler)
+	END_SHADER_PARAMETER_STRUCT()
+};
+
+IMPLEMENT_GLOBAL_SHADER(FNiagaraSimpleElementVolumeAttribute, "/Plugin/FX/Niagara/Private/NiagaraBatchedElements.usf", "MainPS", SF_Pixel);
+
+void FBatchedElementNiagaraVolumeAttribute::BindShaders(FRHICommandList& RHICmdList, FGraphicsPipelineStateInitializer& GraphicsPSOInit, ERHIFeatureLevel::Type InFeatureLevel, const FMatrix& InTransform, const float InGamma, const FMatrix& ColorWeights, const FTexture* Texture)
+{
+	FRHITexture* RHITexture = nullptr;
+	FRHISamplerState* RHISamplerState = nullptr;
+	if (GetTextureAndSampler != nullptr)
+	{
+		GetTextureAndSampler(RHITexture, RHISamplerState);
+	}
+	else if (Texture != nullptr)
+	{
+		RHITexture = Texture->TextureRHI;
+		RHISamplerState = Texture->SamplerStateRHI;
+	}
+
+	TShaderMapRef<FSimpleElementVS> VertexShader(GetGlobalShaderMap(InFeatureLevel));
+	TShaderMapRef<FNiagaraSimpleElementVolumeAttribute> PixelShader(GetGlobalShaderMap(InFeatureLevel));
+
+	GraphicsPSOInit.BoundShaderState.VertexDeclarationRHI = GSimpleElementVertexDeclaration.VertexDeclarationRHI;
+	GraphicsPSOInit.BoundShaderState.VertexShaderRHI = VertexShader.GetVertexShader();
+	GraphicsPSOInit.BoundShaderState.PixelShaderRHI = PixelShader.GetPixelShader();
+	GraphicsPSOInit.PrimitiveType = PT_TriangleList;
+	GraphicsPSOInit.BlendState = TStaticBlendState<>::GetRHI();
+
+	RHICmdList.ApplyCachedRenderTargets(GraphicsPSOInit);
+	SetGraphicsPipelineState(RHICmdList, GraphicsPSOInit, 0);
+
+	VertexShader->SetParameters(RHICmdList, InTransform);
+
+	FNiagaraSimpleElementVolumeAttribute::FParameters PassParameters;
+	PassParameters.InNumAttributes	= AttributeUVs.Num();
+	for (int32 i=0; i < 4; ++i)
+	{
+		const FVector3f AttributeUV = AttributeUVs.IsValidIndex(i) ? AttributeUVs[i] : FVector3f::ZeroVector;
+		const int AttributeChannel = AttributeChannels.IsValidIndex(i) ? AttributeChannels[i] : 0;
+		PassParameters.InAttributeUVAndChannel[i].X = AttributeUV.X;
+		PassParameters.InAttributeUVAndChannel[i].Y = AttributeUV.Y;
+		PassParameters.InAttributeUVAndChannel[i].Z = AttributeUV.Z;
+		reinterpret_cast<int32&>(PassParameters.InAttributeUVAndChannel[i].W) = AttributeChannel;
+	}
+	PassParameters.InTileUV			= TileUVs;
+	PassParameters.InColorWeights	= FMatrix44f(ColorWeights);
+	PassParameters.InGamma			= InGamma;
+	PassParameters.InTexture		= RHITexture ? RHITexture : GWhiteTexture->TextureRHI.GetReference();
+	PassParameters.InTextureSampler	= RHISamplerState ? RHISamplerState : TStaticSamplerState<SF_Bilinear>::GetRHI();
+	SetShaderParameters(RHICmdList, PixelShader, PixelShader.GetPixelShader(), PassParameters);
+}
+
+//////////////////////////////////////////////////////////////////////////
+
 class FNiagaraSimpleElementPS : public FGlobalShader
 {
 	DECLARE_GLOBAL_SHADER(FNiagaraSimpleElementPS);
diff --git a/Engine/Plugins/FX/Niagara/Source/NiagaraShader/Public/NiagaraBatchedElements.h b/Engine/Plugins/FX/Niagara/Source/NiagaraShader/Public/NiagaraBatchedElements.h
--- a/Engine/Plugins/FX/Niagara/Source/NiagaraShader/Public/NiagaraBatchedElements.h
+++ b/Engine/Plugins/FX/Niagara/Source/NiagaraShader/Public/NiagaraBatchedElements.h
@@ -26,6 +26,33 @@ private:
 	FGetTextureAndSamplerDelegate GetTextureAndSampler;
 };
 
+/**
+ * Batched element parameters for gathering attributes from a volume texture into a single color
+ */
+class NIAGARASHADER_API FBatchedElementNiagaraVolumeAttribute : public FBatchedElementParameters
+{
+public:
+	typedef TFunction<void(FRHITexture*&, FRHISamplerState*&)> FGetTextureAndSamplerDelegate;
+
+	FBatchedElementNiagaraVolumeAttribute(const FVector2f& InTileUVs, TConstArrayView<FVector3f> InAttributeUVs, TConstArrayView<int32> InAttributeChannels, FGetTextureAndSamplerDelegate InGetTextureAndSampler = nullptr)
+		: TileUVs(InTileUVs)
+		, AttributeUVs(InAttributeUVs)
+		, AttributeChannels(InAttributeChannels)
+		, GetTextureAndSampler(InGetTextureAndSampler)
+	{
+	}
+
+	/** Binds vertex and pixel shaders for this element */
+	virtual void BindShaders(FRHICommandList& RHICmdList, FGraphicsPipelineStateInitializer& GraphicsPSOInit, ERHIFeatureLevel::Type InFeatureLevel, const FMatrix& InTransform, const float InGamma, const FMatrix& ColorWeights, const FTexture* Texture) override;
+
+private:
+	FVector2f								TileUVs;
+	TArray<FVector3f, TInlineAllocator<4>>	AttributeUVs;
+	TArray<int32, TInlineAllocator<4>>		AttributeChannels;
+
+	FGetTextureAndSamplerDelegate			GetTextureAndSampler;
+};
+
 /**
  * Batched element parameters for inverting a color channel
  */
