removed direct camera access from NiagaraWorldManager

diff --git a/Engine/Plugins/FX/Niagara/Source/Niagara/Private/NiagaraWorldManager.cpp b/Engine/Plugins/FX/Niagara/Source/Niagara/Private/NiagaraWorldManager.cpp
--- a/Engine/Plugins/FX/Niagara/Source/Niagara/Private/NiagaraWorldManager.cpp
+++ b/Engine/Plugins/FX/Niagara/Source/Niagara/Private/NiagaraWorldManager.cpp
@@ -26,6 +26,7 @@
 #include "NiagaraCullProxyComponent.h"
 #include "GameDelegates.h"
 #include "DrawDebugHelpers.h"
+#include "Engine/LocalPlayer.h"
 
 #if WITH_EDITORONLY_DATA
 #include "EditorViewportClient.h"
@@ -948,34 +949,32 @@ void FNiagaraWorldManager::MarkSimulationsForEndOfFrameWait(FNiagaraSystemSimula
 	}
 }
 
-#ifndef WITH_CUSTOM_NIAGARA_VIEWINFO
 bool FNiagaraWorldManager::PrepareCachedViewInfo(const APlayerController* PlayerController, FNiagaraCachedViewInfo& OutViewInfo)
 {
-	if (PlayerController && PlayerController->PlayerCameraManager)
-	{
-		FVector POVLoc;
-		FRotator POVRotation;
-		PlayerController->GetPlayerViewPoint(POVLoc, POVRotation);
-		FRotationTranslationMatrix ViewToWorld(POVRotation, POVLoc);
-		FMatrix ViewMat = ViewToWorld.InverseFast() * FMatrix(
-			FPlane(0, 0, 1, 0),
-			FPlane(1, 0, 0, 0),
-			FPlane(0, 1, 0, 0),
-			FPlane(0, 0, 0, 1));
-
-		FWorldCachedViewInfo WorldViewInfo;
-		WorldViewInfo.ViewMatrix = ViewMat;
-		WorldViewInfo.ProjectionMatrix = PlayerController->PlayerCameraManager->GetCameraCacheView().CalculateProjectionMatrix();
-		WorldViewInfo.ViewProjectionMatrix = ViewMat * WorldViewInfo.ProjectionMatrix;
-		WorldViewInfo.ViewToWorld = ViewToWorld;
-
-		OutViewInfo.Init(WorldViewInfo);
-		return true;
+	const ULocalPlayer* LocalPlayer = PlayerController ? PlayerController->GetLocalPlayer() : nullptr;
+	if (PlayerController && LocalPlayer && LocalPlayer->ViewportClient)
+	{
+		FSceneViewProjectionData ProjectionData;
+		if (LocalPlayer->GetProjectionData(LocalPlayer->ViewportClient->Viewport, /*out*/ ProjectionData))
+		{
+			FVector POVLoc;
+			FRotator POVRotation;
+			PlayerController->GetPlayerViewPoint(POVLoc, POVRotation);
+			FRotationTranslationMatrix ViewToWorld(POVRotation, POVLoc);
+
+			FWorldCachedViewInfo WorldViewInfo;
+			WorldViewInfo.ViewMatrix = FTranslationMatrix(-ProjectionData.ViewOrigin) * ProjectionData.ViewRotationMatrix;
+			WorldViewInfo.ProjectionMatrix = ProjectionData.ProjectionMatrix;
+			WorldViewInfo.ViewProjectionMatrix = ProjectionData.ComputeViewProjectionMatrix();
+			WorldViewInfo.ViewToWorld = ViewToWorld;
+
+			OutViewInfo.Init(WorldViewInfo);
+			return true;
+		}
 	}
 
 	return false;
 }
-#endif // WITH_CUSTOM_NIAGARA_VIEWINFO
 
 void FNiagaraWorldManager::Tick(ETickingGroup TickGroup, float DeltaSeconds, ELevelTick TickType, ENamedThreads::Type CurrentThread, const FGraphEventRef& MyCompletionGraphEvent)
 {
