- Fixes for API changes to FCanvas and FSceneViewFamilyContext

#rnx

diff --git a/Engine/Plugins/FX/Niagara/Source/NiagaraEditor/Private/NiagaraBakerRenderer.cpp b/Engine/Plugins/FX/Niagara/Source/NiagaraEditor/Private/NiagaraBakerRenderer.cpp
--- a/Engine/Plugins/FX/Niagara/Source/NiagaraEditor/Private/NiagaraBakerRenderer.cpp
+++ b/Engine/Plugins/FX/Niagara/Source/NiagaraEditor/Private/NiagaraBakerRenderer.cpp
@@ -347,7 +347,7 @@ void FNiagaraBakerRenderer::RenderSceneCapture(UTextureRenderTarget2D* RenderTar
 	const float WorldTime = GetWorldTime();
 	UWorld* World = GetWorld();
 
-	FCanvas Canvas(RenderTarget->GameThread_GetRenderTargetResource(), nullptr, WorldTime, FApp::GetDeltaTime(), WorldTime, GetFeatureLevel());
+	FCanvas Canvas(RenderTarget->GameThread_GetRenderTargetResource(), nullptr, FGameTime::CreateUndilated(WorldTime, FApp::GetDeltaTime()), GetFeatureLevel());
 	Canvas.Clear(FLinearColor::Black);
 
 	SceneCaptureComponent->RegisterComponentWithWorld(World);
@@ -406,13 +406,13 @@ void FNiagaraBakerRenderer::RenderBufferVisualization(UTextureRenderTarget2D* Re
 	const float WorldTime = GetWorldTime();
 	UWorld* World = GetWorld();
 
-	FCanvas Canvas(RenderTarget->GameThread_GetRenderTargetResource(), nullptr, WorldTime, FApp::GetDeltaTime(), WorldTime, GetFeatureLevel());
+	FCanvas Canvas(RenderTarget->GameThread_GetRenderTargetResource(), nullptr, FGameTime::CreateUndilated(WorldTime, FApp::GetDeltaTime()), GetFeatureLevel());
 	Canvas.Clear(FLinearColor::Black);
 
 	// Create View Family
 	FSceneViewFamilyContext ViewFamily(
 		FSceneViewFamily::ConstructionValues(RenderTarget->GameThread_GetRenderTargetResource(), World->Scene, FEngineShowFlags(ESFIM_Game))
-		.SetWorldTimes(WorldTime, FApp::GetDeltaTime(), WorldTime)
+		.SetTime(FGameTime::CreateUndilated(WorldTime, FApp::GetDeltaTime()))
 		.SetGammaCorrection(GammaCorrection)
 	);
 	
@@ -457,7 +457,7 @@ void FNiagaraBakerRenderer::RenderBufferVisualization(UTextureRenderTarget2D* Re
 void FNiagaraBakerRenderer::RenderDataInterface(UTextureRenderTarget2D* RenderTarget, FName BindingName) const
 {
 	const float WorldTime = GetWorldTime();
-	FCanvas Canvas(RenderTarget->GameThread_GetRenderTargetResource(), nullptr, WorldTime, FApp::GetDeltaTime(), WorldTime, GetFeatureLevel());
+	FCanvas Canvas(RenderTarget->GameThread_GetRenderTargetResource(), nullptr, FGameTime::CreateUndilated(WorldTime, FApp::GetDeltaTime()), GetFeatureLevel());
 	Canvas.Clear(FLinearColor::Black);
 
 	ON_SCOPE_EXIT{ Canvas.Flush_GameThread(); };
@@ -511,7 +511,7 @@ void FNiagaraBakerRenderer::RenderDataInterface(UTextureRenderTarget2D* RenderTa
 void FNiagaraBakerRenderer::RenderParticleAttribute(UTextureRenderTarget2D* RenderTarget, FName BindingName) const
 {
 	const float WorldTime = GetWorldTime();
-	FCanvas Canvas(RenderTarget->GameThread_GetRenderTargetResource(), nullptr, WorldTime, FApp::GetDeltaTime(), WorldTime, GetFeatureLevel());
+	FCanvas Canvas(RenderTarget->GameThread_GetRenderTargetResource(), nullptr, FGameTime::CreateUndilated(WorldTime, FApp::GetDeltaTime()), GetFeatureLevel());
 	Canvas.Clear(FLinearColor::Black);
 
 	ON_SCOPE_EXIT { Canvas.Flush_GameThread(); };
diff --git a/Engine/Plugins/FX/Niagara/Source/NiagaraEditor/Private/NiagaraBakerRendererOutputTexture2D.cpp b/Engine/Plugins/FX/Niagara/Source/NiagaraEditor/Private/NiagaraBakerRendererOutputTexture2D.cpp
--- a/Engine/Plugins/FX/Niagara/Source/NiagaraEditor/Private/NiagaraBakerRendererOutputTexture2D.cpp
+++ b/Engine/Plugins/FX/Niagara/Source/NiagaraEditor/Private/NiagaraBakerRendererOutputTexture2D.cpp
@@ -91,7 +91,7 @@ void FNiagaraBakerRendererOutputTexture2D::RenderGenerated(UNiagaraBakerOutput*
 	}
 
 	const float WorldTime = BakerRenderer.GetWorldTime();
-	FCanvas Canvas(RenderTarget->GameThread_GetRenderTargetResource(), nullptr, WorldTime, FApp::GetDeltaTime(), WorldTime, BakerRenderer.GetFeatureLevel());
+	FCanvas Canvas(RenderTarget->GameThread_GetRenderTargetResource(), nullptr, FGameTime::CreateUndilated(WorldTime, FApp::GetDeltaTime()), BakerRenderer.GetFeatureLevel());
 
 	const FNiagaraBakerOutputFrameIndices FrameIndices = BakerGeneratedSettings->GetOutputFrameIndices(BakerOutput, WorldTime);
 
diff --git a/Engine/Plugins/FX/Niagara/Source/NiagaraEditor/Private/NiagaraBakerRendererOutputVolumeTexture.cpp b/Engine/Plugins/FX/Niagara/Source/NiagaraEditor/Private/NiagaraBakerRendererOutputVolumeTexture.cpp
--- a/Engine/Plugins/FX/Niagara/Source/NiagaraEditor/Private/NiagaraBakerRendererOutputVolumeTexture.cpp
+++ b/Engine/Plugins/FX/Niagara/Source/NiagaraEditor/Private/NiagaraBakerRendererOutputVolumeTexture.cpp
@@ -156,7 +156,7 @@ void FNiagaraBakerRendererOutputVolumeTexture::RenderPreview(UNiagaraBakerOutput
 	UNiagaraBakerOutputVolumeTexture* OutputVolumeTexture = CastChecked<UNiagaraBakerOutputVolumeTexture>(InBakerOutput);
 
 	const float WorldTime = BakerRenderer.GetWorldTime();
-	FCanvas Canvas(RenderTarget->GameThread_GetRenderTargetResource(), nullptr, WorldTime, FApp::GetDeltaTime(), WorldTime, BakerRenderer.GetFeatureLevel());
+	FCanvas Canvas(RenderTarget->GameThread_GetRenderTargetResource(), nullptr, FGameTime::CreateUndilated(WorldTime, FApp::GetDeltaTime()), BakerRenderer.GetFeatureLevel());
 	Canvas.Clear(FLinearColor::Black);
 	ON_SCOPE_EXIT{ Canvas.Flush_GameThread(); };
 
@@ -324,7 +324,7 @@ void FNiagaraBakerRendererOutputVolumeTexture::RenderGenerated(UNiagaraBakerOutp
 	}
 
 	const float WorldTime = BakerRenderer.GetWorldTime();
-	FCanvas Canvas(RenderTarget->GameThread_GetRenderTargetResource(), nullptr, WorldTime, FApp::GetDeltaTime(), WorldTime, BakerRenderer.GetFeatureLevel());
+	FCanvas Canvas(RenderTarget->GameThread_GetRenderTargetResource(), nullptr, FGameTime::CreateUndilated(WorldTime, FApp::GetDeltaTime()), BakerRenderer.GetFeatureLevel());
 	Canvas.Clear(FLinearColor::Black);
 
 	const FIntVector AtlasTextureSize(AtlasVolumeTexture->GetSizeX(), AtlasVolumeTexture->GetSizeY(), AtlasVolumeTexture->GetSizeZ());
