Added fallback to match pins via name if no pin could be found via persistent guid

diff --git a/Engine/Plugins/FX/Niagara/Source/NiagaraEditor/Private/NiagaraNode.cpp b/Engine/Plugins/FX/Niagara/Source/NiagaraEditor/Private/NiagaraNode.cpp
--- a/Engine/Plugins/FX/Niagara/Source/NiagaraEditor/Private/NiagaraNode.cpp
+++ b/Engine/Plugins/FX/Niagara/Source/NiagaraEditor/Private/NiagaraNode.cpp
@@ -218,7 +218,7 @@ bool UNiagaraNode::ReallocatePins(bool bMarkNeedsResynchronizeOnChange)
 
 		// When matching pins, use the pin id if either pin id is valid, otherwise match by name.
 		// we typically don't allow orphaned pins to find a match, unless types match each other again
-		auto PinMatchPredicate = [&](UEdGraphPin* Pin) 
+		auto PinMatchPredicateGuid = [&](UEdGraphPin* Pin) 
 		{
 			// early out conditions for when matching up is not wanted
 			if (Pin->Direction != OldPin->Direction)
@@ -226,13 +226,32 @@ bool UNiagaraNode::ReallocatePins(bool bMarkNeedsResynchronizeOnChange)
 				return false;
 			}
 			
-			bool bPinsHaveIds = Pin->PersistentGuid.IsValid() || OldPin->PersistentGuid.IsValid();
-			bool bAllowPin = OldPin->bOrphanedPin == false || OldPin->PinType == Pin->PinType;
-			return bAllowPin && ((bPinsHaveIds && Pin->PersistentGuid == OldPin->PersistentGuid) || (bPinsHaveIds == false && Pin->PinName == OldPin->PinName));
+			bool bPinsHaveIds = Pin->PersistentGuid.IsValid() && OldPin->PersistentGuid.IsValid();
+			bool bAllowPin = OldPin->bOrphanedPin == false && OldPin->PinType == Pin->PinType;
+			return bAllowPin && (bPinsHaveIds && Pin->PersistentGuid == OldPin->PersistentGuid);
+		};
+
+		auto PinMatchPredicateName = [&](UEdGraphPin* Pin) 
+		{
+			// early out conditions for when matching up is not wanted
+			if (Pin->Direction != OldPin->Direction)
+			{
+				return false;
+			}
+			
+			bool bAllowPin = OldPin->bOrphanedPin == false && OldPin->PinType == Pin->PinType;
+			return bAllowPin && Pin->PinName == OldPin->PinName;
 		};
 
 		OldPin->Modify();
-		if (UEdGraphPin** MatchingNewPin = Pins.FindByPredicate(PinMatchPredicate))
+		UEdGraphPin** MatchingNewPin = Pins.FindByPredicate(PinMatchPredicateGuid);
+
+		if(MatchingNewPin == nullptr)
+		{
+			MatchingNewPin = Pins.FindByPredicate(PinMatchPredicateName);
+		}
+		
+		if (MatchingNewPin)
 		{
 			// If the pin types don't match, CopyPersistentDataFromOldPin could very well overwrite our Matching pin with bad data.
 			// Let's cache it for now and reset it after copying the other relevant data off the old pin.
diff --git a/Engine/Source/Runtime/Core/Private/UObject/DevObjectVersion.cpp b/Engine/Source/Runtime/Core/Private/UObject/DevObjectVersion.cpp
--- a/Engine/Source/Runtime/Core/Private/UObject/DevObjectVersion.cpp
+++ b/Engine/Source/Runtime/Core/Private/UObject/DevObjectVersion.cpp
@@ -121,7 +121,7 @@ FDevSystemGuids::FDevSystemGuids()
 	, MATERIALSHADERMAP_DERIVEDDATA_VER(0x2579AAFE, 0x8F1D4E4F, 0xB04278C6, 0x35917C17)
 	, NANITE_DERIVEDDATA_VER(0xBEB0226A, 0x070E4ECA, 0x972C1E7D, 0xD8599E68)
 	, NIAGARASHADERMAP_DERIVEDDATA_VER(0x7BBD9913, 0xC1554D20, 0xADAE9F17, 0xB006299E)
-	, Niagara_LatestScriptCompileVersion(0x6D32B8EE, 0x909FCA7E, 0xA5CE4F17, 0x066A5F25)
+	, Niagara_LatestScriptCompileVersion(0x6D32B8CC, 0x90AFCB7E, 0xA5CE5F17, 0x126A5F34)
 	, SkeletalMeshDerivedDataVersion(0x9B5F4544, 0x76D7481C, 0x9AD3F614, 0xA6C07904)
 	, STATICMESH_DERIVEDDATA_VER(0x5B7A05A7, 0x7371440C, 0xA314621A, 0x9B1E08BC)
 	, POSESEARCHDB_DERIVEDDATA_VER(0x389117E4, 0x807A4CC0, 0x9F37C2E6, 0xD808A78D)
