Niagara - Updated from Mateo's change - Added fallback to match pins via name and type if no pin could be found via persistent guid.  Also removed code which refreshed variable ids on refresh since it's not compatible with applying script changes.

diff --git a/Engine/Plugins/FX/Niagara/Source/NiagaraEditor/Private/NiagaraGraph.cpp b/Engine/Plugins/FX/Niagara/Source/NiagaraEditor/Private/NiagaraGraph.cpp
--- a/Engine/Plugins/FX/Niagara/Source/NiagaraEditor/Private/NiagaraGraph.cpp
+++ b/Engine/Plugins/FX/Niagara/Source/NiagaraEditor/Private/NiagaraGraph.cpp
@@ -1083,26 +1083,6 @@ void UNiagaraGraph::BeginDestroy()
 	ReleaseCompilationCopy();
 }
 
-void UNiagaraGraph::PostDuplicate(EDuplicateMode::Type DuplicateMode)
-{
-	Super::PostDuplicate(DuplicateMode);
-
-	// when duplicating a module script, we need to change the guids of the script variables in the graph
-	// the IsAsset() check makes sure that we don't change anything when opening a graph to be edited, as it's duplicated there as well
-	if (UNiagaraScript* Outer = GetTypedOuter<UNiagaraScript>(); DuplicateMode == EDuplicateMode::Normal && Outer && Outer->IsAsset())
-	{
-		ensure(bIsForCompilationOnly == false);
-		
-		for (auto It : VariableToScriptVariable)
-		{
-			if (TObjectPtr<UNiagaraScriptVariable> ScriptVariable = It.Value)
-			{
-				ScriptVariable->Metadata.CreateNewGuid();
-			}
-		}
-	}
-}
-
 class UNiagaraScriptSource* UNiagaraGraph::GetSource() const
 {
 	return CastChecked<UNiagaraScriptSource>(GetOuter());
diff --git a/Engine/Plugins/FX/Niagara/Source/NiagaraEditor/Private/NiagaraNode.cpp b/Engine/Plugins/FX/Niagara/Source/NiagaraEditor/Private/NiagaraNode.cpp
--- a/Engine/Plugins/FX/Niagara/Source/NiagaraEditor/Private/NiagaraNode.cpp
+++ b/Engine/Plugins/FX/Niagara/Source/NiagaraEditor/Private/NiagaraNode.cpp
@@ -218,7 +218,7 @@ bool UNiagaraNode::ReallocatePins(bool bMarkNeedsResynchronizeOnChange)
 
 		// When matching pins, use the pin id if either pin id is valid, otherwise match by name.
 		// we typically don't allow orphaned pins to find a match, unless types match each other again
-		auto PinMatchPredicate = [&](UEdGraphPin* Pin) 
+		auto PinMatchPredicateGuid = [&](UEdGraphPin* Pin) 
 		{
 			// early out conditions for when matching up is not wanted
 			if (Pin->Direction != OldPin->Direction)
@@ -231,8 +231,32 @@ bool UNiagaraNode::ReallocatePins(bool bMarkNeedsResynchronizeOnChange)
 			return bAllowPin && ((bPinsHaveIds && Pin->PersistentGuid == OldPin->PersistentGuid) || (bPinsHaveIds == false && Pin->PinName == OldPin->PinName));
 		};
 
+		auto PinMatchPredicateName = [&](UEdGraphPin* Pin)
+		{
+			// early out conditions for when matching up is not wanted
+			if (Pin->Direction != OldPin->Direction || Pin->PinName != OldPin->PinName)
+			{
+				return false;
+			}
+
+			if (OldPin->PinType == Pin->PinType)
+			{
+				return true;
+			}
+
+			// Allow numeric pins and wildcard pins to match any type.
+			FNiagaraTypeDefinition NewType = Schema->PinTypeToTypeDefinition(Pin->PinType);
+			return NewType == FNiagaraTypeDefinition::GetGenericNumericDef() || NewType == FNiagaraTypeDefinition::GetWildcardDef();
+		};
+
 		OldPin->Modify();
-		if (UEdGraphPin** MatchingNewPin = Pins.FindByPredicate(PinMatchPredicate))
+		UEdGraphPin** MatchingNewPin = Pins.FindByPredicate(PinMatchPredicateGuid);
+		if (MatchingNewPin == nullptr)
+		{
+			MatchingNewPin = Pins.FindByPredicate(PinMatchPredicateName);
+		}
+
+		if (MatchingNewPin)
 		{
 			// If the pin types don't match, CopyPersistentDataFromOldPin could very well overwrite our Matching pin with bad data.
 			// Let's cache it for now and reset it after copying the other relevant data off the old pin.
diff --git a/Engine/Plugins/FX/Niagara/Source/NiagaraEditor/Public/NiagaraGraph.h b/Engine/Plugins/FX/Niagara/Source/NiagaraEditor/Public/NiagaraGraph.h
--- a/Engine/Plugins/FX/Niagara/Source/NiagaraEditor/Public/NiagaraGraph.h
+++ b/Engine/Plugins/FX/Niagara/Source/NiagaraEditor/Public/NiagaraGraph.h
@@ -147,7 +147,6 @@ class UNiagaraGraph : public UEdGraph
 	virtual void PostLoad() override;
 	virtual void PostEditChangeProperty(struct FPropertyChangedEvent& PropertyChangedEvent) override;
 	virtual void BeginDestroy() override;
-	virtual void PostDuplicate(EDuplicateMode::Type DuplicateMode) override;
 	//~ End UObjet Interface
 	
 	/** Get the source that owns this graph */
