Fix for Normalmaps being wrong with negative UV scale.

diff --git a/Engine/Plugins/FX/Niagara/Shaders/Private/NiagaraSpriteVertexFactory.ush b/Engine/Plugins/FX/Niagara/Shaders/Private/NiagaraSpriteVertexFactory.ush
--- a/Engine/Plugins/FX/Niagara/Shaders/Private/NiagaraSpriteVertexFactory.ush
+++ b/Engine/Plugins/FX/Niagara/Shaders/Private/NiagaraSpriteVertexFactory.ush
@@ -186,7 +186,8 @@ struct FVertexFactoryIntermediates
 	float4 Color;
 	/** The velocity of the particle, XYZ: direction, W: speed. */
 	float4 ParticleVelocity;
-
+	/** The uv scale of the sprite. */
+	float2 UVScale;
 	/** The particle worldspace position */
 	FLWCVector3 Position;
 	
@@ -552,15 +553,22 @@ float3x3 CalcTangentBasis(FVertexFactoryIntermediates Intermediates)
 	// So there are artifacts with tangent space calculations like the TransformVector node,
 	// But this allows lighting based on a world space shape via the normal while still having normal maps camera aligned.
 	float3x3 Result;
-	Result[0] = Intermediates.TangentRight;
-	Result[1] = Intermediates.TangentUp;
+	const float2 UVSign = sign(Intermediates.UVScale);
+	Result[0] = Intermediates.TangentRight * UVSign.x;
+	Result[1] = Intermediates.TangentUp * UVSign.y;
 
 	// ENM_CameraFacing
 	//@todo - use static branching
 	if (NiagaraSpriteVF.NormalsType < .5f)
 	{
 		Result[2] = normalize(cross(Result[0],Result[1]));
-	} 
+		
+		if (UVSign.x != UVSign.y)
+		{
+			Result[2] *= -1.0f;
+		}
+		
+	}
 	// ENM_Spherical
 	else if (NiagaraSpriteVF.NormalsType < 1.5f)
 	{
@@ -818,6 +826,7 @@ FVertexFactoryIntermediates GetVertexFactoryIntermediates(FVertexFactoryInput In
 	Intermediates.Color	= GetNiagaraParticleColor(ParticleID);
 	Intermediates.CustomFacingVector = CustomFacing;
 	Intermediates.CustomAlignmentVector = SafeNormalize(GetNiagaraParticleAlignmentVector(ParticleID));
+	Intermediates.UVScale = UVScale;
 
 #if (DYNAMIC_PARAMETERS_MASK & 1)
 	Intermediates.DynamicParameter = GetNiagaraParticleDynamicParameters(ParticleID);
