- Ensure we always reset data interfaces when calling SetAsset
- This matches the editor path, but also means that when switching between systems which contain a user parameter of same name & type we don't accidentally end up with the wrong parameters on the data interface (for example not resetting would cause skel mesh bone lists to come from the original system)

#rnx

diff --git a/Engine/Plugins/FX/Niagara/Source/Niagara/Private/NiagaraComponent.cpp b/Engine/Plugins/FX/Niagara/Source/Niagara/Private/NiagaraComponent.cpp
--- a/Engine/Plugins/FX/Niagara/Source/Niagara/Private/NiagaraComponent.cpp
+++ b/Engine/Plugins/FX/Niagara/Source/Niagara/Private/NiagaraComponent.cpp
@@ -2979,7 +2979,7 @@ void UNiagaraComponent::CopyParametersFromAsset(bool bResetExistingOverrideParam
 		bool bNewParam = OverrideParameters.AddParameter(Param, true);
 
 		bool bCopyAssetValue = bResetExistingOverrideParameters || bNewParam;
-		if (bCopyAssetValue)
+		if (bCopyAssetValue || Param.IsDataInterface())
 		{
 			Asset->GetExposedParameters().CopyParameterData(OverrideParameters, Param);
 		}
