- Do not reset GPU buffers as we need to allow the tick process to reallocate them
- This fixes a multi-tick issue, where the first tick would not reset and the second would.  Because we reallocate the GPU buffesr out of sync the first tick would assume we had source particles when they had been killed here.

#rnx

diff --git a/Engine/Plugins/FX/Niagara/Source/Niagara/Private/NiagaraEmitterInstance.cpp b/Engine/Plugins/FX/Niagara/Source/Niagara/Private/NiagaraEmitterInstance.cpp
--- a/Engine/Plugins/FX/Niagara/Source/Niagara/Private/NiagaraEmitterInstance.cpp
+++ b/Engine/Plugins/FX/Niagara/Source/Niagara/Private/NiagaraEmitterInstance.cpp
@@ -1204,7 +1204,7 @@ void FNiagaraEmitterInstance::PreTick()
 
 	if (bResetPending)
 	{
-		Data.ResetBuffers();
+		bResetPending = false;
 
 		if (EventInstanceData.IsValid())
 		{
@@ -1218,13 +1218,15 @@ void FNiagaraEmitterInstance::PreTick()
 			}
 		}
 
-		bResetPending = false;
-
 		if ( GPUExecContext )
 		{
 			GPUExecContext->bResetPending_GT = true;
 			GPUExecContext->GpuSpawnInfo_GT.Reset();
 		}
+		else
+		{
+			Data.ResetBuffers();
+		}
 	}
 
 	++TickCount;
