Downgrade error to a warning during cooking when we detect mismatching rapid iteration parameters, but only in monolithic editor targets. Works around a cooker crash that needs expert Niagara team eyes on to fix properly

#rnx

diff --git a/Engine/Plugins/FX/Niagara/Source/Niagara/Private/NiagaraScript.cpp b/Engine/Plugins/FX/Niagara/Source/Niagara/Private/NiagaraScript.cpp
--- a/Engine/Plugins/FX/Niagara/Source/Niagara/Private/NiagaraScript.cpp
+++ b/Engine/Plugins/FX/Niagara/Source/Niagara/Private/NiagaraScript.cpp
@@ -1806,7 +1806,11 @@ void UNiagaraScript::PostLoad()
 		// disable the owning script and system
 		if (!HasValidParameterBindings())
 		{
+#if WITH_EDITOR && IS_MONOLITHIC
+			UE_LOG(LogNiagara, Warning, TEXT("Mismatch between binding between RapidIterationParamters and ScriptExecutionParameters for system %s"), *GetFullName());
+#else
 			UE_LOG(LogNiagara, Error, TEXT("Mismatch between binding between RapidIterationParamters and ScriptExecutionParameters for system %s"), *GetFullName());
+#endif
 
 			CachedScriptVM.Reset();
 			return;
