- Don't generate particle state error if we don't write particles
- This fixes a false error with an emty simulation stage with state iteration enabled

#rnx

diff --git a/Engine/Plugins/FX/Niagara/Source/NiagaraEditor/Private/NiagaraHlslTranslator.cpp b/Engine/Plugins/FX/Niagara/Source/NiagaraEditor/Private/NiagaraHlslTranslator.cpp
--- a/Engine/Plugins/FX/Niagara/Source/NiagaraEditor/Private/NiagaraHlslTranslator.cpp
+++ b/Engine/Plugins/FX/Niagara/Source/NiagaraEditor/Private/NiagaraHlslTranslator.cpp
@@ -3260,7 +3260,7 @@ void FHlslNiagaraTranslator::DefineMainGPUFunctions(
 		{
 			// Do we have particle state iteration enable
 			SimRunSpawnUpdateLogicString += TEXT("	bool bRunSpawnUpdateLogic = true;\n");
-			if ( TranslationStage.bParticleIterationStateEnabled )
+			if ( TranslationStage.bParticleIterationStateEnabled && TranslationStage.bWritesParticles )
 			{
 				if ( TranslationStage.bPartialParticleUpdate == false )
 				{
