- Fix for uninitialized dummy volume texture

#rnx

diff --git a/Engine/Plugins/FX/Niagara/Source/Niagara/Private/NiagaraRenderer.cpp b/Engine/Plugins/FX/Niagara/Source/Niagara/Private/NiagaraRenderer.cpp
--- a/Engine/Plugins/FX/Niagara/Source/Niagara/Private/NiagaraRenderer.cpp
+++ b/Engine/Plugins/FX/Niagara/Source/Niagara/Private/NiagaraRenderer.cpp
@@ -128,7 +128,13 @@ public:
 
 			case ETextureDimension::Texture3D:
 			{
-				Texture = RHICreateTexture(Desc);
+				FTexture3DRHIRef Texture3D = RHICreateTexture(Desc);
+				Texture = Texture3D;
+
+				const FPixelFormatInfo& Info = GPixelFormats[PixelFormat];
+				TArray<uint8, TInlineAllocator<16>> Data;
+				Data.AddZeroed(Info.BlockBytes);
+				RHIUpdateTexture3D(Texture3D, 0, FUpdateTextureRegion3D(0, 0, 0, 0, 0, 0, 1, 1, 1), Info.BlockBytes, Info.BlockBytes, Data.GetData());
 				break;
 			}
 
