- When flushing pending ticks call begin / end scene otherwise we can run into issues with static buffers being unbound
- For example, if we write a fence it will call SubmitCommandsHint which will ClearState on some RHIs which will unbind all the static buffers

diff --git a/Engine/Plugins/FX/Niagara/Source/Niagara/Private/NiagaraGpuComputeDispatch.cpp b/Engine/Plugins/FX/Niagara/Source/Niagara/Private/NiagaraGpuComputeDispatch.cpp
--- a/Engine/Plugins/FX/Niagara/Source/Niagara/Private/NiagaraGpuComputeDispatch.cpp
+++ b/Engine/Plugins/FX/Niagara/Source/Niagara/Private/NiagaraGpuComputeDispatch.cpp
@@ -493,7 +493,11 @@ void FNiagaraGpuComputeDispatch::ProcessPendingTicksFlush(FRHICommandListImmedia
 
 			TConstArrayView<FViewInfo> DummyViews = MakeArrayView(DummyView, 1);
 			const bool bAllowGPUParticleUpdate = true;
-			
+
+			// Notify that we are about to begin rendering the 'scene' this is required because some RHIs will ClearState
+			// in the event of submitting commands, i.e. when we write a fence, or indeed perform a manual flush.
+			RHICmdList.BeginScene();
+
 			// Execute all ticks that we can support without invalid simulations
 			FRDGBuilder GraphBuilder(RHICmdList);
 			CreateSystemTextures(GraphBuilder);
@@ -517,6 +521,8 @@ void FNiagaraGpuComputeDispatch::ProcessPendingTicksFlush(FRHICommandListImmedia
 			// Properly clear the reference to ViewUniformBuffer before memstack wipes the memory
 			DummyView->~FViewInfo();
 
+			// We have completed flushing the commands
+			RHICmdList.EndScene();
 			break;
 		}
 
