Fix off-by-one error in Niagara ribbon ray tracing mesh vertex count calculation

diff --git a/Engine/Plugins/FX/Niagara/Source/Niagara/Private/NiagaraRendererRibbons.cpp b/Engine/Plugins/FX/Niagara/Source/Niagara/Private/NiagaraRendererRibbons.cpp
--- a/Engine/Plugins/FX/Niagara/Source/Niagara/Private/NiagaraRendererRibbons.cpp
+++ b/Engine/Plugins/FX/Niagara/Source/Niagara/Private/NiagaraRendererRibbons.cpp
@@ -1499,7 +1499,9 @@ void FNiagaraRendererRibbons::GetDynamicRayTracingInstances(FRayTracingMaterialG
 	// Use the internal vertex buffer only when initialized otherwise used the shared vertex buffer - needs to be updated every frame
 	FRWBuffer* VertexBuffer = RayTracingDynamicVertexBuffer.NumBytes > 0 ? &RayTracingDynamicVertexBuffer : nullptr;
 
-	const uint32 VertexCount = DynamicIndexAllocation.MaxUsedIndex;
+	// Vertex count is the maximum value in the index buffer + 1
+	const uint32 VertexCount = DynamicIndexAllocation.MaxUsedIndex + 1;
+
 	Context.DynamicRayTracingGeometriesToUpdate.Add(
 		FRayTracingDynamicGeometryUpdateParams
 		{
