Make sure that DataInterface curves are consistent with respect to UpdateTimeRanges & BuildLUT for WITH_EDITORONLY_DATA

diff --git a/Engine/Plugins/FX/Niagara/Source/Niagara/Classes/NiagaraDataInterfaceColorCurve.h b/Engine/Plugins/FX/Niagara/Source/Niagara/Classes/NiagaraDataInterfaceColorCurve.h
--- a/Engine/Plugins/FX/Niagara/Source/Niagara/Classes/NiagaraDataInterfaceColorCurve.h
+++ b/Engine/Plugins/FX/Niagara/Source/Niagara/Classes/NiagaraDataInterfaceColorCurve.h
@@ -44,10 +44,8 @@ public:
 		CurveLUTNumElems = 4,
 	};
 
-#if WITH_EDITORONLY_DATA
 	virtual void UpdateTimeRanges() override;
 	virtual TArray<float> BuildLUT(int32 NumEntries) const override;
-#endif
 
 	virtual void GetFunctions(TArray<FNiagaraFunctionSignature>& OutFunctions)override;
 	virtual void GetVMExternalFunction(const FVMExternalFunctionBindingInfo& BindingInfo, void* InstanceData, FVMExternalFunction &OutFunc) override;
diff --git a/Engine/Plugins/FX/Niagara/Source/Niagara/Private/NiagaraDataInterfaceColorCurve.cpp b/Engine/Plugins/FX/Niagara/Source/Niagara/Private/NiagaraDataInterfaceColorCurve.cpp
--- a/Engine/Plugins/FX/Niagara/Source/Niagara/Private/NiagaraDataInterfaceColorCurve.cpp
+++ b/Engine/Plugins/FX/Niagara/Source/Niagara/Private/NiagaraDataInterfaceColorCurve.cpp
@@ -69,7 +69,6 @@ void UNiagaraDataInterfaceColorCurve::Serialize(FArchive& Ar)
 	}
 }
 
-#if WITH_EDITORONLY_DATA
 void UNiagaraDataInterfaceColorCurve::UpdateTimeRanges()
 {
 	if ((RedCurve.GetNumKeys() > 0 || GreenCurve.GetNumKeys() > 0 || BlueCurve.GetNumKeys() > 0 || AlphaCurve.GetNumKeys() > 0))
@@ -112,7 +111,6 @@ TArray<float> UNiagaraDataInterfaceColorCurve::BuildLUT(int32 NumEntries) const
 	}
 	return OutputLUT;
 }
-#endif
 
 bool UNiagaraDataInterfaceColorCurve::CopyToInternal(UNiagaraDataInterface* Destination) const 
 {
