- Add camera bookmarks to the baker

diff --git a/Engine/Plugins/FX/Niagara/Source/Niagara/Classes/NiagaraBakerSettings.h b/Engine/Plugins/FX/Niagara/Source/Niagara/Classes/NiagaraBakerSettings.h
--- a/Engine/Plugins/FX/Niagara/Source/Niagara/Classes/NiagaraBakerSettings.h
+++ b/Engine/Plugins/FX/Niagara/Source/Niagara/Classes/NiagaraBakerSettings.h
@@ -26,6 +26,52 @@ enum class ENiagaraBakerViewMode
 	Num UMETA(Hidden)
 };
 
+USTRUCT()
+struct FNiagaraBakerCameraSettings
+{
+	GENERATED_BODY()
+
+	UPROPERTY()
+	ENiagaraBakerViewMode ViewMode = ENiagaraBakerViewMode::Perspective;
+
+	UPROPERTY()
+	FVector ViewportLocation = FVector::ZeroVector;
+
+	UPROPERTY()
+	FRotator ViewportRotation = FRotator::ZeroRotator;
+
+	UPROPERTY()
+	float OrbitDistance = 200.f;
+
+	UPROPERTY()
+	float FOV = 90.0f;
+
+	UPROPERTY()
+	float OrthoWidth = 512.0f;
+
+	UPROPERTY()
+	bool bUseAspectRatio = false;
+
+	UPROPERTY()
+	float AspectRatio = 1.0f;
+
+	bool IsOrthographic() const { return ViewMode != ENiagaraBakerViewMode::Perspective; }
+	bool IsPerspective() const { return ViewMode == ENiagaraBakerViewMode::Perspective; }
+
+	bool operator==(const FNiagaraBakerCameraSettings& Other) const
+	{
+		return
+			ViewMode == Other.ViewMode &&
+			ViewportLocation.Equals(Other.ViewportLocation) &&
+			ViewportRotation.Equals(Other.ViewportRotation) &&
+			OrbitDistance == Other.OrbitDistance &&
+			FOV == Other.FOV &&
+			OrthoWidth == Other.OrthoWidth &&
+			bUseAspectRatio == Other.bUseAspectRatio &&
+			AspectRatio == Other.AspectRatio;
+	}
+};
+
 USTRUCT()
 struct FNiagaraBakerTextureSettings
 {
@@ -112,49 +158,18 @@ public:
 	UPROPERTY(EditAnywhere, Category = "Settings")
 	TArray<TObjectPtr<UNiagaraBakerOutput>> Outputs;
 
-	/** Current active viewport we will render from. */
+	/** Camera Settings, will always be at least ENiagaraBakerViewMode::Num elements and those are fixed cameras. */
 	UPROPERTY(EditAnywhere, Category = "Settings")
-	ENiagaraBakerViewMode CameraViewportMode = ENiagaraBakerViewMode::Perspective;
+	TArray<FNiagaraBakerCameraSettings> CameraSettings;
 
-	/** Per viewport camera position.. */
+	/** Active camera that we were saved with */
 	UPROPERTY(EditAnywhere, Category = "Settings")
-	FVector CameraViewportLocation[(int)ENiagaraBakerViewMode::Num];
-
-	/** Per viewport camera rotation.. */
-	UPROPERTY(EditAnywhere, Category = "Settings")
-	FRotator CameraViewportRotation[(int)ENiagaraBakerViewMode::Num];
-
-	/** Perspective camera orbit distance. */
-	UPROPERTY(EditAnywhere, Category = "Settings", meta=(EditConditionHides, EditCondition="CameraViewportMode == ENiagaraBakerViewMode::Perspective", ClampMin = "0.01"))
-	float CameraOrbitDistance = 200.f;
-
-	/** Camera FOV to use when in perspective mode. */
-	UPROPERTY(EditAnywhere, Category = "Settings", meta=(EditConditionHides, EditCondition="CameraViewportMode == ENiagaraBakerViewMode::Perspective", ClampMin="1.0", ClampMax="179.0"))
-	float CameraFOV = 90.0f;
-
-	/** Camera Orthographic width to use with in orthographic mode. */
-	UPROPERTY(EditAnywhere, Category = "Settings", meta=(EditConditionHides, EditCondition = "CameraViewportMode != ENiagaraBakerViewMode::Perspective", ClampMin="1.0"))
-	float CameraOrthoWidth = 512.0f;
-
-	UPROPERTY(EditAnywhere, Category = "Settings", meta=(PinHiddenByDefault, InlineEditConditionToggle))
-	uint8 bUseCameraAspectRatio : 1;
-
-	/** Custom aspect ratio to use rather than using the width & height to automatically calculate. */
-	UPROPERTY(EditAnywhere, Category = "Settings", meta=(EditCondition = "bUseCameraAspectRatio", ClampMin="0.01"))
-	float CameraAspectRatio = 1.0f;
+	int32 CurrentCameraIndex = 0;
 
 	/** Should we render just the component or the whole scene. */
 	UPROPERTY(EditAnywhere, Category = "Settings", AdvancedDisplay)
 	uint8 bRenderComponentOnly : 1;
 
-	///** Type of level setup to use. */
-	//UPROPERTY(EditAnywhere, Category = "Environment", meta = (PinHiddenByDefault, InlineEditConditionToggle))
-	//uint8 bLoadLevel : 1;
-
-	///** Used to determine type of level setup. */
-	//UPROPERTY(EditAnywhere, Category = "Environment", meta = (EditCondition = "bLoadLevel"))
-	//TSoftObjectPtr<class ULevel> LevelEnvironment;
-
 	bool Equals(const UNiagaraBakerSettings& Other) const;
 
 	float GetSeekDelta() const { return 1.0f / float(FramesPerSecond); }
@@ -165,8 +180,8 @@ public:
 	FMatrix GetViewMatrix() const;
 	FMatrix GetProjectionMatrix() const;
 
-	bool IsOrthographic() const { return CameraViewportMode != ENiagaraBakerViewMode::Perspective; }
-	bool IsPerspective() const { return CameraViewportMode == ENiagaraBakerViewMode::Perspective; }
+	FNiagaraBakerCameraSettings& GetCurrentCamera() { return CameraSettings[CurrentCameraIndex]; }
+	const FNiagaraBakerCameraSettings& GetCurrentCamera() const { return CameraSettings[CurrentCameraIndex]; }
 
 	int GetOutputNumFrames(UNiagaraBakerOutput* BakerOutput) const;
 	FNiagaraBakerOutputFrameIndices GetOutputFrameIndices(UNiagaraBakerOutput* BakerOutput, float RelativeTime) const;
@@ -182,4 +197,28 @@ public:
 	// Deprecated properties
 	UPROPERTY()
 	TArray<FNiagaraBakerTextureSettings> OutputTextures_DEPRECATED;
+
+	UPROPERTY()
+	ENiagaraBakerViewMode CameraViewportMode_DEPRECATED = ENiagaraBakerViewMode::Perspective;
+
+	UPROPERTY()
+	FVector CameraViewportLocation_DEPRECATED[(int)ENiagaraBakerViewMode::Num];
+
+	UPROPERTY()
+	FRotator CameraViewportRotation_DEPRECATED[(int)ENiagaraBakerViewMode::Num];
+
+	UPROPERTY()
+	float CameraOrbitDistance_DEPRECATED = 200.f;
+
+	UPROPERTY()
+	float CameraFOV_DEPRECATED = 90.0f;
+
+	UPROPERTY()
+	float CameraOrthoWidth_DEPRECATED = 512.0f;
+
+	UPROPERTY()
+	uint8 bUseCameraAspectRatio_DEPRECATED : 1;
+
+	UPROPERTY()
+	float CameraAspectRatio_DEPRECATED = 1.0f;
 };
diff --git a/Engine/Plugins/FX/Niagara/Source/Niagara/Private/NiagaraBakerSettings.cpp b/Engine/Plugins/FX/Niagara/Source/Niagara/Private/NiagaraBakerSettings.cpp
--- a/Engine/Plugins/FX/Niagara/Source/Niagara/Private/NiagaraBakerSettings.cpp
+++ b/Engine/Plugins/FX/Niagara/Source/Niagara/Private/NiagaraBakerSettings.cpp
@@ -2,6 +2,7 @@
 
 #include "NiagaraBakerSettings.h"
 #include "NiagaraBakerOutputTexture2D.h"
+#include "NiagaraCustomVersion.h"
 #include "NiagaraSystem.h"
 
 #include "Engine/Texture2D.h"
@@ -14,13 +15,23 @@ UNiagaraBakerSettings::UNiagaraBakerSettings(const FObjectInitializer& Init)
 	bRenderComponentOnly = true;
 	Outputs.Add(Init.CreateDefaultSubobject<UNiagaraBakerOutputTexture2D>(this, "DefaultOutput"));
 
+	// Fill in default camera settings
+	CameraSettings.AddDefaulted(int(ENiagaraBakerViewMode::Num));
+	for (int i=0; i < int(ENiagaraBakerViewMode::Num); ++i)
+	{
+		CameraSettings[i].ViewMode = ENiagaraBakerViewMode(i);
+	}
+	CameraSettings[int(ENiagaraBakerViewMode::Perspective)].ViewportLocation = FVector(0.0f, -200.0f, 0.0f);
+	CameraSettings[int(ENiagaraBakerViewMode::Perspective)].ViewportRotation = FRotator(180.0f, 0.0f, 90.0f);
+
+	// Deprecate properties
 	for ( int i=0; i < (int)ENiagaraBakerViewMode::Num; ++i )
 	{
-		CameraViewportLocation[i] = FVector::ZeroVector;
-		CameraViewportRotation[i] = FRotator::ZeroRotator;
+		CameraViewportLocation_DEPRECATED[i] = FVector::ZeroVector;
+		CameraViewportRotation_DEPRECATED[i] = FRotator::ZeroRotator;
 	}
-	CameraViewportLocation[(int)ENiagaraBakerViewMode::Perspective] = FVector(0.0f, -200.0f, 0.0f);
-	CameraViewportRotation[(int)ENiagaraBakerViewMode::Perspective] = FRotator(180.0f, 0.0f, 90.0f);
+	CameraViewportLocation_DEPRECATED[(int)ENiagaraBakerViewMode::Perspective] = FVector(0.0f, -200.0f, 0.0f);
+	CameraViewportRotation_DEPRECATED[(int)ENiagaraBakerViewMode::Perspective] = FRotator(180.0f, 0.0f, 90.0f);
 }
 
 bool UNiagaraBakerSettings::Equals(const UNiagaraBakerSettings& Other) const
@@ -44,58 +55,40 @@ bool UNiagaraBakerSettings::Equals(const UNiagaraBakerSettings& Other) const
 			return true;
 		};
 
-	auto CamerasEquals =
-		[&]() -> bool
-		{
-			for ( int i=0; i < (int)ENiagaraBakerViewMode::Num; ++i )
-			{
-				if ( !CameraViewportLocation[i].Equals(Other.CameraViewportLocation[i]) ||
-					 !CameraViewportRotation[i].Equals(Other.CameraViewportRotation[i]) )
-				{
-					return false;
-				}
-			}
-			return true;
-		};
-
 	return
+		CameraSettings == Other.CameraSettings &&
+		CurrentCameraIndex == Other.CurrentCameraIndex &&
 		OutputsEquals() &&
-		CamerasEquals() &&
 		FMath::IsNearlyEqual(StartSeconds, Other.StartSeconds) &&
 		FMath::IsNearlyEqual(DurationSeconds, Other.DurationSeconds) &&
 		FramesPerSecond == Other.FramesPerSecond &&
 		FramesPerDimension == Other.FramesPerDimension &&
 		bPreviewLooping == Other.bPreviewLooping &&
-		CameraViewportMode == Other.CameraViewportMode &&
-		FMath::IsNearlyEqual(CameraOrbitDistance, Other.CameraOrbitDistance) &&
-		FMath::IsNearlyEqual(CameraFOV, Other.CameraFOV) &&
-		FMath::IsNearlyEqual(CameraOrthoWidth, Other.CameraOrthoWidth) &&
-		bUseCameraAspectRatio == Other.bUseCameraAspectRatio &&
-		FMath::IsNearlyEqual(CameraAspectRatio, Other.CameraAspectRatio) &&
 		bRenderComponentOnly == Other.bRenderComponentOnly;
 }
 
 FVector UNiagaraBakerSettings::GetCameraLocation() const
 {
-	if (IsPerspective())
+	const FNiagaraBakerCameraSettings& CurrentCamera = GetCurrentCamera();
+	if (CurrentCamera.IsPerspective())
 	{
-		const FVector OrbitOffset = CameraViewportRotation[(int)CameraViewportMode].RotateVector(FVector(0.0f, 0.0f, CameraOrbitDistance));
-		return CameraViewportLocation[(int)CameraViewportMode] - OrbitOffset;
+		const FVector OrbitOffset = CurrentCamera.ViewportRotation.RotateVector(FVector(0.0f, 0.0f, CurrentCamera.OrbitDistance));
+		return CurrentCamera.ViewportLocation - OrbitOffset;
 	}
 	else
 	{
-		return CameraViewportLocation[(int)CameraViewportMode];
+		return CurrentCamera.ViewportLocation;
 	}
 }
 
 FRotator UNiagaraBakerSettings::GetCameraRotation() const
 {
-	return CameraViewportRotation[(int)CameraViewportMode];
+	return GetCurrentCamera().ViewportRotation;
 }
 
 FMatrix UNiagaraBakerSettings::GetViewportMatrix() const
 {
-	switch (CameraViewportMode)
+	switch (GetCurrentCamera().ViewMode)
 	{
 		case ENiagaraBakerViewMode::OrthoFront:		return FMatrix(-FVector::ZAxisVector, -FVector::XAxisVector,  FVector::YAxisVector, FVector::ZeroVector);
 		case ENiagaraBakerViewMode::OrthoBack:		return FMatrix( FVector::ZAxisVector,  FVector::XAxisVector,  FVector::YAxisVector, FVector::ZeroVector);
@@ -115,17 +108,18 @@ FMatrix UNiagaraBakerSettings::GetViewMatrix() const
 
 FMatrix UNiagaraBakerSettings::GetProjectionMatrix() const
 {
-	const float AspectRatioY = bUseCameraAspectRatio ? CameraAspectRatio : 1.0f;
-	if (CameraViewportMode == ENiagaraBakerViewMode::Perspective)
+	const FNiagaraBakerCameraSettings& CurrentCamera = GetCurrentCamera();
+	const float AspectRatioY = CurrentCamera.bUseAspectRatio ? CurrentCamera.AspectRatio : 1.0f;
+	if (CurrentCamera.ViewMode == ENiagaraBakerViewMode::Perspective)
 	{
-		const float HalfXFOV = FMath::DegreesToRadians(CameraFOV) * 0.5f;
+		const float HalfXFOV = FMath::DegreesToRadians(CurrentCamera.FOV) * 0.5f;
 		const float HalfYFOV = FMath::Atan(FMath::Tan(HalfXFOV) / AspectRatioY);
 		return FReversedZPerspectiveMatrix(HalfXFOV, HalfYFOV, 1.0f, 1.0, GNearClippingPlane, GNearClippingPlane);
 	}
 	else
 	{
 		const float ZRange = WORLD_MAX;
-		return FReversedZOrthoMatrix(CameraOrthoWidth / 2.0f, CameraOrthoWidth * AspectRatioY / 2.0f, 0.5f / ZRange, ZRange);
+		return FReversedZOrthoMatrix(CurrentCamera.OrthoWidth / 2.0f, CurrentCamera.OrthoWidth * AspectRatioY / 2.0f, 0.5f / ZRange, ZRange);
 	}
 }
 
@@ -163,6 +157,8 @@ void UNiagaraBakerSettings::PostLoad()
 {
 	Super::PostLoad();
 
+	const int32 NiagaraVersion = GetLinkerCustomVersion(FNiagaraCustomVersion::GUID);
+
 	if ( OutputTextures_DEPRECATED.Num() > 0 )
 	{
 		Outputs.Empty();
@@ -201,6 +197,23 @@ void UNiagaraBakerSettings::PostLoad()
 
 		OutputTextures_DEPRECATED.Empty();
 	}
+
+
+	if (NiagaraVersion < FNiagaraCustomVersion::AddBakerCameraBookmarks)
+	{
+		for ( int32 i=0; i < int(ENiagaraBakerViewMode::Num); ++i )
+		{
+			CameraSettings[i].ViewMode = ENiagaraBakerViewMode(i);
+			CameraSettings[i].ViewportLocation = CameraViewportLocation_DEPRECATED[i];
+			CameraSettings[i].ViewportRotation = CameraViewportRotation_DEPRECATED[i];
+			CameraSettings[i].OrbitDistance = CameraOrbitDistance_DEPRECATED;
+			CameraSettings[i].FOV = CameraFOV_DEPRECATED;
+			CameraSettings[i].OrthoWidth = CameraOrthoWidth_DEPRECATED;
+			CameraSettings[i].bUseAspectRatio = bUseCameraAspectRatio_DEPRECATED;
+			CameraSettings[i].AspectRatio = CameraAspectRatio_DEPRECATED;
+		}
+		CurrentCameraIndex = int(CameraViewportMode_DEPRECATED);
+	}
 }
 
 #if WITH_EDITORONLY_DATA
diff --git a/Engine/Plugins/FX/Niagara/Source/NiagaraCore/Public/NiagaraCustomVersion.h b/Engine/Plugins/FX/Niagara/Source/NiagaraCore/Public/NiagaraCustomVersion.h
--- a/Engine/Plugins/FX/Niagara/Source/NiagaraCore/Public/NiagaraCustomVersion.h
+++ b/Engine/Plugins/FX/Niagara/Source/NiagaraCore/Public/NiagaraCustomVersion.h
@@ -172,6 +172,8 @@ struct FNiagaraCustomVersion
 
 		VisibilityCullingImprovements, // Extended visibility culling options and moved properties into their own struct.
 
+		AddBakerCameraBookmarks,
+
 		// DO NOT ADD A NEW VERSION UNLESS YOU HAVE TALKED TO THE NIAGARA LEAD. Mismanagement of these versions can lead to data loss if it is adjusted in multiple streams simultaneously.
 		// -----<new versions can be added above this line>  -------------------------------------------------
 		VersionPlusOne,
diff --git a/Engine/Plugins/FX/Niagara/Source/NiagaraEditor/Private/NiagaraBakerRenderer.cpp b/Engine/Plugins/FX/Niagara/Source/NiagaraEditor/Private/NiagaraBakerRenderer.cpp
--- a/Engine/Plugins/FX/Niagara/Source/NiagaraEditor/Private/NiagaraBakerRenderer.cpp
+++ b/Engine/Plugins/FX/Niagara/Source/NiagaraEditor/Private/NiagaraBakerRenderer.cpp
@@ -355,15 +355,16 @@ void FNiagaraBakerRenderer::RenderSceneCapture(UTextureRenderTarget2D* RenderTar
 	SceneCaptureComponent->CaptureSource = CaptureSource;
 	
 	// Set view location
-	if (BakerSettings->IsOrthographic())
+	const FNiagaraBakerCameraSettings& CurrentCamera = BakerSettings->GetCurrentCamera();
+	if (CurrentCamera.IsOrthographic())
 	{
 		SceneCaptureComponent->ProjectionType = ECameraProjectionMode::Orthographic;
-		SceneCaptureComponent->OrthoWidth = BakerSettings->CameraOrthoWidth;
+		SceneCaptureComponent->OrthoWidth = CurrentCamera.OrthoWidth;
 	}
 	else
 	{
 		SceneCaptureComponent->ProjectionType = ECameraProjectionMode::Perspective;
-		SceneCaptureComponent->FOVAngle = BakerSettings->CameraFOV;
+		SceneCaptureComponent->FOVAngle = CurrentCamera.FOV;
 	}
 	
 	const FMatrix SceneCaptureMatrix = FMatrix(FPlane(0, 0, 1, 0), FPlane(1, 0, 0, 0), FPlane(0, 1, 0, 0), FPlane(0, 0, 0, 1));
diff --git a/Engine/Plugins/FX/Niagara/Source/NiagaraEditor/Private/ViewModels/NiagaraBakerViewModel.cpp b/Engine/Plugins/FX/Niagara/Source/NiagaraEditor/Private/ViewModels/NiagaraBakerViewModel.cpp
--- a/Engine/Plugins/FX/Niagara/Source/NiagaraEditor/Private/ViewModels/NiagaraBakerViewModel.cpp
+++ b/Engine/Plugins/FX/Niagara/Source/NiagaraEditor/Private/ViewModels/NiagaraBakerViewModel.cpp
@@ -83,115 +83,168 @@ bool FNiagaraBakerViewModel::IsPlaybackLooping() const
 	return BakerSettings ? BakerSettings->bPreviewLooping : false;
 }
 
-void FNiagaraBakerViewModel::SetCameraViewMode(ENiagaraBakerViewMode ViewMode)
+void FNiagaraBakerViewModel::SetCameraSettingsIndex(int CamerSettingsIndex)
 {
 	if (UNiagaraBakerSettings* BakerSettings = GetBakerSettings())
 	{
-		const FScopedTransaction Transaction(LOCTEXT("SetCameraViewMode", "Set Camera View Mode"));
+		const FScopedTransaction Transaction(LOCTEXT("SetCameraSettingsIndex", "Set Camera Settings Index"));
 		BakerSettings->Modify();
-		BakerSettings->CameraViewportMode = ViewMode;
+		BakerSettings->CurrentCameraIndex = FMath::Clamp(CamerSettingsIndex, 0, BakerSettings->CameraSettings.Num() - 1);
 	}
 }
 
-bool FNiagaraBakerViewModel::IsCameraViewMode(ENiagaraBakerViewMode ViewMode)
+bool FNiagaraBakerViewModel::IsCameraSettingIndex(int CamerSettingsIndex) const
 {
 	UNiagaraBakerSettings* BakerSettings = GetBakerSettings();
-	return BakerSettings ? BakerSettings->CameraViewportMode == ViewMode : false;
+	return BakerSettings ? BakerSettings->CurrentCameraIndex == CamerSettingsIndex : false;
+}
+
+void FNiagaraBakerViewModel::AddCameraBookmark()
+{
+	if (UNiagaraBakerSettings* BakerSettings = GetBakerSettings())
+	{
+		const FScopedTransaction Transaction(LOCTEXT("AddCameraBookmark", "Add Camera Bookmark"));
+		BakerSettings->Modify();
+		FNiagaraBakerCameraSettings CameraSettings = BakerSettings->GetCurrentCamera();
+		BakerSettings->CameraSettings.Emplace(CameraSettings);
+	}
+}
+
+void FNiagaraBakerViewModel::RemoveCameraBookmark(int32 CameraIndex)
+{
+	if (UNiagaraBakerSettings* BakerSettings = GetBakerSettings())
+	{
+		if (CameraIndex >= int32(ENiagaraBakerViewMode::Num) && CameraIndex < BakerSettings->CameraSettings.Num())
+		const FScopedTransaction Transaction(LOCTEXT("RemoveCameraBookmark", "Remove Camera Bookmark"));
+		BakerSettings->Modify();
+		BakerSettings->CameraSettings.RemoveAt(CameraIndex);
+		BakerSettings->CurrentCameraIndex = FMath::Clamp(BakerSettings->CurrentCameraIndex, 0, BakerSettings->CameraSettings.Num() - 1);
+	}
 }
 
 FText FNiagaraBakerViewModel::GetCurrentCameraModeText() const
 {
 	UNiagaraBakerSettings* BakerSettings = GetBakerSettings();
-	return BakerSettings ? GetCameraModeText(BakerSettings->CameraViewportMode) : FText::GetEmpty();
+	return BakerSettings ? GetCameraSettingsText(BakerSettings->CurrentCameraIndex) : FText::GetEmpty();
 }
 
 FName FNiagaraBakerViewModel::GetCurrentCameraModeIconName() const
 {
 	UNiagaraBakerSettings* BakerSettings = GetBakerSettings();
-	return GetCameraModeIconName(BakerSettings ? BakerSettings->CameraViewportMode : ENiagaraBakerViewMode::Perspective);
+	return BakerSettings ? GetCameraSettingsIconName(BakerSettings->CurrentCameraIndex) : NAME_None;
 }
 
 FSlateIcon FNiagaraBakerViewModel::GetCurrentCameraModeIcon() const
 {
 	UNiagaraBakerSettings* BakerSettings = GetBakerSettings();
-	return GetCameraModeIcon(BakerSettings ? BakerSettings->CameraViewportMode : ENiagaraBakerViewMode::Perspective);
+	return BakerSettings ? GetCameraSettingsIcon(BakerSettings->CurrentCameraIndex) : FSlateIcon();
 }
 
-FText FNiagaraBakerViewModel::GetCameraModeText(ENiagaraBakerViewMode Mode)
+FText FNiagaraBakerViewModel::GetCameraSettingsText(int32 CameraSettingsIndex) const
 {
-	switch (Mode)
+	if ( UNiagaraBakerSettings* BakerSettings = GetBakerSettings() )
 	{
-		case ENiagaraBakerViewMode::Perspective:	return LOCTEXT("Perspective", "Perspective");
-		case ENiagaraBakerViewMode::OrthoFront:		return LOCTEXT("OrthoFront", "Front");
-		case ENiagaraBakerViewMode::OrthoBack:		return LOCTEXT("OrthoBack", "Back");
-		case ENiagaraBakerViewMode::OrthoLeft:		return LOCTEXT("OrthoLeft", "Left");
-		case ENiagaraBakerViewMode::OrthoRight:		return LOCTEXT("OrthoRight", "Right");
-		case ENiagaraBakerViewMode::OrthoTop:		return LOCTEXT("OrthoTop", "Top");
-		case ENiagaraBakerViewMode::OrthoBottom:	return LOCTEXT("OrthoBottom", "Bottom");
-		default:									return LOCTEXT("Error", "Error");
+		if (BakerSettings->CameraSettings.IsValidIndex(CameraSettingsIndex))
+		{
+			FText ModeText;
+			switch (BakerSettings->CameraSettings[CameraSettingsIndex].ViewMode)
+			{
+				case ENiagaraBakerViewMode::Perspective:	ModeText = LOCTEXT("Perspective", "Perspective"); break;
+				case ENiagaraBakerViewMode::OrthoFront:		ModeText = LOCTEXT("OrthoFront", "Front"); break;
+				case ENiagaraBakerViewMode::OrthoBack:		ModeText = LOCTEXT("OrthoBack", "Back"); break;
+				case ENiagaraBakerViewMode::OrthoLeft:		ModeText = LOCTEXT("OrthoLeft", "Left"); break;
+				case ENiagaraBakerViewMode::OrthoRight:		ModeText = LOCTEXT("OrthoRight", "Right"); break;
+				case ENiagaraBakerViewMode::OrthoTop:		ModeText = LOCTEXT("OrthoTop", "Top"); break;
+				case ENiagaraBakerViewMode::OrthoBottom:	ModeText = LOCTEXT("OrthoBottom", "Bottom"); break;
+				default:									ModeText = LOCTEXT("Error", "Error"); break;
+			}
+
+			if (CameraSettingsIndex < int32(ENiagaraBakerViewMode::Num))
+			{
+				return ModeText;
+			}
+			else
+			{
+				return FText::Format(LOCTEXT("CameraBookmarkFormat", "Bookmark {0} - {1}"), FText::AsNumber(CameraSettingsIndex - int32(ENiagaraBakerViewMode::Num)), ModeText);
+			}
+		}
 	}
+	return FText::GetEmpty();
 }
 
-FName FNiagaraBakerViewModel::GetCameraModeIconName(ENiagaraBakerViewMode Mode)
+FName FNiagaraBakerViewModel::GetCameraSettingsIconName(int32 CameraSettingsIndex) const
 {
-	static FName PerspectiveIcon("EditorViewport.Perspective");
-	static FName TopIcon("EditorViewport.Top");
-	static FName LeftIcon("EditorViewport.Left");
-	static FName FrontIcon("EditorViewport.Front");
-	static FName BottomIcon("EditorViewport.Bottom");
-	static FName RightIcon("EditorViewport.Right");
-	static FName BackIcon("EditorViewport.Back");
-
-	switch (Mode)
+	if (UNiagaraBakerSettings* BakerSettings = GetBakerSettings())
 	{
-		case ENiagaraBakerViewMode::Perspective:	return PerspectiveIcon;
-		case ENiagaraBakerViewMode::OrthoFront:		return FrontIcon;
-		case ENiagaraBakerViewMode::OrthoBack:		return BackIcon;
-		case ENiagaraBakerViewMode::OrthoLeft:		return LeftIcon;
-		case ENiagaraBakerViewMode::OrthoRight:		return RightIcon;
-		case ENiagaraBakerViewMode::OrthoTop:		return TopIcon;
-		case ENiagaraBakerViewMode::OrthoBottom:	return BottomIcon;
-		default:									return PerspectiveIcon;
+		static FName PerspectiveIcon("EditorViewport.Perspective");
+		static FName TopIcon("EditorViewport.Top");
+		static FName LeftIcon("EditorViewport.Left");
+		static FName FrontIcon("EditorViewport.Front");
+		static FName BottomIcon("EditorViewport.Bottom");
+		static FName RightIcon("EditorViewport.Right");
+		static FName BackIcon("EditorViewport.Back");
+
+		if (BakerSettings->CameraSettings.IsValidIndex(CameraSettingsIndex))
+		{
+			switch (BakerSettings->CameraSettings[CameraSettingsIndex].ViewMode)
+			{
+				case ENiagaraBakerViewMode::Perspective:	return PerspectiveIcon;
+				case ENiagaraBakerViewMode::OrthoFront:		return FrontIcon;
+				case ENiagaraBakerViewMode::OrthoBack:		return BackIcon;
+				case ENiagaraBakerViewMode::OrthoLeft:		return LeftIcon;
+				case ENiagaraBakerViewMode::OrthoRight:		return RightIcon;
+				case ENiagaraBakerViewMode::OrthoTop:		return TopIcon;
+				case ENiagaraBakerViewMode::OrthoBottom:	return BottomIcon;
+				default:									return PerspectiveIcon;
+			}
+		}
 	}
+	return NAME_None;
+}
+
+FSlateIcon FNiagaraBakerViewModel::GetCameraSettingsIcon(int32 CameraSettingsIndex) const
+{
+	return FSlateIcon(FEditorStyle::GetStyleSetName(), GetCameraSettingsIconName(CameraSettingsIndex));
 }
 
-FSlateIcon FNiagaraBakerViewModel::GetCameraModeIcon(ENiagaraBakerViewMode Mode)
+bool FNiagaraBakerViewModel::IsCurrentCameraPerspective() const
 {
-	return FSlateIcon(FEditorStyle::GetStyleSetName(), GetCameraModeIconName(Mode));
+	UNiagaraBakerSettings* BakerSettings = GetBakerSettings();
+	return BakerSettings ? BakerSettings->GetCurrentCamera().ViewMode == ENiagaraBakerViewMode::Perspective : true;
 }
 
 FVector FNiagaraBakerViewModel::GetCurrentCameraLocation() const
 {
 	UNiagaraBakerSettings* BakerSettings = GetBakerSettings();
-	return BakerSettings ? BakerSettings->CameraViewportLocation[int(BakerSettings->CameraViewportMode)] : FVector::ZeroVector;
+	return BakerSettings ? BakerSettings->GetCurrentCamera().ViewportLocation : FVector::ZeroVector;
 }
 
 void FNiagaraBakerViewModel::SetCurrentCameraLocation(const FVector Value)
 {
 	if ( UNiagaraBakerSettings* BakerSettings = GetBakerSettings() )
 	{
-		BakerSettings->CameraViewportLocation[int(BakerSettings->CameraViewportMode)] = Value;
+		BakerSettings->GetCurrentCamera().ViewportLocation = Value;
 	}
 }
 
 FRotator FNiagaraBakerViewModel::GetCurrentCameraRotation() const
 {
 	UNiagaraBakerSettings* BakerSettings = GetBakerSettings();
-	return BakerSettings ? BakerSettings->CameraViewportRotation[int(BakerSettings->CameraViewportMode)] : FRotator::ZeroRotator;
+	return BakerSettings ? BakerSettings->GetCurrentCamera().ViewportRotation : FRotator::ZeroRotator;
 }
 
 void FNiagaraBakerViewModel::SetCurrentCameraRotation(const FRotator Value) const
 {
 	if (UNiagaraBakerSettings* BakerSettings = GetBakerSettings())
 	{
-		BakerSettings->CameraViewportRotation[int(BakerSettings->CameraViewportMode)] = Value;
+		BakerSettings->GetCurrentCamera().ViewportRotation = Value;
 	}
 }
 
 float FNiagaraBakerViewModel::GetCameraFOV() const
 {
 	UNiagaraBakerSettings* BakerSettings = GetBakerSettings();
-	return BakerSettings ? BakerSettings->CameraFOV : 0.0f;
+	return BakerSettings ? BakerSettings->GetCurrentCamera().FOV : 0.0f;
 }
 
 void FNiagaraBakerViewModel::SetCameraFOV(float InFOV)
@@ -200,14 +253,14 @@ void FNiagaraBakerViewModel::SetCameraFOV(float InFOV)
 	{
 		const FScopedTransaction Transaction(LOCTEXT("SetCameraFOV", "Set Camera FOV"));
 		BakerSettings->Modify();
-		BakerSettings->CameraFOV = InFOV;
+		BakerSettings->GetCurrentCamera().FOV = InFOV;
 	}
 }
 
 float FNiagaraBakerViewModel::GetCameraOrbitDistance() const
 {
 	UNiagaraBakerSettings* BakerSettings = GetBakerSettings();
-	return BakerSettings ? BakerSettings->CameraOrbitDistance : 0.0f;
+	return BakerSettings ? BakerSettings->GetCurrentCamera().OrbitDistance : 0.0f;
 }
 
 void FNiagaraBakerViewModel::SetCameraOrbitDistance(float InOrbitDistance)
@@ -216,14 +269,14 @@ void FNiagaraBakerViewModel::SetCameraOrbitDistance(float InOrbitDistance)
 	{
 		const FScopedTransaction Transaction(LOCTEXT("SetCameraOrbitDistance", "Set Camera Orbit Distance"));
 		BakerSettings->Modify();
-		BakerSettings->CameraOrbitDistance = InOrbitDistance;
+		BakerSettings->GetCurrentCamera().OrbitDistance = InOrbitDistance;
 	}
 }
 
 float FNiagaraBakerViewModel::GetCameraOrthoWidth() const
 {
 	UNiagaraBakerSettings* BakerSettings = GetBakerSettings();
-	return BakerSettings ? BakerSettings->CameraOrthoWidth : 0.0f;
+	return BakerSettings ? BakerSettings->GetCurrentCamera().OrthoWidth : 0.0f;
 }
 
 void FNiagaraBakerViewModel::SetCameraOrthoWidth(float InOrthoWidth)
@@ -232,7 +285,7 @@ void FNiagaraBakerViewModel::SetCameraOrthoWidth(float InOrthoWidth)
 	{
 		const FScopedTransaction Transaction(LOCTEXT("SetCameraOrthoWidth", "Set Camera Ortho Width"));
 		BakerSettings->Modify();
-		BakerSettings->CameraOrthoWidth = InOrthoWidth;
+		BakerSettings->GetCurrentCamera().OrthoWidth = InOrthoWidth;
 	}
 }
 
@@ -242,20 +295,20 @@ void FNiagaraBakerViewModel::ToggleCameraAspectRatioEnabled()
 	{
 		const FScopedTransaction Transaction(LOCTEXT("ToggleUseCameraAspectRatio", "Toggle Use Camera Aspect Ratio"));
 		BakerSettings->Modify();
-		BakerSettings->bUseCameraAspectRatio = !BakerSettings->bUseCameraAspectRatio;
+		BakerSettings->GetCurrentCamera().bUseAspectRatio = !BakerSettings->GetCurrentCamera().bUseAspectRatio;
 	}
 }
 
 bool FNiagaraBakerViewModel::IsCameraAspectRatioEnabled() const
 {
 	UNiagaraBakerSettings* BakerSettings = GetBakerSettings();
-	return BakerSettings ? BakerSettings->bUseCameraAspectRatio : false;
+	return BakerSettings ? BakerSettings->GetCurrentCamera().bUseAspectRatio : false;
 }
 
 float FNiagaraBakerViewModel::GetCameraAspectRatio() const
 {
 	UNiagaraBakerSettings* BakerSettings = GetBakerSettings();
-	return BakerSettings ? BakerSettings->CameraAspectRatio : 0.0f;
+	return BakerSettings ? BakerSettings->GetCurrentCamera().AspectRatio : 0.0f;
 }
 
 void FNiagaraBakerViewModel::SetCameraAspectRatio(float InAspectRatio)
@@ -264,7 +317,7 @@ void FNiagaraBakerViewModel::SetCameraAspectRatio(float InAspectRatio)
 	{
 		const FScopedTransaction Transaction(LOCTEXT("SetCameraAspectRatio", "Set Camera Aspect Ratio"));
 		BakerSettings->Modify();
-		BakerSettings->CameraAspectRatio = InAspectRatio;
+		BakerSettings->GetCurrentCamera().AspectRatio = InAspectRatio;
 	}
 }
 
@@ -498,36 +551,45 @@ void FNiagaraBakerViewModel::RenderBaker()
 	}
 
 	// Render frames
-	const int32 TotalFrames = BakerSettings->FramesPerDimension.X * BakerSettings->FramesPerDimension.Y;
-	const float FrameDeltaSeconds = BakerSettings->DurationSeconds / float(TotalFrames);
-
-	FScopedSlowTask SlowTask(TotalFrames);
-	SlowTask.MakeDialog();
-
-	for ( int32 iFrame=0; iFrame < TotalFrames; ++iFrame )
 	{
-		SlowTask.EnterProgressFrame(1, FText::Format(LOCTEXT("BakingFormat", "Baking Frame ({0} / {1})"), iFrame + 1, TotalFrames));
+		const int32 TotalFrames = BakerSettings->FramesPerDimension.X * BakerSettings->FramesPerDimension.Y;
+		const float FrameDeltaSeconds = BakerSettings->DurationSeconds / float(TotalFrames);
 
-		// Tick to right location
-		const float FrameTime = BakerSettings->StartSeconds + (float(iFrame) * FrameDeltaSeconds);
-		BakerRenderer->SetAbsoluteTime(FrameTime);
+		FScopedSlowTask SlowTask(TotalFrames);
+		SlowTask.MakeDialog();
 
-		// Capture frames
-		for ( int32 i=0; i < OutputRenderers.Num(); ++i )
+		for ( int32 iFrame=0; iFrame < TotalFrames; ++iFrame )
 		{
-			if (OutputRenderers[i].IsValid())
+			SlowTask.EnterProgressFrame(1, FText::Format(LOCTEXT("BakingFormat", "Baking Frame ({0} / {1})"), iFrame + 1, TotalFrames));
+
+			// Tick to right location
+			const float FrameTime = BakerSettings->StartSeconds + (float(iFrame) * FrameDeltaSeconds);
+			BakerRenderer->SetAbsoluteTime(FrameTime);
+
+			// Capture frames
+			for ( int32 i=0; i < OutputRenderers.Num(); ++i )
 			{
-				OutputRenderers[i]->BakeFrame(BakerSettings->Outputs[i], iFrame, *BakerRenderer.Get());
+				if (OutputRenderers[i].IsValid())
+				{
+					OutputRenderers[i]->BakeFrame(BakerSettings->Outputs[i], iFrame, *BakerRenderer.Get());
+				}
 			}
 		}
 	}
 
 	// Complete bake
-	for (int32 i = 0; i < OutputRenderers.Num(); ++i)
 	{
-		if (OutputRenderers[i].IsValid())
+		FScopedSlowTask SlowTask(OutputRenderers.Num());
+		SlowTask.MakeDialog();
+
+		for (int32 i = 0; i < OutputRenderers.Num(); ++i)
 		{
-			OutputRenderers[i]->EndBake(BakerSettings->Outputs[i]);
+			if (OutputRenderers[i].IsValid())
+			{
+				SlowTask.EnterProgressFrame(1, FText::Format(LOCTEXT("FinishBakeFormat", "Finish Bake for Output '{0}'"), FText::FromString(BakerSettings->Outputs[i]->OutputName)));
+
+				OutputRenderers[i]->EndBake(BakerSettings->Outputs[i]);
+			}
 		}
 	}
 
diff --git a/Engine/Plugins/FX/Niagara/Source/NiagaraEditor/Private/ViewModels/NiagaraBakerViewModel.h b/Engine/Plugins/FX/Niagara/Source/NiagaraEditor/Private/ViewModels/NiagaraBakerViewModel.h
--- a/Engine/Plugins/FX/Niagara/Source/NiagaraEditor/Private/ViewModels/NiagaraBakerViewModel.h
+++ b/Engine/Plugins/FX/Niagara/Source/NiagaraEditor/Private/ViewModels/NiagaraBakerViewModel.h
@@ -64,17 +64,21 @@ public:
 	bool ShowInfoText() const { return bShowInfoText; }
 	void ToggleInfoText() { bShowInfoText = !bShowInfoText; }
 
-	void SetCameraViewMode(ENiagaraBakerViewMode ViewMode);
-	bool IsCameraViewMode(ENiagaraBakerViewMode ViewMode);
+	void SetCameraSettingsIndex(int CamerSettingsIndex);
+	bool IsCameraSettingIndex(int CamerSettingsIndex) const;
+
+	void AddCameraBookmark();
+	void RemoveCameraBookmark(int32 CameraIndex);
 
 	FText GetCurrentCameraModeText() const;
 	FName GetCurrentCameraModeIconName() const;
 	FSlateIcon GetCurrentCameraModeIcon() const;
 
-	static FText GetCameraModeText(ENiagaraBakerViewMode Mode);
-	static FName GetCameraModeIconName(ENiagaraBakerViewMode Mode);
-	static FSlateIcon GetCameraModeIcon(ENiagaraBakerViewMode Mode);
+	FText GetCameraSettingsText(int32 CameraSettingsIndex) const;
+	FName GetCameraSettingsIconName(int32 CameraSettingsIndex) const;
+	FSlateIcon GetCameraSettingsIcon(int32 CameraSettingsIndex) const;
 
+	bool IsCurrentCameraPerspective() const;
 	FVector GetCurrentCameraLocation() const;
 	void SetCurrentCameraLocation(const FVector Value);
 	FRotator GetCurrentCameraRotation() const;
diff --git a/Engine/Plugins/FX/Niagara/Source/NiagaraEditor/Private/Widgets/SNiagaraBakerViewport.cpp b/Engine/Plugins/FX/Niagara/Source/NiagaraEditor/Private/Widgets/SNiagaraBakerViewport.cpp
--- a/Engine/Plugins/FX/Niagara/Source/NiagaraEditor/Private/Widgets/SNiagaraBakerViewport.cpp
+++ b/Engine/Plugins/FX/Niagara/Source/NiagaraEditor/Private/Widgets/SNiagaraBakerViewport.cpp
@@ -2,7 +2,6 @@
 
 #include "SNiagaraBakerViewport.h"
 #include "ViewModels/NiagaraBakerViewModel.h"
-#include "SNiagaraBakerViewportToolbar.h"
 #include "NiagaraBakerRenderer.h"
 #include "NiagaraBatchedElements.h"
 #include "NiagaraComponent.h"
@@ -68,7 +67,7 @@ public:
 				bFocus |= Viewport->KeyState(FEditorViewportCommands::Get().FocusViewportToSelection->GetActiveChord(ChordIndex)->Key);
 			}
 
-			if ( BakerSettings->IsOrthographic() )
+			if ( BakerSettings->GetCurrentCamera().IsOrthographic() )
 			{
 				LocalMovement.X += bLeftKeyState ? KeyboardMoveSpeed : 0.0f;
 				LocalMovement.X -= bRightKeyState ? KeyboardMoveSpeed : 0.0f;
@@ -105,7 +104,7 @@ public:
 		UNiagaraBakerSettings* BakerSettings = ViewModel ? ViewModel->GetBakerSettings() : nullptr;
 		if ( ViewModel->ShowRealtimePreview() && BakerSettings && (Key == EKeys::MouseX || Key == EKeys::MouseY) )
 		{
-			if (BakerSettings->IsOrthographic())
+			if (BakerSettings->GetCurrentCamera().IsOrthographic())
 			{
 				if (InViewport->KeyState(EKeys::RightMouseButton))
 				{
@@ -173,8 +172,10 @@ public:
 			const FVector YAxis = ViewMatrix.GetUnitAxis(EAxis::Y);
 			const FVector ZAxis = ViewMatrix.GetUnitAxis(EAxis::Z);
 
+			FNiagaraBakerCameraSettings& CurrentCamera = BakerSettings->GetCurrentCamera();
+
 			FVector WorldMovement = FVector::ZeroVector;
-			if (BakerSettings->IsOrthographic())
+			if (CurrentCamera.IsOrthographic())
 			{
 				const FVector2D MoveSpeed = GetPreviewOrthoUnits();
 
@@ -182,11 +183,11 @@ public:
 				WorldMovement -= LocalMovement.Y * MoveSpeed.Y * YAxis;
 				//WorldMovement += LocalMovement.Z * MoveSpeed * ZAxis;
 
-				BakerSettings->CameraViewportLocation[(int)BakerSettings->CameraViewportMode] += WorldMovement;
+				CurrentCamera.ViewportLocation += WorldMovement;
 			}
 			else
 			{
-				FRotator& WorldRotation = BakerSettings->CameraViewportRotation[(int)BakerSettings->CameraViewportMode];
+				FRotator& WorldRotation = CurrentCamera.ViewportRotation;
 				WorldRotation.Yaw = FRotator::ClampAxis(WorldRotation.Yaw + LocalRotation.X);
 				WorldRotation.Roll = FMath::Clamp(WorldRotation.Roll + LocalRotation.Y, 0.0f, 180.0f);
 
@@ -194,28 +195,28 @@ public:
 				WorldMovement -= LocalMovement.X * MoveSpeed * XAxis;
 				WorldMovement -= LocalMovement.Y * MoveSpeed * YAxis;
 				//WorldMovement -= LocalMovement.Z * MoveSpeed * ZAxis;
-				BakerSettings->CameraViewportLocation[(int)BakerSettings->CameraViewportMode] += WorldMovement;
+				CurrentCamera.ViewportLocation += WorldMovement;
 
-				BakerSettings->CameraOrbitDistance = FMath::Max(BakerSettings->CameraOrbitDistance + LocalMovement.Z, 0.01f);
+				CurrentCamera.OrbitDistance = FMath::Max(CurrentCamera.OrbitDistance + LocalMovement.Z, 0.01f);
 			}
 
 			if (!FMath::IsNearlyZero(LocalZoom))
 			{
-				if (BakerSettings->IsPerspective())
+				if (CurrentCamera.IsPerspective())
 				{
-					BakerSettings->CameraFOV = FMath::Clamp(BakerSettings->CameraFOV + LocalZoom, 0.001f, 179.0f);
+					CurrentCamera.FOV = FMath::Clamp(CurrentCamera.FOV + LocalZoom, 0.001f, 179.0f);
 				}
 				else
 				{
-					BakerSettings->CameraOrthoWidth = FMath::Max(BakerSettings->CameraOrthoWidth + LocalZoom, 1.0f);
+					CurrentCamera.OrthoWidth = FMath::Max(CurrentCamera.OrthoWidth + LocalZoom, 1.0f);
 				}
 			}
 
 			if (!FMath::IsNearlyZero(LocalAspect))
 			{
-				if (BakerSettings->bUseCameraAspectRatio)
+				if (CurrentCamera.bUseAspectRatio)
 				{
-					BakerSettings->CameraAspectRatio = FMath::Max(BakerSettings->CameraAspectRatio + (LocalAspect / 50.0f), 0.01f);
+					CurrentCamera.AspectRatio = FMath::Max(CurrentCamera.AspectRatio + (LocalAspect / 50.0f), 0.01f);
 				}
 			}
 		}
@@ -555,19 +556,20 @@ public:
 		{
 			//-TODO: Should take aspect ratio into account here
 			const FBoxSphereBounds ComponentBounds = NiagaraComponent->CalcBounds(NiagaraComponent->GetComponentTransform());
-			if ( BakerSettings->IsOrthographic() )
+			FNiagaraBakerCameraSettings& CurrentCamera = BakerSettings->GetCurrentCamera();
+			if (CurrentCamera.IsOrthographic() )
 			{
-				BakerSettings->CameraViewportLocation[(int)BakerSettings->CameraViewportMode] = ComponentBounds.Origin;
-				BakerSettings->CameraOrthoWidth = ComponentBounds.SphereRadius * 2.0f;
+				CurrentCamera.ViewportLocation = ComponentBounds.Origin;
+				CurrentCamera.OrthoWidth = ComponentBounds.SphereRadius * 2.0f;
 			}
 			else
 			{
-				const float HalfFOVRadians = FMath::DegreesToRadians(BakerSettings->CameraFOV) * 0.5f;
+				const float HalfFOVRadians = FMath::DegreesToRadians(CurrentCamera.FOV) * 0.5f;
 				const float CameraDistance = ComponentBounds.SphereRadius / FMath::Tan(HalfFOVRadians);
 				//const FVector CameraOffset = BakerSettings->GetViewMatrix().Inverse().GetUnitAxis(EAxis::Z) * CameraDistance;
 				//BakerSettings->CameraViewportLocation[(int)ENiagaraBakerViewMode::Perspective] = ComponentBounds.Origin - CameraOffset;
-				BakerSettings->CameraViewportLocation[(int)BakerSettings->CameraViewportMode] = ComponentBounds.Origin;
-				BakerSettings->CameraOrbitDistance = CameraDistance;
+				CurrentCamera.ViewportLocation = ComponentBounds.Origin;
+				CurrentCamera.OrbitDistance = CameraDistance;
 			}
 		}
 	}
@@ -580,9 +582,10 @@ public:
 		const UNiagaraBakerSettings* BakerSettings = ViewModel ? ViewModel->GetBakerSettings() : nullptr;
 		if (BakerSettings && (PreviewViewRect.Area() > 0))
 		{
-			const float AspectRatioY = BakerSettings->bUseCameraAspectRatio ? BakerSettings->CameraAspectRatio : 1.0f;
-			OrthoUnits.X = BakerSettings->CameraOrthoWidth / float(PreviewViewRect.Width());
-			OrthoUnits.Y = BakerSettings->CameraOrthoWidth * AspectRatioY / float(PreviewViewRect.Height());
+			const FNiagaraBakerCameraSettings& CurrentCamera = BakerSettings->GetCurrentCamera();
+			const float AspectRatioY = CurrentCamera.bUseAspectRatio ? CurrentCamera.AspectRatio : 1.0f;
+			OrthoUnits.X = CurrentCamera.OrthoWidth / float(PreviewViewRect.Width());
+			OrthoUnits.Y = CurrentCamera.OrthoWidth * AspectRatioY / float(PreviewViewRect.Height());
 		}
 		return OrthoUnits;
 	}
diff --git a/Engine/Plugins/FX/Niagara/Source/NiagaraEditor/Private/Widgets/SNiagaraBakerViewportToolbar.cpp b/Engine/Plugins/FX/Niagara/Source/NiagaraEditor/Private/Widgets/SNiagaraBakerViewportToolbar.cpp
deleted file mode 100644
diff --git a/Engine/Plugins/FX/Niagara/Source/NiagaraEditor/Private/Widgets/SNiagaraBakerViewportToolbar.h b/Engine/Plugins/FX/Niagara/Source/NiagaraEditor/Private/Widgets/SNiagaraBakerViewportToolbar.h
deleted file mode 100644
diff --git a/Engine/Plugins/FX/Niagara/Source/NiagaraEditor/Private/Widgets/SNiagaraBakerWidget.cpp b/Engine/Plugins/FX/Niagara/Source/NiagaraEditor/Private/Widgets/SNiagaraBakerWidget.cpp
--- a/Engine/Plugins/FX/Niagara/Source/NiagaraEditor/Private/Widgets/SNiagaraBakerWidget.cpp
+++ b/Engine/Plugins/FX/Niagara/Source/NiagaraEditor/Private/Widgets/SNiagaraBakerWidget.cpp
@@ -210,111 +210,6 @@ namespace NiagaraBakerWidgetLocal
 
 //////////////////////////////////////////////////////////////////////////
 
-class FBakerSettingsDetails : public IDetailCustomization
-{
-public:
-	FBakerSettingsDetails(TWeakPtr<FNiagaraBakerViewModel> InWeakViewModel)
-		: WeakViewModel(InWeakViewModel)
-	{
-	}
-
-	static TSharedRef<IDetailCustomization> MakeInstance(TWeakPtr<FNiagaraBakerViewModel> InWeakViewModel)
-	{
-		return MakeShared<FBakerSettingsDetails>(InWeakViewModel);
-	}
-
-	virtual void CustomizeDetails(IDetailLayoutBuilder& DetailBuilder) override
-	{
-		// We only support customization on 1 object
-		TArray<TWeakObjectPtr<UObject>> ObjectsCustomized;
-		DetailBuilder.GetObjectsBeingCustomized(ObjectsCustomized);
-		if (ObjectsCustomized.Num() != 1 || !ObjectsCustomized[0]->IsA<UNiagaraBakerSettings>())
-		{
-			return;
-		}
-
-		FNiagaraBakerViewModel* ViewModel = WeakViewModel.Pin().Get();
-
-		UNiagaraBakerSettings* BakerSettings = CastChecked<UNiagaraBakerSettings>(ObjectsCustomized[0]);
-
-		// Func for only showing information for current camera
-		auto PerCameraProperty =
-			[&DetailBuilder, BakerSettings](IDetailCategoryBuilder& DetailCategory, TSharedPtr<IPropertyHandle> PropertyHandle, FText DisplayName, FText ToolTip)
-			{
-				TSharedPtr<IPropertyHandleArray> ArrayPropertyHandle = PropertyHandle->AsArray();
-				if (ArrayPropertyHandle == nullptr)
-				{
-					return;
-				}
-				uint32 NumElements = 0;
-				if ( ArrayPropertyHandle->GetNumElements(NumElements) != FPropertyAccess::Success )
-				{
-					return;
-				}
-
-				DetailBuilder.HideProperty(PropertyHandle->AsShared());
-				for (uint32 i=0; i < NumElements; ++i)
-				{
-					TSharedRef<IPropertyHandle> ArrayValuePropertyHandle = ArrayPropertyHandle->GetElement(i);
-
-					DetailCategory.AddProperty(ArrayValuePropertyHandle)
-						.DisplayName(DisplayName)
-						.ToolTip(ToolTip)
-						.Visibility(TAttribute<EVisibility>::CreateLambda([=] { return uint32(BakerSettings->CameraViewportMode) == i ? EVisibility::Visible : EVisibility::Hidden; }));
-				}
-			};
-
-		// Settings
-		{
-			IDetailCategoryBuilder& DetailCategory = DetailBuilder.EditCategory(FName("Settings"));
-
-			DetailCategory.AddProperty(GET_MEMBER_NAME_CHECKED(UNiagaraBakerSettings, StartSeconds));
-			DetailCategory.AddProperty(GET_MEMBER_NAME_CHECKED(UNiagaraBakerSettings, DurationSeconds));
-			DetailBuilder.HideProperty(GET_MEMBER_NAME_CHECKED(UNiagaraBakerSettings, FramesPerSecond));
-			DetailCategory.AddProperty(GET_MEMBER_NAME_CHECKED(UNiagaraBakerSettings, FramesPerDimension));
-
-			DetailBuilder.HideProperty(GET_MEMBER_NAME_CHECKED(UNiagaraBakerSettings, CameraViewportMode));
-			DetailBuilder.HideProperty(GET_MEMBER_NAME_CHECKED(UNiagaraBakerSettings, CameraViewportLocation));
-			DetailBuilder.HideProperty(GET_MEMBER_NAME_CHECKED(UNiagaraBakerSettings, CameraViewportRotation));
-			DetailBuilder.HideProperty(GET_MEMBER_NAME_CHECKED(UNiagaraBakerSettings, CameraOrbitDistance));
-			DetailBuilder.HideProperty(GET_MEMBER_NAME_CHECKED(UNiagaraBakerSettings, CameraFOV));
-			DetailBuilder.HideProperty(GET_MEMBER_NAME_CHECKED(UNiagaraBakerSettings, CameraOrthoWidth));
-			DetailBuilder.HideProperty(GET_MEMBER_NAME_CHECKED(UNiagaraBakerSettings, CameraAspectRatio));
-
-			DetailBuilder.HideProperty(GET_MEMBER_NAME_CHECKED(UNiagaraBakerSettings, bPreviewLooping));
-			DetailBuilder.HideProperty(GET_MEMBER_NAME_CHECKED(UNiagaraBakerSettings, bRenderComponentOnly));
-
-			//DetailBuilder.HideProperty(GET_MEMBER_NAME_CHECKED(UNiagaraBakerSettings, Outputs));
-		}
-
-		// Output Settings
-		if ( UNiagaraBakerOutput* CurrentOutput = ViewModel->GetCurrentOutput() )
-		{
-			IDetailCategoryBuilder& OutputTextureCategory = DetailBuilder.EditCategory(FName("OutputSettings"));
-
-			TSharedPtr<IPropertyHandleArray> OutputsArrayPropertyHandle = DetailBuilder.GetProperty(GET_MEMBER_NAME_CHECKED(UNiagaraBakerSettings, Outputs))->AsArray();
-			TSharedPtr<IPropertyHandle> OutputPropertyHandle = OutputsArrayPropertyHandle->GetElement(ViewModel->GetCurrentOutputIndex());
-
-			for (TFieldIterator<FProperty> PropertyIt(CurrentOutput->GetClass()); PropertyIt; ++PropertyIt)
-			{
-				if (PropertyIt->HasAnyPropertyFlags(CPF_Transient | CPF_Deprecated))
-				{
-					continue;
-				}
-				TSharedPtr<IPropertyHandle> PropertyHandle = OutputPropertyHandle->GetChildHandle(PropertyIt->GetFName());
-				if ( PropertyHandle )
-				{
-					OutputTextureCategory.AddProperty(PropertyHandle);
-				}
-			}
-		}
-	}
-
-	TWeakPtr<FNiagaraBakerViewModel> WeakViewModel;
-};
-
-//////////////////////////////////////////////////////////////////////////
-
 void SNiagaraBakerWidget::Construct(const FArguments& InArgs)
 {
 	WeakViewModel = InArgs._WeakViewModel;
@@ -614,20 +509,50 @@ TSharedRef<SWidget> SNiagaraBakerWidget::MakeCameraModeMenu()
 
 	FMenuBuilder MenuBuilder(true, nullptr);
 
-	for (int i = 0; i < int(ENiagaraBakerViewMode::Num); ++i)
+	if ( UNiagaraBakerSettings* BakerSettings = ViewModel->GetBakerSettings() )
 	{
+		for (int i = 0; i < BakerSettings->CameraSettings.Num(); ++i)
+		{
+			if ( i == int32(ENiagaraBakerViewMode::Num) )
+			{
+				MenuBuilder.AddSeparator();
+			}
+
+			MenuBuilder.AddMenuEntry(
+				ViewModel->GetCameraSettingsText(i),
+				FText::GetEmpty(),
+				ViewModel->GetCameraSettingsIcon(i),
+				FUIAction(
+					FExecuteAction::CreateSP(ViewModel, &FNiagaraBakerViewModel::SetCameraSettingsIndex, i),
+					FCanExecuteAction(),
+					FIsActionChecked::CreateSP(ViewModel, &FNiagaraBakerViewModel::IsCameraSettingIndex, i)
+				),
+				NAME_None,
+				EUserInterfaceActionType::ToggleButton
+			);
+		}
+
+		MenuBuilder.AddSeparator();
 		MenuBuilder.AddMenuEntry(
-			FNiagaraBakerViewModel::GetCameraModeText(ENiagaraBakerViewMode(i)),
+			LOCTEXT("AddBookmark", "Add Bookmark"),
 			FText::GetEmpty(),
-			FNiagaraBakerViewModel::GetCameraModeIcon(ENiagaraBakerViewMode(i)),
+			FSlateIcon(),
 			FUIAction(
-				FExecuteAction::CreateSP(ViewModel, &FNiagaraBakerViewModel::SetCameraViewMode, ENiagaraBakerViewMode(i)),
-				FCanExecuteAction(),
-				FIsActionChecked::CreateSP(ViewModel, &FNiagaraBakerViewModel::IsCameraViewMode, ENiagaraBakerViewMode(i))
-			),
-			NAME_None,
-			EUserInterfaceActionType::ToggleButton
+				FExecuteAction::CreateSP(ViewModel, &FNiagaraBakerViewModel::AddCameraBookmark)
+			)
 		);
+
+		if ( BakerSettings->CurrentCameraIndex >= int32(ENiagaraBakerViewMode::Num) )
+		{
+			MenuBuilder.AddMenuEntry(
+				FText::Format(LOCTEXT("RemoveBookmarkFormat", "Remove Bookmark '{0}'"), ViewModel->GetCameraSettingsText(BakerSettings->CurrentCameraIndex)),
+				FText::GetEmpty(),
+				FSlateIcon(),
+				FUIAction(
+					FExecuteAction::CreateSP(ViewModel, &FNiagaraBakerViewModel::RemoveCameraBookmark, BakerSettings->CurrentCameraIndex)
+				)
+			);
+		}
 	}
 
 	return MenuBuilder.MakeWidget();
@@ -718,7 +643,7 @@ TSharedRef<SWidget> SNiagaraBakerWidget::MakeViewOptionsMenu()
 			LOCTEXT("CameraLocation", "Camera Rotation")
 		);
 
-		if ( ViewModel->IsCameraViewMode(ENiagaraBakerViewMode::Perspective) )
+		if ( ViewModel->IsCurrentCameraPerspective() )
 		{
 			MenuBuilder.AddWidget(
 				MakeSpinBox<float>(1.0f, 170.0f, 1.0f, 170.0f, TAttribute<float>::CreateSP(ViewModel, &FNiagaraBakerViewModel::GetCameraFOV), SSpinBox<float>::FOnValueChanged::CreateSP(ViewModel, &FNiagaraBakerViewModel::SetCameraFOV)),
