Niagara - Optimize graph post load.
+ Refactor the code which fixes up names and types on dependent pins in stack graphs so that it uses the variable metadata directly rather than using the traversal to get the stack function input pins.
+ Remove the id refresh on graph duplication because it was refreshing the variable ids when applying module changes.
+ Bump the custom version so that initial popluation of the name bindings won'd happen again once an asset has been saved.
+ Move name binding updates to where the override pin is created to avoid having to call update from as many places.

This results in an around 10x speedup to the calls to RefreshFromExternalChanges on function call nodes when loading all systems in engine tests.  Also all tests pass.

diff --git a/Engine/Plugins/FX/Niagara/Source/NiagaraCore/Public/NiagaraCustomVersion.h b/Engine/Plugins/FX/Niagara/Source/NiagaraCore/Public/NiagaraCustomVersion.h
--- a/Engine/Plugins/FX/Niagara/Source/NiagaraCore/Public/NiagaraCustomVersion.h
+++ b/Engine/Plugins/FX/Niagara/Source/NiagaraCore/Public/NiagaraCustomVersion.h
@@ -174,6 +174,8 @@ struct FNiagaraCustomVersion
 
 		AddBakerCameraBookmarks,
 
+		PopulateFunctionCallNodePinNameBindings, // Function call node refresh from external changes has been refactored so that they don't need to populate their name bindings every load.
+
 		// DO NOT ADD A NEW VERSION UNLESS YOU HAVE TALKED TO THE NIAGARA LEAD. Mismanagement of these versions can lead to data loss if it is adjusted in multiple streams simultaneously.
 		// -----<new versions can be added above this line>  -------------------------------------------------
 		VersionPlusOne,
diff --git a/Engine/Plugins/FX/Niagara/Source/NiagaraEditor/Private/NiagaraNodeFunctionCall.cpp b/Engine/Plugins/FX/Niagara/Source/NiagaraEditor/Private/NiagaraNodeFunctionCall.cpp
--- a/Engine/Plugins/FX/Niagara/Source/NiagaraEditor/Private/NiagaraNodeFunctionCall.cpp
+++ b/Engine/Plugins/FX/Niagara/Source/NiagaraEditor/Private/NiagaraNodeFunctionCall.cpp
@@ -93,6 +93,11 @@ void UNiagaraNodeFunctionCall::PostLoad()
 		{
 			UpdateStaticSwitchPinsWithPersistentGuids();
 		}
+
+		if (NiagaraCustomVersion < FNiagaraCustomVersion::PopulateFunctionCallNodePinNameBindings)
+		{
+			FNiagaraStackGraphUtilities::PopulateFunctionCallNameBindings(*this);
+		}
 	}
 
 	// Allow data interfaces an opportunity to intercept changes
@@ -118,12 +123,6 @@ void UNiagaraNodeFunctionCall::PostLoad()
 	{
 		ComputeNodeName();
 	}
-	
-	// check if maybe the parameter names in the referenced module were changed and try to move over existing values
-	UNiagaraGraph* Graph = GetCalledGraph();
-	if (Graph)
-		Graph->ConditionalPostLoad();
-	FixupPinNames();
 }
 
 void UNiagaraNodeFunctionCall::UpgradeDIFunctionCalls()
@@ -412,227 +411,6 @@ bool UNiagaraNodeFunctionCall::IsDeprecated() const
 	return false;
 }
 
-bool UNiagaraNodeFunctionCall::FixupPinNames()
-{
-	// This checks the override pins to see if any of the inputs were renamed since we last checked.
-	// Note that this does *not* fix inputs set via rapid iteration parameters, they are fixed by UNiagaraScriptSource::FixupRenamedParameters.
-	if (UNiagaraNodeParameterMapSet* OverrideNode = FNiagaraStackGraphUtilities::GetStackFunctionOverrideNode(*this))
-	{
-		FixupFunctionScriptVersion();
-		UNiagaraGraph* Graph = GetCalledGraph();
-		if (Graph == nullptr)
-		{
-			return false;
-		}
-		const UEdGraphSchema_Niagara* NiagaraSchema = Graph->GetNiagaraSchema();
-		const UNiagaraSettings* Settings = GetDefault<UNiagaraSettings>();
-		
-		// find out which inputs the module offers
-		TSet<const UEdGraphPin*> HiddenModulePins;
-		TArray<const UEdGraphPin*> ModuleInputPins;
-		FCompileConstantResolver ConstantResolver;
-		GetStackFunctionInputPins(*this, ModuleInputPins, HiddenModulePins, ConstantResolver, FNiagaraStackGraphUtilities::ENiagaraGetStackFunctionInputPinsOptions::ModuleInputsOnly);
-		TMap<FName, const UEdGraphPin*> ModulePinNames;
-		TMap<FGuid, FNiagaraVariable> ModuleInputGuidMapping;
-		for (const UEdGraphPin* ModulePin : ModuleInputPins)
-		{
-			FName InputName = FNiagaraParameterHandle(ModulePin->PinName).GetName();
-			FNiagaraVariable InputVar = NiagaraSchema->PinToNiagaraVariable(ModulePin);
-			TOptional<FNiagaraVariableMetaData> VariableMetaData = Graph->GetMetaData(InputVar);
-			if (VariableMetaData.IsSet()) // check if there is metadata with a guid available
-			{
-				ModuleInputGuidMapping.Add(VariableMetaData->GetVariableGuid(), InputVar);
-				ModulePinNames.Add(InputName, ModulePin);
-			}
-		}
-
-		// gather inputs that are overriden
-		TArray<UEdGraphPin*> OverrideInputPins;
-		TArray<UEdGraphPin*> LinkedModuleOutputs;
-		OverrideNode->GetInputPins(OverrideInputPins);
-		bool bNodeChanged = false;
-		for (UEdGraphPin* OverridePin : OverrideInputPins)
-		{
-			if (OverridePin->Direction != EGPD_Input || !OverridePin->PinName.ToString().StartsWith(GetFunctionName()))
-			{
-				continue;
-			}
-
-			// If the type of our input was changed from vec3 to position we update the override pin type as well
-			FNiagaraVariable InputVar = NiagaraSchema->PinToNiagaraVariable(OverridePin);
-			FNiagaraParameterHandle CurrentPinHandle = FNiagaraParameterHandle(OverridePin->PinName);
-			FName InputName = CurrentPinHandle.GetName();
-			if (ModulePinNames.Contains(InputName) && InputVar.GetType() == FNiagaraTypeDefinition::GetVec3Def())
-			{
-				FNiagaraVariable ModuleInputVar = NiagaraSchema->PinToNiagaraVariable(ModulePinNames[InputName]);
-				if (ModuleInputVar.GetType() == FNiagaraTypeDefinition::GetPositionDef())
-				{
-					OverridePin->PinType = ModulePinNames[InputName]->PinType;
-					bNodeChanged = true;
-				}
-			}
-
-			// Gather module output linked to our inputs
-			if (UEdGraphPin* LinkedValuePin = FNiagaraStackGraphUtilities::GetLinkedValueHandleForFunctionInput(*OverridePin))
-			{
-				FNiagaraParameterHandle PinHandle(LinkedValuePin->PinName);
-				FString LinkedToName = PinHandle.GetParameterHandleString().ToString();
-				if (PinHandle.IsOutputHandle())
-				{
-					LinkedModuleOutputs.Add(LinkedValuePin);
-				}
-
-				// when changing the module input type to position, changes the linked value pin as well
-				if (InputVar.GetType() == FNiagaraTypeDefinition::GetPositionDef() && OverridePin->PinType != LinkedValuePin->PinType && Settings->bEnforceStrictStackTypes)
-				{
-					LinkedValuePin->PinType = OverridePin->PinType;
-					bNodeChanged = true;
-				}
-			}
-
-			// check if the pin is still a valid input name from the module
-			if (ModulePinNames.Contains(InputName))
-			{
-				// update the existing binding data
-				for (auto Entry : ModuleInputGuidMapping)
-				{
-					FNiagaraParameterHandle ModuleInputName = FNiagaraParameterHandle(Entry.Value.GetName());
-					if (ModuleInputName.GetName() == InputName && Entry.Value.GetType() == InputVar.GetType())
-					{
-						UpdateInputNameBinding(Entry.Key, InputName);
-						break;
-					}
-				}
-				continue;
-			}
-
-			// see if we can find a match with a previously set guid and name
-			TArray<FGuid> PossibleGuidMatches = GetBoundPinGuidsByName(InputName);
-			for (FGuid InputGuid : PossibleGuidMatches)
-			{
-				FNiagaraVariable* ModuleVar = ModuleInputGuidMapping.Find(InputGuid);
-				if (ModuleVar && ModuleVar->GetType() == InputVar.GetType())
-				{
-					// found a match! Lets rename the pin
-					FNiagaraParameterHandle CurrentVarName(ModuleVar->GetName());
-					FNiagaraParameterHandle NewName(CurrentPinHandle.GetNamespace(), CurrentVarName.GetName());
-					OverridePin->PinName = NewName.GetParameterHandleString();
-					bNodeChanged = true;
-					break;
-				}
-			}
-		}
-
-		// see if any of the linked inputs were renamed
-		UNiagaraGraph* Owner = GetNiagaraGraph();
-		TMap<FGuid, TArray<FNiagaraVariable>> GraphVariablesGuidMapping;
-		if (Owner && LinkedModuleOutputs.Num() > 0)
-		{
-			// Search for guids of the linked inputs by looking backwards in the current script graph (not the graph this function node calls, but the one it's embedded in)
-			TArray<UNiagaraNode*> NodesTraversed;
-			UNiagaraGraph::BuildTraversal(NodesTraversed, this);
-			for (UNiagaraNode* Node : NodesTraversed)
-			{
-				if (Node != this && Node->IsA<UNiagaraNodeFunctionCall>())
-				{
-					UNiagaraNodeFunctionCall* FunctionNode = CastChecked<UNiagaraNodeFunctionCall>(Node);
-					FunctionNode->ConditionalPostLoad();
-					UNiagaraGraph* CalledGraph = FunctionNode->GetCalledGraph();
-					if (!CalledGraph)
-					{
-						continue;
-					}
-					const TMap<FNiagaraVariable, TObjectPtr<UNiagaraScriptVariable>>& MetaData = CalledGraph->GetAllMetaData();
-					for (const auto& Entry : MetaData)
-					{
-						if (Entry.Value && Entry.Key.IsInNameSpace(FNiagaraConstants::OutputNamespaceString))
-						{
-							TArray<FName> HandleParts = FNiagaraParameterHandle(Entry.Key.GetName()).GetHandleParts();
-							if (HandleParts.Num() >= 3 && HandleParts[1] == FNiagaraConstants::ModuleNamespace)
-							{
-								FString AliasedName = *FString::Printf(TEXT("%s.%s"), *FNiagaraConstants::OutputNamespace.ToString(), *FunctionNode->GetFunctionName());
-								for (int i = 2; i < HandleParts.Num(); i++)
-								{
-									AliasedName.Appendf(TEXT(".%s"), *HandleParts[i].ToString());
-								}
-								FNiagaraVariable OutputVar = Entry.Key;
-								OutputVar.SetName(FName(AliasedName));
-								GraphVariablesGuidMapping.FindOrAdd(Entry.Value->Metadata.GetVariableGuid()).Add(OutputVar);
-							}
-						}
-					}
-				}
-			}
-
-			// Now compare the linked inputs with the guids and names we have
-			for (UEdGraphPin* LinkedOutPin : LinkedModuleOutputs)
-			{
-				// See if we have a matching guid for this input
-				bool bFoundMatch = false;
-				TArray<FGuid> GuidMatches = GetBoundPinGuidsByName(LinkedOutPin->PinName);
-				for (FGuid Guid : GuidMatches)
-				{
-					// see if the name still matches and rename the linked input if not
-					TArray<FNiagaraVariable>* VarsFromGraph = GraphVariablesGuidMapping.Find(Guid);
-					if (VarsFromGraph)
-					{
-						bFoundMatch = true;
-
-						// there might be several modules in different version giving us the same guid for the output, so we check them all
-						if (!VarsFromGraph->ContainsByPredicate([LinkedOutPin](const FNiagaraVariable& Var) { return LinkedOutPin->PinName == Var.GetName(); }))
-						{
-							FNiagaraParameterHandle PinHandle(LinkedOutPin->PinName);
-							FNiagaraVariable* OriginalModuleVar = VarsFromGraph->FindByPredicate([PinHandle](const FNiagaraVariable& Var)
-							{
-								FNiagaraParameterHandle ModuleVarHandle(Var.GetName());
-								return PinHandle.GetNamespace() == ModuleVarHandle.GetNamespace();
-							});
-							
-							if (OriginalModuleVar)
-							{
-								// we found a reference from the original module namespace, so use that
-								LinkedOutPin->PinName = OriginalModuleVar->GetName();
-							}
-							else
-							{
-								// seems like the module was lost or got replaced, so just use any one of the new handles
-								LinkedOutPin->PinName = (*VarsFromGraph)[0].GetName();
-							}
-							
-							bNodeChanged = true;
-							UpdateInputNameBinding(Guid, LinkedOutPin->PinName);
-							break;
-						}
-					}
-				}
-
-				// We have not found a match, see if we can still find a matching var to save the guid for upcoming renames
-				if (!bFoundMatch)
-				{
-					FNiagaraVariable InputVariable = NiagaraSchema->PinToNiagaraVariable(LinkedOutPin);
-					for (auto& Entry : GraphVariablesGuidMapping)
-					{
-						FNiagaraVariable* MatchingVar = Entry.Value.FindByPredicate([InputVariable](const FNiagaraVariable& Var){ return InputVariable == Var; });
-						if (MatchingVar)
-						{
-							UpdateInputNameBinding(Entry.Key, LinkedOutPin->PinName);
-						}
-					}
-				}
-			}
-		}
-
-		if (bNodeChanged)
-		{
-			UNiagaraGraph* LinkedPinGraph = Cast<UNiagaraGraph>(OverrideNode->GetGraph());
-			LinkedPinGraph->NotifyGraphNeedsRecompile();
-		}
-		
-		return bNodeChanged;
-	}
-	return false;
-}
-
 void UNiagaraNodeFunctionCall::UpdateInputNameBinding(const FGuid& BoundVariableGuid, const FName& BoundName)
 {
 	if (!BoundVariableGuid.IsValid() || BoundName.IsNone())
@@ -1178,54 +956,6 @@ void UNiagaraNodeFunctionCall::UpdateStaticSwitchPinsWithPersistentGuids()
 	}
 }
 
-void FixupLinkedOutputs(UNiagaraNode* StartNode)
-{
-	// This searches for function call nodes linked to this node via the graph and calls FixupPinNames on them so they can react to changes in this node's called graph
-	TSet<UNiagaraNode*> SeenNodes;
-	TArray<UNiagaraNode*> NodesToCheck;
-	NodesToCheck.Add(StartNode);
-	int LoopGuard = 0;
-	while (NodesToCheck.Num() > 0)
-	{
-		UNiagaraNode* Node = NodesToCheck.Pop();
-		if (SeenNodes.Contains(Node))
-		{
-			continue;
-		}
-		SeenNodes.Add(Node);
-
-		if (UNiagaraNodeFunctionCall* FunctionNode = Cast<UNiagaraNodeFunctionCall>(Node))
-		{
-			FunctionNode->ConditionalPostLoad();
-			FunctionNode->FixupPinNames();
-		}
-		
-		for (UEdGraphPin* Pin : Node->GetAllPins())
-		{
-			if (Pin->Direction == EGPD_Output && Pin->LinkedTo.Num() > 0)
-			{
-				for (UEdGraphPin* LinkedPin : Pin->LinkedTo)
-				{
-					if (LinkedPin != nullptr)
-					{
-						UNiagaraNode* LinkedNode = Cast<UNiagaraNode>(LinkedPin->GetOwningNode());
-						if (LinkedNode)
-						{
-							NodesToCheck.Add(LinkedNode);
-						}
-					}
-				}
-			}
-		}
-
-		LoopGuard++;
-		if (!ensure(LoopGuard < 100000)) // Guard against loops in the graph or other bugs
-		{
-			return;
-		}
-	}
-}
-
 bool UNiagaraNodeFunctionCall::RefreshFromExternalChanges()
 {
 	bool bReload = false;
@@ -1269,14 +999,11 @@ bool UNiagaraNodeFunctionCall::RefreshFromExternalChanges()
 		}
 	}
 
-	// check if one of our overriden input parameters was renamed 
-	FixupPinNames();
-
 	if (bReload)
 	{
 		// TODO - Leverage code in reallocate pins to determine if any pins have changed...
 		ReallocatePins(false);
-		FixupLinkedOutputs(this); // If our called graph changed we might need to update downstream function call nodes as well
+		FNiagaraStackGraphUtilities::SynchronizeReferencingMapPinsWithFunctionCall(*this);
 		return true;
 	}
 	else
diff --git a/Engine/Plugins/FX/Niagara/Source/NiagaraEditor/Private/NiagaraScriptMergeManager.cpp b/Engine/Plugins/FX/Niagara/Source/NiagaraEditor/Private/NiagaraScriptMergeManager.cpp
--- a/Engine/Plugins/FX/Niagara/Source/NiagaraEditor/Private/NiagaraScriptMergeManager.cpp
+++ b/Engine/Plugins/FX/Niagara/Source/NiagaraEditor/Private/NiagaraScriptMergeManager.cpp
@@ -2820,6 +2820,12 @@ FNiagaraScriptMergeManager::FApplyDiffResults FNiagaraScriptMergeManager::AddInp
 					DesiredMode = ENiagaraDefaultMode::FailIfPreviouslyNotSet;
 				TSet KnownParameters = { Parameter };
 				FNiagaraStackGraphUtilities::SetLinkedValueHandleForFunctionInput(InputOverridePin, OverrideToAdd->GetLinkedValueHandle().GetValue(), KnownParameters, DesiredMode, OverrideToAdd->GetOverrideNodeId());
+				FGuid LinkedOutputId = FNiagaraStackGraphUtilities::GetScriptVariableIdForLinkedOutputHandle(
+					OverrideToAdd->GetLinkedValueHandle().GetValue(), OverrideToAdd->GetType(), *TargetFunctionCall.GetNiagaraGraph());
+				if (LinkedOutputId.IsValid())
+				{
+					TargetFunctionCall.UpdateInputNameBinding(LinkedOutputId, OverrideToAdd->GetLinkedValueHandle()->GetName());
+				}
 				Results.bSucceeded = true;
 			}
 			else if (OverrideToAdd->GetDataValueInputName().IsSet() && OverrideToAdd->GetDataValueObject() != nullptr)
diff --git a/Engine/Plugins/FX/Niagara/Source/NiagaraEditor/Private/NiagaraStackFunctionInputBinder.cpp b/Engine/Plugins/FX/Niagara/Source/NiagaraEditor/Private/NiagaraStackFunctionInputBinder.cpp
--- a/Engine/Plugins/FX/Niagara/Source/NiagaraEditor/Private/NiagaraStackFunctionInputBinder.cpp
+++ b/Engine/Plugins/FX/Niagara/Source/NiagaraEditor/Private/NiagaraStackFunctionInputBinder.cpp
@@ -121,6 +121,9 @@ bool FNiagaraStackFunctionInputBinder::TryBindInternal(
 
 			InputType = InputVariable.GetType();
 
+			TOptional<FNiagaraVariableMetaData> InputMetadata = InFunctionCallNode->GetNiagaraGraph()->GetMetaData(InputVariable);
+			InputVariableGuid = InputMetadata.IsSet() ? InputMetadata->GetVariableGuid() : FGuid();
+
 			AliasedParameterHandle = FNiagaraParameterHandle::CreateAliasedModuleParameterHandle(FNiagaraParameterHandle(InputVariable.GetName()), FunctionCallNode.Get());
 				
 			RefreshGraphPins();
@@ -283,7 +286,7 @@ void FNiagaraStackFunctionInputBinder::SetData(const uint8* InValue, int32 InSiz
 
 			if (OverridePin == nullptr)
 			{
-				OverridePin = &FNiagaraStackGraphUtilities::GetOrCreateStackFunctionInputOverridePin(*FunctionCallNode, AliasedParameterHandle, InputType);
+				OverridePin = &FNiagaraStackGraphUtilities::GetOrCreateStackFunctionInputOverridePin(*FunctionCallNode, AliasedParameterHandle, InputType, InputVariableGuid);
 			}
 
 			OverridePin->Modify();
diff --git a/Engine/Plugins/FX/Niagara/Source/NiagaraEditor/Private/ViewModels/NiagaraPlaceholderDataInterfaceManager.cpp b/Engine/Plugins/FX/Niagara/Source/NiagaraEditor/Private/ViewModels/NiagaraPlaceholderDataInterfaceManager.cpp
--- a/Engine/Plugins/FX/Niagara/Source/NiagaraEditor/Private/ViewModels/NiagaraPlaceholderDataInterfaceManager.cpp
+++ b/Engine/Plugins/FX/Niagara/Source/NiagaraEditor/Private/ViewModels/NiagaraPlaceholderDataInterfaceManager.cpp
@@ -122,8 +122,12 @@ void FNiagaraPlaceholderDataInterfaceManager::PlaceholderDataInterfaceChanged(TW
 		if (OverrideDataInterface == nullptr)
 		{
 			FNiagaraParameterHandle AliasedInputParameterHandle = FNiagaraParameterHandle::CreateAliasedModuleParameterHandle(PlaceholderDataInterfaceInfo->InputHandle, PlaceholderDataInterfaceInfo->OwningFunctionCall.Get());
+			FNiagaraVariable ParameterVariable(FNiagaraTypeDefinition(PlaceholderDataInterface->GetClass()), PlaceholderDataInterfaceInfo->InputHandle.GetParameterHandleString());
+			TOptional<FNiagaraVariableMetaData> ParameterMetaData = PlaceholderDataInterfaceInfo->OwningFunctionCall->GetNiagaraGraph()->GetMetaData(ParameterVariable);
+			FGuid ParameterVariableGuid = ParameterMetaData.IsSet() ? ParameterMetaData->GetVariableGuid() : FGuid();
 			UEdGraphPin& OverridePin = FNiagaraStackGraphUtilities::GetOrCreateStackFunctionInputOverridePin(
-				*PlaceholderDataInterfaceInfo->OwningFunctionCall.Get(), AliasedInputParameterHandle, FNiagaraTypeDefinition(PlaceholderDataInterface->GetClass()));
+				*PlaceholderDataInterfaceInfo->OwningFunctionCall.Get(), AliasedInputParameterHandle,
+				FNiagaraTypeDefinition(PlaceholderDataInterface->GetClass()), ParameterVariableGuid);
 
 			if (OverridePin.LinkedTo.Num() == 1)
 			{
diff --git a/Engine/Plugins/FX/Niagara/Source/NiagaraEditor/Private/ViewModels/Stack/NiagaraStackFunctionInput.cpp b/Engine/Plugins/FX/Niagara/Source/NiagaraEditor/Private/ViewModels/Stack/NiagaraStackFunctionInput.cpp
--- a/Engine/Plugins/FX/Niagara/Source/NiagaraEditor/Private/ViewModels/Stack/NiagaraStackFunctionInput.cpp
+++ b/Engine/Plugins/FX/Niagara/Source/NiagaraEditor/Private/ViewModels/Stack/NiagaraStackFunctionInput.cpp
@@ -1251,7 +1251,11 @@ void UNiagaraStackFunctionInput::SetLinkedValueHandle(const FNiagaraParameterHan
 		UEdGraphPin& OverridePin = GetOrCreateOverridePin();
 		TSet<FNiagaraVariable> KnownParameters = FNiagaraStackGraphUtilities::GetParametersForContext(OverridePin.GetOwningNode()->GetGraph(), GetSystemViewModel()->GetSystem());
 		FNiagaraStackGraphUtilities::SetLinkedValueHandleForFunctionInput(OverridePin, InParameterHandle, KnownParameters);
-		OwningFunctionCallNode->FixupPinNames(); // refresh the input guids and update the bound names
+		FGuid LinkedOutputId = FNiagaraStackGraphUtilities::GetScriptVariableIdForLinkedOutputHandle(InParameterHandle, InputType, *OwningFunctionCallNode->GetNiagaraGraph());
+		if (LinkedOutputId.IsValid())
+		{
+			OwningFunctionCallNode->UpdateInputNameBinding(LinkedOutputId, InParameterHandle.GetName());
+		}
 	}
 
 	FNiagaraStackGraphUtilities::RelayoutGraph(*OwningFunctionCallNode->GetGraph());
@@ -1745,12 +1749,6 @@ void UNiagaraStackFunctionInput::SetLocalValue(TSharedRef<FStructOnScope> InLoca
 			OverridePin->DefaultValue = PinDefaultValue;
 			Cast<UNiagaraNode>(OverridePin->GetOwningNode())->MarkNodeRequiresSynchronization(TEXT("OverridePin Default Value Changed"), true);
 		}
-
-		if (InputMetaData.IsSet())
-		{
-			FName VariableName = FNiagaraParameterHandle(RapidIterationParameter.GetName()).GetName();
-			OwningFunctionCallNode->UpdateInputNameBinding(InputMetaData->GetVariableGuid(), VariableName);
-		}
 	}
 	
 	if (bGraphWillNeedRelayout)
@@ -2703,20 +2701,9 @@ UEdGraphPin& UNiagaraStackFunctionInput::GetOrCreateOverridePin()
 	if (OverridePin == nullptr)
 	{
 		TGuardValue<bool> Guard(bUpdatingGraphDirectly, true);
-		OverridePin = &FNiagaraStackGraphUtilities::GetOrCreateStackFunctionInputOverridePin(*OwningFunctionCallNode, AliasedInputParameterHandle, InputType);
+		FGuid InputScriptVariableId = InputMetaData.IsSet() ? InputMetaData->GetVariableGuid() : FGuid();
+		OverridePin = &FNiagaraStackGraphUtilities::GetOrCreateStackFunctionInputOverridePin(*OwningFunctionCallNode, AliasedInputParameterHandle, InputType, InputScriptVariableId);
 		OverridePinCache = OverridePin;
-
-		// update name binding if we have a guid
-		UNiagaraGraph* CalledGraph = OwningFunctionCallNode->GetCalledGraph();
-		if (CalledGraph)
-		{
-			const FName ModuleInputName = FNiagaraParameterHandle::CreateModuleParameterHandle(AliasedInputParameterHandle.GetName()).GetParameterHandleString();
-			TOptional<FNiagaraVariableMetaData> MetaData = CalledGraph->GetMetaData(FNiagaraVariable(InputType, ModuleInputName));
-			if (MetaData.IsSet())
-			{
-				OwningFunctionCallNode->UpdateInputNameBinding(MetaData->GetVariableGuid(), AliasedInputParameterHandle.GetName());
-			}
-		}
 	}
 	return *OverridePin;
 }
diff --git a/Engine/Plugins/FX/Niagara/Source/NiagaraEditor/Private/ViewModels/Stack/NiagaraStackGraphUtilities.cpp b/Engine/Plugins/FX/Niagara/Source/NiagaraEditor/Private/ViewModels/Stack/NiagaraStackGraphUtilities.cpp
--- a/Engine/Plugins/FX/Niagara/Source/NiagaraEditor/Private/ViewModels/Stack/NiagaraStackGraphUtilities.cpp
+++ b/Engine/Plugins/FX/Niagara/Source/NiagaraEditor/Private/ViewModels/Stack/NiagaraStackGraphUtilities.cpp
@@ -311,14 +311,14 @@ template<typename OutputNodeType, typename InputNodeType>
 OutputNodeType* GetEmitterOutputNodeForStackNodeInternal(InputNodeType& StackNode)
 {
 	TArray<InputNodeType*> NodesToCheck;
-	TSet<InputNodeType*> NodesChecked;
+	TSet<InputNodeType*> NodesSeen;
 	FPinCollectorArray OutputPins;
+	NodesSeen.Add(&StackNode);
 	NodesToCheck.Add(&StackNode);
 	while (NodesToCheck.Num() > 0)
 	{
 		InputNodeType* NodeToCheck = NodesToCheck[0];
 		NodesToCheck.RemoveAt(0);
-		NodesChecked.Add(NodeToCheck);
 
 		if (NodeToCheck->GetClass() == UNiagaraNodeOutput::StaticClass())
 		{
@@ -332,8 +332,9 @@ OutputNodeType* GetEmitterOutputNodeForStackNodeInternal(InputNodeType& StackNod
 			for (UEdGraphPin* LinkedPin : OutputPin->LinkedTo)
 			{
 				InputNodeType* LinkedNiagaraNode = Cast<UNiagaraNode>(LinkedPin->GetOwningNode());
-				if (LinkedNiagaraNode != nullptr && NodesChecked.Contains(LinkedNiagaraNode) == false)
+				if (LinkedNiagaraNode != nullptr && NodesSeen.Contains(LinkedNiagaraNode) == false)
 				{
+					NodesSeen.Add(LinkedNiagaraNode);
 					NodesToCheck.Add(LinkedNiagaraNode);
 				}
 			}
@@ -1224,7 +1225,7 @@ UEdGraphPin* FNiagaraStackGraphUtilities::GetStackFunctionInputOverridePin(UNiag
 	return nullptr;
 }
 
-UEdGraphPin& FNiagaraStackGraphUtilities::GetOrCreateStackFunctionInputOverridePin(UNiagaraNodeFunctionCall& StackFunctionCall, FNiagaraParameterHandle AliasedInputParameterHandle, FNiagaraTypeDefinition InputType, const FGuid& PreferredOverrideNodeGuid)
+UEdGraphPin& FNiagaraStackGraphUtilities::GetOrCreateStackFunctionInputOverridePin(UNiagaraNodeFunctionCall& StackFunctionCall, FNiagaraParameterHandle AliasedInputParameterHandle, FNiagaraTypeDefinition InputType, const FGuid& InputScriptVariableId, const FGuid& PreferredOverrideNodeGuid)
 {
 	UEdGraphPin* OverridePin = GetStackFunctionInputOverridePin(StackFunctionCall, AliasedInputParameterHandle);
 	if (OverridePin == nullptr)
@@ -1238,6 +1239,11 @@ UEdGraphPin& FNiagaraStackGraphUtilities::GetOrCreateStackFunctionInputOverrideP
 		const UEdGraphSchema_Niagara* NiagaraSchema = GetDefault<UEdGraphSchema_Niagara>();
 		FEdGraphPinType PinType = NiagaraSchema->TypeDefinitionToPinType(InputType);
 		OverridePin = OverrideNode.CreatePin(EEdGraphPinDirection::EGPD_Input, PinType, AliasedInputParameterHandle.GetParameterHandleString(), OverrideInputPins.Num() - 1);
+
+		if(InputScriptVariableId.IsValid())
+		{
+			StackFunctionCall.UpdateInputNameBinding(InputScriptVariableId, AliasedInputParameterHandle.GetName());
+		}
 	}
 	return *OverridePin;
 }
@@ -2221,147 +2227,6 @@ struct FRapidIterationParameterContext
 	UNiagaraNodeFunctionCall* const OwningFunctionCall;
 };
 
-struct FStackFunctionInputValue
-{
-	FNiagaraTypeDefinition Type;
-	bool bIsOverride;
-	TOptional<FNiagaraVariable> LocalValue;
-	TOptional<FName> LinkedValue;
-	TOptional<UNiagaraDataInterface*> DataValue;
-	TOptional<UNiagaraNodeFunctionCall*> DynamicValue;
-	TMap<FName, TSharedRef<FStackFunctionInputValue>> DynamicValueInputs;
-
-	bool Matches(const FStackFunctionInputValue& Other) const
-	{
-		if (Type != Other.Type)
-		{
-			return false;
-		}
-		if (LocalValue.IsSet() && Other.LocalValue.IsSet())
-		{
-			return LocalValue.GetValue().GetType() == Other.LocalValue.GetValue().GetType() &&
-				FMemory::Memcmp(LocalValue.GetValue().GetData(), Other.LocalValue.GetValue().GetData(), LocalValue.GetValue().GetSizeInBytes()) == 0;
-		}
-		else if (LinkedValue.IsSet() && Other.LinkedValue.IsSet())
-		{
-			return LinkedValue.GetValue() == Other.LinkedValue.GetValue();
-		}
-		else if (DataValue.IsSet() && Other.DataValue.IsSet())
-		{
-			return (DataValue.GetValue() == nullptr && Other.DataValue.GetValue() == nullptr) ||
-				(DataValue.GetValue() != nullptr && Other.DataValue.GetValue() != nullptr && DataValue.GetValue()->Equals(Other.DataValue.GetValue()));
-		}
-		else if (DynamicValue.IsSet() && Other.DynamicValue.IsSet())
-		{
-			if (DynamicValue.GetValue()->FunctionScript == Other.DynamicValue.GetValue()->FunctionScript)
-			{
-				for (auto It = DynamicValueInputs.CreateConstIterator(); It; ++It)
-				{
-					FName InputName = It.Key();
-					TSharedRef<FStackFunctionInputValue> InputValue = It.Value();
-					const TSharedRef<FStackFunctionInputValue>* OtherInputValue = Other.DynamicValueInputs.Find(InputName);
-					if (OtherInputValue == nullptr || InputValue->Matches(**OtherInputValue) == false)
-					{
-						return false;
-					}
-				}
-				return true;
-			}
-			else
-			{
-				return false;
-			}
-		}
-		else
-		{
-			return false;
-		}
-	}
-};
-
-bool TryGetStackFunctionInputValue(UNiagaraScript& OwningScript, const UEdGraphPin* OverridePin, const UEdGraphPin& DefaultPin, FName InputName, FRapidIterationParameterContext RapidIterationParameterContext, FStackFunctionInputValue& OutStackFunctionInputValue)
-{
-	OutStackFunctionInputValue.Type = GetDefault<UEdGraphSchema_Niagara>()->PinToTypeDefinition(&DefaultPin);
-	OutStackFunctionInputValue.bIsOverride = OverridePin != nullptr;
-	const UEdGraphPin& InputPin = OverridePin != nullptr ? *OverridePin : DefaultPin;
-	if (RapidIterationParameterContext.IsValid() && DefaultPin.LinkedTo.Num() == 0 && OverridePin == nullptr)
-	{
-		OutStackFunctionInputValue.LocalValue = RapidIterationParameterContext.GetValue(OwningScript, InputName, OutStackFunctionInputValue.Type);
-	}
-	else if (InputPin.LinkedTo.Num() == 0)
-	{
-		const UEdGraphSchema_Niagara* NiagaraSchema = GetDefault<UEdGraphSchema_Niagara>();
-		if (NiagaraSchema->PinToTypeDefinition(&InputPin).IsDataInterface())
-		{
-			OutStackFunctionInputValue.DataValue = nullptr;
-		}
-		else
-		{
-			OutStackFunctionInputValue.LocalValue = NiagaraSchema->PinToNiagaraVariable(&InputPin, true);
-		}
-	}
-	else if (InputPin.LinkedTo.Num() == 1)
-	{
-		UEdGraphNode* PreviousOwningNode = InputPin.LinkedTo[0]->GetOwningNode();
-		
-		if (PreviousOwningNode->IsA<UNiagaraNodeParameterMapGet>())
-		{
-			OutStackFunctionInputValue.LinkedValue = InputPin.LinkedTo[0]->GetFName();
-		}
-		else if (PreviousOwningNode->IsA<UNiagaraNodeInput>())
-		{
-			OutStackFunctionInputValue.DataValue = CastChecked<UNiagaraNodeInput>(InputPin.LinkedTo[0]->GetOwningNode())->GetDataInterface();
-		}
-		else if (PreviousOwningNode->IsA<UNiagaraNodeFunctionCall>() && 
-			FNiagaraStackGraphUtilities::GetParameterMapInputPin(*(static_cast<UNiagaraNodeFunctionCall*>(PreviousOwningNode))) != nullptr)
-		{
-			UNiagaraNodeFunctionCall* DynamicInputFunctionCall = CastChecked<UNiagaraNodeFunctionCall>(InputPin.LinkedTo[0]->GetOwningNode());
-			OutStackFunctionInputValue.DynamicValue = DynamicInputFunctionCall;
-			TArray<const UEdGraphPin*> DynamicValueInputPins;
-			GetStackFunctionInputPins(*DynamicInputFunctionCall, DynamicValueInputPins, FNiagaraStackGraphUtilities::ENiagaraGetStackFunctionInputPinsOptions::ModuleInputsOnly);
-
-			FRapidIterationParameterContext InputRapidIterationParameterContext = RapidIterationParameterContext.IsValid() 
-				? FRapidIterationParameterContext(RapidIterationParameterContext.UniqueEmitterName, *DynamicInputFunctionCall)
-				: FRapidIterationParameterContext();
-
-			for (const UEdGraphPin* DynamicValueInputPin : DynamicValueInputPins)
-			{
-				FNiagaraParameterHandle ModuleHandle(DynamicValueInputPin->GetFName());
-				UEdGraphPin* DynamicValueOverridePin = FNiagaraStackGraphUtilities::GetStackFunctionInputOverridePin(*DynamicInputFunctionCall,
-					FNiagaraParameterHandle::CreateAliasedModuleParameterHandle(ModuleHandle, DynamicInputFunctionCall));
-
-				UEdGraphPin* DynamicValueInputDefaultPin = DynamicInputFunctionCall->FindParameterMapDefaultValuePin(DynamicValueInputPin->PinName, OwningScript.GetUsage(), FCompileConstantResolver());
-
-				FStackFunctionInputValue InputValue;
-				if (TryGetStackFunctionInputValue(OwningScript, DynamicValueOverridePin, *DynamicValueInputDefaultPin, ModuleHandle.GetName(), InputRapidIterationParameterContext, InputValue))
-				{
-					OutStackFunctionInputValue.DynamicValueInputs.Add(ModuleHandle.GetName(), MakeShared<FStackFunctionInputValue>(InputValue));
-				}
-				else
-				{
-					return false;
-				}
-			}
-		}
-		else
-		{
-			return false;
-		}
-	}
-	else
-	{
-		return false;
-	}
-	return true;
-}
-
-bool FNiagaraStackGraphUtilities::IsValidDefaultDynamicInput(UNiagaraScript& OwningScript, UEdGraphPin& DefaultPin)
-{
-	FStackFunctionInputValue InputValue;
-	return TryGetStackFunctionInputValue(OwningScript, nullptr, DefaultPin, NAME_None, FRapidIterationParameterContext(), InputValue) && InputValue.DynamicValue.IsSet();
-}
-
-
 bool FNiagaraStackGraphUtilities::CanWriteParameterFromUsageViaOutput(FNiagaraVariable Parameter, const UNiagaraNodeOutput* OutputNode)
 {
 	bool bCanWrite = CanWriteParameterFromUsage(Parameter, OutputNode->GetUsage(), OutputNode->GetStackContextOverride(), OutputNode->GetAllStackContextOverrides());	
@@ -2421,123 +2286,6 @@ bool FNiagaraStackGraphUtilities::CanWriteParameterFromUsage(FNiagaraVariable Pa
 	}
 }
 
-bool FNiagaraStackGraphUtilities::DoesDynamicInputMatchDefault(
-	FString EmitterUniqueName, 
-	UNiagaraScript& OwningScript,
-	UNiagaraNodeFunctionCall& OwningFunctionCallNode,
-	UEdGraphPin& OverridePin,
-	FName InputName,
-	UEdGraphPin& DefaultPin)
-{
-	FStackFunctionInputValue CurrentValue;
-	FStackFunctionInputValue DefaultValue;
-	return
-		TryGetStackFunctionInputValue(OwningScript, &OverridePin, DefaultPin, InputName, FRapidIterationParameterContext(EmitterUniqueName, OwningFunctionCallNode), CurrentValue) &&
-		TryGetStackFunctionInputValue(OwningScript, nullptr, DefaultPin, NAME_None, FRapidIterationParameterContext(), DefaultValue) &&
-		CurrentValue.Matches(DefaultValue);
-}
-
-void SetInputValue(
-	TSharedRef<FNiagaraSystemViewModel> SystemViewModel,
-	TSharedPtr<FNiagaraEmitterViewModel> EmitterViewModel,
-	UNiagaraStackEditorData& StackEditorData,
-	UNiagaraScript& SourceScript,
-	const TArray<TWeakObjectPtr<UNiagaraScript>> AffectedScripts,
-	UNiagaraNodeFunctionCall& ModuleNode,
-	UNiagaraNodeFunctionCall& InputFunctionCallNode,
-	FName InputName,
-	const FStackFunctionInputValue& Value)
-{
-	FNiagaraParameterHandle ModuleHandle = FNiagaraParameterHandle::CreateModuleParameterHandle(InputName);
-	FNiagaraParameterHandle AliasedFunctionHandle = FNiagaraParameterHandle::CreateAliasedModuleParameterHandle(ModuleHandle, &InputFunctionCallNode);
-	if (Value.LocalValue.IsSet())
-	{
-		bool bRapidIterationParameterSet = false;
-		if (FNiagaraStackGraphUtilities::IsRapidIterationType(Value.Type))
-		{
-			FCompileConstantResolver ConstantResolver = EmitterViewModel ? FCompileConstantResolver(EmitterViewModel->GetEmitter(), FNiagaraStackGraphUtilities::GetEmitterOutputNodeForStackNode(InputFunctionCallNode)->GetUsage(), InputFunctionCallNode.DebugState) : FCompileConstantResolver();
-			UEdGraphPin* DefaultPin = InputFunctionCallNode.FindParameterMapDefaultValuePin(ModuleHandle.GetParameterHandleString(), SourceScript.GetUsage(), ConstantResolver);
-			if (DefaultPin->LinkedTo.Num() == 0)
-			{
-				FNiagaraVariable RapidIterationParameter = FNiagaraStackGraphUtilities::CreateRapidIterationParameter(
-					EmitterViewModel->GetEmitter()->GetUniqueEmitterName(), SourceScript.GetUsage(), AliasedFunctionHandle.GetParameterHandleString(), Value.Type);
-
-				for (TWeakObjectPtr<UNiagaraScript> AffectedScript : AffectedScripts)
-				{
-					AffectedScript->Modify();
-					AffectedScript->RapidIterationParameters.SetParameterData(Value.LocalValue.GetValue().GetData(), RapidIterationParameter, true);
-				}
-				bRapidIterationParameterSet = true;
-			}
-		}
-
-		if (bRapidIterationParameterSet == false)
-		{
-			const UEdGraphSchema_Niagara* NiagaraSchema = GetDefault<UEdGraphSchema_Niagara>();
-			FString PinDefaultValue;
-			if (ensureMsgf(NiagaraSchema->TryGetPinDefaultValueFromNiagaraVariable(Value.LocalValue.GetValue(), PinDefaultValue),
-				TEXT("Could not generate default value string for non-rapid iteration parameter.")))
-			{
-				UEdGraphPin& OverridePin = FNiagaraStackGraphUtilities::GetOrCreateStackFunctionInputOverridePin(InputFunctionCallNode, AliasedFunctionHandle, Value.Type);
-				OverridePin.Modify();
-				OverridePin.DefaultValue = PinDefaultValue;
-				Cast<UNiagaraNode>(OverridePin.GetOwningNode())->MarkNodeRequiresSynchronization(TEXT("OverridePin Default Value Changed"), true);
-			}
-		}
-	}
-	else if (Value.LinkedValue.IsSet())
-	{
-		UEdGraphPin& OverridePin = FNiagaraStackGraphUtilities::GetOrCreateStackFunctionInputOverridePin(InputFunctionCallNode, AliasedFunctionHandle, Value.Type);
-		TSet<FNiagaraVariable> KnownParameters = FNiagaraStackGraphUtilities::GetParametersForContext(OverridePin.GetOwningNode()->GetGraph(), SystemViewModel->GetSystem());
-		FNiagaraStackGraphUtilities::SetLinkedValueHandleForFunctionInput(OverridePin, FNiagaraParameterHandle(Value.LinkedValue.GetValue()), KnownParameters);
-	}
-	else if (Value.DataValue.IsSet())
-	{
-		UEdGraphPin& OverridePin = FNiagaraStackGraphUtilities::GetOrCreateStackFunctionInputOverridePin(InputFunctionCallNode, AliasedFunctionHandle, Value.Type);
-		FString DataObjectName = Value.DataValue.GetValue() != nullptr ? Value.DataValue.GetValue()->GetName() : Value.Type.GetName();
-		UNiagaraDataInterface* NewDataObject;
-		FNiagaraStackGraphUtilities::SetDataValueObjectForFunctionInput(OverridePin, Value.Type.GetClass(), DataObjectName, NewDataObject);
-		if (Value.DataValue.GetValue() != nullptr)
-		{
-			Value.DataValue.GetValue()->CopyTo(NewDataObject);
-		}
-	}
-	else if (Value.DynamicValue.IsSet())
-	{
-		UEdGraphPin& OverridePin = FNiagaraStackGraphUtilities::GetOrCreateStackFunctionInputOverridePin(InputFunctionCallNode, AliasedFunctionHandle, Value.Type);
-		UNiagaraNodeFunctionCall* NewDynamicInputFunctionCall;
-		FNiagaraStackGraphUtilities::SetDynamicInputForFunctionInput(OverridePin, Value.DynamicValue.GetValue()->FunctionScript, NewDynamicInputFunctionCall);
-		FNiagaraStackGraphUtilities::InitializeStackFunctionInputs(SystemViewModel, EmitterViewModel, StackEditorData, ModuleNode, *NewDynamicInputFunctionCall);
-		for (auto It = Value.DynamicValueInputs.CreateConstIterator(); It; ++It)
-		{
-			FName DynamicValueInputName = It.Key();
-			TSharedRef<FStackFunctionInputValue> DynamicValueInputValue = It.Value();
-			if (DynamicValueInputValue->bIsOverride)
-			{
-				SetInputValue(SystemViewModel, EmitterViewModel, StackEditorData, SourceScript, AffectedScripts, ModuleNode, *NewDynamicInputFunctionCall, DynamicValueInputName, *DynamicValueInputValue);
-			}
-		}
-	}
-}
-
-void FNiagaraStackGraphUtilities::ResetToDefaultDynamicInput(
-	TSharedRef<FNiagaraSystemViewModel> SystemViewModel,
-	TSharedPtr<FNiagaraEmitterViewModel> EmitterViewModel,
-	UNiagaraStackEditorData& StackEditorData,
-	UNiagaraScript& SourceScript,
-	const TArray<TWeakObjectPtr<UNiagaraScript>> AffectedScripts,
-	UNiagaraNodeFunctionCall& ModuleNode,
-	UNiagaraNodeFunctionCall& InputFunctionCallNode,
-	FName InputName,
-	UEdGraphPin& DefaultPin)
-{
-	FStackFunctionInputValue DefaultValue;
-	if (ensureMsgf(TryGetStackFunctionInputValue(SourceScript, nullptr, DefaultPin, NAME_None, FRapidIterationParameterContext(), DefaultValue), TEXT("Default dynamic input was not valid")))
-	{
-		SetInputValue(SystemViewModel, EmitterViewModel, StackEditorData, SourceScript, AffectedScripts, ModuleNode, InputFunctionCallNode, InputName, DefaultValue);
-	}
-}
-
 bool FNiagaraStackGraphUtilities::GetStackIssuesRecursively(const UNiagaraStackEntry* const Entry, TArray<UNiagaraStackErrorItem*>& OutIssues)
 {
 	TArray<UNiagaraStackEntry*> Entries;
@@ -3364,6 +3112,300 @@ void FNiagaraStackGraphUtilities::SynchronizeVariableToLibraryAndApplyToGraph(UN
 #endif
 }
 
+TSharedRef<TMap<FName, FGuid>> GetVariableNameToVariableIdMap(UNiagaraNodeFunctionCall& InFunctionCallNode)
+{
+	TSharedRef<TMap<FName, FGuid>> VariableNameToVariableIdMap = MakeShared<TMap<FName, FGuid>>();
+	UNiagaraGraph* CalledGraph = InFunctionCallNode.GetCalledGraph();
+	if (CalledGraph != nullptr)
+	{
+		for (const TPair<FNiagaraVariable, TObjectPtr<UNiagaraScriptVariable>>& VariableMetadataPair : CalledGraph->GetAllMetaData())
+		{
+			FGuid VariableGuid = VariableMetadataPair.Value->Metadata.GetVariableGuid();
+			if (VariableGuid.IsValid())
+			{
+				VariableNameToVariableIdMap->FindOrAdd(VariableMetadataPair.Key.GetName()) = VariableGuid;
+			}
+		}
+	}
+	return VariableNameToVariableIdMap;
+}
+
+TSharedRef<TMap<FGuid, FNiagaraVariable>> GetVariableIdToVariableMap(UNiagaraNodeFunctionCall& InFunctionCallNode)
+{
+	TSharedRef<TMap<FGuid, FNiagaraVariable>> VariableIdToVariableMap = MakeShared<TMap<FGuid, FNiagaraVariable>>();
+	UNiagaraGraph* CalledGraph = InFunctionCallNode.GetCalledGraph();
+	if (CalledGraph != nullptr)
+	{
+		for (const TPair<FNiagaraVariable, TObjectPtr<UNiagaraScriptVariable>>& VariableMetadataPair : CalledGraph->GetAllMetaData())
+		{
+			FGuid VariableGuid = VariableMetadataPair.Value->Metadata.GetVariableGuid();
+			if (VariableGuid.IsValid())
+			{
+				VariableIdToVariableMap->FindOrAdd(VariableGuid) = VariableMetadataPair.Key;
+			}
+		}
+	}
+	return VariableIdToVariableMap;
+}
+
+// Searches the graph which contains InFunctionCallNode and finds parameter map get pins which reference outputs from the function call node.
+TArray<UEdGraphPin*> GetStackLinkedOutputPinsForFunction(UNiagaraNodeFunctionCall& InFunctionCallNode)
+{
+	TArray<UEdGraphPin*> LinkedOutputPins;
+	TArray<FNiagaraStackGraphUtilities::FStackNodeGroup> StackNodeGroups;
+	FName FunctionCallName = *InFunctionCallNode.GetFunctionName();
+	GetStackNodeGroups(InFunctionCallNode, StackNodeGroups);
+	for (const FNiagaraStackGraphUtilities::FStackNodeGroup& StackNodeGroup : StackNodeGroups)
+	{
+		TArray<UNiagaraNode*> NodesInGroup;
+		StackNodeGroup.GetAllNodesInGroup(NodesInGroup);
+		for (UNiagaraNode* NodeInGroup : NodesInGroup)
+		{
+			UNiagaraNodeParameterMapGet* MapGet = Cast<UNiagaraNodeParameterMapGet>(NodeInGroup);
+			if (MapGet != nullptr)
+			{
+				TArray<UEdGraphPin*> MapGetOutputPins;
+				MapGet->GetOutputPins(MapGetOutputPins);
+				for (UEdGraphPin* MapGetOutputPin : MapGetOutputPins)
+				{
+					FNiagaraParameterHandle MapGetHandle(MapGetOutputPin->PinName);
+					if (MapGetHandle.IsOutputHandle())
+					{
+						TArray<FName> HandleParts = MapGetHandle.GetHandleParts();
+						if (HandleParts.Num() >= 3 && HandleParts[1] == FunctionCallName)
+						{
+							LinkedOutputPins.Add(MapGetOutputPin);
+						}
+					}
+				}
+			}
+		}
+	}
+	return LinkedOutputPins;
+}
+
+// Takes a parameter map output pin which represents a module output and finds the target function call which is linking that output.
+UNiagaraNodeFunctionCall* GetTargetFunctionCallForStackLinkedOutputPin(UEdGraphPin& StackLinkedOutputPin)
+{
+	for (UEdGraphPin* LinkedOverridePin : StackLinkedOutputPin.LinkedTo)
+	{
+		UNiagaraNodeParameterMapSet* LinkedOverrideNode = Cast<UNiagaraNodeParameterMapSet>(LinkedOverridePin->GetOwningNode());
+		if (LinkedOverrideNode != nullptr)
+		{
+			FNiagaraParameterHandle LinkedOverrideHandle(LinkedOverridePin->PinName);
+			UEdGraphPin* ParameterMapOutputPin = FNiagaraStackGraphUtilities::GetParameterMapOutputPin(*LinkedOverrideNode);
+			if (ParameterMapOutputPin != nullptr)
+			{
+				for (UEdGraphPin* LinkedInputPin : ParameterMapOutputPin->LinkedTo)
+				{
+					UNiagaraNodeFunctionCall* LinkedFunctionCallNode = Cast<UNiagaraNodeFunctionCall>(LinkedInputPin->GetOwningNode());
+					if (LinkedFunctionCallNode != nullptr && *LinkedFunctionCallNode->GetFunctionName() == LinkedOverrideHandle.GetHandleParts()[0])
+					{
+						return LinkedFunctionCallNode;
+					}
+				}
+			}
+		}
+	}
+	return nullptr;
+}
+
+void FNiagaraStackGraphUtilities::PopulateFunctionCallNameBindings(UNiagaraNodeFunctionCall& InFunctionCallNode)
+{
+	UNiagaraNodeOutput* OutputNode = GetEmitterOutputNodeForStackNode(InFunctionCallNode);
+	if (OutputNode == nullptr ||
+		OutputNode->GetUsage() == ENiagaraScriptUsage::Module ||
+		OutputNode->GetUsage() == ENiagaraScriptUsage::DynamicInput ||
+		OutputNode->GetUsage() == ENiagaraScriptUsage::Function)
+	{
+		// Early out if an output node can't be found, or if the output node's usage is not a stack usage.
+		return;
+	}
+
+	TSharedPtr<TMap<FName, FGuid>> VariableNameToVariableIdMap;
+	const UEdGraphSchema_Niagara* NiagaraSchema = GetDefault<UEdGraphSchema_Niagara>();
+
+	// Populate input ids from the inputs on this function which are overriden.
+	TArray<UEdGraphPin*> OverridePins;
+	UNiagaraNodeParameterMapSet* OverrideNode = GetStackFunctionOverrideNode(InFunctionCallNode);
+	if (OverrideNode != nullptr)
+	{
+		OverridePins = GetOverridePinsForFunction(*OverrideNode, InFunctionCallNode);
+	}
+	if (OverridePins.Num() > 0)
+	{
+		VariableNameToVariableIdMap = GetVariableNameToVariableIdMap(InFunctionCallNode);
+		for (UEdGraphPin* OverridePin : OverridePins)
+		{
+			FNiagaraVariable OverrideVariable = NiagaraSchema->PinToNiagaraVariable(OverridePin);
+			FNiagaraParameterHandle OverrideHandle(OverrideVariable.GetName());
+			FNiagaraVariable InputVariable = FNiagaraUtilities::ResolveAliases(OverrideVariable, FNiagaraAliasContext().ChangeModuleNameToModule(InFunctionCallNode.GetFunctionName()));
+			FGuid* IdByInputVariableName = VariableNameToVariableIdMap->Find(InputVariable.GetName());
+			if (IdByInputVariableName != nullptr)
+			{
+				InFunctionCallNode.UpdateInputNameBinding(*IdByInputVariableName, OverrideHandle.GetName());
+			}
+		}
+	}
+
+	// Populate ids on target functions which use outputs from this function.
+	TArray<UEdGraphPin*> StackLinkedOutputPins = GetStackLinkedOutputPinsForFunction(InFunctionCallNode);
+	if (StackLinkedOutputPins.Num() > 0)
+	{
+		if (VariableNameToVariableIdMap.IsValid() == false)
+		{
+			VariableNameToVariableIdMap = GetVariableNameToVariableIdMap(InFunctionCallNode);
+		}
+		for (UEdGraphPin* StackLinkedOutputPin : StackLinkedOutputPins)
+		{
+			UNiagaraNodeFunctionCall* TargetFunctionCallNode = GetTargetFunctionCallForStackLinkedOutputPin(*StackLinkedOutputPin);
+			if (TargetFunctionCallNode != nullptr)
+			{
+				FNiagaraVariable LinkedOutputVariable = NiagaraSchema->PinToNiagaraVariable(StackLinkedOutputPin);
+				FNiagaraParameterHandle LinkedOutputHandle(LinkedOutputVariable.GetName());
+				FNiagaraVariable OutputVariable = FNiagaraUtilities::ResolveAliases(LinkedOutputVariable, FNiagaraAliasContext().ChangeModuleNameToModule(InFunctionCallNode.GetFunctionName()));
+				FGuid* IdByVariableName = VariableNameToVariableIdMap->Find(OutputVariable.GetName());
+				if (IdByVariableName != nullptr)
+				{
+					InFunctionCallNode.UpdateInputNameBinding(*IdByVariableName, LinkedOutputHandle.GetName());
+				}
+			}
+		}
+	}
+}
+
+void FNiagaraStackGraphUtilities::SynchronizeReferencingMapPinsWithFunctionCall(UNiagaraNodeFunctionCall& InFunctionCallNode)
+{
+	UNiagaraNodeOutput* OutputNode = GetEmitterOutputNodeForStackNode(InFunctionCallNode);
+	if (OutputNode == nullptr ||
+		OutputNode->GetUsage() == ENiagaraScriptUsage::Module ||
+		OutputNode->GetUsage() == ENiagaraScriptUsage::DynamicInput ||
+		OutputNode->GetUsage() == ENiagaraScriptUsage::Function)
+	{
+		// Early out if an output node can't be found, or if the output node's usage is not a stack usage.
+		return;
+	}
+
+	TSharedPtr<TMap<FGuid, FNiagaraVariable>> VariableIdToVariableMap;
+	const UEdGraphSchema_Niagara* NiagaraSchema = GetDefault<UEdGraphSchema_Niagara>();
+	
+	// Update the override parameter map set nodes which are used to set the function inputs.
+	TArray<UEdGraphPin*> OverridePins;
+	UNiagaraNodeParameterMapSet* OverrideNode = GetStackFunctionOverrideNode(InFunctionCallNode);
+	if (OverrideNode != nullptr)
+	{	
+		OverridePins = GetOverridePinsForFunction(*OverrideNode, InFunctionCallNode);
+	}
+	if (OverridePins.Num() > 0)
+	{
+		VariableIdToVariableMap = GetVariableIdToVariableMap(InFunctionCallNode);
+		for (UEdGraphPin* OverridePin : OverridePins)
+		{
+			FNiagaraVariable OverrideVariable = NiagaraSchema->PinToNiagaraVariable(OverridePin);
+			FNiagaraParameterHandle OverrideHandle(OverrideVariable.GetName());
+			TArray<FGuid> BoundGuids = InFunctionCallNode.GetBoundPinGuidsByName(OverrideHandle.GetName());
+			for (FGuid BoundGuid : BoundGuids)
+			{
+				FNiagaraVariable* BoundVariable = VariableIdToVariableMap->Find(BoundGuid);
+				if (BoundVariable != nullptr)
+				{
+					FNiagaraParameterHandle BoundVariableHandle(BoundVariable->GetName());
+					if (OverrideHandle.GetName() != BoundVariableHandle.GetName())
+					{
+						FNiagaraParameterHandle UpdatedOverrideHandle = FNiagaraParameterHandle::CreateAliasedModuleParameterHandle(BoundVariable->GetName(), &InFunctionCallNode);
+						OverridePin->PinName = UpdatedOverrideHandle.GetParameterHandleString();
+						InFunctionCallNode.UpdateInputNameBinding(BoundGuid, BoundVariableHandle.GetName());
+					}
+					if (OverrideVariable.GetType() != BoundVariable->GetType())
+					{
+						FEdGraphPinType NewPinType = NiagaraSchema->TypeDefinitionToPinType(BoundVariable->GetType());
+						OverridePin->PinType = NewPinType;
+						if (OverridePin->LinkedTo.Num() == 1 &&
+							OverridePin->LinkedTo[0]->GetOwningNode() != nullptr &&
+							OverridePin->LinkedTo[0]->GetOwningNode()->IsA<UNiagaraNodeCustomHlsl>())
+						{
+							// Custom hlsl nodes in the stack are expression dynamic inputs and their output pin type must match their connected input
+							// so update that here.
+							OverridePin->LinkedTo[0]->PinType = NewPinType;
+						}
+					}
+					break;
+				}
+			}
+		}
+	}
+
+	// Update linked outputs referenced by downstream parameter map gets.
+	TArray<UEdGraphPin*> StackLinkedOutputPins = GetStackLinkedOutputPinsForFunction(InFunctionCallNode);
+	if (StackLinkedOutputPins.Num() > 0)
+	{
+		if (VariableIdToVariableMap.IsValid() == false)
+		{
+			VariableIdToVariableMap = GetVariableIdToVariableMap(InFunctionCallNode);
+		}
+		for (UEdGraphPin* StackLinkedOutputPin : StackLinkedOutputPins)
+		{
+			FNiagaraVariable LinkedOutputVariable = NiagaraSchema->PinToNiagaraVariable(StackLinkedOutputPin);
+			FNiagaraParameterHandle LinkedOutputHandle(LinkedOutputVariable.GetName());
+			UNiagaraNodeFunctionCall* TargetFunctionCallNode = GetTargetFunctionCallForStackLinkedOutputPin(*StackLinkedOutputPin);
+			if (TargetFunctionCallNode != nullptr)
+			{
+				TArray<FGuid> BoundGuids = TargetFunctionCallNode->GetBoundPinGuidsByName(LinkedOutputHandle.GetName());
+				for (const FGuid& BoundGuid : BoundGuids)
+				{
+					FNiagaraVariable* BoundVariable = VariableIdToVariableMap->Find(BoundGuid);
+					if (BoundVariable != nullptr)
+					{
+						FNiagaraParameterHandle BoundVariableHandle(BoundVariable->GetName());
+						if (LinkedOutputHandle.GetName() != BoundVariableHandle.GetName())
+						{
+							FNiagaraParameterHandle UpdatedLinkedOutputHandle = FNiagaraParameterHandle::CreateAliasedModuleParameterHandle(BoundVariableHandle, &InFunctionCallNode);
+							StackLinkedOutputPin->PinName = UpdatedLinkedOutputHandle.GetParameterHandleString();
+							TargetFunctionCallNode->UpdateInputNameBinding(BoundGuid, BoundVariableHandle.GetName());
+						}
+						if (LinkedOutputVariable.GetType() != BoundVariable->GetType())
+						{
+							StackLinkedOutputPin->PinType = NiagaraSchema->TypeDefinitionToPinType(BoundVariable->GetType());
+						}
+						break;
+					}
+				}
+			}
+		}
+	}
+}
+
+FGuid FNiagaraStackGraphUtilities::GetScriptVariableIdForLinkedOutputHandle(const FNiagaraParameterHandle& LinkedHandle, FNiagaraTypeDefinition LinkedType, UNiagaraGraph& TargetGraph)
+{
+	if (LinkedHandle.IsOutputHandle() && LinkedHandle.GetHandleParts().Num() >= 3)
+	{
+		// First find the source function call node.
+		FString LinkedFunctionName = LinkedHandle.GetHandleParts()[1].ToString();
+		TArray<UNiagaraNodeFunctionCall*> FunctionCallNodes;
+		TargetGraph.GetNodesOfClass(FunctionCallNodes);
+		UNiagaraNodeFunctionCall** SourceFunctionCallNodePtr = FunctionCallNodes.FindByPredicate([LinkedFunctionName](const UNiagaraNodeFunctionCall* FunctionCallNode)
+			{ return FunctionCallNode->GetFunctionName() == LinkedFunctionName; });
+		if (SourceFunctionCallNodePtr != nullptr)
+		{
+			UNiagaraGraph* CalledGraph = (*SourceFunctionCallNodePtr)->GetCalledGraph();
+			if (CalledGraph != nullptr)
+			{
+				FNiagaraVariable LinkedVariable(LinkedType, LinkedHandle.GetParameterHandleString());
+				FNiagaraVariable LinkedModuleVariable = FNiagaraUtilities::ResolveAliases(LinkedVariable, FNiagaraAliasContext()
+					.ChangeModuleNameToModule(LinkedFunctionName));
+				for (const TPair<FNiagaraVariable, TObjectPtr<UNiagaraScriptVariable>>& VariableMetadataPair : CalledGraph->GetAllMetaData())
+				{
+					if (VariableMetadataPair.Key == LinkedModuleVariable)
+					{
+						return VariableMetadataPair.Value->Metadata.GetVariableGuid();
+					}
+				}
+			}
+		}
+	}
+	return FGuid();
+}
+
 bool FNiagaraStackGraphUtilities::DependencyUtilities::DoesStackModuleProvideDependency(const FNiagaraStackModuleData& StackModuleData, const FNiagaraModuleDependency& SourceModuleRequiredDependency, const UNiagaraNodeOutput& SourceOutputNode)
 {
 	if (StackModuleData.ModuleNode != nullptr && StackModuleData.ModuleNode->FunctionScript != nullptr)
diff --git a/Engine/Plugins/FX/Niagara/Source/NiagaraEditor/Public/NiagaraNodeFunctionCall.h b/Engine/Plugins/FX/Niagara/Source/NiagaraEditor/Public/NiagaraNodeFunctionCall.h
--- a/Engine/Plugins/FX/Niagara/Source/NiagaraEditor/Public/NiagaraNodeFunctionCall.h
+++ b/Engine/Plugins/FX/Niagara/Source/NiagaraEditor/Public/NiagaraNodeFunctionCall.h
@@ -126,9 +126,6 @@ public:
 	virtual bool IsDeprecated() const override;
 	//~ End EdGraphNode Interface
 
-	/** Checks if the existing pin names match the function script parameter names and try to fix them via guid matching if there is a difference */
-	NIAGARAEDITOR_API bool FixupPinNames();
-
 	/** When overriding an input value, this updates which variable guid was bound to which input name, so it can be reassigned when the input is renamed.*/
 	void UpdateInputNameBinding(const FGuid& BoundVariableGuid, const FName& BoundName);
 
@@ -180,6 +177,7 @@ public:
 	NIAGARAEDITOR_API void RemoveCustomNote(const FGuid& MessageKey);
 	NIAGARAEDITOR_API FSimpleDelegate& OnCustomNotesChanged() { return OnCustomNotesChangedDelegate; }
 	void RemoveCustomNoteViaDelegate(const FGuid MessageKey);
+	TArray<FGuid> GetBoundPinGuidsByName(FName InputName) const;
 
 	/** Adds a static switch pin to this function call node by variable id and sets it's default value using the supplied data and marks it as
 		orphaned. This allows a previously available static switch value to be retained on the node even if the the switch is no longer exposed. */
@@ -204,8 +202,6 @@ protected:
 
 	void UpdateNodeErrorMessage();
 
-	TArray<FGuid> GetBoundPinGuidsByName(FName InputName) const;
-
 	void FixupFunctionScriptVersion();
 
 	void UpdatePinTooltips();
diff --git a/Engine/Plugins/FX/Niagara/Source/NiagaraEditor/Public/NiagaraStackFunctionInputBinder.h b/Engine/Plugins/FX/Niagara/Source/NiagaraEditor/Public/NiagaraStackFunctionInputBinder.h
--- a/Engine/Plugins/FX/Niagara/Source/NiagaraEditor/Public/NiagaraStackFunctionInputBinder.h
+++ b/Engine/Plugins/FX/Niagara/Source/NiagaraEditor/Public/NiagaraStackFunctionInputBinder.h
@@ -102,6 +102,8 @@ private:
 
 	FName InputName;
 
+	FGuid InputVariableGuid;
+
 	FNiagaraParameterHandle AliasedParameterHandle;
 
 	mutable FGuid ValidScriptGraphChangeIdForOverridePin;
diff --git a/Engine/Plugins/FX/Niagara/Source/NiagaraEditor/Public/ViewModels/Stack/NiagaraStackGraphUtilities.h b/Engine/Plugins/FX/Niagara/Source/NiagaraEditor/Public/ViewModels/Stack/NiagaraStackGraphUtilities.h
--- a/Engine/Plugins/FX/Niagara/Source/NiagaraEditor/Public/ViewModels/Stack/NiagaraStackGraphUtilities.h
+++ b/Engine/Plugins/FX/Niagara/Source/NiagaraEditor/Public/ViewModels/Stack/NiagaraStackGraphUtilities.h
@@ -118,7 +118,7 @@ namespace FNiagaraStackGraphUtilities
 
 	UEdGraphPin* GetStackFunctionInputOverridePin(UNiagaraNodeFunctionCall& StackFunctionCall, FNiagaraParameterHandle AliasedInputParameterHandle);
 
-	UEdGraphPin& GetOrCreateStackFunctionInputOverridePin(UNiagaraNodeFunctionCall& StackFunctionCall, FNiagaraParameterHandle AliasedInputParameterHandle, FNiagaraTypeDefinition InputType, const FGuid& PreferredOverrideNodeGuid = FGuid());
+	UEdGraphPin& GetOrCreateStackFunctionInputOverridePin(UNiagaraNodeFunctionCall& StackFunctionCall, FNiagaraParameterHandle AliasedInputParameterHandle, FNiagaraTypeDefinition InputType, const FGuid& InputScriptVariableId, const FGuid& PreferredOverrideNodeGuid = FGuid());
 
 	bool IsOverridePinForFunction(UEdGraphPin& OverridePin, UNiagaraNodeFunctionCall& FunctionCallNode);
 
@@ -221,25 +221,6 @@ namespace FNiagaraStackGraphUtilities
 
 	bool CanWriteParameterFromUsage(FNiagaraVariable Parameter, ENiagaraScriptUsage Usage, const TOptional<FName>& StackContextOverride = TOptional<FName>(), const TArray<FName>& StackContextAllOverrides = TArray<FName>());
 	bool CanWriteParameterFromUsageViaOutput(FNiagaraVariable Parameter, const UNiagaraNodeOutput* OutputNode);
-
-	bool DoesDynamicInputMatchDefault(
-		FString EmitterUniqueName,
-		UNiagaraScript& OwningScript,
-		UNiagaraNodeFunctionCall& OwningFunctionCallNode,
-		UEdGraphPin& OverridePin,
-		FName InputName,
-		UEdGraphPin& DefaultPin);
-
-	void ResetToDefaultDynamicInput(
-		TSharedRef<FNiagaraSystemViewModel> SystemViewModel,
-		TSharedPtr<FNiagaraEmitterViewModel> EmitterViewModel,
-		UNiagaraStackEditorData& StackEditorData,
-		UNiagaraScript& SourceScript,
-		const TArray<TWeakObjectPtr<UNiagaraScript>> AffectedScripts,
-		UNiagaraNodeFunctionCall& ModuleNode,
-		UNiagaraNodeFunctionCall& InputFunctionCallNode,
-		FName InputName,
-		UEdGraphPin& DefaultPin);
 	
 	bool GetStackIssuesRecursively(const UNiagaraStackEntry* const Entry, TArray<UNiagaraStackErrorItem*>& OutIssues);
 
@@ -286,6 +267,12 @@ namespace FNiagaraStackGraphUtilities
 
 	void SynchronizeVariableToLibraryAndApplyToGraph(UNiagaraScriptVariable* ScriptVarToSync);
 
+	void PopulateFunctionCallNameBindings(UNiagaraNodeFunctionCall& InFunctionCallNode);
+
+	void SynchronizeReferencingMapPinsWithFunctionCall(UNiagaraNodeFunctionCall& InFunctionCallNode);
+
+	FGuid GetScriptVariableIdForLinkedOutputHandle(const FNiagaraParameterHandle& LinkedOutputHandle, FNiagaraTypeDefinition LinkedType, UNiagaraGraph& TargetGraph);
+
 	namespace DependencyUtilities
 	{
 		bool DoesStackModuleProvideDependency(const FNiagaraStackModuleData& StackModuleData, const FNiagaraModuleDependency& SourceModuleRequiredDependency, const UNiagaraNodeOutput& SourceOutputNode);
