Niagara - Add a console command for loading all systems in a specified folder, and child folders.

diff --git a/Engine/Plugins/FX/Niagara/Source/NiagaraEditor/Private/NiagaraEditorModule.cpp b/Engine/Plugins/FX/Niagara/Source/NiagaraEditor/Private/NiagaraEditorModule.cpp
--- a/Engine/Plugins/FX/Niagara/Source/NiagaraEditor/Private/NiagaraEditorModule.cpp
+++ b/Engine/Plugins/FX/Niagara/Source/NiagaraEditor/Private/NiagaraEditorModule.cpp
@@ -151,6 +151,7 @@
 #include "NiagaraPerfBaseline.h"
 
 #include "NiagaraGraphDataCache.h"
+#include "Misc/ScopedSlowTask.h"
 
 IMPLEMENT_MODULE( FNiagaraEditorModule, NiagaraEditor );
 
@@ -637,6 +638,39 @@ void DumpCompileIdDataForAsset(const TArray<FString>& Arguments)
 	}
 }
 
+void LoadAllSystemsInFolder(const TArray<FString>& Arguments)
+{
+	if (Arguments.Num() == 1)
+	{
+		TArray<FAssetData> SystemAssetsInFolder;
+		const FAssetRegistryModule& AssetRegistryModule = FModuleManager::LoadModuleChecked<FAssetRegistryModule>(TEXT("AssetRegistry"));
+		FARFilter Filter;
+		Filter.ClassNames.Add(UNiagaraSystem::StaticClass()->GetFName());
+		Filter.PackagePaths.Add(*Arguments[0]);
+		Filter.bRecursivePaths = true;
+		AssetRegistryModule.Get().GetAssets(Filter, SystemAssetsInFolder);
+		if (SystemAssetsInFolder.Num() > 0)
+		{
+			const FText SlowTaskText = FText::Format(NSLOCTEXT("NiagaraEditor", "LoadAllSystemsInFolderFormat", "Loading {0} systems in folder {1}"), FText::AsNumber(SystemAssetsInFolder.Num()), FText::FromString(Arguments[0]));
+			FScopedSlowTask SlowTask(SystemAssetsInFolder.Num(), SlowTaskText);
+			SlowTask.MakeDialog(true);
+			int32 ItemNumber = 0;
+			for (FAssetData& SystemAsset : SystemAssetsInFolder)
+			{
+				if (SlowTask.ShouldCancel())
+				{
+					return;
+				}
+				ItemNumber++;
+				FText SlowTaskUpdateText = FText::Format(NSLOCTEXT("NiagaraEditor", "LoadAllSystemsInFolderProgressFormat", "Loading system {0} of {1}\n{2}"),
+					FText::AsNumber(ItemNumber), FText::AsNumber(SystemAssetsInFolder.Num()),  FText::FromString(SystemAsset.GetFullName()));
+				SlowTask.EnterProgressFrame(1, SlowTaskUpdateText);
+				SystemAsset.GetAsset();
+			}
+		}
+	}
+}
+
 void ExecuteInvalidateNiagaraCachedScripts(const TArray< FString >& Args)
 {
 	if (Args.Num() == 0)
@@ -1155,6 +1189,11 @@ void FNiagaraEditorModule::StartupModule()
 		TEXT("Dumps data relevant to generating the compile id for an asset."),
 		FConsoleCommandWithArgsDelegate::CreateStatic(&DumpCompileIdDataForAsset));
 
+	LoadAllSystemsInFolderCommand = IConsoleManager::Get().RegisterConsoleCommand(
+		TEXT("fx.LoadAllNiagaraSystemsInFolder"),
+		TEXT("Loads all niagara systems in the supplied directory and sub-directories."),
+		FConsoleCommandWithArgsDelegate::CreateStatic(&LoadAllSystemsInFolder));
+
 	FNiagaraMessageManager* MessageManager = FNiagaraMessageManager::Get();
 	MessageManager->RegisterMessageTopic(FNiagaraMessageTopics::CompilerTopicName);
 	MessageManager->RegisterMessageTopic(FNiagaraMessageTopics::ObjectTopicName);
diff --git a/Engine/Plugins/FX/Niagara/Source/NiagaraEditor/Public/NiagaraEditorModule.h b/Engine/Plugins/FX/Niagara/Source/NiagaraEditor/Public/NiagaraEditorModule.h
--- a/Engine/Plugins/FX/Niagara/Source/NiagaraEditor/Public/NiagaraEditorModule.h
+++ b/Engine/Plugins/FX/Niagara/Source/NiagaraEditor/Public/NiagaraEditorModule.h
@@ -344,6 +344,7 @@ private:
 	IConsoleCommand* PreventAllSystemRecompilesCommand;
 	IConsoleCommand* UpgradeAllNiagaraAssetsCommand;
 	IConsoleCommand* DumpCompileIdDataForAssetCommand;
+	IConsoleCommand* LoadAllSystemsInFolderCommand;
 
 	FOnCheckScriptToolkitsShouldFocusGraphElement OnCheckScriptToolkitsShouldFocusGraphElement;
 
