Unshelved from pending changelist '19757341':

Niagara - Fix an issue in the merge manager where the graph's change id wouldn't be updated when setting static switch values on function call nodes.  This was causing strange behavior when resetting static switch inputs to their base value because the merge manager would get a stale merge adapter from the cache.

diff --git a/Engine/Plugins/FX/Niagara/Source/NiagaraEditor/Private/NiagaraScriptMergeManager.cpp b/Engine/Plugins/FX/Niagara/Source/NiagaraEditor/Private/NiagaraScriptMergeManager.cpp
--- a/Engine/Plugins/FX/Niagara/Source/NiagaraEditor/Private/NiagaraScriptMergeManager.cpp
+++ b/Engine/Plugins/FX/Niagara/Source/NiagaraEditor/Private/NiagaraScriptMergeManager.cpp
@@ -2912,6 +2912,8 @@ FNiagaraScriptMergeManager::FApplyDiffResults FNiagaraScriptMergeManager::AddInp
 				if (SwitchType == OverrideToAdd->GetType())
 				{
 					MatchingStaticSwitchPin->DefaultValue = OverrideToAdd->GetStaticSwitchValue().GetValue();
+					TargetFunctionCall.MarkNodeRequiresSynchronization(TEXT("Static Switch Value Changed"), true);
+					Results.bModifiedGraph = true;
 					Results.bSucceeded = true;
 				}
 				else
