Refactored type change action to be more robust & clear, fixing various smaller bugs:
- Map get nodes now unregister their pin guid mappings when a pin is removed
- New orphaned pins on a map get get their guid mapping added to to the node, which fixes up visuals
- Added an exception to keep pin connections when changing a parameter type from vector3 -> position, if strict stack stype aren't enforced.
- Added an exception to keep pin connections when the connected pin is a numeric pin & we are changing to a numeric type
- Fixed a bug where the auto generated pin value wasn't updated to the proper default, making orphaned pin handling more awkward
- Niagara parameter pin labels have their name in red if the pin is orphaned

diff --git a/Engine/Plugins/FX/Niagara/Source/NiagaraEditor/Private/NiagaraEditorStyle.cpp b/Engine/Plugins/FX/Niagara/Source/NiagaraEditor/Private/NiagaraEditorStyle.cpp
--- a/Engine/Plugins/FX/Niagara/Source/NiagaraEditor/Private/NiagaraEditorStyle.cpp
+++ b/Engine/Plugins/FX/Niagara/Source/NiagaraEditor/Private/NiagaraEditorStyle.cpp
@@ -265,6 +265,26 @@ void FNiagaraEditorStyle::InitParameters()
 		.SetFont(DEFAULT_FONT("Regular", 8))
 		.SetColorAndOpacity(FLinearColor(1.0f, 1.0f, 1.0f, 0.5f)));
 
+	FTextBlockStyle InlineEditableTextBlockReadOnly = FTextBlockStyle(NormalText)
+		.SetFont(DEFAULT_FONT("Regular", 9))
+		.SetColorAndOpacity(FSlateColor::UseForeground())
+		.SetShadowOffset(FVector2D::ZeroVector)
+		.SetShadowColorAndOpacity(FLinearColor(0.8f, 0.8f, 0.8f, 0.5));
+
+	FEditableTextBoxStyle InlineEditableTextBlockEditable = FEditableTextBoxStyle()
+		.SetFont(DEFAULT_FONT("Regular", 9))
+		.SetForegroundColor(FStyleColors::Foreground)
+		.SetBackgroundImageNormal(FSlateRoundedBoxBrush(FStyleColors::Input, CoreStyleConstants::InputFocusRadius, FStyleColors::InputOutline, CoreStyleConstants::InputFocusThickness))
+		.SetBackgroundImageHovered(FSlateRoundedBoxBrush(FStyleColors::Input, CoreStyleConstants::InputFocusRadius, FStyleColors::Hover, CoreStyleConstants::InputFocusThickness))
+		.SetBackgroundImageFocused(FSlateRoundedBoxBrush(FStyleColors::Input, CoreStyleConstants::InputFocusRadius, FStyleColors::Primary, CoreStyleConstants::InputFocusThickness))
+		.SetBackgroundImageReadOnly(FSlateRoundedBoxBrush(FStyleColors::Header, CoreStyleConstants::InputFocusRadius, FStyleColors::InputOutline, CoreStyleConstants::InputFocusThickness))
+		.SetScrollBarStyle(FAppStyle::Get().GetWidgetStyle<FScrollBarStyle>("ScrollBar"));
+	
+	FInlineEditableTextBlockStyle InlineEditableTextBlockStyle = FInlineEditableTextBlockStyle()
+		.SetTextStyle(InlineEditableTextBlockReadOnly)
+		.SetEditableTextBoxStyle(InlineEditableTextBlockEditable);
+	Set("NiagaraEditor.Graph.Node.InlineEditablePinName", InlineEditableTextBlockStyle);
+
 
 	Set("NiagaraEditor.StaticIcon", new IMAGE_BRUSH("Icons/staticpill_16x", Icon16x16));
 	Set("NiagaraEditor.Pins.StaticConnected", new IMAGE_BRUSH("Icons/StaticPin_Connected", Icon16x16));
diff --git a/Engine/Plugins/FX/Niagara/Source/NiagaraEditor/Private/NiagaraGraph.cpp b/Engine/Plugins/FX/Niagara/Source/NiagaraEditor/Private/NiagaraGraph.cpp
--- a/Engine/Plugins/FX/Niagara/Source/NiagaraEditor/Private/NiagaraGraph.cpp
+++ b/Engine/Plugins/FX/Niagara/Source/NiagaraEditor/Private/NiagaraGraph.cpp
@@ -27,6 +27,7 @@
 #include "NiagaraScript.h"
 #include "NiagaraScriptSource.h"
 #include "NiagaraScriptVariable.h"
+#include "NiagaraSettings.h"
 #include "HAL/FileManager.h"
 #include "Misc/Paths.h"
 #include "Misc/SecureHash.h"
@@ -473,147 +474,208 @@ void UNiagaraGraph::ChangeParameterType(const TArray<FNiagaraVariable>& Paramete
 	TArray<UNiagaraNodeParameterMapBase*> NiagaraParameterNodes;
 	GetNodesOfClass<UNiagaraNodeParameterMapBase>(NiagaraParameterNodes);
 
+	// we will encounter pins that we've already changed, so we cache their previous state here to carry them over to potential new orphaned pins
 	TMap<UEdGraphPin*, FNiagaraTypeDefinition> OldTypesCache;
 	TMap<UEdGraphPin*, FString> OldDefaultValuesCache;
+	TMap<UEdGraphPin*, FString> OldAutogeneratedDefaultValuesCache;
 	
-	auto TryCreateOrphanedPin = [&](UNiagaraNode* Node, FName OrphanedPinName, UEdGraphPin* OriginalPin)
+	// we keep track of pins that will (but haven't yet) change type to keep connections intact that will have pins of the same type after the op is done
+	TSet<UEdGraphPin*> ChangingPins;
+	// a subset of changing pins. This only includes parameter pins (i.e. outputs on map gets & inputs on map sets)
+	TSet<UEdGraphPin*> ChangingParameterPins;
+	TSet<UEdGraphPin*> ChangingDefaultPins;
+	// since default pins are named 'None' but, if orphaned, can exist independently from their outputs, we cache the parameter pin's name here to override the new orphaned pin's name
+	TMap<UEdGraphPin*, FName> OrphanedDefaultPinNameOverrides;
+	
+	auto GatherChangingPins = [&](UNiagaraNodeParameterMapBase* Node)
 	{
-		// We check pin connections for type to determine if we want to keep a connection or discard it
-		TSet<UEdGraphPin*> PinConnectionsToKeep;
-		for (UEdGraphPin* ConnectedPin : OriginalPin->LinkedTo)
+		for(UEdGraphPin* Pin : Node->Pins)
 		{
-			// we might have already visited and changed the connected pin, so we test against its old type cache if available
-			FNiagaraTypeDefinition CurrentPinType = UEdGraphSchema_Niagara::PinToTypeDefinition(ConnectedPin);
-			FNiagaraVariable ConnectedVariable(OldTypesCache.Contains(ConnectedPin) ? OldTypesCache[ConnectedPin] : CurrentPinType, ConnectedPin->PinName);
-			
-			// keep it if the connected pin has a matching type already or will have a matching type after type conversion is done
-			if((ConnectedPin->PinType.PinSubCategory == UNiagaraNodeParameterMapBase::ParameterPinSubCategory && ParametersToChange.Contains(ConnectedVariable))
-				|| CurrentPinType == NewType)
+			if(Node->IsAddPin(Pin) || Pin->bOrphanedPin)
 			{
-				PinConnectionsToKeep.Add(ConnectedPin);
-			}						
-		}
+				continue;
+			}
 
-		// in case we are keeping all connections, all previously existing connections are parameter pins of the same type or will be of the same type after this operation is done
-		if((PinConnectionsToKeep.Num() != OriginalPin->LinkedTo.Num()) || (OriginalPin->DefaultValue != OldDefaultValuesCache[OriginalPin]))
-		{
-			ensure(OldTypesCache.Contains(OriginalPin) && OldDefaultValuesCache.Contains(OriginalPin));
-			
-			UEdGraphPin* NewOrphanedPin = Node->CreatePin(OriginalPin->Direction, OriginalPin->PinType.PinCategory, OrphanedPinName);
-			NewOrphanedPin->PinType = UEdGraphSchema_Niagara::TypeDefinitionToPinType(OldTypesCache[OriginalPin]);
-			NewOrphanedPin->DefaultValue = OldDefaultValuesCache[OriginalPin];
-			NewOrphanedPin->PersistentGuid = FGuid::NewGuid();
-			NewOrphanedPin->bNotConnectable = true;
-			NewOrphanedPin->bOrphanedPin = true;
-			
-			TSet<UEdGraphPin*> PinsToDisconnect;
-			for (UEdGraphPin* PinConnection : OriginalPin->LinkedTo)
+			// we add pins that represent the parameter here
+			FNiagaraVariable Variable(UEdGraphSchema_Niagara::PinToTypeDefinition(Pin), Pin->PinName);			
+			if(ParametersToChange.Contains(Variable) && Pin->PinType.PinSubCategory == UNiagaraNodeParameterMapBase::ParameterPinSubCategory)
 			{
-				// wire up the new orphaned pin with previous connections if we aren't keeping the connection
-				if(!PinConnectionsToKeep.Contains(PinConnection))
+				ChangingPins.Add(Pin);
+				ChangingParameterPins.Add(Pin);
+				OldTypesCache.Add(Pin, Variable.GetType());
+				OldDefaultValuesCache.Add(Pin, Pin->DefaultValue);
+				OldAutogeneratedDefaultValuesCache.Add(Pin, Pin->AutogeneratedDefaultValue);
+				
+				// parameter map get nodes might have additional default pins that we can't match via name
+				if(UNiagaraNodeParameterMapGet* MapGet = Cast<UNiagaraNodeParameterMapGet>(Node))
 				{
-					NewOrphanedPin->LinkedTo.Add(PinConnection);
-					PinConnection->LinkedTo.Add(NewOrphanedPin);
-
-					PinsToDisconnect.Add(PinConnection);
+					if(UEdGraphPin* DefaultPin = MapGet->GetDefaultPin(Pin))
+					{
+						ChangingPins.Add(DefaultPin);
+						ChangingDefaultPins.Add(DefaultPin);
+						OrphanedDefaultPinNameOverrides.Add(DefaultPin, Pin->PinName);
+						OldTypesCache.Add(DefaultPin, Variable.GetType());
+						OldDefaultValuesCache.Add(DefaultPin, DefaultPin->DefaultValue);
+						OldAutogeneratedDefaultValuesCache.Add(DefaultPin, DefaultPin->AutogeneratedDefaultValue);
+					}
 				}
 			}
-
-			// disconnect the previous connections from the original pin
-			for (UEdGraphPin* PinToDisconnect : PinsToDisconnect)
-			{
-				PinToDisconnect->LinkedTo.Remove(OriginalPin);
-				OriginalPin->LinkedTo.Remove(PinToDisconnect);
-			}							
 		}
 	};
-
-	// for every parameter, we loop over every node and every pin until we encounter a pin representing the parameter
-	for(const FNiagaraVariable& CurrentParameter : ParametersToChange)
+	
+	for (UNiagaraNodeParameterMapBase* NiagaraNode : NiagaraParameterNodes)
 	{
-		// a safe guard so we only look for pins that actually represent parameters and not just any pin name with the same name as a parameter
-		if(!VariableToScriptVariable.Contains(CurrentParameter))
-		{
-			continue;
-		}
-		
-		for (UNiagaraNodeParameterMapBase* NiagaraNode : NiagaraParameterNodes)
-		{
-			TSet<UEdGraphPin*> NewOrphanedPins;
-			for (UEdGraphPin* Pin : NiagaraNode->Pins)
-			{
-				if (NiagaraNode->IsAddPin(Pin) || Pin->bOrphanedPin)
+		GatherChangingPins(NiagaraNode);
+	}
+	
+	auto TryCreateOrphanedPin = [&](UNiagaraNode* Node, FName OrphanedPinName, UEdGraphPin* OriginalPin)
+	{
+		if(OriginalPin->LinkedTo.Num() > 0 || (OriginalPin->Direction == EGPD_Input && OriginalPin->LinkedTo.Num() == 0 && OriginalPin->bDefaultValueIsIgnored == false && OldDefaultValuesCache[OriginalPin] != OldAutogeneratedDefaultValuesCache[OriginalPin]))
+		{
+			// We check pin connections for type to determine if we want to keep a connection or orphan it
+			TSet<UEdGraphPin*> PinConnectionsToKeep;
+			for (UEdGraphPin* ConnectedPin : OriginalPin->LinkedTo)
+			{
+				// we might have already visited and changed the connected pin, so we test against its old type cache if available
+				FNiagaraTypeDefinition ConnectedPinType = UEdGraphSchema_Niagara::PinToTypeDefinition(ConnectedPin);
+				FNiagaraVariable ConnectedVariable(OldTypesCache.Contains(ConnectedPin) ? OldTypesCache[ConnectedPin] : ConnectedPinType, ConnectedPin->PinName);
+
+				/* We keep a connection in the following cases:
+				// - Default case: if the connected pin already has the type we are changing to
+				// - If the connected pin will have a matching type after type conversion is done
+				// - If the connected pin is a numeric pin and the new type will be compatible with a numeric pin
+				// - If we change types from vec -> position for easier fixup */
+				if(ConnectedPinType == NewType 
+					|| (ConnectedPin->PinType.PinSubCategory == UNiagaraNodeParameterMapBase::ParameterPinSubCategory && ChangingPins.Contains(ConnectedPin))
+					|| (ConnectedPinType == FNiagaraTypeDefinition::GetGenericNumericDef() && FNiagaraTypeDefinition::GetNumericTypes().Contains(NewType))
+					|| (GetDefault<UNiagaraSettings>()->bEnforceStrictStackTypes == false && ConnectedPinType == FNiagaraTypeDefinition::GetVec3Def() && NewType == FNiagaraTypeDefinition::GetPositionDef()))
 				{
-					continue;
-				}
+					PinConnectionsToKeep.Add(ConnectedPin);
+				}						
+			}
 
-				FNiagaraVariable PinVariable = FNiagaraVariable(UEdGraphSchema_Niagara::PinToTypeDefinition(Pin), Pin->PinName);
-				if (PinVariable == CurrentParameter)
+			// in case we are keeping all connections, all previously existing connections are parameter pins of the same type or will be of the same type after this operation is done
+			// we also check for default values cache as by this point we have already changed the default values of the original pin
+			if((PinConnectionsToKeep.Num() != OriginalPin->LinkedTo.Num()) || (OldDefaultValuesCache[OriginalPin] != OldAutogeneratedDefaultValuesCache[OriginalPin]))
+			{
+				ensure(OldTypesCache.Contains(OriginalPin) && OldDefaultValuesCache.Contains(OriginalPin));
+			
+				UEdGraphPin* NewOrphanedPin = Node->CreatePin(OriginalPin->Direction, OriginalPin->PinType.PinCategory, OrphanedPinName);
+				NewOrphanedPin->PinType = UEdGraphSchema_Niagara::TypeDefinitionToPinType(OldTypesCache[OriginalPin]);
+				NewOrphanedPin->DefaultValue = OldDefaultValuesCache[OriginalPin];
+				NewOrphanedPin->PersistentGuid = FGuid::NewGuid();
+				NewOrphanedPin->bNotConnectable = true;
+				NewOrphanedPin->bOrphanedPin = true;
+			
+				TSet<UEdGraphPin*> PinsToDisconnect;
+				for (UEdGraphPin* ConnectedPin : OriginalPin->LinkedTo)
 				{
-					// we cache off the old type before changing it as we might encounter a pin we already changed later on in connections
-					if(!OldTypesCache.Contains(Pin))
-					{
-						OldTypesCache.Add(Pin, PinVariable.GetType());					
-					}
-					if(!OldDefaultValuesCache.Contains(Pin))
+					ConnectedPin->Modify();
+				
+					// wire up the new orphaned pin with previous connections if we aren't keeping the connection
+					if(!PinConnectionsToKeep.Contains(ConnectedPin))
 					{
-						OldDefaultValuesCache.Add(Pin, Pin->DefaultValue);
+						NewOrphanedPin->LinkedTo.Add(ConnectedPin);
+						ConnectedPin->LinkedTo.Add(NewOrphanedPin);
+
+						PinsToDisconnect.Add(ConnectedPin);
 					}
-					
-					Pin->Modify();
-					Pin->PinType = NiagaraSchema->TypeDefinitionToPinType(NewType);
+				}
+
+				// disconnect the previous connections from the original pin
+				for (UEdGraphPin* PinToDisconnect : PinsToDisconnect)
+				{
+					PinToDisconnect->LinkedTo.Remove(OriginalPin);
+					OriginalPin->LinkedTo.Remove(PinToDisconnect);
+				}
+
+				return NewOrphanedPin;
+			}
+		}
+		
+		return (UEdGraphPin*) nullptr;
+	};
+		
+	for (UNiagaraNodeParameterMapBase* NiagaraNode : NiagaraParameterNodes)
+	{
+		TArray<UEdGraphPin*> NewOrphanedPins;
+
+		bool bNodeChanged = false;
+		TMap<UEdGraphPin*, FName> ValidOrphanedPinCandidates;
+		// we go through all changing pins to change their type, and if applicable, attempt to create an orphaned pin
+		for (UEdGraphPin* Pin : NiagaraNode->Pins)
+		{			
+			if (ChangingPins.Contains(Pin))
+			{					
+				Pin->Modify();
+				Pin->PinType = NiagaraSchema->TypeDefinitionToPinType(NewType);
+				// we need to restore the parameter sub category here, but only for parameter pins
+				if(ChangingParameterPins.Contains(Pin))
+				{
 					Pin->PinType.PinSubCategory = UNiagaraNodeParameterMapBase::ParameterPinSubCategory;
+				}
+				
+				if(bAllowOrphanedPins)
+				{
+					// the name for most parameter pins will remain the same as before, but not for default pins, which will adapt the parameter's name
+					FName OrphanedPinName = OrphanedDefaultPinNameOverrides.Contains(Pin) ? OrphanedDefaultPinNameOverrides[Pin] : Pin->PinName; 
+					ValidOrphanedPinCandidates.Add(Pin, OrphanedPinName);
+				}
 
-					// we only want to create an additional orphaned pin if there is at least one pin connection or if the pin value differs from the 'default' pin value
-					if(bAllowOrphanedPins)
-					{
-						bool bCreateOrphanedPin = (Pin->bDefaultValueIsIgnored == false && Pin->DefaultValue != Pin->AutogeneratedDefaultValue) || Pin->LinkedTo.Num() > 0;
-						if(bCreateOrphanedPin)
-						{
-							TryCreateOrphanedPin(NiagaraNode, Pin->PinName, Pin);
-						}
-					}
+				// since we have now determined if we want to create an orphaned pin based on default values, we can safely reset them.
+				Pin->ResetDefaultValue();
+				UNiagaraNode::SetPinDefaultToTypeDefaultIfUnset(Pin);
+
+				bNodeChanged = true;
+			}
+		}
+
+		TMap<UEdGraphPin*, UEdGraphPin*> OriginalToOrphanedMap;
+		for(auto& OrphanedPinCandidate : ValidOrphanedPinCandidates)
+		{
+			UEdGraphPin* NewOrphanedPin = TryCreateOrphanedPin(NiagaraNode, OrphanedPinCandidate.Value, OrphanedPinCandidate.Key);
+			if(NewOrphanedPin != nullptr)
+			{
+				OriginalToOrphanedMap.Add(OrphanedPinCandidate.Key, NewOrphanedPin);
+				NewOrphanedPins.Add(NewOrphanedPin);
+			}
+		}
 
-					// fix default pin types if necessary. We don't support orphaned pins for default types so we just disconnect the pin
-					if (UNiagaraNodeParameterMapGet* GetNode = Cast<UNiagaraNodeParameterMapGet>(NiagaraNode))
+		// we now match up the new orphaned default & output pins on map get nodes and add them to the node's map
+		TSet<UEdGraphPin*> HandledGuidsOfOrphanedPins;
+		for(auto& Pin : NiagaraNode->Pins)
+		{	
+			if(OriginalToOrphanedMap.Contains(Pin) && !HandledGuidsOfOrphanedPins.Contains(Pin))
+			{
+				UEdGraphPin* OrphanedPin = OriginalToOrphanedMap[Pin];
+				OrphanedPin->PersistentGuid = FGuid::NewGuid();
+				HandledGuidsOfOrphanedPins.Add(Pin);
+				if(UNiagaraNodeParameterMapGet* MapGet = Cast<UNiagaraNodeParameterMapGet>(NiagaraNode))
+				{
+					if(UEdGraphPin* DefaultPin = MapGet->GetDefaultPin(Pin))
 					{
-						UEdGraphPin* DefaultPin = GetNode->GetDefaultPin(Pin);
-						if (Pin->Direction == EGPD_Output && DefaultPin && Pin->PinType != DefaultPin->PinType)
+						if(OriginalToOrphanedMap.Contains(DefaultPin))
 						{
-							ensure(OldTypesCache.Contains(Pin));
-							
-							if(!OldTypesCache.Contains(DefaultPin))
-							{
-								OldTypesCache.Add(DefaultPin, UEdGraphSchema_Niagara::PinToTypeDefinition(DefaultPin));					
-							}
-							if(!OldDefaultValuesCache.Contains(DefaultPin))
-							{
-								OldDefaultValuesCache.Add(DefaultPin, DefaultPin->DefaultValue);
-							}
-							
-							DefaultPin->Modify();
-							DefaultPin->PinType = Pin->PinType;
-							DefaultPin->ResetDefaultValue();
-							UNiagaraNode::SetPinDefaultToTypeDefaultIfUnset(DefaultPin);							
-
-							if(bAllowOrphanedPins)
+							if(UEdGraphPin* NewOrphanedDefaultPin = OriginalToOrphanedMap[DefaultPin])
 							{
-								bool bCreateOrphanedPin = (Pin->bDefaultValueIsIgnored == false && OldDefaultValuesCache[DefaultPin] != DefaultPin->AutogeneratedDefaultValue) || Pin->LinkedTo.Num() > 0;
-
-								if(bCreateOrphanedPin)
-								{
-									TryCreateOrphanedPin(NiagaraNode, Pin->PinName, DefaultPin);
-								}
+								NewOrphanedDefaultPin->PersistentGuid = FGuid::NewGuid();
+								HandledGuidsOfOrphanedPins.Add(DefaultPin);
+								MapGet->AddOrphanedPinPairGuids(OrphanedPin, NewOrphanedDefaultPin);
 							}
 						}
 					}
-
-					NiagaraNode->MarkNodeRequiresSynchronization(__FUNCTION__, true);
-					break;
-				}			
+				}
 			}
+		}		
+
+		if(bNodeChanged)
+		{
+			NiagaraNode->MarkNodeRequiresSynchronization(__FUNCTION__, true);
 		}
-		
+	}
+
+	for(const FNiagaraVariable& CurrentParameter : ParametersToChange)
+	{		
 		TObjectPtr<UNiagaraScriptVariable>* ScriptVariablePtr = VariableToScriptVariable.Find(CurrentParameter);
 		if (ScriptVariablePtr != nullptr && !ScriptVariablePtr->IsNull())
 		{
@@ -639,8 +701,9 @@ void UNiagaraGraph::ChangeParameterType(const TArray<FNiagaraVariable>& Paramete
 			}
 			ParameterToReferencesMap.Remove(CurrentParameter);
 		}
-		InvalidateCachedParameterData();
-	}	
+	}
+	
+	InvalidateCachedParameterData();
 }
 
 TArray<UEdGraphPin*> UNiagaraGraph::FindParameterMapDefaultValuePins(const FName VariableName) const
@@ -670,7 +733,7 @@ TArray<UEdGraphPin*> UNiagaraGraph::FindParameterMapDefaultValuePins(const FName
 			GetNode->GetOutputPins(OutputPins);
 			for (UEdGraphPin* OutputPin : OutputPins)
 			{
-				if (VariableName != OutputPin->PinName)
+				if (VariableName != OutputPin->PinName || OutputPin->PinType.PinSubCategory != UNiagaraNodeParameterMapBase::ParameterPinSubCategory)
 				{
 					continue;
 				}
diff --git a/Engine/Plugins/FX/Niagara/Source/NiagaraEditor/Private/NiagaraNode.cpp b/Engine/Plugins/FX/Niagara/Source/NiagaraEditor/Private/NiagaraNode.cpp
--- a/Engine/Plugins/FX/Niagara/Source/NiagaraEditor/Private/NiagaraNode.cpp
+++ b/Engine/Plugins/FX/Niagara/Source/NiagaraEditor/Private/NiagaraNode.cpp
@@ -188,6 +188,7 @@ bool UNiagaraNode::SetPinDefaultToTypeDefaultIfUnset(UEdGraphPin* InPin)
 	if (Schema->TryGetPinDefaultValueFromNiagaraVariable(Var, PinDefaultValue))
 	{
 		InPin->DefaultValue = PinDefaultValue;
+		InPin->AutogeneratedDefaultValue = PinDefaultValue;
 		return true;
 	}
 
@@ -636,7 +637,7 @@ void UNiagaraNode::NodeConnectionListChanged()
 	TArray<UEdGraphPin*> CurrentPins = Pins;
 	for (UEdGraphPin* Pin : CurrentPins)
 	{
-		if (Pin->bOrphanedPin && Pin->LinkedTo.Num() == 0)
+		if (Pin->bOrphanedPin && Pin->LinkedTo.Num() == 0 && Pin->DefaultValue == Pin->AutogeneratedDefaultValue)
 		{
 			RemovePin(Pin);
 		}
diff --git a/Engine/Plugins/FX/Niagara/Source/NiagaraEditor/Private/NiagaraNodeParameterMapGet.cpp b/Engine/Plugins/FX/Niagara/Source/NiagaraEditor/Private/NiagaraNodeParameterMapGet.cpp
--- a/Engine/Plugins/FX/Niagara/Source/NiagaraEditor/Private/NiagaraNodeParameterMapGet.cpp
+++ b/Engine/Plugins/FX/Niagara/Source/NiagaraEditor/Private/NiagaraNodeParameterMapGet.cpp
@@ -172,9 +172,25 @@ void UNiagaraNodeParameterMapGet::OnPinRenamed(UEdGraphPin* RenamedPin, const FS
 
 void UNiagaraNodeParameterMapGet::OnPinRemoved(UEdGraphPin* InRemovedPin)
 {
+	FGuid PersistentGuidKey;
 	if(UEdGraphPin* DefaultPin = GetDefaultPin(InRemovedPin))
 	{
-		RemovePin(DefaultPin);
+		PersistentGuidKey = InRemovedPin->PersistentGuid;
+
+		if(!InRemovedPin->bOrphanedPin)
+		{
+			RemovePin(DefaultPin);
+		}
+	}
+
+	if(UEdGraphPin* OutputPin = GetOutputPinForDefault(InRemovedPin))
+	{
+		PersistentGuidKey = OutputPin->PersistentGuid;
+	}
+
+	if(PersistentGuidKey.IsValid() && PinOutputToPinDefaultPersistentId.Contains(PersistentGuidKey))
+	{
+		PinOutputToPinDefaultPersistentId.Remove(PersistentGuidKey);
 	}
 	
 	Super::OnPinRemoved(InRemovedPin);
@@ -379,6 +395,13 @@ FText UNiagaraNodeParameterMapGet::GetNodeTitle(ENodeTitleType::Type TitleType)
 	return LOCTEXT("UNiagaraNodeParameterMapGetName", "Map Get");
 }
 
+void UNiagaraNodeParameterMapGet::AddOrphanedPinPairGuids(UEdGraphPin* OutputPin, UEdGraphPin* DefaultPin)
+{
+	ensure(OutputPin && DefaultPin && OutputPin->bOrphanedPin && DefaultPin->bOrphanedPin && OutputPin->PersistentGuid.IsValid() && DefaultPin->PersistentGuid.IsValid());
+
+	PinOutputToPinDefaultPersistentId.Add(OutputPin->PersistentGuid, DefaultPin->PersistentGuid);
+}
+
 void UNiagaraNodeParameterMapGet::BuildParameterMapHistory(FNiagaraParameterMapHistoryBuilder& OutHistory, bool bRecursive /*= true*/, bool bFilterForCompilation /*= true*/) const
 {
 	if (bRecursive)
diff --git a/Engine/Plugins/FX/Niagara/Source/NiagaraEditor/Private/NiagaraNodeParameterMapGet.h b/Engine/Plugins/FX/Niagara/Source/NiagaraEditor/Private/NiagaraNodeParameterMapGet.h
--- a/Engine/Plugins/FX/Niagara/Source/NiagaraEditor/Private/NiagaraNodeParameterMapGet.h
+++ b/Engine/Plugins/FX/Niagara/Source/NiagaraEditor/Private/NiagaraNodeParameterMapGet.h
@@ -35,6 +35,8 @@ public:
 
 	virtual FText GetNodeTitle(ENodeTitleType::Type TitleType) const;
 
+	void AddOrphanedPinPairGuids(UEdGraphPin* OutputPin, UEdGraphPin* DefaultPin);
+	
 	/** Get the default value input pin for one of the output pins specified.*/
 	UEdGraphPin* GetDefaultPin(UEdGraphPin* OutputPin) const;
 
diff --git a/Engine/Plugins/FX/Niagara/Source/NiagaraEditor/Private/SNiagaraParameterName.cpp b/Engine/Plugins/FX/Niagara/Source/NiagaraEditor/Private/SNiagaraParameterName.cpp
--- a/Engine/Plugins/FX/Niagara/Source/NiagaraEditor/Private/SNiagaraParameterName.cpp
+++ b/Engine/Plugins/FX/Niagara/Source/NiagaraEditor/Private/SNiagaraParameterName.cpp
@@ -405,33 +405,49 @@ void SNiagaraParameterNamePinLabel::Construct(const FArguments& InArgs, UEdGraph
 {
 	TargetPin = InTargetPin;
 
-	SNiagaraParameterNameTextBlock::Construct(SNiagaraParameterNameTextBlock::FArguments()
-		.EditableTextStyle(InArgs._EditableTextStyle)
-		.ParameterText(InArgs._ParameterText)
-		.IsReadOnly(InArgs._IsReadOnly)
-		.HighlightText(InArgs._HighlightText)
-		.OnVerifyTextChanged(InArgs._OnVerifyTextChanged)
-		.OnTextCommitted(InArgs._OnTextCommitted)
-		.IsSelected(InArgs._IsSelected)
-		.DecoratorHAlign(InArgs._DecoratorHAlign)
-		.DecoratorPadding(InArgs._DecoratorPadding)
-		.Decorator()
+	ChildSlot
+	[
+		SNew(SBorder)
+		.BorderImage(FAppStyle::Get().GetBrush("NoBorder"))
+		.ForegroundColor(this, &SNiagaraParameterNamePinLabel::GetForegroundColor)
 		[
-			InArgs._Decorator.Widget
-		]);
+			SAssignNew(ParameterNameTextBlock, SNiagaraParameterNameTextBlock)
+			.EditableTextStyle(InArgs._EditableTextStyle)
+			.ParameterText(InArgs._ParameterText)
+			.IsReadOnly(InArgs._IsReadOnly)
+			.HighlightText(InArgs._HighlightText)
+			.OnVerifyTextChanged(InArgs._OnVerifyTextChanged)
+			.OnTextCommitted(InArgs._OnTextCommitted)
+			.IsSelected(InArgs._IsSelected)
+			.DecoratorHAlign(InArgs._DecoratorHAlign)
+			.DecoratorPadding(InArgs._DecoratorPadding)
+			.Decorator()
+			[
+				InArgs._Decorator.Widget
+			]
+		]
+	];		
 }
 
 void SNiagaraParameterNamePinLabel::Tick(const FGeometry& AllottedGeometry, const double InCurrentTime, const float InDeltaTime)
 {
-	SNiagaraParameterNameTextBlock::Tick(AllottedGeometry, InCurrentTime, InDeltaTime);
-
 	UNiagaraNodeParameterMapBase* ParameterMapNode = Cast<UNiagaraNodeParameterMapBase>(TargetPin->GetOwningNode());
 	if (ParameterMapNode != nullptr)
 	{
 		if (ParameterMapNode->GetIsPinEditNamespaceModifierPending(TargetPin))
 		{
 			ParameterMapNode->SetIsPinEditNamespaceModifierPending(TargetPin, false);
-			EnterNamespaceModifierEditingMode();
+			ParameterNameTextBlock->EnterNamespaceModifierEditingMode();
 		}
 	}
-}
\ No newline at end of file
+}
+
+FSlateColor SNiagaraParameterNamePinLabel::GetForegroundColor() const
+{
+	return TargetPin->bOrphanedPin ? FLinearColor::Red : FLinearColor::White;
+}
+
+void SNiagaraParameterNamePinLabel::EnterEditingMode()
+{
+	ParameterNameTextBlock->EnterEditingMode();
+}
diff --git a/Engine/Plugins/FX/Niagara/Source/NiagaraEditor/Private/TNiagaraGraphPinEditableName.h b/Engine/Plugins/FX/Niagara/Source/NiagaraEditor/Private/TNiagaraGraphPinEditableName.h
--- a/Engine/Plugins/FX/Niagara/Source/NiagaraEditor/Private/TNiagaraGraphPinEditableName.h
+++ b/Engine/Plugins/FX/Niagara/Source/NiagaraEditor/Private/TNiagaraGraphPinEditableName.h
@@ -129,7 +129,7 @@ protected:
 			.Padding(FMargin(0.0f, 1.0f, 0.0f, 1.0f))
 			[
 				SAssignNew(CreatedParameterNamePinLabel, SNiagaraParameterNamePinLabel, this->GraphPinObj)
-				.EditableTextStyle(&FEditorStyle::Get().GetWidgetStyle<FInlineEditableTextBlockStyle>("Graph.Node.InlineEditablePinName"))
+				.EditableTextStyle(&FNiagaraEditorStyle::Get().GetWidgetStyle<FInlineEditableTextBlockStyle>("NiagaraEditor.Graph.Node.InlineEditablePinName"))
 				.ParameterText(this, &TNiagaraGraphPinEditableName<BaseClass>::GetParentPinLabel)
 				.IsReadOnly(bIsPinEditable == false)
 				.Visibility(this, &TNiagaraGraphPinEditableName<BaseClass>::GetParentPinVisibility)
diff --git a/Engine/Plugins/FX/Niagara/Source/NiagaraEditor/Private/Widgets/SNiagaraGraphParameterMapGetNode.cpp b/Engine/Plugins/FX/Niagara/Source/NiagaraEditor/Private/Widgets/SNiagaraGraphParameterMapGetNode.cpp
--- a/Engine/Plugins/FX/Niagara/Source/NiagaraEditor/Private/Widgets/SNiagaraGraphParameterMapGetNode.cpp
+++ b/Engine/Plugins/FX/Niagara/Source/NiagaraEditor/Private/Widgets/SNiagaraGraphParameterMapGetNode.cpp
@@ -133,7 +133,7 @@ void SNiagaraGraphParameterMapGetNode::CreatePinWidgets()
 
 	AddRowWidgetToNode(SourceRowWidget);
 	
-	// Deferred pin adding to line up input/output pins by name.
+	// Deferred pin adding to line up input/output pins by persistent guid mapping.
 	for (int32 i = 0; i < OutputPins.Num(); i++)
 	{
 		TSharedRef<SGraphPin> OutputPin = OutputPins[i];
diff --git a/Engine/Plugins/FX/Niagara/Source/NiagaraEditor/Public/NiagaraGraph.h b/Engine/Plugins/FX/Niagara/Source/NiagaraEditor/Public/NiagaraGraph.h
--- a/Engine/Plugins/FX/Niagara/Source/NiagaraEditor/Public/NiagaraGraph.h
+++ b/Engine/Plugins/FX/Niagara/Source/NiagaraEditor/Public/NiagaraGraph.h
@@ -280,7 +280,7 @@ class UNiagaraGraph : public UEdGraph
 	const TMap<FNiagaraVariable, FNiagaraGraphParameterReferenceCollection>& GetParameterReferenceMap() const; // NOTE: The const is a lie! (This indirectly calls RefreshParameterReferences, which can recreate the entire map)
 
 	UNiagaraScriptVariable* GetScriptVariable(FNiagaraVariable Parameter, bool bUpdateIfPending = false) const;
-	UNiagaraScriptVariable* GetScriptVariable(FName ParameterName, bool bUpdateIfPending = false) const;
+	NIAGARAEDITOR_API UNiagaraScriptVariable* GetScriptVariable(FName ParameterName, bool bUpdateIfPending = false) const;
 
 	/** Adds parameter to the VariableToScriptVariable map.*/
 	UNiagaraScriptVariable* AddParameter(const FNiagaraVariable& Parameter, bool bIsStaticSwitch = false);
@@ -313,7 +313,7 @@ class UNiagaraGraph : public UEdGraph
 	 *  Changing multiple parameters at once helps with maintaining connections.
 	 *  CAUTION: Do not allow orphaned pins in the stack graphs, as they aren't user facing.
 	 */
-	void ChangeParameterType(const TArray<FNiagaraVariable>& ParametersToChange, const FNiagaraTypeDefinition& NewType, bool bAllowOrphanedPins = false);
+	NIAGARAEDITOR_API void ChangeParameterType(const TArray<FNiagaraVariable>& ParametersToChange, const FNiagaraTypeDefinition& NewType, bool bAllowOrphanedPins = false);
 
 	/** Gets a delegate which is called whenever a contained data interfaces changes. */
 	FOnDataInterfaceChanged& OnDataInterfaceChanged();
diff --git a/Engine/Plugins/FX/Niagara/Source/NiagaraEditor/Public/Widgets/SNiagaraParameterName.h b/Engine/Plugins/FX/Niagara/Source/NiagaraEditor/Public/Widgets/SNiagaraParameterName.h
--- a/Engine/Plugins/FX/Niagara/Source/NiagaraEditor/Public/Widgets/SNiagaraParameterName.h
+++ b/Engine/Plugins/FX/Niagara/Source/NiagaraEditor/Public/Widgets/SNiagaraParameterName.h
@@ -152,7 +152,7 @@ private:
 	TSharedPtr<SNiagaraParameterName> ParameterName;
 };
 
-class NIAGARAEDITOR_API SNiagaraParameterNamePinLabel : public SNiagaraParameterNameTextBlock
+class NIAGARAEDITOR_API SNiagaraParameterNamePinLabel : public SCompoundWidget
 {
 public:
 	SLATE_BEGIN_ARGS(SNiagaraParameterNamePinLabel)
@@ -179,6 +179,10 @@ public:
 
 	virtual void Tick(const FGeometry& AllottedGeometry, const double InCurrentTime, const float InDeltaTime) override;
 
+	virtual FSlateColor GetForegroundColor() const override;
+
+	void EnterEditingMode();
 private:
 	UEdGraphPin* TargetPin;
+	TSharedPtr<SNiagaraParameterNameTextBlock> ParameterNameTextBlock;
 };
