- Add support for returning bone scale

diff --git a/Engine/Plugins/FX/Niagara/Shaders/Private/NiagaraDataInterfaceSkeletalMesh.ush b/Engine/Plugins/FX/Niagara/Shaders/Private/NiagaraDataInterfaceSkeletalMesh.ush
--- a/Engine/Plugins/FX/Niagara/Shaders/Private/NiagaraDataInterfaceSkeletalMesh.ush
+++ b/Engine/Plugins/FX/Niagara/Shaders/Private/NiagaraDataInterfaceSkeletalMesh.ush
@@ -812,13 +812,14 @@ float4 DISKelMesh_QuatSlerp(float4 Quat1, float4 Quat2, float Slerp)
     return (Scale0 * Quat1) + (Scale1 * Quat2);
 }
 
-void DISkelMesh_GetSkinnedBoneCommon(in FDISkelMeshContext DIContext, in int Bone, in float Interp, in bool bWorldSpace, out float3 OutPosition, out float4 OutRotation, out float3 OutVelocity)
+void DISkelMesh_GetSkinnedBoneCommon(in FDISkelMeshContext DIContext, in int Bone, in float Interp, in bool bWorldSpace, out float3 OutPosition, out float4 OutRotation, out float3 OutScale, out float3 OutVelocity)
 {
 	// Early out for bad data
 	if ( DIContext.NumBones + DIContext.NumFilteredSockets == 0 )
 	{
 		OutPosition = float3(0.0f, 0.0f, 0.0f);
 		OutRotation = float4(0.0f, 0.0f, 0.0f, 1.0f);
+		OutScale	= float3(0.0f, 0.0f, 0.0f);
 		OutVelocity = float3(0.0f, 0.0f, 0.0f);
 		return;
 	}
@@ -827,8 +828,10 @@ void DISkelMesh_GetSkinnedBoneCommon(in FDISkelMeshContext DIContext, in int Bon
 
     float3 PrevPosition = DIContext.MeshPrevSamplingBonesBuffer[Bone * 3].xyz;
     float4 PrevRotation = DIContext.MeshPrevSamplingBonesBuffer[Bone * 3 + 1];
+	float3 PrevScale	= DIContext.MeshPrevSamplingBonesBuffer[Bone * 3 + 2].xyz;
     float3 CurrPosition = DIContext.MeshCurrSamplingBonesBuffer[Bone * 3].xyz;
-    float4 CurrRotation = DIContext.MeshCurrSamplingBonesBuffer[Bone * 3 + 1];
+    float4 CurrRotation	= DIContext.MeshCurrSamplingBonesBuffer[Bone * 3 + 1];
+	float3 CurrScale	= DIContext.MeshCurrSamplingBonesBuffer[Bone * 3 + 2].xyz;
 
     if (bWorldSpace)
     {
@@ -836,32 +839,36 @@ void DISkelMesh_GetSkinnedBoneCommon(in FDISkelMeshContext DIContext, in int Bon
         CurrPosition = mul(float4(CurrPosition, 1), DIContext.InstanceTransform).xyz;
 		PrevRotation = NiagaraGPU_QuatMul(DIContext.InstancePrevRotation, PrevRotation);
 		CurrRotation = NiagaraGPU_QuatMul(DIContext.InstanceRotation, CurrRotation);
+		PrevScale = mul(PrevScale, (float3x3)DIContext.InstancePrevTransform);
+		CurrScale = mul(CurrScale, (float3x3)DIContext.InstanceTransform);
     }
     CurrPosition = lerp(PrevPosition, CurrPosition, Interp);
+	CurrScale	 = lerp(PrevScale, CurrScale, Interp);
 
     OutPosition = CurrPosition;
     OutRotation = DISKelMesh_QuatSlerp(PrevRotation, CurrRotation, Interp);
+	OutScale	= CurrScale;
     OutVelocity = (CurrPosition - PrevPosition) * DIContext.InstanceInvDeltaTime;
 }
 
-void DISkelMesh_GetSkinnedBone(in FDISkelMeshContext DIContext, in int Bone, out float3 OutPosition, out float4 OutRotation, out float3 OutVelocity)
+void DISkelMesh_GetSkinnedBone(in FDISkelMeshContext DIContext, in int Bone, out float3 OutPosition, out float4 OutRotation, out float3 OutScale, out float3 OutVelocity)
 {
-    DISkelMesh_GetSkinnedBoneCommon(DIContext, Bone, 1.0f, false, OutPosition, OutRotation, OutVelocity);
+    DISkelMesh_GetSkinnedBoneCommon(DIContext, Bone, 1.0f, false, OutPosition, OutRotation, OutScale, OutVelocity);
 }
 
-void DISkelMesh_GetSkinnedBoneInterpolated(in FDISkelMeshContext DIContext, in int Bone, in float Interp, out float3 OutPosition, out float4 OutRotation, out float3 OutVelocity)
+void DISkelMesh_GetSkinnedBoneInterpolated(in FDISkelMeshContext DIContext, in int Bone, in float Interp, out float3 OutPosition, out float4 OutRotation, out float3 OutScale, out float3 OutVelocity)
 {
-    DISkelMesh_GetSkinnedBoneCommon(DIContext, Bone, Interp, false, OutPosition, OutRotation, OutVelocity);
+    DISkelMesh_GetSkinnedBoneCommon(DIContext, Bone, Interp, false, OutPosition, OutRotation, OutScale, OutVelocity);
 }
 
-void DISkelMesh_GetSkinnedBoneWS(in FDISkelMeshContext DIContext, in int Bone, out float3 OutPosition, out float4 OutRotation, out float3 OutVelocity)
+void DISkelMesh_GetSkinnedBoneWS(in FDISkelMeshContext DIContext, in int Bone, out float3 OutPosition, out float4 OutRotation, out float3 OutScale, out float3 OutVelocity)
 {
-    DISkelMesh_GetSkinnedBoneCommon(DIContext, Bone, 1.0f, true, OutPosition, OutRotation, OutVelocity);
+    DISkelMesh_GetSkinnedBoneCommon(DIContext, Bone, 1.0f, true, OutPosition, OutRotation, OutScale, OutVelocity);
 }
 
-void DISkelMesh_GetSkinnedBoneInterpolatedWS(in FDISkelMeshContext DIContext, in int Bone, in float Interp, out float3 OutPosition, out float4 OutRotation, out float3 OutVelocity)
+void DISkelMesh_GetSkinnedBoneInterpolatedWS(in FDISkelMeshContext DIContext, in int Bone, in float Interp, out float3 OutPosition, out float4 OutRotation, out float3 OutScale, out float3 OutVelocity)
 {
-    DISkelMesh_GetSkinnedBoneCommon(DIContext, Bone, Interp, true, OutPosition, OutRotation, OutVelocity);
+    DISkelMesh_GetSkinnedBoneCommon(DIContext, Bone, Interp, true, OutPosition, OutRotation, OutScale, OutVelocity);
 }
 
 ////////////////////////////////////////////////////////////////////////////////////////////////////////////////////////////////////////////////////////////////
diff --git a/Engine/Plugins/FX/Niagara/Source/Niagara/Classes/NDISkeletalMeshCommon.h b/Engine/Plugins/FX/Niagara/Source/Niagara/Classes/NDISkeletalMeshCommon.h
--- a/Engine/Plugins/FX/Niagara/Source/Niagara/Classes/NDISkeletalMeshCommon.h
+++ b/Engine/Plugins/FX/Niagara/Source/Niagara/Classes/NDISkeletalMeshCommon.h
@@ -95,8 +95,10 @@ struct FSkinnedPositionAccessorHelper
 	FORCEINLINE void GetSkinnedPreviousTangentBasis(FSkeletalMeshAccessorHelper& Accessor, int32 VertexIndex, FVector3f& OutTangentX, FVector3f& OutTangentY, FVector3f& OutTangentZ) = delete;
 	FORCEINLINE FVector3f GetSkinnedBonePosition(FSkeletalMeshAccessorHelper& Accessor, int32 BoneIndex) = delete;
 	FORCEINLINE FVector3f GetSkinnedBonePreviousPosition(FSkeletalMeshAccessorHelper& Accessor, int32 BoneIndex) = delete;
-	FORCEINLINE_DEBUGGABLE FQuat4f GetSkinnedBoneRotation(FSkeletalMeshAccessorHelper& Accessor, int32 BoneIndex) = delete;
-	FORCEINLINE_DEBUGGABLE FQuat4f GetSkinnedBonePreviousRotation(FSkeletalMeshAccessorHelper& Accessor, int32 BoneIndex) = delete;
+	FORCEINLINE FQuat4f GetSkinnedBoneRotation(FSkeletalMeshAccessorHelper& Accessor, int32 BoneIndex) = delete;
+	FORCEINLINE FQuat4f GetSkinnedBonePreviousRotation(FSkeletalMeshAccessorHelper& Accessor, int32 BoneIndex) = delete;
+	FORCEINLINE FVector3f GetSkinnedBoneScale(FSkeletalMeshAccessorHelper& Accessor, int32 BoneIndex) = delete;
+	FORCEINLINE FVector3f GetSkinnedBonePreviousScale(FSkeletalMeshAccessorHelper& Accessor, int32 BoneIndex) = delete;
 };
 
 template<>
@@ -152,7 +154,7 @@ struct FSkinnedPositionAccessorHelper<TNDISkelMesh_SkinningModeInvalid>
 	}
 
 	// The bone accessor functions are valid if a mesh is present, so don't stub them entirely
-	FORCEINLINE_DEBUGGABLE FVector3f GetSkinnedBonePosition(FSkeletalMeshAccessorHelper& Accessor, int32 BoneIndex)
+	FORCEINLINE FVector3f GetSkinnedBonePosition(FSkeletalMeshAccessorHelper& Accessor, int32 BoneIndex)
 	{
 		const int32 NumRealBones = Accessor.Mesh->GetRefSkeleton().GetRawBoneNum();
 		if (BoneIndex < NumRealBones)
@@ -164,12 +166,12 @@ struct FSkinnedPositionAccessorHelper<TNDISkelMesh_SkinningModeInvalid>
 		return (FVector3f)RefTransform.GetLocation();
 	}
 
-	FORCEINLINE_DEBUGGABLE FVector3f GetSkinnedBonePreviousPosition(FSkeletalMeshAccessorHelper& Accessor, int32 BoneIndex)
+	FORCEINLINE FVector3f GetSkinnedBonePreviousPosition(FSkeletalMeshAccessorHelper& Accessor, int32 BoneIndex)
 	{
 		return GetSkinnedBonePosition(Accessor, BoneIndex);
 	}
 
-	FORCEINLINE_DEBUGGABLE FQuat4f GetSkinnedBoneRotation(FSkeletalMeshAccessorHelper& Accessor, int32 BoneIndex)
+	FORCEINLINE FQuat4f GetSkinnedBoneRotation(FSkeletalMeshAccessorHelper& Accessor, int32 BoneIndex)
 	{
 		const int32 NumRealBones = Accessor.Mesh->GetRefSkeleton().GetRawBoneNum();
 		if (BoneIndex < NumRealBones)
@@ -181,10 +183,27 @@ struct FSkinnedPositionAccessorHelper<TNDISkelMesh_SkinningModeInvalid>
 		return (FQuat4f)RefTransform.GetRotation();
 	}
 
-	FORCEINLINE_DEBUGGABLE FQuat4f GetSkinnedBonePreviousRotation(FSkeletalMeshAccessorHelper& Accessor, int32 BoneIndex)
+	FORCEINLINE FQuat4f GetSkinnedBonePreviousRotation(FSkeletalMeshAccessorHelper& Accessor, int32 BoneIndex)
 	{
 		return GetSkinnedBoneRotation(Accessor, BoneIndex);
 	}
+
+	FORCEINLINE FVector3f GetSkinnedBoneScale(FSkeletalMeshAccessorHelper& Accessor, int32 BoneIndex)
+	{
+		const int32 NumRealBones = Accessor.Mesh->GetRefSkeleton().GetRawBoneNum();
+		if (BoneIndex < NumRealBones)
+		{
+			return (FVector3f)Accessor.Mesh->GetComposedRefPoseMatrix(BoneIndex).ExtractScaling();
+		}
+
+		const FTransform& RefTransform = Accessor.Mesh->GetRefSkeleton().GetRefBonePose()[BoneIndex];
+		return (FVector3f)RefTransform.GetScale3D();
+	}
+
+	FORCEINLINE FVector3f GetSkinnedBonePreviousScale(FSkeletalMeshAccessorHelper& Accessor, int32 BoneIndex)
+	{
+		return GetSkinnedBoneScale(Accessor, BoneIndex);
+	}
 };
 
 template<>
@@ -243,7 +262,7 @@ struct FSkinnedPositionAccessorHelper<TNDISkelMesh_SkinningModeNone>
 		GetSkeletalMeshRefTangentBasis(Accessor.Mesh, *Accessor.LODData, *Accessor.SkinWeightBuffer, VertexIndex, OutTangentX, OutTangentY, OutTangentZ);
 	}
 
-	FORCEINLINE_DEBUGGABLE FVector3f GetSkinnedBonePosition(FSkeletalMeshAccessorHelper& Accessor, int32 BoneIndex)
+	FORCEINLINE FVector3f GetSkinnedBonePosition(FSkeletalMeshAccessorHelper& Accessor, int32 BoneIndex)
 	{
 		const int32 NumRealBones = Accessor.Mesh->GetRefSkeleton().GetRawBoneNum();
 		if (BoneIndex < NumRealBones)
@@ -255,12 +274,12 @@ struct FSkinnedPositionAccessorHelper<TNDISkelMesh_SkinningModeNone>
 		return (FVector3f)RefTransform.GetLocation();
 	}
 
-	FORCEINLINE_DEBUGGABLE FVector3f GetSkinnedBonePreviousPosition(FSkeletalMeshAccessorHelper& Accessor, int32 BoneIndex)
+	FORCEINLINE FVector3f GetSkinnedBonePreviousPosition(FSkeletalMeshAccessorHelper& Accessor, int32 BoneIndex)
 	{
 		return GetSkinnedBonePosition(Accessor, BoneIndex);
 	}
 
-	FORCEINLINE_DEBUGGABLE FQuat4f GetSkinnedBoneRotation(FSkeletalMeshAccessorHelper& Accessor, int32 BoneIndex)
+	FORCEINLINE FQuat4f GetSkinnedBoneRotation(FSkeletalMeshAccessorHelper& Accessor, int32 BoneIndex)
 	{
 		const int32 NumRealBones = Accessor.Mesh->GetRefSkeleton().GetRawBoneNum();
 		if (BoneIndex < NumRealBones)
@@ -272,10 +291,27 @@ struct FSkinnedPositionAccessorHelper<TNDISkelMesh_SkinningModeNone>
 		return (FQuat4f)RefTransform.GetRotation();
 	}
 
-	FORCEINLINE_DEBUGGABLE FQuat4f GetSkinnedBonePreviousRotation(FSkeletalMeshAccessorHelper& Accessor, int32 BoneIndex)
+	FORCEINLINE FQuat4f GetSkinnedBonePreviousRotation(FSkeletalMeshAccessorHelper& Accessor, int32 BoneIndex)
 	{
 		return GetSkinnedBoneRotation(Accessor, BoneIndex);
 	}
+
+	FORCEINLINE FVector3f GetSkinnedBoneScale(FSkeletalMeshAccessorHelper& Accessor, int32 BoneIndex)
+	{
+		const int32 NumRealBones = Accessor.Mesh->GetRefSkeleton().GetRawBoneNum();
+		if (BoneIndex < NumRealBones)
+		{
+			return (FVector3f)Accessor.Mesh->GetComposedRefPoseMatrix(BoneIndex).ExtractScaling();
+		}
+
+		const FTransform& RefTransform = Accessor.Mesh->GetRefSkeleton().GetRefBonePose()[BoneIndex];
+		return (FVector3f)RefTransform.GetScale3D();
+	}
+
+	FORCEINLINE FVector3f GetSkinnedBonePreviousScale(FSkeletalMeshAccessorHelper& Accessor, int32 BoneIndex)
+	{
+		return GetSkinnedBoneScale(Accessor, BoneIndex);
+	}
 };
 
 template<>
@@ -334,25 +370,35 @@ struct FSkinnedPositionAccessorHelper<TNDISkelMesh_SkinningModeOnTheFly>
 		USkeletalMeshComponent::GetSkinnedTangentBasis(Accessor.Comp, VertexIndex, *Accessor.LODData, *Accessor.SkinWeightBuffer, Accessor.SkinningData->PrevBoneRefToLocals(), OutTangentX, OutTangentY, OutTangentZ);
 	}
 
-	FORCEINLINE_DEBUGGABLE FVector3f GetSkinnedBonePosition(FSkeletalMeshAccessorHelper& Accessor, int32 BoneIndex)
+	FORCEINLINE FVector3f GetSkinnedBonePosition(FSkeletalMeshAccessorHelper& Accessor, int32 BoneIndex)
 	{
 		return (FVector3f)Accessor.SkinningData->CurrComponentTransforms()[BoneIndex].GetLocation();
 	}
 
-	FORCEINLINE_DEBUGGABLE FVector3f GetSkinnedBonePreviousPosition(FSkeletalMeshAccessorHelper& Accessor, int32 BoneIndex)
+	FORCEINLINE FVector3f GetSkinnedBonePreviousPosition(FSkeletalMeshAccessorHelper& Accessor, int32 BoneIndex)
 	{
 		return (FVector3f)Accessor.SkinningData->PrevComponentTransforms()[BoneIndex].GetLocation();
 	}
 
-	FORCEINLINE_DEBUGGABLE FQuat4f GetSkinnedBoneRotation(FSkeletalMeshAccessorHelper& Accessor, int32 BoneIndex)
+	FORCEINLINE FQuat4f GetSkinnedBoneRotation(FSkeletalMeshAccessorHelper& Accessor, int32 BoneIndex)
 	{
 		return (FQuat4f)Accessor.SkinningData->CurrComponentTransforms()[BoneIndex].GetRotation();
 	}
 
-	FORCEINLINE_DEBUGGABLE FQuat4f GetSkinnedBonePreviousRotation(FSkeletalMeshAccessorHelper& Accessor, int32 BoneIndex)
+	FORCEINLINE FQuat4f GetSkinnedBonePreviousRotation(FSkeletalMeshAccessorHelper& Accessor, int32 BoneIndex)
 	{
 		return (FQuat4f)Accessor.SkinningData->PrevComponentTransforms()[BoneIndex].GetRotation();
 	}
+
+	FORCEINLINE FVector3f GetSkinnedBoneScale(FSkeletalMeshAccessorHelper& Accessor, int32 BoneIndex)
+	{
+		return (FVector3f)Accessor.SkinningData->CurrComponentTransforms()[BoneIndex].GetScale3D();
+	}
+
+	FORCEINLINE FVector3f GetSkinnedBonePreviousScale(FSkeletalMeshAccessorHelper& Accessor, int32 BoneIndex)
+	{
+		return (FVector3f)Accessor.SkinningData->PrevComponentTransforms()[BoneIndex].GetScale3D();
+	}
 };
 
 template<>
@@ -411,25 +457,35 @@ struct FSkinnedPositionAccessorHelper<TNDISkelMesh_SkinningModePreSkin>
 		Accessor.SkinningData->GetPreviousTangentBasis(Accessor.Usage.GetLODIndex(), VertexIndex, OutTangentX, OutTangentY, OutTangentZ);
 	}
 
-	FORCEINLINE_DEBUGGABLE FVector3f GetSkinnedBonePosition(FSkeletalMeshAccessorHelper& Accessor, int32 BoneIndex)
+	FORCEINLINE FVector3f GetSkinnedBonePosition(FSkeletalMeshAccessorHelper& Accessor, int32 BoneIndex)
 	{
 		return Accessor.SkinningData->CurrComponentTransforms()[BoneIndex].GetLocation();
 	}
 
-	FORCEINLINE_DEBUGGABLE FVector3f GetSkinnedBonePreviousPosition(FSkeletalMeshAccessorHelper& Accessor, int32 BoneIndex)
+	FORCEINLINE FVector3f GetSkinnedBonePreviousPosition(FSkeletalMeshAccessorHelper& Accessor, int32 BoneIndex)
 	{
 		return Accessor.SkinningData->PrevComponentTransforms()[BoneIndex].GetLocation();
 	}
 
-	FORCEINLINE_DEBUGGABLE FQuat4f GetSkinnedBoneRotation(FSkeletalMeshAccessorHelper& Accessor, int32 BoneIndex)
+	FORCEINLINE FQuat4f GetSkinnedBoneRotation(FSkeletalMeshAccessorHelper& Accessor, int32 BoneIndex)
 	{
 		return (FQuat4f)Accessor.SkinningData->CurrComponentTransforms()[BoneIndex].GetRotation();
 	}
 
-	FORCEINLINE_DEBUGGABLE FQuat4f GetSkinnedBonePreviousRotation(FSkeletalMeshAccessorHelper& Accessor, int32 BoneIndex)
+	FORCEINLINE FQuat4f GetSkinnedBonePreviousRotation(FSkeletalMeshAccessorHelper& Accessor, int32 BoneIndex)
 	{
 		return (FQuat4f)Accessor.SkinningData->PrevComponentTransforms()[BoneIndex].GetRotation();
 	}
+
+	FORCEINLINE FVector3f GetSkinnedBoneScale(FSkeletalMeshAccessorHelper& Accessor, int32 BoneIndex)
+	{
+		return (FVector3f)Accessor.SkinningData->CurrComponentTransforms()[BoneIndex].GetScale3D();
+	}
+
+	FORCEINLINE FVector3f GetSkinnedBonePreviousScale(FSkeletalMeshAccessorHelper& Accessor, int32 BoneIndex)
+	{
+		return (FVector3f)Accessor.SkinningData->PrevComponentTransforms()[BoneIndex].GetScale3D();
+	}
 };
 
 //////////////////////////////////////////////////////////////////////////
diff --git a/Engine/Plugins/FX/Niagara/Source/Niagara/Classes/NiagaraDataInterface.h b/Engine/Plugins/FX/Niagara/Source/Niagara/Classes/NiagaraDataInterface.h
--- a/Engine/Plugins/FX/Niagara/Source/Niagara/Classes/NiagaraDataInterface.h
+++ b/Engine/Plugins/FX/Niagara/Source/Niagara/Classes/NiagaraDataInterface.h
@@ -25,6 +25,8 @@ struct FNDITransformHandlerNoop
 	FORCEINLINE void TransformPosition(FVector3d& V, const FMatrix44d& M) const { }
 	FORCEINLINE void TransformVector(FVector3f& V, const FMatrix44f& M) const { }
 	FORCEINLINE void TransformVector(FVector3d& V, const FMatrix44d& M) const { }
+	FORCEINLINE void TransformNotUnitVector(FVector3f& V, const FMatrix44f& M) const { }
+	FORCEINLINE void TransformNotUnitVector(FVector3d& V, const FMatrix44d& M) const { }
 	FORCEINLINE void TransformRotation(FQuat4f& Q1, const FQuat4f& Q2) const { }
 	FORCEINLINE void TransformRotation(FQuat4d& Q1, const FQuat4d& Q2) const { }
 };
@@ -35,6 +37,8 @@ struct FNDITransformHandler
 	FORCEINLINE void TransformPosition(FVector3d& P, const FMatrix44d& M) const { P = M.TransformPosition(P); }
 	FORCEINLINE void TransformVector(FVector3f& V, const FMatrix44f& M) const { V = M.TransformVector(V).GetUnsafeNormal3(); }
 	FORCEINLINE void TransformVector(FVector3d& V, const FMatrix44d& M) const { V = M.TransformVector(V).GetUnsafeNormal3(); }
+	FORCEINLINE void TransformNotUnitVector(FVector3f& V, const FMatrix44f& M) const { V = M.TransformVector(V); }
+	FORCEINLINE void TransformNotUnitVector(FVector3d& V, const FMatrix44d& M) const { V = M.TransformVector(V); }
 	FORCEINLINE void TransformRotation(FQuat4f& Q1, const FQuat4f& Q2) const { Q1 = Q2 * Q1; }
 	FORCEINLINE void TransformRotation(FQuat4d& Q1, const FQuat4d& Q2) const { Q1 = Q2 * Q1; }
 };
diff --git a/Engine/Plugins/FX/Niagara/Source/Niagara/Private/NDISkeletalMesh_BoneSampling.cpp b/Engine/Plugins/FX/Niagara/Source/Niagara/Private/NDISkeletalMesh_BoneSampling.cpp
--- a/Engine/Plugins/FX/Niagara/Source/Niagara/Private/NDISkeletalMesh_BoneSampling.cpp
+++ b/Engine/Plugins/FX/Niagara/Source/Niagara/Private/NDISkeletalMesh_BoneSampling.cpp
@@ -56,6 +56,7 @@ void UNiagaraDataInterfaceSkeletalMesh::GetSkeletonSamplingFunctions(TArray<FNia
 		Sig.Inputs.Add(FNiagaraVariable(FNiagaraTypeDefinition::GetIntDef(), TEXT("Bone")));
 		Sig.Outputs.Add(FNiagaraVariable(FNiagaraTypeDefinition::GetPositionDef(), TEXT("Position")));
 		Sig.Outputs.Add(FNiagaraVariable(FNiagaraTypeDefinition::GetQuatDef(), TEXT("Rotation")));
+		Sig.Outputs.Add(FNiagaraVariable(FNiagaraTypeDefinition::GetVec3Def(), TEXT("Scale")));
 		Sig.Outputs.Add(FNiagaraVariable(FNiagaraTypeDefinition::GetVec3Def(), TEXT("Velocity")));
 		Sig.bMemberFunction = true;
 		Sig.bRequiresContext = false;
@@ -72,6 +73,7 @@ void UNiagaraDataInterfaceSkeletalMesh::GetSkeletonSamplingFunctions(TArray<FNia
 		Sig.Inputs.Add(FNiagaraVariable(FNiagaraTypeDefinition::GetIntDef(), TEXT("Bone")));
 		Sig.Outputs.Add(FNiagaraVariable(FNiagaraTypeDefinition::GetPositionDef(), TEXT("Position")));
 		Sig.Outputs.Add(FNiagaraVariable(FNiagaraTypeDefinition::GetQuatDef(), TEXT("Rotation")));
+		Sig.Outputs.Add(FNiagaraVariable(FNiagaraTypeDefinition::GetVec3Def(), TEXT("Scale")));
 		Sig.Outputs.Add(FNiagaraVariable(FNiagaraTypeDefinition::GetVec3Def(), TEXT("Velocity")));
 		Sig.bMemberFunction = true;
 		Sig.bRequiresContext = false;
@@ -90,6 +92,7 @@ void UNiagaraDataInterfaceSkeletalMesh::GetSkeletonSamplingFunctions(TArray<FNia
 		Sig.Inputs.Add(FNiagaraVariable(FNiagaraTypeDefinition::GetFloatDef(), TEXT("Interpolation")));
 		Sig.Outputs.Add(FNiagaraVariable(FNiagaraTypeDefinition::GetPositionDef(), TEXT("Position")));
 		Sig.Outputs.Add(FNiagaraVariable(FNiagaraTypeDefinition::GetQuatDef(), TEXT("Rotation")));
+		Sig.Outputs.Add(FNiagaraVariable(FNiagaraTypeDefinition::GetVec3Def(), TEXT("Scale")));
 		Sig.Outputs.Add(FNiagaraVariable(FNiagaraTypeDefinition::GetVec3Def(), TEXT("Velocity")));
 		Sig.bMemberFunction = true;
 		Sig.bRequiresContext = false;
@@ -107,6 +110,7 @@ void UNiagaraDataInterfaceSkeletalMesh::GetSkeletonSamplingFunctions(TArray<FNia
 		Sig.Inputs.Add(FNiagaraVariable(FNiagaraTypeDefinition::GetFloatDef(), TEXT("Interpolation")));
 		Sig.Outputs.Add(FNiagaraVariable(FNiagaraTypeDefinition::GetPositionDef(), TEXT("Position")));
 		Sig.Outputs.Add(FNiagaraVariable(FNiagaraTypeDefinition::GetQuatDef(), TEXT("Rotation")));
+		Sig.Outputs.Add(FNiagaraVariable(FNiagaraTypeDefinition::GetVec3Def(), TEXT("Scale")));
 		Sig.Outputs.Add(FNiagaraVariable(FNiagaraTypeDefinition::GetVec3Def(), TEXT("Velocity")));
 		Sig.bMemberFunction = true;
 		Sig.bRequiresContext = false;
@@ -352,22 +356,22 @@ void UNiagaraDataInterfaceSkeletalMesh::BindSkeletonSamplingFunction(const FVMEx
 	//Bone Functions
 	if (BindingInfo.Name == FSkeletalMeshInterfaceHelper::GetSkinnedBoneDataName)
 	{
-		ensure(BindingInfo.GetNumInputs() == 2 && BindingInfo.GetNumOutputs() == 10);
+		ensure(BindingInfo.GetNumInputs() == 2 && BindingInfo.GetNumOutputs() == 13);
 		TSkinningModeBinder<TNDIExplicitBinder<FNDITransformHandlerNoop, TNDIExplicitBinder<TInterpOff, NDI_FUNC_BINDER(UNiagaraDataInterfaceSkeletalMesh, GetSkinnedBoneData)>>>::BindIgnoreCPUAccess(this, BindingInfo, InstanceData, OutFunc);
 	}
 	else if (BindingInfo.Name == FSkeletalMeshInterfaceHelper::GetSkinnedBoneDataWSName)
 	{
-		ensure(BindingInfo.GetNumInputs() == 2 && BindingInfo.GetNumOutputs() == 10);
+		ensure(BindingInfo.GetNumInputs() == 2 && BindingInfo.GetNumOutputs() == 13);
 		TSkinningModeBinder<TNDIExplicitBinder<FNDITransformHandler, TNDIExplicitBinder<TInterpOff, NDI_FUNC_BINDER(UNiagaraDataInterfaceSkeletalMesh, GetSkinnedBoneData)>>>::BindIgnoreCPUAccess(this, BindingInfo, InstanceData, OutFunc);
 	}
 	else if (BindingInfo.Name == FSkeletalMeshInterfaceHelper::GetSkinnedBoneDataInterpolatedName)
 	{
-		ensure(BindingInfo.GetNumInputs() == 3 && BindingInfo.GetNumOutputs() == 10);
+		ensure(BindingInfo.GetNumInputs() == 3 && BindingInfo.GetNumOutputs() == 13);
 		TSkinningModeBinder<TNDIExplicitBinder<FNDITransformHandlerNoop, TNDIExplicitBinder<TInterpOn, NDI_FUNC_BINDER(UNiagaraDataInterfaceSkeletalMesh, GetSkinnedBoneData)>>>::BindIgnoreCPUAccess(this, BindingInfo, InstanceData, OutFunc);
 	}
 	else if (BindingInfo.Name == FSkeletalMeshInterfaceHelper::GetSkinnedBoneDataWSInterpolatedName)
 	{
-		ensure(BindingInfo.GetNumInputs() == 3 && BindingInfo.GetNumOutputs() == 10);
+		ensure(BindingInfo.GetNumInputs() == 3 && BindingInfo.GetNumOutputs() == 13);
 		TSkinningModeBinder<TNDIExplicitBinder<FNDITransformHandler, TNDIExplicitBinder<TInterpOn, NDI_FUNC_BINDER(UNiagaraDataInterfaceSkeletalMesh, GetSkinnedBoneData)>>>::BindIgnoreCPUAccess(this, BindingInfo, InstanceData, OutFunc);
 	}
 	else if (BindingInfo.Name == FSkeletalMeshInterfaceHelper::IsValidBoneName)
@@ -752,21 +756,25 @@ struct FBoneSocketSkinnedDataOutputHandler
 	FBoneSocketSkinnedDataOutputHandler(FVectorVMExternalFunctionContext& Context)
 		: Position(Context)
 		, Rotation(Context)
+		, Scale(Context)
 		, Velocity(Context)
 		, bNeedsPosition(Position.IsValid())
 		, bNeedsRotation(Rotation.IsValid())
+		, bNeedsScale(Scale.IsValid())
 		, bNeedsVelocity(Velocity.IsValid())
 	{
 	}
 
 	FNDIOutputParam<FVector3f> Position;
 	FNDIOutputParam<FQuat4f> Rotation;
+	FNDIOutputParam<FVector3f> Scale;
 	FNDIOutputParam<FVector3f> Velocity;
 
 	//TODO: Rotation + Scale too? Use quats so we can get proper interpolation between bone and parent.
 
 	const bool bNeedsPosition;
 	const bool bNeedsRotation;
+	const bool bNeedsScale;
 	const bool bNeedsVelocity;
 };
 
@@ -834,6 +842,10 @@ void UNiagaraDataInterfaceSkeletalMesh::GetSkinnedBoneData(FVectorVMExternalFunc
 				{
 					Output.Rotation.SetAndAdvance(FQuat4f::Identity);
 				}
+				if (Output.bNeedsScale)
+				{
+					Output.Scale.SetAndAdvance(FVector3f::ZeroVector);
+				}
 			}
 			else if (bIsSocket)
 			{
@@ -862,6 +874,18 @@ void UNiagaraDataInterfaceSkeletalMesh::GetSkinnedBoneData(FVectorVMExternalFunc
 
 					Output.Rotation.SetAndAdvance(Rotation);
 				}
+				if (Output.bNeedsScale)
+				{
+					FVector3f Scale = CurrSocketTransform.GetScale3D();
+					TransformHandler.TransformVector(Scale, InstanceTransform);
+					if (bInterpolated::Value)
+					{
+						FVector3f PrevScale = PrevSocketTransform.GetScale3D();
+						TransformHandler.TransformVector(PrevScale, PrevInstanceTransform);
+						Scale = FMath::Lerp(PrevScale, Scale, Interp);
+					}
+					Output.Scale.SetAndAdvance(Scale);
+				}
 			}
 			// Bone
 			else
@@ -888,6 +912,19 @@ void UNiagaraDataInterfaceSkeletalMesh::GetSkinnedBoneData(FVectorVMExternalFunc
 
 					Output.Rotation.SetAndAdvance(Rotation);
 				}
+
+				if (Output.bNeedsScale)
+				{
+					FVector3f Scale = SkinningHandler.GetSkinnedBoneScale(Accessor, Bone);
+					TransformHandler.TransformNotUnitVector(Scale, InstanceTransform);
+					if (bInterpolated::Value)
+					{
+						FVector3f PrevScale = SkinningHandler.GetSkinnedBonePreviousScale(Accessor, Bone);
+						TransformHandler.TransformNotUnitVector(PrevScale, PrevInstanceTransform);
+						Scale = FMath::Lerp(PrevScale, Scale, Interp);
+					}
+					Output.Scale.SetAndAdvance(Scale);
+				}
 			}
 
 			if (Output.bNeedsVelocity || bInterpolated::Value)
diff --git a/Engine/Plugins/FX/Niagara/Source/Niagara/Private/NiagaraDataInterfaceSkeletalMesh.cpp b/Engine/Plugins/FX/Niagara/Source/Niagara/Private/NiagaraDataInterfaceSkeletalMesh.cpp
--- a/Engine/Plugins/FX/Niagara/Source/Niagara/Private/NiagaraDataInterfaceSkeletalMesh.cpp
+++ b/Engine/Plugins/FX/Niagara/Source/Niagara/Private/NiagaraDataInterfaceSkeletalMesh.cpp
@@ -42,6 +42,7 @@ struct FNiagaraSkelMeshDIFunctionVersion
 		AddedInputBardCoordToGetFilteredTriangleAt = 8,
 		LargeWorldCoordinates = 9,
 		LargeWorldCoordinates2 = 10,
+		AddBoneScale = 11,
 
 		VersionPlusOne,
 		LatestVersion = VersionPlusOne - 1
@@ -3145,22 +3146,22 @@ bool UNiagaraDataInterfaceSkeletalMesh::GetFunctionHLSL(const FNiagaraDataInterf
 	// Bone Sampling
 	else if (FunctionInfo.DefinitionName == FSkeletalMeshInterfaceHelper::GetSkinnedBoneDataName)
 	{
-		static const TCHAR* FormatSample = TEXT("void {InstanceFunctionName} (in int InBone, out float3 OutPosition, out float4 OutRotation, out float3 OutVelocity) { {GetDISkelMeshContextName} DISkelMesh_GetSkinnedBone(DIContext, InBone, OutPosition, OutRotation, OutVelocity); }");
+		static const TCHAR* FormatSample = TEXT("void {InstanceFunctionName} (in int InBone, out float3 OutPosition, out float4 OutRotation, out float3 OutScale, out float3 OutVelocity) { {GetDISkelMeshContextName} DISkelMesh_GetSkinnedBone(DIContext, InBone, OutPosition, OutRotation, OutScale, OutVelocity); }");
 		OutHLSL += FString::Format(FormatSample, ArgsSample);
 	}
 	else if (FunctionInfo.DefinitionName == FSkeletalMeshInterfaceHelper::GetSkinnedBoneDataInterpolatedName)
 	{
-		static const TCHAR* FormatSample = TEXT("void {InstanceFunctionName} (in int InBone, in float Interp, out float3 OutPosition, out float4 OutRotation, out float3 OutVelocity) { {GetDISkelMeshContextName} DISkelMesh_GetSkinnedBoneInterpolated(DIContext, InBone, Interp, OutPosition, OutRotation, OutVelocity); }");
+		static const TCHAR* FormatSample = TEXT("void {InstanceFunctionName} (in int InBone, in float Interp, out float3 OutPosition, out float4 OutRotation, out float3 OutScale, out float3 OutVelocity) { {GetDISkelMeshContextName} DISkelMesh_GetSkinnedBoneInterpolated(DIContext, InBone, Interp, OutPosition, OutRotation, OutScale, OutVelocity); }");
 		OutHLSL += FString::Format(FormatSample, ArgsSample);
 	}
 	else if (FunctionInfo.DefinitionName == FSkeletalMeshInterfaceHelper::GetSkinnedBoneDataWSName)
 	{
-		static const TCHAR* FormatSample = TEXT("void {InstanceFunctionName} (in int InBone, out float3 OutPosition, out float4 OutRotation, out float3 OutVelocity) { {GetDISkelMeshContextName} DISkelMesh_GetSkinnedBoneWS(DIContext, InBone, OutPosition, OutRotation, OutVelocity); }");
+		static const TCHAR* FormatSample = TEXT("void {InstanceFunctionName} (in int InBone, out float3 OutPosition, out float4 OutRotation, out float3 OutScale, out float3 OutVelocity) { {GetDISkelMeshContextName} DISkelMesh_GetSkinnedBoneWS(DIContext, InBone, OutPosition, OutRotation, OutScale, OutVelocity); }");
 		OutHLSL += FString::Format(FormatSample, ArgsSample);
 	}
 	else if (FunctionInfo.DefinitionName == FSkeletalMeshInterfaceHelper::GetSkinnedBoneDataWSInterpolatedName)
 	{
-		static const TCHAR* FormatSample = TEXT("void {InstanceFunctionName} (in int InBone, in float Interp, out float3 OutPosition, out float4 OutRotation, out float3 OutVelocity) { {GetDISkelMeshContextName} DISkelMesh_GetSkinnedBoneInterpolatedWS(DIContext, InBone, Interp, OutPosition, OutRotation, OutVelocity); }");
+		static const TCHAR* FormatSample = TEXT("void {InstanceFunctionName} (in int InBone, in float Interp, out float3 OutPosition, out float4 OutRotation, out float3 OutScale, out float3 OutVelocity) { {GetDISkelMeshContextName} DISkelMesh_GetSkinnedBoneInterpolatedWS(DIContext, InBone, Interp, OutPosition, OutRotation, OutScale, OutVelocity); }");
 		OutHLSL += FString::Format(FormatSample, ArgsSample);
 	}
 	//////////////////////////////////////////////////////////////////////////////////////////////////////////////////////////////
@@ -3544,6 +3545,26 @@ bool UNiagaraDataInterfaceSkeletalMesh::UpgradeFunctionCall(FNiagaraFunctionSign
 		}
 	}
 
+	if (FunctionSignature.FunctionVersion < FNiagaraSkelMeshDIFunctionVersion::AddBoneScale)
+	{
+		if (
+			(FunctionSignature.Name == FSkeletalMeshInterfaceHelper::GetSkinnedBoneDataName) ||
+			(FunctionSignature.Name == FSkeletalMeshInterfaceHelper::GetSkinnedBoneDataWSName) ||
+			(FunctionSignature.Name == FSkeletalMeshInterfaceHelper::GetSkinnedBoneDataInterpolatedName) ||
+			(FunctionSignature.Name == FSkeletalMeshInterfaceHelper::GetSkinnedBoneDataWSInterpolatedName) )
+		{
+			// Some bogus content appears to exist which only contains Position & Velocity outputs
+			if ((FunctionSignature.Outputs.Num() <= 2) && (FunctionSignature.Outputs[1].GetName() == TEXT("Velocity")))
+			{
+				FunctionSignature.Outputs.EmplaceAt(1, FNiagaraTypeDefinition::GetQuatDef(), TEXT("Rotation"));
+			}
+
+			check(FunctionSignature.Outputs[2].GetName() == TEXT("Velocity"));
+			FunctionSignature.Outputs.EmplaceAt(2, FNiagaraTypeDefinition::GetVec3Def(), TEXT("Scale"));
+			bWasChanged = true;
+		}
+	}
+
 	// Set latest version
 	FunctionSignature.FunctionVersion = FNiagaraSkelMeshDIFunctionVersion::LatestVersion;
 
