- Adjust collision stat to be the asset name when stats are disabled (i.e. test builds) as it makes it easier to track who issued the collision

diff --git a/Engine/Plugins/FX/Niagara/Source/Niagara/Classes/NiagaraCollision.h b/Engine/Plugins/FX/Niagara/Source/Niagara/Classes/NiagaraCollision.h
--- a/Engine/Plugins/FX/Niagara/Source/Niagara/Classes/NiagaraCollision.h
+++ b/Engine/Plugins/FX/Niagara/Source/Niagara/Classes/NiagaraCollision.h
@@ -3,8 +3,11 @@
 #pragma once
 
 #include "NiagaraEvents.h"
+#include "NiagaraStats.h"
 #include "WorldCollision.h"
 
+DECLARE_CYCLE_STAT(TEXT("Collision"), STAT_NiagaraCollision, STATGROUP_Niagara);
+
 class FNiagaraDataSet;
 
 UENUM()
@@ -78,10 +81,11 @@ public:
 	}
 
 
-	void Init(FNiagaraSystemInstanceID InBatchID, UWorld *InCollisionWorld)
+	void Init(FNiagaraSystemInstanceID InBatchID, UWorld *InCollisionWorld, TStatId InStatId)
 	{
 		BatchID = InBatchID;
 		CollisionWorld = InCollisionWorld;
+		StatId = InStatId;
 		CollisionTraces[0].Empty();
 		CollisionTraces[1].Empty();
 		CurrBuffer = 0;
@@ -99,6 +103,7 @@ private:
 
 	//TArray<FTraceHandle> CollisionTraceHandles;
 	const static FName CollisionTagName;
+	const static FName TraceTagName;
 
 	FRWLock CollisionTraceLock;
 	TArray<FNiagaraCollisionEventPayload> CollisionEvents;
@@ -109,4 +114,5 @@ private:
 	TArray<FNiagaraDICollsionQueryResult> CollisionResults;
 	uint32 CurrBuffer;
 	UWorld *CollisionWorld = nullptr;
-};
\ No newline at end of file
+	TStatId StatId;
+};
diff --git a/Engine/Plugins/FX/Niagara/Source/Niagara/Private/NiagaraCollision.cpp b/Engine/Plugins/FX/Niagara/Source/Niagara/Private/NiagaraCollision.cpp
--- a/Engine/Plugins/FX/Niagara/Source/Niagara/Private/NiagaraCollision.cpp
+++ b/Engine/Plugins/FX/Niagara/Source/Niagara/Private/NiagaraCollision.cpp
@@ -3,14 +3,13 @@
 #include "NiagaraCollision.h"
 #include "NiagaraEmitterInstance.h"
 #include "NiagaraSystemInstance.h"
-#include "NiagaraStats.h"
 #include "NiagaraComponent.h"
 #include "PhysicalMaterials/PhysicalMaterial.h"
 
-DECLARE_CYCLE_STAT(TEXT("Collision"), STAT_NiagaraCollision, STATGROUP_Niagara);
 DECLARE_CYCLE_STAT(TEXT("Event Emission"), STAT_NiagaraEventWrite, STATGROUP_Niagara);
 
 const FName FNiagaraDICollisionQueryBatch::CollisionTagName = FName("Niagara");
+const FName FNiagaraDICollisionQueryBatch::TraceTagName = FName("NiagaraCollision");
 
 static int32 GNiagaraCollisionCPUEnabled = 1;
 static FAutoConsoleVariableRef CVarNiagaraCollisionEnabled(
@@ -113,7 +112,7 @@ int32 FNiagaraDICollisionQueryBatch::SubmitQuery(FVector StartPos, FVector EndPo
 		return INDEX_NONE;
 	}
 
-	FCollisionQueryParams CollisionQueryParams(SCENE_QUERY_STAT(NiagaraCollision), false);
+	FCollisionQueryParams CollisionQueryParams(TraceTagName, StatId, false);
 	CollisionQueryParams.OwnerTag = CollisionTagName;
 	CollisionQueryParams.bFindInitialOverlaps = false;
 	CollisionQueryParams.bReturnFaceIndex = false;
@@ -143,7 +142,7 @@ bool FNiagaraDICollisionQueryBatch::PerformQuery(FVector StartPos, FVector EndPo
 	}
 
 	SCOPE_CYCLE_COUNTER(STAT_NiagaraCollision);
-	FCollisionQueryParams QueryParams(SCENE_QUERY_STAT(NiagaraCollision), false);
+	FCollisionQueryParams QueryParams(TraceTagName, StatId, false);
 	QueryParams.OwnerTag = CollisionTagName;
 	QueryParams.bFindInitialOverlaps = false;
 	QueryParams.bReturnFaceIndex = false;
diff --git a/Engine/Plugins/FX/Niagara/Source/Niagara/Private/NiagaraDataInterfaceCollisionQuery.cpp b/Engine/Plugins/FX/Niagara/Source/Niagara/Private/NiagaraDataInterfaceCollisionQuery.cpp
--- a/Engine/Plugins/FX/Niagara/Source/Niagara/Private/NiagaraDataInterfaceCollisionQuery.cpp
+++ b/Engine/Plugins/FX/Niagara/Source/Niagara/Private/NiagaraDataInterfaceCollisionQuery.cpp
@@ -82,7 +82,18 @@ bool UNiagaraDataInterfaceCollisionQuery::InitPerInstanceData(void* PerInstanceD
 	PIData->SystemInstance = InSystemInstance;
 	if (InSystemInstance)
 	{
-		PIData->CollisionBatch.Init(InSystemInstance->GetId(), InSystemInstance->GetWorld());
+		TStatId CollisionStatId = SCENE_QUERY_STAT_ONLY(NiagaraCollision);
+
+		// If we don't have stats enabled can we pull a valid stat from the asset?
+		// This will allow us to trace back to the asset when profiling
+		#if !STATS && (ENABLE_STATNAMEDEVENTS && ENABLE_STATNAMEDEVENTS_UOBJECT)
+			TStatId AssetStatId = static_cast<const UObjectBaseUtility*>(InSystemInstance->GetSystem())->GetStatID();
+			if (AssetStatId.IsValidStat())
+			{
+				CollisionStatId = AssetStatId;
+			}
+		#endif
+		PIData->CollisionBatch.Init(InSystemInstance->GetId(), InSystemInstance->GetWorld(), CollisionStatId);
 	}
 
 	return true;
