Fixed a crash when calling SetVector2DParameter on a Niagara Parameter Collection

diff --git a/Engine/Plugins/FX/Niagara/Source/Niagara/Private/NiagaraParameterCollection.cpp b/Engine/Plugins/FX/Niagara/Source/Niagara/Private/NiagaraParameterCollection.cpp
--- a/Engine/Plugins/FX/Niagara/Source/Niagara/Private/NiagaraParameterCollection.cpp
+++ b/Engine/Plugins/FX/Niagara/Source/Niagara/Private/NiagaraParameterCollection.cpp
@@ -283,23 +283,23 @@ int32 UNiagaraParameterCollectionInstance::GetIntParameter(const FString& InVari
 
 FVector2D UNiagaraParameterCollectionInstance::GetVector2DParameter(const FString& InVariableName)
 {
-	return (FVector2D)ParameterStorage.GetParameterValue<FVector2f>(FNiagaraVariable(FNiagaraTypeDefinition::GetVec2Def(), *Collection->ParameterNameFromFriendlyName(InVariableName)));
+	return FVector2D(ParameterStorage.GetParameterValue<FVector2f>(FNiagaraVariable(FNiagaraTypeDefinition::GetVec2Def(), *Collection->ParameterNameFromFriendlyName(InVariableName))));
 }
 
 FVector UNiagaraParameterCollectionInstance::GetVectorParameter(const FString& InVariableName)
 {
-	return (FVector)ParameterStorage.GetParameterValue<FVector3f>(FNiagaraVariable(FNiagaraTypeDefinition::GetVec3Def(), *Collection->ParameterNameFromFriendlyName(InVariableName)));
+	return FVector(ParameterStorage.GetParameterValue<FVector3f>(FNiagaraVariable(FNiagaraTypeDefinition::GetVec3Def(), *Collection->ParameterNameFromFriendlyName(InVariableName))));
 }
 
 FVector4 UNiagaraParameterCollectionInstance::GetVector4Parameter(const FString& InVariableName)
 {
-	return (FVector4)ParameterStorage.GetParameterValue<FVector4f>(FNiagaraVariable(FNiagaraTypeDefinition::GetVec4Def(), *Collection->ParameterNameFromFriendlyName(InVariableName)));
+	return FVector4(ParameterStorage.GetParameterValue<FVector4f>(FNiagaraVariable(FNiagaraTypeDefinition::GetVec4Def(), *Collection->ParameterNameFromFriendlyName(InVariableName))));
 }
 
 
 FQuat UNiagaraParameterCollectionInstance::GetQuatParameter(const FString& InVariableName)
 {
-	return (FQuat)ParameterStorage.GetParameterValue<FQuat4f>(FNiagaraVariable(FNiagaraTypeDefinition::GetQuatDef(), *Collection->ParameterNameFromFriendlyName(InVariableName)));
+	return FQuat(ParameterStorage.GetParameterValue<FQuat4f>(FNiagaraVariable(FNiagaraTypeDefinition::GetQuatDef(), *Collection->ParameterNameFromFriendlyName(InVariableName))));
 }
 
 FLinearColor UNiagaraParameterCollectionInstance::GetColorParameter(const FString& InVariableName)
@@ -365,7 +365,7 @@ void UNiagaraParameterCollectionInstance::SetVector2DParameter(const FString& In
 	const FName ParameterName = *Collection->ParameterNameFromFriendlyName(InVariableName);
 	if (!CheckConflictWithSourceMpc(ParameterName, __FUNCTION__, InValue, Collection))
 	{
-		ParameterStorage.SetParameterValue(InValue, FNiagaraVariable(FNiagaraTypeDefinition::GetVec2Def(), ParameterName));
+		ParameterStorage.SetParameterValue(FVector2f(InValue), FNiagaraVariable(FNiagaraTypeDefinition::GetVec2Def(), ParameterName));
 	}
 }
 
@@ -374,7 +374,7 @@ void UNiagaraParameterCollectionInstance::SetVectorParameter(const FString& InVa
 	const FName ParameterName = *Collection->ParameterNameFromFriendlyName(InVariableName);
 	if (!CheckConflictWithSourceMpc(ParameterName, __FUNCTION__, InValue, Collection))
 	{
-		ParameterStorage.SetParameterValue(InValue, FNiagaraVariable(FNiagaraTypeDefinition::GetVec3Def(), ParameterName));
+		ParameterStorage.SetParameterValue(FVector3f(InValue), FNiagaraVariable(FNiagaraTypeDefinition::GetVec3Def(), ParameterName));
 	}
 }
 
@@ -383,7 +383,7 @@ void UNiagaraParameterCollectionInstance::SetVector4Parameter(const FString& InV
 	const FName ParameterName = *Collection->ParameterNameFromFriendlyName(InVariableName);
 	if (!CheckConflictWithSourceMpc(ParameterName, __FUNCTION__, InValue, Collection))
 	{
-		ParameterStorage.SetParameterValue(InValue, FNiagaraVariable(FNiagaraTypeDefinition::GetVec4Def(), ParameterName));
+		ParameterStorage.SetParameterValue(FVector4f(InValue), FNiagaraVariable(FNiagaraTypeDefinition::GetVec4Def(), ParameterName));
 	}
 }
 
@@ -401,7 +401,7 @@ void UNiagaraParameterCollectionInstance::SetQuatParameter(const FString& InVari
 	const FName ParameterName = *Collection->ParameterNameFromFriendlyName(InVariableName);
 	if (!CheckConflictWithSourceMpc(ParameterName, __FUNCTION__, InValue, Collection))
 	{
-		ParameterStorage.SetParameterValue(InValue, FNiagaraVariable(FNiagaraTypeDefinition::GetQuatDef(), ParameterName));
+		ParameterStorage.SetParameterValue(FQuat4f(InValue), FNiagaraVariable(FNiagaraTypeDefinition::GetQuatDef(), ParameterName));
 	}
 }
 
