Fixed a false warning from the Niagara component pool

diff --git a/Engine/Plugins/FX/Niagara/Source/Niagara/Private/NiagaraComponentPool.cpp b/Engine/Plugins/FX/Niagara/Source/Niagara/Private/NiagaraComponentPool.cpp
--- a/Engine/Plugins/FX/Niagara/Source/Niagara/Private/NiagaraComponentPool.cpp
+++ b/Engine/Plugins/FX/Niagara/Source/Niagara/Private/NiagaraComponentPool.cpp
@@ -457,7 +457,8 @@ void UNiagaraComponentPool::ReclaimWorldParticleSystem(UNiagaraComponent* Compon
 		return;
 	}
 
-	if (GbEnableNiagaraSystemPooling && Asset != nullptr)
+	// WorldParticleSystemPools is empty after world cleanup, so destroy components coming in to be reclaimed instead
+	if (GbEnableNiagaraSystemPooling && Asset != nullptr && WorldParticleSystemPools.Num() > 0)
 	{
 		float CurrentTime = Component->GetWorld()->GetTimeSeconds();
 
