Made copy & paste of Niagara (custom hlsl) nodes more stable by encoding/decoding the exported strings to base64

diff --git a/Engine/Plugins/FX/Niagara/Source/NiagaraEditor/Private/ViewModels/NiagaraScriptGraphViewModel.cpp b/Engine/Plugins/FX/Niagara/Source/NiagaraEditor/Private/ViewModels/NiagaraScriptGraphViewModel.cpp
--- a/Engine/Plugins/FX/Niagara/Source/NiagaraEditor/Private/ViewModels/NiagaraScriptGraphViewModel.cpp
+++ b/Engine/Plugins/FX/Niagara/Source/NiagaraEditor/Private/ViewModels/NiagaraScriptGraphViewModel.cpp
@@ -270,9 +270,12 @@ void FNiagaraScriptGraphViewModel::CopySelectedNodes()
 			}
 		}
 	}
-	
-	FEdGraphUtilities::ExportNodesToText(CachedNodes, ClipboardContent->ExportedNodes);
 
+	FString ExportedNodesText;
+	FEdGraphUtilities::ExportNodesToText(CachedNodes, ExportedNodesText);
+
+	ClipboardContent->ExportedNodes = FBase64::Encode(ExportedNodesText);
+	
 	// we then attempt to find all variables referenced by those nodes to copy the script variables in our clipboard
 	TSet<FNiagaraVariable> Variables;
 	const TMap<FNiagaraVariable, FNiagaraGraphParameterReferenceCollection>& ParameterReferenceMap = GetGraph()->GetParameterReferenceMap();
@@ -338,7 +341,9 @@ void FNiagaraScriptGraphViewModel::PasteNodes()
 	}
 	
 	TSet<UEdGraphNode*> PastedNodes;
-	FEdGraphUtilities::ImportNodesFromText(Graph, ClipboardContent->ExportedNodes, PastedNodes);
+	FString ExportedNodesText;
+	FBase64::Decode(ClipboardContent->ExportedNodes, ExportedNodesText);
+	FEdGraphUtilities::ImportNodesFromText(Graph, ExportedNodesText, PastedNodes);
 	
 	for (UEdGraphNode* PastedNode : PastedNodes)
 	{
@@ -451,7 +456,9 @@ bool FNiagaraScriptGraphViewModel::CanPasteNodes() const
 	}
 
 	const UNiagaraClipboardContent* ClipboardContent = FNiagaraEditorModule::Get().GetClipboard().GetClipboardContent();
-	return ClipboardContent? FEdGraphUtilities::CanImportNodesFromText(Graph, ClipboardContent->ExportedNodes) : false;
+	FString ExportedNodesText;
+	FBase64::Decode(ClipboardContent->ExportedNodes, ExportedNodesText);
+	return ClipboardContent? FEdGraphUtilities::CanImportNodesFromText(Graph, ExportedNodesText) : false;
 }
 
 void FNiagaraScriptGraphViewModel::DuplicateNodes()
