Fixed specific parameters not being added in systems when using set parameter

diff --git a/Engine/Plugins/FX/Niagara/Source/NiagaraEditor/Private/ViewModels/NiagaraParameterPanelViewModel.cpp b/Engine/Plugins/FX/Niagara/Source/NiagaraEditor/Private/ViewModels/NiagaraParameterPanelViewModel.cpp
--- a/Engine/Plugins/FX/Niagara/Source/NiagaraEditor/Private/ViewModels/NiagaraParameterPanelViewModel.cpp
+++ b/Engine/Plugins/FX/Niagara/Source/NiagaraEditor/Private/ViewModels/NiagaraParameterPanelViewModel.cpp
@@ -1500,6 +1500,19 @@ TSharedRef<SWidget> FNiagaraSystemToolkitParameterPanelViewModel::CreateAddParam
 		AssignmentNode->AddParameter(NewParameter, VarDefaultValue);
 	};
 
+	auto AddExistingParameterLambda = [this, AssignmentNode](FNiagaraVariable& NewParameter) {
+		// If an assignment target is already setting the associated parameter, we simply don't do anything
+		if (AssignmentNode->GetAssignmentTargets().Contains(NewParameter))
+		{
+			return;
+		}
+
+		const FString VarDefaultValue = FNiagaraConstants::GetAttributeDefaultValue(NewParameter);
+		FNiagaraParameterPanelCategory TempCategory = FNiagaraParameterPanelCategory(FNiagaraEditorUtilities::GetNamespaceMetaDataForVariableName(NewParameter.GetName()));
+		FindOrAddParameter(NewParameter, TempCategory);
+		AssignmentNode->AddParameter(NewParameter, VarDefaultValue);
+	};
+
 	auto AddScriptVarLambda = [this, AssignmentNode](const UNiagaraScriptVariable* NewScriptVar) {
 		const FString VarDefaultValue = FNiagaraConstants::GetAttributeDefaultValue(NewScriptVar->Variable);
 		AddScriptVariable(NewScriptVar);
@@ -1526,6 +1539,7 @@ TSharedRef<SWidget> FNiagaraSystemToolkitParameterPanelViewModel::CreateAddParam
 		.AvailableParameterDefinitions(SystemViewModel->GetAvailableParameterDefinitions(bSkipSubscribedLibraries))
 		.SubscribedParameterDefinitions(SystemViewModel->GetSubscribedParameterDefinitions())
 		.OnNewParameterRequested_Lambda(AddParameterLambda)
+		.OnSpecificParameterRequested_Lambda(AddExistingParameterLambda)
 		.OnAddScriptVar_Lambda(AddScriptVarLambda)
 		.OnAddParameterDefinitions(this, &FNiagaraSystemToolkitParameterPanelViewModel::AddParameterDefinitions)
 		.OnAllowMakeType_Static(&INiagaraParameterPanelViewModel::CanMakeNewParameterOfType)
