Fixed "Create asset from this" causing a crash when closing the editor for the newly created emitter.
This seems to be happening due to GEditor object selection. The original objects get selected (GEditor->GetSelectedObjects) after duplication is done, adding a ref to it.
Since the UObject is transient and supposed to be garbage collected after our function is done, when garbage collection hits we have an external ref that causes a crash

diff --git a/Engine/Plugins/FX/Niagara/Source/NiagaraEditor/Private/NiagaraEditorUtilities.cpp b/Engine/Plugins/FX/Niagara/Source/NiagaraEditor/Private/NiagaraEditorUtilities.cpp
--- a/Engine/Plugins/FX/Niagara/Source/NiagaraEditor/Private/NiagaraEditorUtilities.cpp
+++ b/Engine/Plugins/FX/Niagara/Source/NiagaraEditor/Private/NiagaraEditorUtilities.cpp
@@ -51,6 +51,7 @@
 #include "NiagaraSystem.h"
 #include "NiagaraSystemEditorData.h"
 #include "ScopedTransaction.h"
+#include "Selection.h"
 #include "UpgradeNiagaraScriptResults.h"
 #include "Misc/FeedbackContext.h"
 #include "Misc/FileHelper.h"
@@ -2358,6 +2359,9 @@ void FNiagaraEditorUtilities::CreateAssetFromEmitter(TSharedRef<FNiagaraEmitterH
 		NewNode->NodeComment = CurrentComment;
 		NewNode->bCommentBubbleVisible = bCommentBubbleVisible;
 		NewNode->bCommentBubblePinned = bCommentBubblePinned;
+
+		// the duplicate action will select the duplicate emitter, adding a reference to it. We deselect here as it's a transient emitter to be garbage collected.
+		GEditor->GetSelectedObjects()->Deselect(DuplicateEmitter);
 	}
 }
 
