Fixed a few syncing issues with mode changes of script variables
- replaced default pin fixup code with a more streamlined version
- fixed undo of changing value modes not properly being reflected in the pin
- fixed undo not breaking pins if needed

diff --git a/Engine/Plugins/FX/Niagara/Source/NiagaraEditor/Private/NiagaraGraph.cpp b/Engine/Plugins/FX/Niagara/Source/NiagaraEditor/Private/NiagaraGraph.cpp
--- a/Engine/Plugins/FX/Niagara/Source/NiagaraEditor/Private/NiagaraGraph.cpp
+++ b/Engine/Plugins/FX/Niagara/Source/NiagaraEditor/Private/NiagaraGraph.cpp
@@ -2220,39 +2220,36 @@ void UNiagaraGraph::ScriptVariableChanged(FNiagaraVariable Variable)
 		return;
 	}
 
-	FNiagaraEditorModule& EditorModule = FNiagaraEditorModule::Get();
-
+	TSet<UNiagaraNodeParameterMapGet*> MapGetNodes;
 	TArray<UEdGraphPin*> Pins = FindParameterMapDefaultValuePins(Variable.GetName());
 	for (UEdGraphPin* Pin : Pins)
 	{
-		Pin->bHidden = ((*ScriptVariable)->DefaultMode == ENiagaraDefaultMode::Binding) || ((*ScriptVariable)->DefaultMode == ENiagaraDefaultMode::FailIfPreviouslyNotSet);
-		if ((*ScriptVariable)->DefaultMode == ENiagaraDefaultMode::Custom)
+		if(UNiagaraNodeParameterMapGet* MapGetNode = Cast<UNiagaraNodeParameterMapGet>(Pin->GetOwningNode()))
 		{
-			Pin->bNotConnectable = false;
-			Pin->bDefaultValueIsReadOnly = false;
+			MapGetNodes.Add(MapGetNode);
 		}
-		else
+		
+		if ((*ScriptVariable)->DefaultMode == ENiagaraDefaultMode::Value && !Variable.GetType().IsDataInterface())
 		{
-			Pin->BreakAllPinLinks(true);
-			Pin->bNotConnectable = true;
-			Pin->bDefaultValueIsReadOnly = true;
-
-			if ((*ScriptVariable)->DefaultMode == ENiagaraDefaultMode::Value && !Variable.GetType().IsDataInterface())
+			FNiagaraEditorModule& EditorModule = FNiagaraEditorModule::Get();
+			auto TypeUtilityValue = EditorModule.GetTypeUtilities(Variable.GetType());
+			if (TypeUtilityValue.IsValid() && TypeUtilityValue->CanHandlePinDefaults())
 			{
-				auto TypeUtilityValue = EditorModule.GetTypeUtilities(Variable.GetType());
-				if (TypeUtilityValue.IsValid() && TypeUtilityValue->CanHandlePinDefaults())
+				if (!Variable.IsDataAllocated())
 				{
-					if (!Variable.IsDataAllocated())
-					{
-						Variable.AllocateData();
-					}
-					FString NewDefaultValue = TypeUtilityValue->GetPinDefaultStringFromValue(Variable);
-					GetDefault<UEdGraphSchema_Niagara>()->TrySetDefaultValue(*Pin, NewDefaultValue, true);
+					Variable.AllocateData();
 				}
+				FString NewDefaultValue = TypeUtilityValue->GetPinDefaultStringFromValue(Variable);
+				GetDefault<UEdGraphSchema_Niagara>()->TrySetDefaultValue(*Pin, NewDefaultValue, true);
 			}
 		}
 	}
 
+	for(UNiagaraNodeParameterMapGet* MapGet : MapGetNodes)
+	{
+		MapGet->SynchronizeDefaultPins();
+	}
+
 	ValidateDefaultPins();
 	NotifyGraphChanged();
 }
diff --git a/Engine/Plugins/FX/Niagara/Source/NiagaraEditor/Private/NiagaraNodeParameterMapGet.cpp b/Engine/Plugins/FX/Niagara/Source/NiagaraEditor/Private/NiagaraNodeParameterMapGet.cpp
--- a/Engine/Plugins/FX/Niagara/Source/NiagaraEditor/Private/NiagaraNodeParameterMapGet.cpp
+++ b/Engine/Plugins/FX/Niagara/Source/NiagaraEditor/Private/NiagaraNodeParameterMapGet.cpp
@@ -142,7 +142,7 @@ UEdGraphPin* UNiagaraNodeParameterMapGet::CreateDefaultPin(UEdGraphPin* OutputPi
 	}
 	PinOutputToPinDefaultPersistentId.Add(OutputPin->PersistentGuid, DefaultPin->PersistentGuid);
 
-	SynchronizeDefaultInputPin(DefaultPin, OutputPin, GetScriptVariable(OutputPin->PinName));
+	SynchronizeDefaultInputPin(DefaultPin);
 	return DefaultPin;
 }
 
@@ -164,7 +164,7 @@ void UNiagaraNodeParameterMapGet::OnPinRenamed(UEdGraphPin* RenamedPin, const FS
 	if (DefaultPin)
 	{
 		DefaultPin->Modify();
-		SynchronizeDefaultInputPin(DefaultPin, RenamedPin, GetScriptVariable(RenamedPin->PinName));
+		SynchronizeDefaultInputPin(DefaultPin);
 	}
 
 	MarkNodeRequiresSynchronization(__FUNCTION__, true);
@@ -289,6 +289,26 @@ UEdGraphPin* UNiagaraNodeParameterMapGet::GetDefaultPin(UEdGraphPin* OutputPin)
 	return nullptr;
 }
 
+void UNiagaraNodeParameterMapGet::SynchronizeDefaultPins()
+{
+	TArray<UEdGraphPin*> OutputPins;
+	GetOutputPins(OutputPins);
+
+	for(UEdGraphPin* OutputPin : OutputPins)
+	{
+		if(IsAddPin(OutputPin) || OutputPin->bOrphanedPin)
+		{
+			continue;
+		}
+
+		UEdGraphPin* DefaultPin = GetDefaultPin(OutputPin);
+		if(ensureAlwaysMsgf(DefaultPin != nullptr, TEXT("There should always be an input pin for every output pin, even if hidden.")))
+		{
+			SynchronizeDefaultInputPin(DefaultPin);
+		}
+	}
+}
+
 
 UEdGraphPin* UNiagaraNodeParameterMapGet::GetOutputPinForDefault(const UEdGraphPin* DefaultPin) const
 {
@@ -332,14 +352,14 @@ void UNiagaraNodeParameterMapGet::PostLoad()
 			continue;
 		}
 
-		UEdGraphPin* InputPin = GetDefaultPin(OutputPin);
-		if (InputPin == nullptr)
+		UEdGraphPin* DefaultPin = GetDefaultPin(OutputPin);
+		if (DefaultPin == nullptr)
 		{
 			CreateDefaultPin(OutputPin);
 		}
 		else
 		{
-			SynchronizeDefaultInputPin(InputPin, OutputPin, GetScriptVariable(OutputPin->PinName));
+			SynchronizeDefaultInputPin(DefaultPin);
 		}
 
 		OutputPin->PinType.PinSubCategory = UNiagaraNodeParameterMapBase::ParameterPinSubCategory;
@@ -347,14 +367,21 @@ void UNiagaraNodeParameterMapGet::PostLoad()
 }
 
 
-void UNiagaraNodeParameterMapGet::SynchronizeDefaultInputPin(UEdGraphPin* DefaultPin, UEdGraphPin* OutputPin, UNiagaraScriptVariable* ScriptVar)
+void UNiagaraNodeParameterMapGet::SynchronizeDefaultInputPin(UEdGraphPin* DefaultPin)
 {
 	const UEdGraphSchema_Niagara* Schema = GetDefault<UEdGraphSchema_Niagara>();
-	if (!DefaultPin)
+	if (!DefaultPin || DefaultPin->bOrphanedPin)
 	{
 		return;
 	}
 
+	UEdGraphPin* OutputPin = GetOutputPinForDefault(DefaultPin);
+	if(!ensureAlwaysMsgf(OutputPin != nullptr, TEXT("There should always be an output pin found for every default pin we attempt to synchronize")))
+	{
+		return;
+	}
+	
+	UNiagaraScriptVariable* ScriptVar = GetScriptVariable(OutputPin->PinName);
 	FNiagaraVariable Var = Schema->PinToNiagaraVariable(OutputPin);
 	if (FNiagaraParameterMapHistory::IsEngineParameter(Var))
 	{
@@ -373,6 +400,7 @@ void UNiagaraNodeParameterMapGet::SynchronizeDefaultInputPin(UEdGraphPin* Defaul
 	if (ScriptVar) {
 		if (ScriptVar->DefaultMode == ENiagaraDefaultMode::Value)
 		{
+			DefaultPin->BreakAllPinLinks();
 			DefaultPin->bNotConnectable = true;
 			DefaultPin->bDefaultValueIsReadOnly = true;
 		}
diff --git a/Engine/Plugins/FX/Niagara/Source/NiagaraEditor/Private/NiagaraNodeParameterMapGet.h b/Engine/Plugins/FX/Niagara/Source/NiagaraEditor/Private/NiagaraNodeParameterMapGet.h
--- a/Engine/Plugins/FX/Niagara/Source/NiagaraEditor/Private/NiagaraNodeParameterMapGet.h
+++ b/Engine/Plugins/FX/Niagara/Source/NiagaraEditor/Private/NiagaraNodeParameterMapGet.h
@@ -40,6 +40,8 @@ public:
 	/** Get the default value input pin for one of the output pins specified.*/
 	UEdGraphPin* GetDefaultPin(UEdGraphPin* OutputPin) const;
 
+	void SynchronizeDefaultPins();
+	
 	virtual void PostLoad() override;
 
 	void GatherExternalDependencyData(ENiagaraScriptUsage InMasterUsage, const FGuid& InMasterUsageId, TArray<FNiagaraCompileHash>& InReferencedCompileHashes, TArray<FString>& InReferencedObjs) const override;
@@ -59,7 +61,7 @@ protected:
 	virtual void RemoveDynamicPin(UEdGraphPin* Pin) override;
 
 	/** Make sure that descriptions match up as well as any other value that can be changed.*/
-	void SynchronizeDefaultInputPin(UEdGraphPin* DefaultPin, UEdGraphPin* OutputPin, UNiagaraScriptVariable* ScriptVar);
+	void SynchronizeDefaultInputPin(UEdGraphPin* DefaultPin);
 	
 	/** Reverse the lookup from GetDefaultPin.*/
 	UEdGraphPin* GetOutputPinForDefault(const UEdGraphPin* DefaultPin) const;
diff --git a/Engine/Plugins/FX/Niagara/Source/NiagaraEditor/Private/NiagaraScriptVariable.cpp b/Engine/Plugins/FX/Niagara/Source/NiagaraEditor/Private/NiagaraScriptVariable.cpp
--- a/Engine/Plugins/FX/Niagara/Source/NiagaraEditor/Private/NiagaraScriptVariable.cpp
+++ b/Engine/Plugins/FX/Niagara/Source/NiagaraEditor/Private/NiagaraScriptVariable.cpp
@@ -31,6 +31,7 @@ void UNiagaraScriptVariable::PostEditChangeProperty(struct FPropertyChangedEvent
 
 	if (UNiagaraGraph* Graph = Cast<UNiagaraGraph>(GetOuter()))
 	{
+		Graph->ScriptVariableChanged(Variable);
 		Graph->NotifyGraphNeedsRecompile();
 	}
 	else if (UNiagaraParameterDefinitions* Definitions = Cast<UNiagaraParameterDefinitions>(GetOuter()))
