- Fix check slow when no rendering thread

#rnx

diff --git a/Engine/Plugins/FX/Niagara/Source/Niagara/Classes/NiagaraDataSet.h b/Engine/Plugins/FX/Niagara/Source/Niagara/Classes/NiagaraDataSet.h
--- a/Engine/Plugins/FX/Niagara/Source/Niagara/Classes/NiagaraDataSet.h
+++ b/Engine/Plugins/FX/Niagara/Source/Niagara/Classes/NiagaraDataSet.h
@@ -577,7 +577,7 @@ FORCEINLINE void FNiagaraDataBuffer::CheckUsage(bool bReadOnly)const
 	//We can read on the RT but any modifications must be GT (or GT Task).
 	//For GPU sims we must be on the RT.
 	checkSlow(
-		(Owner->GetSimTarget() == ENiagaraSimTarget::CPUSim && (IsInGameThread() || bReadOnly || !IsInRenderingThread())) ||
+		(Owner->GetSimTarget() == ENiagaraSimTarget::CPUSim && (IsInGameThread() || bReadOnly || !GIsThreadedRendering || !IsInRenderingThread())) ||
 		(Owner->GetSimTarget() == ENiagaraSimTarget::GPUComputeSim && IsInRenderingThread())
 	);
 }
