fix query for cell count on neighborgrids

diff --git a/Engine/Plugins/FX/Niagara/Source/Niagara/Private/NiagaraDataInterfaceNeighborGrid3D.cpp b/Engine/Plugins/FX/Niagara/Source/Niagara/Private/NiagaraDataInterfaceNeighborGrid3D.cpp
--- a/Engine/Plugins/FX/Niagara/Source/Niagara/Private/NiagaraDataInterfaceNeighborGrid3D.cpp
+++ b/Engine/Plugins/FX/Niagara/Source/Niagara/Private/NiagaraDataInterfaceNeighborGrid3D.cpp
@@ -368,11 +368,15 @@ void UNiagaraDataInterfaceNeighborGrid3D::GetNumCells(FVectorVMExternalFunctionC
 	FNDIOutputParam<int32> NumCellsY(Context);
 	FNDIOutputParam<int32> NumCellsZ(Context);
 
+	int32 TmpNumCellsX = InstData->NumCells.X;
+	int32 TmpNumCellsY = InstData->NumCells.Y;
+	int32 TmpNumCellsZ = InstData->NumCells.Z;
+
 	for (int32 InstanceIdx = 0; InstanceIdx < Context.GetNumInstances(); ++InstanceIdx)
 	{
-		NumCellsX.SetAndAdvance(NumCells.X);
-		NumCellsY.SetAndAdvance(NumCells.Y);
-		NumCellsZ.SetAndAdvance(NumCells.Z);
+		NumCellsX.SetAndAdvance(TmpNumCellsX);
+		NumCellsY.SetAndAdvance(TmpNumCellsY);
+		NumCellsZ.SetAndAdvance(TmpNumCellsZ);
 	}
 }
 
