fix divide by 0 when voxel size or number of cells is set to 0

diff --git a/Engine/Plugins/FX/Niagara/Source/Niagara/Private/NiagaraDataInterfaceGrid3DCollection.cpp b/Engine/Plugins/FX/Niagara/Source/Niagara/Private/NiagaraDataInterfaceGrid3DCollection.cpp
--- a/Engine/Plugins/FX/Niagara/Source/Niagara/Private/NiagaraDataInterfaceGrid3DCollection.cpp
+++ b/Engine/Plugins/FX/Niagara/Source/Niagara/Private/NiagaraDataInterfaceGrid3DCollection.cpp
@@ -2472,10 +2472,9 @@ void UNiagaraDataInterfaceGrid3DCollection::SetNumCells(FVectorVMExternalFunctio
 		{
 			FIntVector OldNumCells = InstData->NumCells;
 
-
-			InstData->NumCells.X = NewNumCellsX;
-			InstData->NumCells.Y = NewNumCellsY;
-			InstData->NumCells.Z = NewNumCellsZ;
+			InstData->NumCells.X = FMath::Max(1, NewNumCellsX);
+			InstData->NumCells.Y = FMath::Max(1, NewNumCellsY);
+			InstData->NumCells.Z = FMath::Max(1, NewNumCellsZ);
 
 			if (!FMath::IsNearlyEqual(GNiagaraGrid3DResolutionMultiplier, 1.0f))
 			{
diff --git a/Engine/Plugins/FX/Niagara/Source/Niagara/Private/NiagaraDataInterfaceNeighborGrid3D.cpp b/Engine/Plugins/FX/Niagara/Source/Niagara/Private/NiagaraDataInterfaceNeighborGrid3D.cpp
--- a/Engine/Plugins/FX/Niagara/Source/Niagara/Private/NiagaraDataInterfaceNeighborGrid3D.cpp
+++ b/Engine/Plugins/FX/Niagara/Source/Niagara/Private/NiagaraDataInterfaceNeighborGrid3D.cpp
@@ -685,10 +685,11 @@ void UNiagaraDataInterfaceNeighborGrid3D::SetNumCells(FVectorVMExternalFunctionC
 		{
 			FIntVector OldNumCells = InstData->NumCells;
 			int OldMaxNeighborsPerCell = InstData->MaxNeighborsPerCell;
+			
+			InstData->NumCells.X = FMath::Max(1, NewNumCellsX);
+			InstData->NumCells.Y = FMath::Max(1, NewNumCellsY);
+			InstData->NumCells.Z = FMath::Max(1, NewNumCellsZ);
 
-			InstData->NumCells.X = NewNumCellsX;
-			InstData->NumCells.Y = NewNumCellsY;
-			InstData->NumCells.Z = NewNumCellsZ;
 			InstData->MaxNeighborsPerCell = NewMaxNeighborsPerCell;
 		
 			InstData->NeedsRealloc = OldNumCells != InstData->NumCells || OldMaxNeighborsPerCell != InstData->MaxNeighborsPerCell;
