- Fixed SA issue with ptr access without nullcheck
- Fixed non unity compile error

diff --git a/Engine/Plugins/FX/Niagara/Source/NiagaraEditor/Private/ViewModels/NiagaraScriptGraphViewModel.cpp b/Engine/Plugins/FX/Niagara/Source/NiagaraEditor/Private/ViewModels/NiagaraScriptGraphViewModel.cpp
--- a/Engine/Plugins/FX/Niagara/Source/NiagaraEditor/Private/ViewModels/NiagaraScriptGraphViewModel.cpp
+++ b/Engine/Plugins/FX/Niagara/Source/NiagaraEditor/Private/ViewModels/NiagaraScriptGraphViewModel.cpp
@@ -21,6 +21,7 @@
 #include "Widgets/SNiagaraParameterName.h"
 #include "Widgets/Notifications/SNotificationList.h"
 #include "NiagaraNotificationWidgetProvider.h"
+#include "Misc/Base64.h"
 
 #define LOCTEXT_NAMESPACE "NiagaraScriptGraphViewModel"
 
@@ -457,8 +458,11 @@ bool FNiagaraScriptGraphViewModel::CanPasteNodes() const
 
 	const UNiagaraClipboardContent* ClipboardContent = FNiagaraEditorModule::Get().GetClipboard().GetClipboardContent();
 	FString ExportedNodesText;
-	FBase64::Decode(ClipboardContent->ExportedNodes, ExportedNodesText);
-	return ClipboardContent? FEdGraphUtilities::CanImportNodesFromText(Graph, ExportedNodesText) : false;
+	if(ClipboardContent)
+	{
+		FBase64::Decode(ClipboardContent->ExportedNodes, ExportedNodesText);
+	}
+	return ClipboardContent ? FEdGraphUtilities::CanImportNodesFromText(Graph, ExportedNodesText) : false;
 }
 
 void FNiagaraScriptGraphViewModel::DuplicateNodes()
