Resubmit of 19841660 for nicholas.goldstein:  Do not error if visiting "known constant" attribute parameter that has type mismatch.
This could occur when evaluating Particles.Position of type Vec3 in graph but the "known constant" is of type Position.

diff --git a/Engine/Plugins/FX/Niagara/Source/Niagara/Classes/NiagaraConstants.h b/Engine/Plugins/FX/Niagara/Source/Niagara/Classes/NiagaraConstants.h
--- a/Engine/Plugins/FX/Niagara/Source/Niagara/Classes/NiagaraConstants.h
+++ b/Engine/Plugins/FX/Niagara/Source/Niagara/Classes/NiagaraConstants.h
@@ -139,6 +139,28 @@
 #define TRANSLATOR_PARAM_BEGIN_DEFAULTS                  INiagaraModule::GetVar_BeginDefaults()
 #define TRANSLATOR_PARAM_CALL_ID                  		 INiagaraModule::GetVar_CallID()
 
+enum class ENiagaraKnownConstantType : uint8
+{
+	Attribute,
+	Other,
+};
+
+struct FNiagaraKnownConstantInfo
+{
+	FNiagaraKnownConstantInfo()
+		: ConstantVar(nullptr)
+		, ConstantType(ENiagaraKnownConstantType::Other)
+	{};
+
+	FNiagaraKnownConstantInfo(const FNiagaraVariable* InConstantVar, const ENiagaraKnownConstantType InConstantType)
+		: ConstantVar(InConstantVar)
+		, ConstantType(InConstantType)
+	{};
+
+	const FNiagaraVariable* ConstantVar = nullptr;
+	const ENiagaraKnownConstantType ConstantType = ENiagaraKnownConstantType::Other;
+};
+
 struct NIAGARA_API FNiagaraConstants
 {
 	static void Init();
@@ -162,6 +184,7 @@ struct NIAGARA_API FNiagaraConstants
 	static const FNiagaraVariableMetaData* GetConstantMetaData(const FNiagaraVariable& InVar);
 
 	static const FNiagaraVariable* GetKnownConstant(const FName& InName, bool bAllowPartialNameMatch);
+	static const FNiagaraKnownConstantInfo GetKnownConstantInfo(const FName& InName, bool bAllowPartialNameMatch);
 	static const FNiagaraVariable *FindStaticSwitchConstant(const FName& InName);
 
 	static bool IsEngineManagedAttribute(const FNiagaraVariable& Var);
diff --git a/Engine/Plugins/FX/Niagara/Source/Niagara/Private/NiagaraConstants.cpp b/Engine/Plugins/FX/Niagara/Source/Niagara/Private/NiagaraConstants.cpp
--- a/Engine/Plugins/FX/Niagara/Source/Niagara/Private/NiagaraConstants.cpp
+++ b/Engine/Plugins/FX/Niagara/Source/Niagara/Private/NiagaraConstants.cpp
@@ -793,6 +793,11 @@ FNiagaraVariableAttributeBinding FNiagaraConstants::GetAttributeDefaultBinding(c
 }
 
 const FNiagaraVariable* FNiagaraConstants::GetKnownConstant(const FName& InName, bool bAllowPartialNameMatch)
+{
+	return GetKnownConstantInfo(InName, bAllowPartialNameMatch).ConstantVar;
+}
+
+const FNiagaraKnownConstantInfo FNiagaraConstants::GetKnownConstantInfo(const FName& InName, bool bAllowPartialNameMatch)
 {
 	const TArray<FNiagaraVariable>& EngineConstants = GetEngineConstants();
 
@@ -805,7 +810,7 @@ const FNiagaraVariable* FNiagaraConstants::GetKnownConstant(const FName& InName,
 
 		if (FoundSystemVar)
 		{
-			return FoundSystemVar;
+			return FNiagaraKnownConstantInfo(FoundSystemVar, ENiagaraKnownConstantType::Other);
 		}
 	}
 	else
@@ -813,7 +818,7 @@ const FNiagaraVariable* FNiagaraConstants::GetKnownConstant(const FName& InName,
 		int32 FoundIdx = FNiagaraVariable::SearchArrayForPartialNameMatch(EngineConstants, InName);
 		if (FoundIdx != INDEX_NONE)
 		{
-			return &EngineConstants[FoundIdx];
+			return FNiagaraKnownConstantInfo(&EngineConstants[FoundIdx], ENiagaraKnownConstantType::Other);
 		}
 	}
 
@@ -827,7 +832,7 @@ const FNiagaraVariable* FNiagaraConstants::GetKnownConstant(const FName& InName,
 
 		if (FoundAttribVar)
 		{
-			return FoundAttribVar;
+			return FNiagaraKnownConstantInfo(FoundAttribVar, ENiagaraKnownConstantType::Attribute);
 		}
 	}
 	else
@@ -835,10 +840,11 @@ const FNiagaraVariable* FNiagaraConstants::GetKnownConstant(const FName& InName,
 		int32 FoundIdx = FNiagaraVariable::SearchArrayForPartialNameMatch(Attributes, InName);
 		if (FoundIdx != INDEX_NONE)
 		{
-			return &Attributes[FoundIdx];
+			return FNiagaraKnownConstantInfo(&Attributes[FoundIdx], ENiagaraKnownConstantType::Attribute);
 		}
 	}
-	return nullptr;
+
+	return FNiagaraKnownConstantInfo(nullptr, ENiagaraKnownConstantType::Other);
 }
 
 #undef LOCTEXT_NAMESPACE
diff --git a/Engine/Plugins/FX/Niagara/Source/NiagaraEditor/Private/NiagaraHlslTranslator.cpp b/Engine/Plugins/FX/Niagara/Source/NiagaraEditor/Private/NiagaraHlslTranslator.cpp
--- a/Engine/Plugins/FX/Niagara/Source/NiagaraEditor/Private/NiagaraHlslTranslator.cpp
+++ b/Engine/Plugins/FX/Niagara/Source/NiagaraEditor/Private/NiagaraHlslTranslator.cpp
@@ -5130,11 +5130,11 @@ void FHlslNiagaraTranslator::ParameterMapSet(UNiagaraNodeParameterMapSet* SetNod
 
 			Var = ActiveHistoryForFunctionCalls.ResolveAliases(Var);
 			
-			const FNiagaraVariable* ConstantVar = FNiagaraConstants::GetKnownConstant(Var.GetName(), false);
-			if (ConstantVar != nullptr && ConstantVar->GetType() != Var.GetType())
+			const FNiagaraKnownConstantInfo ConstantInfo = FNiagaraConstants::GetKnownConstantInfo(Var.GetName(), false);
+			if (ConstantInfo.ConstantVar != nullptr && ConstantInfo.ConstantVar->GetType() != Var.GetType() && ConstantInfo.ConstantType != ENiagaraKnownConstantType::Attribute)
 			{
 				Error(FText::Format(LOCTEXT("MismatchedConstantTypes", "Variable {0} is a system constant, but its type is different! {1} != {2}"), FText::FromName(Var.GetName()),
-					ConstantVar->GetType().GetNameText(), Var.GetType().GetNameText()), nullptr, nullptr);
+					ConstantInfo.ConstantVar->GetType().GetNameText(), Var.GetType().GetNameText()), nullptr, nullptr);
 			}
 
 			if (FNiagaraConstants::IsEngineManagedAttribute(Var))
@@ -6084,11 +6084,11 @@ void FHlslNiagaraTranslator::HandleParameterRead(int32 ParamMapHistoryIdx, const
 	bool bWasEmitterAliased = FNiagaraParameterMapHistory::IsAliasedEmitterParameter(Var);
 	Var = ActiveHistoryForFunctionCalls.ResolveAliases(Var);
 
-	const FNiagaraVariable* ConstantVar = FNiagaraConstants::GetKnownConstant(Var.GetName(), false);
-	if (ConstantVar != nullptr && ConstantVar->GetType() != Var.GetType())
+	const FNiagaraKnownConstantInfo ConstantInfo = FNiagaraConstants::GetKnownConstantInfo(Var.GetName(), false);
+	if (ConstantInfo.ConstantVar != nullptr && ConstantInfo.ConstantVar->GetType() != Var.GetType() && ConstantInfo.ConstantType != ENiagaraKnownConstantType::Attribute)
 	{
 		Error(FText::Format(LOCTEXT("MismatchedConstantTypes", "Variable {0} is a system constant, but its type is different! {1} != {2}"), FText::FromName(Var.GetName()),
-			ConstantVar->GetType().GetNameText(), Var.GetType().GetNameText()), ErrorNode, nullptr);
+			ConstantInfo.ConstantVar->GetType().GetNameText(), Var.GetType().GetNameText()), ErrorNode, nullptr);
 	}
 
 	if (History.IsPrimaryDataSetOutput(Var, GetTargetUsage())) // Note that data interfaces aren't ever in the primary data set even if the namespace matches.
diff --git a/Engine/Source/Runtime/Core/Private/UObject/FortniteMainBranchObjectVersion.cpp b/Engine/Source/Runtime/Core/Private/UObject/FortniteMainBranchObjectVersion.cpp
--- a/Engine/Source/Runtime/Core/Private/UObject/FortniteMainBranchObjectVersion.cpp
+++ b/Engine/Source/Runtime/Core/Private/UObject/FortniteMainBranchObjectVersion.cpp
@@ -11,7 +11,7 @@ TMap<FGuid, FGuid> FFortniteMainBranchObjectVersion::GetSystemGuids()
 	SystemGuids.Add(DevGuids.MATERIALSHADERMAP_DERIVEDDATA_VER, FGuid("1DA85DA5733C4A489A1BE117CDC5ACCC"));
 	SystemGuids.Add(DevGuids.NANITE_DERIVEDDATA_VER, FGuid("49FBF137DBFB4204B6C9E5B91DA15EBD"));
 	SystemGuids.Add(DevGuids.NIAGARASHADERMAP_DERIVEDDATA_VER, FGuid("2289D5116CF94BC0AFEAEC541468E645"));
-	SystemGuids.Add(DevGuids.Niagara_LatestScriptCompileVersion, FGuid("A3D2BC12548B425DB10DCE6D0237381D"));
+	SystemGuids.Add(DevGuids.Niagara_LatestScriptCompileVersion, FGuid("E3B6BEB506864EB697362DA7ADC799F1"));
 	SystemGuids.Add(DevGuids.SkeletalMeshDerivedDataVersion, FGuid("ACF593EAAE354FCBB34CF44F5AA1BFD2"));
 	SystemGuids.Add(DevGuids.STATICMESH_DERIVEDDATA_VER, FGuid("3006d21f867648a68cf4edee185446b7"));
 
