Niagara - Fix a bunch of spots which were handling null scripts, but not null graphs.  Change the stats view model so that it only adds valid shader platforms rather than adding them all and then crashing later because one is invalid.

diff --git a/Engine/Plugins/FX/Niagara/Source/NiagaraEditor/Private/NiagaraCompiler.cpp b/Engine/Plugins/FX/Niagara/Source/NiagaraEditor/Private/NiagaraCompiler.cpp
--- a/Engine/Plugins/FX/Niagara/Source/NiagaraEditor/Private/NiagaraCompiler.cpp
+++ b/Engine/Plugins/FX/Niagara/Source/NiagaraEditor/Private/NiagaraCompiler.cpp
@@ -169,7 +169,7 @@ void FNiagaraCompileRequestDuplicateData::DuplicateReferencedGraphsRecursive(UNi
 			{
 				UNiagaraScript* FunctionScript = FunctionCallNode->FunctionScript;
 				bool bFunctionCallOwnsScript = FunctionScript != nullptr && FunctionScript->GetOuter() == FunctionCallNode;
-				if(FunctionScript != nullptr && bFunctionCallOwnsScript == false)
+				if(FunctionCallNode->HasValidScriptAndGraph() && bFunctionCallOwnsScript == false)
 				{
 					// If the function call doesn't already own the script it's pointing at then script needs to be duplicated since it's a referenced
 					// script and will need to be preprocessed.
diff --git a/Engine/Plugins/FX/Niagara/Source/NiagaraEditor/Private/NiagaraNodeFunctionCall.cpp b/Engine/Plugins/FX/Niagara/Source/NiagaraEditor/Private/NiagaraNodeFunctionCall.cpp
--- a/Engine/Plugins/FX/Niagara/Source/NiagaraEditor/Private/NiagaraNodeFunctionCall.cpp
+++ b/Engine/Plugins/FX/Niagara/Source/NiagaraEditor/Private/NiagaraNodeFunctionCall.cpp
@@ -280,10 +280,9 @@ void UNiagaraNodeFunctionCall::AllocateDefaultPins()
 	}
 
 	const UEdGraphSchema_Niagara* Schema = CastChecked<UEdGraphSchema_Niagara>(GetSchema());
-	if (FunctionScript)
+	if (HasValidScriptAndGraph())
 	{
-		UNiagaraScriptSource* Source = GetFunctionScriptSource();
-		UNiagaraGraph* Graph = Source->NodeGraph;
+		UNiagaraGraph* Graph = GetCalledGraph();
 
 		//These pins must be refreshed and kept in the correct order for the function
 		TArray<FNiagaraVariable> Inputs;
@@ -586,7 +585,7 @@ void UNiagaraNodeFunctionCall::Compile(class FHlslNiagaraTranslator* Translator,
 	const UEdGraphSchema_Niagara* Schema = CastChecked<UEdGraphSchema_Niagara>(GetSchema());
 	UNiagaraGraph* CallerGraph = GetNiagaraGraph();
 	FVersionedNiagaraScriptData* ScriptData = GetScriptData();
-	if (ScriptData)
+	if (ScriptData && HasValidScriptAndGraph())
 	{
 		if (InvalidScriptVersionReference.IsValid())
 		{
@@ -1014,25 +1013,20 @@ bool UNiagaraNodeFunctionCall::RefreshFromExternalChanges()
 
 void UNiagaraNodeFunctionCall::GatherExternalDependencyData(ENiagaraScriptUsage InMasterUsage, const FGuid& InMasterUsageId, TArray<FNiagaraCompileHash>& InReferencedCompileHashes, TArray<FString>& InReferencedObjs) const
 {
-	if (FunctionScript)
+	if (HasValidScriptAndGraph())
 	{
-		UNiagaraScriptSource* Source = GetFunctionScriptSource();
-		UNiagaraGraph* FunctionGraph = CastChecked<UNiagaraGraph>(Source->NodeGraph);
-		
 		// We don't know which graph type we're referencing, so we try them all... may need to replace this with something faster in the future.
-		if (FunctionGraph)
+		UNiagaraGraph* FunctionGraph = GetCalledGraph();
+		FunctionGraph->RebuildCachedCompileIds();
+		for (int32 i = (int32)ENiagaraScriptUsage::Function; i <= (int32)ENiagaraScriptUsage::DynamicInput; i++)
 		{
-			FunctionGraph->RebuildCachedCompileIds();
-			for (int32 i = (int32)ENiagaraScriptUsage::Function; i <= (int32)ENiagaraScriptUsage::DynamicInput; i++)
+			FGuid FoundGuid = FunctionGraph->GetBaseId((ENiagaraScriptUsage)i, FGuid(0, 0, 0, 0));
+			FNiagaraCompileHash FoundCompileHash = FunctionGraph->GetCompileDataHash((ENiagaraScriptUsage)i, FGuid(0, 0, 0, 0));
+			if (FoundGuid.IsValid() && FoundCompileHash.IsValid())
 			{
-				FGuid FoundGuid = FunctionGraph->GetBaseId((ENiagaraScriptUsage)i, FGuid(0, 0, 0, 0));
-				FNiagaraCompileHash FoundCompileHash = FunctionGraph->GetCompileDataHash((ENiagaraScriptUsage)i, FGuid(0, 0, 0, 0));
-				if (FoundGuid.IsValid() && FoundCompileHash.IsValid())
-				{
-					InReferencedCompileHashes.AddUnique(FoundCompileHash);
-					InReferencedObjs.Add(FunctionGraph->GetPathName());
-					FunctionGraph->GatherExternalDependencyData((ENiagaraScriptUsage)i, FGuid(0, 0, 0, 0), InReferencedCompileHashes, InReferencedObjs);
-				}
+				InReferencedCompileHashes.AddUnique(FoundCompileHash);
+				InReferencedObjs.Add(FunctionGraph->GetPathName());
+				FunctionGraph->GatherExternalDependencyData((ENiagaraScriptUsage)i, FGuid(0, 0, 0, 0), InReferencedCompileHashes, InReferencedObjs);
 			}
 		}
 	}
@@ -1044,17 +1038,9 @@ void UNiagaraNodeFunctionCall::UpdateCompileHashForNode(FSHA1& HashState) const
 	HashState.UpdateWithString(*GetFunctionName(), GetFunctionName().Len());
 }
 
-bool UNiagaraNodeFunctionCall::ScriptIsValid() const
+bool UNiagaraNodeFunctionCall::HasValidScriptAndGraph() const
 {
-	if (FunctionScript != nullptr)
-	{
-		return true;
-	}
-	else if (Signature.IsValid())
-	{
-		return true;
-	}
-	return false;
+	return FunctionScript != nullptr && GetCalledGraph() != nullptr;
 }
 
 void UNiagaraNodeFunctionCall::BuildParameterMapHistory(FNiagaraParameterMapHistoryBuilder& OutHistory, bool bRecursive /*= true*/, bool bFilterForCompilation /*= true*/) const
@@ -1067,8 +1053,6 @@ void UNiagaraNodeFunctionCall::BuildParameterMapHistory(FNiagaraParameterMapHist
 		// Looks like this function call is not yet hooked up. Skip it to prevent cascading errors in the compilation
 		return;
 	}
-	UNiagaraScriptSource* Source = GetFunctionScriptSource();
-
 	
 	Super::BuildParameterMapHistory(OutHistory, bRecursive, bFilterForCompilation);
 	if (!IsNodeEnabled() && OutHistory.GetIgnoreDisabled())
@@ -1077,10 +1061,15 @@ void UNiagaraNodeFunctionCall::BuildParameterMapHistory(FNiagaraParameterMapHist
 		return;
 	}
 
-	
-	if (Source)
+	if (HasValidScriptAndGraph() == false && Signature.IsValid() == false)
+	{
+		RouteParameterMapAroundMe(OutHistory, bRecursive);
+		return;
+	}
+
+	if (HasValidScriptAndGraph())
 	{
-		UNiagaraGraph* FunctionGraph = CastChecked<UNiagaraGraph>(Source->NodeGraph);
+		UNiagaraGraph* FunctionGraph = GetCalledGraph();
 
 		UNiagaraNodeOutput* OutputNode = FunctionGraph->FindOutputNode(ENiagaraScriptUsage::Function);
 		if (OutputNode == nullptr)
@@ -1182,7 +1171,7 @@ void UNiagaraNodeFunctionCall::BuildParameterMapHistory(FNiagaraParameterMapHist
 			OutHistory.RegisterConstantPin(MatchedConstants[i].Value, MatchedConstants[i].Key);
 		}
 	}
-	else if (!ScriptIsValid() || Signature.bRequiresExecPin)
+	else if (Signature.bRequiresExecPin)
 	{
 		RouteParameterMapAroundMe(OutHistory, bRecursive);
 	}
diff --git a/Engine/Plugins/FX/Niagara/Source/NiagaraEditor/Private/ViewModels/NiagaraScriptStatsViewModel.cpp b/Engine/Plugins/FX/Niagara/Source/NiagaraEditor/Private/ViewModels/NiagaraScriptStatsViewModel.cpp
--- a/Engine/Plugins/FX/Niagara/Source/NiagaraEditor/Private/ViewModels/NiagaraScriptStatsViewModel.cpp
+++ b/Engine/Plugins/FX/Niagara/Source/NiagaraEditor/Private/ViewModels/NiagaraScriptStatsViewModel.cpp
@@ -761,15 +761,24 @@ EShaderPlatform FNiagaraScriptStatsViewModel::ColumnNameToShaderPlatform(const F
 
 void FNiagaraScriptStatsViewModel::BuildShaderPlatformDetails()
 {
+	TArray<EShaderPlatform, TInlineAllocator<SP_StaticPlatform_Last + 1>> StaticShaderPlatforms;
+	StaticShaderPlatforms.Add(SP_PCD3D_SM5);
+	StaticShaderPlatforms.Add(SP_VULKAN_SM5);
+	StaticShaderPlatforms.Add(SP_OPENGL_ES3_1_ANDROID); 
+	StaticShaderPlatforms.Add(SP_VULKAN_ES3_1_ANDROID);
+	StaticShaderPlatforms.Add(SP_VULKAN_SM5_ANDROID);
+	StaticShaderPlatforms.Add(SP_METAL_SM5);
+	StaticShaderPlatforms.Add(SP_METAL);
+	StaticShaderPlatforms.Add(SP_METAL_MRT);
+
 	TArray<EShaderPlatform, TInlineAllocator<SP_StaticPlatform_Last + 1>> ValidShaderPlatforms;
-	ValidShaderPlatforms.Add(SP_PCD3D_SM5);
-	ValidShaderPlatforms.Add(SP_VULKAN_SM5);
-	ValidShaderPlatforms.Add(SP_OPENGL_ES3_1_ANDROID); 
-	ValidShaderPlatforms.Add(SP_VULKAN_ES3_1_ANDROID);
-	ValidShaderPlatforms.Add(SP_VULKAN_SM5_ANDROID);
-	ValidShaderPlatforms.Add(SP_METAL_SM5);
-	ValidShaderPlatforms.Add(SP_METAL);
-	ValidShaderPlatforms.Add(SP_METAL_MRT);
+	for (EShaderPlatform StaticShaderPlatform : StaticShaderPlatforms)
+	{
+		if (FDataDrivenShaderPlatformInfo::IsValid(StaticShaderPlatform))
+		{
+			ValidShaderPlatforms.Add(StaticShaderPlatform);
+		}
+	}
 
 	for (int32 iPlatform = SP_StaticPlatform_First; iPlatform <= SP_StaticPlatform_Last; ++iPlatform)
 	{
diff --git a/Engine/Plugins/FX/Niagara/Source/NiagaraEditor/Private/ViewModels/Stack/NiagaraStackFilteredObject.cpp b/Engine/Plugins/FX/Niagara/Source/NiagaraEditor/Private/ViewModels/Stack/NiagaraStackFilteredObject.cpp
--- a/Engine/Plugins/FX/Niagara/Source/NiagaraEditor/Private/ViewModels/Stack/NiagaraStackFilteredObject.cpp
+++ b/Engine/Plugins/FX/Niagara/Source/NiagaraEditor/Private/ViewModels/Stack/NiagaraStackFilteredObject.cpp
@@ -119,7 +119,7 @@ void UNiagaraStackFilteredObject::AppendEmitterCategory(TSharedPtr<FNiagaraScrip
 		bool bIsFirst = true;
 		for (UNiagaraNodeFunctionCall* ModuleNode : ModuleNodes)
 		{
-			if (ModuleNode && ModuleNode->ScriptIsValid())
+			if (ModuleNode && (ModuleNode->HasValidScriptAndGraph() || ModuleNode->Signature.IsValid()))
 			{
 				RefreshChildrenForFunctionCall(ModuleNode, ModuleNode, CurrentChildren, NewChildren, NewIssues, true);
 			}
diff --git a/Engine/Plugins/FX/Niagara/Source/NiagaraEditor/Private/ViewModels/Stack/NiagaraStackGraphUtilities.cpp b/Engine/Plugins/FX/Niagara/Source/NiagaraEditor/Private/ViewModels/Stack/NiagaraStackGraphUtilities.cpp
--- a/Engine/Plugins/FX/Niagara/Source/NiagaraEditor/Private/ViewModels/Stack/NiagaraStackGraphUtilities.cpp
+++ b/Engine/Plugins/FX/Niagara/Source/NiagaraEditor/Private/ViewModels/Stack/NiagaraStackGraphUtilities.cpp
@@ -1801,9 +1801,9 @@ UNiagaraNodeFunctionCall* FNiagaraStackGraphUtilities::AddScriptModuleToStack(FA
 	NewModuleNode->FunctionScriptAssetObjectPath = ModuleScriptAsset.ObjectPath;
 	ModuleNodeCreator.Finalize();
 
-	if (NewModuleNode->FunctionScript == nullptr)
+	if (NewModuleNode->HasValidScriptAndGraph() == false)
 	{
-		// If the module script is null, add parameter map inputs and outputs so that the node can be wired into the graph correctly.
+		// If the module script or graph are invalid, add parameter map inputs and outputs so that the node can be wired into the owning graph correctly.
 		const UEdGraphSchema_Niagara* NiagaraSchema = GetDefault<UEdGraphSchema_Niagara>();
 		NewModuleNode->CreatePin(EGPD_Input, NiagaraSchema->TypeDefinitionToPinType(FNiagaraTypeDefinition::GetParameterMapDef()), TEXT("InputMap"));
 		NewModuleNode->CreatePin(EGPD_Output, NiagaraSchema->TypeDefinitionToPinType(FNiagaraTypeDefinition::GetParameterMapDef()), TEXT("OutputMap"));
diff --git a/Engine/Plugins/FX/Niagara/Source/NiagaraEditor/Private/ViewModels/Stack/NiagaraStackModuleItem.cpp b/Engine/Plugins/FX/Niagara/Source/NiagaraEditor/Private/ViewModels/Stack/NiagaraStackModuleItem.cpp
--- a/Engine/Plugins/FX/Niagara/Source/NiagaraEditor/Private/ViewModels/Stack/NiagaraStackModuleItem.cpp
+++ b/Engine/Plugins/FX/Niagara/Source/NiagaraEditor/Private/ViewModels/Stack/NiagaraStackModuleItem.cpp
@@ -143,10 +143,10 @@ void UNiagaraStackModuleItem::RefreshChildrenInternal(const TArray<UNiagaraStack
 	bIsScratchModuleCache.Reset();
 	DisplayNameCache.Reset();
 
-	if (FunctionCallNode != nullptr && FunctionCallNode->ScriptIsValid())
+	if (FunctionCallNode != nullptr && (FunctionCallNode->HasValidScriptAndGraph() || FunctionCallNode->Signature.IsValid()))
 	{
 		// Determine if meta-data requires that we add our own refresh button here.
-		if (FunctionCallNode->FunctionScript)
+		if (FunctionCallNode->HasValidScriptAndGraph())
 		{
 			UNiagaraGraph* Graph = FunctionCallNode->GetCalledGraph();
 			const TMap<FNiagaraVariable, TObjectPtr<UNiagaraScriptVariable>>& MetaDataMap = Graph->GetAllMetaData();
@@ -867,7 +867,7 @@ void UNiagaraStackModuleItem::RefreshIssues(TArray<FStackIssue>& NewIssues)
 				false));
 		}
 	}
-	else if (!FunctionCallNode->ScriptIsValid())
+	else if (FunctionCallNode->HasValidScriptAndGraph() == false && FunctionCallNode->Signature.IsValid() == false)
 	{
 		FStackIssue InvalidScriptError(
 			EStackIssueSeverity::Error,
diff --git a/Engine/Plugins/FX/Niagara/Source/NiagaraEditor/Private/ViewModels/Stack/NiagaraStackSummaryViewInputCollection.cpp b/Engine/Plugins/FX/Niagara/Source/NiagaraEditor/Private/ViewModels/Stack/NiagaraStackSummaryViewInputCollection.cpp
--- a/Engine/Plugins/FX/Niagara/Source/NiagaraEditor/Private/ViewModels/Stack/NiagaraStackSummaryViewInputCollection.cpp
+++ b/Engine/Plugins/FX/Niagara/Source/NiagaraEditor/Private/ViewModels/Stack/NiagaraStackSummaryViewInputCollection.cpp
@@ -115,7 +115,7 @@ void UNiagaraStackSummaryViewObject::AppendEmitterCategory(FFunctionCallNodesSta
 		bool bIsFirst = true;
 		for (UNiagaraNodeFunctionCall* ModuleNode : ModuleNodes)
 		{
-			if (ModuleNode && ModuleNode->ScriptIsValid())
+			if (ModuleNode && (ModuleNode->HasValidScriptAndGraph() || ModuleNode->Signature.IsValid()))
 			{
 				AppendInputsForFunctionCall(State, ModuleNode, ModuleNode, NewIssues, true);
 			}
diff --git a/Engine/Plugins/FX/Niagara/Source/NiagaraEditor/Public/NiagaraNodeFunctionCall.h b/Engine/Plugins/FX/Niagara/Source/NiagaraEditor/Public/NiagaraNodeFunctionCall.h
--- a/Engine/Plugins/FX/Niagara/Source/NiagaraEditor/Public/NiagaraNodeFunctionCall.h
+++ b/Engine/Plugins/FX/Niagara/Source/NiagaraEditor/Public/NiagaraNodeFunctionCall.h
@@ -97,7 +97,7 @@ public:
 	UPROPERTY(EditAnywhere, Category = "Debug")
 	bool bInheritDebugStatus = true;
 
-	NIAGARAEDITOR_API bool ScriptIsValid() const;
+	NIAGARAEDITOR_API bool HasValidScriptAndGraph() const;
 
 	//Begin UObject interface
 	virtual void PostLoad()override;
diff --git a/Engine/Plugins/FX/Niagara/Source/NiagaraEditorWidgets/Private/SNiagaraOverviewStack.cpp b/Engine/Plugins/FX/Niagara/Source/NiagaraEditorWidgets/Private/SNiagaraOverviewStack.cpp
--- a/Engine/Plugins/FX/Niagara/Source/NiagaraEditorWidgets/Private/SNiagaraOverviewStack.cpp
+++ b/Engine/Plugins/FX/Niagara/Source/NiagaraEditorWidgets/Private/SNiagaraOverviewStack.cpp
@@ -925,7 +925,7 @@ TSharedRef<ITableRow> SNiagaraOverviewStack::OnGenerateRowForEntry(UNiagaraStack
 			];
 
 		// in case we have at least one inline input, we add a wrap box to make sure we have enough space for both the module name and parameters
-		if(StackModuleItem && StackModuleItem->GetModuleNode().ScriptIsValid() && StackModuleItem->GetInlineParameterInputs().Num() > 0)
+		if(StackModuleItem && StackModuleItem->GetModuleNode().HasValidScriptAndGraph() && StackModuleItem->GetInlineParameterInputs().Num() > 0)
 		{
 			ContentBox = SNew(SHorizontalBox)
 			+ SHorizontalBox::Slot()
