- ResolveAliases rather than use CreateAliasedModuleParameterHandle to correctly adjust namespaces

#rnx

diff --git a/Engine/Plugins/FX/Niagara/Source/NiagaraEditor/Private/ViewModels/Stack/NiagaraStackGraphUtilities.cpp b/Engine/Plugins/FX/Niagara/Source/NiagaraEditor/Private/ViewModels/Stack/NiagaraStackGraphUtilities.cpp
--- a/Engine/Plugins/FX/Niagara/Source/NiagaraEditor/Private/ViewModels/Stack/NiagaraStackGraphUtilities.cpp
+++ b/Engine/Plugins/FX/Niagara/Source/NiagaraEditor/Private/ViewModels/Stack/NiagaraStackGraphUtilities.cpp
@@ -3359,8 +3359,10 @@ void FNiagaraStackGraphUtilities::SynchronizeReferencingMapPinsWithFunctionCall(
 						FNiagaraParameterHandle BoundVariableHandle(BoundVariable->GetName());
 						if (LinkedOutputHandle.GetName() != BoundVariableHandle.GetName())
 						{
-							FNiagaraParameterHandle UpdatedLinkedOutputHandle = FNiagaraParameterHandle::CreateAliasedModuleParameterHandle(BoundVariableHandle, &InFunctionCallNode);
-							StackLinkedOutputPin->PinName = UpdatedLinkedOutputHandle.GetParameterHandleString();
+							FNiagaraVariable BoundVariableResolvedName = FNiagaraUtilities::ResolveAliases(*BoundVariable,
+								FNiagaraAliasContext().ChangeModuleToModuleName(InFunctionCallNode.GetFunctionName())
+							);
+							StackLinkedOutputPin->PinName = BoundVariableResolvedName.GetName();
 							TargetFunctionCallNode->UpdateInputNameBinding(BoundGuid, BoundVariableHandle.GetName());
 						}
 						if (LinkedOutputVariable.GetType() != BoundVariable->GetType())
