Fixed breakage of Niagara default value tests

Apparently there are cases in which a script var that's set to "Value" should maintain pin connections. Severing these pin connections led to 2 particles color values being cleared

diff --git a/Engine/Plugins/FX/Niagara/Source/NiagaraEditor/Private/NiagaraNodeParameterMapGet.cpp b/Engine/Plugins/FX/Niagara/Source/NiagaraEditor/Private/NiagaraNodeParameterMapGet.cpp
--- a/Engine/Plugins/FX/Niagara/Source/NiagaraEditor/Private/NiagaraNodeParameterMapGet.cpp
+++ b/Engine/Plugins/FX/Niagara/Source/NiagaraEditor/Private/NiagaraNodeParameterMapGet.cpp
@@ -400,7 +400,6 @@ void UNiagaraNodeParameterMapGet::SynchronizeDefaultInputPin(UEdGraphPin* Defaul
 	if (ScriptVar) {
 		if (ScriptVar->DefaultMode == ENiagaraDefaultMode::Value)
 		{
-			DefaultPin->BreakAllPinLinks();
 			DefaultPin->bNotConnectable = true;
 			DefaultPin->bDefaultValueIsReadOnly = true;
 		}
