- Fix for incorrect motion vectors on Niagara Mesh Particles

#rnx

diff --git a/Engine/Plugins/FX/Niagara/Shaders/Private/NiagaraMeshVertexFactory.ush b/Engine/Plugins/FX/Niagara/Shaders/Private/NiagaraMeshVertexFactory.ush
--- a/Engine/Plugins/FX/Niagara/Shaders/Private/NiagaraMeshVertexFactory.ush
+++ b/Engine/Plugins/FX/Niagara/Shaders/Private/NiagaraMeshVertexFactory.ush
@@ -801,7 +801,7 @@ FVertexFactoryInterpolantsVSToPS VertexFactoryGetInterpolantsVSToPS(FVertexFacto
 float4 VertexFactoryGetPreviousWorldPosition(FVertexFactoryInput Input, FVertexFactoryIntermediates Intermediates)
 {
     FLWCMatrix Transform = GetParticlePrevTransform(Intermediates);
-    float3 WorldPosition = LWCToFloat(LWCAdd(LWCMultiply(Input.Position.xyz, Transform), ResolvedView.PreViewTranslation));
+    float3 WorldPosition = LWCToFloat(LWCAdd(LWCMultiply(Input.Position.xyz, Transform), ResolvedView.PrevPreViewTranslation));
     return float4(WorldPosition, Input.Position.w);
 }
 
