Added support for sections to Naigara Component as well as a category with buttons to restart the sim and debug the sim.
Added ability to summon Niagara debugger to view given variables, emitter, system, component etc in context in editor.

diff --git a/Engine/Plugins/FX/Niagara/Source/NiagaraEditor/Private/Customizations/NiagaraComponentDetails.cpp b/Engine/Plugins/FX/Niagara/Source/NiagaraEditor/Private/Customizations/NiagaraComponentDetails.cpp
--- a/Engine/Plugins/FX/Niagara/Source/NiagaraEditor/Private/Customizations/NiagaraComponentDetails.cpp
+++ b/Engine/Plugins/FX/Niagara/Source/NiagaraEditor/Private/Customizations/NiagaraComponentDetails.cpp
@@ -48,6 +48,8 @@
 #include "Widgets/Layout/SBox.h"
 #include "UObject/WeakObjectPtr.h"
 #include "NiagaraUserRedirectionParameterStore.h"
+#include "Widgets/Layout/SUniformGridPanel.h"
+#include "Widgets/SNiagaraDebugger.h"
 
 #define LOCTEXT_NAMESPACE "NiagaraComponentDetails"
 
@@ -577,11 +579,27 @@ void FNiagaraComponentDetails::OnWorldDestroyed(class UWorld* InWorld)
 
 void FNiagaraComponentDetails::CustomizeDetails(IDetailLayoutBuilder& DetailBuilder)
 {
+
+	FPropertyEditorModule& PropertyModule = FModuleManager::GetModuleChecked<FPropertyEditorModule>("PropertyEditor");
+
 	Builder = &DetailBuilder;
 
 	static const FName ParamCategoryName = TEXT("NiagaraComponent_Parameters");
+	static const FName ParamUtilitiesName = TEXT("NiagaraComponent_Utilities");
 	static const FName ScriptCategoryName = TEXT("Parameters");
 
+	static bool bFirstTime = true;
+	if (bFirstTime)
+	{
+		const FText DisplayName = LOCTEXT("Particles", "Particles");
+		TSharedRef<FPropertySection> Section = PropertyModule.FindOrCreateSection("NiagaraComponent", "Particles", DisplayName);
+		Section->AddCategory(TEXT("Niagara"));
+		Section->AddCategory(ParamCategoryName);
+		Section->AddCategory(ParamUtilitiesName); 
+		Section->AddCategory(TEXT("Activation"));
+		bFirstTime = false;
+	}
+
 	TSharedPtr<IPropertyHandle> LocalOverridesPropertyHandle = DetailBuilder.GetProperty("OverrideParameters");
 	if (LocalOverridesPropertyHandle.IsValid())
 	{
@@ -612,7 +630,6 @@ void FNiagaraComponentDetails::CustomizeDetails(IDetailLayoutBuilder& DetailBuil
 	if (ObjectsCustomized.Num() == 1 && ObjectsCustomized[0]->IsA<UNiagaraComponent>())
 	{
 		Component = CastChecked<UNiagaraComponent>(ObjectsCustomized[0].Get());
-
 		if (GEngine)
 		{
 			GEngine->OnWorldDestroyed().AddRaw(this, &FNiagaraComponentDetails::OnWorldDestroyed);
@@ -634,6 +651,118 @@ void FNiagaraComponentDetails::CustomizeDetails(IDetailLayoutBuilder& DetailBuil
 				.Text(LOCTEXT("OverrideParameterMultiselectionUnsupported", "Multiple override parameter sets cannot be edited simultaneously."))
 			];
 	}
+
+
+	IDetailCategoryBuilder& CustomCategory = DetailBuilder.EditCategory(ParamUtilitiesName, LOCTEXT("ParamUtilsCategoryName", "Niagara Utilities"), ECategoryPriority::Important);
+
+	CustomCategory.AddCustomRow(FText::GetEmpty())
+		.WholeRowContent()
+		.HAlign(HAlign_Left)
+		[
+			SNew(SBox)
+			.MaxDesiredWidth(300.f)
+			[
+				SNew(SUniformGridPanel)
+				.SlotPadding(2.0f)
+				+ SUniformGridPanel::Slot(0, 0)
+				[
+					SNew(SButton)
+					.OnClicked(this, &FNiagaraComponentDetails::OnDebugSelectedSystem)
+					.ToolTipText(LOCTEXT("DebugButtonTooltip", "Open Niagara Debugger and point to the first selected particle system"))
+					.HAlign(HAlign_Center)
+					[
+						SNew(STextBlock)
+						.Text(LOCTEXT("DebugButton", "Debug"))
+					]
+				]
+				+ SUniformGridPanel::Slot(1, 0)
+				[
+					SNew(SButton)
+					.OnClicked(this, &FNiagaraComponentDetails::OnResetSelectedSystem)
+					.ToolTipText(LOCTEXT("ResetEmitterButtonTooltip", "Resets the selected particle systems."))
+					.HAlign(HAlign_Center)
+					[
+						SNew(STextBlock)
+						.Text(LOCTEXT("ResetEmitterButton", "Reset"))
+					]
+				]
+			]
+		];
 }
 
+FReply FNiagaraComponentDetails::OnResetSelectedSystem()
+{
+	if (!Builder)
+		return FReply::Handled();
+
+	const TArray< TWeakObjectPtr<UObject> >& SelectedObjects = Builder->GetSelectedObjects();
+
+	for (int32 Idx = 0; Idx < SelectedObjects.Num(); ++Idx)
+	{
+		if (SelectedObjects[Idx].IsValid())
+		{
+			if (AActor* Actor = Cast<AActor>(SelectedObjects[Idx].Get()))
+			{
+				for (UActorComponent* AC : Actor->GetComponents())
+				{
+					UNiagaraComponent* NiagaraComponent = Cast<UNiagaraComponent>(AC);
+					if (NiagaraComponent)
+					{
+						NiagaraComponent->Activate(true);
+						NiagaraComponent->ReregisterComponent();
+					}
+				}
+			}
+			else if (UNiagaraComponent* NiagaraComponent = Cast<UNiagaraComponent>(SelectedObjects[Idx].Get()))
+			{
+				NiagaraComponent->Activate(true);
+				NiagaraComponent->ReregisterComponent();
+			}
+			
+		}
+	}
+
+	return FReply::Handled();
+}
+
+FReply FNiagaraComponentDetails::OnDebugSelectedSystem()
+{
+	if (!Builder)
+		return FReply::Handled();
+
+	const TArray< TWeakObjectPtr<UObject> >& SelectedObjects = Builder->GetSelectedObjects();
+
+	UNiagaraComponent* NiagaraComponentToUse = nullptr;
+	for (int32 Idx = 0; Idx < SelectedObjects.Num(); ++Idx)
+	{
+		if (SelectedObjects[Idx].IsValid())
+		{
+			if (AActor* Actor = Cast<AActor>(SelectedObjects[Idx].Get()))
+			{
+				for (UActorComponent* AC : Actor->GetComponents())
+				{
+					UNiagaraComponent* NiagaraComponent = Cast<UNiagaraComponent>(AC);
+					if (NiagaraComponent)
+					{
+						NiagaraComponentToUse = NiagaraComponent;
+						break;
+					}
+				}
+			}
+			else if (UNiagaraComponent* NiagaraComponent = Cast<UNiagaraComponent>(SelectedObjects[Idx].Get()))
+			{
+				NiagaraComponentToUse = NiagaraComponent;
+				break;
+			}
+		}
+	}
+
+	if (NiagaraComponentToUse)
+	{
+		//FSpawnTabArgs Args;
+		SNiagaraDebugger::InvokeDebugger(NiagaraComponentToUse);
+	}
+
+	return FReply::Handled();
+}
 #undef LOCTEXT_NAMESPACE
diff --git a/Engine/Plugins/FX/Niagara/Source/NiagaraEditor/Private/Customizations/NiagaraComponentDetails.h b/Engine/Plugins/FX/Niagara/Source/NiagaraEditor/Private/Customizations/NiagaraComponentDetails.h
--- a/Engine/Plugins/FX/Niagara/Source/NiagaraEditor/Private/Customizations/NiagaraComponentDetails.h
+++ b/Engine/Plugins/FX/Niagara/Source/NiagaraEditor/Private/Customizations/NiagaraComponentDetails.h
@@ -36,6 +36,8 @@ protected:
 
 	FNiagaraComponentDetails();
 
+	FReply OnResetSelectedSystem();
+	FReply OnDebugSelectedSystem();
 private:
 	TWeakObjectPtr<UNiagaraComponent> Component;
 	IDetailLayoutBuilder* Builder;
diff --git a/Engine/Plugins/FX/Niagara/Source/NiagaraEditor/Private/NiagaraEditorModule.cpp b/Engine/Plugins/FX/Niagara/Source/NiagaraEditor/Private/NiagaraEditorModule.cpp
--- a/Engine/Plugins/FX/Niagara/Source/NiagaraEditor/Private/NiagaraEditorModule.cpp
+++ b/Engine/Plugins/FX/Niagara/Source/NiagaraEditor/Private/NiagaraEditorModule.cpp
@@ -1757,20 +1757,31 @@ void FNiagaraEditorModule::OnExecParticleInvoked(const TCHAR* Str)
 				AActor* Actor = static_cast<AActor*>(*It);
 				checkSlow(Actor->IsA(AActor::StaticClass()));
 
-				ANiagaraActor* Emitter = Cast<ANiagaraActor>(Actor);
-				if (Emitter)
+				for (UActorComponent* AC : Actor->GetComponents())
 				{
-					Emitter->ResetInLevel();
+					UNiagaraComponent* NiagaraComponent = Cast<UNiagaraComponent>(AC);
+					if (NiagaraComponent)
+					{
+						NiagaraComponent->Activate(true);
+						NiagaraComponent->ReregisterComponent();
+					}
 				}
 			}
 		}
 		else if (FParse::Command(&Str, TEXT("ALL")))
 		{
 			// Reset ALL emitters in the level
-			for (TObjectIterator<ANiagaraActor> It; It; ++It)
+			for (TObjectIterator<AActor> It; It; ++It)
 			{
-				ANiagaraActor* Emitter = *It;
-				Emitter->ResetInLevel();
+				for (UActorComponent* AC : It->GetComponents())
+				{
+					UNiagaraComponent* NiagaraComponent = Cast<UNiagaraComponent>(AC);
+					if (NiagaraComponent)
+					{
+						NiagaraComponent->Activate(true);
+						NiagaraComponent->ReregisterComponent();
+					}
+				}
 			}
 		}
 	}
diff --git a/Engine/Plugins/FX/Niagara/Source/NiagaraEditor/Private/NiagaraOverviewNode.cpp b/Engine/Plugins/FX/Niagara/Source/NiagaraEditor/Private/NiagaraOverviewNode.cpp
--- a/Engine/Plugins/FX/Niagara/Source/NiagaraEditor/Private/NiagaraOverviewNode.cpp
+++ b/Engine/Plugins/FX/Niagara/Source/NiagaraEditor/Private/NiagaraOverviewNode.cpp
@@ -221,6 +221,18 @@ void UNiagaraOverviewNode::GetNodeContextMenuActions(class UToolMenu* Menu, clas
 							FExecuteAction::CreateSP(EmitterHandleViewModel, &FNiagaraEmitterHandleViewModel::SetIsRenamePending, true)
 						)
 					);
+
+					Section.AddSeparator("DebugEmitterSplit");
+					Section.AddMenuEntry(
+						"DebugEmitter",
+						LOCTEXT("DebugEmitter", "Watch Emitter In Niagara Debugger"),
+						LOCTEXT("DebugEmitterToolTip", "Open Niagara Debugger and track this emitter in the world"),
+						FSlateIcon(),
+						FUIAction(
+							FExecuteAction::CreateSP(EmitterHandleViewModel, &FNiagaraEmitterHandleViewModel::BeginDebugEmitter)
+						)
+					);
+					Section.AddSeparator("DebugEmitterSplit2");
 				}
 
 				Section.AddMenuEntry(
diff --git a/Engine/Plugins/FX/Niagara/Source/NiagaraEditor/Private/Toolkits/NiagaraSystemToolkit.cpp b/Engine/Plugins/FX/Niagara/Source/NiagaraEditor/Private/Toolkits/NiagaraSystemToolkit.cpp
--- a/Engine/Plugins/FX/Niagara/Source/NiagaraEditor/Private/Toolkits/NiagaraSystemToolkit.cpp
+++ b/Engine/Plugins/FX/Niagara/Source/NiagaraEditor/Private/Toolkits/NiagaraSystemToolkit.cpp
@@ -773,15 +773,7 @@ bool FNiagaraSystemToolkit::IsDrawOptionEnabled(int32 Element) const
 
 void FNiagaraSystemToolkit::OpenDebugHUD()
 {
-#if WITH_NIAGARA_DEBUGGER
-	TSharedPtr<SDockTab> DebugTab = FGlobalTabmanager::Get()->TryInvokeTab(SNiagaraDebugger::DebugWindowName);
-
-	if (DebugTab.IsValid())
-	{
-		TSharedRef<SNiagaraDebugger> Content = StaticCastSharedRef<SNiagaraDebugger>(DebugTab->GetContent());
-		Content->FocusDebugTab();
-	}
-#endif
+	SNiagaraDebugger::InvokeDebugger(&SystemViewModel->GetSystem());
 }
 
 void FNiagaraSystemToolkit::OpenDebugOutliner()
diff --git a/Engine/Plugins/FX/Niagara/Source/NiagaraEditor/Private/ViewModels/NiagaraEmitterHandleViewModel.cpp b/Engine/Plugins/FX/Niagara/Source/NiagaraEditor/Private/ViewModels/NiagaraEmitterHandleViewModel.cpp
--- a/Engine/Plugins/FX/Niagara/Source/NiagaraEditor/Private/ViewModels/NiagaraEmitterHandleViewModel.cpp
+++ b/Engine/Plugins/FX/Niagara/Source/NiagaraEditor/Private/ViewModels/NiagaraEmitterHandleViewModel.cpp
@@ -23,6 +23,7 @@
 #include "ViewModels/Stack/NiagaraStackRenderItemGroup.h"
 #include "ViewModels/Stack/NiagaraStackRendererItem.h"
 #include "NiagaraMessages.h"
+#include "Widgets/SNiagaraDebugger.h"
 
 #define LOCTEXT_NAMESPACE "EmitterHandleViewModel"
 
@@ -352,6 +353,12 @@ void FNiagaraEmitterHandleViewModel::SetIsRenamePending(bool bInIsRenamePending)
 	bIsRenamePending = bInIsRenamePending;
 }
 
+void FNiagaraEmitterHandleViewModel::BeginDebugEmitter()
+{
+	if (EmitterHandle)
+		SNiagaraDebugger::InvokeDebugger(*EmitterHandle);
+}
+
 FNiagaraEmitterHandleViewModel::FOnPropertyChanged& FNiagaraEmitterHandleViewModel::OnPropertyChanged()
 {
 	return OnPropertyChangedDelegate;
diff --git a/Engine/Plugins/FX/Niagara/Source/NiagaraEditor/Private/ViewModels/NiagaraParameterPanelViewModel.cpp b/Engine/Plugins/FX/Niagara/Source/NiagaraEditor/Private/ViewModels/NiagaraParameterPanelViewModel.cpp
--- a/Engine/Plugins/FX/Niagara/Source/NiagaraEditor/Private/ViewModels/NiagaraParameterPanelViewModel.cpp
+++ b/Engine/Plugins/FX/Niagara/Source/NiagaraEditor/Private/ViewModels/NiagaraParameterPanelViewModel.cpp
@@ -37,6 +37,7 @@
 #include "HAL/PlatformApplicationMisc.h"
 #include "ScopedTransaction.h"
 #include "Widgets/SNullWidget.h"
+#include "Widgets/SNiagaraDebugger.h"
 
 #define LOCTEXT_NAMESPACE "NiagaraParameterPanelViewModel"
 
@@ -398,6 +399,22 @@ void INiagaraParameterPanelViewModel::DuplicateParameters(const TArray<FNiagaraP
 		
 }
 
+bool INiagaraParameterPanelViewModel::GetCanDebugParameters(const TArray<FNiagaraParameterPanelItem>& ItemsToDebug) const
+{
+	for (const FNiagaraParameterPanelItem& ItemToDebug : ItemsToDebug)
+	{
+		const FNiagaraVariableBase& Var = ItemToDebug.GetVariable();
+		if (Var.IsInNameSpace(FNiagaraConstants::UserNamespaceString) ||
+			Var.IsInNameSpace(FNiagaraConstants::SystemNamespaceString) ||
+			Var.IsInNameSpace(FNiagaraConstants::EmitterNamespaceString) ||
+			Var.IsInNameSpace(FNiagaraConstants::ParticleAttributeNamespaceString))
+		{
+			return true;
+		}
+	}
+	return false;
+}
+
 bool INiagaraParameterPanelViewModel::GetCanDuplicateParameterAndToolTip(const TArray<FNiagaraParameterPanelItem>& ItemsToDuplicate, FText& OutCanDuplicateParameterToolTip) const
 {
 	for(const FNiagaraParameterPanelItem& ItemToDuplicate : ItemsToDuplicate)
@@ -1316,6 +1333,15 @@ TSharedPtr<SWidget> FNiagaraSystemToolkitParameterPanelViewModel::CreateContextM
 				LOCTEXT("ChangeNamespaceModifierToolTip", "Edit the namespace modifier for the selected parameter."),
 				FNewMenuDelegate::CreateSP(this, &FNiagaraSystemToolkitParameterPanelViewModel::GetChangeNamespaceModifierSubMenu, false, SelectedItem));
 			
+			MenuBuilder.AddMenuSeparator();
+			bool bCanDebugParameter = GetCanDebugParameters(Items);
+			MenuBuilder.AddMenuEntry(
+				LOCTEXT("DebugParameter", "Watch Parameter In Niagara Debugger"),
+				LOCTEXT("DebugParameterToolTip", "Open the Niagara Debugger and add this Parameter to the list of tracked parameters.\r\nParticles parameters will show up alongside each particle.\r\nSystem, Emitter, and User parameters will be attached to the System in the world."), 
+				FSlateIcon(),
+				FUIAction(
+					FExecuteAction::CreateSP(this, &FNiagaraSystemToolkitParameterPanelViewModel::DebugParameters, Items),
+					FCanExecuteAction::CreateLambda([bCanDebugParameter]() {return bCanDebugParameter; })));
 
 			MenuBuilder.AddMenuSeparator();
 
@@ -1351,6 +1377,28 @@ FNiagaraParameterUtilities::EParameterContext FNiagaraSystemToolkitParameterPane
 	return FNiagaraParameterUtilities::EParameterContext::System;
 }
 
+void FNiagaraSystemToolkitParameterPanelViewModel::DebugParameters(const TArray<FNiagaraParameterPanelItem> Items) const
+{
+	TArray<FNiagaraVariableBase> Attributes;
+	for (const FNiagaraParameterPanelItem& Item : Items)
+	{
+		Attributes.Emplace(Item.GetVariable());
+	}
+
+	const TArray<FGuid> EmitterHandles = SystemViewModel->GetSelectionViewModel()->GetSelectedEmitterHandleIds();
+	UNiagaraSystem* System = &SystemViewModel->GetSystem();
+	if (System)
+	{
+		TArray< FNiagaraEmitterHandle> Handles;
+		for (const FNiagaraEmitterHandle& Handle : System->GetEmitterHandles())
+		{
+			if (EmitterHandles.Contains(Handle.GetId() ))
+				Handles.Add(Handle);
+		}
+		SNiagaraDebugger::InvokeDebugger(System, Handles, Attributes);
+	}
+}
+
 const TArray<FNiagaraParameterPanelCategory>& FNiagaraSystemToolkitParameterPanelViewModel::GetDefaultCategories() const
 {
 //	const UNiagaraEditorSettings* NiagaraEditorSettings = GetDefault<UNiagaraEditorSettings>();
diff --git a/Engine/Plugins/FX/Niagara/Source/NiagaraEditor/Private/Widgets/SNiagaraDebugger.cpp b/Engine/Plugins/FX/Niagara/Source/NiagaraEditor/Private/Widgets/SNiagaraDebugger.cpp
--- a/Engine/Plugins/FX/Niagara/Source/NiagaraEditor/Private/Widgets/SNiagaraDebugger.cpp
+++ b/Engine/Plugins/FX/Niagara/Source/NiagaraEditor/Private/Widgets/SNiagaraDebugger.cpp
@@ -29,6 +29,8 @@
 #include "Customizations/NiagaraDebugHUDCustomization.h"
 #include "Customizations/NiagaraOutlinerCustomization.h"
 #include "IStructureDetailsView.h"
+#include "NiagaraComponent.h"
+#include "NiagaraConstants.h"
 
 // the SessionFrontend is a "target" developer tool for talking to other devices, etc, and may get disabled
 #define WITH_SESSION_FRONTEND 	WITH_UNREAL_TARGET_DEVELOPER_TOOLS
@@ -802,6 +804,241 @@ void SNiagaraDebugger::UnregisterTabSpawner()
 	}
 }
 
+void SNiagaraDebugger::InvokeDebugger(UNiagaraComponent* InComponent)
+{
+	TSharedPtr<SDockTab> DebugTab = FGlobalTabmanager::Get()->TryInvokeTab(SNiagaraDebugger::DebugWindowName);
+
+	if (DebugTab.IsValid())
+	{
+		TSharedRef<SNiagaraDebugger> Content = StaticCastSharedRef<SNiagaraDebugger>(DebugTab->GetContent());
+		Content->FocusDebugTab();
+
+		UNiagaraDebugHUDSettings* HudSettings = GetMutableDefault<UNiagaraDebugHUDSettings>();
+		if (HudSettings && InComponent)
+		{
+			HudSettings->Data.bComponentFilterEnabled = true;
+			HudSettings->Data.ComponentFilter = InComponent->GetName();
+			
+			UNiagaraSystem* System = InComponent->GetAsset();
+			if (System)
+			{
+				HudSettings->Data.SystemFilter = System->GetName();
+				HudSettings->Data.bSystemFilterEnabled = true;
+			}
+
+			AActor* Actor = InComponent->GetOwner();
+			if (IsValid(Actor))
+			{
+				HudSettings->Data.ActorFilter = Actor->GetActorNameOrLabel();
+				HudSettings->Data.bActorFilterEnabled = true;
+			}
+
+			HudSettings->Data.bEmitterFilterEnabled = false;
+
+			if (HudSettings->Data.SystemDebugVerbosity == ENiagaraDebugHudVerbosity::None)
+				HudSettings->Data.SystemDebugVerbosity = ENiagaraDebugHudVerbosity::Basic;
+
+			HudSettings->NotifyPropertyChanged();
+		}
+	}
+}
+
+void SNiagaraDebugger::InvokeDebugger(UNiagaraSystem* InSystem)
+{
+	TSharedPtr<SDockTab> DebugTab = FGlobalTabmanager::Get()->TryInvokeTab(SNiagaraDebugger::DebugWindowName);
+
+	if (DebugTab.IsValid())
+	{
+		TSharedRef<SNiagaraDebugger> Content = StaticCastSharedRef<SNiagaraDebugger>(DebugTab->GetContent());
+		Content->FocusDebugTab();
+
+		UNiagaraDebugHUDSettings* HudSettings = GetMutableDefault<UNiagaraDebugHUDSettings>();
+		if (HudSettings && InSystem)
+		{
+			HudSettings->Data.bComponentFilterEnabled = false;
+			HudSettings->Data.ComponentFilter = TEXT("");
+
+			HudSettings->Data.ActorFilter = TEXT("");
+			HudSettings->Data.bActorFilterEnabled = false;
+
+			if (InSystem)
+			{
+				HudSettings->Data.SystemFilter = InSystem->GetName();
+				HudSettings->Data.bSystemFilterEnabled = true;
+			}
+
+			if (HudSettings->Data.SystemDebugVerbosity == ENiagaraDebugHudVerbosity::None)
+				HudSettings->Data.SystemDebugVerbosity = ENiagaraDebugHudVerbosity::Basic;
+
+			HudSettings->NotifyPropertyChanged();
+		}
+	}
+}
+
+
+void SNiagaraDebugger::InvokeDebugger(FNiagaraEmitterHandle& InEmitterHandle)
+{
+	TSharedPtr<SDockTab> DebugTab = FGlobalTabmanager::Get()->TryInvokeTab(SNiagaraDebugger::DebugWindowName);
+
+	if (DebugTab.IsValid())
+	{
+		TSharedRef<SNiagaraDebugger> Content = StaticCastSharedRef<SNiagaraDebugger>(DebugTab->GetContent());
+		Content->FocusDebugTab();
+
+		UNiagaraDebugHUDSettings* HudSettings = GetMutableDefault<UNiagaraDebugHUDSettings>();
+		if (HudSettings && InEmitterHandle.IsValid() && InEmitterHandle.GetInstance())
+		{
+			HudSettings->Data.bComponentFilterEnabled = false;
+			HudSettings->Data.ComponentFilter = TEXT("");
+
+			HudSettings->Data.ActorFilter = TEXT("");
+			HudSettings->Data.bActorFilterEnabled = false;
+
+			UNiagaraSystem* System = Cast<UNiagaraSystem>(InEmitterHandle.GetInstance()->GetOuter());
+			if (System)
+			{
+				HudSettings->Data.SystemFilter = System->GetName();
+				HudSettings->Data.bSystemFilterEnabled = true;
+			}
+
+			HudSettings->Data.EmitterFilter = InEmitterHandle.GetUniqueInstanceName();
+			HudSettings->Data.bEmitterFilterEnabled = true;
+
+			if (HudSettings->Data.SystemDebugVerbosity == ENiagaraDebugHudVerbosity::None)
+				HudSettings->Data.SystemDebugVerbosity = ENiagaraDebugHudVerbosity::Basic;
+
+			if (HudSettings->Data.SystemEmitterVerbosity == ENiagaraDebugHudVerbosity::None)
+				HudSettings->Data.SystemEmitterVerbosity = ENiagaraDebugHudVerbosity::Verbose;
+
+			HudSettings->NotifyPropertyChanged();
+		}
+	}
+}
+
+void SNiagaraDebugger::InvokeDebugger(UNiagaraSystem* InSystem, TArray<FNiagaraEmitterHandle>& InSelectedHandles, TArray<FNiagaraVariableBase>& InAttributes)
+{
+	TSharedPtr<SDockTab> DebugTab = FGlobalTabmanager::Get()->TryInvokeTab(SNiagaraDebugger::DebugWindowName);
+
+	if (DebugTab.IsValid())
+	{
+		TSharedRef<SNiagaraDebugger> Content = StaticCastSharedRef<SNiagaraDebugger>(DebugTab->GetContent());
+		Content->FocusDebugTab();
+
+		UNiagaraDebugHUDSettings* HudSettings = GetMutableDefault<UNiagaraDebugHUDSettings>();
+		if (HudSettings && InSystem)
+		{
+			HudSettings->Data.bComponentFilterEnabled = false;
+			HudSettings->Data.ComponentFilter = TEXT("");
+
+			HudSettings->Data.ActorFilter = TEXT("");
+			HudSettings->Data.bActorFilterEnabled = false;
+
+			if (InSystem)
+			{
+				HudSettings->Data.SystemFilter = InSystem->GetName();
+				HudSettings->Data.bSystemFilterEnabled = true;
+			}
+
+			if (InSelectedHandles.Num() > 0)
+			{
+				HudSettings->Data.EmitterFilter = InSelectedHandles[0].GetUniqueInstanceName();
+				HudSettings->Data.bEmitterFilterEnabled = true;
+			}
+
+			if (HudSettings->Data.SystemDebugVerbosity == ENiagaraDebugHudVerbosity::None)
+				HudSettings->Data.SystemDebugVerbosity = ENiagaraDebugHudVerbosity::Basic;
+
+			if (HudSettings->Data.SystemEmitterVerbosity == ENiagaraDebugHudVerbosity::None)
+				HudSettings->Data.SystemEmitterVerbosity = ENiagaraDebugHudVerbosity::Verbose;
+
+			TArray<FString> ParticleAttribNames;
+			TArray<FString> SystemAttribNames;
+			for (const FNiagaraVariableBase& Var : InAttributes)
+			{
+				if (Var.IsInNameSpace(FNiagaraConstants::UserNamespaceString) || 
+					Var.IsInNameSpace(FNiagaraConstants::SystemNamespaceString))
+				{
+					// No other work needed, use as-is
+					SystemAttribNames.AddUnique(Var.GetName().ToString());
+				}
+				else if (Var.IsInNameSpace(FNiagaraConstants::EmitterNamespaceString))
+				{
+					// Need to replace the emitter namespace with specific emitter name
+					if (InSelectedHandles.Num() > 0)
+					{
+						FString StrippedEmitterName = Var.GetName().ToString();
+						if (StrippedEmitterName.RemoveFromStart(FNiagaraConstants::EmitterNamespaceString + TEXT(".")))
+							SystemAttribNames.AddUnique(InSelectedHandles[0].GetUniqueInstanceName() + TEXT(".") + StrippedEmitterName);
+					}
+				}
+				else if (Var.IsInNameSpace(FNiagaraConstants::ParticleAttributeNamespaceString))
+				{
+					// Need to replace the particle namespace altogether
+					FString StrippedParticleName = Var.GetName().ToString();
+					if (StrippedParticleName.RemoveFromStart(FNiagaraConstants::ParticleAttributeNamespaceString + TEXT(".")))
+						ParticleAttribNames.AddUnique(StrippedParticleName);
+				}
+			}
+			
+			bool bHasSystemAttributes = SystemAttribNames.Num() > 0;
+			bool bHasParticleAttributes = ParticleAttribNames.Num() > 0;
+
+
+			if (bHasParticleAttributes)
+			{
+				HudSettings->Data.bShowParticleVariables = true;
+				for (const FString& Var : ParticleAttribNames)
+				{
+					FNiagaraDebugHUDVariable NewVar;
+					NewVar.bEnabled = true;
+					NewVar.Name = Var;
+					
+					bool bAdd = true;
+					for (const FNiagaraDebugHUDVariable& OldVar : HudSettings->Data.ParticlesVariables)
+					{
+						if (OldVar.Name == NewVar.Name)
+						{
+							bAdd = false;
+							break;
+						}
+					}
+
+					if (bAdd)
+						HudSettings->Data.ParticlesVariables.Add(NewVar);
+				}
+
+				if (InSelectedHandles.Num() > 0 && InSelectedHandles[0].GetInstance() && InSelectedHandles[0].GetInstance()->SimTarget == ENiagaraSimTarget::GPUComputeSim)
+					HudSettings->Data.bEnableGpuParticleReadback = true;
+			}
+
+			if (bHasSystemAttributes)
+			{
+				HudSettings->Data.bShowSystemVariables = true;
+				for (const FString& Var : SystemAttribNames)
+				{
+					FNiagaraDebugHUDVariable NewVar;
+					NewVar.bEnabled = true;
+					NewVar.Name = Var;
+					bool bAdd = true;
+					for (const FNiagaraDebugHUDVariable& OldVar : HudSettings->Data.SystemVariables)
+					{
+						if (OldVar.Name == NewVar.Name)
+						{
+							bAdd = false;
+							break;
+						}
+					}
+
+					if (bAdd)
+						HudSettings->Data.SystemVariables.Add(NewVar);
+				}
+			}
+
+			HudSettings->NotifyPropertyChanged();
+		}
+	}
+}
+
 TSharedRef<SDockTab> SNiagaraDebugger::SpawnNiagaraDebugger(const FSpawnTabArgs& Args)
 {
 	auto NomadTab = SNew(SDockTab)
diff --git a/Engine/Plugins/FX/Niagara/Source/NiagaraEditor/Private/Widgets/SNiagaraDebugger.h b/Engine/Plugins/FX/Niagara/Source/NiagaraEditor/Private/Widgets/SNiagaraDebugger.h
--- a/Engine/Plugins/FX/Niagara/Source/NiagaraEditor/Private/Widgets/SNiagaraDebugger.h
+++ b/Engine/Plugins/FX/Niagara/Source/NiagaraEditor/Private/Widgets/SNiagaraDebugger.h
@@ -30,6 +30,12 @@ public:
 	static void UnregisterTabSpawner();
 	static TSharedRef<class SDockTab> SpawnNiagaraDebugger(const class FSpawnTabArgs& Args);
 
+	static void InvokeDebugger(UNiagaraComponent* InComponent);
+	static void InvokeDebugger(UNiagaraSystem* InSystem);
+	static void InvokeDebugger(FNiagaraEmitterHandle& InEmitterHandle);
+
+	static void InvokeDebugger(UNiagaraSystem* InSystem, TArray<FNiagaraEmitterHandle>& InSelectedHandles, TArray<FNiagaraVariableBase>& InAttributes);
+
 	void FocusDebugTab();
 	void FocusOutlineTab();
 
diff --git a/Engine/Plugins/FX/Niagara/Source/NiagaraEditor/Public/ViewModels/NiagaraEmitterHandleViewModel.h b/Engine/Plugins/FX/Niagara/Source/NiagaraEditor/Public/ViewModels/NiagaraEmitterHandleViewModel.h
--- a/Engine/Plugins/FX/Niagara/Source/NiagaraEditor/Public/ViewModels/NiagaraEmitterHandleViewModel.h
+++ b/Engine/Plugins/FX/Niagara/Source/NiagaraEditor/Public/ViewModels/NiagaraEmitterHandleViewModel.h
@@ -131,6 +131,8 @@ public:
 	/** Remove a serialized message from the Emitter this viewmodel is managing. */
 	NIAGARAEDITOR_API void RemoveMessage(const FGuid& MessageKey) const;
 
+	void BeginDebugEmitter();
+
 private:
 	/** The system view model which owns this emitter handle view model. */
 	TWeakPtr<FNiagaraSystemViewModel> OwningSystemViewModelWeak;
diff --git a/Engine/Plugins/FX/Niagara/Source/NiagaraEditor/Public/ViewModels/NiagaraParameterPanelViewModel.h b/Engine/Plugins/FX/Niagara/Source/NiagaraEditor/Public/ViewModels/NiagaraParameterPanelViewModel.h
--- a/Engine/Plugins/FX/Niagara/Source/NiagaraEditor/Public/ViewModels/NiagaraParameterPanelViewModel.h
+++ b/Engine/Plugins/FX/Niagara/Source/NiagaraEditor/Public/ViewModels/NiagaraParameterPanelViewModel.h
@@ -200,6 +200,7 @@ public:
 
 	virtual void DuplicateParameters(const TArray<FNiagaraParameterPanelItem> ItemToDuplicate) const;
 	virtual bool GetCanDuplicateParameterAndToolTip(const TArray<FNiagaraParameterPanelItem>& ItemsToDuplicate, FText& OutCanDuplicateParameterToolTip) const;
+	bool GetCanDebugParameters(const TArray<FNiagaraParameterPanelItem>& ItemsToDebug) const;
 
 	virtual bool GetCanRenameParameterAndToolTip(const FNiagaraParameterPanelItem& ItemToRename, const FText& NewVariableNameText, bool bCheckEmptyNameText, FText& OutCanRenameParameterToolTip) const;
 
@@ -324,6 +325,7 @@ public:
 
 	TSharedRef<SWidget> CreateAddParameterMenuForAssignmentNode(UNiagaraNodeAssignment* AssignmentNode, const TSharedPtr<SComboButton>& AddButton) const;
 
+	void DebugParameters(const TArray<FNiagaraParameterPanelItem> ItemToDuplicate) const;
 private:
 	const TArray<UNiagaraGraph*> GetAllGraphsConst() const;
 
