Fixed crash when attempting to copy script variables that nodes reference. Sometimes, a parameter of name "None" would be encountered, but no script variable found. This is a short-term fix

diff --git a/Engine/Plugins/FX/Niagara/Source/NiagaraEditor/Private/ViewModels/NiagaraScriptGraphViewModel.cpp b/Engine/Plugins/FX/Niagara/Source/NiagaraEditor/Private/ViewModels/NiagaraScriptGraphViewModel.cpp
--- a/Engine/Plugins/FX/Niagara/Source/NiagaraEditor/Private/ViewModels/NiagaraScriptGraphViewModel.cpp
+++ b/Engine/Plugins/FX/Niagara/Source/NiagaraEditor/Private/ViewModels/NiagaraScriptGraphViewModel.cpp
@@ -296,9 +296,16 @@ void FNiagaraScriptGraphViewModel::CopySelectedNodes()
 	for(FNiagaraVariable& Var : Variables)
 	{
 		UNiagaraScriptVariable* ScriptVariable = GetGraph()->GetScriptVariable(Var);
-		ensureAlwaysMsgf(ScriptVariable != nullptr, TEXT("Script variable should already exist when we copy nodes referencing them!"));
-		UNiagaraScriptVariable* CopiedVariable = Cast<UNiagaraScriptVariable>(StaticDuplicateObject(ScriptVariable, ClipboardContent));
-		ClipboardContent->ScriptVariables.AddUnique({*CopiedVariable});
+
+		if(ScriptVariable)
+		{
+			UNiagaraScriptVariable* CopiedVariable = Cast<UNiagaraScriptVariable>(StaticDuplicateObject(ScriptVariable, ClipboardContent));
+			ClipboardContent->ScriptVariables.AddUnique({*CopiedVariable});
+		}
+		else
+		{
+			UE_LOG(LogNiagaraEditor, Log, TEXT("Variable %s was encountered as parameter during a copy nodes operation, but the corresponding script variable could not be found."), *(Var.GetName().ToString()));
+		}
 	}
 	
 	FNiagaraEditorModule::Get().GetClipboard().SetClipboardContent(ClipboardContent);
