- Fix for previous external cbuffer using the current data
- Fix for not calling PostTick on GPU Context which meant we didn't push the data from current -> previous

diff --git a/Engine/Plugins/FX/Niagara/Source/Niagara/Private/NiagaraGPUSystemTick.cpp b/Engine/Plugins/FX/Niagara/Source/Niagara/Private/NiagaraGPUSystemTick.cpp
--- a/Engine/Plugins/FX/Niagara/Source/Niagara/Private/NiagaraGPUSystemTick.cpp
+++ b/Engine/Plugins/FX/Niagara/Source/Niagara/Private/NiagaraGPUSystemTick.cpp
@@ -169,6 +169,9 @@ void FNiagaraGPUSystemTick::Init(FNiagaraSystemInstance* InSystemInstance)
 
 			GPUContext->CombinedParamStore.CopyParameterDataToPaddedBuffer(InstanceData->ExternalParamData, ParmSize);
 
+			// Calling PostTick will push current -> previous parameters this must be done after copying the parameter data
+			GPUContext->PostTick();
+
 			InstanceData->bStartNewOverlapGroup = bStartNewOverlapGroup;
 			bStartNewOverlapGroup = false;
 
@@ -304,9 +307,9 @@ const uint8* FNiagaraGPUSystemTick::GetUniformBufferSource(EUniformBufferType Ty
 		}
 		case UBT_External:
 		{
-			// External is special and interpolated parameters are already included inside of the combined parameter store
-			check(Instance);
-			return Instance->ExternalParamData;
+			// External parameters are pushed from the combined parameters store, split into two where first half is current second is previous
+			check(Instance && Instance->Context);
+			return Instance->ExternalParamData + (Current ? 0 : Instance->Context->ExternalCBufferLayoutSize);
 		}
 	}
 
