- Add support for empty cube UAV

#rnx

diff --git a/Engine/Plugins/FX/Niagara/Source/Niagara/Private/NiagaraEmptyUAVPool.cpp b/Engine/Plugins/FX/Niagara/Source/Niagara/Private/NiagaraEmptyUAVPool.cpp
--- a/Engine/Plugins/FX/Niagara/Source/Niagara/Private/NiagaraEmptyUAVPool.cpp
+++ b/Engine/Plugins/FX/Niagara/Source/Niagara/Private/NiagaraEmptyUAVPool.cpp
@@ -79,6 +79,17 @@ FRHIUnorderedAccessView* FNiagaraEmptyUAVPool::GetEmptyUAVFromPool(FRHICommandLi
 				break;
 			}
 			
+			case ENiagaraEmptyUAVType::TextureCube:
+			{
+				const FRHITextureCreateDesc Desc =
+					FRHITextureCreateDesc::CreateCube(ResourceName, 1, Format)
+					.SetFlags(ETextureCreateFlags::ShaderResource | ETextureCreateFlags::UAV);
+
+				NewUAV.Texture = RHICreateTexture(Desc);
+				NewUAV.UAV = RHICreateUnorderedAccessView(NewUAV.Texture, 0);
+				break;
+			}
+
 			default:
 			{
 				checkNoEntry();
diff --git a/Engine/Plugins/FX/Niagara/Source/Niagara/Public/NiagaraEmptyUAVPool.h b/Engine/Plugins/FX/Niagara/Source/Niagara/Public/NiagaraEmptyUAVPool.h
--- a/Engine/Plugins/FX/Niagara/Source/Niagara/Public/NiagaraEmptyUAVPool.h
+++ b/Engine/Plugins/FX/Niagara/Source/Niagara/Public/NiagaraEmptyUAVPool.h
@@ -28,6 +28,7 @@ enum class ENiagaraEmptyUAVType
 	Texture2D,
 	Texture2DArray,
 	Texture3D,
+	TextureCube,
 	Num
 };
 
