- Detect if an array is RW on the GPU based upon functions used

#rnx

diff --git a/Engine/Plugins/FX/Niagara/Source/Niagara/Classes/NiagaraDataInterfaceArray.h b/Engine/Plugins/FX/Niagara/Source/Niagara/Classes/NiagaraDataInterfaceArray.h
--- a/Engine/Plugins/FX/Niagara/Source/Niagara/Classes/NiagaraDataInterfaceArray.h
+++ b/Engine/Plugins/FX/Niagara/Source/Niagara/Classes/NiagaraDataInterfaceArray.h
@@ -25,7 +25,7 @@ struct INDIArrayProxyBase : public FNiagaraDataInterfaceProxyRW
 	virtual bool CopyToInternal(INDIArrayProxyBase* Destination) const = 0;
 	virtual bool Equals(const INDIArrayProxyBase* Other) const = 0;
 	virtual int32 PerInstanceDataSize() const = 0;
-	virtual bool InitPerInstanceData(void* InPerInstanceData, FNiagaraSystemInstance* SystemInstance) = 0;
+	virtual bool InitPerInstanceData(UNiagaraDataInterface* DataInterface, void* InPerInstanceData, FNiagaraSystemInstance* SystemInstance) = 0;
 	virtual void DestroyPerInstanceData(void* InPerInstanceData, FNiagaraSystemInstance* SystemInstance) = 0;
 	virtual void ProvidePerInstanceDataForRenderThread(void* DataForRenderThread, void* PerInstanceData, const FNiagaraSystemInstanceID& SystemInstance) = 0;
 	virtual FNiagaraDataInterfaceParametersCS* CreateComputeParameters() const = 0;
@@ -64,7 +64,7 @@ public:
 	virtual bool Equals(const UNiagaraDataInterface* Other) const;
 
 	virtual int32 PerInstanceDataSize() const override { return GetProxyAs<INDIArrayProxyBase>()->PerInstanceDataSize(); }
-	virtual bool InitPerInstanceData(void* InPerInstanceData, FNiagaraSystemInstance* SystemInstance) override { return GetProxyAs<INDIArrayProxyBase>()->InitPerInstanceData(InPerInstanceData, SystemInstance); }
+	virtual bool InitPerInstanceData(void* InPerInstanceData, FNiagaraSystemInstance* SystemInstance) override { return GetProxyAs<INDIArrayProxyBase>()->InitPerInstanceData(this, InPerInstanceData, SystemInstance); }
 	virtual void DestroyPerInstanceData(void* InPerInstanceData, FNiagaraSystemInstance * SystemInstance) override { GetProxyAs<INDIArrayProxyBase>()->DestroyPerInstanceData(InPerInstanceData, SystemInstance); }
 
 	virtual void ProvidePerInstanceDataForRenderThread(void* DataForRenderThread, void* PerInstanceData, const FNiagaraSystemInstanceID& SystemInstance) override { GetProxyAs<INDIArrayProxyBase>()->ProvidePerInstanceDataForRenderThread(DataForRenderThread, PerInstanceData, SystemInstance); }
diff --git a/Engine/Plugins/FX/Niagara/Source/Niagara/Classes/NiagaraDataInterfaceArrayImpl.h b/Engine/Plugins/FX/Niagara/Source/Niagara/Classes/NiagaraDataInterfaceArrayImpl.h
--- a/Engine/Plugins/FX/Niagara/Source/Niagara/Classes/NiagaraDataInterfaceArrayImpl.h
+++ b/Engine/Plugins/FX/Niagara/Source/Niagara/Classes/NiagaraDataInterfaceArrayImpl.h
@@ -4,6 +4,7 @@
 
 #include "NiagaraClearCounts.h"
 #include "NiagaraDataInterfaceArray.h"
+#include "NiagaraDataInterfaceUtilities.h"
 #include "NiagaraGpuComputeDispatchInterface.h"
 #include "Niagara/Private/NiagaraGpuReadbackManager.h"
 #include "Niagara/Private/NiagaraStats.h"
@@ -82,11 +83,10 @@ struct FNiagaraDataInterfaceArrayImplHelper
 #if WITH_EDITORONLY_DATA
 	static bool UpgradeFunctionCall(FNiagaraFunctionSignature& FunctionSignature);
 #endif
-	static const TCHAR* GetHLSLTemplateFile()
+	static const TCHAR* GetHLSLTemplateFile(bool bIsRWArray)
 	{
-		return SupportsGpuRW() ? FNiagaraDataInterfaceArrayImplHelper::HLSLReadWriteTemplateFile : FNiagaraDataInterfaceArrayImplHelper::HLSLReadTemplateFile;
+		return bIsRWArray ? FNiagaraDataInterfaceArrayImplHelper::HLSLReadWriteTemplateFile : FNiagaraDataInterfaceArrayImplHelper::HLSLReadTemplateFile;
 	}
-	static bool SupportsGpuRW();
 };
 
 //////////////////////////////////////////////////////////////////////////
@@ -129,14 +129,14 @@ struct FNDIArrayInstanceData_RenderThread
 
 	bool IsReadOnly() const { return CountOffset == INDEX_NONE; }
 
-	void Initialize(FRHICommandListImmediate& RHICmdList, FNiagaraGpuComputeDispatchInterface* InComputeInterface, int32 InDefaultElements)
+	void Initialize(FRHICommandListImmediate& RHICmdList, FNiagaraGpuComputeDispatchInterface* InComputeInterface, int32 InDefaultElements, bool bRWGpuArray)
 	{
 		ComputeInterface = InComputeInterface;
 		DefaultElements = 0;
 		NumElements = INDEX_NONE;
 		CountOffset = INDEX_NONE;
 
-		if (FNiagaraDataInterfaceArrayImplHelper::SupportsGpuRW())
+		if (bRWGpuArray)
 		{
 			DefaultElements = InDefaultElements;
 			CountOffset = ComputeInterface->GetGPUInstanceCounterManager().AcquireOrAllocateEntry(RHICmdList);
@@ -620,7 +620,7 @@ struct FNDIArrayProxyImpl : public INDIArrayProxyBase
 
 		// Mutable functions
 		FNiagaraFunctionSignature DefaultMutableSig = DefaultImmutableSig;
-		DefaultMutableSig.bSupportsGPU = FNDIArrayImplHelper<TArrayType>::bSupportsGPU && FNiagaraDataInterfaceArrayImplHelper::SupportsGpuRW();
+		DefaultMutableSig.bSupportsGPU = FNDIArrayImplHelper<TArrayType>::bSupportsGPU;
 		DefaultMutableSig.bRequiresExecPin = true;
 		{
 			FNiagaraFunctionSignature& Sig = OutFunctions.Add_GetRef(DefaultMutableSig);
@@ -736,7 +736,27 @@ struct FNDIArrayProxyImpl : public INDIArrayProxyBase
 		GetVMExternalFunction_Internal(BindingInfo, InstanceData, OutFunc);
 	}
 
+	bool IsRWFunction(const FName FunctionName) const
+	{
+		return
+			FunctionName == FNiagaraDataInterfaceArrayImplHelper::Function_ClearName ||
+			FunctionName == FNiagaraDataInterfaceArrayImplHelper::Function_ResizeName ||
+			FunctionName == FNiagaraDataInterfaceArrayImplHelper::Function_SetArrayElemName ||
+			FunctionName == FNiagaraDataInterfaceArrayImplHelper::Function_AddName ||
+			FunctionName == FNiagaraDataInterfaceArrayImplHelper::Function_RemoveLastElemName;
+	}
+
 #if WITH_EDITORONLY_DATA
+	bool IsRWGpuArray(const FNiagaraDataInterfaceGPUParamInfo& ParamInfo) const
+	{
+		return ParamInfo.GeneratedFunctions.ContainsByPredicate(
+			[&](const FNiagaraDataInterfaceGeneratedFunction& Function)
+			{
+				return IsRWFunction(Function.DefinitionName);
+			}
+		);
+	}
+
 	template<typename T = FNDIArrayImplHelper<TArrayType>>
 	typename TEnableIf<T::bSupportsGPU>::Type GetParameterDefinitionHLSL_Internal(const FNiagaraDataInterfaceGPUParamInfo& ParamInfo, FString& OutHLSL) const
 	{
@@ -752,7 +772,7 @@ struct FNDIArrayProxyImpl : public INDIArrayProxyBase
 		};
 
 		FString TemplateFile;
-		LoadShaderSourceFile(FNiagaraDataInterfaceArrayImplHelper::GetHLSLTemplateFile(), EShaderPlatform::SP_PCD3D_SM5, &TemplateFile, nullptr);
+		LoadShaderSourceFile(FNiagaraDataInterfaceArrayImplHelper::GetHLSLTemplateFile(IsRWGpuArray(ParamInfo)), EShaderPlatform::SP_PCD3D_SM5, &TemplateFile, nullptr);
 		OutHLSL += FString::Format(*TemplateFile, TemplateArgs);
 	}
 
@@ -778,16 +798,13 @@ struct FNDIArrayProxyImpl : public INDIArrayProxyBase
 				return true;
 			}
 
-			if ( FNiagaraDataInterfaceArrayImplHelper::SupportsGpuRW() )
+			if ((FunctionInfo.DefinitionName == FNiagaraDataInterfaceArrayImplHelper::Function_ClearName) ||
+				(FunctionInfo.DefinitionName == FNiagaraDataInterfaceArrayImplHelper::Function_ResizeName) ||
+				(FunctionInfo.DefinitionName == FNiagaraDataInterfaceArrayImplHelper::Function_SetArrayElemName) ||
+				(FunctionInfo.DefinitionName == FNiagaraDataInterfaceArrayImplHelper::Function_AddName) ||
+				(FunctionInfo.DefinitionName == FNiagaraDataInterfaceArrayImplHelper::Function_RemoveLastElemName))
 			{
-				if ((FunctionInfo.DefinitionName == FNiagaraDataInterfaceArrayImplHelper::Function_ClearName) ||
-					(FunctionInfo.DefinitionName == FNiagaraDataInterfaceArrayImplHelper::Function_ResizeName) ||
-					(FunctionInfo.DefinitionName == FNiagaraDataInterfaceArrayImplHelper::Function_SetArrayElemName) ||
-					(FunctionInfo.DefinitionName == FNiagaraDataInterfaceArrayImplHelper::Function_AddName) ||
-					(FunctionInfo.DefinitionName == FNiagaraDataInterfaceArrayImplHelper::Function_RemoveLastElemName))
-				{
-					return true;
-				}
+				return true;
 			}
 		}
 		return false;
@@ -797,12 +814,10 @@ struct FNDIArrayProxyImpl : public INDIArrayProxyBase
 	{
 		if (FNDIArrayImplHelper<TArrayType>::bSupportsGPU)
 		{
-			const bool bSupportsRWArrays = FNiagaraDataInterfaceArrayImplHelper::SupportsGpuRW();
-
-			FSHAHash Hash = GetShaderFileHash(FNiagaraDataInterfaceArrayImplHelper::GetHLSLTemplateFile(), EShaderPlatform::SP_PCD3D_SM5);
+			FSHAHash Hash = GetShaderFileHash(FNiagaraDataInterfaceArrayImplHelper::GetHLSLTemplateFile(false), EShaderPlatform::SP_PCD3D_SM5);
+			InVisitor->UpdateString(TEXT("NiagaraDataInterfaceArrayTemplateHLSLSource"), Hash.ToString());
+			Hash = GetShaderFileHash(FNiagaraDataInterfaceArrayImplHelper::GetHLSLTemplateFile(true), EShaderPlatform::SP_PCD3D_SM5);
 			InVisitor->UpdateString(TEXT("NiagaraDataInterfaceArrayTemplateHLSLSource"), Hash.ToString());
-
-			InVisitor->UpdatePOD(TEXT("NDIArray_SupportsGpuRW"), bSupportsRWArrays);
 		}
 		return true;
 	}
@@ -831,7 +846,7 @@ struct FNDIArrayProxyImpl : public INDIArrayProxyBase
 		return sizeof(FNDIArrayInstanceData_GameThread<TArrayType>);
 	}
 
-	virtual bool InitPerInstanceData(void* InPerInstanceData, FNiagaraSystemInstance* SystemInstance) override
+	virtual bool InitPerInstanceData(UNiagaraDataInterface* DataInterface, void* InPerInstanceData, FNiagaraSystemInstance* SystemInstance) override
 	{
 		// Ensure we have the latest sync mode settings
 		CachePropertiesFromOwner();
@@ -843,12 +858,22 @@ struct FNDIArrayProxyImpl : public INDIArrayProxyBase
 
 		if ( FNDIArrayImplHelper<TArrayType>::bSupportsGPU && Owner->IsUsedWithGPUEmitter() )
 		{
+			bool bRWGpuArray = false;
+			FNiagaraDataInterfaceUtilities::ForEachGpuFunction(
+				DataInterface, SystemInstance,
+				[&](const FNiagaraDataInterfaceGeneratedFunction& Function) -> bool
+				{
+					bRWGpuArray = IsRWFunction(Function.DefinitionName);
+					return !bRWGpuArray;
+				}
+			);
+
 			ENQUEUE_RENDER_COMMAND(FNDIArrayProxyImpl_AddProxy)
 			(
-				[Proxy_RT=this, InstanceID_RT=SystemInstance->GetId(), ComputeInterface_RT=SystemInstance->GetComputeDispatchInterface(), MaxElements_RT=Owner->MaxElements](FRHICommandListImmediate& RHICmdList)
+				[Proxy_RT=this, InstanceID_RT=SystemInstance->GetId(), ComputeInterface_RT=SystemInstance->GetComputeDispatchInterface(), MaxElements_RT=Owner->MaxElements, bRWGpuArray_RT = bRWGpuArray](FRHICommandListImmediate& RHICmdList)
 				{
 					FNDIArrayInstanceData_RenderThread<TArrayType>* InstanceData_RT = &Proxy_RT->PerInstanceData_RenderThread.Add(InstanceID_RT);
-					InstanceData_RT->Initialize(RHICmdList, ComputeInterface_RT, MaxElements_RT);
+					InstanceData_RT->Initialize(RHICmdList, ComputeInterface_RT, MaxElements_RT, bRWGpuArray_RT);
 				}
 			);
 		}
diff --git a/Engine/Plugins/FX/Niagara/Source/Niagara/Private/NiagaraDataInterfaceArrayImpl.cpp b/Engine/Plugins/FX/Niagara/Source/Niagara/Private/NiagaraDataInterfaceArrayImpl.cpp
--- a/Engine/Plugins/FX/Niagara/Source/Niagara/Private/NiagaraDataInterfaceArrayImpl.cpp
+++ b/Engine/Plugins/FX/Niagara/Source/Niagara/Private/NiagaraDataInterfaceArrayImpl.cpp
@@ -16,14 +16,6 @@ const FName FNiagaraDataInterfaceArrayImplHelper::Function_SetArrayElemName(TEXT
 const FName FNiagaraDataInterfaceArrayImplHelper::Function_AddName(TEXT("Add"));
 const FName FNiagaraDataInterfaceArrayImplHelper::Function_RemoveLastElemName(TEXT("RemoveLastElem"));
 
-int32 GNiagaraArraySupportRW = 0;
-static FAutoConsoleVariableRef CVarNiagaraArraySupportRW(
-	TEXT("fx.NiagaraArraySupportRW"),
-	GNiagaraArraySupportRW,
-	TEXT("Allows the GPU to RW to the array, this comes with the caveat that all arrays will use a UAV slot."),
-	ECVF_Default
-);
-
 #if WITH_EDITORONLY_DATA
 bool FNiagaraDataInterfaceArrayImplHelper::UpgradeFunctionCall(FNiagaraFunctionSignature& FunctionSignature)
 {
@@ -69,9 +61,4 @@ bool FNiagaraDataInterfaceArrayImplHelper::UpgradeFunctionCall(FNiagaraFunctionS
 }
 #endif
 
-bool FNiagaraDataInterfaceArrayImplHelper::SupportsGpuRW()
-{
-	return GNiagaraArraySupportRW != 0;
-}
-
 IMPLEMENT_TYPE_LAYOUT(FNiagaraDataInterfaceParametersCS_ArrayImpl);
