- Fix SetValueAtIndex function

#rnx

diff --git a/Engine/Plugins/FX/Niagara/Source/Niagara/Private/NiagaraDataInterfaceGrid2DCollection.cpp b/Engine/Plugins/FX/Niagara/Source/Niagara/Private/NiagaraDataInterfaceGrid2DCollection.cpp
--- a/Engine/Plugins/FX/Niagara/Source/Niagara/Private/NiagaraDataInterfaceGrid2DCollection.cpp
+++ b/Engine/Plugins/FX/Niagara/Source/Niagara/Private/NiagaraDataInterfaceGrid2DCollection.cpp
@@ -1394,7 +1394,18 @@ bool UNiagaraDataInterfaceGrid2DCollection::GetFunctionHLSL(const FNiagaraDataIn
 		OutHLSL += FString::Format(FormatBounds, ArgsBounds);
 		return true;
 	}
-	else if (FunctionInfo.DefinitionName == SetValueFunctionName || FunctionInfo.DefinitionName == SetValueAtIndexFunctionName)
+	else if (FunctionInfo.DefinitionName == SetValueAtIndexFunctionName)
+	{
+		static const TCHAR* FormatBounds = TEXT(R"(
+			void {FunctionName}(int In_IndexX, int In_IndexY, int In_AttributeIndex, float In_Value)
+			{
+				RW{OutputGrid}[int3(In_IndexX, In_IndexY, In_AttributeIndex)] = In_Value;
+			}
+		)");
+		OutHLSL += FString::Format(FormatBounds, ArgsBounds);
+		return true;
+	}
+	else if (FunctionInfo.DefinitionName == SetValueFunctionName)
 	{
 		static const TCHAR *FormatBounds = TEXT(R"(
 			void {FunctionName}(int In_IndexX, int In_IndexY, int In_AttributeIndex, float In_Value, out int val)
