Adding WITH_NIAGARA_DEBUGGER to compile for Shipping Editor builds.
Modifying section headers for Niagara component

diff --git a/Engine/Plugins/FX/Niagara/Source/NiagaraEditor/Private/Customizations/NiagaraComponentDetails.cpp b/Engine/Plugins/FX/Niagara/Source/NiagaraEditor/Private/Customizations/NiagaraComponentDetails.cpp
--- a/Engine/Plugins/FX/Niagara/Source/NiagaraEditor/Private/Customizations/NiagaraComponentDetails.cpp
+++ b/Engine/Plugins/FX/Niagara/Source/NiagaraEditor/Private/Customizations/NiagaraComponentDetails.cpp
@@ -591,12 +591,14 @@ void FNiagaraComponentDetails::CustomizeDetails(IDetailLayoutBuilder& DetailBuil
 	static bool bFirstTime = true;
 	if (bFirstTime)
 	{
-		const FText DisplayName = LOCTEXT("Particles", "Particles");
-		TSharedRef<FPropertySection> Section = PropertyModule.FindOrCreateSection("NiagaraComponent", "Particles", DisplayName);
+		const FText DisplayName = LOCTEXT("EffectsSectionName", "Effects");
+		TSharedRef<FPropertySection> Section = PropertyModule.FindOrCreateSection("NiagaraComponent", "Effects", DisplayName);
 		Section->AddCategory(TEXT("Niagara"));
 		Section->AddCategory(ParamCategoryName);
 		Section->AddCategory(ParamUtilitiesName); 
 		Section->AddCategory(TEXT("Activation"));
+		Section->AddCategory(ScriptCategoryName);
+		Section->AddCategory(TEXT("Randomness"));
 		bFirstTime = false;
 	}
 
@@ -759,8 +761,10 @@ FReply FNiagaraComponentDetails::OnDebugSelectedSystem()
 
 	if (NiagaraComponentToUse)
 	{
-		//FSpawnTabArgs Args;
+
+#if WITH_NIAGARA_DEBUGGER
 		SNiagaraDebugger::InvokeDebugger(NiagaraComponentToUse);
+#endif
 	}
 
 	return FReply::Handled();
diff --git a/Engine/Plugins/FX/Niagara/Source/NiagaraEditor/Private/Toolkits/NiagaraSystemToolkit.cpp b/Engine/Plugins/FX/Niagara/Source/NiagaraEditor/Private/Toolkits/NiagaraSystemToolkit.cpp
--- a/Engine/Plugins/FX/Niagara/Source/NiagaraEditor/Private/Toolkits/NiagaraSystemToolkit.cpp
+++ b/Engine/Plugins/FX/Niagara/Source/NiagaraEditor/Private/Toolkits/NiagaraSystemToolkit.cpp
@@ -773,7 +773,9 @@ bool FNiagaraSystemToolkit::IsDrawOptionEnabled(int32 Element) const
 
 void FNiagaraSystemToolkit::OpenDebugHUD()
 {
+#if WITH_NIAGARA_DEBUGGER
 	SNiagaraDebugger::InvokeDebugger(&SystemViewModel->GetSystem());
+#endif
 }
 
 void FNiagaraSystemToolkit::OpenDebugOutliner()
diff --git a/Engine/Plugins/FX/Niagara/Source/NiagaraEditor/Private/ViewModels/NiagaraEmitterHandleViewModel.cpp b/Engine/Plugins/FX/Niagara/Source/NiagaraEditor/Private/ViewModels/NiagaraEmitterHandleViewModel.cpp
--- a/Engine/Plugins/FX/Niagara/Source/NiagaraEditor/Private/ViewModels/NiagaraEmitterHandleViewModel.cpp
+++ b/Engine/Plugins/FX/Niagara/Source/NiagaraEditor/Private/ViewModels/NiagaraEmitterHandleViewModel.cpp
@@ -355,8 +355,11 @@ void FNiagaraEmitterHandleViewModel::SetIsRenamePending(bool bInIsRenamePending)
 
 void FNiagaraEmitterHandleViewModel::BeginDebugEmitter()
 {
+
+#if WITH_NIAGARA_DEBUGGER
 	if (EmitterHandle)
 		SNiagaraDebugger::InvokeDebugger(*EmitterHandle);
+#endif
 }
 
 FNiagaraEmitterHandleViewModel::FOnPropertyChanged& FNiagaraEmitterHandleViewModel::OnPropertyChanged()
