Fixing final known spot for WITH_NIAGARA_DEBUGGER

diff --git a/Engine/Plugins/FX/Niagara/Source/NiagaraEditor/Private/ViewModels/NiagaraParameterPanelViewModel.cpp b/Engine/Plugins/FX/Niagara/Source/NiagaraEditor/Private/ViewModels/NiagaraParameterPanelViewModel.cpp
--- a/Engine/Plugins/FX/Niagara/Source/NiagaraEditor/Private/ViewModels/NiagaraParameterPanelViewModel.cpp
+++ b/Engine/Plugins/FX/Niagara/Source/NiagaraEditor/Private/ViewModels/NiagaraParameterPanelViewModel.cpp
@@ -1395,7 +1395,9 @@ void FNiagaraSystemToolkitParameterPanelViewModel::DebugParameters(const TArray<
 			if (EmitterHandles.Contains(Handle.GetId() ))
 				Handles.Add(Handle);
 		}
+#if WITH_NIAGARA_DEBUGGER
 		SNiagaraDebugger::InvokeDebugger(System, Handles, Attributes);
+#endif
 	}
 }
 
