Disable initial namespace initialization checking pending investigation of false positives

diff --git a/Engine/Plugins/FX/Niagara/Source/NiagaraEditor/Private/NiagaraHlslTranslator.cpp b/Engine/Plugins/FX/Niagara/Source/NiagaraEditor/Private/NiagaraHlslTranslator.cpp
--- a/Engine/Plugins/FX/Niagara/Source/NiagaraEditor/Private/NiagaraHlslTranslator.cpp
+++ b/Engine/Plugins/FX/Niagara/Source/NiagaraEditor/Private/NiagaraHlslTranslator.cpp
@@ -4878,7 +4878,8 @@ void FHlslNiagaraTranslator::InitializeParameterMapDefaults(int32 ParamMapHistor
 				FNiagaraVariable SourceForInitialValue = FNiagaraParameterMapHistory::GetSourceForInitialValue(Var);
 				if (!UniqueVars.Contains(SourceForInitialValue))
 				{
-					Error(FText::Format(LOCTEXT("MissingInitialValueSource", "Variable {0} is used, but its source variable {1} is not set!"), FText::FromName(Var.GetName()), FText::FromName(SourceForInitialValue.GetName())), nullptr, nullptr);
+					//@todo(ng) disabled pending investigation UE-150159
+					//Error(FText::Format(LOCTEXT("MissingInitialValueSource", "Variable {0} is used, but its source variable {1} is not set!"), FText::FromName(Var.GetName()), FText::FromName(SourceForInitialValue.GetName())), nullptr, nullptr);
 				}
 				InitialNamespaceVariablesMissingDefault.Add(Var);
 			}
@@ -6315,7 +6316,8 @@ void FHlslNiagaraTranslator::HandleParameterRead(int32 ParamMapHistoryIdx, const
 				}
 				else
 				{
-					Error(FText::Format(LOCTEXT("MissingInitialValueSource", "Variable {0} is used, but its source variable {1} is not set!"), FText::FromName(Var.GetName()), FText::FromName(SourceForInitialValue.GetName())), nullptr, nullptr);
+					//@todo(ng) disabled pending investigation UE-150159
+					//Error(FText::Format(LOCTEXT("MissingInitialValueSource", "Variable {0} is used, but its source variable {1} is not set!"), FText::FromName(Var.GetName()), FText::FromName(SourceForInitialValue.GetName())), nullptr, nullptr);
 				}
 			}
 			else if (UniqueVars.Contains(Var) && UniqueVarToChunk.Contains(Var))
