Niagara - Fix an issue where continous edits to user parameters on a placed niagara system would stop after the first change.  This was caused by the system being reset which caused the parameter stores to be reset so the system instance would think that the data interfaces had changed and would then raise its UI resync flag.  The component would then respond by broadcasting the OnSynchronizedWithAssetParametersDelegate which causes the niagara component details to refresh all of the user parameters.  This was changed so that it instead refreshes the errors on the data interfaces directly which is then handled by the data interface details.

diff --git a/Engine/Plugins/FX/Niagara/Source/Niagara/Private/NiagaraComponent.cpp b/Engine/Plugins/FX/Niagara/Source/Niagara/Private/NiagaraComponent.cpp
--- a/Engine/Plugins/FX/Niagara/Source/Niagara/Private/NiagaraComponent.cpp
+++ b/Engine/Plugins/FX/Niagara/Source/Niagara/Private/NiagaraComponent.cpp
@@ -1359,7 +1359,13 @@ void UNiagaraComponent::PostSystemTick_GameThread()
 #if WITH_EDITOR
 	if (SystemInstanceController->HandleNeedsUIResync())
 	{
-		OnSynchronizedWithAssetParametersDelegate.Broadcast();
+		for (UNiagaraDataInterface* OverrideDataInterface : OverrideParameters.GetDataInterfaces())
+		{
+			if (OverrideDataInterface != nullptr)
+			{
+				OverrideDataInterface->RefreshErrors();
+			}
+		}
 	}
 #endif
 
