Fixed a bug where particle collision query results can be assigned to the wrong particles

diff --git a/Engine/Plugins/FX/Niagara/Source/Niagara/Private/NiagaraDataInterfaceCollisionQuery.cpp b/Engine/Plugins/FX/Niagara/Source/Niagara/Private/NiagaraDataInterfaceCollisionQuery.cpp
--- a/Engine/Plugins/FX/Niagara/Source/Niagara/Private/NiagaraDataInterfaceCollisionQuery.cpp
+++ b/Engine/Plugins/FX/Niagara/Source/Niagara/Private/NiagaraDataInterfaceCollisionQuery.cpp
@@ -5,14 +5,11 @@
 #include "GlobalDistanceFieldParameters.h"
 #include "NiagaraAsyncGpuTraceHelper.h"
 #include "NiagaraComponent.h"
-#include "NiagaraDataInterfaceUtilities.h"
 #include "NiagaraGpuComputeDispatchInterface.h"
 #include "NiagaraGpuComputeDispatch.h"
-#include "NiagaraSimStageData.h"
 #include "NiagaraStats.h"
 #include "NiagaraTypes.h"
 #include "NiagaraWorldManager.h"
-#include "RenderResource.h"
 #include "Shader.h"
 #include "ShaderCore.h"
 #include "ShaderCompilerCore.h"
@@ -189,7 +186,7 @@ void UNiagaraDataInterfaceCollisionQuery::GetAssetTagsForContext(const UObject*
 	// for readability's sake, we leave them out of non-CDO assets.
 	if (bHaCPUQueriesWarning || (InAsset && InAsset->HasAnyFlags(EObjectFlags::RF_ClassDefaultObject)))
 	{
-		StringKeys.Add("CPUCollision") = bHaCPUQueriesWarning ? TEXT("True") : TEXT("False");
+		StringKeys.Add(FName("CPUCollision")) = bHaCPUQueriesWarning ? TEXT("True") : TEXT("False");
 
 	}
 
@@ -616,17 +613,17 @@ void UNiagaraDataInterfaceCollisionQuery::PerformQueryAsyncCPU(FVectorVMExternal
 		ensure(!TraceEnd.ContainsNaN());
 
 		int QueryID = Skip ? INDEX_NONE : InstanceData->CollisionBatch.SubmitQuery(LWCConverter.ConvertSimulationPositionToWorld(TraceStart), LWCConverter.ConvertSimulationPositionToWorld(TraceEnd), TraceChannel);
-		OutQueryID.SetAndAdvance(QueryID);
+		OutQueryID.SetAndAdvance(QueryID + 1);
 
 		// try to retrieve a query with the supplied query ID
 		FNiagaraDICollsionQueryResult Res;
 		int32 ID = InIDParam.GetAndAdvance();
-		if (ID != INDEX_NONE && InstanceData->CollisionBatch.GetQueryResult(ID, Res))
+		if (ID > 0 && InstanceData->CollisionBatch.GetQueryResult(ID - 1, Res))
 		{
 			OutQueryValid.SetAndAdvance(true);
 			OutInsideMesh.SetAndAdvance(Res.IsInsideMesh);
 			OutCollisionPos.SetAndAdvance(LWCConverter.ConvertWorldToSimulationPosition(Res.CollisionPos));
-			OutCollisionNormal.SetAndAdvance((FVector3f)Res.CollisionNormal);
+			OutCollisionNormal.SetAndAdvance(FVector3f(Res.CollisionNormal));
 			OutFriction.SetAndAdvance(Res.Friction);
 			OutRestitution.SetAndAdvance(Res.Restitution);
 			OutPhysicalMaterialIdx.SetAndAdvance(Res.PhysicalMaterialIdx);
diff --git a/Engine/Plugins/FX/Niagara/Source/Niagara/Private/NiagaraSystem.cpp b/Engine/Plugins/FX/Niagara/Source/Niagara/Private/NiagaraSystem.cpp
--- a/Engine/Plugins/FX/Niagara/Source/Niagara/Private/NiagaraSystem.cpp
+++ b/Engine/Plugins/FX/Niagara/Source/Niagara/Private/NiagaraSystem.cpp
@@ -2596,17 +2596,15 @@ bool UNiagaraSystem::QueryCompileComplete(bool bWait, bool bDoPost, bool bDoNotA
 				UNiagaraScript* SpawnScript = Emitter->SpawnScriptProps.Script;
 				for (const FNiagaraScriptDataInterfaceInfo& UpdateDataInterfaceInfo : CompiledScript->GetCachedDefaultDataInterfaces())
 				{
-					if (UpdateDataInterfaceInfo.RegisteredParameterMapRead == NAME_None && UpdateDataInterfaceInfo.RegisteredParameterMapWrite == NAME_None)
+					// If the data interface isn't being written to a parameter map then it won't be bound properly so we
+					// assign the update scripts copy of the data interface to the spawn scripts copy by pointer so that they will share
+					// the data interface at runtime and will both be updated in the editor.
+					for (FNiagaraScriptDataInterfaceInfo& SpawnDataInterfaceInfo : SpawnScript->GetCachedDefaultDataInterfaces())
 					{
-						// If the data interface isn't being read or written to a parameter map then it won't be bound properly so we
-						// assign the update scripts copy of the data interface to the spawn scripts copy by pointer so that they will share
-						// the data interface at runtime and will both be updated in the editor.
-						for (FNiagaraScriptDataInterfaceInfo& SpawnDataInterfaceInfo : SpawnScript->GetCachedDefaultDataInterfaces())
+						if (SpawnDataInterfaceInfo.RegisteredParameterMapWrite == NAME_None && UpdateDataInterfaceInfo.RegisteredParameterMapWrite == NAME_None && UpdateDataInterfaceInfo.Name == SpawnDataInterfaceInfo.Name)
 						{
-							if (UpdateDataInterfaceInfo.Name == SpawnDataInterfaceInfo.Name)
-							{
-								SpawnDataInterfaceInfo.DataInterface = UpdateDataInterfaceInfo.DataInterface;
-							}
+							SpawnDataInterfaceInfo.DataInterface = UpdateDataInterfaceInfo.DataInterface;
+							break;
 						}
 					}
 				}
