- Add method for data interfaces to draw to the debug hud
- Wrap debug hud with WITH_NIAGARA_DEBUGGER to make sure it's compiled out

diff --git a/Engine/Plugins/FX/Niagara/Source/Niagara/Classes/NiagaraDataInterface.h b/Engine/Plugins/FX/Niagara/Source/Niagara/Classes/NiagaraDataInterface.h
--- a/Engine/Plugins/FX/Niagara/Source/Niagara/Classes/NiagaraDataInterface.h
+++ b/Engine/Plugins/FX/Niagara/Source/Niagara/Classes/NiagaraDataInterface.h
@@ -12,6 +12,7 @@
 #include "NiagaraDataInterface.generated.h"
 
 class INiagaraCompiler;
+class UCanvas;
 class UCurveVector;
 class UCurveLinearColor;
 class UCurveFloat;
@@ -433,6 +434,16 @@ public:
 
 	virtual void PostExecute() {}
 
+#if WITH_NIAGARA_DEBUGGER
+	/**
+	Override this function to provide additional context to the debug HUD.
+	Fill VariableDataString with the data to display on the variable information in the HUD, this information
+	should be kept light to avoid polluting the display.
+	You can also use the Canvas to draw additional information based on verbosity
+	*/
+	virtual void DrawDebugHud(UCanvas* Canvas, FNiagaraSystemInstance* SystemInstance, FString& VariableDataString, bool bVerbose) const {};
+#endif
+
 #if WITH_EDITOR	
 	/** Refreshes and returns the errors detected with the corresponding data, if any.*/
 	virtual TArray<FNiagaraDataInterfaceError> GetErrors() { return TArray<FNiagaraDataInterfaceError>(); }
diff --git a/Engine/Plugins/FX/Niagara/Source/Niagara/Classes/NiagaraDataInterfaceArray.h b/Engine/Plugins/FX/Niagara/Source/Niagara/Classes/NiagaraDataInterfaceArray.h
--- a/Engine/Plugins/FX/Niagara/Source/Niagara/Classes/NiagaraDataInterfaceArray.h
+++ b/Engine/Plugins/FX/Niagara/Source/Niagara/Classes/NiagaraDataInterfaceArray.h
@@ -21,6 +21,9 @@ struct INDIArrayProxyBase : public FNiagaraDataInterfaceProxyRW
 	virtual bool GetFunctionHLSL(const FNiagaraDataInterfaceGPUParamInfo& ParamInfo, const FNiagaraDataInterfaceGeneratedFunction& FunctionInfo, int FunctionInstanceIndex, FString& OutHLSL) const = 0;
 	virtual bool AppendCompileHash(FNiagaraCompileHashVisitor* InVisitor) const = 0;
 	virtual bool UpgradeFunctionCall(FNiagaraFunctionSignature& FunctionSignature) const = 0;
+#endif
+#if WITH_NIAGARA_DEBUGGER
+	virtual void DrawDebugHud(UCanvas* Canvas, FNiagaraSystemInstance* SystemInstance, FString& VariableDataString, bool bVerbose) const = 0;
 #endif
 	virtual bool CopyToInternal(INDIArrayProxyBase* Destination) const = 0;
 	virtual bool Equals(const INDIArrayProxyBase* Other) const = 0;
@@ -60,6 +63,10 @@ public:
 #endif
 	virtual bool CanExecuteOnTarget(ENiagaraSimTarget Target) const override { return true; }
 
+#if WITH_NIAGARA_DEBUGGER
+	virtual void DrawDebugHud(UCanvas* Canvas, FNiagaraSystemInstance* SystemInstance, FString& VariableDataString, bool bVerbose) const { GetProxyAs<INDIArrayProxyBase>()->DrawDebugHud(Canvas, SystemInstance, VariableDataString, bVerbose); }
+#endif
+
 	virtual bool CopyToInternal(UNiagaraDataInterface* Destination) const;
 	virtual bool Equals(const UNiagaraDataInterface* Other) const;
 
diff --git a/Engine/Plugins/FX/Niagara/Source/Niagara/Classes/NiagaraDataInterfaceArrayImpl.h b/Engine/Plugins/FX/Niagara/Source/Niagara/Classes/NiagaraDataInterfaceArrayImpl.h
--- a/Engine/Plugins/FX/Niagara/Source/Niagara/Classes/NiagaraDataInterfaceArrayImpl.h
+++ b/Engine/Plugins/FX/Niagara/Source/Niagara/Classes/NiagaraDataInterfaceArrayImpl.h
@@ -827,6 +827,18 @@ struct FNDIArrayProxyImpl : public INDIArrayProxyBase
 		return FNiagaraDataInterfaceArrayImplHelper::UpgradeFunctionCall(FunctionSignature);
 	}
 #endif
+#if WITH_NIAGARA_DEBUGGER
+	virtual void DrawDebugHud(UCanvas* Canvas, FNiagaraSystemInstance* SystemInstance, FString& VariableDataString, bool bVerbose) const override
+	{
+		FNDIArrayInstanceData_GameThread<TArrayType>* InstanceData = PerInstanceData_GameThread.FindRef(SystemInstance->GetId());
+		if (InstanceData == nullptr )
+		{
+			return;
+		}
+		FReadArrayRef ArrayData(Owner, InstanceData);
+		VariableDataString = FString::Printf(TEXT("ArrayType(%s) CpuLength(%d)"), *FNDIArrayImplHelper<TArrayType>::GetTypeDefinition().GetName(), ArrayData.GetArray().Num());
+	}
+#endif
 
 	virtual bool CopyToInternal(INDIArrayProxyBase* InDestination) const override
 	{
diff --git a/Engine/Plugins/FX/Niagara/Source/Niagara/Classes/NiagaraDataInterfaceSkeletalMesh.h b/Engine/Plugins/FX/Niagara/Source/Niagara/Classes/NiagaraDataInterfaceSkeletalMesh.h
--- a/Engine/Plugins/FX/Niagara/Source/Niagara/Classes/NiagaraDataInterfaceSkeletalMesh.h
+++ b/Engine/Plugins/FX/Niagara/Source/Niagara/Classes/NiagaraDataInterfaceSkeletalMesh.h
@@ -803,6 +803,9 @@ public:
 	virtual void GetVMExternalFunction(const FVMExternalFunctionBindingInfo& BindingInfo, void* InstanceData, FVMExternalFunction &OutFunc)override;
 	virtual bool Equals(const UNiagaraDataInterface* Other) const override;
 	virtual bool CanExecuteOnTarget(ENiagaraSimTarget Target)const override { return true; }
+#if WITH_NIAGARA_DEBUGGER
+	virtual void DrawDebugHud(UCanvas* Canvas, FNiagaraSystemInstance* SystemInstance, FString& VariableDataString, bool bVerbose) const override;
+#endif
 #if WITH_EDITOR
 	virtual void GetFeedback(UNiagaraSystem* Asset, UNiagaraComponent* Component, TArray<FNiagaraDataInterfaceError>& OutErrors,
 		TArray<FNiagaraDataInterfaceFeedback>& Warnings, TArray<FNiagaraDataInterfaceFeedback>& Info) override;
diff --git a/Engine/Plugins/FX/Niagara/Source/Niagara/Private/NiagaraDataInterfaceSkeletalMesh.cpp b/Engine/Plugins/FX/Niagara/Source/Niagara/Private/NiagaraDataInterfaceSkeletalMesh.cpp
--- a/Engine/Plugins/FX/Niagara/Source/Niagara/Private/NiagaraDataInterfaceSkeletalMesh.cpp
+++ b/Engine/Plugins/FX/Niagara/Source/Niagara/Private/NiagaraDataInterfaceSkeletalMesh.cpp
@@ -5,6 +5,7 @@
 #include "Animation/SkeletalMeshActor.h"
 #include "Async/ParallelFor.h"
 #include "Components/SkeletalMeshComponent.h"
+#include "Engine/Canvas.h"
 #include "Engine/SkeletalMeshSocket.h"
 #include "Internationalization/Internationalization.h"
 #include "NDISkeletalMeshCommon.h"
@@ -2707,6 +2708,44 @@ bool UNiagaraDataInterfaceSkeletalMesh::PerInstanceTick(void* PerInstanceData, F
 	return Inst->Tick(this, SystemInstance, InDeltaSeconds);
 }
 
+#if WITH_NIAGARA_DEBUGGER
+void UNiagaraDataInterfaceSkeletalMesh::DrawDebugHud(UCanvas* Canvas, FNiagaraSystemInstance* SystemInstance, FString& VariableDataString, bool bVerbose) const
+{
+	FNDISkeletalMesh_InstanceData* InstanceData_GT = SystemInstance->FindTypedDataInterfaceInstanceData<FNDISkeletalMesh_InstanceData>(this);
+	if (InstanceData_GT == nullptr)
+	{
+		return;
+	}
+
+	USceneComponent* SceneComponent = InstanceData_GT->SceneComponent.Get();
+	USkeletalMesh* SkeletalMesh = InstanceData_GT->SkeletalMesh.Get();
+	VariableDataString = FString::Printf(TEXT("Skeleton(%s) SkelComp(%s)"), *GetNameSafe(SkeletalMesh), *GetNameSafe(SceneComponent));
+
+	if ( bVerbose && SceneComponent && SkeletalMesh )
+	{
+		FSkeletalMeshAccessorHelper MeshAccessor;
+		MeshAccessor.Init<TNDISkelMesh_FilterModeNone, TNDISkelMesh_AreaWeightingOff>(InstanceData_GT);
+		if (MeshAccessor.AreBonesAccessible())
+		{
+			const TArray<FTransform3f>& BoneTransforms = MeshAccessor.SkinningData->CurrComponentTransforms();
+			const FMatrix& InstanceTransform = InstanceData_GT->Transform;
+
+			for (const FTransform3f& BoneTransform : BoneTransforms)
+			{
+				const FVector BoneLocation = InstanceTransform.TransformPosition(FVector(BoneTransform.GetLocation()));
+				const FVector ScreenPos = Canvas->Project(BoneLocation, false);
+				if (ScreenPos.Z <= 0.0f)
+				{
+					continue;
+				}
+
+				Canvas->Canvas->DrawNGon(FVector2D(ScreenPos), FColor::Red, 8, 4.0f);
+			}
+		}
+	}
+}
+#endif
+
 #if WITH_EDITOR
 void UNiagaraDataInterfaceSkeletalMesh::GetFeedback(UNiagaraSystem* Asset, UNiagaraComponent* Component, TArray<FNiagaraDataInterfaceError>& OutErrors,
 	TArray<FNiagaraDataInterfaceFeedback>& OutWarnings, TArray<FNiagaraDataInterfaceFeedback>& OutInfo)
diff --git a/Engine/Plugins/FX/Niagara/Source/Niagara/Private/NiagaraDebugHud.cpp b/Engine/Plugins/FX/Niagara/Source/Niagara/Private/NiagaraDebugHud.cpp
--- a/Engine/Plugins/FX/Niagara/Source/Niagara/Private/NiagaraDebugHud.cpp
+++ b/Engine/Plugins/FX/Niagara/Source/Niagara/Private/NiagaraDebugHud.cpp
@@ -17,6 +17,8 @@
 #include "Particles/FXBudget.h"
 #include "DrawDebugHelpers.h"
 
+#if WITH_NIAGARA_DEBUGGER
+
 namespace NiagaraDebugLocal
 {
 	FCriticalSection RTFramesGuard;
@@ -2638,7 +2640,21 @@ void FNiagaraDebugHud::DrawComponents(FNiagaraWorldManager* WorldManager, UCanva
 								{
 									if (UserVariable.IsDataInterface())
 									{
-										StringBuilder.Appendf(TEXT("%s(%s %s)\n"), *UserVariable.GetName().ToString(), *GetNameSafe(UserVariable.GetType().GetClass()) , *GetNameSafe(ParameterStore->GetDataInterfaces()[UserVariable.Offset]));
+										UNiagaraDataInterface* DataInterface = ParameterStore->GetDataInterfaces()[UserVariable.Offset];
+										FString DataInterfaceString;
+										if ( DataInterface != nullptr && Settings.DataInterfaceVerbosity != ENiagaraDebugHudVerbosity::None )
+										{
+											DataInterface->DrawDebugHud(Canvas, SystemInstance, DataInterfaceString, Settings.DataInterfaceVerbosity == ENiagaraDebugHudVerbosity::Verbose);
+										}
+
+										if ( DataInterfaceString.IsEmpty() )
+										{
+											StringBuilder.Appendf(TEXT("%s(%s %s)\n"), *UserVariable.GetName().ToString(), *GetNameSafe(UserVariable.GetType().GetClass()) , *GetNameSafe(DataInterface));
+										}
+										else
+										{
+											StringBuilder.Appendf(TEXT("%s(%s)\n"), *UserVariable.GetName().ToString(), *DataInterfaceString);
+										}
 									}
 									else if (UserVariable.IsUObject())
 									{
@@ -3113,3 +3129,5 @@ void FNiagaraDebugHUDStatsListener::TickRT()
 #endif
 }
 #endif//WITH_PARTICLE_PERF_STATS
+
+#endif //WITH_NIAGARA_DEBUGGER
diff --git a/Engine/Plugins/FX/Niagara/Source/Niagara/Private/NiagaraDebugHud.h b/Engine/Plugins/FX/Niagara/Source/Niagara/Private/NiagaraDebugHud.h
--- a/Engine/Plugins/FX/Niagara/Source/Niagara/Private/NiagaraDebugHud.h
+++ b/Engine/Plugins/FX/Niagara/Source/Niagara/Private/NiagaraDebugHud.h
@@ -13,6 +13,8 @@ Class used help debugging Niagara simulations
 #include "NiagaraGPUProfilerInterface.h"
 #include "Particles/ParticlePerfStatsManager.h"
 
+#if WITH_NIAGARA_DEBUGGER
+
 #if WITH_PARTICLE_PERF_STATS
 
 class FNiagaraDebugHud;
@@ -354,3 +356,5 @@ public:
 	void AddBox2D(FVector2D Pos, FVector2D Extents, FLinearColor Color, float Thickness, float Lifetime);
 	//Additional debug geometry helpers END
 };
+
+#endif //WITH_NIAGARA_DEBUGGER
diff --git a/Engine/Plugins/FX/Niagara/Source/Niagara/Private/NiagaraWorldManager.cpp b/Engine/Plugins/FX/Niagara/Source/Niagara/Private/NiagaraWorldManager.cpp
--- a/Engine/Plugins/FX/Niagara/Source/Niagara/Private/NiagaraWorldManager.cpp
+++ b/Engine/Plugins/FX/Niagara/Source/Niagara/Private/NiagaraWorldManager.cpp
@@ -291,7 +291,7 @@ void FNiagaraWorldManager::Init(UWorld* InWorld)
 	//Possibly a later hook we can use.
 	//PrimePoolForAllSystems();
 
-#if !UE_BUILD_SHIPPING
+#if WITH_NIAGARA_DEBUGGER
 	NiagaraDebugHud.Reset(new FNiagaraDebugHud(World));
 #endif
 }
@@ -842,11 +842,13 @@ void FNiagaraWorldManager::PostActorTick(float DeltaSeconds)
 		TickFunc.EndTickGroup = GNiagaraAllowAsyncWorkToEndOfFrame ? TG_LastDemotable : (ETickingGroup)TickFunc.TickGroup;
 	}
 
+#if WITH_NIAGARA_DEBUGGER
 	// Tick debug HUD for the world
 	if (NiagaraDebugHud != nullptr)
 	{
 		NiagaraDebugHud->GatherSystemInfo();
 	}
+#endif
 
 	if ( DebugPlaybackMode == ENiagaraDebugPlaybackMode::Step )
 	{
diff --git a/Engine/Plugins/FX/Niagara/Source/Niagara/Public/NiagaraCommon.h b/Engine/Plugins/FX/Niagara/Source/Niagara/Public/NiagaraCommon.h
--- a/Engine/Plugins/FX/Niagara/Source/Niagara/Public/NiagaraCommon.h
+++ b/Engine/Plugins/FX/Niagara/Source/Niagara/Public/NiagaraCommon.h
@@ -34,6 +34,10 @@ struct FNiagaraParameterStore;
 	#define NIAGARA_COMPUTEDEBUG_ENABLED WITH_EDITOR
 #endif
 
+#ifndef WITH_NIAGARA_DEBUGGER
+	#define WITH_NIAGARA_DEBUGGER !UE_BUILD_SHIPPING
+#endif
+
 #define WITH_NIAGARA_GPU_PROFILER_EDITOR (WITH_EDITOR && STATS)
 #define WITH_NIAGARA_GPU_PROFILER ((WITH_PARTICLE_PERF_STATS && !UE_BUILD_SHIPPING) || WITH_NIAGARA_GPU_PROFILER_EDITOR)
 
diff --git a/Engine/Plugins/FX/Niagara/Source/Niagara/Public/NiagaraDebuggerCommon.h b/Engine/Plugins/FX/Niagara/Source/Niagara/Public/NiagaraDebuggerCommon.h
--- a/Engine/Plugins/FX/Niagara/Source/Niagara/Public/NiagaraDebuggerCommon.h
+++ b/Engine/Plugins/FX/Niagara/Source/Niagara/Public/NiagaraDebuggerCommon.h
@@ -12,8 +12,6 @@ All common code shared between the editor side debugger and debugger clients run
 #include "NiagaraCommon.h"
 #include "NiagaraDebuggerCommon.generated.h"
 
-#define WITH_NIAGARA_DEBUGGER (!UE_BUILD_SHIPPING)// || WITH_EDITOR)
-
 //////////////////////////////////////////////////////////////////////////
 // Niagara Outliner.
 
@@ -520,6 +518,10 @@ struct NIAGARA_API FNiagaraDebugHUDSettingsData
 	UPROPERTY(Config, EditAnywhere, Category = "Debug System", meta = (EditCondition = "SystemDebugVerbosity != ENiagaraDebugHudVerbosity::None"))
 	ENiagaraDebugHudVerbosity SystemEmitterVerbosity = ENiagaraDebugHudVerbosity::Basic;
 
+	/** When enabled allows data interfaces to include additional debugging information. */
+	UPROPERTY(Config, EditAnywhere, Category = "Debug System", meta = (EditCondition = "SystemDebugVerbosity != ENiagaraDebugHudVerbosity::None"))
+	ENiagaraDebugHudVerbosity DataInterfaceVerbosity = ENiagaraDebugHudVerbosity::None;
+
 	/** When enabled will show the system bounds for all filtered systems. */
 	UPROPERTY(Config, EditAnywhere, Category = "Debug System")
 	bool bSystemShowBounds = false;
diff --git a/Engine/Plugins/FX/Niagara/Source/Niagara/Public/NiagaraSystemInstance.h b/Engine/Plugins/FX/Niagara/Source/Niagara/Public/NiagaraSystemInstance.h
--- a/Engine/Plugins/FX/Niagara/Source/Niagara/Public/NiagaraSystemInstance.h
+++ b/Engine/Plugins/FX/Niagara/Source/Niagara/Public/NiagaraSystemInstance.h
@@ -280,6 +280,16 @@ public:
 		return nullptr;
 	}
 
+	template<typename TDataType>
+	FORCEINLINE TDataType* FindTypedDataInterfaceInstanceData(const UNiagaraDataInterface* Interface)
+	{
+		if (auto* InstDataOffsetPair = DataInterfaceInstanceDataOffsets.FindByPredicate([&](auto& Pair) { return Pair.Key.Get() == Interface; }))
+		{
+			return reinterpret_cast<TDataType*>(&DataInterfaceInstanceData[InstDataOffsetPair->Value]);
+		}
+		return nullptr;
+	}
+
 	FORCEINLINE const FNiagaraPerInstanceDIFuncInfo& GetPerInstanceDIFunction(ENiagaraSystemSimulationScript ScriptType, int32 FuncIndex)const { return PerInstanceDIFunctions[(int32)ScriptType][FuncIndex]; }
 
 	void EvaluateBoundFunction(FName FunctionName, bool& UsedOnCpu, bool& UsedOnGpu) const;
diff --git a/Engine/Plugins/FX/Niagara/Source/Niagara/Public/NiagaraWorldManager.h b/Engine/Plugins/FX/Niagara/Source/Niagara/Public/NiagaraWorldManager.h
--- a/Engine/Plugins/FX/Niagara/Source/Niagara/Public/NiagaraWorldManager.h
+++ b/Engine/Plugins/FX/Niagara/Source/Niagara/Public/NiagaraWorldManager.h
@@ -202,7 +202,9 @@ public:
 	void SetDebugPlaybackRate(float Rate) { DebugPlaybackRate = FMath::Clamp(Rate, KINDA_SMALL_NUMBER, 10.0f); }
 	float GetDebugPlaybackRate() const { return DebugPlaybackRate; }
 
+#if WITH_NIAGARA_DEBUGGER
 	class FNiagaraDebugHud* GetNiagaraDebugHud() { return NiagaraDebugHud.Get(); }
+#endif
 
 	class FNiagaraDeferredMethodQueue& GetDeferredMethodQueue() { return DeferredMethods; }
 
@@ -317,7 +319,9 @@ private:
 	ENiagaraDebugPlaybackMode DebugPlaybackMode = ENiagaraDebugPlaybackMode::Play;
 	float DebugPlaybackRate = 1.0f;
 
+#if WITH_NIAGARA_DEBUGGER
 	TUniquePtr<class FNiagaraDebugHud> NiagaraDebugHud;
+#endif
 
 	TMap<UNiagaraSystem*, UNiagaraCullProxyComponent*> CullProxyMap;
 
