- Handle missing & bind dummy data for distance fields

#rnx

diff --git a/Engine/Plugins/FX/Niagara/Source/Niagara/Private/DataInterface/NiagaraDistanceFieldParameters.h b/Engine/Plugins/FX/Niagara/Source/Niagara/Private/DataInterface/NiagaraDistanceFieldParameters.h
--- a/Engine/Plugins/FX/Niagara/Source/Niagara/Private/DataInterface/NiagaraDistanceFieldParameters.h
+++ b/Engine/Plugins/FX/Niagara/Source/Niagara/Private/DataInterface/NiagaraDistanceFieldParameters.h
@@ -57,14 +57,15 @@ public:
 	template<typename ShaderRHIParamRef>
 	FORCEINLINE_DEBUGGABLE void Set(FRHICommandList& RHICmdList, const ShaderRHIParamRef ShaderRHI, const FDistanceFieldSceneData* ParameterData) const
 	{
-		if (IsBound() && ParameterData->GetCurrentObjectBuffers() != nullptr)
+		if (IsBound() )
 		{
-			SetSRVParameter(RHICmdList, ShaderRHI, SceneObjectBounds, ParameterData->GetCurrentObjectBuffers()->Bounds->GetSRV());
-			SetSRVParameter(RHICmdList, ShaderRHI, SceneObjectData, ParameterData->GetCurrentObjectBuffers()->Data->GetSRV());
+			const FDistanceFieldObjectBuffers* DistanceFieldObjectBuffers = ParameterData->GetCurrentObjectBuffers();
+			SetSRVParameter(RHICmdList, ShaderRHI, SceneObjectBounds, DistanceFieldObjectBuffers && DistanceFieldObjectBuffers->Bounds.IsValid() ? DistanceFieldObjectBuffers->Bounds->GetSRV() : FNiagaraRenderer::GetDummyFloat4Buffer());
+			SetSRVParameter(RHICmdList, ShaderRHI, SceneObjectData, DistanceFieldObjectBuffers && DistanceFieldObjectBuffers->Data.IsValid() ? DistanceFieldObjectBuffers->Data->GetSRV() : FNiagaraRenderer::GetDummyFloat4Buffer());
 			SetShaderValue(RHICmdList, ShaderRHI, NumSceneObjects, ParameterData->NumObjectsInBuffer);
-			SetSRVParameter(RHICmdList, ShaderRHI, SceneDistanceFieldAssetData, ParameterData->AssetDataBuffer->GetSRV());
-			SetSRVParameter(RHICmdList, ShaderRHI, DistanceFieldIndirectionTable, ParameterData->IndirectionTable->GetSRV());
-			SetTextureParameter(RHICmdList, ShaderRHI, DistanceFieldBrickTexture, ParameterData->DistanceFieldBrickVolumeTexture->GetRHI());
+			SetSRVParameter(RHICmdList, ShaderRHI, SceneDistanceFieldAssetData, ParameterData->AssetDataBuffer.IsValid() ? ParameterData->AssetDataBuffer->GetSRV() : FNiagaraRenderer::GetDummyFloat4Buffer());
+			SetSRVParameter(RHICmdList, ShaderRHI, DistanceFieldIndirectionTable, ParameterData->IndirectionTable.IsValid() ? ParameterData->IndirectionTable->GetSRV() : FNiagaraRenderer::GetDummyUIntBuffer());
+			SetTextureParameter(RHICmdList, ShaderRHI, DistanceFieldBrickTexture, ParameterData->DistanceFieldBrickVolumeTexture.IsValid() ? ParameterData->DistanceFieldBrickVolumeTexture->GetRHI() : GBlackVolumeTexture->TextureRHI->GetTextureReference());
 			SetSamplerParameter(RHICmdList, ShaderRHI, DistanceFieldSampler, TStaticSamplerState<SF_Bilinear, AM_Clamp, AM_Clamp, AM_Clamp>::GetRHI());
 			SetShaderValue(RHICmdList, ShaderRHI, DistanceFieldBrickSize, FVector3f(DistanceField::BrickSize));
 			SetShaderValue(RHICmdList, ShaderRHI, DistanceFieldUniqueDataBrickSize, FVector3f(DistanceField::UniqueDataBrickSize));
