Niagara - Refreshing an assignment node from external changes rebuilds the script, so we need to modify the node and script when this is done so that the state is recorded in any current transactions.  Failure to save the state of these objects results in them becoming invalid when the transaction is undone.

diff --git a/Engine/Plugins/FX/Niagara/Source/NiagaraEditor/Private/NiagaraNodeAssignment.cpp b/Engine/Plugins/FX/Niagara/Source/NiagaraEditor/Private/NiagaraNodeAssignment.cpp
--- a/Engine/Plugins/FX/Niagara/Source/NiagaraEditor/Private/NiagaraNodeAssignment.cpp
+++ b/Engine/Plugins/FX/Niagara/Source/NiagaraEditor/Private/NiagaraNodeAssignment.cpp
@@ -78,6 +78,11 @@ TSharedRef<SWidget> UNiagaraNodeAssignment::CreateAddParameterMenu(const TShared
 
 bool UNiagaraNodeAssignment::RefreshFromExternalChanges()
 {
+	Modify(false);
+	if (FunctionScript != nullptr)
+	{
+		FunctionScript->Modify(false);
+	}
 	FunctionScript = nullptr;
 	GenerateScript();
 	ReallocatePins();
