- Fix crash with component renderer and scalability view

#rnx

diff --git a/Engine/Plugins/FX/Niagara/Source/NiagaraEditor/Private/Customizations/NiagaraComponentRendererPropertiesDetails.cpp b/Engine/Plugins/FX/Niagara/Source/NiagaraEditor/Private/Customizations/NiagaraComponentRendererPropertiesDetails.cpp
--- a/Engine/Plugins/FX/Niagara/Source/NiagaraEditor/Private/Customizations/NiagaraComponentRendererPropertiesDetails.cpp
+++ b/Engine/Plugins/FX/Niagara/Source/NiagaraEditor/Private/Customizations/NiagaraComponentRendererPropertiesDetails.cpp
@@ -76,7 +76,10 @@ void FNiagaraComponentRendererPropertiesDetails::CustomizeDetails(IDetailLayoutB
 	}
 
 	// If we have a template component we hide the default widget and build a new category with all its properties
-	DetailBuilder.EditDefaultProperty(TemplateHandle)->Visibility(EVisibility::Collapsed);
+	if (IDetailPropertyRow* TemplateHandlePropertyRow = DetailBuilder.EditDefaultProperty(TemplateHandle))
+	{
+		TemplateHandlePropertyRow->Visibility(EVisibility::Collapsed);
+	}
 
 	static const FName ComponentCategoryName = TEXT("Component Properties");
 	IDetailCategoryBuilder& CategoryBuilder = DetailBuilder.EditCategory(ComponentCategoryName);
