Fixed issue with GScalabilityMaxUpdatesPerFrame and continuous frequency updates that meant no updates were ever performed.

diff --git a/Engine/Plugins/FX/Niagara/Source/Niagara/Private/NiagaraScalabilityManager.cpp b/Engine/Plugins/FX/Niagara/Source/Niagara/Private/NiagaraScalabilityManager.cpp
--- a/Engine/Plugins/FX/Niagara/Source/Niagara/Private/NiagaraScalabilityManager.cpp
+++ b/Engine/Plugins/FX/Niagara/Source/Niagara/Private/NiagaraScalabilityManager.cpp
@@ -31,7 +31,7 @@ static float GetScalabilityUpdatePeriod(ENiagaraScalabilityUpdateFrequency Frequ
 
 static int32 GetMaxUpdatesPerFrame(const UNiagaraEffectType* EffectType, int32 ItemsRemaining, float UpdatePeriod, float DeltaSeconds)
 {
-	if (GScalabilityMaxUpdatesPerFrame > 0)
+	if (GScalabilityMaxUpdatesPerFrame > 0 && EffectType->UpdateFrequency != ENiagaraScalabilityUpdateFrequency::Continuous)
 	{
 		int32 UpdateCount = ItemsRemaining;
 
