expose max distance

diff --git a/Engine/Plugins/FX/Niagara/Shaders/Private/NiagaraDataInterfaceRigidMeshCollisionQuery.ush b/Engine/Plugins/FX/Niagara/Shaders/Private/NiagaraDataInterfaceRigidMeshCollisionQuery.ush
--- a/Engine/Plugins/FX/Niagara/Shaders/Private/NiagaraDataInterfaceRigidMeshCollisionQuery.ush
+++ b/Engine/Plugins/FX/Niagara/Shaders/Private/NiagaraDataInterfaceRigidMeshCollisionQuery.ush
@@ -407,7 +407,7 @@ void DIRigidMeshCollision_GetClosestDistance(in FDIRigidMeshCollisionContext DIC
 }
 
 // Given a world space position (WorldPosition) compute the static mesh closest point (position,normal,velocity)
-void DIRigidMeshCollision_GetClosestPointMeshDistanceField(in FDIRigidMeshCollisionContext DIContext, in float3 WorldPosition, in float DeltaTime, in float TimeFraction, out float OutClosestDistance, out float3 OutClosestPosition, out float3 OutClosestNormal, out float3 OutClosestVelocity)
+void DIRigidMeshCollision_GetClosestPointMeshDistanceField(in FDIRigidMeshCollisionContext DIContext, in float3 WorldPosition, in float DeltaTime, in float TimeFraction, in float MaxDistance, out float OutClosestDistance, out float3 OutClosestPosition, out float3 OutClosestNormal, out float3 OutClosestVelocity, out bool NormalIsValid)
 {	
 	float3 CollisionPosition = float3(0,0,0);
 	float3 CollisionNormal = float3(0,0,0);
@@ -417,6 +417,8 @@ void DIRigidMeshCollision_GetClosestPointMeshDistanceField(in FDIRigidMeshCollis
 	OutClosestPosition = float3(0,0,0);
 	OutClosestVelocity = float3(0,0,0);
 
+	NormalIsValid = false;
+
 	float OutClosestDistanceTmp;
 	const int ElementIndex = DIRigidMeshCollision_GetClosestElement(DIContext,WorldPosition,CollisionPosition,CollisionNormal,OutClosestDistanceTmp,TimeFraction);
 
@@ -427,7 +429,9 @@ void DIRigidMeshCollision_GetClosestPointMeshDistanceField(in FDIRigidMeshCollis
 		if (DFIndex >= 0 && DFIndex < NumSceneObjects)
 		{
 			// #todo(dmp): this does a dynamic branch based on intersecting the bbox.  Maybe we can factor that out due to the broadphase here?
-			OutClosestDistance = DistanceToNearestSurfaceForObject(DFIndex, WorldPosition, 1e-3);
+			// #todo(dmp): Check the distance from the physics asset for a speedup?
+			MaxDistance = max(.001, MaxDistance);
+			OutClosestDistance = DistanceToNearestSurfaceForObject(DFIndex, WorldPosition, MaxDistance);			 
 
 			FDFObjectData DFObjectData = LoadDFObjectData(DFIndex);
 			float4x4 WorldToVolume = LWCHackToFloat(DFObjectData.WorldToVolume);
@@ -441,17 +445,77 @@ void DIRigidMeshCollision_GetClosestPointMeshDistanceField(in FDIRigidMeshCollis
 			float3 WorldGradient = mul(VolumeNormal, transpose((float3x3)WorldToVolume));
 			float WorldGradientLength = length(WorldGradient);
 			OutClosestNormal = WorldGradientLength > 0.0f ? WorldGradient / WorldGradientLength : 0;
-
-			const float3 PreviousPosition = mul(DIRigidMeshCollision_GetPreviousTransform(DIContext,ElementIndex), float4(CollisionPosition,1.0)).xyz;
-			const float3 CurrentPosition = mul(DIRigidMeshCollision_GetCurrentTransform(DIContext,ElementIndex), float4(CollisionPosition,1.0)).xyz;
+			
+			CollisionPosition = WorldPosition - OutClosestNormal * OutClosestDistance;			
+			const float3 LocalCollisionPosition = mul(DIRigidMeshCollision_GetInverseTransform(DIContext,ElementIndex), float4(CollisionPosition,1.0)).xyz;			
+			
+			const float3 PreviousPosition = mul(DIRigidMeshCollision_GetPreviousTransform(DIContext,ElementIndex), float4(LocalCollisionPosition,1.0)).xyz;
+			const float3 CurrentPosition = mul(DIRigidMeshCollision_GetCurrentTransform(DIContext,ElementIndex), float4(LocalCollisionPosition,1.0)).xyz;
 			OutClosestVelocity = ( CurrentPosition - PreviousPosition ) / DeltaTime;
 			OutClosestPosition = PreviousPosition + TimeFraction * (CurrentPosition-PreviousPosition);
+
+			// output 			
+			NormalIsValid =  all(step(-DFObjectData.VolumePositionExtent, VolumePosition) * step(VolumePosition, DFObjectData.VolumePositionExtent));
+			OutClosestPosition = NormalIsValid ? OutClosestPosition : float3(0,0,0);
 		}
 	}
 }
 
+// Given a world space position (WorldPosition) compute the static mesh closest point (position,normal,velocity) by querying all overlapping mesh distance fields
+void DIRigidMeshCollision_GetClosestPointMeshDistanceFieldAccurate(in FDIRigidMeshCollisionContext DIContext, in float3 WorldPosition, in float DeltaTime, in float TimeFraction, in float MaxDistance, out float OutClosestDistance, out float3 OutClosestPosition, out float3 OutClosestNormal, out float3 OutClosestVelocity, out bool NormalIsValid)
+{	
+	float3 CollisionPosition = float3(0,0,0);
+	float3 CollisionNormal = float3(0,0,0);
+	OutClosestDistance = MAX_DISTANCE;	
+
+	OutClosestNormal = float3(0,0,0);
+	OutClosestPosition = float3(0,0,0);
+	OutClosestVelocity = float3(0,0,0);
+	
+	NormalIsValid = false;
+
+	for (int ElementIndex = 0; ElementIndex < DIContext.ElementOffsets[3]; ++ElementIndex)
+	{
+		uint DFIndex = DIContext.DFIndexBuffer[ElementIndex];
+
+		// #todo(dmp): this does a dynamic branch based on intersecting the bbox.  Maybe we can factor that out due to the broadphase here?
+		// #todo(dmp): Check the distance from the physics asset for a speedup?
+		MaxDistance = max(.001, MaxDistance);
+		float OutClosestDistanceTmp = DistanceToNearestSurfaceForObject(DFIndex, WorldPosition, MaxDistance);			 
+
+		if (OutClosestDistanceTmp < OutClosestDistance)
+		{
+			OutClosestDistance = OutClosestDistanceTmp;
+			FDFObjectData DFObjectData = LoadDFObjectData(DFIndex);
+			float4x4 WorldToVolume = LWCHackToFloat(DFObjectData.WorldToVolume);
+			uint NumMips = LoadDFAssetData(DFObjectData.AssetIndex, 0).NumMips;
+			FDFAssetData DFAssetData = LoadDFAssetData(DFObjectData.AssetIndex, NumMips - 1);
+			float3 VolumePosition = mul(float4(WorldPosition, 1), WorldToVolume).xyz;		
+			float3 VolumeGradient = CalculateMeshSDFGradient(VolumePosition, DFAssetData);
+			float VolumeGradientLength = length(VolumeGradient);
+			float3 VolumeNormal = VolumeGradientLength > 0.0f ? VolumeGradient / VolumeGradientLength : 0;
+			// Transform by transposed inverse to handle non-uniform scaling
+			float3 WorldGradient = mul(VolumeNormal, transpose((float3x3)WorldToVolume));
+			float WorldGradientLength = length(WorldGradient);
+			OutClosestNormal = WorldGradientLength > 0.0f ? WorldGradient / WorldGradientLength : 0;
+			
+			CollisionPosition = WorldPosition - OutClosestNormal * OutClosestDistance;			
+			const float3 LocalCollisionPosition = mul(DIRigidMeshCollision_GetInverseTransform(DIContext,ElementIndex), float4(CollisionPosition,1.0)).xyz;			
+			
+			const float3 PreviousPosition = mul(DIRigidMeshCollision_GetPreviousTransform(DIContext,ElementIndex), float4(LocalCollisionPosition,1.0)).xyz;
+			const float3 CurrentPosition = mul(DIRigidMeshCollision_GetCurrentTransform(DIContext,ElementIndex), float4(LocalCollisionPosition,1.0)).xyz;
+			OutClosestVelocity = ( CurrentPosition - PreviousPosition ) / DeltaTime;
+			OutClosestPosition = PreviousPosition + TimeFraction * (CurrentPosition-PreviousPosition);
+
+			// output 			
+			NormalIsValid =  all(step(-DFObjectData.VolumePositionExtent, VolumePosition) * step(VolumePosition, DFObjectData.VolumePositionExtent));
+			OutClosestPosition = NormalIsValid ? OutClosestPosition : float3(0,0,0);
+		}
+	}	
+}
+
 // Given a world space position (WorldPosition) compute the static mesh closest point (position,normal,velocity)
-void DIRigidMeshCollision_GetClosestPointMeshDistanceFieldNoNormal(in FDIRigidMeshCollisionContext DIContext, in float3 WorldPosition, in float DeltaTime, in float TimeFraction, out float OutClosestDistance, out float3 OutClosestPosition, out float3 OutClosestVelocity)
+void DIRigidMeshCollision_GetClosestPointMeshDistanceFieldNoNormal(in FDIRigidMeshCollisionContext DIContext, in float3 WorldPosition, in float DeltaTime, in float TimeFraction, in float MaxDistance, out float OutClosestDistance, out float3 OutClosestPosition, out float3 OutClosestVelocity)
 {	
 	float3 CollisionPosition = float3(0,0,0);
 	float3 CollisionNormal = float3(0,0,0);
@@ -470,6 +534,7 @@ void DIRigidMeshCollision_GetClosestPointMeshDistanceFieldNoNormal(in FDIRigidMe
 		if (DFIndex >= 0 && DFIndex < NumSceneObjects)
 		{
 			// #todo(dmp): this does a dynamic branch based on intersecting the bbox.  Maybe we can factor that out due to the broadphase here?
+			MaxDistance = max(.001, MaxDistance);
 			OutClosestDistance = DistanceToNearestSurfaceForObject(DFIndex, WorldPosition, 1e-3);		
 
 			const float3 PreviousPosition = mul(DIRigidMeshCollision_GetPreviousTransform(DIContext,ElementIndex), float4(CollisionPosition,1.0)).xyz;
diff --git a/Engine/Plugins/FX/Niagara/Source/Niagara/Private/NiagaraDataInterfaceRigidMeshCollisionQuery.cpp b/Engine/Plugins/FX/Niagara/Source/Niagara/Private/NiagaraDataInterfaceRigidMeshCollisionQuery.cpp
--- a/Engine/Plugins/FX/Niagara/Source/Niagara/Private/NiagaraDataInterfaceRigidMeshCollisionQuery.cpp
+++ b/Engine/Plugins/FX/Niagara/Source/Niagara/Private/NiagaraDataInterfaceRigidMeshCollisionQuery.cpp
@@ -25,6 +25,7 @@ struct FNiagaraRigidMeshCollisionDIFunctionVersion
 	{
 		InitialVersion = 0,
 		LargeWorldCoordinates = 1,
+		SetMaxDistance = 2,
 
 		VersionPlusOne,
 		LatestVersion = VersionPlusOne - 1
@@ -45,6 +46,7 @@ static const FName GetElementDistanceName(TEXT("GetElementDistance"));
 static const FName GetClosestPointName(TEXT("GetClosestPoint"));
 static const FName GetClosestDistanceName(TEXT("GetClosestDistance"));
 static const FName GetClosestPointMeshDistanceFieldName(TEXT("GetClosestPointMeshDistanceField"));
+static const FName GetClosestPointMeshDistanceFieldAccurateName(TEXT("GetClosestPointMeshDistanceFieldAccurate"));
 static const FName GetClosestPointMeshDistanceFieldNoNormalName(TEXT("GetClosestPointMeshDistanceFieldNoNormal"));
 
 
@@ -1008,11 +1010,33 @@ void UNiagaraDataInterfaceRigidMeshCollisionQuery::GetFunctions(TArray<FNiagaraF
 		Sig.Inputs.Add(FNiagaraVariable(FNiagaraTypeDefinition::GetPositionDef(), TEXT("World Position")));
 		Sig.Inputs.Add(FNiagaraVariable(FNiagaraTypeDefinition::GetFloatDef(), TEXT("Delta Time")));
 		Sig.Inputs.Add(FNiagaraVariable(FNiagaraTypeDefinition::GetFloatDef(), TEXT("Time Fraction")));
+		Sig.Inputs.Add(FNiagaraVariable(FNiagaraTypeDefinition(FNiagaraTypeDefinition::GetFloatDef()), TEXT("MaxDistance")));
 		Sig.Outputs.Add(FNiagaraVariable(FNiagaraTypeDefinition::GetFloatDef(), TEXT("Closest Distance")));
 		Sig.Outputs.Add(FNiagaraVariable(FNiagaraTypeDefinition::GetPositionDef(), TEXT("Closest Position")));
 		Sig.Outputs.Add(FNiagaraVariable(FNiagaraTypeDefinition::GetVec3Def(), TEXT("Closest Normal")));
 		Sig.Outputs.Add(FNiagaraVariable(FNiagaraTypeDefinition::GetVec3Def(), TEXT("Closest Velocity")));
+		Sig.Outputs.Add(FNiagaraVariable(FNiagaraTypeDefinition::GetBoolDef(), TEXT("Normal Is Valid")));		
+		OutFunctions.Add(Sig);
+	}
 
+	{
+		FNiagaraFunctionSignature Sig;
+		Sig.Name = GetClosestPointMeshDistanceFieldAccurateName;
+		Sig.SetDescription(LOCTEXT("GetClosestPointMeshDistanceFieldDescription", "Given a world space position, computes the distance to the closest point for the static mesh, using the mesh's distance field."));
+		Sig.SetFunctionVersion(FNiagaraRigidMeshCollisionDIFunctionVersion::LatestVersion);
+		Sig.bSupportsGPU = true;
+		Sig.bSupportsCPU = false;
+		Sig.bMemberFunction = true;
+		Sig.Inputs.Add(FNiagaraVariable(FNiagaraTypeDefinition(GetClass()), TEXT("Collision DI")));
+		Sig.Inputs.Add(FNiagaraVariable(FNiagaraTypeDefinition::GetPositionDef(), TEXT("World Position")));
+		Sig.Inputs.Add(FNiagaraVariable(FNiagaraTypeDefinition::GetFloatDef(), TEXT("Delta Time")));
+		Sig.Inputs.Add(FNiagaraVariable(FNiagaraTypeDefinition::GetFloatDef(), TEXT("Time Fraction")));
+		Sig.Inputs.Add(FNiagaraVariable(FNiagaraTypeDefinition(FNiagaraTypeDefinition::GetFloatDef()), TEXT("MaxDistance")));
+		Sig.Outputs.Add(FNiagaraVariable(FNiagaraTypeDefinition::GetFloatDef(), TEXT("Closest Distance")));
+		Sig.Outputs.Add(FNiagaraVariable(FNiagaraTypeDefinition::GetPositionDef(), TEXT("Closest Position")));
+		Sig.Outputs.Add(FNiagaraVariable(FNiagaraTypeDefinition::GetVec3Def(), TEXT("Closest Normal")));
+		Sig.Outputs.Add(FNiagaraVariable(FNiagaraTypeDefinition::GetVec3Def(), TEXT("Closest Velocity")));
+		Sig.Outputs.Add(FNiagaraVariable(FNiagaraTypeDefinition::GetBoolDef(), TEXT("Normal Is Valid")));
 		OutFunctions.Add(Sig);
 	}
 
@@ -1028,6 +1052,7 @@ void UNiagaraDataInterfaceRigidMeshCollisionQuery::GetFunctions(TArray<FNiagaraF
 		Sig.Inputs.Add(FNiagaraVariable(FNiagaraTypeDefinition::GetPositionDef(), TEXT("World Position")));
 		Sig.Inputs.Add(FNiagaraVariable(FNiagaraTypeDefinition::GetFloatDef(), TEXT("Delta Time")));
 		Sig.Inputs.Add(FNiagaraVariable(FNiagaraTypeDefinition::GetFloatDef(), TEXT("Time Fraction")));
+		Sig.Inputs.Add(FNiagaraVariable(FNiagaraTypeDefinition(FNiagaraTypeDefinition::GetFloatDef()), TEXT("MaxDistance")));
 		Sig.Outputs.Add(FNiagaraVariable(FNiagaraTypeDefinition::GetFloatDef(), TEXT("Closest Distance")));
 		Sig.Outputs.Add(FNiagaraVariable(FNiagaraTypeDefinition::GetPositionDef(), TEXT("Closest Position")));		
 		Sig.Outputs.Add(FNiagaraVariable(FNiagaraTypeDefinition::GetVec3Def(), TEXT("Closest Velocity")));
@@ -1157,23 +1182,36 @@ bool UNiagaraDataInterfaceRigidMeshCollisionQuery::GetFunctionHLSL(const FNiagar
 	else if (FunctionInfo.DefinitionName == GetClosestPointMeshDistanceFieldName)
 	{
 		static const TCHAR* FormatSample = TEXT(R"(
-		void {InstanceFunctionName}(in float3 WorldPosition, in float DeltaTime, in float TimeFraction, out float ClosestDistance, out float3 OutClosestPosition, 
-							out float3 OutClosestNormal, out float3 OutClosestVelocity)
+		void {InstanceFunctionName}(in float3 WorldPosition, in float DeltaTime, in float TimeFraction,  in float MaxDistance, out float ClosestDistance, out float3 OutClosestPosition, 
+							out float3 OutClosestNormal, out float3 OutClosestVelocity, out bool NormalIsValid)
 		{
-			{RigidMeshCollisionContextName} DIRigidMeshCollision_GetClosestPointMeshDistanceField(DIContext,WorldPosition,DeltaTime,TimeFraction, ClosestDistance,
-				OutClosestPosition,OutClosestNormal,OutClosestVelocity);
+			{RigidMeshCollisionContextName} DIRigidMeshCollision_GetClosestPointMeshDistanceField(DIContext,WorldPosition,DeltaTime,TimeFraction, MaxDistance, ClosestDistance,
+				OutClosestPosition,OutClosestNormal,OutClosestVelocity, NormalIsValid);
 		}
 		)");
 		OutHLSL += FString::Format(FormatSample, ArgsSample);
 		return true;
 	}
+	else if (FunctionInfo.DefinitionName == GetClosestPointMeshDistanceFieldAccurateName)
+	{
+	static const TCHAR* FormatSample = TEXT(R"(
+		void {InstanceFunctionName}(in float3 WorldPosition, in float DeltaTime, in float TimeFraction,  in float MaxDistance, out float ClosestDistance, out float3 OutClosestPosition, 
+							out float3 OutClosestNormal, out float3 OutClosestVelocity, out bool NormalIsValid)
+		{
+			{RigidMeshCollisionContextName} DIRigidMeshCollision_GetClosestPointMeshDistanceFieldAccurate(DIContext,WorldPosition,DeltaTime,TimeFraction, MaxDistance, ClosestDistance,
+				OutClosestPosition,OutClosestNormal,OutClosestVelocity, NormalIsValid);
+		}
+		)");
+	OutHLSL += FString::Format(FormatSample, ArgsSample);
+	return true;
+	}
 	else if (FunctionInfo.DefinitionName == GetClosestPointMeshDistanceFieldNoNormalName)
 	{
 	static const TCHAR* FormatSample = TEXT(R"(
-		void {InstanceFunctionName}(in float3 WorldPosition, in float DeltaTime, in float TimeFraction, out float ClosestDistance, out float3 OutClosestPosition, 
+		void {InstanceFunctionName}(in float3 WorldPosition, in float DeltaTime, in float TimeFraction, in float MaxDistance, out float ClosestDistance, out float3 OutClosestPosition, 
 							out float3 OutClosestVelocity)
 		{
-			{RigidMeshCollisionContextName} DIRigidMeshCollision_GetClosestPointMeshDistanceFieldNoNormal(DIContext,WorldPosition,DeltaTime,TimeFraction, ClosestDistance,
+			{RigidMeshCollisionContextName} DIRigidMeshCollision_GetClosestPointMeshDistanceFieldNoNormal(DIContext,WorldPosition,DeltaTime,TimeFraction, MaxDistance, ClosestDistance,
 				OutClosestPosition,OutClosestVelocity);
 		}
 		)");
@@ -1231,6 +1269,22 @@ bool UNiagaraDataInterfaceRigidMeshCollisionQuery::UpgradeFunctionCall(FNiagaraF
 			bChanged = true;
 		}
 	}
+
+	if (FunctionSignature.FunctionVersion < FNiagaraRigidMeshCollisionDIFunctionVersion::SetMaxDistance)
+	{
+		if (FunctionSignature.Name == GetClosestPointMeshDistanceFieldName)
+		{
+			FunctionSignature.Inputs.Add(FNiagaraVariable(FNiagaraTypeDefinition(FNiagaraTypeDefinition::GetFloatDef()), TEXT("MaxDistance")));
+			FunctionSignature.Outputs.Add(FNiagaraVariable(FNiagaraTypeDefinition(FNiagaraTypeDefinition::GetBoolDef()), TEXT("Normal Is Valid")));
+			bChanged = true;
+		}
+		if (FunctionSignature.Name == GetClosestPointMeshDistanceFieldNoNormalName)
+		{
+			FunctionSignature.Inputs.Add(FNiagaraVariable(FNiagaraTypeDefinition(FNiagaraTypeDefinition::GetFloatDef()), TEXT("MaxDistance")));
+			bChanged = true;
+		}
+	}
+
 	FunctionSignature.FunctionVersion = FNiagaraRigidMeshCollisionDIFunctionVersion::LatestVersion;
 
 	return bChanged;
@@ -1252,7 +1306,7 @@ void UNiagaraDataInterfaceRigidMeshCollisionQuery::GetParameterDefinitionHLSL(co
 #if WITH_EDITOR
 void UNiagaraDataInterfaceRigidMeshCollisionQuery::ValidateFunction(const FNiagaraFunctionSignature& Function, TArray<FText>& OutValidationErrors)
 {
-	if (Function.Name == GetClosestPointMeshDistanceFieldName || Function.Name == GetClosestPointMeshDistanceFieldNoNormalName)
+	if (Function.Name == GetClosestPointMeshDistanceFieldName || Function.Name == GetClosestPointMeshDistanceFieldNoNormalName || Function.Name == GetClosestPointMeshDistanceFieldAccurateName)
 	{
 		if (!IsMeshDistanceFieldEnabled())
 		{
