Added vectorfield function (LoadSample) to grab raw values by index rather than interpolated values

diff --git a/Engine/Plugins/FX/Niagara/Source/Niagara/Classes/NiagaraDataInterfaceVectorField.h b/Engine/Plugins/FX/Niagara/Source/Niagara/Classes/NiagaraDataInterfaceVectorField.h
--- a/Engine/Plugins/FX/Niagara/Source/Niagara/Classes/NiagaraDataInterfaceVectorField.h
+++ b/Engine/Plugins/FX/Niagara/Source/Niagara/Classes/NiagaraDataInterfaceVectorField.h
@@ -62,6 +62,7 @@ public:
 	void GetFieldBounds(FVectorVMExternalFunctionContext& Context); 
 	void GetFieldTilingAxes(FVectorVMExternalFunctionContext& Context);
 	void SampleVectorField(FVectorVMExternalFunctionContext& Context);
+	void LoadVectorField(FVectorVMExternalFunctionContext& Context);
 	
 	//	
 	FVector GetTilingAxes() const;
diff --git a/Engine/Plugins/FX/Niagara/Source/Niagara/Private/NiagaraDataInterfaceVectorField.cpp b/Engine/Plugins/FX/Niagara/Source/Niagara/Private/NiagaraDataInterfaceVectorField.cpp
--- a/Engine/Plugins/FX/Niagara/Source/Niagara/Private/NiagaraDataInterfaceVectorField.cpp
+++ b/Engine/Plugins/FX/Niagara/Source/Niagara/Private/NiagaraDataInterfaceVectorField.cpp
@@ -24,6 +24,7 @@ static const FString MaxBoundsBaseName(TEXT("MaxBounds_"));
 
 // Global VM function names, also used by the shaders code generation methods.
 static const FName SampleVectorFieldName("SampleField");
+static const FName LoadVectorFieldName("LoadField");
 static const FName GetVectorFieldTilingAxesName("FieldTilingAxes");
 static const FName GetVectorFieldDimensionsName("FieldDimensions");
 static const FName GetVectorFieldBoundsName("FieldBounds");
@@ -119,6 +120,20 @@ void UNiagaraDataInterfaceVectorField::GetFunctions(TArray<FNiagaraFunctionSigna
 		OutFunctions.Add(Sig);
 	}
 
+	{
+		FNiagaraFunctionSignature Sig;
+		Sig.Name = LoadVectorFieldName;
+		Sig.Inputs.Add(FNiagaraVariable(FNiagaraTypeDefinition(GetClass()), TEXT("Vector Field")));
+		Sig.Inputs.Add(FNiagaraVariable(FNiagaraTypeDefinition::GetIntDef(), TEXT("Index X")));
+		Sig.Inputs.Add(FNiagaraVariable(FNiagaraTypeDefinition::GetIntDef(), TEXT("Index Y")));
+		Sig.Inputs.Add(FNiagaraVariable(FNiagaraTypeDefinition::GetIntDef(), TEXT("Index Z")));
+		Sig.Outputs.Add(FNiagaraVariable(FNiagaraTypeDefinition::GetVec3Def(), TEXT("Loaded Value")));
+
+		Sig.bMemberFunction = true;
+		Sig.bRequiresContext = false;
+		OutFunctions.Add(Sig);
+	}
+
 	{
 		FNiagaraFunctionSignature Sig;
 		Sig.Name = GetVectorFieldDimensionsName;
@@ -155,12 +170,18 @@ void UNiagaraDataInterfaceVectorField::GetFunctions(TArray<FNiagaraFunctionSigna
 }
 
 DEFINE_NDI_DIRECT_FUNC_BINDER(UNiagaraDataInterfaceVectorField, SampleVectorField);
+DEFINE_NDI_DIRECT_FUNC_BINDER(UNiagaraDataInterfaceVectorField, LoadVectorField);
+
 void UNiagaraDataInterfaceVectorField::GetVMExternalFunction(const FVMExternalFunctionBindingInfo& BindingInfo, void* InstanceData, FVMExternalFunction &OutFunc)
 {
 	if (BindingInfo.Name == SampleVectorFieldName && BindingInfo.GetNumInputs() == 3 && BindingInfo.GetNumOutputs() == 3)
 	{
 		NDI_FUNC_BINDER(UNiagaraDataInterfaceVectorField, SampleVectorField)::Bind(this, OutFunc);
 	}
+	if (BindingInfo.Name == LoadVectorFieldName && BindingInfo.GetNumInputs() == 3 && BindingInfo.GetNumOutputs() == 3)
+	{
+		NDI_FUNC_BINDER(UNiagaraDataInterfaceVectorField, LoadVectorField)::Bind(this, OutFunc);
+	}
 	else if (BindingInfo.Name == GetVectorFieldDimensionsName && BindingInfo.GetNumInputs() == 0 && BindingInfo.GetNumOutputs() == 3)
 	{
 		OutFunc = FVMExternalFunction::CreateUObject(this, &UNiagaraDataInterfaceVectorField::GetFieldDimensions);
@@ -267,11 +288,11 @@ void UNiagaraDataInterfaceVectorField::GetFeedback(UNiagaraSystem* InAsset, UNia
 				const auto& DIInfo = Script->GetVMExecutableData().DataInterfaceInfo[Idx];
 				if (DIInfo.MatchesClass(GetClass()))
 				{
-					// Only the SampleField function is relevant for CPU access
+					// Only the SampleField/LoadField functions are relevant for CPU access
 					bool bHasCPUFunctionsReferenced = false;
 					for (const FNiagaraFunctionSignature& Sig : DIInfo.RegisteredFunctions)
 					{
-						if (Sig.Name == SampleVectorFieldName)
+						if (Sig.Name == SampleVectorFieldName || Sig.Name == LoadVectorFieldName)
 						{
 							bHasCPUFunctionsReferenced = true;
 							break;
@@ -380,6 +401,23 @@ bool UNiagaraDataInterfaceVectorField::GetFunctionHLSL(const FNiagaraDataInterfa
 		OutHLSL += FString::Format(FormatSample, ArgsSample);
 		return true;
 	}
+	else if (FunctionInfo.DefinitionName == LoadVectorFieldName)
+	{
+		static const TCHAR* FormatSample = TEXT(R"(
+			void {FunctionName}(int In_IndexX, int In_IndexY, int In_IndexZ, out float3 Out_Sample)
+			{
+				const int3 SampleIndex = int3(In_IndexX, In_IndexY, In_IndexZ);
+				const int MipIndex = 0;
+				Out_Sample = {TextureName}.Load(int4(SampleIndex, MipIndex)).xyz;
+			}
+		)");
+
+		TMap<FString, FStringFormatArg> ArgsSample;
+		ArgsSample.Add(TEXT("FunctionName"), FunctionInfo.InstanceName);
+		ArgsSample.Add(TEXT("TextureName"), TextureBaseName + ParamInfo.DataInterfaceHLSLSymbol);
+		OutHLSL += FString::Format(FormatSample, ArgsSample);
+		return true;
+	}
 	else if (FunctionInfo.DefinitionName == GetVectorFieldTilingAxesName)
 	{
 		static const TCHAR *FormatTilingAxes = TEXT(R"(
@@ -675,6 +713,94 @@ void UNiagaraDataInterfaceVectorField::SampleVectorField(FVectorVMExternalFuncti
 	}
 }
 
+void UNiagaraDataInterfaceVectorField::LoadVectorField(FVectorVMExternalFunctionContext& Context)
+{
+	// Input arguments...
+	VectorVM::FExternalFuncInputHandler<int32> XParam(Context);
+	VectorVM::FExternalFuncInputHandler<int32> YParam(Context);
+	VectorVM::FExternalFuncInputHandler<int32> ZParam(Context);
+
+	// Outputs...
+	VectorVM::FExternalFuncRegisterHandler<float> OutSampleX(Context);
+	VectorVM::FExternalFuncRegisterHandler<float> OutSampleY(Context);
+	VectorVM::FExternalFuncRegisterHandler<float> OutSampleZ(Context);
+
+	UVectorFieldStatic* StaticVectorField = Cast<UVectorFieldStatic>(Field);
+	UVectorFieldAnimated* AnimatedVectorField = Cast<UVectorFieldAnimated>(Field);
+
+	bool bSuccess = false;
+
+	if (StaticVectorField != nullptr && StaticVectorField->bAllowCPUAccess)
+	{
+		const FIntVector MinSampleIndex(0, 0, 0);
+		const FIntVector MaxSampleIndex(StaticVectorField->SizeX - 1, StaticVectorField->SizeY - 1, StaticVectorField->SizeZ - 1);
+
+		if (ensure(StaticVectorField->HasCPUData() && MaxSampleIndex.GetMin() >= 0))
+		{
+			for (int32 InstanceIdx = 0; InstanceIdx < Context.GetNumInstances(); ++InstanceIdx)
+			{
+				// follow d3d convention of returning 0 if index is out of bounds
+				FVector3f Result = FVector3f(0.0f, 0.0f, 0.0f);
+
+				const FIntVector SampleIndex(XParam.Get(), YParam.Get(), ZParam.Get());
+				const FIntVector ClampedIndex(
+					FMath::Clamp(SampleIndex.X, MinSampleIndex.X, MaxSampleIndex.X),
+					FMath::Clamp(SampleIndex.Y, MinSampleIndex.Y, MaxSampleIndex.Y),
+					FMath::Clamp(SampleIndex.Z, MinSampleIndex.Z, MaxSampleIndex.Z)
+				);
+
+				if (ClampedIndex == SampleIndex)
+				{
+					Result = FVector3f(StaticVectorField->Sample(SampleIndex));
+				}
+
+				// Write final output...
+				*OutSampleX.GetDest() = Result.X;
+				*OutSampleY.GetDest() = Result.Y;
+				*OutSampleZ.GetDest() = Result.Z;
+
+				XParam.Advance();
+				YParam.Advance();
+				ZParam.Advance();
+				OutSampleX.Advance();
+				OutSampleY.Advance();
+				OutSampleZ.Advance();
+			}
+		}
+	}
+
+	if (!bSuccess)
+	{
+		// TODO(mv): Add warnings?
+		if (StaticVectorField != nullptr && !StaticVectorField->bAllowCPUAccess)
+		{
+			// No access to static vector data
+		}
+		else if (AnimatedVectorField != nullptr)
+		{
+			// Animated vector field not supported
+		}
+		else if (Field == nullptr)
+		{
+			// Vector field not loaded
+		}
+
+		for (int32 InstanceIdx = 0; InstanceIdx < Context.GetNumInstances(); ++InstanceIdx)
+		{
+			*OutSampleX.GetDest() = 0.0f;
+			*OutSampleY.GetDest() = 0.0f;
+			*OutSampleZ.GetDest() = 0.0f;
+
+			XParam.Advance();
+			YParam.Advance();
+			ZParam.Advance();
+			OutSampleX.Advance();
+			OutSampleY.Advance();
+			OutSampleZ.Advance();
+		}
+	}
+}
+
 /*--------------------------------------------------------------------------------------------------------------------------*/
 
 FVector UNiagaraDataInterfaceVectorField::GetTilingAxes() const
