- Remove loose parameter bindings from main Niagara parameters
- Data Interface should now opt into the new binding path by setting UseLegacyShaderBindings to false
- BuildShaderParameters / SetShaderParameters replaces BindParameters / SetParameters / UnsetParameters
- CreateShaderStorage / GetShaderStorageType replaces CreateComputeParameters / GetComputeParametersTypeDesc.
- Shader storage should no longer be required for the majority of cases, only those which require data based on bound functions will need additional data, everything else should be statically set.

diff --git a/Engine/Plugins/FX/Niagara/Shaders/Private/NiagaraEmitterInstanceShader.usf b/Engine/Plugins/FX/Niagara/Shaders/Private/NiagaraEmitterInstanceShader.usf
--- a/Engine/Plugins/FX/Niagara/Shaders/Private/NiagaraEmitterInstanceShader.usf
+++ b/Engine/Plugins/FX/Niagara/Shaders/Private/NiagaraEmitterInstanceShader.usf
@@ -86,9 +86,7 @@ bool NiagaraAny(bool4 b) { return b.x || b.y || b.z || b.w; }
 
 	// Physics common 
 	//-TOFIX: This should not be included here
-	#if USE_SIMULATION_STAGES == 1
-		#include "NiagaraPhysicsCommon.ush"
-	#endif
+	#include "NiagaraPhysicsCommon.ush"
 
 	uint EmitterTickCounter;
 
@@ -132,7 +130,7 @@ bool NiagaraAny(bool4 b) { return b.x || b.y || b.z || b.w; }
 	static int GGPUExecIndex;
 
 	static uint GSpawnStartInstance;
-	uint SpawnedInstances;
+	uint NumSpawnedInstances;
 
 	void SetupExecIndexForGPU()
 	{
@@ -163,7 +161,7 @@ bool NiagaraAny(bool4 b) { return b.x || b.y || b.z || b.w; }
 		GGPUExecIndex		        = GGPUExecIndex - GroupSpawnStartIndex;
 		if ( SpawnInfoIndex == (NIAGARA_MAX_GPU_SPAWN_INFOS - 1) )
 		{
-			Engine_ExecutionCount   = int(SpawnedInstances) - GroupSpawnStartIndex;
+			Engine_ExecutionCount   = int(NumSpawnedInstances) - GroupSpawnStartIndex;
 		}
 		else
 		{
@@ -585,7 +583,7 @@ int rand_int(int x, int Seed1, int Seed2, int Seed3)
 		MustBe0[DataSetID] = 0;
 		// End static assert
 
-		int PrevIdx = GSpawnStartInstance + SpawnedInstances;	// scratch instance as default; write to that for dead particles
+		int PrevIdx = GSpawnStartInstance + NumSpawnedInstances;	// scratch instance as default; write to that for dead particles
 
 		#if USE_WAVE_INTRINSICS
 
@@ -766,19 +764,19 @@ int rand_int(int x, int Seed1, int Seed2, int Seed3)
 	 */
 	int OutputIndex(const int DataSetID, const bool bStageKillsParticles, const bool bIsValid)
 	{
-		#if USE_SIMULATION_STAGES
-			// If this stage cannot kill particles, we can just write them out in the same order as they
-			// appear in the input. We must use an if here (as opposed to a ternary operator, or some
-			// other branchless construct), because we don't want to call AcquireIndex() at all, since
-			// that manipulates the RWInstanceCounts UAV. The generated code will copy the source count
-			// at the end of the shader.
-			if (!bStageKillsParticles)
-			{
-				return GLinearThreadId;
-			}
-		#endif // USE_SIMULATION_STAGES
+	#if GPU_SIMULATION
+		// If this stage cannot kill particles, we can just write them out in the same order as they
+		// appear in the input. We must use an if here (as opposed to a ternary operator, or some
+		// other branchless construct), because we don't want to call AcquireIndex() at all, since
+		// that manipulates the RWInstanceCounts UAV. The generated code will copy the source count
+		// at the end of the shader.
+		if (!bStageKillsParticles)
+		{
+			return GLinearThreadId;
+		}
+	#endif
 
-			return AcquireIndex(DataSetID, bIsValid);
+		return AcquireIndex(DataSetID, bIsValid);
 	}
 
 ////////////////////////////////////////////////////////////////////////////////////
diff --git a/Engine/Plugins/FX/Niagara/Source/Niagara/Classes/NiagaraDataInterface.h b/Engine/Plugins/FX/Niagara/Source/Niagara/Classes/NiagaraDataInterface.h
--- a/Engine/Plugins/FX/Niagara/Source/Niagara/Classes/NiagaraDataInterface.h
+++ b/Engine/Plugins/FX/Niagara/Source/Niagara/Classes/NiagaraDataInterface.h
@@ -12,13 +12,16 @@
 #include "NiagaraDataInterface.generated.h"
 
 class INiagaraCompiler;
-class UCanvas;
 class UCurveVector;
 class UCurveLinearColor;
 class UCurveFloat;
-class FNiagaraSystemInstance;
-class FNiagaraEmitterInstance;
 struct FNiagaraDataInterfaceProxy;
+struct FNiagaraComputeInstanceData;
+class FNiagaraEmitterInstance;
+class FNiagaraGpuComputeDispatchInterface;
+class FNiagaraGPUSystemTick;
+struct FNiagaraSimStageData;
+class FNiagaraSystemInstance;
 
 struct FNDITransformHandlerNoop
 {
@@ -268,6 +271,77 @@ struct FNiagaraDataInterfaceProxy
 
 //////////////////////////////////////////////////////////////////////////
 
+struct FNiagaraDataInterfaceSetShaderParametersContext
+{
+	friend class FNiagaraGpuComputeDispatch;
+
+	FNiagaraDataInterfaceSetShaderParametersContext(FRHICommandList& InRHICmdList, const FNiagaraGpuComputeDispatchInterface& InComputeDispatchInterface, const FNiagaraGPUSystemTick& InSystemTick, const FNiagaraComputeInstanceData& InComputeInstanceData, const FNiagaraSimStageData& InSimStageData, const FNiagaraShaderRef& InShaderRef, uint8* InBaseParameters)
+		: RHICmdList(InRHICmdList)
+		, ComputeDispatchInterface(InComputeDispatchInterface)
+		, SystemTick(InSystemTick)
+		, ComputeInstanceData(InComputeInstanceData)
+		, SimStageData(InSimStageData)
+		, ShaderRef(InShaderRef)
+		, BaseParameters(InBaseParameters)
+	{
+	}
+
+	template<typename T> T& GetProxy() const { check(DataInterfaceProxy); return static_cast<T&>(*DataInterfaceProxy); }
+	const FNiagaraGpuComputeDispatchInterface& GetComputeDispatchInterface() const { return ComputeDispatchInterface; }
+	const FNiagaraGPUSystemTick& GetSystemTick() const { return SystemTick; }
+	const FNiagaraComputeInstanceData& GetComputeInstanceData() const { return ComputeInstanceData; }
+	const FNiagaraSimStageData& GetSimStageData() const { return SimStageData; }
+	FRHICommandList& GetRHICmdList() const { return RHICmdList; }
+
+	FNiagaraSystemInstanceID GetSystemInstanceID() const;
+	bool IsResourceBound(const void* ResourceAddress) const;
+	bool IsOutputStage() const;
+	bool IsIterationStage() const;
+
+
+	template<typename T> T* GetParameterNestedStruct() const
+	{
+		const uint32 StructOffset = Align(ParametersOffset, TShaderParameterStructTypeInfo<T>::Alignment);
+		ParametersOffset = StructOffset + TShaderParameterStructTypeInfo<T>::GetStructMetadata()->GetSize();
+
+		return reinterpret_cast<T*>(BaseParameters + StructOffset);
+	}
+	template<typename T> TArrayView<T> GetParameterLooseArray(int32 NumElements) const
+	{
+		const uint32 ArrayOffset = Align(ParametersOffset, SHADER_PARAMETER_ARRAY_ELEMENT_ALIGNMENT);
+		ParametersOffset = ArrayOffset + (sizeof(T) * NumElements);
+		return TArrayView<T>(reinterpret_cast<T*>(BaseParameters + ArrayOffset), NumElements);
+	}
+
+	template<typename T> const T& GetShaderStorage() const
+	{
+		check(ShaderStorage != nullptr);
+		return static_cast<const T&>(*ShaderStorage);
+	}
+
+	void SetDataInterface(FNiagaraDataInterfaceProxy* InDataInterfaceProxy, uint32 InParametersOffset, const FNiagaraDataInterfaceParametersCS* InShaderStorage)
+	{
+		DataInterfaceProxy = InDataInterfaceProxy;
+		ParametersOffset = InParametersOffset;
+		ShaderStorage = InShaderStorage;
+	}
+
+private:
+	FRHICommandList& RHICmdList;
+	const FNiagaraGpuComputeDispatchInterface& ComputeDispatchInterface;
+	const FNiagaraGPUSystemTick& SystemTick;
+	const FNiagaraComputeInstanceData& ComputeInstanceData;
+	const FNiagaraSimStageData& SimStageData;
+	const FNiagaraShaderRef& ShaderRef;
+	uint8* BaseParameters = nullptr;
+	const FNiagaraDataInterfaceParametersCS* ShaderStorage = nullptr;
+
+	FNiagaraDataInterfaceProxy* DataInterfaceProxy = nullptr;
+	mutable uint32 ParametersOffset = 0;
+};
+
+//////////////////////////////////////////////////////////////////////////
+
 /** Base class for all Niagara data interfaces. */
 UCLASS(abstract, EditInlineNew)
 class NIAGARA_API UNiagaraDataInterface : public UNiagaraDataInterfaceBase
@@ -491,6 +565,12 @@ public:
 
 	virtual bool ReadsEmitterParticleData(const FString& EmitterName) const { return false; }
 
+	/**
+	Set the shader parameters will only be called if the data interface provided shader parameters.
+	You must write the parameters in the order you added them to the structure.
+	*/
+	virtual void SetShaderParameters(const FNiagaraDataInterfaceSetShaderParametersContext& Context) const {}
+
 protected:
 	virtual void PushToRenderThreadImpl() {}
 
diff --git a/Engine/Plugins/FX/Niagara/Source/Niagara/Classes/NiagaraDataSet.h b/Engine/Plugins/FX/Niagara/Source/Niagara/Classes/NiagaraDataSet.h
--- a/Engine/Plugins/FX/Niagara/Source/Niagara/Classes/NiagaraDataSet.h
+++ b/Engine/Plugins/FX/Niagara/Source/Niagara/Classes/NiagaraDataSet.h
@@ -188,10 +188,6 @@ public:
 	FORCEINLINE TArray<int32>& GetIDTable() { return IDToIndexTable; }
 	FORCEINLINE const TArray<int32>& GetIDTable() const { return IDToIndexTable; }
 
-	static void SetInputShaderParams(FRHICommandList& RHICmdList, class FNiagaraShader* Shader, FNiagaraDataBuffer* Buffer);
-	static void SetOutputShaderParams(FRHICommandList& RHICmdList, class FNiagaraShader* Shader, FNiagaraDataBuffer* Buffer);
-	static void UnsetShaderParams(FRHICommandList& RHICmdList, class FNiagaraShader* Shader);
-
 	FORCEINLINE void ClearGPUInstanceCount() { GPUInstanceCountBufferOffset = INDEX_NONE; }
 
 	void BuildRegisterTable();
diff --git a/Engine/Plugins/FX/Niagara/Source/Niagara/Private/DataInterface/NiagaraDataInterfaceUObjectPropertyReader.cpp b/Engine/Plugins/FX/Niagara/Source/Niagara/Private/DataInterface/NiagaraDataInterfaceUObjectPropertyReader.cpp
--- a/Engine/Plugins/FX/Niagara/Source/Niagara/Private/DataInterface/NiagaraDataInterfaceUObjectPropertyReader.cpp
+++ b/Engine/Plugins/FX/Niagara/Source/Niagara/Private/DataInterface/NiagaraDataInterfaceUObjectPropertyReader.cpp
@@ -607,7 +607,8 @@ bool UNiagaraDataInterfaceUObjectPropertyReader::InitPerInstanceData(void* PerIn
 			}
 			const FNiagaraScriptInstanceParameterStore& ParameterStore = EmitterInstance->GetGPUContext()->CombinedParamStore;
 			const TArray<UNiagaraDataInterface*>& DataInterfaces = ParameterStore.GetDataInterfaces();
-			const TArray<FNiagaraDataInterfaceGPUParamInfo>& DataInterfaceParamInfo = EmitterInstance->GetGPUContext()->GPUScript_RT->GetDataInterfaceParamInfo();
+			const TSharedRef<FNiagaraShaderScriptParametersMetadata> ScriptParametersMetadata = EmitterInstance->GetGPUContext()->GPUScript_RT->GetScriptParametersMetadata();
+			const TArray<FNiagaraDataInterfaceGPUParamInfo>& DataInterfaceParamInfo = ScriptParametersMetadata->DataInterfaceParamInfo;
 			for ( int32 iDataInterface=0; iDataInterface < DataInterfaces.Num(); ++iDataInterface)
 			{
 				if ( (DataInterfaces[iDataInterface] == this) && DataInterfaceParamInfo.IsValidIndex(iDataInterface) )
diff --git a/Engine/Plugins/FX/Niagara/Source/Niagara/Private/DataInterface/NiagaraDataInterfaceUtilities.cpp b/Engine/Plugins/FX/Niagara/Source/Niagara/Private/DataInterface/NiagaraDataInterfaceUtilities.cpp
--- a/Engine/Plugins/FX/Niagara/Source/Niagara/Private/DataInterface/NiagaraDataInterfaceUtilities.cpp
+++ b/Engine/Plugins/FX/Niagara/Source/Niagara/Private/DataInterface/NiagaraDataInterfaceUtilities.cpp
@@ -143,7 +143,7 @@ void FNiagaraDataInterfaceUtilities::ForEachGpuFunctionEquals(class UNiagaraData
 			if (FNiagaraShaderScript* ShaderScript = Script->GetRenderThreadScript())
 			{
 				const TArray<FNiagaraScriptDataInterfaceInfo>& CachedDefaultDIs = Script->GetCachedDefaultDataInterfaces();
-				const TArray<FNiagaraDataInterfaceGPUParamInfo>& DataInterfaceParamInfos = ShaderScript->GetDataInterfaceParamInfo();
+				const TArray<FNiagaraDataInterfaceGPUParamInfo>& DataInterfaceParamInfos = ShaderScript->GetScriptParametersMetadata()->DataInterfaceParamInfo;
 
 				const int NumDataInterfaces = FMath::Min(CachedDefaultDIs.Num(), DataInterfaceParamInfos.Num());	// Note: Should always be equal but lets be safe
 				for (int iDataInterface = 0; iDataInterface < NumDataInterfaces; ++iDataInterface)
@@ -275,8 +275,8 @@ void FNiagaraDataInterfaceUtilities::ForEachGpuFunction(class UNiagaraDataInterf
 
 		const FNiagaraScriptInstanceParameterStore& ParameterStore = EmitterInstance->GetGPUContext()->CombinedParamStore;
 		const TArray<UNiagaraDataInterface*>& DataInterfaces = ParameterStore.GetDataInterfaces();
-		const TArray<FNiagaraDataInterfaceGPUParamInfo>& DataInterfaceParamInfo = EmitterInstance->GetGPUContext()->GPUScript_RT->GetDataInterfaceParamInfo();
-		const int32 NumDataInterface = FMath::Min(DataInterfaces.Num(), DataInterfaceParamInfo.Num());		// Should be equal, but be safe
+		const TArray<FNiagaraDataInterfaceGPUParamInfo>& DataInterfaceParamInfos = EmitterInstance->GetGPUContext()->GPUScript_RT->GetScriptParametersMetadata()->DataInterfaceParamInfo;
+		const int32 NumDataInterface = FMath::Min(DataInterfaces.Num(), DataInterfaceParamInfos.Num());		// Should be equal, but be safe
 		for ( int32 iDataInterface=0; iDataInterface < NumDataInterface; ++iDataInterface)
 		{
 			if ( DataInterfaces[iDataInterface] != DataInterface )
@@ -284,7 +284,7 @@ void FNiagaraDataInterfaceUtilities::ForEachGpuFunction(class UNiagaraDataInterf
 				continue;
 			}
 
-			for (const FNiagaraDataInterfaceGeneratedFunction& GeneratedFunction : DataInterfaceParamInfo[iDataInterface].GeneratedFunctions)
+			for (const FNiagaraDataInterfaceGeneratedFunction& GeneratedFunction : DataInterfaceParamInfos[iDataInterface].GeneratedFunctions)
 			{
 				if (Action(GeneratedFunction) == false)
 				{
diff --git a/Engine/Plugins/FX/Niagara/Source/Niagara/Private/NiagaraComputeExecutionContext.cpp b/Engine/Plugins/FX/Niagara/Source/Niagara/Private/NiagaraComputeExecutionContext.cpp
--- a/Engine/Plugins/FX/Niagara/Source/Niagara/Private/NiagaraComputeExecutionContext.cpp
+++ b/Engine/Plugins/FX/Niagara/Source/Niagara/Private/NiagaraComputeExecutionContext.cpp
@@ -71,8 +71,9 @@ void FNiagaraComputeExecutionContext::InitParams(UNiagaraScript* InGPUComputeScr
 	}
 	else
 	{
-		DIClassNames.Empty(InGPUComputeScript->GetRenderThreadScript()->GetDataInterfaceParamInfo().Num());
-		for (const FNiagaraDataInterfaceGPUParamInfo& DIParams : InGPUComputeScript->GetRenderThreadScript()->GetDataInterfaceParamInfo())
+		TSharedRef<FNiagaraShaderScriptParametersMetadata> ScriptParametersMetadata = InGPUComputeScript->GetRenderThreadScript()->GetScriptParametersMetadata();
+		DIClassNames.Empty(ScriptParametersMetadata->DataInterfaceParamInfo.Num());
+		for (const FNiagaraDataInterfaceGPUParamInfo& DIParams : ScriptParametersMetadata->DataInterfaceParamInfo)
 		{
 			DIClassNames.Add(DIParams.DIClassName);
 		}
diff --git a/Engine/Plugins/FX/Niagara/Source/Niagara/Private/NiagaraDataSet.cpp b/Engine/Plugins/FX/Niagara/Source/Niagara/Private/NiagaraDataSet.cpp
--- a/Engine/Plugins/FX/Niagara/Source/Niagara/Private/NiagaraDataSet.cpp
+++ b/Engine/Plugins/FX/Niagara/Source/Niagara/Private/NiagaraDataSet.cpp
@@ -1213,90 +1213,6 @@ void FNiagaraDataBuffer::Dump(int32 StartIndex, int32 InNumInstances, const FStr
 
 }
 
-/////////////////////////////////////////////////////////////////////////
-
-void FNiagaraDataBuffer::SetInputShaderParams(FRHICommandList& RHICmdList, class FNiagaraShader* Shader, FNiagaraDataBuffer* Buffer)
-{
-	check(IsInRenderingThread());
-
-	if (Buffer != nullptr)
-	{
-		const uint32 SafeBufferSize = Buffer->GetFloatStride() / sizeof(float);
-		FRHIComputeShader* ComputeShader = RHICmdList.GetBoundComputeShader();
-
-		const bool InstancesAllocated = Buffer->GetNumInstancesAllocated() > 0;
-
-		SetSRVParameter(RHICmdList, ComputeShader, Shader->FloatInputBufferParam, InstancesAllocated ? Buffer->GetGPUBufferFloat().SRV.GetReference() : FNiagaraRenderer::GetDummyFloatBuffer());
-		SetSRVParameter(RHICmdList, ComputeShader, Shader->IntInputBufferParam, InstancesAllocated ? Buffer->GetGPUBufferInt().SRV.GetReference() : FNiagaraRenderer::GetDummyIntBuffer());
-		SetSRVParameter(RHICmdList, ComputeShader, Shader->HalfInputBufferParam, InstancesAllocated ? Buffer->GetGPUBufferHalf().SRV.GetReference() : FNiagaraRenderer::GetDummyHalfBuffer());
-		SetShaderValue(RHICmdList, ComputeShader, Shader->ComponentBufferSizeReadParam, SafeBufferSize);
-	}
-	else
-	{
-		check(!Shader->FloatInputBufferParam.IsBound());
-		check(!Shader->IntInputBufferParam.IsBound());
-		check(!Shader->HalfInputBufferParam.IsBound());
-		check(Shader->ComponentBufferSizeReadParam.GetNumBytes() == 0);
-	}
-}
-
-void FNiagaraDataBuffer::SetOutputShaderParams(FRHICommandList& RHICmdList, class FNiagaraShader* Shader, FNiagaraDataBuffer* Buffer)
-{
-	check(IsInRenderingThread());
-
-	if (Buffer != nullptr)
-	{
-		const uint32 SafeBufferSize = Buffer->GetFloatStride() / sizeof(float);
-		FRHIComputeShader* ComputeShader = RHICmdList.GetBoundComputeShader();
-
-		Shader->FloatOutputBufferParam.SetBuffer(RHICmdList, ComputeShader, Buffer->GetGPUBufferFloat());
-		Shader->IntOutputBufferParam.SetBuffer(RHICmdList, ComputeShader, Buffer->GetGPUBufferInt());
-		Shader->HalfOutputBufferParam.SetBuffer(RHICmdList, ComputeShader, Buffer->GetGPUBufferHalf());
-		SetShaderValue(RHICmdList, ComputeShader, Shader->ComponentBufferSizeWriteParam, SafeBufferSize);
-		if (Shader->IDToIndexBufferParam.IsUAVBound())
-		{
-			ensure(Buffer->GPUIDToIndexTable.Buffer);
-			Shader->IDToIndexBufferParam.SetBuffer(RHICmdList, ComputeShader, Buffer->GPUIDToIndexTable);
-		}
-	}
-	else
-	{
-		check(!Shader->FloatOutputBufferParam.IsUAVBound());
-		check(!Shader->IntOutputBufferParam.IsUAVBound());
-		check(!Shader->HalfOutputBufferParam.IsUAVBound());
-		check(!Shader->IDToIndexBufferParam.IsUAVBound());
-		check(Shader->ComponentBufferSizeWriteParam.GetNumBytes() == 0);
-	}
-}
-
-void FNiagaraDataBuffer::UnsetShaderParams(FRHICommandList& RHICmdList, FNiagaraShader* Shader)
-{
-	check(IsInRenderingThread());
-	FRHIComputeShader* ShaderRHI = RHICmdList.GetBoundComputeShader();
-
-	if (Shader->FloatOutputBufferParam.IsUAVBound())
-	{
-		Shader->FloatOutputBufferParam.UnsetUAV(RHICmdList, ShaderRHI);
-	}
-
-	if (Shader->HalfOutputBufferParam.IsUAVBound())
-	{
-#if !PLATFORM_PS4
-		Shader->HalfOutputBufferParam.UnsetUAV(RHICmdList, ShaderRHI);
-#endif
-	}
-
-	if (Shader->IntOutputBufferParam.IsUAVBound())
-	{
-		Shader->IntOutputBufferParam.UnsetUAV(RHICmdList, ShaderRHI);
-	}
-
-	if (Shader->IDToIndexBufferParam.IsUAVBound())
-	{
-		Shader->IDToIndexBufferParam.UnsetUAV(RHICmdList, RHICmdList.GetBoundComputeShader());
-	}
-}
-
 #if WITH_EDITOR
 void FScopedNiagaraDataSetGPUReadback::ReadbackData(FNiagaraGpuComputeDispatchInterface* InComputeDispatchInterface, FNiagaraDataSet* InDataSet)
 {
diff --git a/Engine/Plugins/FX/Niagara/Source/Niagara/Private/NiagaraEmitter.cpp b/Engine/Plugins/FX/Niagara/Source/Niagara/Private/NiagaraEmitter.cpp
--- a/Engine/Plugins/FX/Niagara/Source/Niagara/Private/NiagaraEmitter.cpp
+++ b/Engine/Plugins/FX/Niagara/Source/Niagara/Private/NiagaraEmitter.cpp
@@ -1397,7 +1397,7 @@ void UNiagaraEmitter::CacheFromShaderCompiled()
 			for (int i=0; i < NiagaraShaderScript->GetNumPermutations(); ++i)
 			{
 				FNiagaraShaderRef NiagaraShaderRef = NiagaraShaderScript->GetShaderGameThread(i);
-				if (NiagaraShaderRef.IsValid() && NiagaraShaderRef->ViewUniformBufferParam.IsBound())
+				if (NiagaraShaderRef.IsValid() && NiagaraShaderRef->bNeedsViewUniformBuffer)
 				{
 					bRequiresViewUniformBuffer = true;
 					break;
diff --git a/Engine/Plugins/FX/Niagara/Source/Niagara/Private/NiagaraGpuComputeDispatch.cpp b/Engine/Plugins/FX/Niagara/Source/Niagara/Private/NiagaraGpuComputeDispatch.cpp
--- a/Engine/Plugins/FX/Niagara/Source/Niagara/Private/NiagaraGpuComputeDispatch.cpp
+++ b/Engine/Plugins/FX/Niagara/Source/Niagara/Private/NiagaraGpuComputeDispatch.cpp
@@ -79,6 +79,11 @@ static FAutoConsoleVariableRef CVarNiagaraBatcherFreeBufferEarly(
 
 const FName FNiagaraGpuComputeDispatch::Name(TEXT("FNiagaraGpuComputeDispatch"));
 
+BEGIN_SHADER_PARAMETER_STRUCT(FNiagaraComputePassParameters, )
+	SHADER_PARAMETER_STRUCT_INCLUDE(FNiagaraSceneTextureParameters,	SceneTextures)
+	SHADER_PARAMETER_STRUCT_REF(FViewUniformShaderParameters, View)
+END_SHADER_PARAMETER_STRUCT()
+
 namespace FNiagaraGpuComputeDispatchLocal
 {
 #if WITH_MGPU
@@ -572,7 +577,7 @@ void FNiagaraGpuComputeDispatch::ResetDataInterfaces(FRHICommandList& RHICmdList
 	for (FNiagaraDataInterfaceProxy* Interface : InstanceData.DataInterfaceProxies)
 	{
 		const FNiagaraDataInterfaceParamRef& DIParam = DIParameters[InterfaceIndex];
-		if (DIParam.Parameters.IsValid())
+		if (DIParam.Parameters.IsValid() || (DIParam.ShaderParametersOffset != INDEX_NONE))
 		{
 			const FNiagaraDataInterfaceArgs TmpContext(Interface, Tick.SystemInstanceID, Tick.SystemGpuComputeProxy->GetSystemLWCTile(), this);
 			Interface->ResetData(RHICmdList, TmpContext);
@@ -597,7 +602,7 @@ void FNiagaraGpuComputeDispatch::PreStageInterface(FRHICommandList& RHICmdList,
 	for (FNiagaraDataInterfaceProxy* Interface : InstanceData.DataInterfaceProxies)
 	{
 		const FNiagaraDataInterfaceParamRef& DIParam = DIParameters[InterfaceIndex];
-		if (DIParam.Parameters.IsValid())
+		if (DIParam.Parameters.IsValid() || (DIParam.ShaderParametersOffset != INDEX_NONE))
 		{
 			const FNiagaraDataInterfaceStageArgs TmpContext(Interface, Tick.SystemInstanceID, Tick.SystemGpuComputeProxy->GetSystemLWCTile(), this, &InstanceData, &SimStageData, InstanceData.IsOutputStage(Interface, SimStageData.StageIndex), InstanceData.IsIterationStage(Interface, SimStageData.StageIndex));
 			Interface->PreStage(RHICmdList, TmpContext);
@@ -621,7 +626,7 @@ void FNiagaraGpuComputeDispatch::PostStageInterface(FRHICommandList& RHICmdList,
 	for (FNiagaraDataInterfaceProxy* Interface : InstanceData.DataInterfaceProxies)
 	{
 		const FNiagaraDataInterfaceParamRef& DIParam = DIParameters[InterfaceIndex];
-		if (DIParam.Parameters.IsValid())
+		if (DIParam.Parameters.IsValid() || (DIParam.ShaderParametersOffset != INDEX_NONE))
 		{
 			const FNiagaraDataInterfaceStageArgs TmpContext(Interface, Tick.SystemInstanceID, Tick.SystemGpuComputeProxy->GetSystemLWCTile(), this, &InstanceData, &SimStageData, InstanceData.IsOutputStage(Interface, SimStageData.StageIndex), InstanceData.IsIterationStage(Interface, SimStageData.StageIndex));
 			Interface->PostStage(RHICmdList, TmpContext);
@@ -645,7 +650,7 @@ void FNiagaraGpuComputeDispatch::PostSimulateInterface(FRHICommandList& RHICmdLi
 	for (FNiagaraDataInterfaceProxy* Interface : InstanceData.DataInterfaceProxies)
 	{
 		const FNiagaraDataInterfaceParamRef& DIParam = DIParameters[InterfaceIndex];
-		if (DIParam.Parameters.IsValid())
+		if (DIParam.Parameters.IsValid() || (DIParam.ShaderParametersOffset != INDEX_NONE))
 		{
 			const FNiagaraDataInterfaceArgs TmpContext(Interface, Tick.SystemInstanceID, Tick.SystemGpuComputeProxy->GetSystemLWCTile(), this);
 			Interface->PostSimulate(RHICmdList, TmpContext);
@@ -1215,7 +1220,7 @@ void FNiagaraGpuComputeDispatch::PrepareAllTicks(FRHICommandListImmediate& RHICm
 	}
 }
 
-void FNiagaraGpuComputeDispatch::ExecuteTicks(FRHICommandList& RHICmdList, FRHIUniformBuffer* ViewUniformBuffer, ENiagaraGpuComputeTickStage::Type TickStage)
+void FNiagaraGpuComputeDispatch::ExecuteTicks(FRHICommandList& RHICmdList, ENiagaraGpuComputeTickStage::Type TickStage)
 {
 #if WITH_MGPU
 	WaitForMultiGPUBuffers(RHICmdList, TickStage);
@@ -1346,7 +1351,7 @@ void FNiagaraGpuComputeDispatch::ExecuteTicks(FRHICommandList& RHICmdList, FRHIU
 			}
 
 			FNiagaraGpuProfileScope GpuProfileDispatchScope(RHICmdList, this, DispatchInstance);
-			DispatchStage(RHICmdList, ViewUniformBuffer, DispatchInstance.Tick, DispatchInstance.InstanceData, DispatchInstance.SimStageData);
+			DispatchStage(RHICmdList, DispatchInstance.Tick, DispatchInstance.InstanceData, DispatchInstance.SimStageData);
 		}
 		RHICmdList.EndUAVOverlap(GPUInstanceCounterManager.GetInstanceCountBuffer().UAV);
 
@@ -1444,7 +1449,7 @@ void FNiagaraGpuComputeDispatch::ExecuteTicks(FRHICommandList& RHICmdList, FRHIU
 	DispatchList.DispatchGroups.Empty();
 }
 
-void FNiagaraGpuComputeDispatch::DispatchStage(FRHICommandList& RHICmdList, FRHIUniformBuffer* ViewUniformBuffer, const FNiagaraGPUSystemTick& Tick, const FNiagaraComputeInstanceData& InstanceData, const FNiagaraSimStageData& SimStageData)
+void FNiagaraGpuComputeDispatch::DispatchStage(FRHICommandList& RHICmdList, const FNiagaraGPUSystemTick& Tick, const FNiagaraComputeInstanceData& InstanceData, const FNiagaraSimStageData& SimStageData)
 {
 	// Setup source buffer
 	if ( SimStageData.Source != nullptr )
@@ -1513,11 +1518,6 @@ void FNiagaraGpuComputeDispatch::DispatchStage(FRHICommandList& RHICmdList, FRHI
 
 	checkf(DispatchNumThreads.X * DispatchNumThreads.Y * DispatchNumThreads.Z > 0, TEXT("DispatchNumThreads(%d, %d, %d) is invalid"), DispatchNumThreads.X, DispatchNumThreads.Y, DispatchNumThreads.Z);
 
-	// Get Shader
-	const FNiagaraShaderRef ComputeShader = InstanceData.Context->GPUScript_RT->GetShader(SimStageData.StageIndex);
-	FRHIComputeShader* RHIComputeShader = ComputeShader.GetComputeShader();
-	SetComputePipelineState(RHICmdList, RHIComputeShader);
-
 	SCOPED_DRAW_EVENTF(RHICmdList, NiagaraGPUSimulationCS,
 		TEXT("NiagaraGpuSim(%s) DispatchCount(%dx%dx%d) Stage(%s %u) Iteration(%u) NumThreads(%dx%dx%d)"),
 		InstanceData.Context->GetDebugSimName(),
@@ -1529,46 +1529,35 @@ void FNiagaraGpuComputeDispatch::DispatchStage(FRHICommandList& RHICmdList, FRHI
 	);
 	FNiagaraEmptyUAVPoolScopedAccess UAVPoolAccessScope(GetEmptyUAVPool());
 
+	//-TODO: Optimize, we don't need to keep allocating for each dispatch
+	FMemStack& MemStack = FMemStack::Get();
+	FMemMark Mark(MemStack);
+	const FShaderParametersMetadata& ShaderParametersMetadata = *InstanceData.Context->GPUScript_RT->GetScriptParametersMetadata()->ShaderParametersMetadata.Get();
+	FNiagaraShader::FParameters* DispatchParameters = reinterpret_cast<FNiagaraShader::FParameters*>(MemStack.Alloc(ShaderParametersMetadata.GetSize(), SHADER_PARAMETER_STRUCT_ALIGNMENT));
+	FMemory::Memset(DispatchParameters, 0, ShaderParametersMetadata.GetSize());
+
 	// Set Parameters
 	const bool bRequiresPersistentIDs = InstanceData.Context->MainDataSet->RequiresPersistentIDs();
 
-	SetShaderValue(RHICmdList, RHIComputeShader, ComputeShader->SimStartParam, InstanceData.bResetData ? 1U : 0U);
-	SetShaderValue(RHICmdList, RHIComputeShader, ComputeShader->EmitterTickCounterParam, FNiagaraComputeExecutionContext::TickCounter);
-	SetShaderValue(RHICmdList, RHIComputeShader, ComputeShader->NumSpawnedInstancesParam, InstancesToSpawn);
-	SetSRVParameter(RHICmdList, RHIComputeShader, ComputeShader->FreeIDBufferParam, bRequiresPersistentIDs ? InstanceData.Context->MainDataSet->GetGPUFreeIDs().SRV.GetReference() : FNiagaraRenderer::GetDummyIntBuffer());
+	DispatchParameters->SimStart = InstanceData.bResetData ? 1U : 0U;
+	DispatchParameters->EmitterTickCounter = FNiagaraComputeExecutionContext::TickCounter;
+	DispatchParameters->NumSpawnedInstances = InstancesToSpawn;
+	DispatchParameters->FreeIDList = bRequiresPersistentIDs ? InstanceData.Context->MainDataSet->GetGPUFreeIDs().SRV.GetReference() : FNiagaraRenderer::GetDummyIntBuffer();
 
 	// Set spawn Information
 	// This parameter is an array of structs with 2 floats and 2 ints on CPU, but a float4 array on GPU. The shader uses asint() to cast the integer values. To set the parameter,
 	// we pass the structure array as a float* to SetShaderValueArray() and specify the number of floats (not float vectors).
 	static_assert((sizeof(InstanceData.SpawnInfo.SpawnInfoStartOffsets) % SHADER_PARAMETER_ARRAY_ELEMENT_ALIGNMENT) == 0, "sizeof SpawnInfoStartOffsets should be a multiple of SHADER_PARAMETER_ARRAY_ELEMENT_ALIGNMENT");
 	static_assert((sizeof(InstanceData.SpawnInfo.SpawnInfoParams) % SHADER_PARAMETER_ARRAY_ELEMENT_ALIGNMENT) == 0, "sizeof SpawnInfoParams should be a multiple of SHADER_PARAMETER_ARRAY_ELEMENT_ALIGNMENT");
-	SetShaderValueArray(RHICmdList, RHIComputeShader, ComputeShader->EmitterSpawnInfoOffsetsParam, InstanceData.SpawnInfo.SpawnInfoStartOffsets, NIAGARA_MAX_GPU_SPAWN_INFOS);
-	SetShaderValueArray(RHICmdList, RHIComputeShader, ComputeShader->EmitterSpawnInfoParamsParam, &InstanceData.SpawnInfo.SpawnInfoParams[0].IntervalDt, 4 * NIAGARA_MAX_GPU_SPAWN_INFOS);
-
-	if (ComputeShader->ViewUniformBufferParam.IsBound())
-	{
-		if (ensureMsgf(ViewUniformBuffer != nullptr, TEXT("ViewUniformBuffer is required for '%s' but we do not have one to bind"), InstanceData.Context->GetDebugSimName()))
-		{
-			RHICmdList.SetShaderUniformBuffer(RHIComputeShader, ComputeShader->ViewUniformBufferParam.GetBaseIndex(), ViewUniformBuffer);
-		}
-	}
+	static_assert(sizeof(DispatchParameters->EmitterSpawnInfoOffsets) == sizeof(InstanceData.SpawnInfo.SpawnInfoStartOffsets), "Mismatch between shader parameter size and SpawnInfoStartOffsets");
+	static_assert(sizeof(DispatchParameters->EmitterSpawnInfoParams) == sizeof(InstanceData.SpawnInfo.SpawnInfoParams), "Mismatch between shader parameter size and SpawnInfoParams");
+	FMemory::Memcpy(&DispatchParameters->EmitterSpawnInfoOffsets, &InstanceData.SpawnInfo.SpawnInfoStartOffsets, sizeof(InstanceData.SpawnInfo.SpawnInfoStartOffsets));
+	FMemory::Memcpy(&DispatchParameters->EmitterSpawnInfoParams,	 &InstanceData.SpawnInfo.SpawnInfoParams, sizeof(InstanceData.SpawnInfo.SpawnInfoParams));
 
 	// Setup instance counts
-	if (ComputeShader->InstanceCountsParam.IsBound())
-	{
-		ComputeShader->InstanceCountsParam.SetBuffer(RHICmdList, RHIComputeShader, GPUInstanceCounterManager.GetInstanceCountBuffer());
-
-		if (SimStageData.AlternateIterationSource)
-		{
-			SetShaderValue(RHICmdList, RHIComputeShader, ComputeShader->ReadInstanceCountOffsetParam, -1);
-			SetShaderValue(RHICmdList, RHIComputeShader, ComputeShader->WriteInstanceCountOffsetParam, -1);
-		}
-		else
-		{
-			SetShaderValue(RHICmdList, RHIComputeShader, ComputeShader->ReadInstanceCountOffsetParam, SimStageData.SourceCountOffset);
-			SetShaderValue(RHICmdList, RHIComputeShader, ComputeShader->WriteInstanceCountOffsetParam, SimStageData.DestinationCountOffset);
-		}
-	}
+	DispatchParameters->RWInstanceCounts			= GPUInstanceCounterManager.GetInstanceCountBuffer().UAV;
+	DispatchParameters->ReadInstanceCountOffset	= SimStageData.AlternateIterationSource ? INDEX_NONE : SimStageData.SourceCountOffset;
+	DispatchParameters->WriteInstanceCountOffset	= SimStageData.AlternateIterationSource ? INDEX_NONE : SimStageData.DestinationCountOffset;
 
 	// Simulation Stage Information
 	// X = Count Buffer Instance Count Offset (INDEX_NONE == Use Instance Count)
@@ -1576,42 +1565,113 @@ void FNiagaraGpuComputeDispatch::DispatchStage(FRHICommandList& RHICmdList, FRHI
 	// Z = Iteration Index
 	// W = Num Iterations
 	{
-		FIntVector4 SimulationStageIterationInfo = { INDEX_NONE, -1, 0, 0 };
-		float SimulationStageNormalizedIterationIndex = 0.0f;
-
+		DispatchParameters->SimulationStageIterationInfo = FIntVector4(INDEX_NONE, -1, 0, 0);
+		DispatchParameters->SimulationStageNormalizedIterationIndex = 0.0f;
 		if (SimStageData.AlternateIterationSource != nullptr)
 		{
 			const uint32 IterationInstanceCountOffset = SimStageData.AlternateIterationSource->GetGPUInstanceCountOffset(Tick.SystemInstanceID);
-			SimulationStageIterationInfo.X = IterationInstanceCountOffset;
-			SimulationStageIterationInfo.Y = IterationInstanceCountOffset == INDEX_NONE ? TotalDispatchCount : 0;
+			DispatchParameters->SimulationStageIterationInfo.X = IterationInstanceCountOffset;
+			DispatchParameters->SimulationStageIterationInfo.Y = IterationInstanceCountOffset == INDEX_NONE ? TotalDispatchCount : 0;
 		}
 
 		const int32 NumIterations = InstanceData.PerStageInfo[SimStageData.StageIndex].NumIterations;
 		const int32 IterationIndex = SimStageData.IterationIndex;
-		SimulationStageIterationInfo.Z = IterationIndex;
-		SimulationStageIterationInfo.W = NumIterations;
-		SimulationStageNormalizedIterationIndex = NumIterations > 1 ? float(IterationIndex) / float(NumIterations - 1) : 1.0f;
-
-		SetShaderValue(RHICmdList, RHIComputeShader, ComputeShader->SimulationStageIterationInfoParam, SimulationStageIterationInfo);
-		SetShaderValue(RHICmdList, RHIComputeShader, ComputeShader->SimulationStageNormalizedIterationIndexParam, SimulationStageNormalizedIterationIndex);
+		DispatchParameters->SimulationStageIterationInfo.Z = IterationIndex;
+		DispatchParameters->SimulationStageIterationInfo.W = NumIterations;
+		DispatchParameters->SimulationStageNormalizedIterationIndex = NumIterations > 1 ? float(IterationIndex) / float(NumIterations - 1) : 1.0f;
 	}
 
 	// Set particle iteration state info
 	// Where X = Parameter Binding, YZ = Inclusive Range
+	DispatchParameters->ParticleIterationStateInfo.X = SimStageData.StageMetaData->ParticleIterationStateComponentIndex;
+	DispatchParameters->ParticleIterationStateInfo.Y = SimStageData.StageMetaData->ParticleIterationStateRange.X;
+	DispatchParameters->ParticleIterationStateInfo.Z = SimStageData.StageMetaData->ParticleIterationStateRange.Y;
+
+	// Set static input buffers
+	DispatchParameters->StaticInputFloat = Tick.SystemGpuComputeProxy->StaticFloatBuffer;
+
+	// Set Particle Input Buffer
+	if ( (SimStageData.Source != nullptr) && (SimStageData.Source->GetNumInstancesAllocated() > 0) )
+	{
+		DispatchParameters->InputFloat = SimStageData.Source->GetGPUBufferFloat().SRV;
+		DispatchParameters->InputHalf = SimStageData.Source->GetGPUBufferHalf().SRV;
+		DispatchParameters->InputInt = SimStageData.Source->GetGPUBufferInt().SRV;
+		DispatchParameters->ComponentBufferSizeRead = SimStageData.Source->GetFloatStride() / sizeof(float);
+	}
+	else
 	{
-		int ParticleIterationStateInfo[3] = { SimStageData.StageMetaData->ParticleIterationStateComponentIndex, SimStageData.StageMetaData->ParticleIterationStateRange.X, SimStageData.StageMetaData->ParticleIterationStateRange.Y };
-		SetShaderValue(RHICmdList, RHIComputeShader, ComputeShader->ParticleIterationStateInfoParam, ParticleIterationStateInfo);
+		DispatchParameters->InputFloat = FNiagaraRenderer::GetDummyFloatBuffer();
+		DispatchParameters->InputHalf = FNiagaraRenderer::GetDummyHalfBuffer();
+		DispatchParameters->InputInt = FNiagaraRenderer::GetDummyIntBuffer();
+		DispatchParameters->ComponentBufferSizeRead = 0;
 	}
 
-	// Set static input buffers
-	SetSRVParameter(RHICmdList, RHIComputeShader, ComputeShader->StaticInputFloatParam, Tick.SystemGpuComputeProxy->StaticFloatBuffer);
+	// Set Particle Output Buffer
+	if (SimStageData.Destination != nullptr)
+	{
+		DispatchParameters->RWOutputFloat = SimStageData.Destination->GetGPUBufferFloat().UAV;
+		DispatchParameters->RWOutputHalf = SimStageData.Destination->GetGPUBufferHalf().UAV;
+		DispatchParameters->RWOutputInt = SimStageData.Destination->GetGPUBufferInt().UAV;
+		DispatchParameters->RWIDToIndexTable = SimStageData.Destination->GetGPUIDToIndexTable().UAV;
+		DispatchParameters->ComponentBufferSizeWrite = SimStageData.Destination->GetFloatStride() / sizeof(float);
+	}
+	else
+	{
+		DispatchParameters->RWOutputFloat = nullptr;
+		DispatchParameters->RWOutputHalf = nullptr;
+		DispatchParameters->RWOutputInt = nullptr;
+		DispatchParameters->RWIDToIndexTable = nullptr;
+		DispatchParameters->ComponentBufferSizeWrite = 0;
+	}
 
-	// Set Input & Output buffers
-	FNiagaraDataBuffer::SetInputShaderParams(RHICmdList, ComputeShader.GetShader(), SimStageData.Source);
-	FNiagaraDataBuffer::SetOutputShaderParams(RHICmdList, ComputeShader.GetShader(), SimStageData.Destination);
+	// Set Compute Shader
+	const FNiagaraShaderRef ComputeShader = InstanceData.Context->GPUScript_RT->GetShader(SimStageData.StageIndex);
+	FRHIComputeShader* RHIComputeShader = ComputeShader.GetComputeShader();
+	SetComputePipelineState(RHICmdList, RHIComputeShader);
 
 	// Set data interface parameters
-	SetDataInterfaceParameters(RHICmdList, Tick, InstanceData, ComputeShader, SimStageData);
+	SetDataInterfaceParameters(RHICmdList, Tick, InstanceData, ComputeShader, SimStageData, reinterpret_cast<uint8*>(DispatchParameters));
+
+	// Setup uniform buffers
+	//-TODO: This should be replaced with shader struct bindings instead but we currently have no way to include a untyped uniform buffer
+	auto SetConstantBuffer =
+		[](FRHICommandList& RHICmdList, FRHIComputeShader* ComputeShader, const FShaderUniformBufferParameter& BufferParam, const FUniformBufferRHIRef& UniformBuffer)
+		{
+			if (BufferParam.IsBound() && UniformBuffer.IsValid())
+			{
+				RHICmdList.SetShaderUniformBuffer(ComputeShader, BufferParam.GetBaseIndex(), UniformBuffer);
+			}
+		};
+
+	SetConstantBuffer(RHICmdList, RHIComputeShader, ComputeShader->GlobalConstantBufferParam[0], Tick.GetUniformBuffer(FNiagaraGPUSystemTick::UBT_Global, nullptr, true));
+	SetConstantBuffer(RHICmdList, RHIComputeShader, ComputeShader->SystemConstantBufferParam[0], Tick.GetUniformBuffer(FNiagaraGPUSystemTick::UBT_System, nullptr, true));
+	SetConstantBuffer(RHICmdList, RHIComputeShader, ComputeShader->OwnerConstantBufferParam[0], Tick.GetUniformBuffer(FNiagaraGPUSystemTick::UBT_Owner, nullptr, true));
+	SetConstantBuffer(RHICmdList, RHIComputeShader, ComputeShader->EmitterConstantBufferParam[0], Tick.GetUniformBuffer(FNiagaraGPUSystemTick::UBT_Emitter, &InstanceData, true));
+	SetConstantBuffer(RHICmdList, RHIComputeShader, ComputeShader->ExternalConstantBufferParam[0], Tick.GetUniformBuffer(FNiagaraGPUSystemTick::UBT_External, &InstanceData, true));
+
+	if (InstanceData.Context->HasInterpolationParameters)
+	{
+		SetConstantBuffer(RHICmdList, RHIComputeShader, ComputeShader->GlobalConstantBufferParam[1], Tick.GetUniformBuffer(FNiagaraGPUSystemTick::UBT_Global, nullptr, false));
+		SetConstantBuffer(RHICmdList, RHIComputeShader, ComputeShader->SystemConstantBufferParam[1], Tick.GetUniformBuffer(FNiagaraGPUSystemTick::UBT_System, nullptr, false));
+		SetConstantBuffer(RHICmdList, RHIComputeShader, ComputeShader->OwnerConstantBufferParam[1], Tick.GetUniformBuffer(FNiagaraGPUSystemTick::UBT_Owner, nullptr, false));
+		SetConstantBuffer(RHICmdList, RHIComputeShader, ComputeShader->EmitterConstantBufferParam[1], Tick.GetUniformBuffer(FNiagaraGPUSystemTick::UBT_Emitter, &InstanceData, false));
+		SetConstantBuffer(RHICmdList, RHIComputeShader, ComputeShader->ExternalConstantBufferParam[1], Tick.GetUniformBuffer(FNiagaraGPUSystemTick::UBT_External, &InstanceData, false));
+	}
+
+	//-TODO: Currently we don't have a way to bind raw UBs as shader parameters
+	//DispatchParameters.GlobalParameters			= Tick.GetUniformBuffer(FNiagaraGPUSystemTick::UBT_Global, nullptr, true);
+	//DispatchParameters.OwnerParameters			= Tick.GetUniformBuffer(FNiagaraGPUSystemTick::UBT_Owner, nullptr, true);
+	//DispatchParameters.SystemParameters			= Tick.GetUniformBuffer(FNiagaraGPUSystemTick::UBT_System, nullptr, true);
+	//DispatchParameters.EmitterParameters		= Tick.GetUniformBuffer(FNiagaraGPUSystemTick::UBT_Emitter, nullptr, true);
+	//DispatchParameters.ExternalParameters		= Tick.GetUniformBuffer(FNiagaraGPUSystemTick::UBT_External, nullptr, true);
+	//if (InstanceData.Context->HasInterpolationParameters)
+	//{
+	//	DispatchParameters.PrevGlobalParameters		= Tick.GetUniformBuffer(FNiagaraGPUSystemTick::UBT_Global, nullptr, false);
+	//	DispatchParameters.PrevOwnerParameters		= Tick.GetUniformBuffer(FNiagaraGPUSystemTick::UBT_Owner, nullptr, false);
+	//	DispatchParameters.PrevSystemParameters		= Tick.GetUniformBuffer(FNiagaraGPUSystemTick::UBT_System, nullptr, false);
+	//	DispatchParameters.PrevEmitterParameters	= Tick.GetUniformBuffer(FNiagaraGPUSystemTick::UBT_Emitter, nullptr, false);
+	//	DispatchParameters.PrevExternalParameters	= Tick.GetUniformBuffer(FNiagaraGPUSystemTick::UBT_External, nullptr, false);
+	//}
 
 	// Execute the dispatch
 	{
@@ -1628,20 +1688,18 @@ void FNiagaraGpuComputeDispatch::DispatchStage(FRHICommandList& RHICmdList, FRHI
 		ThreadGroupCount.Y = FMath::DivideAndRoundUp(DispatchCount.Y, DispatchNumThreads.Y);
 		ThreadGroupCount.Z = FMath::DivideAndRoundUp(DispatchCount.Z, DispatchNumThreads.Z);
 
-		SetShaderValue(RHICmdList, RHIComputeShader, ComputeShader->DispatchThreadIdToLinearParam, FIntVector(1, DispatchCount.X, DispatchCount.X * DispatchCount.Y));
-		SetShaderValue(RHICmdList, RHIComputeShader, ComputeShader->DispatchThreadIdBoundsParam, DispatchCount);
-
-		SetConstantBuffers(RHICmdList, ComputeShader, Tick, &InstanceData);
+		DispatchParameters->DispatchThreadIdToLinear	= FUintVector3(1, DispatchCount.X, DispatchCount.X * DispatchCount.Y);
+		DispatchParameters->DispatchThreadIdBounds		= FUintVector3(DispatchCount.X, DispatchCount.Y, DispatchCount.Z);
 
+		SetShaderParameters<FRHICommandList, FNiagaraShader>(RHICmdList, ComputeShader, RHIComputeShader, &ShaderParametersMetadata, *DispatchParameters);
 		DispatchComputeShader(RHICmdList, ComputeShader, ThreadGroupCount.X, ThreadGroupCount.Y, ThreadGroupCount.Z);
+		UnsetShaderUAVs<FRHICommandList, FNiagaraShader>(RHICmdList, ComputeShader, RHIComputeShader);
 
 		INC_DWORD_STAT(STAT_NiagaraGPUDispatches);
 	}
 
 	// Unset UAV parameters
 	UnsetDataInterfaceParameters(RHICmdList, Tick, InstanceData, ComputeShader, SimStageData);
-	FNiagaraDataBuffer::UnsetShaderParams(RHICmdList, ComputeShader.GetShader());
-	ComputeShader->InstanceCountsParam.UnsetUAV(RHICmdList, RHIComputeShader);
 
 	// Optionally submit commands to the GPU
 	// This can be used to avoid accidental TDR detection in the editor especially when issuing multiple ticks in the same frame
@@ -1713,8 +1771,8 @@ void FNiagaraGpuComputeDispatch::PreInitViews(FRDGBuilder& GraphBuilder, bool bA
 				DumpDebugFrame();
 			}
 
-			FNiagaraSceneTextureParameters* PassParameters = GraphBuilder.AllocParameters<FNiagaraSceneTextureParameters>();
-			GNiagaraViewDataManager.GetSceneTextureParameters(GraphBuilder, nullptr, *PassParameters);
+			FNiagaraComputePassParameters* PassParameters = GraphBuilder.AllocParameters<FNiagaraComputePassParameters>();
+			GNiagaraViewDataManager.GetSceneTextureParameters(GraphBuilder, nullptr, PassParameters->SceneTextures);
 
 			GraphBuilder.AddPass(
 				RDG_EVENT_NAME("Niagara::PreInitViews"),
@@ -1722,10 +1780,10 @@ void FNiagaraGpuComputeDispatch::PreInitViews(FRDGBuilder& GraphBuilder, bool bA
 				ERDGPassFlags::Compute | ERDGPassFlags::NeverCull,
 				[this, PassParameters](FRHICommandListImmediate& RHICmdList)
 			{
-				NiagaraSceneTextures = PassParameters;
+				NiagaraSceneTextures = &PassParameters->SceneTextures;
 				ON_SCOPE_EXIT { NiagaraSceneTextures = nullptr; };
 
-				ExecuteTicks(RHICmdList, nullptr, ENiagaraGpuComputeTickStage::PreInitViews);
+				ExecuteTicks(RHICmdList, ENiagaraGpuComputeTickStage::PreInitViews);
 			});
 		}
 	}
@@ -1744,8 +1802,9 @@ void FNiagaraGpuComputeDispatch::PostInitViews(FRDGBuilder& GraphBuilder, TArray
 
 	if (bAllowGPUParticleUpdate && FNiagaraUtilities::AllowGPUParticles(GetShaderPlatform()))
 	{
-		FNiagaraSceneTextureParameters* PassParameters = GraphBuilder.AllocParameters<FNiagaraSceneTextureParameters>();
-		GNiagaraViewDataManager.GetSceneTextureParameters(GraphBuilder, GetViewFamily(Views).GetSceneTexturesChecked(), *PassParameters);
+		FNiagaraComputePassParameters* PassParameters = GraphBuilder.AllocParameters<FNiagaraComputePassParameters>();
+		GNiagaraViewDataManager.GetSceneTextureParameters(GraphBuilder, GetViewFamily(Views).GetSceneTexturesChecked(), PassParameters->SceneTextures);
+		PassParameters->View = GetReferenceViewUniformBuffer(Views);
 
 		GraphBuilder.AddPass(
 			RDG_EVENT_NAME("Niagara::PostInitViews"),
@@ -1753,11 +1812,10 @@ void FNiagaraGpuComputeDispatch::PostInitViews(FRDGBuilder& GraphBuilder, TArray
 			ERDGPassFlags::Compute | ERDGPassFlags::NeverCull,
 			[this, PassParameters, Views](FRHICommandListImmediate& RHICmdList)
 			{
-				NiagaraSceneTextures = PassParameters;
+				NiagaraSceneTextures = &PassParameters->SceneTextures;
 				ON_SCOPE_EXIT{ NiagaraSceneTextures = nullptr; };
 
-				FRHIUniformBuffer* ViewUniformBuffer = GetReferenceViewUniformBuffer(Views);
-				ExecuteTicks(RHICmdList, ViewUniformBuffer, ENiagaraGpuComputeTickStage::PostInitViews);
+				ExecuteTicks(RHICmdList, ENiagaraGpuComputeTickStage::PostInitViews);
 			#if WITH_MGPU
 				TransferMultiGPUBufers(RHICmdList, ENiagaraGpuComputeTickStage::PostInitViews);
 			#endif // WITH_MGPU
@@ -1779,10 +1837,11 @@ void FNiagaraGpuComputeDispatch::PostRenderOpaque(FRDGBuilder& GraphBuilder, TCo
 
 	if (bAllowGPUParticleUpdate && FNiagaraUtilities::AllowGPUParticles(GetShaderPlatform()))
 	{
-		FNiagaraSceneTextureParameters* PassParameters = GraphBuilder.AllocParameters<FNiagaraSceneTextureParameters>();
+		FNiagaraComputePassParameters* PassParameters = GraphBuilder.AllocParameters<FNiagaraComputePassParameters>();
 		// TODO: This will cause a fragment->compute barrier on a scene textures which could be costly especially on mobile GPUs
 		// Will be nice to avoid executing this if we know that there are no simulations that require access to a scene textures
-		GNiagaraViewDataManager.GetSceneTextureParameters(GraphBuilder, GetViewFamily(Views).GetSceneTexturesChecked(), *PassParameters);
+		GNiagaraViewDataManager.GetSceneTextureParameters(GraphBuilder, GetViewFamily(Views).GetSceneTexturesChecked(), PassParameters->SceneTextures);
+		PassParameters->View = GetReferenceViewUniformBuffer(Views);
 
  		GraphBuilder.AddPass(
 			RDG_EVENT_NAME("Niagara::PostRenderOpaque"),
@@ -1790,17 +1849,15 @@ void FNiagaraGpuComputeDispatch::PostRenderOpaque(FRDGBuilder& GraphBuilder, TCo
 			ERDGPassFlags::Compute | ERDGPassFlags::NeverCull,
 			[this, PassParameters, Views](FRHICommandListImmediate& RHICmdList)
 		{
-			NiagaraSceneTextures = PassParameters;
+			NiagaraSceneTextures = &PassParameters->SceneTextures;
 			ON_SCOPE_EXIT{ NiagaraSceneTextures = nullptr; };
 
 			AsyncGpuTraceHelper->PostRenderOpaque(RHICmdList, this, Views);
 
 			CurrentPassViews = Views;
 
-			FRHIUniformBuffer* ViewUniformBuffer = GetReferenceViewUniformBuffer(Views);
-
 			// Setup new readback since if there is no pending request, there is no risk of having invalid data read (offset being allocated after the readback was sent).
-			ExecuteTicks(RHICmdList, ViewUniformBuffer, ENiagaraGpuComputeTickStage::PostOpaqueRender);
+			ExecuteTicks(RHICmdList, ENiagaraGpuComputeTickStage::PostOpaqueRender);
 
 			FinishDispatches();
 
@@ -2069,23 +2126,35 @@ FNiagaraAsyncGpuTraceHelper& FNiagaraGpuComputeDispatch::GetAsyncGpuTraceHelper(
 
 /* Set shader parameters for data interfaces
  */
-void FNiagaraGpuComputeDispatch::SetDataInterfaceParameters(FRHICommandList& RHICmdList, const FNiagaraGPUSystemTick& Tick, const FNiagaraComputeInstanceData& InstanceData, const FNiagaraShaderRef& ComputeShader, const FNiagaraSimStageData& SimStageData) const
+void FNiagaraGpuComputeDispatch::SetDataInterfaceParameters(FRHICommandList& RHICmdList, const FNiagaraGPUSystemTick& Tick, const FNiagaraComputeInstanceData& InstanceData, const FNiagaraShaderRef& ComputeShader, const FNiagaraSimStageData& SimStageData, uint8* ParametersStructure) const
 {
+	const int32 NumDataInterfaces = InstanceData.DataInterfaceProxies.Num();
+	if (NumDataInterfaces == 0)
+	{
+		return;
+	}
+
 	// @todo-threadsafety This is a bit gross. Need to rethink this api.
 	const FNiagaraShaderMapPointerTable& PointerTable = ComputeShader.GetPointerTable();
+	const TMemoryImageArray<FNiagaraDataInterfaceParamRef>& DIParameters = ComputeShader->GetDIParameters();
 
-	uint32 InterfaceIndex = 0;
-	for (FNiagaraDataInterfaceProxy* Interface : InstanceData.DataInterfaceProxies)
+	FNiagaraDataInterfaceSetShaderParametersContext SetParametersContext(RHICmdList , *this, Tick, InstanceData, SimStageData, ComputeShader, ParametersStructure);
+
+	for ( int32 iDataInterface=0; iDataInterface < NumDataInterfaces; ++iDataInterface )
 	{
-		const TMemoryImageArray<FNiagaraDataInterfaceParamRef>& DIParameters = ComputeShader->GetDIParameters();
-		const FNiagaraDataInterfaceParamRef& DIParam = DIParameters[InterfaceIndex];
-		if (DIParam.Parameters.IsValid())
+		FNiagaraDataInterfaceProxy* DataInterfaceProxy = InstanceData.DataInterfaceProxies[iDataInterface];
+
+		const FNiagaraDataInterfaceParamRef& DIParam = DIParameters[iDataInterface];
+		if (DIParam.ShaderParametersOffset != INDEX_NONE)
 		{
-			FNiagaraDataInterfaceSetArgs Context(Interface, Tick.SystemInstanceID, Tick.SystemGpuComputeProxy->GetSystemLWCTile(), this, ComputeShader, &InstanceData, &SimStageData, InstanceData.IsOutputStage(Interface, SimStageData.StageIndex), InstanceData.IsIterationStage(Interface, SimStageData.StageIndex));
+			SetParametersContext.SetDataInterface(DataInterfaceProxy, DIParam.ShaderParametersOffset, DIParam.Parameters);
+			CastChecked<UNiagaraDataInterface>(DIParam.DIType.Get(PointerTable.DITypes))->SetShaderParameters(SetParametersContext);
+		}
+		else if (DIParam.Parameters.IsValid())
+		{
+			FNiagaraDataInterfaceSetArgs Context(DataInterfaceProxy, Tick.SystemInstanceID, Tick.SystemGpuComputeProxy->GetSystemLWCTile(), this, ComputeShader, &InstanceData, &SimStageData, InstanceData.IsOutputStage(DataInterfaceProxy, SimStageData.StageIndex), InstanceData.IsIterationStage(DataInterfaceProxy, SimStageData.StageIndex));
 			DIParam.DIType.Get(PointerTable.DITypes)->SetParameters(DIParam.Parameters.Get(), RHICmdList, Context);
 		}
-
-		InterfaceIndex++;
 	}
 }
 
@@ -2099,7 +2168,7 @@ void FNiagaraGpuComputeDispatch::UnsetDataInterfaceParameters(FRHICommandList& R
 	{
 		const TMemoryImageArray<FNiagaraDataInterfaceParamRef>& DIParameters = ComputeShader->GetDIParameters();
 		const FNiagaraDataInterfaceParamRef& DIParam = DIParameters[InterfaceIndex];
-		if (DIParam.Parameters.IsValid())
+		if (DIParam.ShaderParametersOffset != INDEX_NONE && DIParam.Parameters.IsValid())
 		{
 			FNiagaraDataInterfaceSetArgs Context(Interface, Tick.SystemInstanceID, Tick.SystemGpuComputeProxy->GetSystemLWCTile(), this, ComputeShader, &InstanceData, &SimStageData, InstanceData.IsOutputStage(Interface, SimStageData.StageIndex), InstanceData.IsIterationStage(Interface, SimStageData.StageIndex));
 			DIParam.DIType.Get(PointerTable.DITypes)->UnsetParameters(DIParam.Parameters.Get(), RHICmdList, Context);
@@ -2109,35 +2178,6 @@ void FNiagaraGpuComputeDispatch::UnsetDataInterfaceParameters(FRHICommandList& R
 	}
 }
 
-static void SetConstantBuffer(FRHICommandList &RHICmdList, FRHIComputeShader* ComputeShader, const FShaderUniformBufferParameter& BufferParam, const FUniformBufferRHIRef& UniformBuffer)
-{
-	if (BufferParam.IsBound() && UniformBuffer.IsValid())
-	{
-		RHICmdList.SetShaderUniformBuffer(ComputeShader, BufferParam.GetBaseIndex(), UniformBuffer);
-	}
-}
-
-void FNiagaraGpuComputeDispatch::SetConstantBuffers(FRHICommandList &RHICmdList, const FNiagaraShaderRef& Shader, const FNiagaraGPUSystemTick& Tick, const FNiagaraComputeInstanceData* Instance)
-{
-	FNiagaraComputeExecutionContext* Context = Instance->Context;
-	FRHIComputeShader* ComputeShader = Shader.GetComputeShader();
-
-	SetConstantBuffer(RHICmdList, ComputeShader, Shader->GlobalConstantBufferParam[0], Tick.GetUniformBuffer(FNiagaraGPUSystemTick::UBT_Global, nullptr, true));
-	SetConstantBuffer(RHICmdList, ComputeShader, Shader->SystemConstantBufferParam[0], Tick.GetUniformBuffer(FNiagaraGPUSystemTick::UBT_System, nullptr, true));
-	SetConstantBuffer(RHICmdList, ComputeShader, Shader->OwnerConstantBufferParam[0], Tick.GetUniformBuffer(FNiagaraGPUSystemTick::UBT_Owner, nullptr, true));
-	SetConstantBuffer(RHICmdList, ComputeShader, Shader->EmitterConstantBufferParam[0], Tick.GetUniformBuffer(FNiagaraGPUSystemTick::UBT_Emitter, Instance, true));
-	SetConstantBuffer(RHICmdList, ComputeShader, Shader->ExternalConstantBufferParam[0], Tick.GetUniformBuffer(FNiagaraGPUSystemTick::UBT_External, Instance, true));
-
-	if (Context->HasInterpolationParameters)
-	{
-		SetConstantBuffer(RHICmdList, ComputeShader, Shader->GlobalConstantBufferParam[1], Tick.GetUniformBuffer(FNiagaraGPUSystemTick::UBT_Global, nullptr, false));
-		SetConstantBuffer(RHICmdList, ComputeShader, Shader->SystemConstantBufferParam[1], Tick.GetUniformBuffer(FNiagaraGPUSystemTick::UBT_System, nullptr, false));
-		SetConstantBuffer(RHICmdList, ComputeShader, Shader->OwnerConstantBufferParam[1], Tick.GetUniformBuffer(FNiagaraGPUSystemTick::UBT_Owner, nullptr, false));
-		SetConstantBuffer(RHICmdList, ComputeShader, Shader->EmitterConstantBufferParam[1], Tick.GetUniformBuffer(FNiagaraGPUSystemTick::UBT_Emitter, Instance, false));
-		SetConstantBuffer(RHICmdList, ComputeShader, Shader->ExternalConstantBufferParam[1], Tick.GetUniformBuffer(FNiagaraGPUSystemTick::UBT_External, Instance, false));
-	}
-}
-
 FGPUSortManager* FNiagaraGpuComputeDispatch::GetGPUSortManager() const
 {
 	return GPUSortManager;
diff --git a/Engine/Plugins/FX/Niagara/Source/Niagara/Private/NiagaraGpuComputeDispatch.h b/Engine/Plugins/FX/Niagara/Source/Niagara/Private/NiagaraGpuComputeDispatch.h
--- a/Engine/Plugins/FX/Niagara/Source/Niagara/Private/NiagaraGpuComputeDispatch.h
+++ b/Engine/Plugins/FX/Niagara/Source/Niagara/Private/NiagaraGpuComputeDispatch.h
@@ -94,7 +94,7 @@ public:
 	const FGlobalDistanceFieldParameterData* GetGlobalDistanceFieldParameters() const;
 	const FDistanceFieldSceneData* GetMeshDistanceFieldParameters() const;
 
-	void SetDataInterfaceParameters(FRHICommandList& RHICmdList, const FNiagaraGPUSystemTick& Tick, const FNiagaraComputeInstanceData& InstanceData, const FNiagaraShaderRef& ComputeShader, const FNiagaraSimStageData& SimStageData) const;
+	void SetDataInterfaceParameters(FRHICommandList& RHICmdList, const FNiagaraGPUSystemTick& Tick, const FNiagaraComputeInstanceData& InstanceData, const FNiagaraShaderRef& ComputeShader, const FNiagaraSimStageData& SimStageData, uint8* ParametersStructure) const;
 	void UnsetDataInterfaceParameters(FRHICommandList& RHICmdList, const FNiagaraGPUSystemTick& Tick, const FNiagaraComputeInstanceData& InstanceData, const FNiagaraShaderRef& ComputeShader, const FNiagaraSimStageData& SimStageData) const;
 
 	/** Reset the data interfaces and check if the spawn stages are valid */
@@ -137,8 +137,8 @@ private:
 	void UpdateInstanceCountManager(FRHICommandListImmediate& RHICmdList);
 	void PrepareTicksForProxy(FRHICommandListImmediate& RHICmdList, FNiagaraSystemGpuComputeProxy* ComputeProxy, FNiagaraGpuDispatchList& GpuDispatchList);
 	void PrepareAllTicks(FRHICommandListImmediate& RHICmdList);
-	void ExecuteTicks(FRHICommandList& RHICmdList, FRHIUniformBuffer* ViewUniformBuffer, ENiagaraGpuComputeTickStage::Type TickStage);
-	void DispatchStage(FRHICommandList& RHICmdList, FRHIUniformBuffer* ViewUniformBuffer, const FNiagaraGPUSystemTick& Tick, const FNiagaraComputeInstanceData& InstanceData, const FNiagaraSimStageData& SimStageData);
+	void ExecuteTicks(FRHICommandList& RHICmdList, ENiagaraGpuComputeTickStage::Type TickStage);
+	void DispatchStage(FRHICommandList& RHICmdList, const FNiagaraGPUSystemTick& Tick, const FNiagaraComputeInstanceData& InstanceData, const FNiagaraSimStageData& SimStageData);
 
 	/**
 	 * Generate all the initial keys and values for a GPUSortManager sort batch.
@@ -166,7 +166,6 @@ private:
 	void UpdateFreeIDsListSizesBuffer(FRHICommandList& RHICmdList, uint32 NumInstances);
 	void UpdateFreeIDBuffers(FRHICommandList& RHICmdList, TConstArrayView<FNiagaraComputeExecutionContext*> Instances);
 
-	void SetConstantBuffers(FRHICommandList &RHICmdList, const FNiagaraShaderRef& Shader, const FNiagaraGPUSystemTick& Tick, const FNiagaraComputeInstanceData* Instance);
 	void BuildConstantBuffers(FNiagaraGPUSystemTick& Tick) const;
 
 	/** The shared GPUSortManager, used to register GPU sort tasks in order to generate sorted particle indices per emitter. */
diff --git a/Engine/Plugins/FX/Niagara/Source/Niagara/Private/NiagaraScript.cpp b/Engine/Plugins/FX/Niagara/Source/Niagara/Private/NiagaraScript.cpp
--- a/Engine/Plugins/FX/Niagara/Source/Niagara/Private/NiagaraScript.cpp
+++ b/Engine/Plugins/FX/Niagara/Source/Niagara/Private/NiagaraScript.cpp
@@ -3080,7 +3080,7 @@ void UNiagaraScript::CacheResourceShadersForRendering(bool bRegenerateId, bool b
 				CachedScriptVMId.BaseScriptCompileHash, CachedScriptVMId.ReferencedCompileHashes,
 				CachedScriptVMId.bUsesRapidIterationParams, GetFriendlyName());
 
-			ScriptResource->SetDataInterfaceParamInfo(CachedScriptVM.DIParamInfo);
+			ScriptResource->BuildScriptParametersMetadata(CachedScriptVM.DIParamInfo);
 
 			if (FNiagaraUtilities::SupportsComputeShaders(ShaderPlatform))
 			{
@@ -3512,7 +3512,7 @@ void UNiagaraScript::ProcessSerializedShaderMaps()
 
 	if (HasScriptResource)
 	{
-		ScriptResource->SetDataInterfaceParamInfo(CachedScriptVM.DIParamInfo);
+		ScriptResource->BuildScriptParametersMetadata(CachedScriptVM.DIParamInfo);
 	}
 }
 
diff --git a/Engine/Plugins/FX/Niagara/Source/Niagara/Public/NiagaraTypes.h b/Engine/Plugins/FX/Niagara/Source/Niagara/Public/NiagaraTypes.h
--- a/Engine/Plugins/FX/Niagara/Source/Niagara/Public/NiagaraTypes.h
+++ b/Engine/Plugins/FX/Niagara/Source/Niagara/Public/NiagaraTypes.h
@@ -650,6 +650,13 @@ public:
 		return true;
 	}
 
+	template<typename T>
+	bool UpdateShaderParameters()
+	{
+		const FShaderParametersMetadata* ShaderParametersMetadata = TShaderParameterStructTypeInfo<T>::GetStructMetadata();
+		return UpdatePOD(ShaderParametersMetadata->GetStructTypeName(), ShaderParametersMetadata->GetLayoutHash());
+	}
+
 	/**
 	Adds an string value to the hash.
 	*/
diff --git a/Engine/Plugins/FX/Niagara/Source/NiagaraCore/Public/NiagaraDataInterfaceBase.h b/Engine/Plugins/FX/Niagara/Source/NiagaraCore/Public/NiagaraDataInterfaceBase.h
--- a/Engine/Plugins/FX/Niagara/Source/NiagaraCore/Public/NiagaraDataInterfaceBase.h
+++ b/Engine/Plugins/FX/Niagara/Source/NiagaraCore/Public/NiagaraDataInterfaceBase.h
@@ -11,10 +11,13 @@ class INiagaraCompiler;
 class UCurveVector;
 class UCurveLinearColor;
 class UCurveFloat;
+class RENDERCORE_API FShaderParametersMetadataBuilder;
+
 class UNiagaraEmitter;
 class FNiagaraSystemInstance;
 class FNiagaraShader;
 class FNiagaraShaderMapPointerTable;
+class FNiagaraShaderParametersBuilder;
 class UNiagaraDataInterfaceBase;
 class FNiagaraGpuComputeDispatchInterface;
 struct FNiagaraDataInterfaceGPUParamInfo;
@@ -84,7 +87,7 @@ struct FNiagaraDataInterfaceSetArgs : public FNiagaraDataInterfaceArgs
 /**
  * An interface to the parameter bindings for the data interface used by a Niagara compute shader.
  * This is not using virtual methods, but derived classes may still override the methods listed below.
- * Overriden methods will be correctly called via indirection through UNiagaraDataInterfaceBase's vtable
+ * Overridden methods will be correctly called via indirection through UNiagaraDataInterfaceBase's vtable
  */
 struct FNiagaraDataInterfaceParametersCS
 {
@@ -93,8 +96,6 @@ public:
 	void Bind(const FNiagaraDataInterfaceGPUParamInfo& ParameterInfo, const class FShaderParameterMap& ParameterMap) {}
 	void Set(FRHICommandList& RHICmdList, const FNiagaraDataInterfaceSetArgs& Context) const {}
 	void Unset(FRHICommandList& RHICmdList, const FNiagaraDataInterfaceSetArgs& Context) const {}
-
-	LAYOUT_FIELD(TIndexedPtr<UNiagaraDataInterfaceBase>, DIType);
 };
 
 //////////////////////////////////////////////////////////////////////////
@@ -106,14 +107,53 @@ class NIAGARACORE_API UNiagaraDataInterfaceBase : public UNiagaraMergeable
 	GENERATED_UCLASS_BODY()
 
 public:
-	
+	/** Does this data interface require legacy shader bindings. */
+	//UE_DEPRECATED(5.1, "UseLegacyShaderBindings is deprecated and will be removed.")
+	virtual bool UseLegacyShaderBindings() const { return true; }
+
+	/**
+	Override this method to provide parameters to the GPU (SRV / UAV / Constants / etc)
+	The most common usage will be to provide a single structure which is nested with other parameters.
+	The structure would be declared like this
+		BEGIN_SHADER_PARAMETER_STRUCT(FMyShaderParameters, )
+			SHADER_PARAMETER(float, MyValue)
+			SHADER_PARAMETER_SRV(Buffer<float>, MySRV)
+		END_SHADER_PARAMETER_STRUCT()
+
+	And inside the override BuildShaderParameters you would simple add the structure
+		ShaderParametersBuilder.AddNestedStruct<FMyShaderParameters>();
+
+	You should also ensure that AppendCompileHash includes the parameters like this
+		InVisitor->UpdateShaderParameters<FMyShaderParameters>();
+
+	When filling the data inside the set function you would ask for parameters using the same structure.
+	Note: This function is only called on the CDO only not the instance
+	*/
+	virtual void BuildShaderParameters(FNiagaraShaderParametersBuilder& ShaderParametersBuilder) const {}
+
+	/**
+	Optional storage that can be created per shader when not in legacy binding mode.
+	This can be used to store information about the compilation state to avoid doing runtime checks / look-ups.
+	*/
+	virtual FNiagaraDataInterfaceParametersCS* CreateShaderStorage(const FNiagaraDataInterfaceGPUParamInfo& ParameterInfo, const FShaderParameterMap& ParameterMap) const { return nullptr; }
+	/**
+	If we create shader storage we must also supply the type of the storage so we understand how to serialize it.
+	This function effectively returns &StaticGetTypeLayoutDesc<ParameterType>().
+	*/
+	virtual const FTypeLayoutDesc* GetShaderStorageType() const { return nullptr; }
+
 	/** Constructs the correct CS parameter type for this DI (if any). The object type returned by this can only vary by class and not per object data. */
+	//UE_DEPRECATED(5.1, "Is deprecated please use CreateShaderStorage instead.")
 	virtual FNiagaraDataInterfaceParametersCS* CreateComputeParameters() const { return nullptr; }
+	//UE_DEPRECATED(5.1, "Is deprecated please use GetShaderStorageType instead.")
 	virtual const FTypeLayoutDesc* GetComputeParametersTypeDesc() const { return nullptr; }
 
 	/** Methods that operate on an instance of type FNiagaraDataInterfaceParametersCS*, created by the above CreateComputeParameters() method */
+	//UE_DEPRECATED(5.1, "Is deprecated please convert to use BuildShaderParameters.")
 	virtual void BindParameters(FNiagaraDataInterfaceParametersCS* Base, const FNiagaraDataInterfaceGPUParamInfo& ParameterInfo, const class FShaderParameterMap& ParameterMap) {}
+	//UE_DEPRECATED(5.1, "Is deprecated please convert to use UNiagaraDataInterface::SetShaderParameters.")
 	virtual void SetParameters(const FNiagaraDataInterfaceParametersCS* Base, FRHICommandList& RHICmdList, const FNiagaraDataInterfaceSetArgs& Context) const {}
+	//UE_DEPRECATED(5.1, "Is deprecated and will be removed as it is no longer required.")
 	virtual void UnsetParameters(const FNiagaraDataInterfaceParametersCS* Base, FRHICommandList& RHICmdList, const FNiagaraDataInterfaceSetArgs& Context) const {}
 
 	/** Returns true if the DI (owned by OwnerEmitter) reads any attributes from the Provider emitter */
@@ -122,6 +162,7 @@ public:
 
 /** This goes in class declaration for UNiagaraDataInterfaceBase-derived types, that need custom parameter type */
 #define DECLARE_NIAGARA_DI_PARAMETER() \
+	/*UE_STATIC_DEPRECATE(5.1, true, "DECLARE_NIAGARA_DI_PARAMETER is deprecate please update to use the new BuildShaderParameters / SetShaderParameters path.");*/ \
 	virtual FNiagaraDataInterfaceParametersCS* CreateComputeParameters() const override; \
 	virtual const FTypeLayoutDesc* GetComputeParametersTypeDesc() const override; \
 	virtual void BindParameters(FNiagaraDataInterfaceParametersCS* Base, const FNiagaraDataInterfaceGPUParamInfo& ParameterInfo, const class FShaderParameterMap& ParameterMap) override; \
@@ -130,6 +171,7 @@ public:
 
 /** This goes in cpp file matched with class declartion using DECLARE_NIAGARA_DI_PARAMETER() */
 #define IMPLEMENT_NIAGARA_DI_PARAMETER(T, ParameterType) \
+	/*UE_STATIC_DEPRECATE(5.1, true, "IMPLEMENT_NIAGARA_DI_PARAMETER is deprecate please update to use the new BuildShaderParameters / SetShaderParameters path.");*/ \
 	static_assert(ParameterType::InterfaceType == ETypeLayoutInterface::NonVirtual, "DI ParameterType must be non-virtual"); \
 	FNiagaraDataInterfaceParametersCS* T::CreateComputeParameters() const { return new ParameterType(); } \
 	const FTypeLayoutDesc* T::GetComputeParametersTypeDesc() const { return &StaticGetTypeLayoutDesc<ParameterType>(); } \
diff --git a/Engine/Plugins/FX/Niagara/Source/NiagaraEditor/Private/NiagaraHlslTranslator.cpp b/Engine/Plugins/FX/Niagara/Source/NiagaraEditor/Private/NiagaraHlslTranslator.cpp
--- a/Engine/Plugins/FX/Niagara/Source/NiagaraEditor/Private/NiagaraHlslTranslator.cpp
+++ b/Engine/Plugins/FX/Niagara/Source/NiagaraEditor/Private/NiagaraHlslTranslator.cpp
@@ -113,7 +113,7 @@ void FNiagaraShaderQueueTickable::ProcessQueue()
 
 		FNiagaraComputeShaderCompilationOutput NewCompilationOutput;
 
-		ShaderScript->SetDataInterfaceParamInfo(CompilableScript->GetVMExecutableData().DIParamInfo);
+		ShaderScript->BuildScriptParametersMetadata(CompilableScript->GetVMExecutableData().DIParamInfo);
 		ShaderScript->SetSourceName(TEXT("NiagaraComputeShader"));
 		UNiagaraEmitter* Emitter = Cast<UNiagaraEmitter>(CompilableScript->GetOuter());
 		if (Emitter && Emitter->GetUniqueEmitterName().Len() > 0)
@@ -3168,7 +3168,7 @@ void FHlslNiagaraTranslator::DefineMainGPUFunctions(
 					"\t\t// count isn't updated. In that case we must manually copy the original count here.\n"
 					"\t\tif (WriteInstanceCountOffset != 0xFFFFFFFF && GLinearThreadId == 0) \n"
 					"\t\t{\n"
-					"\t\t	RWInstanceCounts[WriteInstanceCountOffset] = GSpawnStartInstance + SpawnedInstances; \n"
+					"\t\t	RWInstanceCounts[WriteInstanceCountOffset] = GSpawnStartInstance + NumSpawnedInstances; \n"
 					"\t\t}\n"
 				);
 			}
@@ -3295,7 +3295,7 @@ void FHlslNiagaraTranslator::DefineMainGPUFunctions(
 					"	{\n"
 					"		GSpawnStartInstance = RWInstanceCounts[ReadInstanceCountOffset];\n"
 					"	}\n"
-					"	const uint MaxInstances = GSpawnStartInstance + SpawnedInstances;\n"
+					"	const uint MaxInstances = GSpawnStartInstance + NumSpawnedInstances;\n"
 					"	const bool bRunUpdateLogic = bRunSpawnUpdateLogic && GLinearThreadId < GSpawnStartInstance && GLinearThreadId < MaxInstances;\n"
 					"	const bool bRunSpawnLogic = bRunSpawnUpdateLogic && GLinearThreadId >= GSpawnStartInstance && GLinearThreadId < MaxInstances;\n"
 				);
diff --git a/Engine/Plugins/FX/Niagara/Source/NiagaraShader/Private/NiagaraShader.cpp b/Engine/Plugins/FX/Niagara/Source/NiagaraShader/Private/NiagaraShader.cpp
--- a/Engine/Plugins/FX/Niagara/Source/NiagaraShader/Private/NiagaraShader.cpp
+++ b/Engine/Plugins/FX/Niagara/Source/NiagaraShader/Private/NiagaraShader.cpp
@@ -44,6 +44,13 @@ static FAutoConsoleVariableRef CVarNiagaraSkipVectorVMBackendOptimizations(
 	ECVF_Default
 );
 
+int32 GNiagaraShaderForceBindEverything = 0;
+static FAutoConsoleVariableRef CVarNiagaraShaderForceBindEverything(
+	TEXT("fx.Niagara.Shader.ForceBindEverything"),
+	GNiagaraShaderForceBindEverything,
+	TEXT("Forces Niagara to display errors about missing shader bindings.")
+);
+
 #if ENABLE_COOK_STATS
 namespace NiagaraShaderCookStats
 {
@@ -145,9 +152,7 @@ bool FNiagaraShaderMapPointerTable::LoadFromArchive(FArchive& Ar, const FPlatfor
 void FNiagaraShaderMapId::GetScriptHash(FSHAHash& OutHash) const
 {
 	FSHA1 HashState;
-
 	HashState.Update((const uint8*)&CompilerVersionID, sizeof(CompilerVersionID));
-	//HashState.Update((const uint8*)&BaseScriptID, sizeof(BaseScriptID));
 	HashState.Update(BaseCompileHash.Hash, FNiagaraCompileHash::HashSize);
 	HashState.Update((const uint8*)&FeatureLevel, sizeof(FeatureLevel));
 
@@ -170,8 +175,6 @@ void FNiagaraShaderMapId::GetScriptHash(FSHAHash& OutHash) const
 	{
 		HashState.Update(Dependency.SourceHash.Hash, sizeof(Dependency.SourceHash));
 	}
-
-	//ParameterSet.UpdateHash(HashState);		// will need for static switches
 	
 	HashState.Final();
 	HashState.GetHash(&OutHash.Hash[0]);
@@ -185,8 +188,7 @@ void FNiagaraShaderMapId::GetScriptHash(FSHAHash& OutHash) const
 */
 bool FNiagaraShaderMapId::operator==(const FNiagaraShaderMapId& ReferenceSet) const
 {
-	if (/*BaseScriptID != ReferenceSet.BaseScriptID 
-		|| */BaseCompileHash != ReferenceSet.BaseCompileHash
+	if (   BaseCompileHash != ReferenceSet.BaseCompileHash
 		|| FeatureLevel != ReferenceSet.FeatureLevel
 		|| CompilerVersionID != ReferenceSet.CompilerVersionID 
 		|| bUsesRapidIterationParams != ReferenceSet.bUsesRapidIterationParams
@@ -243,39 +245,17 @@ bool FNiagaraShaderMapId::operator==(const FNiagaraShaderMapId& ReferenceSet) co
 		return false;
 	}
 
-	/*
-	if (ParameterSet != ReferenceSet.ParameterSet
-		|| ReferencedFunctions.Num() != ReferenceSet.ReferencedFunctions.Num()
-	{
-		return false;
-	}
-
-	for (int32 RefFunctionIndex = 0; RefFunctionIndex < ReferenceSet.ReferencedFunctions.Num(); RefFunctionIndex++)
-	{
-		const FGuid& ReferenceGuid = ReferenceSet.ReferencedFunctions[RefFunctionIndex];
-
-		if (ReferencedFunctions[RefFunctionIndex] != ReferenceGuid)
-		{
-			return false;
-		}
-	}
-	*/
 	return true;
 }
 
 void FNiagaraShaderMapId::AppendKeyString(FString& KeyString) const
 {
-	//KeyString += BaseScriptID.ToString();
-	//KeyString += TEXT("_");
-
 	KeyString += BaseCompileHash.ToString();
 	KeyString += TEXT("_");
 
 	FString FeatureLevelString;
 	GetFeatureLevelName(FeatureLevel, FeatureLevelString);
 
-	LayoutParams.AppendKeyString(KeyString);
-
 	{
 		const FSHAHash LayoutHash = Freeze::HashLayout(StaticGetTypeLayoutDesc<FNiagaraShaderMapContent>(), LayoutParams);
 		KeyString += TEXT("_");
@@ -290,18 +270,6 @@ void FNiagaraShaderMapId::AppendKeyString(FString& KeyString) const
 		KeyString += TEXT("_");
 	}
 
-	/*for (const FHashedName& DITypeName : DITypeNames)
-	{
-		const FTypeLayoutDesc* DIType = FTypeLayoutDesc::Find(DITypeName.GetHash());
-		if (DIType)
-		{
-			const FSHAHash LayoutHash = Freeze::HashLayout(*DIType, LayoutParams);
-			KeyString += TEXT("_");
-			KeyString += LayoutHash.ToString();
-			KeyString += TEXT("_");
-		}
-	}*/
-
 	KeyString += FeatureLevelString + TEXT("_");
 	KeyString += CompilerVersionID.ToString();
 	KeyString += TEXT("_");
@@ -314,7 +282,10 @@ void FNiagaraShaderMapId::AppendKeyString(FString& KeyString) const
 	{
 		KeyString += TEXT("NORI_");
 	}
-	
+
+	// Add base parameters structure
+	KeyString += FString::Printf(TEXT("%08x"), TShaderParameterStructTypeInfo<FNiagaraShader::FParameters>::GetStructMetadata()->GetLayoutHash());
+
 	// Add additional defines
 	for (int32 DefinesIndex = 0; DefinesIndex < AdditionalDefines.Num(); DefinesIndex++)
 	{
@@ -349,16 +320,10 @@ void FNiagaraShaderMapId::AppendKeyString(FString& KeyString) const
 	for (const FShaderTypeDependency& ShaderTypeDependency : ShaderTypeDependencies)
 	{
 		const FShaderType* ShaderType = FindShaderTypeByName(ShaderTypeDependency.ShaderTypeName);
-
 		KeyString += TEXT("_");
+		KeyString += ShaderType->GetName();
 		KeyString += ShaderTypeDependency.SourceHash.ToString();
 
-		if (const FShaderParametersMetadata* ParameterStructMetadata = ShaderType->GetRootParametersMetadata())
-		{
-			KeyString += FString::Printf(TEXT("%08x"), ParameterStructMetadata->GetLayoutHash());
-		}
-
-		// Add the serialization history to the key string so that we can detect changes to shader serialization without a corresponding .usf change
 		const FSHAHash LayoutHash = Freeze::HashLayout(ShaderType->GetLayout(), LayoutParams);
 		KeyString += LayoutHash.ToString();
 
@@ -383,8 +348,6 @@ void FNiagaraShaderMapId::AppendKeyString(FString& KeyString) const
 }
 
 
-TMap<const FShaderCompileJob*, TArray<FNiagaraDataInterfaceGPUParamInfo> > FNiagaraShaderType::ExtraParamInfo;
-
 #if WITH_EDITOR
 
 /**
@@ -398,8 +361,7 @@ void FNiagaraShaderType::BeginCompileShader(
 	FSharedShaderCompilerEnvironment* CompilationEnvironment,
 	EShaderPlatform Platform,
 	TArray<FShaderCommonCompileJobPtr>& NewJobs,
-	FShaderTarget Target,
-	TArray<FNiagaraDataInterfaceGPUParamInfo>& InDIParamInfo
+	FShaderTarget Target
 	)
 {
 	FShaderCompileJob* NewJob = GShaderCompilingManager->PrepareShaderCompileJob(ShaderMapId, FShaderCompileJobKey(this, nullptr, PermutationId), EShaderCompileJobPriority::Normal);
@@ -408,8 +370,7 @@ void FNiagaraShaderType::BeginCompileShader(
 		return;
 	}
 
-	//NewJob->DIParamInfo = InDIParamInfo;	// from hlsl translation and need to be passed to the FNiagaraShader on completion
-
+	NewJob->ShaderParameters = MakeShared<const FParameters, ESPMode::ThreadSafe>(Script->GetScriptParametersMetadata());
 	NewJob->Input.SharedEnvironment = CompilationEnvironment;
 	NewJob->Input.Target = Target;
 	NewJob->Input.ShaderFormat = LegacyShaderPlatformToShaderFormat(Platform);
@@ -418,14 +379,10 @@ void FNiagaraShaderType::BeginCompileShader(
 	NewJob->Input.Environment.SetDefine(TEXT("GPU_SIMULATION"), 1);
 	NewJob->Input.Environment.SetDefine(TEXT("NIAGARA_MAX_GPU_SPAWN_INFOS"), NIAGARA_MAX_GPU_SPAWN_INFOS);
 	Script->GetScriptHLSLSource(NewJob->Input.Environment.IncludeVirtualPathToContentsMap.Add(TEXT("/Engine/Generated/NiagaraEmitterInstance.ush")));
-	NewJob->Input.Environment.SetDefine(TEXT("SHADER_STAGE_PERMUTATION"), PermutationId);
 
-	NewJob->Input.Environment.SetDefine(TEXT("NIAGARA_SHADER_PERMUTATIONS"), 1);
-	NewJob->Input.Environment.SetDefine(TEXT("DefaultSimulationStageIndex"), 0);
 	NewJob->Input.Environment.SetDefine(TEXT("SimulationStageIndex"), PermutationId);
-	NewJob->Input.Environment.SetDefine(TEXT("USE_SIMULATION_STAGES"), 1);
 
-	TConstArrayView<FSimulationStageMetaData> SimStageMetaDataArray = const_cast<FNiagaraShaderScript*>(Script)->GetBaseVMScript()->GetSimulationStageMetaData();
+	TConstArrayView<FSimulationStageMetaData> SimStageMetaDataArray = Script->GetBaseVMScript()->GetSimulationStageMetaData();
 	const FSimulationStageMetaData& SimStageMetaData = SimStageMetaDataArray[PermutationId];
 	if (SimStageMetaData.bWritesParticles && SimStageMetaData.bPartialParticleUpdate)
 	{
@@ -468,7 +425,6 @@ void FNiagaraShaderType::BeginCompileShader(
 		NewJob->Input
 	);
 
-	ExtraParamInfo.Add(NewJob, InDIParamInfo);
 	NewJobs.Add(FShaderCommonCompileJobPtr(NewJob));
 }
 #endif
@@ -539,7 +495,6 @@ void FNiagaraShaderType::AddReferencedUniformBufferIncludes(FShaderCompilerEnvir
  * @param CurrentJob - Compile job that was enqueued by BeginCompileShader.
  */
 FShader* FNiagaraShaderType::FinishCompileShader(
-	//const FUniformExpressionSet& UniformExpressionSet,
 	const FSHAHash& ShaderMapHash,
 	const FShaderCompileJob& CurrentJob,
 	const FString& InDebugDescription
@@ -547,14 +502,16 @@ FShader* FNiagaraShaderType::FinishCompileShader(
 {
 	check(CurrentJob.bSucceeded);
 
-	TArray<FNiagaraDataInterfaceGPUParamInfo> DIParamInfo;
-	if (!ExtraParamInfo.RemoveAndCopyValue(&CurrentJob, DIParamInfo))
-	{
-		check(false);
-	}
-
-	FShader* Shader = ConstructCompiled(FNiagaraShaderType::CompiledShaderInitializerType(this, CurrentJob.Key.PermutationId, CurrentJob.Output, ShaderMapHash, InDebugDescription, DIParamInfo));
-	//CurrentJob.Output.ParameterMap.VerifyBindingsAreComplete(GetName(), CurrentJob.Output.Target, nullptr); // b/c we don't bind data interfaces yet...
+	FShader* Shader = ConstructCompiled(
+		FNiagaraShaderType::CompiledShaderInitializerType(
+			this,
+			static_cast<const FParameters*>(CurrentJob.ShaderParameters.Get()),
+			CurrentJob.Key.PermutationId,
+			CurrentJob.Output,
+			ShaderMapHash,
+			InDebugDescription
+		)
+	);
 
 	return Shader;
 }
@@ -586,28 +543,6 @@ void FNiagaraShaderMap::FlushShaderTypes(TArray<const FShaderType*>& ShaderTypes
 	}
 }
 
-#if 0
-void FNiagaraShaderMap::FixupShaderTypes(EShaderPlatform Platform, const TMap<FShaderType*, FString>& ShaderTypeNames)
-{
-	FScopeLock ScopeLock(&GIdToNiagaraShaderMapCS);
-
-	TArray<FNiagaraShaderMapId> Keys;
-	FNiagaraShaderMap::GIdToNiagaraShaderMap[Platform].GenerateKeyArray(Keys);
-
-	TArray<FNiagaraShaderMap*> Values;
-	FNiagaraShaderMap::GIdToNiagaraShaderMap[Platform].GenerateValueArray(Values);
-
-	//@todo - what about the shader maps in AllNiagaraShaderMaps that are not in GIdToNiagaraShaderMap?
-	FNiagaraShaderMap::GIdToNiagaraShaderMap[Platform].Empty();
-
-	for (int32 PairIndex = 0; PairIndex < Keys.Num(); PairIndex++)
-	{
-		FNiagaraShaderMapId& Key = Keys[PairIndex];
-		FNiagaraShaderMap::GIdToNiagaraShaderMap[Platform].Add(Key, Values[PairIndex]);
-	}
-}
-#endif
-
 void NiagaraShaderMapAppendKeyString(EShaderPlatform Platform, FString& KeyString)
 {
 	// does nothing at the moment, but needs to append to keystring if runtime options impacting selection of sim shader permutations are added
@@ -629,8 +564,6 @@ static FString GetNiagaraShaderMapKeyString(const FNiagaraShaderMapId& ShaderMap
 	return FDerivedDataCacheInterface::BuildCacheKey(TEXT("NIAGARASM"), *NIAGARASHADERMAP_DERIVEDDATA_VER, *ShaderMapKeyString);
 }
 
-
-
 void FNiagaraShaderMap::LoadFromDerivedDataCache(const FNiagaraShaderScript* Script, const FNiagaraShaderMapId& ShaderMapId, EShaderPlatform Platform, TRefCountPtr<FNiagaraShaderMap>& InOutShaderMap)
 {
 	if (InOutShaderMap != NULL)
@@ -796,8 +729,7 @@ void FNiagaraShaderMap::Compile(
 								CompilationEnvironment,
 								InPlatform,
 								NewJobs,
-								FShaderTarget(ShaderType->GetFrequency(), InPlatform),
-								Script->GetDataInterfaceParamInfo()
+								FShaderTarget(ShaderType->GetFrequency(), InPlatform)
 							);
 						}
 						NumShaders++;
@@ -1194,98 +1126,73 @@ const FNiagaraShaderMap* FNiagaraShaderMap::GetShaderMapBeingCompiled(const FNia
 
 //////////////////////////////////////////////////////////////////////////
 
-
 FNiagaraShader::FNiagaraShader(const FNiagaraShaderType::CompiledShaderInitializerType& Initializer)
 	: FShader(Initializer)
 	, DebugDescription(Initializer.DebugDescription)
 {
 	check(!DebugDescription.IsEmpty());
-	BindParams(Initializer.DIParamInfo, Initializer.ParameterMap);
-}
 
-void FNiagaraShader::BindParams(const TArray<FNiagaraDataInterfaceGPUParamInfo>& InDIParamInfo, const FShaderParameterMap &ParameterMap)
-{
-	FloatInputBufferParam.Bind(ParameterMap, TEXT("InputFloat"));
-	IntInputBufferParam.Bind(ParameterMap, TEXT("InputInt"));
-	HalfInputBufferParam.Bind(ParameterMap, TEXT("InputHalf"));
-	StaticInputFloatParam.Bind(ParameterMap, TEXT("StaticInputFloat"));
-	FloatOutputBufferParam.Bind(ParameterMap, TEXT("OutputFloat"));
-	IntOutputBufferParam.Bind(ParameterMap, TEXT("OutputInt"));
-	HalfOutputBufferParam.Bind(ParameterMap, TEXT("OutputHalf"));
-
-	InstanceCountsParam.Bind(ParameterMap, TEXT("InstanceCounts"));
-	ReadInstanceCountOffsetParam.Bind(ParameterMap, TEXT("ReadInstanceCountOffset"));
-	WriteInstanceCountOffsetParam.Bind(ParameterMap, TEXT("WriteInstanceCountOffset"));
-
-	FreeIDBufferParam.Bind(ParameterMap, TEXT("FreeIDList"));
-	IDToIndexBufferParam.Bind(ParameterMap, TEXT("IDToIndexTable"));
-	
-	SimStartParam.Bind(ParameterMap, TEXT("SimStart"));
-	EmitterTickCounterParam.Bind(ParameterMap, TEXT("EmitterTickCounter"));
-	EmitterSpawnInfoOffsetsParam.Bind(ParameterMap, TEXT("EmitterSpawnInfoOffsets"));
-	EmitterSpawnInfoParamsParam.Bind(ParameterMap, TEXT("EmitterSpawnInfoParams"));
-
-	NumSpawnedInstancesParam.Bind(ParameterMap, TEXT("SpawnedInstances"));
-
-	SimulationStageIterationInfoParam.Bind(ParameterMap, TEXT("SimulationStageIterationInfo"));
-	SimulationStageNormalizedIterationIndexParam.Bind(ParameterMap, TEXT("SimulationStageNormalizedIterationIndex"));
-	ParticleIterationStateInfoParam.Bind(ParameterMap, TEXT("ParticleIterationStateInfo"));
-	DispatchThreadIdBoundsParam.Bind(ParameterMap, TEXT("DispatchThreadIdBounds"));
-	DispatchThreadIdToLinearParam.Bind(ParameterMap, TEXT("DispatchThreadIdToLinear"));
-
-	ComponentBufferSizeReadParam.Bind(ParameterMap, TEXT("ComponentBufferSizeRead"));
-	ComponentBufferSizeWriteParam.Bind(ParameterMap, TEXT("ComponentBufferSizeWrite"));
-	ViewUniformBufferParam.Bind(ParameterMap, TEXT("View"));
-
-	GlobalConstantBufferParam[0].Bind(ParameterMap, TEXT("FNiagaraGlobalParameters"));
-	SystemConstantBufferParam[0].Bind(ParameterMap, TEXT("FNiagaraSystemParameters"));
-	OwnerConstantBufferParam[0].Bind(ParameterMap, TEXT("FNiagaraOwnerParameters"));
-	EmitterConstantBufferParam[0].Bind(ParameterMap, TEXT("FNiagaraEmitterParameters"));
-	ExternalConstantBufferParam[0].Bind(ParameterMap, TEXT("FNiagaraExternalParameters"));
-
-	GlobalConstantBufferParam[1].Bind(ParameterMap, TEXT("PREV_FNiagaraGlobalParameters"));
-	SystemConstantBufferParam[1].Bind(ParameterMap, TEXT("PREV_FNiagaraSystemParameters"));
-	OwnerConstantBufferParam[1].Bind(ParameterMap, TEXT("PREV_FNiagaraOwnerParameters"));
-	EmitterConstantBufferParam[1].Bind(ParameterMap, TEXT("PREV_FNiagaraEmitterParameters"));
-	ExternalConstantBufferParam[1].Bind(ParameterMap, TEXT("PREV_FNiagaraExternalParameters"));
-
-	DataInterfaceParameters.Empty(InDIParamInfo.Num());
-	for(const FNiagaraDataInterfaceGPUParamInfo& DIParamInfo : InDIParamInfo)
-	{
-		FNiagaraDataInterfaceParamRef& ParamRef = DataInterfaceParameters.AddDefaulted_GetRef();
-		ParamRef.Bind(DIParamInfo, ParameterMap);
-	}
-}
+	// Currently we can not bind all parameters due to the external constant buffers
+	// Once they are inside parameter structure and UseLegacyShaderBindings are all true we can completely bind with validation
+	const bool bShouldBindEverything = false || (GNiagaraShaderForceBindEverything != 0);
 
-//////////////////////////////////////////////////////////////////////////
-FNiagaraDataInterfaceParamRef::FNiagaraDataInterfaceParamRef()
-{
+	// Bind parameters
+	const FNiagaraShaderType::FParameters* InitializerParameters = static_cast<const FNiagaraShaderType::FParameters*>(Initializer.Parameters);
+	Bindings.BindForLegacyShaderParameters(
+		this,
+		Initializer.PermutationId,
+		Initializer.ParameterMap,
+		*InitializerParameters->ScriptParametersMetadata->ShaderParametersMetadata.Get(),
+		bShouldBindEverything
+	);
 
-}
+	bNeedsViewUniformBuffer = Initializer.ParameterMap.ContainsParameterAllocation(TEXT("View"));
+
+	// Legacy bindings
+	GlobalConstantBufferParam[0].Bind(Initializer.ParameterMap, TEXT("FNiagaraGlobalParameters"));
+	SystemConstantBufferParam[0].Bind(Initializer.ParameterMap, TEXT("FNiagaraSystemParameters"));
+	OwnerConstantBufferParam[0].Bind(Initializer.ParameterMap, TEXT("FNiagaraOwnerParameters"));
+	EmitterConstantBufferParam[0].Bind(Initializer.ParameterMap, TEXT("FNiagaraEmitterParameters"));
+	ExternalConstantBufferParam[0].Bind(Initializer.ParameterMap, TEXT("FNiagaraExternalParameters"));
+
+	GlobalConstantBufferParam[1].Bind(Initializer.ParameterMap, TEXT("PREV_FNiagaraGlobalParameters"));
+	SystemConstantBufferParam[1].Bind(Initializer.ParameterMap, TEXT("PREV_FNiagaraSystemParameters"));
+	OwnerConstantBufferParam[1].Bind(Initializer.ParameterMap, TEXT("PREV_FNiagaraOwnerParameters"));
+	EmitterConstantBufferParam[1].Bind(Initializer.ParameterMap, TEXT("PREV_FNiagaraEmitterParameters"));
+	ExternalConstantBufferParam[1].Bind(Initializer.ParameterMap, TEXT("PREV_FNiagaraExternalParameters"));
+
+	// Gather data interface bindings
+	INiagaraShaderModule* ShaderModule = INiagaraShaderModule::Get();
+	DataInterfaceParameters.Empty(InitializerParameters->ScriptParametersMetadata->DataInterfaceParamInfo.Num());
+	for (const FNiagaraDataInterfaceGPUParamInfo& DataInterfaceParamInfo : InitializerParameters->ScriptParametersMetadata->DataInterfaceParamInfo)
+	{
+		FNiagaraDataInterfaceParamRef& DataInterfaceParamRef = DataInterfaceParameters.AddDefaulted_GetRef();
+		DataInterfaceParamRef.ShaderParametersOffset = DataInterfaceParamInfo.ShaderParametersOffset;
 
-FNiagaraDataInterfaceParamRef::~FNiagaraDataInterfaceParamRef()
-{
-	Parameters.SafeDelete();
-}
+		UNiagaraDataInterfaceBase* CDODataInterface = ShaderModule->RequestDefaultDataInterface(*DataInterfaceParamInfo.DIClassName);
+		if (CDODataInterface == nullptr)
+		{
+			continue;
+		}
 
-void FNiagaraDataInterfaceParamRef::Bind(const FNiagaraDataInterfaceGPUParamInfo& InParameterInfo, const class FShaderParameterMap& ParameterMap)
-{
-	INiagaraShaderModule* Module = INiagaraShaderModule::Get();
-	UNiagaraDataInterfaceBase* Base = Module->RequestDefaultDataInterface(*InParameterInfo.DIClassName);
-	ensureMsgf(Base, TEXT("Failed to load class for FNiagaraDataInterfaceParamRef. %s"), *InParameterInfo.DIClassName);
-	if (Base)
-	{
-		DIType = TIndexedPtr<UNiagaraDataInterfaceBase>(Base); // TODO - clean up TIndexedPtr::operator=()
-		Parameters = Base->CreateComputeParameters();
-		if (Parameters)
+		DataInterfaceParamRef.DIType = TIndexedPtr<UNiagaraDataInterfaceBase>(CDODataInterface);
+		if (CDODataInterface->UseLegacyShaderBindings() == true)
 		{
-			checkf(Base->GetComputeParametersTypeDesc(), TEXT("Missing TypeDesc for compute parameters for class %s"), *InParameterInfo.DIClassName);
-			Parameters->DIType = DIType;
-			Base->BindParameters(Parameters, InParameterInfo, ParameterMap);
+			DataInterfaceParamRef.Parameters = CDODataInterface->CreateComputeParameters();
+			if (DataInterfaceParamRef.Parameters != nullptr)
+			{
+				CDODataInterface->BindParameters(DataInterfaceParamRef.Parameters, DataInterfaceParamInfo, Initializer.ParameterMap);
+			}
+		}
+		else
+		{
+			DataInterfaceParamRef.Parameters = CDODataInterface->CreateShaderStorage(DataInterfaceParamInfo, Initializer.ParameterMap);
+			checkf(DataInterfaceParamRef.Parameters == nullptr || CDODataInterface->GetShaderStorageType() != nullptr, TEXT("DataInterface(%s) provides shader storage but did not implement GetShaderStorageType"), *GetNameSafe(CDODataInterface->GetClass()));
 		}
 	}
 }
 
+//////////////////////////////////////////////////////////////////////////
 bool FNiagaraDataInterfaceGeneratedFunction::Serialize(FArchive& Ar)
 {
 	Ar << DefinitionName;
@@ -1302,7 +1209,14 @@ bool operator<<(FArchive& Ar, FNiagaraDataInterfaceGeneratedFunction& DIFunction
 void FNiagaraDataInterfaceParamRef::WriteFrozenParameters(FMemoryImageWriter& Writer, const TMemoryImagePtr<FNiagaraDataInterfaceParametersCS>& InParameters) const
 {
 	const UNiagaraDataInterfaceBase* Base = DIType.Get(Writer.TryGetPrevPointerTable());
-	InParameters.WriteMemoryImageWithDerivedType(Writer, Base->GetComputeParametersTypeDesc());
+	if (Base->UseLegacyShaderBindings())
+	{
+		InParameters.WriteMemoryImageWithDerivedType(Writer, Base->GetComputeParametersTypeDesc());
+	}
+	else
+	{
+		InParameters.WriteMemoryImageWithDerivedType(Writer, Base->GetShaderStorageType());
+	}
 }
 
 bool FNiagaraDataInterfaceGPUParamInfo::IsUserParameter() const
diff --git a/Engine/Plugins/FX/Niagara/Source/NiagaraShader/Private/NiagaraShared.cpp b/Engine/Plugins/FX/Niagara/Source/NiagaraShader/Private/NiagaraShared.cpp
--- a/Engine/Plugins/FX/Niagara/Source/NiagaraShader/Private/NiagaraShared.cpp
+++ b/Engine/Plugins/FX/Niagara/Source/NiagaraShader/Private/NiagaraShared.cpp
@@ -6,6 +6,7 @@
 
 #include "NiagaraShared.h"
 #include "NiagaraShaderModule.h"
+#include "NiagaraShaderParametersBuilder.h"
 #include "NiagaraShaderType.h"
 #include "NiagaraShader.h"
 #include "NiagaraScriptBase.h"
@@ -23,6 +24,7 @@
 #if WITH_EDITOR
 #include "Interfaces/ITargetPlatform.h"
 #endif
+#include "ShaderParameterMetadataBuilder.h"
 
 IMPLEMENT_TYPE_LAYOUT(FNiagaraDataInterfaceParamRef);
 IMPLEMENT_TYPE_LAYOUT(FNiagaraShaderMapContent);
@@ -223,8 +225,7 @@ void FNiagaraShaderScript::GetShaderMapId(EShaderPlatform Platform, const ITarge
 	{
 		INiagaraShaderModule* Module = INiagaraShaderModule::Get();
 
-		OutId.FeatureLevel = GetFeatureLevel();/*
-		OutId.BaseScriptID = BaseScriptId;*/
+		OutId.FeatureLevel = GetFeatureLevel();
 		OutId.bUsesRapidIterationParams = bUsesRapidIterationParams;		
 		BaseCompileHash.ToSHAHash(OutId.BaseCompileHash);
 		OutId.CompilerVersionID = FNiagaraCustomVersion::GetLatestScriptCompileVersion();
@@ -254,18 +255,22 @@ void FNiagaraShaderScript::GetShaderMapId(EShaderPlatform Platform, const ITarge
 			OutId.ShaderTypeDependencies.Emplace(ShaderType, Platform);
 		}
 
+#if WITH_EDITOR
 		if (TargetPlatform)
 		{
-#if WITH_EDITOR
-			OutId.LayoutParams.InitializeForPlatform(TargetPlatform);
-#else
-			UE_LOG(LogShaders, Error, TEXT("FNiagaraShaderScript::GetShaderMapId: TargetPlatform is not null, but a cooked executable cannot target platforms other than its own."));
-#endif
+			OutId.LayoutParams.InitializeForPlatform(TargetPlatform->IniPlatformName(), TargetPlatform->HasEditorOnlyData());
 		}
 		else
 		{
 			OutId.LayoutParams.InitializeForCurrent();
 		}
+#else
+		if (TargetPlatform != nullptr)
+		{
+			UE_LOG(LogShaders, Error, TEXT("FNiagaraShaderScript::GetShaderMapId: TargetPlatform is not null, but a cooked executable cannot target platforms other than its own."));
+		}
+		OutId.LayoutParams.InitializeForCurrent();
+#endif
 	}
 }
 
@@ -396,7 +401,6 @@ void FNiagaraShaderScript::SetScript(UNiagaraScriptBase* InScript, ERHIFeatureLe
 	checkf(InBaseCompileHash.IsValid(), TEXT("Invalid base compile hash.  Script caching will fail."))
 	BaseVMScript = InScript;
 	CompilerVersionId = InCompilerVersionID;
-	//BaseScriptId = InBaseScriptID;
 	AdditionalDefines = InAdditionalDefines;
 	AdditionalVariables = InAdditionalVariables;
 	bUsesRapidIterationParams = bInUsesRapidIterationParams;
@@ -518,7 +522,7 @@ void FNiagaraShaderScript::UpdateCachedData_PostCompile(bool bCalledFromSerializ
 				CachedData.bEmitterConstantBufferUsed |= NiagaraShader->EmitterConstantBufferParam[i].IsBound() ? BitToSet : 0;
 				CachedData.bExternalConstantBufferUsed |= NiagaraShader->ExternalConstantBufferParam[i].IsBound() ? BitToSet : 0;
 			}
-			CachedData.bViewUniformBufferUsed |= NiagaraShader->ViewUniformBufferParam.IsBound() ? 1 : 0;
+			CachedData.bViewUniformBufferUsed |= NiagaraShader->bNeedsViewUniformBuffer;
 		}
 	}
 	else
@@ -675,9 +679,31 @@ void FNiagaraShaderScript::FinishCompilation()
 
 #endif
 
-void FNiagaraShaderScript::SetDataInterfaceParamInfo(const TArray< FNiagaraDataInterfaceGPUParamInfo >& InDIParamInfo)
+void FNiagaraShaderScript::BuildScriptParametersMetadata(const TArray<FNiagaraDataInterfaceGPUParamInfo>& InDIParamInfo)
 {
-	DIParamInfo = InDIParamInfo;
+	ScriptParametersMetadata = MakeShared<FNiagaraShaderScriptParametersMetadata>();
+	ScriptParametersMetadata->DataInterfaceParamInfo = InDIParamInfo;
+
+	FShaderParametersMetadataBuilder ShaderMetadataBuilder(TShaderParameterStructTypeInfo<FNiagaraShader::FParameters>::GetStructMetadata());
+
+	// Build meta data for each data interface
+	INiagaraShaderModule* ShaderModule = INiagaraShaderModule::Get();
+	for (FNiagaraDataInterfaceGPUParamInfo& DataInterfaceParamInfo : ScriptParametersMetadata->DataInterfaceParamInfo)
+	{
+		UNiagaraDataInterfaceBase* CDODataInterface = ShaderModule->RequestDefaultDataInterface(*DataInterfaceParamInfo.DIClassName);
+		if (CDODataInterface->UseLegacyShaderBindings() == false)
+		{
+			const uint32 NextMemberOffset = ShaderMetadataBuilder.GetNextMemberOffset();
+			FNiagaraShaderParametersBuilder ShaderParametersBuilder(DataInterfaceParamInfo, ShaderMetadataBuilder);
+			CDODataInterface->BuildShaderParameters(ShaderParametersBuilder);
+			if (NextMemberOffset != ShaderMetadataBuilder.GetNextMemberOffset())
+			{
+				DataInterfaceParamInfo.ShaderParametersOffset = NextMemberOffset;
+			}
+		}
+	}
+
+	ScriptParametersMetadata->ShaderParametersMetadata.Reset(ShaderMetadataBuilder.Build(FShaderParametersMetadata::EUseCase::ShaderParameterStruct, TEXT("FNiagaraShaderScript")));
 }
 
 FNiagaraShaderRef FNiagaraShaderScript::GetShader(int32 PermutationId) const
diff --git a/Engine/Plugins/FX/Niagara/Source/NiagaraShader/Public/NiagaraShader.h b/Engine/Plugins/FX/Niagara/Source/NiagaraShader/Public/NiagaraShader.h
--- a/Engine/Plugins/FX/Niagara/Source/NiagaraShader/Public/NiagaraShader.h
+++ b/Engine/Plugins/FX/Niagara/Source/NiagaraShader/Public/NiagaraShader.h
@@ -6,12 +6,11 @@
 
 #pragma once
 
-//#include "HAL/IConsoleManager.h"
-//#include "RHI.h"
 #include "ShaderParameters.h"
 #include "SceneView.h"
 #include "Shader.h"
 #include "GlobalShader.h"
+#include "Math/Vector.h"
 #include "NiagaraCommon.h"
 #include "NiagaraScriptBase.h"
 #include "NiagaraShared.h"
@@ -29,15 +28,47 @@ class NIAGARASHADER_API FNiagaraShader : public FShader
 {
 public:
 	DECLARE_SHADER_TYPE(FNiagaraShader, Niagara);
+	//SHADER_USE_PARAMETER_STRUCT(FNiagaraShader, FShader);
+
+	BEGIN_SHADER_PARAMETER_STRUCT(FParameters, )
+		SHADER_PARAMETER(uint32,									ComponentBufferSizeRead)
+		SHADER_PARAMETER(uint32,									ComponentBufferSizeWrite)
+		SHADER_PARAMETER(uint32,									SimStart)
+
+		SHADER_PARAMETER_SRV(Buffer<float>,							InputFloat)
+		SHADER_PARAMETER_SRV(Buffer<half>,							InputHalf)
+		SHADER_PARAMETER_SRV(Buffer<int>,							InputInt)
+		SHADER_PARAMETER_SRV(Buffer<float>,							StaticInputFloat)
+	
+		SHADER_PARAMETER_UAV(RWBuffer<float>,						RWOutputFloat)
+		SHADER_PARAMETER_UAV(RWBuffer<half>,						RWOutputHalf)
+		SHADER_PARAMETER_UAV(RWBuffer<int>,							RWOutputInt)
+
+		SHADER_PARAMETER_UAV(RWBuffer<uint>,						RWInstanceCounts)
+		SHADER_PARAMETER(uint32,									ReadInstanceCountOffset)
+		SHADER_PARAMETER(uint32,									WriteInstanceCountOffset)
+
+		SHADER_PARAMETER_SRV(Buffer<int>,							FreeIDList)
+		SHADER_PARAMETER_UAV(RWBuffer<int>,							RWIDToIndexTable)
+
+		SHADER_PARAMETER(FIntVector4,								SimulationStageIterationInfo)
+		SHADER_PARAMETER(float,										SimulationStageNormalizedIterationIndex)
+
+		SHADER_PARAMETER(FIntVector3,								ParticleIterationStateInfo)
+
+		SHADER_PARAMETER(uint32,									EmitterTickCounter)
+		SHADER_PARAMETER_ARRAY(FIntVector4,							EmitterSpawnInfoOffsets,	[(NIAGARA_MAX_GPU_SPAWN_INFOS + 3) / 4])
+		SHADER_PARAMETER_ARRAY(FVector4f,							EmitterSpawnInfoParams,		[NIAGARA_MAX_GPU_SPAWN_INFOS])
+		SHADER_PARAMETER(uint32,									NumSpawnedInstances)
+
+		SHADER_PARAMETER(FUintVector3,								DispatchThreadIdBounds)
+		SHADER_PARAMETER(FUintVector3,								DispatchThreadIdToLinear)
+	END_SHADER_PARAMETER_STRUCT()
 
 	using FPermutationParameters = FNiagaraShaderPermutationParameters;
 
 	static FName UniformBufferLayoutName;
 
-	FNiagaraShader()
-	{
-	}
-
 	static FIntVector GetDefaultThreadGroupSize(ENiagaraGpuDispatchType DispatchType)
 	{
 		//-TODO: Grab this from FDataDrivenShaderPlatformInfo
@@ -59,56 +90,25 @@ public:
 
 	static bool ShouldCompilePermutation(const FNiagaraShaderPermutationParameters& Parameters)
 	{
-		//@todo - lit materials only 
 		return FNiagaraUtilities::SupportsComputeShaders(Parameters.Platform);
 	}
 
+	FNiagaraShader() {}
 	FNiagaraShader(const FNiagaraShaderType::CompiledShaderInitializerType& Initializer);
 
-	// Bind parameters
-	void BindParams(const TArray<FNiagaraDataInterfaceGPUParamInfo>& InDIParamInfo, const FShaderParameterMap &ParameterMap);
-
 	const TMemoryImageArray<FNiagaraDataInterfaceParamRef>& GetDIParameters()
 	{
 		return DataInterfaceParameters;
 	}
 
-	LAYOUT_FIELD(FShaderResourceParameter, FloatInputBufferParam);
-	LAYOUT_FIELD(FShaderResourceParameter, IntInputBufferParam);
-	LAYOUT_FIELD(FShaderResourceParameter, HalfInputBufferParam);
-	LAYOUT_FIELD(FShaderResourceParameter, StaticInputFloatParam);
-	LAYOUT_FIELD(FRWShaderParameter, FloatOutputBufferParam);
-	LAYOUT_FIELD(FRWShaderParameter, IntOutputBufferParam);
-	LAYOUT_FIELD(FRWShaderParameter, HalfOutputBufferParam);
-	LAYOUT_FIELD(FRWShaderParameter, InstanceCountsParam);
-	LAYOUT_FIELD(FShaderParameter, ReadInstanceCountOffsetParam);
-	LAYOUT_FIELD(FShaderParameter, WriteInstanceCountOffsetParam);
-	LAYOUT_FIELD(FShaderResourceParameter, FreeIDBufferParam);
-	LAYOUT_FIELD(FRWShaderParameter, IDToIndexBufferParam);
+	LAYOUT_FIELD(bool, bNeedsViewUniformBuffer);
 	LAYOUT_ARRAY(FShaderUniformBufferParameter, GlobalConstantBufferParam, 2);
 	LAYOUT_ARRAY(FShaderUniformBufferParameter, SystemConstantBufferParam, 2);
 	LAYOUT_ARRAY(FShaderUniformBufferParameter, OwnerConstantBufferParam, 2);
 	LAYOUT_ARRAY(FShaderUniformBufferParameter, EmitterConstantBufferParam, 2);
 	LAYOUT_ARRAY(FShaderUniformBufferParameter, ExternalConstantBufferParam, 2);
-	LAYOUT_FIELD(FShaderUniformBufferParameter, ViewUniformBufferParam);
-	LAYOUT_FIELD(FShaderParameter, SimStartParam);
-	LAYOUT_FIELD(FShaderParameter, EmitterTickCounterParam);
-	LAYOUT_FIELD(FShaderParameter, EmitterSpawnInfoOffsetsParam);
-	LAYOUT_FIELD(FShaderParameter, EmitterSpawnInfoParamsParam);
-	LAYOUT_FIELD(FShaderParameter, NumSpawnedInstancesParam);
-
-	LAYOUT_FIELD(FShaderParameter, SimulationStageIterationInfoParam);
-	LAYOUT_FIELD(FShaderParameter, SimulationStageNormalizedIterationIndexParam);
-	LAYOUT_FIELD(FShaderParameter, ParticleIterationStateInfoParam);
-	LAYOUT_FIELD(FShaderParameter, DispatchThreadIdBoundsParam);
-	LAYOUT_FIELD(FShaderParameter, DispatchThreadIdToLinearParam);
-
-	LAYOUT_FIELD(FShaderParameter, ComponentBufferSizeReadParam);
-	LAYOUT_FIELD(FShaderParameter, ComponentBufferSizeWriteParam);
 
 private:
-	LAYOUT_FIELD(FShaderUniformBufferParameter, NiagaraUniformBuffer);
-
 	// Data about parameters used for each Data Interface.
 	LAYOUT_FIELD(TMemoryImageArray<FNiagaraDataInterfaceParamRef>, DataInterfaceParameters);
 
@@ -117,25 +117,6 @@ private:
 	FRHIUniformBufferLayout		DebugUniformExpressionUBLayout;
 	*/
 	LAYOUT_FIELD(FMemoryImageString, DebugDescription);
-
-	/* OPTODO: ? */
-	/*
-	// If true, cached uniform expressions are allowed.
-	static int32 bAllowCachedUniformExpressions;
-	// Console variable ref to toggle cached uniform expressions.
-	static FAutoConsoleVariableRef CVarAllowCachedUniformExpressions;
-	*/
-
-	/*
-#if !(UE_BUILD_TEST || UE_BUILD_SHIPPING || !WITH_EDITOR)
-	void VerifyExpressionAndShaderMaps(const FMaterialRenderProxy* MaterialRenderProxy, const FMaterial& Material, const FUniformExpressionCache* UniformExpressionCache);
-#endif
-	*/
-};
-
-class FNiagaraEmitterInstanceShader : public FNiagaraShader
-{
-
 };
 
 extern NIAGARASHADER_API int32 GNiagaraSkipVectorVMBackendOptimizations;
diff --git a/Engine/Plugins/FX/Niagara/Source/NiagaraShader/Public/NiagaraShaderParametersBuilder.h b/Engine/Plugins/FX/Niagara/Source/NiagaraShader/Public/NiagaraShaderParametersBuilder.h
new file mode 100644
--- /dev/null
+++ b/Engine/Plugins/FX/Niagara/Source/NiagaraShader/Public/NiagaraShaderParametersBuilder.h
@@ -0,0 +1,37 @@
+// Copyright Epic Games, Inc. All Rights Reserved.
+
+#pragma once
+
+#include "NiagaraShared.h"
+#include "ShaderParameterMetadataBuilder.h"
+
+class NIAGARASHADER_API FNiagaraShaderParametersBuilder
+{
+public:
+	explicit FNiagaraShaderParametersBuilder(FNiagaraDataInterfaceGPUParamInfo& InParamInfo, FShaderParametersMetadataBuilder& InMetadataBuilder)
+		: ParamInfo(InParamInfo)
+		, MetadataBuilder(InMetadataBuilder)
+	{
+	}
+
+	template<typename T> void AddLooseParam(const TCHAR* Name)
+	{
+		ParamInfo.LooseMetadataNames.Emplace(FString::Printf(TEXT("%s_%s"), *ParamInfo.DataInterfaceHLSLSymbol, Name));
+		MetadataBuilder.AddParam<T>(*ParamInfo.LooseMetadataNames.Last());
+	}
+	template<typename T> void AddLooseParamArray(const TCHAR* Name, int32 NumElements)
+	{
+		ParamInfo.LooseMetadataNames.Emplace(FString::Printf(TEXT("%s_%s"), *ParamInfo.DataInterfaceHLSLSymbol, Name));
+		MetadataBuilder.AddParamArray<T>(*ParamInfo.LooseMetadataNames.Last(), NumElements);
+	}
+	template<typename T> void AddNestedStruct()
+	{
+		MetadataBuilder.AddNestedStruct<T>(*ParamInfo.DataInterfaceHLSLSymbol);
+	}
+
+	TArrayView<FNiagaraDataInterfaceGeneratedFunction> GetGeneratedFunctions() const { return ParamInfo.GeneratedFunctions; }
+
+private:
+	FNiagaraDataInterfaceGPUParamInfo& ParamInfo;
+	FShaderParametersMetadataBuilder& MetadataBuilder;
+};
diff --git a/Engine/Plugins/FX/Niagara/Source/NiagaraShader/Public/NiagaraShaderType.h b/Engine/Plugins/FX/Niagara/Source/NiagaraShader/Public/NiagaraShaderType.h
--- a/Engine/Plugins/FX/Niagara/Source/NiagaraShader/Public/NiagaraShaderType.h
+++ b/Engine/Plugins/FX/Niagara/Source/NiagaraShader/Public/NiagaraShaderType.h
@@ -52,24 +52,32 @@ struct FNiagaraShaderPermutationParameters : public FShaderPermutationParameters
 class FNiagaraShaderType : public FShaderType
 {
 public:
+	struct FParameters : public FShaderType::FParameters
+	{
+		TSharedRef<FNiagaraShaderScriptParametersMetadata> ScriptParametersMetadata;
+
+		FParameters(const TSharedRef<FNiagaraShaderScriptParametersMetadata>& InScriptParametersMetadata)
+			: ScriptParametersMetadata(InScriptParametersMetadata)
+		{
+		}
+	};
+
 	struct CompiledShaderInitializerType : FShaderCompiledShaderInitializerType
 	{
-		//const FUniformExpressionSet& UniformExpressionSet;
 		const FString DebugDescription;
-		TArray< FNiagaraDataInterfaceGPUParamInfo > DIParamInfo;
 
 		CompiledShaderInitializerType(
 			const FShaderType* InType,
+			const FParameters* InParameters,
 			int32 InPermutationId,
 			const FShaderCompilerOutput& CompilerOutput,
 			const FSHAHash& InNiagaraShaderMapHash,
-			const FString& InDebugDescription,
-			const TArray< FNiagaraDataInterfaceGPUParamInfo > &InDIParamInfo
-			)
-		: FShaderCompiledShaderInitializerType(InType,nullptr,InPermutationId,CompilerOutput, InNiagaraShaderMapHash,nullptr,nullptr)
-		, DebugDescription(InDebugDescription)
-		, DIParamInfo(InDIParamInfo)
-		{}
+			const FString& InDebugDescription
+		)
+			: FShaderCompiledShaderInitializerType(InType, InParameters, InPermutationId, CompilerOutput, InNiagaraShaderMapHash, nullptr, nullptr)
+			, DebugDescription(InDebugDescription)
+		{
+		}
 	};
 
 	FNiagaraShaderType(
@@ -86,15 +94,16 @@ public:
 		ValidateCompiledResultType InValidateCompiledResultRef,
 		uint32 InTypeSize,
 		const FShaderParametersMetadata* InRootParametersMetadata = nullptr
-		):
-		FShaderType(EShaderTypeForDynamicCast::Niagara, InTypeLayout, InName, InSourceFilename, InFunctionName, SF_Compute, InTotalPermutationCount,
+	)
+		: FShaderType(EShaderTypeForDynamicCast::Niagara, InTypeLayout, InName, InSourceFilename, InFunctionName, SF_Compute, InTotalPermutationCount,
 			InConstructSerializedRef,
 			InConstructCompiledRef,
 			InModifyCompilationEnvironmentRef,
 			InShouldCompilePermutationRef,
 			InValidateCompiledResultRef,
 			InTypeSize,
-			InRootParametersMetadata)
+			InRootParametersMetadata
+		)
 	{
 		check(InTotalPermutationCount == 1);
 	}
@@ -110,8 +119,7 @@ public:
 			FSharedShaderCompilerEnvironment* CompilationEnvironment,
 			EShaderPlatform Platform,
 			TArray<TRefCountPtr<class FShaderCommonCompileJob>>& NewJobs,
-			FShaderTarget Target,
-			TArray<FNiagaraDataInterfaceGPUParamInfo>& InDIParamInfo
+			FShaderTarget Target
 		);
 
 	/**
@@ -153,7 +161,4 @@ protected:
 	{
 		ModifyCompilationEnvironment(FNiagaraShaderPermutationParameters(Platform, Script), Environment);
 	}
-
-private:
-	static TMap<const FShaderCompileJob*, TArray<FNiagaraDataInterfaceGPUParamInfo> > ExtraParamInfo;
 };
diff --git a/Engine/Plugins/FX/Niagara/Source/NiagaraShader/Public/NiagaraShared.h b/Engine/Plugins/FX/Niagara/Source/NiagaraShader/Public/NiagaraShared.h
--- a/Engine/Plugins/FX/Niagara/Source/NiagaraShader/Public/NiagaraShared.h
+++ b/Engine/Plugins/FX/Niagara/Source/NiagaraShader/Public/NiagaraShared.h
@@ -210,6 +210,12 @@ struct NIAGARASHADER_API FNiagaraDataInterfaceGPUParamInfo
 	UPROPERTY()
 	FString DIClassName;
 
+	UPROPERTY()
+	uint32 ShaderParametersOffset = INDEX_NONE;
+
+	UPROPERTY()
+	TArray<FString> LooseMetadataNames;
+
 	/** Information about all the functions generated by the translator for this data interface. */
 	UPROPERTY()
 	TArray<FNiagaraDataInterfaceGeneratedFunction> GeneratedFunctions;
@@ -227,6 +233,14 @@ template<> struct TStructOpsTypeTraits<FNiagaraDataInterfaceGPUParamInfo> : publ
 	};
 };
 
+struct NIAGARASHADER_API FNiagaraShaderScriptParametersMetadata
+{
+	TArray<FNiagaraDataInterfaceGPUParamInfo> DataInterfaceParamInfo;
+
+	TUniquePtr<FShaderParametersMetadata> ShaderParametersMetadata;
+};
+
+
 /** 
 * Shader side data needed for binding data interface parameters.
 */
@@ -234,12 +248,11 @@ struct FNiagaraDataInterfaceParamRef
 {
 	DECLARE_TYPE_LAYOUT(FNiagaraDataInterfaceParamRef, NonVirtual);
 public:
-	FNiagaraDataInterfaceParamRef();
-	~FNiagaraDataInterfaceParamRef();
+	FNiagaraDataInterfaceParamRef() { }
+	~FNiagaraDataInterfaceParamRef() { Parameters.SafeDelete(); }
 
-	void Bind(const FNiagaraDataInterfaceGPUParamInfo& InParameterInfo, const class FShaderParameterMap& ParameterMap);
+	LAYOUT_FIELD(uint32, ShaderParametersOffset);
 
-public:
 	/** Type of Parameters */
 	LAYOUT_FIELD(TIndexedPtr<UNiagaraDataInterfaceBase>, DIType);
 
@@ -816,20 +829,14 @@ public:
 	}
 #endif
 
-	UNiagaraScriptBase* GetBaseVMScript()
-	{
-		return BaseVMScript;
-	}
+	UNiagaraScriptBase* GetBaseVMScript() { return BaseVMScript; }
+	const UNiagaraScriptBase* GetBaseVMScript() const { return BaseVMScript; }
 
 	NIAGARASHADER_API  FNiagaraShaderRef GetShader(int32 PermutationId) const;
 	NIAGARASHADER_API  FNiagaraShaderRef GetShaderGameThread(int32 PermutationId) const;
 	
-	NIAGARASHADER_API void SetDataInterfaceParamInfo(const TArray< FNiagaraDataInterfaceGPUParamInfo >& InDIParamInfo);
-
-	TArray< FNiagaraDataInterfaceGPUParamInfo >& GetDataInterfaceParamInfo()
-	{
-		return DIParamInfo;
-	}
+	NIAGARASHADER_API void BuildScriptParametersMetadata(const TArray<FNiagaraDataInterfaceGPUParamInfo>& InDIParamInfo);
+	TSharedRef<FNiagaraShaderScriptParametersMetadata> GetScriptParametersMetadata() const { return ScriptParametersMetadata.ToSharedRef(); }
 
 	NIAGARASHADER_API FOnNiagaraScriptCompilationComplete& OnCompilationComplete()
 	{
@@ -890,8 +897,7 @@ private:
 	 */
 	FNiagaraShaderMap* RenderingThreadShaderMap;
 
-	// Information describing data interface parameters. These come from the HLSL translators and need to be passed down to the shader for binding
-	TArray< FNiagaraDataInterfaceGPUParamInfo > DIParamInfo;
+	TSharedPtr<FNiagaraShaderScriptParametersMetadata> ScriptParametersMetadata;
 
 	/** Guid id for base script*/
 	FGuid BaseScriptId_DEPRECATED;
diff --git a/Engine/Source/Runtime/Core/Private/UObject/UE5MainStreamObjectVersion.cpp b/Engine/Source/Runtime/Core/Private/UObject/UE5MainStreamObjectVersion.cpp
--- a/Engine/Source/Runtime/Core/Private/UObject/UE5MainStreamObjectVersion.cpp
+++ b/Engine/Source/Runtime/Core/Private/UObject/UE5MainStreamObjectVersion.cpp
@@ -14,7 +14,7 @@ TMap<FGuid, FGuid> FUE5MainStreamObjectVersion::GetSystemGuids()
 	SystemGuids.Add(DevGuids.MATERIALSHADERMAP_DERIVEDDATA_VER, FGuid("65C01D817C9A4EEFAE9E988D41A1F3DD"));
 	SystemGuids.Add(DevGuids.NANITE_DERIVEDDATA_VER, FGuid("7A3EA93546624597BEBE2ADFDE3A7628"));
 	SystemGuids.Add(DevGuids.NIAGARASHADERMAP_DERIVEDDATA_VER, FGuid("8A37C45D24F2423CBE5F8F371DE33575"));
-	SystemGuids.Add(DevGuids.Niagara_LatestScriptCompileVersion, FGuid("7AF02382DFED45A08DA36E13DEC058EC"));
+	SystemGuids.Add(DevGuids.Niagara_LatestScriptCompileVersion, FGuid("B6633FF4A1A52A4FB85F554EA49FCED7"));
 	SystemGuids.Add(DevGuids.SkeletalMeshDerivedDataVersion, FGuid("FB7B2FB546DE49B2A2F88B7F69763CEE"));
 	SystemGuids.Add(DevGuids.STATICMESH_DERIVEDDATA_VER, FGuid("98F7D79A5811013825E75BBCC41ED3E9"));
 	SystemGuids.Add(DevGuids.POSESEARCHDB_DERIVEDDATA_VER, FGuid("4E595C2AC5E947D6BA9ABC874353E5BC"));
diff --git a/Engine/Source/Runtime/Engine/Private/Particles/FXSystem.cpp b/Engine/Source/Runtime/Engine/Private/Particles/FXSystem.cpp
--- a/Engine/Source/Runtime/Engine/Private/Particles/FXSystem.cpp
+++ b/Engine/Source/Runtime/Engine/Private/Particles/FXSystem.cpp
@@ -93,7 +93,7 @@ void FFXSystemInterface::UnregisterCustomFXSystem(const FName& InterfaceName)
 	CreateCustomFXDelegates.Remove(InterfaceName);
 }
 
-FRHIUniformBuffer* FFXSystemInterface::GetReferenceViewUniformBuffer(TConstArrayView<FViewInfo> Views)
+TUniformBufferRef<FViewUniformShaderParameters> FFXSystemInterface::GetReferenceViewUniformBuffer(TConstArrayView<FViewInfo> Views)
 {
 	check(Views.Num() > 0);
 	return Views[0].ViewUniformBuffer;
diff --git a/Engine/Source/Runtime/Engine/Public/FXSystem.h b/Engine/Source/Runtime/Engine/Public/FXSystem.h
--- a/Engine/Source/Runtime/Engine/Public/FXSystem.h
+++ b/Engine/Source/Runtime/Engine/Public/FXSystem.h
@@ -275,7 +275,8 @@ protected:
 	/** By making the destructor protected, an instance must be destroyed via FFXSystemInterface::Destroy. */
 	ENGINE_API virtual ~FFXSystemInterface() {}
 
-	ENGINE_API static FRHIUniformBuffer* GetReferenceViewUniformBuffer(TConstArrayView<FViewInfo> Views);
+	
+	ENGINE_API static TUniformBufferRef<FViewUniformShaderParameters> GetReferenceViewUniformBuffer(TConstArrayView<FViewInfo> Views);
 	ENGINE_API static bool GetReferenceAllowGPUUpdate(TConstArrayView<FViewInfo> Views);
 	ENGINE_API static const FGlobalDistanceFieldParameterData* GetReferenceGlobalDistanceFieldData(TConstArrayView<FViewInfo> Views);
 
diff --git a/Engine/Source/Runtime/RenderCore/Public/ShaderParameterMetadataBuilder.h b/Engine/Source/Runtime/RenderCore/Public/ShaderParameterMetadataBuilder.h
--- a/Engine/Source/Runtime/RenderCore/Public/ShaderParameterMetadataBuilder.h
+++ b/Engine/Source/Runtime/RenderCore/Public/ShaderParameterMetadataBuilder.h
@@ -46,6 +46,33 @@ public:
 		NextMemberOffset += sizeof(typename TParamTypeInfo::TAlignedType);
 	}
 
+	template<typename T>
+	void AddParamArray(
+		const TCHAR* Name,
+		int32 NumElements,
+		EShaderPrecisionModifier::Type Precision = EShaderPrecisionModifier::Float
+		)
+	{
+		using TParamTypeInfo = TShaderParameterTypeInfo<T>;
+
+		NextMemberOffset = Align(NextMemberOffset, SHADER_PARAMETER_ARRAY_ELEMENT_ALIGNMENT);
+
+		new(Members) FShaderParametersMetadata::FMember(
+			Name,
+			TEXT(""),
+			__LINE__,
+			NextMemberOffset,
+			TParamTypeInfo::BaseType,
+			Precision,
+			TParamTypeInfo::NumRows,
+			TParamTypeInfo::NumColumns,
+			NumElements,
+			TParamTypeInfo::GetStructMetadata()
+			);
+
+		NextMemberOffset += sizeof(typename TParamTypeInfo::TAlignedType) * NumElements;
+	}
+
 	template<typename T>
 	void AddReferencedStruct(
 		const TCHAR* Name,
@@ -119,6 +146,8 @@ public:
 		EShaderPrecisionModifier::Type Precision = EShaderPrecisionModifier::Float
 		);
 
+	uint32 GetNextMemberOffset() const { return NextMemberOffset; }
+
 	FShaderParametersMetadata* Build(
 		FShaderParametersMetadata::EUseCase UseCase,
 		const TCHAR* ShaderParameterName
