- Convert GBuffer data interface to use template & parameter structure binding

diff --git a/Engine/Plugins/FX/Niagara/Shaders/Private/NiagaraDataInterfaceGBuffer.ush b/Engine/Plugins/FX/Niagara/Shaders/Private/NiagaraDataInterfaceGBuffer.ush
deleted file mode 100644
diff --git a/Engine/Plugins/FX/Niagara/Shaders/Private/NiagaraDataInterfaceGBufferTemplate.ush b/Engine/Plugins/FX/Niagara/Shaders/Private/NiagaraDataInterfaceGBufferTemplate.ush
new file mode 100644
--- /dev/null
+++ b/Engine/Plugins/FX/Niagara/Shaders/Private/NiagaraDataInterfaceGBufferTemplate.ush
@@ -0,0 +1,264 @@
+// Copyright Epic Games, Inc. All Rights Reserved.
+
+#include "/Engine/Private/DeferredShadingCommon.ush"
+
+#define NDIGBUFFER_HAS_GBUFFER				(FEATURE_LEVEL >= FEATURE_LEVEL_SM5)
+#define NDIGBUFFER_USE_GBUFFER_VELOCITY		(false && NDIGBUFFER_HAS_GBUFFER)
+#define NDIGBUFFER_IS_MOBILE				(FEATURE_LEVEL <= FEATURE_LEVEL_ES3_1)
+
+////////////////////////////////////////////////////////////////////////////////////////////////////////////////////////////////////////////////////////////////
+
+Texture2D {ParameterName}_VelocityTexture;
+SamplerState {ParameterName}_VelocityTextureSampler;
+
+float2 ApplyViewportOffset_{ParameterName}(float2 UV, bool bApplyViewportOffset)
+{
+	if (bApplyViewportOffset)
+	{
+		UV = (UV * float2(2.0f, -2.0f) + float2(-1.0f, 1.0f)) * View.ScreenPositionScaleBias.xy + View.ScreenPositionScaleBias.wz;
+	}
+	return UV;
+}
+
+void DecodeDiffuseColor_{ParameterName}(in float2 ScreenUV, in bool bApplyViewportOffset, out bool IsValid, out float3 OutValue)
+{
+#if NDIGBUFFER_HAS_GBUFFER
+	ScreenUV = ApplyViewportOffset_{ParameterName}(ScreenUV, bApplyViewportOffset);
+	FGBufferData GBufferData = GetGBufferData(ScreenUV);
+	IsValid = true;
+	OutValue = GBufferData.DiffuseColor;
+#else
+	IsValid = false;
+	OutValue = 0;
+#endif
+}
+
+void DecodeWorldNormal_{ParameterName}(in float2 ScreenUV, in bool bApplyViewportOffset, out bool IsValid, out float3 OutValue)
+{
+#if NDIGBUFFER_HAS_GBUFFER
+	ScreenUV = ApplyViewportOffset_{ParameterName}(ScreenUV, bApplyViewportOffset);
+	FGBufferData GBufferData = GetGBufferData(ScreenUV);
+	IsValid = true;
+	OutValue = GBufferData.WorldNormal;
+#else
+	IsValid = false;
+	OutValue = 0;
+#endif
+}
+
+void DecodeScreenVelocity_{ParameterName}(in float2 ScreenUV, in bool bApplyViewportOffset, out bool IsValid, out float3 OutValue)
+{
+	IsValid = false;
+	OutValue = 0;
+#if NDIGBUFFER_HAS_GBUFFER
+	ScreenUV = ApplyViewportOffset_{ParameterName}(ScreenUV, bApplyViewportOffset);
+	float4 EncodedVelocity = 0;
+	#if NDIGBUFFER_USE_GBUFFER_VELOCITY
+		EncodedVelocity = GetGBufferData(ScreenUV).Velocity;
+	#else
+		EncodedVelocity = {ParameterName}_VelocityTexture.SampleLevel({ParameterName}_VelocityTextureSampler, ScreenUV, 0);
+	#endif
+	IsValid = true;
+	if ( EncodedVelocity.x > 0 )
+	{
+	#if VELOCITY_ENCODE_DEPTH
+		OutValue = DecodeVelocityFromTexture(EncodedVelocity);
+	#else
+		OutValue.xy = DecodeVelocityFromTexture(EncodedVelocity).xy;
+	#endif
+	}
+#endif
+}
+
+void DecodeWorldVelocity_{ParameterName}(in float2 ScreenUV, in bool bApplyViewportOffset, out bool IsValid, out float3 OutValue)
+{
+	IsValid = false;
+	OutValue = 0;
+#if NDIGBUFFER_HAS_GBUFFER
+	ScreenUV = ApplyViewportOffset_{ParameterName}(ScreenUV, bApplyViewportOffset);
+	float4 EncodedVelocity = 0;
+	#if NDIGBUFFER_USE_GBUFFER_VELOCITY
+		EncodedVelocity = GetGBufferData(ScreenUV).Velocity;
+	#else
+		EncodedVelocity = {ParameterName}_VelocityTexture.SampleLevel({ParameterName}_VelocityTextureSampler, ScreenUV, 0);
+	#endif
+	IsValid = true;
+	if ( EncodedVelocity.x > 0 )
+	{
+	#if VELOCITY_ENCODE_DEPTH
+		OutValue = DecodeVelocityFromTexture(EncodedVelocity);
+		OutValue = mul(float4(OutValue, 0), View.ViewToTranslatedWorld).xyz;
+	#else
+		float2 ScreenVelocity = DecodeVelocityFromTexture(EncodedVelocity).xy;
+		ScreenVelocity.y = -ScreenVelocity.y;
+
+		{
+			float2 PrevScreenUV = ScreenUV - ScreenVelocity;
+
+			float SceneDepth = CalcSceneDepth(ScreenUV);
+			float2 ScreenPosition = (ScreenUV - View.ScreenPositionScaleBias.wz) / View.ScreenPositionScaleBias.xy;
+			float3 WorldPosition = mul(float4(ScreenPosition * SceneDepth, SceneDepth, 1), LWCHackToFloat(PrimaryView.ScreenToWorld)).xyz;
+
+			// Ideally we would use CalcPrevSceneDepth(PrevScreenUV) here to get an accurate location, however we do not store PrevSceneDepth so this is a very rough approximation
+			float PrevSceneDepth = SceneDepth;	//CalcPrevSceneDepth(PrevScreenUV);
+			float2 PrevScreenPosition = (PrevScreenUV - View.ScreenPositionScaleBias.wz) / View.ScreenPositionScaleBias.xy;
+			float3 PrevWorldPosition = mul(float4(PrevScreenPosition * PrevSceneDepth, PrevSceneDepth, 1), LWCHackToFloat(PrimaryView.ScreenToWorld)).xyz;
+
+			OutValue = WorldPosition - PrevWorldPosition;
+		}
+	#endif
+	}
+#endif
+}
+
+void DecodeBaseColor_{ParameterName}(in float2 ScreenUV, in bool bApplyViewportOffset, out bool IsValid, out float3 OutValue)
+{
+#if NDIGBUFFER_HAS_GBUFFER
+	ScreenUV = ApplyViewportOffset_{ParameterName}(ScreenUV, bApplyViewportOffset);
+	FGBufferData GBufferData = GetGBufferData(ScreenUV);
+	IsValid = true;
+	OutValue = GBufferData.BaseColor;
+#else
+	IsValid = false;
+	OutValue = 0;
+#endif
+}
+
+void DecodeMetallic_{ParameterName}(in float2 ScreenUV, in bool bApplyViewportOffset, out bool IsValid, out float OutValue)
+{
+#if NDIGBUFFER_HAS_GBUFFER
+	ScreenUV = ApplyViewportOffset_{ParameterName}(ScreenUV, bApplyViewportOffset);
+	FGBufferData GBufferData = GetGBufferData(ScreenUV);
+	IsValid = true;
+	OutValue = GBufferData.Metallic;
+#else
+	IsValid = false;
+	OutValue = 0;
+#endif
+}
+
+void DecodeSpecular_{ParameterName}(in float2 ScreenUV, in bool bApplyViewportOffset, out bool IsValid, out float OutValue)
+{
+#if NDIGBUFFER_HAS_GBUFFER
+	ScreenUV = ApplyViewportOffset_{ParameterName}(ScreenUV, bApplyViewportOffset);
+	FGBufferData GBufferData = GetGBufferData(ScreenUV);
+	IsValid = true;
+	OutValue = GBufferData.Specular;
+#else
+	IsValid = false;
+	OutValue = 0;
+#endif
+}
+
+void DecodeRoughness_{ParameterName}(in float2 ScreenUV, in bool bApplyViewportOffset, out bool IsValid, out float OutValue)
+{
+#if NDIGBUFFER_HAS_GBUFFER
+	ScreenUV = ApplyViewportOffset_{ParameterName}(ScreenUV, bApplyViewportOffset);
+	FGBufferData GBufferData = GetGBufferData(ScreenUV);
+	IsValid = true;
+	OutValue = GBufferData.Roughness;
+#else
+	IsValid = false;
+	OutValue = 0;
+#endif
+}
+
+void DecodeDepth_{ParameterName}(in float2 ScreenUV, in bool bApplyViewportOffset, out bool IsValid, out float OutValue)
+{
+#if NDIGBUFFER_HAS_GBUFFER
+	ScreenUV = ApplyViewportOffset_{ParameterName}(ScreenUV, bApplyViewportOffset);
+	FGBufferData GBufferData = GetGBufferData(ScreenUV);
+	IsValid = true;
+	OutValue = GBufferData.Depth;
+#elif NDIGBUFFER_IS_MOBILE
+	IsValid = true;
+	OutValue = CalcSceneDepth(ScreenUV);
+#else
+	IsValid = false;
+	OutValue = 0;
+#endif
+}
+
+#if 0
+// This has not been tested yet
+void DecodeStencil_{ParameterName}(in float2 ScreenUV, in bool bApplyViewportOffset, out bool IsValid, out int OutValue)
+{
+#if NDIGBUFFER_HAS_GBUFFER
+	IsValid = true;
+
+	ScreenUV = ApplyViewportOffset_{ParameterName}(ScreenUV, bApplyViewportOffset);
+	uint2 IntUV = (uint2)(ScreenUV.xy * (float2)View.BufferSizeAndInvSize.xy);
+	uint Stencil = SceneStencilTexture.Load(uint3(IntUV, 0)) STENCIL_COMPONENT_SWIZZLE;
+	OutValue = int(Stencil);
+#else
+	IsValid = false;
+	OutValue = 0;
+#endif
+}
+#endif
+
+void DecodeCustomDepth_{ParameterName}(in float2 ScreenUV, in bool bApplyViewportOffset, out bool IsValid, out float OutValue)
+{
+#if NDIGBUFFER_HAS_GBUFFER
+	ScreenUV = ApplyViewportOffset_{ParameterName}(ScreenUV, bApplyViewportOffset);
+	FGBufferData GBufferData = GetGBufferData(ScreenUV);
+	IsValid = true;
+	OutValue = GBufferData.CustomDepth;
+#elif NDIGBUFFER_IS_MOBILE
+	IsValid = true;
+	OutValue = ConvertFromDeviceZ(Texture2DSampleLevel(MobileSceneTextures.CustomDepthTexture, MobileSceneTextures.CustomDepthTextureSampler, ScreenUV, 0).r);
+#else
+	IsValid = false;
+	OutValue = 0;
+#endif
+}
+
+void DecodeCustomStencil_{ParameterName}(in float2 ScreenUV, in bool bApplyViewportOffset, out bool IsValid, out int OutValue)
+{
+#if NDIGBUFFER_HAS_GBUFFER
+	ScreenUV = ApplyViewportOffset_{ParameterName}(ScreenUV, bApplyViewportOffset);
+	FGBufferData GBufferData = GetGBufferData(ScreenUV);
+	IsValid = true;
+	OutValue = GBufferData.CustomStencil;
+#elif NDIGBUFFER_IS_MOBILE
+	IsValid = true;
+	#if MOBILE_SUPPORT_FETCH_BINDED_CUSTOM_STENCIL_BUFFER
+		OutValue = MobileSceneTextures.CustomStencilTexture.Load(int3(Parameters.SvPosition.xy, 0)).x;
+	#else
+		OutValue = Texture2DSampleLevel(MobileSceneTextures.MobileCustomStencilTexture, MobileSceneTextures.MobileCustomStencilTextureSampler, ScreenUV, 0).r * 255.0f;
+	#endif
+#else
+	IsValid = false;
+	OutValue = 0;
+#endif
+}
+
+void DecodeSceneColor_{ParameterName}(in float2 ScreenUV, in bool bApplyViewportOffset, out bool IsValid, out float4 OutValue)
+{
+#if NDIGBUFFER_HAS_GBUFFER
+	IsValid = true;
+	ScreenUV = ApplyViewportOffset_{ParameterName}(ScreenUV, bApplyViewportOffset);
+	OutValue = Texture2DSampleLevel(SceneTexturesStruct_SceneColorTexture, SceneTexturesStruct_SceneColorTextureSampler, ScreenUV, 0);
+	OutValue.rgb *= View.OneOverPreExposure.xxx;
+#elif NDIGBUFFER_IS_MOBILE
+	IsValid = true;
+	OutValue = Texture2DSampleLevel(MobileSceneTextures.SceneColorTexture, MobileSceneTextures.SceneColorTextureSampler, ScreenUV, 0);
+	OutValue.rgb *= View.OneOverPreExposure.xxx;
+#else
+	IsValid = false;
+	OutValue = 0;
+#endif
+}
+
+void DecodeShadingModelID_{ParameterName}(in float2 ScreenUV, in bool bApplyViewportOffset, out bool IsValid, out int OutValue)
+{
+#if NDIGBUFFER_HAS_GBUFFER
+	ScreenUV = ApplyViewportOffset_{ParameterName}(ScreenUV, bApplyViewportOffset);
+	FGBufferData GBufferData = GetGBufferData(ScreenUV);
+	IsValid = true;
+	OutValue = GBufferData.ShadingModelID;
+#else
+	IsValid = false;
+	OutValue = 0;
+#endif
+}
diff --git a/Engine/Plugins/FX/Niagara/Source/Niagara/Private/DataInterface/NiagaraDataInterfaceGBuffer.cpp b/Engine/Plugins/FX/Niagara/Source/Niagara/Private/DataInterface/NiagaraDataInterfaceGBuffer.cpp
--- a/Engine/Plugins/FX/Niagara/Source/Niagara/Private/DataInterface/NiagaraDataInterfaceGBuffer.cpp
+++ b/Engine/Plugins/FX/Niagara/Source/Niagara/Private/DataInterface/NiagaraDataInterfaceGBuffer.cpp
@@ -5,15 +5,24 @@
 #include "NiagaraGpuComputeDispatch.h"
 #include "NiagaraTypes.h"
 #include "NiagaraRenderViewDataManager.h"
+#include "NiagaraShaderParametersBuilder.h"
+#include "NiagaraSystemInstance.h"
 #include "NiagaraWorldManager.h"
-#include "ShaderParameterUtils.h"
+
 #include "Internationalization/Internationalization.h"
-#include "NiagaraSystemInstance.h"
+#include "ShaderParameterUtils.h"
 
 //////////////////////////////////////////////////////////////////////////
 
 namespace NiagaraDataInterfaceGBufferLocal
 {
+	BEGIN_SHADER_PARAMETER_STRUCT(FShaderParameters, )
+		SHADER_PARAMETER_TEXTURE(Texture2D, VelocityTexture)
+		SHADER_PARAMETER_SAMPLER(SamplerState, VelocityTextureSampler)
+	END_SHADER_PARAMETER_STRUCT()
+
+	static const TCHAR* TemplateShaderFile = TEXT("/Plugin/FX/Niagara/Private/NiagaraDataInterfaceGBufferTemplate.ush");
+
 	struct EDIFunctionVersion
 	{
 		enum Type
@@ -108,42 +117,6 @@ struct FNiagaraDataIntefaceProxyGBuffer : public FNiagaraDataInterfaceProxy
 	virtual int32 PerInstanceDataPassedToRenderThreadSize() const override { return 0; }
 };
 
-struct FNiagaraDataInterfaceParametersCS_GBuffer : public FNiagaraDataInterfaceParametersCS
-{
-	DECLARE_TYPE_LAYOUT(FNiagaraDataInterfaceParametersCS_GBuffer, NonVirtual);
-public:
-	void Bind(const FNiagaraDataInterfaceGPUParamInfo& ParameterInfo, const class FShaderParameterMap& ParameterMap)
-	{
-		VelocityTextureParam.Bind(ParameterMap, TEXT("NDIGBuffer_VelocityTexture"));
-		VelocityTextureSamplerParam.Bind(ParameterMap, TEXT("NDIGBuffer_VelocityTextureSampler"));
-	}
-
-	void Set(FRHICommandList& RHICmdList, const FNiagaraDataInterfaceSetArgs& Context) const
-	{
-		check(IsInRenderingThread());
-		FRHIComputeShader* ComputeShaderRHI = RHICmdList.GetBoundComputeShader();
-
-		FRHISamplerState* VelocitySamplerState = TStaticSamplerState<SF_Bilinear, AM_Clamp, AM_Clamp, AM_Clamp>::GetRHI();
-		FRHITexture* VelocityRHITexture = GBlackTexture->TextureRHI;
-		if ( FNiagaraSceneTextureParameters* NiagaraSceneTextures = static_cast<const FNiagaraGpuComputeDispatch*>(Context.ComputeDispatchInterface)->GetNiagaraSceneTextures() )	//-BATCHERTODO:
-		{
-			if ( FRDGTexture* VelocityRDGTexture = NiagaraSceneTextures->Velocity.GetTexture() )
-			{
-				VelocityRHITexture = VelocityRDGTexture->GetRHI();
-			}
-		}
-		SetTextureParameter(RHICmdList, ComputeShaderRHI, VelocityTextureParam, VelocityTextureSamplerParam, VelocitySamplerState, VelocityRHITexture);
-	}
-
-private:
-	LAYOUT_FIELD(FShaderResourceParameter, VelocityTextureParam);
-	LAYOUT_FIELD(FShaderResourceParameter, VelocityTextureSamplerParam);
-};
-
-IMPLEMENT_TYPE_LAYOUT(FNiagaraDataInterfaceParametersCS_GBuffer);
-
-IMPLEMENT_NIAGARA_DI_PARAMETER(UNiagaraDataInterfaceGBuffer, FNiagaraDataInterfaceParametersCS_GBuffer);
-
 //////////////////////////////////////////////////////////////////////////
 
 UNiagaraDataInterfaceGBuffer::UNiagaraDataInterfaceGBuffer(FObjectInitializer const& ObjectInitializer)
@@ -193,40 +166,57 @@ void UNiagaraDataInterfaceGBuffer::GetFunctions(TArray<FNiagaraFunctionSignature
 	}
 }
 
+void UNiagaraDataInterfaceGBuffer::BuildShaderParameters(FNiagaraShaderParametersBuilder& ShaderParametersBuilder) const
+{
+	ShaderParametersBuilder.AddNestedStruct<NiagaraDataInterfaceGBufferLocal::FShaderParameters>();
+}
+
+void UNiagaraDataInterfaceGBuffer::SetShaderParameters(const FNiagaraDataInterfaceSetShaderParametersContext& Context) const
+{
+	FRHITexture* VelocityTexture = GBlackTexture->TextureRHI;
+
+	if (FNiagaraSceneTextureParameters* NiagaraSceneTextures = static_cast<const FNiagaraGpuComputeDispatch&>(Context.GetComputeDispatchInterface()).GetNiagaraSceneTextures())	//-BATCHERTODO:
+	{
+		if (FRDGTexture* VelocityRDGTexture = NiagaraSceneTextures->Velocity.GetTexture())
+		{
+			VelocityTexture = VelocityRDGTexture->GetRHI();
+		}
+	}
+	NiagaraDataInterfaceGBufferLocal::FShaderParameters* Parameters = Context.GetParameterNestedStruct<NiagaraDataInterfaceGBufferLocal::FShaderParameters>();
+	Parameters->VelocityTexture = VelocityTexture;
+	Parameters->VelocityTextureSampler = TStaticSamplerState<SF_Bilinear, AM_Clamp, AM_Clamp, AM_Clamp>::GetRHI();
+}
+
 #if WITH_EDITORONLY_DATA
 bool UNiagaraDataInterfaceGBuffer::AppendCompileHash(FNiagaraCompileHashVisitor* InVisitor) const
 {
-	if (!Super::AppendCompileHash(InVisitor))
-		return false;
-
-	FSHAHash Hash = GetShaderFileHash((TEXT("/Plugin/FX/Niagara/Private/NiagaraDataInterfaceGBuffer.ush")), EShaderPlatform::SP_PCD3D_SM5);
-	InVisitor->UpdateString(TEXT("NiagaraDataInterfaceGBufferHLSLSource"), Hash.ToString());
-	return true;
+	bool bSuccess = Super::AppendCompileHash(InVisitor);
+	FSHAHash Hash = GetShaderFileHash(NiagaraDataInterfaceGBufferLocal::TemplateShaderFile, EShaderPlatform::SP_PCD3D_SM5);
+	InVisitor->UpdateString(TEXT("UNiagaraDataInterfaceGBufferTemplateHLSLSource"), Hash.ToString());
+	InVisitor->UpdateShaderParameters<NiagaraDataInterfaceGBufferLocal::FShaderParameters>();
+	return bSuccess;
 }
 
-void UNiagaraDataInterfaceGBuffer::GetCommonHLSL(FString& OutHLSL)
+void UNiagaraDataInterfaceGBuffer::GetParameterDefinitionHLSL(const FNiagaraDataInterfaceGPUParamInfo& ParamInfo, FString& OutHLSL)
 {
-	OutHLSL += TEXT("#include \"/Plugin/FX/Niagara/Private/NiagaraDataInterfaceGBuffer.ush\"\n");
+	TMap<FString, FStringFormatArg> TemplateArgs =
+	{
+		{TEXT("ParameterName"),	ParamInfo.DataInterfaceHLSLSymbol},
+	};
+
+	FString TemplateFile;
+	LoadShaderSourceFile(NiagaraDataInterfaceGBufferLocal::TemplateShaderFile, EShaderPlatform::SP_PCD3D_SM5, &TemplateFile, nullptr);
+	OutHLSL += FString::Format(*TemplateFile, TemplateArgs);
 }
 
 bool UNiagaraDataInterfaceGBuffer::GetFunctionHLSL(const FNiagaraDataInterfaceGPUParamInfo& ParamInfo, const FNiagaraDataInterfaceGeneratedFunction& FunctionInfo, int FunctionInstanceIndex, FString& OutHLSL)
 {
 	using namespace NiagaraDataInterfaceGBufferLocal;
 
-	TMap<FString, FStringFormatArg> ArgsSample =
-	{
-		{TEXT("InstanceFunctionName"), FunctionInfo.InstanceName},
-	};
-
 	for ( const FGBufferAttribute& Attribute : GetGBufferAttributes() )
 	{
 		if (FunctionInfo.DefinitionName == Attribute.ScreenUVFunctionName)
 		{
-			ArgsSample.Emplace(TEXT("AttributeName"), Attribute.AttributeName);
-			ArgsSample.Emplace(TEXT("AttributeType"), Attribute.AttributeType);
-
-			static const TCHAR* FormatSample = TEXT("void {InstanceFunctionName}(float2 ScreenUV, bool bApplyViewportOffset, out bool IsValid, out {AttributeType} {AttributeName}) { DIGBuffer_Decode{AttributeName}(ScreenUV, bApplyViewportOffset, IsValid, {AttributeName}); }\n");
-			OutHLSL += FString::Format(FormatSample, ArgsSample);
 			return true;
 		}
 	}
diff --git a/Engine/Plugins/FX/Niagara/Source/Niagara/Private/DataInterface/NiagaraDataInterfaceGBuffer.h b/Engine/Plugins/FX/Niagara/Source/Niagara/Private/DataInterface/NiagaraDataInterfaceGBuffer.h
--- a/Engine/Plugins/FX/Niagara/Source/Niagara/Private/DataInterface/NiagaraDataInterfaceGBuffer.h
+++ b/Engine/Plugins/FX/Niagara/Source/Niagara/Private/DataInterface/NiagaraDataInterfaceGBuffer.h
@@ -12,9 +12,6 @@ class NIAGARA_API UNiagaraDataInterfaceGBuffer : public UNiagaraDataInterface
 	GENERATED_UCLASS_BODY()
 
 public:
-
-	DECLARE_NIAGARA_DI_PARAMETER();
-
 	//UObject Interface
 	virtual void PostInitProperties() override;
 	//UObject Interface End
@@ -22,9 +19,13 @@ public:
 	//UNiagaraDataInterface Interface
 	virtual void GetFunctions(TArray<FNiagaraFunctionSignature>& OutFunctions)override;
 
+	virtual bool UseLegacyShaderBindings() const  override { return false; }
+	virtual void BuildShaderParameters(FNiagaraShaderParametersBuilder& ShaderParametersBuilder) const override;
+	virtual void SetShaderParameters(const FNiagaraDataInterfaceSetShaderParametersContext& Context) const override;
+
 #if WITH_EDITORONLY_DATA
 	virtual bool AppendCompileHash(FNiagaraCompileHashVisitor* InVisitor) const override;
-	virtual void GetCommonHLSL(FString& OutHLSL) override;
+	virtual void GetParameterDefinitionHLSL(const FNiagaraDataInterfaceGPUParamInfo& ParamInfo, FString& OutHLSL) override;
 	virtual bool GetFunctionHLSL(const FNiagaraDataInterfaceGPUParamInfo& ParamInfo, const FNiagaraDataInterfaceGeneratedFunction& FunctionInfo, int FunctionInstanceIndex, FString& OutHLSL) override;
 	virtual bool UpgradeFunctionCall(FNiagaraFunctionSignature& FunctionSignature);
 #endif
