- Convert Export data interface to use template & parameter structure binding

diff --git a/Engine/Plugins/FX/Niagara/Shaders/Private/NiagaraDataInterfaceExportTemplate.ush b/Engine/Plugins/FX/Niagara/Shaders/Private/NiagaraDataInterfaceExportTemplate.ush
--- a/Engine/Plugins/FX/Niagara/Shaders/Private/NiagaraDataInterfaceExportTemplate.ush
+++ b/Engine/Plugins/FX/Niagara/Shaders/Private/NiagaraDataInterfaceExportTemplate.ush
@@ -1,30 +1,30 @@
 // Copyright Epic Games, Inc. All Rights Reserved.
 
-uint			WriteBufferSize_{ParameterName};
-RWBuffer<uint>	RWWriteBuffer_{ParameterName};
+uint			{ParameterName}_WriteBufferSize;
+RWBuffer<uint>	{ParameterName}_RWWriteBuffer;
 
 void ExportParticleData_{ParameterName}_UEImpureCall(in bool bStoreData, in float3 Position, in float Size, in float3 Velocity, out bool bSuccess)
 {
 	bSuccess = false;
-	if ( bStoreData && (WriteBufferSize_{ParameterName} > 0) )
+	if ( bStoreData && ({ParameterName}_WriteBufferSize > 0) )
 	{
 		uint StoreIndex;
-		InterlockedAdd(RWWriteBuffer_{ParameterName}[0], 1u, StoreIndex);
-		if ( StoreIndex < WriteBufferSize_{ParameterName})
+		InterlockedAdd({ParameterName}_RWWriteBuffer[0], 1u, StoreIndex);
+		if ( StoreIndex < {ParameterName}_WriteBufferSize)
 		{
 			uint StoreOffset = 1 + (StoreIndex * 7);
-			RWWriteBuffer_{ParameterName}[StoreOffset + 0] = asuint(Position.x);
-			RWWriteBuffer_{ParameterName}[StoreOffset + 1] = asuint(Position.y);
-			RWWriteBuffer_{ParameterName}[StoreOffset + 2] = asuint(Position.z);
-			RWWriteBuffer_{ParameterName}[StoreOffset + 3] = asuint(Size);
-			RWWriteBuffer_{ParameterName}[StoreOffset + 4] = asuint(Velocity.x);
-			RWWriteBuffer_{ParameterName}[StoreOffset + 5] = asuint(Velocity.y);
-			RWWriteBuffer_{ParameterName}[StoreOffset + 6] = asuint(Velocity.z);
+			{ParameterName}_RWWriteBuffer[StoreOffset + 0] = asuint(Position.x);
+			{ParameterName}_RWWriteBuffer[StoreOffset + 1] = asuint(Position.y);
+			{ParameterName}_RWWriteBuffer[StoreOffset + 2] = asuint(Position.z);
+			{ParameterName}_RWWriteBuffer[StoreOffset + 3] = asuint(Size);
+			{ParameterName}_RWWriteBuffer[StoreOffset + 4] = asuint(Velocity.x);
+			{ParameterName}_RWWriteBuffer[StoreOffset + 5] = asuint(Velocity.y);
+			{ParameterName}_RWWriteBuffer[StoreOffset + 6] = asuint(Velocity.z);
 			bSuccess = true;
 		}
 		else
 		{
-			InterlockedAdd(RWWriteBuffer_{ParameterName}[0], -1u, StoreIndex);
+			InterlockedAdd({ParameterName}_RWWriteBuffer[0], -1u, StoreIndex);
 		}
 	}
 }
diff --git a/Engine/Plugins/FX/Niagara/Source/Niagara/Classes/NiagaraDataInterfaceExport.h b/Engine/Plugins/FX/Niagara/Source/Niagara/Classes/NiagaraDataInterfaceExport.h
--- a/Engine/Plugins/FX/Niagara/Source/Niagara/Classes/NiagaraDataInterfaceExport.h
+++ b/Engine/Plugins/FX/Niagara/Source/Niagara/Classes/NiagaraDataInterfaceExport.h
@@ -54,8 +54,6 @@ class NIAGARA_API UNiagaraDataInterfaceExport : public UNiagaraDataInterface
 	GENERATED_UCLASS_BODY()
 
 public:
-	DECLARE_NIAGARA_DI_PARAMETER();
-
 	/** Reference to a user parameter that should receive the particle data after the simulation tick. The supplied parameter object needs to implement the INiagaraParticleCallbackHandler interface. */
 	UPROPERTY(EditAnywhere, Category = "Export")
 	FNiagaraUserParameterBinding CallbackHandlerParameter;
@@ -89,6 +87,10 @@ public:
 	virtual bool GetFunctionHLSL(const FNiagaraDataInterfaceGPUParamInfo& ParamInfo, const FNiagaraDataInterfaceGeneratedFunction& FunctionInfo, int FunctionInstanceIndex, FString& OutHLSL) override;
 	virtual bool AppendCompileHash(FNiagaraCompileHashVisitor* InVisitor) const override;
 #endif
+	virtual bool UseLegacyShaderBindings() const  override { return false; }
+	virtual void BuildShaderParameters(FNiagaraShaderParametersBuilder& ShaderParametersBuilder) const override;
+	virtual void SetShaderParameters(const FNiagaraDataInterfaceSetShaderParametersContext& Context) const override;
+
 	virtual bool CanExecuteOnTarget(ENiagaraSimTarget Target) const override { return true; }
 
 	virtual bool HasPreSimulateTick() const override { return true; }
diff --git a/Engine/Plugins/FX/Niagara/Source/Niagara/Private/NiagaraDataInterface.cpp b/Engine/Plugins/FX/Niagara/Source/Niagara/Private/NiagaraDataInterface.cpp
--- a/Engine/Plugins/FX/Niagara/Source/Niagara/Private/NiagaraDataInterface.cpp
+++ b/Engine/Plugins/FX/Niagara/Source/Niagara/Private/NiagaraDataInterface.cpp
@@ -5,6 +5,7 @@
 #include "Curves/CurveLinearColor.h"
 #include "NiagaraTypes.h"
 #include "ShaderParameterUtils.h"
+#include "NiagaraGPUSystemTick.h"
 #include "NiagaraShader.h"
 #include "NiagaraComponent.h"
 #include "ShaderCompilerCore.h"
@@ -220,5 +221,32 @@ FSimpleMulticastDelegate& UNiagaraDataInterface::OnErrorsRefreshed()
 
 #endif
 
-#undef LOCTEXT_NAMESPACE
+FNiagaraSystemInstanceID FNiagaraDataInterfaceSetShaderParametersContext::GetSystemInstanceID() const
+{
+	return SystemTick.SystemInstanceID;
+}
+
+bool FNiagaraDataInterfaceSetShaderParametersContext::IsResourceBound(const void* ResourceAddress) const
+{
+	const uint16 ByteOffset = uint16(uintptr_t(ResourceAddress) - uintptr_t(BaseParameters));
+	for (const FShaderParameterBindings::FResourceParameter& ResourceParameter : ShaderRef->Bindings.ResourceParameters)
+	{
+		if ( ResourceParameter.ByteOffset == ByteOffset )
+		{
+			return true;
+		}
+	}
+	return false;
+}
+
+bool FNiagaraDataInterfaceSetShaderParametersContext::IsOutputStage() const
+{
+	return ComputeInstanceData.IsOutputStage(DataInterfaceProxy, SimStageData.StageIndex);
+}
 
+bool FNiagaraDataInterfaceSetShaderParametersContext::IsIterationStage() const
+{
+	return ComputeInstanceData.IsIterationStage(DataInterfaceProxy, SimStageData.StageIndex);
+}
+
+#undef LOCTEXT_NAMESPACE
diff --git a/Engine/Plugins/FX/Niagara/Source/Niagara/Private/NiagaraDataInterfaceExport.cpp b/Engine/Plugins/FX/Niagara/Source/Niagara/Private/NiagaraDataInterfaceExport.cpp
--- a/Engine/Plugins/FX/Niagara/Source/Niagara/Private/NiagaraDataInterfaceExport.cpp
+++ b/Engine/Plugins/FX/Niagara/Source/Niagara/Private/NiagaraDataInterfaceExport.cpp
@@ -7,22 +7,25 @@
 #include "NiagaraGpuComputeDispatchInterface.h"
 #include "NiagaraGpuReadbackManager.h"
 #include "NiagaraGPUSystemTick.h"
+#include "NiagaraShaderParametersBuilder.h"
 #include "NiagaraSystemInstance.h"
+
 #include "Internationalization/Internationalization.h"
-#include "ShaderParameterUtils.h"
 #include "ShaderCompilerCore.h"
+#include "ShaderParameterUtils.h"
 
 namespace NDIExportLocal
 {
+	BEGIN_SHADER_PARAMETER_STRUCT(FShaderParameters, )
+		SHADER_PARAMETER(uint32,				WriteBufferSize)
+		SHADER_PARAMETER_UAV(RWBuffer<uint>,	RWWriteBuffer)
+	END_SHADER_PARAMETER_STRUCT()
+
 	static const TCHAR*		TemplateShaderFile = TEXT("/Plugin/FX/Niagara/Private/NiagaraDataInterfaceExportTemplate.ush");
 
 	static const FName		StoreDataName_DEPRECATED(TEXT("StoreParticleData"));
 	static const FName		ExportDataName(TEXT("ExportParticleData"));
 
-	static const FString	WriteBufferSizeName(TEXT("WriteBufferSize_"));
-	static const FString	RWWriteBufferName(TEXT("RWWriteBuffer_"));
-	static const FString	WriteBufferName(TEXT("WriteBuffer_"));
-
 	static constexpr uint32 NumFloatsPerInstance = 7;
 
 	static int GGPUMaxReadbackCount = 1000;
@@ -219,65 +222,6 @@ struct FNDIExportProxy : public FNiagaraDataInterfaceProxy
 	TMap<FNiagaraSystemInstanceID, FNDIExportInstanceData_RenderThread> SystemInstancesToProxyData_RT;
 };
 
-//////////////////////////////////////////////////////////////////////////
-// Compute Shader Binding
-
-struct FNDIExportCS : public FNiagaraDataInterfaceParametersCS
-{
-	DECLARE_TYPE_LAYOUT(FNDIExportCS, NonVirtual);
-
-public:
-	void Bind(const FNiagaraDataInterfaceGPUParamInfo& ParameterInfo, const class FShaderParameterMap& ParameterMap)
-	{
-		WriteBufferSizeParam.Bind(ParameterMap, *(NDIExportLocal::WriteBufferSizeName + ParameterInfo.DataInterfaceHLSLSymbol));
-		WriteBufferParam.Bind(ParameterMap, *(NDIExportLocal::WriteBufferName + ParameterInfo.DataInterfaceHLSLSymbol));
-	}
-
-	void Set(FRHICommandList& RHICmdList, const FNiagaraDataInterfaceSetArgs& Context) const
-	{
-		check(IsInRenderingThread());
-
-		uint32 WriteBufferSize = 0;
-
-		FRHIComputeShader* ComputeShaderRHI = Context.Shader.GetComputeShader();
-		if (WriteBufferParam.IsUAVBound())
-		{
-			// Get proxy data
-			FNDIExportProxy* DIProxy = static_cast<FNDIExportProxy*>(Context.DataInterface);
-			FNDIExportInstanceData_RenderThread* InstanceData = DIProxy->SystemInstancesToProxyData_RT.Find(Context.SystemInstanceID);
-			check(InstanceData);
-
-			if (InstanceData->WeakCallbackHandler.IsExplicitlyNull())
-			{
-				WriteBufferSize = 0;
-				RHICmdList.SetUAVParameter(Context.Shader.GetComputeShader(), WriteBufferParam.GetUAVIndex(), Context.ComputeDispatchInterface->GetEmptyUAVFromPool(RHICmdList, PF_R32_UINT, ENiagaraEmptyUAVType::Buffer));
-			}
-			else
-			{
-				InstanceData->bWriteBufferUsed = true;
-				WriteBufferSize = InstanceData->WriteBufferInstanceCount;
-				RHICmdList.SetUAVParameter(Context.Shader.GetComputeShader(), WriteBufferParam.GetUAVIndex(), InstanceData->WriteBuffer.UAV);
-			}
-		}
-		SetShaderValue(RHICmdList, ComputeShaderRHI, WriteBufferSizeParam, WriteBufferSize);
-	}
-
-	void Unset(FRHICommandList& RHICmdList, const FNiagaraDataInterfaceSetArgs& Context) const
-	{
-		check(IsInRenderingThread());
-
-		WriteBufferParam.UnsetUAV(RHICmdList, Context.Shader.GetComputeShader());
-	}
-
-private:
-	LAYOUT_FIELD(FShaderParameter, WriteBufferSizeParam);
-	LAYOUT_FIELD(FRWShaderParameter, WriteBufferParam);
-};
-
-IMPLEMENT_TYPE_LAYOUT(FNDIExportCS);
-
-IMPLEMENT_NIAGARA_DI_PARAMETER(UNiagaraDataInterfaceExport, FNDIExportCS);
-
 //////////////////////////////////////////////////////////////////////////
 // Data Interface
 
@@ -492,6 +436,38 @@ bool UNiagaraDataInterfaceExport::AppendCompileHash(FNiagaraCompileHashVisitor*
 }
 #endif
 
+void UNiagaraDataInterfaceExport::BuildShaderParameters(FNiagaraShaderParametersBuilder& ShaderParametersBuilder) const
+{
+	ShaderParametersBuilder.AddNestedStruct<NDIExportLocal::FShaderParameters>();
+}
+
+void UNiagaraDataInterfaceExport::SetShaderParameters(const FNiagaraDataInterfaceSetShaderParametersContext& Context) const
+{
+	FNDIExportProxy& DIProxy = Context.GetProxy<FNDIExportProxy>();
+	FNDIExportInstanceData_RenderThread* InstanceData_RT = DIProxy.SystemInstancesToProxyData_RT.Find(Context.GetSystemInstanceID());
+
+	NDIExportLocal::FShaderParameters* Parameters = Context.GetParameterNestedStruct<NDIExportLocal::FShaderParameters>();
+	if ( Context.IsResourceBound(&Parameters->RWWriteBuffer) )
+	{
+		if ( (InstanceData_RT != nullptr) && !InstanceData_RT->WeakCallbackHandler.IsExplicitlyNull() )
+		{
+			InstanceData_RT->bWriteBufferUsed = true;
+			Parameters->WriteBufferSize = InstanceData_RT->WriteBufferInstanceCount;
+			Parameters->RWWriteBuffer = InstanceData_RT->WriteBuffer.UAV;
+		}
+		else
+		{
+			Parameters->WriteBufferSize = 0;
+			Parameters->RWWriteBuffer = Context.GetComputeDispatchInterface().GetEmptyUAVFromPool(Context.GetRHICmdList(), PF_R32_UINT, ENiagaraEmptyUAVType::Buffer);
+		}
+	}
+	else
+	{
+		Parameters->WriteBufferSize	= 0;
+		Parameters->RWWriteBuffer = nullptr;
+	}
+}
+
 DEFINE_NDI_DIRECT_FUNC_BINDER(UNiagaraDataInterfaceExport, StoreData);
 DEFINE_NDI_DIRECT_FUNC_BINDER(UNiagaraDataInterfaceExport, ExportData);
 void UNiagaraDataInterfaceExport::GetVMExternalFunction(const FVMExternalFunctionBindingInfo& BindingInfo, void* InstanceData, FVMExternalFunction &OutFunc)
