- Convert Curve data interface to use template & parameter structure binding

diff --git a/Engine/Plugins/FX/Niagara/Shaders/Private/NiagaraDataInterfaceCurveTemplate.ush b/Engine/Plugins/FX/Niagara/Shaders/Private/NiagaraDataInterfaceCurveTemplate.ush
new file mode 100644
--- /dev/null
+++ b/Engine/Plugins/FX/Niagara/Shaders/Private/NiagaraDataInterfaceCurveTemplate.ush
@@ -0,0 +1,59 @@
+// Copyright Epic Games, Inc. All Rights Reserved.
+
+float			{ParameterName}_MinTime;
+float			{ParameterName}_MaxTime;
+float			{ParameterName}_InvTimeRange;
+uint			{ParameterName}_CurveLUTNumMinusOne;
+uint			{ParameterName}_LUTOffset;
+Buffer<float>	{ParameterName}_CurveLUT;
+
+void GetCurveLUTIndices_{ParameterName}(float Time, out uint IndexA, out uint IndexB, out float Fraction)
+{
+	float RemappedTime = saturate((Time - {ParameterName}_MinTime) * {ParameterName}_InvTimeRange) * float({ParameterName}_CurveLUTNumMinusOne);
+	IndexA = floor(RemappedTime);
+	IndexB = min(IndexA + 1, {ParameterName}_CurveLUTNumMinusOne);
+	Fraction = frac(RemappedTime);
+}
+
+float SampleCurveLUT_{ParameterName}(uint index)
+{
+#if {UseStaticBuffer} == 1
+	return {ParameterName}_CurveLUT[index];
+#else
+	return StaticInputFloat[{ParameterName}_LUTOffset + index];
+#endif
+}
+
+#if {NumElements} == 1
+void {CurveSampleFunctionName}_{ParameterName}(float Time, out float Value)
+{
+	uint IndexA, IndexB;
+	float Fraction;
+	GetCurveLUTIndices_{ParameterName}(Time, IndexA, IndexB, Fraction);
+
+	IndexA *= {NumElements};
+	IndexB *= {NumElements};
+
+	[unroll]
+	for ( int i=0; i < {NumElements}; ++i )
+	{
+		Value = lerp(SampleCurveLUT_{ParameterName}(IndexA + i), SampleCurveLUT_{ParameterName}(IndexB + i), Fraction);
+	}
+}
+#else
+void {CurveSampleFunctionName}_{ParameterName}(float Time, out float{NumElements} Value)
+{
+	uint IndexA, IndexB;
+	float Fraction;
+	GetCurveLUTIndices_{ParameterName}(Time, IndexA, IndexB, Fraction);
+
+	IndexA *= {NumElements};
+	IndexB *= {NumElements};
+
+	[unroll]
+	for ( int i=0; i < {NumElements}; ++i )
+	{
+		Value[i] = lerp(SampleCurveLUT_{ParameterName}(IndexA + i), SampleCurveLUT_{ParameterName}(IndexB + i), Fraction);
+	}
+}
+#endif
diff --git a/Engine/Plugins/FX/Niagara/Source/Niagara/Classes/NiagaraDataInterfaceColorCurve.h b/Engine/Plugins/FX/Niagara/Source/Niagara/Classes/NiagaraDataInterfaceColorCurve.h
--- a/Engine/Plugins/FX/Niagara/Source/Niagara/Classes/NiagaraDataInterfaceColorCurve.h
+++ b/Engine/Plugins/FX/Niagara/Source/Niagara/Classes/NiagaraDataInterfaceColorCurve.h
@@ -62,7 +62,11 @@ public:
 	virtual bool GetFunctionHLSL(const FNiagaraDataInterfaceGPUParamInfo& ParamInfo, const FNiagaraDataInterfaceGeneratedFunction& FunctionInfo, int FunctionInstanceIndex, FString& OutHLSL) override;
 #endif
 
-	virtual int32 GetCurveNumElems()const { return CurveLUTNumElems; }
+	virtual int32 GetCurveNumElems() const override { return CurveLUTNumElems; }
+#if WITH_EDITORONLY_DATA
+	virtual FName GetCurveSampleFunctionName() const override { return SampleCurveName; }
+#endif
+
 protected:
 	virtual bool CopyToInternal(UNiagaraDataInterface* Destination) const override;
 
diff --git a/Engine/Plugins/FX/Niagara/Source/Niagara/Classes/NiagaraDataInterfaceCurve.h b/Engine/Plugins/FX/Niagara/Source/Niagara/Classes/NiagaraDataInterfaceCurve.h
--- a/Engine/Plugins/FX/Niagara/Source/Niagara/Classes/NiagaraDataInterfaceCurve.h
+++ b/Engine/Plugins/FX/Niagara/Source/Niagara/Classes/NiagaraDataInterfaceCurve.h
@@ -48,7 +48,11 @@ public:
 	virtual bool GetFunctionHLSL(const FNiagaraDataInterfaceGPUParamInfo& ParamInfo, const FNiagaraDataInterfaceGeneratedFunction& FunctionInfo, int FunctionInstanceIndex, FString& OutHLSL) override;
 #endif
 
-	virtual int32 GetCurveNumElems()const { return CurveLUTNumElems; }
+	virtual int32 GetCurveNumElems() const override { return CurveLUTNumElems; }
+#if WITH_EDITORONLY_DATA
+	virtual FName GetCurveSampleFunctionName() const override { return SampleCurveName; }
+#endif
+
 protected:
 	virtual bool CopyToInternal(UNiagaraDataInterface* Destination) const override;
 
diff --git a/Engine/Plugins/FX/Niagara/Source/Niagara/Classes/NiagaraDataInterfaceCurveBase.h b/Engine/Plugins/FX/Niagara/Source/Niagara/Classes/NiagaraDataInterfaceCurveBase.h
--- a/Engine/Plugins/FX/Niagara/Source/Niagara/Classes/NiagaraDataInterfaceCurveBase.h
+++ b/Engine/Plugins/FX/Niagara/Source/Niagara/Classes/NiagaraDataInterfaceCurveBase.h
@@ -23,7 +23,7 @@ struct FNiagaraDataInterfaceProxyCurveBase : public FNiagaraDataInterfaceProxy
 	float LUTMinTime = 0.0f;
 	float LUTMaxTime = 0.0f;
 	float LUTInvTimeRange = 0.0f;
-	float CurveLUTNumMinusOne = 0.0f;
+	uint32 CurveLUTNumMinusOne = 0;
 	uint32 LUTOffset = INDEX_NONE;
 	FReadBuffer CurveLUT;
 
@@ -76,8 +76,6 @@ protected:
 	}
 
 public:
-	DECLARE_NIAGARA_DI_PARAMETER();
-
 	UPROPERTY(EditAnywhere, AdvancedDisplay, Category = "Curve")
 	uint32 bUseLUT : 1;
 
@@ -184,8 +182,12 @@ public:
 
 	virtual void CacheStaticBuffers(struct FNiagaraSystemStaticBuffers& StaticBuffers, const struct FNiagaraScriptDataInterfaceInfo& DataInterfaceInfo, bool bUsedByCPU, bool bUsedByGPU) override;
 #if WITH_EDITORONLY_DATA
+	virtual bool AppendCompileHash(FNiagaraCompileHashVisitor* InVisitor) const override;
 	virtual void GetParameterDefinitionHLSL(const FNiagaraDataInterfaceGPUParamInfo& ParamInfo, FString& OutHLSL) override;
 #endif
+	virtual bool UseLegacyShaderBindings() const  override { return false; }
+	virtual void BuildShaderParameters(FNiagaraShaderParametersBuilder& ShaderParametersBuilder) const override;
+	virtual void SetShaderParameters(const FNiagaraDataInterfaceSetShaderParametersContext& Context) const override;
 
 	void SetDefaultLUT();
 #if WITH_EDITORONLY_DATA
@@ -202,7 +204,10 @@ public:
 	virtual void GetExposedVariables(TArray<FNiagaraVariableBase>& OutVariables) const override;
 	virtual bool GetExposedVariableValue(const FNiagaraVariableBase& InVariable, void* InPerInstanceData, FNiagaraSystemInstance* InSystemInstance, void* OutData) const override;
 
-	virtual int32 GetCurveNumElems() const { checkf(false, TEXT("You must implement this function in your derived class")); return 0; }
+	virtual int32 GetCurveNumElems() const PURE_VIRTUAL(UNiagaraDataInterfaceCurveBase::GetCurveNumElems(), return 0;)
+#if WITH_EDITORONLY_DATA
+	virtual FName GetCurveSampleFunctionName() const PURE_VIRTUAL(UNiagaraDataInterfaceCurveBase::GetCurveSampleFunctionName(), return NAME_None;)
+#endif
 
 	virtual void UpdateTimeRanges() { checkf(false, TEXT("You must implement this function in your derived class")); }
 	virtual TArray<float> BuildLUT(int32 NumEntries) const { checkf(false, TEXT("You must implement this function in your derived class")); return TArray<float>(); }
@@ -236,18 +241,3 @@ struct TCurveUseLUTBinder
 		}
 	}
 };
-
-struct FNiagaraDataInterfaceParametersCS_Curve : public FNiagaraDataInterfaceParametersCS
-{
-	DECLARE_TYPE_LAYOUT(FNiagaraDataInterfaceParametersCS_Curve, NonVirtual);
-
-	void Bind(const FNiagaraDataInterfaceGPUParamInfo& ParameterInfo, const class FShaderParameterMap& ParameterMap);
-	void Set(FRHICommandList& RHICmdList, const FNiagaraDataInterfaceSetArgs& Context) const;
-
-	LAYOUT_FIELD(FShaderParameter, MinTime);
-	LAYOUT_FIELD(FShaderParameter, MaxTime);
-	LAYOUT_FIELD(FShaderParameter, InvTimeRange);
-	LAYOUT_FIELD(FShaderParameter, CurveLUTNumMinusOne);
-	LAYOUT_FIELD(FShaderParameter, LUTOffset);
-	LAYOUT_FIELD(FShaderResourceParameter, CurveLUT);
-};
diff --git a/Engine/Plugins/FX/Niagara/Source/Niagara/Classes/NiagaraDataInterfaceVector2DCurve.h b/Engine/Plugins/FX/Niagara/Source/Niagara/Classes/NiagaraDataInterfaceVector2DCurve.h
--- a/Engine/Plugins/FX/Niagara/Source/Niagara/Classes/NiagaraDataInterfaceVector2DCurve.h
+++ b/Engine/Plugins/FX/Niagara/Source/Niagara/Classes/NiagaraDataInterfaceVector2DCurve.h
@@ -53,7 +53,11 @@ public:
 	virtual bool GetFunctionHLSL(const FNiagaraDataInterfaceGPUParamInfo& ParamInfo, const FNiagaraDataInterfaceGeneratedFunction& FunctionInfo, int FunctionInstanceIndex, FString& OutHLSL) override;
 #endif
 
-	virtual int32 GetCurveNumElems()const { return CurveLUTNumElems; }
+	virtual int32 GetCurveNumElems() const override { return CurveLUTNumElems; }
+#if WITH_EDITORONLY_DATA
+	virtual FName GetCurveSampleFunctionName() const override { return SampleCurveName; }
+#endif
+
 protected:
 	virtual bool CopyToInternal(UNiagaraDataInterface* Destination) const override;
 
diff --git a/Engine/Plugins/FX/Niagara/Source/Niagara/Classes/NiagaraDataInterfaceVector4Curve.h b/Engine/Plugins/FX/Niagara/Source/Niagara/Classes/NiagaraDataInterfaceVector4Curve.h
--- a/Engine/Plugins/FX/Niagara/Source/Niagara/Classes/NiagaraDataInterfaceVector4Curve.h
+++ b/Engine/Plugins/FX/Niagara/Source/Niagara/Classes/NiagaraDataInterfaceVector4Curve.h
@@ -61,7 +61,11 @@ public:
 	virtual bool GetFunctionHLSL(const FNiagaraDataInterfaceGPUParamInfo& ParamInfo, const FNiagaraDataInterfaceGeneratedFunction& FunctionInfo, int FunctionInstanceIndex, FString& OutHLSL) override;
 #endif
 
-	virtual int32 GetCurveNumElems()const { return CurveLUTNumElems; }
+	virtual int32 GetCurveNumElems() const override { return CurveLUTNumElems; }
+#if WITH_EDITORONLY_DATA
+	virtual FName GetCurveSampleFunctionName() const override { return SampleCurveName; }
+#endif
+
 protected:
 	virtual bool CopyToInternal(UNiagaraDataInterface* Destination) const override;
 
diff --git a/Engine/Plugins/FX/Niagara/Source/Niagara/Classes/NiagaraDataInterfaceVectorCurve.h b/Engine/Plugins/FX/Niagara/Source/Niagara/Classes/NiagaraDataInterfaceVectorCurve.h
--- a/Engine/Plugins/FX/Niagara/Source/Niagara/Classes/NiagaraDataInterfaceVectorCurve.h
+++ b/Engine/Plugins/FX/Niagara/Source/Niagara/Classes/NiagaraDataInterfaceVectorCurve.h
@@ -55,7 +55,10 @@ public:
 	//~ UNiagaraDataInterfaceCurveBase interface
 	virtual void GetCurveData(TArray<FCurveData>& OutCurveData) override;
 	
-	virtual int32 GetCurveNumElems()const { return CurveLUTNumElems; }
+	virtual int32 GetCurveNumElems() const override { return CurveLUTNumElems; }
+#if WITH_EDITORONLY_DATA
+	virtual FName GetCurveSampleFunctionName() const override { return SampleCurveName; }
+#endif
 
 #if WITH_EDITORONLY_DATA
 	virtual bool GetFunctionHLSL(const FNiagaraDataInterfaceGPUParamInfo& ParamInfo, const FNiagaraDataInterfaceGeneratedFunction& FunctionInfo, int FunctionInstanceIndex, FString& OutHLSL) override;
diff --git a/Engine/Plugins/FX/Niagara/Source/Niagara/Private/NiagaraDataInterfaceColorCurve.cpp b/Engine/Plugins/FX/Niagara/Source/Niagara/Private/NiagaraDataInterfaceColorCurve.cpp
--- a/Engine/Plugins/FX/Niagara/Source/Niagara/Private/NiagaraDataInterfaceColorCurve.cpp
+++ b/Engine/Plugins/FX/Niagara/Source/Niagara/Private/NiagaraDataInterfaceColorCurve.cpp
@@ -176,26 +176,11 @@ void UNiagaraDataInterfaceColorCurve::GetFunctions(TArray<FNiagaraFunctionSignat
 #if WITH_EDITORONLY_DATA
 bool UNiagaraDataInterfaceColorCurve::GetFunctionHLSL(const FNiagaraDataInterfaceGPUParamInfo& ParamInfo, const FNiagaraDataInterfaceGeneratedFunction& FunctionInfo, int FunctionInstanceIndex, FString& OutHLSL)
 {
-	FString TimeToLUTFrac = TEXT("TimeToLUTFraction_") + ParamInfo.DataInterfaceHLSLSymbol;
-	FString Sample = TEXT("SampleCurve_") + ParamInfo.DataInterfaceHLSLSymbol;
-	FString NumSamples = TEXT("CurveLUTNumMinusOne_") + ParamInfo.DataInterfaceHLSLSymbol;
-	OutHLSL += FString::Printf(TEXT("\
-void %s(in float In_X, out float4 Out_Value) \n\
-{ \n\
-	float RemappedX = %s(In_X) * %s; \n\
-	float Prev = floor(RemappedX); \n\
-	float Next = Prev < %s ? Prev + 1.0 : Prev; \n\
-	float Interp = RemappedX - Prev; \n\
-	Prev *= %u; \n\
-	Next *= %u; \n\
-	float4 A = float4(%s(Prev), %s(Prev + 1), %s(Prev + 2), %s(Prev + 3)); \n\
-	float4 B = float4(%s(Next), %s(Next + 1), %s(Next + 2), %s(Next + 3)); \n\
-	Out_Value = lerp(A, B, Interp); \n\
-}\n")
-, *FunctionInfo.InstanceName, *TimeToLUTFrac, *NumSamples, *NumSamples, CurveLUTNumElems, CurveLUTNumElems
-, *Sample, *Sample, *Sample, *Sample, *Sample, *Sample, *Sample, *Sample);
-
-	return true;
+	if (FunctionInfo.DefinitionName == SampleCurveName)
+	{
+		return true;
+	}
+	return false;
 }
 #endif
 
diff --git a/Engine/Plugins/FX/Niagara/Source/Niagara/Private/NiagaraDataInterfaceCurve.cpp b/Engine/Plugins/FX/Niagara/Source/Niagara/Private/NiagaraDataInterfaceCurve.cpp
--- a/Engine/Plugins/FX/Niagara/Source/Niagara/Private/NiagaraDataInterfaceCurve.cpp
+++ b/Engine/Plugins/FX/Niagara/Source/Niagara/Private/NiagaraDataInterfaceCurve.cpp
@@ -140,23 +140,11 @@ TArray<float> UNiagaraDataInterfaceCurve::BuildLUT(int32 NumEntries) const
 #if WITH_EDITORONLY_DATA
 bool UNiagaraDataInterfaceCurve::GetFunctionHLSL(const FNiagaraDataInterfaceGPUParamInfo& ParamInfo, const FNiagaraDataInterfaceGeneratedFunction& FunctionInfo, int FunctionInstanceIndex, FString& OutHLSL)
 {
-	FString TimeToLUTFrac = TEXT("TimeToLUTFraction_") + ParamInfo.DataInterfaceHLSLSymbol;
-	FString Sample = TEXT("SampleCurve_") + ParamInfo.DataInterfaceHLSLSymbol;
-	FString NumSamples = TEXT("CurveLUTNumMinusOne_") + ParamInfo.DataInterfaceHLSLSymbol;
-	OutHLSL += FString::Printf(TEXT("\
-void %s(in float In_X, out float Out_Value) \n\
-{ \n\
-	float RemappedX = %s(In_X) * %s; \n\
-	float Prev = floor(RemappedX); \n\
-	float Next = Prev < %s ? Prev + 1.0 : Prev; \n\
-	float Interp = RemappedX - Prev; \n\
-	float A = %s(Prev); \n\
-	float B = %s(Next); \n\
-	Out_Value = lerp(A, B, Interp); \n\
-}\n")
-, *FunctionInfo.InstanceName, *TimeToLUTFrac, *NumSamples, *NumSamples, *Sample, *Sample);
-
-	return true;
+	if (FunctionInfo.DefinitionName == SampleCurveName)
+	{
+		return true;
+	}
+	return false;
 }
 #endif
 
diff --git a/Engine/Plugins/FX/Niagara/Source/Niagara/Private/NiagaraDataInterfaceCurveBase.cpp b/Engine/Plugins/FX/Niagara/Source/Niagara/Private/NiagaraDataInterfaceCurveBase.cpp
--- a/Engine/Plugins/FX/Niagara/Source/Niagara/Private/NiagaraDataInterfaceCurveBase.cpp
+++ b/Engine/Plugins/FX/Niagara/Source/Niagara/Private/NiagaraDataInterfaceCurveBase.cpp
@@ -12,42 +12,12 @@
 #include "NiagaraEmitter.h"
 #include "NiagaraRenderer.h"
 #include "NiagaraShader.h"
+#include "NiagaraShaderParametersBuilder.h"
 #include "NiagaraSystemStaticBuffers.h"
 #include "NiagaraTypes.h"
 
 #define LOCTEXT_NAMESPACE "NiagaraDataInterfaceCurveBase"
 
-IMPLEMENT_TYPE_LAYOUT(FNiagaraDataInterfaceParametersCS_Curve);
-
-IMPLEMENT_NIAGARA_DI_PARAMETER(UNiagaraDataInterfaceCurveBase, FNiagaraDataInterfaceParametersCS_Curve);
-
-void FNiagaraDataInterfaceParametersCS_Curve::Bind(const FNiagaraDataInterfaceGPUParamInfo& ParameterInfo, const class FShaderParameterMap& ParameterMap)
-{
-	MinTime.Bind(ParameterMap, *(TEXT("MinTime_") + ParameterInfo.DataInterfaceHLSLSymbol));
-	MaxTime.Bind(ParameterMap, *(TEXT("MaxTime_") + ParameterInfo.DataInterfaceHLSLSymbol));
-	InvTimeRange.Bind(ParameterMap, *(TEXT("InvTimeRange_") + ParameterInfo.DataInterfaceHLSLSymbol));
-	CurveLUTNumMinusOne.Bind(ParameterMap, *(TEXT("CurveLUTNumMinusOne_") + ParameterInfo.DataInterfaceHLSLSymbol));
-	LUTOffset.Bind(ParameterMap, *(TEXT("LUTOffset_") + ParameterInfo.DataInterfaceHLSLSymbol));
-	CurveLUT.Bind(ParameterMap, *(TEXT("CurveLUT_") + ParameterInfo.DataInterfaceHLSLSymbol));
-}
-
-void FNiagaraDataInterfaceParametersCS_Curve::Set(FRHICommandList& RHICmdList, const FNiagaraDataInterfaceSetArgs& Context) const
-{
-	check(IsInRenderingThread());
-
-	FRHIComputeShader* ComputeShaderRHI = RHICmdList.GetBoundComputeShader();
-	FNiagaraDataInterfaceProxyCurveBase* CurveDI = (FNiagaraDataInterfaceProxyCurveBase*)Context.DataInterface;
-
-	ensure((CurveDI->LUTOffset != INDEX_NONE) || CurveDI->CurveLUT.SRV.IsValid());
-
-	SetShaderValue(RHICmdList, ComputeShaderRHI, MinTime, CurveDI->LUTMinTime);
-	SetShaderValue(RHICmdList, ComputeShaderRHI, MaxTime, CurveDI->LUTMaxTime);
-	SetShaderValue(RHICmdList, ComputeShaderRHI, InvTimeRange, CurveDI->LUTInvTimeRange);
-	SetShaderValue(RHICmdList, ComputeShaderRHI, CurveLUTNumMinusOne, CurveDI->CurveLUTNumMinusOne);
-	SetShaderValue(RHICmdList, ComputeShaderRHI, LUTOffset, CurveDI->LUTOffset);
-	SetSRVParameter(RHICmdList, ComputeShaderRHI, CurveLUT, CurveDI->CurveLUT.SRV.IsValid() ? CurveDI->CurveLUT.SRV.GetReference() : FNiagaraRenderer::GetDummyFloatBuffer());
-}
-
 float GNiagaraLUTOptimizeThreshold = UNiagaraDataInterfaceCurveBase::DefaultOptimizeThreshold;
 static FAutoConsoleVariableRef CVarNiagaraLUTOptimizeThreshold(
 	TEXT("fx.Niagara.LUT.OptimizeThreshold"),
@@ -64,6 +34,17 @@ static FAutoConsoleVariableRef CVarNiagaraLUTVerifyPostLoad(
 	ECVF_Default
 );
 
+BEGIN_SHADER_PARAMETER_STRUCT(FNiagaraDataInterfaceCurveParameters, )
+	SHADER_PARAMETER(float, MinTime)
+	SHADER_PARAMETER(float, MaxTime)
+	SHADER_PARAMETER(float, InvTimeRange)
+	SHADER_PARAMETER(uint32, CurveLUTNumMinusOne)
+	SHADER_PARAMETER(uint32, LUTOffset)
+	SHADER_PARAMETER_SRV(Buffer<float>, CurveLUT)
+END_SHADER_PARAMETER_STRUCT()
+
+static const TCHAR* NDICurve_TemplateShaderFile = TEXT("/Plugin/FX/Niagara/Private/NiagaraDataInterfaceCurveTemplate.ush");
+
 void UNiagaraDataInterfaceCurveBase::PostLoad()
 {
 	Super::PostLoad();
@@ -175,7 +156,7 @@ void UNiagaraDataInterfaceCurveBase::UpdateLUT(bool bFromSerialize)
 	{
 		ShaderLUT = BuildLUT(CurveLUTDefaultWidth);
 		OptimizeLUT();
-		LUTNumSamplesMinusOne = float(ShaderLUT.Num() / GetCurveNumElems()) - 1;
+		LUTNumSamplesMinusOne = (ShaderLUT.Num() / GetCurveNumElems()) - 1;
 	}
 
 	if (!bUseLUT || (ShaderLUT.Num() == 0))
@@ -343,66 +324,52 @@ void UNiagaraDataInterfaceCurveBase::CacheStaticBuffers(struct FNiagaraSystemSta
 	}
 }
 
+void UNiagaraDataInterfaceCurveBase::BuildShaderParameters(FNiagaraShaderParametersBuilder& ShaderParametersBuilder) const
+{
+	ShaderParametersBuilder.AddNestedStruct<FNiagaraDataInterfaceCurveParameters>();
+}
+
+void UNiagaraDataInterfaceCurveBase::SetShaderParameters(const FNiagaraDataInterfaceSetShaderParametersContext& Context) const
+{
+	const FNiagaraDataInterfaceProxyCurveBase& DIProxy = Context.GetProxy<FNiagaraDataInterfaceProxyCurveBase>();
+	FNiagaraDataInterfaceCurveParameters* Parameters = Context.GetParameterNestedStruct<FNiagaraDataInterfaceCurveParameters>();
+	Parameters->MinTime = DIProxy.LUTMinTime;
+	Parameters->MaxTime = DIProxy.LUTMaxTime;
+	Parameters->InvTimeRange = DIProxy.LUTInvTimeRange;
+	Parameters->CurveLUTNumMinusOne = DIProxy.CurveLUTNumMinusOne;
+	Parameters->LUTOffset = DIProxy.LUTOffset;
+	Parameters->CurveLUT = DIProxy.CurveLUT.SRV.IsValid() ? DIProxy.CurveLUT.SRV.GetReference() : FNiagaraRenderer::GetDummyFloatBuffer();
+}
+
 #if WITH_EDITORONLY_DATA
-void UNiagaraDataInterfaceCurveBase::GetParameterDefinitionHLSL(const FNiagaraDataInterfaceGPUParamInfo& ParamInfo, FString& OutHLSL)
+bool UNiagaraDataInterfaceCurveBase::AppendCompileHash(FNiagaraCompileHashVisitor* InVisitor) const
 {
-	OutHLSL += TEXT("\n");
-	FString MinTimeStr = TEXT("MinTime_") + ParamInfo.DataInterfaceHLSLSymbol;
-	FString MaxTimeStr = TEXT("MaxTime_") + ParamInfo.DataInterfaceHLSLSymbol;
-	FString InvTimeRangeStr = TEXT("InvTimeRange_") + ParamInfo.DataInterfaceHLSLSymbol;
-	FString CurveLUTNumMinusOneStr = TEXT("CurveLUTNumMinusOne_") + ParamInfo.DataInterfaceHLSLSymbol;
-	FString LUTOffsetStr = TEXT("LUTOffset_") + ParamInfo.DataInterfaceHLSLSymbol;
-
-	FString BufferName;
-	if (ParamInfo.IsUserParameter() == true)
-	{
-		BufferName = "CurveLUT_" + ParamInfo.DataInterfaceHLSLSymbol;
-		OutHLSL += TEXT("Buffer<float> ") + BufferName + TEXT(";\n");
-	}
-	else
-	{
-		BufferName = "StaticInputFloat";
-	}
+	bool bSuccess = Super::AppendCompileHash(InVisitor);
+	FSHAHash Hash = GetShaderFileHash(NDICurve_TemplateShaderFile, EShaderPlatform::SP_PCD3D_SM5);
+	InVisitor->UpdateString(TEXT("UNiagaraDataInterfaceCurveTemplateHLSLSource"), Hash.ToString());
+	InVisitor->UpdateShaderParameters<FNiagaraDataInterfaceCurveParameters>();
+	return bSuccess;
+}
 
-	OutHLSL += TEXT("float ") + MinTimeStr + TEXT(";\n");
-	OutHLSL += TEXT("float ") + MaxTimeStr + TEXT(";\n");
-	OutHLSL += TEXT("float ") + InvTimeRangeStr + TEXT(";\n");
-	OutHLSL += TEXT("float ") + CurveLUTNumMinusOneStr + TEXT(";\n");
-	OutHLSL += TEXT("uint ") + LUTOffsetStr + TEXT(";\n");
-	OutHLSL += TEXT("\n");
+void UNiagaraDataInterfaceCurveBase::GetParameterDefinitionHLSL(const FNiagaraDataInterfaceGPUParamInfo& ParamInfo, FString& OutHLSL)
+{
+	const int32 NumElements = GetCurveNumElems();
+	check(NumElements > 0 && NumElements <= 4);
 
-	//TODO: Create a Unitiliy/Common funcitons hlsl def shared between all instances of the same data interface class for these.
-	OutHLSL += FString::Printf(TEXT("float TimeToLUTFraction_%s(float T)\n{\n\treturn saturate((T - %s) * %s);\n}\n"), *ParamInfo.DataInterfaceHLSLSymbol, *MinTimeStr, *InvTimeRangeStr);
-	if (ParamInfo.IsUserParameter() == true)
-	{
-		OutHLSL += FString::Printf(TEXT("float SampleCurve_%s(float T)\n{\n\treturn %s[(uint)T];\n}\n"), *ParamInfo.DataInterfaceHLSLSymbol, *BufferName);
-	}
-	else
+	TMap<FString, FStringFormatArg> TemplateArgs =
 	{
-		OutHLSL += FString::Printf(TEXT("float SampleCurve_%s(float T)\n{\n\treturn %s[%s + (uint)T];\n}\n"), *ParamInfo.DataInterfaceHLSLSymbol, *BufferName, *LUTOffsetStr);
-	}
-	OutHLSL += TEXT("\n");
+		{TEXT("ParameterName"),				ParamInfo.DataInterfaceHLSLSymbol},
+		{TEXT("UseStaticBuffer"),			ParamInfo.IsUserParameter() ? TEXT("1") : TEXT("0")},
+		{TEXT("NumElements"),				FString::FromInt(NumElements)},
+		{TEXT("CurveSampleFunctionName"),	GetCurveSampleFunctionName().ToString()},
+	};
+
+	FString TemplateFile;
+	LoadShaderSourceFile(NDICurve_TemplateShaderFile, EShaderPlatform::SP_PCD3D_SM5, &TemplateFile, nullptr);
+	OutHLSL += FString::Format(*TemplateFile, TemplateArgs);
 }
 #endif
 
-//FReadBuffer& UNiagaraDataInterfaceCurveBase::GetCurveLUTGPUBuffer()
-//{
-//	//TODO: This isn't really very thread safe. Need to move to a proxy like system where DIs can push data to the RT safely.
-//	if (GPUBufferDirty)
-//	{
-//		int32 ElemSize = GetCurveNumElems();
-//		CurveLUT.Release();
-//		CurveLUT.Initialize(sizeof(float), CurveLUTWidth * ElemSize, EPixelFormat::PF_R32_FLOAT, BUF_Static);
-//		uint32 BufferSize = ShaderLUT.Num() * sizeof(float);
-//		int32 *BufferData = static_cast<int32*>(RHILockBuffer(CurveLUT.Buffer, 0, BufferSize, EResourceLockMode::RLM_WriteOnly));
-//		FPlatformMemory::Memcpy(BufferData, ShaderLUT.GetData(), BufferSize);
-//		RHIUnlockBuffer(CurveLUT.Buffer);
-//		GPUBufferDirty = false;
-//	}
-//
-//	return CurveLUT;
-//}
-
 void UNiagaraDataInterfaceCurveBase::GetExposedVariables(TArray<FNiagaraVariableBase>& OutVariables) const
 {
 	if (bExposeCurve)
diff --git a/Engine/Plugins/FX/Niagara/Source/Niagara/Private/NiagaraDataInterfaceVector2DCurve.cpp b/Engine/Plugins/FX/Niagara/Source/Niagara/Private/NiagaraDataInterfaceVector2DCurve.cpp
--- a/Engine/Plugins/FX/Niagara/Source/Niagara/Private/NiagaraDataInterfaceVector2DCurve.cpp
+++ b/Engine/Plugins/FX/Niagara/Source/Niagara/Private/NiagaraDataInterfaceVector2DCurve.cpp
@@ -155,25 +155,11 @@ void UNiagaraDataInterfaceVector2DCurve::GetFunctions(TArray<FNiagaraFunctionSig
 #if WITH_EDITORONLY_DATA
 bool UNiagaraDataInterfaceVector2DCurve::GetFunctionHLSL(const FNiagaraDataInterfaceGPUParamInfo& ParamInfo, const FNiagaraDataInterfaceGeneratedFunction& FunctionInfo, int FunctionInstanceIndex, FString& OutHLSL)
 {
-	FString TimeToLUTFrac = TEXT("TimeToLUTFraction_") + ParamInfo.DataInterfaceHLSLSymbol;
-	FString Sample = TEXT("SampleCurve_") + ParamInfo.DataInterfaceHLSLSymbol;
-	FString NumSamples = TEXT("CurveLUTNumMinusOne_") + ParamInfo.DataInterfaceHLSLSymbol;
-	OutHLSL += FString::Printf(TEXT("\
-void %s(in float In_X, out float2 Out_Value) \n\
-{ \n\
-	float RemappedX = %s(In_X) * %s; \n\
-	float Prev = floor(RemappedX); \n\
-	float Next = Prev < %s ? Prev + 1.0 : Prev; \n\
-	float Interp = RemappedX - Prev; \n\
-	Prev *= %u; \n\
-	Next *= %u; \n\
-	float2 A = float2(%s(Prev), %s(Prev + 1)); \n\
-	float2 B = float2(%s(Next), %s(Next + 1)); \n\
-	Out_Value = lerp(A, B, Interp); \n\
-}\n")
-, *FunctionInfo.InstanceName, *TimeToLUTFrac, *NumSamples, *NumSamples, CurveLUTNumElems, CurveLUTNumElems, *Sample, *Sample, *Sample, *Sample);
-
-	return true;
+	if (FunctionInfo.DefinitionName == SampleCurveName)
+	{
+		return true;
+	}
+	return false;
 }
 #endif
 
diff --git a/Engine/Plugins/FX/Niagara/Source/Niagara/Private/NiagaraDataInterfaceVector4Curve.cpp b/Engine/Plugins/FX/Niagara/Source/Niagara/Private/NiagaraDataInterfaceVector4Curve.cpp
--- a/Engine/Plugins/FX/Niagara/Source/Niagara/Private/NiagaraDataInterfaceVector4Curve.cpp
+++ b/Engine/Plugins/FX/Niagara/Source/Niagara/Private/NiagaraDataInterfaceVector4Curve.cpp
@@ -176,26 +176,11 @@ void UNiagaraDataInterfaceVector4Curve::GetFunctions(TArray<FNiagaraFunctionSign
 #if WITH_EDITORONLY_DATA
 bool UNiagaraDataInterfaceVector4Curve::GetFunctionHLSL(const FNiagaraDataInterfaceGPUParamInfo& ParamInfo, const FNiagaraDataInterfaceGeneratedFunction& FunctionInfo, int FunctionInstanceIndex, FString& OutHLSL)
 {
-	FString TimeToLUTFrac = TEXT("TimeToLUTFraction_") + ParamInfo.DataInterfaceHLSLSymbol;
-	FString Sample = TEXT("SampleCurve_") + ParamInfo.DataInterfaceHLSLSymbol;
-	FString NumSamples = TEXT("CurveLUTNumMinusOne_") + ParamInfo.DataInterfaceHLSLSymbol;
-	OutHLSL += FString::Printf(TEXT("\
-void %s(in float In_X, out float4 Out_Value) \n\
-{ \n\
-	float RemappedX = %s(In_X) * %s; \n\
-	float Prev = floor(RemappedX); \n\
-	float Next = Prev < %s ? Prev + 1.0 : Prev; \n\
-	float Interp = RemappedX - Prev; \n\
-	Prev *= %u; \n\
-	Next *= %u; \n\
-	float4 A = float4(%s(Prev), %s(Prev + 1), %s(Prev + 2), %s(Prev + 3)); \n\
-	float4 B = float4(%s(Next), %s(Next + 1), %s(Next + 2), %s(Next + 3)); \n\
-	Out_Value = lerp(A, B, Interp); \n\
-}\n")
-, *FunctionInfo.InstanceName, *TimeToLUTFrac, *NumSamples, *NumSamples, CurveLUTNumElems, CurveLUTNumElems
-, *Sample, *Sample, *Sample, *Sample, *Sample, *Sample, *Sample, *Sample);
-
-	return true;
+	if (FunctionInfo.DefinitionName == SampleCurveName)
+	{
+		return true;
+	}
+	return false;
 }
 #endif
 
diff --git a/Engine/Plugins/FX/Niagara/Source/Niagara/Private/NiagaraDataInterfaceVectorCurve.cpp b/Engine/Plugins/FX/Niagara/Source/Niagara/Private/NiagaraDataInterfaceVectorCurve.cpp
--- a/Engine/Plugins/FX/Niagara/Source/Niagara/Private/NiagaraDataInterfaceVectorCurve.cpp
+++ b/Engine/Plugins/FX/Niagara/Source/Niagara/Private/NiagaraDataInterfaceVectorCurve.cpp
@@ -167,26 +167,11 @@ void UNiagaraDataInterfaceVectorCurve::GetFunctions(TArray<FNiagaraFunctionSigna
 #if WITH_EDITORONLY_DATA
 bool UNiagaraDataInterfaceVectorCurve::GetFunctionHLSL(const FNiagaraDataInterfaceGPUParamInfo& ParamInfo, const FNiagaraDataInterfaceGeneratedFunction& FunctionInfo, int FunctionInstanceIndex, FString& OutHLSL)
 {
-	FString TimeToLUTFrac = TEXT("TimeToLUTFraction_") + ParamInfo.DataInterfaceHLSLSymbol;
-	FString Sample = TEXT("SampleCurve_") + ParamInfo.DataInterfaceHLSLSymbol;
-	FString NumSamples = TEXT("CurveLUTNumMinusOne_") + ParamInfo.DataInterfaceHLSLSymbol;
-	OutHLSL += FString::Printf(TEXT("\
-void %s(in float In_X, out float3 Out_Value) \n\
-{ \n\
-	float RemappedX = %s(In_X) * %s; \n\
-	float Prev = floor(RemappedX); \n\
-	float Next = Prev < %s ? Prev + 1.0 : Prev; \n\
-	float Interp = RemappedX - Prev; \n\
-	Prev *= %u; \n\
-	Next *= %u; \n\
-	float3 A = float3(%s(Prev), %s(Prev + 1), %s(Prev + 2)); \n\
-	float3 B = float3(%s(Next), %s(Next + 1), %s(Next + 2)); \n\
-	Out_Value = lerp(A, B, Interp); \n\
-}\n") 
-, *FunctionInfo.InstanceName, *TimeToLUTFrac, *NumSamples, *NumSamples, CurveLUTNumElems, CurveLUTNumElems
-, *Sample, *Sample, *Sample, *Sample, *Sample, *Sample);
-
-	return true;
+	if (FunctionInfo.DefinitionName == SampleCurveName)
+	{
+		return true;
+	}
+	return false;
 }
 #endif
 
