- Expose FNiagaraDataInterfaceSetShaderParametersContext outside of the Niagara plugin

#rnx

diff --git a/Engine/Plugins/FX/Niagara/Source/Niagara/Classes/NiagaraDataInterface.h b/Engine/Plugins/FX/Niagara/Source/Niagara/Classes/NiagaraDataInterface.h
--- a/Engine/Plugins/FX/Niagara/Source/Niagara/Classes/NiagaraDataInterface.h
+++ b/Engine/Plugins/FX/Niagara/Source/Niagara/Classes/NiagaraDataInterface.h
@@ -271,7 +271,7 @@ struct FNiagaraDataInterfaceProxy
 
 //////////////////////////////////////////////////////////////////////////
 
-struct FNiagaraDataInterfaceSetShaderParametersContext
+struct NIAGARA_API FNiagaraDataInterfaceSetShaderParametersContext
 {
 	friend class FNiagaraGpuComputeDispatch;
 
