- Convert RenderTarget data interface to use template & parameter structure binding

diff --git a/Engine/Plugins/FX/Niagara/Shaders/Private/NiagaraDataInterfaceRenderTarget2DArrayTemplate.ush b/Engine/Plugins/FX/Niagara/Shaders/Private/NiagaraDataInterfaceRenderTarget2DArrayTemplate.ush
new file mode 100644
--- /dev/null
+++ b/Engine/Plugins/FX/Niagara/Shaders/Private/NiagaraDataInterfaceRenderTarget2DArrayTemplate.ush
@@ -0,0 +1,43 @@
+// Copyright Epic Games, Inc. All Rights Reserved.
+
+int3						{ParameterName}_TextureSize;
+RWTexture2DArray<float4>	{ParameterName}_RWTexture;
+Texture2DArray<float4>		{ParameterName}_Texture;
+SamplerState				{ParameterName}_TextureSampler;
+
+void SetRenderTargetValue_{ParameterName}_UEImpureCall(bool bExecute, int IndexX, int IndexY, int IndexZ, float4 Value)
+{
+	if ( bExecute )
+	{
+		{ParameterName}_RWTexture[int3(IndexX, IndexY, IndexZ)] = Value;
+	}
+}
+
+void GetRenderTargetValue_{ParameterName}(int IndexX, int IndexY, int IndexZ, out float4 Value)
+{
+	Value = {ParameterName}_Texture.Load(int4(IndexX, IndexY, IndexZ, 0));
+}
+
+void SampleRenderTargetValue_{ParameterName}(float2 UV, int Slice, out float4 Value)
+{
+	Value = {ParameterName}_Texture.SampleLevel({ParameterName}_TextureSampler, float3(UV.x, UV.y, Slice), 0.0f);
+}
+
+void GetRenderTargetSize_{ParameterName}(out int Width, out int Height, out int Slices)
+{
+	Width = {ParameterName}_TextureSize.x;
+	Height = {ParameterName}_TextureSize.y;
+	Slices = {ParameterName}_TextureSize.z;
+}
+
+void LinearToIndex_{ParameterName}(int Linear, out int IndexX, out int IndexY, out int IndexZ)
+{
+	IndexX = Linear % {ParameterName}_TextureSize.x;
+	IndexY = (Linear / {ParameterName}_TextureSize.x) % {ParameterName}_TextureSize.y;
+	IndexZ = Linear / ({ParameterName}_TextureSize.x * {ParameterName}_TextureSize.y);
+}
+
+void ExecToIndex_{ParameterName}(out int IndexX, out int IndexY, out int IndexZ)
+{
+	LinearToIndex_{ParameterName}(ExecIndex(), IndexX, IndexY, IndexZ);
+}
diff --git a/Engine/Plugins/FX/Niagara/Shaders/Private/NiagaraDataInterfaceRenderTarget2DTemplate.ush b/Engine/Plugins/FX/Niagara/Shaders/Private/NiagaraDataInterfaceRenderTarget2DTemplate.ush
new file mode 100644
--- /dev/null
+++ b/Engine/Plugins/FX/Niagara/Shaders/Private/NiagaraDataInterfaceRenderTarget2DTemplate.ush
@@ -0,0 +1,41 @@
+// Copyright Epic Games, Inc. All Rights Reserved.
+
+int2					{ParameterName}_TextureSize;
+RWTexture2D<float4>		{ParameterName}_RWTexture;
+Texture2D<float4>		{ParameterName}_Texture;
+SamplerState			{ParameterName}_TextureSampler;
+
+void SetRenderTargetValue_{ParameterName}_UEImpureCall(bool bExecute, int IndexX, int IndexY, float4 Value)
+{
+	if ( bExecute )
+	{
+		{ParameterName}_RWTexture[int2(IndexX, IndexY)] = Value;
+	}
+}
+
+void GetRenderTargetValue_{ParameterName}(int IndexX, int IndexY, out float4 Value)
+{
+	Value = {ParameterName}_Texture.Load(int3(IndexX, IndexY, 0));
+}
+
+void SampleRenderTargetValue_{ParameterName}(float2 UV, out float4 Value)
+{
+	Value = {ParameterName}_Texture.SampleLevel({ParameterName}_TextureSampler, UV, 0.0f);
+}
+
+void GetRenderTargetSize_{ParameterName}(out int Width, out int Height)
+{
+	Width = {ParameterName}_TextureSize.x;
+	Height = {ParameterName}_TextureSize.y;
+}
+
+void LinearToIndex_{ParameterName}(int Linear, out int IndexX, out int IndexY)
+{
+	IndexX = Linear % {ParameterName}_TextureSize.x;
+	IndexY = Linear / {ParameterName}_TextureSize.x;
+}
+
+void ExecToIndex_{ParameterName}(out int IndexX, out int IndexY)
+{
+	LinearToIndex_{ParameterName}(ExecIndex(), IndexX, IndexY);
+}
diff --git a/Engine/Plugins/FX/Niagara/Shaders/Private/NiagaraDataInterfaceRenderTargetCubeTemplate.ush b/Engine/Plugins/FX/Niagara/Shaders/Private/NiagaraDataInterfaceRenderTargetCubeTemplate.ush
new file mode 100644
--- /dev/null
+++ b/Engine/Plugins/FX/Niagara/Shaders/Private/NiagaraDataInterfaceRenderTargetCubeTemplate.ush
@@ -0,0 +1,42 @@
+// Copyright Epic Games, Inc. All Rights Reserved.
+
+int						{ParameterName}_TextureSize;
+RWTextureCube<float4>	{ParameterName}_RWTexture;
+TextureCube<float4>		{ParameterName}_Texture;
+SamplerState			{ParameterName}_TextureSampler;
+
+void SetRenderTargetValue_{ParameterName}_UEImpureCall(bool bExecute, int IndexX, int IndexY, int Face, float4 Value)
+{
+	if ( bExecute )
+	{
+		{ParameterName}_RWTexture[int3(IndexX, IndexY, Face)] = Value;
+	}
+}
+
+// Not possible to load from Cube textures
+//void GetRenderTargetValue_{ParameterName}(int IndexX, int IndexY, int Face, out float4 Value)
+//{
+//	Value = {ParameterName}_Texture.Load(int4(IndexX, IndexY, Face, 0));
+//}
+
+void SampleRenderTargetValue_{ParameterName}(float3 UVW, out float4 Value)
+{
+	Value = {ParameterName}_Texture.SampleLevel({ParameterName}_TextureSampler, UVW, 0.0f);
+}
+
+void GetRenderTargetSize_{ParameterName}(out int Size)
+{
+	Size = {ParameterName}_TextureSize;
+}
+
+void LinearToIndex_{ParameterName}(int Linear, out int IndexX, out int IndexY, out int Face)
+{
+	IndexX = Linear % {ParameterName}_TextureSize;
+	IndexY = (Linear / {ParameterName}_TextureSize) % {ParameterName}_TextureSize;
+	Face = Linear / ({ParameterName}_TextureSize * {ParameterName}_TextureSize);
+}
+
+void ExecToIndex_{ParameterName}(out int IndexX, out int IndexY, out int Face)
+{
+	LinearToIndex_{ParameterName}(ExecIndex(), IndexX, IndexY, Face);
+}
diff --git a/Engine/Plugins/FX/Niagara/Shaders/Private/NiagaraDataInterfaceRenderTargetVolumeTemplate.ush b/Engine/Plugins/FX/Niagara/Shaders/Private/NiagaraDataInterfaceRenderTargetVolumeTemplate.ush
new file mode 100644
--- /dev/null
+++ b/Engine/Plugins/FX/Niagara/Shaders/Private/NiagaraDataInterfaceRenderTargetVolumeTemplate.ush
@@ -0,0 +1,43 @@
+// Copyright Epic Games, Inc. All Rights Reserved.
+
+int3					{ParameterName}_TextureSize;
+RWTexture3D<float4>		{ParameterName}_RWTexture;
+Texture3D<float4>		{ParameterName}_Texture;
+SamplerState			{ParameterName}_TextureSampler;
+
+void SetRenderTargetValue_{ParameterName}_UEImpureCall(bool bExecute, int IndexX, int IndexY, int IndexZ, float4 Value)
+{
+	if ( bExecute )
+	{
+		{ParameterName}_RWTexture[int3(IndexX, IndexY, IndexZ)] = Value;
+	}
+}
+
+void GetRenderTargetValue_{ParameterName}(int IndexX, int IndexY, int IndexZ, out float4 Value)
+{
+	Value = {ParameterName}_Texture.Load(int4(IndexX, IndexY, IndexZ, 0));
+}
+
+void SampleRenderTargetValue_{ParameterName}(float3 UVW, out float4 Value)
+{
+	Value = {ParameterName}_Texture.SampleLevel({ParameterName}_TextureSampler, UVW, 0.0f);
+}
+
+void GetRenderTargetSize_{ParameterName}(out int Width, out int Height, out int Slices)
+{
+	Width = {ParameterName}_TextureSize.x;
+	Height = {ParameterName}_TextureSize.y;
+	Slices = {ParameterName}_TextureSize.z;
+}
+
+void LinearToIndex_{ParameterName}(int Linear, out int IndexX, out int IndexY, out int IndexZ)
+{
+	IndexX = Linear % {ParameterName}_TextureSize.x;
+	IndexY = (Linear / {ParameterName}_TextureSize.x) % {ParameterName}_TextureSize.y;
+	IndexZ = Linear / ({ParameterName}_TextureSize.x * {ParameterName}_TextureSize.y);
+}
+
+void ExecToIndex_{ParameterName}(out int IndexX, out int IndexY, out int IndexZ)
+{
+	LinearToIndex_{ParameterName}(ExecIndex(), IndexX, IndexY, IndexZ);
+}
diff --git a/Engine/Plugins/FX/Niagara/Source/Niagara/Classes/NiagaraDataInterfaceRenderTarget2D.h b/Engine/Plugins/FX/Niagara/Source/Niagara/Classes/NiagaraDataInterfaceRenderTarget2D.h
--- a/Engine/Plugins/FX/Niagara/Source/Niagara/Classes/NiagaraDataInterfaceRenderTarget2D.h
+++ b/Engine/Plugins/FX/Niagara/Source/Niagara/Classes/NiagaraDataInterfaceRenderTarget2D.h
@@ -11,7 +11,6 @@
 class FNiagaraSystemInstance;
 class UTextureRenderTarget2D;
 
-
 struct FRenderTarget2DRWInstanceData_GameThread
 {
 	FRenderTarget2DRWInstanceData_GameThread()
@@ -48,6 +47,7 @@ struct FRenderTarget2DRWInstanceData_RenderThread
 	bool bRebuildMips = false;
 	bool bReadThisFrame = false;
 	bool bWroteThisFrame = false;
+	bool bNeedsTransition = false;
 
 	FSamplerStateRHIRef SamplerStateRHI;
 	FTexture2DRHIRef TextureRHI;
@@ -83,8 +83,6 @@ class NIAGARA_API UNiagaraDataInterfaceRenderTarget2D : public UNiagaraDataInter
 	GENERATED_UCLASS_BODY()
 
 public:
-	DECLARE_NIAGARA_DI_PARAMETER();	
-		
 	virtual void PostInitProperties() override;
 	
 	//~ UNiagaraDataInterface interface
@@ -101,9 +99,13 @@ public:
 
 	// GPU sim functionality
 #if WITH_EDITORONLY_DATA
+	virtual bool AppendCompileHash(FNiagaraCompileHashVisitor* InVisitor) const override;
 	virtual void GetParameterDefinitionHLSL(const FNiagaraDataInterfaceGPUParamInfo& ParamInfo, FString& OutHLSL) override;
 	virtual bool GetFunctionHLSL(const FNiagaraDataInterfaceGPUParamInfo& ParamInfo, const FNiagaraDataInterfaceGeneratedFunction& FunctionInfo, int FunctionInstanceIndex, FString& OutHLSL) override;
 #endif
+	virtual bool UseLegacyShaderBindings() const  override { return false; }
+	virtual void BuildShaderParameters(FNiagaraShaderParametersBuilder& ShaderParametersBuilder) const override;
+	virtual void SetShaderParameters(const FNiagaraDataInterfaceSetShaderParametersContext& Context) const override;
 
 	virtual void ProvidePerInstanceDataForRenderThread(void* DataForRenderThread, void* PerInstanceData, const FNiagaraSystemInstanceID& SystemInstance) override {}
 	virtual bool InitPerInstanceData(void* PerInstanceData, FNiagaraSystemInstance* SystemInstance) override;
@@ -123,20 +125,8 @@ public:
 	virtual bool RenderVariableToCanvas(FNiagaraSystemInstanceID SystemInstanceID, FName VariableName, class FCanvas* Canvas, const FIntRect& DrawRect) const override;
 	//~ UNiagaraDataInterface interface END
 
-	void GetSize(FVectorVMExternalFunctionContext& Context); 
-	void SetSize(FVectorVMExternalFunctionContext& Context);
-
-	static const FName SetValueFunctionName;
-	static const FName GetValueFunctionName;
-	static const FName SampleValueFunctionName;
-	static const FName SetSizeFunctionName;
-	static const FName GetSizeFunctionName;
-	static const FName LinearToIndexName;
-
-	static const FString SizeName;
-	static const FString RWOutputName;
-	static const FString OutputName;
-	static const FString InputName;
+	void VMGetSize(FVectorVMExternalFunctionContext& Context); 
+	void VMSetSize(FVectorVMExternalFunctionContext& Context);
 
 	UPROPERTY(EditAnywhere, Category = "Render Target", meta = (EditCondition = "!bInheritUserParameterSettings"))
 	FIntPoint Size;
diff --git a/Engine/Plugins/FX/Niagara/Source/Niagara/Classes/NiagaraDataInterfaceRenderTarget2DArray.h b/Engine/Plugins/FX/Niagara/Source/Niagara/Classes/NiagaraDataInterfaceRenderTarget2DArray.h
--- a/Engine/Plugins/FX/Niagara/Source/Niagara/Classes/NiagaraDataInterfaceRenderTarget2DArray.h
+++ b/Engine/Plugins/FX/Niagara/Source/Niagara/Classes/NiagaraDataInterfaceRenderTarget2DArray.h
@@ -42,6 +42,7 @@ struct FRenderTarget2DArrayRWInstanceData_RenderThread
 	FIntVector Size = FIntVector(EForceInit::ForceInitToZero);
 	bool bWroteThisFrame = false;
 	bool bReadThisFrame = false;
+	bool bNeedsTransition = false;
 
 	FSamplerStateRHIRef SamplerStateRHI;
 	FTexture2DArrayRHIRef TextureRHI;
@@ -80,8 +81,6 @@ class NIAGARA_API UNiagaraDataInterfaceRenderTarget2DArray : public UNiagaraData
 	GENERATED_UCLASS_BODY()
 
 public:
-	DECLARE_NIAGARA_DI_PARAMETER();	
-		
 	virtual void PostInitProperties() override;
 	
 	//~ UNiagaraDataInterface interface
@@ -98,9 +97,13 @@ public:
 
 	// GPU sim functionality
 #if WITH_EDITORONLY_DATA
+	virtual bool AppendCompileHash(FNiagaraCompileHashVisitor* InVisitor) const override;
 	virtual void GetParameterDefinitionHLSL(const FNiagaraDataInterfaceGPUParamInfo& ParamInfo, FString& OutHLSL) override;
 	virtual bool GetFunctionHLSL(const FNiagaraDataInterfaceGPUParamInfo& ParamInfo, const FNiagaraDataInterfaceGeneratedFunction& FunctionInfo, int FunctionInstanceIndex, FString& OutHLSL) override;
 #endif
+	virtual bool UseLegacyShaderBindings() const  override { return false; }
+	virtual void BuildShaderParameters(FNiagaraShaderParametersBuilder& ShaderParametersBuilder) const override;
+	virtual void SetShaderParameters(const FNiagaraDataInterfaceSetShaderParametersContext& Context) const override;
 
 	virtual void ProvidePerInstanceDataForRenderThread(void* DataForRenderThread, void* PerInstanceData, const FNiagaraSystemInstanceID& SystemInstance) override {}
 	virtual bool InitPerInstanceData(void* PerInstanceData, FNiagaraSystemInstance* SystemInstance) override;
@@ -116,20 +119,8 @@ public:
 	virtual bool GetExposedVariableValue(const FNiagaraVariableBase& InVariable, void* InPerInstanceData, FNiagaraSystemInstance* InSystemInstance, void* OutData) const override;
 
 	//~ UNiagaraDataInterface interface END
-	void GetSize(FVectorVMExternalFunctionContext& Context); 
-	void SetSize(FVectorVMExternalFunctionContext& Context);
-
-	static const FName SetValueFunctionName;
-	static const FName GetValueFunctionName;
-	static const FName SampleValueFunctionName;
-	static const FName SetSizeFunctionName;
-	static const FName GetSizeFunctionName;
-	static const FName LinearToIndexName;
-
-	static const FString SizeName;
-	static const FString RWOutputName;
-	static const FString OutputName;
-	static const FString InputName;
+	void VMGetSize(FVectorVMExternalFunctionContext& Context); 
+	void VMSetSize(FVectorVMExternalFunctionContext& Context);
 
 	UPROPERTY(EditAnywhere, Category = "Render Target", meta = (EditCondition = "!bInheritUserParameterSettings"))
 	FIntVector Size;
diff --git a/Engine/Plugins/FX/Niagara/Source/Niagara/Classes/NiagaraDataInterfaceRenderTargetCube.h b/Engine/Plugins/FX/Niagara/Source/Niagara/Classes/NiagaraDataInterfaceRenderTargetCube.h
--- a/Engine/Plugins/FX/Niagara/Source/Niagara/Classes/NiagaraDataInterfaceRenderTargetCube.h
+++ b/Engine/Plugins/FX/Niagara/Source/Niagara/Classes/NiagaraDataInterfaceRenderTargetCube.h
@@ -41,6 +41,7 @@ struct FRenderTargetCubeRWInstanceData_RenderThread
 	int Size = 0;
 	bool bWroteThisFrame = false;
 	bool bReadThisFrame = false;
+	bool bNeedsTransition = false;
 
 	FSamplerStateRHIRef SamplerStateRHI;
 	FTextureCubeRHIRef TextureRHI;
@@ -57,6 +58,7 @@ struct FRenderTargetCubeRWInstanceData_RenderThread
 struct FNiagaraDataInterfaceProxyRenderTargetCubeProxy : public FNiagaraDataInterfaceProxyRW
 {
 	virtual int32 PerInstanceDataPassedToRenderThreadSize() const override { return 0; }
+	virtual void PostStage(FRHICommandList& RHICmdList, const FNiagaraDataInterfaceStageArgs& Context) override;
 	virtual void PostSimulate(FRHICommandList& RHICmdList, const FNiagaraDataInterfaceArgs& Context) override;
 
 	virtual FIntVector GetElementCount(FNiagaraSystemInstanceID SystemInstanceID) const override;
@@ -72,8 +74,6 @@ class NIAGARA_API UNiagaraDataInterfaceRenderTargetCube : public UNiagaraDataInt
 	GENERATED_UCLASS_BODY()
 
 public:
-	DECLARE_NIAGARA_DI_PARAMETER();	
-		
 	virtual void PostInitProperties() override;
 	
 	//~ UNiagaraDataInterface interface
@@ -90,9 +90,13 @@ public:
 
 	// GPU sim functionality
 #if WITH_EDITORONLY_DATA
+	virtual bool AppendCompileHash(FNiagaraCompileHashVisitor* InVisitor) const override;
 	virtual void GetParameterDefinitionHLSL(const FNiagaraDataInterfaceGPUParamInfo& ParamInfo, FString& OutHLSL) override;
 	virtual bool GetFunctionHLSL(const FNiagaraDataInterfaceGPUParamInfo& ParamInfo, const FNiagaraDataInterfaceGeneratedFunction& FunctionInfo, int FunctionInstanceIndex, FString& OutHLSL) override;
 #endif
+	virtual bool UseLegacyShaderBindings() const  override { return false; }
+	virtual void BuildShaderParameters(FNiagaraShaderParametersBuilder& ShaderParametersBuilder) const override;
+	virtual void SetShaderParameters(const FNiagaraDataInterfaceSetShaderParametersContext& Context) const override;
 
 	virtual void ProvidePerInstanceDataForRenderThread(void* DataForRenderThread, void* PerInstanceData, const FNiagaraSystemInstanceID& SystemInstance) override {}
 	virtual bool InitPerInstanceData(void* PerInstanceData, FNiagaraSystemInstance* SystemInstance) override;
@@ -108,20 +112,8 @@ public:
 	virtual bool GetExposedVariableValue(const FNiagaraVariableBase& InVariable, void* InPerInstanceData, FNiagaraSystemInstance* InSystemInstance, void* OutData) const override;
 	//~ UNiagaraDataInterface interface END
 	
-	void GetSize(FVectorVMExternalFunctionContext& Context); 
-	void SetSize(FVectorVMExternalFunctionContext& Context);
-
-	static const FName SetValueFunctionName;
-	static const FName GetValueFunctionName;
-	static const FName SampleValueFunctionName;
-	static const FName SetSizeFunctionName;
-	static const FName GetSizeFunctionName;
-	static const FName LinearToIndexName;
-
-	static const FString RWOutputName;
-	static const FString OutputName;
-	static const FString InputName;
-	static const FString SizeName;
+	void VMGetSize(FVectorVMExternalFunctionContext& Context);
+	void VMSetSize(FVectorVMExternalFunctionContext& Context);
 
 	UPROPERTY(EditAnywhere, Category = "Render Target", meta = (EditCondition = "!bInheritUserParameterSettings"))
 	int Size;
diff --git a/Engine/Plugins/FX/Niagara/Source/Niagara/Classes/NiagaraDataInterfaceRenderTargetVolume.h b/Engine/Plugins/FX/Niagara/Source/Niagara/Classes/NiagaraDataInterfaceRenderTargetVolume.h
--- a/Engine/Plugins/FX/Niagara/Source/Niagara/Classes/NiagaraDataInterfaceRenderTargetVolume.h
+++ b/Engine/Plugins/FX/Niagara/Source/Niagara/Classes/NiagaraDataInterfaceRenderTargetVolume.h
@@ -41,6 +41,7 @@ struct FRenderTargetVolumeRWInstanceData_RenderThread
 	FIntVector Size = FIntVector(EForceInit::ForceInitToZero);
 	bool bWroteThisFrame = false;
 	bool bReadThisFrame = false;
+	bool bNeedsTransition = false;
 
 	FSamplerStateRHIRef SamplerStateRHI;
 	FTexture3DRHIRef TextureRHI;
@@ -57,6 +58,7 @@ struct FRenderTargetVolumeRWInstanceData_RenderThread
 struct FNiagaraDataInterfaceProxyRenderTargetVolumeProxy : public FNiagaraDataInterfaceProxyRW
 {
 	virtual int32 PerInstanceDataPassedToRenderThreadSize() const override { return 0; }
+	virtual void PostStage(FRHICommandList& RHICmdList, const FNiagaraDataInterfaceStageArgs& Context) override;
 	virtual void PostSimulate(FRHICommandList& RHICmdList, const FNiagaraDataInterfaceArgs& Context) override;
 
 	virtual FIntVector GetElementCount(FNiagaraSystemInstanceID SystemInstanceID) const override;
@@ -72,8 +74,6 @@ class NIAGARA_API UNiagaraDataInterfaceRenderTargetVolume : public UNiagaraDataI
 	GENERATED_UCLASS_BODY()
 
 public:
-	DECLARE_NIAGARA_DI_PARAMETER();	
-		
 	virtual void PostInitProperties() override;
 	
 	//~ UNiagaraDataInterface interface
@@ -90,9 +90,14 @@ public:
 
 	// GPU sim functionality
 #if WITH_EDITORONLY_DATA
+	virtual bool AppendCompileHash(FNiagaraCompileHashVisitor* InVisitor) const override;
 	virtual void GetParameterDefinitionHLSL(const FNiagaraDataInterfaceGPUParamInfo& ParamInfo, FString& OutHLSL) override;
 	virtual bool GetFunctionHLSL(const FNiagaraDataInterfaceGPUParamInfo& ParamInfo, const FNiagaraDataInterfaceGeneratedFunction& FunctionInfo, int FunctionInstanceIndex, FString& OutHLSL) override;
 #endif
+	virtual bool UseLegacyShaderBindings() const  override { return false; }
+	virtual void BuildShaderParameters(FNiagaraShaderParametersBuilder& ShaderParametersBuilder) const override;
+	virtual void SetShaderParameters(const FNiagaraDataInterfaceSetShaderParametersContext& Context) const override;
+
 #if WITH_EDITOR
 	virtual ENiagaraGpuDispatchType GetGpuDispatchType() const override { return ENiagaraGpuDispatchType::ThreeD; }
 #endif
@@ -111,20 +116,8 @@ public:
 	virtual bool GetExposedVariableValue(const FNiagaraVariableBase& InVariable, void* InPerInstanceData, FNiagaraSystemInstance* InSystemInstance, void* OutData) const override;
 	//~ UNiagaraDataInterface interface END
 	
-	void GetSize(FVectorVMExternalFunctionContext& Context); 
-	void SetSize(FVectorVMExternalFunctionContext& Context);
-
-	static const FName SetValueFunctionName;
-	static const FName GetValueFunctionName;
-	static const FName SampleValueFunctionName;
-	static const FName SetSizeFunctionName;
-	static const FName GetSizeFunctionName;
-	static const FName LinearToIndexName;
-
-	static const FString RWOutputName;
-	static const FString OutputName;
-	static const FString InputName;
-	static const FString SizeName;
+	void VMGetSize(FVectorVMExternalFunctionContext& Context); 
+	void VMSetSize(FVectorVMExternalFunctionContext& Context);
 
 	UPROPERTY(EditAnywhere, Category = "Render Target", meta = (EditCondition = "!bInheritUserParameterSettings"))
 	FIntVector Size;
diff --git a/Engine/Plugins/FX/Niagara/Source/Niagara/Private/NiagaraDataInterfaceRenderTarget2D.cpp b/Engine/Plugins/FX/Niagara/Source/Niagara/Private/NiagaraDataInterfaceRenderTarget2D.cpp
--- a/Engine/Plugins/FX/Niagara/Source/Niagara/Private/NiagaraDataInterfaceRenderTarget2D.cpp
+++ b/Engine/Plugins/FX/Niagara/Source/Niagara/Private/NiagaraDataInterfaceRenderTarget2D.cpp
@@ -1,6 +1,5 @@
 // Copyright Epic Games, Inc. All Rights Reserved.
 #include "NiagaraDataInterfaceRenderTarget2D.h"
-#include "NiagaraShader.h"
 #include "ShaderParameterUtils.h"
 #include "ClearQuad.h"
 #include "TextureResource.h"
@@ -15,24 +14,34 @@
 #include "NiagaraGpuComputeDebug.h"
 #include "NiagaraGpuComputeDispatchInterface.h"
 #include "NiagaraGPUProfilerInterface.h"
+#include "NiagaraShader.h"
+#include "NiagaraShaderParametersBuilder.h"
 #include "NiagaraStats.h"
 
 #define LOCTEXT_NAMESPACE "NiagaraDataInterfaceRenderTarget2D"
 
-const FString UNiagaraDataInterfaceRenderTarget2D::SizeName(TEXT("RWSize_"));
-const FString UNiagaraDataInterfaceRenderTarget2D::RWOutputName(TEXT("RWOutput_"));
-const FString UNiagaraDataInterfaceRenderTarget2D::OutputName(TEXT("Output_"));
-const FString UNiagaraDataInterfaceRenderTarget2D::InputName(TEXT("Input_"));
-
-// Global VM function names, also used by the shaders code generation methods.
-const FName UNiagaraDataInterfaceRenderTarget2D::SetValueFunctionName("SetRenderTargetValue");
-const FName UNiagaraDataInterfaceRenderTarget2D::GetValueFunctionName("GetRenderTargetValue");
-const FName UNiagaraDataInterfaceRenderTarget2D::SampleValueFunctionName("SampleRenderTargetValue");
-const FName UNiagaraDataInterfaceRenderTarget2D::SetSizeFunctionName("SetRenderTargetSize");
-const FName UNiagaraDataInterfaceRenderTarget2D::GetSizeFunctionName("GetRenderTargetSize");
-const FName UNiagaraDataInterfaceRenderTarget2D::LinearToIndexName("LinearToIndex");
-
-static const FName GNiagaraRenderTarget2DGenerateMipsName("RenderTarget2D::GenerateMips");
+namespace NDIRenderTarget2DLocal
+{
+	BEGIN_SHADER_PARAMETER_STRUCT(FShaderParameters, )
+		SHADER_PARAMETER(FIntPoint,					TextureSize)
+		SHADER_PARAMETER_UAV(RWTexture2D<float4>,	RWTexture)
+		SHADER_PARAMETER_TEXTURE(Texture2D<float4>, Texture)
+		SHADER_PARAMETER_SAMPLER(SamplerState,		TextureSampler)
+	END_SHADER_PARAMETER_STRUCT()
+
+	static const TCHAR* TemplateShaderFile = TEXT("/Plugin/FX/Niagara/Private/NiagaraDataInterfaceRenderTarget2DTemplate.ush");
+
+	// Global VM function names, also used by the shaders code generation methods.
+	static const FName SetValueFunctionName("SetRenderTargetValue");
+	static const FName GetValueFunctionName("GetRenderTargetValue");
+	static const FName SampleValueFunctionName("SampleRenderTargetValue");
+	static const FName SetSizeFunctionName("SetRenderTargetSize");
+	static const FName GetSizeFunctionName("GetRenderTargetSize");
+	static const FName LinearToIndexName("LinearToIndex");
+	static const FName ExecToIndexName("ExecToIndex");
+
+	static const FName GenerateMipsName("RenderTarget2D::GenerateMips");
+}
 
 FNiagaraVariableBase UNiagaraDataInterfaceRenderTarget2D::ExposedRTVar;
 
@@ -163,115 +172,6 @@ struct FNDIRenderTarget2DFunctionVersion
 	};
 };
 
-/*--------------------------------------------------------------------------------------------------------------------------*/
-struct FNiagaraDataInterfaceParametersCS_RenderTarget2D : public FNiagaraDataInterfaceParametersCS
-{
-	DECLARE_TYPE_LAYOUT(FNiagaraDataInterfaceParametersCS_RenderTarget2D, NonVirtual);
-public:
-	void Bind(const FNiagaraDataInterfaceGPUParamInfo& ParameterInfo, const class FShaderParameterMap& ParameterMap)
-	{			
-		SizeParam.Bind(ParameterMap, *(UNiagaraDataInterfaceRenderTarget2D::SizeName + ParameterInfo.DataInterfaceHLSLSymbol));
-		OutputParam.Bind(ParameterMap, *(UNiagaraDataInterfaceRenderTarget2D::OutputName + ParameterInfo.DataInterfaceHLSLSymbol));
-
-		InputParam.Bind(ParameterMap, *(UNiagaraDataInterfaceRenderTarget2D::InputName + ParameterInfo.DataInterfaceHLSLSymbol));
-		InputSamplerStateParam.Bind(ParameterMap, *(UNiagaraDataInterfaceRenderTarget2D::InputName + ParameterInfo.DataInterfaceHLSLSymbol + TEXT("SamplerState")));
-	}
-
-	void Set(FRHICommandList& RHICmdList, const FNiagaraDataInterfaceSetArgs& Context) const
-	{
-		check(IsInRenderingThread());
-
-		// Get shader and DI
-		FRHIComputeShader* ComputeShaderRHI = Context.Shader.GetComputeShader();
-		FNiagaraDataInterfaceProxyRenderTarget2DProxy* VFDI = static_cast<FNiagaraDataInterfaceProxyRenderTarget2DProxy*>(Context.DataInterface);
-
-		FRenderTarget2DRWInstanceData_RenderThread* ProxyData = VFDI->SystemInstancesToProxyData_RT.Find(Context.SystemInstanceID);
-		check(ProxyData);
-
-		SetShaderValue(RHICmdList, ComputeShaderRHI, SizeParam, ProxyData->Size);
-	
-		if (OutputParam.IsUAVBound())
-		{
-			FRHIUnorderedAccessView* OutputUAV = ProxyData->UnorderedAccessViewRHI;
-			if (OutputUAV)
-			{
-				// FIXME: this transition needs to happen in FNiagaraDataInterfaceProxyRenderTarget2DProxy::PreStage so it doesn't break up the overlap group,
-				// but for some reason it stops working if I move it in there.
-				RHICmdList.Transition(FRHITransitionInfo(OutputUAV, ERHIAccess::Unknown, ERHIAccess::UAVCompute));
-				ProxyData->bRebuildMips = true;
-				ProxyData->bWroteThisFrame = true;
-			}
-			else
-			{
-				OutputUAV = Context.ComputeDispatchInterface->GetEmptyUAVFromPool(RHICmdList, EPixelFormat::PF_A16B16G16R16, ENiagaraEmptyUAVType::Buffer);
-			}
-
-			RHICmdList.SetUAVParameter(ComputeShaderRHI, OutputParam.GetUAVIndex(), OutputUAV);
-		}
-
-		if ( InputParam.IsBound() )
-		{
-			ProxyData->bReadThisFrame = true;
-
-			FRHITexture* TextureRHI = ProxyData->TextureRHI;
-			if (!ensureMsgf(!OutputParam.IsUAVBound(), TEXT("NiagaraDIRenderTarget2D(%s) is bound as both read & write, read will be ignored."), *Context.DataInterface->SourceDIName.ToString()))
-			{
-				//-TODO: Feedback to the user that read & write is bound
-				TextureRHI = nullptr;
-			}
-
-			if (TextureRHI == nullptr)
-			{
-				TextureRHI = GBlackTexture->TextureRHI;
-			}
-
-			FRHISamplerState* SamplerStateRHI = ProxyData->SamplerStateRHI;
-			if (SamplerStateRHI == nullptr)
-			{
-				SamplerStateRHI = TStaticSamplerState<SF_Point, AM_Clamp, AM_Clamp, AM_Clamp>::GetRHI();
-			}
-
-			SetTextureParameter(
-				RHICmdList,
-				ComputeShaderRHI,
-				InputParam,
-				InputSamplerStateParam,
-				SamplerStateRHI,
-				TextureRHI
-			);
-		}
-	}
-
-	void Unset(FRHICommandList& RHICmdList, const FNiagaraDataInterfaceSetArgs& Context) const 
-	{
-		if (OutputParam.IsBound())
-		{
-			OutputParam.UnsetUAV(RHICmdList, Context.Shader.GetComputeShader());
-
-			FNiagaraDataInterfaceProxyRenderTarget2DProxy* VFDI = static_cast<FNiagaraDataInterfaceProxyRenderTarget2DProxy*>(Context.DataInterface);
-			if ( FRenderTarget2DRWInstanceData_RenderThread* ProxyData = VFDI->SystemInstancesToProxyData_RT.Find(Context.SystemInstanceID) )
-			{
-				if (FRHIUnorderedAccessView* OutputUAV = ProxyData->UnorderedAccessViewRHI)
-				{
-					// FIXME: move to FNiagaraDataInterfaceProxyRenderTarget2DProxy::PostStage, same as for the transition in Set() above.
-					RHICmdList.Transition(FRHITransitionInfo(OutputUAV, ERHIAccess::UAVCompute, ERHIAccess::SRVMask));
-				}
-			}
-		}
-	}
-
-private:
-	LAYOUT_FIELD(FShaderParameter, SizeParam);
-	LAYOUT_FIELD(FRWShaderParameter, OutputParam);
-
-	LAYOUT_FIELD(FShaderResourceParameter, InputParam);
-	LAYOUT_FIELD(FShaderResourceParameter, InputSamplerStateParam);
-};
-
-IMPLEMENT_TYPE_LAYOUT(FNiagaraDataInterfaceParametersCS_RenderTarget2D);
-
-IMPLEMENT_NIAGARA_DI_PARAMETER(UNiagaraDataInterfaceRenderTarget2D, FNiagaraDataInterfaceParametersCS_RenderTarget2D);
-
 /*--------------------------------------------------------------------------------------------------------------------------*/
 
 #if STATS
@@ -317,6 +217,8 @@ void UNiagaraDataInterfaceRenderTarget2D::PostInitProperties()
 
 void UNiagaraDataInterfaceRenderTarget2D::GetFunctions(TArray<FNiagaraFunctionSignature>& OutFunctions)
 {
+	using namespace NDIRenderTarget2DLocal;
+
 	Super::GetFunctions(OutFunctions);
 
 	const int32 EmitterSystemOnlyBitmask = ENiagaraScriptUsageMask::Emitter | ENiagaraScriptUsageMask::System;
@@ -434,11 +336,31 @@ void UNiagaraDataInterfaceRenderTarget2D::GetFunctions(TArray<FNiagaraFunctionSi
 		Sig.FunctionVersion = FNDIRenderTarget2DFunctionVersion::LatestVersion;
 	#endif
 	}
+
+	{
+		FNiagaraFunctionSignature& Sig = OutFunctions.AddDefaulted_GetRef();
+		Sig.Name = ExecToIndexName;
+		Sig.Inputs.Add(FNiagaraVariable(FNiagaraTypeDefinition(GetClass()), TEXT("RenderTarget")));
+		Sig.Outputs.Add(FNiagaraVariable(FNiagaraTypeDefinition::GetIntDef(), TEXT("IndexX")));
+		Sig.Outputs.Add(FNiagaraVariable(FNiagaraTypeDefinition::GetIntDef(), TEXT("IndexY")));
+
+		Sig.bExperimental = true;
+		Sig.bMemberFunction = true;
+		Sig.bRequiresContext = false;
+		Sig.bWriteFunction = true;
+		Sig.bSupportsCPU = false;
+		Sig.bSupportsGPU = true;
+	#if WITH_EDITORONLY_DATA
+		Sig.FunctionVersion = FNDIRenderTarget2DFunctionVersion::LatestVersion;
+	#endif
+	}
 }
 
 #if WITH_EDITORONLY_DATA
 bool UNiagaraDataInterfaceRenderTarget2D::UpgradeFunctionCall(FNiagaraFunctionSignature& FunctionSignature)
 {
+	using namespace NDIRenderTarget2DLocal;
+
 	bool bWasChanged = false;
 
 	if ( FunctionSignature.FunctionVersion < FNDIRenderTarget2DFunctionVersion::AddedOptionalExecute )
@@ -458,20 +380,22 @@ bool UNiagaraDataInterfaceRenderTarget2D::UpgradeFunctionCall(FNiagaraFunctionSi
 }
 #endif
 
-DEFINE_NDI_DIRECT_FUNC_BINDER(UNiagaraDataInterfaceRenderTarget2D, GetSize);
-DEFINE_NDI_DIRECT_FUNC_BINDER(UNiagaraDataInterfaceRenderTarget2D, SetSize);
+DEFINE_NDI_DIRECT_FUNC_BINDER(UNiagaraDataInterfaceRenderTarget2D, VMGetSize);
+DEFINE_NDI_DIRECT_FUNC_BINDER(UNiagaraDataInterfaceRenderTarget2D, VMSetSize);
 void UNiagaraDataInterfaceRenderTarget2D::GetVMExternalFunction(const FVMExternalFunctionBindingInfo& BindingInfo, void* InstanceData, FVMExternalFunction &OutFunc)
 {
+	using namespace NDIRenderTarget2DLocal;
+
 	Super::GetVMExternalFunction(BindingInfo, InstanceData, OutFunc);
 	if (BindingInfo.Name == GetSizeFunctionName)
 	{
 		check(BindingInfo.GetNumInputs() == 1 && BindingInfo.GetNumOutputs() == 2);
-		NDI_FUNC_BINDER(UNiagaraDataInterfaceRenderTarget2D, GetSize)::Bind(this, OutFunc);
+		NDI_FUNC_BINDER(UNiagaraDataInterfaceRenderTarget2D, VMGetSize)::Bind(this, OutFunc);
 	}
 	else if (BindingInfo.Name == SetSizeFunctionName)
 	{
 		check(BindingInfo.GetNumInputs() == 3 && BindingInfo.GetNumOutputs() == 1);
-		NDI_FUNC_BINDER(UNiagaraDataInterfaceRenderTarget2D, SetSize)::Bind(this, OutFunc);
+		NDI_FUNC_BINDER(UNiagaraDataInterfaceRenderTarget2D, VMSetSize)::Bind(this, OutFunc);
 	}
 
 }
@@ -525,105 +449,105 @@ bool UNiagaraDataInterfaceRenderTarget2D::CopyToInternal(UNiagaraDataInterface*
 }
 
 #if WITH_EDITORONLY_DATA
+bool UNiagaraDataInterfaceRenderTarget2D::AppendCompileHash(FNiagaraCompileHashVisitor* InVisitor) const
+{
+	bool bSuccess = Super::AppendCompileHash(InVisitor);
+	FSHAHash Hash = GetShaderFileHash(NDIRenderTarget2DLocal::TemplateShaderFile, EShaderPlatform::SP_PCD3D_SM5);
+	InVisitor->UpdateString(TEXT("UNiagaraDataInterfaceRenderTarget2DTemplateHLSLSource"), Hash.ToString());
+	InVisitor->UpdateShaderParameters<NDIRenderTarget2DLocal::FShaderParameters>();
+	return bSuccess;
+}
+
 void UNiagaraDataInterfaceRenderTarget2D::GetParameterDefinitionHLSL(const FNiagaraDataInterfaceGPUParamInfo& ParamInfo, FString& OutHLSL)
 {
-	Super::GetParameterDefinitionHLSL(ParamInfo, OutHLSL);
-
-	static const TCHAR *FormatDeclarations = TEXT(R"(
-		RWTexture2D<float4> {OutputName};
-		Texture2D<float4> {InputName};
-		SamplerState {InputName}SamplerState;
-		int2 {SizeName};
-	)");
-	TMap<FString, FStringFormatArg> ArgsDeclarations =
-	{
-		{ TEXT("OutputName"),	RWOutputName + ParamInfo.DataInterfaceHLSLSymbol },
-		{ TEXT("InputName"),	InputName + ParamInfo.DataInterfaceHLSLSymbol },
-		{ TEXT("SizeName"),		SizeName + ParamInfo.DataInterfaceHLSLSymbol },
+	TMap<FString, FStringFormatArg> TemplateArgs =
+	{
+		{TEXT("ParameterName"),				ParamInfo.DataInterfaceHLSLSymbol},
 	};
-	OutHLSL += FString::Format(FormatDeclarations, ArgsDeclarations);
+
+	FString TemplateFile;
+	LoadShaderSourceFile(NDIRenderTarget2DLocal::TemplateShaderFile, EShaderPlatform::SP_PCD3D_SM5, &TemplateFile, nullptr);
+	OutHLSL += FString::Format(*TemplateFile, TemplateArgs);
 }
 
 bool UNiagaraDataInterfaceRenderTarget2D::GetFunctionHLSL(const FNiagaraDataInterfaceGPUParamInfo& ParamInfo, const FNiagaraDataInterfaceGeneratedFunction& FunctionInfo, int FunctionInstanceIndex, FString& OutHLSL)
 {
+	using namespace NDIRenderTarget2DLocal;
+
 	bool ParentRet = Super::GetFunctionHLSL(ParamInfo, FunctionInfo, FunctionInstanceIndex, OutHLSL);
 	if (ParentRet)
 	{
 		return true;
 	} 
 
-	const TMap<FString, FStringFormatArg> ArgsBounds =
-	{
-		{TEXT("FunctionName"),	FunctionInfo.InstanceName},
-		{TEXT("OutputName"),	RWOutputName + ParamInfo.DataInterfaceHLSLSymbol},
-		{TEXT("InputName"),		InputName + ParamInfo.DataInterfaceHLSLSymbol },
-		{TEXT("SizeName"),		SizeName + ParamInfo.DataInterfaceHLSLSymbol },
-	};
-
-	if (FunctionInfo.DefinitionName == SetValueFunctionName)
-	{
-		static const TCHAR* FormatBounds = TEXT(R"(
-			void {FunctionName}(bool bEnabled, int In_IndexX, int In_IndexY, float4 In_Value)
-			{			
-				if ( bEnabled )
-				{
-					{OutputName}[int2(In_IndexX, In_IndexY)] = In_Value;
-				}
-			}
-		)");
-		OutHLSL += FString::Format(FormatBounds, ArgsBounds);
-		return true;
-	}
-	else if (FunctionInfo.DefinitionName == GetValueFunctionName)
+	if ((FunctionInfo.DefinitionName == SetValueFunctionName) ||
+		(FunctionInfo.DefinitionName == GetValueFunctionName) ||
+		(FunctionInfo.DefinitionName == SampleValueFunctionName) ||
+		(FunctionInfo.DefinitionName == GetSizeFunctionName) ||
+		(FunctionInfo.DefinitionName == LinearToIndexName) ||
+		(FunctionInfo.DefinitionName == ExecToIndexName))
 	{
-		static const TCHAR* FormatBounds = TEXT(R"(
-			void {FunctionName}(int In_IndexX, int In_IndexY, out float4 Out_Value)
-			{			
-				Out_Value = {InputName}.Load(int3(In_IndexX, In_IndexY, 0));
-			}
-		)");
-		OutHLSL += FString::Format(FormatBounds, ArgsBounds);
 		return true;
 	}
-	else if (FunctionInfo.DefinitionName == SampleValueFunctionName)
+
+	return false;
+}
+#endif
+
+void UNiagaraDataInterfaceRenderTarget2D::BuildShaderParameters(FNiagaraShaderParametersBuilder& ShaderParametersBuilder) const
+{
+	ShaderParametersBuilder.AddNestedStruct<NDIRenderTarget2DLocal::FShaderParameters>();
+}
+
+void UNiagaraDataInterfaceRenderTarget2D::SetShaderParameters(const FNiagaraDataInterfaceSetShaderParametersContext& Context) const
+{
+	FNiagaraDataInterfaceProxyRenderTarget2DProxy& DIProxy = Context.GetProxy<FNiagaraDataInterfaceProxyRenderTarget2DProxy>();
+	FRenderTarget2DRWInstanceData_RenderThread* InstanceData_RT = DIProxy.SystemInstancesToProxyData_RT.Find(Context.GetSystemInstanceID());
+	check(InstanceData_RT);
+
+	NDIRenderTarget2DLocal::FShaderParameters* Parameters = Context.GetParameterNestedStruct<NDIRenderTarget2DLocal::FShaderParameters>();
+	Parameters->TextureSize = InstanceData_RT->Size;
+
+	const bool bRTWrite = Context.IsResourceBound(&Parameters->RWTexture);
+	const bool bRTRead = Context.IsResourceBound(&Parameters->Texture);
+
+	if (bRTWrite)
 	{
-		static const TCHAR* FormatBounds = TEXT(R"(
-			void {FunctionName}(float2 UV, out float4 Out_Value)
-			{			
-				Out_Value = {InputName}.SampleLevel({InputName}SamplerState, UV, 0.0f);
-			}
-		)");
-		OutHLSL += FString::Format(FormatBounds, ArgsBounds);
-		return true;
+		Parameters->RWTexture = InstanceData_RT->UnorderedAccessViewRHI;
+		if (Parameters->RWTexture)
+		{
+			// FIXME: this transition needs to happen in FNiagaraDataInterfaceProxyRenderTarget2DProxy::PreStage so it doesn't break up the overlap group,
+			// but for some reason it stops working if I move it in there.
+			Context.GetRHICmdList().Transition(FRHITransitionInfo(Parameters->RWTexture, ERHIAccess::Unknown, ERHIAccess::UAVCompute));
+			InstanceData_RT->bRebuildMips = true;
+			InstanceData_RT->bWroteThisFrame = true;
+			InstanceData_RT->bNeedsTransition = true;
+		}
+		else
+		{
+			Parameters->RWTexture = Context.GetComputeDispatchInterface().GetEmptyUAVFromPool(Context.GetRHICmdList(), EPixelFormat::PF_A16B16G16R16, ENiagaraEmptyUAVType::Texture2D);
+		}
 	}
-	else if (FunctionInfo.DefinitionName == LinearToIndexName)
+	
+	if (bRTRead)
 	{
-		static const TCHAR* FormatBounds = TEXT(R"(
-			void {FunctionName}(int Linear, out int OutIndexX, out int OutIndexY)
+		Parameters->Texture = GBlackTexture->TextureRHI;
+		Parameters->TextureSampler = TStaticSamplerState<SF_Point, AM_Clamp, AM_Clamp, AM_Clamp>::GetRHI();
+		if (ensureMsgf(bRTWrite == false, TEXT("RenderTarget DataInterface is both wrote and read from in the same stage, this is not allowed, read will be invalid")))
+		{
+			InstanceData_RT->bReadThisFrame = true;
+
+			if (InstanceData_RT->TextureRHI)
 			{
-				OutIndexX = Linear % {SizeName}.x;
-				OutIndexY = Linear / {SizeName}.x;
+				Parameters->Texture = InstanceData_RT->TextureRHI;
 			}
-		)");
-		OutHLSL += FString::Format(FormatBounds, ArgsBounds);
-		return true;
-	}
-	else if (FunctionInfo.DefinitionName == GetSizeFunctionName)
-	{
-		static const TCHAR* FormatBounds = TEXT(R"(
-			void {FunctionName}(out int Out_Width, out int Out_Height)
-			{			
-				Out_Width = {SizeName}.x;
-				Out_Height = {SizeName}.y;
+			if (InstanceData_RT->SamplerStateRHI)
+			{
+				Parameters->TextureSampler = InstanceData_RT->SamplerStateRHI;
 			}
-		)");
-		OutHLSL += FString::Format(FormatBounds, ArgsBounds);
-		return true;
+		}
 	}
-
-	return false;
 }
-#endif
 
 bool UNiagaraDataInterfaceRenderTarget2D::InitPerInstanceData(void* PerInstanceData, FNiagaraSystemInstance* SystemInstance)
 {
@@ -738,7 +662,7 @@ bool UNiagaraDataInterfaceRenderTarget2D::RenderVariableToCanvas(FNiagaraSystemI
 	return true;
 }
 
-void UNiagaraDataInterfaceRenderTarget2D::SetSize(FVectorVMExternalFunctionContext& Context)
+void UNiagaraDataInterfaceRenderTarget2D::VMSetSize(FVectorVMExternalFunctionContext& Context)
 {
 	// This should only be called from a system or emitter script due to a need for only setting up initially.
 	VectorVM::FUserPtrHandler<FRenderTarget2DRWInstanceData_GameThread> InstData(Context);
@@ -761,7 +685,7 @@ void UNiagaraDataInterfaceRenderTarget2D::SetSize(FVectorVMExternalFunctionConte
 	}
 }
 
-void UNiagaraDataInterfaceRenderTarget2D::GetSize(FVectorVMExternalFunctionContext& Context)
+void UNiagaraDataInterfaceRenderTarget2D::VMGetSize(FVectorVMExternalFunctionContext& Context)
 {
 	VectorVM::FUserPtrHandler<FRenderTarget2DRWInstanceData_GameThread> InstData(Context);
 	VectorVM::FExternalFuncRegisterHandler<int> OutSizeX(Context);
@@ -911,10 +835,19 @@ void FNiagaraDataInterfaceProxyRenderTarget2DProxy::PostStage(FRHICommandList& R
 {
 	if (FRenderTarget2DRWInstanceData_RenderThread* ProxyData = SystemInstancesToProxyData_RT.Find(Context.SystemInstanceID))
 	{
+		if (ProxyData->bNeedsTransition)
+		{
+			if (FRHIUnorderedAccessView* OutputUAV = ProxyData->UnorderedAccessViewRHI)
+			{
+				// FIXME: move to FNiagaraDataInterfaceProxyRenderTarget2DProxy::PostStage, same as for the transition in Set() above.
+				RHICmdList.Transition(FRHITransitionInfo(OutputUAV, ERHIAccess::UAVCompute, ERHIAccess::SRVMask));
+			}
+		}
+
 		if (ProxyData->bRebuildMips && (ProxyData->MipMapGeneration == ENiagaraMipMapGeneration::PostStage))
 		{
 			ProxyData->bRebuildMips = false;
-			FNiagaraGpuProfileScope GpuProfileScope(RHICmdList, Context, GNiagaraRenderTarget2DGenerateMipsName);
+			FNiagaraGpuProfileScope GpuProfileScope(RHICmdList, Context, NDIRenderTarget2DLocal::GenerateMipsName);
 			NiagaraGenerateMips::GenerateMips(RHICmdList, ProxyData->TextureRHI, ProxyData->MipMapGenerationType);
 		}
 	}
@@ -930,7 +863,7 @@ void FNiagaraDataInterfaceProxyRenderTarget2DProxy::PostSimulate(FRHICommandList
 
 	if (ProxyData->bRebuildMips && (ProxyData->MipMapGeneration == ENiagaraMipMapGeneration::PostSimulate))
 	{
-		FNiagaraGpuProfileScope GpuProfileScope(RHICmdList, Context, GNiagaraRenderTarget2DGenerateMipsName);
+		FNiagaraGpuProfileScope GpuProfileScope(RHICmdList, Context, NDIRenderTarget2DLocal::GenerateMipsName);
 		NiagaraGenerateMips::GenerateMips(RHICmdList, ProxyData->TextureRHI, ProxyData->MipMapGenerationType);
 	}
 
diff --git a/Engine/Plugins/FX/Niagara/Source/Niagara/Private/NiagaraDataInterfaceRenderTarget2DArray.cpp b/Engine/Plugins/FX/Niagara/Source/Niagara/Private/NiagaraDataInterfaceRenderTarget2DArray.cpp
--- a/Engine/Plugins/FX/Niagara/Source/Niagara/Private/NiagaraDataInterfaceRenderTarget2DArray.cpp
+++ b/Engine/Plugins/FX/Niagara/Source/Niagara/Private/NiagaraDataInterfaceRenderTarget2DArray.cpp
@@ -1,6 +1,5 @@
 // Copyright Epic Games, Inc. All Rights Reserved.
 #include "NiagaraDataInterfaceRenderTarget2DArray.h"
-#include "NiagaraShader.h"
 #include "ShaderParameterUtils.h"
 #include "ClearQuad.h"
 #include "TextureResource.h"
@@ -12,23 +11,33 @@
 #include "NiagaraStats.h"
 #include "NiagaraRenderer.h"
 #include "NiagaraSettings.h"
+#include "NiagaraShader.h"
+#include "NiagaraShaderParametersBuilder.h"
 #include "NiagaraGpuComputeDebug.h"
 #include "NiagaraGpuComputeDispatchInterface.h"
 
 #define LOCTEXT_NAMESPACE "NiagaraDataInterfaceRenderTarget2DArray"
 
-const FString UNiagaraDataInterfaceRenderTarget2DArray::SizeName(TEXT("RWSize_"));
-const FString UNiagaraDataInterfaceRenderTarget2DArray::RWOutputName(TEXT("RWOutput_"));
-const FString UNiagaraDataInterfaceRenderTarget2DArray::OutputName(TEXT("Output_"));
-const FString UNiagaraDataInterfaceRenderTarget2DArray::InputName(TEXT("Input_"));
-
-// Global VM function names, also used by the shaders code generation methods.
-const FName UNiagaraDataInterfaceRenderTarget2DArray::SetValueFunctionName("SetRenderTargetValue");
-const FName UNiagaraDataInterfaceRenderTarget2DArray::GetValueFunctionName("GetRenderTargetValue");
-const FName UNiagaraDataInterfaceRenderTarget2DArray::SampleValueFunctionName("SampleRenderTargetValue");
-const FName UNiagaraDataInterfaceRenderTarget2DArray::SetSizeFunctionName("SetRenderTargetSize");
-const FName UNiagaraDataInterfaceRenderTarget2DArray::GetSizeFunctionName("GetRenderTargetSize");
-const FName UNiagaraDataInterfaceRenderTarget2DArray::LinearToIndexName("LinearToIndex");
+namespace NDIRenderTarget2DArrayLocal
+{
+	BEGIN_SHADER_PARAMETER_STRUCT(FShaderParameters, )
+		SHADER_PARAMETER(FIntVector3,						TextureSize)
+		SHADER_PARAMETER_UAV(RWTexture2DArray<float4>,		RWTexture)
+		SHADER_PARAMETER_TEXTURE(Texture2DArray<float4>,	Texture)
+		SHADER_PARAMETER_SAMPLER(SamplerState,				TextureSampler)
+	END_SHADER_PARAMETER_STRUCT()
+
+	static const TCHAR* TemplateShaderFile = TEXT("/Plugin/FX/Niagara/Private/NiagaraDataInterfaceRenderTarget2DArrayTemplate.ush");
+
+	// Global VM function names, also used by the shaders code generation methods.
+	static const FName SetValueFunctionName("SetRenderTargetValue");
+	static const FName GetValueFunctionName("GetRenderTargetValue");
+	static const FName SampleValueFunctionName("SampleRenderTargetValue");
+	static const FName SetSizeFunctionName("SetRenderTargetSize");
+	static const FName GetSizeFunctionName("GetRenderTargetSize");
+	static const FName LinearToIndexName("LinearToIndex");
+	static const FName ExecToIndexName("ExecToIndex");
+}
 
 FNiagaraVariableBase UNiagaraDataInterfaceRenderTarget2DArray::ExposedRTVar;
 
@@ -46,110 +55,6 @@ struct FNDIRenderTarget2DArrayFunctionVersion
 	};
 };
 
-/*--------------------------------------------------------------------------------------------------------------------------*/
-struct FNiagaraDataInterfaceParametersCS_RenderTarget2DArray : public FNiagaraDataInterfaceParametersCS
-{
-	DECLARE_TYPE_LAYOUT(FNiagaraDataInterfaceParametersCS_RenderTarget2DArray, NonVirtual);
-public:
-	void Bind(const FNiagaraDataInterfaceGPUParamInfo& ParameterInfo, const class FShaderParameterMap& ParameterMap)
-	{			
-		SizeParam.Bind(ParameterMap, *(UNiagaraDataInterfaceRenderTarget2DArray::SizeName + ParameterInfo.DataInterfaceHLSLSymbol));
-		OutputParam.Bind(ParameterMap, *(UNiagaraDataInterfaceRenderTarget2DArray::OutputName + ParameterInfo.DataInterfaceHLSLSymbol));
-
-		InputParam.Bind(ParameterMap, *(UNiagaraDataInterfaceRenderTarget2DArray::InputName + ParameterInfo.DataInterfaceHLSLSymbol));
-		InputSamplerStateParam.Bind(ParameterMap, *(UNiagaraDataInterfaceRenderTarget2DArray::InputName + ParameterInfo.DataInterfaceHLSLSymbol + TEXT("SamplerState")));
-	}
-
-	void Set(FRHICommandList& RHICmdList, const FNiagaraDataInterfaceSetArgs& Context) const
-	{
-		check(IsInRenderingThread());
-
-		// Get shader and DI
-		FRHIComputeShader* ComputeShaderRHI = Context.Shader.GetComputeShader();
-		FNiagaraDataInterfaceProxyRenderTarget2DArrayProxy* VFDI = static_cast<FNiagaraDataInterfaceProxyRenderTarget2DArrayProxy*>(Context.DataInterface);
-		
-		FRenderTarget2DArrayRWInstanceData_RenderThread* ProxyData = VFDI->SystemInstancesToProxyData_RT.Find(Context.SystemInstanceID);
-		check(ProxyData);
-
-		SetShaderValue(RHICmdList, ComputeShaderRHI, SizeParam, ProxyData->Size);
-	
-		if ( OutputParam.IsUAVBound())
-		{
-			FRHIUnorderedAccessView* OutputUAV = ProxyData->UnorderedAccessViewRHI;
-			if (OutputUAV)
-			{
-				RHICmdList.Transition(FRHITransitionInfo(OutputUAV, ERHIAccess::Unknown, ERHIAccess::UAVCompute));
-				ProxyData->bWroteThisFrame = true;
-			}
-			else
-			{
-				OutputUAV = Context.ComputeDispatchInterface->GetEmptyUAVFromPool(RHICmdList, EPixelFormat::PF_A16B16G16R16, ENiagaraEmptyUAVType::Texture2DArray);
-			}
-
-			RHICmdList.SetUAVParameter(ComputeShaderRHI, OutputParam.GetUAVIndex(), OutputUAV);
-		}
-
-		if (InputParam.IsBound())
-		{
-			ProxyData->bReadThisFrame = true;
-
-			FRHITexture* TextureRHI = ProxyData->TextureRHI;
-			if (!ensureMsgf(!OutputParam.IsUAVBound(), TEXT("NiagaraDIRenderTarget2DArray(%s) is bound as both read & write, read will be ignored."), *Context.DataInterface->SourceDIName.ToString()))
-			{
-				//-TODO: Feedback to the user that read & write is bound
-				TextureRHI = nullptr;
-			}
-
-			if (TextureRHI == nullptr)
-			{
-				TextureRHI = GBlackTexture->TextureRHI;
-			}
-			FRHISamplerState* SamplerStateRHI = ProxyData->SamplerStateRHI;
-			if (SamplerStateRHI == nullptr)
-			{
-				SamplerStateRHI = TStaticSamplerState<SF_Point, AM_Clamp, AM_Clamp, AM_Clamp>::GetRHI();
-			}
-
-			SetTextureParameter(
-				RHICmdList,
-				ComputeShaderRHI,
-				InputParam,
-				InputSamplerStateParam,
-				SamplerStateRHI,
-				TextureRHI
-			);
-		}
-	}
-
-	void Unset(FRHICommandList& RHICmdList, const FNiagaraDataInterfaceSetArgs& Context) const 
-	{
-		if (OutputParam.IsBound())
-		{
-			OutputParam.UnsetUAV(RHICmdList, Context.Shader.GetComputeShader());
-
-			FNiagaraDataInterfaceProxyRenderTarget2DArrayProxy* VFDI = static_cast<FNiagaraDataInterfaceProxyRenderTarget2DArrayProxy*>(Context.DataInterface);
-			if (FRenderTarget2DArrayRWInstanceData_RenderThread* ProxyData = VFDI->SystemInstancesToProxyData_RT.Find(Context.SystemInstanceID))
-			{
-				if (FRHIUnorderedAccessView* OutputUAV = ProxyData->UnorderedAccessViewRHI)
-				{
-					RHICmdList.Transition(FRHITransitionInfo(OutputUAV, ERHIAccess::UAVCompute, ERHIAccess::SRVMask));
-				}
-			}
-		}
-	}
-
-private:
-	LAYOUT_FIELD(FShaderParameter, SizeParam);
-	LAYOUT_FIELD(FRWShaderParameter, OutputParam);
-
-	LAYOUT_FIELD(FShaderResourceParameter, InputParam);
-	LAYOUT_FIELD(FShaderResourceParameter, InputSamplerStateParam);
-};
-
-IMPLEMENT_TYPE_LAYOUT(FNiagaraDataInterfaceParametersCS_RenderTarget2DArray);
-
-IMPLEMENT_NIAGARA_DI_PARAMETER(UNiagaraDataInterfaceRenderTarget2DArray, FNiagaraDataInterfaceParametersCS_RenderTarget2DArray);
-
 /*--------------------------------------------------------------------------------------------------------------------------*/
 
 #if STATS
@@ -194,6 +99,8 @@ void UNiagaraDataInterfaceRenderTarget2DArray::PostInitProperties()
 
 void UNiagaraDataInterfaceRenderTarget2DArray::GetFunctions(TArray<FNiagaraFunctionSignature>& OutFunctions)
 {
+	using namespace NDIRenderTarget2DArrayLocal;
+
 	Super::GetFunctions(OutFunctions);
 
 	const int32 EmitterSystemOnlyBitmask = ENiagaraScriptUsageMask::Emitter | ENiagaraScriptUsageMask::System;
@@ -318,6 +225,25 @@ void UNiagaraDataInterfaceRenderTarget2DArray::GetFunctions(TArray<FNiagaraFunct
 		Sig.FunctionVersion = FNDIRenderTarget2DArrayFunctionVersion::LatestVersion;
 	#endif
 	}
+
+	{
+		FNiagaraFunctionSignature& Sig = OutFunctions.AddDefaulted_GetRef();
+		Sig.Name = ExecToIndexName;
+		Sig.Inputs.Add(FNiagaraVariable(FNiagaraTypeDefinition(GetClass()), TEXT("RenderTarget")));
+		Sig.Outputs.Add(FNiagaraVariable(FNiagaraTypeDefinition::GetIntDef(), TEXT("IndexX")));
+		Sig.Outputs.Add(FNiagaraVariable(FNiagaraTypeDefinition::GetIntDef(), TEXT("IndexY")));
+		Sig.Outputs.Add(FNiagaraVariable(FNiagaraTypeDefinition::GetIntDef(), TEXT("IndexZ")));
+
+		Sig.bExperimental = true;
+		Sig.bMemberFunction = true;
+		Sig.bRequiresContext = false;
+		Sig.bWriteFunction = true;
+		Sig.bSupportsCPU = false;
+		Sig.bSupportsGPU = true;
+	#if WITH_EDITORONLY_DATA
+		Sig.FunctionVersion = FNDIRenderTarget2DArrayFunctionVersion::LatestVersion;
+	#endif
+	}
 }
 
 #if WITH_EDITORONLY_DATA
@@ -327,7 +253,7 @@ bool UNiagaraDataInterfaceRenderTarget2DArray::UpgradeFunctionCall(FNiagaraFunct
 
 	if (FunctionSignature.FunctionVersion < FNDIRenderTarget2DArrayFunctionVersion::AddedOptionalExecute)
 	{
-		if (FunctionSignature.Name == SetValueFunctionName)
+		if (FunctionSignature.Name == NDIRenderTarget2DArrayLocal::SetValueFunctionName)
 		{
 			check(FunctionSignature.Inputs.Num() == 5);
 			FunctionSignature.Inputs.Insert_GetRef(FNiagaraVariable(FNiagaraTypeDefinition::GetBoolDef(), TEXT("Enabled")), 1).SetValue(true);
@@ -342,20 +268,22 @@ bool UNiagaraDataInterfaceRenderTarget2DArray::UpgradeFunctionCall(FNiagaraFunct
 }
 #endif
 
-DEFINE_NDI_DIRECT_FUNC_BINDER(UNiagaraDataInterfaceRenderTarget2DArray, GetSize);
-DEFINE_NDI_DIRECT_FUNC_BINDER(UNiagaraDataInterfaceRenderTarget2DArray, SetSize);
+DEFINE_NDI_DIRECT_FUNC_BINDER(UNiagaraDataInterfaceRenderTarget2DArray, VMGetSize);
+DEFINE_NDI_DIRECT_FUNC_BINDER(UNiagaraDataInterfaceRenderTarget2DArray, VMSetSize);
 void UNiagaraDataInterfaceRenderTarget2DArray::GetVMExternalFunction(const FVMExternalFunctionBindingInfo& BindingInfo, void* InstanceData, FVMExternalFunction &OutFunc)
 {
+	using namespace NDIRenderTarget2DArrayLocal;
+
 	Super::GetVMExternalFunction(BindingInfo, InstanceData, OutFunc);
 	if (BindingInfo.Name == GetSizeFunctionName)
 	{
 		check(BindingInfo.GetNumInputs() == 1 && BindingInfo.GetNumOutputs() == 3);
-		NDI_FUNC_BINDER(UNiagaraDataInterfaceRenderTarget2DArray, GetSize)::Bind(this, OutFunc);
+		NDI_FUNC_BINDER(UNiagaraDataInterfaceRenderTarget2DArray, VMGetSize)::Bind(this, OutFunc);
 	}
 	else if (BindingInfo.Name == SetSizeFunctionName)
 	{
 		check(BindingInfo.GetNumInputs() == 4 && BindingInfo.GetNumOutputs() == 1);
-		NDI_FUNC_BINDER(UNiagaraDataInterfaceRenderTarget2DArray, SetSize)::Bind(this, OutFunc);
+		NDI_FUNC_BINDER(UNiagaraDataInterfaceRenderTarget2DArray, VMSetSize)::Bind(this, OutFunc);
 	}
 
 }
@@ -405,107 +333,104 @@ bool UNiagaraDataInterfaceRenderTarget2DArray::CopyToInternal(UNiagaraDataInterf
 }
 
 #if WITH_EDITORONLY_DATA
+bool UNiagaraDataInterfaceRenderTarget2DArray::AppendCompileHash(FNiagaraCompileHashVisitor* InVisitor) const
+{
+	bool bSuccess = Super::AppendCompileHash(InVisitor);
+	FSHAHash Hash = GetShaderFileHash(NDIRenderTarget2DArrayLocal::TemplateShaderFile, EShaderPlatform::SP_PCD3D_SM5);
+	InVisitor->UpdateString(TEXT("UNiagaraDataInterfaceRenderTarget2DArrayTemplateHLSLSource"), Hash.ToString());
+	InVisitor->UpdateShaderParameters<NDIRenderTarget2DArrayLocal::FShaderParameters>();
+	return bSuccess;
+}
+
 void UNiagaraDataInterfaceRenderTarget2DArray::GetParameterDefinitionHLSL(const FNiagaraDataInterfaceGPUParamInfo& ParamInfo, FString& OutHLSL)
 {
-	Super::GetParameterDefinitionHLSL(ParamInfo, OutHLSL);
-
-	static const TCHAR *FormatDeclarations = TEXT(R"(
-		RWTexture2DArray<float4> {OutputName};
-		Texture2DArray<float4> {InputName};
-		SamplerState {InputName}SamplerState;
-		int3 {SizeName};
-	)");
-	TMap<FString, FStringFormatArg> ArgsDeclarations =
-	{
-		{ TEXT("OutputName"),	RWOutputName + ParamInfo.DataInterfaceHLSLSymbol },
-		{ TEXT("InputName"),	InputName + ParamInfo.DataInterfaceHLSLSymbol },
-		{ TEXT("SizeName"),		SizeName + ParamInfo.DataInterfaceHLSLSymbol },
+	TMap<FString, FStringFormatArg> TemplateArgs =
+	{
+		{TEXT("ParameterName"),				ParamInfo.DataInterfaceHLSLSymbol},
 	};
-	OutHLSL += FString::Format(FormatDeclarations, ArgsDeclarations);
+
+	FString TemplateFile;
+	LoadShaderSourceFile(NDIRenderTarget2DArrayLocal::TemplateShaderFile, EShaderPlatform::SP_PCD3D_SM5, &TemplateFile, nullptr);
+	OutHLSL += FString::Format(*TemplateFile, TemplateArgs);
 }
 
 bool UNiagaraDataInterfaceRenderTarget2DArray::GetFunctionHLSL(const FNiagaraDataInterfaceGPUParamInfo& ParamInfo, const FNiagaraDataInterfaceGeneratedFunction& FunctionInfo, int FunctionInstanceIndex, FString& OutHLSL)
 {
+	using namespace NDIRenderTarget2DArrayLocal;
+
 	bool ParentRet = Super::GetFunctionHLSL(ParamInfo, FunctionInfo, FunctionInstanceIndex, OutHLSL);
 	if (ParentRet)
 	{
 		return true;
-	} 
-
-	const TMap<FString, FStringFormatArg> ArgsBounds =
-	{
-		{TEXT("FunctionName"),	FunctionInfo.InstanceName},
-		{TEXT("OutputName"),	RWOutputName + ParamInfo.DataInterfaceHLSLSymbol},
-		{TEXT("InputName"),		InputName + ParamInfo.DataInterfaceHLSLSymbol },
-		{TEXT("SizeName"),		SizeName + ParamInfo.DataInterfaceHLSLSymbol },
-	};
-
-	if (FunctionInfo.DefinitionName == SetValueFunctionName)
-	{
-		static const TCHAR* FormatBounds = TEXT(R"(
-			void {FunctionName}(bool bEnabled, int IndexX, int IndexY, int IndexZ, float4 Value)
-			{			
-				if ( bEnabled )
-				{
-					{OutputName}[int3(IndexX, IndexY, IndexZ)] = Value;
-				}
-			}
-		)");
-		OutHLSL += FString::Format(FormatBounds, ArgsBounds);
-		return true;
 	}
-	else if (FunctionInfo.DefinitionName == GetValueFunctionName)
+
+	if ((FunctionInfo.DefinitionName == SetValueFunctionName) ||
+		(FunctionInfo.DefinitionName == GetValueFunctionName) ||
+		(FunctionInfo.DefinitionName == SampleValueFunctionName) ||
+		(FunctionInfo.DefinitionName == GetSizeFunctionName) ||
+		(FunctionInfo.DefinitionName == LinearToIndexName) ||
+		(FunctionInfo.DefinitionName == ExecToIndexName))
 	{
-		static const TCHAR* FormatBounds = TEXT(R"(
-			void {FunctionName}(int In_IndexX, int In_IndexY, int In_IndexZ, out float4 Out_Value)
-			{			
-				Out_Value = {InputName}.Load(int4(In_IndexX, In_IndexY, In_IndexZ, 0));
-			}
-		)");
-		OutHLSL += FString::Format(FormatBounds, ArgsBounds);
 		return true;
 	}
-	else if (FunctionInfo.DefinitionName == SampleValueFunctionName)
+
+	return false;
+}
+#endif
+
+void UNiagaraDataInterfaceRenderTarget2DArray::BuildShaderParameters(FNiagaraShaderParametersBuilder& ShaderParametersBuilder) const
+{
+	ShaderParametersBuilder.AddNestedStruct<NDIRenderTarget2DArrayLocal::FShaderParameters>();
+}
+
+void UNiagaraDataInterfaceRenderTarget2DArray::SetShaderParameters(const FNiagaraDataInterfaceSetShaderParametersContext& Context) const
+{
+	FNiagaraDataInterfaceProxyRenderTarget2DArrayProxy& DIProxy = Context.GetProxy<FNiagaraDataInterfaceProxyRenderTarget2DArrayProxy>();
+	FRenderTarget2DArrayRWInstanceData_RenderThread* InstanceData_RT = DIProxy.SystemInstancesToProxyData_RT.Find(Context.GetSystemInstanceID());
+	check(InstanceData_RT);
+
+	NDIRenderTarget2DArrayLocal::FShaderParameters* Parameters = Context.GetParameterNestedStruct<NDIRenderTarget2DArrayLocal::FShaderParameters>();
+	Parameters->TextureSize = InstanceData_RT->Size;
+
+	const bool bRTWrite = Context.IsResourceBound(&Parameters->RWTexture);
+	const bool bRTRead = Context.IsResourceBound(&Parameters->Texture);
+
+	if (bRTWrite)
 	{
-		static const TCHAR* FormatBounds = TEXT(R"(
-			void {FunctionName}(float2 UV, int Slice, out float4 Out_Value)
-			{			
-				Out_Value = {InputName}.SampleLevel({InputName}SamplerState, float3(UV.x, UV.y, Slice), 0.0f);
-			}
-		)");
-		OutHLSL += FString::Format(FormatBounds, ArgsBounds);
-		return true;
+		Parameters->RWTexture = InstanceData_RT->UnorderedAccessViewRHI;
+		if (Parameters->RWTexture)
+		{
+			// FIXME: this transition needs to happen in FNiagaraDataInterfaceProxyRenderTarget2DProxy::PreStage so it doesn't break up the overlap group,
+			// but for some reason it stops working if I move it in there.
+			Context.GetRHICmdList().Transition(FRHITransitionInfo(Parameters->RWTexture, ERHIAccess::Unknown, ERHIAccess::UAVCompute));
+			InstanceData_RT->bWroteThisFrame = true;
+			InstanceData_RT->bNeedsTransition = true;
+		}
+		else
+		{
+			Parameters->RWTexture = Context.GetComputeDispatchInterface().GetEmptyUAVFromPool(Context.GetRHICmdList(), EPixelFormat::PF_A16B16G16R16, ENiagaraEmptyUAVType::Texture2DArray);
+		}
 	}
-	else if (FunctionInfo.DefinitionName == LinearToIndexName)
+	
+	if (bRTRead)
 	{
-		static const TCHAR* FormatBounds = TEXT(R"(
-			void {FunctionName}(int Linear, out int OutIndexX, out int OutIndexY, out int OutIndexZ)
+		Parameters->Texture = GBlackArrayTexture->TextureRHI;
+		Parameters->TextureSampler = TStaticSamplerState<SF_Point, AM_Clamp, AM_Clamp, AM_Clamp>::GetRHI();
+		if (ensureMsgf(bRTWrite == false, TEXT("RenderTarget DataInterface is both wrote and read from in the same stage, this is not allowed, read will be invalid")))
+		{
+			InstanceData_RT->bReadThisFrame = true;
+
+			if (InstanceData_RT->TextureRHI)
 			{
-				OutIndexX = Linear % {SizeName}.x;
-				OutIndexY = (Linear / {SizeName}.x) % {SizeName}.y;
-				OutIndexZ = Linear / ({SizeName}.x * {SizeName}.y);
+				Parameters->Texture = InstanceData_RT->TextureRHI;
 			}
-		)");
-		OutHLSL += FString::Format(FormatBounds, ArgsBounds);
-		return true;
-	}
-	else if (FunctionInfo.DefinitionName == GetSizeFunctionName)
-	{
-		static const TCHAR* FormatBounds = TEXT(R"(
-			void {FunctionName}(out int OutWidth, out int OutHeight, out int OutSlices)
-			{			
-				OutWidth = {SizeName}.x;
-				OutHeight = {SizeName}.y;
-				OutSlices = {SizeName}.z;
+			if (InstanceData_RT->SamplerStateRHI)
+			{
+				Parameters->TextureSampler = InstanceData_RT->SamplerStateRHI;
 			}
-		)");
-		OutHLSL += FString::Format(FormatBounds, ArgsBounds);
-		return true;
+		}
 	}
-
-	return false;
 }
-#endif
 
 bool UNiagaraDataInterfaceRenderTarget2DArray::InitPerInstanceData(void* PerInstanceData, FNiagaraSystemInstance* SystemInstance)
 {
@@ -583,7 +508,7 @@ bool UNiagaraDataInterfaceRenderTarget2DArray::GetExposedVariableValue(const FNi
 	return false;
 }
 
-void UNiagaraDataInterfaceRenderTarget2DArray::SetSize(FVectorVMExternalFunctionContext& Context)
+void UNiagaraDataInterfaceRenderTarget2DArray::VMSetSize(FVectorVMExternalFunctionContext& Context)
 {
 	// This should only be called from a system or emitter script due to a need for only setting up initially.
 	VectorVM::FUserPtrHandler<FRenderTarget2DArrayRWInstanceData_GameThread> InstData(Context);
@@ -609,7 +534,7 @@ void UNiagaraDataInterfaceRenderTarget2DArray::SetSize(FVectorVMExternalFunction
 	}
 }
 
-void UNiagaraDataInterfaceRenderTarget2DArray::GetSize(FVectorVMExternalFunctionContext& Context)
+void UNiagaraDataInterfaceRenderTarget2DArray::VMGetSize(FVectorVMExternalFunctionContext& Context)
 {
 	VectorVM::FUserPtrHandler<FRenderTarget2DArrayRWInstanceData_GameThread> InstData(Context);
 	FNDIOutputParam<int> OutSizeX(Context);
diff --git a/Engine/Plugins/FX/Niagara/Source/Niagara/Private/NiagaraDataInterfaceRenderTargetCube.cpp b/Engine/Plugins/FX/Niagara/Source/Niagara/Private/NiagaraDataInterfaceRenderTargetCube.cpp
--- a/Engine/Plugins/FX/Niagara/Source/Niagara/Private/NiagaraDataInterfaceRenderTargetCube.cpp
+++ b/Engine/Plugins/FX/Niagara/Source/Niagara/Private/NiagaraDataInterfaceRenderTargetCube.cpp
@@ -1,6 +1,5 @@
 // Copyright Epic Games, Inc. All Rights Reserved.
 #include "NiagaraDataInterfaceRenderTargetCube.h"
-#include "NiagaraShader.h"
 #include "ShaderParameterUtils.h"
 #include "ClearQuad.h"
 #include "TextureResource.h"
@@ -10,142 +9,47 @@
 #include "NiagaraStats.h"
 #include "NiagaraRenderer.h"
 #include "NiagaraSettings.h"
+#include "NiagaraShader.h"
+#include "NiagaraShaderParametersBuilder.h"
 #include "NiagaraGpuComputeDebug.h"
 #include "NiagaraGpuComputeDispatchInterface.h"
 
 #define LOCTEXT_NAMESPACE "NiagaraDataInterfaceRenderTargetCube"
 
-const FString UNiagaraDataInterfaceRenderTargetCube::SizeName(TEXT("Size_"));
-const FString UNiagaraDataInterfaceRenderTargetCube::RWOutputName(TEXT("RWOutput_"));
-const FString UNiagaraDataInterfaceRenderTargetCube::OutputName(TEXT("Output_"));
-const FString UNiagaraDataInterfaceRenderTargetCube::InputName(TEXT("Input_"));
-
-// Global VM function names, also used by the shaders code generation methods.
-const FName UNiagaraDataInterfaceRenderTargetCube::SetValueFunctionName("SetRenderTargetValue");
-const FName UNiagaraDataInterfaceRenderTargetCube::GetValueFunctionName("GetRenderTargetValue");
-const FName UNiagaraDataInterfaceRenderTargetCube::SampleValueFunctionName("SampleRenderTargetValue");
-const FName UNiagaraDataInterfaceRenderTargetCube::SetSizeFunctionName("SetRenderTargetSize");
-const FName UNiagaraDataInterfaceRenderTargetCube::GetSizeFunctionName("GetRenderTargetSize");
-const FName UNiagaraDataInterfaceRenderTargetCube::LinearToIndexName("LinearToIndex");
-
-FNiagaraVariableBase UNiagaraDataInterfaceRenderTargetCube::ExposedRTVar;
-
-/*--------------------------------------------------------------------------------------------------------------------------*/
-
-struct FNDIRenderTargetCubeFunctionVersion
-{
-	enum Type
-	{
-		InitialVersion = 0,
-		AddedOptionalExecute = 1,
-
-		VersionPlusOne,
-		LatestVersion = VersionPlusOne - 1
-	};
-};
-
-/*--------------------------------------------------------------------------------------------------------------------------*/
-struct FNiagaraDataInterfaceParametersCS_RenderTargetCube : public FNiagaraDataInterfaceParametersCS
+namespace NDIRenderTargetCubeLocal
 {
-	DECLARE_TYPE_LAYOUT(FNiagaraDataInterfaceParametersCS_RenderTargetCube, NonVirtual);
-public:
-	void Bind(const FNiagaraDataInterfaceGPUParamInfo& ParameterInfo, const class FShaderParameterMap& ParameterMap)
-	{			
-		SizeParam.Bind(ParameterMap, *(UNiagaraDataInterfaceRenderTargetCube::SizeName + ParameterInfo.DataInterfaceHLSLSymbol));
-		OutputParam.Bind(ParameterMap, *(UNiagaraDataInterfaceRenderTargetCube::OutputName + ParameterInfo.DataInterfaceHLSLSymbol));
-
-		InputParam.Bind(ParameterMap, *(UNiagaraDataInterfaceRenderTargetCube::InputName + ParameterInfo.DataInterfaceHLSLSymbol));
-		InputSamplerStateParam.Bind(ParameterMap, *(UNiagaraDataInterfaceRenderTargetCube::InputName + ParameterInfo.DataInterfaceHLSLSymbol + TEXT("SamplerState")));
-	}
-
-	void Set(FRHICommandList& RHICmdList, const FNiagaraDataInterfaceSetArgs& Context) const
-	{
-		check(IsInRenderingThread());
-
-		// Get shader and DI
-		FRHIComputeShader* ComputeShaderRHI = Context.Shader.GetComputeShader();
-		FNiagaraDataInterfaceProxyRenderTargetCubeProxy* VFDI = static_cast<FNiagaraDataInterfaceProxyRenderTargetCubeProxy*>(Context.DataInterface);
-		
-		FRenderTargetCubeRWInstanceData_RenderThread* ProxyData = VFDI->SystemInstancesToProxyData_RT.Find(Context.SystemInstanceID);
-		check(ProxyData);
-
-		SetShaderValue(RHICmdList, ComputeShaderRHI, SizeParam, ProxyData->Size);	
-		
-		if ( OutputParam.IsUAVBound())
+	BEGIN_SHADER_PARAMETER_STRUCT(FShaderParameters, )
+		SHADER_PARAMETER(int, TextureSize)
+		SHADER_PARAMETER_UAV(RWTextureCube<float4>, RWTexture)
+		SHADER_PARAMETER_TEXTURE(TextureCube<float4>, Texture)
+		SHADER_PARAMETER_SAMPLER(SamplerState, TextureSampler)
+	END_SHADER_PARAMETER_STRUCT()
+
+	static const TCHAR* TemplateShaderFile = TEXT("/Plugin/FX/Niagara/Private/NiagaraDataInterfaceRenderTargetCubeTemplate.ush");
+
+	// Global VM function names, also used by the shaders code generation methods.
+	static const FName SetValueFunctionName("SetRenderTargetValue");
+	//static const FName GetValueFunctionName("GetRenderTargetValue");
+	static const FName SampleValueFunctionName("SampleRenderTargetValue");
+	static const FName SetSizeFunctionName("SetRenderTargetSize");
+	static const FName GetSizeFunctionName("GetRenderTargetSize");
+	static const FName LinearToIndexName("LinearToIndex");
+	static const FName ExecToIndexName("ExecToIndex");
+
+	struct EFunctionVersion
+	{
+		enum Type
 		{
-			FRHIUnorderedAccessView* OutputUAV = ProxyData->UnorderedAccessViewRHI;
-			if (OutputUAV)
-			{
-				RHICmdList.Transition(FRHITransitionInfo(OutputUAV, ERHIAccess::Unknown, ERHIAccess::UAVCompute));
-				ProxyData->bWroteThisFrame = true;
-			}
-			else
-			{
-				OutputUAV = Context.ComputeDispatchInterface->GetEmptyUAVFromPool(RHICmdList, EPixelFormat::PF_A16B16G16R16, ENiagaraEmptyUAVType::Texture2DArray);
-			}
-			RHICmdList.SetUAVParameter(ComputeShaderRHI, OutputParam.GetUAVIndex(), OutputUAV);
-		}
+			InitialVersion = 0,
+			AddedOptionalExecute = 1,
 
-		if (InputParam.IsBound())
-		{
-			ProxyData->bReadThisFrame = true;
-
-			FRHITexture* TextureRHI = ProxyData->TextureRHI;
-			if (!ensureMsgf(!OutputParam.IsUAVBound(), TEXT("NiagaraDIRenderTargetCube(%s) is bound as both read & write, read will be ignored."), *Context.DataInterface->SourceDIName.ToString()))
-			{
-				//-TODO: Feedback to the user that read & write is bound
-				TextureRHI = nullptr;
-			}
-
-			if (TextureRHI == nullptr)
-			{
-				TextureRHI = GBlackTextureCube->TextureRHI;
-			}
-			FRHISamplerState* SamplerStateRHI = ProxyData->SamplerStateRHI;
-			if (SamplerStateRHI == nullptr)
-			{
-				SamplerStateRHI = TStaticSamplerState<SF_Point, AM_Clamp, AM_Clamp, AM_Clamp>::GetRHI();
-			}
-
-			SetTextureParameter(
-				RHICmdList,
-				ComputeShaderRHI,
-				InputParam,
-				InputSamplerStateParam,
-				SamplerStateRHI,
-				TextureRHI
-			);
-		}
-	}
-
-	void Unset(FRHICommandList& RHICmdList, const FNiagaraDataInterfaceSetArgs& Context) const 
-	{
-		if (OutputParam.IsBound())
-		{
-			OutputParam.UnsetUAV(RHICmdList, Context.Shader.GetComputeShader());
-
-			FNiagaraDataInterfaceProxyRenderTargetCubeProxy* VFDI = static_cast<FNiagaraDataInterfaceProxyRenderTargetCubeProxy*>(Context.DataInterface);
-			if (FRenderTargetCubeRWInstanceData_RenderThread* ProxyData = VFDI->SystemInstancesToProxyData_RT.Find(Context.SystemInstanceID))
-			{
-				if (FRHIUnorderedAccessView* OutputUAV = ProxyData->UnorderedAccessViewRHI)
-				{
-					RHICmdList.Transition(FRHITransitionInfo(OutputUAV, ERHIAccess::UAVCompute, ERHIAccess::SRVMask));
-				}
-			}
-		}
-	}
-
-private:
-	LAYOUT_FIELD(FShaderParameter, SizeParam);
-	LAYOUT_FIELD(FRWShaderParameter, OutputParam);
-
-	LAYOUT_FIELD(FShaderResourceParameter, InputParam);
-	LAYOUT_FIELD(FShaderResourceParameter, InputSamplerStateParam);
-};
-
-IMPLEMENT_TYPE_LAYOUT(FNiagaraDataInterfaceParametersCS_RenderTargetCube);
+			VersionPlusOne,
+			LatestVersion = VersionPlusOne - 1
+		};
+	};
+}
 
-IMPLEMENT_NIAGARA_DI_PARAMETER(UNiagaraDataInterfaceRenderTargetCube, FNiagaraDataInterfaceParametersCS_RenderTargetCube);
+FNiagaraVariableBase UNiagaraDataInterfaceRenderTargetCube::ExposedRTVar;
 
 /*--------------------------------------------------------------------------------------------------------------------------*/
 
@@ -191,10 +95,12 @@ void UNiagaraDataInterfaceRenderTargetCube::PostInitProperties()
 
 void UNiagaraDataInterfaceRenderTargetCube::GetFunctions(TArray<FNiagaraFunctionSignature>& OutFunctions)
 {
+	using namespace NDIRenderTargetCubeLocal;
+
 	Super::GetFunctions(OutFunctions);
 
 	const int32 EmitterSystemOnlyBitmask = ENiagaraScriptUsageMask::Emitter | ENiagaraScriptUsageMask::System;
-	OutFunctions.Reserve(OutFunctions.Num() + 4);
+	OutFunctions.Reserve(OutFunctions.Num() + 7);
 
 	{
 		FNiagaraFunctionSignature& Sig = OutFunctions.AddDefaulted_GetRef();
@@ -206,7 +112,7 @@ void UNiagaraDataInterfaceRenderTargetCube::GetFunctions(TArray<FNiagaraFunction
 		Sig.bMemberFunction = true;
 		Sig.bRequiresContext = false;
 	#if WITH_EDITORONLY_DATA
-		Sig.FunctionVersion = FNDIRenderTargetCubeFunctionVersion::LatestVersion;
+		Sig.FunctionVersion = EFunctionVersion::LatestVersion;
 	#endif
 	}
 
@@ -225,7 +131,7 @@ void UNiagaraDataInterfaceRenderTargetCube::GetFunctions(TArray<FNiagaraFunction
 		Sig.bSupportsCPU = true;
 		Sig.bSupportsGPU = false;
 	#if WITH_EDITORONLY_DATA
-		Sig.FunctionVersion = FNDIRenderTargetCubeFunctionVersion::LatestVersion;
+		Sig.FunctionVersion = EFunctionVersion::LatestVersion;
 	#endif
 	}
 
@@ -247,18 +153,38 @@ void UNiagaraDataInterfaceRenderTargetCube::GetFunctions(TArray<FNiagaraFunction
 		Sig.bSupportsCPU = false;
 		Sig.bSupportsGPU = true;
 	#if WITH_EDITORONLY_DATA
-		Sig.FunctionVersion = FNDIRenderTargetCubeFunctionVersion::LatestVersion;
+		Sig.FunctionVersion = EFunctionVersion::LatestVersion;
 	#endif
 	}
 
 	extern int GNiagaraRenderTargetAllowReads;
+// CubeTexture.Load() is not available
+//	{
+//		FNiagaraFunctionSignature& Sig = OutFunctions.AddDefaulted_GetRef();
+//		Sig.Name = GetValueFunctionName;
+//		Sig.Inputs.Add(FNiagaraVariable(FNiagaraTypeDefinition(GetClass()), TEXT("RenderTarget")));
+//		Sig.Inputs.Add(FNiagaraVariable(FNiagaraTypeDefinition::GetIntDef(), TEXT("IndexX")));
+//		Sig.Inputs.Add(FNiagaraVariable(FNiagaraTypeDefinition::GetIntDef(), TEXT("IndexY")));
+//		Sig.Inputs.Add(FNiagaraVariable(FNiagaraTypeDefinition::GetIntDef(), TEXT("Face")));
+//		Sig.Outputs.Add(FNiagaraVariable(FNiagaraTypeDefinition::GetColorDef(), TEXT("Value")));
+//
+//		Sig.bHidden = GNiagaraRenderTargetAllowReads != 1;
+//		Sig.bExperimental = true;
+//		Sig.bMemberFunction = true;
+//		Sig.bRequiresContext = false;
+//		Sig.bWriteFunction = true;
+//		Sig.bSupportsCPU = false;
+//		Sig.bSupportsGPU = true;
+//#if WITH_EDITORONLY_DATA
+//		Sig.FunctionVersion = EFunctionVersion::LatestVersion;
+//#endif
+//	}
+
 	{
 		FNiagaraFunctionSignature& Sig = OutFunctions.AddDefaulted_GetRef();
-		Sig.Name = GetValueFunctionName;
+		Sig.Name = SampleValueFunctionName;
 		Sig.Inputs.Add(FNiagaraVariable(FNiagaraTypeDefinition(GetClass()), TEXT("RenderTarget")));
-		Sig.Inputs.Add(FNiagaraVariable(FNiagaraTypeDefinition::GetIntDef(), TEXT("IndexX")));
-		Sig.Inputs.Add(FNiagaraVariable(FNiagaraTypeDefinition::GetIntDef(), TEXT("IndexY")));
-		Sig.Inputs.Add(FNiagaraVariable(FNiagaraTypeDefinition::GetIntDef(), TEXT("Face")));
+		Sig.Inputs.Add(FNiagaraVariable(FNiagaraTypeDefinition::GetVec3Def(), TEXT("UVW")));
 		Sig.Outputs.Add(FNiagaraVariable(FNiagaraTypeDefinition::GetColorDef(), TEXT("Value")));
 
 		Sig.bHidden = GNiagaraRenderTargetAllowReads != 1;
@@ -269,34 +195,34 @@ void UNiagaraDataInterfaceRenderTargetCube::GetFunctions(TArray<FNiagaraFunction
 		Sig.bSupportsCPU = false;
 		Sig.bSupportsGPU = true;
 #if WITH_EDITORONLY_DATA
-		Sig.FunctionVersion = FNDIRenderTargetCubeFunctionVersion::LatestVersion;
+		Sig.FunctionVersion = EFunctionVersion::LatestVersion;
 #endif
 	}
 
 	{
 		FNiagaraFunctionSignature& Sig = OutFunctions.AddDefaulted_GetRef();
-		Sig.Name = SampleValueFunctionName;
+		Sig.Name = LinearToIndexName;
 		Sig.Inputs.Add(FNiagaraVariable(FNiagaraTypeDefinition(GetClass()), TEXT("RenderTarget")));
-		Sig.Inputs.Add(FNiagaraVariable(FNiagaraTypeDefinition::GetVec3Def(), TEXT("UVW")));
-		Sig.Outputs.Add(FNiagaraVariable(FNiagaraTypeDefinition::GetColorDef(), TEXT("Value")));
+		Sig.Inputs.Add(FNiagaraVariable(FNiagaraTypeDefinition::GetIntDef(), TEXT("Linear")));
+		Sig.Outputs.Add(FNiagaraVariable(FNiagaraTypeDefinition::GetIntDef(), TEXT("IndexX")));
+		Sig.Outputs.Add(FNiagaraVariable(FNiagaraTypeDefinition::GetIntDef(), TEXT("IndexY")));
+		Sig.Outputs.Add(FNiagaraVariable(FNiagaraTypeDefinition::GetIntDef(), TEXT("Face")));
 
-		Sig.bHidden = GNiagaraRenderTargetAllowReads != 1;
 		Sig.bExperimental = true;
 		Sig.bMemberFunction = true;
 		Sig.bRequiresContext = false;
 		Sig.bWriteFunction = true;
 		Sig.bSupportsCPU = false;
 		Sig.bSupportsGPU = true;
-#if WITH_EDITORONLY_DATA
-		Sig.FunctionVersion = FNDIRenderTargetCubeFunctionVersion::LatestVersion;
-#endif
+	#if WITH_EDITORONLY_DATA
+		Sig.FunctionVersion = EFunctionVersion::LatestVersion;
+	#endif
 	}
 
 	{
 		FNiagaraFunctionSignature& Sig = OutFunctions.AddDefaulted_GetRef();
-		Sig.Name = LinearToIndexName;
+		Sig.Name = ExecToIndexName;
 		Sig.Inputs.Add(FNiagaraVariable(FNiagaraTypeDefinition(GetClass()), TEXT("RenderTarget")));
-		Sig.Inputs.Add(FNiagaraVariable(FNiagaraTypeDefinition::GetIntDef(), TEXT("Linear")));
 		Sig.Outputs.Add(FNiagaraVariable(FNiagaraTypeDefinition::GetIntDef(), TEXT("IndexX")));
 		Sig.Outputs.Add(FNiagaraVariable(FNiagaraTypeDefinition::GetIntDef(), TEXT("IndexY")));
 		Sig.Outputs.Add(FNiagaraVariable(FNiagaraTypeDefinition::GetIntDef(), TEXT("Face")));
@@ -308,7 +234,7 @@ void UNiagaraDataInterfaceRenderTargetCube::GetFunctions(TArray<FNiagaraFunction
 		Sig.bSupportsCPU = false;
 		Sig.bSupportsGPU = true;
 	#if WITH_EDITORONLY_DATA
-		Sig.FunctionVersion = FNDIRenderTargetCubeFunctionVersion::LatestVersion;
+		Sig.FunctionVersion = EFunctionVersion::LatestVersion;
 	#endif
 	}
 }
@@ -316,9 +242,11 @@ void UNiagaraDataInterfaceRenderTargetCube::GetFunctions(TArray<FNiagaraFunction
 #if WITH_EDITORONLY_DATA
 bool UNiagaraDataInterfaceRenderTargetCube::UpgradeFunctionCall(FNiagaraFunctionSignature& FunctionSignature)
 {
+	using namespace NDIRenderTargetCubeLocal;
+
 	bool bWasChanged = false;
 
-	if (FunctionSignature.FunctionVersion < FNDIRenderTargetCubeFunctionVersion::AddedOptionalExecute)
+	if (FunctionSignature.FunctionVersion < EFunctionVersion::AddedOptionalExecute)
 	{
 		if (FunctionSignature.Name == SetValueFunctionName)
 		{
@@ -329,26 +257,28 @@ bool UNiagaraDataInterfaceRenderTargetCube::UpgradeFunctionCall(FNiagaraFunction
 	}
 
 	// Set latest version
-	FunctionSignature.FunctionVersion = FNDIRenderTargetCubeFunctionVersion::LatestVersion;
+	FunctionSignature.FunctionVersion = EFunctionVersion::LatestVersion;
 
 	return bWasChanged;
 }
 #endif
 
-DEFINE_NDI_DIRECT_FUNC_BINDER(UNiagaraDataInterfaceRenderTargetCube, GetSize);
-DEFINE_NDI_DIRECT_FUNC_BINDER(UNiagaraDataInterfaceRenderTargetCube, SetSize);
+DEFINE_NDI_DIRECT_FUNC_BINDER(UNiagaraDataInterfaceRenderTargetCube, VMGetSize);
+DEFINE_NDI_DIRECT_FUNC_BINDER(UNiagaraDataInterfaceRenderTargetCube, VMSetSize);
 void UNiagaraDataInterfaceRenderTargetCube::GetVMExternalFunction(const FVMExternalFunctionBindingInfo& BindingInfo, void* InstanceData, FVMExternalFunction &OutFunc)
 {
+	using namespace NDIRenderTargetCubeLocal;
+
 	Super::GetVMExternalFunction(BindingInfo, InstanceData, OutFunc);
 	if (BindingInfo.Name == GetSizeFunctionName)
 	{
 		check(BindingInfo.GetNumInputs() == 1 && BindingInfo.GetNumOutputs() == 1);
-		NDI_FUNC_BINDER(UNiagaraDataInterfaceRenderTargetCube, GetSize)::Bind(this, OutFunc);
+		NDI_FUNC_BINDER(UNiagaraDataInterfaceRenderTargetCube, VMGetSize)::Bind(this, OutFunc);
 	}
 	else if (BindingInfo.Name == SetSizeFunctionName)
 	{
 		check(BindingInfo.GetNumInputs() == 2 && BindingInfo.GetNumOutputs() == 1);
-		NDI_FUNC_BINDER(UNiagaraDataInterfaceRenderTargetCube, SetSize)::Bind(this, OutFunc);
+		NDI_FUNC_BINDER(UNiagaraDataInterfaceRenderTargetCube, VMSetSize)::Bind(this, OutFunc);
 	}
 }
 
@@ -397,105 +327,99 @@ bool UNiagaraDataInterfaceRenderTargetCube::CopyToInternal(UNiagaraDataInterface
 }
 
 #if WITH_EDITORONLY_DATA
+bool UNiagaraDataInterfaceRenderTargetCube::AppendCompileHash(FNiagaraCompileHashVisitor* InVisitor) const
+{
+	bool bSuccess = Super::AppendCompileHash(InVisitor);
+	FSHAHash Hash = GetShaderFileHash(NDIRenderTargetCubeLocal::TemplateShaderFile, EShaderPlatform::SP_PCD3D_SM5);
+	InVisitor->UpdateString(TEXT("UNiagaraDataInterfaceRenderTargetCubeTemplateHLSLSource"), Hash.ToString());
+	InVisitor->UpdateShaderParameters<NDIRenderTargetCubeLocal::FShaderParameters>();
+	return bSuccess;
+}
+
 void UNiagaraDataInterfaceRenderTargetCube::GetParameterDefinitionHLSL(const FNiagaraDataInterfaceGPUParamInfo& ParamInfo, FString& OutHLSL)
 {
-	Super::GetParameterDefinitionHLSL(ParamInfo, OutHLSL);
-
-	static const TCHAR *FormatDeclarations = TEXT(R"(				
-		RWTextureCube<float4> {OutputName};
-		TextureCube<float4> {InputName};
-		SamplerState {InputName}SamplerState;
-		int {SizeName};
-	)");
-	TMap<FString, FStringFormatArg> ArgsDeclarations =
-	{
-		{TEXT("OutputName"),	RWOutputName + ParamInfo.DataInterfaceHLSLSymbol},
-		{TEXT("InputName"),		InputName + ParamInfo.DataInterfaceHLSLSymbol },
-		{TEXT("SizeName"),		SizeName + ParamInfo.DataInterfaceHLSLSymbol},
+	TMap<FString, FStringFormatArg> TemplateArgs =
+	{
+		{TEXT("ParameterName"),				ParamInfo.DataInterfaceHLSLSymbol},
 	};
-	OutHLSL += FString::Format(FormatDeclarations, ArgsDeclarations);
+
+	FString TemplateFile;
+	LoadShaderSourceFile(NDIRenderTargetCubeLocal::TemplateShaderFile, EShaderPlatform::SP_PCD3D_SM5, &TemplateFile, nullptr);
+	OutHLSL += FString::Format(*TemplateFile, TemplateArgs);
 }
 
 bool UNiagaraDataInterfaceRenderTargetCube::GetFunctionHLSL(const FNiagaraDataInterfaceGPUParamInfo& ParamInfo, const FNiagaraDataInterfaceGeneratedFunction& FunctionInfo, int FunctionInstanceIndex, FString& OutHLSL)
 {
-	bool ParentRet = Super::GetFunctionHLSL(ParamInfo, FunctionInfo, FunctionInstanceIndex, OutHLSL);
-	if (ParentRet)
-	{
-		return true;
-	} 
+	using namespace NDIRenderTargetCubeLocal;
 
-	TMap<FString, FStringFormatArg> ArgsDeclarations =
+	if ((FunctionInfo.DefinitionName == SetValueFunctionName) ||
+		//(FunctionInfo.DefinitionName == GetValueFunctionName) ||
+		(FunctionInfo.DefinitionName == SampleValueFunctionName) ||
+		(FunctionInfo.DefinitionName == GetSizeFunctionName) ||
+		(FunctionInfo.DefinitionName == LinearToIndexName) ||
+		(FunctionInfo.DefinitionName == ExecToIndexName))
 	{
-		{TEXT("FunctionName"),	FunctionInfo.InstanceName},
-		{TEXT("OutputName"),	RWOutputName + ParamInfo.DataInterfaceHLSLSymbol},
-		{TEXT("InputName"),		InputName + ParamInfo.DataInterfaceHLSLSymbol },
-		{TEXT("SizeName"),		SizeName + ParamInfo.DataInterfaceHLSLSymbol},
-	};
-
-	if (FunctionInfo.DefinitionName == SetValueFunctionName)
-	{
-		static const TCHAR* FormatBounds = TEXT(R"(
-			void {FunctionName}(bool bEnabled, int IndexX, int IndexY, int Face, float4 Value)
-			{
-				if ( bEnabled )
-				{
-					{OutputName}[int3(IndexX, IndexY, Face)] = Value;
-				}
-			}
-		)");
-		OutHLSL += FString::Format(FormatBounds, ArgsDeclarations);
 		return true;
 	}
-	else if (FunctionInfo.DefinitionName == GetValueFunctionName)
-	{
-		static const TCHAR* FormatBounds = TEXT(R"(
-			void {FunctionName}(int In_IndexX, int In_IndexY, int In_Face, out float4 Out_Value)
-			{			
-				Out_Value = {InputName}.Load(int4(In_IndexX, In_IndexY, In_Face, 0));
-			}
-		)");
-		OutHLSL += FString::Format(FormatBounds, ArgsDeclarations);
-		return true;
-	}
-	else if (FunctionInfo.DefinitionName == SampleValueFunctionName)
+
+	return false;
+}
+#endif
+
+void UNiagaraDataInterfaceRenderTargetCube::BuildShaderParameters(FNiagaraShaderParametersBuilder& ShaderParametersBuilder) const
+{
+	ShaderParametersBuilder.AddNestedStruct<NDIRenderTargetCubeLocal::FShaderParameters>();
+}
+
+void UNiagaraDataInterfaceRenderTargetCube::SetShaderParameters(const FNiagaraDataInterfaceSetShaderParametersContext& Context) const
+{
+	FNiagaraDataInterfaceProxyRenderTargetCubeProxy& DIProxy = Context.GetProxy<FNiagaraDataInterfaceProxyRenderTargetCubeProxy>();
+	FRenderTargetCubeRWInstanceData_RenderThread* InstanceData_RT = DIProxy.SystemInstancesToProxyData_RT.Find(Context.GetSystemInstanceID());
+	check(InstanceData_RT);
+
+	NDIRenderTargetCubeLocal::FShaderParameters* Parameters = Context.GetParameterNestedStruct<NDIRenderTargetCubeLocal::FShaderParameters>();
+	Parameters->TextureSize = InstanceData_RT->Size;
+
+	const bool bRTWrite = Context.IsResourceBound(&Parameters->RWTexture);
+	const bool bRTRead = Context.IsResourceBound(&Parameters->Texture);
+
+	if (bRTWrite)
 	{
-		static const TCHAR* FormatBounds = TEXT(R"(
-			void {FunctionName}(float3 UVW, out float4 Out_Value)
-			{			
-				Out_Value = {InputName}.SampleLevel({InputName}SamplerState, UVW, 0.0f);
-			}
-		)");
-		OutHLSL += FString::Format(FormatBounds, ArgsDeclarations);
-		return true;
+		Parameters->RWTexture = InstanceData_RT->UnorderedAccessViewRHI;
+		if (Parameters->RWTexture)
+		{
+			// FIXME: this transition needs to happen in FNiagaraDataInterfaceProxyRenderTargetCubeProxy::PreStage so it doesn't break up the overlap group,
+			// but for some reason it stops working if I move it in there.
+			Context.GetRHICmdList().Transition(FRHITransitionInfo(Parameters->RWTexture, ERHIAccess::Unknown, ERHIAccess::UAVCompute));
+			//InstanceData_RT->bRebuildMips = true;
+			InstanceData_RT->bWroteThisFrame = true;
+			InstanceData_RT->bNeedsTransition = true;
+		}
+		else
+		{
+			Parameters->RWTexture = Context.GetComputeDispatchInterface().GetEmptyUAVFromPool(Context.GetRHICmdList(), EPixelFormat::PF_A16B16G16R16, ENiagaraEmptyUAVType::TextureCube);
+		}
 	}
-	else if (FunctionInfo.DefinitionName == LinearToIndexName)
+
+	if (bRTRead)
 	{
-		static const TCHAR* FormatBounds = TEXT(R"(
-			void {FunctionName}(int Linear, out int OutIndexX, out int OutIndexY, out int OutFace)
+		Parameters->Texture = GBlackTextureCube->TextureRHI;
+		Parameters->TextureSampler = TStaticSamplerState<SF_Point, AM_Clamp, AM_Clamp, AM_Clamp>::GetRHI();
+		if (ensureMsgf(bRTWrite == false, TEXT("RenderTarget DataInterface is both wrote and read from in the same stage, this is not allowed, read will be invalid")))
+		{
+			InstanceData_RT->bReadThisFrame = true;
+
+			if (InstanceData_RT->TextureRHI)
 			{
-				OutIndexX = Linear % {SizeName};
-				OutIndexY = (Linear / {SizeName}) % {SizeName};
-				OutFace = Linear / ({SizeName} * {SizeName});
+				Parameters->Texture = InstanceData_RT->TextureRHI;
 			}
-		)");
-		OutHLSL += FString::Format(FormatBounds, ArgsDeclarations);
-		return true;
-	}
-	else if (FunctionInfo.DefinitionName == GetSizeFunctionName)
-	{
-		static const TCHAR* FormatBounds = TEXT(R"(
-			void {FunctionName}(out int OutSize)
-			{			
-				OutSize = {SizeName};
+			if (InstanceData_RT->SamplerStateRHI)
+			{
+				Parameters->TextureSampler = InstanceData_RT->SamplerStateRHI;
 			}
-		)");
-		OutHLSL += FString::Format(FormatBounds, ArgsDeclarations);
-		return true;
+		}
 	}
-
-	return false;
 }
-#endif
 
 bool UNiagaraDataInterfaceRenderTargetCube::InitPerInstanceData(void* PerInstanceData, FNiagaraSystemInstance* SystemInstance)
 {
@@ -570,7 +494,7 @@ bool UNiagaraDataInterfaceRenderTargetCube::GetExposedVariableValue(const FNiaga
 	return false;
 }
 
-void UNiagaraDataInterfaceRenderTargetCube::SetSize(FVectorVMExternalFunctionContext& Context)
+void UNiagaraDataInterfaceRenderTargetCube::VMSetSize(FVectorVMExternalFunctionContext& Context)
 {
 	// This should only be called from a system or emitter script due to a need for only setting up initially.
 	VectorVM::FUserPtrHandler<FRenderTargetCubeRWInstanceData_GameThread> InstData(Context);
@@ -590,7 +514,7 @@ void UNiagaraDataInterfaceRenderTargetCube::SetSize(FVectorVMExternalFunctionCon
 	}
 }
 
-void UNiagaraDataInterfaceRenderTargetCube::GetSize(FVectorVMExternalFunctionContext& Context)
+void UNiagaraDataInterfaceRenderTargetCube::VMGetSize(FVectorVMExternalFunctionContext& Context)
 {
 	VectorVM::FUserPtrHandler<FRenderTargetCubeRWInstanceData_GameThread> InstData(Context);
 	FNDIOutputParam<int> OutSize(Context);
@@ -732,6 +656,21 @@ bool UNiagaraDataInterfaceRenderTargetCube::PerInstanceTickPostSimulate(void* Pe
 	return false;
 }
 
+void FNiagaraDataInterfaceProxyRenderTargetCubeProxy::PostStage(FRHICommandList& RHICmdList, const FNiagaraDataInterfaceStageArgs& Context)
+{
+	if (FRenderTargetCubeRWInstanceData_RenderThread* ProxyData = SystemInstancesToProxyData_RT.Find(Context.SystemInstanceID))
+	{
+		if (ProxyData->bNeedsTransition)
+		{
+			if (FRHIUnorderedAccessView* OutputUAV = ProxyData->UnorderedAccessViewRHI)
+			{
+				// FIXME: move to FNiagaraDataInterfaceProxyRenderTargetCubeProxy::PostStage, same as for the transition in Set() above.
+				RHICmdList.Transition(FRHITransitionInfo(OutputUAV, ERHIAccess::UAVCompute, ERHIAccess::SRVMask));
+			}
+		}
+	}
+}
+
 void FNiagaraDataInterfaceProxyRenderTargetCubeProxy::PostSimulate(FRHICommandList& RHICmdList, const FNiagaraDataInterfaceArgs& Context)
 {
 	FRenderTargetCubeRWInstanceData_RenderThread* ProxyData = SystemInstancesToProxyData_RT.Find(Context.SystemInstanceID);
diff --git a/Engine/Plugins/FX/Niagara/Source/Niagara/Private/NiagaraDataInterfaceRenderTargetVolume.cpp b/Engine/Plugins/FX/Niagara/Source/Niagara/Private/NiagaraDataInterfaceRenderTargetVolume.cpp
--- a/Engine/Plugins/FX/Niagara/Source/Niagara/Private/NiagaraDataInterfaceRenderTargetVolume.cpp
+++ b/Engine/Plugins/FX/Niagara/Source/Niagara/Private/NiagaraDataInterfaceRenderTargetVolume.cpp
@@ -1,6 +1,5 @@
 // Copyright Epic Games, Inc. All Rights Reserved.
 #include "NiagaraDataInterfaceRenderTargetVolume.h"
-#include "NiagaraShader.h"
 #include "ShaderParameterUtils.h"
 #include "ClearQuad.h"
 #include "TextureResource.h"
@@ -12,141 +11,46 @@
 #include "NiagaraStats.h"
 #include "NiagaraRenderer.h"
 #include "NiagaraSettings.h"
+#include "NiagaraShader.h"
+#include "NiagaraShaderParametersBuilder.h"
 #include "NiagaraGpuComputeDebug.h"
 
 #define LOCTEXT_NAMESPACE "NiagaraDataInterfaceRenderTargetVolume"
 
-const FString UNiagaraDataInterfaceRenderTargetVolume::SizeName(TEXT("Size_"));
-const FString UNiagaraDataInterfaceRenderTargetVolume::RWOutputName(TEXT("RWOutput_"));
-const FString UNiagaraDataInterfaceRenderTargetVolume::OutputName(TEXT("Output_"));
-const FString UNiagaraDataInterfaceRenderTargetVolume::InputName(TEXT("Input_"));
-
-// Global VM function names, also used by the shaders code generation methods.
-const FName UNiagaraDataInterfaceRenderTargetVolume::SetValueFunctionName("SetRenderTargetValue");
-const FName UNiagaraDataInterfaceRenderTargetVolume::GetValueFunctionName("GetRenderTargetValue");
-const FName UNiagaraDataInterfaceRenderTargetVolume::SampleValueFunctionName("SampleRenderTargetValue");
-const FName UNiagaraDataInterfaceRenderTargetVolume::SetSizeFunctionName("SetRenderTargetSize");
-const FName UNiagaraDataInterfaceRenderTargetVolume::GetSizeFunctionName("GetRenderTargetSize");
-const FName UNiagaraDataInterfaceRenderTargetVolume::LinearToIndexName("LinearToIndex");
-
-FNiagaraVariableBase UNiagaraDataInterfaceRenderTargetVolume::ExposedRTVar;
-
-/*--------------------------------------------------------------------------------------------------------------------------*/
-
-struct FNDIRenderTargetVolumeFunctionVersion
+namespace NDIRenderTargetVolumeLocal
 {
-	enum Type
-	{
-		InitialVersion = 0,
-		AddedOptionalExecute = 1,
-
-		VersionPlusOne,
-		LatestVersion = VersionPlusOne - 1
-	};
-};
-
-/*--------------------------------------------------------------------------------------------------------------------------*/
-struct FNiagaraDataInterfaceParametersCS_RenderTargetVolume : public FNiagaraDataInterfaceParametersCS
-{
-	DECLARE_TYPE_LAYOUT(FNiagaraDataInterfaceParametersCS_RenderTargetVolume, NonVirtual);
-public:
-	void Bind(const FNiagaraDataInterfaceGPUParamInfo& ParameterInfo, const class FShaderParameterMap& ParameterMap)
-	{			
-		SizeParam.Bind(ParameterMap, *(UNiagaraDataInterfaceRenderTargetVolume::SizeName + ParameterInfo.DataInterfaceHLSLSymbol));
-		OutputParam.Bind(ParameterMap, *(UNiagaraDataInterfaceRenderTargetVolume::OutputName + ParameterInfo.DataInterfaceHLSLSymbol));
-
-		InputParam.Bind(ParameterMap, *(UNiagaraDataInterfaceRenderTargetVolume::InputName + ParameterInfo.DataInterfaceHLSLSymbol));
-		InputSamplerStateParam.Bind(ParameterMap, *(UNiagaraDataInterfaceRenderTargetVolume::InputName + ParameterInfo.DataInterfaceHLSLSymbol + TEXT("SamplerState")));
-	}
-
-	void Set(FRHICommandList& RHICmdList, const FNiagaraDataInterfaceSetArgs& Context) const
-	{
-		check(IsInRenderingThread());
-
-		// Get shader and DI
-		FRHIComputeShader* ComputeShaderRHI = Context.Shader.GetComputeShader();
-		FNiagaraDataInterfaceProxyRenderTargetVolumeProxy* VFDI = static_cast<FNiagaraDataInterfaceProxyRenderTargetVolumeProxy*>(Context.DataInterface);
-		
-		FRenderTargetVolumeRWInstanceData_RenderThread* ProxyData = VFDI->SystemInstancesToProxyData_RT.Find(Context.SystemInstanceID);
-		check(ProxyData);
-
-		SetShaderValue(RHICmdList, ComputeShaderRHI, SizeParam, ProxyData->Size);	
-		
-		if ( OutputParam.IsUAVBound())
+	BEGIN_SHADER_PARAMETER_STRUCT(FShaderParameters, )
+		SHADER_PARAMETER(FIntVector3, TextureSize)
+		SHADER_PARAMETER_UAV(RWTexture3D<float4>, RWTexture)
+		SHADER_PARAMETER_TEXTURE(Texture3D<float4>, Texture)
+		SHADER_PARAMETER_SAMPLER(SamplerState, TextureSampler)
+	END_SHADER_PARAMETER_STRUCT()
+
+	static const TCHAR* TemplateShaderFile = TEXT("/Plugin/FX/Niagara/Private/NiagaraDataInterfaceRenderTargetVolumeTemplate.ush");
+
+	// Global VM function names, also used by the shaders code generation methods.
+	static const FName SetValueFunctionName("SetRenderTargetValue");
+	static const FName GetValueFunctionName("GetRenderTargetValue");
+	static const FName SampleValueFunctionName("SampleRenderTargetValue");
+	static const FName SetSizeFunctionName("SetRenderTargetSize");
+	static const FName GetSizeFunctionName("GetRenderTargetSize");
+	static const FName LinearToIndexName("LinearToIndex");
+	static const FName ExecToIndexName("ExecToIndex");
+
+	struct EFunctionVersion
+	{
+		enum Type
 		{
-			FRHIUnorderedAccessView* OutputUAV = ProxyData->UnorderedAccessViewRHI;
-			if (OutputUAV)
-			{
-				RHICmdList.Transition(FRHITransitionInfo(OutputUAV, ERHIAccess::Unknown, ERHIAccess::UAVCompute));
-				ProxyData->bWroteThisFrame = true;
-			}
-			else
-			{
-				OutputUAV = Context.ComputeDispatchInterface->GetEmptyUAVFromPool(RHICmdList, EPixelFormat::PF_A16B16G16R16, ENiagaraEmptyUAVType::Texture3D);
-			}
-			RHICmdList.SetUAVParameter(ComputeShaderRHI, OutputParam.GetUAVIndex(), OutputUAV);
-		}
-
-		if (InputParam.IsBound())
-		{
-			ProxyData->bReadThisFrame = true;
-
-			FRHITexture* TextureRHI = ProxyData->TextureRHI;
-			if (!ensureMsgf(!OutputParam.IsUAVBound(), TEXT("NiagaraDIRenderTargetVolume(%s) is bound as both read & write, read will be ignored."), *Context.DataInterface->SourceDIName.ToString()))
-			{
-				//-TODO: Feedback to the user that read & write is bound
-				TextureRHI = nullptr;
-			}
+			InitialVersion = 0,
+			AddedOptionalExecute = 1,
 
-			if (TextureRHI == nullptr)
-			{
-				TextureRHI = GBlackTexture->TextureRHI;
-			}
-			FRHISamplerState* SamplerStateRHI = ProxyData->SamplerStateRHI;
-			if (SamplerStateRHI == nullptr)
-			{
-				SamplerStateRHI = TStaticSamplerState<SF_Point, AM_Clamp, AM_Clamp, AM_Clamp>::GetRHI();
-			}
-
-			SetTextureParameter(
-				RHICmdList,
-				ComputeShaderRHI,
-				InputParam,
-				InputSamplerStateParam,
-				SamplerStateRHI,
-				TextureRHI
-			);
-		}
-	}
-
-	void Unset(FRHICommandList& RHICmdList, const FNiagaraDataInterfaceSetArgs& Context) const 
-	{
-		if (OutputParam.IsBound())
-		{
-			OutputParam.UnsetUAV(RHICmdList, Context.Shader.GetComputeShader());
-
-			FNiagaraDataInterfaceProxyRenderTargetVolumeProxy* VFDI = static_cast<FNiagaraDataInterfaceProxyRenderTargetVolumeProxy*>(Context.DataInterface);
-			if (FRenderTargetVolumeRWInstanceData_RenderThread* ProxyData = VFDI->SystemInstancesToProxyData_RT.Find(Context.SystemInstanceID))
-			{
-				if (FRHIUnorderedAccessView* OutputUAV = ProxyData->UnorderedAccessViewRHI)
-				{
-					RHICmdList.Transition(FRHITransitionInfo(OutputUAV, ERHIAccess::UAVCompute, ERHIAccess::SRVMask));
-				}
-			}
-		}
-	}
-
-private:
-	LAYOUT_FIELD(FShaderParameter, SizeParam);
-	LAYOUT_FIELD(FRWShaderParameter, OutputParam);
-
-	LAYOUT_FIELD(FShaderResourceParameter, InputParam);
-	LAYOUT_FIELD(FShaderResourceParameter, InputSamplerStateParam);
-};
-
-IMPLEMENT_TYPE_LAYOUT(FNiagaraDataInterfaceParametersCS_RenderTargetVolume);
+			VersionPlusOne,
+			LatestVersion = VersionPlusOne - 1
+		};
+	};
+}
 
-IMPLEMENT_NIAGARA_DI_PARAMETER(UNiagaraDataInterfaceRenderTargetVolume, FNiagaraDataInterfaceParametersCS_RenderTargetVolume);
+FNiagaraVariableBase UNiagaraDataInterfaceRenderTargetVolume::ExposedRTVar;
 
 /*--------------------------------------------------------------------------------------------------------------------------*/
 
@@ -192,10 +96,12 @@ void UNiagaraDataInterfaceRenderTargetVolume::PostInitProperties()
 
 void UNiagaraDataInterfaceRenderTargetVolume::GetFunctions(TArray<FNiagaraFunctionSignature>& OutFunctions)
 {
+	using namespace NDIRenderTargetVolumeLocal;
+
 	Super::GetFunctions(OutFunctions);
 
 	const int32 EmitterSystemOnlyBitmask = ENiagaraScriptUsageMask::Emitter | ENiagaraScriptUsageMask::System;
-	OutFunctions.Reserve(OutFunctions.Num() + 4);
+	OutFunctions.Reserve(OutFunctions.Num() + 7);
 
 	{
 		FNiagaraFunctionSignature& Sig = OutFunctions.AddDefaulted_GetRef();
@@ -209,7 +115,7 @@ void UNiagaraDataInterfaceRenderTargetVolume::GetFunctions(TArray<FNiagaraFuncti
 		Sig.bMemberFunction = true;
 		Sig.bRequiresContext = false;
 	#if WITH_EDITORONLY_DATA
-		Sig.FunctionVersion = FNDIRenderTargetVolumeFunctionVersion::LatestVersion;
+		Sig.FunctionVersion = EFunctionVersion::LatestVersion;
 	#endif
 	}
 
@@ -230,7 +136,7 @@ void UNiagaraDataInterfaceRenderTargetVolume::GetFunctions(TArray<FNiagaraFuncti
 		Sig.bSupportsCPU = true;
 		Sig.bSupportsGPU = false;
 	#if WITH_EDITORONLY_DATA
-		Sig.FunctionVersion = FNDIRenderTargetVolumeFunctionVersion::LatestVersion;
+		Sig.FunctionVersion = EFunctionVersion::LatestVersion;
 	#endif
 	}
 
@@ -252,7 +158,7 @@ void UNiagaraDataInterfaceRenderTargetVolume::GetFunctions(TArray<FNiagaraFuncti
 		Sig.bSupportsCPU = false;
 		Sig.bSupportsGPU = true;
 	#if WITH_EDITORONLY_DATA
-		Sig.FunctionVersion = FNDIRenderTargetVolumeFunctionVersion::LatestVersion;
+		Sig.FunctionVersion = EFunctionVersion::LatestVersion;
 	#endif
 	}
 
@@ -274,7 +180,7 @@ void UNiagaraDataInterfaceRenderTargetVolume::GetFunctions(TArray<FNiagaraFuncti
 		Sig.bSupportsCPU = false;
 		Sig.bSupportsGPU = true;
 #if WITH_EDITORONLY_DATA
-		Sig.FunctionVersion = FNDIRenderTargetVolumeFunctionVersion::LatestVersion;
+		Sig.FunctionVersion = EFunctionVersion::LatestVersion;
 #endif
 	}
 
@@ -293,7 +199,7 @@ void UNiagaraDataInterfaceRenderTargetVolume::GetFunctions(TArray<FNiagaraFuncti
 		Sig.bSupportsCPU = false;
 		Sig.bSupportsGPU = true;
 #if WITH_EDITORONLY_DATA
-		Sig.FunctionVersion = FNDIRenderTargetVolumeFunctionVersion::LatestVersion;
+		Sig.FunctionVersion = EFunctionVersion::LatestVersion;
 #endif
 	}
 
@@ -313,7 +219,26 @@ void UNiagaraDataInterfaceRenderTargetVolume::GetFunctions(TArray<FNiagaraFuncti
 		Sig.bSupportsCPU = false;
 		Sig.bSupportsGPU = true;
 	#if WITH_EDITORONLY_DATA
-		Sig.FunctionVersion = FNDIRenderTargetVolumeFunctionVersion::LatestVersion;
+		Sig.FunctionVersion = EFunctionVersion::LatestVersion;
+	#endif
+	}
+
+	{
+		FNiagaraFunctionSignature& Sig = OutFunctions.AddDefaulted_GetRef();
+		Sig.Name = ExecToIndexName;
+		Sig.Inputs.Add(FNiagaraVariable(FNiagaraTypeDefinition(GetClass()), TEXT("RenderTarget")));
+		Sig.Outputs.Add(FNiagaraVariable(FNiagaraTypeDefinition::GetIntDef(), TEXT("IndexX")));
+		Sig.Outputs.Add(FNiagaraVariable(FNiagaraTypeDefinition::GetIntDef(), TEXT("IndexY")));
+		Sig.Outputs.Add(FNiagaraVariable(FNiagaraTypeDefinition::GetIntDef(), TEXT("IndexZ")));
+
+		Sig.bExperimental = true;
+		Sig.bMemberFunction = true;
+		Sig.bRequiresContext = false;
+		Sig.bWriteFunction = true;
+		Sig.bSupportsCPU = false;
+		Sig.bSupportsGPU = true;
+	#if WITH_EDITORONLY_DATA
+		Sig.FunctionVersion = EFunctionVersion::LatestVersion;
 	#endif
 	}
 }
@@ -321,9 +246,11 @@ void UNiagaraDataInterfaceRenderTargetVolume::GetFunctions(TArray<FNiagaraFuncti
 #if WITH_EDITORONLY_DATA
 bool UNiagaraDataInterfaceRenderTargetVolume::UpgradeFunctionCall(FNiagaraFunctionSignature& FunctionSignature)
 {
+	using namespace NDIRenderTargetVolumeLocal;
+
 	bool bWasChanged = false;
 
-	if (FunctionSignature.FunctionVersion < FNDIRenderTargetVolumeFunctionVersion::AddedOptionalExecute)
+	if (FunctionSignature.FunctionVersion < EFunctionVersion::AddedOptionalExecute)
 	{
 		if (FunctionSignature.Name == SetValueFunctionName)
 		{
@@ -334,26 +261,28 @@ bool UNiagaraDataInterfaceRenderTargetVolume::UpgradeFunctionCall(FNiagaraFuncti
 	}
 
 	// Set latest version
-	FunctionSignature.FunctionVersion = FNDIRenderTargetVolumeFunctionVersion::LatestVersion;
+	FunctionSignature.FunctionVersion = EFunctionVersion::LatestVersion;
 
 	return bWasChanged;
 }
 #endif
 
-DEFINE_NDI_DIRECT_FUNC_BINDER(UNiagaraDataInterfaceRenderTargetVolume, GetSize);
-DEFINE_NDI_DIRECT_FUNC_BINDER(UNiagaraDataInterfaceRenderTargetVolume, SetSize);
+DEFINE_NDI_DIRECT_FUNC_BINDER(UNiagaraDataInterfaceRenderTargetVolume, VMGetSize);
+DEFINE_NDI_DIRECT_FUNC_BINDER(UNiagaraDataInterfaceRenderTargetVolume, VMSetSize);
 void UNiagaraDataInterfaceRenderTargetVolume::GetVMExternalFunction(const FVMExternalFunctionBindingInfo& BindingInfo, void* InstanceData, FVMExternalFunction &OutFunc)
 {
+	using namespace NDIRenderTargetVolumeLocal;
+
 	Super::GetVMExternalFunction(BindingInfo, InstanceData, OutFunc);
 	if (BindingInfo.Name == GetSizeFunctionName)
 	{
 		check(BindingInfo.GetNumInputs() == 1 && BindingInfo.GetNumOutputs() == 3);
-		NDI_FUNC_BINDER(UNiagaraDataInterfaceRenderTargetVolume, GetSize)::Bind(this, OutFunc);
+		NDI_FUNC_BINDER(UNiagaraDataInterfaceRenderTargetVolume, VMGetSize)::Bind(this, OutFunc);
 	}
 	else if (BindingInfo.Name == SetSizeFunctionName)
 	{
 		check(BindingInfo.GetNumInputs() == 4 && BindingInfo.GetNumOutputs() == 1);
-		NDI_FUNC_BINDER(UNiagaraDataInterfaceRenderTargetVolume, SetSize)::Bind(this, OutFunc);
+		NDI_FUNC_BINDER(UNiagaraDataInterfaceRenderTargetVolume, VMSetSize)::Bind(this, OutFunc);
 	}
 }
 
@@ -402,107 +331,99 @@ bool UNiagaraDataInterfaceRenderTargetVolume::CopyToInternal(UNiagaraDataInterfa
 }
 
 #if WITH_EDITORONLY_DATA
+bool UNiagaraDataInterfaceRenderTargetVolume::AppendCompileHash(FNiagaraCompileHashVisitor* InVisitor) const
+{
+	bool bSuccess = Super::AppendCompileHash(InVisitor);
+	FSHAHash Hash = GetShaderFileHash(NDIRenderTargetVolumeLocal::TemplateShaderFile, EShaderPlatform::SP_PCD3D_SM5);
+	InVisitor->UpdateString(TEXT("UNiagaraDataInterfaceRenderTargetVolumeTemplateHLSLSource"), Hash.ToString());
+	InVisitor->UpdateShaderParameters<NDIRenderTargetVolumeLocal::FShaderParameters>();
+	return bSuccess;
+}
+
 void UNiagaraDataInterfaceRenderTargetVolume::GetParameterDefinitionHLSL(const FNiagaraDataInterfaceGPUParamInfo& ParamInfo, FString& OutHLSL)
 {
-	Super::GetParameterDefinitionHLSL(ParamInfo, OutHLSL);
-
-	static const TCHAR *FormatDeclarations = TEXT(R"(				
-		RWTexture3D<float4> {OutputName};
-		Texture3D<float4> {InputName};
-		SamplerState {InputName}SamplerState;
-		int3 {SizeName};
-	)");
-	TMap<FString, FStringFormatArg> ArgsDeclarations =
-	{
-		{TEXT("OutputName"),	RWOutputName + ParamInfo.DataInterfaceHLSLSymbol},
-		{TEXT("InputName"),		InputName + ParamInfo.DataInterfaceHLSLSymbol },
-		{TEXT("SizeName"),		SizeName + ParamInfo.DataInterfaceHLSLSymbol},
+	TMap<FString, FStringFormatArg> TemplateArgs =
+	{
+		{TEXT("ParameterName"),				ParamInfo.DataInterfaceHLSLSymbol},
 	};
-	OutHLSL += FString::Format(FormatDeclarations, ArgsDeclarations);
+
+	FString TemplateFile;
+	LoadShaderSourceFile(NDIRenderTargetVolumeLocal::TemplateShaderFile, EShaderPlatform::SP_PCD3D_SM5, &TemplateFile, nullptr);
+	OutHLSL += FString::Format(*TemplateFile, TemplateArgs);
 }
 
 bool UNiagaraDataInterfaceRenderTargetVolume::GetFunctionHLSL(const FNiagaraDataInterfaceGPUParamInfo& ParamInfo, const FNiagaraDataInterfaceGeneratedFunction& FunctionInfo, int FunctionInstanceIndex, FString& OutHLSL)
 {
-	bool ParentRet = Super::GetFunctionHLSL(ParamInfo, FunctionInfo, FunctionInstanceIndex, OutHLSL);
-	if (ParentRet)
-	{
-		return true;
-	} 
+	using namespace NDIRenderTargetVolumeLocal;
 
-	TMap<FString, FStringFormatArg> ArgsDeclarations =
+	if ((FunctionInfo.DefinitionName == SetValueFunctionName) ||
+		(FunctionInfo.DefinitionName == GetValueFunctionName) ||
+		(FunctionInfo.DefinitionName == SampleValueFunctionName) ||
+		(FunctionInfo.DefinitionName == GetSizeFunctionName) ||
+		(FunctionInfo.DefinitionName == LinearToIndexName) ||
+		(FunctionInfo.DefinitionName == ExecToIndexName))
 	{
-		{TEXT("FunctionName"),	FunctionInfo.InstanceName},
-		{TEXT("OutputName"),	RWOutputName + ParamInfo.DataInterfaceHLSLSymbol},
-		{TEXT("InputName"),		InputName + ParamInfo.DataInterfaceHLSLSymbol },
-		{TEXT("SizeName"),		SizeName + ParamInfo.DataInterfaceHLSLSymbol},
-	};
-
-	if (FunctionInfo.DefinitionName == SetValueFunctionName)
-	{
-		static const TCHAR* FormatBounds = TEXT(R"(
-			void {FunctionName}(bool bEnabled, int IndexX, int IndexY, int IndexZ, float4 Value)
-			{
-				if ( bEnabled )
-				{
-					{OutputName}[int3(IndexX, IndexY, IndexZ)] = Value;
-				}
-			}
-		)");
-		OutHLSL += FString::Format(FormatBounds, ArgsDeclarations);
-		return true;
-	}
-	else if (FunctionInfo.DefinitionName == GetValueFunctionName)
-	{
-		static const TCHAR* FormatBounds = TEXT(R"(
-			void {FunctionName}(int In_IndexX, int In_IndexY, int In_IndexZ, out float4 Out_Value)
-			{			
-				Out_Value = {InputName}.Load(int4(In_IndexX, In_IndexY, In_IndexZ, 0));
-			}
-		)");
-		OutHLSL += FString::Format(FormatBounds, ArgsDeclarations);
 		return true;
 	}
-	else if (FunctionInfo.DefinitionName == SampleValueFunctionName)
+
+	return false;
+}
+#endif
+
+void UNiagaraDataInterfaceRenderTargetVolume::BuildShaderParameters(FNiagaraShaderParametersBuilder& ShaderParametersBuilder) const
+{
+	ShaderParametersBuilder.AddNestedStruct<NDIRenderTargetVolumeLocal::FShaderParameters>();
+}
+
+void UNiagaraDataInterfaceRenderTargetVolume::SetShaderParameters(const FNiagaraDataInterfaceSetShaderParametersContext& Context) const
+{
+	FNiagaraDataInterfaceProxyRenderTargetVolumeProxy& DIProxy = Context.GetProxy<FNiagaraDataInterfaceProxyRenderTargetVolumeProxy>();
+	FRenderTargetVolumeRWInstanceData_RenderThread* InstanceData_RT = DIProxy.SystemInstancesToProxyData_RT.Find(Context.GetSystemInstanceID());
+	check(InstanceData_RT);
+
+	NDIRenderTargetVolumeLocal::FShaderParameters* Parameters = Context.GetParameterNestedStruct<NDIRenderTargetVolumeLocal::FShaderParameters>();
+	Parameters->TextureSize = InstanceData_RT->Size;
+
+	const bool bRTWrite = Context.IsResourceBound(&Parameters->RWTexture);
+	const bool bRTRead = Context.IsResourceBound(&Parameters->Texture);
+
+	if (bRTWrite)
 	{
-		static const TCHAR* FormatBounds = TEXT(R"(
-			void {FunctionName}(float3 UVW, out float4 Out_Value)
-			{			
-				Out_Value = {InputName}.SampleLevel({InputName}SamplerState, UVW, 0.0f);
-			}
-		)");
-		OutHLSL += FString::Format(FormatBounds, ArgsDeclarations);
-		return true;
+		Parameters->RWTexture = InstanceData_RT->UnorderedAccessViewRHI;
+		if (Parameters->RWTexture)
+		{
+			// FIXME: this transition needs to happen in FNiagaraDataInterfaceProxyRenderTargetVolumeProxy::PreStage so it doesn't break up the overlap group,
+			// but for some reason it stops working if I move it in there.
+			Context.GetRHICmdList().Transition(FRHITransitionInfo(Parameters->RWTexture, ERHIAccess::Unknown, ERHIAccess::UAVCompute));
+			//InstanceData_RT->bRebuildMips = true;
+			InstanceData_RT->bWroteThisFrame = true;
+			InstanceData_RT->bNeedsTransition = true;
+		}
+		else
+		{
+			Parameters->RWTexture = Context.GetComputeDispatchInterface().GetEmptyUAVFromPool(Context.GetRHICmdList(), EPixelFormat::PF_A16B16G16R16, ENiagaraEmptyUAVType::Texture3D);
+		}
 	}
-	else if (FunctionInfo.DefinitionName == LinearToIndexName)
+
+	if (bRTRead)
 	{
-		static const TCHAR* FormatBounds = TEXT(R"(
-			void {FunctionName}(int Linear, out int OutIndexX, out int OutIndexY, out int OutIndexZ)
+		Parameters->Texture = GBlackVolumeTexture->TextureRHI;
+		Parameters->TextureSampler = TStaticSamplerState<SF_Point, AM_Clamp, AM_Clamp, AM_Clamp>::GetRHI();
+		if (ensureMsgf(bRTWrite == false, TEXT("RenderTarget DataInterface is both wrote and read from in the same stage, this is not allowed, read will be invalid")))
+		{
+			InstanceData_RT->bReadThisFrame = true;
+
+			if (InstanceData_RT->TextureRHI)
 			{
-				OutIndexX = Linear % {SizeName}.x;
-				OutIndexY = (Linear / {SizeName}.x) % {SizeName}.y;
-				OutIndexZ = Linear / ({SizeName}.x * {SizeName}.y);
+				Parameters->Texture = InstanceData_RT->TextureRHI;
 			}
-		)");
-		OutHLSL += FString::Format(FormatBounds, ArgsDeclarations);
-		return true;
-	}
-	else if (FunctionInfo.DefinitionName == GetSizeFunctionName)
-	{
-		static const TCHAR* FormatBounds = TEXT(R"(
-			void {FunctionName}(out int OutWidth, out int OutHeight, out int OutDepth)
-			{			
-				OutWidth = {SizeName}.x;
-				OutHeight = {SizeName}.y;
-				OutDepth = {SizeName}.z;
+			if (InstanceData_RT->SamplerStateRHI)
+			{
+				Parameters->TextureSampler = InstanceData_RT->SamplerStateRHI;
 			}
-		)");
-		OutHLSL += FString::Format(FormatBounds, ArgsDeclarations);
-		return true;
+		}
 	}
-
-	return false;
 }
-#endif
 
 bool UNiagaraDataInterfaceRenderTargetVolume::InitPerInstanceData(void* PerInstanceData, FNiagaraSystemInstance* SystemInstance)
 {
@@ -580,7 +501,7 @@ bool UNiagaraDataInterfaceRenderTargetVolume::GetExposedVariableValue(const FNia
 	return false;
 }
 
-void UNiagaraDataInterfaceRenderTargetVolume::SetSize(FVectorVMExternalFunctionContext& Context)
+void UNiagaraDataInterfaceRenderTargetVolume::VMSetSize(FVectorVMExternalFunctionContext& Context)
 {
 	// This should only be called from a system or emitter script due to a need for only setting up initially.
 	VectorVM::FUserPtrHandler<FRenderTargetVolumeRWInstanceData_GameThread> InstData(Context);
@@ -606,7 +527,7 @@ void UNiagaraDataInterfaceRenderTargetVolume::SetSize(FVectorVMExternalFunctionC
 	}
 }
 
-void UNiagaraDataInterfaceRenderTargetVolume::GetSize(FVectorVMExternalFunctionContext& Context)
+void UNiagaraDataInterfaceRenderTargetVolume::VMGetSize(FVectorVMExternalFunctionContext& Context)
 {
 	VectorVM::FUserPtrHandler<FRenderTargetVolumeRWInstanceData_GameThread> InstData(Context);
 	FNDIOutputParam<int> OutSizeX(Context);
@@ -754,6 +675,21 @@ bool UNiagaraDataInterfaceRenderTargetVolume::PerInstanceTickPostSimulate(void*
 	return false;
 }
 
+void FNiagaraDataInterfaceProxyRenderTargetVolumeProxy::PostStage(FRHICommandList& RHICmdList, const FNiagaraDataInterfaceStageArgs& Context)
+{
+	if (FRenderTargetVolumeRWInstanceData_RenderThread* ProxyData = SystemInstancesToProxyData_RT.Find(Context.SystemInstanceID))
+	{
+		if (ProxyData->bNeedsTransition)
+		{
+			if (FRHIUnorderedAccessView* OutputUAV = ProxyData->UnorderedAccessViewRHI)
+			{
+				// FIXME: move to FNiagaraDataInterfaceProxyRenderTargetVolumeProxy::PostStage, same as for the transition in Set() above.
+				RHICmdList.Transition(FRHITransitionInfo(OutputUAV, ERHIAccess::UAVCompute, ERHIAccess::SRVMask));
+			}
+		}
+	}
+}
+
 void FNiagaraDataInterfaceProxyRenderTargetVolumeProxy::PostSimulate(FRHICommandList& RHICmdList, const FNiagaraDataInterfaceArgs& Context)
 {
 	FRenderTargetVolumeRWInstanceData_RenderThread* ProxyData = SystemInstancesToProxyData_RT.Find(Context.SystemInstanceID);
