Restoring UE5 sytling for Niagara Systems

diff --git a/Engine/Plugins/FX/Niagara/Source/NiagaraEditor/Private/Toolkits/SystemToolkitModes/NiagaraSystemToolkitMode_Default.cpp b/Engine/Plugins/FX/Niagara/Source/NiagaraEditor/Private/Toolkits/SystemToolkitModes/NiagaraSystemToolkitMode_Default.cpp
--- a/Engine/Plugins/FX/Niagara/Source/NiagaraEditor/Private/Toolkits/SystemToolkitModes/NiagaraSystemToolkitMode_Default.cpp
+++ b/Engine/Plugins/FX/Niagara/Source/NiagaraEditor/Private/Toolkits/SystemToolkitModes/NiagaraSystemToolkitMode_Default.cpp
@@ -115,7 +115,7 @@ void FNiagaraSystemToolkitMode_Default::ExtendToolbar()
 				{
 					ToolbarBuilder.AddToolBarButton(FNiagaraEditorCommands::Get().Apply,
 						NAME_None, TAttribute<FText>(), TAttribute<FText>(),
-						FSlateIcon(FNiagaraEditorStyle::Get().GetStyleSetName(), "NiagaraEditor.Apply"),
+						FSlateIcon(FAppStyle::Get().GetStyleSetName(), "AssetEditor.Apply"),
 						FName(TEXT("ApplyNiagaraEmitter")));
 				}
 				ToolbarBuilder.AddToolBarButton(FNiagaraEditorCommands::Get().ApplyScratchPadChanges,
@@ -147,7 +147,7 @@ void FNiagaraSystemToolkitMode_Default::ExtendToolbar()
 				ToolbarBuilder.AddToolBarButton(FNiagaraEditorCommands::Get().SaveThumbnailImage, NAME_None,
 					LOCTEXT("GenerateThumbnail", "Thumbnail"),
 					LOCTEXT("GenerateThumbnailTooltip","Generate a thumbnail image."),
-					FSlateIcon(FEditorStyle::GetStyleSetName(), "Cascade.SaveThumbnailImage"));
+					FSlateIcon(FAppStyle::Get().GetStyleSetName(), "AssetEditor.SaveThumbnail"));
 			}
 			ToolbarBuilder.EndSection();
 
@@ -156,13 +156,13 @@ void FNiagaraSystemToolkitMode_Default::ExtendToolbar()
 				ToolbarBuilder.AddToolBarButton(FNiagaraEditorCommands::Get().ToggleBounds, NAME_None,
 					LOCTEXT("ShowBounds", "Bounds"),
 					LOCTEXT("ShowBoundsTooltip", "Show the bounds for the scene."),
-					FSlateIcon(FEditorStyle::GetStyleSetName(), "Cascade.ToggleBounds"));
+					FSlateIcon(FAppStyle::Get().GetStyleSetName(), "AssetEditor.ToggleShowBounds"));
 				ToolbarBuilder.AddComboButton(
 					FUIAction(),
 					FOnGetContent::CreateRaw(Toolkit, &FNiagaraSystemToolkit::GenerateBoundsMenuContent, Toolkit->GetToolkitCommands()),
 					LOCTEXT("BoundsMenuCombo_Label", "Bounds Options"),
 					LOCTEXT("BoundsMenuCombo_ToolTip", "Bounds options"),
-					FSlateIcon(FEditorStyle::GetStyleSetName(), "Cascade.ToggleBounds"),
+					FSlateIcon(FAppStyle::Get().GetStyleSetName(), "AssetEditor.ToggleShowBounds"),
 					true
 				);
 			}
@@ -174,20 +174,20 @@ void FNiagaraSystemToolkitMode_Default::ExtendToolbar()
 				ToolbarBuilder.AddToolBarButton(FNiagaraEditorCommands::Get().ToggleStatPerformance, NAME_None,
                     LOCTEXT("NiagaraShowPerformance", "Performance"),
                     LOCTEXT("NiagaraShowPerformanceTooltip", "Show runtime performance for particle scripts."),
-                    FSlateIcon(FEditorStyle::GetStyleSetName(), "MaterialEditor.ToggleMaterialStats"));
+                    FSlateIcon(FAppStyle::Get().GetStyleSetName(), "AssetEditor.ToggleStats"));
 				ToolbarBuilder.AddComboButton(
                     FUIAction(),
                     FOnGetContent::CreateRaw(Toolkit, &FNiagaraSystemToolkit::GenerateStatConfigMenuContent, Toolkit->GetToolkitCommands()),
                     FText(),
                     LOCTEXT("NiagaraShowPerformanceCombo_ToolTip", "Runtime performance options"),
-                    FSlateIcon(FEditorStyle::GetStyleSetName(), "MaterialEditor.ToggleMaterialStats"),
+                    FSlateIcon(FAppStyle::Get().GetStyleSetName(), "AssetEditor.ToggleStats"),
                     true);
 				ToolbarBuilder.AddComboButton(
 					FUIAction(),
 					FOnGetContent::CreateStatic(Local::FillDebugOptionsMenu, Toolkit),
 					LOCTEXT("DebugOptions", "Debug"),
 					LOCTEXT("DebugOptionsTooltip", "Debug options"),
-					FSlateIcon(FNiagaraEditorStyle::Get().GetStyleSetName(), "Tab.Debugger")
+					FSlateIcon(FAppStyle::Get().GetStyleSetName(), "Debug")
 				);
 			}
 			ToolbarBuilder.EndSection();
@@ -200,7 +200,7 @@ void FNiagaraSystemToolkitMode_Default::ExtendToolbar()
 					FOnGetContent::CreateStatic(Local::FillSimulationOptionsMenu, Toolkit),
 					LOCTEXT("SimulationOptions", "Simulation"),
 					LOCTEXT("SimulationOptionsTooltip", "Simulation options"),
-					FSlateIcon(FNiagaraEditorStyle::Get().GetStyleSetName(), "NiagaraEditor.SimulationOptions")
+					FSlateIcon(FAppStyle::Get().GetStyleSetName(), "AssetEditor.Simulate")
 				);
 			}
 			ToolbarBuilder.EndSection();
diff --git a/Engine/Plugins/FX/Niagara/Source/NiagaraEditorWidgets/Private/SNiagaraOverviewStack.cpp b/Engine/Plugins/FX/Niagara/Source/NiagaraEditorWidgets/Private/SNiagaraOverviewStack.cpp
--- a/Engine/Plugins/FX/Niagara/Source/NiagaraEditorWidgets/Private/SNiagaraOverviewStack.cpp
+++ b/Engine/Plugins/FX/Niagara/Source/NiagaraEditorWidgets/Private/SNiagaraOverviewStack.cpp
@@ -191,6 +191,25 @@ class SNiagaraSystemOverviewEntryListRow : public STableRow<UNiagaraStackEntry*>
 		],
 		InOwnerTableView);
 
+		if (StackEntry->IsA<UNiagaraStackItem>())
+		{
+			// Execution Category Item Highlight
+			RowOverlay->AddSlot()
+			.HAlign(HAlign_Left)
+			.VAlign(VAlign_Fill)
+			.Padding(FMargin(0))
+			[
+				SNew(SBorder)
+				.BorderImage(FAppStyle::Get().GetBrush("WhiteBrush"))
+				.BorderBackgroundColor(IconColor)
+				.Padding(FMargin(0))
+				[
+					SNew(SBox)
+					.WidthOverride(2)
+				]
+			];
+		}
+		
 		// we special case the red overlay for renderer items as they have their own scalability settings
 		if(UNiagaraStackRendererItem* RendererItem = Cast<UNiagaraStackRendererItem>(StackEntry))
 		{			
