- Remove PURE_VIRTUAL as it's used by other outputs now

#rnx

diff --git a/Engine/Plugins/FX/Niagara/Source/Niagara/Classes/NiagaraBakerOutput.h b/Engine/Plugins/FX/Niagara/Source/Niagara/Classes/NiagaraBakerOutput.h
--- a/Engine/Plugins/FX/Niagara/Source/Niagara/Classes/NiagaraBakerOutput.h
+++ b/Engine/Plugins/FX/Niagara/Source/Niagara/Classes/NiagaraBakerOutput.h
@@ -32,7 +32,7 @@ public:
 	FString OutputName;
 
 	/** Used to determine if the bake is out of date or not. */
-	virtual bool Equals(const UNiagaraBakerOutput& Other) const PURE_VIRTUAL(UNiagaraBakerOutput::Equals, return OutputName == Other.OutputName;);
+	virtual bool Equals(const UNiagaraBakerOutput& Other) const { return OutputName == Other.OutputName; }
 
 	/** Make an OutputName used when creating a new output */
 	virtual FString MakeOutputName() const;
