- Mark base curve class as abstract
- Convert other abstract functions to be PURE_VIRTUAL

#rnx

diff --git a/Engine/Plugins/FX/Niagara/Source/Niagara/Classes/NiagaraDataInterfaceCurveBase.h b/Engine/Plugins/FX/Niagara/Source/Niagara/Classes/NiagaraDataInterfaceCurveBase.h
--- a/Engine/Plugins/FX/Niagara/Source/Niagara/Classes/NiagaraDataInterfaceCurveBase.h
+++ b/Engine/Plugins/FX/Niagara/Source/Niagara/Classes/NiagaraDataInterfaceCurveBase.h
@@ -37,7 +37,7 @@ struct FNiagaraDataInterfaceProxyCurveBase : public FNiagaraDataInterfaceProxy
 };
 
 /** Base class for curve data interfaces which facilitates handling the curve data in a standardized way. */
-UCLASS(EditInlineNew, Category = "Curves", meta = (DisplayName = "Float Curve"))
+UCLASS(EditInlineNew, abstract)
 class NIAGARA_API UNiagaraDataInterfaceCurveBase : public UNiagaraDataInterface
 {
 public:
@@ -209,8 +209,8 @@ public:
 	virtual FName GetCurveSampleFunctionName() const PURE_VIRTUAL(UNiagaraDataInterfaceCurveBase::GetCurveSampleFunctionName(), return NAME_None;)
 #endif
 
-	virtual void UpdateTimeRanges() { checkf(false, TEXT("You must implement this function in your derived class")); }
-	virtual TArray<float> BuildLUT(int32 NumEntries) const { checkf(false, TEXT("You must implement this function in your derived class")); return TArray<float>(); }
+	virtual void UpdateTimeRanges() PURE_VIRTUAL(UNiagaraDataInterfaceCurveBase::UpdateTimeRanges(), )
+	virtual TArray<float> BuildLUT(int32 NumEntries) const PURE_VIRTUAL(UNiagaraDataInterfaceCurveBase::UpdateTimeRanges(), return TArray<float>(); )
 
 	FORCEINLINE float GetMinTime()const { return LUTMinTime; }
 	FORCEINLINE float GetMaxTime()const { return LUTMaxTime; }
