Added file include support for CustomHlsl node (based on contribution from amuTBKT)

#9018

diff --git a/Engine/Plugins/FX/Niagara/Source/NiagaraEditor/Private/NiagaraHlslTranslator.cpp b/Engine/Plugins/FX/Niagara/Source/NiagaraEditor/Private/NiagaraHlslTranslator.cpp
--- a/Engine/Plugins/FX/Niagara/Source/NiagaraEditor/Private/NiagaraHlslTranslator.cpp
+++ b/Engine/Plugins/FX/Niagara/Source/NiagaraEditor/Private/NiagaraHlslTranslator.cpp
@@ -39,11 +39,11 @@
 #include "NiagaraSimulationStageBase.h"
 #include "NiagaraTrace.h"
 #include "ShaderCore.h"
+#include "Misc/FileHelper.h"
 
 #define LOCTEXT_NAMESPACE "NiagaraCompiler"
 
 DECLARE_CYCLE_STAT(TEXT("Niagara - HlslTranslator - Translate"), STAT_NiagaraEditor_HlslTranslator_Translate, STATGROUP_NiagaraEditor);
-DECLARE_CYCLE_STAT(TEXT("Niagara - HlslTranslator - CloneGraphAndPrepareForCompilation"), STAT_NiagaraEditor_HlslTranslator_CloneGraphAndPrepareForCompilation, STATGROUP_NiagaraEditor);
 DECLARE_CYCLE_STAT(TEXT("Niagara - HlslTranslator - BuildParameterMapHlslDefinitions"), STAT_NiagaraEditor_HlslTranslator_BuildParameterMapHlslDefinitions, STATGROUP_NiagaraEditor);
 DECLARE_CYCLE_STAT(TEXT("Niagara - HlslTranslator - Emitter"), STAT_NiagaraEditor_Module_NiagaraHLSLTranslator_Emitter, STATGROUP_NiagaraEditor);
 DECLARE_CYCLE_STAT(TEXT("Niagara - HlslTranslator - MapGet"), STAT_NiagaraEditor_Module_NiagaraHLSLTranslator_MapGet, STATGROUP_NiagaraEditor);
@@ -76,8 +76,6 @@ DECLARE_CYCLE_STAT(TEXT("Niagara - HlslTranslator - GenerateFunctionSignature_Fi
 
 #define NIAGARA_SCOPE_CYCLE_COUNTER(x) //SCOPE_CYCLE_COUNTER(x)
 
-static FNiagaraShaderQueueTickable NiagaraShaderQueueProcessor;
-
 // not pretty. TODO: refactor
 // this will be called via a delegate from UNiagaraScript's cache for cook function,
 // because editor tickables aren't ticked during cooking
@@ -565,7 +563,6 @@ void FHlslNiagaraTranslator::GenerateFunctionSignature(ENiagaraScriptUsage Scrip
 
 FHlslNiagaraTranslator::FHlslNiagaraTranslator()
 	: Schema(GetDefault<UEdGraphSchema_Niagara>())
-	, TranslateResults()
 	, CurrentBodyChunkMode(ENiagaraCodeChunkMode::Body)
 	, ActiveStageIdx(-1)
 	, bInitializedDefaults(false)
@@ -580,6 +577,12 @@ FString FHlslNiagaraTranslator::GetFunctionDefinitions()
 	FString FwdDeclString;
 	FString DefinitionsString;
 
+	// add includes from custom hlsl nodes
+	for (const FNiagaraCustomHlslInclude& Include : FunctionIncludeFilePaths)
+	{
+		DefinitionsString += GetFunctionIncludeStatement(Include);
+	}
+
 	for (const auto& FuncPair : Functions)
 	{
 		FString Sig = GetFunctionSignature(FuncPair.Key);
@@ -5820,7 +5823,7 @@ FString FHlslNiagaraTranslator::GetParameterMapInstanceName(int32 ParamMapHistor
 	return ParameterMapInstanceName;
 }
 
-void FHlslNiagaraTranslator::Emitter(class UNiagaraNodeEmitter* EmitterNode, TArray<int32>& Inputs, TArray<int32>& Outputs)
+void FHlslNiagaraTranslator::Emitter(UNiagaraNodeEmitter* EmitterNode, TArray<int32>& Inputs, TArray<int32>& Outputs)
 {
 	NIAGARA_SCOPE_CYCLE_COUNTER(STAT_NiagaraEditor_Module_NiagaraHLSLTranslator_Emitter);
 
@@ -5896,7 +5899,7 @@ void FHlslNiagaraTranslator::Emitter(class UNiagaraNodeEmitter* EmitterNode, TAr
 	}
 
 	// We act like a function call here as the semantics are identical.
-	RegisterFunctionCall(ScriptUsage, Name, FullName, EmitterNode->NodeGuid, EmitterNode->GetEmitterHandleId().ToString(EGuidFormats::Digits), Source, Signature, false, FString(), Inputs, CallInputs, CallOutputs, Signature);
+	RegisterFunctionCall(ScriptUsage, Name, FullName, EmitterNode->NodeGuid, EmitterNode->GetEmitterHandleId().ToString(EGuidFormats::Digits), Source, Signature, false, FString(), {}, Inputs, CallInputs, CallOutputs, Signature);
 	GenerateFunctionCall(ScriptUsage, Signature, Inputs, Outputs);
 
 	// Clear out the parameter map writes to emitter module parameters as they should not be shared across emitters.
@@ -5948,8 +5951,8 @@ void FHlslNiagaraTranslator::ParameterMapGet(UNiagaraNodeParameterMapGet* GetNod
 		for (int32 i = 0; i < Outputs.Num(); i++)
 		{
 			Outputs[i] = INDEX_NONE;
-			return;
 		}
+		return;
 	}
 	else if (ParamMapHistoryIdx >= ParamMapHistories.Num())
 	{
@@ -5957,8 +5960,8 @@ void FHlslNiagaraTranslator::ParameterMapGet(UNiagaraNodeParameterMapGet* GetNod
 		for (int32 i = 0; i < Outputs.Num(); i++)
 		{
 			Outputs[i] = INDEX_NONE;
-			return;
 		}
+		return;
 	}
 
 	FString ParameterMapInstanceName = GetParameterMapInstanceName(ParamMapHistoryIdx);
@@ -6507,8 +6510,8 @@ void FHlslNiagaraTranslator::ReadDataSet(const FNiagaraDataSetID DataSet, const
 		for (int32 i = 0; i < Outputs.Num(); i++)
 		{
 			Outputs[i] = INDEX_NONE;
-			return;
 		}
+		return;
 	}
 	else if (ParamMapHistoryIdx >= ParamMapHistories.Num())
 	{
@@ -6516,8 +6519,8 @@ void FHlslNiagaraTranslator::ReadDataSet(const FNiagaraDataSetID DataSet, const
 		for (int32 i = 0; i < Outputs.Num(); i++)
 		{
 			Outputs[i] = INDEX_NONE;
-			return;
 		}
+		return;
 	}
 
 	TMap<int32, FDataSetAccessInfo>& Reads = DataSetReadInfo[(int32)AccessMode].FindOrAdd(DataSet);
@@ -6572,8 +6575,8 @@ void FHlslNiagaraTranslator::WriteDataSet(const FNiagaraDataSetID DataSet, const
 		for (int32 i = 0; i < Outputs.Num(); i++)
 		{
 			Outputs[i] = INDEX_NONE;
-			return;
 		}
+		return;
 	}
 	else if (ParamMapHistoryIdx >= ParamMapHistories.Num())
 	{
@@ -6581,8 +6584,8 @@ void FHlslNiagaraTranslator::WriteDataSet(const FNiagaraDataSetID DataSet, const
 		for (int32 i = 0; i < Outputs.Num(); i++)
 		{
 			Outputs[i] = INDEX_NONE;
-			return;
 		}
+		return;
 	}
 
 	if (DataSetWriteInfo[(int32)AccessMode].Find(DataSet))
@@ -6692,7 +6695,7 @@ int32 FHlslNiagaraTranslator::RegisterDataInterface(FNiagaraVariable& Var, UNiag
 	return Idx;
 }
 
-void FHlslNiagaraTranslator::Operation(class UNiagaraNodeOp* Operation, TArray<int32>& Inputs, TArray<int32>& Outputs)
+void FHlslNiagaraTranslator::Operation(UNiagaraNodeOp* Operation, TArray<int32>& Inputs, TArray<int32>& Outputs)
 {
 	NIAGARA_SCOPE_CYCLE_COUNTER(STAT_NiagaraEditor_HlslTranslator_Operation);
 
@@ -6849,7 +6852,6 @@ void FHlslNiagaraTranslator::FunctionCall(UNiagaraNodeFunctionCall* FunctionNode
 	if (!FunctionNode->IsNodeEnabled())
 	{
 		int32 InputPinIdx = INDEX_NONE;
-		int32 OutputPinIdx = INDEX_NONE;
 
 		for (int32 i = 0; i < CallInputs.Num(); i++)
 		{
@@ -6936,6 +6938,7 @@ void FHlslNiagaraTranslator::FunctionCall(UNiagaraNodeFunctionCall* FunctionNode
 	UNiagaraScriptSource* Source = nullptr;
 	bool bCustomHlsl = false;
 	FString CustomHlsl;
+	TArray<FNiagaraCustomHlslInclude> CustomHlslIncludeFilePaths;
 	FNiagaraFunctionSignature Signature = FunctionNode->Signature;
 
 	if (FunctionNode->FunctionScript)
@@ -6998,7 +7001,7 @@ void FHlslNiagaraTranslator::FunctionCall(UNiagaraNodeFunctionCall* FunctionNode
 						ResolvedVariable = FNiagaraVariable(TypeDef, *Prefix);
 					}
 
-					Signature.FunctionSpecifiers.Add(TEXT("CompilerTagKey")) =  ResolvedVariable.GetName();
+					Signature.FunctionSpecifiers.Add(FName("CompilerTagKey")) =  ResolvedVariable.GetName();
 
 					break;
 				}
@@ -7016,14 +7019,13 @@ void FHlslNiagaraTranslator::FunctionCall(UNiagaraNodeFunctionCall* FunctionNode
 		}
 	}
 
-	UNiagaraNodeCustomHlsl* CustomFunctionHlsl = Cast<UNiagaraNodeCustomHlsl>(FunctionNode);
-	if (CustomFunctionHlsl != nullptr)
+	if (UNiagaraNodeCustomHlsl* CustomFunctionHlsl = Cast<UNiagaraNodeCustomHlsl>(FunctionNode))
 	{
 		// All of the arguments here are resolved withing the HandleCustomHlsl function..
-		HandleCustomHlslNode(CustomFunctionHlsl, ScriptUsage, Name, FullName, bCustomHlsl, CustomHlsl, Signature, Inputs);
+		HandleCustomHlslNode(CustomFunctionHlsl, ScriptUsage, Name, FullName, bCustomHlsl, CustomHlsl, CustomHlslIncludeFilePaths, Signature, Inputs);
 	}
 
-	RegisterFunctionCall(ScriptUsage, Name, FullName, FunctionNode->NodeGuid, FString(), Source, Signature, bCustomHlsl, CustomHlsl, Inputs, CallInputs, CallOutputs, OutputSignature);
+	RegisterFunctionCall(ScriptUsage, Name, FullName, FunctionNode->NodeGuid, FString(), Source, Signature, bCustomHlsl, CustomHlsl, CustomHlslIncludeFilePaths, Inputs, CallInputs, CallOutputs, OutputSignature);
 
 	if (OutputSignature.IsValid() == false)
 	{
@@ -7619,7 +7621,7 @@ void FHlslNiagaraTranslator::ProcessCustomHlsl(const FString& InCustomHlsl, ENia
 	OutCustomHlsl = TEXT("\n") + OutCustomHlsl + TEXT("\n");
 }
 
-void FHlslNiagaraTranslator::HandleCustomHlslNode(UNiagaraNodeCustomHlsl* CustomFunctionHlsl, ENiagaraScriptUsage& OutScriptUsage, FString& OutName, FString& OutFullName, bool& bOutCustomHlsl, FString& OutCustomHlsl,
+void FHlslNiagaraTranslator::HandleCustomHlslNode(UNiagaraNodeCustomHlsl* CustomFunctionHlsl, ENiagaraScriptUsage& OutScriptUsage, FString& OutName, FString& OutFullName, bool& bOutCustomHlsl, FString& OutCustomHlsl, TArray<FNiagaraCustomHlslInclude>& OutCustomHlslIncludeFilePaths,
 	FNiagaraFunctionSignature& OutSignature, TArray<int32>& Inputs)
 {
 	NIAGARA_SCOPE_CYCLE_COUNTER(STAT_NiagaraEditor_Module_NiagaraHLSLTranslator_CustomHLSL);
@@ -7636,10 +7638,10 @@ void FHlslNiagaraTranslator::HandleCustomHlslNode(UNiagaraNodeCustomHlsl* Custom
 	OutSignature.Name = *OutName; // Force the name to be set to include the node guid for safety...
 	bOutCustomHlsl = true;
 	OutCustomHlsl = CustomFunctionHlsl->GetCustomHlsl();
+	CustomFunctionHlsl->GetIncludeFilePaths(OutCustomHlslIncludeFilePaths);
 
 	FNiagaraFunctionSignature InSignature = CustomFunctionHlsl->Signature;
 	ProcessCustomHlsl(CustomFunctionHlsl->GetCustomHlsl(), OutScriptUsage, InSignature, Inputs, CustomFunctionHlsl, OutCustomHlsl, OutSignature);
-	
 }
 
 void FHlslNiagaraTranslator::HandleDataInterfaceCall(FNiagaraScriptDataInterfaceCompileInfo& Info, const FNiagaraFunctionSignature& InMatchingSignature)
@@ -7702,7 +7704,7 @@ bool IsVariableWriteBeforeRead(const TArray<FNiagaraParameterMapHistory::FReadHi
 }
 
 void FHlslNiagaraTranslator::RegisterFunctionCall(ENiagaraScriptUsage ScriptUsage, const FString& InName, const FString& InFullName, const FGuid& CallNodeId, const FString& InFunctionNameSuffix, UNiagaraScriptSource* Source,
-                                                  FNiagaraFunctionSignature& InSignature, bool bIsCustomHlsl, const FString& InCustomHlsl, TArray<int32>& Inputs, TArrayView<UEdGraphPin* const> CallInputs, TArrayView<UEdGraphPin* const> CallOutputs,
+                                                  FNiagaraFunctionSignature& InSignature, bool bIsCustomHlsl, const FString& InCustomHlsl, const TArray<FNiagaraCustomHlslInclude>& InCustomHlslIncludeFilePaths, TArray<int32>& Inputs, TArrayView<UEdGraphPin* const> CallInputs, TArrayView<UEdGraphPin* const> CallOutputs,
                                                   FNiagaraFunctionSignature& OutSignature)
 {
 	NIAGARA_SCOPE_CYCLE_COUNTER(STAT_NiagaraEditor_Module_NiagaraHLSLTranslator_RegisterFunctionCall);
@@ -7993,6 +7995,11 @@ void FHlslNiagaraTranslator::RegisterFunctionCall(ENiagaraScriptUsage ScriptUsag
 
 				FunctionStageWriteTargets.Add(OutSignature, ActiveStageWriteTargets.Top());
 
+				for (const FNiagaraCustomHlslInclude& FileInclude : InCustomHlslIncludeFilePaths)
+				{
+					FunctionIncludeFilePaths.AddUnique(FileInclude);
+				}
+
 				ExitFunction();
 			}
 			else
@@ -8360,6 +8367,24 @@ FName FHlslNiagaraTranslator::GetDataInterfaceName(FName BaseName, const FString
 	return BaseName;
 }
 
+FString FHlslNiagaraTranslator::GetFunctionIncludeStatement(const FNiagaraCustomHlslInclude& Include) const
+{
+	TStringBuilder<128> IncludeStatement;
+
+	if (Include.bIsVirtual)
+	{
+		IncludeStatement.Appendf(TEXT("#include \"%s\"\n"), *Include.FilePath);
+	}
+	else if (FString FileContents; FFileHelper::LoadFileToString(FileContents, *Include.FilePath))
+	{
+		IncludeStatement.Appendf(TEXT("\n// included from %s\n"), *Include.FilePath);
+		IncludeStatement.Append(FileContents);
+		IncludeStatement.Append("\n");
+	}
+
+	return IncludeStatement.ToString();
+}
+
 FString FHlslNiagaraTranslator::GetFunctionSignature(const FNiagaraFunctionSignature& Sig)
 {
 	FString SigStr = TEXT("void ") + GetFunctionSignatureSymbol(Sig);
@@ -8540,7 +8565,6 @@ FString FHlslNiagaraTranslator::NamePathToString(const FString& Prefix, const FN
 	FString Value = Prefix;
 	FNiagaraTypeDefinition CurrentType = RootType;
 	bool bParentWasMatrix = (RootType == FNiagaraTypeDefinition::GetMatrix4Def());
-	int32 ParentMatrixRow = -1;
 	for (int32 i = 0; i < NamePath.Num(); i++)
 	{
 		FString Name = NamePath[i].ToString();
@@ -8582,7 +8606,7 @@ FString FHlslNiagaraTranslator::GenerateAssignment(const FNiagaraTypeDefinition&
 	return DestinationDefinition + " = " + SourceDefinition;
 }
 
-void FHlslNiagaraTranslator::Convert(class UNiagaraNodeConvert* Convert, TArrayView<const int32> Inputs, TArray<int32>& Outputs)
+void FHlslNiagaraTranslator::Convert(UNiagaraNodeConvert* Convert, TArrayView<const int32> Inputs, TArray<int32>& Outputs)
 {
 	if (ValidateTypePins(Convert) == false)
 	{
@@ -9456,7 +9480,6 @@ FString FHlslNiagaraTranslator::BuildHLSLStructDecl(const FNiagaraTypeDefinition
 		FString Decl = "struct " + GetStructHlslTypeName(Type) + "\n{\n";
 
 		int32 StructSize = 0;
-		int32 DummyIndex = 0;
 		for (TFieldIterator<FProperty> PropertyIt(Type.GetStruct(), EFieldIteratorFlags::IncludeSuper); PropertyIt; ++PropertyIt)
 		{
 			FProperty* Property = *PropertyIt;
diff --git a/Engine/Plugins/FX/Niagara/Source/NiagaraEditor/Private/NiagaraHlslTranslator.h b/Engine/Plugins/FX/Niagara/Source/NiagaraEditor/Private/NiagaraHlslTranslator.h
--- a/Engine/Plugins/FX/Niagara/Source/NiagaraEditor/Private/NiagaraHlslTranslator.h
+++ b/Engine/Plugins/FX/Niagara/Source/NiagaraEditor/Private/NiagaraHlslTranslator.h
@@ -8,7 +8,6 @@
 #include "NiagaraParameterMapHistory.h"
 #include "NiagaraShaderCompilationManager.h"
 #include "NiagaraDataInterface.h"
-#include "TickableEditorObject.h"
 #include "NiagaraScriptSource.h"
 
 class Error;
@@ -32,12 +31,23 @@ enum class ENiagaraDataSetAccessMode : uint8
 	Num UMETA(Hidden),
 };
 
+struct FNiagaraCustomHlslInclude
+{
+	bool bIsVirtual = false;
+	FString FilePath;
+
+	bool operator==(const FNiagaraCustomHlslInclude& Other) const 
+	{ 
+		return bIsVirtual == Other.bIsVirtual && FilePath == Other.FilePath; 
+	}
+};
+
 
 /** Defines information about the results of a Niagara script compile. */
 struct FNiagaraTranslateResults
 {
 	/** Whether or not HLSL generation was successful */
-	bool bHLSLGenSucceeded;
+	bool bHLSLGenSucceeded = false;;
 
 	/** A results log with messages, warnings, and errors which occurred during the compile. */
 	TArray<FNiagaraCompileEvent> CompileEvents;
@@ -125,8 +135,8 @@ public:
 	bool bUseRapidIterationParams = true;
 	bool bDisableDebugSwitches = false;
 
-	UEnum* ENiagaraScriptCompileStatusEnum;
-	UEnum* ENiagaraScriptUsageEnum;
+	UEnum* ENiagaraScriptCompileStatusEnum = nullptr;
+	UEnum* ENiagaraScriptUsageEnum = nullptr;
 
 	TArray<FNiagaraVariable> RapidIterationParams;
 
@@ -477,12 +487,14 @@ protected:
 	};
 
 	TMap<FNiagaraFunctionSignature, FNiagaraFunctionBody> Functions;
-	TMap<FNiagaraFunctionSignature, TArray<FName> > FunctionStageWriteTargets;
+	TArray<FNiagaraCustomHlslInclude> FunctionIncludeFilePaths;
+	TMap<FNiagaraFunctionSignature, TArray<FName>> FunctionStageWriteTargets;
 	TArray<TArray<FName>> ActiveStageWriteTargets;
 
-	void RegisterFunctionCall(ENiagaraScriptUsage ScriptUsage, const FString& InName, const FString& InFullName, const FGuid& CallNodeId, const FString& InFunctionNameSuffix, UNiagaraScriptSource* Source, FNiagaraFunctionSignature& InSignature, bool bIsCustomHlsl, const FString& InCustomHlsl, TArray<int32>& Inputs, TArrayView<UEdGraphPin* const> CallInputs, TArrayView<UEdGraphPin* const> CallOutputs,
+	void RegisterFunctionCall(ENiagaraScriptUsage ScriptUsage, const FString& InName, const FString& InFullName, const FGuid& CallNodeId, const FString& InFunctionNameSuffix, UNiagaraScriptSource* Source, FNiagaraFunctionSignature& InSignature, bool bIsCustomHlsl, const FString& InCustomHlsl, const TArray<FNiagaraCustomHlslInclude>& InCustomHlslIncludeFilePaths, TArray<int32>& Inputs, TArrayView<UEdGraphPin* const> CallInputs, TArrayView<UEdGraphPin* const> CallOutputs,
 		FNiagaraFunctionSignature& OutSignature);
 	void GenerateFunctionCall(ENiagaraScriptUsage ScriptUsage, FNiagaraFunctionSignature& FunctionSignature, TArrayView<const int32> Inputs, TArray<int32>& Outputs);
+	FString GetFunctionIncludeStatement(const FNiagaraCustomHlslInclude& Include) const;
 	FString GetFunctionSignature(const FNiagaraFunctionSignature& Sig);
 
 	/** Compiles an output Pin on a graph node. Caches the result for any future inputs connected to it. */
@@ -600,8 +612,6 @@ public:
 	void GatherVariableForDataSetAccess(const FNiagaraVariable& Variable, FString Format, int32& RegisterIdxInt, int32& RegisterIdxFloat, int32& RegisterIdxHalf, int32 DataSetIndex, FString InstanceIdxSymbol, FString &HlslOutput, bool bWriteHLSL = true);
 	void GatherComponentsForDataSetAccess(UScriptStruct* Struct, FString VariableSymbol, bool bMatrixRoot, TArray<FString>& Components, TArray<ENiagaraBaseTypes>& Types);
 
-	FString CompileDataInterfaceFunction(UNiagaraDataInterface* DataInterface, FNiagaraFunctionSignature& Signature);
-
 	virtual void FunctionCall(UNiagaraNodeFunctionCall* FunctionNode, TArray<int32>& Inputs, TArray<int32>& Outputs);
 	void EnterFunctionCallNode(const TSet<FName>& UnusedInputs);
 	void ExitFunctionCallNode();
@@ -697,7 +707,7 @@ private:
 	FString ComputeMatrixRowAccess(const FString& Name);
 
 	bool ParseDIFunctionSpecifiers(UNiagaraNode* NodeForErrorReporting, FNiagaraFunctionSignature& Sig, TArray<FString>& Tokens, int32& TokenIdx);
-	void HandleCustomHlslNode(UNiagaraNodeCustomHlsl* CustomFunctionHlsl, ENiagaraScriptUsage& OutScriptUsage, FString& OutName, FString& OutFullName, bool& bOutCustomHlsl, FString& OutCustomHlsl,
+	void HandleCustomHlslNode(UNiagaraNodeCustomHlsl* CustomFunctionHlsl, ENiagaraScriptUsage& OutScriptUsage, FString& OutName, FString& OutFullName, bool& bOutCustomHlsl, FString& OutCustomHlsl, TArray<FNiagaraCustomHlslInclude>& OutCustomHlslIncludeFilePaths,
 		FNiagaraFunctionSignature& OutSignature, TArray<int32>& Inputs);
 	void ProcessCustomHlsl(const FString& InCustomHlsl, ENiagaraScriptUsage InUsage, const FNiagaraFunctionSignature& InSignature, const TArray<int32>& Inputs, UNiagaraNode* NodeForErrorReporting, FString& OutCustomHlsl, FNiagaraFunctionSignature& OutSignature);
 	void HandleSimStageSetupAndTeardown(int32 InWhichStage, FString& OutHlsl);
@@ -740,8 +750,6 @@ private:
 	void GenerateFunctionSignature(ENiagaraScriptUsage ScriptUsage, FString InName, const FString& InFullName, const FString& InFunctionNameSuffix, UNiagaraGraph* FuncGraph, TArray<int32>& Inputs, 
 		bool bHadNumericInputs, bool bHasParameterMapParameters, TArray<UEdGraphPin*> StaticSwitchValues, FNiagaraFunctionSignature& OutSig) const;
 
-	UNiagaraGraph* CloneGraphAndPrepareForCompilation(const UNiagaraScript* InScript, const UNiagaraScriptSource* InSource, bool bClearErrors);
-
 	bool ShouldInterpolateParameter(const FNiagaraVariable& Parameter);
 
 	void UpdateStaticSwitchConstants(UEdGraphNode* Node);
diff --git a/Engine/Plugins/FX/Niagara/Source/NiagaraEditor/Private/NiagaraNodeCustomHlsl.cpp b/Engine/Plugins/FX/Niagara/Source/NiagaraEditor/Private/NiagaraNodeCustomHlsl.cpp
--- a/Engine/Plugins/FX/Niagara/Source/NiagaraEditor/Private/NiagaraNodeCustomHlsl.cpp
+++ b/Engine/Plugins/FX/Niagara/Source/NiagaraEditor/Private/NiagaraNodeCustomHlsl.cpp
@@ -1,12 +1,11 @@
 // Copyright Epic Games, Inc. All Rights Reserved.
 
 #include "NiagaraNodeCustomHlsl.h"
-#include "NiagaraCustomVersion.h"
 #include "SNiagaraGraphNodeCustomHlsl.h"
 #include "EdGraphSchema_Niagara.h"
 #include "ScopedTransaction.h"
-#include "EdGraphSchema_Niagara.h"
 #include "NiagaraGraph.h"
+#include "Misc/FileHelper.h"
 
 #define LOCTEXT_NAMESPACE "NiagaraNodeCustomHlsl"
 
@@ -39,6 +38,24 @@ void UNiagaraNodeCustomHlsl::SetCustomHlsl(const FString& InCustomHlsl)
 	}
 }
 
+void UNiagaraNodeCustomHlsl::GetIncludeFilePaths(TArray<FNiagaraCustomHlslInclude>& OutCustomHlslIncludeFilePaths) const
+{
+	for (const FString& FilePath : VirtualIncludeFilePaths)
+	{
+		if (!FilePath.IsEmpty())
+		{			
+			OutCustomHlslIncludeFilePaths.Add({true, FilePath});
+		}
+	}
+	for (const auto& [FilePath] : AbsoluteIncludeFilePaths)
+	{
+		if (!FilePath.IsEmpty())
+		{			
+			OutCustomHlslIncludeFilePaths.Add({false, FilePath});
+		}
+	}
+}
+
 TSharedPtr<SGraphNode> UNiagaraNodeCustomHlsl::CreateVisualWidget()
 {
 	return SNew(SNiagaraGraphNodeCustomHlsl, this);
@@ -174,17 +191,33 @@ bool UNiagaraNodeCustomHlsl::GetTokensFromString(const FString& InHlsl, TArray<F
 bool UNiagaraNodeCustomHlsl::GetTokens(TArray<FString>& OutTokens, bool IncludeComments) const
 {
 	FString HlslData = *CustomHlsl;
-	return UNiagaraNodeCustomHlsl::GetTokensFromString(HlslData, OutTokens, IncludeComments);
+	return GetTokensFromString(HlslData, OutTokens, IncludeComments);
 }
 
-
-#if WITH_EDITOR
-
-void UNiagaraNodeCustomHlsl::PostEditChangeProperty(struct FPropertyChangedEvent& PropertyChangedEvent)
+void UNiagaraNodeCustomHlsl::PostEditChangeProperty(FPropertyChangedEvent& PropertyChangedEvent)
 {
 	Super::PostEditChangeProperty(PropertyChangedEvent);
 
-	if (PropertyChangedEvent.Property && PropertyChangedEvent.Property->GetFName() == GET_MEMBER_NAME_CHECKED(UNiagaraNodeCustomHlsl, CustomHlsl))
+	bool bRequiresRecompilation = false;
+
+	if (PropertyChangedEvent.Property)
+	{
+		const FName PropertyName = PropertyChangedEvent.Property->GetFName();
+		if (PropertyName == GET_MEMBER_NAME_CHECKED(UNiagaraNodeCustomHlsl, CustomHlsl))
+		{
+			bRequiresRecompilation = true;
+		}
+		else if (PropertyName == GET_MEMBER_NAME_CHECKED(UNiagaraNodeCustomHlsl, AbsoluteIncludeFilePaths))
+		{
+			bRequiresRecompilation = true;
+		}
+		else if (PropertyName == GET_MEMBER_NAME_CHECKED(UNiagaraNodeCustomHlsl, VirtualIncludeFilePaths))
+		{
+			bRequiresRecompilation = true;
+		}
+	}
+
+	if (bRequiresRecompilation)
 	{
 		RefreshFromExternalChanges();
 		GetNiagaraGraph()->NotifyGraphNeedsRecompile();
@@ -200,8 +233,6 @@ void UNiagaraNodeCustomHlsl::InitAsCustomHlslDynamicInput(const FNiagaraTypeDefi
 	ScriptUsage = ENiagaraScriptUsage::DynamicInput;
 }
 
-#endif
-
 bool UNiagaraNodeCustomHlsl::IsPinNameEditableUponCreation(const UEdGraphPin* GraphPinObj) const
 {
 	if (GraphPinObj == PinPendingRename && ScriptUsage != ENiagaraScriptUsage::DynamicInput)
@@ -340,10 +371,10 @@ void UNiagaraNodeCustomHlsl::BuildParameterMapHistory(FNiagaraParameterMapHistor
 	GetOutputPins(OutputPins);
 
 	int32 ParamMapIdx = INDEX_NONE;
-	TArray<FNiagaraVariable> LocalVars;
 	// This only works currently if the input pins are in the same order as the signature pins.
 	if (InputPins.Num() == Signature.Inputs.Num() + 1 && OutputPins.Num() == Signature.Outputs.Num() + 1)// the add pin is extra
 	{
+		TArray<FNiagaraVariable> LocalVars;
 		bool bHasParamMapInput = false;
 		bool bHasParamMapOutput = false;
 		for (int32 i = 0; i < InputPins.Num(); i++)
@@ -407,7 +438,6 @@ void UNiagaraNodeCustomHlsl::BuildParameterMapHistory(FNiagaraParameterMapHistor
 						// There is one possible path here, one where we're using the namespace as-is from the valid list.
 						if (Tokens[i].StartsWith(ValidNamespace, ESearchCase::CaseSensitive))
 						{
-							bool bUsedDefault = false;
 							OutHistory.HandleExternalVariableRead(ParamMapIdx, *Tokens[i]);
 						}
 					}
@@ -417,6 +447,51 @@ void UNiagaraNodeCustomHlsl::BuildParameterMapHistory(FNiagaraParameterMapHistor
 	}
 }
 
+FNiagaraCompileHash GetFileContentHash(const FString& FileContents)
+{
+	FSHA1 CompileHash;
+	CompileHash.UpdateWithString(*FileContents, FileContents.Len());
+	CompileHash.Final();
+
+	TArray<uint8> DataHash;
+	DataHash.AddUninitialized(FSHA1::DigestSize);
+	CompileHash.GetHash(DataHash.GetData());
+
+	return FNiagaraCompileHash(DataHash);
+}
+
+void UNiagaraNodeCustomHlsl::GatherExternalDependencyData(ENiagaraScriptUsage InMasterUsage, const FGuid& InMasterUsageId, TArray<FNiagaraCompileHash>& InReferencedCompileHashes, TArray<FString>& InReferencedObjs) const
+{
+	for (const auto& [IncludePath] : AbsoluteIncludeFilePaths)
+	{
+		if (IncludePath.IsEmpty())
+		{
+			continue;
+		}
+
+		if (FString FileContents; FFileHelper::LoadFileToString(FileContents, *IncludePath))
+		{
+			InReferencedObjs.AddUnique(IncludePath);
+			InReferencedCompileHashes.AddUnique(GetFileContentHash(FileContents));
+		}
+	}
+
+	for (const FString& IncludePath : VirtualIncludeFilePaths)
+	{
+		if (IncludePath.IsEmpty())
+		{
+			continue;
+		}
+		
+		FString FileContents;
+		if(LoadShaderSourceFile(*IncludePath, SP_PCD3D_SM5, &FileContents, nullptr))
+		{
+			InReferencedObjs.AddUnique(IncludePath);
+			InReferencedCompileHashes.AddUnique(GetFileContentHash(FileContents));
+		}
+	}
+}
+
 // Replace items in the tokens array if they start with the src string or optionally src string and a namespace delimiter
 uint32 UNiagaraNodeCustomHlsl::ReplaceExactMatchTokens(TArray<FString>& Tokens, const FString& SrcString, const FString& ReplaceString, bool bAllowNamespaceSeparation)
 {
@@ -443,17 +518,14 @@ uint32 UNiagaraNodeCustomHlsl::ReplaceExactMatchTokens(TArray<FString>& Tokens,
 
 bool UNiagaraNodeCustomHlsl::AllowNiagaraTypeForAddPin(const FNiagaraTypeDefinition& InType) const
 {
-	if (Super::AllowNiagaraTypeForAddPin(InType) || InType.IsDataInterface())
-		return true;
-	else
-		return false;
+	return Super::AllowNiagaraTypeForAddPin(InType) || InType.IsDataInterface();
 }
 
 bool UNiagaraNodeCustomHlsl::AllowNiagaraTypeForAddPin(const FNiagaraTypeDefinition& InType, EEdGraphPinDirection InDirection) const
 {
 	if (AllowNiagaraTypeForAddPin(InType))
 	{
-		if (InType.IsStatic() && InDirection == EEdGraphPinDirection::EGPD_Output)
+		if (InType.IsStatic() && InDirection == EGPD_Output)
 			return false;
 		else
 			return true;
diff --git a/Engine/Plugins/FX/Niagara/Source/NiagaraEditor/Public/NiagaraNodeCustomHlsl.h b/Engine/Plugins/FX/Niagara/Source/NiagaraEditor/Public/NiagaraNodeCustomHlsl.h
--- a/Engine/Plugins/FX/Niagara/Source/NiagaraEditor/Public/NiagaraNodeCustomHlsl.h
+++ b/Engine/Plugins/FX/Niagara/Source/NiagaraEditor/Public/NiagaraNodeCustomHlsl.h
@@ -3,13 +3,12 @@
 #pragma once
 
 #include "CoreMinimal.h"
-#include "UObject/ObjectMacros.h"
-#include "NiagaraEditorCommon.h"
 #include "NiagaraNodeFunctionCall.h"
 #include "NiagaraNodeInput.h"
 #include "NiagaraNodeCustomHlsl.generated.h"
 
 class UNiagaraScript;
+struct FNiagaraCustomHlslInclude;
 
 UCLASS(MinimalAPI)
 class UNiagaraNodeCustomHlsl : public UNiagaraNodeFunctionCall
@@ -20,6 +19,8 @@ public:
 	const FString& GetCustomHlsl() const;
 	void SetCustomHlsl(const FString& InCustomHlsl);
 
+	void GetIncludeFilePaths(TArray<FNiagaraCustomHlslInclude>& OutCustomHlslIncludeFilePaths) const;
+
 	UPROPERTY()
 	ENiagaraScriptUsage ScriptUsage;
 
@@ -33,6 +34,7 @@ public:
 	bool GetTokens(TArray<FString>& OutTokens, bool IncludeComments = true) const;
 
 	virtual void BuildParameterMapHistory(FNiagaraParameterMapHistoryBuilder& OutHistory, bool bRecursive = true, bool bFilterForCompilation = true) const override;
+	virtual void GatherExternalDependencyData(ENiagaraScriptUsage InMasterUsage, const FGuid& InMasterUsageId, TArray<FNiagaraCompileHash>& InReferencedCompileHashes, TArray<FString>& InReferencedObjs) const override;
 
 	// Replace items in the tokens array if they start with the src string or optionally src string and a namespace delimiter
 	static uint32 ReplaceExactMatchTokens(TArray<FString>& Tokens, const FString& SrcString, const FString& ReplaceString, bool bAllowNamespaceSeparation);
@@ -43,12 +45,10 @@ public:
 
 	static bool GetTokensFromString(const FString& InHlsl, TArray<FString>& OutTokens, bool IncludeComments = true);
 
-#if WITH_EDITOR
 	virtual void PostEditChangeProperty(struct FPropertyChangedEvent& PropertyChangedEvent) override;
 
 	void InitAsCustomHlslDynamicInput(const FNiagaraTypeDefinition& OutputType);
 
-#endif
 protected:
 	virtual bool AllowDynamicPins() const override { return true; }
 	virtual bool GetValidateDataInterfaces() const override { return false; }
@@ -83,4 +83,13 @@ protected:
 private:
 	UPROPERTY(EditAnywhere, Category = "HLSL", meta = (MultiLine = true))
 	FString CustomHlsl;
+
+	// Links to hlsl files that will be included by the translator. These external files are not watched by the engine, so changes to them do not automatically trigger a recompile of Niagara scripts.
+	UPROPERTY(EditAnywhere, Category = "HLSL")
+	TArray<FFilePath> AbsoluteIncludeFilePaths;
+
+	// Links to hlsl files that will be included by the translator. These paths are resolved with the virtual shader paths registered in the engine.
+	// For example, /Plugin/FX/Niagara maps to /Engine/Plugins/FX/Niagara/Shaders. Custom mappings can be added via AddShaderSourceDirectoryMapping().
+	UPROPERTY(EditAnywhere, Category = "HLSL")
+	TArray<FString> VirtualIncludeFilePaths;
 };
