- Fix for creating interpolated cbuffers when the emitter does not require them
- This fixes invalid memory access

#rnx

diff --git a/Engine/Plugins/FX/Niagara/Source/Niagara/Classes/NiagaraGPUSystemTick.h b/Engine/Plugins/FX/Niagara/Source/Niagara/Classes/NiagaraGPUSystemTick.h
--- a/Engine/Plugins/FX/Niagara/Source/Niagara/Classes/NiagaraGPUSystemTick.h
+++ b/Engine/Plugins/FX/Niagara/Source/Niagara/Classes/NiagaraGPUSystemTick.h
@@ -36,6 +36,7 @@ struct FNiagaraComputeInstanceData
 	FNiagaraGpuSpawnInfo SpawnInfo;
 	uint8* EmitterParamData = nullptr;
 	uint8* ExternalParamData = nullptr;
+	uint32 ExternalParamDataSize = 0;
 	FNiagaraComputeExecutionContext* Context = nullptr;
 	TArray<FNiagaraDataInterfaceProxy*> DataInterfaceProxies;
 	TArray<FNiagaraDataInterfaceProxyRW*> IterationDataInterfaceProxies;
diff --git a/Engine/Plugins/FX/Niagara/Source/Niagara/Private/NiagaraGPUSystemTick.cpp b/Engine/Plugins/FX/Niagara/Source/Niagara/Private/NiagaraGPUSystemTick.cpp
--- a/Engine/Plugins/FX/Niagara/Source/Niagara/Private/NiagaraGPUSystemTick.cpp
+++ b/Engine/Plugins/FX/Niagara/Source/Niagara/Private/NiagaraGPUSystemTick.cpp
@@ -152,13 +152,12 @@ void FNiagaraGPUSystemTick::Init(FNiagaraSystemInstance* InSystemInstance)
 			}
 			InstanceData->ParticleCountFence = GPUContext->ParticleCountReadFence;
 
-			int32 ParmSize = GPUContext->CombinedParamStore.GetPaddedParameterSizeInBytes();
-
 			InstanceData->EmitterParamData = ParamDataBufferPtr;
 			ParamDataBufferPtr += InterpFactor * sizeof(FNiagaraEmitterParameters);
 
 			InstanceData->ExternalParamData = ParamDataBufferPtr;
-			ParamDataBufferPtr += ParmSize;
+			InstanceData->ExternalParamDataSize = GPUContext->CombinedParamStore.GetPaddedParameterSizeInBytes();
+			ParamDataBufferPtr += InstanceData->ExternalParamDataSize;
 
 			// actually copy all of the data over
 			FMemory::Memcpy(InstanceData->EmitterParamData, &InSystemInstance->GetEmitterParameters(EmitterIdx), sizeof(FNiagaraEmitterParameters));
@@ -167,7 +166,7 @@ void FNiagaraGPUSystemTick::Init(FNiagaraSystemInstance* InSystemInstance)
 				FMemory::Memcpy(InstanceData->EmitterParamData + sizeof(FNiagaraEmitterParameters), &InSystemInstance->GetEmitterParameters(EmitterIdx, true), sizeof(FNiagaraEmitterParameters));
 			}
 
-			GPUContext->CombinedParamStore.CopyParameterDataToPaddedBuffer(InstanceData->ExternalParamData, ParmSize);
+			GPUContext->CombinedParamStore.CopyParameterDataToPaddedBuffer(InstanceData->ExternalParamData, InstanceData->ExternalParamDataSize);
 
 			// Calling PostTick will push current -> previous parameters this must be done after copying the parameter data
 			GPUContext->PostTick();
@@ -309,7 +308,9 @@ const uint8* FNiagaraGPUSystemTick::GetUniformBufferSource(EUniformBufferType Ty
 		{
 			// External parameters are pushed from the combined parameters store, split into two where first half is current second is previous
 			check(Instance && Instance->Context);
-			return Instance->ExternalParamData + (Current ? 0 : Instance->Context->ExternalCBufferLayoutSize);
+			check((Current ? 0 : Instance->Context->ExternalCBufferLayout->ConstantBufferSize) + Instance->Context->ExternalCBufferLayout->ConstantBufferSize <= Instance->ExternalParamDataSize);
+
+			return Instance->ExternalParamData + (Current ? 0 : Instance->Context->ExternalCBufferLayout->ConstantBufferSize);
 		}
 	}
 
diff --git a/Engine/Plugins/FX/Niagara/Source/Niagara/Private/NiagaraGpuComputeDispatch.cpp b/Engine/Plugins/FX/Niagara/Source/Niagara/Private/NiagaraGpuComputeDispatch.cpp
--- a/Engine/Plugins/FX/Niagara/Source/Niagara/Private/NiagaraGpuComputeDispatch.cpp
+++ b/Engine/Plugins/FX/Niagara/Source/Niagara/Private/NiagaraGpuComputeDispatch.cpp
@@ -322,8 +322,6 @@ void FNiagaraGpuComputeDispatch::BuildConstantBuffers(FNiagaraGPUSystemTick& Tic
 			BoundParameterCounts[FNiagaraGPUSystemTick::UBT_Global][InterpIt] += Instance.Context->GPUScript_RT->IsGlobalConstantBufferUsed_RenderThread(InterpIt) ? 1 : 0;
 			BoundParameterCounts[FNiagaraGPUSystemTick::UBT_System][InterpIt] += Instance.Context->GPUScript_RT->IsSystemConstantBufferUsed_RenderThread(InterpIt) ? 1 : 0;
 			BoundParameterCounts[FNiagaraGPUSystemTick::UBT_Owner][InterpIt] += Instance.Context->GPUScript_RT->IsOwnerConstantBufferUsed_RenderThread(InterpIt) ? 1 : 0;
-			BoundParameterCounts[FNiagaraGPUSystemTick::UBT_Emitter][InterpIt] += Instance.Context->GPUScript_RT->IsEmitterConstantBufferUsed_RenderThread(InterpIt) ? 1 : 0;
-			BoundParameterCounts[FNiagaraGPUSystemTick::UBT_External][InterpIt] += Instance.Context->GPUScript_RT->IsExternalConstantBufferUsed_RenderThread(InterpIt) ? 1 : 0;
 		}
 	}
 
@@ -331,6 +329,7 @@ void FNiagaraGpuComputeDispatch::BuildConstantBuffers(FNiagaraGPUSystemTick& Tic
 	const int32 BufferCount = InterpScale * (FNiagaraGPUSystemTick::UBT_NumSystemTypes + FNiagaraGPUSystemTick::UBT_NumInstanceTypes * Tick.InstanceCount);
 
 	Tick.UniformBuffers.Empty(BufferCount);
+	Tick.UniformBuffers.AddDefaulted(BufferCount);
 
 	const FRHIUniformBufferLayout* SystemLayouts[FNiagaraGPUSystemTick::UBT_NumSystemTypes] =
 	{
@@ -339,47 +338,59 @@ void FNiagaraGpuComputeDispatch::BuildConstantBuffers(FNiagaraGPUSystemTick& Tic
 		OwnerCBufferLayout
 	};
 
-	for (int32 InterpIt = 0; InterpIt < InterpScale; ++InterpIt)
+	int32 CurrentBuffer = 0;
+	for (int32 InterpIt=0; InterpIt < InterpScale; ++InterpIt)
 	{
 		for (int32 SystemTypeIt = FNiagaraGPUSystemTick::UBT_FirstSystemType; SystemTypeIt < FNiagaraGPUSystemTick::UBT_NumSystemTypes; ++SystemTypeIt)
 		{
-			FUniformBufferRHIRef BufferRef;
-
 			if (BoundParameterCounts[SystemTypeIt][InterpIt])
 			{
-				BufferRef = RHICreateUniformBuffer(
+				Tick.UniformBuffers[CurrentBuffer] = RHICreateUniformBuffer(
 					Tick.GetUniformBufferSource((FNiagaraGPUSystemTick::EUniformBufferType) SystemTypeIt, nullptr, !InterpIt),
 					SystemLayouts[SystemTypeIt],
 					((BoundParameterCounts[SystemTypeIt][InterpIt] > 1) || HasMultipleStages)
 						? EUniformBufferUsage::UniformBuffer_SingleFrame
 						: EUniformBufferUsage::UniformBuffer_SingleDraw);
 			}
-
-			Tick.UniformBuffers.Add(BufferRef);
+			++CurrentBuffer;
 		}
 
-		for (int32 InstanceTypeIt = FNiagaraGPUSystemTick::UBT_FirstInstanceType; InstanceTypeIt < FNiagaraGPUSystemTick::UBT_NumTypes; ++InstanceTypeIt)
+		// Build emitter constant buffers
+		for (const FNiagaraComputeInstanceData& Instance : InstanceData)
 		{
-			for (const FNiagaraComputeInstanceData& Instance : InstanceData)
+			if (InterpIt == 0 || Instance.Context->HasInterpolationParameters)
 			{
-				FUniformBufferRHIRef BufferRef;
+				if (Instance.Context->GPUScript_RT->IsEmitterConstantBufferUsed_RenderThread(InterpIt))
+				{
+					Tick.UniformBuffers[CurrentBuffer] = RHICreateUniformBuffer(
+						Tick.GetUniformBufferSource(FNiagaraGPUSystemTick::UBT_Emitter, &Instance, !InterpIt),
+						EmitterCBufferLayout,
+						Instance.bHasMultipleStages ? EUniformBufferUsage::UniformBuffer_SingleFrame : EUniformBufferUsage::UniformBuffer_SingleDraw
+					);
+				}
+			}
+			++CurrentBuffer;
+		}
 
-				if (BoundParameterCounts[InstanceTypeIt][InterpIt])
+		// Build external constant buffers
+		for (const FNiagaraComputeInstanceData& Instance : InstanceData)
+		{
+			if (InterpIt == 0 || Instance.Context->HasInterpolationParameters)
+			{
+				FNiagaraRHIUniformBufferLayout* ExternalCBufferLayout = Instance.Context->ExternalCBufferLayout;
+				if (Instance.Context->GPUScript_RT->IsExternalConstantBufferUsed_RenderThread(InterpIt))
 				{
-					FRHIUniformBufferLayout* Layout = InstanceTypeIt == FNiagaraGPUSystemTick::UBT_Emitter ? EmitterCBufferLayout : Instance.Context->ExternalCBufferLayout;
-					if (Layout && (Layout->Resources.Num() > 0 || Layout->ConstantBufferSize > 0))
+					if (ensure(ExternalCBufferLayout && (ExternalCBufferLayout->Resources.Num() > 0 || ExternalCBufferLayout->ConstantBufferSize > 0)))
 					{
-						BufferRef = RHICreateUniformBuffer(
-							Tick.GetUniformBufferSource((FNiagaraGPUSystemTick::EUniformBufferType) InstanceTypeIt, &Instance, !InterpIt),
-							Layout,
-							((BoundParameterCounts[InstanceTypeIt][InterpIt] > 1) || HasMultipleStages)
-							? EUniformBufferUsage::UniformBuffer_SingleFrame
-							: EUniformBufferUsage::UniformBuffer_SingleDraw);
+						Tick.UniformBuffers[CurrentBuffer] = RHICreateUniformBuffer(
+							Tick.GetUniformBufferSource(FNiagaraGPUSystemTick::UBT_External, &Instance, !InterpIt),
+							ExternalCBufferLayout,
+							Instance.bHasMultipleStages ? EUniformBufferUsage::UniformBuffer_SingleFrame : EUniformBufferUsage::UniformBuffer_SingleDraw
+						);
 					}
 				}
-
-				Tick.UniformBuffers.Add(BufferRef);
 			}
+			++CurrentBuffer;
 		}
 	}
 }
@@ -1637,6 +1648,7 @@ void FNiagaraGpuComputeDispatch::DispatchStage(FRHICommandList& RHICmdList, cons
 	auto SetConstantBuffer =
 		[](FRHICommandList& RHICmdList, FRHIComputeShader* ComputeShader, const FShaderUniformBufferParameter& BufferParam, const FUniformBufferRHIRef& UniformBuffer)
 		{
+			check(BufferParam.IsBound() == false || UniformBuffer.IsValid());
 			if (BufferParam.IsBound() && UniformBuffer.IsValid())
 			{
 				RHICmdList.SetShaderUniformBuffer(ComputeShader, BufferParam.GetBaseIndex(), UniformBuffer);
