- Convert Texture Data Interfaces to use template & parameter structure bindings

diff --git a/Engine/Plugins/FX/Niagara/Shaders/Private/NiagaraDataInterfaceCubeTextureTemplate.ush b/Engine/Plugins/FX/Niagara/Shaders/Private/NiagaraDataInterfaceCubeTextureTemplate.ush
new file mode 100644
--- /dev/null
+++ b/Engine/Plugins/FX/Niagara/Shaders/Private/NiagaraDataInterfaceCubeTextureTemplate.ush
@@ -0,0 +1,16 @@
+// Copyright Epic Games, Inc. All Rights Reserved.
+
+int2			{ParameterName}_TextureSize;
+TextureCube		{ParameterName}_Texture;
+SamplerState	{ParameterName}_TextureSampler;
+
+void SampleCubeTexture_{ParameterName}(in float3 UVW, in float MipLevel, out float4 OutValue)
+{
+	OutValue = {ParameterName}_Texture.SampleLevel({ParameterName}_TextureSampler, UVW, MipLevel);
+}
+
+void TextureDimensions_{ParameterName}(out int OutWidth, out int OutHeight)
+{
+	OutWidth = {ParameterName}_TextureSize.x;
+	OutHeight = {ParameterName}_TextureSize.y;
+}
diff --git a/Engine/Plugins/FX/Niagara/Shaders/Private/NiagaraDataInterfaceTexture2DArrayTemplate.ush b/Engine/Plugins/FX/Niagara/Shaders/Private/NiagaraDataInterfaceTexture2DArrayTemplate.ush
new file mode 100644
--- /dev/null
+++ b/Engine/Plugins/FX/Niagara/Shaders/Private/NiagaraDataInterfaceTexture2DArrayTemplate.ush
@@ -0,0 +1,15 @@
+// Copyright Epic Games, Inc. All Rights Reserved.
+
+float3			{ParameterName}_TextureSize;
+Texture2DArray	{ParameterName}_Texture;
+SamplerState	{ParameterName}_TextureSampler;
+
+void SampleTexture_{ParameterName}(in float3 UVW, in float MipLevel, out float4 OutValue)
+{
+	OutValue = {ParameterName}_Texture.SampleLevel({ParameterName}_TextureSampler, UVW, MipLevel);
+}
+
+void TextureDimensions_{ParameterName}(out float2 OutValue)
+{
+	OutValue = {ParameterName}_TextureSize;
+}
diff --git a/Engine/Plugins/FX/Niagara/Shaders/Private/NiagaraDataInterfaceTextureTemplate.ush b/Engine/Plugins/FX/Niagara/Shaders/Private/NiagaraDataInterfaceTextureTemplate.ush
new file mode 100644
--- /dev/null
+++ b/Engine/Plugins/FX/Niagara/Shaders/Private/NiagaraDataInterfaceTextureTemplate.ush
@@ -0,0 +1,20 @@
+// Copyright Epic Games, Inc. All Rights Reserved.
+
+float2			{ParameterName}_TextureSize;
+Texture2D		{ParameterName}_Texture;
+SamplerState	{ParameterName}_TextureSampler;
+
+void SampleTexture2D_{ParameterName}(in float2 UV, out float4 OutValue)
+{
+	OutValue = {ParameterName}_Texture.SampleLevel({ParameterName}_TextureSampler, UV, 0);
+}
+
+void TextureDimensions2D_{ParameterName}(out float2 OutValue)
+{
+	OutValue = {ParameterName}_TextureSize;
+}
+
+void SamplePseudoVolumeTexture_{ParameterName}(in float3 UVW, in float2 XYNumFrames, in float TotalNumFrames, in int MipMode, in float MipLevel, in float2 DDX, in float2 DDY, out float4 OutValue)
+{
+	OutValue = PseudoVolumeTexture({ParameterName}_Texture, {ParameterName}_TextureSampler, UVW, XYNumFrames, TotalNumFrames, (uint)MipMode, MipLevel, DDX, DDY);
+}
diff --git a/Engine/Plugins/FX/Niagara/Shaders/Private/NiagaraDataInterfaceVolumeTextureTemplate.ush b/Engine/Plugins/FX/Niagara/Shaders/Private/NiagaraDataInterfaceVolumeTextureTemplate.ush
new file mode 100644
--- /dev/null
+++ b/Engine/Plugins/FX/Niagara/Shaders/Private/NiagaraDataInterfaceVolumeTextureTemplate.ush
@@ -0,0 +1,15 @@
+// Copyright Epic Games, Inc. All Rights Reserved.
+
+float3			{ParameterName}_TextureSize;
+Texture3D		{ParameterName}_Texture;
+SamplerState	{ParameterName}_TextureSampler;
+
+void SampleVolumeTexture_{ParameterName}(in float3 UV, in float MipLevel, out float4 OutValue)
+{
+	OutValue = {ParameterName}_Texture.SampleLevel({ParameterName}_TextureSampler, UV, MipLevel);
+}
+
+void TextureDimensions3D_{ParameterName}(out float3 OutValue)
+{
+	OutValue = {ParameterName}_TextureSize;
+}
diff --git a/Engine/Plugins/FX/Niagara/Source/Niagara/Classes/NiagaraDataInterface2DArrayTexture.h b/Engine/Plugins/FX/Niagara/Source/Niagara/Classes/NiagaraDataInterface2DArrayTexture.h
--- a/Engine/Plugins/FX/Niagara/Source/Niagara/Classes/NiagaraDataInterface2DArrayTexture.h
+++ b/Engine/Plugins/FX/Niagara/Source/Niagara/Classes/NiagaraDataInterface2DArrayTexture.h
@@ -14,10 +14,14 @@ UCLASS(EditInlineNew, Category = "Texture", meta = (DisplayName = "2D Array Text
 class NIAGARA_API UNiagaraDataInterface2DArrayTexture : public UNiagaraDataInterface
 {
 	GENERATED_UCLASS_BODY()
-public:
 
-	DECLARE_NIAGARA_DI_PARAMETER();
+	BEGIN_SHADER_PARAMETER_STRUCT(FShaderParameters, )
+		SHADER_PARAMETER(FVector3f,					TextureSize)
+		SHADER_PARAMETER_TEXTURE(Texture2DArray,	Texture)
+		SHADER_PARAMETER_SAMPLER(SamplerState,		TextureSampler)
+	END_SHADER_PARAMETER_STRUCT()
 
+public:
 	UPROPERTY(EditAnywhere, Category = "Texture")
 	TObjectPtr<UTexture2DArray> Texture;
 
@@ -47,14 +51,13 @@ public:
 
 	// GPU sim functionality
 #if WITH_EDITORONLY_DATA
+	virtual bool AppendCompileHash(FNiagaraCompileHashVisitor* InVisitor) const override;
 	virtual void GetParameterDefinitionHLSL(const FNiagaraDataInterfaceGPUParamInfo& ParamInfo, FString& OutHLSL) override;
 	virtual bool GetFunctionHLSL(const FNiagaraDataInterfaceGPUParamInfo& ParamInfo, const FNiagaraDataInterfaceGeneratedFunction& FunctionInfo, int FunctionInstanceIndex, FString& OutHLSL) override;
 #endif
-
-	//FRWBuffer& GetGPUBuffer();
-	static const FString TextureName;
-	static const FString SamplerName;
-	static const FString DimensionsBaseName;
+	virtual bool UseLegacyShaderBindings() const override { return false; }
+	virtual void BuildShaderParameters(FNiagaraShaderParametersBuilder& ShaderParametersBuilder) const override;
+	virtual void SetShaderParameters(const FNiagaraDataInterfaceSetShaderParametersContext& Context) const override;
 
 	void SetTexture(UTexture2DArray* InTexture);
 
@@ -62,6 +65,7 @@ protected:
 	virtual bool CopyToInternal(UNiagaraDataInterface* Destination) const override;
 
 protected:
+	static const TCHAR* TemplateShaderFilePath;
 	static const FName SampleTextureName;
 	static const FName TextureDimsName;
 };
diff --git a/Engine/Plugins/FX/Niagara/Source/Niagara/Classes/NiagaraDataInterfaceCubeTexture.h b/Engine/Plugins/FX/Niagara/Source/Niagara/Classes/NiagaraDataInterfaceCubeTexture.h
--- a/Engine/Plugins/FX/Niagara/Source/Niagara/Classes/NiagaraDataInterfaceCubeTexture.h
+++ b/Engine/Plugins/FX/Niagara/Source/Niagara/Classes/NiagaraDataInterfaceCubeTexture.h
@@ -14,10 +14,14 @@ UCLASS(EditInlineNew, Category = "Texture", meta = (DisplayName = "Cube Texture
 class NIAGARA_API UNiagaraDataInterfaceCubeTexture : public UNiagaraDataInterface
 {
 	GENERATED_UCLASS_BODY()
-public:
 
-	DECLARE_NIAGARA_DI_PARAMETER();
+	BEGIN_SHADER_PARAMETER_STRUCT(FShaderParameters, )
+		SHADER_PARAMETER(FIntPoint,				TextureSize)
+		SHADER_PARAMETER_TEXTURE(TextureCube,	Texture)
+		SHADER_PARAMETER_SAMPLER(SamplerState,	TextureSampler)
+	END_SHADER_PARAMETER_STRUCT()
 
+public:
 	UPROPERTY(EditAnywhere, Category = "Texture")
 	TObjectPtr<UTextureCube> Texture;
 
@@ -48,14 +52,13 @@ public:
 
 	// GPU sim functionality
 #if WITH_EDITORONLY_DATA
+	virtual bool AppendCompileHash(FNiagaraCompileHashVisitor* InVisitor) const override;
 	virtual void GetParameterDefinitionHLSL(const FNiagaraDataInterfaceGPUParamInfo& ParamInfo, FString& OutHLSL) override;
 	virtual bool GetFunctionHLSL(const FNiagaraDataInterfaceGPUParamInfo& ParamInfo, const FNiagaraDataInterfaceGeneratedFunction& FunctionInfo, int FunctionInstanceIndex, FString& OutHLSL) override;
 #endif
-
-	//FRWBuffer& GetGPUBuffer();
-	static const FString TextureName;
-	static const FString SamplerName;
-	static const FString DimensionsBaseName;
+	virtual bool UseLegacyShaderBindings() const override { return false; }
+	virtual void BuildShaderParameters(FNiagaraShaderParametersBuilder& ShaderParametersBuilder) const override;
+	virtual void SetShaderParameters(const FNiagaraDataInterfaceSetShaderParametersContext& Context) const override;
 
 	void SetTexture(UTextureCube* InTexture);
 
@@ -63,6 +66,7 @@ protected:
 	virtual bool CopyToInternal(UNiagaraDataInterface* Destination) const override;
 
 protected:
+	static const TCHAR* TemplateShaderFilePath;
 	static const FName SampleCubeTextureName;
 	static const FName TextureDimsName;
 };
diff --git a/Engine/Plugins/FX/Niagara/Source/Niagara/Classes/NiagaraDataInterfaceTexture.h b/Engine/Plugins/FX/Niagara/Source/Niagara/Classes/NiagaraDataInterfaceTexture.h
--- a/Engine/Plugins/FX/Niagara/Source/Niagara/Classes/NiagaraDataInterfaceTexture.h
+++ b/Engine/Plugins/FX/Niagara/Source/Niagara/Classes/NiagaraDataInterfaceTexture.h
@@ -12,10 +12,14 @@ UCLASS(EditInlineNew, Category = "Texture", meta = (DisplayName = "Texture Sampl
 class NIAGARA_API UNiagaraDataInterfaceTexture : public UNiagaraDataInterface
 {
 	GENERATED_UCLASS_BODY()
-public:
 
-	DECLARE_NIAGARA_DI_PARAMETER();
+	BEGIN_SHADER_PARAMETER_STRUCT(FShaderParameters, )
+		SHADER_PARAMETER(FVector2f,				TextureSize)
+		SHADER_PARAMETER_TEXTURE(Texture2D,		Texture)
+		SHADER_PARAMETER_SAMPLER(SamplerState,	TextureSampler)
+	END_SHADER_PARAMETER_STRUCT()
 
+public:
 	UPROPERTY(EditAnywhere, Category = "Texture")
 	TObjectPtr<UTexture> Texture;
 
@@ -49,23 +53,22 @@ public:
 
 	// GPU sim functionality
 #if WITH_EDITORONLY_DATA
+	virtual bool AppendCompileHash(FNiagaraCompileHashVisitor* InVisitor) const override;
 	virtual void GetParameterDefinitionHLSL(const FNiagaraDataInterfaceGPUParamInfo& ParamInfo, FString& OutHLSL) override;
 	virtual bool GetFunctionHLSL(const FNiagaraDataInterfaceGPUParamInfo& ParamInfo, const FNiagaraDataInterfaceGeneratedFunction& FunctionInfo, int FunctionInstanceIndex, FString& OutHLSL) override;
 #endif
+	virtual bool UseLegacyShaderBindings() const override { return false; }
+	virtual void BuildShaderParameters(FNiagaraShaderParametersBuilder& ShaderParametersBuilder) const override;
+	virtual void SetShaderParameters(const FNiagaraDataInterfaceSetShaderParametersContext& Context) const override;
 
 	void SetTexture(UTexture* InTexture);
 
-	//FRWBuffer& GetGPUBuffer();
-	static const FString TextureName;
-	static const FString SamplerName;
-	static const FString DimensionsBaseName;
-
 protected:
 	virtual bool CopyToInternal(UNiagaraDataInterface* Destination) const override;
 
 protected:
+	static const TCHAR* TemplateShaderFilePath;
 	static const FName SampleTexture2DName;
-	static const FName SampleVolumeTextureName;
 	static const FName SamplePseudoVolumeTextureName;
 	static const FName TextureDimsName;
 };
diff --git a/Engine/Plugins/FX/Niagara/Source/Niagara/Classes/NiagaraDataInterfaceVolumeTexture.h b/Engine/Plugins/FX/Niagara/Source/Niagara/Classes/NiagaraDataInterfaceVolumeTexture.h
--- a/Engine/Plugins/FX/Niagara/Source/Niagara/Classes/NiagaraDataInterfaceVolumeTexture.h
+++ b/Engine/Plugins/FX/Niagara/Source/Niagara/Classes/NiagaraDataInterfaceVolumeTexture.h
@@ -14,10 +14,14 @@ UCLASS(EditInlineNew, Category = "Texture", meta = (DisplayName = "Volume Textur
 class NIAGARA_API UNiagaraDataInterfaceVolumeTexture : public UNiagaraDataInterface
 {
 	GENERATED_UCLASS_BODY()
-public:
 
-	DECLARE_NIAGARA_DI_PARAMETER();
+	BEGIN_SHADER_PARAMETER_STRUCT(FShaderParameters, )
+		SHADER_PARAMETER(FVector3f,				TextureSize)
+		SHADER_PARAMETER_TEXTURE(Texture3D,		Texture)
+		SHADER_PARAMETER_SAMPLER(SamplerState,	TextureSampler)
+	END_SHADER_PARAMETER_STRUCT()
 
+public:
 	UPROPERTY(EditAnywhere, Category = "Texture")
 	TObjectPtr<UVolumeTexture> Texture;
 
@@ -48,14 +52,13 @@ public:
 
 	// GPU sim functionality
 #if WITH_EDITORONLY_DATA
+	virtual bool AppendCompileHash(FNiagaraCompileHashVisitor* InVisitor) const override;
 	virtual void GetParameterDefinitionHLSL(const FNiagaraDataInterfaceGPUParamInfo& ParamInfo, FString& OutHLSL) override;
 	virtual bool GetFunctionHLSL(const FNiagaraDataInterfaceGPUParamInfo& ParamInfo, const FNiagaraDataInterfaceGeneratedFunction& FunctionInfo, int FunctionInstanceIndex, FString& OutHLSL) override;
 #endif
-
-	//FRWBuffer& GetGPUBuffer();
-	static const FString TextureName;
-	static const FString SamplerName;
-	static const FString DimensionsBaseName;
+	virtual bool UseLegacyShaderBindings() const override { return false; }
+	virtual void BuildShaderParameters(FNiagaraShaderParametersBuilder& ShaderParametersBuilder) const override;
+	virtual void SetShaderParameters(const FNiagaraDataInterfaceSetShaderParametersContext& Context) const override;
 
 	void SetTexture(UVolumeTexture* InTexture);
 
@@ -63,6 +66,7 @@ protected:
 	virtual bool CopyToInternal(UNiagaraDataInterface* Destination) const override;
 
 protected:
+	static const TCHAR* TemplateShaderFilePath;
 	static const FName SampleVolumeTextureName;
 	static const FName TextureDimsName;
 };
diff --git a/Engine/Plugins/FX/Niagara/Source/Niagara/Private/NiagaraDataInterface2DArrayTexture.cpp b/Engine/Plugins/FX/Niagara/Source/Niagara/Private/NiagaraDataInterface2DArrayTexture.cpp
--- a/Engine/Plugins/FX/Niagara/Source/Niagara/Private/NiagaraDataInterface2DArrayTexture.cpp
+++ b/Engine/Plugins/FX/Niagara/Source/Niagara/Private/NiagaraDataInterface2DArrayTexture.cpp
@@ -1,21 +1,20 @@
 // Copyright Epic Games, Inc. All Rights Reserved.
 
 #include "NiagaraDataInterface2DArrayTexture.h"
+#include "NiagaraComputeExecutionContext.h"
 #include "NiagaraShader.h"
-#include "ShaderParameterUtils.h"
-#include "NiagaraCustomVersion.h"
+#include "NiagaraShaderParametersBuilder.h"
+#include "NiagaraSystemInstance.h"
+
 #include "Engine/Texture2DArray.h"
 #include "Engine/TextureRenderTarget2DArray.h"
-#include "NiagaraSystemInstance.h"
-#include "NiagaraComputeExecutionContext.h"
+#include "ShaderParameterUtils.h"
 
 #define LOCTEXT_NAMESPACE "UNiagaraDataInterface2DArrayTexture"
 
+const TCHAR* UNiagaraDataInterface2DArrayTexture::TemplateShaderFilePath = TEXT("/Plugin/FX/Niagara/Private/NiagaraDataInterfaceTexture2DArrayTemplate.ush");
 const FName UNiagaraDataInterface2DArrayTexture::SampleTextureName(TEXT("SampleTexture"));
 const FName UNiagaraDataInterface2DArrayTexture::TextureDimsName(TEXT("TextureDimensions"));
-const FString UNiagaraDataInterface2DArrayTexture::TextureName(TEXT("Texture_"));
-const FString UNiagaraDataInterface2DArrayTexture::SamplerName(TEXT("Sampler_"));
-const FString UNiagaraDataInterface2DArrayTexture::DimensionsBaseName(TEXT("Dimensions_"));
 
 struct FNDITexture2DArrayInstanceData_GameThread
 {
@@ -228,96 +227,61 @@ void UNiagaraDataInterface2DArrayTexture::GetTextureDimensions(FVectorVMExternal
 }
 
 #if WITH_EDITORONLY_DATA
-bool UNiagaraDataInterface2DArrayTexture::GetFunctionHLSL(const FNiagaraDataInterfaceGPUParamInfo& ParamInfo, const FNiagaraDataInterfaceGeneratedFunction& FunctionInfo, int FunctionInstanceIndex, FString& OutHLSL)
+bool UNiagaraDataInterface2DArrayTexture::AppendCompileHash(FNiagaraCompileHashVisitor* InVisitor) const
 {
-	if (FunctionInfo.DefinitionName == SampleTextureName)
+	bool bSuccess = Super::AppendCompileHash(InVisitor);
+	InVisitor->UpdateString(TEXT("UNiagaraDataInterface2DArrayTextureHLSLSource"), GetShaderFileHash(TemplateShaderFilePath, EShaderPlatform::SP_PCD3D_SM5).ToString());
+	bSuccess &= InVisitor->UpdateShaderParameters<FShaderParameters>();
+	return bSuccess;
+}
+
+void UNiagaraDataInterface2DArrayTexture::GetParameterDefinitionHLSL(const FNiagaraDataInterfaceGPUParamInfo& ParamInfo, FString& OutHLSL)
+{
+	TMap<FString, FStringFormatArg> TemplateArgs =
 	{
-		FString HLSLTextureName = TextureName + ParamInfo.DataInterfaceHLSLSymbol;
-		FString HLSLSamplerName = SamplerName + ParamInfo.DataInterfaceHLSLSymbol;
-		OutHLSL += TEXT("void ") + FunctionInfo.InstanceName + TEXT("(in float3 In_UV, in float MipLevel, out float4 Out_Value) \n{\n");
-		OutHLSL += TEXT("\t Out_Value = ") + HLSLTextureName + TEXT(".SampleLevel(") + HLSLSamplerName + TEXT(", In_UV, MipLevel);\n");
-		OutHLSL += TEXT("\n}\n");
-		return true;
-	}
-	else if (FunctionInfo.DefinitionName == TextureDimsName)
+		{TEXT("ParameterName"),	ParamInfo.DataInterfaceHLSLSymbol},
+	};
+
+	FString TemplateFile;
+	LoadShaderSourceFile(TemplateShaderFilePath, EShaderPlatform::SP_PCD3D_SM5, &TemplateFile, nullptr);
+	OutHLSL += FString::Format(*TemplateFile, TemplateArgs);
+}
+
+bool UNiagaraDataInterface2DArrayTexture::GetFunctionHLSL(const FNiagaraDataInterfaceGPUParamInfo& ParamInfo, const FNiagaraDataInterfaceGeneratedFunction& FunctionInfo, int FunctionInstanceIndex, FString& OutHLSL)
+{
+	if ((FunctionInfo.DefinitionName == SampleTextureName) ||
+		(FunctionInfo.DefinitionName == TextureDimsName) )
 	{
-		FString DimsVar = DimensionsBaseName + ParamInfo.DataInterfaceHLSLSymbol;
-		OutHLSL += TEXT("void ") + FunctionInfo.InstanceName + TEXT("(out float3 Out_Value) \n{\n");
-		OutHLSL += TEXT("\t Out_Value = ") + DimsVar + TEXT(";\n");
-		OutHLSL += TEXT("\n}\n");
 		return true;
 	}
 	return false;
 }
+#endif
 
-void UNiagaraDataInterface2DArrayTexture::GetParameterDefinitionHLSL(const FNiagaraDataInterfaceGPUParamInfo& ParamInfo, FString& OutHLSL)
+void UNiagaraDataInterface2DArrayTexture::BuildShaderParameters(FNiagaraShaderParametersBuilder& ShaderParametersBuilder) const
 {
-	FString HLSLTextureName = TextureName + ParamInfo.DataInterfaceHLSLSymbol;
-	FString HLSLSamplerName = SamplerName + ParamInfo.DataInterfaceHLSLSymbol;
-	OutHLSL += TEXT("Texture2DArray ") + HLSLTextureName + TEXT(";\n");
-	OutHLSL += TEXT("SamplerState ") + HLSLSamplerName + TEXT(";\n");
-	OutHLSL += TEXT("float3 ") + DimensionsBaseName + ParamInfo.DataInterfaceHLSLSymbol + TEXT(";\n");
+	ShaderParametersBuilder.AddNestedStruct<FShaderParameters>();
 }
-#endif
 
-struct FNiagaraDataInterfaceParametersCS_2DArrayTexture : public FNiagaraDataInterfaceParametersCS
+void UNiagaraDataInterface2DArrayTexture::SetShaderParameters(const FNiagaraDataInterfaceSetShaderParametersContext& Context) const
 {
-	DECLARE_TYPE_LAYOUT(FNiagaraDataInterfaceParametersCS_2DArrayTexture, NonVirtual);
-public:
-	void Bind(const FNiagaraDataInterfaceGPUParamInfo& ParameterInfo, const class FShaderParameterMap& ParameterMap)
-	{
-		FString TexName = UNiagaraDataInterface2DArrayTexture::TextureName + ParameterInfo.DataInterfaceHLSLSymbol;
-		FString SampleName = (UNiagaraDataInterface2DArrayTexture::SamplerName + ParameterInfo.DataInterfaceHLSLSymbol);
-		TextureParam.Bind(ParameterMap, *TexName);
-		SamplerParam.Bind(ParameterMap, *SampleName);
-
-		Dimensions.Bind(ParameterMap, *(UNiagaraDataInterface2DArrayTexture::DimensionsBaseName + ParameterInfo.DataInterfaceHLSLSymbol));
+	const FNiagaraDataInterfaceProxyTexture2DArray& TextureProxy = Context.GetProxy<FNiagaraDataInterfaceProxyTexture2DArray>();
+	const FNDITexture2DArrayInstanceData_RenderThread* InstanceData = TextureProxy.InstanceData_RT.Find(Context.GetSystemInstanceID());
 
+	FShaderParameters* Parameters = Context.GetParameterNestedStruct<FShaderParameters>();
+	if (InstanceData && InstanceData->ResolvedTextureRHI.IsValid())
+	{
+		Parameters->TextureSize		= InstanceData->TextureSize;
+		Parameters->Texture			= InstanceData->ResolvedTextureRHI;
+		Parameters->TextureSampler	= InstanceData->SamplerStateRHI ? InstanceData->SamplerStateRHI : GBlackArrayTexture->SamplerStateRHI;
 	}
-
-	void Set(FRHICommandList& RHICmdList, const FNiagaraDataInterfaceSetArgs& Context) const
+	else
 	{
-		check(IsInRenderingThread());
-
-		FRHIComputeShader* ComputeShaderRHI = Context.Shader.GetComputeShader();
-		FNiagaraDataInterfaceProxyTexture2DArray* TextureDI = static_cast<FNiagaraDataInterfaceProxyTexture2DArray*>(Context.DataInterface);
-		FNDITexture2DArrayInstanceData_RenderThread* InstanceData = TextureDI->InstanceData_RT.Find(Context.SystemInstanceID);
-
-		if (InstanceData && InstanceData->ResolvedTextureRHI.IsValid())
-		{
-			FRHISamplerState* SamplerStateRHI = InstanceData->SamplerStateRHI ? InstanceData->SamplerStateRHI : GBlackTexture->SamplerStateRHI;
-			SetTextureParameter(
-				RHICmdList,
-				ComputeShaderRHI,
-				TextureParam,
-				SamplerParam,
-				SamplerStateRHI,
-				InstanceData->ResolvedTextureRHI
-			);
-			SetShaderValue(RHICmdList, ComputeShaderRHI, Dimensions, InstanceData->TextureSize);
-		}
-		else
-		{
-			SetTextureParameter(
-				RHICmdList,
-				ComputeShaderRHI,
-				TextureParam,
-				SamplerParam,
-				GBlackArrayTexture->SamplerStateRHI,
-				GBlackArrayTexture->TextureRHI
-			);
-			SetShaderValue(RHICmdList, ComputeShaderRHI, Dimensions, FVector3f::ZeroVector);
-		}
+		Parameters->TextureSize		= FVector3f::ZeroVector;
+		Parameters->Texture			= GBlackArrayTexture->TextureRHI;
+		Parameters->TextureSampler	= GBlackArrayTexture->SamplerStateRHI;
 	}
-private:
-	LAYOUT_FIELD(FShaderResourceParameter, TextureParam);
-	LAYOUT_FIELD(FShaderResourceParameter, SamplerParam);
-	LAYOUT_FIELD(FShaderParameter, Dimensions);
-};
-
-IMPLEMENT_TYPE_LAYOUT(FNiagaraDataInterfaceParametersCS_2DArrayTexture);
-
-IMPLEMENT_NIAGARA_DI_PARAMETER(UNiagaraDataInterface2DArrayTexture, FNiagaraDataInterfaceParametersCS_2DArrayTexture);
+}
 
 void UNiagaraDataInterface2DArrayTexture::SetTexture(UTexture2DArray* InTexture)
 {
@@ -327,4 +291,4 @@ void UNiagaraDataInterface2DArrayTexture::SetTexture(UTexture2DArray* InTexture)
 	}
 }
 
-#undef LOCTEXT_NAMESPACE
\ No newline at end of file
+#undef LOCTEXT_NAMESPACE
diff --git a/Engine/Plugins/FX/Niagara/Source/Niagara/Private/NiagaraDataInterfaceCubeTexture.cpp b/Engine/Plugins/FX/Niagara/Source/Niagara/Private/NiagaraDataInterfaceCubeTexture.cpp
--- a/Engine/Plugins/FX/Niagara/Source/Niagara/Private/NiagaraDataInterfaceCubeTexture.cpp
+++ b/Engine/Plugins/FX/Niagara/Source/Niagara/Private/NiagaraDataInterfaceCubeTexture.cpp
@@ -1,21 +1,20 @@
 // Copyright Epic Games, Inc. All Rights Reserved.
 
 #include "NiagaraDataInterfaceCubeTexture.h"
+#include "NiagaraComputeExecutionContext.h"
 #include "NiagaraShader.h"
-#include "ShaderParameterUtils.h"
-#include "NiagaraCustomVersion.h"
+#include "NiagaraShaderParametersBuilder.h"
+#include "NiagaraSystemInstance.h"
+
 #include "Engine/TextureCube.h"
 #include "Engine/TextureRenderTargetCube.h"
-#include "NiagaraSystemInstance.h"
-#include "NiagaraComputeExecutionContext.h"
+#include "ShaderParameterUtils.h"
 
 #define LOCTEXT_NAMESPACE "UNiagaraDataInterfaceCubeTexture"
 
+const TCHAR* UNiagaraDataInterfaceCubeTexture::TemplateShaderFilePath = TEXT("/Plugin/FX/Niagara/Private/NiagaraDataInterfaceCubeTextureTemplate.ush");
 const FName UNiagaraDataInterfaceCubeTexture::SampleCubeTextureName(TEXT("SampleCubeTexture"));
 const FName UNiagaraDataInterfaceCubeTexture::TextureDimsName(TEXT("TextureDimensions"));
-const FString UNiagaraDataInterfaceCubeTexture::TextureName(TEXT("Texture_"));
-const FString UNiagaraDataInterfaceCubeTexture::SamplerName(TEXT("Sampler_"));
-const FString UNiagaraDataInterfaceCubeTexture::DimensionsBaseName(TEXT("Dimensions_"));
 
 struct FNDICubeTextureInstanceData_GameThread
 {
@@ -248,102 +247,61 @@ void UNiagaraDataInterfaceCubeTexture::SampleCubeTexture(FVectorVMExternalFuncti
 }
 
 #if WITH_EDITORONLY_DATA
-bool UNiagaraDataInterfaceCubeTexture::GetFunctionHLSL(const FNiagaraDataInterfaceGPUParamInfo& ParameterInfo, const FNiagaraDataInterfaceGeneratedFunction& FunctionInfo, int FunctionInstanceIndex, FString& OutHLSL)
+bool UNiagaraDataInterfaceCubeTexture::AppendCompileHash(FNiagaraCompileHashVisitor* InVisitor) const
 {
-	if (FunctionInfo.DefinitionName == SampleCubeTextureName)
-	{
-		const FString HLSLTextureName = TextureName + ParameterInfo.DataInterfaceHLSLSymbol;
-		const FString HLSLSamplerName = SamplerName + ParameterInfo.DataInterfaceHLSLSymbol;
+	bool bSuccess = Super::AppendCompileHash(InVisitor);
+	InVisitor->UpdateString(TEXT("UNiagaraDataInterfaceCubeTextureHLSLSource"), GetShaderFileHash(TemplateShaderFilePath, EShaderPlatform::SP_PCD3D_SM5).ToString());
+	bSuccess &= InVisitor->UpdateShaderParameters<FShaderParameters>();
+	return bSuccess;
+}
 
-		OutHLSL.Appendf(TEXT("void %s(in float3 In_UVW, in float MipLevel, out float4 Out_Value)\n"), *FunctionInfo.InstanceName);
-		OutHLSL.Append(TEXT("{\n"));
-		OutHLSL.Appendf(TEXT("\tOut_Value = %s.SampleLevel(%s, In_UVW, MipLevel);\n"), *HLSLTextureName, *HLSLSamplerName);
-		OutHLSL.Append(TEXT("}\n"));
-		return true;
-	}
-	else if (FunctionInfo.DefinitionName == TextureDimsName)
+void UNiagaraDataInterfaceCubeTexture::GetParameterDefinitionHLSL(const FNiagaraDataInterfaceGPUParamInfo& ParamInfo, FString& OutHLSL)
+{
+	TMap<FString, FStringFormatArg> TemplateArgs =
 	{
-		const FString DimsVar = DimensionsBaseName + ParameterInfo.DataInterfaceHLSLSymbol;
+		{TEXT("ParameterName"),	ParamInfo.DataInterfaceHLSLSymbol},
+	};
 
-		OutHLSL.Appendf(TEXT("void %s(out int Out_Width, out int Out_Height)\n"), *FunctionInfo.InstanceName);
-		OutHLSL.Append (TEXT("{\n"));
-		OutHLSL.Appendf(TEXT("\tOut_Width = %s.x;\n"), *DimsVar);
-		OutHLSL.Appendf(TEXT("\tOut_Width = %s.y;\n"), *DimsVar);
-		OutHLSL.Append(TEXT("}\n"));
+	FString TemplateFile;
+	LoadShaderSourceFile(TemplateShaderFilePath, EShaderPlatform::SP_PCD3D_SM5, &TemplateFile, nullptr);
+	OutHLSL += FString::Format(*TemplateFile, TemplateArgs);
+}
+
+bool UNiagaraDataInterfaceCubeTexture::GetFunctionHLSL(const FNiagaraDataInterfaceGPUParamInfo& ParameterInfo, const FNiagaraDataInterfaceGeneratedFunction& FunctionInfo, int FunctionInstanceIndex, FString& OutHLSL)
+{
+	if ((FunctionInfo.DefinitionName == SampleCubeTextureName) ||
+		(FunctionInfo.DefinitionName == TextureDimsName))
+	{
 		return true;
 	}
 	return false;
 }
+#endif
 
-void UNiagaraDataInterfaceCubeTexture::GetParameterDefinitionHLSL(const FNiagaraDataInterfaceGPUParamInfo& ParameterInfo, FString& OutHLSL)
+void UNiagaraDataInterfaceCubeTexture::BuildShaderParameters(FNiagaraShaderParametersBuilder& ShaderParametersBuilder) const
 {
-	const FString HLSLTextureName = TextureName + ParameterInfo.DataInterfaceHLSLSymbol;
-	const FString HLSLSamplerName = SamplerName + ParameterInfo.DataInterfaceHLSLSymbol;
-	const FString DimsVar = DimensionsBaseName + ParameterInfo.DataInterfaceHLSLSymbol;
-
-	OutHLSL.Appendf(TEXT("TextureCube %s;\n"), *HLSLTextureName);
-	OutHLSL.Appendf(TEXT("SamplerState %s;\n"), *HLSLSamplerName);
-	OutHLSL.Appendf(TEXT("int2 %s;\n"), *DimsVar);
+	ShaderParametersBuilder.AddNestedStruct<FShaderParameters>();
 }
-#endif
 
-struct FNiagaraDataInterfaceParametersCS_CubeTexture : public FNiagaraDataInterfaceParametersCS
+void UNiagaraDataInterfaceCubeTexture::SetShaderParameters(const FNiagaraDataInterfaceSetShaderParametersContext& Context) const
 {
-	DECLARE_TYPE_LAYOUT(FNiagaraDataInterfaceParametersCS_CubeTexture, NonVirtual);
-public:
-	void Bind(const FNiagaraDataInterfaceGPUParamInfo& ParameterInfo, const class FShaderParameterMap& ParameterMap)
-	{
-		const FString HLSLTextureName = UNiagaraDataInterfaceCubeTexture::TextureName + ParameterInfo.DataInterfaceHLSLSymbol;
-		const FString HLSLSamplerName = UNiagaraDataInterfaceCubeTexture::SamplerName + ParameterInfo.DataInterfaceHLSLSymbol;
-		const FString DimsVar = UNiagaraDataInterfaceCubeTexture::DimensionsBaseName + ParameterInfo.DataInterfaceHLSLSymbol;
+	const FNiagaraDataInterfaceProxyCubeTexture& TextureProxy = Context.GetProxy<FNiagaraDataInterfaceProxyCubeTexture>();
+	const FNDICubeTextureInstanceData_RenderThread* InstanceData = TextureProxy.InstanceData_RT.Find(Context.GetSystemInstanceID());
 
-		TextureParam.Bind(ParameterMap, *HLSLTextureName);
-		SamplerParam.Bind(ParameterMap, *HLSLSamplerName);
-		Dimensions.Bind(ParameterMap, *DimsVar);
+	FShaderParameters* Parameters = Context.GetParameterNestedStruct<FShaderParameters>();
+	if (InstanceData && InstanceData->ResolvedTextureRHI.IsValid())
+	{
+		Parameters->TextureSize		= InstanceData->TextureSize;
+		Parameters->Texture			= InstanceData->ResolvedTextureRHI;
+		Parameters->TextureSampler	= InstanceData->SamplerStateRHI ? InstanceData->SamplerStateRHI : GBlackTextureCube->SamplerStateRHI;
 	}
-
-	void Set(FRHICommandList& RHICmdList, const FNiagaraDataInterfaceSetArgs& Context) const
+	else
 	{
-		check(IsInRenderingThread());
-
-		FRHIComputeShader* ComputeShaderRHI = Context.Shader.GetComputeShader();
-		FNiagaraDataInterfaceProxyCubeTexture* TextureDI = static_cast<FNiagaraDataInterfaceProxyCubeTexture*>(Context.DataInterface);
-		FNDICubeTextureInstanceData_RenderThread* InstanceData = TextureDI->InstanceData_RT.Find(Context.SystemInstanceID);
-
-		if (InstanceData && InstanceData->ResolvedTextureRHI.IsValid())
-		{
-			FRHISamplerState* SamplerStateRHI = InstanceData->SamplerStateRHI ? InstanceData->SamplerStateRHI : GBlackTexture->SamplerStateRHI;
-			SetTextureParameter(
-				RHICmdList,
-				ComputeShaderRHI,
-				TextureParam,
-				SamplerParam,
-				SamplerStateRHI,
-				InstanceData->ResolvedTextureRHI
-			);
-			SetShaderValue(RHICmdList, ComputeShaderRHI, Dimensions, InstanceData->TextureSize);
-		}
-		else
-		{
-			SetTextureParameter(
-				RHICmdList,
-				ComputeShaderRHI,
-				TextureParam,
-				SamplerParam,
-				GBlackTextureCube->SamplerStateRHI,
-				GBlackTextureCube->TextureRHI
-			);
-			SetShaderValue(RHICmdList, ComputeShaderRHI, Dimensions, FVector3f::ZeroVector);
-		}
+		Parameters->TextureSize		= FIntPoint::ZeroValue;
+		Parameters->Texture			= GBlackTextureCube->TextureRHI;
+		Parameters->TextureSampler	= GBlackTextureCube->SamplerStateRHI;
 	}
-private:
-	LAYOUT_FIELD(FShaderResourceParameter, TextureParam);
-	LAYOUT_FIELD(FShaderResourceParameter, SamplerParam);
-	LAYOUT_FIELD(FShaderParameter, Dimensions);
-};
-
-IMPLEMENT_TYPE_LAYOUT(FNiagaraDataInterfaceParametersCS_CubeTexture);
-IMPLEMENT_NIAGARA_DI_PARAMETER(UNiagaraDataInterfaceCubeTexture, FNiagaraDataInterfaceParametersCS_CubeTexture);
+}
 
 void UNiagaraDataInterfaceCubeTexture::SetTexture(UTextureCube* InTexture)
 {
diff --git a/Engine/Plugins/FX/Niagara/Source/Niagara/Private/NiagaraDataInterfaceTexture.cpp b/Engine/Plugins/FX/Niagara/Source/Niagara/Private/NiagaraDataInterfaceTexture.cpp
--- a/Engine/Plugins/FX/Niagara/Source/Niagara/Private/NiagaraDataInterfaceTexture.cpp
+++ b/Engine/Plugins/FX/Niagara/Source/Niagara/Private/NiagaraDataInterfaceTexture.cpp
@@ -1,21 +1,20 @@
 // Copyright Epic Games, Inc. All Rights Reserved.
 
 #include "NiagaraDataInterfaceTexture.h"
+#include "NiagaraComputeExecutionContext.h"
+#include "NiagaraCustomVersion.h"
 #include "NiagaraShader.h"
-#include "ShaderParameterUtils.h"
+#include "NiagaraShaderParametersBuilder.h"
 #include "NiagaraSystemInstance.h"
-#include "NiagaraCustomVersion.h"
-#include "NiagaraComputeExecutionContext.h"
+
+#include "ShaderParameterUtils.h"
 
 #define LOCTEXT_NAMESPACE "UNiagaraDataInterfaceTexture"
 
+const TCHAR* UNiagaraDataInterfaceTexture::TemplateShaderFilePath = TEXT("/Plugin/FX/Niagara/Private/NiagaraDataInterfaceTextureTemplate.ush");
 const FName UNiagaraDataInterfaceTexture::SampleTexture2DName(TEXT("SampleTexture2D"));
-const FName UNiagaraDataInterfaceTexture::SampleVolumeTextureName(TEXT("SampleVolumeTexture"));
 const FName UNiagaraDataInterfaceTexture::SamplePseudoVolumeTextureName(TEXT("SamplePseudoVolumeTexture"));
 const FName UNiagaraDataInterfaceTexture::TextureDimsName(TEXT("TextureDimensions2D"));
-const FString UNiagaraDataInterfaceTexture::TextureName(TEXT("Texture_"));
-const FString UNiagaraDataInterfaceTexture::SamplerName(TEXT("Sampler_"));
-const FString UNiagaraDataInterfaceTexture::DimensionsBaseName(TEXT("Dimensions_"));
 
 struct FNDITextureInstanceData_GameThread
 {
@@ -161,20 +160,6 @@ void UNiagaraDataInterfaceTexture::GetFunctions(TArray<FNiagaraFunctionSignature
 		OutFunctions.Add(Sig);
 	}
 
-	//{
-	//	FNiagaraFunctionSignature Sig;
-	//	Sig.Name = SampleVolumeTextureName;
-	//	Sig.bMemberFunction = true;
-	//	Sig.bRequiresContext = false;
-	//	Sig.Inputs.Add(FNiagaraVariable(FNiagaraTypeDefinition(GetClass()), TEXT("Texture")));
-	//	Sig.Inputs.Add(FNiagaraVariable(FNiagaraTypeDefinition::GetVec3Def(), TEXT("UVW")));
-	//	Sig.SetDescription(LOCTEXT("TextureSampleVolumeTextureDesc", "Sample mip level 0 of the input 3d texture at the specified UVW coordinates. The UVW origin (0,0) is in the bottom left hand corner of the volume."));
-	//	Sig.Outputs.Add(FNiagaraVariable(FNiagaraTypeDefinition::GetVec4Def(), TEXT("Value")));
-	//	//Sig.Owner = *GetName();
-
-	//	OutFunctions.Add(Sig);
-	//}
-
 	{
 		FNiagaraFunctionSignature Sig;
 		Sig.Name = SamplePseudoVolumeTextureName;
@@ -385,118 +370,66 @@ void UNiagaraDataInterfaceTexture::SamplePseudoVolumeTexture(FVectorVMExternalFu
 }
 
 #if WITH_EDITORONLY_DATA
-bool UNiagaraDataInterfaceTexture::GetFunctionHLSL(const FNiagaraDataInterfaceGPUParamInfo& ParamInfo, const FNiagaraDataInterfaceGeneratedFunction& FunctionInfo, int FunctionInstanceIndex, FString& OutHLSL)
+bool UNiagaraDataInterfaceTexture::AppendCompileHash(FNiagaraCompileHashVisitor* InVisitor) const
 {
-	if (FunctionInfo.DefinitionName == SampleTexture2DName)
-	{
-		FString HLSLTextureName = TextureName + ParamInfo.DataInterfaceHLSLSymbol;
-		FString HLSLSamplerName = SamplerName + ParamInfo.DataInterfaceHLSLSymbol;
-		OutHLSL += TEXT("void ") + FunctionInfo.InstanceName + TEXT("(in float2 In_UV, out float4 Out_Value) \n{\n");
-		OutHLSL += TEXT("\t Out_Value = ") + HLSLTextureName + TEXT(".SampleLevel(") + HLSLSamplerName + TEXT(", In_UV, 0);\n");
-		OutHLSL += TEXT("\n}\n");
-		return true;
-	}
-	/*else if (FunctionInfo.DefinitionName == SampleVolumeTextureName)
-	{
-		FString HLSLTextureName = TextureName + ParamInfo.DataInterfaceHLSLSymbol;
-		FString HLSLSamplerName = SamplerName + ParamInfo.DataInterfaceHLSLSymbol;
-		OutHLSL += TEXT("void ") + FunctionInfo.InstanceName + TEXT("(in float3 In_UV, out float4 Out_Value) \n{\n");
-		OutHLSL += TEXT("\t Out_Value = ") + HLSLTextureName + TEXT(".SampleLevel(") + HLSLSamplerName + TEXT(", In_UV, 0);\n");
-		OutHLSL += TEXT("\n}\n");
-		return true;
-	}*/
-	else if (FunctionInfo.DefinitionName == SamplePseudoVolumeTextureName)
+	bool bSuccess = Super::AppendCompileHash(InVisitor);
+	InVisitor->UpdateString(TEXT("UNiagaraDataInterfaceTextureHLSLSource"), GetShaderFileHash(TemplateShaderFilePath, EShaderPlatform::SP_PCD3D_SM5).ToString());
+	bSuccess &= InVisitor->UpdateShaderParameters<FShaderParameters>();
+	return bSuccess;
+}
+
+void UNiagaraDataInterfaceTexture::GetParameterDefinitionHLSL(const FNiagaraDataInterfaceGPUParamInfo& ParamInfo, FString& OutHLSL)
+{
+	TMap<FString, FStringFormatArg> TemplateArgs =
 	{
-		FString HLSLTextureName = TextureName + ParamInfo.DataInterfaceHLSLSymbol;
-		FString HLSLSamplerName = SamplerName + ParamInfo.DataInterfaceHLSLSymbol;
-		OutHLSL += TEXT("void ") + FunctionInfo.InstanceName + TEXT("(in float3 In_UVW, in float2 In_XYNumFrames, in float In_TotalNumFrames, in int In_MipMode, in float In_MipLevel, in float2 In_DDX, in float2 In_DDY, out float4 Out_Value) \n{\n");
-		OutHLSL += TEXT("\t Out_Value = PseudoVolumeTexture(") + HLSLTextureName + TEXT(", ") + HLSLSamplerName + TEXT(", In_UVW, In_XYNumFrames, In_TotalNumFrames, (uint) In_MipMode, In_MipLevel, In_DDX, In_DDY); \n");
-		OutHLSL += TEXT("\n}\n");
-		return true;
-	}
-	else if (FunctionInfo.DefinitionName == TextureDimsName)
+		{TEXT("ParameterName"),	ParamInfo.DataInterfaceHLSLSymbol},
+	};
+
+	FString TemplateFile;
+	LoadShaderSourceFile(TemplateShaderFilePath, EShaderPlatform::SP_PCD3D_SM5, &TemplateFile, nullptr);
+	OutHLSL += FString::Format(*TemplateFile, TemplateArgs);
+}
+
+bool UNiagaraDataInterfaceTexture::GetFunctionHLSL(const FNiagaraDataInterfaceGPUParamInfo& ParamInfo, const FNiagaraDataInterfaceGeneratedFunction& FunctionInfo, int FunctionInstanceIndex, FString& OutHLSL)
+{
+	if ((FunctionInfo.DefinitionName == SampleTexture2DName) ||
+		(FunctionInfo.DefinitionName == SamplePseudoVolumeTextureName) ||
+		(FunctionInfo.DefinitionName == TextureDimsName))
 	{
-		FString DimsVar = DimensionsBaseName + ParamInfo.DataInterfaceHLSLSymbol;
-		OutHLSL += TEXT("void ") + FunctionInfo.InstanceName + TEXT("(out float2 Out_Value) \n{\n");
-		OutHLSL += TEXT("\t Out_Value = ") + DimsVar + TEXT(";\n");
-		OutHLSL += TEXT("\n}\n");
 		return true;
 	}
 	return false;
 }
+#endif
 
-void UNiagaraDataInterfaceTexture::GetParameterDefinitionHLSL(const FNiagaraDataInterfaceGPUParamInfo& ParamInfo, FString& OutHLSL)
+void UNiagaraDataInterfaceTexture::BuildShaderParameters(FNiagaraShaderParametersBuilder& ShaderParametersBuilder) const
 {
-	FString HLSLTextureName = TextureName + ParamInfo.DataInterfaceHLSLSymbol;
-	FString HLSLSamplerName = SamplerName + ParamInfo.DataInterfaceHLSLSymbol;
-	OutHLSL += TEXT("Texture2D ") + HLSLTextureName + TEXT(";\n");
-	OutHLSL += TEXT("SamplerState ") + HLSLSamplerName + TEXT(";\n");
-	OutHLSL += TEXT("float2 ") + DimensionsBaseName + ParamInfo.DataInterfaceHLSLSymbol + TEXT(";\n");
+	ShaderParametersBuilder.AddNestedStruct<FShaderParameters>();
 }
-#endif
 
-struct FNiagaraDataInterfaceParametersCS_Texture : public FNiagaraDataInterfaceParametersCS
+void UNiagaraDataInterfaceTexture::SetShaderParameters(const FNiagaraDataInterfaceSetShaderParametersContext& Context) const
 {
-	DECLARE_TYPE_LAYOUT(FNiagaraDataInterfaceParametersCS_Texture, NonVirtual);
-public:
-	void Bind(const FNiagaraDataInterfaceGPUParamInfo& ParameterInfo, const class FShaderParameterMap& ParameterMap)
+	const FNiagaraDataInterfaceProxyTexture& TextureProxy = Context.GetProxy<FNiagaraDataInterfaceProxyTexture>();
+	const FNDITextureInstanceData_RenderThread* InstanceData = TextureProxy.InstanceData_RT.Find(Context.GetSystemInstanceID());
+
+	FShaderParameters* Parameters = Context.GetParameterNestedStruct<FShaderParameters>();
+	if (InstanceData && InstanceData->ResolvedTextureRHI.IsValid())
 	{
-		FString TexName = UNiagaraDataInterfaceTexture::TextureName + ParameterInfo.DataInterfaceHLSLSymbol;
-		FString SampleName = (UNiagaraDataInterfaceTexture::SamplerName + ParameterInfo.DataInterfaceHLSLSymbol);
-		TextureParam.Bind(ParameterMap, *TexName);
-		SamplerParam.Bind(ParameterMap, *SampleName);
-		
-		Dimensions.Bind(ParameterMap, *(UNiagaraDataInterfaceTexture::DimensionsBaseName + ParameterInfo.DataInterfaceHLSLSymbol));
+		Parameters->TextureSize		= InstanceData->TextureSize;
+		Parameters->Texture			= InstanceData->ResolvedTextureRHI;
+		Parameters->TextureSampler	= InstanceData->SamplerStateRHI ? InstanceData->SamplerStateRHI : GBlackTexture->SamplerStateRHI;
 	}
-
-	void Set(FRHICommandList& RHICmdList, const FNiagaraDataInterfaceSetArgs& Context) const
+	else
 	{
-		check(IsInRenderingThread());
-
-		FRHIComputeShader* ComputeShaderRHI = Context.Shader.GetComputeShader();
-		FNiagaraDataInterfaceProxyTexture* TextureDI = static_cast<FNiagaraDataInterfaceProxyTexture*>(Context.DataInterface);
-		FNDITextureInstanceData_RenderThread* InstanceData = TextureDI->InstanceData_RT.Find(Context.SystemInstanceID);
-		if ( InstanceData && InstanceData->ResolvedTextureRHI.IsValid() )
-		{
-			FRHISamplerState* SamplerStateRHI = InstanceData->SamplerStateRHI ? InstanceData->SamplerStateRHI : GBlackTexture->SamplerStateRHI;
-
-			SetTextureParameter(
-				RHICmdList,
-				ComputeShaderRHI,
-				TextureParam,
-				SamplerParam,
-				SamplerStateRHI,
-				InstanceData->ResolvedTextureRHI
-			);
-			SetShaderValue(RHICmdList, ComputeShaderRHI, Dimensions, InstanceData->TextureSize);
-		}
-		else
-		{
-			SetTextureParameter(
-				RHICmdList,
-				ComputeShaderRHI,
-				TextureParam,
-				SamplerParam,
-				GBlackTexture->SamplerStateRHI,
-				GBlackTexture->TextureRHI
-			);
-			SetShaderValue(RHICmdList, ComputeShaderRHI, Dimensions, FVector2f::ZeroVector);
-		}
+		Parameters->TextureSize		= FVector2f::ZeroVector;
+		Parameters->Texture			= GBlackTexture->TextureRHI;
+		Parameters->TextureSampler	= GBlackTexture->SamplerStateRHI;
 	}
-
-private:
-	LAYOUT_FIELD(FShaderResourceParameter, TextureParam);
-	LAYOUT_FIELD(FShaderResourceParameter, SamplerParam);
-	LAYOUT_FIELD(FShaderParameter, Dimensions);
-};
-
-IMPLEMENT_TYPE_LAYOUT(FNiagaraDataInterfaceParametersCS_Texture);
-
-IMPLEMENT_NIAGARA_DI_PARAMETER(UNiagaraDataInterfaceTexture, FNiagaraDataInterfaceParametersCS_Texture);
+}
 
 void UNiagaraDataInterfaceTexture::SetTexture(UTexture* InTexture)
 {
 	Texture = InTexture;
 }
 
-#undef LOCTEXT_NAMESPACE
\ No newline at end of file
+#undef LOCTEXT_NAMESPACE
diff --git a/Engine/Plugins/FX/Niagara/Source/Niagara/Private/NiagaraDataInterfaceVolumeTexture.cpp b/Engine/Plugins/FX/Niagara/Source/Niagara/Private/NiagaraDataInterfaceVolumeTexture.cpp
--- a/Engine/Plugins/FX/Niagara/Source/Niagara/Private/NiagaraDataInterfaceVolumeTexture.cpp
+++ b/Engine/Plugins/FX/Niagara/Source/Niagara/Private/NiagaraDataInterfaceVolumeTexture.cpp
@@ -1,21 +1,20 @@
 // Copyright Epic Games, Inc. All Rights Reserved.
 
 #include "NiagaraDataInterfaceVolumeTexture.h"
+#include "NiagaraComputeExecutionContext.h"
 #include "NiagaraShader.h"
-#include "ShaderParameterUtils.h"
-#include "NiagaraCustomVersion.h"
-#include "Engine/VolumeTexture.h"
-#include "Engine/TextureRenderTargetVolume.h"
+#include "NiagaraShaderParametersBuilder.h"
 #include "NiagaraSystemInstance.h"
-#include "NiagaraComputeExecutionContext.h"
+
+#include "Engine/TextureRenderTargetVolume.h"
+#include "Engine/VolumeTexture.h"
+#include "ShaderParameterUtils.h"
 
 #define LOCTEXT_NAMESPACE "UNiagaraDataInterfaceVolumeTexture"
 
+const TCHAR* UNiagaraDataInterfaceVolumeTexture::TemplateShaderFilePath = TEXT("/Plugin/FX/Niagara/Private/NiagaraDataInterfaceVolumeTextureTemplate.ush");
 const FName UNiagaraDataInterfaceVolumeTexture::SampleVolumeTextureName(TEXT("SampleVolumeTexture"));
 const FName UNiagaraDataInterfaceVolumeTexture::TextureDimsName(TEXT("TextureDimensions3D"));
-const FString UNiagaraDataInterfaceVolumeTexture::TextureName(TEXT("Texture_"));
-const FString UNiagaraDataInterfaceVolumeTexture::SamplerName(TEXT("Sampler_"));
-const FString UNiagaraDataInterfaceVolumeTexture::DimensionsBaseName(TEXT("Dimensions_"));
 
 struct FNDIVolumeTextureInstanceData_GameThread
 {
@@ -268,96 +267,61 @@ void UNiagaraDataInterfaceVolumeTexture::SampleVolumeTexture(FVectorVMExternalFu
 }
 
 #if WITH_EDITORONLY_DATA
-bool UNiagaraDataInterfaceVolumeTexture::GetFunctionHLSL(const FNiagaraDataInterfaceGPUParamInfo& ParamInfo, const FNiagaraDataInterfaceGeneratedFunction& FunctionInfo, int FunctionInstanceIndex, FString& OutHLSL)
+bool UNiagaraDataInterfaceVolumeTexture::AppendCompileHash(FNiagaraCompileHashVisitor* InVisitor) const
 {
-	if (FunctionInfo.DefinitionName == SampleVolumeTextureName)
+	bool bSuccess = Super::AppendCompileHash(InVisitor);
+	InVisitor->UpdateString(TEXT("UNiagaraDataInterfaceVolumeTextureHLSLSource"), GetShaderFileHash(TemplateShaderFilePath, EShaderPlatform::SP_PCD3D_SM5).ToString());
+	bSuccess &= InVisitor->UpdateShaderParameters<FShaderParameters>();
+	return bSuccess;
+}
+
+void UNiagaraDataInterfaceVolumeTexture::GetParameterDefinitionHLSL(const FNiagaraDataInterfaceGPUParamInfo& ParamInfo, FString& OutHLSL)
+{
+	TMap<FString, FStringFormatArg> TemplateArgs =
 	{
-		FString HLSLTextureName = TextureName + ParamInfo.DataInterfaceHLSLSymbol;
-		FString HLSLSamplerName = SamplerName + ParamInfo.DataInterfaceHLSLSymbol;
-		OutHLSL += TEXT("void ") + FunctionInfo.InstanceName + TEXT("(in float3 In_UV, in float MipLevel, out float4 Out_Value) \n{\n");
-		OutHLSL += TEXT("\t Out_Value = ") + HLSLTextureName + TEXT(".SampleLevel(") + HLSLSamplerName + TEXT(", In_UV, MipLevel);\n");
-		OutHLSL += TEXT("\n}\n");
-		return true;
-	}
-	else if (FunctionInfo.DefinitionName == TextureDimsName)
+		{TEXT("ParameterName"),	ParamInfo.DataInterfaceHLSLSymbol},
+	};
+
+	FString TemplateFile;
+	LoadShaderSourceFile(TemplateShaderFilePath, EShaderPlatform::SP_PCD3D_SM5, &TemplateFile, nullptr);
+	OutHLSL += FString::Format(*TemplateFile, TemplateArgs);
+}
+
+bool UNiagaraDataInterfaceVolumeTexture::GetFunctionHLSL(const FNiagaraDataInterfaceGPUParamInfo& ParamInfo, const FNiagaraDataInterfaceGeneratedFunction& FunctionInfo, int FunctionInstanceIndex, FString& OutHLSL)
+{
+	if ((FunctionInfo.DefinitionName == SampleVolumeTextureName) ||
+		(FunctionInfo.DefinitionName == TextureDimsName) )
 	{
-		FString DimsVar = DimensionsBaseName + ParamInfo.DataInterfaceHLSLSymbol;
-		OutHLSL += TEXT("void ") + FunctionInfo.InstanceName + TEXT("(out float3 Out_Value) \n{\n");
-		OutHLSL += TEXT("\t Out_Value = ") + DimsVar + TEXT(";\n");
-		OutHLSL += TEXT("\n}\n");
 		return true;
 	}
 	return false;
 }
+#endif
 
-void UNiagaraDataInterfaceVolumeTexture::GetParameterDefinitionHLSL(const FNiagaraDataInterfaceGPUParamInfo& ParamInfo, FString& OutHLSL)
+void UNiagaraDataInterfaceVolumeTexture::BuildShaderParameters(FNiagaraShaderParametersBuilder& ShaderParametersBuilder) const
 {
-	FString HLSLTextureName = TextureName + ParamInfo.DataInterfaceHLSLSymbol;
-	FString HLSLSamplerName = SamplerName + ParamInfo.DataInterfaceHLSLSymbol;
-	OutHLSL += TEXT("Texture3D ") + HLSLTextureName + TEXT(";\n");
-	OutHLSL += TEXT("SamplerState ") + HLSLSamplerName + TEXT(";\n");
-	OutHLSL += TEXT("float3 ") + DimensionsBaseName + ParamInfo.DataInterfaceHLSLSymbol + TEXT(";\n");
+	ShaderParametersBuilder.AddNestedStruct<FShaderParameters>();
 }
-#endif
 
-struct FNiagaraDataInterfaceParametersCS_VolumeTexture : public FNiagaraDataInterfaceParametersCS
+void UNiagaraDataInterfaceVolumeTexture::SetShaderParameters(const FNiagaraDataInterfaceSetShaderParametersContext& Context) const
 {
-	DECLARE_TYPE_LAYOUT(FNiagaraDataInterfaceParametersCS_VolumeTexture, NonVirtual);
-public:
-	void Bind(const FNiagaraDataInterfaceGPUParamInfo& ParameterInfo, const class FShaderParameterMap& ParameterMap)
-	{
-		FString TexName = UNiagaraDataInterfaceVolumeTexture::TextureName + ParameterInfo.DataInterfaceHLSLSymbol;
-		FString SampleName = (UNiagaraDataInterfaceVolumeTexture::SamplerName + ParameterInfo.DataInterfaceHLSLSymbol);
-		TextureParam.Bind(ParameterMap, *TexName);
-		SamplerParam.Bind(ParameterMap, *SampleName);
-
-
-		Dimensions.Bind(ParameterMap, *(UNiagaraDataInterfaceVolumeTexture::DimensionsBaseName + ParameterInfo.DataInterfaceHLSLSymbol));
+	const FNiagaraDataInterfaceProxyVolumeTexture& TextureProxy = Context.GetProxy<FNiagaraDataInterfaceProxyVolumeTexture>();
+	const FNDIVolumeTextureInstanceData_RenderThread* InstanceData = TextureProxy.InstanceData_RT.Find(Context.GetSystemInstanceID());
 
+	FShaderParameters* Parameters = Context.GetParameterNestedStruct<FShaderParameters>();
+	if (InstanceData && InstanceData->ResolvedTextureRHI.IsValid())
+	{
+		Parameters->TextureSize		= InstanceData->TextureSize;
+		Parameters->Texture			= InstanceData->ResolvedTextureRHI;
+		Parameters->TextureSampler	= InstanceData->SamplerStateRHI ? InstanceData->SamplerStateRHI : GBlackVolumeTexture->SamplerStateRHI;
 	}
-
-	void Set(FRHICommandList& RHICmdList, const FNiagaraDataInterfaceSetArgs& Context) const
+	else
 	{
-		check(IsInRenderingThread());
-
-		FRHIComputeShader* ComputeShaderRHI = Context.Shader.GetComputeShader();
-		FNiagaraDataInterfaceProxyVolumeTexture* TextureDI = static_cast<FNiagaraDataInterfaceProxyVolumeTexture*>(Context.DataInterface);
-		FNDIVolumeTextureInstanceData_RenderThread* InstanceData = TextureDI->InstanceData_RT.Find(Context.SystemInstanceID);
-
-		if (InstanceData && InstanceData->ResolvedTextureRHI.IsValid())
-		{
-			FRHISamplerState* SamplerStateRHI = InstanceData->SamplerStateRHI ? InstanceData->SamplerStateRHI : GBlackTexture->SamplerStateRHI;
-			SetTextureParameter(
-				RHICmdList,
-				ComputeShaderRHI,
-				TextureParam,
-				SamplerParam,
-				SamplerStateRHI,
-				InstanceData->ResolvedTextureRHI
-			);
-			SetShaderValue(RHICmdList, ComputeShaderRHI, Dimensions, InstanceData->TextureSize);
-		}
-		else
-		{
-			SetTextureParameter(
-				RHICmdList,
-				ComputeShaderRHI,
-				TextureParam,
-				SamplerParam,
-				GBlackVolumeTexture->SamplerStateRHI,
-				GBlackVolumeTexture->TextureRHI
-			);
-			SetShaderValue(RHICmdList, ComputeShaderRHI, Dimensions, FVector3f::ZeroVector);
-		}
+		Parameters->TextureSize		= FVector3f::ZeroVector;
+		Parameters->Texture			= GBlackVolumeTexture->TextureRHI;
+		Parameters->TextureSampler	= GBlackVolumeTexture->SamplerStateRHI;
 	}
-private:
-	LAYOUT_FIELD(FShaderResourceParameter, TextureParam);
-	LAYOUT_FIELD(FShaderResourceParameter, SamplerParam);
-	LAYOUT_FIELD(FShaderParameter, Dimensions);
-};
-
-IMPLEMENT_TYPE_LAYOUT(FNiagaraDataInterfaceParametersCS_VolumeTexture);
-IMPLEMENT_NIAGARA_DI_PARAMETER(UNiagaraDataInterfaceVolumeTexture, FNiagaraDataInterfaceParametersCS_VolumeTexture);
+}
 
 void UNiagaraDataInterfaceVolumeTexture::SetTexture(UVolumeTexture* InTexture)
 {
