Niagara scalability context widget now only works on stack item selections and not node selections
Additionally, if nothing is selected a user hint text is displayed asking the user to select an applicable item

diff --git a/Engine/Plugins/FX/Niagara/Source/NiagaraEditor/Private/Widgets/Scalability/SNiagaraScalabilityContext.cpp b/Engine/Plugins/FX/Niagara/Source/NiagaraEditor/Private/Widgets/Scalability/SNiagaraScalabilityContext.cpp
--- a/Engine/Plugins/FX/Niagara/Source/NiagaraEditor/Private/Widgets/Scalability/SNiagaraScalabilityContext.cpp
+++ b/Engine/Plugins/FX/Niagara/Source/NiagaraEditor/Private/Widgets/Scalability/SNiagaraScalabilityContext.cpp
@@ -7,13 +7,16 @@
 #include "NiagaraMeshRendererProperties.h"
 #include "NiagaraObjectSelection.h"
 #include "Modules/ModuleManager.h"
-#include "Widgets/Input/SButton.h"
 #include "ViewModels/NiagaraOverviewGraphViewModel.h"
 #include "ViewModels/NiagaraSystemSelectionViewModel.h"
 #include "ViewModels/NiagaraSystemViewModel.h"
 #include "ViewModels/Stack/NiagaraStackEntry.h"
 #include "ViewModels/Stack/NiagaraStackRendererItem.h"
 #include "PropertyEditorModule.h"
+#include "ViewModels/NiagaraEmitterViewModel.h"
+#include "ViewModels/Stack/NiagaraStackEmitterPropertiesGroup.h"
+#include "ViewModels/Stack/NiagaraStackSystemSettingsGroup.h"
+#include "Widgets/Layout/SWidgetSwitcher.h"
 
 #define LOCTEXT_NAMESPACE "NiagaraScalabilityContext"
 
@@ -22,7 +25,6 @@ void SNiagaraScalabilityContext::Construct(const FArguments& InArgs, UNiagaraSys
 	ScalabilityViewModel = &InScalabilityViewModel;
 	
 	ScalabilityViewModel->GetSystemViewModel().Pin()->GetSelectionViewModel()->OnEntrySelectionChanged().AddSP(this, &SNiagaraScalabilityContext::UpdateScalabilityContent);
-	ScalabilityViewModel->GetSystemViewModel().Pin()->GetOverviewGraphViewModel()->GetNodeSelection()->OnSelectedObjectsChanged().AddSP(this, &SNiagaraScalabilityContext::UpdateScalabilityContent);
 
 	FPropertyEditorModule& PropertyEditorModule = FModuleManager::Get().LoadModuleChecked<FPropertyEditorModule>("PropertyEditor");
 	
@@ -42,62 +44,76 @@ void SNiagaraScalabilityContext::Construct(const FArguments& InArgs, UNiagaraSys
 
 	ChildSlot
 	[
-		DetailsView.ToSharedRef()
+		SNew(SWidgetSwitcher)
+		.WidgetIndex_Lambda([this]
+		{
+			return DetailsView->GetSelectedObjects().Num() == 1 ? 0 : 1;  
+		})
+		+ SWidgetSwitcher::Slot()
+		[
+			DetailsView.ToSharedRef()
+		]
+		+ SWidgetSwitcher::Slot()
+		.HAlign(HAlign_Center)
+		.VAlign(VAlign_Center)
+		.Padding(10.f)
+		[
+			SNew(STextBlock)
+			.Text(LOCTEXT("EmptyScalabilityContextDescription", "This tab displays scalability settings of the selected emitters, system, or renderers.\nPlease select emitter or system properties or a renderer."))
+			.ColorAndOpacity(FLinearColor(0.8f, 0.8f, 0.8f, 0.5f))
+			.AutoWrapText(true)
+		]
 	];
 }
 
 void SNiagaraScalabilityContext::SetObject(UObject* Object)
-{	
-	if(UNiagaraEmitter* Emitter = Cast<UNiagaraEmitter>(Object))
-	{		
-		DetailsView->SetObject(Emitter, true);
+{
+	if(Object == nullptr)
+	{
+		DetailsView->SetObject(nullptr);
+	}
+	else if(UNiagaraEmitter* Emitter = Cast<UNiagaraEmitter>(Object))
+	{
+		DetailsView->SetObject(Emitter);
 	}
 	else if(UNiagaraSystem* System = Cast<UNiagaraSystem>(Object))
 	{
-		DetailsView->SetObject(System, true);
+		DetailsView->SetObject(System);
 	}
 	else if(UNiagaraStackRendererItem* RendererItem = Cast<UNiagaraStackRendererItem>(Object))
 	{
-		DetailsView->SetObject(RendererItem->GetRendererProperties(), true);
+		DetailsView->SetObject(RendererItem->GetRendererProperties());
 	}
 }
 
 void SNiagaraScalabilityContext::UpdateScalabilityContent()
 {
+	UObject* NewSelection = nullptr;
+	
 	TArray<UNiagaraStackEntry*> StackEntries;
 	ScalabilityViewModel->GetSystemViewModel().Pin()->GetSelectionViewModel()->GetSelectedEntries(StackEntries);
-	const TSet<UObject*>& SelectedNodes = ScalabilityViewModel->GetSystemViewModel().Pin()->GetOverviewGraphViewModel()->GetNodeSelection()->GetSelectedObjects();
-
+	
 	if(StackEntries.Num() == 1)
 	{
 		UNiagaraStackEntry* StackEntry = StackEntries[0];
-		bool bDoesSupportScalabilityContext =  StackEntry->IsA(UNiagaraStackRendererItem::StaticClass());
 
-		if(bDoesSupportScalabilityContext)
+		if(StackEntry->IsA(UNiagaraStackRendererItem::StaticClass()))
 		{
-			SetObject(StackEntry);
-			return;
+			NewSelection = StackEntry;
 		}
-	}
-	
-	if(SelectedNodes.Num() > 0)
-	{
-		UObject* Object = SelectedNodes.Array()[0];
-		UNiagaraOverviewNode* OverviewNode = Cast<UNiagaraOverviewNode>(Object);
-		if(OverviewNode)
+		else if(StackEntry->IsA(UNiagaraStackEmitterPropertiesGroup::StaticClass()))
 		{
-			if(FNiagaraEmitterHandle* Handle = OverviewNode->TryGetEmitterHandle())
-			{
-				SetObject(Handle->GetInstance());
-				return;
-			}
-			else
-			{
-				SetObject(OverviewNode->GetOwningSystem());
-				return;
-			}
+			UNiagaraEmitter* Emitter = Cast<UNiagaraStackEmitterPropertiesGroup>(StackEntry)->GetEmitterViewModel()->GetEmitter();
+			NewSelection = Emitter;
+		}
+		else if(StackEntry->IsA(UNiagaraStackSystemPropertiesGroup::StaticClass()))
+		{
+			UNiagaraSystem& System = Cast<UNiagaraStackSystemPropertiesGroup>(StackEntry)->GetSystemViewModel()->GetSystem();
+			NewSelection = &System;
 		}
 	}
+	
+	SetObject(NewSelection);
 }
 
 bool SNiagaraScalabilityContext::FilterScalabilityProperties(const FPropertyAndParent& InPropertyAndParent) const
