Removed legacy parameter widgets & Removed some unneeded includes

diff --git a/Engine/Plugins/FX/Niagara/Source/NiagaraEditor/Private/NiagaraNodeParameterMapBase.cpp b/Engine/Plugins/FX/Niagara/Source/NiagaraEditor/Private/NiagaraNodeParameterMapBase.cpp
--- a/Engine/Plugins/FX/Niagara/Source/NiagaraEditor/Private/NiagaraNodeParameterMapBase.cpp
+++ b/Engine/Plugins/FX/Niagara/Source/NiagaraEditor/Private/NiagaraNodeParameterMapBase.cpp
@@ -3,29 +3,17 @@
 #include "NiagaraNodeParameterMapBase.h"
 #include "EdGraphSchema_Niagara.h"
 #include "NiagaraEditorUtilities.h"
-#include "SNiagaraGraphNodeConvert.h"
-#include "INiagaraCompiler.h"
 #include "NiagaraNodeOutput.h"
 #include "ScopedTransaction.h"
 #include "SNiagaraGraphPinAdd.h"
 #include "NiagaraGraph.h"
 #include "NiagaraComponent.h"
-#include "EdGraphSchema_Niagara.h"
 #include "NiagaraScript.h"
-#include "NiagaraScriptSource.h"
 #include "NiagaraConstants.h"
-#include "NiagaraParameterCollection.h"
-#include "Widgets/SNiagaraParameterMapView.h"
-#include "Widgets/Input/SEditableTextBox.h"
 #include "ToolMenus.h"
 #include "NiagaraScriptVariable.h"
-#include "Widgets/SNiagaraParameterName.h"
-
-#include "IAssetTools.h"
 #include "AssetRegistry/AssetRegistryModule.h"
 
-#include "Framework/MultiBox/MultiBoxBuilder.h"
-
 #define LOCTEXT_NAMESPACE "NiagaraNodeParameterMapBase"
 
 const FName UNiagaraNodeParameterMapBase::ParameterPinSubCategory("ParameterPin");
diff --git a/Engine/Plugins/FX/Niagara/Source/NiagaraEditor/Private/Toolkits/NiagaraScriptToolkit.cpp b/Engine/Plugins/FX/Niagara/Source/NiagaraEditor/Private/Toolkits/NiagaraScriptToolkit.cpp
--- a/Engine/Plugins/FX/Niagara/Source/NiagaraEditor/Private/Toolkits/NiagaraScriptToolkit.cpp
+++ b/Engine/Plugins/FX/Niagara/Source/NiagaraEditor/Private/Toolkits/NiagaraScriptToolkit.cpp
@@ -21,7 +21,6 @@
 #include "Widgets/SNiagaraParameterDefinitionsPanel.h"
 #include "NiagaraScriptSource.h"
 #include "NiagaraStandaloneScriptViewModel.h"
-#include "NiagaraSystem.h"
 #include "NiagaraVersionMetaData.h"
 #include "PropertyEditorModule.h"
 #include "SGraphActionMenu.h"
@@ -36,10 +35,8 @@
 #include "Misc/MessageDialog.h"
 #include "Modules/ModuleManager.h"
 #include "UObject/Package.h"
-#include "UObject/UObjectIterator.h"
 #include "ViewModels/NiagaraParameterPanelViewModel.h"
 #include "ViewModels/NiagaraScriptViewModel.h"
-#include "Widgets/SNiagaraParameterMapView.h"
 #include "Widgets/SNiagaraParameterPanel.h"
 #include "Widgets/SNiagaraScriptGraph.h"
 #include "Widgets/SNiagaraSelectedObjectsDetails.h"
@@ -109,11 +106,6 @@ void FNiagaraScriptToolkit::RegisterTabSpawners(const TSharedRef<class FTabManag
 		.SetDisplayName(LOCTEXT("ParametersTab", "Parameters"))
 		.SetGroup(WorkspaceMenuCategoryRef)
 		.SetIcon(FSlateIcon(FNiagaraEditorStyle::Get().GetStyleSetName(), "Tab.Parameters"));
-
-	InTabManager->RegisterTabSpawner(ParametersTabId2, FOnSpawnTab::CreateSP(this, &FNiagaraScriptToolkit::SpawnTabScriptParameters2)) //@todo(ng) remove this
-		.SetDisplayName(LOCTEXT("ParametersTab2", "Legacy Parameters"))
-		.SetGroup(WorkspaceMenuCategoryRef)
-		.SetIcon(FSlateIcon(FNiagaraEditorStyle::Get().GetStyleSetName(), "Tab.Parameters"));
 		
 	//@todo(ng) disable parameter definitions panel pending bug fixes
 // 	InTabManager->RegisterTabSpawner(ParameterDefinitionsTabId, FOnSpawnTab::CreateSP(this, &FNiagaraScriptToolkit::SpawnTabParameterDefinitions))
@@ -497,23 +489,6 @@ TSharedRef<SDockTab> FNiagaraScriptToolkit::SpawnTabScriptParameters(const FSpaw
 	return SpawnedTab;
 }
 
-TSharedRef<SDockTab> FNiagaraScriptToolkit::SpawnTabScriptParameters2(const FSpawnTabArgs& Args)
-{
-	checkf(Args.GetTabId().TabType == ParametersTabId2, TEXT("Wrong tab ID in NiagaraScriptToolkit"));
-
-	TArray<TSharedRef<FNiagaraObjectSelection>> Array;
-	Array.Push(DetailsScriptSelection.ToSharedRef());
-	Array.Push(ScriptViewModel->GetVariableSelection());
-
-	TSharedRef<SDockTab> SpawnedTab =
-		SNew(SDockTab)
-		[
-			SNew(SNiagaraParameterMapView, Array, SNiagaraParameterMapView::EToolkitType::SCRIPT, GetToolkitCommands())
-		];
-
-	return SpawnedTab;
-}
-
 TSharedRef<SDockTab> FNiagaraScriptToolkit::SpawnTabParameterDefinitions(const FSpawnTabArgs& Args)
 {
 	checkf(Args.GetTabId().TabType == ParameterDefinitionsTabId, TEXT("Wrong tab ID in NiagaraScriptToolkit"));
diff --git a/Engine/Plugins/FX/Niagara/Source/NiagaraEditor/Private/Toolkits/NiagaraScriptToolkit.h b/Engine/Plugins/FX/Niagara/Source/NiagaraEditor/Private/Toolkits/NiagaraScriptToolkit.h
--- a/Engine/Plugins/FX/Niagara/Source/NiagaraEditor/Private/Toolkits/NiagaraScriptToolkit.h
+++ b/Engine/Plugins/FX/Niagara/Source/NiagaraEditor/Private/Toolkits/NiagaraScriptToolkit.h
@@ -102,7 +102,6 @@ private:
 
 	/** Spawns the tab with the parameter view. */
 	TSharedRef<SDockTab> SpawnTabScriptParameters(const FSpawnTabArgs& Args);
-	TSharedRef<SDockTab> SpawnTabScriptParameters2(const FSpawnTabArgs& Args);
 	TSharedRef<SDockTab> SpawnTabParameterDefinitions(const FSpawnTabArgs& Args);
 	TSharedRef<SDockTab> SpawnTabStats(const FSpawnTabArgs& Args);
 
diff --git a/Engine/Plugins/FX/Niagara/Source/NiagaraEditor/Private/Toolkits/NiagaraSystemToolkit.cpp b/Engine/Plugins/FX/Niagara/Source/NiagaraEditor/Private/Toolkits/NiagaraSystemToolkit.cpp
--- a/Engine/Plugins/FX/Niagara/Source/NiagaraEditor/Private/Toolkits/NiagaraSystemToolkit.cpp
+++ b/Engine/Plugins/FX/Niagara/Source/NiagaraEditor/Private/Toolkits/NiagaraSystemToolkit.cpp
@@ -5,7 +5,6 @@
 #include "NiagaraSystem.h"
 #include "NiagaraEmitter.h"
 #include "NiagaraEmitterHandle.h"
-#include "NiagaraScriptSource.h"
 #include "NiagaraObjectSelection.h"
 #include "ViewModels/NiagaraSystemViewModel.h"
 #include "ViewModels/NiagaraEmitterHandleViewModel.h"
@@ -17,12 +16,8 @@
 #include "NiagaraSystemScriptViewModel.h"
 #include "Widgets/SNiagaraSystemScript.h"
 #include "Widgets/SNiagaraSystemViewport.h"
-#include "Widgets/SNiagaraSelectedObjectsDetails.h"
-#include "Widgets/SNiagaraParameterMapView.h"
 #include "Widgets/SNiagaraParameterPanel.h"
 #include "Widgets/SNiagaraSpreadsheetView.h"
-#include "Widgets/SNiagaraGeneratedCodeView.h"
-#include "Widgets/SNiagaraScriptGraph.h"
 #include "Widgets/SNiagaraDebugger.h"
 #include "NiagaraEditorCommands.h"
 #include "NiagaraEditorStyle.h"
@@ -37,16 +32,13 @@
 #include "IContentBrowserSingleton.h"
 #include "ContentBrowserModule.h"
 
-#include "EditorStyleSet.h"
 #include "Toolkits/AssetEditorToolkit.h"
 #include "ScopedTransaction.h"
 
 #include "Framework/Application/SlateApplication.h"
 #include "Widgets/SBoxPanel.h"
 #include "Widgets/Layout/SBox.h"
-#include "Widgets/Layout/SSplitter.h"
 #include "Widgets/Docking/SDockTab.h"
-#include "AdvancedPreviewSceneModule.h"
 #include "BusyCursor.h"
 #include "Misc/FeedbackContext.h"
 #include "Editor.h"
@@ -57,17 +49,12 @@
 #include "Modules/ModuleManager.h"
 #include "Misc/FileHelper.h"
 #include "NiagaraMessageLogViewModel.h"
-#include "NiagaraScriptStatsViewModel.h"
-#include "Widgets/Input/SCheckBox.h"
 #include "ViewModels/NiagaraParameterPanelViewModel.h"
 #include "NiagaraEditor/Private/SNiagaraAssetPickerList.h"
 #include "SystemToolkitModes/NiagaraSystemToolkitModeBase.h"
 #include "SystemToolkitModes/NiagaraSystemToolkitMode_Default.h"
 #include "SystemToolkitModes/NiagaraSystemToolkitMode_Scalability.h"
-#include "Widgets/SNiagaraParameterDefinitionsPanel.h"
 #include "ViewModels/NiagaraParameterDefinitionsPanelViewModel.h"
-#include "ViewModels/Stack/NiagaraStackViewModel.h"
-
 
 #define LOCTEXT_NAMESPACE "NiagaraSystemEditor"
 
diff --git a/Engine/Plugins/FX/Niagara/Source/NiagaraEditor/Private/Toolkits/NiagaraSystemToolkit.h b/Engine/Plugins/FX/Niagara/Source/NiagaraEditor/Private/Toolkits/NiagaraSystemToolkit.h
--- a/Engine/Plugins/FX/Niagara/Source/NiagaraEditor/Private/Toolkits/NiagaraSystemToolkit.h
+++ b/Engine/Plugins/FX/Niagara/Source/NiagaraEditor/Private/Toolkits/NiagaraSystemToolkit.h
@@ -229,8 +229,6 @@ private:
 	/** Baker preview */
 	TSharedPtr<FNiagaraBakerViewModel> BakerViewModel;
 
-	TSharedPtr<class SNiagaraParameterMapView> ParameterMapView;
-
 	TSharedPtr<FNiagaraSystemToolkitParameterPanelViewModel> ParameterPanelViewModel;
 	TSharedPtr<FNiagaraSystemToolkitParameterDefinitionsPanelViewModel> ParameterDefinitionsPanelViewModel;
 	TSharedPtr<class SNiagaraParameterPanel> ParameterPanel;
diff --git a/Engine/Plugins/FX/Niagara/Source/NiagaraEditor/Private/Toolkits/SystemToolkitModes/NiagaraSystemToolkitModeBase.cpp b/Engine/Plugins/FX/Niagara/Source/NiagaraEditor/Private/Toolkits/SystemToolkitModes/NiagaraSystemToolkitModeBase.cpp
--- a/Engine/Plugins/FX/Niagara/Source/NiagaraEditor/Private/Toolkits/SystemToolkitModes/NiagaraSystemToolkitModeBase.cpp
+++ b/Engine/Plugins/FX/Niagara/Source/NiagaraEditor/Private/Toolkits/SystemToolkitModes/NiagaraSystemToolkitModeBase.cpp
@@ -9,6 +9,7 @@
 #include "NiagaraScriptGraphViewModel.h"
 #include "NiagaraSystemToolkit.h"
 #include "NiagaraScriptSource.h"
+#include "NiagaraScriptStatsViewModel.h"
 #include "ViewModels/NiagaraEmitterHandleViewModel.h"
 #include "ViewModels/NiagaraEmitterViewModel.h"
 #include "ViewModels/NiagaraScriptViewModel.h"
@@ -19,7 +20,6 @@
 #include "Widgets/SNiagaraSystemViewport.h"
 #include "Widgets/SNiagaraSystemScript.h"
 #include "Widgets/SNiagaraParameterPanel.h"
-#include "Widgets/SNiagaraParameterMapView.h"
 #include "Widgets/SNiagaraParameterDefinitionsPanel.h"
 #include "Widgets/SNiagaraScriptGraph.h"
 #include "Widgets/SNiagaraSelectedObjectsDetails.h"
@@ -34,7 +34,6 @@ const FName FNiagaraSystemToolkitModeBase::SequencerTabID(TEXT("NiagaraSystemEdi
 const FName FNiagaraSystemToolkitModeBase::SystemScriptTabID(TEXT("NiagaraSystemEditor_SystemScript"));
 const FName FNiagaraSystemToolkitModeBase::SystemDetailsTabID(TEXT("NiagaraSystemEditor_SystemDetails"));
 const FName FNiagaraSystemToolkitModeBase::SystemParametersTabID(TEXT("NiagaraSystemEditor_SystemParameters"));
-const FName FNiagaraSystemToolkitModeBase::SystemParametersTabID2(TEXT("NiagaraSystemEditor_SystemParameters2"));
 const FName FNiagaraSystemToolkitModeBase::SystemParameterDefinitionsTabID(TEXT("NiagaraSystemEditor_SystemParameterDefinitions"));
 const FName FNiagaraSystemToolkitModeBase::SelectedEmitterStackTabID(TEXT("NiagaraSystemEditor_SelectedEmitterStack"));
 const FName FNiagaraSystemToolkitModeBase::SelectedEmitterGraphTabID(TEXT("NiagaraSystemEditor_SelectedEmitterGraph"));
@@ -78,11 +77,6 @@ void FNiagaraSystemToolkitModeBase::RegisterTabFactories(TSharedPtr<FTabManager>
 		.SetGroup(WorkspaceMenuCategory.ToSharedRef())
 		.SetIcon(FSlateIcon(FNiagaraEditorStyle::Get().GetStyleSetName(), "Tab.Parameters"));
 
-	InTabManager->RegisterTabSpawner(SystemParametersTabID2, FOnSpawnTab::CreateSP(this, &FNiagaraSystemToolkitModeBase::SpawnTab_SystemParameters2))
-		.SetDisplayName(LOCTEXT("SystemParameters2", "Legacy Parameters"))
-		.SetGroup(WorkspaceMenuCategory.ToSharedRef())
-		.SetIcon(FSlateIcon(FNiagaraEditorStyle::Get().GetStyleSetName(), "Tab.Parameters"));
-
 //@todo(ng) disable parameter definitions panel pending bug fixes
 // 	InTabManager->RegisterTabSpawner(SystemParameterDefinitionsTabID, FOnSpawnTab::CreateSP(this, &FNiagaraSystemToolkit::SpawnTab_SystemParameterDefinitions)) 
 // 		.SetDisplayName(LOCTEXT("SystemParameterDefinitions", "Parameter Definitions"))
@@ -370,24 +364,6 @@ TSharedRef<SDockTab> FNiagaraSystemToolkitModeBase::SpawnTab_SystemParameters(co
 	return SpawnedTab;
 }
 
-//@todo(ng) cleanup
-TSharedRef<SDockTab> FNiagaraSystemToolkitModeBase::SpawnTab_SystemParameters2(const FSpawnTabArgs& Args)
-{
-	check(Args.GetTabId().TabType == SystemParametersTabID2);
-
-	TArray<TSharedRef<FNiagaraObjectSelection>> ObjectSelections;
-	ObjectSelections.Add(SystemToolkit.Pin()->ObjectSelectionForParameterMapView.ToSharedRef());
-
-	TSharedRef<SDockTab> SpawnedTab =
-		SNew(SDockTab)
-		[
-			SAssignNew(SystemToolkit.Pin()->ParameterMapView, SNiagaraParameterMapView, ObjectSelections, SNiagaraParameterMapView::EToolkitType::SYSTEM, SystemToolkit.Pin()->GetToolkitCommands())
-		];
-	SystemToolkit.Pin()->RefreshParameters();
-
-	return SpawnedTab;
-}
-
 TSharedRef<SDockTab> FNiagaraSystemToolkitModeBase::SpawnTab_SystemParameterDefinitions(const FSpawnTabArgs& Args)
 {
 	checkf(Args.GetTabId().TabType == SystemParameterDefinitionsTabID, TEXT("Wrong tab ID in NiagaraScriptToolkit"));
diff --git a/Engine/Plugins/FX/Niagara/Source/NiagaraEditor/Private/Toolkits/SystemToolkitModes/NiagaraSystemToolkitModeBase.h b/Engine/Plugins/FX/Niagara/Source/NiagaraEditor/Private/Toolkits/SystemToolkitModes/NiagaraSystemToolkitModeBase.h
--- a/Engine/Plugins/FX/Niagara/Source/NiagaraEditor/Private/Toolkits/SystemToolkitModes/NiagaraSystemToolkitModeBase.h
+++ b/Engine/Plugins/FX/Niagara/Source/NiagaraEditor/Private/Toolkits/SystemToolkitModes/NiagaraSystemToolkitModeBase.h
@@ -3,10 +3,7 @@
 #pragma once
 
 #include "NiagaraBakerViewModel.h"
-#include "NiagaraScriptStatsViewModel.h"
 #include "WorkflowOrientedApp/ApplicationMode.h"
-#include "WorkflowOrientedApp/WorkflowCentricApplication.h"
-#include "Framework/MultiBox/MultiBoxBuilder.h"
 #include "Widgets/Docking/SDockTab.h"
 
 class FNiagaraSystemToolkit;
@@ -25,7 +22,6 @@ private:
 	TSharedRef<SDockTab> SpawnTab_Sequencer(const FSpawnTabArgs& Args);
 	TSharedRef<SDockTab> SpawnTab_SystemScript(const FSpawnTabArgs& Args);
 	TSharedRef<SDockTab> SpawnTab_SystemParameters(const FSpawnTabArgs& Args);
-	TSharedRef<SDockTab> SpawnTab_SystemParameters2(const FSpawnTabArgs& Args); //@todo(ng) cleanup
 	TSharedRef<SDockTab> SpawnTab_SystemParameterDefinitions(const FSpawnTabArgs& Args);
 	TSharedRef<SDockTab> SpawnTab_SelectedEmitterStack(const FSpawnTabArgs& Args);
 	TSharedRef<SDockTab> SpawnTab_SelectedEmitterGraph(const FSpawnTabArgs& Args);
@@ -48,7 +44,6 @@ public:
 	static const FName SystemScriptTabID;
 	static const FName SystemDetailsTabID;
 	static const FName SystemParametersTabID;
-	static const FName SystemParametersTabID2;
 	static const FName SystemParameterDefinitionsTabID;
 	static const FName SelectedEmitterStackTabID;
 	static const FName SelectedEmitterGraphTabID;
diff --git a/Engine/Plugins/FX/Niagara/Source/NiagaraEditor/Private/Widgets/SNiagaraParameterMapPaletteItem.cpp b/Engine/Plugins/FX/Niagara/Source/NiagaraEditor/Private/Widgets/SNiagaraParameterMapPaletteItem.cpp
deleted file mode 100644
diff --git a/Engine/Plugins/FX/Niagara/Source/NiagaraEditor/Private/Widgets/SNiagaraParameterMapPaletteItem.h b/Engine/Plugins/FX/Niagara/Source/NiagaraEditor/Private/Widgets/SNiagaraParameterMapPaletteItem.h
deleted file mode 100644
diff --git a/Engine/Plugins/FX/Niagara/Source/NiagaraEditor/Private/Widgets/SNiagaraParameterMapView.cpp b/Engine/Plugins/FX/Niagara/Source/NiagaraEditor/Private/Widgets/SNiagaraParameterMapView.cpp
deleted file mode 100644
diff --git a/Engine/Plugins/FX/Niagara/Source/NiagaraEditor/Private/Widgets/SNiagaraParameterPanel.cpp b/Engine/Plugins/FX/Niagara/Source/NiagaraEditor/Private/Widgets/SNiagaraParameterPanel.cpp
--- a/Engine/Plugins/FX/Niagara/Source/NiagaraEditor/Private/Widgets/SNiagaraParameterPanel.cpp
+++ b/Engine/Plugins/FX/Niagara/Source/NiagaraEditor/Private/Widgets/SNiagaraParameterPanel.cpp
@@ -21,7 +21,6 @@
 #include "Widgets/SToolTip.h"
 #include "Widgets/Layout/SScaleBox.h"
 #include "Widgets/Text/SRichTextBlock.h"
-#include "Widgets/SNiagaraParameterMapView.h"
 #include "Widgets/SNiagaraParameterName.h"
 #include "Widgets/SItemSelector.h"
 
diff --git a/Engine/Plugins/FX/Niagara/Source/NiagaraEditor/Public/Widgets/SNiagaraParameterMapView.h b/Engine/Plugins/FX/Niagara/Source/NiagaraEditor/Public/Widgets/SNiagaraParameterMapView.h
deleted file mode 100644
diff --git a/Engine/Plugins/FX/Niagara/Source/NiagaraEditorWidgets/Private/SNiagaraScratchPad.cpp b/Engine/Plugins/FX/Niagara/Source/NiagaraEditorWidgets/Private/SNiagaraScratchPad.cpp
--- a/Engine/Plugins/FX/Niagara/Source/NiagaraEditorWidgets/Private/SNiagaraScratchPad.cpp
+++ b/Engine/Plugins/FX/Niagara/Source/NiagaraEditorWidgets/Private/SNiagaraScratchPad.cpp
@@ -14,8 +14,6 @@
 #include "NiagaraEditorWidgetsStyle.h"
 #include "NiagaraScratchPadCommandContext.h"
 #include "Widgets/SNiagaraParameterPanel.h"
-#include "Widgets/SNiagaraParameterMapView.h"
-
 #include "Widgets/Text/SInlineEditableTextBlock.h"
 #include "Widgets/Input/SButton.h"
 #include "Widgets/Layout/SSpacer.h"
@@ -830,22 +828,7 @@ private:
 		{
 			ScriptViewModelWeak = NewScriptViewModel;
 			if (NewScriptViewModel.IsValid())
-			{
-				TSharedPtr<SWidget> ParameterPanelWidget;
-				if (NewScriptViewModel->GetParameterPanelViewModel().IsValid())
-				{
-					ParameterPanelWidget = SNew(SNiagaraParameterPanel, NewScriptViewModel->GetParameterPanelViewModel(), NewScriptViewModel->GetParameterPanelCommands());
-				}
-				else
-				{
-					TSharedRef<FNiagaraObjectSelection> ScriptSelection = MakeShared<FNiagaraObjectSelection>();
-					const FVersionedNiagaraScript& EditScript = NewScriptViewModel->GetEditScript();
-					ScriptSelection->SetSelectedObject(EditScript.Script, &EditScript.Version);
-					TArray<TSharedRef<FNiagaraObjectSelection>> ParameterSelections;
-					ParameterSelections.Add(ScriptSelection);
-					ParameterSelections.Add(NewScriptViewModel->GetVariableSelection());
-					ParameterPanelWidget = SNew(SNiagaraParameterMapView, ParameterSelections, SNiagaraParameterMapView::EToolkitType::SCRIPT, NewScriptViewModel->GetParameterPanelCommands());
-				}
+			{				
 				ChildSlot
 				[
 					SNew(SVerticalBox)
@@ -859,7 +842,7 @@ private:
 					]
 					+ SVerticalBox::Slot()
 					[
-						ParameterPanelWidget.ToSharedRef()
+						SNew(SNiagaraParameterPanel, NewScriptViewModel->GetParameterPanelViewModel(), NewScriptViewModel->GetParameterPanelCommands())
 					]
 				];
 			}
