Niagara - Additional fixes to allow testing creation and editing of niagara assets without the plugin supplied data.
+ Fix crashes with scratch pad scripts.
+ Fix issues creating empty emitters.
+ Add menu entries for creating empty emitters to the system overview graph and emitter wizard.
+ Fix issues creating empty scripts.

diff --git a/Engine/Plugins/FX/Niagara/Source/NiagaraEditor/Private/NiagaraEditorCommands.cpp b/Engine/Plugins/FX/Niagara/Source/NiagaraEditor/Private/NiagaraEditorCommands.cpp
--- a/Engine/Plugins/FX/Niagara/Source/NiagaraEditor/Private/NiagaraEditorCommands.cpp
+++ b/Engine/Plugins/FX/Niagara/Source/NiagaraEditor/Private/NiagaraEditorCommands.cpp
@@ -17,7 +17,7 @@ void FNiagaraEditorCommands::RegisterCommands()
 	UI_COMMAND(SelectNextUsage, "Select Next Usage", "Selects the next usage of the selected item.", EUserInterfaceActionType::Button, FInputChord());
 	UI_COMMAND(CreateAssetFromSelection, "Create Asset...", "Creates an asset from the current selection.", EUserInterfaceActionType::Button, FInputChord());
 
-	UI_COMMAND(OpenAddEmitterMenu, "Add Emitter", "Adds an emitter to the system.", EUserInterfaceActionType::Button, FInputChord(EModifierKey::None, EKeys::E));
+	UI_COMMAND(OpenAddEmitterMenu, "Add Emitter...", "Adds an emitter to the system.", EUserInterfaceActionType::Button, FInputChord(EModifierKey::None, EKeys::E));
 
 	UI_COMMAND(ResetSimulation, "Reset", "Resets the current simulation", EUserInterfaceActionType::Button, FInputChord(EKeys::SpaceBar));
 	UI_COMMAND(TogglePreviewGrid, "Grid", "Toggles the preview pane's grid.", EUserInterfaceActionType::ToggleButton, FInputChord() );
diff --git a/Engine/Plugins/FX/Niagara/Source/NiagaraEditor/Private/NiagaraEmitterFactoryNew.cpp b/Engine/Plugins/FX/Niagara/Source/NiagaraEditor/Private/NiagaraEmitterFactoryNew.cpp
--- a/Engine/Plugins/FX/Niagara/Source/NiagaraEditor/Private/NiagaraEmitterFactoryNew.cpp
+++ b/Engine/Plugins/FX/Niagara/Source/NiagaraEditor/Private/NiagaraEmitterFactoryNew.cpp
@@ -74,7 +74,7 @@ bool UNiagaraEmitterFactoryNew::ConfigureProperties()
 	}
 	else
 	{
-		// User selected an empty emitter so set the emitter to copy to null, and disable creationg of default modules.
+		// User selected an empty emitter so set the emitter to copy to null, and disable creation of default modules.
 		EmitterToCopy = nullptr;
 		bAddDefaultModulesAndRenderersToEmptyEmitter = false;
 	}
@@ -141,6 +141,14 @@ UObject* UNiagaraEmitterFactoryNew::FactoryCreateNew(UClass* Class, UObject* InP
 	{
 		// Create an empty emitter, source, and graph.
 		NewEmitter = NewObject<UNiagaraEmitter>(InParent, Class, Name, Flags | RF_Transactional);
+		InitializeEmitter(NewEmitter, bAddDefaultModulesAndRenderersToEmptyEmitter);
+	}
+	return NewEmitter;
+}
+
+void UNiagaraEmitterFactoryNew::InitializeEmitter(UNiagaraEmitter* NewEmitter, bool bAddDefaultModulesAndRenderers)
+{
+	{
 		NewEmitter->SimTarget = ENiagaraSimTarget::CPUSim;
 
 		UNiagaraScriptSource* Source = NewObject<UNiagaraScriptSource>(NewEmitter, NAME_None, RF_Transactional);
@@ -164,7 +172,7 @@ UObject* UNiagaraEmitterFactoryNew::FactoryCreateNew(UClass* Class, UObject* InP
 		checkf(EmitterSpawnOutputNode != nullptr && EmitterUpdateOutputNode != nullptr && ParticleSpawnOutputNode != nullptr && ParticleUpdateOutputNode != nullptr,
 			TEXT("Failed to create output nodes for emitter scripts."));
 
-		if (bAddDefaultModulesAndRenderersToEmptyEmitter)
+		if (bAddDefaultModulesAndRenderers)
 		{
 			NewEmitter->AddRenderer(NewObject<UNiagaraSpriteRendererProperties>(NewEmitter, "Renderer"));
 
@@ -213,8 +221,6 @@ UObject* UNiagaraEmitterFactoryNew::FactoryCreateNew(UClass* Class, UObject* InP
 		NewEmitter->TemplateSpecification = ENiagaraScriptTemplateSpecification::None;
 		NewEmitter->SpawnScriptProps.Script->SetUsage(ENiagaraScriptUsage::ParticleSpawnScriptInterpolated);
 	}
-	
-	return NewEmitter;
 }
 
 #undef LOCTEXT_NAMESPACE
diff --git a/Engine/Plugins/FX/Niagara/Source/NiagaraEditor/Private/NiagaraScriptFactory.cpp b/Engine/Plugins/FX/Niagara/Source/NiagaraEditor/Private/NiagaraScriptFactory.cpp
--- a/Engine/Plugins/FX/Niagara/Source/NiagaraEditor/Private/NiagaraScriptFactory.cpp
+++ b/Engine/Plugins/FX/Niagara/Source/NiagaraEditor/Private/NiagaraScriptFactory.cpp
@@ -42,7 +42,10 @@ UObject* UNiagaraScriptFactoryNew::FactoryCreateNew(UClass* Class, UObject* InPa
 
 	UNiagaraScript* NewScript;
 	const FSoftObjectPath& SettingDefaultScript = GetDefaultScriptFromSettings(Settings);
-	if (UNiagaraScript* Default = Cast<UNiagaraScript>(SettingDefaultScript.TryLoad()))
+	UNiagaraScript* Default = Cast<UNiagaraScript>(SettingDefaultScript.TryLoad());
+	if (Default != nullptr &&
+		Cast<UNiagaraScriptSource>(Default->GetLatestSource()) != nullptr &&
+		Cast<UNiagaraScriptSource>(Default->GetLatestSource())->NodeGraph != nullptr)
 	{
 		NewScript = Cast<UNiagaraScript>(StaticDuplicateObject(Default, InParent, Name, Flags, Class));
 	}
diff --git a/Engine/Plugins/FX/Niagara/Source/NiagaraEditor/Private/NiagaraScriptSource.cpp b/Engine/Plugins/FX/Niagara/Source/NiagaraEditor/Private/NiagaraScriptSource.cpp
--- a/Engine/Plugins/FX/Niagara/Source/NiagaraEditor/Private/NiagaraScriptSource.cpp
+++ b/Engine/Plugins/FX/Niagara/Source/NiagaraEditor/Private/NiagaraScriptSource.cpp
@@ -302,7 +302,7 @@ bool UNiagaraScriptSource::AddModuleIfMissing(FString ModulePath, ENiagaraScript
 void UNiagaraScriptSource::FixupRenamedParameters(UNiagaraNode* Node, FNiagaraParameterStore& RapidIterationParameters, const TArray<FNiagaraVariable>& OldRapidIterationVariables, const UNiagaraEmitter* Emitter, ENiagaraScriptUsage ScriptUsage) const
 {
 	UNiagaraNodeFunctionCall* FunctionCallNode = Cast<UNiagaraNodeFunctionCall>(Node);
-	if (FunctionCallNode == nullptr || FunctionCallNode->FunctionScript == nullptr)
+	if (FunctionCallNode == nullptr || FunctionCallNode->HasValidScriptAndGraph() == false)
 	{
 		return;
 	}
diff --git a/Engine/Plugins/FX/Niagara/Source/NiagaraEditor/Private/SNewEmitterDialog.cpp b/Engine/Plugins/FX/Niagara/Source/NiagaraEditor/Private/SNewEmitterDialog.cpp
--- a/Engine/Plugins/FX/Niagara/Source/NiagaraEditor/Private/SNewEmitterDialog.cpp
+++ b/Engine/Plugins/FX/Niagara/Source/NiagaraEditor/Private/SNewEmitterDialog.cpp
@@ -58,7 +58,20 @@ void SNewEmitterDialog::Construct(const FArguments& InArgs)
 					bUseInheritance = false;
 				}),
 				CopyAssetPicker.ToSharedRef(), CopyAssetPicker->GetSearchBox()
-				)
+				),
+			SNiagaraNewAssetDialog::FNiagaraNewAssetDialogOption(
+				LOCTEXT("CreateEmptyLabel", "Create an empty emitter"),
+				LOCTEXT("CreateEmptyDescription", "Create an empty emitter with no modules or renderers"),
+				LOCTEXT("EmptyLabel", "Empty Emitter"),
+				SNiagaraNewAssetDialog::FOnGetSelectedAssetsFromPicker(),
+				SNiagaraNewAssetDialog::FOnSelectionConfirmed(),
+				SNew(SBox)
+				.HAlign(HAlign_Center)
+				.VAlign(VAlign_Center)
+				[
+					SNew(STextBlock)
+					.Text(LOCTEXT("NoOptionsLabel", "No Options"))
+				])
 		});
 }
 
diff --git a/Engine/Plugins/FX/Niagara/Source/NiagaraEditor/Private/ViewModels/NiagaraScratchPadViewModel.cpp b/Engine/Plugins/FX/Niagara/Source/NiagaraEditor/Private/ViewModels/NiagaraScratchPadViewModel.cpp
--- a/Engine/Plugins/FX/Niagara/Source/NiagaraEditor/Private/ViewModels/NiagaraScratchPadViewModel.cpp
+++ b/Engine/Plugins/FX/Niagara/Source/NiagaraEditor/Private/ViewModels/NiagaraScratchPadViewModel.cpp
@@ -19,6 +19,7 @@
 #include "NiagaraClipboard.h"
 #include "Toolkits/NiagaraSystemToolkit.h"
 #include "NiagaraEditorUtilities.h"
+#include "NiagaraScriptFactoryNew.h"
 
 #include "ScopedTransaction.h"
 #include "Modules/ModuleManager.h"
@@ -402,7 +403,9 @@ TSharedPtr<FNiagaraScratchPadScriptViewModel> UNiagaraScratchPadViewModel::Creat
 	case ENiagaraScriptUsage::DynamicInput:
 	{
 		UNiagaraScript* DefaultDynamicInput = Cast<UNiagaraScript>(GetDefault<UNiagaraEditorSettings>()->DefaultDynamicInputScript.TryLoad());
-		if (DefaultDynamicInput != nullptr)
+		if (DefaultDynamicInput != nullptr &&
+			Cast<UNiagaraScriptSource>(DefaultDynamicInput->GetLatestSource()) != nullptr &&
+			Cast<UNiagaraScriptSource>(DefaultDynamicInput->GetLatestSource())->NodeGraph != nullptr)
 		{
 			NewScript = CastChecked<UNiagaraScript>(StaticDuplicateObject(DefaultDynamicInput, ScriptOuter, GetUniqueScriptName(ScriptOuter, TEXT("ScratchDynamicInput"))));
 			TArray<UNiagaraNodeOutput*> OutputNodes;
@@ -429,7 +432,9 @@ TSharedPtr<FNiagaraScratchPadScriptViewModel> UNiagaraScratchPadViewModel::Creat
 	case ENiagaraScriptUsage::Module:
 	{
 		UNiagaraScript* DefaultModule = Cast<UNiagaraScript>(GetDefault<UNiagaraEditorSettings>()->DefaultModuleScript.TryLoad());
-		if (DefaultModule != nullptr)
+		if (DefaultModule != nullptr &&
+			Cast<UNiagaraScriptSource>(DefaultModule->GetLatestSource()) != nullptr &&
+			Cast<UNiagaraScriptSource>(DefaultModule->GetLatestSource())->NodeGraph != nullptr)
 		{
 			NewScript = CastChecked<UNiagaraScript>(StaticDuplicateObject(DefaultModule, ScriptOuter, GetUniqueScriptName(ScriptOuter, TEXT("ScratchModule"))));
 		}
@@ -437,16 +442,20 @@ TSharedPtr<FNiagaraScratchPadScriptViewModel> UNiagaraScratchPadViewModel::Creat
 	}
 	}
 
-	if (NewScript != nullptr)
+	if (NewScript == nullptr)
 	{
-		NewScript->ClearFlags(RF_Public | RF_Standalone);
-		ScriptOuter->Modify();
-		TargetScripts->Add(NewScript);
-		NewScript->GetLatestScriptData()->ModuleUsageBitmask |= (1 << (int32)InTargetSupportedUsage);
-		RefreshScriptViewModels();
-		UpdateChangeId(GetSystemViewModel());
+		NewScript = NewObject<UNiagaraScript>(ScriptOuter);
+		NewScript->Usage = InScriptUsage;
+		UNiagaraScriptFactoryNew::InitializeScript(NewScript);
 	}
 
+	NewScript->ClearFlags(RF_Public | RF_Standalone);
+	ScriptOuter->Modify();
+	TargetScripts->Add(NewScript);
+	NewScript->GetLatestScriptData()->ModuleUsageBitmask |= (1 << (int32)InTargetSupportedUsage);
+	RefreshScriptViewModels();
+	UpdateChangeId(GetSystemViewModel());
+
 	return GetViewModelForScript(NewScript);
 }
 
diff --git a/Engine/Plugins/FX/Niagara/Source/NiagaraEditor/Private/ViewModels/NiagaraSystemViewModel.cpp b/Engine/Plugins/FX/Niagara/Source/NiagaraEditor/Private/ViewModels/NiagaraSystemViewModel.cpp
--- a/Engine/Plugins/FX/Niagara/Source/NiagaraEditor/Private/ViewModels/NiagaraSystemViewModel.cpp
+++ b/Engine/Plugins/FX/Niagara/Source/NiagaraEditor/Private/ViewModels/NiagaraSystemViewModel.cpp
@@ -36,6 +36,7 @@
 #include "EdGraphUtilities.h"
 #include "ViewModels/NiagaraScratchPadUtilities.h"
 #include "NiagaraParameterDefinitions.h"
+#include "NiagaraEmitterFactoryNew.h"
 
 #include "Editor.h"
 
@@ -449,6 +450,18 @@ TSharedPtr<FNiagaraEmitterHandleViewModel> FNiagaraSystemViewModel::AddEmitter(U
 	return NewEmitterHandleViewModel;
 }
 
+NIAGARAEDITOR_API TSharedPtr<FNiagaraEmitterHandleViewModel> FNiagaraSystemViewModel::AddEmptyEmitter()
+{
+	UNiagaraEmitter* EmptyEmitter = NewObject<UNiagaraEmitter>(GetTransientPackage());
+	bool bAddDefaultModulesAndRenderers = false;
+	UNiagaraEmitterFactoryNew::InitializeEmitter(EmptyEmitter, bAddDefaultModulesAndRenderers);
+	EmptyEmitter->TemplateSpecification = ENiagaraScriptTemplateSpecification::Template;
+	FName EmptyEmitterName = FNiagaraEditorUtilities::GetUniqueObjectName<UNiagaraEmitter>(GetTransientPackage(), TEXT("Empty"));
+	EmptyEmitter->SetUniqueEmitterName(EmptyEmitterName.ToString());
+	EmptyEmitter->SetFlags(RF_Transactional);
+	return AddEmitter(*EmptyEmitter);
+}
+
 void FNiagaraSystemViewModel::DuplicateEmitters(TArray<FEmitterHandleToDuplicate> EmitterHandlesToDuplicate)
 {
 	if (EmitterHandlesToDuplicate.Num() <= 0)
diff --git a/Engine/Plugins/FX/Niagara/Source/NiagaraEditor/Public/NiagaraEmitterFactoryNew.h b/Engine/Plugins/FX/Niagara/Source/NiagaraEditor/Public/NiagaraEmitterFactoryNew.h
--- a/Engine/Plugins/FX/Niagara/Source/NiagaraEditor/Public/NiagaraEmitterFactoryNew.h
+++ b/Engine/Plugins/FX/Niagara/Source/NiagaraEditor/Public/NiagaraEmitterFactoryNew.h
@@ -21,6 +21,9 @@ class UNiagaraEmitterFactoryNew : public UFactory
 	virtual bool ConfigureProperties() override;
 	virtual UObject* FactoryCreateNew(UClass* Class, UObject* InParent, FName Name, EObjectFlags Flags, UObject* Context, FFeedbackContext* Warn) override;
 	//~ Begin UFactory Interface	
+
+public:
+	static void InitializeEmitter(UNiagaraEmitter* NewEmitter, bool bAddDefaultModulesAndRenderers);
 };
 
 
diff --git a/Engine/Plugins/FX/Niagara/Source/NiagaraEditor/Public/ViewModels/NiagaraSystemViewModel.h b/Engine/Plugins/FX/Niagara/Source/NiagaraEditor/Public/ViewModels/NiagaraSystemViewModel.h
--- a/Engine/Plugins/FX/Niagara/Source/NiagaraEditor/Public/ViewModels/NiagaraSystemViewModel.h
+++ b/Engine/Plugins/FX/Niagara/Source/NiagaraEditor/Public/ViewModels/NiagaraSystemViewModel.h
@@ -220,6 +220,9 @@ public:
 	/** Adds a new emitter to the System. */
 	NIAGARAEDITOR_API TSharedPtr<FNiagaraEmitterHandleViewModel> AddEmitter(UNiagaraEmitter& Emitter);
 
+	/** Adds an empty emitter to the system. */
+	NIAGARAEDITOR_API TSharedPtr<FNiagaraEmitterHandleViewModel> AddEmptyEmitter();
+
 	/** Deletes the emitters with the supplied ids from the system */
 	NIAGARAEDITOR_API void DeleteEmitters(TSet<FGuid> EmitterHandleIdsToDelete);
 
diff --git a/Engine/Plugins/FX/Niagara/Source/NiagaraEditorWidgets/Private/SNiagaraOverviewGraph.cpp b/Engine/Plugins/FX/Niagara/Source/NiagaraEditorWidgets/Private/SNiagaraOverviewGraph.cpp
--- a/Engine/Plugins/FX/Niagara/Source/NiagaraEditorWidgets/Private/SNiagaraOverviewGraph.cpp
+++ b/Engine/Plugins/FX/Niagara/Source/NiagaraEditorWidgets/Private/SNiagaraOverviewGraph.cpp
@@ -255,7 +255,11 @@ FActionMenuContent SNiagaraOverviewGraph::OnCreateGraphActionMenu(UEdGraph* InGr
 		MenuBuilder.BeginSection(TEXT("NiagaraOverview_EditGraph"), LOCTEXT("EditGraph", "Edit Graph"));
 		{
 			MenuBuilder.AddMenuEntry(FNiagaraEditorCommands::Get().OpenAddEmitterMenu);
-
+			MenuBuilder.AddMenuEntry(
+				LOCTEXT("EmptyEmitterLabel", "Add empty emitter"),
+				LOCTEXT("CommentsToolTip", "Adds an empty emitter without any modules or renderers."),
+				FSlateIcon(),
+				FExecuteAction::CreateSP(this, &SNiagaraOverviewGraph::OnCreateEmptyEmitter));
 			MenuBuilder.AddMenuEntry(
 				LOCTEXT("CommentsLabel", "Add Comment"),
 				LOCTEXT("CommentsToolTip", "Add a comment box"),
@@ -290,6 +294,11 @@ FActionMenuContent SNiagaraOverviewGraph::OnCreateGraphActionMenu(UEdGraph* InGr
 	return FActionMenuContent(SNullWidget::NullWidget, SNullWidget::NullWidget);
 }
 
+void SNiagaraOverviewGraph::OnCreateEmptyEmitter()
+{
+	ViewModel->GetSystemViewModel()->AddEmptyEmitter();
+}
+
 void SNiagaraOverviewGraph::OnCreateComment()
 {
 	FNiagaraSchemaAction_NewComment CommentAction = FNiagaraSchemaAction_NewComment(GraphEditor);
diff --git a/Engine/Plugins/FX/Niagara/Source/NiagaraEditorWidgets/Private/SNiagaraOverviewGraph.h b/Engine/Plugins/FX/Niagara/Source/NiagaraEditorWidgets/Private/SNiagaraOverviewGraph.h
--- a/Engine/Plugins/FX/Niagara/Source/NiagaraEditorWidgets/Private/SNiagaraOverviewGraph.h
+++ b/Engine/Plugins/FX/Niagara/Source/NiagaraEditorWidgets/Private/SNiagaraOverviewGraph.h
@@ -35,6 +35,7 @@ private:
 	/** Called to create context menu when right-clicking on graph */
 	FActionMenuContent OnCreateGraphActionMenu(UEdGraph* InGraph, const FVector2D& InNodePosition, const TArray<UEdGraphPin*>& InDraggedPins, bool bAutoExpand, SGraphEditor::FActionMenuClosed InOnMenuClosed);
 
+	void OnCreateEmptyEmitter();
 	void OnCreateComment();
 	void OnClearIsolated();
 
