- Make ParticleRead failed binding log contain the system / emitter

#rnx

diff --git a/Engine/Plugins/FX/Niagara/Source/Niagara/Private/NiagaraDataInterfaceParticleRead.cpp b/Engine/Plugins/FX/Niagara/Source/Niagara/Private/NiagaraDataInterfaceParticleRead.cpp
--- a/Engine/Plugins/FX/Niagara/Source/Niagara/Private/NiagaraDataInterfaceParticleRead.cpp
+++ b/Engine/Plugins/FX/Niagara/Source/Niagara/Private/NiagaraDataInterfaceParticleRead.cpp
@@ -1417,7 +1417,9 @@ void UNiagaraDataInterfaceParticleRead::GetVMExternalFunction(const FVMExternalF
 
 	if (!bBindSuccessful)
 	{
-		UE_LOG(LogNiagara, Warning, TEXT("Failed to bind VMExternalFunction '%s' with attribute '%s'! Check that the attribute is named correctly."), *BindingInfo.Name.ToString(), *AttributeToRead.ToString());
+		UNiagaraSystem* NiagaraSystem = PIData->SystemInstance ? PIData->SystemInstance->GetSystem() : nullptr;
+		UNiagaraEmitter* NiagaraEmitter = PIData->EmitterInstance ? PIData->EmitterInstance->GetCachedEmitter() : nullptr;
+		UE_LOG(LogNiagara, Warning, TEXT("ParticleRead: Failed to '%s' attribute '%s' for System '%s' Emitter '%s'! Check that the attribute is named correctly."), *BindingInfo.Name.ToString(), *AttributeToRead.ToString(), *GetNameSafe(NiagaraSystem), *GetNameSafe(NiagaraEmitter));
 	}
 }
 
