- Ensure we always have valid meta data

#rnx

diff --git a/Engine/Plugins/FX/Niagara/Source/NiagaraShader/Public/NiagaraShared.h b/Engine/Plugins/FX/Niagara/Source/NiagaraShader/Public/NiagaraShared.h
--- a/Engine/Plugins/FX/Niagara/Source/NiagaraShader/Public/NiagaraShared.h
+++ b/Engine/Plugins/FX/Niagara/Source/NiagaraShader/Public/NiagaraShared.h
@@ -670,12 +670,14 @@ public:
 		: BaseVMScript(nullptr)
 		, GameThreadShaderMap(nullptr)
 		, RenderingThreadShaderMap(nullptr)
+		, ScriptParametersMetadata(MakeShared<FNiagaraShaderScriptParametersMetadata>())
 		, FeatureLevel(GMaxRHIFeatureLevel)
 		, ShaderPlatform(SP_NumPlatforms)
 		, bLoadedCookedShaderMapId(false)
 		, bLoadedFromCookedMaterial(false)
 		, bQueuedForRelease(false)
-	{}
+	{
+	}
 
 	/**
 	 * Destructor
