- Convert Grid2D data interface to use template & parameter structure binding
- Convert Grid3D data interface to use template & parameter structure binding
- Convert all RW data interface to use correct formatting for shader parameters

diff --git a/Engine/Plugins/FX/Niagara/Source/Niagara/Classes/NiagaraDataInterfaceGrid2DCollection.h b/Engine/Plugins/FX/Niagara/Source/Niagara/Classes/NiagaraDataInterfaceGrid2DCollection.h
--- a/Engine/Plugins/FX/Niagara/Source/Niagara/Classes/NiagaraDataInterfaceGrid2DCollection.h
+++ b/Engine/Plugins/FX/Niagara/Source/Niagara/Classes/NiagaraDataInterfaceGrid2DCollection.h
@@ -129,23 +129,6 @@ struct FNiagaraDataInterfaceParametersCS_Grid2DCollection : public FNiagaraDataI
 {
 	DECLARE_TYPE_LAYOUT(FNiagaraDataInterfaceParametersCS_Grid2DCollection, NonVirtual);
 
-public:
-	void Bind(const FNiagaraDataInterfaceGPUParamInfo& ParameterInfo, const class FShaderParameterMap& ParameterMap);
-	void Set(FRHICommandList& RHICmdList, const FNiagaraDataInterfaceSetArgs& Context) const;
-	void Unset(FRHICommandList& RHICmdList, const FNiagaraDataInterfaceSetArgs& Context) const;
-
-private:
-	LAYOUT_FIELD(FShaderParameter, NumAttributesParam);
-	LAYOUT_FIELD(FShaderParameter, UnitToUVParam);
-	LAYOUT_FIELD(FShaderParameter, NumCellsParam);
-	LAYOUT_FIELD(FShaderParameter, CellSizeParam);
-	LAYOUT_FIELD(FShaderParameter, WorldBBoxSizeParam);
-
-	LAYOUT_FIELD(FShaderResourceParameter, GridParam);
-	LAYOUT_FIELD(FRWShaderParameter, OutputGridParam);
-	LAYOUT_FIELD(FShaderParameter, AttributeIndicesParam);
-
-	LAYOUT_FIELD(FShaderResourceParameter, SamplerParam);
 	LAYOUT_FIELD(TMemoryImageArray<FMemoryImageName>, AttributeNames);
 	LAYOUT_FIELD(TMemoryImageArray<uint32>, AttributeChannelCount);
 };
@@ -156,8 +139,6 @@ class NIAGARA_API UNiagaraDataInterfaceGrid2DCollection : public UNiagaraDataInt
 	GENERATED_UCLASS_BODY()
 
 public:
-	DECLARE_NIAGARA_DI_PARAMETER();
-
 	/** Reference to a user parameter if we're reading one. */
 	UPROPERTY(EditAnywhere, Category = "Grid")
 	FNiagaraUserParameterBinding RenderTargetUserParameter;
@@ -201,6 +182,11 @@ public:
 	virtual bool GetFunctionHLSL(const FNiagaraDataInterfaceGPUParamInfo& ParamInfo, const FNiagaraDataInterfaceGeneratedFunction& FunctionInfo, int FunctionInstanceIndex, FString& OutHLSL) override;
 	virtual bool AppendCompileHash(FNiagaraCompileHashVisitor* InVisitor) const override;
 #endif
+	virtual bool UseLegacyShaderBindings() const  override { return false; }
+	virtual void BuildShaderParameters(FNiagaraShaderParametersBuilder& ShaderParametersBuilder) const override;
+	virtual void SetShaderParameters(const FNiagaraDataInterfaceSetShaderParametersContext& Context) const override;
+	virtual FNiagaraDataInterfaceParametersCS* CreateShaderStorage(const FNiagaraDataInterfaceGPUParamInfo& ParameterInfo, const FShaderParameterMap& ParameterMap) const override;
+	virtual const FTypeLayoutDesc* GetShaderStorageType() const override;
 
 	virtual void ProvidePerInstanceDataForRenderThread(void* DataForRenderThread, void* PerInstanceData, const FNiagaraSystemInstanceID& SystemInstance) override {}
 	virtual bool InitPerInstanceData(void* PerInstanceData, FNiagaraSystemInstance* SystemInstance) override;
@@ -242,11 +228,11 @@ public:
 	UFUNCTION(BlueprintCallable, Category = Niagara)
 	virtual void GetTextureSize(const UNiagaraComponent *Component, int &SizeX, int &SizeY);
 
-	void GetWorldBBoxSize(FVectorVMExternalFunctionContext& Context);
-	void GetCellSize(FVectorVMExternalFunctionContext& Context);
-	void GetNumCells(FVectorVMExternalFunctionContext& Context);
-	void SetNumCells(FVectorVMExternalFunctionContext& Context);
-	void GetAttributeIndex(FVectorVMExternalFunctionContext& Context, const FName& InName, int32 NumChannels);
+	void VMGetWorldBBoxSize(FVectorVMExternalFunctionContext& Context);
+	void VMGetCellSize(FVectorVMExternalFunctionContext& Context);
+	void VMGetNumCells(FVectorVMExternalFunctionContext& Context);
+	void VMSetNumCells(FVectorVMExternalFunctionContext& Context);
+	void VMGetAttributeIndex(FVectorVMExternalFunctionContext& Context, const FName& InName, int32 NumChannels);
 
 	static const FString GridName;
 	static const FString OutputGridName;
diff --git a/Engine/Plugins/FX/Niagara/Source/Niagara/Classes/NiagaraDataInterfaceGrid2DCollectionReader.h b/Engine/Plugins/FX/Niagara/Source/Niagara/Classes/NiagaraDataInterfaceGrid2DCollectionReader.h
--- a/Engine/Plugins/FX/Niagara/Source/Niagara/Classes/NiagaraDataInterfaceGrid2DCollectionReader.h
+++ b/Engine/Plugins/FX/Niagara/Source/Niagara/Classes/NiagaraDataInterfaceGrid2DCollectionReader.h
@@ -20,15 +20,14 @@ class NIAGARA_API UNiagaraDataInterfaceGrid2DCollectionReader : public UNiagaraD
 	GENERATED_UCLASS_BODY()
 
 public:
-	DECLARE_NIAGARA_DI_PARAMETER();
 
 	// Name of the emitter to read from
 	UPROPERTY(EditAnywhere, Category = "Reader")
-		FString EmitterName;
+	FString EmitterName;
 
 	// Name of the Grid2DCollection Data Interface on the emitter
 	UPROPERTY(EditAnywhere, Category = "Reader")
-		FString DIName;
+	FString DIName;
 
 	virtual bool Equals(const UNiagaraDataInterface* Other) const override;
 	virtual bool InitPerInstanceData(void* PerInstanceData, FNiagaraSystemInstance* SystemInstance) override;
diff --git a/Engine/Plugins/FX/Niagara/Source/Niagara/Classes/NiagaraDataInterfaceGrid3DCollection.h b/Engine/Plugins/FX/Niagara/Source/Niagara/Classes/NiagaraDataInterfaceGrid3DCollection.h
--- a/Engine/Plugins/FX/Niagara/Source/Niagara/Classes/NiagaraDataInterfaceGrid3DCollection.h
+++ b/Engine/Plugins/FX/Niagara/Source/Niagara/Classes/NiagaraDataInterfaceGrid3DCollection.h
@@ -121,29 +121,7 @@ struct FNiagaraDataInterfaceProxyGrid3DCollectionProxy : public FNiagaraDataInte
 struct FNiagaraDataInterfaceParametersCS_Grid3DCollection : public FNiagaraDataInterfaceParametersCS
 {
 	DECLARE_TYPE_LAYOUT(FNiagaraDataInterfaceParametersCS_Grid3DCollection, NonVirtual);
-public:
-	void Bind(const FNiagaraDataInterfaceGPUParamInfo& ParameterInfo, const class FShaderParameterMap& ParameterMap);
-	void Set(FRHICommandList& RHICmdList, const FNiagaraDataInterfaceSetArgs& Context) const;
-	void Unset(FRHICommandList& RHICmdList, const FNiagaraDataInterfaceSetArgs& Context) const;
 
-private:
-	LAYOUT_FIELD(FShaderParameter, NumAttributesParam);
-	LAYOUT_FIELD(FShaderParameter, NumNamedAttributesParam);
-	LAYOUT_FIELD(FShaderParameter, UnitToUVParam);
-	LAYOUT_FIELD(FShaderParameter, NumCellsParam);
-	LAYOUT_FIELD(FShaderParameter, NumTilesParam);
-	LAYOUT_FIELD(FShaderParameter, OneOverNumTilesParam);
-	LAYOUT_FIELD(FShaderParameter, UnitClampMinParam);
-	LAYOUT_FIELD(FShaderParameter, UnitClampMaxParam);
-	LAYOUT_FIELD(FShaderParameter, CellSizeParam);
-	LAYOUT_FIELD(FShaderParameter, WorldBBoxSizeParam);
-
-	LAYOUT_FIELD(FShaderResourceParameter, GridParam);
-	LAYOUT_FIELD(FRWShaderParameter, OutputGridParam);
-	LAYOUT_FIELD(FShaderParameter, AttributeIndicesParam);
-	LAYOUT_FIELD(FShaderResourceParameter, PerAttributeDataParam);
-
-	LAYOUT_FIELD(FShaderResourceParameter, SamplerParam);
 	LAYOUT_FIELD(TMemoryImageArray<FMemoryImageName>, AttributeNames);
 	LAYOUT_FIELD(TMemoryImageArray<uint32>, AttributeChannelCount);
 };
@@ -154,8 +132,6 @@ class NIAGARA_API UNiagaraDataInterfaceGrid3DCollection : public UNiagaraDataInt
 	GENERATED_UCLASS_BODY()
 
 public:
-	DECLARE_NIAGARA_DI_PARAMETER();
-
 	// Number of attributes stored on the grid
 	UPROPERTY(EditAnywhere, Category = "Grid")
 	int32 NumAttributes;
@@ -203,6 +179,12 @@ public:
 	virtual void GetParameterDefinitionHLSL(const FNiagaraDataInterfaceGPUParamInfo& ParamInfo, FString& OutHLSL) override;
 	virtual bool GetFunctionHLSL(const FNiagaraDataInterfaceGPUParamInfo& ParamInfo, const FNiagaraDataInterfaceGeneratedFunction& FunctionInfo, int FunctionInstanceIndex, FString& OutHLSL) override;
 #endif
+	virtual bool UseLegacyShaderBindings() const  override { return false; }
+	virtual void BuildShaderParameters(FNiagaraShaderParametersBuilder& ShaderParametersBuilder) const override;
+	virtual void SetShaderParameters(const FNiagaraDataInterfaceSetShaderParametersContext& Context) const override;
+	virtual FNiagaraDataInterfaceParametersCS* CreateShaderStorage(const FNiagaraDataInterfaceGPUParamInfo& ParameterInfo, const FShaderParameterMap& ParameterMap) const override;
+	virtual const FTypeLayoutDesc* GetShaderStorageType() const override;
+
 #if WITH_EDITOR
 	virtual ENiagaraGpuDispatchType GetGpuDispatchType() const override { return ENiagaraGpuDispatchType::ThreeD; }
 #endif
@@ -244,14 +226,14 @@ public:
 	UFUNCTION(BlueprintCallable, Category = Niagara)
 	virtual void GetTextureSize(const UNiagaraComponent *Component, int &SizeX, int &SizeY, int &SizeZ);	
 
-	void GetWorldBBoxSize(FVectorVMExternalFunctionContext& Context);
-	void GetCellSize(FVectorVMExternalFunctionContext& Context);
+	void VMGetWorldBBoxSize(FVectorVMExternalFunctionContext& Context);
+	void VMGetCellSize(FVectorVMExternalFunctionContext& Context);
 
-	void GetNumCells(FVectorVMExternalFunctionContext& Context);
-	void SetNumCells(FVectorVMExternalFunctionContext& Context);
-	void UnitToFloatIndex(FVectorVMExternalFunctionContext& Context);
+	void VMGetNumCells(FVectorVMExternalFunctionContext& Context);
+	void VMSetNumCells(FVectorVMExternalFunctionContext& Context);
+	void VMUnitToFloatIndex(FVectorVMExternalFunctionContext& Context);
 
-	void GetAttributeIndex(FVectorVMExternalFunctionContext& Context, const FName& InName, int32 NumChannels);
+	void VMGetAttributeIndex(FVectorVMExternalFunctionContext& Context, const FName& InName, int32 NumChannels);
 
 	static const FString NumTilesName;
 	static const FString OneOverNumTilesName;
diff --git a/Engine/Plugins/FX/Niagara/Source/Niagara/Classes/NiagaraDataInterfaceGrid3DCollectionReader.h b/Engine/Plugins/FX/Niagara/Source/Niagara/Classes/NiagaraDataInterfaceGrid3DCollectionReader.h
--- a/Engine/Plugins/FX/Niagara/Source/Niagara/Classes/NiagaraDataInterfaceGrid3DCollectionReader.h
+++ b/Engine/Plugins/FX/Niagara/Source/Niagara/Classes/NiagaraDataInterfaceGrid3DCollectionReader.h
@@ -20,8 +20,6 @@ class NIAGARA_API UNiagaraDataInterfaceGrid3DCollectionReader : public UNiagaraD
 	GENERATED_UCLASS_BODY()
 
 public:
-	DECLARE_NIAGARA_DI_PARAMETER();
-
 	// Name of the emitter to read from
 	UPROPERTY(EditAnywhere, Category = "Reader")
 	FString EmitterName;
diff --git a/Engine/Plugins/FX/Niagara/Source/Niagara/Private/NiagaraDataInterfaceGrid2DCollection.cpp b/Engine/Plugins/FX/Niagara/Source/Niagara/Private/NiagaraDataInterfaceGrid2DCollection.cpp
--- a/Engine/Plugins/FX/Niagara/Source/Niagara/Private/NiagaraDataInterfaceGrid2DCollection.cpp
+++ b/Engine/Plugins/FX/Niagara/Source/Niagara/Private/NiagaraDataInterfaceGrid2DCollection.cpp
@@ -1,27 +1,44 @@
 // Copyright Epic Games, Inc. All Rights Reserved.
 #include "NiagaraDataInterfaceGrid2DCollection.h"
-#include "NiagaraShader.h"
-#include "ShaderParameterUtils.h"
-#include "ClearQuad.h"
-#include "Engine/Canvas.h"
-#include "CanvasItem.h"
-#include "TextureResource.h"
-#include "Engine/Texture2DArray.h"
+#include "NiagaraBatchedElements.h"
+#include "NiagaraConstants.h"
+#include "NiagaraGpuComputeDebug.h"
 #include "NiagaraGpuComputeDispatchInterface.h"
 #include "NiagaraRenderer.h"
-#include "NiagaraSystemInstance.h"
 #include "NiagaraSettings.h"
-#include "NiagaraBatchedElements.h"
-#include "NiagaraGpuComputeDebug.h"
+#include "NiagaraShader.h"
+#include "NiagaraShaderParametersBuilder.h"
+#include "NiagaraSystemInstance.h"
+
+#include "Engine/Canvas.h"
+#include "Engine/Texture2DArray.h"
 #include "Engine/TextureRenderTarget2D.h"
 #include "Engine/TextureRenderTarget2DArray.h"
-#include "NiagaraConstants.h"
+#include "ClearQuad.h"
+#include "CanvasItem.h"
+#include "ShaderParameterMetadataBuilder.h"
+#include "ShaderParameterUtils.h"
+#include "TextureResource.h"
 
 #define LOCTEXT_NAMESPACE "NiagaraDataInterfaceGrid2DCollection"
 
-const FString UNiagaraDataInterfaceGrid2DCollection::GridName(TEXT("Grid_"));
-const FString UNiagaraDataInterfaceGrid2DCollection::OutputGridName(TEXT("OutputGrid_"));
-const FString UNiagaraDataInterfaceGrid2DCollection::SamplerName(TEXT("Sampler_"));
+IMPLEMENT_TYPE_LAYOUT(FNiagaraDataInterfaceParametersCS_Grid2DCollection);
+
+BEGIN_SHADER_PARAMETER_STRUCT(FNDIGrid2DShaderParameters, )
+	SHADER_PARAMETER(int,							NumAttributes)
+	SHADER_PARAMETER(FVector2f,						UnitToUV)
+	SHADER_PARAMETER(FIntPoint,						NumCells)
+	SHADER_PARAMETER(FVector2f,						CellSize)
+	SHADER_PARAMETER(FVector2f,						WorldBBoxSize)
+
+	SHADER_PARAMETER_SRV(Texture2DArray<float>,		Grid)
+	SHADER_PARAMETER_SAMPLER(SamplerState,			SamplerState)
+	SHADER_PARAMETER_UAV(RWTexture2DArray<float>,	OutputGrid)
+END_SHADER_PARAMETER_STRUCT()
+
+const FString UNiagaraDataInterfaceGrid2DCollection::GridName(TEXT("_Grid"));
+const FString UNiagaraDataInterfaceGrid2DCollection::OutputGridName(TEXT("_OutputGrid"));
+const FString UNiagaraDataInterfaceGrid2DCollection::SamplerName(TEXT("_Sampler"));
 
 const FName UNiagaraDataInterfaceGrid2DCollection::SetNumCellsFunctionName("SetNumCells");
 
@@ -66,7 +83,7 @@ const FString UNiagaraDataInterfaceGrid2DCollection::AnonymousAttributeString("A
 const FName UNiagaraDataInterfaceGrid2DCollection::ClearCellFunctionName("ClearCell");
 const FName UNiagaraDataInterfaceGrid2DCollection::CopyPreviousToCurrentForCellFunctionName("CopyPreviousToCurrentForCell");
 
-const FString UNiagaraDataInterfaceGrid2DCollection::AttributeIndicesBaseName(TEXT("AttributeIndices_"));
+const FString UNiagaraDataInterfaceGrid2DCollection::AttributeIndicesBaseName(TEXT("AttributeIndices"));
 const TCHAR* UNiagaraDataInterfaceGrid2DCollection::VectorComponentNames[] = { TEXT(".x"), TEXT(".y"), TEXT(".z"), TEXT(".w") };
 
 const FName UNiagaraDataInterfaceGrid2DCollection::SampleGridFunctionName("SampleGrid");
@@ -202,141 +219,6 @@ struct FNiagaraGrid2DLegacyTiled2DInfo
 	FIntPoint Size = FIntPoint::ZeroValue;
 };
 
-/*--------------------------------------------------------------------------------------------------------------------------*/
-
-void FNiagaraDataInterfaceParametersCS_Grid2DCollection::Bind(const FNiagaraDataInterfaceGPUParamInfo& ParameterInfo, const class FShaderParameterMap& ParameterMap)
-{
-	NumAttributesParam.Bind(ParameterMap, *(UNiagaraDataInterfaceRWBase::NumAttributesName + ParameterInfo.DataInterfaceHLSLSymbol));
-	NumCellsParam.Bind(ParameterMap, *(UNiagaraDataInterfaceRWBase::NumCellsName + ParameterInfo.DataInterfaceHLSLSymbol));
-	UnitToUVParam.Bind(ParameterMap, *(UNiagaraDataInterfaceRWBase::UnitToUVName + ParameterInfo.DataInterfaceHLSLSymbol));
-	CellSizeParam.Bind(ParameterMap, *(UNiagaraDataInterfaceRWBase::CellSizeName + ParameterInfo.DataInterfaceHLSLSymbol));
-	WorldBBoxSizeParam.Bind(ParameterMap, *(UNiagaraDataInterfaceRWBase::WorldBBoxSizeName + ParameterInfo.DataInterfaceHLSLSymbol));
-
-	GridParam.Bind(ParameterMap, *(UNiagaraDataInterfaceGrid2DCollection::GridName + ParameterInfo.DataInterfaceHLSLSymbol));
-	OutputGridParam.Bind(ParameterMap, *(UNiagaraDataInterfaceGrid2DCollection::OutputGridName + ParameterInfo.DataInterfaceHLSLSymbol));
-
-	SamplerParam.Bind(ParameterMap, *(UNiagaraDataInterfaceGrid2DCollection::SamplerName + ParameterInfo.DataInterfaceHLSLSymbol));
-	AttributeIndicesParam.Bind(ParameterMap, *(UNiagaraDataInterfaceGrid2DCollection::AttributeIndicesBaseName + ParameterInfo.DataInterfaceHLSLSymbol));
-
-	// Gather up all the attribute names referenced. Note that there may be multiple in the list of the same name,
-	// but we only deal with this by the number of bound methods.
-	{
-		int32 NumFuncs = ParameterInfo.GeneratedFunctions.Num();
-
-		for (int32 FuncIdx = 0; FuncIdx < NumFuncs; ++FuncIdx)
-		{
-			const FNiagaraDataInterfaceGeneratedFunction& Func = ParameterInfo.GeneratedFunctions[FuncIdx];
-			static const FName NAME_Attribute("Attribute");
-			const FName* AttributeName = Func.FindSpecifierValue(NAME_Attribute);
-			if (AttributeName != nullptr)
-			{
-				int32 ComponentCount = UNiagaraDataInterfaceGrid2DCollection::GetComponentCountFromFuncName(Func.DefinitionName);
-				AttributeNames.Add(*AttributeName);
-				AttributeChannelCount.Add(ComponentCount);
-			}
-			else
-			{
-				AttributeNames.Add(FName());
-				AttributeChannelCount.Add(INDEX_NONE);
-			}
-		}
-	}
-}
-
-void FNiagaraDataInterfaceParametersCS_Grid2DCollection::Set(FRHICommandList& RHICmdList, const FNiagaraDataInterfaceSetArgs& Context) const
-{
-	check(IsInRenderingThread());
-
-	// Get shader and DI
-	FRHIComputeShader* ComputeShaderRHI = Context.Shader.GetComputeShader();
-	FNiagaraDataInterfaceProxyGrid2DCollectionProxy* VFDI = static_cast<FNiagaraDataInterfaceProxyGrid2DCollectionProxy*>(Context.DataInterface);
-
-	FGrid2DCollectionRWInstanceData_RenderThread* ProxyData = VFDI->SystemInstancesToProxyData_RT.Find(Context.SystemInstanceID);
-	FGrid2DCollectionRWInstanceData_RenderThread* OriginalProxyData = ProxyData;
-	check(ProxyData);
-
-	if (ProxyData->OtherProxy != nullptr)
-	{
-		FNiagaraDataInterfaceProxyGrid2DCollectionProxy* OtherGrid3DProxy = static_cast<FNiagaraDataInterfaceProxyGrid2DCollectionProxy*>(ProxyData->OtherProxy);
-		ProxyData = OtherGrid3DProxy->SystemInstancesToProxyData_RT.Find(Context.SystemInstanceID);
-		check(ProxyData);
-	}
-
-	if (OriginalProxyData->AttributeIndices.Num() == 0 && AttributeNames.Num() > 0)
-	{
-		int NumAttrIndices = Align(AttributeNames.Num(), 4);
-		OriginalProxyData->AttributeIndices.SetNumZeroed(NumAttrIndices);
-
-		// TODO handle mismatched types!
-		for (int32 i = 0; i < AttributeNames.Num(); i++)
-		{
-			int32 FoundIdx = ProxyData->Vars.Find(AttributeNames[i]);
-			check(AttributeNames.Num() == AttributeChannelCount.Num());
-			check(ProxyData->Offsets.Num() == ProxyData->VarComponents.Num());
-			check(ProxyData->Offsets.Num() == ProxyData->Vars.Num());
-			if (ProxyData->Offsets.IsValidIndex(FoundIdx) && AttributeChannelCount[i] == ProxyData->VarComponents[FoundIdx])
-			{
-				OriginalProxyData->AttributeIndices[i] = ProxyData->Offsets[FoundIdx];
-			}
-			else
-			{
-				OriginalProxyData->AttributeIndices[i] = -1; // We may need to protect against this in the hlsl as this might underflow an array lookup if used incorrectly.
-			}
-		}
-	}
-
-	SetShaderValue(RHICmdList, ComputeShaderRHI, NumAttributesParam, ProxyData->NumAttributes);
-	SetShaderValue(RHICmdList, ComputeShaderRHI, UnitToUVParam, FVector2f(1.0f) / FVector2f(ProxyData->NumCells));
-	SetShaderValue(RHICmdList, ComputeShaderRHI, NumCellsParam, ProxyData->NumCells);
-	SetShaderValue(RHICmdList, ComputeShaderRHI, CellSizeParam, FVector2f(ProxyData->CellSize));
-	SetShaderValue(RHICmdList, ComputeShaderRHI, WorldBBoxSizeParam, FVector2f(ProxyData->WorldBBoxSize));	// LWC_TODO: Precision loss?
-
-	SetShaderValueArray(RHICmdList, ComputeShaderRHI, AttributeIndicesParam, OriginalProxyData->AttributeIndices.GetData(), OriginalProxyData->AttributeIndices.Num());
-	FRHISamplerState *SamplerState = TStaticSamplerState<SF_Bilinear, AM_Clamp, AM_Clamp, AM_Clamp>::GetRHI();
-	SetSamplerParameter(RHICmdList, ComputeShaderRHI, SamplerParam, SamplerState);
-
-	if (GridParam.IsBound())
-	{
-		FRHIShaderResourceView* InputGridBuffer;
-		if (ProxyData->CurrentData != nullptr)
-		{
-			InputGridBuffer = ProxyData->CurrentData->GridSRV;
-		}
-		else
-		{
-			InputGridBuffer = FNiagaraRenderer::GetDummyTextureReadBuffer2DArray();
-		}
-		SetSRVParameter(RHICmdList, Context.Shader.GetComputeShader(), GridParam, InputGridBuffer);
-	}
-
-	if ( OutputGridParam.IsUAVBound() )
-	{
-		FRHIUnorderedAccessView* OutputGridUAV;
-		if (Context.IsOutputStage && ProxyData->DestinationData != nullptr)
-		{
-			OutputGridUAV = ProxyData->DestinationData->GridUAV;
-		}
-		else
-		{
-			OutputGridUAV = Context.ComputeDispatchInterface->GetEmptyUAVFromPool(RHICmdList, PF_R32_FLOAT, ENiagaraEmptyUAVType::Texture2DArray);
-		}
-		RHICmdList.SetUAVParameter(ComputeShaderRHI, OutputGridParam.GetUAVIndex(), OutputGridUAV);
-	}
-}
-
-void FNiagaraDataInterfaceParametersCS_Grid2DCollection::Unset(FRHICommandList& RHICmdList, const FNiagaraDataInterfaceSetArgs& Context) const
-{
-	if (OutputGridParam.IsBound())
-	{
-		OutputGridParam.UnsetUAV(RHICmdList, Context.Shader.GetComputeShader());
-	}
-}
-
-IMPLEMENT_TYPE_LAYOUT(FNiagaraDataInterfaceParametersCS_Grid2DCollection);
-
-IMPLEMENT_NIAGARA_DI_PARAMETER(UNiagaraDataInterfaceGrid2DCollection, FNiagaraDataInterfaceParametersCS_Grid2DCollection);
-
-
 UNiagaraDataInterfaceGrid2DCollection::UNiagaraDataInterfaceGrid2DCollection(FObjectInitializer const& ObjectInitializer)
 	: Super(ObjectInitializer)
 {
@@ -368,6 +250,7 @@ bool UNiagaraDataInterfaceGrid2DCollection::AppendCompileHash(FNiagaraCompileHas
 
 
 	InVisitor->UpdatePOD(TEXT("UNiagaraDataInterfaceGrid2DCollectionVersion"), (int32)FNiagaraGridCollection2DDIFunctionVersion::LatestVersion);
+	InVisitor->UpdateShaderParameters<FNDIGrid2DShaderParameters>();
 	return true;
 }
 #endif
@@ -994,11 +877,11 @@ void UNiagaraDataInterfaceGrid2DCollection::GetFunctions(TArray<FNiagaraFunction
 // #todo(dmp): expose more CPU functionality
 // #todo(dmp): ideally these would be exposed on the parent class, but we can't bind functions of parent classes but need to work on the interface
 // for sharing an instance data object with the super class
-DEFINE_NDI_DIRECT_FUNC_BINDER(UNiagaraDataInterfaceGrid2DCollection, GetWorldBBoxSize);
-DEFINE_NDI_DIRECT_FUNC_BINDER(UNiagaraDataInterfaceGrid2DCollection, GetCellSize);
-DEFINE_NDI_DIRECT_FUNC_BINDER(UNiagaraDataInterfaceGrid2DCollection, GetNumCells);
-DEFINE_NDI_DIRECT_FUNC_BINDER(UNiagaraDataInterfaceGrid2DCollection, SetNumCells);
-DEFINE_NDI_DIRECT_FUNC_BINDER_WITH_PAYLOAD(UNiagaraDataInterfaceGrid2DCollection, GetAttributeIndex);
+DEFINE_NDI_DIRECT_FUNC_BINDER(UNiagaraDataInterfaceGrid2DCollection, VMGetWorldBBoxSize);
+DEFINE_NDI_DIRECT_FUNC_BINDER(UNiagaraDataInterfaceGrid2DCollection, VMGetCellSize);
+DEFINE_NDI_DIRECT_FUNC_BINDER(UNiagaraDataInterfaceGrid2DCollection, VMGetNumCells);
+DEFINE_NDI_DIRECT_FUNC_BINDER(UNiagaraDataInterfaceGrid2DCollection, VMSetNumCells);
+DEFINE_NDI_DIRECT_FUNC_BINDER_WITH_PAYLOAD(UNiagaraDataInterfaceGrid2DCollection, VMGetAttributeIndex);
 void UNiagaraDataInterfaceGrid2DCollection::GetVMExternalFunction(const FVMExternalFunctionBindingInfo& BindingInfo, void* InstanceData, FVMExternalFunction &OutFunc)
 {
 	Super::GetVMExternalFunction(BindingInfo, InstanceData, OutFunc);
@@ -1008,51 +891,51 @@ void UNiagaraDataInterfaceGrid2DCollection::GetVMExternalFunction(const FVMExter
 	if (BindingInfo.Name == UNiagaraDataInterfaceRWBase::WorldBBoxSizeFunctionName)
 	{
 		check(BindingInfo.GetNumInputs() == 1 && BindingInfo.GetNumOutputs() == 2);
-		NDI_FUNC_BINDER(UNiagaraDataInterfaceGrid2DCollection, GetWorldBBoxSize)::Bind(this, OutFunc);
+		NDI_FUNC_BINDER(UNiagaraDataInterfaceGrid2DCollection, VMGetWorldBBoxSize)::Bind(this, OutFunc);
 	}
 	else if (BindingInfo.Name == UNiagaraDataInterfaceRWBase::CellSizeFunctionName)
 	{
 		// #todo(dmp): this will override the base class definition for GetCellSize because the data interface instance data computes cell size
 		// it would be nice to refactor this so it can be part of the super class
 		check(BindingInfo.GetNumInputs() == 1 && BindingInfo.GetNumOutputs() == 2);
-		NDI_FUNC_BINDER(UNiagaraDataInterfaceGrid2DCollection, GetCellSize)::Bind(this, OutFunc);
+		NDI_FUNC_BINDER(UNiagaraDataInterfaceGrid2DCollection, VMGetCellSize)::Bind(this, OutFunc);
 	}
 	else if (BindingInfo.Name == UNiagaraDataInterfaceRWBase::NumCellsFunctionName)
 	{
 		// #todo(dmp): this will override the base class definition for GetCellSize because the data interface instance data computes cell size
 		// it would be nice to refactor this so it can be part of the super class
 		check(BindingInfo.GetNumInputs() == 1 && BindingInfo.GetNumOutputs() == 2);
-		NDI_FUNC_BINDER(UNiagaraDataInterfaceGrid2DCollection, GetNumCells)::Bind(this, OutFunc);
+		NDI_FUNC_BINDER(UNiagaraDataInterfaceGrid2DCollection, VMGetNumCells)::Bind(this, OutFunc);
 	}
 	else if (BindingInfo.Name == SetNumCellsFunctionName)
 	{
 		check(BindingInfo.GetNumInputs() == 3 && BindingInfo.GetNumOutputs() == 1);
-		NDI_FUNC_BINDER(UNiagaraDataInterfaceGrid2DCollection, SetNumCells)::Bind(this, OutFunc);
+		NDI_FUNC_BINDER(UNiagaraDataInterfaceGrid2DCollection, VMSetNumCells)::Bind(this, OutFunc);
 	}
 
 	else if (BindingInfo.Name == GetVector4AttributeIndexFunctionName)
 	{
 		FName AttributeName = BindingInfo.FindSpecifier(NAME_Attribute)->Value;
 		check(BindingInfo.GetNumInputs() == 1 && BindingInfo.GetNumOutputs() == 1);
-		NDI_FUNC_BINDER(UNiagaraDataInterfaceGrid2DCollection, GetAttributeIndex)::Bind(this, OutFunc, AttributeName, 4);
+		NDI_FUNC_BINDER(UNiagaraDataInterfaceGrid2DCollection, VMGetAttributeIndex)::Bind(this, OutFunc, AttributeName, 4);
 	}
 	else if (BindingInfo.Name == GetVectorAttributeIndexFunctionName)
 	{
 		FName AttributeName = BindingInfo.FindSpecifier(NAME_Attribute)->Value;
 		check(BindingInfo.GetNumInputs() == 1 && BindingInfo.GetNumOutputs() == 1);
-		NDI_FUNC_BINDER(UNiagaraDataInterfaceGrid2DCollection, GetAttributeIndex)::Bind(this, OutFunc, AttributeName, 3);
+		NDI_FUNC_BINDER(UNiagaraDataInterfaceGrid2DCollection, VMGetAttributeIndex)::Bind(this, OutFunc, AttributeName, 3);
 	}
 	else if (BindingInfo.Name == GetVector2DAttributeIndexFunctionName)
 	{
 		FName AttributeName = BindingInfo.FindSpecifier(NAME_Attribute)->Value;
 		check(BindingInfo.GetNumInputs() == 1 && BindingInfo.GetNumOutputs() == 1);
-		NDI_FUNC_BINDER(UNiagaraDataInterfaceGrid2DCollection, GetAttributeIndex)::Bind(this, OutFunc, AttributeName, 2);
+		NDI_FUNC_BINDER(UNiagaraDataInterfaceGrid2DCollection, VMGetAttributeIndex)::Bind(this, OutFunc, AttributeName, 2);
 	}
 	else if (BindingInfo.Name == GetFloatAttributeIndexFunctionName)
 	{
 		FName AttributeName = BindingInfo.FindSpecifier(NAME_Attribute)->Value;
 		check(BindingInfo.GetNumInputs() == 1 && BindingInfo.GetNumOutputs() == 1);
-		NDI_FUNC_BINDER(UNiagaraDataInterfaceGrid2DCollection, GetAttributeIndex)::Bind(this, OutFunc, AttributeName, 1);
+		NDI_FUNC_BINDER(UNiagaraDataInterfaceGrid2DCollection, VMGetAttributeIndex)::Bind(this, OutFunc, AttributeName, 1);
 	}
 }
 
@@ -1081,7 +964,7 @@ void UNiagaraDataInterfaceGrid2DCollection::GetParameterDefinitionHLSL(const FNi
 
 	static const TCHAR *FormatDeclarations = TEXT(R"(
 		Texture2DArray<float> {GridName};
-		RWTexture2DArray<float> RW{OutputGridName};
+		RWTexture2DArray<float> {OutputGridName};
 		SamplerState {SamplerName};
 		int4 {AttributeIndicesName}[{AttributeInt4Count}];
 		int {NumAttributesName};
@@ -1094,12 +977,12 @@ void UNiagaraDataInterfaceGrid2DCollection::GetParameterDefinitionHLSL(const FNi
 	const int AttributeInt4Count = FMath::Max(1, FMath::DivideAndRoundUp(ParamInfo.GeneratedFunctions.Num(), 4));
 
 	TMap<FString, FStringFormatArg> ArgsDeclarations = {
-		{ TEXT("GridName"),    GridName + ParamInfo.DataInterfaceHLSLSymbol },
-		{ TEXT("SamplerName"),    SamplerName + ParamInfo.DataInterfaceHLSLSymbol },
-		{ TEXT("OutputGridName"),    OutputGridName + ParamInfo.DataInterfaceHLSLSymbol },
-		{ TEXT("AttributeIndicesName"), AttributeIndicesBaseName + ParamInfo.DataInterfaceHLSLSymbol},
-		{ TEXT("AttributeInt4Count"), AttributeInt4Count},
-		{ TEXT("NumAttributesName"), UNiagaraDataInterfaceRWBase::NumAttributesName + ParamInfo.DataInterfaceHLSLSymbol},
+		{ TEXT("GridName"),				ParamInfo.DataInterfaceHLSLSymbol + GridName},
+		{ TEXT("SamplerName"),			ParamInfo.DataInterfaceHLSLSymbol + SamplerName},
+		{ TEXT("OutputGridName"),		ParamInfo.DataInterfaceHLSLSymbol + OutputGridName},
+		{ TEXT("AttributeIndicesName"),	ParamInfo.DataInterfaceHLSLSymbol + TEXT("_") + AttributeIndicesBaseName},
+		{ TEXT("AttributeInt4Count"),	AttributeInt4Count},
+		{ TEXT("NumAttributesName"),	ParamInfo.DataInterfaceHLSLSymbol + UNiagaraDataInterfaceRWBase::NumAttributesName},
 	};
 	OutHLSL += FString::Format(FormatDeclarations, ArgsDeclarations);
 }
@@ -1153,19 +1036,19 @@ void UNiagaraDataInterfaceGrid2DCollection::WriteSetHLSL(const FNiagaraDataInter
 					})");
 	}
 	FormatBounds += TEXT(R"(
-					RW{OutputGrid}[int3(In_IndexX, In_IndexY, In_AttributeIndex + i)] = Val;
+					{OutputGrid}[int3(In_IndexX, In_IndexY, In_AttributeIndex + i)] = Val;
 				}
 			}
 		)");
 	TMap<FString, FStringFormatArg> ArgsBounds = {
-		{TEXT("FunctionName"), FunctionInfo.InstanceName},
-		{TEXT("OutputGrid"), OutputGridName + ParamInfo.DataInterfaceHLSLSymbol},
-		{TEXT("NumCellsName"), UNiagaraDataInterfaceRWBase::NumCellsName + ParamInfo.DataInterfaceHLSLSymbol},
-		{TEXT("UnitToUVName"), UNiagaraDataInterfaceRWBase::UnitToUVName + ParamInfo.DataInterfaceHLSLSymbol},
-		{TEXT("AttributeIndicesName"), AttributeIndicesBaseName + ParamInfo.DataInterfaceHLSLSymbol},
-		{TEXT("AttributeIndexGroup"), FunctionInstanceIndex / 4},
-		{TEXT("AttributeIndexComponent"), VectorComponentNames[FunctionInstanceIndex % 4]},
-		{TEXT("NumChannels"), FString::FromInt(InNumChannels)},
+		{TEXT("FunctionName"),				FunctionInfo.InstanceName},
+		{TEXT("OutputGrid"),				ParamInfo.DataInterfaceHLSLSymbol + OutputGridName},
+		{TEXT("NumCellsName"),				ParamInfo.DataInterfaceHLSLSymbol + UNiagaraDataInterfaceRWBase::NumCellsName},
+		{TEXT("UnitToUVName"),				ParamInfo.DataInterfaceHLSLSymbol + UNiagaraDataInterfaceRWBase::UnitToUVName},
+		{TEXT("AttributeIndicesName"),		ParamInfo.DataInterfaceHLSLSymbol + TEXT("_") + AttributeIndicesBaseName},
+		{TEXT("AttributeIndexGroup"),		FunctionInstanceIndex / 4},
+		{TEXT("AttributeIndexComponent"),	VectorComponentNames[FunctionInstanceIndex % 4]},
+		{TEXT("NumChannels"),				FString::FromInt(InNumChannels)},
 		{TEXT("NumChannelsVariableSuffix"), InNumChannels > 1 ? FString::FromInt(InNumChannels) : TEXT("")},
 
 	};
@@ -1225,16 +1108,16 @@ void UNiagaraDataInterfaceGrid2DCollection::WriteGetHLSL(const FNiagaraDataInter
 			}
 		)");
 	TMap<FString, FStringFormatArg> ArgsBounds = {
-		{TEXT("FunctionName"), FunctionInfo.InstanceName},
-		{TEXT("OutputGrid"), OutputGridName + ParamInfo.DataInterfaceHLSLSymbol},
-		{TEXT("Grid"), GridName + ParamInfo.DataInterfaceHLSLSymbol},
-		{TEXT("NumCellsName"), UNiagaraDataInterfaceRWBase::NumCellsName + ParamInfo.DataInterfaceHLSLSymbol},
-		{TEXT("UnitToUVName"), UNiagaraDataInterfaceRWBase::UnitToUVName + ParamInfo.DataInterfaceHLSLSymbol},
-		{TEXT("AttributeIndicesName"), AttributeIndicesBaseName + ParamInfo.DataInterfaceHLSLSymbol},
-		{TEXT("AttributeIndexGroup"), FunctionInstanceIndex / 4},
-		{TEXT("AttributeIndexComponent"), VectorComponentNames[FunctionInstanceIndex % 4]},
-		{TEXT("NumChannels"), FString::FromInt(InNumChannels)},
-		{TEXT("NumChannelsVariableSuffix"), InNumChannels > 1 ? FString::FromInt(InNumChannels) : TEXT("")},
+		{TEXT("FunctionName"),				FunctionInfo.InstanceName},
+		{TEXT("OutputGrid"),				ParamInfo.DataInterfaceHLSLSymbol + OutputGridName},
+		{TEXT("Grid"),						ParamInfo.DataInterfaceHLSLSymbol + GridName},
+		{TEXT("NumCellsName"),				ParamInfo.DataInterfaceHLSLSymbol + UNiagaraDataInterfaceRWBase::NumCellsName},
+		{TEXT("UnitToUVName"),				ParamInfo.DataInterfaceHLSLSymbol + UNiagaraDataInterfaceRWBase::UnitToUVName},
+		{TEXT("AttributeIndicesName"),		ParamInfo.DataInterfaceHLSLSymbol + TEXT("_") + AttributeIndicesBaseName},
+		{TEXT("AttributeIndexGroup"),		FunctionInstanceIndex / 4},
+		{TEXT("AttributeIndexComponent"),	VectorComponentNames[FunctionInstanceIndex % 4]},
+		{TEXT("NumChannels"),				FString::FromInt(InNumChannels)},
+		{TEXT("NumChannelsVariableSuffix"),	InNumChannels > 1 ? FString::FromInt(InNumChannels) : TEXT("")},
 
 	};
 	OutHLSL += FString::Format(*FormatBounds, ArgsBounds);
@@ -1293,16 +1176,16 @@ void UNiagaraDataInterfaceGrid2DCollection::WriteSampleHLSL(const FNiagaraDataIn
 		)");
 
 	TMap<FString, FStringFormatArg> ArgsBounds = {
-		{TEXT("FunctionName"), FunctionInfo.InstanceName},
-		{TEXT("Grid"), GridName + ParamInfo.DataInterfaceHLSLSymbol},
-		{TEXT("SamplerName"),    SamplerName + ParamInfo.DataInterfaceHLSLSymbol },
-		{TEXT("NumCellsName"), UNiagaraDataInterfaceRWBase::NumCellsName + ParamInfo.DataInterfaceHLSLSymbol},
-		{TEXT("UnitToUVName"), UNiagaraDataInterfaceRWBase::UnitToUVName + ParamInfo.DataInterfaceHLSLSymbol},
-		{TEXT("NumChannels"), FString::FromInt(InNumChannels)},
-		{TEXT("NumChannelsVariableSuffix"), InNumChannels > 1 ? FString::FromInt(InNumChannels) : TEXT("")},
-		{TEXT("AttributeIndicesName"), AttributeIndicesBaseName + ParamInfo.DataInterfaceHLSLSymbol},
-		{TEXT("AttributeIndexGroup"), FunctionInstanceIndex / 4},
-		{TEXT("AttributeIndexComponent"), VectorComponentNames[FunctionInstanceIndex % 4]},
+		{TEXT("FunctionName"),				FunctionInfo.InstanceName},
+		{TEXT("Grid"),						ParamInfo.DataInterfaceHLSLSymbol + GridName},
+		{TEXT("SamplerName"),				ParamInfo.DataInterfaceHLSLSymbol + SamplerName},
+		{TEXT("NumCellsName"),				ParamInfo.DataInterfaceHLSLSymbol + UNiagaraDataInterfaceRWBase::NumCellsName},
+		{TEXT("UnitToUVName"),				ParamInfo.DataInterfaceHLSLSymbol + UNiagaraDataInterfaceRWBase::UnitToUVName},
+		{TEXT("NumChannels"),				FString::FromInt(InNumChannels)},
+		{TEXT("NumChannelsVariableSuffix"),	InNumChannels > 1 ? FString::FromInt(InNumChannels) : TEXT("")},
+		{TEXT("AttributeIndicesName"),		ParamInfo.DataInterfaceHLSLSymbol + TEXT("_") + AttributeIndicesBaseName},
+		{TEXT("AttributeIndexGroup"),		FunctionInstanceIndex / 4},
+		{TEXT("AttributeIndexComponent"),	VectorComponentNames[FunctionInstanceIndex % 4]},
 	};
 	OutHLSL += FString::Format(*FormatBounds, ArgsBounds);
 }
@@ -1320,7 +1203,7 @@ void UNiagaraDataInterfaceGrid2DCollection::WriteAttributeGetIndexHLSL(const FNi
 
 	TMap<FString, FStringFormatArg> ArgsBounds = {
 		{TEXT("FunctionName"), FunctionInfo.InstanceName},
-		{TEXT("AttributeIndicesName"), AttributeIndicesBaseName + ParamInfo.DataInterfaceHLSLSymbol},
+		{TEXT("AttributeIndicesName"), ParamInfo.DataInterfaceHLSLSymbol + TEXT("_") + AttributeIndicesBaseName},
 		{TEXT("AttributeIndexGroup"), FunctionInstanceIndex / 4},
 		{TEXT("AttributeIndexComponent"), VectorComponentNames[FunctionInstanceIndex % 4]},
 	};
@@ -1374,13 +1257,13 @@ bool UNiagaraDataInterfaceGrid2DCollection::GetFunctionHLSL(const FNiagaraDataIn
 
 	TMap<FString, FStringFormatArg> ArgsBounds =
 	{
-		{TEXT("FunctionName"), FunctionInfo.InstanceName},
-		{TEXT("Grid"), GridName + ParamInfo.DataInterfaceHLSLSymbol},
-		{TEXT("OutputGrid"), OutputGridName + ParamInfo.DataInterfaceHLSLSymbol},
-		{TEXT("NumAttributes"), UNiagaraDataInterfaceRWBase::NumAttributesName + ParamInfo.DataInterfaceHLSLSymbol},
-		{TEXT("NumCells"), UNiagaraDataInterfaceRWBase::NumCellsName + ParamInfo.DataInterfaceHLSLSymbol},
-		{TEXT("UnitToUVName"), UNiagaraDataInterfaceRWBase::UnitToUVName + ParamInfo.DataInterfaceHLSLSymbol},
-		{TEXT("SamplerName"), SamplerName + ParamInfo.DataInterfaceHLSLSymbol},
+		{TEXT("FunctionName"),	FunctionInfo.InstanceName},
+		{TEXT("Grid"),			ParamInfo.DataInterfaceHLSLSymbol + GridName},
+		{TEXT("OutputGrid"),	ParamInfo.DataInterfaceHLSLSymbol + OutputGridName},
+		{TEXT("NumAttributes"),	ParamInfo.DataInterfaceHLSLSymbol + UNiagaraDataInterfaceRWBase::NumAttributesName},
+		{TEXT("NumCells"),		ParamInfo.DataInterfaceHLSLSymbol + UNiagaraDataInterfaceRWBase::NumCellsName},
+		{TEXT("UnitToUVName"),	ParamInfo.DataInterfaceHLSLSymbol + UNiagaraDataInterfaceRWBase::UnitToUVName},
+		{TEXT("SamplerName"),	ParamInfo.DataInterfaceHLSLSymbol + SamplerName},
 	};
 
 	if (FunctionInfo.DefinitionName == GetValueFunctionName || FunctionInfo.DefinitionName == GetPreviousValueAtIndexFunctionName)
@@ -1399,7 +1282,7 @@ bool UNiagaraDataInterfaceGrid2DCollection::GetFunctionHLSL(const FNiagaraDataIn
 		static const TCHAR* FormatBounds = TEXT(R"(
 			void {FunctionName}(int In_IndexX, int In_IndexY, int In_AttributeIndex, float In_Value)
 			{
-				RW{OutputGrid}[int3(In_IndexX, In_IndexY, In_AttributeIndex)] = In_Value;
+				{OutputGrid}[int3(In_IndexX, In_IndexY, In_AttributeIndex)] = In_Value;
 			}
 		)");
 		OutHLSL += FString::Format(FormatBounds, ArgsBounds);
@@ -1411,7 +1294,7 @@ bool UNiagaraDataInterfaceGrid2DCollection::GetFunctionHLSL(const FNiagaraDataIn
 			void {FunctionName}(int In_IndexX, int In_IndexY, int In_AttributeIndex, float In_Value, out int val)
 			{
 				val = 0;
-				RW{OutputGrid}[int3(In_IndexX, In_IndexY, In_AttributeIndex)] = In_Value;
+				{OutputGrid}[int3(In_IndexX, In_IndexY, In_AttributeIndex)] = In_Value;
 			}
 		)");
 		OutHLSL += FString::Format(FormatBounds, ArgsBounds);
@@ -1425,7 +1308,7 @@ bool UNiagaraDataInterfaceGrid2DCollection::GetFunctionHLSL(const FNiagaraDataIn
 				for (int AttributeIndex = 0; AttributeIndex < {NumAttributes}.x; AttributeIndex++)
 				{
 					float Val = {Grid}.Load(int4(In_IndexX, In_IndexY, AttributeIndex, 0));
-					RW{OutputGrid}[int3(In_IndexX, In_IndexY, AttributeIndex)] = Val;
+					{OutputGrid}[int3(In_IndexX, In_IndexY, AttributeIndex)] = Val;
 				}
 			}
 		)");
@@ -1440,7 +1323,7 @@ bool UNiagaraDataInterfaceGrid2DCollection::GetFunctionHLSL(const FNiagaraDataIn
 				for (int AttributeIndex = 0; AttributeIndex < {NumAttributes}.x; AttributeIndex++)
 				{
 					float Val = 0.0f;
-					RW{OutputGrid}[int3(In_IndexX, In_IndexY, AttributeIndex)] = Val;
+					{OutputGrid}[int3(In_IndexX, In_IndexY, AttributeIndex)] = Val;
 				}
 			}
 		)");
@@ -1543,6 +1426,120 @@ bool UNiagaraDataInterfaceGrid2DCollection::GetFunctionHLSL(const FNiagaraDataIn
 }
 #endif
 
+void UNiagaraDataInterfaceGrid2DCollection::BuildShaderParameters(FNiagaraShaderParametersBuilder& ShaderParametersBuilder) const
+{
+	ShaderParametersBuilder.AddNestedStruct<FNDIGrid2DShaderParameters>();
+
+	const int32 IndirectionTableSize = FMath::Max(FMath::DivideAndRoundUp(ShaderParametersBuilder.GetGeneratedFunctions().Num(), 4), 1);
+	ShaderParametersBuilder.AddLooseParamArray<FIntVector4>(*UNiagaraDataInterfaceGrid2DCollection::AttributeIndicesBaseName, IndirectionTableSize);
+}
+
+void UNiagaraDataInterfaceGrid2DCollection::SetShaderParameters(const FNiagaraDataInterfaceSetShaderParametersContext& Context) const
+{
+	FNiagaraDataInterfaceProxyGrid2DCollectionProxy& DIProxy = Context.GetProxy<FNiagaraDataInterfaceProxyGrid2DCollectionProxy>();
+	FGrid2DCollectionRWInstanceData_RenderThread* ProxyData = DIProxy.SystemInstancesToProxyData_RT.Find(Context.GetSystemInstanceID());
+	FGrid2DCollectionRWInstanceData_RenderThread* OriginalProxyData = ProxyData;
+	check(ProxyData);
+
+	if (ProxyData->OtherProxy != nullptr)
+	{
+		FNiagaraDataInterfaceProxyGrid2DCollectionProxy* OtherGrid3DProxy = static_cast<FNiagaraDataInterfaceProxyGrid2DCollectionProxy*>(ProxyData->OtherProxy);
+		ProxyData = OtherGrid3DProxy->SystemInstancesToProxyData_RT.Find(Context.GetSystemInstanceID());
+		check(ProxyData);
+	}
+
+	const FNiagaraDataInterfaceParametersCS_Grid2DCollection& ShaderStorage = Context.GetShaderStorage<FNiagaraDataInterfaceParametersCS_Grid2DCollection>();
+	if (OriginalProxyData->AttributeIndices.Num() == 0 && ShaderStorage.AttributeNames.Num() > 0)
+	{
+		OriginalProxyData->AttributeIndices.SetNumZeroed(Align(ShaderStorage.AttributeNames.Num(), 4));
+
+		// TODO handle mismatched types!
+		for (int32 i = 0; i < ShaderStorage.AttributeNames.Num(); i++)
+		{
+			const int32 FoundIdx = ProxyData->Vars.Find(ShaderStorage.AttributeNames[i]);
+			check(ShaderStorage.AttributeNames.Num() == ShaderStorage.AttributeChannelCount.Num());
+			check(ProxyData->Offsets.Num() == ProxyData->VarComponents.Num());
+			check(ProxyData->Offsets.Num() == ProxyData->Vars.Num());
+			if (ProxyData->Offsets.IsValidIndex(FoundIdx) && ShaderStorage.AttributeChannelCount[i] == ProxyData->VarComponents[FoundIdx])
+			{
+				OriginalProxyData->AttributeIndices[i] = ProxyData->Offsets[FoundIdx];
+			}
+			else
+			{
+				OriginalProxyData->AttributeIndices[i] = -1; // We may need to protect against this in the hlsl as this might underflow an array lookup if used incorrectly.
+			}
+		}
+	}
+
+	FNDIGrid2DShaderParameters* Parameters = Context.GetParameterNestedStruct<FNDIGrid2DShaderParameters>();
+	Parameters->NumAttributes = ProxyData->NumAttributes;
+	Parameters->UnitToUV = FVector2f(1.0f) / FVector2f(ProxyData->NumCells);
+	Parameters->NumCells = ProxyData->NumCells;
+	Parameters->CellSize = FVector2f(ProxyData->CellSize);
+	Parameters->WorldBBoxSize = FVector2f(ProxyData->WorldBBoxSize);	// LWC_TODO: Precision loss?
+	Parameters->SamplerState = TStaticSamplerState<SF_Bilinear, AM_Clamp, AM_Clamp, AM_Clamp>::GetRHI();
+	Parameters->Grid = ProxyData->CurrentData ? ProxyData->CurrentData->GridSRV.GetReference() : FNiagaraRenderer::GetDummyTextureReadBuffer2DArray();
+	if (Context.IsResourceBound(&Parameters->OutputGrid))
+	{
+		if (Context.IsOutputStage() && ProxyData->DestinationData)
+		{
+			Parameters->OutputGrid = ProxyData->DestinationData->GridUAV;
+		}
+		else
+		{
+			Parameters->OutputGrid = Context.GetComputeDispatchInterface().GetEmptyUAVFromPool(Context.GetRHICmdList(), PF_R32_FLOAT, ENiagaraEmptyUAVType::Texture2DArray);
+		}
+	}
+
+	// Do we need to set the indirection table?
+	if (ShaderStorage.AttributeNames.Num() > 0)
+	{
+		const int NumAttributesVector4s = FMath::Max(FMath::DivideAndRoundUp(ShaderStorage.AttributeNames.Num(), 4), 1);
+		TArrayView<FIntVector4> AttributeIndices = Context.GetParameterLooseArray<FIntVector4>(NumAttributesVector4s);
+		check(OriginalProxyData->AttributeIndices.Num() * OriginalProxyData->AttributeIndices.GetTypeSize() == AttributeIndices.Num() * AttributeIndices.GetTypeSize());
+		FMemory::Memcpy(AttributeIndices.GetData(), OriginalProxyData->AttributeIndices.GetData(), AttributeIndices.Num() * AttributeIndices.GetTypeSize());
+	}
+}
+
+FNiagaraDataInterfaceParametersCS* UNiagaraDataInterfaceGrid2DCollection::CreateShaderStorage(const FNiagaraDataInterfaceGPUParamInfo& ParameterInfo, const FShaderParameterMap& ParameterMap) const
+{
+	FNiagaraDataInterfaceParametersCS_Grid2DCollection* ShaderStorage = new FNiagaraDataInterfaceParametersCS_Grid2DCollection();
+
+	// Gather up all the attribute names referenced. Note that there may be multiple in the list of the same name,
+	// but we only deal with this by the number of bound methods.
+	const int32 NumFuncs = ParameterInfo.GeneratedFunctions.Num();
+	ShaderStorage->AttributeNames.Reserve(NumFuncs);
+	ShaderStorage->AttributeChannelCount.Reserve(NumFuncs);
+
+	for (int32 FuncIdx=0; FuncIdx < NumFuncs; ++FuncIdx)
+	{
+		const FNiagaraDataInterfaceGeneratedFunction& Func = ParameterInfo.GeneratedFunctions[FuncIdx];
+		static const FName NAME_Attribute("Attribute");
+		const FName* AttributeName = Func.FindSpecifierValue(NAME_Attribute);
+		if (AttributeName != nullptr)
+		{
+			int32 ComponentCount = UNiagaraDataInterfaceGrid2DCollection::GetComponentCountFromFuncName(Func.DefinitionName);
+			ShaderStorage->AttributeNames.Add(*AttributeName);
+			ShaderStorage->AttributeChannelCount.Add(ComponentCount);
+		}
+		else
+		{
+			ShaderStorage->AttributeNames.Add(FName());
+			ShaderStorage->AttributeChannelCount.Add(INDEX_NONE);
+		}
+	}
+
+	ShaderStorage->AttributeNames.Shrink();
+	ShaderStorage->AttributeChannelCount.Shrink();
+
+	return ShaderStorage;
+}
+
+const FTypeLayoutDesc* UNiagaraDataInterfaceGrid2DCollection::GetShaderStorageType() const
+{
+	return &StaticGetTypeLayoutDesc<FNiagaraDataInterfaceParametersCS_Grid2DCollection>();
+}
+
 #if WITH_EDITOR
 bool UNiagaraDataInterfaceGrid2DCollection::GenerateIterationSourceNamespaceReadAttributesHLSL(FNiagaraDataInterfaceGPUParamInfo& DIInstanceInfo, const FNiagaraVariable& IterationSourceVar, TConstArrayView<FNiagaraVariable> InArguments, TConstArrayView<FNiagaraVariable> InAttributes, TConstArrayView<FString> InAttributeHLSLNames, bool bInSetToDefaults, bool bSpawnOnly, bool bPartialWrites, TArray<FText>& OutErrors, FString& OutHLSL) const
 {
@@ -2288,7 +2285,7 @@ void UNiagaraDataInterfaceGrid2DCollection::GetTextureSize(const UNiagaraCompone
 	SizeY = Grid2DInstanceData->NumCells.Y;
 }
 
-void UNiagaraDataInterfaceGrid2DCollection::GetWorldBBoxSize(FVectorVMExternalFunctionContext& Context)
+void UNiagaraDataInterfaceGrid2DCollection::VMGetWorldBBoxSize(FVectorVMExternalFunctionContext& Context)
 {
 	VectorVM::FUserPtrHandler<FGrid2DCollectionRWInstanceData_GameThread> InstData(Context);
 	FNDIOutputParam<FVector2f> OutWorldBounds(Context);
@@ -2301,7 +2298,7 @@ void UNiagaraDataInterfaceGrid2DCollection::GetWorldBBoxSize(FVectorVMExternalFu
 }
 
 
-void UNiagaraDataInterfaceGrid2DCollection::GetCellSize(FVectorVMExternalFunctionContext& Context)
+void UNiagaraDataInterfaceGrid2DCollection::VMGetCellSize(FVectorVMExternalFunctionContext& Context)
 {
 	VectorVM::FUserPtrHandler<FGrid2DCollectionRWInstanceData_GameThread> InstData(Context);
 	FNDIOutputParam<FVector2f> OutCellSize(Context);
@@ -2313,7 +2310,7 @@ void UNiagaraDataInterfaceGrid2DCollection::GetCellSize(FVectorVMExternalFunctio
 	}
 }
 
-void UNiagaraDataInterfaceGrid2DCollection::GetNumCells(FVectorVMExternalFunctionContext& Context)
+void UNiagaraDataInterfaceGrid2DCollection::VMGetNumCells(FVectorVMExternalFunctionContext& Context)
 {
 	VectorVM::FUserPtrHandler<FGrid2DCollectionRWInstanceData_GameThread> InstData(Context);
 	FNDIOutputParam<int> OutNumCellsX(Context);
@@ -2334,7 +2331,7 @@ void UNiagaraDataInterfaceGrid2DCollection::GetNumCells(FVectorVMExternalFunctio
 	}
 }
 
-void UNiagaraDataInterfaceGrid2DCollection::SetNumCells(FVectorVMExternalFunctionContext& Context)
+void UNiagaraDataInterfaceGrid2DCollection::VMSetNumCells(FVectorVMExternalFunctionContext& Context)
 {
 	// This should only be called from a system or emitter script due to a need for only setting up initially.
 	VectorVM::FUserPtrHandler<FGrid2DCollectionRWInstanceData_GameThread> InstData(Context);
@@ -2420,7 +2417,7 @@ bool UNiagaraDataInterfaceGrid2DCollection::PerInstanceTickPostSimulate(void* Pe
 }
 
 
-void UNiagaraDataInterfaceGrid2DCollection::GetAttributeIndex(FVectorVMExternalFunctionContext& Context, const FName& InName, int32 NumChannels)
+void UNiagaraDataInterfaceGrid2DCollection::VMGetAttributeIndex(FVectorVMExternalFunctionContext& Context, const FName& InName, int32 NumChannels)
 {
 	VectorVM::FUserPtrHandler<FGrid2DCollectionRWInstanceData_GameThread> InstData(Context);
 	VectorVM::FExternalFuncRegisterHandler<int> OutIndex(Context);
diff --git a/Engine/Plugins/FX/Niagara/Source/Niagara/Private/NiagaraDataInterfaceGrid2DCollectionReader.cpp b/Engine/Plugins/FX/Niagara/Source/Niagara/Private/NiagaraDataInterfaceGrid2DCollectionReader.cpp
--- a/Engine/Plugins/FX/Niagara/Source/Niagara/Private/NiagaraDataInterfaceGrid2DCollectionReader.cpp
+++ b/Engine/Plugins/FX/Niagara/Source/Niagara/Private/NiagaraDataInterfaceGrid2DCollectionReader.cpp
@@ -17,8 +17,6 @@
 
 #define LOCTEXT_NAMESPACE "NiagaraDataInterfaceGrid2DCollectionReader"
 
-IMPLEMENT_NIAGARA_DI_PARAMETER(UNiagaraDataInterfaceGrid2DCollectionReader, FNiagaraDataInterfaceParametersCS_Grid2DCollection);
-
 UNiagaraDataInterfaceGrid2DCollectionReader::UNiagaraDataInterfaceGrid2DCollectionReader(FObjectInitializer const& ObjectInitializer)
 	: Super(ObjectInitializer)
 {
diff --git a/Engine/Plugins/FX/Niagara/Source/Niagara/Private/NiagaraDataInterfaceGrid3DCollection.cpp b/Engine/Plugins/FX/Niagara/Source/Niagara/Private/NiagaraDataInterfaceGrid3DCollection.cpp
--- a/Engine/Plugins/FX/Niagara/Source/Niagara/Private/NiagaraDataInterfaceGrid3DCollection.cpp
+++ b/Engine/Plugins/FX/Niagara/Source/Niagara/Private/NiagaraDataInterfaceGrid3DCollection.cpp
@@ -1,31 +1,53 @@
 // Copyright Epic Games, Inc. All Rights Reserved.
 #include "NiagaraDataInterfaceGrid3DCollection.h"
-#include "NiagaraShader.h"
-#include "ShaderParameterUtils.h"
-#include "ClearQuad.h"
-#include "TextureResource.h"
+#include "NiagaraConstants.h"
 #include "NiagaraGpuComputeDispatchInterface.h"
 #include "NiagaraSimStageData.h"
-#include "NiagaraSystemInstance.h"
 #include "NiagaraRenderer.h"
-#include "Engine/VolumeTexture.h"
-#include "Engine/TextureRenderTargetVolume.h"
 #include "NiagaraSettings.h"
-#include "NiagaraConstants.h"
+#include "NiagaraShader.h"
+#include "NiagaraShaderParametersBuilder.h"
+#include "NiagaraSystemInstance.h"
 #if WITH_EDITOR
 #include "NiagaraGpuComputeDebug.h"
 #endif
 
+#include "Engine/TextureRenderTargetVolume.h"
+#include "Engine/VolumeTexture.h"
+#include "ClearQuad.h"
+#include "ShaderParameterUtils.h"
+#include "TextureResource.h"
+
 #define LOCTEXT_NAMESPACE "NiagaraDataInterfaceGrid3DCollection"
 
-const FString UNiagaraDataInterfaceGrid3DCollection::NumTilesName(TEXT("NumTiles_"));
-const FString UNiagaraDataInterfaceGrid3DCollection::OneOverNumTilesName(TEXT("OneOverNumTiles_"));
-const FString UNiagaraDataInterfaceGrid3DCollection::UnitClampMinName(TEXT("UnitClampMin_"));
-const FString UNiagaraDataInterfaceGrid3DCollection::UnitClampMaxName(TEXT("UnitClampMax_"));
+IMPLEMENT_TYPE_LAYOUT(FNiagaraDataInterfaceParametersCS_Grid3DCollection);
 
-const FString UNiagaraDataInterfaceGrid3DCollection::GridName(TEXT("Grid_"));
-const FString UNiagaraDataInterfaceGrid3DCollection::OutputGridName(TEXT("OutputGrid_"));
-const FString UNiagaraDataInterfaceGrid3DCollection::SamplerName(TEXT("Sampler_"));
+BEGIN_SHADER_PARAMETER_STRUCT(FNDIGrid3DShaderParameters, )
+	SHADER_PARAMETER(int,							NumAttributes)
+	SHADER_PARAMETER(int,							NumNamedAttributes)
+	SHADER_PARAMETER(FVector3f,						UnitToUV)
+	SHADER_PARAMETER(FIntVector,					NumCells)
+	SHADER_PARAMETER(FVector3f,						CellSize)
+	SHADER_PARAMETER(FIntVector,					NumTiles)
+	SHADER_PARAMETER(FVector3f,						OneOverNumTiles)
+	SHADER_PARAMETER(FVector3f,						UnitClampMin)
+	SHADER_PARAMETER(FVector3f,						UnitClampMax)
+	SHADER_PARAMETER(FVector3f,						WorldBBoxSize)
+
+	SHADER_PARAMETER_SRV(Texture3D<float>,			Grid)
+	SHADER_PARAMETER_SAMPLER(SamplerState,			SamplerState)
+	SHADER_PARAMETER_UAV(RWTexture3D<float>,		OutputGrid)
+	SHADER_PARAMETER_SRV(Buffer<float4>,			PerAttributeData)
+END_SHADER_PARAMETER_STRUCT()
+
+const FString UNiagaraDataInterfaceGrid3DCollection::NumTilesName(TEXT("_NumTiles"));
+const FString UNiagaraDataInterfaceGrid3DCollection::OneOverNumTilesName(TEXT("_OneOverNumTiles"));
+const FString UNiagaraDataInterfaceGrid3DCollection::UnitClampMinName(TEXT("_UnitClampMin"));
+const FString UNiagaraDataInterfaceGrid3DCollection::UnitClampMaxName(TEXT("_UnitClampMax"));
+
+const FString UNiagaraDataInterfaceGrid3DCollection::GridName(TEXT("_Grid"));
+const FString UNiagaraDataInterfaceGrid3DCollection::OutputGridName(TEXT("_OutputGrid"));
+const FString UNiagaraDataInterfaceGrid3DCollection::SamplerName(TEXT("_Sampler"));
 
 const FName UNiagaraDataInterfaceGrid3DCollection::SetNumCellsFunctionName("SetNumCells");
 
@@ -67,8 +89,8 @@ const FName UNiagaraDataInterfaceGrid3DCollection::GetFloatAttributeIndexFunctio
 
 const FName UNiagaraDataInterfaceGrid3DCollection::SampleGridFunctionName("SampleGrid");
 
-const FString UNiagaraDataInterfaceGrid3DCollection::AttributeIndicesBaseName(TEXT("AttributeIndices_"));
-const FString UNiagaraDataInterfaceGrid3DCollection::PerAttributeDataName(TEXT("PerAttributeDataName_"));
+const FString UNiagaraDataInterfaceGrid3DCollection::AttributeIndicesBaseName(TEXT("AttributeIndices"));
+const FString UNiagaraDataInterfaceGrid3DCollection::PerAttributeDataName(TEXT("PerAttributeData"));
 
 const TCHAR* UNiagaraDataInterfaceGrid3DCollection::VectorComponentNames[] = { TEXT(".x"), TEXT(".y"), TEXT(".z"), TEXT(".w") };
 
@@ -143,162 +165,6 @@ static FAutoConsoleVariableRef CVarNiagaraGrid3DOverrideFormat(
 	ECVF_Default
 );
 
-/*--------------------------------------------------------------------------------------------------------------------------*/
-
-void FNiagaraDataInterfaceParametersCS_Grid3DCollection::Bind(const FNiagaraDataInterfaceGPUParamInfo& ParameterInfo, const class FShaderParameterMap& ParameterMap)
-	{
-		NumAttributesParam.Bind(ParameterMap, *(UNiagaraDataInterfaceRWBase::NumAttributesName + ParameterInfo.DataInterfaceHLSLSymbol));
-		NumNamedAttributesParam.Bind(ParameterMap, *(UNiagaraDataInterfaceRWBase::NumNamedAttributesName + ParameterInfo.DataInterfaceHLSLSymbol));
-		NumCellsParam.Bind(ParameterMap, *(UNiagaraDataInterfaceRWBase::NumCellsName + ParameterInfo.DataInterfaceHLSLSymbol));
-		NumTilesParam.Bind(ParameterMap, *(UNiagaraDataInterfaceGrid3DCollection::NumTilesName + ParameterInfo.DataInterfaceHLSLSymbol));
-		OneOverNumTilesParam.Bind(ParameterMap, *(UNiagaraDataInterfaceGrid3DCollection::OneOverNumTilesName + ParameterInfo.DataInterfaceHLSLSymbol));
-		UnitClampMinParam.Bind(ParameterMap, *(UNiagaraDataInterfaceGrid3DCollection::UnitClampMinName + ParameterInfo.DataInterfaceHLSLSymbol));
-		UnitClampMaxParam.Bind(ParameterMap, *(UNiagaraDataInterfaceGrid3DCollection::UnitClampMaxName + ParameterInfo.DataInterfaceHLSLSymbol));
-
-		UnitToUVParam.Bind(ParameterMap, *(UNiagaraDataInterfaceRWBase::UnitToUVName + ParameterInfo.DataInterfaceHLSLSymbol));
-		CellSizeParam.Bind(ParameterMap, *(UNiagaraDataInterfaceRWBase::CellSizeName + ParameterInfo.DataInterfaceHLSLSymbol));
-
-		WorldBBoxSizeParam.Bind(ParameterMap, *(UNiagaraDataInterfaceRWBase::WorldBBoxSizeName + ParameterInfo.DataInterfaceHLSLSymbol));
-
-		GridParam.Bind(ParameterMap, *(UNiagaraDataInterfaceGrid3DCollection::GridName + ParameterInfo.DataInterfaceHLSLSymbol));
-		OutputGridParam.Bind(ParameterMap, *(UNiagaraDataInterfaceGrid3DCollection::OutputGridName + ParameterInfo.DataInterfaceHLSLSymbol));
-
-		SamplerParam.Bind(ParameterMap, *(UNiagaraDataInterfaceGrid3DCollection::SamplerName + ParameterInfo.DataInterfaceHLSLSymbol));
-		AttributeIndicesParam.Bind(ParameterMap, *(UNiagaraDataInterfaceGrid3DCollection::AttributeIndicesBaseName + ParameterInfo.DataInterfaceHLSLSymbol));
-		PerAttributeDataParam.Bind(ParameterMap, *(UNiagaraDataInterfaceGrid3DCollection::PerAttributeDataName + ParameterInfo.DataInterfaceHLSLSymbol));
-
-
-		// Gather up all the attribute names referenced. Note that there may be multiple in the list of the same name,
-		// but we only deal with this by the number of bound methods.
-		{
-			int32 NumFuncs = ParameterInfo.GeneratedFunctions.Num();
-
-			for (int32 FuncIdx = 0; FuncIdx < NumFuncs; ++FuncIdx)
-			{
-				const FNiagaraDataInterfaceGeneratedFunction& Func = ParameterInfo.GeneratedFunctions[FuncIdx];
-				static const FName NAME_Attribute("Attribute");
-				const FName* AttributeName = Func.FindSpecifierValue(NAME_Attribute);
-				if (AttributeName != nullptr)
-				{
-					int32 ComponentCount = UNiagaraDataInterfaceGrid3DCollection::GetComponentCountFromFuncName(Func.DefinitionName);
-					AttributeNames.Add(*AttributeName);
-					AttributeChannelCount.Add(ComponentCount);
-				}
-				else
-				{
-					AttributeNames.Add(FName());
-					AttributeChannelCount.Add(INDEX_NONE);
-				}
-			}
-		}
-	}
-void FNiagaraDataInterfaceParametersCS_Grid3DCollection::Set(FRHICommandList& RHICmdList, const FNiagaraDataInterfaceSetArgs& Context) const
-	{
-		check(IsInRenderingThread());
-
-		// Get shader and DI
-		FRHIComputeShader* ComputeShaderRHI = Context.Shader.GetComputeShader();
-		FNiagaraDataInterfaceProxyGrid3DCollectionProxy* VFDI = static_cast<FNiagaraDataInterfaceProxyGrid3DCollectionProxy*>(Context.DataInterface);
-
-		FGrid3DCollectionRWInstanceData_RenderThread* ProxyData = VFDI->SystemInstancesToProxyData_RT.Find(Context.SystemInstanceID);
-		FGrid3DCollectionRWInstanceData_RenderThread* OriginalProxyData = ProxyData;
-
-		check(ProxyData);		
-		if (ProxyData->OtherProxy != nullptr)
-		{			
-			FNiagaraDataInterfaceProxyGrid3DCollectionProxy* OtherGrid3DProxy = static_cast<FNiagaraDataInterfaceProxyGrid3DCollectionProxy*>(ProxyData->OtherProxy);
-			ProxyData = OtherGrid3DProxy->SystemInstancesToProxyData_RT.Find(Context.SystemInstanceID);
-			check(ProxyData);
-		}
-
-		SetShaderValue(RHICmdList, ComputeShaderRHI, NumCellsParam, ProxyData->NumCells);	
-		SetShaderValue(RHICmdList, ComputeShaderRHI, NumTilesParam, ProxyData->NumTiles);
-		SetShaderValue(RHICmdList, ComputeShaderRHI, OneOverNumTilesParam, FVector3f(1.0f) / FVector3f(ProxyData->NumTiles));
-
-		const FVector3f HalfPixelOffset = FVector3f(0.5f / ProxyData->NumCells.X, 0.5f / ProxyData->NumCells.Y, 0.5f / ProxyData->NumCells.Z);
-		SetShaderValue(RHICmdList, ComputeShaderRHI, UnitClampMinParam, HalfPixelOffset);
-		SetShaderValue(RHICmdList, ComputeShaderRHI, UnitClampMaxParam, FVector3f::OneVector - HalfPixelOffset);
-
-		if (OriginalProxyData->AttributeIndices.Num() == 0 && AttributeNames.Num() > 0)
-		{
-			int NumAttrIndices = Align(AttributeNames.Num(), 4);
-			OriginalProxyData->AttributeIndices.SetNumZeroed(NumAttrIndices);
-
-			// TODO handle mismatched types!
-			for (int32 i = 0; i < AttributeNames.Num(); i++)
-			{
-				int32 FoundIdx = ProxyData->Vars.Find(AttributeNames[i]);
-				check(AttributeNames.Num() == AttributeChannelCount.Num());
-				check(ProxyData->Offsets.Num() == ProxyData->VarComponents.Num());
-				check(ProxyData->Offsets.Num() == ProxyData->Vars.Num());
-				if (ProxyData->Offsets.IsValidIndex(FoundIdx) && AttributeChannelCount[i] == ProxyData->VarComponents[FoundIdx])
-				{
-					OriginalProxyData->AttributeIndices[i] = ProxyData->Offsets[FoundIdx];
-				}
-				else
-				{
-					OriginalProxyData->AttributeIndices[i] = -1; // We may need to protect against this in the hlsl as this might underflow an array lookup if used incorrectly.
-				}
-			}
-		}
-
-		SetShaderValue(RHICmdList, ComputeShaderRHI, NumAttributesParam, ProxyData->TotalNumAttributes);
-		SetShaderValue(RHICmdList, ComputeShaderRHI, NumNamedAttributesParam, ProxyData->TotalNumNamedAttributes);
-
-		SetShaderValue(RHICmdList, ComputeShaderRHI, UnitToUVParam, FVector3f(1.0f) / FVector3f(ProxyData->NumCells));
-
-		SetShaderValue(RHICmdList, ComputeShaderRHI, CellSizeParam, (FVector3f)ProxyData->CellSize);		
-
-		SetShaderValue(RHICmdList, ComputeShaderRHI, WorldBBoxSizeParam, (FVector3f)ProxyData->WorldBBoxSize);
-
-		SetShaderValueArray(RHICmdList, ComputeShaderRHI, AttributeIndicesParam, OriginalProxyData->AttributeIndices.GetData(), OriginalProxyData->AttributeIndices.Num());
-		
-		SetSRVParameter(RHICmdList, ComputeShaderRHI, PerAttributeDataParam, ProxyData->PerAttributeData.SRV);
-
-		FRHISamplerState *SamplerState = TStaticSamplerState<SF_Bilinear, AM_Clamp, AM_Clamp, AM_Clamp>::GetRHI();
-		SetSamplerParameter(RHICmdList, ComputeShaderRHI, SamplerParam, SamplerState);
-
-		if (GridParam.IsBound())
-		{
-			FRHIShaderResourceView* InputGridBuffer;
-			if (ProxyData->CurrentData != nullptr)
-			{
-				InputGridBuffer = ProxyData->CurrentData->GridBuffer.SRV;
-			}
-			else
-			{
-				InputGridBuffer = FNiagaraRenderer::GetDummyTextureReadBuffer3D();
-			}
-			SetSRVParameter(RHICmdList, ComputeShaderRHI, GridParam, InputGridBuffer);
-		}
-
-		if ( OutputGridParam.IsUAVBound() )
-		{
-			FRHIUnorderedAccessView* OutputGridUAV;
-			if (Context.IsOutputStage && ProxyData->DestinationData != nullptr)
-			{
-				OutputGridUAV = ProxyData->DestinationData->GridBuffer.UAV;
-			}
-			else
-			{
-				OutputGridUAV = Context.ComputeDispatchInterface->GetEmptyUAVFromPool(RHICmdList, PF_R32_FLOAT, ENiagaraEmptyUAVType::Texture3D);
-			}
-			RHICmdList.SetUAVParameter(ComputeShaderRHI, OutputGridParam.GetUAVIndex(), OutputGridUAV);
-		}
-	}
-void FNiagaraDataInterfaceParametersCS_Grid3DCollection::Unset(FRHICommandList& RHICmdList, const FNiagaraDataInterfaceSetArgs& Context) const
-	{
-		if (OutputGridParam.IsBound())
-		{
-			FRHIComputeShader* ComputeShaderRHI = Context.Shader.GetComputeShader();
-			OutputGridParam.UnsetUAV(RHICmdList, ComputeShaderRHI);
-		}
-	}
-
-IMPLEMENT_TYPE_LAYOUT(FNiagaraDataInterfaceParametersCS_Grid3DCollection);
-
-IMPLEMENT_NIAGARA_DI_PARAMETER(UNiagaraDataInterfaceGrid3DCollection, FNiagaraDataInterfaceParametersCS_Grid3DCollection);
-
 UNiagaraDataInterfaceGrid3DCollection::UNiagaraDataInterfaceGrid3DCollection(FObjectInitializer const& ObjectInitializer)
 	: Super(ObjectInitializer)
 	, NumAttributes(1)
@@ -330,6 +196,7 @@ bool UNiagaraDataInterfaceGrid3DCollection::AppendCompileHash(FNiagaraCompileHas
 
 
 	InVisitor->UpdatePOD(TEXT("UNiagaraDataInterfaceGrid3DCollectionVersion"), (int32)FNiagaraGridCollection3DDIFunctionVersion::LatestVersion);
+	InVisitor->UpdateShaderParameters<FNDIGrid3DShaderParameters>();
 	return true;
 }
 #endif
@@ -974,12 +841,12 @@ void UNiagaraDataInterfaceGrid3DCollection::GetFunctions(TArray<FNiagaraFunction
 // #todo(dmp): expose more CPU functionality
 // #todo(dmp): ideally these would be exposed on the parent class, but we can't bind functions of parent classes but need to work on the interface
 // for sharing an instance data object with the super class
-DEFINE_NDI_DIRECT_FUNC_BINDER(UNiagaraDataInterfaceGrid3DCollection, GetWorldBBoxSize);
-DEFINE_NDI_DIRECT_FUNC_BINDER(UNiagaraDataInterfaceGrid3DCollection, GetCellSize);
-DEFINE_NDI_DIRECT_FUNC_BINDER(UNiagaraDataInterfaceGrid3DCollection, SetNumCells);
-DEFINE_NDI_DIRECT_FUNC_BINDER(UNiagaraDataInterfaceGrid3DCollection, GetNumCells);
-DEFINE_NDI_DIRECT_FUNC_BINDER(UNiagaraDataInterfaceGrid3DCollection, UnitToFloatIndex);
-DEFINE_NDI_DIRECT_FUNC_BINDER_WITH_PAYLOAD(UNiagaraDataInterfaceGrid3DCollection, GetAttributeIndex);
+DEFINE_NDI_DIRECT_FUNC_BINDER(UNiagaraDataInterfaceGrid3DCollection, VMGetWorldBBoxSize);
+DEFINE_NDI_DIRECT_FUNC_BINDER(UNiagaraDataInterfaceGrid3DCollection, VMGetCellSize);
+DEFINE_NDI_DIRECT_FUNC_BINDER(UNiagaraDataInterfaceGrid3DCollection, VMSetNumCells);
+DEFINE_NDI_DIRECT_FUNC_BINDER(UNiagaraDataInterfaceGrid3DCollection, VMGetNumCells);
+DEFINE_NDI_DIRECT_FUNC_BINDER(UNiagaraDataInterfaceGrid3DCollection, VMUnitToFloatIndex);
+DEFINE_NDI_DIRECT_FUNC_BINDER_WITH_PAYLOAD(UNiagaraDataInterfaceGrid3DCollection, VMGetAttributeIndex);
 void UNiagaraDataInterfaceGrid3DCollection::GetVMExternalFunction(const FVMExternalFunctionBindingInfo& BindingInfo, void* InstanceData, FVMExternalFunction &OutFunc)
 {
 	Super::GetVMExternalFunction(BindingInfo, InstanceData, OutFunc);
@@ -989,57 +856,57 @@ void UNiagaraDataInterfaceGrid3DCollection::GetVMExternalFunction(const FVMExter
 	if (BindingInfo.Name == UNiagaraDataInterfaceRWBase::WorldBBoxSizeFunctionName)
 	{
 		check(BindingInfo.GetNumInputs() == 1 && BindingInfo.GetNumOutputs() == 3);
-		NDI_FUNC_BINDER(UNiagaraDataInterfaceGrid3DCollection, GetWorldBBoxSize)::Bind(this, OutFunc);
+		NDI_FUNC_BINDER(UNiagaraDataInterfaceGrid3DCollection, VMGetWorldBBoxSize)::Bind(this, OutFunc);
 	}
 	else if (BindingInfo.Name == UNiagaraDataInterfaceRWBase::CellSizeFunctionName)
 	{
 		// #todo(dmp): this will override the base class definition for GetCellSize because the data interface instance data computes cell size
 		// it would be nice to refactor this so it can be part of the super class
 		check(BindingInfo.GetNumInputs() == 1 && BindingInfo.GetNumOutputs() == 3);
-		NDI_FUNC_BINDER(UNiagaraDataInterfaceGrid3DCollection, GetCellSize)::Bind(this, OutFunc);
+		NDI_FUNC_BINDER(UNiagaraDataInterfaceGrid3DCollection, VMGetCellSize)::Bind(this, OutFunc);
 	}
 	else if (BindingInfo.Name == SetNumCellsFunctionName)
 	{
 		check(BindingInfo.GetNumInputs() == 4 && BindingInfo.GetNumOutputs() == 1);
-		NDI_FUNC_BINDER(UNiagaraDataInterfaceGrid3DCollection, SetNumCells)::Bind(this, OutFunc);
+		NDI_FUNC_BINDER(UNiagaraDataInterfaceGrid3DCollection, VMSetNumCells)::Bind(this, OutFunc);
 	}
 	else if (BindingInfo.Name == NumCellsFunctionName)
 	{
 		check(BindingInfo.GetNumInputs() == 1 && BindingInfo.GetNumOutputs() == 3);
-		NDI_FUNC_BINDER(UNiagaraDataInterfaceGrid3DCollection, GetNumCells)::Bind(this, OutFunc);
+		NDI_FUNC_BINDER(UNiagaraDataInterfaceGrid3DCollection, VMGetNumCells)::Bind(this, OutFunc);
 	}
 	else if (BindingInfo.Name == UnitToFloatIndexFunctionName)
 	{		
 		check(BindingInfo.GetNumInputs() == 4 && BindingInfo.GetNumOutputs() == 3);
-		NDI_FUNC_BINDER(UNiagaraDataInterfaceGrid3DCollection, UnitToFloatIndex)::Bind(this, OutFunc);
+		NDI_FUNC_BINDER(UNiagaraDataInterfaceGrid3DCollection, VMUnitToFloatIndex)::Bind(this, OutFunc);
 	}
 	else if (BindingInfo.Name == GetVector4AttributeIndexFunctionName)
 	{
 		FName AttributeName = BindingInfo.FindSpecifier(NAME_Attribute)->Value;
 		check(BindingInfo.GetNumInputs() == 1 && BindingInfo.GetNumOutputs() == 1);
-		NDI_FUNC_BINDER(UNiagaraDataInterfaceGrid3DCollection, GetAttributeIndex)::Bind(this, OutFunc, AttributeName, 4);
+		NDI_FUNC_BINDER(UNiagaraDataInterfaceGrid3DCollection, VMGetAttributeIndex)::Bind(this, OutFunc, AttributeName, 4);
 	}
 	else if (BindingInfo.Name == GetVectorAttributeIndexFunctionName)
 	{
 		FName AttributeName = BindingInfo.FindSpecifier(NAME_Attribute)->Value;
 		check(BindingInfo.GetNumInputs() == 1 && BindingInfo.GetNumOutputs() == 1);
-		NDI_FUNC_BINDER(UNiagaraDataInterfaceGrid3DCollection, GetAttributeIndex)::Bind(this, OutFunc, AttributeName, 3);
+		NDI_FUNC_BINDER(UNiagaraDataInterfaceGrid3DCollection, VMGetAttributeIndex)::Bind(this, OutFunc, AttributeName, 3);
 	}
 	else if (BindingInfo.Name == GetVector2DAttributeIndexFunctionName)
 	{
 		FName AttributeName = BindingInfo.FindSpecifier(NAME_Attribute)->Value;
 		check(BindingInfo.GetNumInputs() == 1 && BindingInfo.GetNumOutputs() == 1);
-		NDI_FUNC_BINDER(UNiagaraDataInterfaceGrid3DCollection, GetAttributeIndex)::Bind(this, OutFunc, AttributeName, 2);
+		NDI_FUNC_BINDER(UNiagaraDataInterfaceGrid3DCollection, VMGetAttributeIndex)::Bind(this, OutFunc, AttributeName, 2);
 	}
 	else if (BindingInfo.Name == GetFloatAttributeIndexFunctionName)
 	{
 		FName AttributeName = BindingInfo.FindSpecifier(NAME_Attribute)->Value;
 		check(BindingInfo.GetNumInputs() == 1 && BindingInfo.GetNumOutputs() == 1);
-		NDI_FUNC_BINDER(UNiagaraDataInterfaceGrid3DCollection, GetAttributeIndex)::Bind(this, OutFunc, AttributeName, 1);
+		NDI_FUNC_BINDER(UNiagaraDataInterfaceGrid3DCollection, VMGetAttributeIndex)::Bind(this, OutFunc, AttributeName, 1);
 	}
 }
 
-void UNiagaraDataInterfaceGrid3DCollection::GetNumCells(FVectorVMExternalFunctionContext& Context)
+void UNiagaraDataInterfaceGrid3DCollection::VMGetNumCells(FVectorVMExternalFunctionContext& Context)
 {
 	VectorVM::FUserPtrHandler<FGrid3DCollectionRWInstanceData_GameThread> InstData(Context);
 
@@ -1111,7 +978,7 @@ void UNiagaraDataInterfaceGrid3DCollection::GetParameterDefinitionHLSL(const FNi
 
 	static const TCHAR *FormatDeclarations = TEXT(R"(
 		Texture3D<float> {GridName};
-		RWTexture3D<float> RW{OutputGridName};
+		RWTexture3D<float> {OutputGridName};
 		int3 {NumTiles};
 		float3 {OneOverNumTiles};
 		float3 {UnitClampMin};
@@ -1131,20 +998,20 @@ void UNiagaraDataInterfaceGrid3DCollection::GetParameterDefinitionHLSL(const FNi
 	const int AttributeInt4Count = FMath::Max(1, FMath::DivideAndRoundUp(AttributeCount, 4));
 
 	TMap<FString, FStringFormatArg> ArgsDeclarations = {				
-		{ TEXT("GridName"),    GridName + ParamInfo.DataInterfaceHLSLSymbol },
-		{ TEXT("SamplerName"),    SamplerName + ParamInfo.DataInterfaceHLSLSymbol },
-		{ TEXT("OutputGridName"),    OutputGridName + ParamInfo.DataInterfaceHLSLSymbol },
-		{ TEXT("NumTiles"),    NumTilesName + ParamInfo.DataInterfaceHLSLSymbol },
-		{ TEXT("OneOverNumTiles"), OneOverNumTilesName + ParamInfo.DataInterfaceHLSLSymbol },
-		{ TEXT("UnitClampMin"), UnitClampMinName + ParamInfo.DataInterfaceHLSLSymbol },
-		{ TEXT("UnitClampMax"), UnitClampMaxName + ParamInfo.DataInterfaceHLSLSymbol },
-
-		{ TEXT("UnitToUVName"), UNiagaraDataInterfaceRWBase::UnitToUVName + ParamInfo.DataInterfaceHLSLSymbol},
-		{ TEXT("AttributeIndicesName"), AttributeIndicesBaseName + ParamInfo.DataInterfaceHLSLSymbol},
-		{ TEXT("PerAttributeDataName"), PerAttributeDataName + ParamInfo.DataInterfaceHLSLSymbol},
-		{ TEXT("AttributeInt4Count"), AttributeInt4Count},
-		{ TEXT("NumAttributesName"), UNiagaraDataInterfaceRWBase::NumAttributesName + ParamInfo.DataInterfaceHLSLSymbol},
-		{ TEXT("NumNamedAttributesName"), UNiagaraDataInterfaceRWBase::NumNamedAttributesName + ParamInfo.DataInterfaceHLSLSymbol},
+		{ TEXT("GridName"),			ParamInfo.DataInterfaceHLSLSymbol + GridName},
+		{ TEXT("SamplerName"),		ParamInfo.DataInterfaceHLSLSymbol + SamplerName},
+		{ TEXT("OutputGridName"),	ParamInfo.DataInterfaceHLSLSymbol + OutputGridName},
+		{ TEXT("NumTiles"),			ParamInfo.DataInterfaceHLSLSymbol + NumTilesName},
+		{ TEXT("OneOverNumTiles"),	ParamInfo.DataInterfaceHLSLSymbol + OneOverNumTilesName},
+		{ TEXT("UnitClampMin"),		ParamInfo.DataInterfaceHLSLSymbol + UnitClampMinName},
+		{ TEXT("UnitClampMax"),		ParamInfo.DataInterfaceHLSLSymbol + UnitClampMaxName},
+
+		{ TEXT("UnitToUVName"),				ParamInfo.DataInterfaceHLSLSymbol + UNiagaraDataInterfaceRWBase::UnitToUVName},
+		{ TEXT("AttributeIndicesName"),		ParamInfo.DataInterfaceHLSLSymbol + TEXT("_") + AttributeIndicesBaseName},
+		{ TEXT("PerAttributeDataName"),		ParamInfo.DataInterfaceHLSLSymbol + TEXT("_") + PerAttributeDataName},
+		{ TEXT("AttributeInt4Count"),		AttributeInt4Count},
+		{ TEXT("NumAttributesName"),		ParamInfo.DataInterfaceHLSLSymbol + UNiagaraDataInterfaceRWBase::NumAttributesName},
+		{ TEXT("NumNamedAttributesName"),	ParamInfo.DataInterfaceHLSLSymbol + UNiagaraDataInterfaceRWBase::NumNamedAttributesName},
 	};
 	OutHLSL += FString::Format(FormatDeclarations, ArgsDeclarations);
 }
@@ -1175,25 +1042,25 @@ void UNiagaraDataInterfaceGrid3DCollection::WriteSetHLSL(const FNiagaraDataInter
 	FormatBounds += TEXT(R"(	
 					int CurAttributeIndex = In_AttributeIndex + i;
 					int3 TileOffset = {PerAttributeDataName}[CurAttributeIndex].xyz;
-					RW{OutputGrid}[int3(In_IndexX + TileOffset.x, In_IndexY + TileOffset.y, In_IndexZ + TileOffset.z)] = Val;
+					{OutputGrid}[int3(In_IndexX + TileOffset.x, In_IndexY + TileOffset.y, In_IndexZ + TileOffset.z)] = Val;
 				}
 			}
 		)");
 	TMap<FString, FStringFormatArg> ArgsBounds = {
-		{TEXT("FunctionName"), FunctionInfo.InstanceName},
-		{TEXT("OutputGrid"), OutputGridName + ParamInfo.DataInterfaceHLSLSymbol},
-		{TEXT("NumCellsName"), UNiagaraDataInterfaceRWBase::NumCellsName + ParamInfo.DataInterfaceHLSLSymbol},
-		{TEXT("UnitToUVName"), UNiagaraDataInterfaceRWBase::UnitToUVName + ParamInfo.DataInterfaceHLSLSymbol},
-		{TEXT("PerAttributeDataName"), PerAttributeDataName + ParamInfo.DataInterfaceHLSLSymbol},
-		{TEXT("AttributeIndicesName"), AttributeIndicesBaseName + ParamInfo.DataInterfaceHLSLSymbol},
-		{TEXT("AttributeIndexGroup"), FunctionInstanceIndex / 4},
-		{TEXT("AttributeIndexComponent"), VectorComponentNames[FunctionInstanceIndex % 4]},
-		{TEXT("NumChannels"), FString::FromInt(InNumChannels)},
-		{TEXT("NumChannelsVariableSuffix"), InNumChannels > 1 ? FString::FromInt(InNumChannels) : TEXT("")},
-		{TEXT("NumTiles"), NumTilesName + ParamInfo.DataInterfaceHLSLSymbol},
-		{TEXT("UnitToUVName"), UNiagaraDataInterfaceRWBase::UnitToUVName + ParamInfo.DataInterfaceHLSLSymbol},
-		{TEXT("NumCellsName"), UNiagaraDataInterfaceRWBase::NumCellsName + ParamInfo.DataInterfaceHLSLSymbol},
-		{TEXT("NumAttributesName"), UNiagaraDataInterfaceRWBase::NumAttributesName + ParamInfo.DataInterfaceHLSLSymbol},
+		{TEXT("FunctionName"),				FunctionInfo.InstanceName},
+		{TEXT("OutputGrid"),				ParamInfo.DataInterfaceHLSLSymbol + OutputGridName},
+		{TEXT("NumCellsName"),				ParamInfo.DataInterfaceHLSLSymbol + UNiagaraDataInterfaceRWBase::NumCellsName},
+		{TEXT("UnitToUVName"),				ParamInfo.DataInterfaceHLSLSymbol + UNiagaraDataInterfaceRWBase::UnitToUVName},
+		{TEXT("PerAttributeDataName"),		ParamInfo.DataInterfaceHLSLSymbol + TEXT("_") + PerAttributeDataName},
+		{TEXT("AttributeIndicesName"),		ParamInfo.DataInterfaceHLSLSymbol + TEXT("_") + AttributeIndicesBaseName},
+		{TEXT("AttributeIndexGroup"),		FunctionInstanceIndex / 4},
+		{TEXT("AttributeIndexComponent"),	VectorComponentNames[FunctionInstanceIndex % 4]},
+		{TEXT("NumChannels"),				FString::FromInt(InNumChannels)},
+		{TEXT("NumChannelsVariableSuffix"),	InNumChannels > 1 ? FString::FromInt(InNumChannels) : TEXT("")},
+		{TEXT("NumTiles"),					ParamInfo.DataInterfaceHLSLSymbol + NumTilesName},
+		{TEXT("UnitToUVName"),				ParamInfo.DataInterfaceHLSLSymbol + UNiagaraDataInterfaceRWBase::UnitToUVName},
+		{TEXT("NumCellsName"),				ParamInfo.DataInterfaceHLSLSymbol + UNiagaraDataInterfaceRWBase::NumCellsName},
+		{TEXT("NumAttributesName"),			ParamInfo.DataInterfaceHLSLSymbol + UNiagaraDataInterfaceRWBase::NumAttributesName},
 	};
 	OutHLSL += FString::Format(*FormatBounds, ArgsBounds);
 }
@@ -1225,19 +1092,19 @@ void UNiagaraDataInterfaceGrid3DCollection::WriteGetHLSL(const FNiagaraDataInter
 			}
 		)");
 	TMap<FString, FStringFormatArg> ArgsBounds = {
-		{TEXT("FunctionName"), FunctionInfo.InstanceName},
-		{TEXT("OutputGrid"), OutputGridName + ParamInfo.DataInterfaceHLSLSymbol},
-		{TEXT("Grid"), GridName + ParamInfo.DataInterfaceHLSLSymbol},
-		{TEXT("NumCellsName"), UNiagaraDataInterfaceRWBase::NumCellsName + ParamInfo.DataInterfaceHLSLSymbol},
-		{TEXT("UnitToUVName"), UNiagaraDataInterfaceRWBase::UnitToUVName + ParamInfo.DataInterfaceHLSLSymbol},
-		{TEXT("PerAttributeDataName"), PerAttributeDataName + ParamInfo.DataInterfaceHLSLSymbol},
-		{TEXT("AttributeIndicesName"), AttributeIndicesBaseName + ParamInfo.DataInterfaceHLSLSymbol},
-		{TEXT("AttributeIndexGroup"), FunctionInstanceIndex / 4},
-		{TEXT("AttributeIndexComponent"), VectorComponentNames[FunctionInstanceIndex % 4]},
-		{TEXT("NumChannels"), FString::FromInt(InNumChannels)},
-		{TEXT("NumChannelsVariableSuffix"), InNumChannels > 1 ? FString::FromInt(InNumChannels) : TEXT("")},
-		{TEXT("NumTiles"), NumTilesName + ParamInfo.DataInterfaceHLSLSymbol},
-		{TEXT("NumAttributesName"), UNiagaraDataInterfaceRWBase::NumAttributesName + ParamInfo.DataInterfaceHLSLSymbol},
+		{TEXT("FunctionName"),				FunctionInfo.InstanceName},
+		{TEXT("OutputGrid"),				ParamInfo.DataInterfaceHLSLSymbol + OutputGridName},
+		{TEXT("Grid"),						ParamInfo.DataInterfaceHLSLSymbol + GridName},
+		{TEXT("NumCellsName"),				ParamInfo.DataInterfaceHLSLSymbol + UNiagaraDataInterfaceRWBase::NumCellsName},
+		{TEXT("UnitToUVName"),				ParamInfo.DataInterfaceHLSLSymbol + UNiagaraDataInterfaceRWBase::UnitToUVName},
+		{TEXT("PerAttributeDataName"),		ParamInfo.DataInterfaceHLSLSymbol + TEXT("_") + PerAttributeDataName},
+		{TEXT("AttributeIndicesName"),		ParamInfo.DataInterfaceHLSLSymbol + TEXT("_") + AttributeIndicesBaseName},
+		{TEXT("AttributeIndexGroup"),		FunctionInstanceIndex / 4},
+		{TEXT("AttributeIndexComponent"),	VectorComponentNames[FunctionInstanceIndex % 4]},
+		{TEXT("NumChannels"),				FString::FromInt(InNumChannels)},
+		{TEXT("NumChannelsVariableSuffix"),	InNumChannels > 1 ? FString::FromInt(InNumChannels) : TEXT("")},
+		{TEXT("NumTiles"),					ParamInfo.DataInterfaceHLSLSymbol + NumTilesName},
+		{TEXT("NumAttributesName"),			ParamInfo.DataInterfaceHLSLSymbol + UNiagaraDataInterfaceRWBase::NumAttributesName},
 	};
 	OutHLSL += FString::Format(*FormatBounds, ArgsBounds);
 }
@@ -1272,22 +1139,22 @@ void UNiagaraDataInterfaceGrid3DCollection::WriteSampleHLSL(const FNiagaraDataIn
 		)");
 
 	TMap<FString, FStringFormatArg> ArgsBounds = {
-		{TEXT("FunctionName"), FunctionInfo.InstanceName},
-		{TEXT("Grid"), GridName + ParamInfo.DataInterfaceHLSLSymbol},
-		{TEXT("SamplerName"),    SamplerName + ParamInfo.DataInterfaceHLSLSymbol },
-		{TEXT("NumCellsName"), UNiagaraDataInterfaceRWBase::NumCellsName + ParamInfo.DataInterfaceHLSLSymbol},
-		{TEXT("UnitToUVName"), UNiagaraDataInterfaceRWBase::UnitToUVName + ParamInfo.DataInterfaceHLSLSymbol},
-		{TEXT("NumChannels"), FString::FromInt(InNumChannels)},
-		{TEXT("NumChannelsVariableSuffix"), InNumChannels > 1 ? FString::FromInt(InNumChannels) : TEXT("")},
-		{TEXT("AttributeIndicesName"), AttributeIndicesBaseName + ParamInfo.DataInterfaceHLSLSymbol},
-		{TEXT("PerAttributeDataName"), PerAttributeDataName + ParamInfo.DataInterfaceHLSLSymbol},
-		{TEXT("AttributeIndexGroup"), FunctionInstanceIndex / 4},
-		{TEXT("AttributeIndexComponent"), VectorComponentNames[FunctionInstanceIndex % 4]},
-		{TEXT("NumTiles"), NumTilesName + ParamInfo.DataInterfaceHLSLSymbol},
-		{TEXT("OneOverNumTiles"), OneOverNumTilesName + ParamInfo.DataInterfaceHLSLSymbol },
-		{TEXT("UnitClampMin"), UnitClampMinName + ParamInfo.DataInterfaceHLSLSymbol },
-		{TEXT("UnitClampMax"), UnitClampMaxName + ParamInfo.DataInterfaceHLSLSymbol },
-		{TEXT("NumAttributesName"), UNiagaraDataInterfaceRWBase::NumAttributesName + ParamInfo.DataInterfaceHLSLSymbol},
+		{TEXT("FunctionName"),				FunctionInfo.InstanceName},
+		{TEXT("Grid"),						ParamInfo.DataInterfaceHLSLSymbol + GridName},
+		{TEXT("SamplerName"),				ParamInfo.DataInterfaceHLSLSymbol + SamplerName },
+		{TEXT("NumCellsName"),				ParamInfo.DataInterfaceHLSLSymbol + UNiagaraDataInterfaceRWBase::NumCellsName},
+		{TEXT("UnitToUVName"),				ParamInfo.DataInterfaceHLSLSymbol + UNiagaraDataInterfaceRWBase::UnitToUVName},
+		{TEXT("NumChannels"),				FString::FromInt(InNumChannels)},
+		{TEXT("NumChannelsVariableSuffix"),	InNumChannels > 1 ? FString::FromInt(InNumChannels) : TEXT("")},
+		{TEXT("AttributeIndicesName"),		ParamInfo.DataInterfaceHLSLSymbol + TEXT("_") + AttributeIndicesBaseName},
+		{TEXT("PerAttributeDataName"),		ParamInfo.DataInterfaceHLSLSymbol + TEXT("_") + PerAttributeDataName},
+		{TEXT("AttributeIndexGroup"),		FunctionInstanceIndex / 4},
+		{TEXT("AttributeIndexComponent"),	VectorComponentNames[FunctionInstanceIndex % 4]},
+		{TEXT("NumTiles"),					ParamInfo.DataInterfaceHLSLSymbol + NumTilesName},
+		{TEXT("OneOverNumTiles"),			ParamInfo.DataInterfaceHLSLSymbol + OneOverNumTilesName},
+		{TEXT("UnitClampMin"),				ParamInfo.DataInterfaceHLSLSymbol + UnitClampMinName},
+		{TEXT("UnitClampMax"),				ParamInfo.DataInterfaceHLSLSymbol + UnitClampMaxName},
+		{TEXT("NumAttributesName"),			ParamInfo.DataInterfaceHLSLSymbol + UNiagaraDataInterfaceRWBase::NumAttributesName},
 	};
 	OutHLSL += FString::Format(*FormatBounds, ArgsBounds);
 }
@@ -1305,7 +1172,7 @@ void UNiagaraDataInterfaceGrid3DCollection::WriteAttributeGetIndexHLSL(const FNi
 
 	TMap<FString, FStringFormatArg> ArgsBounds = {
 		{TEXT("FunctionName"), FunctionInfo.InstanceName},
-		{TEXT("AttributeIndicesName"), AttributeIndicesBaseName + ParamInfo.DataInterfaceHLSLSymbol},
+		{TEXT("AttributeIndicesName"), ParamInfo.DataInterfaceHLSLSymbol + TEXT("_") + AttributeIndicesBaseName},
 		{TEXT("AttributeIndexGroup"), FunctionInstanceIndex / 4},
 		{TEXT("AttributeIndexComponent"), VectorComponentNames[FunctionInstanceIndex % 4]},
 	};
@@ -1372,21 +1239,21 @@ bool UNiagaraDataInterfaceGrid3DCollection::GetFunctionHLSL(const FNiagaraDataIn
 
 	TMap<FString, FStringFormatArg> ArgsBounds =
 	{
-		{TEXT("FunctionName"), FunctionInfo.InstanceName},
-		{TEXT("Grid"), GridName + ParamInfo.DataInterfaceHLSLSymbol},
-		{TEXT("OutputGrid"), OutputGridName + ParamInfo.DataInterfaceHLSLSymbol},
-		{TEXT("NumAttributes"), UNiagaraDataInterfaceRWBase::NumAttributesName + ParamInfo.DataInterfaceHLSLSymbol},
-		{TEXT("NumNamedAttributes"), UNiagaraDataInterfaceRWBase::NumNamedAttributesName + ParamInfo.DataInterfaceHLSLSymbol},
-		{TEXT("NumCells"), UNiagaraDataInterfaceRWBase::NumCellsName + ParamInfo.DataInterfaceHLSLSymbol},
-		{TEXT("UnitToUVName"), UNiagaraDataInterfaceRWBase::UnitToUVName + ParamInfo.DataInterfaceHLSLSymbol},
-		{TEXT("SamplerName"), SamplerName + ParamInfo.DataInterfaceHLSLSymbol},
-		{TEXT("NumTiles"), NumTilesName + ParamInfo.DataInterfaceHLSLSymbol},
-		{TEXT("OneOverNumTiles"), OneOverNumTilesName + ParamInfo.DataInterfaceHLSLSymbol },
-		{TEXT("UnitClampMin"), UnitClampMinName + ParamInfo.DataInterfaceHLSLSymbol },
-		{TEXT("UnitClampMax"), UnitClampMaxName + ParamInfo.DataInterfaceHLSLSymbol },
-		{TEXT("NumCellsName"), UNiagaraDataInterfaceRWBase::NumCellsName + ParamInfo.DataInterfaceHLSLSymbol},
-		{TEXT("PerAttributeDataName"), PerAttributeDataName + ParamInfo.DataInterfaceHLSLSymbol},
-		{TEXT("NumAttributesName"), UNiagaraDataInterfaceRWBase::NumAttributesName + ParamInfo.DataInterfaceHLSLSymbol},
+		{TEXT("FunctionName"),			FunctionInfo.InstanceName},
+		{TEXT("Grid"),					ParamInfo.DataInterfaceHLSLSymbol + GridName},
+		{TEXT("OutputGrid"),			ParamInfo.DataInterfaceHLSLSymbol + OutputGridName},
+		{TEXT("NumAttributes"),			ParamInfo.DataInterfaceHLSLSymbol + UNiagaraDataInterfaceRWBase::NumAttributesName},
+		{TEXT("NumNamedAttributes"),	ParamInfo.DataInterfaceHLSLSymbol + UNiagaraDataInterfaceRWBase::NumNamedAttributesName},
+		{TEXT("NumCells"),				ParamInfo.DataInterfaceHLSLSymbol + UNiagaraDataInterfaceRWBase::NumCellsName},
+		{TEXT("UnitToUVName"),			ParamInfo.DataInterfaceHLSLSymbol + UNiagaraDataInterfaceRWBase::UnitToUVName},
+		{TEXT("SamplerName"),			ParamInfo.DataInterfaceHLSLSymbol + SamplerName},
+		{TEXT("NumTiles"),				ParamInfo.DataInterfaceHLSLSymbol + NumTilesName},
+		{TEXT("OneOverNumTiles"),		ParamInfo.DataInterfaceHLSLSymbol + OneOverNumTilesName},
+		{TEXT("UnitClampMin"),			ParamInfo.DataInterfaceHLSLSymbol + UnitClampMinName},
+		{TEXT("UnitClampMax"),			ParamInfo.DataInterfaceHLSLSymbol + UnitClampMaxName},
+		{TEXT("NumCellsName"),			ParamInfo.DataInterfaceHLSLSymbol + UNiagaraDataInterfaceRWBase::NumCellsName},
+		{TEXT("PerAttributeDataName"),	ParamInfo.DataInterfaceHLSLSymbol + TEXT("_") + PerAttributeDataName},
+		{TEXT("NumAttributesName"),		ParamInfo.DataInterfaceHLSLSymbol + UNiagaraDataInterfaceRWBase::NumAttributesName},
 	};
 
 	if (FunctionInfo.DefinitionName == GetValueFunctionName || FunctionInfo.DefinitionName == GetPreviousValueAtIndexFunctionName)
@@ -1425,7 +1292,7 @@ bool UNiagaraDataInterfaceGrid3DCollection::GetFunctionHLSL(const FNiagaraDataIn
 				if ( In_AttributeIndex < {NumAttributesName} )
 				{
 					int3 TileOffset = {PerAttributeDataName}[In_AttributeIndex].xyz;
-					RW{OutputGrid}[int3(In_IndexX + TileOffset.x, In_IndexY + TileOffset.y, In_IndexZ + TileOffset.z)] = In_Value;
+					{OutputGrid}[int3(In_IndexX + TileOffset.x, In_IndexY + TileOffset.y, In_IndexZ + TileOffset.z)] = In_Value;
 				}
 			}
 		)");
@@ -1438,7 +1305,7 @@ bool UNiagaraDataInterfaceGrid3DCollection::GetFunctionHLSL(const FNiagaraDataIn
 			void {FunctionName}(int In_IndexX, int In_IndexY, int In_IndexZ, float In_Value, out int val)
 			{			
 				val = 0;
-				RW{OutputGrid}[int3(In_IndexX, In_IndexY, In_IndexZ)] = In_Value;
+				{OutputGrid}[int3(In_IndexX, In_IndexY, In_IndexZ)] = In_Value;
 			}
 		)");
 		OutHLSL += FString::Format(FormatBounds, ArgsBounds);
@@ -1454,7 +1321,7 @@ bool UNiagaraDataInterfaceGrid3DCollection::GetFunctionHLSL(const FNiagaraDataIn
 					int3 TileOffset = {PerAttributeDataName}[AttributeIndex].xyz;
 
 					float Val = {Grid}.Load(int4(In_IndexX + TileOffset.x, In_IndexY + TileOffset.y, In_IndexZ + TileOffset.z, 0));	
-					RW{OutputGrid}[int3(In_IndexX + TileOffset.x, In_IndexY + TileOffset.y, In_IndexZ + TileOffset.z)] = Val;
+					{OutputGrid}[int3(In_IndexX + TileOffset.x, In_IndexY + TileOffset.y, In_IndexZ + TileOffset.z)] = Val;
 				}
 			}
 		)");
@@ -1484,7 +1351,7 @@ bool UNiagaraDataInterfaceGrid3DCollection::GetFunctionHLSL(const FNiagaraDataIn
 						int3 TileOffset = {PerAttributeDataName}[AttributeIndex].xyz;
 
 						float Val = {Grid}.Load(int4(In_IndexX + TileOffset.x, In_IndexY + TileOffset.y, In_IndexZ + TileOffset.z, 0));	
-						RW{OutputGrid}[int3(In_IndexX + TileOffset.x, In_IndexY + TileOffset.y, In_IndexZ + TileOffset.z)] = Val;
+						{OutputGrid}[int3(In_IndexX + TileOffset.x, In_IndexY + TileOffset.y, In_IndexZ + TileOffset.z)] = Val;
 					}
 				}
 
@@ -1501,7 +1368,7 @@ bool UNiagaraDataInterfaceGrid3DCollection::GetFunctionHLSL(const FNiagaraDataIn
 					
 					int3 TileOffset = {PerAttributeDataName}[AttributeIndex].xyz;
 					float Val = {Grid}.Load(int4(In_IndexX + TileOffset.x, In_IndexY + TileOffset.y, In_IndexZ + TileOffset.z, 0));	
-					RW{OutputGrid}[int3(In_IndexX + TileOffset.x, In_IndexY + TileOffset.y, In_IndexZ + TileOffset.z)] = Val;
+					{OutputGrid}[int3(In_IndexX + TileOffset.x, In_IndexY + TileOffset.y, In_IndexZ + TileOffset.z)] = Val;
 				}
 */
 			}
@@ -1520,7 +1387,7 @@ bool UNiagaraDataInterfaceGrid3DCollection::GetFunctionHLSL(const FNiagaraDataIn
 
 					int3 TileOffset = {PerAttributeDataName}[AttributeIndex].xyz;
 
-					RW{OutputGrid}[int3(In_IndexX + TileOffset.x, In_IndexY + TileOffset.y, In_IndexZ + TileOffset.z)] = In_Value;
+					{OutputGrid}[int3(In_IndexX + TileOffset.x, In_IndexY + TileOffset.y, In_IndexZ + TileOffset.z)] = In_Value;
 				}
 			}
 		)");
@@ -1641,6 +1508,156 @@ bool UNiagaraDataInterfaceGrid3DCollection::GetFunctionHLSL(const FNiagaraDataIn
 }
 #endif
 
+void UNiagaraDataInterfaceGrid3DCollection::BuildShaderParameters(FNiagaraShaderParametersBuilder& ShaderParametersBuilder) const
+{
+	ShaderParametersBuilder.AddNestedStruct<FNDIGrid3DShaderParameters>();
+
+	const int32 IndirectionTableSize = FMath::Max(FMath::DivideAndRoundUp(ShaderParametersBuilder.GetGeneratedFunctions().Num(), 4), 1);
+	ShaderParametersBuilder.AddLooseParamArray<FIntVector4>(*UNiagaraDataInterfaceGrid3DCollection::AttributeIndicesBaseName, IndirectionTableSize);
+}
+
+void UNiagaraDataInterfaceGrid3DCollection::SetShaderParameters(const FNiagaraDataInterfaceSetShaderParametersContext& Context) const
+{
+	FNiagaraDataInterfaceProxyGrid3DCollectionProxy& DIProxy = Context.GetProxy<FNiagaraDataInterfaceProxyGrid3DCollectionProxy>();
+	FGrid3DCollectionRWInstanceData_RenderThread* ProxyData = DIProxy.SystemInstancesToProxyData_RT.Find(Context.GetSystemInstanceID());
+	FGrid3DCollectionRWInstanceData_RenderThread* OriginalProxyData = ProxyData;
+	check(ProxyData);
+
+	if (ProxyData->OtherProxy != nullptr)
+	{
+		FNiagaraDataInterfaceProxyGrid3DCollectionProxy* OtherGrid3DProxy = static_cast<FNiagaraDataInterfaceProxyGrid3DCollectionProxy*>(ProxyData->OtherProxy);
+		ProxyData = OtherGrid3DProxy->SystemInstancesToProxyData_RT.Find(Context.GetSystemInstanceID());
+		check(ProxyData);
+	}
+
+	const FNiagaraDataInterfaceParametersCS_Grid3DCollection& ShaderStorage = Context.GetShaderStorage<FNiagaraDataInterfaceParametersCS_Grid3DCollection>();
+	if (OriginalProxyData->AttributeIndices.Num() == 0 && ShaderStorage.AttributeNames.Num() > 0)
+	{
+		OriginalProxyData->AttributeIndices.SetNumZeroed(Align(ShaderStorage.AttributeNames.Num(), 4));
+
+		// TODO handle mismatched types!
+		for (int32 i = 0; i < ShaderStorage.AttributeNames.Num(); i++)
+		{
+			const int32 FoundIdx = ProxyData->Vars.Find(ShaderStorage.AttributeNames[i]);
+			check(ShaderStorage.AttributeNames.Num() == ShaderStorage.AttributeChannelCount.Num());
+			check(ProxyData->Offsets.Num() == ProxyData->VarComponents.Num());
+			check(ProxyData->Offsets.Num() == ProxyData->Vars.Num());
+			if (ProxyData->Offsets.IsValidIndex(FoundIdx) && ShaderStorage.AttributeChannelCount[i] == ProxyData->VarComponents[FoundIdx])
+			{
+				OriginalProxyData->AttributeIndices[i] = ProxyData->Offsets[FoundIdx];
+			}
+			else
+			{
+				OriginalProxyData->AttributeIndices[i] = -1; // We may need to protect against this in the hlsl as this might underflow an array lookup if used incorrectly.
+			}
+		}
+	}
+
+
+	// Generate per-attribute data
+	if (ProxyData->PerAttributeData.NumBytes == 0)
+	{
+		TResourceArray<FVector4f> PerAttributeData;
+		PerAttributeData.AddUninitialized((ProxyData->TotalNumAttributes * 2) + 1);
+		for (int32 iAttribute = 0; iAttribute < ProxyData->TotalNumAttributes; ++iAttribute)
+		{
+			const FIntVector AttributeTileIndex(iAttribute % ProxyData->NumTiles.X, (iAttribute / ProxyData->NumTiles.X) % ProxyData->NumTiles.Y, iAttribute / (ProxyData->NumTiles.X * ProxyData->NumTiles.Y));
+			PerAttributeData[iAttribute] = FVector4f(
+				AttributeTileIndex.X * ProxyData->NumCells.X,
+				AttributeTileIndex.Y * ProxyData->NumCells.Y,
+				AttributeTileIndex.Z * ProxyData->NumCells.Z,
+				0
+			);
+			PerAttributeData[iAttribute + ProxyData->TotalNumAttributes] = FVector4f(
+				(1.0f / ProxyData->NumTiles.X) * float(AttributeTileIndex.X),
+				(1.0f / ProxyData->NumTiles.Y) * float(AttributeTileIndex.Y),
+				(1.0f / ProxyData->NumTiles.Z) * float(AttributeTileIndex.Z),
+				0.0f
+			);
+		}
+		PerAttributeData[ProxyData->TotalNumAttributes * 2] = FVector4f(65535, 65535, 65535, 65535);
+		ProxyData->PerAttributeData.Initialize(TEXT("Grid3D::PerAttributeData"), sizeof(FVector4f), PerAttributeData.Num(), EPixelFormat::PF_A32B32G32R32F, BUF_Static, &PerAttributeData);
+	}
+
+	// Set parameters
+	const FVector3f HalfPixelOffset = FVector3f(0.5f / ProxyData->NumCells.X, 0.5f / ProxyData->NumCells.Y, 0.5f / ProxyData->NumCells.Z);
+
+	FNDIGrid3DShaderParameters* Parameters = Context.GetParameterNestedStruct<FNDIGrid3DShaderParameters>();
+	Parameters->NumAttributes		= ProxyData->TotalNumAttributes;
+	Parameters->NumNamedAttributes	= ProxyData->TotalNumNamedAttributes;
+	Parameters->UnitToUV			= FVector3f(1.0f) / FVector3f(ProxyData->NumCells);
+	Parameters->NumCells			= ProxyData->NumCells;
+	Parameters->CellSize			= FVector3f(ProxyData->CellSize);
+	Parameters->NumTiles			= ProxyData->NumTiles;
+	Parameters->OneOverNumTiles		= FVector3f(1.0f) / FVector3f(ProxyData->NumTiles);
+	Parameters->UnitClampMin		= HalfPixelOffset;
+	Parameters->UnitClampMax		= FVector3f::OneVector - HalfPixelOffset;
+	Parameters->WorldBBoxSize		= FVector3f(ProxyData->WorldBBoxSize);
+
+	Parameters->SamplerState = TStaticSamplerState<SF_Bilinear, AM_Clamp, AM_Clamp, AM_Clamp>::GetRHI();
+	Parameters->Grid = ProxyData->CurrentData ? ProxyData->CurrentData->GridBuffer.SRV.GetReference() : FNiagaraRenderer::GetDummyTextureReadBuffer3D();
+	if (Context.IsResourceBound(&Parameters->OutputGrid))
+	{
+		if (Context.IsOutputStage() && ProxyData->DestinationData)
+		{
+			Parameters->OutputGrid = ProxyData->DestinationData->GridBuffer.UAV;
+		}
+		else
+		{
+			Parameters->OutputGrid = Context.GetComputeDispatchInterface().GetEmptyUAVFromPool(Context.GetRHICmdList(), PF_R32_FLOAT, ENiagaraEmptyUAVType::Texture3D);
+		}
+	}
+	Parameters->PerAttributeData = ProxyData->PerAttributeData.SRV;
+
+	if (OriginalProxyData->AttributeIndices.Num() > 0)
+	{
+		const int NumAttributesVector4s = FMath::Max(FMath::DivideAndRoundUp(ShaderStorage.AttributeNames.Num(), 4), 1);
+		TArrayView<FIntVector4> AttributeIndices = Context.GetParameterLooseArray<FIntVector4>(NumAttributesVector4s);
+		check(OriginalProxyData->AttributeIndices.Num() * OriginalProxyData->AttributeIndices.GetTypeSize() == AttributeIndices.Num() * AttributeIndices.GetTypeSize());
+		FMemory::Memcpy(AttributeIndices.GetData(), OriginalProxyData->AttributeIndices.GetData(), AttributeIndices.Num() * AttributeIndices.GetTypeSize());
+	}
+}
+
+FNiagaraDataInterfaceParametersCS* UNiagaraDataInterfaceGrid3DCollection::CreateShaderStorage(const FNiagaraDataInterfaceGPUParamInfo& ParameterInfo, const FShaderParameterMap& ParameterMap) const
+{
+	FNiagaraDataInterfaceParametersCS_Grid3DCollection* ShaderStorage = new FNiagaraDataInterfaceParametersCS_Grid3DCollection();
+
+	// Gather up all the attribute names referenced. Note that there may be multiple in the list of the same name,
+	// but we only deal with this by the number of bound methods.
+	const int32 NumFuncs = ParameterInfo.GeneratedFunctions.Num();
+	ShaderStorage->AttributeNames.Reserve(NumFuncs);
+	ShaderStorage->AttributeChannelCount.Reserve(NumFuncs);
+
+	for (int32 FuncIdx=0; FuncIdx < NumFuncs; ++FuncIdx)
+	{
+		const FNiagaraDataInterfaceGeneratedFunction& Func = ParameterInfo.GeneratedFunctions[FuncIdx];
+		static const FName NAME_Attribute("Attribute");
+		const FName* AttributeName = Func.FindSpecifierValue(NAME_Attribute);
+		if (AttributeName != nullptr)
+		{
+			int32 ComponentCount = UNiagaraDataInterfaceGrid3DCollection::GetComponentCountFromFuncName(Func.DefinitionName);
+			ShaderStorage->AttributeNames.Add(*AttributeName);
+			ShaderStorage->AttributeChannelCount.Add(ComponentCount);
+		}
+		else
+		{
+			ShaderStorage->AttributeNames.Add(FName());
+			ShaderStorage->AttributeChannelCount.Add(INDEX_NONE);
+		}
+	}
+
+	ShaderStorage->AttributeNames.Shrink();
+	ShaderStorage->AttributeChannelCount.Shrink();
+
+	return ShaderStorage;
+}
+
+const FTypeLayoutDesc* UNiagaraDataInterfaceGrid3DCollection::GetShaderStorageType() const
+{
+	return &StaticGetTypeLayoutDesc<FNiagaraDataInterfaceParametersCS_Grid3DCollection>();
+}
+
+
 #if WITH_EDITOR
 bool UNiagaraDataInterfaceGrid3DCollection::GenerateIterationSourceNamespaceReadAttributesHLSL(FNiagaraDataInterfaceGPUParamInfo& DIInstanceInfo, const FNiagaraVariable& IterationSourceVar, TConstArrayView<FNiagaraVariable> InArguments, TConstArrayView<FNiagaraVariable> InAttributes, TConstArrayView<FString> InAttributeHLSLNames, bool bInSetToDefaults, bool bSpawnOnly, bool bPartialWrites, TArray<FText>& OutErrors, FString& OutHLSL) const
 {
@@ -2436,7 +2453,7 @@ void UNiagaraDataInterfaceGrid3DCollection::GetTextureSize(const UNiagaraCompone
 	SizeZ = Grid3DInstanceData->NumCells.Z;
 }
 
-void UNiagaraDataInterfaceGrid3DCollection::GetWorldBBoxSize(FVectorVMExternalFunctionContext& Context)
+void UNiagaraDataInterfaceGrid3DCollection::VMGetWorldBBoxSize(FVectorVMExternalFunctionContext& Context)
 {
 	VectorVM::FUserPtrHandler<FGrid3DCollectionRWInstanceData_GameThread> InstData(Context);
 	VectorVM::FExternalFuncRegisterHandler<float> OutWorldBoundsX(Context);
@@ -2451,7 +2468,7 @@ void UNiagaraDataInterfaceGrid3DCollection::GetWorldBBoxSize(FVectorVMExternalFu
 	}
 }
 
-void UNiagaraDataInterfaceGrid3DCollection::SetNumCells(FVectorVMExternalFunctionContext& Context)
+void UNiagaraDataInterfaceGrid3DCollection::VMSetNumCells(FVectorVMExternalFunctionContext& Context)
 {
 	// This should only be called from a system or emitter script due to a need for only setting up initially.
 	VectorVM::FUserPtrHandler<FGrid3DCollectionRWInstanceData_GameThread> InstData(Context);
@@ -2488,7 +2505,7 @@ void UNiagaraDataInterfaceGrid3DCollection::SetNumCells(FVectorVMExternalFunctio
 	}
 }
 
-void UNiagaraDataInterfaceGrid3DCollection::UnitToFloatIndex(FVectorVMExternalFunctionContext& Context)
+void UNiagaraDataInterfaceGrid3DCollection::VMUnitToFloatIndex(FVectorVMExternalFunctionContext& Context)
 {
 	// This should only be called from a system or emitter script due to a need for only setting up initially.
 	VectorVM::FUserPtrHandler<FGrid3DCollectionRWInstanceData_GameThread> InstData(Context);
@@ -2593,7 +2610,7 @@ bool UNiagaraDataInterfaceGrid3DCollection::PerInstanceTickPostSimulate(void* Pe
 	return false;
 }
 
-void UNiagaraDataInterfaceGrid3DCollection::GetCellSize(FVectorVMExternalFunctionContext& Context)
+void UNiagaraDataInterfaceGrid3DCollection::VMGetCellSize(FVectorVMExternalFunctionContext& Context)
 {
 	VectorVM::FUserPtrHandler<FGrid3DCollectionRWInstanceData_GameThread> InstData(Context);
 	VectorVM::FExternalFuncRegisterHandler<float> OutCellSizeX(Context);
@@ -2608,7 +2625,7 @@ void UNiagaraDataInterfaceGrid3DCollection::GetCellSize(FVectorVMExternalFunctio
 	}
 }
 
-void UNiagaraDataInterfaceGrid3DCollection::GetAttributeIndex(FVectorVMExternalFunctionContext& Context, const FName& InName, int32 NumChannels)
+void UNiagaraDataInterfaceGrid3DCollection::VMGetAttributeIndex(FVectorVMExternalFunctionContext& Context, const FName& InName, int32 NumChannels)
 {
 	VectorVM::FUserPtrHandler<FGrid3DCollectionRWInstanceData_GameThread> InstData(Context);
 	VectorVM::FExternalFuncRegisterHandler<int> OutIndex(Context);
@@ -2707,31 +2724,6 @@ void FNiagaraDataInterfaceProxyGrid3DCollectionProxy::PreStage(FRHICommandList&
 {
 	FGrid3DCollectionRWInstanceData_RenderThread* ProxyData = SystemInstancesToProxyData_RT.Find(Context.SystemInstanceID);
 
-	// Generate per-attribute data
-	if (ProxyData->PerAttributeData.NumBytes == 0)
-	{
-		TResourceArray<FVector4f> PerAttributeData;
-		PerAttributeData.AddUninitialized((ProxyData->TotalNumAttributes * 2) + 1);
-		for (int32 iAttribute = 0; iAttribute < ProxyData->TotalNumAttributes; ++iAttribute)
-		{
-			const FIntVector AttributeTileIndex(iAttribute % ProxyData->NumTiles.X, (iAttribute / ProxyData->NumTiles.X) % ProxyData->NumTiles.Y, iAttribute / (ProxyData->NumTiles.X * ProxyData->NumTiles.Y));
-			PerAttributeData[iAttribute] = FVector4f(
-				AttributeTileIndex.X * ProxyData->NumCells.X,
-				AttributeTileIndex.Y * ProxyData->NumCells.Y,
-				AttributeTileIndex.Z * ProxyData->NumCells.Z,
-				0
-			);
-			PerAttributeData[iAttribute + ProxyData->TotalNumAttributes] = FVector4f(
-				(1.0f / ProxyData->NumTiles.X) * float(AttributeTileIndex.X),
-				(1.0f / ProxyData->NumTiles.Y) * float(AttributeTileIndex.Y),
-				(1.0f / ProxyData->NumTiles.Z) * float(AttributeTileIndex.Z),
-				0.0f
-			);
-		}
-		PerAttributeData[ProxyData->TotalNumAttributes * 2] = FVector4f(65535, 65535, 65535, 65535);
-		ProxyData->PerAttributeData.Initialize(TEXT("Grid3D::PerAttributeData"), sizeof(FVector4f), PerAttributeData.Num(), EPixelFormat::PF_A32B32G32R32F, BUF_Static, &PerAttributeData);
-	}
-
 	// #todo(dmp): Context doesnt need to specify if a stage is output or not since we moved pre/post stage to the DI itself.  Not sure which design is better for the future
 	if (Context.IsOutputStage)
 	{
diff --git a/Engine/Plugins/FX/Niagara/Source/Niagara/Private/NiagaraDataInterfaceGrid3DCollectionReader.cpp b/Engine/Plugins/FX/Niagara/Source/Niagara/Private/NiagaraDataInterfaceGrid3DCollectionReader.cpp
--- a/Engine/Plugins/FX/Niagara/Source/Niagara/Private/NiagaraDataInterfaceGrid3DCollectionReader.cpp
+++ b/Engine/Plugins/FX/Niagara/Source/Niagara/Private/NiagaraDataInterfaceGrid3DCollectionReader.cpp
@@ -17,8 +17,6 @@
 
 #define LOCTEXT_NAMESPACE "NiagaraDataInterfaceGrid3DCollectionReader"
 
-IMPLEMENT_NIAGARA_DI_PARAMETER(UNiagaraDataInterfaceGrid3DCollectionReader, FNiagaraDataInterfaceParametersCS_Grid3DCollection);
-
 UNiagaraDataInterfaceGrid3DCollectionReader::UNiagaraDataInterfaceGrid3DCollectionReader(FObjectInitializer const& ObjectInitializer)
 	: Super(ObjectInitializer)
 {
diff --git a/Engine/Plugins/FX/Niagara/Source/Niagara/Private/NiagaraDataInterfaceNeighborGrid3D.cpp b/Engine/Plugins/FX/Niagara/Source/Niagara/Private/NiagaraDataInterfaceNeighborGrid3D.cpp
--- a/Engine/Plugins/FX/Niagara/Source/Niagara/Private/NiagaraDataInterfaceNeighborGrid3D.cpp
+++ b/Engine/Plugins/FX/Niagara/Source/Niagara/Private/NiagaraDataInterfaceNeighborGrid3D.cpp
@@ -9,11 +9,11 @@
 
 #define LOCTEXT_NAMESPACE "NiagaraDataInterfaceNeighborGrid3D"
 
-static const FString MaxNeighborsPerCellName(TEXT("MaxNeighborsPerCellValue_"));
-static const FString ParticleNeighborsName(TEXT("ParticleNeighbors_"));
-static const FString ParticleNeighborCountName(TEXT("ParticleNeighborCount_"));
-static const FString OutputParticleNeighborsName(TEXT("OutputParticleNeighbors_"));
-static const FString OutputParticleNeighborCountName(TEXT("OutputParticleNeighborCount_"));
+static const FString MaxNeighborsPerCellName(TEXT("_MaxNeighborsPerCellValue"));
+static const FString ParticleNeighborsName(TEXT("_ParticleNeighbors"));
+static const FString ParticleNeighborCountName(TEXT("_ParticleNeighborCount"));
+static const FString OutputParticleNeighborsName(TEXT("_OutputParticleNeighbors"));
+static const FString OutputParticleNeighborCountName(TEXT("_OutputParticleNeighborCount"));
 
 
 const FName UNiagaraDataInterfaceNeighborGrid3D::SetNumCellsFunctionName("SetNumCells");
@@ -44,18 +44,18 @@ struct FNiagaraDataInterfaceParametersCS_NeighborGrid3D : public FNiagaraDataInt
 public:
 	void Bind(const FNiagaraDataInterfaceGPUParamInfo& ParameterInfo, const class FShaderParameterMap& ParameterMap)
 	{
-		NumCellsParam.Bind(ParameterMap, *(UNiagaraDataInterfaceRWBase::NumCellsName + ParameterInfo.DataInterfaceHLSLSymbol));
-		UnitToUVParam.Bind(ParameterMap, *(UNiagaraDataInterfaceRWBase::UnitToUVName + ParameterInfo.DataInterfaceHLSLSymbol));
-		CellSizeParam.Bind(ParameterMap, *(UNiagaraDataInterfaceRWBase::CellSizeName + ParameterInfo.DataInterfaceHLSLSymbol));
+		NumCellsParam.Bind(ParameterMap, *(ParameterInfo.DataInterfaceHLSLSymbol + UNiagaraDataInterfaceRWBase::NumCellsName));
+		UnitToUVParam.Bind(ParameterMap, *(ParameterInfo.DataInterfaceHLSLSymbol + UNiagaraDataInterfaceRWBase::UnitToUVName));
+		CellSizeParam.Bind(ParameterMap, *(ParameterInfo.DataInterfaceHLSLSymbol + UNiagaraDataInterfaceRWBase::CellSizeName));
 
-		MaxNeighborsPerCellParam.Bind(ParameterMap, *(MaxNeighborsPerCellName + ParameterInfo.DataInterfaceHLSLSymbol));
-		WorldBBoxSizeParam.Bind(ParameterMap, *(UNiagaraDataInterfaceRWBase::WorldBBoxSizeName + ParameterInfo.DataInterfaceHLSLSymbol));
+		MaxNeighborsPerCellParam.Bind(ParameterMap, *(ParameterInfo.DataInterfaceHLSLSymbol + MaxNeighborsPerCellName));
+		WorldBBoxSizeParam.Bind(ParameterMap, *(ParameterInfo.DataInterfaceHLSLSymbol + UNiagaraDataInterfaceRWBase::WorldBBoxSizeName));
 		
-		ParticleNeighborsGridParam.Bind(ParameterMap,  *(ParticleNeighborsName + ParameterInfo.DataInterfaceHLSLSymbol));
-		ParticleNeighborCountGridParam.Bind(ParameterMap,  *(ParticleNeighborCountName + ParameterInfo.DataInterfaceHLSLSymbol));
+		ParticleNeighborsGridParam.Bind(ParameterMap,  *(ParameterInfo.DataInterfaceHLSLSymbol + ParticleNeighborsName));
+		ParticleNeighborCountGridParam.Bind(ParameterMap,  *(ParameterInfo.DataInterfaceHLSLSymbol + ParticleNeighborCountName));
 
-		OutputParticleNeighborsGridParam.Bind(ParameterMap, *(OutputParticleNeighborsName + ParameterInfo.DataInterfaceHLSLSymbol));
-		OutputParticleNeighborCountGridParam.Bind(ParameterMap, *(OutputParticleNeighborCountName + ParameterInfo.DataInterfaceHLSLSymbol));
+		OutputParticleNeighborsGridParam.Bind(ParameterMap, *(ParameterInfo.DataInterfaceHLSLSymbol + OutputParticleNeighborsName));
+		OutputParticleNeighborCountGridParam.Bind(ParameterMap, *(ParameterInfo.DataInterfaceHLSLSymbol + OutputParticleNeighborCountName));
 	}
 
 	// #todo(dmp): make resource transitions batched
@@ -428,11 +428,11 @@ void UNiagaraDataInterfaceNeighborGrid3D::GetParameterDefinitionHLSL(const FNiag
 		RWBuffer<int> RW{OutputParticleNeighborCountName};
 	)");
 	TMap<FString, FStringFormatArg> ArgsDeclarations = {
-		{ TEXT("MaxNeighborsPerCellName"),  MaxNeighborsPerCellName + ParamInfo.DataInterfaceHLSLSymbol },
-		{ TEXT("ParticleNeighborsName"),    ParticleNeighborsName + ParamInfo.DataInterfaceHLSLSymbol },
-		{ TEXT("ParticleNeighborCountName"),    ParticleNeighborCountName + ParamInfo.DataInterfaceHLSLSymbol },
-		{ TEXT("OutputParticleNeighborsName"),    OutputParticleNeighborsName + ParamInfo.DataInterfaceHLSLSymbol },
-		{ TEXT("OutputParticleNeighborCountName"),    OutputParticleNeighborCountName + ParamInfo.DataInterfaceHLSLSymbol },
+		{ TEXT("MaxNeighborsPerCellName"),  ParamInfo.DataInterfaceHLSLSymbol + MaxNeighborsPerCellName },
+		{ TEXT("ParticleNeighborsName"),    ParamInfo.DataInterfaceHLSLSymbol + ParticleNeighborsName },
+		{ TEXT("ParticleNeighborCountName"),    ParamInfo.DataInterfaceHLSLSymbol + ParticleNeighborCountName },
+		{ TEXT("OutputParticleNeighborsName"),    ParamInfo.DataInterfaceHLSLSymbol + OutputParticleNeighborsName },
+		{ TEXT("OutputParticleNeighborCountName"),    ParamInfo.DataInterfaceHLSLSymbol + OutputParticleNeighborCountName },
 	};
 	OutHLSL += FString::Format(FormatDeclarations, ArgsDeclarations);
 }
@@ -457,8 +457,8 @@ bool UNiagaraDataInterfaceNeighborGrid3D::GetFunctionHLSL(const FNiagaraDataInte
 		TMap<FString, FStringFormatArg> FormatArgs =
 		{
 			{TEXT("FunctionName"), FunctionInfo.InstanceName},
-			{TEXT("NumCellsName"), UNiagaraDataInterfaceRWBase::NumCellsName + ParamInfo.DataInterfaceHLSLSymbol},
-			{TEXT("UnitToUVName"), UNiagaraDataInterfaceRWBase::UnitToUVName + ParamInfo.DataInterfaceHLSLSymbol},
+			{TEXT("NumCellsName"), ParamInfo.DataInterfaceHLSLSymbol + UNiagaraDataInterfaceRWBase::NumCellsName},
+			{TEXT("UnitToUVName"), ParamInfo.DataInterfaceHLSLSymbol + UNiagaraDataInterfaceRWBase::UnitToUVName},
 		};
 		OutHLSL += FString::Format(FormatHLSL, FormatArgs);
 	}
@@ -472,7 +472,7 @@ bool UNiagaraDataInterfaceNeighborGrid3D::GetFunctionHLSL(const FNiagaraDataInte
 		)");
 		TMap<FString, FStringFormatArg> ArgsSample = {
 			{TEXT("FunctionName"), FunctionInfo.InstanceName},
-			{TEXT("MaxNeighborsPerCellName"), MaxNeighborsPerCellName + ParamInfo.DataInterfaceHLSLSymbol},
+			{TEXT("MaxNeighborsPerCellName"), ParamInfo.DataInterfaceHLSLSymbol + MaxNeighborsPerCellName},
 
 		};
 		OutHLSL += FString::Format(FormatSample, ArgsSample);
@@ -488,9 +488,9 @@ bool UNiagaraDataInterfaceNeighborGrid3D::GetFunctionHLSL(const FNiagaraDataInte
 		)");
 		TMap<FString, FStringFormatArg> ArgsBounds = {
 			{TEXT("FunctionName"), FunctionInfo.InstanceName},
-			{TEXT("MaxNeighborsPerCellName"), MaxNeighborsPerCellName + ParamInfo.DataInterfaceHLSLSymbol},
-			{TEXT("NumCellsName"), UNiagaraDataInterfaceRWBase::NumCellsName + ParamInfo.DataInterfaceHLSLSymbol},
-			{TEXT("UnitToUVName"), UNiagaraDataInterfaceRWBase::UnitToUVName + ParamInfo.DataInterfaceHLSLSymbol},
+			{TEXT("MaxNeighborsPerCellName"), ParamInfo.DataInterfaceHLSLSymbol + MaxNeighborsPerCellName},
+			{TEXT("NumCellsName"), ParamInfo.DataInterfaceHLSLSymbol + UNiagaraDataInterfaceRWBase::NumCellsName},
+			{TEXT("UnitToUVName"), ParamInfo.DataInterfaceHLSLSymbol + UNiagaraDataInterfaceRWBase::UnitToUVName},
 		};
 		OutHLSL += FString::Format(FormatBounds, ArgsBounds);
 		return true;
@@ -505,7 +505,7 @@ bool UNiagaraDataInterfaceNeighborGrid3D::GetFunctionHLSL(const FNiagaraDataInte
 		)");
 		TMap<FString, FStringFormatArg> ArgsBounds = {
 			{TEXT("FunctionName"), FunctionInfo.InstanceName},
-			{TEXT("ParticleNeighbors"), ParticleNeighborsName + ParamInfo.DataInterfaceHLSLSymbol},
+			{TEXT("ParticleNeighbors"), ParamInfo.DataInterfaceHLSLSymbol + ParticleNeighborsName},
 		};
 		OutHLSL += FString::Format(FormatBounds, ArgsBounds);
 		return true;
@@ -521,7 +521,7 @@ bool UNiagaraDataInterfaceNeighborGrid3D::GetFunctionHLSL(const FNiagaraDataInte
 		)");
 		TMap<FString, FStringFormatArg> ArgsBounds = {
 			{TEXT("FunctionName"), FunctionInfo.InstanceName},
-			{TEXT("OutputParticleNeighbors"), OutputParticleNeighborsName + ParamInfo.DataInterfaceHLSLSymbol},
+			{TEXT("OutputParticleNeighbors"), ParamInfo.DataInterfaceHLSLSymbol + OutputParticleNeighborsName},
 		};
 		OutHLSL += FString::Format(FormatBounds, ArgsBounds);
 		return true;
@@ -536,7 +536,7 @@ bool UNiagaraDataInterfaceNeighborGrid3D::GetFunctionHLSL(const FNiagaraDataInte
 		)");
 		TMap<FString, FStringFormatArg> ArgsBounds = {
 			{TEXT("FunctionName"), FunctionInfo.InstanceName},
-			{TEXT("ParticleNeighborCount"), ParticleNeighborCountName + ParamInfo.DataInterfaceHLSLSymbol},
+			{TEXT("ParticleNeighborCount"), ParamInfo.DataInterfaceHLSLSymbol + ParticleNeighborCountName},
 		};
 		OutHLSL += FString::Format(FormatBounds, ArgsBounds);
 		return true;
@@ -551,7 +551,7 @@ bool UNiagaraDataInterfaceNeighborGrid3D::GetFunctionHLSL(const FNiagaraDataInte
 		)");
 		TMap<FString, FStringFormatArg> ArgsBounds = {
 			{TEXT("FunctionName"), FunctionInfo.InstanceName},
-			{TEXT("OutputParticleNeighborCount"), OutputParticleNeighborCountName + ParamInfo.DataInterfaceHLSLSymbol},
+			{TEXT("OutputParticleNeighborCount"), ParamInfo.DataInterfaceHLSLSymbol + OutputParticleNeighborCountName},
 		};
 		OutHLSL += FString::Format(FormatBounds, ArgsBounds);
 		return true;
diff --git a/Engine/Plugins/FX/Niagara/Source/Niagara/Private/NiagaraDataInterfaceRW.cpp b/Engine/Plugins/FX/Niagara/Source/Niagara/Private/NiagaraDataInterfaceRW.cpp
--- a/Engine/Plugins/FX/Niagara/Source/Niagara/Private/NiagaraDataInterfaceRW.cpp
+++ b/Engine/Plugins/FX/Niagara/Source/Niagara/Private/NiagaraDataInterfaceRW.cpp
@@ -8,12 +8,12 @@
 #define LOCTEXT_NAMESPACE "NiagaraDataInterfaceRW"
 
 // Global HLSL variable base names, used by HLSL.
-const FString UNiagaraDataInterfaceRWBase::NumAttributesName(TEXT("NumAttributes_"));
-const FString UNiagaraDataInterfaceRWBase::NumNamedAttributesName(TEXT("NumNamedAttributes_"));
-const FString UNiagaraDataInterfaceRWBase::NumCellsName(TEXT("NumCells_"));
-const FString UNiagaraDataInterfaceRWBase::UnitToUVName(TEXT("UnitToUV_"));
-const FString UNiagaraDataInterfaceRWBase::CellSizeName(TEXT("CellSize_"));
-const FString UNiagaraDataInterfaceRWBase::WorldBBoxSizeName(TEXT("WorldBBoxSize_"));
+const FString UNiagaraDataInterfaceRWBase::NumAttributesName(TEXT("_NumAttributes"));
+const FString UNiagaraDataInterfaceRWBase::NumNamedAttributesName(TEXT("_NumNamedAttributes"));
+const FString UNiagaraDataInterfaceRWBase::NumCellsName(TEXT("_NumCells"));
+const FString UNiagaraDataInterfaceRWBase::UnitToUVName(TEXT("_UnitToUV"));
+const FString UNiagaraDataInterfaceRWBase::CellSizeName(TEXT("_CellSize"));
+const FString UNiagaraDataInterfaceRWBase::WorldBBoxSizeName(TEXT("_WorldBBoxSize"));
 
 // Global VM function names, also used by the shaders code generation methods.
 const FName UNiagaraDataInterfaceRWBase::NumCellsFunctionName("GetNumCells");
@@ -266,10 +266,10 @@ void UNiagaraDataInterfaceGrid3D::GetParameterDefinitionHLSL(const FNiagaraDataI
 		float3 {WorldBBoxSizeName};
 	)");
 	TMap<FString, FStringFormatArg> ArgsDeclarations = {
-		{ TEXT("NumCellsName"), UNiagaraDataInterfaceRWBase::NumCellsName + ParamInfo.DataInterfaceHLSLSymbol },
-		{ TEXT("UnitToUVName"), UNiagaraDataInterfaceRWBase::UnitToUVName + ParamInfo.DataInterfaceHLSLSymbol },
-		{ TEXT("CellSizeName"), UNiagaraDataInterfaceRWBase::CellSizeName + ParamInfo.DataInterfaceHLSLSymbol },
-		{ TEXT("WorldBBoxSizeName"),    UNiagaraDataInterfaceRWBase::WorldBBoxSizeName + ParamInfo.DataInterfaceHLSLSymbol },
+		{ TEXT("NumCellsName"), ParamInfo.DataInterfaceHLSLSymbol + UNiagaraDataInterfaceRWBase::NumCellsName },
+		{ TEXT("UnitToUVName"), ParamInfo.DataInterfaceHLSLSymbol + UNiagaraDataInterfaceRWBase::UnitToUVName },
+		{ TEXT("CellSizeName"), ParamInfo.DataInterfaceHLSLSymbol + UNiagaraDataInterfaceRWBase::CellSizeName },
+		{ TEXT("WorldBBoxSizeName"), ParamInfo.DataInterfaceHLSLSymbol + UNiagaraDataInterfaceRWBase::WorldBBoxSizeName },
 	};
 	OutHLSL += FString::Format(FormatDeclarations, ArgsDeclarations);
 }
@@ -278,10 +278,10 @@ bool UNiagaraDataInterfaceGrid3D::GetFunctionHLSL(const FNiagaraDataInterfaceGPU
 {
 	TMap<FString, FStringFormatArg> ArgsDeclarations = {
 		{ TEXT("FunctionName"), FunctionInfo.InstanceName},
-		{ TEXT("NumCellsName"), UNiagaraDataInterfaceRWBase::NumCellsName + ParamInfo.DataInterfaceHLSLSymbol },
-		{ TEXT("UnitToUVName"), UNiagaraDataInterfaceRWBase::UnitToUVName + ParamInfo.DataInterfaceHLSLSymbol },
-		{ TEXT("CellSizeName"), UNiagaraDataInterfaceRWBase::CellSizeName + ParamInfo.DataInterfaceHLSLSymbol },
-		{ TEXT("WorldBBoxSizeName"), UNiagaraDataInterfaceRWBase::WorldBBoxSizeName + ParamInfo.DataInterfaceHLSLSymbol },
+		{ TEXT("NumCellsName"), ParamInfo.DataInterfaceHLSLSymbol + UNiagaraDataInterfaceRWBase::NumCellsName },
+		{ TEXT("UnitToUVName"), ParamInfo.DataInterfaceHLSLSymbol + UNiagaraDataInterfaceRWBase::UnitToUVName },
+		{ TEXT("CellSizeName"), ParamInfo.DataInterfaceHLSLSymbol + UNiagaraDataInterfaceRWBase::CellSizeName },
+		{ TEXT("WorldBBoxSizeName"), ParamInfo.DataInterfaceHLSLSymbol + UNiagaraDataInterfaceRWBase::WorldBBoxSizeName },
 	};
 
 	if (FunctionInfo.DefinitionName == UNiagaraDataInterfaceRWBase::WorldBBoxSizeFunctionName)
@@ -736,10 +736,10 @@ void UNiagaraDataInterfaceGrid2D::GetParameterDefinitionHLSL(const FNiagaraDataI
 		float2 {WorldBBoxSizeName};
 	)");
 	TMap<FString, FStringFormatArg> ArgsDeclarations = {
-		{ TEXT("NumCellsName"), UNiagaraDataInterfaceRWBase::NumCellsName + ParamInfo.DataInterfaceHLSLSymbol },
-		{ TEXT("UnitToUVName"), UNiagaraDataInterfaceRWBase::UnitToUVName + ParamInfo.DataInterfaceHLSLSymbol },
-		{ TEXT("CellSizeName"), UNiagaraDataInterfaceRWBase::CellSizeName + ParamInfo.DataInterfaceHLSLSymbol },		
-		{ TEXT("WorldBBoxSizeName"),    UNiagaraDataInterfaceRWBase::WorldBBoxSizeName + ParamInfo.DataInterfaceHLSLSymbol },
+		{ TEXT("NumCellsName"), ParamInfo.DataInterfaceHLSLSymbol + UNiagaraDataInterfaceRWBase::NumCellsName },
+		{ TEXT("UnitToUVName"), ParamInfo.DataInterfaceHLSLSymbol + UNiagaraDataInterfaceRWBase::UnitToUVName },
+		{ TEXT("CellSizeName"), ParamInfo.DataInterfaceHLSLSymbol + UNiagaraDataInterfaceRWBase::CellSizeName },
+		{ TEXT("WorldBBoxSizeName"),    ParamInfo.DataInterfaceHLSLSymbol + UNiagaraDataInterfaceRWBase::WorldBBoxSizeName },
 	};
 	OutHLSL += FString::Format(FormatDeclarations, ArgsDeclarations);
 }
@@ -748,10 +748,10 @@ bool UNiagaraDataInterfaceGrid2D::GetFunctionHLSL(const FNiagaraDataInterfaceGPU
 {
 	TMap<FString, FStringFormatArg> ArgsDeclarations = {
 		{ TEXT("FunctionName"), FunctionInfo.InstanceName},
-		{ TEXT("NumCellsName"), UNiagaraDataInterfaceRWBase::NumCellsName + ParamInfo.DataInterfaceHLSLSymbol },
-		{ TEXT("UnitToUVName"), UNiagaraDataInterfaceRWBase::UnitToUVName + ParamInfo.DataInterfaceHLSLSymbol },
-		{ TEXT("CellSizeName"), UNiagaraDataInterfaceRWBase::CellSizeName + ParamInfo.DataInterfaceHLSLSymbol },		
-		{ TEXT("WorldBBoxSizeName"), UNiagaraDataInterfaceRWBase::WorldBBoxSizeName + ParamInfo.DataInterfaceHLSLSymbol },
+		{ TEXT("NumCellsName"), ParamInfo.DataInterfaceHLSLSymbol + UNiagaraDataInterfaceRWBase::NumCellsName },
+		{ TEXT("UnitToUVName"), ParamInfo.DataInterfaceHLSLSymbol + UNiagaraDataInterfaceRWBase::UnitToUVName },
+		{ TEXT("CellSizeName"), ParamInfo.DataInterfaceHLSLSymbol + UNiagaraDataInterfaceRWBase::CellSizeName },		
+		{ TEXT("WorldBBoxSizeName"), ParamInfo.DataInterfaceHLSLSymbol + UNiagaraDataInterfaceRWBase::WorldBBoxSizeName },
 	};
 
 	if (FunctionInfo.DefinitionName == UNiagaraDataInterfaceRWBase::WorldBBoxSizeFunctionName)
diff --git a/Engine/Plugins/FX/Niagara/Source/Niagara/Private/NiagaraDataInterfaceRasterizationGrid3D.cpp b/Engine/Plugins/FX/Niagara/Source/Niagara/Private/NiagaraDataInterfaceRasterizationGrid3D.cpp
--- a/Engine/Plugins/FX/Niagara/Source/Niagara/Private/NiagaraDataInterfaceRasterizationGrid3D.cpp
+++ b/Engine/Plugins/FX/Niagara/Source/Niagara/Private/NiagaraDataInterfaceRasterizationGrid3D.cpp
@@ -9,14 +9,14 @@
 
 #define LOCTEXT_NAMESPACE "NiagaraDataInterfaceRasterizationGrid3D"
 
-static const FString IntGridName(TEXT("IntGrid_"));
-static const FString OutputIntGridName(TEXT("OutputIntGrid_"));
-static const FString PrecisionName(TEXT("Precision_"));
+static const FString IntGridName(TEXT("_IntGrid"));
+static const FString OutputIntGridName(TEXT("_OutputIntGrid"));
+static const FString PrecisionName(TEXT("_Precision"));
 
 const FName UNiagaraDataInterfaceRasterizationGrid3D::SetNumCellsFunctionName("SetNumCells");
 const FName UNiagaraDataInterfaceRasterizationGrid3D::SetFloatResetValueFunctionName("SetFloatResetValue");
 
-const FString UNiagaraDataInterfaceRasterizationGrid3D::PerAttributeDataName(TEXT("PerAttributeDataName_"));
+const FString UNiagaraDataInterfaceRasterizationGrid3D::PerAttributeDataName(TEXT("_PerAttributeDataName"));
 
 // Global VM function names, also used by the shaders code generation methods.
 
@@ -58,16 +58,16 @@ struct FNiagaraDataInterfaceParametersCS_RasterizationGrid3D : public FNiagaraDa
 public:
 	void Bind(const FNiagaraDataInterfaceGPUParamInfo& ParameterInfo, const class FShaderParameterMap& ParameterMap)
 	{
-		NumAttributesParam.Bind(ParameterMap, *(UNiagaraDataInterfaceRWBase::NumAttributesName + ParameterInfo.DataInterfaceHLSLSymbol));
-		NumCellsParam.Bind(ParameterMap, *(UNiagaraDataInterfaceRWBase::NumCellsName + ParameterInfo.DataInterfaceHLSLSymbol));
-		UnitToUVParam.Bind(ParameterMap, *(UNiagaraDataInterfaceRWBase::UnitToUVName + ParameterInfo.DataInterfaceHLSLSymbol));		
-		PrecisionParam.Bind(ParameterMap, *(PrecisionName + ParameterInfo.DataInterfaceHLSLSymbol));
+		NumAttributesParam.Bind(ParameterMap, *(ParameterInfo.DataInterfaceHLSLSymbol + UNiagaraDataInterfaceRWBase::NumAttributesName));
+		NumCellsParam.Bind(ParameterMap, *(ParameterInfo.DataInterfaceHLSLSymbol + UNiagaraDataInterfaceRWBase::NumCellsName));
+		UnitToUVParam.Bind(ParameterMap, *(ParameterInfo.DataInterfaceHLSLSymbol + UNiagaraDataInterfaceRWBase::UnitToUVName));
+		PrecisionParam.Bind(ParameterMap, *(ParameterInfo.DataInterfaceHLSLSymbol + PrecisionName));
 				
-		IntGridParam.Bind(ParameterMap,  *(IntGridName + ParameterInfo.DataInterfaceHLSLSymbol));
+		IntGridParam.Bind(ParameterMap,  *(ParameterInfo.DataInterfaceHLSLSymbol + IntGridName));
 		
-		OutputIntGridParam.Bind(ParameterMap, *(OutputIntGridName + ParameterInfo.DataInterfaceHLSLSymbol));
+		OutputIntGridParam.Bind(ParameterMap, *(ParameterInfo.DataInterfaceHLSLSymbol + OutputIntGridName));
 
-		PerAttributeDataParam.Bind(ParameterMap, *(UNiagaraDataInterfaceRasterizationGrid3D::PerAttributeDataName + ParameterInfo.DataInterfaceHLSLSymbol));
+		PerAttributeDataParam.Bind(ParameterMap, *(ParameterInfo.DataInterfaceHLSLSymbol + UNiagaraDataInterfaceRasterizationGrid3D::PerAttributeDataName));
 	}
 
 	// #todo(dmp): make resource transitions batched
@@ -607,11 +607,11 @@ void UNiagaraDataInterfaceRasterizationGrid3D::GetParameterDefinitionHLSL(const
 		int {NumAttributesName};
 	)");
 	TMap<FString, FStringFormatArg> ArgsDeclarations = {				
-		{ TEXT("IntGridName"),    IntGridName + ParamInfo.DataInterfaceHLSLSymbol },		
-		{ TEXT("OutputIntGridName"),    OutputIntGridName + ParamInfo.DataInterfaceHLSLSymbol },
-		{ TEXT("Precision"), PrecisionName + ParamInfo.DataInterfaceHLSLSymbol},
-		{ TEXT("PerAttributeDataName"), PerAttributeDataName + ParamInfo.DataInterfaceHLSLSymbol},
-		{ TEXT("NumAttributesName"), UNiagaraDataInterfaceRWBase::NumAttributesName + ParamInfo.DataInterfaceHLSLSymbol},
+		{ TEXT("IntGridName"),    ParamInfo.DataInterfaceHLSLSymbol + IntGridName },
+		{ TEXT("OutputIntGridName"),    ParamInfo.DataInterfaceHLSLSymbol + OutputIntGridName},
+		{ TEXT("Precision"), ParamInfo.DataInterfaceHLSLSymbol + PrecisionName},
+		{ TEXT("PerAttributeDataName"), ParamInfo.DataInterfaceHLSLSymbol + PerAttributeDataName},
+		{ TEXT("NumAttributesName"), ParamInfo.DataInterfaceHLSLSymbol + UNiagaraDataInterfaceRWBase::NumAttributesName},
 	};
 	OutHLSL += FString::Format(FormatDeclarations, ArgsDeclarations);
 
@@ -626,9 +626,9 @@ void UNiagaraDataInterfaceRasterizationGrid3D::GetParameterDefinitionHLSL(const
 			)");
 		TMap<FString, FStringFormatArg> FormatArgs =
 		{
-			{ TEXT("IntToFloatFunction"), IntToFloatFunctionName.ToString() + ParamInfo.DataInterfaceHLSLSymbol},
-			{ TEXT("FloatToIntFunction"), FloatToIntFunctionName.ToString() + ParamInfo.DataInterfaceHLSLSymbol},
-			{ TEXT("Precision"), PrecisionName + ParamInfo.DataInterfaceHLSLSymbol},
+			{ TEXT("IntToFloatFunction"), ParamInfo.DataInterfaceHLSLSymbol + IntToFloatFunctionName.ToString()},
+			{ TEXT("FloatToIntFunction"), ParamInfo.DataInterfaceHLSLSymbol + FloatToIntFunctionName.ToString()},
+			{ TEXT("Precision"), ParamInfo.DataInterfaceHLSLSymbol + PrecisionName},
 		};
 		OutHLSL += FString::Format(FormatHLSL, FormatArgs);
 	}
@@ -643,9 +643,9 @@ void UNiagaraDataInterfaceRasterizationGrid3D::GetParameterDefinitionHLSL(const
 
 		TMap<FString, FStringFormatArg> FormatArgs =
 		{
-			{ TEXT("IntToFloatFunction"), IntToFloatFunctionName.ToString() + ParamInfo.DataInterfaceHLSLSymbol},
-			{ TEXT("FloatToIntFunction"), FloatToIntFunctionName.ToString() + ParamInfo.DataInterfaceHLSLSymbol},
-			{ TEXT("Precision"), PrecisionName + ParamInfo.DataInterfaceHLSLSymbol},
+			{ TEXT("IntToFloatFunction"), ParamInfo.DataInterfaceHLSLSymbol + IntToFloatFunctionName.ToString()},
+			{ TEXT("FloatToIntFunction"), ParamInfo.DataInterfaceHLSLSymbol + FloatToIntFunctionName.ToString()},
+			{ TEXT("Precision"), ParamInfo.DataInterfaceHLSLSymbol + PrecisionName},
 		};
 		OutHLSL += FString::Format(FormatHLSL, FormatArgs);
 	}
@@ -659,15 +659,15 @@ bool UNiagaraDataInterfaceRasterizationGrid3D::GetFunctionHLSL(const FNiagaraDat
 	TMap<FString, FStringFormatArg> ArgsBounds =
 	{
 		{ TEXT("FunctionName"), FunctionInfo.InstanceName},
-		{ TEXT("IntGrid"), IntGridName + ParamInfo.DataInterfaceHLSLSymbol},
-		{ TEXT("OutputIntGrid"), OutputIntGridName + ParamInfo.DataInterfaceHLSLSymbol},		
-		{TEXT("NumAttributesName"), UNiagaraDataInterfaceRWBase::NumAttributesName + ParamInfo.DataInterfaceHLSLSymbol},
-		{ TEXT("NumCellsName"), UNiagaraDataInterfaceRWBase::NumCellsName + ParamInfo.DataInterfaceHLSLSymbol},
-		{ TEXT("UnitToUVName"), UNiagaraDataInterfaceRWBase::UnitToUVName + ParamInfo.DataInterfaceHLSLSymbol},		
-		{ TEXT("IntToFloatFunctionName"), IntToFloatFunctionName.ToString() + ParamInfo.DataInterfaceHLSLSymbol},
-		{ TEXT("FloatToIntFunctionName"), FloatToIntFunctionName.ToString() + ParamInfo.DataInterfaceHLSLSymbol},
-		{ TEXT("Precision"), PrecisionName + ParamInfo.DataInterfaceHLSLSymbol},
-		{TEXT("PerAttributeDataName"), PerAttributeDataName + ParamInfo.DataInterfaceHLSLSymbol},
+		{ TEXT("IntGrid"), ParamInfo.DataInterfaceHLSLSymbol + IntGridName},
+		{ TEXT("OutputIntGrid"), ParamInfo.DataInterfaceHLSLSymbol + OutputIntGridName},
+		{TEXT("NumAttributesName"), ParamInfo.DataInterfaceHLSLSymbol + UNiagaraDataInterfaceRWBase::NumAttributesName},
+		{ TEXT("NumCellsName"), ParamInfo.DataInterfaceHLSLSymbol + UNiagaraDataInterfaceRWBase::NumCellsName},
+		{ TEXT("UnitToUVName"), ParamInfo.DataInterfaceHLSLSymbol + UNiagaraDataInterfaceRWBase::UnitToUVName},
+		{ TEXT("IntToFloatFunctionName"), ParamInfo.DataInterfaceHLSLSymbol + IntToFloatFunctionName.ToString()},
+		{ TEXT("FloatToIntFunctionName"), ParamInfo.DataInterfaceHLSLSymbol + FloatToIntFunctionName.ToString()},
+		{ TEXT("Precision"), ParamInfo.DataInterfaceHLSLSymbol + PrecisionName},
+		{TEXT("PerAttributeDataName"), ParamInfo.DataInterfaceHLSLSymbol + PerAttributeDataName},
 	};
 
 	if (ParentRet)
diff --git a/Engine/Source/Runtime/Core/Private/UObject/UE5MainStreamObjectVersion.cpp b/Engine/Source/Runtime/Core/Private/UObject/UE5MainStreamObjectVersion.cpp
--- a/Engine/Source/Runtime/Core/Private/UObject/UE5MainStreamObjectVersion.cpp
+++ b/Engine/Source/Runtime/Core/Private/UObject/UE5MainStreamObjectVersion.cpp
@@ -14,7 +14,7 @@ TMap<FGuid, FGuid> FUE5MainStreamObjectVersion::GetSystemGuids()
 	SystemGuids.Add(DevGuids.MATERIALSHADERMAP_DERIVEDDATA_VER, FGuid("65C01D817C9A4EEFAE9E988D41A1F3DD"));
 	SystemGuids.Add(DevGuids.NANITE_DERIVEDDATA_VER, FGuid("7A3EA93546624597BEBE2ADFDE3A7628"));
 	SystemGuids.Add(DevGuids.NIAGARASHADERMAP_DERIVEDDATA_VER, FGuid("8A37C45D24F2423CBE5F8F371DE33575"));
-	SystemGuids.Add(DevGuids.Niagara_LatestScriptCompileVersion, FGuid("CEB479E7F34CCB429B078D7827C9CBBC"));
+	SystemGuids.Add(DevGuids.Niagara_LatestScriptCompileVersion, FGuid("D3F577819D21054287503193E522CDA4"));
 	SystemGuids.Add(DevGuids.SkeletalMeshDerivedDataVersion, FGuid("FB7B2FB546DE49B2A2F88B7F69763CEE"));
 	SystemGuids.Add(DevGuids.STATICMESH_DERIVEDDATA_VER, FGuid("98F7D79A5811013825E75BBCC41ED3E9"));
 	SystemGuids.Add(DevGuids.POSESEARCHDB_DERIVEDDATA_VER, FGuid("4E595C2AC5E947D6BA9ABC874353E5BC"));
