- Add toggle to enable / disable HUD rendering
- Allow for pausing / stepping simulations without text overlayed

#rnx

diff --git a/Engine/Plugins/FX/Niagara/Source/Niagara/Private/NiagaraDebugHud.cpp b/Engine/Plugins/FX/Niagara/Source/Niagara/Private/NiagaraDebugHud.cpp
--- a/Engine/Plugins/FX/Niagara/Source/Niagara/Private/NiagaraDebugHud.cpp
+++ b/Engine/Plugins/FX/Niagara/Source/Niagara/Private/NiagaraDebugHud.cpp
@@ -1045,13 +1045,17 @@ void FNiagaraDebugHud::Draw(FNiagaraWorldManager* WorldManager, UCanvas* Canvas,
 	float CurrTime = WorldManager->GetWorld()->GetRealTimeSeconds();
 	DeltaSeconds = CurrTime - LastDrawTime;
 
-	// Draw in world components
-	DrawComponents(WorldManager, Canvas);
+	// We may want to use the HUD to pause / step but not actually display anything on screen
+	if (Settings.bHudRenderingEnabled)
+	{
+		// Draw in world components
+		DrawComponents(WorldManager, Canvas);
 
-	DrawDebugGeomerty(WorldManager, Canvas);
+		DrawDebugGeomerty(WorldManager, Canvas);
 
-	// Draw overview
-	DrawOverview(WorldManager, Canvas->Canvas);
+		// Draw overview
+		DrawOverview(WorldManager, Canvas->Canvas);
+	}
 
 	// Scrub any gpu cached emitters we haven't used in a while
 	{
diff --git a/Engine/Plugins/FX/Niagara/Source/Niagara/Public/NiagaraDebuggerCommon.h b/Engine/Plugins/FX/Niagara/Source/Niagara/Public/NiagaraDebuggerCommon.h
--- a/Engine/Plugins/FX/Niagara/Source/Niagara/Public/NiagaraDebuggerCommon.h
+++ b/Engine/Plugins/FX/Niagara/Source/Niagara/Public/NiagaraDebuggerCommon.h
@@ -431,6 +431,10 @@ struct NIAGARA_API FNiagaraDebugHUDSettingsData
 	UPROPERTY(EditAnywhere, Category = "Debug General", meta = (DisplayName = "Debug HUD Enabled"))
 	bool bHudEnabled = true;
 
+	/** Master control for HUD rendering. */
+	UPROPERTY(EditAnywhere, Category = "Debug General", meta = (DisplayName = "Debug HUD Rendering Enabled"))
+	bool bHudRenderingEnabled = true;
+
 	/**
 	When enabled all Niagara systems that pass the filter will have the simulation data buffers validation.
 	i.e. we will look for NaN or other invalidate data  inside it
