Fixed a Niagara gpu tick problem introduced by CL 20052164

diff --git a/Engine/Plugins/FX/Niagara/Source/NiagaraEditor/Private/NiagaraEditorModule.cpp b/Engine/Plugins/FX/Niagara/Source/NiagaraEditor/Private/NiagaraEditorModule.cpp
--- a/Engine/Plugins/FX/Niagara/Source/NiagaraEditor/Private/NiagaraEditorModule.cpp
+++ b/Engine/Plugins/FX/Niagara/Source/NiagaraEditor/Private/NiagaraEditorModule.cpp
@@ -170,6 +170,9 @@ static FAutoConsoleVariableRef CVarShowNiagaraDeveloperWindows(
 	ECVF_Default
 	);
 
+// this is required for gpu script compilation ticks
+static FNiagaraShaderQueueTickable NiagaraShaderQueueProcessor;
+
 //////////////////////////////////////////////////////////////////////////
 
 class FNiagaraEditorOnlyDataUtilities : public INiagaraEditorOnlyDataUtilities
