Fixed non unity build issue

#rnx

diff --git a/Engine/Plugins/FX/Niagara/Source/NiagaraEditor/Private/Toolkits/SystemToolkitModes/NiagaraSystemToolkitMode_Default.cpp b/Engine/Plugins/FX/Niagara/Source/NiagaraEditor/Private/Toolkits/SystemToolkitModes/NiagaraSystemToolkitMode_Default.cpp
--- a/Engine/Plugins/FX/Niagara/Source/NiagaraEditor/Private/Toolkits/SystemToolkitModes/NiagaraSystemToolkitMode_Default.cpp
+++ b/Engine/Plugins/FX/Niagara/Source/NiagaraEditor/Private/Toolkits/SystemToolkitModes/NiagaraSystemToolkitMode_Default.cpp
@@ -8,6 +8,8 @@
 #include "NiagaraEditorModule.h"
 #include "NiagaraEditorSettings.h"
 
+#include "Framework/MultiBox/MultiBoxBuilder.h"
+
 #define LOCTEXT_NAMESPACE "NiagaraSystemToolkitMode_Default"
 
 FNiagaraSystemToolkitMode_Default::FNiagaraSystemToolkitMode_Default(TWeakPtr<FNiagaraSystemToolkit> InSystemToolkit) : FNiagaraSystemToolkitModeBase(FNiagaraSystemToolkit::DefaultModeName, InSystemToolkit)
