- Convert MeshRendererInfo DI to Shader Parameters
- Convert VectorField DI to Shader Parameters
- Convert CurlNoise DI to Shader Parameters
- Convert Camera DI to Shader Parameters
- Convert Occlusion DI to Shader Parameters

diff --git a/Engine/Plugins/FX/Niagara/Shaders/Private/NiagaraDataInterfaceCamera.ush b/Engine/Plugins/FX/Niagara/Shaders/Private/NiagaraDataInterfaceCamera.ush
--- a/Engine/Plugins/FX/Niagara/Shaders/Private/NiagaraDataInterfaceCamera.ush
+++ b/Engine/Plugins/FX/Niagara/Shaders/Private/NiagaraDataInterfaceCamera.ush
@@ -1,27 +1,33 @@
 // Copyright Epic Games, Inc. All Rights Reserved.
 
-#pragma once
+float3 {ParameterName}_SystemLWCTile;
 
-struct FNDICameraContext
+void GetFieldOfView_{ParameterName}(out float Out_FieldOfViewAngle)
 {
-	float3 SystemLWCTile;
-};
+	Out_FieldOfViewAngle = degrees(View.FieldOfViewWideAngles.x);
+}
+
+void GetTAAJitter_{ParameterName}(out float2 Out_CurrentJitter, out float2 Out_PreviousJitter)
+{
+	Out_CurrentJitter = View.TemporalAAJitter.xy;
+	Out_PreviousJitter = View.TemporalAAJitter.zw;
+}
 
-void DICamera_GetCameraProperties(in FNDICameraContext DIContext, out float3 Out_CameraPositionWorld, out float3 Out_ViewForwardVector, out float3 Out_ViewUpVector, out float3 Out_ViewRightVector)
+void GetCameraPropertiesCPUASC47GPU_{ParameterName}(out float3 Out_CameraPositionWorld, out float3 Out_ViewForwardVector, out float3 Out_ViewUpVector, out float3 Out_ViewRightVector)
 {
 	FLWCVector3 CameraPos = PrimaryView.WorldCameraOrigin;
-	LWCSetTile(CameraPos, LWCGetTile(CameraPos) - DIContext.SystemLWCTile); // convert to simulation space
+	LWCSetTile(CameraPos, LWCGetTile(CameraPos) - {ParameterName}_SystemLWCTile); // convert to simulation space
 	Out_CameraPositionWorld = LWCToFloat(CameraPos);
 	Out_ViewForwardVector = View.ViewForward;
 	Out_ViewUpVector = View.ViewUp;
 	Out_ViewRightVector = View.ViewRight;
 }
 
-void DICamera_GetViewProperties(in FNDICameraContext DIContext, out float3 Out_ViewPositionWorld, out float3 Out_ViewForwardVector, out float3 Out_ViewUpVector, out float3 Out_ViewRightVector, out float4 Out_ViewSizeAndInverseSize,
+void GetViewPropertiesGPU_{ParameterName}(out float3 Out_ViewPositionWorld, out float3 Out_ViewForwardVector, out float3 Out_ViewUpVector, out float3 Out_ViewRightVector, out float4 Out_ViewSizeAndInverseSize,
 	out float4 Out_ScreenToViewSpace, out float2 Out_Current_TAAJitter, out float2 Out_Previous_TAAJitter, out float3 Out_PreViewTranslation, out float4 Out_BufferSizeAndInverseSize, out float2 Out_ViewportOffset)
 {
 	FLWCVector3 ViewPos = PrimaryView.WorldViewOrigin;
-	LWCSetTile(ViewPos, LWCGetTile(ViewPos) - DIContext.SystemLWCTile); // convert to simulation space
+	LWCSetTile(ViewPos, LWCGetTile(ViewPos) - {ParameterName}_SystemLWCTile); // convert to simulation space
 	Out_ViewPositionWorld.xyz = LWCToFloat(ViewPos);
 	Out_ViewForwardVector = View.ViewForward;
 	Out_ViewUpVector = View.ViewUp;
@@ -35,7 +41,7 @@ void DICamera_GetViewProperties(in FNDICameraContext DIContext, out float3 Out_V
 	Out_ViewportOffset = View.ViewRectMin.xy;
 }
 
-void DICamera_GetClipSpaceTransforms(in FNDICameraContext DIContext, out float4x4 Out_WorldToClipTransform, out float4x4 Out_TranslatedWorldToClipTransform, out float4x4 Out_ClipToWorldTransform, out float4x4 Out_ClipToViewTransform,
+void GetClipSpaceTransformsGPU_{ParameterName}(out float4x4 Out_WorldToClipTransform, out float4x4 Out_TranslatedWorldToClipTransform, out float4x4 Out_ClipToWorldTransform, out float4x4 Out_ClipToViewTransform,
 	out float4x4 Out_ClipToTranslatedWorldTransform, out float4x4 Out_ScreenToWorldTransform, out float4x4 Out_ScreenToTranslatedWorldTransform, out float4x4 Out_ClipToPreviousClipTransform)
 {
 	Out_WorldToClipTransform = LWCHackToFloat(PrimaryView.WorldToClip);
@@ -48,7 +54,7 @@ void DICamera_GetClipSpaceTransforms(in FNDICameraContext DIContext, out float4x
 	Out_ClipToPreviousClipTransform = View.ClipToPrevClip;
 }
 
-void DICamera_GetViewSpaceTransforms(in FNDICameraContext DIContext, out float4x4 Out_TranslatedWorldToViewTransform, out float4x4 Out_ViewToTranslatedWorldTransform, out float4x4 Out_TranslatedWorldToCameraViewTransform,
+void GetViewSpaceTransformsGPU_{ParameterName}(out float4x4 Out_TranslatedWorldToViewTransform, out float4x4 Out_ViewToTranslatedWorldTransform, out float4x4 Out_TranslatedWorldToCameraViewTransform,
 	out float4x4 Out_CameraViewToTranslatedWorldTransform, out float4x4 Out_ViewToClipTransform, out float4x4 Out_ViewToClipNoAATransform)
 {
 	Out_TranslatedWorldToViewTransform = View.TranslatedWorldToView;
@@ -58,12 +64,3 @@ void DICamera_GetViewSpaceTransforms(in FNDICameraContext DIContext, out float4x
 	Out_ViewToClipTransform = View.ViewToClip;
 	Out_ViewToClipNoAATransform = View.ViewToClipNoAA;
 }
-
-////////////////////////////////////////////////////////////////////////////////////////////////////////////////////////////////////////////////////////////////
-
-#define NDICAMERA_DECLARE_CONSTANTS(NAME) \
-	float3				SystemLWCTile_ ## NAME;
-
-#define NDICAMERA_MAKE_CONTEXT(NAME) \
-	FNDICameraContext DIContext; \
-	DIContext.SystemLWCTile = SystemLWCTile_ ## NAME;
diff --git a/Engine/Plugins/FX/Niagara/Shaders/Private/NiagaraDataInterfaceOcclusion.ush b/Engine/Plugins/FX/Niagara/Shaders/Private/NiagaraDataInterfaceOcclusion.ush
--- a/Engine/Plugins/FX/Niagara/Shaders/Private/NiagaraDataInterfaceOcclusion.ush
+++ b/Engine/Plugins/FX/Niagara/Shaders/Private/NiagaraDataInterfaceOcclusion.ush
@@ -1,17 +1,10 @@
 // Copyright Epic Games, Inc. All Rights Reserved.
 
-#pragma once
+float3 {ParameterName}_SystemLWCTile;
 
-#include "/Engine/Private/DeferredShadingCommon.ush"
-
-struct FNDIOcclusionContext
-{
-	float3 SystemLWCTile;
-};
-
-void DIOcclusion_Rectangle(in FNDIOcclusionContext DIContext, in float3 In_SampleCenterWorldPos, in float In_SampleWindowWidthWorld, in float In_SampleWindowHeightWorld, in float In_SampleSteps, out float Out_VisibilityFraction, out float Out_SampleFraction)
+void QueryOcclusionFactorWithRectangleGPU_{ParameterName}(in float3 In_SampleCenterWorldPos, in float In_SampleWindowWidthWorld, in float In_SampleWindowHeightWorld, in float In_SampleSteps, out float Out_VisibilityFraction, out float Out_SampleFraction)
 {
-	FLWCVector3 LwcSamplePos = MakeLWCVector3(DIContext.SystemLWCTile, In_SampleCenterWorldPos);
+	FLWCVector3 LwcSamplePos = MakeLWCVector3({ParameterName}_SystemLWCTile, In_SampleCenterWorldPos);
 	float CameraDistance = LWCToFloat(LWCAbs(LWCDot(LWCSubtract(LwcSamplePos, PrimaryView.WorldViewOrigin), View.ViewForward.xyz)));
 
 	float3 ResolvedLwcSamplePos = LWCToFloat(LWCAdd(LwcSamplePos, PrimaryView.PreViewTranslation));
@@ -62,11 +55,11 @@ void DIOcclusion_Rectangle(in FNDIOcclusionContext DIContext, in float3 In_Sampl
 	Out_SampleFraction = Steps == 0 ? 0 : (TotalSamples / (Steps * Steps));
 }
 
-void DIOcclusion_Circle(in FNDIOcclusionContext DIContext, in float3 In_SampleCenterWorldPos, in float In_SampleWindowDiameterWorld, in float In_SampleRays, in float In_SampleStepsPerRay, out float Out_VisibilityFraction, out float Out_SampleFraction)
+void QueryOcclusionFactorWithCircleGPU_{ParameterName}(in float3 In_SampleCenterWorldPos, in float In_SampleWindowDiameterWorld, in float In_SampleRays, in float In_SampleStepsPerRay, out float Out_VisibilityFraction, out float Out_SampleFraction)
 {
 	const float PI = 3.14159265;
 	const float SPIRAL_TURN = 2 * PI * 0.61803399; // use golden ratio to rotate sample pattern each ring so we get a spiral
-	FLWCVector3 LwcSamplePos = MakeLWCVector3(DIContext.SystemLWCTile, In_SampleCenterWorldPos);
+	FLWCVector3 LwcSamplePos = MakeLWCVector3({ParameterName}_SystemLWCTile, In_SampleCenterWorldPos);
 	float CameraDistance = LWCToFloat(LWCAbs(LWCDot(LWCSubtract(LwcSamplePos, PrimaryView.WorldViewOrigin), View.ViewForward.xyz)));
 
 	float3 ResolvedLwcSamplePos = LWCToFloat(LWCAdd(LwcSamplePos, PrimaryView.PreViewTranslation));
@@ -128,12 +121,3 @@ void DIOcclusion_Circle(in FNDIOcclusionContext DIContext, in float3 In_SampleCe
 	Out_VisibilityFraction = TotalSamples > 0 ? 1 - OccludedSamples / TotalSamples : 0;
 	Out_SampleFraction = Steps == 0 ? 0 : (TotalSamples / (Rays * Steps + 1));
 }
-
-////////////////////////////////////////////////////////////////////////////////////////////////////////////////////////////////////////////////////////////////
-
-#define NDIOCCLUSION_DECLARE_CONSTANTS(NAME) \
-	float3				SystemLWCTile_ ## NAME;
-
-#define NDIOCCLUSION_MAKE_CONTEXT(NAME) \
-	FNDIOcclusionContext DIContext; \
-	DIContext.SystemLWCTile = SystemLWCTile_ ## NAME;
diff --git a/Engine/Plugins/FX/Niagara/Source/Niagara/Classes/NiagaraDataInterfaceCamera.h b/Engine/Plugins/FX/Niagara/Source/Niagara/Classes/NiagaraDataInterfaceCamera.h
--- a/Engine/Plugins/FX/Niagara/Source/Niagara/Classes/NiagaraDataInterfaceCamera.h
+++ b/Engine/Plugins/FX/Niagara/Source/Niagara/Classes/NiagaraDataInterfaceCamera.h
@@ -28,9 +28,11 @@ class NIAGARA_API UNiagaraDataInterfaceCamera : public UNiagaraDataInterface
 {
 	GENERATED_UCLASS_BODY()
 
+	BEGIN_SHADER_PARAMETER_STRUCT(FShaderParameters, )
+		SHADER_PARAMETER(FVector3f, SystemLWCTile)
+	END_SHADER_PARAMETER_STRUCT();
+
 public:
-	DECLARE_NIAGARA_DI_PARAMETER();
-	
 	/** This is used to determine which camera position to query for cpu emitters. If no valid index is supplied, the first controller is used as camera reference. */
 	UPROPERTY(EditAnywhere, Category = "Camera")
 	int32 PlayerControllerIndex = 0;
@@ -52,12 +54,14 @@ public:
 	virtual int32 PerInstanceDataSize() const override { return sizeof(FCameraDataInterface_InstanceData); }
 	virtual bool PerInstanceTick(void* PerInstanceData, FNiagaraSystemInstance* SystemInstance, float DeltaSeconds) override;
 #if WITH_EDITORONLY_DATA
+	virtual bool AppendCompileHash(FNiagaraCompileHashVisitor* InVisitor) const override;
 	virtual bool GetFunctionHLSL(const FNiagaraDataInterfaceGPUParamInfo& ParamInfo, const FNiagaraDataInterfaceGeneratedFunction& FunctionInfo, int FunctionInstanceIndex, FString& OutHLSL) override;
 	virtual bool UpgradeFunctionCall(FNiagaraFunctionSignature& FunctionSignature) override;
-	virtual bool AppendCompileHash(FNiagaraCompileHashVisitor* InVisitor) const override;
-	virtual void GetCommonHLSL(FString& OutHLSL) override;
 	virtual void GetParameterDefinitionHLSL(const FNiagaraDataInterfaceGPUParamInfo& ParamInfo, FString& OutHLSL) override;
 #endif
+	virtual bool UseLegacyShaderBindings() const override { return false; }
+	virtual void BuildShaderParameters(FNiagaraShaderParametersBuilder& ShaderParametersBuilder) const override;
+	virtual void SetShaderParameters(const FNiagaraDataInterfaceSetShaderParametersContext& Context) const override;
 	virtual bool CanExecuteOnTarget(ENiagaraSimTarget Target) const override { return true; }
 	virtual bool HasTickGroupPrereqs() const override { return true; }
 	virtual ETickingGroup CalculateTickGroup(const void* PerInstanceData) const override;
@@ -78,6 +82,7 @@ public:
 protected:
 	virtual bool CopyToInternal(UNiagaraDataInterface* Destination) const override;
 private:
+	static const TCHAR* TemplateShaderFilePath;
 	static const FName CalculateDistancesName;
 	static const FName QueryClosestName;
 	static const FName GetViewPropertiesName;
diff --git a/Engine/Plugins/FX/Niagara/Source/Niagara/Classes/NiagaraDataInterfaceCurlNoise.h b/Engine/Plugins/FX/Niagara/Source/Niagara/Classes/NiagaraDataInterfaceCurlNoise.h
--- a/Engine/Plugins/FX/Niagara/Source/Niagara/Classes/NiagaraDataInterfaceCurlNoise.h
+++ b/Engine/Plugins/FX/Niagara/Source/Niagara/Classes/NiagaraDataInterfaceCurlNoise.h
@@ -12,10 +12,12 @@ UCLASS(EditInlineNew, Category = "Curl Noise LUT", meta = (DisplayName = "Curl N
 class NIAGARA_API UNiagaraDataInterfaceCurlNoise : public UNiagaraDataInterface
 {
 	GENERATED_UCLASS_BODY()
-public:
 
-	DECLARE_NIAGARA_DI_PARAMETER();
+	BEGIN_SHADER_PARAMETER_STRUCT(FShaderParameters, )
+		SHADER_PARAMETER(FVector3f, OffsetFromSeed)
+	END_SHADER_PARAMETER_STRUCT();
 
+public:
 	UPROPERTY(EditAnywhere, Category = "Curl Noise")
 	uint32 Seed;
 
@@ -42,9 +44,13 @@ public:
 
 	// GPU sim functionality
 #if WITH_EDITORONLY_DATA
+	virtual bool AppendCompileHash(FNiagaraCompileHashVisitor* InVisitor) const override;
 	virtual void GetParameterDefinitionHLSL(const FNiagaraDataInterfaceGPUParamInfo& ParamInfo, FString& OutHLSL) override;
 	virtual bool GetFunctionHLSL(const FNiagaraDataInterfaceGPUParamInfo& ParamInfo, const FNiagaraDataInterfaceGeneratedFunction& FunctionInfo, int FunctionInstanceIndex, FString& OutHLSL) override;
 #endif
+	virtual bool UseLegacyShaderBindings() const override { return false; }
+	virtual void BuildShaderParameters(FNiagaraShaderParametersBuilder& ShaderParametersBuilder) const override;
+	virtual void SetShaderParameters(const FNiagaraDataInterfaceSetShaderParametersContext& Context) const override;
 
 protected:
 	virtual bool CopyToInternal(UNiagaraDataInterface* Destination) const override;
diff --git a/Engine/Plugins/FX/Niagara/Source/Niagara/Classes/NiagaraDataInterfaceMeshRendererInfo.h b/Engine/Plugins/FX/Niagara/Source/Niagara/Classes/NiagaraDataInterfaceMeshRendererInfo.h
--- a/Engine/Plugins/FX/Niagara/Source/Niagara/Classes/NiagaraDataInterfaceMeshRendererInfo.h
+++ b/Engine/Plugins/FX/Niagara/Source/Niagara/Classes/NiagaraDataInterfaceMeshRendererInfo.h
@@ -19,9 +19,12 @@ class NIAGARA_API UNiagaraDataInterfaceMeshRendererInfo : public UNiagaraDataInt
 {
 	GENERATED_UCLASS_BODY()
 
-public:
-	DECLARE_NIAGARA_DI_PARAMETER();
+	BEGIN_SHADER_PARAMETER_STRUCT(FShaderParameters,)
+		SHADER_PARAMETER(uint32,			NumMeshes)
+		SHADER_PARAMETER_SRV(Buffer<float>,	MeshDataBuffer)
+	END_SHADER_PARAMETER_STRUCT()
 
+public:
 	UNiagaraMeshRendererProperties* GetMeshRenderer() const { return MeshRenderer; }
 
 	//UObject Interface
@@ -39,9 +42,14 @@ public:
 	virtual void GetVMExternalFunction(const FVMExternalFunctionBindingInfo& BindingInfo, void* InstanceData, FVMExternalFunction& OutFunc) override;
 	virtual bool Equals(const UNiagaraDataInterface* Other) const override;
 #if WITH_EDITORONLY_DATA
+	virtual bool AppendCompileHash(FNiagaraCompileHashVisitor* InVisitor) const override;
 	virtual void GetParameterDefinitionHLSL(const FNiagaraDataInterfaceGPUParamInfo& ParamInfo, FString& OutHLSL) override;
 	virtual bool GetFunctionHLSL(const FNiagaraDataInterfaceGPUParamInfo& ParamInfo, const FNiagaraDataInterfaceGeneratedFunction& FunctionInfo, int FunctionInstanceIndex, FString& OutHLSL) override;
 #endif
+	virtual bool UseLegacyShaderBindings() const override { return false; }
+	virtual void BuildShaderParameters(FNiagaraShaderParametersBuilder& ShaderParametersBuilder) const override;
+	virtual void SetShaderParameters(const FNiagaraDataInterfaceSetShaderParametersContext& Context) const override;
+
 	virtual bool CanExecuteOnTarget(ENiagaraSimTarget Target) const override { return true; }
 #if WITH_EDITOR
 	virtual bool UpgradeFunctionCall(FNiagaraFunctionSignature& FunctionSignature) override;
diff --git a/Engine/Plugins/FX/Niagara/Source/Niagara/Classes/NiagaraDataInterfaceOcclusion.h b/Engine/Plugins/FX/Niagara/Source/Niagara/Classes/NiagaraDataInterfaceOcclusion.h
--- a/Engine/Plugins/FX/Niagara/Source/Niagara/Classes/NiagaraDataInterfaceOcclusion.h
+++ b/Engine/Plugins/FX/Niagara/Source/Niagara/Classes/NiagaraDataInterfaceOcclusion.h
@@ -11,10 +11,11 @@ class NIAGARA_API UNiagaraDataInterfaceOcclusion : public UNiagaraDataInterface
 {
 	GENERATED_UCLASS_BODY()
 
-public:
+	BEGIN_SHADER_PARAMETER_STRUCT(FShaderParameters, )
+		SHADER_PARAMETER(FVector3f, SystemLWCTile)
+	END_SHADER_PARAMETER_STRUCT();
 
-	DECLARE_NIAGARA_DI_PARAMETER();
-	
+public:
 	//UObject Interface
 	virtual void PostInitProperties() override;
 	//UObject Interface End
@@ -28,11 +29,16 @@ public:
 	virtual bool UpgradeFunctionCall(FNiagaraFunctionSignature& FunctionSignature) override;
 	virtual void GetParameterDefinitionHLSL(const FNiagaraDataInterfaceGPUParamInfo& ParamInfo, FString& OutHLSL) override;
 #endif
+	virtual bool UseLegacyShaderBindings() const override { return false; }
+	virtual void BuildShaderParameters(FNiagaraShaderParametersBuilder& ShaderParametersBuilder) const override;
+	virtual void SetShaderParameters(const FNiagaraDataInterfaceSetShaderParametersContext& Context) const override;
+
 	virtual bool CanExecuteOnTarget(ENiagaraSimTarget Target) const override { return Target == ENiagaraSimTarget::GPUComputeSim; }
 	virtual bool RequiresDepthBuffer() const override { return true; }
 	//UNiagaraDataInterface Interface
 	
 private:
+	static const TCHAR* TemplateShaderFilePath;
 	static const FName GetCameraOcclusionRectangleName;
 	static const FName GetCameraOcclusionCircleName;
 };
@@ -45,4 +51,4 @@ struct FNiagaraDataIntefaceProxyOcclusionQuery : public FNiagaraDataInterfacePro
 	{
 		return 0;
 	}
-};
\ No newline at end of file
+};
diff --git a/Engine/Plugins/FX/Niagara/Source/Niagara/Classes/NiagaraDataInterfaceVectorField.h b/Engine/Plugins/FX/Niagara/Source/Niagara/Classes/NiagaraDataInterfaceVectorField.h
--- a/Engine/Plugins/FX/Niagara/Source/Niagara/Classes/NiagaraDataInterfaceVectorField.h
+++ b/Engine/Plugins/FX/Niagara/Source/Niagara/Classes/NiagaraDataInterfaceVectorField.h
@@ -12,6 +12,15 @@ class NIAGARA_API UNiagaraDataInterfaceVectorField : public UNiagaraDataInterfac
 {
 	GENERATED_UCLASS_BODY()
 
+	BEGIN_SHADER_PARAMETER_STRUCT(FShaderParameters, )
+		SHADER_PARAMETER(FVector3f,				TilingAxes)
+		SHADER_PARAMETER(FVector3f,				Dimensions)
+		SHADER_PARAMETER(FVector3f,				MinBounds)
+		SHADER_PARAMETER(FVector3f,				MaxBounds)
+		SHADER_PARAMETER_TEXTURE(Texture3D,		Texture)
+		SHADER_PARAMETER_SAMPLER(SamplerState,	Sampler)
+	END_SHADER_PARAMETER_STRUCT();
+
 public:
 	/** Vector field to sample from. */
 	UPROPERTY(EditAnywhere, Category = VectorField)
@@ -26,9 +35,6 @@ public:
 
 public:
 	//~ UObject interface
-
-	DECLARE_NIAGARA_DI_PARAMETER();
-
 	virtual void PostInitProperties() override;
 	virtual void PostLoad() override; 
 #if WITH_EDITOR
@@ -52,9 +58,13 @@ public:
 
 	// GPU sim functionality
 #if WITH_EDITORONLY_DATA
+	virtual bool AppendCompileHash(FNiagaraCompileHashVisitor* InVisitor) const override;
 	virtual void GetParameterDefinitionHLSL(const FNiagaraDataInterfaceGPUParamInfo& ParamInfo, FString& OutHLSL) override;
 	virtual bool GetFunctionHLSL(const FNiagaraDataInterfaceGPUParamInfo& ParamInfo, const FNiagaraDataInterfaceGeneratedFunction& FunctionInfo, int FunctionInstanceIndex, FString& OutHLSL) override;
 #endif
+	virtual bool UseLegacyShaderBindings() const override { return false; }
+	virtual void BuildShaderParameters(FNiagaraShaderParametersBuilder& ShaderParametersBuilder) const override;
+	virtual void SetShaderParameters(const FNiagaraDataInterfaceSetShaderParametersContext& Context) const override;
 	//~ UNiagaraDataInterface interface END
 
 	// VM functions
@@ -98,19 +108,3 @@ struct FNiagaraDataInterfaceProxyVectorField : public FNiagaraDataInterfaceProxy
 		return 0;
 	}
 };
-
-struct FNiagaraDataInterfaceParametersCS_VectorField : public FNiagaraDataInterfaceParametersCS
-{
-	DECLARE_TYPE_LAYOUT(FNiagaraDataInterfaceParametersCS_VectorField, NonVirtual);
-public:
-	void Bind(const FNiagaraDataInterfaceGPUParamInfo& ParameterInfo, const class FShaderParameterMap& ParameterMap);
-	void Set(FRHICommandList& RHICmdList, const FNiagaraDataInterfaceSetArgs& Context) const;
-
-private:
-	LAYOUT_FIELD(FShaderResourceParameter, VectorFieldSampler);
-	LAYOUT_FIELD(FShaderResourceParameter, VectorFieldTexture);
-	LAYOUT_FIELD(FShaderParameter, TilingAxes);
-	LAYOUT_FIELD(FShaderParameter, Dimensions);
-	LAYOUT_FIELD(FShaderParameter, MinBounds);
-	LAYOUT_FIELD(FShaderParameter, MaxBounds);
-};
diff --git a/Engine/Plugins/FX/Niagara/Source/Niagara/Private/NiagaraDataInterfaceCamera.cpp b/Engine/Plugins/FX/Niagara/Source/Niagara/Private/NiagaraDataInterfaceCamera.cpp
--- a/Engine/Plugins/FX/Niagara/Source/Niagara/Private/NiagaraDataInterfaceCamera.cpp
+++ b/Engine/Plugins/FX/Niagara/Source/Niagara/Private/NiagaraDataInterfaceCamera.cpp
@@ -6,6 +6,9 @@
 #include "NiagaraTypes.h"
 #include "NiagaraWorldManager.h"
 #include "Internationalization/Internationalization.h"
+#include "NiagaraGPUSystemTick.h"
+#include "NiagaraShaderParametersBuilder.h"
+#include "NiagaraSystemGpuComputeProxy.h"
 #include "NiagaraSystemInstance.h"
 #include "GameFramework/PlayerController.h"
 #include "ShaderParameterUtils.h"
@@ -31,6 +34,7 @@ struct FNiagaraCameraDIFunctionVersion
     };
 };
 
+const TCHAR* UNiagaraDataInterfaceCamera::TemplateShaderFilePath = TEXT("/Plugin/FX/Niagara/Private/NiagaraDataInterfaceCamera.ush");
 const FName UNiagaraDataInterfaceCamera::GetViewPropertiesName(TEXT("GetViewPropertiesGPU"));
 const FName UNiagaraDataInterfaceCamera::GetClipSpaceTransformsName(TEXT("GetClipSpaceTransformsGPU"));
 const FName UNiagaraDataInterfaceCamera::GetViewSpaceTransformsName(TEXT("GetViewSpaceTransformsGPU"));
@@ -276,83 +280,13 @@ void UNiagaraDataInterfaceCamera::GetFunctions(TArray<FNiagaraFunctionSignature>
 #if WITH_EDITORONLY_DATA
 bool UNiagaraDataInterfaceCamera::GetFunctionHLSL(const FNiagaraDataInterfaceGPUParamInfo& ParamInfo, const FNiagaraDataInterfaceGeneratedFunction& FunctionInfo, int FunctionInstanceIndex, FString& OutHLSL)
 {
-	TMap<FString, FStringFormatArg> ArgsSample =
+	if ((FunctionInfo.DefinitionName == GetViewPropertiesName) ||
+		(FunctionInfo.DefinitionName == GetFieldOfViewName) ||
+		(FunctionInfo.DefinitionName == GetClipSpaceTransformsName) ||
+		(FunctionInfo.DefinitionName == GetViewSpaceTransformsName) ||
+		(FunctionInfo.DefinitionName == GetCameraPropertiesName) ||
+		(FunctionInfo.DefinitionName == GetTAAJitterName))
 	{
-		{TEXT("FunctionName"), FunctionInfo.InstanceName},
-		{TEXT("NDIGetContextName"), TEXT("NDICAMERA_MAKE_CONTEXT(") + ParamInfo.DataInterfaceHLSLSymbol + TEXT(")")},
-	};
-	
-	if (FunctionInfo.DefinitionName == GetViewPropertiesName)
-	{
-		static const TCHAR *FormatSample = TEXT(R"(
-			void {FunctionName}(out float3 Out_ViewPositionWorld, out float3 Out_ViewForwardVector, out float3 Out_ViewUpVector, out float3 Out_ViewRightVector, out float4 Out_ViewSizeAndInverseSize, out float4 Out_ScreenToViewSpace, out float2 Out_Current_TAAJitter, out float2 Out_Previous_TAAJitter, out float3 Out_PreViewTranslation, out float4 Out_BufferSizeAndInverseSize, out float2 Out_ViewportOffset)
-			{
-				{NDIGetContextName}
-				DICamera_GetViewProperties(DIContext, Out_ViewPositionWorld, Out_ViewForwardVector, Out_ViewUpVector, Out_ViewRightVector, Out_ViewSizeAndInverseSize, Out_ScreenToViewSpace, Out_Current_TAAJitter, Out_Previous_TAAJitter, Out_PreViewTranslation, Out_BufferSizeAndInverseSize, Out_ViewportOffset);
-			}
-		)");
-		OutHLSL += FString::Format(FormatSample, ArgsSample);
-		return true;
-	}
-	if (FunctionInfo.DefinitionName == GetFieldOfViewName)
-	{
-		static const TCHAR *FormatSample = TEXT(R"(
-			void {FunctionName}(out float Out_FieldOfViewAngle)
-			{
-				Out_FieldOfViewAngle = degrees(View.FieldOfViewWideAngles.x);
-			}
-		)");
-		OutHLSL += FString::Format(FormatSample, ArgsSample);
-		return true;
-	}
-	if (FunctionInfo.DefinitionName == GetClipSpaceTransformsName)
-	{
-		static const TCHAR *FormatSample = TEXT(R"(
-			void {FunctionName}(out float4x4 Out_WorldToClipTransform, out float4x4 Out_TranslatedWorldToClipTransform, out float4x4 Out_ClipToWorldTransform, out float4x4 Out_ClipToViewTransform,
-				out float4x4 Out_ClipToTranslatedWorldTransform, out float4x4 Out_ScreenToWorldTransform, out float4x4 Out_ScreenToTranslatedWorldTransform, out float4x4 Out_ClipToPreviousClipTransform)
-			{
-				{NDIGetContextName}
-				DICamera_GetClipSpaceTransforms(DIContext, Out_WorldToClipTransform, Out_TranslatedWorldToClipTransform, Out_ClipToWorldTransform, Out_ClipToViewTransform, Out_ClipToTranslatedWorldTransform, Out_ScreenToWorldTransform, Out_ScreenToTranslatedWorldTransform, Out_ClipToPreviousClipTransform);
-			}
-		)");
-		OutHLSL += FString::Format(FormatSample, ArgsSample);
-		return true;
-	}
-	if (FunctionInfo.DefinitionName == GetViewSpaceTransformsName)
-	{
-		static const TCHAR *FormatSample = TEXT(R"(
-			void {FunctionName}(out float4x4 Out_TranslatedWorldToViewTransform, out float4x4 Out_ViewToTranslatedWorldTransform, out float4x4 Out_TranslatedWorldToCameraViewTransform,
-				out float4x4 Out_CameraViewToTranslatedWorldTransform, out float4x4 Out_ViewToClipTransform, out float4x4 Out_ViewToClipNoAATransform)
-			{
-				{NDIGetContextName}
-				DICamera_GetViewSpaceTransforms(DIContext, Out_TranslatedWorldToViewTransform, Out_ViewToTranslatedWorldTransform, Out_TranslatedWorldToCameraViewTransform, Out_CameraViewToTranslatedWorldTransform, Out_ViewToClipTransform, Out_ViewToClipNoAATransform);
-			}
-		)");
-		OutHLSL += FString::Format(FormatSample, ArgsSample);
-		return true;
-	}
-	if (FunctionInfo.DefinitionName == GetCameraPropertiesName)
-	{
-		static const TCHAR *FormatSample = TEXT(R"(
-			void {FunctionName}(out float3 Out_CameraPositionWorld, out float3 Out_ViewForwardVector, out float3 Out_ViewUpVector, out float3 Out_ViewRightVector)
-			{
-				{NDIGetContextName}
-				DICamera_GetCameraProperties(DIContext, Out_CameraPositionWorld, Out_ViewForwardVector, Out_ViewUpVector, Out_ViewRightVector);
-			}
-		)");
-		OutHLSL += FString::Format(FormatSample, ArgsSample);
-		return true;
-	}
-	if (FunctionInfo.DefinitionName == GetTAAJitterName)
-	{
-		static const TCHAR *FormatSample = TEXT(R"(
-			void {FunctionName}(out float2 Out_CurrentJitter, out float2 Out_PreviousJitter)
-			{
-				Out_CurrentJitter = View.TemporalAAJitter.xy;
-				Out_PreviousJitter = View.TemporalAAJitter.zw;
-			}
-		)");
-		OutHLSL += FString::Format(FormatSample, ArgsSample);
 		return true;
 	}
 	return false;
@@ -412,22 +346,34 @@ bool UNiagaraDataInterfaceCamera::AppendCompileHash(FNiagaraCompileHashVisitor*
 	{
 		return false;
 	}
-	FSHAHash Hash = GetShaderFileHash((TEXT("/Plugin/FX/Niagara/Private/NiagaraDataInterfaceCamera.ush")), EShaderPlatform::SP_PCD3D_SM5);
-	InVisitor->UpdateString(TEXT("NiagaraDataInterfaceCameraHLSLSource"), Hash.ToString());
+	InVisitor->UpdateString(TEXT("UNiagaraDataInterfaceCameraHLSLSource"), GetShaderFileHash(TemplateShaderFilePath, EShaderPlatform::SP_PCD3D_SM5).ToString());
+	InVisitor->UpdateShaderParameters<FShaderParameters>();
 	return true;
 }
 
-void UNiagaraDataInterfaceCamera::GetCommonHLSL(FString& OutHLSL)
+void UNiagaraDataInterfaceCamera::GetParameterDefinitionHLSL(const FNiagaraDataInterfaceGPUParamInfo& ParamInfo, FString& OutHLSL)
+{
+	TMap<FString, FStringFormatArg> TemplateArgs =
+	{
+		{TEXT("ParameterName"),	ParamInfo.DataInterfaceHLSLSymbol},
+	};
+
+	FString TemplateFile;
+	LoadShaderSourceFile(TemplateShaderFilePath, EShaderPlatform::SP_PCD3D_SM5, &TemplateFile, nullptr);
+	OutHLSL += FString::Format(*TemplateFile, TemplateArgs);
+}
+#endif
+
+void UNiagaraDataInterfaceCamera::BuildShaderParameters(FNiagaraShaderParametersBuilder& ShaderParametersBuilder) const
 {
-	OutHLSL += TEXT("#include \"/Plugin/FX/Niagara/Private/NiagaraDataInterfaceCamera.ush\"\n");
+	ShaderParametersBuilder.AddNestedStruct<FShaderParameters>();
 }
 
-void UNiagaraDataInterfaceCamera::GetParameterDefinitionHLSL(const FNiagaraDataInterfaceGPUParamInfo& ParamInfo, FString& OutHLSL)
+void UNiagaraDataInterfaceCamera::SetShaderParameters(const FNiagaraDataInterfaceSetShaderParametersContext& Context) const
 {
-	Super::GetParameterDefinitionHLSL(ParamInfo, OutHLSL);
-	OutHLSL += TEXT("NDICAMERA_DECLARE_CONSTANTS(") + ParamInfo.DataInterfaceHLSLSymbol + TEXT(")\n");
+	FShaderParameters* ShaderParameters = Context.GetParameterNestedStruct<FShaderParameters>();
+	ShaderParameters->SystemLWCTile = Context.GetSystemTick().SystemGpuComputeProxy->GetSystemLWCTile();
 }
-#endif
 
 DEFINE_NDI_DIRECT_FUNC_BINDER(UNiagaraDataInterfaceCamera, GetClosestParticles);
 DEFINE_NDI_DIRECT_FUNC_BINDER(UNiagaraDataInterfaceCamera, CalculateParticleDistances);
@@ -614,34 +560,4 @@ bool UNiagaraDataInterfaceCamera::Equals(const UNiagaraDataInterface* Other) con
 		OtherTyped->bRequireCurrentFrameData == bRequireCurrentFrameData ;
 }
 
-struct FNiagaraDataInterfaceParametersCS_Camera : public FNiagaraDataInterfaceParametersCS
-{
-	DECLARE_TYPE_LAYOUT(FNiagaraDataInterfaceParametersCS_Camera, NonVirtual);
-
-	void Bind(const FNiagaraDataInterfaceGPUParamInfo& ParameterInfo, const class FShaderParameterMap& ParameterMap)
-	{
-		SystemLWCTileParam.Bind(ParameterMap, *(SystemLWCTileName + ParameterInfo.DataInterfaceHLSLSymbol));
-	}
-
-	void Set(FRHICommandList& RHICmdList, const FNiagaraDataInterfaceSetArgs& Context) const
-	{
-		check(IsInRenderingThread());
-
-		FRHIComputeShader* ComputeShaderRHI = Context.Shader.GetComputeShader();
-		SetShaderValue(RHICmdList, ComputeShaderRHI, SystemLWCTileParam, Context.SystemLWCTile);
-	}
-
-private:
-	// LWC data
-	LAYOUT_FIELD(FShaderParameter, SystemLWCTileParam);
-
-	static const FString SystemLWCTileName;
-};
-
-// LWC
-const FString FNiagaraDataInterfaceParametersCS_Camera::SystemLWCTileName(TEXT("SystemLWCTile_"));
-
-IMPLEMENT_TYPE_LAYOUT(FNiagaraDataInterfaceParametersCS_Camera);
-IMPLEMENT_NIAGARA_DI_PARAMETER(UNiagaraDataInterfaceCamera, FNiagaraDataInterfaceParametersCS_Camera);
-
 #undef LOCTEXT_NAMESPACE
diff --git a/Engine/Plugins/FX/Niagara/Source/Niagara/Private/NiagaraDataInterfaceCurlNoise.cpp b/Engine/Plugins/FX/Niagara/Source/Niagara/Private/NiagaraDataInterfaceCurlNoise.cpp
--- a/Engine/Plugins/FX/Niagara/Source/Niagara/Private/NiagaraDataInterfaceCurlNoise.cpp
+++ b/Engine/Plugins/FX/Niagara/Source/Niagara/Private/NiagaraDataInterfaceCurlNoise.cpp
@@ -2,11 +2,11 @@
 
 #include "NiagaraDataInterfaceCurlNoise.h"
 #include "NiagaraShader.h"
+#include "NiagaraShaderParametersBuilder.h"
 #include "ShaderParameterUtils.h"
 #include "Math/IntVector.h"
 
 static const FName SampleNoiseFieldName(TEXT("SampleNoiseField"));
-static const FString OffsetFromSeedBaseName(TEXT("OffsetFromSeed_"));
 
 // ----------------------------------------------------- start noise helpers -----------------------------------------------------
 // Fairly straightforward CPU implementation of the equivalent HLSL code found in Engine/Shaders/Private/Random.ush
@@ -425,9 +425,16 @@ void UNiagaraDataInterfaceCurlNoise::SampleNoiseField(FVectorVMExternalFunctionC
 }
 
 #if WITH_EDITORONLY_DATA
+bool UNiagaraDataInterfaceCurlNoise::AppendCompileHash(FNiagaraCompileHashVisitor* InVisitor) const
+{
+	bool bSuccess = Super::AppendCompileHash(InVisitor);
+	bSuccess &= InVisitor->UpdateShaderParameters<FShaderParameters>();
+	return bSuccess;
+}
+
 bool UNiagaraDataInterfaceCurlNoise::GetFunctionHLSL(const FNiagaraDataInterfaceGPUParamInfo& ParamInfo, const FNiagaraDataInterfaceGeneratedFunction& FunctionInfo, int FunctionInstanceIndex, FString& OutHLSL)
 {
-	static const TCHAR *FormatSample = TEXT(R"(
+	static const TCHAR* FormatSample = TEXT(R"(
 		void {FunctionName}(float3 In_XYZ, out float3 Out_Value)
 		{
 			// NOTE(mv): The comments in random.ush claims that the unused part is optimized away, so it only uses 6 out of 12 values in our case.
@@ -437,23 +444,30 @@ bool UNiagaraDataInterfaceCurlNoise::GetFunctionHLSL(const FNiagaraDataInterface
 	)");
 	TMap<FString, FStringFormatArg> ArgsSample;
 	ArgsSample.Add(TEXT("FunctionName"), FunctionInfo.InstanceName);
-	ArgsSample.Add(TEXT("OffsetFromSeedName"), OffsetFromSeedBaseName + ParamInfo.DataInterfaceHLSLSymbol);
+	ArgsSample.Add(TEXT("OffsetFromSeedName"), ParamInfo.DataInterfaceHLSLSymbol + TEXT("_OffsetFromSeed"));
 	OutHLSL += FString::Format(FormatSample, ArgsSample);
 	return true;
 }
 
 void UNiagaraDataInterfaceCurlNoise::GetParameterDefinitionHLSL(const FNiagaraDataInterfaceGPUParamInfo& ParamInfo, FString& OutHLSL)
 {
-	static const TCHAR *FormatDeclarations = TEXT(R"(
-		float3 {OffsetFromSeedName};
-	)");
-
-	TMap<FString, FStringFormatArg> ArgsDeclarations;
-	ArgsDeclarations.Add(TEXT("OffsetFromSeedName"), OffsetFromSeedBaseName + ParamInfo.DataInterfaceHLSLSymbol);
-	OutHLSL += FString::Format(FormatDeclarations, ArgsDeclarations);
+	OutHLSL.Appendf(TEXT("float3 %s_OffsetFromSeed;\n"), *ParamInfo.DataInterfaceHLSLSymbol);
 }
 #endif
 
+void UNiagaraDataInterfaceCurlNoise::BuildShaderParameters(FNiagaraShaderParametersBuilder& ShaderParametersBuilder) const
+{
+	ShaderParametersBuilder.AddNestedStruct<FShaderParameters>();
+}
+
+void UNiagaraDataInterfaceCurlNoise::SetShaderParameters(const FNiagaraDataInterfaceSetShaderParametersContext& Context) const
+{
+	FNiagaraDataInterfaceProxyCurlNoise& DIProxy = Context.GetProxy<FNiagaraDataInterfaceProxyCurlNoise>();
+	
+	FShaderParameters* ShaderParameters = Context.GetParameterNestedStruct<FShaderParameters>();
+	ShaderParameters->OffsetFromSeed = FVector3f(DIProxy.OffsetFromSeed);
+}
+
 void UNiagaraDataInterfaceCurlNoise::PushToRenderThreadImpl()
 {
 	FNiagaraDataInterfaceProxyCurlNoise* RT_Proxy = GetProxyAs<FNiagaraDataInterfaceProxyCurlNoise>();
@@ -467,32 +481,3 @@ void UNiagaraDataInterfaceCurlNoise::PushToRenderThreadImpl()
 		RT_Proxy->OffsetFromSeed = RT_Offset;
 	});
 }
-
-struct FNiagaraDataInterfaceParametersCS_CurlNoise : public FNiagaraDataInterfaceParametersCS
-{
-	DECLARE_TYPE_LAYOUT(FNiagaraDataInterfaceParametersCS_CurlNoise, NonVirtual);
-public:
-	void Bind(const FNiagaraDataInterfaceGPUParamInfo& ParameterInfo, const class FShaderParameterMap& ParameterMap)
-	{
-		OffsetFromSeed.Bind(ParameterMap, *(OffsetFromSeedBaseName + ParameterInfo.DataInterfaceHLSLSymbol));
-	}
-
-	void Set(FRHICommandList& RHICmdList, const FNiagaraDataInterfaceSetArgs& Context) const
-	{
-		check(IsInRenderingThread());
-
-		// Get shader and DI
-		FRHIComputeShader* ComputeShaderRHI = Context.Shader.GetComputeShader();
-		FNiagaraDataInterfaceProxyCurlNoise* CNDI = static_cast<FNiagaraDataInterfaceProxyCurlNoise*>(Context.DataInterface);
-
-		// Set parameters
-		SetShaderValue(RHICmdList, ComputeShaderRHI, OffsetFromSeed, (FVector3f)CNDI->OffsetFromSeed);
-	}
-
-private:
-	LAYOUT_FIELD(FShaderParameter, OffsetFromSeed);
-};
-
-IMPLEMENT_TYPE_LAYOUT(FNiagaraDataInterfaceParametersCS_CurlNoise);
-
-IMPLEMENT_NIAGARA_DI_PARAMETER(UNiagaraDataInterfaceCurlNoise, FNiagaraDataInterfaceParametersCS_CurlNoise);
diff --git a/Engine/Plugins/FX/Niagara/Source/Niagara/Private/NiagaraDataInterfaceMeshRendererInfo.cpp b/Engine/Plugins/FX/Niagara/Source/Niagara/Private/NiagaraDataInterfaceMeshRendererInfo.cpp
--- a/Engine/Plugins/FX/Niagara/Source/Niagara/Private/NiagaraDataInterfaceMeshRendererInfo.cpp
+++ b/Engine/Plugins/FX/Niagara/Source/Niagara/Private/NiagaraDataInterfaceMeshRendererInfo.cpp
@@ -6,6 +6,7 @@
 #include "ShaderParameterUtils.h"
 #include "NiagaraStats.h"
 #include "NiagaraRenderer.h"
+#include "NiagaraShaderParametersBuilder.h"
 #include "Engine/StaticMesh.h"
 
 #define LOCTEXT_NAMESPACE "NiagaraDataInterfaceMeshRendererInfo"
@@ -18,9 +19,6 @@ namespace NDIMeshRendererInfoInternal
 {
 	const FName GetNumMeshesName("GetNumMeshes");
 	const FName GetMeshLocalBoundsName("GetMeshLocalBounds");
-
-	const FString NumMeshesPrefix = TEXT("NumMeshes_");
-	const FString MeshDataBufferPrefix = TEXT("MeshDataBuffer_");
 }
 
 enum class ENDIMeshRendererInfoVersion : uint32
@@ -130,49 +128,6 @@ struct FNDIMeshRendererInfoProxy : public FNiagaraDataInterfaceProxy
 	virtual int32 PerInstanceDataPassedToRenderThreadSize() const override { return 0; }
 };
 
-/** The parameters used by the data interface in GPU emitters */
-struct FNiagaraDataInterfaceParametersCS_MeshRendererInfo : public FNiagaraDataInterfaceParametersCS
-{
-	DECLARE_TYPE_LAYOUT(FNiagaraDataInterfaceParametersCS_MeshRendererInfo, NonVirtual);
-public:
-	void Bind(const FNiagaraDataInterfaceGPUParamInfo& ParameterInfo, const class FShaderParameterMap& ParameterMap)
-	{
-		NumMeshesParam.Bind(ParameterMap, *(NDIMeshRendererInfoInternal::NumMeshesPrefix + ParameterInfo.DataInterfaceHLSLSymbol));
-		MeshDataBuffer.Bind(ParameterMap, *(NDIMeshRendererInfoInternal::MeshDataBufferPrefix + ParameterInfo.DataInterfaceHLSLSymbol));
-	}
-
-	void Set(FRHICommandList& RHICmdList, const FNiagaraDataInterfaceSetArgs& Context) const
-	{
-		check(IsInRenderingThread());
-		FRHIComputeShader* ComputeShaderRHI = RHICmdList.GetBoundComputeShader();
-
-		auto Proxy = static_cast<FNDIMeshRendererInfoProxy*>(Context.DataInterface);
-		FShaderResourceViewRHIRef MeshDataBufferSRV;
-		uint32 NumMeshes = 0;
-		if (Proxy->GPUData.IsValid())
-		{
-			NumMeshes = Proxy->GPUData->GetNumMeshes();
-			MeshDataBufferSRV = Proxy->GPUData->GetMeshDataBufferSRV();
-		}
-
-		if (!MeshDataBufferSRV.IsValid())
-		{
-			MeshDataBufferSRV = FNiagaraRenderer::GetDummyFloatBuffer();
-		}
-		
-		SetShaderValue(RHICmdList, ComputeShaderRHI, NumMeshesParam, NumMeshes);
-		SetSRVParameter(RHICmdList, ComputeShaderRHI, MeshDataBuffer, MeshDataBufferSRV);
-	}
-
-private:
-	LAYOUT_FIELD(FShaderParameter, NumMeshesParam);
-	LAYOUT_FIELD(FShaderResourceParameter, MeshDataBuffer);
-};
-
-IMPLEMENT_TYPE_LAYOUT(FNiagaraDataInterfaceParametersCS_MeshRendererInfo);
-
-IMPLEMENT_NIAGARA_DI_PARAMETER(UNiagaraDataInterfaceMeshRendererInfo, FNiagaraDataInterfaceParametersCS_MeshRendererInfo);
-
 ///////////////////////////////////////////////////////////////////////////////////////////////
 
 FNDIMeshRendererInfo::FCacheMap FNDIMeshRendererInfo::CachedData;
@@ -426,16 +381,17 @@ bool UNiagaraDataInterfaceMeshRendererInfo::Equals(const UNiagaraDataInterface*
 }
 
 #if WITH_EDITORONLY_DATA
+bool UNiagaraDataInterfaceMeshRendererInfo::AppendCompileHash(FNiagaraCompileHashVisitor* InVisitor) const
+{
+	bool bSuccess = Super::AppendCompileHash(InVisitor);
+	bSuccess &= InVisitor->UpdateShaderParameters<FShaderParameters>();
+	return bSuccess;
+}
+
 void UNiagaraDataInterfaceMeshRendererInfo::GetParameterDefinitionHLSL(const FNiagaraDataInterfaceGPUParamInfo& ParamInfo, FString& OutHLSL)
 {
-	const FStringFormatOrderedArguments Args = {
-		ParamInfo.DataInterfaceHLSLSymbol
-	};
-	OutHLSL += FString::Format(TEXT(R"(
-		uint NumMeshes_{0};
-		Buffer<float> MeshDataBuffer_{0};
-	)"),
-	Args);
+	OutHLSL.Appendf(TEXT("uint %s_NumMeshes;\n"), *ParamInfo.DataInterfaceHLSLSymbol);
+	OutHLSL.Appendf(TEXT("Buffer<float> %s_MeshDataBuffer;\n"), *ParamInfo.DataInterfaceHLSLSymbol);
 }
 
 bool UNiagaraDataInterfaceMeshRendererInfo::GetFunctionHLSL(const FNiagaraDataInterfaceGPUParamInfo& ParamInfo, const FNiagaraDataInterfaceGeneratedFunction& FunctionInfo, int FunctionInstanceIndex, FString& OutHLSL)
@@ -450,7 +406,7 @@ bool UNiagaraDataInterfaceMeshRendererInfo::GetFunctionHLSL(const FNiagaraDataIn
 		OutHLSL += FString::Format(TEXT(R"(
 			void {FuncName}(out int OutNumMeshes)
 			{
-				OutNumMeshes = NumMeshes_{DIName};
+				OutNumMeshes = {DIName}_NumMeshes;
 			}
 			)"),
 			Args);
@@ -465,19 +421,19 @@ bool UNiagaraDataInterfaceMeshRendererInfo::GetFunctionHLSL(const FNiagaraDataIn
 				OutMinBounds = (float3)0;
 				OutMaxBounds = (float3)0;
 				OutSize = (float3)0;
-				if (NumMeshes_{DIName} > 0)
+				if ({DIName}_NumMeshes > 0)
 				{
 					const uint MeshDataNumFloats = 6;
-					const uint BufferOffs = clamp(MeshIndex, 0, int(NumMeshes_{DIName} - 1)) * MeshDataNumFloats;
+					const uint BufferOffs = clamp(MeshIndex, 0, int({DIName}_NumMeshes - 1)) * MeshDataNumFloats;
 					OutMinBounds = float3(
-						MeshDataBuffer_{DIName}[BufferOffs + 0],
-						MeshDataBuffer_{DIName}[BufferOffs + 1],
-						MeshDataBuffer_{DIName}[BufferOffs + 2]
+						{DIName}_MeshDataBuffer[BufferOffs + 0],
+						{DIName}_MeshDataBuffer[BufferOffs + 1],
+						{DIName}_MeshDataBuffer[BufferOffs + 2]
 					);
 					OutMaxBounds = float3(
-						MeshDataBuffer_{DIName}[BufferOffs + 3],
-						MeshDataBuffer_{DIName}[BufferOffs + 4],
-						MeshDataBuffer_{DIName}[BufferOffs + 5]
+						{DIName}_MeshDataBuffer[BufferOffs + 3],
+						{DIName}_MeshDataBuffer[BufferOffs + 4],
+						{DIName}_MeshDataBuffer[BufferOffs + 5]
 					);
 					OutSize = OutMaxBounds - OutMinBounds;
 				}
@@ -492,6 +448,27 @@ bool UNiagaraDataInterfaceMeshRendererInfo::GetFunctionHLSL(const FNiagaraDataIn
 }
 #endif
 
+void UNiagaraDataInterfaceMeshRendererInfo::BuildShaderParameters(FNiagaraShaderParametersBuilder& ShaderParametersBuilder) const
+{
+	ShaderParametersBuilder.AddNestedStruct<FShaderParameters>();
+}
+
+void UNiagaraDataInterfaceMeshRendererInfo::SetShaderParameters(const FNiagaraDataInterfaceSetShaderParametersContext& Context) const
+{
+	FNDIMeshRendererInfoProxy& DIProxy = Context.GetProxy<FNDIMeshRendererInfoProxy>();
+	FShaderParameters* ShaderParameters = Context.GetParameterNestedStruct<FShaderParameters>();
+	if (DIProxy.GPUData.IsValid())
+	{
+		ShaderParameters->NumMeshes			= DIProxy.GPUData->GetNumMeshes();
+		ShaderParameters->MeshDataBuffer	= DIProxy.GPUData->GetMeshDataBufferSRV();
+	}
+	else
+	{
+		ShaderParameters->NumMeshes			= 0;
+		ShaderParameters->MeshDataBuffer	= FNiagaraRenderer::GetDummyFloatBuffer();
+	}
+}
+
 #if WITH_EDITOR
 
 bool UNiagaraDataInterfaceMeshRendererInfo::UpgradeFunctionCall(FNiagaraFunctionSignature& FunctionSignature)
diff --git a/Engine/Plugins/FX/Niagara/Source/Niagara/Private/NiagaraDataInterfaceOcclusion.cpp b/Engine/Plugins/FX/Niagara/Source/Niagara/Private/NiagaraDataInterfaceOcclusion.cpp
--- a/Engine/Plugins/FX/Niagara/Source/Niagara/Private/NiagaraDataInterfaceOcclusion.cpp
+++ b/Engine/Plugins/FX/Niagara/Source/Niagara/Private/NiagaraDataInterfaceOcclusion.cpp
@@ -5,10 +5,14 @@
 #include "NiagaraWorldManager.h"
 #include "ShaderParameterUtils.h"
 #include "Internationalization/Internationalization.h"
+#include "NiagaraGPUSystemTick.h"
+#include "NiagaraShaderParametersBuilder.h"
+#include "NiagaraSystemGpuComputeProxy.h"
 #include "NiagaraSystemInstance.h"
 
 #define LOCTEXT_NAMESPACE "NiagaraDataInterfaceOcclusion"
 
+const TCHAR* UNiagaraDataInterfaceOcclusion::TemplateShaderFilePath = TEXT("/Plugin/FX/Niagara/Private/NiagaraDataInterfaceOcclusion.ush");
 const FName UNiagaraDataInterfaceOcclusion::GetCameraOcclusionRectangleName(TEXT("QueryOcclusionFactorWithRectangleGPU"));
 const FName UNiagaraDataInterfaceOcclusion::GetCameraOcclusionCircleName(TEXT("QueryOcclusionFactorWithCircleGPU"));
 
@@ -88,46 +92,21 @@ bool UNiagaraDataInterfaceOcclusion::AppendCompileHash(FNiagaraCompileHashVisito
 	if (!Super::AppendCompileHash(InVisitor))
 		return false;
 
-	FSHAHash Hash = GetShaderFileHash((TEXT("/Plugin/FX/Niagara/Private/NiagaraDataInterfaceOcclusion.ush")), EShaderPlatform::SP_PCD3D_SM5);
-	InVisitor->UpdateString(TEXT("NiagaraDataInterfaceOcclusionHLSLSource"), Hash.ToString());
+	InVisitor->UpdateString(TEXT("NiagaraDataInterfaceOcclusionHLSLSource"), GetShaderFileHash(TemplateShaderFilePath, EShaderPlatform::SP_PCD3D_SM5).ToString());
+	InVisitor->UpdateShaderParameters<FShaderParameters>();
 	return true;
 }
 
 void UNiagaraDataInterfaceOcclusion::GetCommonHLSL(FString& OutHLSL)
 {
-	OutHLSL += TEXT("#include \"/Plugin/FX/Niagara/Private/NiagaraDataInterfaceOcclusion.ush\"\n");
+	OutHLSL += TEXT("#include \"/Engine/Private/DeferredShadingCommon.ush\"\n");
 }
 
 bool UNiagaraDataInterfaceOcclusion::GetFunctionHLSL(const FNiagaraDataInterfaceGPUParamInfo& ParamInfo, const FNiagaraDataInterfaceGeneratedFunction& FunctionInfo, int FunctionInstanceIndex, FString& OutHLSL)
 {
-	TMap<FString, FStringFormatArg> ArgsSample =
+	if ((FunctionInfo.DefinitionName == GetCameraOcclusionRectangleName) ||
+		(FunctionInfo.DefinitionName == GetCameraOcclusionCircleName))
 	{
-		{TEXT("FunctionName"), FunctionInfo.InstanceName},
-		{TEXT("NDIGetContextName"), TEXT("NDIOCCLUSION_MAKE_CONTEXT(") + ParamInfo.DataInterfaceHLSLSymbol + TEXT(")")},
-	};
-
-	if (FunctionInfo.DefinitionName == GetCameraOcclusionRectangleName)
-	{
-		static const TCHAR *FormatSample = TEXT(R"(
-			void {FunctionName}(in float3 In_SampleCenterWorldPos, in float In_SampleWindowWidthWorld, in float In_SampleWindowHeightWorld, in float In_SampleSteps, out float Out_VisibilityFraction, out float Out_SampleFraction)
-			{
-				{NDIGetContextName}
-				DIOcclusion_Rectangle(DIContext, In_SampleCenterWorldPos, In_SampleWindowWidthWorld, In_SampleWindowHeightWorld, In_SampleSteps, Out_VisibilityFraction, Out_SampleFraction);
-			}
-		)");
-		OutHLSL += FString::Format(FormatSample, ArgsSample);
-		return true;
-	}
-	if (FunctionInfo.DefinitionName == GetCameraOcclusionCircleName)
-	{
-		static const TCHAR *FormatSample = TEXT(R"(
-			void {FunctionName}(in float3 In_SampleCenterWorldPos, in float In_SampleWindowDiameterWorld, in float In_SampleRays, in float In_SampleStepsPerRay, out float Out_VisibilityFraction, out float Out_SampleFraction)
-			{
-				{NDIGetContextName}
-				DIOcclusion_Circle(DIContext, In_SampleCenterWorldPos, In_SampleWindowDiameterWorld, In_SampleRays, In_SampleStepsPerRay, Out_VisibilityFraction, Out_SampleFraction);
-			}
-		)");
-		OutHLSL += FString::Format(FormatSample, ArgsSample);
 		return true;
 	}
 	return false;
@@ -156,38 +135,26 @@ void UNiagaraDataInterfaceOcclusion::GetParameterDefinitionHLSL(const FNiagaraDa
 {
 	Super::GetParameterDefinitionHLSL(ParamInfo, OutHLSL);
 
-	OutHLSL += TEXT("NDIOCCLUSION_DECLARE_CONSTANTS(") + ParamInfo.DataInterfaceHLSLSymbol + TEXT(")\n");
+	TMap<FString, FStringFormatArg> TemplateArgs =
+	{
+		{TEXT("ParameterName"),	ParamInfo.DataInterfaceHLSLSymbol},
+	};
+
+	FString TemplateFile;
+	LoadShaderSourceFile(TemplateShaderFilePath, EShaderPlatform::SP_PCD3D_SM5, &TemplateFile, nullptr);
+	OutHLSL += FString::Format(*TemplateFile, TemplateArgs);
 }
 #endif
 
-struct FNiagaraDataInterfaceParametersCS_Occlusion : public FNiagaraDataInterfaceParametersCS
+void UNiagaraDataInterfaceOcclusion::BuildShaderParameters(FNiagaraShaderParametersBuilder& ShaderParametersBuilder) const
 {
-	DECLARE_TYPE_LAYOUT(FNiagaraDataInterfaceParametersCS_Occlusion, NonVirtual);
-
-	void Bind(const FNiagaraDataInterfaceGPUParamInfo& ParameterInfo, const class FShaderParameterMap& ParameterMap)
-	{
-		SystemLWCTileParam.Bind(ParameterMap, *(SystemLWCTileName + ParameterInfo.DataInterfaceHLSLSymbol));
-	}
-
-	void Set(FRHICommandList& RHICmdList, const FNiagaraDataInterfaceSetArgs& Context) const
-	{
-		check(IsInRenderingThread());
-
-		FRHIComputeShader* ComputeShaderRHI = Context.Shader.GetComputeShader();
-		SetShaderValue(RHICmdList, ComputeShaderRHI, SystemLWCTileParam, Context.SystemLWCTile);
-	}
-
-private:
-	// LWC data
-	LAYOUT_FIELD(FShaderParameter, SystemLWCTileParam);
-
-	static const FString SystemLWCTileName;
-};
-
-// LWC
-const FString FNiagaraDataInterfaceParametersCS_Occlusion::SystemLWCTileName(TEXT("SystemLWCTile_"));
+	ShaderParametersBuilder.AddNestedStruct<FShaderParameters>();
+}
 
-IMPLEMENT_TYPE_LAYOUT(FNiagaraDataInterfaceParametersCS_Occlusion);
-IMPLEMENT_NIAGARA_DI_PARAMETER(UNiagaraDataInterfaceOcclusion, FNiagaraDataInterfaceParametersCS_Occlusion);
+void UNiagaraDataInterfaceOcclusion::SetShaderParameters(const FNiagaraDataInterfaceSetShaderParametersContext& Context) const
+{
+	FShaderParameters* ShaderParameters = Context.GetParameterNestedStruct<FShaderParameters>();
+	ShaderParameters->SystemLWCTile = Context.GetSystemTick().SystemGpuComputeProxy->GetSystemLWCTile();
+}
 
 #undef LOCTEXT_NAMESPACE
diff --git a/Engine/Plugins/FX/Niagara/Source/Niagara/Private/NiagaraDataInterfaceVectorField.cpp b/Engine/Plugins/FX/Niagara/Source/Niagara/Private/NiagaraDataInterfaceVectorField.cpp
--- a/Engine/Plugins/FX/Niagara/Source/Niagara/Private/NiagaraDataInterfaceVectorField.cpp
+++ b/Engine/Plugins/FX/Niagara/Source/Niagara/Private/NiagaraDataInterfaceVectorField.cpp
@@ -8,20 +8,13 @@
 #include "NiagaraSystem.h"
 #include "NiagaraEmitterHandle.h"
 #include "NiagaraComponent.h"
+#include "NiagaraShaderParametersBuilder.h"
 #if INTEL_ISPC
 #include "NiagaraDataInterfaceVectorField.ispc.generated.h"
 #endif
 
 #define LOCTEXT_NAMESPACE "NiagaraDataInterfaceVectorField"
 
-// Global HLSL variable base names, used by HLSL.
-static const FString SamplerBaseName(TEXT("VectorFieldSampler_"));
-static const FString TextureBaseName(TEXT("VectorFieldTexture_"));
-static const FString TilingAxesBaseName(TEXT("TilingAxes_"));
-static const FString DimensionsBaseName(TEXT("Dimensions_"));
-static const FString MinBoundsBaseName(TEXT("MinBounds_"));
-static const FString MaxBoundsBaseName(TEXT("MaxBounds_"));
-
 // Global VM function names, also used by the shaders code generation methods.
 static const FName SampleVectorFieldName("SampleField");
 static const FName LoadVectorFieldName("LoadField");
@@ -361,49 +354,49 @@ void UNiagaraDataInterfaceVectorField::GetFeedback(UNiagaraSystem* InAsset, UNia
 /*--------------------------------------------------------------------------------------------------------------------------*/
 
 #if WITH_EDITORONLY_DATA
+bool UNiagaraDataInterfaceVectorField::AppendCompileHash(FNiagaraCompileHashVisitor* InVisitor) const
+{
+	bool bSuccess = Super::AppendCompileHash(InVisitor);
+	bSuccess &= InVisitor->UpdateShaderParameters<FShaderParameters>();
+	return bSuccess;
+}
+
 void UNiagaraDataInterfaceVectorField::GetParameterDefinitionHLSL(const FNiagaraDataInterfaceGPUParamInfo& ParamInfo, FString& OutHLSL)
 {
-	static const TCHAR *FormatDeclarations = TEXT(R"(
-		float3 {TilingAxesName};
-		float3 {DimensionsName};
-		float3 {MinBoundsName};
-		float3 {MaxBoundsName};
-		Texture3D {TextureName};
-		SamplerState {SamplerName};
-	)");
-	TMap<FString, FStringFormatArg> ArgsDeclarations;
-	ArgsDeclarations.Add(TEXT("TilingAxesName"), TilingAxesBaseName + ParamInfo.DataInterfaceHLSLSymbol);
-	ArgsDeclarations.Add(TEXT("DimensionsName"), DimensionsBaseName + ParamInfo.DataInterfaceHLSLSymbol);
-	ArgsDeclarations.Add(TEXT("MinBoundsName"), MinBoundsBaseName + ParamInfo.DataInterfaceHLSLSymbol);
-	ArgsDeclarations.Add(TEXT("MaxBoundsName"), MaxBoundsBaseName + ParamInfo.DataInterfaceHLSLSymbol);
-	ArgsDeclarations.Add(TEXT("TextureName"), TextureBaseName + ParamInfo.DataInterfaceHLSLSymbol);
-	ArgsDeclarations.Add(TEXT("SamplerName"), SamplerBaseName + ParamInfo.DataInterfaceHLSLSymbol);
-	OutHLSL += FString::Format(FormatDeclarations, ArgsDeclarations);
+	OutHLSL.Appendf(TEXT("float3 %s_TilingAxes;\n"), *ParamInfo.DataInterfaceHLSLSymbol);
+	OutHLSL.Appendf(TEXT("float3 %s_Dimensions;\n"), *ParamInfo.DataInterfaceHLSLSymbol);
+	OutHLSL.Appendf(TEXT("float3 %s_MinBounds;\n"), *ParamInfo.DataInterfaceHLSLSymbol);
+	OutHLSL.Appendf(TEXT("float3 %s_MaxBounds;\n"), *ParamInfo.DataInterfaceHLSLSymbol);
+	OutHLSL.Appendf(TEXT("Texture3D %s_Texture;\n"), *ParamInfo.DataInterfaceHLSLSymbol);
+	OutHLSL.Appendf(TEXT("SamplerState %s_Sampler;\n"), *ParamInfo.DataInterfaceHLSLSymbol);
 }
 
 bool UNiagaraDataInterfaceVectorField::GetFunctionHLSL(const FNiagaraDataInterfaceGPUParamInfo& ParamInfo, const FNiagaraDataInterfaceGeneratedFunction& FunctionInfo, int FunctionInstanceIndex, FString& OutHLSL)
 {
+	TMap<FString, FStringFormatArg> FormatArgs;
+	FormatArgs.Add(TEXT("FunctionName"),	FunctionInfo.InstanceName);
+	FormatArgs.Add(TEXT("TilingAxesName"),	ParamInfo.DataInterfaceHLSLSymbol + TEXT("_TilingAxes"));
+	FormatArgs.Add(TEXT("DimensionsName"),	ParamInfo.DataInterfaceHLSLSymbol + TEXT("_Dimensions"));
+	FormatArgs.Add(TEXT("MinBoundsName"),	ParamInfo.DataInterfaceHLSLSymbol + TEXT("_MinBounds"));
+	FormatArgs.Add(TEXT("MaxBoundsName"),	ParamInfo.DataInterfaceHLSLSymbol + TEXT("_MaxBounds"));
+	FormatArgs.Add(TEXT("TextureName"),		ParamInfo.DataInterfaceHLSLSymbol + TEXT("_Texture"));
+	FormatArgs.Add(TEXT("SamplerName"),		ParamInfo.DataInterfaceHLSLSymbol + TEXT("_Sampler"));
+
 	if (FunctionInfo.DefinitionName == SampleVectorFieldName)
 	{
-		static const TCHAR *FormatSample = TEXT(R"(
+		static const TCHAR* FunctionFormat = TEXT(R"(
 			void {FunctionName}(float3 In_SamplePoint, out float3 Out_Sample)
 			{
 				float3 SamplePoint = (In_SamplePoint - {MinBoundsName}) / ({MaxBoundsName} - {MinBoundsName});
 				Out_Sample = Texture3DSample({TextureName}, {SamplerName}, SamplePoint).xyz;
 			}
 		)");
-		TMap<FString, FStringFormatArg> ArgsSample;
-		ArgsSample.Add(TEXT("FunctionName"), FunctionInfo.InstanceName);
-		ArgsSample.Add(TEXT("TextureName"), TextureBaseName + ParamInfo.DataInterfaceHLSLSymbol);
-		ArgsSample.Add(TEXT("MinBoundsName"), MinBoundsBaseName + ParamInfo.DataInterfaceHLSLSymbol);
-		ArgsSample.Add(TEXT("MaxBoundsName"), MaxBoundsBaseName + ParamInfo.DataInterfaceHLSLSymbol);
-		ArgsSample.Add(TEXT("SamplerName"), SamplerBaseName + ParamInfo.DataInterfaceHLSLSymbol);
-		OutHLSL += FString::Format(FormatSample, ArgsSample);
+		OutHLSL += FString::Format(FunctionFormat, FormatArgs);
 		return true;
 	}
 	else if (FunctionInfo.DefinitionName == LoadVectorFieldName)
 	{
-		static const TCHAR* FormatSample = TEXT(R"(
+		static const TCHAR* FunctionFormat = TEXT(R"(
 			void {FunctionName}(int In_IndexX, int In_IndexY, int In_IndexZ, out float3 Out_Sample)
 			{
 				const int3 SampleIndex = int3(In_IndexX, In_IndexY, In_IndexZ);
@@ -411,75 +404,54 @@ bool UNiagaraDataInterfaceVectorField::GetFunctionHLSL(const FNiagaraDataInterfa
 				Out_Sample = {TextureName}.Load(int4(SampleIndex, MipIndex)).xyz;
 			}
 		)");
-
-		TMap<FString, FStringFormatArg> ArgsSample;
-		ArgsSample.Add(TEXT("FunctionName"), FunctionInfo.InstanceName);
-		ArgsSample.Add(TEXT("TextureName"), TextureBaseName + ParamInfo.DataInterfaceHLSLSymbol);
-		OutHLSL += FString::Format(FormatSample, ArgsSample);
+		OutHLSL += FString::Format(FunctionFormat, FormatArgs);
 		return true;
 	}
 	else if (FunctionInfo.DefinitionName == GetVectorFieldTilingAxesName)
 	{
-		static const TCHAR *FormatTilingAxes = TEXT(R"(
+		static const TCHAR* FunctionFormat = TEXT(R"(
 			void {FunctionName}(out float3 Out_TilingAxes)
 			{
 				Out_TilingAxes = {TilingAxesName};
 			}
 		)");
-		TMap<FString, FStringFormatArg> ArgsTilingAxes;
-		ArgsTilingAxes.Add(TEXT("FunctionName"), FunctionInfo.InstanceName);
-		ArgsTilingAxes.Add(TEXT("TilingAxesName"), TilingAxesBaseName + ParamInfo.DataInterfaceHLSLSymbol);
-		OutHLSL += FString::Format(FormatTilingAxes, ArgsTilingAxes);
+		OutHLSL += FString::Format(FunctionFormat, FormatArgs);
 		return true;
 	}
 	else if (FunctionInfo.DefinitionName == GetVectorFieldDimensionsName)
 	{
-		static const TCHAR *FormatDimensions = TEXT(R"(
+		static const TCHAR* FunctionFormat = TEXT(R"(
 			void {FunctionName}(out float3 Out_Dimensions)
 			{
 				Out_Dimensions = {DimensionsName};
 			}
 		)");
-		TMap<FString, FStringFormatArg> ArgsDimensions;
-		ArgsDimensions.Add(TEXT("FunctionName"), FunctionInfo.InstanceName);
-		ArgsDimensions.Add(TEXT("DimensionsName"), DimensionsBaseName + ParamInfo.DataInterfaceHLSLSymbol);
-		OutHLSL += FString::Format(FormatDimensions, ArgsDimensions);
+		OutHLSL += FString::Format(FunctionFormat, FormatArgs);
 		return true;
 	}
 	else if (FunctionInfo.DefinitionName == GetVectorFieldBoundsName)
 	{
-		static const TCHAR *FormatBounds = TEXT(R"(
+		static const TCHAR* FunctionFormat = TEXT(R"(
 			void {FunctionName}(out float3 Out_MinBounds, out float3 Out_MaxBounds)
 			{
 				Out_MinBounds = {MinBoundsName};
 				Out_MaxBounds = {MaxBoundsName};
 			}
 		)");
-		TMap<FString, FStringFormatArg> ArgsBounds;
-		ArgsBounds.Add(TEXT("FunctionName"), FunctionInfo.InstanceName);
-		ArgsBounds.Add(TEXT("MinBoundsName"), MinBoundsBaseName + ParamInfo.DataInterfaceHLSLSymbol);
-		ArgsBounds.Add(TEXT("MaxBoundsName"), MaxBoundsBaseName + ParamInfo.DataInterfaceHLSLSymbol);
-		OutHLSL += FString::Format(FormatBounds, ArgsBounds);
+		OutHLSL += FString::Format(FunctionFormat, FormatArgs);
 		return true;
 	}
 	return false;
 }
 #endif
 
-void FNiagaraDataInterfaceParametersCS_VectorField::Bind(const FNiagaraDataInterfaceGPUParamInfo& ParameterInfo, const class FShaderParameterMap& ParameterMap)
+void UNiagaraDataInterfaceVectorField::BuildShaderParameters(FNiagaraShaderParametersBuilder& ShaderParametersBuilder) const
 {
-	VectorFieldSampler.Bind(ParameterMap, *(SamplerBaseName + ParameterInfo.DataInterfaceHLSLSymbol));
-	VectorFieldTexture.Bind(ParameterMap, *(TextureBaseName + ParameterInfo.DataInterfaceHLSLSymbol));
-	TilingAxes.Bind(ParameterMap, *(TilingAxesBaseName + ParameterInfo.DataInterfaceHLSLSymbol));
-	Dimensions.Bind(ParameterMap, *(DimensionsBaseName + ParameterInfo.DataInterfaceHLSLSymbol));
-	MinBounds.Bind(ParameterMap, *(MinBoundsBaseName + ParameterInfo.DataInterfaceHLSLSymbol));
-	MaxBounds.Bind(ParameterMap, *(MaxBoundsBaseName + ParameterInfo.DataInterfaceHLSLSymbol));
+	ShaderParametersBuilder.AddNestedStruct<FShaderParameters>();
 }
 
-void FNiagaraDataInterfaceParametersCS_VectorField::Set(FRHICommandList& RHICmdList, const FNiagaraDataInterfaceSetArgs& Context) const
+void UNiagaraDataInterfaceVectorField::SetShaderParameters(const FNiagaraDataInterfaceSetShaderParametersContext& Context) const
 {
-	check(IsInRenderingThread());
-
 	// Different sampler states used by the computer shader to sample 3D vector field. 
 	// Encoded as bitflags. To sample: 
 	//     1st bit: X-axis tiling flag
@@ -498,37 +470,17 @@ void FNiagaraDataInterfaceParametersCS_VectorField::Set(FRHICommandList& RHICmdL
 		SamplerStates[7] = TStaticSamplerState<SF_Bilinear, AM_Wrap,  AM_Wrap,  AM_Wrap>::GetRHI();
 	}
 
-	// Get shader and DI
-	FRHIComputeShader* ComputeShaderRHI = Context.Shader.GetComputeShader();
-	FNiagaraDataInterfaceProxyVectorField* VFDI = static_cast<FNiagaraDataInterfaceProxyVectorField*>(Context.DataInterface);
-		
-	// Note: There is a flush in PreEditChange to make sure everything is synced up at this point 
+	FNiagaraDataInterfaceProxyVectorField& DIProxy = Context.GetProxy<FNiagaraDataInterfaceProxyVectorField>();
 
-	// Get and set 3D texture handle from the currently bound vector field.
-	if (VFDI->TextureRHI)
-	{
-		SetTextureParameter(RHICmdList, ComputeShaderRHI, VectorFieldTexture, VFDI->TextureRHI);
-	}
-	else
-	{
-		SetTextureParameter(RHICmdList, ComputeShaderRHI, VectorFieldTexture, GBlackVolumeTexture->TextureRHI);
-	}
-		
-	// Get and set sampler state
-	FRHISamplerState* SamplerState = SamplerStates[int(VFDI->bTileX) + 2 * int(VFDI->bTileY) + 4 * int(VFDI->bTileZ)];
-	SetSamplerParameter(RHICmdList, ComputeShaderRHI, VectorFieldSampler, SamplerState);
-
-	//
-	SetShaderValue(RHICmdList, ComputeShaderRHI, TilingAxes, (FVector3f)VFDI->GetTilingAxes());
-	SetShaderValue(RHICmdList, ComputeShaderRHI, Dimensions, (FVector3f)VFDI->Dimensions);
-	SetShaderValue(RHICmdList, ComputeShaderRHI, MinBounds, (FVector3f)VFDI->MinBounds);
-	SetShaderValue(RHICmdList, ComputeShaderRHI, MaxBounds, (FVector3f)VFDI->MaxBounds);
+	FShaderParameters* ShaderParameters = Context.GetParameterNestedStruct<FShaderParameters>();
+	ShaderParameters->TilingAxes	= (FVector3f)DIProxy.GetTilingAxes();
+	ShaderParameters->Dimensions	= (FVector3f)DIProxy.Dimensions;
+	ShaderParameters->MinBounds		= (FVector3f)DIProxy.MinBounds;
+	ShaderParameters->MaxBounds		= (FVector3f)DIProxy.MaxBounds;
+	ShaderParameters->Texture		= DIProxy.TextureRHI.IsValid() ? DIProxy.TextureRHI : GBlackVolumeTexture->TextureRHI;
+	ShaderParameters->Sampler		= SamplerStates[int(DIProxy.bTileX) + 2 * int(DIProxy.bTileY) + 4 * int(DIProxy.bTileZ)];
 }
 
-IMPLEMENT_TYPE_LAYOUT(FNiagaraDataInterfaceParametersCS_VectorField);
-
-IMPLEMENT_NIAGARA_DI_PARAMETER(UNiagaraDataInterfaceVectorField, FNiagaraDataInterfaceParametersCS_VectorField);
-
 /*--------------------------------------------------------------------------------------------------------------------------*/
 
 void UNiagaraDataInterfaceVectorField::GetFieldTilingAxes(FVectorVMExternalFunctionContext& Context)
diff --git a/Engine/Source/Runtime/Core/Private/UObject/UE5MainStreamObjectVersion.cpp b/Engine/Source/Runtime/Core/Private/UObject/UE5MainStreamObjectVersion.cpp
--- a/Engine/Source/Runtime/Core/Private/UObject/UE5MainStreamObjectVersion.cpp
+++ b/Engine/Source/Runtime/Core/Private/UObject/UE5MainStreamObjectVersion.cpp
@@ -14,7 +14,7 @@ TMap<FGuid, FGuid> FUE5MainStreamObjectVersion::GetSystemGuids()
 	SystemGuids.Add(DevGuids.MATERIALSHADERMAP_DERIVEDDATA_VER, FGuid("65C01D817C9A4EEFAE9E988D41A1F3DD"));
 	SystemGuids.Add(DevGuids.NANITE_DERIVEDDATA_VER, FGuid("7A3EA93546624597BEBE2ADFDE3A7628"));
 	SystemGuids.Add(DevGuids.NIAGARASHADERMAP_DERIVEDDATA_VER, FGuid("8A37C45D24F2423CBE5F8F371DE33575"));
-	SystemGuids.Add(DevGuids.Niagara_LatestScriptCompileVersion, FGuid("D3F577819D21054287503193E522CDA4"));
+	SystemGuids.Add(DevGuids.Niagara_LatestScriptCompileVersion, FGuid("7E5282F3B014E540BAA1E77751460B15"));
 	SystemGuids.Add(DevGuids.SkeletalMeshDerivedDataVersion, FGuid("FB7B2FB546DE49B2A2F88B7F69763CEE"));
 	SystemGuids.Add(DevGuids.STATICMESH_DERIVEDDATA_VER, FGuid("98F7D79A5811013825E75BBCC41ED3E9"));
 	SystemGuids.Add(DevGuids.POSESEARCHDB_DERIVEDDATA_VER, FGuid("4E595C2AC5E947D6BA9ABC874353E5BC"));
