- Fix for broken GPU buffers from LWC changes

#rnx

diff --git a/Engine/Plugins/FX/Niagara/Source/Niagara/Private/NiagaraDataInterfaceSpline.cpp b/Engine/Plugins/FX/Niagara/Source/Niagara/Private/NiagaraDataInterfaceSpline.cpp
--- a/Engine/Plugins/FX/Niagara/Source/Niagara/Private/NiagaraDataInterfaceSpline.cpp
+++ b/Engine/Plugins/FX/Niagara/Source/Niagara/Private/NiagaraDataInterfaceSpline.cpp
@@ -674,30 +674,33 @@ bool UNiagaraDataInterfaceSpline::PerInstanceTick(void* PerInstanceData, FNiagar
 				uint32 BufferSize;
 		
 				// Bind positions
-				TargetData->SplinePositionsLUT.Initialize(TEXT("SplinePositionsLUT"), sizeof(FVector4), rtShaderLUT.Positions.Num(), EPixelFormat::PF_A32B32G32R32F, BUF_Static);
-				BufferSize = rtShaderLUT.Positions.Num() * sizeof(FVector4);
-				FVector4* PositionBufferData = static_cast<FVector4*>(RHILockBuffer(TargetData->SplinePositionsLUT.Buffer, 0, BufferSize, EResourceLockMode::RLM_WriteOnly));
+				TargetData->SplinePositionsLUT.Initialize(TEXT("SplinePositionsLUT"), sizeof(FVector4f), rtShaderLUT.Positions.Num(), EPixelFormat::PF_A32B32G32R32F, BUF_Static);
+				BufferSize = rtShaderLUT.Positions.Num() * sizeof(FVector4f);
+				FVector4f* PositionBufferData = static_cast<FVector4f*>(RHILockBuffer(TargetData->SplinePositionsLUT.Buffer, 0, BufferSize, EResourceLockMode::RLM_WriteOnly));
 				for (int32 Index = 0; Index < rtShaderLUT.Positions.Num(); Index++)
 				{
-					PositionBufferData[Index] = FVector4(rtShaderLUT.Positions[Index]);
+					PositionBufferData[Index] = FVector4f(rtShaderLUT.Positions[Index].X, rtShaderLUT.Positions[Index].Y, rtShaderLUT.Positions[Index].Z);
 				}
 				RHIUnlockBuffer(TargetData->SplinePositionsLUT.Buffer);
 		
 				// Bind scales
-				TargetData->SplineScalesLUT.Initialize(TEXT("SplineScalesLUT"), sizeof(FVector4), rtShaderLUT.Scales.Num(), EPixelFormat::PF_A32B32G32R32F, BUF_Static);
-				BufferSize = rtShaderLUT.Scales.Num() * sizeof(FVector4);
-				FVector4* ScaleBufferData = static_cast<FVector4*>(RHILockBuffer(TargetData->SplineScalesLUT.Buffer, 0, BufferSize, EResourceLockMode::RLM_WriteOnly));
+				TargetData->SplineScalesLUT.Initialize(TEXT("SplineScalesLUT"), sizeof(FVector4f), rtShaderLUT.Scales.Num(), EPixelFormat::PF_A32B32G32R32F, BUF_Static);
+				BufferSize = rtShaderLUT.Scales.Num() * sizeof(FVector4f);
+				FVector4f* ScaleBufferData = static_cast<FVector4f*>(RHILockBuffer(TargetData->SplineScalesLUT.Buffer, 0, BufferSize, EResourceLockMode::RLM_WriteOnly));
 				for (int32 Index = 0; Index < rtShaderLUT.Scales.Num(); Index++)
 				{
-					ScaleBufferData[Index] = FVector4(rtShaderLUT.Scales[Index]);
+					ScaleBufferData[Index] = FVector4f(rtShaderLUT.Scales[Index].X, rtShaderLUT.Scales[Index].Y, rtShaderLUT.Scales[Index].Z);
 				}
 				RHIUnlockBuffer(TargetData->SplineScalesLUT.Buffer);
 				
 				// Bind rotations
-				TargetData->SplineRotationsLUT.Initialize(TEXT("SplineRotationsLUT"), sizeof(FQuat), rtShaderLUT.Rotations.Num(), EPixelFormat::PF_A32B32G32R32F, BUF_Static);
-				BufferSize = rtShaderLUT.Rotations.Num() * sizeof(FQuat);
-				FQuat* RotationBufferData = static_cast<FQuat*>(RHILockBuffer(TargetData->SplineRotationsLUT.Buffer, 0, BufferSize, EResourceLockMode::RLM_WriteOnly));
-				FPlatformMemory::Memcpy(RotationBufferData, rtShaderLUT.Rotations.GetData(), BufferSize);
+				TargetData->SplineRotationsLUT.Initialize(TEXT("SplineRotationsLUT"), sizeof(FQuat4f), rtShaderLUT.Rotations.Num(), EPixelFormat::PF_A32B32G32R32F, BUF_Static);
+				BufferSize = rtShaderLUT.Rotations.Num() * sizeof(FQuat4f);
+				FQuat4f* RotationBufferData = static_cast<FQuat4f*>(RHILockBuffer(TargetData->SplineRotationsLUT.Buffer, 0, BufferSize, EResourceLockMode::RLM_WriteOnly));
+				for (int32 Index=0; Index < rtShaderLUT.Rotations.Num(); Index++)
+				{
+					RotationBufferData[Index] = FQuat4f(rtShaderLUT.Rotations[Index]);
+				}
 				RHIUnlockBuffer(TargetData->SplineRotationsLUT.Buffer);
 				
 				INC_MEMORY_STAT_BY(STAT_NiagaraGPUDataInterfaceMemory, TargetData->SplinePositionsLUT.NumBytes);
