- Convert EmitterProperties DI to shader parameters
- Convert ActorComponent DI to shader parameters
- Convert SimpleCounter DI to shader parameters

diff --git a/Engine/Plugins/FX/Niagara/Shaders/Private/NiagaraDataInterfaceActorComponentTemplate.ush b/Engine/Plugins/FX/Niagara/Shaders/Private/NiagaraDataInterfaceActorComponentTemplate.ush
--- a/Engine/Plugins/FX/Niagara/Shaders/Private/NiagaraDataInterfaceActorComponentTemplate.ush
+++ b/Engine/Plugins/FX/Niagara/Shaders/Private/NiagaraDataInterfaceActorComponentTemplate.ush
@@ -1,20 +1,20 @@
 // Copyright Epic Games, Inc. All Rights Reserved.
 
-bool			Valid_{ParameterName};
-float4x4		Matrix_{ParameterName};
-float4			Rotation_{ParameterName};
-float3			Scale_{ParameterName};
+uint			{ParameterName}_Valid;
+float4x4		{ParameterName}_Matrix;
+float4			{ParameterName}_Rotation;
+float3			{ParameterName}_Scale;
 
 void GetMatrix_{ParameterName}(out bool bOutValid, out float4x4 OutMatrix)
 {
-	bOutValid = Valid_{ParameterName};
-	OutMatrix = Matrix_{ParameterName};
+	bOutValid = {ParameterName}_Valid != 0;
+	OutMatrix = {ParameterName}_Matrix;
 }
 
 void GetTransform_{ParameterName}(out bool bOutValid, out float3 OutPosition, out float4 OutRotation, out float3 OutScale)
 {
-	bOutValid = Valid_{ParameterName};
-	OutPosition = Matrix_{ParameterName}[3].xyz;
-	OutRotation = Rotation_{ParameterName};
-	OutScale = Scale_{ParameterName};
+	bOutValid = {ParameterName}_Valid != 0;
+	OutPosition = {ParameterName}_Matrix[3].xyz;
+	OutRotation = {ParameterName}_Rotation;
+	OutScale = {ParameterName}_Scale;
 }
diff --git a/Engine/Plugins/FX/Niagara/Shaders/Private/NiagaraDataInterfaceEmitterPropertiesTemplate.ush b/Engine/Plugins/FX/Niagara/Shaders/Private/NiagaraDataInterfaceEmitterPropertiesTemplate.ush
--- a/Engine/Plugins/FX/Niagara/Shaders/Private/NiagaraDataInterfaceEmitterPropertiesTemplate.ush
+++ b/Engine/Plugins/FX/Niagara/Shaders/Private/NiagaraDataInterfaceEmitterPropertiesTemplate.ush
@@ -1,26 +1,26 @@
 // Copyright Epic Games, Inc. All Rights Reserved.
 
-bool			LocalSpace_{ParameterName};
-bool			FixedBoundsValid_{ParameterName};
-float3			FixedBoundsMin_{ParameterName};
-float3			FixedBoundsMax_{ParameterName};
+uint			{ParameterName}_LocalSpace;
+uint			{ParameterName}_FixedBoundsValid;
+float3			{ParameterName}_FixedBoundsMin;
+float3			{ParameterName}_FixedBoundsMax;
 
 void GetLocalSpace_{ParameterName}(out bool bLocalSpace)
 {
-	bLocalSpace = LocalSpace_{ParameterName};
+	bLocalSpace = {ParameterName}_LocalSpace != 0;
 }
 
 void GetBounds_{ParameterName}(out bool bValid, out float3 Min, out float3 Max)
 {
 	// GPU is always fixed no dynamic bounds currently
-	bValid = FixedBoundsValid_{ParameterName};
-	Min = FixedBoundsMin_{ParameterName};
-	Max = FixedBoundsMax_{ParameterName};
+	bValid = {ParameterName}_FixedBoundsValid != 0;
+	Min = {ParameterName}_FixedBoundsMin;
+	Max = {ParameterName}_FixedBoundsMax;
 }
 
 void GetFixedBounds_{ParameterName}(out bool bValid, out float3 Min, out float3 Max)
 {
-	bValid = FixedBoundsValid_{ParameterName};
-	Min = FixedBoundsMin_{ParameterName};
-	Max = FixedBoundsMax_{ParameterName};
+	bValid = {ParameterName}_FixedBoundsValid != 0;
+	Min = {ParameterName}_FixedBoundsMin;
+	Max = {ParameterName}_FixedBoundsMax;
 }
diff --git a/Engine/Plugins/FX/Niagara/Shaders/Private/NiagaraDataInterfaceSimpleCounterTemplate.ush b/Engine/Plugins/FX/Niagara/Shaders/Private/NiagaraDataInterfaceSimpleCounterTemplate.ush
--- a/Engine/Plugins/FX/Niagara/Shaders/Private/NiagaraDataInterfaceSimpleCounterTemplate.ush
+++ b/Engine/Plugins/FX/Niagara/Shaders/Private/NiagaraDataInterfaceSimpleCounterTemplate.ush
@@ -1,22 +1,22 @@
 // Copyright Epic Games, Inc. All Rights Reserved.
 // RWInstanceCounts come from NiagaraEmitterInstanceSimulation.ush as we use counts from that buffer
 
-int CountOffset_{ParameterName};
+int {ParameterName}_CountOffset;
 
 void Get_{ParameterName}(out int CurrValue)
 {
-	CurrValue = RWInstanceCounts[CountOffset_{ParameterName}];
+	CurrValue = RWInstanceCounts[{ParameterName}_CountOffset];
 }
 
 void Exchange_{ParameterName}_UEImpureCall(in bool bExecute, in int NewValue, out int PrevValue)
 {
 	if ( bExecute )
 	{
-		InterlockedExchange(RWInstanceCounts[CountOffset_{ParameterName}], NewValue, PrevValue);
+		InterlockedExchange(RWInstanceCounts[{ParameterName}_CountOffset], NewValue, PrevValue);
 	}
 	else
 	{
-		PrevValue = RWInstanceCounts[CountOffset_{ParameterName}];
+		PrevValue = RWInstanceCounts[{ParameterName}_CountOffset];
 	}
 }
 
@@ -24,12 +24,12 @@ void Add_{ParameterName}_UEImpureCall(in bool bExecute, in int Amount, out int P
 {
 	if ( bExecute )
 	{
-		InterlockedAdd(RWInstanceCounts[CountOffset_{ParameterName}], Amount, PrevValue);
+		InterlockedAdd(RWInstanceCounts[{ParameterName}_CountOffset], Amount, PrevValue);
 		CurrValue = PrevValue + Amount;
 	}
 	else
 	{
-		PrevValue = RWInstanceCounts[CountOffset_{ParameterName}];
+		PrevValue = RWInstanceCounts[{ParameterName}_CountOffset];
 		CurrValue = PrevValue;
 	}
 }
diff --git a/Engine/Plugins/FX/Niagara/Source/Niagara/Private/DataInterface/NiagaraDataInterfaceActorComponent.cpp b/Engine/Plugins/FX/Niagara/Source/Niagara/Private/DataInterface/NiagaraDataInterfaceActorComponent.cpp
--- a/Engine/Plugins/FX/Niagara/Source/Niagara/Private/DataInterface/NiagaraDataInterfaceActorComponent.cpp
+++ b/Engine/Plugins/FX/Niagara/Source/Niagara/Private/DataInterface/NiagaraDataInterfaceActorComponent.cpp
@@ -3,6 +3,7 @@
 #include "NiagaraDataInterfaceActorComponent.h"
 #include "NiagaraTypes.h"
 #include "NiagaraCustomVersion.h"
+#include "NiagaraShaderParametersBuilder.h"
 #include "NiagaraSystemInstance.h"
 
 #include "Components/SceneComponent.h"
@@ -80,52 +81,6 @@ namespace NDIActorComponentLocal
 	};
 }
 
-//////////////////////////////////////////////////////////////////////////
-// Compute Shader Binding
-struct FNDIActorComponentCS : public FNiagaraDataInterfaceParametersCS
-{
-	DECLARE_TYPE_LAYOUT(FNDIActorComponentCS, NonVirtual);
-
-public:
-	void Bind(const FNiagaraDataInterfaceGPUParamInfo& ParameterInfo, const class FShaderParameterMap& ParameterMap)
-	{
-		using namespace NDIActorComponentLocal;
-
-		ValidParam.Bind(ParameterMap, *(NDIActorComponentLocal::ValidString + TEXT("_") + ParameterInfo.DataInterfaceHLSLSymbol));
-		MatrixParam.Bind(ParameterMap, *(NDIActorComponentLocal::MatrixString + TEXT("_") + ParameterInfo.DataInterfaceHLSLSymbol));
-		RotationParam.Bind(ParameterMap, *(NDIActorComponentLocal::RotationString + TEXT("_") + ParameterInfo.DataInterfaceHLSLSymbol));
-		ScaleParam.Bind(ParameterMap, *(NDIActorComponentLocal::ScaleString + TEXT("_") + ParameterInfo.DataInterfaceHLSLSymbol));
-	}
-
-	void Set(FRHICommandList& RHICmdList, const FNiagaraDataInterfaceSetArgs& Context) const
-	{
-		using namespace NDIActorComponentLocal;
-
-		FRHIComputeShader* ComputeShaderRHI = Context.Shader.GetComputeShader();
-		FNDIProxy* DataInterfaceProxy = static_cast<FNDIProxy*>(Context.DataInterface);
-		FInstanceData_RenderThread* InstanceData = DataInterfaceProxy->SystemInstancesToInstanceData_RT.Find(Context.SystemInstanceID);
-		check(InstanceData != nullptr);
-
-		const FMatrix44f InstanceMatrix = (FMatrix44f)InstanceData->CachedTransform.ToMatrixWithScale();
-		const FQuat4f InstanceRotation = (FQuat4f)InstanceData->CachedTransform.GetRotation();
-		const FVector3f InstanceScale = (FVector3f)InstanceData->CachedTransform.GetScale3D();
-		SetShaderValue(RHICmdList, ComputeShaderRHI, ValidParam, InstanceData->bCachedValid ? 1 : 0);
-		SetShaderValue(RHICmdList, ComputeShaderRHI, MatrixParam, InstanceMatrix);
-		SetShaderValue(RHICmdList, ComputeShaderRHI, RotationParam, InstanceRotation);
-		SetShaderValue(RHICmdList, ComputeShaderRHI, ScaleParam, InstanceScale);
-	}
-
-private:
-	LAYOUT_FIELD(FShaderParameter, ValidParam);
-	LAYOUT_FIELD(FShaderParameter, MatrixParam);
-	LAYOUT_FIELD(FShaderParameter, RotationParam);
-	LAYOUT_FIELD(FShaderParameter, ScaleParam);
-};
-
-IMPLEMENT_TYPE_LAYOUT(FNDIActorComponentCS);
-
-IMPLEMENT_NIAGARA_DI_PARAMETER(UNiagaraDataInterfaceActorComponent, FNDIActorComponentCS);
-
 //////////////////////////////////////////////////////////////////////////
 // Data Interface
 UNiagaraDataInterfaceActorComponent::UNiagaraDataInterfaceActorComponent(FObjectInitializer const& ObjectInitializer)
@@ -218,6 +173,7 @@ bool UNiagaraDataInterfaceActorComponent::AppendCompileHash(FNiagaraCompileHashV
 	bool bSuccess = Super::AppendCompileHash(InVisitor);
 	FSHAHash Hash = GetShaderFileHash(NDIActorComponentLocal::TemplateShaderFile, EShaderPlatform::SP_PCD3D_SM5);
 	InVisitor->UpdateString(TEXT("NiagaraDataInterfaceActorComponentTemplateHLSLSource"), Hash.ToString());
+	InVisitor->UpdateShaderParameters<FShaderParameters>();
 	return bSuccess;
 }
 
@@ -259,6 +215,25 @@ bool UNiagaraDataInterfaceActorComponent::UpgradeFunctionCall(FNiagaraFunctionSi
 }
 #endif
 
+void UNiagaraDataInterfaceActorComponent::BuildShaderParameters(FNiagaraShaderParametersBuilder& ShaderParametersBuilder) const
+{
+	ShaderParametersBuilder.AddNestedStruct<FShaderParameters>();
+}
+
+void UNiagaraDataInterfaceActorComponent::SetShaderParameters(const FNiagaraDataInterfaceSetShaderParametersContext& Context) const
+{
+	using namespace NDIActorComponentLocal;
+
+	FNDIProxy& DIProxy = Context.GetProxy<FNDIProxy>();
+	FInstanceData_RenderThread& InstanceData = DIProxy.SystemInstancesToInstanceData_RT.FindChecked(Context.GetSystemInstanceID());
+
+	FShaderParameters* ShaderParameters = Context.GetParameterNestedStruct<FShaderParameters>();
+	ShaderParameters->Valid		= InstanceData.bCachedValid ? 1 : 0;
+	ShaderParameters->Matrix	= (FMatrix44f)InstanceData.CachedTransform.ToMatrixWithScale();
+	ShaderParameters->Rotation	= (FQuat4f)InstanceData.CachedTransform.GetRotation();
+	ShaderParameters->Scale		= (FVector3f)InstanceData.CachedTransform.GetScale3D();
+}
+
 bool UNiagaraDataInterfaceActorComponent::InitPerInstanceData(void* PerInstanceData, FNiagaraSystemInstance* SystemInstance)
 {
 	using namespace NDIActorComponentLocal;
diff --git a/Engine/Plugins/FX/Niagara/Source/Niagara/Private/DataInterface/NiagaraDataInterfaceActorComponent.h b/Engine/Plugins/FX/Niagara/Source/Niagara/Private/DataInterface/NiagaraDataInterfaceActorComponent.h
--- a/Engine/Plugins/FX/Niagara/Source/Niagara/Private/DataInterface/NiagaraDataInterfaceActorComponent.h
+++ b/Engine/Plugins/FX/Niagara/Source/Niagara/Private/DataInterface/NiagaraDataInterfaceActorComponent.h
@@ -15,9 +15,14 @@ class NIAGARA_API UNiagaraDataInterfaceActorComponent : public UNiagaraDataInter
 {
 	GENERATED_UCLASS_BODY()
 
-public:
-	DECLARE_NIAGARA_DI_PARAMETER();
+	BEGIN_SHADER_PARAMETER_STRUCT(FShaderParameters, )
+		SHADER_PARAMETER(uint32,		Valid)
+		SHADER_PARAMETER(FMatrix44f,	Matrix)
+		SHADER_PARAMETER(FQuat4f,		Rotation)
+		SHADER_PARAMETER(FVector3f,		Scale)
+	END_SHADER_PARAMETER_STRUCT();
 
+public:
 	/** When this option is disabled, we use the previous frame's data for the skeletal mesh and can often issue the simulation early. This greatly
 	reduces overhead and allows the game thread to run faster, but comes at a tradeoff if the dependencies might leave gaps or other visual artifacts.*/
 	UPROPERTY(EditAnywhere, Category = "Performance")
@@ -48,6 +53,9 @@ public:
 	virtual bool GetFunctionHLSL(const FNiagaraDataInterfaceGPUParamInfo& ParamInfo, const FNiagaraDataInterfaceGeneratedFunction& FunctionInfo, int FunctionInstanceIndex, FString& OutHLSL) override;
 	virtual bool UpgradeFunctionCall(FNiagaraFunctionSignature& FunctionSignature) override;
 #endif
+	virtual bool UseLegacyShaderBindings() const override { return false; }
+	virtual void BuildShaderParameters(FNiagaraShaderParametersBuilder& ShaderParametersBuilder) const override;
+	virtual void SetShaderParameters(const FNiagaraDataInterfaceSetShaderParametersContext& Context) const override;
 
 	virtual bool InitPerInstanceData(void* PerInstanceData, FNiagaraSystemInstance* SystemInstance) override;
 	virtual void DestroyPerInstanceData(void* PerInstanceData, FNiagaraSystemInstance* SystemInstance) override;
diff --git a/Engine/Plugins/FX/Niagara/Source/Niagara/Private/DataInterface/NiagaraDataInterfaceEmitterProperties.cpp b/Engine/Plugins/FX/Niagara/Source/Niagara/Private/DataInterface/NiagaraDataInterfaceEmitterProperties.cpp
--- a/Engine/Plugins/FX/Niagara/Source/Niagara/Private/DataInterface/NiagaraDataInterfaceEmitterProperties.cpp
+++ b/Engine/Plugins/FX/Niagara/Source/Niagara/Private/DataInterface/NiagaraDataInterfaceEmitterProperties.cpp
@@ -3,6 +3,7 @@
 #include "NiagaraDataInterfaceEmitterProperties.h"
 #include "NiagaraEmitterInstance.h"
 #include "NiagaraComputeExecutionContext.h"
+#include "NiagaraShaderParametersBuilder.h"
 #include "NiagaraSystemInstance.h"
 
 #include "Internationalization/Internationalization.h"
@@ -73,47 +74,6 @@ namespace NDIEmitterPropertiesLocal
 
 //////////////////////////////////////////////////////////////////////////
 
-struct FNDIEmitterPropertiesCS : public FNiagaraDataInterfaceParametersCS
-{
-	DECLARE_TYPE_LAYOUT(FNDIEmitterPropertiesCS, NonVirtual);
-
-public:
-	void Bind(const FNiagaraDataInterfaceGPUParamInfo& ParameterInfo, const class FShaderParameterMap& ParameterMap)
-	{
-		using namespace NDIEmitterPropertiesLocal;
-
-		LocalSpaceParam.Bind(ParameterMap, *(TEXT("LocalSpace_") + ParameterInfo.DataInterfaceHLSLSymbol));
-		FixedBoundsValidParam.Bind(ParameterMap, *(TEXT("FixedBoundsValid_") + ParameterInfo.DataInterfaceHLSLSymbol));
-		FixedBoundsMinParam.Bind(ParameterMap, *(TEXT("FixedBoundsMin_") + ParameterInfo.DataInterfaceHLSLSymbol));
-		FixedBoundsMaxParam.Bind(ParameterMap, *(TEXT("FixedBoundsMax_") + ParameterInfo.DataInterfaceHLSLSymbol));
-	}
-
-	void Set(FRHICommandList& RHICmdList, const FNiagaraDataInterfaceSetArgs& Context) const
-	{
-		using namespace NDIEmitterPropertiesLocal;
-
-		FRHIComputeShader* ComputeShaderRHI = Context.Shader.GetComputeShader();
-		const FNDIProxy* DIProxy = static_cast<FNDIProxy*>(Context.DataInterface);
-		const FInstanceData_RenderThread* InstanceData_RT = &DIProxy->PerInstanceData_RenderThread.FindChecked(Context.SystemInstanceID);
-
-		SetShaderValue(RHICmdList, ComputeShaderRHI, LocalSpaceParam, InstanceData_RT->bLocalSpace ? 1 : 0);
-		SetShaderValue(RHICmdList, ComputeShaderRHI, FixedBoundsValidParam, InstanceData_RT->FixedBounds.IsValid ? 1 : 0);
-		SetShaderValue(RHICmdList, ComputeShaderRHI, FixedBoundsMinParam, FVector3f(InstanceData_RT->FixedBounds.Min));
-		SetShaderValue(RHICmdList, ComputeShaderRHI, FixedBoundsMaxParam, FVector3f(InstanceData_RT->FixedBounds.Max));
-	}
-
-private:
-	LAYOUT_FIELD(FShaderParameter, LocalSpaceParam);
-	LAYOUT_FIELD(FShaderParameter, FixedBoundsValidParam);
-	LAYOUT_FIELD(FShaderParameter, FixedBoundsMinParam);
-	LAYOUT_FIELD(FShaderParameter, FixedBoundsMaxParam);
-};
-
-IMPLEMENT_TYPE_LAYOUT(FNDIEmitterPropertiesCS);
-IMPLEMENT_NIAGARA_DI_PARAMETER(UNiagaraDataInterfaceEmitterProperties, FNDIEmitterPropertiesCS);
-
-//////////////////////////////////////////////////////////////////////////
-
 UNiagaraDataInterfaceEmitterProperties::UNiagaraDataInterfaceEmitterProperties(FObjectInitializer const& ObjectInitializer)
 	: Super(ObjectInitializer)
 {
@@ -317,6 +277,7 @@ bool UNiagaraDataInterfaceEmitterProperties::AppendCompileHash(FNiagaraCompileHa
 	bool bSuccess = Super::AppendCompileHash(InVisitor);
 	FSHAHash Hash = GetShaderFileHash(NDIEmitterPropertiesLocal::TemplateShaderFile, EShaderPlatform::SP_PCD3D_SM5);
 	InVisitor->UpdateString(TEXT("NiagaraDataInterfaceEmitterPropertiesTemplateHLSLSource"), Hash.ToString());
+	InVisitor->UpdateShaderParameters<FShaderParameters>();
 	return bSuccess;
 }
 
@@ -346,6 +307,26 @@ bool UNiagaraDataInterfaceEmitterProperties::GetFunctionHLSL(const FNiagaraDataI
 	return false;
 }
 #endif
+
+void UNiagaraDataInterfaceEmitterProperties::BuildShaderParameters(FNiagaraShaderParametersBuilder& ShaderParametersBuilder) const
+{
+	ShaderParametersBuilder.AddNestedStruct<FShaderParameters>();
+}
+
+void UNiagaraDataInterfaceEmitterProperties::SetShaderParameters(const FNiagaraDataInterfaceSetShaderParametersContext& Context) const
+{
+	using namespace NDIEmitterPropertiesLocal;
+
+	FNDIProxy& DIProxy = Context.GetProxy<FNDIProxy>();
+	FInstanceData_RenderThread& InstanceData_RT = DIProxy.PerInstanceData_RenderThread.FindChecked(Context.GetSystemInstanceID());
+
+	FShaderParameters* ShaderParameters	= Context.GetParameterNestedStruct<FShaderParameters>();
+	ShaderParameters->LocalSpace		= InstanceData_RT.bLocalSpace ? 1 : 0;
+	ShaderParameters->FixedBoundsValid	= InstanceData_RT.FixedBounds.IsValid;
+	ShaderParameters->FixedBoundsMin	= FVector3f(InstanceData_RT.FixedBounds.Min);
+	ShaderParameters->FixedBoundsMax	= FVector3f(InstanceData_RT.FixedBounds.Max);
+}
+
 #if WITH_EDITOR	
 void UNiagaraDataInterfaceEmitterProperties::GetFeedback(UNiagaraSystem* Asset, UNiagaraComponent* Component, TArray<FNiagaraDataInterfaceError>& OutErrors, TArray<FNiagaraDataInterfaceFeedback>& OutWarnings, TArray<FNiagaraDataInterfaceFeedback>& OutInfo)
 {
diff --git a/Engine/Plugins/FX/Niagara/Source/Niagara/Private/DataInterface/NiagaraDataInterfaceEmitterProperties.h b/Engine/Plugins/FX/Niagara/Source/Niagara/Private/DataInterface/NiagaraDataInterfaceEmitterProperties.h
--- a/Engine/Plugins/FX/Niagara/Source/Niagara/Private/DataInterface/NiagaraDataInterfaceEmitterProperties.h
+++ b/Engine/Plugins/FX/Niagara/Source/Niagara/Private/DataInterface/NiagaraDataInterfaceEmitterProperties.h
@@ -16,9 +16,14 @@ class UNiagaraDataInterfaceEmitterProperties : public UNiagaraDataInterface
 {
 	GENERATED_UCLASS_BODY()
 
-public:
-	DECLARE_NIAGARA_DI_PARAMETER();
+	BEGIN_SHADER_PARAMETER_STRUCT(FShaderParameters, )
+		SHADER_PARAMETER(uint32,	LocalSpace)
+		SHADER_PARAMETER(uint32,	FixedBoundsValid)
+		SHADER_PARAMETER(FVector3f,	FixedBoundsMin)
+		SHADER_PARAMETER(FVector3f,	FixedBoundsMax)
+	END_SHADER_PARAMETER_STRUCT();
 
+public:
 	/** Selects which emitter the data interface will bind to, i.e the emitter we are contained within or a named emitter. */
 	UPROPERTY(EditAnywhere, Category = "Emitter")
 	FNiagaraDataInterfaceEmitterBinding EmitterBinding;
@@ -49,6 +54,9 @@ public:
 	virtual void GetParameterDefinitionHLSL(const FNiagaraDataInterfaceGPUParamInfo& ParamInfo, FString& OutHLSL) override;
 	virtual bool GetFunctionHLSL(const FNiagaraDataInterfaceGPUParamInfo& ParamInfo, const FNiagaraDataInterfaceGeneratedFunction& FunctionInfo, int FunctionInstanceIndex, FString& OutHLSL) override;
 #endif
+	virtual bool UseLegacyShaderBindings() const override { return false; }
+	virtual void BuildShaderParameters(FNiagaraShaderParametersBuilder& ShaderParametersBuilder) const override;
+	virtual void SetShaderParameters(const FNiagaraDataInterfaceSetShaderParametersContext& Context) const override;
 #if WITH_EDITOR	
 	virtual void GetFeedback(UNiagaraSystem* Asset, UNiagaraComponent* Component, TArray<FNiagaraDataInterfaceError>& OutErrors, TArray<FNiagaraDataInterfaceFeedback>& Warnings, TArray<FNiagaraDataInterfaceFeedback>& Info) override;
 #endif
diff --git a/Engine/Plugins/FX/Niagara/Source/Niagara/Private/DataInterface/NiagaraDataInterfaceSimpleCounter.cpp b/Engine/Plugins/FX/Niagara/Source/Niagara/Private/DataInterface/NiagaraDataInterfaceSimpleCounter.cpp
--- a/Engine/Plugins/FX/Niagara/Source/Niagara/Private/DataInterface/NiagaraDataInterfaceSimpleCounter.cpp
+++ b/Engine/Plugins/FX/Niagara/Source/Niagara/Private/DataInterface/NiagaraDataInterfaceSimpleCounter.cpp
@@ -4,6 +4,7 @@
 #include "NiagaraClearCounts.h"
 #include "NiagaraGpuComputeDispatchInterface.h"
 #include "NiagaraGpuReadbackManager.h"
+#include "NiagaraShaderParametersBuilder.h"
 #include "NiagaraSystemInstance.h"
 
 #include "Internationalization/Internationalization.h"
@@ -111,40 +112,6 @@ struct FNDISimpleCounterProxy : public FNiagaraDataInterfaceProxy
 	TMap<FNiagaraSystemInstanceID, FNDISimpleCounterInstanceData_GameThread*>	PerInstanceData_GameThread;
 };
 
-struct FNDISimpleCounterCS : public FNiagaraDataInterfaceParametersCS
-{
-	DECLARE_TYPE_LAYOUT(FNDISimpleCounterCS, NonVirtual);
-
-public:
-	void Bind(const FNiagaraDataInterfaceGPUParamInfo& ParameterInfo, const class FShaderParameterMap& ParameterMap)
-	{
-		CountOffsetParam.Bind(ParameterMap, *(TEXT("CountOffset_") + ParameterInfo.DataInterfaceHLSLSymbol));
-	}
-
-	void Set(FRHICommandList& RHICmdList, const FNiagaraDataInterfaceSetArgs& Context) const
-	{
-		using namespace NDISimpleCounterLocal;
-
-		check(IsInRenderingThread());
-
-		if (CountOffsetParam.IsBound())
-		{
-			FNDISimpleCounterProxy* DIProxy = static_cast<FNDISimpleCounterProxy*>(Context.DataInterface);
-			FNDISimpleCounterInstanceData_RenderThread* InstanceData = &DIProxy->PerInstanceData_RenderThread.FindOrAdd(Context.SystemInstanceID);
-			check(InstanceData->CountOffset != INDEX_NONE);
-
-			FRHIComputeShader* ComputeShaderRHI = Context.Shader.GetComputeShader();
-			SetShaderValue(RHICmdList, ComputeShaderRHI, CountOffsetParam, InstanceData->CountOffset);
-		}
-	}
-
-private:
-	LAYOUT_FIELD(FShaderParameter, CountOffsetParam);
-};
-
-IMPLEMENT_TYPE_LAYOUT(FNDISimpleCounterCS);
-IMPLEMENT_NIAGARA_DI_PARAMETER(UNiagaraDataInterfaceSimpleCounter, FNDISimpleCounterCS);
-
 //////////////////////////////////////////////////////////////////////////
 
 UNiagaraDataInterfaceSimpleCounter::UNiagaraDataInterfaceSimpleCounter(FObjectInitializer const& ObjectInitializer)
@@ -404,10 +371,27 @@ bool UNiagaraDataInterfaceSimpleCounter::AppendCompileHash(FNiagaraCompileHashVi
 	bool bSuccess = Super::AppendCompileHash(InVisitor);
 	FSHAHash Hash = GetShaderFileHash(NDISimpleCounterLocal::TemplateShaderFile, EShaderPlatform::SP_PCD3D_SM5);
 	InVisitor->UpdateString(TEXT("NiagaraDataInterfaceSimpleCounterTemplateHLSLSource"), Hash.ToString());
+	InVisitor->UpdateShaderParameters<FShaderParameters>();
 	return bSuccess;
 }
 #endif
 
+void UNiagaraDataInterfaceSimpleCounter::BuildShaderParameters(FNiagaraShaderParametersBuilder& ShaderParametersBuilder) const
+{
+	ShaderParametersBuilder.AddNestedStruct<FShaderParameters>();
+}
+
+void UNiagaraDataInterfaceSimpleCounter::SetShaderParameters(const FNiagaraDataInterfaceSetShaderParametersContext& Context) const
+{
+	using namespace NDISimpleCounterLocal;
+
+	FNDISimpleCounterProxy& DIProxy = Context.GetProxy<FNDISimpleCounterProxy>();
+	FNDISimpleCounterInstanceData_RenderThread& InstanceData = DIProxy.PerInstanceData_RenderThread.FindOrAdd(Context.GetSystemInstanceID());
+
+	FShaderParameters* ShaderParameters = Context.GetParameterNestedStruct<FShaderParameters>();
+	ShaderParameters->CountOffset = InstanceData.CountOffset;
+}
+
 void UNiagaraDataInterfaceSimpleCounter::PushToRenderThreadImpl()
 {
 	TArray<TPair<FNiagaraSystemInstanceID, int32>, TInlineAllocator<4>> DataToPush;
diff --git a/Engine/Plugins/FX/Niagara/Source/Niagara/Private/DataInterface/NiagaraDataInterfaceSimpleCounter.h b/Engine/Plugins/FX/Niagara/Source/Niagara/Private/DataInterface/NiagaraDataInterfaceSimpleCounter.h
--- a/Engine/Plugins/FX/Niagara/Source/Niagara/Private/DataInterface/NiagaraDataInterfaceSimpleCounter.h
+++ b/Engine/Plugins/FX/Niagara/Source/Niagara/Private/DataInterface/NiagaraDataInterfaceSimpleCounter.h
@@ -16,9 +16,12 @@ UCLASS(EditInlineNew, Category = "Counting", meta = (DisplayName = "Simple Count
 class NIAGARA_API UNiagaraDataInterfaceSimpleCounter : public UNiagaraDataInterface
 {
 	GENERATED_UCLASS_BODY()
-public:
-	DECLARE_NIAGARA_DI_PARAMETER();
 
+	BEGIN_SHADER_PARAMETER_STRUCT(FShaderParameters, )
+		SHADER_PARAMETER(int32,		CountOffset)
+	END_SHADER_PARAMETER_STRUCT();
+
+public:
 	// UObject Interface
 	virtual void PostInitProperties() override;
 #if WITH_EDITOR
@@ -43,6 +46,9 @@ public:
 	virtual bool GetFunctionHLSL(const FNiagaraDataInterfaceGPUParamInfo& ParamInfo, const FNiagaraDataInterfaceGeneratedFunction& FunctionInfo, int FunctionInstanceIndex, FString& OutHLSL) override;
 	virtual bool AppendCompileHash(FNiagaraCompileHashVisitor* InVisitor) const override;
 #endif
+	virtual bool UseLegacyShaderBindings() const override { return false; }
+	virtual void BuildShaderParameters(FNiagaraShaderParametersBuilder& ShaderParametersBuilder) const override;
+	virtual void SetShaderParameters(const FNiagaraDataInterfaceSetShaderParametersContext& Context) const override;
 	virtual void PushToRenderThreadImpl() override;
 	// UNiagaraDataInterface Interface End
 
