- PVS-Studio fix for uninitialized variable

#rnx

diff --git a/Engine/Plugins/FX/Niagara/Source/Niagara/Private/NiagaraGPUProfiler.cpp b/Engine/Plugins/FX/Niagara/Source/Niagara/Private/NiagaraGPUProfiler.cpp
--- a/Engine/Plugins/FX/Niagara/Source/Niagara/Private/NiagaraGPUProfiler.cpp
+++ b/Engine/Plugins/FX/Niagara/Source/Niagara/Private/NiagaraGPUProfiler.cpp
@@ -203,8 +203,8 @@ bool FNiagaraGPUProfiler::ProcessFrame(FRHICommandListImmediate& RHICmdList, FGp
 		StageResults.DurationMicroseconds = 0;
 		if ( StageTimer.StartQuery.GetQuery() != nullptr )
 		{
-			uint64 StartMicroseconds;
-			uint64 EndMicroseconds;
+			uint64 StartMicroseconds = 0;
+			uint64 EndMicroseconds = 0;
 			ensure(RHICmdList.GetRenderQueryResult(StageTimer.StartQuery.GetQuery(), StartMicroseconds, false));
 			ensure(RHICmdList.GetRenderQueryResult(StageTimer.EndQuery.GetQuery(), EndMicroseconds, false));
 			StageResults.DurationMicroseconds = EndMicroseconds - StartMicroseconds;
@@ -219,8 +219,8 @@ bool FNiagaraGPUProfiler::ProcessFrame(FRHICommandListImmediate& RHICmdList, FGp
 	{
 		auto& DispatchResults = FrameResults->DispatchResults.AddDefaulted_GetRef();
 			
-		uint64 StartMicroseconds;
-		uint64 EndMicroseconds;
+		uint64 StartMicroseconds = 0;
+		uint64 EndMicroseconds = 0;
 		ensure(RHICmdList.GetRenderQueryResult(DispatchTimer.StartQuery.GetQuery(), StartMicroseconds, false));
 		ensure(RHICmdList.GetRenderQueryResult(DispatchTimer.EndQuery.GetQuery(), EndMicroseconds, false));
 		DispatchTimer.StartQuery.ReleaseQuery();
