- Initial pass for Niagara Simulation Cache
- Currently a full restart cache per frame generated
- Only contains system simulation & particle data set information
- Highly experimental, while it works in cooked and uncooked we fully expect the data formats will change

diff --git a/Engine/Plugins/FX/Niagara/Source/Niagara/Classes/NiagaraBakerOutputSimCache.h b/Engine/Plugins/FX/Niagara/Source/Niagara/Classes/NiagaraBakerOutputSimCache.h
new file mode 100644
--- /dev/null
+++ b/Engine/Plugins/FX/Niagara/Source/Niagara/Classes/NiagaraBakerOutputSimCache.h
@@ -0,0 +1,31 @@
+// Copyright Epic Games, Inc. All Rights Reserved.
+
+#pragma once
+
+#include "NiagaraBakerOutput.h"
+#include "NiagaraSimCache.h"
+#include "NiagaraBakerOutputSimCache.generated.h"
+
+UCLASS(meta = (DisplayName = "Simulation Cache Output"))
+class NIAGARA_API UNiagaraBakerOutputSimCache : public UNiagaraBakerOutput
+{
+	GENERATED_BODY()
+
+public:
+	UNiagaraBakerOutputSimCache(const FObjectInitializer& Init)
+	{
+	}
+
+	UPROPERTY(EditAnywhere, Category = "Settings")
+	FString SimCacheAssetPathFormat = TEXT("{AssetFolder}/{AssetName}_SimCache_{OutputName}");
+
+	virtual bool Equals(const UNiagaraBakerOutput& Other) const override;
+
+#if WITH_EDITOR
+	FString MakeOutputName() const override;
+#endif
+
+#if WITH_EDITORONLY_DATA
+	virtual void PostEditChangeProperty(struct FPropertyChangedEvent& PropertyChangedEvent) override;
+#endif
+};
diff --git a/Engine/Plugins/FX/Niagara/Source/Niagara/Classes/NiagaraDataSetDebugAccessor.h b/Engine/Plugins/FX/Niagara/Source/Niagara/Classes/NiagaraDataSetDebugAccessor.h
--- a/Engine/Plugins/FX/Niagara/Source/Niagara/Classes/NiagaraDataSetDebugAccessor.h
+++ b/Engine/Plugins/FX/Niagara/Source/Niagara/Classes/NiagaraDataSetDebugAccessor.h
@@ -14,7 +14,7 @@ struct FNiagaraDataSetDebugAccessor
 	int32 ReadInt(const FNiagaraDataBuffer* DataBuffer, uint32 Instance, uint32 Component) const;
 
 	template<typename TString>
-	void StringAppend(TString& StringType, FNiagaraDataBuffer* DataBuffer, uint32 Instance) const
+	void StringAppend(TString& StringType, const FNiagaraDataBuffer* DataBuffer, uint32 Instance) const
 	{
 		if (IsFloat() || IsHalf())
 		{
diff --git a/Engine/Plugins/FX/Niagara/Source/Niagara/Classes/NiagaraEmitterInstance.h b/Engine/Plugins/FX/Niagara/Source/Niagara/Classes/NiagaraEmitterInstance.h
--- a/Engine/Plugins/FX/Niagara/Source/Niagara/Classes/NiagaraEmitterInstance.h
+++ b/Engine/Plugins/FX/Niagara/Source/Niagara/Classes/NiagaraEmitterInstance.h
@@ -24,6 +24,8 @@ struct FNiagaraEmitterCompiledData;
 */
 class FNiagaraEmitterInstance
 {
+	friend class UNiagaraSimCache;
+
 private:
 	struct FEventInstanceData
 	{
@@ -88,9 +90,6 @@ public:
 	FORCEINLINE bool IsInactive()const { return ExecutionState == ENiagaraExecutionState::Inactive; }
 	FORCEINLINE bool IsComplete()const { return ExecutionState == ENiagaraExecutionState::Complete || ExecutionState == ENiagaraExecutionState::Disabled; }
 
-	/** Create a new NiagaraRenderer. The old renderer is not immediately deleted, but instead put in the ToBeRemoved list.*/
-	//void NIAGARA_API UpdateEmitterRenderer(ERHIFeatureLevel::Type FeatureLevel, TArray<NiagaraRenderer*>& ToBeAddedList, TArray<NiagaraRenderer*>& ToBeRemovedList);
-
 private:
 	NIAGARA_API int32 GetNumParticlesGPUInternal() const;
 public:
diff --git a/Engine/Plugins/FX/Niagara/Source/Niagara/Classes/NiagaraSimCache.h b/Engine/Plugins/FX/Niagara/Source/Niagara/Classes/NiagaraSimCache.h
new file mode 100644
--- /dev/null
+++ b/Engine/Plugins/FX/Niagara/Source/Niagara/Classes/NiagaraSimCache.h
@@ -0,0 +1,196 @@
+// Copyright Epic Games, Inc. All Rights Reserved.
+
+#pragma once
+
+#include "NiagaraCommon.h"
+#include "NiagaraScript.h"
+#include "NiagaraSimCache.generated.h"
+
+class UNiagaraComponent;
+
+USTRUCT()
+struct FNiagaraSimCacheDataBuffers
+{
+	GENERATED_BODY()
+
+	UPROPERTY()
+	uint32 NumInstances = 0;
+
+	UPROPERTY()
+	TArray<uint8> FloatData;
+
+	UPROPERTY()
+	TArray<uint8> HalfData;
+
+	UPROPERTY()
+	TArray<uint8> Int32Data;
+
+	UPROPERTY()
+	TArray<int32> IDToIndexTable;
+
+	UPROPERTY()
+	uint32 IDAcquireTag = 0;
+};
+
+USTRUCT()
+struct FNiagaraSimCacheEmitterFrame
+{
+	GENERATED_BODY()
+
+	//-TODO: We may not require these
+	UPROPERTY()
+	FBox LocalBounds = FBox(EForceInit::ForceInit);
+
+	UPROPERTY()
+	int32 TotalSpawnedParticles = 0;
+
+	UPROPERTY()
+	FNiagaraSimCacheDataBuffers ParticleDataBuffers;
+};
+
+USTRUCT()
+struct FNiagaraSimCacheSystemFrame
+{
+	GENERATED_BODY()
+		
+	UPROPERTY()
+	FBox LocalBounds = FBox(EForceInit::ForceInit);
+
+	UPROPERTY()
+	FNiagaraSimCacheDataBuffers SystemDataBuffers;
+};
+
+USTRUCT()
+struct FNiagaraSimCacheFrame
+{
+	GENERATED_BODY()
+
+	UPROPERTY()
+	FNiagaraSimCacheSystemFrame SystemData;
+
+	UPROPERTY()
+	TArray<FNiagaraSimCacheEmitterFrame> EmitterData;
+
+	//-TODO: DataInterface cache of data
+};
+
+USTRUCT()
+struct FNiagaraSimCacheVariable
+{
+	GENERATED_BODY()
+
+	UPROPERTY()
+	FNiagaraVariableBase Variable;
+
+	UPROPERTY()
+	uint16 FloatOffset = INDEX_NONE;
+
+	UPROPERTY()
+	uint16 FloatCount = 0;
+
+	UPROPERTY()
+	uint16 HalfOffset = INDEX_NONE;
+
+	UPROPERTY()
+	uint16 HalfCount = 0;
+
+	UPROPERTY()
+	uint16 Int32Offset = INDEX_NONE;
+
+	UPROPERTY()
+	uint16 Int32Count = 0;
+};
+
+USTRUCT()
+struct FNiagaraSimCacheDataBuffersLayout
+{
+	GENERATED_BODY()
+
+	UPROPERTY()
+	FName LayoutName;
+
+	UPROPERTY()
+	ENiagaraSimTarget SimTarget = ENiagaraSimTarget::CPUSim;
+
+	UPROPERTY()
+	TArray<FNiagaraSimCacheVariable> Variables;
+
+	UPROPERTY()
+	uint16 FloatCount = 0;
+
+	UPROPERTY()
+	uint16 HalfCount = 0;
+
+	UPROPERTY()
+	uint16 Int32Count = 0;
+
+	UPROPERTY(transient)
+	TArray<uint16> ComponentMappingsToDataBuffer;
+
+	UPROPERTY(transient)
+	TArray<uint16> ComponentMappingsFromDataBuffer;
+};
+
+USTRUCT()
+struct FNiagaraSimCacheLayout
+{
+	GENERATED_BODY()
+
+	UPROPERTY()
+	FNiagaraSimCacheDataBuffersLayout SystemLayout;
+
+	UPROPERTY()
+	TArray<FNiagaraSimCacheDataBuffersLayout> EmitterLayouts;
+};
+
+UCLASS(Experimental)
+class NIAGARA_API UNiagaraSimCache : public UObject
+{
+	GENERATED_UCLASS_BODY()
+
+	UPROPERTY(VisibleAnywhere, Category=SimCache, meta=(DisplayName="Niagara System"))
+	TSoftObjectPtr<UNiagaraSystem> SoftNiagaraSystem;
+
+	UPROPERTY(VisibleAnywhere, Category=SimCache)
+	float StartSeconds = 0.0f;
+
+	UPROPERTY(VisibleAnywhere, Category=SimCache)
+	float DurationSeconds = 0.0f;
+
+	UPROPERTY(transient)
+	bool bNeedsReadComponentMappingRecache = true;
+
+#if WITH_EDITORONLY_DATA
+	UPROPERTY(transient)
+	TArray<FNiagaraVMExecutableDataId> CachedScriptVMIds;
+#endif
+
+	UPROPERTY()
+	FNiagaraSimCacheLayout CacheLayout;
+
+	UPROPERTY()
+	TArray<FNiagaraSimCacheFrame> CacheFrames;
+
+	// UObject Interface
+	virtual bool IsReadyForFinishDestroy() override;
+	// UObject Interface
+
+	bool IsCacheValid() const { return SoftNiagaraSystem.IsNull() == false; }
+
+	void BeginWrite(UNiagaraComponent* NiagaraComponent);
+	void WriteFrame(UNiagaraComponent* NiagaraComponent);
+	void EndWrite();
+
+	bool CanRead(UNiagaraSystem* NiagaraSystem);
+	bool Read(float TimeSeconds, FNiagaraSystemInstance* SystemInstance) const;
+	bool ReadFrame(int32 FrameIndex, FNiagaraSystemInstance* SystemInstance) const;
+
+	/*
+	Create a simulation cache from the component's current data.
+	The cache can be applied to a component to warmup the system, for example, or can be used to debug the system.
+	*/
+	static UNiagaraSimCache* CreateSingleFrame(UObject* OuterObject, UNiagaraComponent* NiagaraComponent);
+
+private:
+	mutable std::atomic<int32> PendingCommandsInFlight;
+};
diff --git a/Engine/Plugins/FX/Niagara/Source/Niagara/Private/NiagaraBakerOutputSimCache.cpp b/Engine/Plugins/FX/Niagara/Source/Niagara/Private/NiagaraBakerOutputSimCache.cpp
new file mode 100644
--- /dev/null
+++ b/Engine/Plugins/FX/Niagara/Source/Niagara/Private/NiagaraBakerOutputSimCache.cpp
@@ -0,0 +1,26 @@
+// Copyright Epic Games, Inc. All Rights Reserved.
+
+#include "NiagaraBakerOutputSimCache.h"
+#include "NiagaraBakerSettings.h"
+
+bool UNiagaraBakerOutputSimCache::Equals(const UNiagaraBakerOutput& OtherBase) const
+{
+	const UNiagaraBakerOutputSimCache& Other = *CastChecked<UNiagaraBakerOutputSimCache>(&OtherBase);
+	return
+		Super::Equals(Other);
+}
+
+#if WITH_EDITOR
+FString UNiagaraBakerOutputSimCache::MakeOutputName() const
+{
+	return FString::Printf(TEXT("SimCache_%d"), GetFName().GetNumber());
+}
+#endif
+
+#if WITH_EDITORONLY_DATA
+void UNiagaraBakerOutputSimCache::PostEditChangeProperty(struct FPropertyChangedEvent& PropertyChangedEvent)
+{
+	Super::PostEditChangeProperty(PropertyChangedEvent);
+
+}
+#endif
diff --git a/Engine/Plugins/FX/Niagara/Source/Niagara/Private/NiagaraComponent.cpp b/Engine/Plugins/FX/Niagara/Source/Niagara/Private/NiagaraComponent.cpp
--- a/Engine/Plugins/FX/Niagara/Source/Niagara/Private/NiagaraComponent.cpp
+++ b/Engine/Plugins/FX/Niagara/Source/Niagara/Private/NiagaraComponent.cpp
@@ -704,7 +704,28 @@ void UNiagaraComponent::TickComponent(float DeltaSeconds, enum ELevelTick TickTy
 			SystemInstanceController->Reset(FNiagaraSystemInstance::EResetMode::ReInit);
 		}
 
-		if (AgeUpdateMode == ENiagaraAgeUpdateMode::TickDeltaTime)
+		// If we have a sim cache attached then use that
+		if ( SimCache != nullptr )
+		{
+			float Age = SystemInstanceController->GetAge();
+
+			if ((AgeUpdateMode == ENiagaraAgeUpdateMode::DesiredAge) || (AgeUpdateMode == ENiagaraAgeUpdateMode::DesiredAgeNoSeek))
+			{
+				if ( DesiredAge < Age )
+				{
+					SystemInstanceController->Reset(FNiagaraSystemInstance::EResetMode::ResetAll);
+				}
+				Age = DesiredAge;
+			}
+			else
+			{
+				Age += DeltaSeconds;
+			}
+
+			SystemInstanceController->SimCacheTick_GameThread(SimCache, Age, DeltaSeconds, (ThisTickFunction && ThisTickFunction->IsCompletionHandleValid()) ? ThisTickFunction->GetCompletionHandle() : nullptr);
+		}
+		// Otherwise we need to manually tick
+		else if (AgeUpdateMode == ENiagaraAgeUpdateMode::TickDeltaTime)
 		{
 			SystemInstanceController->ManualTick(DeltaSeconds * CustomTimeDilation, (ThisTickFunction && ThisTickFunction->IsCompletionHandleValid()) ? ThisTickFunction->GetCompletionHandle() : nullptr);
 		}
@@ -3447,6 +3468,17 @@ void UNiagaraComponent::SetPreviewLODDistance(bool bInEnablePreviewLODDistance,
 	}
 }
 
+void UNiagaraComponent::SetSimCache(UNiagaraSimCache* InSimCache)
+{
+	SimCache = InSimCache;
+	UpdateInstanceSoloMode();
+}
+
+UNiagaraSimCache* UNiagaraComponent::GetSimCache() const
+{
+	return SimCache;
+}
+
 void UNiagaraComponent::SetAllowScalability(bool bAllow)
 {
 	bAllowScalability = bAllow; 
@@ -3580,7 +3612,7 @@ void UNiagaraComponent::SetForceSolo(bool bInForceSolo)
 
 bool UNiagaraComponent::RequiresSoloMode() const
 {
-	return bForceSolo || !FMath::IsNearlyEqual(CustomTimeDilation, 1.0f);
+	return bForceSolo || SimCache != nullptr || !FMath::IsNearlyEqual(CustomTimeDilation, 1.0f);
 }
 
 void UNiagaraComponent::UpdateInstanceSoloMode()
diff --git a/Engine/Plugins/FX/Niagara/Source/Niagara/Private/NiagaraDataSetReadback.cpp b/Engine/Plugins/FX/Niagara/Source/Niagara/Private/NiagaraDataSetReadback.cpp
new file mode 100644
--- /dev/null
+++ b/Engine/Plugins/FX/Niagara/Source/Niagara/Private/NiagaraDataSetReadback.cpp
@@ -0,0 +1,155 @@
+// Copyright Epic Games, Inc. All Rights Reserved.
+
+#include "NiagaraDataSetReadback.h"
+#include "NiagaraComputeExecutionContext.h"
+#include "NiagaraEmitterInstance.h"
+#include "NiagaraSystemInstance.h"
+#include "NiagaraGpuComputeDispatchInterface.h"
+#include "NiagaraGpuReadbackManager.h"
+
+#include "Async/Async.h"
+
+class FNiagaraEmitterInstance;
+
+void FNiagaraDataSetReadback::SetReadbackRead(FOnReadbackReady InOnReadbackReady)
+{
+	check(IsReady());
+	OnReadbackReady = InOnReadbackReady;
+}
+
+void FNiagaraDataSetReadback::EnqueueReadback(FNiagaraEmitterInstance* EmitterInstance)
+{
+	check(EmitterInstance);
+	check(IsReady());
+
+	SourceName = EmitterInstance->GetEmitterHandle().GetName();
+	DataSet.Init(&EmitterInstance->GetData().GetCompiledData());
+	if ( FNiagaraComputeExecutionContext* GPUExecContext = EmitterInstance->GetGPUContext() )
+	{
+		ParameterStore = GPUExecContext->CombinedParamStore;
+
+		FNiagaraSystemInstance* SystemInstance = EmitterInstance->GetParentSystemInstance();
+
+		++PendingReadbacks;
+		ENQUEUE_RENDER_COMMAND(NiagaraDataSetReadback)
+		(
+			[RT_DataSetReadback=AsShared(), RT_ComputeDispatchInterface=SystemInstance->GetComputeDispatchInterface(), RT_GPUExecContext=GPUExecContext](FRHICommandListImmediate& RHICmdList)
+			{
+				RT_DataSetReadback->GPUReadbackInternal(RHICmdList, RT_ComputeDispatchInterface, RT_GPUExecContext);
+			}
+		);
+	}
+	else
+	{
+		const FNiagaraDataSet& SourceDataSet = EmitterInstance->GetData();
+		if (FNiagaraDataBuffer* SourceDataBuffer = SourceDataSet.GetCurrentData())
+		{
+			EmitterInstance->GetData().CopyTo(DataSet, 0, SourceDataBuffer->GetNumInstances());
+		}
+		else
+		{
+			DataSet.CopyFromGPUReadback(nullptr, nullptr, nullptr, 0, 0, 0, 0, 0);
+		}
+		ParameterStore = EmitterInstance->GetUpdateExecutionContext().Parameters;
+	}
+}
+
+void FNiagaraDataSetReadback::ImmediateReadback(FNiagaraEmitterInstance* EmitterInstance)
+{
+	EnqueueReadback(EmitterInstance);
+	if ( !IsReady() )
+	{
+		FNiagaraSystemInstance* SystemInstance = EmitterInstance->GetParentSystemInstance();
+		FNiagaraGpuComputeDispatchInterface* DispatchInterface = SystemInstance->GetComputeDispatchInterface();
+		ENQUEUE_RENDER_COMMAND(NiagaraFlushReadback)
+		(
+			[RT_ComputeDispatchInterface=SystemInstance->GetComputeDispatchInterface()](FRHICommandListImmediate& RHICmdList)
+			{
+				RT_ComputeDispatchInterface->GetGpuReadbackManager()->WaitCompletion(RHICmdList);
+			}
+		);
+		FlushRenderingCommands();
+		check(IsReady());
+	}
+}
+
+void FNiagaraDataSetReadback::GPUReadbackInternal(FRHICommandListImmediate& RHICmdList, FNiagaraGpuComputeDispatchInterface* DispatchInterface, FNiagaraComputeExecutionContext* GPUContext)
+{
+	FNiagaraGpuReadbackManager* ReadbackManager = DispatchInterface->GetGpuReadbackManager();
+	FNiagaraDataBuffer* CurrentDataBuffer = GPUContext->MainDataSet->GetCurrentData();
+	if (CurrentDataBuffer == nullptr || ReadbackManager == nullptr)
+	{
+		DataSet.CopyFromGPUReadback(nullptr, nullptr, nullptr, 0, 0, 0, 0, 0);
+		return;
+	}
+
+	const uint32 CountOffset = CurrentDataBuffer->GetGPUInstanceCountBufferOffset();
+	if (CountOffset == INDEX_NONE)
+	{
+		DataSet.CopyFromGPUReadback(nullptr, nullptr, nullptr, 0, 0, 0, 0, 0);
+		return;
+	}
+
+	const FNiagaraGPUInstanceCountManager& CountManager = DispatchInterface->GetGPUInstanceCounterManager();
+	FRWBuffer& FloatBuffer = CurrentDataBuffer->GetGPUBufferFloat();
+	FRWBuffer& HalfBuffer = CurrentDataBuffer->GetGPUBufferHalf();
+	FRWBuffer& IntBuffer = CurrentDataBuffer->GetGPUBufferInt();
+	FRWBuffer& IDtoIndexBuffer = CurrentDataBuffer->GetGPUIDToIndexTable();
+
+	constexpr int32 NumReadbackBuffers = 5;
+	TArray<FNiagaraGpuReadbackManager::FBufferRequest, TInlineAllocator<NumReadbackBuffers>> ReadbackBuffers;
+	ReadbackBuffers.Emplace(CountManager.GetInstanceCountBuffer().Buffer, CountOffset * sizeof(uint32), sizeof(uint32));
+	const int32 FloatBufferIndex = FloatBuffer.NumBytes == 0 ? INDEX_NONE : ReadbackBuffers.Emplace(FloatBuffer.Buffer, 0, FloatBuffer.NumBytes);
+	const int32 HalfBufferIndex = HalfBuffer.NumBytes == 0 ? INDEX_NONE : ReadbackBuffers.Emplace(HalfBuffer.Buffer, 0, HalfBuffer.NumBytes);
+	const int32 IntBufferIndex = IntBuffer.NumBytes == 0 ? INDEX_NONE : ReadbackBuffers.Emplace(IntBuffer.Buffer, 0, IntBuffer.NumBytes);
+	const int32 IDtoIndexBufferIndex = IDtoIndexBuffer.NumBytes == 0 ? INDEX_NONE : ReadbackBuffers.Emplace(IDtoIndexBuffer.Buffer, 0, IDtoIndexBuffer.NumBytes);
+
+	const int32 FloatBufferStride = CurrentDataBuffer->GetFloatStride();
+	const int32 HalfBufferStride = CurrentDataBuffer->GetHalfStride();
+	const int32 IntBufferStride = CurrentDataBuffer->GetInt32Stride();
+
+	ReadbackManager->EnqueueReadbacks(
+		RHICmdList,
+		MakeArrayView(ReadbackBuffers),
+		[=, Readback=AsShared()](TConstArrayView<TPair<void*, uint32>> BufferData)
+		{
+			const int32 InstanceCount = *reinterpret_cast<int32*>(BufferData[0].Key);
+
+			// Copy dataset databuffer
+			const float* FloatDataBuffer = FloatBufferIndex == INDEX_NONE ? nullptr : reinterpret_cast<float*>(BufferData[FloatBufferIndex].Key);
+			const FFloat16* HalfDataBuffer = HalfBufferIndex == INDEX_NONE ? nullptr : reinterpret_cast<FFloat16*>(BufferData[HalfBufferIndex].Key);
+			const int32* IntDataBuffer = IntBufferIndex == INDEX_NONE ? nullptr : reinterpret_cast<int32*>(BufferData[IntBufferIndex].Key);
+			Readback->DataSet.CopyFromGPUReadback(FloatDataBuffer, IntDataBuffer, HalfDataBuffer, 0, InstanceCount, FloatBufferStride, IntBufferStride, HalfBufferStride);
+
+			// Copy ID to Index table
+			const int32* IDtoIndexBuffer = IDtoIndexBufferIndex == INDEX_NONE ? nullptr : reinterpret_cast<int32*>(BufferData[IDtoIndexBufferIndex].Key);
+			if (IDtoIndexBufferIndex != INDEX_NONE)
+			{
+				check(BufferData[IDtoIndexBufferIndex].Value >= InstanceCount * sizeof(int32));
+				Readback->IDToIndexTable.Empty();
+				Readback->IDToIndexTable.Reset(InstanceCount);
+				FMemory::Memcpy(Readback->IDToIndexTable.GetData(), IDtoIndexBuffer, InstanceCount * sizeof(int32));
+			}
+			else
+			{
+				Readback->IDToIndexTable.Empty();
+			}
+
+			if ( Readback->OnReadbackReady.IsBound() )
+			{
+				AsyncTask(
+					ENamedThreads::GameThread,
+					[Readback]()
+					{
+						Readback->PendingReadbacks--;
+						Readback->OnReadbackReady.Execute(*Readback);
+					}
+				);
+			}
+			else
+			{
+				Readback->PendingReadbacks--;
+			}
+		}
+	);
+}
diff --git a/Engine/Plugins/FX/Niagara/Source/Niagara/Private/NiagaraDebugHud.cpp b/Engine/Plugins/FX/Niagara/Source/Niagara/Private/NiagaraDebugHud.cpp
--- a/Engine/Plugins/FX/Niagara/Source/Niagara/Private/NiagaraDebugHud.cpp
+++ b/Engine/Plugins/FX/Niagara/Source/Niagara/Private/NiagaraDebugHud.cpp
@@ -2,6 +2,7 @@
 
 #include "NiagaraDebugHud.h"
 #include "NiagaraComponent.h"
+#include "NiagaraDataSetReadback.h"
 #include "NiagaraFunctionLibrary.h"
 #include "NiagaraGpuComputeDispatchInterface.h"
 #include "NiagaraComputeExecutionContext.h"
@@ -942,7 +943,7 @@ void FNiagaraDebugHud::GatherSystemInfo()
 	}
 }
 
-FNiagaraDataSet* FNiagaraDebugHud::GetParticleDataSet(FNiagaraSystemInstance* SystemInstance, FNiagaraEmitterInstance* EmitterInstance, int32 iEmitter)
+const FNiagaraDataSet* FNiagaraDebugHud::GetParticleDataSet(FNiagaraSystemInstance* SystemInstance, FNiagaraEmitterInstance* EmitterInstance, int32 iEmitter)
 {
 	using namespace NiagaraDebugLocal;
 
@@ -966,7 +967,7 @@ FNiagaraDataSet* FNiagaraDebugHud::GetParticleDataSet(FNiagaraSystemInstance* Sy
 		GpuCachedData->LastAccessedCycles = FPlatformTime::Cycles64();
 
 		// Pending readback complete?
-		if (GpuCachedData->PendingEmitterData[iEmitter] && GpuCachedData->PendingEmitterData[iEmitter]->bWritten)
+		if (GpuCachedData->PendingEmitterData[iEmitter] && GpuCachedData->PendingEmitterData[iEmitter]->IsReady())
 		{
 			GpuCachedData->CurrentEmitterData[iEmitter] = GpuCachedData->PendingEmitterData[iEmitter];
 			GpuCachedData->PendingEmitterData[iEmitter] = nullptr;
@@ -975,24 +976,14 @@ FNiagaraDataSet* FNiagaraDebugHud::GetParticleDataSet(FNiagaraSystemInstance* Sy
 		// Enqueue a readback?
 		if ( GpuCachedData->PendingEmitterData[iEmitter] == nullptr )
 		{
-			const TArray<TSharedRef<const FNiagaraEmitterCompiledData>>& AllEmittersCompiledData = SystemInstance->GetSystem()->GetEmitterCompiledData();
-
-			GpuCachedData->PendingEmitterData[iEmitter] = MakeShared<FNiagaraScriptDebuggerInfo, ESPMode::ThreadSafe>();
-			GpuCachedData->PendingEmitterData[iEmitter]->Parameters = GPUExecContext->CombinedParamStore;
-			GpuCachedData->PendingEmitterData[iEmitter]->Frame.Init(&AllEmittersCompiledData[iEmitter]->DataSetCompiledData);
-
-			ENQUEUE_RENDER_COMMAND(NiagaraReadbackGpuSim)(
-				[RT_ComputeDispatchInterface=SystemInstance->GetComputeDispatchInterface(), RT_InstanceID=SystemInstance->GetId(), RT_DebugInfo=GpuCachedData->PendingEmitterData[iEmitter], RT_Context=GPUExecContext](FRHICommandListImmediate& RHICmdList)
-				{
-					RT_ComputeDispatchInterface->AddDebugReadback(RT_InstanceID, RT_DebugInfo, RT_Context);
-				}
-			);
+			GpuCachedData->PendingEmitterData[iEmitter] = MakeShared<FNiagaraDataSetReadback, ESPMode::ThreadSafe>();
+			GpuCachedData->PendingEmitterData[iEmitter]->EnqueueReadback(EmitterInstance);
 		}
 
 		// Pull current data if we have one
 		if ( GpuCachedData->CurrentEmitterData[iEmitter] )
 		{
-			return &GpuCachedData->CurrentEmitterData[iEmitter]->Frame;
+			return &GpuCachedData->CurrentEmitterData[iEmitter]->GetDataSet();
 		}
 		return nullptr;
 	}
@@ -2223,13 +2214,13 @@ void FNiagaraDebugHud::DrawValidation(class FNiagaraWorldManager* WorldManager,
 					continue;
 				}
 
-				FNiagaraDataSet* ParticleDataSet = GetParticleDataSet(SystemInstance, EmitterInstance, iEmitter);
+				const FNiagaraDataSet* ParticleDataSet = GetParticleDataSet(SystemInstance, EmitterInstance, iEmitter);
 				if (ParticleDataSet == nullptr)
 				{
 					continue;
 				}
 
-				FNiagaraDataBuffer* DataBuffer = ParticleDataSet->GetCurrentData();
+				const FNiagaraDataBuffer* DataBuffer = ParticleDataSet->GetCurrentData();
 				if (!DataBuffer || !DataBuffer->GetNumInstances())
 				{
 					continue;
@@ -2400,13 +2391,13 @@ void FNiagaraDebugHud::DrawComponents(FNiagaraWorldManager* WorldManager, UCanva
 				}
 
 				FNiagaraEmitterInstance* EmitterInstance = &SystemInstance->GetEmitters()[iEmitter].Get();
-				FNiagaraDataSet* ParticleDataSet = GetParticleDataSet(SystemInstance, EmitterInstance, iEmitter);
+				const FNiagaraDataSet* ParticleDataSet = GetParticleDataSet(SystemInstance, EmitterInstance, iEmitter);
 				if (ParticleDataSet == nullptr)
 				{
 					continue;
 				}
 
-				FNiagaraDataBuffer* DataBuffer = ParticleDataSet->GetCurrentData();
+				const FNiagaraDataBuffer* DataBuffer = ParticleDataSet->GetCurrentData();
 				if (!DataBuffer || !DataBuffer->GetNumInstances())
 				{
 					continue;
@@ -2689,13 +2680,13 @@ void FNiagaraDebugHud::DrawComponents(FNiagaraWorldManager* WorldManager, UCanva
 								}
 
 								FNiagaraEmitterInstance* EmitterInstance = &SystemInstance->GetEmitters()[iEmitter].Get();
-								FNiagaraDataSet* ParticleDataSet = GetParticleDataSet(SystemInstance, EmitterInstance, iEmitter);
+								const FNiagaraDataSet* ParticleDataSet = GetParticleDataSet(SystemInstance, EmitterInstance, iEmitter);
 								if (ParticleDataSet == nullptr)
 								{
 									continue;
 								}
 
-								FNiagaraDataBuffer* DataBuffer = ParticleDataSet->GetCurrentData();
+								const FNiagaraDataBuffer* DataBuffer = ParticleDataSet->GetCurrentData();
 								if (!DataBuffer || !DataBuffer->GetNumInstances())
 								{
 									continue;
diff --git a/Engine/Plugins/FX/Niagara/Source/Niagara/Private/NiagaraDebugHud.h b/Engine/Plugins/FX/Niagara/Source/Niagara/Private/NiagaraDebugHud.h
--- a/Engine/Plugins/FX/Niagara/Source/Niagara/Private/NiagaraDebugHud.h
+++ b/Engine/Plugins/FX/Niagara/Source/Niagara/Private/NiagaraDebugHud.h
@@ -17,6 +17,7 @@ Class used help debugging Niagara simulations
 
 #if WITH_PARTICLE_PERF_STATS
 
+class FNiagaraDataSet;
 class FNiagaraDebugHud;
 
 struct FNiagaraDebugHudFrameStat
@@ -87,7 +88,7 @@ public:
 
 class FNiagaraDebugHud
 {
-	typedef TSharedPtr<struct FNiagaraScriptDebuggerInfo, ESPMode::ThreadSafe> FGpuDataSetPtr;
+	typedef TSharedPtr<class FNiagaraDataSetReadback, ESPMode::ThreadSafe> FGpuDataSetPtr;
 
 	struct FSystemDebugInfo
 	{
@@ -229,7 +230,7 @@ public:
 	void RemoveMessage(FName Key);
 
 private:
-	class FNiagaraDataSet* GetParticleDataSet(class FNiagaraSystemInstance* SystemInstance, class FNiagaraEmitterInstance* EmitterInstance, int32 iEmitter);
+	const FNiagaraDataSet* GetParticleDataSet(class FNiagaraSystemInstance* SystemInstance, class FNiagaraEmitterInstance* EmitterInstance, int32 iEmitter);
 	FValidationErrorInfo& GetValidationErrorInfo(class UNiagaraComponent* NiagaraComponent);
 
 	static void DebugDrawCallback(class UCanvas* Canvas, class APlayerController* PC);
diff --git a/Engine/Plugins/FX/Niagara/Source/Niagara/Private/NiagaraGpuComputeDispatch.cpp b/Engine/Plugins/FX/Niagara/Source/Niagara/Private/NiagaraGpuComputeDispatch.cpp
--- a/Engine/Plugins/FX/Niagara/Source/Niagara/Private/NiagaraGpuComputeDispatch.cpp
+++ b/Engine/Plugins/FX/Niagara/Source/Niagara/Private/NiagaraGpuComputeDispatch.cpp
@@ -419,6 +419,20 @@ void FNiagaraGpuComputeDispatch::FlushPendingTicks_GameThread()
 	);
 }
 
+void FNiagaraGpuComputeDispatch::FlushAndWait_GameThread()
+{
+	check(IsInGameThread());
+	ENQUEUE_RENDER_COMMAND(NiagaraFlushPendingTicks)(
+		[RT_NiagaraBatcher=this](FRHICommandListImmediate& RHICmdList)
+		{
+			RT_NiagaraBatcher->ProcessPendingTicksFlush(RHICmdList, true);
+			RT_NiagaraBatcher->GetGPUInstanceCounterManager().FlushIndirectArgsPool();
+			RT_NiagaraBatcher->GetGpuReadbackManager()->WaitCompletion(RHICmdList);
+		}
+	);
+	FlushRenderingCommands();
+}
+
 void FNiagaraGpuComputeDispatch::ProcessPendingTicksFlush(FRHICommandListImmediate& RHICmdList, bool bForceFlush)
 {
 	// Test to see if we have any proxies, if not we have nothing to do
@@ -475,6 +489,7 @@ void FNiagaraGpuComputeDispatch::ProcessPendingTicksFlush(FRHICommandListImmedia
 			}
 			if (bHasPendingTicks == false)
 			{
+				GpuReadbackManagerPtr->Tick();
 				return;
 			}
 
diff --git a/Engine/Plugins/FX/Niagara/Source/Niagara/Private/NiagaraGpuComputeDispatch.h b/Engine/Plugins/FX/Niagara/Source/Niagara/Private/NiagaraGpuComputeDispatch.h
--- a/Engine/Plugins/FX/Niagara/Source/Niagara/Private/NiagaraGpuComputeDispatch.h
+++ b/Engine/Plugins/FX/Niagara/Source/Niagara/Private/NiagaraGpuComputeDispatch.h
@@ -74,6 +74,7 @@ public:
 
 	// FNiagaraGpuComputeDispatchInterface Impl
 	virtual void FlushPendingTicks_GameThread() override;
+	virtual void FlushAndWait_GameThread() override;
 	// FNiagaraGpuComputeDispatchInterface Impl
 
 	/**
diff --git a/Engine/Plugins/FX/Niagara/Source/Niagara/Private/NiagaraGpuReadbackManager.h b/Engine/Plugins/FX/Niagara/Source/Niagara/Private/NiagaraGpuReadbackManager.h
--- a/Engine/Plugins/FX/Niagara/Source/Niagara/Private/NiagaraGpuReadbackManager.h
+++ b/Engine/Plugins/FX/Niagara/Source/Niagara/Private/NiagaraGpuReadbackManager.h
@@ -22,6 +22,9 @@ protected:
 public:
 	struct FBufferRequest
 	{
+		FBufferRequest() {}
+		FBufferRequest(FRHIBuffer* InBuffer, uint32 InOffset, uint32 InSize) : Buffer(InBuffer), Offset(InOffset), Size(InSize) {}
+
 		FRHIBuffer*	Buffer = nullptr;
 		uint32		Offset = 0;
 		uint32		Size = 0;
diff --git a/Engine/Plugins/FX/Niagara/Source/Niagara/Private/NiagaraSimCache.cpp b/Engine/Plugins/FX/Niagara/Source/Niagara/Private/NiagaraSimCache.cpp
new file mode 100644
--- /dev/null
+++ b/Engine/Plugins/FX/Niagara/Source/Niagara/Private/NiagaraSimCache.cpp
@@ -0,0 +1,671 @@
+// Copyright Epic Games, Inc. All Rights Reserved.
+
+#include "NiagaraSimCache.h"
+
+#include "NiagaraClearCounts.h"
+#include "NiagaraComponent.h"
+#include "NiagaraComputeExecutionContext.h"
+#include "NiagaraDataSetReadback.h"
+#include "NiagaraEmitterInstance.h"
+#include "NiagaraGPUInstanceCountManager.h"
+#include "NiagaraGpuComputeDispatchInterface.h"
+#include "NiagaraSystem.h"
+#include "NiagaraSystemInstance.h"
+
+struct FNiagaraSimCacheHelper
+{
+	explicit FNiagaraSimCacheHelper(FNiagaraSystemInstance* InSystemInstance)
+	{
+		SystemInstance = InSystemInstance;
+		SystemSimulation = SystemInstance->GetSystemSimulation();
+		check(SystemSimulation);
+		SystemSimulationDataBuffer = SystemSimulation->MainDataSet.GetCurrentData();
+		NiagaraSystem = SystemSimulation->GetSystem();
+	}
+
+	explicit FNiagaraSimCacheHelper(UNiagaraComponent* NiagaraComponent)
+	{
+		NiagaraSystem = NiagaraComponent->GetAsset();
+		if ( NiagaraSystem == nullptr )
+		{
+			return;
+		}
+
+		FNiagaraSystemInstanceControllerPtr SystemInstanceController = NiagaraComponent->GetSystemInstanceController();
+		if (SystemInstanceController.IsValid() == false)
+		{
+			return;
+		}
+
+		SystemInstance = SystemInstanceController->GetSystemInstance_Unsafe();
+		if (SystemInstance == nullptr)
+		{
+			return;
+		}
+
+		SystemSimulation = SystemInstance->GetSystemSimulation();
+		if (SystemSimulation == nullptr)
+		{
+			return;
+		}
+
+		SystemSimulationDataBuffer = SystemSimulation->MainDataSet.GetCurrentData();
+		if (SystemSimulationDataBuffer == nullptr)
+		{
+			return;
+		}
+	}
+
+	FNiagaraDataSet& GetSystemSimulationDataSet() { return SystemSimulation->MainDataSet; }
+
+	bool HasValidSimulation() const { return SystemSimulation != nullptr; }
+	bool HasValidSimulationData() const { return SystemSimulationDataBuffer != nullptr; }
+
+	void BuildCacheLayout(FNiagaraSimCacheDataBuffersLayout& CacheLayout, const FNiagaraDataSetCompiledData& CompiledData, FName LayoutName) const
+	{
+		CacheLayout.LayoutName = LayoutName;
+		CacheLayout.SimTarget = CompiledData.SimTarget;
+
+		const int32 NumVariables = CompiledData.Variables.Num();
+		CacheLayout.Variables.AddDefaulted(NumVariables);
+
+		const int32 CacheTotalComponents = CompiledData.TotalFloatComponents + CompiledData.TotalFloatComponents + CompiledData.TotalInt32Components;
+		CacheLayout.ComponentMappingsToDataBuffer.Empty(CacheTotalComponents);
+		CacheLayout.ComponentMappingsToDataBuffer.AddDefaulted(CacheTotalComponents);
+		CacheLayout.ComponentMappingsFromDataBuffer.Empty(CacheTotalComponents);
+		CacheLayout.ComponentMappingsFromDataBuffer.AddDefaulted(CacheTotalComponents);
+
+		int32 FloatOffset = 0;
+		int32 HalfOffset = CacheLayout.FloatCount;
+		int32 Int32Offset = HalfOffset + CacheLayout.HalfCount;
+		for ( int32 iVariable=0; iVariable < NumVariables; ++iVariable)
+		{
+			const FNiagaraVariableLayoutInfo& DataSetVariableLayout = CompiledData.VariableLayouts[iVariable];
+			FNiagaraSimCacheVariable& CacheVariable = CacheLayout.Variables[iVariable];
+
+			CacheVariable.Variable = CompiledData.Variables[iVariable];
+			CacheVariable.FloatOffset = DataSetVariableLayout.GetNumFloatComponents() > 0 ? CacheLayout.FloatCount : INDEX_NONE;
+			CacheVariable.FloatCount = DataSetVariableLayout.GetNumFloatComponents();
+			CacheVariable.HalfOffset = DataSetVariableLayout.GetNumHalfComponents() > 0 ? CacheLayout.HalfCount : INDEX_NONE;
+			CacheVariable.HalfCount = DataSetVariableLayout.GetNumHalfComponents();
+			CacheVariable.Int32Offset = DataSetVariableLayout.GetNumInt32Components() > 0 ? CacheLayout.Int32Count : INDEX_NONE;
+			CacheVariable.Int32Count = DataSetVariableLayout.GetNumInt32Components();
+
+			CacheLayout.FloatCount += DataSetVariableLayout.GetNumFloatComponents();
+			CacheLayout.HalfCount += DataSetVariableLayout.GetNumHalfComponents();
+			CacheLayout.Int32Count += DataSetVariableLayout.GetNumInt32Components();
+
+			// Build DataBuffer To / From mappings
+			for (int32 iComponent=0; iComponent < CacheVariable.FloatCount; ++iComponent)
+			{
+				CacheLayout.ComponentMappingsToDataBuffer[FloatOffset] = DataSetVariableLayout.FloatComponentStart + iComponent;
+				CacheLayout.ComponentMappingsFromDataBuffer[FloatOffset] = CacheVariable.FloatOffset + iComponent;
+				++FloatOffset;
+			}
+
+			for (int32 iComponent=0; iComponent < CacheVariable.HalfCount; ++iComponent)
+			{
+				CacheLayout.ComponentMappingsToDataBuffer[HalfOffset] = DataSetVariableLayout.HalfComponentStart + iComponent;
+				CacheLayout.ComponentMappingsFromDataBuffer[HalfOffset] = CacheVariable.HalfOffset + iComponent;
+				++HalfOffset;
+			}
+
+			for (int32 iComponent=0; iComponent < CacheVariable.Int32Count; ++iComponent)
+			{
+				CacheLayout.ComponentMappingsToDataBuffer[Int32Offset] = DataSetVariableLayout.Int32ComponentStart + iComponent;
+				CacheLayout.ComponentMappingsFromDataBuffer[Int32Offset] = CacheVariable.Int32Offset + iComponent;
+				++Int32Offset;
+			}
+		}
+	}
+
+	static bool BuildCacheReadMappings(FNiagaraSimCacheDataBuffersLayout& CacheLayout, const FNiagaraDataSetCompiledData& CompiledData)
+	{
+		const int32 CacheTotalComponents = CacheLayout.FloatCount + CacheLayout.HalfCount + CacheLayout.Int32Count;
+		CacheLayout.ComponentMappingsToDataBuffer.Empty(CacheTotalComponents);
+		CacheLayout.ComponentMappingsToDataBuffer.AddDefaulted(CacheTotalComponents);
+
+		int32 FloatOffset = 0;
+		int32 HalfOffset = CacheLayout.FloatCount;
+		int32 Int32Offset = HalfOffset + CacheLayout.HalfCount;
+		for (const FNiagaraSimCacheVariable& SourceVariable : CacheLayout.Variables)
+		{
+			// Find variable, if it doesn't exist that's ok as the cache contains more data than is required
+			const int32 DataSetVariableIndex = CompiledData.Variables.IndexOfByPredicate([&](const FNiagaraVariable& DataSetVariable) { return DataSetVariable == SourceVariable.Variable; });
+			const FNiagaraVariableLayoutInfo* DestVariableLayout = nullptr;
+			if (DataSetVariableIndex != INDEX_NONE)
+			{
+				DestVariableLayout = &CompiledData.VariableLayouts[DataSetVariableIndex];
+
+				// If the variable exists but types not match the cache is invalid
+				if ( (DestVariableLayout->GetNumFloatComponents() != SourceVariable.FloatCount) ||
+					 (DestVariableLayout->GetNumHalfComponents() != SourceVariable.HalfCount) ||
+					 (DestVariableLayout->GetNumInt32Components() != SourceVariable.Int32Count) )
+				{
+					return false;
+				}
+			}
+
+			for ( int32 i=0; i < SourceVariable.FloatCount; ++i )
+			{
+				CacheLayout.ComponentMappingsToDataBuffer[FloatOffset++] = DestVariableLayout ? DestVariableLayout->FloatComponentStart + i : INDEX_NONE;
+			}
+
+			for ( int32 i=0; i < SourceVariable.HalfCount; ++i)
+			{
+				CacheLayout.ComponentMappingsToDataBuffer[HalfOffset++] = DestVariableLayout ? DestVariableLayout->HalfComponentStart + i : INDEX_NONE;
+			}
+
+			for ( int32 i=0; i < SourceVariable.Int32Count; ++i)
+			{
+				CacheLayout.ComponentMappingsToDataBuffer[Int32Offset++] = DestVariableLayout ? DestVariableLayout->Int32ComponentStart + i : INDEX_NONE;
+			}
+		}
+
+		return true;
+	}
+
+	static void CheckedMemcpy(TConstArrayView<uint8> DstArray, uint8* Dst, TConstArrayView<uint8> SrcArray, const uint8* Src, uint32 Size)
+	{
+		checkf(Src >= SrcArray.GetData() && Src + Size <= SrcArray.GetData() + SrcArray.Num(), TEXT("Source %p-%p is out of bounds, start %p end %p"), Src, Src + Size, SrcArray.GetData(), SrcArray.GetData() + SrcArray.Num());
+		checkf(Dst >= DstArray.GetData() && Dst + Size <= DstArray.GetData() + DstArray.Num(), TEXT("Dest %p-%p is out of bounds, start %p end %p"), Dst, Dst + Size, DstArray.GetData(), DstArray.GetData() + DstArray.Num());
+		FMemory::Memcpy(Dst, Src, Size);
+	}
+
+	void WriteDataBuffer(const FNiagaraDataBuffer& DataBuffer, const FNiagaraSimCacheDataBuffersLayout& CacheLayout, FNiagaraSimCacheDataBuffers& CacheBuffer, int32 FirstInstance, int32 NumInstances)
+	{
+		if ( NumInstances == 0 )
+		{
+			return;
+		}
+
+		CacheBuffer.NumInstances = NumInstances;
+
+		int32 iComponent = 0;
+
+		// Copy Float
+		CacheBuffer.FloatData.AddDefaulted(CacheLayout.FloatCount * NumInstances * sizeof(float));
+		for ( uint32 i=0; i < CacheLayout.FloatCount; ++i )
+		{
+			const uint32 Component = CacheLayout.ComponentMappingsFromDataBuffer[iComponent++];
+			const uint8* Source = DataBuffer.GetComponentPtrFloat(Component) + (FirstInstance * sizeof(float));
+			uint8* Dest = CacheBuffer.FloatData.GetData() + (i * NumInstances * sizeof(float));
+			CheckedMemcpy(CacheBuffer.FloatData, Dest, DataBuffer.GetFloatBuffer(), Source, sizeof(float) * NumInstances);
+			//FMemory::Memcpy(Dest, Source, sizeof(float) * NumInstances);
+		}
+		
+		// Copy Half
+		CacheBuffer.HalfData.AddDefaulted(CacheLayout.HalfCount * NumInstances * sizeof(FFloat16));
+		for (uint32 i = 0; i < CacheLayout.HalfCount; ++i)
+		{
+			const uint32 Component = CacheLayout.ComponentMappingsFromDataBuffer[iComponent++];
+			const uint8* Source = DataBuffer.GetComponentPtrHalf(Component) + (FirstInstance * sizeof(FFloat16));
+			uint8* Dest = CacheBuffer.HalfData.GetData() + (i * NumInstances * sizeof(FFloat16));
+			CheckedMemcpy(CacheBuffer.HalfData, Dest, DataBuffer.GetHalfBuffer(), Source, sizeof(FFloat16) * NumInstances);
+			//FMemory::Memcpy(Dest, Source, sizeof(FFloat16) * NumInstances);
+		}
+
+		// Copy Int32
+		CacheBuffer.Int32Data.AddDefaulted(CacheLayout.Int32Count * NumInstances * sizeof(int32));
+		for (uint32 i = 0; i < CacheLayout.Int32Count; ++i)
+		{
+			const uint32 Component = CacheLayout.ComponentMappingsFromDataBuffer[iComponent++];
+			const uint8* Source = DataBuffer.GetComponentPtrInt32(Component) + (FirstInstance * sizeof(int32));
+			uint8* Dest = CacheBuffer.Int32Data.GetData() + (i * NumInstances * sizeof(int32));
+			CheckedMemcpy(CacheBuffer.Int32Data, Dest, DataBuffer.GetInt32Buffer(), Source, sizeof(int32) * NumInstances);
+			//FMemory::Memcpy(Dest, Source, sizeof(int32) * NumInstances);
+		}
+
+		// Copy ID to Index Table
+		CacheBuffer.IDToIndexTable = DataBuffer.GetIDTable();
+		CacheBuffer.IDAcquireTag = DataBuffer.GetIDAcquireTag();
+	}
+
+	void WriteDataBufferGPU(FNiagaraEmitterInstance& EmitterInstance, const FNiagaraDataBuffer& DataBuffer, const FNiagaraSimCacheDataBuffersLayout& CacheLayout, FNiagaraSimCacheDataBuffers& CacheBuffer)
+	{
+		//-TODO: Make async
+		TSharedRef<FNiagaraDataSetReadback> ReadbackRequest = MakeShared<FNiagaraDataSetReadback>();
+		ReadbackRequest->ImmediateReadback(&EmitterInstance);
+		if ( FNiagaraDataBuffer* CurrentData = ReadbackRequest->GetDataSet().GetCurrentData() )
+		{
+			WriteDataBuffer(*CurrentData, CacheLayout, CacheBuffer, 0, CurrentData->GetNumInstances());
+		}
+	}
+
+	void ReadDataBuffer(const FNiagaraSimCacheDataBuffersLayout& CacheLayout, const FNiagaraSimCacheDataBuffers& CacheBuffer, FNiagaraDataSet& DataSet)
+	{
+		FNiagaraDataBuffer& DataBuffer = DataSet.BeginSimulate();
+		DataBuffer.Allocate(CacheBuffer.NumInstances);
+		DataBuffer.SetNumInstances(CacheBuffer.NumInstances);
+		if ( CacheBuffer.NumInstances > 0 )
+		{
+			int32 iComponent = 0;
+			const int32 NumInstances = CacheBuffer.NumInstances;
+
+			// Copy Float
+			for (uint32 i=0; i < CacheLayout.FloatCount; ++i)
+			{
+				const uint32 Component = CacheLayout.ComponentMappingsToDataBuffer[iComponent++];
+				if (Component == InvalidComponent)
+				{
+					continue;
+				}
+				const uint8* Source = CacheBuffer.FloatData.GetData() + (i * NumInstances * sizeof(float));
+				uint8* Dest = DataBuffer.GetComponentPtrFloat(Component);
+				CheckedMemcpy(DataBuffer.GetFloatBuffer(), Dest, CacheBuffer.FloatData, Source, sizeof(float) * NumInstances);
+				//FMemory::Memcpy(Dest, Source, sizeof(float) * NumInstances);
+			}
+
+			// Copy Half
+			for (uint32 i = 0; i < CacheLayout.HalfCount; ++i)
+			{
+				const uint32 Component = CacheLayout.ComponentMappingsToDataBuffer[iComponent++];
+				if (Component == InvalidComponent)
+				{
+					continue;
+				}
+				const uint8* Source = CacheBuffer.HalfData.GetData() + (i * NumInstances * sizeof(FFloat16));
+				uint8* Dest = DataBuffer.GetComponentPtrHalf(Component);
+				CheckedMemcpy(DataBuffer.GetHalfBuffer(), Dest, CacheBuffer.HalfData, Source, sizeof(FFloat16) * NumInstances);
+				//FMemory::Memcpy(Dest, Source, sizeof(FFloat16) * NumInstances);
+			}
+
+			// Copy Int32
+			for (uint32 i = 0; i < CacheLayout.Int32Count; ++i)
+			{
+				const uint32 Component = CacheLayout.ComponentMappingsToDataBuffer[iComponent++];
+				if (Component == InvalidComponent)
+				{
+					continue;
+				}
+				const uint8* Source = CacheBuffer.Int32Data.GetData() + (i * NumInstances * sizeof(int32));
+				uint8* Dest = DataBuffer.GetComponentPtrInt32(Component);
+				CheckedMemcpy(DataBuffer.GetInt32Buffer(), Dest, CacheBuffer.Int32Data, Source, sizeof(int32) * NumInstances);
+				//FMemory::Memcpy(Dest, Source, sizeof(int32) * NumInstances);
+			}
+		}
+
+		//-TODO:DestinationDataBuffer.SetIDTable(CacheBuffer.IDToIndexTable);
+		DataBuffer.SetIDAcquireTag(CacheBuffer.IDAcquireTag);
+
+		DataSet.EndSimulate();
+	}
+
+	void ReadDataBufferGPU(FNiagaraEmitterInstance& EmitterInstance, const FNiagaraSimCacheDataBuffersLayout& InCacheLayout, const FNiagaraSimCacheDataBuffers& InCacheBuffer, FNiagaraDataSet& InDataSet, std::atomic<int32>& InPendingCommandsCounter)
+	{
+		++InPendingCommandsCounter;
+
+		FNiagaraGpuComputeDispatchInterface* DispathInterface = EmitterInstance.GetParentSystemInstance()->GetComputeDispatchInterface();
+		ENQUEUE_RENDER_COMMAND(NiagaraSimCacheGpuReadFrame)(
+			[DispathInterface, GPUExecContext=EmitterInstance.GetGPUContext(), CacheLayout=&InCacheLayout, CacheBuffer=&InCacheBuffer, DataSet=&InDataSet, PendingCommandsCounter=&InPendingCommandsCounter](FRHICommandListImmediate& RHICmdList)
+			{
+				const int32 NumInstances = CacheBuffer->NumInstances;
+
+				// Set Instance Count
+				{
+					FNiagaraGPUInstanceCountManager& CountManager = DispathInterface->GetGPUInstanceCounterManager();
+					if (GPUExecContext->CountOffset_RT == INDEX_NONE)
+					{
+						GPUExecContext->CountOffset_RT = CountManager.AcquireOrAllocateEntry(RHICmdList);
+					}
+
+					const FRWBuffer& CountBuffer = CountManager.GetInstanceCountBuffer();
+					const TPair<uint32, int32> DataToSet(GPUExecContext->CountOffset_RT, NumInstances);
+					RHICmdList.Transition(FRHITransitionInfo(CountBuffer.UAV, FNiagaraGPUInstanceCountManager::kCountBufferDefaultState, ERHIAccess::UAVCompute));
+					NiagaraClearCounts::ClearCountsInt(RHICmdList, CountBuffer.UAV, MakeArrayView(&DataToSet, 1));
+					RHICmdList.Transition(FRHITransitionInfo(CountBuffer.UAV, ERHIAccess::UAVCompute, FNiagaraGPUInstanceCountManager::kCountBufferDefaultState));
+				}
+
+				// Copy instance counts
+				FNiagaraDataBuffer& DataBuffer = DataSet->GetCurrentDataChecked();
+				DataBuffer.AllocateGPU(RHICmdList, NumInstances, DispathInterface->GetFeatureLevel(), TEXT("NiagaraSimCache"));
+				DataBuffer.SetNumInstances(NumInstances);
+				DataBuffer.SetGPUDataReadyStage(ENiagaraGpuComputeTickStage::PreInitViews);
+				GPUExecContext->SetDataToRender(&DataBuffer);
+
+				if (CacheBuffer->NumInstances > 0 )
+				{
+					int32 iComponent = 0;
+
+					// Copy Float
+					if ( CacheLayout->FloatCount > 0 )
+					{
+						FRWBuffer& RWBuffer = DataBuffer.GetGPUBufferFloat();
+						const int32 RWComponentStride = DataBuffer.GetFloatStride();
+						uint8* RWBufferMemory = reinterpret_cast<uint8*>(RHILockBuffer(RWBuffer.Buffer, 0, RWBuffer.NumBytes, RLM_WriteOnly));
+
+						for (uint32 i=0; i < CacheLayout->FloatCount; ++i)
+						{
+							const uint32 Component = CacheLayout->ComponentMappingsToDataBuffer[iComponent++];
+							if (Component == InvalidComponent)
+							{
+								continue;
+							}
+							const uint8* Source = CacheBuffer->FloatData.GetData() + (i * NumInstances * sizeof(float));
+							uint8* Dest = RWBufferMemory + (uint32(Component) * RWComponentStride);
+							CheckedMemcpy(MakeArrayView(RWBufferMemory, RWBuffer.NumBytes), Dest, CacheBuffer->FloatData, Source, sizeof(float) * NumInstances);
+							//FMemory::Memcpy(Dest, Source, sizeof(float) * NumInstances);
+						}
+
+						RHIUnlockBuffer(RWBuffer.Buffer);
+					}
+
+					// Copy Half
+					if (CacheLayout->HalfCount > 0)
+					{
+						FRWBuffer& RWBuffer = DataBuffer.GetGPUBufferHalf();
+						const int32 RWComponentStride = DataBuffer.GetHalfStride();
+						uint8* RWBufferMemory = reinterpret_cast<uint8*>(RHILockBuffer(RWBuffer.Buffer, 0, RWBuffer.NumBytes, RLM_WriteOnly));
+
+						for (uint32 i = 0; i < CacheLayout->HalfCount; ++i)
+						{
+							const uint32 Component = CacheLayout->ComponentMappingsToDataBuffer[iComponent++];
+							if (Component == InvalidComponent)
+							{
+								continue;
+							}
+							const uint8* Source = CacheBuffer->HalfData.GetData() + (i * NumInstances * sizeof(FFloat16));
+							uint8* Dest = RWBufferMemory + (uint32(Component) * RWComponentStride);
+							CheckedMemcpy(MakeArrayView(RWBufferMemory, RWBuffer.NumBytes), Dest, CacheBuffer->HalfData, Source, sizeof(FFloat16) * NumInstances);
+							//FMemory::Memcpy(Dest, Source, sizeof(FFloat16) * NumInstances);
+						}
+
+						RHIUnlockBuffer(RWBuffer.Buffer);
+					}
+
+
+					// Copy Int32
+					if (CacheLayout->Int32Count > 0)
+					{
+						FRWBuffer& RWBuffer = DataBuffer.GetGPUBufferInt();
+						const int32 RWComponentStride = DataBuffer.GetInt32Stride();
+						uint8* RWBufferMemory = reinterpret_cast<uint8*>(RHILockBuffer(RWBuffer.Buffer, 0, RWBuffer.NumBytes, RLM_WriteOnly));
+
+						for (uint32 i = 0; i < CacheLayout->Int32Count; ++i)
+						{
+							const uint32 Component = CacheLayout->ComponentMappingsToDataBuffer[iComponent++];
+							if (Component == InvalidComponent)
+							{
+								continue;
+							}
+							const uint8* Source = CacheBuffer->Int32Data.GetData() + (i * NumInstances * sizeof(int32));
+							uint8* Dest = RWBufferMemory + (uint32(Component) * RWComponentStride);
+							CheckedMemcpy(MakeArrayView(RWBufferMemory, RWBuffer.NumBytes), Dest, CacheBuffer->Int32Data, Source, sizeof(int32) * NumInstances);
+							//FMemory::Memcpy(Dest, Source, sizeof(int32) * NumInstances);
+						}
+
+						RHIUnlockBuffer(RWBuffer.Buffer);
+					}
+				}
+
+				//-TODO:DestinationDataBuffer.SetIDTable(CacheBuffer.IDToIndexTable);
+				DataBuffer.SetIDAcquireTag(CacheBuffer->IDAcquireTag);
+
+				// Ensure we decrement our counter so the GameThread knows the state of things
+				--(*PendingCommandsCounter);
+			}
+		);
+	}
+
+	UNiagaraSystem*						NiagaraSystem = nullptr;
+	FNiagaraSystemInstance*				SystemInstance = nullptr;
+	FNiagaraSystemSimulationPtr			SystemSimulation = nullptr;
+	FNiagaraDataBuffer*					SystemSimulationDataBuffer = nullptr;
+
+	static constexpr uint16				InvalidComponent = INDEX_NONE;
+};
+
+//////////////////////////////////////////////////////////////////////////
+
+UNiagaraSimCache::UNiagaraSimCache(const FObjectInitializer& ObjectInitializer)
+{
+}
+
+bool UNiagaraSimCache::IsReadyForFinishDestroy()
+{
+	return PendingCommandsInFlight == 0;
+}
+
+void UNiagaraSimCache::BeginWrite(UNiagaraComponent* NiagaraComponent)
+{
+	check(PendingCommandsInFlight == 0);
+
+	FNiagaraSimCacheHelper Helper(NiagaraComponent);
+	if (Helper.HasValidSimulation() == false)
+	{
+		return;
+	}
+
+	// Reset to defaults
+	SoftNiagaraSystem = Helper.NiagaraSystem;
+	StartSeconds = 0.0f;
+	DurationSeconds = 0.0f;
+	CacheLayout = FNiagaraSimCacheLayout();
+	CacheFrames.Empty();
+
+	// Build new layout for system / emitters
+	Helper.BuildCacheLayout(CacheLayout.SystemLayout, Helper.NiagaraSystem->GetSystemCompiledData().DataSetCompiledData, Helper.NiagaraSystem->GetFName());
+
+	const int32 NumEmitters = Helper.NiagaraSystem->GetEmitterHandles().Num();
+	CacheLayout.EmitterLayouts.AddDefaulted(NumEmitters);
+	for ( int32 i=0; i < NumEmitters; ++i )
+	{
+		const FNiagaraEmitterCompiledData& EmitterCompiledData = Helper.NiagaraSystem->GetEmitterCompiledData()[i].Get();
+		Helper.BuildCacheLayout(CacheLayout.EmitterLayouts[i], EmitterCompiledData.DataSetCompiledData, Helper.NiagaraSystem->GetEmitterHandle(i).GetName());
+	}
+}
+
+void UNiagaraSimCache::WriteFrame(UNiagaraComponent* NiagaraComponent)
+{
+	FNiagaraSimCacheHelper Helper(NiagaraComponent);
+	if (Helper.HasValidSimulationData() == false)
+	{
+		SoftNiagaraSystem.Reset();
+		return;
+	}
+
+	if ( SoftNiagaraSystem.Get() != Helper.NiagaraSystem )
+	{
+		SoftNiagaraSystem.Reset();
+		return;
+	}
+
+	// Simulation is complete nothing to cache
+	if ( Helper.SystemInstance->IsComplete() )
+	{
+		return;
+	}
+
+	// Is the simulation running?  If not nothing to cache yet
+	if ( Helper.SystemInstance->SystemInstanceState != ENiagaraSystemInstanceState::Running )
+	{
+		return;
+	}
+
+	// First frame we are about to cache?
+	if ( CacheFrames.Num() == 0 )
+	{
+		StartSeconds = Helper.SystemInstance->GetAge();
+	}
+
+	// Invalid we have reset for some reason
+	if ( Helper.SystemInstance->GetAge() < StartSeconds + DurationSeconds )
+	{
+		SoftNiagaraSystem.Reset();
+		return;
+	}
+
+	DurationSeconds = Helper.SystemInstance->GetAge() - StartSeconds;
+
+	// Cache frame
+	FNiagaraSimCacheFrame& CacheFrame = CacheFrames.AddDefaulted_GetRef();
+	CacheFrame.SystemData.LocalBounds = Helper.SystemInstance->GetLocalBounds();
+
+	const int32 NumEmitters = CacheLayout.EmitterLayouts.Num();
+	CacheFrame.EmitterData.AddDefaulted(NumEmitters);
+
+	Helper.WriteDataBuffer(*Helper.SystemSimulationDataBuffer, CacheLayout.SystemLayout, CacheFrame.SystemData.SystemDataBuffers, Helper.SystemInstance->GetSystemInstanceIndex(), 1);
+
+	for ( int32 i=0; i < NumEmitters; ++i )
+	{
+		FNiagaraSimCacheEmitterFrame& CacheEmitterFrame = CacheFrame.EmitterData[i];
+		FNiagaraEmitterInstance& EmitterInstance = Helper.SystemInstance->GetEmitters()[i].Get();
+		FNiagaraDataBuffer* EmitterCurrentData = EmitterInstance.GetData().GetCurrentData();
+		if (EmitterInstance.IsComplete() || !EmitterCurrentData)
+		{
+			continue;
+		}
+
+		CacheEmitterFrame.LocalBounds = EmitterInstance.GetBounds();
+		CacheEmitterFrame.TotalSpawnedParticles = EmitterInstance.GetTotalSpawnedParticles();
+		if (CacheLayout.EmitterLayouts[i].SimTarget == ENiagaraSimTarget::GPUComputeSim)
+		{
+			Helper.WriteDataBufferGPU(EmitterInstance, *EmitterCurrentData, CacheLayout.EmitterLayouts[i], CacheEmitterFrame.ParticleDataBuffers);
+		}
+		else
+		{
+			Helper.WriteDataBuffer(*EmitterCurrentData, CacheLayout.EmitterLayouts[i], CacheEmitterFrame.ParticleDataBuffers, 0, EmitterCurrentData->GetNumInstances());
+		}
+	}
+}
+
+void UNiagaraSimCache::EndWrite()
+{
+	check(PendingCommandsInFlight == 0);
+	if ( CacheFrames.Num() == 0 )
+	{
+		SoftNiagaraSystem.Reset();
+	}
+}
+
+bool UNiagaraSimCache::CanRead(UNiagaraSystem* NiagaraSystem)
+{
+	check(IsInGameThread());
+
+	if ( NiagaraSystem != SoftNiagaraSystem.Get() )
+	{
+		return false;
+	}
+
+	if ( NiagaraSystem->IsReadyToRun() == false )
+	{
+		return false;
+	}
+
+	// Uncooked platforms can recompile the system so we need to detect if a recache is required
+	//-OPT: This should use the changed notification delegate to avoid checks
+#if WITH_EDITORONLY_DATA
+	if ( !bNeedsReadComponentMappingRecache )
+	{
+		uint32 CacheVMIndex = 0;
+		NiagaraSystem->ForEachScript(
+			[&](UNiagaraScript* Script)
+			{
+				if (CachedScriptVMIds.IsValidIndex(CacheVMIndex))
+				{
+					bNeedsReadComponentMappingRecache |= CachedScriptVMIds[CacheVMIndex] != Script->GetVMExecutableDataCompilationId();
+				}
+				else
+				{
+					bNeedsReadComponentMappingRecache = true;
+				}
+			}
+		);
+	}
+#endif
+
+	if (bNeedsReadComponentMappingRecache)
+	{
+		const int32 NumEmitters = NiagaraSystem->GetEmitterHandles().Num();
+		if (NumEmitters != CacheLayout.EmitterLayouts.Num())
+		{
+			return false;
+		}
+
+		bool bCacheValid = true;
+		bCacheValid &= FNiagaraSimCacheHelper::BuildCacheReadMappings(CacheLayout.SystemLayout, NiagaraSystem->GetSystemCompiledData().DataSetCompiledData);
+
+		for ( int32 i=0; i < NumEmitters; ++i )
+		{
+			const FNiagaraEmitterCompiledData& EmitterCompiledData = NiagaraSystem->GetEmitterCompiledData()[i].Get();
+			bCacheValid &= FNiagaraSimCacheHelper::BuildCacheReadMappings(CacheLayout.EmitterLayouts[i], EmitterCompiledData.DataSetCompiledData);
+		}
+
+		if (bCacheValid == false)
+		{
+			return false;
+		}
+
+#if WITH_EDITORONLY_DATA
+		// Gather all the CachedScriptVMIds
+		CachedScriptVMIds.Empty();
+		NiagaraSystem->ForEachScript(
+			[&](UNiagaraScript* Script)
+			{
+				CachedScriptVMIds.Add(Script->GetVMExecutableDataCompilationId());
+			}
+		);
+		CachedScriptVMIds.Shrink();
+#endif
+		bNeedsReadComponentMappingRecache = false;
+	}
+
+	return true;
+}
+
+bool UNiagaraSimCache::Read(float TimeSeconds, FNiagaraSystemInstance* SystemInstance) const
+{
+	const float RelativeTime = FMath::Max(TimeSeconds - StartSeconds, 0.0f);
+	if ( RelativeTime > DurationSeconds )
+	{
+		// Complete
+		return false;
+	}
+
+	const float FrameTime		= (RelativeTime / DurationSeconds) * float(CacheFrames.Num() - 1);
+	const float FrameIndexA		= FMath::Floor(FrameTime);
+	const float FrameIndexB		= FMath::Min(FrameIndexA + 1 , CacheFrames.Num() - 1);
+	const float FrameFraction	= FrameTime - float(FrameIndexA);
+
+	return ReadFrame(FrameIndexA, SystemInstance);
+}
+
+bool UNiagaraSimCache::ReadFrame(int32 FrameIndex, FNiagaraSystemInstance* SystemInstance) const
+{
+	FNiagaraSimCacheHelper Helper(SystemInstance);
+	if ( !Helper.HasValidSimulation() )
+	{
+		return false;
+	}
+
+	const FNiagaraSimCacheFrame& CacheFrame = CacheFrames[FrameIndex];
+	Helper.SystemInstance->LocalBounds = CacheFrame.SystemData.LocalBounds;
+	Helper.ReadDataBuffer(CacheLayout.SystemLayout, CacheFrame.SystemData.SystemDataBuffers, Helper.GetSystemSimulationDataSet());
+
+	const int32 NumEmitters = CacheLayout.EmitterLayouts.Num();
+	for (int32 i=0; i < NumEmitters; ++i)
+	{
+		const FNiagaraSimCacheEmitterFrame& CacheEmitterFrame = CacheFrame.EmitterData[i];
+		FNiagaraEmitterInstance& EmitterInstance = Helper.SystemInstance->GetEmitters()[i].Get();
+		EmitterInstance.CachedBounds = CacheEmitterFrame.LocalBounds;
+		EmitterInstance.TotalSpawnedParticles = CacheEmitterFrame.TotalSpawnedParticles;
+		if (CacheLayout.EmitterLayouts[i].SimTarget == ENiagaraSimTarget::GPUComputeSim)
+		{
+			Helper.ReadDataBufferGPU(EmitterInstance, CacheLayout.EmitterLayouts[i], CacheEmitterFrame.ParticleDataBuffers, EmitterInstance.GetData(), PendingCommandsInFlight);
+		}
+		else
+		{
+			Helper.ReadDataBuffer(CacheLayout.EmitterLayouts[i], CacheEmitterFrame.ParticleDataBuffers, EmitterInstance.GetData());
+		}
+	}
+	return true;
+}
+
+UNiagaraSimCache* UNiagaraSimCache::CreateSingleFrame(UObject* OuterObject, UNiagaraComponent* NiagaraComponent)
+{
+	UNiagaraSimCache* SimCache = NewObject<UNiagaraSimCache>(OuterObject);
+	SimCache->BeginWrite(NiagaraComponent);
+	SimCache->WriteFrame(NiagaraComponent);
+	SimCache->EndWrite();
+	return SimCache;
+}
diff --git a/Engine/Plugins/FX/Niagara/Source/Niagara/Private/NiagaraSystemInstance.cpp b/Engine/Plugins/FX/Niagara/Source/Niagara/Private/NiagaraSystemInstance.cpp
--- a/Engine/Plugins/FX/Niagara/Source/Niagara/Private/NiagaraSystemInstance.cpp
+++ b/Engine/Plugins/FX/Niagara/Source/Niagara/Private/NiagaraSystemInstance.cpp
@@ -12,6 +12,7 @@
 #include "NiagaraWorldManager.h"
 #include "NiagaraComponent.h"
 #include "NiagaraRenderer.h"
+#include "NiagaraSimCache.h"
 #include "NiagaraGpuComputeDebug.h"
 #include "NiagaraGpuComputeDispatchInterface.h"
 #include "NiagaraCrashReporterHandler.h"
@@ -2166,6 +2167,92 @@ void FNiagaraSystemInstance::ManualTick(float DeltaSeconds, const FGraphEventRef
 	SystemSim->Tick_GameThread(DeltaSeconds, MyCompletionGraphEvent);
 }
 
+void FNiagaraSystemInstance::SimCacheTick_GameThread(UNiagaraSimCache* SimCache, float DesiredAge, float DeltaSeconds, const FGraphEventRef& MyCompletionGraphEvent)
+{
+	SCOPE_CYCLE_COUNTER(STAT_NiagaraOverview_GT);
+	SCOPE_CYCLE_COUNTER(STAT_NiagaraSystemInst_ComponentTickGT);
+	SCOPE_CYCLE_COUNTER(STAT_NiagaraSystemInst_TickGT);
+	CSV_SCOPED_TIMING_STAT_EXCLUSIVE(Effects);
+	LLM_SCOPE(ELLMTag::Niagara);
+
+	check(IsInGameThread());
+	check(bSolo);
+
+	FNiagaraCrashReporterScope CRScope(this);
+
+	UNiagaraSystem* System = GetSystem();
+	FScopeCycleCounter SystemStat(System->GetStatID(true, false));
+
+	if (IsComplete() || IsDisabled())
+	{
+		return;
+	}
+
+	// If the attached component is marked pending kill the instance is no longer valid
+	if (GetAttachComponent() == nullptr)
+	{
+		Complete(true);
+		return;
+	}
+
+	if ( DesiredAge >= SimCache->StartSeconds + SimCache->DurationSeconds )
+	{
+		Complete(false);
+		return;
+	}
+
+	CachedDeltaSeconds = DeltaSeconds;
+	//FixedBounds_CNC = FixedBounds_GT;
+
+	//TickInstanceParameters_GameThread(DeltaSeconds);
+
+	//TickDataInterfaces(DeltaSeconds, false);
+
+	Age = DesiredAge;
+	TickCount += 1;
+
+#if WITH_EDITOR
+	//// We need to tick the rapid iteration parameters when in the editor
+	//for (auto& EmitterInstance : Emitters)
+	//{
+	//	if (EmitterInstance->ShouldTick())
+	//	{
+	//		EmitterInstance->TickRapidIterationParameters();
+	//	}
+	//}
+#endif
+
+	//-OPT: On cooked builds we do not need to do this every frame
+	if ( SimCache->CanRead(GetSystem()) == false )
+	{
+		Complete(false);
+		return;
+	}
+
+	//-TODO: This can run async
+	const bool bCacheValid = SimCache->Read(Age, this);
+	if (bCacheValid == false)
+	{
+		Complete(false);
+		return;
+	}
+
+	if ( SystemInstanceState == ENiagaraSystemInstanceState::PendingSpawn )
+	{
+		SystemSimulation->SetInstanceState(this, ENiagaraSystemInstanceState::Running);
+	}
+
+	if (OnPostTickDelegate.IsBound())
+	{
+		OnPostTickDelegate.Execute();
+	}
+}
+
+void FNiagaraSystemInstance::SimCacheTick_Concurrent(UNiagaraSimCache* SimCache)
+{
+	//-OPT: Move SimCache read from GT to concurrent work
+}
+
 void FNiagaraSystemInstance::DumpStalledInfo()
 {
 	TStringBuilder<128> Builder;
diff --git a/Engine/Plugins/FX/Niagara/Source/Niagara/Public/NiagaraComponent.h b/Engine/Plugins/FX/Niagara/Source/Niagara/Public/NiagaraComponent.h
--- a/Engine/Plugins/FX/Niagara/Source/Niagara/Public/NiagaraComponent.h
+++ b/Engine/Plugins/FX/Niagara/Source/Niagara/Public/NiagaraComponent.h
@@ -697,6 +697,12 @@ public:
 	uint32 bWaitForCompilationOnActivate : 1;
 #endif
 
+	UFUNCTION(BlueprintCallable, Category = SimCache)
+	void SetSimCache(UNiagaraSimCache* InSimCache);
+
+	UFUNCTION(BlueprintCallable, Category = SimCache)
+	UNiagaraSimCache* GetSimCache() const;
+
 	/** Set whether this component is allowed to perform scalability checks and potentially be culled etc. Occasionally it is useful to disable this for specific components. E.g. Effects on the local player. */
 	UFUNCTION(BlueprintSetter, Category = Scalability, meta = (Keywords = "LOD scalability"))
 	void SetAllowScalability(bool bAllow);
@@ -777,6 +783,9 @@ private:
 
 	float CustomTimeDilation = 1.0f;
 
+	UPROPERTY(transient)
+	TObjectPtr<class UNiagaraSimCache> SimCache;
+
 	UPROPERTY(transient)
 	TObjectPtr<class UNiagaraCullProxyComponent> CullProxy;
 
diff --git a/Engine/Plugins/FX/Niagara/Source/Niagara/Public/NiagaraDataSetReadback.h b/Engine/Plugins/FX/Niagara/Source/Niagara/Public/NiagaraDataSetReadback.h
new file mode 100644
--- /dev/null
+++ b/Engine/Plugins/FX/Niagara/Source/Niagara/Public/NiagaraDataSetReadback.h
@@ -0,0 +1,44 @@
+// Copyright Epic Games, Inc. All Rights Reserved.
+
+#pragma once
+
+#include "NiagaraCommon.h"
+#include "NiagaraDataSet.h"
+#include "NiagaraParameterStore.h"
+
+struct FNiagaraComputeExecutionContext;
+class FNiagaraEmitterInstance;
+class FNiagaraGpuComputeDispatchInterface;
+
+class FNiagaraDataSetReadback : public TSharedFromThis<FNiagaraDataSetReadback, ESPMode::ThreadSafe>
+{
+public:
+	DECLARE_DELEGATE_OneParam(FOnReadbackReady, const FNiagaraDataSetReadback&);
+
+public:
+	bool IsReady() const { return PendingReadbacks == 0; }
+	void SetReadbackRead(FOnReadbackReady InOnReadbackReady);
+
+	FName GetSourceName() const { return SourceName; }
+	//ENiagaraScriptUsage GetSourceScriptUsage() const { return SourceScriptUsage; }
+	const FNiagaraDataSet& GetDataSet() const { check(IsReady()); return DataSet; }
+	const FNiagaraParameterStore& GetParameterStore() const { check(IsReady()); return ParameterStore; }
+
+	void EnqueueReadback(FNiagaraEmitterInstance* EmitterInstance);
+	void ImmediateReadback(FNiagaraEmitterInstance* EmitterInstance);
+
+private:
+	void GPUReadbackInternal(FRHICommandListImmediate& RHICmdList, FNiagaraGpuComputeDispatchInterface* DispatchInterface, FNiagaraComputeExecutionContext* GPUContext);
+
+private:
+	std::atomic<int>		PendingReadbacks;
+
+	FName					SourceName;
+	//ENiagaraScriptUsage		SourceScriptUsage;
+	//FGuid					SourceUsageId;
+	FNiagaraDataSet			DataSet;
+	FNiagaraParameterStore	ParameterStore;
+	TArray<int32>			IDToIndexTable;
+
+	FOnReadbackReady		OnReadbackReady;
+};
diff --git a/Engine/Plugins/FX/Niagara/Source/Niagara/Public/NiagaraGpuComputeDispatchInterface.h b/Engine/Plugins/FX/Niagara/Source/Niagara/Public/NiagaraGpuComputeDispatchInterface.h
--- a/Engine/Plugins/FX/Niagara/Source/Niagara/Public/NiagaraGpuComputeDispatchInterface.h
+++ b/Engine/Plugins/FX/Niagara/Source/Niagara/Public/NiagaraGpuComputeDispatchInterface.h
@@ -77,6 +77,12 @@ public:
 	*/
 	virtual void FlushPendingTicks_GameThread() = 0;
 
+	/**
+	This will flush all pending ticks & readbacks from the dispatcher.
+	Note: This is a GameThread blocking call and will impact performance
+	*/
+	virtual void FlushAndWait_GameThread() = 0;
+
 	/** Debug only function to readback data. */
 	virtual void AddDebugReadback(FNiagaraSystemInstanceID InstanceID, TSharedPtr<FNiagaraScriptDebuggerInfo, ESPMode::ThreadSafe> DebugInfo, FNiagaraComputeExecutionContext* Context) = 0;
 
diff --git a/Engine/Plugins/FX/Niagara/Source/Niagara/Public/NiagaraSystemInstance.h b/Engine/Plugins/FX/Niagara/Source/Niagara/Public/NiagaraSystemInstance.h
--- a/Engine/Plugins/FX/Niagara/Source/Niagara/Public/NiagaraSystemInstance.h
+++ b/Engine/Plugins/FX/Niagara/Source/Niagara/Public/NiagaraSystemInstance.h
@@ -16,6 +16,7 @@ class FNiagaraSystemInstance;
 class FNiagaraSystemSimulation;
 class FNiagaraGpuComputeDispatchInterface;
 class FNiagaraGPUSystemTick;
+class UNiagaraSimCache;
 class FNiagaraSystemGpuComputeProxy;
 
 using FNiagaraSystemInstancePtr = TSharedPtr<FNiagaraSystemInstance, ESPMode::ThreadSafe>;
@@ -70,6 +71,7 @@ class NIAGARA_API FNiagaraSystemInstance
 	friend class FNiagaraSystemSimulation;
 	friend class FNiagaraGPUSystemTick;
 	friend class FNiagaraDebugHud;
+	friend class UNiagaraSimCache;
 
 public:
 	DECLARE_DELEGATE(FOnPostTick);
@@ -169,6 +171,11 @@ public:
 
 	void ManualTick(float DeltaSeconds, const FGraphEventRef& MyCompletionGraphEvent);
 
+	/** Ticks the system using the a SimCache. */
+	void SimCacheTick_GameThread(UNiagaraSimCache* SimCache, float DesiredAge, float DeltaSeconds, const FGraphEventRef& MyCompletionGraphEvent);
+	/** Concurrent work for SimCache tick */
+	void SimCacheTick_Concurrent(UNiagaraSimCache* SimCache);
+
 	/** Initial phase of system instance tick. Must be executed on the game thread. */
 	void Tick_GameThread(float DeltaSeconds);
 	/** Secondary phase of the system instance tick that can be executed on any thread. */
diff --git a/Engine/Plugins/FX/Niagara/Source/Niagara/Public/NiagaraSystemInstanceController.h b/Engine/Plugins/FX/Niagara/Source/Niagara/Public/NiagaraSystemInstanceController.h
--- a/Engine/Plugins/FX/Niagara/Source/Niagara/Public/NiagaraSystemInstanceController.h
+++ b/Engine/Plugins/FX/Niagara/Source/Niagara/Public/NiagaraSystemInstanceController.h
@@ -120,6 +120,7 @@ public:
 	NIAGARA_SYSTEM_INSTANCE_CONTROLLER_SHIM(IsPendingSpawn, const)
 	NIAGARA_SYSTEM_INSTANCE_CONTROLLER_SHIM(Reset,)
 	NIAGARA_SYSTEM_INSTANCE_CONTROLLER_SHIM(ManualTick,)
+	NIAGARA_SYSTEM_INSTANCE_CONTROLLER_SHIM(SimCacheTick_GameThread,)
 	NIAGARA_SYSTEM_INSTANCE_CONTROLLER_SHIM(SetTickBehavior,)
 	NIAGARA_SYSTEM_INSTANCE_CONTROLLER_SHIM(GetTickBehavior,)
 	NIAGARA_SYSTEM_INSTANCE_CONTROLLER_SHIM(AdvanceSimulation,)
diff --git a/Engine/Plugins/FX/Niagara/Source/Niagara/Public/NiagaraSystemSimulation.h b/Engine/Plugins/FX/Niagara/Source/Niagara/Public/NiagaraSystemSimulation.h
--- a/Engine/Plugins/FX/Niagara/Source/Niagara/Public/NiagaraSystemSimulation.h
+++ b/Engine/Plugins/FX/Niagara/Source/Niagara/Public/NiagaraSystemSimulation.h
@@ -237,6 +237,7 @@ public:
 class FNiagaraSystemSimulation : public TSharedFromThis<FNiagaraSystemSimulation, ESPMode::ThreadSafe>, FGCObject
 {
 	friend FNiagaraSystemSimulationTickContext;
+	friend struct FNiagaraSimCacheHelper;
 	friend class FNiagaraDebugHud;
 
 public:
diff --git a/Engine/Plugins/FX/Niagara/Source/NiagaraEditor/Private/Customizations/NiagaraBakerOutputCustomization.cpp b/Engine/Plugins/FX/Niagara/Source/NiagaraEditor/Private/Customizations/NiagaraBakerOutputCustomization.cpp
--- a/Engine/Plugins/FX/Niagara/Source/NiagaraEditor/Private/Customizations/NiagaraBakerOutputCustomization.cpp
+++ b/Engine/Plugins/FX/Niagara/Source/NiagaraEditor/Private/Customizations/NiagaraBakerOutputCustomization.cpp
@@ -1,6 +1,7 @@
 // Copyright Epic Games, Inc. All Rights Reserved.
 
 #include "NiagaraBakerOutputCustomization.h"
+#include "NiagaraBakerOutputSimCache.h"
 #include "NiagaraBakerOutputTexture2D.h"
 #include "NiagaraBakerOutputVolumeTexture.h"
 
@@ -18,411 +19,338 @@
 
 #define LOCTEXT_NAMESPACE "NiagaraBakerOutputCustomization"
 
-namespace NiagaraBakerOutputCustomization
-{
-
 //////////////////////////////////////////////////////////////////////////
 
-struct FNiagaraAssetPathDetailBuilder : public IDetailCustomNodeBuilder, public TSharedFromThis<FNiagaraAssetPathDetailBuilder>
+void FNiagaraBakerOutputDetails::FocusContentBrowserToAsset(const FString& AssetPath)
 {
-public:
-	FNiagaraAssetPathDetailBuilder(TSharedPtr<IPropertyHandle> InPropertyHandle)
-		: PropertyHandle(InPropertyHandle)
+	FAssetRegistryModule& AssetRegistryModule = FModuleManager::LoadModuleChecked<FAssetRegistryModule>(AssetRegistryConstants::ModuleName);
+	IAssetRegistry& AssetRegistry = AssetRegistryModule.Get();
+	TArray<FAssetData> AssetDatas;
+	AssetRegistry.GetAssetsByPackageName(FName(*AssetPath), AssetDatas);
+
+	TArray<UObject*> AssetObjects;
+	AssetObjects.Reserve(AssetDatas.Num());
+	for (const FAssetData& AssetData : AssetDatas)
+	{
+		if (UObject* AssetObject = AssetData.GetAsset())
+		{
+			AssetObjects.Add(AssetObject);
+		}
+	}
+
+	if (AssetObjects.Num() > 0)
 	{
+		GEditor->SyncBrowserToObjects(AssetObjects);
 	}
+}
 
-	virtual void GenerateHeaderRowContent(FDetailWidgetRow& NodeRow) override
+void FNiagaraBakerOutputDetails::ExploreFolder(const FString& Folder)
+{
+	if ( FPaths::DirectoryExists(Folder) )
 	{
-		TSharedPtr<SWidget> NameWidget = PropertyHandle->CreatePropertyNameWidget();
-		TSharedPtr<SWidget> ValueWidget = PropertyHandle->CreatePropertyValueWidget();
-
-		//if (bCanResetToDefault)
-		//{
-		//	TSharedPtr<SHorizontalBox> Container;
-		//	SAssignNew(Container, SHorizontalBox)
-		//		+ SHorizontalBox::Slot()
-		//		.AutoWidth()
-		//		[
-		//			ValueWidget.ToSharedRef()
-		//		]
-		//	+ SHorizontalBox::Slot()
-		//		[
-		//			ParentProperty->CreateDefaultPropertyButtonWidgets()
-		//		];
-
-		//	ValueWidget = Container;
-		//}
-
-		FUIAction CopyAction;
-		FUIAction PasteAction;
-		PropertyHandle->CreateDefaultPropertyCopyPasteActions(CopyAction, PasteAction);
-			NodeRow
-			.NameContent()
+		FPlatformProcess::ExploreFolder(*Folder);
+	}
+}
+
+void FNiagaraBakerOutputDetails::BuildAssetPathWidget(IDetailCategoryBuilder& DetailCategory, TSharedRef<IPropertyHandle>& PropertyHandle, TAttribute<FText>::FGetter TooltipGetter, FOnClicked OnClicked)
+{
+	DetailCategory.AddProperty(PropertyHandle)
+		.CustomWidget()
+		.NameContent()
+		[
+			PropertyHandle->CreatePropertyNameWidget()
+		]
+		.ValueContent()
+		[
+			SNew(SHorizontalBox)
+			+SHorizontalBox::Slot()
+			.FillWidth(1.0f)
 			[
-				NameWidget.ToSharedRef()
+				PropertyHandle->CreatePropertyValueWidget()
 			]
-			.ValueContent()
-			.MinDesiredWidth(300.0f)
+			+SHorizontalBox::Slot()
+			.AutoWidth()
+			.HAlign(HAlign_Right)
 			[
-				SNew(SHorizontalBox)
-				+SHorizontalBox::Slot()
-				[
-					ValueWidget.ToSharedRef()
-				]
-				+SHorizontalBox::Slot()
+				SNew(SButton)
+				.ToolTipText(TAttribute<FText>::Create(TooltipGetter))
+				.OnClicked(OnClicked)
 				[
-					SNew(STextBlock).Text(LOCTEXT("Ooooh", "Ooooh"))
+					SNew(SImage)
+					.Image(FAppStyle::GetBrush(TEXT("Icons.Search")))
 				]
 			]
-			.CopyAction(CopyAction)
-			.PasteAction(PasteAction);
-	}
+		];
+}
+
+//////////////////////////////////////////////////////////////////////////
+
+TSharedRef<IDetailCustomization> FNiagaraBakerOutputSimCacheDetails::MakeInstance()
+{
+	return MakeShared<FNiagaraBakerOutputSimCacheDetails>();
+}
 
-	virtual void GenerateChildContent(IDetailChildrenBuilder& ChildrenBuilder) override
+FText FNiagaraBakerOutputSimCacheDetails::BrowseAssetToolTipText(TWeakObjectPtr<UNiagaraBakerOutputSimCache> WeakOutput)
+{
+	UNiagaraBakerOutputSimCache* Output = WeakOutput.Get();
+	return FText::Format(LOCTEXT("BrowseAssetToolTipFormat", "Browse to asset '{0}'"), FText::FromString(Output ? Output->GetAssetPath(Output->SimCacheAssetPathFormat, 0) : FString()));
+}
+
+FReply FNiagaraBakerOutputSimCacheDetails::BrowseToAsset(TWeakObjectPtr<UNiagaraBakerOutputSimCache> WeakOutput)
+{
+	if (UNiagaraBakerOutputSimCache* Output = WeakOutput.Get() )
 	{
-		//uint32 NumChildren = 0;
-		//if (ParentProperty->GetNumChildren(NumChildren) != FPropertyAccess::Fail)
-		//{
-		//	for (uint32 ChildIndex = 0; ChildIndex < NumChildren; ++ChildIndex)
-		//	{
-		//		TSharedPtr<IPropertyHandle> ChildProperty = ParentProperty->GetChildHandle(ChildIndex);
-		//		if (ChildProperty.IsValid() && ChildProperty->GetProperty()->GetFName() != MeshProperty->GetProperty()->GetFName())
-		//		{
-		//			ChildrenBuilder.AddProperty(ChildProperty.ToSharedRef())
-		//				.ShouldAutoExpand(false);
-		//		}
-		//	}
-		//}
+		const FString AssetPath = Output->GetAssetPath(Output->SimCacheAssetPathFormat, 0);
+		FocusContentBrowserToAsset(AssetPath);
 	}
+	return FReply::Handled();
+}
 
-	virtual FName GetName() const override
+void FNiagaraBakerOutputSimCacheDetails::CustomizeDetails(IDetailLayoutBuilder& DetailBuilder)
+{
+	// We only support customization on 1 object
+	TArray<TWeakObjectPtr<UObject>> ObjectsCustomized;
+	DetailBuilder.GetObjectsBeingCustomized(ObjectsCustomized);
+	if (ObjectsCustomized.Num() != 1 || !ObjectsCustomized[0]->IsA<UNiagaraBakerOutputSimCache>())
 	{
-		static const FName NAME_NiagaraAssetPathDetailBuilder("NiagaraAssetPathDetailBuilder");
-		return NAME_NiagaraAssetPathDetailBuilder;
+		return;
 	}
+	TWeakObjectPtr<UNiagaraBakerOutputSimCache> WeakOutput = CastChecked<UNiagaraBakerOutputSimCache>(ObjectsCustomized[0]);
 
-	//virtual void SetOnRebuildChildren(FSimpleDelegate InOnRegenerateChildren) override
-	//{
-	//	OnRebuildChildren = InOnRegenerateChildren;
-	//}
+	static FName NAME_Settings("Settings");
+	IDetailCategoryBuilder& DetailCategory = DetailBuilder.EditCategory(NAME_Settings);
 
-	virtual void Tick(float DeltaTime) override {}
-	virtual bool RequiresTick() const override { return false; }
-	virtual bool InitiallyCollapsed() const override { return true; }
-	virtual TSharedPtr<IPropertyHandle> GetPropertyHandle() const override { return PropertyHandle; }
+	TArray<TSharedRef<IPropertyHandle>> CategoryProperties;
+	DetailCategory.GetDefaultProperties(CategoryProperties);
 
-	//void Rebuild()
-	//{
-	//	OnRebuildChildren.ExecuteIfBound();
-	//}
-
-private:
-	TSharedPtr<IPropertyHandle> PropertyHandle;
-	//TSharedPtr<IPropertyHandle> MeshProperty;
-	//FSimpleDelegate OnRebuildChildren;
-	//bool bCanResetToDefault = true;
-};
-
-//////////////////////////////////////////////////////////////////////////
-
-struct FNiagaraBakerOutputDetails : public IDetailCustomization
-{
-	static void FocusContentBrowserToAsset(const FString& AssetPath)
+	for (TSharedRef<IPropertyHandle>& PropertyHandle : CategoryProperties)
 	{
-		FAssetRegistryModule& AssetRegistryModule = FModuleManager::LoadModuleChecked<FAssetRegistryModule>(AssetRegistryConstants::ModuleName);
-		IAssetRegistry& AssetRegistry = AssetRegistryModule.Get();
-		TArray<FAssetData> AssetDatas;
-		AssetRegistry.GetAssetsByPackageName(FName(*AssetPath), AssetDatas);
-
-		TArray<UObject*> AssetObjects;
-		AssetObjects.Reserve(AssetDatas.Num());
-		for (const FAssetData& AssetData : AssetDatas)
+		const FName PropertyName = PropertyHandle->GetProperty() ? PropertyHandle->GetProperty()->GetFName() : NAME_None;
+		if (PropertyName == GET_MEMBER_NAME_CHECKED(UNiagaraBakerOutputSimCache, SimCacheAssetPathFormat))
 		{
-			if (UObject* AssetObject = AssetData.GetAsset())
-			{
-				AssetObjects.Add(AssetObject);
-			}
+			BuildAssetPathWidget(
+				DetailCategory, PropertyHandle,
+				TAttribute<FText>::FGetter::CreateStatic(&FNiagaraBakerOutputSimCacheDetails::BrowseAssetToolTipText, WeakOutput),
+				FOnClicked::CreateStatic(&FNiagaraBakerOutputSimCacheDetails::BrowseToAsset, WeakOutput)
+			);
 		}
-
-		if (AssetObjects.Num() > 0)
+		else
 		{
-			GEditor->SyncBrowserToObjects(AssetObjects);
+			DetailCategory.AddProperty(PropertyHandle);
 		}
 	}
+}
 
-	static void ExploreFolder(const FString& Folder)
-	{
-		if ( FPaths::DirectoryExists(Folder) )
-		{
-			FPlatformProcess::ExploreFolder(*Folder);
-		}
-	}
+//////////////////////////////////////////////////////////////////////////
 
-	static void BuildAssetPathWidget(IDetailCategoryBuilder& DetailCategory, TSharedRef<IPropertyHandle>& PropertyHandle, TAttribute<FText>::FGetter TooltipGetter, FOnClicked OnClicked)
-	{
-		DetailCategory.AddProperty(PropertyHandle)
-			.CustomWidget()
-			.NameContent()
-			[
-				PropertyHandle->CreatePropertyNameWidget()
-			]
-			.ValueContent()
-			[
-				SNew(SHorizontalBox)
-				+SHorizontalBox::Slot()
-				.FillWidth(1.0f)
-				[
-					PropertyHandle->CreatePropertyValueWidget()
-				]
-				+SHorizontalBox::Slot()
-				.AutoWidth()
-				.HAlign(HAlign_Right)
-				[
-					SNew(SButton)
-					.ToolTipText(TAttribute<FText>::Create(TooltipGetter))
-					.OnClicked(OnClicked)
-					[
-						SNew(SImage)
-						.Image(FAppStyle::GetBrush(TEXT("Icons.Search")))
-					]
-				]
-			];
-	}
-};
+TSharedRef<IDetailCustomization> FNiagaraBakerOutputTexture2DDetails::MakeInstance()
+{
+	return MakeShared<FNiagaraBakerOutputTexture2DDetails>();
+}
 
-//////////////////////////////////////////////////////////////////////////
+FText FNiagaraBakerOutputTexture2DDetails::BrowseAtlasToolTipText(TWeakObjectPtr<UNiagaraBakerOutputTexture2D> WeakOutput)
+{
+	UNiagaraBakerOutputTexture2D* Output = WeakOutput.Get();
+	return FText::Format(LOCTEXT("BrowseAtlasToolTipFormat", "Browse to atlas asset '{0}'"), FText::FromString(Output ? Output->GetAssetPath(Output->AtlasAssetPathFormat, 0) : FString()));
+}
 
-struct FNiagaraBakerOutputTexture2DDetails : public FNiagaraBakerOutputDetails
+FReply FNiagaraBakerOutputTexture2DDetails::BrowseToAtlas(TWeakObjectPtr<UNiagaraBakerOutputTexture2D> WeakOutput)
 {
-	static TSharedRef<IDetailCustomization> MakeInstance()
+	if ( UNiagaraBakerOutputTexture2D* Output = WeakOutput.Get() )
 	{
-		return MakeShared<FNiagaraBakerOutputTexture2DDetails>();
+		const FString AssetPath = Output->GetAssetPath(Output->AtlasAssetPathFormat, 0);
+		FocusContentBrowserToAsset(AssetPath);
 	}
+	return FReply::Handled();
+}
 
-	static FText BrowseAtlasToolTipText(TWeakObjectPtr<UNiagaraBakerOutputTexture2D> WeakOutput)
-	{
-		UNiagaraBakerOutputTexture2D* Output = WeakOutput.Get();
-		return FText::Format(LOCTEXT("BrowseAtlasToolTipFormat", "Browse to atlas asset '{0}'"), FText::FromString(Output ? Output->GetAssetPath(Output->AtlasAssetPathFormat, 0) : FString()));
-	}
+FText FNiagaraBakerOutputTexture2DDetails::BrowseFrameAssetsToolTipText(TWeakObjectPtr<UNiagaraBakerOutputTexture2D> WeakOutput)
+{
+	UNiagaraBakerOutputTexture2D* Output = WeakOutput.Get();
+	return FText::Format(LOCTEXT("BrowseFrameAssetsToolTipFormat", "Browse to frames assets '{0}'"), FText::FromString(Output ? Output->GetAssetPath(Output->FramesAssetPathFormat, 0) : FString()));
+}
 
-	static FReply BrowseToAtlas(TWeakObjectPtr<UNiagaraBakerOutputTexture2D> WeakOutput)
+FReply FNiagaraBakerOutputTexture2DDetails::BrowseToFrameAssets(TWeakObjectPtr<UNiagaraBakerOutputTexture2D> WeakOutput)
+{
+	if (UNiagaraBakerOutputTexture2D* Output = WeakOutput.Get())
 	{
-		if ( UNiagaraBakerOutputTexture2D* Output = WeakOutput.Get() )
-		{
-			const FString AssetPath = Output->GetAssetPath(Output->AtlasAssetPathFormat, 0);
-			FocusContentBrowserToAsset(AssetPath);
-		}
-		return FReply::Handled();
+		const FString AssetPath = Output->GetAssetPath(Output->FramesAssetPathFormat, 0);
+		FocusContentBrowserToAsset(AssetPath);
 	}
+	return FReply::Handled();
+}
 
-	static FText BrowseFrameAssetsToolTipText(TWeakObjectPtr<UNiagaraBakerOutputTexture2D> WeakOutput)
-	{
-		UNiagaraBakerOutputTexture2D* Output = WeakOutput.Get();
-		return FText::Format(LOCTEXT("BrowseFrameAssetsToolTipFormat", "Browse to frames assets '{0}'"), FText::FromString(Output ? Output->GetAssetPath(Output->FramesAssetPathFormat, 0) : FString()));
-	}
+FText FNiagaraBakerOutputTexture2DDetails::BrowseExportAssetsToolTipText(TWeakObjectPtr<UNiagaraBakerOutputTexture2D> WeakOutput)
+{
+	UNiagaraBakerOutputTexture2D* Output = WeakOutput.Get();
+	return FText::Format(LOCTEXT("BrowseExportAssetsToolTipFormat", "Browse to exported assets '{0}'"), FText::FromString(Output ? Output->GetExportFolder(Output->FramesExportPathFormat, 0) : FString()));
+}
 
-	static FReply BrowseToFrameAssets(TWeakObjectPtr<UNiagaraBakerOutputTexture2D> WeakOutput)
+FReply FNiagaraBakerOutputTexture2DDetails::BrowseToExportAssets(TWeakObjectPtr<UNiagaraBakerOutputTexture2D> WeakOutput)
+{
+	if (UNiagaraBakerOutputTexture2D* Output = WeakOutput.Get())
 	{
-		if (UNiagaraBakerOutputTexture2D* Output = WeakOutput.Get())
-		{
-			const FString AssetPath = Output->GetAssetPath(Output->FramesAssetPathFormat, 0);
-			FocusContentBrowserToAsset(AssetPath);
-		}
-		return FReply::Handled();
+		const FString ExportFolder = Output->GetExportFolder(Output->FramesExportPathFormat, 0);
+		ExploreFolder(ExportFolder);
 	}
+	return FReply::Handled();
+}
 
-	static FText BrowseExportAssetsToolTipText(TWeakObjectPtr<UNiagaraBakerOutputTexture2D> WeakOutput)
+void FNiagaraBakerOutputTexture2DDetails::CustomizeDetails(IDetailLayoutBuilder& DetailBuilder)
+{
+	// We only support customization on 1 object
+	TArray<TWeakObjectPtr<UObject>> ObjectsCustomized;
+	DetailBuilder.GetObjectsBeingCustomized(ObjectsCustomized);
+	if (ObjectsCustomized.Num() != 1 || !ObjectsCustomized[0]->IsA<UNiagaraBakerOutputTexture2D>())
 	{
-		UNiagaraBakerOutputTexture2D* Output = WeakOutput.Get();
-		return FText::Format(LOCTEXT("BrowseExportAssetsToolTipFormat", "Browse to exported assets '{0}'"), FText::FromString(Output ? Output->GetExportFolder(Output->FramesExportPathFormat, 0) : FString()));
+		return;
 	}
+	TWeakObjectPtr<UNiagaraBakerOutputTexture2D> WeakOutput = CastChecked<UNiagaraBakerOutputTexture2D>(ObjectsCustomized[0]);
+
+	static FName NAME_Settings("Settings");
+	IDetailCategoryBuilder& DetailCategory = DetailBuilder.EditCategory(NAME_Settings);
 
-	static FReply BrowseToExportAssets(TWeakObjectPtr<UNiagaraBakerOutputTexture2D> WeakOutput)
+	TArray<TSharedRef<IPropertyHandle>> CategoryProperties;
+	DetailCategory.GetDefaultProperties(CategoryProperties);
+
+	for (TSharedRef<IPropertyHandle>& PropertyHandle : CategoryProperties)
 	{
-		if (UNiagaraBakerOutputTexture2D* Output = WeakOutput.Get())
+		const FName PropertyName = PropertyHandle->GetProperty() ? PropertyHandle->GetProperty()->GetFName() : NAME_None;
+		if (PropertyName == GET_MEMBER_NAME_CHECKED(UNiagaraBakerOutputTexture2D, AtlasAssetPathFormat))
 		{
-			const FString ExportFolder = Output->GetExportFolder(Output->FramesExportPathFormat, 0);
-			ExploreFolder(ExportFolder);
+			BuildAssetPathWidget(
+				DetailCategory, PropertyHandle,
+				TAttribute<FText>::FGetter::CreateStatic(&FNiagaraBakerOutputTexture2DDetails::BrowseAtlasToolTipText, WeakOutput),
+				FOnClicked::CreateStatic(&FNiagaraBakerOutputTexture2DDetails::BrowseToAtlas, WeakOutput)
+			);
 		}
-		return FReply::Handled();
-	}
-
-	virtual void CustomizeDetails(IDetailLayoutBuilder& DetailBuilder) override
-	{
-		// We only support customization on 1 object
-		TArray<TWeakObjectPtr<UObject>> ObjectsCustomized;
-		DetailBuilder.GetObjectsBeingCustomized(ObjectsCustomized);
-		if (ObjectsCustomized.Num() != 1 || !ObjectsCustomized[0]->IsA<UNiagaraBakerOutputTexture2D>())
+		else if (PropertyName == GET_MEMBER_NAME_CHECKED(UNiagaraBakerOutputTexture2D, FramesAssetPathFormat))
 		{
-			return;
+			BuildAssetPathWidget(
+				DetailCategory, PropertyHandle,
+				TAttribute<FText>::FGetter::CreateStatic(&FNiagaraBakerOutputTexture2DDetails::BrowseFrameAssetsToolTipText, WeakOutput),
+				FOnClicked::CreateStatic(&FNiagaraBakerOutputTexture2DDetails::BrowseToFrameAssets, WeakOutput)
+			);
 		}
-		TWeakObjectPtr<UNiagaraBakerOutputTexture2D> WeakOutput = CastChecked<UNiagaraBakerOutputTexture2D>(ObjectsCustomized[0]);
-
-		static FName NAME_Settings("Settings");
-		IDetailCategoryBuilder& DetailCategory = DetailBuilder.EditCategory(NAME_Settings);
-
-		TArray<TSharedRef<IPropertyHandle>> CategoryProperties;
-		DetailCategory.GetDefaultProperties(CategoryProperties);
-
-		for (TSharedRef<IPropertyHandle>& PropertyHandle : CategoryProperties)
+		else if (PropertyName == GET_MEMBER_NAME_CHECKED(UNiagaraBakerOutputTexture2D, FramesExportPathFormat))
+		{
+			BuildAssetPathWidget(
+				DetailCategory, PropertyHandle,
+				TAttribute<FText>::FGetter::CreateStatic(&FNiagaraBakerOutputTexture2DDetails::BrowseExportAssetsToolTipText, WeakOutput),
+				FOnClicked::CreateStatic(&FNiagaraBakerOutputTexture2DDetails::BrowseToExportAssets, WeakOutput)
+			);
+		}
+		else
 		{
-			const FName PropertyName = PropertyHandle->GetProperty() ? PropertyHandle->GetProperty()->GetFName() : NAME_None;
-			if (PropertyName == GET_MEMBER_NAME_CHECKED(UNiagaraBakerOutputTexture2D, AtlasAssetPathFormat))
-			{
-				BuildAssetPathWidget(
-					DetailCategory, PropertyHandle,
-					TAttribute<FText>::FGetter::CreateStatic(&FNiagaraBakerOutputTexture2DDetails::BrowseAtlasToolTipText, WeakOutput),
-					FOnClicked::CreateStatic(&FNiagaraBakerOutputTexture2DDetails::BrowseToAtlas, WeakOutput)
-				);
-			}
-			else if (PropertyName == GET_MEMBER_NAME_CHECKED(UNiagaraBakerOutputTexture2D, FramesAssetPathFormat))
-			{
-				BuildAssetPathWidget(
-					DetailCategory, PropertyHandle,
-					TAttribute<FText>::FGetter::CreateStatic(&FNiagaraBakerOutputTexture2DDetails::BrowseFrameAssetsToolTipText, WeakOutput),
-					FOnClicked::CreateStatic(&FNiagaraBakerOutputTexture2DDetails::BrowseToFrameAssets, WeakOutput)
-				);
-			}
-			else if (PropertyName == GET_MEMBER_NAME_CHECKED(UNiagaraBakerOutputTexture2D, FramesExportPathFormat))
-			{
-				BuildAssetPathWidget(
-					DetailCategory, PropertyHandle,
-					TAttribute<FText>::FGetter::CreateStatic(&FNiagaraBakerOutputTexture2DDetails::BrowseExportAssetsToolTipText, WeakOutput),
-					FOnClicked::CreateStatic(&FNiagaraBakerOutputTexture2DDetails::BrowseToExportAssets, WeakOutput)
-				);
-			}
-			else
-			{
-				DetailCategory.AddProperty(PropertyHandle);
-			}
+			DetailCategory.AddProperty(PropertyHandle);
 		}
 	}
-};
+}
 
 //////////////////////////////////////////////////////////////////////////
 
-struct NiagaraBakerOutputVolumeTexture : public FNiagaraBakerOutputDetails
+TSharedRef<IDetailCustomization> FNiagaraBakerOutputVolumeTextureDetails::MakeInstance()
 {
-	static TSharedRef<IDetailCustomization> MakeInstance()
-	{
-		return MakeShared<NiagaraBakerOutputVolumeTexture>();
-	}
+	return MakeShared<FNiagaraBakerOutputVolumeTextureDetails>();
+}
 
-	static FText BrowseAtlasToolTipText(TWeakObjectPtr<UNiagaraBakerOutputVolumeTexture> WeakOutput)
-	{
-		UNiagaraBakerOutputVolumeTexture* Output = WeakOutput.Get();
-		return FText::Format(LOCTEXT("BrowseAtlasToolTipFormat", "Browse to atlas asset '{0}'"), FText::FromString(Output ? Output->GetAssetPath(Output->AtlasAssetPathFormat, 0) : FString()));
-	}
+FText FNiagaraBakerOutputVolumeTextureDetails::BrowseAtlasToolTipText(TWeakObjectPtr<UNiagaraBakerOutputVolumeTexture> WeakOutput)
+{
+	UNiagaraBakerOutputVolumeTexture* Output = WeakOutput.Get();
+	return FText::Format(LOCTEXT("BrowseAtlasToolTipFormat", "Browse to atlas asset '{0}'"), FText::FromString(Output ? Output->GetAssetPath(Output->AtlasAssetPathFormat, 0) : FString()));
+}
 
-	static FReply BrowseToAtlas(TWeakObjectPtr<UNiagaraBakerOutputVolumeTexture> WeakOutput)
+FReply FNiagaraBakerOutputVolumeTextureDetails::BrowseToAtlas(TWeakObjectPtr<UNiagaraBakerOutputVolumeTexture> WeakOutput)
+{
+	if (UNiagaraBakerOutputVolumeTexture* Output = WeakOutput.Get() )
 	{
-		if (UNiagaraBakerOutputVolumeTexture* Output = WeakOutput.Get() )
-		{
-			const FString AssetPath = Output->GetAssetPath(Output->AtlasAssetPathFormat, 0);
-			FocusContentBrowserToAsset(AssetPath);
-		}
-		return FReply::Handled();
+		const FString AssetPath = Output->GetAssetPath(Output->AtlasAssetPathFormat, 0);
+		FocusContentBrowserToAsset(AssetPath);
 	}
+	return FReply::Handled();
+}
+
+FText FNiagaraBakerOutputVolumeTextureDetails::BrowseFrameAssetsToolTipText(TWeakObjectPtr<UNiagaraBakerOutputVolumeTexture> WeakOutput)
+{
+	UNiagaraBakerOutputVolumeTexture* Output = WeakOutput.Get();
+	return FText::Format(LOCTEXT("BrowseFrameAssetsToolTipFormat", "Browse to frames assets '{0}'"), FText::FromString(Output ? Output->GetAssetPath(Output->FramesAssetPathFormat, 0) : FString()));
+}
 
-	static FText BrowseFrameAssetsToolTipText(TWeakObjectPtr<UNiagaraBakerOutputVolumeTexture> WeakOutput)
+FReply FNiagaraBakerOutputVolumeTextureDetails::BrowseToFrameAssets(TWeakObjectPtr<UNiagaraBakerOutputVolumeTexture> WeakOutput)
+{
+	if (UNiagaraBakerOutputVolumeTexture* Output = WeakOutput.Get())
 	{
-		UNiagaraBakerOutputVolumeTexture* Output = WeakOutput.Get();
-		return FText::Format(LOCTEXT("BrowseFrameAssetsToolTipFormat", "Browse to frames assets '{0}'"), FText::FromString(Output ? Output->GetAssetPath(Output->FramesAssetPathFormat, 0) : FString()));
+		const FString AssetPath = Output->GetAssetPath(Output->FramesAssetPathFormat, 0);
+		FocusContentBrowserToAsset(AssetPath);
 	}
+	return FReply::Handled();
+}
 
-	static FReply BrowseToFrameAssets(TWeakObjectPtr<UNiagaraBakerOutputVolumeTexture> WeakOutput)
+FText FNiagaraBakerOutputVolumeTextureDetails::BrowseExportAssetsToolTipText(TWeakObjectPtr<UNiagaraBakerOutputVolumeTexture> WeakOutput)
+{
+	UNiagaraBakerOutputVolumeTexture* Output = WeakOutput.Get();
+	return FText::Format(LOCTEXT("BrowseExportAssetsToolTipFormat", "Browse to exported assets '{0}'"), FText::FromString(Output ? Output->GetExportFolder(Output->FramesExportPathFormat, 0) : FString()));
+}
+
+FReply FNiagaraBakerOutputVolumeTextureDetails::BrowseToExportAssets(TWeakObjectPtr<UNiagaraBakerOutputVolumeTexture> WeakOutput)
+{
+	if (UNiagaraBakerOutputVolumeTexture* Output = WeakOutput.Get())
 	{
-		if (UNiagaraBakerOutputVolumeTexture* Output = WeakOutput.Get())
-		{
-			const FString AssetPath = Output->GetAssetPath(Output->FramesAssetPathFormat, 0);
-			FocusContentBrowserToAsset(AssetPath);
-		}
-		return FReply::Handled();
+		const FString ExportFolder = Output->GetExportFolder(Output->FramesExportPathFormat, 0);
+		ExploreFolder(ExportFolder);
 	}
+	return FReply::Handled();
+}
 
-	static FText BrowseExportAssetsToolTipText(TWeakObjectPtr<UNiagaraBakerOutputVolumeTexture> WeakOutput)
+void FNiagaraBakerOutputVolumeTextureDetails::CustomizeDetails(IDetailLayoutBuilder& DetailBuilder)
+{
+	// We only support customization on 1 object
+	TArray<TWeakObjectPtr<UObject>> ObjectsCustomized;
+	DetailBuilder.GetObjectsBeingCustomized(ObjectsCustomized);
+	if (ObjectsCustomized.Num() != 1 || !ObjectsCustomized[0]->IsA<UNiagaraBakerOutputVolumeTexture>())
 	{
-		UNiagaraBakerOutputVolumeTexture* Output = WeakOutput.Get();
-		return FText::Format(LOCTEXT("BrowseExportAssetsToolTipFormat", "Browse to exported assets '{0}'"), FText::FromString(Output ? Output->GetExportFolder(Output->FramesExportPathFormat, 0) : FString()));
+		return;
 	}
+	TWeakObjectPtr<UNiagaraBakerOutputVolumeTexture> WeakOutput = CastChecked<UNiagaraBakerOutputVolumeTexture>(ObjectsCustomized[0]);
 
-	static FReply BrowseToExportAssets(TWeakObjectPtr<UNiagaraBakerOutputVolumeTexture> WeakOutput)
+	static FName NAME_Settings("Settings");
+	IDetailCategoryBuilder& DetailCategory = DetailBuilder.EditCategory(NAME_Settings);
+
+	TArray<TSharedRef<IPropertyHandle>> CategoryProperties;
+	DetailCategory.GetDefaultProperties(CategoryProperties);
+
+	for (TSharedRef<IPropertyHandle>& PropertyHandle : CategoryProperties)
 	{
-		if (UNiagaraBakerOutputVolumeTexture* Output = WeakOutput.Get())
+		const FName PropertyName = PropertyHandle->GetProperty() ? PropertyHandle->GetProperty()->GetFName() : NAME_None;
+		if (PropertyName == GET_MEMBER_NAME_CHECKED(UNiagaraBakerOutputVolumeTexture, AtlasAssetPathFormat))
 		{
-			const FString ExportFolder = Output->GetExportFolder(Output->FramesExportPathFormat, 0);
-			ExploreFolder(ExportFolder);
+			BuildAssetPathWidget(
+				DetailCategory, PropertyHandle,
+				TAttribute<FText>::FGetter::CreateStatic(&FNiagaraBakerOutputVolumeTextureDetails::BrowseAtlasToolTipText, WeakOutput),
+				FOnClicked::CreateStatic(&FNiagaraBakerOutputVolumeTextureDetails::BrowseToAtlas, WeakOutput)
+			);
 		}
-		return FReply::Handled();
-	}
-
-	virtual void CustomizeDetails(IDetailLayoutBuilder& DetailBuilder) override
-	{
-		// We only support customization on 1 object
-		TArray<TWeakObjectPtr<UObject>> ObjectsCustomized;
-		DetailBuilder.GetObjectsBeingCustomized(ObjectsCustomized);
-		if (ObjectsCustomized.Num() != 1 || !ObjectsCustomized[0]->IsA<UNiagaraBakerOutputVolumeTexture>())
+		else if (PropertyName == GET_MEMBER_NAME_CHECKED(UNiagaraBakerOutputVolumeTexture, FramesAssetPathFormat))
 		{
-			return;
+			BuildAssetPathWidget(
+				DetailCategory, PropertyHandle,
+				TAttribute<FText>::FGetter::CreateStatic(&FNiagaraBakerOutputVolumeTextureDetails::BrowseFrameAssetsToolTipText, WeakOutput),
+				FOnClicked::CreateStatic(&FNiagaraBakerOutputVolumeTextureDetails::BrowseToFrameAssets, WeakOutput)
+			);
 		}
-		TWeakObjectPtr<UNiagaraBakerOutputVolumeTexture> WeakOutput = CastChecked<UNiagaraBakerOutputVolumeTexture>(ObjectsCustomized[0]);
-
-		static FName NAME_Settings("Settings");
-		IDetailCategoryBuilder& DetailCategory = DetailBuilder.EditCategory(NAME_Settings);
-
-		TArray<TSharedRef<IPropertyHandle>> CategoryProperties;
-		DetailCategory.GetDefaultProperties(CategoryProperties);
-
-		for (TSharedRef<IPropertyHandle>& PropertyHandle : CategoryProperties)
+		else if (PropertyName == GET_MEMBER_NAME_CHECKED(UNiagaraBakerOutputVolumeTexture, FramesExportPathFormat))
 		{
-			const FName PropertyName = PropertyHandle->GetProperty() ? PropertyHandle->GetProperty()->GetFName() : NAME_None;
-			if (PropertyName == GET_MEMBER_NAME_CHECKED(UNiagaraBakerOutputVolumeTexture, AtlasAssetPathFormat))
-			{
-				BuildAssetPathWidget(
-					DetailCategory, PropertyHandle,
-					TAttribute<FText>::FGetter::CreateStatic(&NiagaraBakerOutputVolumeTexture::BrowseAtlasToolTipText, WeakOutput),
-					FOnClicked::CreateStatic(&NiagaraBakerOutputVolumeTexture::BrowseToAtlas, WeakOutput)
-				);
-			}
-			else if (PropertyName == GET_MEMBER_NAME_CHECKED(UNiagaraBakerOutputVolumeTexture, FramesAssetPathFormat))
-			{
-				BuildAssetPathWidget(
-					DetailCategory, PropertyHandle,
-					TAttribute<FText>::FGetter::CreateStatic(&NiagaraBakerOutputVolumeTexture::BrowseFrameAssetsToolTipText, WeakOutput),
-					FOnClicked::CreateStatic(&NiagaraBakerOutputVolumeTexture::BrowseToFrameAssets, WeakOutput)
-				);
-			}
-			else if (PropertyName == GET_MEMBER_NAME_CHECKED(UNiagaraBakerOutputVolumeTexture, FramesExportPathFormat))
-			{
-				BuildAssetPathWidget(
-					DetailCategory, PropertyHandle,
-					TAttribute<FText>::FGetter::CreateStatic(&NiagaraBakerOutputVolumeTexture::BrowseExportAssetsToolTipText, WeakOutput),
-					FOnClicked::CreateStatic(&NiagaraBakerOutputVolumeTexture::BrowseToExportAssets, WeakOutput)
-				);
-			}
-			else
-			{
-				DetailCategory.AddProperty(PropertyHandle);
-			}
+			BuildAssetPathWidget(
+				DetailCategory, PropertyHandle,
+				TAttribute<FText>::FGetter::CreateStatic(&FNiagaraBakerOutputVolumeTextureDetails::BrowseExportAssetsToolTipText, WeakOutput),
+				FOnClicked::CreateStatic(&FNiagaraBakerOutputVolumeTextureDetails::BrowseToExportAssets, WeakOutput)
+			);
+		}
+		else
+		{
+			DetailCategory.AddProperty(PropertyHandle);
 		}
 	}
-};
-
-//////////////////////////////////////////////////////////////////////////
-
-void RegisterCustomization(IDetailsView* DetailsView)
-{
-	DetailsView->RegisterInstancedCustomPropertyLayout(
-		UNiagaraBakerOutputTexture2D::StaticClass(),
-		FOnGetDetailCustomizationInstance::CreateStatic(&FNiagaraBakerOutputTexture2DDetails::MakeInstance)
-	);
-
-	DetailsView->RegisterInstancedCustomPropertyLayout(
-		UNiagaraBakerOutputVolumeTexture::StaticClass(),
-		FOnGetDetailCustomizationInstance::CreateStatic(&NiagaraBakerOutputVolumeTexture::MakeInstance)
-	);
 }
 
-} //NiagaraBakerOutputCustomization
-
 #undef LOCTEXT_NAMESPACE
diff --git a/Engine/Plugins/FX/Niagara/Source/NiagaraEditor/Private/Customizations/NiagaraBakerOutputCustomization.h b/Engine/Plugins/FX/Niagara/Source/NiagaraEditor/Private/Customizations/NiagaraBakerOutputCustomization.h
--- a/Engine/Plugins/FX/Niagara/Source/NiagaraEditor/Private/Customizations/NiagaraBakerOutputCustomization.h
+++ b/Engine/Plugins/FX/Niagara/Source/NiagaraEditor/Private/Customizations/NiagaraBakerOutputCustomization.h
@@ -2,9 +2,62 @@
 
 #pragma once
 
-class IDetailsView;
+#include "IDetailCustomization.h"
+#include "Framework/SlateDelegates.h"
+#include "Misc/Attribute.h"
 
-namespace NiagaraBakerOutputCustomization
+class IDetailCategoryBuilder;
+class IPropertyHandle;
+
+class UNiagaraBakerOutputSimCache;
+class UNiagaraBakerOutputTexture2D;
+class UNiagaraBakerOutputVolumeTexture;
+
+struct FNiagaraBakerOutputDetails : public IDetailCustomization
+{
+	static void FocusContentBrowserToAsset(const FString& AssetPath);
+	static void ExploreFolder(const FString& Folder);
+	static void BuildAssetPathWidget(IDetailCategoryBuilder& DetailCategory, TSharedRef<IPropertyHandle>& PropertyHandle, TAttribute<FText>::FGetter TooltipGetter, FOnClicked OnClicked);
+};
+
+struct FNiagaraBakerOutputSimCacheDetails : public FNiagaraBakerOutputDetails
+{
+	static TSharedRef<IDetailCustomization> MakeInstance();
+
+	static FText BrowseAssetToolTipText(TWeakObjectPtr<UNiagaraBakerOutputSimCache> WeakOutput);
+	static FReply BrowseToAsset(TWeakObjectPtr<UNiagaraBakerOutputSimCache> WeakOutput);
+
+	virtual void CustomizeDetails(IDetailLayoutBuilder& DetailBuilder) override;
+};
+
+struct FNiagaraBakerOutputTexture2DDetails : public FNiagaraBakerOutputDetails
+{
+	static TSharedRef<IDetailCustomization> MakeInstance();
+
+	static FText BrowseAtlasToolTipText(TWeakObjectPtr<UNiagaraBakerOutputTexture2D> WeakOutput);
+	static FReply BrowseToAtlas(TWeakObjectPtr<UNiagaraBakerOutputTexture2D> WeakOutput);
+
+	static FText BrowseFrameAssetsToolTipText(TWeakObjectPtr<UNiagaraBakerOutputTexture2D> WeakOutput);
+	static FReply BrowseToFrameAssets(TWeakObjectPtr<UNiagaraBakerOutputTexture2D> WeakOutput);
+
+	static FText BrowseExportAssetsToolTipText(TWeakObjectPtr<UNiagaraBakerOutputTexture2D> WeakOutput);
+	static FReply BrowseToExportAssets(TWeakObjectPtr<UNiagaraBakerOutputTexture2D> WeakOutput);
+
+	virtual void CustomizeDetails(IDetailLayoutBuilder& DetailBuilder) override;
+};
+
+struct FNiagaraBakerOutputVolumeTextureDetails : public FNiagaraBakerOutputDetails
 {
-	void RegisterCustomization(IDetailsView* DetailsView);
-}
+	static TSharedRef<IDetailCustomization> MakeInstance();
+
+	static FText BrowseAtlasToolTipText(TWeakObjectPtr<UNiagaraBakerOutputVolumeTexture> WeakOutput);
+	static FReply BrowseToAtlas(TWeakObjectPtr<UNiagaraBakerOutputVolumeTexture> WeakOutput);
+
+	static FText BrowseFrameAssetsToolTipText(TWeakObjectPtr<UNiagaraBakerOutputVolumeTexture> WeakOutput);
+	static FReply BrowseToFrameAssets(TWeakObjectPtr<UNiagaraBakerOutputVolumeTexture> WeakOutput);
+
+	static FText BrowseExportAssetsToolTipText(TWeakObjectPtr<UNiagaraBakerOutputVolumeTexture> WeakOutput);
+	static FReply BrowseToExportAssets(TWeakObjectPtr<UNiagaraBakerOutputVolumeTexture> WeakOutput);
+
+	virtual void CustomizeDetails(IDetailLayoutBuilder& DetailBuilder) override;
+};
diff --git a/Engine/Plugins/FX/Niagara/Source/NiagaraEditor/Private/NiagaraBakerOutputRegistry.cpp b/Engine/Plugins/FX/Niagara/Source/NiagaraEditor/Private/NiagaraBakerOutputRegistry.cpp
new file mode 100644
--- /dev/null
+++ b/Engine/Plugins/FX/Niagara/Source/NiagaraEditor/Private/NiagaraBakerOutputRegistry.cpp
@@ -0,0 +1,92 @@
+// Copyright Epic Games, Inc. All Rights Reserved.
+
+#include "NiagaraBakerOutputRegistry.h"
+
+#include "NiagaraBakerOutputSimCache.h"
+#include "NiagaraBakerOutputTexture2D.h"
+#include "NiagaraBakerOutputVolumeTexture.h"
+#include "NiagaraBakerRendererOutputSimCache.h"
+#include "NiagaraBakerRendererOutputTexture2D.h"
+#include "NiagaraBakerRendererOutputVolumeTexture.h"
+#include "Customizations/NiagaraBakerOutputCustomization.h"
+
+static TAutoConsoleVariable<bool> CVarNiagaraBakerSimCacheEnabled(
+	TEXT("fx.Niagara.Baker.SimCache.Enabled"),
+	false,
+	TEXT("SimCache is experimental and will undergo high fluctuation in the data, therefore it is meant to be used for exerpimentation only."),
+	ECVF_Default
+);
+
+FNiagaraBakerOutputRegistry& FNiagaraBakerOutputRegistry::Get()
+{
+	static TUniquePtr<FNiagaraBakerOutputRegistry> Singleton;
+	if ( Singleton.IsValid() == false )
+	{
+		Singleton = MakeUnique<FNiagaraBakerOutputRegistry>();
+
+		// Sim Cache
+		{
+			FRegistryEntry& Entry = Singleton->Registry.AddDefaulted_GetRef();
+			Entry.WeakClass				= UNiagaraBakerOutputSimCache::StaticClass();
+			Entry.ShouldShowCVar		= IConsoleManager::Get().FindConsoleVariable(TEXT("fx.Niagara.Baker.SimCache.Enabled"));
+			Entry.CreateRenderer		= []() -> FNiagaraBakerOutputRenderer* { return new FNiagaraBakerRendererOutputSimCache(); };
+			Entry.CreateCustomization	= FOnGetDetailCustomizationInstance::CreateStatic(&FNiagaraBakerOutputSimCacheDetails::MakeInstance);
+		}
+		// Texture 2D
+		{
+			FRegistryEntry& Entry		= Singleton->Registry.AddDefaulted_GetRef();
+			Entry.WeakClass				= UNiagaraBakerOutputTexture2D::StaticClass();
+			Entry.CreateRenderer		= []() -> FNiagaraBakerOutputRenderer* { return new FNiagaraBakerRendererOutputTexture2D(); };
+			Entry.CreateCustomization	= FOnGetDetailCustomizationInstance::CreateStatic(&FNiagaraBakerOutputTexture2DDetails::MakeInstance);
+		}
+		// Volume Texture
+		{
+			FRegistryEntry& Entry = Singleton->Registry.AddDefaulted_GetRef();
+			Entry.WeakClass				= UNiagaraBakerOutputVolumeTexture::StaticClass();
+			Entry.CreateRenderer		= []() -> FNiagaraBakerOutputRenderer* { return new FNiagaraBakerRendererOutputVolumeTexture(); };
+			Entry.CreateCustomization	= FOnGetDetailCustomizationInstance::CreateStatic(&FNiagaraBakerOutputVolumeTextureDetails::MakeInstance);
+		}
+	}
+	return *Singleton.Get();
+}
+
+TArray<UClass*> FNiagaraBakerOutputRegistry::GetOutputClasses() const
+{
+	TArray<UClass*> Classes;
+	for ( const FRegistryEntry& Entry : Registry )
+	{
+		if (Entry.WeakClass.Get() == nullptr)
+		{
+			continue;
+		}
+		if ( Entry.ShouldShowCVar != nullptr && Entry.ShouldShowCVar->GetBool() == false )
+		{
+			continue;
+		}
+		Classes.Add(Entry.WeakClass.Get());
+	}
+	return Classes;
+}
+
+FNiagaraBakerOutputRenderer* FNiagaraBakerOutputRegistry::GetRendererForClass(UClass* Class) const
+{
+	for (const FRegistryEntry& Entry : Registry)
+	{
+		if (Entry.WeakClass.Get() == Class)
+		{
+			return Entry.CreateRenderer();
+		}
+	}
+	return nullptr;
+}
+
+void FNiagaraBakerOutputRegistry::RegisterCustomizations(IDetailsView* DetailsView) const
+{
+	for (const FRegistryEntry& Entry : Registry)
+	{
+		if (Entry.WeakClass.Get() && Entry.CreateCustomization.IsBound() )
+		{
+			DetailsView->RegisterInstancedCustomPropertyLayout(Entry.WeakClass.Get(), Entry.CreateCustomization);
+		}
+	}
+}
diff --git a/Engine/Plugins/FX/Niagara/Source/NiagaraEditor/Private/NiagaraBakerOutputRegistry.h b/Engine/Plugins/FX/Niagara/Source/NiagaraEditor/Private/NiagaraBakerOutputRegistry.h
new file mode 100644
--- /dev/null
+++ b/Engine/Plugins/FX/Niagara/Source/NiagaraEditor/Private/NiagaraBakerOutputRegistry.h
@@ -0,0 +1,31 @@
+// Copyright Epic Games, Inc. All Rights Reserved.
+
+#pragma once
+
+#include "CoreMinimal.h"
+#include "PropertyEditorDelegates.h"
+
+class FNiagaraBakerOutputRenderer;
+class IDetailsView;
+
+class FNiagaraBakerOutputRegistry
+{
+	struct FRegistryEntry
+	{
+		TWeakObjectPtr<UClass>						WeakClass;
+		IConsoleVariable*							ShouldShowCVar = nullptr;
+		TFunction<FNiagaraBakerOutputRenderer*()>	CreateRenderer;
+		FOnGetDetailCustomizationInstance			CreateCustomization;
+	};
+
+public:
+	static FNiagaraBakerOutputRegistry& Get();
+
+	TArray<UClass*> GetOutputClasses() const;
+	FNiagaraBakerOutputRenderer* GetRendererForClass(UClass* Class) const;
+	void RegisterCustomizations(IDetailsView* DetailsView) const;
+
+private:
+	TArray<FRegistryEntry>	Registry;
+};
+
diff --git a/Engine/Plugins/FX/Niagara/Source/NiagaraEditor/Private/NiagaraBakerRenderer.cpp b/Engine/Plugins/FX/Niagara/Source/NiagaraEditor/Private/NiagaraBakerRenderer.cpp
--- a/Engine/Plugins/FX/Niagara/Source/NiagaraEditor/Private/NiagaraBakerRenderer.cpp
+++ b/Engine/Plugins/FX/Niagara/Source/NiagaraEditor/Private/NiagaraBakerRenderer.cpp
@@ -2,15 +2,12 @@
 
 #include "NiagaraBakerRenderer.h"
 #include "NiagaraBakerSettings.h"
-#include "NiagaraBakerOutputTexture2D.h"
-#include "NiagaraBakerOutputVolumeTexture.h"
+#include "NiagaraBakerOutputRegistry.h"
 #include "NiagaraComponent.h"
 #include "NiagaraSystem.h"
 #include "NiagaraComputeExecutionContext.h"
 #include "NiagaraGpuComputeDispatchInterface.h"
 #include "NiagaraBatchedElements.h"
-#include "NiagaraBakerRendererOutputTexture2D.h"
-#include "NiagaraBakerRendererOutputVolumeTexture.h"
 
 #include "Niagara/Classes/NiagaraDataInterfaceRenderTarget2D.h"
 #include "Niagara/Classes/NiagaraDataInterfaceGrid2DCollection.h"
@@ -275,33 +272,13 @@ FNiagaraBakerRenderer::FNiagaraBakerRenderer(UNiagaraSystem* InNiagaraSystem)
 	SceneCaptureComponent->bCaptureEveryFrame = false;
 	SceneCaptureComponent->bCaptureOnMovement = false;
 
-	PreviewComponent = NewObject<UNiagaraComponent>(GetTransientPackage(), NAME_None, RF_Transient);
-	PreviewComponent->CastShadow = 1;
-	PreviewComponent->bCastDynamicShadow = 1;
-	PreviewComponent->SetAllowScalability(false);
-	PreviewComponent->SetAsset(NiagaraSystem);
-	PreviewComponent->SetForceSolo(true);
-	PreviewComponent->SetAgeUpdateMode(ENiagaraAgeUpdateMode::DesiredAge);
-	PreviewComponent->SetCanRenderWhileSeeking(true);
-	PreviewComponent->SetMaxSimTime(0.0f);
-	PreviewComponent->Activate(true);
-
-	AdvancedPreviewScene = MakeShareable(new FAdvancedPreviewScene(FPreviewScene::ConstructionValues()));
-	AdvancedPreviewScene->SetFloorVisibility(false);
-	AdvancedPreviewScene->AddComponent(PreviewComponent, PreviewComponent->GetRelativeTransform());
+	CreatePreviewScene(NiagaraSystem, PreviewComponent, AdvancedPreviewScene);
 }
 
 FNiagaraBakerRenderer::~FNiagaraBakerRenderer()
 {
-	if (AdvancedPreviewScene && PreviewComponent)
-	{
-		AdvancedPreviewScene->RemoveComponent(PreviewComponent);
-	}
-
-	if (PreviewComponent)
-	{
-		PreviewComponent->DestroyComponent();
-	}
+	DestroyPreviewScene(PreviewComponent, AdvancedPreviewScene);
+	DestroyPreviewScene(SimCachePreviewComponent, SimCacheAdvancedPreviewScene);
 }
 
 void FNiagaraBakerRenderer::SetAbsoluteTime(float AbsoluteTime, bool bShouldTickComponent)
@@ -337,15 +314,20 @@ void FNiagaraBakerRenderer::SetAbsoluteTime(float AbsoluteTime, bool bShouldTick
 }
 
 void FNiagaraBakerRenderer::RenderSceneCapture(UTextureRenderTarget2D* RenderTarget, ESceneCaptureSource CaptureSource) const
+{
+	RenderSceneCapture(RenderTarget, PreviewComponent, CaptureSource);
+}
+
+void FNiagaraBakerRenderer::RenderSceneCapture(UTextureRenderTarget2D* RenderTarget, UNiagaraComponent* NiagaraComponent, ESceneCaptureSource CaptureSource) const
 {
 	UNiagaraBakerSettings* BakerSettings = GetBakerSettings();
-	if ( !RenderTarget || !BakerSettings )
+	if (!NiagaraComponent || !RenderTarget || !BakerSettings)
 	{
 		return;
 	}
 
 	const float WorldTime = GetWorldTime();
-	UWorld* World = GetWorld();
+	UWorld* World = NiagaraComponent->GetWorld();
 
 	FCanvas Canvas(RenderTarget->GameThread_GetRenderTargetResource(), nullptr, FGameTime::CreateUndilated(WorldTime, FApp::GetDeltaTime()), GetFeatureLevel());
 	Canvas.Clear(FLinearColor::Black);
@@ -353,7 +335,7 @@ void FNiagaraBakerRenderer::RenderSceneCapture(UTextureRenderTarget2D* RenderTar
 	SceneCaptureComponent->RegisterComponentWithWorld(World);
 	SceneCaptureComponent->TextureTarget = RenderTarget;
 	SceneCaptureComponent->CaptureSource = CaptureSource;
-	
+
 	// Set view location
 	const FNiagaraBakerCameraSettings& CurrentCamera = BakerSettings->GetCurrentCamera();
 	if (CurrentCamera.IsOrthographic())
@@ -366,25 +348,25 @@ void FNiagaraBakerRenderer::RenderSceneCapture(UTextureRenderTarget2D* RenderTar
 		SceneCaptureComponent->ProjectionType = ECameraProjectionMode::Perspective;
 		SceneCaptureComponent->FOVAngle = CurrentCamera.FOV;
 	}
-	
+
 	const FMatrix SceneCaptureMatrix = FMatrix(FPlane(0, 0, 1, 0), FPlane(1, 0, 0, 0), FPlane(0, 1, 0, 0), FPlane(0, 0, 0, 1));
 	FMatrix ViewMatrix = SceneCaptureMatrix * BakerSettings->GetViewportMatrix().Inverse() * FRotationTranslationMatrix(BakerSettings->GetCameraRotation(), BakerSettings->GetCameraLocation());
 	SceneCaptureComponent->SetWorldLocationAndRotation(ViewMatrix.GetOrigin(), ViewMatrix.Rotator());
-	
+
 	SceneCaptureComponent->bUseCustomProjectionMatrix = true;
 	SceneCaptureComponent->CustomProjectionMatrix = BakerSettings->GetProjectionMatrix();
-	
+
 	SceneCaptureComponent->PrimitiveRenderMode = ESceneCapturePrimitiveRenderMode::PRM_UseShowOnlyList;
 	SceneCaptureComponent->ShowOnlyComponents.Empty(1);
-	SceneCaptureComponent->ShowOnlyComponents.Add(PreviewComponent);
-	
+	SceneCaptureComponent->ShowOnlyComponents.Add(NiagaraComponent);
+
 	SceneCaptureComponent->CaptureScene();
-	
+
 	SceneCaptureComponent->TextureTarget = nullptr;
 	SceneCaptureComponent->UnregisterComponent();
-	
+
 	// Alpha from a scene capture is 1- so we need to invert
-	if ( SceneCaptureComponent->CaptureSource == ESceneCaptureSource::SCS_SceneColorHDR )
+	if (SceneCaptureComponent->CaptureSource == ESceneCaptureSource::SCS_SceneColorHDR)
 	{
 		FCanvasTileItem TileItem(FVector2D(0, 0), FVector2D(RenderTarget->GetSurfaceWidth(), RenderTarget->GetSurfaceHeight()), FLinearColor::White);
 		TileItem.BlendMode = SE_BLEND_Opaque;
@@ -599,6 +581,35 @@ void FNiagaraBakerRenderer::RenderParticleAttribute(UTextureRenderTarget2D* Rend
 	}
 }
 
+void FNiagaraBakerRenderer::RenderSimCache(UTextureRenderTarget2D* RenderTarget, UNiagaraSimCache* SimCache) const
+{
+	UNiagaraBakerSettings* BakerSettings = GetBakerSettings();
+	if (!SimCache)
+	{
+		return;
+	}
+
+	if (SimCachePreviewComponent == nullptr)
+	{
+		CreatePreviewScene(NiagaraSystem, SimCachePreviewComponent, SimCacheAdvancedPreviewScene);
+	}
+
+	const float SeekDelta = BakerSettings->GetSeekDelta();
+
+	SimCachePreviewComponent->SetSimCache(SimCache);
+	SimCachePreviewComponent->SetSeekDelta(SeekDelta);
+	SimCachePreviewComponent->SeekToDesiredAge(GetWorldTime());
+	SimCachePreviewComponent->TickComponent(SeekDelta, ELevelTick::LEVELTICK_All, nullptr);
+
+	SimCachePreviewComponent->MarkRenderDynamicDataDirty();
+	UWorld* World = SimCachePreviewComponent->GetWorld();
+	World->SendAllEndOfFrameUpdates();
+
+	RenderSceneCapture(RenderTarget, SimCachePreviewComponent, ESceneCaptureSource::SCS_SceneColorHDR);
+
+	SimCachePreviewComponent->SetSimCache(nullptr);
+}
+
 UWorld* FNiagaraBakerRenderer::GetWorld() const
 {
 	return PreviewComponent->GetWorld();
@@ -624,19 +635,12 @@ void FNiagaraBakerRenderer::AddReferencedObjects(FReferenceCollector& Collector)
 	Collector.AddReferencedObject(NiagaraSystem);
 	Collector.AddReferencedObject(PreviewComponent);
 	Collector.AddReferencedObject(SceneCaptureComponent);
+	Collector.AddReferencedObject(SimCachePreviewComponent);
 }
 
 FNiagaraBakerOutputRenderer* FNiagaraBakerRenderer::GetOutputRenderer(UClass* Class)
 {
-	if ( Class == UNiagaraBakerOutputTexture2D::StaticClass() )
-	{
-		return new FNiagaraBakerRendererOutputTexture2D();
-	}
-	else if (Class == UNiagaraBakerOutputVolumeTexture::StaticClass())
-	{
-		return new FNiagaraBakerRendererOutputVolumeTexture();
-	}
-	return nullptr;
+	return FNiagaraBakerOutputRegistry::Get().GetRendererForClass(Class);
 }
 
 bool FNiagaraBakerRenderer::ExportImage(FStringView FilePath, FIntPoint ImageSize, TArrayView<FFloat16Color> ImageData)
@@ -693,3 +697,37 @@ bool FNiagaraBakerRenderer::ExportImage(FStringView FilePath, FIntPoint ImageSiz
 	const TArray64<uint8> TempData = ImageWrapper->GetCompressed();
 	return FFileHelper::SaveArrayToFile(TempData, FilePath.GetData());
 }
+
+void FNiagaraBakerRenderer::CreatePreviewScene(UNiagaraSystem* NiagaraSystem, UNiagaraComponent*& OutComponent, TSharedPtr<FAdvancedPreviewScene>& OutPreviewScene)
+{
+	check(NiagaraSystem);
+	OutComponent = NewObject<UNiagaraComponent>(GetTransientPackage(), NAME_None, RF_Transient);
+	OutComponent->CastShadow = 1;
+	OutComponent->bCastDynamicShadow = 1;
+	OutComponent->SetAllowScalability(false);
+	OutComponent->SetAsset(NiagaraSystem);
+	OutComponent->SetForceSolo(true);
+	OutComponent->SetAgeUpdateMode(ENiagaraAgeUpdateMode::DesiredAge);
+	OutComponent->SetCanRenderWhileSeeking(true);
+	OutComponent->SetMaxSimTime(0.0f);
+	OutComponent->Activate(true);
+
+	OutPreviewScene = MakeShareable(new FAdvancedPreviewScene(FPreviewScene::ConstructionValues()));
+	OutPreviewScene->SetFloorVisibility(false);
+	OutPreviewScene->AddComponent(OutComponent, OutComponent->GetRelativeTransform());
+}
+
+void FNiagaraBakerRenderer::DestroyPreviewScene(UNiagaraComponent*& InOutComponent, TSharedPtr<FAdvancedPreviewScene>& InOutPreviewScene)
+{
+	if ( InOutPreviewScene && InOutComponent )
+	{
+		InOutPreviewScene->RemoveComponent(InOutComponent);
+		InOutPreviewScene = nullptr;
+	}
+
+	if ( InOutComponent )
+	{
+		InOutComponent->DestroyComponent();
+		InOutComponent = nullptr;
+	}
+}
diff --git a/Engine/Plugins/FX/Niagara/Source/NiagaraEditor/Private/NiagaraBakerRenderer.h b/Engine/Plugins/FX/Niagara/Source/NiagaraEditor/Private/NiagaraBakerRenderer.h
--- a/Engine/Plugins/FX/Niagara/Source/NiagaraEditor/Private/NiagaraBakerRenderer.h
+++ b/Engine/Plugins/FX/Niagara/Source/NiagaraEditor/Private/NiagaraBakerRenderer.h
@@ -9,6 +9,7 @@
 class FNiagaraBakerRenderer;
 class UNiagaraComponent;
 class UNiagaraBakerSettings;
+class UNiagaraSimCache;
 class UTextureRenderTarget2D;
 class USceneCaptureComponent2D;
 class FAdvancedPreviewScene;
@@ -90,10 +91,12 @@ public:
 	void SetAbsoluteTime(float AbsoluteTime, bool bShouldTickComponent = true);
 
 	void RenderSceneCapture(UTextureRenderTarget2D* RenderTarget, ESceneCaptureSource CaptureSource) const;
+	void RenderSceneCapture(UTextureRenderTarget2D* RenderTarget, UNiagaraComponent* NiagaraComponent, ESceneCaptureSource CaptureSource) const;
 	void RenderBufferVisualization(UTextureRenderTarget2D* RenderTarget, FName BufferVisualizationMode = NAME_None) const;
 	void RenderDataInterface(UTextureRenderTarget2D* RenderTarget, FName BindingName) const;
 	void RenderParticleAttribute(UTextureRenderTarget2D* RenderTarget, FName BindingName) const;
-	
+	void RenderSimCache(UTextureRenderTarget2D* RenderTarget, UNiagaraSimCache* SimCache) const;
+
 	UWorld* GetWorld() const;
 	float GetWorldTime() const;
 	ERHIFeatureLevel::Type GetFeatureLevel() const;
@@ -114,9 +117,16 @@ public:
 
 	static bool ExportImage(FStringView FilePath, FIntPoint ImageSize, TArrayView<FFloat16Color> ImageData);
 
+private:
+	static void CreatePreviewScene(UNiagaraSystem* NiagaraSystem, UNiagaraComponent*& OutComponent, TSharedPtr<FAdvancedPreviewScene>& OutPreviewScene);
+	static void DestroyPreviewScene(UNiagaraComponent*& InOutComponent, TSharedPtr<FAdvancedPreviewScene>& InOutPreviewScene);
+
 private:
 	UNiagaraSystem* NiagaraSystem = nullptr;
 	UNiagaraComponent* PreviewComponent = nullptr;
 	TSharedPtr<FAdvancedPreviewScene> AdvancedPreviewScene;
 	USceneCaptureComponent2D* SceneCaptureComponent = nullptr;
+
+	mutable UNiagaraComponent* SimCachePreviewComponent = nullptr;
+	mutable TSharedPtr<FAdvancedPreviewScene> SimCacheAdvancedPreviewScene;
 };
diff --git a/Engine/Plugins/FX/Niagara/Source/NiagaraEditor/Private/NiagaraBakerRendererOutputSimCache.cpp b/Engine/Plugins/FX/Niagara/Source/NiagaraEditor/Private/NiagaraBakerRendererOutputSimCache.cpp
new file mode 100644
--- /dev/null
+++ b/Engine/Plugins/FX/Niagara/Source/NiagaraEditor/Private/NiagaraBakerRendererOutputSimCache.cpp
@@ -0,0 +1,99 @@
+// Copyright Epic Games, Inc. All Rights Reserved.
+
+#include "NiagaraBakerRendererOutputSimCache.h"
+#include "NiagaraBakerOutputSimCache.h"
+
+#include "NiagaraComponent.h"
+#include "NiagaraGpuComputeDispatchInterface.h"
+#include "NiagaraSimCacheFactoryNew.h"
+#include "NiagaraSystemInstance.h"
+
+#include "AssetToolsModule.h"
+
+TArray<FNiagaraBakerOutputBinding> FNiagaraBakerRendererOutputSimCache::GetRendererBindings(UNiagaraBakerOutput* InBakerOutput) const
+{
+	return TArray<FNiagaraBakerOutputBinding>();
+}
+
+FIntPoint FNiagaraBakerRendererOutputSimCache::GetPreviewSize(UNiagaraBakerOutput* InBakerOutput, FIntPoint InAvailableSize) const
+{
+	return InAvailableSize;
+}
+
+void FNiagaraBakerRendererOutputSimCache::RenderPreview(UNiagaraBakerOutput* InBakerOutput, const FNiagaraBakerRenderer& BakerRenderer, UTextureRenderTarget2D* RenderTarget, TOptional<FString>& OutErrorString) const
+{
+	BakerRenderer.RenderSceneCapture(RenderTarget, ESceneCaptureSource::SCS_SceneColorHDR);
+}
+
+FIntPoint FNiagaraBakerRendererOutputSimCache::GetGeneratedSize(UNiagaraBakerOutput* InBakerOutput, FIntPoint InAvailableSize) const
+{
+	return InAvailableSize;
+}
+
+void FNiagaraBakerRendererOutputSimCache::RenderGenerated(UNiagaraBakerOutput* InBakerOutput, const FNiagaraBakerRenderer& BakerRenderer, UTextureRenderTarget2D* RenderTarget, TOptional<FString>& OutErrorString) const
+{
+	static FString SimCacheNotFoundError(TEXT("SimCache asset not found.\nPlease bake to see the result."));
+
+	UNiagaraBakerOutputSimCache* BakerOutput = CastChecked<UNiagaraBakerOutputSimCache>(InBakerOutput);
+	UNiagaraBakerSettings* BakerSettings = BakerOutput->GetTypedOuter<UNiagaraBakerSettings>();
+
+	UNiagaraSimCache* SimCache = BakerOutput->GetAsset<UNiagaraSimCache>(BakerOutput->SimCacheAssetPathFormat, 0);
+	if (SimCache == nullptr)
+	{
+		OutErrorString = SimCacheNotFoundError;
+		return;
+	}
+
+	BakerRenderer.RenderSimCache(RenderTarget, SimCache);
+}
+
+bool FNiagaraBakerRendererOutputSimCache::BeginBake(UNiagaraBakerOutput* InBakerOutput)
+{
+	UNiagaraBakerOutputSimCache* BakerOutput = CastChecked<UNiagaraBakerOutputSimCache>(InBakerOutput);
+	UNiagaraBakerSettings* BakerSettings = BakerOutput->GetTypedOuter<UNiagaraBakerSettings>();
+	check(BakerSettings && BakerOutput);
+
+	const FString AssetFullName = BakerOutput->GetAssetPath(BakerOutput->SimCacheAssetPathFormat, 0);
+	const FString AssetName = FString(FPathViews::GetCleanFilename(AssetFullName));
+	const FString PackagePath = FString(FPathViews::GetPath(AssetFullName));
+
+	IAssetTools& AssetTools = FModuleManager::Get().LoadModuleChecked<FAssetToolsModule>("AssetTools").Get();
+	UNiagaraSimCacheFactoryNew* SimCacheFactory = NewObject<UNiagaraSimCacheFactoryNew>();
+	BakeSimCache = Cast<UNiagaraSimCache>(AssetTools.CreateAsset(AssetName, PackagePath, UNiagaraSimCache::StaticClass(), SimCacheFactory));
+	if (BakeSimCache == nullptr)
+	{
+		return false;
+	}
+
+	BakeSimCache->AddToRoot();
+
+	return true;
+}
+
+void FNiagaraBakerRendererOutputSimCache::BakeFrame(UNiagaraBakerOutput* InBakerOutput, int FrameIndex, const FNiagaraBakerRenderer& BakerRenderer)
+{
+	UNiagaraBakerOutputSimCache* BakerOutput = CastChecked<UNiagaraBakerOutputSimCache>(InBakerOutput);
+	if (BakeSimCache == nullptr)
+	{
+		return;
+	}
+
+	if ( FrameIndex == 0 )
+	{
+		BakeSimCache->BeginWrite(BakerRenderer.GetPreviewComponent());
+	}
+	BakeSimCache->WriteFrame(BakerRenderer.GetPreviewComponent());
+}
+
+void FNiagaraBakerRendererOutputSimCache::EndBake(UNiagaraBakerOutput* InBakerOutput)
+{
+	UNiagaraBakerOutputSimCache* BakerOutput = CastChecked<UNiagaraBakerOutputSimCache>(InBakerOutput);
+	if (BakeSimCache == nullptr)
+	{
+		return;
+	}
+
+	BakeSimCache->EndWrite();
+
+	BakeSimCache->RemoveFromRoot();
+}
diff --git a/Engine/Plugins/FX/Niagara/Source/NiagaraEditor/Private/NiagaraBakerRendererOutputSimCache.h b/Engine/Plugins/FX/Niagara/Source/NiagaraEditor/Private/NiagaraBakerRendererOutputSimCache.h
new file mode 100644
--- /dev/null
+++ b/Engine/Plugins/FX/Niagara/Source/NiagaraEditor/Private/NiagaraBakerRendererOutputSimCache.h
@@ -0,0 +1,26 @@
+// Copyright Epic Games, Inc. All Rights Reserved.
+
+#pragma once
+
+#include "NiagaraBakerRenderer.h"
+
+class UNiagaraSimCache;
+
+class FNiagaraBakerRendererOutputSimCache : public FNiagaraBakerOutputRenderer
+{
+public:
+	virtual TArray<FNiagaraBakerOutputBinding> GetRendererBindings(UNiagaraBakerOutput* InBakerOutput) const override;
+
+	virtual FIntPoint GetPreviewSize(UNiagaraBakerOutput* InBakerOutput, FIntPoint InAvailableSize) const override;
+	virtual void RenderPreview(UNiagaraBakerOutput* InBakerOutput, const FNiagaraBakerRenderer& BakerRenderer, UTextureRenderTarget2D* RenderTarget, TOptional<FString>& OutErrorString) const override;
+
+	virtual FIntPoint GetGeneratedSize(UNiagaraBakerOutput* InBakerOutput, FIntPoint InAvailableSize) const override;
+	virtual void RenderGenerated(UNiagaraBakerOutput* InBakerOutput, const FNiagaraBakerRenderer& BakerRenderer, UTextureRenderTarget2D* RenderTarget, TOptional<FString>& OutErrorString) const override;
+
+	virtual bool BeginBake(UNiagaraBakerOutput* InBakerOutput) override;
+	virtual void BakeFrame(UNiagaraBakerOutput* InBakerOutput, int FrameIndex, const FNiagaraBakerRenderer& BakerRenderer) override;
+	virtual void EndBake(UNiagaraBakerOutput* InBakerOutput) override;
+
+protected:
+	UNiagaraSimCache* BakeSimCache = nullptr;
+};
diff --git a/Engine/Plugins/FX/Niagara/Source/NiagaraEditor/Private/NiagaraSimCacheFactoryNew.cpp b/Engine/Plugins/FX/Niagara/Source/NiagaraEditor/Private/NiagaraSimCacheFactoryNew.cpp
new file mode 100644
--- /dev/null
+++ b/Engine/Plugins/FX/Niagara/Source/NiagaraEditor/Private/NiagaraSimCacheFactoryNew.cpp
@@ -0,0 +1,21 @@
+// Copyright Epic Games, Inc. All Rights Reserved.
+
+#include "NiagaraSimCacheFactoryNew.h"
+#include "NiagaraSimCache.h"
+
+UNiagaraSimCacheFactoryNew::UNiagaraSimCacheFactoryNew(const FObjectInitializer& ObjectInitializer)
+{
+	SupportedClass = UNiagaraSimCache::StaticClass();
+	bCreateNew = true;
+	bEditAfterNew = true;
+}
+
+bool UNiagaraSimCacheFactoryNew::ShouldShowInNewMenu() const
+{
+	return false;
+}
+
+UObject* UNiagaraSimCacheFactoryNew::FactoryCreateNew(UClass* InClass, UObject* InParent, FName InName, EObjectFlags Flags, UObject* Context, FFeedbackContext* Warn)
+{
+	return NewObject<UNiagaraSimCache>(InParent, InClass, InName, Flags);
+}
diff --git a/Engine/Plugins/FX/Niagara/Source/NiagaraEditor/Private/NiagaraSimCacheFactoryNew.h b/Engine/Plugins/FX/Niagara/Source/NiagaraEditor/Private/NiagaraSimCacheFactoryNew.h
new file mode 100644
--- /dev/null
+++ b/Engine/Plugins/FX/Niagara/Source/NiagaraEditor/Private/NiagaraSimCacheFactoryNew.h
@@ -0,0 +1,17 @@
+// Copyright Epic Games, Inc. All Rights Reserved.
+
+#pragma once
+
+#include "Factories/Factory.h"
+#include "NiagaraSimCacheFactoryNew.generated.h"
+
+class UNiagaraSimCache;
+
+UCLASS(hidecategories=Object, MinimalAPI)
+class UNiagaraSimCacheFactoryNew : public UFactory
+{
+	GENERATED_UCLASS_BODY()
+
+	virtual bool ShouldShowInNewMenu() const override;
+	virtual UObject* FactoryCreateNew(UClass* InClass, UObject* InParent, FName InName, EObjectFlags Flags, UObject* Context, FFeedbackContext* Warn) override;
+};
diff --git a/Engine/Plugins/FX/Niagara/Source/NiagaraEditor/Private/ViewModels/NiagaraBakerViewModel.cpp b/Engine/Plugins/FX/Niagara/Source/NiagaraEditor/Private/ViewModels/NiagaraBakerViewModel.cpp
--- a/Engine/Plugins/FX/Niagara/Source/NiagaraEditor/Private/ViewModels/NiagaraBakerViewModel.cpp
+++ b/Engine/Plugins/FX/Niagara/Source/NiagaraEditor/Private/ViewModels/NiagaraBakerViewModel.cpp
@@ -577,6 +577,15 @@ void FNiagaraBakerViewModel::RenderBaker()
 		}
 	}
 
+	// Ensure we flush everything from the render thread
+	if ( UWorld* World = BakerRenderer->GetWorld() )
+	{
+		if ( FNiagaraGpuComputeDispatchInterface* DispatchInterface = FNiagaraGpuComputeDispatchInterface::Get(World) )
+		{
+			DispatchInterface->FlushAndWait_GameThread();
+		}
+	}
+
 	// Complete bake
 	{
 		FScopedSlowTask SlowTask(OutputRenderers.Num());
diff --git a/Engine/Plugins/FX/Niagara/Source/NiagaraEditor/Private/Widgets/SNiagaraBakerWidget.cpp b/Engine/Plugins/FX/Niagara/Source/NiagaraEditor/Private/Widgets/SNiagaraBakerWidget.cpp
--- a/Engine/Plugins/FX/Niagara/Source/NiagaraEditor/Private/Widgets/SNiagaraBakerWidget.cpp
+++ b/Engine/Plugins/FX/Niagara/Source/NiagaraEditor/Private/Widgets/SNiagaraBakerWidget.cpp
@@ -1,13 +1,13 @@
 // Copyright Epic Games, Inc. All Rights Reserved.
 
 #include "SNiagaraBakerWidget.h"
+#include "NiagaraBakerOutputRegistry.h"
 #include "SNiagaraBakerTimelineWidget.h"
 #include "SNiagaraBakerViewport.h"
 #include "NiagaraEditorCommon.h"
 #include "NiagaraEditorStyle.h"
 #include "NiagaraComponent.h"
 #include "NiagaraSystem.h"
-#include "Customizations/NiagaraBakerOutputCustomization.h"
 #include "ViewModels/NiagaraBakerViewModel.h"
 
 #include "Modules/ModuleManager.h"
@@ -309,7 +309,7 @@ void SNiagaraBakerWidget::Construct(const FArguments& InArgs)
 	DetailsArgs.bAllowSearch = false;
 
 	BakerSettingsDetails = PropertyModule.CreateDetailView(DetailsArgs);
-	NiagaraBakerOutputCustomization::RegisterCustomization(BakerSettingsDetails.Get());
+	FNiagaraBakerOutputRegistry::Get().RegisterCustomizations(BakerSettingsDetails.Get());
 
 	// Transport control args
 	{
@@ -726,17 +726,14 @@ TSharedRef<SWidget> SNiagaraBakerWidget::MakeAddOutputMenu()
 
 	FMenuBuilder MenuBuilder(true, nullptr);
 
-	for (TObjectIterator<UClass> It; It; ++It)
+	for ( UClass* OutputClass : FNiagaraBakerOutputRegistry::Get().GetOutputClasses() )
 	{
-		if ( It->IsChildOf(UNiagaraBakerOutput::StaticClass()) && !It->HasAnyClassFlags(CLASS_Abstract) )
-		{
-			MenuBuilder.AddMenuEntry(
-				It->GetDisplayNameText(),
-				FText::GetEmpty(),
-				FSlateIcon(),
-				FUIAction(FExecuteAction::CreateSP(ViewModel, &FNiagaraBakerViewModel::AddOutput, *It))
-			);
-		}
+		MenuBuilder.AddMenuEntry(
+			OutputClass->GetDisplayNameText(),
+			FText::GetEmpty(),
+			FSlateIcon(),
+			FUIAction(FExecuteAction::CreateSP(ViewModel, &FNiagaraBakerViewModel::AddOutput, OutputClass))
+		);
 	}
 
 	return MenuBuilder.MakeWidget();
diff --git a/Engine/Plugins/FX/Niagara/Source/NiagaraEditorWidgets/Private/DetailCustomizations/NiagaraSimCacheDetailsCustomization.cpp b/Engine/Plugins/FX/Niagara/Source/NiagaraEditorWidgets/Private/DetailCustomizations/NiagaraSimCacheDetailsCustomization.cpp
new file mode 100644
--- /dev/null
+++ b/Engine/Plugins/FX/Niagara/Source/NiagaraEditorWidgets/Private/DetailCustomizations/NiagaraSimCacheDetailsCustomization.cpp
@@ -0,0 +1,639 @@
+// Copyright Epic Games, Inc. All Rights Reserved.
+
+#include "NiagaraSimCacheDetailsCustomization.h"
+#include "NiagaraEditorCommon.h"
+#include "NiagaraSimCache.h"
+
+#include "Widgets/Input/SComboBox.h"
+#include "Widgets/Input/SEditableTextBox.h"
+#include "Widgets/Input/SSpinBox.h"
+#include "Widgets/Layout/SScrollBox.h"
+#include "Widgets/Views/SListView.h"
+#include "Widgets/Views/STableRow.h"
+#include "Widgets/Views/STreeView.h"
+#include "DetailCategoryBuilder.h"
+#include "DetailLayoutBuilder.h"
+#include "DetailWidgetRow.h"
+#include "IDetailChildrenBuilder.h"
+#include "IDetailCustomization.h"
+
+#define LOCTEXT_NAMESPACE "NiagaraSimCacheDetailsCustomization"
+
+namespace NiagaraSimCacheDetailsCustomizationLocal
+{
+	static const FName NAME_Instance("Instance");
+
+	struct FSimCacheBufferReader
+	{
+		struct FComponentInfo
+		{
+			FName			Name = NAME_None;
+			uint32			ComponentOffset = INDEX_NONE;
+			bool			bIsFloat = false;
+			bool			bIsHalf = false;
+			bool			bIsInt32 = false;
+			bool			bShowAsBool = false;
+			UEnum*			Enum = nullptr;
+		};
+
+		explicit FSimCacheBufferReader(UNiagaraSimCache* SimCache, uint32 InFrameIndex = 0, uint32 InEmitterIndex = INDEX_NONE)
+			: WeakSimCache(SimCache)
+			, FrameIndex(InFrameIndex)
+			, EmitterIndex(InEmitterIndex)
+		{
+			UpdateLayout();
+		}
+
+		void UpdateLayout()
+		{
+			ComponentInfos.Empty();
+			FoundFloats = 0;
+			FoundHalfs = 0;
+			FoundInt32s = 0;
+
+			if ( const FNiagaraSimCacheDataBuffersLayout* SimCacheLayout = GetSimCacheBufferLayout() )
+			{
+				for (const FNiagaraSimCacheVariable& Variable : SimCacheLayout->Variables)
+				{
+					const FNiagaraTypeDefinition& TypeDef = Variable.Variable.GetType();
+					if ( TypeDef.IsEnum() )
+					{
+						FComponentInfo& ComponentInfo = ComponentInfos.AddDefaulted_GetRef();
+						ComponentInfo.Name = Variable.Variable.GetName();
+						ComponentInfo.ComponentOffset = FoundInt32s++;
+						ComponentInfo.bIsInt32 = true;
+						ComponentInfo.Enum = TypeDef.GetEnum();
+					}
+					else
+					{
+						Build(Variable.Variable.GetName(), TypeDef.GetScriptStruct());
+					}
+				}
+
+				if ((FoundFloats != SimCacheLayout->FloatCount) ||
+					(FoundHalfs != SimCacheLayout->HalfCount) ||
+					(FoundInt32s != SimCacheLayout->Int32Count) )
+				{
+					UE_LOG(LogNiagaraEditor, Warning, TEXT("SimCache Layout doesn't appear to match iterating the variables"));
+					ComponentInfos.Empty();
+				}
+			}
+		}
+
+		TConstArrayView<FComponentInfo> GetComponentInfos() const
+		{
+			return MakeArrayView(ComponentInfos);
+		}
+
+		int32 GetNumInstances() const
+		{
+			const FNiagaraSimCacheDataBuffers* DataBuffer = GetSimCacheDataBuffers();
+			return DataBuffer ? DataBuffer->NumInstances : 0;
+		}
+
+		int32 GetNumFrames() const
+		{
+			UNiagaraSimCache* SimCache = WeakSimCache.Get();
+			return SimCache ? SimCache->CacheFrames.Num() : 0;
+		}
+
+		int32 GetFrameIndex() const
+		{
+			return FrameIndex;
+		}
+
+		void SetFrameIndex(int32 InFrameIndex)
+		{
+			FrameIndex = InFrameIndex;
+		}
+
+		int32 GetEmitterIndex() const
+		{
+			return EmitterIndex;
+		}
+
+		void SetEmitterIndex(int32 InEmitterIndex)
+		{
+			EmitterIndex = InEmitterIndex;
+			UpdateLayout();
+		}
+
+		const FNiagaraSimCacheDataBuffersLayout* GetSimCacheBufferLayout() const
+		{
+			if (UNiagaraSimCache* SimCache = WeakSimCache.Get())
+			{
+				if (SimCache->CacheFrames.IsValidIndex(FrameIndex))
+				{
+					if (EmitterIndex == INDEX_NONE)
+					{
+						return &SimCache->CacheLayout.SystemLayout;
+					}
+					else if (SimCache->CacheLayout.EmitterLayouts.IsValidIndex(EmitterIndex))
+					{
+						return &SimCache->CacheLayout.EmitterLayouts[EmitterIndex];
+					}
+				}
+			}
+			return nullptr;
+		}
+
+		const FNiagaraSimCacheDataBuffers* GetSimCacheDataBuffers() const
+		{
+			if ( UNiagaraSimCache* SimCache = WeakSimCache.Get() )
+			{
+				if ( SimCache->CacheFrames.IsValidIndex(FrameIndex) )
+				{
+					if ( EmitterIndex == INDEX_NONE )
+					{
+						return &SimCache->CacheFrames[FrameIndex].SystemData.SystemDataBuffers;
+					}
+					else if (SimCache->CacheFrames[FrameIndex].EmitterData.IsValidIndex(EmitterIndex))
+					{
+						return &SimCache->CacheFrames[FrameIndex].EmitterData[EmitterIndex].ParticleDataBuffers;
+					}
+				}
+			}
+			return nullptr;
+		}
+
+		FText GetComponentText(FName ComponentName, int32 InstanceIndex) const
+		{
+			const FComponentInfo* ComponentInfo = ComponentInfos.FindByPredicate([ComponentName](const FComponentInfo& Info) { return Info.Name == ComponentName; });
+			const FNiagaraSimCacheDataBuffers* SimCacheDataBuffers = GetSimCacheDataBuffers();
+			const FNiagaraSimCacheDataBuffersLayout* SimCacheBufferLayout = GetSimCacheBufferLayout();
+			if ( ComponentInfo && SimCacheDataBuffers)
+			{
+				const int32 NumInstances = SimCacheDataBuffers->NumInstances;
+				if (InstanceIndex >= 0 && InstanceIndex < NumInstances)
+				{
+					if ( ComponentInfo->bIsFloat )
+					{
+						const float* Value = reinterpret_cast<const float*>(SimCacheDataBuffers->FloatData.GetData()) + (ComponentInfo->ComponentOffset * NumInstances) + InstanceIndex;
+						return FText::AsNumber(*Value);
+					}
+					else if ( ComponentInfo->bIsHalf )
+					{
+						const FFloat16* Value = reinterpret_cast<const FFloat16*>(SimCacheDataBuffers->HalfData.GetData()) + (ComponentInfo->ComponentOffset * NumInstances) + InstanceIndex;
+						return FText::AsNumber(Value->GetFloat());
+					}
+					else if ( ComponentInfo->bIsInt32 )
+					{
+						const int32* Value = reinterpret_cast<const int32*>(SimCacheDataBuffers->Int32Data.GetData()) + (ComponentInfo->ComponentOffset * NumInstances) + InstanceIndex;
+						if ( ComponentInfo->bShowAsBool )
+						{
+							return *Value == 0 ? LOCTEXT("False", "False") : LOCTEXT("True", "True");
+						}
+						else if ( ComponentInfo->Enum != nullptr )
+						{
+							return ComponentInfo->Enum->GetDisplayNameTextByValue(*Value);
+						}
+						else
+						{
+							return FText::AsNumber(*Value);
+						}
+					}
+				}
+			}
+			return LOCTEXT("Error", "Error");
+		}
+
+	protected:
+		TWeakObjectPtr<UNiagaraSimCache>	WeakSimCache;
+		uint32								FrameIndex = 0;
+		uint32								EmitterIndex = INDEX_NONE;
+
+		TArray<FComponentInfo>				ComponentInfos;
+		uint32								FoundFloats = 0;
+		uint32								FoundHalfs = 0;
+		uint32								FoundInt32s = 0;
+
+		void Build(FName Name, UScriptStruct* Struct)
+		{
+			int32 NumProperties = 0;
+			for (TFieldIterator<FProperty> PropertyIt(Struct, EFieldIteratorFlags::IncludeSuper); PropertyIt; ++PropertyIt)
+			{
+				++NumProperties;
+			}
+
+			for (TFieldIterator<FProperty> PropertyIt(Struct, EFieldIteratorFlags::IncludeSuper); PropertyIt; ++PropertyIt)
+			{
+				FProperty* Property = *PropertyIt;
+				const FName PropertyName = NumProperties > 1 ? FName(*FString::Printf(TEXT("%s.%s"), *Name.ToString(), *Property->GetName())) : Name;
+				if (Property->IsA(FFloatProperty::StaticClass()))
+				{
+					FComponentInfo& ComponentInfo = ComponentInfos.AddDefaulted_GetRef();
+					ComponentInfo.Name = PropertyName;
+					ComponentInfo.ComponentOffset = FoundFloats++;
+					ComponentInfo.bIsFloat = true;
+				}
+				else if (Property->IsA(FUInt16Property::StaticClass()))
+				{
+					FComponentInfo& ComponentInfo = ComponentInfos.AddDefaulted_GetRef();
+					ComponentInfo.Name = PropertyName;
+					ComponentInfo.ComponentOffset = FoundHalfs++;
+					ComponentInfo.bIsHalf = true;
+				}
+				else if (Property->IsA(FIntProperty::StaticClass()))
+				{
+					FComponentInfo& ComponentInfo = ComponentInfos.AddDefaulted_GetRef();
+					ComponentInfo.Name = PropertyName;
+					ComponentInfo.ComponentOffset = FoundInt32s++;
+					ComponentInfo.bIsInt32 = true;
+					ComponentInfo.bShowAsBool = (NumProperties == 1) && (Struct == FNiagaraTypeDefinition::GetBoolStruct());
+				}
+				else if ( Property->IsA(FBoolProperty::StaticClass()))
+				{
+					FComponentInfo& ComponentInfo = ComponentInfos.AddDefaulted_GetRef();
+					ComponentInfo.Name = PropertyName;
+					ComponentInfo.ComponentOffset = FoundInt32s++;
+					ComponentInfo.bIsInt32 = true;
+					ComponentInfo.bShowAsBool = true;
+				}
+				else if (Property->IsA(FEnumProperty::StaticClass()))
+				{
+					FComponentInfo& ComponentInfo = ComponentInfos.AddDefaulted_GetRef();
+					ComponentInfo.Name = PropertyName;
+					ComponentInfo.ComponentOffset = FoundInt32s++;
+					ComponentInfo.bIsInt32 = true;
+					ComponentInfo.Enum = CastFieldChecked<FEnumProperty>(Property)->GetEnum();
+				}
+				else if ( Property->IsA(FStructProperty::StaticClass()) )
+				{
+					FStructProperty* StructProperty = CastFieldChecked<FStructProperty>(Property);
+					Build(PropertyName, FNiagaraTypeHelper::FindNiagaraFriendlyTopLevelStruct(StructProperty->Struct, ENiagaraStructConversion::Simulation));
+				}
+				else
+				{
+					// Fail
+				}
+			}
+		}
+	};
+
+	class SSimCacheDataBufferRowWidget : public SMultiColumnTableRow<TSharedPtr<int32>>
+	{
+	public:
+		SLATE_BEGIN_ARGS(SSimCacheDataBufferRowWidget) {}
+			SLATE_ARGUMENT(TSharedPtr<int32>, RowIndexPtr)
+			SLATE_ARGUMENT(TSharedPtr<FSimCacheBufferReader>, BufferReader)
+		SLATE_END_ARGS()
+
+		/** Construct function for this widget */
+		void Construct(const FArguments& InArgs, const TSharedRef<STableViewBase>& InOwnerTableView)
+		{
+			RowIndexPtr = InArgs._RowIndexPtr;
+			BufferReader = InArgs._BufferReader;
+
+			SMultiColumnTableRow<TSharedPtr<int32>>::Construct(
+				FSuperRowType::FArguments()
+				.Style(FAppStyle::Get(), "DataTableEditor.CellListViewRow"),
+				InOwnerTableView
+			);
+		}
+
+		virtual TSharedRef<SWidget> GenerateWidgetForColumn(const FName& InColumnName) override
+		{
+			const int32 InstanceIndex = *RowIndexPtr;
+			if ( InColumnName == NAME_Instance )
+			{
+				return SNew(STextBlock)
+					.Text(FText::AsNumber(InstanceIndex));
+			}
+			else
+			{
+				return SNew(STextBlock)
+					.Text(BufferReader->GetComponentText(InColumnName, InstanceIndex));
+			}
+		}
+
+		TSharedPtr<int32>						RowIndexPtr;
+		TSharedPtr<FSimCacheBufferReader>		BufferReader;
+	};
+
+	class SSimCacheDataBufferWidget : public SCompoundWidget
+	{
+	public:
+		SLATE_BEGIN_ARGS(SSimCacheDataBufferWidget) {}
+			SLATE_ARGUMENT(UNiagaraSimCache*, SimCache)
+		SLATE_END_ARGS()
+
+		using FBufferSelectionInfo = TPair<int32, FText>;
+
+		void Construct(const FArguments& InArgs)
+		{
+			UNiagaraSimCache* SimCache = InArgs._SimCache;
+			WeakSimCache = SimCache;
+
+			BufferReader = MakeShared<FSimCacheBufferReader>(SimCache);
+
+			HeaderRowWidget = SNew(SHeaderRow);
+
+			UpdateColumns(false);
+			UpdateRows(false);
+			UpdateBufferSelectionList();
+
+			TSharedRef<SScrollBar> HorizontalScrollBar =
+				SNew(SScrollBar)
+				.AlwaysShowScrollbar(true)
+				.Thickness(12.0f)
+				.Orientation(Orient_Horizontal);
+			TSharedRef<SScrollBar> VerticalScrollBar =
+				SNew(SScrollBar)
+				.AlwaysShowScrollbar(true)
+				.Thickness(12.0f)
+				.Orientation(Orient_Vertical);
+
+			// Widget
+			ChildSlot
+			[
+				SNew(SVerticalBox)
+				+SVerticalBox::Slot()
+				.AutoHeight()
+				[
+					// Select cache buffer
+					SNew(SHorizontalBox)
+					+SHorizontalBox::Slot()
+					.AutoWidth()
+					.Padding(10.0f)
+					.VAlign(VAlign_Center)
+					.HAlign(HAlign_Right)
+					[
+						SNew(STextBlock)
+						.Text(LOCTEXT("CacheBufferSelection", "Cache Buffer Selection"))
+						.Font(IDetailLayoutBuilder::GetDetailFont())
+					]
+					+SHorizontalBox::Slot()
+					.AutoWidth()
+					.VAlign(VAlign_Center)
+					.HAlign(HAlign_Left)
+					[
+						SNew(SComboBox<TSharedPtr<FBufferSelectionInfo>>)
+						.OptionsSource(&BufferSelectionList)
+						.OnGenerateWidget(this, &SSimCacheDataBufferWidget::BufferSelectionGenerateWidget)
+						.OnSelectionChanged(this, &SSimCacheDataBufferWidget::BufferSelectionChanged)
+						.InitiallySelectedItem(BufferSelectionList[0])
+						[
+							SNew(STextBlock)
+							.Text(this, &SSimCacheDataBufferWidget::GetBufferSelectionText)
+						]
+					]
+					// Select Cache Frame
+					+SHorizontalBox::Slot()
+					.AutoWidth()
+					.Padding(10.0f)
+					.VAlign(VAlign_Center)
+					.HAlign(HAlign_Right)
+					[
+						SNew(STextBlock)
+						.Text(LOCTEXT("FrameIndex", "Frame Index"))
+						.Font(IDetailLayoutBuilder::GetDetailFont())
+					]
+					+SHorizontalBox::Slot()
+					.AutoWidth()
+					.VAlign(VAlign_Center)
+					.HAlign(HAlign_Left)
+					[
+						SNew(SSpinBox<int32>)
+						.MinValue(0)
+						.MaxValue(this, &SSimCacheDataBufferWidget::GetNumFrames)
+						.Value(this, &SSimCacheDataBufferWidget::GetFrameIndex)
+						.OnValueChanged(this, &SSimCacheDataBufferWidget::SetFrameIndex)
+					]
+					// Component Filter
+					+SHorizontalBox::Slot()
+					.AutoWidth()
+					.Padding(10.0f)
+					.VAlign(VAlign_Center)
+					.HAlign(HAlign_Right)
+					[
+						SNew(STextBlock)
+						.Text(LOCTEXT("ComponentFilter", "Component Filter"))
+						.Font(IDetailLayoutBuilder::GetDetailFont())
+					]
+					+SHorizontalBox::Slot()
+					.AutoWidth()
+					.VAlign(VAlign_Center)
+					.HAlign(HAlign_Left)
+					[
+						SNew(SEditableTextBox)
+						.OnTextChanged(this, &SSimCacheDataBufferWidget::OnComponentFilterChange)
+						.MinDesiredWidth(100)
+					]
+				]
+				//// Main Spreadsheet View
+				+SVerticalBox::Slot()
+				//.FillHeight(1)
+				//.AutoHeight()
+				[
+					SNew(SHorizontalBox)
+					+ SHorizontalBox::Slot()
+					[
+						SNew(SScrollBox)
+						.Orientation(Orient_Horizontal)
+						.ExternalScrollbar(HorizontalScrollBar)
+						+SScrollBox::Slot()
+						[
+							SAssignNew(ListViewWidget, SListView<TSharedPtr<int32>>)
+							.ListItemsSource(&RowItems)
+							.OnGenerateRow(this, &SSimCacheDataBufferWidget::MakeRowWidget)
+							.Visibility(EVisibility::Visible)
+							.SelectionMode(ESelectionMode::Single)
+							.ExternalScrollbar(VerticalScrollBar)
+							.HeaderRow(HeaderRowWidget)
+						]
+					]
+					+ SHorizontalBox::Slot()
+					.AutoWidth()
+					[
+						VerticalScrollBar
+					]
+				]
+				+ SVerticalBox::Slot()
+				.AutoHeight()
+				[
+					HorizontalScrollBar
+				]
+			];
+		}
+
+		TSharedRef<ITableRow> MakeRowWidget(const TSharedPtr<int32> RowIndexPtr, const TSharedRef<STableViewBase>& OwnerTable) const
+		{
+			return
+				SNew(SSimCacheDataBufferRowWidget, OwnerTable)
+				.RowIndexPtr(RowIndexPtr)
+				.BufferReader(BufferReader);
+		}
+
+		void UpdateColumns(const bool bRefresh)
+		{
+			HeaderRowWidget->ClearColumns();
+
+			HeaderRowWidget->AddColumn(
+				SHeaderRow::Column(NAME_Instance)
+				.DefaultLabel(FText::FromName(NAME_Instance))
+				.HAlignHeader(EHorizontalAlignment::HAlign_Center)
+				.VAlignHeader(EVerticalAlignment::VAlign_Fill)
+				.HAlignCell(EHorizontalAlignment::HAlign_Center)
+				.VAlignCell(EVerticalAlignment::VAlign_Fill)
+				.SortMode(EColumnSortMode::None)
+			);
+
+			const bool bFilterActive = ComponentFilterArray.Num() > 0;
+			for (const FSimCacheBufferReader::FComponentInfo& ComponentInfo : BufferReader->GetComponentInfos())
+			{
+				if (bFilterActive)
+				{
+					const FString ComponentInfoString = ComponentInfo.Name.ToString();
+					const bool bPassedFilter = ComponentFilterArray.ContainsByPredicate(
+						[ComponentInfoString=ComponentInfo.Name.ToString()](const FString& ComponentFilter)
+						{
+							return ComponentInfoString.Contains(ComponentFilter);
+						}
+					);
+					if ( bPassedFilter == false)
+					{
+						continue;
+					}
+				}
+
+				HeaderRowWidget->AddColumn(
+					SHeaderRow::Column(ComponentInfo.Name)
+					.DefaultLabel(FText::FromName(ComponentInfo.Name))
+					.HAlignHeader(EHorizontalAlignment::HAlign_Center)
+					.VAlignHeader(EVerticalAlignment::VAlign_Fill)
+					.HAlignCell(EHorizontalAlignment::HAlign_Center)
+					.VAlignCell(EVerticalAlignment::VAlign_Fill)
+					.SortMode(EColumnSortMode::None)
+				);
+			}
+
+			HeaderRowWidget->ResetColumnWidths();
+			HeaderRowWidget->RefreshColumns();
+			if (bRefresh && ListViewWidget)
+			{
+				ListViewWidget->RequestListRefresh();
+			}
+		}
+
+		void UpdateRows(const bool bRefresh)
+		{
+			RowItems.Reset(BufferReader->GetNumInstances());
+			for (int32 i = 0; i < BufferReader->GetNumInstances(); ++i)
+			{
+				RowItems.Emplace(MakeShared<int32>(i));
+			}
+
+			if (bRefresh && ListViewWidget )
+			{
+				ListViewWidget->RequestListRefresh();
+			}
+		}
+
+		void UpdateBufferSelectionList()
+		{
+			BufferSelectionList.Empty();
+
+			UNiagaraSimCache* SimCache = WeakSimCache.Get();
+			if (SimCache && SimCache->IsCacheValid())
+			{
+				BufferSelectionList.Emplace(MakeShared<FBufferSelectionInfo>(INDEX_NONE, LOCTEXT("SystemInstance", "System Instance")));
+				for ( int32 i=0; i < SimCache->CacheLayout.EmitterLayouts.Num(); ++i )
+				{
+					BufferSelectionList.Emplace(
+						MakeShared<FBufferSelectionInfo>(i, FText::Format(LOCTEXT("EmitterFormat", "Emitter - {0}"), FText::FromName(SimCache->CacheLayout.EmitterLayouts[i].LayoutName)))
+					);
+				}
+			}
+			else
+			{
+				BufferSelectionList.Emplace(MakeShared<FBufferSelectionInfo>(INDEX_NONE, LOCTEXT("InvalidCache", "Invalid Cache")));
+			}
+		}
+
+		TOptional<int32> GetNumFrames() const { return BufferReader->GetNumFrames(); }
+		int32 GetFrameIndex() const { return BufferReader->GetFrameIndex(); }
+
+		void SetFrameIndex(int32 Value)
+		{
+			BufferReader->SetFrameIndex(Value);
+			UpdateRows(true);
+		}
+
+		TSharedRef<SWidget> BufferSelectionGenerateWidget(TSharedPtr<FBufferSelectionInfo> InItem)
+		{
+			return
+				SNew(STextBlock)
+				.Text(InItem->Value)
+				.Font(IDetailLayoutBuilder::GetDetailFont());
+		}
+
+		void BufferSelectionChanged(TSharedPtr<FBufferSelectionInfo> NewSelection, ESelectInfo::Type SelectInfo)
+		{
+			BufferReader->SetEmitterIndex(NewSelection.IsValid() ? NewSelection->Key : INDEX_NONE);
+			UpdateColumns(false);
+			UpdateRows(true);
+		}
+
+		FText GetBufferSelectionText() const
+		{
+			const int32 EmitterIndex = BufferReader->GetEmitterIndex();
+			for ( const TSharedPtr<FBufferSelectionInfo>& SelectionInfo : BufferSelectionList )
+			{
+				if ( SelectionInfo->Key == EmitterIndex )
+				{
+					return SelectionInfo->Value;
+				}
+			}
+			return FText::GetEmpty();
+		}
+
+		void OnComponentFilterChange(const FText& InFilter)
+		{
+			ComponentFilterArray.Empty();
+			InFilter.ToString().ParseIntoArray(ComponentFilterArray, TEXT(","));
+			UpdateColumns(true);
+		}
+
+		TWeakObjectPtr<UNiagaraSimCache>			WeakSimCache;
+		TArray<TSharedPtr<int32>>					RowItems;
+		TSharedPtr<FSimCacheBufferReader>			BufferReader;
+
+		TArray<TSharedPtr<FBufferSelectionInfo>>	BufferSelectionList;
+
+		int32										FrameIndex = 0;
+		
+		TSharedPtr<SHeaderRow>						HeaderRowWidget;
+		TSharedPtr<SListView<TSharedPtr<int32>>>	ListViewWidget;
+
+		TArray<FString>								ComponentFilterArray;
+	};
+}
+
+TSharedRef<IDetailCustomization> FNiagaraSimCacheDetailsCustomization::MakeInstance()
+{
+	return MakeShared<FNiagaraSimCacheDetailsCustomization>();
+}
+
+void FNiagaraSimCacheDetailsCustomization::CustomizeDetails(IDetailLayoutBuilder& DetailBuilder)
+{
+	static const FName NAME_SimCache("SimCache");
+
+	TArray<TWeakObjectPtr<UObject>> ObjectsCustomized;
+	DetailBuilder.GetObjectsBeingCustomized(ObjectsCustomized);
+	if (ObjectsCustomized.Num() != 1 || !ObjectsCustomized[0]->IsA<UNiagaraSimCache>())
+	{
+		return;
+	}
+
+	UNiagaraSimCache* SimcCache = CastChecked<UNiagaraSimCache>(ObjectsCustomized[0]);
+	IDetailCategoryBuilder& DetailCategory = DetailBuilder.EditCategory(NAME_SimCache);
+
+	DetailCategory.AddCustomRow(LOCTEXT("DataBuffers", "DataBuffers"))
+	.WholeRowWidget
+	[
+		SNew(NiagaraSimCacheDetailsCustomizationLocal::SSimCacheDataBufferWidget)
+		.SimCache(SimcCache)
+	];
+}
+
+#undef LOCTEXT_NAMESPACE
diff --git a/Engine/Plugins/FX/Niagara/Source/NiagaraEditorWidgets/Private/DetailCustomizations/NiagaraSimCacheDetailsCustomization.h b/Engine/Plugins/FX/Niagara/Source/NiagaraEditorWidgets/Private/DetailCustomizations/NiagaraSimCacheDetailsCustomization.h
new file mode 100644
--- /dev/null
+++ b/Engine/Plugins/FX/Niagara/Source/NiagaraEditorWidgets/Private/DetailCustomizations/NiagaraSimCacheDetailsCustomization.h
@@ -0,0 +1,18 @@
+// Copyright Epic Games, Inc. All Rights Reserved.
+
+#pragma once
+
+#include "IDetailCustomization.h"
+
+class IDetailLayoutBuilder;
+class UNiagaraSimCache;
+
+class FNiagaraSimCacheDetailsCustomization : public IDetailCustomization
+{
+public:
+	static TSharedRef<IDetailCustomization> MakeInstance();
+
+	// Begin IDetailCustomization interface
+	virtual void CustomizeDetails(class IDetailLayoutBuilder& DetailBuilder) override;
+	// End IDetailCustomization interface
+};
diff --git a/Engine/Plugins/FX/Niagara/Source/NiagaraEditorWidgets/Private/NiagaraEditorWidgetsModule.cpp b/Engine/Plugins/FX/Niagara/Source/NiagaraEditorWidgets/Private/NiagaraEditorWidgetsModule.cpp
--- a/Engine/Plugins/FX/Niagara/Source/NiagaraEditorWidgets/Private/NiagaraEditorWidgetsModule.cpp
+++ b/Engine/Plugins/FX/Niagara/Source/NiagaraEditorWidgets/Private/NiagaraEditorWidgetsModule.cpp
@@ -12,6 +12,7 @@
 #include "DetailCustomizations/NiagaraDataInterfaceStaticMeshDetails.h"
 #include "DetailCustomizations/NiagaraDataInterfaceMeshRendererInfoDetails.h"
 #include "DetailCustomizations/NiagaraMeshRendererDetails.h"
+#include "DetailCustomizations/NiagaraSimCacheDetailsCustomization.h"
 #include "ViewModels/NiagaraSystemViewModel.h"
 #include "SNiagaraOverviewGraph.h"
 #include "NiagaraEditorWidgetsUtilities.h"
@@ -130,6 +131,7 @@ void FNiagaraEditorWidgetsModule::StartupModule()
 	PropertyModule.RegisterCustomClassLayout("NiagaraDataInterfaceGrid3DCollection", FOnGetDetailCustomizationInstance::CreateStatic(&FNiagaraDataInterfaceGrid3DCollectionDetails::MakeInstance));
 	PropertyModule.RegisterCustomClassLayout("NiagaraDataInterfaceMeshRendererInfo", FOnGetDetailCustomizationInstance::CreateStatic(&FNiagaraDataInterfaceMeshRendererInfoDetails::MakeInstance));
 	PropertyModule.RegisterCustomClassLayout("NiagaraMeshRendererProperties", FOnGetDetailCustomizationInstance::CreateStatic(&FNiagaraMeshRendererDetails::MakeInstance));
+	PropertyModule.RegisterCustomClassLayout("NiagaraSimCache", FOnGetDetailCustomizationInstance::CreateStatic(&FNiagaraSimCacheDetailsCustomization::MakeInstance));
 
 	ReinitializeStyleCommand = IConsoleManager::Get().RegisterConsoleCommand(
 		TEXT("fx.NiagaraEditorWidgets.ReinitializeStyle"),
