- Fix for missing include that got lost somewhere during merges

#rnx

diff --git a/Engine/Plugins/FX/Niagara/Source/NiagaraEditor/Private/ViewModels/NiagaraBakerViewModel.cpp b/Engine/Plugins/FX/Niagara/Source/NiagaraEditor/Private/ViewModels/NiagaraBakerViewModel.cpp
--- a/Engine/Plugins/FX/Niagara/Source/NiagaraEditor/Private/ViewModels/NiagaraBakerViewModel.cpp
+++ b/Engine/Plugins/FX/Niagara/Source/NiagaraEditor/Private/ViewModels/NiagaraBakerViewModel.cpp
@@ -4,6 +4,7 @@
 #include "ViewModels/NiagaraSystemViewModel.h"
 #include "Widgets/SNiagaraBakerWidget.h"
 #include "NiagaraBakerRenderer.h"
+#include "NiagaraGpuComputeDispatchInterface.h"
 #include "NiagaraComponent.h"
 #include "NiagaraSystem.h"
 
