Added versioning for Niagara emitters

diff --git a/Engine/Plugins/FX/CascadeToNiagaraConverter/Source/CascadeToNiagaraConverter/Private/NiagaraStackGraphUtilitiesAdapterLibrary.cpp b/Engine/Plugins/FX/CascadeToNiagaraConverter/Source/CascadeToNiagaraConverter/Private/NiagaraStackGraphUtilitiesAdapterLibrary.cpp
--- a/Engine/Plugins/FX/CascadeToNiagaraConverter/Source/CascadeToNiagaraConverter/Private/NiagaraStackGraphUtilitiesAdapterLibrary.cpp
+++ b/Engine/Plugins/FX/CascadeToNiagaraConverter/Source/CascadeToNiagaraConverter/Private/NiagaraStackGraphUtilitiesAdapterLibrary.cpp
@@ -43,7 +43,6 @@
 #include "Particles/Velocity/ParticleModuleVelocity.h"
 #include "Particles/Velocity/ParticleModuleVelocityOverLifetime.h"
 #include "Particles/Acceleration/ParticleModuleAccelerationConstant.h"
-#include "Particles/TypeData/ParticleModuleTypeDataBase.h"
 #include "Particles/TypeData/ParticleModuleTypeDataGpu.h"
 #include "Particles/TypeData/ParticleModuleTypeDataMesh.h"
 #include "Particles/TypeData/ParticleModuleTypeDataRibbon.h"
@@ -55,11 +54,8 @@
 #include "NiagaraEditorModule.h"
 #include "ViewModels/NiagaraEmitterHandleViewModel.h"
 #include "ViewModels/NiagaraEmitterViewModel.h"
-#include "ViewModels/NiagaraScriptViewModel.h"
 
-#include "ViewModels/Stack/NiagaraStackClipboardUtilities.h"
 #include "ViewModels/Stack/NiagaraStackViewModel.h"
-#include "NiagaraEmitterFactoryNew.h"
 #include "ViewModels/Stack/NiagaraStackItemGroup.h"
 #include "ViewModels/Stack/NiagaraStackEventScriptItemGroup.h"
 #include "NiagaraRendererProperties.h"
@@ -87,10 +83,7 @@
 #include "Distributions/DistributionVectorUniform.h"
 #include "Distributions/DistributionVectorUniformCurve.h"
 #include "Distributions/DistributionVectorParticleParameter.h"
-#include "MessageLogModule.h"
 #include "IMessageLogListing.h"
-#include "EdGraph/EdGraphNode.h"
-#include "NiagaraNode.h"
 #include "NiagaraNodeFunctionCall.h"
 #include "CascadeToNiagaraConverterModule.h"
 #include "Curves/RichCurve.h"
@@ -1586,7 +1579,7 @@ UNiagaraEmitterConversionContext* UNiagaraSystemConversionContext::AddEmptyEmitt
 	FName NewEmitterName = FNiagaraEditorUtilities::GetUniqueObjectName(System, UNiagaraEmitter::StaticClass(), NewEmitterNameString);
 	UNiagaraEmitter* NewEmitter = CastChecked<UNiagaraEmitter>(StaticLoadObject(UNiagaraEmitter::StaticClass(), nullptr, TEXT("/Niagara/DefaultAssets/Templates/CascadeConversion/CompletelyEmpty")));
 	
-	const TSharedPtr<FNiagaraEmitterHandleViewModel>& NewEmitterHandleViewModel = SystemViewModel->AddEmitter(*NewEmitter);
+	const TSharedPtr<FNiagaraEmitterHandleViewModel>& NewEmitterHandleViewModel = SystemViewModel->AddEmitter(*NewEmitter, NewEmitter->GetExposedVersion().VersionGuid);
 	NewEmitterHandleViewModel->SetName(NewEmitterName);
 	NewEmitterName = NewEmitterHandleViewModel->GetName();
 
@@ -1618,7 +1611,7 @@ UNiagaraEmitterConversionContext* const* UNiagaraSystemConversionContext::FindEm
 
 void UNiagaraEmitterConversionContext::Cleanup()
 {
-	Emitter = nullptr;
+	Emitter = FVersionedNiagaraEmitter();
 	EmitterHandleViewModel.Reset();
 	StackEntryAddActions.Empty();
 	RendererNameToStagedRendererPropertiesMap.Empty();
@@ -1757,7 +1750,7 @@ void UNiagaraEmitterConversionContext::InternalFinalizeEvents(UNiagaraSystemConv
 				UNiagaraEmitterConversionContext* SourceEmitterConversionContext = *SourceEmitterConversionContextPtr;
 
 				// Enable persistent IDs for the event generator emitter as this is a requisite for events.
-				SourceEmitterConversionContext->Emitter->bRequiresPersistentIDs = true;
+				SourceEmitterConversionContext->Emitter.GetEmitterData()->bRequiresPersistentIDs = true;
 
 				// Add the event generator.
 				SourceEmitterConversionContext->FindOrAddModuleScript("EventGenerator", EventHandlerAddAction.AddEventGeneratorOptions.EventGeneratorScriptAssetData, EScriptExecutionCategory::ParticleUpdate);
@@ -1826,8 +1819,8 @@ void UNiagaraEmitterConversionContext::InternalFinalizeStackEntryAddActions()
 		}
 		else if (ExecutionCategory == EScriptExecutionCategory::ParticleEvent)
 		{
-			auto GetEventSourceEmitterId = [](const FName EventName, UNiagaraEmitter* InEmitter)->const FGuid {
-				const FNiagaraEventScriptProperties* FoundEventProps = InEmitter->GetEventHandlers().FindByPredicate(
+			auto GetEventSourceEmitterId = [](const FName EventName, FVersionedNiagaraEmitter InEmitter)->const FGuid {
+				const FNiagaraEventScriptProperties* FoundEventProps = InEmitter.GetEmitterData()->GetEventHandlers().FindByPredicate(
 					[EventName](const FNiagaraEventScriptProperties& EventProps) {return EventProps.SourceEventName == EventName; }
 				);
 				if (FoundEventProps != nullptr)
@@ -1837,7 +1830,7 @@ void UNiagaraEmitterConversionContext::InternalFinalizeStackEntryAddActions()
 				return FGuid();
 			};
 
-			UNiagaraEmitter* EmitterInstance = EmitterHandleViewModel->GetEmitterHandle()->GetInstance();
+			FVersionedNiagaraEmitter EmitterInstance = EmitterHandleViewModel->GetEmitterHandle()->GetInstance();
 			const FGuid EventSourceEmitterId = GetEventSourceEmitterId(StackEntryID.EventName, EmitterInstance);
 
 			auto FilterStackEntryForEventSourceEmitterId = [EventSourceEmitterId](UNiagaraStackItemGroup* ItemGroup)->bool {
@@ -1990,7 +1983,7 @@ void UNiagaraEmitterConversionContext::InternalFinalizeStackEntryAddActions()
 	// Push the emitter messages.
 	for (FGenericConverterMessage& Message : EmitterMessages)
 	{
-		UNiagaraMessageDataText* NewMessageDataText = NewObject<UNiagaraMessageDataText>(Emitter);
+		UNiagaraMessageDataText* NewMessageDataText = NewObject<UNiagaraMessageDataText>(Emitter.Emitter);
 		const FName TopicName = Message.bIsVerbose ? FNiagaraConverterMessageTopics::VerboseConversionEventTopicName : FNiagaraConverterMessageTopics::ConversionEventTopicName;
 		NewMessageDataText->Init(FText::FromString(Message.Message), Message.MessageSeverity, TopicName);
 		EmitterHandleViewModel->AddMessage(NewMessageDataText);
diff --git a/Engine/Plugins/FX/CascadeToNiagaraConverter/Source/CascadeToNiagaraConverter/Public/NiagaraStackGraphUtilitiesAdapterLibrary.h b/Engine/Plugins/FX/CascadeToNiagaraConverter/Source/CascadeToNiagaraConverter/Public/NiagaraStackGraphUtilitiesAdapterLibrary.h
--- a/Engine/Plugins/FX/CascadeToNiagaraConverter/Source/CascadeToNiagaraConverter/Public/NiagaraStackGraphUtilitiesAdapterLibrary.h
+++ b/Engine/Plugins/FX/CascadeToNiagaraConverter/Source/CascadeToNiagaraConverter/Public/NiagaraStackGraphUtilitiesAdapterLibrary.h
@@ -24,12 +24,15 @@
 #include "Particles/TypeData/ParticleModuleTypeDataRibbon.h"
 #include "Curves/RichCurve.h"
 #include "NiagaraEmitter.h"
+#include "NiagaraMessages.h"
 #include "Particles/Location/ParticleModuleLocationPrimitiveCylinder.h"
 #include "Particles/Location/ParticleModuleLocationBoneSocket.h"
 #include "Particles/Orientation/ParticleModuleOrientationAxisLock.h"
 #include "Particles/Velocity/ParticleModuleVelocityInheritParent.h"
 #include "NiagaraStackGraphUtilitiesAdapterLibrary.generated.h"
 
+class UNiagaraScriptConversionContextInput;
+class UNiagaraScriptConversionContext;
 class UNiagaraSystem;
 class UNiagaraEmitter;
 class UNiagaraRendererProperties;
@@ -672,7 +675,7 @@ public:
 	void Finalize();
 
 	UFUNCTION(BlueprintCallable, Category = "FXConverterUtilities")
-	UNiagaraEmitter* GetEmitter() {return Emitter;};
+	UNiagaraEmitter* GetEmitter() {return Emitter.Emitter;}
 
 	void InternalFinalizeEvents(UNiagaraSystemConversionContext* OwningSystemConversionContext);
 
@@ -683,7 +686,7 @@ public:
 	 * @param InEmitter						The Emitter to convert.
 	 * @param InEmitterHandleViewModelGuid	A Guid key to the FNiagaraEmitterHandleViewModel pointing at the InEmitter.
 	 */
-	void Init(UNiagaraEmitter* InEmitter, const TSharedPtr<FNiagaraEmitterHandleViewModel>& InEmitterHandleViewModel)
+	void Init(FVersionedNiagaraEmitter InEmitter, const TSharedPtr<FNiagaraEmitterHandleViewModel>& InEmitterHandleViewModel)
 	{
 		Emitter = InEmitter;
 		EmitterHandleViewModel = InEmitterHandleViewModel;
@@ -708,7 +711,7 @@ public:
 	void RemoveModuleScriptsForAssets(TArray<FAssetData> ScriptsToRemove);
 
 private:
-	UNiagaraEmitter* Emitter;
+	FVersionedNiagaraEmitter Emitter;
 
 	TSharedPtr<FNiagaraEmitterHandleViewModel> EmitterHandleViewModel;
 
diff --git a/Engine/Plugins/FX/Niagara/Source/Niagara/Classes/NiagaraComputeExecutionContext.h b/Engine/Plugins/FX/Niagara/Source/Niagara/Classes/NiagaraComputeExecutionContext.h
--- a/Engine/Plugins/FX/Niagara/Source/Niagara/Classes/NiagaraComputeExecutionContext.h
+++ b/Engine/Plugins/FX/Niagara/Source/Niagara/Classes/NiagaraComputeExecutionContext.h
@@ -145,7 +145,7 @@ public:
 	FString DebugSimName;
 #endif
 	TWeakObjectPtr<class USceneComponent>	ProfilingComponentPtr;
-	TWeakObjectPtr<UNiagaraEmitter>			ProfilingEmitterPtr;
+	FVersionedNiagaraEmitterWeakPtr			ProfilingEmitterPtr;
 
 	const TArray<UNiagaraDataInterface*>& GetDataInterfaces()const { return CombinedParamStore.GetDataInterfaces(); }
 
diff --git a/Engine/Plugins/FX/Niagara/Source/Niagara/Classes/NiagaraDataInterface.h b/Engine/Plugins/FX/Niagara/Source/Niagara/Classes/NiagaraDataInterface.h
--- a/Engine/Plugins/FX/Niagara/Source/Niagara/Classes/NiagaraDataInterface.h
+++ b/Engine/Plugins/FX/Niagara/Source/Niagara/Classes/NiagaraDataInterface.h
@@ -4,9 +4,6 @@
 #include "Niagara/Public/NiagaraCommon.h"
 #include "NiagaraShared.h"
 #include "VectorVM.h"
-#include "StaticMeshResources.h"
-#include "Curves/RichCurve.h"
-#include "NiagaraMergeable.h"
 #include "NiagaraDataInterfaceBase.h"
 #include "NiagaraScriptBase.h"
 #include "NiagaraDataInterface.generated.h"
@@ -543,7 +540,7 @@ public:
 #endif
 
     /** Method to add asset tags that are specific to this data interface. By default we add in how many instances of this class exist in the list.*/
-	virtual void GetAssetTagsForContext(const UObject* InAsset, const TArray<const UNiagaraDataInterface*>& InProperties, TMap<FName, uint32>& NumericKeys, TMap<FName, FString>& StringKeys) const;
+	virtual void GetAssetTagsForContext(const UObject* InAsset, FGuid AssetVersion, const TArray<const UNiagaraDataInterface*>& InProperties, TMap<FName, uint32>& NumericKeys, TMap<FName, FString>& StringKeys) const;
 	virtual bool CanExposeVariables() const { return false; }
 	virtual void GetExposedVariables(TArray<FNiagaraVariableBase>& OutVariables) const {}
 	virtual bool GetExposedVariableValue(const FNiagaraVariableBase& InVariable, void* InPerInstanceData, FNiagaraSystemInstance* InSystemInstance, void* OutData) const { return false; }
@@ -561,7 +558,7 @@ public:
 	* Allows a DI to specify data dependencies between emitters, so the system can ensure that the emitter instances are executed in the correct order.
 	* The Dependencies array may already contain items, and this method should only append to it.
 	*/
-	virtual void GetEmitterDependencies(UNiagaraSystem* Asset, TArray<UNiagaraEmitter*>& Dependencies) const {}
+	virtual void GetEmitterDependencies(UNiagaraSystem* Asset, TArray<FVersionedNiagaraEmitter>& Dependencies) const {}
 
 	virtual bool ReadsEmitterParticleData(const FString& EmitterName) const { return false; }
 
diff --git a/Engine/Plugins/FX/Niagara/Source/Niagara/Classes/NiagaraDataInterfaceCollisionQuery.h b/Engine/Plugins/FX/Niagara/Source/Niagara/Classes/NiagaraDataInterfaceCollisionQuery.h
--- a/Engine/Plugins/FX/Niagara/Source/Niagara/Classes/NiagaraDataInterfaceCollisionQuery.h
+++ b/Engine/Plugins/FX/Niagara/Source/Niagara/Classes/NiagaraDataInterfaceCollisionQuery.h
@@ -45,7 +45,7 @@ public:
 	
 	virtual void GetFunctions(TArray<FNiagaraFunctionSignature>& OutFunctions) override;
 	virtual void GetVMExternalFunction(const FVMExternalFunctionBindingInfo& BindingInfo, void* InstanceData, FVMExternalFunction &OutFunc) override;
-	virtual void GetAssetTagsForContext(const UObject* InAsset, const TArray<const UNiagaraDataInterface*>& InProperties, TMap<FName, uint32>& NumericKeys, TMap<FName, FString>& StringKeys) const override;
+	virtual void GetAssetTagsForContext(const UObject* InAsset, FGuid AssetVersion, const TArray<const UNiagaraDataInterface*>& InProperties, TMap<FName, uint32>& NumericKeys, TMap<FName, FString>& StringKeys) const override;
 
 	// VM functions
 	void PerformQuerySyncCPU(FVectorVMExternalFunctionContext& Context);
diff --git a/Engine/Plugins/FX/Niagara/Source/Niagara/Classes/NiagaraDataInterfaceGrid2DCollectionReader.h b/Engine/Plugins/FX/Niagara/Source/Niagara/Classes/NiagaraDataInterfaceGrid2DCollectionReader.h
--- a/Engine/Plugins/FX/Niagara/Source/Niagara/Classes/NiagaraDataInterfaceGrid2DCollectionReader.h
+++ b/Engine/Plugins/FX/Niagara/Source/Niagara/Classes/NiagaraDataInterfaceGrid2DCollectionReader.h
@@ -31,7 +31,7 @@ public:
 
 	virtual bool Equals(const UNiagaraDataInterface* Other) const override;
 	virtual bool InitPerInstanceData(void* PerInstanceData, FNiagaraSystemInstance* SystemInstance) override;
-	virtual void GetEmitterDependencies(UNiagaraSystem* Asset, TArray<UNiagaraEmitter*>& Dependencies) const override;
+	virtual void GetEmitterDependencies(UNiagaraSystem* Asset, TArray<FVersionedNiagaraEmitter>& Dependencies) const override;
 	virtual void GetFunctions(TArray<FNiagaraFunctionSignature>& OutFunctions) override;
 	virtual bool PerInstanceTickPostSimulate(void* PerInstanceData, FNiagaraSystemInstance* SystemInstance, float DeltaSeconds) override;
 
diff --git a/Engine/Plugins/FX/Niagara/Source/Niagara/Classes/NiagaraDataInterfaceGrid3DCollectionReader.h b/Engine/Plugins/FX/Niagara/Source/Niagara/Classes/NiagaraDataInterfaceGrid3DCollectionReader.h
--- a/Engine/Plugins/FX/Niagara/Source/Niagara/Classes/NiagaraDataInterfaceGrid3DCollectionReader.h
+++ b/Engine/Plugins/FX/Niagara/Source/Niagara/Classes/NiagaraDataInterfaceGrid3DCollectionReader.h
@@ -30,7 +30,7 @@ public:
 
 	virtual bool Equals(const UNiagaraDataInterface* Other) const override;
 	virtual bool InitPerInstanceData(void* PerInstanceData, FNiagaraSystemInstance* SystemInstance) override;
-	virtual void GetEmitterDependencies(UNiagaraSystem* Asset, TArray<UNiagaraEmitter*>& Dependencies) const override;
+	virtual void GetEmitterDependencies(UNiagaraSystem* Asset, TArray<FVersionedNiagaraEmitter>& Dependencies) const override;
 	virtual void GetFunctions(TArray<FNiagaraFunctionSignature>& OutFunctions) override;
 	virtual bool PerInstanceTickPostSimulate(void* PerInstanceData, FNiagaraSystemInstance* SystemInstance, float DeltaSeconds) override;
 
diff --git a/Engine/Plugins/FX/Niagara/Source/Niagara/Classes/NiagaraEffectType.h b/Engine/Plugins/FX/Niagara/Source/Niagara/Classes/NiagaraEffectType.h
--- a/Engine/Plugins/FX/Niagara/Source/Niagara/Classes/NiagaraEffectType.h
+++ b/Engine/Plugins/FX/Niagara/Source/Niagara/Classes/NiagaraEffectType.h
@@ -508,11 +508,11 @@ void UNiagaraEffectType::ForEachPlatformSet(TAction Func)
 {
 	for (FNiagaraSystemScalabilitySettings& Setting : SystemScalabilitySettings.Settings)
 	{
-		Func(this, Setting.Platforms);
+		Func(Setting.Platforms);
 	}
 
 	for (FNiagaraEmitterScalabilitySettings& Setting : EmitterScalabilitySettings.Settings)
 	{
-		Func(this, Setting.Platforms);
+		Func(Setting.Platforms);
 	}
 }
diff --git a/Engine/Plugins/FX/Niagara/Source/Niagara/Classes/NiagaraEmitter.h b/Engine/Plugins/FX/Niagara/Source/Niagara/Classes/NiagaraEmitter.h
--- a/Engine/Plugins/FX/Niagara/Source/Niagara/Classes/NiagaraEmitter.h
+++ b/Engine/Plugins/FX/Niagara/Source/Niagara/Classes/NiagaraEmitter.h
@@ -15,11 +15,13 @@
 #include "NiagaraRendererProperties.h"
 #include "NiagaraParameterDefinitionsBase.h"
 #include "NiagaraParameterDefinitionsSubscriber.h"
+#include "NiagaraScratchPadContainer.h"
 
 #include "NiagaraDataInterfacePlatformSet.h"
 
 #include "NiagaraEmitter.generated.h"
 
+class UNiagaraScratchPadContainer;
 class UMaterial;
 class UNiagaraEmitter;
 class UNiagaraEventReceiverEmitterAction;
@@ -213,11 +215,290 @@
 	float Cine;
 };
 
-struct MemoryRuntimeEstimation
+struct NIAGARA_API FMemoryRuntimeEstimation
 {
 	TMap<uint64, int32> RuntimeAllocations;
 	bool IsEstimationDirty = false;
 	int32 AllocationEstimate = 0;
+
+	FCriticalSection* GetCriticalSection();
+	void Init();
+private:
+	TSharedPtr<FCriticalSection> EstimationCriticalSection;
 };
 
+/** Struct containing all of the data that can be different between different emitter versions.*/
+USTRUCT(BlueprintInternalUseOnly)
+struct NIAGARA_API FVersionedNiagaraEmitterData
+{
+	GENERATED_BODY()
+
+	FVersionedNiagaraEmitterData();
+	
+	UPROPERTY()
+	FNiagaraAssetVersion Version;
+
+	/** What changed in this version compared to the last? Displayed to the user when upgrading to a new script version. */
+	UPROPERTY()
+	FText VersionChangeDescription;
+
+	/** Reference to a python script that is executed when the user updates from a previous version to this version. */
+	UPROPERTY()
+	ENiagaraPythonUpdateScriptReference UpdateScriptExecution = ENiagaraPythonUpdateScriptReference::None;
+
+	/** Python script to run when updating to this script version. */
+	UPROPERTY()
+	FString PythonUpdateScript;
+
+	/** Asset reference to a python script to run when updating to this script version. */
+	UPROPERTY()
+	FFilePath ScriptAsset;
+
+	/** Toggles whether or not the particles within this emitter are relative to the emitter origin or in global space. */ 
+	UPROPERTY(EditAnywhere, BlueprintReadWrite, Category = "Emitter")
+	bool bLocalSpace = false;
+
+	/** Toggles whether to globally make the random number generator be deterministic or non-deterministic. Any random calculation that is set to the emitter defaults will inherit this value. It is still possible to tweak individual random to be deterministic or not. In this case deterministic means that it will return the same results for the same configuration of the emitter as long as delta time is not variable. Any changes to the emitter's individual scripts will adjust the results. */
+	UPROPERTY(EditAnywhere, BlueprintReadWrite, Category = "Emitter")
+	bool bDeterminism = false;
+
+	/** When enabled, this will spawn using interpolated parameter values and perform a partial update at spawn time. This adds significant additional cost for spawning but will produce much smoother spawning for high spawn rates, erratic frame rates and fast moving emitters. */
+	UPROPERTY(EditAnywhere, BlueprintReadWrite, Category = "Emitter")
+	uint32 bInterpolatedSpawning : 1;
+
+	/** An emitter-based seed for the deterministic random number generator. */
+	UPROPERTY(EditAnywhere, BlueprintReadWrite, Category = "Emitter", meta = (EditCondition = "bDeterminism"))
+	int32 RandomSeed = 0;
+
+	UPROPERTY(EditAnywhere, BlueprintReadWrite, Category = "Emitter")
+	ENiagaraSimTarget SimTarget = ENiagaraSimTarget::CPUSim;
+
+	/** How should we calculate bounds for the emitter. */
+	UPROPERTY(EditAnywhere, Category = "Emitter")
+	ENiagaraEmitterCalculateBoundMode CalculateBoundsMode = ENiagaraEmitterCalculateBoundMode::Dynamic;
+	
+	/** The fixed bounding box value. CalculateBoundsMode is the condition whether the fixed bounds can be edited. */
+	UPROPERTY(EditAnywhere, BlueprintReadWrite, Category = "Emitter", meta = (EditConditionHides, EditCondition = "CalculateBoundsMode == ENiagaraEmitterCalculateBoundMode::Fixed"))
+	FBox FixedBounds;
+
+	/** Do particles in this emitter require a persistent ID? */
+	UPROPERTY(EditAnywhere, BlueprintReadWrite, Category = "Emitter")
+	uint32 bRequiresPersistentIDs : 1;
+
+	/** Performance option to allow event based spawning to be combined into a single spawn.  This will result in a single exec from 0 to number of particles rather than several, when using ExecIndex() it is recommended not to do this. */
+	UPROPERTY(EditAnywhere, BlueprintReadWrite, Category = "Emitter")
+	uint32 bCombineEventSpawn : 1;
+
+	UPROPERTY(meta=(NiagaraNoMerge))
+	TArray<FNiagaraEventScriptProperties> EventHandlerScriptProps;
+	
+	UPROPERTY(EditAnywhere, Category = "Scalability", meta=(DisplayInScalabilityContext))
+	FNiagaraPlatformSet Platforms;
+
+	UPROPERTY(EditAnywhere, Category = "Scalability", meta=(DisplayInScalabilityContext))
+	FNiagaraEmitterScalabilityOverrides ScalabilityOverrides;
+
+	/** Whether to limit the max tick delta time or not. */
+	UPROPERTY(EditAnywhere, BlueprintReadWrite, AdvancedDisplay, Category = "Emitter", meta = (InlineEditConditionToggle))
+	uint32 bLimitDeltaTime : 1;
+
+	/** An override on the max number of GPU particles we expect to spawn in a single frame. A value of 0 means it'll use fx.MaxNiagaraGPUParticlesSpawnPerFrame.*/
+	UPROPERTY(EditAnywhere, BlueprintReadWrite, AdvancedDisplay, Category = "Emitter", meta = (EditCondition = "SimTarget == ENiagaraSimTarget::GPUComputeSim", DisplayName = "Max GPU Particles Spawn per Frame"))
+	int32 MaxGPUParticlesSpawnPerFrame;
+
+	/** Limits the delta time per tick to prevent simulation spikes due to frame lags. */
+	UPROPERTY(EditAnywhere, BlueprintReadWrite, AdvancedDisplay, Category = "Emitter", meta = (EditCondition = "bLimitDeltaTime", ForceUnits=s))
+	float MaxDeltaTimePerTick = 0.125;
+
+	/**
+	The emitter needs to allocate memory for the particles each tick.
+	To prevent reallocations, the emitter should allocate as much memory as is needed for the max particle count.
+	This setting controls if the allocation size should be automatically determined or manually entered.
+	*/
+	UPROPERTY(EditAnywhere, BlueprintReadWrite, AdvancedDisplay, Category = "Emitter")
+	EParticleAllocationMode AllocationMode = EParticleAllocationMode::AutomaticEstimate;
+	
+	/** 
+	The emitter will allocate at least this many particles on it's first tick.
+	This can aid performance by avoiding many allocations as an emitter ramps up to it's max size.
+	*/
+	UPROPERTY(EditAnywhere, AdvancedDisplay, Category = "Emitter", meta = (EditCondition = "AllocationMode == EParticleAllocationMode::ManualEstimate"))
+	int32 PreAllocationCount = 0;
+
+	UPROPERTY()
+	FNiagaraEmitterScriptProperties UpdateScriptProps;
+
+	UPROPERTY()
+	FNiagaraEmitterScriptProperties SpawnScriptProps;
+
+	UPROPERTY()
+	FNiagaraParameterStore RendererBindings;
+
+	void CopyFrom(const FVersionedNiagaraEmitterData& Source);
+	void PostLoad(UNiagaraEmitter& Emitter, bool bIsCooked);
+	void PostInitProperties(UNiagaraEmitter* Outer);
+	bool UsesCollection(const UNiagaraParameterCollection* Collection) const;
+	bool UsesScript(const UNiagaraScript* Script) const;
+	bool IsReadyToRun() const;
+	const TArray<UNiagaraRendererProperties*>& GetRenderers() const { return RendererProperties; }
+	void GetScripts(TArray<UNiagaraScript*>& OutScripts, bool bCompilableOnly = true, bool bEnabledOnly = false) const;
+	UNiagaraScript* GetScript(ENiagaraScriptUsage Usage, FGuid UsageId);
+	UNiagaraScript* GetGPUComputeScript() { return GPUComputeScript; }
+	const UNiagaraScript* GetGPUComputeScript() const { return GPUComputeScript; }
+	FORCEINLINE const TArray<FNiagaraEventScriptProperties>& GetEventHandlers() const { return EventHandlerScriptProps; }
+	void CacheFromCompiledData(const FNiagaraDataSetCompiledData* CompiledData, const UNiagaraEmitter& Emitter);
+	void CacheFromShaderCompiled();
+	bool RequiresViewUniformBuffer() const { return bRequiresViewUniformBuffer; }
+	uint32 GetMaxInstanceCount() const { return MaxInstanceCount; }
+	uint32 GetMaxAllocationCount() const { return MaxAllocationCount; }
+	TConstArrayView<TSharedPtr<FNiagaraBoundsCalculator>> GetBoundsCalculators() const { return MakeArrayView(BoundsCalculators); }
+	bool RequiresPersistentIDs() const;
+	const TArray<UNiagaraSimulationStageBase*>& GetSimulationStages() const { return SimulationStages; }
+	FORCEINLINE const FNiagaraEmitterScalabilitySettings& GetScalabilitySettings()const { return CurrentScalabilitySettings; }
+	const FNiagaraEmitterScalabilityOverride& GetCurrentOverrideSettings() const;
+	UNiagaraSimulationStageBase* GetSimulationStageById(FGuid ScriptUsageId) const;
+	void RebuildRendererBindings(const UNiagaraEmitter& Emitter);
+	FORCEINLINE static FBox GetDefaultFixedBounds() { return FBox(FVector(-100), FVector(100)); }
+	
+	FVersionedNiagaraEmitter GetParent() const;
+	FVersionedNiagaraEmitter GetParentAtLastMerge() const;
+	void RemoveParent();
+	void Reparent(const FVersionedNiagaraEmitter& InParent);
+	
+	bool IsValid() const;
+	void UpdateDebugName(const UNiagaraEmitter& Emitter, const FNiagaraDataSetCompiledData* CompiledData);
+	void SyncEmitterAlias(const FString& InOldName, const UNiagaraEmitter& InEmitter);
+
+	template<typename TAction>
+	void ForEachRenderer(TAction Func) const;
+
+	template<typename TAction>
+	void ForEachEnabledRenderer(TAction Func) const;
+
+	template<typename TAction>
+	void ForEachScript(TAction Func) const;
+
+	template<typename TAction>
+	void ForEachPlatformSet(TAction Func);
+
+	/** Returns true if this emitter's platform filter allows it on this platform and quality level. */
+	bool IsAllowedByScalability() const;
+
+	/* Returns the number of max expected particles for memory allocations. */
+	int32 GetMaxParticleCountEstimate();
+
+	/* Gets whether or not the supplied event generator id matches an event generator which is shared between the particle spawn and update scrips. */
+	bool IsEventGeneratorShared(FName EventGeneratorId) const;
+
+	/* This is used by the emitter instances to report runtime allocations to reduce reallocation in future simulation runs. */
+	int32 AddRuntimeAllocation(uint64 ReporterHandle, int32 AllocationCount);
+	void ClearRuntimeAllocationEstimate(uint64 ReportHandle = INDEX_NONE);
+
+	/* Gets a pointer to an event handler by script usage id.  This method is potentially unsafe because modifications to
+	the event handler array can make this pointer become invalid without warning. */
+	FNiagaraEventScriptProperties* GetEventHandlerByIdUnsafe(FGuid ScriptUsageId);
+ 
+#if WITH_EDITORONLY_DATA
+	/** An allow list of Particle attributes (e.g. "Particle.Position" or "Particle.Age") that will not be removed from the DataSet  even if they aren't read by the VM.
+	    Used in conjunction with UNiagaraSystem::bTrimAttributes */
+	UPROPERTY(EditAnywhere, AdvancedDisplay, Category = "Emitter")
+	TArray<FString> AttributesToPreserve;
+
+	UPROPERTY()
+	FNiagaraEmitterScriptProperties EmitterSpawnScriptProps;
+
+	UPROPERTY()
+	FNiagaraEmitterScriptProperties EmitterUpdateScriptProps;
+
+	/** 'Source' data/graphs for the scripts used by this emitter. */
+	UPROPERTY()
+	TObjectPtr<UNiagaraScriptSourceBase> GraphSource = nullptr;
+
+	UPROPERTY()
+	TObjectPtr<UNiagaraScratchPadContainer> ScratchPads = nullptr;
+
+	UPROPERTY()
+	TObjectPtr<UNiagaraScratchPadContainer> ParentScratchPads = nullptr;
+
+	UPROPERTY()
+	FVersionedNiagaraEmitter VersionedParent;
+
+	UPROPERTY()
+	FVersionedNiagaraEmitter VersionedParentAtLastMerge;
+
+	bool AreAllScriptAndSourcesSynchronized() const;
+	void InvalidateCompileResults();
+	bool IsSynchronizedWithParent() const;
+	bool UsesEmitter(const UNiagaraEmitter& InEmitter) const;
+	void GatherStaticVariables(TArray<FNiagaraVariable>& OutVars) const;
+	
+	UNiagaraEditorDataBase* GetEditorData() const;
+	UNiagaraEditorParametersAdapterBase* GetEditorParameters();
+#endif
+
+#if STATS
+	FNiagaraStatDatabase& GetStatData() { return StatDatabase; }
+#endif
+
+#if WITH_NIAGARA_DEBUG_EMITTER_NAME
+	const TCHAR* GetDebugSimName() const { return *DebugSimName; }
+#else
+	const TCHAR* GetDebugSimName() const { return TEXT(""); }
+#endif
+	
+private:
+	UPROPERTY()
+	TArray<TObjectPtr<UNiagaraRendererProperties>> RendererProperties;
+
+	UPROPERTY(meta = (NiagaraNoMerge))
+	TArray<TObjectPtr<UNiagaraSimulationStageBase>> SimulationStages;
+
+	UPROPERTY()
+	TObjectPtr<UNiagaraScript> GPUComputeScript = nullptr;
+	
+	UPROPERTY()
+	TArray<FName> SharedEventGeneratorIds;
+
+	UPROPERTY(Transient)
+	FNiagaraEmitterScalabilitySettings CurrentScalabilitySettings;
+	
+	/** Indicates that the GPU script requires the view uniform buffer. */
+	uint32 bRequiresViewUniformBuffer : 1;
+
+	/** Maximum number of instances we can create for this emitter. */
+	uint32 MaxInstanceCount = 0;
+
+	/** Maximum instance allocations size for the emitter, can be larger than MaxInstanceCount */
+	uint32 MaxAllocationCount = 0;
+
+	/** Optional list of bounds calculators. */
+	TArray<TSharedPtr<FNiagaraBoundsCalculator>, TInlineAllocator<1>> BoundsCalculators;
+	
+	FMemoryRuntimeEstimation RuntimeEstimation;
+#if STATS
+	FNiagaraStatDatabase StatDatabase;
+#endif
+	
+#if WITH_NIAGARA_DEBUG_EMITTER_NAME
+	FString DebugSimName;
+#endif
+
+#if WITH_EDITORONLY_DATA
+	/** Data used by the editor to maintain UI state etc.. */
+	UPROPERTY()
+	TObjectPtr<UNiagaraEditorDataBase> EditorData;
+
+	/** Wrapper for editor only parameters. */
+	UPROPERTY()
+	TObjectPtr<UNiagaraEditorParametersAdapterBase> EditorParameters;
+
+	mutable TSharedPtr<FNiagaraGraphCachedDataBase, ESPMode::ThreadSafe> CachedTraversalData;
+#endif
+
+	friend UNiagaraEmitter;
+
+	void EnsureScriptsPostLoaded();
+	void OnPostCompile(const UNiagaraEmitter& InEmitter);
+};
+
 /** 
@@ -225,17 +506,16 @@
  *	that need to be serialized and are used for its initialization 
  */
 UCLASS(MinimalAPI)
-class UNiagaraEmitter : public UObject, public INiagaraParameterDefinitionsSubscriber
+class UNiagaraEmitter : public UObject, public INiagaraParameterDefinitionsSubscriber, public FNiagaraVersionedObject
 {
 	GENERATED_UCLASS_BODY()
 
 	friend struct FNiagaraEmitterHandle;
-
 public:
 #if WITH_EDITOR
 	DECLARE_MULTICAST_DELEGATE(FOnPropertiesChanged);
 	DECLARE_MULTICAST_DELEGATE(FOnRenderersChanged);
-	DECLARE_MULTICAST_DELEGATE_OneParam(FOnEmitterCompiled, UNiagaraEmitter*);
+	DECLARE_MULTICAST_DELEGATE_OneParam(FOnEmitterCompiled, FVersionedNiagaraEmitter);
 
 	struct NIAGARA_API PrivateMemberNames
 	{
@@ -246,7 +526,7 @@
 public:
 #if WITH_EDITOR
 	/** Creates a new emitter with the supplied emitter as a parent emitter and the supplied system as its owner. */
-	NIAGARA_API static UNiagaraEmitter* CreateWithParentAndOwner(UNiagaraEmitter& InParentEmitter, UObject* InOwner, FName InName, EObjectFlags FlagMask);
+	NIAGARA_API static UNiagaraEmitter* CreateWithParentAndOwner(FVersionedNiagaraEmitter InParentEmitter, UObject* InOwner, FName InName, EObjectFlags FlagMask);
 
 	/** Creates a new emitter by duplicating an existing emitter. The new emitter will reference the same parent emitter if one is available. */
 	static UNiagaraEmitter* CreateAsDuplicate(const UNiagaraEmitter& InEmitterToDuplicate, FName InDuplicateName, UNiagaraSystem& InDuplicateOwnerSystem);
@@ -254,25 +534,24 @@
 	//Begin UObject Interface
 	virtual void PostRename(UObject* OldOuter, const FName OldName) override;
 	virtual void PostDuplicate(EDuplicateMode::Type DuplicateMode) override;
-	virtual void PostEditChangeProperty(struct FPropertyChangedEvent& PropertyChangedEvent) override;
+	virtual void PostEditChangeProperty(FPropertyChangedEvent& PropertyChangedEvent) override;
 	PRAGMA_DISABLE_DEPRECATION_WARNINGS // Suppress compiler warning on override of deprecated function
 	UE_DEPRECATED(5.0, "Use version that takes FObjectPreSaveContext instead.")
 	PRAGMA_ENABLE_DEPRECATION_WARNINGS
 	virtual void PreSave(FObjectPreSaveContext ObjectSaveContext) override;
+	virtual void PostEditChangeVersionedProperty(FPropertyChangedEvent& PropertyChangedEvent, const FGuid& Version);
 	NIAGARA_API FOnPropertiesChanged& OnPropertiesChanged();
 	NIAGARA_API FOnRenderersChanged& OnRenderersChanged();
 	/** Helper method for when a rename has been detected within the graph. Covers renaming the internal renderer bindings.*/
-	NIAGARA_API void HandleVariableRenamed(const FNiagaraVariable& InOldVariable, const FNiagaraVariable& InNewVariable, bool bUpdateContexts);
+	NIAGARA_API void HandleVariableRenamed(const FNiagaraVariable& InOldVariable, const FNiagaraVariable& InNewVariable, bool bUpdateContexts, FGuid EmitterVersion);
 	/** Helper method for when a rename has been detected within the graph. Covers resetting the internal renderer bindings.*/
-	NIAGARA_API void HandleVariableRemoved(const FNiagaraVariable& InOldVariable, bool bUpdateContexts);
+	NIAGARA_API void HandleVariableRemoved(const FNiagaraVariable& InOldVariable, bool bUpdateContexts, FGuid EmitterVersion);
 
 	/** Helper method for binding the notifications needed for proper editor integration. */
-	NIAGARA_API void BindNotifications();
-
-	NIAGARA_API void GatherStaticVariables(TArray<FNiagaraVariable>& OutVars) const;
+	NIAGARA_API void RebindNotifications();
 #endif
 	virtual bool NeedsLoadForTargetPlatform(const ITargetPlatform* TargetPlatform) const override;
-	void Serialize(FArchive& Ar)override;
+	virtual void Serialize(FArchive& Ar)override;
 	virtual void PostInitProperties() override;
 	virtual void PostLoad() override;
 	virtual bool IsEditorOnly() const override;
@@ -281,7 +560,7 @@
 
 #if WITH_EDITORONLY_DATA
 	//~ Begin INiagaraParameterDefinitionsSubscriber Interface
-	virtual const TArray<FParameterDefinitionsSubscription>& GetParameterDefinitionsSubscriptions() const override { return ParameterDefinitionsSubscriptions; };
+	virtual const TArray<FParameterDefinitionsSubscription>& GetParameterDefinitionsSubscriptions() const override { return ParameterDefinitionsSubscriptions; }
 	virtual TArray<FParameterDefinitionsSubscription>& GetParameterDefinitionsSubscriptions() override { return ParameterDefinitionsSubscriptions; };
 
 	/** Get all UNiagaraScriptSourceBase of this subscriber. */
@@ -295,69 +574,155 @@
 	//~ End INiagaraParameterDefinitionsSubscriber Interface
 
 	/** Get the cached parameter map traversal for this emitter.  */
-	NIAGARA_API const TSharedPtr<FNiagaraGraphCachedDataBase, ESPMode::ThreadSafe>& GetCachedTraversalData() const;
+	NIAGARA_API const TSharedPtr<FNiagaraGraphCachedDataBase, ESPMode::ThreadSafe>& GetCachedTraversalData(const FGuid& EmitterVersion) const;
 #endif
 
-	bool IsEnabledOnPlatform(const FString& PlatformName)const;
+	bool IsEnabledOnPlatform(const FString& PlatformName) const;
 
-	/** Toggles whether or not the particles within this emitter are relative to the emitter origin or in global space. */ 
-	UPROPERTY(EditAnywhere, Category = "Emitter")
-	bool bLocalSpace;
+	/** Returns the emitter data for latest exposed version. */
+	NIAGARA_API FVersionedNiagaraEmitterData* GetLatestEmitterData();
+	NIAGARA_API const FVersionedNiagaraEmitterData* GetLatestEmitterData() const;
 
-	/** Toggles whether to globally make the random number generator be deterministic or non-deterministic. Any random calculation that is set to the emitter defaults will inherit this value. It is still possible to tweak individual random to be deterministic or not. In this case deterministic means that it will return the same results for the same configuration of the emitter as long as delta time is not variable. Any changes to the emitter's individual scripts will adjust the results. */
-	UPROPERTY(EditAnywhere, Category = "Emitter")
-	bool bDeterminism;
+	/** Returns the emitter data for a specific version or nullptr if no such version is found. For the null-Guid it returns the exposed version.  */
+	NIAGARA_API FVersionedNiagaraEmitterData* GetEmitterData(const FGuid& VersionGuid);
+	NIAGARA_API const FVersionedNiagaraEmitterData* GetEmitterData(const FGuid& VersionGuid) const;
 
-	/** An emitter-based seed for the deterministic random number generator. */
-	UPROPERTY(EditAnywhere, Category = "Emitter", meta = (EditCondition = "bDeterminism"))
-	int32 RandomSeed;
+	/** Returns all available versions for this emitter. */
+	NIAGARA_API virtual TArray<FNiagaraAssetVersion> GetAllAvailableVersions() const override;
+	
+#if WITH_EDITORONLY_DATA
+	NIAGARA_API virtual TSharedPtr<FNiagaraVersionDataAccessor> GetVersionDataAccessor(const FGuid& Version) override; 
 
-	/**
-	The emitter needs to allocate memory for the particles each tick.
-	To prevent reallocations, the emitter should allocate as much memory as is needed for the max particle count.
-	This setting controls if the allocation size should be automatically determined or manually entered.
-	*/
-	UPROPERTY(EditAnywhere, AdvancedDisplay, Category = "Emitter")
-	EParticleAllocationMode AllocationMode;
+	/** If true then this script asset uses active version control to track changes. */
+	NIAGARA_API virtual bool IsVersioningEnabled() const override { return bVersioningEnabled; }
+
+	/** Returns the version of the exposed version data (i.e. the version used when adding an emitter to a system) */
+	NIAGARA_API virtual FNiagaraAssetVersion GetExposedVersion() const override;
+
+	/** Returns the version data for the given guid, if it exists. Otherwise returns nullptr. */
+	NIAGARA_API virtual FNiagaraAssetVersion const* FindVersionData(const FGuid& VersionGuid) const override;
 	
-	/** 
-	The emitter will allocate at least this many particles on it's first tick.
-	This can aid performance by avoiding many allocations as an emitter ramps up to it's max size.
-	*/
-	UPROPERTY(EditAnywhere, AdvancedDisplay, Category = "Emitter", meta = (EditCondition = "AllocationMode == EParticleAllocationMode::ManualEstimate || AllocationMode == EParticleAllocationMode::FixedCount"))
-	int32 PreAllocationCount;
+	/** Creates a new data entry for the given version number. The version must be > 1.0 and must not collide with an already existing version. The data will be a copy of the previous minor version. */
+	NIAGARA_API virtual FGuid AddNewVersion(int32 MajorVersion, int32 MinorVersion) override;
 
+	/** Deletes the version data for an existing version. The exposed version cannot be deleted and will result in an error. Does nothing if the guid does not exist in the version data. */
+	NIAGARA_API virtual void DeleteVersion(const FGuid& VersionGuid) override;
+
+	/** Changes the exposed version. Does nothing if the guid does not exist in the version data. */
+	NIAGARA_API virtual void ExposeVersion(const FGuid& VersionGuid) override;
+
+	/** Enables versioning for this emitter asset. */
+	NIAGARA_API virtual void EnableVersioning() override;
+
+	/** Disables versioning and keeps only the data from the given version guid. Note that this breaks ALL references from existing assets and should only be used when creating a copy of an emitter, as the effect is very destructive.  */
+	NIAGARA_API virtual void DisableVersioning(const FGuid& VersionGuidToUse) override;
+
+	/** Makes sure that the default version data is available and fixes old emitter assets. */
+	NIAGARA_API void CheckVersionDataAvailable();
+#endif
+
+private:
+	/** The exposed version is the version that is used by default when a user adds this emitter somewhere. It is basically the published version and allows a user to create and test newer versions. */
 	UPROPERTY()
-	FNiagaraEmitterScriptProperties UpdateScriptProps;
+	FGuid ExposedVersion;
 
+	/** If true then this emitter asset uses active version control to track changes. */
 	UPROPERTY()
-	FNiagaraEmitterScriptProperties SpawnScriptProps;
+	bool bVersioningEnabled = false;
 
+	/** Contains all of the versioned emitter data. */
+	UPROPERTY()
+	TArray<FVersionedNiagaraEmitterData> VersionData;
+
+public:
+	
 #if WITH_EDITORONLY_DATA
+	UPROPERTY(EditAnywhere, AdvancedDisplay, Category = "Asset Options", AssetRegistrySearchable)
+	ENiagaraScriptTemplateSpecification TemplateSpecification;
+	
+	UPROPERTY(EditAnywhere, AdvancedDisplay, Category = "Asset Options", AssetRegistrySearchable)
+	FText TemplateAssetDescription;
+
+	/** Category to collate this emitter into for "add new emitter" dialogs.*/
+	UPROPERTY(AssetRegistrySearchable, EditAnywhere, Category = Script)
+	FText Category;
+
+	/** Internal: The thumbnail image.*/
 	UPROPERTY()
-	FNiagaraEmitterScriptProperties EmitterSpawnScriptProps;
+	TObjectPtr<UTexture2D> ThumbnailImage;
 
+	/** Internal: Indicates the thumbnail image is out of date.*/
 	UPROPERTY()
-	FNiagaraEmitterScriptProperties EmitterUpdateScriptProps;
+	uint32 ThumbnailImageOutOfDate : 1;
+	
+	/** If this emitter is exposed to the library, or should be explicitly hidden. */
+    UPROPERTY(EditAnywhere, AdvancedDisplay, Category = "Asset Options", AssetRegistrySearchable)
+    ENiagaraScriptLibraryVisibility LibraryVisibility = ENiagaraScriptLibraryVisibility::Unexposed;
 
-	/** An allow list of Particle attributes (e.g. "Particle.Position" or "Particle.Age") that will not be removed from the DataSet  even if they aren't read by the VM.
-	    Used in conjunction with UNiagaraSystem::bTrimAttributes */
-	UPROPERTY(EditAnywhere, AdvancedDisplay, Category = "Emitter")
-	TArray<FString> AttributesToPreserve;
+	/** This is used as a transient value to open a specific version in the editor */
+	UPROPERTY(Transient)
+	FGuid VersionToOpenInEditor;
 
 	const static FGuid EmitterMergeMessageId;
 #endif
 
-	UPROPERTY(EditAnywhere, Category = "Emitter")
-	ENiagaraSimTarget SimTarget;
+	// ----------------------------------------------------------------------------------------------------------------------------------------------------------------------------------------------------------------------
+	// ------------------ Most properties below this point are deprecated and stored in the versioned emitter data instead!  ------------------------------------------------------------------------------------------------
+	// ----------------------------------------------------------------------------------------------------------------------------------------------------------------------------------------------------------------------
 	
-	/** How should we calculate bounds for the emitter. */
-	UPROPERTY(EditAnywhere, Category = "Emitter")
-	ENiagaraEmitterCalculateBoundMode CalculateBoundsMode = ENiagaraEmitterCalculateBoundMode::Dynamic;
+#if WITH_EDITORONLY_DATA
+	/** Use property in struct returned from GetEmitterData() instead */ 
+	UPROPERTY(meta = (DeprecatedProperty))
+	bool bLocalSpace_DEPRECATED;
 
-	/** The fixed bounding box value. bFixedBounds is the condition whether the fixed bounds can be edited. */
-	UPROPERTY(EditAnywhere, Category = "Emitter", meta = (EditConditionHides, EditCondition = "CalculateBoundsMode == ENiagaraEmitterCalculateBoundMode::Fixed"))
-	FBox FixedBounds;
+	/** Use property in struct returned from GetEmitterData() instead */ 
+	UPROPERTY(meta = (DeprecatedProperty))
+	bool bDeterminism_DEPRECATED;
+
+	/** Use property in struct returned from GetEmitterData() instead */ 
+	UPROPERTY(meta = (DeprecatedProperty))
+	int32 RandomSeed_DEPRECATED;
+
+	/** Use property in struct returned from GetEmitterData() instead */ 
+	UPROPERTY(meta = (DeprecatedProperty))
+	EParticleAllocationMode AllocationMode_DEPRECATED;
+	
+	/** Use property in struct returned from GetEmitterData() instead */ 
+	UPROPERTY(meta = (DeprecatedProperty))
+	int32 PreAllocationCount_DEPRECATED;
+
+	/** Use property in struct returned from GetEmitterData() instead */ 
+	UPROPERTY(meta = (DeprecatedProperty))
+	FNiagaraEmitterScriptProperties UpdateScriptProps_DEPRECATED;
+
+	/** Use property in struct returned from GetEmitterData() instead */ 
+	UPROPERTY(meta = (DeprecatedProperty))
+	FNiagaraEmitterScriptProperties SpawnScriptProps_DEPRECATED;
+
+
+	/** Use property in struct returned from GetEmitterData() instead */ 
+	UPROPERTY(meta = (DeprecatedProperty))
+	FNiagaraEmitterScriptProperties EmitterSpawnScriptProps_DEPRECATED;
+
+	/** Use property in struct returned from GetEmitterData() instead */ 
+	UPROPERTY(meta = (DeprecatedProperty))
+	FNiagaraEmitterScriptProperties EmitterUpdateScriptProps_DEPRECATED;
+
+	/** Use property in struct returned from GetEmitterData() instead */ 
+	UPROPERTY(meta = (DeprecatedProperty))
+	TArray<FString> AttributesToPreserve_DEPRECATED;
+
+	/** Use property in struct returned from GetEmitterData() instead */ 
+	UPROPERTY(meta = (DeprecatedProperty))
+	TArray<TObjectPtr<UNiagaraScript>> ParentScratchPadScripts_DEPRECATED;
+
+	/** Use property in struct returned from GetEmitterData() instead */ 
+	UPROPERTY(meta = (DeprecatedProperty))
+	ENiagaraSimTarget SimTarget_DEPRECATED;
+
+	/** Use property in struct returned from GetEmitterData() instead */ 
+	UPROPERTY(meta = (DeprecatedProperty, ScriptNoExport))
+	FBox FixedBounds_DEPRECATED;
 	
 	UPROPERTY()
 	int32 MinDetailLevel_DEPRECATED;
@@ -366,20 +731,24 @@
 	UPROPERTY()
 	FNiagaraDetailsLevelScaleOverrides GlobalSpawnCountScaleOverrides_DEPRECATED;
 	
-	UPROPERTY(EditAnywhere, Category = "Scalability", meta=(DisplayInScalabilityContext))
-	FNiagaraPlatformSet Platforms;
+	/** Use property in struct returned from GetEmitterData() instead */ 
+	UPROPERTY(meta = (DeprecatedProperty))
+	FNiagaraPlatformSet Platforms_DEPRECATED;
 
-	UPROPERTY(EditAnywhere, Category = "Scalability", meta=(DisplayInScalabilityContext))
-	FNiagaraEmitterScalabilityOverrides ScalabilityOverrides;
+	/** Use property in struct returned from GetEmitterData() instead */ 
+	UPROPERTY(meta = (DeprecatedProperty))
+	FNiagaraEmitterScalabilityOverrides ScalabilityOverrides_DEPRECATED;
 
-	/** When enabled, this will spawn using interpolated parameter values and perform a partial update at spawn time. This adds significant additional cost for spawning but will produce much smoother spawning for high spawn rates, erratic frame rates and fast moving emitters. */
-	UPROPERTY(EditAnywhere, Category = "Emitter")
-	uint32 bInterpolatedSpawning : 1;
+	/** Use property in struct returned from GetEmitterData() instead */ 
+	UPROPERTY(meta = (DeprecatedProperty))
+	uint32 bInterpolatedSpawning_DEPRECATED : 1;
 
-#if WITH_EDITORONLY_DATA
-	UPROPERTY()
+	UPROPERTY(meta = (DeprecatedProperty))
+	FNiagaraParameterStore RendererBindings_DEPRECATED;
+
+	/** Use property in struct returned from GetEmitterData() instead */ 
+	UPROPERTY(meta = (DeprecatedProperty, ScriptNoExport))
 	uint32 bFixedBounds_DEPRECATED : 1;
-#endif
 
 	/** Whether to use the min detail or not. */
 	UPROPERTY()
@@ -393,94 +762,51 @@
 	UPROPERTY()
 	uint32 bOverrideGlobalSpawnCountScale_DEPRECATED : 1;
 
-	/** Do particles in this emitter require a persistent ID? */
-	UPROPERTY(EditAnywhere, Category = "Emitter")
-	uint32 bRequiresPersistentIDs : 1;
+	/** Use property in struct returned from GetEmitterData() instead */ 
+	UPROPERTY(meta = (DeprecatedProperty))
+	uint32 bRequiresPersistentIDs_DEPRECATED : 1;
 
-	/** Performance option to allow event based spawning to be combined into a single spawn.  This will result in a single exec from 0 to number of particles rather than several, when using ExecIndex() it is recommended not to do this. */
-	UPROPERTY(EditAnywhere, Category = "Emitter")
-	uint32 bCombineEventSpawn : 1;
+	/** Use property in struct returned from GetEmitterData() instead */ 
+	UPROPERTY(meta = (DeprecatedProperty))
+	uint32 bCombineEventSpawn_DEPRECATED : 1;
 
-	/** Limits the delta time per tick to prevent simulation spikes due to frame lags. */
-	UPROPERTY(EditAnywhere, AdvancedDisplay, Category = "Emitter", meta = (EditCondition = "bLimitDeltaTime", ForceUnits=s))
-	float MaxDeltaTimePerTick;
+	/** Use property in struct returned from GetEmitterData() instead */ 
+	UPROPERTY(meta = (DeprecatedProperty))
+	float MaxDeltaTimePerTick_DEPRECATED;
 
+	/** Use property in struct returned from GetEmitterData() instead */ 
+	UPROPERTY(meta = (DeprecatedProperty))
+	uint32 bLimitDeltaTime_DEPRECATED : 1;
 
-	/** Whether to limit the max tick delta time or not. */
-	UPROPERTY(EditAnywhere, AdvancedDisplay, Category = "Emitter", meta = (InlineEditConditionToggle))
-	uint32 bLimitDeltaTime : 1;
-
-	/** An override on the max number of GPU particles we expect to spawn in a single frame. A value of 0 means it'll use fx.MaxNiagaraGPUParticlesSpawnPerFrame.*/
-	UPROPERTY(EditAnywhere, AdvancedDisplay, Category = "Emitter", meta = (EditCondition = "SimTarget == ENiagaraSimTarget::GPUComputeSim", DisplayName = "Max GPU Particles Spawn per Frame"))
-	int32 MaxGPUParticlesSpawnPerFrame;
-
-	NIAGARA_API void GetScripts(TArray<UNiagaraScript*>& OutScripts, bool bCompilableOnly = true, bool bEnabledOnly = false) const;
-
-	NIAGARA_API UNiagaraScript* GetScript(ENiagaraScriptUsage Usage, FGuid UsageId);
-
-	NIAGARA_API UNiagaraScript* GetGPUComputeScript() { return GPUComputeScript; }
-	NIAGARA_API const UNiagaraScript* GetGPUComputeScript() const { return GPUComputeScript; }
-
-	void CacheFromCompiledData(const FNiagaraDataSetCompiledData* CompiledData);
-	void CacheFromShaderCompiled();
-
+	/** Use property in struct returned from GetEmitterData() instead */ 
+	UPROPERTY(meta = (DeprecatedProperty))
+	uint32 MaxGPUParticlesSpawnPerFrame_DEPRECATED;
+#endif
+	
 	NIAGARA_API void UpdateEmitterAfterLoad();
 
 #if WITH_EDITORONLY_DATA
-	/** 'Source' data/graphs for the scripts used by this emitter. */
-	UPROPERTY()
-	TObjectPtr<class UNiagaraScriptSourceBase>	GraphSource;
+	/** Use property in struct returned from GetEmitterData() instead */ 
+	UPROPERTY(meta = (DeprecatedProperty))
+	TObjectPtr<UNiagaraScriptSourceBase> GraphSource_DEPRECATED;
 
-	NIAGARA_API bool AreAllScriptAndSourcesSynchronized() const;
-	NIAGARA_API void OnPostCompile();
+	NIAGARA_API void OnPostCompile(const FGuid& EmitterVersion);
 
-	NIAGARA_API void InvalidateCompileResults();
-
 	/* Gets a Guid which is updated any time data in this emitter is changed. */
 	NIAGARA_API FGuid GetChangeId() const;
 
-	NIAGARA_API UNiagaraEditorDataBase* GetEditorData() const;
-	NIAGARA_API UNiagaraEditorParametersAdapterBase* GetEditorParameters();
-
-	NIAGARA_API void SetEditorData(UNiagaraEditorDataBase* InEditorData);
-
-	/** Internal: The thumbnail image.*/
-	UPROPERTY()
-	TObjectPtr<class UTexture2D> ThumbnailImage;
-
-	/** Internal: Indicates the thumbnail image is out of date.*/
-	UPROPERTY()
-	uint32 ThumbnailImageOutOfDate : 1;
-
 	/** Deprecated library exposure bool. Use the LibraryVisibility enum instead. FNiagaraEditorUtilities has accessor functions that takes deprecation into account */
 	UPROPERTY()
 	bool bExposeToLibrary_DEPRECATED;
 	
-	/** If this emitter is exposed to the library, or should be explicitly hidden. */
-	UPROPERTY(EditAnywhere, AdvancedDisplay, Category = "Asset Options", AssetRegistrySearchable)
-	ENiagaraScriptLibraryVisibility LibraryVisibility = ENiagaraScriptLibraryVisibility::Unexposed;
-
 	/** Deprecated template asset bool. Use the TemplateSpecification enum instead. */
 	UPROPERTY()
 	bool bIsTemplateAsset_DEPRECATED;
 
-	/** If this emitter is a standard parent emitter, a template or a behavior example. Templates and behavior examples get copied rather than inherited from. */
-	UPROPERTY(EditAnywhere, AdvancedDisplay, Category = "Asset Options", AssetRegistrySearchable)
-	ENiagaraScriptTemplateSpecification TemplateSpecification = ENiagaraScriptTemplateSpecification::None;
-	
-	UPROPERTY(EditAnywhere, AdvancedDisplay, Category = "Asset Options", AssetRegistrySearchable)
-	FText TemplateAssetDescription;
+	/** Use property in struct returned from GetEmitterData() instead */ 
+	UPROPERTY(meta = (DeprecatedProperty))
+	TArray<TObjectPtr<UNiagaraScript>> ScratchPadScripts_DEPRECATED;
 
-	/** Category to collate this emitter into for "add new emitter" dialogs.*/
-	UPROPERTY(EditAnywhere, AdvancedDisplay, Category = "Asset Options", AssetRegistrySearchable)
-	FText Category;
-
-	UPROPERTY()
-	TArray<TObjectPtr<UNiagaraScript>> ScratchPadScripts;
-
-	UPROPERTY()
-	TArray<TObjectPtr<UNiagaraScript>> ParentScratchPadScripts;
-	
 	/** Callback issued whenever a VM compilation successfully happened (even if the results are a script that cannot be executed due to errors)*/
 	NIAGARA_API FOnEmitterCompiled& OnEmitterVMCompiled();
 
@@ -499,168 +825,70 @@
 	NIAGARA_API bool IsSynchronizedWithParent() const;
 
 	/** Merges in any changes from the parent emitter into this emitter. */
-	NIAGARA_API INiagaraMergeManager::FMergeEmitterResults MergeChangesFromParent();
+	NIAGARA_API TArray<INiagaraMergeManager::FMergeEmitterResults> MergeChangesFromParent();	
 
-	/** Whether or not this emitter uses the supplied emitter */
-	bool UsesEmitter(const UNiagaraEmitter& InEmitter) const;
-
 	/** Duplicates this emitter, but prevents the duplicate from merging in changes from the parent emitter.  The resulting duplicate will have no parent information. */
 	NIAGARA_API UNiagaraEmitter* DuplicateWithoutMerging(UObject* InOuter);
+
+	NIAGARA_API void SetEditorData(UNiagaraEditorDataBase* InEditorData, const FGuid& VersionGuid);
 #endif
 
-	FORCEINLINE const FNiagaraEmitterScalabilitySettings& GetScalabilitySettings()const { return CurrentScalabilitySettings; }
-	NIAGARA_API const FNiagaraEmitterScalabilityOverride& GetCurrentOverrideSettings() const;
-
-	/** Returns true if this emitter's platform filter allows it on this platform and quality level. */
-	NIAGARA_API bool IsAllowedByScalability()const;
-
-	NIAGARA_API bool RequiresPersistentIDs() const;
-
-	NIAGARA_API bool IsValid()const;
-	NIAGARA_API bool IsReadyToRun() const;
-	bool UsesScript(const UNiagaraScript* Script)const;
-	//bool UsesDataInterface(UNiagaraDataInterface* Interface);
-	bool UsesCollection(const class UNiagaraParameterCollection* Collection)const;
-	bool CanObtainParticleAttribute(const FNiagaraVariableBase& InVar, FNiagaraTypeDefinition& OutBoundType) const;
+	bool CanObtainParticleAttribute(const FNiagaraVariableBase& InVar, const FGuid& EmitterVersion, FNiagaraTypeDefinition& OutBoundType) const;
 	bool CanObtainEmitterAttribute(const FNiagaraVariableBase& InVarWithUniqueNameNamespace, FNiagaraTypeDefinition& OutBoundType) const;
 	bool CanObtainSystemAttribute(const FNiagaraVariableBase& InVar, FNiagaraTypeDefinition& OutBoundType) const;
 	bool CanObtainUserVariable(const FNiagaraVariableBase& InVar) const;
 
-#if WITH_NIAGARA_DEBUG_EMITTER_NAME
-	const TCHAR* GetDebugSimName() const { return *DebugSimName; }
-#else
-	const TCHAR* GetDebugSimName() const { return TEXT(""); }
-#endif
-
 	NIAGARA_API const FString& GetUniqueEmitterName() const;
 	NIAGARA_API bool SetUniqueEmitterName(const FString& InName);
 
-	const TArray<UNiagaraRendererProperties*>& GetRenderers() const { return RendererProperties; }
+	void NIAGARA_API AddRenderer(UNiagaraRendererProperties* Renderer, FGuid EmitterVersion);
+	void NIAGARA_API RemoveRenderer(UNiagaraRendererProperties* Renderer, FGuid EmitterVersion);
+	void NIAGARA_API AddEventHandler(FNiagaraEventScriptProperties EventHandler, FGuid EmitterVersion);
+	void NIAGARA_API RemoveEventHandlerByUsageId(FGuid EventHandlerUsageId, FGuid EmitterVersion);
+	void NIAGARA_API AddSimulationStage(UNiagaraSimulationStageBase* SimulationStage, FGuid EmitterVersion);
+	void NIAGARA_API RemoveSimulationStage(UNiagaraSimulationStageBase* SimulationStage, FGuid EmitterVersion);
+	void NIAGARA_API MoveSimulationStageToIndex(UNiagaraSimulationStageBase* SimulationStage, int32 TargetIndex, FGuid EmitterVersion);
 
-	template<typename TAction>
-	void ForEachRenderer(TAction Func) const;
-	
-	template<typename TAction>
-	void ForEachEnabledRenderer(TAction Func) const;
-
-	template<typename TAction>
-	void ForEachScript(TAction Func) const;
-
-	template<typename TAction>
-	void ForEachPlatformSet(TAction Func);
-
-	NIAGARA_API void AddRenderer(UNiagaraRendererProperties* Renderer);
-
-	NIAGARA_API void RemoveRenderer(UNiagaraRendererProperties* Renderer);
-
-	FORCEINLINE const TArray<FNiagaraEventScriptProperties>& GetEventHandlers() const { return EventHandlerScriptProps; }
-
-	/* Gets a pointer to an event handler by script usage id.  This method is potentially unsafe because modifications to
-	   the event handler array can make this pointer become invalid without warning. */
-	NIAGARA_API FNiagaraEventScriptProperties* GetEventHandlerByIdUnsafe(FGuid ScriptUsageId);
-
-	NIAGARA_API void AddEventHandler(FNiagaraEventScriptProperties EventHandler);
-
-	NIAGARA_API void RemoveEventHandlerByUsageId(FGuid EventHandlerUsageId);
-
-	NIAGARA_API const TArray<UNiagaraSimulationStageBase*>& GetSimulationStages() const { return SimulationStages; }
-
-	NIAGARA_API UNiagaraSimulationStageBase* GetSimulationStageById(FGuid ScriptUsageId) const;
-
-	NIAGARA_API void AddSimulationStage(UNiagaraSimulationStageBase* SimulationStage);
-
-	NIAGARA_API void RemoveSimulationStage(UNiagaraSimulationStageBase* SimulationStage);
-
-	NIAGARA_API void MoveSimulationStageToIndex(UNiagaraSimulationStageBase* SimulationStage, int32 TargetIndex);
-
-	/* Gets whether or not the supplied event generator id matches an event generator which is shared between the particle spawn and update scrips. */
-	bool IsEventGeneratorShared(FName EventGeneratorId) const;
-
 	TStatId GetStatID(bool bGameThread, bool bConcurrent) const;
 
-	void ClearRuntimeAllocationEstimate(uint64 ReportHandle = INDEX_NONE);
-	/* This is used by the emitter instances to report runtime allocations to reduce reallocation in future simulation runs. */
-	int32 AddRuntimeAllocation(uint64 ReporterHandle, int32 AllocationCount);
-#if STATS
-	NIAGARA_API FNiagaraStatDatabase& GetStatData() { return StatDatabase; }
-#endif
-
-	/* Returns the number of max expected particles for memory allocations. */
-	NIAGARA_API int32 GetMaxParticleCountEstimate();
-	NIAGARA_API uint32 GetMaxInstanceCount() { return MaxInstanceCount; }
-
-	/* Returns the maximum size we can allocate and hold for instances */
-	uint32 GetMaxAllocationCount() const { return MaxAllocationCount; }
-
-#if WITH_EDITORONLY_DATA
-	NIAGARA_API UNiagaraEmitter* GetParent() const;
-
-	NIAGARA_API UNiagaraEmitter* GetParentAtLastMerge() const;
-
-	NIAGARA_API void RemoveParent();
-
-	NIAGARA_API void SetParent(UNiagaraEmitter& InParent);
-
-	NIAGARA_API	void Reparent(UNiagaraEmitter& InParent);
-
-	NIAGARA_API void NotifyScratchPadScriptsChanged();
-#endif
-
 	void UpdateScalability();
 
 #if WITH_EDITORONLY_DATA
+	NIAGARA_API void SetParent(const FVersionedNiagaraEmitter& InParent);
+	NIAGARA_API void ChangeParentVersion(const FGuid& NewParentVersion, const FGuid& EmitterVersion);
+
+	NIAGARA_API void NotifyScratchPadScriptsChanged();
 	NIAGARA_API const TMap<FGuid, TObjectPtr<UNiagaraMessageDataBase>>& GetMessages() const { return MessageKeyToMessageMap; };
 	NIAGARA_API void AddMessage(const FGuid& MessageKey, UNiagaraMessageDataBase* NewMessage) { MessageKeyToMessageMap.Add(MessageKey, NewMessage); };
 	NIAGARA_API void RemoveMessage(const FGuid& MessageKey) { MessageKeyToMessageMap.Remove(MessageKey); };
 	void RemoveMessageDelegateable(const FGuid MessageKey) { MessageKeyToMessageMap.Remove(MessageKey); };
 #endif
 
-	bool RequiresViewUniformBuffer() const { return bRequiresViewUniformBuffer; }
-
-	FORCEINLINE static FBox GetDefaultFixedBounds() { return FBox(FVector(-100), FVector(100)); }
-	TConstArrayView<TUniquePtr<FNiagaraBoundsCalculator>> GetBoundsCalculators() const { return MakeArrayView(BoundsCalculators); }
-
-	UPROPERTY()
-	FNiagaraParameterStore RendererBindings;
-
-	void RebuildRendererBindings();
-
 protected:
 	virtual void BeginDestroy() override;
 
-	void ResolveScalabilitySettings();
-
 #if WITH_EDITORONLY_DATA
 private:
-	void UpdateFromMergedCopy(const INiagaraMergeManager& MergeManager, UNiagaraEmitter* MergedEmitter);
-
-	void SyncEmitterAlias(const FString& InOldName, const FString& InNewName);
-
+	void UpdateFromMergedCopy(const INiagaraMergeManager& MergeManager, UNiagaraEmitter* MergedEmitter, FVersionedNiagaraEmitterData* EmitterData);
 	void UpdateChangeId(const FString& Reason);
-
 	void ScriptRapidIterationParameterChanged();
-
 	void SimulationStageChanged();
-
 	void RendererChanged();
-
 	void GraphSourceChanged();
-
 	void PersistentEditorDataChanged();
 
-	mutable TSharedPtr<FNiagaraGraphCachedDataBase, ESPMode::ThreadSafe> CachedTraversalData;
 private:
 	/** Adjusted every time that we compile this emitter. Lets us know that we might differ from any cached versions.*/
 	UPROPERTY()
 	FGuid ChangeId;
 
-	/** Data used by the editor to maintain UI state etc.. */
-	UPROPERTY()
-	TObjectPtr<UNiagaraEditorDataBase> EditorData;
+	/** Use property in struct returned from GetEmitterData() instead */ 
+	UPROPERTY(meta = (DeprecatedProperty))
+	TObjectPtr<UNiagaraEditorDataBase> EditorData_DEPRECATED;
 
-	/** Wrapper for editor only parameters. */
-	UPROPERTY()
-	TObjectPtr<UNiagaraEditorParametersAdapterBase> EditorParameters;
+	/** Use property in struct returned from GetEmitterData() instead */ 
+	UPROPERTY(meta = (DeprecatedProperty))
+	TObjectPtr<UNiagaraEditorParametersAdapterBase> EditorParameters_DEPRECATED;
 
 	/** A multicast delegate which is called whenever all the scripts for this emitter have been compiled (successfully or not). */
 	FOnEmitterCompiled OnVMScriptCompiledDelegate;
@@ -673,34 +901,35 @@
 
 	bool bFullyLoaded = false;
 
-#if WITH_NIAGARA_DEBUG_EMITTER_NAME
-	FString DebugSimName;
-#endif
-
 	UPROPERTY()
 	FString UniqueEmitterName;
 
-	UPROPERTY()
-	TArray<TObjectPtr<UNiagaraRendererProperties>> RendererProperties;
+	/** Use property in struct returned from GetEmitterData() instead */ 
+	UPROPERTY(meta = (DeprecatedProperty))
+	TArray<TObjectPtr<UNiagaraRendererProperties>> RendererProperties_DEPRECATED;
+	
+	/** Use property in struct returned from GetEmitterData() instead */ 
+	UPROPERTY(meta = (DeprecatedProperty))
+	TArray<FNiagaraEventScriptProperties> EventHandlerScriptProps_DEPRECATED;
 
-	UPROPERTY(EditAnywhere, Category = "Events", meta=(NiagaraNoMerge))
-	TArray<FNiagaraEventScriptProperties> EventHandlerScriptProps;
+	/** Use property in struct returned from GetEmitterData() instead */ 
+	UPROPERTY(meta = (DeprecatedProperty))
+	TArray<TObjectPtr<UNiagaraSimulationStageBase>> SimulationStages_DEPRECATED;
 
-	UPROPERTY(meta = (NiagaraNoMerge))
-	TArray<TObjectPtr<UNiagaraSimulationStageBase>> SimulationStages;
+	/** Use property in struct returned from GetEmitterData() instead */ 
+	UPROPERTY(meta = (DeprecatedProperty))
+	TObjectPtr<UNiagaraScript> GPUComputeScript_DEPRECATED;
 
-	UPROPERTY()
-	TObjectPtr<UNiagaraScript> GPUComputeScript;
+	/** Use property in struct returned from GetEmitterData() instead */ 
+	UPROPERTY(meta = (DeprecatedProperty))
+	TArray<FName> SharedEventGeneratorIds_DEPRECATED;
 
-	UPROPERTY()
-	TArray<FName> SharedEventGeneratorIds;
-
 #if WITH_EDITORONLY_DATA
-	UPROPERTY()
-	TObjectPtr<UNiagaraEmitter> Parent;
+	UPROPERTY(meta = (DeprecatedProperty))
+	TObjectPtr<UNiagaraEmitter> Parent_DEPRECATED;
 
-	UPROPERTY()
-	TObjectPtr<UNiagaraEmitter> ParentAtLastMerge;
+	UPROPERTY(meta = (DeprecatedProperty))
+	TObjectPtr<UNiagaraEmitter> ParentAtLastMerge_DEPRECATED;
 
 	/** Subscriptions to definitions of parameters. */
 	UPROPERTY()
@@ -713,48 +942,28 @@
 	FOnRenderersChanged OnRenderersChangedDelegate;
 #endif
 
-	void EnsureScriptsPostLoaded();
-
 	void GenerateStatID()const;
 #if STATS
 	mutable TStatId StatID_GT;
 	mutable TStatId StatID_GT_CNC;
 	mutable TStatId StatID_RT;
 	mutable TStatId StatID_RT_CNC;
-	FNiagaraStatDatabase StatDatabase;
 #endif
 
-	/** Indicates that the GPU script requires the view uniform buffer. */
-	uint32 bRequiresViewUniformBuffer : 1;
-
 #if WITH_EDITORONLY_DATA
 	/* Flag set on load based on whether the serialized data includes editor only data */
 	uint32 IsCooked : 1;
-#endif
-
-	/** Maximum number of instances we can create for this emitter. */
-	uint32 MaxInstanceCount = 0;
-
-	/** Maximum instance allocations size for the emitter, can be larger than MaxInstanceCount */
-	uint32 MaxAllocationCount = 0;
-
-	/** Optional list of bounds calculators. */
-	TArray<TUniquePtr<FNiagaraBoundsCalculator>, TInlineAllocator<1>> BoundsCalculators;
-
-	MemoryRuntimeEstimation RuntimeEstimation;
-	FCriticalSection EstimationCriticalSection;
-
-	FNiagaraEmitterScalabilitySettings CurrentScalabilitySettings;
-
-#if WITH_EDITORONLY_DATA
+	
 	/** Messages associated with the Emitter asset. */
 	UPROPERTY()
 	TMap<FGuid, TObjectPtr<UNiagaraMessageDataBase>> MessageKeyToMessageMap;
 #endif
+
+	void ResolveScalabilitySettings();
 };
 
-template<typename TAction>
-void UNiagaraEmitter::ForEachRenderer(TAction Func) const
+template <typename TAction>
+void FVersionedNiagaraEmitterData::ForEachRenderer(TAction Func) const
 {
 	for (UNiagaraRendererProperties* Renderer : RendererProperties)
 	{
@@ -766,7 +975,7 @@
 }
 
 template<typename TAction>
-void UNiagaraEmitter::ForEachEnabledRenderer(TAction Func) const
+void FVersionedNiagaraEmitterData::ForEachEnabledRenderer(TAction Func) const
 {
 	for (UNiagaraRendererProperties* Renderer : RendererProperties)
 	{
@@ -777,7 +986,7 @@
 	}
 }
 template<typename TAction>
-void UNiagaraEmitter::ForEachScript(TAction Func) const
+void FVersionedNiagaraEmitterData::ForEachScript(TAction Func) const
 {
 	Func(SpawnScriptProps.Script);
 	Func(UpdateScriptProps.Script);
@@ -791,13 +1000,13 @@
 }
 
 template<typename TAction>
-void UNiagaraEmitter::ForEachPlatformSet(TAction Func)
+void FVersionedNiagaraEmitterData::ForEachPlatformSet(TAction Func)
 {
-	Func(this, Platforms);
+	Func(Platforms);
 
 	for (FNiagaraEmitterScalabilityOverride& Override : ScalabilityOverrides.Overrides)
 	{
-		Func(this, Override.Platforms);
+		Func(Override.Platforms);
 	}
 
 	for (UNiagaraRendererProperties* Renderer : RendererProperties)
@@ -816,7 +1025,7 @@
 			{
 				if (UNiagaraDataInterfacePlatformSet* PlatformSetDI = Cast<UNiagaraDataInterfacePlatformSet>(DataInterfaceInfo.DataInterface))
 				{
-					Func(PlatformSetDI, PlatformSetDI->Platforms);
+					Func(PlatformSetDI->Platforms);
 				}
 			}
 		}
diff --git a/Engine/Plugins/FX/Niagara/Source/Niagara/Classes/NiagaraEmitterHandle.h b/Engine/Plugins/FX/Niagara/Source/Niagara/Classes/NiagaraEmitterHandle.h
--- a/Engine/Plugins/FX/Niagara/Source/Niagara/Classes/NiagaraEmitterHandle.h
+++ b/Engine/Plugins/FX/Niagara/Source/Niagara/Classes/NiagaraEmitterHandle.h
@@ -3,10 +3,9 @@
 #pragma once
 
 #include "CoreMinimal.h"
-#include "UObject/ObjectMacros.h"
-#include "UObject/Object.h"
+
+#include "NiagaraEmitter.h"
 #include "Misc/Guid.h"
-#include "NiagaraModule.h"
 #include "NiagaraScript.h"
 #include "NiagaraEmitterHandle.generated.h"
 
@@ -26,7 +25,10 @@ public:
 
 #if WITH_EDITORONLY_DATA
 	/** Creates a new emitter handle from an emitter. */
-	FNiagaraEmitterHandle(UNiagaraEmitter& InEmitter);
+	FNiagaraEmitterHandle(UNiagaraEmitter& InEmitter, const FGuid& Version);
+	
+	/** Creates a new emitter handle from an emitter. */
+	FNiagaraEmitterHandle(const FVersionedNiagaraEmitter& InEmitter);
 #endif
 
 	/** Whether or not this is a valid emitter handle. */
@@ -57,7 +59,10 @@ public:
 #endif
 
 	/** Gets the copied instance of the emitter this handle references. */
-	UNiagaraEmitter* GetInstance() const;
+	FVersionedNiagaraEmitter GetInstance() const;
+	FVersionedNiagaraEmitter& GetInstance();
+
+	FVersionedNiagaraEmitterData* GetEmitterData() const;
 
 	/** Gets a unique name for this emitter instance for use in scripts and parameter stores etc.*/
 	FString GetUniqueInstanceName()const;
@@ -70,6 +75,7 @@ public:
 	void ConditionalPostLoad(int32 NiagaraCustomVersion);
 
 	/** Whether or not this handle uses the supplied emitter. */
+	bool UsesEmitter(const FVersionedNiagaraEmitter& InEmitter) const;
 	bool UsesEmitter(const UNiagaraEmitter& InEmitter) const;
 
 	void ClearEmitter();
@@ -113,5 +119,9 @@ private:
 
 	/** The copied instance of the emitter this handle references. */
 	UPROPERTY()
-	TObjectPtr<UNiagaraEmitter> Instance;
+	TObjectPtr<UNiagaraEmitter> Instance_DEPRECATED;
+
+	/** The copied instance of the emitter this handle references. */
+	UPROPERTY()
+	FVersionedNiagaraEmitter VersionedInstance;
 };
diff --git a/Engine/Plugins/FX/Niagara/Source/Niagara/Classes/NiagaraEmitterInstance.h b/Engine/Plugins/FX/Niagara/Source/Niagara/Classes/NiagaraEmitterInstance.h
--- a/Engine/Plugins/FX/Niagara/Source/Niagara/Classes/NiagaraEmitterInstance.h
+++ b/Engine/Plugins/FX/Niagara/Source/Niagara/Classes/NiagaraEmitterInstance.h
@@ -111,7 +111,7 @@ public:
 	}
 
 	FORCEINLINE int32 GetTotalSpawnedParticles()const { return TotalSpawnedParticles; }
-	FORCEINLINE const FNiagaraEmitterScalabilitySettings& GetScalabilitySettings()const { return CachedEmitter->GetScalabilitySettings(); }
+	FORCEINLINE const FNiagaraEmitterScalabilitySettings& GetScalabilitySettings()const { return GetCachedEmitterData()->GetScalabilitySettings(); }
 
 	NIAGARA_API const FNiagaraEmitterHandle& GetEmitterHandle() const;
 
@@ -131,7 +131,8 @@ public:
 	TArrayView<FNiagaraScriptExecutionContext> GetEventExecutionContexts();
 
 	FORCEINLINE FName GetCachedIDName()const { return CachedIDName; }
-	FORCEINLINE UNiagaraEmitter* GetCachedEmitter()const { return CachedEmitter; }
+	FORCEINLINE FVersionedNiagaraEmitter GetCachedEmitter()const { return CachedEmitter; }
+	FORCEINLINE FVersionedNiagaraEmitterData* GetCachedEmitterData()const { return CachedEmitter.GetEmitterData(); }
 
 	TArray<FNiagaraSpawnInfo>& GetSpawnInfo() { return SpawnInfos; }
 
@@ -215,7 +216,7 @@ private:
 	 ** EmitterInstance and the cached bounds will be unset. */
 	FBox CachedSystemFixedBounds;
 
-	FNiagaraSystemInstanceID OwnerSystemInstanceID;
+	FNiagaraSystemInstanceID OwnerSystemInstanceID = 0;
 
 	FNiagaraGpuComputeDispatchInterface* ComputeDispatchInterface = nullptr;
 
@@ -225,7 +226,7 @@ private:
 	FNiagaraSystemInstance *ParentSystemInstance = nullptr;
 
 	/** Raw pointer to the emitter that we're instanced from. Raw ptr should be safe here as we check for the validity of the system and it's emitters higher up before any ticking. */
-	UNiagaraEmitter* CachedEmitter = nullptr;
+	FVersionedNiagaraEmitter CachedEmitter;
 	FName CachedIDName;
 
 	/** The index of our emitter in our parent system instance. */
diff --git a/Engine/Plugins/FX/Niagara/Source/Niagara/Classes/NiagaraGPUProfilerInterface.h b/Engine/Plugins/FX/Niagara/Source/Niagara/Classes/NiagaraGPUProfilerInterface.h
--- a/Engine/Plugins/FX/Niagara/Source/Niagara/Classes/NiagaraGPUProfilerInterface.h
+++ b/Engine/Plugins/FX/Niagara/Source/Niagara/Classes/NiagaraGPUProfilerInterface.h
@@ -3,6 +3,7 @@
 #pragma once
 
 #include "NiagaraCommon.h"
+#include "NiagaraEmitter.h"
 #include "NiagaraSystemGpuComputeProxy.h"
 
 //////////////////////////////////////////////////////////////////////////
@@ -46,7 +47,7 @@ struct FNiagaraGpuFrameResults : public TSharedFromThis<FNiagaraGpuFrameResults,
 	{
 		uint32									bUniqueInstance : 1;		// Set only once for all dispatches from an instance across all ticks
 		TWeakObjectPtr<class USceneComponent>	OwnerComponent;				// Optional pointer back to owning Component
-		TWeakObjectPtr<class UNiagaraEmitter>	OwnerEmitter;				// Optional pointer back to owning Emitter
+		FVersionedNiagaraEmitterWeakPtr			OwnerEmitter;				// Optional pointer back to owning Emitter
 		FName									StageName;					// Generally the simulation stage but may be a DataInterface name
 		uint64									DurationMicroseconds;		// Duration in microseconds of the dispatch
 	};
diff --git a/Engine/Plugins/FX/Niagara/Source/Niagara/Classes/NiagaraScratchPadContainer.h b/Engine/Plugins/FX/Niagara/Source/Niagara/Classes/NiagaraScratchPadContainer.h
new file mode 100644
--- /dev/null
+++ b/Engine/Plugins/FX/Niagara/Source/Niagara/Classes/NiagaraScratchPadContainer.h
@@ -0,0 +1,29 @@
+﻿// Copyright Epic Games, Inc. All Rights Reserved.
+
+#pragma once
+
+#include "CoreMinimal.h"
+#include "UObject/Object.h"
+#include "NiagaraScript.h"
+#include "NiagaraScratchPadContainer.generated.h"
+
+/**
+ * Wrapper class for scratch pad scripts in an emitter. This is needed because each emitter version can have it's own copy of a scratch pad with the same name and they can't all be outered to the emitter that way.
+ */
+UCLASS()
+class NIAGARA_API UNiagaraScratchPadContainer : public UObject
+{
+	GENERATED_BODY()
+	
+public:
+	virtual void PostLoad() override;
+	void CheckConsistency();
+	void SetScripts(const TArray<TObjectPtr<UNiagaraScript>>& InScripts);
+
+#if WITH_EDITORONLY_DATA
+
+	UPROPERTY()
+	TArray<TObjectPtr<UNiagaraScript>> Scripts;
+
+#endif
+};
diff --git a/Engine/Plugins/FX/Niagara/Source/Niagara/Classes/NiagaraScript.h b/Engine/Plugins/FX/Niagara/Source/Niagara/Classes/NiagaraScript.h
--- a/Engine/Plugins/FX/Niagara/Source/Niagara/Classes/NiagaraScript.h
+++ b/Engine/Plugins/FX/Niagara/Source/Niagara/Classes/NiagaraScript.h
@@ -16,6 +16,7 @@
 #include "NiagaraScriptHighlight.h"
 #include "NiagaraStackSection.h"
 #include "NiagaraParameterDefinitionsSubscriber.h"
+#include "NiagaraVersionedObject.h"
 #include "HAL/CriticalSection.h"
 #include "HAL/PlatformAtomics.h"
 #include "VectorVM.h"
@@ -652,7 +653,7 @@ private:
 
 /** Runtime script for a Niagara system */
 UCLASS(MinimalAPI)
-class UNiagaraScript : public UNiagaraScriptBase
+class UNiagaraScript : public UNiagaraScriptBase, public FNiagaraVersionedObject
 {
 	GENERATED_UCLASS_BODY()
 public:
@@ -660,7 +661,7 @@ public:
 
 #if WITH_EDITORONLY_DATA
 	/** If true then this script asset uses active version control to track changes. */
-	NIAGARA_API bool IsVersioningEnabled() const { return bVersioningEnabled; }
+	NIAGARA_API virtual bool IsVersioningEnabled() const override { return bVersioningEnabled; }
 
 	/** Returns the script data for latest exposed version. */
 	NIAGARA_API FVersionedNiagaraScriptData* GetLatestScriptData();
@@ -671,28 +672,30 @@ public:
 	NIAGARA_API const FVersionedNiagaraScriptData* GetScriptData(const FGuid& VersionGuid) const;
 
 	/** Returns all available versions for this script. */
-	NIAGARA_API TArray<FNiagaraAssetVersion> GetAllAvailableVersions() const;
+	NIAGARA_API virtual TArray<FNiagaraAssetVersion> GetAllAvailableVersions() const override;
 
 	/** Returns the version of the exposed version data (i.e. the version used when adding a module to the stack) */
-	NIAGARA_API FNiagaraAssetVersion GetExposedVersion() const;
+	NIAGARA_API virtual FNiagaraAssetVersion GetExposedVersion() const override;
 
 	/** Returns the version data for the given guid, if it exists. Otherwise returns nullptr. */
-	NIAGARA_API FNiagaraAssetVersion const* FindVersionData(const FGuid& VersionGuid) const;
+	NIAGARA_API virtual FNiagaraAssetVersion const* FindVersionData(const FGuid& VersionGuid) const override;
 	
 	/** Creates a new data entry for the given version number. The version must be > 1.0 and must not collide with an already existing version. The data will be a copy of the previous minor version. */
-	NIAGARA_API FGuid AddNewVersion(int32 MajorVersion, int32 MinorVersion);
+	NIAGARA_API virtual FGuid AddNewVersion(int32 MajorVersion, int32 MinorVersion) override;
 
 	/** Deletes the version data for an existing version. The exposed version cannot be deleted and will result in an error. Does nothing if the guid does not exist in the script's version data. */
-	NIAGARA_API void DeleteVersion(const FGuid& VersionGuid);
+	NIAGARA_API virtual void DeleteVersion(const FGuid& VersionGuid) override;
 
 	/** Changes the exposed version. Does nothing if the guid does not exist in the script's version data. */
-	NIAGARA_API void ExposeVersion(const FGuid& VersionGuid);
+	NIAGARA_API virtual void ExposeVersion(const FGuid& VersionGuid) override;
 
 	/** Enables versioning for this script asset. */
-	NIAGARA_API void EnableVersioning();
+	NIAGARA_API virtual void EnableVersioning() override;
 
 	/** Disables versioning and keeps only the data from the given version guid. Note that this breaks ALL references from existing assets and should only be used when creating a copy of a script, as the effect is very destructive.  */
-	NIAGARA_API void DisableVersioning(const FGuid& VersionGuidToUse);
+	NIAGARA_API virtual void DisableVersioning(const FGuid& VersionGuidToUse) override;
+
+	NIAGARA_API virtual TSharedPtr<FNiagaraVersionDataAccessor> GetVersionDataAccessor(const FGuid& Version) override; 
 
 	/** Makes sure that the default version data is available and fixes old script assets. */
 	NIAGARA_API void CheckVersionDataAvailable();
@@ -701,6 +704,9 @@ public:
 	NIAGARA_API UNiagaraScript* CreateCompilationCopy();
 #endif
 
+	/** Workaround for emitter versioning because we used a lot of Script->GetOuter() previously. */
+	NIAGARA_API FVersionedNiagaraEmitter GetOuterEmitter() const;
+
 	// how this script is to be used. cannot be private due to use of GET_MEMBER_NAME_CHECKED
 	UPROPERTY(AssetRegistrySearchable)
 	ENiagaraScriptUsage Usage;
@@ -1212,7 +1218,7 @@ private:
 	private :
 
 #if WITH_EDITORONLY_DATA
-		void ComputeVMCompilationId_EmitterShared(FNiagaraVMExecutableDataId& Id, UNiagaraEmitter* Emitter, UNiagaraSystem* EmitterOwner, ENiagaraRendererSourceDataMode InSourceMode) const;
+		void ComputeVMCompilationId_EmitterShared(FNiagaraVMExecutableDataId& Id, const FVersionedNiagaraEmitter& Emitter, ENiagaraRendererSourceDataMode InSourceMode) const;
 #endif
 };
 
diff --git a/Engine/Plugins/FX/Niagara/Source/Niagara/Classes/NiagaraScriptSourceBase.h b/Engine/Plugins/FX/Niagara/Source/Niagara/Classes/NiagaraScriptSourceBase.h
--- a/Engine/Plugins/FX/Niagara/Source/Niagara/Classes/NiagaraScriptSourceBase.h
+++ b/Engine/Plugins/FX/Niagara/Source/Niagara/Classes/NiagaraScriptSourceBase.h
@@ -3,9 +3,9 @@
 #pragma once
 
 #include "CoreMinimal.h"
-#include "UObject/ObjectMacros.h"
 #include "UObject/Object.h"
 #include "NiagaraCommon.h"
+#include "NiagaraTypes.h"
 #include "NiagaraCompileHash.h"
 #include "Misc/Guid.h"
 #include "NiagaraScriptSourceBase.generated.h"
@@ -114,7 +114,7 @@ class UNiagaraScriptSourceBase : public UObject
 	/** Enforce that the source graph is now out of sync with the script.*/
 	virtual void MarkNotSynchronized(FString Reason) {}
 
-	virtual FGuid GetChangeID() { return FGuid(); };
+	virtual FGuid GetChangeID() { return FGuid(); }
 
 	virtual void ComputeVMCompilationId(struct FNiagaraVMExecutableDataId& Id, ENiagaraScriptUsage InUsage, const FGuid& InUsageId, bool bForceRebuild = false) const {};
 
@@ -131,13 +131,13 @@ class UNiagaraScriptSourceBase : public UObject
 	 * Allows the derived editor only script source to handle a post load requested by an owning emitter. 
 	 * @param OwningEmitter The emitter requesting the post load.
 	 */
-	virtual void PostLoadFromEmitter(UNiagaraEmitter& OwningEmitter) { }
+	virtual void PostLoadFromEmitter(FVersionedNiagaraEmitter OwningEmitter) { }
 
 	/** Adds a module if it isn't already in the graph. If the module isn't found bOutFoundModule will be false. If it is found and it did need to be added, the function returns true. If it already exists, it returns false. */
 	NIAGARA_API virtual bool AddModuleIfMissing(FString ModulePath, ENiagaraScriptUsage Usage, bool& bOutFoundModule) { bOutFoundModule = false; return false; }
 
 #if WITH_EDITOR
-	virtual void CleanUpOldAndInitializeNewRapidIterationParameters(const UNiagaraEmitter* Emitter, ENiagaraScriptUsage ScriptUsage, FGuid ScriptUsageId, FNiagaraParameterStore& RapidIterationParameters) const { checkf(false, TEXT("Not implemented")); }
+	virtual void CleanUpOldAndInitializeNewRapidIterationParameters(const FVersionedNiagaraEmitter& Emitter, ENiagaraScriptUsage ScriptUsage, FGuid ScriptUsageId, FNiagaraParameterStore& RapidIterationParameters) const { checkf(false, TEXT("Not implemented")); }
 
 	FOnChanged& OnChanged() { return OnChangedDelegate; }
 
@@ -170,8 +170,9 @@ class UNiagaraScriptSourceBase : public UObject
 
 	/** Checks if any of the provided variables are linked to function inputs of position type data
 	 */
-	virtual void GetLinkedPositionTypeInputs(const TArray<FNiagaraVariable>& ParametersToCheck, TSet<FNiagaraVariable>& OutLinkedParameters) {};
-	virtual void ChangedLinkedInputTypes(const FNiagaraVariable& ParametersToCheck, const FNiagaraTypeDefinition& NewType) {};
+	virtual void GetLinkedPositionTypeInputs(const TArray<FNiagaraVariable>& ParametersToCheck, TSet<FNiagaraVariable>& OutLinkedParameters) {}
+	virtual void ChangedLinkedInputTypes(const FNiagaraVariable& ParametersToCheck, const FNiagaraTypeDefinition& NewType) {}
+	virtual void ReplaceScriptReferences(UNiagaraScript* OldScript, UNiagaraScript* NewScript) {}
 
 protected:
 	FOnChanged OnChangedDelegate;
diff --git a/Engine/Plugins/FX/Niagara/Source/Niagara/Classes/NiagaraSystem.h b/Engine/Plugins/FX/Niagara/Source/Niagara/Classes/NiagaraSystem.h
--- a/Engine/Plugins/FX/Niagara/Source/Niagara/Classes/NiagaraSystem.h
+++ b/Engine/Plugins/FX/Niagara/Source/Niagara/Classes/NiagaraSystem.h
@@ -19,6 +19,7 @@
 #include "Particles/ParticleSystem.h"
 #include "UObject/Object.h"
 #include "UObject/ObjectMacros.h"
+#include "Components/PrimitiveComponent.h"
 
 #include "NiagaraDataInterfacePlatformSet.h"
 
@@ -231,7 +232,8 @@ public:
 
 #if WITH_EDITORONLY_DATA
 	//~ Begin INiagaraParameterDefinitionsSubscriber interface
-	virtual const TArray<FParameterDefinitionsSubscription>& GetParameterDefinitionsSubscriptions() const override { return ParameterDefinitionsSubscriptions; };
+	virtual const TArray<FParameterDefinitionsSubscription>& GetParameterDefinitionsSubscriptions() const override { return ParameterDefinitionsSubscriptions; }
+	
 	virtual TArray<FParameterDefinitionsSubscription>& GetParameterDefinitionsSubscriptions() override { return ParameterDefinitionsSubscriptions; };
 
 	/** Get all UNiagaraScriptSourceBase of this subscriber. */
@@ -248,6 +250,8 @@ public:
 	 */
 	virtual TArray<INiagaraParameterDefinitionsSubscriber*> GetOwnedParameterDefinitionsSubscribers() override;
 	//~ End INiagaraParameterDefinitionsSubscriber interface
+
+	virtual bool ChangeEmitterVersion(const FVersionedNiagaraEmitter& Emitter, const FGuid& NewVersion);
 #endif 
 
 	/** Gets an array of the emitter handles. */
@@ -270,7 +274,7 @@ public:
 #if WITH_EDITORONLY_DATA
 	/** Adds a new emitter handle to this System.  The new handle exposes an Instance value which is a copy of the
 		original asset. */
-	FNiagaraEmitterHandle AddEmitterHandle(UNiagaraEmitter& SourceEmitter, FName EmitterName);
+	FNiagaraEmitterHandle AddEmitterHandle(UNiagaraEmitter& SourceEmitter, FName EmitterName, FGuid EmitterVersion);
 
 	/** Adds a new emitter handle to this system without copying the original asset. This should only be used for temporary systems and never for live assets. */
 	void AddEmitterHandleDirect(FNiagaraEmitterHandle& EmitterHandleToAdd);
@@ -356,7 +360,7 @@ public:
 	bool HasOutstandingCompilationRequests(bool bIncludingGPUShaders = false) const;
 
 	/** Determines if this system has the supplied emitter as an editable and simulating emitter instance. */
-	bool ReferencesInstanceEmitter(UNiagaraEmitter& Emitter);
+	bool ReferencesInstanceEmitter(const FVersionedNiagaraEmitter& Emitter) const;
 
 	/** Updates the system's rapid iteration parameters from a specific emitter. */
 	void RefreshSystemParametersFromEmitter(const FNiagaraEmitterHandle& EmitterHandle);
@@ -457,12 +461,13 @@ public:
 	FORCEINLINE bool ShouldDisableExperimentalVM() const { return bDisableExperimentalVM; }
 
 #if WITH_EDITORONLY_DATA
-	bool UsesEmitter(const UNiagaraEmitter* Emitter) const;
+	bool UsesEmitter(UNiagaraEmitter* Emitter) const;
+	bool UsesEmitter(const FVersionedNiagaraEmitter& VersionedEmitter) const;
 	bool UsesScript(const UNiagaraScript* Script)const; 
 	void ForceGraphToRecompileOnNextCheck();
 
-	static void RequestCompileForEmitter(UNiagaraEmitter* InEmitter);
-	static void RecomputeExecutionOrderForEmitter(UNiagaraEmitter* InEmitter);
+	static void RequestCompileForEmitter(const FVersionedNiagaraEmitter& InEmitter);
+	static void RecomputeExecutionOrderForEmitter(const FVersionedNiagaraEmitter& InEmitter);
 	static void RecomputeExecutionOrderForDataInterface(class UNiagaraDataInterface* DataInterface);
 
 	FORCEINLINE bool ShouldUseRapidIterationParameters() const { return bCompileForEdit ? !bBakeOutRapidIteration : !bBakeOutRapidIterationOnCook; }
@@ -581,10 +586,10 @@ public:
 	bool ComputeEmitterPriority(int32 EmitterIdx, TArray<int32, TInlineAllocator<32>>& EmitterPriorities, const TBitArray<TInlineAllocator<32>>& EmitterDependencyGraph);
 
 	/** Queries all the data interfaces in the array for emitter dependencies. */
-	void FindDataInterfaceDependencies(UNiagaraEmitter* Emitter, UNiagaraScript* Script, TArray<class UNiagaraEmitter*>& Dependencies);
+	void FindDataInterfaceDependencies(FVersionedNiagaraEmitterData* EmitterData, UNiagaraScript* Script, TArray<FVersionedNiagaraEmitter>& Dependencies);
 
 	/** Looks at all the event handlers in the emitter to determine which other emitters it depends on. */
-	void FindEventDependencies(UNiagaraEmitter* Emitter, TArray<UNiagaraEmitter*>& Dependencies);
+	void FindEventDependencies(FVersionedNiagaraEmitterData* EmitterData, TArray<FVersionedNiagaraEmitter>& Dependencies);
 
 	/** Computes the order in which the emitters in the Emitters array will be ticked and stores the results in EmitterExecutionOrder. */
 	void ComputeEmittersExecutionOrder();
@@ -720,7 +725,7 @@ private:
 	const FName GetEmitterVariableAliasName(const FNiagaraVariable& InEmitterVar, const UNiagaraEmitter* InEmitter) const;
 
 	/** Helper for filling in attribute datasets per emitter. */
-	void InitEmitterDataSetCompiledData(FNiagaraDataSetCompiledData& DataSetToInit, const UNiagaraEmitter* InAssociatedEmitter, const FNiagaraEmitterHandle& InAssociatedEmitterHandle);
+	void InitEmitterDataSetCompiledData(FNiagaraDataSetCompiledData& DataSetToInit, const FNiagaraEmitterHandle& InAssociatedEmitterHandle);
 	void PrepareRapidIterationParametersForCompilation();
 
 	/** Resets internal data leaving it in a state which would have minimal cost to exist in headless builds (servers) */
@@ -975,9 +980,9 @@ void UNiagaraSystem::ForEachScript(TAction Func) const
 
 	for (const FNiagaraEmitterHandle& Handle : EmitterHandles)
 	{
-		if (UNiagaraEmitter* Emitter = Handle.GetInstance())
+		if (FVersionedNiagaraEmitterData* EmitterData = Handle.GetEmitterData())
 		{
-			Emitter->ForEachScript(Func);
+			EmitterData->ForEachScript(Func);
 		}
 	}
 }
@@ -989,7 +994,7 @@ void UNiagaraSystem::ForEachPlatformSet(TAction Func)
 	//Handle our scalability overrides
 	for (FNiagaraSystemScalabilityOverride& Override : SystemScalabilityOverrides.Overrides)
 	{
-		Func(this, Override.Platforms);
+		Func(Override.Platforms);
 	}
 
 	//Handle and platform set User DIs.
@@ -997,7 +1002,7 @@ void UNiagaraSystem::ForEachPlatformSet(TAction Func)
 	{
 		if (UNiagaraDataInterfacePlatformSet* PlatformSetDI = Cast<UNiagaraDataInterfacePlatformSet>(DI))
 		{
-			Func(PlatformSetDI, PlatformSetDI->Platforms);
+			Func(PlatformSetDI->Platforms);
 		}
 	}
 
@@ -1010,7 +1015,7 @@ void UNiagaraSystem::ForEachPlatformSet(TAction Func)
 			{
 				if (UNiagaraDataInterfacePlatformSet* PlatformSetDI = Cast<UNiagaraDataInterfacePlatformSet>(DataInterfaceInfo.DataInterface))
 				{
-					Func(PlatformSetDI, PlatformSetDI->Platforms);
+					Func(PlatformSetDI->Platforms);
 				}
 			}
 		}
@@ -1021,7 +1026,7 @@ void UNiagaraSystem::ForEachPlatformSet(TAction Func)
 	//Finally handle all our emitters.
 	for (FNiagaraEmitterHandle& Handle : EmitterHandles)
 	{
-		if (UNiagaraEmitter* Emitter = Handle.GetInstance())
+		if (FVersionedNiagaraEmitterData* Emitter = Handle.GetEmitterData())
 		{
 			Emitter->ForEachPlatformSet(Func);
 		}
diff --git a/Engine/Plugins/FX/Niagara/Source/Niagara/Classes/NiagaraVersionedObject.h b/Engine/Plugins/FX/Niagara/Source/Niagara/Classes/NiagaraVersionedObject.h
new file mode 100644
--- /dev/null
+++ b/Engine/Plugins/FX/Niagara/Source/Niagara/Classes/NiagaraVersionedObject.h
@@ -0,0 +1,52 @@
+﻿// Copyright Epic Games, Inc. All Rights Reserved.
+
+#pragma once
+
+#include "CoreMinimal.h"
+#include "NiagaraTypes.h"
+
+struct NIAGARA_API FNiagaraVersionDataAccessor
+{
+	virtual ~FNiagaraVersionDataAccessor() = default;
+
+	virtual FNiagaraAssetVersion& GetObjectVersion();
+	virtual FText& GetVersionChangeDescription();
+	virtual ENiagaraPythonUpdateScriptReference& GetUpdateScriptExecutionType();
+	virtual FString& GetPythonUpdateScript();
+	virtual FFilePath& GetScriptAsset();
+};
+
+class NIAGARA_API FNiagaraVersionedObject
+{
+public:
+	virtual ~FNiagaraVersionedObject() = default;
+
+	virtual TSharedPtr<FNiagaraVersionDataAccessor> GetVersionDataAccessor(const FGuid& Version); 
+
+	/** If true then this script asset uses active version control to track changes. */
+	virtual bool IsVersioningEnabled() const;
+	
+	/** Returns all available versions for this object. */
+	virtual TArray<FNiagaraAssetVersion> GetAllAvailableVersions() const;
+
+	/** Returns the version of the exposed version data (i.e. the version used when adding a module to the stack) */
+	virtual FNiagaraAssetVersion GetExposedVersion() const;
+
+	/** Returns the version data for the given guid, if it exists. Otherwise returns nullptr. */
+	virtual FNiagaraAssetVersion const* FindVersionData(const FGuid& VersionGuid) const;
+	
+	/** Creates a new data entry for the given version number. The version must be > 1.0 and must not collide with an already existing version. The data will be a copy of the previous minor version. */
+	virtual FGuid AddNewVersion(int32 MajorVersion, int32 MinorVersion);
+
+	/** Deletes the version data for an existing version. The exposed version cannot be deleted and will result in an error. Does nothing if the guid does not exist in the object's version data. */
+	virtual void DeleteVersion(const FGuid& VersionGuid);
+
+	/** Changes the exposed version. Does nothing if the guid does not exist in the object's version data. */
+	virtual void ExposeVersion(const FGuid& VersionGuid);
+
+	/** Enables versioning for this object. */
+	virtual void EnableVersioning();
+
+	/** Disables versioning and keeps only the data from the given version guid. Note that this breaks ALL references from existing assets and should only be used when creating a copy of an object, as the effect is very destructive.  */
+	virtual void DisableVersioning(const FGuid& VersionGuidToUse);
+};
diff --git a/Engine/Plugins/FX/Niagara/Source/Niagara/Private/DataInterface/NiagaraDataInterfaceDebugDraw.cpp b/Engine/Plugins/FX/Niagara/Source/Niagara/Private/DataInterface/NiagaraDataInterfaceDebugDraw.cpp
--- a/Engine/Plugins/FX/Niagara/Source/Niagara/Private/DataInterface/NiagaraDataInterfaceDebugDraw.cpp
+++ b/Engine/Plugins/FX/Niagara/Source/Niagara/Private/DataInterface/NiagaraDataInterfaceDebugDraw.cpp
@@ -1107,14 +1107,14 @@ struct FNDIDebugDrawInstanceData_GameThread
 
 				for (const FNiagaraEmitterHandle& Handle : System->GetEmitterHandles())
 				{
-					UNiagaraEmitter* Emitter = Handle.GetInstance();
-					if (Emitter)
+					FVersionedNiagaraEmitterData* EmitterData = Handle.GetEmitterData();
+					if (EmitterData)
 					{
-						bool bEmitterIsLocal = Emitter->bLocalSpace;
-						if (Emitter->SimTarget == ENiagaraSimTarget::CPUSim)
+						bool bEmitterIsLocal = EmitterData->bLocalSpace;
+						if (EmitterData->SimTarget == ENiagaraSimTarget::CPUSim)
 						{
 							int32 ScriptCap = Scripts.Num();
-							Emitter->GetScripts(Scripts, true, true);
+							EmitterData->GetScripts(Scripts, true, true);
 
 							for (int32 i = ScriptCap; i < Scripts.Num(); i++)
 								ScriptIsLocal.Add(bEmitterIsLocal);
@@ -1124,8 +1124,8 @@ struct FNDIDebugDrawInstanceData_GameThread
 							// It's a little weird to do this, but ultimately all the rapid iteration values are 
 							// referenced by the compile tags from these scripts and we want to get the most up-to-date 
 							// values here. If we reference the GPU script here, it will have stale values for some reason.
-							Scripts.Add(Emitter->SpawnScriptProps.Script);
-							Scripts.Add(Emitter->UpdateScriptProps.Script);
+							Scripts.Add(EmitterData->SpawnScriptProps.Script);
+							Scripts.Add(EmitterData->UpdateScriptProps.Script);
 							ScriptIsLocal.Add(bEmitterIsLocal);
 							ScriptIsLocal.Add(bEmitterIsLocal);
 						}
diff --git a/Engine/Plugins/FX/Niagara/Source/Niagara/Private/DataInterface/NiagaraDataInterfaceEmitterProperties.cpp b/Engine/Plugins/FX/Niagara/Source/Niagara/Private/DataInterface/NiagaraDataInterfaceEmitterProperties.cpp
--- a/Engine/Plugins/FX/Niagara/Source/Niagara/Private/DataInterface/NiagaraDataInterfaceEmitterProperties.cpp
+++ b/Engine/Plugins/FX/Niagara/Source/Niagara/Private/DataInterface/NiagaraDataInterfaceEmitterProperties.cpp
@@ -125,9 +125,9 @@ bool UNiagaraDataInterfaceEmitterProperties::InitPerInstanceData(void* PerInstan
 
 	FInstanceData_GameThread* InstanceData_GT = new(PerInstanceData) FInstanceData_GameThread();
 	InstanceData_GT->EmitterInstance = EmitterBinding.Resolve(SystemInstance, this);
-	if (InstanceData_GT->EmitterInstance)
+	if (InstanceData_GT->EmitterInstance && InstanceData_GT->EmitterInstance->GetCachedEmitterData())
 	{
-		InstanceData_GT->bLocalSpace = InstanceData_GT->EmitterInstance->GetCachedEmitter()->bLocalSpace;
+		InstanceData_GT->bLocalSpace = InstanceData_GT->EmitterInstance->GetCachedEmitterData()->bLocalSpace;
 	}
 
 	if ( IsUsedByGPUEmitter() )
diff --git a/Engine/Plugins/FX/Niagara/Source/Niagara/Private/DataInterface/NiagaraDataInterfaceUObjectPropertyReader.cpp b/Engine/Plugins/FX/Niagara/Source/Niagara/Private/DataInterface/NiagaraDataInterfaceUObjectPropertyReader.cpp
--- a/Engine/Plugins/FX/Niagara/Source/Niagara/Private/DataInterface/NiagaraDataInterfaceUObjectPropertyReader.cpp
+++ b/Engine/Plugins/FX/Niagara/Source/Niagara/Private/DataInterface/NiagaraDataInterfaceUObjectPropertyReader.cpp
@@ -601,7 +601,7 @@ bool UNiagaraDataInterfaceUObjectPropertyReader::InitPerInstanceData(void* PerIn
 		// We shouldn't need to do this per init, we should be able to cache once and once only
 		for (const TSharedRef<FNiagaraEmitterInstance, ESPMode::ThreadSafe>& EmitterInstance : SystemInstance->GetEmitters())
 		{
-			if (EmitterInstance->IsDisabled() || EmitterInstance->GetCachedEmitter() == nullptr || EmitterInstance->GetGPUContext() == nullptr)
+			if (EmitterInstance->IsDisabled() || EmitterInstance->GetCachedEmitterData() == nullptr || EmitterInstance->GetGPUContext() == nullptr)
 			{
 				continue;
 			}
diff --git a/Engine/Plugins/FX/Niagara/Source/Niagara/Private/DataInterface/NiagaraDataInterfaceUtilities.cpp b/Engine/Plugins/FX/Niagara/Source/Niagara/Private/DataInterface/NiagaraDataInterfaceUtilities.cpp
--- a/Engine/Plugins/FX/Niagara/Source/Niagara/Private/DataInterface/NiagaraDataInterfaceUtilities.cpp
+++ b/Engine/Plugins/FX/Niagara/Source/Niagara/Private/DataInterface/NiagaraDataInterfaceUtilities.cpp
@@ -234,7 +234,7 @@ void FNiagaraDataInterfaceUtilities::ForEachVMFunction(class UNiagaraDataInterfa
 	// Search emitter scripts
 	for (const TSharedRef<FNiagaraEmitterInstance, ESPMode::ThreadSafe>& EmitterInstance : SystemInstance->GetEmitters())
 	{
-		if (EmitterInstance->IsDisabled() || EmitterInstance->GetCachedEmitter() == nullptr || EmitterInstance->GetGPUContext() != nullptr)
+		if (EmitterInstance->IsDisabled() || EmitterInstance->GetCachedEmitter().Emitter == nullptr || EmitterInstance->GetGPUContext() != nullptr)
 		{
 			continue;
 		}
@@ -268,7 +268,7 @@ void FNiagaraDataInterfaceUtilities::ForEachGpuFunction(class UNiagaraDataInterf
 
 	for (const TSharedRef<FNiagaraEmitterInstance, ESPMode::ThreadSafe>& EmitterInstance : SystemInstance->GetEmitters())
 	{
-		if (EmitterInstance->IsDisabled() || EmitterInstance->GetCachedEmitter() == nullptr || EmitterInstance->GetGPUContext() == nullptr)
+		if (EmitterInstance->IsDisabled() || EmitterInstance->GetCachedEmitter().Emitter == nullptr || EmitterInstance->GetGPUContext() == nullptr)
 		{
 			continue;
 		}
diff --git a/Engine/Plugins/FX/Niagara/Source/Niagara/Private/NiagaraAsyncCompile.cpp b/Engine/Plugins/FX/Niagara/Source/Niagara/Private/NiagaraAsyncCompile.cpp
--- a/Engine/Plugins/FX/Niagara/Source/Niagara/Private/NiagaraAsyncCompile.cpp
+++ b/Engine/Plugins/FX/Niagara/Source/Niagara/Private/NiagaraAsyncCompile.cpp
@@ -45,11 +45,11 @@ TSharedPtr<FNiagaraCompileRequestDataBase, ESPMode::ThreadSafe> FNiagaraLazyPrec
 		for (int32 i = 0; i < EmitterHandles.Num(); i++)
 		{
 			const FNiagaraEmitterHandle& Handle = EmitterHandles[i];
-			if (Handle.GetInstance() && Handle.GetIsEnabled())
+			if (Handle.GetEmitterData() && Handle.GetIsEnabled())
 			{
 				TSharedPtr<FNiagaraCompileRequestDataBase, ESPMode::ThreadSafe> EmitterPrecompiledData = SystemPrecompiledData->GetDependentRequest(i);
 				TArray<UNiagaraScript*> EmitterScripts;
-				Handle.GetInstance()->GetScripts(EmitterScripts, false, true);
+				Handle.GetEmitterData()->GetScripts(EmitterScripts, false, true);
 				check(EmitterScripts.Num() > 0);
 				for (UNiagaraScript* EmitterScript : EmitterScripts)
 				{
@@ -264,7 +264,7 @@ void FNiagaraAsyncCompileTask::PrecompileData()
 	DDCKey = CurrentDDCKey;
 	
 	ComputedPrecompileData = PrecompileReference->GetPrecompileData(ScriptPair.CompiledScript);
-	TSharedPtr<FNiagaraCompileRequestDuplicateDataBase, ESPMode::ThreadSafe> SystemPrecompileDuplicateData = PrecompileReference->GetPrecompileDuplicateData(ScriptPair.Emitter, ScriptPair.CompiledScript);
+	TSharedPtr<FNiagaraCompileRequestDuplicateDataBase, ESPMode::ThreadSafe> SystemPrecompileDuplicateData = PrecompileReference->GetPrecompileDuplicateData(ScriptPair.VersionedEmitter.Emitter, ScriptPair.CompiledScript);
 
 	if (ComputedPrecompileData == nullptr || SystemPrecompileDuplicateData == nullptr)
 	{
@@ -286,8 +286,8 @@ void FNiagaraAsyncCompileTask::PrecompileData()
 	}
 	else
 	{
-		UNiagaraEmitter* OwningEmitter = ScriptPair.Emitter;
-		int32 EmitterIndex = OwningSystem->GetEmitterHandles().IndexOfByPredicate([OwningEmitter](const FNiagaraEmitterHandle& EmitterHandle) { return EmitterHandle.GetInstance() == OwningEmitter; });
+		UNiagaraEmitter* OwningEmitter = ScriptPair.VersionedEmitter.Emitter;
+		int32 EmitterIndex = OwningSystem->GetEmitterHandles().IndexOfByPredicate([OwningEmitter](const FNiagaraEmitterHandle& EmitterHandle) { return EmitterHandle.GetInstance().Emitter == OwningEmitter; });
 		if(EmitterIndex != INDEX_NONE)
 		{
 			ComputedPrecompileDuplicateData = SystemPrecompileDuplicateData->GetDependentRequest(EmitterIndex);
diff --git a/Engine/Plugins/FX/Niagara/Source/Niagara/Private/NiagaraAsyncCompile.h b/Engine/Plugins/FX/Niagara/Source/Niagara/Private/NiagaraAsyncCompile.h
--- a/Engine/Plugins/FX/Niagara/Source/Niagara/Private/NiagaraAsyncCompile.h
+++ b/Engine/Plugins/FX/Niagara/Source/Niagara/Private/NiagaraAsyncCompile.h
@@ -17,7 +17,7 @@ struct FEmitterCompiledScriptPair
 	GENERATED_USTRUCT_BODY()
 	
 	bool bResultsReady = false;
-	UNiagaraEmitter* Emitter = nullptr;
+	FVersionedNiagaraEmitter VersionedEmitter;
 	UNiagaraScript* CompiledScript = nullptr;
 	uint32 PendingJobID = INDEX_NONE; // this is the ID for any active shader compiler worker job
 	FNiagaraVMExecutableDataId CompileId;
