Added versioning for Niagara emitters

diff --git a/Engine/Plugins/FX/Niagara/Source/Niagara/Private/NiagaraCommon.cpp b/Engine/Plugins/FX/Niagara/Source/Niagara/Private/NiagaraCommon.cpp
--- a/Engine/Plugins/FX/Niagara/Source/Niagara/Private/NiagaraCommon.cpp
+++ b/Engine/Plugins/FX/Niagara/Source/Niagara/Private/NiagaraCommon.cpp
@@ -151,7 +151,7 @@ void FNiagaraSystemUpdateContext::Add(const UNiagaraSystem* System, bool bReInit
 }
 #if WITH_EDITORONLY_DATA
 
-void FNiagaraSystemUpdateContext::Add(const UNiagaraEmitter* Emitter, bool bReInit)
+void FNiagaraSystemUpdateContext::Add(const FVersionedNiagaraEmitter& Emitter, bool bReInit)
 {
 	for (TObjectIterator<UNiagaraComponent> It; It; ++It)
 	{
@@ -280,7 +280,7 @@ void FNiagaraStatDatabase::AddStatCapture(FStatReportKey ReportKey, TMap<TStatId
 	{
 		return;
 	}
-	FScopeLock Lock(&CriticalSection);
+	FScopeLock Lock(GetCriticalSection());
 	if (StatCaptures.Num() > GNiagaraMaxStatInstanceReports)
 	{
 		// we don't need data from too many emitter instances. If we already have enough, delete an old data point.
@@ -298,13 +298,13 @@ void FNiagaraStatDatabase::AddStatCapture(FStatReportKey ReportKey, TMap<TStatId
 
 void FNiagaraStatDatabase::ClearStatCaptures()
 {
-	FScopeLock Lock(&CriticalSection);
+	FScopeLock Lock(GetCriticalSection());
 	StatCaptures.Empty();
 }
 
 float FNiagaraStatDatabase::GetRuntimeStat(FName StatName, ENiagaraScriptUsage Usage, ENiagaraStatEvaluationType EvaluationType)
 {
-	FScopeLock Lock(&CriticalSection);
+	FScopeLock Lock(GetCriticalSection());
 	int32 ValueCount = 0;
 	float Sum = 0;
 	float Max = 0;
@@ -338,7 +338,7 @@ float FNiagaraStatDatabase::GetRuntimeStat(FName StatName, ENiagaraScriptUsage U
 
 float FNiagaraStatDatabase::GetRuntimeStat(ENiagaraScriptUsage Usage, ENiagaraStatEvaluationType EvaluationType)
 {
-	FScopeLock Lock(&CriticalSection);
+	FScopeLock Lock(GetCriticalSection());
 	int32 ValueCount = 0;
 	float Sum = 0;
 	float Max = 0;
@@ -368,7 +368,7 @@ float FNiagaraStatDatabase::GetRuntimeStat(ENiagaraScriptUsage Usage, ENiagaraSt
 
 TMap<ENiagaraScriptUsage, TSet<FName>> FNiagaraStatDatabase::GetAvailableStatNames()
 {
-	FScopeLock Lock(&CriticalSection);
+	FScopeLock Lock(GetCriticalSection());
 	TMap<ENiagaraScriptUsage, TSet<FName>> Result;
 	for (const auto& EmitterEntry : StatCaptures)
 	{
@@ -380,15 +380,30 @@ TMap<ENiagaraScriptUsage, TSet<FName>> FNiagaraStatDatabase::GetAvailableStatNam
 	}
 	return Result;
 }
+
+void FNiagaraStatDatabase::Init()
+{
+	if (CriticalSection.IsValid() == false)
+	{
+		CriticalSection = MakeShared<FCriticalSection>();
+	}
+}
+
+FCriticalSection* FNiagaraStatDatabase::GetCriticalSection() const
+{
+	ensure(CriticalSection.IsValid());
+	return CriticalSection.Get();
+}
 #endif
 
-void FNiagaraVariableAttributeBinding::SetValue(const FName& InValue, const UNiagaraEmitter* InEmitter, ENiagaraRendererSourceDataMode InSourceMode)
+void FNiagaraVariableAttributeBinding::SetValue(const FName& InValue, const FVersionedNiagaraEmitter& InVersionedEmitter, ENiagaraRendererSourceDataMode InSourceMode)
 {
+	UNiagaraEmitter* Emitter = InVersionedEmitter.Emitter;
 	RootVariable.SetName(InValue);
 
 	const bool bIsRootParticleValue = RootVariable.IsInNameSpace(FNiagaraConstants::ParticleAttributeNamespaceString);
 	const bool bIsRootUnaliasedEmitterValue = RootVariable.IsInNameSpace(FNiagaraConstants::EmitterNamespaceString);
-	const bool bIsAliasedEmitterValue = InEmitter ? RootVariable.IsInNameSpace(InEmitter->GetUniqueEmitterName()) : false;
+	const bool bIsAliasedEmitterValue = Emitter ? RootVariable.IsInNameSpace(Emitter->GetUniqueEmitterName()) : false;
 	const bool bIsRootSystemValue = RootVariable.IsInNameSpace(FNiagaraConstants::SystemNamespaceString);
 	const bool bIsRootUserValue = RootVariable.IsInNameSpace(FNiagaraConstants::UserNamespaceString);
 	const bool bIsStackContextValue = RootVariable.IsInNameSpace(FNiagaraConstants::StackContextNamespaceString);
@@ -408,7 +423,7 @@ void FNiagaraVariableAttributeBinding::SetValue(const FName& InValue, const UNia
 	{
 		// First, replace unaliased emitter namespace with "Emitter" namespace
 		RootVariable = FNiagaraUtilities::ResolveAliases(RootVariable, FNiagaraAliasContext()
-			.ChangeEmitterNameToEmitter(InEmitter->GetUniqueEmitterName()));
+			.ChangeEmitterNameToEmitter(Emitter->GetUniqueEmitterName()));
 
 		// Now strip out "Emitter"
 		RootVariable.SetName(FNiagaraConstants::GetAttributeAsEmitterDataSetKey(RootVariable).GetName());
@@ -422,10 +437,10 @@ void FNiagaraVariableAttributeBinding::SetValue(const FName& InValue, const UNia
 	else if (bIsRootUnaliasedEmitterValue || bIsAliasedEmitterValue)
 	{
 		// First, replace unaliased emitter namespace with "Emitter" namespace
-		if (InEmitter != nullptr)
+		if (Emitter != nullptr)
 		{
 			RootVariable = FNiagaraUtilities::ResolveAliases(RootVariable, FNiagaraAliasContext()
-				.ChangeEmitterNameToEmitter(InEmitter->GetUniqueEmitterName()));
+				.ChangeEmitterNameToEmitter(Emitter->GetUniqueEmitterName()));
 		}
 
 		// Now strip out "Emitter"
@@ -445,14 +460,15 @@ void FNiagaraVariableAttributeBinding::SetValue(const FName& InValue, const UNia
 		ensureMsgf(!bIsStackContextValue, TEXT("Should not get to this point! Should be covered by first two branch expresssions."));
 	}
 
-	CacheValues(InEmitter, InSourceMode);
+	CacheValues(InVersionedEmitter, InSourceMode);
 }
 
-void FNiagaraVariableAttributeBinding::SetAsPreviousValue(const FNiagaraVariableBase& Src, const UNiagaraEmitter* InEmitter, ENiagaraRendererSourceDataMode InSourceMode)
+void FNiagaraVariableAttributeBinding::SetAsPreviousValue(const FNiagaraVariableBase& Src, const FVersionedNiagaraEmitter& InEmitter, ENiagaraRendererSourceDataMode InSourceMode)
 {
 	static const FString PreviousNamespace = FNiagaraConstants::PreviousNamespace.ToString();
 
-	RootVariable = ParamMapVariable = DataSetVariable = Src;
+	ParamMapVariable = Src;
+	RootVariable = DataSetVariable = Src;
 
 	// Split out the name and it's namespace
 	TArray<FString> SplitName;
@@ -484,11 +500,12 @@ void FNiagaraVariableAttributeBinding::SetAsPreviousValue(const FNiagaraVariable
 	}
 }
 
-void FNiagaraVariableAttributeBinding::SetAsPreviousValue(const FNiagaraVariableAttributeBinding& Src, const UNiagaraEmitter* InEmitter, ENiagaraRendererSourceDataMode InSourceMode)
+void FNiagaraVariableAttributeBinding::SetAsPreviousValue(const FNiagaraVariableAttributeBinding& Src, const FVersionedNiagaraEmitter& InEmitter, ENiagaraRendererSourceDataMode InSourceMode)
 {
 	static const FString PreviousNamespace = FNiagaraConstants::PreviousNamespace.ToString();
 
-	RootVariable = ParamMapVariable = DataSetVariable = Src.RootVariable;
+	ParamMapVariable = Src.RootVariable;
+	RootVariable = DataSetVariable = Src.RootVariable;
 
 	// Split out the name and it's namespace
 	TStringBuilder<128> VarName;
@@ -542,7 +559,7 @@ void FNiagaraVariableAttributeBinding::Setup(const FNiagaraVariableBase& InRootV
 	{
 		RootVariable.SetData(InDefaultValue.GetData());
 	}
-	SetValue(InRootVar.GetName(), nullptr, InSourceMode);
+	SetValue(InRootVar.GetName(), FVersionedNiagaraEmitter(), InSourceMode);
 }
 
 #if WITH_EDITORONLY_DATA
@@ -566,7 +583,7 @@ void FNiagaraVariableAttributeBinding::PostLoad(ENiagaraRendererSourceDataMode I
 	if (BoundVariable.IsValid())
 	{
 		RootVariable.SetType(DataSetVariable.GetType()); //Sometimes the BoundVariable was bogus in the past. THe DataSet shouldn't be though.
-		SetValue(BoundVariable.GetName(), nullptr, InSourceMode);
+		SetValue(BoundVariable.GetName(), FVersionedNiagaraEmitter(), InSourceMode);
 		BoundVariable = FNiagaraVariable();
 	}
 #endif
@@ -587,7 +604,7 @@ void FNiagaraVariableAttributeBinding::Dump() const
 		bBindingExistsOnSource ? 1 : 0, bIsCachedParticleValue ? 1 : 0 );
 }
 
-void FNiagaraVariableAttributeBinding::ResetToDefault(const FNiagaraVariableAttributeBinding& InOther, const UNiagaraEmitter* InEmitter, ENiagaraRendererSourceDataMode InSourceMode)
+void FNiagaraVariableAttributeBinding::ResetToDefault(const FNiagaraVariableAttributeBinding& InOther, const FVersionedNiagaraEmitter& InEmitter, ENiagaraRendererSourceDataMode InSourceMode)
 {
 	if (InOther.BindingSourceMode == ImplicitFromSource || InOther.BindingSourceMode == ENiagaraBindingSource::ExplicitEmitter || InOther.BindingSourceMode == ENiagaraBindingSource::ExplicitParticles)
 	{
@@ -606,7 +623,7 @@ void FNiagaraVariableAttributeBinding::ResetToDefault(const FNiagaraVariableAttr
 			TempVar.SetNamespacedName(FNiagaraConstants::ParticleAttributeNamespaceString, InOther.DataSetVariable.GetName());
 		}
 
-		SetValue(TempVar.GetName(), nullptr, InSourceMode);
+		SetValue(TempVar.GetName(), FVersionedNiagaraEmitter(), InSourceMode);
 	}
 	else
 	{
@@ -624,7 +641,7 @@ bool FNiagaraVariableAttributeBinding::MatchesDefault(const FNiagaraVariableAttr
 }
 
 
-bool FNiagaraVariableAttributeBinding::RenameVariableIfMatching(const FNiagaraVariableBase& OldVariable, const FNiagaraVariableBase& NewVariable, const UNiagaraEmitter* InEmitter, ENiagaraRendererSourceDataMode InSourceMode)
+bool FNiagaraVariableAttributeBinding::RenameVariableIfMatching(const FNiagaraVariableBase& OldVariable, const FNiagaraVariableBase& NewVariable, const FVersionedNiagaraEmitter& InEmitter, ENiagaraRendererSourceDataMode InSourceMode)
 {
 	// First try a namespace mangling - free match.
 	if (OldVariable.GetName() == ParamMapVariable.GetName() && OldVariable.GetType() == ParamMapVariable.GetType())
@@ -639,7 +656,7 @@ bool FNiagaraVariableAttributeBinding::RenameVariableIfMatching(const FNiagaraVa
 	{
 		// First, resolve any aliases
 		OldVarAliased = FNiagaraUtilities::ResolveAliases(OldVariable, FNiagaraAliasContext()
-			.ChangeEmitterToEmitterName(InEmitter->GetUniqueEmitterName()));
+			.ChangeEmitterToEmitterName(InEmitter.Emitter->GetUniqueEmitterName()));
 	}
 	if (OldVarAliased.GetName() == ParamMapVariable.GetName() && OldVarAliased.GetType() == ParamMapVariable.GetType())
 	{
@@ -649,7 +666,7 @@ bool FNiagaraVariableAttributeBinding::RenameVariableIfMatching(const FNiagaraVa
 	return false;
 }
 
-bool FNiagaraVariableAttributeBinding::Matches(const FNiagaraVariableBase& OldVariable,  const UNiagaraEmitter* InEmitter, ENiagaraRendererSourceDataMode InSourceMode)
+bool FNiagaraVariableAttributeBinding::Matches(const FNiagaraVariableBase& OldVariable, const FVersionedNiagaraEmitter& InEmitter, ENiagaraRendererSourceDataMode InSourceMode)
 {
 	// First try a namespace mangling - free match.
 	if (OldVariable.GetName() == ParamMapVariable.GetName() && OldVariable.GetType() == ParamMapVariable.GetType())
@@ -659,11 +676,11 @@ bool FNiagaraVariableAttributeBinding::Matches(const FNiagaraVariableBase& OldVa
 
 	// Now we need to deal with any aliased emitter namespaces for the match. If so resolve the aliases then try the match.
 	FNiagaraVariable OldVarAliased = OldVariable;
-	if (InEmitter && OldVariable.IsInNameSpace(FNiagaraConstants::EmitterNamespaceString))
+	if (InEmitter.Emitter && OldVariable.IsInNameSpace(FNiagaraConstants::EmitterNamespaceString))
 	{
 		// First, resolve any aliases
 		OldVarAliased = FNiagaraUtilities::ResolveAliases(OldVariable, FNiagaraAliasContext()
-			.ChangeEmitterToEmitterName(InEmitter->GetUniqueEmitterName()));
+			.ChangeEmitterToEmitterName(InEmitter.Emitter->GetUniqueEmitterName()));
 	}
 	if (OldVarAliased.GetName() == ParamMapVariable.GetName() && OldVarAliased.GetType() == ParamMapVariable.GetType())
 	{
@@ -672,15 +689,16 @@ bool FNiagaraVariableAttributeBinding::Matches(const FNiagaraVariableBase& OldVa
 	return false;
 }
 
-void FNiagaraVariableAttributeBinding::CacheValues(const UNiagaraEmitter* InEmitter, ENiagaraRendererSourceDataMode InSourceMode)
+void FNiagaraVariableAttributeBinding::CacheValues(const FVersionedNiagaraEmitter& InVersionedEmitter, ENiagaraRendererSourceDataMode InSourceMode)
 {
 	// Some older values may have had the root with the emitter unique name as the namespace, fix this up
 	// to meet the new assumptions.
-	if (InEmitter && RootVariable.IsInNameSpace(InEmitter->GetUniqueEmitterName()))
+	UNiagaraEmitter* Emitter = InVersionedEmitter.Emitter;
+	if (Emitter && RootVariable.IsInNameSpace(Emitter->GetUniqueEmitterName()))
 	{
 		// First, replace unaliased emitter namespace with "Emitter" namespace
 		RootVariable = FNiagaraUtilities::ResolveAliases(RootVariable, FNiagaraAliasContext()
-			.ChangeEmitterNameToEmitter(InEmitter->GetUniqueEmitterName()));
+			.ChangeEmitterNameToEmitter(Emitter->GetUniqueEmitterName()));
 
 		// Now strip out "Emitter"
 		RootVariable.SetName(FNiagaraConstants::GetAttributeAsEmitterDataSetKey(RootVariable).GetName());
@@ -722,33 +740,33 @@ void FNiagaraVariableAttributeBinding::CacheValues(const UNiagaraEmitter* InEmit
 #endif
 
 	// Now resolve if this variable actually exists.
-	if (InEmitter)
+	if (Emitter)
 	{
 		if (BindingSourceMode == ENiagaraBindingSource::ExplicitEmitter || (InSourceMode == ENiagaraRendererSourceDataMode::Emitter && BindingSourceMode == ENiagaraBindingSource::ImplicitFromSource))
 		{
 			// Replace  "Emitter" namespace with unaliased emitter namespace
 			FNiagaraAliasContext ResolveAliasesContext(FNiagaraAliasContext::ERapidIterationParameterMode::EmitterOrParticleScript);
-			ResolveAliasesContext.ChangeEmitterToEmitterName(InEmitter->GetUniqueEmitterName());
-			ParamMapVariable = FNiagaraUtilities::ResolveAliases(ParamMapVariable,  ResolveAliasesContext);
+			ResolveAliasesContext.ChangeEmitterToEmitterName(Emitter->GetUniqueEmitterName());
+			ParamMapVariable = FNiagaraUtilities::ResolveAliases(ParamMapVariable, ResolveAliasesContext);
 			DataSetVariable = FNiagaraUtilities::ResolveAliases(DataSetVariable, ResolveAliasesContext);
 		}
 
 		FNiagaraTypeDefinition BoundVarType = ParamMapVariable.GetType();
 		if (BindingSourceMode == ENiagaraBindingSource::ExplicitParticles || (InSourceMode == ENiagaraRendererSourceDataMode::Particles && BindingSourceMode == ENiagaraBindingSource::ImplicitFromSource))
 		{
-			bBindingExistsOnSource = InEmitter->CanObtainParticleAttribute(DataSetVariable, BoundVarType);
+			bBindingExistsOnSource = Emitter->CanObtainParticleAttribute(DataSetVariable, InVersionedEmitter.Version, BoundVarType);
 		}
 		else if (BindingSourceMode == ENiagaraBindingSource::ExplicitEmitter || (InSourceMode == ENiagaraRendererSourceDataMode::Emitter && BindingSourceMode == ENiagaraBindingSource::ImplicitFromSource))
 		{
-			bBindingExistsOnSource = InEmitter->CanObtainEmitterAttribute(ParamMapVariable, BoundVarType);
+			bBindingExistsOnSource = Emitter->CanObtainEmitterAttribute(ParamMapVariable, BoundVarType);
 		}
 		else if (BindingSourceMode == ENiagaraBindingSource::ExplicitSystem)
 		{
-			bBindingExistsOnSource = InEmitter->CanObtainSystemAttribute(ParamMapVariable, BoundVarType);
+			bBindingExistsOnSource = Emitter->CanObtainSystemAttribute(ParamMapVariable, BoundVarType);
 		}
 		else if (BindingSourceMode == ENiagaraBindingSource::ExplicitUser)
 		{
-			bBindingExistsOnSource = InEmitter->CanObtainUserVariable(ParamMapVariable);
+			bBindingExistsOnSource = Emitter->CanObtainUserVariable(ParamMapVariable);
 		}
 
 		if (bBindingExistsOnSource && BoundVarType != ParamMapVariable.GetType())
@@ -1320,7 +1338,7 @@ FNiagaraVariable FNiagaraUtilities::ResolveAliases(const FNiagaraVariable& InVar
 }
 
 #if WITH_EDITORONLY_DATA
-void FNiagaraUtilities::PrepareRapidIterationParameters(const TArray<UNiagaraScript*>& Scripts, const TMap<UNiagaraScript*, UNiagaraScript*>& ScriptDependencyMap, const TMap<UNiagaraScript*, const UNiagaraEmitter*>& ScriptToEmitterMap)
+void FNiagaraUtilities::PrepareRapidIterationParameters(const TArray<UNiagaraScript*>& Scripts, const TMap<UNiagaraScript*, UNiagaraScript*>& ScriptDependencyMap, const TMap<UNiagaraScript*, FVersionedNiagaraEmitter>& ScriptToEmitterMap)
 {
 	SCOPE_CYCLE_COUNTER(STAT_Niagara_Utilities_PrepareRapidIterationParameters);
 
@@ -1333,7 +1351,7 @@ void FNiagaraUtilities::PrepareRapidIterationParameters(const TArray<UNiagaraScr
 		Script->RapidIterationParameters.CopyParametersTo(ParameterStoreToPrepare, false, FNiagaraParameterStore::EDataInterfaceCopyMethod::None);
 		ParameterStoreToPrepare.ParameterGuidMapping = Script->RapidIterationParameters.ParameterGuidMapping;
 		checkf(ScriptToEmitterMap.Find(Script) != nullptr, TEXT("Script to emitter name map must have an entry for each script to be processed."));
-		if (const UNiagaraEmitter* const* Emitter = ScriptToEmitterMap.Find(Script))
+		if (const FVersionedNiagaraEmitter* Emitter = ScriptToEmitterMap.Find(Script))
 		{
 			if (Script->GetLatestSource())
 			{
diff --git a/Engine/Plugins/FX/Niagara/Source/Niagara/Private/NiagaraComponentRendererProperties.cpp b/Engine/Plugins/FX/Niagara/Source/Niagara/Private/NiagaraComponentRendererProperties.cpp
--- a/Engine/Plugins/FX/Niagara/Source/Niagara/Private/NiagaraComponentRendererProperties.cpp
+++ b/Engine/Plugins/FX/Niagara/Source/Niagara/Private/NiagaraComponentRendererProperties.cpp
@@ -242,8 +242,8 @@ void UNiagaraComponentRendererProperties::PostLoad()
 
 void UNiagaraComponentRendererProperties::UpdateSourceModeDerivates(ENiagaraRendererSourceDataMode InSourceMode, bool bFromPropertyEdit)
 {
-	UNiagaraEmitter* SrcEmitter = GetTypedOuter<UNiagaraEmitter>();
-	if (SrcEmitter)
+	FVersionedNiagaraEmitter SrcEmitter = GetOuterEmitter();
+	if (SrcEmitter.Emitter)
 	{
 		EnabledBinding.CacheValues(SrcEmitter, InSourceMode);
 		RendererVisibilityTagBinding.CacheValues(SrcEmitter, InSourceMode);
@@ -485,7 +485,8 @@ void UNiagaraComponentRendererProperties::CreateTemplateComponent()
 	TemplateComponent->SetComponentTickEnabled(false);
 
 	// set some defaults on the component
-	bool IsWorldSpace = EmitterPtr ? !EmitterPtr->bLocalSpace : true;
+	FVersionedNiagaraEmitterData* EmitterData = EmitterPtr.GetEmitterData();
+	bool IsWorldSpace = EmitterData ? !EmitterData->bLocalSpace : true;
 	TemplateComponent->SetAbsolute(IsWorldSpace, IsWorldSpace, IsWorldSpace);
 }
 
@@ -570,7 +571,7 @@ void UNiagaraComponentRendererProperties::GetRendererTooltipWidgets(const FNiaga
 	OutWidgets.Add(Tooltip);
 }
 
-void UNiagaraComponentRendererProperties::GetRendererFeedback(UNiagaraEmitter* InEmitter,	TArray<FNiagaraRendererFeedback>& OutErrors, TArray<FNiagaraRendererFeedback>& OutWarnings,	TArray<FNiagaraRendererFeedback>& OutInfo) const
+void UNiagaraComponentRendererProperties::GetRendererFeedback(const FVersionedNiagaraEmitter& InEmitter, TArray<FNiagaraRendererFeedback>& OutErrors, TArray<FNiagaraRendererFeedback>& OutWarnings, TArray<FNiagaraRendererFeedback>& OutInfo) const
 {
 	OutInfo.Add(FNiagaraRendererFeedback(FText::FromString(TEXT("The component renderer is still a very experimental feature that offers great flexibility, \nbut is *not* optimized for performance or safety. \nWith great power comes great responsibility."))));
 
@@ -581,10 +582,10 @@ void UNiagaraComponentRendererProperties::GetRendererFeedback(UNiagaraEmitter* I
 		OutErrors.Add(FNiagaraRendererFeedback(ErrorDescription, ErrorSummary));
 	}
 
-	if (InEmitter && TemplateComponent)
+	FVersionedNiagaraEmitterData* EmitterData = EmitterPtr.GetEmitterData();
+	if (InEmitter.Emitter && TemplateComponent && EmitterData)
 	{
-		const UNiagaraSettings* Settings = GetDefault<UNiagaraSettings>();
-		if (Settings)
+		if (const UNiagaraSettings* Settings = GetDefault<UNiagaraSettings>())
 		{
 			for (const TPair<FString, FText>& Pair : Settings->ComponentRendererWarningsPerClass)
 			{
@@ -596,8 +597,8 @@ void UNiagaraComponentRendererProperties::GetRendererFeedback(UNiagaraEmitter* I
 			}
 		}
 
-		bool IsWorldSpace = !InEmitter->bLocalSpace;
-		FNiagaraRendererFeedbackFix LocalspaceFix = FNiagaraRendererFeedbackFix::CreateLambda([InEmitter]() {	InEmitter->bLocalSpace = !InEmitter->bLocalSpace; });
+		bool IsWorldSpace = !EmitterData->bLocalSpace;
+		FNiagaraRendererFeedbackFix LocalspaceFix = FNiagaraRendererFeedbackFix::CreateLambda([EmitterData]() {	EmitterData->bLocalSpace = !EmitterData->bLocalSpace; });
 		if (TemplateComponent->IsUsingAbsoluteLocation() != IsWorldSpace && !HasPropertyBinding(FName("bAbsoluteLocation")))
 		{
 			FText ErrorDescription = LOCTEXT("NiagaraComponentLocalspaceLocationWarning", "The component location is configured to use a different localspace setting than the emitter.");
diff --git a/Engine/Plugins/FX/Niagara/Source/Niagara/Private/NiagaraDataInterface.cpp b/Engine/Plugins/FX/Niagara/Source/Niagara/Private/NiagaraDataInterface.cpp
--- a/Engine/Plugins/FX/Niagara/Source/Niagara/Private/NiagaraDataInterface.cpp
+++ b/Engine/Plugins/FX/Niagara/Source/Niagara/Private/NiagaraDataInterface.cpp
@@ -60,7 +60,7 @@ void UNiagaraDataInterface::ModifyCompilationEnvironment(EShaderPlatform ShaderP
 }
 #endif
 
-void UNiagaraDataInterface::GetAssetTagsForContext(const UObject* InAsset, const TArray<const UNiagaraDataInterface*>& InProperties, TMap<FName, uint32>& NumericKeys, TMap<FName, FString>& StringKeys) const
+void UNiagaraDataInterface::GetAssetTagsForContext(const UObject* InAsset,  FGuid AssetVersion, const TArray<const UNiagaraDataInterface*>& InProperties, TMap<FName, uint32>& NumericKeys, TMap<FName, FString>& StringKeys) const
 {
 	UClass* Class = GetClass();
 
diff --git a/Engine/Plugins/FX/Niagara/Source/Niagara/Private/NiagaraDataInterfaceCollisionQuery.cpp b/Engine/Plugins/FX/Niagara/Source/Niagara/Private/NiagaraDataInterfaceCollisionQuery.cpp
--- a/Engine/Plugins/FX/Niagara/Source/Niagara/Private/NiagaraDataInterfaceCollisionQuery.cpp
+++ b/Engine/Plugins/FX/Niagara/Source/Niagara/Private/NiagaraDataInterfaceCollisionQuery.cpp
@@ -115,7 +115,7 @@ void UNiagaraDataInterfaceCollisionQuery::PostInitProperties()
 	}
 }
 
-void UNiagaraDataInterfaceCollisionQuery::GetAssetTagsForContext(const UObject* InAsset, const TArray<const UNiagaraDataInterface*>& InProperties, TMap<FName, uint32>& NumericKeys, TMap<FName, FString>& StringKeys) const
+void UNiagaraDataInterfaceCollisionQuery::GetAssetTagsForContext(const UObject* InAsset, FGuid AssetVersion, const TArray<const UNiagaraDataInterface*>& InProperties, TMap<FName, uint32>& NumericKeys, TMap<FName, FString>& StringKeys) const
 {
 #if WITH_EDITOR
 	const UNiagaraSystem* System = Cast<UNiagaraSystem>(InAsset);
@@ -131,7 +131,7 @@ void UNiagaraDataInterfaceCollisionQuery::GetAssetTagsForContext(const UObject*
 		Scripts.Add(System->GetSystemUpdateScript());
 		for (auto&& EmitterHandle : System->GetEmitterHandles())
 		{
-			const UNiagaraEmitter* HandleEmitter = EmitterHandle.GetInstance();
+			FVersionedNiagaraEmitterData* HandleEmitter = EmitterHandle.GetEmitterData();
 			if (HandleEmitter)
 			{
 				if (HandleEmitter->SimTarget == ENiagaraSimTarget::GPUComputeSim)
@@ -147,10 +147,11 @@ void UNiagaraDataInterfaceCollisionQuery::GetAssetTagsForContext(const UObject*
 	}
 	if (Emitter)
 	{
-		if (Emitter->SimTarget != ENiagaraSimTarget::GPUComputeSim)
+		const FVersionedNiagaraEmitterData* EmitterData = Emitter->GetEmitterData(AssetVersion);
+		if (EmitterData && EmitterData->SimTarget != ENiagaraSimTarget::GPUComputeSim)
 		{
 			TArray<UNiagaraScript*> OutScripts;
-			Emitter->GetScripts(OutScripts, false);
+			EmitterData->GetScripts(OutScripts, false);
 			Scripts.Append(OutScripts);
 		}
 	}
@@ -193,8 +194,7 @@ void UNiagaraDataInterfaceCollisionQuery::GetAssetTagsForContext(const UObject*
 #endif
 	
 	// Make sure and get the base implementation tags
-	Super::GetAssetTagsForContext(InAsset, InProperties, NumericKeys, StringKeys);
-	
+	Super::GetAssetTagsForContext(InAsset, AssetVersion, InProperties, NumericKeys, StringKeys);
 }
 
 void UNiagaraDataInterfaceCollisionQuery::GetFunctions(TArray<FNiagaraFunctionSignature>& OutFunctions)
diff --git a/Engine/Plugins/FX/Niagara/Source/Niagara/Private/NiagaraDataInterfaceCurveBase.cpp b/Engine/Plugins/FX/Niagara/Source/Niagara/Private/NiagaraDataInterfaceCurveBase.cpp
--- a/Engine/Plugins/FX/Niagara/Source/Niagara/Private/NiagaraDataInterfaceCurveBase.cpp
+++ b/Engine/Plugins/FX/Niagara/Source/Niagara/Private/NiagaraDataInterfaceCurveBase.cpp
@@ -2,8 +2,6 @@
 
 #include "NiagaraDataInterfaceCurveBase.h"
 #include "Curves/CurveVector.h"
-#include "Curves/CurveLinearColor.h"
-#include "Curves/CurveFloat.h"
 #include "Engine/Texture2D.h"
 #include "Internationalization/Internationalization.h"
 #include "ShaderParameterUtils.h"
@@ -430,19 +428,10 @@ TArray<FNiagaraDataInterfaceError> UNiagaraDataInterfaceCurveBase::GetErrors()
 {
 	// Trace down the root emitter (if there is one)
 	TArray<FNiagaraDataInterfaceError> Errors;
-	UObject* Obj = GetOuter();
-	UNiagaraEmitter* Emitter = nullptr;
-	while (Obj && !Obj->IsA(UPackage::StaticClass()))
-	{
-		Emitter = Cast<UNiagaraEmitter>(Obj);
-		if (!Emitter)
-			Obj = Obj->GetOuter();
-		else
-			break;
-	}
+	UNiagaraEmitter* Emitter = GetTypedOuter<UNiagaraEmitter>();
 
 	// If there is a root emitter, assume that we are in its particle stack and point out that we need bUseLUT true for GPU sims.
-	if (Emitter && Emitter->SimTarget == ENiagaraSimTarget::GPUComputeSim && !bUseLUT)
+	if (Emitter && Emitter->GetLatestEmitterData()->SimTarget == ENiagaraSimTarget::GPUComputeSim && !bUseLUT)
 	{
 		FNiagaraDataInterfaceError LUTsNeededForGPUSimsError(LOCTEXT("LUTsNeededForGPUSims", "This Data Interface must have bUseLUT set to true for GPU sims."),
 			LOCTEXT("LUTsNeededForGPUSimsSummary", "bUseLUT Required"),
diff --git a/Engine/Plugins/FX/Niagara/Source/Niagara/Private/NiagaraDataInterfaceEmitterBinding.cpp b/Engine/Plugins/FX/Niagara/Source/Niagara/Private/NiagaraDataInterfaceEmitterBinding.cpp
--- a/Engine/Plugins/FX/Niagara/Source/Niagara/Private/NiagaraDataInterfaceEmitterBinding.cpp
+++ b/Engine/Plugins/FX/Niagara/Source/Niagara/Private/NiagaraDataInterfaceEmitterBinding.cpp
@@ -13,7 +13,7 @@ FNiagaraEmitterInstance* FNiagaraDataInterfaceEmitterBinding::Resolve(FNiagaraSy
 		{
 			for (TSharedPtr<FNiagaraEmitterInstance, ESPMode::ThreadSafe> EmitterInstance : SystemInstance->GetEmitters())
 			{
-				if (EmitterInstance->GetCachedEmitter() == OwnerEmitter)
+				if (EmitterInstance->GetCachedEmitter().Emitter == OwnerEmitter)
 				{
 					return EmitterInstance.Get();
 				}
@@ -53,7 +53,7 @@ FNiagaraEmitterInstance* FNiagaraDataInterfaceEmitterBinding::Resolve(FNiagaraSy
 					SelfEmitterView = SelfEmitterView.Mid(0, DotIndex);
 					for (TSharedPtr<FNiagaraEmitterInstance, ESPMode::ThreadSafe> EmitterInstance : SystemInstance->GetEmitters())
 					{
-						if (const UNiagaraEmitter* CachedEmitter = EmitterInstance->GetCachedEmitter())
+						if (const UNiagaraEmitter* CachedEmitter = EmitterInstance->GetCachedEmitter().Emitter)
 						{
 							if (SelfEmitterView.Equals(CachedEmitter->GetUniqueEmitterName()))
 							{
@@ -86,7 +86,7 @@ FNiagaraEmitterInstance* FNiagaraDataInterfaceEmitterBinding::Resolve(FNiagaraSy
 
 		for (TSharedPtr<FNiagaraEmitterInstance, ESPMode::ThreadSafe> EmitterInstance : SystemInstance->GetEmitters())
 		{
-			if (const UNiagaraEmitter* CachedEmitter = EmitterInstance->GetCachedEmitter())
+			if (const UNiagaraEmitter* CachedEmitter = EmitterInstance->GetCachedEmitter().Emitter)
 			{
 				//-TODO: UniqueEmitterName should probably be a FName?
 				if (EmitterNameStringView.Equals(CachedEmitter->GetUniqueEmitterName()) )
@@ -119,7 +119,7 @@ UNiagaraEmitter* FNiagaraDataInterfaceEmitterBinding::Resolve(UNiagaraSystem* Ni
 
 		for ( const FNiagaraEmitterHandle& EmitterHandle : NiagaraSystem->GetEmitterHandles() )
 		{
-			if ( UNiagaraEmitter* NiagaraEmitter = EmitterHandle.GetInstance() )
+			if (UNiagaraEmitter* NiagaraEmitter = EmitterHandle.GetInstance().Emitter)
 			{
 				//-TODO: UniqueEmitterName should probably be a FName?
 				if (EmitterNameStringView.Equals(NiagaraEmitter->GetUniqueEmitterName()))
diff --git a/Engine/Plugins/FX/Niagara/Source/Niagara/Private/NiagaraDataInterfaceGrid2DCollection.cpp b/Engine/Plugins/FX/Niagara/Source/Niagara/Private/NiagaraDataInterfaceGrid2DCollection.cpp
--- a/Engine/Plugins/FX/Niagara/Source/Niagara/Private/NiagaraDataInterfaceGrid2DCollection.cpp
+++ b/Engine/Plugins/FX/Niagara/Source/Niagara/Private/NiagaraDataInterfaceGrid2DCollection.cpp
@@ -11,7 +11,6 @@
 #include "NiagaraSystemInstance.h"
 
 #include "Engine/Canvas.h"
-#include "Engine/Texture2DArray.h"
 #include "Engine/TextureRenderTarget2D.h"
 #include "Engine/TextureRenderTarget2DArray.h"
 #include "ClearQuad.h"
@@ -168,7 +167,6 @@ struct FNiagaraGrid2DLegacyTiled2DInfo
 		: NumAttributes(InNumAttributes)
 		, NumCells(InNumCells)
 	{
-		const int MaxTextureDim = GMaxTextureDimensions;
 		const int MaxTilesX = FMath::DivideAndRoundDown<int>(GMaxTextureDimensions, NumCells.X);
 		const int MaxTilesY = FMath::DivideAndRoundDown<int>(GMaxTextureDimensions, NumCells.Y);
 		const int MaxAttributes = MaxTilesX * MaxTilesY;
@@ -220,7 +218,10 @@ struct FNiagaraGrid2DLegacyTiled2DInfo
 };
 
 UNiagaraDataInterfaceGrid2DCollection::UNiagaraDataInterfaceGrid2DCollection(FObjectInitializer const& ObjectInitializer)
-	: Super(ObjectInitializer)
+	: Super(ObjectInitializer), OverrideBufferFormat(ENiagaraGpuBufferFormat::Float), bOverrideFormat(false)
+#if WITH_EDITORONLY_DATA
+	, bPreviewGrid(false)
+#endif
 {
 	Proxy.Reset(new FNiagaraDataInterfaceProxyGrid2DCollectionProxy());
 
@@ -1989,10 +1990,10 @@ void UNiagaraDataInterfaceGrid2DCollection::FindAttributesByName(FName VariableN
 	int32 TotalAttributes = NumAttributes;
 	for (const FNiagaraEmitterHandle& EmitterHandle : OwnerSystem->GetEmitterHandles())
 	{
-		UNiagaraEmitter* Emitter = EmitterHandle.GetInstance();
-		if (Emitter && EmitterHandle.GetIsEnabled() && Emitter->IsValid() && (Emitter->SimTarget == ENiagaraSimTarget::GPUComputeSim))
+		FVersionedNiagaraEmitterData* EmitterData = EmitterHandle.GetEmitterData();
+		if (EmitterData && EmitterHandle.GetIsEnabled() && EmitterData->IsValid() && (EmitterData->SimTarget == ENiagaraSimTarget::GPUComputeSim))
 		{
-			CollectAttributesForScript(Emitter->GetGPUComputeScript(), VariableName, OutVariables, OutVariableOffsets, TotalAttributes, OutWarnings);
+			CollectAttributesForScript(EmitterData->GetGPUComputeScript(), VariableName, OutVariables, OutVariableOffsets, TotalAttributes, OutWarnings);
 		}
 	}
 	OutNumAttribChannelsFound = TotalAttributes - NumAttributes;
@@ -2011,18 +2012,18 @@ void UNiagaraDataInterfaceGrid2DCollection::FindAttributes(TArray<FNiagaraVariab
 	int32 TotalAttributes = NumAttributes;
 	for (const FNiagaraEmitterHandle& EmitterHandle : OwnerSystem->GetEmitterHandles())
 	{
-		UNiagaraEmitter* Emitter = EmitterHandle.GetInstance();
-		if (Emitter && EmitterHandle.GetIsEnabled() && Emitter->IsValid() && (Emitter->SimTarget == ENiagaraSimTarget::GPUComputeSim))
+		FVersionedNiagaraEmitterData* EmitterData = EmitterHandle.GetEmitterData();
+		if (EmitterData && EmitterHandle.GetIsEnabled() && EmitterData->IsValid() && (EmitterData->SimTarget == ENiagaraSimTarget::GPUComputeSim))
 		{
 			// Search scripts for this data interface so we get the variable name
 			auto FindDataInterfaceVariable =
-				[&OwnerSystem, &Emitter](const UNiagaraDataInterface* DataInterface) -> FName
+				[&OwnerSystem, &EmitterData](const UNiagaraDataInterface* DataInterface) -> FName
 				{
 					UNiagaraScript* Scripts[] =
 					{
 						OwnerSystem->GetSystemSpawnScript(),
 						OwnerSystem->GetSystemUpdateScript(),
-						Emitter->GetGPUComputeScript(),
+						EmitterData->GetGPUComputeScript(),
 					};
 
 					for (UNiagaraScript* Script : Scripts)
@@ -2041,7 +2042,7 @@ void UNiagaraDataInterfaceGrid2DCollection::FindAttributes(TArray<FNiagaraVariab
 			const FName VariableName = FindDataInterfaceVariable(this);
 			if (!VariableName.IsNone() )
 			{
-				CollectAttributesForScript(Emitter->GetGPUComputeScript(), VariableName, OutVariables, OutVariableOffsets, TotalAttributes, OutWarnings);
+				CollectAttributesForScript(EmitterData->GetGPUComputeScript(), VariableName, OutVariables, OutVariableOffsets, TotalAttributes, OutWarnings);
 			}
 		}
 	}
@@ -2125,7 +2126,6 @@ bool UNiagaraDataInterfaceGrid2DCollection::RenderVariableToCanvas(FNiagaraSyste
 	return true;
 }
 
-UFUNCTION(BlueprintCallable, Category = Niagara)
 bool UNiagaraDataInterfaceGrid2DCollection::FillTexture2D(const UNiagaraComponent *Component, UTextureRenderTarget2D *Dest, int AttributeIndex)
 {
 	if (!Component || !Dest)
diff --git a/Engine/Plugins/FX/Niagara/Source/Niagara/Private/NiagaraDataInterfaceGrid2DCollectionReader.cpp b/Engine/Plugins/FX/Niagara/Source/Niagara/Private/NiagaraDataInterfaceGrid2DCollectionReader.cpp
--- a/Engine/Plugins/FX/Niagara/Source/Niagara/Private/NiagaraDataInterfaceGrid2DCollectionReader.cpp
+++ b/Engine/Plugins/FX/Niagara/Source/Niagara/Private/NiagaraDataInterfaceGrid2DCollectionReader.cpp
@@ -3,7 +3,6 @@
 #include "NiagaraShader.h"
 #include "ShaderParameterUtils.h"
 #include "ClearQuad.h"
-#include "TextureResource.h"
 #include "NiagaraSystemInstance.h"
 #include "NiagaraRenderer.h"
 #include "Engine/VolumeTexture.h"
@@ -63,7 +62,7 @@ bool UNiagaraDataInterfaceGrid2DCollectionReader::InitPerInstanceData(void* PerI
 	FNiagaraEmitterInstance* EmitterInstanceToUse = nullptr;
 	for (TSharedPtr<FNiagaraEmitterInstance, ESPMode::ThreadSafe> EmitterInstance : SystemInstance->GetEmitters())
 	{
-		UNiagaraEmitter* Emitter = EmitterInstance->GetCachedEmitter();
+		UNiagaraEmitter* Emitter = EmitterInstance->GetCachedEmitter().Emitter;
 		if (Emitter == nullptr)
 		{
 			continue;
@@ -131,18 +130,17 @@ bool UNiagaraDataInterfaceGrid2DCollectionReader::InitPerInstanceData(void* PerI
 	return false;
 }
 
-void UNiagaraDataInterfaceGrid2DCollectionReader::GetEmitterDependencies(UNiagaraSystem* Asset, TArray<UNiagaraEmitter*>& Dependencies) const
+void UNiagaraDataInterfaceGrid2DCollectionReader::GetEmitterDependencies(UNiagaraSystem* Asset, TArray<FVersionedNiagaraEmitter>& Dependencies) const
 {
 	if (!Asset)
 	{
 		return;
 	}
 
-	UNiagaraEmitter* FoundSourceEmitter = nullptr;
 	for (const FNiagaraEmitterHandle& EmitterHandle : Asset->GetEmitterHandles())
 	{
-		UNiagaraEmitter* EmitterInstance = EmitterHandle.GetInstance();
-		if (EmitterInstance && EmitterInstance->GetUniqueEmitterName() == EmitterName)
+		FVersionedNiagaraEmitter EmitterInstance = EmitterHandle.GetInstance();
+		if (EmitterInstance.Emitter && EmitterInstance.Emitter->GetUniqueEmitterName() == EmitterName)
 		{
 			Dependencies.Add(EmitterInstance);
 			return;
diff --git a/Engine/Plugins/FX/Niagara/Source/Niagara/Private/NiagaraDataInterfaceGrid3DCollection.cpp b/Engine/Plugins/FX/Niagara/Source/Niagara/Private/NiagaraDataInterfaceGrid3DCollection.cpp
--- a/Engine/Plugins/FX/Niagara/Source/Niagara/Private/NiagaraDataInterfaceGrid3DCollection.cpp
+++ b/Engine/Plugins/FX/Niagara/Source/Niagara/Private/NiagaraDataInterfaceGrid3DCollection.cpp
@@ -168,6 +168,11 @@ static FAutoConsoleVariableRef CVarNiagaraGrid3DOverrideFormat(
 UNiagaraDataInterfaceGrid3DCollection::UNiagaraDataInterfaceGrid3DCollection(FObjectInitializer const& ObjectInitializer)
 	: Super(ObjectInitializer)
 	, NumAttributes(1)
+	, OverrideBufferFormat(ENiagaraGpuBufferFormat::Float)
+	, bOverrideFormat(false)
+#if WITH_EDITORONLY_DATA
+	, bPreviewGrid(false)
+#endif
 {
 	Proxy.Reset(new FNiagaraDataInterfaceProxyGrid3DCollectionProxy());
 
@@ -2194,10 +2199,10 @@ void UNiagaraDataInterfaceGrid3DCollection::FindAttributesByName(FName VariableN
 	int32 TotalAttributes = NumAttributes;
 	for (const FNiagaraEmitterHandle& EmitterHandle : OwnerSystem->GetEmitterHandles())
 	{
-		UNiagaraEmitter* Emitter = EmitterHandle.GetInstance();
-		if (Emitter && EmitterHandle.GetIsEnabled() && Emitter->IsValid() && (Emitter->SimTarget == ENiagaraSimTarget::GPUComputeSim))
+		FVersionedNiagaraEmitterData* EmitterData = EmitterHandle.GetEmitterData();
+		if (EmitterData && EmitterHandle.GetIsEnabled() && EmitterData->IsValid() && (EmitterData->SimTarget == ENiagaraSimTarget::GPUComputeSim))
 		{
-			CollectAttributesForScript(Emitter->GetGPUComputeScript(), VariableName, OutVariables, OutVariableOffsets, TotalAttributes, OutWarnings);
+			CollectAttributesForScript(EmitterData->GetGPUComputeScript(), VariableName, OutVariables, OutVariableOffsets, TotalAttributes, OutWarnings);
 		}
 	}
 	OutNumAttribChannelsFound = TotalAttributes - NumAttributes;
@@ -2216,18 +2221,18 @@ void UNiagaraDataInterfaceGrid3DCollection::FindAttributes(TArray<FNiagaraVariab
 	int32 TotalAttributes = NumAttributes;
 	for (const FNiagaraEmitterHandle& EmitterHandle : OwnerSystem->GetEmitterHandles())
 	{
-		UNiagaraEmitter* Emitter = EmitterHandle.GetInstance();
-		if (Emitter && EmitterHandle.GetIsEnabled() && Emitter->IsValid() && (Emitter->SimTarget == ENiagaraSimTarget::GPUComputeSim))
+		FVersionedNiagaraEmitterData* EmitterData = EmitterHandle.GetEmitterData();
+		if (EmitterData && EmitterHandle.GetIsEnabled() && EmitterData->IsValid() && (EmitterData->SimTarget == ENiagaraSimTarget::GPUComputeSim))
 		{
 			// Search scripts for this data interface so we get the variable name
 			auto FindDataInterfaceVariable =
-				[&OwnerSystem, &Emitter](const UNiagaraDataInterface* DataInterface) -> FName
+				[&OwnerSystem, &EmitterData](const UNiagaraDataInterface* DataInterface) -> FName
 				{
 					UNiagaraScript* Scripts[] =
 					{
 						OwnerSystem->GetSystemSpawnScript(),
 						OwnerSystem->GetSystemUpdateScript(),
-						Emitter->GetGPUComputeScript(),
+						EmitterData->GetGPUComputeScript(),
 					};
 
 					for (UNiagaraScript* Script : Scripts)
@@ -2246,14 +2251,13 @@ void UNiagaraDataInterfaceGrid3DCollection::FindAttributes(TArray<FNiagaraVariab
 			const FName VariableName = FindDataInterfaceVariable(this);
 			if (!VariableName.IsNone() )
 			{
-				CollectAttributesForScript(Emitter->GetGPUComputeScript(), VariableName, OutVariables, OutVariableOffsets, TotalAttributes, OutWarnings);
+				CollectAttributesForScript(EmitterData->GetGPUComputeScript(), VariableName, OutVariables, OutVariableOffsets, TotalAttributes, OutWarnings);
 			}
 		}
 	}
 	OutNumAttribChannelsFound = TotalAttributes - NumAttributes;
 }
 
-UFUNCTION(BlueprintCallable, Category = Niagara)
 bool UNiagaraDataInterfaceGrid3DCollection::FillVolumeTexture(const UNiagaraComponent *Component, UVolumeTexture*Dest, int AttributeIndex)
 {
 	/*
@@ -2314,7 +2318,6 @@ bool UNiagaraDataInterfaceGrid3DCollection::FillVolumeTexture(const UNiagaraComp
 	return false;
 }
 
-UFUNCTION(BlueprintCallable, Category = Niagara)
 bool UNiagaraDataInterfaceGrid3DCollection::FillRawVolumeTexture(const UNiagaraComponent *Component, UVolumeTexture *Dest, int &TilesX, int &TilesY, int&TilesZ)
 {
 	/*
@@ -2383,7 +2386,6 @@ bool UNiagaraDataInterfaceGrid3DCollection::FillRawVolumeTexture(const UNiagaraC
 	return false;
 }
 
-UFUNCTION(BlueprintCallable, Category = Niagara)
 void UNiagaraDataInterfaceGrid3DCollection::GetRawTextureSize(const UNiagaraComponent *Component, int &SizeX, int &SizeY, int &SizeZ)
 {
 	if (!Component)
@@ -2418,7 +2420,6 @@ void UNiagaraDataInterfaceGrid3DCollection::GetRawTextureSize(const UNiagaraComp
 	SizeZ = Grid3DInstanceData->NumCells.Z * Grid3DInstanceData->NumTiles.Z;
 }
 
-UFUNCTION(BlueprintCallable, Category = Niagara)
 void UNiagaraDataInterfaceGrid3DCollection::GetTextureSize(const UNiagaraComponent *Component, int &SizeX, int &SizeY, int &SizeZ)
 {
 	if (!Component)
diff --git a/Engine/Plugins/FX/Niagara/Source/Niagara/Private/NiagaraDataInterfaceGrid3DCollectionReader.cpp b/Engine/Plugins/FX/Niagara/Source/Niagara/Private/NiagaraDataInterfaceGrid3DCollectionReader.cpp
--- a/Engine/Plugins/FX/Niagara/Source/Niagara/Private/NiagaraDataInterfaceGrid3DCollectionReader.cpp
+++ b/Engine/Plugins/FX/Niagara/Source/Niagara/Private/NiagaraDataInterfaceGrid3DCollectionReader.cpp
@@ -2,10 +2,7 @@
 #include "NiagaraDataInterfaceGrid3DCollectionReader.h"
 #include "NiagaraShader.h"
 #include "ShaderParameterUtils.h"
-#include "ClearQuad.h"
-#include "TextureResource.h"
 #include "NiagaraSystemInstance.h"
-#include "NiagaraRenderer.h"
 #include "Engine/VolumeTexture.h"
 #include "Engine/TextureRenderTargetVolume.h"
 #include "NiagaraSettings.h"
@@ -63,7 +60,7 @@ bool UNiagaraDataInterfaceGrid3DCollectionReader::InitPerInstanceData(void* PerI
 	FNiagaraEmitterInstance* EmitterInstanceToUse = nullptr;
 	for (TSharedPtr<FNiagaraEmitterInstance, ESPMode::ThreadSafe> EmitterInstance : SystemInstance->GetEmitters())
 	{
-		UNiagaraEmitter* Emitter = EmitterInstance->GetCachedEmitter();
+		UNiagaraEmitter* Emitter = EmitterInstance->GetCachedEmitter().Emitter;
 		if (Emitter == nullptr)
 		{
 			continue;
@@ -131,27 +128,24 @@ bool UNiagaraDataInterfaceGrid3DCollectionReader::InitPerInstanceData(void* PerI
 	return false;
 }
 
-void UNiagaraDataInterfaceGrid3DCollectionReader::GetEmitterDependencies(UNiagaraSystem* Asset, TArray<UNiagaraEmitter*>& Dependencies) const
+void UNiagaraDataInterfaceGrid3DCollectionReader::GetEmitterDependencies(UNiagaraSystem* Asset, TArray<FVersionedNiagaraEmitter>& Dependencies) const
 {
 	if (!Asset)
 	{
 		return;
 	}
 
-	UNiagaraEmitter* FoundSourceEmitter = nullptr;
 	for (const FNiagaraEmitterHandle& EmitterHandle : Asset->GetEmitterHandles())
 	{
-		UNiagaraEmitter* EmitterInstance = EmitterHandle.GetInstance();
+		UNiagaraEmitter* EmitterInstance = EmitterHandle.GetInstance().Emitter;
 		if (EmitterInstance && EmitterInstance->GetUniqueEmitterName() == EmitterName)
 		{
-			Dependencies.Add(EmitterInstance);
+			Dependencies.Add(EmitterHandle.GetInstance());
 			return;
 		}
 	}
 }
 
-
-
 void UNiagaraDataInterfaceGrid3DCollectionReader::GetFunctions(TArray<FNiagaraFunctionSignature>& OutFunctions)
 {
 	int32 NumFunctionsBefore = OutFunctions.Num();
diff --git a/Engine/Plugins/FX/Niagara/Source/Niagara/Private/NiagaraDataInterfaceParticleRead.cpp b/Engine/Plugins/FX/Niagara/Source/Niagara/Private/NiagaraDataInterfaceParticleRead.cpp
--- a/Engine/Plugins/FX/Niagara/Source/Niagara/Private/NiagaraDataInterfaceParticleRead.cpp
+++ b/Engine/Plugins/FX/Niagara/Source/Niagara/Private/NiagaraDataInterfaceParticleRead.cpp
@@ -667,13 +667,11 @@ bool UNiagaraDataInterfaceParticleRead::InitPerInstanceData(void* PerInstanceDat
 	PIData->EmitterInstance = nullptr;
 	for (TSharedPtr<FNiagaraEmitterInstance, ESPMode::ThreadSafe> EmitterInstance : SystemInstance->GetEmitters())
 	{
-		if (const UNiagaraEmitter* CachedEmitter = EmitterInstance->GetCachedEmitter())
+		const FVersionedNiagaraEmitter& CachedEmitter = EmitterInstance->GetCachedEmitter();
+		if (CachedEmitter.Emitter && (EmitterName == CachedEmitter.Emitter->GetUniqueEmitterName()))
 		{
-			if (EmitterName == CachedEmitter->GetUniqueEmitterName())
-			{
-				PIData->EmitterInstance = EmitterInstance.Get();
-				break;
-			}
+			PIData->EmitterInstance = EmitterInstance.Get();
+			break;
 		}
 	}
 
@@ -1418,7 +1416,7 @@ void UNiagaraDataInterfaceParticleRead::GetVMExternalFunction(const FVMExternalF
 	if (!bBindSuccessful)
 	{
 		UNiagaraSystem* NiagaraSystem = PIData->SystemInstance ? PIData->SystemInstance->GetSystem() : nullptr;
-		UNiagaraEmitter* NiagaraEmitter = PIData->EmitterInstance ? PIData->EmitterInstance->GetCachedEmitter() : nullptr;
+		UNiagaraEmitter* NiagaraEmitter = PIData->EmitterInstance ? PIData->EmitterInstance->GetCachedEmitter().Emitter : nullptr;
 		UE_LOG(LogNiagara, Warning, TEXT("ParticleRead: Failed to '%s' attribute '%s' for System '%s' Emitter '%s'! Check that the attribute is named correctly."), *BindingInfo.Name.ToString(), *AttributeToRead.ToString(), *GetNameSafe(NiagaraSystem), *GetNameSafe(NiagaraEmitter));
 	}
 }
@@ -1429,8 +1427,8 @@ void UNiagaraDataInterfaceParticleRead::VMGetLocalSpace(FVectorVMExternalFunctio
 	FNDIOutputParam<bool> OutIsLocalSpace(Context);
 
 	const FNiagaraEmitterInstance* EmitterInstance = InstData.Get()->EmitterInstance;
-	UNiagaraEmitter* NiagaraEmitter = EmitterInstance ? EmitterInstance->GetCachedEmitter() : nullptr;
-	const bool bIsLocalSpace = NiagaraEmitter ? NiagaraEmitter->bLocalSpace : false;
+	FVersionedNiagaraEmitterData* EmitterData = EmitterInstance ? EmitterInstance->GetCachedEmitterData() : nullptr;
+	const bool bIsLocalSpace = EmitterData ? EmitterData->bLocalSpace : false;
 
 	for (int32 InstanceIdx = 0; InstanceIdx < Context.GetNumInstances(); ++InstanceIdx)
 	{
@@ -2421,18 +2419,18 @@ void UNiagaraDataInterfaceParticleRead::GetFeedback(UNiagaraSystem* Asset, UNiag
 		return;
 	}
 
-	UNiagaraEmitter* FoundSourceEmitter = nullptr;
+	FVersionedNiagaraEmitter FoundSourceEmitter;
 	for (const FNiagaraEmitterHandle& EmitterHandle : Asset->GetEmitterHandles())
 	{
-		UNiagaraEmitter* EmitterInstance = EmitterHandle.GetInstance();
-		if (EmitterInstance && EmitterInstance->GetUniqueEmitterName() == EmitterName)
+		FVersionedNiagaraEmitter EmitterInstance = EmitterHandle.GetInstance();
+		if (EmitterInstance.Emitter && EmitterInstance.Emitter->GetUniqueEmitterName() == EmitterName)
 		{
 			FoundSourceEmitter = EmitterInstance;
 			break; 
 		}
 	}
 
-	if (!FoundSourceEmitter)
+	if (!FoundSourceEmitter.Emitter)
 	{
 		Warnings.Emplace(
 			LOCTEXT("SourceEmitterNotFound", "Source emitter was not found."),
@@ -2448,7 +2446,7 @@ void UNiagaraDataInterfaceParticleRead::GetFeedback(UNiagaraSystem* Asset, UNiag
 	for (auto&& EmitterHandle : Asset->GetEmitterHandles())
 	{
 		TArray<UNiagaraScript*> OutScripts;
-		EmitterHandle.GetInstance()->GetScripts(OutScripts, false);
+		EmitterHandle.GetEmitterData()->GetScripts(OutScripts, false);
 		Scripts.Append(OutScripts);
 	}
 
@@ -2496,12 +2494,12 @@ void UNiagaraDataInterfaceParticleRead::GetFeedback(UNiagaraSystem* Asset, UNiag
 				if (CachedDefaultDIs.IsValidIndex(Idx) && CachedDefaultDIs[Idx].DataInterface != nullptr && CachedDefaultDIs[Idx].DataInterface->Equals(this))
 				{
 					bMatchFound = true;
-					UNiagaraEmitter* OuterEmitter = Script->GetTypedOuter<UNiagaraEmitter>();
-					if (OuterEmitter && FoundSourceEmitter)
+					FVersionedNiagaraEmitter OuterEmitter = Script->GetOuterEmitter();
+					if (OuterEmitter.GetEmitterData() && FoundSourceEmitter.GetEmitterData())
 					{
-						if (OuterEmitter->SimTarget != FoundSourceEmitter->SimTarget)
+						if (OuterEmitter.GetEmitterData()->SimTarget != FoundSourceEmitter.GetEmitterData()->SimTarget)
 						{
-							FText Msg = FText::Format(LOCTEXT("SourceEmitterSimTypeMismatchError", "Emitter \"{0}\" SimTarget not compatible (CPU vs GPU)!"), FText::FromName(OuterEmitter->GetFName()));
+							FText Msg = FText::Format(LOCTEXT("SourceEmitterSimTypeMismatchError", "Emitter \"{0}\" SimTarget not compatible (CPU vs GPU)!"), FText::FromName(OuterEmitter.Emitter->GetFName()));
 							FNiagaraDataInterfaceError SourceEmitterNotFoundError(
 								Msg, Msg,
 								FNiagaraDataInterfaceFix());
@@ -2523,16 +2521,16 @@ void UNiagaraDataInterfaceParticleRead::GetFeedback(UNiagaraSystem* Asset, UNiag
 							AttributeType = GetValueTypeFromFuncName(Func.Name);
 						}
 
-						if (AttributeName && FoundSourceEmitter)
+						if (AttributeName && FoundSourceEmitter.GetEmitterData())
 						{
 							if (AttributeType != ENiagaraParticleDataValueType::Invalid)
 							{
-								auto AttribFilter = [&Func, AttributeName](const FNiagaraVariable& Var)
+								auto AttribFilter = [AttributeName](const FNiagaraVariable& Var)
 								{
 									return Var.GetName() == *AttributeName;
 								};
 
-								const FNiagaraVariableBase* FoundVar = FoundSourceEmitter->SpawnScriptProps.Script->GetVMExecutableData().Attributes.FindByPredicate(AttribFilter);
+								const FNiagaraVariableBase* FoundVar = FoundSourceEmitter.GetEmitterData()->SpawnScriptProps.Script->GetVMExecutableData().Attributes.FindByPredicate(AttribFilter);
 								if (FoundVar && !CheckVariableType(FoundVar->GetType(), AttributeType))
 								{
 									FText Msg = FText::Format(LOCTEXT("SourceEmitterTypeMismatchError", "Source Emitter has attribute named, \"{0}\" but the type isn't compatible with the function \"{1}\", and will not succeed."), FText::FromName(*AttributeName), FText::FromName(Func.Name));
@@ -2561,37 +2559,34 @@ void UNiagaraDataInterfaceParticleRead::GetFeedback(UNiagaraSystem* Asset, UNiag
 	}
 
 	// If we found persistent ID functions in use and the target emitter isn't set to expose them, trigger a fixable warning.
-	if (bHasPersistenIDAccessWarning && FoundSourceEmitter && FoundSourceEmitter->bRequiresPersistentIDs == false)
+	FVersionedNiagaraEmitterData* EmitterData = FoundSourceEmitter.GetEmitterData();
+	if (bHasPersistenIDAccessWarning && EmitterData && EmitterData->bRequiresPersistentIDs == false)
 	{
 		FNiagaraDataInterfaceError SourceEmitterNeedsPersistentIDError(LOCTEXT("SourceEmitterNeedsPersistenIDError", "Source Emitter Needs PersistentIDs set."),
 			LOCTEXT("SourceEmitterNeedsPersistenIDErrorSummary", "Source emitter needs persistent id's set."),
 			FNiagaraDataInterfaceFix::CreateLambda([=]()
 				{
-					if (FoundSourceEmitter)
+					FName PropertyName = GET_MEMBER_NAME_CHECKED(FVersionedNiagaraEmitterData, bRequiresPersistentIDs);
+					FProperty* FoundProp = nullptr;
+					for (TFieldIterator<FProperty> PropertyIt(FVersionedNiagaraEmitterData::StaticStruct()); PropertyIt; ++PropertyIt)
 					{
-						FName PropertyName = GET_MEMBER_NAME_CHECKED(UNiagaraEmitter, bRequiresPersistentIDs);
-						FProperty* FoundProp = nullptr;
-						for (TFieldIterator<FProperty> PropertyIt(UNiagaraEmitter::StaticClass(), EFieldIteratorFlags::IncludeSuper); PropertyIt; ++PropertyIt)
+						FProperty* Property = *PropertyIt;
+						if (Property && Property->GetFName() == PropertyName)
 						{
-							FProperty* Property = *PropertyIt;
-							if (Property && Property->GetFName() == PropertyName)
-							{
-								FoundProp = Property;
-								break;
-							}
+							FoundProp = Property;
+							break;
 						}
+					}
 
-						FoundSourceEmitter->Modify();
-						if (FoundProp)
-						{
-							FPropertyChangedEvent EmptyPropertyUpdateStruct(FoundProp);
-
-							// Go through Pre/Post edit change cycle on these because changing them will invoke a recompile on the target emitter.
-							FoundSourceEmitter->PreEditChange(FoundProp);
-							FoundSourceEmitter->bRequiresPersistentIDs = true;
-							FoundSourceEmitter->PostEditChangeProperty(EmptyPropertyUpdateStruct);
-						}
+					FoundSourceEmitter.Emitter->Modify();
+					if (FoundProp)
+					{
+						FPropertyChangedEvent EmptyPropertyUpdateStruct(FoundProp);
 
+						// Go through Pre/Post edit change cycle on these because changing them will invoke a recompile on the target emitter.
+						FoundSourceEmitter.Emitter->PreEditChange(FoundProp);
+						EmitterData->bRequiresPersistentIDs = true;
+						FoundSourceEmitter.Emitter->PostEditChangeProperty(EmptyPropertyUpdateStruct);
 					}
 					return true;
 				}));
@@ -2600,18 +2595,17 @@ void UNiagaraDataInterfaceParticleRead::GetFeedback(UNiagaraSystem* Asset, UNiag
 }
 #endif
 
-void UNiagaraDataInterfaceParticleRead::GetEmitterDependencies(UNiagaraSystem* Asset, TArray<UNiagaraEmitter*>& Dependencies) const
+void UNiagaraDataInterfaceParticleRead::GetEmitterDependencies(UNiagaraSystem* Asset, TArray<FVersionedNiagaraEmitter>& Dependencies) const
 {
 	if (!Asset)
 	{
 		return;
 	}
 
-	UNiagaraEmitter* FoundSourceEmitter = nullptr;
 	for (const FNiagaraEmitterHandle& EmitterHandle : Asset->GetEmitterHandles())
 	{
-		UNiagaraEmitter* EmitterInstance = EmitterHandle.GetInstance();
-		if (EmitterInstance && EmitterInstance->GetUniqueEmitterName() == EmitterName)
+		FVersionedNiagaraEmitter EmitterInstance = EmitterHandle.GetInstance();
+		if (EmitterInstance.Emitter && EmitterInstance.Emitter->GetUniqueEmitterName() == EmitterName)
 		{
 			Dependencies.Add(EmitterInstance);
 			return;
diff --git a/Engine/Plugins/FX/Niagara/Source/Niagara/Private/NiagaraDataInterfaceSkeletalMesh.cpp b/Engine/Plugins/FX/Niagara/Source/Niagara/Private/NiagaraDataInterfaceSkeletalMesh.cpp
--- a/Engine/Plugins/FX/Niagara/Source/Niagara/Private/NiagaraDataInterfaceSkeletalMesh.cpp
+++ b/Engine/Plugins/FX/Niagara/Source/Niagara/Private/NiagaraDataInterfaceSkeletalMesh.cpp
@@ -22,7 +22,6 @@
 #include "Templates/AlignmentTemplates.h"
 #include "ShaderParameterUtils.h"
 #include "ShaderCore.h"
-#include "SkeletalMeshTypes.h"
 
 #define LOCTEXT_NAMESPACE "NiagaraDataInterfaceSkeletalMesh"
 
@@ -2785,7 +2784,7 @@ void UNiagaraDataInterfaceSkeletalMesh::GetFeedback(UNiagaraSystem* Asset, UNiag
 			Scripts.Add(Asset->GetSystemUpdateScript());
 			for (auto&& EmitterHandle : Asset->GetEmitterHandles())
 			{
-				EmitterHandle.GetInstance()->GetScripts(Scripts, false);
+				EmitterHandle.GetEmitterData()->GetScripts(Scripts, false);
 			}
 
 			// Now check if any script uses functions that require CPU access
diff --git a/Engine/Plugins/FX/Niagara/Source/Niagara/Private/NiagaraDataInterfaceVectorField.cpp b/Engine/Plugins/FX/Niagara/Source/Niagara/Private/NiagaraDataInterfaceVectorField.cpp
--- a/Engine/Plugins/FX/Niagara/Source/Niagara/Private/NiagaraDataInterfaceVectorField.cpp
+++ b/Engine/Plugins/FX/Niagara/Source/Niagara/Private/NiagaraDataInterfaceVectorField.cpp
@@ -268,7 +268,7 @@ void UNiagaraDataInterfaceVectorField::GetFeedback(UNiagaraSystem* InAsset, UNia
 		for (auto&& EmitterHandle : InAsset->GetEmitterHandles())
 		{
 			TArray<UNiagaraScript*> OutScripts;
-			EmitterHandle.GetInstance()->GetScripts(OutScripts, false);
+			EmitterHandle.GetEmitterData()->GetScripts(OutScripts, false);
 			Scripts.Append(OutScripts);
 		}
 
@@ -294,14 +294,12 @@ void UNiagaraDataInterfaceVectorField::GetFeedback(UNiagaraSystem* InAsset, UNia
 
 					if (bHasCPUFunctionsReferenced)
 					{
-						bool bMatchFound = false;
 						// We assume that if the properties match or we are referencing an external variable whose name is in the list of candidates found in the prior search, it's a valid match for us.
 						if (CachedDefaultDIs.IsValidIndex(Idx) && CachedDefaultDIs[Idx].DataInterface != nullptr &&
 							(CachedDefaultDIs[Idx].DataInterface->Equals(this) || DIAliases.Contains(CachedDefaultDIs[Idx].Name)))
 						{
-							bMatchFound = true;
-							UNiagaraEmitter* OuterEmitter = Script->GetTypedOuter<UNiagaraEmitter>();
-							if (OuterEmitter && (OuterEmitter->SimTarget == ENiagaraSimTarget::CPUSim || Script->IsSystemScript(Script->Usage)))
+							FVersionedNiagaraEmitter OuterEmitter = Script->GetOuterEmitter();
+							if (OuterEmitter.Emitter && (OuterEmitter.GetEmitterData()->SimTarget == ENiagaraSimTarget::CPUSim || Script->IsSystemScript(Script->Usage)))
 							{
 								bHasCPUFunctions = true;
 							}
diff --git a/Engine/Plugins/FX/Niagara/Source/Niagara/Private/NiagaraDebugHud.cpp b/Engine/Plugins/FX/Niagara/Source/Niagara/Private/NiagaraDebugHud.cpp
--- a/Engine/Plugins/FX/Niagara/Source/Niagara/Private/NiagaraDebugHud.cpp
+++ b/Engine/Plugins/FX/Niagara/Source/Niagara/Private/NiagaraDebugHud.cpp
@@ -16,7 +16,6 @@
 #include "Engine/Canvas.h"
 #include "GameFramework/PlayerController.h"
 #include "Particles/FXBudget.h"
-#include "DrawDebugHelpers.h"
 
 #if WITH_NIAGARA_DEBUGGER
 
@@ -559,7 +558,7 @@ FNiagaraDebugHud::FNiagaraDebugHud(UWorld* World)
 
 				for ( const auto& DispatchResult : GpuResults->DispatchResults )
 				{
-					if ( DispatchResult.OwnerEmitter.IsExplicitlyNull() )
+					if ( DispatchResult.OwnerEmitter.Emitter.IsExplicitlyNull() )
 					{
 						FGpuUsagePerEvent& EventUsage = GpuUsagePerEvent.FindOrAdd(DispatchResult.StageName);
 						EventUsage.InstanceCount.Accumulate(GpuResultsGameFrameCounter, 1);
@@ -567,7 +566,7 @@ FNiagaraDebugHud::FNiagaraDebugHud(UWorld* World)
 					}
 					else
 					{
-						UNiagaraEmitter* OwnerEmitter = DispatchResult.OwnerEmitter.Get();
+						UNiagaraEmitter* OwnerEmitter = DispatchResult.OwnerEmitter.Emitter.Get();
 						UNiagaraSystem* OwnerSystem = OwnerEmitter ? OwnerEmitter->GetTypedOuter<UNiagaraSystem>() : nullptr;
 						if (OwnerSystem == nullptr)
 						{
@@ -858,7 +857,7 @@ void FNiagaraDebugHud::GatherSystemInfo()
 		{
 			for (const TSharedRef<FNiagaraEmitterInstance, ESPMode::ThreadSafe>& EmitterInstance : SystemInstance->GetEmitters())
 			{
-				UNiagaraEmitter* NiagaraEmitter = EmitterInstance->GetCachedEmitter();
+				UNiagaraEmitter* NiagaraEmitter = EmitterInstance->GetCachedEmitter().Emitter;
 				if (NiagaraEmitter == nullptr)
 				{
 					continue;
@@ -880,7 +879,7 @@ void FNiagaraDebugHud::GatherSystemInfo()
 
 			for (const TSharedRef<FNiagaraEmitterInstance, ESPMode::ThreadSafe>& EmitterInstance : SystemInstance->GetEmitters())
 			{
-				UNiagaraEmitter* NiagaraEmitter = EmitterInstance->GetCachedEmitter();
+				UNiagaraEmitter* NiagaraEmitter = EmitterInstance->GetCachedEmitter().Emitter;
 				if (NiagaraEmitter == nullptr)
 				{
 					continue;
@@ -2232,7 +2231,7 @@ void FNiagaraDebugHud::DrawValidation(class FNiagaraWorldManager* WorldManager,
 					[&](const FNiagaraVariable& Variable, int32 InstanceIndex, int32 ComponentIndex)
 					{
 						auto& ValidationError = GetValidationErrorInfo(NiagaraComponent);
-						ValidationError.ParticleVariablesWithErrors.FindOrAdd(EmitterInstance->GetCachedEmitter()->GetFName()).AddUnique(Variable.GetName());
+						ValidationError.ParticleVariablesWithErrors.FindOrAdd(EmitterInstance->GetCachedEmitter().Emitter->GetFName()).AddUnique(Variable.GetName());
 					}
 				);
 			}
@@ -2413,7 +2412,7 @@ void FNiagaraDebugHud::DrawComponents(FNiagaraWorldManager* WorldManager, UCanva
 				FSceneView* SceneView = Canvas->SceneView;
 
 				const FTransform& SystemTransform = SystemInstance->GetWorldTransform();
-				const bool bParticlesLocalSpace = EmitterInstance->GetCachedEmitter()->bLocalSpace;
+				const bool bParticlesLocalSpace = EmitterInstance->GetCachedEmitterData()->bLocalSpace;
 				//const float ClipRadius = Settings.bUseParticleDisplayRadius ? 1.0f : 0.0f;
 				const float ParticleDisplayCenterRadiusSq = Settings.bUseParticleDisplayCenterRadius ? (Settings.ParticleDisplayCenterRadius * Settings.ParticleDisplayCenterRadius) : 0.0f;
 				const float ParticleDisplayClipNearPlane = Settings.bUseParticleDisplayClip ? FMath::Max(Settings.ParticleDisplayClip.X, 0.0f) : 0.0f;
@@ -2549,7 +2548,7 @@ void FNiagaraDebugHud::DrawComponents(FNiagaraWorldManager* WorldManager, UCanva
 					int64 TotalBytes = 0;
 					for (const TSharedRef<FNiagaraEmitterInstance, ESPMode::ThreadSafe>& EmitterInstance : SystemInstance->GetEmitters())
 					{
-						if ( UNiagaraEmitter* NiagaraEmitter = EmitterInstance->GetCachedEmitter() )
+						if ( UNiagaraEmitter* NiagaraEmitter = EmitterInstance->GetCachedEmitter().Emitter )
 						{
 							TotalBytes += EmitterInstance->GetTotalBytesUsed();
 						}
@@ -2564,7 +2563,7 @@ void FNiagaraDebugHud::DrawComponents(FNiagaraWorldManager* WorldManager, UCanva
 					int32 ActiveParticles = 0;
 					for (const TSharedRef<FNiagaraEmitterInstance, ESPMode::ThreadSafe>& EmitterInstance : SystemInstance->GetEmitters())
 					{
-						UNiagaraEmitter* NiagaraEmitter = EmitterInstance->GetCachedEmitter();
+						UNiagaraEmitter* NiagaraEmitter = EmitterInstance->GetCachedEmitter().Emitter;
 						if (NiagaraEmitter == nullptr)
 						{
 							continue;
@@ -2692,7 +2691,7 @@ void FNiagaraDebugHud::DrawComponents(FNiagaraWorldManager* WorldManager, UCanva
 									continue;
 								}
 
-								StringBuilder.Appendf(TEXT("Emitter (%s)\n"), *EmitterInstance->GetCachedEmitter()->GetUniqueEmitterName());
+								StringBuilder.Appendf(TEXT("Emitter (%s)\n"), *EmitterInstance->GetCachedEmitter().Emitter->GetUniqueEmitterName());
 								const uint32 NumParticles = Settings.bUseMaxParticlesToDisplay ? FMath::Min((uint32)Settings.MaxParticlesToDisplay, DataBuffer->GetNumInstances()) : DataBuffer->GetNumInstances();
 								for (uint32 iInstance = 0; iInstance < NumParticles; ++iInstance)
 								{
diff --git a/Engine/Plugins/FX/Niagara/Source/Niagara/Private/NiagaraDebuggerClient.cpp b/Engine/Plugins/FX/Niagara/Source/Niagara/Private/NiagaraDebuggerClient.cpp
--- a/Engine/Plugins/FX/Niagara/Source/Niagara/Private/NiagaraDebuggerClient.cpp
+++ b/Engine/Plugins/FX/Niagara/Source/Niagara/Private/NiagaraDebuggerClient.cpp
@@ -72,7 +72,7 @@ void FNiagaraDebuggerClient::UpdateClientInfo()
 				NewInfo->Systems.Add(System->GetName());
 				for (const FNiagaraEmitterHandle& Handle : System->GetEmitterHandles())
 				{
-					if (Handle.GetInstance())
+					if (Handle.GetInstance().Emitter)
 					{
 						NewInfo->Emitters.AddUnique(Handle.GetUniqueInstanceName());
 					}					
@@ -381,18 +381,19 @@ bool FNiagaraDebuggerClient::UpdateOutliner(float DeltaSeconds)
 						for (TSharedRef<FNiagaraEmitterInstance, ESPMode::ThreadSafe>& EmitterInst : Inst->GetEmitters())
 						{
 							FNiagaraOutlinerEmitterInstanceData& EmitterData = InstData.Emitters.AddDefaulted_GetRef();
-							if (UNiagaraEmitter* NiagaraEmitter = EmitterInst->GetCachedEmitter())
+							FVersionedNiagaraEmitter VersionedEmitter = EmitterInst->GetCachedEmitter();
+							if (VersionedEmitter.Emitter)
 							{
 								//TODO: This is a bit wasteful to copy the name into each instance data. Though we can't rely on the debugger side data matchin the actul running data on the device.
 								//We need to build a shared representation of the asset data from the client that we then reference from this per instance data.
-								EmitterData.EmitterName = NiagaraEmitter->GetUniqueEmitterName();
-								EmitterData.SimTarget = NiagaraEmitter->SimTarget;
+								EmitterData.EmitterName = VersionedEmitter.Emitter->GetUniqueEmitterName();
+								EmitterData.SimTarget = VersionedEmitter.GetEmitterData()->SimTarget;
 								//Move all above to a shared asset representation.
 
 								EmitterData.ExecState = EmitterInst->GetExecutionState();
 								EmitterData.NumParticles = EmitterInst->GetNumParticles();
 
-								EmitterData.bRequiresPersistentIDs = NiagaraEmitter->RequiresPersistentIDs();
+								EmitterData.bRequiresPersistentIDs = VersionedEmitter.GetEmitterData()->RequiresPersistentIDs();
 							}
 						}
 					}
diff --git a/Engine/Plugins/FX/Niagara/Source/Niagara/Private/NiagaraEffectType.cpp b/Engine/Plugins/FX/Niagara/Source/Niagara/Private/NiagaraEffectType.cpp
--- a/Engine/Plugins/FX/Niagara/Source/Niagara/Private/NiagaraEffectType.cpp
+++ b/Engine/Plugins/FX/Niagara/Source/Niagara/Private/NiagaraEffectType.cpp
@@ -64,7 +64,7 @@ void UNiagaraEffectType::PostLoad()
 	}
 
 	//Apply platform set redirectors
-	auto ApplyPlatformSetRedirects = [](UObject* Owner, FNiagaraPlatformSet& Platforms)
+	auto ApplyPlatformSetRedirects = [](FNiagaraPlatformSet& Platforms)
 	{
 		Platforms.ApplyRedirects();
 	};
