Added versioning for Niagara emitters

diff --git a/Engine/Plugins/FX/Niagara/Source/Niagara/Private/NiagaraEmitter.cpp b/Engine/Plugins/FX/Niagara/Source/Niagara/Private/NiagaraEmitter.cpp
--- a/Engine/Plugins/FX/Niagara/Source/Niagara/Private/NiagaraEmitter.cpp
+++ b/Engine/Plugins/FX/Niagara/Source/Niagara/Private/NiagaraEmitter.cpp
@@ -20,9 +20,10 @@
 #include "Modules/ModuleManager.h"
 #include "UObject/ObjectSaveContext.h"
 #include "UObject/Package.h"
+#include "UObject/UObjectIterator.h"
 
 #if WITH_EDITOR
-const FName UNiagaraEmitter::PrivateMemberNames::EventHandlerScriptProps = GET_MEMBER_NAME_CHECKED(UNiagaraEmitter, EventHandlerScriptProps);
+const FName UNiagaraEmitter::PrivateMemberNames::EventHandlerScriptProps = GET_MEMBER_NAME_CHECKED(FVersionedNiagaraEmitterData, EventHandlerScriptProps);
 
 const FString InitialNotSynchronizedReason("Emitter created");
 
@@ -145,23 +146,24 @@ bool FNiagaraEmitterScriptProperties::DataSetAccessSynchronized() const
 
 UNiagaraEmitter::UNiagaraEmitter(const FObjectInitializer& Initializer)
 : Super(Initializer)
-, PreAllocationCount(0)
-, FixedBounds(GetDefaultFixedBounds())
+, PreAllocationCount_DEPRECATED(0)
+, FixedBounds_DEPRECATED(FBox(FVector(-100), FVector(100)))
 , MinDetailLevel_DEPRECATED(0)
 , MaxDetailLevel_DEPRECATED(4)
-, bInterpolatedSpawning(false)
+, bInterpolatedSpawning_DEPRECATED(false)
+, bFixedBounds_DEPRECATED(false)
 , bUseMinDetailLevel_DEPRECATED(false)
 , bUseMaxDetailLevel_DEPRECATED(false)
-, bRequiresPersistentIDs(false)
-, bCombineEventSpawn(false)
-, MaxDeltaTimePerTick(0.125)
-, bLimitDeltaTime(true)
-, MaxGPUParticlesSpawnPerFrame(0)
+, bRequiresPersistentIDs_DEPRECATED(false)
+, bCombineEventSpawn_DEPRECATED(false)
+, MaxDeltaTimePerTick_DEPRECATED(0.125)
+, bLimitDeltaTime_DEPRECATED(true)
+, MaxGPUParticlesSpawnPerFrame_DEPRECATED(0)
+{
 #if WITH_EDITORONLY_DATA
-, ThumbnailImageOutOfDate(true)
-, IsCooked(false)
+	ThumbnailImageOutOfDate = true;
+	IsCooked = false;
 #endif
-{
 }
 
 void UNiagaraEmitter::PostInitProperties()
@@ -169,42 +171,119 @@ void UNiagaraEmitter::PostInitProperties()
 	Super::PostInitProperties();
 	if (HasAnyFlags(RF_ClassDefaultObject | RF_NeedLoad) == false)
 	{
-		SpawnScriptProps.Script = NewObject<UNiagaraScript>(this, "SpawnScript", EObjectFlags::RF_Transactional);
+		for (FVersionedNiagaraEmitterData& Data : VersionData)
+		{
+			Data.PostInitProperties(this);
+		}
+	}
+
+#if WITH_EDITORONLY_DATA
+	for (FVersionedNiagaraEmitterData& Data : VersionData)
+	{
+		if (Data.GPUComputeScript)
+		{
+			Data.GPUComputeScript->OnGPUScriptCompiled().RemoveAll(this);
+			Data.GPUComputeScript->OnGPUScriptCompiled().AddUObject(this, &UNiagaraEmitter::RaiseOnEmitterGPUCompiled);
+		}
+	}
+#endif
+
+	UniqueEmitterName = TEXT("Emitter");
+
+	ResolveScalabilitySettings();
+}
+
+void FVersionedNiagaraEmitterData::CopyFrom(const FVersionedNiagaraEmitterData& Source)
+{
+	// copy over uproperties
+    for (TFieldIterator<FProperty> PropertyIt(StaticStruct()); PropertyIt; ++PropertyIt)
+    {
+    	FProperty* Property = *PropertyIt;
+    	const uint8* SourceAddr = Property->ContainerPtrToValuePtr<uint8>(&Source);
+        uint8* DestinationAddr = Property->ContainerPtrToValuePtr<uint8>(this);
+
+        Property->CopyCompleteValue(DestinationAddr, SourceAddr);
+    }
+
+#if STATS
+	StatDatabase.Init();
+#endif
+	RuntimeEstimation.Init();
+}
+
+FName GetVersionedName(const FString& BaseName, const FNiagaraAssetVersion& VersionPostfix)
+{
+	return FName(FString::Printf(TEXT("%s_%i.%i"), *BaseName, VersionPostfix.MajorVersion, VersionPostfix.MinorVersion));
+}
+
+void FVersionedNiagaraEmitterData::PostInitProperties(UNiagaraEmitter* Outer)
+{
+		SpawnScriptProps.Script = NewObject<UNiagaraScript>(Outer, GetVersionedName("SpawnScript", Version), RF_Transactional);
 		SpawnScriptProps.Script->SetUsage(ENiagaraScriptUsage::ParticleSpawnScript);
 
-		UpdateScriptProps.Script = NewObject<UNiagaraScript>(this, "UpdateScript", EObjectFlags::RF_Transactional);
+		UpdateScriptProps.Script = NewObject<UNiagaraScript>(Outer, GetVersionedName("UpdateScript", Version), RF_Transactional);
 		UpdateScriptProps.Script->SetUsage(ENiagaraScriptUsage::ParticleUpdateScript);
 
 #if WITH_EDITORONLY_DATA
-		EmitterSpawnScriptProps.Script = NewObject<UNiagaraScript>(this, "EmitterSpawnScript", EObjectFlags::RF_Transactional);
+		EmitterSpawnScriptProps.Script = NewObject<UNiagaraScript>(Outer, GetVersionedName("EmitterSpawnScript", Version), RF_Transactional);
 		EmitterSpawnScriptProps.Script->SetUsage(ENiagaraScriptUsage::EmitterSpawnScript);
 		
-		EmitterUpdateScriptProps.Script = NewObject<UNiagaraScript>(this, "EmitterUpdateScript", EObjectFlags::RF_Transactional);
+		EmitterUpdateScriptProps.Script = NewObject<UNiagaraScript>(Outer, GetVersionedName("EmitterUpdateScript", Version), RF_Transactional);
 		EmitterUpdateScriptProps.Script->SetUsage(ENiagaraScriptUsage::EmitterUpdateScript);
 #endif
 
-		GPUComputeScript = NewObject<UNiagaraScript>(this, "GPUComputeScript", EObjectFlags::RF_Transactional);
+		GPUComputeScript = NewObject<UNiagaraScript>(Outer, GetVersionedName("GPUComputeScript", Version), RF_Transactional);
 		GPUComputeScript->SetUsage(ENiagaraScriptUsage::ParticleGPUComputeScript);
 
 #if WITH_EDITORONLY_DATA && WITH_EDITOR
-		if (EditorParameters == nullptr)
+	if (EditorParameters == nullptr)
+	{
+		INiagaraModule& NiagaraModule = FModuleManager::GetModuleChecked<INiagaraModule>("Niagara");
+		EditorParameters = NiagaraModule.GetEditorOnlyDataUtilities().CreateDefaultEditorParameters(Outer);
+	}
+#endif
+}
+
+bool FVersionedNiagaraEmitterData::UsesCollection(const UNiagaraParameterCollection* Collection) const
+{
+	if (SpawnScriptProps.Script && SpawnScriptProps.Script->UsesCollection(Collection))
+	{
+		return true;
+	}
+	if (UpdateScriptProps.Script && UpdateScriptProps.Script->UsesCollection(Collection))
+	{
+		return true;
+	}
+	for (int32 i = 0; i < EventHandlerScriptProps.Num(); i++)
+	{
+		if (EventHandlerScriptProps[i].Script && EventHandlerScriptProps[i].Script->UsesCollection(Collection))
 		{
-			INiagaraModule& NiagaraModule = FModuleManager::GetModuleChecked<INiagaraModule>("Niagara");
-			EditorParameters = NiagaraModule.GetEditorOnlyDataUtilities().CreateDefaultEditorParameters(this);
+			return true;
 		}
-#endif
 	}
+	return false;
+}
 
+bool FVersionedNiagaraEmitterData::UsesScript(const UNiagaraScript* Script) const
+{
+	if (SpawnScriptProps.Script == Script || UpdateScriptProps.Script == Script)
+	{
+		return true;
+	}
 #if WITH_EDITORONLY_DATA
-	if (GPUComputeScript)
+	if (EmitterSpawnScriptProps.Script == Script || EmitterUpdateScriptProps.Script == Script)
 	{
-		GPUComputeScript->OnGPUScriptCompiled().AddUObject(this, &UNiagaraEmitter::RaiseOnEmitterGPUCompiled);
+		return true;
 	}
 #endif
-
-	UniqueEmitterName = TEXT("Emitter");
-
-	ResolveScalabilitySettings();
+	for (int32 i = 0; i < EventHandlerScriptProps.Num(); i++)
+	{
+		if (EventHandlerScriptProps[i].Script == Script)
+		{
+			return true;
+		}
+	}
+	return false;
 }
 
 #if WITH_EDITORONLY_DATA
@@ -215,98 +294,176 @@ bool UNiagaraEmitter::GetForceCompileOnLoad()
 
 bool UNiagaraEmitter::IsSynchronizedWithParent() const
 {
-	if (Parent == nullptr)
+	for (const FVersionedNiagaraEmitterData& Data : VersionData)
+	{
+		if (!Data.IsSynchronizedWithParent())
+		{
+			return false;
+		}
+	}
+	return true;
+}
+
+bool FVersionedNiagaraEmitterData::IsSynchronizedWithParent() const
+{
+	if (VersionedParent.Emitter == nullptr)
 	{
 		// If the emitter has no parent than it is synchronized by default.
 		return true;
 	}
 
-	if (ParentAtLastMerge == nullptr)
+	if (VersionedParentAtLastMerge.Emitter == nullptr)
 	{
 		// If the parent was valid but the parent at last merge isn't, they we don't know if it's up to date so we say it's not, and let 
 		// the actual merge code print an appropriate message to the log.
 		return false;
 	}
 
-	if (Parent->GetChangeId().IsValid() == false ||
-		ParentAtLastMerge->GetChangeId().IsValid() == false)
+	if (VersionedParent.Emitter->GetChangeId().IsValid() == false ||
+		VersionedParentAtLastMerge.Emitter->GetChangeId().IsValid() == false)
 	{
 		// If any of the change Ids aren't valid then we assume we're out of sync.
 		return false;
 	}
 
 	// Otherwise check the change ids, and the force flag.
-	return Parent->GetChangeId() == ParentAtLastMerge->GetChangeId() && GbForceNiagaraMergeOnLoad <= 0;
+	return VersionedParent.Emitter->GetChangeId() == VersionedParentAtLastMerge.Emitter->GetChangeId() && GbForceNiagaraMergeOnLoad <= 0;
 }
 
-INiagaraMergeManager::FMergeEmitterResults UNiagaraEmitter::MergeChangesFromParent()
+TArray<INiagaraMergeManager::FMergeEmitterResults> UNiagaraEmitter::MergeChangesFromParent()
 {
 	TRACE_CPUPROFILER_EVENT_SCOPE(MergeEmitter);
 	TRACE_CPUPROFILER_EVENT_SCOPE_TEXT_ON_CHANNEL(*GetPathName(), NiagaraChannel);
 
-	if (GbEnableEmitterChangeIdMergeLogging)
+	// we need to merge each version separately
+	TArray<INiagaraMergeManager::FMergeEmitterResults> Results;
+	for (FVersionedNiagaraEmitterData& Data : VersionData)
 	{
-		UE_LOG(LogNiagara, Log, TEXT("Emitter %s is merging changes from parent %s because its Change ID was updated."), *GetPathName(),
-			Parent != nullptr ? *Parent->GetPathName() : TEXT("(null)"));
+		FVersionedNiagaraEmitter& VersionedParent = Data.VersionedParent;
+		FVersionedNiagaraEmitter& VersionedParentAtLastMerge = Data.VersionedParentAtLastMerge;
 
-		UE_LOG(LogNiagara, Log, TEXT("\nEmitter %s Id=%s \nParentAtLastMerge %s id=%s \nParent %s Id=%s."), 
-			*GetPathName(), *ChangeId.ToString(),
-			ParentAtLastMerge != nullptr ? *ParentAtLastMerge->GetPathName() : TEXT("(null)"), ParentAtLastMerge != nullptr ? *ParentAtLastMerge->GetChangeId().ToString() : TEXT("(null)"),
-			Parent != nullptr ? *Parent->GetPathName() : TEXT("(null)"), Parent != nullptr ? *Parent->GetChangeId().ToString() : TEXT("(null)"));
-	}
+		if (VersionedParent.Emitter && VersionedParent.Emitter->IsVersioningEnabled() && VersionedParent.Version.IsValid() == false)
+		{
+			// parent was versioned after this emitter was created, so we switch to version 1.0 of the parent.
+			VersionedParent.Version = VersionedParent.Emitter->GetAllAvailableVersions()[0].VersionGuid;
+		}
+		if (GbEnableEmitterChangeIdMergeLogging)
+		{
+			UE_LOG(LogNiagara, Log, TEXT("Emitter %s is merging changes from parent %s because its Change ID was updated."), *GetPathName(),
+				Data.VersionedParent.Emitter != nullptr ? *Data.VersionedParent.Emitter->GetPathName() : TEXT("(null)"));
 
-	if (Parent == nullptr)
-	{
-		// If we don't have a copy of the parent emitter, this emitter can't safely be merged.
-		INiagaraMergeManager::FMergeEmitterResults MergeResults;
-		MergeResults.MergeResult = INiagaraMergeManager::EMergeEmitterResult::FailedToDiff;
-		MergeResults.bModifiedGraph = false;
-		MergeResults.ErrorMessages.Add(NSLOCTEXT("NiagaraEmitter", "NoParentErrorMessage", "This emitter has no 'Parent' so changes can't be merged in."));
-		return MergeResults;
-	}
+			UE_LOG(LogNiagara, Log, TEXT("\nEmitter %s Id=%s \nParentAtLastMerge %s id=%s \nParent %s Id=%s."), 
+				*GetPathName(), *ChangeId.ToString(),
+				VersionedParentAtLastMerge.Emitter != nullptr ? *VersionedParentAtLastMerge.Emitter->GetPathName() : TEXT("(null)"), VersionedParentAtLastMerge.Emitter != nullptr ? *VersionedParentAtLastMerge.Emitter->GetChangeId().ToString() : TEXT("(null)"),
+				VersionedParent.Emitter != nullptr ? *VersionedParent.Emitter->GetPathName() : TEXT("(null)"), VersionedParent.Emitter != nullptr ? *VersionedParent.Emitter->GetChangeId().ToString() : TEXT("(null)"));
+		}
 
-	const bool bNoParentAtLastMerge = (ParentAtLastMerge == nullptr);
+		if (VersionedParent.GetEmitterData() == nullptr)
+		{
+			// If we don't have a copy of the parent emitter, this emitter can't safely be merged.
+			INiagaraMergeManager::FMergeEmitterResults MergeResults;
+			MergeResults.MergeResult = INiagaraMergeManager::EMergeEmitterResult::FailedToDiff;
+			MergeResults.bModifiedGraph = false;
+			MergeResults.ErrorMessages.Add(NSLOCTEXT("NiagaraEmitter", "NoParentErrorMessage", "This emitter has no 'Parent' so changes can't be merged in."));
+			Results.Add(MergeResults);
+			continue;
+		}
 
-	INiagaraModule& NiagaraModule = FModuleManager::Get().GetModuleChecked<INiagaraModule>("Niagara");
-	const INiagaraMergeManager& MergeManager = NiagaraModule.GetMergeManager();
-	INiagaraMergeManager::FMergeEmitterResults MergeResults = MergeManager.MergeEmitter(*Parent, ParentAtLastMerge, *this);
-	if (MergeResults.MergeResult == INiagaraMergeManager::EMergeEmitterResult::SucceededDifferencesApplied || MergeResults.MergeResult == INiagaraMergeManager::EMergeEmitterResult::SucceededNoDifferences)
-	{
-		if (MergeResults.MergeResult == INiagaraMergeManager::EMergeEmitterResult::SucceededDifferencesApplied)
+		INiagaraModule& NiagaraModule = FModuleManager::Get().GetModuleChecked<INiagaraModule>("Niagara");
+		const INiagaraMergeManager& MergeManager = NiagaraModule.GetMergeManager();
+		FVersionedNiagaraEmitter InstanceToMerge = FVersionedNiagaraEmitter(this, Data.Version.VersionGuid);
+		INiagaraMergeManager::FMergeEmitterResults MergeResults = MergeManager.MergeEmitter(VersionedParent, VersionedParentAtLastMerge, InstanceToMerge);
+		Results.Add(MergeResults);
+		if (MergeResults.MergeResult == INiagaraMergeManager::EMergeEmitterResult::SucceededDifferencesApplied || MergeResults.MergeResult == INiagaraMergeManager::EMergeEmitterResult::SucceededNoDifferences)
+		{
+			if (MergeResults.MergeResult == INiagaraMergeManager::EMergeEmitterResult::SucceededDifferencesApplied)
+			{
+				UpdateFromMergedCopy(MergeManager, MergeResults.MergedInstance, &Data);
+			}
+
+			// Update the last merged source and clear it's stand alone and public flags since it's not an asset.
+			VersionedParentAtLastMerge.Emitter = VersionedParent.Emitter->DuplicateWithoutMerging(this);
+			VersionedParentAtLastMerge.Emitter->ClearFlags(RF_Standalone | RF_Public);
+			VersionedParentAtLastMerge.Emitter->DisableVersioning(VersionedParent.Version);
+			VersionedParentAtLastMerge.Version = VersionedParent.Version;
+		}
+		else
 		{
-			UpdateFromMergedCopy(MergeManager, MergeResults.MergedInstance);
+			UE_LOG(LogNiagara, Warning, TEXT("Failed to merge changes for parent emitter.  Emitter: %s  Parent Emitter: %s  Error Message: %s"),
+				*GetPathName(), VersionedParent.Emitter != nullptr ? *VersionedParent.Emitter->GetPathName() : TEXT("(null)"), *MergeResults.GetErrorMessagesString());
 		}
 
-		// Update the last merged source and clear it's stand alone and public flags since it's not an asset.
-		ParentAtLastMerge = Parent->DuplicateWithoutMerging(this);
-		ParentAtLastMerge->ClearFlags(RF_Standalone | RF_Public);
-	}
-	else
-	{
-		UE_LOG(LogNiagara, Warning, TEXT("Failed to merge changes for parent emitter.  Emitter: %s  Parent Emitter: %s  Error Message: %s"),
-			*GetPathName(), Parent != nullptr ? *Parent->GetPathName() : TEXT("(null)"), *MergeResults.GetErrorMessagesString());
+		RemoveMessage(EmitterMergeMessageId);
+		if (MergeResults.MergeNiagaraMessage != nullptr)
+		{
+			AddMessage(EmitterMergeMessageId, CastChecked<UNiagaraMessageDataBase>(StaticDuplicateObject(MergeResults.MergeNiagaraMessage, this)));
+		}
 	}
+	return Results;
+}
+
+bool FVersionedNiagaraEmitterData::UsesEmitter(const UNiagaraEmitter& InEmitter) const
+{
+	return VersionedParent.Emitter == &InEmitter || (VersionedParent.GetEmitterData() && VersionedParent.GetEmitterData()->UsesEmitter(InEmitter));
+}
+
+void FVersionedNiagaraEmitterData::GatherStaticVariables(TArray<FNiagaraVariable>& OutVars) const
+{
+	TArray<UNiagaraScript*> OutScripts;
+	GetScripts(OutScripts, false, true);
 
-	RemoveMessage(EmitterMergeMessageId);
-	if (MergeResults.MergeNiagaraMessage != nullptr)
+	for (UNiagaraScript* Script : OutScripts)
 	{
-		AddMessage(EmitterMergeMessageId, CastChecked<UNiagaraMessageDataBase>(StaticDuplicateObject(MergeResults.MergeNiagaraMessage, this)));
-	}
+		TArray<FNiagaraVariable> StoreParams;
+		Script->RapidIterationParameters.GetParameters(StoreParams);
 
-	return MergeResults;
+		for (int32 i = 0; i < StoreParams.Num(); i++)
+		{
+			if (StoreParams[i].GetType().IsStatic())
+			{
+				const int32* Index = Script->RapidIterationParameters.FindParameterOffset(StoreParams[i]);
+				if (Index != nullptr)
+				{
+					StoreParams[i].SetData(Script->RapidIterationParameters.GetParameterData(*Index)); // This will memcopy the data in.			
+					OutVars.AddUnique(StoreParams[i]);					
+					//UE_LOG(LogNiagara, Log, TEXT("UNiagaraEmitter::GatherStaticVariables Added %s"), *StoreParams[i].ToString());
+				}
+			}
+		}
+	}
 }
 
-bool UNiagaraEmitter::UsesEmitter(const UNiagaraEmitter& InEmitter) const
+struct TParentGuard : private FNoncopyable
 {
-	return Parent == &InEmitter || (Parent != nullptr && Parent->UsesEmitter(InEmitter));
-}
+	TParentGuard(TArray<FVersionedNiagaraEmitterData>& InVersionData) : VersionData(InVersionData)
+	{
+		for (FVersionedNiagaraEmitterData& Data : VersionData)
+		{
+			OldParentValues.Add(Data.VersionedParent);
+			OldParentLastMergeValues.Add(Data.VersionedParentAtLastMerge);
+			Data.RemoveParent();
+		}
+	}
+	~TParentGuard()
+	{
+		for (int i = 0; i < VersionData.Num(); i++)
+		{
+			VersionData[i].VersionedParent = OldParentValues[i];
+			VersionData[i].VersionedParentAtLastMerge = OldParentLastMergeValues[i];
+		}
+	}
+
+private:
+	TArray<FVersionedNiagaraEmitterData>& VersionData;
+	TArray<FVersionedNiagaraEmitter> OldParentValues;
+	TArray<FVersionedNiagaraEmitter> OldParentLastMergeValues;
+};
 
 UNiagaraEmitter* UNiagaraEmitter::DuplicateWithoutMerging(UObject* InOuter)
 {
 	UNiagaraEmitter* Duplicate;
 	{
-		TGuardValue<decltype(Parent)> ParentGuard(Parent, nullptr);
-		TGuardValue<decltype(ParentAtLastMerge)> ParentAtLastMergeGuard(ParentAtLastMerge, nullptr);
+		TParentGuard ParentGuard(VersionData);
 		Duplicate = Cast<UNiagaraEmitter>(StaticDuplicateObject(this, InOuter));
 	}
 	return Duplicate;
@@ -317,31 +474,33 @@ void UNiagaraEmitter::Serialize(FArchive& Ar)
 {
 
 #if WITH_EDITORONLY_DATA
-	for (UNiagaraSimulationStageBase* Stage : SimulationStages)
+	bool bHasTransientFlag = HasAnyFlags(RF_Transient);
+	for (FVersionedNiagaraEmitterData& Data : VersionData)
 	{
-		if (Stage)
+		for (UNiagaraSimulationStageBase* Stage : Data.SimulationStages)
 		{
-			if (Stage->Script)
+			if (Stage)
 			{
-				if (!HasAnyFlags(RF_Transient))
+				if (Stage->Script)
 				{
-					if (Stage->Script->HasAnyFlags(RF_Transient))
+					if (!bHasTransientFlag && Stage->Script->HasAnyFlags(RF_Transient))
 					{
 						UE_LOG(LogNiagara, Error, TEXT("Emitter \"%s\" has a simulation stage with a Transient script and the emitter itself isn't transient!"), *GetPathName());
 					}
 				}
+				else
+				{
+					UE_LOG(LogNiagara, Error, TEXT("Emitter \"%s\" has a simulation stage with a null Script entry!"), *GetPathName());
+				}
+
 			}
 			else
 			{
-				UE_LOG(LogNiagara, Error, TEXT("Emitter \"%s\" has a simulation stage with a null Script entry!"), *GetPathName());
+				UE_LOG(LogNiagara, Error, TEXT("Emitter \"%s\" has a simulation stage with a null entry!"), *GetPathName());
 			}
-
-		}
-		else
-		{
-			UE_LOG(LogNiagara, Error, TEXT("Emitter \"%s\" has a simulation stage with a null entry!"), *GetPathName());
 		}
 	}
+	
 #endif
 	Super::Serialize(Ar);
 
@@ -355,7 +514,7 @@ void UNiagaraEmitter::Serialize(FArchive& Ar)
 	Ar.UsingCustomVersion(FNiagaraCustomVersion::GUID);
 }
 
-void UNiagaraEmitter::EnsureScriptsPostLoaded()
+void FVersionedNiagaraEmitterData::EnsureScriptsPostLoaded()
 {
 	TArray<UNiagaraScript*> AllScripts;
 	GetScripts(AllScripts, false);
@@ -380,6 +539,10 @@ void UNiagaraEmitter::PostLoad()
 {
 	Super::PostLoad();
 
+#if WITH_EDITORONLY_DATA
+	CheckVersionDataAvailable();
+#endif
+
 	if (GIsEditor)
 	{
 		SetFlags(RF_Transactional);
@@ -397,36 +560,37 @@ void UNiagaraEmitter::PostLoad()
 			NewQLMask |= (1 << QL);
 		}
 
-		Platforms = FNiagaraPlatformSet(NewQLMask);
+		FVersionedNiagaraEmitterData* LatestEmitterData = GetLatestEmitterData();
+		LatestEmitterData->Platforms = FNiagaraPlatformSet(NewQLMask);
 
 		//Transfer spawn rate scaling overrides
 		if (bOverrideGlobalSpawnCountScale_DEPRECATED)
 		{
-			FNiagaraEmitterScalabilityOverride& LowOverride = ScalabilityOverrides.Overrides.AddDefaulted_GetRef();
+			FNiagaraEmitterScalabilityOverride& LowOverride = LatestEmitterData->ScalabilityOverrides.Overrides.AddDefaulted_GetRef();
 			LowOverride.Platforms = FNiagaraPlatformSet(FNiagaraPlatformSet::CreateQualityLevelMask(0));
 			LowOverride.bOverrideSpawnCountScale = true;
 			LowOverride.bScaleSpawnCount = true;
 			LowOverride.SpawnCountScale = GlobalSpawnCountScaleOverrides_DEPRECATED.Low;
 
-			FNiagaraEmitterScalabilityOverride& MediumOverride = ScalabilityOverrides.Overrides.AddDefaulted_GetRef();
+			FNiagaraEmitterScalabilityOverride& MediumOverride = LatestEmitterData->ScalabilityOverrides.Overrides.AddDefaulted_GetRef();
 			MediumOverride.Platforms = FNiagaraPlatformSet(FNiagaraPlatformSet::CreateQualityLevelMask(1));
 			MediumOverride.bOverrideSpawnCountScale = true;
 			MediumOverride.bScaleSpawnCount = true;
 			MediumOverride.SpawnCountScale = GlobalSpawnCountScaleOverrides_DEPRECATED.Medium;
 
-			FNiagaraEmitterScalabilityOverride& HighOverride = ScalabilityOverrides.Overrides.AddDefaulted_GetRef();
+			FNiagaraEmitterScalabilityOverride& HighOverride = LatestEmitterData->ScalabilityOverrides.Overrides.AddDefaulted_GetRef();
 			HighOverride.Platforms = FNiagaraPlatformSet(FNiagaraPlatformSet::CreateQualityLevelMask(2));
 			HighOverride.bOverrideSpawnCountScale = true;
 			HighOverride.bScaleSpawnCount = true;
 			HighOverride.SpawnCountScale = GlobalSpawnCountScaleOverrides_DEPRECATED.High;
 
-			FNiagaraEmitterScalabilityOverride& EpicOverride = ScalabilityOverrides.Overrides.AddDefaulted_GetRef();
+			FNiagaraEmitterScalabilityOverride& EpicOverride = LatestEmitterData->ScalabilityOverrides.Overrides.AddDefaulted_GetRef();
 			EpicOverride.Platforms = FNiagaraPlatformSet(FNiagaraPlatformSet::CreateQualityLevelMask(3));
 			EpicOverride.bOverrideSpawnCountScale = true;
 			EpicOverride.bScaleSpawnCount = true;
 			EpicOverride.SpawnCountScale = GlobalSpawnCountScaleOverrides_DEPRECATED.Epic;
 
-			FNiagaraEmitterScalabilityOverride& CineOverride = ScalabilityOverrides.Overrides.AddDefaulted_GetRef();
+			FNiagaraEmitterScalabilityOverride& CineOverride = LatestEmitterData->ScalabilityOverrides.Overrides.AddDefaulted_GetRef();
 			CineOverride.Platforms = FNiagaraPlatformSet(FNiagaraPlatformSet::CreateQualityLevelMask(4));
 			CineOverride.bOverrideSpawnCountScale = true;
 			CineOverride.bScaleSpawnCount = true;
@@ -435,12 +599,6 @@ void UNiagaraEmitter::PostLoad()
 	}
 
 #if WITH_EDITORONLY_DATA
-	if (EditorParameters == nullptr)
-	{
-		INiagaraModule& NiagaraModule = FModuleManager::GetModuleChecked<INiagaraModule>("Niagara");
-		EditorParameters = NiagaraModule.GetEditorOnlyDataUtilities().CreateDefaultEditorParameters(this);
-	}
-
 	// this can only ever be true for old assets that haven't been loaded yet, so this won't overwrite subsequent changes to the template specification
 	if(bIsTemplateAsset_DEPRECATED)
 	{
@@ -451,37 +609,70 @@ void UNiagaraEmitter::PostLoad()
 	{
 		LibraryVisibility = ENiagaraScriptLibraryVisibility::Library;
 	}
-	
-	if (bFixedBounds_DEPRECATED)
+
+#endif
+
+	for (FVersionedNiagaraEmitterData& Data : VersionData)
 	{
-		CalculateBoundsMode = ENiagaraEmitterCalculateBoundMode::Fixed;
+#if WITH_EDITORONLY_DATA
+		Data.PostLoad(*this, IsCooked);
+		Data.GPUComputeScript->OnGPUScriptCompiled().RemoveAll(this);
+		Data.GPUComputeScript->OnGPUScriptCompiled().AddUObject(this, &UNiagaraEmitter::RaiseOnEmitterGPUCompiled);
+#else
+		Data.PostLoad(*this, true);
+#endif
 	}
+
+#if !WITH_EDITOR
+	// When running without the editor in a cooked build we run the update immediately in post load since
+	// there will be no merging or compiling which makes it safe to do so.
+	UpdateEmitterAfterLoad();
+#endif
+}
+
+void FVersionedNiagaraEmitterData::PostLoad(UNiagaraEmitter& Emitter, bool bIsCooked)
+{
+#if STATS
+	StatDatabase.Init();
 #endif
+	RuntimeEstimation.Init();
 
+#if WITH_EDITORONLY_DATA && WITH_EDITOR
+	if (EditorParameters == nullptr)
+	{
+		INiagaraModule& NiagaraModule = FModuleManager::GetModuleChecked<INiagaraModule>("Niagara");
+		EditorParameters = NiagaraModule.GetEditorOnlyDataUtilities().CreateDefaultEditorParameters(&Emitter);
+	}
+#endif
+	
 	for (int32 RendererIndex = RendererProperties.Num() - 1; RendererIndex >= 0; --RendererIndex)
 	{
 		if (RendererProperties[RendererIndex] == nullptr)
 		{
 #if WITH_EDITORONLY_DATA
 			//In cooked builds these can be cooked out and null on purpose.
-			ensureMsgf(IsCooked, TEXT("Null renderer found in %s at index %i, removing it to prevent crashes."), *GetPathName(), RendererIndex);
+			ensureMsgf(bIsCooked, TEXT("Null renderer found in %s at index %i, removing it to prevent crashes."), *Emitter.GetPathName(), RendererIndex);
 #endif
 			RendererProperties.RemoveAt(RendererIndex);
 		}
 		else
 		{
+			RendererProperties[RendererIndex]->OuterEmitterVersion = Version.VersionGuid;
 			RendererProperties[RendererIndex]->ConditionalPostLoad();
 		}
 	}
 
 	for (int32 SimulationStageIndex = SimulationStages.Num() - 1; SimulationStageIndex >= 0; --SimulationStageIndex)
 	{
-		if (ensureMsgf(SimulationStages[SimulationStageIndex] != nullptr && SimulationStages[SimulationStageIndex]->Script != nullptr, TEXT("Null simulation stage, or simulation stage with a null script found in %s at index %i, removing it to prevent crashes."), *GetPathName(), SimulationStageIndex) == false)
+		if (ensureMsgf(SimulationStages[SimulationStageIndex] != nullptr && SimulationStages[SimulationStageIndex]->Script != nullptr, TEXT("Null simulation stage, or simulation stage with a null script found in %s at index %i, removing it to prevent crashes."), *Emitter.GetPathName(), SimulationStageIndex) == false)
 		{
 			SimulationStages.RemoveAt(SimulationStageIndex);
 		}
 		else
 		{
+#if WITH_EDITORONLY_DATA
+			SimulationStages[SimulationStageIndex]->OuterEmitterVersion = Version.VersionGuid;
+#endif
 			SimulationStages[SimulationStageIndex]->ConditionalPostLoad();
 		}
 	}
@@ -499,18 +690,17 @@ void UNiagaraEmitter::PostLoad()
 	
 	if (!GPUComputeScript)
 	{
-		GPUComputeScript = NewObject<UNiagaraScript>(this, "GPUComputeScript", EObjectFlags::RF_Transactional);
+		GPUComputeScript = NewObject<UNiagaraScript>(&Emitter, "GPUComputeScript", RF_Transactional);
 		GPUComputeScript->SetUsage(ENiagaraScriptUsage::ParticleGPUComputeScript);
 		GPUComputeScript->SetLatestSource(SpawnScriptProps.Script ? SpawnScriptProps.Script->GetLatestSource() : nullptr);
 	}
-	GPUComputeScript->OnGPUScriptCompiled().AddUObject(this, &UNiagaraEmitter::RaiseOnEmitterGPUCompiled);
 
 	if (EmitterSpawnScriptProps.Script == nullptr || EmitterUpdateScriptProps.Script == nullptr)
 	{
-		EmitterSpawnScriptProps.Script = NewObject<UNiagaraScript>(this, "EmitterSpawnScript", EObjectFlags::RF_Transactional);
+		EmitterSpawnScriptProps.Script = NewObject<UNiagaraScript>(&Emitter, "EmitterSpawnScript", RF_Transactional);
 		EmitterSpawnScriptProps.Script->SetUsage(ENiagaraScriptUsage::EmitterSpawnScript);
 
-		EmitterUpdateScriptProps.Script = NewObject<UNiagaraScript>(this, "EmitterUpdateScript", EObjectFlags::RF_Transactional);
+		EmitterUpdateScriptProps.Script = NewObject<UNiagaraScript>(&Emitter, "EmitterUpdateScript", RF_Transactional);
 		EmitterUpdateScriptProps.Script->SetUsage(ENiagaraScriptUsage::EmitterUpdateScript);
 
 		if (SpawnScriptProps.Script)
@@ -520,41 +710,41 @@ void UNiagaraEmitter::PostLoad()
 		}
 	}
 
-	if (!GetOutermost()->bIsCookedForEditor)
+	if (!Emitter.GetOutermost()->bIsCookedForEditor)
 	{
 		GraphSource->ConditionalPostLoad();
-		GraphSource->PostLoadFromEmitter(*this);
+		GraphSource->PostLoadFromEmitter(FVersionedNiagaraEmitter(&Emitter, Version.VersionGuid));
 		
 		// Prepare for emitter inheritance.
-		if (Parent != nullptr)
+		if (GetParent().Emitter != nullptr)
 		{
-			Parent->ConditionalPostLoad();
+			GetParent().Emitter->ConditionalPostLoad();
 		}
-		if (ParentAtLastMerge != nullptr)
+		if (GetParentAtLastMerge().Emitter != nullptr)
 		{
-			ParentAtLastMerge->ConditionalPostLoad();
+			GetParentAtLastMerge().Emitter->ConditionalPostLoad();
 		}
 
-		for (auto ScratchPadScript : ScratchPadScripts)
+		for (auto ParentScratchPadScript : ParentScratchPads->Scripts)
 		{
-			if (ScratchPadScript)
+			if (ParentScratchPadScript)
 			{
-				ScratchPadScript->ConditionalPostLoad();
+				ParentScratchPadScript->ConditionalPostLoad();
 			}
 		}
 
-		for (auto ParentScratchPadScript : ParentScratchPadScripts)
+		for (auto ScratchPadScript : ScratchPads->Scripts)
 		{
-			if (ParentScratchPadScript)
+			if (ScratchPadScript)
 			{
-				ParentScratchPadScript->ConditionalPostLoad();
+				ScratchPadScript->ConditionalPostLoad();
 			}
 		}
 
-		if (IsSynchronizedWithParent() == false && IsRunningCommandlet())
+		if (Emitter.IsSynchronizedWithParent() == false && IsRunningCommandlet())
 		{
 			// Modify here so that the asset will be marked dirty when using the resave commandlet.  This will be ignored during regular post load.
-			Modify();
+			Emitter.Modify();
 		}
 	}
 #else
@@ -576,22 +766,15 @@ void UNiagaraEmitter::PostLoad()
 #endif
 				SpawnScriptProps.Script->SetUsage(ENiagaraScriptUsage::ParticleSpawnScript);
 			}
-			UE_LOG(LogNiagara, Warning, TEXT("Disabling interpolated spawn because emitter flag and script type don't match. Did you adjust this value in the UI? Emitter may need recompile.. %s"), *GetFullName());
+			UE_LOG(LogNiagara, Warning, TEXT("Disabling interpolated spawn because emitter flag and script type don't match. Did you adjust this value in the UI? Emitter may need recompile.. %s"), *Emitter.GetFullName());
 		}
 	}
 
 	EnsureScriptsPostLoaded();
 
-#if !WITH_EDITOR
-	// When running without the editor in a cooked build we run the update immediately in post load since
-	// there will be no merging or compiling which makes it safe to do so.
-	UpdateEmitterAfterLoad();
-#endif
-
-	//-TODO: should this be a version
 #if WITH_EDITORONLY_DATA
 	// If we have a simulation stage using the old bSpawnOnly we need to modify all generic stages to the new method
-	if ( SimulationStages.ContainsByPredicate([](UNiagaraSimulationStageBase* Stage) { UNiagaraSimulationStageGeneric* StageAsGeneric = Cast<UNiagaraSimulationStageGeneric>(Stage); return StageAsGeneric && StageAsGeneric->bSpawnOnly_DEPRECATED; }) )
+	if (SimulationStages.ContainsByPredicate([](UNiagaraSimulationStageBase* Stage) { UNiagaraSimulationStageGeneric* StageAsGeneric = Cast<UNiagaraSimulationStageGeneric>(Stage); return StageAsGeneric && StageAsGeneric->bSpawnOnly_DEPRECATED; }))
 	{
 		for (UNiagaraSimulationStageBase* Stage : SimulationStages)
 		{
@@ -618,7 +801,7 @@ bool UNiagaraEmitter::IsEditorOnly() const
 	{
 		for (const auto& SystemEmitterHandle : SystemOwner->GetEmitterHandles())
 		{
-			if (SystemEmitterHandle.GetInstance() == this)
+			if (SystemEmitterHandle.GetInstance().Emitter == this)
 			{
 				return false;
 			}
@@ -635,16 +818,19 @@ bool UNiagaraEmitter::IsEditorOnly() const
 void UNiagaraEmitter::GetAssetRegistryTags(TArray<FAssetRegistryTag>& OutTags) const
 {
 #if WITH_EDITOR
-	OutTags.Add(FAssetRegistryTag("HasGPUEmitter", (SimTarget == ENiagaraSimTarget::GPUComputeSim) ? TEXT("True") : TEXT("False"), FAssetRegistryTag::TT_Alphabetical));
+	const FVersionedNiagaraEmitterData* AssetData = GetLatestEmitterData();
+	FVersionedNiagaraEmitterData DefaultData;
+	const FVersionedNiagaraEmitterData& EmitterData = AssetData ? *AssetData : DefaultData; // the CDO does not have any version data, so just use the default struct values in that case 
+	OutTags.Add(FAssetRegistryTag("HasGPUEmitter", ( EmitterData.SimTarget == ENiagaraSimTarget::GPUComputeSim) ? TEXT("True") : TEXT("False"), FAssetRegistryTag::TT_Alphabetical));
 
-	const float BoundsSize = FixedBounds.GetSize().GetMax();
-	OutTags.Add(FAssetRegistryTag("FixedBoundsSize", CalculateBoundsMode == ENiagaraEmitterCalculateBoundMode::Fixed ? FString::Printf(TEXT("%.2f"), BoundsSize) : FString(TEXT("None")), FAssetRegistryTag::TT_Numerical));
+	const float BoundsSize =  EmitterData.FixedBounds.GetSize().GetMax();
+	OutTags.Add(FAssetRegistryTag("FixedBoundsSize",  EmitterData.CalculateBoundsMode == ENiagaraEmitterCalculateBoundMode::Fixed ? FString::Printf(TEXT("%.2f"), BoundsSize) : FString(TEXT("None")), FAssetRegistryTag::TT_Numerical));
 
 
 	uint32 NumActiveRenderers = 0;
 	TArray<const UNiagaraRendererProperties*> ActiveRenderers;
 
-	for (const UNiagaraRendererProperties* Props : GetRenderers())
+	for (const UNiagaraRendererProperties* Props :  EmitterData.GetRenderers())
 	{
 		if (Props)
 		{
@@ -656,8 +842,7 @@ void UNiagaraEmitter::GetAssetRegistryTags(TArray<FAssetRegistryTag>& OutTags) c
 	OutTags.Add(FAssetRegistryTag("ActiveRenderers", LexToString(NumActiveRenderers), FAssetRegistryTag::TT_Numerical));
 
 	// Gather up NumActive emitters based off of quality level.
-	const UNiagaraSettings* Settings = GetDefault<UNiagaraSettings>();
-	if (Settings)
+	if (const UNiagaraSettings* Settings = GetDefault<UNiagaraSettings>())
 	{
 		int32 NumQualityLevels = Settings->QualityLevels.Num();
 		TArray<int32> QualityLevelsNumActive;
@@ -666,7 +851,7 @@ void UNiagaraEmitter::GetAssetRegistryTags(TArray<FAssetRegistryTag>& OutTags) c
 		// Keeping structure from UNiagaraSystem for easy code comparison
 		for (int32 i = 0; i < NumQualityLevels; i++)
 		{
-			if (Platforms.IsEffectQualityEnabled(i))
+			if ( EmitterData.Platforms.IsEffectQualityEnabled(i))
 			{
 				QualityLevelsNumActive[i]++;
 			}
@@ -691,7 +876,7 @@ void UNiagaraEmitter::GetAssetRegistryTags(TArray<FAssetRegistryTag>& OutTags) c
 			const UNiagaraRendererProperties* PropDefault = RendererClass->GetDefaultObject< UNiagaraRendererProperties>();
 			if (PropDefault)
 			{
-				PropDefault->GetAssetTagsForContext(this, ActiveRenderers, NumericKeys, StringKeys);
+				PropDefault->GetAssetTagsForContext(this, EmitterData.Version.VersionGuid, ActiveRenderers, NumericKeys, StringKeys);
 			}
 		}
 	}
@@ -713,7 +898,7 @@ void UNiagaraEmitter::GetAssetRegistryTags(TArray<FAssetRegistryTag>& OutTags) c
 
 	
 		TArray<UNiagaraScript*> Scripts;
-		GetScripts(Scripts);
+		EmitterData.GetScripts(Scripts);
 		for (UNiagaraScript* Script : Scripts)
 		{
 			AddDIs(Script);
@@ -727,7 +912,7 @@ void UNiagaraEmitter::GetAssetRegistryTags(TArray<FAssetRegistryTag>& OutTags) c
 			const UNiagaraDataInterface* PropDefault = DIClass->GetDefaultObject< UNiagaraDataInterface>();
 			if (PropDefault)
 			{
-				PropDefault->GetAssetTagsForContext(this, DataInterfaces, NumericKeys, StringKeys);
+				PropDefault->GetAssetTagsForContext(this, EmitterData.Version.VersionGuid, DataInterfaces, NumericKeys, StringKeys);
 			}
 		}
 		OutTags.Add(FAssetRegistryTag("ActiveDIs", LexToString(DataInterfaces.Num()), FAssetRegistryTag::TT_Numerical));
@@ -768,7 +953,7 @@ bool UNiagaraEmitter::NeedsLoadForTargetPlatform(const ITargetPlatform* TargetPl
 	{
 		for (const FNiagaraEmitterHandle& EmitterHandle : OwnerSystem->GetEmitterHandles())
 		{
-			if (EmitterHandle.GetInstance() == this)
+			if (EmitterHandle.GetInstance().Emitter == this)
 			{
 				if (!EmitterHandle.GetIsEnabled())
 				{
@@ -794,18 +979,21 @@ bool UNiagaraEmitter::NeedsLoadForTargetPlatform(const ITargetPlatform* TargetPl
 
 #if WITH_EDITOR
 /** Creates a new emitter with the supplied emitter as a parent emitter and the supplied system as it's owner. */
-UNiagaraEmitter* UNiagaraEmitter::CreateWithParentAndOwner(UNiagaraEmitter& InParentEmitter, UObject* InOwner, FName InName, EObjectFlags FlagMask)
-{
-	UNiagaraEmitter* NewEmitter = Cast<UNiagaraEmitter>(StaticDuplicateObject(&InParentEmitter, InOwner, InName, FlagMask));
-	NewEmitter->Parent = &InParentEmitter;
-	NewEmitter->ParentAtLastMerge = Cast<UNiagaraEmitter>(StaticDuplicateObject(&InParentEmitter, NewEmitter));
-	NewEmitter->ParentAtLastMerge->ClearFlags(RF_Standalone | RF_Public);
-	NewEmitter->ParentScratchPadScripts.Append(NewEmitter->ScratchPadScripts);
-	NewEmitter->ScratchPadScripts.Empty();
+UNiagaraEmitter* UNiagaraEmitter::CreateWithParentAndOwner(FVersionedNiagaraEmitter InParentEmitter, UObject* InOwner, FName InName, EObjectFlags FlagMask)
+{
+	UNiagaraEmitter* NewEmitter = Cast<UNiagaraEmitter>(StaticDuplicateObject(InParentEmitter.Emitter, InOwner, InName, FlagMask));
+	NewEmitter->DisableVersioning(InParentEmitter.Version);
 	NewEmitter->SetUniqueEmitterName(InName.GetPlainNameString());
-	NewEmitter->GraphSource->MarkNotSynchronized(InitialNotSynchronizedReason);
-	NewEmitter->BindNotifications();
-
+	FVersionedNiagaraEmitterData* EmitterData = NewEmitter->GetLatestEmitterData();
+	EmitterData->VersionedParent = InParentEmitter;
+	EmitterData->VersionedParentAtLastMerge.Emitter = Cast<UNiagaraEmitter>(StaticDuplicateObject(InParentEmitter.Emitter, NewEmitter));
+	EmitterData->VersionedParentAtLastMerge.Emitter->ClearFlags(RF_Standalone | RF_Public);
+	EmitterData->VersionedParentAtLastMerge.Emitter->DisableVersioning(InParentEmitter.Version);
+	EmitterData->VersionedParentAtLastMerge.Version = InParentEmitter.Version;
+	EmitterData->ParentScratchPads->SetScripts(EmitterData->ScratchPads->Scripts);
+	EmitterData->ScratchPads->Scripts.Empty();
+	EmitterData->GraphSource->MarkNotSynchronized(InitialNotSynchronizedReason);
+	NewEmitter->RebindNotifications();
 	return NewEmitter;
 }
 
@@ -814,15 +1002,26 @@ UNiagaraEmitter* UNiagaraEmitter::CreateAsDuplicate(const UNiagaraEmitter& InEmi
 {
 	UNiagaraEmitter* NewEmitter = Cast<UNiagaraEmitter>(StaticDuplicateObject(&InEmitterToDuplicate, &InDuplicateOwnerSystem));
 	NewEmitter->ClearFlags(RF_Standalone | RF_Public);
-	NewEmitter->Parent = InEmitterToDuplicate.Parent;
-	if (InEmitterToDuplicate.ParentAtLastMerge != nullptr)
+	for (FVersionedNiagaraEmitterData& Data : NewEmitter->VersionData)
 	{
-		NewEmitter->ParentAtLastMerge = Cast<UNiagaraEmitter>(StaticDuplicateObject(InEmitterToDuplicate.ParentAtLastMerge, NewEmitter));
-		NewEmitter->ParentAtLastMerge->ClearFlags(RF_Standalone | RF_Public);
+		if (const FVersionedNiagaraEmitterData* OldData = InEmitterToDuplicate.GetEmitterData(Data.Version.VersionGuid))
+		{
+			ensure(Data.VersionedParent == OldData->VersionedParent);
+			if (OldData->VersionedParentAtLastMerge.Emitter != nullptr)
+			{
+				Data.VersionedParentAtLastMerge.Emitter = Cast<UNiagaraEmitter>(StaticDuplicateObject(OldData->VersionedParentAtLastMerge.Emitter, NewEmitter));
+				Data.VersionedParentAtLastMerge.Emitter->ClearFlags(RF_Standalone | RF_Public);
+				Data.VersionedParentAtLastMerge.Emitter->DisableVersioning(OldData->VersionedParentAtLastMerge.Version);
+				Data.VersionedParentAtLastMerge.Version = OldData->VersionedParentAtLastMerge.Version;
+			}
+		}
 	}
 	NewEmitter->SetUniqueEmitterName(InDuplicateName.GetPlainNameString());
-	NewEmitter->GraphSource->MarkNotSynchronized(InitialNotSynchronizedReason);
-	NewEmitter->BindNotifications();
+	for (FVersionedNiagaraEmitterData& Data : NewEmitter->VersionData)
+	{
+		Data.GraphSource->MarkNotSynchronized(InitialNotSynchronizedReason);
+	}
+	NewEmitter->RebindNotifications();
 
 	return NewEmitter;
 }
@@ -836,6 +1035,14 @@ void UNiagaraEmitter::PostDuplicate(EDuplicateMode::Type DuplicateMode)
 	{
 		SetUniqueEmitterName(GetFName().GetPlainNameString());
 	}
+	
+	for (FVersionedNiagaraEmitterData& Data : VersionData)
+	{
+#if STATS
+		Data.StatDatabase.Init();
+#endif
+		Data.RuntimeEstimation.Init();
+	}
 }
 
 void UNiagaraEmitter::PostRename(UObject* OldOuter, const FName OldName)
@@ -848,7 +1055,27 @@ void UNiagaraEmitter::PostRename(UObject* OldOuter, const FName OldName)
 	}
 }
 
-void UNiagaraEmitter::PostEditChangeProperty(struct FPropertyChangedEvent& PropertyChangedEvent)
+void UNiagaraEmitter::PostEditChangeProperty(FPropertyChangedEvent& PropertyChangedEvent)
+{
+	Super::PostEditChangeProperty(PropertyChangedEvent);
+
+	// post undo, some of the properties need to be resolved again, as they are reset to their default
+	ResolveScalabilitySettings();
+	for (FVersionedNiagaraEmitterData& Data : VersionData)
+	{
+#if STATS
+		Data.StatDatabase.Init();
+#endif
+		Data.RuntimeEstimation.Init();
+		Data.UpdateDebugName(*this, nullptr);
+	}
+	ThumbnailImageOutOfDate = true;
+	UpdateChangeId(TEXT("PostEditChangeProperty"));
+	
+	OnPropertiesChangedDelegate.Broadcast();
+}
+
+void UNiagaraEmitter::PostEditChangeVersionedProperty(FPropertyChangedEvent& PropertyChangedEvent, const FGuid& Version)
 {
 	Super::PostEditChangeProperty(PropertyChangedEvent);
 
@@ -857,54 +1084,60 @@ void UNiagaraEmitter::PostEditChangeProperty(struct FPropertyChangedEvent& Prope
 	{
 		PropertyName = PropertyChangedEvent.Property->GetFName();
 	}
+	FVersionedNiagaraEmitterData* EmitterData = GetEmitterData(Version);
+	if (!EmitterData)
+	{
+		PostEditChangeProperty(PropertyChangedEvent);
+		return;
+	}
 
 	bool bNeedsRecompile = false;
 	bool bRecomputeExecutionOrder = false;
-	if (PropertyName == GET_MEMBER_NAME_CHECKED(UNiagaraEmitter, bInterpolatedSpawning))
+	if (PropertyName == GET_MEMBER_NAME_CHECKED(FVersionedNiagaraEmitterData, bInterpolatedSpawning))
 	{
-		bool bActualInterpolatedSpawning = SpawnScriptProps.Script->IsInterpolatedParticleSpawnScript();
-		if (bInterpolatedSpawning != bActualInterpolatedSpawning)
+		bool bActualInterpolatedSpawning = EmitterData->SpawnScriptProps.Script->IsInterpolatedParticleSpawnScript();
+		if (EmitterData->bInterpolatedSpawning != bActualInterpolatedSpawning)
 		{
 			//Recompile spawn script if we've altered the interpolated spawn property.
-			SpawnScriptProps.Script->SetUsage(bInterpolatedSpawning ? ENiagaraScriptUsage::ParticleSpawnScriptInterpolated : ENiagaraScriptUsage::ParticleSpawnScript);
-			UE_LOG(LogNiagara, Log, TEXT("Updating script usage: Script->IsInterpolatdSpawn %d Emitter->bInterpolatedSpawning %d"), (int32)SpawnScriptProps.Script->IsInterpolatedParticleSpawnScript(), bInterpolatedSpawning);
-			if (GraphSource != nullptr)
+			EmitterData->SpawnScriptProps.Script->SetUsage(EmitterData->bInterpolatedSpawning ? ENiagaraScriptUsage::ParticleSpawnScriptInterpolated : ENiagaraScriptUsage::ParticleSpawnScript);
+			UE_LOG(LogNiagara, Log, TEXT("Updating script usage: Script->IsInterpolatdSpawn %d Emitter->bInterpolatedSpawning %d"), (int32)EmitterData->SpawnScriptProps.Script->IsInterpolatedParticleSpawnScript(), EmitterData->bInterpolatedSpawning);
+			if (EmitterData->GraphSource != nullptr)
 			{
-				GraphSource->MarkNotSynchronized(TEXT("Emitter interpolated spawn changed"));
+				EmitterData->GraphSource->MarkNotSynchronized(TEXT("Emitter interpolated spawn changed"));
 			}
 			bNeedsRecompile = true;
 		}
 	}
-	else if (PropertyName == GET_MEMBER_NAME_CHECKED(UNiagaraEmitter, SimTarget))
+	else if (PropertyName == GET_MEMBER_NAME_CHECKED(FVersionedNiagaraEmitterData, SimTarget))
 	{
-		if (GraphSource != nullptr)
+		if (EmitterData->GraphSource != nullptr)
 		{
-			GraphSource->MarkNotSynchronized(TEXT("Emitter simulation target changed."));
+			EmitterData->GraphSource->MarkNotSynchronized(TEXT("Emitter simulation target changed."));
 		}
 		bNeedsRecompile = true;
 	}
-	else if (PropertyName == GET_MEMBER_NAME_CHECKED(UNiagaraEmitter, bRequiresPersistentIDs))
+	else if (PropertyName == GET_MEMBER_NAME_CHECKED(FVersionedNiagaraEmitterData, bRequiresPersistentIDs))
 	{
-		if (GraphSource != nullptr)
+		if (EmitterData->GraphSource != nullptr)
 		{
-			GraphSource->MarkNotSynchronized(TEXT("Emitter Requires Persistent IDs changed."));
+			EmitterData->GraphSource->MarkNotSynchronized(TEXT("Emitter Requires Persistent IDs changed."));
 		}
 		bNeedsRecompile = true;
 	}
-	else if (PropertyName == GET_MEMBER_NAME_CHECKED(UNiagaraEmitter, bLocalSpace))
+	else if (PropertyName == GET_MEMBER_NAME_CHECKED(FVersionedNiagaraEmitterData, bLocalSpace))
 	{
-		if (GraphSource != nullptr)
+		if (EmitterData->GraphSource != nullptr)
 		{
-			GraphSource->MarkNotSynchronized(TEXT("Emitter LocalSpace changed."));
+			EmitterData->GraphSource->MarkNotSynchronized(TEXT("Emitter LocalSpace changed."));
 		}
 
 		bNeedsRecompile = true;
 	}
-	else if (PropertyName == GET_MEMBER_NAME_CHECKED(UNiagaraEmitter, bDeterminism))
+	else if (PropertyName == GET_MEMBER_NAME_CHECKED(FVersionedNiagaraEmitterData, bDeterminism))
 	{
-		if (GraphSource != nullptr)
+		if (EmitterData->GraphSource != nullptr)
 		{
-			GraphSource->MarkNotSynchronized(TEXT("Emitter Determinism changed."));
+			EmitterData->GraphSource->MarkNotSynchronized(TEXT("Emitter Determinism changed."));
 		}
 
 		bNeedsRecompile = true;
@@ -913,19 +1146,19 @@ void UNiagaraEmitter::PostEditChangeProperty(struct FPropertyChangedEvent& Prope
 	{
 		bRecomputeExecutionOrder = true;
 	}
-	else if (PropertyName == GET_MEMBER_NAME_CHECKED(UNiagaraEmitter, AttributesToPreserve))
+	else if (PropertyName == GET_MEMBER_NAME_CHECKED(FVersionedNiagaraEmitterData, AttributesToPreserve))
 	{
-		if (GraphSource != nullptr)
+		if (EmitterData->GraphSource != nullptr)
 		{
-			GraphSource->MarkNotSynchronized(TEXT("AttributesToPreserve changed."));
+			EmitterData->GraphSource->MarkNotSynchronized(TEXT("AttributesToPreserve changed."));
 		}
 		bNeedsRecompile = true;
 	}
-	else if (PropertyName == GET_MEMBER_NAME_CHECKED(UNiagaraEmitter, MaxDeltaTimePerTick))
+	else if (PropertyName == GET_MEMBER_NAME_CHECKED(FVersionedNiagaraEmitterData, MaxDeltaTimePerTick))
 	{
-		if (GraphSource != nullptr)
+		if (EmitterData->GraphSource != nullptr)
 		{
-			GraphSource->MarkNotSynchronized(TEXT("MaxDeltaTimePerTick changed."));
+			EmitterData->GraphSource->MarkNotSynchronized(TEXT("MaxDeltaTimePerTick changed."));
 		}
 		bNeedsRecompile = true;
 	}
@@ -939,11 +1172,11 @@ void UNiagaraEmitter::PostEditChangeProperty(struct FPropertyChangedEvent& Prope
 #if WITH_EDITORONLY_DATA
 	if (bNeedsRecompile)
 	{
-		UNiagaraSystem::RequestCompileForEmitter(this);
+		UNiagaraSystem::RequestCompileForEmitter(FVersionedNiagaraEmitter(this, Version));
 	}
 	else if (bRecomputeExecutionOrder)
 	{
-		UNiagaraSystem::RecomputeExecutionOrderForEmitter(this);
+		UNiagaraSystem::RecomputeExecutionOrderForEmitter(FVersionedNiagaraEmitter(this, Version));
 	}
 #endif
 }
@@ -965,17 +1198,18 @@ UNiagaraEmitter::FOnPropertiesChanged& UNiagaraEmitter::OnRenderersChanged()
 	return OnRenderersChangedDelegate;
 }
 
-void UNiagaraEmitter::HandleVariableRenamed(const FNiagaraVariable& InOldVariable, const FNiagaraVariable& InNewVariable, bool bUpdateContexts)
+void UNiagaraEmitter::HandleVariableRenamed(const FNiagaraVariable& InOldVariable, const FNiagaraVariable& InNewVariable, bool bUpdateContexts, FGuid EmitterVersion)
 {
 	// Rename the variable if it is in use by any renderer properties
-	for (UNiagaraRendererProperties* Prop : GetRenderers())
+	FVersionedNiagaraEmitterData* Data = GetEmitterData(EmitterVersion);
+	for (UNiagaraRendererProperties* Prop : Data->GetRenderers())
 	{
 		Prop->Modify(false);
-		Prop->RenameVariable(InOldVariable, InNewVariable, this);
+		Prop->RenameVariable(InOldVariable, InNewVariable, FVersionedNiagaraEmitter(this, EmitterVersion));
 	}
 
 	// Rename any simulation stage iteration sources
-	for (UNiagaraSimulationStageBase* SimStage : SimulationStages)
+	for (UNiagaraSimulationStageBase* SimStage : Data->SimulationStages)
 	{
 		UNiagaraSimulationStageGeneric* GenericStage = Cast<UNiagaraSimulationStageGeneric>(SimStage);
 		if (GenericStage && GenericStage->DataInterface.BoundVariable.GetName() == InOldVariable.GetName())
@@ -985,28 +1219,29 @@ void UNiagaraEmitter::HandleVariableRenamed(const FNiagaraVariable& InOldVariabl
 		}
 	}
 
-	RebuildRendererBindings();
+	Data->RebuildRendererBindings(*this);
 
 	if (bUpdateContexts)
 	{
-		FNiagaraSystemUpdateContext UpdateCtx(this, true);
+		FNiagaraSystemUpdateContext UpdateCtx(FVersionedNiagaraEmitter(this, EmitterVersion), true);
 	}
 }
 
-void UNiagaraEmitter::HandleVariableRemoved(const FNiagaraVariable& InOldVariable, bool bUpdateContexts)
+void UNiagaraEmitter::HandleVariableRemoved(const FNiagaraVariable& InOldVariable, bool bUpdateContexts, FGuid EmitterVersion)
 {
 	// Reset the variable if it is in use by any renderer properties
-	for (UNiagaraRendererProperties* Prop : GetRenderers())
+	FVersionedNiagaraEmitterData* Data = GetEmitterData(EmitterVersion);
+	for (UNiagaraRendererProperties* Prop : Data->GetRenderers())
 	{
 		Prop->Modify(false);
-		Prop->RemoveVariable(InOldVariable, this);
+		Prop->RemoveVariable(InOldVariable, FVersionedNiagaraEmitter(this, EmitterVersion));
 	}
 
-	RebuildRendererBindings();
+	Data->RebuildRendererBindings(*this);
 
 	if (bUpdateContexts)
 	{
-		FNiagaraSystemUpdateContext UpdateCtx(this, true);
+		FNiagaraSystemUpdateContext UpdateCtx(FVersionedNiagaraEmitter(this, EmitterVersion), true);
 	}
 }
 #endif
@@ -1016,7 +1251,10 @@ TArray<UNiagaraScriptSourceBase*> UNiagaraEmitter::GetAllSourceScripts()
 {
 	TArray<UNiagaraScriptSourceBase*> OutScriptSources;
 	TArray<UNiagaraScript*> Scripts;
-	GetScripts(Scripts, false);
+	for (FVersionedNiagaraEmitterData& Data : VersionData)
+	{
+		Data.GetScripts(Scripts, false);
+	}
 	for (UNiagaraScript* Script : Scripts)
 	{
 		if (Script != nullptr)
@@ -1038,26 +1276,36 @@ FString UNiagaraEmitter::GetSourceObjectPathName() const
 
 TArray<UNiagaraEditorParametersAdapterBase*> UNiagaraEmitter::GetEditorOnlyParametersAdapters()
 {
-	return { GetEditorParameters() };
+	return { GetLatestEmitterData()->GetEditorParameters() };
 }
 
-const TSharedPtr<FNiagaraGraphCachedDataBase, ESPMode::ThreadSafe>& UNiagaraEmitter::GetCachedTraversalData() const
+const TSharedPtr<FNiagaraGraphCachedDataBase, ESPMode::ThreadSafe>& UNiagaraEmitter::GetCachedTraversalData(const FGuid& EmitterVersion) const
 {
-	if (CachedTraversalData.IsValid())
-		return CachedTraversalData;
+	const FVersionedNiagaraEmitterData* EmitterData = GetEmitterData(EmitterVersion);
+	if (EmitterData->CachedTraversalData.IsValid())
+	{
+		return EmitterData->CachedTraversalData;
+	}
 
 	INiagaraModule& NiagaraModule = FModuleManager::Get().LoadModuleChecked<INiagaraModule>(TEXT("Niagara"));
-	CachedTraversalData = NiagaraModule.CacheGraphTraversal(this, FGuid());
-	return CachedTraversalData;
+	EmitterData->CachedTraversalData = NiagaraModule.CacheGraphTraversal(this, EmitterVersion);
+	return EmitterData->CachedTraversalData;
 }
 #endif // WITH_EDITORONLY_DATA
 
 bool UNiagaraEmitter::IsEnabledOnPlatform(const FString& PlatformName)const
 {
-	return Platforms.IsEnabledForPlatform(PlatformName);
+	for (const FVersionedNiagaraEmitterData& EmitterData : VersionData)
+	{
+		if (EmitterData.Platforms.IsEnabledForPlatform(PlatformName))
+		{
+			return true;
+		}
+	}
+	return false;
 }
 
-bool UNiagaraEmitter::IsValid()const
+bool FVersionedNiagaraEmitterData::IsValid() const
 {
 	if (!SpawnScriptProps.Script || !UpdateScriptProps.Script)
 	{
@@ -1098,7 +1346,7 @@ bool UNiagaraEmitter::IsValid()const
 	return true;
 }
 
-bool UNiagaraEmitter::IsReadyToRun() const
+bool FVersionedNiagaraEmitterData::IsReadyToRun() const
 {
 	//Check for various failure conditions and bail.
 	if (!UpdateScriptProps.Script || !SpawnScriptProps.Script)
@@ -1139,7 +1387,7 @@ bool UNiagaraEmitter::IsReadyToRun() const
 	return true;
 }
 
-void UNiagaraEmitter::GetScripts(TArray<UNiagaraScript*>& OutScripts, bool bCompilableOnly, bool bEnabledOnly) const
+void FVersionedNiagaraEmitterData::GetScripts(TArray<UNiagaraScript*>& OutScripts, bool bCompilableOnly, bool bEnabledOnly) const
 {
 	OutScripts.Add(SpawnScriptProps.Script);
 	OutScripts.Add(UpdateScriptProps.Script);
@@ -1179,7 +1427,7 @@ void UNiagaraEmitter::GetScripts(TArray<UNiagaraScript*>& OutScripts, bool bComp
 	}
 }
 
-UNiagaraScript* UNiagaraEmitter::GetScript(ENiagaraScriptUsage Usage, FGuid UsageId)
+UNiagaraScript* FVersionedNiagaraEmitterData::GetScript(ENiagaraScriptUsage Usage, FGuid UsageId)
 {
 	TArray<UNiagaraScript*> Scripts;
 	GetScripts(Scripts, false);
@@ -1193,7 +1441,7 @@ UNiagaraScript* UNiagaraEmitter::GetScript(ENiagaraScriptUsage Usage, FGuid Usag
 	return nullptr;
 }
 
-void UNiagaraEmitter::CacheFromCompiledData(const FNiagaraDataSetCompiledData* CompiledData)
+void FVersionedNiagaraEmitterData::CacheFromCompiledData(const FNiagaraDataSetCompiledData* CompiledData, const UNiagaraEmitter& Emitter)
 {
 	bRequiresViewUniformBuffer = false;
 
@@ -1210,7 +1458,7 @@ void UNiagaraEmitter::CacheFromCompiledData(const FNiagaraDataSetCompiledData* C
 	// Initialize bounds calculators - skip creating if we won't ever use it.  We leave the GPU sims in there with the editor so that we can
 	// generate the bounds from the readback in the tool.
 #if !WITH_EDITOR
-	const bool bUseDynamicBounds = CalculateBoundsMode == ENiagaraEmitterCalculateBoundMode::Dynamic && SimTarget == ENiagaraSimTarget::CPUSim;
+	bool bUseDynamicBounds = !bFixedBoundsEnabled && SimTarget == ENiagaraSimTarget::CPUSim;
 	if (bUseDynamicBounds)
 #endif
 	{
@@ -1219,8 +1467,7 @@ void UNiagaraEmitter::CacheFromCompiledData(const FNiagaraDataSetCompiledData* C
 		{
 			if ((Renderer != nullptr) && Renderer->GetIsEnabled())
 			{
-				FNiagaraBoundsCalculator* BoundsCalculator = Renderer->CreateBoundsCalculator();
-				if (BoundsCalculator != nullptr)
+				if (FNiagaraBoundsCalculator* BoundsCalculator = Renderer->CreateBoundsCalculator(); BoundsCalculator != nullptr)
 				{
 					BoundsCalculator->InitAccessors(CompiledData);
 					BoundsCalculators.Emplace(BoundsCalculator);
@@ -1231,7 +1478,7 @@ void UNiagaraEmitter::CacheFromCompiledData(const FNiagaraDataSetCompiledData* C
 
 	// Cache shaders for all GPU scripts
 #if WITH_EDITORONLY_DATA
-	if ( AreAllScriptAndSourcesSynchronized() )
+	if (AreAllScriptAndSourcesSynchronized())
 	{
 		ForEachScript(
 			[](UNiagaraScript* Script)
@@ -1300,23 +1547,35 @@ void UNiagaraEmitter::CacheFromCompiledData(const FNiagaraDataSetCompiledData* C
 	{
 		MaxInstanceCount = 0;
 	}
+	UpdateDebugName(Emitter, CompiledData);
+
+}
 
+void FVersionedNiagaraEmitterData::UpdateDebugName(const UNiagaraEmitter& Emitter, const FNiagaraDataSetCompiledData* CompiledData)
+{
 #if WITH_NIAGARA_DEBUG_EMITTER_NAME
 	// Ensure our debug simulation name is up to date
 	// Only required on uncooked as it can change due to compilation
 	if (!FPlatformProperties::RequiresCookedData())
 	{
 		DebugSimName.Empty();
-		if (const UNiagaraSystem* SystemOwner = Cast<const UNiagaraSystem>(GetOuter()))
+		if (const UNiagaraSystem* SystemOwner = Emitter.GetTypedOuter<class UNiagaraSystem>())
 		{
 			DebugSimName = SystemOwner->GetName();
 			DebugSimName.AppendChar(':');
 		}
-		DebugSimName.Append(GetName());
+		DebugSimName.Append(Emitter.GetName());
+		if (Emitter.IsVersioningEnabled())
+		{
+			DebugSimName.AppendChar(':');
+			DebugSimName.AppendInt(Version.MajorVersion);
+			DebugSimName.AppendChar('.');
+			DebugSimName.AppendInt(Version.MinorVersion);
+		}
 	}
 #endif
 
-	RebuildRendererBindings();
+	RebuildRendererBindings(Emitter);
 
 	if ( GPUComputeScript && CompiledData )
 	{
@@ -1333,7 +1592,7 @@ void UNiagaraEmitter::CacheFromCompiledData(const FNiagaraDataSetCompiledData* C
 	}
 }
 
-void UNiagaraEmitter::RebuildRendererBindings()
+void FVersionedNiagaraEmitterData::RebuildRendererBindings(const UNiagaraEmitter& Emitter)
 {
 #if WITH_EDITORONLY_DATA
 
@@ -1350,7 +1609,7 @@ void UNiagaraEmitter::RebuildRendererBindings()
 	}
 
 	TArray<UNiagaraScript*> TargetScripts;
-	if (UNiagaraSystem* SystemOwner = Cast<UNiagaraSystem>(GetOuter()))
+	if (UNiagaraSystem* SystemOwner = Emitter.GetTypedOuter<UNiagaraSystem>())
 	{
 		TargetScripts.Add(SystemOwner->GetSystemSpawnScript());
 		TargetScripts.Add(SystemOwner->GetSystemUpdateScript());
@@ -1387,7 +1646,7 @@ void UNiagaraEmitter::RebuildRendererBindings()
 #endif
 }
 
-void UNiagaraEmitter::CacheFromShaderCompiled()
+void FVersionedNiagaraEmitterData::CacheFromShaderCompiled()
 {
 	bRequiresViewUniformBuffer = false;
 	if (GPUComputeScript && (SimTarget == ENiagaraSimTarget::GPUComputeSim))
@@ -1429,21 +1688,26 @@ void UNiagaraEmitter::UpdateEmitterAfterLoad()
 	// Synchronize with definitions before merging.
 	PostLoadDefinitionsSubscriptions();
 
-	// Merge with parent if necessary.
+	// If we're owned by another emitter then we're a VersionedParentAtLastMerge
 	if (GetOuter()->IsA<UNiagaraEmitter>())
 	{
+		ensure(VersionData.Num() == 1);
+		
 		// If this emitter is owned by another emitter, remove it's inheritance information so that it doesn't try to merge changes.
-		Parent = nullptr;
-		ParentAtLastMerge = nullptr;
+		VersionData[0].VersionedParent = FVersionedNiagaraEmitter();
+		VersionData[0].VersionedParentAtLastMerge = FVersionedNiagaraEmitter();
 	}
 
-	if (Parent != nullptr)
-	{
-		Parent->UpdateEmitterAfterLoad();
-	}
-	if (ParentAtLastMerge != nullptr)
+	for (FVersionedNiagaraEmitterData& Data : VersionData)
 	{
-		ParentAtLastMerge->UpdateEmitterAfterLoad();
+		if (Data.VersionedParent.Emitter != nullptr)
+		{
+			Data.VersionedParent.Emitter->UpdateEmitterAfterLoad();
+		}
+		if (Data.VersionedParentAtLastMerge.Emitter != nullptr)
+		{
+			Data.VersionedParentAtLastMerge.Emitter->UpdateEmitterAfterLoad();
+		}
 	}
 	
 	if (!GetOutermost()->bIsCookedForEditor)
@@ -1468,8 +1732,11 @@ void UNiagaraEmitter::UpdateEmitterAfterLoad()
 			UObject* OuterObj = GetOuter();
 			if (OuterObj == GetOutermost())
 			{
-				GraphSource->ForceGraphToRecompileOnNextCheck();
-				bGenerateNewChangeId = true;
+				for (FVersionedNiagaraEmitterData& Data : VersionData)
+				{
+					Data.GraphSource->ForceGraphToRecompileOnNextCheck();
+				}
+				bGenerateNewChangeId = true;
 				GenerateNewChangeIdReason = TEXT("PostLoad - Force compile on load");
 				if (GEnableVerboseNiagaraChangeIdLogging)
 				{
@@ -1494,24 +1761,19 @@ void UNiagaraEmitter::UpdateEmitterAfterLoad()
 			UpdateChangeId(GenerateNewChangeIdReason);
 		}
 
-		BindNotifications();
+		RebindNotifications();
 	}
 #endif
 
 	ResolveScalabilitySettings();
 
-#if WITH_NIAGARA_DEBUG_EMITTER_NAME
-	DebugSimName.Empty();
-	if (const UNiagaraSystem* SystemOwner = Cast<const UNiagaraSystem>(GetOuter()))
+	for (FVersionedNiagaraEmitterData& Data : VersionData)
 	{
-		DebugSimName = SystemOwner->GetName();
-		DebugSimName.AppendChar(':');
+		Data.UpdateDebugName(*this, nullptr);
 	}
-	DebugSimName.Append(GetName());
-#endif
 }
 
-const FNiagaraEmitterScalabilityOverride& UNiagaraEmitter::GetCurrentOverrideSettings() const
+const FNiagaraEmitterScalabilityOverride& FVersionedNiagaraEmitterData::GetCurrentOverrideSettings() const
 {
 	for (const FNiagaraEmitterScalabilityOverride& Override : ScalabilityOverrides.Overrides)
 	{
@@ -1525,12 +1787,12 @@ const FNiagaraEmitterScalabilityOverride& UNiagaraEmitter::GetCurrentOverrideSet
 	return Dummy;
 }
 
-bool UNiagaraEmitter::IsAllowedByScalability()const
+bool FVersionedNiagaraEmitterData::IsAllowedByScalability() const
 {
 	return Platforms.IsActive();
 }
 
-bool UNiagaraEmitter::RequiresPersistentIDs() const
+bool FVersionedNiagaraEmitterData::RequiresPersistentIDs() const
 {
 	return bRequiresPersistentIDs;
 }
@@ -1542,36 +1804,41 @@ FGuid UNiagaraEmitter::GetChangeId() const
 	return ChangeId;
 }
 
-UNiagaraEditorDataBase* UNiagaraEmitter::GetEditorData() const
+UNiagaraEditorDataBase* FVersionedNiagaraEmitterData::GetEditorData() const
 {
 	return EditorData;
 }
 
-UNiagaraEditorParametersAdapterBase* UNiagaraEmitter::GetEditorParameters()
+UNiagaraEditorParametersAdapterBase* FVersionedNiagaraEmitterData::GetEditorParameters()
 {
 	return EditorParameters;
 }
 
-void UNiagaraEmitter::SetEditorData(UNiagaraEditorDataBase* InEditorData)
+void UNiagaraEmitter::SetEditorData(UNiagaraEditorDataBase* InEditorData, const FGuid& VersionGuid)
 {
-	if (EditorData == InEditorData)
+	FVersionedNiagaraEmitterData* EmitterData = GetEmitterData(VersionGuid);
+	if (EmitterData == nullptr)
 	{
 		return;
 	}
- 
-	if (EditorData != nullptr)
+	if (EmitterData->EditorData == InEditorData)
+	{
+		return;
+	}
+	if (EmitterData->EditorData != nullptr)
 	{
-		EditorData->OnPersistentDataChanged().RemoveAll(this);
+		EmitterData->EditorData->OnPersistentDataChanged().RemoveAll(this);
 	}
-	EditorData = InEditorData;
+
+	EmitterData->EditorData = InEditorData;
 	
-	if (EditorData != nullptr)
+	if (EmitterData->EditorData != nullptr)
 	{
-		EditorData->OnPersistentDataChanged().AddUObject(this, &UNiagaraEmitter::PersistentEditorDataChanged);
+		EmitterData->EditorData->OnPersistentDataChanged().AddUObject(this, &UNiagaraEmitter::PersistentEditorDataChanged);
 	}
 }
 
-bool UNiagaraEmitter::AreAllScriptAndSourcesSynchronized() const
+bool FVersionedNiagaraEmitterData::AreAllScriptAndSourcesSynchronized() const
 {
 	if (SpawnScriptProps.Script->IsCompilable() && !SpawnScriptProps.Script->AreScriptAndSourceSynchronized())
 	{
@@ -1628,7 +1895,7 @@ UNiagaraEmitter::FOnEmitterCompiled& UNiagaraEmitter::OnEmitterGPUCompiled()
 	return OnGPUScriptCompiledDelegate;
 }
 
-void  UNiagaraEmitter::InvalidateCompileResults()
+void FVersionedNiagaraEmitterData::InvalidateCompileResults()
 {
 	TArray<UNiagaraScript*> Scripts;
 	GetScripts(Scripts, false);
@@ -1638,9 +1905,21 @@ void  UNiagaraEmitter::InvalidateCompileResults()
 	}
 }
 
-void UNiagaraEmitter::OnPostCompile()
+void UNiagaraEmitter::OnPostCompile(const FGuid& EmitterVersion)
 {
-	SyncEmitterAlias(TEXT("Emitter"), UniqueEmitterName);
+	FVersionedNiagaraEmitterData* Data = GetEmitterData(EmitterVersion);
+	if (Data == nullptr)
+	{
+		return;
+	}
+	Data->OnPostCompile(*this);
+
+	OnEmitterVMCompiled().Broadcast(FVersionedNiagaraEmitter(this, EmitterVersion));
+}
+
+void FVersionedNiagaraEmitterData::OnPostCompile(const UNiagaraEmitter& InEmitter)
+{
+	SyncEmitterAlias(TEXT("Emitter"), InEmitter);
 
 	SpawnScriptProps.InitDataSetAccess();
 	UpdateScriptProps.InitDataSetAccess();
@@ -1683,161 +1962,108 @@ void UNiagaraEmitter::OnPostCompile()
 		GPUComputeScript->InvalidateCompileResults(TEXT("Not a GPU emitter."));
 	}
 
-	RuntimeEstimation = MemoryRuntimeEstimation();
+	RuntimeEstimation.RuntimeAllocations.Empty();
+	RuntimeEstimation.AllocationEstimate = 0;
+	RuntimeEstimation.IsEstimationDirty = false;
 #if STATS
 	StatDatabase.ClearStatCaptures();
 #endif
-
-	OnEmitterVMCompiled().Broadcast(this);
 }
 
-void UNiagaraEmitter::BindNotifications()
+void UNiagaraEmitter::RebindNotifications()
 {
-	if (GraphSource)
-	{
-		GraphSource->OnChanged().AddUObject(this, &UNiagaraEmitter::GraphSourceChanged);
-	}
-
-	if (EmitterSpawnScriptProps.Script)
+	for (FVersionedNiagaraEmitterData& EmitterData : VersionData)
 	{
-		EmitterSpawnScriptProps.Script->RapidIterationParameters.AddOnChangedHandler(
-			FNiagaraParameterStore::FOnChanged::FDelegate::CreateUObject(this, &UNiagaraEmitter::ScriptRapidIterationParameterChanged));
-	}
-
-	if (EmitterUpdateScriptProps.Script)
-	{
-		EmitterUpdateScriptProps.Script->RapidIterationParameters.AddOnChangedHandler(
-			FNiagaraParameterStore::FOnChanged::FDelegate::CreateUObject(this, &UNiagaraEmitter::ScriptRapidIterationParameterChanged));
-	}
-
-	if (SpawnScriptProps.Script)
-	{
-		SpawnScriptProps.Script->RapidIterationParameters.AddOnChangedHandler(
-			FNiagaraParameterStore::FOnChanged::FDelegate::CreateUObject(this, &UNiagaraEmitter::ScriptRapidIterationParameterChanged));
-	}
-
-	if (UpdateScriptProps.Script)
-	{
-		UpdateScriptProps.Script->RapidIterationParameters.AddOnChangedHandler(
-			FNiagaraParameterStore::FOnChanged::FDelegate::CreateUObject(this, &UNiagaraEmitter::ScriptRapidIterationParameterChanged));
-	}
-
-	for (FNiagaraEventScriptProperties& EventScriptProperties : EventHandlerScriptProps)
-	{
-		if (EventScriptProperties.Script)
+		if (EmitterData.GraphSource)
 		{
-			EventScriptProperties.Script->RapidIterationParameters.AddOnChangedHandler(
-				FNiagaraParameterStore::FOnChanged::FDelegate::CreateUObject(this, &UNiagaraEmitter::ScriptRapidIterationParameterChanged));
+			EmitterData.GraphSource->OnChanged().RemoveAll(this);
+			EmitterData.GraphSource->OnChanged().AddUObject(this, &UNiagaraEmitter::GraphSourceChanged);
 		}
-	}
 
-	for (UNiagaraSimulationStageBase* SimulationStage : SimulationStages)
-	{
-		if (SimulationStage)
+		if (EmitterData.EmitterSpawnScriptProps.Script)
 		{
-			SimulationStage->OnChanged().AddUObject(this, &UNiagaraEmitter::SimulationStageChanged);
-
-			if (SimulationStage->Script)
-			{
-				SimulationStage->Script->RapidIterationParameters.AddOnChangedHandler(
-					FNiagaraParameterStore::FOnChanged::FDelegate::CreateUObject(this, &UNiagaraEmitter::ScriptRapidIterationParameterChanged));
-			}
+			FNiagaraParameterStore& Store = EmitterData.EmitterSpawnScriptProps.Script->RapidIterationParameters;
+			Store.RemoveAllOnChangedHandlers(this);
+			Store.AddOnChangedHandler(FNiagaraParameterStore::FOnChanged::FDelegate::CreateUObject(this, &UNiagaraEmitter::ScriptRapidIterationParameterChanged));
 		}
-	}
 
-	for (UNiagaraRendererProperties* Renderer : RendererProperties)
-	{
-		if (Renderer)
+		if (EmitterData.EmitterUpdateScriptProps.Script)
 		{
-			Renderer->OnChanged().AddUObject(this, &UNiagaraEmitter::RendererChanged);
+			FNiagaraParameterStore& Store = EmitterData.EmitterUpdateScriptProps.Script->RapidIterationParameters;
+			Store.RemoveAllOnChangedHandlers(this);
+			Store.AddOnChangedHandler(FNiagaraParameterStore::FOnChanged::FDelegate::CreateUObject(this, &UNiagaraEmitter::ScriptRapidIterationParameterChanged));
 		}
-	}
 
-	if (EditorData != nullptr)
-	{
-		EditorData->OnPersistentDataChanged().AddUObject(this, &UNiagaraEmitter::PersistentEditorDataChanged);
-	}
-}
+		if (EmitterData.SpawnScriptProps.Script)
+		{
+			FNiagaraParameterStore& Store = EmitterData.SpawnScriptProps.Script->RapidIterationParameters;
+			Store.RemoveAllOnChangedHandlers(this);
+			Store.AddOnChangedHandler(FNiagaraParameterStore::FOnChanged::FDelegate::CreateUObject(this, &UNiagaraEmitter::ScriptRapidIterationParameterChanged));
+		}
 
-void UNiagaraEmitter::GatherStaticVariables(TArray<FNiagaraVariable>& OutVars) const
-{
-	TArray<UNiagaraScript*> OutScripts;
-	GetScripts(OutScripts, false, true);
+		if (EmitterData.UpdateScriptProps.Script)
+		{
+			FNiagaraParameterStore& Store = EmitterData.UpdateScriptProps.Script->RapidIterationParameters;
+			Store.RemoveAllOnChangedHandlers(this);
+			Store.AddOnChangedHandler(FNiagaraParameterStore::FOnChanged::FDelegate::CreateUObject(this, &UNiagaraEmitter::ScriptRapidIterationParameterChanged));
+		}
 
-	for (UNiagaraScript* Script : OutScripts)
-	{
-		TArray<FNiagaraVariable> StoreParams;
-		Script->RapidIterationParameters.GetParameters(StoreParams);
+		for (FNiagaraEventScriptProperties& EventScriptProperties : EmitterData.EventHandlerScriptProps)
+		{
+			if (EventScriptProperties.Script)
+			{
+				FNiagaraParameterStore& Store = EventScriptProperties.Script->RapidIterationParameters;
+				Store.RemoveAllOnChangedHandlers(this);
+				Store.AddOnChangedHandler(FNiagaraParameterStore::FOnChanged::FDelegate::CreateUObject(this, &UNiagaraEmitter::ScriptRapidIterationParameterChanged));
+			}
+		}
 
-		for (int32 i = 0; i < StoreParams.Num(); i++)
+		for (UNiagaraSimulationStageBase* SimulationStage : EmitterData.SimulationStages)
 		{
-			if (StoreParams[i].GetType().IsStatic())
+			if (SimulationStage)
 			{
-				const int32* Index = Script->RapidIterationParameters.FindParameterOffset(StoreParams[i]);
-				if (Index != nullptr)
+				SimulationStage->OuterEmitterVersion = EmitterData.Version.VersionGuid;
+				SimulationStage->OnChanged().RemoveAll(this);
+				SimulationStage->OnChanged().AddUObject(this, &UNiagaraEmitter::SimulationStageChanged);
+
+				if (SimulationStage->Script)
 				{
-					StoreParams[i].SetData(Script->RapidIterationParameters.GetParameterData(*Index)); // This will memcopy the data in.			
-					OutVars.AddUnique(StoreParams[i]);					
-					//UE_LOG(LogNiagara, Log, TEXT("UNiagaraEmitter::GatherStaticVariables Added %s"), *StoreParams[i].ToString());
+					FNiagaraParameterStore& Store = SimulationStage->Script->RapidIterationParameters;
+					Store.RemoveAllOnChangedHandlers(this);
+					Store.AddOnChangedHandler(FNiagaraParameterStore::FOnChanged::FDelegate::CreateUObject(this, &UNiagaraEmitter::ScriptRapidIterationParameterChanged));
 				}
 			}
 		}
-	}
-}
-
 
-#endif
-
-bool UNiagaraEmitter::UsesScript(const UNiagaraScript* Script)const
-{
-	if (SpawnScriptProps.Script == Script || UpdateScriptProps.Script == Script)
-	{
-		return true;
-	}
-#if WITH_EDITORONLY_DATA
-	if (EmitterSpawnScriptProps.Script == Script || EmitterUpdateScriptProps.Script == Script)
-	{
-		return true;
-	}
-#endif
-	for (int32 i = 0; i < EventHandlerScriptProps.Num(); i++)
-	{
-		if (EventHandlerScriptProps[i].Script == Script)
+		for (UNiagaraRendererProperties* Renderer : EmitterData.RendererProperties)
 		{
-			return true;
+			if (Renderer)
+			{
+				Renderer->OuterEmitterVersion = EmitterData.Version.VersionGuid;
+				Renderer->OnChanged().RemoveAll(this);
+				Renderer->OnChanged().AddUObject(this, &UNiagaraEmitter::RendererChanged);
+			}
 		}
-	}
-	return false;
-}
-
-bool UNiagaraEmitter::UsesCollection(const class UNiagaraParameterCollection* Collection)const
-{
-	if (SpawnScriptProps.Script && SpawnScriptProps.Script->UsesCollection(Collection))
-	{
-		return true;
-	}
-	if (UpdateScriptProps.Script && UpdateScriptProps.Script->UsesCollection(Collection))
-	{
-		return true;
-	}
-	for (int32 i = 0; i < EventHandlerScriptProps.Num(); i++)
-	{
-		if (EventHandlerScriptProps[i].Script && EventHandlerScriptProps[i].Script->UsesCollection(Collection))
+		
+		if (EmitterData.EditorData != nullptr)
 		{
-			return true;
+			EmitterData.EditorData->OnPersistentDataChanged().RemoveAll(this);
+			EmitterData.EditorData->OnPersistentDataChanged().AddUObject(this, &UNiagaraEmitter::PersistentEditorDataChanged);
 		}
 	}
-	return false;
 }
 
+#endif
 
-bool UNiagaraEmitter::CanObtainParticleAttribute(const FNiagaraVariableBase& InVar, FNiagaraTypeDefinition& OutBoundType) const
+bool UNiagaraEmitter::CanObtainParticleAttribute(const FNiagaraVariableBase& InVar, const FGuid& EmitterVersion, FNiagaraTypeDefinition& OutBoundType) const
 {
 	check(!HasAnyFlags(RF_NeedPostLoad));
 
-	if (SpawnScriptProps.Script)
+	const FVersionedNiagaraEmitterData* EmitterData = GetEmitterData(EmitterVersion);
+	if (EmitterData && EmitterData->SpawnScriptProps.Script)
 	{
 		// make sure that this isn't called before our dependents are fully loaded
+		const FNiagaraEmitterScriptProperties& SpawnScriptProps = EmitterData->SpawnScriptProps;
 		check(!SpawnScriptProps.Script->HasAnyFlags(RF_NeedPostLoad));
 
 		bool bContainsAttribute = SpawnScriptProps.Script->GetVMExecutableData().Attributes.Contains(InVar);
@@ -1853,12 +2079,12 @@ bool UNiagaraEmitter::CanObtainParticleAttribute(const FNiagaraVariableBase& InV
 	}
 	return false;
 }
+
 bool UNiagaraEmitter::CanObtainEmitterAttribute(const FNiagaraVariableBase& InVarWithUniqueNameNamespace, FNiagaraTypeDefinition& OutBoundType) const
 {
 	check(!HasAnyFlags(RF_NeedPostLoad));
 
-	const UNiagaraSystem* Sys = GetTypedOuter<UNiagaraSystem>();
-	if (Sys)
+	if (const UNiagaraSystem* Sys = GetTypedOuter<UNiagaraSystem>())
 	{
 		// make sure that this isn't called before our dependents are fully loaded
 		check(!Sys->HasAnyFlags(RF_NeedPostLoad));
@@ -1867,12 +2093,12 @@ bool UNiagaraEmitter::CanObtainEmitterAttribute(const FNiagaraVariableBase& InVa
 	}
 	return false;
 }
+
 bool UNiagaraEmitter::CanObtainSystemAttribute(const FNiagaraVariableBase& InVar, FNiagaraTypeDefinition& OutBoundType) const
 {
 	check(!HasAnyFlags(RF_NeedPostLoad));
 
-	const UNiagaraSystem* Sys = GetTypedOuter<UNiagaraSystem>();
-	if (Sys)
+	if (const UNiagaraSystem* Sys = GetTypedOuter<UNiagaraSystem>())
 	{
 		// make sure that this isn't called before our dependents are fully loaded
 		check(!Sys->HasAnyFlags(RF_NeedPostLoad));
@@ -1881,12 +2107,12 @@ bool UNiagaraEmitter::CanObtainSystemAttribute(const FNiagaraVariableBase& InVar
 	}
 	return false;
 }
+
 bool UNiagaraEmitter::CanObtainUserVariable(const FNiagaraVariableBase& InVar) const
 {
 	check(!HasAnyFlags(RF_NeedPostLoad));
 
-	const UNiagaraSystem* Sys = GetTypedOuter<UNiagaraSystem>();
-	if (Sys)
+	if (const UNiagaraSystem* Sys = GetTypedOuter<UNiagaraSystem>())
 	{
 		// make sure that this isn't called before our dependents are fully loaded
 		check(!Sys->HasAnyFlags(RF_NeedPostLoad));
@@ -1901,9 +2127,432 @@ const FString& UNiagaraEmitter::GetUniqueEmitterName()const
 	return UniqueEmitterName;
 }
 
+FVersionedNiagaraEmitterData* UNiagaraEmitter::GetLatestEmitterData()
+{
+	if (VersionData.Num() == 0)
+	{
+		return nullptr;
+	}
+	if (!bVersioningEnabled)
+	{
+		return &VersionData[0];
+	}
+	FVersionedNiagaraEmitterData* EmitterData = GetEmitterData(ExposedVersion);
+	ensureMsgf(EmitterData, TEXT("Invalid exposed version for Niagara emitter %s, asset might be corrupted!"), *GetPathName());
+	return EmitterData;
+}
+
+const FVersionedNiagaraEmitterData* UNiagaraEmitter::GetLatestEmitterData() const
+{
+	return const_cast<UNiagaraEmitter*>(this)->GetLatestEmitterData();
+}
+
+const FVersionedNiagaraEmitterData* UNiagaraEmitter::GetEmitterData(const FGuid& VersionGuid) const
+{
+	return const_cast<UNiagaraEmitter*>(this)->GetEmitterData(VersionGuid);
+}
+
+FVersionedNiagaraEmitterData* UNiagaraEmitter::GetEmitterData(const FGuid& VersionGuid)
+{
+	if (VersionData.Num() == 0)
+	{
+		return nullptr;
+	}
+	
+	// check if we even need to support different versions
+	if (!bVersioningEnabled)
+	{
+		return &VersionData[0];
+	}
+
+	if (!VersionGuid.IsValid())
+	{
+		for (FVersionedNiagaraEmitterData& Data : VersionData)
+		{
+			if (Data.Version.VersionGuid == ExposedVersion)
+			{
+				return &Data;
+			}
+		}
+		ensureMsgf(false, TEXT("Invalid exposed version for Niagara emitter %s, asset might be corrupted!"), *GetPathName());
+		return nullptr;
+	}
+	for (FVersionedNiagaraEmitterData& Data : VersionData)
+	{
+		if (Data.Version.VersionGuid == VersionGuid)
+		{
+			return &Data;
+		}
+	}
+	return nullptr;
+}
+
+TArray<FNiagaraAssetVersion> UNiagaraEmitter::GetAllAvailableVersions() const
+{
+	TArray<FNiagaraAssetVersion> Versions;
+	for (const FVersionedNiagaraEmitterData& Data : VersionData)
+	{
+		Versions.Add(Data.Version);
+	}
+	return Versions;
+}
+
 #if WITH_EDITORONLY_DATA
 
-void UNiagaraEmitter::UpdateFromMergedCopy(const INiagaraMergeManager& MergeManager, UNiagaraEmitter* MergedEmitter)
+struct FNiagaraEmitterVersionDataAccessor final : FNiagaraVersionDataAccessor
+{
+	virtual ~FNiagaraEmitterVersionDataAccessor() override = default;
+	explicit FNiagaraEmitterVersionDataAccessor(FVersionedNiagaraEmitterData* InEmitterData) : EmitterData(InEmitterData) {}
+
+	virtual FNiagaraAssetVersion& GetObjectVersion() override { return EmitterData->Version; }
+	virtual FText& GetVersionChangeDescription() override { return EmitterData->VersionChangeDescription; }
+	virtual ENiagaraPythonUpdateScriptReference& GetUpdateScriptExecutionType() override { return EmitterData->UpdateScriptExecution; }
+	virtual FString& GetPythonUpdateScript() override { return EmitterData->PythonUpdateScript; }
+	virtual FFilePath& GetScriptAsset() override { return EmitterData->ScriptAsset; }
+	
+	FVersionedNiagaraEmitterData* EmitterData;
+};
+
+TSharedPtr<FNiagaraVersionDataAccessor> UNiagaraEmitter::GetVersionDataAccessor(const FGuid& Version)
+{
+	if (FVersionedNiagaraEmitterData* EmitterData = GetEmitterData(Version))
+	{
+		return MakeShared<FNiagaraEmitterVersionDataAccessor>(EmitterData);
+	}
+	return TSharedPtr<FNiagaraVersionDataAccessor>();
+}
+
+FNiagaraAssetVersion UNiagaraEmitter::GetExposedVersion() const
+{
+	const FVersionedNiagaraEmitterData* EmitterData = GetLatestEmitterData();
+	return EmitterData ? EmitterData->Version : FNiagaraAssetVersion();
+}
+
+FNiagaraAssetVersion const* UNiagaraEmitter::FindVersionData(const FGuid& VersionGuid) const
+{
+	for (const FVersionedNiagaraEmitterData& Data : VersionData)
+	{
+		if (Data.Version.VersionGuid == VersionGuid)
+		{
+			return &Data.Version;
+		}
+	}
+	return nullptr;
+}
+
+UNiagaraScriptSourceBase* DuplicateScriptSource(UNiagaraScriptSourceBase* Source, UNiagaraEmitter* Outer, const FNiagaraAssetVersion& NewVersion)
+{
+	if (Source)
+	{
+		FObjectDuplicationParameters ObjParameters(Source, Outer);
+		ObjParameters.DestClass = Source->GetClass();
+		ObjParameters.DestName = MakeUniqueObjectName(Outer, Source->GetClass(), GetVersionedName("NiagaraScriptSource", NewVersion));
+		if (UNiagaraScriptSourceBase* NewSource = Cast<UNiagaraScriptSourceBase>(StaticDuplicateObjectEx(ObjParameters)))
+		{
+			return NewSource;
+		}
+	}
+	return nullptr;
+}
+
+UNiagaraScript* DuplicateScript(UNiagaraScript* Script, UObject* Outer, UNiagaraScriptSourceBase* Source, FString ScriptName, const FNiagaraAssetVersion& NewVersion)
+{
+	if (Script)
+	{
+		FObjectDuplicationParameters ObjParameters(Script, Outer);
+		ObjParameters.DestClass = UNiagaraScript::StaticClass();
+		ObjParameters.DestName = MakeUniqueObjectName(Outer, UNiagaraScript::StaticClass(), GetVersionedName(ScriptName, NewVersion));
+		ObjParameters.ApplyFlags = RF_Transactional;
+		if (UNiagaraScript* NewScript = Cast<UNiagaraScript>(StaticDuplicateObjectEx(ObjParameters)))
+		{
+			if (Source)
+			{
+				NewScript->SetLatestSource(Source);
+			}
+			return NewScript;
+		}
+	}
+	return nullptr;
+}
+
+UNiagaraSimulationStageBase* DuplicateSimStage(UNiagaraSimulationStageBase* SimulationStage, UNiagaraEmitter* Outer, UNiagaraScriptSourceBase* Source, const FNiagaraAssetVersion& NewVersion)
+{
+	if (SimulationStage)
+	{
+		FObjectDuplicationParameters ObjParameters(SimulationStage, Outer);
+		ObjParameters.DestClass = SimulationStage->GetClass();
+		ObjParameters.DestName = MakeUniqueObjectName(Outer, SimulationStage->GetClass(), GetVersionedName(SimulationStage->GetClass()->GetName(), NewVersion));
+		ObjParameters.ApplyFlags = RF_Transactional;
+		if (UNiagaraSimulationStageBase* NewStage = Cast<UNiagaraSimulationStageBase>(StaticDuplicateObjectEx(ObjParameters)))
+		{
+			NewStage->Script = DuplicateScript(SimulationStage->Script, NewStage, Source, "SimulationStage", NewVersion);
+			return NewStage;
+		}
+	}
+	return nullptr;
+}
+
+FGuid UNiagaraEmitter::AddNewVersion(int32 MajorVersion, int32 MinorVersion)
+{
+	// check preconditions
+	check(MajorVersion >= 1);
+	check(MajorVersion != 1 || MinorVersion != 0);
+	Modify();
+
+	FVersionedNiagaraEmitterData NewVersionData;
+	FNiagaraAssetVersion NewVersion = { MajorVersion, MinorVersion, FGuid::NewGuid() };
+	
+	for (int i = VersionData.Num() - 1; i >= 0; i--)
+	{
+		FVersionedNiagaraEmitterData& Data = VersionData[i];
+		check(Data.Version.MajorVersion != MajorVersion || Data.Version.MinorVersion != MinorVersion); // the version should not already exist
+
+		if (Data.Version.MajorVersion < MajorVersion || (Data.Version.MajorVersion == MajorVersion && Data.Version.MinorVersion < MinorVersion))
+		{
+			// copy the data
+			NewVersionData.CopyFrom(Data);
+
+			// duplicate the scripts and source graph
+			NewVersionData.GraphSource = DuplicateScriptSource(Data.GraphSource, this, NewVersion);
+			NewVersionData.SpawnScriptProps.Script = DuplicateScript(Data.SpawnScriptProps.Script, this, NewVersionData.GraphSource, "SpawnScript", NewVersion);
+			NewVersionData.UpdateScriptProps.Script = DuplicateScript(Data.UpdateScriptProps.Script, this, NewVersionData.GraphSource, "UpdateScript", NewVersion);
+			NewVersionData.EmitterSpawnScriptProps.Script = DuplicateScript(Data.EmitterSpawnScriptProps.Script, this, NewVersionData.GraphSource, "EmitterSpawnScript", NewVersion);
+			NewVersionData.EmitterUpdateScriptProps.Script = DuplicateScript(Data.EmitterUpdateScriptProps.Script, this, NewVersionData.GraphSource, "EmitterUpdateScript", NewVersion);
+			NewVersionData.GPUComputeScript = DuplicateScript(Data.GPUComputeScript, this, NewVersionData.GraphSource, "GPUComputeScript", NewVersion);
+
+			// duplicate the event scripts
+			for (FNiagaraEventScriptProperties& EventProp : NewVersionData.EventHandlerScriptProps)
+			{
+				EventProp.Script = DuplicateScript(EventProp.Script, this, NewVersionData.GraphSource, "EventScript", NewVersion);
+			}
+
+			// duplicate the simulation stages
+			NewVersionData.SimulationStages.Empty();
+			for (UNiagaraSimulationStageBase* SimulationStage : Data.SimulationStages)
+			{
+				UNiagaraSimulationStageBase* NewStage = DuplicateSimStage(SimulationStage, this, NewVersionData.GraphSource, NewVersion);
+				NewStage->OuterEmitterVersion = NewVersion.VersionGuid;
+				NewVersionData.SimulationStages.Add(NewStage);
+			}
+
+			// duplicate the renderer settings
+			NewVersionData.RendererProperties.Empty();
+			for (UNiagaraRendererProperties* Renderer : Data.RendererProperties)
+			{
+				UNiagaraRendererProperties* NewRenderer = Cast<UNiagaraRendererProperties>(StaticDuplicateObject(Renderer, this, NAME_None, RF_Transactional));
+				NewRenderer->OuterEmitterVersion = NewVersion.VersionGuid;
+				NewVersionData.RendererProperties.Add(NewRenderer);
+			}
+
+			// duplicate the scratch pad scripts and replace their usage in the graphs
+			Data.ScratchPads->CheckConsistency();
+			NewVersionData.ScratchPads = CastChecked<UNiagaraScratchPadContainer>(StaticDuplicateObject(Data.ScratchPads, this));
+			for (int k = 0; k < Data.ScratchPads->Scripts.Num(); k++) {
+				UNiagaraScript* OldScript = Data.ScratchPads->Scripts[k];
+				UNiagaraScript* NewScript = NewVersionData.ScratchPads->Scripts[k];
+				NewVersionData.GraphSource->ReplaceScriptReferences(OldScript, NewScript);
+			}
+
+			// duplicate the parent scratch pad script container.
+			Data.ParentScratchPads->CheckConsistency();
+			NewVersionData.ParentScratchPads = CastChecked<UNiagaraScratchPadContainer>(StaticDuplicateObject(Data.ParentScratchPads, this));
+			for (int k = 0; k < Data.ParentScratchPads->Scripts.Num(); k++) {
+				UNiagaraScript* OldScript = Data.ParentScratchPads->Scripts[k];
+				UNiagaraScript* NewScript = NewVersionData.ParentScratchPads->Scripts[k];
+				NewVersionData.GraphSource->ReplaceScriptReferences(OldScript, NewScript);
+			}
+
+			// duplicate the editor data
+			if (Data.EditorData)
+			{
+				NewVersionData.EditorData = CastChecked<UNiagaraEditorDataBase>(StaticDuplicateObject(Data.EditorData, this));
+			}
+			if (Data.EditorParameters)
+			{
+				NewVersionData.EditorParameters = CastChecked<UNiagaraEditorParametersAdapterBase>(StaticDuplicateObject(Data.EditorParameters, this));
+			}
+
+			// init stats critical sections
+#if STATS
+			NewVersionData.StatDatabase.Init();
+#endif
+			NewVersionData.RuntimeEstimation.Init();
+			
+			RebindNotifications();
+			break;
+		}
+	}
+
+	// set default data for fields we don't want to copy
+	NewVersionData.Version = NewVersion;
+	NewVersionData.VersionChangeDescription = FText();
+	NewVersionData.UpdateScriptExecution = ENiagaraPythonUpdateScriptReference::None;
+	NewVersionData.PythonUpdateScript = FString();
+	NewVersionData.ScriptAsset = FFilePath();
+	
+	VersionData.Add(NewVersionData);
+	VersionData.Sort([](const FVersionedNiagaraEmitterData& A, const FVersionedNiagaraEmitterData& B) { return A.Version < B.Version; });
+
+	return NewVersionData.Version.VersionGuid;
+}
+
+void UNiagaraEmitter::DeleteVersion(const FGuid& VersionGuid)
+{
+	check(VersionGuid != ExposedVersion);
+	
+	for (int i = 0; i < VersionData.Num(); i++)
+	{
+		FNiagaraAssetVersion& AssetVersion = VersionData[i].Version;
+		if (AssetVersion.VersionGuid == VersionGuid)
+		{
+			check(AssetVersion.MajorVersion != 1 || AssetVersion.MinorVersion != 0);
+			Modify();
+			VersionData.RemoveAt(i);
+			return;
+		}
+	}
+}
+
+void UNiagaraEmitter::ExposeVersion(const FGuid& VersionGuid)
+{
+	// check if the requested version already exists
+	for (FVersionedNiagaraEmitterData& Data : VersionData)
+	{
+		if (Data.Version.VersionGuid == VersionGuid)
+		{
+			Modify();
+			ExposedVersion = VersionGuid;
+			Data.Version.bIsVisibleInVersionSelector = true;
+			return;
+		}
+	}
+}
+
+void UNiagaraEmitter::EnableVersioning()
+{
+	if (bVersioningEnabled)
+	{
+		return;
+	}
+	Modify();
+	ensure(VersionData.Num() == 1);
+	bVersioningEnabled = true;	
+	ExposedVersion = VersionData[0].Version.VersionGuid;
+}
+
+void UNiagaraEmitter::DisableVersioning(const FGuid& VersionGuidToUse)
+{
+	CheckVersionDataAvailable();
+	if (!bVersioningEnabled)
+	{
+		return;
+	}
+	
+	bVersioningEnabled = false;
+	FVersionedNiagaraEmitterData DataToUse;
+	DataToUse.CopyFrom(VersionData[0]);
+	if (VersionGuidToUse.IsValid())
+	{
+		for (const FVersionedNiagaraEmitterData& Data : VersionData)
+		{
+			if (Data.Version.VersionGuid == VersionGuidToUse)
+			{
+				DataToUse.CopyFrom(Data);
+				break;
+			}
+		}
+	}
+	DataToUse.Version = FNiagaraAssetVersion(); // reset and create new guid
+	ExposedVersion = DataToUse.Version.VersionGuid;	
+	VersionData.Empty();
+	VersionData.Add(DataToUse);
+}
+
+void UNiagaraEmitter::CheckVersionDataAvailable()
+{
+	if (VersionData.Num() > 0) {
+		return;
+	}
+
+	// copy over existing data of assets that were created pre-versioning
+	FVersionedNiagaraEmitterData& Data = VersionData.AddDefaulted_GetRef();
+	Data.SimTarget = SimTarget_DEPRECATED;
+	Data.Platforms = Platforms_DEPRECATED;
+	Data.RandomSeed = RandomSeed_DEPRECATED;
+	Data.EditorData = EditorData_DEPRECATED;
+	Data.bLocalSpace = bLocalSpace_DEPRECATED;
+	Data.FixedBounds = FixedBounds_DEPRECATED;
+	Data.GraphSource = GraphSource_DEPRECATED;
+	Data.bDeterminism = bDeterminism_DEPRECATED;
+	Data.AllocationMode = AllocationMode_DEPRECATED;
+	Data.bLimitDeltaTime = bLimitDeltaTime_DEPRECATED;
+	Data.SpawnScriptProps = SpawnScriptProps_DEPRECATED;
+	Data.SimulationStages = SimulationStages_DEPRECATED;
+	Data.EditorParameters = EditorParameters_DEPRECATED;
+	Data.GPUComputeScript = GPUComputeScript_DEPRECATED;
+	Data.UpdateScriptProps = UpdateScriptProps_DEPRECATED;
+	Data.PreAllocationCount = PreAllocationCount_DEPRECATED;
+	Data.bCombineEventSpawn = bCombineEventSpawn_DEPRECATED;
+	Data.RendererProperties = RendererProperties_DEPRECATED;
+	Data.MaxDeltaTimePerTick = MaxDeltaTimePerTick_DEPRECATED;
+	Data.ScalabilityOverrides = ScalabilityOverrides_DEPRECATED;
+	Data.AttributesToPreserve = AttributesToPreserve_DEPRECATED;
+	Data.bInterpolatedSpawning = bInterpolatedSpawning_DEPRECATED;
+	Data.bRequiresPersistentIDs = bRequiresPersistentIDs_DEPRECATED;
+	Data.EventHandlerScriptProps = EventHandlerScriptProps_DEPRECATED;
+	Data.SharedEventGeneratorIds = SharedEventGeneratorIds_DEPRECATED;
+	Data.EmitterSpawnScriptProps = EmitterSpawnScriptProps_DEPRECATED;
+	Data.EmitterUpdateScriptProps = EmitterUpdateScriptProps_DEPRECATED;
+	Data.MaxGPUParticlesSpawnPerFrame = MaxGPUParticlesSpawnPerFrame_DEPRECATED;
+
+	if (bFixedBounds_DEPRECATED)
+	{
+		Data.CalculateBoundsMode = ENiagaraEmitterCalculateBoundMode::Fixed;
+	}
+
+	Data.ScratchPads = NewObject<UNiagaraScratchPadContainer>(this);
+	Data.ScratchPads->SetScripts(ScratchPadScripts_DEPRECATED);
+
+	Data.ParentScratchPads = NewObject<UNiagaraScratchPadContainer>(this);
+	Data.ParentScratchPads->SetScripts(ParentScratchPadScripts_DEPRECATED);
+
+	if (Parent_DEPRECATED)
+	{
+		Data.VersionedParent.Emitter = Parent_DEPRECATED;
+	}
+
+	if (ParentAtLastMerge_DEPRECATED)
+	{
+		// we don't need to store a version here, as the parent at last merge is always just a snapshot of one version
+		Data.VersionedParentAtLastMerge.Emitter = ParentAtLastMerge_DEPRECATED;
+	}
+
+	ExposedVersion = Data.Version.VersionGuid;
+
+	for (UNiagaraRendererProperties* Renderer : Data.RendererProperties)
+	{
+		Renderer->OuterEmitterVersion = ExposedVersion;
+	}
+	for (UNiagaraSimulationStageBase* SimStage : Data.SimulationStages)
+	{
+		SimStage->OuterEmitterVersion = ExposedVersion;
+	}
+}
+
+#endif
+
+FVersionedNiagaraEmitterData::FVersionedNiagaraEmitterData() :
+	bInterpolatedSpawning(false), bRequiresPersistentIDs(false), bCombineEventSpawn(false), bLimitDeltaTime(true), MaxGPUParticlesSpawnPerFrame(0), bRequiresViewUniformBuffer(false)
+#if WITH_EDITORONLY_DATA
+	, EditorData(nullptr), EditorParameters(nullptr)
+#endif
+{
+	FixedBounds = GetDefaultFixedBounds();
+}
+
+#if WITH_EDITORONLY_DATA
+
+void UNiagaraEmitter::UpdateFromMergedCopy(const INiagaraMergeManager& MergeManager, UNiagaraEmitter* MergedEmitter, FVersionedNiagaraEmitterData* EmitterData)
 {
 	auto ReouterMergedObject = [](UObject* NewOuter, UObject* TargetObject)
 	{
@@ -1920,79 +2569,85 @@ void UNiagaraEmitter::UpdateFromMergedCopy(const INiagaraMergeManager& MergeMana
 
 	// Copy base editable emitter properties.
 	TArray<FProperty*> DifferentProperties;
-	MergeManager.DiffEditableProperties(this, MergedEmitter, *UNiagaraEmitter::StaticClass(), DifferentProperties);
+	MergeManager.DiffEditableProperties(this, MergedEmitter, *StaticClass(), DifferentProperties);
 	MergeManager.CopyPropertiesToBase(this, MergedEmitter, DifferentProperties);
 
+	ensure(MergedEmitter->IsVersioningEnabled() == false);
+	FVersionedNiagaraEmitterData* MergedData = MergedEmitter->GetLatestEmitterData();
+	DifferentProperties.Empty();
+	MergeManager.DiffEditableProperties(EmitterData, MergedData, *FVersionedNiagaraEmitterData::StaticStruct(), DifferentProperties);
+	MergeManager.CopyPropertiesToBase(EmitterData, MergedData, DifferentProperties);
+
 	// Copy source and scripts
-	ReouterMergedObject(this, MergedEmitter->GraphSource);
-	GraphSource->OnChanged().RemoveAll(this);
-	GraphSource = MergedEmitter->GraphSource;
-	GraphSource->OnChanged().AddUObject(this, &UNiagaraEmitter::GraphSourceChanged);
-
-	ReouterMergedObject(this, MergedEmitter->SpawnScriptProps.Script);
-	SpawnScriptProps.Script->RapidIterationParameters.RemoveAllOnChangedHandlers(this);
-	SpawnScriptProps.Script = MergedEmitter->SpawnScriptProps.Script;
-	SpawnScriptProps.Script->RapidIterationParameters.AddOnChangedHandler(
+	ReouterMergedObject(this, MergedData->GraphSource);
+	EmitterData->GraphSource->OnChanged().RemoveAll(this);
+	EmitterData->GraphSource = MergedData->GraphSource;
+	EmitterData->GraphSource->OnChanged().AddUObject(this, &UNiagaraEmitter::GraphSourceChanged);
+
+	ReouterMergedObject(this, MergedData->SpawnScriptProps.Script);
+	EmitterData->SpawnScriptProps.Script->RapidIterationParameters.RemoveAllOnChangedHandlers(this);
+	EmitterData->SpawnScriptProps.Script = MergedData->SpawnScriptProps.Script;
+	EmitterData->SpawnScriptProps.Script->RapidIterationParameters.AddOnChangedHandler(
 		FNiagaraParameterStore::FOnChanged::FDelegate::CreateUObject(this, &UNiagaraEmitter::ScriptRapidIterationParameterChanged));
 
-	ReouterMergedObject(this, MergedEmitter->UpdateScriptProps.Script);
-	UpdateScriptProps.Script->RapidIterationParameters.RemoveAllOnChangedHandlers(this);
-	UpdateScriptProps.Script = MergedEmitter->UpdateScriptProps.Script;
-	UpdateScriptProps.Script->RapidIterationParameters.AddOnChangedHandler(
+	ReouterMergedObject(this, MergedData->UpdateScriptProps.Script);
+	EmitterData->UpdateScriptProps.Script->RapidIterationParameters.RemoveAllOnChangedHandlers(this);
+	EmitterData->UpdateScriptProps.Script = MergedData->UpdateScriptProps.Script;
+	EmitterData->UpdateScriptProps.Script->RapidIterationParameters.AddOnChangedHandler(
 		FNiagaraParameterStore::FOnChanged::FDelegate::CreateUObject(this, &UNiagaraEmitter::ScriptRapidIterationParameterChanged));
 
-	ReouterMergedObject(this, MergedEmitter->EmitterSpawnScriptProps.Script);
-	EmitterSpawnScriptProps.Script->RapidIterationParameters.RemoveAllOnChangedHandlers(this);
-	EmitterSpawnScriptProps.Script = MergedEmitter->EmitterSpawnScriptProps.Script;
-	EmitterSpawnScriptProps.Script->RapidIterationParameters.AddOnChangedHandler(
+	ReouterMergedObject(this, MergedData->EmitterSpawnScriptProps.Script);
+	EmitterData->EmitterSpawnScriptProps.Script->RapidIterationParameters.RemoveAllOnChangedHandlers(this);
+	EmitterData->EmitterSpawnScriptProps.Script = MergedData->EmitterSpawnScriptProps.Script;
+	EmitterData->EmitterSpawnScriptProps.Script->RapidIterationParameters.AddOnChangedHandler(
 		FNiagaraParameterStore::FOnChanged::FDelegate::CreateUObject(this, &UNiagaraEmitter::ScriptRapidIterationParameterChanged));
 
-	ReouterMergedObject(this, MergedEmitter->EmitterUpdateScriptProps.Script);
-	EmitterUpdateScriptProps.Script->RapidIterationParameters.RemoveAllOnChangedHandlers(this);
-	EmitterUpdateScriptProps.Script = MergedEmitter->EmitterUpdateScriptProps.Script;
-	EmitterUpdateScriptProps.Script->RapidIterationParameters.AddOnChangedHandler(
+	ReouterMergedObject(this, MergedData->EmitterUpdateScriptProps.Script);
+	EmitterData->EmitterUpdateScriptProps.Script->RapidIterationParameters.RemoveAllOnChangedHandlers(this);
+	EmitterData->EmitterUpdateScriptProps.Script = MergedData->EmitterUpdateScriptProps.Script;
+	EmitterData->EmitterUpdateScriptProps.Script->RapidIterationParameters.AddOnChangedHandler(
 		FNiagaraParameterStore::FOnChanged::FDelegate::CreateUObject(this, &UNiagaraEmitter::ScriptRapidIterationParameterChanged));
 
-	ReouterMergedObject(this, MergedEmitter->GPUComputeScript);
-	GPUComputeScript->RapidIterationParameters.RemoveAllOnChangedHandlers(this);
-	GPUComputeScript = MergedEmitter->GPUComputeScript;
-	GPUComputeScript->RapidIterationParameters.AddOnChangedHandler(
+	ReouterMergedObject(this, MergedData->GPUComputeScript);
+	EmitterData->GPUComputeScript->RapidIterationParameters.RemoveAllOnChangedHandlers(this);
+	EmitterData->GPUComputeScript = MergedData->GPUComputeScript;
+	EmitterData->GPUComputeScript->RapidIterationParameters.AddOnChangedHandler(
 		FNiagaraParameterStore::FOnChanged::FDelegate::CreateUObject(this, &UNiagaraEmitter::ScriptRapidIterationParameterChanged));
 
 	// Copy event handlers
-	for (FNiagaraEventScriptProperties& EventScriptProperties : EventHandlerScriptProps)
+	for (FNiagaraEventScriptProperties& EventScriptProperties : EmitterData->EventHandlerScriptProps)
 	{
 		EventScriptProperties.Script->RapidIterationParameters.RemoveAllOnChangedHandlers(this);
 	}
-	EventHandlerScriptProps.Empty();
+	EmitterData->EventHandlerScriptProps.Empty();
 
-	for (FNiagaraEventScriptProperties& MergedEventScriptProperties : MergedEmitter->EventHandlerScriptProps)
+	for (FNiagaraEventScriptProperties& MergedEventScriptProperties : MergedData->EventHandlerScriptProps)
 	{
-		EventHandlerScriptProps.Add(MergedEventScriptProperties);
+		EmitterData->EventHandlerScriptProps.Add(MergedEventScriptProperties);
 		ReouterMergedObject(this, MergedEventScriptProperties.Script);
 		MergedEventScriptProperties.Script->RapidIterationParameters.AddOnChangedHandler(
 			FNiagaraParameterStore::FOnChanged::FDelegate::CreateUObject(this, &UNiagaraEmitter::ScriptRapidIterationParameterChanged));
 	}
 
 	// Copy shader stages
-	for (TObjectPtr<UNiagaraSimulationStageBase>& SimulationStage : SimulationStages)
+	for (TObjectPtr<UNiagaraSimulationStageBase>& SimulationStage : EmitterData->SimulationStages)
 	{
 		SimulationStage->OnChanged().RemoveAll(this);
 		SimulationStage->Script->RapidIterationParameters.RemoveAllOnChangedHandlers(this);
 	}
-	SimulationStages.Empty();
+	EmitterData->SimulationStages.Empty();
 
-	for (UNiagaraSimulationStageBase* MergedSimulationStage : MergedEmitter->SimulationStages)
+	for (UNiagaraSimulationStageBase* MergedSimulationStage : MergedData->SimulationStages)
 	{
 		ReouterMergedObject(this, MergedSimulationStage);
-		SimulationStages.Add(MergedSimulationStage);
+		EmitterData->SimulationStages.Add(MergedSimulationStage);
 		MergedSimulationStage->OnChanged().AddUObject(this, &UNiagaraEmitter::SimulationStageChanged);
 		MergedSimulationStage->Script->RapidIterationParameters.AddOnChangedHandler(
 			FNiagaraParameterStore::FOnChanged::FDelegate::CreateUObject(this, &UNiagaraEmitter::ScriptRapidIterationParameterChanged));
 	}
 
 	// Copy renderers
-	for (UNiagaraRendererProperties* Renderer : RendererProperties)
+	for (UNiagaraRendererProperties* Renderer : EmitterData->RendererProperties)
 	{
 		Renderer->OnChanged().RemoveAll(this);
 
@@ -2000,34 +2655,33 @@ void UNiagaraEmitter::UpdateFromMergedCopy(const INiagaraMergeManager& MergeMana
 		// they will be included in a cook; so clear the flag while we're removing them
 		Renderer->ClearFlags(RF_Public);
 	}
-	RendererProperties.Empty();
+	EmitterData->RendererProperties.Empty();
 
-	for (UNiagaraRendererProperties* MergedRenderer : MergedEmitter->RendererProperties)
+	for (UNiagaraRendererProperties* MergedRenderer : MergedData->RendererProperties)
 	{
 		ReouterMergedObject(this, MergedRenderer);
-		RendererProperties.Add(MergedRenderer);
+		EmitterData->RendererProperties.Add(MergedRenderer);
 		MergedRenderer->OnChanged().AddUObject(this, &UNiagaraEmitter::RendererChanged);
 	}
 
 	// Copy scratch pad scripts.
-	ParentScratchPadScripts.Empty();
-	ScratchPadScripts.Empty();
-
-	for (UNiagaraScript* MergedParentScratchPadScript : MergedEmitter->ParentScratchPadScripts)
+	EmitterData->ParentScratchPads->Scripts.Empty();
+	for (UNiagaraScript* MergedParentScratchPadScript : MergedData->ParentScratchPads->Scripts)
 	{
-		ReouterMergedObject(this, MergedParentScratchPadScript);
-		ParentScratchPadScripts.Add(MergedParentScratchPadScript);
+		ReouterMergedObject(EmitterData->ParentScratchPads, MergedParentScratchPadScript);
+		EmitterData->ParentScratchPads->Scripts.Add(MergedParentScratchPadScript);
 	}
 
-	for (UNiagaraScript* MergedScratchPadScript : MergedEmitter->ScratchPadScripts)
+	EmitterData->ScratchPads->Scripts.Empty();
+	for (UNiagaraScript* MergedScratchPadScript : MergedData->ScratchPads->Scripts)
 	{
-		ReouterMergedObject(this, MergedScratchPadScript);
-		ScratchPadScripts.Add(MergedScratchPadScript);
+		ReouterMergedObject(EmitterData->ScratchPads, MergedScratchPadScript);
+		EmitterData->ScratchPads->Scripts.Add(MergedScratchPadScript);
 	}
 
-	UNiagaraEditorDataBase* NewEditorData = MergedEmitter->GetEditorData();
+	UNiagaraEditorDataBase* NewEditorData = MergedData->GetEditorData();
 	ReouterMergedObject(this, NewEditorData);
-	SetEditorData(NewEditorData);
+	SetEditorData(MergedData->GetEditorData(), EmitterData->Version.VersionGuid);
 
 	// Update messages
 	for (const TPair<FGuid, TObjectPtr<UNiagaraMessageDataBase>>& Message : MergedEmitter->GetMessages())
@@ -2044,7 +2698,7 @@ void UNiagaraEmitter::UpdateFromMergedCopy(const INiagaraMergeManager& MergeMana
 	UpdateChangeId(TEXT("Updated from merged copy"));
 }
 
-void UNiagaraEmitter::SyncEmitterAlias(const FString& InOldName, const FString& InNewName)
+void FVersionedNiagaraEmitterData::SyncEmitterAlias(const FString& InOldName, const UNiagaraEmitter& InEmitter)
 {
 	TArray<UNiagaraScript*> Scripts;
 	GetScripts(Scripts, false, true); // Get all the scripts...
@@ -2055,23 +2709,24 @@ void UNiagaraEmitter::SyncEmitterAlias(const FString& InOldName, const FString&
 		// due to compilation, in cases where the package should be marked dirty an previous modify would have already done this.
 		Script->Modify(false);
 		Script->SyncAliases(FNiagaraAliasContext(Script->GetUsage())
-			.ChangeEmitterName(InOldName, InNewName));
+			.ChangeEmitterName(InOldName, InEmitter.GetUniqueEmitterName()));
 	}
 
 	// if we haven't yet been postloaded then we'll hold off on updating the renderers as they are dependent on everything
 	// (System/Emitter/Scripts) being fully loaded.
-	if (!HasAnyFlags(RF_NeedPostLoad))
+	if (!InEmitter.HasAnyFlags(RF_NeedPostLoad))
 	{
 		for (UNiagaraRendererProperties* Renderer : RendererProperties)
 		{
 			if (Renderer)
 			{
 				Renderer->Modify(false);
-				Renderer->RenameEmitter(*InOldName, this);
+				Renderer->RenameEmitter(*InOldName, &InEmitter);
 			}
 		}
 	}
 }
+
 #endif
 bool UNiagaraEmitter::SetUniqueEmitterName(const FString& InName)
 {
@@ -2084,12 +2739,15 @@ bool UNiagaraEmitter::SetUniqueEmitterName(const FString& InName)
 		if (GetName() != InName)
 		{
 			// Also rename the underlying uobject to keep things consistent.
-			FName UniqueObjectName = MakeUniqueObjectName(GetOuter(), UNiagaraEmitter::StaticClass(), *InName);
+			FName UniqueObjectName = MakeUniqueObjectName(GetOuter(), StaticClass(), *InName);
 			Rename(*UniqueObjectName.ToString(), GetOuter(), REN_ForceNoResetLoaders);
 		}
 
 #if WITH_EDITORONLY_DATA
-		SyncEmitterAlias(OldName, UniqueEmitterName);
+		for (FVersionedNiagaraEmitterData& EmitterData : VersionData)
+		{
+			EmitterData.SyncEmitterAlias(OldName, *this);
+		}
 #endif
 		return true;
 	}
@@ -2108,16 +2766,18 @@ bool UNiagaraEmitter::SetUniqueEmitterName(const FString& InName)
 //	}
 //}
 
-void UNiagaraEmitter::AddRenderer(UNiagaraRendererProperties* Renderer)
+void UNiagaraEmitter::AddRenderer(UNiagaraRendererProperties* Renderer, FGuid EmitterVersion)
 {
 	Modify();
-	RendererProperties.Add(Renderer);
+	Renderer->OuterEmitterVersion = EmitterVersion;
+	FVersionedNiagaraEmitterData* EmitterData = GetEmitterData(EmitterVersion);
+	EmitterData->RendererProperties.Add(Renderer);
 #if WITH_EDITOR
 	Renderer->OnChanged().AddUObject(this, &UNiagaraEmitter::RendererChanged);
 	UpdateChangeId(TEXT("Renderer added"));
 	OnRenderersChangedDelegate.Broadcast();
 #endif
-	RebuildRendererBindings();
+	EmitterData->RebuildRendererBindings(*this);
 
 	if (UNiagaraSystem* Owner = GetTypedOuter<UNiagaraSystem>())
 	{
@@ -2126,16 +2786,17 @@ void UNiagaraEmitter::AddRenderer(UNiagaraRendererProperties* Renderer)
 	}
 }
 
-void UNiagaraEmitter::RemoveRenderer(UNiagaraRendererProperties* Renderer)
+void UNiagaraEmitter::RemoveRenderer(UNiagaraRendererProperties* Renderer, FGuid EmitterVersion)
 {
 	Modify();
-	RendererProperties.Remove(Renderer);
+	FVersionedNiagaraEmitterData* EmitterData = GetEmitterData(EmitterVersion);
+	EmitterData->RendererProperties.Remove(Renderer);
 #if WITH_EDITOR
 	Renderer->OnChanged().RemoveAll(this);
 	UpdateChangeId(TEXT("Renderer removed"));
 	OnRenderersChangedDelegate.Broadcast();
 #endif
-	RebuildRendererBindings();
+	EmitterData->RebuildRendererBindings(*this);
 
 	if (UNiagaraSystem* Owner = GetTypedOuter<UNiagaraSystem>())
 	{
@@ -2143,7 +2804,7 @@ void UNiagaraEmitter::RemoveRenderer(UNiagaraRendererProperties* Renderer)
 	}
 }
 
-FNiagaraEventScriptProperties* UNiagaraEmitter::GetEventHandlerByIdUnsafe(FGuid ScriptUsageId)
+FNiagaraEventScriptProperties* FVersionedNiagaraEmitterData::GetEventHandlerByIdUnsafe(FGuid ScriptUsageId)
 {
 	for (FNiagaraEventScriptProperties& EventScriptProperties : EventHandlerScriptProps)
 	{
@@ -2155,10 +2816,10 @@ FNiagaraEventScriptProperties* UNiagaraEmitter::GetEventHandlerByIdUnsafe(FGuid
 	return nullptr;
 }
 
-void UNiagaraEmitter::AddEventHandler(FNiagaraEventScriptProperties EventHandler)
+void UNiagaraEmitter::AddEventHandler(FNiagaraEventScriptProperties EventHandler, FGuid EmitterVersion)
 {
 	Modify();
-	EventHandlerScriptProps.Add(EventHandler);
+	GetEmitterData(EmitterVersion)->EventHandlerScriptProps.Add(EventHandler);
 #if WITH_EDITOR
 	EventHandler.Script->RapidIterationParameters.AddOnChangedHandler(
 		FNiagaraParameterStore::FOnChanged::FDelegate::CreateUObject(this, &UNiagaraEmitter::ScriptRapidIterationParameterChanged));
@@ -2166,10 +2827,11 @@ void UNiagaraEmitter::AddEventHandler(FNiagaraEventScriptProperties EventHandler
 #endif
 }
 
-void UNiagaraEmitter::RemoveEventHandlerByUsageId(FGuid EventHandlerUsageId)
+void UNiagaraEmitter::RemoveEventHandlerByUsageId(FGuid EventHandlerUsageId, FGuid EmitterVersion)
 {
 	Modify();
 	auto FindEventHandlerById = [=](const FNiagaraEventScriptProperties& EventHandler) { return EventHandler.Script->GetUsageId() == EventHandlerUsageId; };
+	TArray<FNiagaraEventScriptProperties>& EventHandlerScriptProps = GetEmitterData(EmitterVersion)->EventHandlerScriptProps;
 #if WITH_EDITOR
 	FNiagaraEventScriptProperties* EventHandler = EventHandlerScriptProps.FindByPredicate(FindEventHandlerById);
 	if (EventHandler != nullptr)
@@ -2183,17 +2845,18 @@ void UNiagaraEmitter::RemoveEventHandlerByUsageId(FGuid EventHandlerUsageId)
 #endif
 }
 
-UNiagaraSimulationStageBase* UNiagaraEmitter::GetSimulationStageById(FGuid ScriptUsageId) const
+UNiagaraSimulationStageBase* FVersionedNiagaraEmitterData::GetSimulationStageById(FGuid ScriptUsageId) const
 {
 	TObjectPtr<UNiagaraSimulationStageBase> const* FoundSimulationStagePtr = SimulationStages.FindByPredicate([&ScriptUsageId](UNiagaraSimulationStageBase* SimulationStage) { return SimulationStage->Script->GetUsageId() == ScriptUsageId; });
 	return FoundSimulationStagePtr != nullptr ? *FoundSimulationStagePtr : nullptr;
 }
 
-void UNiagaraEmitter::AddSimulationStage(UNiagaraSimulationStageBase* SimulationStage)
+void UNiagaraEmitter::AddSimulationStage(UNiagaraSimulationStageBase* SimulationStage, FGuid EmitterVersion)
 {
 	Modify();
-	SimulationStages.Add(SimulationStage);
+	GetEmitterData(EmitterVersion)->SimulationStages.Add(SimulationStage);
 #if WITH_EDITOR
+	SimulationStage->OuterEmitterVersion = EmitterVersion;
 	SimulationStage->OnChanged().AddUObject(this, &UNiagaraEmitter::SimulationStageChanged);
 	SimulationStage->Script->RapidIterationParameters.AddOnChangedHandler(
 		FNiagaraParameterStore::FOnChanged::FDelegate::CreateUObject(this, &UNiagaraEmitter::ScriptRapidIterationParameterChanged));
@@ -2201,10 +2864,10 @@ void UNiagaraEmitter::AddSimulationStage(UNiagaraSimulationStageBase* Simulation
 #endif
 }
 
-void UNiagaraEmitter::RemoveSimulationStage(UNiagaraSimulationStageBase* SimulationStage)
+void UNiagaraEmitter::RemoveSimulationStage(UNiagaraSimulationStageBase* SimulationStage, FGuid EmitterVersion)
 {
 	Modify();
-	bool bRemoved = SimulationStages.Remove(SimulationStage) != 0;
+	bool bRemoved = GetEmitterData(EmitterVersion)->SimulationStages.Remove(SimulationStage) != 0;
 #if WITH_EDITOR
 	if (bRemoved)
 	{
@@ -2215,8 +2878,9 @@ void UNiagaraEmitter::RemoveSimulationStage(UNiagaraSimulationStageBase* Simulat
 #endif
 }
 
-void UNiagaraEmitter::MoveSimulationStageToIndex(UNiagaraSimulationStageBase* SimulationStageToMove, int32 TargetIndex)
+void UNiagaraEmitter::MoveSimulationStageToIndex(UNiagaraSimulationStageBase* SimulationStageToMove, int32 TargetIndex, FGuid EmitterVersion)
 {
+	TArray<TObjectPtr<UNiagaraSimulationStageBase>>& SimulationStages = GetEmitterData(EmitterVersion)->SimulationStages;
 	int32 CurrentIndex = SimulationStages.IndexOfByKey(SimulationStageToMove);
 	checkf(CurrentIndex != INDEX_NONE, TEXT("Simulation stage could not be moved because it is not owned by this emitter."));
 	if (TargetIndex != CurrentIndex)
@@ -2233,7 +2897,7 @@ void UNiagaraEmitter::MoveSimulationStageToIndex(UNiagaraSimulationStageBase* Si
 	}
 }
 
-bool UNiagaraEmitter::IsEventGeneratorShared(FName EventGeneratorId) const
+bool FVersionedNiagaraEmitterData::IsEventGeneratorShared(FName EventGeneratorId) const
 {
 	return SharedEventGeneratorIds.Contains(EventGeneratorId);
 }
@@ -2241,13 +2905,16 @@ bool UNiagaraEmitter::IsEventGeneratorShared(FName EventGeneratorId) const
 void UNiagaraEmitter::BeginDestroy()
 {
 #if WITH_EDITOR
-	if (GraphSource != nullptr)
-	{
-		GraphSource->OnChanged().RemoveAll(this);
-	}
-	if (GPUComputeScript)
+	for (FVersionedNiagaraEmitterData& Data : VersionData)
 	{
-		GPUComputeScript->OnGPUScriptCompiled().RemoveAll(this);
+		if (Data.GraphSource != nullptr)
+		{
+			Data.GraphSource->OnChanged().RemoveAll(this);
+		}
+		if (Data.GPUComputeScript)
+		{
+			Data.GPUComputeScript->OnGPUScriptCompiled().RemoveAll(this);
+		}
 	}
 	CleanupDefinitionsSubscriptions();
 #endif
@@ -2268,10 +2935,15 @@ void UNiagaraEmitter::UpdateChangeId(const FString& Reason)
 		UE_LOG(LogNiagara, Log, TEXT("Emitter %s change id updated. Reason: %s OldId: %s NewId: %s"),
 			*GetPathName(), *Reason, *OldId.ToString(), *ChangeId.ToString());
 	}
+
+	for (FVersionedNiagaraEmitterData& Data : VersionData)
+	{
+		Data.CachedTraversalData.Reset();
 #if STATS
-	StatDatabase.ClearStatCaptures();
+		Data.StatDatabase.Init();
+		Data.StatDatabase.ClearStatCaptures();
 #endif
-	CachedTraversalData.Reset();
+	}
 }
 
 void UNiagaraEmitter::ScriptRapidIterationParameterChanged()
@@ -2300,7 +2972,7 @@ void UNiagaraEmitter::GraphSourceChanged()
 
 void UNiagaraEmitter::RaiseOnEmitterGPUCompiled(UNiagaraScript* InScript, const FGuid& ScriptVersion)
 {
-	OnGPUScriptCompiledDelegate.Broadcast(this);
+	OnGPUScriptCompiledDelegate.Broadcast(InScript->GetOuterEmitter());
 }
 
 void UNiagaraEmitter::PersistentEditorDataChanged()
@@ -2343,9 +3015,9 @@ TStatId UNiagaraEmitter::GetStatID(bool bGameThread, bool bConcurrent)const
 	return TStatId();
 }
 
-void UNiagaraEmitter::ClearRuntimeAllocationEstimate(uint64 ReportHandle)
+void FVersionedNiagaraEmitterData::ClearRuntimeAllocationEstimate(uint64 ReportHandle)
 {
-	FScopeLock Lock(&EstimationCriticalSection);
+	FScopeLock Lock(RuntimeEstimation.GetCriticalSection());
 	if (ReportHandle == INDEX_NONE)
 	{
 		RuntimeEstimation.AllocationEstimate = 0;
@@ -2359,9 +3031,9 @@ void UNiagaraEmitter::ClearRuntimeAllocationEstimate(uint64 ReportHandle)
 	}
 }
 
-int32 UNiagaraEmitter::AddRuntimeAllocation(uint64 ReporterHandle, int32 AllocationCount)
+int32 FVersionedNiagaraEmitterData::AddRuntimeAllocation(uint64 ReporterHandle, int32 AllocationCount)
 {
-	FScopeLock Lock(&EstimationCriticalSection);
+	FScopeLock Lock(RuntimeEstimation.GetCriticalSection());
 	int32* Estimate = RuntimeEstimation.RuntimeAllocations.Find(ReporterHandle);
 	if (!Estimate || *Estimate < AllocationCount)
 	{
@@ -2379,7 +3051,7 @@ int32 UNiagaraEmitter::AddRuntimeAllocation(uint64 ReporterHandle, int32 Allocat
 	return RuntimeEstimation.RuntimeAllocations.Num();
 }
 
-int32 UNiagaraEmitter::GetMaxParticleCountEstimate()
+int32 FVersionedNiagaraEmitterData::GetMaxParticleCountEstimate()
 {
 	if ((AllocationMode == EParticleAllocationMode::ManualEstimate)
 		|| (AllocationMode == EParticleAllocationMode::FixedCount))
@@ -2389,7 +3061,7 @@ int32 UNiagaraEmitter::GetMaxParticleCountEstimate()
 	
 	if (RuntimeEstimation.IsEstimationDirty)
 	{
-		FScopeLock lock(&EstimationCriticalSection);
+		FScopeLock lock(RuntimeEstimation.GetCriticalSection());
 		int32 EstimationCount = RuntimeEstimation.RuntimeAllocations.Num();
 		RuntimeEstimation.AllocationEstimate = 0;
 		if (EstimationCount > 0)
@@ -2412,6 +3084,17 @@ int32 UNiagaraEmitter::GetMaxParticleCountEstimate()
 	return RuntimeEstimation.AllocationEstimate;
 }
 
+FCriticalSection* FMemoryRuntimeEstimation::GetCriticalSection()
+{
+	ensure(EstimationCriticalSection.IsValid());
+	return EstimationCriticalSection.Get();
+}
+
+void FMemoryRuntimeEstimation::Init()
+{
+	EstimationCriticalSection = MakeShared<FCriticalSection>();
+}
+
 void UNiagaraEmitter::GenerateStatID()const
 {
 #if STATS
@@ -2425,41 +3108,58 @@ void UNiagaraEmitter::GenerateStatID()const
 }
 
 #if WITH_EDITORONLY_DATA
-UNiagaraEmitter* UNiagaraEmitter::GetParent() const
+FVersionedNiagaraEmitter FVersionedNiagaraEmitterData::GetParent() const
 {
-	return Parent;
+	return VersionedParent;
 }
 
-UNiagaraEmitter* UNiagaraEmitter::GetParentAtLastMerge() const
+FVersionedNiagaraEmitter FVersionedNiagaraEmitterData::GetParentAtLastMerge() const
 {
-	return ParentAtLastMerge;
+	return VersionedParentAtLastMerge;
 }
 
-void UNiagaraEmitter::RemoveParent()
+void FVersionedNiagaraEmitterData::RemoveParent()
 {
-	Parent = nullptr;
-	ParentAtLastMerge = nullptr;
+	VersionedParent = FVersionedNiagaraEmitter();
+	VersionedParentAtLastMerge = FVersionedNiagaraEmitter();
 }
 
-void UNiagaraEmitter::SetParent(UNiagaraEmitter& InParent)
+void UNiagaraEmitter::SetParent(const FVersionedNiagaraEmitter& InParent)
 {
-	Parent = &InParent;
-	ParentAtLastMerge = InParent.DuplicateWithoutMerging(this);
-	ParentAtLastMerge->ClearFlags(RF_Standalone | RF_Public);
+	if (ensure(VersionData.Num() == 1))
+	{
+		FVersionedNiagaraEmitterData& EmitterData = VersionData[0];
+		EmitterData.VersionedParent = InParent;
+		EmitterData.VersionedParentAtLastMerge.Emitter = InParent.Emitter->DuplicateWithoutMerging(this);
+		EmitterData.VersionedParentAtLastMerge.Emitter->ClearFlags(RF_Standalone | RF_Public);
+		EmitterData.VersionedParentAtLastMerge.Emitter->DisableVersioning(InParent.Version);
+		EmitterData.VersionedParentAtLastMerge.Version = InParent.Version;
 
-	// Since this API is only valid for the "Create duplicate parent" operation we move the emitters scratch pad script to the parent array since that's where they're defined now.
-	ParentScratchPadScripts.Append(ScratchPadScripts);
-	ScratchPadScripts.Empty();
-	UpdateChangeId(TEXT("Parent Set"));
+		// Since this API is only valid for the "Create duplicate parent" operation we move the emitters scratch pad script to the parent array since that's where they're defined now.
+		// Normally we would duplicate the parent scratch pad scripts here, but since the whole parent is already a duplicate, we can skip this here.
+		EmitterData.ParentScratchPads->SetScripts(EmitterData.ScratchPads->Scripts);
+		EmitterData.ScratchPads->Scripts.Empty();
+		EmitterData.GraphSource->MarkNotSynchronized(TEXT("Emitter parent changed"));
+		UpdateChangeId(TEXT("Parent Set"));
+	}
+}
 
+void FVersionedNiagaraEmitterData::Reparent(const FVersionedNiagaraEmitter& InParent)
+{
+	VersionedParent = InParent;
+	VersionedParentAtLastMerge = FVersionedNiagaraEmitter();
 	GraphSource->MarkNotSynchronized(TEXT("Emitter parent changed"));
 }
 
-void UNiagaraEmitter::Reparent(UNiagaraEmitter& InParent)
+void UNiagaraEmitter::ChangeParentVersion(const FGuid& NewParentVersion, const FGuid& EmitterVersion)
 {
-	Parent = &InParent;
-	ParentAtLastMerge = nullptr;
-	GraphSource->MarkNotSynchronized(TEXT("Emitter parent changed"));
+	FVersionedNiagaraEmitterData* EmitterData = GetEmitterData(EmitterVersion);
+	if (ensure(EmitterData) && EmitterData->VersionedParent.Emitter && EmitterData->VersionedParent.Emitter->IsVersioningEnabled() && EmitterData->VersionedParent.Emitter->GetEmitterData(NewParentVersion))
+	{
+		Modify();
+		EmitterData->VersionedParent.Version = NewParentVersion;
+		MergeChangesFromParent();
+	}
 }
 
 void UNiagaraEmitter::NotifyScratchPadScriptsChanged()
@@ -2470,24 +3170,27 @@ void UNiagaraEmitter::NotifyScratchPadScriptsChanged()
 
 void UNiagaraEmitter::ResolveScalabilitySettings()
 {
-	CurrentScalabilitySettings.Clear();
-
-	if (UNiagaraSystem* Owner = GetTypedOuter<UNiagaraSystem>())
+	for (FVersionedNiagaraEmitterData& EmitterData : VersionData)
 	{
-		if(UNiagaraEffectType* ActualEffectType = Owner->GetEffectType())
+		EmitterData.CurrentScalabilitySettings.Clear();
+
+		if (UNiagaraSystem* Owner = GetTypedOuter<UNiagaraSystem>())
 		{
-			CurrentScalabilitySettings = ActualEffectType->GetActiveEmitterScalabilitySettings();
+			if(UNiagaraEffectType* ActualEffectType = Owner->GetEffectType())
+			{
+				EmitterData.CurrentScalabilitySettings = ActualEffectType->GetActiveEmitterScalabilitySettings();
+			}
 		}
-	}
 
-	for (FNiagaraEmitterScalabilityOverride& Override : ScalabilityOverrides.Overrides)
-	{
-		if (Override.Platforms.IsActive())
+		for (FNiagaraEmitterScalabilityOverride& Override : EmitterData.ScalabilityOverrides.Overrides)
 		{
-			if (Override.bOverrideSpawnCountScale)
+			if (Override.Platforms.IsActive())
 			{
-				CurrentScalabilitySettings.bScaleSpawnCount = Override.bScaleSpawnCount;
-				CurrentScalabilitySettings.SpawnCountScale = Override.SpawnCountScale;
+				if (Override.bOverrideSpawnCountScale)
+				{
+					EmitterData.CurrentScalabilitySettings.bScaleSpawnCount = Override.bScaleSpawnCount;
+					EmitterData.CurrentScalabilitySettings.SpawnCountScale = Override.SpawnCountScale;
+				}
 			}
 		}
 	}
@@ -2497,3 +3200,4 @@ void UNiagaraEmitter::UpdateScalability()
 {
 	ResolveScalabilitySettings();
 }
+
