Added versioning for Niagara emitters

diff --git a/Engine/Plugins/FX/Niagara/Source/Niagara/Private/NiagaraEmitterHandle.cpp b/Engine/Plugins/FX/Niagara/Source/Niagara/Private/NiagaraEmitterHandle.cpp
--- a/Engine/Plugins/FX/Niagara/Source/Niagara/Private/NiagaraEmitterHandle.cpp
+++ b/Engine/Plugins/FX/Niagara/Source/Niagara/Private/NiagaraEmitterHandle.cpp
@@ -16,12 +16,12 @@ FNiagaraEmitterHandle::FNiagaraEmitterHandle()
 	, LastMergedSource_DEPRECATED(nullptr)
 	, bIsolated(false)
 #endif
-	, Instance(nullptr)
+	, Instance_DEPRECATED(nullptr)
 {
 }
 
 #if WITH_EDITORONLY_DATA
-FNiagaraEmitterHandle::FNiagaraEmitterHandle(UNiagaraEmitter& InEmitter)
+FNiagaraEmitterHandle::FNiagaraEmitterHandle(UNiagaraEmitter& InEmitter, const FGuid& Version)
 	: Id(FGuid::NewGuid())
 	, IdName(*Id.ToString())
 	, bIsEnabled(true)
@@ -29,7 +29,21 @@ FNiagaraEmitterHandle::FNiagaraEmitterHandle(UNiagaraEmitter& InEmitter)
 	, Source_DEPRECATED(nullptr)
 	, LastMergedSource_DEPRECATED(nullptr)
 	, bIsolated(false)
-	, Instance(&InEmitter)
+	, Instance_DEPRECATED(nullptr)
+	, VersionedInstance(FVersionedNiagaraEmitter(&InEmitter, Version))
+{
+}
+
+FNiagaraEmitterHandle::FNiagaraEmitterHandle(const FVersionedNiagaraEmitter& InEmitter)
+	: Id(FGuid::NewGuid())
+	, IdName(*Id.ToString())
+	, bIsEnabled(true)
+	, Name(*InEmitter.Emitter->GetUniqueEmitterName())
+	, Source_DEPRECATED(nullptr)
+	, LastMergedSource_DEPRECATED(nullptr)
+	, bIsolated(false)
+	, Instance_DEPRECATED(nullptr)
+	, VersionedInstance(InEmitter)
 {
 }
 #endif
@@ -73,7 +87,7 @@ void FNiagaraEmitterHandle::SetName(FName InName, UNiagaraSystem& InOwnerSystem)
 	FName UniqueName = FNiagaraUtilities::GetUniqueName(SanitizedName, OtherEmitterNames);
 
 	Name = UniqueName;
-	if (Instance->SetUniqueEmitterName(Name.ToString()))
+	if (VersionedInstance.Emitter->SetUniqueEmitterName(Name.ToString()))
 	{
  #if WITH_EDITOR
 		if (InOwnerSystem.GetSystemSpawnScript() && InOwnerSystem.GetSystemSpawnScript()->GetLatestSource())
@@ -109,9 +123,10 @@ bool FNiagaraEmitterHandle::SetIsEnabled(bool bInIsEnabled, UNiagaraSystem& InOw
 			InOwnerSystem.GetSystemUpdateScript()->InvalidateCompileResults(InvalidateReason);
 
 			// Clean out the emitter's compile results for cleanliness.
-			if (Instance)
+			FVersionedNiagaraEmitterData* EmitterData = VersionedInstance.GetEmitterData();
+			if (EmitterData)
 			{
-				Instance->InvalidateCompileResults();
+				EmitterData->InvalidateCompileResults();
 			}
 
 			// In some cases we may do the recompile now.
@@ -126,22 +141,29 @@ bool FNiagaraEmitterHandle::SetIsEnabled(bool bInIsEnabled, UNiagaraSystem& InOw
 	return false;
 }
 
-UNiagaraEmitter* FNiagaraEmitterHandle::GetInstance() const
+FVersionedNiagaraEmitter FNiagaraEmitterHandle::GetInstance() const
+{
+	return VersionedInstance;
+}
+
+FVersionedNiagaraEmitter& FNiagaraEmitterHandle::GetInstance()
 {
-	return Instance;
+	return VersionedInstance;
 }
 
-FString FNiagaraEmitterHandle::GetUniqueInstanceName()const
+FVersionedNiagaraEmitterData* FNiagaraEmitterHandle::GetEmitterData() const
+{
+	return VersionedInstance.GetEmitterData();
+}
+
+FString FNiagaraEmitterHandle::GetUniqueInstanceName() const
 {
 	// We might not have an instance if this is a cooked object that we're loading in the editor
-	if (Instance)
+	if (VersionedInstance.Emitter)
 	{
-		return Instance->GetUniqueEmitterName();
-	}
-	else
-	{
-		return FString();
+		return VersionedInstance.Emitter->GetUniqueEmitterName();
 	}
+	return FString();
 }
 
 #if WITH_EDITORONLY_DATA
@@ -151,7 +173,7 @@ bool FNiagaraEmitterHandle::NeedsRecompile() const
 	if (GetIsEnabled())
 	{
 		TArray<UNiagaraScript*> Scripts;
-		Instance->GetScripts(Scripts);
+		VersionedInstance.GetEmitterData()->GetScripts(Scripts);
 
 		for (UNiagaraScript* Script : Scripts)
 		{
@@ -166,30 +188,37 @@ bool FNiagaraEmitterHandle::NeedsRecompile() const
 
 void FNiagaraEmitterHandle::ConditionalPostLoad(int32 NiagaraCustomVersion)
 {
-	if (Instance != nullptr)
+	if (Instance_DEPRECATED != nullptr)
 	{
-		Instance->ConditionalPostLoad();
+		VersionedInstance.Emitter = Instance_DEPRECATED;
+		VersionedInstance.Version = Instance_DEPRECATED->IsVersioningEnabled() ? Instance_DEPRECATED->GetExposedVersion().VersionGuid : FGuid(); 
+	}
+
+	if (UNiagaraEmitter* InstanceEmitter = VersionedInstance.Emitter)
+	{
+		InstanceEmitter->ConditionalPostLoad();
 		if (NiagaraCustomVersion < FNiagaraCustomVersion::MoveInheritanceDataFromTheEmitterHandleToTheEmitter)
 		{
-			if (Source_DEPRECATED != nullptr)
+			FVersionedNiagaraEmitterData* EmitterData = GetEmitterData();
+			if (Source_DEPRECATED && EmitterData)
 			{
 				Source_DEPRECATED->ConditionalPostLoad();
-				Instance->Parent = Source_DEPRECATED;
+				EmitterData->VersionedParent = FVersionedNiagaraEmitter(Source_DEPRECATED, Source_DEPRECATED->GetExposedVersion().VersionGuid);
 				Source_DEPRECATED = nullptr;
 			}
-			if (LastMergedSource_DEPRECATED != nullptr)
+			if (LastMergedSource_DEPRECATED && EmitterData)
 			{
 				LastMergedSource_DEPRECATED->ConditionalPostLoad();
-				Instance->ParentAtLastMerge = LastMergedSource_DEPRECATED;
-				Instance->ParentAtLastMerge->Rename(nullptr, Instance, REN_ForceNoResetLoaders);
+				EmitterData->VersionedParentAtLastMerge = FVersionedNiagaraEmitter(LastMergedSource_DEPRECATED, LastMergedSource_DEPRECATED->GetExposedVersion().VersionGuid);
+				EmitterData->VersionedParentAtLastMerge.Emitter->Rename(nullptr, InstanceEmitter, REN_ForceNoResetLoaders);
 				LastMergedSource_DEPRECATED = nullptr;
 			}
 		}
 
 		FText Reason;
-		if (Instance->GetFName().IsValidObjectName(Reason) == false)
+		if (InstanceEmitter->GetFName().IsValidObjectName(Reason) == false)
 		{
-			UNiagaraSystem* OwningSystem = Instance->GetTypedOuter<UNiagaraSystem>();
+			UNiagaraSystem* OwningSystem = InstanceEmitter->GetTypedOuter<UNiagaraSystem>();
 			if (OwningSystem != nullptr)
 			{
 				// If the name isn't a valid object name, set the name again so that it will be properly sanitized.
@@ -199,14 +228,20 @@ void FNiagaraEmitterHandle::ConditionalPostLoad(int32 NiagaraCustomVersion)
 	}
 }
 
+bool FNiagaraEmitterHandle::UsesEmitter(const FVersionedNiagaraEmitter& InEmitter) const
+{
+	return VersionedInstance == InEmitter || (VersionedInstance.GetEmitterData() && InEmitter.Emitter && VersionedInstance.GetEmitterData()->UsesEmitter(*InEmitter.Emitter));
+}
+
 bool FNiagaraEmitterHandle::UsesEmitter(const UNiagaraEmitter& InEmitter) const
 {
-	return Instance == &InEmitter || (Instance != nullptr && Instance->UsesEmitter(InEmitter));
+	return VersionedInstance.Emitter == &InEmitter || (VersionedInstance.GetEmitterData() && VersionedInstance.GetEmitterData()->UsesEmitter(InEmitter));
 }
 
 void FNiagaraEmitterHandle::ClearEmitter()
 {
-	Instance = nullptr;
+	VersionedInstance = FVersionedNiagaraEmitter();
+	Instance_DEPRECATED = nullptr;
 	Source_DEPRECATED = nullptr;
 	LastMergedSource_DEPRECATED = nullptr;
 }
diff --git a/Engine/Plugins/FX/Niagara/Source/Niagara/Private/NiagaraEmitterInstance.cpp b/Engine/Plugins/FX/Niagara/Source/Niagara/Private/NiagaraEmitterInstance.cpp
--- a/Engine/Plugins/FX/Niagara/Source/Niagara/Private/NiagaraEmitterInstance.cpp
+++ b/Engine/Plugins/FX/Niagara/Source/Niagara/Private/NiagaraEmitterInstance.cpp
@@ -5,7 +5,6 @@
 #include "NiagaraConstants.h"
 #include "NiagaraRenderer.h"
 #include "NiagaraSystemInstance.h"
-#include "NiagaraSystemGpuComputeProxy.h"
 #include "NiagaraComputeExecutionContext.h"
 #include "NiagaraDataInterface.h"
 #include "NiagaraGpuComputeDispatchInterface.h"
@@ -125,6 +124,8 @@ FNiagaraEmitterInstance::FNiagaraEmitterInstance(FNiagaraSystemInstance* InParen
 	, FixedBounds(ForceInit)
 	, CachedSystemFixedBounds(ForceInit)
 	, ParentSystemInstance(InParentSystemInstance)
+	, bResetPending(false)
+	, bCombineEventSpawn(false)
 {
 	ParticleDataSet = new FNiagaraDataSet();
 
@@ -135,7 +136,7 @@ FNiagaraEmitterInstance::FNiagaraEmitterInstance(FNiagaraSystemInstance* InParen
 FNiagaraEmitterInstance::~FNiagaraEmitterInstance()
 {
 	// Clear the cached emitter as it is not safe to access the CacheEmitter due to deferred deleted which can happen after the CachedEmitter has been GCed
-	CachedEmitter = nullptr;
+	CachedEmitter = FVersionedNiagaraEmitter();
 
 	//UE_LOG(LogNiagara, Warning, TEXT("~Simulator %p"), this);
 	CachedBounds.Init();
@@ -185,7 +186,7 @@ TArrayView<FNiagaraScriptExecutionContext> FNiagaraEmitterInstance::GetEventExec
 
 bool FNiagaraEmitterInstance::IsReadyToRun() const
 {
-	if (!IsDisabled() && !CachedEmitter->IsReadyToRun())
+	if (!IsDisabled() && !CachedEmitter.GetEmitterData()->IsReadyToRun())
 	{
 		return false;
 	}
@@ -200,12 +201,12 @@ void FNiagaraEmitterInstance::Dump()const
 		return;
 	}
 
-	UE_LOG(LogNiagara, Log, TEXT("==  %s ========"), *CachedEmitter->GetUniqueEmitterName());
+	UE_LOG(LogNiagara, Log, TEXT("==  %s ========"), *CachedEmitter.Emitter->GetUniqueEmitterName());
 	UE_LOG(LogNiagara, Log, TEXT(".................Spawn................."));
 	SpawnExecContext.Parameters.DumpParameters(true);
 	UE_LOG(LogNiagara, Log, TEXT(".................Update................."));
 	UpdateExecContext.Parameters.DumpParameters(true);
-	if (CachedEmitter->SimTarget == ENiagaraSimTarget::GPUComputeSim && GPUExecContext != nullptr)
+	if (CachedEmitter.GetEmitterData()->SimTarget == ENiagaraSimTarget::GPUComputeSim && GPUExecContext != nullptr)
 	{
 		UE_LOG(LogNiagara, Log, TEXT("................. %s Combined Parameters ................."), TEXT("GPU Script"));
 		GPUExecContext->CombinedParamStore.DumpParameters();
@@ -219,13 +220,13 @@ void FNiagaraEmitterInstance::Dump()const
 
 bool FNiagaraEmitterInstance::IsAllowedToExecute() const
 {
-	if (!GetEmitterHandle().GetIsEnabled()
-		|| !CachedEmitter->IsAllowedByScalability())
+	FVersionedNiagaraEmitterData* EmitterData = CachedEmitter.GetEmitterData();
+	if (!GetEmitterHandle().GetIsEnabled() || EmitterData == nullptr || !EmitterData->IsAllowedByScalability())
 	{
 		return false;
 	}
 
-	if (CachedEmitter->SimTarget == ENiagaraSimTarget::GPUComputeSim)
+	if (EmitterData->SimTarget == ENiagaraSimTarget::GPUComputeSim)
 	{
 		//-TODO: Could replace with CPU side sim in some cases
 		if ( !ComputeDispatchInterface || !FNiagaraUtilities::AllowGPUParticles(ComputeDispatchInterface->GetShaderPlatform()) )
@@ -258,16 +259,17 @@ void FNiagaraEmitterInstance::Init(int32 InEmitterIdx, FNiagaraSystemInstanceID
 
 	bResetPending = false;
 
-	if (CachedEmitter == nullptr)
+	FVersionedNiagaraEmitterData* EmitterData = CachedEmitter.GetEmitterData();
+	if (EmitterData == nullptr)
 	{
 		//@todo(message manager) Error bubbling here
 		ExecutionState = ENiagaraExecutionState::Disabled;
 		return;
 	}
 
-	RandomSeed = CachedEmitter->RandomSeed + ParentSystemInstance->GetRandomSeedOffset();
+	RandomSeed = EmitterData->RandomSeed + ParentSystemInstance->GetRandomSeedOffset();
 
-	MaxAllocationCount = CachedEmitter->GetMaxParticleCountEstimate();
+	MaxAllocationCount = EmitterData->GetMaxParticleCountEstimate();
 	if (!IsAllowedToExecute())
 	{
 		ExecutionState = ENiagaraExecutionState::Disabled;
@@ -310,16 +312,16 @@ void FNiagaraEmitterInstance::Init(int32 InEmitterIdx, FNiagaraSystemInstanceID
 #if !UE_BUILD_SHIPPING && !UE_BUILD_TEST 
 		if (bVerboseAttributeLogging)
 		{
-			for (FNiagaraVariable& Attr : CachedEmitter->UpdateScriptProps.Script->GetVMExecutableData().Attributes)
+			for (FNiagaraVariable& Attr : EmitterData->UpdateScriptProps.Script->GetVMExecutableData().Attributes)
 			{
 				int32 FoundIdx;
-				if (!CachedEmitter->SpawnScriptProps.Script->GetVMExecutableData().Attributes.Find(Attr, FoundIdx))
+				if (!EmitterData->SpawnScriptProps.Script->GetVMExecutableData().Attributes.Find(Attr, FoundIdx))
 				{
 					UE_LOG(LogNiagara, Warning, TEXT("Attribute %s is used in the Update script for %s but it is not initialised in the Spawn script!"), *Attr.GetName().ToString(), *EmitterHandle.GetName().ToString());
 				}
-				for (int32 i = 0; i < CachedEmitter->GetEventHandlers().Num(); i++)
+				for (int32 i = 0; i < EmitterData->GetEventHandlers().Num(); i++)
 				{
-					if (CachedEmitter->GetEventHandlers()[i].Script && !CachedEmitter->GetEventHandlers()[i].Script->GetVMExecutableData().Attributes.Find(Attr, FoundIdx))
+					if (EmitterData->GetEventHandlers()[i].Script && !EmitterData->GetEventHandlers()[i].Script->GetVMExecutableData().Attributes.Find(Attr, FoundIdx))
 					{
 						UE_LOG(LogNiagara, Warning, TEXT("Attribute %s is used in the event handler script for %s but it is not initialised in the Spawn script!"), *Attr.GetName().ToString(), *EmitterHandle.GetName().ToString());
 					}
@@ -330,13 +332,13 @@ void FNiagaraEmitterInstance::Init(int32 InEmitterIdx, FNiagaraSystemInstanceID
 	}
 
 	{
-		ensure(CachedEmitter->UpdateScriptProps.DataSetAccessSynchronized());
-		const int32 UpdateEventGeneratorCount = CachedEmitter->UpdateScriptProps.EventGenerators.Num();
+		ensure(EmitterData->UpdateScriptProps.DataSetAccessSynchronized());
+		const int32 UpdateEventGeneratorCount = EmitterData->UpdateScriptProps.EventGenerators.Num();
 
-		ensure(CachedEmitter->SpawnScriptProps.DataSetAccessSynchronized());
-		const int32 SpawnEventGeneratorCount = CachedEmitter->SpawnScriptProps.EventGenerators.Num();
+		ensure(EmitterData->SpawnScriptProps.DataSetAccessSynchronized());
+		const int32 SpawnEventGeneratorCount = EmitterData->SpawnScriptProps.EventGenerators.Num();
 
-		const int32 NumEvents = CachedEmitter->GetEventHandlers().Num();
+		const int32 NumEvents = EmitterData->GetEventHandlers().Num();
 
 		if (UpdateEventGeneratorCount || SpawnEventGeneratorCount || NumEvents)
 		{
@@ -344,26 +346,26 @@ void FNiagaraEmitterInstance::Init(int32 InEmitterIdx, FNiagaraSystemInstanceID
 			EventInstanceData->UpdateScriptEventDataSets.Empty(UpdateEventGeneratorCount);
 			EventInstanceData->UpdateEventGeneratorIsSharedByIndex.SetNumZeroed(UpdateEventGeneratorCount);
 			int32 UpdateEventGeneratorIndex = 0;
-			for (const FNiagaraEventGeneratorProperties &GeneratorProps : CachedEmitter->UpdateScriptProps.EventGenerators)
+			for (const FNiagaraEventGeneratorProperties &GeneratorProps : EmitterData->UpdateScriptProps.EventGenerators)
 			{
 				FNiagaraDataSet *Set = ParentSystemInstance->CreateEventDataSet(EmitterHandle.GetIdName(), GeneratorProps.ID);
 				Set->Init(&GeneratorProps.DataSetCompiledData);
 
 				EventInstanceData->UpdateScriptEventDataSets.Add(Set);
-				EventInstanceData->UpdateEventGeneratorIsSharedByIndex[UpdateEventGeneratorIndex] = CachedEmitter->IsEventGeneratorShared(GeneratorProps.ID);
+				EventInstanceData->UpdateEventGeneratorIsSharedByIndex[UpdateEventGeneratorIndex] = EmitterData->IsEventGeneratorShared(GeneratorProps.ID);
 				++UpdateEventGeneratorIndex;
 			}
 
 			EventInstanceData->SpawnScriptEventDataSets.Empty(SpawnEventGeneratorCount);
 			EventInstanceData->SpawnEventGeneratorIsSharedByIndex.SetNumZeroed(SpawnEventGeneratorCount);
 			int32 SpawnEventGeneratorIndex = 0;
-			for (const FNiagaraEventGeneratorProperties &GeneratorProps : CachedEmitter->SpawnScriptProps.EventGenerators)
+			for (const FNiagaraEventGeneratorProperties &GeneratorProps : EmitterData->SpawnScriptProps.EventGenerators)
 			{
 				FNiagaraDataSet *Set = ParentSystemInstance->CreateEventDataSet(EmitterHandle.GetIdName(), GeneratorProps.ID);
 				Set->Init(&GeneratorProps.DataSetCompiledData);
 
 				EventInstanceData->SpawnScriptEventDataSets.Add(Set);
-				EventInstanceData->SpawnEventGeneratorIsSharedByIndex[SpawnEventGeneratorIndex] = CachedEmitter->IsEventGeneratorShared(GeneratorProps.ID);
+				EventInstanceData->SpawnEventGeneratorIsSharedByIndex[SpawnEventGeneratorIndex] = EmitterData->IsEventGeneratorShared(GeneratorProps.ID);
 				++SpawnEventGeneratorIndex;
 			}
 
@@ -372,9 +374,9 @@ void FNiagaraEmitterInstance::Init(int32 InEmitterIdx, FNiagaraSystemInstanceID
 
 			for (int32 i = 0; i < NumEvents; i++)
 			{
-				ensure(CachedEmitter->GetEventHandlers()[i].DataSetAccessSynchronized());
+				ensure(EmitterData->GetEventHandlers()[i].DataSetAccessSynchronized());
 
-				UNiagaraScript* EventScript = CachedEmitter->GetEventHandlers()[i].Script;
+				UNiagaraScript* EventScript = EmitterData->GetEventHandlers()[i].Script;
 
 				//This is cpu explicitly? Are we doing event handlers on GPU?
 				EventInstanceData->EventExecContexts[i].Init(EventScript, ENiagaraSimTarget::CPUSim);
@@ -384,19 +386,19 @@ void FNiagaraEmitterInstance::Init(int32 InEmitterIdx, FNiagaraSystemInstanceID
 	}
 
 	{
-		SpawnExecContext.Init(CachedEmitter->SpawnScriptProps.Script, CachedEmitter->SimTarget);
-		UpdateExecContext.Init(CachedEmitter->UpdateScriptProps.Script, CachedEmitter->SimTarget);
+		SpawnExecContext.Init(EmitterData->SpawnScriptProps.Script, EmitterData->SimTarget);
+		UpdateExecContext.Init(EmitterData->UpdateScriptProps.Script, EmitterData->SimTarget);
 
 		// setup the parameer store for the GPU execution context; since spawn and update are combined here, we build one with params from both script props
-		if (CachedEmitter->SimTarget == ENiagaraSimTarget::GPUComputeSim)
+		if (EmitterData->SimTarget == ENiagaraSimTarget::GPUComputeSim)
 		{
 			GPUExecContext = new FNiagaraComputeExecutionContext();
-			GPUExecContext->InitParams(CachedEmitter->GetGPUComputeScript(), CachedEmitter->SimTarget);
-			GPUExecContext->SetDebugSimName(CachedEmitter->GetDebugSimName());
+			GPUExecContext->InitParams(EmitterData->GetGPUComputeScript(), EmitterData->SimTarget);
+			GPUExecContext->SetDebugSimName(EmitterData->GetDebugSimName());
 			GPUExecContext->ProfilingComponentPtr = ParentSystemInstance ? ParentSystemInstance->GetAttachComponent() : nullptr;
-			GPUExecContext->ProfilingEmitterPtr = GetEmitterHandle().GetInstance();
+			GPUExecContext->ProfilingEmitterPtr = CachedEmitter.ToWeakPtr();
 			GPUExecContext->MainDataSet = ParticleDataSet;
-			GPUExecContext->GPUScript_RT = CachedEmitter->GetGPUComputeScript()->GetRenderThreadScript();
+			GPUExecContext->GPUScript_RT = EmitterData->GetGPUComputeScript()->GetRenderThreadScript();
 		}
 	}
 
@@ -415,17 +417,17 @@ void FNiagaraEmitterInstance::Init(int32 InEmitterIdx, FNiagaraSystemInstanceID
 	{
 		// Collect script defined data interface parameters.
 		TArray<UNiagaraScript*, TInlineAllocator<8>> Scripts;
-		Scripts.Add(CachedEmitter->SpawnScriptProps.Script);
-		Scripts.Add(CachedEmitter->UpdateScriptProps.Script);
-		for (const FNiagaraEventScriptProperties& EventHandler : CachedEmitter->GetEventHandlers())
+		Scripts.Add(EmitterData->SpawnScriptProps.Script);
+		Scripts.Add(EmitterData->UpdateScriptProps.Script);
+		for (const FNiagaraEventScriptProperties& EventHandler : EmitterData->GetEventHandlers())
 		{
 			Scripts.Add(EventHandler.Script);
 		}
-		for (const UNiagaraSimulationStageBase* SimStage : CachedEmitter->GetSimulationStages())
+		for (const UNiagaraSimulationStageBase* SimStage : EmitterData->GetSimulationStages())
 		{
 			Scripts.Add(SimStage->Script);
 		}
-		FNiagaraUtilities::CollectScriptDataInterfaceParameters(*CachedEmitter, MakeArrayView(Scripts), ScriptDefinedDataInterfaceParameters);
+		FNiagaraUtilities::CollectScriptDataInterfaceParameters(*CachedEmitter.Emitter, MakeArrayView(Scripts), ScriptDefinedDataInterfaceParameters);
 
 		//Bind some stores and unbind immediately just to prime some data from those stores.
 		FNiagaraParameterStore& SystemScriptDefinedDataInterfaceParameters = ParentSystemInstance->GetSystemSimulation()->GetScriptDefinedDataInterfaceParameters();
@@ -454,12 +456,12 @@ void FNiagaraEmitterInstance::Init(int32 InEmitterIdx, FNiagaraSystemInstanceID
 				EventContext.Parameters.UnbindFromSourceStores();
 			}
 
-			const int32 NumEventHandlers = CachedEmitter->GetEventHandlers().Num();
+			const int32 NumEventHandlers = EmitterData->GetEventHandlers().Num();
 			EventInstanceData->EventHandlingInfo.Reset();
 			EventInstanceData->EventHandlingInfo.SetNum(NumEventHandlers);
 			for (int32 i = 0; i < NumEventHandlers; i++)
 			{
-				const FNiagaraEventScriptProperties& EventHandlerProps = CachedEmitter->GetEventHandlers()[i];
+				const FNiagaraEventScriptProperties& EventHandlerProps = EmitterData->GetEventHandlers()[i];
 				FNiagaraEventHandlingInfo& Info = EventInstanceData->EventHandlingInfo[i];
 				Info.SourceEmitterGuid = EventHandlerProps.SourceEmitterID;
 				Info.SourceEmitterName = Info.SourceEmitterGuid.IsValid() ? *Info.SourceEmitterGuid.ToString() : CachedIDName;
@@ -473,7 +475,7 @@ void FNiagaraEmitterInstance::Init(int32 InEmitterIdx, FNiagaraSystemInstanceID
 		RendererBindings.UnbindAll();
 
 		bool bAnyRendererBindingsAdded = false;
-		for (UNiagaraRendererProperties* Props : CachedEmitter->GetRenderers())
+		for (UNiagaraRendererProperties* Props : EmitterData->GetRenderers())
 		{
 			if (Props && Props->bIsEnabled)
 			{
@@ -508,21 +510,21 @@ void FNiagaraEmitterInstance::Init(int32 InEmitterIdx, FNiagaraSystemInstanceID
 			SystemScriptDefinedDataInterfaceParameters.Bind(&RendererBindings);
 			ScriptDefinedDataInterfaceParameters.Bind(&RendererBindings);
 
-			if (GPUExecContext && CachedEmitter->SimTarget == ENiagaraSimTarget::GPUComputeSim)
+			if (GPUExecContext && EmitterData->SimTarget == ENiagaraSimTarget::GPUComputeSim)
 			{
 				GPUExecContext->CombinedParamStore.Bind(&RendererBindings);
 			}
 
 			// Handle populating static variable values into the list
-			CachedEmitter->RendererBindings.Bind(&RendererBindings);
+			CachedEmitter.GetEmitterData()->RendererBindings.Bind(&RendererBindings);
 		}
 	}	
 
-	MaxInstanceCount = CachedEmitter->GetMaxInstanceCount();
+	MaxInstanceCount = EmitterData->GetMaxInstanceCount();
 	ParticleDataSet->SetMaxInstanceCount(MaxInstanceCount);
-	ParticleDataSet->SetMaxAllocationCount(CachedEmitter->GetMaxAllocationCount());
+	ParticleDataSet->SetMaxAllocationCount(EmitterData->GetMaxAllocationCount());
 
-	bCombineEventSpawn = GbNiagaraAllowEventSpawnCombine && (CachedEmitter->bCombineEventSpawn || (GbNiagaraAllowEventSpawnCombine == 2));
+	bCombineEventSpawn = GbNiagaraAllowEventSpawnCombine && (EmitterData->bCombineEventSpawn || (GbNiagaraAllowEventSpawnCombine == 2));
 }
 
 void FNiagaraEmitterInstance::ResetSimulation(bool bKillExisting /*= true*/)
@@ -546,8 +548,9 @@ void FNiagaraEmitterInstance::ResetSimulation(bool bKillExisting /*= true*/)
 		return;
 	}
 
-	check(CachedEmitter);
-	RandomSeed = CachedEmitter->RandomSeed + ParentSystemInstance->GetRandomSeedOffset();
+	FVersionedNiagaraEmitterData* EmitterData = CachedEmitter.GetEmitterData();
+	check(EmitterData);
+	RandomSeed = EmitterData->RandomSeed + ParentSystemInstance->GetRandomSeedOffset();
 
 	SetExecutionState(ENiagaraExecutionState::Active);
 
@@ -581,58 +584,59 @@ void FNiagaraEmitterInstance::CheckForErrors()
 {
 	SCOPE_CYCLE_COUNTER(STAT_NiagaraEmitterErrorCheck);
 	
-	checkSlow(CachedEmitter);
+	checkSlow(CachedEmitter.Emitter);
 
 	//Check for various failure conditions and bail.
-	if (!CachedEmitter->UpdateScriptProps.Script || !CachedEmitter->SpawnScriptProps.Script )
+	FVersionedNiagaraEmitterData* EmitterData = CachedEmitter.GetEmitterData();
+	if (!EmitterData->UpdateScriptProps.Script || !EmitterData->SpawnScriptProps.Script )
 	{
 		//TODO - Arbitrary named scripts. Would need some base functionality for Spawn/Udpate to be called that can be overriden in BPs for emitters with custom scripts.
-		UE_LOG(LogNiagara, Error, TEXT("Emitter cannot be enabled because it's doesn't have both an update and spawn script."), *CachedEmitter->GetFullName());
+		UE_LOG(LogNiagara, Error, TEXT("Emitter cannot be enabled because it's doesn't have both an update and spawn script."), *CachedEmitter.Emitter->GetFullName());
 		SetExecutionState(ENiagaraExecutionState::Disabled);
 		return;
 	}
 
-	if (!CachedEmitter->UpdateScriptProps.Script->IsReadyToRun(ENiagaraSimTarget::CPUSim) || !CachedEmitter->SpawnScriptProps.Script->IsReadyToRun(ENiagaraSimTarget::CPUSim))
+	if (!EmitterData->UpdateScriptProps.Script->IsReadyToRun(ENiagaraSimTarget::CPUSim) || !EmitterData->SpawnScriptProps.Script->IsReadyToRun(ENiagaraSimTarget::CPUSim))
 	{
 		//TODO - Arbitrary named scripts. Would need some base functionality for Spawn/Udpate to be called that can be overriden in BPs for emitters with custom scripts.
-		UE_LOG(LogNiagara, Error, TEXT("Emitter cannot be enabled because it's doesn't have both an update and spawn script ready to run CPU scripts."), *CachedEmitter->GetFullName());
+		UE_LOG(LogNiagara, Error, TEXT("Emitter cannot be enabled because it's doesn't have both an update and spawn script ready to run CPU scripts."), *CachedEmitter.Emitter->GetFullName());
 		SetExecutionState(ENiagaraExecutionState::Disabled);
 		return;
 	}
 
-	if (CachedEmitter->SpawnScriptProps.Script->GetVMExecutableData().DataUsage.bReadsAttributeData)
+	if (EmitterData->SpawnScriptProps.Script->GetVMExecutableData().DataUsage.bReadsAttributeData)
 	{
-		UE_LOG(LogNiagara, Error, TEXT("%s reads attribute data and so cannot be used as a spawn script. The data being read would be invalid."), *CachedEmitter->SpawnScriptProps.Script->GetName());
+		UE_LOG(LogNiagara, Error, TEXT("%s reads attribute data and so cannot be used as a spawn script. The data being read would be invalid."), *EmitterData->SpawnScriptProps.Script->GetName());
 		SetExecutionState(ENiagaraExecutionState::Disabled);
 		return;
 	}
-	if (CachedEmitter->UpdateScriptProps.Script->GetVMExecutableData().Attributes.Num() == 0 || CachedEmitter->SpawnScriptProps.Script->GetVMExecutableData().Attributes.Num() == 0)
+	if (EmitterData->UpdateScriptProps.Script->GetVMExecutableData().Attributes.Num() == 0 || EmitterData->SpawnScriptProps.Script->GetVMExecutableData().Attributes.Num() == 0)
 	{
 		UE_LOG(LogNiagara, Error, TEXT("This emitter cannot be enabled because its spawn or update script doesn't have any attributes.."));
 		SetExecutionState(ENiagaraExecutionState::Disabled);
 		return;
 	}
 
-	if (CachedEmitter->SimTarget == ENiagaraSimTarget::CPUSim)
+	if (EmitterData->SimTarget == ENiagaraSimTarget::CPUSim)
 	{
 		bool bFailed = false;
-		if (!CachedEmitter->SpawnScriptProps.Script->DidScriptCompilationSucceed(false))
+		if (!EmitterData->SpawnScriptProps.Script->DidScriptCompilationSucceed(false))
 		{
 			bFailed = true;
 			UE_LOG(LogNiagara, Error, TEXT("This emitter cannot be enabled because its CPU Spawn script failed to compile."));
 		}
 
-		if (!CachedEmitter->UpdateScriptProps.Script->DidScriptCompilationSucceed(false))
+		if (!EmitterData->UpdateScriptProps.Script->DidScriptCompilationSucceed(false))
 		{
 			bFailed = true;
 			UE_LOG(LogNiagara, Error, TEXT("This emitter cannot be enabled because its CPU Update script failed to compile."));
 		}
 
-		if (CachedEmitter->GetEventHandlers().Num() != 0)
+		if (EmitterData->GetEventHandlers().Num() != 0)
 		{
-			for (int32 i = 0; i < CachedEmitter->GetEventHandlers().Num(); i++)
+			for (int32 i = 0; i < EmitterData->GetEventHandlers().Num(); i++)
 			{
-				if (!CachedEmitter->GetEventHandlers()[i].Script->DidScriptCompilationSucceed(false))
+				if (!EmitterData->GetEventHandlers()[i].Script->DidScriptCompilationSucceed(false))
 				{
 					bFailed = true;
 					UE_LOG(LogNiagara, Error, TEXT("This emitter cannot be enabled because one of its CPU Event scripts failed to compile."));
@@ -647,15 +651,15 @@ void FNiagaraEmitterInstance::CheckForErrors()
 		}
 	}
 
-	if (CachedEmitter->SimTarget == ENiagaraSimTarget::GPUComputeSim)
+	if (EmitterData->SimTarget == ENiagaraSimTarget::GPUComputeSim)
 	{
-		if (CachedEmitter->GetGPUComputeScript()->IsScriptCompilationPending(true))
+		if (EmitterData->GetGPUComputeScript()->IsScriptCompilationPending(true))
 		{
 			UE_LOG(LogNiagara, Error, TEXT("This emitter cannot be enabled because its GPU script hasn't been compiled.."));
 			SetExecutionState(ENiagaraExecutionState::Disabled);
 			return;
 		}
-		if (!CachedEmitter->GetGPUComputeScript()->DidScriptCompilationSucceed(true))
+		if (!EmitterData->GetGPUComputeScript()->DidScriptCompilationSucceed(true))
 		{
 			UE_LOG(LogNiagara, Error, TEXT("This emitter cannot be enabled because its GPU script failed to compile."));
 			SetExecutionState(ENiagaraExecutionState::Disabled);
@@ -675,7 +679,7 @@ void FNiagaraEmitterInstance::DirtyDataInterfaces()
 	SpawnExecContext.DirtyDataInterfaces();
 	UpdateExecContext.DirtyDataInterfaces();
 
-	if (CachedEmitter->SimTarget == ENiagaraSimTarget::GPUComputeSim && GPUExecContext != nullptr)
+	if (CachedEmitter.GetEmitterData()->SimTarget == ENiagaraSimTarget::GPUComputeSim && GPUExecContext != nullptr)
 	{
 		GPUExecContext->DirtyDataInterfaces();
 	}
@@ -766,12 +770,12 @@ void FNiagaraEmitterInstance::BindParameters(bool bExternalOnly)
 			else
 			{
 				UE_LOG(LogNiagara, Error, TEXT("Emitter attempting to bind to a null Parameter Collection Instance.\nEmitter:%s\nCollection:%s")
-				, CachedEmitter ? *CachedEmitter->GetPathName() : TEXT("null emitter!"), *Collection->GetPathName());
+				, CachedEmitter.Emitter ? *CachedEmitter.Emitter->GetPathName() : TEXT("null emitter!"), *Collection->GetPathName());
 			}
 		}
 		else
 		{
-			UE_LOG(LogNiagara, Error, TEXT("Emitter attempting to bind to a null Parameter Collection.\nEmitter:%s"), CachedEmitter ? *CachedEmitter->GetPathName() : TEXT("null emitter!" ));
+			UE_LOG(LogNiagara, Error, TEXT("Emitter attempting to bind to a null Parameter Collection.\nEmitter:%s"), CachedEmitter.Emitter ? *CachedEmitter.Emitter->GetPathName() : TEXT("null emitter!" ));
 		}
 	};
 
@@ -784,7 +788,8 @@ void FNiagaraEmitterInstance::BindParameters(bool bExternalOnly)
 		BindToParameterCollection(Collection, UpdateExecContext.Parameters);
 	}
 
-	if (CachedEmitter->SimTarget == ENiagaraSimTarget::GPUComputeSim)
+	FVersionedNiagaraEmitterData* EmitterData = CachedEmitter.GetEmitterData();
+	if (EmitterData->SimTarget == ENiagaraSimTarget::GPUComputeSim)
 	{
 		for (UNiagaraParameterCollection* Collection : SpawnExecContext.Script->GetCachedParameterCollectionReferences())
 		{
@@ -804,8 +809,8 @@ void FNiagaraEmitterInstance::BindParameters(bool bExternalOnly)
 		}
 	}
 
-	FNiagaraScriptInstanceParameterStore& TargetParamStore = CachedEmitter->SimTarget == ENiagaraSimTarget::GPUComputeSim ? GPUExecContext->CombinedParamStore : UpdateExecContext.Parameters;
-	for (const UNiagaraSimulationStageBase* SimStage : CachedEmitter->GetSimulationStages())
+	FNiagaraScriptInstanceParameterStore& TargetParamStore = EmitterData->SimTarget == ENiagaraSimTarget::GPUComputeSim ? GPUExecContext->CombinedParamStore : UpdateExecContext.Parameters;
+	for (const UNiagaraSimulationStageBase* SimStage : EmitterData->GetSimulationStages())
 	{
 		if (SimStage->bEnabled == false)
 		{
@@ -832,29 +837,29 @@ void FNiagaraEmitterInstance::BindParameters(bool bExternalOnly)
 		}
 
 #if WITH_EDITORONLY_DATA
-		CachedEmitter->SpawnScriptProps.Script->RapidIterationParameters.Bind(&SpawnExecContext.Parameters);
-		CachedEmitter->UpdateScriptProps.Script->RapidIterationParameters.Bind(&UpdateExecContext.Parameters);
+		EmitterData->SpawnScriptProps.Script->RapidIterationParameters.Bind(&SpawnExecContext.Parameters);
+		EmitterData->UpdateScriptProps.Script->RapidIterationParameters.Bind(&UpdateExecContext.Parameters);
 
 		if (EventInstanceData.IsValid())
 		{
-			ensure(CachedEmitter->GetEventHandlers().Num() == EventInstanceData->EventExecContexts.Num());
-			for (int32 i = 0; i < CachedEmitter->GetEventHandlers().Num(); i++)
+			ensure(EmitterData->GetEventHandlers().Num() == EventInstanceData->EventExecContexts.Num());
+			for (int32 i = 0; i < EmitterData->GetEventHandlers().Num(); i++)
 			{
-				CachedEmitter->GetEventHandlers()[i].Script->RapidIterationParameters.Bind(&EventInstanceData->EventExecContexts[i].Parameters);
+				EmitterData->GetEventHandlers()[i].Script->RapidIterationParameters.Bind(&EventInstanceData->EventExecContexts[i].Parameters);
 			}
 		}
 	#endif
 
-		if (CachedEmitter->SimTarget == ENiagaraSimTarget::GPUComputeSim)
+		if (EmitterData->SimTarget == ENiagaraSimTarget::GPUComputeSim)
 		{
 			InstanceParams.Bind(&GPUExecContext->CombinedParamStore);
 #if WITH_EDITORONLY_DATA
-			CachedEmitter->SpawnScriptProps.Script->RapidIterationParameters.Bind(&GPUExecContext->CombinedParamStore);
-			CachedEmitter->UpdateScriptProps.Script->RapidIterationParameters.Bind(&GPUExecContext->CombinedParamStore);
+			EmitterData->SpawnScriptProps.Script->RapidIterationParameters.Bind(&GPUExecContext->CombinedParamStore);
+			EmitterData->UpdateScriptProps.Script->RapidIterationParameters.Bind(&GPUExecContext->CombinedParamStore);
 
-			for (int32 i = 0; i < CachedEmitter->GetSimulationStages().Num(); i++)
+			for (int32 i = 0; i < EmitterData->GetSimulationStages().Num(); i++)
 			{
-				CachedEmitter->GetSimulationStages()[i]->Script->RapidIterationParameters.Bind(&GPUExecContext->CombinedParamStore);
+				EmitterData->GetSimulationStages()[i]->Script->RapidIterationParameters.Bind(&GPUExecContext->CombinedParamStore);
 			}
 #endif
 		}
@@ -931,12 +936,13 @@ int64 FNiagaraEmitterInstance::GetTotalBytesUsed()
 
 FBox FNiagaraEmitterInstance::InternalCalculateDynamicBounds(int32 ParticleCount) const
 {
-	if (!ParticleCount || !CachedEmitter || !ParentSystemInstance)
+	FVersionedNiagaraEmitterData* EmitterData = CachedEmitter.GetEmitterData();
+	if (!ParticleCount || !EmitterData || !ParentSystemInstance)
 	{
 		return FBox(ForceInit);
 	}
 
-	const auto BoundsCalculators = CachedEmitter->GetBoundsCalculators();
+	const auto BoundsCalculators = EmitterData->GetBoundsCalculators();
 	if (!BoundsCalculators.Num())
 	{
 		return FBox(ForceInit);
@@ -946,7 +952,7 @@ FBox FNiagaraEmitterInstance::InternalCalculateDynamicBounds(int32 ParticleCount
 	Ret.Init();
 
 	const FTransform& Transform = ParentSystemInstance->GetWorldTransform();
-	for (const TUniquePtr<FNiagaraBoundsCalculator>& BoundsCalculator : BoundsCalculators)
+	for (const TSharedPtr<FNiagaraBoundsCalculator>& BoundsCalculator : BoundsCalculators)
 	{
 		Ret += BoundsCalculator->CalculateBounds(Transform, *ParticleDataSet, ParticleCount);
 	}
@@ -957,17 +963,16 @@ FBox FNiagaraEmitterInstance::InternalCalculateDynamicBounds(int32 ParticleCount
 #if WITH_EDITOR
 void FNiagaraEmitterInstance::CalculateFixedBounds(const FTransform& ToWorldSpace)
 {
-	check(CachedEmitter);
-
-	if (IsComplete() || CachedEmitter == nullptr)
+	check(CachedEmitter.Emitter);
+	if (IsComplete())
 	{
-		return ;
+		return;
 	}
 
 	FScopedNiagaraDataSetGPUReadback ScopedGPUReadback;
-
 	int32 NumInstances = 0;
-	if (CachedEmitter->SimTarget == ENiagaraSimTarget::GPUComputeSim)
+	FVersionedNiagaraEmitterData* EmitterData = CachedEmitter.GetEmitterData();
+	if (EmitterData->SimTarget == ENiagaraSimTarget::GPUComputeSim)
 	{
 		if (GPUExecContext == nullptr)
 		{
@@ -991,15 +996,15 @@ void FNiagaraEmitterInstance::CalculateFixedBounds(const FTransform& ToWorldSpac
 	if (!Bounds.IsValid)
 		return;
 
-	CachedEmitter->Modify();
-	CachedEmitter->CalculateBoundsMode = ENiagaraEmitterCalculateBoundMode::Fixed;
-	if (CachedEmitter->bLocalSpace)
+	CachedEmitter.Emitter->Modify();
+	EmitterData->CalculateBoundsMode = ENiagaraEmitterCalculateBoundMode::Fixed;
+	if (EmitterData->bLocalSpace)
 	{
-		CachedEmitter->FixedBounds = Bounds;
+		EmitterData->FixedBounds = Bounds;
 	}
 	else
 	{
-		CachedEmitter->FixedBounds = Bounds.TransformBy(ToWorldSpace);
+		EmitterData->FixedBounds = Bounds.TransformBy(ToWorldSpace);
 	}
 
 	bCachedBoundsDynamic = false;
@@ -1014,7 +1019,7 @@ void FNiagaraEmitterInstance::PostTick()
 {
 	SCOPE_CYCLE_COUNTER(STAT_NiagaraEmitterPostTick);
 
-	checkSlow(CachedEmitter);
+	checkSlow(CachedEmitter.Emitter);
 
 	if (EventInstanceData.IsValid())
 	{
@@ -1027,6 +1032,8 @@ void FNiagaraEmitterInstance::PostTick()
 
 	bCachedBoundsDynamic = false;
 	CachedBounds.Init();
+	FVersionedNiagaraEmitterData* EmitterData = CachedEmitter.GetEmitterData();
+	if (EmitterData)
 	{
 		// Read lock can be smaller in scope, but probably not necessary
 		FRWScopeLock ScopeLock(FixedBoundsGuard, SLT_ReadOnly);
@@ -1038,17 +1045,17 @@ void FNiagaraEmitterInstance::PostTick()
 		{
 			CachedBounds = CachedSystemFixedBounds;
 		}
-		else if ( CachedEmitter->CalculateBoundsMode == ENiagaraEmitterCalculateBoundMode::Fixed )
+		else if ( EmitterData->CalculateBoundsMode == ENiagaraEmitterCalculateBoundMode::Fixed )
 		{
-			CachedBounds = CachedEmitter->FixedBounds;
+			CachedBounds = EmitterData->FixedBounds;
 		}
-		else if ( CachedEmitter->CalculateBoundsMode == ENiagaraEmitterCalculateBoundMode::Dynamic )
+		else if ( EmitterData->CalculateBoundsMode == ENiagaraEmitterCalculateBoundMode::Dynamic )
 		{
 			// We do not support dynamic bounds for GPU yet
 			// Therefore we need to return the default bounds in this case to avoid flickering (some content relies on this path)
-			if (CachedEmitter->SimTarget == ENiagaraSimTarget::GPUComputeSim)
+			if (EmitterData->SimTarget == ENiagaraSimTarget::GPUComputeSim)
 			{
-				CachedBounds = UNiagaraEmitter::GetDefaultFixedBounds();
+				CachedBounds = FVersionedNiagaraEmitterData::GetDefaultFixedBounds();
 			}
 			else
 			{
@@ -1056,7 +1063,7 @@ void FNiagaraEmitterInstance::PostTick()
 				if (DynamicBounds.IsValid)
 				{
 					bCachedBoundsDynamic = true;
-					if (CachedEmitter->bLocalSpace)
+					if (EmitterData->bLocalSpace)
 					{
 						CachedBounds = DynamicBounds;
 					}
@@ -1067,15 +1074,12 @@ void FNiagaraEmitterInstance::PostTick()
 				}
 			}
 		}
-	}
 
 #if STATS
-	if (UNiagaraEmitter* Emitter = GetCachedEmitter())
-	{
-		Emitter->GetStatData().AddStatCapture(TTuple<uint64, ENiagaraScriptUsage>((uint64)this, ENiagaraScriptUsage::ParticleSpawnScript), GetSpawnExecutionContext().ReportStats());
-		Emitter->GetStatData().AddStatCapture(TTuple<uint64, ENiagaraScriptUsage>((uint64)this, ENiagaraScriptUsage::ParticleUpdateScript), GetUpdateExecutionContext().ReportStats());
-	}
+		EmitterData->GetStatData().AddStatCapture(TTuple<uint64, ENiagaraScriptUsage>((uint64)this, ENiagaraScriptUsage::ParticleSpawnScript), GetSpawnExecutionContext().ReportStats());
+		EmitterData->GetStatData().AddStatCapture(TTuple<uint64, ENiagaraScriptUsage>((uint64)this, ENiagaraScriptUsage::ParticleUpdateScript), GetUpdateExecutionContext().ReportStats());
 #endif
+	}
 }
 
 bool FNiagaraEmitterInstance::HandleCompletion(bool bForce)
@@ -1118,7 +1122,7 @@ bool FNiagaraEmitterInstance::HandleCompletion(bool bForce)
 bool FNiagaraEmitterInstance::RequiresPersistentIDs() const
 {
 	//TODO: can we have this be enabled at runtime from outside the system?
-	return GetEmitterHandle().GetInstance()->RequiresPersistentIDs() || ParticleDataSet->HasVariable(SYS_PARAM_PARTICLES_ID);
+	return GetEmitterHandle().GetEmitterData()->RequiresPersistentIDs() || ParticleDataSet->HasVariable(SYS_PARAM_PARTICLES_ID);
 }
 
 #if WITH_EDITOR
@@ -1129,14 +1133,15 @@ void FNiagaraEmitterInstance::TickRapidIterationParameters()
 		return;
 	}
 
-	CachedEmitter->SpawnScriptProps.Script->RapidIterationParameters.Tick();
-	CachedEmitter->UpdateScriptProps.Script->RapidIterationParameters.Tick();
+	FVersionedNiagaraEmitterData* EmitterData = CachedEmitter.GetEmitterData();
+	EmitterData->SpawnScriptProps.Script->RapidIterationParameters.Tick();
+	EmitterData->UpdateScriptProps.Script->RapidIterationParameters.Tick();
 	if (EventInstanceData.IsValid())
 	{
-		ensure(CachedEmitter->GetEventHandlers().Num() == EventInstanceData->EventExecContexts.Num());
-		for (int32 i = 0; i < CachedEmitter->GetEventHandlers().Num(); i++)
+		ensure(EmitterData->GetEventHandlers().Num() == EventInstanceData->EventExecContexts.Num());
+		for (int32 i = 0; i < EmitterData->GetEventHandlers().Num(); i++)
 		{
-			CachedEmitter->GetEventHandlers()[i].Script->RapidIterationParameters.Tick();
+			EmitterData->GetEventHandlers()[i].Script->RapidIterationParameters.Tick();
 		}
 	}
 }
@@ -1151,27 +1156,28 @@ void FNiagaraEmitterInstance::PreTick()
 	{
 		return;
 	}
+	FVersionedNiagaraEmitterData* NiagaraEmitterData = CachedEmitter.GetEmitterData();
 
 #if STATS
-	FScopeCycleCounter SystemStatCounter(CachedEmitter->GetStatID(true, true));
+	FScopeCycleCounter SystemStatCounter(CachedEmitter.Emitter->GetStatID(true, true));
 #endif
 
 	checkSlow(ParticleDataSet);
 	FNiagaraDataSet& Data = *ParticleDataSet;
 
 	bool bOk = true;
-	bOk &= SpawnExecContext.Tick(ParentSystemInstance, CachedEmitter->SimTarget);
-	bOk &= UpdateExecContext.Tick(ParentSystemInstance, CachedEmitter->SimTarget);
+	bOk &= SpawnExecContext.Tick(ParentSystemInstance, NiagaraEmitterData->SimTarget);
+	bOk &= UpdateExecContext.Tick(ParentSystemInstance, NiagaraEmitterData->SimTarget);
 
 	// @todo THREADSAFETY We should not tick GPU contexts on the game thread!
-	if (CachedEmitter->SimTarget == ENiagaraSimTarget::GPUComputeSim && GPUExecContext != nullptr)
+	if (NiagaraEmitterData->SimTarget == ENiagaraSimTarget::GPUComputeSim && GPUExecContext != nullptr)
 	{
 		bOk &= GPUExecContext->Tick(ParentSystemInstance);
 	}
 
 	for (FNiagaraScriptExecutionContext& EventContext : GetEventExecutionContexts())
 	{
-		bOk &= EventContext.Tick(ParentSystemInstance, CachedEmitter->SimTarget);
+		bOk &= EventContext.Tick(ParentSystemInstance, NiagaraEmitterData->SimTarget);
 	}
 
 	if (!bOk)
@@ -1186,7 +1192,7 @@ void FNiagaraEmitterInstance::PreTick()
 		//On our very first frame we prime any previous params (for interpolation).
 		SpawnExecContext.PostTick();
 		UpdateExecContext.PostTick();
-		if (CachedEmitter->SimTarget == ENiagaraSimTarget::GPUComputeSim && GPUExecContext != nullptr)
+		if (NiagaraEmitterData->SimTarget == ENiagaraSimTarget::GPUComputeSim && GPUExecContext != nullptr)
 		{
 			//We PostTick the GPUExecContext here to prime crucial PREV parameters (such as PREV_Engine.Owner.Position). This PostTick call is necessary as the GPUExecContext has not been sent to the dispatch interface yet.
 			GPUExecContext->PostTick();
@@ -1199,8 +1205,8 @@ void FNiagaraEmitterInstance::PreTick()
 	}
 
 	checkSlow(Data.GetNumVariables() > 0);
-	checkSlow(CachedEmitter->SpawnScriptProps.Script);
-	checkSlow(CachedEmitter->UpdateScriptProps.Script);
+	checkSlow(CachedEmitter.GetEmitterData()->SpawnScriptProps.Script);
+	checkSlow(CachedEmitter.GetEmitterData()->UpdateScriptProps.Script);
 
 	if (bResetPending)
 	{
@@ -1236,7 +1242,7 @@ void FNiagaraEmitterInstance::PreTick()
 bool FNiagaraEmitterInstance::WaitForDebugInfo()
 {
 	FNiagaraComputeExecutionContext* DebugContext = GPUExecContext;
-	if (CachedEmitter->SimTarget == ENiagaraSimTarget::GPUComputeSim && DebugContext)
+	if (CachedEmitter.GetEmitterData()->SimTarget == ENiagaraSimTarget::GPUComputeSim && DebugContext)
 	{
 		ENQUEUE_RENDER_COMMAND(CaptureCommand)([=](FRHICommandListImmediate& RHICmdList) { ComputeDispatchInterface->ProcessDebugReadbacks(RHICmdList, true); });
 		FlushRenderingCommands(); 
@@ -1259,12 +1265,12 @@ FBox FNiagaraEmitterInstance::GetFixedBounds() const
 			return FixedBounds;
 		}
 	}
-
-	if ( CachedEmitter && CachedEmitter->CalculateBoundsMode == ENiagaraEmitterCalculateBoundMode::Fixed )
+	FVersionedNiagaraEmitterData* EmitterData = CachedEmitter.GetEmitterData();
+	if (EmitterData && EmitterData->CalculateBoundsMode == ENiagaraEmitterCalculateBoundMode::Fixed)
 	{
-		return CachedEmitter->FixedBounds;
+		return EmitterData->FixedBounds;
 	}
-	return FBox(EForceInit::ForceInit);
+	return FBox(ForceInit);
 }
 
 static int32 GbTriggerCrash = 0;
@@ -1283,11 +1289,11 @@ FORCENOINLINE void NiagaraTestCrash()
 void FNiagaraEmitterInstance::Tick(float DeltaSeconds)
 {
 	SCOPE_CYCLE_COUNTER(STAT_NiagaraTick);
-	FScopeCycleCounterUObject AdditionalScope(CachedEmitter, GET_STATID(STAT_NiagaraTick));
+	FScopeCycleCounterUObject AdditionalScope(CachedEmitter.Emitter, GET_STATID(STAT_NiagaraTick));
 	FNiagaraEditorOnlyCycleTimer<false> TickTime(CPUTimeCycles);
 
 #if STATS
-	FScopeCycleCounter SystemStatCounter(CachedEmitter->GetStatID(true, true));
+	FScopeCycleCounter SystemStatCounter(CachedEmitter.Emitter->GetStatID(true, true));
 #endif
 
 	if (HandleCompletion())
@@ -1321,7 +1327,8 @@ void FNiagaraEmitterInstance::Tick(float DeltaSeconds)
 		return;
 	}
 
-	if (CachedEmitter->SimTarget == ENiagaraSimTarget::CPUSim && Data.GetCurrentDataChecked().GetNumInstances() == 0 && ExecutionState != ENiagaraExecutionState::Active)
+	FVersionedNiagaraEmitterData* EmitterData = CachedEmitter.GetEmitterData();
+	if (EmitterData->SimTarget == ENiagaraSimTarget::CPUSim && Data.GetCurrentDataChecked().GetNumInstances() == 0 && ExecutionState != ENiagaraExecutionState::Active)
 	{
 		Data.ResetBuffers();
 		return;
@@ -1331,24 +1338,24 @@ void FNiagaraEmitterInstance::Tick(float DeltaSeconds)
 
 	if (GbDumpParticleData || System->bDumpDebugEmitterInfo)
 	{
-		UE_LOG(LogNiagara, Log, TEXT("|||||||||||||||||||||||||||||||||||||||||||||||||||||||||||||||||||||||||||||||||||||||||||||||||||||||||||||"), *CachedEmitter->GetPathName());
-		UE_LOG(LogNiagara, Log, TEXT("|=== FNiagaraEmitterInstance::Tick [ %s ] ===============|"), *CachedEmitter->GetPathName());
+		UE_LOG(LogNiagara, Log, TEXT("|||||||||||||||||||||||||||||||||||||||||||||||||||||||||||||||||||||||||||||||||||||||||||||||||||||||||||||"), *CachedEmitter.Emitter->GetPathName());
+		UE_LOG(LogNiagara, Log, TEXT("|=== FNiagaraEmitterInstance::Tick [ %s ] ===============|"), *CachedEmitter.Emitter->GetPathName());
 	}
 
 
 	checkSlow(Data.GetNumVariables() > 0);
-	checkSlow(CachedEmitter->SpawnScriptProps.Script);
-	checkSlow(CachedEmitter->UpdateScriptProps.Script);
+	checkSlow(CachedEmitter.GetEmitterData()->SpawnScriptProps.Script);
+	checkSlow(CachedEmitter.GetEmitterData()->UpdateScriptProps.Script);
 	
 	if (EventInstanceData.IsValid())
 	{
 		// Set up the spawn counts and source datasets for the events. The system ensures that we will run after any emitters
 		// we're receiving from, so we can use the data buffers that our sources have computed this tick.
-		const int32 NumEventHandlers = CachedEmitter->GetEventHandlers().Num();
+		const int32 NumEventHandlers = EmitterData->GetEventHandlers().Num();
 		EventInstanceData->EventSpawnTotal = 0;
 		for (int32 i = 0; i < NumEventHandlers; i++)
 		{
-			const FNiagaraEventScriptProperties& EventHandlerProps = CachedEmitter->GetEventHandlers()[i];
+			const FNiagaraEventScriptProperties& EventHandlerProps = EmitterData->GetEventHandlers()[i];
 			FNiagaraEventHandlingInfo& Info = EventInstanceData->EventHandlingInfo[i];
 
 			Info.TotalSpawnCount = 0;//This was being done every frame but should be done in init?
@@ -1380,7 +1387,7 @@ void FNiagaraEmitterInstance::Tick(float DeltaSeconds)
 
 	int32 EventSpawnTotal = (EventInstanceData.IsValid() ? EventInstanceData->EventSpawnTotal : 0);
 	int32 OrigNumParticles = GetNumParticles();
-	int32 AllocationEstimate = CachedEmitter->GetMaxParticleCountEstimate();
+	int32 AllocationEstimate = EmitterData->GetMaxParticleCountEstimate();
 	int32 RequiredSize = OrigNumParticles + SpawnTotal + EventSpawnTotal;
 
 	if (RequiredSize == 0)
@@ -1396,7 +1403,7 @@ void FNiagaraEmitterInstance::Tick(float DeltaSeconds)
 	{
 		ReallocationCount++;
 		MaxAllocationCount = AllocationSize;
-		int32 Estimations = CachedEmitter->AddRuntimeAllocation((uint64)this, MaxAllocationCount);
+		int32 Estimations = EmitterData->AddRuntimeAllocation((uint64)this, MaxAllocationCount);
 		if (GbNiagaraShowAllocationWarnings && Estimations >= 5 && ReallocationCount == 3)
 		{
 			FString SystemName = System->GetName();
@@ -1423,15 +1430,15 @@ void FNiagaraEmitterInstance::Tick(float DeltaSeconds)
 
 	/* GPU simulation -  we just create an FNiagaraComputeExecutionContext, queue it, and let the dispatch interface take care of the rest
 	 */
-	if (CachedEmitter->SimTarget == ENiagaraSimTarget::GPUComputeSim && GPUExecContext != nullptr)
+	if (EmitterData->SimTarget == ENiagaraSimTarget::GPUComputeSim && GPUExecContext != nullptr)
 	{
-		check(GPUExecContext->GPUScript_RT == CachedEmitter->GetGPUComputeScript()->GetRenderThreadScript());
-		GPUExecContext->GPUScript_RT = CachedEmitter->GetGPUComputeScript()->GetRenderThreadScript();
+		check(GPUExecContext->GPUScript_RT == EmitterData->GetGPUComputeScript()->GetRenderThreadScript());
+		GPUExecContext->GPUScript_RT = EmitterData->GetGPUComputeScript()->GetRenderThreadScript();
 
 #if WITH_EDITOR
 		if (ParentSystemInstance->ShouldCaptureThisFrame())
 		{
-			TSharedPtr<struct FNiagaraScriptDebuggerInfo, ESPMode::ThreadSafe> DebugInfo = ParentSystemInstance->GetActiveCaptureWrite(CachedIDName, ENiagaraScriptUsage::ParticleGPUComputeScript, FGuid());
+			TSharedPtr<FNiagaraScriptDebuggerInfo, ESPMode::ThreadSafe> DebugInfo = ParentSystemInstance->GetActiveCaptureWrite(CachedIDName, ENiagaraScriptUsage::ParticleGPUComputeScript, FGuid());
 			if (DebugInfo && ComputeDispatchInterface != nullptr)
 			{
 				//Data.Dump(DebugInfo->Frame, true, 0, OrigNumParticles);
@@ -1476,11 +1483,11 @@ void FNiagaraEmitterInstance::Tick(float DeltaSeconds)
 						// the destination dataset for simulation.
 						NumSpawnedOnGPUThisFrame += Info.Count;
 
-						int32 MaxParticlesSpawnedPerFrame = CachedEmitter->MaxGPUParticlesSpawnPerFrame <= 0 ? GMaxNiagaraGPUParticlesSpawnPerFrame : CachedEmitter->MaxGPUParticlesSpawnPerFrame;
+						int32 MaxParticlesSpawnedPerFrame = EmitterData->MaxGPUParticlesSpawnPerFrame <= 0 ? GMaxNiagaraGPUParticlesSpawnPerFrame : EmitterData->MaxGPUParticlesSpawnPerFrame;
 
 						if (NumSpawnedOnGPUThisFrame > MaxParticlesSpawnedPerFrame)
 						{
-							FString DebugMsg = FString::Printf(TEXT("%s has attempted to execeed max GPU per frame spawn! | Max: %d | Requested: %d | SpawnInfoEntry: %d"), *CachedEmitter->GetFullName(), MaxParticlesSpawnedPerFrame, NumSpawnedOnGPUThisFrame, SpawnInfoIdx);
+							FString DebugMsg = FString::Printf(TEXT("%s has attempted to execeed max GPU per frame spawn! | Max: %d | Requested: %d | SpawnInfoEntry: %d"), *CachedEmitter.Emitter->GetFullName(), MaxParticlesSpawnedPerFrame, NumSpawnedOnGPUThisFrame, SpawnInfoIdx);
 							UE_LOG(LogNiagara, Warning, TEXT("%s"), *DebugMsg);
 							GEngine->AddOnScreenDebugMessage(-1, 5.f, FColor::Yellow, *DebugMsg);
 							break;
@@ -1498,7 +1505,7 @@ void FNiagaraEmitterInstance::Tick(float DeltaSeconds)
 					}
 					else if (Info.Count > 0)
 					{
-						FString DebugMsg = FString::Printf(TEXT("%s Exceeded Gpu spawn info count, see NIAGARA_MAX_GPU_SPAWN_INFOS for more information!"), *CachedEmitter->GetUniqueEmitterName());
+						FString DebugMsg = FString::Printf(TEXT("%s Exceeded Gpu spawn info count, see NIAGARA_MAX_GPU_SPAWN_INFOS for more information!"), *CachedEmitter.Emitter->GetUniqueEmitterName());
 						UE_LOG(LogNiagara, Warning, TEXT("%s"), *DebugMsg);
 						GEngine->AddOnScreenDebugMessage(-1, 5.f, FColor::Yellow, *DebugMsg);
 						break;
@@ -1569,14 +1576,14 @@ void FNiagaraEmitterInstance::Tick(float DeltaSeconds)
 
 	//Ensure we don't blow our current hard limits on cpu particle count.
 	//TODO: These current limits can be improved relatively easily. Though perf in at these counts will obviously be an issue anyway.
-	if (CachedEmitter->SimTarget == ENiagaraSimTarget::CPUSim && AllocationSize > GMaxNiagaraCPUParticlesPerEmitter)
+	if (EmitterData->SimTarget == ENiagaraSimTarget::CPUSim && AllocationSize > GMaxNiagaraCPUParticlesPerEmitter)
 	{
-		FString DebugMsg = FString::Printf(TEXT("%s has attempted to exceed the max CPU particle count! | Max: %d | Requested: %u"), *CachedEmitter->GetFullName(), GMaxNiagaraCPUParticlesPerEmitter, AllocationSize);
+		FString DebugMsg = FString::Printf(TEXT("%s has attempted to exceed the max CPU particle count! | Max: %d | Requested: %u"), *CachedEmitter.Emitter->GetFullName(), GMaxNiagaraCPUParticlesPerEmitter, AllocationSize);
 		UE_LOG(LogNiagara, Warning, TEXT("%s"), *DebugMsg);
 		GEngine->AddOnScreenDebugMessage(-1, 5.f, FColor::Yellow, *DebugMsg);
 
 		//We clear the emitters estimate otherwise we get stuck in this state forever.
-		CachedEmitter->ClearRuntimeAllocationEstimate();
+		EmitterData->ClearRuntimeAllocationEstimate();
 
 		//For now we completely bail out of spawning new particles. Possibly should improve this in future.
 		AllocationSize = OrigNumParticles;
@@ -1757,7 +1764,6 @@ void FNiagaraEmitterInstance::Tick(float DeltaSeconds)
 			const int32 AdjustedSpawnCount = FMath::Min(Info.Count, SpawnCountRemaining);
 			if ( AdjustedSpawnCount > 0 )
 			{
-				auto& EmitterParameters = ParentSystemInstance->EditEmitterParameters(EmitterIdx);
 				SpawnIntervalBinding.SetValue(Info.IntervalDt);
 				InterpSpawnStartBinding.SetValue(Info.InterpStartDt);
 				SpawnGroupBinding.SetValue(Info.SpawnGroup);
@@ -1773,7 +1779,7 @@ void FNiagaraEmitterInstance::Tick(float DeltaSeconds)
 			{
 				int32 EventParticlesToSpawn = 0;
 
-				for (int32 EventScriptIdx = 0; EventScriptIdx < CachedEmitter->GetEventHandlers().Num(); EventScriptIdx++)
+				for (int32 EventScriptIdx = 0; EventScriptIdx < EmitterData->GetEventHandlers().Num(); EventScriptIdx++)
 				{
 					FNiagaraEventHandlingInfo& Info = EventInstanceData->EventHandlingInfo[EventScriptIdx];
 					for (int32 i = 0; i < Info.SpawnCounts.Num(); i++)
@@ -1788,8 +1794,6 @@ void FNiagaraEmitterInstance::Tick(float DeltaSeconds)
 
 				if (EventParticlesToSpawn > 0)
 				{
-					const int32 CurrNumParticles = Data.GetDestinationDataChecked().GetNumInstances();
-
 					SpawnIntervalBinding.SetValue(0.0f);
 					InterpSpawnStartBinding.SetValue(DeltaSeconds * 0.5f);
 					SpawnGroupBinding.SetValue(0);
@@ -1798,7 +1802,7 @@ void FNiagaraEmitterInstance::Tick(float DeltaSeconds)
 			}
 			else
 			{
-				for (int32 EventScriptIdx = 0; EventScriptIdx < CachedEmitter->GetEventHandlers().Num(); EventScriptIdx++)
+				for (int32 EventScriptIdx = 0; EventScriptIdx < EmitterData->GetEventHandlers().Num(); EventScriptIdx++)
 				{
 					FNiagaraEventHandlingInfo& Info = EventInstanceData->EventHandlingInfo[EventScriptIdx];
 
@@ -1860,7 +1864,7 @@ void FNiagaraEmitterInstance::Tick(float DeltaSeconds)
 	if (ParentSystemInstance->ShouldCaptureThisFrame())
 	{
 		//Pull out update data.
-		TSharedPtr<struct FNiagaraScriptDebuggerInfo, ESPMode::ThreadSafe> DebugInfo = ParentSystemInstance->GetActiveCaptureWrite(CachedIDName, ENiagaraScriptUsage::ParticleUpdateScript, FGuid());
+		TSharedPtr<FNiagaraScriptDebuggerInfo, ESPMode::ThreadSafe> DebugInfo = ParentSystemInstance->GetActiveCaptureWrite(CachedIDName, ENiagaraScriptUsage::ParticleUpdateScript, FGuid());
 		if (DebugInfo)
 		{
 			Data.CopyTo(DebugInfo->Frame, 0, OrigNumParticles);
@@ -1904,10 +1908,8 @@ void FNiagaraEmitterInstance::Tick(float DeltaSeconds)
 			//}
 		}
 
-		int32 SpawnEventScriptStartIndex = EventSpawnStart;
-		for (int32 EventScriptIdx = 0; EventScriptIdx < CachedEmitter->GetEventHandlers().Num(); EventScriptIdx++)
+		for (int32 EventScriptIdx = 0; EventScriptIdx < EmitterData->GetEventHandlers().Num(); EventScriptIdx++)
 		{
-			const FNiagaraEventScriptProperties &EventHandlerProps = CachedEmitter->GetEventHandlers()[EventScriptIdx];
 			FNiagaraEventHandlingInfo& Info = EventInstanceData->EventHandlingInfo[EventScriptIdx];
 
 			FScriptExecutionConstantBufferTable EventConstantBufferTable;
@@ -1943,7 +1945,7 @@ void FNiagaraEmitterInstance::Tick(float DeltaSeconds)
 						if (ParentSystemInstance->ShouldCaptureThisFrame())
 						{
 							FGuid EventGuid = EventInstanceData->EventExecContexts[EventScriptIdx].Script->GetUsageId();
-							TSharedPtr<struct FNiagaraScriptDebuggerInfo, ESPMode::ThreadSafe> DebugInfo = ParentSystemInstance->GetActiveCaptureWrite(CachedIDName, ENiagaraScriptUsage::ParticleEventScript, EventGuid);
+							TSharedPtr<FNiagaraScriptDebuggerInfo, ESPMode::ThreadSafe> DebugInfo = ParentSystemInstance->GetActiveCaptureWrite(CachedIDName, ENiagaraScriptUsage::ParticleEventScript, EventGuid);
 							if (DebugInfo)
 							{
 								Data.CopyTo(DebugInfo->Frame, EventSpawnStart, EventSpawnsStillAlive);
@@ -1966,9 +1968,9 @@ void FNiagaraEmitterInstance::Tick(float DeltaSeconds)
 		}
 
 		// Update events need a copy per event so that the previous event's data can be used.
-		for (int32 EventScriptIdx = 0; EventScriptIdx < CachedEmitter->GetEventHandlers().Num(); EventScriptIdx++)
+		for (int32 EventScriptIdx = 0; EventScriptIdx < EmitterData->GetEventHandlers().Num(); EventScriptIdx++)
 		{
-			const FNiagaraEventScriptProperties &EventHandlerProps = CachedEmitter->GetEventHandlers()[EventScriptIdx];
+			const FNiagaraEventScriptProperties &EventHandlerProps = EmitterData->GetEventHandlers()[EventScriptIdx];
 			FNiagaraDataBuffer* EventData = EventInstanceData->EventHandlingInfo[EventScriptIdx].EventData;
 
 			FScriptExecutionConstantBufferTable EventConstantBufferTable;
@@ -2009,7 +2011,7 @@ void FNiagaraEmitterInstance::Tick(float DeltaSeconds)
 						if (ParentSystemInstance->ShouldCaptureThisFrame())
 						{
 							FGuid EventGuid = EventInstanceData->EventExecContexts[EventScriptIdx].Script->GetUsageId();
-							TSharedPtr<struct FNiagaraScriptDebuggerInfo, ESPMode::ThreadSafe> DebugInfo = ParentSystemInstance->GetActiveCaptureWrite(CachedIDName, ENiagaraScriptUsage::ParticleEventScript, EventGuid);
+							TSharedPtr<FNiagaraScriptDebuggerInfo, ESPMode::ThreadSafe> DebugInfo = ParentSystemInstance->GetActiveCaptureWrite(CachedIDName, ENiagaraScriptUsage::ParticleEventScript, EventGuid);
 							if (DebugInfo)
 							{
 								Data.CopyTo(DebugInfo->Frame, 0, NumInstancesPrev);
@@ -2096,8 +2098,8 @@ void FNiagaraEmitterInstance::Tick(float DeltaSeconds)
 
 	if (GbDumpParticleData || System->bDumpDebugEmitterInfo)
 	{
-		UE_LOG(LogNiagara, Log, TEXT("|=== END OF FNiagaraEmitterInstance::Tick [ %s ] ===============|"), *CachedEmitter->GetPathName());
-		UE_LOG(LogNiagara, Log, TEXT("|||||||||||||||||||||||||||||||||||||||||||||||||||||||||||||||||||||||||||||||||||||||||||||||||||||||||||||"), *CachedEmitter->GetPathName());
+		UE_LOG(LogNiagara, Log, TEXT("|=== END OF FNiagaraEmitterInstance::Tick [ %s ] ===============|"), *CachedEmitter.Emitter->GetPathName());
+		UE_LOG(LogNiagara, Log, TEXT("|||||||||||||||||||||||||||||||||||||||||||||||||||||||||||||||||||||||||||||||||||||||||||||||||||||||||||||"), *CachedEmitter.Emitter->GetPathName());
 	}
 
 
@@ -2191,7 +2193,7 @@ void FNiagaraEmitterInstance::SetExecutionState(ENiagaraExecutionState InState)
 	if (ensureMsgf(InState >= ENiagaraExecutionState::Active && InState < ENiagaraExecutionState::Num, 
 					TEXT("Setting invalid emitter execution state! %d\nEmitter=%s\nSystem=%s\nComponent=%s"),
 					(int32)InState,
-					*GetFullNameSafe(CachedEmitter),
+					*GetFullNameSafe(CachedEmitter.Emitter),
 					*GetFullNameSafe(ParentSystemInstance ? ParentSystemInstance->GetSystem() : nullptr),
 					*GetFullNameSafe(ParentSystemInstance ? ParentSystemInstance->GetAttachComponent() : nullptr))
 		)
diff --git a/Engine/Plugins/FX/Niagara/Source/Niagara/Private/NiagaraGPUProfiler.cpp b/Engine/Plugins/FX/Niagara/Source/Niagara/Private/NiagaraGPUProfiler.cpp
--- a/Engine/Plugins/FX/Niagara/Source/Niagara/Private/NiagaraGPUProfiler.cpp
+++ b/Engine/Plugins/FX/Niagara/Source/Niagara/Private/NiagaraGPUProfiler.cpp
@@ -4,7 +4,6 @@
 #include "NiagaraGPUSystemTick.h"
 #include "NiagaraSimStageData.h"
 #include "HAL/IConsoleManager.h"
-#include "Misc/StringBuilder.h"
 
 bool GNiagaraGpuProfilingEnabled = true;
 static FAutoConsoleVariableRef CVarNiagaraGpuProfilingEnabled(
@@ -157,7 +156,7 @@ void FNiagaraGPUProfiler::BeginDispatch(FRHICommandList& RHICmdList, FName Stage
 	FGpuDispatchTimer& DispatchTimer = ActiveWriteFrame->DispatchTimers.AddDefaulted_GetRef();
 	DispatchTimer.bUniqueInstance = false;
 	DispatchTimer.OwnerComponent = nullptr;
-	DispatchTimer.OwnerEmitter = nullptr;
+	DispatchTimer.OwnerEmitter = FVersionedNiagaraEmitterWeakPtr();
 	DispatchTimer.StageName = StageName;
 	DispatchTimer.StartQuery = QueryPool->AllocateQuery();
 	RHICmdList.EndRenderQuery(DispatchTimer.StartQuery.GetQuery());
diff --git a/Engine/Plugins/FX/Niagara/Source/Niagara/Private/NiagaraGPUProfiler.h b/Engine/Plugins/FX/Niagara/Source/Niagara/Private/NiagaraGPUProfiler.h
--- a/Engine/Plugins/FX/Niagara/Source/Niagara/Private/NiagaraGPUProfiler.h
+++ b/Engine/Plugins/FX/Niagara/Source/Niagara/Private/NiagaraGPUProfiler.h
@@ -3,6 +3,7 @@
 #pragma once
 
 #include "NiagaraCommon.h"
+#include "NiagaraEmitter.h"
 #include "RHI.h"
 #include "RHICommandList.h"
 #include "RHIResources.h"
@@ -31,7 +32,7 @@ class FNiagaraGPUProfiler : public FNiagaraGPUProfilerInterface
 	{
 		uint32									bUniqueInstance : 1;
 		TWeakObjectPtr<class USceneComponent>	OwnerComponent;
-		TWeakObjectPtr<class UNiagaraEmitter>	OwnerEmitter;
+		FVersionedNiagaraEmitterWeakPtr			OwnerEmitter;
 		FName									StageName;
 		FRHIPooledRenderQuery					StartQuery;
 		FRHIPooledRenderQuery					EndQuery;
diff --git a/Engine/Plugins/FX/Niagara/Source/Niagara/Private/NiagaraGPUProfilerInterface.cpp b/Engine/Plugins/FX/Niagara/Source/Niagara/Private/NiagaraGPUProfilerInterface.cpp
--- a/Engine/Plugins/FX/Niagara/Source/Niagara/Private/NiagaraGPUProfilerInterface.cpp
+++ b/Engine/Plugins/FX/Niagara/Source/Niagara/Private/NiagaraGPUProfilerInterface.cpp
@@ -6,6 +6,7 @@
 #include "NiagaraGPUSystemTick.h"
 #include "NiagaraDataInterfaceBase.h"
 #include "NiagaraComponent.h"
+#include "NiagaraEmitter.h"
 
 #include "Async/Async.h"
 
@@ -144,13 +145,13 @@ void SetResultsForEditorStats(const FNiagaraGpuFrameResultsPtr& FrameResults)
 #if WITH_NIAGARA_GPU_PROFILER_EDITOR
 	// Send data for editor stat display
 	//-TODO: editor stats should merge with particle perf stats
-	TMap<UNiagaraEmitter*, TMap<TStatIdData const*, float>> CapturedStats;
+	TMap<FVersionedNiagaraEmitterData*, TMap<TStatIdData const*, float>> CapturedStats;
 	CapturedStats.Reserve(FrameResults->DispatchResults.Num());
 
 	for (const auto& DispatchResult : FrameResults->DispatchResults)
 	{
-		UNiagaraEmitter* Emitter = DispatchResult.OwnerEmitter.Get();
-		if ( Emitter == nullptr )
+		FVersionedNiagaraEmitterData* EmitterData = DispatchResult.OwnerEmitter.GetEmitterData();
+		if ( EmitterData == nullptr )
 		{
 			continue;
 		}
@@ -162,7 +163,7 @@ void SetResultsForEditorStats(const FNiagaraGpuFrameResultsPtr& FrameResults)
 		const FName StatFName(StatNameBuilder.ToString());
 
 		const TStatId StatId = FDynamicStats::CreateStatId<FStatGroup_STATGROUP_NiagaraDetailed>(StatFName);
-		CapturedStats.FindOrAdd(Emitter).FindOrAdd(StatId.GetRawPointer()) += DispatchResult.DurationMicroseconds;
+		CapturedStats.FindOrAdd(EmitterData).FindOrAdd(StatId.GetRawPointer()) += DispatchResult.DurationMicroseconds;
 	}
 
 	for (auto& CapturedStat : CapturedStats)
diff --git a/Engine/Plugins/FX/Niagara/Source/Niagara/Private/NiagaraGPUSystemTick.cpp b/Engine/Plugins/FX/Niagara/Source/Niagara/Private/NiagaraGPUSystemTick.cpp
--- a/Engine/Plugins/FX/Niagara/Source/Niagara/Private/NiagaraGPUSystemTick.cpp
+++ b/Engine/Plugins/FX/Niagara/Source/Niagara/Private/NiagaraGPUSystemTick.cpp
@@ -115,12 +115,12 @@ void FNiagaraGPUSystemTick::Init(FNiagaraSystemInstance* InSystemInstance)
 				continue;
 			}
 
-			const UNiagaraEmitter* Emitter = EmitterInstance->GetCachedEmitter();
+			const FVersionedNiagaraEmitterData* EmitterData = EmitterInstance->GetCachedEmitterData();
 			FNiagaraComputeExecutionContext* GPUContext = EmitterInstance->GetGPUContext();
 
-			check(Emitter);
+			check(EmitterData);
 
-			if (!Emitter || !GPUContext || Emitter->SimTarget != ENiagaraSimTarget::GPUComputeSim)
+			if (!EmitterData || !GPUContext || EmitterData->SimTarget != ENiagaraSimTarget::GPUComputeSim)
 			{
 				continue;
 			}
diff --git a/Engine/Plugins/FX/Niagara/Source/Niagara/Private/NiagaraGpuComputeDispatch.cpp b/Engine/Plugins/FX/Niagara/Source/Niagara/Private/NiagaraGpuComputeDispatch.cpp
--- a/Engine/Plugins/FX/Niagara/Source/Niagara/Private/NiagaraGpuComputeDispatch.cpp
+++ b/Engine/Plugins/FX/Niagara/Source/Niagara/Private/NiagaraGpuComputeDispatch.cpp
@@ -18,7 +18,6 @@
 #include "NiagaraGPUProfilerInterface.h"
 #include "NiagaraGpuReadbackManager.h"
 #include "NiagaraRenderer.h"
-#include "NiagaraScriptExecutionContext.h"
 #include "NiagaraSystemGpuComputeProxy.h"
 #include "NiagaraShader.h"
 #include "NiagaraShaderParticleID.h"
diff --git a/Engine/Plugins/FX/Niagara/Source/Niagara/Private/NiagaraLightRendererProperties.cpp b/Engine/Plugins/FX/Niagara/Source/Niagara/Private/NiagaraLightRendererProperties.cpp
--- a/Engine/Plugins/FX/Niagara/Source/Niagara/Private/NiagaraLightRendererProperties.cpp
+++ b/Engine/Plugins/FX/Niagara/Source/Niagara/Private/NiagaraLightRendererProperties.cpp
@@ -158,7 +158,7 @@ void UNiagaraLightRendererProperties::GetRendererTooltipWidgets(const FNiagaraEm
 	OutWidgets.Add(LightTooltip);
 }
 
-void UNiagaraLightRendererProperties::GetRendererFeedback(const UNiagaraEmitter* InEmitter, TArray<FText>& OutErrors, TArray<FText>& OutWarnings, TArray<FText>& OutInfo) const
+void UNiagaraLightRendererProperties::GetRendererFeedback(const FVersionedNiagaraEmitter& InEmitter, TArray<FText>& OutErrors, TArray<FText>& OutWarnings, TArray<FText>& OutInfo) const
 {
 	Super::GetRendererFeedback(InEmitter, OutErrors, OutWarnings, OutInfo);
 }
diff --git a/Engine/Plugins/FX/Niagara/Source/Niagara/Private/NiagaraMeshRendererProperties.cpp b/Engine/Plugins/FX/Niagara/Source/Niagara/Private/NiagaraMeshRendererProperties.cpp
--- a/Engine/Plugins/FX/Niagara/Source/Niagara/Private/NiagaraMeshRendererProperties.cpp
+++ b/Engine/Plugins/FX/Niagara/Source/Niagara/Private/NiagaraMeshRendererProperties.cpp
@@ -8,6 +8,7 @@
 #include "NiagaraCustomVersion.h"
 #include "NiagaraEmitterInstance.h"
 #include "Modules/ModuleManager.h"
+#include "NiagaraGPUSortInfo.h"
 
 #if WITH_EDITOR
 #include "Editor.h"
@@ -161,9 +162,9 @@ FNiagaraBoundsCalculator* UNiagaraMeshRendererProperties::CreateBoundsCalculator
 	FVector MaxWorldMeshOffset(ForceInitToZero);
 
 	bool bLocalSpace = false;
-	if (UNiagaraEmitter* Emitter = Cast<UNiagaraEmitter>(GetOuter()))
+	if (FVersionedNiagaraEmitterData* EmitterData = GetEmitterData())
 	{
-		bLocalSpace = Emitter->bLocalSpace;
+		bLocalSpace = EmitterData->bLocalSpace;
 	}
 
 	for (const auto& MeshProperties : Meshes)
@@ -281,10 +282,10 @@ void UNiagaraMeshRendererProperties::InitBindings()
 		CustomSortingBinding = FNiagaraConstants::GetAttributeDefaultBinding(SYS_PARAM_PARTICLES_NORMALIZED_AGE);
 	}
 
-	SetPreviousBindings(nullptr, SourceMode);
+	SetPreviousBindings(FVersionedNiagaraEmitter(), SourceMode);
 }
 
-void UNiagaraMeshRendererProperties::SetPreviousBindings(const UNiagaraEmitter* SrcEmitter, ENiagaraRendererSourceDataMode InSourceMode)
+void UNiagaraMeshRendererProperties::SetPreviousBindings(const FVersionedNiagaraEmitter& SrcEmitter, ENiagaraRendererSourceDataMode InSourceMode)
 {
 	PrevPositionBinding.SetAsPreviousValue(PositionBinding, SrcEmitter, InSourceMode);
 	PrevScaleBinding.SetAsPreviousValue(ScaleBinding, SrcEmitter, InSourceMode);
@@ -297,12 +298,12 @@ void UNiagaraMeshRendererProperties::UpdateSourceModeDerivates(ENiagaraRendererS
 {
 	Super::UpdateSourceModeDerivates(InSourceMode, bFromPropertyEdit);
 
-	UNiagaraEmitter* SrcEmitter = GetTypedOuter<UNiagaraEmitter>();
-	if (SrcEmitter)
+	FVersionedNiagaraEmitter SrcEmitter = GetOuterEmitter();
+	if (SrcEmitter.Emitter)
 	{
 		for (FNiagaraMaterialAttributeBinding& MaterialParamBinding : MaterialParameterBindings)
 		{
-			MaterialParamBinding.CacheValues(SrcEmitter);
+			MaterialParamBinding.CacheValues(SrcEmitter.Emitter);
 		}
 
 		SetPreviousBindings(SrcEmitter, InSourceMode);
@@ -555,7 +556,7 @@ void UNiagaraMeshRendererProperties::PostLoad()
 	PostLoadBindings(SourceMode);
 	
 	// Fix up these bindings from their loaded source bindings
-	SetPreviousBindings(nullptr, SourceMode);
+	SetPreviousBindings(FVersionedNiagaraEmitter(), SourceMode);
 
 	for ( const FNiagaraMeshMaterialOverride& OverrideMaterial : OverrideMaterials )
 	{
@@ -667,7 +668,7 @@ void UNiagaraMeshRendererProperties::GetRendererTooltipWidgets(const FNiagaraEmi
 }
 
 
-void UNiagaraMeshRendererProperties::GetRendererFeedback(const UNiagaraEmitter* InEmitter, TArray<FText>& OutErrors, TArray<FText>& OutWarnings, TArray<FText>& OutInfo) const
+void UNiagaraMeshRendererProperties::GetRendererFeedback(const FVersionedNiagaraEmitter& InEmitter, TArray<FText>& OutErrors, TArray<FText>& OutWarnings, TArray<FText>& OutInfo) const
 {
 	Super::GetRendererFeedback(InEmitter, OutErrors, OutWarnings, OutInfo);
 }
@@ -801,28 +802,28 @@ void UNiagaraMeshRendererProperties::PostEditChangeProperty(FPropertyChangedEven
 	Super::PostEditChangeProperty(PropertyChangedEvent);
 }
 
-void UNiagaraMeshRendererProperties::RenameVariable(const FNiagaraVariableBase& OldVariable, const FNiagaraVariableBase& NewVariable, const UNiagaraEmitter* InEmitter)
+void UNiagaraMeshRendererProperties::RenameVariable(const FNiagaraVariableBase& OldVariable, const FNiagaraVariableBase& NewVariable, const FVersionedNiagaraEmitter& InEmitter)
 {
 	Super::RenameVariable(OldVariable, NewVariable, InEmitter);
 
 	// Handle renaming material bindings
 	for (FNiagaraMaterialAttributeBinding& Binding : MaterialParameterBindings)
 	{
-		Binding.RenameVariableIfMatching(OldVariable, NewVariable, InEmitter, GetCurrentSourceMode());
+		Binding.RenameVariableIfMatching(OldVariable, NewVariable, InEmitter.Emitter, GetCurrentSourceMode());
 	}
 }
 
-void UNiagaraMeshRendererProperties::RemoveVariable(const FNiagaraVariableBase& OldVariable, const UNiagaraEmitter* InEmitter)
+void UNiagaraMeshRendererProperties::RemoveVariable(const FNiagaraVariableBase& OldVariable, const FVersionedNiagaraEmitter& InEmitter)
 {
 	Super::RemoveVariable(OldVariable, InEmitter);
 
 	// Handle resetting material bindings to defaults
 	for (FNiagaraMaterialAttributeBinding& Binding : MaterialParameterBindings)
 	{
-		if (Binding.Matches(OldVariable, InEmitter, GetCurrentSourceMode()))
+		if (Binding.Matches(OldVariable, InEmitter.Emitter, GetCurrentSourceMode()))
 		{
 			Binding.NiagaraVariable = FNiagaraVariable();
-			Binding.CacheValues(InEmitter);
+			Binding.CacheValues(InEmitter.Emitter);
 		}
 	}
 }
@@ -831,10 +832,10 @@ void UNiagaraMeshRendererProperties::OnMeshChanged()
 {
 	FNiagaraSystemUpdateContext ReregisterContext;
 
-	UNiagaraEmitter* Emitter = Cast<UNiagaraEmitter>(GetOuter());
-	if (Emitter != nullptr)
+	FVersionedNiagaraEmitter Outer = GetOuterEmitter();
+	if (Outer.Emitter)
 	{
-		ReregisterContext.Add(Emitter, true);
+		ReregisterContext.Add(Outer, true);
 	}
 
 	CheckMaterialUsage();
diff --git a/Engine/Plugins/FX/Niagara/Source/Niagara/Private/NiagaraPlatformSet.cpp b/Engine/Plugins/FX/Niagara/Source/Niagara/Private/NiagaraPlatformSet.cpp
--- a/Engine/Plugins/FX/Niagara/Source/Niagara/Private/NiagaraPlatformSet.cpp
+++ b/Engine/Plugins/FX/Niagara/Source/Niagara/Private/NiagaraPlatformSet.cpp
@@ -9,6 +9,7 @@
 #include "NiagaraSystem.h"
 #include "NiagaraSettings.h"
 #include "SystemSettings.h"
+#include "UObject/UObjectIterator.h"
 
 #if WITH_EDITOR
 #include "PlatformInfo.h"
@@ -958,7 +959,7 @@ void FNiagaraPlatformSet::RefreshScalability()
 			check(System);
 
 			bool bNeedsUpdate = false;
-			auto NeedsScalabilityUpdate = [&](UObject* Owner, FNiagaraPlatformSet& PSet)
+			auto NeedsScalabilityUpdate = [&](FNiagaraPlatformSet& PSet)
 			{
 				for (FNiagaraPlatformSetCVarCondition& CVarCondition : PSet.CVarConditions)
 				{
diff --git a/Engine/Plugins/FX/Niagara/Source/Niagara/Private/NiagaraRenderer.cpp b/Engine/Plugins/FX/Niagara/Source/Niagara/Private/NiagaraRenderer.cpp
--- a/Engine/Plugins/FX/Niagara/Source/Niagara/Private/NiagaraRenderer.cpp
+++ b/Engine/Plugins/FX/Niagara/Source/Niagara/Private/NiagaraRenderer.cpp
@@ -2,12 +2,9 @@
 
 #include "NiagaraRenderer.h"
 #include "ParticleResources.h"
-#include "ParticleBeamTrailVertexFactory.h"
 #include "NiagaraDataSet.h"
 #include "NiagaraStats.h"
-#include "NiagaraVertexFactory.h"
 #include "NiagaraComponent.h"
-#include "Engine/Engine.h"
 #include "DynamicBufferAllocator.h"
 #include "NiagaraComputeExecutionContext.h"
 #include "NiagaraGPUSortInfo.h"
@@ -369,14 +366,14 @@ void FNiagaraDynamicDataBase::SetVertexFactoryData(class FNiagaraVertexFactoryBa
 
 FNiagaraRenderer::FNiagaraRenderer(ERHIFeatureLevel::Type InFeatureLevel, const UNiagaraRendererProperties *InProps, const FNiagaraEmitterInstance* Emitter)
 	: DynamicDataRender(nullptr)
-	, bLocalSpace(Emitter->GetCachedEmitter()->bLocalSpace)
+	, bLocalSpace(Emitter->GetCachedEmitterData()->bLocalSpace)
 	, bHasLights(false)
 	, bMotionBlurEnabled(InProps ? InProps->MotionVectorSetting != ENiagaraRendererMotionVectorSetting::Disable : false)
-	, SimTarget(Emitter->GetCachedEmitter()->SimTarget)
+	, SimTarget(Emitter->GetCachedEmitterData()->SimTarget)
 	, FeatureLevel(InFeatureLevel)
 {
 #if STATS
-	EmitterStatID = Emitter->GetCachedEmitter()->GetStatID(false, false);
+	EmitterStatID = Emitter->GetCachedEmitter().Emitter->GetStatID(false, false);
 #endif
 }
 
diff --git a/Engine/Plugins/FX/Niagara/Source/Niagara/Private/NiagaraRendererComponents.cpp b/Engine/Plugins/FX/Niagara/Source/Niagara/Private/NiagaraRendererComponents.cpp
--- a/Engine/Plugins/FX/Niagara/Source/Niagara/Private/NiagaraRendererComponents.cpp
+++ b/Engine/Plugins/FX/Niagara/Source/Niagara/Private/NiagaraRendererComponents.cpp
@@ -521,7 +521,7 @@ void FNiagaraRendererComponents::PostSystemTick_GameThread(const UNiagaraRendere
 		}
 
 		FComponentPoolEntry& PoolEntry = ComponentPool[PoolIndex];
-		FNiagaraLWCConverter LwcConverter = SystemInstance->GetLWCConverter(Emitter->GetCachedEmitter()->bLocalSpace);
+		FNiagaraLWCConverter LwcConverter = SystemInstance->GetLWCConverter(Emitter->GetCachedEmitter().GetEmitterData()->bLocalSpace);
 		TickPropertyBindings(Properties, SceneComponent, Data, ParticleIndex, PoolEntry, LwcConverter);
 
 		// Activate the component.
diff --git a/Engine/Plugins/FX/Niagara/Source/Niagara/Private/NiagaraRendererProperties.cpp b/Engine/Plugins/FX/Niagara/Source/Niagara/Private/NiagaraRendererProperties.cpp
--- a/Engine/Plugins/FX/Niagara/Source/Niagara/Private/NiagaraRendererProperties.cpp
+++ b/Engine/Plugins/FX/Niagara/Source/Niagara/Private/NiagaraRendererProperties.cpp
@@ -161,7 +161,7 @@ FNiagaraVariable UNiagaraRendererProperties::GetBoundAttribute(const FNiagaraVar
 	return FNiagaraVariable();
 }
 
-void UNiagaraRendererProperties::GetRendererFeedback(UNiagaraEmitter* InEmitter,	TArray<FNiagaraRendererFeedback>& OutErrors, TArray<FNiagaraRendererFeedback>& OutWarnings,	TArray<FNiagaraRendererFeedback>& OutInfo) const
+void UNiagaraRendererProperties::GetRendererFeedback(const FVersionedNiagaraEmitter& InEmitter, TArray<FNiagaraRendererFeedback>& OutErrors, TArray<FNiagaraRendererFeedback>& OutWarnings,	TArray<FNiagaraRendererFeedback>& OutInfo) const
 {
 	TArray<FText> Errors;
 	TArray<FText> Warnings;
@@ -191,7 +191,7 @@ FText UNiagaraRendererProperties::GetWidgetDisplayName() const
 	return GetClass()->GetDisplayNameText();
 }
 
-void UNiagaraRendererProperties::RenameVariable(const FNiagaraVariableBase& OldVariable, const FNiagaraVariableBase& NewVariable, const UNiagaraEmitter* InEmitter)
+void UNiagaraRendererProperties::RenameVariable(const FNiagaraVariableBase& OldVariable, const FNiagaraVariableBase& NewVariable, const FVersionedNiagaraEmitter& InEmitter)
 {
 	// Handle the renaming of generic renderer bindings...
 	for (const FNiagaraVariableAttributeBinding* AttributeBinding : AttributeBindings)
@@ -201,7 +201,7 @@ void UNiagaraRendererProperties::RenameVariable(const FNiagaraVariableBase& OldV
 			Binding->RenameVariableIfMatching(OldVariable, NewVariable, InEmitter, GetCurrentSourceMode());
 	}
 }
-void UNiagaraRendererProperties::RemoveVariable(const FNiagaraVariableBase& OldVariable,const UNiagaraEmitter* InEmitter)
+void UNiagaraRendererProperties::RemoveVariable(const FNiagaraVariableBase& OldVariable,const FVersionedNiagaraEmitter& InEmitter)
 {
 	// Handle the reset to defaults of generic renderer bindings
 	for (const FNiagaraVariableAttributeBinding* AttributeBinding : AttributeBindings)
@@ -287,7 +287,7 @@ uint32 UNiagaraRendererProperties::ComputeMaxUsedComponents(const FNiagaraDataSe
 	return MaxNumComponents;
 }
 
-void UNiagaraRendererProperties::GetAssetTagsForContext(const UObject* InAsset, const TArray<const UNiagaraRendererProperties*>& InProperties, TMap<FName, uint32>& NumericKeys, TMap<FName, FString>& StringKeys) const
+void UNiagaraRendererProperties::GetAssetTagsForContext(const UObject* InAsset, FGuid AssetVersion, const TArray<const UNiagaraRendererProperties*>& InProperties, TMap<FName, uint32>& NumericKeys, TMap<FName, FString>& StringKeys) const
 {
 	UClass* Class = GetClass();
 
@@ -310,7 +310,7 @@ void UNiagaraRendererProperties::GetAssetTagsForContext(const UObject* InAsset,
 			FString Key = Class->GetName();
 			Key.ReplaceInline(TEXT("Niagara"), TEXT(""));
 			Key.ReplaceInline(TEXT("Properties"), TEXT(""));
-			NumericKeys.Add(*Key) = NumInstances;
+			NumericKeys.Add(FName(Key)) = NumInstances;
 		}
 	}
 }
@@ -378,26 +378,26 @@ void UNiagaraRendererProperties::PostEditChangeProperty(struct FPropertyChangedE
 {
 	Super::PostEditChangeProperty(PropertyChangedEvent);
 
-	if (UNiagaraEmitter* Emitter = GetTypedOuter<UNiagaraEmitter>())
+	if (FVersionedNiagaraEmitterData* EmitterData = GetEmitterData())
 	{
 		// Check for properties changing that invalidate the current script compilation for the emitter
 		bool bNeedsRecompile = false;
 		if (PropertyChangedEvent.GetPropertyName() == GET_MEMBER_NAME_CHECKED(UNiagaraRendererProperties, MotionVectorSetting))
 		{
-			if (Emitter->GraphSource)
+			if (EmitterData->GraphSource)
 			{
-				Emitter->GraphSource->MarkNotSynchronized(TEXT("Renderer MotionVectorSetting changed"));
+				EmitterData->GraphSource->MarkNotSynchronized(TEXT("Renderer MotionVectorSetting changed"));
 			}
 			bNeedsRecompile = true;
 		}
 
 		if (bNeedsRecompile)
 		{
-			UNiagaraSystem::RequestCompileForEmitter(Emitter);
+			UNiagaraSystem::RequestCompileForEmitter(GetOuterEmitter());
 		}
 
 		// Just in case we changed something that needs static params, refresh that cached list.
-		Emitter->RebuildRendererBindings();
+		EmitterData->RebuildRendererBindings(*GetOuterEmitter().Emitter);
 	}
 
 }
@@ -411,7 +411,8 @@ void UNiagaraRendererProperties::SetIsEnabled(bool bInIsEnabled)
 #if WITH_EDITORONLY_DATA
 		// Changing the enabled state will add or remove its renderer binding data stored on the emitters RenderBindings
 		// parameter store, so we need to reset to clear any binding references or add new ones
-		if (UNiagaraEmitter* SrcEmitter = GetTypedOuter<UNiagaraEmitter>())
+		FVersionedNiagaraEmitter SrcEmitter = GetOuterEmitter();
+		if (SrcEmitter.Emitter)
 		{
 			FNiagaraSystemUpdateContext(SrcEmitter, true);
 		}
@@ -423,8 +424,8 @@ void UNiagaraRendererProperties::SetIsEnabled(bool bInIsEnabled)
 
 void UNiagaraRendererProperties::UpdateSourceModeDerivates(ENiagaraRendererSourceDataMode InSourceMode, bool bFromPropertyEdit)
 {
-	UNiagaraEmitter* SrcEmitter = GetTypedOuter<UNiagaraEmitter>();
-	if (SrcEmitter)
+	FVersionedNiagaraEmitter SrcEmitter = GetOuterEmitter();
+	if (SrcEmitter.Emitter)
 	{
 		for (const FNiagaraVariableAttributeBinding* Binding : AttributeBindings)
 		{
@@ -436,18 +437,33 @@ void UNiagaraRendererProperties::UpdateSourceModeDerivates(ENiagaraRendererSourc
 		// to ensure new ones get bound by the simulation
 		if (bFromPropertyEdit)
 		{
-			if (SrcEmitter)
-			{
-				// We may need to refresh internal variables because this may be the first binding to it, so request a recompile as that will pull data 
-				// into the right place.
-				UNiagaraSystem::RequestCompileForEmitter(SrcEmitter);
-			}
+			// We may need to refresh internal variables because this may be the first binding to it, so request a recompile as that will pull data 
+			// into the right place.
+			UNiagaraSystem::RequestCompileForEmitter(GetOuterEmitter());
 			FNiagaraSystemUpdateContext Context(SrcEmitter, true);
 		}
 #endif
 	}
 }
 
+FVersionedNiagaraEmitterData* UNiagaraRendererProperties::GetEmitterData() const
+{
+	if (UNiagaraEmitter* SrcEmitter = GetTypedOuter<UNiagaraEmitter>())
+	{
+		return SrcEmitter->GetEmitterData(OuterEmitterVersion);
+	}
+	return nullptr;
+}
+
+FVersionedNiagaraEmitter UNiagaraRendererProperties::GetOuterEmitter() const
+{
+	if (UNiagaraEmitter* SrcEmitter = GetTypedOuter<UNiagaraEmitter>())
+	{
+		return FVersionedNiagaraEmitter(SrcEmitter, OuterEmitterVersion);
+	}
+	return FVersionedNiagaraEmitter();
+}
+
 bool UNiagaraRendererProperties::NeedsPreciseMotionVectors() const
 {
 	if (MotionVectorSetting == ENiagaraRendererMotionVectorSetting::AutoDetect)
diff --git a/Engine/Plugins/FX/Niagara/Source/Niagara/Private/NiagaraRibbonRendererProperties.cpp b/Engine/Plugins/FX/Niagara/Source/Niagara/Private/NiagaraRibbonRendererProperties.cpp
--- a/Engine/Plugins/FX/Niagara/Source/Niagara/Private/NiagaraRibbonRendererProperties.cpp
+++ b/Engine/Plugins/FX/Niagara/Source/Niagara/Private/NiagaraRibbonRendererProperties.cpp
@@ -256,10 +256,10 @@ void UNiagaraRibbonRendererProperties::InitBindings()
 		V1RangeOverrideBinding = FNiagaraConstants::GetAttributeDefaultBinding(SYS_PARAM_PARTICLES_RIBBONV1RANGEOVERRIDE);
 	}
 
-	SetPreviousBindings(nullptr);
+	SetPreviousBindings(FVersionedNiagaraEmitter());
 }
 
-void UNiagaraRibbonRendererProperties::SetPreviousBindings(const UNiagaraEmitter* SrcEmitter)
+void UNiagaraRibbonRendererProperties::SetPreviousBindings(const FVersionedNiagaraEmitter& SrcEmitter)
 {
 	PrevPositionBinding.SetAsPreviousValue(PositionBinding, SrcEmitter, ENiagaraRendererSourceDataMode::Particles);
 	PrevRibbonWidthBinding.SetAsPreviousValue(RibbonWidthBinding, SrcEmitter, ENiagaraRendererSourceDataMode::Particles);
@@ -467,7 +467,7 @@ void UNiagaraRibbonRendererProperties::GetRendererTooltipWidgets(const FNiagaraE
 }
 
 
-void UNiagaraRibbonRendererProperties::GetRendererFeedback(const UNiagaraEmitter* InEmitter, TArray<FText>& OutErrors, TArray<FText>& OutWarnings, TArray<FText>& OutInfo) const
+void UNiagaraRibbonRendererProperties::GetRendererFeedback(const FVersionedNiagaraEmitter& InEmitter, TArray<FText>& OutErrors, TArray<FText>& OutWarnings, TArray<FText>& OutInfo) const
 {
 	Super::GetRendererFeedback(InEmitter, OutErrors, OutWarnings, OutInfo);
 }
@@ -496,28 +496,28 @@ bool UNiagaraRibbonRendererProperties::CanEditChange(const FProperty* InProperty
 	return Super::CanEditChange(InProperty);
 }
 
-void UNiagaraRibbonRendererProperties::RenameVariable(const FNiagaraVariableBase& OldVariable, const FNiagaraVariableBase& NewVariable, const UNiagaraEmitter* InEmitter)
+void UNiagaraRibbonRendererProperties::RenameVariable(const FNiagaraVariableBase& OldVariable, const FNiagaraVariableBase& NewVariable, const FVersionedNiagaraEmitter& InEmitter)
 {
 	Super::RenameVariable(OldVariable, NewVariable, InEmitter);
 
 	// Handle renaming material bindings
 	for (FNiagaraMaterialAttributeBinding& Binding : MaterialParameterBindings)
 	{
-		Binding.RenameVariableIfMatching(OldVariable, NewVariable, InEmitter, GetCurrentSourceMode());
+		Binding.RenameVariableIfMatching(OldVariable, NewVariable, InEmitter.Emitter, GetCurrentSourceMode());
 	}
 }
 
-void UNiagaraRibbonRendererProperties::RemoveVariable(const FNiagaraVariableBase& OldVariable, const UNiagaraEmitter* InEmitter)
+void UNiagaraRibbonRendererProperties::RemoveVariable(const FNiagaraVariableBase& OldVariable, const FVersionedNiagaraEmitter& InEmitter)
 {
 	Super::RemoveVariable(OldVariable, InEmitter);
 
 	// Handle resetting material bindings to defaults
 	for (FNiagaraMaterialAttributeBinding& Binding : MaterialParameterBindings)
 	{
-		if (Binding.Matches(OldVariable, InEmitter, GetCurrentSourceMode()))
+		if (Binding.Matches(OldVariable, InEmitter.Emitter, GetCurrentSourceMode()))
 		{
 			Binding.NiagaraVariable = FNiagaraVariable();
-			Binding.CacheValues(InEmitter);
+			Binding.CacheValues(InEmitter.Emitter);
 		}
 	}
 }
diff --git a/Engine/Plugins/FX/Niagara/Source/Niagara/Private/NiagaraScratchPadContainer.cpp b/Engine/Plugins/FX/Niagara/Source/Niagara/Private/NiagaraScratchPadContainer.cpp
new file mode 100644
--- /dev/null
+++ b/Engine/Plugins/FX/Niagara/Source/Niagara/Private/NiagaraScratchPadContainer.cpp
@@ -0,0 +1,32 @@
+﻿// Copyright Epic Games, Inc. All Rights Reserved.
+
+
+#include "NiagaraScratchPadContainer.h"
+
+void UNiagaraScratchPadContainer::PostLoad()
+{
+	CheckConsistency();
+
+	Super::PostLoad();
+}
+
+void UNiagaraScratchPadContainer::CheckConsistency()
+{
+#if WITH_EDITORONLY_DATA
+	for (UNiagaraScript* Script : Scripts)
+	{
+		if (Script && Script->GetOuter() != this)
+		{
+			Script->Rename(*Script->GetName(), this, REN_ForceNoResetLoaders | REN_NonTransactional);
+		}
+	}
+#endif
+}
+
+void UNiagaraScratchPadContainer::SetScripts(const TArray<TObjectPtr<UNiagaraScript>>& InScripts)
+{
+#if WITH_EDITORONLY_DATA
+	Scripts = InScripts;
+	CheckConsistency();
+#endif
+}
