Added versioning for Niagara emitters

diff --git a/Engine/Plugins/FX/Niagara/Source/Niagara/Private/NiagaraScript.cpp b/Engine/Plugins/FX/Niagara/Source/Niagara/Private/NiagaraScript.cpp
--- a/Engine/Plugins/FX/Niagara/Source/Niagara/Private/NiagaraScript.cpp
+++ b/Engine/Plugins/FX/Niagara/Source/Niagara/Private/NiagaraScript.cpp
@@ -95,7 +95,7 @@ static FAutoConsoleVariableRef CVarLogCompileStaticVars(
 	ECVF_Default
 );
 
-FNiagaraScriptDebuggerInfo::FNiagaraScriptDebuggerInfo() : bWaitForGPU(false), FrameLastWriteId(-1), bWritten(false)
+FNiagaraScriptDebuggerInfo::FNiagaraScriptDebuggerInfo() : bWaitForGPU(false), Usage(ENiagaraScriptUsage::Function), FrameLastWriteId(-1), bWritten(false)
 {
 }
 
@@ -409,6 +409,7 @@ FGuid UNiagaraScript::AddNewVersion(int32 MajorVersion, int32 MinorVersion)
 	// check preconditions
 	check(MajorVersion >= 1);
 	check(MajorVersion != 1 || MinorVersion != 0);
+	Modify();
 
 	FVersionedNiagaraScriptData NewVersionData;
 	for (int i = VersionData.Num() - 1; i >= 0; i--)
@@ -447,6 +448,7 @@ void UNiagaraScript::DeleteVersion(const FGuid& VersionGuid)
 		if (AssetVersion.VersionGuid == VersionGuid)
 		{
 			check(AssetVersion.MajorVersion != 1 || AssetVersion.MinorVersion != 0);
+			Modify();
 			VersionData.RemoveAt(i);
 			return;
 		}
@@ -460,6 +462,7 @@ void UNiagaraScript::ExposeVersion(const FGuid& VersionGuid)
 	{
 		if (Data.Version.VersionGuid == VersionGuid)
 		{
+			Modify();
 			ExposedVersion = VersionGuid;
 			Data.Version.bIsVisibleInVersionSelector = true;
 			return;
@@ -473,7 +476,7 @@ void UNiagaraScript::EnableVersioning()
 	{
 		return;
 	}
-	
+	Modify();
 	ensure(VersionData.Num() == 1);
 	bVersioningEnabled = true;	
 	ExposedVersion = VersionData[0].Version.VersionGuid;
@@ -501,6 +504,29 @@ void UNiagaraScript::DisableVersioning(const FGuid& VersionGuidToUse)
 	VersionData.Add(DataToUse);
 }
 
+struct FNiagaraScriptVersionDataAccessor final : FNiagaraVersionDataAccessor
+{
+	virtual ~FNiagaraScriptVersionDataAccessor() override = default;
+	explicit FNiagaraScriptVersionDataAccessor(FVersionedNiagaraScriptData* InScriptData) : ScriptData(InScriptData) {}
+
+	virtual FNiagaraAssetVersion& GetObjectVersion() override { return ScriptData->Version; }
+	virtual FText& GetVersionChangeDescription() override { return ScriptData->VersionChangeDescription; }
+	virtual ENiagaraPythonUpdateScriptReference& GetUpdateScriptExecutionType() override { return ScriptData->UpdateScriptExecution; }
+	virtual FString& GetPythonUpdateScript() override { return ScriptData->PythonUpdateScript; }
+	virtual FFilePath& GetScriptAsset() override { return ScriptData->ScriptAsset; }
+	
+	FVersionedNiagaraScriptData* ScriptData;
+};
+
+TSharedPtr<FNiagaraVersionDataAccessor> UNiagaraScript::GetVersionDataAccessor(const FGuid& Version)
+{
+	if (FVersionedNiagaraScriptData* ScriptData = GetScriptData(Version))
+	{
+		return MakeShared<FNiagaraScriptVersionDataAccessor>(ScriptData);
+	}
+	return TSharedPtr<FNiagaraVersionDataAccessor>();
+}
+
 void UNiagaraScript::CheckVersionDataAvailable()
 {
 	if (VersionData.Num() > 0) {
@@ -550,6 +576,26 @@ UNiagaraScript* UNiagaraScript::CreateCompilationCopy()
 
 #endif
 
+FVersionedNiagaraEmitter UNiagaraScript::GetOuterEmitter() const
+{
+	UNiagaraEmitter* Emitter = GetTypedOuter<UNiagaraEmitter>();
+	if (Emitter == nullptr)
+	{
+		return FVersionedNiagaraEmitter();
+	}
+	for (FNiagaraAssetVersion& Version : Emitter->GetAllAvailableVersions())
+	{
+		FVersionedNiagaraEmitterData* EmitterData = Emitter->GetEmitterData(Version.VersionGuid);
+		TArray<UNiagaraScript*> OutScripts;
+		EmitterData->GetScripts(OutScripts, false);
+		if (OutScripts.Contains(this))
+		{
+			return FVersionedNiagaraEmitter(Emitter, Version.VersionGuid);
+		}
+	}
+	return FVersionedNiagaraEmitter();
+}
+
 bool FNiagaraVMExecutableDataId::IsValid() const
 {
 	return CompilerVersionID.IsValid();
@@ -833,15 +879,15 @@ void UNiagaraScript::ComputeVMCompilationId(FNiagaraVMExecutableDataId& Id, FGui
 	ENiagaraSimTarget SimTargetToBuild = ENiagaraSimTarget::CPUSim;
 	// Ideally we wouldn't want to do this but rather than push the data down
 	// from the emitter.  Checking all outers here to pick up simulation stages too.
-	UNiagaraEmitter* OuterEmitter = GetTypedOuter<UNiagaraEmitter>();
+	FVersionedNiagaraEmitter Outer = GetOuterEmitter();
+	FVersionedNiagaraEmitterData* EmitterData = Outer.GetEmitterData();
 	int32 SystemSpawnIdx = INDEX_NONE;
 	int32 SystemUpdateIdx = INDEX_NONE;
 	TArray<UNiagaraScript*> Scripts;
 
-	if (OuterEmitter != nullptr)
+	if (EmitterData != nullptr)
 	{
-		UNiagaraEmitter* Emitter = OuterEmitter;
-
+		UNiagaraEmitter* Emitter = Outer.Emitter;
 		if (UNiagaraSystem* EmitterOwner = Cast<UNiagaraSystem>(Emitter->GetOuter()))
 		{
 			SystemSpawnIdx = Scripts.Add(EmitterOwner->GetSystemSpawnScript());
@@ -861,16 +907,16 @@ void UNiagaraScript::ComputeVMCompilationId(FNiagaraVMExecutableDataId& Id, FGui
 			bool TrimAttributes = EmitterOwner->ShouldTrimAttributes();
 			if (TrimAttributes)
 			{
-				auto TrimAttributesSupported = [=](const UNiagaraEmitter* OtherEmitter)
+				auto TrimAttributesSupported = [=](const FVersionedNiagaraEmitter& OtherEmitter)
 				{
 					TArray<const UNiagaraDataInterfaceBase*> DataInterfaces;
-					if (OtherEmitter && OtherEmitter->GraphSource)
+					if (OtherEmitter.GetEmitterData() && OtherEmitter.GetEmitterData()->GraphSource)
 					{
-						OtherEmitter->GraphSource->CollectDataInterfaces(DataInterfaces);
+						OtherEmitter.GetEmitterData()->GraphSource->CollectDataInterfaces(DataInterfaces);
 
 						for (const UNiagaraDataInterfaceBase* DataInterface : DataInterfaces)
 						{
-							if (DataInterface->HasInternalAttributeReads(OtherEmitter, Emitter))
+							if (DataInterface->HasInternalAttributeReads(OtherEmitter.Emitter, Emitter))
 							{
 								return false;
 							}
@@ -912,14 +958,14 @@ void UNiagaraScript::ComputeVMCompilationId(FNiagaraVMExecutableDataId& Id, FGui
 					};
 
 				// preserve the attributes that have been defined on the emitter directly
-				for (const FString& Attribute : Emitter->AttributesToPreserve)
+				for (const FString& Attribute : EmitterData->AttributesToPreserve)
 				{
 					const FString PreserveDefine = TEXT("PreserveAttribute=") + Attribute;
 					PreserveAttributes.AddUnique(PreserveDefine);
 				}
 
 				// Now preserve the attributes that have been defined on the renderers in use
-				for (UNiagaraRendererProperties* RendererProperty : Emitter->GetRenderers())
+				for (UNiagaraRendererProperties* RendererProperty : EmitterData->GetRenderers())
 				{
 					if (RendererProperty->bIsEnabled)
 					{
@@ -930,7 +976,7 @@ void UNiagaraScript::ComputeVMCompilationId(FNiagaraVMExecutableDataId& Id, FGui
 					}
 				}
 
-				for (const UNiagaraSimulationStageBase* SimStageBase : Emitter->GetSimulationStages())
+				for (const UNiagaraSimulationStageBase* SimStageBase : EmitterData->GetSimulationStages())
 				{
 					if (const UNiagaraSimulationStageGeneric* SimStageGeneric = Cast<const UNiagaraSimulationStageGeneric>(SimStageBase))
 					{
@@ -946,37 +992,37 @@ void UNiagaraScript::ComputeVMCompilationId(FNiagaraVMExecutableDataId& Id, FGui
 				Id.AdditionalDefines.Append(PreserveAttributes);
 			}
 
-			ComputeVMCompilationId_EmitterShared(Id, Emitter, EmitterOwner, ENiagaraRendererSourceDataMode::Particles);
+			ComputeVMCompilationId_EmitterShared(Id, Outer, ENiagaraRendererSourceDataMode::Particles);
 		}
 
 		// Gather up the scripts we will use to iterate for static rapid iteration vars later...
-		Emitter->GetScripts(Scripts, false, true);
+		EmitterData->GetScripts(Scripts, false, true);
 
-		if ((Emitter->bInterpolatedSpawning && Usage == ENiagaraScriptUsage::ParticleGPUComputeScript) ||
-			(Emitter->bInterpolatedSpawning && Usage == ENiagaraScriptUsage::ParticleSpawnScript) ||
+		if ((EmitterData->bInterpolatedSpawning && Usage == ENiagaraScriptUsage::ParticleGPUComputeScript) ||
+			(EmitterData->bInterpolatedSpawning && Usage == ENiagaraScriptUsage::ParticleSpawnScript) ||
 			Usage == ENiagaraScriptUsage::ParticleSpawnScriptInterpolated)
 		{
 			Id.bInterpolatedSpawn = true;
 			Id.AdditionalDefines.Add(TEXT("InterpolatedSpawn"));
 		}
-		if (Emitter->bInterpolatedSpawning)
+		if (EmitterData->bInterpolatedSpawning)
 		{
 			Id.AdditionalDefines.Add(TEXT("Emitter.InterpolatedSpawn"));
 		}
-		if (IsParticleScript(Usage))
+		if (IsParticleScript(Usage) || IsParticleStageScript(Usage))
 		{
-			SimTargetToBuild = Emitter->SimTarget;
+			SimTargetToBuild = EmitterData->SimTarget;
 		}
-		if (Emitter->RequiresPersistentIDs())
+		if (EmitterData->RequiresPersistentIDs())
 		{
 			Id.bRequiresPersistentIDs = true;
 			Id.AdditionalDefines.Add(TEXT("RequiresPersistentIDs"));
 		}
-		if (Emitter->bLocalSpace)
+		if (EmitterData->bLocalSpace)
 		{
 			Id.AdditionalDefines.Add(TEXT("Emitter.Localspace"));
 		}
-		if (Emitter->bDeterminism)
+		if (EmitterData->bDeterminism)
 		{
 			Id.AdditionalDefines.Add(TEXT("Emitter.Determinism"));
 		}
@@ -986,7 +1032,7 @@ void UNiagaraScript::ComputeVMCompilationId(FNiagaraVMExecutableDataId& Id, FGui
 
 			FSHA1 HashState;
 			FNiagaraCompileHashVisitor Visitor(HashState);
-			for (UNiagaraSimulationStageBase* Base : Emitter->GetSimulationStages())
+			for (UNiagaraSimulationStageBase* Base : EmitterData->GetSimulationStages())
 			{
 				// bool AppendCompileHash(FNiagaraCompileHashVisitor* InVisitor) const;
 				if (Base && Base->bEnabled)
@@ -1007,7 +1053,7 @@ void UNiagaraScript::ComputeVMCompilationId(FNiagaraVMExecutableDataId& Id, FGui
 
 		if (IsParticleEventScript(Usage))
 		{
-			if (const FNiagaraEventScriptProperties* EventScriptProps = Emitter->GetEventHandlerByIdUnsafe(UsageId))
+			if (const FNiagaraEventScriptProperties* EventScriptProps = EmitterData->GetEventHandlerByIdUnsafe(UsageId))
 			{
 				if (EventScriptProps->ExecutionMode == EScriptExecutionMode::SpawnedParticles)
 				{
@@ -1032,23 +1078,23 @@ void UNiagaraScript::ComputeVMCompilationId(FNiagaraVMExecutableDataId& Id, FGui
 
 		for (const FNiagaraEmitterHandle& EmitterHandle: System->GetEmitterHandles())
 		{
-			UNiagaraEmitter* Emitter = Cast<UNiagaraEmitter>(EmitterHandle.GetInstance());
-			if (Emitter && EmitterHandle.GetIsEnabled())
+			FVersionedNiagaraEmitterData* HandleData = EmitterHandle.GetEmitterData();
+			if (HandleData && EmitterHandle.GetIsEnabled())
 			{
-				if (Emitter->bLocalSpace)
+				if (HandleData->bLocalSpace)
 				{
-					Id.AdditionalDefines.Add(Emitter->GetUniqueEmitterName() + TEXT(".Localspace"));
+					Id.AdditionalDefines.Add(EmitterHandle.GetUniqueInstanceName() + TEXT(".Localspace"));
 				}
-				if (Emitter->bDeterminism)
+				if (HandleData->bDeterminism)
 				{
-					Id.AdditionalDefines.Add(Emitter->GetUniqueEmitterName() + TEXT(".Determinism"));
+					Id.AdditionalDefines.Add(EmitterHandle.GetUniqueInstanceName() + TEXT(".Determinism"));
 				}
-				if (Emitter->bInterpolatedSpawning)
+				if (HandleData->bInterpolatedSpawning)
 				{
-					Id.AdditionalDefines.Add(Emitter->GetUniqueEmitterName() + TEXT(".InterpolatedSpawn"));
+					Id.AdditionalDefines.Add(EmitterHandle.GetUniqueInstanceName() + TEXT(".InterpolatedSpawn"));
 				}
 
-				ComputeVMCompilationId_EmitterShared(Id, Emitter, System, ENiagaraRendererSourceDataMode::Emitter);
+				ComputeVMCompilationId_EmitterShared(Id, EmitterHandle.GetInstance(), ENiagaraRendererSourceDataMode::Emitter);
 			}
 		}
 	}
@@ -1215,10 +1261,10 @@ void UNiagaraScript::ComputeVMCompilationId(FNiagaraVMExecutableDataId& Id, FGui
 }
 
 
-void UNiagaraScript::ComputeVMCompilationId_EmitterShared(FNiagaraVMExecutableDataId& Id, UNiagaraEmitter* Emitter, UNiagaraSystem* EmitterOwner, ENiagaraRendererSourceDataMode InSourceMode) const
+void UNiagaraScript::ComputeVMCompilationId_EmitterShared(FNiagaraVMExecutableDataId& Id, const FVersionedNiagaraEmitter& Emitter, ENiagaraRendererSourceDataMode InSourceMode) const
 {
 	// Gather additional variables from renderers
-	for (UNiagaraRendererProperties* RendererProperty : Emitter->GetRenderers())
+	for (UNiagaraRendererProperties* RendererProperty : Emitter.GetEmitterData()->GetRenderers())
 	{
 		// Because we currently generate the CompilationId during PostLoad (as a part of AreScriptAndSourceSynchronized)
 		// we need to ensure that the RendererProperties have done their PostLoad.
@@ -1328,11 +1374,11 @@ TOptional<ENiagaraSimTarget> UNiagaraScript::GetSimTarget() const
 	case ENiagaraScriptUsage::ParticleEventScript:
 	case ENiagaraScriptUsage::ParticleSimulationStageScript:
 	case ENiagaraScriptUsage::ParticleGPUComputeScript:
-		if (UNiagaraEmitter* OwningEmitter = GetTypedOuter<UNiagaraEmitter>())
+		if (FVersionedNiagaraEmitterData* EmitterData = GetOuterEmitter().GetEmitterData())
 		{
-			if (OwningEmitter->SimTarget != ENiagaraSimTarget::CPUSim || CachedScriptVM.IsValid())
+			if (EmitterData->SimTarget != ENiagaraSimTarget::CPUSim || CachedScriptVM.IsValid())
 			{
-				return OwningEmitter->SimTarget;
+				return EmitterData->SimTarget;
 			}
 		}
 		break;
@@ -1804,6 +1850,7 @@ void UNiagaraScript::PostLoad()
 
 		// if our bindings aren't valid, then something has gone wrong with our cook and we need to disable this Script, which will in turn
 		// disable the owning script and system
+		
 		if (!HasValidParameterBindings())
 		{
 #if WITH_EDITOR && IS_MONOLITHIC
@@ -1811,7 +1858,7 @@ void UNiagaraScript::PostLoad()
 #else
 			UE_LOG(LogNiagara, Error, TEXT("Mismatch between binding between RapidIterationParamters and ScriptExecutionParameters for system %s"), *GetFullName());
 #endif
-
+			
 			CachedScriptVM.Reset();
 			return;
 		}
@@ -2002,8 +2049,9 @@ bool UNiagaraScript::ShouldCacheShadersForCooking(const ITargetPlatform* TargetP
 			}
 		}
 
-		UNiagaraEmitter* OwningEmitter = GetTypedOuter<UNiagaraEmitter>();
-		if (OwningEmitter != nullptr && OwningEmitter->SimTarget == ENiagaraSimTarget::GPUComputeSim && OwningEmitter->NeedsLoadForTargetPlatform(TargetPlatform))
+		FVersionedNiagaraEmitter OuterEmitter = GetOuterEmitter();
+		FVersionedNiagaraEmitterData* EmitterData = OuterEmitter.GetEmitterData();
+		if (EmitterData != nullptr && EmitterData->SimTarget == ENiagaraSimTarget::GPUComputeSim && OuterEmitter.Emitter->NeedsLoadForTargetPlatform(TargetPlatform))
 		{
 			return true;
 		}
@@ -2681,9 +2729,9 @@ void UNiagaraScript::RaiseOnGPUCompilationComplete()
 	OnGPUScriptCompiled().Broadcast(this, FGuid());
 	FNiagaraSystemUpdateContext(this, true);
 
-	if (UNiagaraEmitter* EmitterOwner = Cast<UNiagaraEmitter>(GetOuter()) )
+	if (FVersionedNiagaraEmitterData* EmitterData = GetOuterEmitter().GetEmitterData())
 	{
-		EmitterOwner->CacheFromShaderCompiled();
+		EmitterData->CacheFromShaderCompiled();
 	}
 #endif
 }
@@ -2827,18 +2875,11 @@ bool UNiagaraScript::IsEditorOnly() const
 		return false;
 	}
 
-	if (const UNiagaraEmitter* EmitterOwner = Cast<const UNiagaraEmitter>(GetOuter()))
+	UNiagaraSystem* OuterSystem = GetTypedOuter<UNiagaraSystem>();
+	bool bIsSystemScript = OuterSystem && (OuterSystem->GetSystemSpawnScript() == this || OuterSystem->GetSystemUpdateScript() == this);
+	if (!bIsSystemScript && !GetOuterEmitter().GetEmitterData())
 	{
-		// we want to only cook scripts that are referenced by systems (as opposed to standalone scripts that may be getting
-		// referenced via an emitter's parent, this will also take care of GPUScripts that are created for CPU emitters
-		TArray<UNiagaraScript*> OwnerScripts;
-
-		EmitterOwner->GetScripts(OwnerScripts, false);
-
-		if (!OwnerScripts.Contains(this))
-		{
-			return true;
-		}
+		return true;
 	}
 #endif
 	return Super::IsEditorOnly();
@@ -3590,12 +3631,11 @@ bool UNiagaraScript::CanBeRunOnGpu()const
 	return true;
 }
 
-
 bool UNiagaraScript::OwnerCanBeRunOnGpu() const
 {
-	if (UNiagaraEmitter* Emitter = GetTypedOuter<UNiagaraEmitter>())
+	if (FVersionedNiagaraEmitterData* EmitterData = GetOuterEmitter().GetEmitterData())
 	{
-		if (Emitter->SimTarget == ENiagaraSimTarget::GPUComputeSim)
+		if (EmitterData->SimTarget == ENiagaraSimTarget::GPUComputeSim)
 		{
 			return true;
 		}
@@ -3606,9 +3646,9 @@ bool UNiagaraScript::OwnerCanBeRunOnGpu() const
 
 bool UNiagaraScript::LegacyCanBeRunOnGpu() const
 {
-	if (UNiagaraEmitter* Emitter = GetTypedOuter<UNiagaraEmitter>())
+	if (FVersionedNiagaraEmitterData* EmitterData = GetOuterEmitter().GetEmitterData())
 	{
-		if (Emitter->SimTarget == ENiagaraSimTarget::CPUSim)
+		if (EmitterData->SimTarget == ENiagaraSimTarget::CPUSim)
 		{
 			return false;
 		}
diff --git a/Engine/Plugins/FX/Niagara/Source/Niagara/Private/NiagaraSimulationStageBase.cpp b/Engine/Plugins/FX/Niagara/Source/Niagara/Private/NiagaraSimulationStageBase.cpp
--- a/Engine/Plugins/FX/Niagara/Source/Niagara/Private/NiagaraSimulationStageBase.cpp
+++ b/Engine/Plugins/FX/Niagara/Source/Niagara/Private/NiagaraSimulationStageBase.cpp
@@ -38,16 +38,15 @@ void UNiagaraSimulationStageBase::SetEnabled(bool bInEnabled)
 
 void UNiagaraSimulationStageBase::RequestRecompile()
 {
-	UNiagaraEmitter* Emitter = Cast< UNiagaraEmitter>(GetOuter());
-	if (Emitter)
+	FVersionedNiagaraEmitter OuterEmitter = GetOuterEmitter();
+	if (FVersionedNiagaraEmitterData* EmitterData = OuterEmitter.GetEmitterData())
 	{
-		UNiagaraScriptSourceBase* GraphSource = Emitter->UpdateScriptProps.Script->GetLatestSource();
-		if (GraphSource != nullptr)
+		if (UNiagaraScriptSourceBase* GraphSource = EmitterData->UpdateScriptProps.Script->GetLatestSource())
 		{
 			GraphSource->MarkNotSynchronized(TEXT("SimulationStage changed."));
 		}
 
-		UNiagaraSystem::RequestCompileForEmitter(Emitter);
+		UNiagaraSystem::RequestCompileForEmitter(OuterEmitter);
 	}
 }
 
@@ -218,6 +217,24 @@ void UNiagaraSimulationStageGeneric::PostEditChangeProperty(struct FPropertyChan
 	}
 }
 
+FVersionedNiagaraEmitterData* UNiagaraSimulationStageBase::GetEmitterData() const
+{
+	if (UNiagaraEmitter* SrcEmitter = GetTypedOuter<UNiagaraEmitter>())
+	{
+		return SrcEmitter->GetEmitterData(OuterEmitterVersion);
+	}
+	return nullptr;
+}
+
+FVersionedNiagaraEmitter UNiagaraSimulationStageBase::GetOuterEmitter() const
+{
+	if (UNiagaraEmitter* SrcEmitter = GetTypedOuter<UNiagaraEmitter>())
+	{
+		return FVersionedNiagaraEmitter(SrcEmitter, OuterEmitterVersion);
+	}
+	return FVersionedNiagaraEmitter();
+}
+
 FName UNiagaraSimulationStageGeneric::GetStackContextReplacementName() const 
 {
 	if (IterationSource == ENiagaraIterationSource::Particles)
diff --git a/Engine/Plugins/FX/Niagara/Source/Niagara/Private/NiagaraSpriteRendererProperties.cpp b/Engine/Plugins/FX/Niagara/Source/Niagara/Private/NiagaraSpriteRendererProperties.cpp
--- a/Engine/Plugins/FX/Niagara/Source/Niagara/Private/NiagaraSpriteRendererProperties.cpp
+++ b/Engine/Plugins/FX/Niagara/Source/Niagara/Private/NiagaraSpriteRendererProperties.cpp
@@ -148,7 +148,7 @@ void UNiagaraSpriteRendererProperties::PostLoad()
 	PostLoadBindings(SourceMode);
 
 	// Fix up these bindings from their loaded source bindings
-	SetPreviousBindings(nullptr, SourceMode);
+	SetPreviousBindings(FVersionedNiagaraEmitter(), SourceMode);
 #endif // WITH_EDITORONLY_DATA
 }
 
@@ -235,10 +235,10 @@ void UNiagaraSpriteRendererProperties::InitBindings()
 		CustomSortingBinding = FNiagaraConstants::GetAttributeDefaultBinding(SYS_PARAM_PARTICLES_NORMALIZED_AGE);
 	}
 
-	SetPreviousBindings(nullptr, SourceMode);
+	SetPreviousBindings(FVersionedNiagaraEmitter(), SourceMode);
 }
 
-void UNiagaraSpriteRendererProperties::SetPreviousBindings(const UNiagaraEmitter* SrcEmitter, ENiagaraRendererSourceDataMode InSourceMode)
+void UNiagaraSpriteRendererProperties::SetPreviousBindings(const FVersionedNiagaraEmitter& SrcEmitter, ENiagaraRendererSourceDataMode InSourceMode)
 {
 	PrevPositionBinding.SetAsPreviousValue(PositionBinding, SrcEmitter, InSourceMode);
 	PrevVelocityBinding.SetAsPreviousValue(VelocityBinding, SrcEmitter, InSourceMode);
@@ -355,7 +355,7 @@ void UNiagaraSpriteRendererProperties::UpdateSourceModeDerivates(ENiagaraRendere
 			MaterialParamBinding.CacheValues(SrcEmitter);
 		}
 
-		SetPreviousBindings(SrcEmitter, InSourceMode);
+		SetPreviousBindings(FVersionedNiagaraEmitter(), InSourceMode);
 	}
 }
 
@@ -440,28 +440,28 @@ void UNiagaraSpriteRendererProperties::PostEditChangeProperty(struct FPropertyCh
 }
 
 
-void UNiagaraSpriteRendererProperties::RenameVariable(const FNiagaraVariableBase& OldVariable, const FNiagaraVariableBase& NewVariable, const UNiagaraEmitter* InEmitter)
+void UNiagaraSpriteRendererProperties::RenameVariable(const FNiagaraVariableBase& OldVariable, const FNiagaraVariableBase& NewVariable, const FVersionedNiagaraEmitter& InEmitter)
 {
 	Super::RenameVariable(OldVariable, NewVariable, InEmitter);
 
 	// Handle renaming material bindings
 	for (FNiagaraMaterialAttributeBinding& Binding : MaterialParameterBindings)
 	{
-		Binding.RenameVariableIfMatching(OldVariable, NewVariable, InEmitter, GetCurrentSourceMode());
+		Binding.RenameVariableIfMatching(OldVariable, NewVariable, InEmitter.Emitter, GetCurrentSourceMode());
 	}
 }
 
-void UNiagaraSpriteRendererProperties::RemoveVariable(const FNiagaraVariableBase& OldVariable, const UNiagaraEmitter* InEmitter)
+void UNiagaraSpriteRendererProperties::RemoveVariable(const FNiagaraVariableBase& OldVariable, const FVersionedNiagaraEmitter& InEmitter)
 {
 	Super::RemoveVariable(OldVariable, InEmitter);
 
 	// Handle resetting material bindings to defaults
 	for (FNiagaraMaterialAttributeBinding& Binding : MaterialParameterBindings)
 	{
-		if (Binding.Matches(OldVariable, InEmitter, GetCurrentSourceMode()))
+		if (Binding.Matches(OldVariable, InEmitter.Emitter, GetCurrentSourceMode()))
 		{
 			Binding.NiagaraVariable = FNiagaraVariable();
-			Binding.CacheValues(InEmitter);
+			Binding.CacheValues(InEmitter.Emitter);
 		}
 	}
 }
@@ -535,10 +535,10 @@ void UNiagaraSpriteRendererProperties::GetRendererWidgets(const FNiagaraEmitterI
 }
 
 
-void UNiagaraSpriteRendererProperties::GetRendererFeedback(const UNiagaraEmitter* InEmitter, TArray<FText>& OutErrors, TArray<FText>& OutWarnings, TArray<FText>& OutInfo) const
+void UNiagaraSpriteRendererProperties::GetRendererFeedback(const FVersionedNiagaraEmitter& InEmitter, TArray<FText>& OutErrors, TArray<FText>& OutWarnings, TArray<FText>& OutInfo) const
 {
 	Super::GetRendererFeedback(InEmitter, OutErrors, OutWarnings, OutInfo);
-	if (InEmitter->SpawnScriptProps.Script->GetVMExecutableData().IsValid())
+	if (InEmitter.GetEmitterData()->SpawnScriptProps.Script->GetVMExecutableData().IsValid())
 	{
 		if (bUseMaterialCutoutTexture || CutoutTexture)
 		{
diff --git a/Engine/Plugins/FX/Niagara/Source/Niagara/Private/NiagaraSystem.cpp b/Engine/Plugins/FX/Niagara/Source/Niagara/Private/NiagaraSystem.cpp
--- a/Engine/Plugins/FX/Niagara/Source/Niagara/Private/NiagaraSystem.cpp
+++ b/Engine/Plugins/FX/Niagara/Source/Niagara/Private/NiagaraSystem.cpp
@@ -32,9 +32,6 @@
 #include "UObject/Package.h"
 #include "UObject/UObjectThreadContext.h"
 #include "ShaderCompiler.h"
-#include "NiagaraSimulationStageBase.h"
-#include "DeviceProfiles/DeviceProfile.h"
-#include "DeviceProfiles/DeviceProfileManager.h"
 
 #define LOCTEXT_NAMESPACE "NiagaraSystem"
 
@@ -135,6 +132,10 @@ UNiagaraSystem::UNiagaraSystem(const FObjectInitializer& ObjectInitializer)
 	AssetGuid = FGuid::NewGuid();
 #endif
 
+#if STATS
+	StatDatabase.Init();
+#endif
+	
 	const UNiagaraSettings* Settings = GetDefault<UNiagaraSettings>();
 	bLwcEnabledSettingCached = Settings ? Settings->bSystemsSupportLargeWorldCoordinates : true;
 }
@@ -201,10 +202,10 @@ void UNiagaraSystem::HandleVariableRenamed(const FNiagaraVariable& InOldVariable
 
 	for (const FNiagaraEmitterHandle& Handle : GetEmitterHandles())
 	{
-		UNiagaraEmitter* Emitter = Handle.GetInstance();
-		if (Emitter)
+		FVersionedNiagaraEmitter VersionedEmitter = Handle.GetInstance();
+		if (VersionedEmitter.Emitter)
 		{
-			Emitter->HandleVariableRenamed(InOldVariable, InNewVariable, false);
+			VersionedEmitter.Emitter->HandleVariableRenamed(InOldVariable, InNewVariable, false, VersionedEmitter.Version);
 		}
 	}
 
@@ -224,10 +225,10 @@ void UNiagaraSystem::HandleVariableRemoved(const FNiagaraVariable& InOldVariable
 	}
 	for (const FNiagaraEmitterHandle& Handle : GetEmitterHandles())
 	{
-		UNiagaraEmitter* Emitter = Handle.GetInstance();
-		if (Emitter)
+		FVersionedNiagaraEmitter VersionedEmitter = Handle.GetInstance();
+		if (VersionedEmitter.Emitter)
 		{
-			Emitter->HandleVariableRemoved(InOldVariable, false);
+			VersionedEmitter.Emitter->HandleVariableRemoved(InOldVariable, false, VersionedEmitter.Version);
 		}
 	}
 	if (bUpdateContexts)
@@ -257,11 +258,32 @@ TArray<INiagaraParameterDefinitionsSubscriber*> UNiagaraSystem::GetOwnedParamete
 	TArray<INiagaraParameterDefinitionsSubscriber*> OutSubscribers;
 	for (const FNiagaraEmitterHandle& EmitterHandle : GetEmitterHandles())
 	{
-		OutSubscribers.Add(EmitterHandle.GetInstance());
+		OutSubscribers.Add(EmitterHandle.GetInstance().Emitter);
 	}
 	return OutSubscribers;
 }
 
+bool UNiagaraSystem::ChangeEmitterVersion(const FVersionedNiagaraEmitter& VersionedEmitter, const FGuid& NewVersion)
+{
+	if (VersionedEmitter.Emitter == nullptr || VersionedEmitter.Emitter->IsVersioningEnabled() == false)
+	{
+		return false;
+	}
+	for (int i = 0; i < EmitterHandles.Num(); i++)
+	{
+		FNiagaraEmitterHandle& Handle = EmitterHandles[i];
+		if (Handle.GetInstance() == VersionedEmitter)
+		{
+			KillAllActiveCompilations();
+			Modify();
+			EmitterHandles[i].GetInstance().Version = NewVersion;
+			RequestCompile(false);
+			return true;
+		}
+	}
+	return false;
+}
+
 #endif
 
 void UNiagaraSystem::PostInitProperties()
@@ -317,7 +339,8 @@ bool UNiagaraSystem::UsesCollection(const UNiagaraParameterCollection* Collectio
 
 	for (const FNiagaraEmitterHandle& EmitterHandle : GetEmitterHandles())
 	{
-		if (EmitterHandle.GetInstance() && EmitterHandle.GetInstance()->UsesCollection(Collection))
+		FVersionedNiagaraEmitterData* EmitterData = EmitterHandle.GetEmitterData();
+		if (EmitterData && EmitterData->UsesCollection(Collection))
 		{
 			return true;
 		}
@@ -346,9 +369,9 @@ void UNiagaraSystem::UpdateSystemAfterLoad()
 
 	for (FNiagaraEmitterHandle& EmitterHandle : EmitterHandles)
 	{
-		if (EmitterHandle.GetInstance() != nullptr)
+		if (EmitterHandle.GetInstance().Emitter)
 		{
-			EmitterHandle.GetInstance()->UpdateEmitterAfterLoad();
+			EmitterHandle.GetInstance().Emitter->UpdateEmitterAfterLoad();
 		}
 	}
 
@@ -503,7 +526,7 @@ void UNiagaraSystem::UpdateSystemAfterLoad()
 		bool bEmitterScriptsAreSynchronized = true;
 		for (FNiagaraEmitterHandle& EmitterHandle : EmitterHandles)
 		{
-			if (EmitterHandle.GetIsEnabled() && EmitterHandle.GetInstance() && !EmitterHandle.GetInstance()->AreAllScriptAndSourcesSynchronized())
+			if (EmitterHandle.GetIsEnabled() && EmitterHandle.GetEmitterData() && !EmitterHandle.GetEmitterData()->AreAllScriptAndSourcesSynchronized())
 			{
 				bEmitterScriptsAreSynchronized = false;
 			}
@@ -551,7 +574,8 @@ bool UNiagaraSystem::UsesScript(const UNiagaraScript* Script) const
 
 	for (const FNiagaraEmitterHandle& EmitterHandle : GetEmitterHandles())
 	{
-		if (EmitterHandle.GetInstance() && EmitterHandle.GetInstance()->UsesScript(Script))
+		FVersionedNiagaraEmitterData* EmitterData = EmitterHandle.GetEmitterData();
+		if (EmitterData && EmitterData->UsesScript(Script))
 		{
 			return true;
 		}
@@ -560,22 +584,39 @@ bool UNiagaraSystem::UsesScript(const UNiagaraScript* Script) const
 	return false;
 }
 
-bool UNiagaraSystem::UsesEmitter(const UNiagaraEmitter* Emitter) const
+bool UNiagaraSystem::UsesEmitter(UNiagaraEmitter* Emitter) const
 {
-	if (Emitter != nullptr)
+	if (!Emitter)
 	{
-		for (const FNiagaraEmitterHandle& EmitterHandle : GetEmitterHandles())
+		return false;
+	}
+	for (const FNiagaraEmitterHandle& EmitterHandle : GetEmitterHandles())
+	{
+		if (EmitterHandle.UsesEmitter(*Emitter))
 		{
-			if (EmitterHandle.UsesEmitter(*Emitter))
-			{
-				return true;
-			}
+			return true;
 		}
 	}
 	return false;
 }
 
-void UNiagaraSystem::RequestCompileForEmitter(UNiagaraEmitter* InEmitter)
+bool UNiagaraSystem::UsesEmitter(const FVersionedNiagaraEmitter& VersionedEmitter) const
+{
+	if (!VersionedEmitter.Emitter)
+	{
+		return false;
+	}
+	for (const FNiagaraEmitterHandle& EmitterHandle : GetEmitterHandles())
+	{
+		if (EmitterHandle.UsesEmitter(VersionedEmitter))
+		{
+			return true;
+		}
+	}
+	return false;
+}
+
+void UNiagaraSystem::RequestCompileForEmitter(const FVersionedNiagaraEmitter& InEmitter)
 {
 	for (TObjectIterator<UNiagaraSystem> It; It; ++It)
 	{
@@ -587,7 +628,7 @@ void UNiagaraSystem::RequestCompileForEmitter(UNiagaraEmitter* InEmitter)
 	}
 }
 
-void UNiagaraSystem::RecomputeExecutionOrderForEmitter(UNiagaraEmitter* InEmitter)
+void UNiagaraSystem::RecomputeExecutionOrderForEmitter(const FVersionedNiagaraEmitter& InEmitter)
 {
 	for (TObjectIterator<UNiagaraSystem> It; It; ++It)
 	{
@@ -603,7 +644,7 @@ void UNiagaraSystem::RecomputeExecutionOrderForDataInterface(class UNiagaraDataI
 {
 	if (UNiagaraEmitter* Emitter = DataInterface->GetTypedOuter<UNiagaraEmitter>())
 	{
-		RecomputeExecutionOrderForEmitter(Emitter);
+		RecomputeExecutionOrderForEmitter(FVersionedNiagaraEmitter(Emitter, Emitter->GetExposedVersion().VersionGuid));
 	}
 	else
 	{
@@ -794,7 +835,7 @@ void UNiagaraSystem::PostLoad()
 #if WITH_EDITORONLY_DATA
 		EmitterHandle.ConditionalPostLoad(NiagaraVer);
 #else
-		if (UNiagaraEmitter* NiagaraEmitter = EmitterHandle.GetInstance())
+		if (UNiagaraEmitter* NiagaraEmitter = EmitterHandle.GetInstance().Emitter)
 		{
 			NiagaraEmitter->ConditionalPostLoad();
 		}
@@ -833,7 +874,7 @@ void UNiagaraSystem::PostLoad()
 #endif // WITH_EDITORONLY_DATA
 
 	//Apply platform set redirectors
-	auto ApplyPlatformSetRedirects = [](UObject* Owner, FNiagaraPlatformSet& Platforms)
+	auto ApplyPlatformSetRedirects = [](FNiagaraPlatformSet& Platforms)
 	{
 		Platforms.ApplyRedirects();
 	};
@@ -871,16 +912,16 @@ UNiagaraEditorParametersAdapterBase* UNiagaraSystem::GetEditorParameters()
 	return EditorParameters;
 }
 
-bool UNiagaraSystem::ReferencesInstanceEmitter(UNiagaraEmitter& Emitter)
+bool UNiagaraSystem::ReferencesInstanceEmitter(const FVersionedNiagaraEmitter& Emitter) const
 {
-	if (&Emitter == nullptr)
+	if (Emitter.Emitter == nullptr)
 	{
 		return false;
 	}
 
-	for (FNiagaraEmitterHandle& Handle : EmitterHandles)
+	for (const FNiagaraEmitterHandle& Handle : EmitterHandles)
 	{
-		if (&Emitter == Handle.GetInstance())
+		if (Emitter == Handle.GetInstance())
 		{
 			return true;
 		}
@@ -894,12 +935,10 @@ void UNiagaraSystem::RefreshSystemParametersFromEmitter(const FNiagaraEmitterHan
 	if (ensureMsgf(EmitterHandles.ContainsByPredicate([=](const FNiagaraEmitterHandle& OwnedEmitterHandle) { return OwnedEmitterHandle.GetId() == EmitterHandle.GetId(); }),
 		TEXT("Can't refresh parameters from an emitter handle this system doesn't own.")))
 	{
-		if (EmitterHandle.GetInstance())
+		if (FVersionedNiagaraEmitterData* EmitterData = EmitterHandle.GetEmitterData())
 		{
-			EmitterHandle.GetInstance()->EmitterSpawnScriptProps.Script->RapidIterationParameters.CopyParametersTo(
-				SystemSpawnScript->RapidIterationParameters, false, FNiagaraParameterStore::EDataInterfaceCopyMethod::None);
-			EmitterHandle.GetInstance()->EmitterUpdateScriptProps.Script->RapidIterationParameters.CopyParametersTo(
-				SystemUpdateScript->RapidIterationParameters, false, FNiagaraParameterStore::EDataInterfaceCopyMethod::None);
+			EmitterData->EmitterSpawnScriptProps.Script->RapidIterationParameters.CopyParametersTo(SystemSpawnScript->RapidIterationParameters, false, FNiagaraParameterStore::EDataInterfaceCopyMethod::None);
+			EmitterData->EmitterUpdateScriptProps.Script->RapidIterationParameters.CopyParametersTo(SystemUpdateScript->RapidIterationParameters, false, FNiagaraParameterStore::EDataInterfaceCopyMethod::None);
 		}
 	}
 }
@@ -910,10 +949,10 @@ void UNiagaraSystem::RemoveSystemParametersForEmitter(const FNiagaraEmitterHandl
 	if (ensureMsgf(EmitterHandles.ContainsByPredicate([=](const FNiagaraEmitterHandle& OwnedEmitterHandle) { return OwnedEmitterHandle.GetId() == EmitterHandle.GetId(); }),
 		TEXT("Can't remove parameters for an emitter handle this system doesn't own.")))
 	{
-		if (EmitterHandle.GetInstance())
+		if (FVersionedNiagaraEmitterData* EmitterData = EmitterHandle.GetEmitterData())
 		{
-			EmitterHandle.GetInstance()->EmitterSpawnScriptProps.Script->RapidIterationParameters.RemoveParameters(SystemSpawnScript->RapidIterationParameters);
-			EmitterHandle.GetInstance()->EmitterUpdateScriptProps.Script->RapidIterationParameters.RemoveParameters(SystemUpdateScript->RapidIterationParameters);
+			EmitterData->EmitterSpawnScriptProps.Script->RapidIterationParameters.RemoveParameters(SystemSpawnScript->RapidIterationParameters);
+			EmitterData->EmitterUpdateScriptProps.Script->RapidIterationParameters.RemoveParameters(SystemUpdateScript->RapidIterationParameters);
 		}
 	}
 }
@@ -982,14 +1021,15 @@ bool UNiagaraSystem::IsReadyToRunInternal() const
 	for (int32 EmitterIt = 0; EmitterIt < EmitterCount; ++EmitterIt)
 	{
 		const FNiagaraEmitterHandle& Handle = EmitterHandles[EmitterIt];
-		if (Handle.GetInstance() && !Handle.GetInstance()->IsReadyToRun())
+		FVersionedNiagaraEmitterData* EmitterData = Handle.GetEmitterData();
+		if (EmitterData && !EmitterData->IsReadyToRun())
 		{
 			if (FPlatformProperties::RequiresCookedData())
 			{
 				UE_LOG(LogNiagara, Warning, TEXT("%s IsReadyToRunInternal() failed due to Emitter not being ready to run.  Emitter #%d - %s"),
 					*GetFullName(),
 					EmitterIt,
-					Handle.GetInstance() ? *Handle.GetInstance()->GetUniqueEmitterName() : TEXT("<none>"));
+					Handle.GetInstance().Emitter ? *Handle.GetInstance().Emitter->GetUniqueEmitterName() : TEXT("<none>"));
 			}
 
 			return false;
@@ -1038,16 +1078,15 @@ void UNiagaraSystem::GetAssetRegistryTags(TArray<FAssetRegistryTag>& OutTags) co
 		{
 
 			NumActiveEmitters++;
-			const UNiagaraEmitter* Emitter = Handle.GetInstance();
-			if (Emitter)
+			if (FVersionedNiagaraEmitterData* EmitterData = Handle.GetEmitterData())
 			{
 				// Only register fixed bounds requirement for GPU if the system itself isn't fixed bounds.
-				if (bFixedBounds == false && Emitter->CalculateBoundsMode == ENiagaraEmitterCalculateBoundMode::Dynamic && Emitter->SimTarget == ENiagaraSimTarget::GPUComputeSim)
+				if (bFixedBounds == false && EmitterData->CalculateBoundsMode == ENiagaraEmitterCalculateBoundMode::Dynamic && EmitterData->SimTarget == ENiagaraSimTarget::GPUComputeSim)
 				{
 					GPUSimsMissingFixedBounds++;
 				}
 
-				for (const UNiagaraRendererProperties* Props : Emitter->GetRenderers())
+				for (const UNiagaraRendererProperties* Props : EmitterData->GetRenderers())
 				{
 					if (Props)
 					{
@@ -1080,12 +1119,13 @@ void UNiagaraSystem::GetAssetRegistryTags(TArray<FAssetRegistryTag>& OutTags) co
 		{
 			if (Handle.GetIsEnabled())
 			{
-				const UNiagaraEmitter* Emitter = Handle.GetInstance();
-				if (Emitter)
+				FVersionedNiagaraEmitter Emitter = Handle.GetInstance();
+				FVersionedNiagaraEmitterData* EmitterData = Emitter.GetEmitterData();
+				if (EmitterData)
 				{
 					for (int32 i = 0; i < NumQualityLevels; i++)
 					{
-						if (Emitter->Platforms.IsEffectQualityEnabled(i))
+						if (EmitterData->Platforms.IsEffectQualityEnabled(i))
 						{
 							QualityLevelsNumActive[i]++;
 						}
@@ -1115,7 +1155,7 @@ void UNiagaraSystem::GetAssetRegistryTags(TArray<FAssetRegistryTag>& OutTags) co
 			const UNiagaraRendererProperties* PropDefault = RendererClass->GetDefaultObject< UNiagaraRendererProperties>();
 			if (PropDefault)
 			{
-				PropDefault->GetAssetTagsForContext(this, ActiveRenderers, NumericKeys, StringKeys);
+				PropDefault->GetAssetTagsForContext(this, FGuid(), ActiveRenderers, NumericKeys, StringKeys);
 			}
 		}
 	}
@@ -1141,10 +1181,10 @@ void UNiagaraSystem::GetAssetRegistryTags(TArray<FAssetRegistryTag>& OutTags) co
 		{
 			if (Handle.GetIsEnabled())
 			{
-				if (UNiagaraEmitter* Emitter = Handle.GetInstance())
+				if (FVersionedNiagaraEmitterData* EmitterData = Handle.GetEmitterData())
 				{
 					TArray<UNiagaraScript*> Scripts;
-					Emitter->GetScripts(Scripts);
+					EmitterData->GetScripts(Scripts);
 					for (UNiagaraScript* Script : Scripts)
 					{
 						AddDIs(Script);
@@ -1161,7 +1201,7 @@ void UNiagaraSystem::GetAssetRegistryTags(TArray<FAssetRegistryTag>& OutTags) co
 			const UNiagaraDataInterface* PropDefault = DIClass->GetDefaultObject< UNiagaraDataInterface>();
 			if (PropDefault)
 			{
-				PropDefault->GetAssetTagsForContext(this, DataInterfaces, NumericKeys, StringKeys);
+				PropDefault->GetAssetTagsForContext(this, FGuid(), DataInterfaces, NumericKeys, StringKeys);
 			}
 		}
 		OutTags.Add(FAssetRegistryTag("ActiveDIs", LexToString(DataInterfaces.Num()), FAssetRegistryTag::TT_Numerical));
@@ -1222,9 +1262,9 @@ bool UNiagaraSystem::HasOutstandingCompilationRequests(bool bIncludingGPUShaders
 	{
 		for (const FNiagaraEmitterHandle& EmitterHandle : GetEmitterHandles())
 		{
-			if (const UNiagaraEmitter* Emitter = EmitterHandle.GetInstance())
+			if (FVersionedNiagaraEmitterData* EmitterData = EmitterHandle.GetEmitterData())
 			{
-				if (const UNiagaraScript* GPUComputeScript = Emitter->GetGPUComputeScript())
+				if (const UNiagaraScript* GPUComputeScript = EmitterData->GetGPUComputeScript())
 				{
 					if (const FNiagaraShaderScript* ShaderScript = GPUComputeScript->GetRenderThreadScript())
 					{
@@ -1288,14 +1328,14 @@ bool UNiagaraSystem::ComputeEmitterPriority(int32 EmitterIdx, TArray<int32, TInl
 	return true;
 }
 
-void UNiagaraSystem::FindEventDependencies(UNiagaraEmitter* Emitter, TArray<UNiagaraEmitter*>& Dependencies)
+void UNiagaraSystem::FindEventDependencies(FVersionedNiagaraEmitterData* EmitterData, TArray<FVersionedNiagaraEmitter>& Dependencies)
 {
-	if (!Emitter)
+	if (!EmitterData)
 	{
 		return;
 	}
 
-	const TArray<FNiagaraEventScriptProperties>& EventHandlers = Emitter->GetEventHandlers();
+	const TArray<FNiagaraEventScriptProperties>& EventHandlers = EmitterData->GetEventHandlers();
 	for (const FNiagaraEventScriptProperties& Handler : EventHandlers)
 	{
 		// An empty ID means the event reads from the same emitter, so we don't need to record a dependency.
@@ -1312,19 +1352,18 @@ void UNiagaraSystem::FindEventDependencies(UNiagaraEmitter* Emitter, TArray<UNia
 			if (EmitterIDName.ToString() == SourceEmitterIDName)
 			{
 				// The Emitters array is in the same order as the EmitterHandles array.
-				UNiagaraEmitter* Sender = EmitterHandles[EmitterIdx].GetInstance();
-				Dependencies.Add(Sender);
+				Dependencies.Add(EmitterHandles[EmitterIdx].GetInstance());
 				break;
 			}
 		}
 	}
 }
 
-void UNiagaraSystem::FindDataInterfaceDependencies(UNiagaraEmitter* Emitter, UNiagaraScript* Script, TArray<UNiagaraEmitter*>& Dependencies)
+void UNiagaraSystem::FindDataInterfaceDependencies(FVersionedNiagaraEmitterData* EmitterData, UNiagaraScript* Script, TArray<FVersionedNiagaraEmitter>& Dependencies)
 {
-	if (const FNiagaraScriptExecutionParameterStore* ParameterStore = Script->GetExecutionReadyParameterStore(Emitter->SimTarget))
+	if (const FNiagaraScriptExecutionParameterStore* ParameterStore = Script->GetExecutionReadyParameterStore(EmitterData->SimTarget))
 	{
-		if (Emitter->SimTarget == ENiagaraSimTarget::CPUSim)
+		if (EmitterData->SimTarget == ENiagaraSimTarget::CPUSim)
 		{
 			for (UNiagaraDataInterface* DataInterface : ParameterStore->GetDataInterfaces())
 			{
@@ -1389,7 +1428,7 @@ void UNiagaraSystem::ComputeEmittersExecutionOrder()
 	EmitterPriorities.SetNum(NumEmitters);
 	EmitterDependencyGraph.Init(false, NumEmitters * NumEmitters);
 
-	TArray<UNiagaraEmitter*> EmitterDependencies;
+	TArray<FVersionedNiagaraEmitter> EmitterDependencies;
 	EmitterDependencies.Reserve(3 * NumEmitters);
 
 	RendererPostTickOrder.Reset();
@@ -1400,36 +1439,36 @@ void UNiagaraSystem::ComputeEmittersExecutionOrder()
 	for (int32 EmitterIdx = 0; EmitterIdx < NumEmitters; ++EmitterIdx)
 	{
 		const FNiagaraEmitterHandle& EmitterHandle = EmitterHandles[EmitterIdx];
-		UNiagaraEmitter* Emitter = EmitterHandle.GetInstance();
+		FVersionedNiagaraEmitterData* EmitterData = EmitterHandle.GetEmitterData();
 
 		EmitterExecutionOrder[EmitterIdx].EmitterIndex = EmitterIdx;
 		EmitterPriorities[EmitterIdx] = -1;
 
-		if (Emitter == nullptr)
+		if (EmitterData == nullptr)
 		{
 			continue;
 		}
 
 		if (!EmitterHandle.GetIsEnabled())
 		{
-			Emitter->ForEachEnabledRenderer([&] (const UNiagaraRendererProperties*) { ++SystemRendererIndex; });
+			EmitterData->ForEachEnabledRenderer([&] (const UNiagaraRendererProperties*) { ++SystemRendererIndex; });
 			continue;
 		}
 
 		EmitterDependencies.SetNum(0, false);
 
-		if (Emitter->SimTarget == ENiagaraSimTarget::GPUComputeSim && Emitter->GetGPUComputeScript())
+		if (EmitterData->SimTarget == ENiagaraSimTarget::GPUComputeSim && EmitterData->GetGPUComputeScript())
 		{
 			// GPU emitters have a combined execution context for spawn and update.
-			FindDataInterfaceDependencies(Emitter, Emitter->GetGPUComputeScript(), EmitterDependencies);
+			FindDataInterfaceDependencies(EmitterData, EmitterData->GetGPUComputeScript(), EmitterDependencies);
 		}
 		else
 		{
 			// CPU emitters have separate contexts for spawn and update, so we need to gather DIs from both. They also support events,
 			// so we need to look at the event sources for extra dependencies.
-			FindDataInterfaceDependencies(Emitter, Emitter->SpawnScriptProps.Script, EmitterDependencies);
-			FindDataInterfaceDependencies(Emitter, Emitter->UpdateScriptProps.Script, EmitterDependencies);
-			FindEventDependencies(Emitter, EmitterDependencies);
+			FindDataInterfaceDependencies(EmitterData, EmitterData->SpawnScriptProps.Script, EmitterDependencies);
+			FindDataInterfaceDependencies(EmitterData, EmitterData->UpdateScriptProps.Script, EmitterDependencies);
+			FindEventDependencies(EmitterData, EmitterDependencies);
 		}
 
 		// Map the pointers returned by the emitter to indices inside the Emitters array. This is O(N^2), but we expect
@@ -1442,8 +1481,8 @@ void UNiagaraSystem::ComputeEmittersExecutionOrder()
 			{
 				if (EmitterDependencies[DepIdx] == EmitterHandles[OtherEmitterIdx].GetInstance())
 				{
-					const bool HasSourceEmitter = EmitterHandles[EmitterIdx].GetInstance() != nullptr;
-					const bool HasDependentEmitter = EmitterHandles[OtherEmitterIdx].GetInstance() != nullptr;
+					const bool HasSourceEmitter = EmitterHandles[EmitterIdx].GetInstance().Emitter != nullptr;
+					const bool HasDependentEmitter = EmitterHandles[OtherEmitterIdx].GetInstance().Emitter != nullptr;
 
 					// check to see if the emitter we're dependent on may have been culled during the cook
 					if (HasSourceEmitter && !HasDependentEmitter)
@@ -1465,10 +1504,10 @@ void UNiagaraSystem::ComputeEmittersExecutionOrder()
 		}
 
 		// Determine renderer execution order for PostTick and Completion for any renderers that opt into it
-		for (int32 RendererIndex = 0; RendererIndex < Emitter->GetRenderers().Num(); ++RendererIndex)
+		for (int32 RendererIndex = 0; RendererIndex < EmitterData->GetRenderers().Num(); ++RendererIndex)
 		{
-			const UNiagaraRendererProperties* Renderer = Emitter->GetRenderers()[RendererIndex];
-			if (Renderer && Renderer->GetIsEnabled() && Renderer->IsSimTargetSupported(Emitter->SimTarget))
+			const UNiagaraRendererProperties* Renderer = EmitterData->GetRenderers()[RendererIndex];
+			if (Renderer && Renderer->GetIsEnabled() && Renderer->IsSimTargetSupported(EmitterData->SimTarget))
 			{
 				FNiagaraRendererExecutionIndex ExecutionIndex;
 				ExecutionIndex.EmitterIndex = EmitterIdx;
@@ -1528,7 +1567,7 @@ void UNiagaraSystem::ComputeEmittersExecutionOrder()
 	// likely been cooked out because of scalability
 	EmitterExecutionOrder.SetNum(Algo::StableRemoveIf(EmitterExecutionOrder, [this](FNiagaraEmitterExecutionIndex EmitterExecIdx)
 	{
-		return EmitterHandles[EmitterExecIdx.EmitterIndex].GetInstance() == nullptr;
+		return EmitterHandles[EmitterExecIdx.EmitterIndex].GetInstance().Emitter == nullptr;
 	}));
 }
 
@@ -1544,9 +1583,9 @@ void UNiagaraSystem::ComputeRenderersDrawOrder()
 
 	for (const FNiagaraEmitterHandle& EmitterHandle : EmitterHandles)
 	{
-		if (UNiagaraEmitter* Emitter = EmitterHandle.GetInstance())
+		if (FVersionedNiagaraEmitterData* EmitterData = EmitterHandle.GetEmitterData())
 		{
-			Emitter->ForEachEnabledRenderer(
+			EmitterData->ForEachEnabledRenderer(
 				[&](UNiagaraRendererProperties* Properties)
 				{
 					RendererSortInfo.Emplace(Properties->SortOrderHint, RendererSortInfo.Num());
@@ -1592,12 +1631,12 @@ void UNiagaraSystem::CacheFromCompiledData()
 	for (int32 i=0; i < EmitterHandles.Num(); ++i)
 	{
 		FNiagaraEmitterHandle& Handle = EmitterHandles[i];
-		UNiagaraEmitter* NiagaraEmitter = EmitterHandles[i].GetInstance();
-		if (Handle.GetIsEnabled() && NiagaraEmitter)
+		FVersionedNiagaraEmitterData* EmitterData = Handle.GetEmitterData();
+		if (Handle.GetIsEnabled() && EmitterData)
 		{
 			// Cache system instance accessors
 			ExecutionStateNameBuilder.Reset();
-			ExecutionStateNameBuilder << NiagaraEmitter->GetUniqueEmitterName();
+			ExecutionStateNameBuilder << Handle.GetInstance().Emitter->GetUniqueEmitterName();
 			ExecutionStateNameBuilder << TEXT(".ExecutionState");
 			const FName ExecutionStateName(ExecutionStateNameBuilder.ToString());
 
@@ -1614,20 +1653,20 @@ void UNiagaraSystem::CacheFromCompiledData()
 
 				DataSetCompiledData = &EmitterCompiledData[i]->DataSetCompiledData;
 
-				if (NiagaraEmitter->bLimitDeltaTime)
+				if (EmitterData->bLimitDeltaTime)
 				{
-					MaxDeltaTime = MaxDeltaTime.IsSet() ? FMath::Min(MaxDeltaTime.GetValue(), NiagaraEmitter->MaxDeltaTimePerTick) : NiagaraEmitter->MaxDeltaTimePerTick;
+					MaxDeltaTime = MaxDeltaTime.IsSet() ? FMath::Min(MaxDeltaTime.GetValue(), EmitterData->MaxDeltaTimePerTick) : EmitterData->MaxDeltaTimePerTick;
 				}
 			}
-			NiagaraEmitter->ConditionalPostLoad();
-			NiagaraEmitter->CacheFromCompiledData(DataSetCompiledData);
+			Handle.GetInstance().Emitter->ConditionalPostLoad();
+			EmitterData->CacheFromCompiledData(DataSetCompiledData, *Handle.GetInstance().Emitter);
 
 			// Allow data interfaces to cache static buffers
 			UNiagaraScript* NiagaraEmitterScripts[] =
 			{
-				NiagaraEmitter->SpawnScriptProps.Script,
-				NiagaraEmitter->UpdateScriptProps.Script,
-				NiagaraEmitter->GetGPUComputeScript(),
+				EmitterData->SpawnScriptProps.Script,
+				EmitterData->UpdateScriptProps.Script,
+				EmitterData->GetGPUComputeScript(),
 			};
 
 			for (UNiagaraScript* NiagaraScript : NiagaraEmitterScripts)
@@ -1637,8 +1676,8 @@ void UNiagaraSystem::CacheFromCompiledData()
 					continue;
 				}
 
-				const bool bUsedByCPU = NiagaraEmitter->SimTarget == ENiagaraSimTarget::CPUSim;
-				const bool bUsedByGPU = NiagaraEmitter->SimTarget == ENiagaraSimTarget::GPUComputeSim;
+				const bool bUsedByCPU = EmitterData->SimTarget == ENiagaraSimTarget::CPUSim;
+				const bool bUsedByGPU = EmitterData->SimTarget == ENiagaraSimTarget::GPUComputeSim;
 				for (const FNiagaraScriptDataInterfaceInfo& DataInterfaceInfo : NiagaraScript->GetCachedDefaultDataInterfaces())
 				{
 					if (UNiagaraDataInterface* DataInterface = DataInterfaceInfo.DataInterface)
@@ -1727,13 +1766,14 @@ void UNiagaraSystem::UpdatePostCompileDIInfo()
 
 	for (const FNiagaraEmitterHandle& EmitterHandle : GetEmitterHandles())
 	{
-		if (EmitterHandle.GetIsEnabled() == false || !EmitterHandle.GetInstance())
+		FVersionedNiagaraEmitterData* EmitterData = EmitterHandle.GetEmitterData();
+		if (EmitterHandle.GetIsEnabled() == false || EmitterData == nullptr)
 		{
 			continue;
 		}
-		if (EmitterHandle.GetInstance()->SimTarget == ENiagaraSimTarget::GPUComputeSim)
+		if (EmitterData->SimTarget == ENiagaraSimTarget::GPUComputeSim)
 		{
-			UNiagaraScript* GPUScript = EmitterHandle.GetInstance()->GetGPUComputeScript();
+			UNiagaraScript* GPUScript = EmitterData->GetGPUComputeScript();
 			if (GPUScript)
 			{
 				FNiagaraVMExecutableData& VMData = GPUScript->GetVMExecutableData();
@@ -1772,10 +1812,10 @@ void UNiagaraSystem::UpdateDITickFlags()
 	{
 		if (Handle.GetIsEnabled())
 		{
-			if (UNiagaraEmitter* Emitter = Handle.GetInstance())
+			if (FVersionedNiagaraEmitterData* EmitterData = Handle.GetEmitterData())
 			{
 				TArray<UNiagaraScript*> Scripts;
-				Emitter->GetScripts(Scripts);
+				EmitterData->GetScripts(Scripts);
 				for (UNiagaraScript* Script : Scripts)
 				{
 					CheckPostSimTick(Script);
@@ -1792,9 +1832,9 @@ void UNiagaraSystem::UpdateHasGPUEmitters()
 	{
 		if (Handle.GetIsEnabled())
 		{
-			if (UNiagaraEmitter* Emitter = Handle.GetInstance())
+			if (FVersionedNiagaraEmitterData* EmitterData = Handle.GetEmitterData())
 			{
-				bHasAnyGPUEmitters |= Emitter->SimTarget == ENiagaraSimTarget::GPUComputeSim;
+				bHasAnyGPUEmitters |= EmitterData->SimTarget == ENiagaraSimTarget::GPUComputeSim;
 			}
 		}
 	}
@@ -1820,11 +1860,11 @@ bool UNiagaraSystem::IsValidInternal() const
 
 	for (const FNiagaraEmitterHandle& Handle : EmitterHandles)
 	{
-		if ( Handle.GetIsEnabled())
+		if (Handle.GetIsEnabled())
 		{
-			if ( UNiagaraEmitter* NiagaraEmitter = Handle.GetInstance() )
+			if (FVersionedNiagaraEmitterData* EmitterData = Handle.GetEmitterData())
 			{
-				if ( !NiagaraEmitter->IsValid() && NiagaraEmitter->IsAllowedByScalability() )
+				if (!EmitterData->IsValid() && EmitterData->IsAllowedByScalability())
 				{
 					return false;
 				}
@@ -1872,15 +1912,16 @@ bool UNiagaraSystem::CanObtainUserVariable(const FNiagaraVariableBase& InVar) co
 
 #if WITH_EDITORONLY_DATA
 
-FNiagaraEmitterHandle UNiagaraSystem::AddEmitterHandle(UNiagaraEmitter& InEmitter, FName EmitterName)
+FNiagaraEmitterHandle UNiagaraSystem::AddEmitterHandle(UNiagaraEmitter& InEmitter, FName EmitterName, FGuid EmitterVersion)
 {
-	UNiagaraEmitter* NewEmitter = UNiagaraEmitter::CreateWithParentAndOwner(InEmitter, this, EmitterName, ~(RF_Public | RF_Standalone));
-	FNiagaraEmitterHandle EmitterHandle(*NewEmitter);
+	UNiagaraEmitter* NewEmitter = UNiagaraEmitter::CreateWithParentAndOwner(FVersionedNiagaraEmitter(&InEmitter, EmitterVersion), this, EmitterName, ~(RF_Public | RF_Standalone));
+	FNiagaraEmitterHandle EmitterHandle(*NewEmitter, EmitterVersion);
 	if (InEmitter.TemplateSpecification == ENiagaraScriptTemplateSpecification::Template || InEmitter.TemplateSpecification == ENiagaraScriptTemplateSpecification::Behavior)
 	{
+		NewEmitter->DisableVersioning(EmitterVersion);
 		NewEmitter->TemplateSpecification = ENiagaraScriptTemplateSpecification::None;
 		NewEmitter->TemplateAssetDescription = FText();
-		NewEmitter->RemoveParent();
+		NewEmitter->GetLatestEmitterData()->RemoveParent();
 	}
 	EmitterHandles.Add(EmitterHandle);
 	RefreshSystemParametersFromEmitter(EmitterHandle);
@@ -1895,8 +1936,8 @@ void UNiagaraSystem::AddEmitterHandleDirect(FNiagaraEmitterHandle& EmitterHandle
 
 FNiagaraEmitterHandle UNiagaraSystem::DuplicateEmitterHandle(const FNiagaraEmitterHandle& EmitterHandleToDuplicate, FName EmitterName)
 {
-	UNiagaraEmitter* DuplicateEmitter = UNiagaraEmitter::CreateAsDuplicate(*EmitterHandleToDuplicate.GetInstance(), EmitterName, *this);
-	FNiagaraEmitterHandle EmitterHandle(*DuplicateEmitter);
+	UNiagaraEmitter* DuplicateEmitter = UNiagaraEmitter::CreateAsDuplicate(*EmitterHandleToDuplicate.GetInstance().Emitter, EmitterName, *this);
+	FNiagaraEmitterHandle EmitterHandle(*DuplicateEmitter, EmitterHandleToDuplicate.GetInstance().Version);
 	EmitterHandle.SetIsEnabled(EmitterHandleToDuplicate.GetIsEnabled(), *this, false);
 	EmitterHandles.Add(EmitterHandle);
 	RefreshSystemParametersFromEmitter(EmitterHandle);
@@ -1905,7 +1946,6 @@ FNiagaraEmitterHandle UNiagaraSystem::DuplicateEmitterHandle(const FNiagaraEmitt
 
 void UNiagaraSystem::RemoveEmitterHandle(const FNiagaraEmitterHandle& EmitterHandleToDelete)
 {
-	UNiagaraEmitter* EditableEmitter = EmitterHandleToDelete.GetInstance();
 	RemoveSystemParametersForEmitter(EmitterHandleToDelete);
 	auto RemovePredicate = [&](const FNiagaraEmitterHandle& EmitterHandle) { return EmitterHandle.GetId() == EmitterHandleToDelete.GetId(); };
 	EmitterHandles.RemoveAll(RemovePredicate);
@@ -1991,9 +2031,9 @@ void UNiagaraSystem::ForceGraphToRecompileOnNextCheck()
 
 	for (FNiagaraEmitterHandle Handle : EmitterHandles)
 	{
-		if (Handle.GetInstance())
+		if (FVersionedNiagaraEmitterData* EmitterData = Handle.GetEmitterData())
 		{
-			UNiagaraScriptSourceBase* GraphSource = Handle.GetInstance()->GraphSource;
+			UNiagaraScriptSourceBase* GraphSource = EmitterData->GraphSource;
 			GraphSource->ForceGraphToRecompileOnNextCheck();
 		}
 	}
@@ -2010,9 +2050,9 @@ void UNiagaraSystem::WaitForCompilationComplete(bool bIncludingGPUShaders, bool
 	{
 		for (FNiagaraEmitterHandle& EmitterHandle : EmitterHandles)
 		{
-			if (UNiagaraEmitter* Emitter = EmitterHandle.GetInstance())
+			if (FVersionedNiagaraEmitterData* EmitterData = EmitterHandle.GetEmitterData())
 			{
-				if (UNiagaraScript* GPUComputeScript = Emitter->GetGPUComputeScript())
+				if (UNiagaraScript* GPUComputeScript = EmitterData->GetGPUComputeScript())
 				{
 					if (FNiagaraShaderScript* ShaderScript = GPUComputeScript->GetRenderThreadScript())
 					{
@@ -2253,35 +2293,36 @@ void UNiagaraSystem::EvaluateCompileResultDependencies() const
 	// Particle Sim Stages
 	for (const FNiagaraEmitterHandle& Handle : EmitterHandles)
 	{
-		if (Handle.GetInstance() && Handle.GetIsEnabled())
+		if (Handle.GetIsEnabled() && Handle.GetEmitterData())
 		{
-			UNiagaraEmitter* Emitter = Handle.GetInstance();
+			UNiagaraEmitter* Emitter = Handle.GetInstance().Emitter;
+			FVersionedNiagaraEmitterData* EmitterData = Handle.GetEmitterData();
 
-			const int32 EmitterSpawnScriptIdx = AddCompiledScriptForValidation(Emitter->EmitterSpawnScriptProps.Script, SystemSpawnVisitOrder, Emitter);
+			const int32 EmitterSpawnScriptIdx = AddCompiledScriptForValidation(EmitterData->EmitterSpawnScriptProps.Script, SystemSpawnVisitOrder, Emitter);
 			TArray<int32> EmitterSpawnVisitOrder = SystemSpawnVisitOrder;
 			TArray<int32> EmitterUpdateVisitOrder = SystemUpdateVisitOrder;
 			EmitterSpawnVisitOrder.Add(EmitterSpawnScriptIdx);
 			EmitterUpdateVisitOrder.Add(EmitterSpawnScriptIdx);
 
-			const int32 EmitterUpdateScriptIdx = AddCompiledScriptForValidation(Emitter->EmitterUpdateScriptProps.Script, EmitterUpdateVisitOrder, Emitter);
+			const int32 EmitterUpdateScriptIdx = AddCompiledScriptForValidation(EmitterData->EmitterUpdateScriptProps.Script, EmitterUpdateVisitOrder, Emitter);
 			EmitterUpdateVisitOrder.Add(EmitterUpdateScriptIdx);
 
 			TArray<int32> ParticleSpawnVisitOrder = EmitterUpdateVisitOrder;
 			TArray<int32> ParticleUpdateVisitOrder = ParticleSpawnVisitOrder;
-			const int32 ParticleSpawnScriptIdx = AddCompiledScriptForValidation(Emitter->SpawnScriptProps.Script, ParticleSpawnVisitOrder, Emitter);
+			const int32 ParticleSpawnScriptIdx = AddCompiledScriptForValidation(EmitterData->SpawnScriptProps.Script, ParticleSpawnVisitOrder, Emitter);
 			ParticleSpawnVisitOrder.Add(ParticleSpawnScriptIdx);
 			ParticleUpdateVisitOrder.Add(ParticleSpawnScriptIdx);
 
-			const int32 ParticleUpdateScriptIdx = AddCompiledScriptForValidation(Emitter->UpdateScriptProps.Script, ParticleUpdateVisitOrder, Emitter);
+			const int32 ParticleUpdateScriptIdx = AddCompiledScriptForValidation(EmitterData->UpdateScriptProps.Script, ParticleUpdateVisitOrder, Emitter);
 			ParticleUpdateVisitOrder.Add(ParticleUpdateScriptIdx);
 
 			// If the spawn script is interpolated, allow the particle update script to resolve dependencies for spawn events and spawn sim stages.
-			if (Emitter->bInterpolatedSpawning)
+			if (EmitterData->bInterpolatedSpawning)
 			{
 				ParticleSpawnVisitOrder.Add(ParticleUpdateScriptIdx);
 			}
 			
-			const TArray<FNiagaraEventScriptProperties>& EventHandlerScriptProps = Emitter->GetEventHandlers();
+			const TArray<FNiagaraEventScriptProperties>& EventHandlerScriptProps = EmitterData->GetEventHandlers();
 			for (int32 i = 0; i < EventHandlerScriptProps.Num(); i++)
 			{
 				if (EventHandlerScriptProps[i].Script)
@@ -2299,7 +2340,7 @@ void UNiagaraSystem::EvaluateCompileResultDependencies() const
 			}
 			
 			// Gather up all the spawn stages first, since they could write to a value before anything else.
-			const TArray<UNiagaraSimulationStageBase*>& SimulationStages = Emitter->GetSimulationStages();
+			const TArray<UNiagaraSimulationStageBase*>& SimulationStages = EmitterData->GetSimulationStages();
 			TArray<UNiagaraSimulationStageBase*> NonSpawnStages;
 			for (int32 i = 0; i < SimulationStages.Num(); i++)
 			{
@@ -2331,9 +2372,9 @@ void UNiagaraSystem::EvaluateCompileResultDependencies() const
 				ParticleUpdateVisitOrder.Add(LastStageIdx);
 			}
 
-			if (Emitter->SimTarget == ENiagaraSimTarget::GPUComputeSim)
+			if (EmitterData->SimTarget == ENiagaraSimTarget::GPUComputeSim)
 			{
-				AddCompiledScriptForValidation(Emitter->GetGPUComputeScript(), ParticleUpdateVisitOrder, Emitter);
+				AddCompiledScriptForValidation(EmitterData->GetGPUComputeScript(), ParticleUpdateVisitOrder, Emitter);
 			}
 		}
 	}
@@ -2564,15 +2605,15 @@ bool UNiagaraSystem::QueryCompileComplete(bool bWait, bool bDoPost, bool bDoNotA
 		{
 			for (FNiagaraEmitterHandle Handle : EmitterHandles)
 			{
-				if (Handle.GetInstance())
+				if (Handle.GetInstance().Emitter)
 				{
 					if (Handle.GetIsEnabled())
 					{
-						Handle.GetInstance()->OnPostCompile();
+						Handle.GetInstance().Emitter->OnPostCompile(Handle.GetInstance().Version);
 					}
 					else
 					{
-						Handle.GetInstance()->InvalidateCompileResults();
+						Handle.GetEmitterData()->InvalidateCompileResults();
 					}
 				}
 			}
@@ -2588,12 +2629,12 @@ bool UNiagaraSystem::QueryCompileComplete(bool bWait, bool bDoPost, bool bDoNotA
 		for (auto& AsyncTask : CompileRequest.DDCTasks)
 		{
 			FEmitterCompiledScriptPair& EmitterCompiledScriptPair = AsyncTask->ScriptPair;
-			UNiagaraEmitter* Emitter = EmitterCompiledScriptPair.Emitter;
+			FVersionedNiagaraEmitter Emitter = EmitterCompiledScriptPair.VersionedEmitter;
 			UNiagaraScript* CompiledScript = EmitterCompiledScriptPair.CompiledScript;
 
 			if (UNiagaraScript::IsEquivalentUsage(CompiledScript->GetUsage(), ENiagaraScriptUsage::ParticleUpdateScript))
 			{
-				UNiagaraScript* SpawnScript = Emitter->SpawnScriptProps.Script;
+				UNiagaraScript* SpawnScript = Emitter.GetEmitterData()->SpawnScriptProps.Script;
 				for (const FNiagaraScriptDataInterfaceInfo& UpdateDataInterfaceInfo : CompiledScript->GetCachedDefaultDataInterfaces())
 				{
 					// If the data interface isn't being written to a parameter map then it won't be bound properly so we
@@ -2664,27 +2705,27 @@ const FName UNiagaraSystem::GetEmitterVariableAliasName(const FNiagaraVariable&
 	return FName(*InEmitterVar.GetName().ToString().Replace(TEXT("Emitter."), *(InEmitter->GetUniqueEmitterName() + TEXT("."))));
 }
 
-void UNiagaraSystem::InitEmitterDataSetCompiledData(FNiagaraDataSetCompiledData& DataSetToInit, const UNiagaraEmitter* InAssociatedEmitter, const FNiagaraEmitterHandle& InAssociatedEmitterHandle)
+void UNiagaraSystem::InitEmitterDataSetCompiledData(FNiagaraDataSetCompiledData& DataSetToInit, const FNiagaraEmitterHandle& InAssociatedEmitterHandle)
 {
 	DataSetToInit.Empty();
-
-	if (InAssociatedEmitter->SimTarget == ENiagaraSimTarget::GPUComputeSim)
+	const FVersionedNiagaraEmitterData* AssociatedEmitterData = InAssociatedEmitterHandle.GetEmitterData();
+	if (AssociatedEmitterData->SimTarget == ENiagaraSimTarget::GPUComputeSim)
 	{
-		DataSetToInit.Variables = InAssociatedEmitter->GetGPUComputeScript()->GetVMExecutableData().Attributes;
+		DataSetToInit.Variables = AssociatedEmitterData->GetGPUComputeScript()->GetVMExecutableData().Attributes;
 	}
 	else
 	{
-		DataSetToInit.Variables = InAssociatedEmitter->UpdateScriptProps.Script->GetVMExecutableData().Attributes;
+		DataSetToInit.Variables = AssociatedEmitterData->UpdateScriptProps.Script->GetVMExecutableData().Attributes;
 
-		for (const FNiagaraVariable& Var : InAssociatedEmitter->SpawnScriptProps.Script->GetVMExecutableData().Attributes)
+		for (const FNiagaraVariable& Var : AssociatedEmitterData->SpawnScriptProps.Script->GetVMExecutableData().Attributes)
 		{
 			DataSetToInit.Variables.AddUnique(Var);
 		}
 	}
 
-	DataSetToInit.bRequiresPersistentIDs = InAssociatedEmitter->RequiresPersistentIDs() || DataSetToInit.Variables.Contains(SYS_PARAM_PARTICLES_ID);
+	DataSetToInit.bRequiresPersistentIDs = AssociatedEmitterData->RequiresPersistentIDs() || DataSetToInit.Variables.Contains(SYS_PARAM_PARTICLES_ID);
 	DataSetToInit.ID = FNiagaraDataSetID(InAssociatedEmitterHandle.GetIdName(), ENiagaraDataSetType::ParticleData);
-	DataSetToInit.SimTarget = InAssociatedEmitter->SimTarget;
+	DataSetToInit.SimTarget = AssociatedEmitterData->SimTarget;
 
 	DataSetToInit.BuildLayout();
 }
@@ -2693,14 +2734,14 @@ void UNiagaraSystem::PrepareRapidIterationParametersForCompilation()
 {
 	// prepare rapid iteration parameters for compilation
 	TArray<UNiagaraScript*> Scripts;
-	TMap<UNiagaraScript*, const UNiagaraEmitter*> ScriptToEmitterMap;
+	TMap<UNiagaraScript*, FVersionedNiagaraEmitter> ScriptToEmitterMap;
 	for (int32 i = 0; i < GetEmitterHandles().Num(); i++)
 	{
 		const FNiagaraEmitterHandle& Handle = GetEmitterHandle(i);
 		if (Handle.GetIsEnabled()) // Don't pull in the emitter if it isn't going to be used.
 		{
 			TArray<UNiagaraScript*> EmitterScripts;
-			Handle.GetInstance()->GetScripts(EmitterScripts, false, true);
+			Handle.GetEmitterData()->GetScripts(EmitterScripts, false, true);
 
 			for (int32 ScriptIdx = 0; ScriptIdx < EmitterScripts.Num(); ScriptIdx++)
 			{
@@ -2712,54 +2753,81 @@ void UNiagaraSystem::PrepareRapidIterationParametersForCompilation()
 			}
 		}
 	}
+	TArray<TTuple<UNiagaraScript*, FVersionedNiagaraEmitter>> ScriptsToIterate;
+	for (const auto& Entry : ScriptToEmitterMap)
+	{
+		ScriptsToIterate.Add(Entry);
+	}
     		
 	TMap<UNiagaraScript*, UNiagaraScript*> ScriptDependencyMap;
-	for (auto ScriptEmitterPair = ScriptToEmitterMap.CreateIterator(); ScriptEmitterPair; ++ScriptEmitterPair)
+	for (int Index = 0; Index < ScriptsToIterate.Num(); Index++)
 	{
-		UNiagaraScript* CompiledScript = ScriptEmitterPair->Key;
-		const UNiagaraEmitter* Emitter = ScriptEmitterPair->Value;
+		TTuple<UNiagaraScript*, FVersionedNiagaraEmitter> ScriptEmitterPair = ScriptsToIterate[Index];
+		UNiagaraScript* CompiledScript = ScriptEmitterPair.Key;
+		FVersionedNiagaraEmitter VersionedEmitter = ScriptEmitterPair.Value;
 
 		if (UNiagaraScript::IsEquivalentUsage(CompiledScript->GetUsage(), ENiagaraScriptUsage::EmitterSpawnScript))
 		{
 			Scripts.AddUnique(SystemSpawnScript);
 			ScriptDependencyMap.Add(CompiledScript, SystemSpawnScript);
-			ScriptToEmitterMap.Add(SystemSpawnScript, nullptr);
+			if (!ScriptToEmitterMap.Contains(SystemSpawnScript))
+			{
+				ScriptToEmitterMap.Add(SystemSpawnScript.Get());
+				ScriptsToIterate.Emplace(SystemSpawnScript, FVersionedNiagaraEmitter());
+			}
 		}
 
 		if (UNiagaraScript::IsEquivalentUsage(CompiledScript->GetUsage(), ENiagaraScriptUsage::EmitterUpdateScript))
 		{
 			Scripts.AddUnique(SystemUpdateScript);
 			ScriptDependencyMap.Add(CompiledScript, SystemUpdateScript);
-			ScriptToEmitterMap.Add(SystemUpdateScript, nullptr);
+			if (!ScriptToEmitterMap.Contains(SystemUpdateScript))
+			{
+				ScriptToEmitterMap.Add(SystemUpdateScript.Get());
+				ScriptsToIterate.Emplace(SystemUpdateScript, FVersionedNiagaraEmitter());
+			}
 		}
 
+		FVersionedNiagaraEmitterData* EmitterData = VersionedEmitter.GetEmitterData();
 		if (UNiagaraScript::IsEquivalentUsage(CompiledScript->GetUsage(), ENiagaraScriptUsage::ParticleSpawnScript))
 		{
-			if (Emitter && Emitter->SimTarget == ENiagaraSimTarget::GPUComputeSim)
+			if (EmitterData && EmitterData->SimTarget == ENiagaraSimTarget::GPUComputeSim)
 			{
-				UNiagaraScript* ComputeScript = const_cast<UNiagaraScript*>(Emitter->GetGPUComputeScript());
+				UNiagaraScript* ComputeScript = EmitterData->GetGPUComputeScript();
 
 				Scripts.AddUnique(ComputeScript);
 				ScriptDependencyMap.Add(CompiledScript, ComputeScript);
-				ScriptToEmitterMap.Add(ComputeScript, Emitter);
+				if (!ScriptToEmitterMap.Contains(ComputeScript))
+				{
+					ScriptToEmitterMap.Add(ComputeScript, VersionedEmitter);
+					ScriptsToIterate.Emplace(ComputeScript, VersionedEmitter);
+				}
 			}
 		}
 
 		if (UNiagaraScript::IsEquivalentUsage(CompiledScript->GetUsage(), ENiagaraScriptUsage::ParticleUpdateScript))
 		{
-			if (Emitter && Emitter->SimTarget == ENiagaraSimTarget::GPUComputeSim)
+			if (EmitterData && EmitterData->SimTarget == ENiagaraSimTarget::GPUComputeSim)
 			{
-				UNiagaraScript* ComputeScript = const_cast<UNiagaraScript*>(Emitter->GetGPUComputeScript());
+				UNiagaraScript* ComputeScript = EmitterData->GetGPUComputeScript();
 
 				Scripts.AddUnique(ComputeScript);
 				ScriptDependencyMap.Add(CompiledScript, ComputeScript);
-				ScriptToEmitterMap.Add(ComputeScript, Emitter);
+				if (!ScriptToEmitterMap.Contains(ComputeScript))
+				{
+					ScriptToEmitterMap.Add(ComputeScript, VersionedEmitter);
+					ScriptsToIterate.Emplace(ComputeScript, VersionedEmitter);
+				}
 			}
-			else if (Emitter && Emitter->bInterpolatedSpawning)
+			else if (EmitterData && EmitterData->bInterpolatedSpawning)
 			{
-				Scripts.AddUnique(Emitter->SpawnScriptProps.Script);
-				ScriptDependencyMap.Add(CompiledScript, Emitter->SpawnScriptProps.Script);
-				ScriptToEmitterMap.Add(Emitter->SpawnScriptProps.Script, Emitter);
+				Scripts.AddUnique(EmitterData->SpawnScriptProps.Script);
+				ScriptDependencyMap.Add(CompiledScript, EmitterData->SpawnScriptProps.Script);
+				if (!ScriptToEmitterMap.Contains(EmitterData->SpawnScriptProps.Script))
+				{
+					ScriptToEmitterMap.Add(EmitterData->SpawnScriptProps.Script, VersionedEmitter);
+					ScriptsToIterate.Emplace(EmitterData->SpawnScriptProps.Script, VersionedEmitter);
+				}
 			}
 		}
 	}
@@ -2843,25 +2911,25 @@ bool UNiagaraSystem::RequestCompile(bool bForce, FNiagaraSystemUpdateContext* Op
 	{
 		COOK_STAT(auto Timer = NiagaraScriptCookStats::UsageStats.TimeSyncWork());
 		COOK_STAT(Timer.TrackCyclesOnly());
-		INiagaraModule& NiagaraModule = FModuleManager::Get().LoadModuleChecked<INiagaraModule>(TEXT("Niagara"));
 
 		//Compile all emitters
-		bool bAnyUnsynchronized = false;	
 
 		// Pass one... determine if any need to be compiled.
 		{
+			bool bAnyUnsynchronized = false;
 			for (int32 i = 0; i < EmitterHandles.Num(); i++)
 			{
 				FNiagaraEmitterHandle Handle = EmitterHandles[i];
-				if (Handle.GetInstance() && Handle.GetIsEnabled())
+				FVersionedNiagaraEmitterData* EmitterData = Handle.GetEmitterData();
+				if (EmitterData && Handle.GetIsEnabled())
 				{
 					TArray<UNiagaraScript*> EmitterScripts;
-					Handle.GetInstance()->GetScripts(EmitterScripts, false, true);
+					EmitterData->GetScripts(EmitterScripts, false, true);
 					check(EmitterScripts.Num() > 0);
 					for (UNiagaraScript* EmitterScript : EmitterScripts)
 					{
 						FEmitterCompiledScriptPair Pair;
-						Pair.Emitter = Handle.GetInstance();
+						Pair.VersionedEmitter = Handle.GetInstance();
 						Pair.CompiledScript = EmitterScript;
 
 						FNiagaraVMExecutableDataId NewID;
@@ -2879,7 +2947,7 @@ bool UNiagaraSystem::RequestCompile(bool bForce, FNiagaraSystemUpdateContext* Op
 						{
 							AsyncTask->CurrentState = ENiagaraCompilationState::Finished;
 						}
-						AsyncTask->UniqueEmitterName = Handle.GetInstance()->GetUniqueEmitterName();
+						AsyncTask->UniqueEmitterName = Handle.GetInstance().Emitter->GetUniqueEmitterName();
 						ActiveCompilation.DDCTasks.Add(AsyncTask);
 					}
 				}
@@ -2890,6 +2958,7 @@ bool UNiagaraSystem::RequestCompile(bool bForce, FNiagaraSystemUpdateContext* Op
 			// Now add the system scripts for compilation...
 			{
 				FEmitterCompiledScriptPair Pair;
+				Pair.VersionedEmitter = FVersionedNiagaraEmitter();
 				Pair.CompiledScript = SystemSpawnScript;
 				
 				FNiagaraVMExecutableDataId NewID;
@@ -2912,6 +2981,7 @@ bool UNiagaraSystem::RequestCompile(bool bForce, FNiagaraSystemUpdateContext* Op
 
 			{
 				FEmitterCompiledScriptPair Pair;
+				Pair.VersionedEmitter = FVersionedNiagaraEmitter();
 				Pair.CompiledScript = SystemUpdateScript;
 				
 				FNiagaraVMExecutableDataId NewID;
@@ -2942,10 +3012,10 @@ bool UNiagaraSystem::RequestCompile(bool bForce, FNiagaraSystemUpdateContext* Op
 		for (int32 i = 0; i < EmitterHandles.Num(); i++)
 		{
 			FNiagaraEmitterHandle Handle = EmitterHandles[i];
-			if (Handle.GetInstance() && Handle.GetIsEnabled())
+			if (Handle.GetIsEnabled() && Handle.GetEmitterData())
 			{
 				TArray<UNiagaraScript*> EmitterScripts;
-				Handle.GetInstance()->GetScripts(EmitterScripts, false, true);
+				Handle.GetEmitterData()->GetScripts(EmitterScripts, false, true);
 				check(EmitterScripts.Num() > 0);
 				for (UNiagaraScript* EmitterScript : EmitterScripts)
 				{
@@ -3014,7 +3084,7 @@ void UNiagaraSystem::InitEmitterCompiledData()
 		{
 			for (int32 EmitterIdx = 0; EmitterIdx < EmitterHandles.Num(); ++EmitterIdx)
 			{
-				UNiagaraEmitter* Emitter = EmitterHandles[EmitterIdx].GetInstance();
+				UNiagaraEmitter* Emitter = EmitterHandles[EmitterIdx].GetInstance().Emitter;
 				if (Emitter)
 				{
 					FString EmitterName = Emitter->GetUniqueEmitterName() + TEXT(".");
@@ -3030,7 +3100,7 @@ void UNiagaraSystem::InitEmitterCompiledData()
 		{
 			for (int32 EmitterIdx = 0; EmitterIdx < EmitterHandles.Num(); ++EmitterIdx)
 			{
-				UNiagaraEmitter* Emitter = EmitterHandles[EmitterIdx].GetInstance();
+				UNiagaraEmitter* Emitter = EmitterHandles[EmitterIdx].GetInstance().Emitter;
 				if (Emitter)
 				{
 					FString EmitterName = Emitter->GetUniqueEmitterName() + TEXT(".");
@@ -3045,14 +3115,14 @@ void UNiagaraSystem::InitEmitterCompiledData()
 		for (int32 EmitterIdx = 0; EmitterIdx < EmitterHandles.Num(); ++EmitterIdx)
 		{
 			const FNiagaraEmitterHandle& EmitterHandle = EmitterHandles[EmitterIdx];
-			const UNiagaraEmitter* Emitter = EmitterHandle.GetInstance();
+			const UNiagaraEmitter* Emitter = EmitterHandle.GetInstance().Emitter;
 			FNiagaraDataSetCompiledData& EmitterDataSetCompiledData = NewEmitterCompiledData[EmitterIdx]->DataSetCompiledData;
 			FNiagaraDataSetCompiledData& GPUCaptureCompiledData = NewEmitterCompiledData[EmitterIdx]->GPUCaptureDataSetCompiledData;
 			if (ensureMsgf(Emitter != nullptr, TEXT("Failed to get Emitter Instance from Emitter Handle in post compile, please investigate.")))
 			{
 				static FName GPUCaptureDataSetName = TEXT("GPU Capture Dataset");
 				InitEmitterVariableAliasNames(NewEmitterCompiledData[EmitterIdx].Get(), Emitter);
-				InitEmitterDataSetCompiledData(EmitterDataSetCompiledData, Emitter, EmitterHandle);
+				InitEmitterDataSetCompiledData(EmitterDataSetCompiledData, EmitterHandle);
 				GPUCaptureCompiledData.ID = FNiagaraDataSetID(GPUCaptureDataSetName, ENiagaraDataSetType::ParticleData);
 				GPUCaptureCompiledData.Variables = EmitterDataSetCompiledData.Variables;
 				GPUCaptureCompiledData.SimTarget = ENiagaraSimTarget::CPUSim;
@@ -3120,7 +3190,7 @@ void UNiagaraSystem::InitSystemCompiledData()
 	for (int32 EmitterIdx = 0; EmitterIdx < EmitterCount; ++EmitterIdx)
 	{
 		const FNiagaraEmitterHandle& PerEmitterHandle = EmitterHandles[EmitterIdx];
-		const UNiagaraEmitter* Emitter = PerEmitterHandle.GetInstance();
+		const UNiagaraEmitter* Emitter = PerEmitterHandle.GetInstance().Emitter;
 		if (ensureMsgf(Emitter != nullptr, TEXT("Failed to get Emitter Instance from Emitter Handle when post compiling Niagara System %s!"), *GetPathNameSafe(this)))
 		{
 			const FString EmitterName = Emitter->GetUniqueEmitterName();
@@ -3337,10 +3407,11 @@ void UNiagaraSystem::GatherStaticVariables(TArray<FNiagaraVariable>& OutVars, TA
 	for (int32 i = 0; i < EmitterHandles.Num(); i++)
 	{
 		const FNiagaraEmitterHandle& Handle = EmitterHandles[i];
-		if (Handle.GetInstance() && Handle.GetIsEnabled())
+		FVersionedNiagaraEmitterData* EmitterData = Handle.GetEmitterData();
+		if (EmitterData && Handle.GetIsEnabled())
 		{
-			EmitterScripts.Add(Handle.GetInstance()->EmitterSpawnScriptProps.Script);
-			EmitterScripts.Add(Handle.GetInstance()->EmitterUpdateScriptProps.Script);
+			EmitterScripts.Add(EmitterData->EmitterSpawnScriptProps.Script);
+			EmitterScripts.Add(EmitterData->EmitterUpdateScriptProps.Script);
 		}		
 	}
 
@@ -3444,9 +3515,9 @@ void UNiagaraSystem::UpdateScalability()
 
 	for (FNiagaraEmitterHandle& Handle : EmitterHandles)
 	{
-		if (Handle.GetInstance())
+		if (UNiagaraEmitter* Emitter = Handle.GetInstance().Emitter)
 		{
-			Handle.GetInstance()->UpdateScalability();
+			Emitter->UpdateScalability();
 		}
 	}
 
@@ -3521,9 +3592,9 @@ void UNiagaraSystem::FixupPositionUserParameters()
 	// check if any of the renderers have a user param bound to a position binding
 	for (FNiagaraEmitterHandle& EmitterHandle : EmitterHandles)
 	{
-		if (UNiagaraEmitter* NiagaraEmitter = EmitterHandle.GetInstance())
+		if (FVersionedNiagaraEmitterData* EmitterData = EmitterHandle.GetEmitterData())
 		{
-			for (UNiagaraRendererProperties* Renderer : NiagaraEmitter->GetRenderers())
+			for (UNiagaraRendererProperties* Renderer : EmitterData->GetRenderers())
 			{
 				for (const FNiagaraVariableAttributeBinding* AttributeBinding : Renderer->GetAttributeBindings())
 				{
diff --git a/Engine/Plugins/FX/Niagara/Source/Niagara/Private/NiagaraSystemInstance.cpp b/Engine/Plugins/FX/Niagara/Source/Niagara/Private/NiagaraSystemInstance.cpp
--- a/Engine/Plugins/FX/Niagara/Source/Niagara/Private/NiagaraSystemInstance.cpp
+++ b/Engine/Plugins/FX/Niagara/Source/Niagara/Private/NiagaraSystemInstance.cpp
@@ -129,40 +129,45 @@ bool DoSystemDataInterfacesRequireSolo(const UNiagaraSystem& System, const FNiag
 FNiagaraSystemInstance::FNiagaraSystemInstance(UWorld& InWorld, UNiagaraSystem& InAsset, FNiagaraUserRedirectionParameterStore* InOverrideParameters,
                                                USceneComponent* InAttachComponent, ENiagaraTickBehavior InTickBehavior, bool bInPooled)
 	: SystemInstanceIndex(INDEX_NONE)
-	, SignificanceIndex(INDEX_NONE)
-	, World(&InWorld)
-	, Asset(&InAsset)
-	, OverrideParameters(InOverrideParameters)
-	, AttachComponent(InAttachComponent)
-	, TickBehavior(InTickBehavior)
-	, Age(0.0f)
-	, LastRenderTime(0.0f)
-	, TickCount(0)
-	, RandomSeed(0)
-	, RandomSeedOffset(0)
-	, LODDistance(0.0f)
-	, MaxLODDistance(FLT_MAX)
-	, CurrentFrameIndex(1)
-	, ParametersValid(false)
-	, bSolo(false)
-	, bForceSolo(false)
-	, bDataInterfacesHaveTickPrereqs(false)
-	, bDataInterfacesInitialized(false)
-	, bAlreadyBound(false)
-	, bLODDistanceIsValid(false)
-	, bLODDistanceIsOverridden(false)
-	, bPooled(bInPooled)
+	  , SignificanceIndex(INDEX_NONE)
+	  , World(&InWorld)
+	  , Asset(&InAsset)
+	  , OverrideParameters(InOverrideParameters)
+	  , AttachComponent(InAttachComponent)
+	  , TickBehavior(InTickBehavior)
+	  , Age(0.0f)
+	  , LastRenderTime(0.0f)
+	  , TickCount(0)
+	  , RandomSeed(0)
+	  , RandomSeedOffset(0)
+	  , LODDistance(0.0f)
+	  , MaxLODDistance(FLT_MAX)
+#if WITH_EDITORONLY_DATA
+	  , bWasSoloPriorToCaptureRequest(false)
+#endif
+	  , GlobalParameters{}
+	  , SystemParameters{}
+	  , CurrentFrameIndex(1)
+	  , ParametersValid(false)
+	  , bSolo(false)
+	  , bForceSolo(false)
+	  , bNotifyOnCompletion(false), bHasGPUEmitters(false), bDataInterfacesHaveTickPrereqs(false)
+	  , bDataInterfacesInitialized(false)
+	  , bAlreadyBound(false)
+	  , bLODDistanceIsValid(false)
+	  , bLODDistanceIsOverridden(false)
+	  , bPooled(bInPooled)
 #if WITH_EDITOR
-	, bNeedsUIResync(false)
+	  , bNeedsUIResync(false)
 #endif
-	, CachedDeltaSeconds(0.0f)
-	, TimeSinceLastForceUpdateTransform(0.0f)
-	, FixedBounds_GT(EForceInit::ForceInit)
-	, FixedBounds_CNC(EForceInit::ForceInit)
-	, LocalBounds(EForceInit::ForceInit)
-	, RequestedExecutionState(ENiagaraExecutionState::Complete)
-	, ActualExecutionState(ENiagaraExecutionState::Complete)
-	, FeatureLevel(GMaxRHIFeatureLevel)
+	  , CachedDeltaSeconds(0.0f)
+	  , TimeSinceLastForceUpdateTransform(0.0f)
+	  , FixedBounds_GT(ForceInit)
+	  , FixedBounds_CNC(ForceInit)
+	  , LocalBounds(ForceInit)
+	  , RequestedExecutionState(ENiagaraExecutionState::Complete)
+	  , ActualExecutionState(ENiagaraExecutionState::Complete)
+	  , FeatureLevel(GMaxRHIFeatureLevel)
 {
 	static TAtomic<uint64> IDCounter(1);
 	ID = IDCounter.IncrementExchange();
@@ -385,13 +390,14 @@ bool FNiagaraSystemInstance::RequestCapture(const FGuid& RequestId)
 	for (const FNiagaraEmitterHandle& Handle : GetSystem()->GetEmitterHandles())
 	{
 		TArray<UNiagaraScript*> Scripts;
-		if (Handle.GetInstance() && Handle.GetIsEnabled())
+		FVersionedNiagaraEmitterData* EmitterData = Handle.GetEmitterData();
+		if (EmitterData && Handle.GetIsEnabled())
 		{
-			Handle.GetInstance()->GetScripts(Scripts, false);
+			EmitterData->GetScripts(Scripts, false);
 
 			for (UNiagaraScript* Script : Scripts)
 			{
-				if (Script->IsGPUScript(Script->Usage) && Handle.GetInstance()->SimTarget == ENiagaraSimTarget::CPUSim)
+				if (Script->IsGPUScript(Script->Usage) && EmitterData->SimTarget == ENiagaraSimTarget::CPUSim)
 					continue;
 				TSharedPtr<struct FNiagaraScriptDebuggerInfo, ESPMode::ThreadSafe> DebugInfoPtr = MakeShared<FNiagaraScriptDebuggerInfo, ESPMode::ThreadSafe>(Handle.GetIdName(), Script->GetUsage(), Script->GetUsageId());
 				DebugInfoPtr->bWritten = false;
@@ -791,7 +797,7 @@ void FNiagaraSystemInstance::SetPaused(bool bInPaused)
 	}
 }
 
-void FNiagaraSystemInstance::Reset(FNiagaraSystemInstance::EResetMode Mode)
+void FNiagaraSystemInstance::Reset(EResetMode Mode)
 {
 	SCOPE_CYCLE_COUNTER(STAT_NiagaraSystemReset);
 	FScopeCycleCounterUObject AdditionalScope(GetSystem(), GET_STATID(STAT_NiagaraSystemReset));
@@ -1365,15 +1371,10 @@ bool FNiagaraSystemInstance::RequiresViewUniformBuffer() const
 
 	for (const TSharedRef<FNiagaraEmitterInstance, ESPMode::ThreadSafe>& EmitterHandle : Emitters)
 	{
-		if (FNiagaraComputeExecutionContext* GPUContext = EmitterHandle->GetGPUContext())
+		FVersionedNiagaraEmitterData* EmitterData = EmitterHandle->GetCachedEmitterData();
+		if (EmitterData && EmitterHandle->GetGPUContext() && EmitterData->RequiresViewUniformBuffer())
 		{
-			if (UNiagaraEmitter* Emitter = EmitterHandle->GetCachedEmitter())
-			{
-				if (Emitter->RequiresViewUniformBuffer())
-				{
-					return true;
-				}
-			}
+			return true;
 		}
 	}
 
@@ -1522,7 +1523,8 @@ void FNiagaraSystemInstance::InitDataInterfaces()
 			continue;
 		}
 
-		const bool bGPUSimulation = Sim.GetCachedEmitter() && (Sim.GetCachedEmitter()->SimTarget == ENiagaraSimTarget::GPUComputeSim);
+		FVersionedNiagaraEmitterData* EmitterData = Sim.GetCachedEmitterData();
+		const bool bGPUSimulation = EmitterData && (EmitterData->SimTarget == ENiagaraSimTarget::GPUComputeSim);
 
 		CalcInstDataSize(Sim.GetSpawnExecutionContext().Parameters, !bGPUSimulation, bGPUSimulation, false);
 		CalcInstDataSize(Sim.GetUpdateExecutionContext().Parameters, !bGPUSimulation, bGPUSimulation, false);
@@ -1531,7 +1533,7 @@ void FNiagaraSystemInstance::InitDataInterfaces()
 			CalcInstDataSize(Sim.GetEventExecutionContexts()[i].Parameters, !bGPUSimulation, bGPUSimulation, false);
 		}
 
-		if (Sim.GetCachedEmitter() && Sim.GetCachedEmitter()->SimTarget == ENiagaraSimTarget::GPUComputeSim && Sim.GetGPUContext())
+		if (EmitterData && EmitterData->SimTarget == ENiagaraSimTarget::GPUComputeSim && Sim.GetGPUContext())
 		{
 			CalcInstDataSize(Sim.GetGPUContext()->CombinedParamStore, !bGPUSimulation, bGPUSimulation, false);
 		}
@@ -1627,7 +1629,7 @@ void FNiagaraSystemInstance::InitDataInterfaces()
 				continue;
 			}
 
-			if (Sim.GetCachedEmitter()->SimTarget == ENiagaraSimTarget::GPUComputeSim && Sim.GetGPUContext())
+			if (Sim.GetCachedEmitterData()->SimTarget == ENiagaraSimTarget::GPUComputeSim && Sim.GetGPUContext())
 			{
 				Sim.GetGPUContext()->OptionalContexInit(this);
 			}
@@ -2034,35 +2036,6 @@ FNiagaraSystemInstance::GetEventDataSet(FName EmitterName, FName EventName) cons
 
 #if WITH_EDITORONLY_DATA
 
-bool FNiagaraSystemInstance::UsesEmitter(const UNiagaraEmitter* Emitter)const
-{
-	if (GetSystem())
-	{
-		return GetSystem()->UsesEmitter(Emitter);
-	}
-	return false;
-}
-
-bool FNiagaraSystemInstance::UsesScript(const UNiagaraScript* Script)const
-{
-	if (GetSystem())
-	{
-		for (const FNiagaraEmitterHandle& EmitterHandle : GetSystem()->GetEmitterHandles())
-		{
-			if (EmitterHandle.GetInstance() && EmitterHandle.GetInstance()->UsesScript(Script))
-			{
-				return true;
-			}
-		}
-	}
-	return false;
-}
-
-// bool FNiagaraSystemInstance::UsesDataInterface(UNiagaraDataInterface* Interface)
-// {
-//
-// }
-
 bool FNiagaraSystemInstance::UsesCollection(const UNiagaraParameterCollection* Collection)const
 {
 	if (UNiagaraSystem* System = GetSystem())
@@ -2109,7 +2082,7 @@ void FNiagaraSystemInstance::InitEmitters()
 			}
 			else if (EmitterEnabled)
 			{
-				if (const UNiagaraEmitter* EmitterAsset = EmitterHandles[EmitterIdx].GetInstance())
+				if (const FVersionedNiagaraEmitterData* EmitterAsset = EmitterHandles[EmitterIdx].GetEmitterData())
 				{
 					if (EmitterAsset->CalculateBoundsMode == ENiagaraEmitterCalculateBoundMode::Fixed)
 					{
@@ -2126,9 +2099,9 @@ void FNiagaraSystemInstance::InitEmitters()
 				// Only set bHasGPUEmitters if we allow compute shaders on the platform
 				if (bAllowComputeShaders)
 				{
-					if (const UNiagaraEmitter* Emitter = Sim->GetCachedEmitter())
+					if (const FVersionedNiagaraEmitterData* EmitterData = Sim->GetCachedEmitterData())
 					{
-						bHasGPUEmitters |= Emitter->SimTarget == ENiagaraSimTarget::GPUComputeSim;
+						bHasGPUEmitters |= EmitterData->SimTarget == ENiagaraSimTarget::GPUComputeSim;
 					}
 				}
 			}
@@ -2140,7 +2113,7 @@ void FNiagaraSystemInstance::InitEmitters()
 		}
 		else if (!LocalBounds.IsValid)
 		{
-			LocalBounds = FBox(EForceInit::ForceInitToZero);
+			LocalBounds = FBox(ForceInitToZero);
 		}
 	}
 
@@ -2463,7 +2436,8 @@ void FNiagaraSystemInstance::Tick_Concurrent(bool bEnqueueGPUTickIfNeeded)
 			Inst.Tick(CachedDeltaSeconds);
 		}
 
-		if (Inst.GetCachedEmitter() && Inst.GetCachedEmitter()->SimTarget == ENiagaraSimTarget::GPUComputeSim && !Inst.IsComplete())
+		FVersionedNiagaraEmitterData* EmitterData = Inst.GetCachedEmitterData();
+		if (EmitterData && EmitterData->SimTarget == ENiagaraSimTarget::GPUComputeSim && !Inst.IsComplete())
 		{
 			// Handle edge case where an emitter was set to inactive on the first frame by scalability
 			// Since it will not tick we should not execute a GPU tick for it, this test must be symeterical with FNiagaraGPUSystemTick::Init
@@ -2500,8 +2474,8 @@ void FNiagaraSystemInstance::Tick_Concurrent(bool bEnqueueGPUTickIfNeeded)
 	}
 	else
 	{
-		FBox NewDynamicBounds(EForceInit::ForceInit);
-		FBox NewFixedBounds(EForceInit::ForceInit);
+		FBox NewDynamicBounds(ForceInit);
+		FBox NewFixedBounds(ForceInit);
 		for (const auto& Emitter : Emitters)
 		{
 			if (Emitter->AreBoundsDynamic())
@@ -2733,7 +2707,7 @@ FBox FNiagaraSystemInstance::GetSystemFixedBounds() const
 			return NiagaraSystem->GetFixedBounds();
 		}
 	}
-	return FBox(EForceInit::ForceInit);
+	return FBox(ForceInit);
 }
 
 void FNiagaraSystemInstance::SetEmitterFixedBounds(FName EmitterName, const FBox& InLocalBounds)
@@ -2763,7 +2737,7 @@ FBox FNiagaraSystemInstance::GetEmitterFixedBounds(FName EmitterName) const
 
 	// Failed to find emitter
 	UE_LOG(LogNiagara, Warning, TEXT("GetEmitterFixedBounds: Failed to find Emitter(%s) System(%s) Component(%s)"), *EmitterName.ToString(), *GetNameSafe(Asset.Get()), *GetFullNameSafe(AttachComponent.Get()));
-	return FBox(EForceInit::ForceInit);
+	return FBox(ForceInit);
 }
 
 FNiagaraEmitterInstance* FNiagaraSystemInstance::GetEmitterByID(FGuid InID)
diff --git a/Engine/Plugins/FX/Niagara/Source/Niagara/Private/NiagaraSystemInstanceController.cpp b/Engine/Plugins/FX/Niagara/Source/Niagara/Private/NiagaraSystemInstanceController.cpp
--- a/Engine/Plugins/FX/Niagara/Source/Niagara/Private/NiagaraSystemInstanceController.cpp
+++ b/Engine/Plugins/FX/Niagara/Source/Niagara/Private/NiagaraSystemInstanceController.cpp
@@ -103,10 +103,9 @@ int32 FNiagaraSystemInstanceController::GetNumMaterials() const
 	{
 		for (TSharedPtr<FNiagaraEmitterInstance, ESPMode::ThreadSafe> EmitterInst : SystemInstance->GetEmitters())
 		{
-			if (UNiagaraEmitter* Emitter = EmitterInst->GetCachedEmitter())
+			if (FVersionedNiagaraEmitterData* EmitterData = EmitterInst->GetCachedEmitterData())
 			{
-				Emitter->ForEachEnabledRenderer(
-					[&](UNiagaraRendererProperties* Properties)
+				EmitterData->ForEachEnabledRenderer([&](UNiagaraRendererProperties* Properties)
 					{
 						Properties->GetUsedMaterials(EmitterInst.Get(), UsedMaterials);
 					}
@@ -127,9 +126,9 @@ void FNiagaraSystemInstanceController::GetUsedMaterials(TArray<UMaterialInterfac
 
 	for (TSharedRef<FNiagaraEmitterInstance, ESPMode::ThreadSafe> EmitterInst : SystemInstance->GetEmitters())
 	{
-		if (UNiagaraEmitter* Emitter = EmitterInst->GetCachedEmitter())
+		if (FVersionedNiagaraEmitterData* EmitterData = EmitterInst->GetCachedEmitterData())
 		{
-			Emitter->ForEachEnabledRenderer(
+			EmitterData->ForEachEnabledRenderer(
 				[&](UNiagaraRendererProperties* Properties)
 				{
 					bool bCreateMidsForUsedMaterials = Properties->NeedsMIDsForMaterials();
@@ -164,9 +163,9 @@ void FNiagaraSystemInstanceController::GetStreamingMeshInfo(const FBoxSphereBoun
 
 	for (TSharedRef<FNiagaraEmitterInstance, ESPMode::ThreadSafe> EmitterInst : SystemInstance->GetEmitters())
 	{
-		if (UNiagaraEmitter* Emitter = EmitterInst->GetCachedEmitter())
+		if (FVersionedNiagaraEmitterData* EmitterData = EmitterInst->GetCachedEmitter().GetEmitterData())
 		{
-			Emitter->ForEachEnabledRenderer([&](UNiagaraRendererProperties* Properties)
+			EmitterData->ForEachEnabledRenderer([&](UNiagaraRendererProperties* Properties)
 			{
 				Properties->GetStreamingMeshInfo(OwnerBounds, &EmitterInst.Get(), OutStreamingRenderAssets);
 			});
@@ -260,7 +259,7 @@ void FNiagaraSystemInstanceController::DebugDump(bool bFullDump)
 		{
 			for (TSharedRef<FNiagaraEmitterInstance, ESPMode::ThreadSafe> Emitter : SystemInstance->GetEmitters())
 			{
-				if ( Emitter->GetCachedEmitter() != nullptr )
+				if ( Emitter->GetCachedEmitter().Emitter != nullptr )
 				{
 					UE_LOG(LogNiagara, Log, TEXT("\tEmitter '%s' ExecutionState(%s) NumParticles(%d)"), *Emitter->GetEmitterHandle().GetUniqueInstanceName(), *ExecutionStateEnum->GetNameStringByIndex((int32)Emitter->GetExecutionState()), Emitter->GetNumParticles());
 				}
@@ -296,9 +295,9 @@ void FNiagaraSystemInstanceController::UpdateEmitterMaterials()
 		for (int32 i = 0; i < SystemInstance->GetEmitters().Num(); i++)
 		{
 			FNiagaraEmitterInstance* EmitterInst = &SystemInstance->GetEmitters()[i].Get();
-			if (UNiagaraEmitter* Emitter = EmitterInst->GetCachedEmitter())
+			if (FVersionedNiagaraEmitterData* EmitterData = EmitterInst->GetCachedEmitterData())
 			{
-				Emitter->ForEachEnabledRenderer(
+				EmitterData->ForEachEnabledRenderer(
 					[&](UNiagaraRendererProperties* Properties)
 					{
 						// Nothing to do if we don't create MIDs for this material
diff --git a/Engine/Plugins/FX/Niagara/Source/Niagara/Private/NiagaraSystemRenderData.cpp b/Engine/Plugins/FX/Niagara/Source/Niagara/Private/NiagaraSystemRenderData.cpp
--- a/Engine/Plugins/FX/Niagara/Source/Niagara/Private/NiagaraSystemRenderData.cpp
+++ b/Engine/Plugins/FX/Niagara/Source/Niagara/Private/NiagaraSystemRenderData.cpp
@@ -3,7 +3,6 @@
 #include "NiagaraSystemRenderData.h"
 #include "NiagaraSystemInstance.h"
 #include "NiagaraCrashReporterHandler.h"
-#include "NiagaraScalabilityManager.h"
 
 void FNiagaraSystemRenderData::ExecuteDynamicDataCommands_RenderThread(const FSetDynamicDataCommandList& Commands)
 {
@@ -155,19 +154,19 @@ void FNiagaraSystemRenderData::GenerateSetDynamicDataCommands(FSetDynamicDataCom
 	for (int32 i = 0; i < SystemInstance->GetEmitters().Num(); i++)
 	{
 		FNiagaraEmitterInstance* EmitterInst = &SystemInstance->GetEmitters()[i].Get();
-		UNiagaraEmitter* Emitter = EmitterInst->GetCachedEmitter();
+		FVersionedNiagaraEmitterData* EmitterData = EmitterInst->GetCachedEmitterData();
 
-		if (Emitter == nullptr)
+		if (EmitterData == nullptr)
 		{
 			continue;
 		}
 
 #if STATS
-		TStatId EmitterStatID = Emitter->GetStatID(true, true);
+		TStatId EmitterStatID = EmitterInst->GetCachedEmitter().Emitter->GetStatID(true, true);
 		FScopeCycleCounter EmitterStatCounter(EmitterStatID);
 #endif
 
-		Emitter->ForEachEnabledRenderer(
+		EmitterData->ForEachEnabledRenderer(
 			[&](UNiagaraRendererProperties* Properties)
 			{
 				FNiagaraRenderer* Renderer = EmitterRenderers_GT[RendererIndex];
@@ -259,9 +258,9 @@ void FNiagaraSystemRenderData::RecacheRenderers(const FNiagaraSystemInstance& Sy
 	bAnyMotionBlurEnabled = false;
 	for (TSharedRef<const FNiagaraEmitterInstance, ESPMode::ThreadSafe> EmitterInst : SystemInstance.GetEmitters())
 	{
-		if (UNiagaraEmitter* Emitter = EmitterInst->GetCachedEmitter())
+		if (FVersionedNiagaraEmitterData* EmitterData = EmitterInst->GetCachedEmitterData())
 		{
-			Emitter->ForEachEnabledRenderer(
+			EmitterData->ForEachEnabledRenderer(
 				[&](UNiagaraRendererProperties* Properties)
 			{
 				//We can skip creation of the renderer if the current quality level doesn't support it. If the quality level changes all systems are fully reinitialized.
@@ -309,11 +308,11 @@ void FNiagaraSystemRenderData::PostTickRenderers(const FNiagaraSystemInstance& S
 				const FNiagaraEmitterInstance& EmitterInst = SystemInstance.GetEmitters()[ExecIdx.EmitterIndex].Get();
 				if ( !EmitterInst.IsComplete() )
 				{
-					const UNiagaraEmitter* Emitter = EmitterInst.GetCachedEmitter();
+					FVersionedNiagaraEmitterData* EmitterData = EmitterInst.GetCachedEmitterData();
 					FNiagaraRenderer* EmitterRenderer = EmitterRenderers_GT[ExecIdx.SystemRendererIndex];
-					if (Emitter && EmitterRenderer)
+					if (EmitterData && EmitterRenderer)
 					{
-						const UNiagaraRendererProperties* RendererProperties = Emitter->GetRenderers()[ExecIdx.EmitterRendererIndex];
+						const UNiagaraRendererProperties* RendererProperties = EmitterData->GetRenderers()[ExecIdx.EmitterRendererIndex];
 						EmitterRenderer->PostSystemTick_GameThread(RendererProperties, &EmitterInst);
 					}
 				}
@@ -333,11 +332,11 @@ void FNiagaraSystemRenderData::OnSystemComplete(const FNiagaraSystemInstance& Sy
 			if (EmitterRenderers_GT.IsValidIndex(ExecIdx.SystemRendererIndex))
 			{
 				const FNiagaraEmitterInstance& EmitterInst = SystemInstance.GetEmitters()[ExecIdx.EmitterIndex].Get();
-				const UNiagaraEmitter* Emitter = EmitterInst.GetCachedEmitter();
+				FVersionedNiagaraEmitterData* EmitterData = EmitterInst.GetCachedEmitterData();
 				FNiagaraRenderer* EmitterRenderer = EmitterRenderers_GT[ExecIdx.SystemRendererIndex];
-				if (Emitter && EmitterRenderer)
+				if (EmitterData && EmitterRenderer)
 				{
-					const UNiagaraRendererProperties* RendererProperties = Emitter->GetRenderers()[ExecIdx.EmitterRendererIndex];
+					const UNiagaraRendererProperties* RendererProperties = EmitterData->GetRenderers()[ExecIdx.EmitterRendererIndex];
 					EmitterRenderer->OnSystemComplete_GameThread(RendererProperties, &EmitterInst);
 				}
 			}
diff --git a/Engine/Plugins/FX/Niagara/Source/Niagara/Private/NiagaraTypes.cpp b/Engine/Plugins/FX/Niagara/Source/Niagara/Private/NiagaraTypes.cpp
--- a/Engine/Plugins/FX/Niagara/Source/Niagara/Private/NiagaraTypes.cpp
+++ b/Engine/Plugins/FX/Niagara/Source/Niagara/Private/NiagaraTypes.cpp
@@ -4,6 +4,7 @@
 
 #include "NiagaraDataInterface.h"
 #include "NiagaraStats.h"
+#include "NiagaraEmitter.h"
 #include "Misc/StringBuilder.h"
 #include "UObject/Class.h"
 
@@ -30,6 +31,41 @@ void FNiagaraVariableMetaData::CopyUserEditableMetaData(const FNiagaraVariableMe
 	}
 }
 
+FVersionedNiagaraEmitterData* FVersionedNiagaraEmitter::GetEmitterData() const
+{
+	return Emitter ? Emitter->GetEmitterData(Version) : nullptr;
+}
+
+FVersionedNiagaraEmitterWeakPtr FVersionedNiagaraEmitter::ToWeakPtr() const
+{
+	return FVersionedNiagaraEmitterWeakPtr(Emitter, Version);
+}
+
+FVersionedNiagaraEmitterWeakPtr::FVersionedNiagaraEmitterWeakPtr(UNiagaraEmitter* InEmitter, const FGuid& InVersion)
+{
+	Emitter = InEmitter;
+	Version = InVersion;
+}
+
+FVersionedNiagaraEmitter FVersionedNiagaraEmitterWeakPtr::ResolveWeakPtr() const
+{
+	if (Emitter.IsValid())
+	{
+		return FVersionedNiagaraEmitter(Emitter.Get(), Version);
+	}
+	return FVersionedNiagaraEmitter();
+}
+
+FVersionedNiagaraEmitterData* FVersionedNiagaraEmitterWeakPtr::GetEmitterData() const
+{
+	return ResolveWeakPtr().GetEmitterData();
+}
+
+bool FVersionedNiagaraEmitterWeakPtr::IsValid() const
+{
+	return GetEmitterData() != nullptr;
+}
+
 FNiagaraLWCConverter::FNiagaraLWCConverter(FVector InSystemWorldPos)
 {
 	SystemWorldPos = InSystemWorldPos;
diff --git a/Engine/Plugins/FX/Niagara/Source/Niagara/Public/INiagaraMergeManager.h b/Engine/Plugins/FX/Niagara/Source/Niagara/Public/INiagaraMergeManager.h
--- a/Engine/Plugins/FX/Niagara/Source/Niagara/Public/INiagaraMergeManager.h
+++ b/Engine/Plugins/FX/Niagara/Source/Niagara/Public/INiagaraMergeManager.h
@@ -43,7 +43,7 @@ public:
 		}
 	};
 
-	virtual FMergeEmitterResults MergeEmitter(UNiagaraEmitter& Parent, UNiagaraEmitter* ParentAtLastMerge, UNiagaraEmitter& Instance) const = 0;
+	virtual FMergeEmitterResults MergeEmitter(const FVersionedNiagaraEmitter& Parent, const FVersionedNiagaraEmitter& ParentAtLastMerge, const FVersionedNiagaraEmitter& Instance) const = 0;
 
 	virtual void DiffEditableProperties(const void* BaseDataAddress, const void* OtherDataAddress, UStruct& Struct, TArray<FProperty*>& OutDifferentProperties) const = 0;
 
