Added versioning for Niagara emitters

diff --git a/Engine/Plugins/FX/Niagara/Source/Niagara/Public/NiagaraCommon.h b/Engine/Plugins/FX/Niagara/Source/Niagara/Public/NiagaraCommon.h
--- a/Engine/Plugins/FX/Niagara/Source/Niagara/Public/NiagaraCommon.h
+++ b/Engine/Plugins/FX/Niagara/Source/Niagara/Public/NiagaraCommon.h
@@ -15,6 +15,7 @@
 #include "NiagaraCore.h"
 #include "NiagaraCommon.generated.h"
 
+struct FVersionedNiagaraEmitter;
 class UNiagaraComponent;
 class UNiagaraSystem;
 class UNiagaraScript;
@@ -730,7 +731,7 @@ struct NIAGARA_API FNiagaraSystemUpdateContext
 	FNiagaraSystemUpdateContext(UNiagaraComponent* Component, bool bReInit, bool bInDestroyOnAdd = false, bool bInOnlyActive = false, bool bInDestroySystemSim = true) :bDestroyOnAdd(bInDestroyOnAdd), bOnlyActive(bInOnlyActive), bDestroySystemSim(bInDestroySystemSim) { Add(Component, bReInit); }
 	FNiagaraSystemUpdateContext(const UNiagaraSystem* System, bool bReInit, bool bInDestroyOnAdd = false, bool bInOnlyActive = false, bool bInDestroySystemSim = true) :bDestroyOnAdd(bInDestroyOnAdd), bOnlyActive(bInOnlyActive), bDestroySystemSim(bInDestroySystemSim) { Add(System, bReInit); }
 #if WITH_EDITORONLY_DATA
-	FNiagaraSystemUpdateContext(const UNiagaraEmitter* Emitter, bool bReInit, bool bInDestroyOnAdd = false, bool bInOnlyActive = false, bool bInDestroySystemSim = true) :bDestroyOnAdd(bInDestroyOnAdd), bOnlyActive(bInOnlyActive), bDestroySystemSim(bInDestroySystemSim) { Add(Emitter, bReInit); }
+	FNiagaraSystemUpdateContext(const FVersionedNiagaraEmitter& Emitter, bool bReInit, bool bInDestroyOnAdd = false, bool bInOnlyActive = false, bool bInDestroySystemSim = true) :bDestroyOnAdd(bInDestroyOnAdd), bOnlyActive(bInOnlyActive), bDestroySystemSim(bInDestroySystemSim) { Add(Emitter, bReInit); }
 	FNiagaraSystemUpdateContext(const UNiagaraScript* Script, bool bReInit, bool bInDestroyOnAdd = false, bool bInOnlyActive = false, bool bInDestroySystemSim = true) :bDestroyOnAdd(bInDestroyOnAdd), bOnlyActive(bInOnlyActive), bDestroySystemSim(bInDestroySystemSim) { Add(Script, bReInit); }
 	FNiagaraSystemUpdateContext(const UNiagaraParameterCollection* Collection, bool bReInit, bool bInDestroyOnAdd = false, bool bInOnlyActive = false, bool bInDestroySystemSim = true) :bDestroyOnAdd(bInDestroyOnAdd), bOnlyActive(bInOnlyActive), bDestroySystemSim(bInDestroySystemSim) { Add(Collection, bReInit); }
 #endif
@@ -746,7 +747,7 @@ struct NIAGARA_API FNiagaraSystemUpdateContext
 	void Add(UNiagaraComponent* Component, bool bReInit);
 	void Add(const UNiagaraSystem* System, bool bReInit);
 #if WITH_EDITORONLY_DATA
-	void Add(const UNiagaraEmitter* Emitter, bool bReInit);
+	void Add(const FVersionedNiagaraEmitter& Emitter, bool bReInit);
 	void Add(const UNiagaraScript* Script, bool bReInit);
 	void Add(const UNiagaraParameterCollection* Collection, bool bReInit);
 #endif
@@ -914,7 +915,7 @@ private:
  * This is then used by the SNiagaraStackRowPerfWidget to display the data in the ui.
  */
 struct NIAGARA_API FNiagaraStatDatabase
-{
+{	
 	typedef TTuple<uint64, ENiagaraScriptUsage> FStatReportKey;
 
 	/* Used by emitter and system instances to add the recorded data of a frame to this emitter's data store. */
@@ -932,10 +933,14 @@ struct NIAGARA_API FNiagaraStatDatabase
 	/* Returns the names of all captures stat data points. Useful for debugging and to dump the stat data. */
 	TMap<ENiagaraScriptUsage, TSet<FName>> GetAvailableStatNames();
 
+	void Init();
+
 private:
 	/** The captured runtime stat data. The first key is a combination of reporter handle and script usage, the second key is the stat id which correlates to a single recorded scope. */
 	TMap<FStatReportKey, TMap<TStatIdData const*, FStatExecutionTimer>> StatCaptures;
-	FCriticalSection CriticalSection;
+
+	FCriticalSection* GetCriticalSection() const;
+	TSharedPtr<FCriticalSection> CriticalSection;
 };
 #endif
 
@@ -977,12 +982,12 @@ struct FNiagaraVariableAttributeBinding
 	bool NIAGARA_API IsParticleBinding() const { return bIsCachedParticleValue; }
 	bool NIAGARA_API DoesBindingExistOnSource() const { return bBindingExistsOnSource; }
 	bool NIAGARA_API CanBindToHostParameterMap() const { return bBindingExistsOnSource && !bIsCachedParticleValue; }
-	void NIAGARA_API SetValue(const FName& InValue, const UNiagaraEmitter* InEmitter, ENiagaraRendererSourceDataMode InSourceMode);
-	void NIAGARA_API SetAsPreviousValue(const FNiagaraVariableBase& Src, const UNiagaraEmitter* InEmitter, ENiagaraRendererSourceDataMode InSourceMode);
-	void NIAGARA_API SetAsPreviousValue(const FNiagaraVariableAttributeBinding& Src, const UNiagaraEmitter* InEmitter, ENiagaraRendererSourceDataMode InSourceMode);
-	void NIAGARA_API CacheValues(const UNiagaraEmitter* InEmitter, ENiagaraRendererSourceDataMode InSourceMode);
-	bool NIAGARA_API RenameVariableIfMatching(const FNiagaraVariableBase& OldVariable, const FNiagaraVariableBase& NewVariable, const UNiagaraEmitter* InEmitter, ENiagaraRendererSourceDataMode InSourceMode);
-	bool NIAGARA_API Matches(const FNiagaraVariableBase& OldVariable, const UNiagaraEmitter* InEmitter, ENiagaraRendererSourceDataMode InSourceMode);
+	void NIAGARA_API SetValue(const FName& InValue, const FVersionedNiagaraEmitter& InEmitter, ENiagaraRendererSourceDataMode InSourceMode);
+	void NIAGARA_API SetAsPreviousValue(const FNiagaraVariableBase& Src, const FVersionedNiagaraEmitter& InEmitter, ENiagaraRendererSourceDataMode InSourceMode);
+	void NIAGARA_API SetAsPreviousValue(const FNiagaraVariableAttributeBinding& Src, const FVersionedNiagaraEmitter& InEmitter, ENiagaraRendererSourceDataMode InSourceMode);
+	void NIAGARA_API CacheValues(const FVersionedNiagaraEmitter& InEmitter, ENiagaraRendererSourceDataMode InSourceMode);
+	bool NIAGARA_API RenameVariableIfMatching(const FNiagaraVariableBase& OldVariable, const FNiagaraVariableBase& NewVariable, const FVersionedNiagaraEmitter& InEmitter, ENiagaraRendererSourceDataMode InSourceMode);
+	bool NIAGARA_API Matches(const FNiagaraVariableBase& OldVariable, const FVersionedNiagaraEmitter& InEmitter, ENiagaraRendererSourceDataMode InSourceMode);
 
 #if WITH_EDITORONLY_DATA
 	NIAGARA_API const FName& GetName() const { return CachedDisplayName; }
@@ -1006,7 +1011,7 @@ struct FNiagaraVariableAttributeBinding
 	void NIAGARA_API PostLoad(ENiagaraRendererSourceDataMode InSourceMode);
 	void NIAGARA_API Dump() const;
 
-	void NIAGARA_API ResetToDefault(const FNiagaraVariableAttributeBinding& InOther, const UNiagaraEmitter* InEmitter, ENiagaraRendererSourceDataMode InSourceMode);
+	void NIAGARA_API ResetToDefault(const FNiagaraVariableAttributeBinding& InOther, const FVersionedNiagaraEmitter& InEmitter, ENiagaraRendererSourceDataMode InSourceMode);
 	bool NIAGARA_API MatchesDefault(const FNiagaraVariableAttributeBinding& InOther, ENiagaraRendererSourceDataMode InSourceMode) const;
 protected:
 	/** The fully expressed namespace for the variable. If an emitter namespace, this will include the Emitter's unique name.*/
@@ -1260,7 +1265,7 @@ namespace FNiagaraUtilities
 	 * @param ScriptToEmitterNameMap An array of scripts to the name of the emitter than owns them.  If this is a system script the name can be empty.  All scripts in the
 	 * scripts array must have an entry in this map.
 	 */
-	void NIAGARA_API PrepareRapidIterationParameters(const TArray<UNiagaraScript*>& Scripts, const TMap<UNiagaraScript*, UNiagaraScript*>& ScriptDependencyMap, const TMap<UNiagaraScript*, const UNiagaraEmitter*>& ScriptToEmitterNameMap);
+	void NIAGARA_API PrepareRapidIterationParameters(const TArray<UNiagaraScript*>& Scripts, const TMap<UNiagaraScript*, UNiagaraScript*>& ScriptDependencyMap, const TMap<UNiagaraScript*, FVersionedNiagaraEmitter>& ScriptToEmitterNameMap);
 
 	bool NIAGARA_API AreTypesAssignable(const FNiagaraTypeDefinition& TypeA, const FNiagaraTypeDefinition& TypeB);
 #endif
diff --git a/Engine/Plugins/FX/Niagara/Source/Niagara/Public/NiagaraComponentRendererProperties.h b/Engine/Plugins/FX/Niagara/Source/Niagara/Public/NiagaraComponentRendererProperties.h
--- a/Engine/Plugins/FX/Niagara/Source/Niagara/Public/NiagaraComponentRendererProperties.h
+++ b/Engine/Plugins/FX/Niagara/Source/Niagara/Public/NiagaraComponentRendererProperties.h
@@ -5,6 +5,7 @@
 #include <atomic>
 #include "CoreMinimal.h"
 #include "NiagaraCommon.h"
+#include "NiagaraEmitter.h"
 #include "NiagaraRendererProperties.h"
 #include "Components/PointLightComponent.h"
 #include "NiagaraComponentRendererProperties.generated.h"
@@ -71,7 +72,7 @@ public:
 
 	virtual void GetRendererWidgets(const FNiagaraEmitterInstance* InEmitter, TArray<TSharedPtr<SWidget>>& OutWidgets, TSharedPtr<FAssetThumbnailPool> InThumbnailPool) const override;
 	virtual void GetRendererTooltipWidgets(const FNiagaraEmitterInstance* InEmitter, TArray<TSharedPtr<SWidget>>& OutWidgets, TSharedPtr<FAssetThumbnailPool> InThumbnailPool) const override;
-	virtual void GetRendererFeedback(UNiagaraEmitter* InEmitter, TArray<FNiagaraRendererFeedback>& OutErrors, TArray<FNiagaraRendererFeedback>& OutWarnings, TArray<FNiagaraRendererFeedback>& OutInfo) const override;
+	virtual void GetRendererFeedback(const FVersionedNiagaraEmitter& InEmitter, TArray<FNiagaraRendererFeedback>& OutErrors, TArray<FNiagaraRendererFeedback>& OutWarnings, TArray<FNiagaraRendererFeedback>& OutInfo) const override;
 	virtual const FSlateBrush* GetStackIcon() const override;
 	virtual FText GetWidgetDisplayName() const override;
 
@@ -162,7 +163,7 @@ private:
 	static TArray<TWeakObjectPtr<UNiagaraComponentRendererProperties>> ComponentRendererPropertiesToDeferredInit;
 
 	// this is just used to check for localspace when creating a new template component
-	const UNiagaraEmitter* EmitterPtr;
+	FVersionedNiagaraEmitter EmitterPtr;
 
 	void CreateTemplateComponent();
 
diff --git a/Engine/Plugins/FX/Niagara/Source/Niagara/Public/NiagaraComponentSettings.h b/Engine/Plugins/FX/Niagara/Source/Niagara/Public/NiagaraComponentSettings.h
--- a/Engine/Plugins/FX/Niagara/Source/Niagara/Public/NiagaraComponentSettings.h
+++ b/Engine/Plugins/FX/Niagara/Source/Niagara/Public/NiagaraComponentSettings.h
@@ -86,22 +86,23 @@ public:
 			{
 				EmitterRef.SystemName = ParentSystem->GetFName();
 			}
-			const UNiagaraEmitter* CachedEmitter = EmitterInstance->GetCachedEmitter();
-			EmitterRef.EmitterName = CachedEmitter->GetUniqueEmitterName();
+			FVersionedNiagaraEmitter CachedEmitter = EmitterInstance->GetCachedEmitter();
+			EmitterRef.EmitterName = CachedEmitter.Emitter->GetUniqueEmitterName();
 
 			if (bUseSuppressEmitterList && ComponentSettings->SuppressEmitterList.Contains(EmitterRef))
 			{
 				return true;
 			}
 
-			if (bUseGpuEmitterAllowList && (CachedEmitter->SimTarget == ENiagaraSimTarget::GPUComputeSim) && !ComponentSettings->GPUEmitterAllowList.Contains(EmitterRef))
+			FVersionedNiagaraEmitterData* EmitterData = CachedEmitter.GetEmitterData();
+			if (bUseGpuEmitterAllowList && (EmitterData->SimTarget == ENiagaraSimTarget::GPUComputeSim) && !ComponentSettings->GPUEmitterAllowList.Contains(EmitterRef))
 			{
 				return true;
 			}
 
-			if ( bUseGpuDataInterfaceDenyList && (CachedEmitter->SimTarget == ENiagaraSimTarget::GPUComputeSim) )
+			if ( bUseGpuDataInterfaceDenyList && (EmitterData->SimTarget == ENiagaraSimTarget::GPUComputeSim) )
 			{
-				if (const UNiagaraScript* GPUComputeScript = CachedEmitter->GetGPUComputeScript())
+				if (const UNiagaraScript* GPUComputeScript = EmitterData->GetGPUComputeScript())
 				{
 					for (const FNiagaraScriptDataInterfaceInfo& DefaultDIInfo : GPUComputeScript->GetCachedDefaultDataInterfaces())
 					{
diff --git a/Engine/Plugins/FX/Niagara/Source/Niagara/Public/NiagaraLightRendererProperties.h b/Engine/Plugins/FX/Niagara/Source/Niagara/Public/NiagaraLightRendererProperties.h
--- a/Engine/Plugins/FX/Niagara/Source/Niagara/Public/NiagaraLightRendererProperties.h
+++ b/Engine/Plugins/FX/Niagara/Source/Niagara/Public/NiagaraLightRendererProperties.h
@@ -36,7 +36,7 @@ public:
 	virtual const TArray<FNiagaraVariable>& GetOptionalAttributes() override;
 	virtual void GetRendererWidgets(const FNiagaraEmitterInstance* InEmitter, TArray<TSharedPtr<SWidget>>& OutWidgets, TSharedPtr<FAssetThumbnailPool> InThumbnailPool) const override;
 	virtual void GetRendererTooltipWidgets(const FNiagaraEmitterInstance* InEmitter, TArray<TSharedPtr<SWidget>>& OutWidgets, TSharedPtr<FAssetThumbnailPool> InThumbnailPool) const override;
-	virtual void GetRendererFeedback(const UNiagaraEmitter* InEmitter, TArray<FText>& OutErrors, TArray<FText>& OutWarnings, TArray<FText>& OutInfo) const override;
+	virtual void GetRendererFeedback(const FVersionedNiagaraEmitter& InEmitter, TArray<FText>& OutErrors, TArray<FText>& OutWarnings, TArray<FText>& OutInfo) const override;
 #endif // WITH_EDITORONLY_DATA
 	virtual void CacheFromCompiledData(const FNiagaraDataSetCompiledData* CompiledData) override;
 	//UNiagaraRendererProperties Interface END
diff --git a/Engine/Plugins/FX/Niagara/Source/Niagara/Public/NiagaraMeshRendererProperties.h b/Engine/Plugins/FX/Niagara/Source/Niagara/Public/NiagaraMeshRendererProperties.h
--- a/Engine/Plugins/FX/Niagara/Source/Niagara/Public/NiagaraMeshRendererProperties.h
+++ b/Engine/Plugins/FX/Niagara/Source/Niagara/Public/NiagaraMeshRendererProperties.h
@@ -3,9 +3,8 @@
 #pragma once
 
 #include "CoreMinimal.h"
-#include "UObject/ObjectMacros.h"
 #include "NiagaraRendererProperties.h"
-#include "StaticMeshResources.h"
+#include "NiagaraGPUSortInfo.h"
 #include "NiagaraCommon.h"
 #include "NiagaraMeshRendererProperties.generated.h"
 
@@ -166,8 +165,8 @@ public:
 	virtual void BeginDestroy() override;
 	virtual void PreEditChange(class FProperty* PropertyThatWillChange) override;
 	virtual void PostEditChangeProperty(struct FPropertyChangedEvent& PropertyChangedEvent) override;
-	virtual void RenameVariable(const FNiagaraVariableBase& OldVariable, const FNiagaraVariableBase& NewVariable, const UNiagaraEmitter* InEmitter) override;
-	virtual void RemoveVariable(const FNiagaraVariableBase& OldVariable, const UNiagaraEmitter* InEmitter) override;
+	virtual void RenameVariable(const FNiagaraVariableBase& OldVariable, const FNiagaraVariableBase& NewVariable, const FVersionedNiagaraEmitter& InEmitter) override;
+	virtual void RemoveVariable(const FNiagaraVariableBase& OldVariable, const FVersionedNiagaraEmitter& InEmitter) override;
 #endif// WITH_EDITORONLY_DATA
 	//UObject Interface END
 
@@ -185,7 +184,7 @@ public:
 	virtual void GetAdditionalVariables(TArray<FNiagaraVariableBase>& OutArray) const override;
 	virtual	void GetRendererWidgets(const FNiagaraEmitterInstance* InEmitter, TArray<TSharedPtr<SWidget>>& OutWidgets, TSharedPtr<FAssetThumbnailPool> InThumbnailPool) const override;
 	virtual	void GetRendererTooltipWidgets(const FNiagaraEmitterInstance* InEmitter, TArray<TSharedPtr<SWidget>>& OutWidgets, TSharedPtr<FAssetThumbnailPool> InThumbnailPool) const override;
-	virtual void GetRendererFeedback(const UNiagaraEmitter* InEmitter, TArray<FText>& OutErrors, TArray<FText>& OutWarnings, TArray<FText>& OutInfo) const override;
+	virtual void GetRendererFeedback(const FVersionedNiagaraEmitter& InEmitter, TArray<FText>& OutErrors, TArray<FText>& OutWarnings, TArray<FText>& OutInfo) const override;
 	void OnMeshChanged();
 	void OnMeshPostBuild(UStaticMesh*);
 	void OnAssetReimported(UObject*);
@@ -408,7 +407,7 @@ public:
 
 protected:
 	void InitBindings();
-	void SetPreviousBindings(const UNiagaraEmitter* SrcEmitter, ENiagaraRendererSourceDataMode InSourceMode);
+	void SetPreviousBindings(const FVersionedNiagaraEmitter& SrcEmitter, ENiagaraRendererSourceDataMode InSourceMode);
 	virtual void UpdateSourceModeDerivates(ENiagaraRendererSourceDataMode InSourceMode, bool bFromPropertyEdit = false) override;
 	virtual bool NeedsMIDsForMaterials() const override { return MaterialParameterBindings.Num() > 0; }	
 #if WITH_EDITORONLY_DATA
diff --git a/Engine/Plugins/FX/Niagara/Source/Niagara/Public/NiagaraModule.h b/Engine/Plugins/FX/Niagara/Source/Niagara/Public/NiagaraModule.h
--- a/Engine/Plugins/FX/Niagara/Source/Niagara/Public/NiagaraModule.h
+++ b/Engine/Plugins/FX/Niagara/Source/Niagara/Public/NiagaraModule.h
@@ -3,9 +3,7 @@
 #pragma once
 
 #include "CoreMinimal.h"
-#include "Stats/Stats.h"
 #include "Modules/ModuleInterface.h"
-#include "Engine/World.h"
 #include "NiagaraTypes.h"
 #include "Templates/SharedPointer.h"
 #include "NiagaraPerfBaseline.h"
diff --git a/Engine/Plugins/FX/Niagara/Source/Niagara/Public/NiagaraRendererComponents.h b/Engine/Plugins/FX/Niagara/Source/Niagara/Public/NiagaraRendererComponents.h
--- a/Engine/Plugins/FX/Niagara/Source/Niagara/Public/NiagaraRendererComponents.h
+++ b/Engine/Plugins/FX/Niagara/Source/Niagara/Public/NiagaraRendererComponents.h
@@ -6,6 +6,7 @@ NiagaraRendererComponents.h: Renderer for rendering Niagara particles as scene c
 #pragma once
 
 #include "NiagaraRenderer.h"
+#include "UObject/WeakFieldPtr.h"
 #include "CoreMinimal.h"
 
 class UNiagaraComponentRendererProperties;
diff --git a/Engine/Plugins/FX/Niagara/Source/Niagara/Public/NiagaraRendererMeshes.h b/Engine/Plugins/FX/Niagara/Source/Niagara/Public/NiagaraRendererMeshes.h
--- a/Engine/Plugins/FX/Niagara/Source/Niagara/Public/NiagaraRendererMeshes.h
+++ b/Engine/Plugins/FX/Niagara/Source/Niagara/Public/NiagaraRendererMeshes.h
@@ -8,6 +8,8 @@ NiagaraRenderer.h: Base class for Niagara render modules
 #include "NiagaraRenderer.h"
 #include "NiagaraMeshRendererProperties.h"
 #include "NiagaraMeshVertexFactory.h"
+#include "NiagaraGPUSortInfo.h"
+#include "StaticMeshResources.h"
 #include "NiagaraGPUSceneUtils.h"
 
 class FNiagaraDataSet;
@@ -20,7 +22,7 @@ class NIAGARA_API FNiagaraRendererMeshes : public FNiagaraRenderer
 {
 public:
 	FNiagaraRendererMeshes(ERHIFeatureLevel::Type FeatureLevel, const UNiagaraRendererProperties *InProps, const FNiagaraEmitterInstance* Emitter);
-	~FNiagaraRendererMeshes();
+	virtual ~FNiagaraRendererMeshes() override;
 
 	//FNiagaraRenderer Interface
 	virtual void Initialize(const UNiagaraRendererProperties* InProps, const FNiagaraEmitterInstance* Emitter, const FNiagaraSystemInstanceController& InController) override;
diff --git a/Engine/Plugins/FX/Niagara/Source/Niagara/Public/NiagaraRendererProperties.h b/Engine/Plugins/FX/Niagara/Source/Niagara/Public/NiagaraRendererProperties.h
--- a/Engine/Plugins/FX/Niagara/Source/Niagara/Public/NiagaraRendererProperties.h
+++ b/Engine/Plugins/FX/Niagara/Source/Niagara/Public/NiagaraRendererProperties.h
@@ -4,15 +4,14 @@
 
 #include "CoreMinimal.h"
 #include "Containers/ArrayView.h"
-#include "UObject/Object.h"
 #include "RHIDefinitions.h"
 #include "NiagaraTypes.h"
 #include "NiagaraCommon.h"
 #include "NiagaraMergeable.h"
-#include "NiagaraGPUSortInfo.h"
 #include "NiagaraPlatformSet.h"
 #include "NiagaraRendererProperties.generated.h"
 
+struct FVersionedNiagaraEmitterData;
 class FNiagaraRenderer;
 class FNiagaraSystemInstanceController;
 class UMaterial;
@@ -42,7 +41,6 @@ public:
           , SummaryText(InSummaryText)
           , FixDescription(FText())
           , Fix(FNiagaraRendererFeedbackFix())
-		  , Dismissable(false)
 	{}
 
 	FNiagaraRendererFeedback()
@@ -91,7 +89,7 @@ private:
 	FText SummaryText;
 	FText FixDescription;
 	FNiagaraRendererFeedbackFix Fix;
-	bool Dismissable;
+	bool Dismissable = false;
 };
 #endif
 
@@ -209,7 +207,7 @@ public:
 	virtual bool NeedsLoadForTargetPlatform(const class ITargetPlatform* TargetPlatform) const override;
 
 	/** Method to add asset tags that are specific to this renderer. By default we add in how many instances of this class exist in the list.*/
-	virtual void GetAssetTagsForContext(const UObject* InAsset, const TArray<const UNiagaraRendererProperties*>& InProperties, TMap<FName, uint32>& NumericKeys, TMap<FName, FString>& StringKeys) const;
+	virtual void GetAssetTagsForContext(const UObject* InAsset, FGuid AssetVersion, const TArray<const UNiagaraRendererProperties*>& InProperties, TMap<FName, uint32>& NumericKeys, TMap<FName, FString>& StringKeys) const;
 
 	/** In the case that we need parameters bound in that aren't Particle variables, these should be set up here so that the data is appropriately populated after the simulation.*/
 	virtual bool PopulateRequiredBindings(FNiagaraParameterStore& InParameterStore);
@@ -220,8 +218,8 @@ public:
 
 	/** Internal handling of any emitter variable renames. Note that this doesn't modify the renderer, the caller will need to do that if it is desired.*/
 	virtual void RenameEmitter(const FName& InOldName, const UNiagaraEmitter* InRenamedEmitter);
-	virtual void RenameVariable(const FNiagaraVariableBase& OldVariable, const FNiagaraVariableBase& NewVariable, const UNiagaraEmitter* InEmitter);
-	virtual void RemoveVariable(const FNiagaraVariableBase& OldVariable, const UNiagaraEmitter* InEmitter);
+	virtual void RenameVariable(const FNiagaraVariableBase& OldVariable, const FNiagaraVariableBase& NewVariable, const FVersionedNiagaraEmitter& InEmitter);
+	virtual void RemoveVariable(const FNiagaraVariableBase& OldVariable, const FVersionedNiagaraEmitter& InEmitter);
 	virtual bool IsMaterialValidForRenderer(UMaterial* Material, FText& InvalidMessage) { return true; }
 
 	virtual void FixMaterial(UMaterial* Material) { }
@@ -240,8 +238,8 @@ public:
 
 	virtual void GetRendererWidgets(const FNiagaraEmitterInstance* InEmitter, TArray<TSharedPtr<SWidget>>& OutWidgets, TSharedPtr<FAssetThumbnailPool> InThumbnailPool) const PURE_VIRTUAL(UNiagaraRendererProperties::GetRendererWidgets, );
 	virtual void GetRendererTooltipWidgets(const FNiagaraEmitterInstance* InEmitter, TArray<TSharedPtr<SWidget>>& OutWidgets, TSharedPtr<FAssetThumbnailPool> InThumbnailPool) const PURE_VIRTUAL(UNiagaraRendererProperties::GetRendererTooltipWidgets, );
-	virtual void GetRendererFeedback(const UNiagaraEmitter* InEmitter, TArray<FText>& OutErrors, TArray<FText>& OutWarnings, TArray<FText>& OutInfo) const {};
-	virtual void GetRendererFeedback(UNiagaraEmitter* InEmitter, TArray<FNiagaraRendererFeedback>& OutErrors, TArray<FNiagaraRendererFeedback>& OutWarnings, TArray<FNiagaraRendererFeedback>& OutInfo) const;
+	virtual void GetRendererFeedback(const FVersionedNiagaraEmitter& InEmitter, TArray<FText>& OutErrors, TArray<FText>& OutWarnings, TArray<FText>& OutInfo) const {};
+	virtual void GetRendererFeedback(const FVersionedNiagaraEmitter& InEmitter, TArray<FNiagaraRendererFeedback>& OutErrors, TArray<FNiagaraRendererFeedback>& OutWarnings, TArray<FNiagaraRendererFeedback>& OutInfo) const;
 
 	// The icon to display in the niagara stack widget under the renderer section
 	virtual const FSlateBrush* GetStackIcon() const;
@@ -275,6 +273,9 @@ public:
 	template<typename TAction>
 	void ForEachPlatformSet(TAction Func);
 
+	FVersionedNiagaraEmitterData* GetEmitterData() const;
+	FVersionedNiagaraEmitter GetOuterEmitter() const;
+
 	/** Platforms on which this renderer is enabled. */
 	UPROPERTY(EditAnywhere, Category = "Scalability", meta=(DisplayInScalabilityContext))
 	FNiagaraPlatformSet Platforms;
@@ -298,6 +299,9 @@ public:
 	UPROPERTY(EditAnywhere, Category = "Scalability")
 	bool bAllowInCullProxies;
 
+	UPROPERTY()
+	FGuid OuterEmitterVersion;
+	
 protected:
 	UPROPERTY()
 	bool bMotionBlurEnabled_DEPRECATED; // This has been rolled into MotionVectorSetting
@@ -323,5 +327,5 @@ protected:
 template<typename TAction>
 void UNiagaraRendererProperties::ForEachPlatformSet(TAction Func)
 {
-	Func(this, Platforms);
+	Func(Platforms);
 }
diff --git a/Engine/Plugins/FX/Niagara/Source/Niagara/Public/NiagaraRibbonRendererProperties.h b/Engine/Plugins/FX/Niagara/Source/Niagara/Public/NiagaraRibbonRendererProperties.h
--- a/Engine/Plugins/FX/Niagara/Source/Niagara/Public/NiagaraRibbonRendererProperties.h
+++ b/Engine/Plugins/FX/Niagara/Source/Niagara/Public/NiagaraRibbonRendererProperties.h
@@ -199,8 +199,8 @@ public:
 #if WITH_EDITOR
 	virtual void PostEditChangeProperty(struct FPropertyChangedEvent& PropertyChangedEvent) override;
 	virtual bool CanEditChange(const FProperty* InProperty) const override;
-	virtual void RenameVariable(const FNiagaraVariableBase& OldVariable, const FNiagaraVariableBase& NewVariable, const UNiagaraEmitter* InEmitter) override;
-	virtual void RemoveVariable(const FNiagaraVariableBase& OldVariable, const UNiagaraEmitter* InEmitter) override;
+	virtual void RenameVariable(const FNiagaraVariableBase& OldVariable, const FNiagaraVariableBase& NewVariable, const FVersionedNiagaraEmitter& InEmitter) override;
+	virtual void RemoveVariable(const FNiagaraVariableBase& OldVariable, const FVersionedNiagaraEmitter& InEmitter) override;
 #endif
 	//UObject Interface END
 
@@ -219,7 +219,7 @@ public:
 	virtual FNiagaraVariable GetBoundAttribute(const FNiagaraVariableAttributeBinding* Binding) const override;
 	virtual void GetRendererWidgets(const FNiagaraEmitterInstance* InEmitter, TArray<TSharedPtr<SWidget>>& OutWidgets, TSharedPtr<FAssetThumbnailPool> InThumbnailPool) const override;
 	virtual void GetRendererTooltipWidgets(const FNiagaraEmitterInstance* InEmitter, TArray<TSharedPtr<SWidget>>& OutWidgets, TSharedPtr<FAssetThumbnailPool> InThumbnailPool) const override;
-	virtual void GetRendererFeedback(const UNiagaraEmitter* InEmitter, TArray<FText>& OutErrors, TArray<FText>& OutWarnings, TArray<FText>& OutInfo) const override;
+	virtual void GetRendererFeedback(const FVersionedNiagaraEmitter& InEmitter, TArray<FText>& OutErrors, TArray<FText>& OutWarnings, TArray<FText>& OutInfo) const override;
 #endif
 	virtual void CacheFromCompiledData(const FNiagaraDataSetCompiledData* CompiledData) override;
 	
@@ -458,7 +458,7 @@ public:
 
 protected:
 	void InitBindings();
-	void SetPreviousBindings(const UNiagaraEmitter* SrcEmitter);
+	void SetPreviousBindings(const FVersionedNiagaraEmitter& SrcEmitter);
 
 	void UpdateSourceModeDerivates(ENiagaraRendererSourceDataMode InSourceMode, bool bFromPropertyEdit);
 
diff --git a/Engine/Plugins/FX/Niagara/Source/Niagara/Public/NiagaraSimulationStageBase.h b/Engine/Plugins/FX/Niagara/Source/Niagara/Public/NiagaraSimulationStageBase.h
--- a/Engine/Plugins/FX/Niagara/Source/Niagara/Public/NiagaraSimulationStageBase.h
+++ b/Engine/Plugins/FX/Niagara/Source/Niagara/Public/NiagaraSimulationStageBase.h
@@ -37,12 +37,18 @@ public:
 	uint32 bEnabled : 1;
 
 	virtual bool AppendCompileHash(FNiagaraCompileHashVisitor* InVisitor) const;
-#if WITH_EDITOR
+#if WITH_EDITORONLY_DATA
 	/** Return the FName to use in place of the default for the location in the stack context. If this would be the default, return NAME_None.*/
 	virtual FName GetStackContextReplacementName() const { return NAME_None; }
 	void SetEnabled(bool bEnabled);
 	void RequestRecompile();
 	virtual void PostEditChangeProperty(struct FPropertyChangedEvent& PropertyChangedEvent) override;
+
+	UPROPERTY()
+	FGuid OuterEmitterVersion;
+	
+	FVersionedNiagaraEmitterData* GetEmitterData() const;
+	FVersionedNiagaraEmitter GetOuterEmitter() const;
 #endif
 };
 
diff --git a/Engine/Plugins/FX/Niagara/Source/Niagara/Public/NiagaraSpriteRendererProperties.h b/Engine/Plugins/FX/Niagara/Source/Niagara/Public/NiagaraSpriteRendererProperties.h
--- a/Engine/Plugins/FX/Niagara/Source/Niagara/Public/NiagaraSpriteRendererProperties.h
+++ b/Engine/Plugins/FX/Niagara/Source/Niagara/Public/NiagaraSpriteRendererProperties.h
@@ -3,8 +3,8 @@
 #pragma once
 
 #include "CoreMinimal.h"
-#include "UObject/ObjectMacros.h"
 #include "NiagaraCommon.h"
+#include "NiagaraGPUSortInfo.h"
 #include "NiagaraRendererProperties.h"
 #include "Particles/SubUVAnimation.h"
 #include "NiagaraSpriteRendererProperties.generated.h"
@@ -111,8 +111,8 @@ public:
 	virtual void Serialize(FStructuredArchive::FRecord Record) override;
 #if WITH_EDITORONLY_DATA
 	virtual void PostEditChangeProperty(struct FPropertyChangedEvent& PropertyChangedEvent) override;
-	virtual void RenameVariable(const FNiagaraVariableBase& OldVariable, const FNiagaraVariableBase& NewVariable, const UNiagaraEmitter* InEmitter) override;
-	virtual void RemoveVariable(const FNiagaraVariableBase& OldVariable, const UNiagaraEmitter* InEmitter) override;
+	virtual void RenameVariable(const FNiagaraVariableBase& OldVariable, const FNiagaraVariableBase& NewVariable, const FVersionedNiagaraEmitter& InEmitter) override;
+	virtual void RemoveVariable(const FNiagaraVariableBase& OldVariable, const FVersionedNiagaraEmitter& InEmitter) override;
 
 #endif // WITH_EDITORONLY_DATA
 	//UObject Interface END
@@ -130,7 +130,7 @@ public:
 	virtual void GetAdditionalVariables(TArray<FNiagaraVariableBase>& OutArray) const override;
 	virtual void GetRendererWidgets(const FNiagaraEmitterInstance* InEmitter, TArray<TSharedPtr<SWidget>>& OutWidgets, TSharedPtr<FAssetThumbnailPool> InThumbnailPool) const override;
 	virtual void GetRendererTooltipWidgets(const FNiagaraEmitterInstance* InEmitter, TArray<TSharedPtr<SWidget>>& OutWidgets, TSharedPtr<FAssetThumbnailPool> InThumbnailPool) const override;
-	virtual void GetRendererFeedback(const UNiagaraEmitter* InEmitter, TArray<FText>& OutErrors, TArray<FText>& OutWarnings, TArray<FText>& OutInfo) const override;
+	virtual void GetRendererFeedback(const FVersionedNiagaraEmitter& InEmitter, TArray<FText>& OutErrors, TArray<FText>& OutWarnings, TArray<FText>& OutInfo) const override;
 #endif
 	virtual ENiagaraRendererSourceDataMode GetCurrentSourceMode() const override { return SourceMode; }
 
@@ -386,7 +386,7 @@ public:
 
 protected:
 	void InitBindings();
-	void SetPreviousBindings(const UNiagaraEmitter* SrcEmitter, ENiagaraRendererSourceDataMode InSourceMode);
+	void SetPreviousBindings(const FVersionedNiagaraEmitter& SrcEmitter, ENiagaraRendererSourceDataMode InSourceMode);
 	virtual void UpdateSourceModeDerivates(ENiagaraRendererSourceDataMode InSourceMode, bool bFromPropertyEdit = false) override;
 
 #if WITH_EDITORONLY_DATA
diff --git a/Engine/Plugins/FX/Niagara/Source/Niagara/Public/NiagaraSystemInstance.h b/Engine/Plugins/FX/Niagara/Source/Niagara/Public/NiagaraSystemInstance.h
--- a/Engine/Plugins/FX/Niagara/Source/Niagara/Public/NiagaraSystemInstance.h
+++ b/Engine/Plugins/FX/Niagara/Source/Niagara/Public/NiagaraSystemInstance.h
@@ -302,9 +302,6 @@ public:
 	void EvaluateBoundFunction(FName FunctionName, bool& UsedOnCpu, bool& UsedOnGpu) const;
 
 #if WITH_EDITORONLY_DATA
-	bool UsesEmitter(const UNiagaraEmitter* Emitter) const;
-	bool UsesScript(const UNiagaraScript* Script) const;
-	//bool UsesDataInterface(UNiagaraDataInterface* Interface);
 	bool UsesCollection(const UNiagaraParameterCollection* Collection) const;
 #endif
 
diff --git a/Engine/Plugins/FX/Niagara/Source/Niagara/Public/NiagaraTypes.h b/Engine/Plugins/FX/Niagara/Source/Niagara/Public/NiagaraTypes.h
--- a/Engine/Plugins/FX/Niagara/Source/Niagara/Public/NiagaraTypes.h
+++ b/Engine/Plugins/FX/Niagara/Source/Niagara/Public/NiagaraTypes.h
@@ -177,7 +177,7 @@ struct FNiagaraMatrix
 };
 
 USTRUCT()
-struct FNiagaraAssetVersion
+struct NIAGARA_API FNiagaraAssetVersion
 {
 	GENERATED_USTRUCT_BODY()
 
@@ -2083,4 +2083,74 @@ struct alignas(16) FNiagaraEmitterParameters
 	// todo - what else should be inserted here?  we could put an array of spawninfos/interp spawn values
 	int32 _Pad0;
 	int32 _Pad1;
-};
\ No newline at end of file
+};
+
+// Forward decl FVersionedNiagaraEmitterWeakPtr to support FVersionedNiagaraEmitter::ToWeakPtr().
+struct FVersionedNiagaraEmitterWeakPtr;
+struct FVersionedNiagaraEmitterData;
+class UNiagaraEmitter;
+
+/** Struct combining an emitter with a specific version.*/
+USTRUCT()
+struct NIAGARA_API FVersionedNiagaraEmitter
+{
+	GENERATED_USTRUCT_BODY()
+	
+	FVersionedNiagaraEmitter() {}
+	FVersionedNiagaraEmitter(UNiagaraEmitter* InEmitter, const FGuid& InVersion)
+		: Emitter(InEmitter), Version(InVersion)
+	{};
+
+	FVersionedNiagaraEmitterData* GetEmitterData() const;
+	FVersionedNiagaraEmitterWeakPtr ToWeakPtr() const;
+
+	UPROPERTY()
+	TObjectPtr<UNiagaraEmitter> Emitter = nullptr;
+
+	UPROPERTY()
+	FGuid Version;
+
+	bool operator==(const FVersionedNiagaraEmitter& Other) const
+	{
+		return Emitter == Other.Emitter && Version == Other.Version;
+	}
+
+	bool operator!=(const FVersionedNiagaraEmitter& Other) const
+	{
+		return !(*this == Other);
+	}
+};
+
+struct NIAGARA_API FVersionedNiagaraEmitterWeakPtr
+{
+	FVersionedNiagaraEmitterWeakPtr() : Emitter(nullptr) {}
+	FVersionedNiagaraEmitterWeakPtr(UNiagaraEmitter* InEmitter, const FGuid& InVersion);
+	FVersionedNiagaraEmitter ResolveWeakPtr() const;
+	FVersionedNiagaraEmitterData* GetEmitterData() const;
+	bool IsValid() const;
+	bool operator==(const FVersionedNiagaraEmitterWeakPtr& Other) const
+	{
+		return Emitter == Other.Emitter && Version == Other.Version;
+	}
+
+
+	bool operator!=(const FVersionedNiagaraEmitterWeakPtr& Other) const
+	{
+		return !(*this == Other);
+	}
+
+	TWeakObjectPtr<UNiagaraEmitter> Emitter;
+	FGuid Version;
+};
+
+inline uint32 GetTypeHash(const FVersionedNiagaraEmitterWeakPtr& Item)
+{
+	return GetTypeHash(Item.Emitter);
+}
+
+inline uint32 GetTypeHash(const FVersionedNiagaraEmitter& Item)
+{
+	return GetTypeHash(Item.Emitter) + GetTypeHash(Item.Version);
+}
+
+
diff --git a/Engine/Plugins/FX/Niagara/Source/NiagaraEditor/Private/AssetTypeActions/AssetTypeActions_NiagaraEmitter.cpp b/Engine/Plugins/FX/Niagara/Source/NiagaraEditor/Private/AssetTypeActions/AssetTypeActions_NiagaraEmitter.cpp
--- a/Engine/Plugins/FX/Niagara/Source/NiagaraEditor/Private/AssetTypeActions/AssetTypeActions_NiagaraEmitter.cpp
+++ b/Engine/Plugins/FX/Niagara/Source/NiagaraEditor/Private/AssetTypeActions/AssetTypeActions_NiagaraEmitter.cpp
@@ -8,7 +8,6 @@
 #include "NiagaraEditorStyle.h"
 #include "NiagaraSystemFactoryNew.h"
 #include "ViewModels/NiagaraSystemViewModel.h"
-#include "NiagaraSystemScriptViewModel.h"
 #include "ToolMenus.h"
 #include "Styling/AppStyle.h"
 #include "Toolkits/AssetEditorToolkit.h"
@@ -40,9 +39,12 @@ void FAssetTypeActions_NiagaraEmitter::OpenAssetEditor(const TArray<UObject*>& I
 
 	for (auto ObjIt = InObjects.CreateConstIterator(); ObjIt; ++ObjIt)
 	{
-		UNiagaraEmitter* Emitter = Cast<UNiagaraEmitter>(*ObjIt);
-		if (Emitter != nullptr)
+		if (UNiagaraEmitter* Emitter = Cast<UNiagaraEmitter>(*ObjIt))
 		{
+			if (!Emitter->VersionToOpenInEditor.IsValid())
+			{
+				Emitter->VersionToOpenInEditor = Emitter->GetExposedVersion().VersionGuid;
+			}
 			TSharedRef<FNiagaraSystemToolkit> SystemToolkit(new FNiagaraSystemToolkit());
 			SystemToolkit->InitializeWithEmitter(Mode, EditWithinLevelEditor, *Emitter);
 		}
@@ -74,7 +76,18 @@ void FAssetTypeActions_NiagaraEmitter::GetActions(const TArray<UObject*>& InObje
 		LOCTEXT("Emitter_CreateDuplicateParentTooltip", "Duplicate this emitter and set this emitter's parent to the new emitter."),
 		FSlateIcon(),
 		FUIAction(
-			FExecuteAction::CreateSP(this, &FAssetTypeActions_NiagaraEmitter::ExecuteCreateDuplicateParent, NiagaraEmitters)
+			FExecuteAction::CreateSP(this, &FAssetTypeActions_NiagaraEmitter::ExecuteCreateDuplicateParent, NiagaraEmitters),
+			FCanExecuteAction::CreateLambda([NiagaraEmitters]()
+			{
+				for (TWeakObjectPtr<UNiagaraEmitter> EmitterPtr : NiagaraEmitters)
+				{
+					if (EmitterPtr.IsValid() && EmitterPtr->IsVersioningEnabled())
+					{
+						return false;
+					}
+				}
+				return true;
+			})
 		)
 	);
 
@@ -107,7 +120,7 @@ void FAssetTypeActions_NiagaraEmitter::ExecuteNewNiagaraSystem(TArray<TWeakObjec
 
 			// Create the factory used to generate the asset
 			UNiagaraSystemFactoryNew* Factory = NewObject<UNiagaraSystemFactoryNew>();
-			Factory->EmittersToAddToNewSystem.Add(Emitter.Get());
+			Factory->EmittersToAddToNewSystem.Add(FVersionedNiagaraEmitter(Emitter.Get(), Emitter->GetExposedVersion().VersionGuid));
 			FAssetToolsModule& AssetToolsModule = FModuleManager::GetModuleChecked<FAssetToolsModule>("AssetTools");
 			UObject* NewAsset = AssetToolsModule.Get().CreateAsset(Name, FPackageName::GetLongPackagePath(PackageName), UNiagaraSystem::StaticClass(), Factory);
 			
@@ -150,11 +163,11 @@ void FAssetTypeActions_NiagaraEmitter::ExecuteCreateDuplicateParent(TArray<TWeak
 				GEditor->GetEditorSubsystem<UAssetEditorSubsystem>()->CloseAllEditorsForAsset(Emitter);
 
 				NewEmitter->Modify();
-				NewEmitter->GraphSource->MarkNotSynchronized(TEXT("Emitter created"));
+				NewEmitter->GetLatestEmitterData()->GraphSource->MarkNotSynchronized(TEXT("Emitter created"));
 
 				Emitter->Modify();
-				Emitter->SetParent(*NewEmitter);
-				Emitter->GraphSource->MarkNotSynchronized(TEXT("Emitter parent changed"));
+				Emitter->SetParent(FVersionedNiagaraEmitter(NewEmitter, NewEmitter->GetExposedVersion().VersionGuid));
+				Emitter->GetLatestEmitterData()->GraphSource->MarkNotSynchronized(TEXT("Emitter parent changed"));
 
 				ObjectsToSync.Add(NewEmitter);
 			}
diff --git a/Engine/Plugins/FX/Niagara/Source/NiagaraEditor/Private/Commandlets/NiagaraDumpBytecodeCommandlet.cpp b/Engine/Plugins/FX/Niagara/Source/NiagaraEditor/Private/Commandlets/NiagaraDumpBytecodeCommandlet.cpp
--- a/Engine/Plugins/FX/Niagara/Source/NiagaraEditor/Private/Commandlets/NiagaraDumpBytecodeCommandlet.cpp
+++ b/Engine/Plugins/FX/Niagara/Source/NiagaraEditor/Private/Commandlets/NiagaraDumpBytecodeCommandlet.cpp
@@ -105,7 +105,6 @@ void UNiagaraDumpByteCodeCommandlet::ProcessNiagaraScripts()
 	//  Iterate over all scripts
 	const FString DevelopersFolder = FPackageName::FilenameToLongPackageName(FPaths::GameDevelopersDir().LeftChop(1));
 	FString LastPackageName = TEXT("");
-	int32 PackageSwitches = 0;
 	UPackage* CurrentPackage = nullptr;
 	for (const FAssetData& AssetIt : AssetList)
 	{
@@ -175,7 +174,7 @@ void UNiagaraDumpByteCodeCommandlet::ProcessNiagaraScripts()
 				continue;
 			}
 
-			if (UNiagaraEmitter* Emitter = EmitterHandle.GetInstance())
+			if (FVersionedNiagaraEmitterData* Emitter = EmitterHandle.GetEmitterData())
 			{
 				const FString EmitterName = EmitterHandle.GetUniqueInstanceName();
 
@@ -186,7 +185,7 @@ void UNiagaraDumpByteCodeCommandlet::ProcessNiagaraScripts()
 
 				for (const auto* EmitterScript : EmitterScripts)
 				{
-					DumpByteCode(EmitterScript, ShorterSystemName / EmitterName / UsageEnum->GetNameStringByValue((int64)EmitterScript->GetUsage()) + TEXT(".bin"));
+					DumpByteCode(EmitterScript, ShorterSystemName / EmitterName / UsageEnum->GetNameStringByValue(static_cast<int64>(EmitterScript->GetUsage())) + TEXT(".bin"));
 				}
 			}
 		}
diff --git a/Engine/Plugins/FX/Niagara/Source/NiagaraEditor/Private/Commandlets/NiagaraSystemAuditCommandlet.cpp b/Engine/Plugins/FX/Niagara/Source/NiagaraEditor/Private/Commandlets/NiagaraSystemAuditCommandlet.cpp
--- a/Engine/Plugins/FX/Niagara/Source/NiagaraEditor/Private/Commandlets/NiagaraSystemAuditCommandlet.cpp
+++ b/Engine/Plugins/FX/Niagara/Source/NiagaraEditor/Private/Commandlets/NiagaraSystemAuditCommandlet.cpp
@@ -139,7 +139,6 @@ bool UNiagaraSystemAuditCommandlet::ProcessNiagaraSystems()
 	//  Iterate over all systems
 	const FString DevelopersFolder = FPackageName::FilenameToLongPackageName(FPaths::GameDevelopersDir().LeftChop(1));
 	FString LastPackageName = TEXT("");
-	int32 PackageSwitches = 0;
 	UPackage* CurrentPackage = nullptr;
 	TArray<UPackage*> PackagesToSave;
 
@@ -207,8 +206,8 @@ bool UNiagaraSystemAuditCommandlet::ProcessNiagaraSystems()
 
 		for (const FNiagaraEmitterHandle& EmitterHandle : NiagaraSystem->GetEmitterHandles())
 		{
-			UNiagaraEmitter* NiagaraEmitter = EmitterHandle.GetInstance();
-			if (NiagaraEmitter == nullptr)
+			FVersionedNiagaraEmitterData* EmitterData = EmitterHandle.GetEmitterData();
+			if (EmitterData == nullptr)
 			{
 				continue;
 			}
@@ -218,39 +217,39 @@ bool UNiagaraSystemAuditCommandlet::ProcessNiagaraSystems()
 				continue;
 			}
 
-			if (NiagaraEmitter->SimTarget == ENiagaraSimTarget::GPUComputeSim)
+			if (EmitterData->SimTarget == ENiagaraSimTarget::GPUComputeSim)
 			{
 				// Optionally disable GPU emitters
 				for (UDeviceProfile* DeviceProfile : DeviceProfilesToDisableGpu)
 				{
-					const int32 DeviceQualityLevelMask = NiagaraEmitter->Platforms.IsEnabledForDeviceProfile(DeviceProfile);
+					const int32 DeviceQualityLevelMask = EmitterData->Platforms.IsEnabledForDeviceProfile(DeviceProfile);
 					if (DeviceQualityLevelMask != 0)
 					{
 						for (int32 iQualityLevel = 0; iQualityLevel < NiagaraSettings->QualityLevels.Num(); ++iQualityLevel)
 						{
 							if ( (DeviceQualityLevelMask & (1 << iQualityLevel)) != 0 )
 							{
-								NiagaraEmitter->Platforms.SetDeviceProfileState(DeviceProfile, iQualityLevel, ENiagaraPlatformSelectionState::Disabled);
+								EmitterData->Platforms.SetDeviceProfileState(DeviceProfile, iQualityLevel, ENiagaraPlatformSelectionState::Disabled);
 							}
 						}
 
 						PackagesToSave.AddUnique(CurrentPackage);
 						CurrentPackage->SetDirtyFlag(true);
 
-						UE_LOG(LogNiagaraSystemAuditCommandlet, Log, TEXT("Disabling Emitter %s for System %s Device Quality Level Mask 0x%08x"), *GetNameSafe(NiagaraEmitter), *GetNameSafe(NiagaraSystem), DeviceQualityLevelMask);
+						UE_LOG(LogNiagaraSystemAuditCommandlet, Log, TEXT("Disabling Emitter %s for System %s Device Quality Level Mask 0x%08x"), *EmitterHandle.GetUniqueInstanceName(), *GetNameSafe(NiagaraSystem), DeviceQualityLevelMask);
 					}
 				}
 
 				// Build information to write out
 				TStringBuilder<512> GpuEmitterBuilder;
-				GpuEmitterBuilder.Append(NiagaraEmitter->GetDebugSimName());
+				GpuEmitterBuilder.Append(EmitterData->GetDebugSimName());
 				for (int32 iQualityLevel=0; iQualityLevel < NiagaraSettings->QualityLevels.Num(); ++iQualityLevel)
 				{
-					const bool bEnabled = NiagaraEmitter->Platforms.IsEffectQualityEnabled(iQualityLevel);
+					const bool bEnabled = EmitterData->Platforms.IsEffectQualityEnabled(iQualityLevel);
 
 					TArray<UDeviceProfile*> EnabledProfiles;
 					TArray<UDeviceProfile*> DisabledProfiles;
-					NiagaraEmitter->Platforms.GetOverridenDeviceProfiles(iQualityLevel, EnabledProfiles, DisabledProfiles);
+					EmitterData->Platforms.GetOverridenDeviceProfiles(iQualityLevel, EnabledProfiles, DisabledProfiles);
 
 					GpuEmitterBuilder.Append(TEXT(","));
 					GpuEmitterBuilder.Append(bEnabled ? TEXT("Enabled") : TEXT("Disabled"));
@@ -266,7 +265,7 @@ bool UNiagaraSystemAuditCommandlet::ProcessNiagaraSystems()
 				}
 
 				GpuEmitterBuilder.Append(TEXT(","));
-				for (const FNiagaraPlatformSetCVarCondition& Condition : NiagaraEmitter->Platforms.CVarConditions)
+				for (const FNiagaraPlatformSetCVarCondition& Condition : EmitterData->Platforms.CVarConditions)
 				{
 					GpuEmitterBuilder.Appendf(TEXT(" CVarName(%s)"), *Condition.CVarName.ToString());
 				}
@@ -276,9 +275,9 @@ bool UNiagaraSystemAuditCommandlet::ProcessNiagaraSystems()
 				NiagaraSystemsWithGPUEmitters.Add(GpuEmitterBuilder.ToString());
 			}
 
-			bHasEvents |= NiagaraEmitter->GetEventHandlers().Num() > 0;
+			bHasEvents |= EmitterData->GetEventHandlers().Num() > 0;
 
-			for (UNiagaraRendererProperties* RendererProperties : NiagaraEmitter->GetRenderers())
+			for (UNiagaraRendererProperties* RendererProperties : EmitterData->GetRenderers())
 			{
 				if (UNiagaraLightRendererProperties* LightRendererProperties = Cast<UNiagaraLightRendererProperties>(RendererProperties))
 				{
@@ -292,7 +291,7 @@ bool UNiagaraSystemAuditCommandlet::ProcessNiagaraSystems()
 						static UEnum* NiagaraRibbonTessellationModeEnum = StaticEnum<ENiagaraRibbonTessellationMode>();
 
 						TStringBuilder<512> RendererBuilder;
-						RendererBuilder.Append(NiagaraEmitter->GetPathName());
+						RendererBuilder.Append(EmitterHandle.GetInstance().Emitter->GetPathName());
 						RendererBuilder.Append(TEXT(","));
 						RendererBuilder.Append(NiagaraRibbonTessellationModeEnum->GetValueAsString(RibbonRendererProperties->TessellationMode));
 						NiagaraRibbonRenderers.Add(RendererBuilder.ToString());
@@ -300,9 +299,9 @@ bool UNiagaraSystemAuditCommandlet::ProcessNiagaraSystems()
 				}
 			}
 
-			if ( (NiagaraEmitter->CalculateBoundsMode == ENiagaraEmitterCalculateBoundMode::Dynamic) && !NiagaraSystem->bFixedBounds )
+			if ( (EmitterData->CalculateBoundsMode == ENiagaraEmitterCalculateBoundMode::Dynamic) && !NiagaraSystem->bFixedBounds )
 			{
-				EmittersWithDynamicBounds.Append(NiagaraEmitter->GetDebugSimName());
+				EmittersWithDynamicBounds.Append(EmitterData->GetDebugSimName());
 			}
 		}
 
@@ -496,14 +495,14 @@ TArray<class UNiagaraDataInterface*> UNiagaraSystemAuditCommandlet::GetDataInter
 
 	for (const FNiagaraEmitterHandle& EmitterHandle : NiagaraSystem->GetEmitterHandles())
 	{
-		UNiagaraEmitter* NiagaraEmitter = EmitterHandle.GetInstance();
-		if (NiagaraEmitter == nullptr)
+		FVersionedNiagaraEmitterData* EmitterData = EmitterHandle.GetEmitterData();
+		if (EmitterData == nullptr)
 		{
 			continue;
 		}
 
 		TArray<UNiagaraScript*> EmitterScripts;
-		NiagaraEmitter->GetScripts(EmitterScripts);
+		EmitterData->GetScripts(EmitterScripts);
 		for (UNiagaraScript* Script : EmitterScripts)
 		{
 			GatherScriptDIs(Script);
diff --git a/Engine/Plugins/FX/Niagara/Source/NiagaraEditor/Private/Customizations/NiagaraComponentRendererPropertiesDetails.cpp b/Engine/Plugins/FX/Niagara/Source/NiagaraEditor/Private/Customizations/NiagaraComponentRendererPropertiesDetails.cpp
--- a/Engine/Plugins/FX/Niagara/Source/NiagaraEditor/Private/Customizations/NiagaraComponentRendererPropertiesDetails.cpp
+++ b/Engine/Plugins/FX/Niagara/Source/NiagaraEditor/Private/Customizations/NiagaraComponentRendererPropertiesDetails.cpp
@@ -148,10 +148,9 @@ void FNiagaraComponentRendererPropertiesDetails::CustomizeDetails(IDetailLayoutB
 
 			TSharedPtr<SWidget> NameWidget;
 			TSharedPtr<SWidget> ValueWidget;
-			const FNiagaraComponentPropertyBinding* Binding = FindBinding(PropHandle);
 
 			// if we have a binding override, we display a different value ui for the property
-			if (Binding)
+			if (const FNiagaraComponentPropertyBinding* Binding = FindBinding(PropHandle))
 			{
 				FDetailWidgetRow& WidgetRow = PropertyRow.CustomWidget(false);
 				PropertyRow.GetDefaultWidgets(NameWidget, ValueWidget, false);
@@ -281,8 +280,7 @@ TSharedRef<SWidget> FNiagaraComponentRendererPropertiesDetails::GetAddBindingMen
 
 void FNiagaraComponentRendererPropertiesDetails::ChangePropertyBinding(TSharedPtr<IPropertyHandle> PropertyHandle, const FNiagaraVariable& BindingVar)
 {
-	UNiagaraComponentRendererProperties* ComponentProperties = RendererProperties.Get();
-	if (ComponentProperties)
+	if (UNiagaraComponentRendererProperties* ComponentProperties = RendererProperties.Get())
 	{
 		FScopedTransaction Transaction(FText::Format(LOCTEXT("ChangePropertyBinding", "Change component property binding to \"{0}\" "), FText::FromName(BindingVar.GetName())));
 		FNiagaraComponentPropertyBinding NewBinding = ToPropertyBinding(PropertyHandle, ComponentProperties);
@@ -291,8 +289,7 @@ void FNiagaraComponentRendererPropertiesDetails::ChangePropertyBinding(TSharedPt
 		ComponentProperties->Modify();
 		PropertyHandle->NotifyPreChange();
 
-		const FNiagaraComponentPropertyBinding* PropertyBinding = FindBinding(PropertyHandle);
-		if (PropertyBinding)
+		if (const FNiagaraComponentPropertyBinding* PropertyBinding = FindBinding(PropertyHandle))
 		{
 			// this is an update, so remove the preexisting binding
 			for (int i = ComponentProperties->PropertyBindings.Num() - 1; i >= 0; i--)
@@ -361,8 +358,7 @@ FReply FNiagaraComponentRendererPropertiesDetails::ResetBindingButtonPressed(TSh
 void FNiagaraComponentRendererPropertiesDetails::RefreshPropertiesPanel()
 {
 	// Raw because we don't want to keep alive the details builder when calling the force refresh details
-	IDetailLayoutBuilder* DetailLayoutBuilder = DetailBuilderWeakPtr.Pin().Get();
-	if (DetailLayoutBuilder)
+	if (IDetailLayoutBuilder* DetailLayoutBuilder = DetailBuilderWeakPtr.Pin().Get())
 	{
 		DetailLayoutBuilder->ForceRefreshDetails();
 	}
@@ -370,11 +366,9 @@ void FNiagaraComponentRendererPropertiesDetails::RefreshPropertiesPanel()
 
 FText FNiagaraComponentRendererPropertiesDetails::GetCurrentBindingText(TSharedPtr<IPropertyHandle> PropertyHandle) const
 {
-	UNiagaraComponentRendererProperties* ComponentProperties = RendererProperties.Get();
-	if (ComponentProperties)
+	if (UNiagaraComponentRendererProperties* ComponentProperties = RendererProperties.Get())
 	{
-		const FNiagaraComponentPropertyBinding* PropertyBinding = FindBinding(PropertyHandle);
-		if (PropertyBinding)
+		if (const FNiagaraComponentPropertyBinding* PropertyBinding = FindBinding(PropertyHandle))
 		{
 			return FText::FromName(PropertyBinding->AttributeBinding.GetName());
 		}
@@ -382,26 +376,25 @@ FText FNiagaraComponentRendererPropertiesDetails::GetCurrentBindingText(TSharedP
 	return FText::FromString(TEXT("Missing"));
 }
 
-const UNiagaraEmitter* FNiagaraComponentRendererPropertiesDetails::GetCurrentEmitter() const
+FVersionedNiagaraEmitter FNiagaraComponentRendererPropertiesDetails::GetCurrentEmitter() const
 {
-	UNiagaraComponentRendererProperties* ComponentProperties = RendererProperties.Get();
-	if (ComponentProperties)
+	if (UNiagaraComponentRendererProperties* ComponentProperties = RendererProperties.Get())
 	{
-		return Cast<UNiagaraEmitter>(ComponentProperties->GetOuter());
+		return ComponentProperties->GetOuterEmitter();
 	}
-	return nullptr;
+	return FVersionedNiagaraEmitter();
 }
 
 TArray<FNiagaraVariable> FNiagaraComponentRendererPropertiesDetails::GetPossibleBindings(TSharedPtr<IPropertyHandle> PropertyHandle) const
 {
 	const FNiagaraTypeDefinition& PropertyType = UNiagaraComponentRendererProperties::ToNiagaraType(PropertyHandle->GetProperty());
 
-	const UNiagaraEmitter* Emitter = GetCurrentEmitter();
-	if (!Emitter || !PropertyType.IsValid())
+	FVersionedNiagaraEmitterData* EmitterData = GetCurrentEmitter().GetEmitterData();
+	if (!EmitterData || !PropertyType.IsValid())
 	{
 		return TArray<FNiagaraVariable>();
 	}
-	UNiagaraScriptSource* Source = Cast<UNiagaraScriptSource>(Emitter->GraphSource);
+	UNiagaraScriptSource* Source = Cast<UNiagaraScriptSource>(EmitterData->GraphSource);
 	if (!Source)
 	{
 		return TArray<FNiagaraVariable>();
diff --git a/Engine/Plugins/FX/Niagara/Source/NiagaraEditor/Private/Customizations/NiagaraComponentRendererPropertiesDetails.h b/Engine/Plugins/FX/Niagara/Source/NiagaraEditor/Private/Customizations/NiagaraComponentRendererPropertiesDetails.h
--- a/Engine/Plugins/FX/Niagara/Source/NiagaraEditor/Private/Customizations/NiagaraComponentRendererPropertiesDetails.h
+++ b/Engine/Plugins/FX/Niagara/Source/NiagaraEditor/Private/Customizations/NiagaraComponentRendererPropertiesDetails.h
@@ -14,7 +14,7 @@ class SWidget;
 class FNiagaraComponentRendererPropertiesDetails : public IDetailCustomization
 {
 public:
-	~FNiagaraComponentRendererPropertiesDetails();
+	virtual ~FNiagaraComponentRendererPropertiesDetails() override;
 	virtual void CustomizeDetails(IDetailLayoutBuilder& DetailBuilder) override;
 	virtual void CustomizeDetails(const TSharedPtr<IDetailLayoutBuilder>& DetailBuilder) override;
 	static TSharedRef<IDetailCustomization> MakeInstance();
@@ -28,7 +28,7 @@ private:
 	FReply ResetBindingButtonPressed(TSharedPtr<IPropertyHandle> PropertyHandle);
 	void RefreshPropertiesPanel();
 	FText GetCurrentBindingText(TSharedPtr<IPropertyHandle> PropertyHandle) const;
-	const UNiagaraEmitter* GetCurrentEmitter() const;
+	FVersionedNiagaraEmitter GetCurrentEmitter() const;
 	TArray<FNiagaraVariable> GetPossibleBindings(TSharedPtr<IPropertyHandle> PropertyHandle) const;
 	bool IsOverridableType(TSharedPtr<IPropertyHandle> PropertyHandle) const;
 	void ChangePropertyBinding(TSharedPtr<IPropertyHandle> PropertyHandle, const FNiagaraVariable& BindingVar);
diff --git a/Engine/Plugins/FX/Niagara/Source/NiagaraEditor/Private/Customizations/NiagaraEventScriptPropertiesCustomization.cpp b/Engine/Plugins/FX/Niagara/Source/NiagaraEditor/Private/Customizations/NiagaraEventScriptPropertiesCustomization.cpp
--- a/Engine/Plugins/FX/Niagara/Source/NiagaraEditor/Private/Customizations/NiagaraEventScriptPropertiesCustomization.cpp
+++ b/Engine/Plugins/FX/Niagara/Source/NiagaraEditor/Private/Customizations/NiagaraEventScriptPropertiesCustomization.cpp
@@ -3,23 +3,13 @@
 #include "Customizations/NiagaraEventScriptPropertiesCustomization.h"
 #include "Widgets/Text/STextBlock.h"
 #include "Widgets/DeclarativeSyntaxSupport.h"
-#include "Styling/SlateTypes.h"
-#include "SlateOptMacros.h"
-#include "Widgets/Input/SCheckBox.h"
 #include "Widgets/Input/STextComboBox.h"
 #include "DetailWidgetRow.h"
 #include "IDetailPropertyRow.h"
-#include "DetailCategoryBuilder.h"
 #include "DetailLayoutBuilder.h"
-#include "DetailWidgetRow.h"
 #include "NiagaraEmitter.h"
-#include "Textures/SlateIcon.h"
-#include "Framework/Commands/UIAction.h"
 #include "Widgets/Layout/SBorder.h"
-#include "Widgets/Text/STextBlock.h"
-#include "Framework/MultiBox/MultiBoxBuilder.h"
 #include "NiagaraSystem.h"
-#include "NiagaraEmitter.h"
 #include "NiagaraEmitterHandle.h"
 #include "ScopedTransaction.h"
 #include "SGraphActionMenu.h"
@@ -29,14 +19,13 @@
 
 #define LOCTEXT_NAMESPACE "FNiagaraEventScriptPropertiesCustomization"
 
-TSharedRef<IPropertyTypeCustomization> FNiagaraEventScriptPropertiesCustomization::MakeInstance(TWeakObjectPtr<UNiagaraSystem> InSystem,
-TWeakObjectPtr<UNiagaraEmitter> InEmitter)
+TSharedRef<IPropertyTypeCustomization> FNiagaraEventScriptPropertiesCustomization::MakeInstance(TWeakObjectPtr<UNiagaraSystem> InSystem, FVersionedNiagaraEmitterWeakPtr InEmitter)
 {
 	return MakeShareable(new FNiagaraEventScriptPropertiesCustomization(InSystem, InEmitter));
 }
 
 FNiagaraEventScriptPropertiesCustomization::FNiagaraEventScriptPropertiesCustomization(TWeakObjectPtr<UNiagaraSystem> InSystem,
-	TWeakObjectPtr<UNiagaraEmitter> InEmitter) : System(InSystem), Emitter(InEmitter)
+	FVersionedNiagaraEmitterWeakPtr InEmitter) : System(InSystem), Emitter(InEmitter)
 {
 	GEditor->RegisterForUndo(this);
 }
@@ -231,7 +220,7 @@ void FNiagaraEventScriptPropertiesCustomization::ComputeErrorVisibility()
 				bIsMatch = true;
 			}
 
-			if (Emitter == Handle.GetInstance())
+			if (Emitter.ResolveWeakPtr() == Handle.GetInstance())
 			{
 				HandleGuid.Invalidate();
 			}
@@ -314,7 +303,7 @@ void FNiagaraEventScriptPropertiesCustomization::ResolveEmitterName()
 			}
 
 			// Set our handle to zeroes if we are looking at this emitter.
-			if (Emitter == Handle.GetInstance())
+			if (Emitter.ResolveWeakPtr() == Handle.GetInstance())
 			{
 				HandleGuid.Invalidate();
 			}
@@ -348,11 +337,11 @@ TSharedRef<SWidget> FNiagaraEventScriptPropertiesCustomization::OnGetMenuContent
 		];
 }
 
-TArray<FName> FNiagaraEventScriptPropertiesCustomization::GetEventNames(UNiagaraEmitter* InEmitter) const
+TArray<FName> FNiagaraEventScriptPropertiesCustomization::GetEventNames(const FVersionedNiagaraEmitter& InEmitter) const
 {
 	TArray<FName> EventNames;
 	TArray<UNiagaraScript*> Scripts;
-	InEmitter->GetScripts(Scripts);
+	InEmitter.GetEmitterData()->GetScripts(Scripts);
 
 	for (UNiagaraScript* Script : Scripts)
 	{
@@ -380,7 +369,7 @@ void FNiagaraEventScriptPropertiesCustomization::CollectAllActions(FGraphActionL
 		FName EmitterName = Handle.GetName();
 
 		// Set our handle to zeroes if we are looking at this emitter.
-		if (Emitter == Handle.GetInstance())
+		if (Emitter.ResolveWeakPtr() == Handle.GetInstance())
 		{
 			HandleGuid.Invalidate();
 		}
diff --git a/Engine/Plugins/FX/Niagara/Source/NiagaraEditor/Private/Customizations/NiagaraPlatformSetCustomization.cpp b/Engine/Plugins/FX/Niagara/Source/NiagaraEditor/Private/Customizations/NiagaraPlatformSetCustomization.cpp
--- a/Engine/Plugins/FX/Niagara/Source/NiagaraEditor/Private/Customizations/NiagaraPlatformSetCustomization.cpp
+++ b/Engine/Plugins/FX/Niagara/Source/NiagaraEditor/Private/Customizations/NiagaraPlatformSetCustomization.cpp
@@ -10,7 +10,6 @@
 #include "NiagaraTypes.h"
 #include "PlatformInfo.h"
 #include "PropertyHandle.h"
-#include "Scalability.h"
 #include "DeviceProfiles/DeviceProfile.h"
 #include "DeviceProfiles/DeviceProfileManager.h"
 #include "Layout/Visibility.h"
@@ -23,16 +22,55 @@
 #include "Widgets/SToolTip.h"
 #include "NiagaraSettings.h"
 #include "IDetailPropertyRow.h"
+#include "IPropertyUtilities.h"
+#include "PropertyEditor/Private/PropertyNode.h"
+#include "PropertyEditor/Private/StructurePropertyNode.h"
 
 
 #define LOCTEXT_NAMESPACE "FNiagaraPlatformSetCustomization"
 
+uint8* GetBaseAddress(TSharedRef<IPropertyHandle> PropertyHandle)
+{
+	if (PropertyHandle->GetNumOuterObjects() > 0)
+	{
+		TArray<UObject*> Objects;
+		PropertyHandle->GetOuterObjects(Objects);
+		return reinterpret_cast<uint8*>(Objects[0]);
+	}
+	
+	// walk the struct hierarchy to find the parent StructOnScope
+	FStructurePropertyNode* StructurePropertyNode = nullptr;
+	TSharedPtr<IPropertyHandle> LastParent = PropertyHandle;
+	while (StructurePropertyNode == nullptr)
+	{
+		TSharedPtr<IPropertyHandle> NewParent = LastParent->GetParentHandle();
+		if (!ensureMsgf(NewParent.IsValid() && NewParent != LastParent, TEXT("Unable to walk property chain")))
+		{
+			return nullptr;
+		}
+		
+		if (FComplexPropertyNode* ComplexPropertyNode = NewParent->GetPropertyNode()->AsComplexNode())
+		{
+			StructurePropertyNode = ComplexPropertyNode->AsStructureNode();
+		}
+		LastParent = NewParent;
+	}
+	TSharedPtr<FStructOnScope> StructOnScope = StructurePropertyNode->GetStructData();
+	return StructOnScope->GetStructMemory();
+}
+
 void FNiagaraPlatformSetCustomization::CustomizeHeader(TSharedRef<IPropertyHandle> InPropertyHandle, FDetailWidgetRow& HeaderRow, IPropertyTypeCustomizationUtils& CustomizationUtils)
 {
 	PropertyHandle = InPropertyHandle;
-	TArray<UObject*> Objects;
-	PropertyHandle->GetOuterObjects(Objects);
-	TargetPlatformSet = (FNiagaraPlatformSet*)PropertyHandle->GetValueBaseAddress((uint8*)Objects[0]);
+	TArray<TWeakObjectPtr<UObject>> WeakObjectPtrs = CustomizationUtils.GetPropertyUtilities()->GetSelectedObjects();
+
+	uint8* BaseAddress = GetBaseAddress(InPropertyHandle);
+	if (!BaseAddress)
+	{
+		return;
+	}
+	
+	TargetPlatformSet = reinterpret_cast<FNiagaraPlatformSet*>(PropertyHandle->GetValueBaseAddress(BaseAddress));
 
 	if (PlatformSelectionStates.Num() == 0)
 	{
@@ -68,22 +106,22 @@ void FNiagaraPlatformSetCustomization::CustomizeHeader(TSharedRef<IPropertyHandl
 				TSharedPtr<IPropertyHandle> ParentHandle = CurrHandle->GetParentHandle();
 				if (CurrStruct == FNiagaraSystemScalabilitySettingsArray::StaticStruct())
 				{
-					SystemScalabilitySettings = (FNiagaraSystemScalabilitySettingsArray*)ParentHandle->GetValueBaseAddress((uint8*)Objects[0]);
+					SystemScalabilitySettings = reinterpret_cast<FNiagaraSystemScalabilitySettingsArray*>(ParentHandle->GetValueBaseAddress(BaseAddress));
 					break;
 				}
 				else if (CurrStruct == FNiagaraEmitterScalabilitySettingsArray::StaticStruct())
 				{
-					EmitterScalabilitySettings = (FNiagaraEmitterScalabilitySettingsArray*)ParentHandle->GetValueBaseAddress((uint8*)Objects[0]);
+					EmitterScalabilitySettings = reinterpret_cast<FNiagaraEmitterScalabilitySettingsArray*>(ParentHandle->GetValueBaseAddress(BaseAddress));
 					break;
 				}
 				else if (CurrStruct == FNiagaraSystemScalabilityOverrides::StaticStruct())
 				{
-					SystemScalabilityOverrides = (FNiagaraSystemScalabilityOverrides*)ParentHandle->GetValueBaseAddress((uint8*)Objects[0]);
+					SystemScalabilityOverrides = reinterpret_cast<FNiagaraSystemScalabilityOverrides*>(ParentHandle->GetValueBaseAddress(BaseAddress));
 					break;
 				}
 				else if (CurrStruct == FNiagaraEmitterScalabilityOverrides::StaticStruct())
 				{
-					EmitterScalabilityOverrides = (FNiagaraEmitterScalabilityOverrides*)ParentHandle->GetValueBaseAddress((uint8*)Objects[0]);
+					EmitterScalabilityOverrides = reinterpret_cast<FNiagaraEmitterScalabilityOverrides*>(ParentHandle->GetValueBaseAddress(BaseAddress));
 					break;
 				}
 			}
@@ -94,7 +132,6 @@ void FNiagaraPlatformSetCustomization::CustomizeHeader(TSharedRef<IPropertyHandl
 
 	UpdateCachedConflicts();
 	
-	BaseSystem = Objects[0]->GetTypedOuter<UNiagaraSystem>();
 	PropertyHandle->SetOnPropertyValueChanged(FSimpleDelegate::CreateSP(this, &FNiagaraPlatformSetCustomization::OnPropertyValueChanged));
 	PropertyHandle->SetOnChildPropertyValueChanged(FSimpleDelegate::CreateLambda([&]{ TargetPlatformSet->OnChanged(); }));
 
@@ -286,7 +323,7 @@ static TSharedPtr<IPropertyHandle> FindChildPlatformSet(TSharedPtr<IPropertyHand
 	uint32 NumChildren = 0;
 	PropertyHandle->GetNumChildren(NumChildren);
 
-	for (int32 Idx = 0; Idx < (int32) NumChildren; ++Idx)
+	for (int32 Idx = 0; Idx < static_cast<int32>(NumChildren); ++Idx)
 	{
 		TSharedPtr<IPropertyHandle> Child = PropertyHandle->GetChildHandle(Idx);
 
@@ -309,7 +346,7 @@ void FNiagaraPlatformSetCustomization::InvalidateSiblingConflicts() const
 
 	uint32 ArrayCount = 0;
 	PlatformSetArray->GetNumElements(ArrayCount);
-	for (int32 Idx = 0; Idx < (int32) ArrayCount; ++Idx)
+	for (int32 Idx = 0; Idx < static_cast<int32>(ArrayCount); ++Idx)
 	{
 		if (Idx == PlatformSetArrayIndex)
 		{
@@ -1280,6 +1317,7 @@ static const TCHAR* NiagaraCVarHistoryKey = TEXT("NiagaraCVarHistory");
 SNiagaraConsoleInputBox::SNiagaraConsoleInputBox()
 	: bIgnoreUIUpdate(false)
 	, bHasTicked(false)
+	, bConsumeTab(false)
 {
 }
 
@@ -1787,27 +1825,20 @@ void FNiagaraPlatformSetCVarConditionCustomization::OnTextCommitted(const FText&
 
 FNiagaraPlatformSetCVarCondition* FNiagaraPlatformSetCVarConditionCustomization::GetTargetCondition()const
 {
-	TArray<UObject*> Objects;
-	PropertyHandle->GetOuterObjects(Objects);
-	int32 Index = 0;
-	if (Objects.IsValidIndex(Index))
+	if (uint8* BaseAddress = GetBaseAddress(PropertyHandle.ToSharedRef()))
 	{
-		return (FNiagaraPlatformSetCVarCondition*)PropertyHandle->GetValueBaseAddress((uint8*)Objects[Index]);
+		return reinterpret_cast<FNiagaraPlatformSetCVarCondition*>(PropertyHandle->GetValueBaseAddress(BaseAddress));
 	}
 	return nullptr;
 }
 
 FNiagaraPlatformSet* FNiagaraPlatformSetCVarConditionCustomization::GetTargetPlatformSet()const
 {
-	TArray<UObject*> Objects;
 	check(PropertyHandle && PropertyHandle->GetParentHandle());
 	TSharedPtr<IPropertyHandle> ParentHandle = PropertyHandle->GetParentHandle();
-	ParentHandle->GetOuterObjects(Objects);
-
-	int32 Index = 0;
-	if (Objects.IsValidIndex(Index))
+	if (uint8* BaseAddress = GetBaseAddress(ParentHandle.ToSharedRef()))
 	{
-		return (FNiagaraPlatformSet*)ParentHandle->GetValueBaseAddress((uint8*)Objects[Index]);
+		return reinterpret_cast<FNiagaraPlatformSet*>(ParentHandle->GetValueBaseAddress(BaseAddress));
 	}
 	return nullptr;
 }
diff --git a/Engine/Plugins/FX/Niagara/Source/NiagaraEditor/Private/Customizations/NiagaraPlatformSetCustomization.h b/Engine/Plugins/FX/Niagara/Source/NiagaraEditor/Private/Customizations/NiagaraPlatformSetCustomization.h
--- a/Engine/Plugins/FX/Niagara/Source/NiagaraEditor/Private/Customizations/NiagaraPlatformSetCustomization.h
+++ b/Engine/Plugins/FX/Niagara/Source/NiagaraEditor/Private/Customizations/NiagaraPlatformSetCustomization.h
@@ -91,7 +91,6 @@ private:
 
 	TSharedPtr<IPropertyHandle> PropertyHandle;
 	struct FNiagaraPlatformSet* TargetPlatformSet;
-	class UNiagaraSystem* BaseSystem;
 
 	TArray<TSharedPtr<FNiagaraDeviceProfileViewModel>> FullDeviceProfileTree;
 	TArray<TArray<TSharedPtr<FNiagaraDeviceProfileViewModel>>> FilteredDeviceProfileTrees;
diff --git a/Engine/Plugins/FX/Niagara/Source/NiagaraEditor/Private/Customizations/NiagaraTypeCustomizations.cpp b/Engine/Plugins/FX/Niagara/Source/NiagaraEditor/Private/Customizations/NiagaraTypeCustomizations.cpp
--- a/Engine/Plugins/FX/Niagara/Source/NiagaraEditor/Private/Customizations/NiagaraTypeCustomizations.cpp
+++ b/Engine/Plugins/FX/Niagara/Source/NiagaraEditor/Private/Customizations/NiagaraTypeCustomizations.cpp
@@ -121,12 +121,13 @@ void FNiagaraMatrixCustomization::CustomizeChildren(TSharedRef<IPropertyHandle>
 	}
 }
 
-TArray<FNiagaraVariableBase> FNiagaraStackAssetAction_VarBind::FindVariables(UNiagaraEmitter* InEmitter, bool bSystem, bool bEmitter, bool bParticles, bool bUser, bool bAllowStatic)
+TArray<FNiagaraVariableBase> FNiagaraStackAssetAction_VarBind::FindVariables(const FVersionedNiagaraEmitter& InEmitter, bool bSystem, bool bEmitter, bool bParticles, bool bUser, bool bAllowStatic)
 {
 	TArray<FNiagaraVariableBase> Bindings;
 	TArray<FNiagaraParameterMapHistory> Histories;
 
-	UNiagaraScriptSource* Source = Cast<UNiagaraScriptSource>(InEmitter->GraphSource);
+	UNiagaraEmitter* Emitter = InEmitter.Emitter;
+	UNiagaraScriptSource* Source = Cast<UNiagaraScriptSource>(InEmitter.GetEmitterData()->GraphSource);
 	if (Source)
 	{
 		Histories.Append(UNiagaraNodeParameterMapBase::GetParameterMaps(Source->NodeGraph));
@@ -134,8 +135,7 @@ TArray<FNiagaraVariableBase> FNiagaraStackAssetAction_VarBind::FindVariables(UNi
 
 	if (bSystem || bEmitter)
 	{
-		UNiagaraSystem* Sys = InEmitter->GetTypedOuter<UNiagaraSystem>();
-		if (Sys)
+		if (UNiagaraSystem* Sys = Emitter->GetTypedOuter<UNiagaraSystem>())
 		{
 			Source = Cast<UNiagaraScriptSource>(Sys->GetSystemUpdateScript()->GetLatestSource()); 
 			if (Source)
@@ -162,10 +162,10 @@ TArray<FNiagaraVariableBase> FNiagaraStackAssetAction_VarBind::FindVariables(UNi
 			{
 				Bindings.AddUnique(Var);
 			}
-			else if (Var.IsInNameSpace(InEmitter->GetUniqueEmitterName()) && bEmitter)
+			else if (Var.IsInNameSpace(Emitter->GetUniqueEmitterName()) && bEmitter)
 			{
 				Bindings.AddUnique(FNiagaraUtilities::ResolveAliases(Var, FNiagaraAliasContext()
-					.ChangeEmitterNameToEmitter(InEmitter->GetUniqueEmitterName())));
+					.ChangeEmitterNameToEmitter(Emitter->GetUniqueEmitterName())));
 			}
 			else if (FNiagaraParameterMapHistory::IsAliasedEmitterParameter(Var) && bEmitter)
 			{
@@ -184,8 +184,7 @@ TArray<FNiagaraVariableBase> FNiagaraStackAssetAction_VarBind::FindVariables(UNi
 
 	if (bUser)
 	{
-		UNiagaraSystem* Sys = InEmitter->GetTypedOuter<UNiagaraSystem>();
-		if (Sys)
+		if (UNiagaraSystem* Sys = Emitter->GetTypedOuter<UNiagaraSystem>())
 		{
 			for (const FNiagaraVariable Var : Sys->GetExposedParameters().ReadParameterVariables())
 			{
@@ -199,7 +198,7 @@ TArray<FNiagaraVariableBase> FNiagaraStackAssetAction_VarBind::FindVariables(UNi
 
 FName FNiagaraVariableAttributeBindingCustomization::GetVariableName() const
 {
-	if (BaseEmitter && TargetVariableBinding)
+	if (BaseEmitter.Emitter && TargetVariableBinding)
 	{
 		return (TargetVariableBinding->GetName());
 	}
@@ -208,7 +207,7 @@ FName FNiagaraVariableAttributeBindingCustomization::GetVariableName() const
 
 FText FNiagaraVariableAttributeBindingCustomization::GetCurrentText() const
 {
-	if (BaseEmitter && TargetVariableBinding)
+	if (BaseEmitter.Emitter && TargetVariableBinding)
 	{
 		return FText::FromName(TargetVariableBinding->GetName());
 	}
@@ -217,7 +216,7 @@ FText FNiagaraVariableAttributeBindingCustomization::GetCurrentText() const
 
 FText FNiagaraVariableAttributeBindingCustomization::GetTooltipText() const
 {
-	if (BaseEmitter && TargetVariableBinding)
+	if (BaseEmitter.Emitter && TargetVariableBinding)
 	{
 		FString DefaultValueStr = TargetVariableBinding->GetDefaultValueString();
 
@@ -257,7 +256,7 @@ TSharedRef<SWidget> FNiagaraVariableAttributeBindingCustomization::OnGetMenuCont
 	return Widget.ToSharedRef(); 
 }
 
-TArray<FName> FNiagaraVariableAttributeBindingCustomization::GetNames(UNiagaraEmitter* InEmitter) const
+TArray<FName> FNiagaraVariableAttributeBindingCustomization::GetNames(const FVersionedNiagaraEmitter& InEmitter) const
 {
 	TArray<FName> Names;
 	if (!PropertyHandle.IsValid() || !PropertyHandle->GetProperty() || !TargetVariableBinding)
@@ -359,7 +358,7 @@ void FNiagaraVariableAttributeBindingCustomization::ChangeSource(FName InVarName
 	{
 		Obj->Modify();
 	}
-	check(BaseEmitter);
+	check(BaseEmitter.Emitter);
 	check(RenderProps || SimulationStage);
 
 	PropertyHandle->NotifyPreChange();
@@ -376,7 +375,7 @@ void FNiagaraVariableAttributeBindingCustomization::ResetToDefault()
 
 EVisibility FNiagaraVariableAttributeBindingCustomization::IsResetToDefaultsVisible() const
 {
-	check(BaseEmitter);
+	check(BaseEmitter.Emitter);
 	check(RenderProps || SimulationStage);
 	check(TargetVariableBinding);
 	check(DefaultVariableBinding);
@@ -394,7 +393,7 @@ FReply FNiagaraVariableAttributeBindingCustomization::OnResetToDefaultsClicked()
 	{
 		Obj->Modify();
 	}
-	check(BaseEmitter);
+	check(BaseEmitter.Emitter);
 	check(RenderProps || SimulationStage);
 	check(TargetVariableBinding);
 	check(DefaultVariableBinding);
@@ -411,7 +410,7 @@ void FNiagaraVariableAttributeBindingCustomization::CustomizeHeader(TSharedRef<I
 {
 	RenderProps = nullptr;
 	SimulationStage = nullptr;
-	BaseEmitter = nullptr;
+	BaseEmitter = FVersionedNiagaraEmitter();
 	PropertyHandle = InPropertyHandle;
 	TArray<UObject*> Objects;
 	PropertyHandle->GetOuterObjects(Objects);
@@ -433,16 +432,16 @@ void FNiagaraVariableAttributeBindingCustomization::CustomizeHeader(TSharedRef<I
 		RenderProps = Cast<UNiagaraRendererProperties>(Objects[0]);
 		if (RenderProps)
 		{
-			BaseEmitter = Cast<UNiagaraEmitter>(RenderProps->GetOuter());
+			BaseEmitter = RenderProps->GetOuterEmitter();
 		}
 
 		SimulationStage = Cast<UNiagaraSimulationStageBase>(Objects[0]);
 		if ( SimulationStage )
 		{
-			BaseEmitter = SimulationStage->GetTypedOuter<UNiagaraEmitter>();
+			BaseEmitter = SimulationStage->GetOuterEmitter();
 		}
 
-		if (BaseEmitter)
+		if (BaseEmitter.Emitter)
 		{
 			TargetVariableBinding = (FNiagaraVariableAttributeBinding*)PropertyHandle->GetValueBaseAddress((uint8*)Objects[0]);
 			DefaultVariableBinding = (FNiagaraVariableAttributeBinding*)PropertyHandle->GetValueBaseAddress((uint8*)Objects[0]->GetClass()->GetDefaultObject());
@@ -670,7 +669,7 @@ void FNiagaraUserParameterBindingCustomization::CustomizeHeader(TSharedRef<IProp
 			if (UNiagaraDataInterface* ObjectAsDataInterface = Cast<UNiagaraDataInterface>(Objects[0]))
 			{
 				FNiagaraEditorModule& NiagaraEditorModule = FModuleManager::LoadModuleChecked<FNiagaraEditorModule>("NiagaraEditor");
-				UNiagaraEmitter* Emitter = nullptr;
+				FVersionedNiagaraEmitter Emitter;
 				NiagaraEditorModule.GetTargetSystemAndEmitterForDataInterface(ObjectAsDataInterface, BaseSystem, Emitter);
 			}
 		}
@@ -838,7 +837,7 @@ TArray<TPair<FNiagaraVariableBase, FNiagaraVariableBase> > FNiagaraMaterialAttri
 	TArray<TPair<FNiagaraVariableBase, FNiagaraVariableBase>> Names;
 	TArray < FNiagaraVariableBase > BaseVars;
 
-	if (BaseSystem && BaseEmitter && TargetParameterBinding)
+	if (BaseSystem && BaseEmitter.Emitter && TargetParameterBinding)
 	{
 		bool bSystem = true;
 		bool bEmitter = true;
@@ -850,10 +849,10 @@ TArray<TPair<FNiagaraVariableBase, FNiagaraVariableBase> > FNiagaraMaterialAttri
 		TArray<UNiagaraScript*> Scripts;
 		Scripts.Add(BaseSystem->GetSystemUpdateScript());
 		Scripts.Add(BaseSystem->GetSystemSpawnScript());
-		BaseEmitter->GetScripts(Scripts, false);
+		BaseEmitter.GetEmitterData()->GetScripts(Scripts, false);
 
 		TMap<FString, FString> EmitterAlias;
-		EmitterAlias.Emplace(FNiagaraConstants::EmitterNamespace.ToString(), BaseEmitter->GetUniqueEmitterName());
+		EmitterAlias.Emplace(FNiagaraConstants::EmitterNamespace.ToString(), BaseEmitter.Emitter->GetUniqueEmitterName());
 
 		auto FindCachedDI = 
 			[&](const FNiagaraVariableBase& BaseVariable) -> UNiagaraDataInterface*
@@ -862,7 +861,7 @@ TArray<TPair<FNiagaraVariableBase, FNiagaraVariableBase> > FNiagaraMaterialAttri
 				if (BaseVariable.IsInNameSpace(FNiagaraConstants::EmitterNamespaceString))
 				{
 					VariableName = FNiagaraUtilities::ResolveAliases(BaseVariable, FNiagaraAliasContext()
-						.ChangeEmitterToEmitterName(BaseEmitter->GetUniqueEmitterName())).GetName();
+						.ChangeEmitterToEmitterName(BaseEmitter.Emitter->GetUniqueEmitterName())).GetName();
 				}
 
 				for (UNiagaraScript* Script : Scripts)
@@ -1008,7 +1007,7 @@ void FNiagaraMaterialAttributeBindingCustomization::ChangeNiagaraSource(FNiagara
 	PropertyHandle->NotifyPreChange();
 	TargetParameterBinding->NiagaraVariable = InVar->BaseVar;
 	TargetParameterBinding->NiagaraChildVariable = InVar->ChildVar;
-	TargetParameterBinding->CacheValues(BaseEmitter);
+	TargetParameterBinding->CacheValues(BaseEmitter.Emitter);
 	//TargetParameterBinding->Parameter.SetType(FNiagaraTypeDefinition::GetUObjectDef()); Do not override the type here!
 	//TargetVariableBinding->DataSetVariable = FNiagaraConstants::GetAttributeAsDataSetKey(TargetVariableBinding->BoundVariable);
 	PropertyHandle->NotifyPostChange(EPropertyChangeType::ValueSet);
@@ -1207,7 +1206,7 @@ void FNiagaraMaterialAttributeBindingCustomization::CustomizeChildren(TSharedRef
 {
 	RenderProps = nullptr;
 	BaseSystem =  nullptr;
-	BaseEmitter = nullptr;
+	BaseEmitter = FVersionedNiagaraEmitter();
 	  
 	TArray<UObject*> Objects;
 	PropertyHandle->GetOuterObjects(Objects);
@@ -1215,7 +1214,10 @@ void FNiagaraMaterialAttributeBindingCustomization::CustomizeChildren(TSharedRef
 	{
 		RenderProps = Cast<UNiagaraRendererProperties>(Objects[0]);
 		BaseSystem = Objects[0]->GetTypedOuter<UNiagaraSystem>();
-		BaseEmitter = Objects[0]->GetTypedOuter<UNiagaraEmitter>();
+		if (RenderProps)
+		{
+			BaseEmitter = RenderProps->GetOuterEmitter();
+		}
 		if (BaseSystem)
 		{
 			TargetParameterBinding = (FNiagaraMaterialAttributeBinding*)PropertyHandle->GetValueBaseAddress((uint8*)Objects[0]);
@@ -1341,19 +1343,16 @@ TArray<FName> FNiagaraDataInterfaceBindingCustomization::GetNames() const
 
 	if (BaseStage && TargetDataInterfaceBinding)
 	{
-		UNiagaraEmitter* Emitter = BaseStage->GetTypedOuter<UNiagaraEmitter>();
-
-		if (Emitter)
+		if (FVersionedNiagaraEmitterData* EmitterData = BaseStage->GetEmitterData())
 		{
 			// Find all used emitter and particle data interface variables that can be iterated upon.
 			TArray<UNiagaraScript*> AllScripts;
-			Emitter->GetScripts(AllScripts, false);
+			EmitterData->GetScripts(AllScripts, false);
 
 			TArray<UNiagaraGraph*> Graphs;
 			for (const UNiagaraScript* Script : AllScripts)
 			{
-				const UNiagaraScriptSource* Source = Cast<UNiagaraScriptSource>(Script->GetLatestSource());
-				if (Source)
+				if (const UNiagaraScriptSource* Source = Cast<UNiagaraScriptSource>(Script->GetLatestSource()))
 				{
 					Graphs.AddUnique(Source->NodeGraph);
 				}
diff --git a/Engine/Plugins/FX/Niagara/Source/NiagaraEditor/Private/Customizations/NiagaraTypeCustomizations.h b/Engine/Plugins/FX/Niagara/Source/NiagaraEditor/Private/Customizations/NiagaraTypeCustomizations.h
--- a/Engine/Plugins/FX/Niagara/Source/NiagaraEditor/Private/Customizations/NiagaraTypeCustomizations.h
+++ b/Engine/Plugins/FX/Niagara/Source/NiagaraEditor/Private/Customizations/NiagaraTypeCustomizations.h
@@ -3,7 +3,6 @@
 #pragma once
 
 #include "CoreMinimal.h"
-#include "IDetailCustomization.h"
 #include "IPropertyTypeCustomization.h"
 #include "NiagaraCommon.h"
 #include "EdGraph/EdGraphSchema.h"
@@ -96,7 +95,6 @@ public:
 	virtual void CustomizeChildren(TSharedRef<IPropertyHandle> PropertyHandle, IDetailChildrenBuilder& ChildBuilder, IPropertyTypeCustomizationUtils& CustomizationUtils);
 
 
-private:
 };
 
 
@@ -151,7 +149,7 @@ struct FNiagaraStackAssetAction_VarBind : public FEdGraphSchemaAction
 	}
 	//~ End FEdGraphSchemaAction Interface
 
-	static TArray<FNiagaraVariableBase> FindVariables(UNiagaraEmitter* InEmitter, bool bSystem, bool bEmitter, bool bParticles, bool bUser, bool bAllowStatic);
+	static TArray<FNiagaraVariableBase> FindVariables(const FVersionedNiagaraEmitter& InEmitter, bool bSystem, bool bEmitter, bool bParticles, bool bUser, bool bAllowStatic);
 };
 
 class FNiagaraVariableAttributeBindingCustomization : public IPropertyTypeCustomization
@@ -175,7 +173,7 @@ private:
 	FText GetCurrentText() const;
 	FText GetTooltipText() const;
 	TSharedRef<SWidget> OnGetMenuContent() const;
-	TArray<FName> GetNames(class UNiagaraEmitter* InEmitter) const;
+	TArray<FName> GetNames(const FVersionedNiagaraEmitter& InEmitter) const;
 
 	void ChangeSource(FName InVarName);
 	void CollectAllActions(FGraphActionListBuilderBase& OutAllActions);
@@ -184,7 +182,7 @@ private:
 
 	TSharedPtr<IPropertyHandle> PropertyHandle;
 	TSharedPtr<SComboButton> ComboButton;
-	class UNiagaraEmitter* BaseEmitter = nullptr;
+	FVersionedNiagaraEmitter BaseEmitter;
 	class UNiagaraRendererProperties* RenderProps = nullptr;
 	class UNiagaraSimulationStageBase* SimulationStage = nullptr;
 	struct FNiagaraVariableAttributeBinding* TargetVariableBinding = nullptr;
@@ -268,7 +266,7 @@ private:
 	TSharedPtr<SComboButton> MaterialParameterButton;
 	TSharedPtr<SComboButton> NiagaraParameterButton;
 	class UNiagaraSystem* BaseSystem;
-	class UNiagaraEmitter* BaseEmitter;
+	FVersionedNiagaraEmitter BaseEmitter;
 	class UNiagaraRendererProperties* RenderProps;
 	struct FNiagaraMaterialAttributeBinding* TargetParameterBinding;
 
diff --git a/Engine/Plugins/FX/Niagara/Source/NiagaraEditor/Private/EdGraphSchema_Niagara.cpp b/Engine/Plugins/FX/Niagara/Source/NiagaraEditor/Private/EdGraphSchema_Niagara.cpp
--- a/Engine/Plugins/FX/Niagara/Source/NiagaraEditor/Private/EdGraphSchema_Niagara.cpp
+++ b/Engine/Plugins/FX/Niagara/Source/NiagaraEditor/Private/EdGraphSchema_Niagara.cpp
@@ -79,7 +79,7 @@ const FName UEdGraphSchema_Niagara::PinCategoryStaticEnum("StaticEnum");
 namespace NiagaraNodeNumbers
 {
 	// Maximum distance a drag can be off a node edge to require 'push off' from node
-	const int32 NiagaraMinNodeDistance = 60;
+	constexpr int32 NiagaraMinNodeDistance = 60;
 }
 
 UEdGraphNode* FNiagaraSchemaAction_NewNode::PerformAction(class UEdGraph* ParentGraph, UEdGraphPin* FromPin, const FVector2D Location, bool bSelectNewNode/* = true*/)
@@ -280,7 +280,7 @@ const UNiagaraGraph* GetAlternateGraph(const UNiagaraGraph* NiagaraGraph)
 		UNiagaraScript* Script = Cast<UNiagaraScript>(ScriptSource->GetOuter());
 		if (Script != nullptr)
 		{
-			UNiagaraEmitter* EmitterProperties = Cast<UNiagaraEmitter>(Script->GetOuter());
+			FVersionedNiagaraEmitterData* EmitterProperties = Script->GetOuterEmitter().GetEmitterData();
 			if (EmitterProperties != nullptr)
 			{
 				if (EmitterProperties->SpawnScriptProps.Script == Script)
diff --git a/Engine/Plugins/FX/Niagara/Source/NiagaraEditor/Private/NiagaraBakerRenderer.cpp b/Engine/Plugins/FX/Niagara/Source/NiagaraEditor/Private/NiagaraBakerRenderer.cpp
--- a/Engine/Plugins/FX/Niagara/Source/NiagaraEditor/Private/NiagaraBakerRenderer.cpp
+++ b/Engine/Plugins/FX/Niagara/Source/NiagaraEditor/Private/NiagaraBakerRenderer.cpp
@@ -132,8 +132,8 @@ void FNiagaraBakerOutputBindingHelper::ForEachEmitterDataInterface(UNiagaraSyste
 	for (int32 EmitterIndex=0; EmitterIndex < NiagaraSystem->GetEmitterHandles().Num(); ++EmitterIndex)
 	{
 		const FNiagaraEmitterHandle& EmitterHandle = NiagaraSystem->GetEmitterHandle(EmitterIndex);
-		UNiagaraEmitter* EmitterInstance = EmitterHandle.GetInstance();
-		if (!EmitterHandle.IsValid() || !EmitterHandle.GetIsEnabled() || !EmitterInstance)
+		FVersionedNiagaraEmitterData* EmitterData = EmitterHandle.GetInstance().GetEmitterData();
+		if (!EmitterHandle.IsValid() || !EmitterHandle.GetIsEnabled() || !EmitterData)
 		{
 			continue;
 		}
@@ -141,10 +141,10 @@ void FNiagaraBakerOutputBindingHelper::ForEachEmitterDataInterface(UNiagaraSyste
 		const FString EmitterName = EmitterHandle.GetName().ToString();
 		const FString EmitterPrefix = EmitterName + TEXT(".");
 
-		EmitterInstance->ForEachScript(
+		EmitterData->ForEachScript(
 			[&](UNiagaraScript* NiagaraScript)
 			{
-				if (const FNiagaraScriptExecutionParameterStore* SrcStore = NiagaraScript->GetExecutionReadyParameterStore(EmitterInstance->SimTarget))
+				if (const FNiagaraScriptExecutionParameterStore* SrcStore = NiagaraScript->GetExecutionReadyParameterStore(EmitterData->SimTarget))
 				{
 					for (const FNiagaraVariableWithOffset& Variable : SrcStore->ReadParameterVariables())
 					{
@@ -231,8 +231,8 @@ void FNiagaraBakerOutputBindingHelper::GetParticleAttributeBindings(TArray<FNiag
 	for (int32 EmitterIndex = 0; EmitterIndex < NiagaraSystem->GetEmitterHandles().Num(); ++EmitterIndex)
 	{
 		const FNiagaraEmitterHandle& EmitterHandle = NiagaraSystem->GetEmitterHandle(EmitterIndex);
-		UNiagaraEmitter* EmitterInstance = EmitterHandle.GetInstance();
-		if (!EmitterHandle.IsValid() || !EmitterHandle.GetIsEnabled() || !EmitterInstance)
+		FVersionedNiagaraEmitterData* EmitterData = EmitterHandle.GetInstance().GetEmitterData();
+		if (!EmitterHandle.IsValid() || !EmitterHandle.GetIsEnabled() || !EmitterData)
 		{
 			continue;
 		}
@@ -523,7 +523,7 @@ void FNiagaraBakerRenderer::RenderParticleAttribute(UTextureRenderTarget2D* Rend
 	
 	for ( const auto& EmitterInstance : SystemInstance->GetEmitters() )
 	{
-		UNiagaraEmitter* NiagaraEmitter = EmitterInstance->GetCachedEmitter();
+		UNiagaraEmitter* NiagaraEmitter = EmitterInstance->GetCachedEmitter().Emitter;
 		if ( !NiagaraEmitter || (NiagaraEmitter->GetUniqueEmitterName() != EmitterName) )
 		{
 			continue;
diff --git a/Engine/Plugins/FX/Niagara/Source/NiagaraEditor/Private/NiagaraCompiler.cpp b/Engine/Plugins/FX/Niagara/Source/NiagaraEditor/Private/NiagaraCompiler.cpp
--- a/Engine/Plugins/FX/Niagara/Source/NiagaraEditor/Private/NiagaraCompiler.cpp
+++ b/Engine/Plugins/FX/Niagara/Source/NiagaraEditor/Private/NiagaraCompiler.cpp
@@ -12,9 +12,7 @@
 #include "NiagaraNodeInput.h"
 #include "NiagaraFunctionLibrary.h"
 #include "NiagaraDataInterface.h"
-#include "ViewModels/Stack/NiagaraStackGraphUtilities.h"
 #include "NiagaraNodeFunctionCall.h"
-#include "NiagaraNodeParameterMapSet.h"
 #include "INiagaraEditorTypeUtilities.h"
 #include "NiagaraEditorUtilities.h"
 #include "NiagaraNodeOutput.h"
@@ -25,7 +23,6 @@
 #include "ShaderCompiler.h"
 #include "NiagaraShader.h"
 #include "NiagaraScript.h"
-#include "NiagaraRendererProperties.h"
 #include "NiagaraSimulationStageBase.h"
 #include "NiagaraTrace.h"
 #include "Serialization/MemoryReader.h"
@@ -346,9 +343,9 @@ void FNiagaraCompileRequestDuplicateData::DeepCopyGraphs(UNiagaraScriptSource* S
 	DuplicateReferencedGraphs(ScriptSource->NodeGraph, NodeGraphDeepCopy.Get(), InUsage, ConstantResolver);
 }
 
-void FNiagaraCompileRequestDuplicateData::DeepCopyGraphs(UNiagaraEmitter* Emitter)
+void FNiagaraCompileRequestDuplicateData::DeepCopyGraphs(const FVersionedNiagaraEmitter& Emitter)
 {
-	UNiagaraScriptSource* ScriptSource = CastChecked<UNiagaraScriptSource>(Emitter->GraphSource);
+	UNiagaraScriptSource* ScriptSource = CastChecked<UNiagaraScriptSource>(Emitter.GetEmitterData()->GraphSource);
 
 	SourceDeepCopy = ScriptSource->CreateCompilationCopy(ValidUsages);
 	NodeGraphDeepCopy = SourceDeepCopy->NodeGraph;
@@ -914,9 +911,9 @@ TSharedPtr<FNiagaraCompileRequestDataBase, ESPMode::ThreadSafe> FNiagaraEditorMo
 		{
 			const FNiagaraEmitterHandle& Handle = System->GetEmitterHandle(i);
 			TSharedPtr<FNiagaraCompileRequestData, ESPMode::ThreadSafe> EmitterPtr = MakeShared<FNiagaraCompileRequestData, ESPMode::ThreadSafe>();
-			EmitterPtr->EmitterUniqueName = Handle.GetInstance()->GetUniqueEmitterName();
+			EmitterPtr->EmitterUniqueName = Handle.GetInstance().Emitter->GetUniqueEmitterName();
 			EmitterPtr->SourceName = BasePtr->SourceName;
-			EmitterPtr->Source = Cast<UNiagaraScriptSource>(Handle.GetInstance()->GraphSource);
+			EmitterPtr->Source = Cast<UNiagaraScriptSource>(Handle.GetEmitterData()->GraphSource);
 			EmitterPtr->bUseRapidIterationParams = BasePtr->bUseRapidIterationParams;
 			EmitterPtr->bDisableDebugSwitches = BasePtr->bDisableDebugSwitches;
 			//EmitterPtr->bSimulationStagesEnabled = Handle.GetInstance()->bSimulationStagesEnabled;
@@ -944,7 +941,7 @@ TSharedPtr<FNiagaraCompileRequestDataBase, ESPMode::ThreadSafe> FNiagaraEditorMo
 		// now that the scripts have been precompiled we can prepare the rapid iteration parameters, which we need to do before we
 		// actually generate the hlsl in the case of baking out the parameters
 		TArray<UNiagaraScript*> Scripts;
-		TMap<UNiagaraScript*, const UNiagaraEmitter*> ScriptToEmitterMap;
+		TMap<UNiagaraScript*, FVersionedNiagaraEmitter> ScriptToEmitterMap;
 
 		// Now we can finish off the emitters.
 		for (int32 i = 0; i < System->GetEmitterHandles().Num(); i++)
@@ -953,12 +950,13 @@ TSharedPtr<FNiagaraCompileRequestDataBase, ESPMode::ThreadSafe> FNiagaraEditorMo
 			FCompileConstantResolver ConstantResolver(Handle.GetInstance(), ENiagaraScriptUsage::EmitterSpawnScript);
 			if (Handle.GetIsEnabled()) // Don't pull in the emitter if it isn't going to be used.
 			{
-				TSharedPtr<FNiagaraGraphCachedDataBase, ESPMode::ThreadSafe> CachedTraversalEmitterData = Handle.GetInstance()->GetCachedTraversalData();
+				TSharedPtr<FNiagaraGraphCachedDataBase, ESPMode::ThreadSafe> CachedTraversalEmitterData = Handle.GetInstance().Emitter->GetCachedTraversalData(Handle.GetInstance().Version);
 				TArray<FNiagaraVariable> StaticVariablesFromEmitter = StaticVariablesFromSystem;
 				StaticVariablesFromEmitter.Append(((FNiagaraGraphCachedBuiltHistory*)CachedTraversalEmitterData.Get())->StaticVariables);
 
 				TArray<UNiagaraScript*> EmitterScripts;
-				Handle.GetInstance()->GetScripts(EmitterScripts, false, true);
+				FVersionedNiagaraEmitterData* EmitterData = Handle.GetEmitterData();
+				EmitterData->GetScripts(EmitterScripts, false, true);
 
 				for (int32 ScriptIdx = 0; ScriptIdx < EmitterScripts.Num(); ScriptIdx++)
 				{
@@ -986,7 +984,7 @@ TSharedPtr<FNiagaraCompileRequestDataBase, ESPMode::ThreadSafe> FNiagaraEditorMo
 					ENiagaraScriptUsage::ParticleSimulationStageScript };
 
 				TArray<FNiagaraVariable> OldStaticVars = BasePtr->EmitterData[i]->StaticVariables;
-				BasePtr->EmitterData[i]->FinishPrecompile(ParticleEncounterableVars, OldStaticVars, ConstantResolver, ParticleUsages, &Handle.GetInstance()->GetSimulationStages(), EmitterNames);
+				BasePtr->EmitterData[i]->FinishPrecompile(ParticleEncounterableVars, OldStaticVars, ConstantResolver, ParticleUsages, &EmitterData->GetSimulationStages(), EmitterNames);
 			}
 		}
 
@@ -994,17 +992,27 @@ TSharedPtr<FNiagaraCompileRequestDataBase, ESPMode::ThreadSafe> FNiagaraEditorMo
 			TMap<UNiagaraScript*, UNiagaraScript*> ScriptDependencyMap;
 
 			// Prepare rapid iteration parameters for execution.
-			for (auto ScriptEmitterPair = ScriptToEmitterMap.CreateIterator(); ScriptEmitterPair; ++ScriptEmitterPair)
+			TArray<TTuple<UNiagaraScript*, FVersionedNiagaraEmitter>> ScriptsToIterate;
+			for (const auto& Entry : ScriptToEmitterMap)
 			{
-				UNiagaraScript* CompiledScript = ScriptEmitterPair->Key;
-				const UNiagaraEmitter* Emitter = ScriptEmitterPair->Value;
+				ScriptsToIterate.Add(Entry);
+			}
+			for (int Index = 0; Index < ScriptsToIterate.Num(); Index++)
+			{
+				TTuple<UNiagaraScript*, FVersionedNiagaraEmitter> ScriptEmitterPair = ScriptsToIterate[Index];
+				UNiagaraScript* CompiledScript = ScriptEmitterPair.Key;
+				FVersionedNiagaraEmitterData* EmitterData = ScriptEmitterPair.Value.GetEmitterData();
 
 				if (UNiagaraScript::IsEquivalentUsage(CompiledScript->GetUsage(), ENiagaraScriptUsage::EmitterSpawnScript))
 				{
 					UNiagaraScript* SystemSpawnScript = System->GetSystemSpawnScript();
 					Scripts.AddUnique(SystemSpawnScript);
 					ScriptDependencyMap.Add(CompiledScript, SystemSpawnScript);
-					ScriptToEmitterMap.Add(SystemSpawnScript, nullptr);
+					if (!ScriptToEmitterMap.Contains(SystemSpawnScript))
+					{
+						ScriptToEmitterMap.Add(SystemSpawnScript);
+						ScriptsToIterate.Emplace(SystemSpawnScript, FVersionedNiagaraEmitter());
+					}
 				}
 
 				if (UNiagaraScript::IsEquivalentUsage(CompiledScript->GetUsage(), ENiagaraScriptUsage::EmitterUpdateScript))
@@ -1012,36 +1020,53 @@ TSharedPtr<FNiagaraCompileRequestDataBase, ESPMode::ThreadSafe> FNiagaraEditorMo
 					UNiagaraScript* SystemUpdateScript = System->GetSystemUpdateScript();
 					Scripts.AddUnique(SystemUpdateScript);
 					ScriptDependencyMap.Add(CompiledScript, SystemUpdateScript);
-					ScriptToEmitterMap.Add(SystemUpdateScript, nullptr);
+					ScriptToEmitterMap.Add(SystemUpdateScript);
+					if (!ScriptToEmitterMap.Contains(SystemUpdateScript))
+					{
+						ScriptToEmitterMap.Add(SystemUpdateScript);
+						ScriptsToIterate.Emplace(SystemUpdateScript, FVersionedNiagaraEmitter());
+					}
 				}
 
 				if (UNiagaraScript::IsEquivalentUsage(CompiledScript->GetUsage(), ENiagaraScriptUsage::ParticleSpawnScript))
 				{
-					if (Emitter && Emitter->SimTarget == ENiagaraSimTarget::GPUComputeSim)
+					if (EmitterData && EmitterData->SimTarget == ENiagaraSimTarget::GPUComputeSim)
 					{
-						UNiagaraScript* ComputeScript = const_cast<UNiagaraScript*>(Emitter->GetGPUComputeScript());
+						UNiagaraScript* ComputeScript = EmitterData->GetGPUComputeScript();
 
 						Scripts.AddUnique(ComputeScript);
 						ScriptDependencyMap.Add(CompiledScript, ComputeScript);
-						ScriptToEmitterMap.Add(ComputeScript, Emitter);
+						if (!ScriptToEmitterMap.Contains(ComputeScript))
+						{
+							ScriptToEmitterMap.Add(ComputeScript, ScriptEmitterPair.Value);
+							ScriptsToIterate.Emplace(ComputeScript, ScriptEmitterPair.Value);
+						}
 					}
 				}
 
 				if (UNiagaraScript::IsEquivalentUsage(CompiledScript->GetUsage(), ENiagaraScriptUsage::ParticleUpdateScript))
 				{
-					if (Emitter && Emitter->SimTarget == ENiagaraSimTarget::GPUComputeSim)
+					if (EmitterData && EmitterData->SimTarget == ENiagaraSimTarget::GPUComputeSim)
 					{
-						UNiagaraScript* ComputeScript = const_cast<UNiagaraScript*>(Emitter->GetGPUComputeScript());
+						UNiagaraScript* ComputeScript = EmitterData->GetGPUComputeScript();
 
 						Scripts.AddUnique(ComputeScript);
 						ScriptDependencyMap.Add(CompiledScript, ComputeScript);
-						ScriptToEmitterMap.Add(ComputeScript, Emitter);
+						if (!ScriptToEmitterMap.Contains(ComputeScript))
+						{
+							ScriptToEmitterMap.Add(ComputeScript, ScriptEmitterPair.Value);
+							ScriptsToIterate.Emplace(ComputeScript, ScriptEmitterPair.Value);
+						}
 					}
-					else if (Emitter && Emitter->bInterpolatedSpawning)
+					else if (EmitterData && EmitterData->bInterpolatedSpawning)
 					{
-						Scripts.AddUnique(Emitter->SpawnScriptProps.Script);
-						ScriptDependencyMap.Add(CompiledScript, Emitter->SpawnScriptProps.Script);
-						ScriptToEmitterMap.Add(Emitter->SpawnScriptProps.Script, Emitter);
+						Scripts.AddUnique(EmitterData->SpawnScriptProps.Script);
+						ScriptDependencyMap.Add(CompiledScript, EmitterData->SpawnScriptProps.Script);
+						if (!ScriptToEmitterMap.Contains(EmitterData->SpawnScriptProps.Script))
+						{
+							ScriptToEmitterMap.Add(EmitterData->SpawnScriptProps.Script, ScriptEmitterPair.Value);
+							ScriptsToIterate.Emplace(EmitterData->SpawnScriptProps.Script, ScriptEmitterPair.Value);
+						}
 					}
 				}
 			}
@@ -1059,7 +1084,7 @@ TSharedPtr<FNiagaraCompileRequestDataBase, ESPMode::ThreadSafe> FNiagaraEditorMo
 				if (Handle.GetIsEnabled()) // Don't pull in the emitter if it isn't going to be used.
 				{
 					TArray<UNiagaraScript*> EmitterScripts;
-					Handle.GetInstance()->GetScripts(EmitterScripts, false);
+					Handle.GetEmitterData()->GetScripts(EmitterScripts, false);
 
 					for (int32 ScriptIdx = 0; ScriptIdx < EmitterScripts.Num(); ScriptIdx++)
 					{
@@ -1211,13 +1236,13 @@ TSharedPtr<FNiagaraCompileRequestDuplicateDataBase, ESPMode::ThreadSafe> FNiagar
 		{
 			const FNiagaraEmitterHandle& Handle = OwningSystem->GetEmitterHandle(i);
 			TSharedPtr<FNiagaraCompileRequestDuplicateData, ESPMode::ThreadSafe> EmitterPtr = MakeShared<FNiagaraCompileRequestDuplicateData, ESPMode::ThreadSafe>();
-			EmitterPtr->EmitterUniqueName = Handle.GetInstance()->GetUniqueEmitterName();
+			EmitterPtr->EmitterUniqueName = Handle.GetInstance().Emitter->GetUniqueEmitterName();
 			EmitterPtr->ValidUsages = BasePtr->ValidUsages;
 			EmitterPtr->SharedSourceGraphToDuplicatedGraphsMap = BasePtr->SharedSourceGraphToDuplicatedGraphsMap;
 			EmitterPtr->SharedNameToDuplicatedDataInterfaceMap = BasePtr->SharedNameToDuplicatedDataInterfaceMap;
 			EmitterPtr->SharedDataInterfaceClassToDuplicatedCDOMap = BasePtr->SharedDataInterfaceClassToDuplicatedCDOMap;
 			//EmitterPtr->bSimulationStagesEnabled = Handle.GetInstance()->bSimulationStagesEnabled;
-			if (Handle.GetIsEnabled() && (OwningEmitter == nullptr || OwningEmitter == Handle.GetInstance())) // Don't need to copy the graph if we aren't going to use it.
+			if (Handle.GetIsEnabled() && (OwningEmitter == nullptr || OwningEmitter == Handle.GetInstance().Emitter)) // Don't need to copy the graph if we aren't going to use it.
 			{
 				EmitterPtr->DeepCopyGraphs(Handle.GetInstance());
 			}
@@ -1253,14 +1278,14 @@ TSharedPtr<FNiagaraCompileRequestDuplicateDataBase, ESPMode::ThreadSafe> FNiagar
 			TArray<FNiagaraVariable> EncounterableEmitterVariables;
 			OwningSystemRequestData->GetDependentRequest(i)->GatherPreCompiledVariables(FString(), EncounterableEmitterVariables);
 
-			if (Handle.GetIsEnabled() && (OwningEmitter == nullptr || OwningEmitter == Handle.GetInstance()))
+			if (Handle.GetIsEnabled() && (OwningEmitter == nullptr || OwningEmitter == Handle.GetInstance().Emitter))
 			{
-				TSharedPtr<FNiagaraGraphCachedDataBase, ESPMode::ThreadSafe> CachedTraversalEmitterData = Handle.GetInstance()->GetCachedTraversalData();
+				TSharedPtr<FNiagaraGraphCachedDataBase, ESPMode::ThreadSafe> CachedTraversalEmitterData = Handle.GetInstance().Emitter->GetCachedTraversalData(Handle.GetInstance().Version);
 				TArray<FNiagaraVariable> StaticVariablesFromEmitter = StaticVariablesFromSystem;
 				StaticVariablesFromEmitter.Append(((FNiagaraGraphCachedBuiltHistory*)CachedTraversalEmitterData.Get())->StaticVariables);
 
 				FCompileConstantResolver ConstantResolver(Handle.GetInstance(), ENiagaraScriptUsage::EmitterSpawnScript);
-				BasePtr->EmitterData[i]->FinishPrecompileDuplicate(EncounterableEmitterVariables, StaticVariablesFromEmitter, ConstantResolver, &Handle.GetInstance()->GetSimulationStages(), ((FNiagaraCompileRequestData*)OwningSystemRequestData)->EmitterData[i]->RapidIterationParams);
+				BasePtr->EmitterData[i]->FinishPrecompileDuplicate(EncounterableEmitterVariables, StaticVariablesFromEmitter, ConstantResolver, &Handle.GetEmitterData()->GetSimulationStages(), ((FNiagaraCompileRequestData*)OwningSystemRequestData)->EmitterData[i]->RapidIterationParams);
 
 			}
 
@@ -1385,11 +1410,12 @@ TSharedPtr<FNiagaraGraphCachedDataBase, ESPMode::ThreadSafe> FNiagaraEditorModul
 			}*/
 		}
 		SrcUniqueEmitterName = Emitter->GetUniqueEmitterName();
-		Emitter->GatherStaticVariables(SrcStaticVariables);
+		const FVersionedNiagaraEmitterData* EmitterData = Emitter->GetEmitterData(Version);
+		EmitterData->GatherStaticVariables(SrcStaticVariables);
 
-		ScriptSource = CastChecked<UNiagaraScriptSource>(Emitter->GraphSource);
-		SimStages = &Emitter->GetSimulationStages();
-		ConstantResolver = FCompileConstantResolver(Emitter, ENiagaraScriptUsage::EmitterSpawnScript);
+		ScriptSource = CastChecked<UNiagaraScriptSource>(EmitterData->GraphSource);
+		SimStages = &EmitterData->GetSimulationStages();
+		ConstantResolver = FCompileConstantResolver(FVersionedNiagaraEmitter(const_cast<UNiagaraEmitter*>(Emitter), Version), ENiagaraScriptUsage::EmitterSpawnScript);
 	}
 
 
diff --git a/Engine/Plugins/FX/Niagara/Source/NiagaraEditor/Private/NiagaraEditorCommands.cpp b/Engine/Plugins/FX/Niagara/Source/NiagaraEditor/Private/NiagaraEditorCommands.cpp
--- a/Engine/Plugins/FX/Niagara/Source/NiagaraEditor/Private/NiagaraEditorCommands.cpp
+++ b/Engine/Plugins/FX/Niagara/Source/NiagaraEditor/Private/NiagaraEditorCommands.cpp
@@ -14,6 +14,7 @@ void FNiagaraEditorCommands::RegisterCommands()
 	UI_COMMAND(Compile, "Compile", "Compile the current scripts", EUserInterfaceActionType::Button, FInputChord());
 	UI_COMMAND(RefreshNodes, "Refresh", "Refreshes the current graph nodes, and updates pins due to external changes.", EUserInterfaceActionType::Button, FInputChord());
 	UI_COMMAND(ModuleVersioning, "Versioning", "Manage different versions of a module script.", EUserInterfaceActionType::Button, FInputChord());
+	UI_COMMAND(EmitterVersioning, "Versioning", "Manage different versions of an emittert.", EUserInterfaceActionType::Button, FInputChord());
 	UI_COMMAND(SelectNextUsage, "Select Next Usage", "Selects the next usage of the selected item.", EUserInterfaceActionType::Button, FInputChord());
 	UI_COMMAND(CreateAssetFromSelection, "Create Asset...", "Creates an asset from the current selection.", EUserInterfaceActionType::Button, FInputChord());
 
