Added versioning for Niagara emitters

diff --git a/Engine/Plugins/FX/Niagara/Source/NiagaraEditor/Private/NiagaraEditorModule.cpp b/Engine/Plugins/FX/Niagara/Source/NiagaraEditor/Private/NiagaraEditorModule.cpp
--- a/Engine/Plugins/FX/Niagara/Source/NiagaraEditor/Private/NiagaraEditorModule.cpp
+++ b/Engine/Plugins/FX/Niagara/Source/NiagaraEditor/Private/NiagaraEditorModule.cpp
@@ -117,7 +117,6 @@
 #include "NiagaraScriptVariable.h"
 #include "NiagaraScript.h"
 #include "NiagaraCommon.h"
-#include "NiagaraScriptHighlight.h"
 #include "NiagaraComponentRendererProperties.h"
 #include "NiagaraMeshRendererProperties.h"
 
@@ -139,7 +138,6 @@
 #include "ISourceControlModule.h"
 #include "DeviceProfiles/DeviceProfileManager.h"
 #include "Containers/Ticker.h"
-#include "NiagaraConstants.h"
 
 #include "ViewModels/Stack/NiagaraStackObjectIssueGenerator.h"
 #include "NiagaraPlatformSet.h"
@@ -330,10 +328,10 @@ void DumpRapidIterationParametersForScript(UNiagaraScript* Script, const FString
 	DumpParameterStore(Script->RapidIterationParameters);
 }
 
-void DumpRapidIterationParametersForEmitter(UNiagaraEmitter* Emitter, const FString& EmitterName)
+void DumpRapidIterationParametersForEmitter(FVersionedNiagaraEmitter Emitter, const FString& EmitterName)
 {
 	TArray<UNiagaraScript*> Scripts;
-	Emitter->GetScripts(Scripts, false);
+	Emitter.GetEmitterData()->GetScripts(Scripts, false);
 	for (UNiagaraScript* Script : Scripts)
 	{
 		DumpRapidIterationParametersForScript(Script, EmitterName);
@@ -364,7 +362,7 @@ void DumpRapidIterationParamersForAsset(const TArray<FString>& Arguments)
 				UNiagaraEmitter* EmitterAsset = Cast<UNiagaraEmitter>(Asset);
 				if (EmitterAsset != nullptr)
 				{
-					DumpRapidIterationParametersForEmitter(EmitterAsset, EmitterAsset->GetName());
+					DumpRapidIterationParametersForEmitter(FVersionedNiagaraEmitter(EmitterAsset, EmitterAsset->GetExposedVersion().VersionGuid), EmitterAsset->GetName());
 				}
 				else
 				{
@@ -383,35 +381,36 @@ void DumpRapidIterationParamersForAsset(const TArray<FString>& Arguments)
 	}
 }
 
-void CompileEmitterStandAlone(UNiagaraEmitter* Emitter, TSet<UNiagaraEmitter*>& InOutCompiledEmitters)
+void CompileEmitterStandAlone(FVersionedNiagaraEmitter VersionedEmitter, TSet<FVersionedNiagaraEmitter>& InOutCompiledEmitters)
 {
-	if (InOutCompiledEmitters.Contains(Emitter) == false)
+	if (InOutCompiledEmitters.Contains(VersionedEmitter) == false)
 	{
-		if (Emitter->GetParent() != nullptr)
+		FVersionedNiagaraEmitterData* EmitterData = VersionedEmitter.GetEmitterData();
+		if (EmitterData && EmitterData->GetParent().Emitter != nullptr)
 		{
 			// If the emitter has a parent emitter make sure to compile that one first.
-			CompileEmitterStandAlone(Emitter->GetParent(), InOutCompiledEmitters);
+			CompileEmitterStandAlone(EmitterData->GetParent(), InOutCompiledEmitters);
 
-			if (Emitter->IsSynchronizedWithParent() == false)
+			if (VersionedEmitter.Emitter->IsSynchronizedWithParent() == false)
 			{
 				// If compiling the parent caused it to become out of sync with the current emitter merge in changes before compiling.
-				Emitter->MergeChangesFromParent();
+				VersionedEmitter.Emitter->MergeChangesFromParent();
 			}
 		}
 
-		Emitter->MarkPackageDirty();
+		VersionedEmitter.Emitter->MarkPackageDirty();
 		UNiagaraSystem* TransientSystem = NewObject<UNiagaraSystem>(GetTransientPackage(), FName("StandaloneEmitter_TempSystem"), RF_Transient);
 		UNiagaraSystemFactoryNew::InitializeSystem(TransientSystem, true);
-		TransientSystem->AddEmitterHandle(*Emitter, TEXT("Emitter"));
+		TransientSystem->AddEmitterHandle(*VersionedEmitter.Emitter, TEXT("Emitter"), VersionedEmitter.Version);
 		FNiagaraStackGraphUtilities::RebuildEmitterNodes(*TransientSystem);
 		TransientSystem->RequestCompile(false);
 		TransientSystem->WaitForCompilationComplete();
 
-		InOutCompiledEmitters.Add(Emitter);
+		InOutCompiledEmitters.Add(VersionedEmitter);
 	}
 }
 
-void PreventSystemRecompile(FAssetData SystemAsset, TSet<UNiagaraEmitter*>& InOutCompiledEmitters)
+void PreventSystemRecompile(FAssetData SystemAsset, TSet<FVersionedNiagaraEmitter>& InOutCompiledEmitters)
 {
 	UNiagaraSystem* System = Cast<UNiagaraSystem>(SystemAsset.GetAsset());
 	if (System != nullptr)
@@ -442,7 +441,7 @@ void PreventSystemRecompile(const TArray<FString>& Arguments)
 				SystemAsset = AssetsInPackage[0];
 			}
 		}
-		TSet<UNiagaraEmitter*> CompiledEmitters;
+		TSet<FVersionedNiagaraEmitter> CompiledEmitters;
 		PreventSystemRecompile(SystemAsset, CompiledEmitters);
 	}
 }
@@ -457,7 +456,7 @@ void PreventAllSystemRecompiles()
 	TArray<FAssetData> SystemAssets;
 	AssetRegistryModule.Get().GetAssetsByClass(UNiagaraSystem::StaticClass()->GetFName(), SystemAssets);
 
-	TSet<UNiagaraEmitter*> CompiledEmitters;
+	TSet<FVersionedNiagaraEmitter> CompiledEmitters;
 	int32 ItemIndex = 0;
 	for (FAssetData& SystemAsset : SystemAssets)
 	{
@@ -536,7 +535,6 @@ void UpgradeAllNiagaraAssets()
 void MakeIndent(int32 IndentLevel, FString& OutIndentString)
 {
 	OutIndentString.Reserve(IndentLevel * 2);
-	int32 InsertStart = OutIndentString.Len();
 	for (int32 i = 0; i < IndentLevel * 2; i++)
 	{
 		OutIndentString.AppendChar(TCHAR(' '));
@@ -575,15 +573,14 @@ void DumpCompileIdDataForScript(UNiagaraScript* Script, int32 IndentLevel, FStri
 	}
 }
 
-void DumpCompileIdDataForEmitter(UNiagaraEmitter* Emitter, int32 IndentLevel, FString& Dump)
+void DumpCompileIdDataForEmitter(const FVersionedNiagaraEmitter& VersionedEmitter, int32 IndentLevel, FString& Dump)
 {
-
 	FString Indent;
 	MakeIndent(IndentLevel, Indent);
-	Dump.Append(FString::Printf(TEXT("%sEmitter: %s\n"), *Indent, *Emitter->GetUniqueEmitterName()));
+	Dump.Append(FString::Printf(TEXT("%sEmitter: %s\n"), *Indent, *VersionedEmitter.Emitter->GetUniqueEmitterName()));
 
 	TArray<UNiagaraScript*> Scripts;
-	Emitter->GetScripts(Scripts, false);
+	VersionedEmitter.GetEmitterData()->GetScripts(Scripts, false);
 	for (UNiagaraScript* Script : Scripts)
 	{
 		DumpCompileIdDataForScript(Script, IndentLevel + 1, Dump);
@@ -1603,7 +1600,6 @@ void FNiagaraEditorModule::GetDataInterfaceFeedbackSafe(UNiagaraDataInterface* I
 		return;
 
 	UNiagaraSystem* OwningSystem = InDataInterface->GetTypedOuter<UNiagaraSystem>();
-	UNiagaraEmitter* OwningEmitter = InDataInterface->GetTypedOuter<UNiagaraEmitter>();
 	UNiagaraComponent* OwningComponent = InDataInterface->GetTypedOuter<UNiagaraComponent>();
 
 	if (OwningSystem == nullptr)
@@ -1627,7 +1623,6 @@ void FNiagaraEditorModule::GetDataInterfaceFeedbackSafe(UNiagaraDataInterface* I
 			if (SystemViewModel->IsValid() && SystemViewModel->GetPlaceholderDataInterfaceManager()->TryGetOwnerInformation(InDataInterface, OwningEmitterHandle, OwningFunctionCallNode))
 			{
 				OwningSystem = &SystemViewModel->GetSystem();
-				OwningEmitter = OwningEmitterHandle.IsValid() ? SystemViewModel->GetEmitterHandleViewModelById(OwningEmitterHandle)->GetEmitterViewModel()->GetEmitter() : nullptr;
 				break;
 			}
 		}
@@ -1652,10 +1647,10 @@ const TArray<TWeakObjectPtr<UNiagaraParameterDefinitions>>& FNiagaraEditorModule
 	return ParameterDefinitionsAssetCache.Get();
 }
 
-void FNiagaraEditorModule::GetTargetSystemAndEmitterForDataInterface(UNiagaraDataInterface* InDataInterface, UNiagaraSystem*& OutOwningSystem, UNiagaraEmitter*& OutOwningEmitter)
+void FNiagaraEditorModule::GetTargetSystemAndEmitterForDataInterface(UNiagaraDataInterface* InDataInterface, UNiagaraSystem*& OutOwningSystem, FVersionedNiagaraEmitter& OutOwningEmitter)
 {
 	OutOwningSystem = InDataInterface->GetTypedOuter<UNiagaraSystem>();
-	OutOwningEmitter = InDataInterface->GetTypedOuter<UNiagaraEmitter>();
+	OutOwningEmitter.Emitter = InDataInterface->GetTypedOuter<UNiagaraEmitter>();
 	
 	if (OutOwningSystem == nullptr)
 	{
@@ -1679,7 +1674,10 @@ void FNiagaraEditorModule::GetTargetSystemAndEmitterForDataInterface(UNiagaraDat
 			if (SystemViewModel->IsValid() && SystemViewModel->GetPlaceholderDataInterfaceManager()->TryGetOwnerInformation(InDataInterface, OwningEmitterHandle, OwningFunctionCallNode))
 			{
 				OutOwningSystem = &SystemViewModel->GetSystem();
-				OutOwningEmitter = OwningEmitterHandle.IsValid() ? SystemViewModel->GetEmitterHandleViewModelById(OwningEmitterHandle)->GetEmitterViewModel()->GetEmitter() : nullptr;
+				if (OwningEmitterHandle.IsValid())
+				{
+					OutOwningEmitter = SystemViewModel->GetEmitterHandleViewModelById(OwningEmitterHandle)->GetEmitterViewModel()->GetEmitter();
+				}
 				break;
 			}
 		}
diff --git a/Engine/Plugins/FX/Niagara/Source/NiagaraEditor/Private/NiagaraEditorUtilities.cpp b/Engine/Plugins/FX/Niagara/Source/NiagaraEditor/Private/NiagaraEditorUtilities.cpp
--- a/Engine/Plugins/FX/Niagara/Source/NiagaraEditor/Private/NiagaraEditorUtilities.cpp
+++ b/Engine/Plugins/FX/Niagara/Source/NiagaraEditor/Private/NiagaraEditorUtilities.cpp
@@ -28,6 +28,7 @@
 #include "NiagaraParameterMapHistory.h"
 #include "NiagaraParameterDefinitions.h"
 #include "NiagaraScript.h"
+#include "NiagaraScriptMergeManager.h"
 #include "NiagaraScriptSource.h"
 #include "NiagaraSettings.h"
 #include "NiagaraSimulationStageBase.h"
@@ -45,11 +46,11 @@
 #include "Editor/EditorEngine.h"
 #include "Framework/MultiBox/MultiBoxBuilder.h"
 #include "Framework/Notifications/NotificationManager.h"
-#include "HAL/PlatformApplicationMisc.h"
 #include "HAL/PlatformFileManager.h"
 #include "Interfaces/IPluginManager.h"
 #include "NiagaraSystem.h"
 #include "NiagaraSystemEditorData.h"
+#include "NiagaraSystemToolkit.h"
 #include "ScopedTransaction.h"
 #include "Selection.h"
 #include "UpgradeNiagaraScriptResults.h"
@@ -75,7 +76,7 @@ static int GNiagaraDeletePythonFilesOnError = 1;
 static FAutoConsoleVariableRef CVarDeletePythonFilesOnError(
 	TEXT("fx.Niagara.DeletePythonFilesOnError"),
 	GNiagaraDeletePythonFilesOnError,
-	TEXT("This determines whether we keep the intermediate python used by module versioning around when they were executed and resulted in an error."),
+	TEXT("This determines whether we keep the intermediate python used by module/emitter versioning around when they were executed and resulted in an error."),
 	ECVF_Default
 );
 
@@ -141,7 +142,7 @@ void FNiagaraEditorUtilities::InitializeParameterInputNode(UNiagaraNodeInput& In
 	else if(Type.IsDataInterface())
 	{
 		InputNode.Input.AllocateData(); // Frees previously used memory if we're switching from a struct to a class type.
-		InputNode.SetDataInterface(NewObject<UNiagaraDataInterface>(&InputNode, const_cast<UClass*>(Type.GetClass()), NAME_None, RF_Transactional | RF_Public));
+		InputNode.SetDataInterface(NewObject<UNiagaraDataInterface>(&InputNode, Type.GetClass(), NAME_None, RF_Transactional | RF_Public));
 	}
 }
 
@@ -437,8 +438,7 @@ bool FNiagaraEditorUtilities::NestedPropertiesAppendCompileHash(const void* Cont
 	// We special case FNiagaraTypeDefinitions here because they need to write out a lot more than just their standalone uproperties.
 	if (Struct == FNiagaraTypeDefinition::StaticStruct())
 	{
-		FNiagaraTypeDefinition* TypeDef = (FNiagaraTypeDefinition*)Container;
-		if (TypeDef)
+		if (FNiagaraTypeDefinition* TypeDef = (FNiagaraTypeDefinition*)Container)
 		{
 			TypeDef->AppendCompileHash(InVisitor);
 			return true;
@@ -446,8 +446,7 @@ bool FNiagaraEditorUtilities::NestedPropertiesAppendCompileHash(const void* Cont
 	}
 	else if (Struct == FNiagaraTypeDefinitionHandle::StaticStruct())
 	{
-		FNiagaraTypeDefinitionHandle* TypeDef = (FNiagaraTypeDefinitionHandle*)Container;
-		if (TypeDef)
+		if (FNiagaraTypeDefinitionHandle* TypeDef = (FNiagaraTypeDefinitionHandle*)Container)
 		{
 			TypeDef->AppendCompileHash(InVisitor);
 			return true;
@@ -672,7 +671,10 @@ void FNiagaraEditorUtilities::GatherChangeIds(UNiagaraEmitter& Emitter, TMap<FGu
 	ChangeIds.Empty();
 	TArray<UNiagaraGraph*> Graphs;
 	TArray<UNiagaraScript*> Scripts;
-	Emitter.GetScripts(Scripts);
+	for (FNiagaraAssetVersion& Version : Emitter.GetAllAvailableVersions())
+	{
+		Emitter.GetEmitterData(Version.VersionGuid)->GetScripts(Scripts);
+	}
 
 	// First gather all the graphs used by this emitter..
 	for (UNiagaraScript* Script : Scripts)
@@ -726,7 +728,7 @@ void FNiagaraEditorUtilities::GatherChangeIds(UNiagaraEmitter& Emitter, TMap<FGu
 
 	if (bWriteToLogDir)
 	{
-		FNiagaraEditorUtilities::WriteTextFileToDisk(FPaths::ProjectLogDir(), InDebugName + TEXT(".txt"), ExportText, true);
+		WriteTextFileToDisk(FPaths::ProjectLogDir(), InDebugName + TEXT(".txt"), ExportText, true);
 	}
 	
 }
@@ -900,18 +902,18 @@ TSharedPtr<SWidget> FNiagaraEditorUtilities::CreateInlineErrorText(TAttribute<FT
 				];
 }
 
-void FNiagaraEditorUtilities::CompileExistingEmitters(const TArray<UNiagaraEmitter*>& AffectedEmitters)
+void FNiagaraEditorUtilities::CompileExistingEmitters(const TArray<FVersionedNiagaraEmitter>& AffectedEmitters)
 {
 	TArray<TSharedPtr<FNiagaraSystemViewModel>> ExistingSystemViewModels;
 
 	{
 		FNiagaraSystemUpdateContext UpdateCtx;
 
-		TSet<UNiagaraEmitter*> CompiledEmitters;
-		for (UNiagaraEmitter* Emitter : AffectedEmitters)
+		TSet<FVersionedNiagaraEmitter> CompiledEmitters;
+		for (const FVersionedNiagaraEmitter& VersionedEmitter : AffectedEmitters)
 		{
 			// If we've already compiled this emitter, or it's invalid skip it.
-			if (Emitter == nullptr || CompiledEmitters.Contains(Emitter) || !IsValidChecked(Emitter) || Emitter->IsUnreachable())
+			if (VersionedEmitter.Emitter == nullptr || CompiledEmitters.Contains(VersionedEmitter) || !IsValidChecked(VersionedEmitter.Emitter) || VersionedEmitter.Emitter->IsUnreachable())
 			{
 				continue;
 			}
@@ -920,7 +922,7 @@ void FNiagaraEditorUtilities::CompileExistingEmitters(const TArray<UNiagaraEmitt
 			// of a system.
 			for (TObjectIterator<UNiagaraSystem> SystemIterator; SystemIterator; ++SystemIterator)
 			{
-				if (SystemIterator->ReferencesInstanceEmitter(*Emitter))
+				if (SystemIterator->ReferencesInstanceEmitter(VersionedEmitter))
 				{
 					SystemIterator->RequestCompile(false, &UpdateCtx);
 
@@ -945,7 +947,7 @@ bool FNiagaraEditorUtilities::TryGetEventDisplayName(UNiagaraEmitter* Emitter, F
 {
 	if (Emitter != nullptr)
 	{
-		for (const FNiagaraEventScriptProperties& EventScriptProperties : Emitter->GetEventHandlers())
+		for (const FNiagaraEventScriptProperties& EventScriptProperties : Emitter->GetLatestEmitterData()->GetEventHandlers())
 		{
 			if (EventScriptProperties.Script->GetUsageId() == EventUsageId)
 			{
@@ -1188,7 +1190,7 @@ TSharedPtr<FStructOnScope> FNiagaraEditorUtilities::StaticSwitchDefaultIntToStru
 		BoolValue.SetValue(InStaticSwitchDefaultValue != 0);
 
 		TSharedPtr<FStructOnScope> StructValue = MakeShared<FStructOnScope>(InSwitchType.GetStruct());
-		FMemory::Memcpy(StructValue->GetStructMemory(), (uint8*)(&BoolValue), InSwitchType.GetSize());
+		FMemory::Memcpy(StructValue->GetStructMemory(), &BoolValue, InSwitchType.GetSize());
 
 		return StructValue;
 	}
@@ -1200,7 +1202,7 @@ TSharedPtr<FStructOnScope> FNiagaraEditorUtilities::StaticSwitchDefaultIntToStru
 		IntValue.Value = InStaticSwitchDefaultValue;
 
 		TSharedPtr<FStructOnScope> StructValue = MakeShared<FStructOnScope>(InSwitchType.GetStruct());
-		FMemory::Memcpy(StructValue->GetStructMemory(), (uint8*)(&IntValue), InSwitchType.GetSize());
+		FMemory::Memcpy(StructValue->GetStructMemory(), &IntValue, InSwitchType.GetSize());
 
 		return StructValue;
 	}
@@ -1489,25 +1491,26 @@ UNiagaraScript* FNiagaraEditorUtilities::GetScriptFromSystem(UNiagaraSystem& Sys
 			[EmitterHandleId](const FNiagaraEmitterHandle& EmitterHandle) { return EmitterHandle.GetId() == EmitterHandleId; });
 		if (ScriptEmitterHandle != nullptr)
 		{
+			FVersionedNiagaraEmitterData* EmitterData = ScriptEmitterHandle->GetEmitterData();
 			if (UNiagaraScript::IsEquivalentUsage(Usage, ENiagaraScriptUsage::EmitterSpawnScript))
 			{
-				return ScriptEmitterHandle->GetInstance()->EmitterSpawnScriptProps.Script;
+				return EmitterData->EmitterSpawnScriptProps.Script;
 			}
 			else if (UNiagaraScript::IsEquivalentUsage(Usage, ENiagaraScriptUsage::EmitterUpdateScript))
 			{
-				return ScriptEmitterHandle->GetInstance()->EmitterUpdateScriptProps.Script;
+				return EmitterData->EmitterUpdateScriptProps.Script;
 			}
 			else if (UNiagaraScript::IsEquivalentUsage(Usage, ENiagaraScriptUsage::ParticleSpawnScript))
 			{
-				return ScriptEmitterHandle->GetInstance()->SpawnScriptProps.Script;
+				return EmitterData->SpawnScriptProps.Script;
 			}
 			else if (UNiagaraScript::IsEquivalentUsage(Usage, ENiagaraScriptUsage::ParticleUpdateScript))
 			{
-				return ScriptEmitterHandle->GetInstance()->UpdateScriptProps.Script;
+				return EmitterData->UpdateScriptProps.Script;
 			}
 			else if (UNiagaraScript::IsEquivalentUsage(Usage, ENiagaraScriptUsage::ParticleEventScript))
 			{
-				for (const FNiagaraEventScriptProperties& EventScriptProperties : ScriptEmitterHandle->GetInstance()->GetEventHandlers())
+				for (const FNiagaraEventScriptProperties& EventScriptProperties : EmitterData->GetEventHandlers())
 				{
 					if (EventScriptProperties.Script->GetUsageId() == UsageId)
 					{
@@ -1517,7 +1520,7 @@ UNiagaraScript* FNiagaraEditorUtilities::GetScriptFromSystem(UNiagaraSystem& Sys
 			}
 			else if (UNiagaraScript::IsEquivalentUsage(Usage, ENiagaraScriptUsage::ParticleSimulationStageScript))
 			{
-				for (const UNiagaraSimulationStageBase* SimulationStage : ScriptEmitterHandle->GetInstance()->GetSimulationStages())
+				for (const UNiagaraSimulationStageBase* SimulationStage : EmitterData->GetSimulationStages())
 				{
 					if (SimulationStage && SimulationStage->Script && SimulationStage->Script->GetUsageId() == UsageId)
 					{
@@ -1530,10 +1533,10 @@ UNiagaraScript* FNiagaraEditorUtilities::GetScriptFromSystem(UNiagaraSystem& Sys
 	return nullptr;
 }
 
-const FNiagaraEmitterHandle* FNiagaraEditorUtilities::GetEmitterHandleForEmitter(UNiagaraSystem& System, UNiagaraEmitter& Emitter)
+const FNiagaraEmitterHandle* FNiagaraEditorUtilities::GetEmitterHandleForEmitter(UNiagaraSystem& System, const FVersionedNiagaraEmitter& Emitter)
 {
 	return System.GetEmitterHandles().FindByPredicate(
-		[&Emitter](const FNiagaraEmitterHandle& EmitterHandle) { return EmitterHandle.GetInstance() == &Emitter; });
+		[&Emitter](const FNiagaraEmitterHandle& EmitterHandle) { return EmitterHandle.GetInstance() == Emitter; });
 }
 
 ENiagaraScriptLibraryVisibility FNiagaraEditorUtilities::GetScriptAssetVisibility(const FAssetData& ScriptAssetData)
@@ -2026,7 +2029,7 @@ TArray<UNiagaraComponent*> FNiagaraEditorUtilities::GetComponentsThatReferenceSy
 		{
 			for (const auto& EmitterHandle : ReferencedSystemViewModel.GetSystem().GetEmitterHandles())
 			{
-				if (Component->GetAsset()->UsesEmitter(EmitterHandle.GetInstance()->GetParent()))
+				if (Component->GetAsset()->UsesEmitter(EmitterHandle.GetEmitterData()->GetParent()))
 				{
 					ReferencingComponents.Add(Component);
 				}
@@ -2036,7 +2039,7 @@ TArray<UNiagaraComponent*> FNiagaraEditorUtilities::GetComponentsThatReferenceSy
 	return ReferencingComponents;
 }
 
-const FGuid FNiagaraEditorUtilities::AddEmitterToSystem(UNiagaraSystem& InSystem, UNiagaraEmitter& InEmitterToAdd, bool bCreateCopy)
+const FGuid FNiagaraEditorUtilities::AddEmitterToSystem(UNiagaraSystem& InSystem, UNiagaraEmitter& InEmitterToAdd, FGuid EmitterVersion, bool bCreateCopy)
 {
 	// Kill all system instances before modifying the emitter handle list to prevent accessing deleted data.
 	KillSystemInstances(InSystem);
@@ -2052,19 +2055,19 @@ const FGuid FNiagaraEditorUtilities::AddEmitterToSystem(UNiagaraSystem& InSystem
 	if (SystemEditorData->GetOwningSystemIsPlaceholder() == false)
 	{
 		InSystem.Modify();
-		EmitterHandle = InSystem.AddEmitterHandle(InEmitterToAdd, FNiagaraUtilities::GetUniqueName(InEmitterToAdd.GetFName(), EmitterHandleNames));
+		EmitterHandle = InSystem.AddEmitterHandle(InEmitterToAdd, FNiagaraUtilities::GetUniqueName(InEmitterToAdd.GetFName(), EmitterHandleNames), EmitterVersion);
 	}
 	else if (bCreateCopy)
 	{
 		// When editing an emitter asset we add the emitter as a duplicate so that the parent emitter is duplicated, but it's parent emitter
 		// information is maintained.
 		checkf(InSystem.GetNumEmitters() == 0, TEXT("Can not add multiple emitters to a system being edited in emitter asset mode."));
-		FNiagaraEmitterHandle TemporaryEmitterHandle(InEmitterToAdd);
+		FNiagaraEmitterHandle TemporaryEmitterHandle(InEmitterToAdd, EmitterVersion);
 		EmitterHandle = InSystem.DuplicateEmitterHandle(TemporaryEmitterHandle, *InEmitterToAdd.GetUniqueEmitterName());
 	}
 	else
 	{
-		EmitterHandle = FNiagaraEmitterHandle(InEmitterToAdd);
+		EmitterHandle = FNiagaraEmitterHandle(InEmitterToAdd, EmitterVersion);
 		InSystem.AddEmitterHandleDirect(EmitterHandle);		
 	}
 	
@@ -2086,6 +2089,7 @@ void FNiagaraEditorUtilities::RemoveEmittersFromSystemByEmitterHandleId(UNiagara
 	InSystem.Modify();
 	InSystem.RemoveEmitterHandlesById(EmitterHandleIdsToDelete);
 
+	FNiagaraScriptMergeManager::Get()->ClearMergeAdapterCache();
 	FNiagaraStackGraphUtilities::RebuildEmitterNodes(InSystem);
 	UNiagaraSystemEditorData* SystemEditorData = CastChecked<UNiagaraSystemEditorData>(InSystem.GetEditorData(), ECastCheckedType::NullChecked);
 	SystemEditorData->SynchronizeOverviewGraphWithSystem(InSystem);
@@ -2178,15 +2182,21 @@ bool FNiagaraEditorUtilities::AddParameter(FNiagaraVariable& NewParameterVariabl
 void FNiagaraEditorUtilities::ShowParentEmitterInContentBrowser(TSharedRef<FNiagaraEmitterViewModel> Emitter)
 {
 	FContentBrowserModule& ContentBrowserModule = FModuleManager::Get().LoadModuleChecked<FContentBrowserModule>(TEXT("ContentBrowser"));
-	ContentBrowserModule.Get().SyncBrowserToAssets(TArray<FAssetData> { FAssetData(Emitter->GetParentEmitter()) });
+	ContentBrowserModule.Get().SyncBrowserToAssets(TArray<FAssetData> { FAssetData(Emitter->GetParentEmitter().Emitter) });
 }
 
 void FNiagaraEditorUtilities::OpenParentEmitterForEdit(TSharedRef<FNiagaraEmitterViewModel> Emitter)
 {
-	UNiagaraEmitter* ParentEmitter = const_cast<UNiagaraEmitter*>(Emitter->GetParentEmitter());
-	if (ParentEmitter != nullptr)
+	if (UNiagaraEmitter* ParentEmitter = Emitter->GetParentEmitter().Emitter)
 	{
-		GEditor->GetEditorSubsystem<UAssetEditorSubsystem>()->OpenEditorForAsset(ParentEmitter);
+		ParentEmitter->VersionToOpenInEditor = Emitter->GetParentEmitter().Version;
+		if (GEditor->GetEditorSubsystem<UAssetEditorSubsystem>()->OpenEditorForAsset(ParentEmitter))
+		{
+			if (FNiagaraSystemToolkit* EditorInstance = static_cast<FNiagaraSystemToolkit*>(GEditor->GetEditorSubsystem<UAssetEditorSubsystem>()->FindEditorForAsset(ParentEmitter, true)))
+			{
+				EditorInstance->SwitchToVersion(ParentEmitter->VersionToOpenInEditor);
+			}			
+		}
 	}
 }
 
@@ -2299,7 +2309,7 @@ void FNiagaraEditorUtilities::CreateAssetFromEmitter(TSharedRef<FNiagaraEmitterH
 		return;
 	}
 
-	UNiagaraEmitter* EmitterToCopy = EmitterHandleViewModel->GetEmitterViewModel()->GetEmitter();
+	UNiagaraEmitter* EmitterToCopy = EmitterHandleViewModel->GetEmitterViewModel()->GetEmitter().Emitter;
 	const FString PackagePath = FPackageName::GetLongPackagePath(EmitterToCopy->GetOutermost()->GetName());
 	const FName EmitterName = EmitterToCopy->GetFName();
 	
@@ -2307,9 +2317,12 @@ void FNiagaraEditorUtilities::CreateAssetFromEmitter(TSharedRef<FNiagaraEmitterH
 	
 	// First duplicate the asset so that fixes can be made on the duplicate without modifying the system and before it's saved.
 	UNiagaraEmitter* DuplicateEmitter = CastChecked<UNiagaraEmitter>(StaticDuplicateObject(EmitterToCopy, GetTransientPackage()));
-	FNiagaraScratchPadUtilities::FixExternalScratchPadScriptsForEmitter(SystemViewModel->GetSystem(), *DuplicateEmitter);
+	for (const FNiagaraAssetVersion& Version : DuplicateEmitter->GetAllAvailableVersions())
+	{
+		FNiagaraScratchPadUtilities::FixExternalScratchPadScriptsForEmitter(SystemViewModel->GetSystem(), FVersionedNiagaraEmitter(DuplicateEmitter, Version.VersionGuid));
+	}
 
-	// Save the duplicated emiter.
+	// Save the duplicated emitter.
 	UNiagaraEmitter* CreatedAsset = Cast<UNiagaraEmitter>(AssetToolsModule.Get().DuplicateAssetWithDialogAndTitle(EmitterName.GetPlainNameString(), PackagePath, DuplicateEmitter, LOCTEXT("CreateEmitterAssetDialogTitle", "Create Emitter As")));
 	if (CreatedAsset != nullptr)
 	{
@@ -2342,7 +2355,7 @@ void FNiagaraEditorUtilities::CreateAssetFromEmitter(TSharedRef<FNiagaraEmitterH
 
 		// Replace existing emitter
 		SystemViewModel->DeleteEmitters(TSet<FGuid> { EmitterHandleViewModel->GetId() });
-		TSharedPtr<FNiagaraEmitterHandleViewModel> NewEmitterHandleViewModel = SystemViewModel->AddEmitter(*CreatedAsset);
+		TSharedPtr<FNiagaraEmitterHandleViewModel> NewEmitterHandleViewModel = SystemViewModel->AddEmitter(*CreatedAsset, CreatedAsset->GetExposedVersion().VersionGuid);
 
 		NewEmitterHandleViewModel->SetName(EmitterName);
 
@@ -2895,26 +2908,27 @@ void FNiagaraEditorUtilities::RefreshAllScriptsFromExternalChanges(FRefreshAllSc
 
 	// Now determine if any of these scripts were in Emitters. If so, those emitters should be compiled together. If not, go ahead and compile individually.
 	// Use the existing view models if they exist,as they are already wired into the correct UI.
-	TArray<UNiagaraEmitter*> AffectedEmitters;
+	TArray<FVersionedNiagaraEmitter> AffectedEmitters;
 	for (UNiagaraScript* Script : AffectedScripts)
 	{
 		if (Script->IsParticleScript() || Script->IsEmitterSpawnScript() || Script->IsEmitterUpdateScript())
 		{
-			UNiagaraEmitter* Emitter = Cast<UNiagaraEmitter>(Script->GetTypedOuter<UNiagaraEmitter>());
-			if (Emitter)
+			FVersionedNiagaraEmitter Outer = Script->GetOuterEmitter();
+			if (Outer.Emitter)
 			{
-				AffectedEmitters.AddUnique(Emitter);
+				AffectedEmitters.AddUnique(Outer);
 			}
 		}
 		else if (Script->IsSystemSpawnScript() || Script->IsSystemUpdateScript())
 		{
-			UNiagaraSystem* System = Cast<UNiagaraSystem>(Script->GetTypedOuter<UNiagaraSystem>());
-			if (System)
+			if (UNiagaraSystem* System = Cast<UNiagaraSystem>(Script->GetTypedOuter<UNiagaraSystem>()))
 			{
 				for (int32 i = 0; i < System->GetNumEmitters(); i++)
 				{
-					AffectedEmitters.AddUnique(System->GetEmitterHandle(i).GetInstance());
-					AffectedEmitters.AddUnique((UNiagaraEmitter*)System->GetEmitterHandle(i).GetInstance()->GetParent());
+					FVersionedNiagaraEmitter VersionedEmitter = System->GetEmitterHandle(i).GetInstance();
+					AffectedEmitters.AddUnique(VersionedEmitter);
+					//TODO (mg) is the parent compile necessary?
+					//AffectedEmitters.AddUnique(VersionedEmitter.Emitter->GetParent());
 				}
 			}
 		}
@@ -2991,10 +3005,10 @@ void FNiagaraEditorUtilities::RunPythonUpgradeScripts(UNiagaraNodeFunctionCall*
 			FFileHelper::LoadFileToString(PythonScript, *NewData->ScriptAsset.FilePath);
 		}
 
-		bool bLoggedWarning = false;
 		if (!PythonScript.IsEmpty())
 		{
 			// set up script context
+			bool bLoggedWarning = false;
 			if (SourceNode->SelectedScriptVersion != PreviousData->Version.VersionGuid)
 			{
 				SourceNode->SelectedScriptVersion = PreviousData->Version.VersionGuid;
@@ -3066,6 +3080,162 @@ void FNiagaraEditorUtilities::RunPythonUpgradeScripts(UNiagaraNodeFunctionCall*
 	SourceNode->SelectedScriptVersion = SavedVersion;
 }
 
+void FNiagaraEditorUtilities::RunPythonUpgradeScripts(UUpgradeNiagaraEmitterContext* UpgradeContext)
+{
+	TArray<FVersionedNiagaraEmitterData*> UpgradeData = UpgradeContext->GetUpgradeData();
+	if (UpgradeData.Num() <= 1)
+	{
+		// no script executions found, so we're done here
+		return;
+	}
+
+	FString Warnings;
+	for (int32 i = 1; i < UpgradeData.Num(); i++)
+	{
+		FVersionedNiagaraEmitterData* PreviousData = UpgradeData[i - 1];
+		FVersionedNiagaraEmitterData* NewData = UpgradeData[i];
+		
+		FString PythonScript;
+		if (PreviousData == nullptr || NewData == nullptr || NewData->UpdateScriptExecution == ENiagaraPythonUpdateScriptReference::None)
+		{
+			continue;
+		}
+		
+		if (NewData->UpdateScriptExecution == ENiagaraPythonUpdateScriptReference::DirectTextEntry)
+		{
+			PythonScript = NewData->PythonUpdateScript;
+		}
+		else if (NewData->UpdateScriptExecution == ENiagaraPythonUpdateScriptReference::ScriptAsset && !NewData->ScriptAsset.FilePath.IsEmpty())
+		{
+			FFileHelper::LoadFileToString(PythonScript, *NewData->ScriptAsset.FilePath);
+		}
+
+		if (PythonScript.IsEmpty())
+		{
+			continue;
+		}
+		
+		// set up script context
+		bool bLoggedWarning = false;
+		
+		// save python script to a temp file to execute
+		FString TempScriptFile = FPaths::CreateTempFilename(*FPaths::ProjectIntermediateDir(), TEXT("VersionUpgrade-"), TEXT(".py"));
+		IPlatformFile& PlatformFile = FPlatformFileManager::Get().GetPlatformFile();
+		ON_SCOPE_EXIT
+		{
+			// Delete temp script file
+			if (GNiagaraDeletePythonFilesOnError || !UpgradeContext->bCancelledByPythonError)
+			{
+				PlatformFile.DeleteFile(*TempScriptFile);
+			}
+		};
+		if (!FFileHelper::SaveStringToFile(FString::Format(*PythonUpgradeScriptStub, {UpgradeContext->GetPathName(), PythonScript}), *TempScriptFile))
+		{
+			UE_LOG(LogNiagaraEditor, Error, TEXT("Unable to save python script to file %s"), *TempScriptFile);
+			AddStackWarning(PreviousData->Version, NewData->Version, "Cannot create python script file!", bLoggedWarning, Warnings);
+			continue;
+		}
+
+		UpgradeContext->bCancelledByPythonError = true;
+		FPythonCommandEx PythonCommand = FPythonCommandEx();
+		PythonCommand.ExecutionMode = EPythonCommandExecutionMode::ExecuteFile;
+		PythonCommand.Command = TempScriptFile;
+
+		// execute python script
+		IPythonScriptPlugin::Get()->ExecPythonCommandEx(PythonCommand);
+
+		if (UpgradeContext->bCancelledByPythonError)
+		{
+			UE_LOG(LogNiagaraEditor, Error, TEXT("%s\n\nPython script:\n%s\nTo keep the intermediate script around, set fx.Niagara.DeletePythonFilesOnError to 0."), *PythonCommand.CommandResult, *PythonCommand.Command);
+			AddStackWarning(PreviousData->Version, NewData->Version, "Python script ended with error!", bLoggedWarning, Warnings);
+		}
+		else
+		{
+			//TODO: apply changes?
+		}
+		
+		if (PythonCommand.LogOutput.Num() > 0)
+		{
+			for (FPythonLogOutputEntry& Entry : PythonCommand.LogOutput)
+			{
+				AddStackWarning(PreviousData->Version, NewData->Version, Entry.Output, bLoggedWarning, Warnings);
+			}
+		}
+	}
+
+	//TODO: add warnings to stack
+}
+
+TSharedPtr<FNiagaraSystemViewModel> CreateSystemViewModelForUpgrade(UNiagaraSystem* System)
+{
+	TSharedPtr<FNiagaraSystemViewModel> SystemViewModel = MakeShared<FNiagaraSystemViewModel>();
+	FNiagaraSystemViewModelOptions SystemOptions;
+	SystemOptions.bCanModifyEmittersFromTimeline = false;
+	SystemOptions.bCanSimulate = false;
+	SystemOptions.bCanAutoCompile = false;
+	SystemOptions.bIsForDataProcessingOnly = true;
+	SystemOptions.MessageLogGuid = System->GetAssetGuid();
+	SystemViewModel->Initialize(*System, SystemOptions);
+	return SystemViewModel;
+}
+
+void FNiagaraEditorUtilities::SwitchParentEmitterVersion(TSharedRef<FNiagaraEmitterViewModel> EmitterViewModel, TSharedRef<FNiagaraSystemViewModel> SystemViewModel, const FGuid& NewVersionGuid)
+{
+	// if we want to run the python upgrade script, we make a snapshot of the existing emitter viewmodel
+	UNiagaraPythonEmitter* OldPythonEmitter = NewObject<UNiagaraPythonEmitter>(GetTransientPackage());
+	FVersionedNiagaraEmitter Parent = EmitterViewModel->GetEmitter().GetEmitterData()->GetParent();
+	bool bRunPython = Parent.GetEmitterData() && Parent.Emitter->IsVersioningEnabled();
+	int32 ViewModelIndex = 0;
+	TSharedPtr<FNiagaraSystemViewModel> OldSystemViewModel;
+	if (bRunPython)
+	{
+		UNiagaraSystem* ExistingSystem = CastChecked<UNiagaraSystem>(StaticDuplicateObject(&SystemViewModel->GetSystem(), GetTransientPackage()));
+		OldSystemViewModel = CreateSystemViewModelForUpgrade(ExistingSystem);
+
+		TArray<TSharedRef<FNiagaraEmitterHandleViewModel>> EmitterHandleViewModels = SystemViewModel->GetEmitterHandleViewModels();
+		for (ViewModelIndex = 0; ViewModelIndex < EmitterHandleViewModels.Num(); ViewModelIndex++)
+		{
+			if (FNiagaraEmitterHandle* EmitterHandle = EmitterHandleViewModels[ViewModelIndex]->GetEmitterHandle())
+			{
+				if (EmitterHandle->GetInstance() == EmitterViewModel->GetEmitter())
+				{
+					OldPythonEmitter->Init(OldSystemViewModel->GetEmitterHandleViewModels()[ViewModelIndex]);
+					break;
+				}
+			}
+		}
+	}
+	
+	FScopedTransaction Transaction(LOCTEXT("ChangeEmitterVersion", "Change Parent Emitter Version"));
+
+	// change the parent, merge the changes and reset the stack
+	EmitterViewModel->PreviousEmitterVersion = EmitterViewModel->GetParentEmitter().Version;
+	UNiagaraEmitter* Emitter = EmitterViewModel->GetEmitter().Emitter;
+	Emitter->ChangeParentVersion(NewVersionGuid, EmitterViewModel->GetEmitter().Version);
+
+	SystemViewModel->GetSystem().KillAllActiveCompilations();
+	SystemViewModel->RefreshAll();
+	SystemViewModel->GetSelectionViewModel()->EmptySelection();
+	SystemViewModel->GetSelectionViewModel()->AddEntryToSelectionByDisplayedObjectDeferred(Emitter);
+
+	// optionally run python scripts
+	if (bRunPython)
+	{
+		UNiagaraPythonEmitter* NewPythonEmitter = NewObject<UNiagaraPythonEmitter>(GetTransientPackage());
+		NewPythonEmitter->Init(SystemViewModel->GetEmitterHandleViewModels()[ViewModelIndex]);
+		UUpgradeNiagaraEmitterContext* UpgradeContext = NewObject<UUpgradeNiagaraEmitterContext>(GetTransientPackage());
+		UpgradeContext->Init(OldPythonEmitter, NewPythonEmitter);
+		
+		RunPythonUpgradeScripts(UpgradeContext);
+
+		// purge temp objects
+		OldSystemViewModel->GetSystem().MarkAsGarbage();
+		OldSystemViewModel->Cleanup();
+		OldSystemViewModel.Reset();
+		CollectGarbage(GARBAGE_COLLECTION_KEEPFLAGS);
+	}
+}
+
 bool FNiagaraParameterUtilities::DoesParameterNameMatchSearchText(FName ParameterName, const FString& SearchTextString)
 {
 	FNiagaraParameterHandle ParameterHandle(ParameterName);
diff --git a/Engine/Plugins/FX/Niagara/Source/NiagaraEditor/Private/NiagaraEmitterDetailsCustomization.cpp b/Engine/Plugins/FX/Niagara/Source/NiagaraEditor/Private/NiagaraEmitterDetailsCustomization.cpp
--- a/Engine/Plugins/FX/Niagara/Source/NiagaraEditor/Private/NiagaraEmitterDetailsCustomization.cpp
+++ b/Engine/Plugins/FX/Niagara/Source/NiagaraEditor/Private/NiagaraEmitterDetailsCustomization.cpp
@@ -7,6 +7,8 @@
 #include "NiagaraEmitter.h"
 #include "NiagaraEditorModule.h"
 #include "Toolkits/NiagaraSystemToolkit.h"
+#include "ViewModels/NiagaraEmitterViewModel.h"
+#include "ViewModels/TNiagaraViewModelManager.h"
 
 TSharedRef<IDetailCustomization> FNiagaraEmitterDetails::MakeInstance()
 {
@@ -15,11 +17,76 @@ TSharedRef<IDetailCustomization> FNiagaraEmitterDetails::MakeInstance()
 
 void FNiagaraEmitterDetails::CustomizeDetails(IDetailLayoutBuilder& InDetailLayout)
 {
-	TSharedPtr<IPropertyHandle> EventHandlersPropertyHandle = InDetailLayout.GetProperty(UNiagaraEmitter::PrivateMemberNames::EventHandlerScriptProps);
-	EventHandlersPropertyHandle->MarkHiddenByCustomization();
+	TArray<TWeakObjectPtr<UObject>> ObjectsBeingCustomized;
+	InDetailLayout.GetObjectsBeingCustomized(ObjectsBeingCustomized);
+	
+	if (UNiagaraEmitter* EmitterBeingCustomized = CastChecked<UNiagaraEmitter>(ObjectsBeingCustomized[0]))
+	{
+		TArray<TSharedPtr<FNiagaraEmitterViewModel>> ExistingViewModels;
+		TNiagaraViewModelManager<UNiagaraEmitter, FNiagaraEmitterViewModel>::GetAllViewModelsForObject(EmitterBeingCustomized, ExistingViewModels);
+		FVersionedNiagaraEmitter VersionedNiagaraEmitter = ExistingViewModels[0]->GetEmitter();
+		FVersionedNiagaraEmitterData* EmitterData = VersionedNiagaraEmitter.GetEmitterData();
+
+		for (FProperty* ChildProperty : TFieldRange<FProperty>(FVersionedNiagaraEmitterData::StaticStruct()))
+		{
+			if (ChildProperty->HasAllPropertyFlags(CPF_Edit))
+			{
+				FName Category = FName(ChildProperty->GetMetaData(TEXT("Category")));
+				IDetailCategoryBuilder& CategoryBuilder = InDetailLayout.EditCategory(Category);
+				
+				TSharedPtr<FStructOnScope> StructData = MakeShareable(new FStructOnScope(FVersionedNiagaraEmitterData::StaticStruct(), reinterpret_cast<uint8*>(EmitterData)));
+				if (IDetailPropertyRow* PropertyRow = CategoryBuilder.AddExternalStructureProperty(StructData, ChildProperty->GetFName()))
+				{
+					PropertyRow->GetPropertyHandle()->SetOnPropertyValueChanged(FSimpleDelegate::CreateLambda([VersionedNiagaraEmitter, ChildProperty]()
+					{
+						FPropertyChangedEvent ChangeEvent(ChildProperty);
+						VersionedNiagaraEmitter.Emitter->PostEditChangeVersionedProperty(ChangeEvent, VersionedNiagaraEmitter.Version);
+					}));
+				}
+			}
+		}
+	}
+	
 	if(GNiagaraScalabilityModeEnabled)
 	{
 		InDetailLayout.HideCategory("Scalability");
 	}
 }
 
+TSharedRef<IDetailCustomization> FNiagaraEmitterScalabilityDetails::MakeInstance()
+{
+	return MakeShared<FNiagaraEmitterScalabilityDetails>();
+}
+
+void FNiagaraEmitterScalabilityDetails::CustomizeDetails(IDetailLayoutBuilder& InDetailLayout)
+{
+	TArray<TWeakObjectPtr<UObject>> ObjectsBeingCustomized;
+	InDetailLayout.GetObjectsBeingCustomized(ObjectsBeingCustomized);
+	
+	if (UNiagaraEmitter* EmitterBeingCustomized = CastChecked<UNiagaraEmitter>(ObjectsBeingCustomized[0]))
+	{
+		TArray<TSharedPtr<FNiagaraEmitterViewModel>> ExistingViewModels;
+		TNiagaraViewModelManager<UNiagaraEmitter, FNiagaraEmitterViewModel>::GetAllViewModelsForObject(EmitterBeingCustomized, ExistingViewModels);
+		FVersionedNiagaraEmitter VersionedNiagaraEmitter = ExistingViewModels[0]->GetEmitter();
+		FVersionedNiagaraEmitterData* EmitterData = VersionedNiagaraEmitter.GetEmitterData();
+
+		for (FProperty* ChildProperty : TFieldRange<FProperty>(FVersionedNiagaraEmitterData::StaticStruct()))
+		{
+			if (ChildProperty->HasAllPropertyFlags(CPF_Edit) && ChildProperty->HasMetaData(TEXT("DisplayInScalabilityContext")))
+			{
+				FName Category = FName(ChildProperty->GetMetaData(TEXT("Category")));
+				IDetailCategoryBuilder& CategoryBuilder = InDetailLayout.EditCategory(FName(Category));
+				
+				TSharedPtr<FStructOnScope> StructData = MakeShareable(new FStructOnScope(FVersionedNiagaraEmitterData::StaticStruct(), reinterpret_cast<uint8*>(EmitterData)));
+				if (IDetailPropertyRow* PropertyRow = CategoryBuilder.AddExternalStructureProperty(StructData, ChildProperty->GetFName()))
+				{
+					PropertyRow->GetPropertyHandle()->SetOnPropertyValueChanged(FSimpleDelegate::CreateLambda([VersionedNiagaraEmitter, ChildProperty]()
+					{
+						FPropertyChangedEvent ChangeEvent(ChildProperty);
+						VersionedNiagaraEmitter.Emitter->PostEditChangeVersionedProperty(ChangeEvent, VersionedNiagaraEmitter.Version);
+					}));
+				}
+			}
+		}
+	}
+}
diff --git a/Engine/Plugins/FX/Niagara/Source/NiagaraEditor/Private/NiagaraEmitterDetailsCustomization.h b/Engine/Plugins/FX/Niagara/Source/NiagaraEditor/Private/NiagaraEmitterDetailsCustomization.h
--- a/Engine/Plugins/FX/Niagara/Source/NiagaraEditor/Private/NiagaraEmitterDetailsCustomization.h
+++ b/Engine/Plugins/FX/Niagara/Source/NiagaraEditor/Private/NiagaraEmitterDetailsCustomization.h
@@ -20,3 +20,13 @@ public:
 	virtual void CustomizeDetails(IDetailLayoutBuilder& InDetailLayout) override;
 };
 
+// customization for scalability mode
+class FNiagaraEmitterScalabilityDetails : public IDetailCustomization
+{
+public:
+	/** Makes a new instance of this detail layout class for a specific detail view requesting it */
+	static TSharedRef<IDetailCustomization> MakeInstance();
+
+	/** IDetailCustomization interface */
+	virtual void CustomizeDetails(IDetailLayoutBuilder& InDetailLayout) override;
+};
\ No newline at end of file
diff --git a/Engine/Plugins/FX/Niagara/Source/NiagaraEditor/Private/NiagaraEmitterFactoryNew.cpp b/Engine/Plugins/FX/Niagara/Source/NiagaraEditor/Private/NiagaraEmitterFactoryNew.cpp
--- a/Engine/Plugins/FX/Niagara/Source/NiagaraEditor/Private/NiagaraEmitterFactoryNew.cpp
+++ b/Engine/Plugins/FX/Niagara/Source/NiagaraEditor/Private/NiagaraEmitterFactoryNew.cpp
@@ -19,19 +19,17 @@
 #include "Interfaces/IMainFrameModule.h"
 #include "Framework/Application/SlateApplication.h"
 
-#include "Misc/ConfigCacheIni.h"
-
 #define LOCTEXT_NAMESPACE "NiagaraEmitterFactory"
 
 UNiagaraEmitterFactoryNew::UNiagaraEmitterFactoryNew(const FObjectInitializer& ObjectInitializer)
-: Super(ObjectInitializer)
+	: Super(ObjectInitializer)
 {
 	SupportedClass = UNiagaraEmitter::StaticClass();
-	bCreateNew = false;
 	bEditAfterNew = true;
 	bCreateNew = true;
 	EmitterToCopy = nullptr;
 	bUseInheritance = false;
+	bAddDefaultModulesAndRenderersToEmptyEmitter = true;
 }
 
 bool UNiagaraEmitterFactoryNew::ConfigureProperties()
@@ -125,12 +123,13 @@ UObject* UNiagaraEmitterFactoryNew::FactoryCreateNew(UClass* Class, UObject* InP
 	{
 		if (bUseInheritance && EmitterToCopy->TemplateSpecification == ENiagaraScriptTemplateSpecification::None)
 		{
-			NewEmitter = UNiagaraEmitter::CreateWithParentAndOwner(*EmitterToCopy, InParent, Name, Flags);
+			NewEmitter = UNiagaraEmitter::CreateWithParentAndOwner(FVersionedNiagaraEmitter(EmitterToCopy, EmitterToCopy->GetExposedVersion().VersionGuid), InParent, Name, Flags);
 		}
 		else
 		{
 			NewEmitter = Cast<UNiagaraEmitter>(StaticDuplicateObject(EmitterToCopy, InParent, Name, Flags, Class));
 			NewEmitter->SetUniqueEmitterName(Name.GetPlainNameString());
+			NewEmitter->DisableVersioning(EmitterToCopy->GetExposedVersion().VersionGuid);
 		}
 
 		NewEmitter->TemplateSpecification = ENiagaraScriptTemplateSpecification::None;
@@ -149,39 +148,41 @@ UObject* UNiagaraEmitterFactoryNew::FactoryCreateNew(UClass* Class, UObject* InP
 void UNiagaraEmitterFactoryNew::InitializeEmitter(UNiagaraEmitter* NewEmitter, bool bAddDefaultModulesAndRenderers)
 {
 	{
-		NewEmitter->SimTarget = ENiagaraSimTarget::CPUSim;
+		NewEmitter->CheckVersionDataAvailable();
+		FVersionedNiagaraEmitterData* EmitterData = NewEmitter->GetLatestEmitterData();
+		EmitterData->SimTarget = ENiagaraSimTarget::CPUSim;
 
 		UNiagaraScriptSource* Source = NewObject<UNiagaraScriptSource>(NewEmitter, NAME_None, RF_Transactional);
 		UNiagaraGraph* CreatedGraph = NewObject<UNiagaraGraph>(Source, NAME_None, RF_Transactional);
 		Source->NodeGraph = CreatedGraph;
 
 		// Fix up source pointers.
-		NewEmitter->GraphSource = Source;
-		NewEmitter->SpawnScriptProps.Script->SetLatestSource(Source);
-		NewEmitter->UpdateScriptProps.Script->SetLatestSource(Source);
-		NewEmitter->EmitterSpawnScriptProps.Script->SetLatestSource(Source);
-		NewEmitter->EmitterUpdateScriptProps.Script->SetLatestSource(Source);
-		NewEmitter->GetGPUComputeScript()->SetLatestSource(Source);
+		EmitterData->GraphSource = Source;
+		EmitterData->SpawnScriptProps.Script->SetLatestSource(Source);
+		EmitterData->UpdateScriptProps.Script->SetLatestSource(Source);
+		EmitterData->EmitterSpawnScriptProps.Script->SetLatestSource(Source);
+		EmitterData->EmitterUpdateScriptProps.Script->SetLatestSource(Source);
+		EmitterData->GetGPUComputeScript()->SetLatestSource(Source);
 
 		// Initialize the scripts for output.
-		UNiagaraNodeOutput* EmitterSpawnOutputNode = FNiagaraStackGraphUtilities::ResetGraphForOutput(*Source->NodeGraph, ENiagaraScriptUsage::EmitterSpawnScript, NewEmitter->EmitterSpawnScriptProps.Script->GetUsageId());
-		UNiagaraNodeOutput* EmitterUpdateOutputNode = FNiagaraStackGraphUtilities::ResetGraphForOutput(*Source->NodeGraph, ENiagaraScriptUsage::EmitterUpdateScript, NewEmitter->EmitterUpdateScriptProps.Script->GetUsageId());
-		UNiagaraNodeOutput* ParticleSpawnOutputNode = FNiagaraStackGraphUtilities::ResetGraphForOutput(*Source->NodeGraph, ENiagaraScriptUsage::ParticleSpawnScript, NewEmitter->SpawnScriptProps.Script->GetUsageId());
-		UNiagaraNodeOutput* ParticleUpdateOutputNode = FNiagaraStackGraphUtilities::ResetGraphForOutput(*Source->NodeGraph, ENiagaraScriptUsage::ParticleUpdateScript, NewEmitter->UpdateScriptProps.Script->GetUsageId());
+		UNiagaraNodeOutput* EmitterSpawnOutputNode = FNiagaraStackGraphUtilities::ResetGraphForOutput(*Source->NodeGraph, ENiagaraScriptUsage::EmitterSpawnScript, EmitterData->EmitterSpawnScriptProps.Script->GetUsageId());
+		UNiagaraNodeOutput* EmitterUpdateOutputNode = FNiagaraStackGraphUtilities::ResetGraphForOutput(*Source->NodeGraph, ENiagaraScriptUsage::EmitterUpdateScript, EmitterData->EmitterUpdateScriptProps.Script->GetUsageId());
+		UNiagaraNodeOutput* ParticleSpawnOutputNode = FNiagaraStackGraphUtilities::ResetGraphForOutput(*Source->NodeGraph, ENiagaraScriptUsage::ParticleSpawnScript, EmitterData->SpawnScriptProps.Script->GetUsageId());
+		UNiagaraNodeOutput* ParticleUpdateOutputNode = FNiagaraStackGraphUtilities::ResetGraphForOutput(*Source->NodeGraph, ENiagaraScriptUsage::ParticleUpdateScript, EmitterData->UpdateScriptProps.Script->GetUsageId());
 
 		checkf(EmitterSpawnOutputNode != nullptr && EmitterUpdateOutputNode != nullptr && ParticleSpawnOutputNode != nullptr && ParticleUpdateOutputNode != nullptr,
 			TEXT("Failed to create output nodes for emitter scripts."));
 
 		if (bAddDefaultModulesAndRenderers)
 		{
-			NewEmitter->AddRenderer(NewObject<UNiagaraSpriteRendererProperties>(NewEmitter, "Renderer"));
+			NewEmitter->AddRenderer(NewObject<UNiagaraSpriteRendererProperties>(NewEmitter, "Renderer"), EmitterData->Version.VersionGuid);
 
 			AddModuleFromAssetPath(TEXT("/Niagara/Modules/Emitter/EmitterLifeCycle.EmitterLifeCycle"), *EmitterUpdateOutputNode);
 			
 			UNiagaraNodeFunctionCall* SpawnRateNode = AddModuleFromAssetPath(TEXT("/Niagara/Modules/Emitter/SpawnRate.SpawnRate"), *EmitterUpdateOutputNode);
 			if (SpawnRateNode != nullptr)
 			{
-				SetRapidIterationParameter(NewEmitter->GetUniqueEmitterName(), *NewEmitter->EmitterUpdateScriptProps.Script, *SpawnRateNode,
+				SetRapidIterationParameter(NewEmitter->GetUniqueEmitterName(), *EmitterData->EmitterUpdateScriptProps.Script, *SpawnRateNode,
 					"SpawnRate", FNiagaraTypeDefinition::GetFloatDef(), 10.0f);
 			}
 
@@ -190,7 +191,7 @@ void UNiagaraEmitterFactoryNew::InitializeEmitter(UNiagaraEmitter* NewEmitter, b
 			UNiagaraNodeFunctionCall* AddVelocityNode = AddModuleFromAssetPath(TEXT("/Niagara/Modules/Spawn/Velocity/AddVelocity.AddVelocity"), *ParticleSpawnOutputNode);
 			if (AddVelocityNode != nullptr)
 			{
-				SetRapidIterationParameter(NewEmitter->GetUniqueEmitterName(), *NewEmitter->SpawnScriptProps.Script, *AddVelocityNode,
+				SetRapidIterationParameter(NewEmitter->GetUniqueEmitterName(), *EmitterData->SpawnScriptProps.Script, *AddVelocityNode,
 					"Velocity", FNiagaraTypeDefinition::GetVec3Def(), FVector3f(0.0f, 0.0f, 100.0f));
 			}
 
@@ -216,10 +217,10 @@ void UNiagaraEmitterFactoryNew::InitializeEmitter(UNiagaraEmitter* NewEmitter, b
 		}
 
 		FNiagaraStackGraphUtilities::RelayoutGraph(*Source->NodeGraph);
-		NewEmitter->bInterpolatedSpawning = true;
-		NewEmitter->bDeterminism = false; // NOTE: Default to non-determinism
+		EmitterData->bInterpolatedSpawning = true;
+		EmitterData->bDeterminism = false; // NOTE: Default to non-determinism
 		NewEmitter->TemplateSpecification = ENiagaraScriptTemplateSpecification::None;
-		NewEmitter->SpawnScriptProps.Script->SetUsage(ENiagaraScriptUsage::ParticleSpawnScriptInterpolated);
+		EmitterData->SpawnScriptProps.Script->SetUsage(ENiagaraScriptUsage::ParticleSpawnScriptInterpolated);
 	}
 }
 
diff --git a/Engine/Plugins/FX/Niagara/Source/NiagaraEditor/Private/NiagaraGraph.cpp b/Engine/Plugins/FX/Niagara/Source/NiagaraEditor/Private/NiagaraGraph.cpp
--- a/Engine/Plugins/FX/Niagara/Source/NiagaraEditor/Private/NiagaraGraph.cpp
+++ b/Engine/Plugins/FX/Niagara/Source/NiagaraEditor/Private/NiagaraGraph.cpp
@@ -706,6 +706,25 @@ void UNiagaraGraph::ChangeParameterType(const TArray<FNiagaraVariable>& Paramete
 	InvalidateCachedParameterData();
 }
 
+void UNiagaraGraph::ReplaceScriptReferences(UNiagaraScript* OldScript, UNiagaraScript* NewScript)
+{
+	bool bChanged = false;
+	for (UEdGraphNode* Node : Nodes)
+	{
+		UNiagaraNodeFunctionCall* FunctionNode = Cast<UNiagaraNodeFunctionCall>(Node);
+		if (FunctionNode && FunctionNode->FunctionScript == OldScript)
+		{
+			FunctionNode->Modify();
+			FunctionNode->FunctionScript = NewScript;
+			bChanged = true;
+		}
+	}
+	if (bChanged)
+	{
+		MarkGraphRequiresSynchronization(TEXT("Script references updated"));
+	}
+}
+
 TArray<UEdGraphPin*> UNiagaraGraph::FindParameterMapDefaultValuePins(const FName VariableName) const
 {
 	TArray<UEdGraphPin*> DefaultPins;
@@ -1982,7 +2001,8 @@ bool UNiagaraGraph::RenameParameterFromPin(const FNiagaraVariable& Parameter, FN
 		return true;
 	}
 
-	TSharedPtr<FNiagaraScriptToolkitParameterPanelViewModel> ParameterPanelViewModel = TNiagaraViewModelManager<UNiagaraScript, FNiagaraScriptToolkitParameterPanelViewModel>::GetExistingViewModelForObject(GetTypedOuter<UNiagaraScript>());
+	UNiagaraScript* OuterScript = GetTypedOuter<UNiagaraScript>();
+	TSharedPtr<FNiagaraScriptToolkitParameterPanelViewModel> ParameterPanelViewModel = TNiagaraViewModelManager<UNiagaraScript, FNiagaraScriptToolkitParameterPanelViewModel>::GetExistingViewModelForObject(OuterScript);
 	if (ParameterPanelViewModel.IsValid() == false)
 	{
 		ensureMsgf(false, TEXT("Failed to get parameter panel view model when renaming parameter via pin in graph!"));
@@ -2254,6 +2274,24 @@ void UNiagaraGraph::ScriptVariableChanged(FNiagaraVariable Variable)
 	NotifyGraphChanged();
 }
 
+FVersionedNiagaraEmitter UNiagaraGraph::GetOwningEmitter() const
+{
+	UNiagaraEmitter* OwningEmitter = GetTypedOuter<UNiagaraEmitter>();
+	if (OwningEmitter)
+	{
+		for (FNiagaraAssetVersion Version : OwningEmitter->GetAllAvailableVersions())
+		{
+			FVersionedNiagaraEmitterData* EmitterData = OwningEmitter->GetEmitterData(Version.VersionGuid);
+			UNiagaraGraph* EmitterGraph = Cast<UNiagaraScriptSource>(EmitterData->GraphSource)->NodeGraph;
+			if (this == EmitterGraph)
+			{
+				return FVersionedNiagaraEmitter(OwningEmitter, Version.VersionGuid);
+			}
+		}
+	}
+	return FVersionedNiagaraEmitter();
+}
+
 bool UNiagaraGraph::SynchronizeScriptVariable(const UNiagaraScriptVariable* SourceScriptVar, UNiagaraScriptVariable* DestScriptVar /*= nullptr*/, bool bIgnoreChangeId /*= false*/)
 {
 	if (DestScriptVar == nullptr)
diff --git a/Engine/Plugins/FX/Niagara/Source/NiagaraEditor/Private/NiagaraHlslTranslator.h b/Engine/Plugins/FX/Niagara/Source/NiagaraEditor/Private/NiagaraHlslTranslator.h
--- a/Engine/Plugins/FX/Niagara/Source/NiagaraEditor/Private/NiagaraHlslTranslator.h
+++ b/Engine/Plugins/FX/Niagara/Source/NiagaraEditor/Private/NiagaraHlslTranslator.h
@@ -165,7 +165,7 @@ public:
 
 	void DuplicateReferencedGraphs(UNiagaraGraph* InSrcGraph, UNiagaraGraph* InDupeGraph, ENiagaraScriptUsage InUsage, FCompileConstantResolver ConstantResolver, TMap<UNiagaraNodeFunctionCall*, ENiagaraScriptUsage> FunctionsWithUsage = TMap<UNiagaraNodeFunctionCall*, ENiagaraScriptUsage>());
 	void DeepCopyGraphs(UNiagaraScriptSource* ScriptSource, ENiagaraScriptUsage InUsage, FCompileConstantResolver ConstantResolver);
-	void DeepCopyGraphs(UNiagaraEmitter* Emitter);
+	void DeepCopyGraphs(const FVersionedNiagaraEmitter& Emitter);
 	
 	void FinishPrecompileDuplicate(const TArray<FNiagaraVariable>& EncounterableVariables, const TArray<FNiagaraVariable>& InStaticVariables, FCompileConstantResolver ConstantResolver, const TArray<class UNiagaraSimulationStageBase*>* SimStages, const TArray<FNiagaraVariable>& InParamStore);
 	void CreateDataInterfaceCDO(TArrayView<UClass*> VariableDataInterfaces);
diff --git a/Engine/Plugins/FX/Niagara/Source/NiagaraEditor/Private/NiagaraNodeAssignment.cpp b/Engine/Plugins/FX/Niagara/Source/NiagaraEditor/Private/NiagaraNodeAssignment.cpp
--- a/Engine/Plugins/FX/Niagara/Source/NiagaraEditor/Private/NiagaraNodeAssignment.cpp
+++ b/Engine/Plugins/FX/Niagara/Source/NiagaraEditor/Private/NiagaraNodeAssignment.cpp
@@ -1,7 +1,6 @@
 // Copyright Epic Games, Inc. All Rights Reserved.
 
 #include "NiagaraNodeAssignment.h"
-#include "UObject/UnrealType.h"
 #include "NiagaraGraph.h"
 #include "NiagaraScriptSource.h"
 #include "NiagaraScript.h"
@@ -9,14 +8,12 @@
 #include "NiagaraNodeInput.h"
 #include "NiagaraNodeOutput.h"
 #include "EdGraphSchema_Niagara.h"
-#include "Modules/ModuleManager.h"
 #include "NiagaraComponent.h"
 #include "NiagaraEditorUtilities.h"
 #include "NiagaraNodeParameterMapGet.h"
 #include "NiagaraNodeParameterMapSet.h"
 #include "NiagaraConstants.h"
 #include "ViewModels/Stack/NiagaraParameterHandle.h"
-#include "Framework/MultiBox/MultiBoxBuilder.h"
 #include "ScopedTransaction.h"
 #include "ViewModels/NiagaraParameterPanelViewModel.h"
 #include "ViewModels/Stack/NiagaraStackGraphUtilities.h"
@@ -125,8 +122,7 @@ void UNiagaraNodeAssignment::PostLoad()
 		if (Pin != nullptr && Pin->LinkedTo.Num() == 1)
 		{
 			// Likely we have a set node going into us, check to see if it has any variables that need to be cleaned up.
-			UNiagaraNodeParameterMapSet* SetNode = Cast<UNiagaraNodeParameterMapSet>(Pin->LinkedTo[0]->GetOwningNode());
-			if (SetNode)
+			if (UNiagaraNodeParameterMapSet* SetNode = Cast<UNiagaraNodeParameterMapSet>(Pin->LinkedTo[0]->GetOwningNode()))
 			{
 				SetNode->ConditionalPostLoad();
 
@@ -179,7 +175,18 @@ void UNiagaraNodeAssignment::PostLoad()
 			TArray<UNiagaraScript*> Scripts;
 			if (Emitter)
 			{
-				Emitter->GetScripts(Scripts, false);
+				if (UNiagaraScript* OuterScript = GetTypedOuter<UNiagaraScript>())
+				{
+					OuterScript->GetOuterEmitter().GetEmitterData()->GetScripts(Scripts, false);
+				}
+				else if (UNiagaraScriptSource* OuterSource = GetTypedOuter<UNiagaraScriptSource>())
+				{
+					OuterSource->GetOuterEmitter().GetEmitterData()->GetScripts(Scripts, false);
+				}
+				else
+				{
+					ensureMsgf(false, TEXT("Unable to find path from emitter %s to %s"), *Emitter->GetPathName(), *this->GetPathName());
+				}
 			}
 			if (System)
 			{
@@ -489,8 +496,7 @@ void UNiagaraNodeAssignment::InitializeScript(UNiagaraScript* NewScript)
 
 				if (FNiagaraConstants::IsNiagaraConstant(AssignmentTargets[i]))
 				{
-					const FNiagaraVariableMetaData* FoundMetaData = FNiagaraConstants::GetConstantMetaData(AssignmentTargets[i]);
-					if (FoundMetaData)
+					if (const FNiagaraVariableMetaData* FoundMetaData = FNiagaraConstants::GetConstantMetaData(AssignmentTargets[i]))
 					{
 						FNiagaraVariableMetaData NewMetaData;
 						TOptional<FNiagaraVariableMetaData> ExistingMetaData = CreatedGraph->GetMetaData(TargetVar);
diff --git a/Engine/Plugins/FX/Niagara/Source/NiagaraEditor/Private/NiagaraNodeEmitter.cpp b/Engine/Plugins/FX/Niagara/Source/NiagaraEditor/Private/NiagaraNodeEmitter.cpp
--- a/Engine/Plugins/FX/Niagara/Source/NiagaraEditor/Private/NiagaraNodeEmitter.cpp
+++ b/Engine/Plugins/FX/Niagara/Source/NiagaraEditor/Private/NiagaraNodeEmitter.cpp
@@ -5,17 +5,12 @@
 #include "NiagaraEmitter.h"
 #include "NiagaraEditorUtilities.h"
 #include "EdGraphSchema_Niagara.h"
-#include "EdGraph/EdGraph.h"
 #include "NiagaraCommon.h"
-#include "NiagaraDataInterface.h"
 #include "NiagaraScriptSource.h"
 #include "NiagaraGraph.h"
 #include "NiagaraNodeInput.h"
-#include "NiagaraGraph.h"
-#include "NiagaraNodeParameterMapBase.h"
 #include "NiagaraNodeOutput.h"
 #include "NiagaraHlslTranslator.h"
-#include "Framework/MultiBox/MultiBoxBuilder.h"
 #include "Stats/Stats.h"
 #include "NiagaraEditorModule.h"
 
@@ -118,18 +113,6 @@ bool UNiagaraNodeEmitter::GenerateCompileHashForClassMembers(const UClass* InCla
 
 void UNiagaraNodeEmitter::AllocateDefaultPins()
 {
-	UNiagaraEmitter* Emitter = nullptr;
-	if (OwnerSystem)
-	{
-		for (int32 i = 0; i < OwnerSystem->GetNumEmitters(); ++i)
-		{
-			if (OwnerSystem->GetEmitterHandle(i).GetId() == EmitterHandleId)
-			{
-				Emitter = OwnerSystem->GetEmitterHandle(i).GetInstance();
-			}
-		}
-	}
-
 	const UEdGraphSchema_Niagara* NiagaraSchema = Cast<UEdGraphSchema_Niagara>(GetSchema());
 	CreatePin(EGPD_Input, NiagaraSchema->TypeDefinitionToPinType(FNiagaraTypeDefinition::GetParameterMapDef()), TEXT("InputMap"));
 	CreatePin(EGPD_Output, NiagaraSchema->TypeDefinitionToPinType(FNiagaraTypeDefinition::GetParameterMapDef()), TEXT("OutputMap"));
@@ -193,22 +176,23 @@ FString UNiagaraNodeEmitter::GetEmitterUniqueName() const
 UNiagaraScriptSource* UNiagaraNodeEmitter::GetScriptSource() const
 {
 	// First get the emitter that we're referencing..
-	UNiagaraEmitter* Emitter = nullptr;
+	FVersionedNiagaraEmitter VersionedEmitter;
 	if (OwnerSystem)
 	{
 		for (int32 i = 0; i < OwnerSystem->GetNumEmitters(); ++i)
 		{
 			if (OwnerSystem->GetEmitterHandle(i).GetId() == EmitterHandleId)
 			{
-				Emitter = OwnerSystem->GetEmitterHandle(i).GetInstance();
+				VersionedEmitter = OwnerSystem->GetEmitterHandle(i).GetInstance();
 			}
 		}
 	}
 
 	// Now get the graph off that emitter
-	if (Emitter != nullptr && Emitter->GraphSource != nullptr)
+	FVersionedNiagaraEmitterData* EmitterData = VersionedEmitter.GetEmitterData();
+	if (EmitterData && EmitterData->GraphSource)
 	{
-		UNiagaraScriptSource* Source = Cast<UNiagaraScriptSource>(Emitter->GraphSource);
+		UNiagaraScriptSource* Source = Cast<UNiagaraScriptSource>(EmitterData->GraphSource);
 		return Source;
 	}
 
@@ -218,22 +202,23 @@ UNiagaraScriptSource* UNiagaraNodeEmitter::GetScriptSource() const
 UNiagaraGraph* UNiagaraNodeEmitter::GetCalledGraph() const
 {
 	// First get the emitter that we're referencing..
-	UNiagaraEmitter* Emitter = nullptr;
+	FVersionedNiagaraEmitter VersionedEmitter;
 	if (OwnerSystem)
 	{
 		for (int32 i = 0; i < OwnerSystem->GetNumEmitters(); ++i)
 		{
 			if (OwnerSystem->GetEmitterHandle(i).GetId() == EmitterHandleId)
 			{
-				Emitter = OwnerSystem->GetEmitterHandle(i).GetInstance();
+				VersionedEmitter = OwnerSystem->GetEmitterHandle(i).GetInstance();
 			}
 		}
 	}
 
 	// Now get the graph off that emitter
-	if (Emitter != nullptr && Emitter->GraphSource != nullptr)
+	FVersionedNiagaraEmitterData* EmitterData = VersionedEmitter.GetEmitterData();
+	if (EmitterData && EmitterData->GraphSource)
 	{
-		UNiagaraScriptSource* Source = Cast<UNiagaraScriptSource>(Emitter->GraphSource);
+		UNiagaraScriptSource* Source = Cast<UNiagaraScriptSource>(EmitterData->GraphSource);
 		if (Source)
 		{
 			return Source->NodeGraph;
@@ -436,8 +421,6 @@ void UNiagaraNodeEmitter::Compile(FHlslNiagaraTranslator *Translator, TArray<int
 
 	check(Outputs.Num() == 0);
 
-	const UEdGraphSchema_Niagara* Schema = CastChecked<UEdGraphSchema_Niagara>(GetSchema());
-
 	// First compile fully down the hierarchy for our predecessors..
 	UNiagaraGraph* CalledGraph = GetCalledGraph();
 	TArray<UNiagaraNodeInput*> InputsNodes;
diff --git a/Engine/Plugins/FX/Niagara/Source/NiagaraEditor/Private/NiagaraNodeOutput.cpp b/Engine/Plugins/FX/Niagara/Source/NiagaraEditor/Private/NiagaraNodeOutput.cpp
--- a/Engine/Plugins/FX/Niagara/Source/NiagaraEditor/Private/NiagaraNodeOutput.cpp
+++ b/Engine/Plugins/FX/Niagara/Source/NiagaraEditor/Private/NiagaraNodeOutput.cpp
@@ -62,17 +62,18 @@ FText UNiagaraNodeOutput::GetPinNameText(UEdGraphPin* Pin) const
 
 TArray<FName> UNiagaraNodeOutput::GetAllStackContextOverrides() const
 {
-	UNiagaraEmitter* Emitter = GetTypedOuter<UNiagaraEmitter>();
+	FVersionedNiagaraEmitter Outer = GetNiagaraGraph()->GetOwningEmitter();
 	TArray<FName> Overrides;
-	if (Emitter)
+	if (Outer.Emitter)
 	{
 		TArray<UNiagaraScript*> Scripts;
-		Emitter->GetScripts(Scripts, false);
+		FVersionedNiagaraEmitterData* EmitterData = Outer.GetEmitterData();
+		EmitterData->GetScripts(Scripts, false);
 		for (UNiagaraScript* Script : Scripts)
 		{
 			if (Script && Script->GetUsage() == ENiagaraScriptUsage::ParticleSimulationStageScript)
 			{
-				UNiagaraSimulationStageBase* Base = Emitter->GetSimulationStageById(Script->GetUsageId());
+				UNiagaraSimulationStageBase* Base = EmitterData->GetSimulationStageById(Script->GetUsageId());
 				if (Base)
 				{
 					FName StackContextAlias = Base->GetStackContextReplacementName();
@@ -89,10 +90,10 @@ TOptional<FName> UNiagaraNodeOutput::GetStackContextOverride() const
 {
 	
 	{
-		UNiagaraEmitter* Emitter = GetTypedOuter<UNiagaraEmitter>();
-		if (Emitter && GetUsage() == ENiagaraScriptUsage::ParticleSimulationStageScript)
+		FVersionedNiagaraEmitter Outer = GetNiagaraGraph()->GetOwningEmitter();
+		if (Outer.Emitter && GetUsage() == ENiagaraScriptUsage::ParticleSimulationStageScript)
 		{
-			UNiagaraSimulationStageBase* Base = Emitter->GetSimulationStageById(GetUsageId());
+			UNiagaraSimulationStageBase* Base = Outer.GetEmitterData()->GetSimulationStageById(GetUsageId());
 			if (Base)
 			{
 				FName StackContextAlias = Base->GetStackContextReplacementName();
diff --git a/Engine/Plugins/FX/Niagara/Source/NiagaraEditor/Private/NiagaraParameterMapHistory.cpp b/Engine/Plugins/FX/Niagara/Source/NiagaraEditor/Private/NiagaraParameterMapHistory.cpp
--- a/Engine/Plugins/FX/Niagara/Source/NiagaraEditor/Private/NiagaraParameterMapHistory.cpp
+++ b/Engine/Plugins/FX/Niagara/Source/NiagaraEditor/Private/NiagaraParameterMapHistory.cpp
@@ -2165,36 +2165,37 @@ bool FCompileConstantResolver::ResolveConstant(FNiagaraVariable& OutConstant) co
 	}
 
 	// handle emitter case
-	if (Emitter && OutConstant == FNiagaraVariable(FNiagaraTypeDefinition::GetBoolDef(), TEXT("Emitter.Localspace")))
+	FVersionedNiagaraEmitterData* EmitterData = Emitter.GetEmitterData();
+	if (EmitterData && OutConstant == FNiagaraVariable(FNiagaraTypeDefinition::GetBoolDef(), TEXT("Emitter.Localspace")))
 	{
-		OutConstant.SetValue(Emitter->bLocalSpace ? FNiagaraBool(true) : FNiagaraBool(false));
+		OutConstant.SetValue(EmitterData->bLocalSpace ? FNiagaraBool(true) : FNiagaraBool(false));
 		return true;
 	}
-	if (Emitter && OutConstant == FNiagaraVariable(FNiagaraTypeDefinition::GetBoolDef(), TEXT("Emitter.Determinism")))
+	if (EmitterData && OutConstant == FNiagaraVariable(FNiagaraTypeDefinition::GetBoolDef(), TEXT("Emitter.Determinism")))
 	{
-		OutConstant.SetValue(Emitter->bDeterminism ? FNiagaraBool(true) : FNiagaraBool(false));
+		OutConstant.SetValue(EmitterData->bDeterminism ? FNiagaraBool(true) : FNiagaraBool(false));
 		return true;
 	}
-	if (Emitter && OutConstant == FNiagaraVariable(FNiagaraTypeDefinition::GetBoolDef(), TEXT("Emitter.InterpolatedSpawn")))
+	if (EmitterData && OutConstant == FNiagaraVariable(FNiagaraTypeDefinition::GetBoolDef(), TEXT("Emitter.InterpolatedSpawn")))
 	{
-		OutConstant.SetValue(Emitter->bInterpolatedSpawning ? FNiagaraBool(true) : FNiagaraBool(false));
+		OutConstant.SetValue(EmitterData->bInterpolatedSpawning ? FNiagaraBool(true) : FNiagaraBool(false));
 		return true;
 	}
-	if (Emitter && OutConstant == FNiagaraVariable(FNiagaraTypeDefinition::GetSimulationTargetEnum(), TEXT("Emitter.SimulationTarget")))
+	if (EmitterData && OutConstant == FNiagaraVariable(FNiagaraTypeDefinition::GetSimulationTargetEnum(), TEXT("Emitter.SimulationTarget")))
 	{
 		FNiagaraInt32 EnumValue;
-		EnumValue.Value = (uint8)Emitter->SimTarget;
+		EnumValue.Value = (uint8)EmitterData->SimTarget;
 		OutConstant.SetValue(EnumValue);
 		return true;
 	}
-	if (Emitter && OutConstant == FNiagaraVariable(FNiagaraTypeDefinition::GetScriptUsageEnum(), TEXT("Script.Usage")))
+	if (EmitterData && OutConstant == FNiagaraVariable(FNiagaraTypeDefinition::GetScriptUsageEnum(), TEXT("Script.Usage")))
 	{
 		FNiagaraInt32 EnumValue;
 		EnumValue.Value = (uint8)FNiagaraUtilities::ConvertScriptUsageToStaticSwitchUsage(Usage);
 		OutConstant.SetValue(EnumValue);
 		return true;
 	}
-	if (Emitter && OutConstant == FNiagaraVariable(FNiagaraTypeDefinition::GetScriptContextEnum(), TEXT("Script.Context")))
+	if (EmitterData && OutConstant == FNiagaraVariable(FNiagaraTypeDefinition::GetScriptContextEnum(), TEXT("Script.Context")))
 	{
 		FNiagaraInt32 EnumValue;
 		EnumValue.Value = (uint8)FNiagaraUtilities::ConvertScriptUsageToStaticSwitchContext(Usage);
@@ -2239,7 +2240,7 @@ bool FCompileConstantResolver::ResolveConstant(FNiagaraVariable& OutConstant) co
 ENiagaraFunctionDebugState FCompileConstantResolver::CalculateDebugState() const
 {
 	// System or Emitter
-	const UNiagaraSystem* OwnerSystem = System ? System : (Emitter ? Cast<const UNiagaraSystem>(Emitter->GetOuter()) : nullptr);
+	const UNiagaraSystem* OwnerSystem = System ? System : (Emitter.Emitter ? Cast<const UNiagaraSystem>(Emitter.Emitter->GetOuter()) : nullptr);
 	if ( OwnerSystem != nullptr )
 	{
 		return OwnerSystem->ShouldDisableDebugSwitches() ? ENiagaraFunctionDebugState::NoDebug : DebugState;
