Added versioning for Niagara emitters

diff --git a/Engine/Plugins/FX/Niagara/Source/NiagaraEditor/Private/NiagaraScriptMergeManager.cpp b/Engine/Plugins/FX/Niagara/Source/NiagaraEditor/Private/NiagaraScriptMergeManager.cpp
--- a/Engine/Plugins/FX/Niagara/Source/NiagaraEditor/Private/NiagaraScriptMergeManager.cpp
+++ b/Engine/Plugins/FX/Niagara/Source/NiagaraEditor/Private/NiagaraScriptMergeManager.cpp
@@ -55,7 +55,7 @@ static FAutoConsoleVariableRef CVarForceErrorIfMissingDefaultOnMergeh(
 );
 
 FNiagaraStackFunctionInputOverrideMergeAdapter::FNiagaraStackFunctionInputOverrideMergeAdapter(
-	const UNiagaraEmitter& InOwningEmitter,
+	const FVersionedNiagaraEmitter& InOwningEmitter,
 	UNiagaraScript& InOwningScript,
 	UNiagaraNodeFunctionCall& InOwningFunctionCallNode,
 	UEdGraphPin& InOverridePin
@@ -202,18 +202,19 @@ TOptional<FString> FNiagaraStackFunctionInputOverrideMergeAdapter::GetStaticSwit
 	return StaticSwitchValue;
 }
 
-FNiagaraStackFunctionMergeAdapter::FNiagaraStackFunctionMergeAdapter(const UNiagaraEmitter& InOwningEmitter, UNiagaraScript& InOwningScript, UNiagaraNodeFunctionCall& InFunctionCallNode, int32 InStackIndex)
+FNiagaraStackFunctionMergeAdapter::FNiagaraStackFunctionMergeAdapter(const FVersionedNiagaraEmitter& InOwningEmitter, UNiagaraScript& InOwningScript, UNiagaraNodeFunctionCall& InFunctionCallNode, int32 InStackIndex)
 {
 	OwningScript = &InOwningScript;
 	FunctionCallNode = &InFunctionCallNode;
 	StackIndex = InStackIndex;
 
-	bUsesScratchPadScript = InOwningEmitter.ScratchPadScripts.Contains(FunctionCallNode->FunctionScript) ||
-		InOwningEmitter.ParentScratchPadScripts.Contains(FunctionCallNode->FunctionScript);
+	FVersionedNiagaraEmitterData* EmitterData = InOwningEmitter.GetEmitterData();
+	bUsesScratchPadScript = EmitterData->ScratchPads->Scripts.Contains(FunctionCallNode->FunctionScript) ||
+		EmitterData->ParentScratchPads->Scripts.Contains(FunctionCallNode->FunctionScript);
 
-	FString UniqueEmitterName = InOwningEmitter.GetUniqueEmitterName();
+	FString UniqueEmitterName = InOwningEmitter.Emitter->GetUniqueEmitterName();
 
-	FCompileConstantResolver ConstantResolver(&InOwningEmitter, FNiagaraStackGraphUtilities::GetEmitterOutputNodeForStackNode(*FunctionCallNode)->GetUsage());
+	FCompileConstantResolver ConstantResolver(InOwningEmitter, FNiagaraStackGraphUtilities::GetEmitterOutputNodeForStackNode(*FunctionCallNode)->GetUsage());
 
 	// Collect explicit overrides set via parameter map set nodes.
 	TSet<FName> AliasedInputsAdded;
@@ -235,9 +236,9 @@ FNiagaraStackFunctionMergeAdapter::FNiagaraStackFunctionMergeAdapter(const UNiag
 	if (InFunctionCallNode.FunctionScript != nullptr)
 	{
 		TSet<const UEdGraphPin*> HiddenPins;
-		FCompileConstantResolver Resolver(&InOwningEmitter, ENiagaraScriptUsage::Function);
+		FCompileConstantResolver Resolver(InOwningEmitter, ENiagaraScriptUsage::Function);
 		TArray<const UEdGraphPin*> FunctionInputPins;
-		FNiagaraStackGraphUtilities::GetStackFunctionInputPins(*FunctionCallNode, FunctionInputPins, HiddenPins, Resolver, FNiagaraStackGraphUtilities::ENiagaraGetStackFunctionInputPinsOptions::ModuleInputsOnly, false);
+		GetStackFunctionInputPins(*FunctionCallNode, FunctionInputPins, HiddenPins, Resolver, FNiagaraStackGraphUtilities::ENiagaraGetStackFunctionInputPinsOptions::ModuleInputsOnly, false);
 		
 		for (const UEdGraphPin* FunctionInputPin : FunctionInputPins)
 		{
@@ -384,15 +385,15 @@ const TArray<FNiagaraStackMessage>& FNiagaraStackFunctionMergeAdapter::GetMessag
 	return FunctionCallNode->GetCustomNotes();
 }
 
-FNiagaraScriptStackMergeAdapter::FNiagaraScriptStackMergeAdapter(const UNiagaraEmitter& InOwningEmitter, UNiagaraNodeOutput& InOutputNode, UNiagaraScript& InScript)
+FNiagaraScriptStackMergeAdapter::FNiagaraScriptStackMergeAdapter(const FVersionedNiagaraEmitter& InOwningEmitter, UNiagaraNodeOutput& InOutputNode, UNiagaraScript& InScript)
 {
 	OutputNode = &InOutputNode;
 	InputNode.Reset();
 	Script = &InScript;
-	UniqueEmitterName = InOwningEmitter.GetUniqueEmitterName();
+	UniqueEmitterName = InOwningEmitter.Emitter->GetUniqueEmitterName();
 
 	TArray<FNiagaraStackGraphUtilities::FStackNodeGroup> StackGroups;
-	FNiagaraStackGraphUtilities::GetStackNodeGroups(*OutputNode, StackGroups);
+	GetStackNodeGroups(*OutputNode, StackGroups);
 
 	if (StackGroups.Num() >= 2 && StackGroups[0].EndNode->IsA<UNiagaraNodeInput>())
 	{
@@ -459,24 +460,24 @@ void FNiagaraScriptStackMergeAdapter::GatherFunctionCallNodes(TArray<UNiagaraNod
 	}
 }
 
-FNiagaraEventHandlerMergeAdapter::FNiagaraEventHandlerMergeAdapter(const UNiagaraEmitter& InEmitter, const FNiagaraEventScriptProperties* InEventScriptProperties, UNiagaraNodeOutput* InOutputNode)
+FNiagaraEventHandlerMergeAdapter::FNiagaraEventHandlerMergeAdapter(const FVersionedNiagaraEmitter& InEmitter, const FNiagaraEventScriptProperties* InEventScriptProperties, UNiagaraNodeOutput* InOutputNode)
 {
 	Initialize(InEmitter, InEventScriptProperties, nullptr, InOutputNode);
 }
 
-FNiagaraEventHandlerMergeAdapter::FNiagaraEventHandlerMergeAdapter(const UNiagaraEmitter& InEmitter, FNiagaraEventScriptProperties* InEventScriptProperties, UNiagaraNodeOutput* InOutputNode)
+FNiagaraEventHandlerMergeAdapter::FNiagaraEventHandlerMergeAdapter(const FVersionedNiagaraEmitter& InEmitter, FNiagaraEventScriptProperties* InEventScriptProperties, UNiagaraNodeOutput* InOutputNode)
 {
 	Initialize(InEmitter, InEventScriptProperties, InEventScriptProperties, InOutputNode);
 }
 
-FNiagaraEventHandlerMergeAdapter::FNiagaraEventHandlerMergeAdapter(const UNiagaraEmitter& InEmitter, UNiagaraNodeOutput* InOutputNode)
+FNiagaraEventHandlerMergeAdapter::FNiagaraEventHandlerMergeAdapter(const FVersionedNiagaraEmitter& InEmitter, UNiagaraNodeOutput* InOutputNode)
 {
 	Initialize(InEmitter, nullptr, nullptr, InOutputNode);
 }
 
-void FNiagaraEventHandlerMergeAdapter::Initialize(const UNiagaraEmitter& InEmitter, const FNiagaraEventScriptProperties* InEventScriptProperties, FNiagaraEventScriptProperties* InEditableEventScriptProperties, UNiagaraNodeOutput* InOutputNode)
+void FNiagaraEventHandlerMergeAdapter::Initialize(const FVersionedNiagaraEmitter& InEmitter, const FNiagaraEventScriptProperties* InEventScriptProperties, FNiagaraEventScriptProperties* InEditableEventScriptProperties, UNiagaraNodeOutput* InOutputNode)
 {
-	Emitter = MakeWeakObjectPtr(const_cast<UNiagaraEmitter*>(&InEmitter));
+	Emitter =InEmitter.ToWeakPtr();
 
 	EventScriptProperties = InEventScriptProperties;
 	EditableEventScriptProperties = InEditableEventScriptProperties;
@@ -485,14 +486,14 @@ void FNiagaraEventHandlerMergeAdapter::Initialize(const UNiagaraEmitter& InEmitt
 
 	if (EventScriptProperties != nullptr && OutputNode != nullptr)
 	{
-		EventStack = MakeShared<FNiagaraScriptStackMergeAdapter>(*Emitter.Get(), *OutputNode.Get(), *EventScriptProperties->Script);
+		EventStack = MakeShared<FNiagaraScriptStackMergeAdapter>(InEmitter, *OutputNode.Get(), *EventScriptProperties->Script);
 		InputNode = EventStack->GetInputNode();
 	}
 }
 
-const UNiagaraEmitter* FNiagaraEventHandlerMergeAdapter::GetEmitter() const
+FVersionedNiagaraEmitter FNiagaraEventHandlerMergeAdapter::GetEmitter() const
 {
-	return Emitter.Get();
+	return Emitter.ResolveWeakPtr();
 }
 
 FGuid FNiagaraEventHandlerMergeAdapter::GetUsageId() const
@@ -532,24 +533,24 @@ TSharedPtr<FNiagaraScriptStackMergeAdapter> FNiagaraEventHandlerMergeAdapter::Ge
 	return EventStack;
 }
 
-FNiagaraSimulationStageMergeAdapter::FNiagaraSimulationStageMergeAdapter(const UNiagaraEmitter& InEmitter, const UNiagaraSimulationStageBase* InSimulationStage, int32 InSimulationStageIndex, UNiagaraNodeOutput* InOutputNode)
+FNiagaraSimulationStageMergeAdapter::FNiagaraSimulationStageMergeAdapter(const FVersionedNiagaraEmitter& InEmitter, const UNiagaraSimulationStageBase* InSimulationStage, int32 InSimulationStageIndex, UNiagaraNodeOutput* InOutputNode)
 {
 	Initialize(InEmitter, InSimulationStage, nullptr, InSimulationStageIndex, InOutputNode);
 }
 
-FNiagaraSimulationStageMergeAdapter::FNiagaraSimulationStageMergeAdapter(const UNiagaraEmitter& InEmitter, UNiagaraSimulationStageBase* InSimulationStage, int32 InSimulationStageIndex, UNiagaraNodeOutput* InOutputNode)
+FNiagaraSimulationStageMergeAdapter::FNiagaraSimulationStageMergeAdapter(const FVersionedNiagaraEmitter& InEmitter, UNiagaraSimulationStageBase* InSimulationStage, int32 InSimulationStageIndex, UNiagaraNodeOutput* InOutputNode)
 {
 	Initialize(InEmitter, InSimulationStage, InSimulationStage, InSimulationStageIndex, InOutputNode);
 }
 
-FNiagaraSimulationStageMergeAdapter::FNiagaraSimulationStageMergeAdapter(const UNiagaraEmitter& InEmitter, UNiagaraNodeOutput* InOutputNode)
+FNiagaraSimulationStageMergeAdapter::FNiagaraSimulationStageMergeAdapter(const FVersionedNiagaraEmitter& InEmitter, UNiagaraNodeOutput* InOutputNode)
 {
 	Initialize(InEmitter, nullptr, nullptr, INDEX_NONE, InOutputNode);
 }
 
-void FNiagaraSimulationStageMergeAdapter::Initialize(const UNiagaraEmitter& InEmitter, const UNiagaraSimulationStageBase* InSimulationStage, UNiagaraSimulationStageBase* InEditableSimulationStage, int32 InSimulationStageIndex, UNiagaraNodeOutput* InOutputNode)
+void FNiagaraSimulationStageMergeAdapter::Initialize(const FVersionedNiagaraEmitter& InEmitter, const UNiagaraSimulationStageBase* InSimulationStage, UNiagaraSimulationStageBase* InEditableSimulationStage, int32 InSimulationStageIndex, UNiagaraNodeOutput* InOutputNode)
 {
-	Emitter = MakeWeakObjectPtr(const_cast<UNiagaraEmitter*>(&InEmitter));
+	Emitter = InEmitter.ToWeakPtr();
 
 	SimulationStage = InSimulationStage;
 	EditableSimulationStage = InEditableSimulationStage;
@@ -559,14 +560,14 @@ void FNiagaraSimulationStageMergeAdapter::Initialize(const UNiagaraEmitter& InEm
 
 	if (SimulationStage != nullptr && OutputNode != nullptr)
 	{
-		SimulationStageStack = MakeShared<FNiagaraScriptStackMergeAdapter>(*Emitter.Get(), *OutputNode.Get(), *SimulationStage->Script);
+		SimulationStageStack = MakeShared<FNiagaraScriptStackMergeAdapter>(InEmitter, *OutputNode.Get(), *SimulationStage->Script);
 		InputNode = SimulationStageStack->GetInputNode();
 	}
 }
 
-const UNiagaraEmitter* FNiagaraSimulationStageMergeAdapter::GetEmitter() const
+FVersionedNiagaraEmitter FNiagaraSimulationStageMergeAdapter::GetEmitter() const
 {
-	return Emitter.Get();
+	return Emitter.ResolveWeakPtr();
 }
 
 FGuid FNiagaraSimulationStageMergeAdapter::GetUsageId() const
@@ -637,14 +638,11 @@ const FFunctionInputSummaryViewMetadata& FNiagaraInputSummaryMergeAdapter::GetVa
 }
 
 FNiagaraScratchPadMergeAdapter::FNiagaraScratchPadMergeAdapter()
-	: TargetEmitter(nullptr)
-	, InstanceEmitter(nullptr)
-	, ParentEmitter(nullptr)
-	, bIsInitialized(false)
+	: bIsInitialized(false)
 {
 }
 
-FNiagaraScratchPadMergeAdapter::FNiagaraScratchPadMergeAdapter(UNiagaraEmitter* InTargetEmitter, UNiagaraEmitter* InInstanceEmitter, UNiagaraEmitter* InParentEmitter)
+FNiagaraScratchPadMergeAdapter::FNiagaraScratchPadMergeAdapter(const FVersionedNiagaraEmitter InTargetEmitter, const FVersionedNiagaraEmitter& InInstanceEmitter, const FVersionedNiagaraEmitter& InParentEmitter)
 	: TargetEmitter(InTargetEmitter)
 	, InstanceEmitter(InInstanceEmitter)
 	, ParentEmitter(InParentEmitter)
@@ -665,35 +663,37 @@ UNiagaraScript* FNiagaraScratchPadMergeAdapter::GetScratchPadScriptForFunctionId
 
 void FNiagaraScratchPadMergeAdapter::Initialize()
 {
-	if (TargetEmitter->ScratchPadScripts.Num() > 0 || TargetEmitter->ParentScratchPadScripts.Num() > 0)
+	FVersionedNiagaraEmitterData* TargetEmitterData = TargetEmitter.GetEmitterData();
+	if (TargetEmitterData->ScratchPads->Scripts.Num() > 0 || TargetEmitterData->ParentScratchPads->Scripts.Num() > 0)
 	{
 		// Collect the parent emitters in order
-		TArray<UNiagaraEmitter*> ParentEmitters;
-		UNiagaraEmitter* CurrentParent = ParentEmitter;
-		while (CurrentParent != nullptr)
+		TArray<FVersionedNiagaraEmitter> ParentEmitters;
+		FVersionedNiagaraEmitter CurrentParent = ParentEmitter;
+		while (CurrentParent.Emitter != nullptr)
 		{
 			ParentEmitters.Insert(CurrentParent, 0);
-			CurrentParent = CurrentParent->GetParent();
+			CurrentParent = CurrentParent.GetEmitterData()->GetParent();
 		}
 
 		// Create a mapping from the actual parent scratch pad scripts to their copies in the target emitter.
 		TMap<UNiagaraScript*, UNiagaraScript*> SourceScratchPadScriptToTargetCopyScratchPadScript;
 		int32 TargetCopyIndex = 0;
-		for (UNiagaraEmitter* CurrentParentEmitter : ParentEmitters)
+		for (FVersionedNiagaraEmitter CurrentParentEmitter : ParentEmitters)
 		{
 			TArray<UNiagaraScript*> ParentScratchPadScripts;
-			if (CurrentParentEmitter->GetParent() == nullptr)
+			FVersionedNiagaraEmitterData* CurrentParentEmitterData = CurrentParentEmitter.GetEmitterData();
+			if (CurrentParentEmitterData->GetParent().Emitter == nullptr)
 			{
-				ParentScratchPadScripts.Append(CurrentParentEmitter->ParentScratchPadScripts);
+				ParentScratchPadScripts.Append(CurrentParentEmitterData->ParentScratchPads->Scripts);
 			}
-			ParentScratchPadScripts.Append(CurrentParentEmitter->ScratchPadScripts);
+			ParentScratchPadScripts.Append(CurrentParentEmitterData->ScratchPads->Scripts);
 
 			for (UNiagaraScript* ParentScratchPadScript : ParentScratchPadScripts)
 			{
 				UNiagaraScript* TargetCopy;
-				if (ensureMsgf(TargetCopyIndex < TargetEmitter->ParentScratchPadScripts.Num(), TEXT("Parent scratch pad script was missing from the Target's copies")))
+				if (ensureMsgf(TargetCopyIndex < TargetEmitterData->ParentScratchPads->Scripts.Num(), TEXT("Parent scratch pad script was missing from the Target's copies")))
 				{
-					TargetCopy = TargetEmitter->ParentScratchPadScripts[TargetCopyIndex];
+					TargetCopy = TargetEmitterData->ParentScratchPads->Scripts[TargetCopyIndex];
 				}
 				else
 				{
@@ -705,36 +705,37 @@ void FNiagaraScratchPadMergeAdapter::Initialize()
 		}
 
 		// Create a mapping from the instance scratch pad scripts to their copies in the target emitter.
-		if (InstanceEmitter != nullptr)
+		if (FVersionedNiagaraEmitterData* InstanceEmitterData = InstanceEmitter.GetEmitterData())
 		{
-			for (int32 InstanceScratchPadScriptIndex = 0; InstanceScratchPadScriptIndex < InstanceEmitter->ScratchPadScripts.Num(); InstanceScratchPadScriptIndex++)
+			for (int32 InstanceScratchPadScriptIndex = 0; InstanceScratchPadScriptIndex < InstanceEmitterData->ScratchPads->Scripts.Num(); InstanceScratchPadScriptIndex++)
 			{
 				UNiagaraScript* TargetCopy;
-				if (ensureMsgf(InstanceScratchPadScriptIndex < TargetEmitter->ScratchPadScripts.Num(), TEXT("Instance scratch pad script was missing from the Target's copies.")))
+				if (ensureMsgf(InstanceScratchPadScriptIndex < InstanceEmitterData->ScratchPads->Scripts.Num(), TEXT("Instance scratch pad script was missing from the Target's copies.")))
 				{
-					TargetCopy = TargetEmitter->ScratchPadScripts[InstanceScratchPadScriptIndex];
+					TargetCopy = InstanceEmitterData->ScratchPads->Scripts[InstanceScratchPadScriptIndex];
 				}
 				else
 				{
 					TargetCopy = nullptr;
 				}
-				SourceScratchPadScriptToTargetCopyScratchPadScript.Add(InstanceEmitter->ScratchPadScripts[InstanceScratchPadScriptIndex], TargetCopy);
+				SourceScratchPadScriptToTargetCopyScratchPadScript.Add(InstanceEmitterData->ScratchPads->Scripts[InstanceScratchPadScriptIndex], TargetCopy);
 			}
 		}
 
 		// For each source emitter collect up the traversed function ids that use the source scratch pad scripts and cache the corresponding target copies so that
 		// any usages encountered when applying the diff can be hooked up correctly.
-		TArray<UNiagaraEmitter*> SourceEmitters;
+		TArray<FVersionedNiagaraEmitter> SourceEmitters;
 		SourceEmitters.Append(ParentEmitters);
-		if (InstanceEmitter != nullptr)
+		if (InstanceEmitter.Emitter != nullptr)
 		{
 			SourceEmitters.Add(InstanceEmitter);
 		}
-		for (UNiagaraEmitter* SourceEmitter : SourceEmitters)
+		for (FVersionedNiagaraEmitter SourceEmitter : SourceEmitters)
 		{
-			if (SourceEmitter->ScratchPadScripts.Num() > 0 || (SourceEmitter->GetParent() == nullptr && SourceEmitter->ParentScratchPadScripts.Num() > 0))
+			FVersionedNiagaraEmitterData* SourceEmitterData = SourceEmitter.GetEmitterData();
+			if (SourceEmitterData->ScratchPads->Scripts.Num() > 0 || (SourceEmitterData->GetParent().Emitter == nullptr && SourceEmitterData->ParentScratchPads->Scripts.Num() > 0))
 			{
-				UNiagaraScriptSource* ScriptSource = Cast<UNiagaraScriptSource>(SourceEmitter->GraphSource);
+				UNiagaraScriptSource* ScriptSource = Cast<UNiagaraScriptSource>(SourceEmitterData->GraphSource);
 				if (ScriptSource != nullptr)
 				{
 					TArray<UNiagaraNodeOutput*> OutputNodes;
@@ -752,8 +753,8 @@ void FNiagaraScratchPadMergeAdapter::Initialize()
 								TraversedFunctionNode->GetClass() == UNiagaraNodeFunctionCall::StaticClass())
 							{
 								UNiagaraScript* ScratchPadScript = TraversedFunctionNode->FunctionScript;
-								if (SourceEmitter->ScratchPadScripts.Contains(ScratchPadScript) ||
-									(SourceEmitter->GetParent() == nullptr && SourceEmitter->ParentScratchPadScripts.Contains(ScratchPadScript)))
+								if (SourceEmitterData->ScratchPads->Scripts.Contains(ScratchPadScript) ||
+									(SourceEmitterData->GetParent().Emitter == nullptr && SourceEmitterData->ParentScratchPads->Scripts.Contains(ScratchPadScript)))
 								{
 									FunctionIdToScratchPadScript.Add(TraversedFunctionNode->NodeGuid, SourceScratchPadScriptToTargetCopyScratchPadScript[TraversedFunctionNode->FunctionScript]);
 								}
@@ -768,21 +769,21 @@ void FNiagaraScratchPadMergeAdapter::Initialize()
 	bIsInitialized = true;
 }
 
-FNiagaraEmitterMergeAdapter::FNiagaraEmitterMergeAdapter(const UNiagaraEmitter& InEmitter)
+FNiagaraEmitterMergeAdapter::FNiagaraEmitterMergeAdapter(const FVersionedNiagaraEmitter& InEmitter)
 {
-	Initialize(InEmitter, nullptr);
+	Initialize(InEmitter, InEmitter);
 }
 
-FNiagaraEmitterMergeAdapter::FNiagaraEmitterMergeAdapter(UNiagaraEmitter& InEmitter)
+void FNiagaraEmitterMergeAdapter::Initialize(const FVersionedNiagaraEmitter& InEmitter, const FVersionedNiagaraEmitter& InEditableEmitter)
 {
-	Initialize(InEmitter, &InEmitter);
-}
-
-void FNiagaraEmitterMergeAdapter::Initialize(const UNiagaraEmitter& InEmitter, UNiagaraEmitter* InEditableEmitter)
-{
-	Emitter = &InEmitter;
-	EditableEmitter = InEditableEmitter;
-	UNiagaraScriptSource* EmitterScriptSource = Cast<UNiagaraScriptSource>(Emitter->GraphSource);
+	Emitter = InEmitter.ToWeakPtr();
+	EditableEmitter = InEditableEmitter.ToWeakPtr();
+	FVersionedNiagaraEmitterData* EmitterData = Emitter.GetEmitterData();
+	if (EmitterData == nullptr)
+	{
+		return;
+	}
+	UNiagaraScriptSource* EmitterScriptSource = Cast<UNiagaraScriptSource>(EmitterData->GraphSource);
 	UNiagaraGraph* Graph = EmitterScriptSource->NodeGraph;
 	TArray<UNiagaraNodeOutput*> OutputNodes;
 	Graph->GetNodesOfClass<UNiagaraNodeOutput>(OutputNodes);
@@ -793,19 +794,19 @@ void FNiagaraEmitterMergeAdapter::Initialize(const UNiagaraEmitter& InEmitter, U
 	{
 		if (UNiagaraScript::IsEquivalentUsage(OutputNode->GetUsage(), ENiagaraScriptUsage::EmitterSpawnScript))
 		{
-			EmitterSpawnStack = MakeShared<FNiagaraScriptStackMergeAdapter>(*Emitter.Get(), *OutputNode, *Emitter->EmitterSpawnScriptProps.Script);
+			EmitterSpawnStack = MakeShared<FNiagaraScriptStackMergeAdapter>(InEmitter, *OutputNode, *EmitterData->EmitterSpawnScriptProps.Script);
 		}
 		else if (UNiagaraScript::IsEquivalentUsage(OutputNode->GetUsage(), ENiagaraScriptUsage::EmitterUpdateScript))
 		{
-			EmitterUpdateStack = MakeShared<FNiagaraScriptStackMergeAdapter>(*Emitter.Get(), *OutputNode, *Emitter->EmitterUpdateScriptProps.Script);
+			EmitterUpdateStack = MakeShared<FNiagaraScriptStackMergeAdapter>(InEmitter, *OutputNode, *EmitterData->EmitterUpdateScriptProps.Script);
 		}
 		else if (UNiagaraScript::IsEquivalentUsage(OutputNode->GetUsage(), ENiagaraScriptUsage::ParticleSpawnScript))
 		{
-			ParticleSpawnStack = MakeShared<FNiagaraScriptStackMergeAdapter>(*Emitter.Get(), *OutputNode, *Emitter->SpawnScriptProps.Script);
+			ParticleSpawnStack = MakeShared<FNiagaraScriptStackMergeAdapter>(InEmitter, *OutputNode, *EmitterData->SpawnScriptProps.Script);
 		}
 		else if (UNiagaraScript::IsEquivalentUsage(OutputNode->GetUsage(), ENiagaraScriptUsage::ParticleUpdateScript))
 		{
-			ParticleUpdateStack = MakeShared<FNiagaraScriptStackMergeAdapter>(*Emitter.Get(), *OutputNode, *Emitter->UpdateScriptProps.Script);
+			ParticleUpdateStack = MakeShared<FNiagaraScriptStackMergeAdapter>(InEmitter, *OutputNode, *EmitterData->UpdateScriptProps.Script);
 		}
 		else if(UNiagaraScript::IsEquivalentUsage(OutputNode->GetUsage(), ENiagaraScriptUsage::ParticleEventScript))
 		{
@@ -819,21 +820,21 @@ void FNiagaraEmitterMergeAdapter::Initialize(const UNiagaraEmitter& InEmitter, U
 
 	// Create an event handler adapter for each usage id even if it's missing an event script properties struct or an output node.  These
 	// incomplete adapters will be caught if they are diffed.
-	for (const FNiagaraEventScriptProperties& EventScriptProperties : Emitter->GetEventHandlers())
+	for (const FNiagaraEventScriptProperties& EventScriptProperties : EmitterData->GetEventHandlers())
 	{
 		UNiagaraNodeOutput** MatchingOutputNodePtr = EventOutputNodes.FindByPredicate(
 			[=](UNiagaraNodeOutput* EventOutputNode) { return EventOutputNode->GetUsageId() == EventScriptProperties.Script->GetUsageId(); });
 
 		UNiagaraNodeOutput* MatchingOutputNode = MatchingOutputNodePtr != nullptr ? *MatchingOutputNodePtr : nullptr;
 
-		if (EditableEmitter == nullptr)
+		if (EditableEmitter.Emitter == nullptr)
 		{
-			EventHandlers.Add(MakeShared<FNiagaraEventHandlerMergeAdapter>(*Emitter.Get(), &EventScriptProperties, MatchingOutputNode));
+			EventHandlers.Add(MakeShared<FNiagaraEventHandlerMergeAdapter>(InEmitter, &EventScriptProperties, MatchingOutputNode));
 		}
 		else
 		{
-			FNiagaraEventScriptProperties* EditableEventScriptProperties = EditableEmitter->GetEventHandlerByIdUnsafe(EventScriptProperties.Script->GetUsageId());
-			EventHandlers.Add(MakeShared<FNiagaraEventHandlerMergeAdapter>(*Emitter.Get(), EditableEventScriptProperties, MatchingOutputNode));
+			FNiagaraEventScriptProperties* EditableEventScriptProperties = EditableEmitter.GetEmitterData()->GetEventHandlerByIdUnsafe(EventScriptProperties.Script->GetUsageId());
+			EventHandlers.Add(MakeShared<FNiagaraEventHandlerMergeAdapter>(InEmitter, EditableEventScriptProperties, MatchingOutputNode));
 		}
 
 		if (MatchingOutputNode != nullptr)
@@ -844,28 +845,28 @@ void FNiagaraEmitterMergeAdapter::Initialize(const UNiagaraEmitter& InEmitter, U
 
 	for (UNiagaraNodeOutput* EventOutputNode : EventOutputNodes)
 	{
-		EventHandlers.Add(MakeShared<FNiagaraEventHandlerMergeAdapter>(*Emitter.Get(), EventOutputNode));
+		EventHandlers.Add(MakeShared<FNiagaraEventHandlerMergeAdapter>(InEmitter, EventOutputNode));
 	}
 
 	// Create an shader stage adapter for each usage id even if it's missing a shader stage object or an output node.  These
 	// incomplete adapters will be caught if they are diffed.
-	for (int32 SimulationStageIndex = 0; SimulationStageIndex < Emitter->GetSimulationStages().Num(); SimulationStageIndex++)
+	for (int32 SimulationStageIndex = 0; SimulationStageIndex < EmitterData->GetSimulationStages().Num(); SimulationStageIndex++)
 	{
-		const UNiagaraSimulationStageBase* SimulationStage = Emitter->GetSimulationStages()[SimulationStageIndex];
+		const UNiagaraSimulationStageBase* SimulationStage = EmitterData->GetSimulationStages()[SimulationStageIndex];
 	
 		UNiagaraNodeOutput** MatchingOutputNodePtr = SimulationStageOutputNodes.FindByPredicate(
 			[=](UNiagaraNodeOutput* SimulationStageOutputNode) { return SimulationStageOutputNode->GetUsageId() == SimulationStage->Script->GetUsageId(); });
 
 		UNiagaraNodeOutput* MatchingOutputNode = MatchingOutputNodePtr != nullptr ? *MatchingOutputNodePtr : nullptr;
 
-		if (EditableEmitter == nullptr)
+		if (EditableEmitter.Emitter == nullptr)
 		{
-			SimulationStages.Add(MakeShared<FNiagaraSimulationStageMergeAdapter>(*Emitter.Get(), SimulationStage, SimulationStageIndex, MatchingOutputNode));
+			SimulationStages.Add(MakeShared<FNiagaraSimulationStageMergeAdapter>(InEmitter, SimulationStage, SimulationStageIndex, MatchingOutputNode));
 		}
 		else
 		{
-			UNiagaraSimulationStageBase* EditableSimulationStage = EditableEmitter->GetSimulationStageById(SimulationStage->Script->GetUsageId());
-			SimulationStages.Add(MakeShared<FNiagaraSimulationStageMergeAdapter>(*Emitter.Get(), EditableSimulationStage, SimulationStageIndex, MatchingOutputNode));
+			UNiagaraSimulationStageBase* EditableSimulationStage = EditableEmitter.GetEmitterData()->GetSimulationStageById(SimulationStage->Script->GetUsageId());
+			SimulationStages.Add(MakeShared<FNiagaraSimulationStageMergeAdapter>(InEmitter, EditableSimulationStage, SimulationStageIndex, MatchingOutputNode));
 		}
 
 		if (MatchingOutputNode != nullptr)
@@ -876,22 +877,21 @@ void FNiagaraEmitterMergeAdapter::Initialize(const UNiagaraEmitter& InEmitter, U
 
 	for (UNiagaraNodeOutput* SimulationStageOutputNode : SimulationStageOutputNodes)
 	{
-		SimulationStages.Add(MakeShared<FNiagaraSimulationStageMergeAdapter>(*Emitter.Get(), SimulationStageOutputNode));
+		SimulationStages.Add(MakeShared<FNiagaraSimulationStageMergeAdapter>(InEmitter, SimulationStageOutputNode));
 	}
 
 	// Renderers
-	for (UNiagaraRendererProperties* RendererProperties : Emitter->GetRenderers())
+	for (UNiagaraRendererProperties* RendererProperties : EmitterData->GetRenderers())
 	{
 		Renderers.Add(MakeShared<FNiagaraRendererMergeAdapter>(*RendererProperties));
 	}
 
-	EditorData = Cast<const UNiagaraEmitterEditorData>(Emitter->GetEditorData());
-
+	EditorData = Cast<const UNiagaraEmitterEditorData>(EmitterData->GetEditorData());
 }
 
-UNiagaraEmitter* FNiagaraEmitterMergeAdapter::GetEditableEmitter() const
+FVersionedNiagaraEmitter FNiagaraEmitterMergeAdapter::GetEditableEmitter() const
 {
-	return EditableEmitter.Get();
+	return EditableEmitter.ResolveWeakPtr();
 }
 
 TSharedPtr<FNiagaraScriptStackMergeAdapter> FNiagaraEmitterMergeAdapter::GetEmitterSpawnStack() const
@@ -1228,21 +1228,21 @@ void FNiagaraScriptMergeManager::GetForcedChangeIds(
 	}
 }
 
-FNiagaraScriptMergeManager::FApplyDiffResults FNiagaraScriptMergeManager::ForceInstanceChangeIds(TSharedRef<FNiagaraEmitterMergeAdapter> MergedInstanceAdapter, UNiagaraEmitter& OriginalEmitterInstance, const TMap<FGuid, FGuid>& FunctionIdToForcedChangedId) const
+FNiagaraScriptMergeManager::FApplyDiffResults FNiagaraScriptMergeManager::ForceInstanceChangeIds(TSharedRef<FNiagaraEmitterMergeAdapter> MergedInstanceAdapter, const FVersionedNiagaraEmitter& OriginalEmitterInstance, const TMap<FGuid, FGuid>& FunctionIdToForcedChangedId) const
 {
 	FApplyDiffResults DiffResults;
 
 	if (FunctionIdToForcedChangedId.Num() != 0)
 	{
 		auto It = FunctionIdToForcedChangedId.CreateConstIterator();
-		UNiagaraEmitter* Emitter = MergedInstanceAdapter->GetEditableEmitter();
+		FVersionedNiagaraEmitter Emitter = MergedInstanceAdapter->GetEditableEmitter();
 
 		TArray<UNiagaraGraph*> Graphs;
 		TArray<UNiagaraScript*> Scripts;
-		Emitter->GetScripts(Scripts);
+		Emitter.GetEmitterData()->GetScripts(Scripts);
 
 		TArray<UNiagaraScript*> OriginalScripts;
-		OriginalEmitterInstance.GetScripts(OriginalScripts);
+		OriginalEmitterInstance.GetEmitterData()->GetScripts(OriginalScripts);
 
 		// First gather all the graphs used by this emitter..
 		for (UNiagaraScript* Script : Scripts)
@@ -1317,7 +1317,7 @@ FNiagaraScriptMergeManager::FApplyDiffResults FNiagaraScriptMergeManager::ForceI
 	return DiffResults;
 }
 
-void FNiagaraScriptMergeManager::UpdateModuleVersions(UNiagaraEmitter& Instance, const FNiagaraEmitterDiffResults& DiffResults) const
+void FNiagaraScriptMergeManager::UpdateModuleVersions(const FVersionedNiagaraEmitter& Instance, const FNiagaraEmitterDiffResults& DiffResults) const
 {
 	// gather all the changes
 	TArray<TSharedRef<FNiagaraStackFunctionMergeAdapter>> ChangedVersionBaseModules;
@@ -1399,7 +1399,7 @@ void FNiagaraScriptMergeManager::UpdateModuleVersions(UNiagaraEmitter& Instance,
 		}
 		
 		FNiagaraScriptVersionUpgradeContext UpgradeContext;
-		UpgradeContext.ConstantResolver = FCompileConstantResolver(&Instance, FNiagaraStackGraphUtilities::GetOutputNodeUsage(*ChangedModule->GetFunctionCallNode()));
+		UpgradeContext.ConstantResolver = FCompileConstantResolver(Instance, FNiagaraStackGraphUtilities::GetOutputNodeUsage(*ChangedModule->GetFunctionCallNode()));
 		if (ModuleItem)
 		{
 			UpgradeContext.CreateClipboardCallback = [ModuleItem](UNiagaraClipboardContent* ClipboardContent)
@@ -1418,17 +1418,68 @@ void FNiagaraScriptMergeManager::UpdateModuleVersions(UNiagaraEmitter& Instance,
 	}
 }
 
-INiagaraMergeManager::FMergeEmitterResults FNiagaraScriptMergeManager::MergeEmitter(UNiagaraEmitter& Parent, UNiagaraEmitter* ParentAtLastMerge, UNiagaraEmitter& Instance) const
+enum class EmitterMergeState
 {
-	SCOPE_CYCLE_COUNTER(STAT_NiagaraEditor_ScriptMergeManager_MergeEmitter);
-	INiagaraMergeManager::FMergeEmitterResults MergeResults;
-	const bool bNoParentAtLastMerge = (ParentAtLastMerge == nullptr);
-	UNiagaraEmitter* FirstEmitterToDiffAgainst = bNoParentAtLastMerge ? &Parent : ParentAtLastMerge;
-	FNiagaraEmitterDiffResults DiffResults = DiffEmitters(*FirstEmitterToDiffAgainst, Instance);
-	
+	Failed, Unchanged, DuplicateParent, CopyProperties
+};
+
+FNiagaraScriptMergeManager::EmitterMergeState FNiagaraScriptMergeManager::GetEmitterMergeState(const FNiagaraEmitterDiffResults& DiffResults, const FVersionedNiagaraEmitter& Parent, const FVersionedNiagaraEmitter& Instance) const
+{
+	EmitterMergeState State = EmitterMergeState::CopyProperties;
 	if (DiffResults.IsValid() == false)
 	{
-		MergeResults.MergeResult = INiagaraMergeManager::EMergeEmitterResult::FailedToDiff;
+		State = EmitterMergeState::Failed;
+	}
+	else if (DiffResults.IsEmpty())
+	{
+		// If there were no changes made on the instance, check if the instance matches the parent.
+		FNiagaraEmitterDiffResults DiffResultsFromParent = DiffEmitters(Parent, Instance);
+		if (DiffResultsFromParent.IsValid() && DiffResultsFromParent.IsEmpty())
+		{
+			State = EmitterMergeState::Unchanged;
+		}
+		else if (Instance.Emitter->IsVersioningEnabled())
+		{
+			// if this emitter is versioned we cannot just duplicate the parent, as we would lose version data
+			State = EmitterMergeState::CopyProperties;
+		}
+		else
+		{
+			State = EmitterMergeState::DuplicateParent;
+		}
+	}
+	return State;
+}
+
+INiagaraMergeManager::FMergeEmitterResults FNiagaraScriptMergeManager::MergeEmitter(const FVersionedNiagaraEmitter& Parent, const FVersionedNiagaraEmitter& ParentAtLastMerge, const FVersionedNiagaraEmitter& Instance) const
+{
+	SCOPE_CYCLE_COUNTER(STAT_NiagaraEditor_ScriptMergeManager_MergeEmitter);
+	FMergeEmitterResults MergeResults;
+	const bool bNoParentAtLastMerge = (ParentAtLastMerge.Emitter == nullptr);
+	FVersionedNiagaraEmitter FirstEmitterToDiffAgainst = bNoParentAtLastMerge ? Parent : ParentAtLastMerge;
+	FNiagaraEmitterDiffResults DiffResults = DiffEmitters(FirstEmitterToDiffAgainst, Instance);
+
+	// depending on the diff result and versioning state of the emitters, we decide what to do next
+	EmitterMergeState State = GetEmitterMergeState(DiffResults, Parent, Instance);
+
+	if (State == EmitterMergeState::Unchanged)
+	{
+		MergeResults.MergeResult = EMergeEmitterResult::SucceededNoDifferences;
+	}
+	else if (State == EmitterMergeState::DuplicateParent)
+	{
+		// If there were differences from the parent or the parent diff failed we can just return a copy of the parent as the merged instance since there
+		// were no changes in the instance which need to be applied.
+		MergeResults.MergeResult = EMergeEmitterResult::SucceededDifferencesApplied;
+		MergeResults.MergedInstance = Parent.Emitter->DuplicateWithoutMerging(GetTransientPackage());
+		MergeResults.MergedInstance->DisableVersioning(Parent.Version);
+		FVersionedNiagaraEmitterData* MergedEmitterData = MergeResults.MergedInstance->GetLatestEmitterData();
+		MergedEmitterData->ParentScratchPads->SetScripts(MergedEmitterData->ScratchPads->Scripts);
+		MergedEmitterData->ScratchPads->Scripts.Empty();
+	}
+	else if (State == EmitterMergeState::Failed)
+	{
+		MergeResults.MergeResult = EMergeEmitterResult::FailedToDiff;
 		MergeResults.ErrorMessages = DiffResults.GetErrorMessages();
 
 		auto ReportScriptStackDiffErrors = [](FMergeEmitterResults& EmitterMergeResults, const FNiagaraScriptStackDiffResults& ScriptStackDiffResults, FText ScriptName)
@@ -1455,28 +1506,12 @@ INiagaraMergeManager::FMergeEmitterResults FNiagaraScriptMergeManager::MergeEmit
 			ReportScriptStackDiffErrors(MergeResults, EventHandlerDiffResults.ScriptDiffResults, EventHandlerName);
 		}
 	}
-	else if (DiffResults.IsEmpty())
+	else if (State == EmitterMergeState::CopyProperties)
 	{
-		// If there were no changes made on the instance, check if the instance matches the parent.
-		FNiagaraEmitterDiffResults DiffResultsFromParent = DiffEmitters(Parent, Instance);
-		if (DiffResultsFromParent.IsValid() && DiffResultsFromParent.IsEmpty())
-		{
-			MergeResults.MergeResult = INiagaraMergeManager::EMergeEmitterResult::SucceededNoDifferences;
-		}
-		else
-		{
-			// If there were differences from the parent or the parent diff failed we can just return a copy of the parent as the merged instance since there
-			// were no changes in the instance which need to be applied.
-			MergeResults.MergeResult = INiagaraMergeManager::EMergeEmitterResult::SucceededDifferencesApplied;
-			MergeResults.MergedInstance = Parent.DuplicateWithoutMerging(GetTransientPackage());
-			MergeResults.MergedInstance->ParentScratchPadScripts.Append(MergeResults.MergedInstance->ScratchPadScripts);
-			MergeResults.MergedInstance->ScratchPadScripts.Empty();
-		}
-	}
-	else
-	{
-		UNiagaraEmitter* MergedInstance = Parent.DuplicateWithoutMerging(GetTransientPackage());
-		TSharedRef<FNiagaraEmitterMergeAdapter> MergedInstanceAdapter = MakeShared<FNiagaraEmitterMergeAdapter>(*MergedInstance);
+		UNiagaraEmitter* MergedInstance = Parent.Emitter->DuplicateWithoutMerging(GetTransientPackage());
+		MergedInstance->DisableVersioning(Parent.Version);
+		FVersionedNiagaraEmitter VersionedMergedInstance = FVersionedNiagaraEmitter(MergedInstance, Parent.Version);
+		TSharedRef<FNiagaraEmitterMergeAdapter> MergedInstanceAdapter = MakeShared<FNiagaraEmitterMergeAdapter>(VersionedMergedInstance);
 
 		// diff for version changes and apply them first. We need to upgrade both the current emitter and the parent at last merge
 		// to the newest version and then diff them again. Otherwise we won't know which changes were made by the version upgrade
@@ -1488,14 +1523,15 @@ INiagaraMergeManager::FMergeEmitterResults FNiagaraScriptMergeManager::MergeEmit
 			UpdateModuleVersions(Instance, VersionChangeDiffResults);
 
 			// apply version upgrade to the parent at last merge
-			UNiagaraEmitter* ParentAtLastMergeCopy = Cast<UNiagaraEmitter>(StaticDuplicateObject(FirstEmitterToDiffAgainst, GetTransientPackage()));
+			UNiagaraEmitter* ParentAtLastMergeCopy = Cast<UNiagaraEmitter>(StaticDuplicateObject(FirstEmitterToDiffAgainst.Emitter, GetTransientPackage()));
 			ParentAtLastMergeCopy->ClearFlags(RF_Standalone | RF_Public);
-			FNiagaraEmitterDiffResults ParentsDiffResults = DiffEmitters(Parent, *ParentAtLastMergeCopy);
-			UpdateModuleVersions(*ParentAtLastMergeCopy, ParentsDiffResults);
+			FVersionedNiagaraEmitter VersionedParentCopy = FVersionedNiagaraEmitter(ParentAtLastMergeCopy, FirstEmitterToDiffAgainst.Version);
+			FNiagaraEmitterDiffResults ParentsDiffResults = DiffEmitters(Parent, VersionedParentCopy);
+			UpdateModuleVersions(VersionedParentCopy, ParentsDiffResults);
 
 			// now diff again
-			DiffResults = DiffEmitters(*ParentAtLastMergeCopy, Instance);
-			ensureMsgf(DiffResults.HasVersionChanges() == false, TEXT("Emitter still has pending version changes after merge! Asset %s"), *Instance.GetPathName());
+			DiffResults = DiffEmitters(VersionedParentCopy, Instance);
+			ensureMsgf(DiffResults.HasVersionChanges() == false, TEXT("Emitter still has pending version changes after merge! Asset %s"), *Instance.Emitter->GetPathName());
 		}
 
 		TMap<FGuid, UNiagaraNodeFunctionCall*> ParentFunctionIdToNodeMap;
@@ -1519,16 +1555,17 @@ INiagaraMergeManager::FMergeEmitterResults FNiagaraScriptMergeManager::MergeEmit
 		TMap<FGuid, FGuid> FunctionIdToForcedChangeId;
 		GetForcedChangeIds(ParentFunctionIdToNodeMap, LastMergedParentFunctionIdToNodeMap, InstanceFunctionIdToNodeMap, FunctionIdToForcedChangeId);
 
-		MergedInstance->ParentScratchPadScripts.Append(MergedInstance->ScratchPadScripts);
-		MergedInstance->ScratchPadScripts.Empty();
-		CopyInstanceScratchPadScripts(*MergedInstance, Instance);
-		TSharedRef<FNiagaraScratchPadMergeAdapter> ScratchPadAdapter = MakeShared<FNiagaraScratchPadMergeAdapter>(MergedInstance, &Instance, &Parent);
+		FVersionedNiagaraEmitterData* MergedEmitterData = VersionedMergedInstance.GetEmitterData();
+		MergedEmitterData->ParentScratchPads->SetScripts(MergedEmitterData->ScratchPads->Scripts);
+		MergedEmitterData->ScratchPads->Scripts.Empty();
+		CopyInstanceScratchPadScripts(VersionedMergedInstance, Instance.GetEmitterData());
+		TSharedRef<FNiagaraScratchPadMergeAdapter> ScratchPadAdapter = MakeShared<FNiagaraScratchPadMergeAdapter>(VersionedMergedInstance, Instance, Parent);
 
-		MergeResults.MergeResult = INiagaraMergeManager::EMergeEmitterResult::SucceededDifferencesApplied;
+		MergeResults.MergeResult = EMergeEmitterResult::SucceededDifferencesApplied;
 		FApplyDiffResults EmitterSpawnResults = ApplyScriptStackDiff(MergedInstanceAdapter, MergedInstanceAdapter->GetEmitterSpawnStack().ToSharedRef(), ScratchPadAdapter, DiffResults.EmitterSpawnDiffResults, bNoParentAtLastMerge);
 		if (EmitterSpawnResults.bSucceeded == false)
 		{
-			MergeResults.MergeResult = INiagaraMergeManager::EMergeEmitterResult::FailedToMerge;
+			MergeResults.MergeResult = EMergeEmitterResult::FailedToMerge;
 		}
 		MergeResults.bModifiedGraph |= EmitterSpawnResults.bModifiedGraph;
 		MergeResults.ErrorMessages.Append(EmitterSpawnResults.ErrorMessages);
@@ -1536,7 +1573,7 @@ INiagaraMergeManager::FMergeEmitterResults FNiagaraScriptMergeManager::MergeEmit
 		FApplyDiffResults EmitterUpdateResults = ApplyScriptStackDiff(MergedInstanceAdapter, MergedInstanceAdapter->GetEmitterUpdateStack().ToSharedRef(), ScratchPadAdapter, DiffResults.EmitterUpdateDiffResults, bNoParentAtLastMerge);
 		if (EmitterUpdateResults.bSucceeded == false)
 		{
-			MergeResults.MergeResult = INiagaraMergeManager::EMergeEmitterResult::FailedToMerge;
+			MergeResults.MergeResult = EMergeEmitterResult::FailedToMerge;
 		}
 		MergeResults.bModifiedGraph |= EmitterUpdateResults.bModifiedGraph;
 		MergeResults.ErrorMessages.Append(EmitterUpdateResults.ErrorMessages);
@@ -1544,7 +1581,7 @@ INiagaraMergeManager::FMergeEmitterResults FNiagaraScriptMergeManager::MergeEmit
 		FApplyDiffResults ParticleSpawnResults = ApplyScriptStackDiff(MergedInstanceAdapter, MergedInstanceAdapter->GetParticleSpawnStack().ToSharedRef(), ScratchPadAdapter, DiffResults.ParticleSpawnDiffResults, bNoParentAtLastMerge);
 		if (ParticleSpawnResults.bSucceeded == false)
 		{
-			MergeResults.MergeResult = INiagaraMergeManager::EMergeEmitterResult::FailedToMerge;
+			MergeResults.MergeResult = EMergeEmitterResult::FailedToMerge;
 		}
 		MergeResults.bModifiedGraph |= ParticleSpawnResults.bModifiedGraph;
 		MergeResults.ErrorMessages.Append(ParticleSpawnResults.ErrorMessages);
@@ -1552,7 +1589,7 @@ INiagaraMergeManager::FMergeEmitterResults FNiagaraScriptMergeManager::MergeEmit
 		FApplyDiffResults ParticleUpdateResults = ApplyScriptStackDiff(MergedInstanceAdapter, MergedInstanceAdapter->GetParticleUpdateStack().ToSharedRef(), ScratchPadAdapter, DiffResults.ParticleUpdateDiffResults, bNoParentAtLastMerge);
 		if (ParticleUpdateResults.bSucceeded == false)
 		{
-			MergeResults.MergeResult = INiagaraMergeManager::EMergeEmitterResult::FailedToMerge;
+			MergeResults.MergeResult = EMergeEmitterResult::FailedToMerge;
 		}
 		MergeResults.bModifiedGraph |= ParticleUpdateResults.bModifiedGraph;
 		MergeResults.ErrorMessages.Append(ParticleUpdateResults.ErrorMessages);
@@ -1560,7 +1597,7 @@ INiagaraMergeManager::FMergeEmitterResults FNiagaraScriptMergeManager::MergeEmit
 		FApplyDiffResults EventHandlerResults = ApplyEventHandlerDiff(MergedInstanceAdapter, ScratchPadAdapter, DiffResults, bNoParentAtLastMerge);
 		if (EventHandlerResults.bSucceeded == false)
 		{
-			MergeResults.MergeResult = INiagaraMergeManager::EMergeEmitterResult::FailedToMerge;
+			MergeResults.MergeResult = EMergeEmitterResult::FailedToMerge;
 		}
 		MergeResults.bModifiedGraph |= EventHandlerResults.bModifiedGraph;
 		MergeResults.ErrorMessages.Append(EventHandlerResults.ErrorMessages);
@@ -1568,33 +1605,33 @@ INiagaraMergeManager::FMergeEmitterResults FNiagaraScriptMergeManager::MergeEmit
 		FApplyDiffResults SimulationStageResults = ApplySimulationStageDiff(MergedInstanceAdapter, ScratchPadAdapter, DiffResults, bNoParentAtLastMerge);
 		if (SimulationStageResults.bSucceeded == false)
 		{
-			MergeResults.MergeResult = INiagaraMergeManager::EMergeEmitterResult::FailedToMerge;
+			MergeResults.MergeResult = EMergeEmitterResult::FailedToMerge;
 		}
 		MergeResults.bModifiedGraph |= SimulationStageResults.bModifiedGraph;
 		MergeResults.ErrorMessages.Append(SimulationStageResults.ErrorMessages);
 
-		FApplyDiffResults RendererResults = ApplyRendererDiff(*MergedInstance, DiffResults, bNoParentAtLastMerge);
+		FApplyDiffResults RendererResults = ApplyRendererDiff(VersionedMergedInstance, DiffResults, bNoParentAtLastMerge);
 		if (RendererResults.bSucceeded == false)
 		{
-			MergeResults.MergeResult = INiagaraMergeManager::EMergeEmitterResult::FailedToMerge;
+			MergeResults.MergeResult = EMergeEmitterResult::FailedToMerge;
 		}
 		MergeResults.bModifiedGraph |= RendererResults.bModifiedGraph;
 		MergeResults.ErrorMessages.Append(RendererResults.ErrorMessages);
 
-		FApplyDiffResults InputSummaryResults = ApplyEmitterSummaryDiff(*MergedInstance, DiffResults, bNoParentAtLastMerge);
+		FApplyDiffResults InputSummaryResults = ApplyEmitterSummaryDiff(VersionedMergedInstance, DiffResults);
 		if (InputSummaryResults.bSucceeded == false)
 		{
-			MergeResults.MergeResult = INiagaraMergeManager::EMergeEmitterResult::FailedToMerge;
+			MergeResults.MergeResult = EMergeEmitterResult::FailedToMerge;
 		}
 		MergeResults.bModifiedGraph |= InputSummaryResults.bModifiedGraph;
 		MergeResults.ErrorMessages.Append(InputSummaryResults.ErrorMessages);
 
-		CopyPropertiesToBase(MergedInstance, &Instance, DiffResults.DifferentEmitterProperties);
+		CopyPropertiesToBase(VersionedMergedInstance.GetEmitterData(), Instance.GetEmitterData(), DiffResults.DifferentEmitterProperties);
 
-		FApplyDiffResults StackEntryDisplayNameDiffs = ApplyStackEntryDisplayNameDiffs(*MergedInstance, DiffResults);
+		FApplyDiffResults StackEntryDisplayNameDiffs = ApplyStackEntryDisplayNameDiffs(VersionedMergedInstance, DiffResults);
 		if (StackEntryDisplayNameDiffs.bSucceeded == false)
 		{
-			MergeResults.MergeResult = INiagaraMergeManager::EMergeEmitterResult::FailedToMerge;
+			MergeResults.MergeResult = EMergeEmitterResult::FailedToMerge;
 		}
 		MergeResults.bModifiedGraph |= StackEntryDisplayNameDiffs.bModifiedGraph;
 		MergeResults.ErrorMessages.Append(StackEntryDisplayNameDiffs.ErrorMessages);
@@ -1620,7 +1657,7 @@ INiagaraMergeManager::FMergeEmitterResults FNiagaraScriptMergeManager::MergeEmit
 		FApplyDiffResults ChangeIdResults = ForceInstanceChangeIds(MergedInstanceAdapter, Instance, FunctionIdToForcedChangeId);
 		if (ChangeIdResults.bSucceeded == false)
 		{
-			MergeResults.MergeResult = INiagaraMergeManager::EMergeEmitterResult::FailedToMerge;
+			MergeResults.MergeResult = EMergeEmitterResult::FailedToMerge;
 		}
 		MergeResults.bModifiedGraph |= ChangeIdResults.bModifiedGraph;
 		MergeResults.ErrorMessages.Append(ChangeIdResults.ErrorMessages);
@@ -1643,7 +1680,7 @@ INiagaraMergeManager::FMergeEmitterResults FNiagaraScriptMergeManager::MergeEmit
 		}
 #endif
 
-		FNiagaraStackGraphUtilities::CleanUpStaleRapidIterationParameters(*MergedInstance);
+		FNiagaraStackGraphUtilities::CleanUpStaleRapidIterationParameters(VersionedMergedInstance);
 
 #if 0
 		UE_LOG(LogNiagaraEditor, Log, TEXT("C"));
@@ -1662,9 +1699,9 @@ INiagaraMergeManager::FMergeEmitterResults FNiagaraScriptMergeManager::MergeEmit
 			UpdateScript->RapidIterationParameters.DumpParameters();
 		}
 #endif
-		if (MergeResults.MergeResult == INiagaraMergeManager::EMergeEmitterResult::SucceededDifferencesApplied)
+		if (MergeResults.MergeResult == EMergeEmitterResult::SucceededDifferencesApplied)
 		{
-			UNiagaraScriptSource* ScriptSource = Cast<UNiagaraScriptSource>(MergedInstance->GraphSource);
+			UNiagaraScriptSource* ScriptSource = Cast<UNiagaraScriptSource>(MergedEmitterData->GraphSource);
 			FNiagaraStackGraphUtilities::RelayoutGraph(*ScriptSource->NodeGraph);
 			MergeResults.MergedInstance = MergedInstance;
 		}
@@ -1672,11 +1709,15 @@ INiagaraMergeManager::FMergeEmitterResults FNiagaraScriptMergeManager::MergeEmit
 		TMap<FGuid, FGuid> FinalChangeIds;
 		FNiagaraEditorUtilities::GatherChangeIds(*MergedInstance, FinalChangeIds, TEXT("Final"));
 	}
+	else
+	{
+		checkf(false, TEXT("Unknown merge state!"));
+	}
 
 	if (MergeResults.MergeResult != INiagaraMergeManager::EMergeEmitterResult::SucceededNoDifferences && MergeResults.MergeResult != INiagaraMergeManager::EMergeEmitterResult::SucceededDifferencesApplied)
 	{
 		TArray<FText> Errors;
-		Errors.Add(FText::Format(LOCTEXT("MergeFailedFormat", "Failed to merge changes from parent emitter {0}"), FText::FromString(Parent.GetName())));
+		Errors.Add(FText::Format(LOCTEXT("MergeFailedFormat", "Failed to merge changes from parent emitter {0}"), FText::FromString(Parent.Emitter->GetName())));
 		Errors.Append(MergeResults.ErrorMessages);
 
 		UNiagaraMessageDataText* Message = NewObject<UNiagaraMessageDataText>(GetTransientPackage());
@@ -1701,14 +1742,14 @@ bool FNiagaraScriptMergeManager::IsMergeableScriptUsage(ENiagaraScriptUsage Scri
 		ScriptUsage == ENiagaraScriptUsage::ParticleSimulationStageScript;
 }
 
-bool FNiagaraScriptMergeManager::HasBaseModule(const UNiagaraEmitter& BaseEmitter, ENiagaraScriptUsage ScriptUsage, FGuid ScriptUsageId, FGuid ModuleId)
+bool FNiagaraScriptMergeManager::HasBaseModule(const FVersionedNiagaraEmitter& BaseEmitter, ENiagaraScriptUsage ScriptUsage, FGuid ScriptUsageId, FGuid ModuleId)
 {
 	TSharedRef<FNiagaraEmitterMergeAdapter> BaseEmitterAdapter = GetEmitterMergeAdapterUsingCache(BaseEmitter);
 	TSharedPtr<FNiagaraScriptStackMergeAdapter> BaseScriptStackAdapter = BaseEmitterAdapter->GetScriptStack(ScriptUsage, ScriptUsageId);
 	return BaseScriptStackAdapter.IsValid() && BaseScriptStackAdapter->GetModuleFunctionById(ModuleId).IsValid();
 }
 
-bool FNiagaraScriptMergeManager::IsModuleInputDifferentFromBase(UNiagaraEmitter& Emitter, const UNiagaraEmitter& BaseEmitter, ENiagaraScriptUsage ScriptUsage, FGuid ScriptUsageId, FGuid ModuleId, FString InputName)
+bool FNiagaraScriptMergeManager::IsModuleInputDifferentFromBase(const FVersionedNiagaraEmitter& Emitter, const FVersionedNiagaraEmitter& BaseEmitter, ENiagaraScriptUsage ScriptUsage, FGuid ScriptUsageId, FGuid ModuleId, FString InputName)
 {
 	SCOPE_CYCLE_COUNTER(STAT_NiagaraEditor_ScriptMergeManager_IsModuleInputDifferentFromBase);
 
@@ -1747,23 +1788,24 @@ bool FNiagaraScriptMergeManager::IsModuleInputDifferentFromBase(UNiagaraEmitter&
 		ScriptStackDiffResults.ModifiedOtherInputOverrides.FindByPredicate(FindInputOverrideByInputName) != nullptr;
 }
 
-FNiagaraScriptMergeManager::FApplyDiffResults FNiagaraScriptMergeManager::ResetModuleInputToBase(UNiagaraEmitter& Emitter, const UNiagaraEmitter& BaseEmitter, ENiagaraScriptUsage ScriptUsage, FGuid ScriptUsageId, FGuid ModuleId, FString InputName)
+FNiagaraScriptMergeManager::FApplyDiffResults FNiagaraScriptMergeManager::ResetModuleInputToBase(const FVersionedNiagaraEmitter& VersionedEmitter, const FVersionedNiagaraEmitter& VersionedBaseEmitter, ENiagaraScriptUsage ScriptUsage, FGuid ScriptUsageId, FGuid ModuleId, FString InputName)
 {
-	TSharedRef<FNiagaraEmitterMergeAdapter> EmitterAdapter = GetEmitterMergeAdapterUsingCache(Emitter);
-	TSharedRef<FNiagaraEmitterMergeAdapter> BaseEmitterAdapter = GetEmitterMergeAdapterUsingCache(BaseEmitter);
-	TSharedRef<FNiagaraScratchPadMergeAdapter> ScratchPadMergeAdapter = GetScratchPadMergeAdapterUsingCache(Emitter);
+	TSharedRef<FNiagaraEmitterMergeAdapter> EmitterAdapter = GetEmitterMergeAdapterUsingCache(VersionedEmitter);
+	TSharedRef<FNiagaraEmitterMergeAdapter> BaseEmitterAdapter = GetEmitterMergeAdapterUsingCache(VersionedBaseEmitter);
+	TSharedRef<FNiagaraScratchPadMergeAdapter> ScratchPadMergeAdapter = GetScratchPadMergeAdapterUsingCache(VersionedEmitter);
 
 	// Diff from the emitter to the base to create a diff which will reset the emitter back to the base.
 	FNiagaraScriptStackDiffResults ResetDiffResults;
 	DiffScriptStacks(EmitterAdapter->GetScriptStack(ScriptUsage, ScriptUsageId).ToSharedRef(), BaseEmitterAdapter->GetScriptStack(ScriptUsage, ScriptUsageId).ToSharedRef(), ResetDiffResults);
 
+	FText EmitterPath = FText::FromString(VersionedEmitter.Emitter->GetPathName());
 	if (ResetDiffResults.IsValid() == false)
 	{
 		FApplyDiffResults Results;
 		Results.bSucceeded = false;
 		Results.bModifiedGraph = false;
 		Results.ErrorMessages.Add(FText::Format(LOCTEXT("ResetFailedBecauseOfDiffMessage", "Failed to reset input back to it's base value.  It couldn't be diffed successfully.  Emitter: {0}  Input:{1}"),
-			FText::FromString(Emitter.GetPathName()), FText::FromString(InputName)));
+		    EmitterPath, FText::FromString(InputName)));
 		return Results;
 	}
 
@@ -1773,17 +1815,19 @@ FNiagaraScriptMergeManager::FApplyDiffResults FNiagaraScriptMergeManager::ResetM
 		Results.bSucceeded = false;
 		Results.bModifiedGraph = false;
 		Results.ErrorMessages.Add(FText::Format(LOCTEXT("ResetFailedBecauseOfEmptyDiffMessage", "Failed to reset input back to it's base value.  It wasn't different from the base.  Emitter: {0}  Input:{1}"),
-			FText::FromString(Emitter.GetPathName()), FText::FromString(InputName)));
+			EmitterPath, FText::FromString(InputName)));
 		return Results;
 	}
 
-	if (Emitter.ParentScratchPadScripts.Num() != BaseEmitter.ParentScratchPadScripts.Num() + BaseEmitter.ScratchPadScripts.Num())
+	FVersionedNiagaraEmitterData* EmitterData = VersionedEmitter.GetEmitterData();
+	FVersionedNiagaraEmitterData* BaseEmitterData = VersionedBaseEmitter.GetEmitterData();
+	if (EmitterData->ParentScratchPads->Scripts.Num() != (BaseEmitterData->ParentScratchPads->Scripts.Num() + BaseEmitterData->ScratchPads->Scripts.Num()))
 	{
 		FApplyDiffResults Results;
 		Results.bSucceeded = false;
 		Results.bModifiedGraph = false;
 		Results.ErrorMessages.Add(FText::Format(LOCTEXT("ResetFailedBecauseOfScratchPadScripts", "Failed to reset input back to it's base value.  Its scratch pad scripts were out of sync.  Emitter: {0}  Input:{1}"),
-			FText::FromString(Emitter.GetPathName()), FText::FromString(InputName)));
+			EmitterPath, FText::FromString(InputName)));
 	}
 	
 	// Remove items from the diff which are not relevant to this input.
@@ -1803,13 +1847,13 @@ FNiagaraScriptMergeManager::FApplyDiffResults FNiagaraScriptMergeManager::ResetM
 	return ApplyScriptStackDiff(EmitterAdapter, EmitterAdapter->GetScriptStack(ScriptUsage, ScriptUsageId).ToSharedRef(), ScratchPadMergeAdapter, ResetDiffResults, false);
 }
 
-bool FNiagaraScriptMergeManager::HasBaseEventHandler(const UNiagaraEmitter& BaseEmitter, FGuid EventScriptUsageId)
+bool FNiagaraScriptMergeManager::HasBaseEventHandler(const FVersionedNiagaraEmitter& BaseEmitter, FGuid EventScriptUsageId)
 {
 	TSharedRef<FNiagaraEmitterMergeAdapter> BaseEmitterAdapter = GetEmitterMergeAdapterUsingCache(BaseEmitter);
 	return BaseEmitterAdapter->GetEventHandler(EventScriptUsageId).IsValid();
 }
 
-bool FNiagaraScriptMergeManager::IsEventHandlerPropertySetDifferentFromBase(UNiagaraEmitter& Emitter, const UNiagaraEmitter& BaseEmitter, FGuid EventScriptUsageId)
+bool FNiagaraScriptMergeManager::IsEventHandlerPropertySetDifferentFromBase(const FVersionedNiagaraEmitter& Emitter, const FVersionedNiagaraEmitter& BaseEmitter, FGuid EventScriptUsageId)
 {
 	TSharedRef<FNiagaraEmitterMergeAdapter> EmitterAdapter = GetEmitterMergeAdapterUsingCache(Emitter);
 	TSharedRef<FNiagaraEmitterMergeAdapter> BaseEmitterAdapter = GetEmitterMergeAdapterUsingCache(BaseEmitter);
@@ -1827,9 +1871,9 @@ bool FNiagaraScriptMergeManager::IsEventHandlerPropertySetDifferentFromBase(UNia
 	return DifferentProperties.Num() > 0;
 }
 
-void FNiagaraScriptMergeManager::ResetEventHandlerPropertySetToBase(UNiagaraEmitter& Emitter, const UNiagaraEmitter& BaseEmitter, FGuid EventScriptUsageId)
+void FNiagaraScriptMergeManager::ResetEventHandlerPropertySetToBase(const FVersionedNiagaraEmitter& VersionedEmitter, const FVersionedNiagaraEmitter& BaseEmitter, FGuid EventScriptUsageId)
 {
-	TSharedRef<FNiagaraEmitterMergeAdapter> EmitterAdapter = GetEmitterMergeAdapterUsingCache(Emitter);
+	TSharedRef<FNiagaraEmitterMergeAdapter> EmitterAdapter = GetEmitterMergeAdapterUsingCache(VersionedEmitter);
 	TSharedRef<FNiagaraEmitterMergeAdapter> BaseEmitterAdapter = GetEmitterMergeAdapterUsingCache(BaseEmitter);
 
 	TSharedPtr<FNiagaraEventHandlerMergeAdapter> EventHandlerAdapter = EmitterAdapter->GetEventHandler(EventScriptUsageId);
@@ -1844,16 +1888,16 @@ void FNiagaraScriptMergeManager::ResetEventHandlerPropertySetToBase(UNiagaraEmit
 	TArray<FProperty*> DifferentProperties;
 	DiffEditableProperties(BaseEventHandlerAdapter->GetEventScriptProperties(), EventHandlerAdapter->GetEventScriptProperties(), *FNiagaraEventScriptProperties::StaticStruct(), DifferentProperties);
 	CopyPropertiesToBase(EventHandlerAdapter->GetEditableEventScriptProperties(), BaseEventHandlerAdapter->GetEventScriptProperties(), DifferentProperties);
-	Emitter.PostEditChange();
+	VersionedEmitter.Emitter->PostEditChange();
 }
 
-bool FNiagaraScriptMergeManager::HasBaseSimulationStage(const UNiagaraEmitter& BaseEmitter, FGuid SimulationStageScriptUsageId)
+bool FNiagaraScriptMergeManager::HasBaseSimulationStage(const FVersionedNiagaraEmitter& BaseEmitter, FGuid SimulationStageScriptUsageId)
 {
 	TSharedRef<FNiagaraEmitterMergeAdapter> BaseEmitterAdapter = GetEmitterMergeAdapterUsingCache(BaseEmitter);
 	return BaseEmitterAdapter->GetSimulationStage(SimulationStageScriptUsageId).IsValid();
 }
 
-bool FNiagaraScriptMergeManager::IsSimulationStagePropertySetDifferentFromBase(UNiagaraEmitter& Emitter, const UNiagaraEmitter& BaseEmitter, FGuid SimulationStageScriptUsageId)
+bool FNiagaraScriptMergeManager::IsSimulationStagePropertySetDifferentFromBase(const FVersionedNiagaraEmitter& Emitter, const FVersionedNiagaraEmitter& BaseEmitter, FGuid SimulationStageScriptUsageId)
 {
 	TSharedRef<FNiagaraEmitterMergeAdapter> EmitterAdapter = GetEmitterMergeAdapterUsingCache(Emitter);
 	TSharedRef<FNiagaraEmitterMergeAdapter> BaseEmitterAdapter = GetEmitterMergeAdapterUsingCache(BaseEmitter);
@@ -1871,9 +1915,9 @@ bool FNiagaraScriptMergeManager::IsSimulationStagePropertySetDifferentFromBase(U
 	return DifferentProperties.Num() > 0;
 }
 
-void FNiagaraScriptMergeManager::ResetSimulationStagePropertySetToBase(UNiagaraEmitter& Emitter, const UNiagaraEmitter& BaseEmitter, FGuid SimulationStageScriptUsageId)
+void FNiagaraScriptMergeManager::ResetSimulationStagePropertySetToBase(const FVersionedNiagaraEmitter& VersionedEmitter, const FVersionedNiagaraEmitter& BaseEmitter, FGuid SimulationStageScriptUsageId)
 {
-	TSharedRef<FNiagaraEmitterMergeAdapter> EmitterAdapter = GetEmitterMergeAdapterUsingCache(Emitter);
+	TSharedRef<FNiagaraEmitterMergeAdapter> EmitterAdapter = GetEmitterMergeAdapterUsingCache(VersionedEmitter);
 	TSharedRef<FNiagaraEmitterMergeAdapter> BaseEmitterAdapter = GetEmitterMergeAdapterUsingCache(BaseEmitter);
 
 	TSharedPtr<FNiagaraSimulationStageMergeAdapter> SimulationStageAdapter = EmitterAdapter->GetSimulationStage(SimulationStageScriptUsageId);
@@ -1888,16 +1932,16 @@ void FNiagaraScriptMergeManager::ResetSimulationStagePropertySetToBase(UNiagaraE
 	TArray<FProperty*> DifferentProperties;
 	DiffEditableProperties(BaseSimulationStageAdapter->GetSimulationStage(), SimulationStageAdapter->GetSimulationStage(), *BaseSimulationStageAdapter->GetSimulationStage()->GetClass(), DifferentProperties);
 	CopyPropertiesToBase(SimulationStageAdapter->GetEditableSimulationStage(), BaseSimulationStageAdapter->GetSimulationStage(), DifferentProperties);
-	Emitter.PostEditChange();
+	VersionedEmitter.Emitter->PostEditChange();
 }
 
-bool FNiagaraScriptMergeManager::HasBaseRenderer(const UNiagaraEmitter& BaseEmitter, FGuid RendererMergeId)
+bool FNiagaraScriptMergeManager::HasBaseRenderer(const FVersionedNiagaraEmitter& BaseEmitter, FGuid RendererMergeId)
 {
 	TSharedRef<FNiagaraEmitterMergeAdapter> BaseEmitterAdapter = GetEmitterMergeAdapterUsingCache(BaseEmitter);
 	return BaseEmitterAdapter->GetRenderer(RendererMergeId).IsValid();
 }
 
-bool FNiagaraScriptMergeManager::IsRendererDifferentFromBase(UNiagaraEmitter& Emitter, const UNiagaraEmitter& BaseEmitter, FGuid RendererMergeId)
+bool FNiagaraScriptMergeManager::IsRendererDifferentFromBase(const FVersionedNiagaraEmitter& Emitter, const FVersionedNiagaraEmitter& BaseEmitter, FGuid RendererMergeId)
 {
 	TSharedRef<FNiagaraEmitterMergeAdapter> EmitterAdapter = GetEmitterMergeAdapterUsingCache(Emitter);
 	TSharedRef<FNiagaraEmitterMergeAdapter> BaseEmitterAdapter = GetEmitterMergeAdapterUsingCache(BaseEmitter);
@@ -1919,7 +1963,7 @@ bool FNiagaraScriptMergeManager::IsRendererDifferentFromBase(UNiagaraEmitter& Em
 	return DiffResults.ModifiedOtherRenderers.FindByPredicate(FindRendererByMergeId) != nullptr;
 }
 
-void FNiagaraScriptMergeManager::ResetRendererToBase(UNiagaraEmitter& Emitter, const UNiagaraEmitter& BaseEmitter, FGuid RendererMergeId)
+void FNiagaraScriptMergeManager::ResetRendererToBase(FVersionedNiagaraEmitter Emitter, const FVersionedNiagaraEmitter& BaseEmitter, FGuid RendererMergeId)
 {
 	TSharedRef<FNiagaraEmitterMergeAdapter> EmitterAdapter = GetEmitterMergeAdapterUsingCache(Emitter);
 	TSharedRef<FNiagaraEmitterMergeAdapter> BaseEmitterAdapter = GetEmitterMergeAdapterUsingCache(BaseEmitter);
@@ -1942,22 +1986,22 @@ void FNiagaraScriptMergeManager::ResetRendererToBase(UNiagaraEmitter& Emitter, c
 	ApplyRendererDiff(Emitter, ResetDiffResults, false);
 }
 
-bool FNiagaraScriptMergeManager::IsEmitterEditablePropertySetDifferentFromBase(const UNiagaraEmitter& Emitter, const UNiagaraEmitter& BaseEmitter)
+bool FNiagaraScriptMergeManager::IsEmitterEditablePropertySetDifferentFromBase(const FVersionedNiagaraEmitter& Emitter, const FVersionedNiagaraEmitter& BaseEmitter)
 {
 	TArray<FProperty*> DifferentProperties;
-	DiffEditableProperties(&BaseEmitter, &Emitter, *UNiagaraEmitter::StaticClass(), DifferentProperties);
+	DiffEditableProperties(BaseEmitter.GetEmitterData(), Emitter.GetEmitterData(), *FVersionedNiagaraEmitterData::StaticStruct(), DifferentProperties);
 	return DifferentProperties.Num() > 0;
 }
 
-void FNiagaraScriptMergeManager::ResetEmitterEditablePropertySetToBase(UNiagaraEmitter& Emitter, const UNiagaraEmitter& BaseEmitter)
+void FNiagaraScriptMergeManager::ResetEmitterEditablePropertySetToBase(const FVersionedNiagaraEmitter& VersionedEmitter, const FVersionedNiagaraEmitter& BaseEmitter)
 {
 	TArray<FProperty*> DifferentProperties;
-	DiffEditableProperties(&BaseEmitter, &Emitter, *UNiagaraEmitter::StaticClass(), DifferentProperties);
-	CopyPropertiesToBase(&Emitter, &BaseEmitter, DifferentProperties);
-	Emitter.PostEditChange();
+	DiffEditableProperties(BaseEmitter.GetEmitterData(), VersionedEmitter.GetEmitterData(), *FVersionedNiagaraEmitterData::StaticStruct(), DifferentProperties);
+	CopyPropertiesToBase(VersionedEmitter.GetEmitterData(), BaseEmitter.GetEmitterData(), DifferentProperties);
+	VersionedEmitter.Emitter->PostEditChange();
 }
 
-FNiagaraEmitterDiffResults FNiagaraScriptMergeManager::DiffEmitters(UNiagaraEmitter& BaseEmitter, UNiagaraEmitter& OtherEmitter) const
+FNiagaraEmitterDiffResults FNiagaraScriptMergeManager::DiffEmitters(const FVersionedNiagaraEmitter& BaseEmitter, const FVersionedNiagaraEmitter& OtherEmitter) const
 {
 	SCOPE_CYCLE_COUNTER(STAT_NiagaraEditor_ScriptMergeManager_DiffEmitters);
 
@@ -2000,20 +2044,22 @@ FNiagaraEmitterDiffResults FNiagaraScriptMergeManager::DiffEmitters(UNiagaraEmit
 	{
 		EmitterDiffResults.AddError(LOCTEXT("ParticleUpdateStacksInvalidMessage", "One of the particle update script stacks was invalid."));
 	}
+	FVersionedNiagaraEmitterData* BaseEmitterData = BaseEmitter.GetEmitterData();
+	FVersionedNiagaraEmitterData* OtherEmitterData = OtherEmitter.GetEmitterData();
 
 	DiffEventHandlers(BaseEmitterAdapter->GetEventHandlers(), OtherEmitterAdapter->GetEventHandlers(), EmitterDiffResults);
 	DiffSimulationStages(BaseEmitterAdapter->GetSimulationStages(), OtherEmitterAdapter->GetSimulationStages(), EmitterDiffResults);
 	DiffRenderers(BaseEmitterAdapter->GetRenderers(), OtherEmitterAdapter->GetRenderers(), EmitterDiffResults);
 	DiffEmitterSummary(BaseEmitterAdapter->GetEditorData(), OtherEmitterAdapter->GetEditorData(), EmitterDiffResults);
-	DiffEditableProperties(&BaseEmitter, &OtherEmitter, *UNiagaraEmitter::StaticClass(), EmitterDiffResults.DifferentEmitterProperties);
+	DiffEditableProperties(BaseEmitterData, OtherEmitterData, *FVersionedNiagaraEmitterData::StaticStruct(), EmitterDiffResults.DifferentEmitterProperties);
 	DiffStackEntryDisplayNames(BaseEmitterAdapter->GetEditorData(), OtherEmitterAdapter->GetEditorData(), EmitterDiffResults.ModifiedStackEntryDisplayNames);
 
 	TArray<UNiagaraScript*> BaseScratchPadScripts;
-	BaseScratchPadScripts.Append(BaseEmitter.ParentScratchPadScripts);
-	BaseScratchPadScripts.Append(BaseEmitter.ScratchPadScripts);
+	BaseScratchPadScripts.Append(BaseEmitterData->ParentScratchPads->Scripts);
+	BaseScratchPadScripts.Append(BaseEmitterData->ScratchPads->Scripts);
 	TArray<UNiagaraScript*> OtherScratchPadScripts;
-	OtherScratchPadScripts.Append(OtherEmitter.ParentScratchPadScripts);
-	OtherScratchPadScripts.Append(OtherEmitter.ScratchPadScripts);
+	OtherScratchPadScripts.Append(OtherEmitterData->ParentScratchPads->Scripts);
+	OtherScratchPadScripts.Append(OtherEmitterData->ScratchPads->Scripts);
 	DiffScratchPadScripts(BaseScratchPadScripts, OtherScratchPadScripts, EmitterDiffResults);
 
 	EmitterDiffResults.BaseEmitterAdapter = BaseEmitterAdapter;
@@ -2485,7 +2531,7 @@ void FNiagaraScriptMergeManager::DiffScratchPadScripts(const TArray<UNiagaraScri
 	{
 		UNiagaraScript* BaseScratchPadScript = BaseScratchPadScripts[ScratchPadScriptIndex];
 		UNiagaraScript* OtherScratchPadScript = OtherEmitterScratchPadScripts[ScratchPadScriptIndex];
-		if (BaseScratchPadScript->GetName() != OtherScratchPadScript->GetName() ||
+		if (!BaseScratchPadScript->GetFName().IsEqual(OtherScratchPadScript->GetFName(), ENameCase::IgnoreCase, false) ||
 			BaseScratchPadScript->GetBaseChangeID() != OtherScratchPadScript->GetBaseChangeID())
 		{
 			DiffResults.bScratchPadModified = true;
@@ -2890,7 +2936,7 @@ FNiagaraScriptMergeManager::FApplyDiffResults FNiagaraScriptMergeManager::AddInp
 		if (OverrideToAdd->GetLocalValueRapidIterationParameter().IsSet())
 		{
 			FNiagaraVariable RapidIterationParameter = FNiagaraStackGraphUtilities::CreateRapidIterationParameter(
-				BaseEmitterAdapter->GetEditableEmitter()->GetUniqueEmitterName(), OwningScript.GetUsage(), AliasedFunctionInputHandle.GetParameterHandleString(), OverrideToAdd->GetLocalValueRapidIterationParameter().GetValue().GetType());
+				BaseEmitterAdapter->GetEditableEmitter().Emitter->GetUniqueEmitterName(), OwningScript.GetUsage(), AliasedFunctionInputHandle.GetParameterHandleString(), OverrideToAdd->GetLocalValueRapidIterationParameter().GetValue().GetType());
 			const uint8* SourceData = OverrideToAdd->GetOwningScript()->RapidIterationParameters.GetParameterData(OverrideToAdd->GetLocalValueRapidIterationParameter().GetValue());
 			OwningScript.Modify();
 			bool bAddParameterIfMissing = true;
@@ -2949,13 +2995,14 @@ void FNiagaraScriptMergeManager::CopyPropertiesToBase(void* BaseDataAddress, con
 	}
 }
 
-void FNiagaraScriptMergeManager::CopyInstanceScratchPadScripts(UNiagaraEmitter& MergedInstance, const UNiagaraEmitter& SourceInstance) const
+void FNiagaraScriptMergeManager::CopyInstanceScratchPadScripts(const FVersionedNiagaraEmitter& MergedInstance, FVersionedNiagaraEmitterData* SourceInstance) const
 {
-	for (UNiagaraScript* SourceScratchPadScript : SourceInstance.ScratchPadScripts)
+	TObjectPtr<UNiagaraScratchPadContainer> ScratchPadContainer = MergedInstance.GetEmitterData()->ScratchPads;
+	for (UNiagaraScript* SourceScratchPadScript : SourceInstance->ScratchPads->Scripts)
 	{
-		FName UniqueObjectName = FNiagaraEditorUtilities::GetUniqueObjectName<UNiagaraScript>(&MergedInstance, SourceScratchPadScript->GetName());
-		UNiagaraScript* MergedInstanceScratchPadScript = CastChecked<UNiagaraScript>(StaticDuplicateObject(SourceScratchPadScript, &MergedInstance, UniqueObjectName));
-		MergedInstance.ScratchPadScripts.Add(MergedInstanceScratchPadScript);
+		FName UniqueObjectName = FNiagaraEditorUtilities::GetUniqueObjectName<UNiagaraScript>(ScratchPadContainer, SourceScratchPadScript->GetName());
+		UNiagaraScript* MergedInstanceScratchPadScript = CastChecked<UNiagaraScript>(StaticDuplicateObject(SourceScratchPadScript, ScratchPadContainer, UniqueObjectName));
+		ScratchPadContainer->Scripts.Add(MergedInstanceScratchPadScript);
 	}
 }
 
@@ -3210,7 +3257,7 @@ FNiagaraScriptMergeManager::FApplyDiffResults FNiagaraScriptMergeManager::ApplyE
 		}
 	}
 
-	UNiagaraScriptSource* EmitterSource = CastChecked<UNiagaraScriptSource>(BaseEmitterAdapter->GetEditableEmitter()->GraphSource);
+	UNiagaraScriptSource* EmitterSource = CastChecked<UNiagaraScriptSource>(BaseEmitterAdapter->GetEditableEmitter().GetEmitterData()->GraphSource);
 	UNiagaraGraph* EmitterGraph = EmitterSource->NodeGraph;
 	for (TSharedRef<FNiagaraEventHandlerMergeAdapter> AddedEventHandler : DiffResults.AddedOtherEventHandlers)
 	{
@@ -3230,13 +3277,14 @@ FNiagaraScriptMergeManager::FApplyDiffResults FNiagaraScriptMergeManager::ApplyE
 		}
 		else
 		{
-			UNiagaraEmitter* BaseEmitter = BaseEmitterAdapter->GetEditableEmitter();
+			FVersionedNiagaraEmitter BaseEmitter = BaseEmitterAdapter->GetEditableEmitter();
 			FNiagaraEventScriptProperties AddedEventScriptProperties = *AddedEventHandler->GetEventScriptProperties();
-			AddedEventScriptProperties.Script = NewObject<UNiagaraScript>(BaseEmitter, MakeUniqueObjectName(BaseEmitter, UNiagaraScript::StaticClass(), "EventScript"), EObjectFlags::RF_Transactional);
+			AddedEventScriptProperties.Script = NewObject<UNiagaraScript>(BaseEmitter.Emitter, MakeUniqueObjectName(BaseEmitter.Emitter, UNiagaraScript::StaticClass(), "EventScript"),
+			                                                              RF_Transactional);
 			AddedEventScriptProperties.Script->SetUsage(ENiagaraScriptUsage::ParticleEventScript);
 			AddedEventScriptProperties.Script->SetUsageId(AddedEventHandler->GetUsageId());
 			AddedEventScriptProperties.Script->SetLatestSource(EmitterSource);
-			BaseEmitter->AddEventHandler(AddedEventScriptProperties);
+			BaseEmitter.Emitter->AddEventHandler(AddedEventScriptProperties, BaseEmitter.Version);
 
 			FGuid PreferredOutputNodeGuid = AddedEventHandler->GetOutputNode()->NodeGuid;
 			FGuid PreferredInputNodeGuid = AddedEventHandler->GetInputNode()->NodeGuid;
@@ -3316,7 +3364,8 @@ FNiagaraScriptMergeManager::FApplyDiffResults FNiagaraScriptMergeManager::ApplyS
 		}
 	}
 
-	UNiagaraScriptSource* EmitterSource = CastChecked<UNiagaraScriptSource>(BaseEmitterAdapter->GetEditableEmitter()->GraphSource);
+	FVersionedNiagaraEmitterData* BaseEmitterData = BaseEmitterAdapter->GetEditableEmitter().GetEmitterData();
+	UNiagaraScriptSource* EmitterSource = CastChecked<UNiagaraScriptSource>(BaseEmitterData->GraphSource);
 	UNiagaraGraph* EmitterGraph = EmitterSource->NodeGraph;
 	for (TSharedRef<FNiagaraSimulationStageMergeAdapter> AddedOtherSimulationStage : DiffResults.AddedOtherSimulationStages)
 	{
@@ -3336,16 +3385,17 @@ FNiagaraScriptMergeManager::FApplyDiffResults FNiagaraScriptMergeManager::ApplyS
 		}
 		else
 		{
-			UNiagaraEmitter* BaseEmitter = BaseEmitterAdapter->GetEditableEmitter();
+			UNiagaraEmitter* BaseEmitter = BaseEmitterAdapter->GetEditableEmitter().Emitter;
 			UNiagaraSimulationStageBase* AddedSimulationStage = CastChecked<UNiagaraSimulationStageBase>(StaticDuplicateObject(AddedOtherSimulationStage->GetSimulationStage(), BaseEmitter));
-			AddedSimulationStage->Script = NewObject<UNiagaraScript>(AddedSimulationStage, MakeUniqueObjectName(AddedSimulationStage, UNiagaraScript::StaticClass(), "SimulationStage"), EObjectFlags::RF_Transactional);
+			AddedSimulationStage->Script = NewObject<UNiagaraScript>(AddedSimulationStage, MakeUniqueObjectName(AddedSimulationStage, UNiagaraScript::StaticClass(), "SimulationStage"),
+			                                                         RF_Transactional);
 			AddedSimulationStage->Script->SetUsage(ENiagaraScriptUsage::ParticleSimulationStageScript);
 			AddedSimulationStage->Script->SetUsageId(AddedOtherSimulationStage->GetUsageId());
 			AddedSimulationStage->Script->SetLatestSource(EmitterSource);
-			BaseEmitter->AddSimulationStage(AddedSimulationStage);
+			BaseEmitter->AddSimulationStage(AddedSimulationStage, BaseEmitterAdapter->GetEditableEmitter().Version);
 
-			int32 TargetIndex = FMath::Min(AddedOtherSimulationStage->GetSimulationStageIndex(), BaseEmitter->GetSimulationStages().Num() - 1);
-			BaseEmitter->MoveSimulationStageToIndex(AddedSimulationStage, TargetIndex);
+			int32 TargetIndex = FMath::Min(AddedOtherSimulationStage->GetSimulationStageIndex(), BaseEmitterData->GetSimulationStages().Num() - 1);
+			BaseEmitter->MoveSimulationStageToIndex(AddedSimulationStage, TargetIndex, BaseEmitterAdapter->GetEditableEmitter().Version);
 
 			FGuid PreferredOutputNodeGuid = AddedOtherSimulationStage->GetOutputNode()->NodeGuid;
 			FGuid PreferredInputNodeGuid = AddedOtherSimulationStage->GetInputNode()->NodeGuid;
@@ -3371,16 +3421,17 @@ FNiagaraScriptMergeManager::FApplyDiffResults FNiagaraScriptMergeManager::ApplyS
 	return Results;
 }
 
-FNiagaraScriptMergeManager::FApplyDiffResults FNiagaraScriptMergeManager::ApplyRendererDiff(UNiagaraEmitter& BaseEmitter, const FNiagaraEmitterDiffResults& DiffResults, const bool bNoParentAtLastMerge) const
+FNiagaraScriptMergeManager::FApplyDiffResults FNiagaraScriptMergeManager::ApplyRendererDiff(const FVersionedNiagaraEmitter& BaseEmitter, const FNiagaraEmitterDiffResults& DiffResults, const bool bNoParentAtLastMerge) const
 {
 	TArray<UNiagaraRendererProperties*> RenderersToRemove;
 	TArray<UNiagaraRendererProperties*> RenderersToAdd;
 	TArray<UNiagaraRendererProperties*> RenderersToDisable;
-
+	FVersionedNiagaraEmitterData* BaseEmitterData = BaseEmitter.GetEmitterData();
+	
 	for (TSharedRef<FNiagaraRendererMergeAdapter> RemovedRenderer : DiffResults.RemovedBaseRenderers)
 	{
 		auto FindRendererByMergeId = [=](UNiagaraRendererProperties* Renderer) { return Renderer->GetMergeId() == RemovedRenderer->GetRenderer()->GetMergeId(); };
-		UNiagaraRendererProperties* const* MatchingRendererPtr = BaseEmitter.GetRenderers().FindByPredicate(FindRendererByMergeId);
+		UNiagaraRendererProperties* const* MatchingRendererPtr = BaseEmitterData->GetRenderers().FindByPredicate(FindRendererByMergeId);
 		if (MatchingRendererPtr != nullptr)
 		{
 			if (bNoParentAtLastMerge)
@@ -3399,28 +3450,28 @@ FNiagaraScriptMergeManager::FApplyDiffResults FNiagaraScriptMergeManager::ApplyR
 
 	for (TSharedRef<FNiagaraRendererMergeAdapter> AddedRenderer : DiffResults.AddedOtherRenderers)
 	{
-		RenderersToAdd.Add(Cast<UNiagaraRendererProperties>(StaticDuplicateObject(AddedRenderer->GetRenderer(), &BaseEmitter)));
+		RenderersToAdd.Add(Cast<UNiagaraRendererProperties>(StaticDuplicateObject(AddedRenderer->GetRenderer(), BaseEmitter.Emitter)));
 	}
 
 	for (TSharedRef<FNiagaraRendererMergeAdapter> ModifiedRenderer : DiffResults.ModifiedOtherRenderers)
 	{
 		auto FindRendererByMergeId = [=](UNiagaraRendererProperties* Renderer) { return Renderer->GetMergeId() == ModifiedRenderer->GetRenderer()->GetMergeId(); };
-		UNiagaraRendererProperties*const* MatchingRendererPtr = BaseEmitter.GetRenderers().FindByPredicate(FindRendererByMergeId);
+		UNiagaraRendererProperties*const* MatchingRendererPtr = BaseEmitterData->GetRenderers().FindByPredicate(FindRendererByMergeId);
 		if (MatchingRendererPtr != nullptr)
 		{
 			RenderersToRemove.Add(*MatchingRendererPtr);
-			RenderersToAdd.Add(Cast<UNiagaraRendererProperties>(StaticDuplicateObject(ModifiedRenderer->GetRenderer(), &BaseEmitter)));
+			RenderersToAdd.Add(Cast<UNiagaraRendererProperties>(StaticDuplicateObject(ModifiedRenderer->GetRenderer(), BaseEmitter.Emitter)));
 		}
 	}
 
 	for (UNiagaraRendererProperties* RendererToRemove : RenderersToRemove)
 	{
-		BaseEmitter.RemoveRenderer(RendererToRemove);
+		BaseEmitter.Emitter->RemoveRenderer(RendererToRemove, BaseEmitter.Version);
 	}
 
 	for (UNiagaraRendererProperties* RendererToAdd : RenderersToAdd)
 	{
-		BaseEmitter.AddRenderer(RendererToAdd);
+		BaseEmitter.Emitter->AddRenderer(RendererToAdd, BaseEmitter.Version);
 	}
 
 	for (UNiagaraRendererProperties* RendererToDisable : RenderersToDisable)
@@ -3434,19 +3485,20 @@ FNiagaraScriptMergeManager::FApplyDiffResults FNiagaraScriptMergeManager::ApplyR
 	return Results;
 }
 
-FNiagaraScriptMergeManager::FApplyDiffResults FNiagaraScriptMergeManager::ApplyEmitterSummaryDiff(UNiagaraEmitter& Emitter, const FNiagaraEmitterDiffResults& DiffResults, const bool bNoParentAtLastMerge) const
+FNiagaraScriptMergeManager::FApplyDiffResults FNiagaraScriptMergeManager::ApplyEmitterSummaryDiff(const FVersionedNiagaraEmitter& BaseEmitter, const FNiagaraEmitterDiffResults& DiffResults) const
 {
-	UNiagaraEmitterEditorData* EditorData = Cast<UNiagaraEmitterEditorData>(Emitter.GetEditorData());
+	FVersionedNiagaraEmitterData* BaseEmitterData = BaseEmitter.GetEmitterData();
+	UNiagaraEmitterEditorData* EditorData = Cast<UNiagaraEmitterEditorData>(BaseEmitterData->GetEditorData());
 
 	if (EditorData == nullptr)
 	{
-		EditorData = NewObject<UNiagaraEmitterEditorData>(&Emitter, NAME_None, RF_Transactional);
-		Emitter.Modify();
-		Emitter.SetEditorData(EditorData);
+		EditorData = NewObject<UNiagaraEmitterEditorData>(BaseEmitter.Emitter, NAME_None, RF_Transactional);
+		BaseEmitter.Emitter->Modify();
+		BaseEmitter.Emitter->SetEditorData(EditorData, BaseEmitter.Version);
 	}
 
 	check(EditorData != GetDefault<UNiagaraEmitterEditorData>());
-	check(EditorData->GetOuter() == &Emitter);
+	check(EditorData->GetOuter() == BaseEmitter.Emitter);
 	
 	EditorData->Modify();
 	for (TSharedRef<FNiagaraInputSummaryMergeAdapter> RemovedInput : DiffResults.RemovedInputSummaryEntries)
@@ -3510,16 +3562,16 @@ FNiagaraScriptMergeManager::FApplyDiffResults FNiagaraScriptMergeManager::ApplyE
 	return Results;
 }
 
-FNiagaraScriptMergeManager::FApplyDiffResults FNiagaraScriptMergeManager::ApplyStackEntryDisplayNameDiffs(UNiagaraEmitter& Emitter, const FNiagaraEmitterDiffResults& DiffResults) const
+FNiagaraScriptMergeManager::FApplyDiffResults FNiagaraScriptMergeManager::ApplyStackEntryDisplayNameDiffs(FVersionedNiagaraEmitter BaseEmitter, const FNiagaraEmitterDiffResults& DiffResults) const
 {
 	if (DiffResults.ModifiedStackEntryDisplayNames.Num() > 0)
 	{
-		UNiagaraEmitterEditorData* EditorData = Cast<UNiagaraEmitterEditorData>(Emitter.GetEditorData());
+		UNiagaraEmitterEditorData* EditorData = Cast<UNiagaraEmitterEditorData>(BaseEmitter.GetEmitterData()->GetEditorData());
 		if (EditorData == nullptr)
 		{
-			EditorData = NewObject<UNiagaraEmitterEditorData>(&Emitter, NAME_None, RF_Transactional);
+			EditorData = NewObject<UNiagaraEmitterEditorData>(BaseEmitter.Emitter, NAME_None, RF_Transactional);
 			EditorData->Modify();
-			Emitter.SetEditorData(EditorData);
+			BaseEmitter.Emitter->SetEditorData(EditorData, BaseEmitter.Version);
 		}
 
 		for (auto& Pair : DiffResults.ModifiedStackEntryDisplayNames)
@@ -3534,18 +3586,18 @@ FNiagaraScriptMergeManager::FApplyDiffResults FNiagaraScriptMergeManager::ApplyS
 	return Results;
 }
 
-FNiagaraScriptMergeManager::FCachedMergeAdapter* FNiagaraScriptMergeManager::FindOrAddMergeAdapterCacheForEmitter(const UNiagaraEmitter& Emitter)
+FNiagaraScriptMergeManager::FCachedMergeAdapter* FNiagaraScriptMergeManager::FindOrAddMergeAdapterCacheForEmitter(const FVersionedNiagaraEmitter& VersionedEmitter)
 {
-	FCachedMergeAdapter* CachedMergeAdapter = CachedMergeAdapters.Find(FObjectKey(&Emitter));
+	FCachedMergeAdapter* CachedMergeAdapter = CachedMergeAdapters.Find(VersionedEmitter);
 	if (CachedMergeAdapter == nullptr)
 	{
-		CachedMergeAdapter = &CachedMergeAdapters.Add(FObjectKey(&Emitter));
+		CachedMergeAdapter = &CachedMergeAdapters.Add(VersionedEmitter);
 	}
 	else
 	{
-		if (CachedMergeAdapter->ChangeId != Emitter.GetChangeId())
+		if (CachedMergeAdapter->ChangeId != VersionedEmitter.Emitter->GetChangeId())
 		{
-			CachedMergeAdapter->ChangeId = Emitter.GetChangeId();
+			CachedMergeAdapter->ChangeId = VersionedEmitter.Emitter->GetChangeId();
 			CachedMergeAdapter->EmitterMergeAdapter.Reset();
 			CachedMergeAdapter->ScratchPadMergeAdapter.Reset();
 		}
@@ -3553,25 +3605,17 @@ FNiagaraScriptMergeManager::FCachedMergeAdapter* FNiagaraScriptMergeManager::Fin
 	return CachedMergeAdapter;
 }
 
-TSharedRef<FNiagaraEmitterMergeAdapter> FNiagaraScriptMergeManager::GetEmitterMergeAdapterUsingCache(const UNiagaraEmitter& Emitter)
+void FNiagaraScriptMergeManager::ClearMergeAdapterCache()
 {
-	FCachedMergeAdapter* CachedMergeAdapter = FindOrAddMergeAdapterCacheForEmitter(Emitter);
-
-	if (CachedMergeAdapter->EmitterMergeAdapter.IsValid() == false ||
-		CachedMergeAdapter->EmitterMergeAdapter->GetEditableEmitter() != nullptr)
-	{
-		CachedMergeAdapter->EmitterMergeAdapter = MakeShared<FNiagaraEmitterMergeAdapter>(Emitter);
-	}
-
-	return CachedMergeAdapter->EmitterMergeAdapter.ToSharedRef();
+	CachedMergeAdapters.Empty();
 }
 
-TSharedRef<FNiagaraEmitterMergeAdapter> FNiagaraScriptMergeManager::GetEmitterMergeAdapterUsingCache(UNiagaraEmitter& Emitter)
+TSharedRef<FNiagaraEmitterMergeAdapter> FNiagaraScriptMergeManager::GetEmitterMergeAdapterUsingCache(const FVersionedNiagaraEmitter& Emitter)
 {
 	FCachedMergeAdapter* CachedMergeAdapter = FindOrAddMergeAdapterCacheForEmitter(Emitter);
 
 	if (CachedMergeAdapter->EmitterMergeAdapter.IsValid() == false ||
-		CachedMergeAdapter->EmitterMergeAdapter->GetEditableEmitter() == nullptr)
+		CachedMergeAdapter->EmitterMergeAdapter->GetEditableEmitter().Emitter != nullptr)
 	{
 		CachedMergeAdapter->EmitterMergeAdapter = MakeShared<FNiagaraEmitterMergeAdapter>(Emitter);
 	}
@@ -3579,13 +3623,13 @@ TSharedRef<FNiagaraEmitterMergeAdapter> FNiagaraScriptMergeManager::GetEmitterMe
 	return CachedMergeAdapter->EmitterMergeAdapter.ToSharedRef();
 }
 
-TSharedRef<FNiagaraScratchPadMergeAdapter> FNiagaraScriptMergeManager::GetScratchPadMergeAdapterUsingCache(UNiagaraEmitter& Emitter)
+TSharedRef<FNiagaraScratchPadMergeAdapter> FNiagaraScriptMergeManager::GetScratchPadMergeAdapterUsingCache(const FVersionedNiagaraEmitter& VersionedEmitter)
 {
-	FCachedMergeAdapter* CachedMergeAdapter = FindOrAddMergeAdapterCacheForEmitter(Emitter);
+	FCachedMergeAdapter* CachedMergeAdapter = FindOrAddMergeAdapterCacheForEmitter(VersionedEmitter);
 
 	if (CachedMergeAdapter->ScratchPadMergeAdapter.IsValid() == false)
 	{
-		CachedMergeAdapter->ScratchPadMergeAdapter = MakeShared<FNiagaraScratchPadMergeAdapter>(&Emitter, nullptr, Emitter.GetParent());
+		CachedMergeAdapter->ScratchPadMergeAdapter = MakeShared<FNiagaraScratchPadMergeAdapter>(VersionedEmitter, FVersionedNiagaraEmitter(), VersionedEmitter.GetEmitterData()->GetParent());
 	}
 
 	return CachedMergeAdapter->ScratchPadMergeAdapter.ToSharedRef();
diff --git a/Engine/Plugins/FX/Niagara/Source/NiagaraEditor/Private/NiagaraScriptMergeManager.h b/Engine/Plugins/FX/Niagara/Source/NiagaraEditor/Private/NiagaraScriptMergeManager.h
--- a/Engine/Plugins/FX/Niagara/Source/NiagaraEditor/Private/NiagaraScriptMergeManager.h
+++ b/Engine/Plugins/FX/Niagara/Source/NiagaraEditor/Private/NiagaraScriptMergeManager.h
@@ -8,11 +8,11 @@
 #include "INiagaraMergeManager.h"
 #include "NiagaraTypes.h"
 #include "NiagaraCommon.h"
+#include "NiagaraEmitter.h"
 #include "NiagaraMessages.h"
 #include "NiagaraEmitterEditorData.h"
 
 #include "Templates/SharedPointer.h"
-#include "UObject/WeakObjectPtr.h"
 #include "UObject/WeakObjectPtrTemplates.h"
 #include "UObject/ObjectKey.h"
 
@@ -36,7 +36,7 @@ class FNiagaraStackFunctionInputOverrideMergeAdapter
 {
 public:
 	FNiagaraStackFunctionInputOverrideMergeAdapter(
-		const UNiagaraEmitter& InOwningEmitter, 
+		const FVersionedNiagaraEmitter& InOwningEmitter, 
 		UNiagaraScript& InOwningScript, 
 		UNiagaraNodeFunctionCall& InOwningFunctionCallNode, 
 		UEdGraphPin& InOverridePin);
@@ -88,7 +88,7 @@ private:
 class FNiagaraStackFunctionMergeAdapter
 {
 public:
-	FNiagaraStackFunctionMergeAdapter(const UNiagaraEmitter& InOwningEmitter, UNiagaraScript& InOwningScript, UNiagaraNodeFunctionCall& InFunctionCallNode,	int32 InStackIndex);
+	FNiagaraStackFunctionMergeAdapter(const FVersionedNiagaraEmitter& InOwningEmitter, UNiagaraScript& InOwningScript, UNiagaraNodeFunctionCall& InFunctionCallNode,	int32 InStackIndex);
 
 	UNiagaraNodeFunctionCall* GetFunctionCallNode() const;
 
@@ -115,7 +115,7 @@ private:
 class FNiagaraScriptStackMergeAdapter
 {
 public:
-	FNiagaraScriptStackMergeAdapter(const UNiagaraEmitter& InOwningEmitter, UNiagaraNodeOutput& InOutputNode, UNiagaraScript& InScript);
+	FNiagaraScriptStackMergeAdapter(const FVersionedNiagaraEmitter& InOwningEmitter, UNiagaraNodeOutput& InOutputNode, UNiagaraScript& InScript);
 
 	UNiagaraNodeOutput* GetOutputNode() const;
 	UNiagaraNodeInput* GetInputNode() const;
@@ -141,12 +141,12 @@ private:
 class FNiagaraEventHandlerMergeAdapter
 {
 public:
-	FNiagaraEventHandlerMergeAdapter(const UNiagaraEmitter& InEmitter, const FNiagaraEventScriptProperties* InEventScriptProperties, UNiagaraNodeOutput* InOutputNode);
-	FNiagaraEventHandlerMergeAdapter(const UNiagaraEmitter& InEmitter, FNiagaraEventScriptProperties* InEventScriptProperties, UNiagaraNodeOutput* InOutputNode);
-	FNiagaraEventHandlerMergeAdapter(const UNiagaraEmitter& InEmitter, UNiagaraNodeOutput* InOutputNode);
+	FNiagaraEventHandlerMergeAdapter(const FVersionedNiagaraEmitter& InEmitter, const FNiagaraEventScriptProperties* InEventScriptProperties, UNiagaraNodeOutput* InOutputNode);
+	FNiagaraEventHandlerMergeAdapter(const FVersionedNiagaraEmitter& InEmitter, FNiagaraEventScriptProperties* InEventScriptProperties, UNiagaraNodeOutput* InOutputNode);
+	FNiagaraEventHandlerMergeAdapter(const FVersionedNiagaraEmitter& InEmitter, UNiagaraNodeOutput* InOutputNode);
 
 	FGuid GetUsageId() const;
-	const UNiagaraEmitter* GetEmitter() const;
+	FVersionedNiagaraEmitter GetEmitter() const;
 	const FNiagaraEventScriptProperties* GetEventScriptProperties() const;
 	FNiagaraEventScriptProperties* GetEditableEventScriptProperties() const;
 	UNiagaraNodeOutput* GetOutputNode() const;
@@ -154,10 +154,10 @@ public:
 	TSharedPtr<FNiagaraScriptStackMergeAdapter> GetEventStack() const;
 
 private:
-	void Initialize(const UNiagaraEmitter& InEmitter, const FNiagaraEventScriptProperties* InEventScriptProperties, FNiagaraEventScriptProperties* InEditableEventScriptProperties, UNiagaraNodeOutput* InOutputNode);
+	void Initialize(const FVersionedNiagaraEmitter& InEmitter, const FNiagaraEventScriptProperties* InEventScriptProperties, FNiagaraEventScriptProperties* InEditableEventScriptProperties, UNiagaraNodeOutput* InOutputNode);
 
 private:
-	TWeakObjectPtr<UNiagaraEmitter> Emitter;
+	FVersionedNiagaraEmitterWeakPtr Emitter;
 	const FNiagaraEventScriptProperties* EventScriptProperties;
 	FNiagaraEventScriptProperties* EditableEventScriptProperties;
 	TWeakObjectPtr<UNiagaraNodeOutput> OutputNode;
@@ -169,12 +169,12 @@ private:
 class FNiagaraSimulationStageMergeAdapter
 {
 public:
-	FNiagaraSimulationStageMergeAdapter(const UNiagaraEmitter& InEmitter, const UNiagaraSimulationStageBase* InSimulationStage, int32 InSimulationStageIndex, UNiagaraNodeOutput* InOutputNode);
-	FNiagaraSimulationStageMergeAdapter(const UNiagaraEmitter& InEmitter, UNiagaraSimulationStageBase* InSimulationStage, int32 InSimulationStageIndex, UNiagaraNodeOutput* InOutputNode);
-	FNiagaraSimulationStageMergeAdapter(const UNiagaraEmitter& InEmitter, UNiagaraNodeOutput* InOutputNode);
+	FNiagaraSimulationStageMergeAdapter(const FVersionedNiagaraEmitter& InEmitter, const UNiagaraSimulationStageBase* InSimulationStage, int32 InSimulationStageIndex, UNiagaraNodeOutput* InOutputNode);
+	FNiagaraSimulationStageMergeAdapter(const FVersionedNiagaraEmitter& InEmitter, UNiagaraSimulationStageBase* InSimulationStage, int32 InSimulationStageIndex, UNiagaraNodeOutput* InOutputNode);
+	FNiagaraSimulationStageMergeAdapter(const FVersionedNiagaraEmitter& InEmitter, UNiagaraNodeOutput* InOutputNode);
 
 	FGuid GetUsageId() const;
-	const UNiagaraEmitter* GetEmitter() const;
+	FVersionedNiagaraEmitter GetEmitter() const;
 	const UNiagaraSimulationStageBase* GetSimulationStage() const;
 	UNiagaraSimulationStageBase* GetEditableSimulationStage() const;
 	UNiagaraNodeOutput* GetOutputNode() const;
@@ -183,10 +183,10 @@ public:
 	TSharedPtr<FNiagaraScriptStackMergeAdapter> GetSimulationStageStack() const;
 
 private:
-	void Initialize(const UNiagaraEmitter& InEmitter, const UNiagaraSimulationStageBase* InSimulationStage, UNiagaraSimulationStageBase* InEditableSimulationStage, int32 InSimulationStageIndex, UNiagaraNodeOutput* InOutputNode);
+	void Initialize(const FVersionedNiagaraEmitter& InEmitter, const UNiagaraSimulationStageBase* InSimulationStage, UNiagaraSimulationStageBase* InEditableSimulationStage, int32 InSimulationStageIndex, UNiagaraNodeOutput* InOutputNode);
 
 private:
-	TWeakObjectPtr<UNiagaraEmitter> Emitter;
+	FVersionedNiagaraEmitterWeakPtr Emitter;
 	const UNiagaraSimulationStageBase* SimulationStage;
 	UNiagaraSimulationStageBase* EditableSimulationStage;
 	TWeakObjectPtr<UNiagaraNodeOutput> OutputNode;
@@ -225,7 +225,7 @@ class FNiagaraScratchPadMergeAdapter
 public:
 	FNiagaraScratchPadMergeAdapter();
 
-	FNiagaraScratchPadMergeAdapter(UNiagaraEmitter* InTargetEmitter, UNiagaraEmitter* InInstanceEmitter, UNiagaraEmitter* InParentEmitter);
+	FNiagaraScratchPadMergeAdapter(const FVersionedNiagaraEmitter InTargetEmitter, const FVersionedNiagaraEmitter& InInstanceEmitter, const FVersionedNiagaraEmitter& InParentEmitter);
 
 	UNiagaraScript* GetScratchPadScriptForFunctionId(FGuid FunctionId);
 
@@ -233,9 +233,9 @@ private:
 	void Initialize();
 
 private:
-	UNiagaraEmitter* TargetEmitter;
-	UNiagaraEmitter* InstanceEmitter;
-	UNiagaraEmitter* ParentEmitter;
+	FVersionedNiagaraEmitter TargetEmitter;
+	FVersionedNiagaraEmitter InstanceEmitter;
+	FVersionedNiagaraEmitter ParentEmitter;
 	bool bIsInitialized;
 	TMap<FGuid, UNiagaraScript*> FunctionIdToScratchPadScript;
 };
@@ -243,10 +243,9 @@ private:
 class FNiagaraEmitterMergeAdapter
 {
 public:
-	FNiagaraEmitterMergeAdapter(const UNiagaraEmitter& InEmitter);
-	FNiagaraEmitterMergeAdapter(UNiagaraEmitter& InEmitter);
+	FNiagaraEmitterMergeAdapter(const FVersionedNiagaraEmitter& InEmitter);
 
-	UNiagaraEmitter* GetEditableEmitter() const;
+	FVersionedNiagaraEmitter GetEditableEmitter() const;
 
 	TSharedPtr<FNiagaraScriptStackMergeAdapter> GetEmitterSpawnStack() const;
 	TSharedPtr<FNiagaraScriptStackMergeAdapter> GetEmitterUpdateStack() const;
@@ -268,11 +267,11 @@ public:
 	void GatherFunctionCallNodes(TArray<UNiagaraNodeFunctionCall*>& OutFunctionCallNodes) const;
 
 private:
-	void Initialize(const UNiagaraEmitter& InEmitter, UNiagaraEmitter* InEditableEmitter);
+	void Initialize(const FVersionedNiagaraEmitter& InEmitter, const FVersionedNiagaraEmitter& InEditableEmitter);
 
 private:
-	TWeakObjectPtr<const UNiagaraEmitter> Emitter;
-	TWeakObjectPtr<UNiagaraEmitter> EditableEmitter;
+	FVersionedNiagaraEmitterWeakPtr Emitter;
+	FVersionedNiagaraEmitterWeakPtr EditableEmitter;
 	TSharedPtr<FNiagaraScriptStackMergeAdapter> EmitterSpawnStack;
 	TSharedPtr<FNiagaraScriptStackMergeAdapter> EmitterUpdateStack;
 	TSharedPtr<FNiagaraScriptStackMergeAdapter> ParticleSpawnStack;
@@ -420,43 +419,43 @@ public:
 	};
 
 public:
-	void UpdateModuleVersions(UNiagaraEmitter& Instance, const FNiagaraEmitterDiffResults& EmitterDiffResults) const;
+	void UpdateModuleVersions(const FVersionedNiagaraEmitter& Instance, const FNiagaraEmitterDiffResults& EmitterDiffResults) const;
 	
-	virtual INiagaraMergeManager::FMergeEmitterResults MergeEmitter(UNiagaraEmitter& Parent, UNiagaraEmitter* ParentAtLastMerge, UNiagaraEmitter& Instance) const override;
+	virtual INiagaraMergeManager::FMergeEmitterResults MergeEmitter(const FVersionedNiagaraEmitter& Parent, const FVersionedNiagaraEmitter& ParentAtLastMerge, const FVersionedNiagaraEmitter& Instance) const override;
 
 	static TSharedRef<FNiagaraScriptMergeManager> Get();
 
-	FNiagaraEmitterDiffResults DiffEmitters(UNiagaraEmitter& BaseEmitter, UNiagaraEmitter& OtherEmitter) const;
+	FNiagaraEmitterDiffResults DiffEmitters(const FVersionedNiagaraEmitter& BaseEmitter, const FVersionedNiagaraEmitter& OtherEmitter) const;
 
 	bool IsMergeableScriptUsage(ENiagaraScriptUsage ScriptUsage) const;
 
-	bool HasBaseModule(const UNiagaraEmitter& BaseEmitter, ENiagaraScriptUsage ScriptUsage, FGuid ScriptUsageId, FGuid ModuleId);
+	bool HasBaseModule(const FVersionedNiagaraEmitter& BaseEmitter, ENiagaraScriptUsage ScriptUsage, FGuid ScriptUsageId, FGuid ModuleId);
 
-	bool IsModuleInputDifferentFromBase(UNiagaraEmitter& Emitter, const UNiagaraEmitter& BaseEmitter, ENiagaraScriptUsage ScriptUsage, FGuid ScriptUsageId, FGuid ModuleId, FString InputName);
+	bool IsModuleInputDifferentFromBase(const FVersionedNiagaraEmitter& Emitter, const FVersionedNiagaraEmitter& BaseEmitter, ENiagaraScriptUsage ScriptUsage, FGuid ScriptUsageId, FGuid ModuleId, FString InputName);
 
-	FApplyDiffResults ResetModuleInputToBase(UNiagaraEmitter& Emitter, const UNiagaraEmitter& BaseEmitter, ENiagaraScriptUsage ScriptUsage, FGuid ScriptUsageId, FGuid ModuleId, FString InputName);
+	FApplyDiffResults ResetModuleInputToBase(const FVersionedNiagaraEmitter& VersionedEmitter, const FVersionedNiagaraEmitter& VersionedBaseEmitter, ENiagaraScriptUsage ScriptUsage, FGuid ScriptUsageId, FGuid ModuleId, FString InputName);
 
-	bool HasBaseEventHandler(const UNiagaraEmitter& BaseEmitter, FGuid EventScriptUsageId);
+	bool HasBaseEventHandler(const FVersionedNiagaraEmitter& BaseEmitter, FGuid EventScriptUsageId);
 
-	bool IsEventHandlerPropertySetDifferentFromBase(UNiagaraEmitter& Emitter, const UNiagaraEmitter& BaseEmitter, FGuid EventScriptUsageId);
+	bool IsEventHandlerPropertySetDifferentFromBase(const FVersionedNiagaraEmitter& Emitter, const FVersionedNiagaraEmitter& BaseEmitter, FGuid EventScriptUsageId);
 
-	void ResetEventHandlerPropertySetToBase(UNiagaraEmitter& Emitter, const UNiagaraEmitter& BaseEmitter, FGuid EventScriptUsageId);
+	void ResetEventHandlerPropertySetToBase(const FVersionedNiagaraEmitter& VersionedEmitter, const FVersionedNiagaraEmitter& BaseEmitter, FGuid EventScriptUsageId);
 
-	bool HasBaseSimulationStage(const UNiagaraEmitter& BaseEmitter, FGuid SimulationStageScriptUsageId);
+	bool HasBaseSimulationStage(const FVersionedNiagaraEmitter& BaseEmitter, FGuid SimulationStageScriptUsageId);
 
-	bool IsSimulationStagePropertySetDifferentFromBase(UNiagaraEmitter& Emitter, const UNiagaraEmitter& BaseEmitter, FGuid SimulationStageScriptUsageId);
+	bool IsSimulationStagePropertySetDifferentFromBase(const FVersionedNiagaraEmitter& Emitter, const FVersionedNiagaraEmitter& BaseEmitter, FGuid SimulationStageScriptUsageId);
 
-	void ResetSimulationStagePropertySetToBase(UNiagaraEmitter& Emitter, const UNiagaraEmitter& BaseEmitter, FGuid SimulationStageScriptUsageId);
+	void ResetSimulationStagePropertySetToBase(const FVersionedNiagaraEmitter& VersionedEmitter, const FVersionedNiagaraEmitter& BaseEmitter, FGuid SimulationStageScriptUsageId);
 
-	bool HasBaseRenderer(const UNiagaraEmitter& BaseEmitter, FGuid RendererMergeId);
+	bool HasBaseRenderer(const FVersionedNiagaraEmitter& BaseEmitter, FGuid RendererMergeId);
 
-	bool IsRendererDifferentFromBase(UNiagaraEmitter& Emitter, const UNiagaraEmitter& BaseEmitter, FGuid RendererMergeId);
+	bool IsRendererDifferentFromBase(const FVersionedNiagaraEmitter& Emitter, const FVersionedNiagaraEmitter& BaseEmitter, FGuid RendererMergeId);
 
-	void ResetRendererToBase(UNiagaraEmitter& Emitter, const UNiagaraEmitter& BaseEmitter, FGuid RendererMergeId);
+	void ResetRendererToBase(FVersionedNiagaraEmitter Emitter, const FVersionedNiagaraEmitter& BaseEmitter, FGuid RendererMergeId);
 
-	bool IsEmitterEditablePropertySetDifferentFromBase(const UNiagaraEmitter& Emitter, const UNiagaraEmitter& BaseEmitter);
+	bool IsEmitterEditablePropertySetDifferentFromBase(const FVersionedNiagaraEmitter& Emitter, const FVersionedNiagaraEmitter& BaseEmitter);
 
-	void ResetEmitterEditablePropertySetToBase(UNiagaraEmitter& Emitter, const UNiagaraEmitter& BaseEmitter);
+	void ResetEmitterEditablePropertySetToBase(const FVersionedNiagaraEmitter& VersionedEmitter, const FVersionedNiagaraEmitter& BaseEmitter);
 
 	void DiffEventHandlers(const TArray<TSharedRef<FNiagaraEventHandlerMergeAdapter>>& BaseEventHandlers, const TArray<TSharedRef<FNiagaraEventHandlerMergeAdapter>>& OtherEventHandlers, FNiagaraEmitterDiffResults& DiffResults) const;
 
@@ -478,6 +477,7 @@ public:
 
 	virtual void CopyPropertiesToBase(void* BaseDataAddress, const void* OtherDataAddress, TArray<FProperty*> PropertiesToCopy) const override;
 
+	void ClearMergeAdapterCache();
 
 private:
 	struct FCachedMergeAdapter
@@ -487,10 +487,16 @@ private:
 		TSharedPtr<FNiagaraScratchPadMergeAdapter> ScratchPadMergeAdapter;
 	};
 
-private:
+	enum class EmitterMergeState
+	{
+		Failed, Unchanged, DuplicateParent, CopyProperties
+	};
+
+	EmitterMergeState GetEmitterMergeState(const FNiagaraEmitterDiffResults& DiffResults, const FVersionedNiagaraEmitter& Parent, const FVersionedNiagaraEmitter& Instance) const;
+
 	TOptional<bool> DoFunctionInputOverridesMatch(TSharedRef<FNiagaraStackFunctionInputOverrideMergeAdapter> BaseFunctionInputAdapter, TSharedRef<FNiagaraStackFunctionInputOverrideMergeAdapter> OtherFunctionInputAdapter) const;
 
-	void CopyInstanceScratchPadScripts(UNiagaraEmitter& MergedInstance, const UNiagaraEmitter& SourceInstance) const;
+	void CopyInstanceScratchPadScripts(const FVersionedNiagaraEmitter& MergedInstance, FVersionedNiagaraEmitterData* SourceInstance) const;
 
 	FApplyDiffResults ApplyScriptStackDiff(
 		TSharedRef<FNiagaraEmitterMergeAdapter> BaseEmitterAdapter,
@@ -511,11 +517,11 @@ private:
 		const FNiagaraEmitterDiffResults& DiffResults,
 		const bool bNoParentAtLastMerge) const;
 
-	FApplyDiffResults ApplyRendererDiff(UNiagaraEmitter& BaseEmitter, const FNiagaraEmitterDiffResults& DiffResults, const bool bNoParentAtLastMerge) const;
+	FApplyDiffResults ApplyRendererDiff(const FVersionedNiagaraEmitter& BaseEmitter, const FNiagaraEmitterDiffResults& DiffResults, const bool bNoParentAtLastMerge) const;
 
-	FApplyDiffResults ApplyEmitterSummaryDiff(UNiagaraEmitter& Emitter, const FNiagaraEmitterDiffResults& DiffResults, const bool bNoParentAtLastMerge) const;
+	FApplyDiffResults ApplyEmitterSummaryDiff(const FVersionedNiagaraEmitter& BaseEmitter, const FNiagaraEmitterDiffResults& DiffResults) const;
 	
-	FApplyDiffResults ApplyStackEntryDisplayNameDiffs(UNiagaraEmitter& Emitter, const FNiagaraEmitterDiffResults& DiffResults) const;
+	FApplyDiffResults ApplyStackEntryDisplayNameDiffs(FVersionedNiagaraEmitter BaseEmitter, const FNiagaraEmitterDiffResults& DiffResults) const;
 
 	FApplyDiffResults AddModule(
 		TSharedRef<FNiagaraEmitterMergeAdapter> BaseEmitterAdapter,
@@ -533,13 +539,11 @@ private:
 		UNiagaraNodeFunctionCall& TargetFunctionCall,
 		TSharedRef<FNiagaraStackFunctionInputOverrideMergeAdapter> OverrideToAdd) const;
 
-	FCachedMergeAdapter* FindOrAddMergeAdapterCacheForEmitter(const UNiagaraEmitter& Emitter);
-
-	TSharedRef<FNiagaraEmitterMergeAdapter> GetEmitterMergeAdapterUsingCache(const UNiagaraEmitter& Emitter);
+	FCachedMergeAdapter* FindOrAddMergeAdapterCacheForEmitter(const FVersionedNiagaraEmitter& VersionedEmitter);
 
-	TSharedRef<FNiagaraEmitterMergeAdapter> GetEmitterMergeAdapterUsingCache(UNiagaraEmitter& Emitter);
+	TSharedRef<FNiagaraEmitterMergeAdapter> GetEmitterMergeAdapterUsingCache(const FVersionedNiagaraEmitter& Emitter);
 
-	TSharedRef<FNiagaraScratchPadMergeAdapter> GetScratchPadMergeAdapterUsingCache(UNiagaraEmitter& Emitter);
+	TSharedRef<FNiagaraScratchPadMergeAdapter> GetScratchPadMergeAdapterUsingCache(const FVersionedNiagaraEmitter& VersionedEmitter);
 
 	void GetForcedChangeIds(
 		const TMap<FGuid, UNiagaraNodeFunctionCall*>& InParentFunctionIdToNodeMap,
@@ -547,8 +551,7 @@ private:
 		const TMap<FGuid, UNiagaraNodeFunctionCall*>& InInstanceFunctionIdToNodeMap,
 		TMap<FGuid, FGuid>& OutFunctionIdToForcedChangeId) const;
 
-	FApplyDiffResults ForceInstanceChangeIds(TSharedRef<FNiagaraEmitterMergeAdapter> MergedInstanceAdapter, UNiagaraEmitter& OriginalEmitterInstance, const TMap<FGuid, FGuid>& ChangeIdsThatNeedToBeReset) const;
+	FApplyDiffResults ForceInstanceChangeIds(TSharedRef<FNiagaraEmitterMergeAdapter> MergedInstanceAdapter, const FVersionedNiagaraEmitter& OriginalEmitterInstance, const TMap<FGuid, FGuid>& ChangeIdsThatNeedToBeReset) const;
 
-private:
-	TMap<FObjectKey, FCachedMergeAdapter> CachedMergeAdapters;
+	TMap<FVersionedNiagaraEmitter, FCachedMergeAdapter> CachedMergeAdapters;
 };
\ No newline at end of file
