Added versioning for Niagara emitters

diff --git a/Engine/Plugins/FX/Niagara/Source/NiagaraEditor/Private/NiagaraScriptSource.cpp b/Engine/Plugins/FX/Niagara/Source/NiagaraEditor/Private/NiagaraScriptSource.cpp
--- a/Engine/Plugins/FX/Niagara/Source/NiagaraEditor/Private/NiagaraScriptSource.cpp
+++ b/Engine/Plugins/FX/Niagara/Source/NiagaraEditor/Private/NiagaraScriptSource.cpp
@@ -79,8 +79,7 @@ void UNiagaraScriptSource::RefreshFromExternalChanges()
 	{
 		for (UEdGraphNode* Node : NodeGraph->Nodes)
 		{
-			UNiagaraNode* NiagaraNode = Cast<UNiagaraNode>(Node);
-			if (NiagaraNode)
+			if (UNiagaraNode* NiagaraNode = Cast<UNiagaraNode>(Node))
 			{
 				NiagaraNode->RefreshFromExternalChanges();
 			}
@@ -145,16 +144,17 @@ void UNiagaraScriptSource::MarkNotSynchronized(FString Reason)
 	}
 }
 
-void UNiagaraScriptSource::PostLoadFromEmitter(UNiagaraEmitter& OwningEmitter)
+void UNiagaraScriptSource::PostLoadFromEmitter(FVersionedNiagaraEmitter OwningEmitter)
 {
 	const int32 NiagaraCustomVersion = GetLinkerCustomVersion(FNiagaraCustomVersion::GUID);
 	if (NiagaraCustomVersion < FNiagaraCustomVersion::ScriptsNowUseAGuidForIdentificationInsteadOfAnIndex)
 	{
 		TArray<UNiagaraNodeOutput*> OutputNodes;
 		NodeGraph->GetNodesOfClass<UNiagaraNodeOutput>(OutputNodes);
-		for (int32 i = 0; i < OwningEmitter.GetEventHandlers().Num(); i++)
+		FVersionedNiagaraEmitterData* EmitterData = OwningEmitter.GetEmitterData();
+		for (int32 i = 0; i < EmitterData->GetEventHandlers().Num(); i++)
 		{
-			const FNiagaraEventScriptProperties& EventScriptProperties = OwningEmitter.GetEventHandlers()[i];
+			const FNiagaraEventScriptProperties& EventScriptProperties = EmitterData->GetEventHandlers()[i];
 			EventScriptProperties.Script->SetUsageId(FGuid::NewGuid());
 
 			auto FindOutputNodeByUsageIndex = [=](UNiagaraNodeOutput* OutputNode) 
@@ -186,8 +186,7 @@ void FindObjectNamesRecursive(UNiagaraGraph* InGraph, ENiagaraScriptUsage Usage,
 	{
 		if (UNiagaraNode* InNode = Cast<UNiagaraNode>(Node))
 		{
-			UNiagaraNodeInput* InputNode = Cast<UNiagaraNodeInput>(InNode);
-			if (InputNode)
+			if (UNiagaraNodeInput* InputNode = Cast<UNiagaraNodeInput>(InNode))
 			{
 				if (InputNode->Input.IsDataInterface())
 				{
@@ -203,8 +202,7 @@ void FindObjectNamesRecursive(UNiagaraGraph* InGraph, ENiagaraScriptUsage Usage,
 				continue;
 			}
 
-			UNiagaraNodeFunctionCall* FunctionCallNode = Cast<UNiagaraNodeFunctionCall>(InNode);
-			if (FunctionCallNode)
+			if (UNiagaraNodeFunctionCall* FunctionCallNode = Cast<UNiagaraNodeFunctionCall>(InNode))
 			{
 				UNiagaraGraph* FunctionGraph = FunctionCallNode->GetCalledGraph();
 				if(VisitedGraphs.Contains(FunctionGraph) == false)
@@ -261,10 +259,9 @@ TMap<FName, UNiagaraDataInterface*> UNiagaraScriptSource::ComputeObjectNameMap(U
 
 			for (const FNiagaraEmitterHandle& EmitterHandle : System.GetEmitterHandles())
 			{
-				if (EmitterHandle.GetIsEnabled() && EmitterHandle.GetInstance() != nullptr)
+				if (EmitterHandle.GetIsEnabled() && EmitterHandle.GetInstance().Emitter != nullptr)
 				{
-					UNiagaraScriptSource* EmitterSource = Cast<UNiagaraScriptSource>(EmitterHandle.GetInstance()->GraphSource);
-					if (EmitterSource != nullptr)
+					if (UNiagaraScriptSource* EmitterSource = Cast<UNiagaraScriptSource>(EmitterHandle.GetEmitterData()->GraphSource); EmitterSource != nullptr)
 					{
 						FindObjectNamesRecursive(EmitterSource->NodeGraph, EmitterUsage, FGuid(), EmitterHandle.GetUniqueInstanceName(), VisitedGraphs, Result);
 					}
@@ -299,7 +296,7 @@ bool UNiagaraScriptSource::AddModuleIfMissing(FString ModulePath, ENiagaraScript
 	return false;
 }
 
-void UNiagaraScriptSource::FixupRenamedParameters(UNiagaraNode* Node, FNiagaraParameterStore& RapidIterationParameters, const TArray<FNiagaraVariable>& OldRapidIterationVariables, const UNiagaraEmitter* Emitter, ENiagaraScriptUsage ScriptUsage) const
+void UNiagaraScriptSource::FixupRenamedParameters(UNiagaraNode* Node, FNiagaraParameterStore& RapidIterationParameters, const TArray<FNiagaraVariable>& OldRapidIterationVariables, const FVersionedNiagaraEmitter& VersionedEmitter, ENiagaraScriptUsage ScriptUsage) const
 {
 	UNiagaraNodeFunctionCall* FunctionCallNode = Cast<UNiagaraNodeFunctionCall>(Node);
 	if (FunctionCallNode == nullptr || FunctionCallNode->HasValidScriptAndGraph() == false)
@@ -316,7 +313,7 @@ void UNiagaraScriptSource::FixupRenamedParameters(UNiagaraNode* Node, FNiagaraPa
 	// the rapid iteration parameters and the function input pins use different variable naming schemes, so most of this is just used to convert one name to the other 
 	UNiagaraGraph* Graph = FunctionCallNode->GetCalledGraph();
 	const UEdGraphSchema_Niagara* NiagaraSchema = Graph->GetNiagaraSchema();
-	const FString UniqueEmitterName = Emitter ? Emitter->GetUniqueEmitterName() : FString();
+	const FString UniqueEmitterName = VersionedEmitter.Emitter ? VersionedEmitter.Emitter->GetUniqueEmitterName() : FString();
 
 	// go through the existing rapid iteration params to see if they are either still valid or were renamed
 	for (FNiagaraVariable OldRapidIterationVar : OldRapidIterationVariables)
@@ -379,7 +376,7 @@ void UNiagaraScriptSource::FixupRenamedParameters(UNiagaraNode* Node, FNiagaraPa
 	}
 }
 
-void InitializeNewRapidIterationParametersForNode(const UEdGraphSchema_Niagara* Schema, UNiagaraNode* Node, const UNiagaraEmitter* Emitter, ENiagaraScriptUsage ScriptUsage, FNiagaraParameterStore& RapidIterationParameters, TSet<FName>& ValidRapidIterationParameterNames)
+void InitializeNewRapidIterationParametersForNode(const UEdGraphSchema_Niagara* Schema, UNiagaraNode* Node, const FVersionedNiagaraEmitter& Emitter, ENiagaraScriptUsage ScriptUsage, FNiagaraParameterStore& RapidIterationParameters, TSet<FName>& ValidRapidIterationParameterNames)
 {
 	UNiagaraNodeFunctionCall* FunctionCallNode = Cast<UNiagaraNodeFunctionCall>(Node);
 	if (FunctionCallNode == nullptr)
@@ -390,7 +387,7 @@ void InitializeNewRapidIterationParametersForNode(const UEdGraphSchema_Niagara*
 	FCompileConstantResolver Resolver(Emitter, ScriptUsage);
 	TArray<const UEdGraphPin*> FunctionInputPins;
 
-	const FString UniqueEmitterName = Emitter ? Emitter->GetUniqueEmitterName() : FString();
+	const FString UniqueEmitterName = Emitter.Emitter ? Emitter.Emitter->GetUniqueEmitterName() : FString();
 
 	FNiagaraStackGraphUtilities::GetStackFunctionInputPins(*FunctionCallNode, FunctionInputPins, HiddenPins, Resolver, FNiagaraStackGraphUtilities::ENiagaraGetStackFunctionInputPinsOptions::ModuleInputsOnly, false);
 	for (const UEdGraphPin* FunctionInputPin : FunctionInputPins)
@@ -433,7 +430,7 @@ void InitializeNewRapidIterationParametersForNode(const UEdGraphSchema_Niagara*
 	}
 }
 
-void UNiagaraScriptSource::CleanUpOldAndInitializeNewRapidIterationParameters(const UNiagaraEmitter* Emitter, ENiagaraScriptUsage ScriptUsage, FGuid ScriptUsageId, FNiagaraParameterStore& RapidIterationParameters) const
+void UNiagaraScriptSource::CleanUpOldAndInitializeNewRapidIterationParameters(const FVersionedNiagaraEmitter& Emitter, ENiagaraScriptUsage ScriptUsage, FGuid ScriptUsageId, FNiagaraParameterStore& RapidIterationParameters) const
 {
 	SCOPE_CYCLE_COUNTER(STAT_NiagaraEditor_ScriptSource_InitializeNewRapidIterationParameters);
 	TArray<UNiagaraNodeOutput*> OutputNodes;
@@ -548,26 +545,44 @@ FGuid UNiagaraScriptSource::GetChangeID()
 	return NodeGraph->GetChangeID(); 
 }
 
+FVersionedNiagaraEmitter UNiagaraScriptSource::GetOuterEmitter() const
+{
+	UNiagaraEmitter* Emitter = Cast<UNiagaraEmitter>(GetOuter());
+	if (Emitter == nullptr)
+	{
+		return FVersionedNiagaraEmitter();
+	}
+	for (FNiagaraAssetVersion& Version : Emitter->GetAllAvailableVersions())
+	{
+		if (FVersionedNiagaraEmitterData* EmitterData = Emitter->GetEmitterData(Version.VersionGuid); EmitterData->GraphSource == this)
+		{
+			return FVersionedNiagaraEmitter(Emitter, Version.VersionGuid);
+		}
+	}
+	return FVersionedNiagaraEmitter();
+}
+
 void UNiagaraScriptSource::CollectDataInterfaces(TArray<const UNiagaraDataInterfaceBase*>& DataInterfaces) const
 {
-	if (NodeGraph)
+	if (!NodeGraph)
 	{
-		TArray<UNiagaraNodeOutput*> OutputNodes;
-		NodeGraph->FindOutputNodes(OutputNodes);
+		return;
+	}
+	TArray<UNiagaraNodeOutput*> OutputNodes;
+	NodeGraph->FindOutputNodes(OutputNodes);
 
-		for (UNiagaraNodeOutput* OutputNode : OutputNodes)
-		{
-			TArray<UNiagaraNode*> TraversalNodes;
-			NodeGraph->BuildTraversal(TraversalNodes, OutputNode, true);
+	for (UNiagaraNodeOutput* OutputNode : OutputNodes)
+	{
+		TArray<UNiagaraNode*> TraversalNodes;
+		NodeGraph->BuildTraversal(TraversalNodes, OutputNode, true);
 
-			for (const UNiagaraNode* TraversalNode : TraversalNodes)
+		for (const UNiagaraNode* TraversalNode : TraversalNodes)
+		{
+			if (const UNiagaraNodeInput* NodeInput = Cast<const UNiagaraNodeInput>(TraversalNode))
 			{
-				if (const UNiagaraNodeInput* NodeInput = Cast<const UNiagaraNodeInput>(TraversalNode))
+				if (NodeInput->Input.IsDataInterface())
 				{
-					if (NodeInput->Input.IsDataInterface())
-					{
-						DataInterfaces.Add(NodeInput->GetDataInterface());
-					}
+					DataInterfaces.Add(NodeInput->GetDataInterface());
 				}
 			}
 		}
@@ -581,6 +596,10 @@ void UNiagaraScriptSource::SynchronizeGraphParametersWithParameterDefinitions(
 	INiagaraParameterDefinitionsSubscriber* Subscriber,
 	FSynchronizeWithParameterDefinitionsArgs Args)
 {
+	if (!NodeGraph)
+	{
+		return;
+	}
 	TArray<UNiagaraParameterDefinitions*> QualifiedTargetDefinitions = FNiagaraEditorUtilities::DowncastParameterDefinitionsBaseArray(TargetDefinitions);
 	TArray<UNiagaraParameterDefinitions*> QualifiedAllDefinitions = FNiagaraEditorUtilities::DowncastParameterDefinitionsBaseArray(AllDefinitions);
 	NodeGraph->SynchronizeParametersWithParameterDefinitions(QualifiedTargetDefinitions, QualifiedAllDefinitions, AllDefinitionsParameterIds, Subscriber, Args);
@@ -588,6 +607,10 @@ void UNiagaraScriptSource::SynchronizeGraphParametersWithParameterDefinitions(
 
 void UNiagaraScriptSource::RenameGraphAssignmentAndSetNodePins(const FName OldName, const FName NewName)
 {
+	if (!NodeGraph)
+	{
+		return;
+	}
 	NodeGraph->RenameAssignmentAndSetNodePins(OldName, NewName);
 }
 
@@ -658,3 +681,12 @@ void UNiagaraScriptSource::ChangedLinkedInputTypes(const FNiagaraVariable& Param
 	}
 	NodeGraph->ChangeParameterType({ParametersToChange}, NewType);
 }
+
+void UNiagaraScriptSource::ReplaceScriptReferences(UNiagaraScript* OldScript, UNiagaraScript* NewScript)
+{
+	if (!NodeGraph)
+	{
+		return;
+	}
+	NodeGraph->ReplaceScriptReferences(OldScript, NewScript);
+}
diff --git a/Engine/Plugins/FX/Niagara/Source/NiagaraEditor/Private/NiagaraSystemEditorData.cpp b/Engine/Plugins/FX/Niagara/Source/NiagaraEditor/Private/NiagaraSystemEditorData.cpp
--- a/Engine/Plugins/FX/Niagara/Source/NiagaraEditor/Private/NiagaraSystemEditorData.cpp
+++ b/Engine/Plugins/FX/Niagara/Source/NiagaraEditor/Private/NiagaraSystemEditorData.cpp
@@ -12,7 +12,7 @@
 #include "EdGraphSchema_NiagaraSystemOverview.h"
 #include "EdGraph/EdGraph.h"
 
-static const float SystemOverviewNodePadding = 250.0f;
+static constexpr float SystemOverviewNodePadding = 250.0f;
 
 const FName UNiagaraSystemEditorFolder::GetFolderName() const
 {
@@ -174,7 +174,7 @@ void UNiagaraSystemEditorData::UpdatePlaybackRangeFromEmitters(UNiagaraSystem& O
 
 		for (const FNiagaraEmitterHandle& EmitterHandle : OwnerSystem.GetEmitterHandles())
 		{
-			UNiagaraEmitterEditorData* EmitterEditorData = Cast<UNiagaraEmitterEditorData>(EmitterHandle.GetInstance()->GetEditorData());
+			UNiagaraEmitterEditorData* EmitterEditorData = Cast<UNiagaraEmitterEditorData>(EmitterHandle.GetEmitterData()->GetEditorData());
 			if (EmitterEditorData != nullptr)
 			{
 				EmitterPlaybackRangeMin = FMath::Min(PlaybackRangeMin, EmitterEditorData->GetPlaybackRange().GetLowerBoundValue());
@@ -215,8 +215,6 @@ void FindOuterNodes(const TArray<UNiagaraOverviewNode*> OverviewNodes, UNiagaraO
 
 void UNiagaraSystemEditorData::SynchronizeOverviewGraphWithSystem(UNiagaraSystem& OwnerSystem)
 {
-	float NextNodePosX = 0.0;
-
 	TArray<UNiagaraOverviewNode*> OverviewNodes;
 	SystemOverviewGraph->GetNodesOfClass<UNiagaraOverviewNode>(OverviewNodes);
 
diff --git a/Engine/Plugins/FX/Niagara/Source/NiagaraEditor/Private/NiagaraSystemFactoryNew.cpp b/Engine/Plugins/FX/Niagara/Source/NiagaraEditor/Private/NiagaraSystemFactoryNew.cpp
--- a/Engine/Plugins/FX/Niagara/Source/NiagaraEditor/Private/NiagaraSystemFactoryNew.cpp
+++ b/Engine/Plugins/FX/Niagara/Source/NiagaraEditor/Private/NiagaraSystemFactoryNew.cpp
@@ -2,7 +2,6 @@
 
 #include "NiagaraSystemFactoryNew.h"
 #include "CoreMinimal.h"
-#include "Misc/ConfigCacheIni.h"
 #include "NiagaraSystem.h"
 #include "NiagaraScriptSource.h"
 #include "NiagaraGraph.h"
@@ -27,7 +26,6 @@ UNiagaraSystemFactoryNew::UNiagaraSystemFactoryNew(const FObjectInitializer& Obj
 : Super(ObjectInitializer)
 {
 	SupportedClass = UNiagaraSystem::StaticClass();
-	bCreateNew = false;
 	bEditAfterNew = true;
 	bCreateNew = true;
 }
@@ -76,7 +74,7 @@ bool UNiagaraSystemFactoryNew::ConfigureProperties()
 			UNiagaraEmitter* EmitterToAdd = Cast<UNiagaraEmitter>(EmitterAssetToAdd.GetAsset());
 			if (EmitterToAdd != nullptr)
 			{
-				EmittersToAddToNewSystem.Add(EmitterToAdd);
+				EmittersToAddToNewSystem.Add(FVersionedNiagaraEmitter(EmitterToAdd, EmitterToAdd->GetExposedVersion().VersionGuid));
 			}
 			else
 			{
@@ -157,9 +155,9 @@ UObject* UNiagaraSystemFactoryNew::FactoryCreateNew(UClass* Class, UObject* InPa
 		NewSystem = NewObject<UNiagaraSystem>(InParent, Class, Name, Flags | RF_Transactional);
 		InitializeSystem(NewSystem, true);
 
-		for (UNiagaraEmitter* EmitterToAddToNewSystem : EmittersToAddToNewSystem)
+		for (FVersionedNiagaraEmitter& EmitterToAddToNewSystem : EmittersToAddToNewSystem)
 		{
-			FNiagaraEditorUtilities::AddEmitterToSystem(*NewSystem, *EmitterToAddToNewSystem);
+			FNiagaraEditorUtilities::AddEmitterToSystem(*NewSystem, *EmitterToAddToNewSystem.Emitter, EmitterToAddToNewSystem.Version);
 		}
 	}
 	else
diff --git a/Engine/Plugins/FX/Niagara/Source/NiagaraEditor/Private/NiagaraValidationRules.cpp b/Engine/Plugins/FX/Niagara/Source/NiagaraEditor/Private/NiagaraValidationRules.cpp
--- a/Engine/Plugins/FX/Niagara/Source/NiagaraEditor/Private/NiagaraValidationRules.cpp
+++ b/Engine/Plugins/FX/Niagara/Source/NiagaraEditor/Private/NiagaraValidationRules.cpp
@@ -20,7 +20,6 @@
 #include "ViewModels/NiagaraEmitterViewModel.h"
 
 #include "AssetToolsModule.h"
-#include "Toolkits/AssetEditorToolkit.h"
 
 #define LOCTEXT_NAMESPACE "NiagaraValidationRules"
 
@@ -130,8 +129,8 @@ void UNiagaraValidationRule_FixedGPUBoundsSet::CheckValidity(TSharedPtr<FNiagara
 	TArray<TSharedRef<FNiagaraEmitterHandleViewModel>> EmitterHandleViewModels = ViewModel->GetEmitterHandleViewModels();
 	for (TSharedRef<FNiagaraEmitterHandleViewModel> EmitterHandleModel : EmitterHandleViewModels)
 	{
-		UNiagaraEmitter* NiagaraEmitter = EmitterHandleModel.Get().GetEmitterHandle()->GetInstance();
-		if (NiagaraEmitter->SimTarget == ENiagaraSimTarget::GPUComputeSim && NiagaraEmitter->CalculateBoundsMode == ENiagaraEmitterCalculateBoundMode::Dynamic)
+		FVersionedNiagaraEmitterData* EmitterData = EmitterHandleModel.Get().GetEmitterHandle()->GetEmitterData();
+		if (EmitterData->SimTarget == ENiagaraSimTarget::GPUComputeSim && EmitterData->CalculateBoundsMode == ENiagaraEmitterCalculateBoundMode::Dynamic)
 		{
 			UNiagaraStackEmitterPropertiesItem* EmitterProperties = GetStackEntry<UNiagaraStackEmitterPropertiesItem>(EmitterHandleModel.Get().GetEmitterStackViewModel());
 			FNiagaraValidationResult Result(ENiagaraValidationSeverity::Error, LOCTEXT("GpuDynamicBoundsErrorSummary", "GPU emitters do not support dynamic bounds"), LOCTEXT("GpuDynamicBoundsErrorDescription", "Gpu emitter should either not be in dynamic mode or the system must have fixed bounds."), EmitterProperties);
@@ -158,9 +157,8 @@ void UNiagaraValidationRule_BannedRenderers::CheckValidity(TSharedPtr<FNiagaraSy
 	TArray<TSharedRef<FNiagaraEmitterHandleViewModel>> EmitterHandleViewModels = ViewModel->GetEmitterHandleViewModels();
 	for (TSharedRef<FNiagaraEmitterHandleViewModel> EmitterHandleModel : EmitterHandleViewModels)
 	{
-		int32 MaxQualityLevel = 3;
-		UNiagaraEmitter* NiagaraEmitter = EmitterHandleModel.Get().GetEmitterHandle()->GetInstance();
-		NiagaraEmitter->ForEachRenderer([&Results, EmitterHandleModel, MaxQualityLevel, this, &System](UNiagaraRendererProperties* RendererProperties)
+		FVersionedNiagaraEmitterData* EmitterData = EmitterHandleModel.Get().GetEmitterHandle()->GetEmitterData();
+		EmitterData->ForEachRenderer([&Results, EmitterHandleModel, this, &System](UNiagaraRendererProperties* RendererProperties)
 		{
 			if (RendererProperties->GetIsEnabled() && BannedRenderers.Contains(RendererProperties->GetClass()))
 			{
@@ -225,15 +223,15 @@ void UNiagaraValidationRule_BannedModules::CheckValidity(TSharedPtr<FNiagaraSyst
 			{
 				if (BannedModule == FuncCall.FunctionScript)
 				{
-					UNiagaraEmitter* Emitter = Item->GetEmitterViewModel().IsValid() ? Item->GetEmitterViewModel()->GetEmitter() : nullptr;
+					FVersionedNiagaraEmitterData* EmitterData = Item->GetEmitterViewModel().IsValid() ? Item->GetEmitterViewModel()->GetEmitter().GetEmitterData() : nullptr;
 
 					bool bApplyBan = true;
-					if (Emitter)
+					if (EmitterData)
 					{
 						//If we're on an emitter, this emitter may be culled on the platforms the rule applies to.
 						TArray<const FNiagaraPlatformSet*> CheckSets;
 						CheckSets.Add(&Platforms);
-						CheckSets.Add(&Emitter->Platforms);
+						CheckSets.Add(&EmitterData->Platforms);
 						TArray<FNiagaraPlatformSetConflictInfo> Conflicts;
 						FNiagaraPlatformSet::GatherConflicts(CheckSets, Conflicts);
 						bApplyBan = Conflicts.Num() > 0;
@@ -295,7 +293,7 @@ void UNiagaraValidationRule_LWC::CheckValidity(TSharedPtr<FNiagaraSystemViewMode
 	TArray<TSharedRef<FNiagaraEmitterHandleViewModel>> EmitterHandleViewModels = ViewModel->GetEmitterHandleViewModels();
 	for (TSharedRef<FNiagaraEmitterHandleViewModel> EmitterHandleModel : EmitterHandleViewModels)
 	{
-		if (EmitterHandleModel->GetEmitterHandle()->GetInstance()->bLocalSpace == false)
+		if (EmitterHandleModel->GetEmitterHandle()->GetEmitterData()->bLocalSpace == false)
 		{
 			AllModules.Append(GetStackEntries<UNiagaraStackModuleItem>(EmitterHandleModel.Get().GetEmitterStackViewModel()));
 		}
diff --git a/Engine/Plugins/FX/Niagara/Source/NiagaraEditor/Private/Sequencer/NiagaraSequence/MovieSceneNiagaraEmitterTrack.cpp b/Engine/Plugins/FX/Niagara/Source/NiagaraEditor/Private/Sequencer/NiagaraSequence/MovieSceneNiagaraEmitterTrack.cpp
--- a/Engine/Plugins/FX/Niagara/Source/NiagaraEditor/Private/Sequencer/NiagaraSequence/MovieSceneNiagaraEmitterTrack.cpp
+++ b/Engine/Plugins/FX/Niagara/Source/NiagaraEditor/Private/Sequencer/NiagaraSequence/MovieSceneNiagaraEmitterTrack.cpp
@@ -71,7 +71,7 @@ TSharedRef<ISequencerSection> UMovieSceneNiagaraEmitterSectionBase::MakeInvalidS
 }
 
 UMovieSceneNiagaraEmitterTrack::UMovieSceneNiagaraEmitterTrack(const FObjectInitializer& ObjectInitializer)
-	: Super(ObjectInitializer)
+	: Super(ObjectInitializer), SystemViewModel(nullptr)
 {
 	bSectionsWereModified = false;
 }
@@ -179,14 +179,14 @@ void UMovieSceneNiagaraEmitterTrack::Tick(float DeltaTime)
 {
 	if(!HasAnyFlags(RF_ClassDefaultObject) && bScalabilityModeActive)
 	{
-		UNiagaraEmitter* Emitter = nullptr;
+		FVersionedNiagaraEmitterData* EmitterData = nullptr;
 		// the emitter instance might no longer be valid if we deleted the emitter but the track still lives until garbage collection
 		if(EmitterHandleViewModel.IsValid() && EmitterHandleViewModel.Pin()->GetEmitterHandle() != nullptr)
 		{
-			Emitter = EmitterHandleViewModel.Pin()->GetEmitterHandle()->GetInstance();
+			EmitterData = EmitterHandleViewModel.Pin()->GetEmitterHandle()->GetEmitterData();
 		}
 		
-		if(Emitter && !Emitter->IsAllowedByScalability())
+		if(EmitterData && !EmitterData->IsAllowedByScalability())
 		{
 			SetColorTint(FNiagaraEditorStyle::Get().GetColor("NiagaraEditor.SystemOverview.ExcludedFromScalability").ToFColor(true));
 		}
@@ -269,7 +269,7 @@ void UMovieSceneNiagaraEmitterTrack::CreateSections(const FFrameRate& InFrameRes
 {
 	SectionInitializationErrors.Empty();
 
-	UNiagaraScript* EmitterUpdateScript = GetEmitterHandleViewModel()->GetEmitterViewModel()->GetEmitter()->GetScript(ENiagaraScriptUsage::EmitterUpdateScript, FGuid());
+	UNiagaraScript* EmitterUpdateScript = GetEmitterHandleViewModel()->GetEmitterViewModel()->GetEmitter().GetEmitterData()->GetScript(ENiagaraScriptUsage::EmitterUpdateScript, FGuid());
 	UNiagaraScriptSource* ScriptSource = CastChecked<UNiagaraScriptSource>(EmitterUpdateScript->GetLatestSource());
 	UNiagaraNodeOutput* OutputNode = ScriptSource->NodeGraph->FindOutputNode(ENiagaraScriptUsage::EmitterUpdateScript);
 
diff --git a/Engine/Plugins/FX/Niagara/Source/NiagaraEditor/Private/Sequencer/NiagaraSequence/NiagaraEmitterTrackEditor.cpp b/Engine/Plugins/FX/Niagara/Source/NiagaraEditor/Private/Sequencer/NiagaraSequence/NiagaraEmitterTrackEditor.cpp
--- a/Engine/Plugins/FX/Niagara/Source/NiagaraEditor/Private/Sequencer/NiagaraSequence/NiagaraEmitterTrackEditor.cpp
+++ b/Engine/Plugins/FX/Niagara/Source/NiagaraEditor/Private/Sequencer/NiagaraSequence/NiagaraEmitterTrackEditor.cpp
@@ -90,7 +90,7 @@ public:
 		// We refresh the renderer widgets when renderers changed.
 		// The OnRenderersChanged delegate will cause the renderer items to be created that are used for widget creation,
 		// Due to delegate bind order, the item creation happens after the refresh. To solve this, we just wait a frame. 
-		UNiagaraEmitter* NiagaraEmitter = EmitterTrack->GetEmitterHandleViewModel()->GetEmitterViewModel()->GetEmitter();
+		UNiagaraEmitter* NiagaraEmitter = EmitterTrack->GetEmitterHandleViewModel()->GetEmitterViewModel()->GetEmitter().Emitter;
 		WeakNiagaraEmitter = NiagaraEmitter;
 		NiagaraEmitter->OnRenderersChanged().AddSP(this, &SEmitterTrackWidget::RefreshRenderers);
 
@@ -297,7 +297,7 @@ bool FNiagaraEmitterTrackEditor::HandleAssetAdded(UObject* Asset, const FGuid& T
 	UNiagaraSequence* NiagaraSequence = Cast<UNiagaraSequence>(GetSequencer()->GetRootMovieSceneSequence());
 	if (EmitterAsset != nullptr && NiagaraSequence != nullptr && NiagaraSequence->GetSystemViewModel().GetCanModifyEmittersFromTimeline())
 	{
-		NiagaraSequence->GetSystemViewModel().AddEmitter(*EmitterAsset);
+		NiagaraSequence->GetSystemViewModel().AddEmitter(*EmitterAsset, EmitterAsset->GetExposedVersion().VersionGuid);
 	}
 	return false;
 }
diff --git a/Engine/Plugins/FX/Niagara/Source/NiagaraEditor/Private/Sequencer/NiagaraSequence/Sections/MovieSceneNiagaraEmitterSection.cpp b/Engine/Plugins/FX/Niagara/Source/NiagaraEditor/Private/Sequencer/NiagaraSequence/Sections/MovieSceneNiagaraEmitterSection.cpp
--- a/Engine/Plugins/FX/Niagara/Source/NiagaraEditor/Private/Sequencer/NiagaraSequence/Sections/MovieSceneNiagaraEmitterSection.cpp
+++ b/Engine/Plugins/FX/Niagara/Source/NiagaraEditor/Private/Sequencer/NiagaraSequence/Sections/MovieSceneNiagaraEmitterSection.cpp
@@ -157,7 +157,7 @@ bool UMovieSceneNiagaraEmitterSection::TryAddTimeRangeModule(UNiagaraNodeFunctio
 	}
 
 	SectionTimingModule = &InModule;
-	UNiagaraScript* EmitterUpdateScript = GetEmitterHandleViewModel()->GetEmitterHandle()->GetInstance()->GetScript(ENiagaraScriptUsage::EmitterUpdateScript, FGuid());
+	UNiagaraScript* EmitterUpdateScript = GetEmitterHandleViewModel()->GetEmitterHandle()->GetEmitterData()->GetScript(ENiagaraScriptUsage::EmitterUpdateScript, FGuid());
 	UNiagaraScript* SystemUpdateScript = GetSystemViewModel().GetSystem().GetSystemUpdateScript();
 	TArray<UNiagaraScript*> BinderDependentScripts;
 	BinderDependentScripts.Add(SystemUpdateScript);
@@ -241,7 +241,7 @@ bool UMovieSceneNiagaraEmitterSection::TryAddKeyModule(UNiagaraNodeFunctionCall&
 
 bool UMovieSceneNiagaraEmitterSection::TrySetupModuleAndBinders(UNiagaraNodeFunctionCall& InModule, FModuleAndBinders& InOutModuleAndBinders, FText& OutErrorMessage)
 {
-	UNiagaraScript* EmitterUpdateScript = GetEmitterHandleViewModel()->GetEmitterHandle()->GetInstance()->GetScript(ENiagaraScriptUsage::EmitterUpdateScript, FGuid());
+	UNiagaraScript* EmitterUpdateScript = GetEmitterHandleViewModel()->GetEmitterHandle()->GetEmitterData()->GetScript(ENiagaraScriptUsage::EmitterUpdateScript, FGuid());
 	UNiagaraScript* SystemUpdateScript = GetSystemViewModel().GetSystem().GetSystemUpdateScript();
 	TArray<UNiagaraScript*> BinderDependentScripts;
 	BinderDependentScripts.Add(SystemUpdateScript);
@@ -458,7 +458,7 @@ void UMovieSceneNiagaraEmitterSection::UpdateKeyModulesFromSection(FChannelAndMo
 	// Create new modules for new keys.
 	if (NewKeysIndices.Num() > 0)
 	{
-		UNiagaraScript* EmitterUpdateScript = GetEmitterHandleViewModel()->GetEmitterHandle()->GetInstance()->GetScript(ENiagaraScriptUsage::EmitterUpdateScript, FGuid());
+		UNiagaraScript* EmitterUpdateScript = GetEmitterHandleViewModel()->GetEmitterHandle()->GetEmitterData()->GetScript(ENiagaraScriptUsage::EmitterUpdateScript, FGuid());
 		UNiagaraScript* SystemUpdateScript = GetSystemViewModel().GetSystem().GetSystemUpdateScript();
 		TArray<UNiagaraScript*> BinderDependentScripts;
 		BinderDependentScripts.Add(SystemUpdateScript);
diff --git a/Engine/Plugins/FX/Niagara/Source/NiagaraEditor/Private/Tests/NiagaraEditorTestUtilities.cpp b/Engine/Plugins/FX/Niagara/Source/NiagaraEditor/Private/Tests/NiagaraEditorTestUtilities.cpp
--- a/Engine/Plugins/FX/Niagara/Source/NiagaraEditor/Private/Tests/NiagaraEditorTestUtilities.cpp
+++ b/Engine/Plugins/FX/Niagara/Source/NiagaraEditor/Private/Tests/NiagaraEditorTestUtilities.cpp
@@ -16,13 +16,15 @@
 UNiagaraEmitter* FNiagaraEditorTestUtilities::CreateEmptyTestEmitter()
 {
 	UNiagaraEmitter* Emitter = NewObject<UNiagaraEmitter>();
+	Emitter->CheckVersionDataAvailable();
 	UNiagaraScriptSource* ScriptSource = NewObject<UNiagaraScriptSource>(Emitter);
 	UNiagaraGraph* Graph = NewObject<UNiagaraGraph>(ScriptSource);
 
 	ScriptSource->NodeGraph = Graph;
-	Emitter->GraphSource = ScriptSource;
-	Emitter->SpawnScriptProps.Script->SetLatestSource(ScriptSource);
-	Emitter->UpdateScriptProps.Script->SetLatestSource(ScriptSource);
+	FVersionedNiagaraEmitterData* Data = Emitter->GetLatestEmitterData();
+	Data->GraphSource = ScriptSource;
+	Data->SpawnScriptProps.Script->SetLatestSource(ScriptSource);
+	Data->UpdateScriptProps.Script->SetLatestSource(ScriptSource);
 
 	FNiagaraStackGraphUtilities::ResetGraphForOutput(*Graph, ENiagaraScriptUsage::EmitterSpawnScript, FGuid());
 	FNiagaraStackGraphUtilities::ResetGraphForOutput(*Graph, ENiagaraScriptUsage::EmitterUpdateScript, FGuid());
@@ -36,7 +38,7 @@ UNiagaraSystem* FNiagaraEditorTestUtilities::CreateTestSystemForEmitter(UNiagara
 {
 	UNiagaraSystem* System = NewObject<UNiagaraSystem>((UObject*)GetTransientPackage(), NAME_None, RF_Transient);
 	UNiagaraSystemFactoryNew::InitializeSystem(System, false);
-	System->AddEmitterHandle(*Emitter, Emitter->GetFName());
+	System->AddEmitterHandle(*Emitter, Emitter->GetFName(), FGuid());
 	return System;
 }
 
diff --git a/Engine/Plugins/FX/Niagara/Source/NiagaraEditor/Private/Toolkits/NiagaraScriptToolkit.cpp b/Engine/Plugins/FX/Niagara/Source/NiagaraEditor/Private/Toolkits/NiagaraScriptToolkit.cpp
--- a/Engine/Plugins/FX/Niagara/Source/NiagaraEditor/Private/Toolkits/NiagaraScriptToolkit.cpp
+++ b/Engine/Plugins/FX/Niagara/Source/NiagaraEditor/Private/Toolkits/NiagaraScriptToolkit.cpp
@@ -78,6 +78,7 @@ FNiagaraScriptToolkit::~FNiagaraScriptToolkit()
 
 	FNiagaraEditorModule& NiagaraEditorModule = FModuleManager::LoadModuleChecked<FNiagaraEditorModule>("NiagaraEditor");
 	NiagaraEditorModule.GetOnScriptToolkitsShouldFocusGraphElement().RemoveAll(this);
+	GEditor->UnregisterForUndo(this);
 }
 
 void FNiagaraScriptToolkit::RegisterTabSpawners(const TSharedRef<class FTabManager>& InTabManager)
@@ -215,7 +216,7 @@ void FNiagaraScriptToolkit::Initialize( const EToolkitMode::Type Mode, const TSh
 	Stats = MessageLogModule.CreateLogListingWidget(StatsListing.ToSharedRef());
 
 	VersionMetadata = NewObject<UNiagaraVersionMetaData>(InputScript, "VersionMetadata", RF_Transient);
-	SAssignNew(VersionsWidget, SNiagaraScriptVersionWidget, EditedNiagaraScript.Script, VersionMetadata)
+	SAssignNew(VersionsWidget, SNiagaraScriptVersionWidget, EditedNiagaraScript.Script, VersionMetadata, InputScript->GetOutermost()->GetName())
 		.OnChangeToVersion(this, &FNiagaraScriptToolkit::SwitchToVersion)
 		.OnVersionDataChanged_Lambda([this]()
 		{
@@ -638,7 +639,7 @@ void FNiagaraScriptToolkit::ExtendToolbar()
                  FOnGetContent::CreateRaw(ScriptToolkit, &FNiagaraScriptToolkit::GenerateVersioningDropdownMenu, ScriptToolkit->GetToolkitCommands()),
                  FText(),
                  LOCTEXT("NiagaraShowVersions_ToolTip", "Select version to edit"),
-                 FSlateIcon(FAppStyle::GetAppStyleSetName(), "LevelEditor.SourceControl"),
+                 FSlateIcon(FAppStyle::GetAppStyleSetName(), "Versions"),
                  true);
 			}
 			ToolbarBuilder.EndSection();
@@ -804,14 +805,15 @@ void FNiagaraScriptToolkit::UpdateModuleStats()
 bool FNiagaraScriptToolkit::MatchesContext(const FTransactionContext& InContext, const TArray<TPair<UObject*, FTransactionObjectEvent>>& TransactionObjects) const
 {
 	const auto* Graph = ScriptViewModel->GetGraphViewModel()->GetGraph();
-	if (Graph)
+	UNiagaraScriptSourceBase* Source = EditedNiagaraScript.GetScriptData() ? EditedNiagaraScript.Script->GetSource(EditedNiagaraScript.Version) : nullptr;
+	if (Graph || Source)
 	{
 		for (const TPair<UObject*, FTransactionObjectEvent>& TransactionObjectPair : TransactionObjects)
 		{
 			UObject* Object = TransactionObjectPair.Key;
 			while (Object != nullptr)
 			{
-				if (Object == Graph)
+				if (Object == Graph || Object == Source)
 				{
 					return true;
 				}
@@ -824,6 +826,17 @@ bool FNiagaraScriptToolkit::MatchesContext(const FTransactionContext& InContext,
 
 void FNiagaraScriptToolkit::PostUndo(bool bSuccess)
 {
+	if (EditedNiagaraScript.Version.IsValid() && EditedNiagaraScript.GetScriptData() == nullptr)
+	{
+		if (EditedNiagaraScript.Script && EditedNiagaraScript.Script->IsVersioningEnabled())
+		{
+			SwitchToVersion(EditedNiagaraScript.Script->GetExposedVersion().VersionGuid);
+		}
+		else
+		{
+			EditedNiagaraScript.Version = FGuid();
+		}
+	}
 	MarkDirtyWithPendingChanges();
 }
 
diff --git a/Engine/Plugins/FX/Niagara/Source/NiagaraEditor/Private/Toolkits/NiagaraSystemToolkit.cpp b/Engine/Plugins/FX/Niagara/Source/NiagaraEditor/Private/Toolkits/NiagaraSystemToolkit.cpp
--- a/Engine/Plugins/FX/Niagara/Source/NiagaraEditor/Private/Toolkits/NiagaraSystemToolkit.cpp
+++ b/Engine/Plugins/FX/Niagara/Source/NiagaraEditor/Private/Toolkits/NiagaraSystemToolkit.cpp
@@ -40,6 +40,7 @@
 #include "Widgets/SBoxPanel.h"
 #include "Widgets/Layout/SBox.h"
 #include "Widgets/Docking/SDockTab.h"
+#include "Widgets/SNiagaraEmitterVersionWidget.h"
 #include "BusyCursor.h"
 #include "Misc/FeedbackContext.h"
 #include "Editor.h"
@@ -50,6 +51,8 @@
 #include "Modules/ModuleManager.h"
 #include "Misc/FileHelper.h"
 #include "NiagaraMessageLogViewModel.h"
+#include "NiagaraVersionMetaData.h"
+#include "Dialogs/Dialogs.h"
 #include "ViewModels/NiagaraParameterPanelViewModel.h"
 #include "NiagaraEditor/Private/SNiagaraAssetPickerList.h"
 #include "SystemToolkitModes/NiagaraSystemToolkitModeBase.h"
@@ -112,6 +115,7 @@ FNiagaraSystemToolkit::~FNiagaraSystemToolkit()
 void FNiagaraSystemToolkit::AddReferencedObjects(FReferenceCollector& Collector) 
 {
 	Collector.AddReferencedObject(System);
+	Collector.AddReferencedObject(VersionMetadata);
 }
 
 TSharedPtr<SWidget> FNiagaraSystemToolkit::GetSystemOverview() const
@@ -176,46 +180,54 @@ void FNiagaraSystemToolkit::InitializeWithSystem(const EToolkitMode::Type Mode,
 	InitializeInternal(Mode, InitToolkitHost, SystemOptions.MessageLogGuid.GetValue());
 }
 
-void FNiagaraSystemToolkit::InitializeWithEmitter(const EToolkitMode::Type Mode, const TSharedPtr< class IToolkitHost >& InitToolkitHost, UNiagaraEmitter& InEmitter)
+void FNiagaraSystemToolkit::InitializeRapidIterationParameters(const FVersionedNiagaraEmitter& VersionedEmitter)
 {
-	System = NewObject<UNiagaraSystem>(GetTransientPackage(), NAME_None, RF_Transient | RF_Transactional);
-	UNiagaraSystemFactoryNew::InitializeSystem(System, true);
-	System->EnsureFullyLoaded();
-
-	Emitter = &InEmitter;
-	Emitter->UpdateEmitterAfterLoad();
-
 	// Before copying the emitter prepare the rapid iteration parameters so that the post compile prepare doesn't
 	// cause the change ids to become out of sync.
 	TArray<UNiagaraScript*> Scripts;
 	TMap<UNiagaraScript*, UNiagaraScript*> ScriptDependencyMap;
-	TMap<UNiagaraScript*, const UNiagaraEmitter*> ScriptToEmitterMap;
+	TMap<UNiagaraScript*, FVersionedNiagaraEmitter> ScriptToEmitterMap;
 
-	Scripts.Add(Emitter->EmitterSpawnScriptProps.Script);
-	ScriptToEmitterMap.Add(Emitter->EmitterSpawnScriptProps.Script, Emitter);
+	FVersionedNiagaraEmitterData* EmitterData = VersionedEmitter.GetEmitterData();
+	Scripts.Add(EmitterData->EmitterSpawnScriptProps.Script);
+	ScriptToEmitterMap.Add(EmitterData->EmitterSpawnScriptProps.Script, VersionedEmitter);
 
-	Scripts.Add(Emitter->EmitterUpdateScriptProps.Script);
-	ScriptToEmitterMap.Add(Emitter->EmitterUpdateScriptProps.Script, Emitter);
+	Scripts.Add(EmitterData->EmitterUpdateScriptProps.Script);
+	ScriptToEmitterMap.Add(EmitterData->EmitterUpdateScriptProps.Script, VersionedEmitter);
 
-	Scripts.Add(Emitter->SpawnScriptProps.Script);
-	ScriptToEmitterMap.Add(Emitter->SpawnScriptProps.Script, Emitter);
+	Scripts.Add(EmitterData->SpawnScriptProps.Script);
+	ScriptToEmitterMap.Add(EmitterData->SpawnScriptProps.Script, VersionedEmitter);
 
-	Scripts.Add(Emitter->UpdateScriptProps.Script);
-	ScriptToEmitterMap.Add(Emitter->UpdateScriptProps.Script, Emitter);
+	Scripts.Add(EmitterData->UpdateScriptProps.Script);
+	ScriptToEmitterMap.Add(EmitterData->UpdateScriptProps.Script, VersionedEmitter);
 
-	if (Emitter->SimTarget == ENiagaraSimTarget::GPUComputeSim)
+	if (EmitterData->SimTarget == ENiagaraSimTarget::GPUComputeSim)
 	{
-		Scripts.Add(Emitter->GetGPUComputeScript());
-		ScriptToEmitterMap.Add(Emitter->GetGPUComputeScript(), Emitter);
-		ScriptDependencyMap.Add(Emitter->SpawnScriptProps.Script, Emitter->GetGPUComputeScript());
-		ScriptDependencyMap.Add(Emitter->UpdateScriptProps.Script, Emitter->GetGPUComputeScript());
+		Scripts.Add(EmitterData->GetGPUComputeScript());
+		ScriptToEmitterMap.Add(EmitterData->GetGPUComputeScript(), VersionedEmitter);
+		ScriptDependencyMap.Add(EmitterData->SpawnScriptProps.Script, EmitterData->GetGPUComputeScript());
+		ScriptDependencyMap.Add(EmitterData->UpdateScriptProps.Script, EmitterData->GetGPUComputeScript());
 	} 
-	else if (Emitter->bInterpolatedSpawning)
+	else if (EmitterData->bInterpolatedSpawning)
 	{
-		ScriptDependencyMap.Add(Emitter->UpdateScriptProps.Script, Emitter->SpawnScriptProps.Script);
+		ScriptDependencyMap.Add(EmitterData->UpdateScriptProps.Script, EmitterData->SpawnScriptProps.Script);
 	}
 
 	FNiagaraUtilities::PrepareRapidIterationParameters(Scripts, ScriptDependencyMap, ScriptToEmitterMap);
+}
+
+void FNiagaraSystemToolkit::InitializeWithEmitter(const EToolkitMode::Type Mode, const TSharedPtr< class IToolkitHost >& InitToolkitHost, UNiagaraEmitter& InEmitter)
+{
+	System = NewObject<UNiagaraSystem>(GetTransientPackage(), NAME_None, RF_Transient | RF_Transactional);
+	UNiagaraSystemFactoryNew::InitializeSystem(System, true);
+	System->EnsureFullyLoaded();
+
+	InEmitter.UpdateEmitterAfterLoad();
+	Emitter = &InEmitter;
+
+	FGuid VersionGuid = Emitter->IsVersioningEnabled() && Emitter->VersionToOpenInEditor.IsValid() ? Emitter->VersionToOpenInEditor : Emitter->GetExposedVersion().VersionGuid;
+	FVersionedNiagaraEmitter VersionedEmitter = FVersionedNiagaraEmitter(Emitter, VersionGuid);
+	InitializeRapidIterationParameters(VersionedEmitter);
 
 	// No need to reset loader or versioning on the transient package, there should never be any set 
 	
@@ -233,7 +245,7 @@ void FNiagaraSystemToolkit::InitializeWithEmitter(const EToolkitMode::Type Mode,
 
 	SystemViewModel->OnGetWorkflowMode().BindSP(this, &FNiagaraSystemToolkit::GetCurrentMode);
 
-	SystemViewModel->AddEmitter(*Emitter);
+	SystemViewModel->AddEmitter(VersionedEmitter);
 
 	ParameterPanelViewModel = MakeShared<FNiagaraSystemToolkitParameterPanelViewModel>(SystemViewModel);
 	ParameterDefinitionsPanelViewModel = MakeShared<FNiagaraSystemToolkitParameterDefinitionsPanelViewModel>(SystemViewModel);
@@ -247,7 +259,7 @@ void FNiagaraSystemToolkit::InitializeWithEmitter(const EToolkitMode::Type Mode,
 	// Adding the emitter to the system has made a copy of it and we set this to the copy's change id here instead of the original emitter's change 
 	// id because the copy's change id may have been updated from the original as part of post load and we use this id to detect if the editable 
 	// emitter has been changed.
-	LastSyncedEmitterChangeId = SystemViewModel->GetEmitterHandleViewModels()[0]->GetEmitterViewModel()->GetEmitter()->GetChangeId();
+	LastSyncedEmitterChangeId = SystemViewModel->GetEmitterHandleViewModels()[0]->GetEmitterViewModel()->GetEmitter().Emitter->GetChangeId();
 	SystemToolkitMode = ESystemToolkitMode::Emitter;
 
 	if (GbLogNiagaraSystemChanges > 0)
@@ -281,11 +293,8 @@ void FNiagaraSystemToolkit::InitializeInternal(const EToolkitMode::Type Mode, co
 	
 	SystemViewModel->OnGetWorkflowMode().BindSP(this, &FNiagaraSystemToolkit::GetCurrentMode);
 	
-	const float InTime = -0.02f;
-	const float OutTime = 3.2f;
-	
-	const bool bCreateDefaultStandaloneMenu = true;
-	const bool bCreateDefaultToolbar = true;
+	constexpr bool bCreateDefaultStandaloneMenu = true;
+	constexpr bool bCreateDefaultToolbar = true;
 	UObject* ToolkitObject = SystemToolkitMode == ESystemToolkitMode::System ? (UObject*)System : (UObject*)Emitter;
 	// order of registering commands matters. SetupCommands before InitAssetEditor will make the toolkit prioritize niagara commands
 	SetupCommands();
@@ -296,9 +305,23 @@ void FNiagaraSystemToolkit::InitializeInternal(const EToolkitMode::Type Mode, co
 	
 	AddApplicationMode(DefaultModeName, MakeShared<FNiagaraSystemToolkitMode_Default>(SharedThis(this)));
 	AddApplicationMode(ScalabilityModeName, MakeShared<FNiagaraSystemToolkitMode_Scalability>(SharedThis(this)));
+
+	// set up the versioning widget
+	VersionMetadata = NewObject<UNiagaraVersionMetaData>(ToolkitObject, "VersionMetadata", RF_Transient);
+	SAssignNew(VersionsWidget, SNiagaraEmitterVersionWidget, HasEmitter() ? GetEditedEmitterViewModel()->GetEmitter().Emitter : nullptr, VersionMetadata, HasEmitter() ? Emitter->GetOutermost()->GetName() : TEXT(""))
+		.OnChangeToVersion(this, &FNiagaraSystemToolkit::SwitchToVersion)
+		.OnVersionDataChanged_Lambda([this]()
+		{
+			if (TSharedPtr<FNiagaraEmitterViewModel> EditableEmitterViewModel = SystemViewModel->GetEmitterHandleViewModels()[0]->GetEmitterViewModel())
+			{
+				FVersionedNiagaraEmitter EditableEmitter = EditableEmitterViewModel->GetEmitter();
+				FProperty* VersionProperty = FindFProperty<FProperty>(UNiagaraEmitter::StaticClass(), FName("VersionData"));
+				FPropertyChangedEvent ChangeEvent(VersionProperty);
+				EditableEmitter.Emitter->PostEditChangeProperty(ChangeEvent);
+			}
+		});
 	
 	SetCurrentMode(DefaultModeName);
-	//SystemViewModel->GetScalabilityViewModel()->Activate();
 	
 	FNiagaraEditorModule& NiagaraEditorModule = FModuleManager::LoadModuleChecked<FNiagaraEditorModule>("NiagaraEditor");
 	AddMenuExtender(NiagaraEditorModule.GetMenuExtensibilityManager()->GetAllExtenders(GetToolkitCommands(), GetEditingObjects()));
@@ -380,9 +403,6 @@ void FNiagaraSystemToolkit::PostUndo(bool bSuccess)
 
 void FNiagaraSystemToolkit::SetupCommands()
 {
-	FLevelEditorModule& Module = FModuleManager::Get().LoadModuleChecked<FLevelEditorModule>(TEXT("LevelEditor"));
-	const FLevelEditorCommands& LevelEditorCommands = Module.GetLevelEditorCommands();
-	
 	GetToolkitCommands()->MapAction(
 		FNiagaraEditorCommands::Get().Compile,
 		FExecuteAction::CreateRaw(this, &FNiagaraSystemToolkit::CompileSystem, false));
@@ -524,17 +544,31 @@ void FNiagaraSystemToolkit::SetupCommands()
 	GetToolkitCommands()->MapAction(
 		FNiagaraEditorCommands::Get().OpenAttributeSpreadsheet,
 		FExecuteAction::CreateSP(this, &FNiagaraSystemToolkit::OpenAttributeSpreadsheet));
+
+	GetToolkitCommands()->MapAction(
+		FNiagaraEditorCommands::Get().EmitterVersioning,
+		FExecuteAction::CreateSP(this, &FNiagaraSystemToolkit::ManageVersions));
 	
 	// appending the sequencer commands will make the toolkit also check for sequencer commands (last)
 	GetToolkitCommands()->Append(SystemViewModel->GetSequencer()->GetCommandBindings(ESequencerCommandBindings::Sequencer).ToSharedRef());
 	SystemViewModel->GetSequencer()->GetCommandBindings(ESequencerCommandBindings::Sequencer)->Append(GetToolkitCommands());
 }
 
+void FNiagaraSystemToolkit::ManageVersions()
+{
+	TabManager->TryInvokeTab(FNiagaraSystemToolkitModeBase::VersioningTabID);
+}
+
+TSharedPtr<FNiagaraEmitterViewModel> FNiagaraSystemToolkit::GetEditedEmitterViewModel() const
+{
+	return HasEmitter() ? SystemViewModel->GetEmitterHandleViewModels()[0]->GetEmitterViewModel() : TSharedPtr<FNiagaraEmitterViewModel>(); 
+}
+
 void FNiagaraSystemToolkit::OnSaveThumbnailImage()
 {
 	if (Viewport.IsValid())
 	{
-		Viewport->CreateThumbnail(SystemToolkitMode == ESystemToolkitMode::System ? (UObject*)System : Emitter);
+		Viewport->CreateThumbnail(SystemToolkitMode == ESystemToolkitMode::System ? static_cast<UObject*>(System) : Emitter);
 	}
 }
 
@@ -551,7 +585,7 @@ void FNiagaraSystemToolkit::OnThumbnailCaptured(UTexture2D* Thumbnail)
 	else if (SystemToolkitMode == ESystemToolkitMode::Emitter) 
 	{
 		TSharedPtr<FNiagaraEmitterViewModel> EditableEmitterViewModel = SystemViewModel->GetEmitterHandleViewModels()[0]->GetEmitterViewModel();
-		UNiagaraEmitter* EditableEmitter = EditableEmitterViewModel->GetEmitter();
+		UNiagaraEmitter* EditableEmitter = EditableEmitterViewModel->GetEmitter().Emitter;
 		EditableEmitter->ThumbnailImage = Thumbnail;
 		bEmitterThumbnailUpdated = true;
 		// Broadcast an object property changed event to update the content browser
@@ -648,6 +682,81 @@ TSharedRef<SWidget> FNiagaraSystemToolkit::CreateAddEmitterMenuContent()
 		];
 }
 
+FText FNiagaraSystemToolkit::GetVersionButtonLabel() const
+{
+	FText BaseLabel = LOCTEXT("NiagaraShowEmitterVersions", "Versioning");
+	TSharedPtr<FNiagaraEmitterViewModel> EmitterViewModel = GetEditedEmitterViewModel();
+	if (EmitterViewModel && EmitterViewModel->GetEmitter().Emitter->IsVersioningEnabled())
+	{
+		if (FVersionedNiagaraEmitterData* EmitterData = EmitterViewModel->GetEmitter().GetEmitterData())
+		{
+			FNiagaraAssetVersion ExposedVersion = EmitterViewModel->GetEmitter().Emitter->GetExposedVersion();
+			return FText::Format(FText::FromString("{0} ({1}.{2}{3})"), BaseLabel, EmitterData->Version.MajorVersion, EmitterData->Version.MinorVersion, EmitterData->Version <= ExposedVersion ? FText::FromString("*") : FText());
+		}
+	}
+	return BaseLabel;
+}
+
+TArray<FNiagaraAssetVersion> FNiagaraSystemToolkit::GetEmitterVersions() const
+{
+	TArray<FNiagaraAssetVersion> Result;
+	if (TSharedPtr<FNiagaraEmitterViewModel> EmitterViewModel = GetEditedEmitterViewModel())
+	{
+		Result.Append(EmitterViewModel->GetEmitter().Emitter->GetAllAvailableVersions());
+	}
+	return Result;
+}
+
+TSharedRef<SWidget> FNiagaraSystemToolkit::GenerateVersioningDropdownMenu(TSharedRef<FUICommandList> InCommandList)
+{
+	constexpr bool bShouldCloseWindowAfterMenuSelection = true;
+	FMenuBuilder MenuBuilder(bShouldCloseWindowAfterMenuSelection, InCommandList);
+
+	TArray<FNiagaraAssetVersion> AssetVersions = GetEmitterVersions();
+	for (FNiagaraAssetVersion& Version : AssetVersions)
+	{
+		FText Tooltip = LOCTEXT("NiagaraSelectVersion", "Select this emitter version to edit");
+		FUIAction UIAction(FExecuteAction::CreateSP(this, &FNiagaraSystemToolkit::SwitchToVersion, Version.VersionGuid),
+		FCanExecuteAction(),
+		FIsActionChecked::CreateSP(this, &FNiagaraSystemToolkit::IsVersionSelected, Version));
+		TAttribute<FText> Label = TAttribute<FText>::Create(TAttribute<FText>::FGetter::CreateSP(this, &FNiagaraSystemToolkit::GetVersionMenuLabel, Version));
+		MenuBuilder.AddMenuEntry(Label, Tooltip, FSlateIcon(), UIAction, NAME_None, EUserInterfaceActionType::RadioButton);	
+	}
+
+	return MenuBuilder.MakeWidget();
+}
+
+void FNiagaraSystemToolkit::SwitchToVersion(FGuid VersionGuid)
+{
+	if (TSharedPtr<FNiagaraEmitterViewModel> EmitterViewModel = GetEditedEmitterViewModel())
+	{
+		if (EmitterViewModel->GetEmitter().Version == VersionGuid)
+		{
+			return;
+		}
+		
+		FScopedTransaction Transaction(LOCTEXT("ChangeEmitterVersion", "Switch to emitter version"));
+		InitializeRapidIterationParameters(FVersionedNiagaraEmitter(EmitterViewModel->GetEmitter().Emitter, VersionGuid));
+		if (GetSystemViewModel()->ChangeEmitterVersion(EmitterViewModel->GetEmitter(), VersionGuid))
+		{
+			GetSystemViewModel()->GetSelectionViewModel()->EmptySelection();
+			GetSystemViewModel()->GetSelectionViewModel()->AddEntryToSelectionByDisplayedObjectDeferred(EmitterViewModel->GetEmitter().Emitter);
+		}
+	}
+}
+
+bool FNiagaraSystemToolkit::IsVersionSelected(FNiagaraAssetVersion Version) const
+{
+	return HasEmitter() && GetEditedEmitterViewModel()->GetEmitter().Version == Version.VersionGuid;
+}
+
+FText FNiagaraSystemToolkit::GetVersionMenuLabel(FNiagaraAssetVersion Version) const
+{
+	TSharedPtr<FNiagaraEmitterViewModel> EmitterViewModel = GetEditedEmitterViewModel();
+	bool bIsExposed = EmitterViewModel && Version == EmitterViewModel->GetEmitter().Emitter->GetExposedVersion();
+	return FText::Format(FText::FromString("v{0}.{1} {2}"), Version.MajorVersion, Version.MinorVersion, bIsExposed ? LOCTEXT("NiagaraExposedVersionHint", "(exposed)") : FText());
+}
+
 TSharedRef<SWidget> FNiagaraSystemToolkit::GenerateCompileMenuContent()
 {
 	const bool bShouldCloseWindowAfterMenuSelection = true;
@@ -890,43 +999,59 @@ void FNiagaraSystemToolkit::UpdateOriginalEmitter()
 	checkf(SystemToolkitMode == ESystemToolkitMode::Emitter, TEXT("There is no original emitter to update in system mode."));
 
 	TSharedPtr<FNiagaraEmitterViewModel> EditableEmitterViewModel = SystemViewModel->GetEmitterHandleViewModels()[0]->GetEmitterViewModel();
-	UNiagaraEmitter* EditableEmitter = EditableEmitterViewModel->GetEmitter();
-
-	if (EditableEmitter->GetChangeId() != LastSyncedEmitterChangeId)
+	FVersionedNiagaraEmitter EditableEmitter = EditableEmitterViewModel->GetEmitter();
+	FVersionedNiagaraEmitterData* EditableEmitterData = EditableEmitter.GetEmitterData();
+	UNiagaraEmitter* Source = Emitter;
+	if (EditableEmitter.Emitter->GetChangeId() != LastSyncedEmitterChangeId)
 	{
+		if (EditableEmitter.Emitter->IsVersioningEnabled() && EditableEmitterData->Version <= EditableEmitter.Emitter->GetExposedVersion())
+		{
+			FSuppressableWarningDialog::FSetupInfo Info( 
+				LOCTEXT("ApplyExposedVersionChangesPrompt", "You are about to apply changes to an already exposed asset version. Saving these changes will force-push them out to existing usages!\nConsider creating a new version instead to make those changes."), 
+				LOCTEXT("ApplyExposedVersionChangesTitle", "Warning: editing exposed emitter version"), 
+				TEXT("ApplyExposedEmitterChanges"));
+			Info.ConfirmText = LOCTEXT("ApplyExposedVersionChanges_ConfirmText", "Apply Changes");
+			Info.CancelText = LOCTEXT("ApplyExposedVersionChanges_CancelText", "Cancel");
+			Info.CheckBoxText = LOCTEXT("ApplyExposedVersionChanges_CheckBoxText", "Don't Ask Again");
+
+			if (FSuppressableWarningDialog(Info).ShowModal() == FSuppressableWarningDialog::EResult::Cancel)
+			{
+				return;
+			}
+		}
+		
 		const FScopedBusyCursor BusyCursor;
 		const FText LocalizedScriptEditorApply = NSLOCTEXT("UnrealEd", "ToolTip_NiagaraEmitterEditorApply", "Apply changes to original emitter and its use in the world.");
 		GWarn->BeginSlowTask(LocalizedScriptEditorApply, true);
 		GWarn->StatusUpdate(1, 1, LocalizedScriptEditorApply);
 
-		if (Emitter->IsSelected())
+		if (Source->IsSelected())
 		{
-			GEditor->GetSelectedObjects()->Deselect(Emitter);
+			GEditor->GetSelectedObjects()->Deselect(Source);
 		}
 
-		ResetLoaders(Emitter->GetOutermost()); // Make sure that we're not going to get invalid version number linkers into the package we are going into. 
+		ResetLoaders(Source->GetOutermost()); // Make sure that we're not going to get invalid version number linkers into the package we are going into. 
 
 		TArray<UNiagaraScript*> AllScripts;
-		EditableEmitter->GetScripts(AllScripts, true);
+		EditableEmitterData->GetScripts(AllScripts, true);
 		for (UNiagaraScript* Script : AllScripts)
 		{
 			checkfSlow(Script->AreScriptAndSourceSynchronized(), TEXT("Editable Emitter Script change ID is out of date when applying to Original Emitter!"));
 		}
-		Emitter->PreEditChange(nullptr);
+		Source->PreEditChange(nullptr);
 		// overwrite the original script in place by constructing a new one with the same name
-		Emitter = (UNiagaraEmitter*)StaticDuplicateObject(EditableEmitter, Emitter->GetOuter(),
-			Emitter->GetFName(), RF_AllFlags, Emitter->GetClass());
+		Source = Cast<UNiagaraEmitter>(StaticDuplicateObject(EditableEmitter.Emitter, Source->GetOuter(), Source->GetFName(), RF_AllFlags, Source->GetClass()));
 
 		// Restore RF_Standalone and RF_Public on the original emitter, as it had been removed from the preview emitter so that it could be GC'd.
-		Emitter->SetFlags(RF_Standalone | RF_Public);
+		Source->SetFlags(RF_Standalone | RF_Public);
 
-		Emitter->PostEditChange();
+		Source->PostEditChange();
 
 		TArray<UNiagaraScript*> EmitterScripts;
-		Emitter->GetScripts(EmitterScripts, false);
+		Source->GetEmitterData(EditableEmitter.Version)->GetScripts(EmitterScripts, false);
 
 		TArray<UNiagaraScript*> EditableEmitterScripts;
-		EditableEmitter->GetScripts(EditableEmitterScripts, false);
+		EditableEmitterData->GetScripts(EditableEmitterScripts, false);
 
 		// Validate that the change ids on the original emitters match the editable emitters ids to ensure the DDC contents are up to data without having to recompile.
 		if (ensureMsgf(EmitterScripts.Num() == EditableEmitterScripts.Num(), TEXT("Script count mismatch after copying from editable emitter to original emitter.")))
@@ -943,7 +1068,7 @@ void FNiagaraSystemToolkit::UpdateOriginalEmitter()
 		}
 
 		// Record the last synced change id to detect future changes.
-		LastSyncedEmitterChangeId = EditableEmitter->GetChangeId();
+		LastSyncedEmitterChangeId = EditableEmitter.Emitter->GetChangeId();
 		bEmitterThumbnailUpdated = false;
 
 		UpdateExistingEmitters();
@@ -951,14 +1076,14 @@ void FNiagaraSystemToolkit::UpdateOriginalEmitter()
 	}
 	else if(bEmitterThumbnailUpdated)
 	{
-		Emitter->MarkPackageDirty();
-		Emitter->ThumbnailImage = (UTexture2D*)StaticDuplicateObject(EditableEmitter->ThumbnailImage, Emitter);
-		Emitter->PostEditChange();
+		Source->MarkPackageDirty();
+		Source->ThumbnailImage = (UTexture2D*)StaticDuplicateObject(EditableEmitter.Emitter->ThumbnailImage, Source);
+		Source->PostEditChange();
 		bEmitterThumbnailUpdated = false;
 	}
 }
 
-void MergeEmittersRecursively(UNiagaraEmitter* ChangedEmitter, const TMap<UNiagaraEmitter*, TArray<UNiagaraEmitter*>>& EmitterToReferencingEmittersMap, TSet<UNiagaraEmitter*>& OutMergedEmitters)
+void MergeEmittersRecursively(const UNiagaraEmitter* ChangedEmitter, const TMap<UNiagaraEmitter*, TArray<UNiagaraEmitter*>>& EmitterToReferencingEmittersMap, TSet<UNiagaraEmitter*>& OutMergedEmitters)
 {
 	const TArray<UNiagaraEmitter*>* ReferencingEmitters = EmitterToReferencingEmittersMap.Find(ChangedEmitter);
 	if (ReferencingEmitters != nullptr)
@@ -979,14 +1104,21 @@ void FNiagaraSystemToolkit::UpdateExistingEmitters()
 {
 	// Build a tree of references from the currently loaded emitters so that we can efficiently find all emitters that reference the modified emitter.
 	TMap<UNiagaraEmitter*, TArray<UNiagaraEmitter*>> EmitterToReferencingEmittersMap;
-	UNiagaraEmitter* EditableCopy = System->GetEmitterHandles()[0].GetInstance();
+	UNiagaraEmitter* EditableCopy = System->GetEmitterHandles()[0].GetInstance().Emitter;
 	for (TObjectIterator<UNiagaraEmitter> EmitterIterator; EmitterIterator; ++EmitterIterator)
 	{
 		UNiagaraEmitter* LoadedEmitter = *EmitterIterator;
-		if (LoadedEmitter != EditableCopy && LoadedEmitter->GetParent() != nullptr)
+		if (LoadedEmitter != EditableCopy)
 		{
-			TArray<UNiagaraEmitter*>& ReferencingEmitters = EmitterToReferencingEmittersMap.FindOrAdd(LoadedEmitter->GetParent());
-			ReferencingEmitters.Add(LoadedEmitter);
+			for (FNiagaraAssetVersion Version : LoadedEmitter->GetAllAvailableVersions())
+			{
+				FVersionedNiagaraEmitterData* EmitterData = LoadedEmitter->GetEmitterData(Version.VersionGuid);
+				if (UNiagaraEmitter* ParentEmitter = EmitterData->GetParent().Emitter)
+				{
+					TArray<UNiagaraEmitter*>& ReferencingEmitters = EmitterToReferencingEmittersMap.FindOrAdd(ParentEmitter);
+					ReferencingEmitters.AddUnique(LoadedEmitter);
+				}
+			}
 		}
 	}
 
@@ -1006,7 +1138,7 @@ void FNiagaraSystemToolkit::UpdateExistingEmitters()
 			bool bUsesMergedEmitterDirectly = false;
 			for (const FNiagaraEmitterHandle& EmitterHandle : LoadedSystem->GetEmitterHandles())
 			{
-				if (MergedEmitters.Contains(EmitterHandle.GetInstance()))
+				if (MergedEmitters.Contains(EmitterHandle.GetInstance().Emitter))
 				{
 					bUsesMergedEmitterDirectly = true;
 					break;
@@ -1133,7 +1265,7 @@ bool FNiagaraSystemToolkit::OnRequestClose()
 	if (SystemToolkitMode == ESystemToolkitMode::Emitter)
 	{
 		TSharedPtr<FNiagaraEmitterViewModel> EmitterViewModel = SystemViewModel->GetEmitterHandleViewModels()[0]->GetEmitterViewModel();
-		if (bChangesDiscarded == false && (EmitterViewModel->GetEmitter()->GetChangeId() != LastSyncedEmitterChangeId || bEmitterThumbnailUpdated))
+		if (bChangesDiscarded == false && (EmitterViewModel->GetEmitter().Emitter->GetChangeId() != LastSyncedEmitterChangeId || bEmitterThumbnailUpdated))
 		{
 			// find out the user wants to do with this dirty emitter.
 			EAppReturnType::Type YesNoCancelReply = FMessageDialog::Open(EAppMsgType::YesNoCancel,
@@ -1194,7 +1326,7 @@ bool FNiagaraSystemToolkit::OnApplyEnabled() const
 	if (Emitter != nullptr)
 	{
 		TSharedPtr<FNiagaraEmitterViewModel> EmitterViewModel = SystemViewModel->GetEmitterHandleViewModels()[0]->GetEmitterViewModel();
-		return EmitterViewModel->GetEmitter()->GetChangeId() != LastSyncedEmitterChangeId || bEmitterThumbnailUpdated;
+		return EmitterViewModel->GetEmitter().Emitter->GetChangeId() != LastSyncedEmitterChangeId || bEmitterThumbnailUpdated;
 	}
 	return false;
 }
@@ -1239,9 +1371,9 @@ void FNiagaraSystemToolkit::RefreshParameters()
 
 	for (TSharedRef<FNiagaraEmitterHandleViewModel> EmitterHandleToDisplay : EmitterHandlesToDisplay)
 	{
-		if (EmitterHandleToDisplay->IsValid() && EmitterHandleToDisplay->GetEmitterViewModel()->GetEmitter() != nullptr)
+		if (EmitterHandleToDisplay->IsValid() && EmitterHandleToDisplay->GetEmitterViewModel()->GetEmitter().Emitter != nullptr)
 		{
-			NewParameterViewSelection.Add(EmitterHandleToDisplay->GetEmitterViewModel()->GetEmitter());
+			NewParameterViewSelection.Add(EmitterHandleToDisplay->GetEmitterViewModel()->GetEmitter().Emitter);
 		}
 	}
 
diff --git a/Engine/Plugins/FX/Niagara/Source/NiagaraEditor/Private/Toolkits/NiagaraSystemToolkit.h b/Engine/Plugins/FX/Niagara/Source/NiagaraEditor/Private/Toolkits/NiagaraSystemToolkit.h
--- a/Engine/Plugins/FX/Niagara/Source/NiagaraEditor/Private/Toolkits/NiagaraSystemToolkit.h
+++ b/Engine/Plugins/FX/Niagara/Source/NiagaraEditor/Private/Toolkits/NiagaraSystemToolkit.h
@@ -5,19 +5,19 @@
 #include "CoreMinimal.h"
 #include "UObject/GCObject.h"
 #include "Toolkits/IToolkitHost.h"
-#include "Toolkits/AssetEditorToolkit.h"
 #include "EditorUndoClient.h"
 #include "NiagaraEditorModule.h"
 
 #include "ISequencer.h"
 #include "ISequencerTrackEditor.h"
 
-#include "NiagaraCommon.h"
-#include "NiagaraScript.h"
+#include "ViewModels/NiagaraEmitterViewModel.h"
 #include "Widgets/SItemSelector.h"
 #include "ViewModels/NiagaraSystemGraphSelectionViewModel.h"
 #include "WorkflowOrientedApp/WorkflowCentricApplication.h"
 
+class SNiagaraEmitterVersionWidget;
+class UNiagaraVersionMetaData;
 #if WITH_NIAGARA_GPU_PROFILER
 	#include "NiagaraGPUProfilerInterface.h"
 #endif
@@ -68,7 +68,7 @@ public:
 	void InitializeWithEmitter(const EToolkitMode::Type Mode, const TSharedPtr< class IToolkitHost >& InitToolkitHost, UNiagaraEmitter& InEmitter);
 
 	/** Destructor */
-	virtual ~FNiagaraSystemToolkit();
+	virtual ~FNiagaraSystemToolkit() override;
 	
 	//~ Begin IToolkit Interface
 	virtual FName GetToolkitFName() const override;
@@ -105,7 +105,14 @@ public:
 	void SetSystemOverview(const TSharedPtr<SWidget>& SystemOverview);
 	TSharedPtr<SWidget> GetScriptScratchpad() const;
 	void SetScriptScratchpad(const TSharedPtr<SWidget>& ScriptScratchpad);
-	
+	TSharedPtr<SWidget> GetVersioningWidget() const { return VersionsWidget; }
+	FText GetVersionButtonLabel() const;
+	TArray<FNiagaraAssetVersion> GetEmitterVersions() const;
+	TSharedRef<SWidget> GenerateVersioningDropdownMenu(TSharedRef<FUICommandList> InCommandList);
+	void SwitchToVersion(FGuid VersionGuid);
+	bool IsVersionSelected(FNiagaraAssetVersion Version) const;
+	FText GetVersionMenuLabel(FNiagaraAssetVersion Version) const;
+
 	/** Mode exposed functions */
 	TSharedRef<SWidget> GenerateCompileMenuContent();
 	TSharedRef<SWidget> GenerateBoundsMenuContent(TSharedRef<FUICommandList> InCommandList);
@@ -154,6 +161,7 @@ protected:
 	
 private:
 	void InitializeInternal(const EToolkitMode::Type Mode, const TSharedPtr<class IToolkitHost>& InitToolkitHost, const FGuid& MessageLogGuid);
+	void InitializeRapidIterationParameters(const FVersionedNiagaraEmitter& VersionedEmitter);
 
 	void UpdateOriginalEmitter();
 	void UpdateExistingEmitters();
@@ -185,13 +193,19 @@ private:
 	void OnSaveThumbnailImage();
 	void OnThumbnailCaptured(UTexture2D* Thumbnail);
 
+	void ManageVersions();
+	TSharedPtr<FNiagaraEmitterViewModel> GetEditedEmitterViewModel() const;
+
 private:
 
 	/** The System being edited in system mode, or the placeholder system being edited in emitter mode. */
-	UNiagaraSystem* System;
+	UNiagaraSystem* System = nullptr;
 
-	/** The emitter being edited in emitter mode, or null when editing in system mode. */
-	UNiagaraEmitter* Emitter;
+	/** The emitter being edited in emitter mode, or null when editing in system mode. Note that we are NOT using the FVersionedEmitter here, because this is the full asset. The version being edited is defined by the emitter view model. */
+	UNiagaraEmitter* Emitter = nullptr;
+
+	/** Temp object to hold version data being edited in the versioning widget's property editor. */
+	UNiagaraVersionMetaData* VersionMetadata = nullptr;
 
 	/** The value of the emitter change id from the last time it was in sync with the original emitter. */
 	FGuid LastSyncedEmitterChangeId;
@@ -236,11 +250,11 @@ private:
 	TSharedPtr<FNiagaraObjectSelection> ObjectSelectionForParameterMapView;
 
 	bool bChangesDiscarded;
-
 	bool bScratchPadChangesDiscarded;
 
 	TSharedPtr<SWidget> SystemOverview;
 	TSharedPtr<SWidget> ScriptScratchpad;
+	TSharedPtr<SWidget> VersionsWidget;
 	
 	static IConsoleVariable* VmStatEnabledVar;
 
diff --git a/Engine/Plugins/FX/Niagara/Source/NiagaraEditor/Private/Toolkits/SystemToolkitModes/NiagaraSystemToolkitModeBase.cpp b/Engine/Plugins/FX/Niagara/Source/NiagaraEditor/Private/Toolkits/SystemToolkitModes/NiagaraSystemToolkitModeBase.cpp
--- a/Engine/Plugins/FX/Niagara/Source/NiagaraEditor/Private/Toolkits/SystemToolkitModes/NiagaraSystemToolkitModeBase.cpp
+++ b/Engine/Plugins/FX/Niagara/Source/NiagaraEditor/Private/Toolkits/SystemToolkitModes/NiagaraSystemToolkitModeBase.cpp
@@ -45,6 +45,7 @@ const FName FNiagaraSystemToolkitModeBase::SystemOverviewTabID(TEXT("NiagaraSyst
 const FName FNiagaraSystemToolkitModeBase::ScratchPadTabID(TEXT("NiagaraSystemEditor_ScratchPad"));
 const FName FNiagaraSystemToolkitModeBase::ScriptStatsTabID(TEXT("NiagaraSystemEditor_ScriptStats"));
 const FName FNiagaraSystemToolkitModeBase::BakerTabID(TEXT("NiagaraSystemEditor_Baker"));
+const FName FNiagaraSystemToolkitModeBase::VersioningTabID(TEXT("NiagaraSystemEditor_Versioning"));
 
 void FNiagaraSystemToolkitModeBase::RegisterTabFactories(TSharedPtr<FTabManager> InTabManager)
 {
@@ -131,6 +132,11 @@ void FNiagaraSystemToolkitModeBase::RegisterTabFactories(TSharedPtr<FTabManager>
 		.SetDisplayName(LOCTEXT("NiagaraBakerTab", "Baker"))
 		.SetGroup(WorkspaceMenuCategory.ToSharedRef())
 		.SetIcon(FSlateIcon(FNiagaraEditorStyle::Get().GetStyleSetName(), "NiagaraEditor.BakerIcon"));
+
+	InTabManager->RegisterTabSpawner(VersioningTabID, FOnSpawnTab::CreateSP(this, &FNiagaraSystemToolkitModeBase::SpawnTab_Versioning))
+		.SetDisplayName(LOCTEXT("VersioningTab", "Versioning"))
+		.SetGroup(WorkspaceMenuCategory.ToSharedRef())
+		.SetIcon(FSlateIcon(FAppStyle::GetAppStyleSetName(), "Versions"));
 }
 
 class SNiagaraSelectedEmitterGraph : public SCompoundWidget
@@ -151,7 +157,7 @@ public:
 		UpdateGraphWidget();
 	}
 
-	~SNiagaraSelectedEmitterGraph()
+	virtual ~SNiagaraSelectedEmitterGraph() override
 	{
 		if (SystemViewModel.IsValid() && SystemViewModel->GetSelectionViewModel())
 		{
@@ -184,7 +190,7 @@ private:
 					SNew(SNiagaraScriptGraph, SelectedEmitterHandle->GetEmitterViewModel()->GetSharedScriptViewModel()->GetGraphViewModel())
 				];
 
-			UNiagaraEmitter* LastMergedEmitter = SelectedEmitterHandle->GetEmitterViewModel()->GetEmitter()->GetParentAtLastMerge();
+			FVersionedNiagaraEmitterData* LastMergedEmitter = SelectedEmitterHandle->GetEmitterViewModel()->GetEmitter().GetEmitterData()->GetParentAtLastMerge().GetEmitterData();
 			if (LastMergedEmitter != nullptr)
 			{
 				UNiagaraScriptSource* LastMergedScriptSource = CastChecked<UNiagaraScriptSource>(LastMergedEmitter->GraphSource);
@@ -525,4 +531,21 @@ TSharedRef<SDockTab> FNiagaraSystemToolkitModeBase::SpawnTab_Baker(const FSpawnT
 	return SpawnedTab;
 }
 
+TSharedRef<SDockTab> FNiagaraSystemToolkitModeBase::SpawnTab_Versioning(const FSpawnTabArgs& Args)
+{
+	check(Args.GetTabId().TabType == VersioningTabID);
+
+	TSharedRef<SDockTab> SpawnedTab = SNew(SDockTab)
+		.Label(LOCTEXT("EmitterVersioningTitle", "Versioning"))
+		[
+			SNew(SBox)
+			.AddMetaData<FTagMetaData>(FTagMetaData(TEXT("EmitterVersioning")))
+			[
+				SystemToolkit.Pin()->GetVersioningWidget().ToSharedRef()
+			]
+		];
+
+	return SpawnedTab;
+}
+
 #undef LOCTEXT_NAMESPACE
\ No newline at end of file
diff --git a/Engine/Plugins/FX/Niagara/Source/NiagaraEditor/Private/Toolkits/SystemToolkitModes/NiagaraSystemToolkitMode_Default.cpp b/Engine/Plugins/FX/Niagara/Source/NiagaraEditor/Private/Toolkits/SystemToolkitModes/NiagaraSystemToolkitMode_Default.cpp
--- a/Engine/Plugins/FX/Niagara/Source/NiagaraEditor/Private/Toolkits/SystemToolkitModes/NiagaraSystemToolkitMode_Default.cpp
+++ b/Engine/Plugins/FX/Niagara/Source/NiagaraEditor/Private/Toolkits/SystemToolkitModes/NiagaraSystemToolkitMode_Default.cpp
@@ -7,6 +7,8 @@
 #include "NiagaraDebuggerCommon.h"
 #include "NiagaraEditorModule.h"
 #include "NiagaraEditorSettings.h"
+#include "ViewModels/NiagaraEmitterHandleViewModel.h"
+#include "ViewModels/NiagaraSystemViewModel.h"
 
 #include "Framework/MultiBox/MultiBoxBuilder.h"
 
@@ -264,6 +266,28 @@ void FNiagaraSystemToolkitMode_Default::ExtendToolbar()
 				}
 				ToolbarBuilder.EndSection();
 			}
+
+			if (Toolkit->HasEmitter())
+			{
+				ToolbarBuilder.BeginSection("Versioning");
+				{
+					ToolbarBuilder.AddToolBarButton(FNiagaraEditorCommands::Get().EmitterVersioning, NAME_None,
+					TAttribute<FText>(Toolkit, &FNiagaraSystemToolkit::GetVersionButtonLabel),
+					LOCTEXT("NiagaraShowModuleVersionsTooltip", "Manage different versions of this emitter."),
+					FSlateIcon(FAppStyle::GetAppStyleSetName(), "Versions"));
+
+					FUIAction DropdownAction;
+					DropdownAction.IsActionVisibleDelegate = FIsActionButtonVisible::CreateLambda([Toolkit]() { return Toolkit->GetEmitterVersions().Num() > 1; });
+					ToolbarBuilder.AddComboButton(
+					 DropdownAction,
+					 FOnGetContent::CreateRaw(Toolkit, &FNiagaraSystemToolkit::GenerateVersioningDropdownMenu, Toolkit->GetToolkitCommands()),
+					 FText(),
+					 LOCTEXT("NiagaraShowVersions_ToolTip", "Select version to edit"),
+					 FSlateIcon(FAppStyle::GetAppStyleSetName(), "Versions"),
+					 true);
+				}
+				ToolbarBuilder.EndSection();
+			}
 		}
 	};
 
diff --git a/Engine/Plugins/FX/Niagara/Source/NiagaraEditor/Private/Toolkits/SystemToolkitModes/NiagaraSystemToolkitMode_Scalability.cpp b/Engine/Plugins/FX/Niagara/Source/NiagaraEditor/Private/Toolkits/SystemToolkitModes/NiagaraSystemToolkitMode_Scalability.cpp
--- a/Engine/Plugins/FX/Niagara/Source/NiagaraEditor/Private/Toolkits/SystemToolkitModes/NiagaraSystemToolkitMode_Scalability.cpp
+++ b/Engine/Plugins/FX/Niagara/Source/NiagaraEditor/Private/Toolkits/SystemToolkitModes/NiagaraSystemToolkitMode_Scalability.cpp
@@ -88,8 +88,8 @@ void FNiagaraSystemToolkitMode_Scalability::UnbindEmitterUpdates()
 
 	for(const TSharedRef<FNiagaraEmitterHandleViewModel>& EmitterHandleViewModel : SystemViewModel->GetEmitterHandleViewModels())
 	{
-		EmitterHandleViewModel->GetEmitterHandle()->GetInstance()->OnRenderersChanged().RemoveAll(this);
-		EmitterHandleViewModel->GetEmitterHandle()->GetInstance()->OnPropertiesChanged().RemoveAll(this);
+		EmitterHandleViewModel->GetEmitterHandle()->GetInstance().Emitter->OnRenderersChanged().RemoveAll(this);
+		EmitterHandleViewModel->GetEmitterHandle()->GetInstance().Emitter->OnPropertiesChanged().RemoveAll(this);
 	}
 }
 
@@ -100,19 +100,20 @@ void FNiagaraSystemToolkitMode_Scalability::BindEmitterPreviewOverrides()
 	
 	for(const TSharedRef<FNiagaraEmitterHandleViewModel>& EmitterHandleViewModel : SystemViewModel->GetEmitterHandleViewModels())
 	{
+		FVersionedNiagaraEmitterData* EmitterData = EmitterHandleViewModel->GetEmitterHandle()->GetEmitterData();
 		// we bind the quality level and device profile of the emitter itself to use our scalability preview
-		EmitterHandleViewModel->GetEmitterHandle()->GetInstance()->Platforms.OnOverrideQualityLevelDelegate.BindUObject(ScalabilityViewModel, &UNiagaraSystemScalabilityViewModel::GetPreviewQualityLevel);
-		EmitterHandleViewModel->GetEmitterHandle()->GetInstance()->Platforms.OnOverrideActiveDeviceProfileDelegate.BindUObject(ScalabilityViewModel, &UNiagaraSystemScalabilityViewModel::GetPreviewDeviceProfile);
+		EmitterData->Platforms.OnOverrideQualityLevelDelegate.BindUObject(ScalabilityViewModel, &UNiagaraSystemScalabilityViewModel::GetPreviewQualityLevel);
+		EmitterData->Platforms.OnOverrideActiveDeviceProfileDelegate.BindUObject(ScalabilityViewModel, &UNiagaraSystemScalabilityViewModel::GetPreviewDeviceProfile);
 
 		// we do the same for the overrides
-		for(FNiagaraEmitterScalabilityOverride& EmitterScalabilityOverride : EmitterHandleViewModel->GetEmitterHandle()->GetInstance()->ScalabilityOverrides.Overrides)
+		for(FNiagaraEmitterScalabilityOverride& EmitterScalabilityOverride : EmitterData->ScalabilityOverrides.Overrides)
 		{
 			EmitterScalabilityOverride.Platforms.OnOverrideQualityLevelDelegate.BindUObject(ScalabilityViewModel, &UNiagaraSystemScalabilityViewModel::GetPreviewQualityLevel);
 			EmitterScalabilityOverride.Platforms.OnOverrideActiveDeviceProfileDelegate.BindUObject(ScalabilityViewModel, &UNiagaraSystemScalabilityViewModel::GetPreviewDeviceProfile);
 		}
 
 		// and also for the renderers
-		EmitterHandleViewModel->GetEmitterHandle()->GetInstance()->ForEachRenderer([=](UNiagaraRendererProperties* RendererProperties)
+		EmitterData->ForEachRenderer([=](UNiagaraRendererProperties* RendererProperties)
 		{
 			RendererProperties->Platforms.OnOverrideQualityLevelDelegate.BindUObject(ScalabilityViewModel, &UNiagaraSystemScalabilityViewModel::GetPreviewQualityLevel);
 			RendererProperties->Platforms.OnOverrideActiveDeviceProfileDelegate.BindUObject(ScalabilityViewModel, &UNiagaraSystemScalabilityViewModel::GetPreviewDeviceProfile);
@@ -139,16 +140,17 @@ void FNiagaraSystemToolkitMode_Scalability::UnbindEmitterPreviewOverrides()
 	// we undo all the bindings we created in the Bind function here 
 	for(const TSharedRef<FNiagaraEmitterHandleViewModel>& EmitterHandleViewModel : SystemViewModel->GetEmitterHandleViewModels())
 	{
-		EmitterHandleViewModel->GetEmitterHandle()->GetInstance()->Platforms.OnOverrideQualityLevelDelegate.Unbind();
-		EmitterHandleViewModel->GetEmitterHandle()->GetInstance()->Platforms.OnOverrideActiveDeviceProfileDelegate.Unbind();
+		FVersionedNiagaraEmitterData* EmitterData = EmitterHandleViewModel->GetEmitterHandle()->GetEmitterData();
+		EmitterData->Platforms.OnOverrideQualityLevelDelegate.Unbind();
+		EmitterData->Platforms.OnOverrideActiveDeviceProfileDelegate.Unbind();
 		
-		for(FNiagaraEmitterScalabilityOverride& EmitterScalabilityOverride : EmitterHandleViewModel->GetEmitterHandle()->GetInstance()->ScalabilityOverrides.Overrides)
+		for(FNiagaraEmitterScalabilityOverride& EmitterScalabilityOverride : EmitterData->ScalabilityOverrides.Overrides)
 		{
 			EmitterScalabilityOverride.Platforms.OnOverrideQualityLevelDelegate.Unbind();
 			EmitterScalabilityOverride.Platforms.OnOverrideActiveDeviceProfileDelegate.Unbind();
 		}
 
-		EmitterHandleViewModel->GetEmitterHandle()->GetInstance()->ForEachRenderer([](UNiagaraRendererProperties* RendererProperties)
+		EmitterData->ForEachRenderer([](UNiagaraRendererProperties* RendererProperties)
 		{
 			RendererProperties->Platforms.OnOverrideQualityLevelDelegate.Unbind();
 			RendererProperties->Platforms.OnOverrideActiveDeviceProfileDelegate.Unbind();
@@ -176,8 +178,8 @@ void FNiagaraSystemToolkitMode_Scalability::BindEmitterUpdates()
 	TSharedPtr<FNiagaraSystemViewModel> SystemViewModel = SystemToolkit.Pin()->GetSystemViewModel();
 	for(const TSharedRef<FNiagaraEmitterHandleViewModel>& EmitterHandleViewModel : SystemViewModel->GetEmitterHandleViewModels())
 	{		
-		EmitterHandleViewModel->GetEmitterHandle()->GetInstance()->OnRenderersChanged().AddSP(this, &FNiagaraSystemToolkitMode_Scalability::BindEmitterPreviewOverrides);
-		EmitterHandleViewModel->GetEmitterHandle()->GetInstance()->OnPropertiesChanged().AddSP(this, &FNiagaraSystemToolkitMode_Scalability::BindEmitterPreviewOverrides);
+		EmitterHandleViewModel->GetEmitterHandle()->GetInstance().Emitter->OnRenderersChanged().AddSP(this, &FNiagaraSystemToolkitMode_Scalability::BindEmitterPreviewOverrides);
+		EmitterHandleViewModel->GetEmitterHandle()->GetInstance().Emitter->OnPropertiesChanged().AddSP(this, &FNiagaraSystemToolkitMode_Scalability::BindEmitterPreviewOverrides);
 	}
 
 	BindEmitterPreviewOverrides();
diff --git a/Engine/Plugins/FX/Niagara/Source/NiagaraEditor/Private/UpgradeNiagaraScriptResults.cpp b/Engine/Plugins/FX/Niagara/Source/NiagaraEditor/Private/UpgradeNiagaraScriptResults.cpp
--- a/Engine/Plugins/FX/Niagara/Source/NiagaraEditor/Private/UpgradeNiagaraScriptResults.cpp
+++ b/Engine/Plugins/FX/Niagara/Source/NiagaraEditor/Private/UpgradeNiagaraScriptResults.cpp
@@ -3,6 +3,37 @@
 #include "UpgradeNiagaraScriptResults.h"
 
 #include "NiagaraClipboard.h"
+#include "NiagaraEmitterHandle.h"
+#include "NiagaraNodeFunctionCall.h"
+#include "ViewModels/NiagaraEmitterHandleViewModel.h"
+#include "ViewModels/NiagaraEmitterViewModel.h"
+#include "ViewModels/Stack/NiagaraStackModuleItem.h"
+#include "ViewModels/Stack/NiagaraStackViewModel.h"
+
+template<typename T>
+TArray<T*> GetStackEntries(UNiagaraStackViewModel* StackViewModel, bool bRefresh = false)
+{
+	TArray<T*> Results;
+	TArray<UNiagaraStackEntry*> EntriesToCheck;
+	if (UNiagaraStackEntry* RootEntry = StackViewModel->GetRootEntry())
+	{
+		if (bRefresh)
+		{
+			RootEntry->RefreshChildren();
+		}
+		RootEntry->GetUnfilteredChildren(EntriesToCheck);
+	}
+	while (EntriesToCheck.Num() > 0)
+	{
+		UNiagaraStackEntry* Entry = EntriesToCheck.Pop();
+		if (T* ItemToCheck = Cast<T>(Entry))
+		{
+			Results.Add(ItemToCheck);
+		}
+		Entry->GetUnfilteredChildren(EntriesToCheck);
+	}
+	return Results;
+}
 
 template<typename T>
 T GetValue(const UNiagaraClipboardFunctionInput* Input)
@@ -253,4 +284,131 @@ UNiagaraPythonScriptModuleInput* UUpgradeNiagaraScriptResults::GetNewInput(const
 	return nullptr;
 }
 
+void UNiagaraPythonModule::Init(UNiagaraStackModuleItem* InModuleItem)
+{
+	ModuleItem = InModuleItem;
+}
+
+UNiagaraStackModuleItem* UNiagaraPythonModule::Object() const
+{
+	return ModuleItem;
+}
+
+void UNiagaraPythonEmitter::Init(TSharedRef<FNiagaraEmitterHandleViewModel> InEmitterViewModel)
+{
+	EmitterViewModel = InEmitterViewModel;
+}
+
+UNiagaraEmitter* UNiagaraPythonEmitter::Object()
+{
+	return EmitterViewModel->GetEmitterViewModel()->GetEmitter().Emitter;
+}
+
+FVersionedNiagaraEmitterData UNiagaraPythonEmitter::GetProperties() const
+{
+	if (EmitterViewModel->GetEmitterViewModel()->GetEmitter().GetEmitterData())
+	{
+		return *EmitterViewModel->GetEmitterViewModel()->GetEmitter().GetEmitterData();
+	}
+	return FVersionedNiagaraEmitterData();
+}
+
+void UNiagaraPythonEmitter::SetProperties(FVersionedNiagaraEmitterData Data)
+{
+	if (EmitterViewModel->GetEmitterViewModel()->GetEmitter().GetEmitterData())
+	{
+		FVersionedNiagaraEmitterData* EmitterData = EmitterViewModel->GetEmitterViewModel()->GetEmitter().GetEmitterData();
+		for (TFieldIterator<FProperty> PropertyIterator(FVersionedNiagaraEmitterData::StaticStruct()); PropertyIterator; ++PropertyIterator)
+		{
+			if (PropertyIterator->HasAllPropertyFlags(CPF_Edit))
+			{
+				void* Dest = PropertyIterator->ContainerPtrToValuePtr<void>(EmitterData);
+				void* Src = PropertyIterator->ContainerPtrToValuePtr<void>(&Data);
+				PropertyIterator->CopyCompleteValue(Dest, Src);
+			}
+		}
+	}
+}
+
+TArray<UNiagaraPythonModule*> UNiagaraPythonEmitter::GetModules() const
+{
+	TArray<UNiagaraPythonModule*> Result;
+
+	if (UNiagaraStackViewModel* StackViewModel = EmitterViewModel->GetEmitterStackViewModel())
+	{
+		TArray<UNiagaraStackModuleItem*> StackModuleItems =	GetStackEntries<UNiagaraStackModuleItem>(StackViewModel);
+		for (UNiagaraStackModuleItem* ModuleItem : StackModuleItems)
+		{
+			UNiagaraPythonModule* Module = NewObject<UNiagaraPythonModule>();
+			Module->Init(ModuleItem);
+			Result.Add(Module);
+		}
+	}
+
+	return Result;
+}
+
+bool UNiagaraPythonEmitter::HasModule(const FString& ModuleName) const
+{
+	return GetModule(ModuleName)->Object() != nullptr;
+}
+
+UNiagaraPythonModule* UNiagaraPythonEmitter::GetModule(const FString& ModuleName) const
+{
+	UNiagaraPythonModule* Module = NewObject<UNiagaraPythonModule>();
+	if (UNiagaraStackViewModel* StackViewModel = EmitterViewModel->GetEmitterStackViewModel())
+	{
+		TArray<UNiagaraStackModuleItem*> StackModuleItems =	GetStackEntries<UNiagaraStackModuleItem>(StackViewModel);
+		for (UNiagaraStackModuleItem* ModuleItem : StackModuleItems)
+		{
+			if (ModuleItem->GetModuleNode().GetFunctionName() == ModuleName)
+			{
+				Module->Init(ModuleItem);
+				break;
+			}
+		}
+	}
+	return Module;
+}
+
+void UUpgradeNiagaraEmitterContext::Init(UNiagaraPythonEmitter* InOldEmitter, UNiagaraPythonEmitter* InNewEmitter)
+{
+	OldEmitter = InOldEmitter;
+	NewEmitter = InNewEmitter;
+
+	UpgradeVersionData.Empty();
+
+	if (!IsValid())
+	{
+		return;
+	}
+	
+	FVersionedNiagaraEmitterData* SourceData = OldEmitter->EmitterViewModel->GetEmitterHandle()->GetInstance().GetEmitterData();
+	FVersionedNiagaraEmitterData* TargetData = NewEmitter->EmitterViewModel->GetEmitterHandle()->GetInstance().GetEmitterData();
+	FVersionedNiagaraEmitter SourceParent = SourceData->GetParent();
+	FVersionedNiagaraEmitter TargetParent = TargetData->GetParent();
+	FNiagaraAssetVersion SourceVersion = SourceParent.GetEmitterData()->Version;
+	FNiagaraAssetVersion TargetVersion = TargetParent.GetEmitterData()->Version;
+
+	// gather script required to execute
+	for (const FNiagaraAssetVersion& Version : SourceParent.Emitter->GetAllAvailableVersions())
+	{
+		if (SourceVersion <= Version && Version <= TargetVersion)
+		{
+			FVersionedNiagaraEmitterData* ParentData = SourceParent.Emitter->GetEmitterData(Version.VersionGuid);
+			UpgradeVersionData.Add(ParentData);
+		}
+	}
+}
+
+bool UUpgradeNiagaraEmitterContext::IsValid() const
+{
+	return OldEmitter && OldEmitter->IsValid() && NewEmitter && NewEmitter->IsValid();
+}
+
+const TArray<FVersionedNiagaraEmitterData*>& UUpgradeNiagaraEmitterContext::GetUpgradeData() const
+{
+	return UpgradeVersionData;
+}
+
 
diff --git a/Engine/Plugins/FX/Niagara/Source/NiagaraEditor/Private/ViewModels/NiagaraCurveSelectionViewModel.cpp b/Engine/Plugins/FX/Niagara/Source/NiagaraEditor/Private/ViewModels/NiagaraCurveSelectionViewModel.cpp
--- a/Engine/Plugins/FX/Niagara/Source/NiagaraEditor/Private/ViewModels/NiagaraCurveSelectionViewModel.cpp
+++ b/Engine/Plugins/FX/Niagara/Source/NiagaraEditor/Private/ViewModels/NiagaraCurveSelectionViewModel.cpp
@@ -11,7 +11,6 @@
 #include "NiagaraNodeFunctionCall.h"
 #include "NiagaraNodeInput.h"
 #include "NiagaraNodeOutput.h"
-#include "NiagaraNodeParameterMapSet.h"
 #include "NiagaraParameterStore.h"
 #include "NiagaraScriptSource.h"
 #include "NiagaraScript.h"
@@ -22,7 +21,6 @@
 #include "NiagaraSimulationStageBase.h"
 #include "Toolkits/NiagaraSystemToolkit.h"
 #include "ViewModels/NiagaraEmitterHandleViewModel.h"
-#include "ViewModels/NiagaraEmitterViewModel.h"
 #include "ViewModels/NiagaraPlaceholderDataInterfaceManager.h"
 #include "ViewModels/NiagaraSystemViewModel.h"
 #include "ViewModels/Stack/NiagaraParameterHandle.h"
@@ -769,11 +767,11 @@ TSharedPtr<FNiagaraCurveSelectionTreeNode> UNiagaraCurveSelectionViewModel::Crea
 		OldEmitterChildNodes = OldEmitterNode->GetChildNodes();
 	}
 
-	UNiagaraEmitter* Emitter = EmitterHandle.GetInstance();
+	FVersionedNiagaraEmitterData* EmitterData = EmitterHandle.GetEmitterData();
 
-	UNiagaraStackEditorData& StackEditorData = CastChecked<UNiagaraEmitterEditorData>(EmitterHandle.GetInstance()->GetEditorData())->GetStackEditorData();
+	UNiagaraStackEditorData& StackEditorData = CastChecked<UNiagaraEmitterEditorData>(EmitterHandle.GetEmitterData()->GetEditorData())->GetStackEditorData();
 	TSharedPtr<FNiagaraCurveSelectionTreeNode> EmitterSpawnNode = CreateNodeForScript(
-		OldEmitterChildNodes, *Emitter->EmitterSpawnScriptProps.Script, TEXT("Emitter Spawn"), StackEditorData,
+		OldEmitterChildNodes, *EmitterData->EmitterSpawnScriptProps.Script, TEXT("Emitter Spawn"), StackEditorData,
 		UNiagaraStackEntry::FExecutionCategoryNames::Emitter, UNiagaraStackEntry::FExecutionSubcategoryNames::Spawn,
 		&EmitterHandle);
 	if (EmitterSpawnNode.IsValid())
@@ -782,7 +780,7 @@ TSharedPtr<FNiagaraCurveSelectionTreeNode> UNiagaraCurveSelectionViewModel::Crea
 	}
 
 	TSharedPtr<FNiagaraCurveSelectionTreeNode> EmitterUpdateNode = CreateNodeForScript(
-		OldEmitterChildNodes, *Emitter->EmitterUpdateScriptProps.Script, TEXT("Emitter Update"), StackEditorData,
+		OldEmitterChildNodes, *EmitterData->EmitterUpdateScriptProps.Script, TEXT("Emitter Update"), StackEditorData,
 		UNiagaraStackEntry::FExecutionCategoryNames::Emitter, UNiagaraStackEntry::FExecutionSubcategoryNames::Update,
 		&EmitterHandle);
 	if (EmitterUpdateNode.IsValid())
@@ -791,7 +789,7 @@ TSharedPtr<FNiagaraCurveSelectionTreeNode> UNiagaraCurveSelectionViewModel::Crea
 	}
 
 	TSharedPtr<FNiagaraCurveSelectionTreeNode> ParticleSpawnNode = CreateNodeForScript(
-		OldEmitterChildNodes, *Emitter->SpawnScriptProps.Script, TEXT("Particle Spawn"), StackEditorData,
+		OldEmitterChildNodes, *EmitterData->SpawnScriptProps.Script, TEXT("Particle Spawn"), StackEditorData,
 		UNiagaraStackEntry::FExecutionCategoryNames::Particle, UNiagaraStackEntry::FExecutionSubcategoryNames::Spawn,
 		&EmitterHandle);
 	if (ParticleSpawnNode.IsValid())
@@ -800,7 +798,7 @@ TSharedPtr<FNiagaraCurveSelectionTreeNode> UNiagaraCurveSelectionViewModel::Crea
 	}
 
 	TSharedPtr<FNiagaraCurveSelectionTreeNode> ParticleUpdateNode = CreateNodeForScript(
-		OldEmitterChildNodes, *Emitter->UpdateScriptProps.Script, TEXT("Particle Update"), StackEditorData,
+		OldEmitterChildNodes, *EmitterData->UpdateScriptProps.Script, TEXT("Particle Update"), StackEditorData,
 		UNiagaraStackEntry::FExecutionCategoryNames::Particle, UNiagaraStackEntry::FExecutionSubcategoryNames::Update,
 		&EmitterHandle);
 	if (ParticleUpdateNode.IsValid())
@@ -808,7 +806,7 @@ TSharedPtr<FNiagaraCurveSelectionTreeNode> UNiagaraCurveSelectionViewModel::Crea
 		NewEmitterChildNodes.Add(ParticleUpdateNode.ToSharedRef());
 	}
 
-	for (const FNiagaraEventScriptProperties& EventScriptProps : Emitter->GetEventHandlers())
+	for (const FNiagaraEventScriptProperties& EventScriptProps : EmitterData->GetEventHandlers())
 	{
 		TSharedPtr<FNiagaraCurveSelectionTreeNode> EmitterEventNode = CreateNodeForScript(
 			OldEmitterChildNodes, *EventScriptProps.Script, FString::Printf(TEXT("Event Handler - %s"), *EventScriptProps.SourceEventName.ToString()), StackEditorData,
@@ -820,7 +818,7 @@ TSharedPtr<FNiagaraCurveSelectionTreeNode> UNiagaraCurveSelectionViewModel::Crea
 		}
 	}
 
-	for (UNiagaraSimulationStageBase* SimulationStage : Emitter->GetSimulationStages())
+	for (UNiagaraSimulationStageBase* SimulationStage : EmitterData->GetSimulationStages())
 	{
 		TSharedPtr<FNiagaraCurveSelectionTreeNode> EmitterSimulationStageNode = CreateNodeForScript(
 			OldEmitterChildNodes, *SimulationStage->Script, FString::Printf(TEXT("Simulation Stage - %s"), *SimulationStage->SimulationStageName.ToString()), StackEditorData,
@@ -844,7 +842,7 @@ TSharedPtr<FNiagaraCurveSelectionTreeNode> UNiagaraCurveSelectionViewModel::Crea
 			NewEmitterNode->SetDataId(EmitterDataId);
 			NewEmitterNode->SetSecondDisplayName(FText::FromName(FNiagaraConstants::EmitterNamespace));
 			NewEmitterNode->SetStyle(ENiagaraCurveSelectionNodeStyleMode::TopLevelObject, UNiagaraStackEntry::FExecutionCategoryNames::Emitter, UNiagaraStackEntry::FExecutionSubcategoryNames::Settings, false);
-			NewEmitterNode->SetDisplayedObjectKey(FObjectKey(EmitterHandle.GetInstance()));
+			NewEmitterNode->SetDisplayedObjectKey(FObjectKey(EmitterHandle.GetInstance().Emitter));
 			NewEmitterNode->SetShowInTree(true);
 			NewEmitterNode->SetIsExpanded(true);
 
@@ -888,7 +886,7 @@ void UNiagaraCurveSelectionViewModel::Refresh()
 		for (TSharedRef<FNiagaraEmitterHandleViewModel> EmitterHandleViewModel : SystemViewModel->GetEmitterHandleViewModels())
 		{
 			// During transactions Emitters may not have their EditorData set if they were created during the transaction; skip if this is true.
-			if (EmitterHandleViewModel->GetEmitterHandle() && EmitterHandleViewModel->GetEmitterHandle()->GetInstance()->GetEditorData())
+			if (EmitterHandleViewModel->GetEmitterHandle() && EmitterHandleViewModel->GetEmitterHandle()->GetEmitterData()->GetEditorData())
 			{
 				TSharedPtr<FNiagaraCurveSelectionTreeNode> EmitterNode = CreateNodeForEmitter(OldRootChildNodes, *EmitterHandleViewModel->GetEmitterHandle());
 				if (EmitterNode.IsValid())
diff --git a/Engine/Plugins/FX/Niagara/Source/NiagaraEditor/Private/ViewModels/NiagaraEmitterHandleViewModel.cpp b/Engine/Plugins/FX/Niagara/Source/NiagaraEditor/Private/ViewModels/NiagaraEmitterHandleViewModel.cpp
--- a/Engine/Plugins/FX/Niagara/Source/NiagaraEditor/Private/ViewModels/NiagaraEmitterHandleViewModel.cpp
+++ b/Engine/Plugins/FX/Niagara/Source/NiagaraEditor/Private/ViewModels/NiagaraEmitterHandleViewModel.cpp
@@ -6,19 +6,12 @@
 #include "ViewModels/NiagaraSystemViewModel.h"
 #include "ViewModels/NiagaraEmitterViewModel.h"
 #include "ViewModels/Stack/NiagaraStackViewModel.h"
-#include "ViewModels/NiagaraScriptViewModel.h"
 #include "NiagaraScriptGraphViewModel.h"
-#include "NiagaraObjectSelection.h"
 #include "NiagaraScriptSource.h"
 #include "NiagaraNodeOutput.h"
-#include "NiagaraGraph.h"
-#include "NiagaraNodeInput.h"
 #include "NiagaraScriptOutputCollectionViewModel.h"
 #include "Algo/Find.h"
-#include "Framework/Notifications/NotificationManager.h"
-#include "Widgets/Notifications/SNotificationList.h"
 #include "ScopedTransaction.h"
-#include "NiagaraRendererProperties.h"
 #include "ViewModels/Stack/NiagaraStackRoot.h"
 #include "ViewModels/Stack/NiagaraStackRenderItemGroup.h"
 #include "ViewModels/Stack/NiagaraStackRendererItem.h"
@@ -63,7 +56,6 @@ void FNiagaraEmitterHandleViewModel::Cleanup()
 void FNiagaraEmitterHandleViewModel::GetRendererEntries(TArray<UNiagaraStackEntry*>& InRenderingEntries)
 {
 	InRenderingEntries.Empty();
-	UNiagaraEmitter* Emitter = GetEmitterHandle()->GetInstance();
 	UNiagaraStackRoot* StackRoot = Cast<UNiagaraStackRoot>(EmitterStackViewModel->GetRootEntry());
 	if (StackRoot)
 	{
@@ -92,7 +84,7 @@ FGuid FNiagaraEmitterHandleViewModel::AddMessage(UNiagaraMessageData* NewMessage
 	{
 		const FGuid NewGuid = InNewGuid.IsValid() ? InNewGuid : FGuid::NewGuid();
 		
-		EmitterHandle->GetInstance()->AddMessage(NewGuid, static_cast<UNiagaraMessageDataBase*>(NewMessage));
+		EmitterHandle->GetInstance().Emitter->AddMessage(NewGuid, NewMessage);
 		return NewGuid;
 	}
 	return FGuid();
@@ -102,7 +94,7 @@ void FNiagaraEmitterHandleViewModel::RemoveMessage(const FGuid& MessageKey) cons
 {
 	if (ensureMsgf(EmitterHandle != nullptr, TEXT("EmitterHandleViewModel had a null EmitterHandle!")))
 	{
-		EmitterHandle->GetInstance()->RemoveMessage(MessageKey);
+		EmitterHandle->GetInstance().Emitter->RemoveMessage(MessageKey);
 	}
 }
 
@@ -117,7 +109,7 @@ void FNiagaraEmitterHandleViewModel::Initialize(TSharedRef<FNiagaraSystemViewMod
 	OwningSystemViewModelWeak = InOwningSystemViewModel;
 	EmitterHandleIndex = InEmitterHandleIndex;
 	EmitterHandle = &InOwningSystemViewModel->GetSystem().GetEmitterHandle(InEmitterHandleIndex);
-	UNiagaraEmitter* Emitter = EmitterHandle != nullptr ? EmitterHandle->GetInstance() : nullptr;
+	FVersionedNiagaraEmitter Emitter = EmitterHandle != nullptr ? EmitterHandle->GetInstance() : FVersionedNiagaraEmitter();
 	EmitterViewModel->Initialize(Emitter, InSimulation);
 	EmitterStackViewModel->InitializeWithViewModels(InOwningSystemViewModel, this->AsShared(), FNiagaraStackViewModelOptions(false, true));
 }
