Added versioning for Niagara emitters

diff --git a/Engine/Plugins/FX/Niagara/Source/NiagaraEditor/Private/ViewModels/NiagaraEmitterViewModel.cpp b/Engine/Plugins/FX/Niagara/Source/NiagaraEditor/Private/ViewModels/NiagaraEmitterViewModel.cpp
--- a/Engine/Plugins/FX/Niagara/Source/NiagaraEditor/Private/ViewModels/NiagaraEmitterViewModel.cpp
+++ b/Engine/Plugins/FX/Niagara/Source/NiagaraEditor/Private/ViewModels/NiagaraEmitterViewModel.cpp
@@ -15,7 +15,6 @@
 #include "NiagaraGraph.h"
 #include "NiagaraSystem.h"
 #include "NiagaraSystemInstance.h"
-#include "NiagaraParameterDefinitions.h"
 
 #include "ScopedTransaction.h"
 #include "IContentBrowserSingleton.h"
@@ -39,7 +38,8 @@ namespace NiagaraCommands
 }
 
 FNiagaraEmitterViewModel::FNiagaraEmitterViewModel(bool bInIsForDataProcessingOnly)
-	: SharedScriptViewModel(MakeShareable(new FNiagaraScriptViewModel(LOCTEXT("SharedDisplayName", "Graph"), ENiagaraParameterEditMode::EditAll, bInIsForDataProcessingOnly)))
+	: EmitterWeakPtr(FVersionedNiagaraEmitterWeakPtr(nullptr, FGuid()))
+	, SharedScriptViewModel(MakeShareable(new FNiagaraScriptViewModel(LOCTEXT("SharedDisplayName", "Graph"), ENiagaraParameterEditMode::EditAll, bInIsForDataProcessingOnly)))
 	, bUpdatingSelectionInternally(false)
 	, ExecutionStateEnum(StaticEnum<ENiagaraExecutionState>())
 {	
@@ -52,10 +52,10 @@ void FNiagaraEmitterViewModel::Cleanup()
 		NewParentWindow.Reset();
 	}
 
-	if (Emitter.IsValid())
+	if (EmitterWeakPtr.Emitter.IsValid())
 	{
-		Emitter->OnEmitterVMCompiled().RemoveAll(this);
-		Emitter->OnPropertiesChanged().RemoveAll(this);
+		EmitterWeakPtr.Emitter->OnEmitterVMCompiled().RemoveAll(this);
+		EmitterWeakPtr.Emitter->OnPropertiesChanged().RemoveAll(this);
 	}
 
 	if (SharedScriptViewModel.IsValid())
@@ -67,7 +67,7 @@ void FNiagaraEmitterViewModel::Cleanup()
 	RemoveScriptEventHandlers();
 }
 
-bool FNiagaraEmitterViewModel::Initialize(UNiagaraEmitter* InEmitter, TWeakPtr<FNiagaraEmitterInstance, ESPMode::ThreadSafe> InSimulation)
+bool FNiagaraEmitterViewModel::Initialize(const FVersionedNiagaraEmitter& InEmitter, TWeakPtr<FNiagaraEmitterInstance, ESPMode::ThreadSafe> InSimulation)
 {
 	SetEmitter(InEmitter);
 	SetSimulation(InSimulation);
@@ -77,13 +77,13 @@ bool FNiagaraEmitterViewModel::Initialize(UNiagaraEmitter* InEmitter, TWeakPtr<F
 
 void FNiagaraEmitterViewModel::Reset()
 {
-	SetEmitter(nullptr);
+	SetEmitter(FVersionedNiagaraEmitter());
 	SetSimulation(nullptr);
 }
 
 INiagaraParameterDefinitionsSubscriber* FNiagaraEmitterViewModel::GetParameterDefinitionsSubscriber()
 {
-	return GetEmitter();
+	return GetEmitter().Emitter;
 }
 
 FNiagaraEmitterViewModel::~FNiagaraEmitterViewModel()
@@ -94,13 +94,13 @@ FNiagaraEmitterViewModel::~FNiagaraEmitterViewModel()
 	//UE_LOG(LogNiagaraEditor, Warning, TEXT("Deleting Emitter view model %p"), this);
 }
 
-void FNiagaraEmitterViewModel::SetEmitter(UNiagaraEmitter* InEmitter)
+void FNiagaraEmitterViewModel::SetEmitter(FVersionedNiagaraEmitter InEmitter)
 {
-	if (Emitter.IsValid())
+	if (EmitterWeakPtr.IsValid())
 	{
-		Emitter->OnEmitterVMCompiled().RemoveAll(this);
-		Emitter->OnEmitterGPUCompiled().RemoveAll(this);
-		Emitter->OnPropertiesChanged().RemoveAll(this);
+		EmitterWeakPtr.Emitter->OnEmitterVMCompiled().RemoveAll(this);
+		EmitterWeakPtr.Emitter->OnEmitterGPUCompiled().RemoveAll(this);
+		EmitterWeakPtr.Emitter->OnPropertiesChanged().RemoveAll(this);
 		UnregisterViewModelWithMap(RegisteredHandle);
 	}
 
@@ -108,18 +108,18 @@ void FNiagaraEmitterViewModel::SetEmitter(UNiagaraEmitter* InEmitter)
 
 	RemoveScriptEventHandlers();
 
-	Emitter = InEmitter;
+	EmitterWeakPtr = InEmitter.ToWeakPtr();
 
-	if (Emitter.IsValid())
+	if (EmitterWeakPtr.IsValid())
 	{
-		Emitter->OnEmitterVMCompiled().AddSP(this, &FNiagaraEmitterViewModel::OnVMCompiled);
-		Emitter->OnPropertiesChanged().AddSP(this, &FNiagaraEmitterViewModel::OnEmitterPropertiesChanged);
-		Emitter->OnEmitterGPUCompiled().AddSP(this, &FNiagaraEmitterViewModel::OnGPUCompiled);
+		EmitterWeakPtr.Emitter->OnEmitterVMCompiled().AddSP(this, &FNiagaraEmitterViewModel::OnVMCompiled);
+		EmitterWeakPtr.Emitter->OnPropertiesChanged().AddSP(this, &FNiagaraEmitterViewModel::OnEmitterPropertiesChanged);
+		EmitterWeakPtr.Emitter->OnEmitterGPUCompiled().AddSP(this, &FNiagaraEmitterViewModel::OnGPUCompiled);
 	}
 
 	AddScriptEventHandlers();
 
-	RegisteredHandle = RegisterViewModelWithMap(InEmitter, this);
+	RegisteredHandle = RegisterViewModelWithMap(InEmitter.Emitter, this);
 
 	check(SharedScriptViewModel.IsValid());
 	SharedScriptViewModel->SetScripts(InEmitter);
@@ -137,35 +137,36 @@ void FNiagaraEmitterViewModel::SetSimulation(TWeakPtr<FNiagaraEmitterInstance, E
 	Simulation = InSimulation;
 }
 
-UNiagaraEmitter* FNiagaraEmitterViewModel::GetEmitter()
+FVersionedNiagaraEmitter FNiagaraEmitterViewModel::GetEmitter()
 {
-	return Emitter.Get();
+	return EmitterWeakPtr.ResolveWeakPtr();
 }
 
 bool FNiagaraEmitterViewModel::HasParentEmitter() const
 {
-	return Emitter.IsValid() && Emitter->GetParent() != nullptr;
+	return EmitterWeakPtr.IsValid() && EmitterWeakPtr.GetEmitterData()->GetParent().Emitter != nullptr;
 }
 
-const UNiagaraEmitter* FNiagaraEmitterViewModel::GetParentEmitter() const
+FVersionedNiagaraEmitter FNiagaraEmitterViewModel::GetParentEmitter() const
 {
-	return Emitter.IsValid() ? Emitter->GetParent() : nullptr;
+	FVersionedNiagaraEmitterData* EmitterData = EmitterWeakPtr.GetEmitterData();
+	return EmitterData ? EmitterData->GetParent() : FVersionedNiagaraEmitter();
 }
 
 FText FNiagaraEmitterViewModel::GetParentNameText() const
 {
-	if (Emitter.IsValid() && Emitter->GetParent() != nullptr)
+	if (EmitterWeakPtr.IsValid() && EmitterWeakPtr.GetEmitterData()->GetParent().Emitter != nullptr)
 	{
-		return FText::FromString(Emitter->GetParent()->GetName());
+		return FText::FromString(EmitterWeakPtr.GetEmitterData()->GetParent().Emitter->GetName());
 	}
 	return FText();
 }
 
 FText FNiagaraEmitterViewModel::GetParentPathNameText() const
 {
-	if (Emitter.IsValid() && Emitter->GetParent() != nullptr)
+	if (EmitterWeakPtr.IsValid() && EmitterWeakPtr.GetEmitterData()->GetParent().Emitter != nullptr)
 	{
-		return FText::FromString(Emitter->GetParent()->GetPathName());
+		return FText::FromString(EmitterWeakPtr.GetEmitterData()->GetParent().Emitter->GetPathName());
 	}
 	return FText();
 }
@@ -206,30 +207,29 @@ void FNiagaraEmitterViewModel::UpdateParentEmitter(const TArray<FAssetData> & Ac
 	FScopedTransaction ScopedTransaction(LOCTEXT("RemoveParentEmitterTransaction", "Remove Parent Emitter"));
 	if ((ActivationMethod == EAssetTypeActivationMethod::DoubleClicked || ActivationMethod == EAssetTypeActivationMethod::Opened) && ActivatedAssets.Num() == 1)
 	{
-		if (Emitter.IsValid())
+		if (EmitterWeakPtr.IsValid())
 		{
 			FAssetData NewParentData = ActivatedAssets[0];
 			UNiagaraEmitter* NewParentPtr = Cast<UNiagaraEmitter>(NewParentData.GetAsset());
+			UNiagaraEmitter* Emitter = EmitterWeakPtr.Emitter.Get();
 			if (NewParentPtr && NewParentPtr != Emitter)
 			{
 				Emitter->PreEditChange(nullptr);
-				Emitter->Reparent(*NewParentPtr);
+				EmitterWeakPtr.GetEmitterData()->Reparent(FVersionedNiagaraEmitter(NewParentPtr, NewParentPtr->GetExposedVersion().VersionGuid));
 				Emitter->MergeChangesFromParent();
 				Emitter->PostEditChange();
-				EmitterHandleViewModel->GetOwningSystemViewModel()->RefreshAll();
+				EmitterHandleViewModel->GetOwningSystemViewModel()->ResetStack();
 				NewParentWindow->RequestDestroyWindow();
 			}
-			
 		}
-		
 	}
 }
 
 void FNiagaraEmitterViewModel::RemoveParentEmitter()
 {
 	FScopedTransaction ScopedTransaction(LOCTEXT("RemoveParentEmitterTransaction", "Remove Parent Emitter"));
-	Emitter->Modify();
-	Emitter->RemoveParent();
+	EmitterWeakPtr.Emitter->Modify();
+	EmitterWeakPtr.GetEmitterData()->RemoveParent();
 	OnParentRemovedDelegate.Broadcast();
 }
 
@@ -250,31 +250,25 @@ FText FNiagaraEmitterViewModel::GetStatsText() const
 			}
 
 			const FNiagaraEmitterHandle& Handle = SimInstance->GetEmitterHandle();
-			if (Handle.GetInstance())
+			if (FVersionedNiagaraEmitterData* EmitterData = Handle.GetEmitterData())
 			{
 				if (Handle.IsValid() == false)
 				{
 					return LOCTEXT("InvalidHandle", "Invalid handle");
 				}
-
-				UNiagaraEmitter* HandleEmitter = Handle.GetInstance();
-				if (HandleEmitter == nullptr)
-				{
-					return LOCTEXT("NullInstance", "Invalid instance");
-				}
 				
 				if (Handle.GetIsEnabled() == false)
 				{
 					return LOCTEXT("DisabledSimulation", "Emitter is not enabled. Excluded from code.");
 				}
 
-				if (!HandleEmitter->IsValid())
+				if (!Handle.GetEmitterData()->IsValid())
 				{
 					return LOCTEXT("InvalidInstance", "Invalid Emitter! May have compile errors.");
 				}
 
  				int32 QualityLevel = FNiagaraPlatformSet::GetQualityLevel();
-  				if (!HandleEmitter->IsAllowedByScalability())
+  				if (!EmitterData->IsAllowedByScalability())
   				{
 					return FText::Format(ParticleDisabledDueToScalability, FNiagaraPlatformSet::GetQualityLevelText(QualityLevel));
   				}
@@ -285,17 +279,17 @@ FText FNiagaraEmitterViewModel::GetStatsText() const
 				}
 				else
 				{
-					const double Megabyte = 1024 * 1024;
+					constexpr double Megabyte = 1024 * 1024;
 					return FText::Format(StatsFormat,
 						FText::AsNumber(SimInstance->GetNumParticles()),
 						FText::AsNumber(SimInstance->GetTotalCPUTimeMS(), &FractionalFormatOptions),
-						FText::AsNumber(double(SimInstance->GetTotalBytesUsed()) / Megabyte, &FractionalFormatOptions),
-						ExecutionStateEnum->GetDisplayNameTextByValue((int32)SimInstance->GetExecutionState()));
+						FText::AsNumber(static_cast<double>(SimInstance->GetTotalBytesUsed()) / Megabyte, &FractionalFormatOptions),
+						ExecutionStateEnum->GetDisplayNameTextByValue(static_cast<int32>(SimInstance->GetExecutionState())));
 				}
 			}
 		}
 	}
-	else if(!Emitter->IsReadyToRun())
+	else if(!EmitterWeakPtr.GetEmitterData()->IsReadyToRun())
 	{
 		return LOCTEXT("SimulationNotReady", "Preparing simulation...");
 	}
@@ -310,9 +304,9 @@ TSharedRef<FNiagaraScriptViewModel> FNiagaraEmitterViewModel::GetSharedScriptVie
 
 const UNiagaraEmitterEditorData& FNiagaraEmitterViewModel::GetEditorData() const
 {
-	check(Emitter.IsValid());
+	check(EmitterWeakPtr.Emitter.IsValid());
 
-	const UNiagaraEmitterEditorData* EditorData = Cast<UNiagaraEmitterEditorData>(Emitter->GetEditorData());
+	const UNiagaraEmitterEditorData* EditorData = Cast<UNiagaraEmitterEditorData>(EmitterWeakPtr.GetEmitterData()->GetEditorData());
 	if (EditorData == nullptr)
 	{
 		EditorData = GetDefault<UNiagaraEmitterEditorData>();
@@ -322,39 +316,40 @@ const UNiagaraEmitterEditorData& FNiagaraEmitterViewModel::GetEditorData() const
 
 UNiagaraEmitterEditorData& FNiagaraEmitterViewModel::GetOrCreateEditorData()
 {
-	check(Emitter.IsValid());
-	UNiagaraEmitterEditorData* EditorData = Cast<UNiagaraEmitterEditorData>(Emitter->GetEditorData());
+	check(EmitterWeakPtr.IsValid());
+	UNiagaraEmitterEditorData* EditorData = Cast<UNiagaraEmitterEditorData>(EmitterWeakPtr.GetEmitterData()->GetEditorData());
 	if (EditorData == nullptr)
 	{
-		EditorData = NewObject<UNiagaraEmitterEditorData>(Emitter.Get(), NAME_None, RF_Transactional);
-		Emitter->Modify(); 
-		Emitter->SetEditorData(EditorData);
+		EditorData = NewObject<UNiagaraEmitterEditorData>(EmitterWeakPtr.Emitter.Get(), NAME_None, RF_Transactional);
+		EmitterWeakPtr.Emitter->Modify(); 
+		EmitterWeakPtr.Emitter->SetEditorData(EditorData, EmitterWeakPtr.Version);
 	}
 	return *EditorData;
 }
 
 void FNiagaraEmitterViewModel::AddEventHandler(FNiagaraEventScriptProperties& EventScriptProperties, bool bResetGraphForOutput /*= false*/)
 {
-	EventScriptProperties.Script = NewObject<UNiagaraScript>(GetEmitter(), MakeUniqueObjectName(GetEmitter(), UNiagaraScript::StaticClass(), "EventScript"), EObjectFlags::RF_Transactional);
+	EventScriptProperties.Script = NewObject<UNiagaraScript>(GetEmitter().Emitter, MakeUniqueObjectName(GetEmitter().Emitter, UNiagaraScript::StaticClass(), "EventScript"), EObjectFlags::RF_Transactional);
 	EventScriptProperties.Script->SetUsage(ENiagaraScriptUsage::ParticleEventScript);
 	EventScriptProperties.Script->SetUsageId(FGuid::NewGuid());
 	EventScriptProperties.Script->SetLatestSource(GetSharedScriptViewModel()->GetGraphViewModel()->GetScriptSource());
-	Emitter->AddEventHandler(EventScriptProperties);
+	EmitterWeakPtr.Emitter->AddEventHandler(EventScriptProperties, EmitterWeakPtr.Version);
 	if (bResetGraphForOutput)
 	{
 		FNiagaraStackGraphUtilities::ResetGraphForOutput(*GetSharedScriptViewModel()->GetGraphViewModel()->GetGraph(), ENiagaraScriptUsage::ParticleEventScript, EventScriptProperties.Script->GetUsageId());
 	}
 }
 
-void FNiagaraEmitterViewModel::OnGPUCompiled(UNiagaraEmitter* InEmitter)
+void FNiagaraEmitterViewModel::OnGPUCompiled(FVersionedNiagaraEmitter InEmitter)
 {
 	// Flush compilation logs just like the CPU script compiled.
 	OnVMCompiled(InEmitter);
 }
 
-void FNiagaraEmitterViewModel::OnVMCompiled(UNiagaraEmitter* InEmitter)
+void FNiagaraEmitterViewModel::OnVMCompiled(FVersionedNiagaraEmitter InEmitter)
 {
-	if (Emitter.IsValid() && InEmitter == Emitter.Get())
+	FVersionedNiagaraEmitter Emitter = EmitterWeakPtr.ResolveWeakPtr();
+	if (InEmitter == Emitter)
 	{
 		TArray<ENiagaraScriptCompileStatus> CompileStatuses;
 		TArray<FString> CompileErrors;
@@ -365,7 +360,8 @@ void FNiagaraEmitterViewModel::OnVMCompiled(UNiagaraEmitter* InEmitter)
 		FString AggregateErrors;
 
 		TArray<UNiagaraScript*> Scripts;
-		Emitter->GetScripts(Scripts, true);
+		FVersionedNiagaraEmitterData* EmitterData = Emitter.GetEmitterData();
+		EmitterData->GetScripts(Scripts, true);
 
 		int32 EventsFound = 0;
 		for (int32 i = 0; i < Scripts.Num(); i++)
@@ -402,7 +398,7 @@ void FNiagaraEmitterViewModel::OnVMCompiled(UNiagaraEmitter* InEmitter)
 		
 		// Going ahead and making an array for future proofing..
 		{
-			UNiagaraScript* GPUScript = Emitter->GetGPUComputeScript();
+			UNiagaraScript* GPUScript = EmitterData->GetGPUComputeScript();
 			GPUScripts.Add(GPUScript);
 		}
 
@@ -460,11 +456,12 @@ ENiagaraScriptCompileStatus FNiagaraEmitterViewModel::GetLatestCompileStatus()
 	check(SharedScriptViewModel.IsValid());
 	ENiagaraScriptCompileStatus UnionStatus = SharedScriptViewModel->GetLatestCompileStatus();
 
-	if (Emitter->SimTarget == ENiagaraSimTarget::GPUComputeSim)
+	FVersionedNiagaraEmitterData* EmitterData = EmitterWeakPtr.GetEmitterData();
+	if (EmitterData && EmitterData->SimTarget == ENiagaraSimTarget::GPUComputeSim)
 	{
 		if (UnionStatus != ENiagaraScriptCompileStatus::NCS_Dirty)
 		{
-			if (!Emitter->GetGPUComputeScript()->AreScriptAndSourceSynchronized())
+			if (!EmitterData->GetGPUComputeScript()->AreScriptAndSourceSynchronized())
 			{
 				UnionStatus = ENiagaraScriptCompileStatus::NCS_Dirty;
 			}
@@ -506,10 +503,10 @@ FNiagaraEmitterViewModel::FOnScriptParameterStoreChanged& FNiagaraEmitterViewMod
 
 void FNiagaraEmitterViewModel::AddScriptEventHandlers()
 {
-	if (Emitter.IsValid())
+	if (FVersionedNiagaraEmitterData* EmitterData = EmitterWeakPtr.GetEmitterData())
 	{
 		TArray<UNiagaraScript*> Scripts;
-		Emitter->GetScripts(Scripts, false);
+		EmitterData->GetScripts(Scripts, false);
 		for (UNiagaraScript* Script : Scripts)
 		{
 			UNiagaraScriptSource* ScriptSource = CastChecked<UNiagaraScriptSource>(Script->GetLatestSource());
@@ -531,10 +528,10 @@ void FNiagaraEmitterViewModel::AddScriptEventHandlers()
 
 void FNiagaraEmitterViewModel::RemoveScriptEventHandlers()
 {
-	if (Emitter.IsValid())
+	if (FVersionedNiagaraEmitterData* EmitterData = EmitterWeakPtr.GetEmitterData())
 	{
 		TArray<UNiagaraScript*> Scripts;
-		Emitter->GetScripts(Scripts, false);
+		EmitterData->GetScripts(Scripts, false);
 		for (UNiagaraScript* Script : Scripts)
 		{
 			FDelegateHandle* OnGraphChangedHandle = ScriptToOnGraphChangedHandleMap.Find(FObjectKey(Script));
@@ -577,10 +574,10 @@ void FNiagaraEmitterViewModel::ScriptParameterStoreChanged(const FNiagaraParamet
 void FNiagaraEmitterViewModel::OnEmitterPropertiesChanged()
 {
 	// Check that these are valid since post edit changed is called on objects even when they've been deleted as a result of undo/redo.
-	if (SharedScriptViewModel.IsValid() && Emitter.IsValid())
+	if (SharedScriptViewModel.IsValid() && EmitterWeakPtr.Emitter.IsValid())
 	{
 		// When the properties change we reset the scripts on the script view model because gpu/cpu or interpolation may have changed.
-		SharedScriptViewModel->SetScripts(Emitter.Get());
+		SharedScriptViewModel->SetScripts(EmitterWeakPtr.ResolveWeakPtr());
 		OnPropertyChangedDelegate.Broadcast();
 	}
 }
diff --git a/Engine/Plugins/FX/Niagara/Source/NiagaraEditor/Private/ViewModels/NiagaraParameterPanelViewModel.cpp b/Engine/Plugins/FX/Niagara/Source/NiagaraEditor/Private/ViewModels/NiagaraParameterPanelViewModel.cpp
--- a/Engine/Plugins/FX/Niagara/Source/NiagaraEditor/Private/ViewModels/NiagaraParameterPanelViewModel.cpp
+++ b/Engine/Plugins/FX/Niagara/Source/NiagaraEditor/Private/ViewModels/NiagaraParameterPanelViewModel.cpp
@@ -83,7 +83,7 @@ TArray<UNiagaraGraph*> FNiagaraSystemToolkitParameterPanelUtilities::GetAllGraph
 			{
 				continue;
 			}
-			UNiagaraGraph* Graph = Cast<UNiagaraScriptSource>(EmitterHandle->GetInstance()->GraphSource)->NodeGraph;
+			UNiagaraGraph* Graph = Cast<UNiagaraScriptSource>(EmitterHandle->GetEmitterData()->GraphSource)->NodeGraph;
 			if (Graph)
 			{
 				OutGraphs.Add(Graph);
@@ -95,7 +95,7 @@ TArray<UNiagaraGraph*> FNiagaraSystemToolkitParameterPanelUtilities::GetAllGraph
 		FNiagaraEmitterHandle* EmitterHandle = SystemViewModel->GetEmitterHandleViewModels()[0]->GetEmitterHandle();
 		if (EmitterHandle != nullptr)
 		{
-			OutGraphs.Add(Cast<UNiagaraScriptSource>(EmitterHandle->GetInstance()->GraphSource)->NodeGraph);
+			OutGraphs.Add(Cast<UNiagaraScriptSource>(EmitterHandle->GetEmitterData()->GraphSource)->NodeGraph);
 		}
 	}
 	return OutGraphs;
@@ -126,7 +126,7 @@ TArray<UNiagaraGraph*> FNiagaraSystemToolkitParameterPanelUtilities::GetEditable
 	}
 	else
 	{
-		EditableGraphs.Add(Cast<UNiagaraScriptSource>(SystemViewModel->GetEmitterHandleViewModels()[0]->GetEmitterHandle()->GetInstance()->GraphSource)->NodeGraph);
+		EditableGraphs.Add(Cast<UNiagaraScriptSource>(SystemViewModel->GetEmitterHandleViewModels()[0]->GetEmitterHandle()->GetEmitterData()->GraphSource)->NodeGraph);
 	}
 	return EditableGraphs;
 }
@@ -157,7 +157,10 @@ FReply FNiagaraSystemToolkitParameterPanelUtilities::CreateDragEventForParameter
 TArray<UNiagaraGraph*> FNiagaraScriptToolkitParameterPanelUtilities::GetEditableGraphs(const TSharedPtr<FNiagaraScriptViewModel>& ScriptViewModel)
 {
 	TArray<UNiagaraGraph*> EditableGraphs;
-	EditableGraphs.Add(ScriptViewModel->GetGraphViewModel()->GetGraph());
+	if (UNiagaraGraph* Graph = ScriptViewModel->GetGraphViewModel()->GetGraph())
+	{
+		EditableGraphs.Add(Graph);
+	}
 	return EditableGraphs;
 }
 
@@ -893,7 +896,7 @@ void FNiagaraSystemToolkitParameterPanelViewModel::Cleanup()
 		{
 			TArray<UNiagaraScript*> EmitterScripts;
 			const bool bCompilableOnly = false;
-			EmitterHandle->GetInstance()->GetScripts(EmitterScripts, bCompilableOnly);
+			EmitterHandle->GetEmitterData()->GetScripts(EmitterScripts, bCompilableOnly);
 			for (UNiagaraScript* EmitterScript : EmitterScripts)
 			{
 				GraphsToRemoveCallbacks.Add(GetGraphFromScript(EmitterScript));
@@ -948,7 +951,7 @@ void FNiagaraSystemToolkitParameterPanelViewModel::Init(const FSystemToolkitUICo
 		{
 			TArray<UNiagaraScript*> EmitterScripts;
 			const bool bCompilableOnly = false;
-			EmitterHandle->GetInstance()->GetScripts(EmitterScripts, bCompilableOnly);
+			EmitterHandle->GetEmitterData()->GetScripts(EmitterScripts, bCompilableOnly);
 			for (UNiagaraScript* EmitterScript : EmitterScripts)
 			{
 				GraphsToAddCallbacks.Add(GetGraphFromScript(EmitterScript));
@@ -1246,7 +1249,7 @@ void FNiagaraSystemToolkitParameterPanelViewModel::RenameParameter(const FNiagar
 		{
 			for (const FNiagaraEmitterHandle* EmitterHandle : GetEditableEmitterHandles())
 			{
-				EmitterHandle->GetInstance()->HandleVariableRenamed(Parameter, FNiagaraVariableBase(Parameter.GetType(), NewName), true);
+				EmitterHandle->GetInstance().Emitter->HandleVariableRenamed(Parameter, FNiagaraVariableBase(Parameter.GetType(), NewName), true, EmitterHandle->GetInstance().Version);
 			}		
 		}
 		else
@@ -1649,8 +1652,8 @@ const TArray<FNiagaraGraphParameterReference> FNiagaraSystemToolkitParameterPane
 
 			bool bIgnoreDisabled = true;
 			FNiagaraParameterMapHistoryBuilder Builder;
-			UNiagaraEmitter* GraphOwningEmitter = Graph->GetTypedOuter<UNiagaraEmitter>();
-			FCompileConstantResolver ConstantResolver = GraphOwningEmitter != nullptr
+			FVersionedNiagaraEmitter GraphOwningEmitter = Graph->GetOwningEmitter();
+			FCompileConstantResolver ConstantResolver = GraphOwningEmitter.Emitter != nullptr
 				? FCompileConstantResolver(GraphOwningEmitter, ENiagaraScriptUsage::Function)
 				: FCompileConstantResolver();
 
@@ -1658,9 +1661,9 @@ const TArray<FNiagaraGraphParameterReference> FNiagaraSystemToolkitParameterPane
 			Builder.ConstantResolver = ConstantResolver;
 			FName StageName;
 			ENiagaraScriptUsage StageUsage = OutputNode->GetUsage();
-			if (StageUsage == ENiagaraScriptUsage::ParticleSimulationStageScript && GraphOwningEmitter)
+			if (StageUsage == ENiagaraScriptUsage::ParticleSimulationStageScript && GraphOwningEmitter.Emitter)
 			{
-				UNiagaraSimulationStageBase* Base = GraphOwningEmitter->GetSimulationStageById(OutputNode->GetUsageId());
+				UNiagaraSimulationStageBase* Base = GraphOwningEmitter.GetEmitterData()->GetSimulationStageById(OutputNode->GetUsageId());
 				if (Base)
 				{
 					StageName = Base->GetStackContextReplacementName();
@@ -1826,8 +1829,8 @@ TArray<FNiagaraParameterPanelItem> FNiagaraSystemToolkitParameterPanelViewModel:
 
 			bool bIgnoreDisabled = true;
 			FNiagaraParameterMapHistoryBuilder Builder;
-			UNiagaraEmitter* GraphOwningEmitter = EditableGraph->GetTypedOuter<UNiagaraEmitter>();
-			FCompileConstantResolver ConstantResolver = GraphOwningEmitter != nullptr
+			FVersionedNiagaraEmitter GraphOwningEmitter = EditableGraph->GetOwningEmitter();
+			FCompileConstantResolver ConstantResolver = GraphOwningEmitter.Emitter != nullptr
 				? FCompileConstantResolver(GraphOwningEmitter, OutputNode->GetUsage())
 				: FCompileConstantResolver();
 				 
@@ -1835,9 +1838,9 @@ TArray<FNiagaraParameterPanelItem> FNiagaraSystemToolkitParameterPanelViewModel:
 			Builder.ConstantResolver = ConstantResolver;
 			FName StageName;
 			ENiagaraScriptUsage StageUsage = OutputNode->GetUsage();
-			if (StageUsage == ENiagaraScriptUsage::ParticleSimulationStageScript && GraphOwningEmitter)
+			if (StageUsage == ENiagaraScriptUsage::ParticleSimulationStageScript && GraphOwningEmitter.Emitter)
 			{
-				UNiagaraSimulationStageBase* Base = GraphOwningEmitter->GetSimulationStageById(OutputNode->GetUsageId());
+				UNiagaraSimulationStageBase* Base = GraphOwningEmitter.GetEmitterData()->GetSimulationStageById(OutputNode->GetUsageId());
 				if (Base)
 				{
 					StageName = Base->GetStackContextReplacementName();
@@ -1967,7 +1970,7 @@ TArray<TWeakObjectPtr<UNiagaraGraph>> FNiagaraSystemToolkitParameterPanelViewMod
 	else
 	{
 		TArray<TWeakObjectPtr<UNiagaraGraph>> EditableEmitterScriptGraphs;
-		EditableEmitterScriptGraphs.Add(static_cast<UNiagaraScriptSource*>(SystemViewModel->GetEmitterHandleViewModels()[0]->GetEmitterHandle()->GetInstance()->GraphSource)->NodeGraph);
+		EditableEmitterScriptGraphs.Add(static_cast<UNiagaraScriptSource*>(SystemViewModel->GetEmitterHandleViewModels()[0]->GetEmitterHandle()->GetEmitterData()->GraphSource)->NodeGraph);
 		return EditableEmitterScriptGraphs;
 	}
 }
@@ -2072,7 +2075,8 @@ void FNiagaraSystemToolkitParameterPanelViewModel::OnParameterRenamedExternally(
 				if (SelectedEmitterHandleIds.Contains(EmitterHandleViewModel->GetId()))
 				{
 					bool bFound = false;
-					UNiagaraGraph* Graph = static_cast<UNiagaraScriptSource*>(EmitterHandleViewModel->GetEmitterHandle()->GetInstance()->GraphSource)->NodeGraph;
+					FVersionedNiagaraEmitter VersionedEmitter = EmitterHandleViewModel->GetEmitterHandle()->GetInstance();
+					UNiagaraGraph* Graph = static_cast<UNiagaraScriptSource*>(VersionedEmitter.GetEmitterData()->GraphSource)->NodeGraph;
 					if (Graph)
 					{
 						const TMap<FNiagaraVariable, FNiagaraGraphParameterReferenceCollection>& RefMap = Graph->GetParameterReferenceMap();
@@ -2105,7 +2109,7 @@ void FNiagaraSystemToolkitParameterPanelViewModel::OnParameterRenamedExternally(
 		const TArray<TSharedRef<FNiagaraEmitterHandleViewModel>>& EmitterHandleViewModels = SystemViewModel->GetEmitterHandleViewModels();
 		for (const TSharedRef<FNiagaraEmitterHandleViewModel>& EmitterHandleViewModel : EmitterHandleViewModels)
 		{
-			UNiagaraGraph* Graph = static_cast<UNiagaraScriptSource*>(EmitterHandleViewModel->GetEmitterHandle()->GetInstance()->GraphSource)->NodeGraph;
+			UNiagaraGraph* Graph = static_cast<UNiagaraScriptSource*>(EmitterHandleViewModel->GetEmitterHandle()->GetEmitterData()->GraphSource)->NodeGraph;
 			if (Graph)
 			{ 
 				Graphs.Add(Graph);
@@ -2155,7 +2159,8 @@ void FNiagaraSystemToolkitParameterPanelViewModel::OnParameterRemovedExternally(
 				if (SelectedEmitterHandleIds.Contains(EmitterHandleViewModel->GetId()))
 				{
 					bool bFound = false;
-					UNiagaraGraph* Graph = static_cast<UNiagaraScriptSource*>(EmitterHandleViewModel->GetEmitterHandle()->GetInstance()->GraphSource)->NodeGraph;
+					FVersionedNiagaraEmitter VersionedEmitter = EmitterHandleViewModel->GetEmitterHandle()->GetInstance();
+					UNiagaraGraph* Graph = static_cast<UNiagaraScriptSource*>(VersionedEmitter.GetEmitterData()->GraphSource)->NodeGraph;
 					if (Graph)
 					{
 						const TMap<FNiagaraVariable, FNiagaraGraphParameterReferenceCollection>& RefMap = Graph->GetParameterReferenceMap();
@@ -2169,7 +2174,7 @@ void FNiagaraSystemToolkitParameterPanelViewModel::OnParameterRemovedExternally(
 
 					if (!bFound)
 					{
-						EmitterHandleViewModel->GetEmitterHandle()->GetInstance()->HandleVariableRemoved(InOldVar, true);
+						VersionedEmitter.Emitter->HandleVariableRemoved(InOldVar, true, VersionedEmitter.Version);
 					}
 				}
 			}
@@ -2185,7 +2190,7 @@ void FNiagaraSystemToolkitParameterPanelViewModel::OnParameterRemovedExternally(
 		const TArray<TSharedRef<FNiagaraEmitterHandleViewModel>>& EmitterHandleViewModels = SystemViewModel->GetEmitterHandleViewModels();
 		for (const TSharedRef<FNiagaraEmitterHandleViewModel>& EmitterHandleViewModel : EmitterHandleViewModels)
 		{
-			UNiagaraGraph* Graph = static_cast<UNiagaraScriptSource*>(EmitterHandleViewModel->GetEmitterHandle()->GetInstance()->GraphSource)->NodeGraph;
+			UNiagaraGraph* Graph = static_cast<UNiagaraScriptSource*>(EmitterHandleViewModel->GetEmitterHandle()->GetEmitterData()->GraphSource)->NodeGraph;
 			if (Graph)
 				Graphs.Add(Graph);
 
@@ -2232,7 +2237,7 @@ void FNiagaraSystemToolkitParameterPanelViewModel::ReconcileOnGraphChangedBindin
 		{
 			TArray<UNiagaraScript*> EmitterScripts;
 			const bool bCompilableOnly = false;
-			EmitterHandle->GetInstance()->GetScripts(EmitterScripts, bCompilableOnly);
+			EmitterHandle->GetEmitterData()->GetScripts(EmitterScripts, bCompilableOnly);
 			for (UNiagaraScript* EmitterScript : EmitterScripts)
 			{
 				GraphsToAddCallbacks.Add(GetGraphFromScript(EmitterScript));
diff --git a/Engine/Plugins/FX/Niagara/Source/NiagaraEditor/Private/ViewModels/NiagaraScratchPadUtilities.cpp b/Engine/Plugins/FX/Niagara/Source/NiagaraEditor/Private/ViewModels/NiagaraScratchPadUtilities.cpp
--- a/Engine/Plugins/FX/Niagara/Source/NiagaraEditor/Private/ViewModels/NiagaraScratchPadUtilities.cpp
+++ b/Engine/Plugins/FX/Niagara/Source/NiagaraEditor/Private/ViewModels/NiagaraScratchPadUtilities.cpp
@@ -21,17 +21,18 @@ void FNiagaraScratchPadUtilities::FixFunctionInputsFromFunctionScriptRename(UNia
 	FunctionCallNode.SuggestName(FString());
 	const FString NewFunctionName = FunctionCallNode.GetFunctionName();
 	UNiagaraSystem* System = FunctionCallNode.GetTypedOuter<UNiagaraSystem>();
-	UNiagaraEmitter* Emitter = FunctionCallNode.GetTypedOuter<UNiagaraEmitter>();
+	FVersionedNiagaraEmitter Emitter = FunctionCallNode.GetNiagaraGraph()->GetOwningEmitter();
 	FNiagaraStackGraphUtilities::RenameReferencingParameters(System, Emitter, FunctionCallNode, OldFunctionName, NewFunctionName);
 }
 
-void FNiagaraScratchPadUtilities::FixExternalScratchPadScriptsForEmitter(UNiagaraSystem& SourceSystem, UNiagaraEmitter& TargetEmitter)
+void FNiagaraScratchPadUtilities::FixExternalScratchPadScriptsForEmitter(UNiagaraSystem& SourceSystem, const FVersionedNiagaraEmitter& TargetEmitter)
 {
-	UNiagaraSystem* TargetSystem = TargetEmitter.GetTypedOuter<UNiagaraSystem>();
-	if (TargetSystem != &SourceSystem)
+	UNiagaraSystem* TargetSystem = TargetEmitter.Emitter->GetTypedOuter<UNiagaraSystem>();
+	FVersionedNiagaraEmitterData* EmitterData = TargetEmitter.GetEmitterData();
+	if (EmitterData && TargetSystem != &SourceSystem)
 	{
-		UNiagaraScriptSource* EmitterScriptSource = Cast<UNiagaraScriptSource>(TargetEmitter.GraphSource);
-		if (ensureMsgf(EmitterScriptSource != nullptr, TEXT("Emitter script source was null. Target Emitter: %s"), *TargetEmitter.GetPathName()))
+		UNiagaraScriptSource* EmitterScriptSource = Cast<UNiagaraScriptSource>(EmitterData->GraphSource);
+		if (ensureMsgf(EmitterScriptSource != nullptr, TEXT("Emitter script source was null. Target Emitter: %s"), *TargetEmitter.Emitter->GetPathName()))
 		{
 			// Find function call nodes which reference external scratch pad scripts.
 			TArray<UNiagaraNodeFunctionCall*> AllFunctionCallNodes;
@@ -54,7 +55,7 @@ void FNiagaraScratchPadUtilities::FixExternalScratchPadScriptsForEmitter(UNiagar
 
 			// Determine the destination for the copies of the external scratch pad scripts.
 			UObject* TargetScratchPadScriptOuter = nullptr;
-			using ScriptArrayType = decltype(TargetEmitter.ScratchPadScripts);
+			using ScriptArrayType = decltype(EmitterData->ScratchPads->Scripts);
 			ScriptArrayType* TargetScratchPadScriptArray = nullptr;
 			if (TargetSystem != nullptr)
 			{
@@ -63,13 +64,13 @@ void FNiagaraScratchPadUtilities::FixExternalScratchPadScriptsForEmitter(UNiagar
 			}
 			else
 			{
-				TargetScratchPadScriptOuter = &TargetEmitter;
-				TargetScratchPadScriptArray = &TargetEmitter.ScratchPadScripts;
+				TargetScratchPadScriptOuter = TargetEmitter.Emitter;
+				TargetScratchPadScriptArray = &EmitterData->ScratchPads->Scripts;
 			}
 
 			// Fix up the external scratch scripts.
 			if (ensureMsgf(TargetScratchPadScriptOuter != nullptr && TargetScratchPadScriptArray != nullptr,
-				TEXT("Failed to find target outer and array for fixing up  external scratch pad scripts.  Target Emitter: %s"), *TargetEmitter.GetPathName()))
+				TEXT("Failed to find target outer and array for fixing up  external scratch pad scripts.  Target Emitter: %s"), *TargetEmitter.Emitter->GetPathName()))
 			{
 				// Collate the function call nodes by the scratch script they call.
 				TMap<UNiagaraScript*, TArray<UNiagaraNodeFunctionCall*>> ExternalScratchPadScriptToFunctionCallNodes;
diff --git a/Engine/Plugins/FX/Niagara/Source/NiagaraEditor/Private/ViewModels/NiagaraScratchPadUtilities.h b/Engine/Plugins/FX/Niagara/Source/NiagaraEditor/Private/ViewModels/NiagaraScratchPadUtilities.h
--- a/Engine/Plugins/FX/Niagara/Source/NiagaraEditor/Private/ViewModels/NiagaraScratchPadUtilities.h
+++ b/Engine/Plugins/FX/Niagara/Source/NiagaraEditor/Private/ViewModels/NiagaraScratchPadUtilities.h
@@ -2,7 +2,7 @@
 
 #pragma once
 
-#include "UObject/NameTypes.h"
+#include "NiagaraTypes.h"
 
 class UNiagaraSystem;
 class UNiagaraEmitter;
@@ -16,5 +16,5 @@ namespace FNiagaraScratchPadUtilities
 
 	/** Fixes function call nodes in an emitter which reference scratch pad scripts from another system which can happen when copying
 		and pasting emitters between systems, and when saving system emitters as assets. */
-	void FixExternalScratchPadScriptsForEmitter(UNiagaraSystem& SourceSystem, UNiagaraEmitter& TargetEmitter);
+	void FixExternalScratchPadScriptsForEmitter(UNiagaraSystem& SourceSystem, const FVersionedNiagaraEmitter& TargetEmitter);
 }
\ No newline at end of file
diff --git a/Engine/Plugins/FX/Niagara/Source/NiagaraEditor/Private/ViewModels/NiagaraScratchPadViewModel.cpp b/Engine/Plugins/FX/Niagara/Source/NiagaraEditor/Private/ViewModels/NiagaraScratchPadViewModel.cpp
--- a/Engine/Plugins/FX/Niagara/Source/NiagaraEditor/Private/ViewModels/NiagaraScratchPadViewModel.cpp
+++ b/Engine/Plugins/FX/Niagara/Source/NiagaraEditor/Private/ViewModels/NiagaraScratchPadViewModel.cpp
@@ -57,20 +57,21 @@ void UNiagaraScratchPadViewModel::Finalize()
 	ActiveScriptViewModel.Reset();
 }
 
-void GetOuterAndTargetScripts(TSharedRef<FNiagaraSystemViewModel> SystemViewModel, UObject*& OutOuter, TArray<UNiagaraScript*>*& OutTargetScripts)
+void GetOuterAndTargetScripts(TSharedRef<FNiagaraSystemViewModel> SystemViewModel, UObject*& OutOuter, TArray<TObjectPtr<UNiagaraScript>>*& OutTargetScripts)
 {
 	if (SystemViewModel->GetEditMode() == ENiagaraSystemViewModelEditMode::SystemAsset)
 	{
 		OutOuter = &SystemViewModel->GetSystem();
-		OutTargetScripts = &ToRawPtrTArrayUnsafe(SystemViewModel->GetSystem().ScratchPadScripts);
+		OutTargetScripts = &SystemViewModel->GetSystem().ScratchPadScripts;
 	}
 	else
 	{
 		if (SystemViewModel->GetSystem().GetEmitterHandles().Num() == 1)
 		{
-			UNiagaraEmitter* TargetEmitter = SystemViewModel->GetSystem().GetEmitterHandles()[0].GetInstance();
-			OutOuter = TargetEmitter;
-			OutTargetScripts = &ToRawPtrTArrayUnsafe(TargetEmitter->ScratchPadScripts);
+			FVersionedNiagaraEmitter TargetEmitter = SystemViewModel->GetSystem().GetEmitterHandles()[0].GetInstance();
+			FVersionedNiagaraEmitterData* EmitterData = TargetEmitter.GetEmitterData();
+			OutOuter = EmitterData->ScratchPads;
+			OutTargetScripts = &EmitterData->ScratchPads->Scripts;
 		}
 		else
 		{
@@ -86,7 +87,7 @@ void UpdateChangeId(TSharedRef<FNiagaraSystemViewModel> SystemViewModel)
 	{
 		if (SystemViewModel->GetSystem().GetEmitterHandles().Num() == 1)
 		{
-			UNiagaraEmitter* TargetEmitter = SystemViewModel->GetSystem().GetEmitterHandles()[0].GetInstance();
+			UNiagaraEmitter* TargetEmitter = SystemViewModel->GetSystem().GetEmitterHandles()[0].GetInstance().Emitter;
 			TargetEmitter->NotifyScratchPadScriptsChanged();
 		}
 	}
@@ -101,7 +102,7 @@ void UNiagaraScratchPadViewModel::RefreshScriptViewModels()
 	bool bViewModelsChanged = false;
 
 	UObject* ScriptOuter;
-	TArray<UNiagaraScript*>* TargetScripts;
+	TArray<TObjectPtr<UNiagaraScript>>* TargetScripts;
 	GetOuterAndTargetScripts(GetSystemViewModel(), ScriptOuter, TargetScripts);
 
 	if (ScriptOuter != nullptr && TargetScripts != nullptr)
@@ -374,7 +375,7 @@ void UNiagaraScratchPadViewModel::DeleteActiveScript()
 		}
 
 		UObject* ScriptOuter;
-		TArray<UNiagaraScript*>* TargetScripts;
+		TArray<TObjectPtr<UNiagaraScript>>* TargetScripts;
 		GetOuterAndTargetScripts(GetSystemViewModel(), ScriptOuter, TargetScripts);
 
 		ScriptOuter->Modify();
@@ -394,7 +395,7 @@ FName GetUniqueScriptName(UObject* Outer, const FString& CandidateName)
 TSharedPtr<FNiagaraScratchPadScriptViewModel> UNiagaraScratchPadViewModel::CreateNewScript(ENiagaraScriptUsage InScriptUsage, ENiagaraScriptUsage InTargetSupportedUsage, FNiagaraTypeDefinition InOutputType)
 {
 	UObject* ScriptOuter;
-	TArray<UNiagaraScript*>* TargetScripts;
+	TArray<TObjectPtr<UNiagaraScript>>* TargetScripts;
 	GetOuterAndTargetScripts(GetSystemViewModel(), ScriptOuter, TargetScripts);
 
 	UNiagaraScript* NewScript = nullptr;
@@ -462,7 +463,7 @@ TSharedPtr<FNiagaraScratchPadScriptViewModel> UNiagaraScratchPadViewModel::Creat
 TSharedPtr<FNiagaraScratchPadScriptViewModel> UNiagaraScratchPadViewModel::CreateNewScriptAsDuplicate(const UNiagaraScript* ScriptToDuplicate)
 {
 	UObject* ScriptOuter;
-	TArray<UNiagaraScript*>* TargetScripts;
+	TArray<TObjectPtr<UNiagaraScript>>* TargetScripts;
 	GetOuterAndTargetScripts(GetSystemViewModel(), ScriptOuter, TargetScripts);
 
 	UNiagaraScript* NewScript = CastChecked<UNiagaraScript>(StaticDuplicateObject(ScriptToDuplicate, ScriptOuter, GetUniqueScriptName(ScriptOuter, *ScriptToDuplicate->GetFName().ToString())));
diff --git a/Engine/Plugins/FX/Niagara/Source/NiagaraEditor/Private/ViewModels/NiagaraScriptStatsViewModel.cpp b/Engine/Plugins/FX/Niagara/Source/NiagaraEditor/Private/ViewModels/NiagaraScriptStatsViewModel.cpp
--- a/Engine/Plugins/FX/Niagara/Source/NiagaraEditor/Private/ViewModels/NiagaraScriptStatsViewModel.cpp
+++ b/Engine/Plugins/FX/Niagara/Source/NiagaraEditor/Private/ViewModels/NiagaraScriptStatsViewModel.cpp
@@ -10,12 +10,12 @@
 #include "ViewModels/NiagaraSystemSelectionViewModel.h"
 
 #include "Async/Async.h"
+#include "Styling/AppStyle.h"
 #include "Framework/MultiBox/MultiBoxBuilder.h"
 #include "Interfaces/ITargetPlatformManagerModule.h"
 #include "Types/SlateEnums.h"
 #include "Widgets/SCompoundWidget.h"
 #include "Widgets/Text/STextBlock.h"
-#include "Widgets/Layout/SSpacer.h"
 #include "Widgets/Layout/SExpandableArea.h"
 #include "Widgets/Layout/SScrollBox.h"
 #include "Widgets/Views/SListView.h"
@@ -108,13 +108,13 @@ namespace NiagaraScriptStatsLocal
 
 	struct FNiagaraScriptStats
 	{
-		FNiagaraScriptStats(const FGuid& InGuid, UNiagaraEmitter* NiagaraEmitter, EShaderPlatform ShaderPlatform)
+		FNiagaraScriptStats(const FGuid& InGuid, FVersionedNiagaraEmitter NiagaraEmitter, EShaderPlatform ShaderPlatform)
 		{
 			EmitterId = InGuid;
-			bIsGPU = NiagaraEmitter->SimTarget == ENiagaraSimTarget::GPUComputeSim;
+			bIsGPU = NiagaraEmitter.GetEmitterData()->SimTarget == ENiagaraSimTarget::GPUComputeSim;
 		}
 
-		void UpdateStatus(EShaderPlatform ShaderPlatform, UNiagaraEmitter* NiagaraEmitter)
+		void UpdateStatus(EShaderPlatform ShaderPlatform, FVersionedNiagaraEmitter NiagaraEmitter)
 		{
 			if ( bIsCompiling )
 			{
@@ -126,7 +126,7 @@ namespace NiagaraScriptStatsLocal
 					if ( ShaderScripts.Num() == 0 )
 					{
 						bool bScriptsValid = true;
-						NiagaraEmitter->ForEachScript(
+						NiagaraEmitter.GetEmitterData()->ForEachScript(
 							[&](UNiagaraScript* NiagaraScript)
 							{
 								if ( !NiagaraScript->GetVMExecutableDataCompilationId().CompilerVersionID.IsValid() || !NiagaraScript->GetVMExecutableDataCompilationId().BaseScriptCompileHash.IsValid() )
@@ -141,7 +141,7 @@ namespace NiagaraScriptStatsLocal
 							return;
 						}
 
-						NiagaraEmitter->ForEachScript(
+						NiagaraEmitter.GetEmitterData()->ForEachScript(
 							[&](UNiagaraScript* NiagaraScript)
 							{
 								TArray<FNiagaraShaderScript*> NewResources;
@@ -209,7 +209,7 @@ namespace NiagaraScriptStatsLocal
 				else
 				{
 					TArray<UNiagaraScript*> NiagaraScripts;
-					NiagaraEmitter->GetScripts(NiagaraScripts);
+					NiagaraEmitter.GetEmitterData()->GetScripts(NiagaraScripts);
 					GetVMScriptStatus(MakeArrayView(NiagaraScripts), bIsCompiling, bHasError, ResultsString);
 				}
 			}
@@ -622,7 +622,7 @@ void FNiagaraScriptStatsViewModel::RefreshView()
 		{
 			if (EmitterHandle.IsValid() && EmitterHandle.GetIsEnabled())
 			{
-				if (UNiagaraEmitter* NiagaraEmitter = EmitterHandle.GetInstance())
+				if (UNiagaraEmitter* NiagaraEmitter = EmitterHandle.GetInstance().Emitter)
 				{
 					RowIDs.Emplace(MakeShareable(new FGuid(EmitterHandle.GetId())));
 				}
diff --git a/Engine/Plugins/FX/Niagara/Source/NiagaraEditor/Private/ViewModels/NiagaraScriptViewModel.cpp b/Engine/Plugins/FX/Niagara/Source/NiagaraEditor/Private/ViewModels/NiagaraScriptViewModel.cpp
--- a/Engine/Plugins/FX/Niagara/Source/NiagaraEditor/Private/ViewModels/NiagaraScriptViewModel.cpp
+++ b/Engine/Plugins/FX/Niagara/Source/NiagaraEditor/Private/ViewModels/NiagaraScriptViewModel.cpp
@@ -231,9 +231,10 @@ void FNiagaraScriptViewModel::SetScripts(UNiagaraScriptSource* InScriptSource, T
 	}
 }
 
-void FNiagaraScriptViewModel::SetScripts(UNiagaraEmitter* InEmitter)
+void FNiagaraScriptViewModel::SetScripts(FVersionedNiagaraEmitter InEmitter)
 {
-	if (InEmitter == nullptr)
+	FVersionedNiagaraEmitterData* EmitterData = InEmitter.GetEmitterData();
+	if (EmitterData == nullptr)
 	{
 		TArray<FVersionedNiagaraScript> EmptyScripts;
 		SetScripts(nullptr, EmptyScripts);
@@ -241,7 +242,7 @@ void FNiagaraScriptViewModel::SetScripts(UNiagaraEmitter* InEmitter)
 	else
 	{
 		TArray<UNiagaraScript*> InScripts;
-		InEmitter->GetScripts(InScripts);
+		EmitterData->GetScripts(InScripts);
 
 		TArray<FVersionedNiagaraScript> EmitterScripts;
 		for (UNiagaraScript* Script : InScripts)
@@ -249,7 +250,7 @@ void FNiagaraScriptViewModel::SetScripts(UNiagaraEmitter* InEmitter)
 			EmitterScripts.AddDefaulted_GetRef().Script = Script;
 		}
 		
-		SetScripts(Cast<UNiagaraScriptSource>(InEmitter->GraphSource), EmitterScripts);
+		SetScripts(Cast<UNiagaraScriptSource>(EmitterData->GraphSource), EmitterScripts);
 	}
 }
 
diff --git a/Engine/Plugins/FX/Niagara/Source/NiagaraEditor/Private/ViewModels/NiagaraStandaloneScriptViewModel.h b/Engine/Plugins/FX/Niagara/Source/NiagaraEditor/Private/ViewModels/NiagaraStandaloneScriptViewModel.h
--- a/Engine/Plugins/FX/Niagara/Source/NiagaraEditor/Private/ViewModels/NiagaraStandaloneScriptViewModel.h
+++ b/Engine/Plugins/FX/Niagara/Source/NiagaraEditor/Private/ViewModels/NiagaraStandaloneScriptViewModel.h
@@ -11,7 +11,7 @@ class UNiagaraScript;
 class UNiagaraScriptVariable;
 
 
-class FNiagaraStandaloneScriptViewModel : public FNiagaraScriptViewModel
+class NIAGARAEDITOR_API FNiagaraStandaloneScriptViewModel : public FNiagaraScriptViewModel
 {
 public:
 	FNiagaraStandaloneScriptViewModel(
diff --git a/Engine/Plugins/FX/Niagara/Source/NiagaraEditor/Private/ViewModels/NiagaraSystemGraphSelectionViewModel.cpp b/Engine/Plugins/FX/Niagara/Source/NiagaraEditor/Private/ViewModels/NiagaraSystemGraphSelectionViewModel.cpp
--- a/Engine/Plugins/FX/Niagara/Source/NiagaraEditor/Private/ViewModels/NiagaraSystemGraphSelectionViewModel.cpp
+++ b/Engine/Plugins/FX/Niagara/Source/NiagaraEditor/Private/ViewModels/NiagaraSystemGraphSelectionViewModel.cpp
@@ -34,7 +34,7 @@ void FNiagaraSystemGraphSelectionViewModel::RefreshSelectedEmitterScriptGraphs()
 		{
 			if (SelectedEmitterHandleIds.Contains(EmitterHandleViewModel->GetId()))
 			{
-				SelectedEmitterScriptGraphs.Add(static_cast<UNiagaraScriptSource*>(EmitterHandleViewModel->GetEmitterHandle()->GetInstance()->GraphSource)->NodeGraph);
+				SelectedEmitterScriptGraphs.Add(static_cast<UNiagaraScriptSource*>(EmitterHandleViewModel->GetEmitterHandle()->GetEmitterData()->GraphSource)->NodeGraph);
 			}
 		}
 	}
diff --git a/Engine/Plugins/FX/Niagara/Source/NiagaraEditor/Private/ViewModels/NiagaraSystemScriptViewModel.cpp b/Engine/Plugins/FX/Niagara/Source/NiagaraEditor/Private/ViewModels/NiagaraSystemScriptViewModel.cpp
--- a/Engine/Plugins/FX/Niagara/Source/NiagaraEditor/Private/ViewModels/NiagaraSystemScriptViewModel.cpp
+++ b/Engine/Plugins/FX/Niagara/Source/NiagaraEditor/Private/ViewModels/NiagaraSystemScriptViewModel.cpp
@@ -4,7 +4,6 @@
 
 #include "NiagaraEmitterHandle.h"
 #include "NiagaraGraph.h"
-#include "NiagaraNodeEmitter.h"
 #include "NiagaraNodeInput.h"
 #include "NiagaraScriptInputCollectionViewModel.h"
 #include "NiagaraSystem.h"
@@ -58,7 +57,7 @@ void FNiagaraSystemScriptViewModel::OnSystemVMCompiled(UNiagaraSystem* InSystem)
 	for (const FNiagaraEmitterHandle& Handle : InSystem->GetEmitterHandles())
 	{
 		int32 NumScripts = SystemScripts.Num();
-		Handle.GetInstance()->GetScripts(SystemScripts, true);
+		Handle.GetEmitterData()->GetScripts(SystemScripts, true);
 		for (; NumScripts < SystemScripts.Num(); NumScripts++)
 		{
 			if (Handle.GetIsEnabled())
@@ -156,7 +155,7 @@ ENiagaraScriptCompileStatus FNiagaraSystemScriptViewModel::GetLatestCompileStatu
 	{
 		if (Handle.GetIsEnabled())
 		{
-			Handle.GetInstance()->GetScripts(SystemScripts, true);
+			Handle.GetEmitterData()->GetScripts(SystemScripts, true);
 		}
 	}
 
diff --git a/Engine/Plugins/FX/Niagara/Source/NiagaraEditor/Private/ViewModels/NiagaraSystemScriptViewModel.h b/Engine/Plugins/FX/Niagara/Source/NiagaraEditor/Private/ViewModels/NiagaraSystemScriptViewModel.h
--- a/Engine/Plugins/FX/Niagara/Source/NiagaraEditor/Private/ViewModels/NiagaraSystemScriptViewModel.h
+++ b/Engine/Plugins/FX/Niagara/Source/NiagaraEditor/Private/ViewModels/NiagaraSystemScriptViewModel.h
@@ -15,11 +15,11 @@ public:
 	DECLARE_MULTICAST_DELEGATE(FOnSystemCompiled)
 
 public:
-	FNiagaraSystemScriptViewModel(bool bInIsForDataProcessingOnly);
+	explicit FNiagaraSystemScriptViewModel(bool bInIsForDataProcessingOnly);
 
 	void Initialize(UNiagaraSystem& InSystem);
 
-	~FNiagaraSystemScriptViewModel();
+	virtual ~FNiagaraSystemScriptViewModel() override;
 
 	FOnSystemCompiled& OnSystemCompiled();
 
diff --git a/Engine/Plugins/FX/Niagara/Source/NiagaraEditor/Private/ViewModels/NiagaraSystemSelectionViewModel.cpp b/Engine/Plugins/FX/Niagara/Source/NiagaraEditor/Private/ViewModels/NiagaraSystemSelectionViewModel.cpp
--- a/Engine/Plugins/FX/Niagara/Source/NiagaraEditor/Private/ViewModels/NiagaraSystemSelectionViewModel.cpp
+++ b/Engine/Plugins/FX/Niagara/Source/NiagaraEditor/Private/ViewModels/NiagaraSystemSelectionViewModel.cpp
@@ -11,7 +11,6 @@
 #include "ViewModels/NiagaraEmitterViewModel.h"
 #include "NiagaraEditorUtilities.h"
 #include "NiagaraSystemEditorData.h"
-#include "NiagaraEmitterHandle.h"
 
 FGuid StackEntryToEmitterGuid(UNiagaraStackEntry* StackEntry)
 {
@@ -458,7 +457,7 @@ void UNiagaraSystemSelectionViewModel::UpdateExternalSelectionState()
 	bSystemIsSelected = false;
 	TArray<FGuid> OldSelectedEmitterHandleIds = SelectedEmitterHandleIds;
 	SelectedEmitterHandleIds.Empty();
-	TArray<UNiagaraEmitter*> SelectedEmitters;
+	TArray<FVersionedNiagaraEmitter> SelectedEmitters;
 	for (const FSelectionEntry& SelectionEntry : SelectionEntries)
 	{
 		if (SelectionEntry.Entry->GetEmitterViewModel().IsValid())
@@ -472,7 +471,7 @@ void UNiagaraSystemSelectionViewModel::UpdateExternalSelectionState()
 	}
 
 	int32 NumRemoved = 0;
-	for (UNiagaraEmitter* SelectedEmitter : SelectedEmitters)
+	for (const FVersionedNiagaraEmitter& SelectedEmitter : SelectedEmitters)
 	{
 		FGuid EmitterId = GetSystemViewModel()->GetEmitterHandleViewModelForEmitter(SelectedEmitter)->GetId();
 		NumRemoved += OldSelectedEmitterHandleIds.Remove(EmitterId);
diff --git a/Engine/Plugins/FX/Niagara/Source/NiagaraEditor/Private/ViewModels/NiagaraSystemViewModel.cpp b/Engine/Plugins/FX/Niagara/Source/NiagaraEditor/Private/ViewModels/NiagaraSystemViewModel.cpp
--- a/Engine/Plugins/FX/Niagara/Source/NiagaraEditor/Private/ViewModels/NiagaraSystemViewModel.cpp
+++ b/Engine/Plugins/FX/Niagara/Source/NiagaraEditor/Private/ViewModels/NiagaraSystemViewModel.cpp
@@ -35,7 +35,6 @@
 #include "EdGraphSchema_NiagaraSystemOverview.h"
 #include "EdGraphUtilities.h"
 #include "ViewModels/NiagaraScratchPadUtilities.h"
-#include "NiagaraParameterDefinitions.h"
 #include "NiagaraEmitterFactoryNew.h"
 
 #include "Editor.h"
@@ -56,7 +55,6 @@
 #include "NiagaraMessageUtilities.h"
 #include "ViewModels/NiagaraOverviewGraphViewModel.h"
 #include "Framework/Application/SlateApplication.h"
-#include "NiagaraScriptVariable.h"
 #include "NiagaraParameterDefinitionsSubscriber.h"
 
 DECLARE_CYCLE_STAT(TEXT("Niagara - SystemViewModel - CompileSystem"), STAT_NiagaraEditor_SystemViewModel_CompileSystem, STATGROUP_NiagaraEditor);
@@ -268,7 +266,7 @@ INiagaraParameterDefinitionsSubscriber* FNiagaraSystemViewModel::GetParameterDef
 	}
 	else /**EditMode == ENiagaraSystemViewModelEditMode::EmitterAsset */
 	{
-		return GetEmitterHandleViewModels()[0]->GetEmitterHandle()->GetInstance();
+		return GetEmitterHandleViewModels()[0]->GetEmitterHandle()->GetInstance().Emitter;
 	}
 }
 
@@ -294,7 +292,7 @@ TSharedPtr<FNiagaraEmitterHandleViewModel> FNiagaraSystemViewModel::GetEmitterHa
 	return TSharedPtr<FNiagaraEmitterHandleViewModel>();
 }
 
-TSharedPtr<FNiagaraEmitterHandleViewModel> FNiagaraSystemViewModel::GetEmitterHandleViewModelForEmitter(UNiagaraEmitter* InEmitter) const
+TSharedPtr<FNiagaraEmitterHandleViewModel> FNiagaraSystemViewModel::GetEmitterHandleViewModelForEmitter(const FVersionedNiagaraEmitter& InEmitter) const
 {
 	for (TSharedRef<FNiagaraEmitterHandleViewModel> EmitterHandleViewModel : EmitterHandleViewModels)
 	{
@@ -389,13 +387,13 @@ TSharedPtr<FNiagaraEmitterHandleViewModel> FNiagaraSystemViewModel::AddEmitterFr
 	UNiagaraEmitter* Emitter = Cast<UNiagaraEmitter>(AssetData.GetAsset());
 	if (Emitter != nullptr)
 	{
-		return AddEmitter(*Emitter);
+		return AddEmitter(*Emitter, Emitter->GetExposedVersion().VersionGuid);
 	}
 
 	return nullptr;
 }
 
-TSharedPtr<FNiagaraEmitterHandleViewModel> FNiagaraSystemViewModel::AddEmitter(UNiagaraEmitter& Emitter)
+TSharedPtr<FNiagaraEmitterHandleViewModel> FNiagaraSystemViewModel::AddEmitter(UNiagaraEmitter& Emitter, FGuid EmitterVersion)
 {
 	// Reset view models before modifying the emitter handle list to prevent accessing deleted data.
 	ResetEmitterHandleViewModelsAndTracks();
@@ -407,7 +405,7 @@ TSharedPtr<FNiagaraEmitterHandleViewModel> FNiagaraSystemViewModel::AddEmitter(U
 		GEditor->BeginTransaction(LOCTEXT("AddEmitter", "Add emitter"));
 	}
 
-	const FGuid NewEmitterHandleId = FNiagaraEditorUtilities::AddEmitterToSystem(GetSystem(), Emitter, EditMode != ENiagaraSystemViewModelEditMode::EmitterDuringMerge);
+	const FGuid NewEmitterHandleId = FNiagaraEditorUtilities::AddEmitterToSystem(GetSystem(), Emitter, EmitterVersion, EditMode != ENiagaraSystemViewModelEditMode::EmitterDuringMerge);
 
 	if (false == bSystemIsPlaceholder && false == bIsForDataProcessingOnly)
 	{
@@ -459,7 +457,12 @@ NIAGARAEDITOR_API TSharedPtr<FNiagaraEmitterHandleViewModel> FNiagaraSystemViewM
 	FName EmptyEmitterName = FNiagaraEditorUtilities::GetUniqueObjectName<UNiagaraEmitter>(GetTransientPackage(), TEXT("Empty"));
 	EmptyEmitter->SetUniqueEmitterName(EmptyEmitterName.ToString());
 	EmptyEmitter->SetFlags(RF_Transactional);
-	return AddEmitter(*EmptyEmitter);
+	return AddEmitter(FVersionedNiagaraEmitter(EmptyEmitter, FGuid()));
+}
+
+TSharedPtr<FNiagaraEmitterHandleViewModel> FNiagaraSystemViewModel::AddEmitter(const FVersionedNiagaraEmitter& VersionedEmitter)
+{
+	return AddEmitter(*VersionedEmitter.Emitter, VersionedEmitter.Version);
 }
 
 void FNiagaraSystemViewModel::DuplicateEmitters(TArray<FEmitterHandleToDuplicate> EmitterHandlesToDuplicate)
@@ -512,7 +515,7 @@ void FNiagaraSystemViewModel::DuplicateEmitters(TArray<FEmitterHandleToDuplicate
 		if (SourceSystem != nullptr && HandleToDuplicate.IsValid())
 		{
 			const FNiagaraEmitterHandle& EmitterHandle = GetSystem().DuplicateEmitterHandle(HandleToDuplicate, FNiagaraUtilities::GetUniqueName(HandleToDuplicate.GetName(), EmitterHandleNames));
-			FNiagaraScratchPadUtilities::FixExternalScratchPadScriptsForEmitter(*SourceSystem, *EmitterHandle.GetInstance());
+			FNiagaraScratchPadUtilities::FixExternalScratchPadScriptsForEmitter(*SourceSystem, EmitterHandle.GetInstance());
 			EmitterHandleNames.Add(EmitterHandle.GetName());
 			if (EmitterHandleToDuplicate.OverviewNode)
 			{
@@ -690,10 +693,7 @@ void FNiagaraSystemViewModel::AddReferencedObjects(FReferenceCollector& Collecto
 void FNiagaraSystemViewModel::PostUndo(bool bSuccess)
 {
 	// Reset system stack and emitter handle view models to prevent accessing invalid data if they were in the undo operation.
-	SystemStackViewModel->Reset();
-	ResetEmitterHandleViewModelsAndTracks();
-	RefreshAll();
-	GetDefault<UEdGraphSchema_NiagaraSystemOverview>()->ForceVisualizationCacheClear();
+	ResetStack();
 	System->InvalidateActiveCompiles();
 	System->RequestCompile(false);
 }
@@ -818,12 +818,12 @@ void FNiagaraSystemViewModel::SetToolkitCommands(const TSharedRef<FUICommandList
 const TArray<FNiagaraStackModuleData>& FNiagaraSystemViewModel::GetStackModuleDataByModuleEntry(UNiagaraStackEntry* ModuleEntry)
 {
 	FGuid EmitterHandleId = FGuid();
-	UNiagaraEmitter* Emitter = nullptr;
+	FVersionedNiagaraEmitter Emitter;
 	if (ModuleEntry->GetEmitterViewModel().IsValid())
 	{
 		Emitter = ModuleEntry->GetEmitterViewModel()->GetEmitter();
 		TSharedPtr<FNiagaraEmitterHandleViewModel> EmitterHandleViewModel = GetEmitterHandleViewModelForEmitter(Emitter);
-		if (ensureMsgf(EmitterHandleViewModel.IsValid(), TEXT("Failed to get the emitter handle view model for emitter %s while getting stack module data."), *Emitter->GetPathName()))
+		if (ensureMsgf(EmitterHandleViewModel.IsValid(), TEXT("Failed to get the emitter handle view model for emitter %s while getting stack module data."), *Emitter.Emitter->GetPathName()))
 		{
 			EmitterHandleId = EmitterHandleViewModel->GetId();
 		}
@@ -841,7 +841,7 @@ const TArray<FNiagaraStackModuleData>& FNiagaraSystemViewModel::GetStackModuleDa
 		return *StackModuleData;
 	}
 
-	UNiagaraEmitter* Emitter = nullptr;
+	FVersionedNiagaraEmitter Emitter;
 	TSharedPtr<FNiagaraEmitterHandleViewModel> EmitterHandleViewModel = GetEmitterHandleViewModelById(EmitterHandleId);
 	if (ensureMsgf(EmitterHandleViewModel.IsValid(), TEXT("Failed to get the emitter handle view model for emitter handle id %s while getting stack module data."), *EmitterHandleId.ToString(EGuidFormats::DigitsWithHyphens)))
 	{
@@ -851,7 +851,7 @@ const TArray<FNiagaraStackModuleData>& FNiagaraSystemViewModel::GetStackModuleDa
 	return BuildAndCacheStackModuleData(EmitterHandleId, Emitter);
 }
 
-const TArray<FNiagaraStackModuleData>& FNiagaraSystemViewModel::BuildAndCacheStackModuleData(FGuid EmitterHandleId, UNiagaraEmitter* Emitter)
+const TArray<FNiagaraStackModuleData>& FNiagaraSystemViewModel::BuildAndCacheStackModuleData(FGuid EmitterHandleId, const FVersionedNiagaraEmitter& Emitter)
 {
 	TArray<FNiagaraStackModuleData>& StackModuleData = GuidToCachedStackModuleData.Add(EmitterHandleId);
 	TArray<UNiagaraScript*> OrderedScripts;
@@ -870,13 +870,13 @@ UNiagaraEditorParametersAdapter* FNiagaraSystemViewModel::GetEditorOnlyParameter
 		return CastChecked<UNiagaraEditorParametersAdapter>(System->GetEditorParameters());
 	}
 	/** else EditMode == ENiagaraSystemViewModelEditMode::EmitterAsset */
-	UNiagaraEmitter* Emitter = GetEmitterHandleViewModels()[0]->GetEmitterHandle()->GetInstance();
-	return CastChecked<UNiagaraEditorParametersAdapter>(Emitter->GetEditorParameters());
+	FVersionedNiagaraEmitterData* EmitterData = GetEmitterHandleViewModels()[0]->GetEmitterHandle()->GetEmitterData();
+	return CastChecked<UNiagaraEditorParametersAdapter>(EmitterData->GetEditorParameters());
 }
 
 void FNiagaraSystemViewModel::GetOrderedScriptsForEmitterHandleId(FGuid EmitterHandleId, TArray<UNiagaraScript*>& OutScripts)
 {
-	UNiagaraEmitter* Emitter = nullptr;
+	FVersionedNiagaraEmitter Emitter;
 	TSharedPtr<FNiagaraEmitterHandleViewModel> EmitterHandleViewModel = GetEmitterHandleViewModelById(EmitterHandleId);
 	if (EmitterHandleViewModel.IsValid())
 	{
@@ -885,17 +885,18 @@ void FNiagaraSystemViewModel::GetOrderedScriptsForEmitterHandleId(FGuid EmitterH
 	GetOrderedScriptsForEmitter(Emitter, OutScripts);
 }
 
-void FNiagaraSystemViewModel::GetOrderedScriptsForEmitter(UNiagaraEmitter* Emitter, TArray<UNiagaraScript*>& OutScripts)
+void FNiagaraSystemViewModel::GetOrderedScriptsForEmitter(const FVersionedNiagaraEmitter& Emitter, TArray<UNiagaraScript*>& OutScripts)
 {
 	OutScripts.Add(System->GetSystemSpawnScript());
 	OutScripts.Add(System->GetSystemUpdateScript());
-	if (Emitter != nullptr)
+	FVersionedNiagaraEmitterData* EmitterData = Emitter.GetEmitterData();
+	if (EmitterData != nullptr)
 	{
-		OutScripts.Add(Emitter->EmitterSpawnScriptProps.Script);
-		OutScripts.Add(Emitter->EmitterUpdateScriptProps.Script);
-		OutScripts.Add(Emitter->SpawnScriptProps.Script);
-		OutScripts.Add(Emitter->UpdateScriptProps.Script);
-		for (UNiagaraSimulationStageBase* SimulationStage : Emitter->GetSimulationStages())
+		OutScripts.Add(EmitterData->EmitterSpawnScriptProps.Script);
+		OutScripts.Add(EmitterData->EmitterUpdateScriptProps.Script);
+		OutScripts.Add(EmitterData->SpawnScriptProps.Script);
+		OutScripts.Add(EmitterData->UpdateScriptProps.Script);
+		for (UNiagaraSimulationStageBase* SimulationStage : EmitterData->GetSimulationStages())
 		{
 			if (SimulationStage->bEnabled)
 			{
@@ -979,12 +980,12 @@ void FNiagaraSystemViewModel::SendLastCompileMessageJobs() const
 	const TArray<FNiagaraEmitterHandle> EmitterHandles = GetSystem().GetEmitterHandles();
 	for (const FNiagaraEmitterHandle& Handle : EmitterHandles)
 	{
-		const UNiagaraEmitter* EmitterInSystem = Handle.GetInstance();
+		FVersionedNiagaraEmitter EmitterInSystem = Handle.GetInstance();
 		TArray<UNiagaraScript*> EmitterScripts;
-		EmitterInSystem->GetScripts(EmitterScripts, false);
+		EmitterInSystem.GetEmitterData()->GetScripts(EmitterScripts, false);
 		for (UNiagaraScript* EmitterScript : EmitterScripts)
 		{
-			ScriptsToGetCompileEventsFrom.Add(FNiagaraScriptAndOwningScriptNameString(EmitterScript, EmitterInSystem->GetUniqueEmitterName()));
+			ScriptsToGetCompileEventsFrom.Add(FNiagaraScriptAndOwningScriptNameString(EmitterScript, EmitterInSystem.Emitter->GetUniqueEmitterName()));
 		}
 	}
 
@@ -1134,7 +1135,7 @@ void FNiagaraSystemViewModel::TickCompileStatus()
 					{
 						if (EmitterHandle.GetIsEnabled())
 						{
-							EmitterHandle.GetInstance()->GetScripts(ScriptsToCheckForStatus, true);
+							EmitterHandle.GetEmitterData()->GetScripts(ScriptsToCheckForStatus, true);
 						}
 					}
 				}
@@ -1189,6 +1190,14 @@ void FNiagaraSystemViewModel::RefreshAll()
 	bPendingAssetMessagesChanged = true;
 }
 
+void FNiagaraSystemViewModel::ResetStack()
+{
+	SystemStackViewModel->Reset();
+	ResetEmitterHandleViewModelsAndTracks();
+	RefreshAll();
+	GetDefault<UEdGraphSchema_NiagaraSystemOverview>()->ForceVisualizationCacheClear();
+}
+
 void FNiagaraSystemViewModel::NotifyDataObjectChanged(TArray<UObject*> ChangedObjects, ENiagaraDataObjectChange ChangeType)
 {
 	if (ChangedObjects.Num() == 1 && ChangedObjects[0]->IsA<UNiagaraSystem>())
@@ -1216,11 +1225,16 @@ void FNiagaraSystemViewModel::NotifyDataObjectChanged(TArray<UObject*> ChangedOb
 
 			if(UNiagaraEmitter* Emitter = Cast<UNiagaraEmitter>(ChangedObject))
 			{
-				UNiagaraStackViewModel* StackViewModel = GetEmitterHandleViewModelForEmitter(Emitter)->GetEmitterStackViewModel();
-
-				if(ensure(StackViewModel))
+				for (TSharedRef<FNiagaraEmitterHandleViewModel> EmitterHandleViewModel : EmitterHandleViewModels)
 				{
-					StackViewModel->Refresh();
+					if (EmitterHandleViewModel->GetEmitterViewModel()->GetEmitter().Emitter == Emitter)
+					{
+						UNiagaraStackViewModel* StackViewModel = EmitterHandleViewModel->GetEmitterStackViewModel();
+						if(ensure(StackViewModel))
+						{
+							StackViewModel->Refresh();
+						}
+					}
 				}
 			}
 		}
@@ -1660,12 +1674,12 @@ void GetCompiledScriptsAndEmitterNameFromInputNode(UNiagaraNode& StackNode, UNia
 		{
 			const FNiagaraEmitterHandle* OwningEmitterHandle = OwningSystem.GetEmitterHandles().FindByPredicate([&StackNode](const FNiagaraEmitterHandle& EmitterHandle)
 			{
-				return CastChecked<UNiagaraScriptSource>(EmitterHandle.GetInstance()->GraphSource)->NodeGraph == StackNode.GetNiagaraGraph();
+				return CastChecked<UNiagaraScriptSource>(EmitterHandle.GetEmitterData()->GraphSource)->NodeGraph == StackNode.GetNiagaraGraph();
 			});
 
 			if(OwningEmitterHandle != nullptr)
 			{
-				OutEmitterName = OwningEmitterHandle->GetInstance()->GetUniqueEmitterName();
+				OutEmitterName = OwningEmitterHandle->GetInstance().Emitter->GetUniqueEmitterName();
 				switch (OutputNode->GetUsage())
 				{
 				case ENiagaraScriptUsage::EmitterSpawnScript:
@@ -1678,10 +1692,11 @@ void GetCompiledScriptsAndEmitterNameFromInputNode(UNiagaraNode& StackNode, UNia
 				case ENiagaraScriptUsage::ParticleUpdateScript:
 				case ENiagaraScriptUsage::ParticleEventScript:
 				case ENiagaraScriptUsage::ParticleSimulationStageScript:
-					OutCompiledScripts.Add(OwningEmitterHandle->GetInstance()->GetScript(OutputNode->GetUsage(), OutputNode->GetUsageId()));
-					if (OwningEmitterHandle->GetInstance()->SimTarget == ENiagaraSimTarget::GPUComputeSim)
+					//TODO extract emitterdata into one var
+					OutCompiledScripts.Add(OwningEmitterHandle->GetEmitterData()->GetScript(OutputNode->GetUsage(), OutputNode->GetUsageId()));
+					if (OwningEmitterHandle->GetEmitterData()->SimTarget == ENiagaraSimTarget::GPUComputeSim)
 					{
-						OutCompiledScripts.Add(OwningEmitterHandle->GetInstance()->GetScript(ENiagaraScriptUsage::ParticleGPUComputeScript, FGuid()));
+						OutCompiledScripts.Add(OwningEmitterHandle->GetEmitterData()->GetScript(ENiagaraScriptUsage::ParticleGPUComputeScript, FGuid()));
 					}
 					break;
 				}
@@ -2278,14 +2293,25 @@ void FNiagaraSystemViewModel::ClearEmitterStats()
 #if STATS
 	for (const FNiagaraEmitterHandle& Handle : System->GetEmitterHandles())
 	{
-		if (UNiagaraEmitter* Emitter = Handle.GetInstance())
+		if (FVersionedNiagaraEmitterData* EmitterData = Handle.GetEmitterData())
 		{
-			Emitter->GetStatData().ClearStatCaptures();
+			EmitterData->GetStatData().ClearStatCaptures();
 		}
 	}
 #endif
 }
 
+// this switches the viewed emitter version in the editor, for the version selector of the referenced parent version see SNiagaraOverviewStackNode::SwitchToVersion
+bool FNiagaraSystemViewModel::ChangeEmitterVersion(const FVersionedNiagaraEmitter& Emitter, const FGuid& NewVersion)
+{
+	if (System->ChangeEmitterVersion(Emitter, NewVersion))
+	{
+		ResetStack();
+		return true;
+	}
+	return false;
+}
+
 void FNiagaraSystemViewModel::AddSystemEventHandlers()
 {
 	if (System != nullptr)
@@ -2470,9 +2496,9 @@ void FNiagaraSystemViewModel::RefreshAssetMessages()
 		Scripts.Add(ViewedSystem.GetSystemUpdateScript());
 		for (const FNiagaraEmitterHandle& EmitterHandle : ViewedSystem.GetEmitterHandles())
 		{
-			UNiagaraEmitter& Emitter = *EmitterHandle.GetInstance();
-			LastFunctionCallNodeObjectKeys.Append(PublishMessages(Emitter, &UNiagaraEmitter::RemoveMessageDelegateable));
-			Emitter.GetScripts(Scripts, false);
+			FVersionedNiagaraEmitter VersionedEmitter = EmitterHandle.GetInstance();
+			LastFunctionCallNodeObjectKeys.Append(PublishMessages(*VersionedEmitter.Emitter, &UNiagaraEmitter::RemoveMessageDelegateable));
+			VersionedEmitter.GetEmitterData()->GetScripts(Scripts, false);
 		}
 		LastFunctionCallNodeObjectKeys.Append(PublishScriptMessages(Scripts));
 
@@ -2482,11 +2508,11 @@ void FNiagaraSystemViewModel::RefreshAssetMessages()
 		const TArray<FNiagaraEmitterHandle>& EmitterHandles = ViewedSystem.GetEmitterHandles();
 		if (ensureMsgf(EmitterHandles.Num() == 1, TEXT("There was not exactly 1 Emitter Handle for the SystemViewModel in Emitter edit mode!")))
 		{
-			UNiagaraEmitter& ViewedEmitter = *(EmitterHandles[0].GetInstance());
-			LastFunctionCallNodeObjectKeys.Append(PublishMessages(ViewedEmitter, &UNiagaraEmitter::RemoveMessageDelegateable));
+			FVersionedNiagaraEmitter VersionedEmitter = EmitterHandles[0].GetInstance();
+			LastFunctionCallNodeObjectKeys.Append(PublishMessages(*VersionedEmitter.Emitter, &UNiagaraEmitter::RemoveMessageDelegateable));
 
 			TArray<UNiagaraScript*> Scripts;
-			ViewedEmitter.GetScripts(Scripts, false);
+			VersionedEmitter.GetEmitterData()->GetScripts(Scripts, false);
 			LastFunctionCallNodeObjectKeys.Append(PublishScriptMessages(Scripts));
 		}
 	}
diff --git a/Engine/Plugins/FX/Niagara/Source/NiagaraEditor/Private/ViewModels/Stack/NiagaraStackEmitterPropertiesGroup.cpp b/Engine/Plugins/FX/Niagara/Source/NiagaraEditor/Private/ViewModels/Stack/NiagaraStackEmitterPropertiesGroup.cpp
--- a/Engine/Plugins/FX/Niagara/Source/NiagaraEditor/Private/ViewModels/Stack/NiagaraStackEmitterPropertiesGroup.cpp
+++ b/Engine/Plugins/FX/Niagara/Source/NiagaraEditor/Private/ViewModels/Stack/NiagaraStackEmitterPropertiesGroup.cpp
@@ -112,7 +112,8 @@ public:
 			return;
 		}
 
-		UNiagaraEmitter* Emitter = EmitterViewModel->GetEmitter();
+		FVersionedNiagaraEmitter VersionedEmitter = EmitterViewModel->GetEmitter();
+		UNiagaraEmitter* Emitter = VersionedEmitter.Emitter;
 		UNiagaraScriptSource* Source = EmitterViewModel->GetSharedScriptViewModel()->GetGraphViewModel()->GetScriptSource();
 		UNiagaraGraph* Graph = EmitterViewModel->GetSharedScriptViewModel()->GetGraphViewModel()->GetGraph();
 
@@ -139,7 +140,7 @@ public:
 			EventScriptProperties.Script->SetUsage(ENiagaraScriptUsage::ParticleEventScript);
 			EventScriptProperties.Script->SetUsageId(FGuid::NewGuid());
 			EventScriptProperties.Script->SetLatestSource(Source);
-			Emitter->AddEventHandler(EventScriptProperties);
+			Emitter->AddEventHandler(EventScriptProperties, VersionedEmitter.Version);
 			FNiagaraStackGraphUtilities::ResetGraphForOutput(*Graph, ENiagaraScriptUsage::ParticleEventScript, EventScriptProperties.Script->GetUsageId());
 			AddedEventHandlerId = EventScriptProperties.Script->GetUsageId();
 		}
@@ -153,7 +154,7 @@ public:
 			AddedSimulationStage->Script->SetUsage(ENiagaraScriptUsage::ParticleSimulationStageScript);
 			AddedSimulationStage->Script->SetUsageId(AddedSimulationStage->GetMergeId());
 			AddedSimulationStage->Script->SetLatestSource(Source);
-			Emitter->AddSimulationStage(AddedSimulationStage);
+			Emitter->AddSimulationStage(AddedSimulationStage, VersionedEmitter.Version);
 			FNiagaraStackGraphUtilities::ResetGraphForOutput(*Graph, ENiagaraScriptUsage::ParticleSimulationStageScript, AddedSimulationStage->Script->GetUsageId());
 		}
 
@@ -161,7 +162,7 @@ public:
 		// TODO: Move the logic for managing additional scripts into the emitter view model or script view model.
 		TWeakPtr<FNiagaraEmitterInstance, ESPMode::ThreadSafe> Simulation = EmitterViewModel->GetSimulation();
 		EmitterViewModel->Reset();
-		EmitterViewModel->Initialize(Emitter, Simulation);
+		EmitterViewModel->Initialize(VersionedEmitter, Simulation);
 
 		OnItemAdded.ExecuteIfBound(AddedEventHandlerId, AddedSimulationStage);
 	}
@@ -194,12 +195,12 @@ const FSlateBrush* UNiagaraStackEmitterPropertiesGroup::GetSecondaryIconBrush()
 {
 	if (IsFinalized() == false && GetEmitterViewModel().IsValid())
 	{
-		UNiagaraEmitter* Emitter = GetEmitterViewModel()->GetEmitter();
-		if (Emitter->SimTarget == ENiagaraSimTarget::CPUSim)
+		FVersionedNiagaraEmitterData* EmitterData = GetEmitterViewModel()->GetEmitter().GetEmitterData();
+		if (EmitterData->SimTarget == ENiagaraSimTarget::CPUSim)
 		{
 			return FNiagaraEditorStyle::Get().GetBrush("NiagaraEditor.Stack.CPUIcon");
 		}
-		if (Emitter->SimTarget == ENiagaraSimTarget::GPUComputeSim)
+		if (EmitterData->SimTarget == ENiagaraSimTarget::GPUComputeSim)
 		{
 			return FNiagaraEditorStyle::Get().GetBrush("NiagaraEditor.Stack.GPUIcon");
 		}
diff --git a/Engine/Plugins/FX/Niagara/Source/NiagaraEditor/Private/ViewModels/Stack/NiagaraStackEmitterSettingsGroup.cpp b/Engine/Plugins/FX/Niagara/Source/NiagaraEditor/Private/ViewModels/Stack/NiagaraStackEmitterSettingsGroup.cpp
--- a/Engine/Plugins/FX/Niagara/Source/NiagaraEditor/Private/ViewModels/Stack/NiagaraStackEmitterSettingsGroup.cpp
+++ b/Engine/Plugins/FX/Niagara/Source/NiagaraEditor/Private/ViewModels/Stack/NiagaraStackEmitterSettingsGroup.cpp
@@ -7,14 +7,15 @@
 #include "ViewModels/NiagaraEmitterViewModel.h"
 #include "NiagaraEmitter.h"
 #include "NiagaraEmitterEditorData.h"
-#include "NiagaraStackEditorData.h"
 #include "ViewModels/Stack/NiagaraStackGraphUtilities.h"
 #include "NiagaraScriptMergeManager.h"
 #include "NiagaraEmitterDetailsCustomization.h"
 #include "NiagaraEditorStyle.h"
+#include "NiagaraEditorUtilities.h"
 #include "NiagaraSystem.h"
-#include "NiagaraEditorStyle.h"
+#include "ScopedTransaction.h"
 #include "ViewModels/NiagaraEmitterHandleViewModel.h"
+#include "ViewModels/NiagaraSystemSelectionViewModel.h"
 #include "ViewModels/Stack/NiagaraStackViewModel.h"
 #include "Styling/AppStyle.h"
 #include "IDetailTreeNode.h"
@@ -24,15 +25,15 @@
 void UNiagaraStackEmitterPropertiesItem::Initialize(FRequiredEntryData InRequiredEntryData)
 {
 	Super::Initialize(InRequiredEntryData, TEXT("EmitterProperties"));
-	Emitter = GetEmitterViewModel()->GetEmitter();
-	Emitter->OnPropertiesChanged().AddUObject(this, &UNiagaraStackEmitterPropertiesItem::EmitterPropertiesChanged);
+	EmitterWeakPtr = GetEmitterViewModel()->GetEmitter().ToWeakPtr();
+	EmitterWeakPtr.Emitter->OnPropertiesChanged().AddUObject(this, &UNiagaraStackEmitterPropertiesItem::EmitterPropertiesChanged);
 }
 
 void UNiagaraStackEmitterPropertiesItem::FinalizeInternal()
 {
-	if (Emitter.IsValid())
+	if (EmitterWeakPtr.IsValid())
 	{
-		Emitter->OnPropertiesChanged().RemoveAll(this);
+		EmitterWeakPtr.Emitter->OnPropertiesChanged().RemoveAll(this);
 	}
 	Super::FinalizeInternal();
 }
@@ -49,13 +50,17 @@ FText UNiagaraStackEmitterPropertiesItem::GetTooltipText() const
 
 bool UNiagaraStackEmitterPropertiesItem::TestCanResetToBaseWithMessage(FText& OutCanResetToBaseMessage) const
 {
-	if (bCanResetToBaseCache.IsSet() == false)
+	if (GetEmitterViewModel().IsValid() == false)
+	{
+		bCanResetToBaseCache = false;
+	}
+	else if (bCanResetToBaseCache.IsSet() == false)
 	{
-		const UNiagaraEmitter* BaseEmitter = GetEmitterViewModel()->GetEmitter()->GetParent();
-		if (BaseEmitter != nullptr && Emitter != BaseEmitter)
+		FVersionedNiagaraEmitter BaseEmitter = GetEmitterViewModel()->GetParentEmitter();
+		if (BaseEmitter.GetEmitterData() != nullptr && EmitterWeakPtr.Emitter != BaseEmitter.Emitter)
 		{
 			TSharedRef<FNiagaraScriptMergeManager> MergeManager = FNiagaraScriptMergeManager::Get();
-			bCanResetToBaseCache = MergeManager->IsEmitterEditablePropertySetDifferentFromBase(*Emitter.Get(), *BaseEmitter);
+			bCanResetToBaseCache = MergeManager->IsEmitterEditablePropertySetDifferentFromBase(EmitterWeakPtr.ResolveWeakPtr(), BaseEmitter);
 		}
 		else
 		{
@@ -79,21 +84,24 @@ void UNiagaraStackEmitterPropertiesItem::ResetToBase()
 	FText Unused;
 	if (TestCanResetToBaseWithMessage(Unused))
 	{
-		const UNiagaraEmitter* BaseEmitter = GetEmitterViewModel()->GetEmitter()->GetParent();
+		FVersionedNiagaraEmitter BaseEmitter = GetEmitterViewModel()->GetParentEmitter();
 		TSharedRef<FNiagaraScriptMergeManager> MergeManager = FNiagaraScriptMergeManager::Get();
-		MergeManager->ResetEmitterEditablePropertySetToBase(*Emitter, *BaseEmitter);
+		MergeManager->ResetEmitterEditablePropertySetToBase(EmitterWeakPtr.ResolveWeakPtr(), BaseEmitter);
 	}
 }
 
 const FSlateBrush* UNiagaraStackEmitterPropertiesItem::GetIconBrush() const
 {
-	if (Emitter->SimTarget == ENiagaraSimTarget::CPUSim)
+	if (FVersionedNiagaraEmitterData* EmitterData = EmitterWeakPtr.GetEmitterData())
 	{
-		return FNiagaraEditorStyle::Get().GetBrush("NiagaraEditor.Stack.CPUIcon");
-	}
-	if (Emitter->SimTarget == ENiagaraSimTarget::GPUComputeSim)
-	{
-		return FNiagaraEditorStyle::Get().GetBrush("NiagaraEditor.Stack.GPUIcon");
+		if (EmitterData->SimTarget == ENiagaraSimTarget::CPUSim)
+		{
+			return FNiagaraEditorStyle::Get().GetBrush("NiagaraEditor.Stack.CPUIcon");
+		}
+		if (EmitterData->SimTarget == ENiagaraSimTarget::GPUComputeSim)
+		{
+			return FNiagaraEditorStyle::Get().GetBrush("NiagaraEditor.Stack.GPUIcon");
+		}
 	}
 	return FAppStyle::GetBrush("NoBrush");
 }
@@ -105,7 +113,7 @@ void UNiagaraStackEmitterPropertiesItem::RefreshChildrenInternal(const TArray<UN
 		EmitterObject = NewObject<UNiagaraStackObject>(this);
 		FRequiredEntryData RequiredEntryData(GetSystemViewModel(), GetEmitterViewModel(), FExecutionCategoryNames::Emitter, NAME_None, GetStackEditorData());
 		bool bIsTopLevelObject = true;
-		EmitterObject->Initialize(RequiredEntryData, Emitter.Get(), bIsTopLevelObject, GetStackEditorDataKey());
+		EmitterObject->Initialize(RequiredEntryData, EmitterWeakPtr.Emitter.Get(), bIsTopLevelObject, GetStackEditorDataKey());
 		EmitterObject->RegisterInstancedCustomPropertyLayout(UNiagaraEmitter::StaticClass(), FOnGetDetailCustomizationInstance::CreateStatic(&FNiagaraEmitterDetails::MakeInstance));
 	}
 
@@ -115,11 +123,21 @@ void UNiagaraStackEmitterPropertiesItem::RefreshChildrenInternal(const TArray<UN
 	RefreshIssues(NewIssues);
 }
 
+UNiagaraStackEntry::FStackIssueFixDelegate UNiagaraStackEmitterPropertiesItem::GetUpgradeVersionFix()
+{
+	return FStackIssueFixDelegate::CreateLambda([=]()
+	{
+		FGuid NewVersion = GetEmitterViewModel()->GetParentEmitter().Emitter->GetExposedVersion().VersionGuid;
+		FNiagaraEditorUtilities::SwitchParentEmitterVersion(GetEmitterViewModel().ToSharedRef(), GetSystemViewModel(), NewVersion);
+	});
+}
 
 void UNiagaraStackEmitterPropertiesItem::RefreshIssues(TArray<FStackIssue>& NewIssues)
 {
-	UNiagaraEmitter* ActualEmitter = GetEmitterViewModel()->GetEmitter();
-	if (ActualEmitter && ActualEmitter->SimTarget == ENiagaraSimTarget::GPUComputeSim && ActualEmitter->CalculateBoundsMode == ENiagaraEmitterCalculateBoundMode::Dynamic)
+	FNiagaraStackGraphUtilities::CheckForDeprecatedEmitterVersion(GetEmitterViewModel(), GetStackEditorDataKey(), GetUpgradeVersionFix(), NewIssues);
+	
+	FVersionedNiagaraEmitterData* ActualEmitterData = GetEmitterViewModel()->GetEmitter().GetEmitterData();
+	if (ActualEmitterData && ActualEmitterData->SimTarget == ENiagaraSimTarget::GPUComputeSim && ActualEmitterData->CalculateBoundsMode == ENiagaraEmitterCalculateBoundMode::Dynamic)
 	{
 		bool bAddError = true;
 		
@@ -153,9 +171,9 @@ void UNiagaraStackEmitterPropertiesItem::EmitterPropertiesChanged()
 		bCanResetToBaseCache.Reset();
 		RefreshChildren();
 
-		if (Emitter.IsValid())
+		if (EmitterWeakPtr.IsValid())
 		{
-			GetSystemViewModel()->GetEmitterHandleViewModelForEmitter(Emitter.Get()).Get()->GetEmitterStackViewModel()->RequestValidationUpdate();
+			GetSystemViewModel()->GetEmitterHandleViewModelForEmitter(EmitterWeakPtr.ResolveWeakPtr()).Get()->GetEmitterStackViewModel()->RequestValidationUpdate();
 		}
 	}
 }
@@ -163,7 +181,7 @@ void UNiagaraStackEmitterPropertiesItem::EmitterPropertiesChanged()
 void UNiagaraStackEmitterSummaryItem::Initialize(FRequiredEntryData InRequiredEntryData)
 {
 	Super::Initialize(InRequiredEntryData, TEXT("EmitterParameters"));
-	Emitter = GetEmitterViewModel()->GetEmitter();
+	Emitter = GetEmitterViewModel()->GetEmitter().ToWeakPtr();
 }
 
 FText UNiagaraStackEmitterSummaryItem::GetDisplayName() const
@@ -206,7 +224,7 @@ void UNiagaraStackEmitterSummaryItem::RefreshChildrenInternal(const TArray<UNiag
 	{
 		FilteredObject = NewObject<UNiagaraStackSummaryViewObject>(this);
 		FRequiredEntryData RequiredEntryData(GetSystemViewModel(), GetEmitterViewModel(), FExecutionCategoryNames::Emitter, NAME_None, GetStackEditorData());
-		FilteredObject->Initialize(RequiredEntryData, Emitter.Get(), GetStackEditorDataKey());
+		FilteredObject->Initialize(RequiredEntryData, Emitter, GetStackEditorDataKey());
 	}
 
 	NewChildren.Add(FilteredObject);
diff --git a/Engine/Plugins/FX/Niagara/Source/NiagaraEditor/Private/ViewModels/Stack/NiagaraStackEntry.cpp b/Engine/Plugins/FX/Niagara/Source/NiagaraEditor/Private/ViewModels/Stack/NiagaraStackEntry.cpp
--- a/Engine/Plugins/FX/Niagara/Source/NiagaraEditor/Private/ViewModels/Stack/NiagaraStackEntry.cpp
+++ b/Engine/Plugins/FX/Niagara/Source/NiagaraEditor/Private/ViewModels/Stack/NiagaraStackEntry.cpp
@@ -7,7 +7,6 @@
 #include "ViewModels/Stack/NiagaraStackErrorItem.h"
 #include "ViewModels/NiagaraSystemViewModel.h"
 #include "ViewModels/NiagaraEmitterViewModel.h"
-#include "ViewModels/Stack/NiagaraStackItemGroup.h"
 #include "NiagaraStackEditorData.h"
 #include "NiagaraScriptMergeManager.h"
 #include "Misc/SecureHash.h"
@@ -557,7 +556,7 @@ bool UNiagaraStackEntry::HasBaseEmitter() const
 {
 	if (bHasBaseEmitterCache.IsSet() == false)
 	{
-		const UNiagaraEmitter* BaseEmitter = GetEmitterViewModel().IsValid() ? GetEmitterViewModel()->GetParentEmitter() : nullptr;
+		const UNiagaraEmitter* BaseEmitter = GetEmitterViewModel().IsValid() ? GetEmitterViewModel()->GetParentEmitter().Emitter : nullptr;
 		bHasBaseEmitterCache = BaseEmitter != nullptr;
 	}
 	return bHasBaseEmitterCache.GetValue();
diff --git a/Engine/Plugins/FX/Niagara/Source/NiagaraEditor/Private/ViewModels/Stack/NiagaraStackEventScriptItemGroup.cpp b/Engine/Plugins/FX/Niagara/Source/NiagaraEditor/Private/ViewModels/Stack/NiagaraStackEventScriptItemGroup.cpp
--- a/Engine/Plugins/FX/Niagara/Source/NiagaraEditor/Private/ViewModels/Stack/NiagaraStackEventScriptItemGroup.cpp
+++ b/Engine/Plugins/FX/Niagara/Source/NiagaraEditor/Private/ViewModels/Stack/NiagaraStackEventScriptItemGroup.cpp
@@ -13,9 +13,7 @@
 #include "EdGraphSchema_Niagara.h"
 #include "ViewModels/Stack/NiagaraStackErrorItem.h"
 #include "NiagaraScriptSource.h"
-#include "ViewModels/Stack/NiagaraStackGraphUtilities.h"
 #include "NiagaraScriptMergeManager.h"
-#include "NiagaraStackEditorData.h"
 #include "ViewModels/NiagaraSystemViewModel.h"
 #include "Customizations/NiagaraEventScriptPropertiesCustomization.h"
 
@@ -25,6 +23,17 @@
 
 #define LOCTEXT_NAMESPACE "UNiagaraStackEventScriptItemGroup"
 
+void UNiagaraStackEventWrapper::PostEditChangeProperty(FPropertyChangedEvent& PropertyChangedEvent)
+{
+	UObject::PostEditChangeProperty(PropertyChangedEvent);
+
+	if (FVersionedNiagaraEmitterData* EmitterData = EmitterWeakPtr.GetEmitterData())
+	{
+		EmitterData->EventHandlerScriptProps = EventHandlerScriptProps;
+		EmitterWeakPtr.Emitter->PostEditChangeVersionedProperty(PropertyChangedEvent, EmitterWeakPtr.Version);
+	}
+}
+
 void UNiagaraStackEventHandlerPropertiesItem::Initialize(FRequiredEntryData InRequiredEntryData, FGuid InEventScriptUsageId)
 {
 	FString EventStackEditorDataKey = FString::Printf(TEXT("Event-%s-Properties"), *InEventScriptUsageId.ToString(EGuidFormats::DigitsWithHyphens));
@@ -32,15 +41,15 @@ void UNiagaraStackEventHandlerPropertiesItem::Initialize(FRequiredEntryData InRe
 
 	EventScriptUsageId = InEventScriptUsageId;
 
-	Emitter = GetEmitterViewModel()->GetEmitter();
-	Emitter->OnPropertiesChanged().AddUObject(this, &UNiagaraStackEventHandlerPropertiesItem::EventHandlerPropertiesChanged);
+	EmitterWeakPtr = GetEmitterViewModel()->GetEmitter().ToWeakPtr();
+	EmitterWeakPtr.Emitter->OnPropertiesChanged().AddUObject(this, &UNiagaraStackEventHandlerPropertiesItem::EventHandlerPropertiesChanged);
 }
 
 void UNiagaraStackEventHandlerPropertiesItem::FinalizeInternal()
 {
-	if (Emitter.IsValid())
+	if (EmitterWeakPtr.IsValid())
 	{
-		Emitter->OnPropertiesChanged().RemoveAll(this);
+		EmitterWeakPtr.Emitter->OnPropertiesChanged().RemoveAll(this);
 	}
 	Super::FinalizeInternal();
 }
@@ -56,11 +65,11 @@ bool UNiagaraStackEventHandlerPropertiesItem::TestCanResetToBaseWithMessage(FTex
 	{
 		if (HasBaseEventHandler())
 		{
-			const UNiagaraEmitter* BaseEmitter = GetEmitterViewModel()->GetEmitter()->GetParent();
-			if (BaseEmitter != nullptr && Emitter != BaseEmitter)
+			FVersionedNiagaraEmitter BaseEmitter = GetEmitterViewModel()->GetParentEmitter();
+			if (BaseEmitter.Emitter != nullptr && EmitterWeakPtr.Emitter != BaseEmitter.Emitter)
 			{
 				TSharedRef<FNiagaraScriptMergeManager> MergeManager = FNiagaraScriptMergeManager::Get();
-				bCanResetToBaseCache = MergeManager->IsEventHandlerPropertySetDifferentFromBase(*Emitter.Get(), *BaseEmitter, EventScriptUsageId);
+				bCanResetToBaseCache = MergeManager->IsEventHandlerPropertySetDifferentFromBase(EmitterWeakPtr.ResolveWeakPtr(), BaseEmitter, EventScriptUsageId);
 			}
 			else
 			{
@@ -87,25 +96,29 @@ bool UNiagaraStackEventHandlerPropertiesItem::TestCanResetToBaseWithMessage(FTex
 void UNiagaraStackEventHandlerPropertiesItem::ResetToBase()
 {
 	FText Unused;
-	if (TestCanResetToBaseWithMessage(Unused))
+	if (TestCanResetToBaseWithMessage(Unused) && EmitterWeakPtr.IsValid())
 	{
-		const UNiagaraEmitter* BaseEmitter = GetEmitterViewModel()->GetEmitter()->GetParent();
+		FVersionedNiagaraEmitter BaseEmitter = GetEmitterViewModel()->GetParentEmitter();
 		TSharedRef<FNiagaraScriptMergeManager> MergeManager = FNiagaraScriptMergeManager::Get();
-		MergeManager->ResetEventHandlerPropertySetToBase(*Emitter, *BaseEmitter, EventScriptUsageId);
+		MergeManager->ResetEventHandlerPropertySetToBase(EmitterWeakPtr.ResolveWeakPtr(), BaseEmitter, EventScriptUsageId);
 		RefreshChildren();
 	}
 }
 
 void UNiagaraStackEventHandlerPropertiesItem::RefreshChildrenInternal(const TArray<UNiagaraStackEntry*>& CurrentChildren, TArray<UNiagaraStackEntry*>& NewChildren, TArray<FStackIssue>& NewIssues)
 {
-	if (EmitterObject == nullptr)
+	if (EmitterObject == nullptr && EmitterWeakPtr.GetEmitterData())
 	{
+		EventWrapper = NewObject<UNiagaraStackEventWrapper>(this, NAME_None, RF_Transactional);
+		EventWrapper->EventHandlerScriptProps = EmitterWeakPtr.GetEmitterData()->EventHandlerScriptProps;
+		EventWrapper->EmitterWeakPtr = EmitterWeakPtr;
+
 		EmitterObject = NewObject<UNiagaraStackObject>(this);
 		bool bIsTopLevelObject = true;
-		EmitterObject->Initialize(CreateDefaultChildRequiredData(), Emitter.Get(), bIsTopLevelObject, GetStackEditorDataKey());
+		EmitterObject->Initialize(CreateDefaultChildRequiredData(), EventWrapper, bIsTopLevelObject, GetStackEditorDataKey());
 		EmitterObject->RegisterInstancedCustomPropertyTypeLayout(FNiagaraEventScriptProperties::StaticStruct()->GetFName(), 
 			FOnGetPropertyTypeCustomizationInstance::CreateStatic(&FNiagaraEventScriptPropertiesCustomization::MakeInstance, 
-			TWeakObjectPtr<UNiagaraSystem>(&GetSystemViewModel()->GetSystem()), TWeakObjectPtr<UNiagaraEmitter>(GetEmitterViewModel()->GetEmitter())));
+			TWeakObjectPtr<UNiagaraSystem>(&GetSystemViewModel()->GetSystem()), GetEmitterViewModel()->GetEmitter().ToWeakPtr()));
 		EmitterObject->SetOnSelectRootNodes(UNiagaraStackObject::FOnSelectRootNodes::CreateUObject(this, &UNiagaraStackEventHandlerPropertiesItem::SelectEmitterStackObjectRootTreeNodes));
 	}
 
@@ -183,11 +196,11 @@ bool UNiagaraStackEventHandlerPropertiesItem::HasBaseEventHandler() const
 {
 	if (bHasBaseEventHandlerCache.IsSet() == false)
 	{
-		const UNiagaraEmitter* BaseEmitter = GetEmitterViewModel()->GetEmitter()->GetParent();
-		if (BaseEmitter != nullptr && Emitter != BaseEmitter)
+		FVersionedNiagaraEmitter BaseEmitter = GetEmitterViewModel()->GetParentEmitter();
+		if (BaseEmitter.Emitter != nullptr && EmitterWeakPtr.Emitter.Get() != BaseEmitter.Emitter)
 		{
 			TSharedRef<FNiagaraScriptMergeManager> MergeManager = FNiagaraScriptMergeManager::Get();
-			bHasBaseEventHandlerCache = MergeManager->HasBaseEventHandler(*BaseEmitter, EventScriptUsageId);
+			bHasBaseEventHandlerCache = MergeManager->HasBaseEventHandler(BaseEmitter, EventScriptUsageId);
 		}
 		else
 		{
@@ -214,9 +227,9 @@ void UNiagaraStackEventScriptItemGroup::RefreshChildrenInternal(const TArray<UNi
 {
 	bHasBaseEventHandlerCache.Reset();
 
-	UNiagaraEmitter* Emitter = GetEmitterViewModel()->GetEmitter();
+	FVersionedNiagaraEmitterData* EmitterData = GetEmitterViewModel()->GetEmitter().GetEmitterData();
 
-	const FNiagaraEventScriptProperties* EventScriptProperties = Emitter->GetEventHandlers().FindByPredicate(
+	const FNiagaraEventScriptProperties* EventScriptProperties = EmitterData->GetEventHandlers().FindByPredicate(
 		[=](const FNiagaraEventScriptProperties& InEventScriptProperties) { return InEventScriptProperties.Script->GetUsageId() == GetScriptUsageId(); });
 
 	if (EventScriptProperties != nullptr)
@@ -257,8 +270,9 @@ void UNiagaraStackEventScriptItemGroup::Delete()
 	TSharedPtr<FNiagaraScriptViewModel> ScriptViewModelPinned = ScriptViewModel.Pin();
 	checkf(ScriptViewModelPinned.IsValid(), TEXT("Can not delete when the script view model has been deleted."));
 
-	UNiagaraEmitter* Emitter = GetEmitterViewModel()->GetEmitter();
-	UNiagaraScriptSource* Source = Cast<UNiagaraScriptSource>(Emitter->GraphSource);
+	FVersionedNiagaraEmitter VersionedEmitter = GetEmitterViewModel()->GetEmitter();
+	FVersionedNiagaraEmitterData* EmitterData = VersionedEmitter.GetEmitterData();
+	UNiagaraScriptSource* Source = Cast<UNiagaraScriptSource>(EmitterData->GraphSource);
 
 	if (!Source || !Source->NodeGraph)
 	{
@@ -268,10 +282,10 @@ void UNiagaraStackEventScriptItemGroup::Delete()
 	//Need to tear down existing systems now.
 	FNiagaraSystemUpdateContext UpdateCtx;
 	UpdateCtx.SetDestroyOnAdd(true);
-	UpdateCtx.Add(Emitter, true);
+	UpdateCtx.Add(VersionedEmitter, true);
 
 	FScopedTransaction Transaction(FText::Format(LOCTEXT("DeleteEventHandler", "Deleted {0}"), GetDisplayName()));
-	Emitter->Modify();
+	VersionedEmitter.Emitter->Modify();
 	Source->NodeGraph->Modify();
 	TArray<UNiagaraNode*> EventIndexNodes;
 	Source->NodeGraph->BuildTraversal(EventIndexNodes, GetScriptUsage(), GetScriptUsageId());
@@ -281,7 +295,7 @@ void UNiagaraStackEventScriptItemGroup::Delete()
 	}
 	
 	// First, remove the event handler script properties object.
-	Emitter->RemoveEventHandlerByUsageId(GetScriptUsageId());
+	VersionedEmitter.Emitter->RemoveEventHandlerByUsageId(GetScriptUsageId(), VersionedEmitter.Version);
 	
 	// Now remove all graph nodes associated with the event script index.
 	for (UNiagaraNode* Node : EventIndexNodes)
@@ -291,7 +305,7 @@ void UNiagaraStackEventScriptItemGroup::Delete()
 
 	// Set the emitter here to that the internal state of the view model is updated.
 	// TODO: Move the logic for managing event handlers into the emitter view model or script view model.
-	ScriptViewModelPinned->SetScripts(Emitter);
+	ScriptViewModelPinned->SetScripts(VersionedEmitter);
 	
 	OnModifiedEventHandlersDelegate.ExecuteIfBound();
 }
@@ -310,8 +324,8 @@ bool UNiagaraStackEventScriptItemGroup::HasBaseEventHandler() const
 {
 	if (bHasBaseEventHandlerCache.IsSet() == false)
 	{
-		const UNiagaraEmitter* BaseEmitter = GetEmitterViewModel()->GetEmitter()->GetParent();
-		bHasBaseEventHandlerCache = BaseEmitter != nullptr && FNiagaraScriptMergeManager::Get()->HasBaseEventHandler(*BaseEmitter, GetScriptUsageId());
+		FVersionedNiagaraEmitter BaseEmitter = GetEmitterViewModel()->GetParentEmitter();
+		bHasBaseEventHandlerCache = BaseEmitter.Emitter != nullptr && FNiagaraScriptMergeManager::Get()->HasBaseEventHandler(BaseEmitter, GetScriptUsageId());
 	}
 	return bHasBaseEventHandlerCache.GetValue();
 }
diff --git a/Engine/Plugins/FX/Niagara/Source/NiagaraEditor/Private/ViewModels/Stack/NiagaraStackFilteredObject.cpp b/Engine/Plugins/FX/Niagara/Source/NiagaraEditor/Private/ViewModels/Stack/NiagaraStackFilteredObject.cpp
--- a/Engine/Plugins/FX/Niagara/Source/NiagaraEditor/Private/ViewModels/Stack/NiagaraStackFilteredObject.cpp
+++ b/Engine/Plugins/FX/Niagara/Source/NiagaraEditor/Private/ViewModels/Stack/NiagaraStackFilteredObject.cpp
@@ -4,21 +4,15 @@
 
 #include "EdGraphSchema_Niagara.h"
 #include "NiagaraClipboard.h"
-#include "NiagaraDataInterface.h"
 #include "NiagaraGraph.h"
 #include "NiagaraNodeFunctionCall.h"
 #include "NiagaraNodeOutput.h"
-#include "NiagaraNodeParameterMapSet.h"
-#include "ScopedTransaction.h"
-#include "EdGraph/EdGraphPin.h"
 #include "ViewModels/NiagaraEmitterViewModel.h"
 #include "ViewModels/NiagaraSystemViewModel.h"
 #include "ViewModels/NiagaraScriptViewModel.h"
 #include "ViewModels/NiagaraScriptGraphViewModel.h"
 #include "ViewModels/Stack/NiagaraStackFunctionInput.h"
 #include "ViewModels/Stack/NiagaraStackGraphUtilities.h"
-#include "ViewModels/Stack/NiagaraStackInputCategory.h"
-#include "ViewModels/Stack/NiagaraStackModuleItem.h"
 #include "ViewModels/Stack/NiagaraStackFunctionInputCollection.h"
 
 
@@ -26,28 +20,27 @@
 
 
 UNiagaraStackFilteredObject::UNiagaraStackFilteredObject()
-	: Emitter(nullptr)
 {
 }
 
 void UNiagaraStackFilteredObject::Initialize(FRequiredEntryData InRequiredEntryData, FString InOwningStackItemEditorDataKey)
 {
-	checkf(Emitter == nullptr, TEXT("Can only initialize once."));
 	FString ObjectStackEditorDataKey = FString::Printf(TEXT("%s-FilteredView"), *InOwningStackItemEditorDataKey);
 	Super::Initialize(InRequiredEntryData, InOwningStackItemEditorDataKey, ObjectStackEditorDataKey);
 
-	Emitter = GetEmitterViewModel()->GetEmitter();
-	if (Emitter->GetEditorData())
+	VersionedEmitter = GetEmitterViewModel()->GetEmitter();
+	FVersionedNiagaraEmitterData* EmitterData = VersionedEmitter.GetEmitterData();
+	if (EmitterData && EmitterData->GetEditorData())
 	{
-		Cast<UNiagaraEmitterEditorData>(Emitter->GetEditorData())->OnSummaryViewStateChanged().AddUObject(this, &UNiagaraStackFilteredObject::OnViewStateChanged);
+		Cast<UNiagaraEmitterEditorData>(EmitterData->GetEditorData())->OnSummaryViewStateChanged().AddUObject(this, &UNiagaraStackFilteredObject::OnViewStateChanged);
 	}
 }
 
 void UNiagaraStackFilteredObject::FinalizeInternal()
 {
-	if (Emitter->GetEditorData())
+	if (UNiagaraEditorDataBase* EditorData = VersionedEmitter.GetEmitterData()->GetEditorData())
 	{
-		Cast<UNiagaraEmitterEditorData>(Emitter->GetEditorData())->OnSummaryViewStateChanged().RemoveAll(this);
+		Cast<UNiagaraEmitterEditorData>(EditorData)->OnSummaryViewStateChanged().RemoveAll(this);
 	}
 	
 	Super::FinalizeInternal();
@@ -94,9 +87,6 @@ void UNiagaraStackFilteredObject::RefreshChildrenInternal(const TArray<UNiagaraS
 	TSharedPtr<FNiagaraScriptViewModel> ScriptViewModelPinned = ViewModel->GetSharedScriptViewModel();
 	checkf(ScriptViewModelPinned.IsValid(), TEXT("Can not refresh children when the script view model has been deleted."));
 
-	ENiagaraScriptUsage ScriptUsage = ENiagaraScriptUsage::EmitterUpdateScript;
-	FGuid ScriptUsageId = FGuid();
-
 	AppendEmitterCategory(ScriptViewModelPinned, ENiagaraScriptUsage::EmitterSpawnScript, FGuid(), NewChildren, CurrentChildren, NewIssues);
  	AppendEmitterCategory(ScriptViewModelPinned, ENiagaraScriptUsage::EmitterUpdateScript, FGuid(), NewChildren, CurrentChildren, NewIssues);
 	AppendEmitterCategory(ScriptViewModelPinned, ENiagaraScriptUsage::ParticleSpawnScript, FGuid(), NewChildren, CurrentChildren, NewIssues);
@@ -116,7 +106,6 @@ void UNiagaraStackFilteredObject::AppendEmitterCategory(TSharedPtr<FNiagaraScrip
 		TArray<UNiagaraNodeFunctionCall*> ModuleNodes;
 		FNiagaraStackGraphUtilities::GetOrderedModuleNodes(*MatchingOutputNode, ModuleNodes);
 
-		bool bIsFirst = true;
 		for (UNiagaraNodeFunctionCall* ModuleNode : ModuleNodes)
 		{
 			if (ModuleNode && (ModuleNode->HasValidScriptAndGraph() || ModuleNode->Signature.IsValid()))
