Added versioning for Niagara emitters

diff --git a/Engine/Plugins/FX/Niagara/Source/NiagaraEditor/Private/ViewModels/Stack/NiagaraStackFunctionInput.cpp b/Engine/Plugins/FX/Niagara/Source/NiagaraEditor/Private/ViewModels/Stack/NiagaraStackFunctionInput.cpp
--- a/Engine/Plugins/FX/Niagara/Source/NiagaraEditor/Private/ViewModels/Stack/NiagaraStackFunctionInput.cpp
+++ b/Engine/Plugins/FX/Niagara/Source/NiagaraEditor/Private/ViewModels/Stack/NiagaraStackFunctionInput.cpp
@@ -128,7 +128,7 @@ void UNiagaraStackFunctionInput::Initialize(
 	OwningAssignmentNode = Cast<UNiagaraNodeAssignment>(OwningFunctionCallNode.Get());
 
 	UNiagaraSystem& ParentSystem = GetSystemViewModel()->GetSystem();
-	UNiagaraEmitter* ParentEmitter = GetEmitterViewModel().IsValid() ? GetEmitterViewModel()->GetEmitter() : nullptr;
+	FVersionedNiagaraEmitter ParentEmitter = GetEmitterViewModel().IsValid() ? GetEmitterViewModel()->GetEmitter() : FVersionedNiagaraEmitter();
 
 	FNiagaraStackGraphUtilities::FindAffectedScripts(&ParentSystem, ParentEmitter, *OwningModuleNode.Get(), AffectedScripts);
 
@@ -171,7 +171,7 @@ void UNiagaraStackFunctionInput::Initialize(
 	FCompileConstantResolver ConstantResolver = GetEmitterViewModel().IsValid()
 		? FCompileConstantResolver(GetEmitterViewModel()->GetEmitter(), SourceScript->GetUsage())
 		: FCompileConstantResolver(&GetSystemViewModel()->GetSystem(), SourceScript->GetUsage());
-	FString UniqueEmitterName = GetEmitterViewModel().IsValid() ? GetEmitterViewModel()->GetEmitter()->GetUniqueEmitterName() : FString();
+	FString UniqueEmitterName = GetEmitterViewModel().IsValid() ? GetEmitterViewModel()->GetEmitter().Emitter->GetUniqueEmitterName() : FString();
 	EditCondition.Initialize(SourceScript.Get(), AffectedScriptsNotWeak, ConstantResolver, UniqueEmitterName, OwningFunctionCallNode.Get());
 	VisibleCondition.Initialize(SourceScript.Get(), AffectedScriptsNotWeak, ConstantResolver, UniqueEmitterName, OwningFunctionCallNode.Get());
 
@@ -795,11 +795,11 @@ void UNiagaraStackFunctionInput::RefreshChildrenInternal(const TArray<UNiagaraSt
 FString UNiagaraStackFunctionInput::ResolveDisplayNameArgument(const FString& InArg) const
 {
 	//~ Begin helper functions
-	auto GetMaterialsFromEmitter = [](UNiagaraEmitter* InEmitter, FNiagaraEmitterInstance* InEmitterInstance)->TArray<UMaterial*> {
+	auto GetMaterialsFromEmitter = [](const FVersionedNiagaraEmitter& InEmitter, const FNiagaraEmitterInstance* InEmitterInstance)->TArray<UMaterial*> {
 		TArray<UMaterial*> ResultMaterials;
-		if (InEmitter != nullptr)
+		if (InEmitter.Emitter != nullptr)
 		{
-			for (UNiagaraRendererProperties* RenderProperties : InEmitter->GetRenderers())
+			for (UNiagaraRendererProperties* RenderProperties : InEmitter.GetEmitterData()->GetRenderers())
 			{
 				TArray<UMaterialInterface*> UsedMaterialInteraces;
 				RenderProperties->GetUsedMaterials(InEmitterInstance, UsedMaterialInteraces);
@@ -1015,7 +1015,7 @@ void UNiagaraStackFunctionInput::RefreshValues()
 		if (InputValues.Mode == EValueMode::Data)
 		{
 			FGuid EmitterHandleId = GetEmitterViewModel().IsValid()
-				? FNiagaraEditorUtilities::GetEmitterHandleForEmitter(GetSystemViewModel()->GetSystem(), *GetEmitterViewModel()->GetEmitter())->GetId()
+				? FNiagaraEditorUtilities::GetEmitterHandleForEmitter(GetSystemViewModel()->GetSystem(), GetEmitterViewModel()->GetEmitter())->GetId()
 				: FGuid();
 			PlaceholderDataInterfaceHandle = GetSystemViewModel()->GetPlaceholderDataInterfaceManager()->GetPlaceholderDataInterface(EmitterHandleId, *OwningFunctionCallNode.Get(), InputParameterHandle);
 			if (PlaceholderDataInterfaceHandle.IsValid())
@@ -1032,7 +1032,7 @@ void UNiagaraStackFunctionInput::RefreshValues()
 		{
 			// If the input it a data interface but hasn't been edited yet, we need to provide a placeholder data interface to edit.
 			FGuid EmitterHandleId = GetEmitterViewModel().IsValid()
-				? FNiagaraEditorUtilities::GetEmitterHandleForEmitter(GetSystemViewModel()->GetSystem(), *GetEmitterViewModel()->GetEmitter())->GetId()
+				? FNiagaraEditorUtilities::GetEmitterHandleForEmitter(GetSystemViewModel()->GetSystem(), GetEmitterViewModel()->GetEmitter())->GetId()
 				: FGuid();
 			PlaceholderDataInterfaceHandle = GetSystemViewModel()->GetPlaceholderDataInterfaceManager()->GetOrCreatePlaceholderDataInterface(EmitterHandleId, *OwningFunctionCallNode.Get(), InputParameterHandle, InputType.GetClass());
 			InputValues.Mode = EValueMode::Data;
@@ -1262,7 +1262,7 @@ void UNiagaraStackFunctionInput::SetLinkedValueHandle(const FNiagaraParameterHan
 	RefreshValues();
 }
 
-bool UsageRunsBefore(ENiagaraScriptUsage UsageA, ENiagaraScriptUsage UsageB, bool bCheckInterpSpawn , UNiagaraEmitter* InEmitter, FGuid UsageAID, FGuid UsageBID)
+bool UsageRunsBefore(ENiagaraScriptUsage UsageA, ENiagaraScriptUsage UsageB, bool bCheckInterpSpawn, FVersionedNiagaraEmitter InEmitter, FGuid UsageAID, FGuid UsageBID)
 {
 	static TArray<ENiagaraScriptUsage> UsagesOrderedByExecution
 	{
@@ -1290,11 +1290,12 @@ bool UsageRunsBefore(ENiagaraScriptUsage UsageA, ENiagaraScriptUsage UsageB, boo
 		UsagesOrderedByExecution.Find(UsageB, IndexB);
 	}
 
-	if (IndexA == IndexB && UsageA == ENiagaraScriptUsage::ParticleSimulationStageScript && InEmitter)
+	if (IndexA == IndexB && UsageA == ENiagaraScriptUsage::ParticleSimulationStageScript && InEmitter.Emitter)
 	{
-		const TArray<UNiagaraSimulationStageBase*> SimStages = InEmitter->GetSimulationStages();
-		UNiagaraSimulationStageBase* StageA = InEmitter->GetSimulationStageById(UsageAID);
-		UNiagaraSimulationStageBase* StageB = InEmitter->GetSimulationStageById(UsageBID);
+		FVersionedNiagaraEmitterData* EmitterData = InEmitter.GetEmitterData();
+		const TArray<UNiagaraSimulationStageBase*> SimStages = EmitterData->GetSimulationStages();
+		UNiagaraSimulationStageBase* StageA = EmitterData->GetSimulationStageById(UsageAID);
+		UNiagaraSimulationStageBase* StageB = EmitterData->GetSimulationStageById(UsageBID);
 
 		int32 SimIndexA = SimStages.IndexOfByKey(StageA);
 		int32 SimIndexB = SimStages.IndexOfByKey(StageB);
@@ -1385,7 +1386,7 @@ void UNiagaraStackFunctionInput::GetAvailableParameterHandles(TArray<FNiagaraPar
 	}
 
 	// gather variables from the parameter map history
-	UNiagaraEmitter* Emitter = nullptr;
+	FVersionedNiagaraEmitter Emitter;
 	TArray<UNiagaraNodeOutput*> AllOutputNodes;
 	if (GetEmitterViewModel().IsValid())
 	{
@@ -1405,7 +1406,7 @@ void UNiagaraStackFunctionInput::GetAvailableParameterHandles(TArray<FNiagaraPar
 		bool bSpawnScript = false;
 		if (CurrentOutputNode != nullptr && CurrentOutputNode->GetUsage() == ENiagaraScriptUsage::ParticleEventScript)
 		{
-			for (const FNiagaraEventScriptProperties &EventHandlerProps : GetEmitterViewModel()->GetEmitter()->GetEventHandlers())
+			for (const FNiagaraEventScriptProperties &EventHandlerProps : GetEmitterViewModel()->GetEmitter().GetEmitterData()->GetEventHandlers())
 			{
 				if (EventHandlerProps.Script->GetUsageId() == CurrentOutputNode->ScriptTypeId)
 				{
@@ -1414,7 +1415,7 @@ void UNiagaraStackFunctionInput::GetAvailableParameterHandles(TArray<FNiagaraPar
 				}
 			}
 		}
-		bool bInterpolatedSpawn = GetEmitterViewModel().IsValid() && GetEmitterViewModel()->GetEmitter()->bInterpolatedSpawning;
+		bool bInterpolatedSpawn = GetEmitterViewModel().IsValid() && GetEmitterViewModel()->GetEmitter().GetEmitterData()->bInterpolatedSpawning;
 		bool bCheckInterpSpawn = bInterpolatedSpawn || !bSpawnScript;
 		if (OutputNode == CurrentOutputNode || (CurrentOutputNode != nullptr && UsageRunsBefore(OutputNode->GetUsage(), CurrentOutputNode->GetUsage(), bCheckInterpSpawn, Emitter, OutputNode->GetUsageId(), CurrentOutputNode->GetUsageId())) || (CurrentOutputNode != nullptr && IsSpawnUsage(CurrentOutputNode->GetUsage())))
 		{
@@ -1904,10 +1905,10 @@ bool UNiagaraStackFunctionInput::CanResetToBase() const
 				{
 					TSharedPtr<FNiagaraEmitterViewModel> EmitterViewModelPtr = GetEmitterViewModel();
 					ensureMsgf(EmitterViewModelPtr.IsValid(), TEXT("ViewMode is nullptr and it never should be"));
-					const UNiagaraEmitter* BaseEmitter = EmitterViewModelPtr.IsValid() ? EmitterViewModelPtr->GetParentEmitter() : nullptr;
-					bCanResetToBaseCache = BaseEmitter != nullptr && MergeManager->IsModuleInputDifferentFromBase(
-						*EmitterViewModelPtr->GetEmitter(),
-						*BaseEmitter,
+					FVersionedNiagaraEmitter BaseEmitter = EmitterViewModelPtr.IsValid() ? EmitterViewModelPtr->GetParentEmitter() : FVersionedNiagaraEmitter();
+					bCanResetToBaseCache = BaseEmitter.Emitter != nullptr && MergeManager->IsModuleInputDifferentFromBase(
+						EmitterViewModelPtr->GetEmitter(),
+						BaseEmitter,
 						OutputNode->GetUsage(),
 						OutputNode->GetUsageId(),
 						OwningModuleNode->NodeGuid,
@@ -1934,13 +1935,13 @@ void UNiagaraStackFunctionInput::ResetToBase()
 	{
 		TSharedRef<FNiagaraScriptMergeManager> MergeManager = FNiagaraScriptMergeManager::Get();
 
-		const UNiagaraEmitter* BaseEmitter = GetEmitterViewModel()->GetEmitter()->GetParent();
+		FVersionedNiagaraEmitter BaseEmitter = GetEmitterViewModel()->GetEmitter().GetEmitterData()->GetParent();
 		UNiagaraNodeOutput* OutputNode = FNiagaraStackGraphUtilities::GetEmitterOutputNodeForStackNode(*OwningFunctionCallNode.Get());
 
 		FScopedTransaction ScopedTransaction(LOCTEXT("ResetInputToBaseTransaction", "Reset this input to match the parent emitter."));
 		FNiagaraScriptMergeManager::FApplyDiffResults Results = MergeManager->ResetModuleInputToBase(
-			*GetEmitterViewModel()->GetEmitter(),
-			*BaseEmitter,
+			GetEmitterViewModel()->GetEmitter(),
+			BaseEmitter,
 			OutputNode->GetUsage(),
 			OutputNode->GetUsageId(),
 			OwningModuleNode->NodeGuid,
@@ -1951,7 +1952,7 @@ void UNiagaraStackFunctionInput::ResetToBase()
 			// If resetting to the base succeeded, an unknown number of rapid iteration parameters may have been added.  To fix
 			// this copy all of the owning scripts rapid iteration parameters to all other affected scripts.
 			// TODO: Either the merge should take care of this directly, or at least provide more information about what changed.
-			UNiagaraScript* OwningScript = GetEmitterViewModel()->GetEmitter()->GetScript(OutputNode->GetUsage(), OutputNode->GetUsageId());
+			UNiagaraScript* OwningScript = GetEmitterViewModel()->GetEmitter().GetEmitterData()->GetScript(OutputNode->GetUsage(), OutputNode->GetUsageId());
 			TArray<FNiagaraVariable> OwningScriptRapidIterationParameters;
 			OwningScript->RapidIterationParameters.GetParameters(OwningScriptRapidIterationParameters);
 			if (OwningScriptRapidIterationParameters.Num() > 0)
@@ -1978,7 +1979,7 @@ void UNiagaraStackFunctionInput::ResetToBase()
 FNiagaraVariable UNiagaraStackFunctionInput::CreateRapidIterationVariable(const FName& InName)
 {
 	UNiagaraNodeOutput* OutputNode = FNiagaraStackGraphUtilities::GetEmitterOutputNodeForStackNode(*OwningModuleNode.Get());
-	FString UniqueEmitterName = GetEmitterViewModel().IsValid() ? GetEmitterViewModel()->GetEmitter()->GetUniqueEmitterName() : FString();
+	FString UniqueEmitterName = GetEmitterViewModel().IsValid() ? GetEmitterViewModel()->GetEmitter().Emitter->GetUniqueEmitterName() : FString();
 	return FNiagaraStackGraphUtilities::CreateRapidIterationParameter(UniqueEmitterName, OutputNode->GetUsage(), InName, InputType);
 }
 
@@ -2025,7 +2026,7 @@ void UNiagaraStackFunctionInput::OnRenamed(FText NewNameText)
 		TSharedRef<FNiagaraSystemViewModel> CachedSysViewModel = GetSystemViewModel();
 		TSharedPtr<FNiagaraEmitterViewModel> CachedEmitterViewModel = GetEmitterViewModel();
 		UNiagaraSystem& System = GetSystemViewModel()->GetSystem();
-		UNiagaraEmitter* Emitter = GetEmitterViewModel().IsValid() ? GetEmitterViewModel()->GetEmitter() : nullptr;
+		FVersionedNiagaraEmitter Emitter = GetEmitterViewModel().IsValid() ? GetEmitterViewModel()->GetEmitter() : FVersionedNiagaraEmitter();
 
 		FScopedTransaction ScopedTransaction(LOCTEXT("RenameInput", "Rename this function's input."));
 		FNiagaraStackGraphUtilities::RenameAssignmentTarget(
@@ -2037,7 +2038,7 @@ void UNiagaraStackFunctionInput::OnRenamed(FText NewNameText)
 			NewName);
 		ensureMsgf(IsFinalized(), TEXT("Input not finalized when renamed."));
 
-		CachedSysViewModel->NotifyParameterRenamedExternally(OldVar, NewVar, Emitter);
+		CachedSysViewModel->NotifyParameterRenamedExternally(OldVar, NewVar, Emitter.Emitter);
 
 	}
 }
@@ -2053,14 +2054,14 @@ void UNiagaraStackFunctionInput::DeleteInput()
 	{
 		FScopedTransaction ScopedTransaction(LOCTEXT("RemoveInputTransaction", "Remove Input"));
 		TSharedRef<FNiagaraSystemViewModel> CachedSysViewModel = GetSystemViewModel();
-		UNiagaraEmitter* Emitter = GetEmitterViewModel().IsValid() ? GetEmitterViewModel()->GetEmitter() : nullptr;
+		FVersionedNiagaraEmitter VersionedEmitter = GetEmitterViewModel().IsValid() ? GetEmitterViewModel()->GetEmitter() : FVersionedNiagaraEmitter();
 
 		
 		{
 			// Rapid iteration parameters might be affected by this removal, so add them. Variables might also be removed in other bindings, but that is handled elsewhere.
 			UNiagaraSystem& System = GetSystemViewModel()->GetSystem();
 
-			FNiagaraStackGraphUtilities::FindAffectedScripts(&System, Emitter, *OwningModuleNode.Get(), AffectedScripts);
+			FNiagaraStackGraphUtilities::FindAffectedScripts(&System, VersionedEmitter, *OwningModuleNode.Get(), AffectedScripts);
 
 			for (TWeakObjectPtr<UNiagaraScript> AffectedScript : AffectedScripts)
 			{
@@ -2076,8 +2077,7 @@ void UNiagaraStackFunctionInput::DeleteInput()
 		NodeAssignment->Modify();
 		NodeAssignment->RemoveParameter(Var);
 
-		CachedSysViewModel->NotifyParameterRemovedExternally(Var, Emitter);
-
+		CachedSysViewModel->NotifyParameterRemovedExternally(Var, VersionedEmitter.Emitter);
 	}
 }
 
@@ -2234,7 +2234,7 @@ void UNiagaraStackFunctionInput::ReassignDynamicInputScript(UNiagaraScript* Dyna
 
 		const FString NewName = InputValues.DynamicNode->GetFunctionName();
 		UNiagaraSystem& System = GetSystemViewModel()->GetSystem();
-		UNiagaraEmitter* Emitter = GetEmitterViewModel().IsValid() ? GetEmitterViewModel()->GetEmitter() : nullptr;
+		FVersionedNiagaraEmitter Emitter = GetEmitterViewModel().IsValid() ? GetEmitterViewModel()->GetEmitter() : FVersionedNiagaraEmitter();
 		FNiagaraStackGraphUtilities::RenameReferencingParameters(&System, Emitter, *InputValues.DynamicNode.Get(), OldName, NewName);
 
 		InputValues.DynamicNode->RefreshFromExternalChanges();
diff --git a/Engine/Plugins/FX/Niagara/Source/NiagaraEditor/Private/ViewModels/Stack/NiagaraStackFunctionInputCollection.cpp b/Engine/Plugins/FX/Niagara/Source/NiagaraEditor/Private/ViewModels/Stack/NiagaraStackFunctionInputCollection.cpp
--- a/Engine/Plugins/FX/Niagara/Source/NiagaraEditor/Private/ViewModels/Stack/NiagaraStackFunctionInputCollection.cpp
+++ b/Engine/Plugins/FX/Niagara/Source/NiagaraEditor/Private/ViewModels/Stack/NiagaraStackFunctionInputCollection.cpp
@@ -18,7 +18,6 @@
 #include "ViewModels/Stack/NiagaraStackFunctionInput.h"
 #include "ViewModels/Stack/NiagaraStackGraphUtilities.h"
 #include "ViewModels/Stack/NiagaraStackInputCategory.h"
-#include "ViewModels/Stack/NiagaraStackModuleItem.h"
 #include "NiagaraEmitterEditorData.h"
 
 #define LOCTEXT_NAMESPACE "UNiagaraStackFunctionInputCollection"
@@ -54,9 +53,9 @@ static TOptional<FFunctionInputSummaryViewKey> GetSummaryViewInputKeyForFunction
 	return TOptional<FFunctionInputSummaryViewKey>();
 }
 
-static bool ShouldShowInSummaryView(UNiagaraEmitter* Emitter, UNiagaraNodeFunctionCall* FunctionCall, const FNiagaraVariable& InputVariable, const TOptional<FNiagaraVariableMetaData>& InputMetaData, bool bIsSummaryInEditMode)
+static bool ShouldShowInSummaryView(const FVersionedNiagaraEmitterData* EmitterData, UNiagaraNodeFunctionCall* FunctionCall, const FNiagaraVariable& InputVariable, const TOptional<FNiagaraVariableMetaData>& InputMetaData, bool bIsSummaryInEditMode)
 {
-	const UNiagaraEmitterEditorData* EditorData = Emitter? Cast<UNiagaraEmitterEditorData>(Emitter->GetEditorData()) : nullptr;
+	const UNiagaraEmitterEditorData* EditorData = EmitterData? Cast<UNiagaraEmitterEditorData>(EmitterData->GetEditorData()) : nullptr;
 	if (EditorData)
 	{
 		TOptional<FFunctionInputSummaryViewKey> SummaryViewKey = GetSummaryViewInputKeyForFunctionInput(FunctionCall, InputVariable, InputMetaData);
@@ -272,7 +271,8 @@ void UNiagaraStackFunctionInputCollectionBase::RefreshChildrenForFunctionCall(UN
 
 void UNiagaraStackFunctionInputCollectionBase::AppendInputsForFunctionCall(FFunctionCallNodesState& State, UNiagaraNodeFunctionCall* ModuleNode, UNiagaraNodeFunctionCall* InputFunctionCallNode, TArray<FStackIssue>& NewIssues, bool bShouldApplySummaryFilter)
 {
-	UNiagaraEmitter* Emitter = GetEmitterViewModel().IsValid() ? GetEmitterViewModel()->GetEmitter() : nullptr;
+	FVersionedNiagaraEmitter Emitter = GetEmitterViewModel().IsValid() ? GetEmitterViewModel()->GetEmitter() : FVersionedNiagaraEmitter();
+	FVersionedNiagaraEmitterData* EmitterData = Emitter.GetEmitterData();
 
 	TSet<const UEdGraphPin*> HiddenPins;
 	TArray<const UEdGraphPin*> InputPins;
@@ -300,10 +300,10 @@ void UNiagaraStackFunctionInputCollectionBase::AppendInputsForFunctionCall(FFunc
 	UNiagaraGraph* InputFunctionGraph = InputFunctionCallNode->GetCalledGraph();
 	
 	
-	const auto GetSummaryViewState = [&](UNiagaraEmitter* Emitter, UNiagaraNodeFunctionCall* FunctionCall, const FNiagaraVariable& InputVariable,
+	const auto GetSummaryViewState = [&](const UNiagaraNodeFunctionCall* FunctionCall, const FNiagaraVariable& InputVariable,
 		const TOptional<FNiagaraVariableMetaData>& InputMetaData, FText& InputCategory, TOptional<FText>& DisplayName, int32& EditorSortPriority ) -> bool
 	{
-		const UNiagaraEmitterEditorData* EditorData = Emitter? Cast<UNiagaraEmitterEditorData>(Emitter->GetEditorData()) : nullptr;
+		const UNiagaraEmitterEditorData* EditorData = EmitterData ? Cast<UNiagaraEmitterEditorData>(EmitterData->GetEditorData()) : nullptr;
 
 		const TOptional<FFunctionInputSummaryViewKey> SummaryViewKey = GetSummaryViewInputKeyForFunctionInput(InputFunctionCallNode, InputVariable, InputMetaData);
 		if (EditorData && bShouldApplySummaryFilter && SummaryViewKey.IsSet())
@@ -332,7 +332,7 @@ void UNiagaraStackFunctionInputCollectionBase::AppendInputsForFunctionCall(FFunc
 			DisplayName = SummaryViewData.IsSet() && SummaryViewData->DisplayName != NAME_None ? FText::FromName(SummaryViewData->DisplayName) : TOptional<FText>();		
 		}
 		
-		return ShouldShowInSummaryView(Emitter, InputFunctionCallNode, InputVariable, InputMetaData, GetEmitterViewModel().IsValid() && GetEmitterViewModel()->GetSummaryIsInEditMode());
+		return ShouldShowInSummaryView(EmitterData, InputFunctionCallNode, InputVariable, InputMetaData, GetEmitterViewModel().IsValid() && GetEmitterViewModel()->GetSummaryIsInEditMode());
 	};
 
 
@@ -369,7 +369,7 @@ void UNiagaraStackFunctionInputCollectionBase::AppendInputsForFunctionCall(FFunc
 		int32 EditorSortPriority = InputMetaData.IsSet() ? InputMetaData->EditorSortPriority : 0;
 		TOptional<FText> DisplayName;
 
-		bool bShouldShowInSummary = GetSummaryViewState(Emitter, InputFunctionCallNode, InputVariable, InputMetaData, InputCategory, DisplayName, EditorSortPriority);		
+		bool bShouldShowInSummary = GetSummaryViewState(InputFunctionCallNode, InputVariable, InputMetaData, InputCategory, DisplayName, EditorSortPriority);		
 		if (bShouldShowInSummary)
 		{
 			SummaryViewPins.Add(InputPin);
@@ -435,7 +435,7 @@ void UNiagaraStackFunctionInputCollectionBase::AppendInputsForFunctionCall(FFunc
 		int32 EditorSortPriority = InputMetaData.IsSet() ? InputMetaData->EditorSortPriority : 0;
 		TOptional<FText> DisplayName;
 						
-		bool bShouldShowInSummary = GetSummaryViewState(Emitter, InputFunctionCallNode, InputVariable, InputMetaData, InputCategory, DisplayName, EditorSortPriority);		
+		bool bShouldShowInSummary = GetSummaryViewState(InputFunctionCallNode, InputVariable, InputMetaData, InputCategory, DisplayName, EditorSortPriority);		
 		if (bShouldShowInSummary)
 		{
 			SummaryViewPins.Add(InputPin);
diff --git a/Engine/Plugins/FX/Niagara/Source/NiagaraEditor/Private/ViewModels/Stack/NiagaraStackGraphUtilities.cpp b/Engine/Plugins/FX/Niagara/Source/NiagaraEditor/Private/ViewModels/Stack/NiagaraStackGraphUtilities.cpp
--- a/Engine/Plugins/FX/Niagara/Source/NiagaraEditor/Private/ViewModels/Stack/NiagaraStackGraphUtilities.cpp
+++ b/Engine/Plugins/FX/Niagara/Source/NiagaraEditor/Private/ViewModels/Stack/NiagaraStackGraphUtilities.cpp
@@ -169,53 +169,6 @@ void FNiagaraStackGraphUtilities::RelayoutGraph(UEdGraph& Graph)
 	Graph.NotifyGraphChanged();
 }
 
-void FNiagaraStackGraphUtilities::GetWrittenVariablesForGraph(UEdGraph& Graph, TArray<FNiagaraVariable>& OutWrittenVariables)
-{
-	TArray<UNiagaraNodeOutput*> OutputNodes;
-	Graph.GetNodesOfClass<UNiagaraNodeOutput>(OutputNodes);
-	FPinCollectorArray InputPins;
-
-
-	UNiagaraEmitter* Emitter = Graph.GetTypedOuter<UNiagaraEmitter>();
-	UNiagaraSystem* System = Graph.GetTypedOuter<UNiagaraSystem>();
-
-	TArray<FNiagaraVariable> StaticVars;
-	TSharedPtr<FNiagaraGraphCachedDataBase, ESPMode::ThreadSafe> CachedTraversalData;
-	if (Emitter)
-	{
-		CachedTraversalData = Emitter->GetCachedTraversalData();
-	}
-	else if (System)
-	{
-		CachedTraversalData = System->GetCachedTraversalData();
-	}
-	if (CachedTraversalData.IsValid())
-	{
-		CachedTraversalData.Get()->GetStaticVariables(StaticVars);
-	}
-
-
-	for (UNiagaraNodeOutput* OutputNode : OutputNodes)
-	{
-		InputPins.Reset();
-		OutputNode->GetInputPins(InputPins);
-		if (InputPins.Num() == 1)
-		{
-			FNiagaraParameterMapHistoryBuilder Builder;
-			Builder.RegisterExternalStaticVariables(StaticVars);
-			Builder.BuildParameterMaps(OutputNode, true);
-			check(Builder.Histories.Num() == 1);
-			for (int32 i = 0; i < Builder.Histories[0].Variables.Num(); i++)
-			{
-				if (Builder.Histories[0].PerVariableWriteHistory[i].Num() > 0)
-				{
-					OutWrittenVariables.Add(Builder.Histories[0].Variables[i]);
-				}
-			}
-		}
-	}
-}
-
 void FNiagaraStackGraphUtilities::ConnectPinToInputNode(UEdGraphPin& Pin, UNiagaraNodeInput& InputNode)
 {
 	FPinCollectorArray InputPins;
@@ -359,7 +312,7 @@ const UNiagaraNodeOutput* FNiagaraStackGraphUtilities::GetEmitterOutputNodeForSt
 	return GetEmitterOutputNodeForStackNodeInternal<const UNiagaraNodeOutput, const UNiagaraNode>(StackNode);
 }
 
-TArray<FName> FNiagaraStackGraphUtilities::StackContextResolution(UNiagaraEmitter* OwningEmitter, UNiagaraNodeOutput* OutputNodeInChain)
+TArray<FName> FNiagaraStackGraphUtilities::StackContextResolution(FVersionedNiagaraEmitter OwningEmitter, UNiagaraNodeOutput* OutputNodeInChain)
 {
 	TArray<FName> PossibleRootNames;
 	ENiagaraScriptUsage Usage = OutputNodeInChain->GetUsage();
@@ -379,9 +332,9 @@ TArray<FName> FNiagaraStackGraphUtilities::StackContextResolution(UNiagaraEmitte
 			break;
 		case ENiagaraScriptUsage::ParticleSimulationStageScript:
 		{
-			if (OwningEmitter)
+			if (OwningEmitter.Emitter)
 			{
-				UNiagaraSimulationStageBase* Base = OwningEmitter->GetSimulationStageById(OutputNodeInChain->GetUsageId());
+				UNiagaraSimulationStageBase* Base = OwningEmitter.GetEmitterData()->GetSimulationStageById(OutputNodeInChain->GetUsageId());
 				if (Base)
 					StageName = Base->GetStackContextReplacementName();
 			}
@@ -397,9 +350,9 @@ TArray<FName> FNiagaraStackGraphUtilities::StackContextResolution(UNiagaraEmitte
 		case ENiagaraScriptUsage::EmitterUpdateScript:
 			StageName = TEXT("Emitter");
 			{
-				if (OwningEmitter)
+				if (OwningEmitter.Emitter)
 				{
-					FString EmitterAliasStr = OwningEmitter->GetUniqueEmitterName();
+					FString EmitterAliasStr = OwningEmitter.Emitter->GetUniqueEmitterName();
 					if (EmitterAliasStr.Len())
 					{
 						StageName = *EmitterAliasStr;
@@ -422,16 +375,16 @@ TArray<FName> FNiagaraStackGraphUtilities::StackContextResolution(UNiagaraEmitte
 	return PossibleRootNames;
 }
 
-void FNiagaraStackGraphUtilities::BuildParameterMapHistoryWithStackContextResolution(UNiagaraEmitter* OwningEmitter, UNiagaraNodeOutput* OutputNodeInChain, UNiagaraNode* NodeToVisit, FNiagaraParameterMapHistoryBuilder& Builder, bool bRecursive /*= true*/, bool bFilterForCompilation /*= true*/)
+void FNiagaraStackGraphUtilities::BuildParameterMapHistoryWithStackContextResolution(FVersionedNiagaraEmitter OwningEmitter, UNiagaraNodeOutput* OutputNodeInChain, UNiagaraNode* NodeToVisit, FNiagaraParameterMapHistoryBuilder& Builder, bool bRecursive /*= true*/, bool bFilterForCompilation /*= true*/)
 {
 	bool bSetUsage = false;
-	if (OwningEmitter && OutputNodeInChain)
+	if (OwningEmitter.Emitter && OutputNodeInChain)
 	{
 		ENiagaraScriptUsage Usage = OutputNodeInChain->GetUsage();
 		FName StageName;
 		if (Usage == ENiagaraScriptUsage::ParticleSimulationStageScript)
 		{
-			UNiagaraSimulationStageBase* Base = OwningEmitter->GetSimulationStageById(OutputNodeInChain->GetUsageId());
+			UNiagaraSimulationStageBase* Base = OwningEmitter.GetEmitterData()->GetSimulationStageById(OutputNodeInChain->GetUsageId());
 			if (Base)
 				StageName = Base->GetStackContextReplacementName();
 		}
@@ -441,9 +394,9 @@ void FNiagaraStackGraphUtilities::BuildParameterMapHistoryWithStackContextResolu
 
 	TArray<FNiagaraVariable> StaticVars;
 	TSharedPtr<FNiagaraGraphCachedDataBase, ESPMode::ThreadSafe> CachedTraversalData;
-	if (OwningEmitter)
+	if (OwningEmitter.GetEmitterData())
 	{
-		CachedTraversalData = OwningEmitter->GetCachedTraversalData();
+		CachedTraversalData = OwningEmitter.Emitter->GetCachedTraversalData(OwningEmitter.Version);
 	}
 	
 	if (CachedTraversalData.IsValid())
@@ -490,10 +443,10 @@ UNiagaraNodeInput* FNiagaraStackGraphUtilities::GetEmitterInputNodeForStackNode(
 	return nullptr;
 }
 
-FText GetChangelistText(UNiagaraScript* NiagaraScript, const FNiagaraAssetVersion& FromVersion, const FNiagaraAssetVersion& ToVersion)
+FText GetChangelistText(FNiagaraVersionedObject* VersionedObject, const FNiagaraAssetVersion& FromVersion, const FNiagaraAssetVersion& ToVersion)
 {
 	FText Result;
-	for (FNiagaraAssetVersion Version : NiagaraScript->GetAllAvailableVersions())
+	for (FNiagaraAssetVersion Version : VersionedObject->GetAllAvailableVersions())
 	{
 		if (Version <= FromVersion)
 		{
@@ -503,7 +456,7 @@ FText GetChangelistText(UNiagaraScript* NiagaraScript, const FNiagaraAssetVersio
 		{
 			break;
 		}
-		FText ChangeDescription = NiagaraScript->GetScriptData(Version.VersionGuid)->VersionChangeDescription;
+		FText ChangeDescription = VersionedObject->GetVersionDataAccessor(Version.VersionGuid)->GetVersionChangeDescription();
 		if (!ChangeDescription.IsEmpty())
 		{
 			Result = FText::Format(FText::FromString("{0}{1}.{2}: {3}\n"), Result, Version.MajorVersion, Version.MinorVersion, ChangeDescription);
@@ -576,6 +529,93 @@ void FNiagaraStackGraphUtilities::CheckForDeprecatedScriptVersion(UNiagaraNodeFu
 	}
 }
 
+void FNiagaraStackGraphUtilities::CheckForDeprecatedEmitterVersion(TSharedPtr<FNiagaraEmitterViewModel> ViewModel, const FString& StackEditorDataKey, UNiagaraStackEntry::FStackIssueFixDelegate VersionUpgradeFix, TArray<UNiagaraStackEntry::FStackIssue>& OutIssues)
+{
+	if (!ViewModel || !ViewModel->GetParentEmitter().Emitter || !ViewModel->GetParentEmitter().Emitter->IsVersioningEnabled())
+	{
+		return;
+	}
+	FVersionedNiagaraEmitter VersionedEmitter = ViewModel->GetParentEmitter();
+	UNiagaraEmitter* Emitter = VersionedEmitter.Emitter;
+	FVersionedNiagaraEmitterData* EmitterData = VersionedEmitter.GetEmitterData();
+
+	if (Emitter && !EmitterData)
+	{
+		if (VersionUpgradeFix.IsBound())
+		{
+			TArray<UNiagaraStackEntry::FStackIssueFix> Fixes;
+			Fixes.Add(UNiagaraStackEntry::FStackIssueFix(LOCTEXT("InvalidParentVersionFix", "Switch to exposed parent version."), VersionUpgradeFix));
+			
+			UNiagaraStackEntry::FStackIssue UpgradeVersion(
+				EStackIssueSeverity::Error,
+				LOCTEXT("MissingParentVersionFormat", "Invalid parent emitter"),
+				LOCTEXT("MissingParentVersionLongFormat", "The parent version selected for this emitter does not exist.\nPlease select a valid version or changes cannot be merged down from the parent emitter."),
+				StackEditorDataKey,
+				false,
+				Fixes);
+			OutIssues.Add(UpgradeVersion);
+		}
+		return;
+	}
+	
+	// Generate an issue if the emitter version is out of date
+	if (VersionUpgradeFix.IsBound())
+	{
+		FNiagaraAssetVersion ExposedVersion = Emitter->GetExposedVersion();
+		if (FNiagaraAssetVersion ReferencedVersion = EmitterData->Version; ReferencedVersion.MajorVersion < ExposedVersion.MajorVersion)
+		{
+			TArray<UNiagaraStackEntry::FStackIssueFix> Fixes;
+			Fixes.Add(UNiagaraStackEntry::FStackIssueFix(LOCTEXT("UpgradeVersionFix", "Upgrade to newest parent version."), VersionUpgradeFix));
+			
+			FText ChangelistDescriptions = GetChangelistText(Emitter, ReferencedVersion, ExposedVersion);
+			FText LongDescription = FText::Format(LOCTEXT("DeprecatedVersionFormat", "This emitter has a newer parent version available.\nYou can upgrade now, but major version upgrades can sometimes come with breaking changes! So check that everything is still working as expected afterwards.{0}"),
+				ChangelistDescriptions.IsEmpty() ? FText() : FText::Format(LOCTEXT("DeprecatedVersionFormatChanges", "\n\nVersion change description:\n{0}"), ChangelistDescriptions));
+			UNiagaraStackEntry::FStackIssue UpgradeVersion(
+                EStackIssueSeverity::Warning,
+                FText::Format(LOCTEXT("DeprecatedVersionSummaryFormat", "Deprecated emitter version: {0}.{1} -> {2}.{3}"),
+                	FText::AsNumber(ReferencedVersion.MajorVersion), FText::AsNumber(ReferencedVersion.MinorVersion), FText::AsNumber(ExposedVersion.MajorVersion), FText::AsNumber(ExposedVersion.MinorVersion)),
+                LongDescription,
+                StackEditorDataKey,
+                true,
+                Fixes);
+			OutIssues.Add(UpgradeVersion);
+		}
+	}
+
+	// Generate a note of the changelist when the emitter version was manually changed
+	{
+		FGuid SelectedVersionGuid = VersionedEmitter.Version;
+		FGuid PreviousVersionGuid = ViewModel->PreviousEmitterVersion;
+		FVersionedNiagaraEmitterData* PreviousVersion = Emitter->GetEmitterData(PreviousVersionGuid);
+		if (PreviousVersionGuid.IsValid() && PreviousVersionGuid != SelectedVersionGuid && PreviousVersion)
+		{
+			FText ChangelistDescriptions = GetChangelistText(Emitter, PreviousVersion->Version, EmitterData->Version);
+			if (!ChangelistDescriptions.IsEmpty())
+			{
+				UNiagaraStackEntry::FStackIssue UpgradeInfo(
+	                EStackIssueSeverity::Info,
+	                LOCTEXT("VersionUpgradeInfoSummary", "Version upgrade note"),
+	                FText::Format(LOCTEXT("VersionUpgradeFormatChanges", "The version of this emitter was recently upgraded; here is a list of changes:\n{0}"), ChangelistDescriptions),
+	                StackEditorDataKey,
+	                true);
+				OutIssues.Add(UpgradeInfo);
+			}
+		}
+	}
+
+	// Generate a note if the last version upgrade generated warnings
+	if (ViewModel->PythonUpgradeScriptWarnings.IsEmpty() == false)
+	{
+		UNiagaraStackEntry::FStackIssue PythonLog(
+            EStackIssueSeverity::Warning,
+            LOCTEXT("PythonLogSummary", "Python upgrade script log"),
+            FText::Format(LOCTEXT("VersionUpgradePythonLog", "Upgrading the version generated some warnings from the python script:\n{0}"), FText::FromString(ViewModel->PythonUpgradeScriptWarnings)),
+            StackEditorDataKey,
+            true);
+		OutIssues.Add(PythonLog);
+	}
+}
+
 void GetGroupNodesRecursive(const TArray<UNiagaraNode*>& CurrentStartNodes, UNiagaraNode* EndNode, TArray<UNiagaraNode*>& OutAllNodes)
 {
 	FPinCollectorArray InputPins;
@@ -793,6 +833,25 @@ void ExtractInputPinsFromHistory(FNiagaraParameterMapHistory& History, UEdGraph*
 	}
 }
 
+FVersionedNiagaraEmitter GetOuterEmitter(UNiagaraNode* Node)
+{
+	FVersionedNiagaraEmitter Result;
+	if (Node)
+	{
+		if (UNiagaraGraph* NiagaraGraph = Node->GetNiagaraGraph())
+		{
+			Result = NiagaraGraph->GetOwningEmitter();
+		}
+
+		if (!Result.Emitter)
+		{
+			Result.Emitter = Node->GetTypedOuter<UNiagaraEmitter>();
+		}
+	}
+	
+	return Result;
+}
+
 void FNiagaraStackGraphUtilities::GetStackFunctionInputPins(UNiagaraNodeFunctionCall& FunctionCallNode, TArray<const UEdGraphPin*>& OutInputPins, TSet<const UEdGraphPin*>& OutHiddenPins, FCompileConstantResolver ConstantResolver, ENiagaraGetStackFunctionInputPinsOptions Options, bool bIgnoreDisabled)
 {
 	FNiagaraEditorModule::Get().GetGraphDataCache().GetStackFunctionInputPins(FunctionCallNode, OutInputPins, OutHiddenPins, ConstantResolver, Options, bIgnoreDisabled);
@@ -804,14 +863,14 @@ void FNiagaraStackGraphUtilities::GetStackFunctionInputPinsWithoutCache(UNiagara
 	Builder.SetIgnoreDisabled(bIgnoreDisabled);
 	Builder.ConstantResolver = ConstantResolver;
 
-	UNiagaraEmitter* Emitter = FunctionCallNode.GetTypedOuter<UNiagaraEmitter>();
+	FVersionedNiagaraEmitter Emitter = GetOuterEmitter(&FunctionCallNode);
 	UNiagaraSystem* System = FunctionCallNode.GetTypedOuter<UNiagaraSystem>();
 
 	TArray<FNiagaraVariable> StaticVars;
 	TSharedPtr<FNiagaraGraphCachedDataBase, ESPMode::ThreadSafe> CachedTraversalData;
-	if (Emitter)
+	if (Emitter.GetEmitterData())
 	{
-		CachedTraversalData = Emitter->GetCachedTraversalData();
+		CachedTraversalData = Emitter.Emitter->GetCachedTraversalData(Emitter.Version);
 	}
 	else if (System)
 	{
@@ -970,17 +1029,17 @@ void FNiagaraStackGraphUtilities::GetStackFunctionStaticSwitchPins(UNiagaraNodeF
 	FunctionCallNode.GetInputPins(InputPins);
 	FNiagaraEditorUtilities::SetStaticSwitchConstants(FunctionCallGraph, InputPins, ConstantResolver);
 
-	UNiagaraEmitter* Emitter = FunctionCallNode.GetTypedOuter<UNiagaraEmitter>();
+	FVersionedNiagaraEmitter OuterEmitter = GetOuterEmitter(&FunctionCallNode);
 	UNiagaraSystem* System = FunctionCallNode.GetTypedOuter<UNiagaraSystem>();
 
 	TArray<FNiagaraVariable> StaticVars;
 	TSharedPtr<FNiagaraGraphCachedDataBase, ESPMode::ThreadSafe> CachedTraversalData;
 	FString EmitterName;
 	FNiagaraAliasContext::ERapidIterationParameterMode Mode = FNiagaraAliasContext::ERapidIterationParameterMode::None;
-	if (Emitter)
+	if (OuterEmitter.GetEmitterData())
 	{
-		CachedTraversalData = Emitter->GetCachedTraversalData();
-		EmitterName = Emitter->GetUniqueEmitterName();
+		CachedTraversalData = OuterEmitter.Emitter->GetCachedTraversalData(OuterEmitter.Version);
+		EmitterName = OuterEmitter.Emitter->GetUniqueEmitterName();
 		Mode = FNiagaraAliasContext::ERapidIterationParameterMode::EmitterOrParticleScript;
 	}
 	else if (System)
@@ -1038,14 +1097,14 @@ void FNiagaraStackGraphUtilities::GetStackFunctionOutputVariables(UNiagaraNodeFu
 	Builder.SetIgnoreDisabled(false);
 	Builder.ConstantResolver = ConstantResolver;
 
-	UNiagaraEmitter* Emitter = FunctionCallNode.GetTypedOuter<UNiagaraEmitter>();
+	FVersionedNiagaraEmitter OuterEmitter = GetOuterEmitter(&FunctionCallNode);
 	UNiagaraSystem* System = FunctionCallNode.GetTypedOuter<UNiagaraSystem>();
 
 	TArray<FNiagaraVariable> StaticVars;
 	TSharedPtr<FNiagaraGraphCachedDataBase, ESPMode::ThreadSafe> CachedTraversalData;
-	if (Emitter)
+	if (OuterEmitter.GetEmitterData())
 	{
-		CachedTraversalData = Emitter->GetCachedTraversalData();
+		CachedTraversalData = OuterEmitter.Emitter->GetCachedTraversalData(OuterEmitter.Version);
 	}
 	else if (System)
 	{
@@ -1091,14 +1150,14 @@ bool FNiagaraStackGraphUtilities::GetStackFunctionInputAndOutputVariables(UNiaga
 	Builder.SetIgnoreDisabled(false);
 	Builder.ConstantResolver = ConstantResolver;
 
-	UNiagaraEmitter* Emitter = FunctionCallNode.GetTypedOuter<UNiagaraEmitter>();
+	FVersionedNiagaraEmitter OuterEmitter = GetOuterEmitter(&FunctionCallNode);
 	UNiagaraSystem* System = FunctionCallNode.GetTypedOuter<UNiagaraSystem>();
 
 	TArray<FNiagaraVariable> StaticVars;
 	TSharedPtr<FNiagaraGraphCachedDataBase, ESPMode::ThreadSafe> CachedTraversalData;
-	if (Emitter)
+	if (OuterEmitter.GetEmitterData())
 	{
-		CachedTraversalData = Emitter->GetCachedTraversalData();
+		CachedTraversalData = OuterEmitter.Emitter->GetCachedTraversalData(OuterEmitter.Version);
 	}
 	else if (System)
 	{
@@ -1716,7 +1775,7 @@ UNiagaraNodeFunctionCall* FNiagaraStackGraphUtilities::AddScriptModuleToStack(co
 	const FGuid& VersionGuid = Args.VersionGuid;
 
 	// Get the stack graph via the TargetOutputNode.
-	UEdGraph* Graph = TargetOutputNode->GetGraph();
+	UNiagaraGraph* Graph = TargetOutputNode->GetNiagaraGraph();
 	Graph->Modify();
 
 	// Create the UNiagaraNodeFunctionCall to represent the script in the stack.
@@ -2075,7 +2134,7 @@ FNiagaraVariable FNiagaraStackGraphUtilities::CreateRapidIterationParameter(cons
 	return RapidIterationVariable;
 }
 
-void FNiagaraStackGraphUtilities::CleanUpStaleRapidIterationParameters(UNiagaraScript& Script, UNiagaraEmitter& OwningEmitter)
+void FNiagaraStackGraphUtilities::CleanUpStaleRapidIterationParameters(UNiagaraScript& Script, FVersionedNiagaraEmitter OwningEmitter)
 {
 	UNiagaraScriptSource* Source = CastChecked<UNiagaraScriptSource>(Script.GetLatestSource());
 	UNiagaraNodeOutput* OutputNode = Source->NodeGraph->FindOutputNode(Script.GetUsage(), Script.GetUsageId());
@@ -2092,7 +2151,7 @@ void FNiagaraStackGraphUtilities::CleanUpStaleRapidIterationParameters(UNiagaraS
 			FString InputName;
 			if (FNiagaraParameterMapHistory::SplitRapidIterationParameterName(RapidIterationParameter, Script.GetUsage(), EmitterName, FunctionCallName, InputName))
 			{
-				if (EmitterName != OwningEmitter.GetUniqueEmitterName() || ValidFunctionCallNames.Contains(FunctionCallName) == false)
+				if (EmitterName != OwningEmitter.Emitter->GetUniqueEmitterName() || ValidFunctionCallNames.Contains(FunctionCallName) == false)
 				{
 					Script.RapidIterationParameters.RemoveParameter(RapidIterationParameter);
 				}
@@ -2101,10 +2160,10 @@ void FNiagaraStackGraphUtilities::CleanUpStaleRapidIterationParameters(UNiagaraS
 	}
 }
 
-void FNiagaraStackGraphUtilities::CleanUpStaleRapidIterationParameters(UNiagaraEmitter& Emitter)
+void FNiagaraStackGraphUtilities::CleanUpStaleRapidIterationParameters(FVersionedNiagaraEmitter Emitter)
 {
 	TArray<UNiagaraScript*> Scripts;
-	Emitter.GetScripts(Scripts, false);
+	Emitter.GetEmitterData()->GetScripts(Scripts, false);
 	for (UNiagaraScript* Script : Scripts)
 	{
 		CleanUpStaleRapidIterationParameters(*Script, Emitter);
@@ -2490,7 +2549,6 @@ void FNiagaraStackGraphUtilities::MoveModule(UNiagaraScript& SourceScript, UNiag
 			// so we can just add them directly to the target script.
 			for (auto It = SourceFunctionIdToRapidIterationParametersMap.CreateIterator(); It; ++It)
 			{
-				FGuid& FunctionId = It.Key();
 				TArray<FNiagaraVariable>& RapidIterationParameters = It.Value();
 				for (FNiagaraVariable& RapidIterationParameter : RapidIterationParameters)
 				{
@@ -2588,7 +2646,7 @@ void FNiagaraStackGraphUtilities::RebuildEmitterNodes(UNiagaraSystem& System)
 		UEdGraphPin* OutPinLinkedPin = OutPin != nullptr && OutPin->LinkedTo.Num() == 1 ? OutPin->LinkedTo[0] : nullptr;
 		CurrentEmitterNode->DestroyNode();
 
-		if (InPinLinkedPin != nullptr &&& OutPinLinkedPin != nullptr)
+		if (InPinLinkedPin != nullptr && OutPinLinkedPin != nullptr)
 		{
 			MakeLinkTo(InPinLinkedPin, OutPinLinkedPin);
 		}
@@ -2685,16 +2743,16 @@ void FNiagaraStackGraphUtilities::RebuildEmitterNodes(UNiagaraSystem& System)
 	RelayoutGraph(*SystemGraph);
 }
 
-void FNiagaraStackGraphUtilities::FindAffectedScripts(UNiagaraSystem* System, UNiagaraEmitter* Emitter, UNiagaraNodeFunctionCall& ModuleNode, TArray<TWeakObjectPtr<UNiagaraScript>>& OutAffectedScripts)
+void FNiagaraStackGraphUtilities::FindAffectedScripts(UNiagaraSystem* System, FVersionedNiagaraEmitter Emitter, UNiagaraNodeFunctionCall& ModuleNode, TArray<TWeakObjectPtr<UNiagaraScript>>& OutAffectedScripts)
 {
 	UNiagaraNodeOutput* OutputNode = FNiagaraStackGraphUtilities::GetEmitterOutputNodeForStackNode(ModuleNode);
 
 	if (OutputNode)
 	{
 		TArray<UNiagaraScript*> Scripts;
-		if (Emitter != nullptr)
+		if (FVersionedNiagaraEmitterData* EmitterData = Emitter.GetEmitterData())
 		{
-			Emitter->GetScripts(Scripts, false);
+			EmitterData->GetScripts(Scripts, false);
 		}
 
 		if (System != nullptr)
@@ -2721,7 +2779,7 @@ void FNiagaraStackGraphUtilities::FindAffectedScripts(UNiagaraSystem* System, UN
 	}
 }
 
-void FNiagaraStackGraphUtilities::GatherRenamedStackFunctionOutputVariableNames(UNiagaraEmitter* Emitter, UNiagaraNodeFunctionCall& FunctionCallNode, const FString& OldFunctionName, const FString& NewFunctionName, TMap<FName, FName>& OutOldToNewNameMap)
+void FNiagaraStackGraphUtilities::GatherRenamedStackFunctionOutputVariableNames(FVersionedNiagaraEmitter Emitter, UNiagaraNodeFunctionCall& FunctionCallNode, const FString& OldFunctionName, const FString& NewFunctionName, TMap<FName, FName>& OutOldToNewNameMap)
 {
 	if (OldFunctionName == NewFunctionName)
 	{
@@ -2755,7 +2813,7 @@ void FNiagaraStackGraphUtilities::GatherRenamedStackFunctionOutputVariableNames(
 	}
 }
 
-void FNiagaraStackGraphUtilities::GatherRenamedStackFunctionInputAndOutputVariableNames(UNiagaraEmitter* Emitter, UNiagaraNodeFunctionCall& FunctionCallNode, const FString& OldFunctionName, const FString& NewFunctionName, TMap<FName, FName>& OutOldToNewNameMap)
+void FNiagaraStackGraphUtilities::GatherRenamedStackFunctionInputAndOutputVariableNames(FVersionedNiagaraEmitter Emitter, UNiagaraNodeFunctionCall& FunctionCallNode, const FString& OldFunctionName, const FString& NewFunctionName, TMap<FName, FName>& OutOldToNewNameMap)
 {
 	if (OldFunctionName == NewFunctionName)
 	{
@@ -2789,7 +2847,7 @@ void FNiagaraStackGraphUtilities::GatherRenamedStackFunctionInputAndOutputVariab
 	}
 }
 
-void FNiagaraStackGraphUtilities::RenameReferencingParameters(UNiagaraSystem* System, UNiagaraEmitter* Emitter, UNiagaraNodeFunctionCall& FunctionCallNode, const FString& OldModuleName, const FString& NewModuleName)
+void FNiagaraStackGraphUtilities::RenameReferencingParameters(UNiagaraSystem* System, FVersionedNiagaraEmitter Emitter, UNiagaraNodeFunctionCall& FunctionCallNode, const FString& OldModuleName, const FString& NewModuleName)
 {
 	TMap<FName, FName> OldNameToNewNameMap;
 	FNiagaraStackGraphUtilities::GatherRenamedStackFunctionInputAndOutputVariableNames(Emitter, FunctionCallNode, OldModuleName, NewModuleName, OldNameToNewNameMap);
@@ -2816,10 +2874,7 @@ void FNiagaraStackGraphUtilities::RenameReferencingParameters(UNiagaraSystem* Sy
 	TArray<TWeakObjectPtr<UNiagaraScript>> Scripts;
 	FindAffectedScripts(System, Emitter, FunctionCallNode, Scripts);
 
-	const UNiagaraNodeOutput* OutputNode = GetEmitterOutputNodeForStackNode(FunctionCallNode);
-	UNiagaraGraph* OwningGraph = FunctionCallNode.GetNiagaraGraph();
-
-	FString OwningEmitterName = Emitter != nullptr ? Emitter->GetUniqueEmitterName() : FString();
+	FString OwningEmitterName = Emitter.Emitter != nullptr ? Emitter.Emitter->GetUniqueEmitterName() : FString();
 
 	for (TWeakObjectPtr<UNiagaraScript> Script : Scripts)
 	{
@@ -2930,11 +2985,11 @@ bool FNiagaraStackGraphUtilities::TryRenameAssignmentTarget(UNiagaraNodeAssignme
 	if (OutputNode != nullptr)
 	{
 		UNiagaraSystem* OwningSystem = OwningAssignmentNode.GetTypedOuter<UNiagaraSystem>();
-		UNiagaraEmitter* OwningEmitter = OwningAssignmentNode.GetTypedOuter<UNiagaraEmitter>();
+		FVersionedNiagaraEmitter OwningEmitter = OwningAssignmentNode.GetNiagaraGraph()->GetOwningEmitter();
 		UNiagaraScript* OwningScript = nullptr;
-		if (OwningEmitter != nullptr)
+		if (OwningEmitter.Emitter)
 		{
-			OwningScript = OwningEmitter->GetScript(OutputNode->GetUsage(), OutputNode->GetUsageId());
+			OwningScript = OwningEmitter.GetEmitterData()->GetScript(OutputNode->GetUsage(), OutputNode->GetUsageId());
 		}
 		else if(OwningSystem != nullptr)
 		{
@@ -2959,16 +3014,17 @@ bool FNiagaraStackGraphUtilities::TryRenameAssignmentTarget(UNiagaraNodeAssignme
 
 void FNiagaraStackGraphUtilities::RenameAssignmentTarget(
 	UNiagaraSystem& OwningSystem,
-	UNiagaraEmitter* OwningEmitter,
+	FVersionedNiagaraEmitter OwningEmitter,
 	UNiagaraScript& OwningScript,
 	UNiagaraNodeAssignment& OwningAssignmentNode,
 	FNiagaraVariable CurrentAssignmentTarget,
 	FName NewAssignmentTargetName)
 {
 	UNiagaraStackEditorData* StackEditorData;
-	if (OwningEmitter != nullptr)
+	FVersionedNiagaraEmitterData* EmitterData = OwningEmitter.GetEmitterData();
+	if (EmitterData != nullptr)
 	{
-		UNiagaraEmitterEditorData* EmitterEditorData = Cast<UNiagaraEmitterEditorData>(OwningEmitter->GetEditorData());
+		UNiagaraEmitterEditorData* EmitterEditorData = Cast<UNiagaraEmitterEditorData>(EmitterData->GetEditorData());
 		StackEditorData = EmitterEditorData != nullptr ? &EmitterEditorData->GetStackEditorData() : nullptr;
 	}
 	else
@@ -2999,7 +3055,7 @@ void FNiagaraStackGraphUtilities::RenameAssignmentTarget(
 		else if (IsRapidIterationType(CurrentAssignmentTarget.GetType()))
 		{
 			// Otherwise if this is a rapid iteration type check to see if there is an existing rapid iteration value, and if so, rename it.
-			FString UniqueEmitterName = OwningEmitter != nullptr ? OwningEmitter->GetUniqueEmitterName() : FString();
+			FString UniqueEmitterName = OwningEmitter.Emitter != nullptr ? OwningEmitter.Emitter->GetUniqueEmitterName() : FString();
 			FNiagaraVariable CurrentRapidIterationParameter = CreateRapidIterationParameter(UniqueEmitterName, OwningScript.GetUsage(), CurrentAliasedInputParameterHandle.GetParameterHandleString(), CurrentAssignmentTarget.GetType());
 			if (OwningScript.RapidIterationParameters.IndexOf(CurrentRapidIterationParameter) != INDEX_NONE)
 			{
@@ -3350,18 +3406,16 @@ void FNiagaraStackGraphUtilities::SynchronizeReferencingMapPinsWithFunctionCall(
 			UNiagaraNodeFunctionCall* TargetFunctionCallNode = GetTargetFunctionCallForStackLinkedOutputPin(*StackLinkedOutputPin);
 			if (TargetFunctionCallNode != nullptr)
 			{
-				TArray<FGuid> BoundGuids = TargetFunctionCallNode->GetBoundPinGuidsByName(LinkedOutputHandle.GetName());
+				TArray<FGuid> BoundGuids = TargetFunctionCallNode->GetBoundPinGuidsByName(LinkedOutputHandle.GetParameterHandleString());
 				for (const FGuid& BoundGuid : BoundGuids)
 				{
 					FNiagaraVariable* BoundVariable = VariableIdToVariableMap->Find(BoundGuid);
 					if (BoundVariable != nullptr)
 					{
-						FNiagaraParameterHandle BoundVariableHandle(BoundVariable->GetName());
+						FNiagaraVariable BoundVariableResolvedName = FNiagaraUtilities::ResolveAliases(*BoundVariable,FNiagaraAliasContext().ChangeModuleToModuleName(InFunctionCallNode.GetFunctionName()));
+						FNiagaraParameterHandle BoundVariableHandle(BoundVariableResolvedName.GetName());
 						if (LinkedOutputHandle.GetName() != BoundVariableHandle.GetName())
 						{
-							FNiagaraVariable BoundVariableResolvedName = FNiagaraUtilities::ResolveAliases(*BoundVariable,
-								FNiagaraAliasContext().ChangeModuleToModuleName(InFunctionCallNode.GetFunctionName())
-							);
 							StackLinkedOutputPin->PinName = BoundVariableResolvedName.GetName();
 							TargetFunctionCallNode->UpdateInputNameBinding(BoundGuid, BoundVariableHandle.GetName());
 						}
@@ -3459,7 +3513,7 @@ void FNiagaraStackGraphUtilities::DependencyUtilities::GetModuleScriptAssetsByDe
 	}
 }
 
-int32 FNiagaraStackGraphUtilities::DependencyUtilities::FindBestIndexForModuleInStack(UNiagaraNodeFunctionCall& ModuleNode, UEdGraph& EmitterScriptGraph)
+int32 FNiagaraStackGraphUtilities::DependencyUtilities::FindBestIndexForModuleInStack(UNiagaraNodeFunctionCall& ModuleNode, UNiagaraGraph& EmitterScriptGraph)
 {
 	// Check if the new module node has any dependencies to begin with. If not, early exit.
 	FVersionedNiagaraScriptData* ScriptData = ModuleNode.GetScriptData();
@@ -3470,8 +3524,8 @@ int32 FNiagaraStackGraphUtilities::DependencyUtilities::FindBestIndexForModuleIn
 
 	// Get the Emitter and System the emitter script script graph is outered to.
 	UNiagaraSystem* System = EmitterScriptGraph.GetTypedOuter<UNiagaraSystem>();
-	UNiagaraEmitter* Emitter = EmitterScriptGraph.GetTypedOuter<UNiagaraEmitter>();
-	if (System == nullptr || Emitter == nullptr)
+	FVersionedNiagaraEmitter OuterEmitter = EmitterScriptGraph.GetOwningEmitter();
+	if (System == nullptr || OuterEmitter.Emitter == nullptr)
 	{
 		return INDEX_NONE;
 	}
@@ -3483,7 +3537,7 @@ int32 FNiagaraStackGraphUtilities::DependencyUtilities::FindBestIndexForModuleIn
 		ensureMsgf(false, TEXT("Failed to get systemviewmodel for valid system when getting best index in stack for module!"));
 		return INDEX_NONE;
 	}
-	TSharedPtr<FNiagaraEmitterHandleViewModel> EmitterHandleViewModel = SystemViewModel->GetEmitterHandleViewModelForEmitter(Emitter);
+	TSharedPtr<FNiagaraEmitterHandleViewModel> EmitterHandleViewModel = SystemViewModel->GetEmitterHandleViewModelForEmitter(OuterEmitter);
 	if (EmitterHandleViewModel->IsValid() == false)
 	{
 		ensureMsgf(false, TEXT("Failed to get emitterhandleviewmodel for valid emitter when getting best index in stack for module!"));
diff --git a/Engine/Plugins/FX/Niagara/Source/NiagaraEditor/Private/ViewModels/Stack/NiagaraStackModuleItem.cpp b/Engine/Plugins/FX/Niagara/Source/NiagaraEditor/Private/ViewModels/Stack/NiagaraStackModuleItem.cpp
--- a/Engine/Plugins/FX/Niagara/Source/NiagaraEditor/Private/ViewModels/Stack/NiagaraStackModuleItem.cpp
+++ b/Engine/Plugins/FX/Niagara/Source/NiagaraEditor/Private/ViewModels/Stack/NiagaraStackModuleItem.cpp
@@ -9,7 +9,6 @@
 #include "NiagaraConvertInPlaceUtilityBase.h"
 #include "NiagaraEditorStyle.h"
 #include "NiagaraEditorUtilities.h"
-#include "NiagaraEmitter.h"
 #include "NiagaraEmitterHandle.h"
 #include "NiagaraGraph.h"
 #include "NiagaraMessageManager.h"
@@ -573,7 +572,7 @@ void GenerateDependencyIssues(
 	}
 
 	FVersionedNiagaraScriptData* ScriptData = SourceModuleNode.FunctionScript->GetScriptData(SourceModuleNode.SelectedScriptVersion);
-	int32 ModuleIndex = SourceStackModuleData.IndexOfByPredicate([&SourceModuleNode](const FNiagaraStackModuleData& ModuleData) { return ModuleData.ModuleNode == &SourceModuleNode; });
+	int32 ModuleIndex = SourceStackModuleData.IndexOfByPredicate([&SourceModuleNode](const FNiagaraStackModuleData& ModuleData) { return ModuleData.ModuleNode == &SourceModuleNode || ModuleData.ModuleNode->GetFunctionName() == SourceModuleNode.GetFunctionName(); });
 	if (ensureMsgf(ModuleIndex != INDEX_NONE, TEXT("In system %s, module %s (%s) did not exist in the stack module data."),
 		*SourceSystemViewModel->GetSystem().GetPathName(), *SourceModuleNode.GetFunctionName(), *SourceModuleNode.GetName()))
 	{
@@ -1059,10 +1058,10 @@ bool UNiagaraStackModuleItem::CanMoveAndDelete() const
 			// When editing systems only non-base modules can be moved and deleted.
 			TSharedRef<FNiagaraScriptMergeManager> MergeManager = FNiagaraScriptMergeManager::Get();
 
-			const UNiagaraEmitter* BaseEmitter = GetEmitterViewModel()->GetEmitter()->GetParent();
+			FVersionedNiagaraEmitter BaseEmitter = GetEmitterViewModel()->GetParentEmitter();
 
 			bool bIsMergeable = MergeManager->IsMergeableScriptUsage(OutputNode->GetUsage());
-			bool bHasBaseModule = bIsMergeable && BaseEmitter != nullptr && MergeManager->HasBaseModule(*BaseEmitter, OutputNode->GetUsage(), OutputNode->GetUsageId(), FunctionCallNode->NodeGuid);
+			bool bHasBaseModule = bIsMergeable && BaseEmitter.Emitter != nullptr && MergeManager->HasBaseModule(BaseEmitter, OutputNode->GetUsage(), OutputNode->GetUsageId(), FunctionCallNode->NodeGuid);
 			bCanMoveAndDeleteCache = bHasBaseModule == false;
 		}
 	}
@@ -1198,7 +1197,7 @@ void UNiagaraStackModuleItem::ReassignModuleScript(UNiagaraScript* ModuleScript)
 		if (NewName != OldName)
 		{
 			UNiagaraSystem& System = GetSystemViewModel()->GetSystem();
-			UNiagaraEmitter* Emitter = GetEmitterViewModel().IsValid() ? GetEmitterViewModel()->GetEmitter() : nullptr;
+			FVersionedNiagaraEmitter Emitter = GetEmitterViewModel().IsValid() ? GetEmitterViewModel()->GetEmitter() : FVersionedNiagaraEmitter();
 			FNiagaraStackGraphUtilities::RenameReferencingParameters(&System, Emitter, *FunctionCallNode, OldName, NewName);
 			FunctionCallNode->RefreshFromExternalChanges();
 			FunctionCallNode->MarkNodeRequiresSynchronization(TEXT("Module script reassigned."), true);
@@ -1407,13 +1406,12 @@ void UNiagaraStackModuleItem::Delete()
 {
 	checkf(CanMoveAndDelete(), TEXT("This module can't be deleted"));
 	const FNiagaraEmitterHandle* EmitterHandle = GetEmitterViewModel().IsValid()
-		? FNiagaraEditorUtilities::GetEmitterHandleForEmitter(GetSystemViewModel()->GetSystem(), *GetEmitterViewModel()->GetEmitter())
+		? FNiagaraEditorUtilities::GetEmitterHandleForEmitter(GetSystemViewModel()->GetSystem(), GetEmitterViewModel()->GetEmitter())
 		: nullptr;
 	FGuid EmitterHandleId = EmitterHandle != nullptr ? EmitterHandle->GetId() : FGuid();
 
 	TArray<TWeakObjectPtr<UNiagaraNodeInput>> RemovedNodes;
-	bool bRemoved = FNiagaraStackGraphUtilities::RemoveModuleFromStack(GetSystemViewModel()->GetSystem(), EmitterHandleId, *FunctionCallNode, RemovedNodes);
-	if (bRemoved)
+	if (FNiagaraStackGraphUtilities::RemoveModuleFromStack(GetSystemViewModel()->GetSystem(), EmitterHandleId, *FunctionCallNode, RemovedNodes))
 	{
 		UNiagaraGraph* Graph = FunctionCallNode->GetNiagaraGraph();
 		Graph->NotifyGraphNeedsRecompile();
diff --git a/Engine/Plugins/FX/Niagara/Source/NiagaraEditor/Private/ViewModels/Stack/NiagaraStackParameterStoreEntry.cpp b/Engine/Plugins/FX/Niagara/Source/NiagaraEditor/Private/ViewModels/Stack/NiagaraStackParameterStoreEntry.cpp
--- a/Engine/Plugins/FX/Niagara/Source/NiagaraEditor/Private/ViewModels/Stack/NiagaraStackParameterStoreEntry.cpp
+++ b/Engine/Plugins/FX/Niagara/Source/NiagaraEditor/Private/ViewModels/Stack/NiagaraStackParameterStoreEntry.cpp
@@ -14,12 +14,10 @@
 #include "ViewModels/Stack/NiagaraStackObject.h"
 #include "NiagaraNodeParameterMapGet.h"
 #include "ViewModels/Stack/NiagaraStackGraphUtilities.h"
-#include "NiagaraStackEditorData.h"
 
 #include "ScopedTransaction.h"
 #include "Editor.h"
 #include "UObject/StructOnScope.h"
-#include "Modules/ModuleManager.h"
 #include "AssetRegistry/AssetRegistryModule.h"
 #include "AssetRegistry/ARFilter.h"
 #include "EdGraph/EdGraphPin.h"
@@ -194,7 +192,7 @@ TArray<UEdGraphPin*> UNiagaraStackParameterStoreEntry::GetOwningPins()
 	auto EmitterHandles = GetSystemViewModel()->GetSystem().GetEmitterHandles();
 	for (const FNiagaraEmitterHandle& Handle : EmitterHandles)
 	{
-		UNiagaraGraph* EmitterGraph = CastChecked<UNiagaraScriptSource>(Handle.GetInstance()->GraphSource)->NodeGraph;
+		UNiagaraGraph* EmitterGraph = CastChecked<UNiagaraScriptSource>(Handle.GetEmitterData()->GraphSource)->NodeGraph;
 		GraphsToCheck.Add(EmitterGraph);
 	}
 	TArray<UEdGraphPin*> OwningPins;
@@ -260,13 +258,19 @@ void UNiagaraStackParameterStoreEntry::OnRenamed(FText NewName)
 		}
 
 		UObject* OwnerObj = Owner.Get();
-		UNiagaraSystem* System = Cast<UNiagaraSystem>(OwnerObj);
-		UNiagaraEmitter* Emitter = Cast<UNiagaraEmitter>(OwnerObj);
-		if (System)
+		if (UNiagaraSystem* System = Cast<UNiagaraSystem>(OwnerObj))
+		{
 			System->HandleVariableRenamed(FNiagaraVariable(InputType, ParameterName), FNiagaraVariable(InputType, VariableName), true);
-		if (Emitter)
-			Emitter->HandleVariableRenamed(FNiagaraVariable(InputType, ParameterName), FNiagaraVariable(InputType, VariableName), true);
-		
+		}
+		else if (UNiagaraEmitter* Emitter = Cast<UNiagaraEmitter>(OwnerObj))
+		{
+			FGuid Version;
+			if (TSharedPtr<FNiagaraEmitterViewModel> ViewModel = GetEmitterViewModel())
+			{
+				Version = ViewModel->GetEmitter().Version;
+			}
+			Emitter->HandleVariableRenamed(FNiagaraVariable(InputType, ParameterName), FNiagaraVariable(InputType, VariableName), true, Version);	
+		}		
 
 		ParameterName = VariableName;
 		DisplayName = FText::FromName(ParameterName);
@@ -417,13 +421,19 @@ void UNiagaraStackParameterStoreEntry::Delete()
 
 	// Update anything that was referencing that parameter
 	UObject* OwnerObj = Owner.Get();
-	UNiagaraSystem* System = Cast<UNiagaraSystem>(OwnerObj);
-	UNiagaraEmitter* Emitter = Cast<UNiagaraEmitter>(OwnerObj);
-	if (System)
+	if (UNiagaraSystem* System = Cast<UNiagaraSystem>(OwnerObj))
+	{
 		System->HandleVariableRemoved(FNiagaraVariable(InputType, ParameterName),  true);
-	if (Emitter)
-		Emitter->HandleVariableRemoved(FNiagaraVariable(InputType, ParameterName), true);
-
+	}
+	else if (UNiagaraEmitter* Emitter = Cast<UNiagaraEmitter>(OwnerObj))
+	{
+		FGuid Version;
+		if (TSharedPtr<FNiagaraEmitterViewModel> ViewModel = GetEmitterViewModel())
+		{
+			Version = ViewModel->GetEmitter().Version;
+		}
+		Emitter->HandleVariableRemoved(FNiagaraVariable(InputType, ParameterName), true, Version);
+	}	
 
 	// Notify the system view model that the DI has been modifier.
 	if (CachedSystemViewModel.IsValid() && DataInterface != nullptr)
diff --git a/Engine/Plugins/FX/Niagara/Source/NiagaraEditor/Private/ViewModels/Stack/NiagaraStackRenderItemGroup.cpp b/Engine/Plugins/FX/Niagara/Source/NiagaraEditor/Private/ViewModels/Stack/NiagaraStackRenderItemGroup.cpp
--- a/Engine/Plugins/FX/Niagara/Source/NiagaraEditor/Private/ViewModels/Stack/NiagaraStackRenderItemGroup.cpp
+++ b/Engine/Plugins/FX/Niagara/Source/NiagaraEditor/Private/ViewModels/Stack/NiagaraStackRenderItemGroup.cpp
@@ -88,22 +88,21 @@ public:
 
 		FScopedTransaction ScopedTransaction(LOCTEXT("AddNewRendererTransaction", "Add new renderer"));
 
-		UNiagaraEmitter* Emitter = EmitterViewModelPinned->GetEmitter();
-		Emitter->Modify();
-		UNiagaraRendererProperties* RendererProperties = NewObject<UNiagaraRendererProperties>(Emitter, RenderAddAction->GetRendererClass(), NAME_None, RF_Transactional);
-		Emitter->AddRenderer(RendererProperties);
+		FVersionedNiagaraEmitter VersionedEmitter = EmitterViewModelPinned->GetEmitter();
+		UNiagaraRendererProperties* RendererProperties = NewObject<UNiagaraRendererProperties>(VersionedEmitter.Emitter, RenderAddAction->GetRendererClass(), NAME_None, RF_Transactional);
+		VersionedEmitter.Emitter->AddRenderer(RendererProperties, VersionedEmitter.Version);
 
 		bool bVarsAdded = false;
-		TArray<FNiagaraVariable> MissingAttributes = UNiagaraStackRendererItem::GetMissingVariables(RendererProperties, Emitter);
+		TArray<FNiagaraVariable> MissingAttributes = UNiagaraStackRendererItem::GetMissingVariables(RendererProperties, VersionedEmitter.GetEmitterData());
 		for (int32 i = 0; i < MissingAttributes.Num(); i++)
 		{
-			if (UNiagaraStackRendererItem::AddMissingVariable(Emitter, MissingAttributes[i]))
+			if (UNiagaraStackRendererItem::AddMissingVariable(VersionedEmitter.GetEmitterData(), MissingAttributes[i]))
 			{
 				bVarsAdded = true;
 			}
 		}
 
-		FNiagaraSystemUpdateContext SystemUpdate(Emitter, true);
+		FNiagaraSystemUpdateContext SystemUpdate(VersionedEmitter, true);
 
 		if (bVarsAdded)
 		{
@@ -127,8 +126,8 @@ void UNiagaraStackRenderItemGroup::Initialize(FRequiredEntryData InRequiredEntry
 	FText ToolTip = LOCTEXT("RendererGroupTooltip", "Describes how we should display/present each particle. Note that this doesn't have to be visual. Multiple renderers are supported. Order in this stack is not necessarily relevant to draw order.");
 	AddUtilities = MakeShared<FRenderItemGroupAddUtilities>(InRequiredEntryData.EmitterViewModel.ToSharedRef(), FRenderItemGroupAddUtilities::FOnItemAdded::CreateUObject(this, &UNiagaraStackRenderItemGroup::OnRendererAdded));
 	Super::Initialize(InRequiredEntryData, DisplayName, ToolTip, AddUtilities.Get());
-	EmitterWeak = GetEmitterViewModel()->GetEmitter();
-	EmitterWeak->OnRenderersChanged().AddUObject(this, &UNiagaraStackRenderItemGroup::EmitterRenderersChanged);
+	EmitterWeak = GetEmitterViewModel()->GetEmitter().ToWeakPtr();
+	EmitterWeak.Emitter->OnRenderersChanged().AddUObject(this, &UNiagaraStackRenderItemGroup::EmitterRenderersChanged);
 }
 
 bool UNiagaraStackRenderItemGroup::TestCanPasteWithMessage(const UNiagaraClipboardContent* ClipboardContent, FText& OutMessage) const
@@ -151,12 +150,12 @@ void UNiagaraStackRenderItemGroup::Paste(const UNiagaraClipboardContent* Clipboa
 {
 	if (EmitterWeak.IsValid())
 	{
-		UNiagaraEmitter* Emitter = EmitterWeak.Get();
+		UNiagaraEmitter* Emitter = EmitterWeak.Emitter.Get();
 		for (const UNiagaraRendererProperties* ClipboardRenderer : ClipboardContent->Renderers)
 		{
 			if (ClipboardRenderer != nullptr)
 			{
-				EmitterWeak->AddRenderer(ClipboardRenderer->StaticDuplicateWithNewMergeId(Emitter));
+				EmitterWeak.Emitter->AddRenderer(ClipboardRenderer->StaticDuplicateWithNewMergeId(Emitter), EmitterWeak.Version);
 			}
 		}
 	}
@@ -165,7 +164,7 @@ void UNiagaraStackRenderItemGroup::Paste(const UNiagaraClipboardContent* Clipboa
 void UNiagaraStackRenderItemGroup::RefreshChildrenInternal(const TArray<UNiagaraStackEntry*>& CurrentChildren, TArray<UNiagaraStackEntry*>& NewChildren, TArray<FStackIssue>& NewIssues)
 {
 	int32 RendererIndex = 0;
-	for (UNiagaraRendererProperties* RendererProperties : GetEmitterViewModel()->GetEmitter()->GetRenderers())
+	for (UNiagaraRendererProperties* RendererProperties : GetEmitterViewModel()->GetEmitter().GetEmitterData()->GetRenderers())
 	{
 		UNiagaraStackRendererItem* RendererItem = FindCurrentChildOfTypeByPredicate<UNiagaraStackRendererItem>(CurrentChildren,
 			[=](UNiagaraStackRendererItem* CurrentRendererItem) { return CurrentRendererItem->GetRendererProperties() == RendererProperties; });
@@ -217,7 +216,7 @@ void UNiagaraStackRenderItemGroup::FinalizeInternal()
 {
 	if (EmitterWeak.IsValid())
 	{
-		EmitterWeak->OnRenderersChanged().RemoveAll(this);
+		EmitterWeak.Emitter->OnRenderersChanged().RemoveAll(this);
 	}
 	Super::FinalizeInternal();
 }
diff --git a/Engine/Plugins/FX/Niagara/Source/NiagaraEditor/Private/ViewModels/Stack/NiagaraStackRendererItem.cpp b/Engine/Plugins/FX/Niagara/Source/NiagaraEditor/Private/ViewModels/Stack/NiagaraStackRendererItem.cpp
--- a/Engine/Plugins/FX/Niagara/Source/NiagaraEditor/Private/ViewModels/Stack/NiagaraStackRendererItem.cpp
+++ b/Engine/Plugins/FX/Niagara/Source/NiagaraEditor/Private/ViewModels/Stack/NiagaraStackRendererItem.cpp
@@ -8,7 +8,6 @@
 #include "NiagaraScript.h"
 #include "ViewModels/NiagaraSystemViewModel.h"
 #include "ViewModels/NiagaraEmitterViewModel.h"
-#include "ViewModels/NiagaraScriptViewModel.h"
 #include "Internationalization/Internationalization.h"
 #include "NiagaraNodeAssignment.h"
 #include "NiagaraNodeOutput.h"
@@ -16,9 +15,7 @@
 #include "NiagaraScriptGraphViewModel.h"
 #include "NiagaraGraph.h"
 #include "ScopedTransaction.h"
-#include "NiagaraEmitter.h"
 #include "NiagaraScriptSource.h"
-#include "ScopedTransaction.h"
 #include "ViewModels/Stack/NiagaraStackErrorItem.h"
 #include "Framework/Notifications/NotificationManager.h"
 #include "Widgets/Notifications/SNotificationList.h"
@@ -27,7 +24,6 @@
 #include "NiagaraClipboard.h"
 #include "NiagaraEmitterEditorData.h"
 
-#include "Styling/SlateIconFinder.h"
 #include "ViewModels/NiagaraEmitterHandleViewModel.h"
 #include "ViewModels/Stack/NiagaraStackViewModel.h"
 
@@ -56,11 +52,11 @@ void UNiagaraStackRendererItem::FinalizeInternal()
 	Super::FinalizeInternal();
 }
 
-TArray<FNiagaraVariable> UNiagaraStackRendererItem::GetMissingVariables(UNiagaraRendererProperties* RendererProperties, UNiagaraEmitter* Emitter)
+TArray<FNiagaraVariable> UNiagaraStackRendererItem::GetMissingVariables(UNiagaraRendererProperties* RendererProperties, const FVersionedNiagaraEmitterData* EmitterData)
 {
 	TArray<FNiagaraVariable> MissingAttributes;
 	const TArray<FNiagaraVariable>& RequiredAttrs = RendererProperties->GetRequiredAttributes();
-	const UNiagaraScript* Script = Emitter->SpawnScriptProps.Script;
+	const UNiagaraScript* Script = EmitterData->SpawnScriptProps.Script;
 	if (Script != nullptr && Script->IsReadyToRun(ENiagaraSimTarget::CPUSim))
 	{
 		MissingAttributes.Empty();
@@ -84,9 +80,9 @@ TArray<FNiagaraVariable> UNiagaraStackRendererItem::GetMissingVariables(UNiagara
 	return MissingAttributes;
 }
 
-bool UNiagaraStackRendererItem::AddMissingVariable(UNiagaraEmitter* Emitter, const FNiagaraVariable& Variable)
+bool UNiagaraStackRendererItem::AddMissingVariable(const FVersionedNiagaraEmitterData* EmitterData, const FNiagaraVariable& Variable)
 {
-	UNiagaraScript* Script = Emitter->SpawnScriptProps.Script;
+	UNiagaraScript* Script = EmitterData->SpawnScriptProps.Script;
 	if (!Script)
 	{
 		return false;
@@ -145,7 +141,7 @@ bool UNiagaraStackRendererItem::IsExcludedFromScalability() const
 
 bool UNiagaraStackRendererItem::IsOwningEmitterExcludedFromScalability() const
 {
-	return GetEmitterViewModel().IsValid() ? !GetEmitterViewModel()->GetEmitter()->IsAllowedByScalability() : false;
+	return GetEmitterViewModel().IsValid() ? !GetEmitterViewModel()->GetEmitter().GetEmitterData()->IsAllowedByScalability() : false;
 }
 
 UNiagaraRendererProperties* UNiagaraStackRendererItem::GetRendererProperties()
@@ -261,12 +257,12 @@ FText UNiagaraStackRendererItem::GetDeleteTransactionText() const
 
 void UNiagaraStackRendererItem::Delete()
 {
-	UNiagaraEmitter* Emitter = GetEmitterViewModel()->GetEmitter();
+	FVersionedNiagaraEmitter VersionedEmitter = GetEmitterViewModel()->GetEmitter();
 
 	UNiagaraRendererProperties* Renderer = RendererProperties.Get();
-	Emitter->Modify();
-	Emitter->RemoveRenderer(Renderer);
-	if (UNiagaraEmitterEditorData* Data = Cast<UNiagaraEmitterEditorData>(Emitter->GetEditorData()))
+	VersionedEmitter.Emitter->Modify();
+	VersionedEmitter.Emitter->RemoveRenderer(Renderer, VersionedEmitter.Version);
+	if (UNiagaraEmitterEditorData* Data = Cast<UNiagaraEmitterEditorData>(VersionedEmitter.GetEmitterData()->GetEditorData()))
 	{
 		Data->GetStackEditorData().Modify();
 		Data->GetStackEditorData().SetStackEntryDisplayName(GetStackEditorDataKey(), FText());
@@ -294,8 +290,8 @@ bool UNiagaraStackRendererItem::HasBaseRenderer() const
 		if (bHasBaseRendererCache.IsSet() == false)
 		{
 			TSharedRef<FNiagaraScriptMergeManager> MergeManager = FNiagaraScriptMergeManager::Get();
-			const UNiagaraEmitter* BaseEmitter = GetEmitterViewModel()->GetEmitter()->GetParent();
-			bHasBaseRendererCache = BaseEmitter != nullptr && MergeManager->HasBaseRenderer(*BaseEmitter, RendererProperties->GetMergeId());
+			FVersionedNiagaraEmitter BaseEmitter = GetEmitterViewModel()->GetEmitter().GetEmitterData()->GetParent();
+			bHasBaseRendererCache = BaseEmitter.Emitter != nullptr && MergeManager->HasBaseRenderer(BaseEmitter, RendererProperties->GetMergeId());
 		}
 		return bHasBaseRendererCache.GetValue();
 	}
@@ -309,8 +305,8 @@ bool UNiagaraStackRendererItem::TestCanResetToBaseWithMessage(FText& OutCanReset
 		if (HasBaseRenderer())
 		{
 			TSharedRef<FNiagaraScriptMergeManager> MergeManager = FNiagaraScriptMergeManager::Get();
-			const UNiagaraEmitter* BaseEmitter = GetEmitterViewModel()->GetEmitter()->GetParent();
-			bCanResetToBaseCache = BaseEmitter != nullptr && MergeManager->IsRendererDifferentFromBase(*GetEmitterViewModel()->GetEmitter(), *BaseEmitter, RendererProperties->GetMergeId());
+			FVersionedNiagaraEmitter BaseEmitter = GetEmitterViewModel()->GetEmitter().GetEmitterData()->GetParent();
+			bCanResetToBaseCache = BaseEmitter.Emitter != nullptr && MergeManager->IsRendererDifferentFromBase(GetEmitterViewModel()->GetEmitter(), BaseEmitter, RendererProperties->GetMergeId());
 		}
 		else
 		{
@@ -335,8 +331,8 @@ void UNiagaraStackRendererItem::ResetToBase()
 	if (TestCanResetToBaseWithMessage(Unused))
 	{
 		TSharedRef<FNiagaraScriptMergeManager> MergeManager = FNiagaraScriptMergeManager::Get();
-		const UNiagaraEmitter* BaseEmitter = GetEmitterViewModel()->GetEmitter()->GetParent();
-		MergeManager->ResetRendererToBase(*GetEmitterViewModel()->GetEmitter(), *BaseEmitter, RendererProperties->GetMergeId());
+		FVersionedNiagaraEmitter BaseEmitter = GetEmitterViewModel()->GetEmitter().GetEmitterData()->GetParent();
+		MergeManager->ResetRendererToBase(GetEmitterViewModel()->GetEmitter(), BaseEmitter, RendererProperties->GetMergeId());
 		ModifiedGroupItemsDelegate.Broadcast();
 	}
 }
@@ -391,7 +387,7 @@ void UNiagaraStackRendererItem::RefreshChildrenInternal(const TArray<UNiagaraSta
 	}
 
 	NewChildren.Add(RendererObject);
-	MissingAttributes = GetMissingVariables(RendererProperties.Get(), GetEmitterViewModel()->GetEmitter());
+	MissingAttributes = GetMissingVariables(RendererProperties.Get(), GetEmitterViewModel()->GetEmitter().GetEmitterData());
 	bHasBaseRendererCache.Reset();
 	bCanResetToBaseCache.Reset();
 	DisplayNameCache.Reset();
@@ -421,6 +417,7 @@ void UNiagaraStackRendererItem::RefreshIssues(TArray<FStackIssue>& NewIssues)
 		NewIssues.Empty();
 		return;
 	}
+	FVersionedNiagaraEmitterData* EmitterData = GetEmitterViewModel()->GetEmitter().GetEmitterData();
 	for (FNiagaraVariable Attribute : MissingAttributes)
 	{
 		FText FixDescription = LOCTEXT("AddMissingVariable", "Add missing variable");
@@ -429,7 +426,7 @@ void UNiagaraStackRendererItem::RefreshIssues(TArray<FStackIssue>& NewIssues)
 			FStackIssueFixDelegate::CreateLambda([=]()
 		{
 			FScopedTransaction ScopedTransaction(FixDescription);
-			if (AddMissingVariable(GetEmitterViewModel()->GetEmitter(), Attribute))
+			if (AddMissingVariable(EmitterData, Attribute))
 			{
 				FNotificationInfo Info(FText::Format(LOCTEXT("AddedVariableForFix", "Added {0} to the Spawn script to support the renderer."), FText::FromName(Attribute.GetName())));
 				Info.ExpireDuration = 5.0f;
@@ -450,13 +447,13 @@ void UNiagaraStackRendererItem::RefreshIssues(TArray<FStackIssue>& NewIssues)
 		NewIssues.Add(MissingAttributeError);
 	}
 
-	if (RendererProperties->GetIsEnabled() && !RendererProperties->IsSimTargetSupported(GetEmitterViewModel()->GetEmitter()->SimTarget))
+	if (RendererProperties->GetIsEnabled() && !RendererProperties->IsSimTargetSupported(EmitterData->SimTarget))
 	{
 		
 		FStackIssue TargetSupportError(
 			EStackIssueSeverity::Error,
 			LOCTEXT("FailedRendererDueToSimTarget", "Renderer incompatible with SimTarget mode."),
-			FText::Format(LOCTEXT("FailedRendererDueToSimTargetLong", "Renderer incompatible with SimTarget mode \"{0}\"."), FText::FromName(UEnum::GetValueAsName(GetEmitterViewModel()->GetEmitter()->SimTarget))),
+			FText::Format(LOCTEXT("FailedRendererDueToSimTargetLong", "Renderer incompatible with SimTarget mode \"{0}\"."), FText::FromName(UEnum::GetValueAsName(EmitterData->SimTarget))),
 			GetStackEditorDataKey(),
 			false);
 
diff --git a/Engine/Plugins/FX/Niagara/Source/NiagaraEditor/Private/ViewModels/Stack/NiagaraStackRoot.cpp b/Engine/Plugins/FX/Niagara/Source/NiagaraEditor/Private/ViewModels/Stack/NiagaraStackRoot.cpp
--- a/Engine/Plugins/FX/Niagara/Source/NiagaraEditor/Private/ViewModels/Stack/NiagaraStackRoot.cpp
+++ b/Engine/Plugins/FX/Niagara/Source/NiagaraEditor/Private/ViewModels/Stack/NiagaraStackRoot.cpp
@@ -9,8 +9,8 @@
 #include "ViewModels/Stack/NiagaraStackSimulationStageGroup.h"
 #include "ViewModels/NiagaraSystemViewModel.h"
 #include "ViewModels/NiagaraEmitterViewModel.h"
-#include "NiagaraSystemScriptViewModel.h"
 #include "ViewModels/NiagaraScriptViewModel.h"
+#include "ViewModels/NiagaraSystemScriptViewModel.h"
 #include "NiagaraSystem.h"
 #include "NiagaraSystemEditorData.h"
 #include "NiagaraEmitterEditorData.h"
@@ -49,12 +49,12 @@ void UNiagaraStackRoot::Initialize(FRequiredEntryData InRequiredEntryData, bool
 	if (bInIncludeEmitterInformation && GetEmitterViewModel())
 	{
 		GetEmitterViewModel()->GetOrCreateEditorData().OnSummaryViewStateChanged().AddUObject(this, &UNiagaraStackRoot::OnSummaryViewStateChanged);
-	}	
+	}
 }
 
 void UNiagaraStackRoot::FinalizeInternal()
 {
-	if (bIncludeEmitterInformation && GetEmitterViewModel())
+	if (bIncludeEmitterInformation && GetEmitterViewModel() && GetEmitterViewModel()->GetEmitter().GetEmitterData())
 	{
 		GetEmitterViewModel()->GetOrCreateEditorData().OnSummaryViewStateChanged().RemoveAll(this);
 	}	
@@ -230,7 +230,7 @@ void UNiagaraStackRoot::RefreshChildrenInternal(const TArray<UNiagaraStackEntry*
 			NewChildren.Add(ParticleSpawnGroup);
 			NewChildren.Add(ParticleUpdateGroup);
 
-			for (const FNiagaraEventScriptProperties& EventScriptProperties : GetEmitterViewModel()->GetEmitter()->GetEventHandlers())
+			for (const FNiagaraEventScriptProperties& EventScriptProperties : GetEmitterViewModel()->GetEmitter().GetEmitterData()->GetEventHandlers())
 			{
 				UNiagaraStackEventScriptItemGroup* EventHandlerGroup = FindCurrentChildOfTypeByPredicate<UNiagaraStackEventScriptItemGroup>(CurrentChildren,
 					[&](UNiagaraStackEventScriptItemGroup* CurrentEventHandlerGroup) { return CurrentEventHandlerGroup->GetScriptUsageId() == EventScriptProperties.Script->GetUsageId(); });
@@ -248,7 +248,7 @@ void UNiagaraStackRoot::RefreshChildrenInternal(const TArray<UNiagaraStackEntry*
 				NewChildren.Add(EventHandlerGroup);
 			}
 
-			for (UNiagaraSimulationStageBase* SimulationStage : GetEmitterViewModel()->GetEmitter()->GetSimulationStages())
+			for (UNiagaraSimulationStageBase* SimulationStage : GetEmitterViewModel()->GetEmitter().GetEmitterData()->GetSimulationStages())
 			{
 				UNiagaraStackSimulationStageGroup* SimulationStageGroup = FindCurrentChildOfTypeByPredicate<UNiagaraStackSimulationStageGroup>(CurrentChildren,
 					[SimulationStage](UNiagaraStackSimulationStageGroup* CurrentSimulationStageGroup) { return CurrentSimulationStageGroup->GetSimulationStage() == SimulationStage; });
@@ -279,7 +279,10 @@ void UNiagaraStackRoot::EmitterArraysChanged()
 
 void UNiagaraStackRoot::OnSummaryViewStateChanged()
 {
-	RefreshChildren();
+	if (!IsFinalized())
+	{
+		RefreshChildren();
+	}
 }
 
 #undef LOCTEXT_NAMESPACE
diff --git a/Engine/Plugins/FX/Niagara/Source/NiagaraEditor/Private/ViewModels/Stack/NiagaraStackScriptItemGroup.cpp b/Engine/Plugins/FX/Niagara/Source/NiagaraEditor/Private/ViewModels/Stack/NiagaraStackScriptItemGroup.cpp
--- a/Engine/Plugins/FX/Niagara/Source/NiagaraEditor/Private/ViewModels/Stack/NiagaraStackScriptItemGroup.cpp
+++ b/Engine/Plugins/FX/Niagara/Source/NiagaraEditor/Private/ViewModels/Stack/NiagaraStackScriptItemGroup.cpp
@@ -208,7 +208,6 @@ private:
 	{
 	}
 
-private:
 	FText DisplayName;
 	TArray<FString> Categories;
 	FText Description;
@@ -219,7 +218,6 @@ private:
 	bool bRenameParameterOnAdd;
 	FAssetData ModuleAssetData;
 	UNiagaraScript* ModuleScript;
-	bool bIsMaterialParameterModuleAction;
 	bool bIsNewScratchModuleAction;
 	bool bIsNewSetSpecificModuleAction;
 	FNiagaraActionSourceData SourceData;
@@ -407,7 +405,7 @@ void UNiagaraStackScriptItemGroup::Initialize(
 	}
 	else if (GetEmitterViewModel().IsValid())
 	{
-		OwningParticleScriptWeak = GetEmitterViewModel()->GetEmitter()->GetScript(ScriptUsage, ScriptUsageId);
+		OwningParticleScriptWeak = GetEmitterViewModel()->GetEmitter().GetEmitterData()->GetScript(ScriptUsage, ScriptUsageId);
 		if (OwningParticleScriptWeak.IsValid())
 		{
 			OwningParticleScriptWeak->OnVMScriptCompiled().AddUObject(this, &UNiagaraStackScriptItemGroup::OnParticleScriptCompiled);
@@ -1062,7 +1060,7 @@ TOptional<UNiagaraStackEntry::FDropRequestResponse> UNiagaraStackScriptItemGroup
 
 	UNiagaraStackModuleItem* SourceModuleItem = CastChecked<UNiagaraStackModuleItem>(StackEntryDragDropOp->GetDraggedEntries()[0]);
 	const FNiagaraEmitterHandle* SourceEmitterHandle = SourceModuleItem->GetEmitterViewModel().IsValid()
-		? FNiagaraEditorUtilities::GetEmitterHandleForEmitter(SourceModuleItem->GetSystemViewModel()->GetSystem(), *SourceModuleItem->GetEmitterViewModel()->GetEmitter())
+		? FNiagaraEditorUtilities::GetEmitterHandleForEmitter(SourceModuleItem->GetSystemViewModel()->GetSystem(), SourceModuleItem->GetEmitterViewModel()->GetEmitter())
 		: nullptr;
 	FGuid SourceEmitterHandleId = SourceEmitterHandle != nullptr
 		? SourceEmitterHandle->GetId()
@@ -1071,9 +1069,8 @@ TOptional<UNiagaraStackEntry::FDropRequestResponse> UNiagaraStackScriptItemGroup
 	UNiagaraScript* SourceModuleScript = FNiagaraEditorUtilities::GetScriptFromSystem(SourceModuleItem->GetSystemViewModel()->GetSystem(), SourceEmitterHandleId,
 		SourceModuleOutputNode->GetUsage(), SourceModuleOutputNode->GetUsageId());
 
-	const UNiagaraStackModuleItem* TargetModuleItem = Cast<UNiagaraStackModuleItem>(&TargetEntry);
 	const FNiagaraEmitterHandle* TargetEmitterHandle = GetEmitterViewModel().IsValid()
-		? FNiagaraEditorUtilities::GetEmitterHandleForEmitter(GetSystemViewModel()->GetSystem(), *GetEmitterViewModel()->GetEmitter())
+		? FNiagaraEditorUtilities::GetEmitterHandleForEmitter(GetSystemViewModel()->GetSystem(), GetEmitterViewModel()->GetEmitter())
 		: nullptr;
 	FGuid TargetEmitterHandleId = TargetEmitterHandle != nullptr
 		? TargetEmitterHandle->GetId()
@@ -1221,7 +1218,7 @@ void GatherRenamedModuleOutputs(
 	TArray<TPair<const UNiagaraClipboardFunction*, UNiagaraNodeFunctionCall*>> InClipboardFunctionAndNodeFunctionPairs,
 	TMap<FName, FName>& OutOldModuleOutputNameToNewModuleOutputNameMap)
 {
-	UNiagaraEmitter* Emitter = InOwnerEntry->GetEmitterViewModel().IsValid() ? InOwnerEntry->GetEmitterViewModel()->GetEmitter() : nullptr;
+	FVersionedNiagaraEmitter Emitter = InOwnerEntry->GetEmitterViewModel().IsValid() ? InOwnerEntry->GetEmitterViewModel()->GetEmitter() : FVersionedNiagaraEmitter();
 
 	for (TPair<const UNiagaraClipboardFunction*, UNiagaraNodeFunctionCall*>& ClipboardFunctionAndNodeFunctionPair : InClipboardFunctionAndNodeFunctionPairs)
 	{
diff --git a/Engine/Plugins/FX/Niagara/Source/NiagaraEditor/Private/ViewModels/Stack/NiagaraStackSimulationStageGroup.cpp b/Engine/Plugins/FX/Niagara/Source/NiagaraEditor/Private/ViewModels/Stack/NiagaraStackSimulationStageGroup.cpp
--- a/Engine/Plugins/FX/Niagara/Source/NiagaraEditor/Private/ViewModels/Stack/NiagaraStackSimulationStageGroup.cpp
+++ b/Engine/Plugins/FX/Niagara/Source/NiagaraEditor/Private/ViewModels/Stack/NiagaraStackSimulationStageGroup.cpp
@@ -4,7 +4,6 @@
 #include "NiagaraEmitter.h"
 #include "NiagaraSimulationStageBase.h"
 #include "ViewModels/Stack/NiagaraStackObject.h"
-#include "ViewModels/Stack/NiagaraStackModuleItem.h"
 #include "ViewModels/NiagaraScriptViewModel.h"
 #include "NiagaraScriptGraphViewModel.h"
 #include "ViewModels/NiagaraEmitterViewModel.h"
@@ -14,14 +13,11 @@
 #include "EdGraphSchema_Niagara.h"
 #include "ViewModels/Stack/NiagaraStackErrorItem.h"
 #include "NiagaraScriptSource.h"
-#include "ViewModels/Stack/NiagaraStackGraphUtilities.h"
 #include "NiagaraScriptMergeManager.h"
-#include "NiagaraStackEditorData.h"
 #include "ViewModels/NiagaraSystemViewModel.h"
 
 #include "Internationalization/Internationalization.h"
 #include "ScopedTransaction.h"
-#include "IDetailTreeNode.h"
 
 #define LOCTEXT_NAMESPACE "UNiagaraStackSimulationStageGroup"
 
@@ -55,12 +51,12 @@ bool UNiagaraStackSimulationStagePropertiesItem::TestCanResetToBaseWithMessage(F
 	{
 		if (HasBaseSimulationStage())
 		{
-			UNiagaraEmitter* Emitter = GetEmitterViewModel()->GetEmitter();
-			const UNiagaraEmitter* BaseEmitter = Emitter->GetParent();
-			if (BaseEmitter != nullptr && Emitter != BaseEmitter)
+			FVersionedNiagaraEmitter VersionedEmitter = GetEmitterViewModel()->GetEmitter();
+			FVersionedNiagaraEmitter BaseEmitter = VersionedEmitter.GetEmitterData()->GetParent();
+			if (BaseEmitter.Emitter != nullptr && VersionedEmitter.Emitter != BaseEmitter.Emitter)
 			{
 				TSharedRef<FNiagaraScriptMergeManager> MergeManager = FNiagaraScriptMergeManager::Get();
-				bCanResetToBaseCache = MergeManager->IsSimulationStagePropertySetDifferentFromBase(*Emitter, *BaseEmitter, SimulationStage->Script->GetUsageId());
+				bCanResetToBaseCache = MergeManager->IsSimulationStagePropertySetDifferentFromBase(VersionedEmitter, BaseEmitter, SimulationStage->Script->GetUsageId());
 			}
 			else
 			{
@@ -89,10 +85,10 @@ void UNiagaraStackSimulationStagePropertiesItem::ResetToBase()
 	FText Unused;
 	if (TestCanResetToBaseWithMessage(Unused))
 	{
-		UNiagaraEmitter* Emitter = GetEmitterViewModel()->GetEmitter();
-		const UNiagaraEmitter* BaseEmitter = Emitter->GetParent();
+		FVersionedNiagaraEmitter VersionedNiagaraEmitter = GetEmitterViewModel()->GetEmitter();
+		FVersionedNiagaraEmitter BaseEmitter = VersionedNiagaraEmitter.GetEmitterData()->GetParent();
 		TSharedRef<FNiagaraScriptMergeManager> MergeManager = FNiagaraScriptMergeManager::Get();
-		MergeManager->ResetSimulationStagePropertySetToBase(*Emitter, *BaseEmitter, SimulationStage->Script->GetUsageId());
+		MergeManager->ResetSimulationStagePropertySetToBase(VersionedNiagaraEmitter, BaseEmitter, SimulationStage->Script->GetUsageId());
 		RefreshChildren();
 	}
 }
@@ -108,8 +104,8 @@ void UNiagaraStackSimulationStagePropertiesItem::RefreshChildrenInternal(const T
 
 	if ( SimulationStage.IsValid() )
 	{
-		UNiagaraEmitter* Emitter = GetEmitterViewModel()->GetEmitter();
-		if ( Emitter && (Emitter->SimTarget != ENiagaraSimTarget::GPUComputeSim) && SimulationStage->bEnabled )
+		FVersionedNiagaraEmitterData* EmitterData = GetEmitterViewModel()->GetEmitter().GetEmitterData();
+		if (EmitterData && (EmitterData->SimTarget != ENiagaraSimTarget::GPUComputeSim) && SimulationStage->bEnabled )
 		{
 			TArray<FStackIssueFix> IssueFixes;
 			IssueFixes.Emplace(
@@ -119,13 +115,13 @@ void UNiagaraStackSimulationStagePropertiesItem::RefreshChildrenInternal(const T
 			IssueFixes.Emplace(
 				LOCTEXT("SetGpuSimulationFix", "Set GPU simulation"),
 				FStackIssueFixDelegate::CreateLambda(
-					[WeakEmitter=TWeakObjectPtr<UNiagaraEmitter>(GetEmitterViewModel()->GetEmitter())]()
+					[WeakEmitter=GetEmitterViewModel()->GetEmitter().ToWeakPtr()]()
 					{
-						if ( UNiagaraEmitter* NiagaraEmitter = WeakEmitter.Get() )
+						if (WeakEmitter.IsValid())
 						{
 							FScopedTransaction Transaction(LOCTEXT("SetGpuSimulation", "Set Gpu Simulation"));
-							NiagaraEmitter->Modify();
-							NiagaraEmitter->SimTarget = ENiagaraSimTarget::GPUComputeSim;
+							WeakEmitter.Emitter.Get()->Modify();
+							WeakEmitter.GetEmitterData()->SimTarget = ENiagaraSimTarget::GPUComputeSim;
 						}
 					}
 				)
@@ -159,12 +155,12 @@ bool UNiagaraStackSimulationStagePropertiesItem::HasBaseSimulationStage() const
 {
 	if (bHasBaseSimulationStageCache.IsSet() == false)
 	{
-		UNiagaraEmitter* Emitter = GetEmitterViewModel()->GetEmitter();
-		const UNiagaraEmitter* BaseEmitter = Emitter->GetParent();
-		if (BaseEmitter != nullptr && Emitter != BaseEmitter)
+		FVersionedNiagaraEmitter VersionedEmitter = GetEmitterViewModel()->GetEmitter();
+		FVersionedNiagaraEmitter BaseEmitter = VersionedEmitter.GetEmitterData()->GetParent();
+		if (BaseEmitter.Emitter != nullptr && VersionedEmitter.Emitter != BaseEmitter.Emitter)
 		{
 			TSharedRef<FNiagaraScriptMergeManager> MergeManager = FNiagaraScriptMergeManager::Get();
-			bHasBaseSimulationStageCache = MergeManager->HasBaseSimulationStage(*BaseEmitter, SimulationStage->Script->GetUsageId());
+			bHasBaseSimulationStageCache = MergeManager->HasBaseSimulationStage(BaseEmitter, SimulationStage->Script->GetUsageId());
 		}
 		else
 		{
@@ -233,8 +229,6 @@ void UNiagaraStackSimulationStageGroup::RefreshChildrenInternal(const TArray<UNi
 {
 	bHasBaseSimulationStageCache.Reset();
 
-	UNiagaraEmitter* Emitter = GetEmitterViewModel()->GetEmitter();
-
 	SetDisplayName(FText::FromName(SimulationStage->SimulationStageName));
 
 	if (SimulationStageProperties == nullptr)
@@ -293,7 +287,7 @@ TOptional<UNiagaraStackEntry::FDropRequestResponse> UNiagaraStackSimulationStage
 			return Super::CanDropInternal(DropRequest);
 		}
 
-		UNiagaraEmitter* OwningEmitter = GetEmitterViewModel()->GetEmitter();
+		FVersionedNiagaraEmitterData* OwningEmitter = GetEmitterViewModel()->GetEmitter().GetEmitterData();
 		int32 SourceIndex = OwningEmitter->GetSimulationStages().IndexOfByKey(SourceSimulationStageGroup->GetSimulationStage());
 		if (SourceIndex == INDEX_NONE)
 		{
@@ -328,15 +322,16 @@ TOptional<UNiagaraStackEntry::FDropRequestResponse> UNiagaraStackSimulationStage
 		UNiagaraStackSimulationStageGroup* SourceSimulationStageGroup = Cast<UNiagaraStackSimulationStageGroup>(StackEntryDragDropOp->GetDraggedEntries()[0]);
 		if (SourceSimulationStageGroup)
 		{
-			UNiagaraEmitter* OwningEmitter = GetEmitterViewModel()->GetEmitter();
-			int32 SourceIndex = OwningEmitter->GetSimulationStages().IndexOfByKey(SourceSimulationStageGroup->GetSimulationStage());
+			FVersionedNiagaraEmitter OwningEmitter = GetEmitterViewModel()->GetEmitter();
+			FVersionedNiagaraEmitterData* EmitterData = OwningEmitter.GetEmitterData();
+			int32 SourceIndex = EmitterData->GetSimulationStages().IndexOfByKey(SourceSimulationStageGroup->GetSimulationStage());
 			if (SourceIndex != INDEX_NONE)
 			{
 				int32 TargetOffset = DropRequest.DropZone == EItemDropZone::BelowItem ? 1 : 0;
-				int32 TargetIndex = OwningEmitter->GetSimulationStages().IndexOfByKey(SimulationStage.Get()) + TargetOffset;
+				int32 TargetIndex = EmitterData->GetSimulationStages().IndexOfByKey(SimulationStage.Get()) + TargetOffset;
 
 				FScopedTransaction Transaction(FText::Format(LOCTEXT("MoveSimulationStage", "Move Shader Stage {0}"), GetDisplayName()));
-				OwningEmitter->MoveSimulationStageToIndex(SourceSimulationStageGroup->GetSimulationStage(), TargetIndex);
+				OwningEmitter.Emitter->MoveSimulationStageToIndex(SourceSimulationStageGroup->GetSimulationStage(), TargetIndex, OwningEmitter.Version);
 
 				OnRequestFullRefreshDeferred().Broadcast();
 				return FDropRequestResponse(DropRequest.DropZone, FText());
@@ -365,8 +360,9 @@ void UNiagaraStackSimulationStageGroup::Delete()
 	TSharedPtr<FNiagaraScriptViewModel> ScriptViewModelPinned = ScriptViewModel.Pin();
 	checkf(ScriptViewModelPinned.IsValid(), TEXT("Can not delete when the script view model has been deleted."));
 
-	UNiagaraEmitter* Emitter = GetEmitterViewModel()->GetEmitter();
-	UNiagaraScriptSource* Source = Cast<UNiagaraScriptSource>(Emitter->GraphSource);
+	FVersionedNiagaraEmitter VersionedEmitter = GetEmitterViewModel()->GetEmitter();
+	FVersionedNiagaraEmitterData* EmitterData = VersionedEmitter.GetEmitterData();
+	UNiagaraScriptSource* Source = Cast<UNiagaraScriptSource>(EmitterData->GraphSource);
 
 	if (!Source || !Source->NodeGraph)
 	{
@@ -375,7 +371,7 @@ void UNiagaraStackSimulationStageGroup::Delete()
 
 	FScopedTransaction Transaction(FText::Format(LOCTEXT("DeleteSimulationStage", "Delete {0}"), GetDisplayName()));
 
-	Emitter->Modify();
+	VersionedEmitter.Emitter->Modify();
 	Source->NodeGraph->Modify();
 	TArray<UNiagaraNode*> SimulationStageNodes;
 	Source->NodeGraph->BuildTraversal(SimulationStageNodes, GetScriptUsage(), GetScriptUsageId());
@@ -385,7 +381,7 @@ void UNiagaraStackSimulationStageGroup::Delete()
 	}
 	
 	// First, remove the simulation stage object.
-	Emitter->RemoveSimulationStage(SimulationStage.Get());
+	VersionedEmitter.Emitter->RemoveSimulationStage(SimulationStage.Get(), VersionedEmitter.Version);
 	
 	// Now remove all graph nodes associated with the simulation stage.
 	for (UNiagaraNode* Node : SimulationStageNodes)
@@ -395,7 +391,7 @@ void UNiagaraStackSimulationStageGroup::Delete()
 
 	// Set the emitter here to that the internal state of the view model is updated.
 	// TODO: Move the logic for managing additional scripts into the emitter view model or script view model.
-	ScriptViewModelPinned->SetScripts(Emitter);
+	ScriptViewModelPinned->SetScripts(VersionedEmitter);
 	
 	OnModifiedSimulationStagesDelegate.ExecuteIfBound();
 }
@@ -419,8 +415,8 @@ bool UNiagaraStackSimulationStageGroup::HasBaseSimulationStage() const
 {
 	if (bHasBaseSimulationStageCache.IsSet() == false)
 	{
-		const UNiagaraEmitter* BaseEmitter = GetEmitterViewModel()->GetEmitter()->GetParent();
-		bHasBaseSimulationStageCache = BaseEmitter != nullptr && FNiagaraScriptMergeManager::Get()->HasBaseSimulationStage(*BaseEmitter, GetScriptUsageId());
+		FVersionedNiagaraEmitter BaseEmitter = GetEmitterViewModel()->GetEmitter().GetEmitterData()->GetParent();
+		bHasBaseSimulationStageCache = BaseEmitter.Emitter != nullptr && FNiagaraScriptMergeManager::Get()->HasBaseSimulationStage(BaseEmitter, GetScriptUsageId());
 	}
 	return bHasBaseSimulationStageCache.GetValue();
 }
diff --git a/Engine/Plugins/FX/Niagara/Source/NiagaraEditor/Private/ViewModels/Stack/NiagaraStackSummaryViewInputCollection.cpp b/Engine/Plugins/FX/Niagara/Source/NiagaraEditor/Private/ViewModels/Stack/NiagaraStackSummaryViewInputCollection.cpp
--- a/Engine/Plugins/FX/Niagara/Source/NiagaraEditor/Private/ViewModels/Stack/NiagaraStackSummaryViewInputCollection.cpp
+++ b/Engine/Plugins/FX/Niagara/Source/NiagaraEditor/Private/ViewModels/Stack/NiagaraStackSummaryViewInputCollection.cpp
@@ -3,22 +3,16 @@
 #include "ViewModels/Stack/NiagaraStackSummaryViewInputCollection.h"
 #include "EdGraphSchema_Niagara.h"
 #include "NiagaraClipboard.h"
-#include "NiagaraDataInterface.h"
 #include "NiagaraGraph.h"
 #include "NiagaraNodeFunctionCall.h"
 #include "NiagaraNodeOutput.h"
-#include "NiagaraNodeParameterMapSet.h"
 #include "NiagaraSimulationStageBase.h"
-#include "ScopedTransaction.h"
-#include "EdGraph/EdGraphPin.h"
 #include "ViewModels/NiagaraEmitterViewModel.h"
 #include "ViewModels/NiagaraSystemViewModel.h"
 #include "ViewModels/NiagaraScriptViewModel.h"
 #include "ViewModels/NiagaraScriptGraphViewModel.h"
 #include "ViewModels/Stack/NiagaraStackFunctionInput.h"
 #include "ViewModels/Stack/NiagaraStackGraphUtilities.h"
-#include "ViewModels/Stack/NiagaraStackInputCategory.h"
-#include "ViewModels/Stack/NiagaraStackModuleItem.h"
 #include "ViewModels/Stack/NiagaraStackFunctionInputCollection.h"
 #include "ViewModels/Stack/NiagaraStackSimulationStageGroup.h"
 
@@ -27,13 +21,12 @@
 
 
 UNiagaraStackSummaryViewObject::UNiagaraStackSummaryViewObject()
-	: Emitter(nullptr)
 {
 }
 
-void UNiagaraStackSummaryViewObject::Initialize(FRequiredEntryData InRequiredEntryData, UNiagaraEmitter* InEmitter, FString InOwningStackItemEditorDataKey)
+void UNiagaraStackSummaryViewObject::Initialize(FRequiredEntryData InRequiredEntryData, FVersionedNiagaraEmitterWeakPtr InEmitter, FString InOwningStackItemEditorDataKey)
 {
-	checkf(Emitter == nullptr, TEXT("Can only initialize once."));
+	checkf(Emitter.Emitter == nullptr, TEXT("Can only initialize once."));
 	FString ObjectStackEditorDataKey = FString::Printf(TEXT("%s-FilteredView"), *InOwningStackItemEditorDataKey);
 	Super::Initialize(InRequiredEntryData, InOwningStackItemEditorDataKey, ObjectStackEditorDataKey);
 
@@ -43,7 +36,10 @@ void UNiagaraStackSummaryViewObject::Initialize(FRequiredEntryData InRequiredEnt
 
 void UNiagaraStackSummaryViewObject::FinalizeInternal()
 {
-	GetEmitterViewModel()->GetOrCreateEditorData().OnSummaryViewStateChanged().RemoveAll(this);
+	if (GetEmitterViewModel() && GetEmitterViewModel()->GetEmitter().GetEmitterData())
+	{
+		GetEmitterViewModel()->GetOrCreateEditorData().OnSummaryViewStateChanged().RemoveAll(this);
+	}
 	Super::FinalizeInternal();
 }
 
@@ -72,12 +68,12 @@ void UNiagaraStackSummaryViewObject::RefreshChildrenInternal(const TArray<UNiaga
 	AppendEmitterCategory(State, ScriptViewModelPinned, ENiagaraScriptUsage::ParticleSpawnScript, FGuid(), NewIssues);
 	AppendEmitterCategory(State, ScriptViewModelPinned, ENiagaraScriptUsage::ParticleUpdateScript, FGuid(), NewIssues);
 
-	for (const FNiagaraEventScriptProperties& EventScriptProperties : ViewModel->GetEmitter()->GetEventHandlers())
+	for (const FNiagaraEventScriptProperties& EventScriptProperties : ViewModel->GetEmitter().GetEmitterData()->GetEventHandlers())
 	{
 		AppendEmitterCategory(State, ScriptViewModelPinned, ENiagaraScriptUsage::ParticleEventScript, EventScriptProperties.Script->GetUsageId(), NewIssues);
 	}
 	
-	for (UNiagaraSimulationStageBase* SimulationStage : GetEmitterViewModel()->GetEmitter()->GetSimulationStages())
+	for (UNiagaraSimulationStageBase* SimulationStage : GetEmitterViewModel()->GetEmitter().GetEmitterData()->GetSimulationStages())
 	{
 		AppendEmitterCategory(State, ScriptViewModelPinned, ENiagaraScriptUsage::ParticleSimulationStageScript, SimulationStage->Script->GetUsageId(), NewIssues);
 	}	
diff --git a/Engine/Plugins/FX/Niagara/Source/NiagaraEditor/Private/ViewModels/Stack/NiagaraStackViewModel.cpp b/Engine/Plugins/FX/Niagara/Source/NiagaraEditor/Private/ViewModels/Stack/NiagaraStackViewModel.cpp
--- a/Engine/Plugins/FX/Niagara/Source/NiagaraEditor/Private/ViewModels/Stack/NiagaraStackViewModel.cpp
+++ b/Engine/Plugins/FX/Niagara/Source/NiagaraEditor/Private/ViewModels/Stack/NiagaraStackViewModel.cpp
@@ -134,7 +134,7 @@ void UNiagaraStackViewModel::InitializeWithViewModels(TSharedPtr<FNiagaraSystemV
 		{
 			EmitterViewModel->OnParentRemoved().AddUObject(this, &UNiagaraStackViewModel::EmitterParentRemoved);
 
-			if (UNiagaraEmitter* Emitter = EmitterViewModel->GetEmitter())
+			if (UNiagaraEmitter* Emitter = EmitterViewModel->GetEmitter().Emitter)
 			{
 				EmitterViewModel->GetOrCreateEditorData().OnSummaryViewStateChanged().AddUObject(this, &UNiagaraStackViewModel::EntryRequestFullRefreshDeferred);
 			}
@@ -207,15 +207,18 @@ void UNiagaraStackViewModel::Reset()
 	}
 	RootEntries.Empty();
 
-	if (EmitterHandleViewModel.IsValid())
+	TSharedPtr<FNiagaraEmitterViewModel> EmitterViewModel = EmitterHandleViewModel.IsValid() ? EmitterHandleViewModel.Pin()->GetEmitterViewModel() : TSharedPtr<FNiagaraEmitterViewModel>();
+	if (EmitterViewModel.IsValid())
 	{
-		EmitterHandleViewModel.Reset();
+		EmitterViewModel->OnParentRemoved().RemoveAll(this);
+		if (EmitterViewModel->GetEmitter().GetEmitterData())
+		{
+			EmitterViewModel->GetOrCreateEditorData().OnSummaryViewStateChanged().RemoveAll(this);
+		}
 	}
 
-	if (SystemViewModel.IsValid())
-	{
-		SystemViewModel.Reset();
-	}
+	EmitterHandleViewModel.Reset();
+	SystemViewModel.Reset();
 
 	TopLevelViewModels.Empty();
 
diff --git a/Engine/Plugins/FX/Niagara/Source/NiagaraEditor/Private/Widgets/SNiagaraBakerWidget.cpp b/Engine/Plugins/FX/Niagara/Source/NiagaraEditor/Private/Widgets/SNiagaraBakerWidget.cpp
--- a/Engine/Plugins/FX/Niagara/Source/NiagaraEditor/Private/Widgets/SNiagaraBakerWidget.cpp
+++ b/Engine/Plugins/FX/Niagara/Source/NiagaraEditor/Private/Widgets/SNiagaraBakerWidget.cpp
@@ -760,13 +760,13 @@ void SNiagaraBakerWidget::FindWarnings()
 			}
 			for (const FNiagaraEmitterHandle& EmitterHandle : NiagaraSystem->GetEmitterHandles() )
 			{
-				UNiagaraEmitter* NiagaraEmitter = EmitterHandle.GetInstance();
-				if (NiagaraEmitter == nullptr || NiagaraEmitter->bDeterminism)
+				FVersionedNiagaraEmitterData* EmitterData = EmitterHandle.GetInstance().GetEmitterData();
+				if (EmitterData == nullptr || EmitterData->bDeterminism)
 				{
 					continue;
 				}
 
-				FoundWarnings.Emplace(FText::Format(LOCTEXT("EmitterDeterminismFormat", "Emitter '{0}' is not set to deterministic, results will vary each bake"), FText::FromString(NiagaraEmitter->GetUniqueEmitterName())));
+				FoundWarnings.Emplace(FText::Format(LOCTEXT("EmitterDeterminismFormat", "Emitter '{0}' is not set to deterministic, results will vary each bake"), FText::FromString(EmitterHandle.GetInstance().Emitter->GetUniqueEmitterName())));
 			}
 		}
 	}
diff --git a/Engine/Plugins/FX/Niagara/Source/NiagaraEditor/Private/Widgets/SNiagaraDebugger.cpp b/Engine/Plugins/FX/Niagara/Source/NiagaraEditor/Private/Widgets/SNiagaraDebugger.cpp
--- a/Engine/Plugins/FX/Niagara/Source/NiagaraEditor/Private/Widgets/SNiagaraDebugger.cpp
+++ b/Engine/Plugins/FX/Niagara/Source/NiagaraEditor/Private/Widgets/SNiagaraDebugger.cpp
@@ -885,7 +885,7 @@ void SNiagaraDebugger::InvokeDebugger(FNiagaraEmitterHandle& InEmitterHandle)
 		Content->FocusDebugTab();
 
 		UNiagaraDebugHUDSettings* HudSettings = GetMutableDefault<UNiagaraDebugHUDSettings>();
-		if (HudSettings && InEmitterHandle.IsValid() && InEmitterHandle.GetInstance())
+		if (HudSettings && InEmitterHandle.IsValid() && InEmitterHandle.GetInstance().Emitter)
 		{
 			HudSettings->Data.bComponentFilterEnabled = false;
 			HudSettings->Data.ComponentFilter = TEXT("");
@@ -893,7 +893,7 @@ void SNiagaraDebugger::InvokeDebugger(FNiagaraEmitterHandle& InEmitterHandle)
 			HudSettings->Data.ActorFilter = TEXT("");
 			HudSettings->Data.bActorFilterEnabled = false;
 
-			UNiagaraSystem* System = Cast<UNiagaraSystem>(InEmitterHandle.GetInstance()->GetOuter());
+			UNiagaraSystem* System = Cast<UNiagaraSystem>(InEmitterHandle.GetInstance().Emitter->GetOuter());
 			if (System)
 			{
 				HudSettings->Data.SystemFilter = System->GetName();
@@ -1006,7 +1006,7 @@ void SNiagaraDebugger::InvokeDebugger(UNiagaraSystem* InSystem, TArray<FNiagaraE
 						HudSettings->Data.ParticlesVariables.Add(NewVar);
 				}
 
-				if (InSelectedHandles.Num() > 0 && InSelectedHandles[0].GetInstance() && InSelectedHandles[0].GetInstance()->SimTarget == ENiagaraSimTarget::GPUComputeSim)
+				if (InSelectedHandles.Num() > 0 && InSelectedHandles[0].GetEmitterData() && InSelectedHandles[0].GetEmitterData()->SimTarget == ENiagaraSimTarget::GPUComputeSim)
 					HudSettings->Data.bEnableGpuParticleReadback = true;
 			}
 
diff --git a/Engine/Plugins/FX/Niagara/Source/NiagaraEditor/Private/Widgets/SNiagaraEmitterVersionWidget.h b/Engine/Plugins/FX/Niagara/Source/NiagaraEditor/Private/Widgets/SNiagaraEmitterVersionWidget.h
new file mode 100644
--- /dev/null
+++ b/Engine/Plugins/FX/Niagara/Source/NiagaraEditor/Private/Widgets/SNiagaraEmitterVersionWidget.h
@@ -0,0 +1,24 @@
+﻿// Copyright Epic Games, Inc. All Rights Reserved.
+
+#pragma once
+
+#include "CoreMinimal.h"
+
+#include "NiagaraActions.h"
+#include "SNiagaraVersionWidget.h"
+
+
+
+class NIAGARAEDITOR_API SNiagaraEmitterVersionWidget : public SNiagaraVersionWidget
+{
+public:
+	void Construct(const FArguments& InArgs, UNiagaraEmitter* InEmitter, UNiagaraVersionMetaData* InMetadata, const FString& BasePackageName);
+
+protected:
+	virtual FText GetInfoHeaderText() const override;
+	virtual void ExecuteSaveAsAssetAction(FNiagaraAssetVersion AssetVersion) override;
+
+private:
+	UNiagaraEmitter* Emitter = nullptr;
+	FString BasePackageName;
+};
diff --git a/Engine/Plugins/FX/Niagara/Source/NiagaraEditor/Private/Widgets/SNiagaraGeneratedCodeView.cpp b/Engine/Plugins/FX/Niagara/Source/NiagaraEditor/Private/Widgets/SNiagaraGeneratedCodeView.cpp
--- a/Engine/Plugins/FX/Niagara/Source/NiagaraEditor/Private/Widgets/SNiagaraGeneratedCodeView.cpp
+++ b/Engine/Plugins/FX/Niagara/Source/NiagaraEditor/Private/Widgets/SNiagaraGeneratedCodeView.cpp
@@ -370,7 +370,7 @@ void SNiagaraGeneratedCodeView::UpdateUI()
 		if (Handle)
 		{
 			TArray<UNiagaraScript*> EmitterScripts;
-			Handle->GetInstance()->GetScripts(EmitterScripts);
+			Handle->GetEmitterData()->GetScripts(EmitterScripts);
 			Scripts.Append(EmitterScripts);
 		}
 	}
diff --git a/Engine/Plugins/FX/Niagara/Source/NiagaraEditor/Private/Widgets/SNiagaraParameterDefinitionsPanel.h b/Engine/Plugins/FX/Niagara/Source/NiagaraEditor/Private/Widgets/SNiagaraParameterDefinitionsPanel.h
--- a/Engine/Plugins/FX/Niagara/Source/NiagaraEditor/Private/Widgets/SNiagaraParameterDefinitionsPanel.h
+++ b/Engine/Plugins/FX/Niagara/Source/NiagaraEditor/Private/Widgets/SNiagaraParameterDefinitionsPanel.h
@@ -40,7 +40,7 @@ public:
 	{}
 	SLATE_END_ARGS();
 
-	NIAGARAEDITOR_API ~SNiagaraParameterDefinitionsPanel();
+	NIAGARAEDITOR_API virtual ~SNiagaraParameterDefinitionsPanel() override;
 
 	NIAGARAEDITOR_API void Construct(const FArguments& InArgs, const TSharedPtr<INiagaraParameterDefinitionsPanelViewModel>& InParameterDefinitionsPanelViewModel, const TSharedPtr<FUICommandList>& InToolkitCommands);
 
