Added versioning for Niagara emitters

diff --git a/Engine/Plugins/FX/Niagara/Source/NiagaraEditor/Private/Widgets/SNiagaraScriptVersionWidget.cpp b/Engine/Plugins/FX/Niagara/Source/NiagaraEditor/Private/Widgets/SNiagaraScriptVersionWidget.cpp
--- a/Engine/Plugins/FX/Niagara/Source/NiagaraEditor/Private/Widgets/SNiagaraScriptVersionWidget.cpp
+++ b/Engine/Plugins/FX/Niagara/Source/NiagaraEditor/Private/Widgets/SNiagaraScriptVersionWidget.cpp
@@ -3,35 +3,20 @@
 #include "SNiagaraScriptVersionWidget.h"
 
 #include "AssetToolsModule.h"
-#include "DetailLayoutBuilder.h"
 #include "Editor.h"
-#include "IDetailsView.h"
-#include "PropertyEditorModule.h"
-#include "SGraphActionMenu.h"
-#include "SlateOptMacros.h"
 #include "Editor/EditorEngine.h"
-#include "Framework/MultiBox/MultiBoxBuilder.h"
 #include "Modules/ModuleManager.h"
 #include "NiagaraEditor/Private/NiagaraVersionMetaData.h"
 #include "NiagaraEditor/Public/NiagaraActions.h"
 #include "Subsystems/AssetEditorSubsystem.h"
-#include "Widgets/Images/SImage.h"
-#include "Widgets/Input/SButton.h"
-#include "Widgets/Layout/SBox.h"
-#include "Widgets/Layout/SSeparator.h"
-#include "Widgets/Layout/SWidgetSwitcher.h"
-#include "Widgets/Text/STextBlock.h"
-#include "Widgets/Views/STableRow.h"
-#include "Widgets/Views/STableViewBase.h"
 
 #define LOCTEXT_NAMESPACE "NiagaraScriptVersionWidget"
 
-BEGIN_SLATE_FUNCTION_BUILD_OPTIMIZATION
-
-FText SNiagaraScriptVersionWidget::FormatVersionLabel(const FNiagaraAssetVersion& Version) const
+void SNiagaraScriptVersionWidget::Construct(const FArguments& InArgs, UNiagaraScript* InScript,	UNiagaraVersionMetaData* InMetadata, const FString& InBasePackageName)
 {
-	FText Format = Version == Script->GetExposedVersion() ? FText::FromString("{0}.{1} (exposed)") : FText::FromString("{0}.{1}");
-	return FText::Format(Format, Version.MajorVersion, Version.MinorVersion);
+	Script = InScript;
+	BasePackageName = InBasePackageName;
+	SNiagaraVersionWidget::Construct(InArgs, InScript, InMetadata);
 }
 
 FText SNiagaraScriptVersionWidget::GetInfoHeaderText() const
@@ -40,277 +25,18 @@ FText SNiagaraScriptVersionWidget::GetInfoHeaderText() const
 	{
 		return LOCTEXT("NiagaraManageVersionInfoHeader", "Here you can see and manage available script versions. Adding versions allows you to make changes without breaking existing assets.\nThe exposed version is the one that users will see when adding the module, function or dynamic input to the stack. If the exposed version has a higher major version than the one used in an asset, the user will see a prompt to upgrade to the new version.");	
 	}
-	return LOCTEXT("NiagaraManageVersionInfoHeaderDisabled", "Versioning is not yet enabled for this script. Enabling versioning allows you to make changes to the module without breaking exising usages. But it also forces users to manually upgrade to new versions, as changes are no longer pushed out automatically.\n\nDo you want to enable versioning for this script? This will convert the current script to version 1.0 and automatically expose it.");
-}
-
-FNiagaraVersionMenuAction::FNiagaraVersionMenuAction(FText InNodeCategory, FText InMenuDesc, FText InToolTip, const int32 InGrouping, FText InKeywords, FOnExecuteStackAction InAction, int32 InSectionID, FNiagaraAssetVersion InVersion)
-	: FNiagaraMenuAction(MoveTemp(InNodeCategory), MoveTemp(InMenuDesc), MoveTemp(InToolTip), InGrouping, MoveTemp(InKeywords), MoveTemp(InAction), InSectionID)
-	, AssetVersion(InVersion)
-{}
-
-void SNiagaraScriptVersionWidget::Construct(const FArguments& InArgs, UNiagaraScript* InScript, UNiagaraVersionMetaData* InMetadata)
-{
-	OnVersionDataChanged = InArgs._OnVersionDataChanged;
-	OnChangeToVersion = InArgs._OnChangeToVersion;
-	Script = InScript;
-	VersionMetadata = InMetadata; 
-	
-	FDetailsViewArgs DetailsArgs;
-	DetailsArgs.bHideSelectionTip = true;
-	DetailsArgs.bAllowSearch = false;
-	DetailsArgs.NotifyHook = this;
-
-	FPropertyEditorModule& PropertyModule = FModuleManager::LoadModuleChecked<FPropertyEditorModule>("PropertyEditor");
-	VersionSettingsDetails = PropertyModule.CreateDetailView(DetailsArgs);
-	VersionSettingsDetails->SetObject(VersionMetadata);
-
-	// the list of available versions
-	SAssignNew(VersionListWidget, SGraphActionMenu)
-	    .OnActionSelected(this, &SNiagaraScriptVersionWidget::OnActionSelected)
-		.OnCollectAllActions(this, &SNiagaraScriptVersionWidget::CollectAllVersionActions)
-	    .OnCollectStaticSections_Lambda([](TArray<int32>& StaticSectionIDs) {StaticSectionIDs.Add(1);})
-	    .OnGetSectionTitle_Lambda([](int32) {return LOCTEXT("NiagaraVersionManagementTitle", "Versions");})
-	    .AutoExpandActionMenu(true)
-	    .ShowFilterTextBox(false)
-	    .OnCreateCustomRowExpander_Static(&SNiagaraScriptVersionWidget::CreateCustomActionExpander)
-	    .UseSectionStyling(true)
-		.AlphaSortItems(false)
-	    .OnGetSectionWidget(this, &SNiagaraScriptVersionWidget::GetVersionSelectionHeaderWidget)
-		.OnContextMenuOpening(this, &SNiagaraScriptVersionWidget::OnVersionContextMenuOpening)
-	    .OnCreateWidgetForAction_Lambda([](const FCreateWidgetForActionData* InData)
-	    {
-	        TSharedPtr<FNiagaraMenuAction> NiagaraAction = StaticCastSharedPtr<FNiagaraMenuAction>(InData->Action);
-	        return SNew(STextBlock).Text(NiagaraAction->GetMenuDescription());
-	    });
-	FText ItemName = FormatVersionLabel(Script->GetExposedVersion());
-	VersionListWidget->SelectItemByName(FName(ItemName.ToString()));
-
-	ChildSlot
-    [
-	 SNew(SBox)
-        .MinDesiredWidth(400)
-		.Padding(10)
-        [
-            SNew(SVerticalBox)
-
-            // the top description text
-            + SVerticalBox::Slot()
-            .AutoHeight()
-            [
-            	SNew(STextBlock)
-            	.AutoWrapText(true)
-				.Text(this, &SNiagaraScriptVersionWidget::GetInfoHeaderText)
-            ]
-
-            // the main part of the widget
-            + SVerticalBox::Slot()
-            .FillHeight(1)
-            .Padding(FMargin(5, 15))
-            .VAlign(VAlign_Fill)
-            [
-            	// display button to enable versioning or versioning details
-	            SNew(SWidgetSwitcher)
-		        .WidgetIndex(this, &SNiagaraScriptVersionWidget::GetDetailWidgetIndex)
-		        + SWidgetSwitcher::Slot()
-		        [
-					SNew(SOverlay)
-
-					+ SOverlay::Slot()
-					.VAlign(VAlign_Top)
-					.HAlign(HAlign_Center)
-					[
-						// enable versioning button
-				        SNew(SButton)
-		                .ForegroundColor(FAppStyle::GetSlateColor("DefaultForeground"))
-		                .ButtonStyle(FAppStyle::Get(), "FlatButton.Dark")
-		                .ContentPadding(FMargin(6, 2))
-		                .TextStyle(FAppStyle::Get(), "ContentBrowser.TopBar.Font")
-		                .Text(LOCTEXT("NiagaraAEnableVersioning", "Enable versioning"))
-		                .OnClicked(this, &SNiagaraScriptVersionWidget::EnableVersioning)
-	                ]
-		        ]
-		        + SWidgetSwitcher::Slot()
-		        [
-			        SNew(SHorizontalBox)
-
-		            // version selector
-		            + SHorizontalBox::Slot()
-		            .AutoWidth()
-		            [
-		                SNew(SVerticalBox)
-
-		                // version list
-		                + SVerticalBox::Slot()
-		                .FillHeight(1)
-		                [
-		                    SNew(SBox)
-		                    .MaxDesiredWidth(200)
-		                    .MinDesiredWidth(200)
-		                    [
-		                        VersionListWidget.ToSharedRef()
-		                    ]
-		                ]
-		            ]
-
-		            // separator
-		            + SHorizontalBox::Slot()
-		            .AutoWidth()
-		            .Padding(5)
-		            [
-		                SNew(SSeparator)
-		                .Orientation(Orient_Vertical)
-		            ]
-
-		            // version details
-		            + SHorizontalBox::Slot()
-		            [
-		                VersionSettingsDetails.ToSharedRef()
-		            ]
-		        ]
-            ]
-        ]
-    ];
-}
-
-void SNiagaraScriptVersionWidget::NotifyPostChange(const FPropertyChangedEvent& PropertyChangedEvent, FProperty* PropertyThatChanged)
-{
-	// update script data from details view changes
-	FVersionedNiagaraScriptData* ScriptData = Script->GetScriptData(SelectedVersion);
-	check(ScriptData);
-	
-	ScriptData->VersionChangeDescription = VersionMetadata->ChangeDescription;
-	ScriptData->UpdateScriptExecution = VersionMetadata->UpdateScriptExecution;
-	ScriptData->ScriptAsset = VersionMetadata->ScriptAsset;
-	ScriptData->PythonUpdateScript = VersionMetadata->PythonUpdateScript;
-
-	if (PropertyChangedEvent.Property->GetFName() == GET_MEMBER_NAME_CHECKED(UNiagaraVersionMetaData, bIsExposedVersion) && VersionMetadata->bIsExposedVersion)
-	{
-		ExecuteExposeAction(ScriptData->Version);
-	}
-
-	if (PropertyChangedEvent.Property->GetFName() == GET_MEMBER_NAME_CHECKED(UNiagaraVersionMetaData, bIsVisibleInVersionSelector))
-	{
-		ScriptData->Version.bIsVisibleInVersionSelector = VersionMetadata->bIsVisibleInVersionSelector;
-		VersionListWidget->RefreshAllActions(true);
-		VersionInListSelected(ScriptData->Version);
-	}
-
-	OnVersionDataChanged.ExecuteIfBound();
-}
-
-void SNiagaraScriptVersionWidget::SetOnVersionDataChanged(FSimpleDelegate InOnVersionDataChanged)
-{
-	OnVersionDataChanged = InOnVersionDataChanged;
-}
-
-TSharedRef<SExpanderArrow> SNiagaraScriptVersionWidget::CreateCustomActionExpander(const FCustomExpanderData& ActionMenuData)
-{
-	return SNew(SExpanderArrow, ActionMenuData.TableRow)
-        .Visibility(EVisibility::Hidden);
-}
-
-TSharedRef<SWidget> SNiagaraScriptVersionWidget::OnGetAddVersionMenu()
-{
-	FNiagaraAssetVersion LatestVersion = Script->GetAllAvailableVersions().Last();
-	FMenuBuilder MenuBuilder(true, nullptr);
-
-	FText Label = FText::Format(LOCTEXT("NiagaraAddMajorVersion", "New major version ({0}.0)"), LatestVersion.MajorVersion + 1);
-	FText Tooltip = LOCTEXT("NiagaraAddMajorVersion_Tooltip", "Adds a new major version. This should be used for breaking changes (e.g. adding a new parameter).");
-	FUIAction UIAction(FExecuteAction::CreateSP(this, &SNiagaraScriptVersionWidget::AddNewMajorVersion), FCanExecuteAction());
-	MenuBuilder.AddMenuEntry(Label, Tooltip, FSlateIcon(), UIAction);
-
-	Label = FText::Format(LOCTEXT("NiagaraAddMinorVersion", "New minor version ({0}.{1})"), LatestVersion.MajorVersion, LatestVersion.MinorVersion + 1);
-	Tooltip = LOCTEXT("NiagaraAddMinorVersion_Tooltip", "Adds a new minor version. This should be used for non-breaking changes (e.g. adding a comment to the graph).");
-	UIAction = FUIAction(FExecuteAction::CreateSP(this, &SNiagaraScriptVersionWidget::AddNewMinorVersion), FCanExecuteAction());
-	MenuBuilder.AddMenuEntry(Label, Tooltip, FSlateIcon(), UIAction);
-
-	return MenuBuilder.MakeWidget();
-}
-
-TSharedPtr<SWidget> SNiagaraScriptVersionWidget::OnVersionContextMenuOpening()
-{
-	TArray<TSharedPtr<FEdGraphSchemaAction>> SelectedActions;
-	VersionListWidget->GetSelectedActions(SelectedActions);
-	if (SelectedActions.Num() != 1)
-	{
-		return SNullWidget::NullWidget;
-	}
-	FNiagaraVersionMenuAction* SelectedAction = (FNiagaraVersionMenuAction*) SelectedActions[0].Get();
-	
-	FMenuBuilder MenuBuilder(true, nullptr);
-	MenuBuilder.BeginSection("Edit", LOCTEXT("EditMenuHeader", "Edit"));
-	{
-		MenuBuilder.AddMenuEntry(LOCTEXT("ExposeVersion", "Expose version"), LOCTEXT("ExposeVersion_Tooltip", "Exposing this version will make it the default version for any new assets."), FSlateIcon(),
-            FUIAction(
-                FExecuteAction::CreateSP(this, &SNiagaraScriptVersionWidget::ExecuteExposeAction, SelectedAction->AssetVersion),
-                FCanExecuteAction::CreateSP(this, &SNiagaraScriptVersionWidget::CanExecuteExposeAction, SelectedAction->AssetVersion)
-            ));
-
-		MenuBuilder.AddMenuEntry(LOCTEXT("SaveAsAsset", "Save as new asset..."), LOCTEXT("SaveAsAssetVersion_Tooltip", "Creates a new asset with this version as starting point."), FSlateIcon(),
-			FUIAction(FExecuteAction::CreateSP(this, &SNiagaraScriptVersionWidget::ExecuteSaveAsAssetAction, SelectedAction->AssetVersion)));
-
-		MenuBuilder.AddSeparator();
-		
-		MenuBuilder.AddMenuEntry(LOCTEXT("DeleteVersion", "Delete version"), LOCTEXT("DeleteVersion_Tooltip", "Deletes this version and all associated data. This will break existing usages of that version!"), FSlateIcon(),
-            FUIAction(
-                FExecuteAction::CreateSP(this, &SNiagaraScriptVersionWidget::ExecuteDeleteAction, SelectedAction->AssetVersion),
-                FCanExecuteAction::CreateSP(this, &SNiagaraScriptVersionWidget::CanExecuteDeleteAction, SelectedAction->AssetVersion)
-            ));
-	}
-	MenuBuilder.EndSection();
-
-	return MenuBuilder.MakeWidget();
-}
-
-int32 SNiagaraScriptVersionWidget::GetDetailWidgetIndex() const
-{
-	return Script->IsVersioningEnabled() ? 1 : 0;
-}
-
-FReply SNiagaraScriptVersionWidget::EnableVersioning()
-{
-	Script->EnableVersioning();
-	OnVersionDataChanged.ExecuteIfBound();
-	
-	return FReply::Handled();
-}
-
-bool SNiagaraScriptVersionWidget::CanExecuteDeleteAction(FNiagaraAssetVersion AssetVersion)
-{
-	return AssetVersion != Script->GetExposedVersion() && (AssetVersion.MajorVersion != 1 || AssetVersion.MinorVersion != 0);
-}
-
-bool SNiagaraScriptVersionWidget::CanExecuteExposeAction(FNiagaraAssetVersion AssetVersion)
-{
-	return AssetVersion != Script->GetExposedVersion();
-}
-
-void SNiagaraScriptVersionWidget::ExecuteDeleteAction(FNiagaraAssetVersion AssetVersion)
-{
-	OnChangeToVersion.ExecuteIfBound(Script->GetExposedVersion().VersionGuid);
-	Script->DeleteVersion(AssetVersion.VersionGuid);
-
-	VersionListWidget->RefreshAllActions(true);
-	VersionInListSelected(Script->GetExposedVersion());
-	OnVersionDataChanged.ExecuteIfBound();
-}
-
-void SNiagaraScriptVersionWidget::ExecuteExposeAction(FNiagaraAssetVersion AssetVersion)
-{
-	Script->ExposeVersion(AssetVersion.VersionGuid);
-
-	VersionListWidget->RefreshAllActions(true);
-	VersionInListSelected(AssetVersion);
-	OnVersionDataChanged.ExecuteIfBound();
+	return LOCTEXT("NiagaraManageVersionInfoHeaderDisabled", "Versioning is not yet enabled for this script. Enabling versioning allows you to make changes to the module without breaking exising usages. But it also forces users to manually upgrade to new versions, as changes are no longer pushed out automatically.\n\nDo you want to enable versioning for this script? This will convert the current script to version 1.0 and automatically exposes it.");
 }
 
 void SNiagaraScriptVersionWidget::ExecuteSaveAsAssetAction(FNiagaraAssetVersion AssetVersion)
 {
-	UNiagaraScript* ScriptToCopy = Script;
-	const FString StartingPath = FPackageName::GetLongPackagePath(ScriptToCopy->GetOutermost()->GetName());
+	FString Name;
+	FString PackageName;
 	FAssetToolsModule& AssetToolsModule = FModuleManager::LoadModuleChecked<FAssetToolsModule>("AssetTools");
-	UNiagaraScript* NewAssetScript = Cast<UNiagaraScript>(AssetToolsModule.Get().DuplicateAssetWithDialogAndTitle(
-		ScriptToCopy->GetName(), StartingPath, ScriptToCopy, LOCTEXT("SaveVersionAsAssetTitle", "Create Script As")));
+	AssetToolsModule.Get().CreateUniqueAssetName(BasePackageName, TEXT(""), PackageName, Name);
+	
+	UNiagaraScript* NewAssetScript = Cast<UNiagaraScript>(AssetToolsModule.Get().DuplicateAssetWithDialogAndTitle(Name, FPackageName::GetLongPackagePath(PackageName), Script, LOCTEXT("SaveVersionAsAssetTitle", "Create Script As")));
+	
 	if (NewAssetScript != nullptr)
 	{
 		NewAssetScript->DisableVersioning(AssetVersion.VersionGuid);
@@ -318,115 +44,4 @@ void SNiagaraScriptVersionWidget::ExecuteSaveAsAssetAction(FNiagaraAssetVersion
 	}
 }
 
-TSharedRef<ITableRow> SNiagaraScriptVersionWidget::HandleVersionViewGenerateRow(TSharedRef<FNiagaraAssetVersion> Item, const TSharedRef<STableViewBase>& OwnerTable)
-{
-	return
-        SNew(STableRow<TSharedRef<FText>>, OwnerTable)
-        [
-            SNew(STextBlock)
-            .Text(FormatVersionLabel(Item.Get()))
-        ];
-}
-
-void SNiagaraScriptVersionWidget::OnActionSelected(const TArray<TSharedPtr<FEdGraphSchemaAction>>& SelectedActions,	ESelectInfo::Type)
-{
-	for (TSharedPtr<FEdGraphSchemaAction> Action : SelectedActions)
-	{
-		FNiagaraMenuAction* SelectedAction = (FNiagaraMenuAction*) Action.Get();
-		if (SelectedAction)
-		{
-			SelectedAction->ExecuteAction();
-		}
-	}
-}
-
-TSharedRef<SWidget> SNiagaraScriptVersionWidget::GetVersionSelectionHeaderWidget(TSharedRef<SWidget>, int32)
-{
-	// creates the add version button
-	return SNew(SComboButton)
-        .ButtonStyle(FAppStyle::Get(), "RoundButton")
-        .ForegroundColor(FAppStyle::GetSlateColor("DefaultForeground"))
-        .ContentPadding(FMargin(2, 0))
-        .OnGetMenuContent(this, &SNiagaraScriptVersionWidget::OnGetAddVersionMenu)
-        .HAlign(HAlign_Center)
-        .VAlign(VAlign_Center)
-        .HasDownArrow(false)
-        .ButtonContent()
-        [
-            SNew(SHorizontalBox)
-
-            + SHorizontalBox::Slot()
-            .AutoWidth()
-            .Padding(FMargin(0, 1))
-            [
-                SNew(SImage)
-                .Image(FAppStyle::GetBrush("Plus"))
-            ]
-
-            + SHorizontalBox::Slot()
-            .VAlign(VAlign_Center)
-            .AutoWidth()
-            .Padding(FMargin(2,0,0,0))
-            [
-                SNew(STextBlock)
-                .Font(IDetailLayoutBuilder::GetDetailFontBold())
-                .Text(LOCTEXT("NiagaraVersionManagementAdd", "Add version"))
-                .ShadowOffset(FVector2D(1,1))
-            ]
-        ];
-}
-
-void SNiagaraScriptVersionWidget::CollectAllVersionActions(FGraphActionListBuilderBase& OutAllActions)
-{
-	TArray<FNiagaraAssetVersion> AssetVersions = Script->GetAllAvailableVersions();
-	for (FNiagaraAssetVersion& Version : AssetVersions)
-	{
-		FText MenuDesc = FormatVersionLabel(Version);
-		FText Tooltip;
-		auto ExecAction = FNiagaraMenuAction::FOnExecuteStackAction::CreateSP(this, &SNiagaraScriptVersionWidget::VersionInListSelected, Version);
-		TSharedPtr<FNiagaraMenuAction> Action = MakeShared<FNiagaraVersionMenuAction>(FText(), MenuDesc, Tooltip, 0, FText(), ExecAction, 1, Version);
-		OutAllActions.AddAction(Action);
-	}
-}
-
-void SNiagaraScriptVersionWidget::VersionInListSelected(FNiagaraAssetVersion InSelectedVersion)
-{
-	SelectedVersion = InSelectedVersion.VersionGuid;
-	FVersionedNiagaraScriptData* ScriptData = Script->GetScriptData(SelectedVersion);
-	check(ScriptData);
-	
-	VersionMetadata->VersionGuid = InSelectedVersion.VersionGuid;
-	VersionMetadata->ChangeDescription = ScriptData->VersionChangeDescription;
-	VersionMetadata->bIsExposedVersion = InSelectedVersion == Script->GetExposedVersion();
-	VersionMetadata->bIsVisibleInVersionSelector = InSelectedVersion.bIsVisibleInVersionSelector;
-
-	VersionMetadata->UpdateScriptExecution = ScriptData->UpdateScriptExecution;
-	VersionMetadata->ScriptAsset = ScriptData->ScriptAsset;
-	VersionMetadata->PythonUpdateScript = ScriptData->PythonUpdateScript;
-
-	VersionSettingsDetails->SetObject(VersionMetadata, true);
-}
-
-void SNiagaraScriptVersionWidget::AddNewMajorVersion()
-{
-	FNiagaraAssetVersion LatestVersion = Script->GetAllAvailableVersions().Last();
-	FGuid NewVersion = Script->AddNewVersion(LatestVersion.MajorVersion + 1, 0);
-
-	VersionListWidget->RefreshAllActions(true);
-	OnVersionDataChanged.ExecuteIfBound();
-	OnChangeToVersion.ExecuteIfBound(NewVersion);
-}
-
-void SNiagaraScriptVersionWidget::AddNewMinorVersion()
-{
-	FNiagaraAssetVersion LatestVersion = Script->GetAllAvailableVersions().Last();
-	FGuid NewVersion = Script->AddNewVersion(LatestVersion.MajorVersion, LatestVersion.MinorVersion + 1);
-
-	VersionListWidget->RefreshAllActions(true);
-	OnVersionDataChanged.ExecuteIfBound();
-	OnChangeToVersion.ExecuteIfBound(NewVersion);
-}
-
-END_SLATE_FUNCTION_BUILD_OPTIMIZATION
-
 #undef LOCTEXT_NAMESPACE
diff --git a/Engine/Plugins/FX/Niagara/Source/NiagaraEditor/Private/Widgets/SNiagaraScriptVersionWidget.h b/Engine/Plugins/FX/Niagara/Source/NiagaraEditor/Private/Widgets/SNiagaraScriptVersionWidget.h
--- a/Engine/Plugins/FX/Niagara/Source/NiagaraEditor/Private/Widgets/SNiagaraScriptVersionWidget.h
+++ b/Engine/Plugins/FX/Niagara/Source/NiagaraEditor/Private/Widgets/SNiagaraScriptVersionWidget.h
@@ -5,86 +5,18 @@
 #include "CoreMinimal.h"
 
 #include "NiagaraActions.h"
-#include "EdGraph/EdGraphSchema.h"
-#include "Misc/NotifyHook.h"
-#include "Widgets/DeclarativeSyntaxSupport.h"
-#include "Widgets/SCompoundWidget.h"
-#include "Widgets/SWidget.h"
+#include "SNiagaraVersionWidget.h"
 
-struct FCustomExpanderData;
-struct FNiagaraAssetVersion;
-class SGraphActionMenu;
-class SExpanderArrow;
-class STableViewBase;
-class ITableRow;
-class IDetailsView;
-class UNiagaraVersionMetaData;
-class UNiagaraScript;
-
-DECLARE_DELEGATE_OneParam(FOnSwitchToVersionDelegate, FGuid);
-
-struct FNiagaraVersionMenuAction : FNiagaraMenuAction
-{
-	FNiagaraVersionMenuAction() {}
-	FNiagaraVersionMenuAction(FText InNodeCategory, FText InMenuDesc, FText InToolTip, const int32 InGrouping, FText InKeywords, FOnExecuteStackAction InAction, int32 InSectionID, FNiagaraAssetVersion InVersion);
-
-	FNiagaraAssetVersion AssetVersion;
-};
-
-class NIAGARAEDITOR_API SNiagaraScriptVersionWidget : public SCompoundWidget, FNotifyHook
+class NIAGARAEDITOR_API SNiagaraScriptVersionWidget : public SNiagaraVersionWidget
 {
 public:
-	SLATE_BEGIN_ARGS(SNiagaraScriptVersionWidget)
-	{}
+	void Construct(const FArguments& InArgs, UNiagaraScript* InScript, UNiagaraVersionMetaData* InMetadata, const FString& BasePackageName);
 
-	/** Called when the version data of the script was edited by the user */
-	SLATE_EVENT(FSimpleDelegate, OnVersionDataChanged)
+protected:
+	virtual FText GetInfoHeaderText() const override;
+	virtual void ExecuteSaveAsAssetAction(FNiagaraAssetVersion AssetVersion) override;
 
-	/** Called when the user does something that prompts the editor to change the current active version, e.g. delete a version or add a new version */
-    SLATE_EVENT(FOnSwitchToVersionDelegate, OnChangeToVersion)
-	
-	SLATE_END_ARGS()
-
-	void Construct(const FArguments& InArgs, UNiagaraScript* InScript, UNiagaraVersionMetaData* InMetadata);
-
-	//~ Begin FNotifyHook Interface
-	virtual void NotifyPostChange(const FPropertyChangedEvent& PropertyChangedEvent, FProperty* PropertyThatChanged) override;
-	// End of FNotifyHook
-
-	/** See OnVersionDataChanged event */
-	void SetOnVersionDataChanged(FSimpleDelegate InOnVersionDataChanged);
-	
 private:
 	UNiagaraScript* Script = nullptr;
-	bool bAssetVersionsChanged = false;
-	UNiagaraVersionMetaData* VersionMetadata = nullptr;
-	FGuid SelectedVersion;
-	
-	FSimpleDelegate OnVersionDataChanged;
-	FOnSwitchToVersionDelegate OnChangeToVersion;
-	TSharedPtr<IDetailsView> VersionSettingsDetails;
-	TSharedPtr<SGraphActionMenu> VersionListWidget;
-
-	void AddNewMajorVersion();
-	void AddNewMinorVersion();
-	
-	FText FormatVersionLabel(const FNiagaraAssetVersion& Version) const;
-	FText GetInfoHeaderText() const;
-	TSharedRef<ITableRow> HandleVersionViewGenerateRow(TSharedRef<FNiagaraAssetVersion> Item, const TSharedRef<STableViewBase>& OwnerTable);
-	void OnActionSelected(const TArray<TSharedPtr<FEdGraphSchemaAction>>& SelectedActions, ESelectInfo::Type InSelectionType);
-	TSharedRef<SWidget> GetVersionSelectionHeaderWidget(TSharedRef<SWidget> RowWidget, int32 SectionID);
-	void CollectAllVersionActions(FGraphActionListBuilderBase& OutAllActions);
-	void VersionInListSelected(FNiagaraAssetVersion SelectedVersion);
-	static TSharedRef<SExpanderArrow> CreateCustomActionExpander(const FCustomExpanderData& ActionMenuData);
-	TSharedRef<SWidget> OnGetAddVersionMenu();
-	TSharedPtr<SWidget> OnVersionContextMenuOpening();
-	int32 GetDetailWidgetIndex() const;
-	FReply EnableVersioning();
-
-	// context menu actions
-	bool CanExecuteDeleteAction(FNiagaraAssetVersion AssetVersion);
-	bool CanExecuteExposeAction(FNiagaraAssetVersion AssetVersion);
-	void ExecuteDeleteAction(FNiagaraAssetVersion AssetVersion);
-	void ExecuteExposeAction(FNiagaraAssetVersion AssetVersion);
-	void ExecuteSaveAsAssetAction(FNiagaraAssetVersion AssetVersion);
+	FString BasePackageName;
 };
diff --git a/Engine/Plugins/FX/Niagara/Source/NiagaraEditor/Private/Widgets/SNiagaraSpreadsheetView.cpp b/Engine/Plugins/FX/Niagara/Source/NiagaraEditor/Private/Widgets/SNiagaraSpreadsheetView.cpp
--- a/Engine/Plugins/FX/Niagara/Source/NiagaraEditor/Private/Widgets/SNiagaraSpreadsheetView.cpp
+++ b/Engine/Plugins/FX/Niagara/Source/NiagaraEditor/Private/Widgets/SNiagaraSpreadsheetView.cpp
@@ -1,37 +1,27 @@
 // Copyright Epic Games, Inc. All Rights Reserved.
 
 #include "SNiagaraSpreadsheetView.h"
-#include "Textures/SlateIcon.h"
-#include "Framework/Commands/UIAction.h"
-#include "Framework/MultiBox/MultiBoxBuilder.h"
-#include "Widgets/Input/SSearchBox.h"
-#include "Widgets/Input/SButton.h"
-#include "Widgets/Input/SComboButton.h"
+#include "EditorStyleSet.h"
 #include "ISequencer.h"
-#include "ViewModels/NiagaraSystemViewModel.h"
-#include "ViewModels/NiagaraEmitterHandleViewModel.h"
-#include "ViewModels/NiagaraSystemSelectionViewModel.h"
-#include "NiagaraEmitterHandle.h"
+#include "NiagaraComponent.h"
+#include "NiagaraEditorStyle.h"
 #include "NiagaraEmitter.h"
+#include "NiagaraEmitterHandle.h"
 #include "NiagaraScript.h"
+#include "NiagaraSystem.h"
+#include "Framework/Commands/UIAction.h"
+#include "Framework/MultiBox/MultiBoxBuilder.h"
 #include "HAL/PlatformApplicationMisc.h"
+#include "Kismet2/DebuggerCommands.h"
 #include "Styling/AppStyle.h"
 #include "Widgets/Layout/SScrollBox.h"
-#include "UObject/UObjectGlobals.h"
 #include "UObject/Class.h"
-#include "UObject/Package.h"
-#include "SequencerSettings.h"
-#include "NiagaraSystem.h"
-#include "NiagaraEditorStyle.h"
-#include "HAL/PlatformApplicationMisc.h"
 #include "UObject/UObjectIterator.h"
-#include "NiagaraComponent.h"
 #include "ViewModels/NiagaraEmitterHandleViewModel.h"
-#include "NiagaraEmitterHandle.h"
-#include "Kismet2/DebuggerCommands.h"
-#include "Editor/EditorEngine.h"
-#include "UnrealEdGlobals.h"
-#include "EngineGlobals.h"
+#include "ViewModels/NiagaraSystemSelectionViewModel.h"
+#include "ViewModels/NiagaraSystemViewModel.h"
+#include "Widgets/Input/SButton.h"
+#include "Widgets/Input/SComboButton.h"
 
 #define LOCTEXT_NAMESPACE "SNiagaraSpreadsheetView"
 #define ARRAY_INDEX_COLUMN_NAME TEXT("Array Index")
@@ -1255,7 +1245,7 @@ void SNiagaraSpreadsheetView::GenerateLayoutInfo(FNiagaraTypeLayoutInfo& Layout,
 
 void SNiagaraSpreadsheetView::ResetColumns(EUITab Tab)
 {
-	int32 i = (int32)Tab;
+	int32 i = Tab;
 
 	if (CaptureData[i].DataSet.GetCurrentData() && CaptureData[i].DataSet.GetCurrentData()->GetNumInstances() != 0)
 	{
@@ -1263,20 +1253,20 @@ void SNiagaraSpreadsheetView::ResetColumns(EUITab Tab)
 
 		// Handle output columns
 		{
-			CaptureData[(int32)i].OutputHeaderRow->ClearColumns();
+			CaptureData[i].OutputHeaderRow->ClearColumns();
 
-			const TArray<FName>& PreviousSupportedFields = CaptureData[(int32)i].SupportedOutputFields.IsValid() ? *CaptureData[(int32)i].SupportedOutputFields.Get() : TArray<FName>();
-			CaptureData[(int32)i].SupportedOutputFields = MakeShared<TArray<FName> >();
-			CaptureData[(int32)i].OutputFieldInfoMap = MakeShared<TMap<FName, FieldInfo> >();
+			const TArray<FName>& PreviousSupportedFields = CaptureData[i].SupportedOutputFields.IsValid() ? *CaptureData[i].SupportedOutputFields.Get() : TArray<FName>();
+			CaptureData[i].SupportedOutputFields = MakeShared<TArray<FName> >();
+			CaptureData[i].OutputFieldInfoMap = MakeShared<TMap<FName, FieldInfo> >();
 			uint32 TotalFloatComponents = 0;
 			uint32 TotalInt32Components = 0;
 			uint32 TotalHalfComponents = 0;
 
-			TArray<FNiagaraVariable> Variables = CaptureData[(int32)i].DataSet.GetVariables();
+			TArray<FNiagaraVariable> Variables = CaptureData[i].DataSet.GetVariables();
 
 			TArray<FName> ColumnNames;
 
-			if (CaptureData[(int32)i].bOutputColumnsAreAttributes)
+			if (CaptureData[i].bOutputColumnsAreAttributes)
 			{
 				ColumnNames.Add(ARRAY_INDEX_COLUMN_NAME);
 			}
@@ -1372,28 +1362,28 @@ void SNiagaraSpreadsheetView::ResetColumns(EUITab Tab)
 					ColumnArgs.DefaultLabel(FText::FromString(TEXT(" ")));
 					ColumnArgs.ManualWidth(ManualWidth);
 				}
-				CaptureData[(int32)i].OutputHeaderRow->AddColumn(ColumnArgs);
+				CaptureData[i].OutputHeaderRow->AddColumn(ColumnArgs);
 			}
 
-			CaptureData[(int32)i].OutputHeaderRow->ResetColumnWidths();
-			CaptureData[(int32)i].OutputHeaderRow->RefreshColumns();
-			CaptureData[(int32)i].OutputsListView->RequestTreeRefresh();
+			CaptureData[i].OutputHeaderRow->ResetColumnWidths();
+			CaptureData[i].OutputHeaderRow->RefreshColumns();
+			CaptureData[i].OutputsListView->RequestTreeRefresh();
 		}
 
 		// Handle input columns
 		{
-			CaptureData[(int32)i].InputHeaderRow->ClearColumns();
+			CaptureData[i].InputHeaderRow->ClearColumns();
 
 
-			CaptureData[(int32)i].SupportedInputFields = MakeShared<TArray<FName> >();
-			CaptureData[(int32)i].InputFieldInfoMap = MakeShared<TMap<FName, FieldInfo> >();
+			CaptureData[i].SupportedInputFields = MakeShared<TArray<FName> >();
+			CaptureData[i].InputFieldInfoMap = MakeShared<TMap<FName, FieldInfo> >();
 
 			TArray<FNiagaraVariable> Variables;
-			CaptureData[(int32)i].InputParams.GetParameters(Variables);
+			CaptureData[i].InputParams.GetParameters(Variables);
 
 			TArray<FName> ColumnNames;
 
-			if (CaptureData[(int32)i].bInputColumnsAreAttributes)
+			if (CaptureData[i].bInputColumnsAreAttributes)
 			{
 				ColumnNames.Add(ARRAY_INDEX_COLUMN_NAME);
 			}
@@ -1415,7 +1405,7 @@ void SNiagaraSpreadsheetView::ResetColumns(EUITab Tab)
 				TArray<FName> PropertyNames;
 				TArray<SNiagaraSpreadsheetView::FieldInfo> FieldInfos;
 
-				int32 ByteOffset = CaptureData[(int32)i].InputParams.IndexOf(Var);
+				int32 ByteOffset = CaptureData[i].InputParams.IndexOf(Var);
 
 				GenerateLayoutInfo(Layout, Struct, Enum, Var.GetName(), PropertyNames, FieldInfos);
 
@@ -1426,11 +1416,11 @@ void SNiagaraSpreadsheetView::ResetColumns(EUITab Tab)
 
 						FieldInfos[VarIdx].GlobalStartOffset += ByteOffset;
 
-						CaptureData[(int32)i].SupportedInputFields->Add(PropertyNames[VarIdx]);
-						CaptureData[(int32)i].InputFieldInfoMap->Add(PropertyNames[VarIdx], FieldInfos[VarIdx]);
+						CaptureData[i].SupportedInputFields->Add(PropertyNames[VarIdx]);
+						CaptureData[i].InputFieldInfoMap->Add(PropertyNames[VarIdx], FieldInfos[VarIdx]);
 					}
 
-					if (CaptureData[(int32)i].bInputColumnsAreAttributes)
+					if (CaptureData[i].bInputColumnsAreAttributes)
 					{
 						ColumnNames.Add(PropertyNames[VarIdx]);
 					}
@@ -1461,12 +1451,12 @@ void SNiagaraSpreadsheetView::ResetColumns(EUITab Tab)
 					ColumnArgs.DefaultLabel(FText::FromString(TEXT(" ")));
 					ColumnArgs.ManualWidth(ManualWidth);
 				}
-				CaptureData[(int32)i].InputHeaderRow->AddColumn(ColumnArgs);
+				CaptureData[i].InputHeaderRow->AddColumn(ColumnArgs);
 			}
 
-			CaptureData[(int32)i].InputHeaderRow->ResetColumnWidths();
-			CaptureData[(int32)i].InputHeaderRow->RefreshColumns();
-			CaptureData[(int32)i].InputsListView->RequestTreeRefresh();
+			CaptureData[i].InputHeaderRow->ResetColumnWidths();
+			CaptureData[i].InputHeaderRow->RefreshColumns();
+			CaptureData[i].InputsListView->RequestTreeRefresh();
 		}
 	}
 }
@@ -1494,11 +1484,12 @@ FReply SNiagaraSpreadsheetView::OnCaptureRequestPressed()
 		SystemInstance->RequestCapture(TargetRequestId);
 
 		TSharedPtr<FNiagaraEmitterHandleViewModel> SelectedEmitterHandle = SystemViewModel->GetEmitterHandleViewModelById(SelectedEmitterHandleIds[0]);
-		UNiagaraEmitter* Emitter = SelectedEmitterHandle->GetEmitterHandle()->GetInstance();
+		FVersionedNiagaraEmitter Emitter = SelectedEmitterHandle->GetEmitterHandle()->GetInstance();
+		FVersionedNiagaraEmitterData* EmitterData = Emitter.GetEmitterData();
 
 		for (int32 i = 0; i < CaptureData.Num(); i++)
 		{
-			CaptureData[(int32)i].DataSource = Emitter;
+			CaptureData[i].DataSource = Emitter.ToWeakPtr();
 			switch (i)
 			{
 				case UIPerParticleUpdate:
@@ -1511,15 +1502,15 @@ FReply SNiagaraSpreadsheetView::OnCaptureRequestPressed()
 					break;
 				case UIPerParticleEvent0:
 					CaptureData[i].TargetUsage = ENiagaraScriptUsage::ParticleEventScript;
-					CaptureData[i].TargetUsageId = Emitter->GetEventHandlers().Num() >= 1 ? Emitter->GetEventHandlers() [0].Script->GetUsageId() : FGuid();
+					CaptureData[i].TargetUsageId = EmitterData->GetEventHandlers().Num() >= 1 ? EmitterData->GetEventHandlers() [0].Script->GetUsageId() : FGuid();
 					break;
 				case UIPerParticleEvent1:
 					CaptureData[i].TargetUsage = ENiagaraScriptUsage::ParticleEventScript;
-					CaptureData[i].TargetUsageId = Emitter->GetEventHandlers().Num() >= 2 ? Emitter->GetEventHandlers()[1].Script->GetUsageId() : FGuid();
+					CaptureData[i].TargetUsageId = EmitterData->GetEventHandlers().Num() >= 2 ? EmitterData->GetEventHandlers()[1].Script->GetUsageId() : FGuid();
 					break;
 				case UIPerParticleEvent2:
 					CaptureData[i].TargetUsage = ENiagaraScriptUsage::ParticleEventScript;
-					CaptureData[i].TargetUsageId = Emitter->GetEventHandlers().Num() >= 3 ? Emitter->GetEventHandlers()[2].Script->GetUsageId() : FGuid();
+					CaptureData[i].TargetUsageId = EmitterData->GetEventHandlers().Num() >= 3 ? EmitterData->GetEventHandlers()[2].Script->GetUsageId() : FGuid();
 					break;
 				case UISystemUpdate:
 					CaptureData[i].TargetUsage = ENiagaraScriptUsage::SystemUpdateScript;
diff --git a/Engine/Plugins/FX/Niagara/Source/NiagaraEditor/Private/Widgets/SNiagaraSpreadsheetView.h b/Engine/Plugins/FX/Niagara/Source/NiagaraEditor/Private/Widgets/SNiagaraSpreadsheetView.h
--- a/Engine/Plugins/FX/Niagara/Source/NiagaraEditor/Private/Widgets/SNiagaraSpreadsheetView.h
+++ b/Engine/Plugins/FX/Niagara/Source/NiagaraEditor/Private/Widgets/SNiagaraSpreadsheetView.h
@@ -13,6 +13,7 @@
 #include "TickableEditorObject.h"
 #include "UObject/WeakObjectPtrTemplates.h"
 #include "NiagaraDataSet.h"
+#include "NiagaraEmitter.h"
 #include "Templates/SharedPointer.h"
 #include "Containers/Map.h"
 #include "Widgets/Input/SCheckBox.h"
@@ -100,7 +101,7 @@ protected:
 		ENiagaraScriptUsage TargetUsage;
 		FGuid TargetUsageId;
 		FGuid LastCaptureHandleId;
-		TWeakObjectPtr<UNiagaraEmitter> DataSource;
+		FVersionedNiagaraEmitterWeakPtr DataSource;
 		TSharedPtr<SScrollBar> OutputHorizontalScrollBar;
 		TSharedPtr<SScrollBar> OutputVerticalScrollBar;
 		TSharedPtr<SScrollBar> InputHorizontalScrollBar;
diff --git a/Engine/Plugins/FX/Niagara/Source/NiagaraEditor/Private/Widgets/SNiagaraSystemViewport.cpp b/Engine/Plugins/FX/Niagara/Source/NiagaraEditor/Private/Widgets/SNiagaraSystemViewport.cpp
--- a/Engine/Plugins/FX/Niagara/Source/NiagaraEditor/Private/Widgets/SNiagaraSystemViewport.cpp
+++ b/Engine/Plugins/FX/Niagara/Source/NiagaraEditor/Private/Widgets/SNiagaraSystemViewport.cpp
@@ -198,8 +198,8 @@ void FNiagaraSystemViewportClient::DrawInstructionCounts(UNiagaraSystem* Particl
 
 	for (const FNiagaraEmitterHandle& EmitterHandle : ParticleSystem->GetEmitterHandles())
 	{
-		UNiagaraEmitter* Emitter = EmitterHandle.GetInstance();
-		if (Emitter == nullptr)
+		FVersionedNiagaraEmitterData* EmitterData = EmitterHandle.GetEmitterData();
+		if (EmitterData == nullptr)
 		{
 			continue;
 		}
@@ -208,7 +208,7 @@ void FNiagaraSystemViewportClient::DrawInstructionCounts(UNiagaraSystem* Particl
 		CurrentY += FontHeight;
 
 		TArray<UNiagaraScript*> EmitterScripts;
-		Emitter->GetScripts(EmitterScripts);
+		EmitterData->GetScripts(EmitterScripts);
 
 		for (UNiagaraScript* Script : EmitterScripts)
 		{
@@ -260,7 +260,7 @@ void FNiagaraSystemViewportClient::DrawParticleCounts(UNiagaraComponent* Compone
 		{
 			const FName EmitterName = EmitterInstance->GetEmitterHandle().GetName();
 			const int32 CurrentCount = EmitterInstance->GetNumParticles();
-			const int32 MaxCount = EmitterInstance->GetEmitterHandle().GetInstance()->GetMaxParticleCountEstimate();
+			const int32 MaxCount = EmitterInstance->GetEmitterHandle().GetEmitterData()->GetMaxParticleCountEstimate();
 			const bool IsIsolated = EmitterInstance->GetEmitterHandle().IsIsolated();
 			const bool IsEnabled = EmitterInstance->GetEmitterHandle().GetIsEnabled();
 			const ENiagaraExecutionState ExecutionState = EmitterInstance->GetExecutionState();
@@ -283,19 +283,19 @@ void FNiagaraSystemViewportClient::DrawParticleCounts(UNiagaraComponent* Compone
 
 		for (const FNiagaraEmitterHandle& EmitterHandle : NiagaraSystem->GetEmitterHandles())
 		{
-			UNiagaraEmitter* Emitter = EmitterHandle.GetInstance();
-			if (Emitter == nullptr)
+			FVersionedNiagaraEmitterData* EmitterData = EmitterHandle.GetEmitterData();
+			if (EmitterData == nullptr)
 			{
 				continue;
 			}
 
-			const bool bEmitterCompiling = bSystemCompiling || !Emitter->IsReadyToRun();
+			const bool bEmitterCompiling = bSystemCompiling || !EmitterData->IsReadyToRun();
 			if (!bEmitterCompiling)
 			{
 				continue;
 			}
 
-			FString CompilingText = FString::Printf(TEXT("%s - %s"), *FText::FromString("Compiling Emitter").ToString(), *Emitter->GetName());
+			FString CompilingText = FString::Printf(TEXT("%s - %s"), *FText::FromString("Compiling Emitter").ToString(), *EmitterHandle.GetUniqueInstanceName());
 			Canvas->DrawShadowedString(CurrentX, CurrentY, *CompilingText, Font, FLinearColor::Yellow);
 			CurrentY += FontHeight;
 		}
@@ -318,7 +318,7 @@ void FNiagaraSystemViewportClient::DrawEmitterExecutionOrder(UNiagaraComponent*
 		for (const FNiagaraEmitterExecutionIndex& EmitterExecIndex : ExecutionOrder)
 		{
 			const FNiagaraEmitterHandle& EmitterHandle = NiagaraSystem->GetEmitterHandle(EmitterExecIndex.EmitterIndex);
-			if (UNiagaraEmitter* NiagaraEmitter = EmitterHandle.GetInstance())
+			if (FVersionedNiagaraEmitterData* NiagaraEmitter = EmitterHandle.GetEmitterData())
 			{
 				Canvas->DrawShadowedString(CurrentX, CurrentY, *FString::Printf(TEXT("%d - %s"), ++DisplayIndex, NiagaraEmitter->GetDebugSimName()), Font, FLinearColor::White);
 				CurrentY += FontHeight;
diff --git a/Engine/Plugins/FX/Niagara/Source/NiagaraEditor/Private/Widgets/Scalability/SNiagaraScalabilityContext.cpp b/Engine/Plugins/FX/Niagara/Source/NiagaraEditor/Private/Widgets/Scalability/SNiagaraScalabilityContext.cpp
--- a/Engine/Plugins/FX/Niagara/Source/NiagaraEditor/Private/Widgets/Scalability/SNiagaraScalabilityContext.cpp
+++ b/Engine/Plugins/FX/Niagara/Source/NiagaraEditor/Private/Widgets/Scalability/SNiagaraScalabilityContext.cpp
@@ -2,8 +2,8 @@
 
 #include "SNiagaraScalabilityContext.h"
 #include "IDetailChildrenBuilder.h"
-#include "IPropertyRowGenerator.h"
 #include "ISinglePropertyView.h"
+#include "NiagaraEmitterDetailsCustomization.h"
 #include "NiagaraMeshRendererProperties.h"
 #include "NiagaraObjectSelection.h"
 #include "Modules/ModuleManager.h"
@@ -37,7 +37,7 @@ void SNiagaraScalabilityContext::Construct(const FArguments& InArgs, UNiagaraSys
 	// puts a category from default into custom categories. This in turn will ignore the property visible delegate.
 	// without this, a few properties from the mesh renderer properties make it into the details panel even though the delegate returns false for them.
 	DetailsView->RegisterInstancedCustomPropertyLayout(UNiagaraMeshRendererProperties::StaticClass(), DetailsView->GetGenericLayoutDetailsDelegate());
-	// add more here...
+	DetailsView->RegisterInstancedCustomPropertyLayout(UNiagaraEmitter::StaticClass(), FOnGetDetailCustomizationInstance::CreateStatic(&FNiagaraEmitterScalabilityDetails::MakeInstance));
 	
 	DetailsView->SetIsPropertyVisibleDelegate(FIsPropertyVisible::CreateSP(this, &SNiagaraScalabilityContext::FilterScalabilityProperties));
 	UpdateScalabilityContent();
@@ -103,7 +103,7 @@ void SNiagaraScalabilityContext::UpdateScalabilityContent()
 		}
 		else if(StackEntry->IsA(UNiagaraStackEmitterPropertiesGroup::StaticClass()))
 		{
-			UNiagaraEmitter* Emitter = Cast<UNiagaraStackEmitterPropertiesGroup>(StackEntry)->GetEmitterViewModel()->GetEmitter();
+			UNiagaraEmitter* Emitter = Cast<UNiagaraStackEmitterPropertiesGroup>(StackEntry)->GetEmitterViewModel()->GetEmitter().Emitter;
 			NewSelection = Emitter;
 		}
 		else if(StackEntry->IsA(UNiagaraStackSystemPropertiesGroup::StaticClass()))
diff --git a/Engine/Plugins/FX/Niagara/Source/NiagaraEditor/Private/Widgets/Scalability/SNiagaraSystemResolvedScalabilitySettings.cpp b/Engine/Plugins/FX/Niagara/Source/NiagaraEditor/Private/Widgets/Scalability/SNiagaraSystemResolvedScalabilitySettings.cpp
--- a/Engine/Plugins/FX/Niagara/Source/NiagaraEditor/Private/Widgets/Scalability/SNiagaraSystemResolvedScalabilitySettings.cpp
+++ b/Engine/Plugins/FX/Niagara/Source/NiagaraEditor/Private/Widgets/Scalability/SNiagaraSystemResolvedScalabilitySettings.cpp
@@ -15,7 +15,6 @@
 #include "Widgets/Input/SButton.h"
 #include "Widgets/Input/SExpandableButton.h"
 #include "Widgets/Layout/SExpandableArea.h"
-#include "Widgets/Layout/SSeparator.h"
 #include "Widgets/Text/SInlineEditableTextBlock.h"
 #include "PropertyEditorModule.h"
 
@@ -130,10 +129,10 @@ TSharedRef<SWidget> SScalabilityResolvedRow::GenerateWidgetForColumn(const FName
 		];
 	}
 	else if(ColumnId == ColumnDefaultValue)
-	{		
-		FText DefaultValueText;
+	{
 		if(ScalabilityRowData->DefaultPropertyHandle.IsValid())
 		{
+			FText DefaultValueText;
 			ScalabilityRowData->DefaultPropertyHandle->GetValueAsDisplayText(DefaultValueText);			
 			RowContent = SNew(SHorizontalBox)
 			+ SHorizontalBox::Slot()
@@ -153,7 +152,7 @@ TSharedRef<SWidget> SScalabilityResolvedRow::GenerateWidgetForColumn(const FName
 		bool bSkipWidget = false;
 		if(UNiagaraEmitter* Emitter = Cast<UNiagaraEmitter>(ScalabilityRowData->OwningObject))
 		{
-			const FNiagaraEmitterScalabilityOverride& CurrentOverrides = Emitter->GetCurrentOverrideSettings();
+			const FNiagaraEmitterScalabilityOverride& CurrentOverrides = Emitter->GetEmitterData(ScalabilityRowData->OwningObjectVersion)->GetCurrentOverrideSettings();
 			if(!IsEmitterScalabilityValueOverridden(ScalabilityRowData->ResolvedTreeNode->GetNodeName(), CurrentOverrides))
 			{
 				RowContent = SNullWidget::NullWidget;
@@ -221,8 +220,7 @@ void SNiagaraSystemResolvedScalabilitySettings::RebuildWidget()
 
 	for(UObject* SelectedNode : SystemViewModel->GetOverviewGraphViewModel()->GetNodeSelection()->GetSelectedObjects())
 	{
-		UNiagaraOverviewNode* OverviewNode = Cast<UNiagaraOverviewNode>(SelectedNode);
-		if(OverviewNode)
+		if(UNiagaraOverviewNode* OverviewNode = Cast<UNiagaraOverviewNode>(SelectedNode))
 		{
 			SelectedOverviewNodes.Add(OverviewNode);
 		}
@@ -232,17 +230,17 @@ void SNiagaraSystemResolvedScalabilitySettings::RebuildWidget()
 	{
 		for(UObject* SelectedNode : SelectedOverviewNodes)
 		{
-			UNiagaraOverviewNode* OverviewNode = Cast<UNiagaraOverviewNode>(SelectedNode);
-			if(OverviewNode)
+			if(UNiagaraOverviewNode* OverviewNode = Cast<UNiagaraOverviewNode>(SelectedNode))
 			{
 				if(FNiagaraEmitterHandle* Handle = OverviewNode->TryGetEmitterHandle())
 				{
-					Handle->GetInstance()->OnPropertiesChanged().RemoveAll(this);
-					Handle->GetInstance()->OnPropertiesChanged().AddSP(this, &SNiagaraSystemResolvedScalabilitySettings::RebuildWidget);
+					FVersionedNiagaraEmitter VersionedEmitter = Handle->GetInstance();
+					VersionedEmitter.Emitter->OnPropertiesChanged().RemoveAll(this);
+					VersionedEmitter.Emitter->OnPropertiesChanged().AddSP(this, &SNiagaraSystemResolvedScalabilitySettings::RebuildWidget);
 					ResolvedScalabilityContainer->AddSlot()
 					.AutoHeight()
 					[
-						GenerateResolvedScalabilityTable(Handle->GetInstance())
+						GenerateResolvedScalabilityTable(VersionedEmitter.Emitter, VersionedEmitter.Version)
 					];
 				}
 				else
@@ -250,7 +248,7 @@ void SNiagaraSystemResolvedScalabilitySettings::RebuildWidget()
 					ResolvedScalabilityContainer->AddSlot()
 					.AutoHeight()
 					[
-						GenerateResolvedScalabilityTable(OverviewNode->GetOwningSystem())
+						GenerateResolvedScalabilityTable(OverviewNode->GetOwningSystem(), FGuid())
 					];
 				}
 			}
@@ -260,16 +258,17 @@ void SNiagaraSystemResolvedScalabilitySettings::RebuildWidget()
 	{
 		ResolvedScalabilityContainer->AddSlot()
 		[
-			GenerateResolvedScalabilityTable(&SystemViewModel->GetSystem())
+			GenerateResolvedScalabilityTable(&SystemViewModel->GetSystem(), FGuid())
 		];
 		
 		for(const TSharedRef<FNiagaraEmitterHandleViewModel>& EmitterHandleViewModel : SystemViewModel->GetEmitterHandleViewModels())
 		{
-			EmitterHandleViewModel->GetEmitterHandle()->GetInstance()->OnPropertiesChanged().RemoveAll(this);
-			EmitterHandleViewModel->GetEmitterHandle()->GetInstance()->OnPropertiesChanged().AddSP(this, &SNiagaraSystemResolvedScalabilitySettings::RebuildWidget);
+			FVersionedNiagaraEmitter VersionedEmitter = EmitterHandleViewModel->GetEmitterHandle()->GetInstance();
+			VersionedEmitter.Emitter->OnPropertiesChanged().RemoveAll(this);
+			VersionedEmitter.Emitter->OnPropertiesChanged().AddSP(this, &SNiagaraSystemResolvedScalabilitySettings::RebuildWidget);
 			ResolvedScalabilityContainer->AddSlot()
 			[
-				GenerateResolvedScalabilityTable(EmitterHandleViewModel->GetEmitterHandle()->GetInstance())
+				GenerateResolvedScalabilityTable(VersionedEmitter.Emitter, VersionedEmitter.Version)
 			];	
 		}
 	}
@@ -285,7 +284,7 @@ TSharedRef<ITableRow> SNiagaraSystemResolvedScalabilitySettings::GenerateScalabi
 	return SNew(SScalabilityResolvedRow, ScalabilityRowData, TableViewBase, SharedThis(this));
 }
 
-TSharedRef<SWidget> SNiagaraSystemResolvedScalabilitySettings::GenerateResolvedScalabilityTable(UObject* Object)
+TSharedRef<SWidget> SNiagaraSystemResolvedScalabilitySettings::GenerateResolvedScalabilityTable(UObject* Object, const FGuid& Version)
 {
 	FString DisplayName = Object->GetName();
 
@@ -307,7 +306,7 @@ TSharedRef<SWidget> SNiagaraSystemResolvedScalabilitySettings::GenerateResolvedS
 			const FNiagaraEmitterScalabilitySettings& DefaultEmitterScalabilitySettings = System->GetEffectType()->GetActiveEmitterScalabilitySettings();
 			DefaultValueStruct = MakeShared<FStructOnScope>(FNiagaraEmitterScalabilitySettings::StaticStruct(), (uint8*) &DefaultEmitterScalabilitySettings);			
 		}
-		const FNiagaraEmitterScalabilitySettings& EmitterScalabilitySettings = Emitter->GetScalabilitySettings();
+		const FNiagaraEmitterScalabilitySettings& EmitterScalabilitySettings = Emitter->GetEmitterData(Version)->GetScalabilitySettings();
 		ResolvedValueStruct = MakeShared<FStructOnScope>(FNiagaraEmitterScalabilitySettings::StaticStruct(), (uint8*) &EmitterScalabilitySettings);
 	}
 	else if(UNiagaraSystem* NiagaraSystem = Cast<UNiagaraSystem>(Object))
@@ -326,12 +325,11 @@ TSharedRef<SWidget> SNiagaraSystemResolvedScalabilitySettings::GenerateResolvedS
     ResolvedPropertyGenerator->SetStructure(ResolvedValueStruct);
 
 	// we are caching the default handles so we can later find the default handle corresponding to a resolved handle
-	TArray<TSharedRef<IDetailTreeNode>> DefaultTreeNodes;
 	TMap<FString, TSharedRef<IDetailTreeNode>> DefaultHandles;
 	if(DefaultValueStruct.IsValid())
 	{
 		DefaultPropertyGenerator->SetStructure(DefaultValueStruct);
-		DefaultTreeNodes = DefaultPropertyGenerator->GetRootTreeNodes();
+		TArray<TSharedRef<IDetailTreeNode>> DefaultTreeNodes = DefaultPropertyGenerator->GetRootTreeNodes();
 		
 		for(TSharedRef<IDetailTreeNode> RootNode : DefaultTreeNodes)
 		{
@@ -374,6 +372,7 @@ TSharedRef<SWidget> SNiagaraSystemResolvedScalabilitySettings::GenerateResolvedS
 					ScalabilityValue->ResolvedPropertyHandle = ChildPropertyHandle.ToSharedRef();
 					ScalabilityValue->ResolvedTreeNode = Child;
 					ScalabilityValue->OwningObject = Object;
+					ScalabilityValue->OwningObjectVersion = Version;
 
 					// if we found the same handle in the default handles, we have both handles needed to display default and override data
 					if(DefaultHandles.Contains(ChildPropertyHandle->GetProperty()->GetName()))
diff --git a/Engine/Plugins/FX/Niagara/Source/NiagaraEditor/Private/Widgets/Scalability/SNiagaraSystemResolvedScalabilitySettings.h b/Engine/Plugins/FX/Niagara/Source/NiagaraEditor/Private/Widgets/Scalability/SNiagaraSystemResolvedScalabilitySettings.h
--- a/Engine/Plugins/FX/Niagara/Source/NiagaraEditor/Private/Widgets/Scalability/SNiagaraSystemResolvedScalabilitySettings.h
+++ b/Engine/Plugins/FX/Niagara/Source/NiagaraEditor/Private/Widgets/Scalability/SNiagaraSystemResolvedScalabilitySettings.h
@@ -16,6 +16,7 @@ struct FScalabilityRowData
 	TSharedPtr<IPropertyHandle> DefaultPropertyHandle;
 	TArray<FProperty*> SourceProperties;
 	UObject* OwningObject;
+	FGuid OwningObjectVersion;
 };
 
 class SScalabilityResolvedRow : public SMultiColumnTableRow<TSharedPtr<FScalabilityRowData>>
@@ -62,7 +63,7 @@ private:
 	TMap<UObject*, FRequiredInstanceInformation> InstanceInformation;
 	
 	/** Generates an expandable widget for an emitter or system */
-	TSharedRef<SWidget> GenerateResolvedScalabilityTable(UObject* Object);
+	TSharedRef<SWidget> GenerateResolvedScalabilityTable(UObject* Object, const FGuid& Version);
 	
 	TSharedRef<ITableRow> GenerateScalabilityValueRow(TSharedRef<FScalabilityRowData> ScalabilityRowData, const TSharedRef< class STableViewBase >& TableViewBase);
 };
diff --git a/Engine/Plugins/FX/Niagara/Source/NiagaraEditor/Public/Customizations/NiagaraEventScriptPropertiesCustomization.h b/Engine/Plugins/FX/Niagara/Source/NiagaraEditor/Public/Customizations/NiagaraEventScriptPropertiesCustomization.h
--- a/Engine/Plugins/FX/Niagara/Source/NiagaraEditor/Public/Customizations/NiagaraEventScriptPropertiesCustomization.h
+++ b/Engine/Plugins/FX/Niagara/Source/NiagaraEditor/Public/Customizations/NiagaraEventScriptPropertiesCustomization.h
@@ -48,13 +48,11 @@ class FNiagaraEventScriptPropertiesCustomization : public IPropertyTypeCustomiza
 {
 public:
 	/** Makes a new instance of this detail layout class for a specific detail view requesting it */
-	static TSharedRef<class IPropertyTypeCustomization> MakeInstance(TWeakObjectPtr<UNiagaraSystem> InSystem,
-		TWeakObjectPtr<UNiagaraEmitter> InEmitter);
+	static TSharedRef<class IPropertyTypeCustomization> MakeInstance(TWeakObjectPtr<UNiagaraSystem> InSystem, FVersionedNiagaraEmitterWeakPtr InEmitter);
 
-	FNiagaraEventScriptPropertiesCustomization(TWeakObjectPtr<UNiagaraSystem> InSystem,
-		TWeakObjectPtr<UNiagaraEmitter> InEmitter);
+	FNiagaraEventScriptPropertiesCustomization(TWeakObjectPtr<UNiagaraSystem> InSystem,	FVersionedNiagaraEmitterWeakPtr InEmitter);
 
-	~FNiagaraEventScriptPropertiesCustomization();
+	virtual ~FNiagaraEventScriptPropertiesCustomization() override;
 
 	/** IPropertyTypeCustomization interface */
 	virtual void CustomizeHeader(TSharedRef<IPropertyHandle> StructPropertyHandle, class FDetailWidgetRow& HeaderRow, IPropertyTypeCustomizationUtils& StructCustomizationUtils) override;
@@ -80,7 +78,7 @@ protected:
 	FText GetErrorText() const;
 	FText GetErrorTextTooltip() const;
 
-	TArray<FName> GetEventNames(UNiagaraEmitter* Emitter) const;
+	TArray<FName> GetEventNames(const FVersionedNiagaraEmitter& Emitter) const;
 private:
 	TSharedPtr<IPropertyHandle> HandleSrcID;
 	TSharedPtr<IPropertyHandle> HandleEventName;
@@ -91,7 +89,7 @@ private:
 	TSharedPtr<IPropertyHandle> HandleMinSpawnNumber;
 
 	TWeakObjectPtr<UNiagaraSystem> System;
-	TWeakObjectPtr<UNiagaraEmitter> Emitter;
+	FVersionedNiagaraEmitterWeakPtr Emitter;
 
 	FName CachedEmitterName;
 	EVisibility CachedVisibility;
diff --git a/Engine/Plugins/FX/Niagara/Source/NiagaraEditor/Public/NiagaraEditorCommands.h b/Engine/Plugins/FX/Niagara/Source/NiagaraEditor/Public/NiagaraEditorCommands.h
--- a/Engine/Plugins/FX/Niagara/Source/NiagaraEditor/Public/NiagaraEditorCommands.h
+++ b/Engine/Plugins/FX/Niagara/Source/NiagaraEditor/Public/NiagaraEditorCommands.h
@@ -28,6 +28,7 @@ public:
 	TSharedPtr<FUICommandInfo> Compile;
 	TSharedPtr<FUICommandInfo> RefreshNodes;
 	TSharedPtr<FUICommandInfo> ModuleVersioning;
+	TSharedPtr<FUICommandInfo> EmitterVersioning;
 	TSharedPtr<FUICommandInfo> ResetSimulation;
 	TSharedPtr<FUICommandInfo> SelectNextUsage;
 	TSharedPtr<FUICommandInfo> CreateAssetFromSelection;
diff --git a/Engine/Plugins/FX/Niagara/Source/NiagaraEditor/Public/NiagaraEditorModule.h b/Engine/Plugins/FX/Niagara/Source/NiagaraEditor/Public/NiagaraEditorModule.h
--- a/Engine/Plugins/FX/Niagara/Source/NiagaraEditor/Public/NiagaraEditorModule.h
+++ b/Engine/Plugins/FX/Niagara/Source/NiagaraEditor/Public/NiagaraEditorModule.h
@@ -201,7 +201,7 @@ public:
 
 	const TArray<TWeakObjectPtr<UNiagaraParameterDefinitions>>& GetCachedParameterDefinitionsAssets();
 
-	NIAGARAEDITOR_API void GetTargetSystemAndEmitterForDataInterface(UNiagaraDataInterface* InDataInterface, UNiagaraSystem*& OutOwningSystem, UNiagaraEmitter*& OutOwningEmitter);
+	NIAGARAEDITOR_API void GetTargetSystemAndEmitterForDataInterface(UNiagaraDataInterface* InDataInterface, UNiagaraSystem*& OutOwningSystem, FVersionedNiagaraEmitter& OutOwningEmitter);
 	NIAGARAEDITOR_API void GetDataInterfaceFeedbackSafe(UNiagaraDataInterface* InDataInterface, TArray<FNiagaraDataInterfaceError>& OutErrors, TArray<FNiagaraDataInterfaceFeedback>& Warnings, TArray<FNiagaraDataInterfaceFeedback>& Info);
 
 	NIAGARAEDITOR_API void EnsureReservedDefinitionUnique(FGuid& UniqueId);
diff --git a/Engine/Plugins/FX/Niagara/Source/NiagaraEditor/Public/NiagaraEditorUtilities.h b/Engine/Plugins/FX/Niagara/Source/NiagaraEditor/Public/NiagaraEditorUtilities.h
--- a/Engine/Plugins/FX/Niagara/Source/NiagaraEditor/Public/NiagaraEditorUtilities.h
+++ b/Engine/Plugins/FX/Niagara/Source/NiagaraEditor/Public/NiagaraEditorUtilities.h
@@ -11,7 +11,6 @@
 #include "NiagaraGraph.h"
 #include "NiagaraEditorSettings.h"
 #include "UpgradeNiagaraScriptResults.h"
-#include "EdGraph/EdGraphSchema.h"
 #include "ViewModels/NiagaraSystemScalabilityViewModel.h"
 
 class UNiagaraNodeInput;
@@ -153,7 +152,7 @@ namespace FNiagaraEditorUtilities
 
 	TSharedPtr<SWidget> CreateInlineErrorText(TAttribute<FText> ErrorMessage, TAttribute<FText> ErrorTooltip);
 
-	void CompileExistingEmitters(const TArray<UNiagaraEmitter*>& AffectedEmitters);
+	void CompileExistingEmitters(const TArray<FVersionedNiagaraEmitter>& AffectedEmitters);
 
 	bool TryGetEventDisplayName(UNiagaraEmitter* Emitter, FGuid EventUsageId, FText& OutEventDisplayName);
 
@@ -225,7 +224,7 @@ namespace FNiagaraEditorUtilities
 	 * @param The emitter to search for in the system.
 	 * @returns The emitter handle for the supplied emitter, or nullptr if the emitter isn't owned by this system.
 	 */
-	const FNiagaraEmitterHandle* GetEmitterHandleForEmitter(UNiagaraSystem& System, UNiagaraEmitter& Emitter);
+	const FNiagaraEmitterHandle* GetEmitterHandleForEmitter(UNiagaraSystem& System, const FVersionedNiagaraEmitter& Emitter);
 
 	NIAGARAEDITOR_API ENiagaraScriptLibraryVisibility GetScriptAssetVisibility(const FAssetData& ScriptAssetData);
 
@@ -254,7 +253,7 @@ namespace FNiagaraEditorUtilities
 
 	TArray<UNiagaraComponent*> GetComponentsThatReferenceSystemViewModel(const FNiagaraSystemViewModel& ReferencedSystemViewModel);
 
-	NIAGARAEDITOR_API const FGuid AddEmitterToSystem(UNiagaraSystem& InSystem, UNiagaraEmitter& InEmitterToAdd, bool bCreateCopy = true);
+	NIAGARAEDITOR_API const FGuid AddEmitterToSystem(UNiagaraSystem& InSystem, UNiagaraEmitter& InEmitterToAdd, FGuid EmitterVersion, bool bCreateCopy = true);
 
 	void RemoveEmittersFromSystemByEmitterHandleId(UNiagaraSystem& InSystem, TSet<FGuid> EmitterHandleIdsToDelete);
 
@@ -278,7 +277,7 @@ namespace FNiagaraEditorUtilities
 
 	void ShowParentEmitterInContentBrowser(TSharedRef<FNiagaraEmitterViewModel> EmitterViewModel);
 
-	void OpenParentEmitterForEdit(TSharedRef<FNiagaraEmitterViewModel> Emitter);
+	NIAGARAEDITOR_API void OpenParentEmitterForEdit(TSharedRef<FNiagaraEmitterViewModel> Emitter);
 	ECheckBoxState GetSelectedEmittersEnabledCheckState(TSharedRef<FNiagaraSystemViewModel> SystemViewModel);
 	void ToggleSelectedEmittersEnabled(TSharedRef<FNiagaraSystemViewModel> SystemViewModel);
 
@@ -334,12 +333,19 @@ namespace FNiagaraEditorUtilities
 	// Executes python upgrade scripts on the given source node for all the given in-between versions
 	void RunPythonUpgradeScripts(UNiagaraNodeFunctionCall* SourceNode, const TArray<FVersionedNiagaraScriptData*>& UpgradeVersionData, const FNiagaraScriptVersionUpgradeContext& UpgradeContext, FString& OutWarnings);
 
+	// Executes python upgrade scripts on the given source node for all the given in-between versions
+	void RunPythonUpgradeScripts(UUpgradeNiagaraEmitterContext* UpgradeContext);
+
+	// Changes the referenced parent version and optionally runs the python upgrade scripts
+	NIAGARAEDITOR_API void SwitchParentEmitterVersion(TSharedRef<FNiagaraEmitterViewModel> EmitterViewModel, TSharedRef<FNiagaraSystemViewModel> SystemViewModel, const FGuid& NewVersionGuid);
+
 	void RefreshAllScriptsFromExternalChanges(FRefreshAllScriptsFromExternalChangesArgs Args);
 
 	DECLARE_DELEGATE_OneParam(FNodeVisitor, UEdGraphNode* /*VisitedNode*/);
 	void VisitAllNodesConnectedToInputs(UEdGraphNode* StartNode, FNodeVisitor Visitor);
 	
 	NIAGARAEDITOR_API float GetScalabilityTintAlpha(FNiagaraEmitterHandle* EmitterHandle);
+	
 };
 
 namespace FNiagaraParameterUtilities
diff --git a/Engine/Plugins/FX/Niagara/Source/NiagaraEditor/Public/NiagaraGraph.h b/Engine/Plugins/FX/Niagara/Source/NiagaraEditor/Public/NiagaraGraph.h
--- a/Engine/Plugins/FX/Niagara/Source/NiagaraEditor/Public/NiagaraGraph.h
+++ b/Engine/Plugins/FX/Niagara/Source/NiagaraEditor/Public/NiagaraGraph.h
@@ -315,6 +315,8 @@ class UNiagaraGraph : public UEdGraph
 	 */
 	NIAGARAEDITOR_API void ChangeParameterType(const TArray<FNiagaraVariable>& ParametersToChange, const FNiagaraTypeDefinition& NewType, bool bAllowOrphanedPins = false);
 
+	void ReplaceScriptReferences(UNiagaraScript* OldScript, UNiagaraScript* NewScript);
+
 	/** Gets a delegate which is called whenever a contained data interfaces changes. */
 	FOnDataInterfaceChanged& OnDataInterfaceChanged();
 
@@ -349,6 +351,8 @@ class UNiagaraGraph : public UEdGraph
 
 	void ScriptVariableChanged(FNiagaraVariable Variable);
 
+	FVersionedNiagaraEmitter GetOwningEmitter() const;
+
 	/** Synchronize all the properties of DestScriptVar to those of SourceScriptVar, as well as propagating those changes through the graph (pin variable names and default values on pins.) 
 	 *  If DestScriptVar is not set, find a script variable with the same key as the SourceScriptVar.
 	 *  Returns bool to signify if DestScriptVar was modified.
diff --git a/Engine/Plugins/FX/Niagara/Source/NiagaraEditor/Public/NiagaraParameterMapHistory.h b/Engine/Plugins/FX/Niagara/Source/NiagaraEditor/Public/NiagaraParameterMapHistory.h
--- a/Engine/Plugins/FX/Niagara/Source/NiagaraEditor/Public/NiagaraParameterMapHistory.h
+++ b/Engine/Plugins/FX/Niagara/Source/NiagaraEditor/Public/NiagaraParameterMapHistory.h
@@ -4,6 +4,7 @@
 
 #include "CoreMinimal.h"
 #include "NiagaraCommon.h"
+#include "NiagaraEmitter.h"
 #include "Templates/Tuple.h"
 #include "ViewModels/Stack/NiagaraStackGraphUtilities.h"
 #include "NiagaraScript.h"
@@ -17,12 +18,12 @@ class FHlslNiagaraTranslator;
 class FCompileConstantResolver
 {
 public:
-	FCompileConstantResolver() : Emitter(nullptr), System(nullptr), Translator(nullptr), Usage(ENiagaraScriptUsage::Function), DebugState(ENiagaraFunctionDebugState::NoDebug) {}
-	FCompileConstantResolver(const UNiagaraEmitter* Emitter, ENiagaraScriptUsage Usage, ENiagaraFunctionDebugState DebugState = ENiagaraFunctionDebugState::NoDebug) : Emitter(Emitter), System(nullptr), Translator(nullptr), Usage(Usage), DebugState(DebugState) {}
-	FCompileConstantResolver(const UNiagaraSystem* System, ENiagaraScriptUsage Usage, ENiagaraFunctionDebugState DebugState = ENiagaraFunctionDebugState::NoDebug) : Emitter(nullptr), System(System), Translator(nullptr), Usage(Usage), DebugState(DebugState) {}
-	FCompileConstantResolver(const FHlslNiagaraTranslator* Translator, ENiagaraFunctionDebugState DebugState = ENiagaraFunctionDebugState::NoDebug) : Emitter(nullptr), System(nullptr), Translator(Translator), Usage(ENiagaraScriptUsage::Function), DebugState(DebugState) {}
+	FCompileConstantResolver() : Emitter(FVersionedNiagaraEmitter()), System(nullptr), Translator(nullptr), Usage(ENiagaraScriptUsage::Function), DebugState(ENiagaraFunctionDebugState::NoDebug) {}
+	FCompileConstantResolver(const FVersionedNiagaraEmitter& Emitter, ENiagaraScriptUsage Usage, ENiagaraFunctionDebugState DebugState = ENiagaraFunctionDebugState::NoDebug) : Emitter(Emitter), System(nullptr), Translator(nullptr), Usage(Usage), DebugState(DebugState) {}
+	FCompileConstantResolver(const UNiagaraSystem* System, ENiagaraScriptUsage Usage, ENiagaraFunctionDebugState DebugState = ENiagaraFunctionDebugState::NoDebug) : Emitter(FVersionedNiagaraEmitter()), System(System), Translator(nullptr), Usage(Usage), DebugState(DebugState) {}
+	FCompileConstantResolver(const FHlslNiagaraTranslator* Translator, ENiagaraFunctionDebugState DebugState = ENiagaraFunctionDebugState::NoDebug) : Emitter(FVersionedNiagaraEmitter()), System(nullptr), Translator(Translator), Usage(ENiagaraScriptUsage::Function), DebugState(DebugState) {}
 
-	const UNiagaraEmitter* GetEmitter() const { return Emitter; }
+	const UNiagaraEmitter* GetEmitter() const { return Emitter.Emitter; }
 	const UNiagaraSystem* GetSystem() const { return System; }
 	const FHlslNiagaraTranslator* GetTranslator() const { return Translator; }
 	ENiagaraScriptUsage GetUsage() const { return Usage; }
@@ -34,7 +35,7 @@ public:
 	FCompileConstantResolver WithDebugState(ENiagaraFunctionDebugState InDebugState) const;
 	FCompileConstantResolver WithUsage(ENiagaraScriptUsage ScriptUsage) const;
 private:
-	const UNiagaraEmitter* Emitter;
+	FVersionedNiagaraEmitter Emitter;
 	const UNiagaraSystem* System;
 	const FHlslNiagaraTranslator* Translator;
 	ENiagaraScriptUsage Usage;
diff --git a/Engine/Plugins/FX/Niagara/Source/NiagaraEditor/Public/NiagaraScriptSource.h b/Engine/Plugins/FX/Niagara/Source/NiagaraEditor/Public/NiagaraScriptSource.h
--- a/Engine/Plugins/FX/Niagara/Source/NiagaraEditor/Public/NiagaraScriptSource.h
+++ b/Engine/Plugins/FX/Niagara/Source/NiagaraEditor/Public/NiagaraScriptSource.h
@@ -35,7 +35,8 @@ class UNiagaraScriptSource : public UNiagaraScriptSourceBase
 	virtual bool IsSynchronized(const FGuid& InChangeId) override;
 	virtual void MarkNotSynchronized(FString Reason) override;
 
-	virtual FGuid GetChangeID();
+	virtual FGuid GetChangeID() override;
+	FVersionedNiagaraEmitter GetOuterEmitter() const;
 
 	virtual void ComputeVMCompilationId(struct FNiagaraVMExecutableDataId& Id, ENiagaraScriptUsage InUsage, const FGuid& InUsageId, bool bForceRebuild = false) const override;
 
@@ -43,14 +44,14 @@ class UNiagaraScriptSource : public UNiagaraScriptSourceBase
 
 	virtual FNiagaraCompileHash GetCompileHash(ENiagaraScriptUsage InUsage, const FGuid& InUsageId) const override;
 
-	virtual void PostLoadFromEmitter(UNiagaraEmitter& OwningEmitter) override;
+	virtual void PostLoadFromEmitter(FVersionedNiagaraEmitter OwningEmitter) override;
 
 	virtual TMap<FName, UNiagaraDataInterface*> ComputeObjectNameMap(UNiagaraSystem& System, ENiagaraScriptUsage Usage, FGuid UsageId, FString EmitterUniqueName) const override;
 
 	NIAGARAEDITOR_API virtual bool AddModuleIfMissing(FString ModulePath, ENiagaraScriptUsage Usage, bool& bOutFoundModule)override;
 
-	virtual void FixupRenamedParameters(UNiagaraNode* Node, FNiagaraParameterStore& RapidIterationParameters, const TArray<FNiagaraVariable>& OldRapidIterationVariables, const UNiagaraEmitter* Emitter, ENiagaraScriptUsage ScriptUsage) const;
-	virtual void CleanUpOldAndInitializeNewRapidIterationParameters(const UNiagaraEmitter* Emitter, ENiagaraScriptUsage ScriptUsage, FGuid ScriptUsageId, FNiagaraParameterStore& RapidIterationParameters) const override;
+	virtual void FixupRenamedParameters(UNiagaraNode* Node, FNiagaraParameterStore& RapidIterationParameters, const TArray<FNiagaraVariable>& OldRapidIterationVariables, const FVersionedNiagaraEmitter& Emitter, ENiagaraScriptUsage ScriptUsage) const;
+	virtual void CleanUpOldAndInitializeNewRapidIterationParameters(const FVersionedNiagaraEmitter& Emitter, ENiagaraScriptUsage ScriptUsage, FGuid ScriptUsageId, FNiagaraParameterStore& RapidIterationParameters) const override;
 	virtual void ForceGraphToRecompileOnNextCheck() override;
 	virtual void RefreshFromExternalChanges() override;
 
@@ -79,6 +80,7 @@ class UNiagaraScriptSource : public UNiagaraScriptSourceBase
 
 	virtual void GetLinkedPositionTypeInputs(const TArray<FNiagaraVariable>& ParametersToCheck, TSet<FNiagaraVariable>& OutLinkedParameters) override;
 	virtual void ChangedLinkedInputTypes(const FNiagaraVariable& ParametersToChange, const FNiagaraTypeDefinition& NewType) override;
+	virtual void ReplaceScriptReferences(UNiagaraScript* OldScript, UNiagaraScript* NewScript) override;
 private:
 	void OnGraphChanged(const FEdGraphEditAction &Action);
 	void OnGraphDataInterfaceChanged();
diff --git a/Engine/Plugins/FX/Niagara/Source/NiagaraEditor/Public/NiagaraSystemFactoryNew.h b/Engine/Plugins/FX/Niagara/Source/NiagaraEditor/Public/NiagaraSystemFactoryNew.h
--- a/Engine/Plugins/FX/Niagara/Source/NiagaraEditor/Public/NiagaraSystemFactoryNew.h
+++ b/Engine/Plugins/FX/Niagara/Source/NiagaraEditor/Public/NiagaraSystemFactoryNew.h
@@ -3,6 +3,8 @@
 #pragma once
 
 #include "CoreMinimal.h"
+
+#include "NiagaraTypes.h"
 #include "UObject/ObjectMacros.h"
 #include "Factories/Factory.h"
 #include "NiagaraSystemFactoryNew.generated.h"
@@ -15,8 +17,8 @@ class UNiagaraSystemFactoryNew : public UFactory
 {
 	GENERATED_UCLASS_BODY()
 
-	UNiagaraSystem* SystemToCopy;
-	TArray<UNiagaraEmitter*> EmittersToAddToNewSystem;
+	UNiagaraSystem* SystemToCopy = nullptr;
+	TArray<FVersionedNiagaraEmitter> EmittersToAddToNewSystem;
 
 	//~ Begin UFactory Interface
 	virtual bool ConfigureProperties() override;
diff --git a/Engine/Plugins/FX/Niagara/Source/NiagaraEditor/Public/UpgradeNiagaraScriptResults.h b/Engine/Plugins/FX/Niagara/Source/NiagaraEditor/Public/UpgradeNiagaraScriptResults.h
--- a/Engine/Plugins/FX/Niagara/Source/NiagaraEditor/Public/UpgradeNiagaraScriptResults.h
+++ b/Engine/Plugins/FX/Niagara/Source/NiagaraEditor/Public/UpgradeNiagaraScriptResults.h
@@ -7,8 +7,10 @@
 #include "NiagaraParameterMapHistory.h"
 #include "UObject/Object.h"
 #include "NiagaraTypes.h"
+#include "ViewModels/NiagaraEmitterHandleViewModel.h"
 #include "UpgradeNiagaraScriptResults.generated.h"
 
+class UNiagaraStackModuleItem;
 class UNiagaraClipboardFunctionInput;
 
 UENUM()
@@ -30,6 +32,8 @@ enum class ENiagaraPythonScriptInputSource : uint32
     Num UMETA(Hidden)
 };
 
+// ---------------------- Script upgrade data -------------------------------------
+
 /** Wrapper for setting the value on a parameter of a UNiagaraScript, applied through a UUpgradeNiagaraScriptResults. */
 UCLASS(BlueprintInternalUseOnly)
 class UNiagaraPythonScriptModuleInput : public UObject
@@ -144,3 +148,101 @@ struct NIAGARAEDITOR_API FNiagaraScriptVersionUpgradeContext {
 	FCompileConstantResolver ConstantResolver;
 	bool bSkipPythonScript = false;
 };
+
+
+
+// ---------------------- Emitter upgrade data -------------------------------------
+
+/** Wrapper for a module from the emitter stack. */
+UCLASS(BlueprintInternalUseOnly)
+class UNiagaraPythonModule : public UObject
+{
+	GENERATED_BODY()
+
+public:
+	UNiagaraPythonModule() {};
+
+	void Init(UNiagaraStackModuleItem* InModuleItem);
+
+	// Returns the raw underlying object
+	UFUNCTION(BlueprintCallable, Category = "Scripting")
+	UNiagaraStackModuleItem* Object() const;
+
+	//TODO: extend the api here to fully support module actions
+
+private:
+	UPROPERTY()
+	UNiagaraStackModuleItem* ModuleItem = nullptr;
+};
+
+/** Wrapper for an emitter stack. */
+UCLASS(BlueprintInternalUseOnly)
+class UNiagaraPythonEmitter : public UObject
+{
+	GENERATED_BODY()
+
+public:
+	UNiagaraPythonEmitter() {}
+
+	void Init(TSharedRef<FNiagaraEmitterHandleViewModel> EmitterViewModel);
+
+	// Returns the raw underlying object
+	UFUNCTION(BlueprintCallable, Category = "Scripting")
+	UNiagaraEmitter* Object();
+
+	// returns the emitter properties, such as determinism or interpolated spawning
+	UFUNCTION(BlueprintCallable, Category = "Scripting")
+	FVersionedNiagaraEmitterData GetProperties() const;
+
+	// sets the new emitter properties
+	UFUNCTION(BlueprintCallable, Category = "Scripting")
+	void SetProperties(FVersionedNiagaraEmitterData Data);
+
+	// returns a list of all modules contained in this emitter
+	UFUNCTION(BlueprintCallable, Category = "Scripting")
+	TArray<UNiagaraPythonModule*> GetModules() const;
+
+	// returns true if the emitter contains a certain module
+	UFUNCTION(BlueprintCallable, Category = "Scripting")
+	bool HasModule(const FString& ModuleName) const;
+
+	// returns a module by name
+	UFUNCTION(BlueprintCallable, Category = "Scripting")
+	UNiagaraPythonModule* GetModule(const FString& ModuleName) const;
+
+	//TODO: extend the api here to support renderers
+
+	bool IsValid() const { return EmitterViewModel.IsValid(); }
+
+	TSharedPtr<FNiagaraEmitterHandleViewModel> EmitterViewModel;
+};
+
+/**
+ * Wrapper class for passing results back from the version upgrade python script.
+ */
+UCLASS(BlueprintType)
+class NIAGARAEDITOR_API UUpgradeNiagaraEmitterContext : public UObject
+{
+	GENERATED_BODY()
+
+public:
+
+	UUpgradeNiagaraEmitterContext() {}
+
+	void Init(UNiagaraPythonEmitter* InOldEmitter, UNiagaraPythonEmitter* InNewEmitter);
+	bool IsValid() const;
+	const TArray<FVersionedNiagaraEmitterData*>& GetUpgradeData() const;
+	
+	// Whether the converter process was cancelled due to an unrecoverable error in the python script process.
+	UPROPERTY(BlueprintReadWrite, Category = "Scripting")
+	bool bCancelledByPythonError = false;
+
+	UPROPERTY(BlueprintReadWrite, Category = "Scripting")
+	UNiagaraPythonEmitter* OldEmitter = nullptr;
+
+	UPROPERTY(BlueprintReadWrite, Category = "Scripting")
+	UNiagaraPythonEmitter* NewEmitter = nullptr;
+
+private:
+	TArray<FVersionedNiagaraEmitterData*> UpgradeVersionData;
+};
\ No newline at end of file
diff --git a/Engine/Plugins/FX/Niagara/Source/NiagaraEditor/Public/ViewModels/NiagaraEmitterHandleViewModel.h b/Engine/Plugins/FX/Niagara/Source/NiagaraEditor/Public/ViewModels/NiagaraEmitterHandleViewModel.h
--- a/Engine/Plugins/FX/Niagara/Source/NiagaraEditor/Public/ViewModels/NiagaraEmitterHandleViewModel.h
+++ b/Engine/Plugins/FX/Niagara/Source/NiagaraEditor/Public/ViewModels/NiagaraEmitterHandleViewModel.h
@@ -27,8 +27,8 @@ public:
 public:
 	/** Creates a new emitter editor view model.  This must be initialized before it can be used. */
 	FNiagaraEmitterHandleViewModel(bool bInIsForDataProcessingOnly);
-	
-	~FNiagaraEmitterHandleViewModel();
+
+	virtual ~FNiagaraEmitterHandleViewModel() override;
 
 	/** Initializes the emitter editor view model with the supplied emitter handle and simulation.*/
 	void Initialize(TSharedRef<FNiagaraSystemViewModel> InOwningSystemViewModel, int32 InEmitterHandleIndex, TWeakPtr<FNiagaraEmitterInstance, ESPMode::ThreadSafe> InSimulation);
@@ -75,9 +75,6 @@ public:
 	EVisibility GetErrorTextVisibility() const;
 	FSlateColor GetErrorTextColor() const;
 
-	/** Called to get the synch state of the emitter handle to its source.*/
-	bool IsSynchronized() const;
-
 	/** Gets whether or not this emitter handle is enabled. */
 	NIAGARAEDITOR_API bool GetIsEnabled() const;
 
diff --git a/Engine/Plugins/FX/Niagara/Source/NiagaraEditor/Public/ViewModels/NiagaraEmitterViewModel.h b/Engine/Plugins/FX/Niagara/Source/NiagaraEditor/Public/ViewModels/NiagaraEmitterViewModel.h
--- a/Engine/Plugins/FX/Niagara/Source/NiagaraEditor/Public/ViewModels/NiagaraEmitterViewModel.h
+++ b/Engine/Plugins/FX/Niagara/Source/NiagaraEditor/Public/ViewModels/NiagaraEmitterViewModel.h
@@ -8,6 +8,7 @@
 #include "UObject/ObjectKey.h"
 #include "IAssetTypeActions.h"
 #include "NiagaraEmitterEditorData.h"
+#include "NiagaraEmitter.h"
 
 class UNiagaraEmitter;
 class UNiagaraScript;
@@ -42,10 +43,10 @@ public:
 public:
 	/** Creates a new emitter editor view model.  It must be initialized before use. */
 	FNiagaraEmitterViewModel(bool bInIsForDataProcessingOnly);
-	virtual ~FNiagaraEmitterViewModel();
+	virtual ~FNiagaraEmitterViewModel() override;
 
 	/** Initialize this view model with an emitter and simulation. */
-	bool Initialize(UNiagaraEmitter* InEmitter, TWeakPtr<FNiagaraEmitterInstance, ESPMode::ThreadSafe> InSimulation);
+	bool Initialize(const FVersionedNiagaraEmitter& InEmitter, TWeakPtr<FNiagaraEmitterInstance, ESPMode::ThreadSafe> InSimulation);
 
 	/** Resets this view model to initial conditions. */
 	void Reset();
@@ -63,13 +64,13 @@ public:
 	void SetSimulation(TWeakPtr<FNiagaraEmitterInstance, ESPMode::ThreadSafe> InSimulation);
 
 	/** Gets the emitter represented by this view model. */
-	NIAGARAEDITOR_API UNiagaraEmitter* GetEmitter();
+	NIAGARAEDITOR_API FVersionedNiagaraEmitter GetEmitter();
 
 	/** Gets whether or not this emitter has a parent emitter. */
 	NIAGARAEDITOR_API bool HasParentEmitter() const;
 
 	/** Gets the parent emitter for the emitter represented by this view model, if it has one. */
-	NIAGARAEDITOR_API const UNiagaraEmitter* GetParentEmitter() const;
+	NIAGARAEDITOR_API FVersionedNiagaraEmitter GetParentEmitter() const;
 
 	/** Gets the text representation of the parent emitter name. */
 	NIAGARAEDITOR_API FText GetParentNameText() const;
@@ -125,12 +126,18 @@ public:
 
 	void Cleanup();
 
+	/** When the parent emitter version is changed, this is set to generate the changelist in the stack. */
+	FGuid PreviousEmitterVersion;
+
+	/** Can be used by the ui after a version change to display change notes */
+	FString PythonUpgradeScriptWarnings;
+	
 private:
 	/** Sets this view model to a different emitter. */
-	void SetEmitter(UNiagaraEmitter* InEmitter);
+	void SetEmitter(FVersionedNiagaraEmitter InEmitter);
 
-	void OnVMCompiled(UNiagaraEmitter* InEmitter);
-	void OnGPUCompiled(UNiagaraEmitter* InEmitter);
+	void OnVMCompiled(FVersionedNiagaraEmitter InEmitter);
+	void OnGPUCompiled(FVersionedNiagaraEmitter InEmitter);
 
 	void AddScriptEventHandlers();
 
@@ -142,7 +149,6 @@ private:
 
 	void OnEmitterPropertiesChanged();
 
-private:
 	/** The text format stats display .*/
 	static const FText StatsFormat;
 
@@ -153,7 +159,7 @@ private:
 	static const FText ParticleDisabledDueToScalability;
 
 	/** The emitter object being displayed by the control .*/
-	TWeakObjectPtr<UNiagaraEmitter> Emitter;
+	FVersionedNiagaraEmitterWeakPtr EmitterWeakPtr;
 
 	/** The runtime simulation for the emitter being displayed by the control */
 	TWeakPtr<FNiagaraEmitterInstance, ESPMode::ThreadSafe> Simulation;
diff --git a/Engine/Plugins/FX/Niagara/Source/NiagaraEditor/Public/ViewModels/NiagaraParameterDefinitionsPanelViewModel.h b/Engine/Plugins/FX/Niagara/Source/NiagaraEditor/Public/ViewModels/NiagaraParameterDefinitionsPanelViewModel.h
--- a/Engine/Plugins/FX/Niagara/Source/NiagaraEditor/Public/ViewModels/NiagaraParameterDefinitionsPanelViewModel.h
+++ b/Engine/Plugins/FX/Niagara/Source/NiagaraEditor/Public/ViewModels/NiagaraParameterDefinitionsPanelViewModel.h
@@ -2,7 +2,6 @@
 
 #pragma once
 
-#include "NiagaraEditorCommon.h"
 #include "NiagaraParameterPanelViewModel.h"
 #include "NiagaraToolkitCommon.h"
 #include "Templates/SharedPointer.h"
@@ -13,7 +12,7 @@ class UNiagaraParameterDefinitions;
 class FNiagaraObjectSelection;
 
 /** Interface for view models to the parameter panel. */
-class INiagaraParameterDefinitionsPanelViewModel : public INiagaraImmutableParameterPanelViewModel
+class NIAGARAEDITOR_API INiagaraParameterDefinitionsPanelViewModel : public INiagaraImmutableParameterPanelViewModel
 {
 public:
 	//~ Begin INiagaraImmutableParameterPanelViewModel interface
@@ -22,7 +21,7 @@ public:
 	virtual const TArray<FNiagaraGraphParameterReference> GetGraphParameterReferencesForItem(const FNiagaraParameterPanelItemBase& Item) const override;
 
 	//~ Begin Pure Virtual Methods
-	virtual const TArray<UNiagaraGraph*> GetEditableGraphsConst() const = 0;
+	virtual const TArray<UNiagaraGraph*> GetEditableGraphsConst() const override = 0;
 
 	virtual const TArray<UNiagaraParameterDefinitions*> GetAvailableParameterDefinitionsAssets(bool bSkipSubscribedParameterDefinitions) const = 0;
 	//~ End INiagaraImmutableParameterPanelViewModel interface
diff --git a/Engine/Plugins/FX/Niagara/Source/NiagaraEditor/Public/ViewModels/NiagaraParameterPanelViewModel.h b/Engine/Plugins/FX/Niagara/Source/NiagaraEditor/Public/ViewModels/NiagaraParameterPanelViewModel.h
--- a/Engine/Plugins/FX/Niagara/Source/NiagaraEditor/Public/ViewModels/NiagaraParameterPanelViewModel.h
+++ b/Engine/Plugins/FX/Niagara/Source/NiagaraEditor/Public/ViewModels/NiagaraParameterPanelViewModel.h
@@ -76,7 +76,7 @@ struct FMenuAndSearchBoxWidgets
 };
 
 /** Base Interface for view models to SiagaraParameterPanel and SNiagaraParameterDefinitionsPanel. */
-class INiagaraImmutableParameterPanelViewModel : public TSharedFromThis<INiagaraImmutableParameterPanelViewModel>, public FSelfRegisteringEditorUndoClient
+class NIAGARAEDITOR_API INiagaraImmutableParameterPanelViewModel : public TSharedFromThis<INiagaraImmutableParameterPanelViewModel>, public FSelfRegisteringEditorUndoClient
 {
 public:
 	/** Delegate to signal the view model's state has changed. */
@@ -141,7 +141,7 @@ public:
 	/** Delegate to get the names of all selected parameter items. */
 	DECLARE_DELEGATE_RetVal(TArray<FName>, FOnGetSelectedParameterNames);
 
-	~INiagaraParameterPanelViewModel();
+	virtual ~INiagaraParameterPanelViewModel() override;
 
 	//~ Begin INiagaraImmutableParameterPanelViewModel interface
 	virtual const TArray<UNiagaraGraph*> GetEditableGraphsConst() const = 0;
@@ -258,7 +258,7 @@ protected:
 	mutable TArray<FNiagaraParameterPanelItem> CachedViewedItems; //@todo(ng) consider moving to tset in future
 
 	/** Re-entrancy guard for adding parameters. */
-	mutable bool bIsAddingParameter;
+	mutable bool bIsAddingParameter = false;
 
 	/** Transient UNiagaraScriptVariables used to pass to new FNiagaraParameterPanelItems when the source FNiagaraVariable is not associated with a UNiagaraScriptVariable in a graph. */
 	mutable TMap<FNiagaraVariable, TObjectPtr<UNiagaraScriptVariable>> TransientParameterToScriptVarMap;
@@ -372,7 +372,7 @@ class FNiagaraScriptToolkitParameterPanelViewModel : public INiagaraParameterPan
 {
 public:
 	/** Construct a ScriptToolkit Parameter Panel View Model from a Script View Model. */
-	FNiagaraScriptToolkitParameterPanelViewModel(TSharedPtr<FNiagaraScriptViewModel> InScriptViewModel);
+	NIAGARAEDITOR_API FNiagaraScriptToolkitParameterPanelViewModel(TSharedPtr<FNiagaraScriptViewModel> InScriptViewModel);
 
 	void Init(const FScriptToolkitUIContext& InUIContext);
 
diff --git a/Engine/Plugins/FX/Niagara/Source/NiagaraEditor/Public/ViewModels/NiagaraScriptViewModel.h b/Engine/Plugins/FX/Niagara/Source/NiagaraEditor/Public/ViewModels/NiagaraScriptViewModel.h
--- a/Engine/Plugins/FX/Niagara/Source/NiagaraEditor/Public/ViewModels/NiagaraScriptViewModel.h
+++ b/Engine/Plugins/FX/Niagara/Source/NiagaraEditor/Public/ViewModels/NiagaraScriptViewModel.h
@@ -31,7 +31,7 @@ class FNiagaraScriptViewModel
 public:
 	FNiagaraScriptViewModel(TAttribute<FText> DisplayName, ENiagaraParameterEditMode InParameterEditMode, bool bInIsForDataProcessingOnly);
 
-	virtual ~FNiagaraScriptViewModel();
+	virtual ~FNiagaraScriptViewModel() override;
 
 	//~ Begin NiagaraParameterDefinitionsSubscriberViewModel Interface
 protected:
@@ -46,7 +46,7 @@ public:
 	/** Sets the view model to a different script. */
 	void SetScript(FVersionedNiagaraScript InScript);
 
-	void SetScripts(UNiagaraEmitter* InEmitter);
+	void SetScripts(FVersionedNiagaraEmitter InEmitter);
 
 	/** Gets the view model for the input parameter collection. */
 	TSharedRef<FNiagaraScriptInputCollectionViewModel> GetInputCollectionViewModel();
@@ -69,14 +69,14 @@ public:
 	void CompileStandaloneScript(bool bForceCompile = false);
 
 	/** Get the latest status of this view-model's script compilation.*/
-	virtual ENiagaraScriptCompileStatus GetLatestCompileStatus(FGuid VersionGuid = FGuid());
+	NIAGARAEDITOR_API virtual ENiagaraScriptCompileStatus GetLatestCompileStatus(FGuid VersionGuid = FGuid());
 
 	/** Refreshes the nodes in the script graph, updating the pins to match external changes. */
 	void RefreshNodes();
 
 	//~ FEditorUndoClient Interface
-	virtual void PostUndo(bool bSuccess) override;
-	virtual void PostRedo(bool bSuccess) override { PostUndo(bSuccess); }
+	NIAGARAEDITOR_API virtual void PostUndo(bool bSuccess) override;
+	NIAGARAEDITOR_API virtual void PostRedo(bool bSuccess) override { PostUndo(bSuccess); }
 
 	UNiagaraScript* GetContainerScript(ENiagaraScriptUsage InUsage, FGuid InUsageId = FGuid());
 	UNiagaraScript* GetScript(ENiagaraScriptUsage InUsage, FGuid InUsageId = FGuid());
@@ -110,9 +110,9 @@ protected:
 	/** The script which provides the data for this view model. */
 	TArray<FVersionedNiagaraScriptWeakPtr> Scripts;
 
-	virtual void OnVMScriptCompiled(UNiagaraScript* InScript, const FGuid& ScriptVersion);
+	NIAGARAEDITOR_API virtual void OnVMScriptCompiled(UNiagaraScript* InScript, const FGuid& ScriptVersion);
 
-	virtual void OnGPUScriptCompiled(UNiagaraScript* InScript, const FGuid& ScriptVersion);
+	NIAGARAEDITOR_API virtual void OnGPUScriptCompiled(UNiagaraScript* InScript, const FGuid& ScriptVersion);
 
 	TWeakObjectPtr<UNiagaraScriptSource> Source;
 
diff --git a/Engine/Plugins/FX/Niagara/Source/NiagaraEditor/Public/ViewModels/NiagaraSystemViewModel.h b/Engine/Plugins/FX/Niagara/Source/NiagaraEditor/Public/ViewModels/NiagaraSystemViewModel.h
--- a/Engine/Plugins/FX/Niagara/Source/NiagaraEditor/Public/ViewModels/NiagaraSystemViewModel.h
+++ b/Engine/Plugins/FX/Niagara/Source/NiagaraEditor/Public/ViewModels/NiagaraSystemViewModel.h
@@ -68,7 +68,7 @@ struct NIAGARAEDITOR_API FNiagaraSystemViewModelOptions
 	FNiagaraSystemViewModelOptions();
 
 	/** Whether or not the user can edit emitters from the timeline. */
-	bool bCanModifyEmittersFromTimeline;
+	bool bCanModifyEmittersFromTimeline = true;
 
 	/** A delegate which is used to generate the content for the add menu in sequencer. */
 	FOnGetAddMenuContent OnGetSequencerAddMenuContent;
@@ -83,7 +83,7 @@ struct NIAGARAEDITOR_API FNiagaraSystemViewModelOptions
 	TOptional<const FGuid> MessageLogGuid;
 
 	/** Gets the current editing mode for this system. */
-	ENiagaraSystemViewModelEditMode EditMode;
+	ENiagaraSystemViewModelEditMode EditMode = ENiagaraSystemViewModelEditMode::SystemAsset;
 
 	/** Specifies that the view model is being constructed for data processing only and will not be displayed in the UI. */
 	bool bIsForDataProcessingOnly;
@@ -171,8 +171,8 @@ public:
 	/** Returns whether or not this view model is initialized and safe to use. */
 	bool IsValid() const;
 
-	NIAGARAEDITOR_API ~FNiagaraSystemViewModel();
-
+	NIAGARAEDITOR_API virtual ~FNiagaraSystemViewModel() override;
+	
 	//~ Begin NiagaraParameterDefinitionsSubscriberViewModel Interface
 protected:
 	virtual INiagaraParameterDefinitionsSubscriber* GetParameterDefinitionsSubscriber() override;
@@ -188,7 +188,7 @@ public:
 	NIAGARAEDITOR_API TSharedPtr<FNiagaraEmitterHandleViewModel> GetEmitterHandleViewModelById(FGuid InEmitterHandleId);
 
 	/** Gets an emitter handle view model for the given emitter. Returns an invalid shared ptr if it can't be found. */
-	NIAGARAEDITOR_API TSharedPtr<FNiagaraEmitterHandleViewModel> GetEmitterHandleViewModelForEmitter(UNiagaraEmitter* InEmitter) const;
+	NIAGARAEDITOR_API TSharedPtr<FNiagaraEmitterHandleViewModel> GetEmitterHandleViewModelForEmitter(const FVersionedNiagaraEmitter& InEmitter) const;
 
 	/** Gets the view model for the System script. */
 	TSharedPtr<FNiagaraSystemScriptViewModel> GetSystemScriptViewModel();
@@ -218,7 +218,8 @@ public:
 	NIAGARAEDITOR_API TSharedPtr<FNiagaraEmitterHandleViewModel> AddEmitterFromAssetData(const FAssetData& AssetData);
 
 	/** Adds a new emitter to the System. */
-	NIAGARAEDITOR_API TSharedPtr<FNiagaraEmitterHandleViewModel> AddEmitter(UNiagaraEmitter& Emitter);
+	NIAGARAEDITOR_API TSharedPtr<FNiagaraEmitterHandleViewModel> AddEmitter(UNiagaraEmitter& Emitter, FGuid EmitterVersion);
+	NIAGARAEDITOR_API TSharedPtr<FNiagaraEmitterHandleViewModel> AddEmitter(const FVersionedNiagaraEmitter& Emitter);
 
 	/** Adds an empty emitter to the system. */
 	NIAGARAEDITOR_API TSharedPtr<FNiagaraEmitterHandleViewModel> AddEmptyEmitter();
@@ -284,6 +285,8 @@ public:
 
 	/** Reinitializes all System instances, and rebuilds emitter handle view models and tracks. */
 	NIAGARAEDITOR_API void RefreshAll();
+	
+	NIAGARAEDITOR_API void ResetStack();
 
 	/** Called to notify the system view model that one of the data objects in the system was modified. */
 	void NotifyDataObjectChanged(TArray<UObject*> ChangedObjects, ENiagaraDataObjectChange ChangeType);
@@ -297,6 +300,9 @@ public:
 	/** Clear the captures stats for all the emitters in the current system. */
 	void ClearEmitterStats();
 
+	/** Changes the version of an existing emitter in the system */
+	bool ChangeEmitterVersion(const FVersionedNiagaraEmitter& Emitter, const FGuid& NewVersion);
+
 	/** Isolates the supplied emitters.  This will remove all other emitters from isolation. */
 	NIAGARAEDITOR_API void IsolateEmitters(TArray<FGuid> EmitterHandlesIdsToIsolate);
 
@@ -347,7 +353,7 @@ public:
 	void GetOrderedScriptsForEmitterHandleId(FGuid EmitterHandleId, TArray<UNiagaraScript*>& OutScripts);
 
 	/** Gets all non-event scripts which will execute for an emitter. */
-	void GetOrderedScriptsForEmitter(UNiagaraEmitter* Emitter, TArray<UNiagaraScript*>& OutScripts);
+	void GetOrderedScriptsForEmitter(const FVersionedNiagaraEmitter& Emitter, TArray<UNiagaraScript*>& OutScripts);
 
 	/** Gets the ViewModel for the system overview graph. */
 	NIAGARAEDITOR_API TSharedPtr<FNiagaraOverviewGraphViewModel> GetOverviewGraphViewModel() const;
@@ -436,9 +442,6 @@ private:
 	/** Callback function for when editor settings change. Used to snap to the closest available speed. */
 	void SnapToNextSpeed(const FString& PropertyName, const UNiagaraEditorSettings* Settings);
 
-	/** Gets the content for the add menu in sequencer. */
-	void GetSequencerAddMenuContent(FMenuBuilder& MenuBuilder, TSharedRef<ISequencer> Sequencer);
-
 	/** Updates the compiled versions of data interfaces when their sources change. */
 	void UpdateCompiledDataInterfaces(UNiagaraDataInterface* ChangedDataInterface);
 
@@ -531,7 +534,7 @@ private:
 	/** Called whenever the map of messages associated with the managed Emitter/System changes. */
 	void RefreshAssetMessages();
 
-	const TArray<FNiagaraStackModuleData>& BuildAndCacheStackModuleData(FGuid EmitterHandleId, UNiagaraEmitter* Emitter);
+	const TArray<FNiagaraStackModuleData>& BuildAndCacheStackModuleData(FGuid EmitterHandleId, const FVersionedNiagaraEmitter& Emitter);
 
 private:
 	/** The System being viewed and edited by this view model. */
diff --git a/Engine/Plugins/FX/Niagara/Source/NiagaraEditor/Public/ViewModels/Stack/NiagaraStackEmitterSettingsGroup.h b/Engine/Plugins/FX/Niagara/Source/NiagaraEditor/Public/ViewModels/Stack/NiagaraStackEmitterSettingsGroup.h
--- a/Engine/Plugins/FX/Niagara/Source/NiagaraEditor/Public/ViewModels/Stack/NiagaraStackEmitterSettingsGroup.h
+++ b/Engine/Plugins/FX/Niagara/Source/NiagaraEditor/Public/ViewModels/Stack/NiagaraStackEmitterSettingsGroup.h
@@ -2,6 +2,7 @@
 
 #pragma once
 
+#include "NiagaraEmitter.h"
 #include "ViewModels/Stack/NiagaraStackScriptItemGroup.h"
 #include "ViewModels/Stack/NiagaraStackItem.h"
 #include "NiagaraStackEmitterSettingsGroup.generated.h"
@@ -42,11 +43,11 @@ protected:
 
 private:
 	void EmitterPropertiesChanged();
+	FStackIssueFixDelegate GetUpgradeVersionFix();
 
-private:
 	mutable TOptional<bool> bCanResetToBaseCache;
 
-	TWeakObjectPtr<UNiagaraEmitter> Emitter;
+	FVersionedNiagaraEmitterWeakPtr EmitterWeakPtr;
 
 	UPROPERTY()
 	TObjectPtr<UNiagaraStackObject> EmitterObject;
@@ -80,7 +81,7 @@ private:
 	void SelectSummaryNodesFromEmitterEditorDataRootNodes(TArray<TSharedRef<IDetailTreeNode>> Source, TArray<TSharedRef<IDetailTreeNode>>* Selected);
 
 private:
-	TWeakObjectPtr<UNiagaraEmitter> Emitter;
+	FVersionedNiagaraEmitterWeakPtr Emitter;
 
 	UPROPERTY()
 	TObjectPtr<UNiagaraStackSummaryViewObject> FilteredObject;
diff --git a/Engine/Plugins/FX/Niagara/Source/NiagaraEditor/Public/ViewModels/Stack/NiagaraStackEventScriptItemGroup.h b/Engine/Plugins/FX/Niagara/Source/NiagaraEditor/Public/ViewModels/Stack/NiagaraStackEventScriptItemGroup.h
--- a/Engine/Plugins/FX/Niagara/Source/NiagaraEditor/Public/ViewModels/Stack/NiagaraStackEventScriptItemGroup.h
+++ b/Engine/Plugins/FX/Niagara/Source/NiagaraEditor/Public/ViewModels/Stack/NiagaraStackEventScriptItemGroup.h
@@ -13,6 +13,22 @@ class UNiagaraStackObject;
 struct FNiagaraEventScriptProperties;
 class IDetailTreeNode;
 
+// This is a wrapper class used for the details customization in the stack, since the event script properties were moved from the emitter object into the version data struct
+UCLASS()
+class NIAGARAEDITOR_API UNiagaraStackEventWrapper : public UObject
+{
+	GENERATED_BODY()
+
+public:
+
+	UPROPERTY(EditAnywhere, Category = "Events")
+	TArray<FNiagaraEventScriptProperties> EventHandlerScriptProps;
+
+	FVersionedNiagaraEmitterWeakPtr EmitterWeakPtr;
+
+	virtual void PostEditChangeProperty(FPropertyChangedEvent& PropertyChangedEvent) override;
+};
+
 UCLASS()
 class NIAGARAEDITOR_API UNiagaraStackEventHandlerPropertiesItem : public UNiagaraStackItem
 {
@@ -50,10 +66,13 @@ private:
 
 	mutable TOptional<bool> bCanResetToBaseCache;
 
-	TWeakObjectPtr<UNiagaraEmitter> Emitter;
+	FVersionedNiagaraEmitterWeakPtr EmitterWeakPtr;
 
 	UPROPERTY()
 	TObjectPtr<UNiagaraStackObject> EmitterObject;
+
+	UPROPERTY()
+	TObjectPtr<UNiagaraStackEventWrapper> EventWrapper;
 };
 
 UCLASS()
diff --git a/Engine/Plugins/FX/Niagara/Source/NiagaraEditor/Public/ViewModels/Stack/NiagaraStackFilteredObject.h b/Engine/Plugins/FX/Niagara/Source/NiagaraEditor/Public/ViewModels/Stack/NiagaraStackFilteredObject.h
--- a/Engine/Plugins/FX/Niagara/Source/NiagaraEditor/Public/ViewModels/Stack/NiagaraStackFilteredObject.h
+++ b/Engine/Plugins/FX/Niagara/Source/NiagaraEditor/Public/ViewModels/Stack/NiagaraStackFilteredObject.h
@@ -3,7 +3,6 @@
 #pragma once
 
 #include "ViewModels/Stack/NiagaraStackItem.h"
-#include "Misc/NotifyHook.h"
 #include "PropertyEditorDelegates.h"
 #include "NiagaraStackFunctionInputCollection.h"
 #include "NiagaraStackFilteredObject.generated.h"
@@ -27,7 +26,7 @@ public:
 
 	virtual FText GetDisplayName() const override;
 	virtual bool GetShouldShowInStack() const override;
-	virtual bool GetIsEnabled() const;
+	virtual bool GetIsEnabled() const override;
 protected:
 
 	virtual void RefreshChildrenInternal(const TArray<UNiagaraStackEntry*>& CurrentChildren, TArray<UNiagaraStackEntry*>& NewChildren, TArray<FStackIssue>& NewIssues) override;
@@ -43,7 +42,7 @@ private:
 	void OnViewStateChanged();
 private:
 
-	UNiagaraEmitter* Emitter;
+	FVersionedNiagaraEmitter VersionedEmitter;
 
 	TMap<FGuid, UNiagaraStackFunctionInputCollection*> KnownInputCollections;
 
diff --git a/Engine/Plugins/FX/Niagara/Source/NiagaraEditor/Public/ViewModels/Stack/NiagaraStackGraphUtilities.h b/Engine/Plugins/FX/Niagara/Source/NiagaraEditor/Public/ViewModels/Stack/NiagaraStackGraphUtilities.h
--- a/Engine/Plugins/FX/Niagara/Source/NiagaraEditor/Public/ViewModels/Stack/NiagaraStackGraphUtilities.h
+++ b/Engine/Plugins/FX/Niagara/Source/NiagaraEditor/Public/ViewModels/Stack/NiagaraStackGraphUtilities.h
@@ -39,8 +39,6 @@ namespace FNiagaraStackGraphUtilities
 
 	void RelayoutGraph(UEdGraph& Graph);
 
-	void GetWrittenVariablesForGraph(UEdGraph& Graph, TArray<FNiagaraVariable>& OutWrittenVariables);
-
 	void ConnectPinToInputNode(UEdGraphPin& Pin, UNiagaraNodeInput& InputNode);
 
 	UEdGraphPin* GetParameterMapInputPin(UNiagaraNode& Node);
@@ -62,6 +60,7 @@ namespace FNiagaraStackGraphUtilities
 	UNiagaraNodeInput* GetEmitterInputNodeForStackNode(UNiagaraNode& StackNode);
 
 	void CheckForDeprecatedScriptVersion(UNiagaraNodeFunctionCall* InputFunctionCallNode, const FString& StackEditorDataKey, UNiagaraStackEntry::FStackIssueFixDelegate VersionUpgradeFix, TArray<UNiagaraStackEntry::FStackIssue>& OutIssues);
+	void CheckForDeprecatedEmitterVersion(TSharedPtr<FNiagaraEmitterViewModel> ViewModel, const FString& StackEditorDataKey, UNiagaraStackEntry::FStackIssueFixDelegate VersionUpgradeFix, TArray<UNiagaraStackEntry::FStackIssue>& OutIssues);
 
 	struct FStackNodeGroup
 	{
@@ -84,8 +83,8 @@ namespace FNiagaraStackGraphUtilities
 
 	FString GenerateStackModuleEditorDataKey(UNiagaraNodeFunctionCall& ModuleNode);
 
-	TArray<FName> StackContextResolution(UNiagaraEmitter* OwningEmitter, UNiagaraNodeOutput* OutputNodeInChain);
-	void BuildParameterMapHistoryWithStackContextResolution(UNiagaraEmitter* OwningEmitter, UNiagaraNodeOutput* OutputNodeInChain, UNiagaraNode* NodeToVisit, FNiagaraParameterMapHistoryBuilder& Builder, bool bRecursive = true, bool bFilterForCompilation = true);
+	TArray<FName> StackContextResolution(FVersionedNiagaraEmitter OwningEmitter, UNiagaraNodeOutput* OutputNodeInChain);
+	void BuildParameterMapHistoryWithStackContextResolution(FVersionedNiagaraEmitter OwningEmitter, UNiagaraNodeOutput* OutputNodeInChain, UNiagaraNode* NodeToVisit, FNiagaraParameterMapHistoryBuilder& Builder, bool bRecursive = true, bool bFilterForCompilation = true);
 
 	enum class ENiagaraGetStackFunctionInputPinsOptions
 	{
@@ -208,9 +207,9 @@ namespace FNiagaraStackGraphUtilities
 
 	FNiagaraVariable CreateRapidIterationParameter(const FString& UniqueEmitterName, ENiagaraScriptUsage ScriptUsage, const FName& AliasedInputName, const FNiagaraTypeDefinition& InputType);
 
-	void CleanUpStaleRapidIterationParameters(UNiagaraScript& Script, UNiagaraEmitter& OwningEmitter);
+	void CleanUpStaleRapidIterationParameters(UNiagaraScript& Script, FVersionedNiagaraEmitter OwningEmitter);
 
-	void CleanUpStaleRapidIterationParameters(UNiagaraEmitter& Emitter);
+	void CleanUpStaleRapidIterationParameters(FVersionedNiagaraEmitter Emitter);
 
 	void GetNewParameterAvailableTypes(TArray<FNiagaraTypeDefinition>& OutAvailableTypes, FName Namespace);
 
@@ -232,12 +231,12 @@ namespace FNiagaraStackGraphUtilities
 
 	void RebuildEmitterNodes(UNiagaraSystem& System);
 
-	void FindAffectedScripts(UNiagaraSystem* System, UNiagaraEmitter* Emitter, UNiagaraNodeFunctionCall& ModuleNode, TArray<TWeakObjectPtr<UNiagaraScript>>& OutAffectedScripts);
+	void FindAffectedScripts(UNiagaraSystem* System, FVersionedNiagaraEmitter Emitter, UNiagaraNodeFunctionCall& ModuleNode, TArray<TWeakObjectPtr<UNiagaraScript>>& OutAffectedScripts);
 
-	void RenameReferencingParameters(UNiagaraSystem* System, UNiagaraEmitter* Emitter, UNiagaraNodeFunctionCall& FunctionCallNode, const FString& OldName, const FString& NewName);
+	void RenameReferencingParameters(UNiagaraSystem* System, FVersionedNiagaraEmitter Emitter, UNiagaraNodeFunctionCall& FunctionCallNode, const FString& OldName, const FString& NewName);
 
-	void GatherRenamedStackFunctionOutputVariableNames(UNiagaraEmitter* Emitter, UNiagaraNodeFunctionCall& FunctionCallNode, const FString& OldFunctionName, const FString& NewFunctionName, TMap<FName, FName>& OutOldToNewNameMap);
-	void GatherRenamedStackFunctionInputAndOutputVariableNames(UNiagaraEmitter* Emitter, UNiagaraNodeFunctionCall& FunctionCallNode, const FString& OldFunctionName, const FString& NewFunctionName, TMap<FName, FName>& OutOldToNewNameMap);
+	void GatherRenamedStackFunctionOutputVariableNames(FVersionedNiagaraEmitter Emitter, UNiagaraNodeFunctionCall& FunctionCallNode, const FString& OldFunctionName, const FString& NewFunctionName, TMap<FName, FName>& OutOldToNewNameMap);
+	void GatherRenamedStackFunctionInputAndOutputVariableNames(FVersionedNiagaraEmitter Emitter, UNiagaraNodeFunctionCall& FunctionCallNode, const FString& OldFunctionName, const FString& NewFunctionName, TMap<FName, FName>& OutOldToNewNameMap);
 
 	enum class EStackEditContext
 	{
@@ -255,7 +254,7 @@ namespace FNiagaraStackGraphUtilities
 
 	void RenameAssignmentTarget(
 		UNiagaraSystem& OwningSystem,
-		UNiagaraEmitter* OwningEmitter,
+		FVersionedNiagaraEmitter OwningEmitter,
 		UNiagaraScript& OwningScript,
 		UNiagaraNodeAssignment& OwningAssignmentNode,
 		FNiagaraVariable CurrentAssignmentTarget,
@@ -269,7 +268,7 @@ namespace FNiagaraStackGraphUtilities
 
 	void PopulateFunctionCallNameBindings(UNiagaraNodeFunctionCall& InFunctionCallNode);
 
-	void SynchronizeReferencingMapPinsWithFunctionCall(UNiagaraNodeFunctionCall& InFunctionCallNode);
+	NIAGARAEDITOR_API void SynchronizeReferencingMapPinsWithFunctionCall(UNiagaraNodeFunctionCall& InFunctionCallNode);
 
 	FGuid GetScriptVariableIdForLinkedOutputHandle(const FNiagaraParameterHandle& LinkedOutputHandle, FNiagaraTypeDefinition LinkedType, UNiagaraGraph& TargetGraph);
 
@@ -279,6 +278,6 @@ namespace FNiagaraStackGraphUtilities
 
 		void GetModuleScriptAssetsByDependencyProvided(FName DependencyName, TOptional<ENiagaraScriptUsage> RequiredUsage, TArray<FAssetData>& OutAssets);
 
-		int32 FindBestIndexForModuleInStack(UNiagaraNodeFunctionCall& ModuleNode, UEdGraph& EmitterScriptGraph);
+		int32 FindBestIndexForModuleInStack(UNiagaraNodeFunctionCall& ModuleNode, UNiagaraGraph& EmitterScriptGraph);
 	}
 }
diff --git a/Engine/Plugins/FX/Niagara/Source/NiagaraEditor/Public/ViewModels/Stack/NiagaraStackRenderItemGroup.h b/Engine/Plugins/FX/Niagara/Source/NiagaraEditor/Public/ViewModels/Stack/NiagaraStackRenderItemGroup.h
--- a/Engine/Plugins/FX/Niagara/Source/NiagaraEditor/Public/ViewModels/Stack/NiagaraStackRenderItemGroup.h
+++ b/Engine/Plugins/FX/Niagara/Source/NiagaraEditor/Public/ViewModels/Stack/NiagaraStackRenderItemGroup.h
@@ -2,6 +2,7 @@
 
 #pragma once
 
+#include "NiagaraEmitter.h"
 #include "ViewModels/Stack/NiagaraStackItemGroup.h"
 #include "NiagaraStackRenderItemGroup.generated.h"
 
@@ -35,5 +36,5 @@ private:
 	void OnRendererAdded(UNiagaraRendererProperties* RendererProperties) const;
 	TSharedPtr<INiagaraStackItemGroupAddUtilities> AddUtilities;
 
-	TWeakObjectPtr<UNiagaraEmitter> EmitterWeak;
+	FVersionedNiagaraEmitterWeakPtr EmitterWeak;
 };
