Added versioning for Niagara emitters

diff --git a/Engine/Plugins/FX/Niagara/Source/NiagaraEditor/Public/ViewModels/Stack/NiagaraStackRendererItem.h b/Engine/Plugins/FX/Niagara/Source/NiagaraEditor/Public/ViewModels/Stack/NiagaraStackRendererItem.h
--- a/Engine/Plugins/FX/Niagara/Source/NiagaraEditor/Public/ViewModels/Stack/NiagaraStackRendererItem.h
+++ b/Engine/Plugins/FX/Niagara/Source/NiagaraEditor/Public/ViewModels/Stack/NiagaraStackRendererItem.h
@@ -7,6 +7,7 @@
 #include "NiagaraTypes.h"
 #include "NiagaraStackRendererItem.generated.h"
 
+struct FVersionedNiagaraEmitterData;
 class UNiagaraEmitter;
 class UNiagaraRendererProperties;
 class UNiagaraStackObject;
@@ -63,8 +64,8 @@ public:
 	virtual bool TestCanResetToBaseWithMessage(FText& OutCanResetToBaseMessage) const override;
 	virtual void ResetToBase() override;
 
-	static TArray<FNiagaraVariable> GetMissingVariables(UNiagaraRendererProperties* RendererProperties, UNiagaraEmitter* Emitter);
-	static bool AddMissingVariable(UNiagaraEmitter* Emitter, const FNiagaraVariable& Variable);
+	static TArray<FNiagaraVariable> GetMissingVariables(UNiagaraRendererProperties* RendererProperties, const FVersionedNiagaraEmitterData* EmitterData);
+	static bool AddMissingVariable(const FVersionedNiagaraEmitterData* EmitterData, const FNiagaraVariable& Variable);
 
 	bool IsExcludedFromScalability() const;
 	bool IsOwningEmitterExcludedFromScalability() const;
diff --git a/Engine/Plugins/FX/Niagara/Source/NiagaraEditor/Public/ViewModels/Stack/NiagaraStackRoot.h b/Engine/Plugins/FX/Niagara/Source/NiagaraEditor/Public/ViewModels/Stack/NiagaraStackRoot.h
--- a/Engine/Plugins/FX/Niagara/Source/NiagaraEditor/Public/ViewModels/Stack/NiagaraStackRoot.h
+++ b/Engine/Plugins/FX/Niagara/Source/NiagaraEditor/Public/ViewModels/Stack/NiagaraStackRoot.h
@@ -23,7 +23,7 @@ public:
 	UNiagaraStackRoot();
 	
 	void Initialize(FRequiredEntryData InRequiredEntryData, bool bInIncludeSystemInformation, bool bInIncludeEmitterInformation);
-	void FinalizeInternal() override;
+	virtual void FinalizeInternal() override;
 
 	virtual bool GetCanExpand() const override;
 	virtual bool GetShouldShowInStack() const override;
diff --git a/Engine/Plugins/FX/Niagara/Source/NiagaraEditor/Public/ViewModels/Stack/NiagaraStackSummaryViewInputCollection.h b/Engine/Plugins/FX/Niagara/Source/NiagaraEditor/Public/ViewModels/Stack/NiagaraStackSummaryViewInputCollection.h
--- a/Engine/Plugins/FX/Niagara/Source/NiagaraEditor/Public/ViewModels/Stack/NiagaraStackSummaryViewInputCollection.h
+++ b/Engine/Plugins/FX/Niagara/Source/NiagaraEditor/Public/ViewModels/Stack/NiagaraStackSummaryViewInputCollection.h
@@ -2,8 +2,8 @@
 
 #pragma once
 
+#include "NiagaraEmitter.h"
 #include "ViewModels/Stack/NiagaraStackItem.h"
-#include "Misc/NotifyHook.h"
 #include "PropertyEditorDelegates.h"
 #include "NiagaraStackFunctionInputCollection.h"
 #include "NiagaraStackSummaryViewInputCollection.generated.h"
@@ -25,7 +25,7 @@ public:
 public:
 	UNiagaraStackSummaryViewObject();
 
-	void Initialize(FRequiredEntryData InRequiredEntryData, UNiagaraEmitter* InEmitter, FString InOwningStackItemEditorDataKey);
+	void Initialize(FRequiredEntryData InRequiredEntryData, FVersionedNiagaraEmitterWeakPtr InEmitter, FString InOwningStackItemEditorDataKey);
 	virtual void FinalizeInternal() override;
 
 	virtual FText GetDisplayName() const override;
@@ -46,6 +46,6 @@ private:
 	void OnViewStateChanged();
 private:
 
-	UNiagaraEmitter* Emitter;
+	FVersionedNiagaraEmitterWeakPtr Emitter;
 	TMap<FGuid, UNiagaraStackFunctionInputCollection*> KnownInputCollections;
 };
diff --git a/Engine/Plugins/FX/Niagara/Source/NiagaraEditorWidgets/Private/DetailCustomizations/NiagaraDataInterfaceMeshRendererInfoDetails.cpp b/Engine/Plugins/FX/Niagara/Source/NiagaraEditorWidgets/Private/DetailCustomizations/NiagaraDataInterfaceMeshRendererInfoDetails.cpp
--- a/Engine/Plugins/FX/Niagara/Source/NiagaraEditorWidgets/Private/DetailCustomizations/NiagaraDataInterfaceMeshRendererInfoDetails.cpp
+++ b/Engine/Plugins/FX/Niagara/Source/NiagaraEditorWidgets/Private/DetailCustomizations/NiagaraDataInterfaceMeshRendererInfoDetails.cpp
@@ -7,8 +7,8 @@
 #include "NiagaraMeshRendererProperties.h"
 #include "DetailLayoutBuilder.h"
 #include "DetailCategoryBuilder.h"
-#include "IDetailGroup.h"
 #include "NiagaraComponent.h"
+#include "NiagaraEmitter.h"
 #include "NiagaraStackEditorData.h"
 #include "NiagaraEmitterEditorData.h"
 #include "NiagaraEditorModule.h"
@@ -110,11 +110,11 @@ void FNiagaraDataInterfaceMeshRendererInfoDetails::GenerateRendererList()
 	if (auto Interface = DataInterface.Get())
 	{
 		UNiagaraSystem* System = nullptr;
-		UNiagaraEmitter* Emitter = nullptr;
-		FNiagaraEditorModule::Get().GetTargetSystemAndEmitterForDataInterface(Interface, System, Emitter);
-		if (System && Emitter)
+		FVersionedNiagaraEmitter EmitterHandle;
+		FNiagaraEditorModule::Get().GetTargetSystemAndEmitterForDataInterface(Interface, System, EmitterHandle);
+		if (System && EmitterHandle.GetEmitterData())
 		{
-			for (auto RendererProps : Emitter->GetRenderers())
+			for (auto RendererProps : EmitterHandle.GetEmitterData()->GetRenderers())
 			{
 				if (auto MeshProps = Cast<UNiagaraMeshRendererProperties>(RendererProps))
 				{
@@ -186,7 +186,7 @@ FText FNiagaraDataInterfaceMeshRendererInfoDetails::CreateRendererTextLabel(cons
 	UNiagaraEmitter* Emitter = Properties->GetTypedOuter<UNiagaraEmitter>();
 	check(Emitter);
 
-	UNiagaraEmitterEditorData* EmitterEditorData = static_cast<UNiagaraEmitterEditorData*>(Emitter->GetEditorData());
+	UNiagaraEmitterEditorData* EmitterEditorData = static_cast<UNiagaraEmitterEditorData*>(Properties->GetEmitterData()->GetEditorData());
 	if (EmitterEditorData == nullptr)
 	{
 		return LOCTEXT("NoneOption", "None");
diff --git a/Engine/Plugins/FX/Niagara/Source/NiagaraEditorWidgets/Private/SNiagaraOverviewStackNode.cpp b/Engine/Plugins/FX/Niagara/Source/NiagaraEditorWidgets/Private/SNiagaraOverviewStackNode.cpp
--- a/Engine/Plugins/FX/Niagara/Source/NiagaraEditorWidgets/Private/SNiagaraOverviewStackNode.cpp
+++ b/Engine/Plugins/FX/Niagara/Source/NiagaraEditorWidgets/Private/SNiagaraOverviewStackNode.cpp
@@ -1,6 +1,8 @@
 // Copyright Epic Games, Inc. All Rights Reserved.
 
 #include "SNiagaraOverviewStackNode.h"
+
+#include "EditorFontGlyphs.h"
 #include "NiagaraOverviewNode.h"
 #include "ViewModels/NiagaraSystemViewModel.h"
 #include "ViewModels/NiagaraEmitterViewModel.h"
@@ -19,7 +21,6 @@
 #include "Widgets/Layout/SBox.h"
 #include "Widgets/Input/SButton.h"
 #include "Widgets/Input/SCheckBox.h"
-#include "AssetThumbnail.h"
 #include "NiagaraEmitterInstance.h"
 #include "NiagaraEditorUtilities.h"
 #include "Widgets/SToolTip.h"
@@ -27,17 +28,17 @@
 #include "NiagaraRendererProperties.h"
 #include "SGraphPanel.h"
 #include "Widgets/SWidget.h"
-#include "Subsystems/AssetEditorSubsystem.h"
 #include "ViewModels/Stack/NiagaraStackRendererItem.h"
 #include "Materials/MaterialInterface.h"
 #include "Materials/Material.h"
 #include "Framework/MultiBox/MultiBoxBuilder.h"
 #include "Styling/StyleColors.h"
+#include "Widgets/Text/SInlineEditableTextBlock.h"
 #include "Styling/AppStyle.h"
 
 #define LOCTEXT_NAMESPACE "NiagaraOverviewStackNode"
 
-const float ThumbnailSize = 24.0f;
+constexpr float ThumbnailSize = 24.0f;
 
 void SNiagaraOverviewStackNode::Construct(const FArguments& InArgs, UNiagaraOverviewNode* InNode)
 {
@@ -145,6 +146,8 @@ TSharedRef<SWidget> SNiagaraOverviewStackNode::CreateTitleRightWidget()
 		return SNullWidget::NullWidget;
 	}
 	return SNew(SHorizontalBox)
+
+		// open parent button
 		+ SHorizontalBox::Slot()
 		.AutoWidth()
 		.HAlign(HAlign_Right)
@@ -153,7 +156,7 @@ TSharedRef<SWidget> SNiagaraOverviewStackNode::CreateTitleRightWidget()
 		[
 			SNew(SButton)
 			.IsFocusable(false)
-			.ToolTipText(LOCTEXT("OpenAndFocusParentEmitterToolTip", "Open and Focus Parent Emitter"))
+			.ToolTipText(this, &SNiagaraOverviewStackNode::OpenParentEmitterTooltip)
 			.ButtonStyle(FAppStyle::Get(), "HoverHintOnly")
 			.ContentPadding(2)
 			.OnClicked(this, &SNiagaraOverviewStackNode::OpenParentEmitter)
@@ -166,6 +169,31 @@ TSharedRef<SWidget> SNiagaraOverviewStackNode::CreateTitleRightWidget()
 				.ColorAndOpacity(FSlateColor::UseForeground())
 			]
 		]
+	
+		// version selector
+		+ SHorizontalBox::Slot()
+		.HAlign(HAlign_Right)
+		.AutoWidth()
+		.Padding(1, 0, 2, 0)
+		[
+			SNew(SComboButton)
+			.HasDownArrow(false)
+			.ToolTipText(LOCTEXT("ChangeEmitterVersionToolTip", "Change the parent emitter version"))
+			.ButtonStyle(FAppStyle::Get(), "HoverHintOnly")
+			.ForegroundColor(FSlateColor::UseForeground())
+			.OnGetMenuContent(this, &SNiagaraOverviewStackNode::GetVersionSelectorDropdownMenu)
+			.ContentPadding(FMargin(2))
+			.Visibility(this, &SNiagaraOverviewStackNode::GetVersionSelectorVisibility)
+			.ButtonContent()
+			[
+				SNew(STextBlock)
+				.Font(FAppStyle::Get().GetFontStyle("FontAwesome.10"))
+				.ColorAndOpacity(this, &SNiagaraOverviewStackNode::GetVersionSelectorColor)
+				.Text(FEditorFontGlyphs::Random)
+			]
+		]
+
+		// issue/error icon
 		+ SHorizontalBox::Slot()
 		.AutoWidth()
 		.HAlign(HAlign_Right)
@@ -285,9 +313,9 @@ void SNiagaraOverviewStackNode::Tick(const FGeometry& AllottedGeometry, const do
 
 void SNiagaraOverviewStackNode::OnMaterialCompiled(class UMaterialInterface* MaterialInterface)
 {
-	bool bUsingThisMaterial = false;
 	if (EmitterHandleViewModelWeak.IsValid())
 	{
+		bool bUsingThisMaterial = false;
 		EmitterHandleViewModelWeak.Pin()->GetRendererEntries(PreviewStackEntries);
 		FNiagaraEmitterInstance* InInstance = EmitterHandleViewModelWeak.Pin()->GetEmitterViewModel()->GetSimulation().IsValid() ? EmitterHandleViewModelWeak.Pin()->GetEmitterViewModel()->GetSimulation().Pin().Get() : nullptr;
 		for (UNiagaraStackEntry* Entry : PreviewStackEntries)
@@ -313,7 +341,7 @@ void SNiagaraOverviewStackNode::OnMaterialCompiled(class UMaterialInterface* Mat
 
 void SNiagaraOverviewStackNode::CreateBottomSummaryExpander()
 {
-	UNiagaraEmitter* Emitter = EmitterHandleViewModelWeak.IsValid()? EmitterHandleViewModelWeak.Pin()->GetEmitterViewModel()->GetEmitter() : nullptr;
+	UNiagaraEmitter* Emitter = EmitterHandleViewModelWeak.IsValid()? EmitterHandleViewModelWeak.Pin()->GetEmitterViewModel()->GetEmitter().Emitter : nullptr;
 	if (BottomSummaryExpander.IsValid() || !Emitter)
 	{
 		return;
@@ -640,7 +668,7 @@ EVisibility SNiagaraOverviewStackNode::ShowExcludedOverlay() const
 	// we only want actual results in scalability mode and for nodes representing emitters (not system nodes)
 	if(bScalabilityModeActive && EmitterHandleViewModelWeak.IsValid())
 	{		
-		return EmitterHandleViewModelWeak.Pin()->GetEmitterHandle()->GetInstance()->IsAllowedByScalability() ? EVisibility::Hidden : EVisibility::HitTestInvisible;
+		return EmitterHandleViewModelWeak.Pin()->GetEmitterHandle()->GetEmitterData()->IsAllowedByScalability() ? EVisibility::Hidden : EVisibility::HitTestInvisible;
 	}
 	
 	return EVisibility::Hidden; 
@@ -656,27 +684,108 @@ float SNiagaraOverviewStackNode::GetGraphZoomDistanceAlphaMultiplier() const
 FReply SNiagaraOverviewStackNode::OpenParentEmitter()
 {
 	TSharedPtr<FNiagaraEmitterHandleViewModel> EmitterHandleViewModel = EmitterHandleViewModelWeak.Pin();
-
 	if (EmitterHandleViewModel.IsValid())
 	{
-		UNiagaraEmitter* ParentEmitter = const_cast<UNiagaraEmitter*>(EmitterHandleViewModel->GetEmitterViewModel()->GetParentEmitter());
-		if (ParentEmitter != nullptr)
+		FNiagaraEditorUtilities::OpenParentEmitterForEdit(EmitterHandleViewModel->GetEmitterViewModel());
+	}
+	return FReply::Handled();
+}
+
+FText SNiagaraOverviewStackNode::OpenParentEmitterTooltip() const
+{
+	TSharedPtr<FNiagaraEmitterHandleViewModel> EmitterHandleViewModel = EmitterHandleViewModelWeak.Pin();
+	if (EmitterHandleViewModel.IsValid() && EmitterHandleViewModel->GetEmitterViewModel()->HasParentEmitter() && EmitterHandleViewModel->GetEmitterViewModel()->GetParentEmitter().Emitter->IsVersioningEnabled())
+	{
+		FVersionedNiagaraEmitter ParentEmitter = EmitterHandleViewModel->GetEmitterViewModel()->GetParentEmitter();
+		FText ParentName = FText::FromString(ParentEmitter.Emitter->GetUniqueEmitterName());
+		if (FVersionedNiagaraEmitterData* EmitterData = ParentEmitter.GetEmitterData())
 		{
-			GEditor->GetEditorSubsystem<UAssetEditorSubsystem>()->OpenEditorForAsset(ParentEmitter);
+			return FText::Format(LOCTEXT("OpenAndFocusVersionedParentToolTip", "Open and Focus Parent Emitter:\n{0} - v{1}.{2}"), ParentName, EmitterData->Version.MajorVersion, EmitterData->Version.MinorVersion);
 		}
 	}
-	return FReply::Handled();
+	return LOCTEXT("OpenAndFocusParentEmitterToolTip", "Open and Focus Parent Emitter");
 }
 
 EVisibility SNiagaraOverviewStackNode::GetOpenParentEmitterVisibility() const
 {
 	TSharedPtr<FNiagaraEmitterHandleViewModel> EmitterHandleViewModel = EmitterHandleViewModelWeak.Pin();
 	return EmitterHandleViewModel.IsValid() && 
-		EmitterHandleViewModel->GetEmitterViewModel()->GetParentEmitter() != nullptr
+		EmitterHandleViewModel->GetEmitterViewModel()->HasParentEmitter()
+		? EVisibility::Visible 
+		: EVisibility::Collapsed;
+}
+
+EVisibility SNiagaraOverviewStackNode::GetVersionSelectorVisibility() const
+{
+	TSharedPtr<FNiagaraEmitterHandleViewModel> EmitterHandleViewModel = EmitterHandleViewModelWeak.Pin();
+	return EmitterHandleViewModel.IsValid() && 
+		EmitterHandleViewModel->GetEmitterViewModel()->HasParentEmitter() &&
+		EmitterHandleViewModel->GetEmitterViewModel()->GetParentEmitter().Emitter->IsVersioningEnabled()
 		? EVisibility::Visible 
 		: EVisibility::Collapsed;
 }
 
+FSlateColor SNiagaraOverviewStackNode::GetVersionSelectorColor() const
+{
+	if (TSharedPtr<FNiagaraEmitterHandleViewModel> EmitterHandleViewModel = EmitterHandleViewModelWeak.Pin())
+	{
+		FVersionedNiagaraEmitter ParentEmitter = EmitterHandleViewModel->GetEmitterViewModel()->GetParentEmitter();
+		if (ParentEmitter.Emitter && ParentEmitter.Emitter->IsVersioningEnabled())
+		{
+			FNiagaraAssetVersion ExposedVersion = ParentEmitter.Emitter->GetExposedVersion();
+			if (ParentEmitter.GetEmitterData() == nullptr || ParentEmitter.GetEmitterData()->Version < ExposedVersion)
+			{
+				return FNiagaraEditorWidgetsStyle::Get().GetColor("NiagaraEditor.Stack.IconColor.VersionUpgrade");
+			}
+		}
+	}
+	return FNiagaraEditorWidgetsStyle::Get().GetColor("NiagaraEditor.Stack.FlatButtonColor");
+}
+
+TSharedRef<SWidget> SNiagaraOverviewStackNode::GetVersionSelectorDropdownMenu()
+{
+	constexpr bool bShouldCloseWindowAfterMenuSelection = true;
+	FMenuBuilder MenuBuilder(bShouldCloseWindowAfterMenuSelection, nullptr);
+
+	if (TSharedPtr<FNiagaraEmitterHandleViewModel> EmitterHandleViewModel = EmitterHandleViewModelWeak.Pin())
+	{
+		FVersionedNiagaraEmitter ParentEmitter = EmitterHandleViewModel->GetEmitterViewModel()->GetParentEmitter();
+		TArray<FNiagaraAssetVersion> AssetVersions = ParentEmitter.Emitter->GetAllAvailableVersions();
+		for (FNiagaraAssetVersion& Version : AssetVersions)
+		{
+			if (!Version.bIsVisibleInVersionSelector)
+			{
+				continue;
+			}
+			FVersionedNiagaraEmitterData* EmitterData = ParentEmitter.Emitter->GetEmitterData(Version.VersionGuid);
+			bool bIsSelected = ParentEmitter.Version == Version.VersionGuid;
+		
+			FText Tooltip = LOCTEXT("NiagaraSelectVersion_Tooltip", "Select this version to use for the emitter");
+			if (EmitterData && !EmitterData->VersionChangeDescription.IsEmpty())
+			{
+				Tooltip = FText::Format(LOCTEXT("NiagaraSelectVersionChangelist_Tooltip", "Select this version to use for the emitter. Change description for this version:\n{0}"), EmitterData->VersionChangeDescription);
+			}
+		
+			FUIAction UIAction(FExecuteAction::CreateSP(this, &SNiagaraOverviewStackNode::SwitchToVersion, Version),
+			FCanExecuteAction(),
+			FIsActionChecked::CreateLambda([bIsSelected]() { return bIsSelected; }));
+			FText Format = (Version == ParentEmitter.Emitter->GetExposedVersion()) ? FText::FromString("{0}.{1}*") : FText::FromString("{0}.{1}");
+			FText Label = FText::Format(Format, Version.MajorVersion, Version.MinorVersion);
+			MenuBuilder.AddMenuEntry(Label, Tooltip, FSlateIcon(), UIAction, NAME_None, EUserInterfaceActionType::RadioButton);	
+		}
+	}
+
+	return MenuBuilder.MakeWidget();
+}
+
+// this switches the referenced parent version, for the version selector in the toolbar see FNiagaraSystemViewModel::ChangeEmitterVersion
+void SNiagaraOverviewStackNode::SwitchToVersion(FNiagaraAssetVersion Version)
+{
+	if (TSharedPtr<FNiagaraEmitterHandleViewModel> EmitterHandleViewModel = EmitterHandleViewModelWeak.Pin())
+	{
+		FNiagaraEditorUtilities::SwitchParentEmitterVersion(EmitterHandleViewModel->GetEmitterViewModel(), EmitterHandleViewModel->GetOwningSystemViewModel(), Version.VersionGuid);
+	}
+}
 
 const FSlateBrush* SNiagaraOverviewStackNode::GetSummaryViewButtonBrush() const
 {
@@ -703,8 +812,7 @@ FText SNiagaraOverviewStackNode::GetSummaryViewCollapseTooltipText() const
 
 FReply SNiagaraOverviewStackNode::ExpandSummaryViewClicked()
 {
-	UNiagaraEmitterEditorData* EditorData = EmitterHandleViewModelWeak.IsValid()? &EmitterHandleViewModelWeak.Pin()->GetEmitterViewModel()->GetOrCreateEditorData() : nullptr;
-	if (EditorData)
+	if (UNiagaraEmitterEditorData* EditorData = EmitterHandleViewModelWeak.IsValid()? &EmitterHandleViewModelWeak.Pin()->GetEmitterViewModel()->GetOrCreateEditorData() : nullptr)
 	{
 		EditorData->ToggleShowSummaryView();
 	}
diff --git a/Engine/Plugins/FX/Niagara/Source/NiagaraEditorWidgets/Private/SNiagaraOverviewStackNode.h b/Engine/Plugins/FX/Niagara/Source/NiagaraEditorWidgets/Private/SNiagaraOverviewStackNode.h
--- a/Engine/Plugins/FX/Niagara/Source/NiagaraEditorWidgets/Private/SNiagaraOverviewStackNode.h
+++ b/Engine/Plugins/FX/Niagara/Source/NiagaraEditorWidgets/Private/SNiagaraOverviewStackNode.h
@@ -25,7 +25,7 @@ public:
 	virtual void Tick(const FGeometry& AllottedGeometry, const double InCurrentTime, const float InDeltaTime) override;
 	virtual bool ShouldAllowCulling() const override { return false; }
 
-	virtual ~SNiagaraOverviewStackNode();
+	virtual ~SNiagaraOverviewStackNode() override;
 protected:
 	virtual TSharedRef<SWidget> CreateTitleWidget(TSharedPtr<SNodeTitle> NodeTitle) override;
 	virtual TSharedRef<SWidget> CreateTitleRightWidget() override;
@@ -67,15 +67,20 @@ private:
 
 	FReply OnCycleThroughIssues();
 	FReply OpenParentEmitter();
+	FText OpenParentEmitterTooltip() const;
 	EVisibility GetOpenParentEmitterVisibility() const;
+	EVisibility GetVersionSelectorVisibility() const;
+	FSlateColor GetVersionSelectorColor() const;
+	TSharedRef<SWidget> GetVersionSelectorDropdownMenu();
+	void SwitchToVersion(FNiagaraAssetVersion Version);
 
 	const FSlateBrush* GetSummaryViewButtonBrush() const;
 	FText GetSummaryViewCollapseTooltipText() const;
 	FReply ExpandSummaryViewClicked();
-private:
-	UNiagaraOverviewNode* OverviewStackNode;
-	UNiagaraStackViewModel* StackViewModel;
-	UNiagaraSystemSelectionViewModel* OverviewSelectionViewModel;
+
+	UNiagaraOverviewNode* OverviewStackNode = nullptr;
+	UNiagaraStackViewModel* StackViewModel = nullptr;
+	UNiagaraSystemSelectionViewModel* OverviewSelectionViewModel = nullptr;
 	TWeakObjectPtr<UNiagaraSystemScalabilityViewModel> ScalabilityViewModel;
 	TWeakPtr<FNiagaraEmitterHandleViewModel> EmitterHandleViewModelWeak;
 	/** The top content bar houses the isolate button and the thumbnails */
@@ -83,8 +88,8 @@ private:
 	TSharedPtr<SWidget> BottomSummaryExpander;
 
 	TArray<UNiagaraStackEntry*> PreviewStackEntries;
-	bool bIsHoveringThumbnail;
-	bool bTopContentBarRefreshPending;
-	int32 CurrentIssueIndex;
-	bool bScalabilityModeActive;
+	bool bIsHoveringThumbnail = false;
+	bool bTopContentBarRefreshPending = false;
+	int32 CurrentIssueIndex = 0;
+	bool bScalabilityModeActive = false;
 };
\ No newline at end of file
diff --git a/Engine/Plugins/FX/Niagara/Source/NiagaraEditorWidgets/Private/SNiagaraScalabilityPreviewSettings.cpp b/Engine/Plugins/FX/Niagara/Source/NiagaraEditorWidgets/Private/SNiagaraScalabilityPreviewSettings.cpp
--- a/Engine/Plugins/FX/Niagara/Source/NiagaraEditorWidgets/Private/SNiagaraScalabilityPreviewSettings.cpp
+++ b/Engine/Plugins/FX/Niagara/Source/NiagaraEditorWidgets/Private/SNiagaraScalabilityPreviewSettings.cpp
@@ -205,23 +205,19 @@ FSlateColor SNiagaraScalabilityPreviewSettings::GetQualityLevelButtonTextColor(i
 
 const FSlateBrush* SNiagaraScalabilityPreviewSettings::GetActivePreviewPlatformImage() const
 {
-	UDeviceProfile* DeviceProfile = ScalabilityViewModel->GetPreviewDeviceProfile().Get(nullptr);
-
-	FName PlatformName;
-	if(DeviceProfile)
+	if (UDeviceProfile* DeviceProfile = ScalabilityViewModel->GetPreviewDeviceProfile().Get(nullptr))
 	{
-		PlatformName = *DeviceProfile->DeviceType;
-	}
-	
-	if (const PlatformInfo::FTargetPlatformInfo* Info = PlatformInfo::FindPlatformInfo(PlatformName))
-	{
-		const FSlateBrush* DeviceProfileTypeIcon = FAppStyle::GetBrush(Info->GetIconStyleName(EPlatformIconSize::Normal));
-		if (DeviceProfileTypeIcon != FAppStyle::Get().GetDefaultBrush())
+		FName PlatformName = *DeviceProfile->DeviceType;
+		if (const PlatformInfo::FTargetPlatformInfo* Info = PlatformInfo::FindPlatformInfo(PlatformName))
 		{
-			return DeviceProfileTypeIcon;
+			const FSlateBrush* DeviceProfileTypeIcon = FAppStyle::GetBrush(Info->GetIconStyleName(EPlatformIconSize::Normal));
+			if (DeviceProfileTypeIcon != FAppStyle::Get().GetDefaultBrush())
+			{
+				return DeviceProfileTypeIcon;
+			}
 		}
 	}
-
+	
 	return nullptr;
 }
 
diff --git a/Engine/Plugins/FX/Niagara/Source/NiagaraEditorWidgets/Private/SNiagaraStack.cpp b/Engine/Plugins/FX/Niagara/Source/NiagaraEditorWidgets/Private/SNiagaraStack.cpp
--- a/Engine/Plugins/FX/Niagara/Source/NiagaraEditorWidgets/Private/SNiagaraStack.cpp
+++ b/Engine/Plugins/FX/Niagara/Source/NiagaraEditorWidgets/Private/SNiagaraStack.cpp
@@ -24,7 +24,6 @@
 #include "ViewModels/Stack/NiagaraStackModuleItem.h"
 #include "ViewModels/Stack/NiagaraStackSystemSettingsGroup.h"
 #include "ViewModels/Stack/NiagaraStackParameterStoreEntry.h"
-#include "Framework/Application/SlateApplication.h"
 #include "Framework/Commands/GenericCommands.h"
 #include "IContentBrowserSingleton.h"
 #include "ContentBrowserModule.h"
@@ -51,14 +50,10 @@
 #include "NiagaraEditorWidgetsUtilities.h"
 #include "DragAndDrop/DecoratedDragDropOp.h"
 #include "Stack/SNiagaraStackErrorItem.h"
-#include "NiagaraStackEditorData.h"
-#include "ScopedTransaction.h"
 #include "Widgets/Layout/SWrapBox.h"
 #include "ViewModels/Stack/NiagaraStackRoot.h"
 #include "NiagaraStackCommandContext.h"
 #include "NiagaraEditorUtilities.h"
-#include "Subsystems/AssetEditorSubsystem.h"
-#include "EditorFontGlyphs.h"
 #include "Framework/Commands/UICommandList.h"
 #include "ViewModels/NiagaraOverviewGraphViewModel.h"
 #include "Widgets/SNiagaraParameterName.h"
@@ -230,7 +225,7 @@ private:
 
 	EVisibility GetOpenSourceEmitterVisibility() const
 	{
-		return EmitterHandleViewModel->GetEmitterViewModel()->GetEmitter()->GetParent() != nullptr ? EVisibility::Visible : EVisibility::Collapsed;
+		return EmitterHandleViewModel->GetEmitterViewModel()->HasParentEmitter() ? EVisibility::Visible : EVisibility::Collapsed;
 	}
 
 	FReply OnIssueIconClicked() const
@@ -494,7 +489,7 @@ void SNiagaraStack::ShowEmitterInContentBrowser(TWeakPtr<FNiagaraEmitterHandleVi
 		TArray<FAssetData> Assets;
 		if (EmitterHandleViewModel->GetEmitterViewModel()->HasParentEmitter())
 		{
-			Assets.Add(FAssetData(EmitterHandleViewModel->GetEmitterViewModel()->GetParentEmitter()));
+			Assets.Add(FAssetData(EmitterHandleViewModel->GetEmitterViewModel()->GetParentEmitter().Emitter));
 			ContentBrowserModule.Get().SyncBrowserToAssets(Assets);
 		}
 	}
diff --git a/Engine/Plugins/FX/Niagara/Source/NiagaraEditorWidgets/Private/Stack/SNiagaraStackFunctionInputValue.cpp b/Engine/Plugins/FX/Niagara/Source/NiagaraEditorWidgets/Private/Stack/SNiagaraStackFunctionInputValue.cpp
--- a/Engine/Plugins/FX/Niagara/Source/NiagaraEditorWidgets/Private/Stack/SNiagaraStackFunctionInputValue.cpp
+++ b/Engine/Plugins/FX/Niagara/Source/NiagaraEditorWidgets/Private/Stack/SNiagaraStackFunctionInputValue.cpp
@@ -780,7 +780,7 @@ FReply SNiagaraStackFunctionInputValue::ResetButtonPressed() const
 
 EVisibility SNiagaraStackFunctionInputValue::GetResetToBaseButtonVisibility() const
 {
-	if (FunctionInput->HasBaseEmitter())
+	if (FunctionInput->HasBaseEmitter() && FunctionInput->GetEmitterViewModel().IsValid())
 	{
 		return FunctionInput->CanResetToBase() ? EVisibility::Visible : EVisibility::Hidden;
 	}
diff --git a/Engine/Plugins/FX/Niagara/Source/NiagaraEditorWidgets/Private/Stack/SNiagaraStackRowPerfWidget.cpp b/Engine/Plugins/FX/Niagara/Source/NiagaraEditorWidgets/Private/Stack/SNiagaraStackRowPerfWidget.cpp
--- a/Engine/Plugins/FX/Niagara/Source/NiagaraEditorWidgets/Private/Stack/SNiagaraStackRowPerfWidget.cpp
+++ b/Engine/Plugins/FX/Niagara/Source/NiagaraEditorWidgets/Private/Stack/SNiagaraStackRowPerfWidget.cpp
@@ -110,13 +110,13 @@ void SNiagaraStackRowPerfWidget::Tick(const FGeometry& AllottedGeometry, const d
 	{
 		return;
 	}
-	GroupOverallTime = CalculateGroupOverallTime(IsEmitterStack() ? GetEmitter()->GetUniqueEmitterName() : "Main");
+	GroupOverallTime = CalculateGroupOverallTime(IsEmitterStack() ? GetEmitter().Emitter->GetUniqueEmitterName() : "Main");
 	if (IsSystemStack())
 	{
 		EmitterTimeTotal = 0;
 		for (const FNiagaraEmitterHandle& Handle : StackEntry->GetSystemViewModel()->GetSystem().GetEmitterHandles())
 		{
-			EmitterTimeTotal += CalculateGroupOverallTime(Handle.GetInstance()->GetUniqueEmitterName());
+			EmitterTimeTotal += CalculateGroupOverallTime(Handle.GetInstance().Emitter->GetUniqueEmitterName());
 		}
 	}
 	if (IsGpuEmitter())
@@ -401,11 +401,11 @@ bool SNiagaraStackRowPerfWidget::IsEntrySelected() const
 	return StackEntry->GetSystemViewModel()->GetSelectionViewModel()->ContainsEntry(StackEntry.Get());
 }
 
-UNiagaraEmitter* SNiagaraStackRowPerfWidget::GetEmitter() const
+FVersionedNiagaraEmitter SNiagaraStackRowPerfWidget::GetEmitter() const
 {
 	if (!StackEntry->GetEmitterViewModel())
 	{
-		return nullptr;
+		return FVersionedNiagaraEmitter();
 	}
 	return StackEntry->GetEmitterViewModel().Get()->GetEmitter();
 }
@@ -440,17 +440,18 @@ ENiagaraStatDisplayMode SNiagaraStackRowPerfWidget::GetDisplayMode() const
 
 bool SNiagaraStackRowPerfWidget::IsInterpolatedSpawnEnabled() const
 {
-	UNiagaraEmitter* Emitter = GetEmitter();
-	if (!Emitter)
+	FVersionedNiagaraEmitterData* EmitterData = GetEmitter().GetEmitterData();
+	if (!EmitterData)
 	{
 		return false;
 	}
-	return Emitter->bInterpolatedSpawning;
+	return EmitterData->bInterpolatedSpawning;
 }
 
 bool SNiagaraStackRowPerfWidget::IsGpuEmitter() const
 {
-	return GetEmitter() && GetEmitter()->SimTarget == ENiagaraSimTarget::GPUComputeSim;
+	FVersionedNiagaraEmitterData* EmitterData = GetEmitter().GetEmitterData();
+	return EmitterData && EmitterData->SimTarget == ENiagaraSimTarget::GPUComputeSim;
 }
 
 float SNiagaraStackRowPerfWidget::CalculateGroupOverallTime(FString StatScopeName) const
@@ -462,17 +463,17 @@ float SNiagaraStackRowPerfWidget::CalculateGroupOverallTime(FString StatScopeNam
 	}
 	if (IsParticleStack())
 	{
-		UNiagaraEmitter* Emitter = GetEmitter();
-		if (!Emitter)
+		FVersionedNiagaraEmitterData* EmitterData = GetEmitter().GetEmitterData();
+		if (!EmitterData)
 		{
 			return 0;
 		}
 		TStatId StatId = FDynamicStats::CreateStatId<FStatGroup_STATGROUP_NiagaraDetailed>(StatScopeName);
 		if (IsGpuEmitter())
 		{
-			return Emitter->GetStatData().GetRuntimeStat(StatId.GetName(), ENiagaraScriptUsage::ParticleGPUComputeScript, GetEvaluationType()) / 1000.0f;
+			return EmitterData->GetStatData().GetRuntimeStat(StatId.GetName(), ENiagaraScriptUsage::ParticleGPUComputeScript, GetEvaluationType()) / 1000.0f;
 		}
-		return FPlatformTime::ToMilliseconds(Emitter->GetStatData().GetRuntimeStat(StatId.GetName(), GetUsage(), GetEvaluationType()));
+		return FPlatformTime::ToMilliseconds(EmitterData->GetStatData().GetRuntimeStat(StatId.GetName(), GetUsage(), GetEvaluationType()));
 	}
 	if (IsSystemStack() || IsEmitterStack())
 	{
@@ -502,8 +503,8 @@ float SNiagaraStackRowPerfWidget::CalculateStackEntryTime() const
 	UNiagaraNodeFunctionCall& FunctionNode = ModuleItem->GetModuleNode();
 	if (IsParticleStack())
 	{
-		UNiagaraEmitter* Emitter = GetEmitter();
-		if (!Emitter)
+		FVersionedNiagaraEmitterData* EmitterData = GetEmitter().GetEmitterData();
+		if (!EmitterData)
 		{
 			return 0;
 		}
@@ -511,11 +512,11 @@ float SNiagaraStackRowPerfWidget::CalculateStackEntryTime() const
 		FString StatScopeName = FunctionNode.GetFunctionName() + "_Emitter";
 		TStatId StatId = FDynamicStats::CreateStatId<FStatGroup_STATGROUP_NiagaraDetailed>(StatScopeName);
 		FString StatName = StatId.GetName().ToString();
-		return FPlatformTime::ToMilliseconds(Emitter->GetStatData().GetRuntimeStat(StatId.GetName(), GetUsage(), GetEvaluationType()));
+		return FPlatformTime::ToMilliseconds(EmitterData->GetStatData().GetRuntimeStat(StatId.GetName(), GetUsage(), GetEvaluationType()));
 	}
 	if (IsEmitterStack())
 	{
-		UNiagaraEmitter* Emitter = GetEmitter();
+		UNiagaraEmitter* Emitter = GetEmitter().Emitter;
 		if (!Emitter)
 		{
 			return 0;
diff --git a/Engine/Plugins/FX/Niagara/Source/NiagaraEditorWidgets/Private/Stack/SNiagaraStackRowPerfWidget.h b/Engine/Plugins/FX/Niagara/Source/NiagaraEditorWidgets/Private/Stack/SNiagaraStackRowPerfWidget.h
--- a/Engine/Plugins/FX/Niagara/Source/NiagaraEditorWidgets/Private/Stack/SNiagaraStackRowPerfWidget.h
+++ b/Engine/Plugins/FX/Niagara/Source/NiagaraEditorWidgets/Private/Stack/SNiagaraStackRowPerfWidget.h
@@ -38,7 +38,7 @@ private:
 	bool IsGroupHeaderEntry() const;
 	bool IsModuleEntry() const;
 	bool IsEntrySelected() const;
-	UNiagaraEmitter* GetEmitter() const;
+	FVersionedNiagaraEmitter GetEmitter() const;
 	ENiagaraScriptUsage GetUsage() const;
 	ENiagaraStatEvaluationType GetEvaluationType() const;
 	ENiagaraStatDisplayMode GetDisplayMode() const;
diff --git a/Engine/Plugins/FX/Niagara/Source/NiagaraEditorWidgets/Private/Stack/SNiagaraStackTableRow.cpp b/Engine/Plugins/FX/Niagara/Source/NiagaraEditorWidgets/Private/Stack/SNiagaraStackTableRow.cpp
--- a/Engine/Plugins/FX/Niagara/Source/NiagaraEditorWidgets/Private/Stack/SNiagaraStackTableRow.cpp
+++ b/Engine/Plugins/FX/Niagara/Source/NiagaraEditorWidgets/Private/Stack/SNiagaraStackTableRow.cpp
@@ -522,11 +522,10 @@ void SNiagaraStackTableRow::ToggleShowInSummaryView()
 		{
 			// TODO: Move this parent handling to the UNiagaraStackFunctionInput and merge manager.
 			bool bHasParentSummaryData = false;;
-			UNiagaraEmitter* ParentEmitter = FunctionInput->GetEmitterViewModel()->GetEmitter()->GetParent();
-			if (ParentEmitter != nullptr)
+			FVersionedNiagaraEmitter ParentEmitter = FunctionInput->GetEmitterViewModel()->GetEmitter().GetEmitterData()->GetParent();
+			if (ParentEmitter.GetEmitterData())
 			{
-				UNiagaraEmitterEditorData* ParentEmitterEditorData = Cast<UNiagaraEmitterEditorData>(ParentEmitter->GetEditorData());
-				if (ParentEmitterEditorData != nullptr)
+				if (UNiagaraEmitterEditorData* ParentEmitterEditorData = Cast<UNiagaraEmitterEditorData>(ParentEmitter.GetEmitterData()->GetEditorData()))
 				{
 					bHasParentSummaryData = ParentEmitterEditorData->GetSummaryViewMetaData(Key.GetValue()).IsSet();
 				}
diff --git a/Engine/Plugins/Runtime/GeometryCache/Source/GeometryCache/Classes/NiagaraGeometryCacheRendererProperties.h b/Engine/Plugins/Runtime/GeometryCache/Source/GeometryCache/Classes/NiagaraGeometryCacheRendererProperties.h
--- a/Engine/Plugins/Runtime/GeometryCache/Source/GeometryCache/Classes/NiagaraGeometryCacheRendererProperties.h
+++ b/Engine/Plugins/Runtime/GeometryCache/Source/GeometryCache/Classes/NiagaraGeometryCacheRendererProperties.h
@@ -58,7 +58,7 @@ public:
 	virtual const FSlateBrush* GetStackIcon() const override;
 	virtual void GetRendererTooltipWidgets(const FNiagaraEmitterInstance* InEmitter, TArray<TSharedPtr<SWidget>>& OutWidgets, TSharedPtr<FAssetThumbnailPool> InThumbnailPool) const override;
 	virtual const TArray<FNiagaraVariable>& GetOptionalAttributes() override;
-	virtual void GetRendererFeedback(const UNiagaraEmitter* InEmitter, TArray<FText>& OutErrors, TArray<FText>& OutWarnings, TArray<FText>& OutInfo) const override;
+	virtual void GetRendererFeedback(const FVersionedNiagaraEmitter& InEmitter, TArray<FText>& OutErrors, TArray<FText>& OutWarnings, TArray<FText>& OutInfo) const override;
 #endif // WITH_EDITORONLY_DATA
 
 	/** Reference to the geometry cache assets to use. If ArrayIndexBinding is not set, a random element is used for each particle. */
diff --git a/Engine/Plugins/Runtime/GeometryCache/Source/GeometryCache/Private/NiagaraGeometryCacheRendererProperties.cpp b/Engine/Plugins/Runtime/GeometryCache/Source/GeometryCache/Private/NiagaraGeometryCacheRendererProperties.cpp
--- a/Engine/Plugins/Runtime/GeometryCache/Source/GeometryCache/Private/NiagaraGeometryCacheRendererProperties.cpp
+++ b/Engine/Plugins/Runtime/GeometryCache/Source/GeometryCache/Private/NiagaraGeometryCacheRendererProperties.cpp
@@ -333,7 +333,7 @@ const TArray<FNiagaraVariable>& UNiagaraGeometryCacheRendererProperties::GetOpti
 	return Attrs;
 }
 
-void UNiagaraGeometryCacheRendererProperties::GetRendererFeedback(const UNiagaraEmitter* InEmitter,	TArray<FText>& OutErrors, TArray<FText>& OutWarnings, TArray<FText>& OutInfo) const
+void UNiagaraGeometryCacheRendererProperties::GetRendererFeedback(const FVersionedNiagaraEmitter& InEmitter, TArray<FText>& OutErrors, TArray<FText>& OutWarnings, TArray<FText>& OutInfo) const
 {
 	OutInfo.Add(FText::FromString(TEXT("The geometry cache renderer is still an experimental feature.")));
 }
diff --git a/Engine/Plugins/Runtime/GeometryCache/Source/GeometryCache/Private/NiagaraRendererGeometryCache.cpp b/Engine/Plugins/Runtime/GeometryCache/Source/GeometryCache/Private/NiagaraRendererGeometryCache.cpp
--- a/Engine/Plugins/Runtime/GeometryCache/Source/GeometryCache/Private/NiagaraRendererGeometryCache.cpp
+++ b/Engine/Plugins/Runtime/GeometryCache/Source/GeometryCache/Private/NiagaraRendererGeometryCache.cpp
@@ -233,7 +233,7 @@ void FNiagaraRendererGeometryCache::PostSystemTick_GameThread(const UNiagaraRend
 		if (!GeometryComponent || GeometryComponent->HasAnyFlags(RF_BeginDestroyed | RF_FinishDestroyed))
 		{
 			int32 CacheIndex = ArrayIndexAccessor.GetSafe(ParticleIndex, INDEX_NONE);
-			if (CacheIndex == INDEX_NONE && Emitter->GetCachedEmitter()->bDeterminism)
+			if (CacheIndex == INDEX_NONE && Emitter->GetCachedEmitterData()->bDeterminism)
 			{
 				int32 Seed = Properties->bAssignComponentsOnParticleID ? ParticleID : ParticleIndex;
 				FRandomStream RandomStream = FRandomStream(Seed * 907633515U); // multiply the seed, otherwise we get very poor randomness for small seeds
@@ -277,7 +277,7 @@ void FNiagaraRendererGeometryCache::PostSystemTick_GameThread(const UNiagaraRend
 			GeometryComponent->SetGeometryCache(Cache);
 			GeometryComponent->SetManualTick(true); // we want to tick the component with the delta time of the niagara sim
 
-			if (Emitter->GetCachedEmitter()->bLocalSpace)
+			if (Emitter->GetCachedEmitterData()->bLocalSpace)
 			{
 				GeometryComponent->SetAbsolute(false, false, false);
 			}
@@ -312,7 +312,7 @@ void FNiagaraRendererGeometryCache::PostSystemTick_GameThread(const UNiagaraRend
 		const auto RotationAccessor = FNiagaraDataSetAccessor<FVector3f>::CreateReader(Data, Properties->RotationBinding.GetDataSetBindableVariable().GetName());
 		const auto ScaleAccessor = FNiagaraDataSetAccessor<FVector3f>::CreateReader(Data, Properties->ScaleBinding.GetDataSetBindableVariable().GetName());
 		const auto ElapsedTimeAccessor = FNiagaraDataSetAccessor<float>::CreateReader(Data, Properties->ElapsedTimeBinding.GetDataSetBindableVariable().GetName());
-		FNiagaraLWCConverter LwcConverter = SystemInstance->GetLWCConverter(Emitter->GetCachedEmitter()->bLocalSpace);
+		FNiagaraLWCConverter LwcConverter = SystemInstance->GetLWCConverter(Emitter->GetCachedEmitterData()->bLocalSpace);
 
 		FVector Position = LwcConverter.ConvertSimulationPositionToWorld(PositionAccessor.GetSafe(ParticleIndex, FNiagaraPosition(ForceInit)));
 		FVector Scale = (FVector)ScaleAccessor.GetSafe(ParticleIndex, FVector3f::OneVector);
