Added versioning for Niagara emitters

diff --git a/Engine/Plugins/FX/Niagara/Source/Niagara/Private/NiagaraEmitter.cpp b/Engine/Plugins/FX/Niagara/Source/Niagara/Private/NiagaraEmitter.cpp
--- a/Engine/Plugins/FX/Niagara/Source/Niagara/Private/NiagaraEmitter.cpp
+++ b/Engine/Plugins/FX/Niagara/Source/Niagara/Private/NiagaraEmitter.cpp
@@ -146,6 +146,7 @@ bool FNiagaraEmitterScriptProperties::DataSetAccessSynchronized() const
 
 UNiagaraEmitter::UNiagaraEmitter(const FObjectInitializer& Initializer)
 : Super(Initializer)
+#if WITH_EDITORONLY_DATA
 , PreAllocationCount_DEPRECATED(0)
 , FixedBounds_DEPRECATED(FBox(FVector(-100), FVector(100)))
 , MinDetailLevel_DEPRECATED(0)
@@ -159,6 +160,7 @@ UNiagaraEmitter::UNiagaraEmitter(const FObjectInitializer& Initializer)
 , MaxDeltaTimePerTick_DEPRECATED(0.125)
 , bLimitDeltaTime_DEPRECATED(true)
 , MaxGPUParticlesSpawnPerFrame_DEPRECATED(0)
+#endif
 {
 #if WITH_EDITORONLY_DATA
 	ThumbnailImageOutOfDate = true;
diff --git a/Engine/Plugins/Runtime/GeometryCache/Source/GeometryCache/Private/NiagaraRendererGeometryCache.cpp b/Engine/Plugins/Runtime/GeometryCache/Source/GeometryCache/Private/NiagaraRendererGeometryCache.cpp
--- a/Engine/Plugins/Runtime/GeometryCache/Source/GeometryCache/Private/NiagaraRendererGeometryCache.cpp
+++ b/Engine/Plugins/Runtime/GeometryCache/Source/GeometryCache/Private/NiagaraRendererGeometryCache.cpp
@@ -7,6 +7,7 @@
 #include "NiagaraDataSet.h"
 #include "NiagaraGeometryCacheRendererProperties.h"
 #include "NiagaraSystemInstance.h"
+#include "NiagaraEmitter.h"
 #include "Async/Async.h"
 
 static float GNiagaraGeometryComponentRenderPoolInactiveTimeLimit = 5;
