Added missing copyright header

diff --git a/Engine/Plugins/FX/Niagara/Source/Niagara/Private/NiagaraVersionedObject.cpp b/Engine/Plugins/FX/Niagara/Source/Niagara/Private/NiagaraVersionedObject.cpp
--- a/Engine/Plugins/FX/Niagara/Source/Niagara/Private/NiagaraVersionedObject.cpp
+++ b/Engine/Plugins/FX/Niagara/Source/Niagara/Private/NiagaraVersionedObject.cpp
@@ -1,4 +1,6 @@
-﻿#include "NiagaraVersionedObject.h"
+﻿// Copyright Epic Games, Inc. All Rights Reserved.
+
+#include "NiagaraVersionedObject.h"
 
 #include "NiagaraTypes.h"
 
