Fixed non-editor build error

diff --git a/Engine/Plugins/FX/Niagara/Source/Niagara/Private/NiagaraEmitter.cpp b/Engine/Plugins/FX/Niagara/Source/Niagara/Private/NiagaraEmitter.cpp
--- a/Engine/Plugins/FX/Niagara/Source/Niagara/Private/NiagaraEmitter.cpp
+++ b/Engine/Plugins/FX/Niagara/Source/Niagara/Private/NiagaraEmitter.cpp
@@ -551,6 +551,8 @@ void UNiagaraEmitter::PostLoad()
 	}
 
 	const int32 NiagaraVer = GetLinkerCustomVersion(FNiagaraCustomVersion::GUID);
+
+#if WITH_EDITORONLY_DATA
 	if (NiagaraVer < FNiagaraCustomVersion::PlatformScalingRefactor)
 	{
 		int32 MinDetailLevel = bUseMaxDetailLevel_DEPRECATED ? MinDetailLevel_DEPRECATED : 0;
@@ -600,7 +602,6 @@ void UNiagaraEmitter::PostLoad()
 		}
 	}
 
-#if WITH_EDITORONLY_DATA
 	// this can only ever be true for old assets that haven't been loaded yet, so this won't overwrite subsequent changes to the template specification
 	if(bIsTemplateAsset_DEPRECATED)
 	{
@@ -1460,7 +1461,7 @@ void FVersionedNiagaraEmitterData::CacheFromCompiledData(const FNiagaraDataSetCo
 	// Initialize bounds calculators - skip creating if we won't ever use it.  We leave the GPU sims in there with the editor so that we can
 	// generate the bounds from the readback in the tool.
 #if !WITH_EDITOR
-	bool bUseDynamicBounds = !bFixedBoundsEnabled && SimTarget == ENiagaraSimTarget::CPUSim;
+	bool bUseDynamicBounds = CalculateBoundsMode == ENiagaraEmitterCalculateBoundMode::Dynamic && SimTarget == ENiagaraSimTarget::CPUSim;
 	if (bUseDynamicBounds)
 #endif
 	{
