Fixed static analysis warnings

diff --git a/Engine/Plugins/FX/Niagara/Source/NiagaraEditor/Public/UpgradeNiagaraScriptResults.h b/Engine/Plugins/FX/Niagara/Source/NiagaraEditor/Public/UpgradeNiagaraScriptResults.h
--- a/Engine/Plugins/FX/Niagara/Source/NiagaraEditor/Public/UpgradeNiagaraScriptResults.h
+++ b/Engine/Plugins/FX/Niagara/Source/NiagaraEditor/Public/UpgradeNiagaraScriptResults.h
@@ -172,7 +172,7 @@ public:
 
 private:
 	UPROPERTY()
-	UNiagaraStackModuleItem* ModuleItem = nullptr;
+	TObjectPtr<UNiagaraStackModuleItem> ModuleItem = nullptr;
 };
 
 /** Wrapper for an emitter stack. */
@@ -238,10 +238,10 @@ public:
 	bool bCancelledByPythonError = false;
 
 	UPROPERTY(BlueprintReadWrite, Category = "Scripting")
-	UNiagaraPythonEmitter* OldEmitter = nullptr;
+	TObjectPtr<UNiagaraPythonEmitter> OldEmitter = nullptr;
 
 	UPROPERTY(BlueprintReadWrite, Category = "Scripting")
-	UNiagaraPythonEmitter* NewEmitter = nullptr;
+	TObjectPtr<UNiagaraPythonEmitter> NewEmitter = nullptr;
 
 private:
 	TArray<FVersionedNiagaraEmitterData*> UpgradeVersionData;
