- Convert Spline DI to Shader Parameters
- Object Reader conversion to use shader structures
- Neighbor Grid 3D conversion to use shader structures

diff --git a/Engine/Plugins/FX/Niagara/Shaders/Private/NiagaraDataInterfaceSplineTemplate.ush b/Engine/Plugins/FX/Niagara/Shaders/Private/NiagaraDataInterfaceSplineTemplate.ush
--- a/Engine/Plugins/FX/Niagara/Shaders/Private/NiagaraDataInterfaceSplineTemplate.ush
+++ b/Engine/Plugins/FX/Niagara/Shaders/Private/NiagaraDataInterfaceSplineTemplate.ush
@@ -1,102 +1,101 @@
 ﻿// Copyright Epic Games, Inc. All Rights Reserved.
 	
-
-Buffer<float4> SplinePositionsLUT_{ParameterHLSLSymbol};
-Buffer<float4> SplineScalesLUT_{ParameterHLSLSymbol};
-Buffer<float4> SplineRotationsLUT_{ParameterHLSLSymbol};
-
-float4x4 SplineTransform_{ParameterHLSLSymbol};
-float4x4 SplineTransformRotationMat_{ParameterHLSLSymbol};
-float4x4 SplineTransformInverseTranspose_{ParameterHLSLSymbol};
-float4 SplineTransformRotation_{ParameterHLSLSymbol};
-float3 DefaultUpVector_{ParameterHLSLSymbol};
-float SplineLength_{ParameterHLSLSymbol};
-float SplineDistanceStep_{ParameterHLSLSymbol};
-float InvSplineDistanceStep_{ParameterHLSLSymbol};
-int MaxIndex_{ParameterHLSLSymbol};
-
-void FindNeighborKeys_{ParameterHLSLSymbol}(float InDistance, out int PrevKey, out int NextKey, out float Alpha)
-{
-	const float Key = InDistance * {InvSplineDistanceStep};
-	PrevKey = clamp(int(floor(Key)), 0, {MaxIndex});
-	NextKey = clamp(int(ceil(Key)), 0, {MaxIndex});
+float4x4 {ParameterName}_SplineTransform;
+float4x4 {ParameterName}_SplineTransformRotationMat;
+float4x4 {ParameterName}_SplineTransformInverseTranspose;
+float4	{ParameterName}_SplineTransformRotation;
+float3	{ParameterName}_DefaultUpVector;
+float	{ParameterName}_SplineLength;
+float	{ParameterName}_SplineDistanceStep;
+float	{ParameterName}_InvSplineDistanceStep;
+int		{ParameterName}_MaxIndex;
+
+Buffer<float4> {ParameterName}_SplinePositionsLUT;
+Buffer<float4> {ParameterName}_SplineScalesLUT;
+Buffer<float4> {ParameterName}_SplineRotationsLUT;
+
+void FindNeighborKeys_{ParameterName}(float InDistance, out int PrevKey, out int NextKey, out float Alpha)
+{
+	const float Key = InDistance * {ParameterName}_InvSplineDistanceStep;
+	PrevKey = clamp(int(floor(Key)), 0, {ParameterName}_MaxIndex);
+	NextKey = clamp(int(ceil(Key)), 0, {ParameterName}_MaxIndex);
 	Alpha = frac(Key);
 }
 
-float3 EvaluatePosition_{ParameterHLSLSymbol}(float InDistance)
+float3 EvaluatePosition_{ParameterName}(float InDistance)
 {
-	if ({MaxIndex} < 0)
+	if ({ParameterName}_MaxIndex < 0)
 	{
 		return float3(0,0,0);
 	}
 	
 	int PrevKey, NextKey;
 	float Alpha;
-	FindNeighborKeys_{ParameterHLSLSymbol}(InDistance, PrevKey, NextKey, Alpha);
+	FindNeighborKeys_{ParameterName}(InDistance, PrevKey, NextKey, Alpha);
 	
 	if (NextKey == PrevKey)
 	{
-		return {SplinePositionsLUT}[PrevKey].xyz;
+		return {ParameterName}_SplinePositionsLUT[PrevKey].xyz;
 	}
-	return lerp({SplinePositionsLUT}[PrevKey].xyz, {SplinePositionsLUT}[NextKey].xyz, Alpha).xyz;
+	return lerp({ParameterName}_SplinePositionsLUT[PrevKey].xyz, {ParameterName}_SplinePositionsLUT[NextKey].xyz, Alpha).xyz;
 }
 
-float3 EvaluateScale_{ParameterHLSLSymbol}(float InDistance)
+float3 EvaluateScale_{ParameterName}(float InDistance)
 {
-	if ({MaxIndex} < 0)
+	if ({ParameterName}_MaxIndex < 0)
 	{
 		return float3(1,1,1);
 	}
 	
 	int PrevKey, NextKey;
 	float Alpha;
-	FindNeighborKeys_{ParameterHLSLSymbol}(InDistance, PrevKey, NextKey, Alpha);
+	FindNeighborKeys_{ParameterName}(InDistance, PrevKey, NextKey, Alpha);
 	
 	if (NextKey == PrevKey)
 	{
-		return {SplineScalesLUT}[PrevKey].xyz;
+		return {ParameterName}_SplineScalesLUT[PrevKey].xyz;
 	}
-	return lerp({SplineScalesLUT}[PrevKey].xyz, {SplineScalesLUT}[NextKey].xyz, Alpha).xyz;
+	return lerp({ParameterName}_SplineScalesLUT[PrevKey].xyz, {ParameterName}_SplineScalesLUT[NextKey].xyz, Alpha).xyz;
 }
 	
-float4 EvaluateRotation_{ParameterHLSLSymbol}(float InDistance)
+float4 EvaluateRotation_{ParameterName}(float InDistance)
 {
-	if ({MaxIndex} < 0)
+	if ({ParameterName}_MaxIndex < 0)
 	{
 		return float4(0,0,0,1);
 	}
 	
 	int PrevKey, NextKey;
 	float Alpha;
-	FindNeighborKeys_{ParameterHLSLSymbol}(InDistance, PrevKey, NextKey, Alpha);
+	FindNeighborKeys_{ParameterName}(InDistance, PrevKey, NextKey, Alpha);
 
 	if (NextKey == PrevKey)
 	{
-		return {SplineRotationsLUT}[PrevKey];
+		return {ParameterName}_SplineRotationsLUT[PrevKey];
 	}
-	return NiagaraQuatSLerp({SplineRotationsLUT}[PrevKey], {SplineRotationsLUT}[NextKey], Alpha);
+	return NiagaraQuatSLerp({ParameterName}_SplineRotationsLUT[PrevKey], {ParameterName}_SplineRotationsLUT[NextKey], Alpha);
 }
 	
-float4 GetRotationAtDistanceAlongSpline_{ParameterHLSLSymbol}(float InDistance)
+float4 GetRotationAtDistanceAlongSpline_{ParameterName}(float InDistance)
 {
-	float4 Quat = EvaluateRotation_{ParameterHLSLSymbol}(InDistance);
+	float4 Quat = EvaluateRotation_{ParameterName}(InDistance);
 	return Quat;
 }
 	
-float3 EvaluateDerivativePosition_{ParameterHLSLSymbol}(float InDistance)
+float3 EvaluateDerivativePosition_{ParameterName}(float InDistance)
 {
-	if ({MaxIndex} < 0)
+	if ({ParameterName}_MaxIndex < 0)
 	{
 		return float3(0,0,0);
 	}
 	
 	int PrevKey, NextKey;
 	float Alpha;
-	FindNeighborKeys_{ParameterHLSLSymbol}(InDistance, PrevKey, NextKey, Alpha);
+	FindNeighborKeys_{ParameterName}(InDistance, PrevKey, NextKey, Alpha);
 
 	if (NextKey == PrevKey)
 	{
-		if (NextKey < {MaxIndex})
+		if (NextKey < {ParameterName}_MaxIndex)
 		{
 			NextKey++;
 		}
@@ -109,115 +108,111 @@ float3 EvaluateDerivativePosition_{ParameterHLSLSymbol}(float InDistance)
 			return float3(0,0,0);
 		}
 	}
-	return {SplinePositionsLUT}[NextKey] - {SplinePositionsLUT}[PrevKey];
+	return {ParameterName}_SplinePositionsLUT[NextKey] - {ParameterName}_SplinePositionsLUT[PrevKey];
 }
 
-float3 EvaluateFindNearestPosition_{ParameterHLSLSymbol}(float3 InPosition)
+float3 EvaluateFindNearestPosition_{ParameterName}(float3 InPosition)
 {
-	if ({MaxIndex} < 0)
+	if ({ParameterName}_MaxIndex < 0)
 	{
 		return 0.0f;
 	}
 	
-	float MinDistance = dot(InPosition, {SplinePositionsLUT}[0]);
+	float MinDistance = dot(InPosition, {ParameterName}_SplinePositionsLUT[0]);
 	float KeyToNearest = 0.0f;
-	for (int Index = 1; Index <= {MaxIndex}; Index++)
+	for (int Index = 1; Index <= {ParameterName}_MaxIndex; Index++)
 	{
-		const float Distance = dot(InPosition, {SplinePositionsLUT}[Index]);
+		const float Distance = dot(InPosition, {ParameterName}_SplinePositionsLUT[Index]);
 		if (Distance < MinDistance)
 		{
 			MinDistance = Distance;
-			KeyToNearest = Index * {SplineDistanceStep};
+			KeyToNearest = Index * {ParameterName}_SplineDistanceStep;
 		}
 	}
 	return KeyToNearest;
 }
 
-void SampleSplinePositionByUnitDistance_{ParameterHLSLSymbol}(float U, out float3 Position)
+void SampleSplinePositionByUnitDistance_{ParameterName}(float U, out float3 Position)
 {
-	Position = EvaluatePosition_{ParameterHLSLSymbol}(U * {SplineLength});
+	Position = EvaluatePosition_{ParameterName}(U * {ParameterName}_SplineLength);
 }
 
-void SampleSplinePositionByUnitDistanceWS_{ParameterHLSLSymbol}(float U, out float3 Position)
+void SampleSplinePositionByUnitDistanceWS_{ParameterName}(float U, out float3 Position)
 {
-	Position = EvaluatePosition_{ParameterHLSLSymbol}(U * {SplineLength});
-	Position = mul(float4(Position, 1.0f), {SplineTransform}).xyz;
+	Position = EvaluatePosition_{ParameterName}(U * {ParameterName}_SplineLength);
+	Position = mul(float4(Position, 1.0f), {ParameterName}_SplineTransform).xyz;
 }
 
-void SampleSplineRotationByUnitDistance_{ParameterHLSLSymbol}(float U, out float4 Rotation)
+void SampleSplineRotationByUnitDistance_{ParameterName}(float U, out float4 Rotation)
 {
-	Rotation = GetRotationAtDistanceAlongSpline_{ParameterHLSLSymbol}(U * {SplineLength});
+	Rotation = GetRotationAtDistanceAlongSpline_{ParameterName}(U * {ParameterName}_SplineLength);
 }		
 
-void SampleSplineRotationByUnitDistanceWS_{ParameterHLSLSymbol}(float U, out float4 Rotation)
+void SampleSplineRotationByUnitDistanceWS_{ParameterName}(float U, out float4 Rotation)
 {
-	Rotation = GetRotationAtDistanceAlongSpline_{ParameterHLSLSymbol}(U * {SplineLength});
-	Rotation = NiagaraQuatMul({SplineTransformRotation}, Rotation);
+	Rotation = GetRotationAtDistanceAlongSpline_{ParameterName}(U * {ParameterName}_SplineLength);
+	Rotation = NiagaraQuatMul({ParameterName}_SplineTransformRotation, Rotation);
 }
 
-void SampleSplineDirectionByUnitDistance_{ParameterHLSLSymbol}(float U, out float3 Direction)
+void SampleSplineDirectionByUnitDistance_{ParameterName}(float U, out float3 Direction)
 {
-	Direction = normalize(EvaluateDerivativePosition_{ParameterHLSLSymbol}(U * {SplineLength}));
+	Direction = normalize(EvaluateDerivativePosition_{ParameterName}(U * {ParameterName}_SplineLength));
 }
 
-void SampleSplineDirectionByUnitDistanceWS_{ParameterHLSLSymbol}(float U, out float3 Direction)
+void SampleSplineDirectionByUnitDistanceWS_{ParameterName}(float U, out float3 Direction)
 {
-	Direction = normalize(EvaluateDerivativePosition_{ParameterHLSLSymbol}(U * {SplineLength}));
-	Direction = mul(float4(Direction, 1.0f), {SplineTransformRotationMat}).xyz;
+	Direction = normalize(EvaluateDerivativePosition_{ParameterName}(U * {ParameterName}_SplineLength));
+	Direction = mul(float4(Direction, 1.0f), {ParameterName}_SplineTransformRotationMat).xyz;
 }
 
-void SampleSplineUpVectorByUnitDistance_{ParameterHLSLSymbol}(float U, out float3 UpVector)
+void SampleSplineUpVectorByUnitDistance_{ParameterName}(float U, out float3 UpVector)
 {
-	float4 Rotation = GetRotationAtDistanceAlongSpline_{ParameterHLSLSymbol}(U * {SplineLength});
+	float4 Rotation = GetRotationAtDistanceAlongSpline_{ParameterName}(U * {ParameterName}_SplineLength);
 	UpVector = NiagaraQuatRotateVector(Rotation, float3(0,0,1));
 }
 
-void SampleSplineUpVectorByUnitDistanceWS_{ParameterHLSLSymbol}(float U, out float3 UpVector)
+void SampleSplineUpVectorByUnitDistanceWS_{ParameterName}(float U, out float3 UpVector)
 {
-	float4 Rotation = GetRotationAtDistanceAlongSpline_{ParameterHLSLSymbol}(U * {SplineLength});
+	float4 Rotation = GetRotationAtDistanceAlongSpline_{ParameterName}(U * {ParameterName}_SplineLength);
 	UpVector = NiagaraQuatRotateVector(Rotation, float3(0,0,1));
-	UpVector = mul(float4(UpVector, 1.0f), {SplineTransformRotationMat}).xyz;
+	UpVector = mul(float4(UpVector, 1.0f), {ParameterName}_SplineTransformRotationMat).xyz;
 }
 
-void SampleSplineRightVectorByUnitDistance_{ParameterHLSLSymbol}(float U, out float3 RightVector)
+void SampleSplineRightVectorByUnitDistance_{ParameterName}(float U, out float3 RightVector)
 {
-	float4 Rotation = GetRotationAtDistanceAlongSpline_{ParameterHLSLSymbol}(U * {SplineLength});
+	float4 Rotation = GetRotationAtDistanceAlongSpline_{ParameterName}(U * {ParameterName}_SplineLength);
 	RightVector = NiagaraQuatRotateVector(Rotation, float3(0,1,0));
 }
 
-void SampleSplineRightVectorByUnitDistanceWS_{ParameterHLSLSymbol}(float U, out float3 RightVector)
+void SampleSplineRightVectorByUnitDistanceWS_{ParameterName}(float U, out float3 RightVector)
 {
-	float4 Rotation = GetRotationAtDistanceAlongSpline_{ParameterHLSLSymbol}(U * {SplineLength});
+	float4 Rotation = GetRotationAtDistanceAlongSpline_{ParameterName}(U * {ParameterName}_SplineLength);
 	RightVector = NiagaraQuatRotateVector(Rotation, float3(0,1,0));
-	RightVector = mul(float4(RightVector, 1.0f), {SplineTransformRotationMat}).xyz;
+	RightVector = mul(float4(RightVector, 1.0f), {ParameterName}_SplineTransformRotationMat).xyz;
 }
 
-void SampleSplineTangentByUnitDistance_{ParameterHLSLSymbol}(float U, out float3 Tangent)
+void SampleSplineTangentByUnitDistance_{ParameterName}(float U, out float3 Tangent)
 {
-	Tangent = EvaluateDerivativePosition_{ParameterHLSLSymbol}(U * {SplineLength});
+	Tangent = EvaluateDerivativePosition_{ParameterName}(U * {ParameterName}_SplineLength);
 }
 
-void SampleSplineTangentByUnitDistanceWS_{ParameterHLSLSymbol}(float U, out float3 Tangent)
+void SampleSplineTangentByUnitDistanceWS_{ParameterName}(float U, out float3 Tangent)
 {
-	Tangent = EvaluateDerivativePosition_{ParameterHLSLSymbol}(U * {SplineLength});
-	Tangent = mul(Tangent, {SplineTransformRotation});
+	Tangent = EvaluateDerivativePosition_{ParameterName}(U * {ParameterName}_SplineLength);
+	Tangent = mul(Tangent, {ParameterName}_SplineTransformRotation);
 }
 
-void GetSplineLocalToWorld_{ParameterHLSLSymbol}(out float4x4 Transform)
+void GetSplineLocalToWorld_{ParameterName}(out float4x4 Transform)
 {
-	Transform = {SplineTransform};
+	Transform = {ParameterName}_SplineTransform;
 }
 
-void GetSplineLocalToWorldInverseTransposed_{ParameterHLSLSymbol}(out float4x4 Transform)
+void GetSplineLocalToWorldInverseTransposed_{ParameterName}(out float4x4 Transform)
 {
-	Transform = {SplineTransformInverseTranspose};
+	Transform = {ParameterName}_SplineTransformInverseTranspose;
 }
 
-void FindClosestUnitDistanceFromPositionWS_{ParameterHLSLSymbol}(float3 PositionWS, out float U)
+void FindClosestUnitDistanceFromPositionWS_{ParameterName}(float3 PositionWS, out float U)
 {
-	U = EvaluateFindNearestPosition_{ParameterHLSLSymbol}(U * {SplineLength});
+	U = EvaluateFindNearestPosition_{ParameterName}(U * {ParameterName}_SplineLength);
 }
-
-
-
-
diff --git a/Engine/Plugins/FX/Niagara/Source/Niagara/Classes/NiagaraDataInterfaceNeighborGrid3D.h b/Engine/Plugins/FX/Niagara/Source/Niagara/Classes/NiagaraDataInterfaceNeighborGrid3D.h
--- a/Engine/Plugins/FX/Niagara/Source/Niagara/Classes/NiagaraDataInterfaceNeighborGrid3D.h
+++ b/Engine/Plugins/FX/Niagara/Source/Niagara/Classes/NiagaraDataInterfaceNeighborGrid3D.h
@@ -62,14 +62,22 @@ class NIAGARA_API UNiagaraDataInterfaceNeighborGrid3D : public UNiagaraDataInter
 {
 	GENERATED_UCLASS_BODY()
 
-public:	
+	BEGIN_SHADER_PARAMETER_STRUCT(FShaderParameters, )
+		SHADER_PARAMETER(FIntVector,		NumCells)
+		SHADER_PARAMETER(FVector3f,			UnitToUV)
+		SHADER_PARAMETER(FVector3f,			CellSize)
+		SHADER_PARAMETER(FVector3f,			WorldBBoxSize)
+
+		SHADER_PARAMETER(int32,				MaxNeighborsPerCellValue)
+		SHADER_PARAMETER_SRV(Buffer<int>,	ParticleNeighbors)
+		SHADER_PARAMETER_SRV(Buffer<int>,	ParticleNeighborCount)
+		SHADER_PARAMETER_UAV(RWBuffer<int>, OutputParticleNeighbors)
+		SHADER_PARAMETER_UAV(RWBuffer<int>, OutputParticleNeighborCount)
+	END_SHADER_PARAMETER_STRUCT()
 
+public:	
 	UPROPERTY(EditAnywhere, Category = "Grid")
-		uint32 MaxNeighborsPerCell;
-
-public:
-
-	DECLARE_NIAGARA_DI_PARAMETER();
+	uint32 MaxNeighborsPerCell;
 
 	virtual void PostInitProperties() override
 	{
@@ -92,9 +100,13 @@ public:
 
 	// GPU sim functionality
 #if WITH_EDITORONLY_DATA
+	virtual bool AppendCompileHash(FNiagaraCompileHashVisitor* InVisitor) const override;
 	virtual void GetParameterDefinitionHLSL(const FNiagaraDataInterfaceGPUParamInfo& ParamInfo, FString& OutHLSL) override;
 	virtual bool GetFunctionHLSL(const FNiagaraDataInterfaceGPUParamInfo& ParamInfo, const FNiagaraDataInterfaceGeneratedFunction& FunctionInfo, int FunctionInstanceIndex, FString& OutHLSL) override;
 #endif
+	virtual bool UseLegacyShaderBindings() const override { return false; }
+	virtual void BuildShaderParameters(FNiagaraShaderParametersBuilder& ShaderParametersBuilder) const override;
+	virtual void SetShaderParameters(const FNiagaraDataInterfaceSetShaderParametersContext& Context) const override;
 #if WITH_EDITOR
 	virtual bool ShouldCompile(EShaderPlatform ShaderPlatform) const override;
 #endif
diff --git a/Engine/Plugins/FX/Niagara/Source/Niagara/Classes/NiagaraDataInterfaceSpline.h b/Engine/Plugins/FX/Niagara/Source/Niagara/Classes/NiagaraDataInterfaceSpline.h
--- a/Engine/Plugins/FX/Niagara/Source/Niagara/Classes/NiagaraDataInterfaceSpline.h
+++ b/Engine/Plugins/FX/Niagara/Source/Niagara/Classes/NiagaraDataInterfaceSpline.h
@@ -188,9 +188,7 @@ class NIAGARA_API UNiagaraDataInterfaceSpline : public UNiagaraDataInterface
 {
 	GENERATED_UCLASS_BODY()
 
-	DECLARE_NIAGARA_DI_PARAMETER();
 public:
-	
 	/** The source actor from which to sample.  Note that this can only be set when used as a user variable on a component in the world.*/
 	UPROPERTY(EditAnywhere, Category = "Spline")
 	TObjectPtr<AActor> Source;
@@ -256,6 +254,9 @@ protected:
 
 	virtual bool AppendCompileHash(FNiagaraCompileHashVisitor* InVisitor) const override;
 #endif
+	virtual bool UseLegacyShaderBindings() const override { return false; }
+	virtual void BuildShaderParameters(FNiagaraShaderParametersBuilder& ShaderParametersBuilder) const override;
+	virtual void SetShaderParameters(const FNiagaraDataInterfaceSetShaderParametersContext& Context) const override;
 
 private:
 	
diff --git a/Engine/Plugins/FX/Niagara/Source/Niagara/Private/DataInterface/NiagaraDataInterfaceUObjectPropertyReader.cpp b/Engine/Plugins/FX/Niagara/Source/Niagara/Private/DataInterface/NiagaraDataInterfaceUObjectPropertyReader.cpp
--- a/Engine/Plugins/FX/Niagara/Source/Niagara/Private/DataInterface/NiagaraDataInterfaceUObjectPropertyReader.cpp
+++ b/Engine/Plugins/FX/Niagara/Source/Niagara/Private/DataInterface/NiagaraDataInterfaceUObjectPropertyReader.cpp
@@ -3,6 +3,7 @@
 #include "NiagaraDataInterfaceUObjectPropertyReader.h"
 #include "NiagaraComputeExecutionContext.h"
 #include "NiagaraFunctionLibrary.h"
+#include "NiagaraShaderParametersBuilder.h"
 #include "NiagaraSystemInstance.h"
 
 #include "Internationalization/Internationalization.h"
@@ -193,7 +194,7 @@ namespace NDIUObjectPropertyReaderLocal
 			return nullptr;
 		}
 		static constexpr TCHAR const* HlslBufferType = TEXT("float");
-		static constexpr TCHAR const* HlslBufferRead = TEXT("BUFFER[OFFSET] : 0.0f");
+		static constexpr TCHAR const* HlslBufferRead = TEXT("asfloat(BUFFER[OFFSET]) : 0.0f");
 	};
 
 	template<>
@@ -214,7 +215,7 @@ namespace NDIUObjectPropertyReaderLocal
 			return nullptr;
 		}
 		static constexpr TCHAR const* HlslBufferType = TEXT("float2");
-		static constexpr TCHAR const* HlslBufferRead = TEXT("float2(BUFFER[OFFSET + 0], BUFFER[OFFSET + 1]) : float2(0.0f, 0.0f)");
+		static constexpr TCHAR const* HlslBufferRead = TEXT("float2(asfloat(BUFFER[OFFSET + 0]), asfloat(BUFFER[OFFSET + 1])) : float2(0.0f, 0.0f)");
 	};
 
 	template<>
@@ -235,7 +236,7 @@ namespace NDIUObjectPropertyReaderLocal
 			return nullptr;
 		}
 		static constexpr TCHAR const* HlslBufferType = TEXT("float3");
-		static constexpr TCHAR const* HlslBufferRead = TEXT("float3(BUFFER[OFFSET + 0], BUFFER[OFFSET + 1], BUFFER[OFFSET + 2]) : float3(0.0f, 0.0f, 0.0f)");
+		static constexpr TCHAR const* HlslBufferRead = TEXT("float3(asfloat(BUFFER[OFFSET + 0]), asfloat(BUFFER[OFFSET + 1]), asfloat(BUFFER[OFFSET + 2])) : float3(0.0f, 0.0f, 0.0f)");
 	};
 
 	template<>
@@ -256,7 +257,7 @@ namespace NDIUObjectPropertyReaderLocal
 			return nullptr;
 		}
 		static constexpr TCHAR const* HlslBufferType = TEXT("float4");
-		static constexpr TCHAR const* HlslBufferRead = TEXT("float4(BUFFER[OFFSET + 0], BUFFER[OFFSET + 1], BUFFER[OFFSET + 2], BUFFER[OFFSET + 3]) : float4(0.0f, 0.0f, 0.0f, 0.0f)");
+		static constexpr TCHAR const* HlslBufferRead = TEXT("float4(asfloat(BUFFER[OFFSET + 0]), asfloat(BUFFER[OFFSET + 1]), asfloat(BUFFER[OFFSET + 2]), asfloat(BUFFER[OFFSET + 3])) : float4(0.0f, 0.0f, 0.0f, 0.0f)");
 	};
 
 	template<>
@@ -281,7 +282,7 @@ namespace NDIUObjectPropertyReaderLocal
 			return nullptr;
 		}
 		static constexpr TCHAR const* HlslBufferType = TEXT("float4");
-		static constexpr TCHAR const* HlslBufferRead = TEXT("float4(BUFFER[OFFSET + 0], BUFFER[OFFSET + 1], BUFFER[OFFSET + 2], BUFFER[OFFSET + 3]) : float4(0.0f, 0.0f, 0.0f, 0.0f)");
+		static constexpr TCHAR const* HlslBufferRead = TEXT("float4(asfloat(BUFFER[OFFSET + 0]), asfloat(BUFFER[OFFSET + 1]), asfloat(BUFFER[OFFSET + 2]), asfloat(BUFFER[OFFSET + 3])) : float4(0.0f, 0.0f, 0.0f, 0.0f)");
 	};
 
 	template<>
@@ -302,7 +303,7 @@ namespace NDIUObjectPropertyReaderLocal
 			return nullptr;
 		}
 		static constexpr TCHAR const* HlslBufferType = TEXT("float3");
-		static constexpr TCHAR const* HlslBufferRead = TEXT("float3(BUFFER[OFFSET + 0], BUFFER[OFFSET + 1], BUFFER[OFFSET + 2]) : float3(0.0f, 0.0f, 0.0f)");
+		static constexpr TCHAR const* HlslBufferRead = TEXT("float3(asfloat(BUFFER[OFFSET + 0)], asfloat(BUFFER[OFFSET + 1]), asfloat(BUFFER[OFFSET + 2])) : float3(0.0f, 0.0f, 0.0f)");
 	};
 
 	template<>
@@ -445,93 +446,6 @@ namespace NDIUObjectPropertyReaderLocal
 
 //////////////////////////////////////////////////////////////////////////
 
-struct FNDIUObjectPropertyReaderCS : public FNiagaraDataInterfaceParametersCS
-{
-	DECLARE_TYPE_LAYOUT(FNDIUObjectPropertyReaderCS, NonVirtual);
-
-public:
-	void Bind(const FNiagaraDataInterfaceGPUParamInfo& ParameterInfo, const class FShaderParameterMap& ParameterMap)
-	{
-		using namespace NDIUObjectPropertyReaderLocal;
-
-		PropertiesToRead.Reserve(ParameterInfo.GeneratedFunctions.Num());
-		for ( const FNiagaraDataInterfaceGeneratedFunction& FunctionInfo : ParameterInfo.GeneratedFunctions )
-		{
-			#define NDI_PROPERTY_TYPE(TYPE) \
-				else if (FunctionInfo.DefinitionName == FTypeHelper<TYPE>::GetFunctionName()) { PropertiesToRead.Emplace(FTypeHelper<TYPE>::GetTypeDef(), FunctionInfo.Specifiers[0].Value); }
-
-				if (false) { } NDI_PROPERTY_TYPES
-			#undef NDI_PROPERTY_TYPE
-		}
-
-		TransformInfoParam.Bind(ParameterMap, *(TEXT("TransformInfo_") + ParameterInfo.DataInterfaceHLSLSymbol));
-		InvTransformInfoParam.Bind(ParameterMap, *(TEXT("InvTransformInfo_") + ParameterInfo.DataInterfaceHLSLSymbol));
-		PropertyDataParam.Bind(ParameterMap, *(TEXT("PropertyData_") + ParameterInfo.DataInterfaceHLSLSymbol));
-		PropertyOffsetsParam.Bind(ParameterMap, *(TEXT("PropertyOffsets_") + ParameterInfo.DataInterfaceHLSLSymbol));
-	}
-
-	void Set(FRHICommandList& RHICmdList, const FNiagaraDataInterfaceSetArgs& Context) const
-	{
-		using namespace NDIUObjectPropertyReaderLocal;
-
-		check(IsInRenderingThread());
-
-		FRHIComputeShader* ComputeShaderRHI = Context.Shader.GetComputeShader();
-		const FNDIProxy* DIProxy = static_cast<FNDIProxy*>(Context.DataInterface);
-		const FInstanceData_RenderThread* InstanceData_RT = &DIProxy->PerInstanceData_RenderThread.FindChecked(Context.SystemInstanceID);
-		SetSRVParameter(RHICmdList, ComputeShaderRHI, PropertyDataParam, InstanceData_RT->PropertyData.SRV);
-		SetShaderValueArray(RHICmdList, ComputeShaderRHI, PropertyOffsetsParam, InstanceData_RT->PropertyOffsets.GetData(), InstanceData_RT->PropertyOffsets.Num());
-
-		{
-			FVector4f TransformData[3];
-			if (InstanceData_RT->CachedTransform.IsSet())
-			{
-				const FQuat4f TransformRotation = FQuat4f(InstanceData_RT->CachedTransform->GetRotation());
-				TransformData[0] = FVector4f(FVector3f(InstanceData_RT->CachedTransform->GetLocation()), 1.0f);
-				TransformData[1] = FVector4f(TransformRotation.X, TransformRotation.Y, TransformRotation.Z, TransformRotation.W);
-				TransformData[2] = FVector4f(FVector3f(InstanceData_RT->CachedTransform->GetScale3D()));
-			}
-			else
-			{
-				TransformData[0] = FVector4f(0.0f, 0.0f, 0.0f, 0.0f);
-				TransformData[1] = FVector4f(0.0f, 0.0f, 0.0f, 1.0f);
-				TransformData[2] = FVector4f(1.0f, 1.0f, 1.0f, 0.0f);
-			}
-			SetShaderValueArray(RHICmdList, ComputeShaderRHI, TransformInfoParam, TransformData, UE_ARRAY_COUNT(TransformData));
-		}
-
-		{
-			FVector4f InvTransformData[3];
-			if (InstanceData_RT->CachedInvTransform.IsSet())
-			{
-				const FQuat4f InvTransformRotation = FQuat4f(InstanceData_RT->CachedInvTransform->GetRotation());
-				InvTransformData[0] = FVector4f(FVector3f(InstanceData_RT->CachedInvTransform->GetLocation()), 1.0f);
-				InvTransformData[1] = FVector4f(InvTransformRotation.X, InvTransformRotation.Y, InvTransformRotation.Z, InvTransformRotation.W);
-				InvTransformData[2] = FVector4f(FVector3f(InstanceData_RT->CachedInvTransform->GetScale3D()));
-			}
-			else
-			{
-				InvTransformData[0] = FVector4f(0.0f, 0.0f, 0.0f, 0.0f);
-				InvTransformData[1] = FVector4f(0.0f, 0.0f, 0.0f, 1.0f);
-				InvTransformData[2] = FVector4f(1.0f, 1.0f, 1.0f, 0.0f);
-			}
-			SetShaderValueArray(RHICmdList, ComputeShaderRHI, InvTransformInfoParam, InvTransformData, UE_ARRAY_COUNT(InvTransformData));
-		}
-	}
-
-private:
-	LAYOUT_FIELD(FShaderParameter,							TransformInfoParam);
-	LAYOUT_FIELD(FShaderParameter,							InvTransformInfoParam);
-	LAYOUT_FIELD(TMemoryImageArray<FNiagaraVariableBase>,	PropertiesToRead);
-	LAYOUT_FIELD(FShaderResourceParameter,					PropertyDataParam);
-	LAYOUT_FIELD(FShaderParameter,							PropertyOffsetsParam);
-};
-
-IMPLEMENT_TYPE_LAYOUT(FNDIUObjectPropertyReaderCS);
-IMPLEMENT_NIAGARA_DI_PARAMETER(UNiagaraDataInterfaceUObjectPropertyReader, FNDIUObjectPropertyReaderCS);
-
-//////////////////////////////////////////////////////////////////////////
-
 UNiagaraDataInterfaceUObjectPropertyReader::UNiagaraDataInterfaceUObjectPropertyReader(FObjectInitializer const& ObjectInitializer)
 	: Super(ObjectInitializer)
 {
@@ -854,12 +768,27 @@ void UNiagaraDataInterfaceUObjectPropertyReader::GetVMExternalFunction(const FVM
 }
 
 #if WITH_EDITORONLY_DATA
+bool UNiagaraDataInterfaceUObjectPropertyReader::AppendCompileHash(FNiagaraCompileHashVisitor* InVisitor) const
+{
+	bool bSuccess = Super::AppendCompileHash(InVisitor);
+	bSuccess &= InVisitor->UpdateShaderParameters<FShaderParameters>();
+	return bSuccess;
+}
+
 void UNiagaraDataInterfaceUObjectPropertyReader::GetParameterDefinitionHLSL(const FNiagaraDataInterfaceGPUParamInfo& ParameterInfo, FString& OutHLSL)
 {
-	OutHLSL.Appendf(TEXT("float4 TransformInfo_%s[3];\n"), *ParameterInfo.DataInterfaceHLSLSymbol);
-	OutHLSL.Appendf(TEXT("float4 InvTransformInfo_%s[3];\n"), *ParameterInfo.DataInterfaceHLSLSymbol);
-	OutHLSL.Appendf(TEXT("Buffer<float> PropertyData_%s;\n"), *ParameterInfo.DataInterfaceHLSLSymbol);
-	OutHLSL.Appendf(TEXT("uint4 PropertyOffsets_%s[%d];\n"), *ParameterInfo.DataInterfaceHLSLSymbol, FMath::DivideAndRoundUp(ParameterInfo.GeneratedFunctions.Num(), 4));
+	OutHLSL.Appendf(TEXT("float3	%s_TransformLocation;\n"), *ParameterInfo.DataInterfaceHLSLSymbol);
+	OutHLSL.Appendf(TEXT("uint		%s_TransformValid;\n"), *ParameterInfo.DataInterfaceHLSLSymbol);
+	OutHLSL.Appendf(TEXT("float4	%s_TransformRotation;\n"), *ParameterInfo.DataInterfaceHLSLSymbol);
+	OutHLSL.Appendf(TEXT("float3	%s_TransformScale;\n"), *ParameterInfo.DataInterfaceHLSLSymbol);
+
+	OutHLSL.Appendf(TEXT("float3	%s_InvTransformLocation;\n"), *ParameterInfo.DataInterfaceHLSLSymbol);
+	OutHLSL.Appendf(TEXT("uint		%s_InvTransformValid;\n"), *ParameterInfo.DataInterfaceHLSLSymbol);
+	OutHLSL.Appendf(TEXT("float4	%s_InvTransformRotation;\n"), *ParameterInfo.DataInterfaceHLSLSymbol);
+	OutHLSL.Appendf(TEXT("float3	%s_InvTransformScale;\n"), *ParameterInfo.DataInterfaceHLSLSymbol);
+
+	OutHLSL.Appendf(TEXT("Buffer<uint>	%s_PropertyData;\n"), *ParameterInfo.DataInterfaceHLSLSymbol);
+	OutHLSL.Appendf(TEXT("uint4			%s_PropertyOffsets[%d];\n"), *ParameterInfo.DataInterfaceHLSLSymbol, FMath::DivideAndRoundUp(ParameterInfo.GeneratedFunctions.Num(), 4));
 }
 
 bool UNiagaraDataInterfaceUObjectPropertyReader::GetFunctionHLSL(const FNiagaraDataInterfaceGPUParamInfo& ParamInfo, const FNiagaraDataInterfaceGeneratedFunction& FunctionInfo, int FunctionInstanceIndex, FString& OutHLSL)
@@ -873,10 +802,10 @@ bool UNiagaraDataInterfaceUObjectPropertyReader::GetFunctionHLSL(const FNiagaraD
 	{
 		OutHLSL.Appendf(TEXT("void %s(out bool bSuccess, out float3 Position, out float4 Rotation, out float3 Scale)\n"), *FunctionInfo.InstanceName);
 		OutHLSL.Append(TEXT("{\n"));
-		OutHLSL.Appendf(TEXT("	bSuccess = TransformInfo_%s[0].w > 0.0f;\n"), *ParamInfo.DataInterfaceHLSLSymbol);
-		OutHLSL.Appendf(TEXT("	Position = TransformInfo_%s[0].xyz;\n"), *ParamInfo.DataInterfaceHLSLSymbol);
-		OutHLSL.Appendf(TEXT("	Rotation = TransformInfo_%s[1];\n"), *ParamInfo.DataInterfaceHLSLSymbol);
-		OutHLSL.Appendf(TEXT("	Scale = TransformInfo_%s[2].xyz;\n"), *ParamInfo.DataInterfaceHLSLSymbol);
+		OutHLSL.Appendf(TEXT("	bSuccess = %s_TransformValid != 0;\n"), *ParamInfo.DataInterfaceHLSLSymbol);
+		OutHLSL.Appendf(TEXT("	Position = %s_TransformLocation;\n"), *ParamInfo.DataInterfaceHLSLSymbol);
+		OutHLSL.Appendf(TEXT("	Rotation = %s_TransformRotation;\n"), *ParamInfo.DataInterfaceHLSLSymbol);
+		OutHLSL.Appendf(TEXT("	Scale = %s_TransformScale;\n"), *ParamInfo.DataInterfaceHLSLSymbol);
 		OutHLSL.Append(TEXT("}\n"));
 		return true;
 	}
@@ -884,10 +813,10 @@ bool UNiagaraDataInterfaceUObjectPropertyReader::GetFunctionHLSL(const FNiagaraD
 	{
 		OutHLSL.Appendf(TEXT("void %s(out bool bSuccess, out float3 Position, out float4 Rotation, out float3 Scale)\n"), *FunctionInfo.InstanceName);
 		OutHLSL.Append(TEXT("{\n"));
-		OutHLSL.Appendf(TEXT("	bSuccess = InvTransformInfo_%s[0].w > 0.0f;\n"), *ParamInfo.DataInterfaceHLSLSymbol);
-		OutHLSL.Appendf(TEXT("	Position = InvTransformInfo_%s[0].xyz;\n"), *ParamInfo.DataInterfaceHLSLSymbol);
-		OutHLSL.Appendf(TEXT("	Rotation = InvTransformInfo_%s[1];\n"), *ParamInfo.DataInterfaceHLSLSymbol);
-		OutHLSL.Appendf(TEXT("	Scale = InvTransformInfo_%s[2].xyz;\n"), *ParamInfo.DataInterfaceHLSLSymbol);
+		OutHLSL.Appendf(TEXT("	bSuccess = %s_InvTransformValid != 0;\n"), *ParamInfo.DataInterfaceHLSLSymbol);
+		OutHLSL.Appendf(TEXT("	Position = %s_InvTransformLocation;\n"), *ParamInfo.DataInterfaceHLSLSymbol);
+		OutHLSL.Appendf(TEXT("	Rotation = %s_InvTransformRotation;\n"), *ParamInfo.DataInterfaceHLSLSymbol);
+		OutHLSL.Appendf(TEXT("	Scale = %s_InvTransformScale;\n"), *ParamInfo.DataInterfaceHLSLSymbol);
 		OutHLSL.Append(TEXT("}\n"));
 		return true;
 	}
@@ -904,12 +833,12 @@ bool UNiagaraDataInterfaceUObjectPropertyReader::GetFunctionHLSL(const FNiagaraD
 	if ( HlslBufferType && HlslBufferRead )
 	{
 		FString HlslBufferReadString(HlslBufferRead);
-		HlslBufferReadString.ReplaceInline(TEXT("BUFFER"), *(TEXT("PropertyData_") + ParamInfo.DataInterfaceHLSLSymbol));
+		HlslBufferReadString.ReplaceInline(TEXT("BUFFER"), *(ParamInfo.DataInterfaceHLSLSymbol + TEXT("_PropertyData")));
 		HlslBufferReadString.ReplaceInline(TEXT("OFFSET"), TEXT("BufferOffset"));
 
 		OutHLSL.Appendf(TEXT("void %s(out bool bSuccess, out %s Value)\n"), *FunctionInfo.InstanceName, HlslBufferType);
 		OutHLSL.Append(TEXT("{\n"));
-		OutHLSL.Appendf(TEXT("	uint BufferOffset = PropertyOffsets_%s[%d][%d];\n"), *ParamInfo.DataInterfaceHLSLSymbol, FunctionInstanceIndex / 4, FunctionInstanceIndex % 4);
+		OutHLSL.Appendf(TEXT("	uint BufferOffset = %s_PropertyOffsets[%d][%d];\n"), *ParamInfo.DataInterfaceHLSLSymbol, FunctionInstanceIndex / 4, FunctionInstanceIndex % 4);
 		OutHLSL.Append(TEXT("	bSuccess = BufferOffset != 0xffffffff;\n"));
 		OutHLSL.Appendf(TEXT("	Value = bSuccess ? %s;\n"), *HlslBufferReadString);
 		OutHLSL.Append(TEXT("}\n"));
@@ -920,6 +849,38 @@ bool UNiagaraDataInterfaceUObjectPropertyReader::GetFunctionHLSL(const FNiagaraD
 }
 #endif
 
+void UNiagaraDataInterfaceUObjectPropertyReader::BuildShaderParameters(FNiagaraShaderParametersBuilder& ShaderParametersBuilder) const
+{
+	ShaderParametersBuilder.AddNestedStruct<FShaderParameters>();
+
+	const int32 ProperyOffsetsSize = FMath::Max(FMath::DivideAndRoundUp(ShaderParametersBuilder.GetGeneratedFunctions().Num(), 4), 1);
+	ShaderParametersBuilder.AddLooseParamArray<FUintVector4>(TEXT("PropertyOffsets"), ProperyOffsetsSize);
+}
+
+void UNiagaraDataInterfaceUObjectPropertyReader::SetShaderParameters(const FNiagaraDataInterfaceSetShaderParametersContext& Context) const
+{
+	using namespace NDIUObjectPropertyReaderLocal;
+
+	const FNDIProxy& DIProxy = Context.GetProxy<FNDIProxy>();
+	const FInstanceData_RenderThread& InstanceData_RT = DIProxy.PerInstanceData_RenderThread.FindChecked(Context.GetSystemInstanceID());
+
+	FShaderParameters* ShaderParameters		= Context.GetParameterNestedStruct<FShaderParameters>();
+	ShaderParameters->TransformLocation		= InstanceData_RT.CachedTransform.IsSet() ? FVector3f(InstanceData_RT.CachedTransform->GetLocation()) : FVector3f::ZeroVector;
+	ShaderParameters->TransformValid		= InstanceData_RT.CachedTransform.IsSet() ? 1 : 0;
+	ShaderParameters->TransformRotation		= InstanceData_RT.CachedTransform.IsSet() ? FQuat4f(InstanceData_RT.CachedTransform->GetRotation()) : FQuat4f::Identity;
+	ShaderParameters->TransformScale		= InstanceData_RT.CachedTransform.IsSet() ? FVector3f(InstanceData_RT.CachedTransform->GetScale3D()) : FVector3f::OneVector;
+	ShaderParameters->InvTransformLocation	= InstanceData_RT.CachedInvTransform.IsSet() ? FVector3f(InstanceData_RT.CachedInvTransform->GetLocation()) : FVector3f::ZeroVector;
+	ShaderParameters->InvTransformValid		= InstanceData_RT.CachedInvTransform.IsSet() ? 1 : 0;
+	ShaderParameters->InvTransformRotation	= InstanceData_RT.CachedInvTransform.IsSet() ? FQuat4f(InstanceData_RT.CachedInvTransform->GetRotation()) : FQuat4f::Identity;
+	ShaderParameters->InvTransformScale		= InstanceData_RT.CachedInvTransform.IsSet() ? FVector3f(InstanceData_RT.CachedInvTransform->GetScale3D()) : FVector3f::OneVector;
+	ShaderParameters->PropertyData			= InstanceData_RT.PropertyData.SRV;
+
+	check((InstanceData_RT.PropertyOffsets.Num() % 4) == 0 && InstanceData_RT.PropertyOffsets.Num() > 0);
+	const int32 ProperyOffsetsSize = InstanceData_RT.PropertyOffsets.Num() / 4;
+	TArrayView<FUintVector4> PropertyOffsets = Context.GetParameterLooseArray<FUintVector4>(ProperyOffsetsSize);
+	FMemory::Memcpy(PropertyOffsets.GetData(), InstanceData_RT.PropertyOffsets.GetData(), InstanceData_RT.PropertyOffsets.Num() * InstanceData_RT.PropertyOffsets.GetTypeSize());
+}
+
 void UNiagaraDataInterfaceUObjectPropertyReader::SetUObjectReaderPropertyRemap(UNiagaraComponent* NiagaraComponent, FName UserParameterName, FName GraphName, FName RemapName)
 {
 	UNiagaraDataInterfaceUObjectPropertyReader* ReaderDI = UNiagaraFunctionLibrary::GetDataInterface<UNiagaraDataInterfaceUObjectPropertyReader>(NiagaraComponent, UserParameterName);
diff --git a/Engine/Plugins/FX/Niagara/Source/Niagara/Private/DataInterface/NiagaraDataInterfaceUObjectPropertyReader.h b/Engine/Plugins/FX/Niagara/Source/Niagara/Private/DataInterface/NiagaraDataInterfaceUObjectPropertyReader.h
--- a/Engine/Plugins/FX/Niagara/Source/Niagara/Private/DataInterface/NiagaraDataInterfaceUObjectPropertyReader.h
+++ b/Engine/Plugins/FX/Niagara/Source/Niagara/Private/DataInterface/NiagaraDataInterfaceUObjectPropertyReader.h
@@ -33,9 +33,21 @@ class UNiagaraDataInterfaceUObjectPropertyReader : public UNiagaraDataInterface
 {
 	GENERATED_UCLASS_BODY()
 
-public:
-	DECLARE_NIAGARA_DI_PARAMETER();
+	BEGIN_SHADER_PARAMETER_STRUCT(FShaderParameters, )
+		SHADER_PARAMETER(FVector3f,			TransformLocation)
+		SHADER_PARAMETER(uint32,			TransformValid)
+		SHADER_PARAMETER(FQuat4f,			TransformRotation)
+		SHADER_PARAMETER(FVector3f,			TransformScale)
+
+		SHADER_PARAMETER(FVector3f,			InvTransformLocation)
+		SHADER_PARAMETER(uint32,			InvTransformValid)
+		SHADER_PARAMETER(FQuat4f,			InvTransformRotation)
+		SHADER_PARAMETER(FVector3f,			InvTransformScale)
 
+		SHADER_PARAMETER_SRV(Buffer<uint>,	PropertyData)
+	END_SHADER_PARAMETER_STRUCT()
+
+public:
 	/** User parameter Object binding to read properties from. */
 	UPROPERTY(EditAnywhere, Category = "UObjectPropertyReader")
 	FNiagaraUserParameterBinding UObjectParameterBinding;
@@ -78,10 +90,13 @@ public:
 	virtual void GetFunctions(TArray<FNiagaraFunctionSignature>& OutFunctions) override;
 	virtual void GetVMExternalFunction(const FVMExternalFunctionBindingInfo& BindingInfo, void* InstanceData, FVMExternalFunction& OutFunc) override;
 #if WITH_EDITORONLY_DATA
-	//virtual bool AppendCompileHash(FNiagaraCompileHashVisitor* InVisitor) const override;
+	virtual bool AppendCompileHash(FNiagaraCompileHashVisitor* InVisitor) const override;
 	virtual void GetParameterDefinitionHLSL(const FNiagaraDataInterfaceGPUParamInfo& ParamInfo, FString& OutHLSL) override;
 	virtual bool GetFunctionHLSL(const FNiagaraDataInterfaceGPUParamInfo& ParamInfo, const FNiagaraDataInterfaceGeneratedFunction& FunctionInfo, int FunctionInstanceIndex, FString& OutHLSL) override;
 #endif
+	virtual bool UseLegacyShaderBindings() const override { return false; }
+	virtual void BuildShaderParameters(FNiagaraShaderParametersBuilder& ShaderParametersBuilder) const override;
+	virtual void SetShaderParameters(const FNiagaraDataInterfaceSetShaderParametersContext& Context) const override;
 	//UNiagaraDataInterface Interface
 
 	/** Remaps a property reader, where the  */
diff --git a/Engine/Plugins/FX/Niagara/Source/Niagara/Private/NiagaraDataInterfaceNeighborGrid3D.cpp b/Engine/Plugins/FX/Niagara/Source/Niagara/Private/NiagaraDataInterfaceNeighborGrid3D.cpp
--- a/Engine/Plugins/FX/Niagara/Source/Niagara/Private/NiagaraDataInterfaceNeighborGrid3D.cpp
+++ b/Engine/Plugins/FX/Niagara/Source/Niagara/Private/NiagaraDataInterfaceNeighborGrid3D.cpp
@@ -3,6 +3,7 @@
 #include "NiagaraShader.h"
 #include "ShaderParameterUtils.h"
 #include "NiagaraGpuComputeDispatchInterface.h"
+#include "NiagaraShaderParametersBuilder.h"
 #include "NiagaraSystemInstance.h"
 #include "NiagaraRenderer.h"
 #include "NiagaraShaderParticleID.h"
@@ -15,10 +16,8 @@ static const FString ParticleNeighborCountName(TEXT("_ParticleNeighborCount"));
 static const FString OutputParticleNeighborsName(TEXT("_OutputParticleNeighbors"));
 static const FString OutputParticleNeighborCountName(TEXT("_OutputParticleNeighborCount"));
 
-
 const FName UNiagaraDataInterfaceNeighborGrid3D::SetNumCellsFunctionName("SetNumCells");
 
-
 // Global VM function names, also used by the shaders code generation methods.
 
 static const FName MaxNeighborsPerCellFunctionName("MaxNeighborsPerCell");
@@ -36,143 +35,6 @@ static FAutoConsoleVariableRef CVarMaxNiagaraNeighborGridCells(
 	ECVF_Default
 );
 
-
-/*--------------------------------------------------------------------------------------------------------------------------*/
-struct FNiagaraDataInterfaceParametersCS_NeighborGrid3D : public FNiagaraDataInterfaceParametersCS
-{
-	DECLARE_TYPE_LAYOUT(FNiagaraDataInterfaceParametersCS_NeighborGrid3D, NonVirtual);
-public:
-	void Bind(const FNiagaraDataInterfaceGPUParamInfo& ParameterInfo, const class FShaderParameterMap& ParameterMap)
-	{
-		NumCellsParam.Bind(ParameterMap, *(ParameterInfo.DataInterfaceHLSLSymbol + UNiagaraDataInterfaceRWBase::NumCellsName));
-		UnitToUVParam.Bind(ParameterMap, *(ParameterInfo.DataInterfaceHLSLSymbol + UNiagaraDataInterfaceRWBase::UnitToUVName));
-		CellSizeParam.Bind(ParameterMap, *(ParameterInfo.DataInterfaceHLSLSymbol + UNiagaraDataInterfaceRWBase::CellSizeName));
-
-		MaxNeighborsPerCellParam.Bind(ParameterMap, *(ParameterInfo.DataInterfaceHLSLSymbol + MaxNeighborsPerCellName));
-		WorldBBoxSizeParam.Bind(ParameterMap, *(ParameterInfo.DataInterfaceHLSLSymbol + UNiagaraDataInterfaceRWBase::WorldBBoxSizeName));
-		
-		ParticleNeighborsGridParam.Bind(ParameterMap,  *(ParameterInfo.DataInterfaceHLSLSymbol + ParticleNeighborsName));
-		ParticleNeighborCountGridParam.Bind(ParameterMap,  *(ParameterInfo.DataInterfaceHLSLSymbol + ParticleNeighborCountName));
-
-		OutputParticleNeighborsGridParam.Bind(ParameterMap, *(ParameterInfo.DataInterfaceHLSLSymbol + OutputParticleNeighborsName));
-		OutputParticleNeighborCountGridParam.Bind(ParameterMap, *(ParameterInfo.DataInterfaceHLSLSymbol + OutputParticleNeighborCountName));
-	}
-
-	// #todo(dmp): make resource transitions batched
-	void Set(FRHICommandList& RHICmdList, const FNiagaraDataInterfaceSetArgs& Context) const
-	{
-		check(IsInRenderingThread());
-
-		// Get shader and DI
-		FRHIComputeShader* ComputeShaderRHI = Context.Shader.GetComputeShader();
-		FNiagaraDataInterfaceProxyNeighborGrid3D* VFDI = static_cast<FNiagaraDataInterfaceProxyNeighborGrid3D*>(Context.DataInterface);
-
-		NeighborGrid3DRWInstanceData* ProxyData = VFDI->SystemInstancesToProxyData.Find(Context.SystemInstanceID);
-		float CellSizeTmp[3];
-
-		if (!(ProxyData && ProxyData->NeighborhoodBuffer.Buffer.IsValid()))
-		{
-			CellSizeTmp[0] = CellSizeTmp[1] = CellSizeTmp[2] = 1.0f;
-			SetShaderValue(RHICmdList, ComputeShaderRHI, NumCellsParam, 0);
-			SetShaderValue(RHICmdList, ComputeShaderRHI, UnitToUVParam, FVector3f::ZeroVector);
-			SetShaderValue(RHICmdList, ComputeShaderRHI, CellSizeParam, CellSizeTmp);
-			SetShaderValue(RHICmdList, ComputeShaderRHI, MaxNeighborsPerCellParam, 0);
-			SetShaderValue(RHICmdList, ComputeShaderRHI, WorldBBoxSizeParam, FVector3f(0.0f, 0.0f, 0.0f));
-			SetSRVParameter(RHICmdList, ComputeShaderRHI, ParticleNeighborsGridParam, FNiagaraRenderer::GetDummyIntBuffer());
-			SetSRVParameter(RHICmdList, ComputeShaderRHI, ParticleNeighborCountGridParam, FNiagaraRenderer::GetDummyIntBuffer());
-			if (OutputParticleNeighborsGridParam.IsUAVBound())
-			{
-				RHICmdList.SetUAVParameter(ComputeShaderRHI, OutputParticleNeighborsGridParam.GetUAVIndex(), Context.ComputeDispatchInterface->GetEmptyUAVFromPool(RHICmdList, PF_R32_SINT, ENiagaraEmptyUAVType::Buffer));
-			}
-
-			if (OutputParticleNeighborCountGridParam.IsUAVBound())
-			{
-				RHICmdList.SetUAVParameter(ComputeShaderRHI, OutputParticleNeighborCountGridParam.GetUAVIndex(), Context.ComputeDispatchInterface->GetEmptyUAVFromPool(RHICmdList, PF_R32_SINT, ENiagaraEmptyUAVType::Buffer));
-			}
-
-			return;
-		}
-
-		SetShaderValue(RHICmdList, ComputeShaderRHI, NumCellsParam, ProxyData->NumCells);
-		SetShaderValue(RHICmdList, ComputeShaderRHI, UnitToUVParam, FVector3f(1.0f) / FVector3f(ProxyData->NumCells));
-
-		// #todo(dmp): remove this computation here
-		CellSizeTmp[0] = ProxyData->WorldBBoxSize.X / ProxyData->NumCells.X;
-		CellSizeTmp[1] = ProxyData->WorldBBoxSize.Y / ProxyData->NumCells.Y;
-		CellSizeTmp[2] = ProxyData->WorldBBoxSize.Z / ProxyData->NumCells.Z;
-		SetShaderValue(RHICmdList, ComputeShaderRHI, CellSizeParam, CellSizeTmp);
-		SetShaderValue(RHICmdList, ComputeShaderRHI, MaxNeighborsPerCellParam, ProxyData->MaxNeighborsPerCell);
-		SetShaderValue(RHICmdList, ComputeShaderRHI, WorldBBoxSizeParam, FVector3f(ProxyData->WorldBBoxSize));	// LWC_TODO: Potential issues for large world representation here.
-		
-		if (!Context.IsOutputStage)
-		{
-			if (ParticleNeighborsGridParam.IsBound())
-			{
-				RHICmdList.Transition(FRHITransitionInfo(ProxyData->NeighborhoodBuffer.UAV, ERHIAccess::Unknown, ERHIAccess::SRVCompute));
-				SetSRVParameter(RHICmdList, ComputeShaderRHI, ParticleNeighborsGridParam, ProxyData->NeighborhoodBuffer.SRV);
-			}
-
-			if (ParticleNeighborCountGridParam.IsBound())
-			{
-				RHICmdList.Transition(FRHITransitionInfo(ProxyData->NeighborhoodCountBuffer.UAV, ERHIAccess::Unknown, ERHIAccess::SRVCompute));
-				SetSRVParameter(RHICmdList, ComputeShaderRHI, ParticleNeighborCountGridParam, ProxyData->NeighborhoodCountBuffer.SRV);
-			}
-
-			if (OutputParticleNeighborsGridParam.IsUAVBound())
-			{
-				RHICmdList.SetUAVParameter(ComputeShaderRHI, OutputParticleNeighborsGridParam.GetUAVIndex(), Context.ComputeDispatchInterface->GetEmptyUAVFromPool(RHICmdList, PF_R32_SINT, ENiagaraEmptyUAVType::Buffer));
-			}
-
-			if (OutputParticleNeighborCountGridParam.IsUAVBound())
-			{
-				RHICmdList.SetUAVParameter(ComputeShaderRHI, OutputParticleNeighborCountGridParam.GetUAVIndex(), Context.ComputeDispatchInterface->GetEmptyUAVFromPool(RHICmdList, PF_R32_SINT, ENiagaraEmptyUAVType::Buffer));
-			}
-		}
-		else
-		{
-			SetSRVParameter(RHICmdList, ComputeShaderRHI, ParticleNeighborsGridParam, FNiagaraRenderer::GetDummyIntBuffer());
-			SetSRVParameter(RHICmdList, ComputeShaderRHI, ParticleNeighborCountGridParam, FNiagaraRenderer::GetDummyIntBuffer());
-
-			if (OutputParticleNeighborsGridParam.IsBound())
-			{
-				RHICmdList.Transition(FRHITransitionInfo(ProxyData->NeighborhoodBuffer.UAV, ERHIAccess::Unknown, ERHIAccess::UAVCompute));
-				OutputParticleNeighborsGridParam.SetBuffer(RHICmdList, ComputeShaderRHI, ProxyData->NeighborhoodBuffer);
-			}
-
-			if (OutputParticleNeighborCountGridParam.IsBound())
-			{
-				RHICmdList.Transition(FRHITransitionInfo(ProxyData->NeighborhoodCountBuffer.UAV, ERHIAccess::Unknown, ERHIAccess::UAVCompute));
-				OutputParticleNeighborCountGridParam.SetBuffer(RHICmdList, ComputeShaderRHI, ProxyData->NeighborhoodCountBuffer);
-			}
-		}
-		// Note: There is a flush in PreEditChange to make sure everything is synced up at this point 
-	}
-
-	void Unset(FRHICommandList& RHICmdList, const FNiagaraDataInterfaceSetArgs& Context) const 
-	{
-		if (OutputParticleNeighborsGridParam.IsBound())
-		{
-			OutputParticleNeighborsGridParam.UnsetUAV(RHICmdList, Context.Shader.GetComputeShader());
-		}
-
-		if (OutputParticleNeighborCountGridParam.IsBound())
-		{
-			OutputParticleNeighborCountGridParam.UnsetUAV(RHICmdList, Context.Shader.GetComputeShader());
-		}
-	}
-
-private:
-	LAYOUT_FIELD(FShaderParameter, NumCellsParam);
-	LAYOUT_FIELD(FShaderParameter, UnitToUVParam);
-	LAYOUT_FIELD(FShaderParameter, CellSizeParam);
-	LAYOUT_FIELD(FShaderParameter, MaxNeighborsPerCellParam);
-	LAYOUT_FIELD(FShaderParameter, WorldBBoxSizeParam);
-	LAYOUT_FIELD(FShaderResourceParameter, ParticleNeighborsGridParam);
-	LAYOUT_FIELD(FShaderResourceParameter, ParticleNeighborCountGridParam);
-	LAYOUT_FIELD(FRWShaderParameter, OutputParticleNeighborCountGridParam);
-	LAYOUT_FIELD(FRWShaderParameter, OutputParticleNeighborsGridParam);
-};
-
 void NeighborGrid3DRWInstanceData::ResizeBuffers()
 {
 	uint32 NumTotalCells = NumCells.X * NumCells.Y * NumCells.Z;
@@ -192,9 +54,7 @@ void NeighborGrid3DRWInstanceData::ResizeBuffers()
 	#endif
 }
 
-IMPLEMENT_TYPE_LAYOUT(FNiagaraDataInterfaceParametersCS_NeighborGrid3D);
-
-IMPLEMENT_NIAGARA_DI_PARAMETER(UNiagaraDataInterfaceNeighborGrid3D, FNiagaraDataInterfaceParametersCS_NeighborGrid3D);
+//////////////////////////////////////////////////////////////////////////
 
 UNiagaraDataInterfaceNeighborGrid3D::UNiagaraDataInterfaceNeighborGrid3D(FObjectInitializer const& ObjectInitializer)
 	: Super(ObjectInitializer)
@@ -416,25 +276,22 @@ bool UNiagaraDataInterfaceNeighborGrid3D::ShouldCompile(EShaderPlatform ShaderPl
 #endif
 
 #if WITH_EDITORONLY_DATA
+bool UNiagaraDataInterfaceNeighborGrid3D::AppendCompileHash(FNiagaraCompileHashVisitor* InVisitor) const
+{
+	bool bSuccess = Super::AppendCompileHash(InVisitor);
+	bSuccess &= InVisitor->UpdateShaderParameters<FShaderParameters>();
+	return bSuccess;
+}
+
 void UNiagaraDataInterfaceNeighborGrid3D::GetParameterDefinitionHLSL(const FNiagaraDataInterfaceGPUParamInfo& ParamInfo, FString& OutHLSL)
 {
 	Super::GetParameterDefinitionHLSL(ParamInfo, OutHLSL);
 
-	static const TCHAR *FormatDeclarations = TEXT(R"(
-		int {MaxNeighborsPerCellName};
-		Buffer<int> {ParticleNeighborsName};
-		Buffer<int> {ParticleNeighborCountName};
-		RWBuffer<int> RW{OutputParticleNeighborsName};
-		RWBuffer<int> RW{OutputParticleNeighborCountName};
-	)");
-	TMap<FString, FStringFormatArg> ArgsDeclarations = {
-		{ TEXT("MaxNeighborsPerCellName"),  ParamInfo.DataInterfaceHLSLSymbol + MaxNeighborsPerCellName },
-		{ TEXT("ParticleNeighborsName"),    ParamInfo.DataInterfaceHLSLSymbol + ParticleNeighborsName },
-		{ TEXT("ParticleNeighborCountName"),    ParamInfo.DataInterfaceHLSLSymbol + ParticleNeighborCountName },
-		{ TEXT("OutputParticleNeighborsName"),    ParamInfo.DataInterfaceHLSLSymbol + OutputParticleNeighborsName },
-		{ TEXT("OutputParticleNeighborCountName"),    ParamInfo.DataInterfaceHLSLSymbol + OutputParticleNeighborCountName },
-	};
-	OutHLSL += FString::Format(FormatDeclarations, ArgsDeclarations);
+	OutHLSL.Appendf(TEXT("int			%s%s;\n"), *ParamInfo.DataInterfaceHLSLSymbol, *MaxNeighborsPerCellName);
+	OutHLSL.Appendf(TEXT("Buffer<int>	%s%s;\n"), *ParamInfo.DataInterfaceHLSLSymbol, *ParticleNeighborsName);
+	OutHLSL.Appendf(TEXT("Buffer<int>	%s%s;\n"), *ParamInfo.DataInterfaceHLSLSymbol, *ParticleNeighborCountName);
+	OutHLSL.Appendf(TEXT("RWBuffer<int>	%s%s;\n"), *ParamInfo.DataInterfaceHLSLSymbol, *OutputParticleNeighborsName);
+	OutHLSL.Appendf(TEXT("RWBuffer<int>	%s%s;\n"), *ParamInfo.DataInterfaceHLSLSymbol, *OutputParticleNeighborCountName);
 }
 
 bool UNiagaraDataInterfaceNeighborGrid3D::GetFunctionHLSL(const FNiagaraDataInterfaceGPUParamInfo& ParamInfo, const FNiagaraDataInterfaceGeneratedFunction& FunctionInfo, int FunctionInstanceIndex, FString& OutHLSL)
@@ -444,7 +301,20 @@ bool UNiagaraDataInterfaceNeighborGrid3D::GetFunctionHLSL(const FNiagaraDataInte
 	{
 		return true;
 	}
-	else if (FunctionInfo.DefinitionName == UNiagaraDataInterfaceRWBase::NumCellsFunctionName)
+
+	TMap<FString, FStringFormatArg> FormatArgs =
+	{
+		{TEXT("FunctionName"),					FunctionInfo.InstanceName},
+		{TEXT("NumCellsName"),					ParamInfo.DataInterfaceHLSLSymbol + UNiagaraDataInterfaceRWBase::NumCellsName},
+		{TEXT("UnitToUVName"),					ParamInfo.DataInterfaceHLSLSymbol + UNiagaraDataInterfaceRWBase::UnitToUVName},
+		{TEXT("MaxNeighborsPerCellName"),		ParamInfo.DataInterfaceHLSLSymbol + MaxNeighborsPerCellName},
+		{TEXT("ParticleNeighbors"),				ParamInfo.DataInterfaceHLSLSymbol + ParticleNeighborsName},
+		{TEXT("ParticleNeighborCount"),			ParamInfo.DataInterfaceHLSLSymbol + ParticleNeighborCountName},
+		{TEXT("OutputParticleNeighbors"),		ParamInfo.DataInterfaceHLSLSymbol + OutputParticleNeighborsName},
+		{TEXT("OutputParticleNeighborCount"),	ParamInfo.DataInterfaceHLSLSymbol + OutputParticleNeighborCountName},
+	};
+
+	if (FunctionInfo.DefinitionName == UNiagaraDataInterfaceRWBase::NumCellsFunctionName)
 	{
 		static const TCHAR* FormatHLSL = TEXT(R"(
 			void {FunctionName}(out int OutNumCellsX, out int OutNumCellsY, out int OutNumCellsZ)
@@ -454,12 +324,6 @@ bool UNiagaraDataInterfaceNeighborGrid3D::GetFunctionHLSL(const FNiagaraDataInte
 				OutNumCellsZ = {NumCellsName}.z;
 			}
 		)");
-		TMap<FString, FStringFormatArg> FormatArgs =
-		{
-			{TEXT("FunctionName"), FunctionInfo.InstanceName},
-			{TEXT("NumCellsName"), ParamInfo.DataInterfaceHLSLSymbol + UNiagaraDataInterfaceRWBase::NumCellsName},
-			{TEXT("UnitToUVName"), ParamInfo.DataInterfaceHLSLSymbol + UNiagaraDataInterfaceRWBase::UnitToUVName},
-		};
 		OutHLSL += FString::Format(FormatHLSL, FormatArgs);
 	}
 	else if (FunctionInfo.DefinitionName == MaxNeighborsPerCellFunctionName)
@@ -470,12 +334,7 @@ bool UNiagaraDataInterfaceNeighborGrid3D::GetFunctionHLSL(const FNiagaraDataInte
 				Out_MaxNeighborsPerCell = {MaxNeighborsPerCellName};
 			}
 		)");
-		TMap<FString, FStringFormatArg> ArgsSample = {
-			{TEXT("FunctionName"), FunctionInfo.InstanceName},
-			{TEXT("MaxNeighborsPerCellName"), ParamInfo.DataInterfaceHLSLSymbol + MaxNeighborsPerCellName},
-
-		};
-		OutHLSL += FString::Format(FormatSample, ArgsSample);
+		OutHLSL += FString::Format(FormatSample, FormatArgs);
 		return true;
 	}
 	else if (FunctionInfo.DefinitionName == NeighborGridIndexToLinearFunctionName)
@@ -486,13 +345,7 @@ bool UNiagaraDataInterfaceNeighborGrid3D::GetFunctionHLSL(const FNiagaraDataInte
 				Out_Linear = In_Neighbor + In_IndexX * {MaxNeighborsPerCellName} + In_IndexY * {MaxNeighborsPerCellName}*{NumCellsName}.x + In_IndexZ * {MaxNeighborsPerCellName}*{NumCellsName}.x*{NumCellsName}.y;
 			}
 		)");
-		TMap<FString, FStringFormatArg> ArgsBounds = {
-			{TEXT("FunctionName"), FunctionInfo.InstanceName},
-			{TEXT("MaxNeighborsPerCellName"), ParamInfo.DataInterfaceHLSLSymbol + MaxNeighborsPerCellName},
-			{TEXT("NumCellsName"), ParamInfo.DataInterfaceHLSLSymbol + UNiagaraDataInterfaceRWBase::NumCellsName},
-			{TEXT("UnitToUVName"), ParamInfo.DataInterfaceHLSLSymbol + UNiagaraDataInterfaceRWBase::UnitToUVName},
-		};
-		OutHLSL += FString::Format(FormatBounds, ArgsBounds);
+		OutHLSL += FString::Format(FormatBounds, FormatArgs);
 		return true;
 	}
 	else if (FunctionInfo.DefinitionName == GetParticleNeighborFunctionName)
@@ -503,11 +356,7 @@ bool UNiagaraDataInterfaceNeighborGrid3D::GetFunctionHLSL(const FNiagaraDataInte
 				Out_ParticleNeighborIndex = {ParticleNeighbors}[In_Index];				
 			}
 		)");
-		TMap<FString, FStringFormatArg> ArgsBounds = {
-			{TEXT("FunctionName"), FunctionInfo.InstanceName},
-			{TEXT("ParticleNeighbors"), ParamInfo.DataInterfaceHLSLSymbol + ParticleNeighborsName},
-		};
-		OutHLSL += FString::Format(FormatBounds, ArgsBounds);
+		OutHLSL += FString::Format(FormatBounds, FormatArgs);
 		return true;
 	}
 	else if (FunctionInfo.DefinitionName == SetParticleNeighborFunctionName)
@@ -515,15 +364,11 @@ bool UNiagaraDataInterfaceNeighborGrid3D::GetFunctionHLSL(const FNiagaraDataInte
 		static const TCHAR *FormatBounds = TEXT(R"(
 			void {FunctionName}(int In_Index, int In_ParticleNeighborIndex, out int Out_Ignore)
 			{
-				RW{OutputParticleNeighbors}[In_Index] = In_ParticleNeighborIndex;				
+				{OutputParticleNeighbors}[In_Index] = In_ParticleNeighborIndex;				
 				Out_Ignore = 0;
 			}
 		)");
-		TMap<FString, FStringFormatArg> ArgsBounds = {
-			{TEXT("FunctionName"), FunctionInfo.InstanceName},
-			{TEXT("OutputParticleNeighbors"), ParamInfo.DataInterfaceHLSLSymbol + OutputParticleNeighborsName},
-		};
-		OutHLSL += FString::Format(FormatBounds, ArgsBounds);
+		OutHLSL += FString::Format(FormatBounds, FormatArgs);
 		return true;
 	}
 	else if (FunctionInfo.DefinitionName == GetParticleNeighborCountFunctionName)
@@ -534,11 +379,7 @@ bool UNiagaraDataInterfaceNeighborGrid3D::GetFunctionHLSL(const FNiagaraDataInte
 				Out_ParticleNeighborIndex = {ParticleNeighborCount}[In_Index];				
 			}
 		)");
-		TMap<FString, FStringFormatArg> ArgsBounds = {
-			{TEXT("FunctionName"), FunctionInfo.InstanceName},
-			{TEXT("ParticleNeighborCount"), ParamInfo.DataInterfaceHLSLSymbol + ParticleNeighborCountName},
-		};
-		OutHLSL += FString::Format(FormatBounds, ArgsBounds);
+		OutHLSL += FString::Format(FormatBounds, FormatArgs);
 		return true;
 	}
 	else if (FunctionInfo.DefinitionName == SetParticleNeighborCountFunctionName)
@@ -546,14 +387,10 @@ bool UNiagaraDataInterfaceNeighborGrid3D::GetFunctionHLSL(const FNiagaraDataInte
 		static const TCHAR *FormatBounds = TEXT(R"(
 			void {FunctionName}(int In_Index, int In_Increment, out int PreviousNeighborCount)
 			{				
-				InterlockedAdd(RW{OutputParticleNeighborCount}[In_Index], In_Increment, PreviousNeighborCount);				
+				InterlockedAdd({OutputParticleNeighborCount}[In_Index], In_Increment, PreviousNeighborCount);				
 			}
 		)");
-		TMap<FString, FStringFormatArg> ArgsBounds = {
-			{TEXT("FunctionName"), FunctionInfo.InstanceName},
-			{TEXT("OutputParticleNeighborCount"), ParamInfo.DataInterfaceHLSLSymbol + OutputParticleNeighborCountName},
-		};
-		OutHLSL += FString::Format(FormatBounds, ArgsBounds);
+		OutHLSL += FString::Format(FormatBounds, FormatArgs);
 		return true;
 	}
 
@@ -561,6 +398,94 @@ bool UNiagaraDataInterfaceNeighborGrid3D::GetFunctionHLSL(const FNiagaraDataInte
 }
 #endif
 
+void UNiagaraDataInterfaceNeighborGrid3D::BuildShaderParameters(FNiagaraShaderParametersBuilder& ShaderParametersBuilder) const
+{
+	ShaderParametersBuilder.AddNestedStruct<FShaderParameters>();
+}
+
+void UNiagaraDataInterfaceNeighborGrid3D::SetShaderParameters(const FNiagaraDataInterfaceSetShaderParametersContext& Context) const
+{
+	FNiagaraDataInterfaceProxyNeighborGrid3D& DIProxy = Context.GetProxy<FNiagaraDataInterfaceProxyNeighborGrid3D>();
+	NeighborGrid3DRWInstanceData* ProxyData = DIProxy.SystemInstancesToProxyData.Find(Context.GetSystemInstanceID());
+
+	FShaderParameters* ShaderParameters = Context.GetParameterNestedStruct<FShaderParameters>();
+	if (ProxyData && ProxyData->NeighborhoodBuffer.Buffer.IsValid())
+	{
+		ShaderParameters->NumCells		= ProxyData->NumCells;
+		ShaderParameters->UnitToUV		= FVector3f(1.0f) / FVector3f(ProxyData->NumCells);
+		ShaderParameters->CellSize.X	= ProxyData->WorldBBoxSize.X / float(ProxyData->NumCells.X);
+		ShaderParameters->CellSize.Y	= ProxyData->WorldBBoxSize.Y / float(ProxyData->NumCells.Y);
+		ShaderParameters->CellSize.Z	= ProxyData->WorldBBoxSize.Z / float(ProxyData->NumCells.Z);
+		ShaderParameters->WorldBBoxSize	= FVector3f(ProxyData->WorldBBoxSize);
+		ShaderParameters->MaxNeighborsPerCellValue = ProxyData->MaxNeighborsPerCell;
+
+		TArray<FRHITransitionInfo, TInlineAllocator<4>> PendingTransitions;
+		if (Context.IsOutputStage())
+		{
+			ShaderParameters->ParticleNeighbors = FNiagaraRenderer::GetDummyIntBuffer();
+			ShaderParameters->ParticleNeighborCount = FNiagaraRenderer::GetDummyIntBuffer();
+
+			if (Context.IsResourceBound(&ShaderParameters->OutputParticleNeighbors))
+			{
+				PendingTransitions.Emplace(ProxyData->NeighborhoodBuffer.UAV, ERHIAccess::Unknown, ERHIAccess::UAVCompute);
+				ShaderParameters->OutputParticleNeighbors = ProxyData->NeighborhoodBuffer.UAV;
+			}
+			if (Context.IsResourceBound(&ShaderParameters->OutputParticleNeighborCount))
+			{
+				PendingTransitions.Emplace(ProxyData->NeighborhoodCountBuffer.UAV, ERHIAccess::Unknown, ERHIAccess::UAVCompute);
+				ShaderParameters->OutputParticleNeighborCount = ProxyData->NeighborhoodCountBuffer.UAV;
+			}
+		}
+		else
+		{
+			if (Context.IsResourceBound(&ShaderParameters->ParticleNeighbors))
+			{
+				PendingTransitions.Emplace(ProxyData->NeighborhoodBuffer.UAV, ERHIAccess::Unknown, ERHIAccess::SRVCompute);
+				ShaderParameters->ParticleNeighbors = ProxyData->NeighborhoodBuffer.SRV;
+			}
+
+			if (Context.IsResourceBound(&ShaderParameters->ParticleNeighborCount))
+			{
+				PendingTransitions.Emplace(ProxyData->NeighborhoodCountBuffer.UAV, ERHIAccess::Unknown, ERHIAccess::SRVCompute);
+				ShaderParameters->ParticleNeighborCount = ProxyData->NeighborhoodCountBuffer.SRV;
+			}
+
+			if (Context.IsResourceBound(&ShaderParameters->OutputParticleNeighbors))
+			{
+				ShaderParameters->OutputParticleNeighbors = Context.GetComputeDispatchInterface().GetEmptyUAVFromPool(Context.GetRHICmdList(), PF_R32_SINT, ENiagaraEmptyUAVType::Buffer);
+			}
+			if (Context.IsResourceBound(&ShaderParameters->OutputParticleNeighborCount))
+			{
+				ShaderParameters->OutputParticleNeighborCount = Context.GetComputeDispatchInterface().GetEmptyUAVFromPool(Context.GetRHICmdList(), PF_R32_SINT, ENiagaraEmptyUAVType::Buffer);
+			}
+		}
+
+		if (PendingTransitions.Num() > 0)
+		{
+			Context.GetRHICmdList().Transition(PendingTransitions);
+		}
+	}
+	else
+	{
+		ShaderParameters->NumCells		= FIntVector::ZeroValue;
+		ShaderParameters->UnitToUV		= FVector3f::ZeroVector;
+		ShaderParameters->CellSize		= FVector3f::ZeroVector;
+		ShaderParameters->WorldBBoxSize	= FVector3f::ZeroVector;
+
+		ShaderParameters->MaxNeighborsPerCellValue	= 0;
+		ShaderParameters->ParticleNeighbors			= FNiagaraRenderer::GetDummyIntBuffer();
+		ShaderParameters->ParticleNeighborCount		= FNiagaraRenderer::GetDummyIntBuffer();
+		if (Context.IsResourceBound(&ShaderParameters->OutputParticleNeighbors))
+		{
+			ShaderParameters->OutputParticleNeighbors = Context.GetComputeDispatchInterface().GetEmptyUAVFromPool(Context.GetRHICmdList(), PF_R32_SINT, ENiagaraEmptyUAVType::Buffer);
+		}
+		if (Context.IsResourceBound(&ShaderParameters->OutputParticleNeighborCount))
+		{
+			ShaderParameters->OutputParticleNeighborCount = Context.GetComputeDispatchInterface().GetEmptyUAVFromPool(Context.GetRHICmdList(), PF_R32_SINT, ENiagaraEmptyUAVType::Buffer);
+		}
+	}
+}
+
 bool UNiagaraDataInterfaceNeighborGrid3D::InitPerInstanceData(void* PerInstanceData, FNiagaraSystemInstance* SystemInstance)
 {
 
diff --git a/Engine/Plugins/FX/Niagara/Source/Niagara/Private/NiagaraDataInterfaceSpline.cpp b/Engine/Plugins/FX/Niagara/Source/Niagara/Private/NiagaraDataInterfaceSpline.cpp
--- a/Engine/Plugins/FX/Niagara/Source/Niagara/Private/NiagaraDataInterfaceSpline.cpp
+++ b/Engine/Plugins/FX/Niagara/Source/Niagara/Private/NiagaraDataInterfaceSpline.cpp
@@ -3,6 +3,7 @@
 #include "NiagaraDataInterfaceSpline.h"
 #include "NiagaraEmitterInstance.h"
 #include "NiagaraComponent.h"
+#include "NiagaraShaderParametersBuilder.h"
 #include "NiagaraSystemInstance.h"
 #include "Internationalization/Internationalization.h"
 #include "ShaderParameterUtils.h"
@@ -24,7 +25,23 @@ struct FNiagaraSplineDIFunctionVersion
 
 namespace NDISplineLocal
 {
-	static const TCHAR*		TemplateShaderFile = TEXT("/Plugin/FX/Niagara/Private/NiagaraDataInterfaceSplineTemplate.ush");
+	BEGIN_SHADER_PARAMETER_STRUCT(FShaderParameters, )
+		SHADER_PARAMETER(FMatrix44f,	SplineTransform)
+		SHADER_PARAMETER(FMatrix44f,	SplineTransformRotationMat)
+		SHADER_PARAMETER(FMatrix44f,	SplineTransformInverseTranspose)
+		SHADER_PARAMETER(FQuat4f,		SplineTransformRotation)
+		SHADER_PARAMETER(FVector3f,		DefaultUpVector)
+		SHADER_PARAMETER(float,			SplineLength)
+		SHADER_PARAMETER(float,			SplineDistanceStep)
+		SHADER_PARAMETER(float,			InvSplineDistanceStep)
+		SHADER_PARAMETER(int,			MaxIndex)
+
+		SHADER_PARAMETER_SRV(Buffer<float4>, SplinePositionsLUT)
+		SHADER_PARAMETER_SRV(Buffer<float4>, SplineScalesLUT)
+		SHADER_PARAMETER_SRV(Buffer<float4>, SplineRotationsLUT)
+	END_SHADER_PARAMETER_STRUCT()
+
+	static const TCHAR* TemplateShaderFile = TEXT("/Plugin/FX/Niagara/Private/NiagaraDataInterfaceSplineTemplate.ush");
 
 	static const FName SampleSplinePositionByUnitDistanceName("SampleSplinePositionByUnitDistance");
 	static const FName SampleSplinePositionByUnitDistanceWSName("SampleSplinePositionByUnitDistanceWS");
@@ -427,23 +444,7 @@ void UNiagaraDataInterfaceSpline::GetCommonHLSL(FString& OutHLSL)
 void UNiagaraDataInterfaceSpline::GetParameterDefinitionHLSL(const FNiagaraDataInterfaceGPUParamInfo& ParamInfo, FString& OutHLSL)
 {
 	const TMap<FString, FStringFormatArg> TemplateArgs = {
-		{TEXT("ParameterHLSLSymbol"), ParamInfo.DataInterfaceHLSLSymbol},
-				
-		{TEXT("SplineTransform"), TEXT("SplineTransform_") + ParamInfo.DataInterfaceHLSLSymbol},
-		{TEXT("SplineTransformRotationMat"), TEXT("SplineTransformRotationMat_") + ParamInfo.DataInterfaceHLSLSymbol},
-		{TEXT("SplineTransformInverseTranspose"), TEXT("SplineTransformInverseTranspose_") + ParamInfo.DataInterfaceHLSLSymbol},
-		{TEXT("SplineTransformRotation"), TEXT("SplineTransformRotation_") + ParamInfo.DataInterfaceHLSLSymbol},
-				
-		{TEXT("DefaultUpVector"), TEXT("DefaultUpVector_") + ParamInfo.DataInterfaceHLSLSymbol},		
-				
-		{TEXT("SplineLength"), TEXT("SplineLength_") + ParamInfo.DataInterfaceHLSLSymbol},
-		{TEXT("SplineDistanceStep"), TEXT("SplineDistanceStep_") + ParamInfo.DataInterfaceHLSLSymbol},
-		{TEXT("InvSplineDistanceStep"), TEXT("InvSplineDistanceStep_") + ParamInfo.DataInterfaceHLSLSymbol},
-		{TEXT("MaxIndex"), TEXT("MaxIndex_") + ParamInfo.DataInterfaceHLSLSymbol},
-						
-		{TEXT("SplinePositionsLUT"), TEXT("SplinePositionsLUT_") + ParamInfo.DataInterfaceHLSLSymbol},
-		{TEXT("SplineScalesLUT"), TEXT("SplineScalesLUT_") + ParamInfo.DataInterfaceHLSLSymbol},
-		{TEXT("SplineRotationsLUT"), TEXT("SplineRotationsLUT_") + ParamInfo.DataInterfaceHLSLSymbol},
+		{TEXT("ParameterName"), ParamInfo.DataInterfaceHLSLSymbol},
 	};
 	
 	FString TemplateFile;
@@ -477,10 +478,52 @@ bool UNiagaraDataInterfaceSpline::AppendCompileHash(FNiagaraCompileHashVisitor*
 	bool bSuccess = Super::AppendCompileHash(InVisitor);
 	FSHAHash Hash = GetShaderFileHash(NDISplineLocal::TemplateShaderFile, EShaderPlatform::SP_PCD3D_SM5);
 	InVisitor->UpdateString(TEXT("NiagaraDataInterfaceExportTemplateHLSLSource"), Hash.ToString());
+	InVisitor->UpdateShaderParameters<NDISplineLocal::FShaderParameters>();
 	return bSuccess;
 }
 #endif
 
+void UNiagaraDataInterfaceSpline::BuildShaderParameters(FNiagaraShaderParametersBuilder& ShaderParametersBuilder) const
+{
+	ShaderParametersBuilder.AddNestedStruct<NDISplineLocal::FShaderParameters>();
+}
+
+void UNiagaraDataInterfaceSpline::SetShaderParameters(const FNiagaraDataInterfaceSetShaderParametersContext& Context) const
+{
+	FNiagaraDataInterfaceProxySpline& DIProxy = Context.GetProxy<FNiagaraDataInterfaceProxySpline>();
+	NDISplineLocal::FShaderParameters* ShaderParameters = Context.GetParameterNestedStruct<NDISplineLocal::FShaderParameters>();
+	if (FNDISpline_InstanceData_RenderThread* InstanceData_RT = DIProxy.SystemInstancesToProxyData_RT.Find(Context.GetSystemInstanceID()))
+	{
+		ShaderParameters->SplineTransform					= InstanceData_RT->SplineTransform;
+		ShaderParameters->SplineTransformRotationMat		= InstanceData_RT->SplineTransformRotationMat;
+		ShaderParameters->SplineTransformInverseTranspose	= InstanceData_RT->SplineTransformInverseTranspose;
+		ShaderParameters->SplineTransformRotation			= InstanceData_RT->SplineTransformRotation;
+		ShaderParameters->DefaultUpVector					= InstanceData_RT->DefaultUpVector;
+		ShaderParameters->SplineLength						= InstanceData_RT->SplineLength;
+		ShaderParameters->SplineDistanceStep				= InstanceData_RT->SplineDistanceStep;
+		ShaderParameters->InvSplineDistanceStep				= InstanceData_RT->InvSplineDistanceStep;
+		ShaderParameters->MaxIndex							= InstanceData_RT->MaxIndex;
+		ShaderParameters->SplinePositionsLUT				= InstanceData_RT->SplinePositionsLUT.SRV;
+		ShaderParameters->SplineScalesLUT					= InstanceData_RT->SplineScalesLUT.SRV;
+		ShaderParameters->SplineRotationsLUT				= InstanceData_RT->SplineRotationsLUT.SRV;
+	}
+	else
+	{
+		ShaderParameters->SplineTransform					= FMatrix44f::Identity;
+		ShaderParameters->SplineTransformRotationMat		= FMatrix44f::Identity;
+		ShaderParameters->SplineTransformInverseTranspose	= FMatrix44f::Identity;
+		ShaderParameters->SplineTransformRotation			= FQuat4f::Identity;
+		ShaderParameters->DefaultUpVector					= FVector3f::UnitZ();
+		ShaderParameters->SplineLength						= 0.0f;
+		ShaderParameters->SplineDistanceStep				= 0.0f;
+		ShaderParameters->InvSplineDistanceStep				= 0.0f;
+		ShaderParameters->MaxIndex							= 0;
+		ShaderParameters->SplinePositionsLUT				= FNiagaraRenderer::GetDummyFloat4Buffer();
+		ShaderParameters->SplineScalesLUT					= FNiagaraRenderer::GetDummyFloat4Buffer();
+		ShaderParameters->SplineRotationsLUT				= FNiagaraRenderer::GetDummyFloat4Buffer();
+	}
+}
+
 bool UNiagaraDataInterfaceSpline::Equals(const UNiagaraDataInterface* Other) const
 {
 	if (!Super::Equals(Other))
@@ -714,77 +757,6 @@ bool UNiagaraDataInterfaceSpline::PerInstanceTick(void* PerInstanceData, FNiagar
 	return false;
 }
 
-struct FNiagaraDataInterfaceParametersCS_Spline : public FNiagaraDataInterfaceParametersCS
-{
-	DECLARE_TYPE_LAYOUT(FNiagaraDataInterfaceParametersCS_Spline, NonVirtual);
-
-	void Bind(const FNiagaraDataInterfaceGPUParamInfo& ParameterInfo, const class FShaderParameterMap& ParameterMap)
-	{
-		SplineTransform.Bind(ParameterMap, *(TEXT("SplineTransform_") + ParameterInfo.DataInterfaceHLSLSymbol));
-		SplineTransformRotationMat.Bind(ParameterMap, *(TEXT("SplineTransformRotationMat_") + ParameterInfo.DataInterfaceHLSLSymbol));
-		SplineTransformInverseTranspose.Bind(ParameterMap, *(TEXT("SplineTransformInverseTranspose_") + ParameterInfo.DataInterfaceHLSLSymbol));
-		SplineTransformRotation.Bind(ParameterMap, *(TEXT("SplineTransformRotation_") + ParameterInfo.DataInterfaceHLSLSymbol));
-	
-		DefaultUpVector.Bind(ParameterMap, *(TEXT("DefaultUpVector_") + ParameterInfo.DataInterfaceHLSLSymbol));
-	
-		SplineLength.Bind(ParameterMap, *(TEXT("SplineLength_") + ParameterInfo.DataInterfaceHLSLSymbol));
-		SplineDistanceStep.Bind(ParameterMap, *(TEXT("SplineDistanceStep_") + ParameterInfo.DataInterfaceHLSLSymbol));
-		InvSplineDistanceStep.Bind(ParameterMap, *(TEXT("InvSplineDistanceStep_") + ParameterInfo.DataInterfaceHLSLSymbol));
-		MaxIndex.Bind(ParameterMap, *(TEXT("MaxIndex_") + ParameterInfo.DataInterfaceHLSLSymbol));
-	
-		SplinePositionsLUT.Bind(ParameterMap, *(TEXT("SplinePositionsLUT_") + ParameterInfo.DataInterfaceHLSLSymbol));
-		SplineScalesLUT.Bind(ParameterMap, *(TEXT("SplineScalesLUT_") + ParameterInfo.DataInterfaceHLSLSymbol));
-		SplineRotationsLUT.Bind(ParameterMap, *(TEXT("SplineRotationsLUT_") + ParameterInfo.DataInterfaceHLSLSymbol));
-	}
-	
-	void Set(FRHICommandList& RHICmdList, const FNiagaraDataInterfaceSetArgs& Context) const
-	{
-		check(IsInRenderingThread());
-
-		FRHIComputeShader* ComputeShaderRHI = RHICmdList.GetBoundComputeShader();
-		FNiagaraDataInterfaceProxySpline* RT_Proxy = (FNiagaraDataInterfaceProxySpline*)Context.DataInterface;
-
-	
-		if (FNDISpline_InstanceData_RenderThread* Instance_RT_Proxy = RT_Proxy->SystemInstancesToProxyData_RT.Find(Context.SystemInstanceID))
-		{
-			SetShaderValue(RHICmdList, ComputeShaderRHI, SplineTransform, Instance_RT_Proxy->SplineTransform);
-			SetShaderValue(RHICmdList, ComputeShaderRHI, SplineTransformRotationMat, Instance_RT_Proxy->SplineTransformRotationMat);
-			SetShaderValue(RHICmdList, ComputeShaderRHI, SplineTransformInverseTranspose, Instance_RT_Proxy->SplineTransformInverseTranspose);
-			SetShaderValue(RHICmdList, ComputeShaderRHI, SplineTransformRotation, Instance_RT_Proxy->SplineTransformRotation);
-	
-			SetShaderValue(RHICmdList, ComputeShaderRHI, DefaultUpVector, Instance_RT_Proxy->DefaultUpVector);
-	
-			SetShaderValue(RHICmdList, ComputeShaderRHI, SplineLength, Instance_RT_Proxy->SplineLength);
-			SetShaderValue(RHICmdList, ComputeShaderRHI, SplineDistanceStep, Instance_RT_Proxy->SplineDistanceStep);
-			SetShaderValue(RHICmdList, ComputeShaderRHI, InvSplineDistanceStep, Instance_RT_Proxy->InvSplineDistanceStep);
-			SetShaderValue(RHICmdList, ComputeShaderRHI, MaxIndex, Instance_RT_Proxy->MaxIndex);
-	
-			SetSRVParameter(RHICmdList, ComputeShaderRHI, SplinePositionsLUT, Instance_RT_Proxy->SplinePositionsLUT.SRV);
-			SetSRVParameter(RHICmdList, ComputeShaderRHI, SplineScalesLUT, Instance_RT_Proxy->SplineScalesLUT.SRV);
-			SetSRVParameter(RHICmdList, ComputeShaderRHI, SplineRotationsLUT, Instance_RT_Proxy->SplineRotationsLUT.SRV);
-		}
-	}
-
-	LAYOUT_FIELD(FShaderParameter, SplineTransform);
-	LAYOUT_FIELD(FShaderParameter, SplineTransformRotationMat);
-	LAYOUT_FIELD(FShaderParameter, SplineTransformInverseTranspose);
-	LAYOUT_FIELD(FShaderParameter, SplineTransformRotation);
-	
-	LAYOUT_FIELD(FShaderParameter, DefaultUpVector);
-	
-	LAYOUT_FIELD(FShaderParameter, SplineLength);
-	LAYOUT_FIELD(FShaderParameter, SplineDistanceStep);
-	LAYOUT_FIELD(FShaderParameter, InvSplineDistanceStep);
-	LAYOUT_FIELD(FShaderParameter, MaxIndex);
-	
-	LAYOUT_FIELD(FShaderResourceParameter, SplinePositionsLUT);
-	LAYOUT_FIELD(FShaderResourceParameter, SplineScalesLUT);
-	LAYOUT_FIELD(FShaderResourceParameter, SplineRotationsLUT);
-};
-
-IMPLEMENT_TYPE_LAYOUT(FNiagaraDataInterfaceParametersCS_Spline);
-IMPLEMENT_NIAGARA_DI_PARAMETER(UNiagaraDataInterfaceSpline, FNiagaraDataInterfaceParametersCS_Spline);
-
 void FNiagaraDataInterfaceSplineLUT::BuildLUT(const FSplineCurves& SplineCurves, int32 NumSteps)
 {
 	Positions.Empty(NumSteps);
diff --git a/Engine/Source/Runtime/Core/Private/UObject/UE5MainStreamObjectVersion.cpp b/Engine/Source/Runtime/Core/Private/UObject/UE5MainStreamObjectVersion.cpp
--- a/Engine/Source/Runtime/Core/Private/UObject/UE5MainStreamObjectVersion.cpp
+++ b/Engine/Source/Runtime/Core/Private/UObject/UE5MainStreamObjectVersion.cpp
@@ -14,7 +14,7 @@ TMap<FGuid, FGuid> FUE5MainStreamObjectVersion::GetSystemGuids()
 	SystemGuids.Add(DevGuids.MATERIALSHADERMAP_DERIVEDDATA_VER, FGuid("65C01D817C9A4EEFAE9E988D41A1F3DD"));
 	SystemGuids.Add(DevGuids.NANITE_DERIVEDDATA_VER, FGuid("7A3EA93546624597BEBE2ADFDE3A7628"));
 	SystemGuids.Add(DevGuids.NIAGARASHADERMAP_DERIVEDDATA_VER, FGuid("8A37C45D24F2423CBE5F8F371DE33575"));
-	SystemGuids.Add(DevGuids.Niagara_LatestScriptCompileVersion, FGuid("7E5282F3B014E540BAA1E77751460B15"));
+	SystemGuids.Add(DevGuids.Niagara_LatestScriptCompileVersion, FGuid("D40FAE8F632C6B469CEE5F17C49C5BBE"));
 	SystemGuids.Add(DevGuids.SkeletalMeshDerivedDataVersion, FGuid("FB7B2FB546DE49B2A2F88B7F69763CEE"));
 	SystemGuids.Add(DevGuids.STATICMESH_DERIVEDDATA_VER, FGuid("98F7D79A5811013825E75BBCC41ED3E9"));
 	SystemGuids.Add(DevGuids.POSESEARCHDB_DERIVEDDATA_VER, FGuid("4E595C2AC5E947D6BA9ABC874353E5BC"));
