- Fix not clearing the transition flag & fix missing post stage for 2DArray

#rnx

diff --git a/Engine/Plugins/FX/Niagara/Source/Niagara/Classes/NiagaraDataInterfaceRenderTarget2DArray.h b/Engine/Plugins/FX/Niagara/Source/Niagara/Classes/NiagaraDataInterfaceRenderTarget2DArray.h
--- a/Engine/Plugins/FX/Niagara/Source/Niagara/Classes/NiagaraDataInterfaceRenderTarget2DArray.h
+++ b/Engine/Plugins/FX/Niagara/Source/Niagara/Classes/NiagaraDataInterfaceRenderTarget2DArray.h
@@ -66,6 +66,7 @@ struct FNiagaraDataInterfaceProxyRenderTarget2DArrayProxy : public FNiagaraDataI
 	}
 
 	virtual void ClearBuffers(FRHICommandList& RHICmdList) {}
+	virtual void PostStage(FRHICommandList& RHICmdList, const FNiagaraDataInterfaceStageArgs& Context) override;
 	virtual void PostSimulate(FRHICommandList& RHICmdList, const FNiagaraDataInterfaceArgs& Context) override;
 
 	virtual FIntVector GetElementCount(FNiagaraSystemInstanceID SystemInstanceID) const override;
diff --git a/Engine/Plugins/FX/Niagara/Source/Niagara/Private/NiagaraDataInterfaceRenderTarget2D.cpp b/Engine/Plugins/FX/Niagara/Source/Niagara/Private/NiagaraDataInterfaceRenderTarget2D.cpp
--- a/Engine/Plugins/FX/Niagara/Source/Niagara/Private/NiagaraDataInterfaceRenderTarget2D.cpp
+++ b/Engine/Plugins/FX/Niagara/Source/Niagara/Private/NiagaraDataInterfaceRenderTarget2D.cpp
@@ -837,6 +837,8 @@ void FNiagaraDataInterfaceProxyRenderTarget2DProxy::PostStage(FRHICommandList& R
 	{
 		if (ProxyData->bNeedsTransition)
 		{
+			ProxyData->bNeedsTransition = false;
+
 			if (FRHIUnorderedAccessView* OutputUAV = ProxyData->UnorderedAccessViewRHI)
 			{
 				// FIXME: move to FNiagaraDataInterfaceProxyRenderTarget2DProxy::PostStage, same as for the transition in Set() above.
diff --git a/Engine/Plugins/FX/Niagara/Source/Niagara/Private/NiagaraDataInterfaceRenderTarget2DArray.cpp b/Engine/Plugins/FX/Niagara/Source/Niagara/Private/NiagaraDataInterfaceRenderTarget2DArray.cpp
--- a/Engine/Plugins/FX/Niagara/Source/Niagara/Private/NiagaraDataInterfaceRenderTarget2DArray.cpp
+++ b/Engine/Plugins/FX/Niagara/Source/Niagara/Private/NiagaraDataInterfaceRenderTarget2DArray.cpp
@@ -680,6 +680,23 @@ bool UNiagaraDataInterfaceRenderTarget2DArray::PerInstanceTickPostSimulate(void*
 	return false;
 }
 
+void FNiagaraDataInterfaceProxyRenderTarget2DArrayProxy::PostStage(FRHICommandList& RHICmdList, const FNiagaraDataInterfaceStageArgs& Context)
+{
+	if (FRenderTarget2DArrayRWInstanceData_RenderThread* ProxyData = SystemInstancesToProxyData_RT.Find(Context.SystemInstanceID))
+	{
+		if (ProxyData->bNeedsTransition)
+		{
+			ProxyData->bNeedsTransition = false;
+
+			if (FRHIUnorderedAccessView* OutputUAV = ProxyData->UnorderedAccessViewRHI)
+			{
+				// FIXME: move to FNiagaraDataInterfaceProxyRenderTargetVolumeProxy::PostStage, same as for the transition in Set() above.
+				RHICmdList.Transition(FRHITransitionInfo(OutputUAV, ERHIAccess::UAVCompute, ERHIAccess::SRVMask));
+			}
+		}
+	}
+}
+
 void FNiagaraDataInterfaceProxyRenderTarget2DArrayProxy::PostSimulate(FRHICommandList& RHICmdList, const FNiagaraDataInterfaceArgs& Context)
 {
 	FRenderTarget2DArrayRWInstanceData_RenderThread* ProxyData = SystemInstancesToProxyData_RT.Find(Context.SystemInstanceID);
diff --git a/Engine/Plugins/FX/Niagara/Source/Niagara/Private/NiagaraDataInterfaceRenderTargetCube.cpp b/Engine/Plugins/FX/Niagara/Source/Niagara/Private/NiagaraDataInterfaceRenderTargetCube.cpp
--- a/Engine/Plugins/FX/Niagara/Source/Niagara/Private/NiagaraDataInterfaceRenderTargetCube.cpp
+++ b/Engine/Plugins/FX/Niagara/Source/Niagara/Private/NiagaraDataInterfaceRenderTargetCube.cpp
@@ -662,6 +662,8 @@ void FNiagaraDataInterfaceProxyRenderTargetCubeProxy::PostStage(FRHICommandList&
 	{
 		if (ProxyData->bNeedsTransition)
 		{
+			ProxyData->bNeedsTransition = false;
+
 			if (FRHIUnorderedAccessView* OutputUAV = ProxyData->UnorderedAccessViewRHI)
 			{
 				// FIXME: move to FNiagaraDataInterfaceProxyRenderTargetCubeProxy::PostStage, same as for the transition in Set() above.
diff --git a/Engine/Plugins/FX/Niagara/Source/Niagara/Private/NiagaraDataInterfaceRenderTargetVolume.cpp b/Engine/Plugins/FX/Niagara/Source/Niagara/Private/NiagaraDataInterfaceRenderTargetVolume.cpp
--- a/Engine/Plugins/FX/Niagara/Source/Niagara/Private/NiagaraDataInterfaceRenderTargetVolume.cpp
+++ b/Engine/Plugins/FX/Niagara/Source/Niagara/Private/NiagaraDataInterfaceRenderTargetVolume.cpp
@@ -681,6 +681,8 @@ void FNiagaraDataInterfaceProxyRenderTargetVolumeProxy::PostStage(FRHICommandLis
 	{
 		if (ProxyData->bNeedsTransition)
 		{
+			ProxyData->bNeedsTransition = false;
+
 			if (FRHIUnorderedAccessView* OutputUAV = ProxyData->UnorderedAccessViewRHI)
 			{
 				// FIXME: move to FNiagaraDataInterfaceProxyRenderTargetVolumeProxy::PostStage, same as for the transition in Set() above.
