- Fix for SamplerState not being bound post shader parameters conversion

#rnx

diff --git a/Engine/Plugins/FX/Niagara/Source/Niagara/Private/NiagaraDataInterfaceGrid2DCollection.cpp b/Engine/Plugins/FX/Niagara/Source/Niagara/Private/NiagaraDataInterfaceGrid2DCollection.cpp
--- a/Engine/Plugins/FX/Niagara/Source/Niagara/Private/NiagaraDataInterfaceGrid2DCollection.cpp
+++ b/Engine/Plugins/FX/Niagara/Source/Niagara/Private/NiagaraDataInterfaceGrid2DCollection.cpp
@@ -31,13 +31,13 @@ BEGIN_SHADER_PARAMETER_STRUCT(FNDIGrid2DShaderParameters, )
 	SHADER_PARAMETER(FVector2f,						WorldBBoxSize)
 
 	SHADER_PARAMETER_SRV(Texture2DArray<float>,		Grid)
-	SHADER_PARAMETER_SAMPLER(SamplerState,			SamplerState)
+	SHADER_PARAMETER_SAMPLER(SamplerState,			GridSampler)
 	SHADER_PARAMETER_UAV(RWTexture2DArray<float>,	OutputGrid)
 END_SHADER_PARAMETER_STRUCT()
 
 const FString UNiagaraDataInterfaceGrid2DCollection::GridName(TEXT("_Grid"));
 const FString UNiagaraDataInterfaceGrid2DCollection::OutputGridName(TEXT("_OutputGrid"));
-const FString UNiagaraDataInterfaceGrid2DCollection::SamplerName(TEXT("_Sampler"));
+const FString UNiagaraDataInterfaceGrid2DCollection::SamplerName(TEXT("_GridSampler"));
 
 const FName UNiagaraDataInterfaceGrid2DCollection::SetNumCellsFunctionName("SetNumCells");
 
@@ -1478,8 +1478,9 @@ void UNiagaraDataInterfaceGrid2DCollection::SetShaderParameters(const FNiagaraDa
 	Parameters->NumCells = ProxyData->NumCells;
 	Parameters->CellSize = FVector2f(ProxyData->CellSize);
 	Parameters->WorldBBoxSize = FVector2f(ProxyData->WorldBBoxSize);	// LWC_TODO: Precision loss?
-	Parameters->SamplerState = TStaticSamplerState<SF_Bilinear, AM_Clamp, AM_Clamp, AM_Clamp>::GetRHI();
 	Parameters->Grid = ProxyData->CurrentData ? ProxyData->CurrentData->GridSRV.GetReference() : FNiagaraRenderer::GetDummyTextureReadBuffer2DArray();
+	Parameters->GridSampler = TStaticSamplerState<SF_Bilinear, AM_Clamp, AM_Clamp, AM_Clamp>::GetRHI();
+
 	if (Context.IsResourceBound(&Parameters->OutputGrid))
 	{
 		if (Context.IsOutputStage() && ProxyData->DestinationData)
diff --git a/Engine/Plugins/FX/Niagara/Source/Niagara/Private/NiagaraDataInterfaceGrid3DCollection.cpp b/Engine/Plugins/FX/Niagara/Source/Niagara/Private/NiagaraDataInterfaceGrid3DCollection.cpp
--- a/Engine/Plugins/FX/Niagara/Source/Niagara/Private/NiagaraDataInterfaceGrid3DCollection.cpp
+++ b/Engine/Plugins/FX/Niagara/Source/Niagara/Private/NiagaraDataInterfaceGrid3DCollection.cpp
@@ -35,7 +35,7 @@ BEGIN_SHADER_PARAMETER_STRUCT(FNDIGrid3DShaderParameters, )
 	SHADER_PARAMETER(FVector3f,						WorldBBoxSize)
 
 	SHADER_PARAMETER_SRV(Texture3D<float>,			Grid)
-	SHADER_PARAMETER_SAMPLER(SamplerState,			SamplerState)
+	SHADER_PARAMETER_SAMPLER(SamplerState,			GridSampler)
 	SHADER_PARAMETER_UAV(RWTexture3D<float>,		OutputGrid)
 	SHADER_PARAMETER_SRV(Buffer<float4>,			PerAttributeData)
 END_SHADER_PARAMETER_STRUCT()
@@ -47,7 +47,7 @@ const FString UNiagaraDataInterfaceGrid3DCollection::UnitClampMaxName(TEXT("_Uni
 
 const FString UNiagaraDataInterfaceGrid3DCollection::GridName(TEXT("_Grid"));
 const FString UNiagaraDataInterfaceGrid3DCollection::OutputGridName(TEXT("_OutputGrid"));
-const FString UNiagaraDataInterfaceGrid3DCollection::SamplerName(TEXT("_Sampler"));
+const FString UNiagaraDataInterfaceGrid3DCollection::SamplerName(TEXT("_GridSampler"));
 
 const FName UNiagaraDataInterfaceGrid3DCollection::SetNumCellsFunctionName("SetNumCells");
 
@@ -1599,8 +1599,9 @@ void UNiagaraDataInterfaceGrid3DCollection::SetShaderParameters(const FNiagaraDa
 	Parameters->UnitClampMax		= FVector3f::OneVector - HalfPixelOffset;
 	Parameters->WorldBBoxSize		= FVector3f(ProxyData->WorldBBoxSize);
 
-	Parameters->SamplerState = TStaticSamplerState<SF_Bilinear, AM_Clamp, AM_Clamp, AM_Clamp>::GetRHI();
 	Parameters->Grid = ProxyData->CurrentData ? ProxyData->CurrentData->GridBuffer.SRV.GetReference() : FNiagaraRenderer::GetDummyTextureReadBuffer3D();
+	Parameters->GridSampler = TStaticSamplerState<SF_Bilinear, AM_Clamp, AM_Clamp, AM_Clamp>::GetRHI();
+
 	if (Context.IsResourceBound(&Parameters->OutputGrid))
 	{
 		if (Context.IsOutputStage() && ProxyData->DestinationData)
