Fix for potential access of stale/GC'd actor component when using the ActorComponent DI

diff --git a/Engine/Plugins/FX/Niagara/Source/Niagara/Private/DataInterface/NiagaraDataInterfaceActorComponent.cpp b/Engine/Plugins/FX/Niagara/Source/Niagara/Private/DataInterface/NiagaraDataInterfaceActorComponent.cpp
--- a/Engine/Plugins/FX/Niagara/Source/Niagara/Private/DataInterface/NiagaraDataInterfaceActorComponent.cpp
+++ b/Engine/Plugins/FX/Niagara/Source/Niagara/Private/DataInterface/NiagaraDataInterfaceActorComponent.cpp
@@ -42,6 +42,11 @@ namespace NDIActorComponentLocal
 		FNiagaraParameterDirectBinding<UObject*>	UserParamBinding;
 		bool										bCachedValid = false;
 		FTransform									CachedTransform = FTransform::Identity;
+
+		// our use of UserParamBinding can occur within CalculateTickGroup which occurs before we tick our parameter stores.  This
+		// can lead to a stale UObject reference being accessed (if the actor we're pointing at is deleted).  For now we cache
+		// the results during PerInstanceTick and re-use the result (if it remains valid) for calculating the tick group.
+		TWeakObjectPtr<UActorComponent>				CachedActorForCalcTickGroup;
 	};
 
 	struct FGameToRenderInstanceData
@@ -240,6 +245,7 @@ bool UNiagaraDataInterfaceActorComponent::InitPerInstanceData(void* PerInstanceD
 
 	FInstanceData_GameThread* InstanceData = new (PerInstanceData) FInstanceData_GameThread;
 	InstanceData->UserParamBinding.Init(SystemInstance->GetInstanceParameters(), ActorOrComponentParameter.Parameter);
+	InstanceData->CachedActorForCalcTickGroup = ResolveComponent(InstanceData);
 
 	return true;
 }
@@ -278,7 +284,9 @@ bool UNiagaraDataInterfaceActorComponent::PerInstanceTick(void* PerInstanceData,
 
 	InstanceData->bCachedValid = false;
 	InstanceData->CachedTransform = FTransform::Identity;
-	if ( UActorComponent* ActorComponent = ResolveComponent(PerInstanceData) )
+	UActorComponent* ActorComponent = ResolveComponent(PerInstanceData);
+
+	if (ActorComponent)
 	{
 		if (USceneComponent* SceneComponent = Cast<USceneComponent>(ActorComponent) )
 		{
@@ -294,6 +302,11 @@ bool UNiagaraDataInterfaceActorComponent::PerInstanceTick(void* PerInstanceData,
 		}
 	}
 
+	if (bRequireCurrentFrameData)
+	{
+		InstanceData->CachedActorForCalcTickGroup = ActorComponent;
+	}
+
 	return false;
 }
 
@@ -306,7 +319,9 @@ ETickingGroup UNiagaraDataInterfaceActorComponent::CalculateTickGroup(const void
 {
 	if ( bRequireCurrentFrameData )
 	{
-		if (UActorComponent* ActorComponent = ResolveComponent(PerInstanceData))
+		const NDIActorComponentLocal::FInstanceData_GameThread* InstanceData = (NDIActorComponentLocal::FInstanceData_GameThread*)PerInstanceData;
+		const UActorComponent* ActorComponent = InstanceData ? InstanceData->CachedActorForCalcTickGroup.Get() : nullptr;
+		if (ActorComponent)
 		{
 			ETickingGroup FinalTickGroup = FMath::Max(ActorComponent->PrimaryComponentTick.TickGroup, ActorComponent->PrimaryComponentTick.EndTickGroup);
 			//-TODO: Do we need to do this?
