- Ensure Emitter version data is ready as graph post load can use the emitter before it has been PostLoaded

#rnx

diff --git a/Engine/Plugins/FX/Niagara/Source/NiagaraEditor/Private/NiagaraEditorUtilities.cpp b/Engine/Plugins/FX/Niagara/Source/NiagaraEditor/Private/NiagaraEditorUtilities.cpp
--- a/Engine/Plugins/FX/Niagara/Source/NiagaraEditor/Private/NiagaraEditorUtilities.cpp
+++ b/Engine/Plugins/FX/Niagara/Source/NiagaraEditor/Private/NiagaraEditorUtilities.cpp
@@ -947,6 +947,7 @@ bool FNiagaraEditorUtilities::TryGetEventDisplayName(UNiagaraEmitter* Emitter, F
 {
 	if (Emitter != nullptr)
 	{
+		Emitter->CheckVersionDataAvailable();
 		for (const FNiagaraEventScriptProperties& EventScriptProperties : Emitter->GetLatestEmitterData()->GetEventHandlers())
 		{
 			if (EventScriptProperties.Script->GetUsageId() == EventUsageId)
