Temp fix for UE-151483 - Crash with emitter versioning in UI.

#rnx

diff --git a/Engine/Plugins/FX/Niagara/Source/NiagaraEditor/Private/ViewModels/NiagaraSystemViewModel.cpp b/Engine/Plugins/FX/Niagara/Source/NiagaraEditor/Private/ViewModels/NiagaraSystemViewModel.cpp
--- a/Engine/Plugins/FX/Niagara/Source/NiagaraEditor/Private/ViewModels/NiagaraSystemViewModel.cpp
+++ b/Engine/Plugins/FX/Niagara/Source/NiagaraEditor/Private/ViewModels/NiagaraSystemViewModel.cpp
@@ -296,7 +296,9 @@ TSharedPtr<FNiagaraEmitterHandleViewModel> FNiagaraSystemViewModel::GetEmitterHa
 {
 	for (TSharedRef<FNiagaraEmitterHandleViewModel> EmitterHandleViewModel : EmitterHandleViewModels)
 	{
-		if (EmitterHandleViewModel->GetEmitterViewModel()->GetEmitter() == InEmitter)
+		FVersionedNiagaraEmitterData* InEmitterData = InEmitter.GetEmitterData();
+		FVersionedNiagaraEmitterData* ViewModelEmitterData = EmitterHandleViewModel->GetEmitterViewModel()->GetEmitter().GetEmitterData();
+		if (InEmitterData == ViewModelEmitterData)
 		{
 			return EmitterHandleViewModel;
 		}
diff --git a/Engine/Plugins/FX/Niagara/Source/NiagaraEditor/Private/ViewModels/Stack/NiagaraStackGraphUtilities.cpp b/Engine/Plugins/FX/Niagara/Source/NiagaraEditor/Private/ViewModels/Stack/NiagaraStackGraphUtilities.cpp
--- a/Engine/Plugins/FX/Niagara/Source/NiagaraEditor/Private/ViewModels/Stack/NiagaraStackGraphUtilities.cpp
+++ b/Engine/Plugins/FX/Niagara/Source/NiagaraEditor/Private/ViewModels/Stack/NiagaraStackGraphUtilities.cpp
@@ -3532,13 +3532,13 @@ int32 FNiagaraStackGraphUtilities::DependencyUtilities::FindBestIndexForModuleIn
 
 	// Get the stack module data for the emitter stack to find dependencies.
 	TSharedPtr<FNiagaraSystemViewModel> SystemViewModel = TNiagaraViewModelManager<UNiagaraSystem, FNiagaraSystemViewModel>::GetExistingViewModelForObject(System);
-	if (SystemViewModel->IsValid() == false)
+	if (SystemViewModel.IsValid() == false)
 	{
 		ensureMsgf(false, TEXT("Failed to get systemviewmodel for valid system when getting best index in stack for module!"));
 		return INDEX_NONE;
 	}
 	TSharedPtr<FNiagaraEmitterHandleViewModel> EmitterHandleViewModel = SystemViewModel->GetEmitterHandleViewModelForEmitter(OuterEmitter);
-	if (EmitterHandleViewModel->IsValid() == false)
+	if (EmitterHandleViewModel.IsValid() == false)
 	{
 		ensureMsgf(false, TEXT("Failed to get emitterhandleviewmodel for valid emitter when getting best index in stack for module!"));
 		return INDEX_NONE;
