Improved log output

diff --git a/Engine/Plugins/FX/Niagara/Source/NiagaraEditor/Private/NiagaraScriptMergeManager.cpp b/Engine/Plugins/FX/Niagara/Source/NiagaraEditor/Private/NiagaraScriptMergeManager.cpp
--- a/Engine/Plugins/FX/Niagara/Source/NiagaraEditor/Private/NiagaraScriptMergeManager.cpp
+++ b/Engine/Plugins/FX/Niagara/Source/NiagaraEditor/Private/NiagaraScriptMergeManager.cpp
@@ -691,7 +691,7 @@ void FNiagaraScratchPadMergeAdapter::Initialize()
 			for (UNiagaraScript* ParentScratchPadScript : ParentScratchPadScripts)
 			{
 				UNiagaraScript* TargetCopy;
-				if (ensureMsgf(TargetCopyIndex < TargetEmitterData->ParentScratchPads->Scripts.Num(), TEXT("Parent scratch pad script was missing from the Target's copies")))
+				if (ensureMsgf(TargetCopyIndex < TargetEmitterData->ParentScratchPads->Scripts.Num(), TEXT("Parent scratch pad script was missing from the Target's copies. Emitter %s"), *GetPathNameSafe(InstanceEmitter.Emitter)))
 				{
 					TargetCopy = TargetEmitterData->ParentScratchPads->Scripts[TargetCopyIndex];
 				}
@@ -710,7 +710,7 @@ void FNiagaraScratchPadMergeAdapter::Initialize()
 			for (int32 InstanceScratchPadScriptIndex = 0; InstanceScratchPadScriptIndex < InstanceEmitterData->ScratchPads->Scripts.Num(); InstanceScratchPadScriptIndex++)
 			{
 				UNiagaraScript* TargetCopy;
-				if (ensureMsgf(InstanceScratchPadScriptIndex < InstanceEmitterData->ScratchPads->Scripts.Num(), TEXT("Instance scratch pad script was missing from the Target's copies.")))
+				if (ensureMsgf(InstanceScratchPadScriptIndex < InstanceEmitterData->ScratchPads->Scripts.Num(), TEXT("Instance scratch pad script was missing from the Target's copies. Emitter %s"), *GetPathNameSafe(InstanceEmitter.Emitter)))
 				{
 					TargetCopy = InstanceEmitterData->ScratchPads->Scripts[InstanceScratchPadScriptIndex];
 				}
