Fixes for build errors

diff --git a/Engine/Plugins/FX/Niagara/Source/Niagara/Private/DataInterface/NiagaraDataInterfaceActorComponent.cpp b/Engine/Plugins/FX/Niagara/Source/Niagara/Private/DataInterface/NiagaraDataInterfaceActorComponent.cpp
--- a/Engine/Plugins/FX/Niagara/Source/Niagara/Private/DataInterface/NiagaraDataInterfaceActorComponent.cpp
+++ b/Engine/Plugins/FX/Niagara/Source/Niagara/Private/DataInterface/NiagaraDataInterfaceActorComponent.cpp
@@ -10,6 +10,7 @@
 #include "GameFramework/Actor.h"
 #include "ShaderCompilerCore.h"
 #include "ShaderParameterUtils.h"
+#include "Misc/LargeWorldRenderPosition.h"
 
 #define LOCTEXT_NAMESPACE "NiagaraDataInterfaceActorComponent"
 
diff --git a/Engine/Plugins/FX/Niagara/Source/Niagara/Private/NiagaraDataInterfaceExport.cpp b/Engine/Plugins/FX/Niagara/Source/Niagara/Private/NiagaraDataInterfaceExport.cpp
--- a/Engine/Plugins/FX/Niagara/Source/Niagara/Private/NiagaraDataInterfaceExport.cpp
+++ b/Engine/Plugins/FX/Niagara/Source/Niagara/Private/NiagaraDataInterfaceExport.cpp
@@ -13,6 +13,7 @@
 #include "Internationalization/Internationalization.h"
 #include "ShaderCompilerCore.h"
 #include "ShaderParameterUtils.h"
+#include "Misc/LargeWorldRenderPosition.h"
 
 namespace NDIExportLocal
 {
diff --git a/Engine/Plugins/FX/Niagara/Source/NiagaraEditor/Private/NiagaraValidationRules.cpp b/Engine/Plugins/FX/Niagara/Source/NiagaraEditor/Private/NiagaraValidationRules.cpp
--- a/Engine/Plugins/FX/Niagara/Source/NiagaraEditor/Private/NiagaraValidationRules.cpp
+++ b/Engine/Plugins/FX/Niagara/Source/NiagaraEditor/Private/NiagaraValidationRules.cpp
@@ -23,85 +23,88 @@
 
 #define LOCTEXT_NAMESPACE "NiagaraValidationRules"
 
-template<typename T>
-TArray<T*> GetStackEntries(UNiagaraStackViewModel* StackViewModel, bool bRefresh = false)
+namespace NiagaraValidation
 {
-	TArray<T*> Results;
-	TArray<UNiagaraStackEntry*> EntriesToCheck;
-	if (UNiagaraStackEntry* RootEntry = StackViewModel->GetRootEntry())
+	template<typename T>
+	TArray<T*> GetStackEntries(UNiagaraStackViewModel* StackViewModel, bool bRefresh = false)
 	{
-		if (bRefresh)
+		TArray<T*> Results;
+		TArray<UNiagaraStackEntry*> EntriesToCheck;
+		if (UNiagaraStackEntry* RootEntry = StackViewModel->GetRootEntry())
 		{
-			RootEntry->RefreshChildren();
+			if (bRefresh)
+			{
+				RootEntry->RefreshChildren();
+			}
+			RootEntry->GetUnfilteredChildren(EntriesToCheck);
 		}
-		RootEntry->GetUnfilteredChildren(EntriesToCheck);
-	}
-	while (EntriesToCheck.Num() > 0)
-	{
-		UNiagaraStackEntry* Entry = EntriesToCheck.Pop();
-		if (T* ItemToCheck = Cast<T>(Entry))
+		while (EntriesToCheck.Num() > 0)
 		{
-			Results.Add(ItemToCheck);
+			UNiagaraStackEntry* Entry = EntriesToCheck.Pop();
+			if (T* ItemToCheck = Cast<T>(Entry))
+			{
+				Results.Add(ItemToCheck);
+			}
+			Entry->GetUnfilteredChildren(EntriesToCheck);
 		}
-		Entry->GetUnfilteredChildren(EntriesToCheck);
+		return Results;
 	}
-	return Results;
-}
 
-template<typename T>
-TArray<T*> GetAllStackEntriesInSystem(TSharedPtr<FNiagaraSystemViewModel> ViewModel, bool bRefresh = false)
-{
-	TArray<T*> Results;
-	Results.Append(GetStackEntries<T>(ViewModel->GetSystemStackViewModel(), bRefresh));
-	TArray<TSharedRef<FNiagaraEmitterHandleViewModel>> EmitterHandleViewModels = ViewModel->GetEmitterHandleViewModels();
-	for (TSharedRef<FNiagaraEmitterHandleViewModel> EmitterHandleModel : EmitterHandleViewModels)
+	template<typename T>
+	TArray<T*> GetAllStackEntriesInSystem(TSharedPtr<FNiagaraSystemViewModel> ViewModel, bool bRefresh = false)
 	{
-		Results.Append(GetStackEntries<T>(EmitterHandleModel.Get().GetEmitterStackViewModel(), bRefresh));
+		TArray<T*> Results;
+		Results.Append(GetStackEntries<T>(ViewModel->GetSystemStackViewModel(), bRefresh));
+		TArray<TSharedRef<FNiagaraEmitterHandleViewModel>> EmitterHandleViewModels = ViewModel->GetEmitterHandleViewModels();
+		for (TSharedRef<FNiagaraEmitterHandleViewModel> EmitterHandleModel : EmitterHandleViewModels)
+		{
+			Results.Append(GetStackEntries<T>(EmitterHandleModel.Get().GetEmitterStackViewModel(), bRefresh));
+		}
+		return Results;
 	}
-	return Results;
-}
 
-// helper function to retrieve a single stack entry from the system or emitter view model
-template<typename T>
-T* GetStackEntry(UNiagaraStackViewModel* StackViewModel, bool bRefresh = false)
-{
-	TArray<T*> StackEntries = GetStackEntries<T>(StackViewModel, bRefresh);
-	if (StackEntries.Num() > 0)
+	// helper function to retrieve a single stack entry from the system or emitter view model
+	template<typename T>
+	T* GetStackEntry(UNiagaraStackViewModel* StackViewModel, bool bRefresh = false)
 	{
-		return StackEntries[0];
+		TArray<T*> StackEntries = GetStackEntries<T>(StackViewModel, bRefresh);
+		if (StackEntries.Num() > 0)
+		{
+			return StackEntries[0];
+		}
+		return nullptr;
 	}
-	return nullptr;
-}
 
-// --------------------------------------------------------------------------------------------------------------------------------------------
-// Common fixes and links
+	// --------------------------------------------------------------------------------------------------------------------------------------------
+	// Common fixes and links
 
-void AddGoToFXTypeLink(FNiagaraValidationResult& Result, UNiagaraEffectType* FXType)
-{
-	if (FXType == nullptr)
+	void AddGoToFXTypeLink(FNiagaraValidationResult& Result, UNiagaraEffectType* FXType)
 	{
-		return;
-	}
-
-	FNiagaraValidationFix& GoToValidationRulesLink = Result.Links.AddDefaulted_GetRef();
-	GoToValidationRulesLink.Description = LOCTEXT("GoToValidationRulesFix", "Go To Validation Rules");
-	TWeakObjectPtr<UNiagaraEffectType> WeakFXType = FXType;
-	GoToValidationRulesLink.FixDelegate = FNiagaraValidationFixDelegate::CreateLambda([WeakFXType]
+		if (FXType == nullptr)
 		{
-			FAssetToolsModule& AssetToolsModule = FModuleManager::LoadModuleChecked<FAssetToolsModule>(TEXT("AssetTools"));
-			TWeakPtr<IAssetTypeActions> WeakAssetTypeActions = AssetToolsModule.Get().GetAssetTypeActionsForClass(UNiagaraEffectType::StaticClass());
+			return;
+		}
 
-			if (UNiagaraEffectType* FXType = WeakFXType.Get())
+		FNiagaraValidationFix& GoToValidationRulesLink = Result.Links.AddDefaulted_GetRef();
+		GoToValidationRulesLink.Description = LOCTEXT("GoToValidationRulesFix", "Go To Validation Rules");
+		TWeakObjectPtr<UNiagaraEffectType> WeakFXType = FXType;
+		GoToValidationRulesLink.FixDelegate = FNiagaraValidationFixDelegate::CreateLambda([WeakFXType]
 			{
-				if (TSharedPtr<IAssetTypeActions> AssetTypeActions = WeakAssetTypeActions.Pin())
+				FAssetToolsModule& AssetToolsModule = FModuleManager::LoadModuleChecked<FAssetToolsModule>(TEXT("AssetTools"));
+				TWeakPtr<IAssetTypeActions> WeakAssetTypeActions = AssetToolsModule.Get().GetAssetTypeActionsForClass(UNiagaraEffectType::StaticClass());
+
+				if (UNiagaraEffectType* FXType = WeakFXType.Get())
 				{
-					TArray<UObject*> AssetsToEdit;
-					AssetsToEdit.Add(FXType);
-					AssetTypeActions->OpenAssetEditor(AssetsToEdit);
-					//TODO: Is there a way for us to auto navigate to and open up the validation rules inside FXType?
+					if (TSharedPtr<IAssetTypeActions> AssetTypeActions = WeakAssetTypeActions.Pin())
+					{
+						TArray<UObject*> AssetsToEdit;
+						AssetsToEdit.Add(FXType);
+						AssetTypeActions->OpenAssetEditor(AssetsToEdit);
+						//TODO: Is there a way for us to auto navigate to and open up the validation rules inside FXType?
+					}
 				}
-			}
-		});
+			});
+	}
 }
 
 // --------------------------------------------------------------------------------------------------------------------------------------------
@@ -111,7 +114,7 @@ void UNiagaraValidationRule_NoWarmupTime::CheckValidity(TSharedPtr<FNiagaraSyste
 	UNiagaraSystem& System = ViewModel->GetSystem();
 	if (System.NeedsWarmup())
 	{
-		UNiagaraStackSystemPropertiesItem* SystemProperties = GetStackEntry<UNiagaraStackSystemPropertiesItem>(ViewModel->GetSystemStackViewModel());
+		UNiagaraStackSystemPropertiesItem* SystemProperties = NiagaraValidation::GetStackEntry<UNiagaraStackSystemPropertiesItem>(ViewModel->GetSystemStackViewModel());
 		FNiagaraValidationResult Result(ENiagaraValidationSeverity::Error, LOCTEXT("WarumupSummary", "Warmuptime > 0 is not allowed"), LOCTEXT("WarmupDescription", "Systems with the chosen effect type do not allow warmup time, as it costs too much performance.\nPlease set the warmup time to 0 in the system properties."), SystemProperties);
 		Results.Add(Result);
 	}
@@ -132,7 +135,7 @@ void UNiagaraValidationRule_FixedGPUBoundsSet::CheckValidity(TSharedPtr<FNiagara
 		FVersionedNiagaraEmitterData* EmitterData = EmitterHandleModel.Get().GetEmitterHandle()->GetEmitterData();
 		if (EmitterData->SimTarget == ENiagaraSimTarget::GPUComputeSim && EmitterData->CalculateBoundsMode == ENiagaraEmitterCalculateBoundMode::Dynamic)
 		{
-			UNiagaraStackEmitterPropertiesItem* EmitterProperties = GetStackEntry<UNiagaraStackEmitterPropertiesItem>(EmitterHandleModel.Get().GetEmitterStackViewModel());
+			UNiagaraStackEmitterPropertiesItem* EmitterProperties = NiagaraValidation::GetStackEntry<UNiagaraStackEmitterPropertiesItem>(EmitterHandleModel.Get().GetEmitterStackViewModel());
 			FNiagaraValidationResult Result(ENiagaraValidationSeverity::Error, LOCTEXT("GpuDynamicBoundsErrorSummary", "GPU emitters do not support dynamic bounds"), LOCTEXT("GpuDynamicBoundsErrorDescription", "Gpu emitter should either not be in dynamic mode or the system must have fixed bounds."), EmitterProperties);
 			Results.Add(Result);
 		}
@@ -170,7 +173,7 @@ void UNiagaraValidationRule_BannedRenderers::CheckValidity(TSharedPtr<FNiagaraSy
 				FNiagaraPlatformSet::GatherConflicts(CheckSets, Conflicts);
 				if (Conflicts.Num() > 0)
 				{
-					TArray<UNiagaraStackRendererItem*> RendererItems = GetStackEntries<UNiagaraStackRendererItem>(EmitterHandleModel.Get().GetEmitterStackViewModel());
+					TArray<UNiagaraStackRendererItem*> RendererItems = NiagaraValidation::GetStackEntries<UNiagaraStackRendererItem>(EmitterHandleModel.Get().GetEmitterStackViewModel());
 					for (UNiagaraStackRendererItem* Item : RendererItems)
 					{
 						if (Item->GetRendererProperties() != RendererProperties)
@@ -184,7 +187,7 @@ void UNiagaraValidationRule_BannedRenderers::CheckValidity(TSharedPtr<FNiagaraSy
 						Result.Description = LOCTEXT("BannedRenderDescription", "Please ensure only allowed renderers are used for each platform according to the validation rules in the System's Effect Type.");
 						Result.SourceObject = Item;
 						
-						AddGoToFXTypeLink(Result, System.GetEffectType());
+						NiagaraValidation::AddGoToFXTypeLink(Result, System.GetEffectType());
 
 						//Add autofix to disable the module
 						FNiagaraValidationFix& DisableRendererFix = Result.Fixes.AddDefaulted_GetRef();
@@ -211,7 +214,7 @@ void UNiagaraValidationRule_BannedModules::CheckValidity(TSharedPtr<FNiagaraSyst
 {
 	UNiagaraSystem& System = ViewModel->GetSystem();
 
-	TArray<UNiagaraStackModuleItem*> StackModuleItems =	GetAllStackEntriesInSystem<UNiagaraStackModuleItem>(ViewModel);
+	TArray<UNiagaraStackModuleItem*> StackModuleItems =	NiagaraValidation::GetAllStackEntriesInSystem<UNiagaraStackModuleItem>(ViewModel);
 
 	for (UNiagaraStackModuleItem* Item : StackModuleItems)
 	{
@@ -251,7 +254,7 @@ void UNiagaraValidationRule_BannedModules::CheckValidity(TSharedPtr<FNiagaraSyst
 					Result.Description = LOCTEXT("BanndeModulesDescription", "Check this module against the Effect Type's Banned Modules validators");
 					Result.SourceObject = Item;
 
-					AddGoToFXTypeLink(Result, System.GetEffectType());
+					NiagaraValidation::AddGoToFXTypeLink(Result, System.GetEffectType());
 
 					//Add autofix to disable the module
 					FNiagaraValidationFix& DisableModuleFix = Result.Fixes.AddDefaulted_GetRef();
@@ -272,7 +275,7 @@ void UNiagaraValidationRule_BannedModules::CheckValidity(TSharedPtr<FNiagaraSyst
 
 void UNiagaraValidationRule_InvalidEffectType::CheckValidity(TSharedPtr<FNiagaraSystemViewModel> ViewModel, TArray<FNiagaraValidationResult>& Results)  const
 {
-	UNiagaraStackSystemPropertiesItem* SystemProperties = GetStackEntry<UNiagaraStackSystemPropertiesItem>(ViewModel->GetSystemStackViewModel());
+	UNiagaraStackSystemPropertiesItem* SystemProperties = NiagaraValidation::GetStackEntry<UNiagaraStackSystemPropertiesItem>(ViewModel->GetSystemStackViewModel());
 	FNiagaraValidationResult Result(ENiagaraValidationSeverity::Error, LOCTEXT("InvalidEffectSummary", "Invalid Effect Type"), LOCTEXT("InvalidEffectDescription", "The effect type on this system was marked as invalid for production content and should only be used as placeholder."), SystemProperties);
 	Results.Add(Result);
 }
@@ -289,13 +292,13 @@ void UNiagaraValidationRule_LWC::CheckValidity(TSharedPtr<FNiagaraSystemViewMode
 
 	// gather all the modules in the system, excluding localspace emitters
 	TArray<UNiagaraStackModuleItem*> AllModules;
-	AllModules.Append(GetStackEntries<UNiagaraStackModuleItem>(ViewModel->GetSystemStackViewModel()));
+	AllModules.Append(NiagaraValidation::GetStackEntries<UNiagaraStackModuleItem>(ViewModel->GetSystemStackViewModel()));
 	TArray<TSharedRef<FNiagaraEmitterHandleViewModel>> EmitterHandleViewModels = ViewModel->GetEmitterHandleViewModels();
 	for (TSharedRef<FNiagaraEmitterHandleViewModel> EmitterHandleModel : EmitterHandleViewModels)
 	{
 		if (EmitterHandleModel->GetEmitterHandle()->GetEmitterData()->bLocalSpace == false)
 		{
-			AllModules.Append(GetStackEntries<UNiagaraStackModuleItem>(EmitterHandleModel.Get().GetEmitterStackViewModel()));
+			AllModules.Append(NiagaraValidation::GetStackEntries<UNiagaraStackModuleItem>(EmitterHandleModel.Get().GetEmitterStackViewModel()));
 		}
 	}
 
diff --git a/Engine/Plugins/FX/Niagara/Source/NiagaraEditor/Private/UpgradeNiagaraScriptResults.cpp b/Engine/Plugins/FX/Niagara/Source/NiagaraEditor/Private/UpgradeNiagaraScriptResults.cpp
--- a/Engine/Plugins/FX/Niagara/Source/NiagaraEditor/Private/UpgradeNiagaraScriptResults.cpp
+++ b/Engine/Plugins/FX/Niagara/Source/NiagaraEditor/Private/UpgradeNiagaraScriptResults.cpp
@@ -10,54 +10,57 @@
 #include "ViewModels/Stack/NiagaraStackModuleItem.h"
 #include "ViewModels/Stack/NiagaraStackViewModel.h"
 
-template<typename T>
-TArray<T*> GetStackEntries(UNiagaraStackViewModel* StackViewModel, bool bRefresh = false)
+namespace NiagaraScriptResults
 {
-	TArray<T*> Results;
-	TArray<UNiagaraStackEntry*> EntriesToCheck;
-	if (UNiagaraStackEntry* RootEntry = StackViewModel->GetRootEntry())
+	template<typename T>
+	TArray<T*> GetStackEntries(UNiagaraStackViewModel* StackViewModel, bool bRefresh = false)
 	{
-		if (bRefresh)
+		TArray<T*> Results;
+		TArray<UNiagaraStackEntry*> EntriesToCheck;
+		if (UNiagaraStackEntry* RootEntry = StackViewModel->GetRootEntry())
 		{
-			RootEntry->RefreshChildren();
+			if (bRefresh)
+			{
+				RootEntry->RefreshChildren();
+			}
+			RootEntry->GetUnfilteredChildren(EntriesToCheck);
 		}
-		RootEntry->GetUnfilteredChildren(EntriesToCheck);
-	}
-	while (EntriesToCheck.Num() > 0)
-	{
-		UNiagaraStackEntry* Entry = EntriesToCheck.Pop();
-		if (T* ItemToCheck = Cast<T>(Entry))
+		while (EntriesToCheck.Num() > 0)
 		{
-			Results.Add(ItemToCheck);
+			UNiagaraStackEntry* Entry = EntriesToCheck.Pop();
+			if (T* ItemToCheck = Cast<T>(Entry))
+			{
+				Results.Add(ItemToCheck);
+			}
+			Entry->GetUnfilteredChildren(EntriesToCheck);
 		}
-		Entry->GetUnfilteredChildren(EntriesToCheck);
+		return Results;
 	}
-	return Results;
-}
 
-template<typename T>
-T GetValue(const UNiagaraClipboardFunctionInput* Input)
-{
-	T Value;
-	if (Input == nullptr || Input->InputType.GetSize() != sizeof(T) || Input->Local.Num() != sizeof(T))
+	template<typename T>
+	T GetValue(const UNiagaraClipboardFunctionInput* Input)
 	{
-		FMemory::Memzero(Value);
+		T Value;
+		if (Input == nullptr || Input->InputType.GetSize() != sizeof(T) || Input->Local.Num() != sizeof(T))
+		{
+			FMemory::Memzero(Value);
+		}
+		else
+		{
+			FMemory::Memcpy(&Value, Input->Local.GetData(), sizeof(T));
+		}
+		return Value;
 	}
-	else
+
+	template<typename T>
+	void SetValue(UNiagaraPythonScriptModuleInput* ModuleInput, T Data)
 	{
-		FMemory::Memcpy(&Value, Input->Local.GetData(), sizeof(T));
+		TArray<uint8> LocalData;
+		LocalData.SetNumZeroed(sizeof(T));
+		FMemory::Memcpy(LocalData.GetData(), &Data, sizeof(T));
+		const UNiagaraClipboardFunctionInput* Input = ModuleInput->Input;
+		ModuleInput->Input = UNiagaraClipboardFunctionInput::CreateLocalValue(ModuleInput, Input->InputName, Input->InputType, Input->bEditConditionValue, LocalData);
 	}
-	return Value;
-}
-
-template<typename T>
-void SetValue(UNiagaraPythonScriptModuleInput* ModuleInput, T Data)
-{
-	TArray<uint8> LocalData;
-	LocalData.SetNumZeroed(sizeof(T));
-	FMemory::Memcpy(LocalData.GetData(), &Data, sizeof(T));
-	const UNiagaraClipboardFunctionInput* Input = ModuleInput->Input;
-	ModuleInput->Input = UNiagaraClipboardFunctionInput::CreateLocalValue(ModuleInput, Input->InputName, Input->InputType, Input->bEditConditionValue, LocalData);
 }
 
 bool UNiagaraPythonScriptModuleInput::IsSet() const
@@ -75,7 +78,7 @@ float UNiagaraPythonScriptModuleInput::AsFloat() const
 	if (IsSet() && Input->InputType == FNiagaraTypeDefinition::GetFloatDef())
 	{
 		FNiagaraVariable LocalInput;
-		return GetValue<float>(Input);
+		return NiagaraScriptResults::GetValue<float>(Input);
 	}
 	return 0;
 }
@@ -84,7 +87,7 @@ int32 UNiagaraPythonScriptModuleInput::AsInt() const
 {
 	if (IsSet() && Input->InputType == FNiagaraTypeDefinition::GetIntDef())
 	{
-		return GetValue<int32>(Input);
+		return NiagaraScriptResults::GetValue<int32>(Input);
 	}
 	return 0;
 }
@@ -107,7 +110,7 @@ FVector2D UNiagaraPythonScriptModuleInput::AsVec2() const
 {
 	if (IsSet() && Input->InputType == FNiagaraTypeDefinition::GetVec2Def())
 	{
-		return GetValue<FVector2D>(Input);
+		return NiagaraScriptResults::GetValue<FVector2D>(Input);
 	}
 	return FVector2D();
 }
@@ -116,7 +119,7 @@ FVector UNiagaraPythonScriptModuleInput::AsVec3() const
 {
 	if (IsSet() && Input->InputType == FNiagaraTypeDefinition::GetVec3Def())
 	{
-		return (FVector)GetValue<FVector3f>(Input);
+		return (FVector)NiagaraScriptResults::GetValue<FVector3f>(Input);
 	}
 	return FVector();
 }
@@ -125,7 +128,7 @@ FVector4 UNiagaraPythonScriptModuleInput::AsVec4() const
 {
 	if (IsSet() && Input->InputType == FNiagaraTypeDefinition::GetVec4Def())
 	{
-		return FVector4(GetValue<FVector4f>(Input));
+		return FVector4(NiagaraScriptResults::GetValue<FVector4f>(Input));
 	}
 	return FVector4();
 }
@@ -134,7 +137,7 @@ FLinearColor UNiagaraPythonScriptModuleInput::AsColor() const
 {
 	if (IsSet() && Input->InputType == FNiagaraTypeDefinition::GetColorDef())
 	{
-		return GetValue<FLinearColor>(Input);
+		return NiagaraScriptResults::GetValue<FLinearColor>(Input);
 	}
 	return FLinearColor();
 }
@@ -143,7 +146,7 @@ FQuat UNiagaraPythonScriptModuleInput::AsQuat() const
 {
 	if (IsSet() && Input->InputType == FNiagaraTypeDefinition::GetQuatDef())
 	{
-		return GetValue<FQuat>(Input);
+		return NiagaraScriptResults::GetValue<FQuat>(Input);
 	}
 	return FQuat();
 }
@@ -152,7 +155,7 @@ FString UNiagaraPythonScriptModuleInput::AsEnum() const
 {
 	if (IsSet() && Input->InputType.IsEnum())
 	{
-		int32 Value = GetValue<int32>(Input);
+		int32 Value = NiagaraScriptResults::GetValue<int32>(Input);
 		return Input->InputType.GetEnum()->GetNameStringByValue(Value);
 	}
 	return FString();
@@ -195,7 +198,7 @@ void UUpgradeNiagaraScriptResults::SetFloatInput(const FString& InputName, float
 	UNiagaraPythonScriptModuleInput* ModuleInput = GetNewInput(FName(InputName));
 	if (ModuleInput && ModuleInput->Input->InputType == FNiagaraTypeDefinition::GetFloatDef())
 	{
-		SetValue(ModuleInput, Value);
+		NiagaraScriptResults::SetValue(ModuleInput, Value);
 	}
 }
 
@@ -204,7 +207,7 @@ void UUpgradeNiagaraScriptResults::SetIntInput(const FString& InputName, int32 V
 	UNiagaraPythonScriptModuleInput* ModuleInput = GetNewInput(FName(InputName));
 	if (ModuleInput && ModuleInput->Input->InputType == FNiagaraTypeDefinition::GetIntDef())
 	{
-		SetValue(ModuleInput, Value);
+		NiagaraScriptResults::SetValue(ModuleInput, Value);
 	}
 }
 
@@ -213,7 +216,7 @@ void UUpgradeNiagaraScriptResults::SetBoolInput(const FString& InputName, bool V
 	UNiagaraPythonScriptModuleInput* ModuleInput = GetNewInput(FName(InputName));
 	if (ModuleInput && ModuleInput->Input->InputType == FNiagaraTypeDefinition::GetBoolDef())
 	{
-		SetValue(ModuleInput, Value);
+		NiagaraScriptResults::SetValue(ModuleInput, Value);
 	}
 }
 
@@ -222,7 +225,7 @@ void UUpgradeNiagaraScriptResults::SetVec2Input(const FString& InputName, FVecto
 	UNiagaraPythonScriptModuleInput* ModuleInput = GetNewInput(FName(InputName));
 	if (ModuleInput && ModuleInput->Input->InputType == FNiagaraTypeDefinition::GetVec2Def())
 	{
-		SetValue(ModuleInput, Value);
+		NiagaraScriptResults::SetValue(ModuleInput, Value);
 	}
 }
 
@@ -231,7 +234,7 @@ void UUpgradeNiagaraScriptResults::SetVec3Input(const FString& InputName, FVecto
 	UNiagaraPythonScriptModuleInput* ModuleInput = GetNewInput(FName(InputName));
 	if (ModuleInput && ModuleInput->Input->InputType == FNiagaraTypeDefinition::GetVec3Def())
 	{
-		SetValue(ModuleInput, Value);
+		NiagaraScriptResults::SetValue(ModuleInput, Value);
 	}
 }
 
@@ -240,7 +243,7 @@ void UUpgradeNiagaraScriptResults::SetVec4Input(const FString& InputName, FVecto
 	UNiagaraPythonScriptModuleInput* ModuleInput = GetNewInput(FName(InputName));
 	if (ModuleInput && ModuleInput->Input->InputType == FNiagaraTypeDefinition::GetVec4Def())
 	{
-		SetValue(ModuleInput, Value);
+		NiagaraScriptResults::SetValue(ModuleInput, Value);
 	}
 }
 
@@ -249,7 +252,7 @@ void UUpgradeNiagaraScriptResults::SetColorInput(const FString& InputName, FLine
 	UNiagaraPythonScriptModuleInput* ModuleInput = GetNewInput(FName(InputName));
 	if (ModuleInput && ModuleInput->Input->InputType == FNiagaraTypeDefinition::GetColorDef())
 	{
-		SetValue(ModuleInput, Value);
+		NiagaraScriptResults::SetValue(ModuleInput, Value);
 	}
 }
 
@@ -258,7 +261,7 @@ void UUpgradeNiagaraScriptResults::SetQuatInput(const FString& InputName, FQuat
 	UNiagaraPythonScriptModuleInput* ModuleInput = GetNewInput(FName(InputName));
 	if (ModuleInput && ModuleInput->Input->InputType == FNiagaraTypeDefinition::GetQuatDef())
 	{
-		SetValue(ModuleInput, Value);
+		NiagaraScriptResults::SetValue(ModuleInput, Value);
 	}
 }
 
@@ -268,7 +271,7 @@ void UUpgradeNiagaraScriptResults::SetEnumInput(const FString& InputName, FStrin
 	if (ModuleInput && ModuleInput->Input->InputType.IsEnum())
 	{
 		int32 EnumValue = ModuleInput->Input->InputType.GetEnum()->GetValueByNameString(Value, EGetByNameFlags::ErrorIfNotFound | EGetByNameFlags::CheckAuthoredName);
-		SetValue(ModuleInput, EnumValue);
+		NiagaraScriptResults::SetValue(ModuleInput, EnumValue);
 	}
 }
 
@@ -336,7 +339,7 @@ TArray<UNiagaraPythonModule*> UNiagaraPythonEmitter::GetModules() const
 
 	if (UNiagaraStackViewModel* StackViewModel = EmitterViewModel->GetEmitterStackViewModel())
 	{
-		TArray<UNiagaraStackModuleItem*> StackModuleItems =	GetStackEntries<UNiagaraStackModuleItem>(StackViewModel);
+		TArray<UNiagaraStackModuleItem*> StackModuleItems =	NiagaraScriptResults::GetStackEntries<UNiagaraStackModuleItem>(StackViewModel);
 		for (UNiagaraStackModuleItem* ModuleItem : StackModuleItems)
 		{
 			UNiagaraPythonModule* Module = NewObject<UNiagaraPythonModule>();
@@ -358,7 +361,7 @@ UNiagaraPythonModule* UNiagaraPythonEmitter::GetModule(const FString& ModuleName
 	UNiagaraPythonModule* Module = NewObject<UNiagaraPythonModule>();
 	if (UNiagaraStackViewModel* StackViewModel = EmitterViewModel->GetEmitterStackViewModel())
 	{
-		TArray<UNiagaraStackModuleItem*> StackModuleItems =	GetStackEntries<UNiagaraStackModuleItem>(StackViewModel);
+		TArray<UNiagaraStackModuleItem*> StackModuleItems =	NiagaraScriptResults::GetStackEntries<UNiagaraStackModuleItem>(StackViewModel);
 		for (UNiagaraStackModuleItem* ModuleItem : StackModuleItems)
 		{
 			if (ModuleItem->GetModuleNode().GetFunctionName() == ModuleName)
