Fix for crash generating baselines.

#rnx

diff --git a/Engine/Plugins/FX/Niagara/Source/Niagara/Private/NiagaraPerfBaseline.cpp b/Engine/Plugins/FX/Niagara/Source/Niagara/Private/NiagaraPerfBaseline.cpp
--- a/Engine/Plugins/FX/Niagara/Source/Niagara/Private/NiagaraPerfBaseline.cpp
+++ b/Engine/Plugins/FX/Niagara/Source/Niagara/Private/NiagaraPerfBaseline.cpp
@@ -272,7 +272,10 @@ UNiagaraSystem* UNiagaraBaselineController::GetSystem()
 void UNiagaraBaselineController_Basic::OnBeginTest_Implementation()
 {
 	Super::OnBeginTest_Implementation();
-	SpawnedComponents.SetNumZeroed(NumInstances);
+	if(NumInstances > 0)
+	{
+		SpawnedComponents.SetNumZeroed(NumInstances);
+	}
 }
 
 bool UNiagaraBaselineController_Basic::OnTickTest_Implementation()
