- Add missing null check if  RequestDefaultDataInterface does not return a valid data interface
- Invalidate the compile results if this does happen as the system can no longer execute

#rnx

diff --git a/Engine/Plugins/FX/Niagara/Source/NiagaraShader/Private/NiagaraShared.cpp b/Engine/Plugins/FX/Niagara/Source/NiagaraShader/Private/NiagaraShared.cpp
--- a/Engine/Plugins/FX/Niagara/Source/NiagaraShader/Private/NiagaraShared.cpp
+++ b/Engine/Plugins/FX/Niagara/Source/NiagaraShader/Private/NiagaraShared.cpp
@@ -691,6 +691,12 @@ void FNiagaraShaderScript::BuildScriptParametersMetadata(const TArray<FNiagaraDa
 	for (FNiagaraDataInterfaceGPUParamInfo& DataInterfaceParamInfo : ScriptParametersMetadata->DataInterfaceParamInfo)
 	{
 		UNiagaraDataInterfaceBase* CDODataInterface = ShaderModule->RequestDefaultDataInterface(*DataInterfaceParamInfo.DIClassName);
+		if (CDODataInterface == nullptr)
+		{
+			Invalidate();
+			return;
+		}
+
 		if (CDODataInterface->UseLegacyShaderBindings() == false)
 		{
 			const uint32 NextMemberOffset = ShaderMetadataBuilder.GetNextMemberOffset();
