Fixed a crash caused by emitters not having all the parent scratch pads during a merge

diff --git a/Engine/Plugins/FX/Niagara/Source/Niagara/Classes/NiagaraScratchPadContainer.h b/Engine/Plugins/FX/Niagara/Source/Niagara/Classes/NiagaraScratchPadContainer.h
--- a/Engine/Plugins/FX/Niagara/Source/Niagara/Classes/NiagaraScratchPadContainer.h
+++ b/Engine/Plugins/FX/Niagara/Source/Niagara/Classes/NiagaraScratchPadContainer.h
@@ -19,6 +19,8 @@ public:
 	virtual void PostLoad() override;
 	void CheckConsistency();
 	void SetScripts(const TArray<TObjectPtr<UNiagaraScript>>& InScripts);
+	void AppendScripts(const TArray<TObjectPtr<UNiagaraScript>>& InScripts);
+	void AppendScripts(TObjectPtr<UNiagaraScratchPadContainer> InScripts);
 
 #if WITH_EDITORONLY_DATA
 
diff --git a/Engine/Plugins/FX/Niagara/Source/Niagara/Private/NiagaraEmitter.cpp b/Engine/Plugins/FX/Niagara/Source/Niagara/Private/NiagaraEmitter.cpp
--- a/Engine/Plugins/FX/Niagara/Source/Niagara/Private/NiagaraEmitter.cpp
+++ b/Engine/Plugins/FX/Niagara/Source/Niagara/Private/NiagaraEmitter.cpp
@@ -215,7 +215,7 @@ void FVersionedNiagaraEmitterData::CopyFrom(const FVersionedNiagaraEmitterData&
 
 FName GetVersionedName(const FString& BaseName, const FNiagaraAssetVersion& VersionPostfix)
 {
-	return FName(FString::Printf(TEXT("%s_%i.%i"), *BaseName, VersionPostfix.MajorVersion, VersionPostfix.MinorVersion));
+	return FName(FString::Printf(TEXT("%s_%i_%i"), *BaseName, VersionPostfix.MajorVersion, VersionPostfix.MinorVersion));
 }
 
 void FVersionedNiagaraEmitterData::PostInitProperties(UNiagaraEmitter* Outer)
@@ -993,8 +993,7 @@ UNiagaraEmitter* UNiagaraEmitter::CreateWithParentAndOwner(FVersionedNiagaraEmit
 	EmitterData->VersionedParentAtLastMerge.Emitter->ClearFlags(RF_Standalone | RF_Public);
 	EmitterData->VersionedParentAtLastMerge.Emitter->DisableVersioning(InParentEmitter.Version);
 	EmitterData->VersionedParentAtLastMerge.Version = InParentEmitter.Version;
-	EmitterData->ParentScratchPads->SetScripts(EmitterData->ScratchPads->Scripts);
-	EmitterData->ScratchPads->Scripts.Empty();
+	EmitterData->ParentScratchPads->AppendScripts(EmitterData->ScratchPads);
 	EmitterData->GraphSource->MarkNotSynchronized(InitialNotSynchronizedReason);
 	NewEmitter->RebindNotifications();
 	return NewEmitter;
@@ -3140,8 +3139,7 @@ void UNiagaraEmitter::SetParent(const FVersionedNiagaraEmitter& InParent)
 
 		// Since this API is only valid for the "Create duplicate parent" operation we move the emitters scratch pad script to the parent array since that's where they're defined now.
 		// Normally we would duplicate the parent scratch pad scripts here, but since the whole parent is already a duplicate, we can skip this here.
-		EmitterData.ParentScratchPads->SetScripts(EmitterData.ScratchPads->Scripts);
-		EmitterData.ScratchPads->Scripts.Empty();
+		EmitterData.ParentScratchPads->AppendScripts(EmitterData.ScratchPads);
 		EmitterData.GraphSource->MarkNotSynchronized(TEXT("Emitter parent changed"));
 		UpdateChangeId(TEXT("Parent Set"));
 	}
diff --git a/Engine/Plugins/FX/Niagara/Source/Niagara/Private/NiagaraEmitterInstance.cpp b/Engine/Plugins/FX/Niagara/Source/Niagara/Private/NiagaraEmitterInstance.cpp
--- a/Engine/Plugins/FX/Niagara/Source/Niagara/Private/NiagaraEmitterInstance.cpp
+++ b/Engine/Plugins/FX/Niagara/Source/Niagara/Private/NiagaraEmitterInstance.cpp
@@ -1018,8 +1018,9 @@ void FNiagaraEmitterInstance::CalculateFixedBounds(const FTransform& ToWorldSpac
 void FNiagaraEmitterInstance::PostTick()
 {
 	SCOPE_CYCLE_COUNTER(STAT_NiagaraEmitterPostTick);
-
-	checkSlow(CachedEmitter.Emitter);
+	
+	FVersionedNiagaraEmitterData* EmitterData = CachedEmitter.GetEmitterData();
+	checkSlow(EmitterData);
 
 	if (EventInstanceData.IsValid())
 	{
@@ -1032,7 +1033,7 @@ void FNiagaraEmitterInstance::PostTick()
 
 	bCachedBoundsDynamic = false;
 	CachedBounds.Init();
-	FVersionedNiagaraEmitterData* EmitterData = CachedEmitter.GetEmitterData();
+	
 	if (EmitterData)
 	{
 		// Read lock can be smaller in scope, but probably not necessary
diff --git a/Engine/Plugins/FX/Niagara/Source/Niagara/Private/NiagaraScratchPadContainer.cpp b/Engine/Plugins/FX/Niagara/Source/Niagara/Private/NiagaraScratchPadContainer.cpp
--- a/Engine/Plugins/FX/Niagara/Source/Niagara/Private/NiagaraScratchPadContainer.cpp
+++ b/Engine/Plugins/FX/Niagara/Source/Niagara/Private/NiagaraScratchPadContainer.cpp
@@ -30,3 +30,23 @@ void UNiagaraScratchPadContainer::SetScripts(const TArray<TObjectPtr<UNiagaraScr
 	CheckConsistency();
 #endif
 }
+
+void UNiagaraScratchPadContainer::AppendScripts(const TArray<TObjectPtr<UNiagaraScript>>& InScripts)
+{
+#if WITH_EDITORONLY_DATA
+	Scripts.Append(InScripts);
+	CheckConsistency();
+#endif
+}
+
+void UNiagaraScratchPadContainer::AppendScripts(TObjectPtr<UNiagaraScratchPadContainer> InScripts)
+{
+#if WITH_EDITORONLY_DATA
+	if (InScripts)
+	{
+		Scripts.Append(InScripts->Scripts);
+		InScripts->Scripts.Empty();
+		CheckConsistency();
+	}
+#endif
+}
diff --git a/Engine/Plugins/FX/Niagara/Source/NiagaraEditor/Private/NiagaraScriptMergeManager.cpp b/Engine/Plugins/FX/Niagara/Source/NiagaraEditor/Private/NiagaraScriptMergeManager.cpp
--- a/Engine/Plugins/FX/Niagara/Source/NiagaraEditor/Private/NiagaraScriptMergeManager.cpp
+++ b/Engine/Plugins/FX/Niagara/Source/NiagaraEditor/Private/NiagaraScriptMergeManager.cpp
@@ -1474,8 +1474,7 @@ INiagaraMergeManager::FMergeEmitterResults FNiagaraScriptMergeManager::MergeEmit
 		MergeResults.MergedInstance = Parent.Emitter->DuplicateWithoutMerging(GetTransientPackage());
 		MergeResults.MergedInstance->DisableVersioning(Parent.Version);
 		FVersionedNiagaraEmitterData* MergedEmitterData = MergeResults.MergedInstance->GetLatestEmitterData();
-		MergedEmitterData->ParentScratchPads->SetScripts(MergedEmitterData->ScratchPads->Scripts);
-		MergedEmitterData->ScratchPads->Scripts.Empty();
+		MergedEmitterData->ParentScratchPads->AppendScripts(MergedEmitterData->ScratchPads);
 	}
 	else if (State == EmitterMergeState::Failed)
 	{
@@ -1556,8 +1555,7 @@ INiagaraMergeManager::FMergeEmitterResults FNiagaraScriptMergeManager::MergeEmit
 		GetForcedChangeIds(ParentFunctionIdToNodeMap, LastMergedParentFunctionIdToNodeMap, InstanceFunctionIdToNodeMap, FunctionIdToForcedChangeId);
 
 		FVersionedNiagaraEmitterData* MergedEmitterData = VersionedMergedInstance.GetEmitterData();
-		MergedEmitterData->ParentScratchPads->SetScripts(MergedEmitterData->ScratchPads->Scripts);
-		MergedEmitterData->ScratchPads->Scripts.Empty();
+		MergedEmitterData->ParentScratchPads->AppendScripts(MergedEmitterData->ScratchPads);
 		CopyInstanceScratchPadScripts(VersionedMergedInstance, Instance.GetEmitterData());
 		TSharedRef<FNiagaraScratchPadMergeAdapter> ScratchPadAdapter = MakeShared<FNiagaraScratchPadMergeAdapter>(VersionedMergedInstance, Instance, Parent);
 
