Rejigging to hopefully bypass overzealous(wong) static analysis.

#rnx

diff --git a/Engine/Plugins/FX/Niagara/Source/NiagaraEditor/Private/ViewModels/Stack/NiagaraStackGraphUtilities.cpp b/Engine/Plugins/FX/Niagara/Source/NiagaraEditor/Private/ViewModels/Stack/NiagaraStackGraphUtilities.cpp
--- a/Engine/Plugins/FX/Niagara/Source/NiagaraEditor/Private/ViewModels/Stack/NiagaraStackGraphUtilities.cpp
+++ b/Engine/Plugins/FX/Niagara/Source/NiagaraEditor/Private/ViewModels/Stack/NiagaraStackGraphUtilities.cpp
@@ -531,15 +531,21 @@ void FNiagaraStackGraphUtilities::CheckForDeprecatedScriptVersion(UNiagaraNodeFu
 
 void FNiagaraStackGraphUtilities::CheckForDeprecatedEmitterVersion(TSharedPtr<FNiagaraEmitterViewModel> ViewModel, const FString& StackEditorDataKey, UNiagaraStackEntry::FStackIssueFixDelegate VersionUpgradeFix, TArray<UNiagaraStackEntry::FStackIssue>& OutIssues)
 {
-	if (!ViewModel || !ViewModel->GetParentEmitter().Emitter || !ViewModel->GetParentEmitter().Emitter->IsVersioningEnabled())
+	if (!ViewModel)
 	{
 		return;
 	}
+	
 	FVersionedNiagaraEmitter VersionedEmitter = ViewModel->GetParentEmitter();
 	UNiagaraEmitter* Emitter = VersionedEmitter.Emitter;
 	FVersionedNiagaraEmitterData* EmitterData = VersionedEmitter.GetEmitterData();
 
-	if (Emitter && !EmitterData)
+	if (!Emitter || !Emitter->IsVersioningEnabled())
+	{
+		return;
+	}
+
+	if (!EmitterData)
 	{
 		if (VersionUpgradeFix.IsBound())
 		{
