Fixed a bug where scratch pad modules were not merged correctly to child emitters

diff --git a/Engine/Plugins/FX/Niagara/Source/Niagara/Private/NiagaraEmitter.cpp b/Engine/Plugins/FX/Niagara/Source/Niagara/Private/NiagaraEmitter.cpp
--- a/Engine/Plugins/FX/Niagara/Source/Niagara/Private/NiagaraEmitter.cpp
+++ b/Engine/Plugins/FX/Niagara/Source/Niagara/Private/NiagaraEmitter.cpp
@@ -2667,19 +2667,11 @@ void UNiagaraEmitter::UpdateFromMergedCopy(const INiagaraMergeManager& MergeMana
 	}
 
 	// Copy scratch pad scripts.
-	EmitterData->ParentScratchPads->Scripts.Empty();
-	for (UNiagaraScript* MergedParentScratchPadScript : MergedData->ParentScratchPads->Scripts)
-	{
-		ReouterMergedObject(EmitterData->ParentScratchPads, MergedParentScratchPadScript);
-		EmitterData->ParentScratchPads->Scripts.Add(MergedParentScratchPadScript);
-	}
+	EmitterData->ParentScratchPads = NewObject<UNiagaraScratchPadContainer>(this);
+	EmitterData->ParentScratchPads->AppendScripts(MergedData->ParentScratchPads->Scripts);
 
-	EmitterData->ScratchPads->Scripts.Empty();
-	for (UNiagaraScript* MergedScratchPadScript : MergedData->ScratchPads->Scripts)
-	{
-		ReouterMergedObject(EmitterData->ScratchPads, MergedScratchPadScript);
-		EmitterData->ScratchPads->Scripts.Add(MergedScratchPadScript);
-	}
+	EmitterData->ScratchPads = NewObject<UNiagaraScratchPadContainer>(this);
+	EmitterData->ScratchPads->AppendScripts(MergedData->ScratchPads->Scripts);
 
 	UNiagaraEditorDataBase* NewEditorData = MergedData->GetEditorData();
 	ReouterMergedObject(this, NewEditorData);
diff --git a/Engine/Plugins/FX/Niagara/Source/Niagara/Private/NiagaraScratchPadContainer.cpp b/Engine/Plugins/FX/Niagara/Source/Niagara/Private/NiagaraScratchPadContainer.cpp
--- a/Engine/Plugins/FX/Niagara/Source/Niagara/Private/NiagaraScratchPadContainer.cpp
+++ b/Engine/Plugins/FX/Niagara/Source/Niagara/Private/NiagaraScratchPadContainer.cpp
@@ -13,7 +13,7 @@ void UNiagaraScratchPadContainer::PostLoad()
 void UNiagaraScratchPadContainer::CheckConsistency()
 {
 #if WITH_EDITORONLY_DATA
-	for (UNiagaraScript* Script : Scripts)
+	for (TObjectPtr<UNiagaraScript> Script : Scripts)
 	{
 		if (Script && Script->GetOuter() != this)
 		{
@@ -44,9 +44,16 @@ void UNiagaraScratchPadContainer::AppendScripts(TObjectPtr<UNiagaraScratchPadCon
 #if WITH_EDITORONLY_DATA
 	if (InScripts)
 	{
-		Scripts.Append(InScripts->Scripts);
+		for (TObjectPtr<UNiagaraScript> Script : InScripts->Scripts)
+		{
+			if (Script)
+			{
+				FName UniqueName = MakeUniqueObjectName(this, Script->GetClass(), Script->GetFName());
+				Script->Rename(*UniqueName.ToString(), this, REN_ForceNoResetLoaders | REN_NonTransactional);
+				Scripts.Add(Script);
+			}
+		}
 		InScripts->Scripts.Empty();
-		CheckConsistency();
 	}
 #endif
 }
diff --git a/Engine/Plugins/FX/Niagara/Source/NiagaraEditor/Private/NiagaraScriptMergeManager.cpp b/Engine/Plugins/FX/Niagara/Source/NiagaraEditor/Private/NiagaraScriptMergeManager.cpp
--- a/Engine/Plugins/FX/Niagara/Source/NiagaraEditor/Private/NiagaraScriptMergeManager.cpp
+++ b/Engine/Plugins/FX/Niagara/Source/NiagaraEditor/Private/NiagaraScriptMergeManager.cpp
@@ -710,9 +710,9 @@ void FNiagaraScratchPadMergeAdapter::Initialize()
 			for (int32 InstanceScratchPadScriptIndex = 0; InstanceScratchPadScriptIndex < InstanceEmitterData->ScratchPads->Scripts.Num(); InstanceScratchPadScriptIndex++)
 			{
 				UNiagaraScript* TargetCopy;
-				if (ensureMsgf(InstanceScratchPadScriptIndex < InstanceEmitterData->ScratchPads->Scripts.Num(), TEXT("Instance scratch pad script was missing from the Target's copies. Emitter %s"), *GetPathNameSafe(InstanceEmitter.Emitter)))
+				if (ensureMsgf(InstanceScratchPadScriptIndex < TargetEmitterData->ScratchPads->Scripts.Num(), TEXT("Instance scratch pad script was missing from the Target's copies. Emitter %s"), *GetPathNameSafe(InstanceEmitter.Emitter)))
 				{
-					TargetCopy = InstanceEmitterData->ScratchPads->Scripts[InstanceScratchPadScriptIndex];
+					TargetCopy = TargetEmitterData->ScratchPads->Scripts[InstanceScratchPadScriptIndex];
 				}
 				else
 				{
