- Fix for using Parameters.SvPosition vs ScreenUV

#rnx

diff --git a/Engine/Plugins/FX/Niagara/Shaders/Private/NiagaraDataInterfaceGBufferTemplate.ush b/Engine/Plugins/FX/Niagara/Shaders/Private/NiagaraDataInterfaceGBufferTemplate.ush
--- a/Engine/Plugins/FX/Niagara/Shaders/Private/NiagaraDataInterfaceGBufferTemplate.ush
+++ b/Engine/Plugins/FX/Niagara/Shaders/Private/NiagaraDataInterfaceGBufferTemplate.ush
@@ -222,7 +222,7 @@ void DecodeCustomStencil_{ParameterName}(in float2 ScreenUV, in bool bApplyViewp
 	OutValue = GBufferData.CustomStencil;
 #elif NDIGBUFFER_IS_MOBILE
 	IsValid = true;
-	OutValue = MobileSceneTextures.CustomStencilTexture.Load(int3(Parameters.SvPosition.xy, 0)) STENCIL_COMPONENT_SWIZZLE;
+	OutValue = MobileSceneTextures.CustomStencilTexture.Load(int3(ScreenUV, 0)) STENCIL_COMPONENT_SWIZZLE;
 #else
 	IsValid = false;
 	OutValue = 0;
