Fixed a crash when creating empty emitters

diff --git a/Engine/Plugins/FX/Niagara/Source/Niagara/Private/NiagaraEmitter.cpp b/Engine/Plugins/FX/Niagara/Source/Niagara/Private/NiagaraEmitter.cpp
--- a/Engine/Plugins/FX/Niagara/Source/Niagara/Private/NiagaraEmitter.cpp
+++ b/Engine/Plugins/FX/Niagara/Source/Niagara/Private/NiagaraEmitter.cpp
@@ -171,6 +171,9 @@ UNiagaraEmitter::UNiagaraEmitter(const FObjectInitializer& Initializer)
 void UNiagaraEmitter::PostInitProperties()
 {
 	Super::PostInitProperties();
+#if WITH_EDITORONLY_DATA
+	CheckVersionDataAvailable();
+#endif
 	if (HasAnyFlags(RF_ClassDefaultObject | RF_NeedLoad) == false)
 	{
 		for (FVersionedNiagaraEmitterData& Data : VersionData)
