Modifying warmup dt now affects the tick count and time properties.
Adding stack warnings on system and emitters if the emitter max dt is smaller than the system's warmup dt.

diff --git a/Engine/Plugins/FX/Niagara/Source/Niagara/Classes/NiagaraSystem.h b/Engine/Plugins/FX/Niagara/Source/Niagara/Classes/NiagaraSystem.h
--- a/Engine/Plugins/FX/Niagara/Source/Niagara/Classes/NiagaraSystem.h
+++ b/Engine/Plugins/FX/Niagara/Source/Niagara/Classes/NiagaraSystem.h
@@ -255,7 +255,7 @@ public:
 #endif 
 
 	/** Gets an array of the emitter handles. */
-	const TArray<FNiagaraEmitterHandle>& GetEmitterHandles();
+	TArray<FNiagaraEmitterHandle>& GetEmitterHandles();
 	const TArray<FNiagaraEmitterHandle>& GetEmitterHandles()const;
 	
 	FNiagaraSystemScalabilityOverrides& GetScalabilityOverrides(){return SystemScalabilityOverrides; }
@@ -351,6 +351,10 @@ public:
 	FORCEINLINE bool NeedsDeterminism() const { return bDeterminism; }
 	FORCEINLINE int32 GetRandomSeed() const { return RandomSeed; }
 
+	FORCEINLINE void SetWarmupTime(float InWarmupTime) { WarmupTime = InWarmupTime; ResolveWarmupTickCount(); }
+	FORCEINLINE void SetWarmupTickDelta(float InWarmupTickDelta) { WarmupTickDelta = InWarmupTickDelta; ResolveWarmupTickCount(); }
+	void ResolveWarmupTickCount();
+
 #if STATS
 	FNiagaraStatDatabase& GetStatData() { return StatDatabase; }
 #endif
diff --git a/Engine/Plugins/FX/Niagara/Source/Niagara/Private/NiagaraSystem.cpp b/Engine/Plugins/FX/Niagara/Source/Niagara/Private/NiagaraSystem.cpp
--- a/Engine/Plugins/FX/Niagara/Source/Niagara/Private/NiagaraSystem.cpp
+++ b/Engine/Plugins/FX/Niagara/Source/Niagara/Private/NiagaraSystem.cpp
@@ -712,6 +712,21 @@ void UNiagaraSystem::PreEditChange(FProperty* PropertyThatWillChange)
 	}
 }
 
+void UNiagaraSystem::ResolveWarmupTickCount()
+{
+	//Set the WarmupTickCount to feed back to the user.
+	if (FMath::IsNearlyZero(WarmupTickDelta))
+	{
+		WarmupTickDelta = 0.0f;
+		WarmupTickCount = 0;
+	}
+	else
+	{
+		WarmupTickCount = WarmupTime / WarmupTickDelta;
+		WarmupTime = WarmupTickDelta * WarmupTickCount;
+	}
+}
+
 void UNiagaraSystem::PostEditChangeProperty(struct FPropertyChangedEvent& PropertyChangedEvent)
 {
 	Super::PostEditChangeProperty(PropertyChangedEvent);
@@ -724,19 +739,15 @@ void UNiagaraSystem::PostEditChangeProperty(struct FPropertyChangedEvent& Proper
 		{
 			//Set the WarmupTime to feed back to the user.
 			WarmupTime = WarmupTickCount * WarmupTickDelta;
+			ResolveWarmupTickCount();
 		}
 		else if (PropertyChangedEvent.Property->GetFName() == GET_MEMBER_NAME_CHECKED(UNiagaraSystem, WarmupTime))
 		{
-			//Set the WarmupTickCount to feed back to the user.
-			if (FMath::IsNearlyZero(WarmupTickDelta))
-			{
-				WarmupTickDelta = 0.0f;
-			}
-			else
-			{
-				WarmupTickCount = WarmupTime / WarmupTickDelta;
-				WarmupTime = WarmupTickDelta * WarmupTickCount;
-			}
+			ResolveWarmupTickCount();
+		}
+		else if (PropertyChangedEvent.Property->GetFName() == GET_MEMBER_NAME_CHECKED(UNiagaraSystem, WarmupTickDelta))
+		{
+			ResolveWarmupTickCount();
 		}
 	}
 	else
@@ -959,7 +970,7 @@ void UNiagaraSystem::RemoveSystemParametersForEmitter(const FNiagaraEmitterHandl
 #endif
 
 
-const TArray<FNiagaraEmitterHandle>& UNiagaraSystem::GetEmitterHandles()
+TArray<FNiagaraEmitterHandle>& UNiagaraSystem::GetEmitterHandles()
 {
 	return EmitterHandles;
 }
diff --git a/Engine/Plugins/FX/Niagara/Source/NiagaraEditor/Private/ViewModels/Stack/NiagaraStackEmitterSettingsGroup.cpp b/Engine/Plugins/FX/Niagara/Source/NiagaraEditor/Private/ViewModels/Stack/NiagaraStackEmitterSettingsGroup.cpp
--- a/Engine/Plugins/FX/Niagara/Source/NiagaraEditor/Private/ViewModels/Stack/NiagaraStackEmitterSettingsGroup.cpp
+++ b/Engine/Plugins/FX/Niagara/Source/NiagaraEditor/Private/ViewModels/Stack/NiagaraStackEmitterSettingsGroup.cpp
@@ -137,29 +137,87 @@ void UNiagaraStackEmitterPropertiesItem::RefreshIssues(TArray<FStackIssue>& NewI
 	FNiagaraStackGraphUtilities::CheckForDeprecatedEmitterVersion(GetEmitterViewModel(), GetStackEditorDataKey(), GetUpgradeVersionFix(), NewIssues);
 	
 	FVersionedNiagaraEmitterData* ActualEmitterData = GetEmitterViewModel()->GetEmitter().GetEmitterData();
-	if (ActualEmitterData && ActualEmitterData->SimTarget == ENiagaraSimTarget::GPUComputeSim && ActualEmitterData->CalculateBoundsMode == ENiagaraEmitterCalculateBoundMode::Dynamic)
+	if (ActualEmitterData)
 	{
-		bool bAddError = true;
-		
-		UNiagaraSystem& Sys = GetSystemViewModel()->GetSystem();
-		if (Sys.bFixedBounds)
+		if (ActualEmitterData->SimTarget == ENiagaraSimTarget::GPUComputeSim && ActualEmitterData->CalculateBoundsMode == ENiagaraEmitterCalculateBoundMode::Dynamic)
 		{
-			bAddError = false;
-		}			
-		
+			bool bAddError = true;
 
-		if (bAddError)
+			UNiagaraSystem& Sys = GetSystemViewModel()->GetSystem();
+			if (Sys.bFixedBounds)
+			{
+				bAddError = false;
+			}
+
+			if (bAddError)
+			{
+				FStackIssue MissingRequiredFixedBoundsModuleError(
+					EStackIssueSeverity::Warning,
+					LOCTEXT("RequiredFixedBoundsWarningFormat", "The emitter is GPU and is using dynamic bounds mode.\r\nPlease update the Emitter or System properties otherwise bounds may be incorrect."),
+					LOCTEXT("MissingFixedBounds", "Missing fixed bounds."),
+					GetStackEditorDataKey(),
+					false);
+
+				NewIssues.Add(MissingRequiredFixedBoundsModuleError);
+			}
+		}
+
+		UNiagaraSystem& System = GetSystemViewModel()->GetSystem();
+		TWeakPtr<FNiagaraSystemViewModel> WeakSysViewModel = GetSystemViewModel();
+		if (System.NeedsWarmup())
 		{
-			FStackIssue MissingRequiredFixedBoundsModuleError(
-				EStackIssueSeverity::Warning,
-				LOCTEXT("RequiredFixedBoundsWarningFormat", "The emitter is GPU and is using dynamic bounds mode.\r\nPlease update the Emitter or System properties otherwise bounds may be incorrect."),
-				LOCTEXT("MissingFixedBounds", "Missing fixed bounds."),
-				GetStackEditorDataKey(),
-				false);
-
-			NewIssues.Add(MissingRequiredFixedBoundsModuleError);
+			float WarmupDelta = System.GetWarmupTickDelta();
+			if (ActualEmitterData->bLimitDeltaTime && ActualEmitterData->MaxDeltaTimePerTick < WarmupDelta)
+			{
+				TArray<FStackIssueFix> Fixes;
+
+				float MaxEmitterDt = ActualEmitterData->MaxDeltaTimePerTick;
+				//This emitter does not allow ticks with a delta time so large.
+				FText FixDescriptionReduceWarmupDt = LOCTEXT("FixWarmupDeltaTime", "Reduce System Warmup Delta Time");
+				Fixes.Emplace(
+					FixDescriptionReduceWarmupDt,
+					FStackIssueFixDelegate::CreateLambda([=]()
+						{
+							if (auto Pinned = WeakSysViewModel.Pin())
+							{
+								FScopedTransaction ScopedTransaction(FixDescriptionReduceWarmupDt);
+								Pinned->GetSystem().Modify();
+								Pinned->GetSystem().SetWarmupTickDelta(MaxEmitterDt);
+								Pinned->RefreshAll();
+							}
+						}));
+
+				FVersionedNiagaraEmitterWeakPtr WeakEmitter = GetEmitterViewModel()->GetEmitter().ToWeakPtr();
+				FText FixDescriptionReduceIncreaseEmitterDt = LOCTEXT("FixEmitterDeltaTime", "Increase Max Emitter Delta Time");
+				Fixes.Emplace(
+					FixDescriptionReduceIncreaseEmitterDt,
+					FStackIssueFixDelegate::CreateLambda([=]()
+						{
+							auto PinnedSysViewModel = WeakSysViewModel.Pin();
+							FVersionedNiagaraEmitter PinnedEmitter = WeakEmitter.ResolveWeakPtr();
+							if (PinnedEmitter.Emitter && PinnedSysViewModel)
+							{
+								FScopedTransaction ScopedTransaction(FixDescriptionReduceIncreaseEmitterDt);
+
+								PinnedEmitter.Emitter->Modify();
+								PinnedEmitter.GetEmitterData()->MaxDeltaTimePerTick = WarmupDelta;
+								PinnedSysViewModel->RefreshAll();
+							}
+						}));
+
+				FStackIssue WarmupDeltaTimeExceedsEmitterDeltaTimeWarning(
+					EStackIssueSeverity::Warning,
+					LOCTEXT("WarmupDeltaTimeExceedsEmitterDeltaTimeWarningSummary", "System Warmup Delta Time Exceeds Emitter Max Delta Time."),
+					LOCTEXT("WarmupDeltaTimeExceedsEmitterDeltaTimeWarningText", "Max Tick Delta Time is smaller than the System's Warmup Delta Time. This could cause unintended results during warmup for this emitter."),
+					GetStackEditorDataKey(),
+					false,
+					Fixes);
+
+				NewIssues.Add(WarmupDeltaTimeExceedsEmitterDeltaTimeWarning);
+			}
 		}
 	}
+
 }
 
 void UNiagaraStackEmitterPropertiesItem::EmitterPropertiesChanged()
diff --git a/Engine/Plugins/FX/Niagara/Source/NiagaraEditor/Private/ViewModels/Stack/NiagaraStackSystemPropertiesItem.cpp b/Engine/Plugins/FX/Niagara/Source/NiagaraEditor/Private/ViewModels/Stack/NiagaraStackSystemPropertiesItem.cpp
--- a/Engine/Plugins/FX/Niagara/Source/NiagaraEditor/Private/ViewModels/Stack/NiagaraStackSystemPropertiesItem.cpp
+++ b/Engine/Plugins/FX/Niagara/Source/NiagaraEditor/Private/ViewModels/Stack/NiagaraStackSystemPropertiesItem.cpp
@@ -70,6 +70,67 @@ void UNiagaraStackSystemPropertiesItem::RefreshChildrenInternal(const TArray<UNi
 		NewIssues.Add(OverridesWithNoEffectTypeWarning);
 	}
 
+	//Check emitters don't have a max delta time that exceeds the warmup delta time.
+	if (System->NeedsWarmup())
+	{
+		float WarmupDelta = System->GetWarmupTickDelta();
+
+		for (FNiagaraEmitterHandle& EmitterHandle : System->GetEmitterHandles())
+		{
+			FVersionedNiagaraEmitter& Emitter = EmitterHandle.GetInstance();
+			const FVersionedNiagaraEmitterData* EmitterData = Emitter.GetEmitterData();			
+			check(EmitterData);
+			if (EmitterData->bLimitDeltaTime && EmitterData->MaxDeltaTimePerTick < WarmupDelta)
+			{
+				TArray<FStackIssueFix> Fixes;
+
+				float MaxEmitterDt = EmitterData->MaxDeltaTimePerTick;
+				//This emitter does not allow ticks with a delta time so large.
+				FText FixDescriptionReduceWarmupDt = LOCTEXT("FixWarmupDeltaTime", "Reduce System Warmup Delta Time");
+				Fixes.Emplace(
+					FixDescriptionReduceWarmupDt,
+					FStackIssueFixDelegate::CreateLambda([=]()
+						{
+							if (auto Pinned = WeakSysViewModel.Pin())
+							{
+								FScopedTransaction ScopedTransaction(FixDescriptionReduceWarmupDt);
+								Pinned->GetSystem().Modify();
+								Pinned->GetSystem().SetWarmupTickDelta(MaxEmitterDt);
+								Pinned->RefreshAll();
+							}
+						}));
+
+				FVersionedNiagaraEmitterWeakPtr WeakEmitter = Emitter.ToWeakPtr();
+				FText FixDescriptionReduceIncreaseEmitterDt = LOCTEXT("FixEmitterDeltaTime", "Increase Max Emitter Delta Time");
+				Fixes.Emplace(
+					FixDescriptionReduceIncreaseEmitterDt,
+					FStackIssueFixDelegate::CreateLambda([=]()
+						{
+							auto PinnedSysViewModel = WeakSysViewModel.Pin();
+							FVersionedNiagaraEmitter PinnedEmitter = WeakEmitter.ResolveWeakPtr();
+							if (PinnedEmitter.Emitter && PinnedSysViewModel)
+							{
+								FScopedTransaction ScopedTransaction(FixDescriptionReduceIncreaseEmitterDt);
+
+								PinnedEmitter.Emitter->Modify();
+								PinnedEmitter.GetEmitterData()->MaxDeltaTimePerTick = WarmupDelta;
+								PinnedSysViewModel->RefreshAll();
+							}
+						}));
+
+				FStackIssue WarmupDeltaTimeExceedsEmitterDeltaTimeWarning(
+					EStackIssueSeverity::Warning,
+					LOCTEXT("WarmupDeltaTimeExceedsEmitterDeltaTimeWarningSummary", "System Warmup Delta Time Exceeds Emitter Max Delta Time."),
+					FText::Format(LOCTEXT("WarmupDeltaTimeExceedsEmitterDeltaTimeWarningText", "Emitter {0} has a Max Tick Delta Time smaller than the System's Warmup Delta Time. This could cause unintended results during warmup for this emitter."), FText::FromString(EmitterHandle.GetUniqueInstanceName())),
+					GetStackEditorDataKey(),
+					false,
+					Fixes);
+
+				NewIssues.Add(WarmupDeltaTimeExceedsEmitterDeltaTimeWarning);
+			}
+		}
+	}
+
 	Super::RefreshChildrenInternal(CurrentChildren, NewChildren, NewIssues);
 }
 
