Hook up global scalabiltiy culling mode to a console command to aid debugging.

diff --git a/Engine/Plugins/FX/Niagara/Source/Niagara/Private/NiagaraWorldManager.cpp b/Engine/Plugins/FX/Niagara/Source/Niagara/Private/NiagaraWorldManager.cpp
--- a/Engine/Plugins/FX/Niagara/Source/Niagara/Private/NiagaraWorldManager.cpp
+++ b/Engine/Plugins/FX/Niagara/Source/Niagara/Private/NiagaraWorldManager.cpp
@@ -197,6 +197,39 @@ FAutoConsoleCommandWithWorldAndArgs GCmdNiagaraPlaybackRate(
 	)
 );
 
+FAutoConsoleCommandWithWorldAndArgs GCmdNiagaraScalabilityCullingMode(
+	TEXT("fx.Niagara.Scalability.CullingMode"),
+	TEXT("Set scalability culling mode\n")
+	TEXT("0 - Enabled. Culling is enabled as normal.\n")
+	TEXT("1 - Paused. No culling will occur but FX will still be tracked internally so culling can be resumed correctly later.\n")
+	TEXT("2 - Disabled. No culling will occur and no FX will be tracked. Culling may not work correctly for some FX if enabled again after this.\n"),
+	FConsoleCommandWithWorldAndArgsDelegate::CreateLambda(
+		[](const TArray<FString>& Args, UWorld* World)
+		{
+			if (FNiagaraWorldManager* WorldManager = FNiagaraWorldManager::Get(World))
+			{
+				if (Args.Num() != 1 || Args[0].IsNumeric() == false)
+				{
+					UE_LOG(LogNiagara, Log, TEXT("fx.Niagara.ScalabililityCullingMode %d"), (int32)WorldManager->GetScalabilityCullingMode());
+					UE_LOG(LogNiagara, Log, TEXT("Set scalability culling mode"));
+					UE_LOG(LogNiagara, Log, TEXT("0 - Enabled. Culling is enabled as normal."));
+					UE_LOG(LogNiagara, Log, TEXT("1 - Paused. No culling will occur but FX will still be tracked internally so culling can be resumed correctly later."));
+					UE_LOG(LogNiagara, Log, TEXT("2 - Disabled. No culling will occur and no FX will be tracked. Culling may not work correctly for some FX if enabled again after this."));
+				}
+				else
+				{
+					const ENiagaraScalabilityCullingMode CullingMode = FMath::Clamp((ENiagaraScalabilityCullingMode)FCString::Atoi(*Args[0]), ENiagaraScalabilityCullingMode::Enabled, ENiagaraScalabilityCullingMode::Disabled);
+					WorldManager->SetScalabilityCullingMode(CullingMode);
+				}
+			}
+			else
+			{
+				UE_LOG(LogNiagara, Warning, TEXT("Cannot set Niagara Scalability Culling Mode on a null world."));
+			}
+		}
+	)
+);
+
 FDelegateHandle FNiagaraWorldManager::OnWorldInitHandle;
 FDelegateHandle FNiagaraWorldManager::OnWorldCleanupHandle;
 FDelegateHandle FNiagaraWorldManager::OnPostWorldCleanupHandle;
