Fixed system scalability override bools not showing the governed properties properly when enabling them

#rnx

diff --git a/Engine/Plugins/FX/Niagara/Source/NiagaraEditor/Private/Customizations/NiagaraTypeCustomizations.cpp b/Engine/Plugins/FX/Niagara/Source/NiagaraEditor/Private/Customizations/NiagaraTypeCustomizations.cpp
--- a/Engine/Plugins/FX/Niagara/Source/NiagaraEditor/Private/Customizations/NiagaraTypeCustomizations.cpp
+++ b/Engine/Plugins/FX/Niagara/Source/NiagaraEditor/Private/Customizations/NiagaraTypeCustomizations.cpp
@@ -1942,8 +1942,19 @@ void FNiagaraSystemScalabilityOverrideCustomization::CustomizeChildren(TSharedRe
 	{
 		if (TSharedPtr<IPropertyHandle> MasterProperty = PropertyHandle->GetChildHandle(MasterPropertyName))
 		{
+			TAttribute<EVisibility> VisibilityAttribute = TAttribute<EVisibility>::CreateLambda([MasterProperty]
+			{
+				bool bMasterProp = false;
+				if(MasterProperty->GetValue(bMasterProp) == FPropertyAccess::Success && bMasterProp)
+				{
+					return EVisibility::Visible;
+				}
+
+				return EVisibility::Collapsed;
+			});
+			
 			bool bMasterProp;
-			if (MasterProperty->GetValue(bMasterProp) == FPropertyAccess::Success && bMasterProp)
+			if (MasterProperty->GetValue(bMasterProp) == FPropertyAccess::Success)
 			{
 				IDetailGroup& PropGroup = ChildBuilder.AddGroup(MasterPropertyName, MasterProperty->GetPropertyDisplayName());
 				PropGroup.HeaderProperty(MasterProperty.ToSharedRef());
@@ -1951,7 +1962,8 @@ void FNiagaraSystemScalabilityOverrideCustomization::CustomizeChildren(TSharedRe
 				{
 					if (TSharedPtr<IPropertyHandle> OverrideProp = PropertyHandle->GetChildHandle(OverridePropName))
 					{
-						PropGroup.AddPropertyRow(OverrideProp.ToSharedRef());
+						PropGroup.AddPropertyRow(OverrideProp.ToSharedRef())
+						.Visibility(VisibilityAttribute);
 					}
 				}
 			}
