Niagara - Files missed in original submit - Additional fixes to allow testing creation and editing of niagara assets without the plugin supplied data.
+ Fix crashes with scratch pad scripts.
+ Fix issues creating empty emitters.
+ Add menu entries for creating empty emitters to the system overview graph and emitter wizard.
+ Fix issues creating empty scripts.

diff --git a/Engine/Plugins/FX/Niagara/Source/Niagara/Private/NiagaraCommon.cpp b/Engine/Plugins/FX/Niagara/Source/Niagara/Private/NiagaraCommon.cpp
--- a/Engine/Plugins/FX/Niagara/Source/Niagara/Private/NiagaraCommon.cpp
+++ b/Engine/Plugins/FX/Niagara/Source/Niagara/Private/NiagaraCommon.cpp
@@ -920,12 +920,21 @@ ENiagaraScriptContextStaticSwitch FNiagaraUtilities::ConvertScriptUsageToStaticS
 
 FName NIAGARA_API FNiagaraUtilities::GetUniqueName(FName CandidateName, const TSet<FName>& ExistingNames)
 {
-	if (ExistingNames.Contains(CandidateName) == false)
+	// This utility function needs to generate a unique name while only considering the text portion of the name and
+	// not the index, so generate names with 0 indices before using them for comparison.
+	TSet<FName> ExistingNamesWithIndexZero;
+	for (FName ExistingName : ExistingNames)
+	{
+		ExistingNamesWithIndexZero.Add(FName(ExistingName, 0));
+	}
+	FName CandidateNameWithIndexZero = FName(CandidateName, 0);
+
+	if (ExistingNamesWithIndexZero.Contains(CandidateNameWithIndexZero) == false)
 	{
 		return CandidateName;
 	}
 
-	FString CandidateNameString = CandidateName.ToString();
+	FString CandidateNameString = CandidateNameWithIndexZero.ToString();
 	FString BaseNameString = CandidateNameString;
 	if (CandidateNameString.Len() >= 3 && CandidateNameString.Right(3).IsNumeric())
 	{
@@ -934,7 +943,7 @@ FName NIAGARA_API FNiagaraUtilities::GetUniqueName(FName CandidateName, const TS
 
 	FName UniqueName = FName(*BaseNameString);
 	int32 NameIndex = 1;
-	while (ExistingNames.Contains(UniqueName))
+	while (ExistingNamesWithIndexZero.Contains(UniqueName))
 	{
 		UniqueName = FName(*FString::Printf(TEXT("%s%03i"), *BaseNameString, NameIndex));
 		NameIndex++;
diff --git a/Engine/Plugins/FX/Niagara/Source/Niagara/Private/NiagaraSystem.cpp b/Engine/Plugins/FX/Niagara/Source/Niagara/Private/NiagaraSystem.cpp
--- a/Engine/Plugins/FX/Niagara/Source/Niagara/Private/NiagaraSystem.cpp
+++ b/Engine/Plugins/FX/Niagara/Source/Niagara/Private/NiagaraSystem.cpp
@@ -3637,7 +3637,7 @@ void UNiagaraSystem::FixupPositionUserParameters()
 		// different types of the same user parameter linked throughout the system.
 		ForEachScript([&LinkedParameter](UNiagaraScript* Script)
 		{
-			if (Script)
+			if (Script && Script->GetLatestSource() != nullptr)
 			{
 				Script->GetLatestSource()->ChangedLinkedInputTypes(LinkedParameter, FNiagaraTypeDefinition::GetPositionDef());
 			}
