Speculative fix for random GC crashes

diff --git a/Engine/Plugins/FX/Niagara/Source/Niagara/Private/NiagaraParameterStore.cpp b/Engine/Plugins/FX/Niagara/Source/Niagara/Private/NiagaraParameterStore.cpp
--- a/Engine/Plugins/FX/Niagara/Source/Niagara/Private/NiagaraParameterStore.cpp
+++ b/Engine/Plugins/FX/Niagara/Source/Niagara/Private/NiagaraParameterStore.cpp
@@ -596,7 +596,7 @@ bool FNiagaraParameterStore::AddParameter(const FNiagaraVariable& Param, bool bI
 #if WITH_EDITORONLY_DATA
 	if (ParameterOffsets.Num())
 	{
-		UE_LOG(LogNiagara, Warning, TEXT("AddParameter: ParameterOffsets.Num() should be 0 is %d, please investigate for %s."), ParameterOffsets.Num() , *GetPathNameSafe(Owner));
+		UE_LOG(LogNiagara, Warning, TEXT("AddParameter: ParameterOffsets.Num() should be 0 is %d, please investigate for %s."), ParameterOffsets.Num() , *GetPathNameSafe(Owner.Get()));
 	}
 	//check(!ParameterOffsets.Num()); // Migration to SortedParameterOffsets
 #endif
@@ -629,7 +629,7 @@ bool FNiagaraParameterStore::AddParameter(const FNiagaraVariable& Param, bool bI
 	if (ParamType.IsDataInterface())
 	{
 		Offset = DataInterfaces.AddZeroed();
-		DataInterfaces[Offset] = bInitInterfaces ? NewObject<UNiagaraDataInterface>(Owner, const_cast<UClass*>(ParamType.GetClass()), NAME_None, RF_Transactional | RF_Public) : nullptr;
+		DataInterfaces[Offset] = bInitInterfaces ? NewObject<UNiagaraDataInterface>(Owner.Get(), const_cast<UClass*>(ParamType.GetClass()), NAME_None, RF_Transactional | RF_Public) : nullptr;
 		bInterfacesDirty = true;
 	}
 	else if (ParamType.IsUObject())
@@ -700,7 +700,7 @@ bool FNiagaraParameterStore::RemoveParameter(const FNiagaraVariableBase& ToRemov
 #if WITH_EDITORONLY_DATA
 	if (ParameterOffsets.Num())
 	{
-		UE_LOG(LogNiagara, Warning, TEXT("RemoveParameter: ParameterOffsets.Num() should be 0 is %d, please investigate for %s"), ParameterOffsets.Num(), *GetPathNameSafe(Owner));
+		UE_LOG(LogNiagara, Warning, TEXT("RemoveParameter: ParameterOffsets.Num() should be 0 is %d, please investigate for %s"), ParameterOffsets.Num(), *GetPathNameSafe(Owner.Get()));
 	}
 	//check(!ParameterOffsets.Num()); // Migration to SortedParameterOffsets
 #endif
@@ -762,7 +762,7 @@ void FNiagaraParameterStore::RenameParameter(const FNiagaraVariableBase& Param,
 #if WITH_EDITORONLY_DATA
 	if (ParameterOffsets.Num())
 	{
-		UE_LOG(LogNiagara, Warning, TEXT("RenameParameter: ParameterOffsets.Num() should be 0 is %d, please investigate for %s"), ParameterOffsets.Num(), *GetPathNameSafe(Owner));
+		UE_LOG(LogNiagara, Warning, TEXT("RenameParameter: ParameterOffsets.Num() should be 0 is %d, please investigate for %s"), ParameterOffsets.Num(), *GetPathNameSafe(Owner.Get()));
 	}
 	//check(!ParameterOffsets.Num()); // Migration to SortedParameterOffsets
 #endif
@@ -856,9 +856,9 @@ void FNiagaraParameterStore::SanityCheckData(bool bInitInterfaces)
 
 					OwnerDirtied = true;
 				}
-				if (DataInterfaces[SrcIndex] == nullptr && bInitInterfaces && Owner)
+				if (DataInterfaces[SrcIndex] == nullptr && bInitInterfaces && Owner.IsValid())
 				{
-					DataInterfaces[SrcIndex] = NewObject<UNiagaraDataInterface>(Owner, const_cast<UClass*>(Parameter.GetType().GetClass()), NAME_None, RF_Transactional | RF_Public);
+					DataInterfaces[SrcIndex] = NewObject<UNiagaraDataInterface>(Owner.Get(), const_cast<UClass*>(Parameter.GetType().GetClass()), NAME_None, RF_Transactional | RF_Public);
 					UE_LOG(LogNiagara, Verbose, TEXT("Had to initialize data interface! %s on %s"), *Parameter.GetName().ToString(), Owner != nullptr ? *Owner->GetPathName() : TEXT("Unknown owner"));
 
 					OwnerDirtied = true;
@@ -894,7 +894,7 @@ void FNiagaraParameterStore::SanityCheckData(bool bInitInterfaces)
 		ParameterData.AddZeroed(ParameterDataSize - ParameterData.Num());
 	}
 
-	if (Owner && OwnerDirtied)
+	if (Owner.IsValid() && OwnerDirtied)
 	{
 		UE_LOG(LogNiagara, Verbose, TEXT("%s needs to be resaved to prevent above warnings due to the parameter state being stale."), *Owner->GetFullName());
 	}
@@ -1028,7 +1028,7 @@ void FNiagaraParameterStore::InitFromSource(const FNiagaraParameterStore* SrcSto
 	ParameterOffsets = SrcStore->ParameterOffsets;
 	if (ParameterOffsets.Num())
 	{
-		UE_LOG(LogNiagara, Warning, TEXT("InitFromSource: ParameterOffsets.Num() should be 0 is %d, please investigate for %s"), ParameterOffsets.Num(), *GetPathNameSafe(Owner));
+		UE_LOG(LogNiagara, Warning, TEXT("InitFromSource: ParameterOffsets.Num() should be 0 is %d, please investigate for %s"), ParameterOffsets.Num(), *GetPathNameSafe(Owner.Get()));
 	}
 #endif // WITH_EDITORONLY_DATA
 	CopySortedParameterOffsets(SrcStore->ReadParameterVariables());
@@ -1281,7 +1281,7 @@ const int32* FNiagaraParameterStore::FindParameterOffset(const FNiagaraVariableB
 #if WITH_EDITORONLY_DATA
 	if (ParameterOffsets.Num())
 	{
-		UE_LOG(LogNiagara, Warning, TEXT("FindParameterOffset: ParameterOffsets.Num() should be 0 is %d, please investigate for %s"), ParameterOffsets.Num(), *GetPathNameSafe(Owner));
+		UE_LOG(LogNiagara, Warning, TEXT("FindParameterOffset: ParameterOffsets.Num() should be 0 is %d, please investigate for %s"), ParameterOffsets.Num(), *GetPathNameSafe(Owner.Get()));
 	}
 	//check(!ParameterOffsets.Num()); // Migration to SortedParameterOffsets
 #endif
diff --git a/Engine/Plugins/FX/Niagara/Source/Niagara/Public/NiagaraParameterStore.h b/Engine/Plugins/FX/Niagara/Source/Niagara/Public/NiagaraParameterStore.h
--- a/Engine/Plugins/FX/Niagara/Source/Niagara/Public/NiagaraParameterStore.h
+++ b/Engine/Plugins/FX/Niagara/Source/Niagara/Public/NiagaraParameterStore.h
@@ -179,7 +179,7 @@ struct NIAGARA_API FNiagaraParameterStore
 private:
 	/** Owner of this store. Used to provide an outer to data interfaces in this store. */
 	UPROPERTY(Transient)
-	TObjectPtr<UObject> Owner;
+	TWeakObjectPtr<UObject> Owner;
 	
 #if WITH_EDITORONLY_DATA
 	/** Map from parameter defs to their offset in the data table or the data interface. TODO: Separate out into a layout and instance class to reduce duplicated data for this?  */
@@ -253,7 +253,7 @@ public:
 #endif
 
 	void SetOwner(UObject* InOwner);
-	UObject* GetOwner()const { return Owner; }
+	UObject* GetOwner() const { return Owner.Get(); }
 
 	void Dump();
 	void DumpParameters(bool bDumpBindings = false)const;
@@ -469,8 +469,8 @@ public:
 				{
 					if (ensureMsgf(DestStore.Owner != nullptr, TEXT("Destination data interface pointer was null and a new one couldn't be created because the destination store's owner pointer was also null.")))
 					{
-						UE_LOG(LogNiagara, Warning, TEXT("While trying to copy parameter data the destination data interface was null, creating a new one.  Parameter: %s Destination Store Owner: %s"), *Parameter.GetName().ToString(), *DestStore.Owner->GetPathName());
-						DestDataInterface = NewObject<UNiagaraDataInterface>(DestStore.Owner, Parameter.GetType().GetClass(), NAME_None, RF_Transactional | RF_Public);
+						UE_LOG(LogNiagara, Warning, TEXT("While trying to copy parameter data the destination data interface was null, creating a new one.  Parameter: %s Destination Store Owner: %s"), *Parameter.GetName().ToString(), *GetPathNameSafe(DestStore.Owner.Get()));
+						DestDataInterface = NewObject<UNiagaraDataInterface>(DestStore.Owner.Get(), Parameter.GetType().GetClass(), NAME_None, RF_Transactional | RF_Public);
 						DestStore.DataInterfaces[DestIndex] = DestDataInterface;
 					}
 					else
