Some code cleanup

diff --git a/Engine/Plugins/FX/Niagara/Source/NiagaraEditor/Private/ViewModels/NiagaraSystemViewModel.cpp b/Engine/Plugins/FX/Niagara/Source/NiagaraEditor/Private/ViewModels/NiagaraSystemViewModel.cpp
--- a/Engine/Plugins/FX/Niagara/Source/NiagaraEditor/Private/ViewModels/NiagaraSystemViewModel.cpp
+++ b/Engine/Plugins/FX/Niagara/Source/NiagaraEditor/Private/ViewModels/NiagaraSystemViewModel.cpp
@@ -294,13 +294,16 @@ TSharedPtr<FNiagaraEmitterHandleViewModel> FNiagaraSystemViewModel::GetEmitterHa
 
 TSharedPtr<FNiagaraEmitterHandleViewModel> FNiagaraSystemViewModel::GetEmitterHandleViewModelForEmitter(const FVersionedNiagaraEmitter& InEmitter) const
 {
-	for (TSharedRef<FNiagaraEmitterHandleViewModel> EmitterHandleViewModel : EmitterHandleViewModels)
+	if (FVersionedNiagaraEmitterData* InEmitterData = InEmitter.GetEmitterData())
 	{
-		FVersionedNiagaraEmitterData* InEmitterData = InEmitter.GetEmitterData();
-		FVersionedNiagaraEmitterData* ViewModelEmitterData = EmitterHandleViewModel->GetEmitterViewModel()->GetEmitter().GetEmitterData();
-		if (InEmitterData == ViewModelEmitterData)
+		for (TSharedRef<FNiagaraEmitterHandleViewModel> EmitterHandleViewModel : EmitterHandleViewModels)
 		{
-			return EmitterHandleViewModel;
+			// we compare the pointer to the emitter data here because depending on the versioning status,
+			// two guids might point to the same or to different emitters 
+			if (InEmitterData == EmitterHandleViewModel->GetEmitterViewModel()->GetEmitter().GetEmitterData())
+			{
+				return EmitterHandleViewModel;
+			}
 		}
 	}
 	return TSharedPtr<FNiagaraEmitterHandleViewModel>();
diff --git a/Engine/Plugins/FX/Niagara/Source/NiagaraEditor/Private/ViewModels/Stack/NiagaraStackGraphUtilities.cpp b/Engine/Plugins/FX/Niagara/Source/NiagaraEditor/Private/ViewModels/Stack/NiagaraStackGraphUtilities.cpp
--- a/Engine/Plugins/FX/Niagara/Source/NiagaraEditor/Private/ViewModels/Stack/NiagaraStackGraphUtilities.cpp
+++ b/Engine/Plugins/FX/Niagara/Source/NiagaraEditor/Private/ViewModels/Stack/NiagaraStackGraphUtilities.cpp
@@ -3538,15 +3538,13 @@ int32 FNiagaraStackGraphUtilities::DependencyUtilities::FindBestIndexForModuleIn
 
 	// Get the stack module data for the emitter stack to find dependencies.
 	TSharedPtr<FNiagaraSystemViewModel> SystemViewModel = TNiagaraViewModelManager<UNiagaraSystem, FNiagaraSystemViewModel>::GetExistingViewModelForObject(System);
-	if (SystemViewModel.IsValid() == false)
+	if (!ensureMsgf(SystemViewModel.IsValid(), TEXT("Failed to get systemviewmodel for valid system when getting best index in stack for module!")))
 	{
-		ensureMsgf(false, TEXT("Failed to get systemviewmodel for valid system when getting best index in stack for module!"));
 		return INDEX_NONE;
 	}
 	TSharedPtr<FNiagaraEmitterHandleViewModel> EmitterHandleViewModel = SystemViewModel->GetEmitterHandleViewModelForEmitter(OuterEmitter);
-	if (EmitterHandleViewModel.IsValid() == false)
+	if (!ensureMsgf(EmitterHandleViewModel.IsValid(), TEXT("Failed to get emitterhandleviewmodel for valid emitter when getting best index in stack for module!")))
 	{
-		ensureMsgf(false, TEXT("Failed to get emitterhandleviewmodel for valid emitter when getting best index in stack for module!"));
 		return INDEX_NONE;
 	}
 	const TArray<FNiagaraStackModuleData>& StackModuleData = SystemViewModel->GetStackModuleDataByEmitterHandleId(EmitterHandleViewModel->GetId());
