Fixed python api that caused ambiguous search path failures

diff --git a/Engine/Plugins/FX/Niagara/Source/NiagaraEditor/Private/UpgradeNiagaraScriptResults.cpp b/Engine/Plugins/FX/Niagara/Source/NiagaraEditor/Private/UpgradeNiagaraScriptResults.cpp
--- a/Engine/Plugins/FX/Niagara/Source/NiagaraEditor/Private/UpgradeNiagaraScriptResults.cpp
+++ b/Engine/Plugins/FX/Niagara/Source/NiagaraEditor/Private/UpgradeNiagaraScriptResults.cpp
@@ -292,7 +292,7 @@ void UNiagaraPythonModule::Init(UNiagaraStackModuleItem* InModuleItem)
 	ModuleItem = InModuleItem;
 }
 
-UNiagaraStackModuleItem* UNiagaraPythonModule::Object() const
+UNiagaraStackModuleItem* UNiagaraPythonModule::GetObject() const
 {
 	return ModuleItem;
 }
@@ -302,7 +302,7 @@ void UNiagaraPythonEmitter::Init(TSharedRef<FNiagaraEmitterHandleViewModel> InEm
 	EmitterViewModel = InEmitterViewModel;
 }
 
-UNiagaraEmitter* UNiagaraPythonEmitter::Object()
+UNiagaraEmitter* UNiagaraPythonEmitter::GetObject()
 {
 	return EmitterViewModel->GetEmitterViewModel()->GetEmitter().Emitter;
 }
diff --git a/Engine/Plugins/FX/Niagara/Source/NiagaraEditor/Public/UpgradeNiagaraScriptResults.h b/Engine/Plugins/FX/Niagara/Source/NiagaraEditor/Public/UpgradeNiagaraScriptResults.h
--- a/Engine/Plugins/FX/Niagara/Source/NiagaraEditor/Public/UpgradeNiagaraScriptResults.h
+++ b/Engine/Plugins/FX/Niagara/Source/NiagaraEditor/Public/UpgradeNiagaraScriptResults.h
@@ -166,7 +166,7 @@ public:
 
 	// Returns the raw underlying object
 	UFUNCTION(BlueprintCallable, Category = "Scripting")
-	UNiagaraStackModuleItem* Object() const;
+	UNiagaraStackModuleItem* GetObject() const;
 
 	//TODO: extend the api here to fully support module actions
 
@@ -188,7 +188,7 @@ public:
 
 	// Returns the raw underlying object
 	UFUNCTION(BlueprintCallable, Category = "Scripting")
-	UNiagaraEmitter* Object();
+	UNiagaraEmitter* GetObject();
 
 	// returns the emitter properties, such as determinism or interpolated spawning
 	UFUNCTION(BlueprintCallable, Category = "Scripting")
