Fixed build error

diff --git a/Engine/Plugins/FX/Niagara/Source/NiagaraEditor/Private/UpgradeNiagaraScriptResults.cpp b/Engine/Plugins/FX/Niagara/Source/NiagaraEditor/Private/UpgradeNiagaraScriptResults.cpp
--- a/Engine/Plugins/FX/Niagara/Source/NiagaraEditor/Private/UpgradeNiagaraScriptResults.cpp
+++ b/Engine/Plugins/FX/Niagara/Source/NiagaraEditor/Private/UpgradeNiagaraScriptResults.cpp
@@ -353,7 +353,7 @@ TArray<UNiagaraPythonModule*> UNiagaraPythonEmitter::GetModules() const
 
 bool UNiagaraPythonEmitter::HasModule(const FString& ModuleName) const
 {
-	return GetModule(ModuleName)->Object() != nullptr;
+	return GetModule(ModuleName)->GetObject() != nullptr;
 }
 
 UNiagaraPythonModule* UNiagaraPythonEmitter::GetModule(const FString& ModuleName) const
