Exporting renderer layout for customer request.

diff --git a/Engine/Plugins/FX/Niagara/Source/Niagara/Public/NiagaraRendererProperties.h b/Engine/Plugins/FX/Niagara/Source/Niagara/Public/NiagaraRendererProperties.h
--- a/Engine/Plugins/FX/Niagara/Source/Niagara/Public/NiagaraRendererProperties.h
+++ b/Engine/Plugins/FX/Niagara/Source/Niagara/Public/NiagaraRendererProperties.h
@@ -129,7 +129,7 @@ struct FNiagaraRendererVariableInfo
 };
 
 /** Used for building renderer layouts for vertex factories */
-struct FNiagaraRendererLayout
+struct NIAGARA_API FNiagaraRendererLayout
 {
 	void Initialize(int32 NumVariables);
 	bool SetVariable(const FNiagaraDataSetCompiledData* CompiledData, const FNiagaraVariableBase& Variable, int32 VFVarOffset);
